Copying //UE4/Dev-Niagara to Dev-Main (//UE4/Dev-Main)

diff --git a/Engine/Plugins/FX/Niagara/Source/Niagara/Private/NiagaraEmitterInstanceBatcher.cpp b/Engine/Plugins/FX/Niagara/Source/Niagara/Private/NiagaraEmitterInstanceBatcher.cpp
--- a/Engine/Plugins/FX/Niagara/Source/Niagara/Private/NiagaraEmitterInstanceBatcher.cpp
+++ b/Engine/Plugins/FX/Niagara/Source/Niagara/Private/NiagaraEmitterInstanceBatcher.cpp
@@ -7,12 +7,14 @@
 #include "NiagaraStats.h"
 #include "NiagaraShader.h"
 #include "NiagaraSortingGPU.h"
+#include "NiagaraWorldManager.h"
 #include "ShaderParameterUtils.h"
 #include "SceneUtils.h"
 #include "ShaderParameterUtils.h"
 #include "ClearQuad.h"
 #include "Runtime/Engine/Private/GPUSort.h"
 
+DECLARE_CYCLE_STAT(TEXT("Niagara Dispatch Setup"), STAT_NiagaraGPUDispatchSetup_RT, STATGROUP_Niagara);
 DECLARE_CYCLE_STAT(TEXT("GPU Emitter Dispatch [RT]"), STAT_NiagaraGPUSimTick_RT, STATGROUP_Niagara);
 DECLARE_CYCLE_STAT(TEXT("GPU Data Readback [RT]"), STAT_NiagaraGPUReadback_RT, STATGROUP_Niagara);
 DECLARE_CYCLE_STAT(TEXT("Allocate GPU Readback Data [RT]"), STAT_NiagaraAllocateGPUReadback_RT, STATGROUP_Niagara);
@@ -25,9 +27,27 @@ DECLARE_DWORD_COUNTER_STAT(TEXT("Readback latency (frames)"), STAT_NiagaraReadba
 DECLARE_GPU_STAT_NAMED(NiagaraGPU, TEXT("Niagara"));
 DECLARE_GPU_STAT_NAMED(NiagaraGPUSimulation, TEXT("Niagara GPU Simulation"));
 DECLARE_GPU_STAT_NAMED(NiagaraGPUSorting, TEXT("Niagara GPU sorting"));
+DECLARE_GPU_STAT_NAMED(NiagaraSimReadback, TEXT("Niagara GPU Simulation read back"));
+DECLARE_GPU_STAT_NAMED(NiagaraIndexBufferClear, TEXT("Niagara index buffer clear"));
 
 uint32 FNiagaraComputeExecutionContext::TickCounter = 0;
 
+int32 GNiagaraOverlapCompute = 1;
+static FAutoConsoleVariableRef CVarNiagaraUseAsyncCompute(
+	TEXT("fx.NiagaraOverlapCompute"),
+	GNiagaraOverlapCompute,
+	TEXT("If 1, use compute dispatch overlap for better performance. \n"),
+	ECVF_Default
+);
+
+// @todo REMOVE THIS HACK
+int32 GNiagaraGpuMaxQueuedRenderFrames = 10;
+static FAutoConsoleVariableRef CVarNiagaraGpuMaxQueuedRenderFrames(
+	TEXT("fx.NiagaraGpuMaxQueuedRenderFrames"),
+	GNiagaraOverlapCompute,
+	TEXT("Number of frames we all to pass before we start to discard GPU ticks.\n"),
+	ECVF_Default
+);
 
 FNiagaraIndicesVertexBuffer::FNiagaraIndicesVertexBuffer(int32 InIndexCount)
 	: IndexCount(InIndexCount)
@@ -48,80 +68,430 @@ FFXSystemInterface* NiagaraEmitterInstanceBatcher::GetInterface(const FName& InN
 NiagaraEmitterInstanceBatcher::~NiagaraEmitterInstanceBatcher()
 {
 	ParticleSortBuffers.ReleaseRHI();
+
+	FinishDispatches();
 }
 
-void NiagaraEmitterInstanceBatcher::Queue(FNiagaraComputeExecutionContext *ExecContext)
+void NiagaraEmitterInstanceBatcher::GiveSystemTick_RenderThread(FNiagaraGPUSystemTick& Tick)
 {
-	if (!IsPendingKill())
-	{
-		//UE_LOG(LogNiagara, Warning, TEXT("Submitted!"));
-		//SimulationQueue[CurQueueIndex]->Add(InContext);
-		TArray<FNiagaraComputeExecutionContext*>* Queue = &SimulationQueue[0];
-		uint32 QueueIndex = CurQueueIndex;
-		ENQUEUE_RENDER_COMMAND(QueueNiagaraDispatch)(
-			[Queue, QueueIndex, ExecContext](FRHICommandListImmediate& RHICmdList)
+	check(IsInRenderingThread());
+
+	// @todo REMOVE THIS HACK
+	if (GFrameNumberRenderThread > LastFrameThatDrainedData + GNiagaraGpuMaxQueuedRenderFrames)
+	{
+		Tick.Destroy();
+		return;
+	}
+
+	// Now we consume DataInterface instance data.
+	if (Tick.DIInstanceData)
+	{
+		uint8* BasePointer = (uint8*) Tick.DIInstanceData->PerInstanceDataForRT;
+
+		//UE_LOG(LogNiagara, Log, TEXT("RT Give DI (dipacket) %p (baseptr) %p"), Tick.DIInstanceData, BasePointer);
+		for(auto& Pair : Tick.DIInstanceData->InterfaceProxiesToOffsets)
 		{
-			const uint32 QueueIndexMask = (1 << QueueIndex);
-			//Don't queue the same context for execution multiple times. TODO: possibly try to combine/accumulate the tick info if we happen to have > 1 before it's executed.
-			if (!(ExecContext->PendingExecutionQueueMask & QueueIndexMask))
+			FNiagaraDataInterfaceProxy* Proxy = Pair.Key;
+			uint8* InstanceDataPtr = BasePointer + Pair.Value;
+
+			//UE_LOG(LogNiagara, Log, TEXT("\tRT DI (proxy) %p (instancebase) %p"), Proxy, InstanceDataPtr);
+			Proxy->ConsumePerInstanceDataFromGameThread(InstanceDataPtr, Tick.SystemInstanceID);
+		}
+	}
+
+	// A note:
+	// This is making a copy of Tick. That structure is small now and we take a copy to avoid
+	// making a bunch of small allocations on the game thread. We may need to revisit this.
+	Ticks_RT.Add(Tick);
+}
+
+void NiagaraEmitterInstanceBatcher::GiveEmitterContextToDestroy_RenderThread(FNiagaraComputeExecutionContext* Context)
+{
+	ContextsToDestroy_RT.Add(Context);
+}
+
+void NiagaraEmitterInstanceBatcher::GiveDataSetToDestroy_RenderThread(FNiagaraDataSet* DataSet)
+{
+	DataSetsToDestroy_RT.Add(DataSet);
+}
+
+void NiagaraEmitterInstanceBatcher::FinishDispatches()
+{
+	ReleaseTicks();
+
+	for (FNiagaraComputeExecutionContext* Context : ContextsToDestroy_RT)
+	{
+		delete Context;
+	}
+	ContextsToDestroy_RT.Reset();
+
+	for (FNiagaraDataSet* DataSet : DataSetsToDestroy_RT)
+	{
+		delete DataSet;
+	}
+	DataSetsToDestroy_RT.Reset();
+
+	for (TSharedPtr<FNiagaraDataInterfaceProxy, ESPMode::ThreadSafe>& Proxy : DIProxyDeferredDeletes_RT)
+	{
+		Proxy->DeferredDestroy();
+	}
+
+	DIProxyDeferredDeletes_RT.Empty();
+}
+
+void NiagaraEmitterInstanceBatcher::ReleaseTicks()
+{
+	check(IsInRenderingThread());
+
+	for (FNiagaraGPUSystemTick& Tick : Ticks_RT)
+	{
+		Tick.Destroy();
+	}
+
+	Ticks_RT.Empty(0);
+}
+
+void NiagaraEmitterInstanceBatcher::ResizeBuffersAndGatherResources(FOverlappableTicks& OverlappableTick, FRHICommandList& RHICmdList, FNiagaraBufferArray& DestDataBuffers, FNiagaraBufferArray& CurrDataBuffers, FNiagaraBufferArray& DestBufferIntFloat, FNiagaraBufferArray& CurrBufferIntFloat)
+{
+	SCOPE_CYCLE_COUNTER(STAT_NiagaraGPUDispatchSetup_RT);
+
+	for (FNiagaraGPUSystemTick* Tick : OverlappableTick)
+	{
+		uint32 DispatchCount = Tick->Count;
+		FNiagaraComputeInstanceData* Instances = Tick->GetInstanceData();
+		for (uint32 Index = 0; Index < DispatchCount; Index++)
+		{
+			FNiagaraComputeInstanceData& Instance = Instances[Index];
+			FNiagaraComputeExecutionContext* Context = Instance.Context;
+			if ( Context == nullptr )
+			{
+				continue;
+			}
+
+			FNiagaraShader* Shader = Context->GPUScript_RT->GetShader();
+			if ( Shader == nullptr )
+			{
+				continue;
+			}
+
+			//The buffer containing current simulation state.
+			Instance.CurrentData = Context->MainDataSet->GetCurrentData();
+			//The buffer we're going to write simulation results to.
+			Instance.DestinationData = &Context->MainDataSet->BeginSimulate();
+
+			check(Instance.CurrentData && Instance.DestinationData);
+			FNiagaraDataBuffer& CurrentData = *Instance.CurrentData;
+			FNiagaraDataBuffer& DestinationData = *Instance.DestinationData;
+
+			const uint32 PrevNumInstances = CurrentData.GetNumInstances();
+			const uint32 NewNumInstances = Instance.SpawnRateInstances + Instance.EventSpawnTotal + PrevNumInstances;
+
+			//We must assume all particles survive when allocating here. 
+			//If this is not true, the read back in ResolveDatasetWrites will shrink the buffers.
+			const uint32 RequiredInstances = FMath::Max(PrevNumInstances, NewNumInstances);
+
+			DestinationData.AllocateGPU(RequiredInstances + 1, RHICmdList);
+			DestinationData.SetNumInstances(RequiredInstances);
+
+			if ( Shader->FloatInputBufferParam.IsBound() )
 			{
-				Queue[QueueIndex].Add(ExecContext);
-				ExecContext->PendingExecutionQueueMask |= QueueIndexMask;
+				CurrDataBuffers.Add(CurrentData.GetGPUBufferFloat().UAV);
 			}
-		});
+			if ( Shader->IntInputBufferParam.IsBound() )
+			{
+				CurrBufferIntFloat.Add(CurrentData.GetGPUBufferInt().UAV);
+			}
+
+			if ( Shader->FloatOutputBufferParam.IsBound() )
+			{
+				DestDataBuffers.Add(DestinationData.GetGPUBufferFloat().UAV);
+			}
+			if ( Shader->IntOutputBufferParam.IsBound() )
+			{
+				DestBufferIntFloat.Add(DestinationData.GetGPUBufferInt().UAV);
+			}
+
+			Context->MainDataSet->EndSimulate();
+			Context->SetDataToRender(Instance.DestinationData);
+		}
 	}
 }
 
-void NiagaraEmitterInstanceBatcher::Remove(FNiagaraComputeExecutionContext *ExecContext)
+void NiagaraEmitterInstanceBatcher::DispatchAllOnCompute(FOverlappableTicks& OverlappableTick, FRHICommandList& RHICmdList, FUniformBufferRHIParamRef ViewUniformBuffer, FNiagaraBufferArray& DestDataBuffers, FNiagaraBufferArray& CurrDataBuffers, FNiagaraBufferArray& DestBufferIntFloat, FNiagaraBufferArray& CurrBufferIntFloat)
 {
-	if (!IsPendingKill())
+	FRHICommandListImmediate& RHICmdListImmediate = FRHICommandListExecutor::GetImmediateCommandList();
+
+	// Disable automatic cache flush so that we can have our compute work overlapping. Barrier will be used as a sync mechanism.
+	RHICmdList.AutomaticCacheFlushAfterComputeShader(false);
+
+	//
+	//	Transition current index buffer ready for compute and clear then all using overlapping compute work items.
+	//
+
+	RHICmdList.TransitionResources(EResourceTransitionAccess::ERWBarrier, EResourceTransitionPipeline::EGfxToCompute, CurrDataBuffers.GetData(), CurrDataBuffers.Num());
+
+	{
+		SCOPED_DRAW_EVENT(RHICmdList, NiagaraIndexBufferClear);
+		SCOPED_GPU_STAT(RHICmdList, NiagaraIndexBufferClear);
+
+		for (FNiagaraGPUSystemTick* Tick : OverlappableTick)
+		{
+			uint32 DispatchCount = Tick->Count;
+			FNiagaraComputeInstanceData* Instances = Tick->GetInstanceData();
+			for (uint32 Index = 0; Index < DispatchCount; Index++)
+			{
+				FNiagaraComputeInstanceData& Instance = Instances[Index];
+				FNiagaraComputeExecutionContext* Context = Instance.Context;
+				if (Context && Context->GPUScript_RT->GetShader())
+				{
+#if WITH_EDITORONLY_DATA
+					if (Context->DebugInfo.IsValid())
+					{
+						ProcessDebugInfo(RHICmdList, Context);
+					}
+#endif // WITH_EDITORONLY_DATA
+
+					// clear data set index buffer for the simulation shader to write number of written instances
+					ClearUAV(RHICmdList, Instance.DestinationData->GetGPUIndices(), 0);
+				}
+
+			}
+		}
+	}
+
+	//
+	//	Add a rw barrier for the destination data buffers we just cleared and mark others as read/write as needed for particles simulation.
+	//
+	RHICmdList.TransitionResources(EResourceTransitionAccess::ERWBarrier, EResourceTransitionPipeline::EComputeToCompute, DestDataBuffers.GetData(), DestDataBuffers.Num());
+	RHICmdList.TransitionResources(EResourceTransitionAccess::EWritable, EResourceTransitionPipeline::EComputeToCompute, DestBufferIntFloat.GetData(), DestBufferIntFloat.Num());
+	RHICmdList.TransitionResources(EResourceTransitionAccess::EReadable, EResourceTransitionPipeline::EComputeToCompute, CurrDataBuffers.GetData(), CurrDataBuffers.Num());
+	RHICmdList.TransitionResources(EResourceTransitionAccess::EReadable, EResourceTransitionPipeline::EComputeToCompute, CurrBufferIntFloat.GetData(), CurrBufferIntFloat.Num());
+	RHICmdList.FlushComputeShaderCache();
+
+	{
+		SCOPED_DRAW_EVENT(RHICmdList, NiagaraGPUSimulation);
+		SCOPED_GPU_STAT(RHICmdList, NiagaraGPUSimulation);
+		for (FNiagaraGPUSystemTick* Tick : OverlappableTick)
+		{
+			uint32 DispatchCount = Tick->Count;
+			FNiagaraComputeInstanceData* Instances = Tick->GetInstanceData();
+			for (uint32 Index = 0; Index < DispatchCount; Index++)
+			{
+				FNiagaraComputeInstanceData& Instance = Instances[Index];
+				FNiagaraComputeExecutionContext* Context = Instance.Context;
+				if (Context && Context->GPUScript_RT->GetShader())
+				{
+					FNiagaraComputeExecutionContext::TickCounter++;
+
+					// run shader, sim and spawn in a single dispatch
+					uint32 UpdateStartInstance = 0;
+					Run<false>(*Tick, &Instance, UpdateStartInstance, Instance.DestinationData->GetNumInstances(), Context->GPUScript_RT->GetShader(), RHICmdList, ViewUniformBuffer);
+				}
+			}
+		}
+	}
+
+
+	//
+	//	Now Copy to staging buffer the data we want to read back (alive particle count). And make buffer ready for that and draw commands on the graphics pipe too.
+	//
+	RHICmdList.TransitionResources(EResourceTransitionAccess::ERWBarrier, EResourceTransitionPipeline::EComputeToCompute, DestDataBuffers.GetData(), DestDataBuffers.Num());
+	RHICmdList.TransitionResources(EResourceTransitionAccess::ERWBarrier, EResourceTransitionPipeline::EComputeToCompute, DestBufferIntFloat.GetData(), DestBufferIntFloat.Num());
+	RHICmdList.TransitionResources(EResourceTransitionAccess::EReadable, EResourceTransitionPipeline::EComputeToGfx, DestDataBuffers.GetData(), DestDataBuffers.Num());
+	RHICmdList.TransitionResources(EResourceTransitionAccess::EReadable, EResourceTransitionPipeline::EComputeToGfx, DestBufferIntFloat.GetData(), DestBufferIntFloat.Num());
+	RHICmdList.FlushComputeShaderCache();
+
+	// We have done all our overlapping compute work on this list so go back to default behavior and flush.
+	RHICmdList.AutomaticCacheFlushAfterComputeShader(true);
+
 	{
-		TArray<FNiagaraComputeExecutionContext*>* Queue = &SimulationQueue[0];
-		ENQUEUE_RENDER_COMMAND(RemoveNiagaraDispatch)(
-			[Queue, ExecContext](FRHICommandListImmediate& RHICmdList)
+		SCOPED_DRAW_EVENT(RHICmdList, NiagaraSimReadback);
+		SCOPED_GPU_STAT(RHICmdList, NiagaraSimReadback);
+		for (FNiagaraGPUSystemTick* Tick : OverlappableTick)
 		{
-			for (int32 i = 0; i < SIMULATION_QUEUE_COUNT; i++)
+			uint32 DispatchCount = Tick->Count;
+			FNiagaraComputeInstanceData* Instances = Tick->GetInstanceData();
+			for (uint32 Index = 0; Index < DispatchCount; Index++)
 			{
-				Queue[i].Remove(ExecContext);
+				FNiagaraComputeInstanceData& Instance = Instances[Index];
+				FNiagaraComputeExecutionContext* Context = Instance.Context;
+				if (Context && Context->GPUScript_RT->GetShader())
+				{
+					// Don't resolve if the data if there are no instances (prevents a transition issue warning).
+					if (Instance.DestinationData->GetNumInstances() > 0)
+					{
+						// resolve data set writes - grabs the number of instances written from the index set during the simulation run
+						ResolveDatasetWrites(RHICmdList, &Instance);
+					}
+					check(Instance.DestinationData->GetGPUIndices().Buffer);
+				}
 			}
-			ExecContext->PendingExecutionQueueMask = 0;
-		});
+		}
 	}
+	// the VF grabs current state for drawing, so need to transition
+	
+	//TODO: Need to set the data to render in the context?
+
+	RHICmdList.TransitionResources(EResourceTransitionAccess::EReadable, EResourceTransitionPipeline::EComputeToGfx, DestBufferIntFloat.GetData(), DestBufferIntFloat.Num());
+
+	// We have done all our compute work
+	RHICmdList.FlushComputeShaderCache();
+	RHICmdList.SubmitCommandsHint();
+}
+
+void NiagaraEmitterInstanceBatcher::PostRenderOpaque(FRHICommandListImmediate& RHICmdList, const FUniformBufferRHIParamRef ViewUniformBuffer, const class FShaderParametersMetadata* SceneTexturesUniformBufferStruct, FUniformBufferRHIParamRef SceneTexturesUniformBuffer)
+{
+	ExecuteAll(RHICmdList, ViewUniformBuffer);
+	// Sort buffer only after the tick, so that the sorting and rendering stay coherent.
+	SortGPUParticles(RHICmdList);
 }
 
 void NiagaraEmitterInstanceBatcher::ExecuteAll(FRHICommandList &RHICmdList, FUniformBufferRHIParamRef ViewUniformBuffer)
 {
+	// @todo REMOVE THIS HACK
+	LastFrameThatDrainedData = GFrameNumberRenderThread;
+
+	// This is always called by the renderer so early out if we have no work.
+	if (Ticks_RT.Num() == 0)
+	{
+		return;
+	}
+
 	SCOPED_DRAW_EVENT(RHICmdList, NiagaraEmitterInstanceBatcher_ExecuteAll);
-	SCOPED_GPU_STAT(RHICmdList, NiagaraGPU);
 
-	const uint32 TickedQueueIndex = (CurQueueIndex ^ 0x1);
-	const uint32 TickedQueueIndexMask = (1 << TickedQueueIndex);
 
-	TArray<FNiagaraComputeExecutionContext*> &WorkQueue = SimulationQueue[TickedQueueIndex];
-	for (FNiagaraComputeExecutionContext *Context : WorkQueue)
+	if (GNiagaraOverlapCompute>0)
 	{
-		if (Context)
+		TArray<FOverlappableTicks, TInlineAllocator<2>> SimPasses;
 		{
-			// need to call RenderThreadInit on data interfaces
-			//ExecuteSingle(Context, RHICmdList);
-			TickSingle(Context, RHICmdList, ViewUniformBuffer);
-			Context->PendingExecutionQueueMask &= ~TickedQueueIndexMask;
+			TMap<FNiagaraComputeExecutionContext*, FOverlappableTicks> ContextToTicks;
+
+			// Ticks are added in order. Two tick with the same context should not overlap so should be in two different batch.
+			// Those ticks should still be executed in order.
+			for (FNiagaraGPUSystemTick& Tick : Ticks_RT)
+			{
+				FNiagaraComputeInstanceData* data = Tick.GetInstanceData();
+				ContextToTicks.FindOrAdd(data->Context).Add(&Tick);
+			}
+
+			// Count the maximum number of tick per context to know the number of passes we need to do
+			uint32 NumSimPass = 0;
+			for (auto& Ticks : ContextToTicks)
+			{
+				NumSimPass = FMath::Max(NumSimPass, (uint32)Ticks.Value.Num());
+			}
+
+			// Transpose now only once the data to get all independent tick per pass
+			SimPasses.AddDefaulted(NumSimPass);
+			for (auto& Ticks : ContextToTicks)
+			{
+				uint32 tickPass = 0;
+				for (auto& Tick : Ticks.Value)
+				{
+					SimPasses[tickPass].Add(Tick);
+					tickPass++;
+				}
+			}
 		}
 
+		for (auto& SimPass : SimPasses)
+		{
+			FNiagaraBufferArray DestDataBuffers;
+			FNiagaraBufferArray CurrDataBuffers;
+			FNiagaraBufferArray DestBufferIntFloat;
+			FNiagaraBufferArray CurrBufferIntFloat;
+
+			// This initial pass gathers all the buffers that are read from and written to so we can do batch resource transitions.
+			// It also ensures the GPU buffers are large enough to hold everything.
+			ResizeBuffersAndGatherResources(SimPass, RHICmdList, DestDataBuffers, CurrDataBuffers, DestBufferIntFloat, CurrBufferIntFloat);
+
+			DispatchAllOnCompute(SimPass, RHICmdList, ViewUniformBuffer, DestDataBuffers, CurrDataBuffers, DestBufferIntFloat, CurrBufferIntFloat);
+		}
+	}
+	else
+	{
+		for (FNiagaraGPUSystemTick& Tick : Ticks_RT)
+		{
+			uint32 DispatchCount = Tick.Count;
+			FNiagaraComputeInstanceData* Instances = Tick.GetInstanceData();
+
+			for (uint32 i = 0; i < DispatchCount; i++)
+			{
+				FNiagaraComputeInstanceData& Instance = Instances[i];
+				TickSingle(Tick, &Instance, RHICmdList, ViewUniformBuffer);
+			}
+		}
 	}
-	WorkQueue.Empty();
+
+	FinishDispatches();
 }
 
-void NiagaraEmitterInstanceBatcher::TickSingle(FNiagaraComputeExecutionContext *Context, FRHICommandList &RHICmdList, FUniformBufferRHIParamRef ViewUniformBuffer) const
+void NiagaraEmitterInstanceBatcher::SimStepClearAndSetup(const FNiagaraComputeInstanceData& Instance, FRHICommandList& RHICmdList) const
+{
+	check(IsInRenderingThread());
+#if 1
+	FNiagaraComputeExecutionContext* Context = Instance.Context;
+	check(Context);
+
+
+	FNiagaraShader* ComputeShader = Context->GPUScript_RT->GetShader();
+	check(ComputeShader);
+	if (!ComputeShader)
+	{
+		return;
+	}
+
+#if WITH_EDITORONLY_DATA
+	if (Context->DebugInfo.IsValid())
+	{
+		ProcessDebugInfo(RHICmdList, Context);
+	}
+#endif // WITH_EDITORONLY_DATA
+
+	// clear data set index buffer for the simulation shader to write number of written instances to
+	//
+	FRWBuffer &DatasetIndexBufferWrite = Instance.DestinationData->GetGPUIndices();
+//	SCOPED_GPU_EVENTF(RHICmdList, NiagaraIndexBufferClear, TEXT("Niagara index buffer clear"));
+	SCOPED_DRAW_EVENT(RHICmdList, NiagaraIndexBufferClear);
+	SCOPED_GPU_STAT(RHICmdList, NiagaraIndexBufferClear);
+
+	ClearUAV(RHICmdList, DatasetIndexBufferWrite, 0);
+#endif
+}
+
+//void NiagaraEmitterInstanceBatcher::SimStepReadback(const FNiagaraComputeInstanceData& Instance, FRHICommandList& RHICmdList) const
+//{
+//#if 1
+//	FNiagaraComputeExecutionContext* Context = Instance.Context;
+//	check(Context);
+//	FNiagaraShader* ComputeShader = Context->GPUScript_RT->GetShader();
+//	if (!ComputeShader)
+//	{
+//		return;
+//	}
+//
+//	// Don't resolve if the data if there are no instances (prevents a transition issue warning).
+//	if (Context->MainDataSet->CurrData().GetNumInstances() > 0)
+//	{
+//		// resolve data set writes - grabs the number of instances written from the index set during the simulation run
+//		ResolveDatasetWrites(RHICmdList, &Instance);
+//	}
+//
+//	check(Context->MainDataSet->HasDatasetIndices());
+//#endif
+//}
+
+void NiagaraEmitterInstanceBatcher::TickSingle(const FNiagaraGPUSystemTick& Tick, FNiagaraComputeInstanceData *Instance, FRHICommandList &RHICmdList, FUniformBufferRHIParamRef ViewUniformBuffer) const
 {
 	SCOPE_CYCLE_COUNTER(STAT_NiagaraGPUSimTick_RT);
 
+	FNiagaraComputeExecutionContext* Context = Instance->Context;
+	
 	check(IsInRenderingThread());
-	Context->MainDataSet->Tick();
 
 	FNiagaraComputeExecutionContext::TickCounter++;
 
-	FNiagaraShader* ComputeShader = Context->RTGPUScript->GetShader();
+	FNiagaraShader* ComputeShader = Context->GPUScript_RT->GetShader();
 	if (!ComputeShader)
 	{
 		return;
@@ -134,40 +504,50 @@ void NiagaraEmitterInstanceBatcher::TickSingle(FNiagaraComputeExecutionContext *
 	}
 #endif // WITH_EDITORONLY_DATA
 
-	uint32 PrevNumInstances = Context->MainDataSet->PrevData().GetNumInstances();
-	uint32 NewNumInstances = Context->SpawnRateInstances + Context->EventSpawnTotal + PrevNumInstances;
-	Context->EventSpawnTotal = GetEventSpawnTotal(Context);
+	//The buffer containing current simulation state.
+	Instance->CurrentData = Context->MainDataSet->GetCurrentData();
+	//The buffer we're going to write simulation results to.
+	Instance->DestinationData = &Context->MainDataSet->BeginSimulate();
 
-	ResizeCurrentBuffer(RHICmdList, Context, NewNumInstances, PrevNumInstances);
+	check(Instance->CurrentData && Instance->DestinationData);
+	FNiagaraDataBuffer& CurrentData = *Instance->CurrentData;
+	FNiagaraDataBuffer& DestinationData = *Instance->DestinationData;
 
-	// set up a data set index buffer, if we don't have one yet
-	//
-	if (!Context->MainDataSet->HasDatasetIndices())
-	{
-		Context->MainDataSet->SetupCurDatasetIndices();
-	}
+	uint32 PrevNumInstances = CurrentData.GetNumInstances();
+	uint32 NewNumInstances = Instance->SpawnRateInstances + Instance->EventSpawnTotal + PrevNumInstances;
+
+	//We must assume all particles survive when allocating here. 
+	//If this is not true, the read back in ResolveDatasetWrites will shrink the buffers.
+	uint32 RequiredInstances = FMath::Max(PrevNumInstances, NewNumInstances);
+
+	DestinationData.AllocateGPU(RequiredInstances + 1, RHICmdList);
+	DestinationData.SetNumInstances(RequiredInstances);
 
 	// clear data set index buffer for the simulation shader to write number of written instances to
-	//
-	ClearIndexBufferCur(RHICmdList, Context);
+	{
+		FRWBuffer& DatasetIndexBufferWrite = DestinationData.GetGPUIndices();
+		SCOPED_DRAW_EVENTF(RHICmdList, NiagaraIndexBufferClear, TEXT("Niagara index buffer clear"));
+		ClearUAV(RHICmdList, DatasetIndexBufferWrite, 0);
+		RHICmdList.TransitionResource(EResourceTransitionAccess::EWritable, EResourceTransitionPipeline::EComputeToCompute, DatasetIndexBufferWrite.UAV);
+	}
 
 	// run shader, sim and spawn in a single dispatch
 	uint32 UpdateStartInstance = 0;
-	Run(Context, UpdateStartInstance, NewNumInstances, ComputeShader, RHICmdList, ViewUniformBuffer);
-
-	// assume all instances survived; ResolveDataSetWrites will change this if the deferred readback was successful; that data may be several frames old
-	Context->MainDataSet->CurrData().SetNumInstances(NewNumInstances);
+	Run<true>(Tick, Instance, UpdateStartInstance, NewNumInstances, ComputeShader, RHICmdList, ViewUniformBuffer);
 
 	// ResolveDatasetWrites may read this, so we must transition it here.
-	RHICmdList.TransitionResource(EResourceTransitionAccess::EReadable, EResourceTransitionPipeline::EComputeToGfx, Context->MainDataSet->GetCurDataSetIndices().UAV);		// transition to readable; we'll be using this next frame
+	RHICmdList.TransitionResource(EResourceTransitionAccess::EReadable, EResourceTransitionPipeline::EComputeToGfx, DestinationData.GetGPUIndices().UAV);		// transition to readable; we'll be using this next frame
 
 	// Don't resolve if the data if there are no instances (prevents a transition issue warning).
 	if (NewNumInstances > 0)
 	{
 		// resolve data set writes - grabs the number of instances written from the index set during the simulation run
-		ResolveDatasetWrites(RHICmdList, Context);
+		ResolveDatasetWrites(RHICmdList, Instance);
 	}
 
+	Context->MainDataSet->EndSimulate();
+	Context->SetDataToRender(Instance->DestinationData);
+
 	/*
 	// TODO: hack - only updating event set 0 on update scripts now; need to match them to their indices and update them all
 	if (Context->UpdateEventWriteDataSets.Num())
@@ -178,10 +558,9 @@ void NiagaraEmitterInstanceBatcher::TickSingle(FNiagaraComputeExecutionContext *
 	*/
 
 	// the VF grabs PrevDataRender for drawing, so need to transition
-	RHICmdList.TransitionResource(EResourceTransitionAccess::EReadable, EResourceTransitionPipeline::EComputeToGfx, Context->MainDataSet->PrevData().GetGPUBufferFloat()->UAV);
-	RHICmdList.TransitionResource(EResourceTransitionAccess::EReadable, EResourceTransitionPipeline::EComputeToGfx, Context->MainDataSet->PrevData().GetGPUBufferInt()->UAV);
-
-	check(Context->MainDataSet->HasDatasetIndices());
+	// Better to do this in the VF set?
+	RHICmdList.TransitionResource(EResourceTransitionAccess::EReadable, EResourceTransitionPipeline::EComputeToGfx, Context->DataToRender->GetGPUBufferFloat().UAV);
+	RHICmdList.TransitionResource(EResourceTransitionAccess::EReadable, EResourceTransitionPipeline::EComputeToGfx, Context->DataToRender->GetGPUBufferInt().UAV);
 }
 
 void NiagaraEmitterInstanceBatcher::PreInitViews()
@@ -195,6 +574,24 @@ void NiagaraEmitterInstanceBatcher::PreInitViews()
 	}
 }
 
+bool NiagaraEmitterInstanceBatcher::UsesGlobalDistanceField() const
+{
+	for (const FNiagaraGPUSystemTick& Tick : Ticks_RT)
+	{
+		if (Tick.bRequiredDistanceFieldData)
+		{
+			return true;
+		}
+	}
+
+	return false;
+}
+
+void NiagaraEmitterInstanceBatcher::PreRender(FRHICommandListImmediate& RHICmdList, const class FGlobalDistanceFieldParameterData* GlobalDistanceFieldParameterData)
+{
+	GlobalDistanceFieldParams = GlobalDistanceFieldParameterData ? *GlobalDistanceFieldParameterData : FGlobalDistanceFieldParameterData();
+}
+
 int32 NiagaraEmitterInstanceBatcher::AddSortedGPUSimulation(const FNiagaraGPUSortInfo& SortInfo)
 {
 	const int32 ResultOffset = SortedParticleCount;
@@ -235,26 +632,51 @@ void NiagaraEmitterInstanceBatcher::SortGPUParticles(FRHICommandListImmediate& R
 			ParticleSortBuffers.InitRHI();
 		}
 
-		SET_DWORD_STAT(STAT_NiagaraGPUSortedParticles, SortedParticleCount);
-		SET_DWORD_STAT(STAT_NiagaraGPUSortedBuffers, ParticleSortBuffers.GetSize());
+		INC_DWORD_STAT_BY(STAT_NiagaraGPUSortedParticles, SortedParticleCount);
+		INC_DWORD_STAT_BY(STAT_NiagaraGPUSortedBuffers, ParticleSortBuffers.GetSize());
 
 		// Make sure our outputs are safe to write to.
-		FUnorderedAccessViewRHIParamRef OutputUAVs[2] = { ParticleSortBuffers.GetKeyBufferUAV(), ParticleSortBuffers.GetVertexBufferUAV() };
+		const int32 InitialSortBufferIndex = 0;
+		FUnorderedAccessViewRHIParamRef OutputUAVs[2] = { ParticleSortBuffers.GetKeyBufferUAV(InitialSortBufferIndex), ParticleSortBuffers.GetVertexBufferUAV(InitialSortBufferIndex) };
 		RHICmdList.TransitionResources(EResourceTransitionAccess::ERWBarrier, EResourceTransitionPipeline::EGfxToCompute, OutputUAVs, 2);
 
+		// EmitterKey = (EmitterIndex & EmitterKeyMask) << EmitterKeyShift.
+		// SortKey = (Key32 >> SortKeyShift) & SortKeyMask.
+		uint32 EmitterKeyMask = (1 << FMath::CeilLogTwo(SimulationsToSort.Num())) - 1;
+		uint32 EmitterKeyShift = 16;
+		uint32 SortKeyMask = 0xFFFF;
+		
 		{
 			SCOPED_DRAW_EVENT(RHICmdList, NiagaraSortKeyGen);
 
 			// Bind the shader
-			TShaderMapRef<FNiagaraSortKeyGenCS> KeyGenCS(GetGlobalShaderMap(FeatureLevel));
+			
+			FNiagaraSortKeyGenCS::FPermutationDomain PermutationVector;
+			PermutationVector.Set<FNiagaraSortKeyGenCS::FSortUsingMaxPrecision>(GNiagaraGPUSortingUseMaxPrecision != 0);
+			
+			TShaderMapRef<FNiagaraSortKeyGenCS> KeyGenCS(GetGlobalShaderMap(FeatureLevel), PermutationVector);
 			RHICmdList.SetComputeShader(KeyGenCS->GetComputeShader());
-			KeyGenCS->SetOutput(RHICmdList, ParticleSortBuffers.GetKeyBufferUAV(), ParticleSortBuffers.GetVertexBufferUAV());
+			KeyGenCS->SetOutput(RHICmdList, ParticleSortBuffers.GetKeyBufferUAV(InitialSortBufferIndex), ParticleSortBuffers.GetVertexBufferUAV(InitialSortBufferIndex));
 
+			// (SortKeyMask, SortKeyShift, SortKeySignBit)
+			FUintVector4 SortKeyParams(SortKeyMask, 0, 0x8000, 0); 
+			if (GNiagaraGPUSortingUseMaxPrecision != 0)
+			{
+				EmitterKeyMask = FMath::Max<uint32>(EmitterKeyMask , 1); // Need at list 1 bit for the above logic
+				uint32 UnusedBits = FPlatformMath::CountLeadingZeros(EmitterKeyMask << EmitterKeyShift);
+				EmitterKeyShift += UnusedBits;
+				SortKeyMask = ~(EmitterKeyMask << EmitterKeyShift);
+
+				SortKeyParams.X = SortKeyMask;
+				SortKeyParams.Y = 16 - UnusedBits;
+				SortKeyParams.Z <<= UnusedBits;
+			}
+			
 			int32 OutputOffset = 0;
 			for (int32 EmitterIndex = 0; EmitterIndex < SimulationsToSort.Num(); ++EmitterIndex)
 			{
 				const FNiagaraGPUSortInfo& SortInfo = SimulationsToSort[EmitterIndex];
-				KeyGenCS->SetParameters(RHICmdList, SortInfo, EmitterIndex, OutputOffset);
+				KeyGenCS->SetParameters(RHICmdList, SortInfo, (uint32)EmitterIndex << EmitterKeyShift, OutputOffset, SortKeyParams);
 				DispatchComputeShader(RHICmdList, *KeyGenCS, FMath::DivideAndRoundUp(SortInfo.ParticleCount, NIAGARA_KEY_GEN_THREAD_COUNT), 1, 1);
 
 				OutputOffset += SortInfo.ParticleCount;
@@ -267,9 +689,8 @@ void NiagaraEmitterInstanceBatcher::SortGPUParticles(FRHICommandListImmediate& R
 
 		// Sort buffers and copy results to index buffers.
 		{
-			const uint32 EmitterKeyMask = (1 << FMath::CeilLogTwo(SimulationsToSort.Num())) - 1;
-			const uint32 KeyMask = (EmitterKeyMask << 16) | 0xFFFF;
-			const int32 ResultBufferIndex = SortGPUBuffers(RHICmdList, ParticleSortBuffers.GetSortBuffers(), 0, KeyMask, SortedParticleCount, FeatureLevel);
+			const uint32 KeyMask = (EmitterKeyMask << EmitterKeyShift) | SortKeyMask;
+			const int32 ResultBufferIndex = SortGPUBuffers(RHICmdList, ParticleSortBuffers.GetSortBuffers(), InitialSortBufferIndex, KeyMask, SortedParticleCount, FeatureLevel);
 			ResolveParticleSortBuffers(RHICmdList, ResultBufferIndex);
 		}
 
@@ -350,6 +771,8 @@ void NiagaraEmitterInstanceBatcher::ResolveParticleSortBuffers(FRHICommandListIm
 			LastCount = SortBuffer.UsedIndexCount;
 		}
 
+		RHICmdList.TransitionResources(EResourceTransitionAccess::ERWBarrier, EResourceTransitionPipeline::EComputeToCompute, UAVs, NumBuffers);
+
 		CopyBufferCS->SetParameters(RHICmdList, ParticleSortBuffers.GetSortedVertexBufferSRV(ResultBufferIndex), UAVs, UsedIndexCounts, StartingIndex, NumBuffers);
 		DispatchComputeShader(RHICmdList, *CopyBufferCS, FMath::DivideAndRoundUp(LastCount - StartingIndex, NIAGARA_COPY_BUFFER_THREAD_COUNT), 1, 1);
 		RHICmdList.TransitionResources(EResourceTransitionAccess::ERWBarrier, EResourceTransitionPipeline::EComputeToGfx, UAVs, NumBuffers);
@@ -359,47 +782,20 @@ void NiagaraEmitterInstanceBatcher::ResolveParticleSortBuffers(FRHICommandListIm
 	CopyBufferCS->UnbindBuffers(RHICmdList);
 }
 
-/* Calculate total number of spawned instances from all spawn events
-*/
-uint32 NiagaraEmitterInstanceBatcher::GetEventSpawnTotal(const FNiagaraComputeExecutionContext *InContext) const
-{
-	uint32 EventSpawnTotal = 0;
-	for (int32 i = 0; i < InContext->GetEventHandlers().Num(); i++)
-	{
-		const FNiagaraEventScriptProperties &EventHandlerProps = InContext->GetEventHandlers()[i];
-		if (EventHandlerProps.ExecutionMode == EScriptExecutionMode::SpawnedParticles && InContext->EventSets[i])
-		{
-			uint32 NumEventsToProcess = InContext->EventSets[i]->PrevData().GetNumInstances();
-			uint32 EventSpawnNum = NumEventsToProcess * EventHandlerProps.SpawnNumber;
-			EventSpawnTotal += EventSpawnNum;
-		}
-	}
-	return EventSpawnTotal;
-}
-
-
-/* Clear the data set index buffer; needs to be called before a sim run
- */
-void NiagaraEmitterInstanceBatcher::ClearIndexBufferCur(FRHICommandList &RHICmdList, FNiagaraComputeExecutionContext *Context) const
-{
-	FRWBuffer &DatasetIndexBufferWrite = Context->MainDataSet->GetCurDataSetIndices();
-	SCOPED_DRAW_EVENTF(RHICmdList, NiagaraIndexBufferClear, TEXT("Niagara index buffer clear"));
-
-	ClearUAV(RHICmdList, DatasetIndexBufferWrite, 0);
-	RHICmdList.TransitionResource(EResourceTransitionAccess::EWritable, EResourceTransitionPipeline::EComputeToCompute, DatasetIndexBufferWrite.UAV);
-}
-
 /* Attempt to read back simulation results (number of live instances) from the GPU via an async readback request; 
  *	If the readback isn't ready to be performed, we accumulate spawn rates and assume all instances have survived, until
  *	the GPU can tell us how many are actually alive; since that data may be several frames old, we'll always end up
  *	overallocating a bit, and the CPU might think we have more particles alive than we actually do;
  *	since we use DrawIndirect with the GPU determining draw call parameters, that's not an issue
  */
-void NiagaraEmitterInstanceBatcher::ResolveDatasetWrites(FRHICommandList &RHICmdList, FNiagaraComputeExecutionContext *Context) const
+void NiagaraEmitterInstanceBatcher::ResolveDatasetWrites(FRHICommandList &RHICmdList, FNiagaraComputeInstanceData *Instance) const
 {
-	FRWBuffer &DatasetIndexBufferWrite = Context->MainDataSet->GetCurDataSetIndices();
-	uint32 SpawnedThisFrame = Context->SpawnRateInstances + Context->EventSpawnTotal;
+	FNiagaraComputeExecutionContext* Context = Instance->Context;
+	FNiagaraDataBuffer& DestinationData = *Instance->DestinationData;
+	FRWBuffer &DatasetIndexBufferWrite = DestinationData.GetGPUIndices();
+	uint32 SpawnedThisFrame = Instance->SpawnRateInstances + Instance->EventSpawnTotal;
 	Context->AccumulatedSpawnRate += SpawnedThisFrame;
+	int32 ExistingDataCount = DestinationData.GetNumInstances();
 	if (!Context->GPUDataReadback)
 	{
 		SCOPE_CYCLE_COUNTER(STAT_NiagaraAllocateGPUReadback_RT);
@@ -407,6 +803,7 @@ void NiagaraEmitterInstanceBatcher::ResolveDatasetWrites(FRHICommandList &RHICmd
 		Context->GPUDataReadback = new FRHIGPUMemoryReadback(TEXT("Niagara GPU Emitter Readback"));
 		INC_DWORD_STAT(STAT_NiagaraReadbackLatency);
 		Context->GPUDataReadback->EnqueueCopy(RHICmdList, DatasetIndexBufferWrite.Buffer);
+		INC_DWORD_STAT_BY(STAT_NiagaraGPUParticles, ExistingDataCount);
 	}
 	else if (Context->GPUDataReadback->IsReady())
 	{
@@ -416,12 +813,10 @@ void NiagaraEmitterInstanceBatcher::ResolveDatasetWrites(FRHICommandList &RHICmd
 		    int32 *NumInstancesAfterSim = static_cast<int32*>(Context->GPUDataReadback->Lock(64 * sizeof(int32)));
 			if (NumInstancesAfterSim)
 			{
-				int32 ExistingDataCount = Context->MainDataSet->CurrData().GetNumInstances();// index 1 is always the count
-				int32 NewExistingDataCount = NumInstancesAfterSim[1] + Context->AccumulatedSpawnRate;
-				Context->MainDataSet->CurrData().SetNumInstances(NewExistingDataCount);
-				FString PathName = Context->GPUScript->GetOutermost()->GetPathName();
-				//UE_LOG(LogNiagara, Log, TEXT("GPU Syncup %s : Was(%d) Now(%d)"), *PathName, ExistingDataCount, NewExistingDataCount );
-				SET_DWORD_STAT(STAT_NiagaraGPUParticles, NewExistingDataCount);
+				int32 NewExistingDataCount = NumInstancesAfterSim[1] + Context->AccumulatedSpawnRate; // index 1 is always the count
+				DestinationData.SetNumInstances(NewExistingDataCount);
+				//UE_LOG(LogNiagara, Log, TEXT("GPU Syncup %s : Was(%d) Now(%d)"), *Context->DebugSimName, ExistingDataCount, NewExistingDataCount );
+				INC_DWORD_STAT_BY(STAT_NiagaraGPUParticles, NewExistingDataCount);
 				SET_DWORD_STAT(STAT_NiagaraReadbackLatency, 0);
 
 				Context->AccumulatedSpawnRate = 0;
@@ -430,6 +825,7 @@ void NiagaraEmitterInstanceBatcher::ResolveDatasetWrites(FRHICommandList &RHICmd
 			else
 			{
 				UE_LOG(LogNiagara, Warning, TEXT("GPUDataReadback said it was ready, but returned an invalid buffer. Skipping this time.."));
+				INC_DWORD_STAT_BY(STAT_NiagaraGPUParticles, ExistingDataCount);
 			}
 			Context->GPUDataReadback->Unlock();
 		}
@@ -440,16 +836,19 @@ void NiagaraEmitterInstanceBatcher::ResolveDatasetWrites(FRHICommandList &RHICmd
 			Context->GPUDataReadback->EnqueueCopy(RHICmdList, DatasetIndexBufferWrite.Buffer);
 		}
 	}
+	else
+	{
+		INC_DWORD_STAT_BY(STAT_NiagaraGPUParticles, ExistingDataCount);
+	}
 }
 
-void NiagaraEmitterInstanceBatcher::ProcessDebugInfo(FRHICommandList &RHICmdList, const FNiagaraComputeExecutionContext *Context) const
+void NiagaraEmitterInstanceBatcher::ProcessDebugInfo(FRHICommandList &RHICmdList, const FNiagaraComputeExecutionContext* Context) const
 {
 #if WITH_EDITORONLY_DATA
 	// This method may be called from one of two places: in the tick or as part of a paused frame looking for the debug info that was submitted previously...
 	// Note that PrevData is where we expect the data to be for rendering, as per NiagaraEmitterInstanceBatcher::TickSingle
-	if (Context->DebugInfo.IsValid())
+	if (Context && Context->DebugInfo.IsValid())
 	{
-		
 		// Fire off the readback if not already doing so
 		if (!Context->GPUDebugDataReadbackFloat && !Context->GPUDebugDataReadbackInt && !Context->GPUDebugDataReadbackCounts)
 		{
@@ -476,7 +875,9 @@ void NiagaraEmitterInstanceBatcher::ProcessDebugInfo(FRHICommandList &RHICmdList
 				{
 					IntDataBuffer = static_cast<int*>(Context->GPUDebugDataReadbackInt->Lock(Context->GPUDebugDataIntSize));
 				}
-				Context->MainDataSet->DumpGPU(Context->DebugInfo->Frame, FloatDataBuffer, IntDataBuffer, 0, NewExistingDataCount);
+
+				Context->DebugInfo->Frame.CopyFromGPUReadback(FloatDataBuffer, IntDataBuffer, 0, NewExistingDataCount);
+
 				Context->DebugInfo->bWritten = true;
 
 				if (Context->GPUDebugDataReadbackFloat)
@@ -515,86 +916,64 @@ void NiagaraEmitterInstanceBatcher::ProcessDebugInfo(FRHICommandList &RHICmdList
 #endif // WITH_EDITORONLY_DATA
 }
 
-
-
-/* Resize data set buffers and set number of instances 
- *	Allocates one additional instance at the end, which is a scratch instance; by setting the default index from AcquireIndex in the shader
- *	to that scratch index, we can avoid branching in every single OutputData function
- */
-void NiagaraEmitterInstanceBatcher::ResizeCurrentBuffer(FRHICommandList &RHICmdList, FNiagaraComputeExecutionContext *Context, uint32 NewNumInstances, uint32 PrevNumInstances) const
-{
-
-	// allocate for additional instances spawned and set the new number in the data set, if the new number is greater (meaning if we're spawning in this run)
-	// TODO: interpolated spawning
-	//
-	if (NewNumInstances > PrevNumInstances)
-	{
-		//UE_LOG(LogNiagara, Warning, TEXT("Resize up!"));
-		Context->MainDataSet->CurrData().AllocateGPU(NewNumInstances + 1, RHICmdList);
-		Context->MainDataSet->CurrData().SetNumInstances(NewNumInstances);
-	}
-	// if we're not spawning, we need to make sure that the current buffer alloc size and number of instances matches the last one
-	// we may have spawned in the last tick, so the buffers may be different sizes
-	//
-	else if (Context->MainDataSet->CurrData().GetNumInstances() < Context->MainDataSet->PrevData().GetNumInstances())
-	{
-		//UE_LOG(LogNiagara, Warning, TEXT("Resize down!"));
-		Context->MainDataSet->CurrData().AllocateGPU(PrevNumInstances + 1, RHICmdList);
-		Context->MainDataSet->CurrData().SetNumInstances(PrevNumInstances);
-	}
-}
-
-
 /* Set shader parameters for data interfaces
  */
-void NiagaraEmitterInstanceBatcher::SetDataInterfaceParameters(const TArray<UNiagaraDataInterface*> &DataInterfaces, FNiagaraShader* Shader, FRHICommandList &RHICmdList, const FNiagaraComputeExecutionContext *Context) const
+void NiagaraEmitterInstanceBatcher::SetDataInterfaceParameters(const TArray<FNiagaraDataInterfaceProxy*>& DataInterfaceProxies, FNiagaraShader* Shader, FRHICommandList &RHICmdList, const FNiagaraComputeInstanceData* Instance, const FNiagaraGPUSystemTick& Tick) const
 {
 	// set up data interface buffers, as defined by the DIs during compilation
 	//
+
+	// @todo-threadsafety This is a bit gross. Need to rethink this api.
+	const FGuid& SystemInstance = Tick.SystemInstanceID;
+
 	uint32 InterfaceIndex = 0;
-	for (UNiagaraDataInterface* Interface : DataInterfaces)
+	for (FNiagaraDataInterfaceProxy* Interface : DataInterfaceProxies)
 	{
 		FNiagaraDataInterfaceParamRef& DIParam = Shader->GetDIParameters()[InterfaceIndex];
 		if (DIParam.Parameters)
 		{
-			void* PerInstanceData = nullptr;
-			int32* OffsetFound = nullptr;
-			if (Context->PerInstanceDataSize != 0 && Context->PerInstanceDataInterfaceOffsets != nullptr)
-			{
-				OffsetFound = Context->PerInstanceDataInterfaceOffsets->Find(Interface);
-				if (OffsetFound != nullptr)
-				{
-					PerInstanceData = (*OffsetFound) + (uint8*)Context->PerInstanceData;
-				}
-			}
-			DIParam.Parameters->Set(RHICmdList, Shader, Interface, PerInstanceData);
+			FNiagaraDataInterfaceSetArgs Context;
+			Context.Shader = Shader;
+			Context.DataInterface = Interface;
+			Context.SystemInstance = SystemInstance;
+			Context.Batcher = this;
+			DIParam.Parameters->Set(RHICmdList, Context);
 		}
 
 		InterfaceIndex++;
 	}
 }
 
-void NiagaraEmitterInstanceBatcher::UnsetDataInterfaceParameters(const TArray<UNiagaraDataInterface*> &DataInterfaces, FNiagaraShader* Shader, FRHICommandList &RHICmdList, const FNiagaraComputeExecutionContext *Context) const
+void NiagaraEmitterInstanceBatcher::UnsetDataInterfaceParameters(const TArray<FNiagaraDataInterfaceProxy*> &DataInterfaceProxies, FNiagaraShader* Shader, FRHICommandList &RHICmdList, const FNiagaraComputeInstanceData* Instance, const FNiagaraGPUSystemTick& Tick) const
 {
 	// set up data interface buffers, as defined by the DIs during compilation
 	//
+
+	// @todo-threadsafety This is a bit gross. Need to rethink this api.
+	const FGuid& SystemInstance = Tick.SystemInstanceID;
+
 	uint32 InterfaceIndex = 0;
-	for (UNiagaraDataInterface* Interface : DataInterfaces)
+	for (FNiagaraDataInterfaceProxy* Interface : DataInterfaceProxies)
 	{
 		FNiagaraDataInterfaceParamRef& DIParam = Shader->GetDIParameters()[InterfaceIndex];
 		if (DIParam.Parameters)
 		{
 			void* PerInstanceData = nullptr;
 			int32* OffsetFound = nullptr;
-			if (Context->PerInstanceDataSize != 0 && Context->PerInstanceDataInterfaceOffsets != nullptr)
+			if (Tick.DIInstanceData && Tick.DIInstanceData->PerInstanceDataSize != 0 && Tick.DIInstanceData->InterfaceProxiesToOffsets.Num() != 0)
 			{
-				OffsetFound = Context->PerInstanceDataInterfaceOffsets->Find(Interface);
+				OffsetFound = Tick.DIInstanceData->InterfaceProxiesToOffsets.Find(Interface);
 				if (OffsetFound != nullptr)
 				{
-					PerInstanceData = (*OffsetFound) + (uint8*)Context->PerInstanceData;
+					PerInstanceData = (*OffsetFound) + (uint8*)Tick.DIInstanceData->PerInstanceDataForRT;
 				}
 			}
-			DIParam.Parameters->Unset(RHICmdList, Shader, Interface, PerInstanceData);
+			FNiagaraDataInterfaceSetArgs Context;
+			Context.Shader = Shader;
+			Context.DataInterface = Interface;
+			Context.SystemInstance = SystemInstance;
+			Context.Batcher = this;
+			DIParam.Parameters->Unset(RHICmdList, Context);
 		}
 
 		InterfaceIndex++;
@@ -604,10 +983,12 @@ void NiagaraEmitterInstanceBatcher::UnsetDataInterfaceParameters(const TArray<UN
 
 /* Kick off a simulation/spawn run
  */
-void NiagaraEmitterInstanceBatcher::Run(const FNiagaraComputeExecutionContext *Context, uint32 UpdateStartInstance, const uint32 TotalNumInstances, FNiagaraShader* Shader,
+ template<bool bDoResourceTransitions>
+void NiagaraEmitterInstanceBatcher::Run(const FNiagaraGPUSystemTick& Tick, const FNiagaraComputeInstanceData* Instance, uint32 UpdateStartInstance, const uint32 TotalNumInstances, FNiagaraShader* Shader,
 	FRHICommandList &RHICmdList, FUniformBufferRHIParamRef ViewUniformBuffer, bool bCopyBeforeStart) const
 {
-	if (TotalNumInstances == 0)
+	FNiagaraComputeExecutionContext* Context = Instance->Context;
+	if (TotalNumInstances == 0 && !bDoResourceTransitions)
 	{
 		SCOPED_DRAW_EVENTF(RHICmdList, NiagaraGPUSimulationCS, TEXT("Niagara Gpu Sim - %s - NumInstances: %u"),
 			*Context->DebugSimName,
@@ -616,13 +997,15 @@ void NiagaraEmitterInstanceBatcher::Run(const FNiagaraComputeExecutionContext *C
 	}
 
 	//UE_LOG(LogNiagara, Warning, TEXT("Run"));
-
-	FNiagaraDataSet *DataSet = Context->MainDataSet;
-	const FNiagaraParameterStore& ParameterStore = Context->CombinedParamStore;
+	
+	const TArray<FNiagaraDataInterfaceProxy*>& DataInterfaceProxies = Instance->DataInterfaceProxies;
 	const FRHIUniformBufferLayout& CBufferLayout = Context->CBufferLayout;
-	const TArray<uint8, TAlignedHeapAllocator<16>> &Params = Context->ParamData_RT;
-	const FRWBuffer &WriteIndexBuffer = DataSet->GetCurDataSetIndices();
-	FRWBuffer &ReadIndexBuffer = DataSet->GetPrevDataSetIndices();
+	check(Instance->CurrentData && Instance->DestinationData);
+	FNiagaraDataBuffer& DestinationData = *Instance->DestinationData;
+	FNiagaraDataBuffer& CurrentData = *Instance->CurrentData;
+
+	const FRWBuffer& WriteIndexBuffer = DestinationData.GetGPUIndices();
+	FRWBuffer& ReadIndexBuffer = CurrentData.GetGPUIndices();
 
 	// if we don't have a previous index buffer, we need to prep one using the maximum number of instances; this should only happen on the first frame
 	//		the data set index buffer is really the param buffer for the indirect draw call; it contains the number of live instances at index 1, and the simulation
@@ -646,14 +1029,12 @@ void NiagaraEmitterInstanceBatcher::Run(const FNiagaraComputeExecutionContext *C
 		RHICmdList.SetShaderUniformBuffer(Shader->GetComputeShader(), Shader->ViewUniformBufferParam.GetBaseIndex(), ViewUniformBuffer);
 	}
 
-	SetDataInterfaceParameters(ParameterStore.GetDataInterfaces(), Shader, RHICmdList, Context);
+	SetDataInterfaceParameters(DataInterfaceProxies, Shader, RHICmdList, Instance, Tick);
 
 	// set the shader and data set params 
 	//
-	uint32 WriteBufferIdx = 0;
-	uint32 ReadBufferIdx = 0;
-	DataSet->SetShaderParams(Shader, RHICmdList, WriteBufferIdx, ReadBufferIdx);
-	
+	CurrentData.SetShaderParams<bDoResourceTransitions>(Shader, RHICmdList, true);
+	DestinationData.SetShaderParams<bDoResourceTransitions>(Shader, RHICmdList, false);
 
 	// set the index buffer uav
 	//
@@ -673,7 +1054,7 @@ void NiagaraEmitterInstanceBatcher::Run(const FNiagaraComputeExecutionContext *C
 
 	RHICmdList.SetShaderParameter(Shader->GetComputeShader(), Shader->UpdateStartInstanceParam.GetBufferIndex(), Shader->UpdateStartInstanceParam.GetBaseIndex(), Shader->UpdateStartInstanceParam.GetNumBytes(), &UpdateStartInstance);					// 0, except for event handler runs
 	RHICmdList.SetShaderParameter(Shader->GetComputeShader(), Shader->NumIndicesPerInstanceParam.GetBufferIndex(), Shader->NumIndicesPerInstanceParam.GetBaseIndex(), Shader->NumIndicesPerInstanceParam.GetNumBytes(), &Context->NumIndicesPerInstance);		// set from the renderer in FNiagaraEmitterInstance::Tick
-	int32 InstancesToSpawnThisFrame = Context->SpawnRateInstances + Context->EventSpawnTotal;
+	int32 InstancesToSpawnThisFrame = Instance->SpawnRateInstances + Instance->EventSpawnTotal;
 	RHICmdList.SetShaderParameter(Shader->GetComputeShader(), Shader->NumSpawnedInstancesParam.GetBufferIndex(), Shader->NumSpawnedInstancesParam.GetBaseIndex(), Shader->NumSpawnedInstancesParam.GetNumBytes(), &InstancesToSpawnThisFrame);				// number of instances in the spawn run
 
 	uint32 NumThreadGroups = 1;
@@ -686,10 +1067,14 @@ void NiagaraEmitterInstanceBatcher::Run(const FNiagaraComputeExecutionContext *C
 	if (CBufferLayout.ConstantBufferSize)
 	{
 		check(CBufferLayout.Resources.Num() == 0);
-		const uint8* ParamData = Params.GetData();
+		const uint8* ParamData = Instance->ParamData;
 		FUniformBufferRHIRef CBuffer = RHICreateUniformBuffer(ParamData, CBufferLayout, EUniformBufferUsage::UniformBuffer_SingleDraw);
 		RHICmdList.SetShaderUniformBuffer(Shader->GetComputeShader(), Shader->EmitterConstantBufferParam.GetBaseIndex(), CBuffer);
 	}
+	else
+	{
+		ensure(!Shader->EmitterConstantBufferParam.IsBound());
+	}
 
 	// Dispatch, if anything needs to be done
 	//
@@ -698,7 +1083,6 @@ void NiagaraEmitterInstanceBatcher::Run(const FNiagaraComputeExecutionContext *C
 		SCOPED_DRAW_EVENTF(RHICmdList, NiagaraGPUSimulationCS, TEXT("Niagara Gpu Sim - %s - NumInstances: %u"),
 			*Context->DebugSimName,
 			TotalNumInstances);
-		SCOPED_GPU_STAT(RHICmdList, NiagaraGPUSimulation);
 		DispatchComputeShader(RHICmdList, Shader, NumThreadGroups, 1, 1);
 	}
 
@@ -708,26 +1092,25 @@ void NiagaraEmitterInstanceBatcher::Run(const FNiagaraComputeExecutionContext *C
 	{
 		//UE_LOG(LogNiagara, Warning, TEXT("Queued up!"));
 
-		if (!Context->GPUDebugDataReadbackFloat && !Context->GPUDebugDataReadbackInt && !Context->GPUDebugDataReadbackCounts && Context->MainDataSet != nullptr)
+		if (!Context->GPUDebugDataReadbackFloat && !Context->GPUDebugDataReadbackInt && !Context->GPUDebugDataReadbackCounts)
 		{
-			FRWBuffer &DatasetIndexBufferWrite = Context->MainDataSet->GetCurDataSetIndices();
+			FRWBuffer &DatasetIndexBufferWrite = DestinationData.GetGPUIndices();
 
-			Context->GPUDebugDataCurrBufferIdx = Context->MainDataSet->GetCurrBufferIdx();
 			Context->GPUDebugDataFloatSize = 0;
 			Context->GPUDebugDataIntSize = 0;
 
-			if (Context->MainDataSet->GetNumFloatComponents() > 0)
+			if (DestinationData.GetGPUBufferFloat().NumBytes > 0)
 			{
 				Context->GPUDebugDataReadbackFloat = new FRHIGPUMemoryReadback(TEXT("Niagara GPU Debug Info Float Emitter Readback"));
-				Context->GPUDebugDataReadbackFloat->EnqueueCopy(RHICmdList, Context->MainDataSet->GetDataByIndex(WriteBufferIdx).GetGPUBufferFloat()->Buffer);
-				Context->GPUDebugDataFloatSize = Context->MainDataSet->GetDataByIndex(WriteBufferIdx).GetGPUBufferFloat()->NumBytes;
+				Context->GPUDebugDataReadbackFloat->EnqueueCopy(RHICmdList, DestinationData.GetGPUBufferFloat().Buffer);
+				Context->GPUDebugDataFloatSize = DestinationData.GetGPUBufferFloat().NumBytes;
 			}
 
-			if (Context->MainDataSet->GetNumInt32Components() > 0)
+			if (DestinationData.GetGPUBufferInt().NumBytes > 0)
 			{
 				Context->GPUDebugDataReadbackInt = new FRHIGPUMemoryReadback(TEXT("Niagara GPU Debug Info Int Emitter Readback"));
-				Context->GPUDebugDataReadbackInt->EnqueueCopy(RHICmdList, Context->MainDataSet->GetDataByIndex(WriteBufferIdx).GetGPUBufferInt()->Buffer);
-				Context->GPUDebugDataIntSize = Context->MainDataSet->GetDataByIndex(WriteBufferIdx).GetGPUBufferInt()->NumBytes;
+				Context->GPUDebugDataReadbackInt->EnqueueCopy(RHICmdList, DestinationData.GetGPUBufferInt().Buffer);
+				Context->GPUDebugDataIntSize = DestinationData.GetGPUBufferInt().NumBytes;
 			}
 
 			Context->GPUDebugDataReadbackCounts = new FRHIGPUMemoryReadback(TEXT("Niagara GPU Emitter Readback"));
@@ -738,154 +1121,8 @@ void NiagaraEmitterInstanceBatcher::Run(const FNiagaraComputeExecutionContext *C
 
 	// Unset UAV parameters and transition resources (TODO: resource transition should be moved to the renderer)
 	// 
-	UnsetDataInterfaceParameters(ParameterStore.GetDataInterfaces(), Shader, RHICmdList, Context);
-	DataSet->UnsetShaderParams(Shader, RHICmdList);
+	UnsetDataInterfaceParameters(DataInterfaceProxies, Shader, RHICmdList, Instance, Tick);
+	CurrentData.UnsetShaderParams(Shader, RHICmdList);
+	DestinationData.UnsetShaderParams(Shader, RHICmdList);
 	Shader->OutputIndexBufferParam.UnsetUAV(RHICmdList, Shader->GetComputeShader());
 }
-
-/* Kick off event handler runs
- * TODO: compile event handlers into the sim/spawn compute shader, so we can do everything within one dispatch
- */
-void NiagaraEmitterInstanceBatcher::RunEventHandlers(const FNiagaraComputeExecutionContext *Context, uint32 NumInstancesAfterSim, uint32 NumInstancesAfterSpawn, uint32 NumInstancesAfterNonEventSpawn, FRHICommandList &RHICmdList) const
-{
-	// Event handler run
-	//
-	/*
-	for (int32 EventScriptIdx = 0; EventScriptIdx < Context->GetEventHandlers().Num(); EventScriptIdx++)
-	{
-		const FNiagaraEventScriptProperties &EventHandlerProps = Context->GetEventHandlers()[EventScriptIdx];
-		FNiagaraDataSet *EventSet = Context->EventSets[EventScriptIdx];
-		if (EventSet)
-		{
-			int32 NumEvents = EventSet->PrevDataRender().GetNumInstances();
-
-			// handle all-particle events
-			if (NumEvents && EventHandlerProps.Script && EventHandlerProps.ExecutionMode == EScriptExecutionMode::EveryParticle && EventSet)
-			{
-				FNiagaraShader* EventHandlerShader = EventHandlerProps.Script->GetRenderThreadScript()->GetShader();
-				if (EventHandlerShader)
-				{
-					SetupDataInterfaceBuffers(EventHandlerProps.Script->DataInterfaceInfo, EventHandlerShader, RHICmdList);
-
-					RHICmdList.SetShaderParameter(EventHandlerShader->GetComputeShader(), 0, EventHandlerShader->NumParticlesPerEventParam.GetBaseIndex(), sizeof(int32), &NumInstancesAfterSpawn);
-					RHICmdList.SetShaderParameter(EventHandlerShader->GetComputeShader(), 0, EventHandlerShader->NumEventsPerParticleParam.GetBaseIndex(), sizeof(int32), &NumEvents);
-
-					// If we have events, Swap buffers, to make sure we don't overwrite previous script results
-					Context->MainDataSet->TickRenderThread(ENiagaraSimTarget::GPUComputeSim);
-					Context->MainDataSet->CurrDataRender().AllocateGPU(NumInstancesAfterSpawn, RHICmdList);
-					Context->MainDataSet->CurrDataRender().SetNumInstances(NumInstancesAfterSpawn);
-
-					SetPrevDataStrideParams(EventSet, EventHandlerShader, RHICmdList);
-
-					// set up event data set buffers
-					FShaderResourceParameter &FloatParam = EventHandlerShader->EventFloatSRVParams[0];
-					FShaderResourceParameter &IntParam = EventHandlerShader->EventIntSRVParams[0];
-					RHICmdList.SetShaderResourceViewParameter(EventHandlerShader->GetComputeShader(), FloatParam.GetBaseIndex(), EventSet->PrevDataRender().GetGPUBufferFloat()->SRV);
-					RHICmdList.SetShaderResourceViewParameter(EventHandlerShader->GetComputeShader(), IntParam.GetBaseIndex(), EventSet->PrevDataRender().GetGPUBufferInt()->SRV);
-
-					TArray<uint8, TAlignedHeapAllocator<16>>  BlankParams;
-					Run(Context->MainDataSet, 0, NumInstancesAfterNonEventSpawn, EventHandlerShader, BlankParams, RHICmdList, DummyWriteIndexBuffer);
-				}
-			}
-
-			// handle spawn events
-			if (EventHandlerProps.Script && EventHandlerProps.ExecutionMode == EScriptExecutionMode::SpawnedParticles && EventSet)
-			{
-				uint32 NumEventsToProcess = EventSet->PrevDataRender().GetNumInstances();
-				uint32 EventSpawnNum = NumEventsToProcess * EventHandlerProps.SpawnNumber;
-
-				//int32 SpawnNum = Context->EventSpawnCounts[EventScriptIdx];
-				FNiagaraShader* EventHandlerShader = EventHandlerProps.Script->GetRenderThreadScript()->GetShader();
-				if (NumEvents && EventSpawnNum && EventHandlerShader)
-				{
-					//SetupDataInterfaceBuffers(EventHandlerProps.Script->DataInterfaceInfo, EventHandlerShader, RHICmdList);
-					int32 OneEvent = 1;
-					int32 ParticlesPerEvent = EventSpawnNum / NumEvents;
-					RHICmdList.SetShaderParameter(EventHandlerShader->GetComputeShader(), 0, EventHandlerShader->NumEventsPerParticleParam.GetBaseIndex(), sizeof(int32), &OneEvent);
-					RHICmdList.SetShaderParameter(EventHandlerShader->GetComputeShader(), 0, EventHandlerShader->NumParticlesPerEventParam.GetBaseIndex(), sizeof(int32), &ParticlesPerEvent);
-
-					// If we have events, Swap buffers
-					Context->MainDataSet->TickRenderThread(ENiagaraSimTarget::GPUComputeSim);
-					Context->MainDataSet->CurrDataRender().AllocateGPU(NumInstancesAfterSpawn, RHICmdList);
-					Context->MainDataSet->CurrDataRender().SetNumInstances(NumInstancesAfterSpawn);
-
-					// set up event data set buffers
-					FShaderResourceParameter &FloatParam = EventHandlerShader->EventFloatSRVParams[0];
-					FShaderResourceParameter &IntParam = EventHandlerShader->EventIntSRVParams[0];
-					RHICmdList.SetShaderResourceViewParameter(EventHandlerShader->GetComputeShader(), FloatParam.GetBaseIndex(), EventSet->PrevDataRender().GetGPUBufferFloat()->SRV);
-					RHICmdList.SetShaderResourceViewParameter(EventHandlerShader->GetComputeShader(), IntParam.GetBaseIndex(), EventSet->PrevDataRender().GetGPUBufferInt()->SRV);
-
-					SetPrevDataStrideParams(EventSet, EventHandlerShader, RHICmdList);
-
-					// we assume event spawns are at the end of the buffer
-					check(NumInstancesAfterNonEventSpawn + EventSpawnNum == Context->MainDataSet->CurrDataRender().GetNumInstances());
-
-					TArray<uint8, TAlignedHeapAllocator<16>>  BlankParams;
-					Run(Context->MainDataSet, NumInstancesAfterNonEventSpawn, EventSpawnNum, EventHandlerShader, BlankParams, RHICmdList, DummyWriteIndexBuffer, true);
-				}
-			}
-		}
-	}
-	*/
-}
-
-/* Set read strides for the sim shader
- */
-void NiagaraEmitterInstanceBatcher::SetPrevDataStrideParams(const FNiagaraDataSet *Set, FNiagaraShader* Shader, FRHICommandList &RHICmdList) const
-{
-	int32 FloatStride = Set->PrevData().GetFloatStride() / sizeof(float);
-	int32 IntStride = Set->PrevData().GetInt32Stride() / sizeof(int32);
-	RHICmdList.SetShaderParameter(Shader->GetComputeShader(), Shader->EventReadFloatStrideParams[0].GetBufferIndex(), Shader->EventReadFloatStrideParams[0].GetBaseIndex(), sizeof(int32), &FloatStride);
-	RHICmdList.SetShaderParameter(Shader->GetComputeShader(), Shader->EventReadIntStrideParams[0].GetBufferIndex(), Shader->EventReadIntStrideParams[0].GetBaseIndex(), sizeof(int32), &IntStride);
-}
-
-
-/* Set up UAVs for the event data sets
-*/
-void NiagaraEmitterInstanceBatcher::SetupEventUAVs(const FNiagaraComputeExecutionContext *Context, uint32 NewNumInstances, FRHICommandList &RHICmdList) const
-{
-	FNiagaraShader* UpdateShader = Context->RTUpdateScript->GetShader();
-
-	uint32 SetIndex = 0;
-	for (FNiagaraDataSet *Set : Context->UpdateEventWriteDataSets)
-	{
-		if (NewNumInstances)
-		{
-			Set->CurrData().AllocateGPU(NewNumInstances, RHICmdList);
-			Set->CurrData().SetNumInstances(NewNumInstances);
-			FRWShaderParameter &FloatParam = UpdateShader->EventFloatUAVParams[SetIndex];
-			FRWShaderParameter &IntParam = UpdateShader->EventIntUAVParams[SetIndex];
-			if (FloatParam.IsUAVBound())
-			{
-				RHICmdList.SetUAVParameter(UpdateShader->GetComputeShader(), FloatParam.GetUAVIndex(), Set->CurrData().GetGPUBufferFloat()->UAV);
-			}
-			if (IntParam.IsUAVBound())
-			{
-				RHICmdList.SetUAVParameter(UpdateShader->GetComputeShader(), IntParam.GetUAVIndex(), Set->CurrData().GetGPUBufferInt()->UAV);
-			}
-
-			uint32 FloatStride = Set->CurrData().GetFloatStride() / sizeof(float);
-			uint32 IntStride = Set->CurrData().GetInt32Stride() / sizeof(int32);
-			RHICmdList.SetShaderParameter(UpdateShader->GetComputeShader(), UpdateShader->EventWriteFloatStrideParams[SetIndex].GetBufferIndex(), UpdateShader->EventWriteFloatStrideParams[SetIndex].GetBaseIndex(), sizeof(int32), &FloatStride);
-			RHICmdList.SetShaderParameter(UpdateShader->GetComputeShader(), UpdateShader->EventWriteIntStrideParams[SetIndex].GetBufferIndex(), UpdateShader->EventWriteIntStrideParams[SetIndex].GetBaseIndex(), sizeof(int32), &IntStride);
-		}
-
-		SetIndex++;
-	}
-
-
-}
-
-
-void NiagaraEmitterInstanceBatcher::UnsetEventUAVs(const FNiagaraComputeExecutionContext *Context, FRHICommandList &RHICmdList) const
-{
-	FNiagaraShader* UpdateShader = Context->RTUpdateScript->GetShader();
-
-	for (int32 SetIndex = 0; SetIndex < Context->UpdateEventWriteDataSets.Num(); SetIndex++)
-	{
-		FRWShaderParameter &FloatParam = UpdateShader->EventFloatUAVParams[SetIndex];
-		FRWShaderParameter &IntParam = UpdateShader->EventIntUAVParams[SetIndex];
-		FloatParam.UnsetUAV(RHICmdList, UpdateShader->GetComputeShader());
-		IntParam.UnsetUAV(RHICmdList, UpdateShader->GetComputeShader());
-	}
-}
\ No newline at end of file
diff --git a/Engine/Plugins/FX/Niagara/Source/Niagara/Private/NiagaraFunctionLibrary.cpp b/Engine/Plugins/FX/Niagara/Source/Niagara/Private/NiagaraFunctionLibrary.cpp
--- a/Engine/Plugins/FX/Niagara/Source/Niagara/Private/NiagaraFunctionLibrary.cpp
+++ b/Engine/Plugins/FX/Niagara/Source/Niagara/Private/NiagaraFunctionLibrary.cpp
@@ -9,6 +9,7 @@
 #include "ContentStreaming.h"
 
 #include "NiagaraWorldManager.h"
+#include "NiagaraDataInterfaceStaticMesh.h"
 
 #if WITH_EDITOR
 int32 GForceNiagaraSpawnAttachedSolo = 0;
@@ -86,14 +87,12 @@ UNiagaraComponent* UNiagaraFunctionLibrary::SpawnSystemAttached(UNiagaraSystem*
 		else
 		{
 			PSC = CreateNiagaraSystem(SystemTemplate, AttachToComponent->GetWorld(), AttachToComponent->GetOwner(), bAutoDestroy, EPSCPoolMethod::None);
-
 #if WITH_EDITOR
 			if (GForceNiagaraSpawnAttachedSolo > 0)
 			{
 				PSC->SetForceSolo(true);
 			}
 #endif
-
 			PSC->RegisterComponentWithWorld(AttachToComponent->GetWorld());
 
 			PSC->AttachToComponent(AttachToComponent, FAttachmentTransformRules::KeepRelativeTransform, AttachPointName);
@@ -211,6 +210,112 @@ void UNiagaraFunctionLibrary::SetUpdateScriptConstant(UNiagaraComponent* Compone
 }
 */
 
+void UNiagaraFunctionLibrary::OverrideSystemUserVariableStaticMeshComponent(UNiagaraComponent* NiagaraSystem, const FString& OverrideName, UStaticMeshComponent* StaticMeshComponent)
+{
+	if (!NiagaraSystem)
+	{
+		UE_LOG(LogNiagara, Warning, TEXT("NiagaraSystem in \"Set Niagara Static Mesh Component\" is NULL, OverrideName \"%s\" and StaticMeshComponent \"%s\", skipping."), *OverrideName, StaticMeshComponent ? *StaticMeshComponent->GetName() : TEXT("NULL"));
+		return;
+	}
+
+	if (!StaticMeshComponent)
+	{
+		UE_LOG(LogNiagara, Warning, TEXT("StaticMeshComponent in \"Set Niagara Static Mesh Component\" is NULL, OverrideName \"%s\" and NiagaraSystem \"%s\", skipping."), *OverrideName, *NiagaraSystem->GetOwner()->GetName());
+		return;
+	}
+
+	const FNiagaraParameterStore& OverrideParameters = NiagaraSystem->GetOverrideParameters();
+
+	FNiagaraVariable Variable(FNiagaraTypeDefinition(UNiagaraDataInterfaceStaticMesh::StaticClass()), *OverrideName);
+	
+	int32 Index = OverrideParameters.IndexOf(Variable);
+	if (Index == INDEX_NONE)
+	{
+		UE_LOG(LogNiagara, Warning, TEXT("Could not find index of variable \"%s\" in the OverrideParameters map of NiagaraSystem \"%s\"."), *OverrideName, *NiagaraSystem->GetOwner()->GetName());
+		return;
+	}
+	
+	UNiagaraDataInterfaceStaticMesh* StaticMeshInterface = Cast<UNiagaraDataInterfaceStaticMesh>(OverrideParameters.GetDataInterface(Index));
+	if (!StaticMeshInterface)
+	{
+		UE_LOG(LogNiagara, Warning, TEXT("Did not find a matching Static Mesh Data Interface variable named \"%s\" in the User variables of NiagaraSystem \"%s\" ."), *OverrideName, *NiagaraSystem->GetOwner()->GetName());
+		return;
+	}
+
+	StaticMeshInterface->SetSourceComponentFromBlueprints(StaticMeshComponent);
+}
+
+void UNiagaraFunctionLibrary::OverrideSystemUserVariableStaticMesh(UNiagaraComponent* NiagaraSystem, const FString& OverrideName, UStaticMesh* StaticMesh)
+{
+	if (!NiagaraSystem)
+	{
+		UE_LOG(LogNiagara, Warning, TEXT("NiagaraSystem in \"Set Niagara Static Mesh Component\" is NULL, OverrideName \"%s\" and StaticMesh \"%s\", skipping."), *OverrideName, StaticMesh ? *StaticMesh->GetName() : TEXT("NULL"));
+		return;
+	}
+
+	if (!StaticMesh)
+	{
+		UE_LOG(LogNiagara, Warning, TEXT("StaticMesh in \"Set Niagara Static Mesh Component\" is NULL, OverrideName \"%s\" and NiagaraSystem \"%s\", skipping."), *OverrideName, *NiagaraSystem->GetOwner()->GetName());
+		return;
+	}
+
+	const FNiagaraParameterStore& OverrideParameters = NiagaraSystem->GetOverrideParameters();
+
+	FNiagaraVariable Variable(FNiagaraTypeDefinition(UNiagaraDataInterfaceStaticMesh::StaticClass()), *OverrideName);
+
+	int32 Index = OverrideParameters.IndexOf(Variable);
+	if (Index == INDEX_NONE)
+	{
+		UE_LOG(LogNiagara, Warning, TEXT("Could not find index of variable \"%s\" in the OverrideParameters map of NiagaraSystem \"%s\"."), *OverrideName, *NiagaraSystem->GetOwner()->GetName());
+		return;
+	}
+
+	UNiagaraDataInterfaceStaticMesh* StaticMeshInterface = Cast<UNiagaraDataInterfaceStaticMesh>(OverrideParameters.GetDataInterface(Index));
+	if (!StaticMeshInterface)
+	{
+		UE_LOG(LogNiagara, Warning, TEXT("Did not find a matching Static Mesh Data Interface variable named \"%s\" in the User variables of NiagaraSystem \"%s\" ."), *OverrideName, *NiagaraSystem->GetOwner()->GetName());
+		return;
+	}
+
+	StaticMeshInterface->SetDefaultMeshFromBlueprints(StaticMesh);
+}
+
+
+void UNiagaraFunctionLibrary::OverrideSystemUserVariableSkeletalMeshComponent(UNiagaraComponent* NiagaraSystem, const FString& OverrideName, USkeletalMeshComponent* SkeletalMeshComponent)
+{
+	if (!NiagaraSystem)
+	{
+		UE_LOG(LogNiagara, Warning, TEXT("NiagaraSystem in \"Set Niagara Skeletal Mesh Component\" is NULL, OverrideName \"%s\" and SkeletalMeshComponent \"%s\", skipping."), *OverrideName, SkeletalMeshComponent ? *SkeletalMeshComponent->GetName() : TEXT("NULL"));
+		return;
+	}
+
+	if (!SkeletalMeshComponent)
+	{
+		UE_LOG(LogNiagara, Warning, TEXT("SkeletalMeshComponent in \"Set Niagara Skeletal Mesh Component\" is NULL, OverrideName \"%s\" and NiagaraSystem \"%s\", skipping."), *OverrideName, *NiagaraSystem->GetOwner()->GetName());
+		return;
+	}
+
+	const FNiagaraParameterStore& OverrideParameters = NiagaraSystem->GetOverrideParameters();
+
+	FNiagaraVariable Variable(FNiagaraTypeDefinition(UNiagaraDataInterfaceSkeletalMesh::StaticClass()), *OverrideName);
+	
+	int32 Index = OverrideParameters.IndexOf(Variable);
+	if (Index == INDEX_NONE)
+	{
+		UE_LOG(LogNiagara, Warning, TEXT("Could not find index of variable \"%s\" in the OverrideParameters map of NiagaraSystem \"%s\"."), *OverrideName, *NiagaraSystem->GetOwner()->GetName());
+		return;
+	}
+	
+	UNiagaraDataInterfaceSkeletalMesh* SkeletalMeshInterface = Cast<UNiagaraDataInterfaceSkeletalMesh>(OverrideParameters.GetDataInterface(Index));
+	if (!SkeletalMeshInterface)
+	{
+		UE_LOG(LogNiagara, Warning, TEXT("Did not find a matching Skeletal Mesh Data Interface variable named \"%s\" in the User variables of NiagaraSystem \"%s\" ."), *OverrideName, *NiagaraSystem->GetOwner()->GetName());
+		return;
+	}
+
+	SkeletalMeshInterface->SetSourceComponentFromBlueprints(SkeletalMeshComponent);
+}
+
 UNiagaraParameterCollectionInstance* UNiagaraFunctionLibrary::GetNiagaraParameterCollection(UObject* WorldContextObject, UNiagaraParameterCollection* Collection)
 {
 	UWorld* World = GEngine->GetWorldFromContextObject(WorldContextObject, EGetWorldErrorMode::LogAndReturnNull);
diff --git a/Engine/Plugins/FX/Niagara/Source/Niagara/Private/NiagaraLightRendererProperties.cpp b/Engine/Plugins/FX/Niagara/Source/Niagara/Private/NiagaraLightRendererProperties.cpp
--- a/Engine/Plugins/FX/Niagara/Source/Niagara/Private/NiagaraLightRendererProperties.cpp
+++ b/Engine/Plugins/FX/Niagara/Source/Niagara/Private/NiagaraLightRendererProperties.cpp
@@ -3,6 +3,8 @@
 #include "NiagaraLightRendererProperties.h"
 #include "NiagaraRenderer.h"
 #include "NiagaraConstants.h"
+#include "NiagaraRendererLights.h"
+
 UNiagaraLightRendererProperties::UNiagaraLightRendererProperties()
 	: bUseInverseSquaredFalloff(1), bAffectsTranslucency(0), bOverrideRenderingEnabled(0), RadiusScale(1.0f), ColorAdd(FVector(0.0f, 0.0f, 0.0f))
 {
@@ -35,9 +37,11 @@ void UNiagaraLightRendererProperties::InitCDOPropertiesAfterModuleStartup()
 	CDO->VolumetricScatteringBinding = FNiagaraConstants::GetAttributeDefaultBinding(SYS_PARAM_PARTICLES_LIGHT_VOLUMETRIC_SCATTERING);
 }
 
-NiagaraRenderer* UNiagaraLightRendererProperties::CreateEmitterRenderer(ERHIFeatureLevel::Type FeatureLevel)
+FNiagaraRenderer* UNiagaraLightRendererProperties::CreateEmitterRenderer(ERHIFeatureLevel::Type FeatureLevel, const FNiagaraEmitterInstance* Emitter)
 {
-	return new NiagaraRendererLights(FeatureLevel, this);
+	FNiagaraRenderer* NewRenderer = new FNiagaraRendererLights(FeatureLevel, this, Emitter);
+	NewRenderer->Initialize(FeatureLevel, this, Emitter);
+	return NewRenderer;
 }
 
 void UNiagaraLightRendererProperties::GetUsedMaterials(TArray<UMaterialInterface*>& OutMaterials) const
diff --git a/Engine/Plugins/FX/Niagara/Source/Niagara/Private/NiagaraMeshRendererProperties.cpp b/Engine/Plugins/FX/Niagara/Source/Niagara/Private/NiagaraMeshRendererProperties.cpp
--- a/Engine/Plugins/FX/Niagara/Source/Niagara/Private/NiagaraMeshRendererProperties.cpp
+++ b/Engine/Plugins/FX/Niagara/Source/Niagara/Private/NiagaraMeshRendererProperties.cpp
@@ -4,6 +4,7 @@
 #include "NiagaraRendererMeshes.h"
 #include "Engine/StaticMesh.h"
 #include "NiagaraConstants.h"
+#include "NiagaraBoundsCalculatorHelper.h"
 
 UNiagaraMeshRendererProperties::UNiagaraMeshRendererProperties()
 	: ParticleMesh(nullptr)
@@ -12,17 +13,34 @@ UNiagaraMeshRendererProperties::UNiagaraMeshRendererProperties()
 {
 }
 
-NiagaraRenderer* UNiagaraMeshRendererProperties::CreateEmitterRenderer(ERHIFeatureLevel::Type FeatureLevel)
+FNiagaraRenderer* UNiagaraMeshRendererProperties::CreateEmitterRenderer(ERHIFeatureLevel::Type FeatureLevel, const FNiagaraEmitterInstance* Emitter)
 {
-	return new NiagaraRendererMeshes(FeatureLevel, this);
+	if (ParticleMesh)
+	{
+		FNiagaraRenderer* NewRenderer = new FNiagaraRendererMeshes(FeatureLevel, this, Emitter);
+		NewRenderer->Initialize(FeatureLevel, this, Emitter);
+		return NewRenderer;
+	}
+
+	return nullptr;
 }
 
+FNiagaraBoundsCalculator* UNiagaraMeshRendererProperties::CreateBoundsCalculator()
+{
+	if (ParticleMesh)
+	{
+		FNiagaraBoundsCalculatorHelper<false, true, false>* BoundsCalculator = new FNiagaraBoundsCalculatorHelper<false, true, false>();
+		BoundsCalculator->MeshExtents = ParticleMesh->GetBounds().BoxExtent;
+		return BoundsCalculator;
+	}
+
+	return nullptr;
 
+}
 
 void UNiagaraMeshRendererProperties::PostInitProperties()
 {
 	Super::PostInitProperties();
-	SyncId = 0;
 	if (HasAnyFlags(RF_ClassDefaultObject) == false)
 	{
 		InitBindings();
@@ -93,7 +111,8 @@ void UNiagaraMeshRendererProperties::GetUsedMaterials(TArray<UMaterialInterface*
 
 uint32 UNiagaraMeshRendererProperties::GetNumIndicesPerInstance() 
 {
-	return ParticleMesh ? ParticleMesh->RenderData->LODResources[0].IndexBuffer.GetNumIndices() : 0;
+	// TODO: Add proper support for multiple mesh sections for GPU mesh particles.
+	return ParticleMesh ? ParticleMesh->RenderData->LODResources[0].Sections[0].NumTriangles * 3 : 0;
 }
 
 
@@ -164,11 +183,6 @@ void UNiagaraMeshRendererProperties::PostEditChangeProperty(FPropertyChangedEven
 			}
 		}
 	}
-	
-	if (PropertyChangedEvent.GetPropertyName() != TEXT("SyncId"))
-	{
-		SyncId++;
-	}
 }
 
 #endif // WITH_EDITORONLY_DATA
\ No newline at end of file
diff --git a/Engine/Plugins/FX/Niagara/Source/Niagara/Private/NiagaraModule.cpp b/Engine/Plugins/FX/Niagara/Source/Niagara/Private/NiagaraModule.cpp
--- a/Engine/Plugins/FX/Niagara/Source/Niagara/Private/NiagaraModule.cpp
+++ b/Engine/Plugins/FX/Niagara/Source/Niagara/Private/NiagaraModule.cpp
@@ -31,6 +31,14 @@ float INiagaraModule::EngineGlobalSpawnCountScale = 1.0f;
 float INiagaraModule::EngineGlobalSystemCountScale = 1.0f;
 int32 INiagaraModule::EngineDetailLevel = 4;
 
+int32 GEnableVerboseNiagaraChangeIdLogging = 0;
+static FAutoConsoleVariableRef CVarEnableVerboseNiagaraChangeIdLogging(
+	TEXT("fx.EnableVerboseNiagaraChangeIdLogging"),
+	GEnableVerboseNiagaraChangeIdLogging,
+	TEXT("If > 0 Verbose change id logging info will be printed. \n"),
+	ECVF_Default
+);
+
 /**
 Detail Level CVar.
 Effectively replaces the DetaiMode feature but allows for a rolling range of new hardware and emitters to target them.
@@ -293,7 +301,15 @@ void INiagaraModule::ShutdownModule()
 
 FNiagaraWorldManager* INiagaraModule::GetWorldManager(UWorld* World)
 {
-	return WorldManagers.FindChecked(World);
+	FNiagaraWorldManager** OutWorld = WorldManagers.Find(World);
+	if (OutWorld == nullptr)
+	{
+		UE_LOG(LogNiagara, Warning, TEXT("Calling INiagaraModule::GetWorldManager \"%s\", but Niagara has never encountered this world before. "
+			" This means that WorldInit never happened. This may happen in some edge cases in the editor, like saving invisible child levels, "
+			"in which case the calling context needs to be safe against this returning nullptr."), World ? *World->GetName() : TEXT("nullptr"));
+		return nullptr;
+	}
+	return *OutWorld;
 }
 
 void INiagaraModule::DestroyAllSystemSimulations(class UNiagaraSystem* System)
diff --git a/Engine/Plugins/FX/Niagara/Source/Niagara/Private/NiagaraParameterStore.cpp b/Engine/Plugins/FX/Niagara/Source/Niagara/Private/NiagaraParameterStore.cpp
--- a/Engine/Plugins/FX/Niagara/Source/Niagara/Private/NiagaraParameterStore.cpp
+++ b/Engine/Plugins/FX/Niagara/Source/Niagara/Private/NiagaraParameterStore.cpp
@@ -291,7 +291,12 @@ bool FNiagaraParameterStore::AddParameter(const FNiagaraVariable& Param, bool bI
 		//Temporary to init param data from FNiagaraVariable storage. This will be removed when we change the UNiagaraScript to use a parameter store too.
 		if (Param.IsDataAllocated())
 		{
-			FMemory::Memcpy(GetParameterData_Internal(Offset), Param.GetData(), ParamSize);
+			uint8* Dest = GetParameterData_Internal(Offset);
+			const uint8* Src = Param.GetData();
+			if (Dest != Src)
+			{
+				FMemory::Memcpy(Dest, Src, ParamSize);
+			}
 		}
 	}
 
@@ -366,7 +371,14 @@ void FNiagaraParameterStore::RenameParameter(const FNiagaraVariable& Param, FNam
 
 		bool bInitInterfaces = false;
 		bool bTriggerRebind = false;
+
+		int32 NumBytesBefore = ParameterData.Num();
 		AddParameter(NewParam, bInitInterfaces, bTriggerRebind);
+		int32 NumBytesAfter = ParameterData.Num();
+		int32 DeltaBytes = NumBytesAfter - NumBytesBefore;
+		int32 SizeInBytes = Param.GetSizeInBytes();
+
+		check(DeltaBytes == SizeInBytes);
 
 		int32 NewIdx = IndexOf(NewParam);
 		if (Param.IsDataInterface())
diff --git a/Engine/Plugins/FX/Niagara/Source/Niagara/Private/NiagaraRenderer.cpp b/Engine/Plugins/FX/Niagara/Source/Niagara/Private/NiagaraRenderer.cpp
--- a/Engine/Plugins/FX/Niagara/Source/Niagara/Private/NiagaraRenderer.cpp
+++ b/Engine/Plugins/FX/Niagara/Source/Niagara/Private/NiagaraRenderer.cpp
@@ -9,87 +9,238 @@
 #include "Engine/Engine.h"
 #include "DynamicBufferAllocator.h"
 
-DECLARE_CYCLE_STAT(TEXT("Generate Particle Lights"), STAT_NiagaraGenLights, STATGROUP_Niagara);
 DECLARE_CYCLE_STAT(TEXT("Sort Particles"), STAT_NiagaraSortParticles, STATGROUP_Niagara);
 
+DECLARE_CYCLE_STAT(TEXT("Global Float Alloc - All"), STAT_NiagaraAllocateGlobalFloatAll, STATGROUP_Niagara);
+DECLARE_CYCLE_STAT(TEXT("Global Float Alloc - InsideLock"), STAT_NiagaraAllocateGlobalFloatInsideLock, STATGROUP_Niagara);
+DECLARE_CYCLE_STAT(TEXT("Global Float Alloc - Alloc New Buffer"), STAT_NiagaraAllocateGlobalFloatAllocNew, STATGROUP_Niagara);
+DECLARE_CYCLE_STAT(TEXT("Global Float Alloc - Map Buffer"), STAT_NiagaraAllocateGlobalFloatMapBuffer, STATGROUP_Niagara);
 
-/** Enable/disable parallelized System renderers */
-int32 GbNiagaraParallelEmitterRenderers = 1;
-static FAutoConsoleVariableRef CVarParallelEmitterRenderers(
-	TEXT("niagara.ParallelEmitterRenderers"),
-	GbNiagaraParallelEmitterRenderers,
-	TEXT("Whether to run Niagara System renderers in parallel"),
-	ECVF_Default
-	);
+class FNiagaraDummyRWBufferFloat : public FRenderResource
+{
+public:
+	FNiagaraDummyRWBufferFloat(const FString InDebugId) : DebugId(InDebugId) {}
+	FString DebugId;
+	FRWBuffer Buffer;
 
+	virtual void InitRHI() override
+	{
+		UE_LOG(LogNiagara, Log, TEXT("FNiagaraDummyRWBufferFloat InitRHI %s"), *DebugId);
+		Buffer.Initialize(sizeof(float), 1, EPixelFormat::PF_R32_FLOAT, BUF_Static, *DebugId);
+	}
 
+	virtual void ReleaseRHI() override
+	{
+		UE_LOG(LogNiagara, Log, TEXT("FNiagaraDummyRWBufferFloat ReleaseRHI %s"), *DebugId);
+		Buffer.Release();
+	}
+};
 
-void FNiagaraDummyRWBufferFloat::InitRHI() 
+class FNiagaraDummyRWBufferFloat4 : public FRenderResource
 {
-	UE_LOG(LogNiagara, Log, TEXT("FNiagaraDummyRWBufferFloat InitRHI %s"), *DebugId);
-	Buffer.Initialize(sizeof(float), 1, EPixelFormat::PF_R32_FLOAT, BUF_Static);
-}
+public:
+	FNiagaraDummyRWBufferFloat4(const FString InDebugId) : DebugId(InDebugId) {}
+	FString DebugId;
+	FRWBuffer Buffer;
 
-void FNiagaraDummyRWBufferFloat::ReleaseRHI() 
-{
-	UE_LOG(LogNiagara, Log, TEXT("FNiagaraDummyRWBufferFloat ReleaseRHI %s"), *DebugId);
-	Buffer.Release();
-}
+	virtual void InitRHI() override
+	{
+		UE_LOG(LogNiagara, Log, TEXT("FNiagaraDummyRWBufferFloat InitRHI %s"), *DebugId);
+		Buffer.Initialize(sizeof(float) * 4, 1, EPixelFormat::PF_A32B32G32R32F, BUF_Static, *DebugId);
+	}
 
-void FNiagaraDummyRWBufferInt::InitRHI() 
+	virtual void ReleaseRHI() override
+	{
+		UE_LOG(LogNiagara, Log, TEXT("FNiagaraDummyRWBufferFloat4 ReleaseRHI %s"), *DebugId);
+		Buffer.Release();
+	}
+};
+
+class FNiagaraDummyRWBufferInt : public FRenderResource
 {
-	UE_LOG(LogNiagara, Log, TEXT("FNiagaraDummyRWBufferInt InitRHI %s"), *DebugId);
-	Buffer.Initialize(sizeof(int32), 1, EPixelFormat::PF_R32_SINT, BUF_Static);
-}
-void FNiagaraDummyRWBufferInt::ReleaseRHI() 
+public:
+	FNiagaraDummyRWBufferInt(const FString InDebugId) : DebugId(InDebugId) {}
+	FString DebugId;
+	FRWBuffer Buffer;
+
+	virtual void InitRHI() override
+	{
+		UE_LOG(LogNiagara, Log, TEXT("FNiagaraDummyRWBufferInt InitRHI %s"), *DebugId);
+		Buffer.Initialize(sizeof(int32), 1, EPixelFormat::PF_R32_SINT, BUF_Static, *DebugId);
+	}
+
+	virtual void ReleaseRHI() override
+	{
+		UE_LOG(LogNiagara, Log, TEXT("FNiagaraDummyRWBufferInt ReleaseRHI %s"), *DebugId);
+		Buffer.Release();
+	}
+};
+
+class FNiagaraDummyRWBufferUInt : public FRenderResource
 {
-	UE_LOG(LogNiagara, Log, TEXT("FNiagaraDummyRWBufferInt ReleaseRHI %s"), *DebugId);
-	Buffer.Release();
-}
+public:
+	FNiagaraDummyRWBufferUInt(const FString InDebugId) : DebugId(InDebugId) {}
+	FString DebugId;
+	FRWBuffer Buffer;
 
+	virtual void InitRHI() override
+	{
+		UE_LOG(LogNiagara, Log, TEXT("FNiagaraDummyRWBufferUInt InitRHI %s"), *DebugId);
+		Buffer.Initialize(sizeof(uint32), 1, EPixelFormat::PF_R32_UINT, BUF_Static, *DebugId);
+	}
+
+	virtual void ReleaseRHI() override
+	{
+		UE_LOG(LogNiagara, Log, TEXT("FNiagaraDummyRWBufferUInt ReleaseRHI %s"), *DebugId);
+		Buffer.Release();
+	}
+};
 
-FRWBuffer& NiagaraRenderer::GetDummyFloatBuffer()
+FRWBuffer& FNiagaraRenderer::GetDummyFloatBuffer()
 {
 	check(IsInRenderingThread());
 	static TGlobalResource<FNiagaraDummyRWBufferFloat> DummyFloatBuffer(TEXT("NiagaraRenderer::DummyFloat"));
 	return DummyFloatBuffer.Buffer;
 }
 
-FRWBuffer& NiagaraRenderer::GetDummyIntBuffer()
+FRWBuffer& FNiagaraRenderer::GetDummyFloat4Buffer()
+{
+	check(IsInRenderingThread());
+	static TGlobalResource<FNiagaraDummyRWBufferFloat4> GetDummyFloat4Buffer(TEXT("NiagaraRenderer::DummyFloat4"));
+	return GetDummyFloat4Buffer.Buffer;
+}
+
+FRWBuffer& FNiagaraRenderer::GetDummyIntBuffer()
 {
 	check(IsInRenderingThread());
 	static TGlobalResource<FNiagaraDummyRWBufferInt> DummyIntBuffer(TEXT("NiagaraRenderer::DummyInt"));
 	return DummyIntBuffer.Buffer;
 }
 
-NiagaraRenderer::NiagaraRenderer()
-	: CPUTimeMS(0.0f)
-	, bLocalSpace(false)
-	, bEnabled(true)
-	, DynamicDataRender(nullptr)
-	, BaseExtents(1.0f, 1.0f, 1.0f)
+FRWBuffer& FNiagaraRenderer::GetDummyUIntBuffer()
+{
+	check(IsInRenderingThread());
+	static TGlobalResource<FNiagaraDummyRWBufferUInt> DummyUIntBuffer(TEXT("NiagaraRenderer::DummyUInt"));
+	return DummyUIntBuffer.Buffer;
+}
+
+//////////////////////////////////////////////////////////////////////////
+
+FNiagaraDynamicDataBase::FNiagaraDynamicDataBase(const FNiagaraEmitterInstance* InEmitter)
+{
+	check(InEmitter);
+
+	FNiagaraDataSet& DataSet = InEmitter->GetData();
+	SimTarget = DataSet.GetSimTarget();
+
+	if (SimTarget == ENiagaraSimTarget::CPUSim)
+	{
+		//On CPU we pass through direct ptr to the most recent data buffer.
+		Data.CPUParticleData = &DataSet.GetCurrentDataChecked();
+
+		//Mark this buffer as in use by this renderer. Prevents this buffer being reused to write new simulation data while it's inuse by the renderer.
+		Data.CPUParticleData->AddReadRef();
+	}
+	else
+	{
+		//On GPU we must access the correct buffer via the GPUExecContext. Probably a way to route this data better outside the dynamic data in future.
+		//During simulation, the correct data buffer for rendering will be placed in the GPUContext and AddReadRef called.
+		check(SimTarget == ENiagaraSimTarget::GPUComputeSim);
+		Data.GPUExecContext = InEmitter->GetGPUContext();
+	}
+}
+
+FNiagaraDynamicDataBase::~FNiagaraDynamicDataBase()
 {
-	Material = UMaterial::GetDefaultMaterial(MD_Surface);
+	if (SimTarget == ENiagaraSimTarget::CPUSim)
+	{
+		check(Data.CPUParticleData);
+		//Release our ref on the buffer so it can be reused as a destination for a new simulation tick.
+		Data.CPUParticleData->ReleaseReadRef();
+	}
+}
+
+FNiagaraDataBuffer* FNiagaraDynamicDataBase::GetParticleDataToRender()const
+{
+	FNiagaraDataBuffer* Ret = nullptr;
+
+	if (SimTarget == ENiagaraSimTarget::CPUSim)
+	{
+		Ret = Data.CPUParticleData;
+	}
+	else
+	{
+		Ret = Data.GPUExecContext->GetDataToRender();
+	}
+
+	checkSlow(Ret == nullptr || Ret->IsBeingRead());
+	return Ret;
 }
 
+//////////////////////////////////////////////////////////////////////////
+
 
-NiagaraRenderer::~NiagaraRenderer() 
+FNiagaraRenderer::FNiagaraRenderer(ERHIFeatureLevel::Type FeatureLevel, const UNiagaraRendererProperties *InProps, const FNiagaraEmitterInstance* Emitter)
+	: DynamicDataRender(nullptr)
+	, CPUTimeMS(0.0f)
+	, bLocalSpace(Emitter->GetCachedEmitter()->bLocalSpace)
+	, bHasLights(false)
 {
+#if STATS
+	EmitterStatID = Emitter->GetCachedEmitter()->GetStatID(false, false);
+#endif
 }
 
-void NiagaraRenderer::Release()
+void FNiagaraRenderer::Initialize(ERHIFeatureLevel::Type FeatureLevel, const UNiagaraRendererProperties *InProps, const FNiagaraEmitterInstance* Emitter)
 {
-	check(IsInGameThread());
-	NiagaraRenderer* Renderer = this;
-	ENQUEUE_RENDER_COMMAND(NiagaraRendererDeletion)(
-		[Renderer](FRHICommandListImmediate& RHICmdList)
+	//Get our list of valid base materials. Fall back to default material if they're not valid.
+	InProps->GetUsedMaterials(BaseMaterials_GT);
+	for (UMaterialInterface*& Mat : BaseMaterials_GT)
+	{
+		if (!IsMaterialValid(Mat))
 		{
-			delete Renderer;
+			Mat = UMaterial::GetDefaultMaterial(MD_Surface);
 		}
-	);
+		BaseMaterialRelevance_GT |= Mat->GetRelevance(FeatureLevel);
+	}
+}
+
+FNiagaraRenderer::~FNiagaraRenderer()
+{
+	ReleaseRenderThreadResources();
+	SetDynamicData_RenderThread(nullptr);
+}
+
+FPrimitiveViewRelevance FNiagaraRenderer::GetViewRelevance(const FSceneView* View, const FNiagaraSceneProxy *SceneProxy)const
+{
+	FPrimitiveViewRelevance Result;
+	bool bHasDynamicData = HasDynamicData();
+
+	//Always draw so our LastRenderTime is updated. We may not have dynamic data if we're disabled from visibility culling.
+	Result.bDrawRelevance =/* bHasDynamicData && */SceneProxy->IsShown(View) && View->Family->EngineShowFlags.Particles;
+	Result.bShadowRelevance = bHasDynamicData && SceneProxy->IsShadowCast(View);
+	Result.bDynamicRelevance = bHasDynamicData;
+	if (bHasDynamicData)
+	{
+		Result.bOpaqueRelevance = View->Family->EngineShowFlags.Bounds;
+		DynamicDataRender->GetMaterialRelevance().SetPrimitiveViewRelevance(Result);
+	}
+
+	return Result;
+}
+
+void FNiagaraRenderer::SetDynamicData_RenderThread(FNiagaraDynamicDataBase* NewDynamicData)
+{
+	check(IsInRenderingThread());
+	if (DynamicDataRender)
+	{
+		delete DynamicDataRender;
+		DynamicDataRender = NULL;
+	}
+	DynamicDataRender = NewDynamicData;
 }
 
-void NiagaraRenderer::SortIndices(ENiagaraSortMode SortMode, int32 SortAttributeOffset, const FNiagaraDataBuffer& Buffer, const FMatrix& LocalToWorld, const FSceneView* View, FGlobalDynamicReadBuffer::FAllocation& OutIndices)const
+void FNiagaraRenderer::SortIndices(ENiagaraSortMode SortMode, int32 SortAttributeOffset, const FNiagaraDataBuffer& Buffer, const FMatrix& LocalToWorld, const FSceneView* View, FGlobalDynamicReadBuffer::FAllocation& OutIndices)const
 {
 	SCOPE_CYCLE_COUNTER(STAT_NiagaraSortParticles);
 
@@ -192,136 +343,3 @@ void NiagaraRenderer::SortIndices(ENiagaraSortMode SortMode, int32 SortAttribute
 		IndexBuffer[i] = ParticleOrder[i].Index;
 	}
 }
-
-	
-//////////////////////////////////////////////////////////////////////////
-// Light renderer
-
-NiagaraRendererLights::NiagaraRendererLights(ERHIFeatureLevel::Type FeatureLevel, UNiagaraRendererProperties *InProps) :
-	NiagaraRenderer()
-{
-	Properties = Cast<UNiagaraLightRendererProperties>(InProps);
-}
-
-
-void NiagaraRendererLights::ReleaseRenderThreadResources()
-{
-}
-
-void NiagaraRendererLights::CreateRenderThreadResources()
-{
-}
-
-
-
-/** Update render data buffer from attributes */
-FNiagaraDynamicDataBase *NiagaraRendererLights::GenerateVertexData(const FNiagaraSceneProxy* Proxy, FNiagaraDataSet &Data, const ENiagaraSimTarget Target)
-{
-	SCOPE_CYCLE_COUNTER(STAT_NiagaraGenLights);
-
-	SimpleTimer VertexDataTimer;
-
-	//Bail if we don't have the required attributes to render this emitter.
-	if (!bEnabled || !Properties)
-	{
-		return nullptr;
-	}
-
-	//I'm not a great fan of pulling scalar components out to a structured vert buffer like this.
-	//TODO: Experiment with a new VF that reads the data directly from the scalar layout.
-	FNiagaraDataSetIterator<FVector> PosItr(Data, Properties->PositionBinding.DataSetVariable);
-	FNiagaraDataSetIterator<FLinearColor> ColItr(Data, Properties->ColorBinding.DataSetVariable);
-	FNiagaraDataSetIterator<float> RadiusItr(Data, Properties->RadiusBinding.DataSetVariable);
-	FNiagaraDataSetIterator<float> ExponentItr(Data, Properties->LightExponentBinding.DataSetVariable);
-	FNiagaraDataSetIterator<float> ScatteringItr(Data, Properties->VolumetricScatteringBinding.DataSetVariable);
-	FNiagaraDataSetIterator<int32> EnabledItr(Data, Properties->LightRenderingEnabledBinding.DataSetVariable);
-
-	FNiagaraDynamicDataLights *DynamicData = new FNiagaraDynamicDataLights;
-	const FMatrix& LocalToWorldMatrix = Proxy->GetLocalToWorld();
-	FVector DefaultColor = FVector(Properties->ColorBinding.DefaultValueIfNonExistent.GetValue<FLinearColor>());
-	FVector DefaultPos = FVector4(LocalToWorldMatrix.GetOrigin());
-	float DefaultRadius = Properties->RadiusBinding.DefaultValueIfNonExistent.GetValue<float>();
-	float DefaultScattering = Properties->VolumetricScatteringBinding.DefaultValueIfNonExistent.GetValue<float>();
-
-	DynamicData->LightArray.Empty();
-
-	for (uint32 ParticleIndex = 0; ParticleIndex < Data.GetNumInstances(); ParticleIndex++)
-	{
-		bool ShouldRenderParticleLight = !Properties->bOverrideRenderingEnabled || (EnabledItr.IsValid() ? (*EnabledItr) : true);
-		float LightRadius = (RadiusItr.IsValid() ? (*RadiusItr) : DefaultRadius) * Properties->RadiusScale;
-		if (ShouldRenderParticleLight && LightRadius > 0)
-		{
-			SimpleLightData LightData;
-			LightData.LightEntry.Radius = LightRadius;
-			LightData.LightEntry.Color = (ColItr.IsValid() ? FVector((*ColItr)) : DefaultColor) + Properties->ColorAdd;
-			LightData.LightEntry.Exponent = Properties->bUseInverseSquaredFalloff ? 0 : (ExponentItr.IsValid() ? (*ExponentItr) : 1);
-			LightData.LightEntry.bAffectTranslucency = Properties->bAffectsTranslucency;
-			LightData.LightEntry.VolumetricScatteringIntensity = ScatteringItr.IsValid() ? (*ScatteringItr) : DefaultScattering;
-			LightData.PerViewEntry.Position = PosItr.IsValid() ? (*PosItr) : DefaultPos;
-			if (bLocalSpace)
-			{
-				LightData.PerViewEntry.Position = LocalToWorldMatrix.TransformPosition(LightData.PerViewEntry.Position);
-			}
-
-			DynamicData->LightArray.Add(LightData);
-		}
-
-		PosItr.Advance();
-		ColItr.Advance();
-		RadiusItr.Advance();
-		ExponentItr.Advance();
-		ScatteringItr.Advance();
-		EnabledItr.Advance();
-	}
-
-	CPUTimeMS = VertexDataTimer.GetElapsedMilliseconds();
-	return DynamicData;
-}
-
-
-void NiagaraRendererLights::GetDynamicMeshElements(const TArray<const FSceneView*>& Views, const FSceneViewFamily& ViewFamily, uint32 VisibilityMap, FMeshElementCollector& Collector, const FNiagaraSceneProxy *SceneProxy) const
-{
-
-}
-
-void NiagaraRendererLights::SetDynamicData_RenderThread(FNiagaraDynamicDataBase* NewDynamicData)
-{
-	if (DynamicDataRender)
-	{
-		delete static_cast<FNiagaraDynamicDataLights*>(DynamicDataRender);
-		DynamicDataRender = NULL;
-	}
-	DynamicDataRender = static_cast<FNiagaraDynamicDataLights*>(NewDynamicData);
-}
-
-int NiagaraRendererLights::GetDynamicDataSize()
-{
-	return 0;
-}
-bool NiagaraRendererLights::HasDynamicData()
-{
-	return false;
-}
-
-bool NiagaraRendererLights::SetMaterialUsage()
-{
-	return false;
-}
-
-void NiagaraRendererLights::TransformChanged()
-{
-}
-
-#if WITH_EDITORONLY_DATA
-
-const TArray<FNiagaraVariable>& NiagaraRendererLights::GetRequiredAttributes()
-{
-	return Properties->GetRequiredAttributes();
-}
-
-const TArray<FNiagaraVariable>& NiagaraRendererLights::GetOptionalAttributes()
-{
-	return Properties->GetOptionalAttributes();
-}
-
-#endif
\ No newline at end of file
diff --git a/Engine/Plugins/FX/Niagara/Source/Niagara/Private/NiagaraRendererLights.cpp b/Engine/Plugins/FX/Niagara/Source/Niagara/Private/NiagaraRendererLights.cpp
new file mode 100644
--- /dev/null
+++ b/Engine/Plugins/FX/Niagara/Source/Niagara/Private/NiagaraRendererLights.cpp
@@ -0,0 +1,125 @@
+// Copyright 1998-2019 Epic Games, Inc. All Rights Reserved.
+
+#include "NiagaraRendererLights.h"
+#include "ParticleResources.h"
+#include "NiagaraDataSet.h"
+#include "NiagaraStats.h"
+#include "NiagaraVertexFactory.h"
+#include "Engine/Engine.h"
+
+#include "NiagaraRendererLights.h"
+
+
+DECLARE_CYCLE_STAT(TEXT("Generate Particle Lights"), STAT_NiagaraGenLights, STATGROUP_Niagara);
+
+struct FNiagaraDynamicDataLights : public FNiagaraDynamicDataBase
+{
+	FNiagaraDynamicDataLights(const FNiagaraEmitterInstance* InEmitter)
+		: FNiagaraDynamicDataBase(InEmitter)
+	{
+	}
+	
+	TArray<FNiagaraRendererLights::SimpleLightData> LightArray;
+};
+
+//////////////////////////////////////////////////////////////////////////
+
+
+FNiagaraRendererLights::FNiagaraRendererLights(ERHIFeatureLevel::Type FeatureLevel, const UNiagaraRendererProperties *InProps, const FNiagaraEmitterInstance* Emitter)
+	: FNiagaraRenderer(FeatureLevel, InProps, Emitter)
+{
+	bHasLights = true;
+}
+
+FPrimitiveViewRelevance FNiagaraRendererLights::GetViewRelevance(const FSceneView* View, const FNiagaraSceneProxy *SceneProxy)const
+{
+	FPrimitiveViewRelevance Result;
+	Result.bDrawRelevance = false;
+	Result.bShadowRelevance = false;
+	Result.bDynamicRelevance = false;
+	Result.bOpaqueRelevance = false;
+	Result.bHasSimpleLights = true;
+	//MaterialRelevance.SetPrimitiveViewRelevance(Result);
+	return Result;
+}
+
+/** Update render data buffer from attributes */
+FNiagaraDynamicDataBase* FNiagaraRendererLights::GenerateDynamicData(const FNiagaraSceneProxy* Proxy, const UNiagaraRendererProperties* InProperties, const FNiagaraEmitterInstance* Emitter) const
+{
+	SCOPE_CYCLE_COUNTER(STAT_NiagaraGenLights);
+
+	SimpleTimer VertexDataTimer;
+
+	//Bail if we don't have the required attributes to render this emitter.
+	const UNiagaraLightRendererProperties* Properties = CastChecked<const UNiagaraLightRendererProperties>(InProperties);
+	FNiagaraDataSet& Data = Emitter->GetData();
+	FNiagaraDataBuffer& ParticleData = Data.GetCurrentDataChecked();
+
+	//I'm not a great fan of pulling scalar components out to a structured vert buffer like this.
+	//TODO: Experiment with a new VF that reads the data directly from the scalar layout.
+	FNiagaraDataSetAccessor<FVector> PosAccessor(Data, Properties->PositionBinding.DataSetVariable);
+	FNiagaraDataSetAccessor<FLinearColor> ColAccessor(Data, Properties->ColorBinding.DataSetVariable);
+	FNiagaraDataSetAccessor<float> RadiusAccessor(Data, Properties->RadiusBinding.DataSetVariable);
+	FNiagaraDataSetAccessor<float> ExponentAccessor(Data, Properties->LightExponentBinding.DataSetVariable);
+	FNiagaraDataSetAccessor<float> ScatteringAccessor(Data, Properties->VolumetricScatteringBinding.DataSetVariable);
+	FNiagaraDataSetAccessor<int32> EnabledAccessor(Data, Properties->LightRenderingEnabledBinding.DataSetVariable);
+
+	ENiagaraSimTarget SimTarget = Emitter->GetData().GetSimTarget();
+	FNiagaraDynamicDataLights* DynamicData = new FNiagaraDynamicDataLights(Emitter);
+
+	const FMatrix& LocalToWorldMatrix = Proxy->GetLocalToWorld();
+	FVector DefaultColor = FVector(Properties->ColorBinding.DefaultValueIfNonExistent.GetValue<FLinearColor>());
+	FVector DefaultPos = FVector4(LocalToWorldMatrix.GetOrigin());
+	float DefaultRadius = Properties->RadiusBinding.DefaultValueIfNonExistent.GetValue<float>();
+	float DefaultScattering = Properties->VolumetricScatteringBinding.DefaultValueIfNonExistent.GetValue<float>();
+
+	for (uint32 ParticleIndex = 0; ParticleIndex < ParticleData.GetNumInstances(); ParticleIndex++)
+	{
+		bool ShouldRenderParticleLight = !Properties->bOverrideRenderingEnabled || EnabledAccessor.GetSafe(ParticleIndex, true);
+		float LightRadius = RadiusAccessor.GetSafe(ParticleIndex, DefaultRadius) * Properties->RadiusScale;
+		if (ShouldRenderParticleLight && LightRadius > 0)
+		{
+			SimpleLightData LightData;
+			LightData.LightEntry.Radius = LightRadius;
+			LightData.LightEntry.Color = FVector(ColAccessor.GetSafe(ParticleIndex, DefaultColor)) + Properties->ColorAdd;
+			LightData.LightEntry.Exponent = Properties->bUseInverseSquaredFalloff ? 0 : ExponentAccessor.GetSafe(ParticleIndex, 1.0f);
+			LightData.LightEntry.bAffectTranslucency = Properties->bAffectsTranslucency;
+			LightData.LightEntry.VolumetricScatteringIntensity = ScatteringAccessor.GetSafe(ParticleIndex, DefaultScattering);
+			LightData.PerViewEntry.Position = PosAccessor.GetSafe(ParticleIndex, DefaultPos);
+			if (bLocalSpace)
+			{
+				LightData.PerViewEntry.Position = LocalToWorldMatrix.TransformPosition(LightData.PerViewEntry.Position);
+			}
+
+			DynamicData->LightArray.Add(LightData);
+		}
+	}
+
+	CPUTimeMS = VertexDataTimer.GetElapsedMilliseconds();
+	return DynamicData;
+}
+
+void FNiagaraRendererLights::GatherSimpleLights(FSimpleLightArray& OutParticleLights)const
+{
+	FNiagaraDynamicDataLights* DynamicData = (FNiagaraDynamicDataLights*)DynamicDataRender;
+	if (DynamicData)
+	{
+		int32 LightCount = DynamicData->LightArray.Num();
+
+		OutParticleLights.InstanceData.Reserve(OutParticleLights.InstanceData.Num() + LightCount);
+		OutParticleLights.PerViewData.Reserve(OutParticleLights.PerViewData.Num() + LightCount);
+
+		for (FNiagaraRendererLights::SimpleLightData &LightData : DynamicData->LightArray)
+		{
+			// When not using camera-offset, output one position for all views to share. 
+			OutParticleLights.PerViewData.Add(LightData.PerViewEntry);
+
+			// Add an entry for the light instance.
+			OutParticleLights.InstanceData.Add(LightData.LightEntry);
+		}
+	}
+}
+
+
+
+
