Copying //UE4/Dev-Niagara to Dev-Main (//UE4/Dev-Main)

diff --git a/Engine/Plugins/FX/Niagara/Source/Niagara/Private/NiagaraRendererMeshes.cpp b/Engine/Plugins/FX/Niagara/Source/Niagara/Private/NiagaraRendererMeshes.cpp
--- a/Engine/Plugins/FX/Niagara/Source/Niagara/Private/NiagaraRendererMeshes.cpp
+++ b/Engine/Plugins/FX/Niagara/Source/Niagara/Private/NiagaraRendererMeshes.cpp
@@ -14,11 +14,28 @@ DECLARE_CYCLE_STAT(TEXT("Generate Mesh Vertex Data [GT]"), STAT_NiagaraGenMeshVe
 DECLARE_CYCLE_STAT(TEXT("Render Meshes [RT]"), STAT_NiagaraRenderMeshes, STATGROUP_Niagara);
 
 
+DECLARE_DWORD_COUNTER_STAT(TEXT("NumMeshesRenderer"), STAT_NiagaraNumMeshes, STATGROUP_Niagara);
+DECLARE_DWORD_COUNTER_STAT(TEXT("NumMesheVerts"), STAT_NiagaraNumMeshVerts, STATGROUP_Niagara);
+
+static int32 GbEnableNiagaraMeshRendering = 1;
+static FAutoConsoleVariableRef CVarEnableNiagaraMeshRendering(
+	TEXT("fx.EnableNiagaraMeshRendering"),
+	GbEnableNiagaraMeshRendering,
+	TEXT("If == 0, Niagara Mesh Renderers are disabled. \n"),
+	ECVF_Default
+);
+
+struct FNiagaraDynamicDataMesh : public FNiagaraDynamicDataBase
+{
+	FNiagaraDynamicDataMesh(const FNiagaraEmitterInstance* InEmitter)
+		: FNiagaraDynamicDataBase(InEmitter)
+	{
+	}
 
-extern int32 GbNiagaraParallelEmitterRenderers;
-
-
+	TArray<FMaterialRenderProxy*, TInlineAllocator<8>> Materials;
+};
 
+//////////////////////////////////////////////////////////////////////////
 
 class FNiagaraMeshCollectorResourcesMesh : public FOneFrameResource
 {
@@ -32,14 +49,16 @@ public:
 	}
 };
 
+//////////////////////////////////////////////////////////////////////////
 
-NiagaraRendererMeshes::NiagaraRendererMeshes(ERHIFeatureLevel::Type FeatureLevel, UNiagaraRendererProperties *InProps) :
-	NiagaraRenderer()
+FNiagaraRendererMeshes::FNiagaraRendererMeshes(ERHIFeatureLevel::Type FeatureLevel, const UNiagaraRendererProperties *Props, const FNiagaraEmitterInstance* Emitter)
+	: FNiagaraRenderer(FeatureLevel, Props, Emitter)
 	, PositionOffset(INDEX_NONE)
 	, VelocityOffset(INDEX_NONE)
 	, ColorOffset(INDEX_NONE)
 	, ScaleOffset(INDEX_NONE)
 	, SizeOffset(INDEX_NONE)
+	, MaterialParamValidMask(0)
 	, MaterialParamOffset(INDEX_NONE)
 	, MaterialParamOffset1(INDEX_NONE)
 	, MaterialParamOffset2(INDEX_NONE)
@@ -48,47 +67,58 @@ NiagaraRendererMeshes::NiagaraRendererMeshes(ERHIFeatureLevel::Type FeatureLevel
 	, NormalizedAgeOffset(INDEX_NONE)
 	, MaterialRandomOffset(INDEX_NONE)
 	, CustomSortingOffset(INDEX_NONE)
-	, LastSyncedId(INDEX_NONE)
 {
-	//check(InProps);
-	VertexFactory = ConstructNiagaraMeshVertexFactory(NVFT_Mesh, FeatureLevel);
-	Properties = Cast<UNiagaraMeshRendererProperties>(InProps);
+	check(Emitter);
+	check(Props);
 
+	VertexFactory = ConstructNiagaraMeshVertexFactory(NVFT_Mesh, FeatureLevel);
+	
+	const UNiagaraMeshRendererProperties* Properties = CastChecked<const UNiagaraMeshRendererProperties>(Props);
+	check(Properties->ParticleMesh);
+
+	MeshRenderData = Properties->ParticleMesh->RenderData.Get();
+
+	FacingMode = Properties->FacingMode;
+	SortMode = Properties->SortMode;
+	bSortOnlyWhenTranslucent = Properties->bSortOnlyWhenTranslucent;
+	bOverrideMaterials = Properties->bOverrideMaterials;
+
+	const FNiagaraDataSet& Data = Emitter->GetData();
+	int32 IntDummy;
+	SizeOffset = -1;
+	Data.GetVariableComponentOffsets(Properties->PositionBinding.DataSetVariable, PositionOffset, IntDummy);
+	Data.GetVariableComponentOffsets(Properties->VelocityBinding.DataSetVariable, VelocityOffset, IntDummy);
+	Data.GetVariableComponentOffsets(Properties->ColorBinding.DataSetVariable, ColorOffset, IntDummy);
+	Data.GetVariableComponentOffsets(Properties->ScaleBinding.DataSetVariable, ScaleOffset, IntDummy);
+	Data.GetVariableComponentOffsets(Properties->DynamicMaterialBinding.DataSetVariable, MaterialParamOffset, IntDummy);
+	Data.GetVariableComponentOffsets(Properties->DynamicMaterial1Binding.DataSetVariable, MaterialParamOffset1, IntDummy);
+	Data.GetVariableComponentOffsets(Properties->DynamicMaterial2Binding.DataSetVariable, MaterialParamOffset2, IntDummy);
+	Data.GetVariableComponentOffsets(Properties->DynamicMaterial3Binding.DataSetVariable, MaterialParamOffset3, IntDummy);
+	Data.GetVariableComponentOffsets(Properties->MeshOrientationBinding.DataSetVariable, TransformOffset, IntDummy);
+	Data.GetVariableComponentOffsets(Properties->NormalizedAgeBinding.DataSetVariable, NormalizedAgeOffset, IntDummy);
+	Data.GetVariableComponentOffsets(Properties->MaterialRandomBinding.DataSetVariable, MaterialRandomOffset, IntDummy);
+	Data.GetVariableComponentOffsets(Properties->CustomSortingBinding.DataSetVariable, CustomSortingOffset, IntDummy);
+
+	MaterialParamValidMask = MaterialParamOffset  == -1 ? 0 : 0x1;
+	MaterialParamValidMask |= MaterialParamOffset1 == -1 ? 0 : 0x2;
+	MaterialParamValidMask |= MaterialParamOffset2 == -1 ? 0 : 0x4;
+	MaterialParamValidMask |= MaterialParamOffset3 == -1 ? 0 : 0x8;
+}
 
-	if (Properties && Properties->ParticleMesh)
-	{
-		if (Properties->bOverrideMaterials && Properties->OverrideMaterials.Num() != 0)
-		{
-			for (UMaterialInterface* Interface : Properties->OverrideMaterials)
-			{
-				if (Interface)
-				{
-					Interface->CheckMaterialUsage_Concurrent(MATUSAGE_NiagaraMeshParticles);
-				}
-			}
-		}
-
-		const FStaticMeshLODResources& LODModel = Properties->ParticleMesh->RenderData->LODResources[0];
-		for (int32 SectionIndex = 0; SectionIndex < LODModel.Sections.Num(); SectionIndex++)
-		{
-			//FMaterialRenderProxy* MaterialProxy = MeshMaterials[SectionIndex]->GetRenderProxy(bSelected);
-			const FStaticMeshSection& Section = LODModel.Sections[SectionIndex];
-			UMaterialInterface* ParticleMeshMaterial = Properties->ParticleMesh->GetMaterial(Section.MaterialIndex);
-			if (ParticleMeshMaterial && Properties->bOverrideMaterials == false)
-			{
-				if (ParticleMeshMaterial)
-				{
-					ParticleMeshMaterial->CheckMaterialUsage_Concurrent(MATUSAGE_NiagaraMeshParticles);
-				}
-			}
-		}
-
-		BaseExtents = Properties->ParticleMesh->GetBounds().BoxExtent;
-	}
+void FNiagaraRendererMeshes::ReleaseRenderThreadResources()
+{
+	FNiagaraRenderer::ReleaseRenderThreadResources();
+	VertexFactory->ReleaseResource();
+	WorldSpacePrimitiveUniformBuffer.ReleaseResource();
+}
 
+void FNiagaraRendererMeshes::CreateRenderThreadResources()
+{
+	FNiagaraRenderer::CreateRenderThreadResources();
+	VertexFactory->InitResource();
 }
 
-void NiagaraRendererMeshes::SetupVertexFactory(FNiagaraMeshVertexFactory *InVertexFactory, const FStaticMeshLODResources& LODResources) const
+void FNiagaraRendererMeshes::SetupVertexFactory(FNiagaraMeshVertexFactory *InVertexFactory, const FStaticMeshLODResources& LODResources) const
 {
 	FStaticMeshDataType Data;
 
@@ -99,21 +129,7 @@ void NiagaraRendererMeshes::SetupVertexFactory(FNiagaraMeshVertexFactory *InVert
 	InVertexFactory->SetData(Data);
 }
 
-
-
-void NiagaraRendererMeshes::ReleaseRenderThreadResources()
-{
-	VertexFactory->ReleaseResource();
-	WorldSpacePrimitiveUniformBuffer.ReleaseResource();
-}
-
-void NiagaraRendererMeshes::CreateRenderThreadResources()
-{
-	VertexFactory->InitResource();
-}
-
-
-void NiagaraRendererMeshes::GetDynamicMeshElements(const TArray<const FSceneView*>& Views, const FSceneViewFamily& ViewFamily, uint32 VisibilityMap, FMeshElementCollector& Collector, const FNiagaraSceneProxy *SceneProxy) const
+void FNiagaraRendererMeshes::GetDynamicMeshElements(const TArray<const FSceneView*>& Views, const FSceneViewFamily& ViewFamily, uint32 VisibilityMap, FMeshElementCollector& Collector, const FNiagaraSceneProxy *SceneProxy) const
 {
 	SCOPE_CYCLE_COUNTER(STAT_NiagaraRender);
 	SCOPE_CYCLE_COUNTER(STAT_NiagaraRenderMeshes);
@@ -121,35 +137,40 @@ void NiagaraRendererMeshes::GetDynamicMeshElements(const TArray<const FSceneView
 	SimpleTimer MeshElementsTimer;
 
 	FNiagaraDynamicDataMesh *DynamicDataMesh = (static_cast<FNiagaraDynamicDataMesh*>(DynamicDataRender));
-	//if (!DynamicDataMesh || !DynamicDataMesh->DataSet || DynamicDataMesh->DataSet->GetNumInstances()==0)
-	if (!DynamicDataMesh 
-		|| DynamicDataMesh->RTParticleData.GetNumInstancesAllocated() == 0
-		|| DynamicDataMesh->RTParticleData.GetNumInstances() == 0
-		|| nullptr == Properties
-		|| !GSupportsResourceView  // Current shader requires SRV to draw properly in all cases.
+	if (!DynamicDataMesh)
+	{
+		return;
+	}
+
+	FNiagaraDataBuffer* SourceParticleData = DynamicDataMesh->GetParticleDataToRender();
+	if( SourceParticleData == nullptr ||
+		MeshRenderData == nullptr ||
+		SourceParticleData->GetNumInstancesAllocated() == 0 ||
+		SourceParticleData->GetNumInstances() == 0 ||
+		GbEnableNiagaraMeshRendering == 0 ||
+		!GSupportsResourceView  // Current shader requires SRV to draw properly in all cases.
 		)
 	{
 		return;
 	}
 
-	int32 NumInstances = DynamicDataMesh->RTParticleData.GetNumInstances();
+#if STATS
+	FScopeCycleCounter EmitterStatsCounter(EmitterStatID);
+#endif
 
-	int32 TotalFloatSize = DynamicDataMesh->RTParticleData.GetFloatBuffer().Num() / sizeof(float);
+	int32 NumInstances = SourceParticleData->GetNumInstances();
 
+	int32 TotalFloatSize = SourceParticleData->GetFloatBuffer().Num() / sizeof(float);
 	FGlobalDynamicReadBuffer& DynamicReadBuffer = Collector.GetDynamicReadBuffer();
 	FGlobalDynamicReadBuffer::FAllocation ParticleData;
-	
+
 	//For cpu sims we allocate render buffers from the global pool. GPU sims own their own.
-	if (DynamicDataMesh->DataSet->GetSimTarget() == ENiagaraSimTarget::CPUSim)
+	if (DynamicDataMesh->GetSimTarget() == ENiagaraSimTarget::CPUSim)
 	{
 		ParticleData = DynamicReadBuffer.AllocateFloat(TotalFloatSize);
-		FMemory::Memcpy(ParticleData.Buffer, DynamicDataMesh->RTParticleData.GetFloatBuffer().GetData(), DynamicDataMesh->RTParticleData.GetFloatBuffer().Num());
+		FMemory::Memcpy(ParticleData.Buffer, SourceParticleData->GetFloatBuffer().GetData(), SourceParticleData->GetFloatBuffer().Num());
 	}
 
-	//check(DynamicDataMesh->DataSet->PrevDataRender().GetNumInstances() > 0);
-	//check(DynamicDataMesh->DataSet->PrevDataRender().GetNumInstancesAllocated() > 0);
-	//check(DynamicDataMesh->DataSet->PrevData().GetGPUBufferFloat()->NumBytes > 0 || DynamicDataMesh->DataSet->PrevData().GetGPUBufferInt()->NumBytes > 0)
-	
 	{
 		// Update the primitive uniform buffer if needed.
 		if (!WorldSpacePrimitiveUniformBuffer.IsInitialized())
@@ -182,7 +203,8 @@ void NiagaraRendererMeshes::GetDynamicMeshElements(const TArray<const FSceneView
 			if (VisibilityMap & (1 << ViewIndex))
 			{
 				const FSceneView* View = Views[ViewIndex];
-				const FStaticMeshLODResources& LODModel = Properties->ParticleMesh->RenderData->LODResources[0];
+				//TODO: Work out LOD and pass through in dynamic data.
+				const FStaticMeshLODResources& LODModel = MeshRenderData->LODResources[0];
 
 				FNiagaraMeshCollectorResourcesMesh& CollectorResources = Collector.AllocateOneFrameResource<FNiagaraMeshCollectorResourcesMesh>();
 				SetupVertexFactory(&CollectorResources.VertexFactory, LODModel);
@@ -197,6 +219,7 @@ void NiagaraRendererMeshes::GetDynamicMeshElements(const TArray<const FSceneView
 				PerViewUniformParameters.TransformDataOffset = TransformOffset;
 				PerViewUniformParameters.ScaleDataOffset = ScaleOffset;
 				PerViewUniformParameters.SizeDataOffset = SizeOffset;
+				PerViewUniformParameters.MaterialParamValidMask = MaterialParamValidMask;
 				PerViewUniformParameters.MaterialParamDataOffset = MaterialParamOffset;
 				PerViewUniformParameters.MaterialParam1DataOffset = MaterialParamOffset1;
 				PerViewUniformParameters.MaterialParam2DataOffset = MaterialParamOffset2;
@@ -204,54 +227,25 @@ void NiagaraRendererMeshes::GetDynamicMeshElements(const TArray<const FSceneView
 				PerViewUniformParameters.NormalizedAgeDataOffset = NormalizedAgeOffset;
 				PerViewUniformParameters.MaterialRandomDataOffset = MaterialRandomOffset;
 				PerViewUniformParameters.DefaultPos = bLocalSpace ? FVector4(0.0f, 0.0f, 0.0f, 1.0f) : FVector4(SceneProxy->GetLocalToWorld().GetOrigin());
-				/*
-				if (Properties)
-				{
-				PerViewUniformParameters.SubImageSize = FVector4(Properties->SubImageInfo.X, Properties->SubImageInfo.Y, 1.0f / Properties->SubImageInfo.X, 1.0f / Properties->SubImageInfo.Y);
-				}
-				*/
 
 				//Grab the material proxies we'll be using for each section and check them for translucency.
 				TArray<FMaterialRenderProxy*, TInlineAllocator<32>> MaterialProxies;
 				bool bHasTranslucentMaterials = false;
-				for (int32 SectionIndex = 0; SectionIndex < LODModel.Sections.Num(); SectionIndex++)
+				for (FMaterialRenderProxy* MaterialProxy : DynamicDataMesh->Materials)
 				{
-					const FStaticMeshSection& Section = LODModel.Sections[SectionIndex];
-					UMaterialInterface* ParticleMeshMaterial = Properties->ParticleMesh->GetMaterial(Section.MaterialIndex);
-					FMaterialRenderProxy* MaterialProxy = nullptr;
-
-					if (Properties->bOverrideMaterials && Properties->OverrideMaterials.Num() > Section.MaterialIndex &&
-						Properties->OverrideMaterials[Section.MaterialIndex] != nullptr)
-					{
-						MaterialProxy = Properties->OverrideMaterials[Section.MaterialIndex]->GetRenderProxy();
-					}
-
-					if (MaterialProxy == nullptr && ParticleMeshMaterial)
-					{
-						MaterialProxy = ParticleMeshMaterial->GetRenderProxy();
-					}
-
-					if (MaterialProxy == nullptr)
-					{
-						MaterialProxy = UMaterial::GetDefaultMaterial(MD_Surface)->GetRenderProxy();
-					}
-
-					MaterialProxies.Add(MaterialProxy);
-					if (MaterialProxy)
-					{
-						EBlendMode BlendMode = MaterialProxy->GetMaterial(VertexFactory->GetFeatureLevel())->GetBlendMode();
-						bHasTranslucentMaterials |= BlendMode == BLEND_AlphaComposite || BlendMode == BLEND_Translucent;
-					}
+					check(MaterialProxy);
+					EBlendMode BlendMode = MaterialProxy->GetMaterial(SceneProxy->GetScene().GetFeatureLevel())->GetBlendMode();
+					bHasTranslucentMaterials |= BlendMode == BLEND_AlphaComposite || BlendMode == BLEND_Translucent;
 				}
 
 				//Sort particles if needed.
 				CollectorResources.VertexFactory.SetSortedIndices(nullptr, 0xFFFFFFFF);
 
 				FNiagaraGPUSortInfo SortInfo;
-				if (View && Properties->SortMode != ENiagaraSortMode::None && (bHasTranslucentMaterials || !Properties->bSortOnlyWhenTranslucent))
+				if (View && SortMode != ENiagaraSortMode::None && (bHasTranslucentMaterials || !bSortOnlyWhenTranslucent))
 				{
 					SortInfo.ParticleCount = NumInstances;
-					SortInfo.SortMode = Properties->SortMode;
+					SortInfo.SortMode = SortMode;
 					SortInfo.SortAttributeOffset = (SortInfo.SortMode == ENiagaraSortMode::CustomAscending || SortInfo.SortMode == ENiagaraSortMode::CustomDecending) ? CustomSortingOffset : PositionOffset;
 					SortInfo.ViewOrigin = View->ViewMatrices.GetViewOrigin();
 					SortInfo.ViewDirection = View->GetViewDirection();
@@ -263,7 +257,7 @@ void NiagaraRendererMeshes::GetDynamicMeshElements(const TArray<const FSceneView
 					}
 				};
 
-				if (DynamicDataMesh->DataSet->GetSimTarget() == ENiagaraSimTarget::CPUSim)//TODO: Compute shader for sorting gpu sims and larger cpu sims.
+				if (DynamicDataMesh->GetSimTarget() == ENiagaraSimTarget::CPUSim)//TODO: Compute shader for sorting gpu sims and larger cpu sims.
 				{
 					check(ParticleData.IsValid());
 					if (SortInfo.SortMode != ENiagaraSortMode::None && SortInfo.SortAttributeOffset != INDEX_NONE)
@@ -273,9 +267,10 @@ void NiagaraRendererMeshes::GetDynamicMeshElements(const TArray<const FSceneView
 							SortInfo.ParticleCount >= GNiagaraGPUSortingCPUToGPUThreshold &&
 							SceneProxy->GetBatcher())
 						{
+							SortInfo.ParticleCount = NumInstances;
 							SortInfo.ParticleDataFloatSRV = ParticleData.ReadBuffer->SRV;
 							SortInfo.FloatDataOffset = ParticleData.FirstIndex / sizeof(float);
-							SortInfo.FloatDataStride = DynamicDataMesh->RTParticleData.GetFloatStride() / sizeof(float);
+							SortInfo.FloatDataStride = SourceParticleData->GetFloatStride() / sizeof(float);
 							const int32 IndexBufferOffset = SceneProxy->GetBatcher()->AddSortedGPUSimulation(SortInfo);
 							if (IndexBufferOffset != INDEX_NONE)
 							{
@@ -286,20 +281,20 @@ void NiagaraRendererMeshes::GetDynamicMeshElements(const TArray<const FSceneView
 						{
 							FGlobalDynamicReadBuffer::FAllocation SortedIndices;
 							SortedIndices = DynamicReadBuffer.AllocateInt32(NumInstances);
-							SortIndices(SortInfo.SortMode, SortInfo.SortAttributeOffset, DynamicDataMesh->RTParticleData, SceneProxy->GetLocalToWorld(), View, SortedIndices);
+							SortIndices(SortInfo.SortMode, SortInfo.SortAttributeOffset, *SourceParticleData, SceneProxy->GetLocalToWorld(), View, SortedIndices);
 							CollectorResources.VertexFactory.SetSortedIndices(SortedIndices.ReadBuffer->SRV, SortedIndices.FirstIndex / sizeof(float));
 						}
 					}
-					CollectorResources.VertexFactory.SetParticleData(ParticleData.ReadBuffer->SRV, ParticleData.FirstIndex / sizeof(float), DynamicDataMesh->RTParticleData.GetFloatStride() / sizeof(float));
+					CollectorResources.VertexFactory.SetParticleData(ParticleData.ReadBuffer->SRV, ParticleData.FirstIndex / sizeof(float), SourceParticleData->GetFloatStride() / sizeof(float));
 				}
 				else
 				{
 					if (SortInfo.SortMode != ENiagaraSortMode::None && SortInfo.SortAttributeOffset != INDEX_NONE && GNiagaraGPUSorting && SceneProxy->GetBatcher())
 					{
-						SortInfo.ParticleDataFloatSRV = DynamicDataMesh->DataSet->CurrData().GetGPUBufferFloat()->SRV;
+						SortInfo.ParticleDataFloatSRV = SourceParticleData->GetGPUBufferFloat().SRV;
 						SortInfo.FloatDataOffset = 0;
-						SortInfo.FloatDataStride = DynamicDataMesh->RTParticleData.GetFloatStride() / sizeof(float);
-						SortInfo.GPUParticleCountSRV = DynamicDataMesh->DataSet->GetCurDataSetIndices().SRV;
+						SortInfo.FloatDataStride = SourceParticleData->GetFloatStride() / sizeof(float);
+						SortInfo.GPUParticleCountSRV = SourceParticleData->GetGPUIndices().SRV;
 						SortInfo.GPUParticleCountOffset = 1;
 						const int32 IndexBufferOffset = SceneProxy->GetBatcher()->AddSortedGPUSimulation(SortInfo);
 						if (IndexBufferOffset != INDEX_NONE)
@@ -307,23 +302,42 @@ void NiagaraRendererMeshes::GetDynamicMeshElements(const TArray<const FSceneView
 							CollectorResources.VertexFactory.SetSortedIndices(SceneProxy->GetBatcher()->GetGPUSortedBuffer().VertexBufferSRV, IndexBufferOffset);
 						}
 					}
-					CollectorResources.VertexFactory.SetParticleData(DynamicDataMesh->DataSet->CurrData().GetGPUBufferFloat()->SRV, 0, DynamicDataMesh->DataSet->CurrData().GetFloatStride() / sizeof(float));
+
+					if (SortInfo.SortMode != ENiagaraSortMode::None && SortInfo.SortAttributeOffset != INDEX_NONE && GNiagaraGPUSorting && SceneProxy->GetBatcher())
+					{
+						// Here we need to be conservative about the InstanceCount, since the final value is only known on the GPU after the simulation.
+						SortInfo.ParticleCount = SourceParticleData->GetNumInstances();
+
+						SortInfo.ParticleDataFloatSRV = SourceParticleData->GetGPUBufferFloat().SRV;
+						SortInfo.FloatDataOffset = 0;
+						SortInfo.FloatDataStride = SourceParticleData->GetFloatStride() / sizeof(float);
+						SortInfo.GPUParticleCountSRV = SourceParticleData->GetGPUIndices().SRV;
+						SortInfo.GPUParticleCountOffset = 1;
+						const int32 IndexBufferOffset = SceneProxy->GetBatcher()->AddSortedGPUSimulation(SortInfo);
+						if (IndexBufferOffset != INDEX_NONE)
+						{
+							CollectorResources.VertexFactory.SetSortedIndices(SceneProxy->GetBatcher()->GetGPUSortedBuffer().VertexBufferSRV, IndexBufferOffset);
+						}
+					}
+					CollectorResources.VertexFactory.SetParticleData(SourceParticleData->GetGPUBufferFloat().SRV, 0, SourceParticleData->GetFloatStride() / sizeof(float));
 				}
 
 				// Collector.AllocateOneFrameResource uses default ctor, initialize the vertex factory
 				CollectorResources.VertexFactory.SetParticleFactoryType(NVFT_Mesh);
-				CollectorResources.VertexFactory.SetMeshFacingMode((uint32)Properties->FacingMode);
+				CollectorResources.VertexFactory.SetMeshFacingMode((uint32)FacingMode);
 				CollectorResources.UniformBuffer = FNiagaraMeshUniformBufferRef::CreateUniformBufferImmediate(PerViewUniformParameters, UniformBuffer_SingleFrame);
 
 				CollectorResources.VertexFactory.InitResource();
 				CollectorResources.VertexFactory.SetUniformBuffer(CollectorResources.UniformBuffer);
 			
+				// GPU mesh rendering currently only supports one mesh section.
+				// TODO: Add proper support for multiple mesh sections for GPU mesh particles.
+				int32 MaxSection = DynamicDataMesh->GetSimTarget() == ENiagaraSimTarget::GPUComputeSim ? 1 : LODModel.Sections.Num();
 				const bool bIsWireframe = AllowDebugViewmodes() && View && View->Family->EngineShowFlags.Wireframe;
-
-				for (int32 SectionIndex = 0; SectionIndex < LODModel.Sections.Num(); SectionIndex++)
+				for (int32 SectionIndex = 0; SectionIndex < MaxSection; SectionIndex++)
 				{
 					const FStaticMeshSection& Section = LODModel.Sections[SectionIndex];
-					FMaterialRenderProxy* MaterialProxy = MaterialProxies[SectionIndex];
+					FMaterialRenderProxy* MaterialProxy = DynamicDataMesh->Materials[SectionIndex];
 					if ((Section.NumTriangles == 0) || (MaterialProxy == NULL))
 					{
 						//@todo. This should never occur, but it does occasionally.
@@ -374,10 +388,10 @@ void NiagaraRendererMeshes::GetDynamicMeshElements(const TArray<const FSceneView
 						BatchElement.NumPrimitives = Section.NumTriangles;
 					}
 
-					if (DynamicDataMesh->DataSet->GetSimTarget() == ENiagaraSimTarget::GPUComputeSim)
+					if (DynamicDataMesh->GetSimTarget() == ENiagaraSimTarget::GPUComputeSim)
 					{
 						BatchElement.NumPrimitives = 0;
-						BatchElement.IndirectArgsBuffer = DynamicDataMesh->DataSet->GetCurDataSetIndices().Buffer;
+						BatchElement.IndirectArgsBuffer = SourceParticleData->GetGPUIndices().Buffer;
 					}
 					else
 					{
@@ -388,6 +402,9 @@ void NiagaraRendererMeshes::GetDynamicMeshElements(const TArray<const FSceneView
 					Mesh.bUseWireframeSelectionColoring = SceneProxy->IsSelected();
 
 					Collector.AddMesh(ViewIndex, Mesh);
+
+					INC_DWORD_STAT_BY(STAT_NiagaraNumMeshVerts, NumInstances * LODModel.GetNumVertices());
+					INC_DWORD_STAT_BY(STAT_NiagaraNumMeshes, NumInstances);
 				}
 			}
 		}
@@ -396,116 +413,101 @@ void NiagaraRendererMeshes::GetDynamicMeshElements(const TArray<const FSceneView
 	CPUTimeMS += MeshElementsTimer.GetElapsedMilliseconds();
 }
 
-
-
-bool NiagaraRendererMeshes::SetMaterialUsage()
-{
-	//Causes deadlock :S Need to look at / rework the setting of materials and render modules.
-	return Material && Material->CheckMaterialUsage_Concurrent(MATUSAGE_NiagaraMeshParticles);
-}
-
-void NiagaraRendererMeshes::TransformChanged()
-{
-	WorldSpacePrimitiveUniformBuffer.ReleaseResource();
-}
-
-/** Update render data buffer from attributes */
-FNiagaraDynamicDataBase *NiagaraRendererMeshes::GenerateVertexData(const FNiagaraSceneProxy* Proxy, FNiagaraDataSet &Data, const ENiagaraSimTarget Target)
+FNiagaraDynamicDataBase *FNiagaraRendererMeshes::GenerateDynamicData(const FNiagaraSceneProxy* Proxy, const UNiagaraRendererProperties* InProperties, const FNiagaraEmitterInstance* Emitter) const
 {
 	SCOPE_CYCLE_COUNTER(STAT_NiagaraRenderGT);
 	SCOPE_CYCLE_COUNTER(STAT_NiagaraGenMeshVertexData);
 
-	if (!Properties || Properties->ParticleMesh == nullptr || !bEnabled)
+	const UNiagaraMeshRendererProperties* Properties = CastChecked<const UNiagaraMeshRendererProperties>(InProperties);
+
+	if (Properties->ParticleMesh == nullptr)
 	{
 		return nullptr;
 	}
 
 	SimpleTimer VertexDataTimer;
-//	TArray<FNiagaraMeshInstanceVertex>& RenderData = DynamicData->VertexData;
-	//TArray< FNiagaraMeshInstanceVertexDynamicParameter>& RenderMaterialVertexData = DynamicData->MaterialParameterVertexData;
 
-	if (PositionOffset == INDEX_NONE || LastSyncedId != Properties->SyncId)
-	{
-		// optional attributes
-		int32 IntDummy;
-		SizeOffset = -1;
-		Data.GetVariableComponentOffsets(Properties->PositionBinding.DataSetVariable, PositionOffset, IntDummy);
-		Data.GetVariableComponentOffsets(Properties->VelocityBinding.DataSetVariable, VelocityOffset, IntDummy);
-		Data.GetVariableComponentOffsets(Properties->ColorBinding.DataSetVariable, ColorOffset, IntDummy);
-		Data.GetVariableComponentOffsets(Properties->ScaleBinding.DataSetVariable, ScaleOffset, IntDummy);
-		Data.GetVariableComponentOffsets(Properties->DynamicMaterialBinding.DataSetVariable, MaterialParamOffset, IntDummy);
-		Data.GetVariableComponentOffsets(Properties->DynamicMaterial1Binding.DataSetVariable, MaterialParamOffset1, IntDummy);
-		Data.GetVariableComponentOffsets(Properties->DynamicMaterial2Binding.DataSetVariable, MaterialParamOffset2, IntDummy);
-		Data.GetVariableComponentOffsets(Properties->DynamicMaterial3Binding.DataSetVariable, MaterialParamOffset3, IntDummy);
-		Data.GetVariableComponentOffsets(Properties->MeshOrientationBinding.DataSetVariable, TransformOffset, IntDummy);
-		Data.GetVariableComponentOffsets(Properties->NormalizedAgeBinding.DataSetVariable, NormalizedAgeOffset, IntDummy);
-		Data.GetVariableComponentOffsets(Properties->MaterialRandomBinding.DataSetVariable, MaterialRandomOffset, IntDummy);
-		Data.GetVariableComponentOffsets(Properties->CustomSortingBinding.DataSetVariable, CustomSortingOffset, IntDummy);
-		LastSyncedId = Properties->SyncId;
-	}
+	FNiagaraDynamicDataMesh *DynamicData = nullptr;
 
-	//Bail if we don't have the required attributes to render this emitter.
-	if (!bEnabled)
+	ENiagaraSimTarget SimTarget = Emitter->GetData().GetSimTarget();
+	if (Properties->ParticleMesh)
 	{
-		return nullptr;
-	}
+		DynamicData = new FNiagaraDynamicDataMesh(Emitter);
 
-	FNiagaraDynamicDataMesh *DynamicData = nullptr;
+		//TODO: Mesh LODs.
+		const FStaticMeshLODResources& LODModel = Properties->ParticleMesh->RenderData->LODResources[0];
 
-	if (Data.CurrData().GetNumInstances() > 0)
-	{
-		DynamicData = new FNiagaraDynamicDataMesh;
+		check(BaseMaterials_GT.Num() == LODModel.Sections.Num());
 
-		//TODO: This buffer is far fatter than needed. Just pull out the data needed for rendering.
-		Data.CurrData().CopyTo(DynamicData->RTParticleData);
+		DynamicData->Materials.Reset(LODModel.Sections.Num());
+		DynamicData->SetMaterialRelevance(BaseMaterialRelevance_GT);
+		for (int32 SectionIndex = 0; SectionIndex < LODModel.Sections.Num(); SectionIndex++)
+		{
+			const FStaticMeshSection& Section = LODModel.Sections[SectionIndex];
+			UMaterialInterface* SectionMat = nullptr;
 
-		DynamicData->DataSet = &Data;
+			//In preparation for a material override feature, we pass our material(s) and relevance in via dynamic data.
+			//The renderer ensures we have the correct usage and relevance for materials in BaseMaterials_GT.
+			//Any override feature must also do the same for materials that are set.
+			check(BaseMaterials_GT[SectionIndex]->CheckMaterialUsage_Concurrent(MATUSAGE_NiagaraMeshParticles));
+			DynamicData->Materials.Add(BaseMaterials_GT[SectionIndex]->GetRenderProxy());
+		}
 	}
 
 	CPUTimeMS = VertexDataTimer.GetElapsedMilliseconds();
 	return DynamicData;  
 }
 
-
-
-void NiagaraRendererMeshes::SetDynamicData_RenderThread(FNiagaraDynamicDataBase* NewDynamicData)
-{
-	check(IsInRenderingThread());
-
-	if (DynamicDataRender)
-	{
-		delete static_cast<FNiagaraDynamicDataMesh*>(DynamicDataRender);
-		DynamicDataRender = NULL;
-	}
-	DynamicDataRender = NewDynamicData;
-}
-
-int NiagaraRendererMeshes::GetDynamicDataSize()
+int FNiagaraRendererMeshes::GetDynamicDataSize()const
 {
 	uint32 Size = sizeof(FNiagaraDynamicDataMesh);
-	if (DynamicDataRender && static_cast<FNiagaraDynamicDataMesh*>(DynamicDataRender)->DataSet)
-	{
-		//Size += (static_cast<FNiagaraDynamicDataMesh*>(DynamicDataRender))->DataSet->PrevDataRender().GetNumInstances() * sizeof(float);
-	}
-
 	return Size;
 }
 
-bool NiagaraRendererMeshes::HasDynamicData()
+void FNiagaraRendererMeshes::TransformChanged()
 {
-	return DynamicDataRender != nullptr;
+	WorldSpacePrimitiveUniformBuffer.ReleaseResource();
 }
 
-#if WITH_EDITORONLY_DATA
-
-const TArray<FNiagaraVariable>& NiagaraRendererMeshes::GetRequiredAttributes()
+bool FNiagaraRendererMeshes::IsMaterialValid(UMaterialInterface* Mat)const
 {
-	return Properties->GetRequiredAttributes();
+	return Mat && Mat->CheckMaterialUsage(MATUSAGE_NiagaraMeshParticles);
 }
 
-const TArray<FNiagaraVariable>& NiagaraRendererMeshes::GetOptionalAttributes()
+
+//////////////////////////////////////////////////////////////////////////
+// Proposed class for ensuring Niagara/Cascade components who's proxies reference render data of other objects (Materials, Meshes etc) do not have data freed from under them.
+// Our components register themselves with the referenced component which then calls InvalidateRenderDependencies() whenever it's render data is changed or when it is destroyed.
+// UNTESTED - DO NOT USE.
+struct FComponentRenderDependencyHandler
 {
-	return Properties->GetOptionalAttributes();
-}
+	void AddDependency(UPrimitiveComponent* Component)
+	{
+		DependentComponents.Add(Component);
+	}
+
+	void RemoveDependancy(UPrimitiveComponent* Component)
+	{
+		DependentComponents.RemoveSwap(Component);
+	}
+
+	void InvalidateRenderDependencies()
+	{
+		int32 i = DependentComponents.Num();
+		while(--i >= 0)
+		{
+			if (UPrimitiveComponent* Comp = DependentComponents[i].Get())
+			{
+				Comp->MarkRenderStateDirty();
+			}
+			else
+			{
+				DependentComponents.RemoveAtSwap(i);
+			}
+		}
+	}
+
+	TArray<TWeakObjectPtr<UPrimitiveComponent>> DependentComponents;
+};
 
-#endif
\ No newline at end of file
+//////////////////////////////////////////////////////////////////////////
\ No newline at end of file
diff --git a/Engine/Plugins/FX/Niagara/Source/Niagara/Private/NiagaraRendererRibbons.cpp b/Engine/Plugins/FX/Niagara/Source/Niagara/Private/NiagaraRendererRibbons.cpp
--- a/Engine/Plugins/FX/Niagara/Source/Niagara/Private/NiagaraRendererRibbons.cpp
+++ b/Engine/Plugins/FX/Niagara/Source/Niagara/Private/NiagaraRendererRibbons.cpp
@@ -64,15 +64,21 @@ static FAutoConsoleVariableRef CVarNiagaraRibbonMinSegmentLength(
 	ECVF_Scalability
 );
 
-float GNiagaraRibbonMaxAveraging = .98f;
+float GNiagaraRibbonMaxAveraging = .5f;
 static FAutoConsoleVariableRef CVarNiagaraRibbonMaxAveraging(
 	TEXT("Niagara.Ribbon.Tessellation.MaxAveraging"),
 	GNiagaraRibbonMaxAveraging,
-	TEXT("Averaging value between the max metrics and the current metrics. (default=95%)"),
+	TEXT("Averaging value between the max metrics and the current metrics. (default=50%)"),
 	ECVF_Scalability
 );
 
-
+static int32 GbEnableNiagaraRibbonRendering = 1;
+static FAutoConsoleVariableRef CVarEnableNiagaraRibbonRendering(
+	TEXT("fx.EnableNiagaraRibbonRendering"),
+	GbEnableNiagaraRibbonRendering,
+	TEXT("If == 0, Niagara Ribbon Renderers are disabled. \n"),
+	ECVF_Default
+);
 
 // max absolute error 9.0x10^-3
 // Eberly's polynomial degree 1 - respect bounds
@@ -87,6 +93,15 @@ FORCEINLINE float AcosFast(float InX)
 
 struct FNiagaraDynamicDataRibbon : public FNiagaraDynamicDataBase
 {
+	FNiagaraDynamicDataRibbon(const FNiagaraEmitterInstance* InEmitter)
+		: FNiagaraDynamicDataBase(InEmitter)
+		, Material(nullptr)
+	{
+	}
+	
+	/** Material to use passed to the Renderer. */
+	FMaterialRenderProxy* Material;
+
 	// The list of all segments, each one connecting SortedIndices[SegmentId] to SortedIndices[SegmentId + 1].
 	// We use this format because the final index buffer gets generated based on view sorting and InterpCount.
 	TArray<int32> SegmentData;
@@ -95,10 +110,6 @@ struct FNiagaraDynamicDataRibbon : public FNiagaraDynamicDataBase
 	TArray<uint32> MultiRibbonIndices;
 	TArray<float> PackedPerRibbonDataByIndex;
 
-	//Direct ptr to the dataset. ONLY FOR USE BE GPU EMITTERS.
-	//TODO: Even this needs to go soon.
-	const FNiagaraDataSet *DataSet;
-
 	// start and end world space position of the ribbon, to figure out draw direction
 	FVector StartPos;
 	FVector EndPos;
@@ -128,8 +139,15 @@ public:
 };
 
 
-NiagaraRendererRibbons::NiagaraRendererRibbons(ERHIFeatureLevel::Type FeatureLevel, UNiagaraRendererProperties *InProps) :
-	NiagaraRenderer()
+FNiagaraRendererRibbons::FNiagaraRendererRibbons(ERHIFeatureLevel::Type FeatureLevel, const UNiagaraRendererProperties *InProps, const FNiagaraEmitterInstance* Emitter)
+	: FNiagaraRenderer(FeatureLevel, InProps, Emitter)
+	, FacingMode(ENiagaraRibbonFacingMode::Screen)
+	, UV0TilingDistance(0.0f)
+	, UV0Scale(FVector2D(1.0f, 1.0f))
+	, UV0AgeOffsetMode(ENiagaraRibbonAgeOffsetMode::Scale)
+	, UV1TilingDistance(0.0f)
+	, UV1Scale(FVector2D(1.0f, 1.0f))
+	, UV1AgeOffsetMode(ENiagaraRibbonAgeOffsetMode::Scale)
 	, PositionDataOffset(INDEX_NONE)
 	, VelocityDataOffset(INDEX_NONE)
 	, WidthDataOffset(INDEX_NONE)
@@ -138,26 +156,63 @@ NiagaraRendererRibbons::NiagaraRendererRibbons(ERHIFeatureLevel::Type FeatureLev
 	, ColorDataOffset(INDEX_NONE)
 	, NormalizedAgeDataOffset(INDEX_NONE)
 	, MaterialRandomDataOffset(INDEX_NONE)
-	, LastSyncedId(INDEX_NONE)
 {
 	VertexFactory = new FNiagaraRibbonVertexFactory(NVFT_Ribbon, FeatureLevel);
-	Properties = Cast<UNiagaraRibbonRendererProperties>(InProps);
+
+	const UNiagaraRibbonRendererProperties* Properties = CastChecked<const UNiagaraRibbonRendererProperties>(InProps);
+	const FNiagaraDataSet& Data = Emitter->GetData();
+
+	FacingMode = Properties->FacingMode;
+	UV0TilingDistance = Properties->UV0TilingDistance;
+	UV0Scale = Properties->UV0Scale;
+	UV0Offset = Properties->UV0Offset;
+	UV0AgeOffsetMode = Properties->UV0AgeOffsetMode;
+	UV1TilingDistance = Properties->UV1TilingDistance;
+	UV1Scale = Properties->UV1Scale;
+	UV1Offset = Properties->UV1Offset;
+	UV1AgeOffsetMode = Properties->UV1AgeOffsetMode;
+	DrawDirection = Properties->DrawDirection;
+
+	// required attributes
+	int32 IntDummy;
+	Data.GetVariableComponentOffsets(Properties->PositionBinding.DataSetVariable, PositionDataOffset, IntDummy);
+	Data.GetVariableComponentOffsets(Properties->VelocityBinding.DataSetVariable, VelocityDataOffset, IntDummy);
+	Data.GetVariableComponentOffsets(Properties->ColorBinding.DataSetVariable, ColorDataOffset, IntDummy);
+
+	// optional attributes
+	Data.GetVariableComponentOffsets(Properties->RibbonWidthBinding.DataSetVariable, WidthDataOffset, IntDummy);
+	Data.GetVariableComponentOffsets(Properties->RibbonTwistBinding.DataSetVariable, TwistDataOffset, IntDummy);
+	Data.GetVariableComponentOffsets(Properties->RibbonFacingBinding.DataSetVariable, FacingDataOffset, IntDummy);
+	Data.GetVariableComponentOffsets(Properties->NormalizedAgeBinding.DataSetVariable, NormalizedAgeDataOffset, IntDummy);
+	Data.GetVariableComponentOffsets(Properties->MaterialRandomBinding.DataSetVariable, MaterialRandomDataOffset, IntDummy);
+
+	Data.GetVariableComponentOffsets(Properties->DynamicMaterialBinding.DataSetVariable, MaterialParamOffset, IntDummy);
+	Data.GetVariableComponentOffsets(Properties->DynamicMaterial1Binding.DataSetVariable, MaterialParamOffset1, IntDummy);
+	Data.GetVariableComponentOffsets(Properties->DynamicMaterial2Binding.DataSetVariable, MaterialParamOffset2, IntDummy);
+	Data.GetVariableComponentOffsets(Properties->DynamicMaterial3Binding.DataSetVariable, MaterialParamOffset3, IntDummy);
+
+	MaterialParamValidMask = MaterialParamOffset != -1 ? 1 : 0;
+	MaterialParamValidMask |= MaterialParamOffset1 != -1 ? 2 : 0;
+	MaterialParamValidMask |= MaterialParamOffset2 != -1 ? 4 : 0;
+	MaterialParamValidMask |= MaterialParamOffset3 != -1 ? 8 : 0;
 }
 
 
-void NiagaraRendererRibbons::ReleaseRenderThreadResources()
+void FNiagaraRendererRibbons::ReleaseRenderThreadResources()
 {
+	FNiagaraRenderer::ReleaseRenderThreadResources();
 	VertexFactory->ReleaseResource();
 	WorldSpacePrimitiveUniformBuffer.ReleaseResource();
 }
 
 // FPrimitiveSceneProxy interface.
-void NiagaraRendererRibbons::CreateRenderThreadResources()
+void FNiagaraRendererRibbons::CreateRenderThreadResources()
 {
+	FNiagaraRenderer::CreateRenderThreadResources();
 	VertexFactory->InitResource();
 }
 
-void NiagaraRendererRibbons::GenerateIndexBuffer(uint16* OutIndices, const TArray<int32>& SegmentData,  int32 InterpCount, bool bInvertOrder, bool bCullTwistedStrips)
+void FNiagaraRendererRibbons::GenerateIndexBuffer(uint16* OutIndices, const TArray<int32>& SegmentData,  int32 InterpCount, bool bInvertOrder, bool bCullTwistedStrips)
 {
 	auto AddTriangleIndices = [&OutIndices, InterpCount, bCullTwistedStrips](int32 SegmentIndex)
 	{
@@ -209,37 +264,48 @@ void NiagaraRendererRibbons::GenerateIndexBuffer(uint16* OutIndices, const TArra
 	}
 }
 
-void NiagaraRendererRibbons::GetDynamicMeshElements(const TArray<const FSceneView*>& Views, const FSceneViewFamily& ViewFamily, uint32 VisibilityMap, FMeshElementCollector& Collector, const FNiagaraSceneProxy *SceneProxy) const
+void FNiagaraRendererRibbons::GetDynamicMeshElements(const TArray<const FSceneView*>& Views, const FSceneViewFamily& ViewFamily, uint32 VisibilityMap, FMeshElementCollector& Collector, const FNiagaraSceneProxy *SceneProxy) const
 {
 	SCOPE_CYCLE_COUNTER(STAT_NiagaraRender);
 	SCOPE_CYCLE_COUNTER(STAT_NiagaraRenderRibbons);
 
 	SimpleTimer MeshElementsTimer;
 	FNiagaraDynamicDataRibbon *DynamicDataRibbon = static_cast<FNiagaraDynamicDataRibbon*>(DynamicDataRender);
-	if (!DynamicDataRibbon 
-		|| DynamicDataRibbon->SegmentData.Num() == 0
-		|| nullptr == Properties
-		|| !GSupportsResourceView // Current shader requires SRV to draw properly in all cases.
+	if (!DynamicDataRibbon)
+	{
+		return;
+	}
+
+	FNiagaraDataBuffer* SourceParticleData = DynamicDataRibbon->GetParticleDataToRender();
+	if( !SourceParticleData ||
+		SourceParticleData->GetNumInstances() < 2 ||
+		DynamicDataRibbon->SegmentData.Num() == 0	||
+		GbEnableNiagaraRibbonRendering == 0 ||
+		!GSupportsResourceView // Current shader requires SRV to draw properly in all cases.
 		)
 	{
 		return;
 	}
 
-	const bool bIsWireframe = ViewFamily.EngineShowFlags.Wireframe;
-	const bool bCullTwistedStrips = Properties->FacingMode == ENiagaraRibbonFacingMode::Screen && GNiagaraRibbonCullTwistedStrips != 0;
+#if STATS
+	FScopeCycleCounter EmitterStatsCounter(EmitterStatID);
+#endif
 
-	FMaterialRenderProxy* MaterialRenderProxy = Material->GetRenderProxy();
+	const bool bIsWireframe = ViewFamily.EngineShowFlags.Wireframe;
+	const bool bCullTwistedStrips = FacingMode == ENiagaraRibbonFacingMode::Screen && GNiagaraRibbonCullTwistedStrips != 0;
+	FMaterialRenderProxy* MaterialRenderProxy = DynamicDataRibbon->Material;
+	check(MaterialRenderProxy);
 
 	FGlobalDynamicIndexBuffer& DynamicIndexBuffer = Collector.GetDynamicIndexBuffer();
 
-	const int32 TotalFloatSize = DynamicDataRibbon->RTParticleData.GetFloatBuffer().Num() / sizeof(float);
+	int32 TotalFloatSize = SourceParticleData->GetFloatBuffer().Num() / sizeof(float);
 	FGlobalDynamicReadBuffer::FAllocation ParticleData;
 
-	if (DynamicDataRibbon->DataSet->GetSimTarget() == ENiagaraSimTarget::CPUSim)
+	if (DynamicDataRibbon->GetSimTarget() == ENiagaraSimTarget::CPUSim)
 	{
 		FGlobalDynamicReadBuffer& DynamicReadBuffer = Collector.GetDynamicReadBuffer();
 		ParticleData = DynamicReadBuffer.AllocateFloat(TotalFloatSize);
-		FMemory::Memcpy(ParticleData.Buffer, DynamicDataRibbon->RTParticleData.GetFloatBuffer().GetData(), DynamicDataRibbon->RTParticleData.GetFloatBuffer().Num());
+		FMemory::Memcpy(ParticleData.Buffer, SourceParticleData->GetFloatBuffer().GetData(), SourceParticleData->GetFloatBuffer().Num());
 	}
 
 	// Update the primitive uniform buffer if needed.
@@ -303,8 +369,8 @@ void NiagaraRendererRibbons::GetDynamicMeshElements(const TArray<const FSceneVie
 			// TODO : This doesn't work with multi-ribbons.
 			const float StartDist = FVector::DotProduct(View->GetViewDirection(), DynamicDataRibbon->StartPos - ViewOriginForDistanceCulling);
 			const float EndDist = FVector::DotProduct(View->GetViewDirection(), DynamicDataRibbon->EndPos - ViewOriginForDistanceCulling);
-			const bool bInvertOrder = ((StartDist > EndDist) && Properties->DrawDirection == ENiagaraRibbonDrawDirection::BackToFront)
-					|| ((StartDist < EndDist) && Properties->DrawDirection == ENiagaraRibbonDrawDirection::FrontToBack);
+			const bool bInvertOrder = ((StartDist > EndDist) && DrawDirection == ENiagaraRibbonDrawDirection::BackToFront)
+					|| ((StartDist < EndDist) && DrawDirection == ENiagaraRibbonDrawDirection::FrontToBack);
 
 			// Copy the vertex data over.
 			FGlobalDynamicIndexBuffer::FAllocation DynamicIndexAllocation = DynamicIndexBuffer.Allocate(NumberOfPrimitives * 3, sizeof(int16));
@@ -318,7 +384,7 @@ void NiagaraRendererRibbons::GetDynamicMeshElements(const TArray<const FSceneVie
 			PerViewUniformParameters.CameraUp = View->GetViewUp(); // FVector4(0.0f, 0.0f, 1.0f, 0.0f);
 			PerViewUniformParameters.CameraRight = View->GetViewRight();//	FVector4(1.0f, 0.0f, 0.0f, 0.0f);
 			PerViewUniformParameters.ScreenAlignment = FVector4(0.0f, 0.0f, 0.0f, 0.0f);
-			PerViewUniformParameters.TotalNumInstances = DynamicDataRibbon->RTParticleData.GetNumInstances();
+			PerViewUniformParameters.TotalNumInstances = SourceParticleData->GetNumInstances();
 			PerViewUniformParameters.InterpCount = SegmentTessellation;
 			PerViewUniformParameters.OneOverInterpCount = 1.f / (float)SegmentTessellation;
 
@@ -327,17 +393,18 @@ void NiagaraRendererRibbons::GetDynamicMeshElements(const TArray<const FSceneVie
 			PerViewUniformParameters.ColorDataOffset = ColorDataOffset;
 			PerViewUniformParameters.WidthDataOffset = WidthDataOffset;
 			PerViewUniformParameters.TwistDataOffset = TwistDataOffset;
-			PerViewUniformParameters.FacingDataOffset = Properties->FacingMode == ENiagaraRibbonFacingMode::Custom ? FacingDataOffset : -1;
+			PerViewUniformParameters.FacingDataOffset = FacingMode == ENiagaraRibbonFacingMode::Custom ? FacingDataOffset : -1;
 			PerViewUniformParameters.NormalizedAgeDataOffset = NormalizedAgeDataOffset;
 			PerViewUniformParameters.MaterialRandomDataOffset = MaterialRandomDataOffset;
+			PerViewUniformParameters.MaterialParamValidMask = MaterialParamValidMask;
 			PerViewUniformParameters.MaterialParamDataOffset = MaterialParamOffset;
 			PerViewUniformParameters.MaterialParam1DataOffset = MaterialParamOffset1;
 			PerViewUniformParameters.MaterialParam2DataOffset = MaterialParamOffset2;
 			PerViewUniformParameters.MaterialParam3DataOffset = MaterialParamOffset3;
-			PerViewUniformParameters.OneOverUV0TilingDistance = Properties->UV0TilingDistance ? 1.f / (Properties->UV0TilingDistance) : 0.f;
-			PerViewUniformParameters.OneOverUV1TilingDistance = Properties->UV1TilingDistance ? 1.f / (Properties->UV1TilingDistance) : 0.f;
-			PerViewUniformParameters.PackedVData = FVector4(Properties->UV0Scale.Y, Properties->UV0Offset.Y, Properties->UV1Scale.Y, Properties->UV1Offset.Y);
-			CollectorResources.VertexFactory.SetParticleData(ParticleData.ReadBuffer->SRV, ParticleData.FirstIndex / sizeof(float), DynamicDataRibbon->RTParticleData.GetFloatStride() / sizeof(float));
+			PerViewUniformParameters.OneOverUV0TilingDistance = UV0TilingDistance ? 1.f / (UV0TilingDistance) : 0.f;
+			PerViewUniformParameters.OneOverUV1TilingDistance = UV1TilingDistance ? 1.f / (UV1TilingDistance) : 0.f;
+			PerViewUniformParameters.PackedVData = FVector4(UV0Scale.Y, UV0Offset.Y, UV1Scale.Y, UV1Offset.Y);
+			CollectorResources.VertexFactory.SetParticleData(ParticleData.ReadBuffer->SRV, ParticleData.FirstIndex / sizeof(float), SourceParticleData->GetFloatStride() / sizeof(float));
 
 			// Collector.AllocateOneFrameResource uses default ctor, initialize the vertex factory
 			CollectorResources.VertexFactory.SetParticleFactoryType(NVFT_Ribbon);
@@ -423,19 +490,7 @@ void NiagaraRendererRibbons::GetDynamicMeshElements(const TArray<const FSceneVie
 	CPUTimeMS += MeshElementsTimer.GetElapsedMilliseconds();
 }
 
-void NiagaraRendererRibbons::SetDynamicData_RenderThread(FNiagaraDynamicDataBase* NewDynamicData)
-{
-	check(IsInRenderingThread());
-
-	if (DynamicDataRender)
-	{
-		delete static_cast<FNiagaraDynamicDataRibbon*>(DynamicDataRender);
-		DynamicDataRender = NULL;
-	}
-	DynamicDataRender = NewDynamicData;
-}
-
-int NiagaraRendererRibbons::GetDynamicDataSize()
+int FNiagaraRendererRibbons::GetDynamicDataSize()const
 {
 	uint32 Size = sizeof(FNiagaraDynamicDataRibbon);
 	if (DynamicDataRender)
@@ -451,32 +506,7 @@ int NiagaraRendererRibbons::GetDynamicDataSize()
 	return Size;
 }
 
-bool NiagaraRendererRibbons::HasDynamicData()
-{
-	return DynamicDataRender && (static_cast<FNiagaraDynamicDataRibbon*>(DynamicDataRender))->SegmentData.Num() > 0;
-}
-
-#if WITH_EDITORONLY_DATA
-
-const TArray<FNiagaraVariable>& NiagaraRendererRibbons::GetRequiredAttributes()
-{
-	return Properties->GetRequiredAttributes();
-}
-
-const TArray<FNiagaraVariable>& NiagaraRendererRibbons::GetOptionalAttributes()
-{
-	return Properties->GetOptionalAttributes();
-}
-
-#endif
-
-
-bool NiagaraRendererRibbons::SetMaterialUsage()
-{
-	return Material && Material->CheckMaterialUsage_Concurrent(MATUSAGE_NiagaraRibbons);
-}
-
-void NiagaraRendererRibbons::TransformChanged()
+void FNiagaraRendererRibbons::TransformChanged()
 {
 	WorldSpacePrimitiveUniformBuffer.ReleaseResource();
 }
@@ -535,24 +565,18 @@ void CalculateUVScaleAndOffsets(
 	}
 }
 
-FNiagaraDynamicDataBase* NiagaraRendererRibbons::GenerateVertexData(const FNiagaraSceneProxy* Proxy, FNiagaraDataSet &Data, const ENiagaraSimTarget Target)
+FNiagaraDynamicDataBase* FNiagaraRendererRibbons::GenerateDynamicData(const FNiagaraSceneProxy* Proxy, const UNiagaraRendererProperties* InProperties, const FNiagaraEmitterInstance* Emitter)const
 {
 	SCOPE_CYCLE_COUNTER(STAT_NiagaraGenRibbonVertexData);
 
-	SimpleTimer VertexDataTimer;
-	if (!bEnabled)
+	if (Emitter->GetData().GetSimTarget() == ENiagaraSimTarget::GPUComputeSim)
 	{
 		return nullptr;
 	}
-	FNiagaraDynamicDataRibbon* DynamicData = new FNiagaraDynamicDataRibbon;
-	TArray<int32>& SegmentData = DynamicData->SegmentData;
-	float TotalSegmentLength = 0;
-	float TotalSegmentAngle = 0;
-	float TotalTwistAngle = 0;
-	float TotalWidth = 0;
-	int NumTotalSamples = 0;
 
-	FNiagaraDataSetAccessor<FVector> PosData(Data, Properties->PositionBinding.DataSetVariable);
+	SimpleTimer VertexDataTimer;
+	FNiagaraDataSet& Data = Emitter->GetData();
+	const UNiagaraRibbonRendererProperties* Properties = CastChecked<const UNiagaraRibbonRendererProperties>(InProperties);
 
 	bool bSortKeyIsAge = false;
 	FNiagaraDataSetAccessor<float> SortKeyData(Data, Properties->RibbonLinkOrderBinding.DataSetVariable);
@@ -562,38 +586,7 @@ FNiagaraDynamicDataBase* NiagaraRendererRibbons::GenerateVertexData(const FNiaga
 		bSortKeyIsAge = true;
 	}
 
-	//Bail if we don't have the required attributes to render this emitter.
-	if (Data.GetNumInstances() < 2 || !PosData.IsValid() || !SortKeyData.IsValid())
-	{
-		return DynamicData;
-	}
-
-	if (PositionDataOffset == INDEX_NONE || LastSyncedId != Properties->SyncId)
-	{
-		// required attributes
-		int32 IntDummy;
-		Data.GetVariableComponentOffsets(Properties->PositionBinding.DataSetVariable, PositionDataOffset, IntDummy);
-		Data.GetVariableComponentOffsets(Properties->VelocityBinding.DataSetVariable, VelocityDataOffset, IntDummy);
-		Data.GetVariableComponentOffsets(Properties->ColorBinding.DataSetVariable, ColorDataOffset, IntDummy);
-
-		// optional attributes
-		Data.GetVariableComponentOffsets(Properties->RibbonWidthBinding.DataSetVariable, WidthDataOffset, IntDummy);
-		Data.GetVariableComponentOffsets(Properties->RibbonTwistBinding.DataSetVariable, TwistDataOffset, IntDummy);
-		Data.GetVariableComponentOffsets(Properties->RibbonFacingBinding.DataSetVariable, FacingDataOffset, IntDummy);
-		Data.GetVariableComponentOffsets(Properties->NormalizedAgeBinding.DataSetVariable, NormalizedAgeDataOffset, IntDummy);
-		Data.GetVariableComponentOffsets(Properties->MaterialRandomBinding.DataSetVariable, MaterialRandomDataOffset, IntDummy);
-
-		Data.GetVariableComponentOffsets(Properties->DynamicMaterialBinding.DataSetVariable, MaterialParamOffset, IntDummy);
-		Data.GetVariableComponentOffsets(Properties->DynamicMaterial1Binding.DataSetVariable, MaterialParamOffset1, IntDummy);
-		Data.GetVariableComponentOffsets(Properties->DynamicMaterial2Binding.DataSetVariable, MaterialParamOffset2, IntDummy);
-		Data.GetVariableComponentOffsets(Properties->DynamicMaterial3Binding.DataSetVariable, MaterialParamOffset3, IntDummy);
-
-		LastSyncedId = Properties->SyncId;
-	}
-
-	DynamicData->DataSet = &Data;
-
-	////////
+	FNiagaraDataSetAccessor<FVector> PosData(Data, Properties->PositionBinding.DataSetVariable);
 	FNiagaraDataSetAccessor<float> SizeData(Data, Properties->RibbonWidthBinding.DataSetVariable);
 	FNiagaraDataSetAccessor<float> TwistData(Data, Properties->RibbonTwistBinding.DataSetVariable);
 	FNiagaraDataSetAccessor<FVector> AlignData(Data, Properties->RibbonFacingBinding.DataSetVariable);
@@ -604,15 +597,40 @@ FNiagaraDynamicDataBase* NiagaraRendererRibbons::GenerateVertexData(const FNiaga
 
 	FNiagaraDataSetAccessor<int32> RibbonIdData;
 	FNiagaraDataSetAccessor<FNiagaraID> RibbonFullIDData;
+
+	ENiagaraSimTarget SimTarget = Data.GetSimTarget();
+	FNiagaraDataBuffer* DataToRender = Emitter->GetData().GetCurrentData();
+	if (DataToRender == nullptr || DataToRender->GetNumInstances() < 2 || !PosData.IsValid() || !SortKeyData.IsValid())
+	{
+		return nullptr;
+	}
+
+	FNiagaraDynamicDataRibbon* DynamicData = new FNiagaraDynamicDataRibbon(Emitter);
+
+	//In preparation for a material override feature, we pass our material(s) and relevance in via dynamic data.
+	//The renderer ensures we have the correct usage and relevance for materials in BaseMaterials_GT.
+	//Any override feature must also do the same for materials that are set.
+	check(BaseMaterials_GT.Num() == 1);
+	check(BaseMaterials_GT[0]->CheckMaterialUsage_Concurrent(MATUSAGE_NiagaraRibbons));
+	DynamicData->Material = BaseMaterials_GT[0]->GetRenderProxy();
+	DynamicData->SetMaterialRelevance(BaseMaterialRelevance_GT);
+
+	TArray<int32>& SegmentData = DynamicData->SegmentData;
+	float TotalSegmentLength = 0;
+	float TotalSegmentAngle = 0;
+	float TotalTwistAngle = 0;
+	float TotalWidth = 0;
+	int NumTotalSamples = 0;
+
 	if (Properties->RibbonIdBinding.DataSetVariable.GetType() == FNiagaraTypeDefinition::GetIDDef())
 	{
 		RibbonFullIDData.Create(&Data, Properties->RibbonIdBinding.DataSetVariable);
-		RibbonFullIDData.InitForAccess(true);
+		RibbonFullIDData.InitForAccess();
 	}
 	else
 	{
 		RibbonIdData.Create(&Data, Properties->RibbonIdBinding.DataSetVariable);
-		RibbonIdData.InitForAccess(true);
+		RibbonIdData.InitForAccess();
 	}
 
 	bool bFullIDs = RibbonFullIDData.IsValid();
@@ -768,10 +786,10 @@ FNiagaraDynamicDataBase* NiagaraRendererRibbons::GenerateVertexData(const FNiaga
 
 	// store the start and end positions for the ribbon for draw distance flipping 
 	DynamicData->StartPos = PosData[0];
-	DynamicData->EndPos = PosData[Data.GetNumInstances() - 1];
+	DynamicData->EndPos = PosData[Data.GetCurrentDataChecked().GetNumInstances() - 1];
 	
 	//TODO: Move sorting to share code with sprite and mesh sorting and support the custom sorting key.
-	int32 TotalIndices = Data.GetNumInstances();
+	int32 TotalIndices = Data.GetCurrentDataChecked().GetNumInstances();
 
 	if (!bMultiRibbons)
 	{
@@ -859,16 +877,23 @@ FNiagaraDynamicDataBase* NiagaraRendererRibbons::GenerateVertexData(const FNiaga
 		TessellationTwistAngle = 0;
 		TessellationTwistCurvature = 0;
 	}
-	 
-	if (Data.CurrData().GetNumInstances() > 0)
-	{
-		//TODO: This buffer is far fatter than needed. Just pull out the data needed for rendering.
-		Data.CurrData().CopyTo(DynamicData->RTParticleData);
-
-		DynamicData->DataSet = &Data;
-	}
 
 	CPUTimeMS = VertexDataTimer.GetElapsedMilliseconds();
 
 	return DynamicData;
 }
+
+void FNiagaraRendererRibbons::AddDynamicParam(TArray<FNiagaraRibbonVertexDynamicParameter>& ParamData, const FVector4& DynamicParam)
+{
+	FNiagaraRibbonVertexDynamicParameter Param;
+	Param.DynamicValue[0] = DynamicParam.X;
+	Param.DynamicValue[1] = DynamicParam.Y;
+	Param.DynamicValue[2] = DynamicParam.Z;
+	Param.DynamicValue[3] = DynamicParam.W;
+	ParamData.Add(Param);
+}
+
+bool FNiagaraRendererRibbons::IsMaterialValid(UMaterialInterface* Mat)const
+{
+	return Mat && Mat->CheckMaterialUsage(MATUSAGE_NiagaraRibbons);
+}
diff --git a/Engine/Plugins/FX/Niagara/Source/Niagara/Private/NiagaraRendererSprites.cpp b/Engine/Plugins/FX/Niagara/Source/Niagara/Private/NiagaraRendererSprites.cpp
--- a/Engine/Plugins/FX/Niagara/Source/Niagara/Private/NiagaraRendererSprites.cpp
+++ b/Engine/Plugins/FX/Niagara/Source/Niagara/Private/NiagaraRendererSprites.cpp
@@ -1,6 +1,6 @@
 // Copyright 1998-2019 Epic Games, Inc. All Rights Reserved.
 
-#include "NiagaraRenderer.h"
+#include "NiagaraRendererSprites.h"
 #include "ParticleResources.h"
 #include "NiagaraSpriteVertexFactory.h"
 #include "NiagaraDataSet.h"
@@ -12,20 +12,37 @@
 #include "RayTracingDynamicGeometryCollection.h"
 #include "RayTracingInstance.h"
 
-DECLARE_CYCLE_STAT(TEXT("Generate Sprite Vertex Data [GT]"), STAT_NiagaraGenSpriteVertexData, STATGROUP_Niagara);
+DECLARE_CYCLE_STAT(TEXT("Generate Sprite Dynamic Data [GT]"), STAT_NiagaraGenSpriteDynamicData, STATGROUP_Niagara);
 DECLARE_CYCLE_STAT(TEXT("Render Sprites [RT]"), STAT_NiagaraRenderSprites, STATGROUP_Niagara);
+DECLARE_CYCLE_STAT(TEXT("Render Sprites - CPU Sim Copy[RT]"), STAT_NiagaraRenderSpritesCPUSimCopy, STATGROUP_Niagara);
+DECLARE_CYCLE_STAT(TEXT("Render Sprites - CPU Sim Memcopy[RT]"), STAT_NiagaraRenderSpritesCPUSimMemCopy, STATGROUP_Niagara);
+DECLARE_CYCLE_STAT(TEXT("Render Sprites - Cutout[RT]"), STAT_NiagaraRenderSpritesCutout, STATGROUP_Niagara);
+DECLARE_CYCLE_STAT(TEXT("Render Sprites - Sorting[RT]"), STAT_NiagaraRenderSpritesSorting, STATGROUP_Niagara);
+DECLARE_CYCLE_STAT(TEXT("Render Sprites - GlobalSortCPU[RT]"), STAT_NiagaraRenderSpritesGlobalSortCPU, STATGROUP_Niagara);
 
 DECLARE_CYCLE_STAT(TEXT("Genereate GPU Buffers"), STAT_NiagaraGenSpriteGpuBuffers, STATGROUP_Niagara);
+DECLARE_DWORD_COUNTER_STAT(TEXT("NumSprites"), STAT_NiagaraNumSprites, STATGROUP_Niagara);
 
+static int32 GbEnableNiagaraSpriteRendering = 1;
+static FAutoConsoleVariableRef CVarEnableNiagaraSpriteRendering(
+	TEXT("fx.EnableNiagaraSpriteRendering"),
+	GbEnableNiagaraSpriteRendering,
+	TEXT("If == 0, Niagara Sprite Renderers are disabled. \n"),
+	ECVF_Default
+);
+
+/** Dynamic data for sprite renderers. */
 struct FNiagaraDynamicDataSprites : public FNiagaraDynamicDataBase
 {
-	//Direct ptr to the dataset. ONLY FOR USE BE GPU EMITTERS.
-	//TODO: Even this needs to go soon.
-	const FNiagaraDataSet *DataSet;
+	FNiagaraDynamicDataSprites(const FNiagaraEmitterInstance* InEmitter)
+		: FNiagaraDynamicDataBase(InEmitter)
+		, Material(nullptr)
+	{
+	}
+	
+	FMaterialRenderProxy* Material;
 };
 
-
-
 /* Mesh collector classes */
 class FNiagaraMeshCollectorResourcesSprite : public FOneFrameResource
 {
@@ -39,41 +56,91 @@ public:
 	}
 };
 
-NiagaraRendererSprites::NiagaraRendererSprites(ERHIFeatureLevel::Type FeatureLevel, UNiagaraRendererProperties *InProps) 
-	: NiagaraRenderer()
+
+//////////////////////////////////////////////////////////////////////////
+
+FNiagaraRendererSprites::FNiagaraRendererSprites(ERHIFeatureLevel::Type FeatureLevel, const UNiagaraRendererProperties *InProps, const FNiagaraEmitterInstance* Emitter)
+	: FNiagaraRenderer(FeatureLevel, InProps, Emitter)
+	, Alignment(ENiagaraSpriteAlignment::Unaligned)
+	, FacingMode(ENiagaraSpriteFacingMode::FaceCamera)
+	, CustomFacingVectorMask(ForceInitToZero)
+	, PivotInUVSpace(0.5f, 0.5f)
+	, SortMode(ENiagaraSortMode::ViewDistance)
+	, SubImageSize(1.0f, 1.0f)
+	, bSubImageBlend(false)
+	, bRemoveHMDRollInVR(false)
+	, bSortOnlyWhenTranslucent(true)
+	, MinFacingCameraBlendDistance(0.0f)
+	, MaxFacingCameraBlendDistance(0.0f)
 	, PositionOffset(INDEX_NONE)
+	, ColorOffset(INDEX_NONE)
 	, VelocityOffset(INDEX_NONE)
 	, RotationOffset(INDEX_NONE)
 	, SizeOffset(INDEX_NONE)
-	, ColorOffset(INDEX_NONE)
 	, FacingOffset(INDEX_NONE)
 	, AlignmentOffset(INDEX_NONE)
 	, SubImageOffset(INDEX_NONE)
+	, MaterialParamValidMask(0)
 	, MaterialParamOffset(INDEX_NONE)
 	, MaterialParamOffset1(INDEX_NONE)
 	, MaterialParamOffset2(INDEX_NONE)
 	, MaterialParamOffset3(INDEX_NONE)
 	, CameraOffsetOffset(INDEX_NONE)
 	, UVScaleOffset(INDEX_NONE)
-	, NormalizedAgeOffset(INDEX_NONE)
 	, MaterialRandomOffset(INDEX_NONE)
 	, CustomSortingOffset(INDEX_NONE)
-	, LastSyncId(INDEX_NONE)
+	, NormalizedAgeOffset(INDEX_NONE)
 {
-	//check(InProps);
+	check(InProps && Emitter);
+
 	VertexFactory = new FNiagaraSpriteVertexFactory(NVFT_Sprite, FeatureLevel);
-	Properties = Cast<UNiagaraSpriteRendererProperties>(InProps);
-	BaseExtents = FVector(0.5f, 0.5f, 0.5f);
 
-	if (Properties)
-	{
-		NumCutoutVertexPerSubImage = Properties->GetNumCutoutVertexPerSubimage();
-		CutoutVertexBuffer.Data = Properties->GetCutoutData();
-	}
+	const UNiagaraSpriteRendererProperties* Properties = CastChecked<const UNiagaraSpriteRendererProperties>(InProps);
+	
+	Alignment = Properties->Alignment;
+	FacingMode = Properties->FacingMode;
+	CustomFacingVectorMask = Properties->CustomFacingVectorMask;
+	PivotInUVSpace = Properties->PivotInUVSpace;
+	SortMode = Properties->SortMode;
+	SubImageSize = Properties->SubImageSize;
+	bSubImageBlend = Properties->bSubImageBlend;
+	bRemoveHMDRollInVR = Properties->bRemoveHMDRollInVR;
+	bSortOnlyWhenTranslucent = Properties->bSortOnlyWhenTranslucent;
+	MinFacingCameraBlendDistance = Properties->MinFacingCameraBlendDistance;
+	MaxFacingCameraBlendDistance = Properties->MaxFacingCameraBlendDistance;
+
+	NumCutoutVertexPerSubImage = Properties->GetNumCutoutVertexPerSubimage();
+	CutoutVertexBuffer.Data = Properties->GetCutoutData();
+
+	int32 IntDummy;
+	const FNiagaraDataSet& Data = Emitter->GetData();
+	Data.GetVariableComponentOffsets(Properties->PositionBinding.DataSetVariable, PositionOffset, IntDummy);
+	Data.GetVariableComponentOffsets(Properties->VelocityBinding.DataSetVariable, VelocityOffset, IntDummy);
+	Data.GetVariableComponentOffsets(Properties->SpriteRotationBinding.DataSetVariable, RotationOffset, IntDummy);
+	Data.GetVariableComponentOffsets(Properties->SpriteSizeBinding.DataSetVariable, SizeOffset, IntDummy);
+	Data.GetVariableComponentOffsets(Properties->ColorBinding.DataSetVariable, ColorOffset, IntDummy);
+	Data.GetVariableComponentOffsets(Properties->SpriteFacingBinding.DataSetVariable, FacingOffset, IntDummy);
+	Data.GetVariableComponentOffsets(Properties->SpriteAlignmentBinding.DataSetVariable, AlignmentOffset, IntDummy);
+	Data.GetVariableComponentOffsets(Properties->SubImageIndexBinding.DataSetVariable, SubImageOffset, IntDummy);
+	Data.GetVariableComponentOffsets(Properties->DynamicMaterialBinding.DataSetVariable, MaterialParamOffset, IntDummy);
+	Data.GetVariableComponentOffsets(Properties->DynamicMaterial1Binding.DataSetVariable, MaterialParamOffset1, IntDummy);
+	Data.GetVariableComponentOffsets(Properties->DynamicMaterial2Binding.DataSetVariable, MaterialParamOffset2, IntDummy);
+	Data.GetVariableComponentOffsets(Properties->DynamicMaterial3Binding.DataSetVariable, MaterialParamOffset3, IntDummy);
+	Data.GetVariableComponentOffsets(Properties->CameraOffsetBinding.DataSetVariable, CameraOffsetOffset, IntDummy);
+	Data.GetVariableComponentOffsets(Properties->UVScaleBinding.DataSetVariable, UVScaleOffset, IntDummy);
+	Data.GetVariableComponentOffsets(Properties->NormalizedAgeBinding.DataSetVariable, NormalizedAgeOffset, IntDummy);
+	Data.GetVariableComponentOffsets(Properties->MaterialRandomBinding.DataSetVariable, MaterialRandomOffset, IntDummy);
+	Data.GetVariableComponentOffsets(Properties->CustomSortingBinding.DataSetVariable, CustomSortingOffset, IntDummy);
+
+	MaterialParamValidMask = MaterialParamOffset != -1 ? 1 : 0;
+	MaterialParamValidMask |= MaterialParamOffset1 != -1 ? 2 : 0;
+	MaterialParamValidMask |= MaterialParamOffset2 != -1 ? 4 : 0;
+	MaterialParamValidMask |= MaterialParamOffset3 != -1 ? 8 : 0;
 }
 
-void NiagaraRendererSprites::ReleaseRenderThreadResources()
+void FNiagaraRendererSprites::ReleaseRenderThreadResources()
 {
+	FNiagaraRenderer::ReleaseRenderThreadResources();
 	VertexFactory->ReleaseResource();
 	WorldSpacePrimitiveUniformBuffer.ReleaseResource();
 
@@ -87,8 +154,9 @@ void NiagaraRendererSprites::ReleaseRenderThreadResources()
 #endif
 }
 
-void NiagaraRendererSprites::CreateRenderThreadResources()
+void FNiagaraRendererSprites::CreateRenderThreadResources()
 {
+	FNiagaraRenderer::CreateRenderThreadResources();
 	CutoutVertexBuffer.InitResource();
 	VertexFactory->InitResource();
 
@@ -114,7 +182,7 @@ void NiagaraRendererSprites::CreateRenderThreadResources()
 #endif
 }
 
-void NiagaraRendererSprites::ConditionalInitPrimitiveUniformBuffer(const FNiagaraSceneProxy *SceneProxy) const
+void FNiagaraRendererSprites::ConditionalInitPrimitiveUniformBuffer(const FNiagaraSceneProxy *SceneProxy) const
 {
 	// Update the primitive uniform buffer if needed.
 	if (!WorldSpacePrimitiveUniformBuffer.IsInitialized())
@@ -142,22 +210,28 @@ void NiagaraRendererSprites::ConditionalInitPrimitiveUniformBuffer(const FNiagar
 	}
 }
 
-NiagaraRendererSprites::FCPUSimParticleDataAllocation NiagaraRendererSprites::ConditionalAllocateCPUSimParticleData(FNiagaraDynamicDataSprites* DynamicDataSprites, FGlobalDynamicReadBuffer& DynamicReadBuffer) const
+FNiagaraRendererSprites::FCPUSimParticleDataAllocation FNiagaraRendererSprites::ConditionalAllocateCPUSimParticleData(FNiagaraDynamicDataSprites* DynamicDataSprites, FGlobalDynamicReadBuffer& DynamicReadBuffer) const
 {
-	int32 TotalFloatSize = DynamicDataSprites->RTParticleData.GetFloatBuffer().Num() / sizeof(float);
+	FNiagaraDataBuffer* SourceParticleData = DynamicDataSprites->GetParticleDataToRender();
+	check(SourceParticleData);//Can be null but should be checked before here.
+	int32 TotalFloatSize = SourceParticleData->GetFloatBuffer().Num() / sizeof(float);
 
 	FCPUSimParticleDataAllocation CPUSimParticleDataAllocation { DynamicReadBuffer };
 
-	if (DynamicDataSprites->DataSet->GetSimTarget() == ENiagaraSimTarget::CPUSim)
+	if (DynamicDataSprites->GetSimTarget() == ENiagaraSimTarget::CPUSim)
 	{
+		SCOPE_CYCLE_COUNTER(STAT_NiagaraRenderSpritesCPUSimCopy);
 		CPUSimParticleDataAllocation.ParticleData = DynamicReadBuffer.AllocateFloat(TotalFloatSize);
-		FMemory::Memcpy(CPUSimParticleDataAllocation.ParticleData.Buffer, DynamicDataSprites->RTParticleData.GetFloatBuffer().GetData(), DynamicDataSprites->RTParticleData.GetFloatBuffer().Num());
+		{
+			SCOPE_CYCLE_COUNTER(STAT_NiagaraRenderSpritesCPUSimMemCopy);
+			FMemory::Memcpy(CPUSimParticleDataAllocation.ParticleData.Buffer, SourceParticleData->GetFloatBuffer().GetData(), SourceParticleData->GetFloatBuffer().Num());
+		}
 	}
 
 	return CPUSimParticleDataAllocation;
 }
 
-FNiagaraSpriteUniformBufferRef NiagaraRendererSprites::CreatePerViewUniformBuffer(const FSceneView* View, const FSceneViewFamily& ViewFamily, const FNiagaraSceneProxy *SceneProxy) const
+FNiagaraSpriteUniformBufferRef FNiagaraRendererSprites::CreatePerViewUniformBuffer(const FSceneView* View, const FSceneViewFamily& ViewFamily, const FNiagaraSceneProxy *SceneProxy) const
 {
 	FNiagaraSpriteUniformParameters PerViewUniformParameters;
 	PerViewUniformParameters.LocalToWorld = bLocalSpace ? SceneProxy->GetLocalToWorld() : FMatrix::Identity;//For now just handle local space like this but maybe in future have a VF variant to avoid the transform entirely?
@@ -169,17 +243,18 @@ FNiagaraSpriteUniformBufferRef NiagaraRendererSprites::CreatePerViewUniformBuffe
 	PerViewUniformParameters.NormalsType = 0.0f;
 	PerViewUniformParameters.NormalsSphereCenter = FVector4(0.0f, 0.0f, 0.0f, 1.0f);
 	PerViewUniformParameters.NormalsCylinderUnitDirection = FVector4(0.0f, 0.0f, 1.0f, 0.0f);
-	PerViewUniformParameters.PivotOffset = Properties->PivotInUVSpace * -1.0f; // We do this because we want to slide the coordinates back since 0,0 is the upper left corner.
+	PerViewUniformParameters.PivotOffset = PivotInUVSpace * -1.0f; // We do this because we want to slide the coordinates back since 0,0 is the upper left corner.
 	PerViewUniformParameters.MacroUVParameters = FVector4(0.0f, 0.0f, 1.0f, 1.0f);
 	PerViewUniformParameters.CameraFacingBlend = FVector4(0.0f, 0.0f, 0.0f, 1.0f);
-	PerViewUniformParameters.RemoveHMDRoll = Properties->bRemoveHMDRollInVR;
+	PerViewUniformParameters.RemoveHMDRoll = bRemoveHMDRollInVR;
 	PerViewUniformParameters.CustomFacingVectorMask = FVector4(0.0f, 0.0f, 0.0f, 0.0f);
-	PerViewUniformParameters.SubImageSize = FVector4(Properties->SubImageSize.X, Properties->SubImageSize.Y, 1.0f / Properties->SubImageSize.X, 1.0f / Properties->SubImageSize.Y);
+	PerViewUniformParameters.SubImageSize = FVector4(SubImageSize.X, SubImageSize.Y, 1.0f / SubImageSize.X, 1.0f / SubImageSize.Y);
 	PerViewUniformParameters.PositionDataOffset = PositionOffset;
 	PerViewUniformParameters.VelocityDataOffset = VelocityOffset;
 	PerViewUniformParameters.RotationDataOffset = RotationOffset;
 	PerViewUniformParameters.SizeDataOffset = SizeOffset;
 	PerViewUniformParameters.ColorDataOffset = ColorOffset;
+	PerViewUniformParameters.MaterialParamValidMask = MaterialParamValidMask;
 	PerViewUniformParameters.MaterialParamDataOffset = MaterialParamOffset;
 	PerViewUniformParameters.MaterialParam1DataOffset = MaterialParamOffset1;
 	PerViewUniformParameters.MaterialParam2DataOffset = MaterialParamOffset2;
@@ -187,30 +262,30 @@ FNiagaraSpriteUniformBufferRef NiagaraRendererSprites::CreatePerViewUniformBuffe
 	PerViewUniformParameters.SubimageDataOffset = SubImageOffset;
 	PerViewUniformParameters.FacingDataOffset = FacingOffset;
 	PerViewUniformParameters.AlignmentDataOffset = AlignmentOffset;
-	PerViewUniformParameters.SubImageBlendMode = Properties->bSubImageBlend;
+	PerViewUniformParameters.SubImageBlendMode = bSubImageBlend;
 	PerViewUniformParameters.CameraOffsetDataOffset = CameraOffsetOffset;
 	PerViewUniformParameters.UVScaleDataOffset = UVScaleOffset;
 	PerViewUniformParameters.NormalizedAgeDataOffset = NormalizedAgeOffset;
 	PerViewUniformParameters.MaterialRandomDataOffset = MaterialRandomOffset;
 	PerViewUniformParameters.DefaultPos = bLocalSpace ? FVector4(0.0f, 0.0f, 0.0f, 1.0f) : FVector4(SceneProxy->GetLocalToWorld().GetOrigin());
 
-	ENiagaraSpriteFacingMode FacingMode = Properties->FacingMode;
-	ENiagaraSpriteAlignment AlignmentMode = Properties->Alignment;
+	ENiagaraSpriteFacingMode ActualFacingMode = FacingMode;
+	ENiagaraSpriteAlignment ActualAlignmentMode = Alignment;
 
 	if (FacingOffset == -1 && FacingMode == ENiagaraSpriteFacingMode::CustomFacingVector)
 	{
-		FacingMode = ENiagaraSpriteFacingMode::FaceCamera;
+		ActualFacingMode = ENiagaraSpriteFacingMode::FaceCamera;
 	}
 
-	if (AlignmentOffset == -1 && AlignmentMode == ENiagaraSpriteAlignment::CustomAlignment)
+	if (AlignmentOffset == -1 && ActualAlignmentMode == ENiagaraSpriteAlignment::CustomAlignment)
 	{
-		AlignmentMode = ENiagaraSpriteAlignment::Unaligned;
+		ActualAlignmentMode = ENiagaraSpriteAlignment::Unaligned;
 	}
 
-	if (FacingMode == ENiagaraSpriteFacingMode::FaceCameraDistanceBlend)
+	if (ActualFacingMode == ENiagaraSpriteFacingMode::FaceCameraDistanceBlend)
 	{
-		float DistanceBlendMinSq = Properties->MinFacingCameraBlendDistance * Properties->MinFacingCameraBlendDistance;
-		float DistanceBlendMaxSq = Properties->MaxFacingCameraBlendDistance * Properties->MaxFacingCameraBlendDistance;
+		float DistanceBlendMinSq = MinFacingCameraBlendDistance * MinFacingCameraBlendDistance;
+		float DistanceBlendMaxSq = MaxFacingCameraBlendDistance * MaxFacingCameraBlendDistance;
 		float InvBlendRange = 1.0f / FMath::Max(DistanceBlendMaxSq - DistanceBlendMinSq, 1.0f);
 		float BlendScaledMinDistance = DistanceBlendMinSq * InvBlendRange;
 
@@ -219,22 +294,22 @@ FNiagaraSpriteUniformBufferRef NiagaraRendererSprites::CreatePerViewUniformBuffe
 		PerViewUniformParameters.CameraFacingBlend.Z = BlendScaledMinDistance;
 	}
 
-	if (Properties->Alignment == ENiagaraSpriteAlignment::VelocityAligned)
+	if (ActualAlignmentMode == ENiagaraSpriteAlignment::VelocityAligned)
 	{
 		// velocity aligned
 		PerViewUniformParameters.RotationScale = 0.0f;
 		PerViewUniformParameters.TangentSelector = FVector4(0.0f, 1.0f, 0.0f, 0.0f);
 	}
 
-	if (Properties->FacingMode == ENiagaraSpriteFacingMode::CustomFacingVector)
+	if (ActualFacingMode == ENiagaraSpriteFacingMode::CustomFacingVector)
 	{
-		PerViewUniformParameters.CustomFacingVectorMask = Properties->CustomFacingVectorMask;
+		PerViewUniformParameters.CustomFacingVectorMask = CustomFacingVectorMask;
 	}
 
 	return FNiagaraSpriteUniformBufferRef::CreateUniformBufferImmediate(PerViewUniformParameters, UniformBuffer_SingleFrame);
 }
 
-void NiagaraRendererSprites::SetVertexFactoryParticleData(
+void FNiagaraRendererSprites::SetVertexFactoryParticleData(
 	FNiagaraSpriteVertexFactory& OutVertexFactory, 
 	FNiagaraDynamicDataSprites* DynamicDataSprites, 
 	FCPUSimParticleDataAllocation& CPUSimParticleDataAllocation,
@@ -242,10 +317,11 @@ void NiagaraRendererSprites::SetVertexFactoryParticleData(
 	const FNiagaraSceneProxy *SceneProxy) const
 {
 	// Cutout geometry.
-	const bool bUseSubImage = Properties->SubImageSize.X != 1 || Properties->SubImageSize.Y != 1;
+	const bool bUseSubImage = SubImageSize.X != 1 || SubImageSize.Y != 1;
 	const bool bUseCutout = CutoutVertexBuffer.VertexBufferRHI.IsValid();
 	if (bUseCutout)
 	{	// Is Accessing Properties safe here? Or should values be cached in the constructor?
+		SCOPE_CYCLE_COUNTER(STAT_NiagaraRenderSpritesCutout);
 		if (bUseSubImage)
 		{
 			OutVertexFactory.SetCutoutParameters(NumCutoutVertexPerSubImage, CutoutVertexBuffer.VertexBufferSRV);
@@ -256,79 +332,93 @@ void NiagaraRendererSprites::SetVertexFactoryParticleData(
 		}
 	}
 
-	//Sort particles if needed.
-	FMaterialRenderProxy* MaterialRenderProxy = Material->GetRenderProxy();
-	EBlendMode BlendMode = MaterialRenderProxy->GetMaterial(VertexFactory->GetFeatureLevel())->GetBlendMode();
-	OutVertexFactory.SetSortedIndices(nullptr, 0xFFFFFFFF);
+	FNiagaraDataBuffer* SourceParticleData = DynamicDataSprites->GetParticleDataToRender();
+	check(SourceParticleData);//Can be null but should be checked before here.
 
-	int32 NumInstances = DynamicDataSprites->RTParticleData.GetNumInstances();
-	FNiagaraGPUSortInfo SortInfo;
-	if (View && Properties->SortMode != ENiagaraSortMode::None && (BlendMode == BLEND_AlphaComposite || BlendMode == BLEND_Translucent || !Properties->bSortOnlyWhenTranslucent))
+	//Sort particles if needed.
 	{
-		SortInfo.ParticleCount = NumInstances;
-		SortInfo.SortMode = Properties->SortMode;
-		SortInfo.SortAttributeOffset = (SortInfo.SortMode == ENiagaraSortMode::CustomAscending || SortInfo.SortMode == ENiagaraSortMode::CustomDecending) ? CustomSortingOffset : PositionOffset;
-		SortInfo.ViewOrigin = View->ViewMatrices.GetViewOrigin();
-		SortInfo.ViewDirection = View->GetViewDirection();
-		if (bLocalSpace)
+		SCOPE_CYCLE_COUNTER(STAT_NiagaraRenderSpritesSorting)
+
+
+		FMaterialRenderProxy* MaterialRenderProxy = DynamicDataSprites->Material;
+		check(MaterialRenderProxy);
+		EBlendMode BlendMode = MaterialRenderProxy->GetMaterial(VertexFactory->GetFeatureLevel())->GetBlendMode();
+		OutVertexFactory.SetSortedIndices(nullptr, 0xFFFFFFFF);
+
+		int32 NumInstances = SourceParticleData->GetNumInstances();
+		FNiagaraGPUSortInfo SortInfo;
+		if (View && SortMode != ENiagaraSortMode::None && (BlendMode == BLEND_AlphaComposite || BlendMode == BLEND_Translucent || !bSortOnlyWhenTranslucent))
 		{
-			FMatrix InvTransform = SceneProxy->GetLocalToWorld().InverseFast();
-			SortInfo.ViewOrigin = InvTransform.TransformPosition(SortInfo.ViewOrigin);
-			SortInfo.ViewDirection = InvTransform.TransformVector(SortInfo.ViewDirection);
-		}
-	};
+			SortInfo.ParticleCount = NumInstances;
+			SortInfo.SortMode = SortMode;
+			SortInfo.SortAttributeOffset = (SortInfo.SortMode == ENiagaraSortMode::CustomAscending || SortInfo.SortMode == ENiagaraSortMode::CustomDecending) ? CustomSortingOffset : PositionOffset;
+			SortInfo.ViewOrigin = View->ViewMatrices.GetViewOrigin();
+			SortInfo.ViewDirection = View->GetViewDirection();
+			if (bLocalSpace)
+			{
+				FMatrix InvTransform = SceneProxy->GetLocalToWorld().InverseFast();
+				SortInfo.ViewOrigin = InvTransform.TransformPosition(SortInfo.ViewOrigin);
+				SortInfo.ViewDirection = InvTransform.TransformVector(SortInfo.ViewDirection);
+			}
+		};
 
-	if (DynamicDataSprites->DataSet->GetSimTarget() == ENiagaraSimTarget::CPUSim)//TODO: Compute shader for sorting gpu sims and larger cpu sims.
-	{
-		check(CPUSimParticleDataAllocation.ParticleData.IsValid());
-		if (SortInfo.SortMode != ENiagaraSortMode::None && SortInfo.SortAttributeOffset != INDEX_NONE)
+		if (DynamicDataSprites->GetSimTarget() == ENiagaraSimTarget::CPUSim)//TODO: Compute shader for sorting gpu sims and larger cpu sims.
 		{
-			if (GNiagaraGPUSorting &&
-				GNiagaraGPUSortingCPUToGPUThreshold != INDEX_NONE &&
-				SortInfo.ParticleCount >= GNiagaraGPUSortingCPUToGPUThreshold &&
-				SceneProxy->GetBatcher())
+			check(CPUSimParticleDataAllocation.ParticleData.IsValid());
+			if (SortInfo.SortMode != ENiagaraSortMode::None && SortInfo.SortAttributeOffset != INDEX_NONE)
 			{
-				SortInfo.ParticleDataFloatSRV = CPUSimParticleDataAllocation.ParticleData.ReadBuffer->SRV;
-				SortInfo.FloatDataOffset = CPUSimParticleDataAllocation.ParticleData.FirstIndex / sizeof(float);
-				SortInfo.FloatDataStride = DynamicDataSprites->RTParticleData.GetFloatStride() / sizeof(float);
-				const int32 IndexBufferOffset = SceneProxy->GetBatcher()->AddSortedGPUSimulation(SortInfo);
-				if (IndexBufferOffset != INDEX_NONE)
+				if (GNiagaraGPUSorting &&
+					GNiagaraGPUSortingCPUToGPUThreshold != INDEX_NONE &&
+					SortInfo.ParticleCount >= GNiagaraGPUSortingCPUToGPUThreshold &&
+					SceneProxy->GetBatcher())
 				{
-					OutVertexFactory.SetSortedIndices(SceneProxy->GetBatcher()->GetGPUSortedBuffer().VertexBufferSRV, IndexBufferOffset);
+					SortInfo.ParticleCount = NumInstances;
+					SortInfo.ParticleDataFloatSRV = CPUSimParticleDataAllocation.ParticleData.ReadBuffer->SRV;
+					SortInfo.FloatDataOffset = CPUSimParticleDataAllocation.ParticleData.FirstIndex / sizeof(float);
+					SortInfo.FloatDataStride = SourceParticleData->GetFloatStride() / sizeof(float);
+					const int32 IndexBufferOffset = SceneProxy->GetBatcher()->AddSortedGPUSimulation(SortInfo);
+					if (IndexBufferOffset != INDEX_NONE)
+					{
+						OutVertexFactory.SetSortedIndices(SceneProxy->GetBatcher()->GetGPUSortedBuffer().VertexBufferSRV, IndexBufferOffset);
+					}
+				}
+				else
+				{
+					SCOPE_CYCLE_COUNTER(STAT_NiagaraRenderSpritesGlobalSortCPU);
+
+					FGlobalDynamicReadBuffer::FAllocation SortedIndices;
+					SortedIndices = CPUSimParticleDataAllocation.DynamicReadBuffer.AllocateInt32(NumInstances);
+					SortIndices(SortInfo.SortMode, SortInfo.SortAttributeOffset, *SourceParticleData, SceneProxy->GetLocalToWorld(), View, SortedIndices);
+					OutVertexFactory.SetSortedIndices(SortedIndices.ReadBuffer->SRV, SortedIndices.FirstIndex / sizeof(float));
 				}
 			}
-			else
-			{
-				FGlobalDynamicReadBuffer::FAllocation SortedIndices;
-				SortedIndices = CPUSimParticleDataAllocation.DynamicReadBuffer.AllocateInt32(NumInstances);
-				SortIndices(SortInfo.SortMode, SortInfo.SortAttributeOffset, DynamicDataSprites->RTParticleData, SceneProxy->GetLocalToWorld(), View, SortedIndices);
-				OutVertexFactory.SetSortedIndices(SortedIndices.ReadBuffer->SRV, SortedIndices.FirstIndex / sizeof(float));
-			}
+			OutVertexFactory.SetParticleData(CPUSimParticleDataAllocation.ParticleData.ReadBuffer->SRV, CPUSimParticleDataAllocation.ParticleData.FirstIndex / sizeof(float), SourceParticleData->GetFloatStride() / sizeof(float));
 		}
-		OutVertexFactory.SetParticleData(CPUSimParticleDataAllocation.ParticleData.ReadBuffer->SRV, CPUSimParticleDataAllocation.ParticleData.FirstIndex / sizeof(float), DynamicDataSprites->RTParticleData.GetFloatStride() / sizeof(float));
-	}
-	else
-	{
-		if (SortInfo.SortMode != ENiagaraSortMode::None && SortInfo.SortAttributeOffset != INDEX_NONE && GNiagaraGPUSorting && SceneProxy->GetBatcher())
+		else
 		{
-			SortInfo.ParticleDataFloatSRV = DynamicDataSprites->DataSet->CurrData().GetGPUBufferFloat()->SRV;
-			SortInfo.FloatDataOffset = 0;
-			SortInfo.FloatDataStride = DynamicDataSprites->RTParticleData.GetFloatStride() / sizeof(float);
-			SortInfo.GPUParticleCountSRV = DynamicDataSprites->DataSet->GetCurDataSetIndices().SRV;
-			SortInfo.GPUParticleCountOffset = 1;
-			const int32 IndexBufferOffset = SceneProxy->GetBatcher()->AddSortedGPUSimulation(SortInfo);
-			if (IndexBufferOffset != INDEX_NONE)
+			if (SortInfo.SortMode != ENiagaraSortMode::None && SortInfo.SortAttributeOffset != INDEX_NONE && GNiagaraGPUSorting && SceneProxy->GetBatcher())
 			{
-				OutVertexFactory.SetSortedIndices(SceneProxy->GetBatcher()->GetGPUSortedBuffer().VertexBufferSRV, IndexBufferOffset);
+				// Here we need to be conservative about the InstanceCount, since the final value is only known on the GPU after the simulation.
+				SortInfo.ParticleCount = SourceParticleData->GetNumInstances();
+
+				SortInfo.ParticleDataFloatSRV = SourceParticleData->GetGPUBufferFloat().SRV;
+				SortInfo.FloatDataOffset = 0;
+				SortInfo.FloatDataStride = SourceParticleData->GetFloatStride() / sizeof(float);
+				SortInfo.GPUParticleCountSRV = SourceParticleData->GetGPUIndices().SRV;
+				SortInfo.GPUParticleCountOffset = 1;
+				const int32 IndexBufferOffset = SceneProxy->GetBatcher()->AddSortedGPUSimulation(SortInfo);
+				if (IndexBufferOffset != INDEX_NONE)
+				{
+					OutVertexFactory.SetSortedIndices(SceneProxy->GetBatcher()->GetGPUSortedBuffer().VertexBufferSRV, IndexBufferOffset);
+				}
 			}
 
+			OutVertexFactory.SetParticleData(SourceParticleData->GetGPUBufferFloat().SRV, 0, SourceParticleData->GetFloatStride() / sizeof(float));
 		}
-
-		OutVertexFactory.SetParticleData(DynamicDataSprites->DataSet->CurrData().GetGPUBufferFloat()->SRV, 0, DynamicDataSprites->DataSet->CurrData().GetFloatStride() / sizeof(float));
 	}
 }
 
-void NiagaraRendererSprites::CreateMeshBatchForView(
+void FNiagaraRendererSprites::CreateMeshBatchForView(
 	const FSceneView* View, 
 	const FSceneViewFamily& ViewFamily, 
 	const FNiagaraSceneProxy *SceneProxy,
@@ -336,14 +426,28 @@ void NiagaraRendererSprites::CreateMeshBatchForView(
 	FMeshBatch& MeshBatch,
 	FNiagaraMeshCollectorResourcesSprite& CollectorResources) const
 {
-	int32 NumInstances = DynamicDataSprites->RTParticleData.GetNumInstances();
+	FNiagaraDataBuffer* SourceParticleData = DynamicDataSprites->GetParticleDataToRender();
+	check(SourceParticleData);//Can be null but should be checked before here.
+	int32 NumInstances = SourceParticleData->GetNumInstances();
 	const bool bIsWireframe = ViewFamily.EngineShowFlags.Wireframe;
-	FMaterialRenderProxy* MaterialRenderProxy = Material->GetRenderProxy();
 
-	ENiagaraSpriteFacingMode FacingMode = Properties->FacingMode;
-	ENiagaraSpriteAlignment AlignmentMode = Properties->Alignment;
+	FMaterialRenderProxy* MaterialRenderProxy = DynamicDataSprites->Material;
+	check(MaterialRenderProxy);
 
-	CollectorResources.VertexFactory.SetAlignmentMode((uint32)AlignmentMode);
+	ENiagaraSpriteFacingMode ActualFacingMode = FacingMode;
+	ENiagaraSpriteAlignment ActualAlignmentMode = Alignment;
+
+	if (FacingOffset == -1 && FacingMode == ENiagaraSpriteFacingMode::CustomFacingVector)
+	{
+		ActualFacingMode = ENiagaraSpriteFacingMode::FaceCamera;
+	}
+
+	if (AlignmentOffset == -1 && ActualAlignmentMode == ENiagaraSpriteAlignment::CustomAlignment)
+	{
+		ActualAlignmentMode = ENiagaraSpriteAlignment::Unaligned;
+	}
+
+	CollectorResources.VertexFactory.SetAlignmentMode((uint32)ActualAlignmentMode);
 	CollectorResources.VertexFactory.SetFacingMode((uint32)FacingMode);
 	CollectorResources.VertexFactory.SetParticleFactoryType(NVFT_Sprite);
 	CollectorResources.VertexFactory.InitResource();
@@ -359,9 +463,9 @@ void NiagaraRendererSprites::CreateMeshBatchForView(
 	VFLooseParams.ParticleFacingMode = CollectorResources.VertexFactory.GetFacingMode();
 	VFLooseParams.SortedIndices = CollectorResources.VertexFactory.GetSortedIndicesSRV() ? CollectorResources.VertexFactory.GetSortedIndicesSRV().GetReference() : GFNiagaraNullSortedIndicesVertexBuffer.VertexBufferSRV.GetReference();
 	VFLooseParams.SortedIndicesOffset = CollectorResources.VertexFactory.GetSortedIndicesOffset();
-	if (DynamicDataSprites->DataSet->GetSimTarget() == ENiagaraSimTarget::GPUComputeSim)
+	if (DynamicDataSprites->GetSimTarget() == ENiagaraSimTarget::GPUComputeSim)
 	{
-		VFLooseParams.IndirectArgsBuffer = DynamicDataSprites->DataSet->GetCurDataSetIndices().SRV;
+		VFLooseParams.IndirectArgsBuffer = SourceParticleData->GetGPUIndices().SRV;
 	}
 	else
 	{
@@ -397,9 +501,9 @@ void NiagaraRendererSprites::CreateMeshBatchForView(
 	MeshElement.MinVertexIndex = 0;
 	MeshElement.MaxVertexIndex = 0;// MeshElement.NumInstances * 4 - 1;
 	MeshElement.PrimitiveUniformBuffer = WorldSpacePrimitiveUniformBuffer.GetUniformBufferRHI();
-	if (DynamicDataSprites->DataSet->GetSimTarget() == ENiagaraSimTarget::GPUComputeSim)
+	if (DynamicDataSprites->GetSimTarget() == ENiagaraSimTarget::GPUComputeSim)
 	{
-		MeshElement.IndirectArgsBuffer = DynamicDataSprites->DataSet->GetCurDataSetIndices().Buffer;
+		MeshElement.IndirectArgsBuffer = SourceParticleData->GetGPUIndices().Buffer;
 		MeshElement.NumPrimitives = 0;
 	}
 
@@ -408,9 +512,11 @@ void NiagaraRendererSprites::CreateMeshBatchForView(
 		MeshElement.NumPrimitives = 6;
 		MeshElement.IndexBuffer = &GSixTriangleParticleIndexBuffer;
 	}
+	
+	INC_DWORD_STAT_BY(STAT_NiagaraNumSprites, NumInstances);
 }
 
-void NiagaraRendererSprites::GetDynamicMeshElements(const TArray<const FSceneView*>& Views, const FSceneViewFamily& ViewFamily, uint32 VisibilityMap, FMeshElementCollector& Collector, const FNiagaraSceneProxy *SceneProxy) const
+void FNiagaraRendererSprites::GetDynamicMeshElements(const TArray<const FSceneView*>& Views, const FSceneViewFamily& ViewFamily, uint32 VisibilityMap, FMeshElementCollector& Collector, const FNiagaraSceneProxy *SceneProxy) const
 {
 	SCOPE_CYCLE_COUNTER(STAT_NiagaraRender);
 	SCOPE_CYCLE_COUNTER(STAT_NiagaraRenderSprites);
@@ -418,18 +524,25 @@ void NiagaraRendererSprites::GetDynamicMeshElements(const TArray<const FSceneVie
 	SimpleTimer MeshElementsTimer;
 
 	//check(DynamicDataRender)
-
 	FNiagaraDynamicDataSprites *DynamicDataSprites = static_cast<FNiagaraDynamicDataSprites*>(DynamicDataRender);
+	if (DynamicDataSprites == nullptr)
+	{
+		return;
+	}
 
-	if (!DynamicDataSprites
-		|| DynamicDataSprites->RTParticleData.GetNumInstancesAllocated() == 0
-		|| DynamicDataSprites->RTParticleData.GetNumInstances() == 0
-		|| nullptr == Properties
-		|| !GSupportsResourceView // Current shader requires SRV to draw properly in all cases.
+	FNiagaraDataBuffer* SourceParticleData = DynamicDataSprites->GetParticleDataToRender();
+	if (SourceParticleData == nullptr	||
+		SourceParticleData->GetNumInstances() == 0 ||
+		GbEnableNiagaraSpriteRendering == 0 ||
+		!GSupportsResourceView // Current shader requires SRV to draw properly in all cases.
 		)
 	{
 		return;
 	}
+
+#if STATS
+	FScopeCycleCounter EmitterStatsCounter(EmitterStatID);
+#endif
 	
 	FCPUSimParticleDataAllocation CPUSimParticleDataAllocation = ConditionalAllocateCPUSimParticleData(DynamicDataSprites, Collector.GetDynamicReadBuffer());
 
@@ -455,7 +568,7 @@ void NiagaraRendererSprites::GetDynamicMeshElements(const TArray<const FSceneVie
 }
 
 #if RHI_RAYTRACING
-void NiagaraRendererSprites::GetDynamicRayTracingInstances(FRayTracingMaterialGatheringContext& Context, TArray<FRayTracingInstance>& OutRayTracingInstances, const FNiagaraSceneProxy* SceneProxy)
+void FNiagaraRendererSprites::GetDynamicRayTracingInstances(FRayTracingMaterialGatheringContext& Context, TArray<FRayTracingInstance>& OutRayTracingInstances, const FNiagaraSceneProxy* SceneProxy)
 {
 	SCOPE_CYCLE_COUNTER(STAT_NiagaraRender);
 	SCOPE_CYCLE_COUNTER(STAT_NiagaraRenderSprites);
@@ -463,12 +576,17 @@ void NiagaraRendererSprites::GetDynamicRayTracingInstances(FRayTracingMaterialGa
 	SimpleTimer MeshElementsTimer;
 
 	FNiagaraDynamicDataSprites *DynamicDataSprites = static_cast<FNiagaraDynamicDataSprites*>(DynamicDataRender);
+	if (DynamicDataSprites == nullptr)
+	{
+		return;
+	}
 
-	if (!DynamicDataSprites
-		|| DynamicDataSprites->RTParticleData.GetNumInstancesAllocated() == 0
-		|| DynamicDataSprites->RTParticleData.GetNumInstances() == 0
-		|| nullptr == Properties
-		|| !GSupportsResourceView // Current shader requires SRV to draw properly in all cases.
+	FNiagaraDataBuffer* SourceParticleData = DynamicDataSprites->GetParticleDataToRender();
+	if (SourceParticleData == nullptr	||
+		SourceParticleData->GetNumInstancesAllocated() == 0 ||
+		SourceParticleData->GetNumInstances() == 0 ||
+		GbEnableNiagaraSpriteRendering == 0 ||
+		!GSupportsResourceView // Current shader requires SRV to draw properly in all cases.
 		)
 	{
 		return;
@@ -499,7 +617,7 @@ void NiagaraRendererSprites::GetDynamicRayTracingInstances(FRayTracingMaterialGa
 			SceneProxy,
 			MeshBatch,
 			RayTracingGeometry,
-			6 * DynamicDataSprites->RTParticleData.GetNumInstances(),
+			6 * SourceParticleData->GetNumInstances(),
 			RayTracingDynamicVertexBuffer);
 	}
 
@@ -511,66 +629,32 @@ void NiagaraRendererSprites::GetDynamicRayTracingInstances(FRayTracingMaterialGa
 }
 #endif
 
-bool NiagaraRendererSprites::SetMaterialUsage()
-{
-	//Causes deadlock :S Need to look at / rework the setting of materials and render modules.
-	return Material && Material->CheckMaterialUsage_Concurrent(MATUSAGE_NiagaraSprites);
-}
-
-void NiagaraRendererSprites::TransformChanged()
-{
-	WorldSpacePrimitiveUniformBuffer.ReleaseResource();
-}
-
 /** Update render data buffer from attributes */
-FNiagaraDynamicDataBase *NiagaraRendererSprites::GenerateVertexData(const FNiagaraSceneProxy* Proxy, FNiagaraDataSet &Data, const ENiagaraSimTarget Target)
+FNiagaraDynamicDataBase *FNiagaraRendererSprites::GenerateDynamicData(const FNiagaraSceneProxy* Proxy, const UNiagaraRendererProperties* InProperties, const FNiagaraEmitterInstance* Emitter)const
 {
 	FNiagaraDynamicDataSprites *DynamicData = nullptr;
+	const UNiagaraSpriteRendererProperties* Properties = CastChecked<const UNiagaraSpriteRendererProperties>(InProperties);
 
-	if (bEnabled && Properties)
+	if (Properties)
 	{
 		SimpleTimer VertexDataTimer;
 
-		SCOPE_CYCLE_COUNTER(STAT_NiagaraGenSpriteVertexData);
-
-		if (PositionOffset == INDEX_NONE || LastSyncId != Properties->SyncId)
-		{			
-			// optional attributes; we pass -1 as the offset so the VF can branch
-			int32 IntDummy;
-			Data.GetVariableComponentOffsets(Properties->PositionBinding.DataSetVariable, PositionOffset, IntDummy);
-			Data.GetVariableComponentOffsets(Properties->VelocityBinding.DataSetVariable, VelocityOffset, IntDummy);
-			Data.GetVariableComponentOffsets(Properties->SpriteRotationBinding.DataSetVariable, RotationOffset, IntDummy);
-			Data.GetVariableComponentOffsets(Properties->SpriteSizeBinding.DataSetVariable, SizeOffset, IntDummy);
-			Data.GetVariableComponentOffsets(Properties->ColorBinding.DataSetVariable, ColorOffset, IntDummy);
-			Data.GetVariableComponentOffsets(Properties->SpriteFacingBinding.DataSetVariable, FacingOffset, IntDummy);
-			Data.GetVariableComponentOffsets(Properties->SpriteAlignmentBinding.DataSetVariable, AlignmentOffset, IntDummy);
-			Data.GetVariableComponentOffsets(Properties->SubImageIndexBinding.DataSetVariable, SubImageOffset, IntDummy);
-			Data.GetVariableComponentOffsets(Properties->DynamicMaterialBinding.DataSetVariable, MaterialParamOffset, IntDummy);
-			Data.GetVariableComponentOffsets(Properties->DynamicMaterial1Binding.DataSetVariable, MaterialParamOffset1, IntDummy);
-			Data.GetVariableComponentOffsets(Properties->DynamicMaterial2Binding.DataSetVariable, MaterialParamOffset2, IntDummy);
-			Data.GetVariableComponentOffsets(Properties->DynamicMaterial3Binding.DataSetVariable, MaterialParamOffset3, IntDummy);
-			Data.GetVariableComponentOffsets(Properties->CameraOffsetBinding.DataSetVariable, CameraOffsetOffset, IntDummy);
-			Data.GetVariableComponentOffsets(Properties->UVScaleBinding.DataSetVariable, UVScaleOffset, IntDummy);
-			Data.GetVariableComponentOffsets(Properties->NormalizedAgeBinding.DataSetVariable, NormalizedAgeOffset, IntDummy);
-			Data.GetVariableComponentOffsets(Properties->MaterialRandomBinding.DataSetVariable, MaterialRandomOffset, IntDummy);
-			Data.GetVariableComponentOffsets(Properties->CustomSortingBinding.DataSetVariable, CustomSortingOffset, IntDummy);
-
-			if (CustomSortingOffset == INDEX_NONE && (Properties->SortMode == ENiagaraSortMode::CustomAscending || Properties->SortMode == ENiagaraSortMode::CustomDecending))
-			{
-				UE_LOG(LogNiagara, Warning, TEXT("Niagara Sprite Emitter using custom sorting but does not have a valid custom sorting attribute binding."));
-			}
+		SCOPE_CYCLE_COUNTER(STAT_NiagaraGenSpriteDynamicData);
 
-			LastSyncId = Properties->SyncId;
-		}
-		
-		if(Data.CurrData().GetNumInstances() > 0)
+		FNiagaraDataSet& Data = Emitter->GetData();
+		ENiagaraSimTarget SimTarget = Data.GetSimTarget();
+		if(SimTarget == ENiagaraSimTarget::GPUComputeSim || Data.GetCurrentDataChecked().GetNumInstances() > 0)
 		{
-			DynamicData = new FNiagaraDynamicDataSprites;
+			DynamicData = new FNiagaraDynamicDataSprites(Emitter);
 
-			//TODO: This buffer is far fatter than needed. Just pull out the data needed for rendering.
-			Data.CurrData().CopyTo(DynamicData->RTParticleData);
 
-			DynamicData->DataSet = &Data;
+			//In preparation for a material override feature, we pass our material(s) and relevance in via dynamic data.
+			//The renderer ensures we have the correct usage and relevance for materials in BaseMaterials_GT.
+			//Any override feature must also do the same for materials that are set.
+			check(BaseMaterials_GT.Num() == 1);
+			check(BaseMaterials_GT[0]->CheckMaterialUsage_Concurrent(MATUSAGE_NiagaraSprites));
+			DynamicData->Material = BaseMaterials_GT[0]->GetRenderProxy();
+			DynamicData->SetMaterialRelevance(BaseMaterialRelevance_GT);
 		}
 
 		CPUTimeMS = VertexDataTimer.GetElapsedMilliseconds();
@@ -579,42 +663,18 @@ FNiagaraDynamicDataBase *NiagaraRendererSprites::GenerateVertexData(const FNiaga
 	return DynamicData;  // for VF that can fetch from particle data directly
 }
 
-
-
-void NiagaraRendererSprites::SetDynamicData_RenderThread(FNiagaraDynamicDataBase* NewDynamicData)
-{
-	check(IsInRenderingThread());
-
-	if (DynamicDataRender)
-	{
-		delete static_cast<FNiagaraDynamicDataSprites*>(DynamicDataRender);
-		DynamicDataRender = NULL;
-	}
-	DynamicDataRender = NewDynamicData;
-}
-
-int NiagaraRendererSprites::GetDynamicDataSize()
+int FNiagaraRendererSprites::GetDynamicDataSize()const
 {
 	uint32 Size = sizeof(FNiagaraDynamicDataSprites);
 	return Size;
 }
 
-bool NiagaraRendererSprites::HasDynamicData()
+void FNiagaraRendererSprites::TransformChanged()
 {
-//	return DynamicDataRender && static_cast<FNiagaraDynamicDataSprites*>(DynamicDataRender)->VertexData.Num() > 0;
-	return DynamicDataRender!=nullptr;// && static_cast<FNiagaraDynamicDataSprites*>(DynamicDataRender)->NumInstances > 0;
-}
-
-#if WITH_EDITORONLY_DATA
-
-const TArray<FNiagaraVariable>& NiagaraRendererSprites::GetRequiredAttributes()
-{
-	return Properties->GetRequiredAttributes();
+	WorldSpacePrimitiveUniformBuffer.ReleaseResource();
 }
 
-const TArray<FNiagaraVariable>& NiagaraRendererSprites::GetOptionalAttributes()
+bool FNiagaraRendererSprites::IsMaterialValid(UMaterialInterface* Mat)const
 {
-	return Properties->GetOptionalAttributes();
-}
-
-#endif
\ No newline at end of file
+	return Mat && Mat->CheckMaterialUsage(MATUSAGE_NiagaraSprites);
+}
\ No newline at end of file
diff --git a/Engine/Plugins/FX/Niagara/Source/Niagara/Private/NiagaraRibbonRendererProperties.cpp b/Engine/Plugins/FX/Niagara/Source/Niagara/Private/NiagaraRibbonRendererProperties.cpp
--- a/Engine/Plugins/FX/Niagara/Source/Niagara/Private/NiagaraRibbonRendererProperties.cpp
+++ b/Engine/Plugins/FX/Niagara/Source/Niagara/Private/NiagaraRibbonRendererProperties.cpp
@@ -3,10 +3,30 @@
 #include "NiagaraRibbonRendererProperties.h"
 #include "NiagaraRendererRibbons.h"
 #include "NiagaraConstants.h"
+#include "NiagaraBoundsCalculatorHelper.h"
+
+UNiagaraRibbonRendererProperties::UNiagaraRibbonRendererProperties()
+	: Material(nullptr)
+	, FacingMode(ENiagaraRibbonFacingMode::Screen)
+	, UV0TilingDistance(0.0f)
+	, UV0Scale(FVector2D(1.0f, 1.0f))
+	, UV0AgeOffsetMode(ENiagaraRibbonAgeOffsetMode::Scale)
+	, UV1TilingDistance(0.0f)
+	, UV1Scale(FVector2D(1.0f, 1.0f))
+	, UV1AgeOffsetMode(ENiagaraRibbonAgeOffsetMode::Scale)
+{
+}
 
-NiagaraRenderer* UNiagaraRibbonRendererProperties::CreateEmitterRenderer(ERHIFeatureLevel::Type FeatureLevel)
+FNiagaraRenderer* UNiagaraRibbonRendererProperties::CreateEmitterRenderer(ERHIFeatureLevel::Type FeatureLevel, const FNiagaraEmitterInstance* Emitter)
 {
-	return new NiagaraRendererRibbons(FeatureLevel, this);
+	FNiagaraRenderer* NewRenderer = new FNiagaraRendererRibbons(FeatureLevel, this, Emitter);
+	NewRenderer->Initialize(FeatureLevel, this, Emitter);
+	return NewRenderer;
+}
+
+FNiagaraBoundsCalculator* UNiagaraRibbonRendererProperties::CreateBoundsCalculator()
+{
+	return new FNiagaraBoundsCalculatorHelper<false, false, true>();
 }
 
 void UNiagaraRibbonRendererProperties::GetUsedMaterials(TArray<UMaterialInterface*>& OutMaterials) const
@@ -17,7 +37,6 @@ void UNiagaraRibbonRendererProperties::GetUsedMaterials(TArray<UMaterialInterfac
 void UNiagaraRibbonRendererProperties::PostInitProperties()
 {
 	Super::PostInitProperties();
-	SyncId = 0;
 	if (HasAnyFlags(RF_ClassDefaultObject) == false)
 	{
 		InitBindings();
@@ -56,10 +75,6 @@ void UNiagaraRibbonRendererProperties::InitBindings()
 void UNiagaraRibbonRendererProperties::PostEditChangeProperty(struct FPropertyChangedEvent& PropertyChangedEvent)
 {
 	Super::PostEditChangeProperty(PropertyChangedEvent);
-	if (PropertyChangedEvent.GetPropertyName() != TEXT("SyncId"))
-	{
-		SyncId++;
-	}
 }
 
 const TArray<FNiagaraVariable>& UNiagaraRibbonRendererProperties::GetRequiredAttributes()
