Copying //UE4/Dev-Niagara to Dev-Main (//UE4/Dev-Main)

diff --git a/Engine/Plugins/FX/Niagara/Source/Niagara/Private/NiagaraScript.cpp b/Engine/Plugins/FX/Niagara/Source/Niagara/Private/NiagaraScript.cpp
--- a/Engine/Plugins/FX/Niagara/Source/Niagara/Private/NiagaraScript.cpp
+++ b/Engine/Plugins/FX/Niagara/Source/Niagara/Private/NiagaraScript.cpp
@@ -26,6 +26,7 @@
 	#include "Interfaces/ITargetPlatform.h"
 #endif
 
+
 #include "UObject/FortniteMainBranchObjectVersion.h"
 #include "UObject/RenderingObjectVersion.h"
 
@@ -113,10 +114,24 @@ bool FNiagaraVMExecutableDataId::operator==(const FNiagaraVMExecutableDataId& Re
 	if (CompilerVersionID != ReferenceSet.CompilerVersionID ||
 		ScriptUsageType != ReferenceSet.ScriptUsageType || 
 		ScriptUsageTypeID != ReferenceSet.ScriptUsageTypeID ||
-		BaseScriptID != ReferenceSet.BaseScriptID)
+		BaseScriptID != ReferenceSet.BaseScriptID ||
+		BaseScriptCompileHash != ReferenceSet.BaseScriptCompileHash)
+	{
+		return false;
+	}
+
+	if (ReferencedCompileHashes.Num() != ReferenceSet.ReferencedCompileHashes.Num())
 	{
 		return false;
 	}
+
+	for (int32 ReferencedHashIndex = 0; ReferencedHashIndex < ReferencedCompileHashes.Num(); ReferencedHashIndex++)
+	{
+		if (ReferencedCompileHashes[ReferencedHashIndex] != ReferenceSet.ReferencedCompileHashes[ReferencedHashIndex])
+		{
+			return false;
+		}
+	}
 	
 	if (ReferencedDependencyIds.Num() != ReferenceSet.ReferencedDependencyIds.Num())
 	{
@@ -161,6 +176,8 @@ void FNiagaraVMExecutableDataId::AppendKeyString(FString& KeyString) const
 	KeyString += TEXT("_");
 	KeyString += BaseScriptID.ToString();
 	KeyString += TEXT("_");
+	KeyString += BaseScriptCompileHash.ToString();
+	KeyString += TEXT("_");
 
 	for (int32 Idx = 0; Idx < AdditionalDefines.Num(); Idx++)
 	{
@@ -172,6 +189,17 @@ void FNiagaraVMExecutableDataId::AppendKeyString(FString& KeyString) const
 		}
 	}
 	
+	// Add any referenced script compile hashes to the key so that we will recompile when they are changed
+	for (int32 HashIndex = 0; HashIndex < ReferencedCompileHashes.Num(); HashIndex++)
+	{
+		KeyString += ReferencedCompileHashes[HashIndex].ToString();
+
+		if (HashIndex < ReferencedCompileHashes.Num() - 1)
+		{
+			KeyString += TEXT("_");
+		}
+	}
+
 	// Add any referenced functions to the key so that we will recompile when they are changed
 	for (int32 FunctionIndex = 0; FunctionIndex < ReferencedDependencyIds.Num(); FunctionIndex++)
 	{
@@ -431,28 +459,51 @@ void UNiagaraScript::PostLoad()
 	
 	bool bNeedsRecompile = false;
 	const int32 NiagaraVer = GetLinkerCustomVersion(FNiagaraCustomVersion::GUID);
-	
-	// Resources can be processed / registered now that we're back on the main thread
-	ProcessSerializedShaderMaps(this, LoadedScriptResources, ScriptResource, ScriptResourcesByFeatureLevel);
 
-	if (GIsEditor)
-	{
-		
 #if WITH_EDITORONLY_DATA
-		// Since we're about to check the synchronized state, we need to make sure that it has been post-loaded (which 
-		// can affect the results of that call).
-		if (Source != nullptr)
-		{	
-			Source->ConditionalPostLoad();
+	if (Source != nullptr)
+	{
+		Source->ConditionalPostLoad();
+		if (NiagaraVer < FNiagaraCustomVersion::UseHashesToIdentifyCompileStateOfTopLevelScripts && CachedScriptVMId.CompilerVersionID.IsValid())
+		{
+			FGuid BaseId = Source->GetCompileBaseId(Usage, UsageId);
+			if (BaseId.IsValid() == false)
+			{
+				UE_LOG(LogNiagara, Warning,
+					TEXT("Invalidating compile ids for script %s because it doesn't have a valid base id.  The owning asset will continue to compile on load until it is resaved."),
+					*GetPathName());
+				bool bForceRebuild = true;
+				Source->InvalidateCachedCompileIds();
+				Source->ComputeVMCompilationId(CachedScriptVMId, Usage, UsageId, bForceRebuild);
+			}
+			else
+			{
+				FNiagaraCompileHash CompileHash = Source->GetCompileHash(Usage, UsageId);
+				if (CompileHash.IsValid())
+				{
+					CachedScriptVMId.BaseScriptCompileHash = CompileHash;
+				}
+				else
+				{
+					// If the compile hash isn't valid, recompute the entire cached VM Id.
+					bool bForceRebuild = true;
+					Source->ComputeVMCompilationId(CachedScriptVMId, Usage, UsageId, bForceRebuild);
+				}
+			}
 		}
-
-#endif
 	}
+#endif
+	
+	// Resources can be processed / registered now that we're back on the main thread
+	ProcessSerializedShaderMaps(this, LoadedScriptResources, ScriptResource, ScriptResourcesByFeatureLevel);
 
 	// for now, force recompile until we can be sure everything is working
 	//bNeedsRecompile = true;
 #if WITH_EDITORONLY_DATA
-	CacheResourceShadersForRendering(false, bNeedsRecompile);
+	if (CachedScriptVMId.BaseScriptID.IsValid() && CachedScriptVMId.BaseScriptCompileHash.IsValid())
+	{
+		CacheResourceShadersForRendering(false, bNeedsRecompile);
+	}
 #endif
 #if STATS
 	GenerateStatScopeIDs();
@@ -478,7 +529,18 @@ bool UNiagaraScript::IsReadyToRun(ENiagaraSimTarget SimTarget) const
 	return false;
 }
 
-
+bool UNiagaraScript::ShouldCacheShadersForCooking() const
+{
+	if (CanBeRunOnGpu())
+	{
+		UNiagaraEmitter* OwningEmitter = GetTypedOuter<UNiagaraEmitter>();
+		if (OwningEmitter != nullptr && OwningEmitter->SimTarget == ENiagaraSimTarget::GPUComputeSim)
+		{
+			return true;
+		}
+	}
+	return false;
+}
 
 #if STATS
 void UNiagaraScript::GenerateStatScopeIDs()
@@ -531,22 +593,53 @@ bool UNiagaraScript::AreScriptAndSourceSynchronized() const
 		{
 			if (NewId != LastReportedVMId)
 			{
-				if (NewId.BaseScriptID != CachedScriptVMId.BaseScriptID)
+				if (GEnableVerboseNiagaraChangeIdLogging)
 				{
-					UE_LOG(LogNiagara, Log, TEXT("AreScriptAndSourceSynchronized base script id's don't match. %s != %s"), *NewId.BaseScriptID.ToString(), *CachedScriptVMId.BaseScriptID.ToString());
-				}
-				if (NewId.ReferencedDependencyIds.Num() != CachedScriptVMId.ReferencedDependencyIds.Num())
-				{
-					UE_LOG(LogNiagara, Log, TEXT("AreScriptAndSourceSynchronized num dependencies don't match. %d != %d"), NewId.ReferencedDependencyIds.Num(), CachedScriptVMId.ReferencedDependencyIds.Num());
-				}
-				else
-				{
-					for (int32 i = 0; i < NewId.ReferencedDependencyIds.Num(); i++)
+					if (NewId.BaseScriptID != CachedScriptVMId.BaseScriptID)
+					{
+						UE_LOG(LogNiagara, Log, TEXT("AreScriptAndSourceSynchronized base script id's don't match. %s != %s, script %s"),
+							*NewId.BaseScriptID.ToString(), *CachedScriptVMId.BaseScriptID.ToString(), *GetPathName());
+					}
+
+					if (NewId.BaseScriptCompileHash != CachedScriptVMId.BaseScriptCompileHash)
+					{
+						UE_LOG(LogNiagara, Log, TEXT("AreScriptAndSourceSynchronized base script compile hashes don't match. %s != %s, script %s"),
+							*NewId.BaseScriptCompileHash.ToString(), *CachedScriptVMId.BaseScriptCompileHash.ToString(), *GetPathName());
+					}
+
+					if (NewId.ReferencedCompileHashes.Num() != CachedScriptVMId.ReferencedCompileHashes.Num())
+					{
+						UE_LOG(LogNiagara, Log, TEXT("AreScriptAndSourceSynchronized num referenced compile hashes don't match. %d != %d, script %s"),
+							NewId.ReferencedCompileHashes.Num(), CachedScriptVMId.ReferencedCompileHashes.Num(), *GetPathName());
+					}
+					else
 					{
-						if (NewId.ReferencedDependencyIds[i] != CachedScriptVMId.ReferencedDependencyIds[i])
+						for (int32 i = 0; i < NewId.ReferencedCompileHashes.Num(); i++)
 						{
-							UE_LOG(LogNiagara, Log, TEXT("AreScriptAndSourceSynchronized reference id %d doesn't match. %s != %s, source %s"), i, *NewId.ReferencedDependencyIds[i].ToString(), *CachedScriptVMId.ReferencedDependencyIds[i].ToString(),
-								NewId.ReferencedObjects[i] != nullptr ? *NewId.ReferencedObjects[i]->GetPathName() : TEXT("nullptr"));
+							if (NewId.ReferencedCompileHashes[i] != CachedScriptVMId.ReferencedCompileHashes[i])
+							{
+								UE_LOG(LogNiagara, Log, TEXT("AreScriptAndSourceSynchronized referenced compile hash %d doesn't match. %s != %s, script %s, source %s"),
+									i, *NewId.ReferencedCompileHashes[i].ToString(), *CachedScriptVMId.ReferencedCompileHashes[i].ToString(), *GetPathName(),
+									NewId.ReferencedObjects[i] != nullptr ? *NewId.ReferencedObjects[i]->GetPathName() : TEXT("nullptr"));
+							}
+						}
+					}
+
+					if (NewId.ReferencedDependencyIds.Num() != CachedScriptVMId.ReferencedDependencyIds.Num())
+					{
+						UE_LOG(LogNiagara, Log, TEXT("AreScriptAndSourceSynchronized num dependencies don't match. %d != %d, script %s"),
+							NewId.ReferencedDependencyIds.Num(), CachedScriptVMId.ReferencedDependencyIds.Num(), *GetPathName());
+					}
+					else
+					{
+						for (int32 i = 0; i < NewId.ReferencedDependencyIds.Num(); i++)
+						{
+							if (NewId.ReferencedDependencyIds[i] != CachedScriptVMId.ReferencedDependencyIds[i])
+							{
+								UE_LOG(LogNiagara, Log, TEXT("AreScriptAndSourceSynchronized reference id %d doesn't match. %s != %s, script %s, source %s"),
+									i, *NewId.ReferencedDependencyIds[i].ToString(), *CachedScriptVMId.ReferencedDependencyIds[i].ToString(), *GetPathName(),
+									NewId.ReferencedObjects[i] != nullptr ? *NewId.ReferencedObjects[i]->GetPathName() : TEXT("nullptr"));
+							}
 						}
 					}
 				}
@@ -811,7 +904,16 @@ void UNiagaraScript::SetVMCompilationResults(const FNiagaraVMExecutableDataId& I
 	// Now go ahead and trigger the GPU script compile now that we have a compiled GPU hlsl script.
 	if (Usage == ENiagaraScriptUsage::ParticleGPUComputeScript)
 	{
-		CacheResourceShadersForRendering(false, true);
+		if (CachedScriptVMId.CompilerVersionID.IsValid() && CachedScriptVMId.BaseScriptID.IsValid() && CachedScriptVMId.BaseScriptCompileHash.IsValid())
+		{
+			CacheResourceShadersForRendering(false, true);
+		}
+		else
+		{
+			UE_LOG(LogNiagara, Error,
+				TEXT("Failed to cache resource shaders for rendering for script %s because it had an invalid cached script id. This should be fixed by force recompiling the owning asset using the 'Full Rebuild' option and then saving the asset."),
+				*GetPathName());
+		}
 	}
 
 	InvalidateExecutionReadyParameterStores();
@@ -961,7 +1063,7 @@ void UNiagaraScript::GetAssetRegistryTags(TArray<FAssetRegistryTag>& OutTags) co
 
 void UNiagaraScript::BeginCacheForCookedPlatformData(const ITargetPlatform *TargetPlatform)
 {
-	if (CanBeRunOnGpu())
+	if (ShouldCacheShadersForCooking())
 	{
 		// Commandlets like DerivedDataCacheCommandlet call BeginCacheForCookedPlatformData directly on objects. This may mean that
 		// we have not properly gotten the HLSL script generated by the time that we get here. This does the awkward work of 
@@ -972,6 +1074,14 @@ void UNiagaraScript::BeginCacheForCookedPlatformData(const ITargetPlatform *Targ
 			SystemOwner->WaitForCompilationComplete();
 		}
 
+		if (CachedScriptVMId.CompilerVersionID.IsValid() == false || CachedScriptVMId.BaseScriptID.IsValid() == false || CachedScriptVMId.BaseScriptCompileHash.IsValid() == false)
+		{
+			UE_LOG(LogNiagara, Error,
+				TEXT("Failed to cache cooked shader for script %s because it had an invalid cached script id.  This should be fixed by running the console command fx.PreventSystemRecompile with the owning system asset path as the argument and then resaving the assets."),
+				*GetPathName());
+			return;
+		}
+
 		TArray<FName> DesiredShaderFormats;
 		TargetPlatform->GetAllTargetedShaderFormats(DesiredShaderFormats);
 
@@ -1000,10 +1110,12 @@ void UNiagaraScript::CacheResourceShadersForCooking(EShaderPlatform ShaderPlatfo
 			ERHIFeatureLevel::Type TargetFeatureLevel = GetMaxSupportedFeatureLevel(ShaderPlatform);
 
 			FNiagaraShaderScript* NewResource = AllocateResource();
-			check(CachedScriptVMId.CompilerVersionID != FGuid());
-			check(CachedScriptVMId.BaseScriptID != FGuid());
+			check(CachedScriptVMId.CompilerVersionID.IsValid());
+			check(CachedScriptVMId.BaseScriptID.IsValid());
+			check(CachedScriptVMId.BaseScriptCompileHash.IsValid());
 
-			NewResource->SetScript(this, (ERHIFeatureLevel::Type)TargetFeatureLevel, CachedScriptVMId.CompilerVersionID, CachedScriptVMId.BaseScriptID, CachedScriptVMId.ReferencedDependencyIds, GetName());
+			NewResource->SetScript(this, (ERHIFeatureLevel::Type)TargetFeatureLevel, CachedScriptVMId.CompilerVersionID, CachedScriptVMId.BaseScriptID,
+				CachedScriptVMId.BaseScriptCompileHash,	CachedScriptVMId.ReferencedCompileHashes, CachedScriptVMId.ReferencedDependencyIds, GetFullName());
 			ResourceToCache = NewResource;
 
 			check(ResourceToCache);
@@ -1067,7 +1179,8 @@ void UNiagaraScript::CacheResourceShadersForRendering(bool bRegenerateId, bool b
 		{
 			FNiagaraShaderScript* ResourceToCache;
 			ERHIFeatureLevel::Type CacheFeatureLevel = GMaxRHIFeatureLevel;
-			ScriptResource.SetScript(this, FeatureLevel, CachedScriptVMId.CompilerVersionID, CachedScriptVMId.BaseScriptID, CachedScriptVMId.ReferencedDependencyIds, GetName());
+			ScriptResource.SetScript(this, FeatureLevel, CachedScriptVMId.CompilerVersionID, CachedScriptVMId.BaseScriptID,
+				CachedScriptVMId.BaseScriptCompileHash, CachedScriptVMId.ReferencedCompileHashes, CachedScriptVMId.ReferencedDependencyIds, GetFullName());
 
 			//if (ScriptResourcesByFeatureLevel[FeatureLevel])
 			{
diff --git a/Engine/Plugins/FX/Niagara/Source/Niagara/Private/NiagaraScriptExecutionContext.cpp b/Engine/Plugins/FX/Niagara/Source/Niagara/Private/NiagaraScriptExecutionContext.cpp
--- a/Engine/Plugins/FX/Niagara/Source/Niagara/Private/NiagaraScriptExecutionContext.cpp
+++ b/Engine/Plugins/FX/Niagara/Source/Niagara/Private/NiagaraScriptExecutionContext.cpp
@@ -6,6 +6,8 @@
 #include "NiagaraDataInterface.h"
 #include "NiagaraSystemInstance.h"
 #include "NiagaraWorldManager.h"
+#include "NiagaraSystemInstance.h"
+#include "NiagaraEmitterInstance.h"
 
 DECLARE_CYCLE_STAT(TEXT("Register Setup"), STAT_NiagaraSimRegisterSetup, STATGROUP_Niagara);
 DECLARE_CYCLE_STAT(TEXT("Context Ticking"), STAT_NiagaraScriptExecContextTick, STATGROUP_Niagara);
@@ -135,7 +137,21 @@ void FNiagaraScriptExecutionContext::PostTick()
 	}
 }
 
-bool FNiagaraScriptExecutionContext::Execute(uint32 NumInstances, TArray<FNiagaraDataSetExecutionInfo, TInlineAllocator<8>>& DataSetInfos)
+void FNiagaraScriptExecutionContext::BindData(int32 Index, FNiagaraDataSet& DataSet, int32 StartInstance, bool bUpdateInstanceCounts)
+{
+	DataSetInfo.SetNum(FMath::Max(DataSetInfo.Num(), Index + 1));
+	
+	DataSetInfo[Index].Init(&DataSet, DataSet.GetCurrentData(), DataSet.GetDestinationData(), StartInstance, bUpdateInstanceCounts);
+}
+
+void FNiagaraScriptExecutionContext::BindData(int32 Index, FNiagaraDataBuffer* Input, FNiagaraDataBuffer* Output, int32 StartInstance, bool bUpdateInstanceCounts)
+{
+	DataSetInfo.SetNum(FMath::Max(DataSetInfo.Num(), Index + 1));
+	check(Input || Output);
+	DataSetInfo[Index].Init(Input ? Input->GetOwner() : Output->GetOwner(), Input, Output, StartInstance, bUpdateInstanceCounts);
+}
+
+bool FNiagaraScriptExecutionContext::Execute(uint32 NumInstances)
 {
 	if (NumInstances == 0)
 	{
@@ -151,28 +167,48 @@ bool FNiagaraScriptExecutionContext::Execute(uint32 NumInstances, TArray<FNiagar
 
 	DataSetMetaTable.Reset();
 
-	bool bRegisterSetupCompleted = true;
+	bool bError = false;
 	{
 		SCOPE_CYCLE_COUNTER(STAT_NiagaraSimRegisterSetup);
-		for (FNiagaraDataSetExecutionInfo& DataSetInfo : DataSetInfos)
+		for (FNiagaraDataSetExecutionInfo& Info : DataSetInfo)
 		{
+			FNiagaraDataBuffer* DestinationData = Info.DataSet->GetDestinationData();
 #if NIAGARA_NAN_CHECKING
 			DataSetInfo.DataSet->CheckForNaNs();
 #endif
-			check(DataSetInfo.DataSet);
-			DataSetMetaTable.Emplace(DataSetInfo.DataSet->GetSizeBytes(), &InputRegisters[NumInputRegisters], NumInputRegisters, DataSetInfo.StartInstance,
-				&DataSetInfo.DataSet->CurrIDTable(), &DataSetInfo.DataSet->GetFreeIDTable(), &DataSetInfo.DataSet->GetNumFreeIDs(), &DataSetInfo.DataSet->GetMaxUsedID(), DataSetInfo.DataSet->GetIDAcquireTag());
-			if (DataSetInfo.bAllocate)
+			check(Info.DataSet);
+			DataSetMetaTable.Emplace(&InputRegisters[NumInputRegisters], NumInputRegisters, Info.StartInstance,
+				DestinationData ? &DestinationData->GetIDTable() : nullptr, &Info.DataSet->GetFreeIDTable(), &Info.DataSet->GetNumFreeIDs(), &Info.DataSet->GetMaxUsedID(), Info.DataSet->GetIDAcquireTag());
+
+			int32 TotalComponents = Info.DataSet->GetNumFloatComponents() + Info.DataSet->GetNumInt32Components();
+			if (NumInputRegisters + TotalComponents > VectorVM::MaxInputRegisters || NumOutputRegisters + TotalComponents > VectorVM::MaxOutputRegisters)
 			{
-				DataSetInfo.DataSet->Allocate(NumInstances);
-				DataSetInfo.DataSet->SetNumInstances(NumInstances);
+				UE_LOG(LogNiagara, Warning, TEXT("VM Script is using too many registers."));//TODO: We can make this a non issue and support any number of registers.
+				bError = true;
+				break;
+			}
+
+			if (Info.Input)
+			{
+				Info.Input->AppendToRegisterTable(InputRegisters, NumInputRegisters, Info.StartInstance);
+			}
+			else
+			{
+				Info.DataSet->ClearRegisterTable(InputRegisters, NumInputRegisters);
+			}
+
+			if (Info.Output)
+			{
+				Info.Output->AppendToRegisterTable(OutputRegisters, NumOutputRegisters, Info.StartInstance);
+			}
+			else
+			{
+				Info.DataSet->ClearRegisterTable(OutputRegisters, NumOutputRegisters);
 			}
-			
-			bRegisterSetupCompleted &= DataSetInfo.DataSet->AppendToRegisterTable(InputRegisters, NumInputRegisters, OutputRegisters, NumOutputRegisters, DataSetInfo.StartInstance);
 		}
 	}
 
-	if (bRegisterSetupCompleted && GbExecVMScripts != 0)
+	if (GbExecVMScripts != 0 && !bError)
 	{
 		VectorVM::Exec(
 			Script->GetVMExecutableData().ByteCode.GetData(),
@@ -192,9 +228,9 @@ bool FNiagaraScriptExecutionContext::Execute(uint32 NumInstances, TArray<FNiagar
 	}
 
 	// Tell the datasets we wrote how many instances were actually written.
-	for (int Idx = 0; Idx < DataSetInfos.Num(); Idx++)
+	for (int Idx = 0; Idx < DataSetInfo.Num(); Idx++)
 	{
-		FNiagaraDataSetExecutionInfo& Info = DataSetInfos[Idx];
+		FNiagaraDataSetExecutionInfo& Info = DataSetInfo[Idx];
 
 #if NIAGARA_NAN_CHECKING
 		Info.DataSet->CheckForNaNs();
@@ -202,10 +238,12 @@ bool FNiagaraScriptExecutionContext::Execute(uint32 NumInstances, TArray<FNiagar
 
 		if (Info.bUpdateInstanceCount)
 		{
-			Info.DataSet->SetNumInstances(Info.StartInstance + DataSetMetaTable[Idx].DataSetAccessIndex + 1);
+			Info.Output->SetNumInstances(Info.StartInstance + DataSetMetaTable[Idx].DataSetAccessIndex + 1);
 		}
 	}
 
+	DataSetInfo.Reset();
+
 	return true;//TODO: Error cases?
 }
 
@@ -218,3 +256,303 @@ bool FNiagaraScriptExecutionContext::CanExecute()const
 {
 	return Script && Script->GetVMExecutableData().IsValid() && Script->GetVMExecutableData().ByteCode.Num() > 0;
 }
+
+struct H2
+{
+	FNiagaraDataInterfaceProxy* Proxy;
+};
+
+void FNiagaraGPUSystemTick::Init(FNiagaraSystemInstance* InSystemInstance)
+{
+	check(IsInGameThread());
+
+	ensure(InSystemInstance != nullptr);
+	CA_ASSUME(InSystemInstance != nullptr);
+	ensure(!InSystemInstance->IsComplete());
+	SystemInstanceID = InSystemInstance->GetId();
+	bRequiredDistanceFieldData = InSystemInstance->RequiresDistanceFieldData();
+	uint32 DataSizeForGPU = InSystemInstance->GPUDataInterfaceInstanceDataSize;
+
+	if (DataSizeForGPU > 0)
+	{
+		uint32 AllocationSize = DataSizeForGPU;
+
+		DIInstanceData = new FNiagaraDataInterfaceInstanceData;
+		DIInstanceData->PerInstanceDataSize = AllocationSize;
+		DIInstanceData->PerInstanceDataForRT = FMemory::Malloc(AllocationSize);
+		DIInstanceData->Instances = InSystemInstance->DataInterfaceInstanceDataOffsets.Num();
+
+		uint8* InstanceDataBase = (uint8*) DIInstanceData->PerInstanceDataForRT;
+		uint32 RunningOffset = 0;
+		for (auto& Pair : InSystemInstance->DataInterfaceInstanceDataOffsets)
+		{
+			UNiagaraDataInterface* Interface = Pair.Key.Get();
+
+			FNiagaraDataInterfaceProxy* Proxy = Interface->GetProxy();
+			int32 Offset = Pair.Value;
+
+			const int32 RTDataSize = Interface->PerInstanceDataPassedToRenderThreadSize();
+			if (RTDataSize > 0)
+			{
+				check(Proxy);
+				void* PerInstanceData = &InSystemInstance->DataInterfaceInstanceData[Offset];
+
+				Interface->ProvidePerInstanceDataForRenderThread(InstanceDataBase, PerInstanceData, SystemInstanceID);
+
+				// @todo rethink this. So ugly.
+				DIInstanceData->InterfaceProxiesToOffsets.Add(Proxy, RunningOffset);
+
+				InstanceDataBase += RTDataSize;
+				RunningOffset += RTDataSize;
+			}
+		}
+	}
+
+	check(MAX_uint32 > InSystemInstance->ActiveGPUEmitterCount);
+
+	// Layout our packet.
+	const uint32 PackedDispatchesSize = InSystemInstance->ActiveGPUEmitterCount * sizeof(FNiagaraComputeInstanceData);
+	// We want the Params after the instance data to be aligned so we can upload to the gpu.
+	uint32 PackedDispatchesSizeAligned = Align(PackedDispatchesSize, SHADER_PARAMETER_STRUCT_ALIGNMENT);
+	uint32 TotalParamSize = InSystemInstance->TotalParamSize;
+
+	uint32 TotalPackedBufferSize = PackedDispatchesSizeAligned + TotalParamSize;
+
+	InstanceData_ParamData_Packed = (uint8*)FMemory::Malloc(TotalPackedBufferSize);
+
+	FNiagaraComputeInstanceData* Instances = (FNiagaraComputeInstanceData*)(InstanceData_ParamData_Packed);
+	uint8* ParamDataBufferPtr = InstanceData_ParamData_Packed + PackedDispatchesSizeAligned;
+
+	// Now we will generate instance data for every GPU simulation we want to run on the render thread.
+	// This is spawn rate as well as DataInterface per instance data and the ParameterData for the emitter.
+	// @todo Ideally we would only update DataInterface and ParameterData bits if they have changed.
+	uint32 InstanceIndex = 0;
+	for (int32 i = 0; i < InSystemInstance->GetEmitters().Num(); i++)
+	{
+		FNiagaraEmitterInstance* Emitter = &InSystemInstance->GetEmitters()[i].Get();
+
+		if (Emitter->GetCachedEmitter()->SimTarget == ENiagaraSimTarget::GPUComputeSim && Emitter->GetGPUContext() != nullptr && Emitter->GetExecutionState() != ENiagaraExecutionState::Complete)
+		{
+			FNiagaraComputeInstanceData* InstanceData = new (&Instances[InstanceIndex]) FNiagaraComputeInstanceData;
+			InstanceIndex++;
+
+			InstanceData->Context = Emitter->GetGPUContext();
+			check(InstanceData->Context->MainDataSet);
+			InstanceData->SpawnRateInstances = Emitter->GetGPUContext()->SpawnRateInstances_GT;
+			InstanceData->EventSpawnTotal = Emitter->GetGPUContext()->EventSpawnTotal_GT;
+
+			int32 ParmSize = Emitter->GetGPUContext()->CombinedParamStore.GetPaddedParameterSizeInBytes();
+
+			Emitter->GetGPUContext()->CombinedParamStore.CopyParameterDataToPaddedBuffer(ParamDataBufferPtr, ParmSize);
+
+			InstanceData->ParamData = ParamDataBufferPtr;
+
+			ParamDataBufferPtr += ParmSize;
+
+			// @todo-threadsafety Think of a better way to do this!
+			FNiagaraComputeExecutionContext* GPUContext = Emitter->GetGPUContext();
+			const TArray<UNiagaraDataInterface*>& DataInterfaces = GPUContext->CombinedParamStore.GetDataInterfaces();
+			InstanceData->DataInterfaceProxies.Reserve(DataInterfaces.Num());
+			for (UNiagaraDataInterface* DI : DataInterfaces)
+			{
+				check(DI->GetProxy());
+				InstanceData->DataInterfaceProxies.Add(DI->GetProxy());
+			}			
+		}
+	}
+
+	check(InSystemInstance->ActiveGPUEmitterCount == InstanceIndex);
+	Count = InstanceIndex;
+}
+
+void FNiagaraGPUSystemTick::Destroy()
+{
+	FNiagaraComputeInstanceData* Instances = GetInstanceData();
+	for (uint32 i = 0; i < Count; i++)
+	{
+		FNiagaraComputeInstanceData& Instance = Instances[i];
+		Instance.~FNiagaraComputeInstanceData();
+	}
+
+	FMemory::Free(InstanceData_ParamData_Packed);
+	if (DIInstanceData)
+	{
+		FMemory::Free(DIInstanceData->PerInstanceDataForRT);
+		delete DIInstanceData;
+	}
+}
+
+//////////////////////////////////////////////////////////////////////////
+
+FNiagaraComputeExecutionContext::FNiagaraComputeExecutionContext()
+	: MainDataSet(nullptr)
+	, GPUScript(nullptr)
+	, GPUScript_RT(nullptr)
+	, CBufferLayout(TEXT("Niagara Compute Sim CBuffer"))
+	, DataToRender(nullptr)
+	, GPUDataReadback(nullptr)
+	, AccumulatedSpawnRate(0)
+	, NumIndicesPerInstance(0)
+#if WITH_EDITORONLY_DATA
+	, GPUDebugDataReadbackFloat(nullptr)
+	, GPUDebugDataReadbackInt(nullptr)
+	, GPUDebugDataReadbackCounts(nullptr)
+	, GPUDebugDataFloatSize(0)
+	, GPUDebugDataIntSize(0)
+#endif	  
+{
+}
+
+FNiagaraComputeExecutionContext::~FNiagaraComputeExecutionContext()
+{
+	checkf(IsInRenderingThread(), TEXT("Can only delete the gpu readback from the render thread"));
+	if (GPUDataReadback)
+	{
+		delete GPUDataReadback;
+		GPUDataReadback = nullptr;
+	}
+
+#if WITH_EDITORONLY_DATA
+	if (GPUDebugDataReadbackFloat)
+	{
+		delete GPUDebugDataReadbackFloat;
+		GPUDebugDataReadbackFloat = nullptr;
+	}
+	if (GPUDebugDataReadbackInt)
+	{
+		delete GPUDebugDataReadbackInt;
+		GPUDebugDataReadbackInt = nullptr;
+	}
+	if (GPUDebugDataReadbackCounts)
+	{
+		delete GPUDebugDataReadbackCounts;
+		GPUDebugDataReadbackCounts = nullptr;
+	}
+#endif
+
+	SetDataToRender(nullptr);
+}
+
+void FNiagaraComputeExecutionContext::Reset()
+{
+	FNiagaraComputeExecutionContext* Context = this;
+	ENQUEUE_RENDER_COMMAND(ResetRT)(
+		[Context](FRHICommandListImmediate& RHICmdList)
+	{
+		Context->ResetInternal();
+	}
+	);
+}
+
+void FNiagaraComputeExecutionContext::InitParams(UNiagaraScript* InGPUComputeScript, ENiagaraSimTarget InSimTarget, const FString& InDebugSimName)
+{
+	DebugSimName = InDebugSimName;
+	GPUScript = InGPUComputeScript;
+	CombinedParamStore.InitFromOwningContext(InGPUComputeScript, InSimTarget, true);
+
+#if DO_CHECK
+	FNiagaraShader *Shader = InGPUComputeScript->GetRenderThreadScript()->GetShaderGameThread();
+	DIParamInfo.Empty();
+	if (Shader)
+	{
+		for (FNiagaraDataInterfaceParamRef& DIParams : Shader->GetDIParameters())
+		{
+			DIParamInfo.Add(DIParams.ParameterInfo);
+		}
+	}
+	else
+	{
+		DIParamInfo = InGPUComputeScript->GetRenderThreadScript()->GetDataInterfaceParamInfo();
+	}
+#endif
+}
+
+void FNiagaraComputeExecutionContext::DirtyDataInterfaces()
+{
+	CombinedParamStore.MarkInterfacesDirty();
+}
+
+bool FNiagaraComputeExecutionContext::Tick(FNiagaraSystemInstance* ParentSystemInstance)
+{
+	if (CombinedParamStore.GetInterfacesDirty())
+	{
+#if DO_CHECK
+		const TArray<UNiagaraDataInterface*> &DataInterfaces = CombinedParamStore.GetDataInterfaces();
+		// We must make sure that the data interfaces match up between the original script values and our overrides...
+		if (DIParamInfo.Num() != DataInterfaces.Num())
+		{
+			UE_LOG(LogNiagara, Warning, TEXT("Mismatch between Niagara GPU Execution Context data interfaces and those in its script!"));
+			return false;
+		}
+
+		for (int32 i = 0; i < DIParamInfo.Num(); ++i)
+		{
+			FString UsedClassName = DataInterfaces[i]->GetClass()->GetName();
+			if (DIParamInfo[i].DIClassName != UsedClassName)
+			{
+				UE_LOG(LogNiagara, Warning, TEXT("Mismatched class between Niagara GPU Execution Context data interfaces and those in its script!\nIndex:%d\nShader:%s\nScript:%s")
+					, i, *DIParamInfo[i].DIClassName, *UsedClassName);
+			}
+		}
+#endif
+
+		CombinedParamStore.Tick();
+	}
+
+	return true;
+}
+
+void FNiagaraComputeExecutionContext::PostTick()
+{
+	//If we're for interpolated spawn, copy over the previous frame's parameters into the Prev parameters.
+	if (GPUScript && GPUScript->GetComputedVMCompilationId().HasInterpolatedParameters())
+	{
+		CombinedParamStore.CopyCurrToPrev();
+	}
+}
+
+void FNiagaraComputeExecutionContext::ResetInternal()
+{
+	checkf(IsInRenderingThread(), TEXT("Can only reset the gpu context from the render thread"));
+	AccumulatedSpawnRate = 0;
+	if (GPUDataReadback)
+	{
+		delete GPUDataReadback;
+		GPUDataReadback = nullptr;
+	}
+
+#if WITH_EDITORONLY_DATA
+	if (GPUDebugDataReadbackFloat)
+	{
+		delete GPUDebugDataReadbackFloat;
+		GPUDebugDataReadbackFloat = nullptr;
+	}
+	if (GPUDebugDataReadbackInt)
+	{
+		delete GPUDebugDataReadbackInt;
+		GPUDebugDataReadbackInt = nullptr;
+	}
+	if (GPUDebugDataReadbackCounts)
+	{
+		delete GPUDebugDataReadbackCounts;
+		GPUDebugDataReadbackCounts = nullptr;
+	}
+#endif
+
+	SetDataToRender(nullptr);
+}
+
+void FNiagaraComputeExecutionContext::SetDataToRender(FNiagaraDataBuffer* InDataToRender)
+{
+	if (DataToRender)
+	{
+		DataToRender->ReleaseReadRef();
+	}
+
+	DataToRender = InDataToRender;
+
+	if (DataToRender)
+	{
+		DataToRender->AddReadRef();
+	}
+}
\ No newline at end of file
diff --git a/Engine/Plugins/FX/Niagara/Source/Niagara/Private/NiagaraSpriteRendererProperties.cpp b/Engine/Plugins/FX/Niagara/Source/Niagara/Private/NiagaraSpriteRendererProperties.cpp
--- a/Engine/Plugins/FX/Niagara/Source/Niagara/Private/NiagaraSpriteRendererProperties.cpp
+++ b/Engine/Plugins/FX/Niagara/Source/Niagara/Private/NiagaraSpriteRendererProperties.cpp
@@ -4,6 +4,8 @@
 #include "NiagaraRenderer.h"
 #include "Internationalization/Internationalization.h"
 #include "NiagaraConstants.h"
+#include "NiagaraRendererSprites.h"
+#include "NiagaraBoundsCalculatorHelper.h"
 #if WITH_EDITOR
 #include "DerivedDataCacheInterface.h"
 #endif
@@ -37,7 +39,6 @@ UNiagaraSpriteRendererProperties::UNiagaraSpriteRendererProperties()
 	, bSortOnlyWhenTranslucent(true)
 	, MinFacingCameraBlendDistance(0.0f)
 	, MaxFacingCameraBlendDistance(0.0f)
-	, SyncId(0)
 #if WITH_EDITORONLY_DATA
 	, BoundingMode(BVC_EightVertices)
 	, AlphaThreshold(0.1f)
@@ -45,9 +46,16 @@ UNiagaraSpriteRendererProperties::UNiagaraSpriteRendererProperties()
 {
 }
 
-NiagaraRenderer* UNiagaraSpriteRendererProperties::CreateEmitterRenderer(ERHIFeatureLevel::Type FeatureLevel)
+FNiagaraRenderer* UNiagaraSpriteRendererProperties::CreateEmitterRenderer(ERHIFeatureLevel::Type FeatureLevel, const FNiagaraEmitterInstance* Emitter)
 {
-	return new NiagaraRendererSprites(FeatureLevel, this);
+	FNiagaraRenderer* NewRenderer = new FNiagaraRendererSprites(FeatureLevel, this, Emitter);	
+	NewRenderer->Initialize(FeatureLevel, this, Emitter);
+	return NewRenderer;
+}
+
+FNiagaraBoundsCalculator* UNiagaraSpriteRendererProperties::CreateBoundsCalculator()
+{
+	return new FNiagaraBoundsCalculatorHelper<true, false, false>();
 }
 
 void UNiagaraSpriteRendererProperties::GetUsedMaterials(TArray<UMaterialInterface*>& OutMaterials) const
@@ -74,7 +82,6 @@ void UNiagaraSpriteRendererProperties::PostLoad()
 void UNiagaraSpriteRendererProperties::PostInitProperties()
 {
 	Super::PostInitProperties();
-	SyncId = 0;
 	if (HasAnyFlags(RF_ClassDefaultObject) == false)
 	{
 		InitBindings();
@@ -149,12 +156,7 @@ void UNiagaraSpriteRendererProperties::PostEditChangeProperty(struct FPropertyCh
 			CacheDerivedData();
 		}
 	}
-
-	if (PropertyChangedEvent.GetPropertyName() != TEXT("SyncId"))
-	{
-		SyncId++;
-	}
-
+	
 	Super::PostEditChangeProperty(PropertyChangedEvent);
 }
 
diff --git a/Engine/Plugins/FX/Niagara/Source/Niagara/Private/NiagaraStats.h b/Engine/Plugins/FX/Niagara/Source/Niagara/Private/NiagaraStats.h
--- a/Engine/Plugins/FX/Niagara/Source/Niagara/Private/NiagaraStats.h
+++ b/Engine/Plugins/FX/Niagara/Source/Niagara/Private/NiagaraStats.h
@@ -8,9 +8,9 @@
 #include "NiagaraModule.h"
 
 DECLARE_STATS_GROUP(TEXT("Niagara"), STATGROUP_Niagara, STATCAT_Advanced);
-DECLARE_CYCLE_STAT(TEXT("Render Total [RT]"), STAT_NiagaraRender, STATGROUP_Niagara);
-DECLARE_CYCLE_STAT(TEXT("Render Total [GT]"), STAT_NiagaraRenderGT, STATGROUP_Niagara);
-DECLARE_DWORD_COUNTER_STAT(TEXT("NumParticles"), STAT_NiagaraNumParticles, STATGROUP_Niagara);
+DECLARE_CYCLE_STAT(TEXT("Niagara Render Total [RT]"), STAT_NiagaraRender, STATGROUP_Niagara);
+DECLARE_CYCLE_STAT(TEXT("Niagara Render Total [GT]"), STAT_NiagaraRenderGT, STATGROUP_Niagara);
+DECLARE_DWORD_COUNTER_STAT(TEXT("# CPU Particles"), STAT_NiagaraNumParticles, STATGROUP_Niagara);
 DECLARE_CYCLE_STAT(TEXT("Constant Setup [CNC]"), STAT_NiagaraConstants, STATGROUP_Niagara);
 DECLARE_CYCLE_STAT(TEXT("Emitter Tick [CNC]"), STAT_NiagaraTick, STATGROUP_Niagara);
 DECLARE_DWORD_COUNTER_STAT(TEXT("NumSystems"), STAT_NiagaraNumSystems, STATGROUP_Niagara);
@@ -22,4 +22,8 @@ DECLARE_STATS_GROUP(TEXT("NiagaraOverview"), STATGROUP_NiagaraOverview, STATCAT_
 DECLARE_CYCLE_STAT(TEXT("GT Total"), STAT_NiagaraOverview_GT, STATGROUP_NiagaraOverview);
 DECLARE_CYCLE_STAT(TEXT("GT Concurrent Total"), STAT_NiagaraOverview_GT_CNC, STATGROUP_NiagaraOverview);
 DECLARE_CYCLE_STAT(TEXT("RT Total"), STAT_NiagaraOverview_RT, STATGROUP_NiagaraOverview);
-DECLARE_CYCLE_STAT(TEXT("RT Concurrent Total"), STAT_NiagaraOverview_RT_CNC, STATGROUP_NiagaraOverview);
\ No newline at end of file
+DECLARE_CYCLE_STAT(TEXT("RT Concurrent Total"), STAT_NiagaraOverview_RT_CNC, STATGROUP_NiagaraOverview);
+
+
+DECLARE_STATS_GROUP(TEXT("Niagara Systems"), STATGROUP_NiagaraSystems, STATCAT_Advanced);
+DECLARE_STATS_GROUP(TEXT("Niagara Emitters"), STATGROUP_NiagaraEmitters, STATCAT_Advanced);
\ No newline at end of file
diff --git a/Engine/Plugins/FX/Niagara/Source/Niagara/Private/NiagaraSystem.cpp b/Engine/Plugins/FX/Niagara/Source/Niagara/Private/NiagaraSystem.cpp
--- a/Engine/Plugins/FX/Niagara/Source/Niagara/Private/NiagaraSystem.cpp
+++ b/Engine/Plugins/FX/Niagara/Source/Niagara/Private/NiagaraSystem.cpp
@@ -91,6 +91,8 @@ void UNiagaraSystem::PostInitProperties()
 		SystemUpdateScript = NewObject<UNiagaraScript>(this, "SystemUpdateScript", RF_Transactional);
 		SystemUpdateScript->SetUsage(ENiagaraScriptUsage::SystemUpdateScript);
 	}
+
+	GenerateStatID();
 }
 
 bool UNiagaraSystem::IsLooping() const
@@ -460,7 +462,7 @@ bool UNiagaraSystem::ReferencesInstanceEmitter(UNiagaraEmitter& Emitter)
 	return false;
 }
 
-void UNiagaraSystem::UpdateFromEmitterChanges(UNiagaraEmitter& ChangedSourceEmitter)
+void UNiagaraSystem::UpdateFromEmitterChanges(UNiagaraEmitter& ChangedSourceEmitter, bool bRecompileOnChange)
 {
 	bool bNeedsCompile = false;
 	for(FNiagaraEmitterHandle& EmitterHandle : EmitterHandles)
@@ -472,7 +474,7 @@ void UNiagaraSystem::UpdateFromEmitterChanges(UNiagaraEmitter& ChangedSourceEmit
 		}
 	}
 
-	if (bNeedsCompile)
+	if (bNeedsCompile && bRecompileOnChange)
 	{
 		RequestCompile(false);
 	}
@@ -1072,4 +1074,43 @@ void UNiagaraSystem::InitEmitterSpawnAttributes()
 			}
 		}
 	}
+}
+
+TStatId UNiagaraSystem::GetStatID(bool bGameThread, bool bConcurrent)const
+{
+#if STATS
+	if (bGameThread)
+	{
+		if (bConcurrent)
+		{
+			return StatID_GT;
+		}
+		else
+		{
+			return StatID_GT_CNC;
+		}
+	}
+	else
+	{
+		if(bConcurrent)
+		{
+			return StatID_RT;
+		}
+		else
+		{
+			return StatID_RT_CNC;
+		}
+	}
+#endif
+	return TStatId();
+}
+
+void UNiagaraSystem::GenerateStatID()
+{
+#if STATS
+	StatID_GT = FDynamicStats::CreateStatId<FStatGroup_STATGROUP_NiagaraSystems>(GetPathName() + TEXT("[GT]"));
+	StatID_GT_CNC = FDynamicStats::CreateStatId<FStatGroup_STATGROUP_NiagaraSystems>(GetPathName() + TEXT("[GT_CNC]"));
+	StatID_RT = FDynamicStats::CreateStatId<FStatGroup_STATGROUP_NiagaraSystems>(GetPathName() + TEXT("[RT]"));
+	StatID_RT_CNC = FDynamicStats::CreateStatId<FStatGroup_STATGROUP_NiagaraSystems>(GetPathName() + TEXT("[RT_CNC]"));
+#endif
 }
\ No newline at end of file
diff --git a/Engine/Plugins/FX/Niagara/Source/Niagara/Private/NiagaraSystemInstance.cpp b/Engine/Plugins/FX/Niagara/Source/Niagara/Private/NiagaraSystemInstance.cpp
--- a/Engine/Plugins/FX/Niagara/Source/Niagara/Private/NiagaraSystemInstance.cpp
+++ b/Engine/Plugins/FX/Niagara/Source/Niagara/Private/NiagaraSystemInstance.cpp
@@ -218,12 +218,13 @@ bool FNiagaraSystemInstance::QueryCaptureResults(const FGuid& RequestId, TArray<
 				if ((*Array)[i]->bWaitForGPU && (*Array)[i]->bWritten == false)
 				{
 					bWaitForGPU = true;
+					break;
 				}
 			}
 			
 			if (bWaitForGPU)
 			{
-				for (TSharedRef<FNiagaraEmitterInstance> CachedEmitter : Emitters)
+				for (TSharedRef<FNiagaraEmitterInstance, ESPMode::ThreadSafe> CachedEmitter : Emitters)
 				{
 					CachedEmitter->WaitForDebugInfo();
 				}
@@ -334,6 +335,38 @@ void FNiagaraSystemInstance::Deactivate(bool bImmediate)
 	}
 }
 
+bool FNiagaraSystemInstance::AllocateSystemInstance(class UNiagaraComponent* InComponent, TUniquePtr< FNiagaraSystemInstance >& OutSystemInstanceAllocation)
+{
+	OutSystemInstanceAllocation = MakeUnique<FNiagaraSystemInstance>(InComponent);
+	return true;
+}
+
+bool FNiagaraSystemInstance::DeallocateSystemInstance(TUniquePtr< FNiagaraSystemInstance >& SystemInstanceAllocation)
+{
+	if (SystemInstanceAllocation.IsValid())
+	{
+		TSharedPtr<FNiagaraSystemSimulation, ESPMode::ThreadSafe> SystemSim = SystemInstanceAllocation->GetSystemSimulation();
+
+		// Make sure we remove the instance
+		if (SystemInstanceAllocation->SystemInstanceIndex != INDEX_NONE)
+		{
+			SystemSim->RemoveInstance(SystemInstanceAllocation.Get());
+		}
+
+		// Queue deferred deletion from the WorldManager
+		FNiagaraWorldManager* WorldManager = SystemInstanceAllocation->GetWorldManager();
+		check(WorldManager != nullptr);
+
+		SystemInstanceAllocation->Component = nullptr;
+
+		WorldManager->DestroySystemInstance(SystemInstanceAllocation);
+		check(SystemInstanceAllocation == nullptr);
+	}
+	SystemInstanceAllocation = nullptr;
+	
+	return true;
+}
+
 void FNiagaraSystemInstance::Complete()
 {
 	SCOPE_CYCLE_COUNTER(STAT_NiagaraSystemComplete);
@@ -341,6 +374,8 @@ void FNiagaraSystemInstance::Complete()
 	// Only notify others if have yet to complete
 	bool bNeedToNotifyOthers = bNotifyOnCompletion;
 
+	//UE_LOG(LogNiagara, Log, TEXT("FNiagaraSystemInstance::Complete { %p"), this);
+
 	if (SystemInstanceIndex != INDEX_NONE)
 	{
 		TSharedPtr<FNiagaraSystemSimulation, ESPMode::ThreadSafe> SystemSim = GetSystemSimulation();
@@ -349,7 +384,7 @@ void FNiagaraSystemInstance::Complete()
 		SetActualExecutionState(ENiagaraExecutionState::Complete);
 		SetRequestedExecutionState(ENiagaraExecutionState::Complete);
 
-		for (TSharedRef<FNiagaraEmitterInstance> Simulation : Emitters)
+		for (TSharedRef<FNiagaraEmitterInstance, ESPMode::ThreadSafe> Simulation : Emitters)
 		{
 			Simulation->HandleCompletion(true);
 		}
@@ -366,16 +401,22 @@ void FNiagaraSystemInstance::Complete()
 
 	if (bNeedToNotifyOthers)
 	{
-		// We've already notified once, no need to do so again.
-		bNotifyOnCompletion = false;
-
+		//UE_LOG(LogNiagara, Log, TEXT("FNiagaraSystemInstance::BroadCast OnCompleteDelegate { %p"), this);
 		OnCompleteDelegate.Broadcast(this);
+		//UE_LOG(LogNiagara, Log, TEXT("FNiagaraSystemInstance::BroadCast OnCompleteDelegate } %p"), this);
 
 		if (Component)
 		{
-			// Note: This call may destroy this instance of FNiagaraSystemInstance, so don't use bNotifyOnCompletion after it!
+			//UE_LOG(LogNiagara, Log, TEXT("FNiagaraSystemInstance::Component->OnSystemComplete { %p"), this);
 			Component->OnSystemComplete();
+			//UE_LOG(LogNiagara, Log, TEXT("FNiagaraSystemInstance::Component->OnSystemComplete } %p"), this);
 		}
+
+		//UE_LOG(LogNiagara, Log, TEXT("FNiagaraSystemInstance::Complete } %p"), this);
+
+		
+		// We've already notified once, no need to do so again.
+		bNotifyOnCompletion = false;
 	}
 }
 
@@ -386,16 +427,19 @@ void FNiagaraSystemInstance::SetPaused(bool bInPaused)
 		return;
 	}
 	
-	FNiagaraSystemSimulation* SystemSim = GetSystemSimulation().Get();
-	if (SystemSim)
+	if (SystemInstanceIndex != INDEX_NONE)
 	{
-		if (bInPaused)
+		FNiagaraSystemSimulation* SystemSim = GetSystemSimulation().Get();
+		if (SystemSim)
 		{
-			SystemSim->PauseInstance(this);
-		}
-		else
-		{
-			SystemSim->UnpauseInstance(this);
+			if (bInPaused)
+			{
+				SystemSim->PauseInstance(this);
+			}
+			else
+			{
+				SystemSim->UnpauseInstance(this);
+			}
 		}
 	}
 
@@ -415,8 +459,11 @@ void FNiagaraSystemInstance::Reset(FNiagaraSystemInstance::EResetMode Mode, bool
 		}*/
 		return;
 	}
-
-	Component->SetLastRenderTime(Component->GetWorld()->GetTimeSeconds());
+	
+	if (Component)
+	{
+		Component->SetLastRenderTime(Component->GetWorld()->GetTimeSeconds());
+	}
 
 	SetPaused(false);
 
@@ -461,7 +508,11 @@ void FNiagaraSystemInstance::Reset(FNiagaraSystemInstance::EResetMode Mode, bool
 	SetRequestedExecutionState(ENiagaraExecutionState::Active);
 	SetActualExecutionState(ENiagaraExecutionState::Active);
 
-	InitDataInterfaces();
+	// We avoid calling InitDataInterfaces in the ResetSystem path so as to not clear out the System's DI on this frame.
+	if (Mode == EResetMode::ResetAll || Mode == EResetMode::ReInit)
+	{
+		InitDataInterfaces();
+	}
 
 	//Interface init can disable the system.
 	if (!IsComplete())
@@ -483,9 +534,12 @@ void FNiagaraSystemInstance::Reset(FNiagaraSystemInstance::EResetMode Mode, bool
 		}
 	}
 
-	// This system may not tick again immediately so we mark the render state dirty here so that
-	// the renderers will be reset this frame.
-	Component->MarkRenderDynamicDataDirty();
+	if (Component)
+	{
+		// This system may not tick again immediately so we mark the render state dirty here so that
+		// the renderers will be reset this frame.
+		Component->MarkRenderDynamicDataDirty();
+	}
 }
 
 void FNiagaraSystemInstance::ResetInternal(bool bResetSimulations)
@@ -519,7 +573,7 @@ void FNiagaraSystemInstance::ResetInternal(bool bResetSimulations)
 		return;
 	}
 
-	for (TSharedRef<FNiagaraEmitterInstance> Simulation : Emitters)
+	for (TSharedRef<FNiagaraEmitterInstance, ESPMode::ThreadSafe> Simulation : Emitters)
 	{
 		Simulation->ResetSimulation(bResetSimulations);
 	}
@@ -562,7 +616,7 @@ bool FNiagaraSystemInstance::IsReadyToRun() const
 		return false;
 	}
 
-	for (TSharedRef<FNiagaraEmitterInstance> Simulation : Emitters)
+	for (TSharedRef<FNiagaraEmitterInstance, ESPMode::ThreadSafe> Simulation : Emitters)
 	{
 		if (!Simulation->IsReadyToRun())
 		{
@@ -679,7 +733,7 @@ void FNiagaraSystemInstance::ReInitInternal()
 	TArray<FNiagaraVariable> TotalSpawnedParticlesVars;
 	for (int32 i = 0; i < Emitters.Num(); i++)
 	{
-		TSharedRef<FNiagaraEmitterInstance> Simulation = Emitters[i];
+		TSharedRef<FNiagaraEmitterInstance, ESPMode::ThreadSafe> Simulation = Emitters[i];
 		FString EmitterName = Simulation->GetEmitterHandle().GetInstance()->GetUniqueEmitterName();
 		
 		{
@@ -748,28 +802,9 @@ void FNiagaraSystemInstance::ReInitInternal()
 
 	TickInstanceParameters(0.01f);
 
-	// This gets a little tricky, but we want to delete any renderers that are no longer in use on the rendering thread, but
-	// first (to be safe), we want to update the proxy to point to the new renderer objects.
-
-	// Step 1: Recreate the renderers on the simulations, we keep the old and new renderers.
-	TArray<NiagaraRenderer*> NewRenderers;
-	TArray<NiagaraRenderer*> OldRenderers;
-
-	UpdateRenderModules(Component->GetWorld()->FeatureLevel, NewRenderers, OldRenderers);
-
-	// Step 2: Update the proxy with the new renderers that were created.
-	UpdateProxy(NewRenderers);
+	//Invalidate the component render state so we recreate the scene proxy and the renderers.
 	Component->MarkRenderStateDirty();
 
-	// Step 3: Queue up the old renderers for deletion on the render thread.
-	for (NiagaraRenderer* Renderer : OldRenderers)
-	{
-		if (Renderer != nullptr)
-		{
-			Renderer->Release();
-		}
-	}
-
 #if WITH_EDITOR
 	//UE_LOG(LogNiagara, Log, TEXT("OnResetInternal %p"), this);
 	OnResetDelegate.Broadcast();
@@ -779,7 +814,7 @@ void FNiagaraSystemInstance::ReInitInternal()
 
 FNiagaraSystemInstance::~FNiagaraSystemInstance()
 {
-	//UE_LOG(LogNiagara, Warning, TEXT("~FNiagaraSystemInstance %p"), this);
+	//UE_LOG(LogNiagara, Log, TEXT("~FNiagaraSystemInstance %p"), this);
 
 	//FlushRenderingCommands();
 
@@ -800,16 +835,6 @@ void FNiagaraSystemInstance::Cleanup()
 
 	DestroyDataInterfaceInstanceData();
 
-	// Clear out the System renderer from the proxy.
-	TArray<NiagaraRenderer*> NewRenderers;
-	UpdateProxy(NewRenderers);
-
-	// Clear out the System renderer from the simulation.
-	for (TSharedRef<FNiagaraEmitterInstance> Simulation : Emitters)
-	{
-		Simulation->ClearRenderer();
-	}
-
 	UnbindParameters();
 
 	// Clear out the emitters.
@@ -822,7 +847,7 @@ void FNiagaraSystemInstance::Cleanup()
 // 	OldInstance->GetParameterStore().Unbind(&InstanceParameters);
 // 	NewInstance->GetParameterStore().Bind(&InstanceParameters);
 // 
-// 	for (TSharedRef<FNiagaraEmitterInstance> Simulation : Emitters)
+// 	for (TSharedRef<FNiagaraEmitterInstance, ESPMode::ThreadSafe> Simulation : Emitters)
 // 	{
 // 		Simulation->RebindParameterCollection(OldInstance, NewInstance);
 // 	}
@@ -834,6 +859,11 @@ void FNiagaraSystemInstance::Cleanup()
 
 void FNiagaraSystemInstance::BindParameters()
 {
+	if (!Component)
+	{
+		return;
+	}
+
 	Component->GetOverrideParameters().Bind(&InstanceParameters);
 
 	if (SystemSimulation->GetIsSolo())
@@ -851,7 +881,7 @@ void FNiagaraSystemInstance::BindParameters()
 		GetSystem()->GetExposedParameters().Bind(&SystemSimulation->GetSpawnExecutionContext().Parameters);
 	}
 
-	for (TSharedRef<FNiagaraEmitterInstance> Simulation : Emitters)
+	for (TSharedRef<FNiagaraEmitterInstance, ESPMode::ThreadSafe> Simulation : Emitters)
 	{
 		Simulation->BindParameters();
 	}
@@ -859,14 +889,21 @@ void FNiagaraSystemInstance::BindParameters()
 
 void FNiagaraSystemInstance::UnbindParameters()
 {
-	Component->GetOverrideParameters().Unbind(&InstanceParameters);
+
+	if (Component)
+	{
+		Component->GetOverrideParameters().Unbind(&InstanceParameters);
+	}
 
 	if (SystemSimulation.IsValid())
 	{
 		if (SystemSimulation->GetIsSolo())
 		{
-			Component->GetOverrideParameters().Unbind(&SystemSimulation->GetSpawnExecutionContext().Parameters);
-			Component->GetOverrideParameters().Unbind(&SystemSimulation->GetUpdateExecutionContext().Parameters);
+			if (Component)
+			{
+				Component->GetOverrideParameters().Unbind(&SystemSimulation->GetSpawnExecutionContext().Parameters);
+				Component->GetOverrideParameters().Unbind(&SystemSimulation->GetUpdateExecutionContext().Parameters);
+			}
 		}
 		else 
 		{
@@ -879,7 +916,7 @@ void FNiagaraSystemInstance::UnbindParameters()
 		}
 	}
 
-	for (TSharedRef<FNiagaraEmitterInstance> Simulation : Emitters)
+	for (TSharedRef<FNiagaraEmitterInstance, ESPMode::ThreadSafe> Simulation : Emitters)
 	{
 		Simulation->UnbindParameters();
 	}
@@ -890,6 +927,31 @@ FNiagaraWorldManager* FNiagaraSystemInstance::GetWorldManager()const
 	return Component ? FNiagaraWorldManager::Get(Component->GetWorld()) : nullptr; 
 }
 
+bool FNiagaraSystemInstance::RequiresDistanceFieldData() const
+{
+	if (!bHasGPUEmitters)
+	{
+		return false;
+	}
+
+	for (const TSharedRef<FNiagaraEmitterInstance, ESPMode::ThreadSafe>& Emitter : Emitters)
+	{
+		FNiagaraComputeExecutionContext* GPUContext = Emitter->GetGPUContext();
+		if (GPUContext)
+		{
+			for (UNiagaraDataInterface* DataInterface : GPUContext->CombinedParamStore.GetDataInterfaces())
+			{
+				if (DataInterface && DataInterface->RequiresDistanceFieldData())
+				{
+					return true;
+				}
+			}
+		}
+	}
+
+	return false;
+}
+
 void FNiagaraSystemInstance::InitDataInterfaces()
 {
 	// If either the System or the component is invalid, it is possible that our cached data interfaces
@@ -909,6 +971,8 @@ void FNiagaraSystemInstance::InitDataInterfaces()
 	
 	DestroyDataInterfaceInstanceData();
 
+	GPUDataInterfaceInstanceDataSize = 0;
+
 	//Now the interfaces in the simulations are all correct, we can build the per instance data table.
 	int32 InstanceDataSize = 0;
 	DataInterfaceInstanceDataOffsets.Empty();
@@ -947,7 +1011,7 @@ void FNiagaraSystemInstance::InitDataInterfaces()
 	}
 
 	//Iterate over interfaces to get size for table and clear their interface bindings.
-	for (TSharedRef<FNiagaraEmitterInstance> Simulation : Emitters)
+	for (TSharedRef<FNiagaraEmitterInstance, ESPMode::ThreadSafe> Simulation : Emitters)
 	{
 		FNiagaraEmitterInstance& Sim = Simulation.Get();
 		CalcInstDataSize(Sim.GetSpawnExecutionContext().GetDataInterfaces());
@@ -970,6 +1034,8 @@ void FNiagaraSystemInstance::InitDataInterfaces()
 		{
 			check(IsAligned(&DataInterfaceInstanceData[Pair.Value], 16));
 
+			GPUDataInterfaceInstanceDataSize += Pair.Key->PerInstanceDataPassedToRenderThreadSize();
+
 			//Ideally when we make the batching changes, we can keep the instance data in big single type blocks that can all be updated together with a single virtual call.
 			bool bResult = Pair.Key->InitPerInstanceData(&DataInterfaceInstanceData[Pair.Value], this);
 			bDataInterfacesInitialized &= bResult;
@@ -1098,6 +1164,11 @@ float FNiagaraSystemInstance::GetLODDistance()
 
 void FNiagaraSystemInstance::TickInstanceParameters(float DeltaSeconds)
 {
+	if (!Component)
+	{
+		return;
+	}
+
 	//TODO: Create helper binding objects to avoid the search in set parameter value.
 	//Set System params.
 	FTransform ComponentTrans = Component->GetComponentTransform();
@@ -1147,15 +1218,19 @@ void FNiagaraSystemInstance::TickInstanceParameters(float DeltaSeconds)
 	SystemTickCountParam.SetValue(TickCount);
 
 	int32 NumAlive = 0;
+	const TArray<FNiagaraEmitterHandle>& EmitterHandles = GetSystem()->GetEmitterHandles();
 	for (int32 i = 0; i < Emitters.Num(); i++)
 	{
-		int32 NumParticles = Emitters[i]->GetNumParticles();
-		if (!Emitters[i]->IsComplete())
+		if (EmitterHandles[i].GetIsEnabled())//TODO: We should just null out the entry to the emitter in the array.
 		{
-			NumAlive++;
+			int32 NumParticles = Emitters[i]->GetNumParticles();
+			if (!Emitters[i]->IsComplete())
+			{
+				NumAlive++;
+			}
+			ParameterNumParticleBindings[i].SetValue(NumParticles);
+			ParameterTotalSpawnedParticlesBindings[i].SetValue(Emitters[i]->GetTotalSpawnedParticles());
 		}
-		ParameterNumParticleBindings[i].SetValue(NumParticles);
-		ParameterTotalSpawnedParticlesBindings[i].SetValue(Emitters[i]->GetTotalSpawnedParticles());
 	}
 	SystemNumEmittersParam.SetValue(Emitters.Num());
 	SystemNumEmittersAliveParam.SetValue(NumAlive);
@@ -1226,65 +1301,29 @@ void FNiagaraSystemInstance::InitEmitters()
 		Component->MarkRenderStateDirty();
 	}
 
-	//STovey: This is done after InitEmitters so removing this duplication.
-// 	// Just in case this ends up being called more than in Init, we need to 
-// 	// clear out the update proxy of any renderers that will be destroyed when Emitters.Empty occurs..
-// 	TArray<NiagaraRenderer*> NewRenderers;
-// 	UpdateProxy(NewRenderers);
-// 
-// 	// Clear out the System renderer from the simulation.
-// 	for (TSharedRef<FNiagaraEmitterInstance> Simulation : Emitters)
-// 	{
-// 		Simulation->ClearRenderer();
-// 	}
-
 	Emitters.Empty();
-	if (GetSystem() != nullptr)
+	UNiagaraSystem* System = GetSystem();
+	if (System != nullptr)
 	{
 		const TArray<FNiagaraEmitterHandle>& EmitterHandles = GetSystem()->GetEmitterHandles();
 		for (int32 EmitterIdx=0; EmitterIdx < GetSystem()->GetEmitterHandles().Num(); ++EmitterIdx)
 		{
 			const FNiagaraEmitterHandle& EmitterHandle = EmitterHandles[EmitterIdx];
-			TSharedRef<FNiagaraEmitterInstance> Sim = MakeShareable(new FNiagaraEmitterInstance(this));
+
+			TSharedRef<FNiagaraEmitterInstance, ESPMode::ThreadSafe> Sim = MakeShared<FNiagaraEmitterInstance, ESPMode::ThreadSafe>(this);
 			Sim->Init(EmitterIdx, IDName);
+			if (System->bFixedBounds)
+			{
+				Sim->SetSystemFixedBoundsOverride(System->GetFixedBounds());
+			}
 			Emitters.Add(Sim);
 		}
 
-		for (TSharedRef<FNiagaraEmitterInstance> Simulation : Emitters)
+		for (TSharedRef<FNiagaraEmitterInstance, ESPMode::ThreadSafe> Simulation : Emitters)
 		{
-			Simulation->PostInitSimulation();
-		}
-	}
-}
+			bHasGPUEmitters |= Simulation->GetCachedEmitter()->SimTarget == ENiagaraSimTarget::GPUComputeSim;
 
-void FNiagaraSystemInstance::UpdateRenderModules(ERHIFeatureLevel::Type InFeatureLevel, TArray<NiagaraRenderer*>& OutNewRenderers, TArray<NiagaraRenderer*>& OutOldRenderers)
-{
-	for (TSharedPtr<FNiagaraEmitterInstance> Sim : Emitters)
-	{
-		Sim->UpdateEmitterRenderer(InFeatureLevel, OutNewRenderers, OutOldRenderers);
-	}
-}
-
-void FNiagaraSystemInstance::UpdateProxy(TArray<NiagaraRenderer*>& InRenderers)
-{
-	if (!Component)
-	{
-		return;
-	}
-
-	FNiagaraSceneProxy *NiagaraProxy = static_cast<FNiagaraSceneProxy*>(Component->SceneProxy);
-	if (NiagaraProxy)
-	{
-		if (Component->GetWorld() != nullptr)
-		{
-			// Tell the scene proxy on the render thread to update its System renderers.
-			TArray<NiagaraRenderer*> InRendererArray = InRenderers;
-			ENQUEUE_RENDER_COMMAND(FChangeNiagaraRenderModule)(
-				[NiagaraProxy, InRendererArray](FRHICommandListImmediate& RHICmdList)
-				{
-					NiagaraProxy->UpdateEmitterRenderers(InRendererArray);
-				}
-			);
+			Simulation->PostInitSimulation();
 		}
 	}
 }
@@ -1311,10 +1350,13 @@ void FNiagaraSystemInstance::FinalizeTick(float DeltaSeconds)
 	//Post tick our interfaces.
 	TickDataInterfaces(DeltaSeconds, true);
 
-	if (HasTickingEmitters())
+	if (Component)
 	{
-		Component->UpdateComponentToWorld();//Needed for bounds updates. Can probably skip if using fixed bounds
-		Component->MarkRenderDynamicDataDirty();
+		if (HasTickingEmitters())
+		{
+			Component->UpdateComponentToWorld();//Needed for bounds updates. Can probably skip if using fixed bounds
+			Component->MarkRenderDynamicDataDirty();
+		}
 	}
 }
 
@@ -1322,7 +1364,7 @@ bool FNiagaraSystemInstance::HandleCompletion()
 {
 	bool bEmittersCompleteOrDisabled = true;
 	bHasTickingEmitters = false;
-	for (TSharedRef<FNiagaraEmitterInstance>&it : Emitters)
+	for (TSharedRef<FNiagaraEmitterInstance, ESPMode::ThreadSafe>&it : Emitters)
 	{
 		FNiagaraEmitterInstance& Inst = *it;
 		bEmittersCompleteOrDisabled &= Inst.HandleCompletion();
@@ -1342,6 +1384,9 @@ bool FNiagaraSystemInstance::HandleCompletion()
 void FNiagaraSystemInstance::PreSimulateTick(float DeltaSeconds)
 {
 	SCOPE_CYCLE_COUNTER(STAT_NiagaraSystemPreSimulateTick);
+	UNiagaraSystem* System = GetSystem();
+	FScopeCycleCounter SystemStat(System->GetStatID(true, true));
+
 	TickInstanceParameters(DeltaSeconds);
 
 	Age += DeltaSeconds;
@@ -1353,17 +1398,17 @@ void FNiagaraSystemInstance::PostSimulateTick(float DeltaSeconds)
 	SCOPE_CYCLE_COUNTER(STAT_NiagaraSystemInstanceTick);
 	SCOPE_CYCLE_COUNTER(STAT_NiagaraOverview_GT_CNC);
 
+	// Reset values that will be accumulated during emitter tick.
+	TotalParamSize = 0;
+	ActiveGPUEmitterCount = 0;
+	UNiagaraSystem* System = GetSystem();
+
 	if (IsComplete() || !bHasTickingEmitters || GetSystem() == nullptr || Component == nullptr || DeltaSeconds < SMALL_NUMBER)
 	{
 		return;
 	}
 
-	// pass the constants down to the emitter
-	// TODO: should probably just pass a pointer to the table
-	for (TPair<FNiagaraDataSetID, FNiagaraDataSet>& EventSetPair : ExternalEvents)
-	{
-		EventSetPair.Value.Tick();
-	}
+	FScopeCycleCounter SystemStat(System->GetStatID(true, true));
 
 	for (int32 EmitterIdx = 0; EmitterIdx < Emitters.Num(); EmitterIdx++)
 	{
@@ -1376,6 +1421,12 @@ void FNiagaraSystemInstance::PostSimulateTick(float DeltaSeconds)
 	{
 		FNiagaraEmitterInstance& Inst = Emitters[EmitterIdx].Get();
 		Inst.Tick(DeltaSeconds);
+
+		if (Inst.GetCachedEmitter()->SimTarget == ENiagaraSimTarget::GPUComputeSim && Inst.GetGPUContext() != nullptr && (Inst.GetExecutionState() != ENiagaraExecutionState::Complete))
+		{
+			TotalParamSize += Inst.GetGPUContext()->CombinedParamStore.GetPaddedParameterSizeInBytes();
+			ActiveGPUEmitterCount++;
+		}
 	}
 }
 
@@ -1404,9 +1455,9 @@ void FNiagaraSystemInstance::DestroyDataInterfaceInstanceData()
 	DataInterfaceInstanceData.Empty();
 }
 
-TSharedPtr<FNiagaraEmitterInstance> FNiagaraSystemInstance::GetSimulationForHandle(const FNiagaraEmitterHandle& EmitterHandle)
+TSharedPtr<FNiagaraEmitterInstance, ESPMode::ThreadSafe> FNiagaraSystemInstance::GetSimulationForHandle(const FNiagaraEmitterHandle& EmitterHandle)
 {
-	for (TSharedPtr<FNiagaraEmitterInstance> Sim : Emitters)
+	for (TSharedPtr<FNiagaraEmitterInstance, ESPMode::ThreadSafe> Sim : Emitters)
 	{
 		if(Sim->GetEmitterHandle().GetId() == EmitterHandle.GetId())
 		{
@@ -1418,12 +1469,19 @@ TSharedPtr<FNiagaraEmitterInstance> FNiagaraSystemInstance::GetSimulationForHand
 
 UNiagaraSystem* FNiagaraSystemInstance::GetSystem()const
 {
-	return Component->GetAsset();
+	if (Component)
+	{
+		return Component->GetAsset();
+	}
+	else
+	{
+		return nullptr;
+	}
 }
 
 FNiagaraEmitterInstance* FNiagaraSystemInstance::GetEmitterByID(FGuid InID)
 {
-	for (TSharedRef<FNiagaraEmitterInstance>& Emitter : Emitters)
+	for (TSharedRef<FNiagaraEmitterInstance, ESPMode::ThreadSafe>& Emitter : Emitters)
 	{
 		if (Emitter->GetEmitterHandle().GetId() == InID)
 		{
@@ -1442,7 +1500,7 @@ FNiagaraDataSet* FNiagaraSystemInstance::GetDataSet(FNiagaraDataSetID SetID, FNa
 			return ExternalSet;
 		}
 	}
-	for (TSharedPtr<FNiagaraEmitterInstance> Emitter : Emitters)
+	for (TSharedPtr<FNiagaraEmitterInstance, ESPMode::ThreadSafe> Emitter : Emitters)
 	{
 		check(Emitter.IsValid());
 		if (!Emitter->IsComplete())
diff --git a/Engine/Plugins/FX/Niagara/Source/Niagara/Private/NiagaraSystemSimulation.cpp b/Engine/Plugins/FX/Niagara/Source/Niagara/Private/NiagaraSystemSimulation.cpp
--- a/Engine/Plugins/FX/Niagara/Source/Niagara/Private/NiagaraSystemSimulation.cpp
+++ b/Engine/Plugins/FX/Niagara/Source/Niagara/Private/NiagaraSystemSimulation.cpp
@@ -13,6 +13,7 @@
 #include "Async/ParallelFor.h"
 #include "NiagaraComponent.h"
 #include "NiagaraWorldManager.h"
+#include "NiagaraEmitterInstanceBatcher.h"
 
 DECLARE_CYCLE_STAT(TEXT("System Simulation [GT]"), STAT_NiagaraSystemSim, STATGROUP_Niagara);
 DECLARE_CYCLE_STAT(TEXT("System Pre Simulate [GT]"), STAT_NiagaraSystemSim_PreSimulate, STATGROUP_Niagara);
@@ -90,25 +91,25 @@ bool FNiagaraSystemSimulation::Init(UNiagaraSystem* InSystem, UWorld* InWorld, b
 
 	if (bCanExecute)
 	{
-		DataSet.Init(FNiagaraDataSetID(), ENiagaraSimTarget::CPUSim);
+		DataSet.Init(FNiagaraDataSetID(), ENiagaraSimTarget::CPUSim, System->GetFullName());
 		DataSet.AddVariables(System->GetSystemSpawnScript()->GetVMExecutableData().Attributes);
 		DataSet.AddVariables(System->GetSystemUpdateScript()->GetVMExecutableData().Attributes);
 		DataSet.Finalize();
 
-		PausedInstanceData.Init(FNiagaraDataSetID(), ENiagaraSimTarget::CPUSim);
+		PausedInstanceData.Init(FNiagaraDataSetID(), ENiagaraSimTarget::CPUSim, System->GetFullName());
 		PausedInstanceData.AddVariables(System->GetSystemSpawnScript()->GetVMExecutableData().Attributes);
 		PausedInstanceData.AddVariables(System->GetSystemUpdateScript()->GetVMExecutableData().Attributes);
 		PausedInstanceData.Finalize();
 
 		{
-			SpawnInstanceParameterDataSet.Init(FNiagaraDataSetID(), ENiagaraSimTarget::CPUSim);
+			SpawnInstanceParameterDataSet.Init(FNiagaraDataSetID(), ENiagaraSimTarget::CPUSim, System->GetFullName());
 			FNiagaraParameters* EngineParamsSpawn = System->GetSystemSpawnScript()->GetVMExecutableData().DataSetToParameters.Find(TEXT("Engine"));
 			if (EngineParamsSpawn != nullptr)
 			{
 				SpawnInstanceParameterDataSet.AddVariables(EngineParamsSpawn->Parameters);
 			}
 			SpawnInstanceParameterDataSet.Finalize();
-			UpdateInstanceParameterDataSet.Init(FNiagaraDataSetID(), ENiagaraSimTarget::CPUSim);
+			UpdateInstanceParameterDataSet.Init(FNiagaraDataSetID(), ENiagaraSimTarget::CPUSim, System->GetFullName());
 			FNiagaraParameters* EngineParamsUpdate = System->GetSystemUpdateScript()->GetVMExecutableData().DataSetToParameters.Find(TEXT("Engine"));
 			if (EngineParamsUpdate != nullptr)
 			{
@@ -236,8 +237,8 @@ void FNiagaraSystemSimulation::TransferInstance(FNiagaraSystemSimulation* Source
 // 		UE_LOG(LogNiagara, Log, TEXT(" ----- Existing values in src. Idx: %d -----"), SystemInstIdx);
 // 		SourceSimulation->DataSet.Dump(true, SystemInstIdx, 1);
 
-		//If we're not pending then the system actually has data to pull over.
-		NewDataIndex = DataSet.TransferInstance(SourceSimulation->DataSet, SystemInstIdx);
+		//If we're not pending then the system actually has data to pull over. This is not fast.
+		NewDataIndex = DataSet.GetCurrentDataChecked().TransferInstance(SourceSimulation->DataSet.GetCurrentDataChecked(), SystemInstIdx);
 
 // 		UE_LOG(LogNiagara, Log, TEXT(" ----- Transfered values in dest. Idx: %d -----"), NewDataIndex);
 // 		DataSet.Dump(true, NewDataIndex, 1);
@@ -266,24 +267,26 @@ void FNiagaraSystemSimulation::DumpInstance(const FNiagaraSystemInstance* Inst)c
 	UE_LOG(LogNiagara, Log, TEXT("==  %s (%d) ========"), *Inst->GetSystem()->GetFullName(), Inst->SystemInstanceIndex);
 	UE_LOG(LogNiagara, Log, TEXT(".................Spawn................."));
 	SpawnExecContext.Parameters.DumpParameters(false);
-	SpawnInstanceParameterDataSet.Dump(false, Inst->SystemInstanceIndex, 1);
+	SpawnInstanceParameterDataSet.Dump(Inst->SystemInstanceIndex, 1, TEXT("Spawn Instance Parameters"));
 	UE_LOG(LogNiagara, Log, TEXT(".................Update................."));
 	UpdateExecContext.Parameters.DumpParameters(false);
-	UpdateInstanceParameterDataSet.Dump(false, Inst->SystemInstanceIndex, 1);
+	UpdateInstanceParameterDataSet.Dump(Inst->SystemInstanceIndex, 1, TEXT("Update Instance Parameters"));
 	UE_LOG(LogNiagara, Log, TEXT("................. System Instance ................."));
-	DataSet.Dump(false, Inst->SystemInstanceIndex, 1);
-	DataSet.Dump(true, Inst->SystemInstanceIndex, 1);
+	DataSet.Dump(Inst->SystemInstanceIndex, 1, TEXT("System Data"));
 }
 
 bool FNiagaraSystemSimulation::Tick(float DeltaSeconds)
 {
 	UNiagaraSystem* System = WeakSystem.Get();
+
 	if (System == nullptr || bCanExecute == false)
 	{
 		// TODO: evaluate whether or not we should have removed this from the world manager instead?
 		return false;
 	}
 
+	FScopeCycleCounter SystemStatCounter(System->GetStatID(true, false));
+
 	if (MaxDeltaTime.IsSet())
 	{
 		DeltaSeconds = FMath::Clamp(DeltaSeconds, 0.0f, MaxDeltaTime.GetValue());
@@ -307,21 +310,23 @@ bool FNiagaraSystemSimulation::Tick(float DeltaSeconds)
 
 		{
 			int32 SystemIndex = 0;
-			while (SystemIndex < SystemInstances.Num())
+			while ( SystemIndex < SystemInstances.Num() )
 			{
 				FNiagaraSystemInstance* Inst = SystemInstances[SystemIndex];
-
 				Inst->TickDataInterfaces(DeltaSeconds, false);
-				if (!Inst->IsComplete())
+
+				// TickDataInterfaces could remove the system so we only increment if the system has changed
+				if (Inst->SystemInstanceIndex != INDEX_NONE)
 				{
+					checkSlow(Inst->SystemInstanceIndex == SystemIndex);
 					++SystemIndex;
 				}
 				else
 				{
-					checkSlow(Inst->SystemInstanceIndex == INDEX_NONE);
+					checkSlow((SystemInstances.Num() <= SystemIndex) || (Inst == SystemInstances[SystemIndex]));
 				}
 			}
-			OrigNum = SystemIndex;
+			OrigNum = SystemInstances.Num();
 
 			//Pre tick and gather any still valid pending instances for spawn.
 			SystemInstances.Reserve(NewNum);
@@ -383,15 +388,6 @@ bool FNiagaraSystemSimulation::Tick(float DeltaSeconds)
 			UE_LOG(LogNiagara, Log, TEXT("=========================================================="));
 		}
 
-		if (SpawnNum)
-		{
-			SCOPE_CYCLE_COUNTER(STAT_NiagaraSystemSim_Spawn);
-			DataSet.Allocate(NewNum, true);
-			DataSet.SetNumInstances(NewNum);
-		}
-
-		TArray<FNiagaraDataSetExecutionInfo, TInlineAllocator<8>> DataSetExecInfos;
-		DataSetExecInfos.SetNum(2);
 		{
 			SCOPE_CYCLE_COUNTER(STAT_NiagaraSystemSim_PrepareForSimulate);
 			auto PreSimulateAndTransferParams = [&](int32 SystemIndex)
@@ -408,23 +404,27 @@ bool FNiagaraSystemSimulation::Tick(float DeltaSeconds)
 				//TODO: Find good way to check that we're not using any instance parameter data interfaces in the system scripts here.
 				//In that case we need to solo and will never get here.
 
-				TArray<TSharedRef<FNiagaraEmitterInstance>>& Emitters = Inst->GetEmitters();
+				TArray<TSharedRef<FNiagaraEmitterInstance, ESPMode::ThreadSafe>>& Emitters = Inst->GetEmitters();
 				for (int32 EmitterIdx = 0; EmitterIdx < Emitters.Num(); ++EmitterIdx)
 				{
 					FNiagaraEmitterInstance& EmitterInst = Emitters[EmitterIdx].Get();
-					if (EmitterExecutionStateAccessors.Num() > EmitterIdx && EmitterExecutionStateAccessors[EmitterIdx].BaseIsValid())
+					if (EmitterExecutionStateAccessors.Num() > EmitterIdx && EmitterExecutionStateAccessors[EmitterIdx].IsValidForWrite())
 					{
 						EmitterExecutionStateAccessors[EmitterIdx].Set(SystemIndex, (int32)EmitterInst.GetExecutionState());
 					}
 				}
 			};
+			
+			//Begin filling the state of the instance paramter datasets.
+			SpawnInstanceParameterDataSet.BeginSimulate();
+			UpdateInstanceParameterDataSet.BeginSimulate();
 
 			SpawnInstanceParameterDataSet.Allocate(NewNum);
 			UpdateInstanceParameterDataSet.Allocate(NewNum);
 
 			for (int32 EmitterIdx = 0; EmitterIdx < System->GetNumEmitters(); ++EmitterIdx)
 			{
-				EmitterExecutionStateAccessors[EmitterIdx].InitForAccess(true);
+				EmitterExecutionStateAccessors[EmitterIdx].InitForAccess();
 			}
 
 			//Transfer any values like execution state from the system instance into the dataset for simulation.
@@ -433,10 +433,12 @@ bool FNiagaraSystemSimulation::Tick(float DeltaSeconds)
 				PreSimulateAndTransferParams(SystemIndex);
 			});
 
-			SpawnInstanceParameterDataSet.SetNumInstances(NewNum);
-			UpdateInstanceParameterDataSet.SetNumInstances(NewNum);
-			SpawnInstanceParameterDataSet.Tick();
-			UpdateInstanceParameterDataSet.Tick();
+			SpawnInstanceParameterDataSet.GetDestinationDataChecked().SetNumInstances(NewNum);
+			UpdateInstanceParameterDataSet.GetDestinationDataChecked().SetNumInstances(NewNum);
+
+			//We're done filling in the current state for the instance parameter datasets.
+			SpawnInstanceParameterDataSet.EndSimulate();
+			UpdateInstanceParameterDataSet.EndSimulate();
 
 			//Setup the few real constants like delta time.
 			float InvDt = 1.0f / DeltaSeconds;
@@ -462,28 +464,32 @@ bool FNiagaraSystemSimulation::Tick(float DeltaSeconds)
 		if (SpawnNum)
 		{
 			SCOPE_CYCLE_COUNTER(STAT_NiagaraSystemSim_Spawn);
+			DataSet.BeginSimulate();
+			DataSet.Allocate(NewNum, true);
+			DataSet.GetDestinationDataChecked().SetNumInstances(NewNum);
 
 			//Run Spawn
 			SpawnExecContext.Tick(SoloSystemInstance);//We can't require a specific instance here as these are for all instances.
-			DataSetExecInfos[0] = FNiagaraDataSetExecutionInfo(&DataSet, OrigNum, false, false);
-			DataSetExecInfos[1] = FNiagaraDataSetExecutionInfo(&SpawnInstanceParameterDataSet, OrigNum, false, false);
-			SpawnExecContext.Execute(SpawnNum, DataSetExecInfos);
+			SpawnExecContext.BindData(0, DataSet, OrigNum, false);
+			SpawnExecContext.BindData(1, SpawnInstanceParameterDataSet, OrigNum, false);
+			SpawnExecContext.Execute(SpawnNum);
 
 			if (GbDumpSystemData || System->bDumpDebugSystemInfo)
 			{
 				UE_LOG(LogNiagara, Log, TEXT("=== Spwaned %d Systems ==="), SpawnNum);
-				DataSet.Dump(true, OrigNum, SpawnNum);
-				SpawnInstanceParameterDataSet.Dump(false, OrigNum, SpawnNum);
+				DataSet.GetDestinationDataChecked().Dump(OrigNum, SpawnNum, TEXT("System Dataset - Post Spawn"));
+				SpawnInstanceParameterDataSet.GetCurrentDataChecked().Dump(OrigNum, SpawnNum, TEXT("Spawn Instance Parameter Data"));
 			}
 
+			DataSet.EndSimulate();
+
 #if WITH_EDITORONLY_DATA
 			if (SoloSystemInstance && SoloSystemInstance->ShouldCaptureThisFrame())
 			{
 				TSharedPtr<struct FNiagaraScriptDebuggerInfo, ESPMode::ThreadSafe> DebugInfo = SoloSystemInstance->GetActiveCaptureWrite(NAME_None, ENiagaraScriptUsage::SystemSpawnScript, FGuid());
 				if (DebugInfo)
 				{
-					DataSet.Dump(DebugInfo->Frame, true, OrigNum, SpawnNum);
-					//DebugInfo->Frame.Dump(true, 0, SpawnNum);
+					DataSet.CopyTo(DebugInfo->Frame, OrigNum, SpawnNum);
 					DebugInfo->Parameters = SpawnExecContext.Parameters;
 					DebugInfo->bWritten = true;
 				}
@@ -491,33 +497,25 @@ bool FNiagaraSystemSimulation::Tick(float DeltaSeconds)
 #endif
 		}
 
-		DataSet.Tick();
-		DataSet.Allocate(NewNum);
 
 		{
 			SCOPE_CYCLE_COUNTER(STAT_NiagaraSystemSim_Update);
-			DataSet.SetNumInstances(NewNum);
-			//UpdateInstanceParameterDataSet.SetNumInstances(OrigNum);
+
+			FNiagaraDataBuffer& DestinationData = DataSet.BeginSimulate();
+			DestinationData.Allocate(NewNum);
+			DestinationData.SetNumInstances(NewNum);
 
 			//Run update.
 			UpdateExecContext.Tick(SystemInstances[0]);
-			DataSetExecInfos[0] = FNiagaraDataSetExecutionInfo(&DataSet, 0, false, false);
-			DataSetExecInfos[1] = FNiagaraDataSetExecutionInfo(&UpdateInstanceParameterDataSet, 0, false, false);
-
-// 			if (GbDumpSystemData || System->bDumpDebugSystemInfo)
-// 			{
-// 				UE_LOG(LogNiagara, Log, TEXT("=== PreUpdated %d Systems ==="), OrigNum);
-// 				DataSet.Dump(false, 0, OrigNum);
-// 				UpdateInstanceParameterDataSet.Dump(false, 0, OrigNum);
-// 			}
-
-			UpdateExecContext.Execute(OrigNum, DataSetExecInfos);
+			UpdateExecContext.BindData(0, DataSet, 0, false);
+			UpdateExecContext.BindData(1, UpdateInstanceParameterDataSet, 0, false);
+			UpdateExecContext.Execute(OrigNum);
 
 			if (GbDumpSystemData || System->bDumpDebugSystemInfo)
 			{
-				UE_LOG(LogNiagara, Log, TEXT("=== Updated %d Systems ==="), OrigNum);
-				DataSet.Dump(true, 0, OrigNum);
-				UpdateInstanceParameterDataSet.Dump(false, 0, OrigNum);
+				UE_LOG(LogNiagara, Log, TEXT("=== Updated %d Systems ==="), OrigNum); 
+				DestinationData.Dump(0, OrigNum, TEXT("System Data - Post Update"));
+				UpdateInstanceParameterDataSet.GetCurrentDataChecked().Dump(0, OrigNum, TEXT("Update Instance Paramter Data"));
 			}
 
  			//Also run the update script on the newly spawned systems too.
@@ -525,25 +523,26 @@ bool FNiagaraSystemSimulation::Tick(float DeltaSeconds)
 			//Ideally this should be compiled directly into the script similarly to interpolated particle spawning.
 			if (SpawnNum && GbSystemUpdateOnSpawn)
 			{
-				DataSet.SetNumInstances(NewNum);
-
 				//Run update.
 				UpdateExecContext.Tick(SystemInstances[0]);
-				DataSetExecInfos[0] = FNiagaraDataSetExecutionInfo(&DataSet, OrigNum, false, false);
-				DataSetExecInfos[1] = FNiagaraDataSetExecutionInfo(&UpdateInstanceParameterDataSet, OrigNum, false, false);
 
 				UpdateExecContext.Parameters.SetParameterValue(0.0001f, SYS_PARAM_ENGINE_DELTA_TIME);
 				UpdateExecContext.Parameters.SetParameterValue(10000.0f, SYS_PARAM_ENGINE_INV_DELTA_TIME);
 
-				UpdateExecContext.Execute(SpawnNum, DataSetExecInfos);
+				UpdateExecContext.BindData(0, DataSet, OrigNum, false);
+				UpdateExecContext.BindData(1, UpdateInstanceParameterDataSet, OrigNum, false);
+				UpdateExecContext.Execute(SpawnNum);
 
 				if (GbDumpSystemData || System->bDumpDebugSystemInfo)
 				{
 					UE_LOG(LogNiagara, Log, TEXT("=== Spawn Updated %d Systems ==="), SpawnNum);
-					DataSet.Dump(true, OrigNum, SpawnNum);
-					UpdateInstanceParameterDataSet.Dump(false, OrigNum, SpawnNum);
+					DestinationData.Dump(OrigNum, SpawnNum, TEXT("System Data - Post Update (new systems)"));
+					UpdateInstanceParameterDataSet.GetCurrentDataChecked().Dump(OrigNum, SpawnNum, TEXT("Update Instance Paramter Data (new systems)"));
 				}
 			}
+			
+			DataSet.EndSimulate();
+
 
 #if WITH_EDITORONLY_DATA
 			if (SoloSystemInstance && SoloSystemInstance->ShouldCaptureThisFrame())
@@ -551,8 +550,7 @@ bool FNiagaraSystemSimulation::Tick(float DeltaSeconds)
 				TSharedPtr<struct FNiagaraScriptDebuggerInfo, ESPMode::ThreadSafe> DebugInfo = SoloSystemInstance->GetActiveCaptureWrite(NAME_None, ENiagaraScriptUsage::SystemUpdateScript, FGuid());
 				if (DebugInfo)
 				{
-					DataSet.Dump(DebugInfo->Frame, true, 0, NewNum);
-					//DebugInfo->Frame.Dump(true, 0, OrigNum);
+					DataSet.CopyTo(DebugInfo->Frame);
 					DebugInfo->Parameters = UpdateExecContext.Parameters;
 					DebugInfo->bWritten = true;
 				}
@@ -560,20 +558,19 @@ bool FNiagaraSystemSimulation::Tick(float DeltaSeconds)
 #endif
 		}
 
-		SystemExecutionStateAccessor.InitForAccess(true);
+		SystemExecutionStateAccessor.InitForAccess();
 		for (int32 EmitterIdx = 0; EmitterIdx < System->GetNumEmitters(); ++EmitterIdx)
 		{
-			EmitterExecutionStateAccessors[EmitterIdx].InitForAccess(true);
+			EmitterExecutionStateAccessors[EmitterIdx].InitForAccess();
 			for (int32 SpawnInfoIdx = 0; SpawnInfoIdx < EmitterSpawnInfoAccessors[EmitterIdx].Num(); ++SpawnInfoIdx)
 			{
-				EmitterSpawnInfoAccessors[EmitterIdx][SpawnInfoIdx].InitForAccess(true);
+				EmitterSpawnInfoAccessors[EmitterIdx][SpawnInfoIdx].InitForAccess();
 			}
 		}
 
 		{
 			SCOPE_CYCLE_COUNTER(STAT_NiagaraSystemSim_TransferParameters);
-			int32 SystemIndex = 0;
-			while (SystemIndex < SystemInstances.Num())
+			for (int32 SystemIndex = 0; SystemIndex < SystemInstances.Num(); ++SystemIndex)
 			{
 				ENiagaraExecutionState ExecutionState = (ENiagaraExecutionState)SystemExecutionStateAccessor.GetSafe(SystemIndex, (int32)ENiagaraExecutionState::Disabled);
 				FNiagaraSystemInstance* SystemInst = SystemInstances[SystemIndex];
@@ -584,7 +581,7 @@ bool FNiagaraSystemSimulation::Tick(float DeltaSeconds)
 				if (!SystemInst->IsDisabled() && !SystemInst->HandleCompletion())
 				{
 					//Now pull data out of the simulation and drive the emitters with it.
-					TArray<TSharedRef<FNiagaraEmitterInstance>>& Emitters = SystemInst->GetEmitters();
+					TArray<TSharedRef<FNiagaraEmitterInstance, ESPMode::ThreadSafe>>& Emitters = SystemInst->GetEmitters();
 					for (int32 EmitterIdx = 0; EmitterIdx < Emitters.Num(); ++EmitterIdx)
 					{
 						FNiagaraEmitterInstance& EmitterInst = Emitters[EmitterIdx].Get();
@@ -595,6 +592,7 @@ bool FNiagaraSystemSimulation::Tick(float DeltaSeconds)
 							continue;
 						}
 
+						check(Emitters.Num() > EmitterIdx);
 						ENiagaraExecutionState State = (ENiagaraExecutionState)EmitterExecutionStateAccessors[EmitterIdx].GetSafe(SystemIndex, (int32)ENiagaraExecutionState::Disabled);
 						EmitterInst.SetExecutionState(State);
 												
@@ -634,9 +632,6 @@ bool FNiagaraSystemSimulation::Tick(float DeltaSeconds)
 							}
 						}
 					}
-
-					//System is still enabled. 
-					++SystemIndex;
 				}
 			}
 		}
@@ -673,13 +668,45 @@ bool FNiagaraSystemSimulation::Tick(float DeltaSeconds)
 		int32 SystemIndex = 0;
 		while (SystemIndex < SystemInstances.Num())
 		{
-			FNiagaraSystemInstance* SystemInst = SystemInstances[SystemIndex];
+			FNiagaraSystemInstance* SystemInstance = SystemInstances[SystemIndex];
 			++SystemIndex;
 			if (SystemIndex < SystemInstances.Num())
 			{
 				FPlatformMisc::Prefetch(SystemInstances[SystemIndex]->GetComponent());
 			}
-			SystemInst->FinalizeTick(DeltaSeconds);
+			SystemInstance->FinalizeTick(DeltaSeconds);
+			// @todo If we buffer up the Tick object we create here we can push all the data for every system in this simulation at once.
+			// @todo It may also be worth considering hoisting the push out to FNiagaraWorldManager::Tick to push every desired simulation at once.
+			if (SystemInstance->ActiveGPUEmitterCount > 0)
+			{
+				ensure(!SystemInstance->IsComplete());
+				FNiagaraGPUSystemTick GPUTick;
+				GPUTick.Init(SystemInstance);
+
+				//if (GPUTick.DIInstanceData)
+				//{
+				//	uint8* BasePointer = (uint8*)GPUTick.DIInstanceData->PerInstanceDataForRT;
+
+				//	//UE_LOG(LogNiagara, Log, TEXT("GT Testing (dipacket) %p (baseptr) %p"), GPUTick.DIInstanceData, BasePointer);
+				//	for (auto& Pair : GPUTick.DIInstanceData->InterfaceProxiesToOffsets)
+				//	{
+				//		FNiagaraDataInterfaceProxy* Proxy = Pair.Key;
+				//		UE_LOG(LogNiagara, Log, TEXT("\tGT (proxy) %p (size) %u"), Proxy, Proxy->PerInstanceDataPassedToRenderThreadSize());
+				//	}
+				//}
+
+				// We will give the data over to the render thread. It is responsible for freeing it.
+				// We no longer own it and cannot modify it after this point.
+				// @todo We are taking a copy of the object here. This object is small so this overhead should
+				// not be very high. And we avoid making a bunch of small allocations here.
+				NiagaraEmitterInstanceBatcher* Batcher = SystemInstance->GetBatcher();
+				ENQUEUE_RENDER_COMMAND(FGiveSystemInstanceTickToRT)(
+					[Batcher, GPUTick](FRHICommandListImmediate& RHICmdList) mutable
+				{
+					Batcher->GiveSystemTick_RenderThread(GPUTick);
+				}
+				);
+			}
 		}
 	}
 
@@ -702,13 +729,14 @@ void FNiagaraSystemSimulation::RemoveInstance(FNiagaraSystemInstance* Instance)
 		return;
 	}
 
+	check(DataSet.GetDestinationData() == nullptr);
 	UNiagaraSystem* System = WeakSystem.Get();
 	if (Instance->IsPendingSpawn())
 	{
 		if (GbDumpSystemData || (System && System->bDumpDebugSystemInfo))
 		{
 			UE_LOG(LogNiagara, Log, TEXT("=== Removing Pending Spawn %d ==="), Instance->SystemInstanceIndex);
-			DataSet.Dump(true, Instance->SystemInstanceIndex, 1);
+			DataSet.GetCurrentDataChecked().Dump(Instance->SystemInstanceIndex, 1, TEXT("System data being removed."));
 		}
 		int32 SystemIndex = Instance->SystemInstanceIndex;
 		check(Instance == PendingSystemInstances[SystemIndex]);
@@ -725,7 +753,7 @@ void FNiagaraSystemSimulation::RemoveInstance(FNiagaraSystemInstance* Instance)
 		if (GbDumpSystemData || (System && System->bDumpDebugSystemInfo))
 		{
 			UE_LOG(LogNiagara, Log, TEXT("=== Removing Paused %d ==="), Instance->SystemInstanceIndex);
-			DataSet.Dump(true, Instance->SystemInstanceIndex, 1);
+			DataSet.GetCurrentDataChecked().Dump(Instance->SystemInstanceIndex, 1, TEXT("System data being removed."));
 		}
 
 		int32 SystemIndex = Instance->SystemInstanceIndex;
@@ -742,16 +770,16 @@ void FNiagaraSystemSimulation::RemoveInstance(FNiagaraSystemInstance* Instance)
 		if (GbDumpSystemData || (System && System->bDumpDebugSystemInfo))
 		{
 			UE_LOG(LogNiagara, Log, TEXT("=== Removing System %d ==="), Instance->SystemInstanceIndex);
-			DataSet.Dump(true, Instance->SystemInstanceIndex, 1);
+			DataSet.GetCurrentDataChecked().Dump(Instance->SystemInstanceIndex, 1, TEXT("System data being removed."));
 		}
 
-		int32 NumInstances = DataSet.GetNumInstances();
+		int32 NumInstances = DataSet.GetCurrentDataChecked().GetNumInstances();
 		check(SystemInstances.Num() == NumInstances);
 
 		int32 SystemIndex = Instance->SystemInstanceIndex;
 		check(Instance == SystemInstances[SystemIndex]);
 		check(SystemInstances.IsValidIndex(SystemIndex));
-		DataSet.KillInstance(SystemIndex);
+		DataSet.GetCurrentDataChecked().KillInstance(SystemIndex);
 		SystemInstances.RemoveAtSwap(SystemIndex);
 		Instance->SystemInstanceIndex = INDEX_NONE;
 
@@ -783,6 +811,8 @@ void FNiagaraSystemSimulation::AddInstance(FNiagaraSystemInstance* Instance)
 void FNiagaraSystemSimulation::PauseInstance(FNiagaraSystemInstance* Instance)
 {
 	check(!Instance->IsPaused());
+	check(!DataSet.GetDestinationData());
+	check(!PausedInstanceData.GetDestinationData());
 
 	UNiagaraSystem* System = WeakSystem.Get();
 	if (Instance->IsPendingSpawn())
@@ -800,15 +830,15 @@ void FNiagaraSystemSimulation::PauseInstance(FNiagaraSystemInstance* Instance)
 	if (GbDumpSystemData || (System && System->bDumpDebugSystemInfo))
 	{
 		UE_LOG(LogNiagara, Log, TEXT("=== Pausing System %d ==="), Instance->SystemInstanceIndex);
-		DataSet.Dump(true, Instance->SystemInstanceIndex, 1);
+		DataSet.GetCurrentDataChecked().Dump(Instance->SystemInstanceIndex, 1, TEXT("System data being paused."));
 	}
 
 	int32 SystemIndex = Instance->SystemInstanceIndex;
 	check(SystemInstances.IsValidIndex(SystemIndex));
 	check(Instance == SystemInstances[SystemIndex]);
 
-	int32 NewSystemIndex = PausedInstanceData.TransferInstance(DataSet, SystemIndex);
-	DataSet.KillInstance(SystemIndex);
+	int32 NewSystemIndex = PausedInstanceData.GetCurrentDataChecked().TransferInstance(DataSet.GetCurrentDataChecked(), SystemIndex);
+	DataSet.GetCurrentDataChecked().KillInstance(SystemIndex);
 
 	check(PausedSystemInstances.Num() == NewSystemIndex);
 	Instance->SystemInstanceIndex = NewSystemIndex;
@@ -824,6 +854,8 @@ void FNiagaraSystemSimulation::PauseInstance(FNiagaraSystemInstance* Instance)
 void FNiagaraSystemSimulation::UnpauseInstance(FNiagaraSystemInstance* Instance)
 {
 	check(Instance->IsPaused());
+	check(!DataSet.GetDestinationData());
+	check(!PausedInstanceData.GetDestinationData());
 
 	UNiagaraSystem* System = WeakSystem.Get();
 	if (Instance->IsPendingSpawn())
@@ -841,15 +873,15 @@ void FNiagaraSystemSimulation::UnpauseInstance(FNiagaraSystemInstance* Instance)
 	if (GbDumpSystemData || (System && System->bDumpDebugSystemInfo))
 	{
 		UE_LOG(LogNiagara, Log, TEXT("=== Unpausing System %d ==="), Instance->SystemInstanceIndex);
-		DataSet.Dump(true, Instance->SystemInstanceIndex, 1);
+		DataSet.GetCurrentDataChecked().Dump(Instance->SystemInstanceIndex, 1, TEXT("System data being unpaused."));
 	}
 
 	int32 SystemIndex = Instance->SystemInstanceIndex;
 	check(PausedSystemInstances.IsValidIndex(SystemIndex));
 	check(Instance == PausedSystemInstances[SystemIndex]);
 
-	int32 NewSystemIndex = DataSet.TransferInstance(PausedInstanceData, SystemIndex);
-	PausedInstanceData.KillInstance(SystemIndex);
+	int32 NewSystemIndex = DataSet.GetCurrentDataChecked().TransferInstance(PausedInstanceData.GetCurrentDataChecked(), SystemIndex);
+	PausedInstanceData.GetCurrentDataChecked().KillInstance(SystemIndex);
 
 	check(SystemInstances.Num() == NewSystemIndex);
 	Instance->SystemInstanceIndex = NewSystemIndex;
@@ -873,7 +905,7 @@ void FNiagaraSystemSimulation::InitParameterDataSetBindings(FNiagaraSystemInstan
 		SpawnInstanceParameterToDataSetBinding.Init(SpawnInstanceParameterDataSet, SystemInst->GetInstanceParameters());
 		UpdateInstanceParameterToDataSetBinding.Init(UpdateInstanceParameterDataSet, SystemInst->GetInstanceParameters());
 
-		TArray<TSharedRef<FNiagaraEmitterInstance>>& Emitters = SystemInst->GetEmitters();
+		TArray<TSharedRef<FNiagaraEmitterInstance, ESPMode::ThreadSafe>>& Emitters = SystemInst->GetEmitters();
 		DataSetToEmitterSpawnParameters.SetNum(Emitters.Num());
 		DataSetToEmitterUpdateParameters.SetNum(Emitters.Num());
 		DataSetToEmitterEventParameters.SetNum(Emitters.Num());
diff --git a/Engine/Plugins/FX/Niagara/Source/Niagara/Private/NiagaraWorldManager.cpp b/Engine/Plugins/FX/Niagara/Source/Niagara/Private/NiagaraWorldManager.cpp
--- a/Engine/Plugins/FX/Niagara/Source/Niagara/Private/NiagaraWorldManager.cpp
+++ b/Engine/Plugins/FX/Niagara/Source/Niagara/Private/NiagaraWorldManager.cpp
@@ -15,6 +15,7 @@
 #include "NiagaraStats.h"
 
 DECLARE_CYCLE_STAT(TEXT("Niagara Manager Tick [GT]"), STAT_NiagaraWorldManTick, STATGROUP_Niagara);
+DECLARE_CYCLE_STAT(TEXT("Niagara Manager Wait On Render [GT]"), STAT_NiagaraWorldManWaitOnRender, STATGROUP_Niagara);
 
 TGlobalResource<FNiagaraViewDataMgr> GNiagaraViewDataManager;
 
@@ -33,8 +34,6 @@ void FNiagaraViewDataMgr::Init()
 
 	GNiagaraViewDataManager.PostOpaqueDelegate.BindRaw(&GNiagaraViewDataManager, &FNiagaraViewDataMgr::PostOpaqueRender);
 	RendererModule.RegisterPostOpaqueRenderDelegate(GNiagaraViewDataManager.PostOpaqueDelegate);
-
-	RendererModule.OnPreSceneRender().AddRaw(&GNiagaraViewDataManager, &FNiagaraViewDataMgr::OnPreSceneRenderCalled);
 }
 
 void FNiagaraViewDataMgr::Shutdown()
@@ -61,8 +60,6 @@ FNiagaraWorldManager::FNiagaraWorldManager(UWorld* InWorld)
 {
 }
 
-
-
 FNiagaraWorldManager* FNiagaraWorldManager::Get(UWorld* World)
 {
 	//INiagaraModule& NiagaraModule = FModuleManager::LoadModuleChecked<INiagaraModule>("Niagara");
@@ -171,6 +168,13 @@ void FNiagaraWorldManager::DestroySystemSimulation(UNiagaraSystem* System)
 	}	
 }
 
+void FNiagaraWorldManager::DestroySystemInstance(TUniquePtr<FNiagaraSystemInstance>& InPtr)
+{
+	check(IsInGameThread());
+	check(InPtr != nullptr);
+	DeferredDeletionQueue[DeferredDeletionQueueIndex].Queue.Emplace(MoveTemp(InPtr));
+}
+
 void FNiagaraWorldManager::OnWorldCleanup(bool bSessionEnded, bool bCleanupResources)
 {
 	for (TPair<UNiagaraSystem*, TSharedRef<FNiagaraSystemSimulation, ESPMode::ThreadSafe>>& SimPair : SystemSimulations)
@@ -179,37 +183,21 @@ void FNiagaraWorldManager::OnWorldCleanup(bool bSessionEnded, bool bCleanupResou
 	}
 	SystemSimulations.Empty();
 	CleanupParameterCollections();
+
+	for ( int32 i=0; i < NumDeferredQueues; ++i)
+	{
+		DeferredDeletionQueue[DeferredDeletionQueueIndex].Fence.Wait();
+		DeferredDeletionQueue[i].Queue.Empty();
+	}
 }
 
 void FNiagaraWorldManager::Tick(float DeltaSeconds)
 {
-/*
-Some experimental tinkering with links to effects quality.
-Going off this idea tbh
-
-	int32 CurrentEffectsQuality = Scalability::GetEffectsQualityDirect(true);
-	if (CachedEffectsQuality != CurrentEffectsQuality)
-	{
-		CachedEffectsQuality = CurrentEffectsQuality;
-		FString SectionName = FString::Printf(TEXT("%s@%d"), TEXT("EffectsQuality"), CurrentEffectsQuality);
-		if (FConfigFile* File = GConfig->FindConfigFileWithBaseName(TEXT("Niagara")))
-		{
-			FString ScalabilityCollectionString;			
-			if (File->GetString(*SectionName, TEXT("ScalabilityCollection"), ScalabilityCollectionString))
-			{
-				//NewLandscape_Material = LoadObject<UMaterialInterface>(NULL, *NewLandsfgcapeMaterialName, NULL, LOAD_NoWarn);
-				UNiagaraParameterCollectionInstance* ScalabilityCollection = LoadObject<UNiagaraParameterCollectionInstance>(NULL, *ScalabilityCollectionString, NULL, LOAD_NoWarn);
-				if (ScalabilityCollection)
-				{
-					SetParameterCollection(ScalabilityCollection);
-				}
-			}
-		}
-	}
-*/
 	SCOPE_CYCLE_COUNTER(STAT_NiagaraWorldManTick);
 	SCOPE_CYCLE_COUNTER(STAT_NiagaraOverview_GT);
 
+	FNiagaraSharedObject::FlushDeletionList();
+
 	SkeletalMeshGeneratedData.TickGeneratedData(DeltaSeconds);
 
 	//Tick our collections to push any changes to bound stores.
@@ -233,4 +221,15 @@ Going off this idea tbh
 	{
 		SystemSimulations.Remove(DeadSystem);
 	}
+
+	// Enqueue fence for deferred deletion
+	DeferredDeletionQueue[DeferredDeletionQueueIndex].Fence.BeginFence();
+
+	// Remove instances from previous frame
+	DeferredDeletionQueueIndex = (DeferredDeletionQueueIndex + 1) % NumDeferredQueues;
+	{
+		SCOPE_CYCLE_COUNTER(STAT_NiagaraWorldManWaitOnRender);
+		DeferredDeletionQueue[DeferredDeletionQueueIndex].Fence.Wait();
+	}
+	DeferredDeletionQueue[DeferredDeletionQueueIndex].Queue.Empty();
 }
diff --git a/Engine/Plugins/FX/Niagara/Source/Niagara/Public/NiagaraBoundsCalculator.h b/Engine/Plugins/FX/Niagara/Source/Niagara/Public/NiagaraBoundsCalculator.h
new file mode 100644
--- /dev/null
+++ b/Engine/Plugins/FX/Niagara/Source/Niagara/Public/NiagaraBoundsCalculator.h
@@ -0,0 +1,16 @@
+// Copyright 1998-2019 Epic Games, Inc. All Rights Reserved.
+
+
+#pragma once
+
+#include "CoreMinimal.h"
+#include "Stats/Stats.h"
+#include "NiagaraModule.h"
+
+class FNiagaraBoundsCalculator
+{
+public:
+	virtual ~FNiagaraBoundsCalculator() { }
+	virtual void InitAccessors(FNiagaraDataSet& DataSet) = 0;
+	virtual FBox CalculateBounds(const int32 NumInstances, bool& bOutContainsNaN) = 0;
+};
diff --git a/Engine/Plugins/FX/Niagara/Source/Niagara/Public/NiagaraBoundsCalculatorHelper.h b/Engine/Plugins/FX/Niagara/Source/Niagara/Public/NiagaraBoundsCalculatorHelper.h
new file mode 100644
--- /dev/null
+++ b/Engine/Plugins/FX/Niagara/Source/Niagara/Public/NiagaraBoundsCalculatorHelper.h
@@ -0,0 +1,142 @@
+// Copyright 1998-2019 Epic Games, Inc. All Rights Reserved.
+
+
+#pragma once
+
+#include "NiagaraBoundsCalculator.h"
+#include "NiagaraDataSet.h"
+
+template<bool bUsedWithSprites, bool bUsedWithMeshes, bool bUsedWithRibbons>
+class FNiagaraBoundsCalculatorHelper : public FNiagaraBoundsCalculator
+{
+public:
+	virtual void InitAccessors(FNiagaraDataSet& DataSet) override final
+	{
+		static const FName PositionName(TEXT("Position"));
+		static const FName SpriteSizeName(TEXT("SpriteSize"));
+		static const FName ScaleName(TEXT("Scale"));
+		static const FName RibbonWidthName(TEXT("RibbonWidth"));
+
+		PositionAccessor = FNiagaraDataSetAccessor<FVector>(DataSet, FNiagaraVariable(FNiagaraTypeDefinition::GetVec3Def(), PositionName));
+
+		if (bUsedWithSprites)
+		{
+			SpriteSizeAccessor = FNiagaraDataSetAccessor<FVector2D>(DataSet, FNiagaraVariable(FNiagaraTypeDefinition::GetVec2Def(), SpriteSizeName));
+		}
+		else
+		{
+			SpriteSizeAccessor = FNiagaraDataSetAccessor<FVector2D>();
+		}
+
+		if (bUsedWithMeshes)
+		{
+			ScaleAccessor = FNiagaraDataSetAccessor<FVector>(DataSet, FNiagaraVariable(FNiagaraTypeDefinition::GetVec3Def(), ScaleName));
+		}
+		else
+		{
+			ScaleAccessor = FNiagaraDataSetAccessor<FVector>();
+		}
+
+		if (bUsedWithRibbons)
+		{
+			RibbonWidthAccessor = FNiagaraDataSetAccessor<float>(DataSet, FNiagaraVariable(FNiagaraTypeDefinition::GetFloatDef(), RibbonWidthName));
+		}
+		else
+		{
+			RibbonWidthAccessor = FNiagaraDataSetAccessor<float>();
+		}
+	}
+
+	virtual FBox CalculateBounds(const int32 NumInstances, bool& bOutContainsNaN) override final
+	{
+		PositionAccessor.InitForAccess();
+		if (PositionAccessor.IsValid() == false)
+			return FBox(ForceInit);
+
+		if (bUsedWithSprites)
+		{
+			SpriteSizeAccessor.InitForAccess();
+		}
+		if (bUsedWithMeshes)
+		{
+			ScaleAccessor.InitForAccess();
+		}
+		if (bUsedWithRibbons)
+		{
+			RibbonWidthAccessor.InitForAccess();
+		}
+
+		constexpr float kDefaultSize = 50.0f;
+		FVector MaxScale(ScaleAccessor.IsValid() ? FVector::ZeroVector : FVector(kDefaultSize, kDefaultSize, kDefaultSize));
+		float MaxSpriteSize(SpriteSizeAccessor.IsValid() ? 0.0f : kDefaultSize);
+		float MaxRibbonWidth(RibbonWidthAccessor.IsValid() ? 0.0f : kDefaultSize);
+
+#if !UE_BUILD_SHIPPING
+		bOutContainsNaN = false;
+#endif
+
+		FBox Ret;
+		Ret.Init();
+
+		for (int32 InstIdx = 0; InstIdx < NumInstances && PositionAccessor.IsValid(); ++InstIdx)
+		{
+			FVector Position;
+			PositionAccessor.Get(InstIdx, Position);
+
+			// Some graphs have a tendency to divide by zero. This ContainsNaN has been added prophylactically
+			// to keep us safe during GDC. It should be removed as soon as we feel safe that scripts are appropriately warned.
+			if (Position.ContainsNaN())
+			{
+				bOutContainsNaN = true;
+				continue;
+			}
+
+			Ret += Position;
+
+			if (bUsedWithMeshes && ScaleAccessor.IsValid())
+			{
+				MaxScale = MaxScale.ComponentMax(ScaleAccessor.Get(InstIdx));
+			}
+
+			if (bUsedWithSprites && SpriteSizeAccessor.IsValid())
+			{
+				const FVector2D InstanceSpriteSize = SpriteSizeAccessor.Get(InstIdx);
+				MaxSpriteSize = FMath::Max3(MaxSpriteSize, InstanceSpriteSize.X, InstanceSpriteSize.Y);
+			}
+
+			if (bUsedWithRibbons && RibbonWidthAccessor.IsValid())
+			{
+				MaxRibbonWidth = FMath::Max(MaxRibbonWidth, RibbonWidthAccessor.Get(InstIdx));
+			}
+		}
+
+		float MaxSize = 0.0001f;
+		if (bUsedWithSprites)
+		{
+			const float SpriteSize = FMath::IsNearlyZero(MaxSpriteSize) ? 1.0f : MaxSpriteSize * 0.5f;
+			MaxSize = FMath::Max(MaxSize, SpriteSize);
+		}
+
+		if (bUsedWithMeshes)
+		{
+			const FVector MeshSize = (MaxScale.IsNearlyZero() ? FVector::OneVector : MaxScale) * MeshExtents;
+			MaxSize = FMath::Max(MaxSize, MeshSize.GetMax());
+		}
+
+		if (bUsedWithRibbons)
+		{
+			const float RibbonSize = FMath::IsNearlyZero(MaxRibbonWidth) ? 1.0f : MaxRibbonWidth * 0.5f;
+			MaxSize = FMath::Max(MaxSize, RibbonSize);
+		}
+
+		Ret = Ret.ExpandBy(MaxSize);
+
+		return Ret;
+	}
+
+	FNiagaraDataSetAccessor<FVector> PositionAccessor;
+	FNiagaraDataSetAccessor<FVector2D> SpriteSizeAccessor;
+	FNiagaraDataSetAccessor<FVector> ScaleAccessor;
+	FNiagaraDataSetAccessor<float> RibbonWidthAccessor;
+	FVector MeshExtents = FVector::OneVector;
+};
diff --git a/Engine/Plugins/FX/Niagara/Source/Niagara/Public/NiagaraCommon.h b/Engine/Plugins/FX/Niagara/Source/Niagara/Public/NiagaraCommon.h
--- a/Engine/Plugins/FX/Niagara/Source/Niagara/Public/NiagaraCommon.h
+++ b/Engine/Plugins/FX/Niagara/Source/Niagara/Public/NiagaraCommon.h
@@ -21,7 +21,7 @@ struct FNiagaraParameterStore;
 #define NIAGARA_NAN_CHECKING 0
 
 const uint32 NIAGARA_COMPUTE_THREADGROUP_SIZE = 64;
-const uint32 NIAGARA_MAX_COMPUTE_THREADGROUPS = 65536;
+const uint32 NIAGARA_MAX_COMPUTE_THREADGROUPS = 65535;
 
 const FString INTERPOLATED_PARAMETER_PREFIX = TEXT("PREV_");
 
diff --git a/Engine/Plugins/FX/Niagara/Source/Niagara/Public/NiagaraComponent.h b/Engine/Plugins/FX/Niagara/Source/Niagara/Public/NiagaraComponent.h
--- a/Engine/Plugins/FX/Niagara/Source/Niagara/Public/NiagaraComponent.h
+++ b/Engine/Plugins/FX/Niagara/Source/Niagara/Public/NiagaraComponent.h
@@ -14,7 +14,7 @@
 #include "NiagaraComponent.generated.h"
 
 class FMeshElementCollector;
-class NiagaraRenderer;
+class FNiagaraRenderer;
 class UNiagaraSystem;
 class UNiagaraParameterCollection;
 class UNiagaraParameterCollectionInstance;
@@ -312,6 +312,8 @@ public:
 
 	const FNiagaraParameterStore& GetOverrideParameters() const { return OverrideParameters; }
 
+	bool IsWorldReadyToRun() const;
+
 	//~ End UObject Interface.
 
 	// Called when the particle system is done
@@ -432,9 +434,10 @@ public:
 	~FNiagaraSceneProxy();
 
 	/** Called on render thread to assign new dynamic data */
-	void SetDynamicData_RenderThread(struct FNiagaraDynamicDataBase* NewDynamicData);
-	TArray<class NiagaraRenderer*>& GetEmitterRenderers() { return EmitterRenderers; }
-	void UpdateEmitterRenderers(const TArray<NiagaraRenderer*>& InRenderers);
+	const TArray<class FNiagaraRenderer*>& GetEmitterRenderers() { return EmitterRenderers; }
+
+	void CreateRenderers(const UNiagaraComponent* InComponent);
+	void ReleaseRenderers();
 
 	/** Gets whether or not this scene proxy should be rendered. */
 	bool GetRenderingEnabled() const;
@@ -478,8 +481,16 @@ private:
 	uint32 GetAllocatedSize() const;
 
 private:
-	//class NiagaraRenderer* EmitterRenderer;
-	TArray<class NiagaraRenderer*>EmitterRenderers;
+	/** Emitter Renderers in the order they appear in the emitters. */
+	TArray<FNiagaraRenderer*> EmitterRenderers;
+	
+	/** Indices of renderers in the order they should be rendered. */
+	TArray<int32> RendererDrawOrder;
+
 	bool bRenderingEnabled;
 	NiagaraEmitterInstanceBatcher* Batcher = nullptr;
+
+#if STATS
+	TStatId SystemStatID;
+#endif
 };
diff --git a/Engine/Plugins/FX/Niagara/Source/Niagara/Public/NiagaraFunctionLibrary.h b/Engine/Plugins/FX/Niagara/Source/Niagara/Public/NiagaraFunctionLibrary.h
--- a/Engine/Plugins/FX/Niagara/Source/Niagara/Public/NiagaraFunctionLibrary.h
+++ b/Engine/Plugins/FX/Niagara/Source/Niagara/Public/NiagaraFunctionLibrary.h
@@ -34,6 +34,17 @@ class NIAGARA_API UNiagaraFunctionLibrary : public UBlueprintFunctionLibrary
 
 	static UNiagaraComponent* SpawnSystemAttached(UNiagaraSystem* SystemTemplate, USceneComponent* AttachToComponent, FName AttachPointName, FVector Location, FRotator Rotation, FVector Scale, EAttachLocation::Type LocationType, bool bAutoDestroy, EPSCPoolMethod PoolingMethod);
 
+	/** Sets a Niagara StaticMesh parameter by name, overriding locally if necessary.*/
+	UFUNCTION(BlueprintCallable, Category = Niagara, meta = (DisplayName = "Set Niagara Static Mesh Component"))
+	static void OverrideSystemUserVariableStaticMeshComponent(UNiagaraComponent* NiagaraSystem, const FString& OverrideName, UStaticMeshComponent* StaticMeshComponent);
+
+	UFUNCTION(BlueprintCallable, Category = Niagara, meta = (DisplayName = "Set Niagara Static Mesh Directly"))
+	static void OverrideSystemUserVariableStaticMesh(UNiagaraComponent* NiagaraSystem, const FString& OverrideName, UStaticMesh* StaticMesh);
+
+	/** Sets a Niagara StaticMesh parameter by name, overriding locally if necessary.*/
+	UFUNCTION(BlueprintCallable, Category = Niagara, meta = (DisplayName = "Set Niagara Skeletal Mesh Component"))
+	static void OverrideSystemUserVariableSkeletalMeshComponent(UNiagaraComponent* NiagaraSystem, const FString& OverrideName, USkeletalMeshComponent* SkeletalMeshComponent);
+	
 	//This is gonna be totally reworked
 // 	UFUNCTION(BlueprintCallable, Category = Niagara, meta = (Keywords = "niagara System", UnsafeDuringActorConstruction = "true"))
 // 	static void SetUpdateScriptConstant(UNiagaraComponent* Component, FName EmitterName, FName ConstantName, FVector Value);
diff --git a/Engine/Plugins/FX/Niagara/Source/Niagara/Public/NiagaraLightRendererProperties.h b/Engine/Plugins/FX/Niagara/Source/Niagara/Public/NiagaraLightRendererProperties.h
--- a/Engine/Plugins/FX/Niagara/Source/Niagara/Public/NiagaraLightRendererProperties.h
+++ b/Engine/Plugins/FX/Niagara/Source/Niagara/Public/NiagaraLightRendererProperties.h
@@ -8,6 +8,8 @@
 #include "NiagaraCommon.h"
 #include "NiagaraLightRendererProperties.generated.h"
 
+class FNiagaraEmitterInstance;
+
 UCLASS(editinlinenew)
 class UNiagaraLightRendererProperties : public UNiagaraRendererProperties
 {
@@ -16,12 +18,15 @@ public:
 
 	UNiagaraLightRendererProperties();
 
+	//UObject Interface
 	virtual void PostInitProperties() override;
+	//UObject Interface END
 
 	static void InitCDOPropertiesAfterModuleStartup();
 
 	//~ UNiagaraRendererProperties interface
-	virtual NiagaraRenderer* CreateEmitterRenderer(ERHIFeatureLevel::Type FeatureLevel) override;
+	virtual FNiagaraRenderer* CreateEmitterRenderer(ERHIFeatureLevel::Type FeatureLevel, const FNiagaraEmitterInstance* Emitter) override;
+	virtual class FNiagaraBoundsCalculator* CreateBoundsCalculator() override { return nullptr; }
 	virtual void GetUsedMaterials(TArray<UMaterialInterface*>& OutMaterials) const override;
 	virtual bool IsSimTargetSupported(ENiagaraSimTarget InSimTarget) const override { return (InSimTarget == ENiagaraSimTarget::CPUSim); };
 #if WITH_EDITORONLY_DATA
diff --git a/Engine/Plugins/FX/Niagara/Source/Niagara/Public/NiagaraMeshRendererProperties.h b/Engine/Plugins/FX/Niagara/Source/Niagara/Public/NiagaraMeshRendererProperties.h
--- a/Engine/Plugins/FX/Niagara/Source/Niagara/Public/NiagaraMeshRendererProperties.h
+++ b/Engine/Plugins/FX/Niagara/Source/Niagara/Public/NiagaraMeshRendererProperties.h
@@ -9,6 +9,8 @@
 #include "NiagaraCommon.h"
 #include "NiagaraMeshRendererProperties.generated.h"
 
+class FNiagaraEmitterInstance;
+
 /** This enum decides how a mesh particle will orient its "facing" axis relative to camera. Must keep these in sync with NiagaraMeshVertexFactory.ush*/
 UENUM()
 enum class ENiagaraMeshFacingMode : uint8
@@ -31,13 +33,22 @@ public:
 
 	UNiagaraMeshRendererProperties();
 
+	//UObject Interface
+	virtual void PostInitProperties() override;
+#if WITH_EDITORONLY_DATA
+	virtual void PostEditChangeProperty(struct FPropertyChangedEvent& PropertyChangedEvent) override;
+#endif// WITH_EDITORONLY_DATA
+	//UObject Interface END
+
+	static void InitCDOPropertiesAfterModuleStartup();
+
 	//~ UNiagaraRendererProperties interface
-	virtual NiagaraRenderer* CreateEmitterRenderer(ERHIFeatureLevel::Type FeatureLevel) override;
+	virtual FNiagaraRenderer* CreateEmitterRenderer(ERHIFeatureLevel::Type FeatureLevel, const FNiagaraEmitterInstance* Emitter) override;
+	virtual class FNiagaraBoundsCalculator* CreateBoundsCalculator() override;
 	virtual void GetUsedMaterials(TArray<UMaterialInterface*>& OutMaterials) const override;
 	virtual bool IsSimTargetSupported(ENiagaraSimTarget InSimTarget) const override { return true; };
 
 #if WITH_EDITORONLY_DATA
-	virtual void PostEditChangeProperty(struct FPropertyChangedEvent& PropertyChangedEvent) override;
 	virtual bool IsMaterialValidForRenderer(UMaterial* Material, FText& InvalidMessage) override;
 	virtual void FixMaterial(UMaterial* Material) override;
 	virtual const TArray<FNiagaraVariable>& GetRequiredAttributes() override;
@@ -71,9 +82,6 @@ public:
 	UPROPERTY(EditAnywhere, Category = "Mesh Rendering")
 	ENiagaraMeshFacingMode FacingMode; 
 	
-	virtual void PostInitProperties() override;
-	static void InitCDOPropertiesAfterModuleStartup();
-	
 	/** Which attribute should we use for position when generating instanced meshes?*/
 	UPROPERTY(EditAnywhere, Category = "Bindings")
 	FNiagaraVariableAttributeBinding PositionBinding;
@@ -122,8 +130,6 @@ public:
 	UPROPERTY(EditAnywhere, Category = "Bindings")
 	FNiagaraVariableAttributeBinding NormalizedAgeBinding;
 
-	UPROPERTY(Transient)
-	int32 SyncId;
 protected:
 	void InitBindings();
 };
diff --git a/Engine/Plugins/FX/Niagara/Source/Niagara/Public/NiagaraModule.h b/Engine/Plugins/FX/Niagara/Source/Niagara/Public/NiagaraModule.h
--- a/Engine/Plugins/FX/Niagara/Source/Niagara/Public/NiagaraModule.h
+++ b/Engine/Plugins/FX/Niagara/Source/Niagara/Public/NiagaraModule.h
@@ -16,6 +16,8 @@ class UNiagaraScript;
 class FNiagaraCompileOptions;
 class FNiagaraCompileRequestDataBase;
 
+extern NIAGARA_API int32 GEnableVerboseNiagaraChangeIdLogging;
+
 /**
 * Niagara module interface
 */
