Copying //UE4/Dev-Niagara to Dev-Main (//UE4/Dev-Main)

diff --git a/Engine/Plugins/FX/Niagara/Source/Niagara/Public/NiagaraParameterStore.h b/Engine/Plugins/FX/Niagara/Source/Niagara/Public/NiagaraParameterStore.h
--- a/Engine/Plugins/FX/Niagara/Source/Niagara/Public/NiagaraParameterStore.h
+++ b/Engine/Plugins/FX/Niagara/Source/Niagara/Public/NiagaraParameterStore.h
@@ -340,7 +340,12 @@ public:
 	FORCEINLINE_DEBUGGABLE void SetParameterData(const uint8* Data, int32 Offset, int32 Size)
 	{
 		checkSlow(Data != nullptr);
-		FMemory::Memcpy(GetParameterData_Internal(Offset), Data, Size);
+		checkSlow((Offset + Size) <= ParameterData.Num());
+		uint8* Dest = GetParameterData_Internal(Offset);
+		if (Dest != Data)
+		{
+			FMemory::Memcpy(Dest, Data, Size);
+		}
 		OnParameterChange();
 	}
 
@@ -351,7 +356,11 @@ public:
 		if (Offset != INDEX_NONE)
 		{
 			checkSlow(!Param.IsDataInterface());
-			FMemory::Memcpy(GetParameterData_Internal(Offset), Data, Param.GetSizeInBytes());
+			uint8* Dest = GetParameterData_Internal(Offset);
+			if (Dest != Data)
+			{
+				FMemory::Memcpy(Dest, Data, Param.GetSizeInBytes());
+			}
 			OnParameterChange();
 			return true;
 		}
@@ -364,7 +373,11 @@ public:
 				AddParameter(Param, bInitInterfaces, bTriggerRebind);
 				Offset = IndexOf(Param);
 				check(Offset != INDEX_NONE);
-				FMemory::Memcpy(GetParameterData_Internal(Offset), Data, Param.GetSizeInBytes());
+				uint8* Dest = GetParameterData_Internal(Offset);
+				if (Dest != Data)
+				{
+					FMemory::Memcpy(Dest, Data, Param.GetSizeInBytes());
+				}
 				OnLayoutChange();
 				return true;
 			}
@@ -382,7 +395,12 @@ public:
 		int32 Offset = IndexOf(Param);
 		if (Offset != INDEX_NONE)
 		{
-			FMemory::Memcpy(GetParameterData_Internal(Offset), Param.GetData(), Param.GetSizeInBytes());
+			uint8* Dest = GetParameterData_Internal(Offset);
+			const uint8* Src = Param.GetData();
+			if (Dest != Src)
+			{
+				FMemory::Memcpy(Dest, Src, Param.GetSizeInBytes());
+			}
 			OnParameterChange();
 		}
 	}
diff --git a/Engine/Plugins/FX/Niagara/Source/Niagara/Public/NiagaraRenderer.h b/Engine/Plugins/FX/Niagara/Source/Niagara/Public/NiagaraRenderer.h
--- a/Engine/Plugins/FX/Niagara/Source/Niagara/Public/NiagaraRenderer.h
+++ b/Engine/Plugins/FX/Niagara/Source/Niagara/Public/NiagaraRenderer.h
@@ -20,15 +20,35 @@ NiagaraRenderer.h: Base class for Niagara render modules
 #include "SceneView.h"
 #include "NiagaraComponent.h"
 #include "NiagaraCutoutVertexBuffer.h"
+#include "NiagaraBoundsCalculator.h"
 
 class FNiagaraDataSet;
 
 /** Struct used to pass dynamic data from game thread to render thread */
 struct FNiagaraDynamicDataBase
 {
-	//Ugh. Temporarily copying the whole buffer over.
-	//After GDC i have a shelf that does just the data we need for rendering. For now, this is safer.
-	FNiagaraDataBuffer RTParticleData;
+	explicit FNiagaraDynamicDataBase(const FNiagaraEmitterInstance* InEmitter);
+	virtual ~FNiagaraDynamicDataBase();
+
+	FNiagaraDynamicDataBase() = delete;
+	FNiagaraDynamicDataBase(FNiagaraDynamicDataBase& Other) = delete;
+	FNiagaraDynamicDataBase& operator=(const FNiagaraDynamicDataBase& Other) = delete;
+
+	FNiagaraDataBuffer* GetParticleDataToRender()const;
+	FORCEINLINE ENiagaraSimTarget GetSimTarget()const { return SimTarget; }
+	FORCEINLINE FMaterialRelevance GetMaterialRelevance()const { return MaterialRelevance; }
+
+	FORCEINLINE void SetMaterialRelevance(FMaterialRelevance NewRelevance) { MaterialRelevance = NewRelevance; }
+protected:
+
+	FMaterialRelevance MaterialRelevance;
+	ENiagaraSimTarget SimTarget;
+
+	union
+	{
+		FNiagaraDataBuffer* CPUParticleData;
+		FNiagaraComputeExecutionContext* GPUExecContext;
+	}Data;
 };
 
 
@@ -54,309 +74,71 @@ private:
 	double StartTime;
 };
 
-class FNiagaraDummyRWBufferInt : public FRenderResource
-{
-public:
-	FNiagaraDummyRWBufferInt(const FString InDebugId) : DebugId(InDebugId) {}
-	FString DebugId;
-	FRWBuffer Buffer;
-	virtual void InitRHI() override;
-	virtual void ReleaseRHI() override;
-};
-
-class FNiagaraDummyRWBufferFloat : public FRenderResource
-{
-public:
-	FNiagaraDummyRWBufferFloat(const FString InDebugId) : DebugId(InDebugId) {}
-	FString DebugId;
-	FRWBuffer Buffer;
-	virtual void InitRHI() override;
-	virtual void ReleaseRHI() override;
-};
+//////////////////////////////////////////////////////////////////////////
 
 /**
-* Base class for Niagara System renderers. System renderers handle generating vertex data for and
-* drawing of simulation data coming out of FNiagaraEmitterInstance instances.
+* Base class for Niagara System renderers.
 */
-class NiagaraRenderer 
+class FNiagaraRenderer
 {
 public:
-	virtual void GetDynamicMeshElements(const TArray<const FSceneView*>& Views, const FSceneViewFamily& ViewFamily, uint32 VisibilityMap, FMeshElementCollector& Collector, const FNiagaraSceneProxy *SceneProxy) const = 0;
 
-#if RHI_RAYTRACING
-	virtual void GetDynamicRayTracingInstances(FRayTracingMaterialGatheringContext& Context, TArray<FRayTracingInstance>& OutRayTracingInstances, const FNiagaraSceneProxy* Proxy) {}
-#endif
-
-	virtual void SetDynamicData_RenderThread(FNiagaraDynamicDataBase* NewDynamicData) = 0;
-	virtual void CreateRenderThreadResources() = 0;
-	virtual void ReleaseRenderThreadResources() = 0;
-	virtual FNiagaraDynamicDataBase *GenerateVertexData(const FNiagaraSceneProxy* Proxy, FNiagaraDataSet &Data, const ENiagaraSimTarget Target) = 0;
-	virtual int GetDynamicDataSize() = 0;
+	FNiagaraRenderer(ERHIFeatureLevel::Type FeatureLevel, const UNiagaraRendererProperties *InProps, const FNiagaraEmitterInstance* Emitter);
+	virtual ~FNiagaraRenderer();
 
-	virtual bool HasDynamicData() = 0;
+	FNiagaraRenderer(const FNiagaraRenderer& Other) = delete;
+	FNiagaraRenderer& operator=(const FNiagaraRenderer& Other) = delete;
 
-	virtual bool SetMaterialUsage() = 0;
+	virtual void Initialize(ERHIFeatureLevel::Type FeatureLevel, const UNiagaraRendererProperties *InProps, const FNiagaraEmitterInstance* Emitter);
+	virtual void CreateRenderThreadResources() {}
+	virtual void ReleaseRenderThreadResources() {}
 
-	virtual void TransformChanged() = 0;
+	virtual FPrimitiveViewRelevance GetViewRelevance(const FSceneView* View, const FNiagaraSceneProxy *SceneProxy)const;
+	virtual void GetDynamicMeshElements(const TArray<const FSceneView*>& Views, const FSceneViewFamily& ViewFamily, uint32 VisibilityMap, FMeshElementCollector& Collector, const FNiagaraSceneProxy *SceneProxy) const {}
+	virtual FNiagaraDynamicDataBase* GenerateDynamicData(const FNiagaraSceneProxy* Proxy, const UNiagaraRendererProperties* InProperties, const FNiagaraEmitterInstance* Emitteride) const { return nullptr; }
 
-#if WITH_EDITORONLY_DATA
-	virtual const TArray<FNiagaraVariable>& GetRequiredAttributes() = 0;
-	virtual const TArray<FNiagaraVariable>& GetOptionalAttributes() = 0;
-#endif
+	virtual void GatherSimpleLights(FSimpleLightArray& OutParticleLights)const {}
+	virtual int32 GetDynamicDataSize()const { return 0; }
+	virtual void TransformChanged() {}
+	virtual bool IsMaterialValid(UMaterialInterface* Mat)const { return Mat != nullptr; }
 
-	virtual FPrimitiveViewRelevance GetViewRelevance(const FSceneView* View, const FNiagaraSceneProxy *SceneProxy)
-	{
-		FPrimitiveViewRelevance Result;
-		bool bHasDynamicData = HasDynamicData();
-
-		//Always draw so our LastRenderTime is updated. We may not have dynamic data if we're disabled from visibility culling.
-		Result.bDrawRelevance =/* bHasDynamicData && */SceneProxy->IsShown(View) && View->Family->EngineShowFlags.Particles;
-		Result.bShadowRelevance = bHasDynamicData && SceneProxy->IsShadowCast(View);
-		Result.bDynamicRelevance = bHasDynamicData;
-		if (bHasDynamicData && View->Family->EngineShowFlags.Bounds)
-		{
-			Result.bOpaqueRelevance = true;
-		}
-		MaterialRelevance.SetPrimitiveViewRelevance(Result);
-
-		return Result;
-	}
-
-
-	UMaterialInterface *GetMaterial()	{ return Material; }
-	void SetMaterial(UMaterialInterface *InMaterial, ERHIFeatureLevel::Type FeatureLevel)
-	{
-		Material = InMaterial;
-		if (Material->GetMaterialResource(FeatureLevel))
-		{
-			if (!Material || !SetMaterialUsage())
-			{
-				Material = UMaterial::GetDefaultMaterial(MD_Surface);
-			}
-			check(Material);
-			MaterialRelevance = Material->GetRelevance(FeatureLevel);
-		}
-	}
-	
-	virtual UClass *GetPropertiesClass() = 0;
-	virtual void SetRendererProperties(UNiagaraRendererProperties *Props) = 0;
-	virtual UNiagaraRendererProperties* GetRendererProperties() const = 0;
-
-	float GetCPUTimeMS() { return CPUTimeMS; }
-
-	void SetLocalSpace(bool bInLocalSpace) { bLocalSpace = bInLocalSpace; }
-
-	/** Release enqueues the System renderer to be killed on the render thread safely.*/
-	void Release();
-
-	FNiagaraDynamicDataBase *GetDynamicData() { return DynamicDataRender; }
+	void SortIndices(ENiagaraSortMode SortMode, int32 SortAttributeOffset, const FNiagaraDataBuffer& Buffer, const FMatrix& LocalToWorld, const FSceneView* View, FGlobalDynamicReadBuffer::FAllocation& OutIndices)const;
 
-	bool IsEnabled() const {return bEnabled;}
+	void SetDynamicData_RenderThread(FNiagaraDynamicDataBase* NewDynamicData);
+	FORCEINLINE FNiagaraDynamicDataBase *GetDynamicData()const { return DynamicDataRender; }
+	FORCEINLINE bool HasDynamicData()const { return DynamicDataRender; }
+	FORCEINLINE float GetCPUTimeMS() const { return CPUTimeMS; }
+	FORCEINLINE bool HasLights()const { return bHasLights; }
 
-	void SetEnabled(bool bInEnabled) { bEnabled = bInEnabled; }
-	
-	const FVector& GetBaseExtents() const {	return BaseExtents; }
-
-	void SortIndices(ENiagaraSortMode SortMode, int32 SortAttributeOffset, const FNiagaraDataBuffer& Buffer, const FMatrix& LocalToWorld, const FSceneView* View, FGlobalDynamicReadBuffer::FAllocation& OutIndices)const;
+#if RHI_RAYTRACING
+	virtual void GetDynamicRayTracingInstances(FRayTracingMaterialGatheringContext& Context, TArray<FRayTracingInstance>& OutRayTracingInstances, const FNiagaraSceneProxy* Proxy) {}
+#endif
 
 	static FRWBuffer& GetDummyFloatBuffer(); 
+	static FRWBuffer& GetDummyFloat4Buffer();
 	static FRWBuffer& GetDummyIntBuffer();
+	static FRWBuffer& GetDummyUIntBuffer();
 	
 protected:
-	NiagaraRenderer();
-	virtual ~NiagaraRenderer();
-
-	mutable float CPUTimeMS;
-	UMaterialInterface* Material;
-	bool bLocalSpace;
-	bool bEnabled;
-
-	FMaterialRelevance MaterialRelevance;
 
 	struct FNiagaraDynamicDataBase *DynamicDataRender;
-
-	FVector BaseExtents;
-
+	
 #if RHI_RAYTRACING
 	FRWBuffer RayTracingDynamicVertexBuffer;
 	FRayTracingGeometry RayTracingGeometry;
 #endif
-};
-
-
 
-struct FNiagaraDynamicDataSprites;
-
-/**
-* NiagaraRendererSprites renders an FNiagaraEmitterInstance as sprite particles
-*/
-class NIAGARA_API NiagaraRendererSprites : public NiagaraRenderer
-{
-public:
-
-	explicit NiagaraRendererSprites(ERHIFeatureLevel::Type FeatureLevel, UNiagaraRendererProperties *Props);
-	~NiagaraRendererSprites()
-	{
-		ReleaseRenderThreadResources();
-	}
-
-
-	virtual void ReleaseRenderThreadResources() override;
-
-	// FPrimitiveSceneProxy interface.
-	virtual void CreateRenderThreadResources() override;
-
-	virtual void GetDynamicMeshElements(const TArray<const FSceneView*>& Views, const FSceneViewFamily& ViewFamily, uint32 VisibilityMap, FMeshElementCollector& Collector, const FNiagaraSceneProxy *SceneProxy) const override;
-#if RHI_RAYTRACING
-	virtual void GetDynamicRayTracingInstances(FRayTracingMaterialGatheringContext& Context, TArray<FRayTracingInstance>& OutRayTracingInstances, const FNiagaraSceneProxy* Proxy) final override;
-#endif
-	virtual bool SetMaterialUsage() override;
-	virtual void TransformChanged() override;
-	/** Update render data buffer from attributes */
-	FNiagaraDynamicDataBase *GenerateVertexData(const FNiagaraSceneProxy* Proxy, FNiagaraDataSet &Data, const ENiagaraSimTarget Target) override;
-
-	virtual void SetDynamicData_RenderThread(FNiagaraDynamicDataBase* NewDynamicData) override;
-	int GetDynamicDataSize() override;
-	bool HasDynamicData() override;
-
-	UClass *GetPropertiesClass() override { return UNiagaraSpriteRendererProperties::StaticClass(); }
-	void SetRendererProperties(UNiagaraRendererProperties *Props) override { Properties = Cast<UNiagaraSpriteRendererProperties>(Props); }
-	virtual UNiagaraRendererProperties* GetRendererProperties() const override {
-		return Properties;
-	}
-
-#if WITH_EDITORONLY_DATA
-	virtual const TArray<FNiagaraVariable>& GetRequiredAttributes() override;
-	virtual const TArray<FNiagaraVariable>& GetOptionalAttributes() override;
-#endif
-
-private:
-
-	struct FCPUSimParticleDataAllocation
-	{
-		FGlobalDynamicReadBuffer& DynamicReadBuffer;
-		FGlobalDynamicReadBuffer::FAllocation ParticleData;
-	};
-
-	void ConditionalInitPrimitiveUniformBuffer(const FNiagaraSceneProxy *SceneProxy) const;
-	FCPUSimParticleDataAllocation ConditionalAllocateCPUSimParticleData(FNiagaraDynamicDataSprites *DynamicDataSprites, FGlobalDynamicReadBuffer& DynamicReadBuffer) const;
-	TUniformBufferRef<class FNiagaraSpriteUniformParameters> CreatePerViewUniformBuffer(const FSceneView* View, const FSceneViewFamily& ViewFamily, const FNiagaraSceneProxy *SceneProxy) const;
-	void SetVertexFactoryParticleData(
-		class FNiagaraSpriteVertexFactory& VertexFactory,
-		FNiagaraDynamicDataSprites *DynamicDataSprites,
-		FCPUSimParticleDataAllocation& CPUSimParticleDataAllocation,
-		const FSceneView* View,
-		const FNiagaraSceneProxy *SceneProxy) const;
-	void CreateMeshBatchForView(
-		const FSceneView* View,
-		const FSceneViewFamily& ViewFamily,
-		const FNiagaraSceneProxy *SceneProxy,
-		FNiagaraDynamicDataSprites *DynamicDataSprites,
-		FMeshBatch& OutMeshBatch,
-		class FNiagaraMeshCollectorResourcesSprite& OutCollectorResources) const;
-
-	UNiagaraSpriteRendererProperties *Properties;
-	mutable TUniformBuffer<FPrimitiveUniformShaderParameters> WorldSpacePrimitiveUniformBuffer;
-	class FNiagaraSpriteVertexFactory* VertexFactory;
-
-	int32 PositionOffset;
-	int32 VelocityOffset;
-	int32 RotationOffset;
-	int32 SizeOffset;
-	int32 ColorOffset;
-
-	int32 FacingOffset;
-	int32 AlignmentOffset;
-	int32 SubImageOffset;
-	int32 MaterialParamOffset;
-	int32 MaterialParamOffset1;
-	int32 MaterialParamOffset2;
-	int32 MaterialParamOffset3;
-	int32 CameraOffsetOffset;
-	int32 UVScaleOffset;
-	int32 NormalizedAgeOffset;
-	int32 MaterialRandomOffset;
-	int32 CustomSortingOffset;
-
-	int32 LastSyncId;
-
-	FNiagaraCutoutVertexBuffer CutoutVertexBuffer;
-	int32 NumCutoutVertexPerSubImage = 0;
-};
-
-
-/**
-* NiagaraRendererLights renders an FNiagaraEmitterInstance as simple lights
-*/
-class NIAGARA_API NiagaraRendererLights : public NiagaraRenderer
-{
-public:
-	struct SimpleLightData
-	{
-		FSimpleLightEntry LightEntry;
-		FSimpleLightPerViewEntry PerViewEntry;
-	};
-	explicit NiagaraRendererLights(ERHIFeatureLevel::Type FeatureLevel, UNiagaraRendererProperties *Props);
-
-	~NiagaraRendererLights()
-	{
-		ReleaseRenderThreadResources();
-	}
-
-
-	virtual FPrimitiveViewRelevance GetViewRelevance(const FSceneView* View, const FNiagaraSceneProxy *SceneProxy) override
-	{
-		FPrimitiveViewRelevance Result;
-		Result.bDrawRelevance = false;
-		Result.bShadowRelevance = false;
-		Result.bDynamicRelevance = false;
-		Result.bOpaqueRelevance = false;
-		Result.bHasSimpleLights = true;
-		//MaterialRelevance.SetPrimitiveViewRelevance(Result);
-		return Result;
-	}
-
-
-	virtual void ReleaseRenderThreadResources() override;
-
-	// FPrimitiveSceneProxy interface.
-	virtual void CreateRenderThreadResources() override;
-
-	virtual void GetDynamicMeshElements(const TArray<const FSceneView*>& Views, const FSceneViewFamily& ViewFamily, uint32 VisibilityMap, FMeshElementCollector& Collector, const FNiagaraSceneProxy *SceneProxy) const override;
-	virtual bool SetMaterialUsage() override;
-	virtual void TransformChanged() override;
-	/** Update render data buffer from attributes */
-	FNiagaraDynamicDataBase *GenerateVertexData(const FNiagaraSceneProxy* Proxy, FNiagaraDataSet &Data, const ENiagaraSimTarget Target) override;
-
-	virtual void SetDynamicData_RenderThread(FNiagaraDynamicDataBase* NewDynamicData) override;
-	int GetDynamicDataSize() override;
-	bool HasDynamicData() override;
+	mutable float CPUTimeMS;
 
-	UClass *GetPropertiesClass() override { return UNiagaraLightRendererProperties::StaticClass(); }
-	void SetRendererProperties(UNiagaraRendererProperties *Props) override { Properties = Cast<UNiagaraLightRendererProperties>(Props); }
-	virtual UNiagaraRendererProperties* GetRendererProperties() const override {
-		return Properties;
-	}
+	uint32 bLocalSpace : 1;
+	uint32 bHasLights : 1;
 
-#if WITH_EDITORONLY_DATA
-	virtual const TArray<FNiagaraVariable>& GetRequiredAttributes() override;
-	virtual const TArray<FNiagaraVariable>& GetOptionalAttributes() override;
+#if STATS
+	TStatId EmitterStatID;
 #endif
 
-	TArray<SimpleLightData> &GetLights() { return LightArray; }
-private:
-	UNiagaraLightRendererProperties *Properties;
-	TArray<SimpleLightData>LightArray;
-};
-
-
-struct FNiagaraDynamicDataLights : public FNiagaraDynamicDataBase
-{
-	TArray<NiagaraRendererLights::SimpleLightData> LightArray;
+	/** Cached array of materials used from the properties data. Validated with usage flags etc. */
+	TArray<UMaterialInterface*> BaseMaterials_GT;
+	FMaterialRelevance BaseMaterialRelevance_GT;
 };
 
-
-
-
-
-
diff --git a/Engine/Plugins/FX/Niagara/Source/Niagara/Public/NiagaraRendererLights.h b/Engine/Plugins/FX/Niagara/Source/Niagara/Public/NiagaraRendererLights.h
new file mode 100644
--- /dev/null
+++ b/Engine/Plugins/FX/Niagara/Source/Niagara/Public/NiagaraRendererLights.h
@@ -0,0 +1,31 @@
+// Copyright 1998-2019 Epic Games, Inc. All Rights Reserved.
+
+/*==============================================================================
+NiagaraRendererLights.h: Renderer for rendering Niagara particles as Lights.
+==============================================================================*/
+#pragma once
+
+#include "NiagaraRenderer.h"
+
+/**
+* NiagaraRendererLights renders an FNiagaraEmitterInstance as simple lights
+*/
+class NIAGARA_API FNiagaraRendererLights : public FNiagaraRenderer
+{
+public:
+	struct SimpleLightData
+	{
+		FSimpleLightEntry LightEntry;
+		FSimpleLightPerViewEntry PerViewEntry;
+	};
+
+	FNiagaraRendererLights(ERHIFeatureLevel::Type FeatureLevel, const UNiagaraRendererProperties *InProps, const FNiagaraEmitterInstance* Emitter);
+
+	//FNiagaraRenderer interface
+	virtual FPrimitiveViewRelevance GetViewRelevance(const FSceneView* View, const FNiagaraSceneProxy *SceneProxy)const override;
+	virtual FNiagaraDynamicDataBase *GenerateDynamicData(const FNiagaraSceneProxy* Proxy, const UNiagaraRendererProperties* InProperties, const FNiagaraEmitterInstance* Emitter) const override;
+	virtual void GatherSimpleLights(FSimpleLightArray& OutParticleLights)const override;
+	//FNiagaraRenderer interface END
+};
+
+
diff --git a/Engine/Plugins/FX/Niagara/Source/Niagara/Public/NiagaraRendererMeshes.h b/Engine/Plugins/FX/Niagara/Source/Niagara/Public/NiagaraRendererMeshes.h
--- a/Engine/Plugins/FX/Niagara/Source/Niagara/Public/NiagaraRendererMeshes.h
+++ b/Engine/Plugins/FX/Niagara/Source/Niagara/Public/NiagaraRendererMeshes.h
@@ -10,90 +10,46 @@ NiagaraRenderer.h: Base class for Niagara render modules
 
 class FNiagaraDataSet;
 
-
-
-struct FNiagaraDynamicDataMesh : public FNiagaraDynamicDataBase
-{
-	//Direct ptr to the dataset. ONLY FOR USE BE GPU EMITTERS.
-	//TODO: Even this needs to go soon.
-	const FNiagaraDataSet *DataSet;
-};
-
-
-
 /**
 * NiagaraRendererSprites renders an FNiagaraEmitterInstance as sprite particles
 */
-class NIAGARA_API NiagaraRendererMeshes : public NiagaraRenderer
+class NIAGARA_API FNiagaraRendererMeshes : public FNiagaraRenderer
 {
 public:
 
-	explicit NiagaraRendererMeshes(ERHIFeatureLevel::Type FeatureLevel, UNiagaraRendererProperties *Props);
-	~NiagaraRendererMeshes()
-	{
-		ReleaseRenderThreadResources();
-	}
-
-
-	virtual void ReleaseRenderThreadResources() override;
-
-	// FPrimitiveSceneProxy interface.
+	FNiagaraRendererMeshes(ERHIFeatureLevel::Type FeatureLevel, const UNiagaraRendererProperties *InProps, const FNiagaraEmitterInstance* Emitter);
+	
+	//FNiagaraRenderer Interface
 	virtual void CreateRenderThreadResources() override;
+	virtual void ReleaseRenderThreadResources() override;
 
 	virtual void GetDynamicMeshElements(const TArray<const FSceneView*>& Views, const FSceneViewFamily& ViewFamily, uint32 VisibilityMap, FMeshElementCollector& Collector, const FNiagaraSceneProxy *SceneProxy) const override;
-	virtual bool SetMaterialUsage() override;
+	virtual FNiagaraDynamicDataBase* GenerateDynamicData(const FNiagaraSceneProxy* Proxy, const UNiagaraRendererProperties* InProperties, const FNiagaraEmitterInstance* Emitter) const override;
+	virtual int32 GetDynamicDataSize()const override;
 	virtual void TransformChanged() override;
-	/** Update render data buffer from attributes */
-	FNiagaraDynamicDataBase *GenerateVertexData(const FNiagaraSceneProxy* Proxy, FNiagaraDataSet &Data, const ENiagaraSimTarget Target) override;
-
-	virtual void SetDynamicData_RenderThread(FNiagaraDynamicDataBase* NewDynamicData) override;
-	int GetDynamicDataSize() override;
-	bool HasDynamicData() override;
-
-	virtual FPrimitiveViewRelevance GetViewRelevance(const FSceneView* View, const FNiagaraSceneProxy *SceneProxy)
-	{
-		FPrimitiveViewRelevance Result;
-		bool bHasDynamicData = HasDynamicData();
-		//Always draw so our LastRenderTime is updated. We may not have dynamic data if we're disabled from visibility culling.
-		Result.bDrawRelevance = /*bHasDynamicData && */SceneProxy->IsShown(View) && View->Family->EngineShowFlags.Particles;
-		Result.bShadowRelevance = bHasDynamicData && SceneProxy->IsShadowCast(View);
-		Result.bDynamicRelevance = bHasDynamicData;
-
-		if (bHasDynamicData)
-		{
-			Result.bOpaqueRelevance = MaterialRelevance.bOpaque;
-			Result.bNormalTranslucencyRelevance = MaterialRelevance.bNormalTranslucency;
-			Result.bSeparateTranslucencyRelevance = MaterialRelevance.bSeparateTranslucency;
-			Result.bDistortionRelevance = MaterialRelevance.bDistortion;
-		}
-
-		return Result;
-	}
+	virtual bool IsMaterialValid(UMaterialInterface* Mat)const override;
+	//FNiagaraRenderer Interface END
 
-
-
-
-	UClass *GetPropertiesClass() override { return UNiagaraMeshRendererProperties::StaticClass(); }
-	void SetRendererProperties(UNiagaraRendererProperties *Props) override { Properties = Cast<UNiagaraMeshRendererProperties>(Props); }
-	virtual UNiagaraRendererProperties* GetRendererProperties()  const override {
-		return Properties;
-	}
-#if WITH_EDITORONLY_DATA
-	virtual const TArray<FNiagaraVariable>& GetRequiredAttributes() override;
-	virtual const TArray<FNiagaraVariable>& GetOptionalAttributes() override;
-#endif
 	void SetupVertexFactory(FNiagaraMeshVertexFactory *InVertexFactory, const FStaticMeshLODResources& LODResources) const;
 
 private:
-	UNiagaraMeshRendererProperties *Properties;
 	mutable TUniformBuffer<FPrimitiveUniformShaderParameters> WorldSpacePrimitiveUniformBuffer;
 	class FNiagaraMeshVertexFactory* VertexFactory;
 
+	/** Render data of the static mesh we use. */
+	FStaticMeshRenderData* MeshRenderData;
+
+	ENiagaraSortMode SortMode;
+	ENiagaraMeshFacingMode FacingMode;
+	uint32 bOverrideMaterials : 1;
+	uint32 bSortOnlyWhenTranslucent : 1;
+
 	int32 PositionOffset;
 	int32 VelocityOffset;
 	int32 ColorOffset;
 	int32 ScaleOffset;
 	int32 SizeOffset;
+	uint32 MaterialParamValidMask;
 	int32 MaterialParamOffset;
 	int32 MaterialParamOffset1;
 	int32 MaterialParamOffset2;
@@ -102,5 +58,4 @@ private:
 	int32 NormalizedAgeOffset;
 	int32 MaterialRandomOffset;
 	int32 CustomSortingOffset;
-	int32 LastSyncedId;
 };
diff --git a/Engine/Plugins/FX/Niagara/Source/Niagara/Public/NiagaraRendererProperties.h b/Engine/Plugins/FX/Niagara/Source/Niagara/Public/NiagaraRendererProperties.h
--- a/Engine/Plugins/FX/Niagara/Source/Niagara/Public/NiagaraRendererProperties.h
+++ b/Engine/Plugins/FX/Niagara/Source/Niagara/Public/NiagaraRendererProperties.h
@@ -19,9 +19,10 @@
 * and handed back to the System renderer on load.
 */
 
-class NiagaraRenderer;
+class FNiagaraRenderer;
 class UMaterial;
 class UMaterialInterface;
+class FNiagaraEmitterInstance;
 
 UCLASS(ABSTRACT)
 class NIAGARA_API UNiagaraRendererProperties : public UNiagaraMergeable
@@ -33,7 +34,10 @@ public:
 	: bIsEnabled(true)
 	{
 	}
-	virtual NiagaraRenderer* CreateEmitterRenderer(ERHIFeatureLevel::Type FeatureLevel) PURE_VIRTUAL ( UNiagaraRendererProperties::CreateEmitterRenderer, return nullptr;);
+
+	//UObject Interface End
+	virtual FNiagaraRenderer* CreateEmitterRenderer(ERHIFeatureLevel::Type FeatureLevel, const FNiagaraEmitterInstance* Emitter) PURE_VIRTUAL ( UNiagaraRendererProperties::CreateEmitterRenderer, return nullptr;);
+	virtual class FNiagaraBoundsCalculator* CreateBoundsCalculator() PURE_VIRTUAL(UNiagaraRendererProperties::CreateBoundsCalculator, return nullptr;);
 	virtual void GetUsedMaterials(TArray<UMaterialInterface*>& OutMaterials) const PURE_VIRTUAL(UNiagaraRendererProperties::GetUsedMaterials,);
 
 	virtual bool IsSimTargetSupported(ENiagaraSimTarget InSimTarget) const { return false; };
@@ -51,15 +55,9 @@ public:
 	// GPU simulation uses DrawIndirect, so the sim step needs to know indices per instance in order to prepare the draw call parameters
 	virtual uint32 GetNumIndicesPerInstance() { return 0; }
 
-	virtual bool GetIsEnabled() const
-	{
-		return bIsEnabled;
-	}
-
-	virtual void SetIsEnabled(bool bInIsEnabled)
-	{
-		bIsEnabled = bInIsEnabled;
-	}
+	virtual bool GetIsEnabled() const { return bIsEnabled; }
+	virtual void SetIsEnabled(bool bInIsEnabled) { bIsEnabled = bInIsEnabled; }
+	
 	/** By default, emitters are drawn in the order that they are added to the system. This value will allow you to control the order in a more fine-grained manner. 
 	Materials of the same type (i.e. Transparent) will draw in order from lowest to highest within the system. The default value is 0.*/
 	UPROPERTY(EditAnywhere, Category = "Sort Order")
diff --git a/Engine/Plugins/FX/Niagara/Source/Niagara/Public/NiagaraRendererRibbons.h b/Engine/Plugins/FX/Niagara/Source/Niagara/Public/NiagaraRendererRibbons.h
--- a/Engine/Plugins/FX/Niagara/Source/Niagara/Public/NiagaraRendererRibbons.h
+++ b/Engine/Plugins/FX/Niagara/Source/Niagara/Public/NiagaraRendererRibbons.h
@@ -14,64 +14,41 @@ class FNiagaraDataSet;
 * NiagaraRendererRibbons renders an FNiagaraEmitterInstance as a ribbon connecting all particles
 * in order by particle age.
 */
-class NIAGARA_API NiagaraRendererRibbons : public NiagaraRenderer
+class NIAGARA_API FNiagaraRendererRibbons : public FNiagaraRenderer
 {
 public:
+	FNiagaraRendererRibbons(ERHIFeatureLevel::Type FeatureLevel, const UNiagaraRendererProperties *InProps, const FNiagaraEmitterInstance* Emitter);	// FNiagaraRenderer Interface 
 
-	NiagaraRendererRibbons(ERHIFeatureLevel::Type FeatureLevel, UNiagaraRendererProperties *Props);
-	~NiagaraRendererRibbons()
-	{
-		ReleaseRenderThreadResources();
-	}
-
-	virtual void ReleaseRenderThreadResources() override;
+	// FNiagaraRenderer Interface 
 	virtual void CreateRenderThreadResources() override;
-
+	virtual void ReleaseRenderThreadResources() override;
 
 	virtual void GetDynamicMeshElements(const TArray<const FSceneView*>& Views, const FSceneViewFamily& ViewFamily, uint32 VisibilityMap, FMeshElementCollector& Collector, const FNiagaraSceneProxy *SceneProxy) const override;
-
-	virtual bool SetMaterialUsage() override;
-
+	virtual FNiagaraDynamicDataBase *GenerateDynamicData(const FNiagaraSceneProxy* Proxy, const UNiagaraRendererProperties* InProperties, const FNiagaraEmitterInstance* Emitteride) const override;
+	virtual int32 GetDynamicDataSize()const override;
 	virtual void TransformChanged() override;
+	virtual bool IsMaterialValid(UMaterialInterface* Mat)const override;
+	//FNiagaraInterface END
 
-	/** Update render data buffer from attributes */
-	FNiagaraDynamicDataBase *GenerateVertexData(const FNiagaraSceneProxy* Proxy, FNiagaraDataSet &Data, const ENiagaraSimTarget Target) override;
-
-	void AddDynamicParam(TArray<FNiagaraRibbonVertexDynamicParameter>& ParamData, const FVector4& DynamicParam)
-	{
-		FNiagaraRibbonVertexDynamicParameter Param;
-		Param.DynamicValue[0] = DynamicParam.X;
-		Param.DynamicValue[1] = DynamicParam.Y;
-		Param.DynamicValue[2] = DynamicParam.Z;
-		Param.DynamicValue[3] = DynamicParam.W;
-		ParamData.Add(Param);
-	}
-
-	virtual void SetDynamicData_RenderThread(FNiagaraDynamicDataBase* NewDynamicData) override;
-	int GetDynamicDataSize() override;
-	bool HasDynamicData() override;
-
-#if WITH_EDITORONLY_DATA
-	virtual const TArray<FNiagaraVariable>& GetRequiredAttributes() override;
-	virtual const TArray<FNiagaraVariable>& GetOptionalAttributes() override;
-#endif
-
-	UClass *GetPropertiesClass() override { return UNiagaraRibbonRendererProperties::StaticClass(); }
-	void SetRendererProperties(UNiagaraRendererProperties *Props) override { Properties = Cast<UNiagaraRibbonRendererProperties>(Props); }
-	virtual UNiagaraRendererProperties* GetRendererProperties() const override 
-	{
-		return Properties;
-	}
-
+	FORCEINLINE void AddDynamicParam(TArray<FNiagaraRibbonVertexDynamicParameter>& ParamData, const FVector4& DynamicParam);
 protected:
-
 	static void GenerateIndexBuffer(uint16* OutIndices, const TArray<int32>& SegmentData, int32 MaxTessellation, bool bInvertOrder, bool bCullTwistedStrips);
 
 private:
-
 	class FNiagaraRibbonVertexFactory *VertexFactory;
-	UNiagaraRibbonRendererProperties *Properties;
 	mutable TUniformBuffer<FPrimitiveUniformShaderParameters> WorldSpacePrimitiveUniformBuffer;
+
+	ENiagaraRibbonFacingMode FacingMode;
+	float UV0TilingDistance;
+	FVector2D UV0Scale;
+	FVector2D UV0Offset;
+	ENiagaraRibbonAgeOffsetMode UV0AgeOffsetMode;
+	float UV1TilingDistance;
+	FVector2D UV1Scale;
+	FVector2D UV1Offset;
+	ENiagaraRibbonAgeOffsetMode UV1AgeOffsetMode;
+	ENiagaraRibbonDrawDirection DrawDirection;
+
 	int32 PositionDataOffset;
 	int32 VelocityDataOffset;
 	int32 WidthDataOffset;
@@ -80,21 +57,18 @@ private:
 	int32 ColorDataOffset;
 	int32 NormalizedAgeDataOffset;
 	int32 MaterialRandomDataOffset;
-	int32 LastSyncedId;
+	uint32 MaterialParamValidMask;
 	int32 MaterialParamOffset;
 	int32 MaterialParamOffset1;
 	int32 MaterialParamOffset2;
 	int32 MaterialParamOffset3;
 
 	// Average curvature of the segments.
-	float TessellationAngle = 0;
+	mutable float TessellationAngle = 0;
 	// Average angle of the curvature of the segments (in radian).
-	float TessellationCurvature = 0;
-
+	mutable float TessellationCurvature = 0;
 	// Average twist of the segments.
-	float TessellationTwistAngle = 0;
+	mutable float TessellationTwistAngle = 0;
 	// Average twist curvature of the segments.
-	float TessellationTwistCurvature = 0;
+	mutable float TessellationTwistCurvature = 0;
 };
-
-
diff --git a/Engine/Plugins/FX/Niagara/Source/Niagara/Public/NiagaraRendererSprites.h b/Engine/Plugins/FX/Niagara/Source/Niagara/Public/NiagaraRendererSprites.h
new file mode 100644
--- /dev/null
+++ b/Engine/Plugins/FX/Niagara/Source/Niagara/Public/NiagaraRendererSprites.h
@@ -0,0 +1,99 @@
+// Copyright 1998-2019 Epic Games, Inc. All Rights Reserved.
+
+/*==============================================================================
+NiagaraRendererSprites.h: Renderer for rendering Niagara particles as sprites.
+==============================================================================*/
+
+#pragma once
+
+#include "NiagaraRenderer.h"
+
+struct FNiagaraDynamicDataSprites;
+
+/**
+* FNiagaraRendererSprites renders an FNiagaraEmitterInstance as sprite particles
+*/
+class NIAGARA_API FNiagaraRendererSprites : public FNiagaraRenderer
+{
+public:
+
+	FNiagaraRendererSprites(ERHIFeatureLevel::Type FeatureLevel, const UNiagaraRendererProperties *InProps, const FNiagaraEmitterInstance* Emitter);
+
+	//FNiagaraRenderer interface
+	virtual void CreateRenderThreadResources()override;
+	virtual void ReleaseRenderThreadResources()override;
+
+	virtual void GetDynamicMeshElements(const TArray<const FSceneView*>& Views, const FSceneViewFamily& ViewFamily, uint32 VisibilityMap, FMeshElementCollector& Collector, const FNiagaraSceneProxy *SceneProxy) const override;
+	virtual FNiagaraDynamicDataBase* GenerateDynamicData(const FNiagaraSceneProxy* Proxy, const UNiagaraRendererProperties* InProperties, const FNiagaraEmitterInstance* Emitter) const override;
+	virtual int GetDynamicDataSize()const override;
+	virtual void TransformChanged()override;
+	virtual bool IsMaterialValid(UMaterialInterface* Mat)const override;
+
+#if RHI_RAYTRACING
+		virtual void GetDynamicRayTracingInstances(FRayTracingMaterialGatheringContext& Context, TArray<FRayTracingInstance>& OutRayTracingInstances, const FNiagaraSceneProxy* Proxy) final override;
+#endif
+	//FNiagaraRenderer interface END
+
+private:
+	struct FCPUSimParticleDataAllocation
+	{
+		FGlobalDynamicReadBuffer& DynamicReadBuffer;
+		FGlobalDynamicReadBuffer::FAllocation ParticleData;
+	};
+
+	void ConditionalInitPrimitiveUniformBuffer(const FNiagaraSceneProxy *SceneProxy) const;
+	FCPUSimParticleDataAllocation ConditionalAllocateCPUSimParticleData(FNiagaraDynamicDataSprites *DynamicDataSprites, FGlobalDynamicReadBuffer& DynamicReadBuffer) const;
+	TUniformBufferRef<class FNiagaraSpriteUniformParameters> CreatePerViewUniformBuffer(const FSceneView* View, const FSceneViewFamily& ViewFamily, const FNiagaraSceneProxy *SceneProxy) const;
+	void SetVertexFactoryParticleData(
+		class FNiagaraSpriteVertexFactory& VertexFactory,
+		FNiagaraDynamicDataSprites *DynamicDataSprites,
+		FCPUSimParticleDataAllocation& CPUSimParticleDataAllocation,
+		const FSceneView* View,
+		const FNiagaraSceneProxy *SceneProxy) const;
+	void CreateMeshBatchForView(
+		const FSceneView* View,
+		const FSceneViewFamily& ViewFamily,
+		const FNiagaraSceneProxy *SceneProxy,
+		FNiagaraDynamicDataSprites *DynamicDataSprites,
+		FMeshBatch& OutMeshBatch,
+		class FNiagaraMeshCollectorResourcesSprite& OutCollectorResources) const;
+
+	mutable TUniformBuffer<FPrimitiveUniformShaderParameters> WorldSpacePrimitiveUniformBuffer;
+	class FNiagaraSpriteVertexFactory* VertexFactory;
+
+	//Cached data from the properties struct.
+	ENiagaraSpriteAlignment Alignment;
+	ENiagaraSpriteFacingMode FacingMode;
+	FVector CustomFacingVectorMask;
+	FVector2D PivotInUVSpace;
+	ENiagaraSortMode SortMode;
+	FVector2D SubImageSize;
+	uint32 bSubImageBlend : 1;
+	uint32 bRemoveHMDRollInVR : 1;
+	uint32 bSortOnlyWhenTranslucent : 1;
+	float MinFacingCameraBlendDistance;
+	float MaxFacingCameraBlendDistance;
+	FNiagaraCutoutVertexBuffer CutoutVertexBuffer;
+	int32 NumCutoutVertexPerSubImage = 0;
+
+
+	//Offsets into the emitter's dataset for each bound attribute.
+	int32 PositionOffset;
+	int32 ColorOffset;
+	int32 VelocityOffset;
+	int32 RotationOffset;
+	int32 SizeOffset;
+	int32 FacingOffset;
+	int32 AlignmentOffset;
+	int32 SubImageOffset;
+	uint32 MaterialParamValidMask;
+	int32 MaterialParamOffset;
+	int32 MaterialParamOffset1;
+	int32 MaterialParamOffset2;
+	int32 MaterialParamOffset3;
+	int32 CameraOffsetOffset;
+	int32 UVScaleOffset;
+	int32 MaterialRandomOffset;
+	int32 CustomSortingOffset;
+	int32 NormalizedAgeOffset;
+};
diff --git a/Engine/Plugins/FX/Niagara/Source/Niagara/Public/NiagaraRibbonRendererProperties.h b/Engine/Plugins/FX/Niagara/Source/Niagara/Public/NiagaraRibbonRendererProperties.h
--- a/Engine/Plugins/FX/Niagara/Source/Niagara/Public/NiagaraRibbonRendererProperties.h
+++ b/Engine/Plugins/FX/Niagara/Source/Niagara/Public/NiagaraRibbonRendererProperties.h
@@ -8,6 +8,8 @@
 #include "NiagaraCommon.h"
 #include "NiagaraRibbonRendererProperties.generated.h"
 
+class FNiagaraEmitterInstance;
+
 UENUM()
 enum class ENiagaraRibbonFacingMode : uint8
 {
@@ -39,17 +41,30 @@ class UNiagaraRibbonRendererProperties : public UNiagaraRendererProperties
 public:
 	GENERATED_BODY()
 
-	UNiagaraRibbonRendererProperties()
-		: Material(nullptr)
-		, FacingMode(ENiagaraRibbonFacingMode::Screen)
-		, UV0TilingDistance(0.0f)
-		, UV0Scale(FVector2D(1.0f, 1.0f))
-		, UV0AgeOffsetMode(ENiagaraRibbonAgeOffsetMode::Scale)
-		, UV1TilingDistance(0.0f)
-		, UV1Scale(FVector2D(1.0f, 1.0f))
-		, UV1AgeOffsetMode(ENiagaraRibbonAgeOffsetMode::Scale)
-	{
-	}
+	UNiagaraRibbonRendererProperties();
+
+	//UObject Interface
+	virtual void PostInitProperties() override;
+#if WITH_EDITOR
+	virtual void PostEditChangeProperty(struct FPropertyChangedEvent& PropertyChangedEvent) override;
+#endif
+	//UObject Interface END
+
+	static void InitCDOPropertiesAfterModuleStartup();
+
+	//UNiagaraRendererProperties Interface
+	virtual FNiagaraRenderer* CreateEmitterRenderer(ERHIFeatureLevel::Type FeatureLevel, const FNiagaraEmitterInstance* Emitter) override;
+	virtual class FNiagaraBoundsCalculator* CreateBoundsCalculator() override;
+	virtual void GetUsedMaterials(TArray<UMaterialInterface*>& OutMaterials) const override;
+	virtual bool IsSimTargetSupported(ENiagaraSimTarget InSimTarget) const override { return (InSimTarget == ENiagaraSimTarget::CPUSim); };
+#if WITH_EDITOR
+	virtual bool IsMaterialValidForRenderer(UMaterial* Material, FText& InvalidMessage) override;
+	virtual void FixMaterial(UMaterial* Material);
+	virtual const TArray<FNiagaraVariable>& GetRequiredAttributes() override;
+	virtual const TArray<FNiagaraVariable>& GetOptionalAttributes() override;
+#endif
+	//UNiagaraRendererProperties Interface END
+
 
 	UPROPERTY(EditAnywhere, Category = "Ribbon Rendering")
 	UMaterialInterface* Material;
@@ -85,21 +100,6 @@ public:
 	UPROPERTY(EditAnywhere, Category = "Ribbon Rendering")
 	ENiagaraRibbonDrawDirection DrawDirection;
 
-	//~ UNiagaraRendererProperties interface
-	virtual NiagaraRenderer* CreateEmitterRenderer(ERHIFeatureLevel::Type FeatureLevel) override;
-	virtual void GetUsedMaterials(TArray<UMaterialInterface*>& OutMaterials) const override;
-	virtual bool IsSimTargetSupported(ENiagaraSimTarget InSimTarget) const override { return (InSimTarget == ENiagaraSimTarget::CPUSim); };
-#if WITH_EDITORONLY_DATA
-	virtual void PostEditChangeProperty(struct FPropertyChangedEvent& PropertyChangedEvent) override;
-	virtual const TArray<FNiagaraVariable>& GetRequiredAttributes() override;
-	virtual const TArray<FNiagaraVariable>& GetOptionalAttributes() override;
-	virtual bool IsMaterialValidForRenderer(UMaterial* Material, FText& InvalidMessage) override;
-	virtual void FixMaterial(UMaterial* Material);
-#endif
-
-	virtual void PostInitProperties() override;
-	static void InitCDOPropertiesAfterModuleStartup();
-
 
 	/** Which attribute should we use for position when generating ribbons?*/
 	UPROPERTY(EditAnywhere, AdvancedDisplay, Category = "Bindings")
@@ -157,8 +157,6 @@ public:
 	UPROPERTY(EditAnywhere, AdvancedDisplay, Category = "Bindings")
 	FNiagaraVariableAttributeBinding DynamicMaterial3Binding;
 
-	UPROPERTY(Transient)
-	int32 SyncId;
 protected:
 	void InitBindings();
 };
diff --git a/Engine/Plugins/FX/Niagara/Source/Niagara/Public/NiagaraSpriteRendererProperties.h b/Engine/Plugins/FX/Niagara/Source/Niagara/Public/NiagaraSpriteRendererProperties.h
--- a/Engine/Plugins/FX/Niagara/Source/Niagara/Public/NiagaraSpriteRendererProperties.h
+++ b/Engine/Plugins/FX/Niagara/Source/Niagara/Public/NiagaraSpriteRendererProperties.h
@@ -47,23 +47,29 @@ public:
 
 	UNiagaraSpriteRendererProperties();
 
+	//UObject Interface
 	virtual void PostLoad() override;
 	virtual void PostInitProperties() override;
 	virtual void Serialize(FStructuredArchive::FRecord Record) override;
+#if WITH_EDITORONLY_DATA
+	virtual void PostEditChangeProperty(struct FPropertyChangedEvent& PropertyChangedEvent) override;
+#endif // WITH_EDITORONLY_DATA
+	//UObject Interface END
 
 	static void InitCDOPropertiesAfterModuleStartup();
 
-	//~ UNiagaraRendererProperties interface
-	virtual NiagaraRenderer* CreateEmitterRenderer(ERHIFeatureLevel::Type FeatureLevel) override;
+	//UNiagaraRendererProperties interface
+	virtual FNiagaraRenderer* CreateEmitterRenderer(ERHIFeatureLevel::Type FeatureLevel, const FNiagaraEmitterInstance* Emitter) override;
+	virtual class FNiagaraBoundsCalculator* CreateBoundsCalculator() override;
 	virtual void GetUsedMaterials(TArray<UMaterialInterface*>& OutMaterials) const override;
 	virtual bool IsSimTargetSupported(ENiagaraSimTarget InSimTarget) const override { return true; };
-#if WITH_EDITORONLY_DATA
-	virtual void PostEditChangeProperty(struct FPropertyChangedEvent& PropertyChangedEvent) override;
+#if WITH_EDITOR
 	virtual bool IsMaterialValidForRenderer(UMaterial* Material, FText& InvalidMessage) override;
 	virtual void FixMaterial(UMaterial* Material) override;
 	virtual const TArray<FNiagaraVariable>& GetRequiredAttributes() override;
 	virtual const TArray<FNiagaraVariable>& GetOptionalAttributes() override;
-#endif // WITH_EDITORONLY_DATA
+#endif
+	//UNiagaraMaterialRendererProperties interface END
 
 	int32 GetNumCutoutVertexPerSubimage() const;
 	virtual uint32 GetNumIndicesPerInstance();
@@ -184,9 +190,6 @@ public:
 	UPROPERTY(EditAnywhere, Category = "Bindings")
 	FNiagaraVariableAttributeBinding NormalizedAgeBinding;
 
-	UPROPERTY(Transient)
-	int32 SyncId;
-
 	void InitBindings();
 
 #if WITH_EDITORONLY_DATA
diff --git a/Engine/Plugins/FX/Niagara/Source/Niagara/Public/NiagaraSystemInstance.h b/Engine/Plugins/FX/Niagara/Source/Niagara/Public/NiagaraSystemInstance.h
--- a/Engine/Plugins/FX/Niagara/Source/Niagara/Public/NiagaraSystemInstance.h
+++ b/Engine/Plugins/FX/Niagara/Source/Niagara/Public/NiagaraSystemInstance.h
@@ -6,7 +6,6 @@
 #include "NiagaraEmitter.h"
 #include "NiagaraEmitterHandle.h"
 #include "NiagaraSystem.h"
-#include "NiagaraParameterBindingInstance.h"
 #include "NiagaraDataInterfaceBindingInstance.h"
 #include "Templates/UniquePtr.h"
 #include "NiagaraCommon.h"
@@ -17,9 +16,12 @@ class UNiagaraComponent;
 class FNiagaraSystemInstance;
 class FNiagaraSystemSimulation;
 class NiagaraEmitterInstanceBatcher;
+class FNiagaraGPUSystemTick;
 
 class NIAGARA_API FNiagaraSystemInstance 
 {
+	friend class FNiagaraGPUSystemTick;
+
 public:
 	DECLARE_MULTICAST_DELEGATE(FOnInitialized);
 	DECLARE_MULTICAST_DELEGATE_OneParam(FOnComplete, FNiagaraSystemInstance*);
@@ -73,6 +75,7 @@ public:
 	FORCEINLINE FNiagaraParameterStore& GetInstanceParameters() { return InstanceParameters; }
 	
 	FNiagaraWorldManager* GetWorldManager()const;
+	bool RequiresDistanceFieldData() const;
 
 	/** Requests the the simulation be reset on the next tick. */
 	void Reset(EResetMode Mode, bool bBindParams = false);
@@ -97,11 +100,11 @@ public:
 	FORCEINLINE bool IsDisabled()const { return ActualExecutionState == ENiagaraExecutionState::Disabled; }
 
 	/** Gets the simulation for the supplied emitter handle. */
-	TSharedPtr<FNiagaraEmitterInstance> GetSimulationForHandle(const FNiagaraEmitterHandle& EmitterHandle);
+	TSharedPtr<FNiagaraEmitterInstance, ESPMode::ThreadSafe> GetSimulationForHandle(const FNiagaraEmitterHandle& EmitterHandle);
 
 	UNiagaraSystem* GetSystem()const;
 	FORCEINLINE UNiagaraComponent *GetComponent() { return Component; }
-	FORCEINLINE TArray<TSharedRef<FNiagaraEmitterInstance> > &GetEmitters()	{ return Emitters; }
+	FORCEINLINE TArray<TSharedRef<FNiagaraEmitterInstance, ESPMode::ThreadSafe> > &GetEmitters()	{ return Emitters; }
 	FORCEINLINE FBox &GetSystemBounds()	{ return SystemBounds;  }
 
 	FNiagaraEmitterInstance* GetEmitterByID(FGuid InID);
@@ -133,6 +136,7 @@ public:
 #endif
 
 	FName GetIDName() { return IDName; }
+	FGuid GetId() { return ID; }
 
 	/** Returns the instance data for a particular interface for this System. */
 	FORCEINLINE void* FindDataInterfaceInstanceData(UNiagaraDataInterface* Interface) 
@@ -204,6 +208,11 @@ public:
 
 	NiagaraEmitterInstanceBatcher* GetBatcher() const { return Batcher; }
 
+	static bool AllocateSystemInstance(class UNiagaraComponent* InComponent, TUniquePtr< FNiagaraSystemInstance >& OutSystemInstanceAllocation);
+	static bool DeallocateSystemInstance(TUniquePtr< FNiagaraSystemInstance >& SystemInstanceAllocation);
+	/*void SetHasGPUEmitters(bool bInHasGPUEmitters) { bHasGPUEmitters = bInHasGPUEmitters; }
+	bool HasGPUEmitters() { return bHasGPUEmitters;  }*/
+
 	int32 GetDetailLevel()const;
 private:
 
@@ -215,12 +224,6 @@ private:
 	/** Resets for restart, assumes no change in emitter setup */
 	void ResetInternal(bool bResetSimulations);
 
-	/** Updates the renders for the simulations. Gathers both the EmitterRenderers that were previously set as well as the ones that we  create within.*/
-	void UpdateRenderModules(ERHIFeatureLevel::Type InFeatureLevel, TArray<NiagaraRenderer*>& OutNewRenderers, TArray<NiagaraRenderer*>& OutOldRenderers);
-
-	/** Updates the scene proxy for the System with the specified EmitterRenderer array. Note that this is pushed onto the rendering thread behind the scenes.*/
-	void UpdateProxy(TArray<NiagaraRenderer*>& InRenderers);
-
 	/** Call PrepareForSImulation on each data source from the simulations and determine which need per-tick updates.*/
 	void InitDataInterfaces();	
 
@@ -243,7 +246,7 @@ private:
 
 	TMap<FNiagaraDataSetID, FNiagaraDataSet> ExternalEvents;
 
-	TArray< TSharedRef<FNiagaraEmitterInstance> > Emitters;
+	TArray< TSharedRef<FNiagaraEmitterInstance, ESPMode::ThreadSafe> > Emitters;
 
 	FOnInitialized OnInitializedDelegate;
 	FOnComplete OnCompleteDelegate;
@@ -318,6 +321,8 @@ private:
 
 	/** If this system is paused. When paused it will not tick and never complete etc. */
 	uint32 bPaused : 1;
+	/** If this system has emitters that will run GPU Simulations */
+	bool bHasGPUEmitters : 1;
 
 	/* Execution state requested by external code/BPs calling Activate/Deactivate. */
 	ENiagaraExecutionState RequestedExecutionState;
@@ -328,4 +333,9 @@ private:
 	bool bDataInterfacesInitialized;
 
 	NiagaraEmitterInstanceBatcher* Batcher = nullptr;
+public:
+	// Transient data that is accumulated during tick.
+	uint32 TotalParamSize = 0;
+	uint32 ActiveGPUEmitterCount = 0;
+	int32 GPUDataInterfaceInstanceDataSize = 0;
 };
diff --git a/Engine/Plugins/FX/Niagara/Source/Niagara/Public/NiagaraSystemSimulation.h b/Engine/Plugins/FX/Niagara/Source/Niagara/Public/NiagaraSystemSimulation.h
--- a/Engine/Plugins/FX/Niagara/Source/Niagara/Public/NiagaraSystemSimulation.h
+++ b/Engine/Plugins/FX/Niagara/Source/Niagara/Public/NiagaraSystemSimulation.h
@@ -68,16 +68,16 @@ struct FNiagaraParameterStoreToDataSetBinding
 		DataSet.CheckForNaNs();
 #endif
 
-		FNiagaraDataBuffer& CurrBuffer = DataSet.CurrData();
+		FNiagaraDataBuffer* CurrBuffer = DataSet.GetCurrentData();
 
 		for (const FDataOffsets& DataOffsets : FloatOffsets)
 		{
-			float* DataSetPtr = CurrBuffer.GetInstancePtrFloat(DataOffsets.DataSetComponentOffset, DataSetInstanceIndex);
+			float* DataSetPtr = CurrBuffer->GetInstancePtrFloat(DataOffsets.DataSetComponentOffset, DataSetInstanceIndex);
 			ParameterStore.SetParameterByOffset(DataOffsets.ParameterOffset, *DataSetPtr);
 		}
 		for (const FDataOffsets& DataOffsets : Int32Offsets)
 		{
-			int32* DataSetPtr = CurrBuffer.GetInstancePtrInt32(DataOffsets.DataSetComponentOffset, DataSetInstanceIndex);
+			int32* DataSetPtr = CurrBuffer->GetInstancePtrInt32(DataOffsets.DataSetComponentOffset, DataSetInstanceIndex);
 			ParameterStore.SetParameterByOffset(DataOffsets.ParameterOffset, *DataSetPtr);
 		}
 
@@ -90,7 +90,7 @@ struct FNiagaraParameterStoreToDataSetBinding
 
 	FORCEINLINE_DEBUGGABLE void ParameterStoreToDataSet(FNiagaraParameterStore& ParameterStore, FNiagaraDataSet& DataSet, int32 DataSetInstanceIndex)
 	{
-		FNiagaraDataBuffer& CurrBuffer = DataSet.CurrData();
+		FNiagaraDataBuffer& CurrBuffer = DataSet.GetDestinationDataChecked();
 		const uint8* ParameterData = ParameterStore.GetParameterDataArray().GetData();
 
 #if NIAGARA_NAN_CHECKING
diff --git a/Engine/Plugins/FX/Niagara/Source/Niagara/Public/NiagaraTypes.h b/Engine/Plugins/FX/Niagara/Source/Niagara/Public/NiagaraTypes.h
--- a/Engine/Plugins/FX/Niagara/Source/Niagara/Public/NiagaraTypes.h
+++ b/Engine/Plugins/FX/Niagara/Source/Niagara/Public/NiagaraTypes.h
@@ -318,47 +318,44 @@ struct NIAGARA_API FNiagaraVariableMetaData
 	GENERATED_USTRUCT_BODY()
 public:
 	FNiagaraVariableMetaData()
-		:EditorSortPriority(0)
-		, bAdvancedDisplay(false)
+		: bAdvancedDisplay(false)
+		, EditorSortPriority(0)
 		, bInlineEditConditionToggle(false)
-		, CallSortPriority(0)
+		, bIsStaticSwitch(false)
 	{
 	}
 public:
-	UPROPERTY(EditAnywhere, Category = "Variable", DisplayName = "Property Metadata", meta = (ToolTip = "Property Metadata"))
-	TMap<FName, FString> PropertyMetaData;
-
 	UPROPERTY(EditAnywhere, Category = "Variable", meta = (MultiLine = true))
 	FText Description;
 
 	UPROPERTY(EditAnywhere, Category = "Variable")
 	FText CategoryName;
 
-	UPROPERTY(EditAnywhere, Category = "Variable", meta = (ToolTip = "Affects the sort order in the editor stacks. Use a smaller number to push it to the top. Defaults to zero."))
-	int32 EditorSortPriority;
-
 	/** Declares that this input is advanced and should only be visible if expanded inputs have been expanded. */
 	UPROPERTY(EditAnywhere, Category = "Variable")
 	bool bAdvancedDisplay;
 
-	/** Declares the associated input is used as an inline edit condition toggle, so it should should be hidden and it should be edited as a 
-	checkbox inline with the input which has designated as its edit condition. */
-	UPROPERTY(EditAnywhere, Category = "Variable")
+	UPROPERTY(EditAnywhere, Category = "Variable", meta = (ToolTip = "Affects the sort order in the editor stacks. Use a smaller number to push it to the top. Defaults to zero."))
+	int32 EditorSortPriority;
+
+	/** Declares the associated input is used as an inline edit condition toggle, so it should be hidden and edited as a 
+	checkbox inline with the input which was designated as its edit condition. */
+	UPROPERTY(EditAnywhere, Category = "Variable", meta = (EditCondition = "!bIsStaticSwitch"))
 	bool bInlineEditConditionToggle;
 
 	/** Declares the associated input should be conditionally editable based on the value of another input. */
-	UPROPERTY(EditAnywhere, Category = "Input Conditions")
+	UPROPERTY(EditAnywhere, Category = "Input Conditions", meta = (EditCondition = "!bIsStaticSwitch"))
 	FNiagaraInputConditionMetadata EditCondition;
 
 	/** Declares the associated input should be conditionally visible based on the value of another input. */
-	UPROPERTY(EditAnywhere, Category = "Input Conditions")
+	UPROPERTY(EditAnywhere, Category = "Input Conditions", meta = (EditCondition = "!bIsStaticSwitch"))
 	FNiagaraInputConditionMetadata VisibleCondition;
 
-	UPROPERTY()
-	int32 CallSortPriority;
+	UPROPERTY(EditAnywhere, Category = "Variable", DisplayName = "Property Metadata", meta = (ToolTip = "Property Metadata"))
+	TMap<FName, FString> PropertyMetaData;
 
 	UPROPERTY()
-	TArray<TWeakObjectPtr<UObject>> ReferencerNodes;
+	bool bIsStaticSwitch;
 };
 
 USTRUCT()
@@ -425,6 +422,19 @@ public:
 #endif
 	}
 
+	FName GetFName() const
+	{
+		if ( IsValid() == false )
+		{
+			return FName();
+		}
+		if ( Enum != nullptr )
+		{
+			return Enum->GetFName();
+		}
+		return GetStruct()->GetFName();
+	}
+
 	FString GetName()const
 	{
 		if (IsValid() == false)
diff --git a/Engine/Plugins/FX/Niagara/Source/Niagara/Public/NiagaraWorldManager.h b/Engine/Plugins/FX/Niagara/Source/Niagara/Public/NiagaraWorldManager.h
--- a/Engine/Plugins/FX/Niagara/Source/Niagara/Public/NiagaraWorldManager.h
+++ b/Engine/Plugins/FX/Niagara/Source/Niagara/Public/NiagaraWorldManager.h
@@ -9,6 +9,7 @@
 #include "NiagaraDataSet.h"
 #include "NiagaraScriptExecutionContext.h"
 #include "NiagaraSystemSimulation.h"
+#include "NiagaraSystemInstance.h"
 #include "GlobalDistanceFieldParameters.h"
 #include "NiagaraDataInterfaceSkeletalMesh.h"
 
@@ -32,26 +33,12 @@ public:
 		ViewUniformBuffer = Params.ViewUniformBuffer;
 		SceneNormalTexture = Params.NormalTexture;
 		SceneTexturesUniformParams = Params.SceneTexturesUniformParams;
-		GlobalDistanceFieldParams = Params.GlobalDistanceFieldParams;
-
-		PreSceneRenderValues = FPreSceneRenderValues();
-	}
-
-	void OnPreSceneRenderCalled(FPreSceneRenderValues& OutValues) const
-	{
-		OutValues.bUsesGlobalDistanceField |= PreSceneRenderValues.bUsesGlobalDistanceField;
-	}
-
-	void SetGlobalDistanceFieldUsage()
-	{
-		PreSceneRenderValues.bUsesGlobalDistanceField = true;
 	}
 
 	FTexture2DRHIParamRef GetSceneDepthTexture() { return SceneDepthTexture; }
 	FTexture2DRHIParamRef GetSceneNormalTexture() { return SceneNormalTexture; }
 	FUniformBufferRHIParamRef GetViewUniformBuffer() { return ViewUniformBuffer; }
 	TUniformBufferRef<FSceneTexturesUniformParameters> GetSceneTextureUniformParameters() { return SceneTexturesUniformParams; }
-	const FGlobalDistanceFieldParameterData* GetGlobalDistanceFieldParameters() { return GlobalDistanceFieldParams; }
 
 	virtual void InitDynamicRHI() override;
 
@@ -61,10 +48,8 @@ private:
 	FTexture2DRHIParamRef SceneDepthTexture;
 	FTexture2DRHIParamRef SceneNormalTexture;
 	FUniformBufferRHIParamRef ViewUniformBuffer;
-	FPreSceneRenderValues PreSceneRenderValues;	
 
 	TUniformBufferRef<FSceneTexturesUniformParameters> SceneTexturesUniformParams;
-	const FGlobalDistanceFieldParameterData* GlobalDistanceFieldParams;
 	FPostOpaqueRenderDelegate PostOpaqueDelegate;
 };
 
@@ -90,6 +75,7 @@ public:
 	void CleanupParameterCollections();
 	TSharedRef<FNiagaraSystemSimulation, ESPMode::ThreadSafe> GetSystemSimulation(UNiagaraSystem* System);
 	void DestroySystemSimulation(UNiagaraSystem* System);
+	void DestroySystemInstance(TUniquePtr<FNiagaraSystemInstance>& InPtr);
 
 	void Tick(float DeltaSeconds);
 
@@ -107,5 +93,17 @@ private:
 
 	/** Generated data used by data interfaces*/
 	FNDI_SkeletalMesh_GeneratedData SkeletalMeshGeneratedData;
+
+	// Deferred deletion queue for system instances
+	// We need to make sure that any enqueued GPU ticks have been processed before we remove the system instances
+	struct FDeferredDeletionQueue
+	{
+		FRenderCommandFence							Fence;
+		TArray<TUniquePtr<FNiagaraSystemInstance>>	Queue;
+	};
+
+	static constexpr int NumDeferredQueues = 2;
+	int DeferredDeletionQueueIndex = 0;
+	FDeferredDeletionQueue DeferredDeletionQueue[NumDeferredQueues];
 };
 
diff --git a/Engine/Plugins/FX/Niagara/Source/NiagaraCore/Private/NiagaraCompileHash.cpp b/Engine/Plugins/FX/Niagara/Source/NiagaraCore/Private/NiagaraCompileHash.cpp
new file mode 100644
--- /dev/null
+++ b/Engine/Plugins/FX/Niagara/Source/NiagaraCore/Private/NiagaraCompileHash.cpp
@@ -0,0 +1,56 @@
+// Copyright 1998-2019 Epic Games, Inc. All Rights Reserved.
+
+#include "NiagaraCompileHash.h"
+
+const uint32 FNiagaraCompileHash::HashSize = 20;
+
+bool FNiagaraCompileHash::operator==(const FNiagaraCompileHash& Other) const
+{
+	return DataHash == Other.DataHash;
+}
+
+bool FNiagaraCompileHash::operator!=(const FNiagaraCompileHash& Other) const
+{
+	return DataHash != Other.DataHash;
+}
+
+bool FNiagaraCompileHash::IsValid() const
+{
+	return DataHash.Num() == HashSize;
+}
+
+uint32 FNiagaraCompileHash::GetTypeHash() const
+{
+	// Use the first 4 bytes of the data hash as the hash id for this object.
+	return DataHash.Num() == HashSize 
+		? *((uint32*)(DataHash.GetData())) 
+		: 0;
+}
+
+const uint8* FNiagaraCompileHash::GetData() const
+{
+	return DataHash.GetData();
+}
+
+FString FNiagaraCompileHash::ToString() const
+{
+	if (DataHash.Num() == HashSize)
+	{
+		FString DataHashString;
+		DataHashString.Reserve(40);
+		for (int i = 0; i < HashSize; i++)
+		{
+			DataHashString.Append(FString::Printf(TEXT("%02x"), DataHash[i]));
+		}
+		return DataHashString;
+	}
+	else
+	{
+		return TEXT("Invalid");
+	}
+}
+
+FArchive& operator<<(FArchive& Ar, FNiagaraCompileHash& Id)
+{
+	return Ar << Id.DataHash;
+}
\ No newline at end of file
diff --git a/Engine/Plugins/FX/Niagara/Source/NiagaraCore/Private/NiagaraCustomVersion.cpp b/Engine/Plugins/FX/Niagara/Source/NiagaraCore/Private/NiagaraCustomVersion.cpp
--- a/Engine/Plugins/FX/Niagara/Source/NiagaraCore/Private/NiagaraCustomVersion.cpp
+++ b/Engine/Plugins/FX/Niagara/Source/NiagaraCore/Private/NiagaraCustomVersion.cpp
@@ -8,4 +8,4 @@ const FGuid FNiagaraCustomVersion::GUID(0xFCF57AFA, 0x50764283, 0xB9A9E658, 0xFF
 // Register the custom version with core
 FCustomVersionRegistration GRegisterNiagaraCustomVersion(FNiagaraCustomVersion::GUID, FNiagaraCustomVersion::LatestVersion, TEXT("NiagaraVer"));
 
-const FGuid FNiagaraCustomVersion::LatestScriptCompileVersion(0xC9A20384, 0x928D48B9, 0x9C80D8BA, 0x3E95F269);
\ No newline at end of file
+const FGuid FNiagaraCustomVersion::LatestScriptCompileVersion(0x4182078E, 0x4D784761, 0xB4205380, 0x84A48D4E);
diff --git a/Engine/Plugins/FX/Niagara/Source/NiagaraCore/Public/NiagaraCompileHash.h b/Engine/Plugins/FX/Niagara/Source/NiagaraCore/Public/NiagaraCompileHash.h
new file mode 100644
--- /dev/null
+++ b/Engine/Plugins/FX/Niagara/Source/NiagaraCore/Public/NiagaraCompileHash.h
@@ -0,0 +1,42 @@
+// Copyright 1998-2019 Epic Games, Inc. All Rights Reserved.
+
+#pragma once
+
+#include "UObject/ObjectMacros.h"
+#include "NiagaraCompileHash.generated.h"
+
+USTRUCT()
+struct NIAGARACORE_API FNiagaraCompileHash
+{
+	GENERATED_BODY()
+
+	FNiagaraCompileHash()
+	{
+	}
+
+	explicit FNiagaraCompileHash(const TArray<uint8>& InDataHash)
+	{
+		checkf(InDataHash.Num() == HashSize, TEXT("Invalid hash data."));
+		DataHash = InDataHash;
+	}
+
+	bool operator==(const FNiagaraCompileHash& Other) const;
+
+	bool operator!=(const FNiagaraCompileHash& Other) const;
+
+	bool IsValid() const;
+
+	uint32 GetTypeHash() const;
+
+	const uint8* GetData() const;
+
+	FString ToString() const;
+
+	NIAGARACORE_API friend FArchive& operator<<(FArchive& Ar, FNiagaraCompileHash& Id);
+
+	static const uint32 HashSize;
+
+private:
+	UPROPERTY()
+	TArray<uint8> DataHash;
+};
\ No newline at end of file
diff --git a/Engine/Plugins/FX/Niagara/Source/NiagaraCore/Public/NiagaraCustomVersion.h b/Engine/Plugins/FX/Niagara/Source/NiagaraCore/Public/NiagaraCustomVersion.h
--- a/Engine/Plugins/FX/Niagara/Source/NiagaraCore/Public/NiagaraCustomVersion.h
+++ b/Engine/Plugins/FX/Niagara/Source/NiagaraCore/Public/NiagaraCustomVersion.h
@@ -116,6 +116,10 @@ struct FNiagaraCustomVersion
 
 		MoveCommonInputMetadataToProperties, // Moved commonly used input metadata out of the strin/string property metadata map to actual properties on the metadata struct.
 
+		UseHashesToIdentifyCompileStateOfTopLevelScripts, // Move to using the traversed graph hash and the base script id for the FNiagaraVMExecutableDataId instead of the change id guid to prevent invalidating the DDC.
+
+		MetaDataAndParametersUpdate, // Reworked how the metadata is stored in NiagaraGraph from storing a Map of FNiagaraVariableMetaData to storing a map of UNiagaraScriptVariable* to be used with the Details panel.
+		
 		// -----<new versions can be added above this line>-------------------------------------------------
 		VersionPlusOne,
 		LatestVersion = VersionPlusOne - 1,
diff --git a/Engine/Plugins/FX/Niagara/Source/NiagaraCore/Public/NiagaraDataInterfaceBase.h b/Engine/Plugins/FX/Niagara/Source/NiagaraCore/Public/NiagaraDataInterfaceBase.h
--- a/Engine/Plugins/FX/Niagara/Source/NiagaraCore/Public/NiagaraDataInterfaceBase.h
+++ b/Engine/Plugins/FX/Niagara/Source/NiagaraCore/Public/NiagaraDataInterfaceBase.h
@@ -13,6 +13,17 @@ class FNiagaraShader;
 struct FNiagaraDataInterfaceParamRef;
 class FRHICommandList;
 
+struct FNiagaraDataInterfaceProxy;
+class NiagaraEmitterInstanceBatcher;
+
+struct FNiagaraDataInterfaceSetArgs
+{
+	FNiagaraShader* Shader;
+	FNiagaraDataInterfaceProxy* DataInterface;
+	FGuid SystemInstance;
+	const NiagaraEmitterInstanceBatcher* Batcher;
+};
+
 /**
 * An interface to the parameter bindings for the data interface used by a Niagara compute shader.
 */
@@ -21,8 +32,8 @@ struct FNiagaraDataInterfaceParametersCS
 	virtual ~FNiagaraDataInterfaceParametersCS() {}
 	virtual void Bind(const FNiagaraDataInterfaceParamRef& ParamRef, const class FShaderParameterMap& ParameterMap) {}
 	virtual void Serialize(FArchive& Ar) { }
-	virtual void Set(FRHICommandList& RHICmdList, FNiagaraShader* Shader, class UNiagaraDataInterface* DataInterface, void* PerInstanceData) const {}
-	virtual void Unset(FRHICommandList& RHICmdList, FNiagaraShader* Shader, class UNiagaraDataInterface* DataInterface, void* PerInstanceData) const {}
+	virtual void Set(FRHICommandList& RHICmdList, const FNiagaraDataInterfaceSetArgs& Context) const {}
+	virtual void Unset(FRHICommandList& RHICmdList, const FNiagaraDataInterfaceSetArgs& Context) const {}
 
 };
 
@@ -43,3 +54,4 @@ public:
 	}
 };
 
+
diff --git a/Engine/Plugins/FX/Niagara/Source/NiagaraEditor/Private/AssetTypeActions/AssetTypeActions_NiagaraEmitter.cpp b/Engine/Plugins/FX/Niagara/Source/NiagaraEditor/Private/AssetTypeActions/AssetTypeActions_NiagaraEmitter.cpp
--- a/Engine/Plugins/FX/Niagara/Source/NiagaraEditor/Private/AssetTypeActions/AssetTypeActions_NiagaraEmitter.cpp
+++ b/Engine/Plugins/FX/Niagara/Source/NiagaraEditor/Private/AssetTypeActions/AssetTypeActions_NiagaraEmitter.cpp
@@ -14,6 +14,7 @@
 #include "IContentBrowserSingleton.h"
 #include "ContentBrowserModule.h"
 #include "NiagaraEditorUtilities.h"
+#include "ViewModels/Stack/NiagaraStackGraphUtilities.h"
 #include "Modules/ModuleManager.h"
 
 #define LOCTEXT_NAMESPACE "AssetTypeActions"
@@ -108,9 +109,7 @@ void FAssetTypeActions_NiagaraEmitter::ExecuteNewNiagaraSystem(TArray<TWeakObjec
 			if (System != nullptr)
 			{
 				System->AddEmitterHandle(*Emitter, FName(*Emitter->GetName()));
-
-				FNiagaraSystemViewModel SystemViewModel = FNiagaraSystemViewModel(*System, SystemOptions);
-				SystemViewModel.GetSystemScriptViewModel()->RebuildEmitterNodes();
+				FNiagaraStackGraphUtilities::RebuildEmitterNodes(*System);
 
 				// Ensure the new System is compiled
 				if (!Emitter->AreAllScriptAndSourcesSynchronized())
diff --git a/Engine/Plugins/FX/Niagara/Source/NiagaraEditor/Private/Customizations/NiagaraMetaDataCustomNodeBuilder.cpp b/Engine/Plugins/FX/Niagara/Source/NiagaraEditor/Private/Customizations/NiagaraMetaDataCustomNodeBuilder.cpp
new file mode 100644
--- /dev/null
+++ b/Engine/Plugins/FX/Niagara/Source/NiagaraEditor/Private/Customizations/NiagaraMetaDataCustomNodeBuilder.cpp
@@ -0,0 +1,127 @@
+﻿// Copyright 1998-2019 Epic Games, Inc. All Rights Reserved.
+
+#include "NiagaraMetaDataCustomNodeBuilder.h"
+#include "NiagaraGraph.h"
+#include "ViewModels/Stack/NiagaraParameterHandle.h"
+#include "IDetailPropertyRow.h"
+#include "IDetailChildrenBuilder.h"
+
+FNiagaraMetaDataCustomNodeBuilder::FNiagaraMetaDataCustomNodeBuilder()
+	: ScriptGraph(nullptr)
+{
+}
+	
+void FNiagaraMetaDataCustomNodeBuilder::Initialize(UNiagaraGraph* InScriptGraph)
+{
+	ScriptGraph = InScriptGraph;
+
+	// We need to add handlers for both graph changed and graph needs compile since needs compile changes will
+	// not broadcast the generic graph changed but they still might modify the metadata collection.
+	OnGraphChangedHandle = ScriptGraph->AddOnGraphChangedHandler(FOnGraphChanged::FDelegate::CreateSP(this, &FNiagaraMetaDataCustomNodeBuilder::OnGraphChanged));
+	OnGraphNeedsRecompileHandle = ScriptGraph->AddOnGraphNeedsRecompileHandler(FOnGraphChanged::FDelegate::CreateSP(this, &FNiagaraMetaDataCustomNodeBuilder::OnGraphChanged));
+}
+
+FNiagaraMetaDataCustomNodeBuilder::~FNiagaraMetaDataCustomNodeBuilder()
+{
+	if (ScriptGraph.IsValid())
+	{
+		ScriptGraph->RemoveOnGraphChangedHandler(OnGraphChangedHandle);
+		ScriptGraph->RemoveOnGraphNeedsRecompileHandler(OnGraphNeedsRecompileHandle);
+	}
+}
+
+void FNiagaraMetaDataCustomNodeBuilder::SetOnRebuildChildren(FSimpleDelegate InOnRegenerateChildren)
+{
+	OnRebuildChildren = InOnRegenerateChildren;
+}
+
+FName FNiagaraMetaDataCustomNodeBuilder::GetName() const
+{
+	static const FName NiagaraMetadataCustomNodeBuilder("NiagaraMetadataCustomNodeBuilder");
+	return NiagaraMetadataCustomNodeBuilder;
+}
+
+void FNiagaraMetaDataCustomNodeBuilder::GenerateChildContent(IDetailChildrenBuilder& ChildrenBuilder)
+{
+	// We only store metadata information if it's actually set but we want to display metadata for all module parameters so
+	// we iterate over that collection here.
+	TArray<TTuple<FNiagaraVariable, FNiagaraParameterHandle, TSharedRef<FStructOnScope>>> MetaDataRowData;
+	const TMap<FNiagaraVariable, FNiagaraGraphParameterReferenceCollection>& ParameterReferenceMap = ScriptGraph->GetParameterReferenceMap();
+	for (const auto& ParameterToReferences : ParameterReferenceMap)
+	{
+		const FNiagaraVariable& ParameterVariable = ParameterToReferences.Key;
+		FNiagaraParameterHandle ParameterHandle(ParameterVariable.GetName());
+		if (ParameterHandle.IsModuleHandle())
+		{
+			TSharedRef<FStructOnScope> MetaDataContainerStruct = MakeShared<FStructOnScope>(FNiagaraVariableMetaDataContainer::StaticStruct());
+			TOptional<FNiagaraVariableMetaData> MetaData = ScriptGraph->GetMetaData(ParameterVariable);
+			if (MetaData.IsSet())
+			{
+				FNiagaraVariableMetaDataContainer* MetaDataContainer = (FNiagaraVariableMetaDataContainer*)MetaDataContainerStruct->GetStructMemory();
+				MetaDataContainer->MetaData = MetaData.GetValue();
+			}
+
+			MetaDataRowData.Add(TTuple<FNiagaraVariable, FNiagaraParameterHandle, TSharedRef<FStructOnScope>>(ParameterVariable, ParameterHandle, MetaDataContainerStruct));
+		}
+	}
+
+	// Sort the metadata using the same sorting that the function inputs use in the stack.
+	MetaDataRowData.Sort([](const TTuple<FNiagaraVariable, FNiagaraParameterHandle, TSharedRef<FStructOnScope>>& MetaDataRowDataA, const TTuple<FNiagaraVariable, FNiagaraParameterHandle, TSharedRef<FStructOnScope>>& MetaDataRowDataB)
+	{
+		const FNiagaraVariableMetaData& MetaDataA = ((FNiagaraVariableMetaDataContainer*)MetaDataRowDataA.Get<2>()->GetStructMemory())->MetaData;
+		const FNiagaraVariableMetaData& MetaDataB = ((FNiagaraVariableMetaDataContainer*)MetaDataRowDataB.Get<2>()->GetStructMemory())->MetaData;
+		if (MetaDataA.CategoryName.IsEmpty() && MetaDataB.CategoryName.IsEmpty() == false)
+		{
+			return true;
+		}
+		if (MetaDataA.CategoryName.IsEmpty() == false && MetaDataB.CategoryName.IsEmpty())
+		{
+			return false;
+		}
+		if (MetaDataA.EditorSortPriority != MetaDataB.EditorSortPriority)
+		{
+			return  MetaDataA.EditorSortPriority < MetaDataB.EditorSortPriority;
+		}
+		else
+		{
+			const FNiagaraVariable& ParameterVariableA = MetaDataRowDataA.Get<0>();
+			const FNiagaraVariable& ParameterVariableB = MetaDataRowDataB.Get<0>();
+			return ParameterVariableA.GetName() < ParameterVariableB.GetName();
+		}
+	});
+
+	// Add the rows for each metadata entry now that they have been sorted.
+	for (TTuple<FNiagaraVariable, FNiagaraParameterHandle, TSharedRef<FStructOnScope>>& MetaDataRowDataItem : MetaDataRowData)
+	{
+		const FNiagaraVariable& ParameterVariable = MetaDataRowDataItem.Get<0>();
+		const FNiagaraParameterHandle& ParameterHandle = MetaDataRowDataItem.Get<1>();
+		TSharedRef<FStructOnScope>& MetaDataContainerStruct = MetaDataRowDataItem.Get<2>();
+		
+		IDetailPropertyRow* MetaDataRow = ChildrenBuilder.AddExternalStructureProperty(MetaDataContainerStruct, GET_MEMBER_NAME_CHECKED(FNiagaraVariableMetaDataContainer, MetaData), ParameterVariable.GetName());
+		MetaDataRow->DisplayName(FText::FromName(ParameterHandle.GetName()));
+		if (MetaDataRow->GetPropertyHandle().IsValid())
+		{
+			MetaDataRow->GetPropertyHandle()->SetOnPropertyValueChanged(FSimpleDelegate::CreateSP(this, &FNiagaraMetaDataCustomNodeBuilder::MetaDataPropertyHandleChanged, ParameterVariable, MetaDataContainerStruct));
+			MetaDataRow->GetPropertyHandle()->SetOnChildPropertyValueChanged(FSimpleDelegate::CreateSP(this, &FNiagaraMetaDataCustomNodeBuilder::MetaDataPropertyHandleChanged, ParameterVariable, MetaDataContainerStruct));
+		}
+	}
+}
+
+void FNiagaraMetaDataCustomNodeBuilder::Rebuild()
+{
+	OnRebuildChildren.ExecuteIfBound();
+}
+
+void FNiagaraMetaDataCustomNodeBuilder::OnGraphChanged(const FEdGraphEditAction& Action)
+{
+	OnRebuildChildren.ExecuteIfBound();
+}
+
+void FNiagaraMetaDataCustomNodeBuilder::MetaDataPropertyHandleChanged(FNiagaraVariable ParameterVariable, TSharedRef<FStructOnScope> MetaDataContainerStruct)
+{
+	if (ScriptGraph.IsValid())
+	{
+		FNiagaraVariableMetaDataContainer* MetaDataContainer = (FNiagaraVariableMetaDataContainer*)MetaDataContainerStruct->GetStructMemory();
+		ScriptGraph->SetMetaData(ParameterVariable, MetaDataContainer->MetaData);
+	}
+}
\ No newline at end of file
diff --git a/Engine/Plugins/FX/Niagara/Source/NiagaraEditor/Private/Customizations/NiagaraMetaDataCustomNodeBuilder.h b/Engine/Plugins/FX/Niagara/Source/NiagaraEditor/Private/Customizations/NiagaraMetaDataCustomNodeBuilder.h
--- a/Engine/Plugins/FX/Niagara/Source/NiagaraEditor/Private/Customizations/NiagaraMetaDataCustomNodeBuilder.h
+++ b/Engine/Plugins/FX/Niagara/Source/NiagaraEditor/Private/Customizations/NiagaraMetaDataCustomNodeBuilder.h
@@ -1,101 +1,54 @@
-﻿// Copyright 1998-2019 Epic Games, Inc. All Rights Reserved.
+// Copyright 1998-2019 Epic Games, Inc. All Rights Reserved.
 
 #pragma once
-#include "IDetailPropertyRow.h"
-#include "NiagaraMetaDataCollectionViewModel.h"
-#include "NiagaraMetaDataViewModel.h"
-#include "IDetailChildrenBuilder.h"
-#include "UObject/StructOnScope.h"
-#include "IDetailGroup.h"
-#include "NiagaraEditorModule.h"
-#include "Modules/ModuleManager.h"
 
+#include "IDetailCustomNodeBuilder.h"
+#include "NiagaraTypes.h"
+#include "NiagaraMetaDataCustomNodeBuilder.generated.h"
 
-#define LOCTEXT_NAMESPACE "NiagaraMetaDataCustomNodeBuilder"
+class UNiagaraGraph;
+class FDetailWidgetRow;
+class IDetailChildrenBuilder;
 
-class FNiagaraMetaDataCustomNodeBuilder : public IDetailCustomNodeBuilder
+USTRUCT()
+struct FNiagaraVariableMetaDataContainer
+{
+	GENERATED_BODY();
+
+	UPROPERTY(EditAnywhere, Category=MetaData)
+	FNiagaraVariableMetaData MetaData;
+};
+
+class FNiagaraMetaDataCustomNodeBuilder : public IDetailCustomNodeBuilder, public TSharedFromThis<FNiagaraMetaDataCustomNodeBuilder>
 {
 public:
-	FNiagaraMetaDataCustomNodeBuilder(TSharedRef<FNiagaraMetaDataCollectionViewModel> InViewModel)
-		: ViewModel(InViewModel)
-	{
-		ViewModel->OnCollectionChanged().AddRaw(this, &FNiagaraMetaDataCustomNodeBuilder::OnCollectionViewModelChanged);
-	}
+	FNiagaraMetaDataCustomNodeBuilder();
+	
+	void Initialize(UNiagaraGraph* InScriptGraph);
 
-	~FNiagaraMetaDataCustomNodeBuilder()
-	{
-		ViewModel->OnCollectionChanged().RemoveAll(this);
-	}
+	~FNiagaraMetaDataCustomNodeBuilder();
 
-	virtual void SetOnRebuildChildren(FSimpleDelegate InOnRegenerateChildren) override
-	{
-		OnRebuildChildren = InOnRegenerateChildren;
-	}
+	virtual void SetOnRebuildChildren(FSimpleDelegate InOnRegenerateChildren) override;
 
 	virtual void GenerateHeaderRowContent(FDetailWidgetRow& NodeRow) {}
 	virtual void Tick(float DeltaTime) override {}
 	virtual bool RequiresTick() const override { return false; }
-	virtual bool InitiallyCollapsed() const { return false; }
+	virtual bool InitiallyCollapsed() const { return true; }
 	
-	virtual FName GetName() const  override
-	{
-		static const FName NiagaraMetadataCustomNodeBuilder("NiagaraMetadataCustomNodeBuilder");
-		return NiagaraMetadataCustomNodeBuilder;
-	}
+	virtual FName GetName() const  override;
 
-	virtual void GenerateChildContent(IDetailChildrenBuilder& ChildrenBuilder) override
-	{
-		const TArray<TSharedRef<FNiagaraMetaDataViewModel>>& VariableModels = ViewModel->GetVariableModels();
+	virtual void GenerateChildContent(IDetailChildrenBuilder& ChildrenBuilder) override;
 
-		FNiagaraEditorModule& NiagaraEditorModule = FModuleManager::GetModuleChecked<FNiagaraEditorModule>("NiagaraEditor");
-
-		for (const TSharedRef<FNiagaraMetaDataViewModel>& MetadataViewModel : VariableModels)
-		{
-			IDetailGroup& MetaDataGroup = ChildrenBuilder.AddGroup(MetadataViewModel->GetName(), FText::FromName(MetadataViewModel->GetName()));
-			MetaDataGroup.ToggleExpansion(true);
-			TSharedRef<FStructOnScope> StructData = MetadataViewModel->GetValueStruct();
-			
-			auto AddGroupedPropertyRow = [&ChildrenBuilder, &MetadataViewModel, &MetaDataGroup, &StructData](FName PropertyName, bool bShouldAutoExpand, bool bShouldWatchChildProperties)
-			{
-				IDetailPropertyRow* PropertyRow = ChildrenBuilder.AddExternalStructureProperty(StructData, PropertyName, PropertyName);
-				if (PropertyRow)
-				{
-					if (bShouldAutoExpand)
-					{
-						PropertyRow->ShouldAutoExpand(true);
-					}
-					PropertyRow->Visibility(EVisibility::Hidden);
-					MetaDataGroup.AddPropertyRow(PropertyRow->GetPropertyHandle()->AsShared());
-
-					TSharedPtr<IPropertyHandle> PropertyHandle = PropertyRow->GetPropertyHandle();
-					PropertyHandle->SetOnPropertyValueChanged(FSimpleDelegate::CreateSP(MetadataViewModel, &FNiagaraMetaDataViewModel::NotifyMetaDataChanged));
-					if (bShouldWatchChildProperties)
-					{
-						PropertyHandle->SetOnChildPropertyValueChanged(FSimpleDelegate::CreateSP(MetadataViewModel, &FNiagaraMetaDataViewModel::NotifyMetaDataChanged));
-					}
-				}
-			};
-
-			AddGroupedPropertyRow(GET_MEMBER_NAME_CHECKED(FNiagaraVariableMetaData, Description), false, false);
-			AddGroupedPropertyRow(GET_MEMBER_NAME_CHECKED(FNiagaraVariableMetaData, CategoryName), false, false);
-			AddGroupedPropertyRow(GET_MEMBER_NAME_CHECKED(FNiagaraVariableMetaData, bAdvancedDisplay), false, false);
-			AddGroupedPropertyRow(GET_MEMBER_NAME_CHECKED(FNiagaraVariableMetaData, bInlineEditConditionToggle), false, false);
-			AddGroupedPropertyRow(GET_MEMBER_NAME_CHECKED(FNiagaraVariableMetaData, EditorSortPriority), false, false);
-			AddGroupedPropertyRow(GET_MEMBER_NAME_CHECKED(FNiagaraVariableMetaData, EditCondition), false, true);
-			AddGroupedPropertyRow(GET_MEMBER_NAME_CHECKED(FNiagaraVariableMetaData, VisibleCondition), false, true);
-			AddGroupedPropertyRow(GET_MEMBER_NAME_CHECKED(FNiagaraVariableMetaData, PropertyMetaData), true, true);
-		}
-	}
+	void Rebuild();
 
 private:
-	void OnCollectionViewModelChanged()
-	{
-		OnRebuildChildren.ExecuteIfBound();
-	}
+	void OnGraphChanged(const struct FEdGraphEditAction& Action);
+
+	void MetaDataPropertyHandleChanged(FNiagaraVariable ParameterVariable, TSharedRef<FStructOnScope> MetaDataContainerStruct);
 
 private:
-	TSharedRef<FNiagaraMetaDataCollectionViewModel> ViewModel;
+	TWeakObjectPtr<UNiagaraGraph> ScriptGraph;
+	FDelegateHandle OnGraphChangedHandle;
+	FDelegateHandle OnGraphNeedsRecompileHandle;
 	FSimpleDelegate OnRebuildChildren;
 };
-
-#undef LOCTEXT_NAMESPACE // "NiagaraMetaDataCustomNodeBuilder"
diff --git a/Engine/Plugins/FX/Niagara/Source/NiagaraEditor/Private/Customizations/NiagaraScriptDetails.cpp b/Engine/Plugins/FX/Niagara/Source/NiagaraEditor/Private/Customizations/NiagaraScriptDetails.cpp
--- a/Engine/Plugins/FX/Niagara/Source/NiagaraEditor/Private/Customizations/NiagaraScriptDetails.cpp
+++ b/Engine/Plugins/FX/Niagara/Source/NiagaraEditor/Private/Customizations/NiagaraScriptDetails.cpp
@@ -10,6 +10,7 @@
 #include "NiagaraScriptInputCollectionViewModel.h"
 #include "NiagaraScriptOutputCollectionViewModel.h"
 #include "NiagaraScriptViewModel.h"
+#include "NiagaraScriptGraphViewModel.h"
 #include "NiagaraParameterViewModel.h"
 #include "NiagaraEditorStyle.h"
 #include "IDetailChildrenBuilder.h"
@@ -24,13 +25,9 @@
 #include "NiagaraScript.h"
 #include "NiagaraParameterCollectionCustomNodeBuilder.h"
 #include "NiagaraMetaDataCustomNodeBuilder.h"
-#include "NiagaraMetaDataCollectionViewModel.h"
-
 
 #define LOCTEXT_NAMESPACE "NiagaraScriptDetails"
 
-
-
 class SAddParameterButton : public SCompoundWidget
 {	
 	SLATE_BEGIN_ARGS(SAddParameterButton)
@@ -134,13 +131,11 @@ void FNiagaraScriptDetails::CustomizeDetails(IDetailLayoutBuilder& DetailBuilder
 	bool bAddParameters = false;
 	TSharedPtr<INiagaraParameterCollectionViewModel> InputCollectionViewModel;
 	TSharedPtr<INiagaraParameterCollectionViewModel> OutputCollectionViewModel;
-	TSharedPtr<FNiagaraMetaDataCollectionViewModel> MetadataCollectionViewModel;
 
 	if (StandaloneScript)
 	{
 		InputCollectionViewModel = ScriptViewModel->GetInputCollectionViewModel();
 		OutputCollectionViewModel = ScriptViewModel->GetOutputCollectionViewModel();
-		MetadataCollectionViewModel = ScriptViewModel->GetMetadataCollectionViewModel();
 		bAddParameters = true;
 	}
 
@@ -164,37 +159,43 @@ void FNiagaraScriptDetails::CustomizeDetails(IDetailLayoutBuilder& DetailBuilder
 		OutputParamCategory.AddCustomBuilder(MakeShared<FNiagaraParameterCollectionCustomNodeBuilder>(OutputCollectionViewModel.ToSharedRef()));
 	}
 	
-	if (MetadataCollectionViewModel.IsValid())
-	{
-		IDetailCategoryBuilder& MetadataDetailCategory = DetailBuilder.EditCategory(MetadataCategoryName, LOCTEXT("MetadataParamCategoryName", "Variable Metadata"));
-		MetadataDetailCategory.HeaderContent(
-			// Checkbox for showing in curve editor
-			SNew(SHorizontalBox)
-			+ SHorizontalBox::Slot()
-			.Padding(5, 0, 5, 0)
-			.HAlign(EHorizontalAlignment::HAlign_Right)
+	// Disable the metadata header in Script Details panel
+	// TODO: Delete when it isn't useful anymore (when the parameters rework is done)
+	/*
+	IDetailCategoryBuilder& MetadataDetailCategory = DetailBuilder.EditCategory(MetadataCategoryName, LOCTEXT("MetadataParamCategoryName", "Variable Metadata"));
+	MetadataDetailCategory.HeaderContent(
+		SNew(SHorizontalBox)
+		+ SHorizontalBox::Slot()
+		.Padding(5, 0, 5, 0)
+		.HAlign(EHorizontalAlignment::HAlign_Right)
+		[
+			SNew(SButton)
+			.ButtonStyle(FEditorStyle::Get(), "HoverHintOnly")
+			.HAlign(HAlign_Right)
+			.VAlign(VAlign_Center)
+			.ContentPadding(1)
+			.ToolTipText(LOCTEXT("RefreshMetadataToolTip", "Refresh the view according to the latest Editor Sort Priority values"))
+			.OnClicked(this, &FNiagaraScriptDetails::OnRefreshMetadata)
+			.Content()
 			[
-				SNew(SButton)
-				.ButtonStyle(FEditorStyle::Get(), "HoverHintOnly")
-				.HAlign(HAlign_Right)
-				.VAlign(VAlign_Center)
-				.ContentPadding(1)
-				.ToolTipText(LOCTEXT("RefreshMetadataToolTip", "Refresh the view according to the latest Editor Sort Priority values"))
-				.OnClicked(this, &FNiagaraScriptDetails::OnRefreshMetadata)
-				.Content()
-				[
-					SNew(SImage)
-					.Image(FEditorStyle::GetBrush("Icons.Refresh"))
-				]
+				SNew(SImage)
+				.Image(FEditorStyle::GetBrush("Icons.Refresh"))
 			]
-		);
-		MetadataDetailCategory.AddCustomBuilder(MakeShared<FNiagaraMetaDataCustomNodeBuilder>(MetadataCollectionViewModel.ToSharedRef()));
-	}	
+		]
+	);
+
+	MetaDataBuilder = MakeShared<FNiagaraMetaDataCustomNodeBuilder>();
+	MetaDataBuilder->Initialize(ScriptViewModel->GetGraphViewModel()->GetGraph());
+	MetadataDetailCategory.AddCustomBuilder(MetaDataBuilder.ToSharedRef());
+	*/
 }
 
 FReply FNiagaraScriptDetails::OnRefreshMetadata()
 {
-	ScriptViewModel->GetMetadataCollectionViewModel()->RequestRefresh();
+	if (MetaDataBuilder.IsValid())
+	{
+		MetaDataBuilder->Rebuild();
+	}
     return FReply::Handled();
 }
 
diff --git a/Engine/Plugins/FX/Niagara/Source/NiagaraEditor/Private/Customizations/NiagaraScriptDetails.h b/Engine/Plugins/FX/Niagara/Source/NiagaraEditor/Private/Customizations/NiagaraScriptDetails.h
--- a/Engine/Plugins/FX/Niagara/Source/NiagaraEditor/Private/Customizations/NiagaraScriptDetails.h
+++ b/Engine/Plugins/FX/Niagara/Source/NiagaraEditor/Private/Customizations/NiagaraScriptDetails.h
@@ -8,6 +8,7 @@
 
 class FNiagaraScriptViewModel;
 class INiagaraParameterCollectionViewModel;
+class FNiagaraMetaDataCustomNodeBuilder;
 
 class FNiagaraScriptDetails : public IDetailCustomization
 {
@@ -22,6 +23,6 @@ public:
 
 private:
 	TSharedPtr<FNiagaraScriptViewModel> ScriptViewModel;
-
+	TSharedPtr<FNiagaraMetaDataCustomNodeBuilder> MetaDataBuilder;
 };
 
diff --git a/Engine/Plugins/FX/Niagara/Source/NiagaraEditor/Private/EdGraphSchema_Niagara.cpp b/Engine/Plugins/FX/Niagara/Source/NiagaraEditor/Private/EdGraphSchema_Niagara.cpp
--- a/Engine/Plugins/FX/Niagara/Source/NiagaraEditor/Private/EdGraphSchema_Niagara.cpp
+++ b/Engine/Plugins/FX/Niagara/Source/NiagaraEditor/Private/EdGraphSchema_Niagara.cpp
@@ -43,6 +43,7 @@
 #include "Modules/ModuleManager.h"
 #include "AssetRegistryModule.h"
 #include "NiagaraNodeSimTargetSelector.h"
+#include "NiagaraNodeStaticSwitch.h"
 
 #define LOCTEXT_NAMESPACE "NiagaraSchema"
 
@@ -834,8 +835,8 @@ TArray<TSharedPtr<FNiagaraSchemaAction_NewNode> > UEdGraphSchema_Niagara::GetGra
 	//TODO: Add quick commands for certain UNiagaraStructs and UNiagaraScripts to be added as functions
 
 	// Add reroute node
+	const FText UtilMenuCat = LOCTEXT("NiagaraUsageSelectorMenuCat", "Util");
 	{
-		const FText UtilMenuCat = LOCTEXT("NiagaraRerouteMenuCat", "Util");
 		const FText RerouteMenuDesc = LOCTEXT("NiagaraRerouteMenuDesc", "Reroute ");
 		TSharedPtr<FNiagaraSchemaAction_NewNode> Action = AddNewNodeAction(NewActions, UtilMenuCat, RerouteMenuDesc, TEXT("Reroute"), FText::GetEmpty());
 		UNiagaraNodeReroute* RerouteNode = NewObject<UNiagaraNodeReroute>(OwnerOfTemporaries);
@@ -843,17 +844,23 @@ TArray<TSharedPtr<FNiagaraSchemaAction_NewNode> > UEdGraphSchema_Niagara::GetGra
 	}
 
 	// Add usage selector node
-	{
-		const FText UtilMenuCat = LOCTEXT("NiagaraUsageSelectorMenuCat", "Util");
+	{		
 		const FText UsageSelectorMenuDesc = LOCTEXT("NiagaraUsageSelectorMenuDesc", "Select By Use");
 		TSharedPtr<FNiagaraSchemaAction_NewNode> Action = AddNewNodeAction(NewActions, UtilMenuCat, UsageSelectorMenuDesc, TEXT("Select By Use"), FText::GetEmpty());
 		UNiagaraNodeUsageSelector* Node = NewObject<UNiagaraNodeUsageSelector>(OwnerOfTemporaries);
 		Action->NodeTemplate = Node;
 	}
 
+	// Add static switch node
+	{
+		const FText UsageSelectorMenuDesc = LOCTEXT("NiagaraUsageSelectorMenuDesc", "Static Switch");
+		TSharedPtr<FNiagaraSchemaAction_NewNode> Action = AddNewNodeAction(NewActions, UtilMenuCat, UsageSelectorMenuDesc, TEXT("Static Switch"), FText::GetEmpty());
+		UNiagaraNodeStaticSwitch* Node = NewObject<UNiagaraNodeStaticSwitch>(OwnerOfTemporaries);
+		Action->NodeTemplate = Node;
+	}
+
 	// Add simulation target selector node
 	{
-		const FText UtilMenuCat = LOCTEXT("NiagaraSimTargetSelectorMenuCat", "Util");
 		const FText SimTargetSelectorMenuDesc = LOCTEXT("NiagaraSimTargetSelectorMenuDesc", "Select By Simulation Target");
 		TSharedPtr<FNiagaraSchemaAction_NewNode> Action = AddNewNodeAction(NewActions, UtilMenuCat, SimTargetSelectorMenuDesc, TEXT("Select By Simulation Target"), FText::GetEmpty());
 		UNiagaraNodeSimTargetSelector* Node = NewObject<UNiagaraNodeSimTargetSelector>(OwnerOfTemporaries);
diff --git a/Engine/Plugins/FX/Niagara/Source/NiagaraEditor/Private/NiagaraEditorCommands.cpp b/Engine/Plugins/FX/Niagara/Source/NiagaraEditor/Private/NiagaraEditorCommands.cpp
--- a/Engine/Plugins/FX/Niagara/Source/NiagaraEditor/Private/NiagaraEditorCommands.cpp
+++ b/Engine/Plugins/FX/Niagara/Source/NiagaraEditor/Private/NiagaraEditorCommands.cpp
@@ -13,6 +13,7 @@ void FNiagaraEditorCommands::RegisterCommands()
 	UI_COMMAND(RefreshNodes, "Refresh", "Refreshes the current graph nodes, and updates pins due to external changes.", EUserInterfaceActionType::Button, FInputChord());
 	UI_COMMAND(ResetSimulation, "Reset", "Resets the current simulation", EUserInterfaceActionType::Button, FInputChord(EKeys::SpaceBar));
 	UI_COMMAND(TogglePreviewGrid, "Grid", "Toggles the preview pane's grid.", EUserInterfaceActionType::ToggleButton, FInputChord() );
+	UI_COMMAND(ToggleInstructionCounts, "InstructionCounts", "Display Instruction Counts", EUserInterfaceActionType::ToggleButton, FInputChord());
 	UI_COMMAND(TogglePreviewBackground, "Background", "Toggles the preview pane's background.", EUserInterfaceActionType::ToggleButton, FInputChord());
 	UI_COMMAND(ToggleUnlockToChanges, "Lock/Unlock To Changes", "Toggles whether or not changes in the source asset get pulled into this asset automatically.", EUserInterfaceActionType::Button, FInputChord());
 	UI_COMMAND(ToggleBounds, "Bounds", "Display Bounds", EUserInterfaceActionType::ToggleButton, FInputChord());
diff --git a/Engine/Plugins/FX/Niagara/Source/NiagaraEditor/Private/NiagaraEditorModule.cpp b/Engine/Plugins/FX/Niagara/Source/NiagaraEditor/Private/NiagaraEditorModule.cpp
--- a/Engine/Plugins/FX/Niagara/Source/NiagaraEditor/Private/NiagaraEditorModule.cpp
+++ b/Engine/Plugins/FX/Niagara/Source/NiagaraEditor/Private/NiagaraEditorModule.cpp
@@ -67,6 +67,7 @@
 #include "NiagaraEmitter.h"
 #include "NiagaraScriptSource.h"
 #include "NiagaraTypes.h"
+#include "NiagaraSystemFactoryNew.h"
 
 #include "MovieScene/Parameters/MovieSceneNiagaraBoolParameterTrack.h"
 #include "MovieScene/Parameters/MovieSceneNiagaraFloatParameterTrack.h"
@@ -90,6 +91,9 @@
 #include "HAL/IConsoleManager.h"
 #include "NiagaraHlslTranslator.h"
 #include "NiagaraThumbnailRenderer.h"
+#include "Misc/FeedbackContext.h"
+#include "Customizations/NiagaraStaticSwitchNodeDetails.h"
+#include "Customizations/NiagaraFunctionCallNodeDetails.h"
 
 IMPLEMENT_MODULE( FNiagaraEditorModule, NiagaraEditor );
 
@@ -273,6 +277,83 @@ void DumpRapidIterationParamersForAsset(const TArray<FString>& Arguments)
 	}
 }
 
+void PreventSystemRecompile(FAssetData SystemAsset, TSet<UNiagaraEmitter*>& InOutCompiledEmitters)
+{
+	UNiagaraSystem* System = Cast<UNiagaraSystem>(SystemAsset.GetAsset());
+	if (System != nullptr)
+	{
+		for (const FNiagaraEmitterHandle& EmitterHandle : System->GetEmitterHandles())
+		{
+			UNiagaraEmitter* SourceEmitter = const_cast<UNiagaraEmitter*>(EmitterHandle.GetSource());
+			if (SourceEmitter != nullptr)
+			{
+				if (InOutCompiledEmitters.Contains(SourceEmitter) == false)
+				{
+					SourceEmitter->MarkPackageDirty();
+					UNiagaraSystem* TransientSystem = NewObject<UNiagaraSystem>(GetTransientPackage(), NAME_None, RF_Transient);
+					UNiagaraSystemFactoryNew::InitializeSystem(TransientSystem, true);
+					TransientSystem->AddEmitterHandle(*SourceEmitter, TEXT("Emitter")); // TODO Frank.Fella, resolve this properly...
+					FNiagaraStackGraphUtilities::RebuildEmitterNodes(*TransientSystem);
+					TransientSystem->RequestCompile(false);
+					TransientSystem->WaitForCompilationComplete();
+					InOutCompiledEmitters.Add(SourceEmitter);
+				}
+				System->UpdateFromEmitterChanges(*SourceEmitter, false);
+			}
+		}
+
+		System->MarkPackageDirty();
+		System->RequestCompile(false);
+		System->WaitForCompilationComplete();
+	}
+}
+
+void PreventSystemRecompile(const TArray<FString>& Arguments)
+{
+	if (Arguments.Num() > 0)
+	{
+		const FAssetRegistryModule& AssetRegistryModule = FModuleManager::LoadModuleChecked<FAssetRegistryModule>(TEXT("AssetRegistry"));
+		FAssetData SystemAsset = AssetRegistryModule.Get().GetAssetByObjectPath(*Arguments[0]);
+		if (SystemAsset.IsValid() == false)
+		{
+			TArray<FAssetData> AssetsInPackage;
+			AssetRegistryModule.Get().GetAssetsByPackageName(*Arguments[0], AssetsInPackage);
+			if (AssetsInPackage.Num() == 1)
+			{
+				SystemAsset = AssetsInPackage[0];
+			}
+		}
+		TSet<UNiagaraEmitter*> CompiledEmitters;
+		PreventSystemRecompile(SystemAsset, CompiledEmitters);
+	}
+}
+
+void PreventAllSystemRecompiles()
+{
+	const FText SlowTaskText = NSLOCTEXT("NiagaraEditor", "PreventAllSystemRecompiles", "Refreshing all systems to prevent recompiles.");
+	GWarn->BeginSlowTask(SlowTaskText, true, true);
+
+	FAssetRegistryModule& AssetRegistryModule = FModuleManager::LoadModuleChecked<FAssetRegistryModule>("AssetRegistry");
+
+	TArray<FAssetData> SystemAssets;
+	AssetRegistryModule.Get().GetAssetsByClass(UNiagaraSystem::StaticClass()->GetFName(), SystemAssets);
+
+	TSet<UNiagaraEmitter*> CompiledEmitters;
+	int32 ItemIndex = 0;
+	for (FAssetData& SystemAsset : SystemAssets)
+	{
+		if (GWarn->ReceivedUserCancel())
+		{
+			return;
+		}
+		GWarn->UpdateProgress(ItemIndex++, SystemAssets.Num());
+
+		PreventSystemRecompile(SystemAsset, CompiledEmitters);
+	}
+
+	GWarn->EndSlowTask();
+}
+
 class FNiagaraSystemBoolParameterTrackEditor : public FNiagaraSystemParameterTrackEditor<UMovieSceneNiagaraBoolParameterTrack, UMovieSceneBoolSection>
 {
 	virtual TSharedRef<ISequencerSection> MakeSectionInterface(UMovieSceneSection& SectionObject, UMovieSceneTrack& Track, FGuid ObjectBinding) override
@@ -313,6 +394,10 @@ void FNiagaraEditorModule::StartupModule()
 	// register details customization
 	FPropertyEditorModule& PropertyModule = FModuleManager::LoadModuleChecked<FPropertyEditorModule>("PropertyEditor");
 	PropertyModule.RegisterCustomClassLayout("NiagaraComponent", FOnGetDetailCustomizationInstance::CreateStatic(&FNiagaraComponentDetails::MakeInstance));
+
+	PropertyModule.RegisterCustomClassLayout("NiagaraNodeStaticSwitch", FOnGetDetailCustomizationInstance::CreateStatic(&FNiagaraStaticSwitchNodeDetails::MakeInstance));
+
+	PropertyModule.RegisterCustomClassLayout("NiagaraNodeFunctionCall", FOnGetDetailCustomizationInstance::CreateStatic(&FNiagaraFunctionCallNodeDetails::MakeInstance));
 	
 	PropertyModule.RegisterCustomPropertyTypeLayout("NiagaraFloat",
 		FOnGetPropertyTypeCustomizationInstance::CreateStatic(&FNiagaraNumericCustomization::MakeInstance)
@@ -487,6 +572,16 @@ void FNiagaraEditorModule::StartupModule()
 		TEXT("Dumps the values of the rapid iteration parameters for the specified asset by path."),
 		FConsoleCommandWithArgsDelegate::CreateStatic(&DumpRapidIterationParamersForAsset));
 
+	PreventSystemRecompileCommand = IConsoleManager::Get().RegisterConsoleCommand(
+		TEXT("fx.PreventSystemRecompile"),
+		TEXT("Forces the system to refresh all it's dependencies so it won't recompile on load.  This may mark multiple assets dirty for re-saving."),
+		FConsoleCommandWithArgsDelegate::CreateStatic(&PreventSystemRecompile));
+
+	PreventSystemRecompileCommand = IConsoleManager::Get().RegisterConsoleCommand(
+		TEXT("fx.PreventAllSystemRecompiles"),
+		TEXT("Loads all of the systems in the project and forces each system to refresh all it's dependencies so it won't recompile on load.  This may mark multiple assets dirty for re-saving."),
+		FConsoleCommandDelegate::CreateStatic(&PreventAllSystemRecompiles));
+
 	if (GIsEditor)
 	{
 		UThumbnailManager::Get().RegisterCustomRenderer(UNiagaraEmitter::StaticClass(), UNiagaraEmitterThumbnailRenderer::StaticClass());
diff --git a/Engine/Plugins/FX/Niagara/Source/NiagaraEditor/Private/NiagaraEditorUtilities.cpp b/Engine/Plugins/FX/Niagara/Source/NiagaraEditor/Private/NiagaraEditorUtilities.cpp
--- a/Engine/Plugins/FX/Niagara/Source/NiagaraEditor/Private/NiagaraEditorUtilities.cpp
+++ b/Engine/Plugins/FX/Niagara/Source/NiagaraEditor/Private/NiagaraEditorUtilities.cpp
@@ -32,6 +32,8 @@
 #include "Misc/FileHelper.h"
 #include "EdGraph/EdGraphPin.h"
 #include "NiagaraNodeWriteDataSet.h"
+#include "NiagaraNodeStaticSwitch.h"
+#include "NiagaraNodeFunctionCall.h"
 
 #define LOCTEXT_NAMESPACE "FNiagaraEditorUtilities"
 
@@ -620,6 +622,87 @@ void FNiagaraEditorUtilities::FixUpNumericPins(const UEdGraphSchema_Niagara* Sch
 	TraverseGraphFromOutputDepthFirst(Schema, Node, FixUpVisitor);
 }
 
+void FNiagaraEditorUtilities::SetStaticSwitchConstants(UNiagaraGraph* Graph, const TArray<UEdGraphPin*>& CallInputs)
+{
+	for (UEdGraphNode* Node : Graph->Nodes)
+	{
+		// if there is a static switch node its value must be set by the caller
+		UNiagaraNodeStaticSwitch* SwitchNode = Cast<UNiagaraNodeStaticSwitch>(Node);
+		if (SwitchNode)
+		{
+			SwitchNode->IsValueSet = false;
+			for (UEdGraphPin* InputPin : CallInputs)
+			{
+				if (InputPin->GetFName().IsEqual(SwitchNode->InputParameterName))
+				{
+					int32 SwitchValue = 0;
+					SwitchNode->IsValueSet = ResolveConstantValue(InputPin, SwitchValue);
+					SwitchNode->SwitchValue = SwitchValue;
+					break;
+				}
+			}
+			
+		}
+
+		// if there is a function node, it might have delegated some of the static switch values inside its script graph
+		// to be set by the next higher caller instead of directly by the user
+		UNiagaraNodeFunctionCall* FunctionNode = Cast<UNiagaraNodeFunctionCall>(Node);
+		if (FunctionNode && FunctionNode->PropagatedStaticSwitchParameters.Num() > 0)
+		{
+			for (const FNiagaraVariable& SwitchValue : FunctionNode->PropagatedStaticSwitchParameters)
+			{
+				UEdGraphPin* ValuePin = FunctionNode->FindPin(SwitchValue.GetName(), EGPD_Input);
+				if (!ValuePin)
+				{
+					continue;
+				}
+				ValuePin->DefaultValue = FString();
+				for (UEdGraphPin* InputPin : CallInputs)
+				{
+					if (InputPin->GetFName().IsEqual(ValuePin->GetFName()) && InputPin->PinType == ValuePin->PinType)
+					{
+						ValuePin->DefaultValue = InputPin->DefaultValue;
+						break;
+					}
+				}				
+			}
+
+		}
+	}
+}
+
+bool FNiagaraEditorUtilities::ResolveConstantValue(UEdGraphPin* Pin, int32& Value)
+{
+	if (Pin->LinkedTo.Num() > 0)
+	{
+		return false;
+	}
+	
+	const FEdGraphPinType& PinType = Pin->PinType;
+	if (PinType.PinCategory == UEdGraphSchema_Niagara::PinCategoryType && PinType.PinSubCategoryObject.IsValid())
+	{
+		FString PinTypeName = PinType.PinSubCategoryObject->GetName();
+		if (PinTypeName.Equals(FString(TEXT("NiagaraBool"))))
+		{
+			Value = Pin->DefaultValue.Equals(FString(TEXT("true"))) ? 1 : 0;
+			return true;
+		}
+		else if (PinTypeName.Equals(FString(TEXT("NiagaraInt32"))))
+		{
+			Value = FCString::Atoi(*Pin->DefaultValue);
+			return true;
+		}
+	}
+	else if (PinType.PinCategory == UEdGraphSchema_Niagara::PinCategoryEnum && PinType.PinSubCategoryObject.IsValid())
+	{
+		UEnum* Enum = Cast<UEnum>(PinType.PinSubCategoryObject);
+		FString FullName = Enum->GenerateFullEnumName(*Pin->DefaultValue);
+		Value = Enum->GetIndexByName(FName(*FullName));
+		return Value != INDEX_NONE;
+	}
+	return false;
+}
+
 /* Go through the graph and attempt to auto-detect the type of any numeric pins by working back from the leaves of the graph. Only change the types of pins, not FNiagaraVariables.*/
 void PreprocessGraph(const UEdGraphSchema_Niagara* Schema, UNiagaraGraph* Graph, UNiagaraNodeOutput* OutputNode)
 {
@@ -769,7 +852,36 @@ void FNiagaraEditorUtilities::PreprocessFunctionGraph(const UEdGraphSchema_Niaga
 	}
 
 	FNiagaraEditorUtilities::FixUpNumericPins(Schema, OutputNode);
+	FNiagaraEditorUtilities::SetStaticSwitchConstants(Graph, CallInputs);
+}
 
+const TMap<FNiagaraVariable, FNiagaraGraphParameterReferenceCollection>& FNiagaraEditorUtilities::GetCompiledGraphParameterMapReferences(UNiagaraGraph* Graph)
+{
+	const UEdGraphSchema_Niagara* Schema = GetDefault<UEdGraphSchema_Niagara>();
+
+	// We clone the graph to prevent any destructive changes to the source
+	UNiagaraGraph* NodeGraphDeepCopy = CastChecked<UNiagaraGraph>(FEdGraphUtilities::CloneGraph(Graph, GetTransientPackage(), 0, true));
+	TArray<UNiagaraNodeFunctionCall*> Nodes;
+	NodeGraphDeepCopy->GetNodesOfClass(Nodes);
+	for (int32 i = 0; i < Nodes.Num(); i++)
+	{
+		UNiagaraNodeFunctionCall* Node = Nodes[i];
+		ENiagaraScriptUsage ScriptUsage = Node->GetCalledUsage();
+		UNiagaraGraph* FunctionGraph = Node->GetCalledGraph();
+		if (FunctionGraph)
+		{
+			// We append all function nodes we find to the end of the queue, basically traversing the graph breadth-first
+			TArray<UNiagaraNodeFunctionCall*> SubFunctionNodes;
+			FunctionGraph->GetNodesOfClass(SubFunctionNodes);
+			Nodes.Append(SubFunctionNodes);
+
+			// We set the static switch values based on the function call inputs to prepare the parameter map history generation
+			TArray<UEdGraphPin*> CallInputs;
+			Node->GetInputPins(CallInputs);
+			SetStaticSwitchConstants(FunctionGraph, CallInputs);
+		}
+	}
+	return NodeGraphDeepCopy->GetParameterReferenceMap();
 }
 
 void FNiagaraEditorUtilities::GetFilteredScriptAssets(FGetFilteredScriptAssetsOptions InFilter, TArray<FAssetData>& OutFilteredScriptAssets)
