Copying //UE4/Dev-Niagara to Dev-Main (//UE4/Dev-Main)

diff --git a/Engine/Plugins/FX/Niagara/Source/NiagaraEditor/Private/NiagaraGraph.cpp b/Engine/Plugins/FX/Niagara/Source/NiagaraEditor/Private/NiagaraGraph.cpp
--- a/Engine/Plugins/FX/Niagara/Source/NiagaraEditor/Private/NiagaraGraph.cpp
+++ b/Engine/Plugins/FX/Niagara/Source/NiagaraEditor/Private/NiagaraGraph.cpp
@@ -29,6 +29,10 @@
 #include "NiagaraNode.h"
 #include "EdGraphSchema_Niagara.h"
 #include "ViewModels/Stack/NiagaraParameterHandle.h"
+#include "NiagaraNodeStaticSwitch.h"
+#include "NiagaraScriptVariable.h"
+#include "NiagaraHlslTranslator.h"
+#include "NiagaraNodeFunctionCall.h"
 
 DECLARE_CYCLE_STAT(TEXT("NiagaraEditor - Graph - FindInputNodes"), STAT_NiagaraEditor_Graph_FindInputNodes, STATGROUP_NiagaraEditor);
 DECLARE_CYCLE_STAT(TEXT("NiagaraEditor - Graph - FindInputNodes_NotFilterUsage"), STAT_NiagaraEditor_Graph_FindInputNodes_NotFilterUsage, STATGROUP_NiagaraEditor);
@@ -54,14 +58,29 @@ bool FNiagaraGraphParameterReferenceCollection::WasCreated() const
 
 FNiagaraGraphScriptUsageInfo::FNiagaraGraphScriptUsageInfo() : UsageType(ENiagaraScriptUsage::Function)
 {
-	DataHash.AddZeroed(sizeof(FSHAHash));
+}
+
+void FNiagaraGraphScriptUsageInfo::PostLoad(UObject* Owner)
+{
+	const int32 NiagaraVer = Owner->GetLinkerCustomVersion(FNiagaraCustomVersion::GUID);
+	if (NiagaraVer < FNiagaraCustomVersion::UseHashesToIdentifyCompileStateOfTopLevelScripts)
+	{
+		// When loading old data use the last generated compile id as the base id to prevent recompiles on load for existing scripts.
+		BaseId = GeneratedCompileId;
+
+		if (CompileHash.IsValid() == false && DataHash_DEPRECATED.Num() == FNiagaraCompileHash::HashSize)
+		{
+			CompileHash = FNiagaraCompileHash(DataHash_DEPRECATED);
+		}
+	}
 }
 
 UNiagaraGraph::UNiagaraGraph(const FObjectInitializer& ObjectInitializer)
 	: Super(ObjectInitializer)
 	, bNeedNumericCacheRebuilt(true)
-	, bFindParametersAllowed(true)
 	, bIsRenamingParameter(false)
+	, bParameterReferenceRefreshPending(true)
+	, bUnreferencedMetaDataPurgePending(true)
 {
 	Schema = UEdGraphSchema_Niagara::StaticClass();
 	ChangeId = FGuid::NewGuid();
@@ -79,7 +98,7 @@ void UNiagaraGraph::RemoveOnGraphNeedsRecompileHandler(FDelegateHandle Handle)
 
 void UNiagaraGraph::NotifyGraphChanged(const FEdGraphEditAction& InAction)
 {
-	FindParameters();
+	InvalidateCachedParameterData();
 	if ((InAction.Action & GRAPHACTION_AddNode) != 0 || (InAction.Action & GRAPHACTION_RemoveNode) != 0 ||
 		(InAction.Action & GRAPHACTION_GenericNeedsRecompile) != 0)
 	{
@@ -95,8 +114,8 @@ void UNiagaraGraph::NotifyGraphChanged(const FEdGraphEditAction& InAction)
 
 void UNiagaraGraph::NotifyGraphChanged()
 {
-	FindParameters();
 	Super::NotifyGraphChanged();
+	InvalidateCachedParameterData();
 	InvalidateNumericCache();
 }
 
@@ -104,6 +123,11 @@ void UNiagaraGraph::PostLoad()
 {
 	Super::PostLoad();
 
+	for (FNiagaraGraphScriptUsageInfo& CachedUsageInfoItem : CachedUsageInfo)
+	{
+		CachedUsageInfoItem.PostLoad(this);
+	}
+
 	// In the past, we didn't bother setting the CallSortPriority and just used lexicographic ordering.
 	// In the event that we have multiple non-matching nodes with a zero call sort priority, this will
 	// give every node a unique order value.
@@ -175,21 +199,40 @@ void UNiagaraGraph::PostLoad()
 		}
 	}
 
-	RebuildCachedData();
+	RebuildCachedCompileIds();
 
 	if (GIsEditor)
 	{
 		SetFlags(RF_Transactional);
 	}
 
-	Parameters.Empty();
-	FindParameters();
-
 	// Migrate input condition metadata
 	const int32 NiagaraVer = GetLinkerCustomVersion(FNiagaraCustomVersion::GUID);
+
+	if (NiagaraVer < FNiagaraCustomVersion::MetaDataAndParametersUpdate)
+	{
+		// If the version of the asset is older than FNiagaraCustomVersion::MetaDataAndParametersUpdate 
+		// we need to migrate the old metadata by looping through VariableToMetaData_DEPRECATED
+		// and create new entries in VariableToScriptVariable
+		for (auto It = VariableToMetaData_DEPRECATED.CreateConstIterator(); It; ++It)
+		{
+			SetMetaData(It.Key(), It.Value());
+
+			FString PathName = GetPathName();
+			int ColonPos;
+			if (PathName.FindChar(TCHAR('.'), ColonPos))
+			{
+				// GetPathName() returns something similar to "/Path/To/ScriptName.ScriptName:NiagaraScriptSource_N.NiagaraGraph_N"
+				// so this will extract "/Path/To/ScriptName"
+				PathName = PathName.Left(ColonPos);
+			}
+			UE_LOG(LogNiagaraEditor, Log, TEXT("Migrated old metadata entry for variable \"%s\" in \"%s\""), *It.Key().GetName().ToString(), *PathName);
+		}
+		VariableToMetaData_DEPRECATED.Empty();
+	}
 	if (NiagaraVer < FNiagaraCustomVersion::MoveCommonInputMetadataToProperties)
 	{
-		auto MigrateInputCondition = [](TMap<FName,FString>& PropertyMetaData, const FName& InputConditionKey, FNiagaraInputConditionMetadata& InOutInputCondition)
+		auto MigrateInputCondition = [](TMap<FName, FString>& PropertyMetaData, const FName& InputConditionKey, FNiagaraInputConditionMetadata& InOutInputCondition)
 		{
 			FString* InputCondition = PropertyMetaData.Find(InputConditionKey);
 			if (InputCondition != nullptr)
@@ -210,29 +253,31 @@ void UNiagaraGraph::PostLoad()
 			}
 		};
 
-		for (auto& VariableToMetaDataItem : VariableToMetaData)
+		for (auto& VariableToScriptVariableItem : VariableToScriptVariable)
 		{
-			FNiagaraVariableMetaData& MetaData = VariableToMetaDataItem.Value;
+			UNiagaraScriptVariable*& MetaData = VariableToScriptVariableItem.Value;
 
 			// Migrate advanced display.
-			if (MetaData.PropertyMetaData.Contains("AdvancedDisplay"))
+			if (MetaData->Metadata.PropertyMetaData.Contains("AdvancedDisplay"))
 			{
-				MetaData.bAdvancedDisplay = true;
-				MetaData.PropertyMetaData.Remove("AdvancedDisplay");
+				MetaData->Metadata.bAdvancedDisplay = true;
+				MetaData->Metadata.PropertyMetaData.Remove("AdvancedDisplay");
 			}
 
 			// Migrate inline edit condition toggle
-			if (MetaData.PropertyMetaData.Contains("InlineEditConditionToggle"))
+			if (MetaData->Metadata.PropertyMetaData.Contains("InlineEditConditionToggle"))
 			{
-				MetaData.bInlineEditConditionToggle = true;
-				MetaData.PropertyMetaData.Remove("InlineEditConditionToggle");
+				MetaData->Metadata.bInlineEditConditionToggle = true;
+				MetaData->Metadata.PropertyMetaData.Remove("InlineEditConditionToggle");
 			}
 
 			// Migrate edit and visible conditions
-			MigrateInputCondition(MetaData.PropertyMetaData, TEXT("EditCondition"), MetaData.EditCondition);
-			MigrateInputCondition(MetaData.PropertyMetaData, TEXT("VisibleCondition"), MetaData.VisibleCondition);
+			MigrateInputCondition(MetaData->Metadata.PropertyMetaData, TEXT("EditCondition"), MetaData->Metadata.EditCondition);
+			MigrateInputCondition(MetaData->Metadata.PropertyMetaData, TEXT("VisibleCondition"), MetaData->Metadata.VisibleCondition);
 		}
 	}
+
+	InvalidateCachedParameterData();
 }
 
 void UNiagaraGraph::PostEditChangeProperty(FPropertyChangedEvent& PropertyChangedEvent)
@@ -245,9 +290,9 @@ class UNiagaraScriptSource* UNiagaraGraph::GetSource() const
 	return CastChecked<UNiagaraScriptSource>(GetOuter());
 }
 
-FGuid UNiagaraGraph::GetCompileID(ENiagaraScriptUsage InUsage, const FGuid& InUsageId)
+FGuid UNiagaraGraph::ComputeCompileID(ENiagaraScriptUsage InUsage, const FGuid& InUsageId)
 {
-	RebuildCachedData();
+	RebuildCachedCompileIds();
 
 	for (int32 j = 0; j < CachedUsageInfo.Num(); j++)
 	{
@@ -261,6 +306,46 @@ FGuid UNiagaraGraph::GetCompileID(ENiagaraScriptUsage InUsage, const FGuid& InUs
 
 }
 
+FNiagaraCompileHash UNiagaraGraph::GetCompileDataHash(ENiagaraScriptUsage InUsage, const FGuid& InUsageId) const
+{
+	for (int32 i = 0; i < CachedUsageInfo.Num(); i++)
+	{
+		if (UNiagaraScript::IsEquivalentUsage(CachedUsageInfo[i].UsageType, InUsage) && CachedUsageInfo[i].UsageId == InUsageId)
+		{
+			return CachedUsageInfo[i].CompileHash;
+		}
+	}
+	return FNiagaraCompileHash();
+}
+
+FGuid UNiagaraGraph::GetBaseId(ENiagaraScriptUsage InUsage, const FGuid& InUsageId) const
+{
+	for (int32 i = 0; i < CachedUsageInfo.Num(); i++)
+	{
+		if (UNiagaraScript::IsEquivalentUsage(CachedUsageInfo[i].UsageType, InUsage) && CachedUsageInfo[i].UsageId == InUsageId)
+		{
+			return CachedUsageInfo[i].BaseId;
+		}
+	}
+	return FGuid();
+}
+
+void UNiagaraGraph::ForceBaseId(ENiagaraScriptUsage InUsage, const FGuid& InUsageId, const FGuid InForcedBaseId)
+{
+	FNiagaraGraphScriptUsageInfo* MatchingCachedUsageInfo = CachedUsageInfo.FindByPredicate([InUsage, InUsageId](const FNiagaraGraphScriptUsageInfo& CachedUsageInfoItem)
+	{ 
+		return CachedUsageInfoItem.UsageType == InUsage && CachedUsageInfoItem.UsageId == InUsageId; 
+	});
+
+	if (MatchingCachedUsageInfo == nullptr)
+	{
+		MatchingCachedUsageInfo = &CachedUsageInfo.AddDefaulted_GetRef();
+		MatchingCachedUsageInfo->UsageType = InUsage;
+		MatchingCachedUsageInfo->UsageId = InUsageId;
+	}
+	MatchingCachedUsageInfo->BaseId = InForcedBaseId;
+}
+
 UEdGraphPin* UNiagaraGraph::FindParameterMapDefaultValuePin(const FName VariableName, ENiagaraScriptUsage InUsage, ENiagaraScriptUsage InParentUsage) const
 {
 	TArray<UEdGraphPin*> MatchingDefaultPins;
@@ -529,6 +614,30 @@ void UNiagaraGraph::FindInputNodes(TArray<UNiagaraNodeInput*>& OutInputNodes, UN
 	}
 }
 
+TArray<FNiagaraVariable> UNiagaraGraph::FindStaticSwitchInputs() const
+{
+	TArray<FNiagaraVariable> Result;
+	for (UEdGraphNode* Node : Nodes)
+	{
+		UNiagaraNodeStaticSwitch* SwitchNode = Cast<UNiagaraNodeStaticSwitch>(Node);
+		if (SwitchNode)
+		{
+			FNiagaraVariable Variable(SwitchNode->GetInputType(), SwitchNode->InputParameterName);
+			Result.AddUnique(Variable);
+		}
+
+		UNiagaraNodeFunctionCall* FunctionNode = Cast<UNiagaraNodeFunctionCall>(Node);
+		if (FunctionNode)
+		{
+			for (FNiagaraVariable Variable : FunctionNode->PropagatedStaticSwitchParameters)
+			{
+				Result.AddUnique(Variable);
+			}			
+		}
+	}
+	return Result;
+}
+
 void UNiagaraGraph::GetParameters(TArray<FNiagaraVariable>& Inputs, TArray<FNiagaraVariable>& Outputs)const
 {
 	Inputs.Empty();
@@ -563,30 +672,57 @@ void UNiagaraGraph::GetParameters(TArray<FNiagaraVariable>& Inputs, TArray<FNiag
 // 	Outputs.Sort(SortVars);
 }
 
-const TMap<FNiagaraVariable, FNiagaraGraphParameterReferenceCollection>& UNiagaraGraph::GetParameterMap() const
+const TMap<FNiagaraVariable, UNiagaraScriptVariable*>& UNiagaraGraph::GetAllMetaData() const
+{
+	if (bUnreferencedMetaDataPurgePending)
+	{
+		PurgeUnreferencedMetaData();
+	}
+	return VariableToScriptVariable;
+}
+
+TMap<FNiagaraVariable, UNiagaraScriptVariable*>& UNiagaraGraph::GetAllMetaData()
 {
-	return Parameters;
+	if (bUnreferencedMetaDataPurgePending)
+	{
+		PurgeUnreferencedMetaData();
+	}
+	return VariableToScriptVariable;
+}
+
+const TMap<FNiagaraVariable, FNiagaraGraphParameterReferenceCollection>& UNiagaraGraph::GetParameterReferenceMap() const
+{
+	if (bParameterReferenceRefreshPending)
+	{
+		RefreshParameterReferences();
+	}
+	return ParameterToReferencesMap;
 }
 
 void UNiagaraGraph::AddParameter(const FNiagaraVariable& Parameter)
 {
-	FNiagaraGraphParameterReferenceCollection* FoundParameterReferenceCollection = Parameters.Find(Parameter);
+	FNiagaraGraphParameterReferenceCollection* FoundParameterReferenceCollection = ParameterToReferencesMap.Find(Parameter);
 	if (!FoundParameterReferenceCollection)
 	{
 		FNiagaraGraphParameterReferenceCollection NewReferenceCollection = FNiagaraGraphParameterReferenceCollection(true /*bCreated*/);
 		NewReferenceCollection.Graph = this;
-		Parameters.Add(Parameter, NewReferenceCollection);
+		ParameterToReferencesMap.Add(Parameter, NewReferenceCollection);
+	}
+
+	UNiagaraScriptVariable** FoundScriptVariable = VariableToScriptVariable.Find(Parameter);
+	if (!FoundScriptVariable)
+	{
+		UNiagaraScriptVariable* NewScriptVariable = NewObject<UNiagaraScriptVariable>(this);
+		NewScriptVariable->Variable = Parameter;
+		VariableToScriptVariable.Add(Parameter, NewScriptVariable);
 	}
 }
 
-void UNiagaraGraph::RemoveParameter(const FNiagaraVariable& Parameter, const bool bNotifyGraphChanged /*= true*/)
+void UNiagaraGraph::RemoveParameter(const FNiagaraVariable& Parameter)
 {
-	FNiagaraGraphParameterReferenceCollection* ReferenceCollection = Parameters.Find(Parameter);
+	FNiagaraGraphParameterReferenceCollection* ReferenceCollection = ParameterToReferencesMap.Find(Parameter);
 	if (ReferenceCollection)
 	{
-		// Prevent finding all parameters and metadata when renaming each pin.
-		SetFindParametersAllowed(false);
-
 		for (int32 Index = 0; Index < ReferenceCollection->ParameterReferences.Num(); Index++)
 		{
 			const FNiagaraGraphParameterReference& Reference = ReferenceCollection->ParameterReferences[Index];
@@ -601,18 +737,14 @@ void UNiagaraGraph::RemoveParameter(const FNiagaraVariable& Parameter, const boo
 			}
 		}
 
-		Parameters.Remove(Parameter);
-
-		SetFindParametersAllowed(true);
-
-		if (bNotifyGraphChanged)
-		{
-			NotifyGraphChanged();
-		}
+		// Remove it from the reference collection directly because it might have been user added and
+		// these aren't removed when the cached data is rebuilt.
+		ParameterToReferencesMap.Remove(Parameter);
+		NotifyGraphChanged();
 	}
 }
 
-bool UNiagaraGraph::RenameParameter(const FNiagaraVariable& Parameter, FName NewName, const bool bInNotifyGraphChanged /*= true*/)
+bool UNiagaraGraph::RenameParameter(const FNiagaraVariable& Parameter, FName NewName)
 {
 	// Block rename when already renaming. This prevents recursion when CommitEditablePinName is called on referenced nodes. 
 	if (bIsRenamingParameter)
@@ -621,14 +753,19 @@ bool UNiagaraGraph::RenameParameter(const FNiagaraVariable& Parameter, FName New
 	}
 	bIsRenamingParameter = true;
 
-	// Prevent finding all parameters and metadata when renaming each pin.
-	SetFindParametersAllowed(false);
-	
 	// Create the new parameter
 	FNiagaraVariable NewParameter = Parameter;
 	NewParameter.SetName(NewName);
 
-	FNiagaraGraphParameterReferenceCollection* ReferenceCollection = Parameters.Find(Parameter);
+	UNiagaraScriptVariable** OldScriptVariable = VariableToScriptVariable.Find(Parameter);
+	FNiagaraVariableMetaData OldMetaData;
+	if (OldScriptVariable)
+	{
+		OldMetaData = (*OldScriptVariable)->Metadata;
+	}
+		
+
+	FNiagaraGraphParameterReferenceCollection* ReferenceCollection = ParameterToReferencesMap.Find(Parameter);
 	if (ReferenceCollection)
 	{
 		const FText NewNameText = FText::FromName(NewName);
@@ -646,26 +783,20 @@ bool UNiagaraGraph::RenameParameter(const FNiagaraVariable& Parameter, FName New
 			}
 		}
 
-		Parameters.Remove(Parameter);
-		Parameters.Add(NewParameter, NewReferences);
+		ParameterToReferencesMap.Remove(Parameter);
+		ParameterToReferencesMap.Add(NewParameter, NewReferences);
 	}
 
 	// Swap metadata to the new parameter
-	FNiagaraVariableMetaData* Metadata = GetMetaData(Parameter);
-	if (Metadata)
+	if (UNiagaraScriptVariable** Metadata = VariableToScriptVariable.Find(Parameter))
 	{
-		FNiagaraVariableMetaData MetadataCopy = *Metadata;
-		VariableToMetaData.Remove(Parameter);
-		VariableToMetaData.Add(NewParameter, MetadataCopy);
+		VariableToScriptVariable.Remove(Parameter);
 	}
+	SetMetaData(NewParameter, OldMetaData);
 
-	SetFindParametersAllowed(true);
 	bIsRenamingParameter = false;
 
-	if (bInNotifyGraphChanged)
-	{
-		NotifyGraphChanged();
-	}
+	NotifyGraphChanged();
 	return true;
 }
 
@@ -794,7 +925,7 @@ FNiagaraTypeDefinition UNiagaraGraph::GetCachedNumericConversion(class UEdGraphP
 	return ReturnDef;
 }
 
-void UNiagaraGraph::RebuildCachedData(bool bForce)
+void UNiagaraGraph::RebuildCachedCompileIds(bool bForce)
 {
 	// If the graph hasn't changed since last rebuild, then do nothing.
 	if (!bForce && ChangeId == LastBuiltTraversalDataChangeId && LastBuiltTraversalDataChangeId.IsValid())
@@ -834,6 +965,16 @@ void UNiagaraGraph::RebuildCachedData(bool bForce)
 			}
 		}
 
+		if (FoundMatchIdx == INDEX_NONE || CachedUsageInfo[FoundMatchIdx].BaseId.IsValid() == false)
+		{
+			NewUsageCache[i].BaseId = FGuid::NewGuid();
+		}
+		else
+		{
+			//Copy the old base id if available and valid.
+			NewUsageCache[i].BaseId = CachedUsageInfo[FoundMatchIdx].BaseId;
+		}
+
 		// Now compare the change id's of all the nodes in the traversal by hashing them up and comparing the hash
 		// now with the hash from previous runs.
 		FSHA1 HashState;
@@ -845,15 +986,17 @@ void UNiagaraGraph::RebuildCachedData(bool bForce)
 		HashState.Final();
 
 		// We can't store in a FShaHash struct directly because you can't UProperty it. Using a standin of the same size.
-		check(sizeof(uint8)*NewUsageCache[i].DataHash.Num() == sizeof(FSHAHash));
-		HashState.GetHash(&NewUsageCache[i].DataHash[0]);
+		TArray<uint8> DataHash;
+		DataHash.AddUninitialized(20);
+		HashState.GetHash(DataHash.GetData());
+		NewUsageCache[i].CompileHash = FNiagaraCompileHash(DataHash);
 
 		bool bNeedsNewCompileId = true;
 
 		// Now compare the hashed data. If it is the same as before, then leave the compile ID as-is. If it is different, generate a new guid.
 		if (FoundMatchIdx != INDEX_NONE)
 		{
-			if (NewUsageCache[i].DataHash == CachedUsageInfo[FoundMatchIdx].DataHash)
+			if (NewUsageCache[i].CompileHash == CachedUsageInfo[FoundMatchIdx].CompileHash)
 			{
 				NewUsageCache[i].GeneratedCompileId = CachedUsageInfo[FoundMatchIdx].GeneratedCompileId;
 				bNeedsNewCompileId = false;
@@ -907,6 +1050,19 @@ void UNiagaraGraph::RebuildCachedData(bool bForce)
 		FNiagaraGraphScriptUsageInfo GpuUsageInfo;
 		GpuUsageInfo.UsageType = ENiagaraScriptUsage::ParticleGPUComputeScript;
 		GpuUsageInfo.UsageId = FGuid();
+
+		FNiagaraGraphScriptUsageInfo* OldGpuInfo = CachedUsageInfo.FindByPredicate(
+			[](const FNiagaraGraphScriptUsageInfo& OldInfo) { return OldInfo.UsageType == ENiagaraScriptUsage::ParticleGPUComputeScript && OldInfo.UsageId == FGuid(); });
+		if (OldGpuInfo == nullptr || OldGpuInfo->BaseId.IsValid() == false)
+		{
+			GpuUsageInfo.BaseId = FGuid::NewGuid();
+		}
+		else
+		{
+			// Copy the old base id if available
+			GpuUsageInfo.BaseId = OldGpuInfo->BaseId;
+		}
+
 		GpuUsageInfo.Traversal.Append(ParticleSpawnUsageInfo->Traversal);
 		GpuUsageInfo.Traversal.Append(ParticleUpdateUsageInfo->Traversal);
 
@@ -918,11 +1074,13 @@ void UNiagaraGraph::RebuildCachedData(bool bForce)
 		}
 		HashState.Final();
 
-		check(sizeof(uint8) * GpuUsageInfo.DataHash.Num() == sizeof(FSHAHash));
-		HashState.GetHash(GpuUsageInfo.DataHash.GetData());
+		TArray<uint8> DataHash;
+		DataHash.AddUninitialized(20);
+		HashState.GetHash(DataHash.GetData());
+		GpuUsageInfo.CompileHash = FNiagaraCompileHash(DataHash);
 
 		FNiagaraGraphScriptUsageInfo* OldGpuUsageInfo = CachedUsageInfo.FindByPredicate([](const FNiagaraGraphScriptUsageInfo& UsageInfo) { return UsageInfo.UsageType == ENiagaraScriptUsage::ParticleGPUComputeScript && UsageInfo.UsageId == FGuid(); });
-		if (OldGpuUsageInfo != nullptr && OldGpuUsageInfo->DataHash == GpuUsageInfo.DataHash)
+		if (OldGpuUsageInfo != nullptr && OldGpuUsageInfo->CompileHash == GpuUsageInfo.CompileHash)
 		{
 			GpuUsageInfo.GeneratedCompileId = OldGpuUsageInfo->GeneratedCompileId;
 		}
@@ -981,6 +1139,12 @@ FString UNiagaraGraph::GetFunctionAliasByContext(const FNiagaraGraphFunctionAlia
 			NiagaraNode->AppendFunctionAliasForContext(FunctionAliasContext, FunctionAlias);
 		}
 	}
+
+	for (UEdGraphPin* Pin : FunctionAliasContext.StaticSwitchValues)
+	{
+		FunctionAlias += TEXT("_") + FHlslNiagaraTranslator::GetSanitizedFunctionNameSuffix(Pin->GetName()) 
+			+ TEXT("_") + FHlslNiagaraTranslator::GetSanitizedFunctionNameSuffix(Pin->DefaultValue);
+	}
 	return FunctionAlias;
 }
 
@@ -995,13 +1159,16 @@ void UNiagaraGraph::ResolveNumerics(TMap<UNiagaraNode*, bool>& VisitedNodes, UEd
 		{
 			if (InputPins[i])
 			{
-				UNiagaraNode* FoundNode = Cast<UNiagaraNode>(InputPins[i]->GetOwningNode());
-				if (!FoundNode || VisitedNodes.Contains(FoundNode))
+				for (int32 j = 0; j < InputPins[i]->LinkedTo.Num(); j++)
 				{
-					continue;
+					UNiagaraNode* FoundNode = Cast<UNiagaraNode>(InputPins[i]->LinkedTo[j]->GetOwningNode());
+					if (!FoundNode || VisitedNodes.Contains(FoundNode))
+					{
+						continue;
+					}
+					VisitedNodes.Add(FoundNode, true);
+					ResolveNumerics(VisitedNodes, FoundNode);
 				}
-				VisitedNodes.Add(FoundNode, true);
-				ResolveNumerics(VisitedNodes, FoundNode);
 			}
 		}
 
@@ -1014,7 +1181,7 @@ void UNiagaraGraph::ResolveNumerics(TMap<UNiagaraNode*, bool>& VisitedNodes, UEd
 void UNiagaraGraph::SynchronizeInternalCacheWithGraph(UNiagaraGraph* Other)
 {
 	// Force us to rebuild the cache, note that this builds traversals and everything else, keeping it in sync if nothing changed from the current version.
-	RebuildCachedData(true);
+	RebuildCachedCompileIds(true);
 	
 	UEnum* FoundEnum = nullptr;
 
@@ -1034,7 +1201,7 @@ void UNiagaraGraph::SynchronizeInternalCacheWithGraph(UNiagaraGraph* Other)
 
 		if (FoundMatchIdx != INDEX_NONE)
 		{
-			if (CachedUsageInfo[i].DataHash == Other->CachedUsageInfo[FoundMatchIdx].DataHash)
+			if (CachedUsageInfo[i].CompileHash == Other->CachedUsageInfo[FoundMatchIdx].CompileHash)
 			{
 				CachedUsageInfo[i].GeneratedCompileId = Other->CachedUsageInfo[FoundMatchIdx].GeneratedCompileId;		
 
@@ -1050,7 +1217,10 @@ void UNiagaraGraph::SynchronizeInternalCacheWithGraph(UNiagaraGraph* Other)
 					{
 						ResultsEnum = FoundEnum->GetNameStringByValue((int64)CachedUsageInfo[i].UsageType);
 					}
-					UE_LOG(LogNiagaraEditor, Log, TEXT("'%s' changes synchronized with master script in %s .. synced guid: %s"), *GetFullName(), *ResultsEnum, *CachedUsageInfo[i].GeneratedCompileId.ToString());
+					if (GEnableVerboseNiagaraChangeIdLogging)
+					{
+						UE_LOG(LogNiagaraEditor, Log, TEXT("'%s' changes synchronized with master script in %s .. synced guid: %s"), *GetFullName(), *ResultsEnum, *CachedUsageInfo[i].GeneratedCompileId.ToString());
+					}
 				}
 			}
 		}
@@ -1071,9 +1241,9 @@ void UNiagaraGraph::InvalidateCachedCompileIds()
 	MarkGraphRequiresSynchronization(__FUNCTION__);
 }
 
-void UNiagaraGraph::GatherExternalDependencyIDs(ENiagaraScriptUsage InUsage, const FGuid& InUsageId, TArray<FGuid>& InReferencedIDs, TArray<UObject*>& InReferencedObjs)
+void UNiagaraGraph::GatherExternalDependencyIDs(ENiagaraScriptUsage InUsage, const FGuid& InUsageId, TArray<FNiagaraCompileHash>& InReferencedCompileHashes, TArray<FGuid>& InReferencedIDs, TArray<UObject*>& InReferencedObjs)
 {
-	RebuildCachedData();
+	RebuildCachedCompileIds();
 	
 	// Particle compute scripts get all particle scripts baked into their dependency chain. 
 	if (InUsage == ENiagaraScriptUsage::ParticleGPUComputeScript)
@@ -1083,12 +1253,12 @@ void UNiagaraGraph::GatherExternalDependencyIDs(ENiagaraScriptUsage InUsage, con
 			// Add all chains that we depend on.
 			if (UNiagaraScript::IsUsageDependentOn(InUsage, CachedUsageInfo[i].UsageType)) 
 			{
-				InReferencedIDs.Add(CachedUsageInfo[i].GeneratedCompileId);
+				InReferencedCompileHashes.Add(CachedUsageInfo[i].CompileHash);
 				InReferencedObjs.Add(CachedUsageInfo[i].Traversal.Last());
 
 				for (UNiagaraNode* Node : CachedUsageInfo[i].Traversal)
 				{
-					Node->GatherExternalDependencyIDs(InUsage, InUsageId, InReferencedIDs, InReferencedObjs);
+					Node->GatherExternalDependencyIDs(InUsage, InUsageId, InReferencedCompileHashes, InReferencedIDs, InReferencedObjs);
 				}
 			}
 		}
@@ -1104,18 +1274,18 @@ void UNiagaraGraph::GatherExternalDependencyIDs(ENiagaraScriptUsage InUsage, con
 				// Skip adding to list because we already did it in GetCompileId above.
 				for (UNiagaraNode* Node : CachedUsageInfo[i].Traversal)
 				{
-					Node->GatherExternalDependencyIDs(InUsage, InUsageId, InReferencedIDs, InReferencedObjs);
+					Node->GatherExternalDependencyIDs(InUsage, InUsageId, InReferencedCompileHashes, InReferencedIDs, InReferencedObjs);
 				}
 			}
 			// Now add any other dependency chains that we might have...
 			else if (UNiagaraScript::IsUsageDependentOn(InUsage, CachedUsageInfo[i].UsageType))
 			{
-				InReferencedIDs.Add(CachedUsageInfo[i].GeneratedCompileId);
+				InReferencedCompileHashes.Add(CachedUsageInfo[i].CompileHash);
 				InReferencedObjs.Add(CachedUsageInfo[i].Traversal.Last());
 
 				for (UNiagaraNode* Node : CachedUsageInfo[i].Traversal)
 				{
-					Node->GatherExternalDependencyIDs(InUsage, InUsageId, InReferencedIDs, InReferencedObjs);
+					Node->GatherExternalDependencyIDs(InUsage, InUsageId, InReferencedCompileHashes, InReferencedIDs, InReferencedObjs);
 				}
 			}
 		}
@@ -1178,74 +1348,85 @@ void UNiagaraGraph::MarkGraphRequiresSynchronization(FString Reason)
 {
 	Modify();
 	ChangeId = FGuid::NewGuid();
-	//UE_LOG(LogNiagaraEditor, Verbose, TEXT("Graph %s was marked requires synchronization.  Reason: %s"), *GetPathName(), *Reason);
+	if (GEnableVerboseNiagaraChangeIdLogging)
+	{
+		UE_LOG(LogNiagaraEditor, Verbose, TEXT("Graph %s was marked requires synchronization.  Reason: %s"), *GetPathName(), *Reason);
+	}
 }
 
-/** Get the meta-data associated with this variable, if it exists.*/
-FNiagaraVariableMetaData* UNiagaraGraph::GetMetaData(const FNiagaraVariable& InVar)
+TOptional<FNiagaraVariableMetaData> UNiagaraGraph::GetMetaData(const FNiagaraVariable& InVar) const
 {
-	return VariableToMetaData.Find(InVar);
+	if (bUnreferencedMetaDataPurgePending)
+	{
+		PurgeUnreferencedMetaData();
+	}
+	
+	if (UNiagaraScriptVariable** MetaData = VariableToScriptVariable.Find(InVar))
+	{
+		if (*MetaData)
+		{
+			return (*MetaData)->Metadata;
+		}
+	}
+	return TOptional<FNiagaraVariableMetaData>();
 }
 
-const FNiagaraVariableMetaData* UNiagaraGraph::GetMetaData(const FNiagaraVariable& InVar) const
+void UNiagaraGraph::SetMetaData(const FNiagaraVariable& InVar, const FNiagaraVariableMetaData& InMetaData)
 {
-	return VariableToMetaData.Find(InVar);
-}
+	ensure(FNiagaraConstants::IsNiagaraConstant(InVar) == false);
 
-/** Return the meta-data associated with this variable. This should only be called on variables defined within this Graph, otherwise meta-data may leak.*/
-FNiagaraVariableMetaData& UNiagaraGraph::FindOrAddMetaData(const FNiagaraVariable& InVar)
-{
-	FNiagaraVariableMetaData* FoundMetaData = VariableToMetaData.Find(InVar);
-	if (FoundMetaData)
+	if (UNiagaraScriptVariable** FoundMetaData = VariableToScriptVariable.Find(InVar))
 	{
-		return *FoundMetaData;
+		if (*FoundMetaData)
+		{
+			// Replace the old metadata..
+			(*FoundMetaData)->Metadata = InMetaData;
+		} 
+	}
+	else 
+	{
+		UNiagaraScriptVariable*& NewScriptVariable = VariableToScriptVariable.Add(InVar, NewObject<UNiagaraScriptVariable>(this));
+		NewScriptVariable->Variable = InVar;
+		NewScriptVariable->Metadata = InMetaData;
 	}
-	// We shouldn't add constants to the graph's meta-data list. Those are stored globally.
-	ensure(FNiagaraConstants::IsNiagaraConstant(InVar) == false);
-	return VariableToMetaData.Add(InVar);
 }
 
-void UNiagaraGraph::PurgeUnreferencedMetaData()
+void UNiagaraGraph::PurgeUnreferencedMetaData() const
 {
-	TArray<FNiagaraVariable> VarsToRemove;
-	for (auto It = VariableToMetaData.CreateConstIterator(); It; ++It)
+	TSet<FNiagaraVariable> ReferencedParameters;
+	ReferencedParameters.Append(FindStaticSwitchInputs());
+	const UEdGraphSchema_Niagara* NiagaraSchema = Cast<UEdGraphSchema_Niagara>(Schema);
+	for (UEdGraphNode* Node : Nodes)
 	{
-		int32 NumValid = 0;
-		for (TWeakObjectPtr<UObject> WeakPtr : It.Value().ReferencerNodes)
+		for (UEdGraphPin* Pin : Node->Pins)
 		{
-			if (WeakPtr.IsValid())
+			if (Pin->PinType.PinSubCategory == UNiagaraNodeParameterMapBase::ParameterPinSubCategory)
 			{
-				NumValid++;
+				const FNiagaraVariable Parameter = NiagaraSchema->PinToNiagaraVariable(Pin, false);
+				const FNiagaraParameterHandle Handle = FNiagaraParameterHandle(Parameter.GetName());
+				if (Handle.IsModuleHandle() && !FNiagaraConstants::IsNiagaraConstant(Parameter))
+				{
+					ReferencedParameters.Add(Parameter);
+				}
 			}
 		}
-
-		if (NumValid == 0)
-		{
-			VarsToRemove.Add(It.Key());
-		}
-	}
-
-	for (FNiagaraVariable& Var : VarsToRemove)
-	{
-		VariableToMetaData.Remove(Var);
 	}
-}
 
-void UNiagaraGraph::PurgeUnreferencedParameters()
-{
 	TArray<FNiagaraVariable> VarsToRemove;
-	for (auto& ParameterEntry : Parameters)
+	for (auto It = VariableToScriptVariable.CreateConstIterator(); It; ++It)
 	{
-		if (!ParameterEntry.Value.WasCreated() && ParameterEntry.Value.ParameterReferences.Num() == 0)
+		if (!ReferencedParameters.Contains(It.Key()))
 		{
-			VarsToRemove.Add(ParameterEntry.Key);
+			VarsToRemove.Add(It.Key());
 		}
 	}
 
 	for (FNiagaraVariable& Var : VarsToRemove)
 	{
-		Parameters.Remove(Var);
+		VariableToScriptVariable.Remove(Var);
 	}
+
+	bUnreferencedMetaDataPurgePending = false;
 }
 
 UNiagaraGraph::FOnDataInterfaceChanged& UNiagaraGraph::OnDataInterfaceChanged()
@@ -1253,55 +1434,79 @@ UNiagaraGraph::FOnDataInterfaceChanged& UNiagaraGraph::OnDataInterfaceChanged()
 	return OnDataInterfaceChangedDelegate;
 }
 
-void UNiagaraGraph::FindParameters()
+void UNiagaraGraph::RefreshParameterReferences() const 
 {
-	if (!bFindParametersAllowed)
-	{
-		return;
-	}
+	// A set of variables to track which parameters are used so that unused parameters can be removed after the reference tracking.
+	TSet<FNiagaraVariable> CandidateUnreferencedParametersToRemove;
 
-	for (auto& ParameterEntry : Parameters)
-	{
-		ParameterEntry.Value.ParameterReferences.Empty();
-	}
+	// The set of pins which has already been handled by add parameters.
+	TSet<const UEdGraphPin*> HandledPins;
 
-	for (auto& MetadataEntry : VariableToMetaData)
+	// Purge existing parameter references and collect candidate unreferenced parameters.
+	for (auto& ParameterToReferences : ParameterToReferencesMap)
 	{
-		MetadataEntry.Value.ReferencerNodes.Empty();
+		ParameterToReferences.Value.ParameterReferences.Empty();
+		if (ParameterToReferences.Value.WasCreated() == false)
+		{
+			// Collect all parameters not created for the user so that they can be removed later if no references are found for them.
+			CandidateUnreferencedParametersToRemove.Add(ParameterToReferences.Key);
+		}
 	}
 
-	auto AddParameterReference = [&](const FNiagaraVariable& Parameter, const UEdGraphPin* Pin, FNiagaraGraphParameterReferenceCollection*& ReferenceCollection)
+	auto AddParameterReference = [&](const FNiagaraVariable& Parameter, const UEdGraphPin* Pin)
 	{
 		if (Pin->PinType.PinSubCategory == UNiagaraNodeParameterMapBase::ParameterPinSubCategory)
 		{
-			const FNiagaraGraphParameterReference Reference(Pin->PersistentGuid, Cast<UNiagaraNode>(Pin->GetOwningNode()));
-			bool bNewReference = true;
-			if (ReferenceCollection)
+			FNiagaraGraphParameterReferenceCollection* ReferenceCollection = ParameterToReferencesMap.Find(Parameter);
+			if (ReferenceCollection == nullptr)
 			{
-				ReferenceCollection->ParameterReferences.AddUnique(Reference);
-				bNewReference = false;
-			}
-			else
-			{
-				FNiagaraGraphParameterReferenceCollection* FoundReferenceCollection = Parameters.Find(Parameter);
-				if (FoundReferenceCollection)
+				FNiagaraGraphParameterReferenceCollection& NewReferenceCollection = ParameterToReferencesMap.Add(Parameter);
+				NewReferenceCollection.Graph = this;
+				ReferenceCollection = &NewReferenceCollection;
+
+				// When a variable is created or added from the graph it won't call AddParameter, 
+				// but instead call this method
+				UNiagaraScriptVariable** FoundScriptVariable = VariableToScriptVariable.Find(Parameter);
+				if (!FoundScriptVariable)
 				{
-					ReferenceCollection = FoundReferenceCollection;
-					FoundReferenceCollection->ParameterReferences.AddUnique(Reference);
-					bNewReference = false;
+					// Warning! This method (RefreshParameterReferences) isn't const at all!
+					// Need to cast away the const to be able to create a serializable UObject here...
+					UNiagaraScriptVariable* NewScriptVariable = NewObject<UNiagaraScriptVariable>(const_cast<UNiagaraGraph*>(this));
+					NewScriptVariable->Variable = Parameter;
+					VariableToScriptVariable.Add(Parameter, NewScriptVariable);
 				}
 			}
+			ReferenceCollection->ParameterReferences.AddUnique(FNiagaraGraphParameterReference(Pin->PersistentGuid, Cast<UNiagaraNode>(Pin->GetOwningNode())));
 
-			if (bNewReference)
-			{
-				FNiagaraGraphParameterReferenceCollection NewReferenceCollection;
-				NewReferenceCollection.ParameterReferences.Add(Reference);
-				NewReferenceCollection.Graph = this;
-				Parameters.Add(Parameter, NewReferenceCollection);
-			}
+			// If we're adding a parameter reference then it needs to be removed from the list of candidate variables to remove since it's been referenced.
+			CandidateUnreferencedParametersToRemove.Remove(Parameter);
+		}
+
+		HandledPins.Add(Pin);
+	};
+
+	auto AddStaticParameterReference = [&](const FNiagaraVariable& Variable, UNiagaraNode* Node)
+	{
+		FNiagaraGraphParameterReferenceCollection* ReferenceCollection = ParameterToReferencesMap.Find(Variable);
+		if (ReferenceCollection == nullptr)
+		{
+			FNiagaraGraphParameterReferenceCollection NewReferenceCollection(true);
+			NewReferenceCollection.Graph = this;
+			ReferenceCollection = &ParameterToReferencesMap.Add(Variable, NewReferenceCollection);
+		}
+		UNiagaraScriptVariable** FoundScriptVariable = VariableToScriptVariable.Find(Variable);
+		if (!FoundScriptVariable)
+		{
+			UNiagaraScriptVariable* NewScriptVariable = NewObject<UNiagaraScriptVariable>(const_cast<UNiagaraGraph*>(this));
+			NewScriptVariable->Variable = Variable;
+			NewScriptVariable->Metadata.bIsStaticSwitch = true;
+			VariableToScriptVariable.Add(Variable, NewScriptVariable);
 		}
+		ReferenceCollection->ParameterReferences.AddUnique(FNiagaraGraphParameterReference(Node->NodeGuid, Node));
+		CandidateUnreferencedParametersToRemove.Remove(Variable);
 	};
 
+	// Add parameter references from parameter map traversals.
 	const TArray<FNiagaraParameterMapHistory> Histories = UNiagaraNodeParameterMapBase::GetParameterMaps(this);
 	for (const FNiagaraParameterMapHistory& History : Histories)
 	{
@@ -1309,71 +1514,59 @@ void UNiagaraGraph::FindParameters()
 		{
 			const FNiagaraVariable& Parameter = History.VariablesWithOriginalAliasesIntact[Index];
 
-			FNiagaraGraphParameterReferenceCollection* FoundReferences = nullptr;
 			for (const UEdGraphPin* WritePin : History.PerVariableWriteHistory[Index])
 			{
-				AddParameterReference(Parameter, WritePin, FoundReferences);
+				AddParameterReference(Parameter, WritePin);
 			}
 
 			for (const TTuple<const UEdGraphPin*, const UEdGraphPin*>& ReadPinTuple : History.PerVariableReadHistory[Index])
 			{
-				AddParameterReference(Parameter, ReadPinTuple.Key, FoundReferences);
+				AddParameterReference(Parameter, ReadPinTuple.Key);
 			}
 		}
 	}
 
-	// Find all the parameters in the graph that have no connection and won't be picked up by the parameter map history.
+	// Check all pins on all nodes in the graph to find parameter pins which may have been missed in the parameter map traversal.  This
+	// can happen for nodes which are not fully connected and therefore don't show up in the traversal.
 	const UEdGraphSchema_Niagara* NiagaraSchema = Cast<UEdGraphSchema_Niagara>(Schema);
 	for (UEdGraphNode* Node : Nodes)
 	{
+		if (UNiagaraNodeStaticSwitch* SwitchNode = Cast<UNiagaraNodeStaticSwitch>(Node))
+		{
+			FNiagaraVariable Variable(SwitchNode->GetInputType(), SwitchNode->InputParameterName);
+			AddStaticParameterReference(Variable, SwitchNode);
+		}
+		else if (UNiagaraNodeFunctionCall* FunctionNode = Cast<UNiagaraNodeFunctionCall>(Node))
+		{
+			for (FNiagaraVariable Variable : FunctionNode->PropagatedStaticSwitchParameters)
+			{
+				AddStaticParameterReference(Variable, FunctionNode);
+			}
+		}
+
 		for (UEdGraphPin* Pin : Node->Pins)
 		{
-			if (Pin->PinType.PinSubCategory == UNiagaraNodeParameterMapBase::ParameterPinSubCategory)
+			if (HandledPins.Contains(Pin) == false)
 			{
 				const FNiagaraVariable Parameter = NiagaraSchema->PinToNiagaraVariable(Pin, false);
-				const FNiagaraParameterHandle Handle = FNiagaraParameterHandle(Parameter.GetName());
-
-				if (Handle.IsModuleHandle() && !FNiagaraConstants::IsNiagaraConstant(Parameter))
-				{
-					FNiagaraVariableMetaData* MetaData = VariableToMetaData.Find(Parameter);
-					if (MetaData)
-					{
-						MetaData->ReferencerNodes.AddUnique(Node);
-					}
-					else
-					{
-						FNiagaraVariableMetaData NewVariableMetadata;
-						NewVariableMetadata.ReferencerNodes.Add(Node);
-						VariableToMetaData.Add(Parameter, NewVariableMetadata);
-					}
-				}
-			
-				const FNiagaraGraphParameterReference Reference(Pin->PersistentGuid, Cast<UNiagaraNode>(Pin->GetOwningNode()));
-				FNiagaraGraphParameterReferenceCollection* FoundParameterReferenceCollection = Parameters.Find(Parameter);
-				if (FoundParameterReferenceCollection)
-				{
-					FoundParameterReferenceCollection->ParameterReferences.AddUnique(Reference);
-				}
-				else
-				{
-					FNiagaraGraphParameterReferenceCollection NewReferenceCollection;
-					NewReferenceCollection.ParameterReferences.Add(Reference);
-					NewReferenceCollection.Graph = this;
-					Parameters.Add(Parameter, NewReferenceCollection);
-				}
+				AddParameterReference(Parameter, Pin);
 			}
 		}
 	}
 
-	// Clean up all parameters and metadata that do not have a reference
-	PurgeUnreferencedParameters();
-	PurgeUnreferencedMetaData();
-}
+	// If there were any previous parameters which didn't have any references added, remove them here.
+	for (const FNiagaraVariable& UnreferencedParameterToRemove : CandidateUnreferencedParametersToRemove)
+	{
+		ParameterToReferencesMap.Remove(UnreferencedParameterToRemove);
+	}
 
+	bParameterReferenceRefreshPending = false;
+}
 
-void UNiagaraGraph::SetFindParametersAllowed(const bool bAllowed)
+void UNiagaraGraph::InvalidateCachedParameterData()
 {
-	bFindParametersAllowed = bAllowed;
+	bParameterReferenceRefreshPending = true;
+	bUnreferencedMetaDataPurgePending = true;
 }
 
 #undef LOCTEXT_NAMESPACE
\ No newline at end of file
diff --git a/Engine/Plugins/FX/Niagara/Source/NiagaraEditor/Private/NiagaraHlslTranslator.cpp b/Engine/Plugins/FX/Niagara/Source/NiagaraEditor/Private/NiagaraHlslTranslator.cpp
--- a/Engine/Plugins/FX/Niagara/Source/NiagaraEditor/Private/NiagaraHlslTranslator.cpp
+++ b/Engine/Plugins/FX/Niagara/Source/NiagaraEditor/Private/NiagaraHlslTranslator.cpp
@@ -85,6 +85,8 @@ FNiagaraShaderProcessorTickable NiagaraShaderProcessor;
 // because editor tickables aren't ticked during cooking
 void FNiagaraShaderQueueTickable::ProcessQueue()
 {
+	check(IsInGameThread());
+
 	for (FNiagaraCompilationQueue::NiagaraCompilationQueueItem &Item : FNiagaraCompilationQueue::Get()->GetQueue())
 	{
 		FNiagaraShaderScript* ShaderScript = Item.Script;
@@ -255,7 +257,7 @@ bool FHlslNiagaraTranslator::ValidateTypePins(UNiagaraNode* NodeToValidate)
 
 
 void FHlslNiagaraTranslator::GenerateFunctionSignature(ENiagaraScriptUsage ScriptUsage, FString InName, const FString& InFullName, UNiagaraGraph* FuncGraph, TArray<int32>& Inputs,
-	bool bHadNumericInputs, bool bHasParameterMapParameters, FNiagaraFunctionSignature& OutSig)const
+	bool bHadNumericInputs, bool bHasParameterMapParameters, TArray<UEdGraphPin*> StaticSwitchValues, FNiagaraFunctionSignature& OutSig)const
 {
 	SCOPE_CYCLE_COUNTER(STAT_NiagaraEditor_Module_NiagaraHLSLTranslator_GenerateFunctionSignature);
 
@@ -314,7 +316,7 @@ void FHlslNiagaraTranslator::GenerateFunctionSignature(ENiagaraScriptUsage Scrip
 				}
 			}
 		}
-
+		
 		// Remove the inputs which will be handled by inline constants
 		for (int32 i = ConstantInputIndicesToRemove.Num() - 1; i >= 0; i--)
 		{
@@ -372,6 +374,7 @@ void FHlslNiagaraTranslator::GenerateFunctionSignature(ENiagaraScriptUsage Scrip
 	{
 		FNiagaraGraphFunctionAliasContext FunctionAliasContext;
 		FunctionAliasContext.CompileUsage = GetCurrentUsage();
+		FunctionAliasContext.StaticSwitchValues = StaticSwitchValues;
 		FString SignatureName = InName + FuncGraph->GetFunctionAliasByContext(FunctionAliasContext);
 		OutSig = FNiagaraFunctionSignature(*SignatureName, InputVars, OutputVars, *InFullName, true, false);
 	}
@@ -993,6 +996,23 @@ const FNiagaraTranslateResults &FHlslNiagaraTranslator::Translate(const FNiagara
 	if (FunctionCtx())
 		return TranslateResults;
 
+	// MASSIVE HACK - Tracked in JIRA UE-69298
+	// Hardcoded random function accessible from inner part of node implementation.
+	// It works for now at least and avoids exposing every random needed in the UI. 
+	// Temporary solution, it will be replaced when a design is validated.
+	if (CompilationTarget == ENiagaraSimTarget::GPUComputeSim)
+	{
+		HlslOutput += TEXT(R"(
+		float NiagaraInternalNoise(uint u, uint v, uint s)
+		{
+			static uint RandomSeedOffset = 0;
+			uint Seed = (u * 1664525u + v) + s + RandomSeedOffset;
+			RandomSeedOffset += Seed;
+			return float(Rand3DPCG32(int3(u,v,Seed)).x) / 4294967296.0f;
+		}
+		)");
+	}
+
 	//Now evaluate all the code chunks to generate the shader code.
 	//FString HlslOutput;
 	if (TranslateResults.bHLSLGenSucceeded)
@@ -1656,8 +1676,6 @@ void FHlslNiagaraTranslator::DefineDataInterfaceHLSL(FString &InHlslOutput)
 		UNiagaraDataInterface* CDO = Cast<UNiagaraDataInterface>(*FoundCDO);
 		if (CDO && CDO->CanExecuteOnTarget(ENiagaraSimTarget::GPUComputeSim))
 		{
-			TArray<FNiagaraFunctionSignature> DataInterfaceFunctions;
-			CDO->GetFunctions(DataInterfaceFunctions);
 			FString OwnerIDString = Info.Name.ToString();
 			FString SanitizedOwnerIDString = GetSanitizedSymbolName(OwnerIDString, true);
 
@@ -1670,14 +1688,22 @@ void FHlslNiagaraTranslator::DefineDataInterfaceHLSL(FString &InHlslOutput)
 
 			// grab the function hlsl from the interface
 			//
-			for (int32 FuncIdx = 0; FuncIdx < DataInterfaceFunctions.Num(); FuncIdx++)
+			const TSet<FNiagaraFunctionSignature>* DataInterfaceFunctions = DataInterfaceRegisteredFunctions.Find(Info.Type.GetFName());
+			if (DataInterfaceFunctions != nullptr)
 			{
-				FNiagaraFunctionSignature Sig = DataInterfaceFunctions[FuncIdx];	// make a copy so we can modify the owner id and get the correct hlsl signature
-				Sig.OwnerName = Info.Name;
-				FString DefStr = GetFunctionSignatureSymbol(Sig);
+				for (const FNiagaraFunctionSignature& OriginalSig : *DataInterfaceFunctions)
+				{
+					FNiagaraFunctionSignature Sig = OriginalSig;	// make a copy so we can modify the owner id and get the correct hlsl signature
+					Sig.OwnerName = Info.Name;
+					FString DefStr = GetFunctionSignatureSymbol(Sig);
 
-				bool HlslOK = CDO->GetFunctionHLSL(Sig.Name, DefStr, DIParamInfo[NewIdx], InterfaceFunctionHLSL);
-				ensure(HlslOK == true);
+					const bool HlslOK = CDO->GetFunctionHLSL(Sig.Name, DefStr, DIParamInfo[NewIdx], InterfaceFunctionHLSL);
+					if (HlslOK == false)
+					{
+						Error(FText::Format(LOCTEXT("GPUDataInterfaceFunctionNotSupported", "DataInterface {0} function {1} cannot run on the GPU or is not implemented."), FText::FromName(Info.Type.GetFName()), FText::FromName(Sig.Name))
+							, nullptr, nullptr);
+					}
+				}
 			}
 		}
 		else
@@ -1697,6 +1723,12 @@ void FHlslNiagaraTranslator::DefineMain(FString &OutHlslOutput,
 	if (CompilationTarget == ENiagaraSimTarget::GPUComputeSim)
 	{
 		OutHlslOutput += TEXT("void SimulateMain(in int InstanceIdx, in int InEventIndex, in int Phase)\n{\n");
+
+		// MASSIVE HACK
+		HlslOutput += TEXT(R"(
+			float RandomSeedInitialisation = NiagaraInternalNoise(InstanceIdx*16384, InEventIndex*8196, Phase*4096 + EmitterTickCounter);	// initialise the random state seed
+		)");
+		
 	}
 	else
 	{
@@ -1940,42 +1972,44 @@ void FHlslNiagaraTranslator::DefineDataSetVariableWrites(FString &OutHlslOutput,
 		for (int32 i = 0; i < ParamMapHistories.Num(); i++)
 		{
 			const UNiagaraNodeOutput* OutputNode = ParamMapHistories[i].GetFinalOutputNode();
-			bool bFound = (INDEX_NONE != ParamMapHistories[i].FindVariable(*(DataSetName + TEXT(".Alive")), FNiagaraTypeDefinition::GetBoolDef()));
-			if (bFound && OutputNode && (OutputNode->GetUsage() == ENiagaraScriptUsage::ParticleSpawnScript || OutputNode->GetUsage() == ENiagaraScriptUsage::ParticleSpawnScriptInterpolated))
+			if (!OutputNode)
 			{
-				bHasPerParticleAliveSpawn = true;
+				continue;
 			}
-			else if (bFound && OutputNode && OutputNode->GetUsage() == ENiagaraScriptUsage::ParticleUpdateScript)
+
+			if (INDEX_NONE == ParamMapHistories[i].FindVariable(*(DataSetName + TEXT(".Alive")), FNiagaraTypeDefinition::GetBoolDef()))
 			{
-				bHasPerParticleAliveUpdate = true;
+				continue;
 			}
-			else if (bFound && OutputNode && OutputNode->GetUsage() == ENiagaraScriptUsage::ParticleEventScript)
+
+			switch (OutputNode->GetUsage())
 			{
+			case ENiagaraScriptUsage::ParticleSpawnScript:
+			case ENiagaraScriptUsage::ParticleSpawnScriptInterpolated:
+				bHasPerParticleAliveSpawn = true;
+				break;
+			case ENiagaraScriptUsage::ParticleUpdateScript:
+				bHasPerParticleAliveUpdate = true;
+				break;
+			case ENiagaraScriptUsage::ParticleEventScript:
 				bHasPerParticleAliveEvent = true;
+				break;
 			}
 		}
 
 		if ((bHasPerParticleAliveSpawn || bHasPerParticleAliveUpdate) && TranslationStages.Num() > 1)
 		{
-			if (bHasPerParticleAliveSpawn && bHasPerParticleAliveUpdate)
-			{
-				OutHlslOutput += TEXT("\tbool bValid = Context.MapUpdate.") + DataSetName + TEXT(".Alive && Context.MapSpawn.DataInstance.Alive;\n");
-			}
-			else if (bHasPerParticleAliveSpawn)
-			{
-				OutHlslOutput += TEXT("\tbool bValid = Context.MapSpawn.") + DataSetName + TEXT(".Alive;\n");
-			}
-			else if (bHasPerParticleAliveUpdate)
-			{
-				OutHlslOutput += TEXT("\tbool bValid = Context.MapUpdate.") + DataSetName + TEXT(".Alive;\n");
-			}
+			// NOTE: TranslationStages.Num() > 1 for GPU Script or CPU Interpolated Spawn CPU scripts
+
+			// NOTE: Context.MapSpawn is copied to Context.MapUpdate before this point in the script, so we might
+			//       as well just keep it simple and check against MapUpdate only instead of redundantly branch.
+			OutHlslOutput += TEXT("\tbool bValid = Context.MapUpdate.") + DataSetName + TEXT(".Alive;\n");
 		}
-		else if ((UNiagaraScript::IsParticleSpawnScript(CompileOptions.TargetUsage) && bHasPerParticleAliveSpawn)
-			|| (UNiagaraScript::IsGPUScript(CompileOptions.TargetUsage) && bHasPerParticleAliveSpawn)
+		else if ((UNiagaraScript::IsParticleSpawnScript(CompileOptions.TargetUsage) && bHasPerParticleAliveSpawn) 
 			|| (UNiagaraScript::IsParticleUpdateScript(CompileOptions.TargetUsage) && bHasPerParticleAliveUpdate)
-			|| (UNiagaraScript::IsParticleEventScript(CompileOptions.TargetUsage) && bHasPerParticleAliveEvent)
-			|| (UNiagaraScript::IsGPUScript(CompileOptions.TargetUsage) && bHasPerParticleAliveUpdate && CompilationTarget == ENiagaraSimTarget::GPUComputeSim))
+			|| (UNiagaraScript::IsParticleEventScript(CompileOptions.TargetUsage) && bHasPerParticleAliveEvent))
 		{
+			// Non-interpolated CPU spawn script
 			OutHlslOutput += TEXT("\tbool bValid = Context.Map.") + DataSetName + TEXT(".Alive;\n");
 		}
 		else
@@ -1983,10 +2017,6 @@ void FHlslNiagaraTranslator::DefineDataSetVariableWrites(FString &OutHlslOutput,
 			OutHlslOutput += "\tbool bValid = true;\n";
 		}
 	}
-	int32 WriteOffsetInt = 0;
-	int32 WriteOffsetFloat = 0;
-	int32 &FloatCounter = WriteOffsetFloat;
-	int32 &IntCounter = CompilationTarget == ENiagaraSimTarget::GPUComputeSim ? WriteOffsetInt : WriteOffsetFloat;
 
 	// grab the current ouput index; currently pass true, but should use an arbitrary bool to determine whether write should happen or not
 	OutHlslOutput += "\tTmpWriteIndex = AcquireIndex(0, bValid);\n";
@@ -1998,22 +2028,23 @@ void FHlslNiagaraTranslator::DefineDataSetVariableWrites(FString &OutHlslOutput,
 		OutHlslOutput += FString::Printf(TEXT("\tUpdateID(0, %s.Particles.ID.Index, TmpWriteIndex);\n"), *MapName);
 	}
 
+	int32 WriteOffsetInt = 0;
+	int32 WriteOffsetFloat = 0;
+	int32 &FloatCounter = WriteOffsetFloat;
+	int32 &IntCounter = CompilationTarget == ENiagaraSimTarget::GPUComputeSim ? WriteOffsetInt : WriteOffsetFloat;
 	for (FNiagaraVariable &Var : WriteVars)
 	{
 		// If coming from a parameter map, use the one on the context, otherwise use the output.
 		FString Fmt;
-
+		if (TranslationStages.Num() > 1)
 		{
-			if (TranslationStages.Num() > 1)
-			{
-				Fmt = TEXT("\tOutputData{1}(0, {2}, {3}, Context.") + TranslationStages[TranslationStages.Num() - 1].PassNamespace + TEXT(".") + GetSanitizedSymbolName(Var.GetName().ToString()) + TEXT("{0});\n");
-			}
-			else
-			{
-				Fmt = TEXT("\tOutputData{1}(0, {2}, {3}, Context.Map.") + GetSanitizedSymbolName(Var.GetName().ToString()) + TEXT("{0});\n");
-			}
-			GatherVariableForDataSetAccess(Var, Fmt, IntCounter, FloatCounter, -1, TEXT("TmpWriteIndex"), OutHlslOutput);
+			Fmt = TEXT("\tOutputData{1}(0, {2}, {3}, Context.") + TranslationStages[TranslationStages.Num() - 1].PassNamespace + TEXT(".") + GetSanitizedSymbolName(Var.GetName().ToString()) + TEXT("{0});\n");
 		}
+		else
+		{
+			Fmt = TEXT("\tOutputData{1}(0, {2}, {3}, Context.Map.") + GetSanitizedSymbolName(Var.GetName().ToString()) + TEXT("{0});\n");
+		}
+		GatherVariableForDataSetAccess(Var, Fmt, IntCounter, FloatCounter, -1, TEXT("TmpWriteIndex"), OutHlslOutput);
 	}
 	OutHlslOutput += "\t}\n";
 }
@@ -2261,27 +2292,7 @@ FString FHlslNiagaraTranslator::GetSanitizedSymbolName(FString SymbolName, bool
 		SplitName[i].ReplaceInline(TEXT(" "), TEXT(""));
 
 		// Handle internationalization of characters..
-		FString ChangedSplitName;
-		ChangedSplitName.Reserve(SplitName[i].Len() * 6); // Assign room for every current char to be 'ASCXXX'
-		for (int32 j = 0; j < SplitName[i].Len(); j++)
-		{
-			if ((SplitName[i][j] >= TCHAR('0') && SplitName[i][j] <= TCHAR('9')) ||
-				(SplitName[i][j] >= TCHAR('A') && SplitName[i][j] <= TCHAR('Z')) ||
-				(SplitName[i][j] >= TCHAR('a') && SplitName[i][j] <= TCHAR('z')) ||
-				SplitName[i][j] == TCHAR('_') ||
-				SplitName[i][j] == TCHAR(' '))
-			{
-				// Do nothing.. these are valid chars..
-				ChangedSplitName.AppendChar(SplitName[i][j]);
-			}
-			else
-			{
-				// Need to replace the bad characters..
-				ChangedSplitName.Append(TEXT("ASC"));
-				ChangedSplitName.AppendInt((int32)SplitName[i][j]);
-			}
-		}
-		SplitName[i] = ChangedSplitName;
+		SplitName[i] = ConvertToAsciiString(SplitName[i]);
 	}
 
 	// Gather back into single string..
@@ -2303,6 +2314,53 @@ FString FHlslNiagaraTranslator::GetSanitizedSymbolName(FString SymbolName, bool
 	return Ret;
 }
 
+FString FHlslNiagaraTranslator::GetSanitizedFunctionNameSuffix(FString Name)
+{
+	if (Name.Len() == 0)
+	{
+		return Name;
+	}
+	FString Ret = Name;
+
+	// remove special characters
+	Ret.ReplaceInline(TEXT("."), TEXT("_"));
+	Ret.ReplaceInline(TEXT("\\"), TEXT("_"));
+	Ret.ReplaceInline(TEXT("/"), TEXT("_"));
+	Ret.ReplaceInline(TEXT(","), TEXT("_"));
+	Ret.ReplaceInline(TEXT("-"), TEXT("_"));
+	Ret.ReplaceInline(TEXT(":"), TEXT("_"));
+	Ret.ReplaceInline(TEXT("\t"), TEXT(""));
+	Ret.ReplaceInline(TEXT(" "), TEXT(""));	
+	Ret.ReplaceInline(TEXT("__"), TEXT("ASC95ASC95")); // Opengl reserves "__" within a name
+
+	// Handle internationalization of characters..
+	return ConvertToAsciiString(Ret);
+}
+
+FString FHlslNiagaraTranslator::ConvertToAsciiString(FString Str)
+{
+	FString AsciiString;
+	AsciiString.Reserve(Str.Len() * 6); // Assign room for every current char to be 'ASCXXX'
+	for (int32 j = 0; j < Str.Len(); j++)
+	{
+		if ((Str[j] >= TCHAR('0') && Str[j] <= TCHAR('9')) ||
+			(Str[j] >= TCHAR('A') && Str[j] <= TCHAR('Z')) ||
+			(Str[j] >= TCHAR('a') && Str[j] <= TCHAR('z')) ||
+			Str[j] == TCHAR('_') || Str[j] == TCHAR(' '))
+		{
+			// Do nothing.. these are valid chars..
+			AsciiString.AppendChar(Str[j]);
+		}
+		else
+		{
+			// Need to replace the bad characters..
+			AsciiString.Append(TEXT("ASC"));
+			AsciiString.AppendInt((int32)Str[j]);
+		}
+	}
+	return AsciiString;
+}
+
 FString FHlslNiagaraTranslator::GetUniqueSymbolName(FName BaseName)
 {
 	FString RetString = GetSanitizedSymbolName(BaseName.ToString());
@@ -4190,6 +4248,22 @@ void FHlslNiagaraTranslator::FunctionCall(UNiagaraNodeFunctionCall* FunctionNode
 	FunctionNode->GetOutputPins(CallOutputs);
 	FunctionNode->GetInputPins(CallInputs);
 
+	// Validate that there are no input pins with the same name and type
+	TMultiMap<FName, FEdGraphPinType> SeenPins;
+	for (UEdGraphPin* Pin : CallInputs)
+	{
+		FEdGraphPinType* SeenType = SeenPins.FindPair(Pin->GetFName(), Pin->PinType);
+		if (SeenType)
+		{
+			Error(LOCTEXT("FunctionCallDuplicateInput", "Function call has duplicated inputs. Please make sure that each function parameter is unique."), FunctionNode, Pin);
+			return;
+		}
+		else
+		{
+			SeenPins.Add(Pin->GetFName(), Pin->PinType);
+		}
+	}
+
 	// If the function call is disabled, we 
 	// need to route the input parameter map pin to the output parameter map pin.
 	// Any other outputs become invalid.
@@ -4526,11 +4600,22 @@ void FHlslNiagaraTranslator::RegisterFunctionCall(ENiagaraScriptUsage ScriptUsag
 				}
 			}
 		}
-
+		TArray<UEdGraphPin*> StaticSwitchValues;
+		for (FNiagaraVariable StaticSwitchInput : SourceGraph->FindStaticSwitchInputs())
+		{
+			for (UEdGraphPin* Pin : CallInputs)
+			{
+				if (StaticSwitchInput.GetName().IsEqual(Pin->GetFName()))
+				{
+					StaticSwitchValues.Add(Pin);
+					break;
+				}
+			}
+		}
 
 		bool bHasParameterMapParameters = SourceGraph->HasParameterMapParameters();
 
-		GenerateFunctionSignature(ScriptUsage, InName, InFullName, SourceGraph, Inputs, bHasNumericInputs, bHasParameterMapParameters, OutSignature);
+		GenerateFunctionSignature(ScriptUsage, InName, InFullName, SourceGraph, Inputs, bHasNumericInputs, bHasParameterMapParameters, StaticSwitchValues, OutSignature);
 
 		// 		//Sort the input and outputs to match the sorted parameters. They may be different.
 		// 		TArray<FNiagaraVariable> OrderedInputs;
@@ -4761,21 +4846,17 @@ void FHlslNiagaraTranslator::RegisterFunctionCall(ENiagaraScriptUsage ScriptUsag
 			{
 				TArray<FNiagaraFunctionSignature> DataInterfaceFunctions;
 				CDO->GetFunctions(DataInterfaceFunctions);
-				bool bFoundMatch = false;
-				for (const FNiagaraFunctionSignature& Sig : DataInterfaceFunctions)
-				{
-					if (Sig == OutSignature)
-					{
-						bFoundMatch = true;
-					}
-				}
 
+				const bool bFoundMatch = DataInterfaceFunctions.ContainsByPredicate([&](const FNiagaraFunctionSignature& Sig) -> bool { return Sig == OutSignature; });
 				if (!bFoundMatch)
 				{
 					Error(LOCTEXT("FunctionCallDataInterfaceMissing", "Function call signature does not match DataInterface possible signatures?"), nullptr, nullptr);
 					return;
 				}
 
+				// We only use this for GPU systems currently so that we emit only the functionality required
+				DataInterfaceRegisteredFunctions.FindOrAdd(Info.Type.GetFName()).Add(OutSignature);
+
 				if (Info.UserPtrIdx != INDEX_NONE && CompilationTarget != ENiagaraSimTarget::GPUComputeSim)
 				{
 					//This interface requires per instance data via a user ptr so place the index to it at the end of the inputs.
@@ -5321,11 +5402,39 @@ int32 FHlslNiagaraTranslator::CompileOutputPin(const UEdGraphPin* InPin)
 	// The incoming pin to compile may be pointing to a reroute node. If so, we just jump over it
 	// to where it really came from.
 	UEdGraphPin* Pin = UNiagaraNode::TraceOutputPin(const_cast<UEdGraphPin*>(InPin));
-
 	check(Pin->Direction == EGPD_Output);
 
-	int32 Ret = INDEX_NONE;
+	// The node can also replace our pin with another pin (e.g. in the case of static switches), so we need to make sure we don't run into a circular dependency
+	TSet<UEdGraphPin*> SeenPins;
+	UNiagaraNode* Node = Cast<UNiagaraNode>(Pin->GetOwningNode());
+	UEdGraphPin* OriginalPin = Pin;
+	while (Node->SubstituteCompiledPin(this, &Pin))
+	{
+		bool bIsAlreadyInSet = false;
+		SeenPins.Add(Pin, &bIsAlreadyInSet);
+		Node = Cast<UNiagaraNode>(Pin->GetOwningNode());
+		if (bIsAlreadyInSet)
+		{
+			Error(LOCTEXT("CircularGraphSubstitutionError", "Circular dependency detected, please check your module graph."), Node, Pin);
+			return INDEX_NONE;
+		}
+	}
+
+	// It is possible that the output pin was substituted by an input pin (e.g. the default value pin on a node).
+	// If that is the case we try to compile that pin directly.
+	if (Pin->Direction == EGPD_Input)
+	{
+		int32* ExistingChunk = PinToCodeChunks.Last().Find(OriginalPin); // Check if the pin was already compiled before
+		if (ExistingChunk)
+		{
+			return *ExistingChunk;
+		}
+		int32 Chunk = CompilePin(Pin);
+		PinToCodeChunks.Last().Add(OriginalPin, Chunk);
+		return Chunk;
+	}
 
+	int32 Ret = INDEX_NONE;
 	int32* Chunk = PinToCodeChunks.Last().Find(Pin);
 	if (Chunk)
 	{
@@ -5334,7 +5443,6 @@ int32 FHlslNiagaraTranslator::CompileOutputPin(const UEdGraphPin* InPin)
 	else
 	{
 		//Otherwise we need to compile the node to get its output pins.
-		UNiagaraNode* Node = Cast<UNiagaraNode>(Pin->GetOwningNode());
 		if (ValidateTypePins(Node))
 		{
 			TArray<int32> Outputs;
diff --git a/Engine/Plugins/FX/Niagara/Source/NiagaraEditor/Private/NiagaraHlslTranslator.h b/Engine/Plugins/FX/Niagara/Source/NiagaraEditor/Private/NiagaraHlslTranslator.h
--- a/Engine/Plugins/FX/Niagara/Source/NiagaraEditor/Private/NiagaraHlslTranslator.h
+++ b/Engine/Plugins/FX/Niagara/Source/NiagaraEditor/Private/NiagaraHlslTranslator.h
@@ -472,6 +472,10 @@ public:
 
 
 	static FString GetSanitizedSymbolName(FString SymbolName, bool bCollapseNamespaces=false);
+	static FString GetSanitizedFunctionNameSuffix(FString Name);
+
+	/** Replaces all non-ascii characters with a "ASCXXX" string, where XXX is their int value */
+	static FString ConvertToAsciiString(FString Name);
 
 	bool AddStructToDefinitionSet(const FNiagaraTypeDefinition& TypeDef);
 
@@ -531,7 +535,7 @@ private:
     void ValidateParticleIDUsage();
 	bool ValidateTypePins(UNiagaraNode* NodeToValidate);
 	void GenerateFunctionSignature(ENiagaraScriptUsage ScriptUsage, FString InName, const FString& InFullName, UNiagaraGraph* FuncGraph, TArray<int32>& Inputs, 
-		bool bHadNumericInputs, bool bHasParameterMapParameters, FNiagaraFunctionSignature& OutSig)const;
+		bool bHadNumericInputs, bool bHasParameterMapParameters, TArray<UEdGraphPin*> StaticSwitchValues, FNiagaraFunctionSignature& OutSig) const;
 
 	UNiagaraGraph* CloneGraphAndPrepareForCompilation(const UNiagaraScript* InScript, const UNiagaraScriptSource* InSource, bool bClearErrors);
 
@@ -622,4 +626,6 @@ private:
 	TArray<FNiagaraVariable> InitialNamespaceVariablesMissingDefault;
 	// Variables that need to be initialized in the body or at the end of spawn.
 	TArray<FNiagaraVariable> DeferredVariablesMissingDefault;
+
+	TMap<FName, TSet<FNiagaraFunctionSignature>> DataInterfaceRegisteredFunctions;
 };
diff --git a/Engine/Plugins/FX/Niagara/Source/NiagaraEditor/Private/NiagaraNode.cpp b/Engine/Plugins/FX/Niagara/Source/NiagaraEditor/Private/NiagaraNode.cpp
--- a/Engine/Plugins/FX/Niagara/Source/NiagaraEditor/Private/NiagaraNode.cpp
+++ b/Engine/Plugins/FX/Niagara/Source/NiagaraEditor/Private/NiagaraNode.cpp
@@ -29,7 +29,7 @@ void UNiagaraNode::PostLoad()
 	}
 }
 
-bool UNiagaraNode::ReallocatePins()
+bool UNiagaraNode::ReallocatePins(bool bMarkNeedsResynchronizeOnChange)
 {
 	Modify();
 
@@ -109,7 +109,7 @@ bool UNiagaraNode::ReallocatePins()
 	}
 
 	//GetGraph()->NotifyGraphChanged();
-	if (!bAllSame)
+	if (bMarkNeedsResynchronizeOnChange && !bAllSame)
 	{
 		MarkNodeRequiresSynchronization(__FUNCTION__, true);
 	}
@@ -354,7 +354,7 @@ bool UNiagaraNode::CanAddToGraph(UNiagaraGraph* TargetGraph, FString& OutErrorMs
 	return true;
 }
 
-void UNiagaraNode::BuildParameterMapHistory(FNiagaraParameterMapHistoryBuilder& OutHistory, bool bRecursive)
+void UNiagaraNode::BuildParameterMapHistory(FNiagaraParameterMapHistoryBuilder& OutHistory, bool bRecursive) const
 {
 	if (bRecursive)
 	{
@@ -362,7 +362,8 @@ void UNiagaraNode::BuildParameterMapHistory(FNiagaraParameterMapHistoryBuilder&
 	}
 }
 
-void UNiagaraNode::GetInputPins(TArray<class UEdGraphPin*>& OutInputPins) const
+template<typename PinType>
+void GetInputPinsInternal(const TArray<UEdGraphPin*>& Pins, TArray<PinType*>& OutInputPins)
 {
 	OutInputPins.Empty();
 
@@ -375,6 +376,16 @@ void UNiagaraNode::GetInputPins(TArray<class UEdGraphPin*>& OutInputPins) const
 	}
 }
 
+void UNiagaraNode::GetInputPins(TArray<class UEdGraphPin*>& OutInputPins) const
+{
+	GetInputPinsInternal<class UEdGraphPin>(Pins, OutInputPins);
+}
+
+void UNiagaraNode::GetInputPins(TArray<const class UEdGraphPin*>& OutInputPins) const
+{
+	GetInputPinsInternal<const class UEdGraphPin>(Pins, OutInputPins);
+}
+
 UEdGraphPin* UNiagaraNode::GetOutputPin(int32 OutputIndex) const
 {
 	for (int32 PinIndex = 0, FoundOutputs = 0; PinIndex < Pins.Num(); PinIndex++)
@@ -395,7 +406,8 @@ UEdGraphPin* UNiagaraNode::GetOutputPin(int32 OutputIndex) const
 	return NULL;
 }
 
-void UNiagaraNode::GetOutputPins(TArray<class UEdGraphPin*>& OutOutputPins) const
+template<typename PinType>
+void GetOutputPinsInternal(const TArray<UEdGraphPin*>& Pins, TArray<PinType*>& OutOutputPins)
 {
 	OutOutputPins.Empty();
 
@@ -408,6 +420,16 @@ void UNiagaraNode::GetOutputPins(TArray<class UEdGraphPin*>& OutOutputPins) cons
 	}
 }
 
+void UNiagaraNode::GetOutputPins(TArray<class UEdGraphPin*>& OutOutputPins) const
+{
+	return GetOutputPinsInternal<class UEdGraphPin>(Pins, OutOutputPins);
+}
+
+void UNiagaraNode::GetOutputPins(TArray<const class UEdGraphPin*>& OutOutputPins) const
+{
+	return GetOutputPinsInternal<const class UEdGraphPin>(Pins, OutOutputPins);
+}
+
 UEdGraphPin* UNiagaraNode::GetPinByPersistentGuid(const FGuid& InPersistentGuid) const
 {
 	for (UEdGraphPin* Pin : Pins)
@@ -534,7 +556,12 @@ UEdGraphPin* UNiagaraNode::TraceOutputPin(UEdGraphPin* LocallyOwnedOutputPin)
 	return LinkedNode->GetTracedOutputPin(LocallyOwnedOutputPin);
 }
 
-void UNiagaraNode::RouteParameterMapAroundMe(FNiagaraParameterMapHistoryBuilder& OutHistory, bool bRecursive)
+bool UNiagaraNode::SubstituteCompiledPin(FHlslNiagaraTranslator* Translator, UEdGraphPin** LocallyOwnedPin)
+{
+	return false;
+}
+
+void UNiagaraNode::RouteParameterMapAroundMe(FNiagaraParameterMapHistoryBuilder& OutHistory, bool bRecursive) const
 {
 	const UEdGraphSchema_Niagara* Schema = CastChecked<UEdGraphSchema_Niagara>(GetSchema());
 
diff --git a/Engine/Plugins/FX/Niagara/Source/NiagaraEditor/Private/NiagaraNodeAssignment.cpp b/Engine/Plugins/FX/Niagara/Source/NiagaraEditor/Private/NiagaraNodeAssignment.cpp
--- a/Engine/Plugins/FX/Niagara/Source/NiagaraEditor/Private/NiagaraNodeAssignment.cpp
+++ b/Engine/Plugins/FX/Niagara/Source/NiagaraEditor/Private/NiagaraNodeAssignment.cpp
@@ -172,11 +172,17 @@ void UNiagaraNodeAssignment::PostLoad()
 	}
 }
 
-void UNiagaraNodeAssignment::BuildParameterMapHistory(FNiagaraParameterMapHistoryBuilder& OutHistory, bool bRecursive)
+void UNiagaraNodeAssignment::BuildParameterMapHistory(FNiagaraParameterMapHistoryBuilder& OutHistory, bool bRecursive) const
 {
 	Super::BuildParameterMapHistory(OutHistory, bRecursive);
 }
 
+void UNiagaraNodeAssignment::GatherExternalDependencyIDs(ENiagaraScriptUsage InMasterUsage, const FGuid& InMasterUsageId, TArray<FNiagaraCompileHash>& InReferencedCompileHashes, TArray<FGuid>& InReferencedIDs, TArray<UObject*>& InReferencedObjs) const
+{
+	// Assignment nodes own their function graphs and therefore have no external dependencies so we override the default function behavior here to avoid 
+	// adding additional non-deterministic guids to the compile id generation which can invalid the DDC for compiled scripts, especially during emitter merging.
+}
+
 void UNiagaraNodeAssignment::GenerateScript()
 {
 	if (FunctionScript == nullptr)
@@ -233,7 +239,7 @@ void UNiagaraNodeAssignment::BuildCreateParameterMenu(FMenuBuilder& MenuBuilder,
 		TSet<FName> Names;
 		for (const UNiagaraGraph* Graph : Graphs)
 		{
-			for (const TPair<FNiagaraVariable, FNiagaraGraphParameterReferenceCollection>& ParameterElement : Graph->GetParameterMap())
+			for (const TPair<FNiagaraVariable, FNiagaraGraphParameterReferenceCollection>& ParameterElement : Graph->GetParameterReferenceMap())
 			{
 				Names.Add(ParameterElement.Key.GetName());
 			}
@@ -495,16 +501,18 @@ void UNiagaraNodeAssignment::InitializeScript(UNiagaraScript* NewScript)
 					const FNiagaraVariableMetaData* FoundMetaData = FNiagaraConstants::GetConstantMetaData(AssignmentTargets[i]);
 					if (FoundMetaData)
 					{
-						FNiagaraVariableMetaData& MetaData = CreatedGraph->FindOrAddMetaData(TargetVar);
-						MetaData.Description = FoundMetaData->Description;
-						MetaData.ReferencerNodes.Empty();
-						MetaData.ReferencerNodes.Add(GetNodes[0]);
+						FNiagaraVariableMetaData NewMetaData;
+						TOptional<FNiagaraVariableMetaData> ExistingMetaData = CreatedGraph->GetMetaData(TargetVar);
+						if (ExistingMetaData.IsSet())
+						{
+							NewMetaData = ExistingMetaData.GetValue();
+						}
+						NewMetaData.Description = FoundMetaData->Description;
+						CreatedGraph->SetMetaData(TargetVar, NewMetaData);
 					}
 				}
 			}
 		}
-
-		CreatedGraph->PurgeUnreferencedMetaData();
 	}
 }
 
diff --git a/Engine/Plugins/FX/Niagara/Source/NiagaraEditor/Private/NiagaraNodeCustomHlsl.cpp b/Engine/Plugins/FX/Niagara/Source/NiagaraEditor/Private/NiagaraNodeCustomHlsl.cpp
--- a/Engine/Plugins/FX/Niagara/Source/NiagaraEditor/Private/NiagaraNodeCustomHlsl.cpp
+++ b/Engine/Plugins/FX/Niagara/Source/NiagaraEditor/Private/NiagaraNodeCustomHlsl.cpp
@@ -191,7 +191,7 @@ void UNiagaraNodeCustomHlsl::MoveDynamicPin(UEdGraphPin* Pin, int32 DirectionToM
 	RebuildSignatureFromPins();
 }
 
-void UNiagaraNodeCustomHlsl::BuildParameterMapHistory(FNiagaraParameterMapHistoryBuilder& OutHistory, bool bRecursive)
+void UNiagaraNodeCustomHlsl::BuildParameterMapHistory(FNiagaraParameterMapHistoryBuilder& OutHistory, bool bRecursive) const
 {
 	Super::BuildParameterMapHistory(OutHistory, bRecursive);
 	if (!IsNodeEnabled() && OutHistory.GetIgnoreDisabled())
diff --git a/Engine/Plugins/FX/Niagara/Source/NiagaraEditor/Private/NiagaraNodeEmitter.cpp b/Engine/Plugins/FX/Niagara/Source/NiagaraEditor/Private/NiagaraNodeEmitter.cpp
--- a/Engine/Plugins/FX/Niagara/Source/NiagaraEditor/Private/NiagaraNodeEmitter.cpp
+++ b/Engine/Plugins/FX/Niagara/Source/NiagaraEditor/Private/NiagaraNodeEmitter.cpp
@@ -256,7 +256,7 @@ FText UNiagaraNodeEmitter::GetNameFromEmitter()
 	return FText();
 }
 
-void UNiagaraNodeEmitter::BuildParameterMapHistory(FNiagaraParameterMapHistoryBuilder& OutHistory, bool bRecursive)
+void UNiagaraNodeEmitter::BuildParameterMapHistory(FNiagaraParameterMapHistoryBuilder& OutHistory, bool bRecursive) const
 {
 	Super::BuildParameterMapHistory(OutHistory, bRecursive);
 
@@ -422,16 +422,16 @@ void UNiagaraNodeEmitter::Compile(FHlslNiagaraTranslator *Translator, TArray<int
 	}
 }
 
-void UNiagaraNodeEmitter::GatherExternalDependencyIDs(ENiagaraScriptUsage InMasterUsage, const FGuid& InMasterUsageId, TArray<FGuid>& InReferencedIDs, TArray<UObject*>& InReferencedObjs) const
+void UNiagaraNodeEmitter::GatherExternalDependencyIDs(ENiagaraScriptUsage InMasterUsage, const FGuid& InMasterUsageId, TArray<FNiagaraCompileHash>& InReferencedCompileHashes, TArray<FGuid>& InReferencedIDs, TArray<UObject*>& InReferencedObjs) const
 {
 	UNiagaraGraph* CalledGraph = GetCalledGraph();
 
 	if (CalledGraph)
 	{
 		ENiagaraScriptUsage TargetUsage = InMasterUsage == ENiagaraScriptUsage::SystemSpawnScript ? ENiagaraScriptUsage::EmitterSpawnScript : ENiagaraScriptUsage::EmitterUpdateScript;
-		InReferencedIDs.Add(CalledGraph->GetCompileID(TargetUsage, FGuid(0,0,0,0)));
+		InReferencedCompileHashes.Add(CalledGraph->GetCompileDataHash(TargetUsage, FGuid(0,0,0,0)));
 		InReferencedObjs.Add(CalledGraph);
-		CalledGraph->GatherExternalDependencyIDs(TargetUsage, FGuid(0, 0, 0, 0), InReferencedIDs, InReferencedObjs);
+		CalledGraph->GatherExternalDependencyIDs(TargetUsage, FGuid(0, 0, 0, 0), InReferencedCompileHashes, InReferencedIDs, InReferencedObjs);
 	}
 }
 
diff --git a/Engine/Plugins/FX/Niagara/Source/NiagaraEditor/Private/NiagaraNodeEmitter.h b/Engine/Plugins/FX/Niagara/Source/NiagaraEditor/Private/NiagaraNodeEmitter.h
--- a/Engine/Plugins/FX/Niagara/Source/NiagaraEditor/Private/NiagaraNodeEmitter.h
+++ b/Engine/Plugins/FX/Niagara/Source/NiagaraEditor/Private/NiagaraNodeEmitter.h
@@ -45,7 +45,7 @@ public:
 	virtual bool VerifyEditablePinName(const FText& InName, FText& OutErrorMessage, const UEdGraphPin* InGraphPinObj) const override;
 	virtual bool CommitEditablePinName(const FText& InName, UEdGraphPin* InGraphPinObj)  override;
 
-	virtual void BuildParameterMapHistory(FNiagaraParameterMapHistoryBuilder& OutHistory, bool bRecursive = true) override;
+	virtual void BuildParameterMapHistory(FNiagaraParameterMapHistoryBuilder& OutHistory, bool bRecursive = true) const override;
 
 	virtual FText GetNodeTitle(ENodeTitleType::Type TitleType) const;
 
@@ -57,7 +57,7 @@ public:
 	UNiagaraGraph* GetCalledGraph() const;
 	
 	virtual void Compile(FHlslNiagaraTranslator *Translator, TArray<int32>& Outputs) override;
-	virtual void GatherExternalDependencyIDs(ENiagaraScriptUsage InMasterUsage, const FGuid& InMasterUsageId, TArray<FGuid>& InReferencedIDs, TArray<UObject*>& InReferencedObjs) const override;
+	virtual void GatherExternalDependencyIDs(ENiagaraScriptUsage InMasterUsage, const FGuid& InMasterUsageId, TArray<FNiagaraCompileHash>& InReferencedCompileHashes, TArray<FGuid>& InReferencedIDs, TArray<UObject*>& InReferencedObjs) const override;
 
 	void SetCachedVariablesForCompilation(const FName& InUniqueName, UNiagaraGraph* InGraph, UNiagaraScriptSourceBase* InSource);
 
diff --git a/Engine/Plugins/FX/Niagara/Source/NiagaraEditor/Private/NiagaraNodeFunctionCall.cpp b/Engine/Plugins/FX/Niagara/Source/NiagaraEditor/Private/NiagaraNodeFunctionCall.cpp
--- a/Engine/Plugins/FX/Niagara/Source/NiagaraEditor/Private/NiagaraNodeFunctionCall.cpp
+++ b/Engine/Plugins/FX/Niagara/Source/NiagaraEditor/Private/NiagaraNodeFunctionCall.cpp
@@ -117,7 +117,7 @@ void UNiagaraNodeFunctionCall::AllocateDefaultPins()
 		Options.bFilterDuplicates = true;
 		Graph->FindInputNodes(InputNodes, Options);
 
-		bool bHasAdvancePins = false;
+		AdvancedPinDisplay = ENodeAdvancedPins::NoPins;
 		for (UNiagaraNodeInput* InputNode : InputNodes)
 		{
 			if (InputNode->IsExposed())
@@ -145,7 +145,7 @@ void UNiagaraNodeFunctionCall::AllocateDefaultPins()
 				if (InputNode->IsHidden())
 				{
 					NewPin->bAdvancedView = true;
-					bHasAdvancePins = true;
+					AdvancedPinDisplay = ENodeAdvancedPins::Hidden;
 				}
 				else
 				{
@@ -154,7 +154,20 @@ void UNiagaraNodeFunctionCall::AllocateDefaultPins()
 			}
 		}
 
-		AdvancedPinDisplay = bHasAdvancePins ? ENodeAdvancedPins::Hidden : ENodeAdvancedPins::NoPins;
+		TArray<FNiagaraVariable> SwitchNodeInputs = Graph->FindStaticSwitchInputs();
+		for (const FNiagaraVariable& Input : SwitchNodeInputs)
+		{
+			UEdGraphPin* NewPin = CreatePin(EGPD_Input, Schema->TypeDefinitionToPinType(Input.GetType()), Input.GetName());
+			NewPin->bDefaultValueIsIgnored = false;
+			NewPin->bNotConnectable = true;
+
+			FString PinDefaultValue;
+			if (Schema->TryGetPinDefaultValueFromNiagaraVariable(Input, PinDefaultValue))
+			{
+				NewPin->DefaultValue = PinDefaultValue;
+			}
+			NewPin->bDefaultValueIsIgnored = PropagatedStaticSwitchParameters.Contains(Input);
+		}
 
 		for (FNiagaraVariable& Output : Outputs)
 		{
@@ -531,10 +544,29 @@ bool UNiagaraNodeFunctionCall::RefreshFromExternalChanges()
 		}
 	}
 
+	// Go over the static switch parameters to set their propagation status on the pins
+	UNiagaraGraph* CalledGraph = GetCalledGraph();
+	if (CalledGraph)
+	{
+		TArray<UEdGraphPin*> InputPins;
+		GetInputPins(InputPins);
+		for (FNiagaraVariable InputVar : CalledGraph->FindStaticSwitchInputs())
+		{
+			for (UEdGraphPin* Pin : InputPins)
+			{
+				if (InputVar.GetName().IsEqual(Pin->GetFName()))
+				{
+					Pin->bDefaultValueIsIgnored = PropagatedStaticSwitchParameters.Contains(InputVar);
+					break;
+				}
+			}
+		}
+	}
+
 	if (bReload)
 	{
 		// TODO - Leverage code in reallocate pins to determine if any pins have changed...
-		ReallocatePins();
+		ReallocatePins(false);
 		return true;
 	}
 	else
@@ -560,7 +592,7 @@ void UNiagaraNodeFunctionCall::SubsumeExternalDependencies(TMap<const UObject*,
 	}
 }
 
-void UNiagaraNodeFunctionCall::GatherExternalDependencyIDs(ENiagaraScriptUsage InMasterUsage, const FGuid& InMasterUsageId, TArray<FGuid>& InReferencedIDs, TArray<UObject*>& InReferencedObjs) const
+void UNiagaraNodeFunctionCall::GatherExternalDependencyIDs(ENiagaraScriptUsage InMasterUsage, const FGuid& InMasterUsageId, TArray<FNiagaraCompileHash>& InReferencedCompileHashes, TArray<FGuid>& InReferencedIDs, TArray<UObject*>& InReferencedObjs) const
 {
 	if (FunctionScript && FunctionScript->GetOutermost() != this->GetOutermost())
 	{
@@ -572,13 +604,13 @@ void UNiagaraNodeFunctionCall::GatherExternalDependencyIDs(ENiagaraScriptUsage I
 		{
 			for (int32 i = (int32)ENiagaraScriptUsage::Function; i <= (int32)ENiagaraScriptUsage::DynamicInput; i++)
 			{
-				FGuid FoundGuid = FunctionGraph->GetCompileID((ENiagaraScriptUsage)i, FGuid(0, 0, 0, 0));
+				FGuid FoundGuid = FunctionGraph->ComputeCompileID((ENiagaraScriptUsage)i, FGuid(0, 0, 0, 0));
 				if (FoundGuid.IsValid())
 				{
 					InReferencedIDs.Add(FoundGuid);
 					InReferencedObjs.Add(FunctionGraph);
 
-					FunctionGraph->GatherExternalDependencyIDs((ENiagaraScriptUsage)i, FGuid(0, 0, 0, 0), InReferencedIDs, InReferencedObjs);
+					FunctionGraph->GatherExternalDependencyIDs((ENiagaraScriptUsage)i, FGuid(0, 0, 0, 0), InReferencedCompileHashes, InReferencedIDs, InReferencedObjs);
 				}
 			}
 		}
@@ -598,7 +630,7 @@ bool UNiagaraNodeFunctionCall::ScriptIsValid() const
 	return false;
 }
 
-void UNiagaraNodeFunctionCall::BuildParameterMapHistory(FNiagaraParameterMapHistoryBuilder& OutHistory, bool bRecursive)
+void UNiagaraNodeFunctionCall::BuildParameterMapHistory(FNiagaraParameterMapHistoryBuilder& OutHistory, bool bRecursive) const
 {
 	Super::BuildParameterMapHistory(OutHistory, bRecursive);
 	if (!IsNodeEnabled() && OutHistory.GetIgnoreDisabled())
@@ -700,6 +732,31 @@ UEdGraphPin* UNiagaraNodeFunctionCall::FindParameterMapDefaultValuePin(const FNa
 	return nullptr;
 }
 
+UEdGraphPin* UNiagaraNodeFunctionCall::FindStaticSwitchInputPin(const FName& VariableName) const
+{
+	UNiagaraGraph* CalledGraph = GetCalledGraph();
+	if (CalledGraph == nullptr)
+	{
+		return nullptr;
+	}
+	TArray<UEdGraphPin*> InputPins;
+	GetInputPins(InputPins);
+	for (FNiagaraVariable InputVar : CalledGraph->FindStaticSwitchInputs())
+	{
+		if (InputVar.GetName().IsEqual(VariableName))
+		{			
+			for (UEdGraphPin* Pin : InputPins)
+			{
+				if (VariableName.IsEqual(Pin->GetFName()))
+				{
+					return Pin;
+				}
+			}
+		}
+	}
+	return nullptr;
+}
+
 void UNiagaraNodeFunctionCall::SuggestName(FString SuggestedName)
 {
 	ComputeNodeName(SuggestedName);
diff --git a/Engine/Plugins/FX/Niagara/Source/NiagaraEditor/Private/NiagaraNodeInput.cpp b/Engine/Plugins/FX/Niagara/Source/NiagaraEditor/Private/NiagaraNodeInput.cpp
--- a/Engine/Plugins/FX/Niagara/Source/NiagaraEditor/Private/NiagaraNodeInput.cpp
+++ b/Engine/Plugins/FX/Niagara/Source/NiagaraEditor/Private/NiagaraNodeInput.cpp
@@ -72,7 +72,7 @@ void UNiagaraNodeInput::PostLoad()
 	}
 }
 
-void UNiagaraNodeInput::BuildParameterMapHistory(FNiagaraParameterMapHistoryBuilder& OutHistory, bool bRecursive)
+void UNiagaraNodeInput::BuildParameterMapHistory(FNiagaraParameterMapHistoryBuilder& OutHistory, bool bRecursive) const
 {
 	if (!IsNodeEnabled() && OutHistory.GetIgnoreDisabled())
 	{
diff --git a/Engine/Plugins/FX/Niagara/Source/NiagaraEditor/Private/NiagaraNodeParameterMapBase.cpp b/Engine/Plugins/FX/Niagara/Source/NiagaraEditor/Private/NiagaraNodeParameterMapBase.cpp
--- a/Engine/Plugins/FX/Niagara/Source/NiagaraEditor/Private/NiagaraNodeParameterMapBase.cpp
+++ b/Engine/Plugins/FX/Niagara/Source/NiagaraEditor/Private/NiagaraNodeParameterMapBase.cpp
@@ -48,7 +48,7 @@ TArray<FNiagaraParameterMapHistory> UNiagaraNodeParameterMapBase::GetParameterMa
 
 }
 
-TArray<FNiagaraParameterMapHistory> UNiagaraNodeParameterMapBase::GetParameterMaps(UNiagaraGraph* InGraph, FString EmitterNameOverride, const TArray<FNiagaraVariable>& EncounterableVariables)
+TArray<FNiagaraParameterMapHistory> UNiagaraNodeParameterMapBase::GetParameterMaps(const UNiagaraGraph* InGraph, FString EmitterNameOverride, const TArray<FNiagaraVariable>& EncounterableVariables)
 {
 	TArray<UNiagaraNodeOutput*> OutputNodes;
 	InGraph->FindOutputNodes(OutputNodes);
@@ -62,7 +62,7 @@ TArray<FNiagaraParameterMapHistory> UNiagaraNodeParameterMapBase::GetParameterMa
 	return OutputParameterMapHistories;
 }
 
-TArray<FNiagaraParameterMapHistory> UNiagaraNodeParameterMapBase::GetParameterMaps(UNiagaraNodeOutput* InGraphEnd, bool bLimitToOutputScriptType, FString EmitterNameOverride, const TArray<FNiagaraVariable>& EncounterableVariables)
+TArray<FNiagaraParameterMapHistory> UNiagaraNodeParameterMapBase::GetParameterMaps(const UNiagaraNodeOutput* InGraphEnd, bool bLimitToOutputScriptType, FString EmitterNameOverride, const TArray<FNiagaraVariable>& EncounterableVariables)
 {
 	const UEdGraphSchema_Niagara* Schema = Cast<UEdGraphSchema_Niagara>(InGraphEnd->GetSchema());
 
@@ -100,11 +100,10 @@ FText UNiagaraNodeParameterMapBase::GetPinDescriptionText(UEdGraphPin* Pin) cons
 	FNiagaraVariable Var = CastChecked<UEdGraphSchema_Niagara>(GetSchema())->PinToNiagaraVariable(Pin);
 
 	const UNiagaraGraph* Graph = GetNiagaraGraph();
-	const FNiagaraVariableMetaData* OldMetaData = Graph->GetMetaData(Var);
-	if (OldMetaData)
+	TOptional<FNiagaraVariableMetaData> MetaData = Graph->GetMetaData(Var);
+	if (MetaData.IsSet())
 	{
-		ensure(OldMetaData->ReferencerNodes.Contains(MakeWeakObjectPtr<UObject>(const_cast<UNiagaraNodeParameterMapBase*>(this))));
-		return OldMetaData->Description;
+		return MetaData->Description;
 	}
 	return FText::GetEmpty();
 }
@@ -119,30 +118,18 @@ void UNiagaraNodeParameterMapBase::PinDescriptionTextCommitted(const FText& Text
 		UE_LOG(LogNiagaraEditor, Error, TEXT("You cannot set the description for a Niagara internal constant \"%s\""),*Var.GetName().ToString());
 		return;
 	}
-	FNiagaraVariableMetaData* OldMetaData = Graph->GetMetaData(Var);
-	bool bSet = !Text.IsEmptyOrWhitespace();
-	if (OldMetaData && !OldMetaData->Description.IsEmptyOrWhitespace())
+
+	TOptional<FNiagaraVariableMetaData> OldMetaData = Graph->GetMetaData(Var);
+	FNiagaraVariableMetaData NewMetaData;
+	if (OldMetaData.IsSet())
 	{
-		bSet = true;
+		NewMetaData = OldMetaData.GetValue();
 	}
 
 	FScopedTransaction AddNewPinTransaction(LOCTEXT("Rename Pin Desc", "Changed variable description"));
-	Modify();
-	Pin->Modify();
-
-	if (OldMetaData)
-	{
-		ensure(OldMetaData->ReferencerNodes.Contains(TWeakObjectPtr<UObject>(this)));
-		Graph->Modify();
-		OldMetaData->Description = Text;
-	}
-	else
-	{
-		Graph->Modify();
-		FNiagaraVariableMetaData& NewMetaData = Graph->FindOrAddMetaData(Var);
-		NewMetaData.Description = Text;
-		NewMetaData.ReferencerNodes.Add(TWeakObjectPtr<UObject>(this));
-	}
+	NewMetaData.Description = Text;
+	Graph->Modify();
+	Graph->SetMetaData(Var, NewMetaData);
 }
 
 void UNiagaraNodeParameterMapBase::CollectAddPinActions(FGraphActionListBuilderBase& OutActions, bool& bOutCreateRemainingActions, UEdGraphPin* Pin)
@@ -186,8 +173,8 @@ void UNiagaraNodeParameterMapBase::GetPinHoverText(const UEdGraphPin& Pin, FStri
 			}
 
 			FNiagaraVariable Var = FNiagaraVariable(TypeDef, Pin.PinName);
-			const FNiagaraVariableMetaData* Metadata = NiagaraGraph->GetMetaData(Var);
-			if (Metadata)
+			TOptional<FNiagaraVariableMetaData> Metadata = NiagaraGraph->GetMetaData(Var);
+			if (Metadata.IsSet())
 			{
 				FText Desc = FText::Format(LOCTEXT("GetVarTooltip", "Name: \"{0}\"\nType: {1}\nDesc: {2}"), FText::FromName(Pin.PinName),
 					TypeDef.GetNameText(), Metadata->Description);
@@ -231,7 +218,7 @@ void UNiagaraNodeParameterMapBase::OnPinRenamed(UEdGraphPin* RenamedPin, const F
 	FNiagaraVariable Var(VarType, *OldName);
 
 	UNiagaraGraph* Graph = GetNiagaraGraph();
-	Graph->RenameParameter(Var, NewUniqueName, false /*bInNotifyGraphChanged*/); // Notify graph changed in child.
+	Graph->RenameParameter(Var, NewUniqueName);
 
 	if (RenamedPin == PinPendingRename)
 	{
diff --git a/Engine/Plugins/FX/Niagara/Source/NiagaraEditor/Private/NiagaraNodeParameterMapBase.h b/Engine/Plugins/FX/Niagara/Source/NiagaraEditor/Private/NiagaraNodeParameterMapBase.h
--- a/Engine/Plugins/FX/Niagara/Source/NiagaraEditor/Private/NiagaraNodeParameterMapBase.h
+++ b/Engine/Plugins/FX/Niagara/Source/NiagaraEditor/Private/NiagaraNodeParameterMapBase.h
@@ -20,8 +20,8 @@ public:
 	/** Traverse the graph looking for the history of the parameter map specified by the input pin. This will return the list of variables discovered, any per-variable warnings (type mismatches, etc)
 		encountered per variable, and an array of pins encountered in order of traversal outward from the input pin.
 	*/
-	static TArray<FNiagaraParameterMapHistory> GetParameterMaps(UNiagaraNodeOutput* InGraphEnd, bool bLimitToOutputScriptType = false, FString EmitterNameOverride = TEXT(""), const TArray<FNiagaraVariable>& EncounterableVariables = TArray<FNiagaraVariable>());
-	static TArray<FNiagaraParameterMapHistory> GetParameterMaps(UNiagaraGraph* InGraph, FString EmitterNameOverride = TEXT(""), const TArray<FNiagaraVariable>& EncounterableVariables = TArray<FNiagaraVariable>());
+	static TArray<FNiagaraParameterMapHistory> GetParameterMaps(const UNiagaraNodeOutput* InGraphEnd, bool bLimitToOutputScriptType = false, FString EmitterNameOverride = TEXT(""), const TArray<FNiagaraVariable>& EncounterableVariables = TArray<FNiagaraVariable>());
+	static TArray<FNiagaraParameterMapHistory> GetParameterMaps(const UNiagaraGraph* InGraph, FString EmitterNameOverride = TEXT(""), const TArray<FNiagaraVariable>& EncounterableVariables = TArray<FNiagaraVariable>());
 	static TArray<FNiagaraParameterMapHistory> GetParameterMaps(class UNiagaraScriptSourceBase* InSource, FString EmitterNameOverride = TEXT(""), const TArray<FNiagaraVariable>& EncounterableVariables = TArray<FNiagaraVariable>());
 
 	virtual bool AllowNiagaraTypeForAddPin(const FNiagaraTypeDefinition& InType) override;
diff --git a/Engine/Plugins/FX/Niagara/Source/NiagaraEditor/Private/NiagaraNodeParameterMapGet.cpp b/Engine/Plugins/FX/Niagara/Source/NiagaraEditor/Private/NiagaraNodeParameterMapGet.cpp
--- a/Engine/Plugins/FX/Niagara/Source/NiagaraEditor/Private/NiagaraNodeParameterMapGet.cpp
+++ b/Engine/Plugins/FX/Niagara/Source/NiagaraEditor/Private/NiagaraNodeParameterMapGet.cpp
@@ -171,32 +171,15 @@ void UNiagaraNodeParameterMapGet::OnNewTypedPinAdded(UEdGraphPin* NewPin)
 
 void UNiagaraNodeParameterMapGet::RemoveDynamicPin(UEdGraphPin* Pin)
 {
-	UEdGraphPin* DefaultPin = nullptr;
+	RemovePin(Pin);
 	if (Pin->Direction == EEdGraphPinDirection::EGPD_Output)
 	{
-		DefaultPin = GetDefaultPin(Pin);
-
-		FNiagaraVariable Var = CastChecked<UEdGraphSchema_Niagara>(GetSchema())->PinToNiagaraVariable(Pin);
-
-		UNiagaraGraph* Graph = GetNiagaraGraph();
-		FNiagaraVariableMetaData* OldMetaData = Graph->GetMetaData(Var);
-
-		if (OldMetaData)
+		UEdGraphPin* DefaultPin = GetDefaultPin(Pin);
+		if (DefaultPin != nullptr)
 		{
-			Graph->Modify();
-			OldMetaData->ReferencerNodes.Remove(TWeakObjectPtr<UObject>(this));
+			RemovePin(DefaultPin);
 		}
 	}
-
-	RemovePin(Pin);
-
-	if (DefaultPin != nullptr)
-	{
-		RemovePin(DefaultPin);
-	}
-
-	MarkNodeRequiresSynchronization(__FUNCTION__, true);
-	//GetGraph()->NotifyGraphChanged();
 }
 
 UEdGraphPin* UNiagaraNodeParameterMapGet::GetDefaultPin(UEdGraphPin* OutputPin) const
@@ -308,7 +291,7 @@ FText UNiagaraNodeParameterMapGet::GetNodeTitle(ENodeTitleType::Type TitleType)
 	return LOCTEXT("UNiagaraNodeParameterMapGetName", "Map Get");
 }
 
-void UNiagaraNodeParameterMapGet::BuildParameterMapHistory(FNiagaraParameterMapHistoryBuilder& OutHistory, bool bRecursive)
+void UNiagaraNodeParameterMapGet::BuildParameterMapHistory(FNiagaraParameterMapHistoryBuilder& OutHistory, bool bRecursive) const
 {
 	if (bRecursive)
 	{
@@ -431,7 +414,7 @@ bool UNiagaraNodeParameterMapGet::CommitEditablePinName(const FText& InName, UEd
 	return false;
 }
 
-void UNiagaraNodeParameterMapGet::GatherExternalDependencyIDs(ENiagaraScriptUsage InMasterUsage, const FGuid& InMasterUsageId, TArray<FGuid>& InReferencedIDs, TArray<UObject*>& InReferencedObjs) const
+void UNiagaraNodeParameterMapGet::GatherExternalDependencyIDs(ENiagaraScriptUsage InMasterUsage, const FGuid& InMasterUsageId, TArray<FNiagaraCompileHash>& InReferencedCompileHashes, TArray<FGuid>& InReferencedIDs, TArray<UObject*>& InReferencedObjs) const
 {
 	// If we are referencing any parameter collections, we need to register them here... might want to speeed this up in the future 
 	// by caching any parameter collections locally.
@@ -484,23 +467,17 @@ void UNiagaraNodeParameterMapGet::GetPinHoverText(const UEdGraphPin& Pin, FStrin
 				const UEdGraphPin* OutputPin = GetOutputPinForDefault(&Pin);
 				if (OutputPin)
 				{
-					TypeDef = Schema->PinToTypeDefinition(OutputPin);
-					FNiagaraVariable Var = FNiagaraVariable(TypeDef, OutputPin->PinName);
-					const FNiagaraVariableMetaData* Metadata = NiagaraGraph->GetMetaData(Var);
-					if (Metadata)
-					{
-						FText Desc = FText::Format(LOCTEXT("DefaultValueTooltip", "Default value for \"{0}\" if no other module has set it previously in the stack."), FText::FromName(OutputPin->PinName));
-						HoverTextOut = Desc.ToString();
-						return;
-					}
+					FText Desc = FText::Format(LOCTEXT("DefaultValueTooltip", "Default value for \"{0}\" if no other module has set it previously in the stack."), FText::FromName(OutputPin->PinName));
+					HoverTextOut = Desc.ToString();
+					return;
 				}
 				
 			}
 			else
 			{
 				FNiagaraVariable Var = FNiagaraVariable(TypeDef, Pin.PinName);
-				const FNiagaraVariableMetaData* Metadata = NiagaraGraph->GetMetaData(Var);
-				if (Metadata)
+				TOptional<FNiagaraVariableMetaData> Metadata = NiagaraGraph->GetMetaData(Var);
+				if (Metadata.IsSet())
 				{
 					FText Desc = FText::Format(LOCTEXT("GetVarTooltip", "Name: \"{0}\"\nType: {1}\nDesc: {2}"), FText::FromName(Pin.PinName),
 						TypeDef.GetNameText(), Metadata->Description);
diff --git a/Engine/Plugins/FX/Niagara/Source/NiagaraEditor/Private/NiagaraNodeParameterMapGet.h b/Engine/Plugins/FX/Niagara/Source/NiagaraEditor/Private/NiagaraNodeParameterMapGet.h
--- a/Engine/Plugins/FX/Niagara/Source/NiagaraEditor/Private/NiagaraNodeParameterMapGet.h
+++ b/Engine/Plugins/FX/Niagara/Source/NiagaraEditor/Private/NiagaraNodeParameterMapGet.h
@@ -30,7 +30,7 @@ public:
 
 	virtual void Compile(class FHlslNiagaraTranslator* Translator, TArray<int32>& Outputs);
 
-	virtual void BuildParameterMapHistory(FNiagaraParameterMapHistoryBuilder& OutHistory, bool bRecursive = true) override;
+	virtual void BuildParameterMapHistory(FNiagaraParameterMapHistoryBuilder& OutHistory, bool bRecursive = true) const override;
 
 	virtual FText GetNodeTitle(ENodeTitleType::Type TitleType) const;
 
@@ -39,7 +39,7 @@ public:
 
 	virtual void PostLoad() override;
 
-	void GatherExternalDependencyIDs(ENiagaraScriptUsage InMasterUsage, const FGuid& InMasterUsageId, TArray<FGuid>& InReferencedIDs, TArray<UObject*>& InReferencedObjs) const override;
+	void GatherExternalDependencyIDs(ENiagaraScriptUsage InMasterUsage, const FGuid& InMasterUsageId, TArray<FNiagaraCompileHash>& InReferencedCompileHashes, TArray<FGuid>& InReferencedIDs, TArray<UObject*>& InReferencedObjs) const override;
 	virtual void GetPinHoverText(const UEdGraphPin& Pin, FString& HoverTextOut) const override;
 
 protected:
diff --git a/Engine/Plugins/FX/Niagara/Source/NiagaraEditor/Private/NiagaraNodeParameterMapSet.cpp b/Engine/Plugins/FX/Niagara/Source/NiagaraEditor/Private/NiagaraNodeParameterMapSet.cpp
--- a/Engine/Plugins/FX/Niagara/Source/NiagaraEditor/Private/NiagaraNodeParameterMapSet.cpp
+++ b/Engine/Plugins/FX/Niagara/Source/NiagaraEditor/Private/NiagaraNodeParameterMapSet.cpp
@@ -191,7 +191,7 @@ FText UNiagaraNodeParameterMapSet::GetNodeTitle(ENodeTitleType::Type TitleType)
 }
 
 
-void UNiagaraNodeParameterMapSet::BuildParameterMapHistory(FNiagaraParameterMapHistoryBuilder& OutHistory, bool bRecursive)
+void UNiagaraNodeParameterMapSet::BuildParameterMapHistory(FNiagaraParameterMapHistoryBuilder& OutHistory, bool bRecursive) const
 {
 	const UEdGraphSchema_Niagara* Schema = GetDefault<UEdGraphSchema_Niagara>();
 	TArray<UEdGraphPin*> InputPins;
diff --git a/Engine/Plugins/FX/Niagara/Source/NiagaraEditor/Private/NiagaraNodeParameterMapSet.h b/Engine/Plugins/FX/Niagara/Source/NiagaraEditor/Private/NiagaraNodeParameterMapSet.h
--- a/Engine/Plugins/FX/Niagara/Source/NiagaraEditor/Private/NiagaraNodeParameterMapSet.h
+++ b/Engine/Plugins/FX/Niagara/Source/NiagaraEditor/Private/NiagaraNodeParameterMapSet.h
@@ -26,7 +26,7 @@ public:
 	virtual bool CommitEditablePinName(const FText& InName, UEdGraphPin* InGraphPinObj)  override;
 	virtual bool CancelEditablePinName(const FText& InName, UEdGraphPin* InGraphPinObj) override;
 
-	virtual void BuildParameterMapHistory(FNiagaraParameterMapHistoryBuilder& OutHistory, bool bRecursive = true) override;
+	virtual void BuildParameterMapHistory(FNiagaraParameterMapHistoryBuilder& OutHistory, bool bRecursive = true) const override;
 
 	virtual void Compile(class FHlslNiagaraTranslator* Translator, TArray<int32>& Outputs);
 
diff --git a/Engine/Plugins/FX/Niagara/Source/NiagaraEditor/Private/NiagaraNodeReadDataSet.cpp b/Engine/Plugins/FX/Niagara/Source/NiagaraEditor/Private/NiagaraNodeReadDataSet.cpp
--- a/Engine/Plugins/FX/Niagara/Source/NiagaraEditor/Private/NiagaraNodeReadDataSet.cpp
+++ b/Engine/Plugins/FX/Niagara/Source/NiagaraEditor/Private/NiagaraNodeReadDataSet.cpp
@@ -109,7 +109,7 @@ void UNiagaraNodeReadDataSet::Compile(class FHlslNiagaraTranslator* Translator,
 	Translator->ReadDataSet(DataSet, Variables, ENiagaraDataSetAccessMode::AppendConsume, Inputs[0], Outputs);
 }
 
-void UNiagaraNodeReadDataSet::BuildParameterMapHistory(FNiagaraParameterMapHistoryBuilder& OutHistory, bool bRecursive)
+void UNiagaraNodeReadDataSet::BuildParameterMapHistory(FNiagaraParameterMapHistoryBuilder& OutHistory, bool bRecursive) const
 {
 	if (bRecursive)
 	{
diff --git a/Engine/Plugins/FX/Niagara/Source/NiagaraEditor/Private/NiagaraNodeReroute.cpp b/Engine/Plugins/FX/Niagara/Source/NiagaraEditor/Private/NiagaraNodeReroute.cpp
--- a/Engine/Plugins/FX/Niagara/Source/NiagaraEditor/Private/NiagaraNodeReroute.cpp
+++ b/Engine/Plugins/FX/Niagara/Source/NiagaraEditor/Private/NiagaraNodeReroute.cpp
@@ -117,7 +117,7 @@ void UNiagaraNodeReroute::PinConnectionListChanged(UEdGraphPin* Pin)
 	Super::PinConnectionListChanged(Pin);
 }
 
-void UNiagaraNodeReroute::BuildParameterMapHistory(FNiagaraParameterMapHistoryBuilder& OutHistory, bool bRecursive)
+void UNiagaraNodeReroute::BuildParameterMapHistory(FNiagaraParameterMapHistoryBuilder& OutHistory, bool bRecursive) const
 {
 	// Should never hit here...
 	check(false);
diff --git a/Engine/Plugins/FX/Niagara/Source/NiagaraEditor/Private/NiagaraNodeReroute.h b/Engine/Plugins/FX/Niagara/Source/NiagaraEditor/Private/NiagaraNodeReroute.h
--- a/Engine/Plugins/FX/Niagara/Source/NiagaraEditor/Private/NiagaraNodeReroute.h
+++ b/Engine/Plugins/FX/Niagara/Source/NiagaraEditor/Private/NiagaraNodeReroute.h
@@ -33,7 +33,7 @@ public:
 	virtual void Compile(class FHlslNiagaraTranslator* Translator, TArray<int32>& Outputs) override;
 	virtual bool RefreshFromExternalChanges() override;
 	virtual void PinConnectionListChanged(UEdGraphPin* Pin) override;
-	void BuildParameterMapHistory(FNiagaraParameterMapHistoryBuilder& OutHistory, bool bRecursive) override;
+	void BuildParameterMapHistory(FNiagaraParameterMapHistoryBuilder& OutHistory, bool bRecursive) const override;
 	virtual ENiagaraNumericOutputTypeSelectionMode GetNumericOutputTypeSelectionMode() const;
 
 	/** Trace to an output pin that is not a reroute node output pin. If the reroute nodes resulted in a dead end and no output pin was found return nullptr.*/
diff --git a/Engine/Plugins/FX/Niagara/Source/NiagaraEditor/Private/NiagaraNodeSimTargetSelector.cpp b/Engine/Plugins/FX/Niagara/Source/NiagaraEditor/Private/NiagaraNodeSimTargetSelector.cpp
--- a/Engine/Plugins/FX/Niagara/Source/NiagaraEditor/Private/NiagaraNodeSimTargetSelector.cpp
+++ b/Engine/Plugins/FX/Niagara/Source/NiagaraEditor/Private/NiagaraNodeSimTargetSelector.cpp
@@ -104,7 +104,7 @@ UEdGraphPin* UNiagaraNodeSimTargetSelector::GetPassThroughPin(const UEdGraphPin*
 	return nullptr;
 }
 
-void UNiagaraNodeSimTargetSelector::BuildParameterMapHistory(FNiagaraParameterMapHistoryBuilder& OutHistory, bool bRecursive)
+void UNiagaraNodeSimTargetSelector::BuildParameterMapHistory(FNiagaraParameterMapHistoryBuilder& OutHistory, bool bRecursive) const
 {
 	UNiagaraNode::BuildParameterMapHistory(OutHistory, bRecursive);
 }
diff --git a/Engine/Plugins/FX/Niagara/Source/NiagaraEditor/Private/NiagaraNodeSimTargetSelector.h b/Engine/Plugins/FX/Niagara/Source/NiagaraEditor/Private/NiagaraNodeSimTargetSelector.h
--- a/Engine/Plugins/FX/Niagara/Source/NiagaraEditor/Private/NiagaraNodeSimTargetSelector.h
+++ b/Engine/Plugins/FX/Niagara/Source/NiagaraEditor/Private/NiagaraNodeSimTargetSelector.h
@@ -19,7 +19,7 @@ public:
 
 	//~ Begin UNiagaraNode Interface
 	virtual void Compile(class FHlslNiagaraTranslator* Translator, TArray<int32>& Outputs) override;
-	void BuildParameterMapHistory(FNiagaraParameterMapHistoryBuilder& OutHistory, bool bRecursive) override;
+	void BuildParameterMapHistory(FNiagaraParameterMapHistoryBuilder& OutHistory, bool bRecursive) const override;
 	virtual UEdGraphPin* GetPassThroughPin(const UEdGraphPin* LocallyOwnedOutputPin, ENiagaraScriptUsage MasterUsage) const override;
 	//~ End UNiagaraNode Interface
 
