Copying //UE4/Dev-Niagara to Dev-Main (//UE4/Dev-Main)

diff --git a/Engine/Plugins/FX/Niagara/Source/NiagaraEditor/Private/NiagaraNodeStaticSwitch.cpp b/Engine/Plugins/FX/Niagara/Source/NiagaraEditor/Private/NiagaraNodeStaticSwitch.cpp
new file mode 100644
--- /dev/null
+++ b/Engine/Plugins/FX/Niagara/Source/NiagaraEditor/Private/NiagaraNodeStaticSwitch.cpp
@@ -0,0 +1,289 @@
+// Copyright 1998-2019 Epic Games, Inc. All Rights Reserved.
+
+#include "NiagaraNodeStaticSwitch.h"
+#include "NiagaraEditorUtilities.h"
+#include "NiagaraHlslTranslator.h"
+
+#define LOCTEXT_NAMESPACE "NiagaraNodeStaticSwitch"
+
+UNiagaraNodeStaticSwitch::UNiagaraNodeStaticSwitch(const FObjectInitializer& ObjectInitializer)
+	: Super(ObjectInitializer), InputParameterName(FName(TEXT("Undefined parameter name"))), IsValueSet(false), SwitchValue(0)
+{
+}
+
+FNiagaraTypeDefinition UNiagaraNodeStaticSwitch::GetInputType() const
+{
+	if (SwitchTypeData.SwitchType == ENiagaraStaticSwitchType::Bool)
+	{
+		return FNiagaraTypeDefinition::GetBoolDef();
+	}
+	else if (SwitchTypeData.SwitchType == ENiagaraStaticSwitchType::Integer)
+	{
+		return FNiagaraTypeDefinition::GetIntDef();
+	}
+	else if (SwitchTypeData.SwitchType == ENiagaraStaticSwitchType::Enum)
+	{
+		return FNiagaraTypeDefinition(SwitchTypeData.Enum);
+	}
+	return FNiagaraTypeDefinition();
+}
+
+void UNiagaraNodeStaticSwitch::AllocateDefaultPins()
+{
+	const UEdGraphSchema_Niagara* Schema = GetDefault<UEdGraphSchema_Niagara>();
+
+	// create the input pins which differ in count depending on the switch type
+	if (SwitchTypeData.SwitchType == ENiagaraStaticSwitchType::Bool)
+	{
+		for (FNiagaraVariable& Var : OutputVars)
+		{
+			FEdGraphPinType PinType = Schema->TypeDefinitionToPinType(Var.GetType());
+			CreatePin(EGPD_Input, PinType, *(Var.GetName().ToString() + TEXT(" if true")));
+		}
+		for (FNiagaraVariable& Var : OutputVars)
+		{
+			FEdGraphPinType PinType = Schema->TypeDefinitionToPinType(Var.GetType());
+			CreatePin(EGPD_Input, PinType, *(Var.GetName().ToString() + TEXT(" if false")));
+		}
+	}
+	else if (SwitchTypeData.SwitchType == ENiagaraStaticSwitchType::Integer)
+	{
+		// the int range is inclusive, so we don't skip the last number
+		for (int32 i = 0; i <= SwitchTypeData.MaxIntCount; i++)
+		{
+			for (FNiagaraVariable& Var : OutputVars)
+			{
+				FEdGraphPinType PinType = Schema->TypeDefinitionToPinType(Var.GetType());
+				CreatePin(EGPD_Input, PinType, *(Var.GetName().ToString() + TEXT(" if ") + FString::FromInt(i)));
+			}
+		}
+	}
+	else if (SwitchTypeData.SwitchType == ENiagaraStaticSwitchType::Enum)
+	{
+		// The last enum value is a special "max" value that we do not want to display, so we skip it
+		for (int32 i = 0; i < SwitchTypeData.Enum->NumEnums() - 1; i++)
+		{
+			for (FNiagaraVariable& Var : OutputVars)
+			{
+				FEdGraphPinType PinType = Schema->TypeDefinitionToPinType(Var.GetType());
+				FText EnumName = SwitchTypeData.Enum->GetDisplayNameTextByIndex(i);
+				CreatePin(EGPD_Input, PinType, *(Var.GetName().ToString() + TEXT(" if ") + EnumName.ToString()));
+			}
+		}
+	}
+	
+
+	// create the output pins
+	for (int32 Index = 0; Index < OutputVars.Num(); Index++)
+	{
+		const FNiagaraVariable& Var = OutputVars[Index];
+		UEdGraphPin* NewPin = CreatePin(EGPD_Output, Schema->TypeDefinitionToPinType(Var.GetType()), Var.GetName());
+		NewPin->PersistentGuid = OutputVarGuids[Index];
+	}
+
+	CreateAddPin(EGPD_Output);
+}
+
+void UNiagaraNodeStaticSwitch::InsertInputPinsFor(const FNiagaraVariable& Var)
+{
+	const UEdGraphSchema_Niagara* Schema = GetDefault<UEdGraphSchema_Niagara>();
+
+	int32 OptionsCount = 0;
+	if (SwitchTypeData.SwitchType == ENiagaraStaticSwitchType::Bool)
+	{
+		OptionsCount = 2;
+	}
+	else if (SwitchTypeData.SwitchType == ENiagaraStaticSwitchType::Integer)
+	{
+		// +1 because the range is inclusive
+		OptionsCount = SwitchTypeData.MaxIntCount + 1;
+	}
+	else if (SwitchTypeData.SwitchType == ENiagaraStaticSwitchType::Enum)
+	{
+		// -1 because the last enum value is a special "max" value
+		OptionsCount = SwitchTypeData.Enum->NumEnums() - 1;
+	}
+
+	TArray<UEdGraphPin*> OldPins(Pins);
+	Pins.Reset(OldPins.Num() + OptionsCount);
+
+	// Create the inputs for each option
+	for (int64 i = 0; i < OptionsCount; i++)
+	{
+		// Add the previous input pins
+		for (int32 k = 0; k < OutputVars.Num() - 1; k++)
+		{
+			Pins.Add(OldPins[k]);
+		}
+		OldPins.RemoveAt(0, OutputVars.Num() - 1);
+
+		// Add the new input pin		
+		FString PathSuffix = TEXT(" if ");
+		if (SwitchTypeData.SwitchType == ENiagaraStaticSwitchType::Bool)
+		{
+			PathSuffix += i == 0 ? TEXT("true") : TEXT("false");
+		}
+		else if (SwitchTypeData.SwitchType == ENiagaraStaticSwitchType::Integer)
+		{
+			PathSuffix += FString::FromInt(i);
+		}
+		else if (SwitchTypeData.SwitchType == ENiagaraStaticSwitchType::Enum)
+		{
+			FText EnumName = SwitchTypeData.Enum->GetDisplayNameTextByIndex(i);
+			PathSuffix += EnumName.ToString();
+		}
+		CreatePin(EGPD_Input, Schema->TypeDefinitionToPinType(Var.GetType()), *(Var.GetName().ToString() + PathSuffix));
+	}
+
+	// Move the rest of the old pins over
+	Pins.Append(OldPins);
+}
+
+bool UNiagaraNodeStaticSwitch::AllowNiagaraTypeForAddPin(const FNiagaraTypeDefinition& InType)
+{
+	// explicitly allow parameter maps and numeric types
+	 return InType.GetScriptStruct() != nullptr;
+}
+
+bool UNiagaraNodeStaticSwitch::GetVarIndex(FHlslNiagaraTranslator* Translator, int32 InputPinCount, int32& VarIndexOut) const
+{	
+	return GetVarIndex(Translator, InputPinCount, SwitchValue, VarIndexOut);
+}
+
+bool UNiagaraNodeStaticSwitch::GetVarIndex(class FHlslNiagaraTranslator* Translator, int32 InputPinCount, int32 Value, int32& VarIndexOut) const
+{
+	bool Success = false;
+	if (SwitchTypeData.SwitchType == ENiagaraStaticSwitchType::Bool)
+	{
+		VarIndexOut = Value ? 0 : InputPinCount / 2;
+		Success = true;
+	}
+	else if (SwitchTypeData.SwitchType == ENiagaraStaticSwitchType::Integer)
+	{
+		int32 MaxValue = SwitchTypeData.MaxIntCount;
+		if (MaxValue >= 0)
+		{
+			if (Translator && (Value > MaxValue || Value < 0))
+			{
+				Translator->Warning(FText::Format(LOCTEXT("InvalidStaticSwitchIntValue", "The supplied int value {0} is outside the bounds for the static switch."), FText::FromString(FString::FromInt(Value))), this, nullptr);
+			}
+			VarIndexOut = FMath::Clamp(Value, 0, MaxValue) * (InputPinCount / MaxValue);
+			Success = true;
+		}
+		else if (Translator)
+		{
+			Translator->Error(FText::Format(LOCTEXT("InvalidSwitchMaxIntValue", "Invalid max int value {0} for static switch."), FText::FromString(FString::FromInt(Value))), this, nullptr);
+		}
+	}
+	else if (SwitchTypeData.SwitchType == ENiagaraStaticSwitchType::Enum && SwitchTypeData.Enum)
+	{
+		int32 MaxValue = SwitchTypeData.Enum->NumEnums() - 1;
+		if (MaxValue > 0)
+		{
+			if (Value <= MaxValue && Value >= 0)
+			{
+				VarIndexOut = Value * (InputPinCount / MaxValue);
+				Success = true;
+			}
+			else if (Translator)
+			{
+				Translator->Error(FText::Format(LOCTEXT("InvalidSwitchEnumIndex", "Invalid static switch value \"{0}\" for enum value index."), FText::FromString(FString::FromInt(SwitchValue))), this, nullptr);
+			}
+		}
+	}
+	else if (Translator)
+	{
+		Translator->Error(LOCTEXT("InvalidSwitchType", "Invalid static switch type."), this, nullptr);
+	}
+	return Success;
+}
+
+void UNiagaraNodeStaticSwitch::Compile(FHlslNiagaraTranslator* Translator, TArray<int32>& Outputs)
+{
+	UNiagaraNode::Compile(Translator, Outputs);
+}
+
+bool UNiagaraNodeStaticSwitch::SubstituteCompiledPin(FHlslNiagaraTranslator* Translator, UEdGraphPin** LocallyOwnedPin)
+{
+	// if we compile the standalone module or function we don't have any valid input yet, so we just take the first option to satisfy the compiler
+	ENiagaraScriptUsage TargetUsage = Translator->GetTargetUsage();
+	bool IsDryRun = TargetUsage == ENiagaraScriptUsage::Module || TargetUsage == ENiagaraScriptUsage::Function;
+	SwitchValue = IsDryRun ? 0 : SwitchValue;
+	if (!IsValueSet && !IsDryRun)
+	{
+		FText ErrorMessage = FText::Format(LOCTEXT("MissingSwitchValue", "The input parameter \"{0}\" is not set to a constant value for the static switch node."), FText::FromString(InputParameterName.ToString()));
+		Translator->Error(ErrorMessage, this, nullptr);
+		return false;
+	}
+
+	TArray<UEdGraphPin*> InputPins;
+	GetInputPins(InputPins);
+	TArray<UEdGraphPin*> OutputPins;
+	GetOutputPins(OutputPins);
+
+	for (int i = 0; i < OutputPins.Num(); i++)
+	{
+		UEdGraphPin* OutPin = OutputPins[i];
+		int32 VarIdx;
+		if (OutPin == *LocallyOwnedPin && GetVarIndex(Translator, InputPins.Num(), IsDryRun ? 0 : SwitchValue, VarIdx))
+		{
+			UEdGraphPin* InputPin = InputPins[VarIdx + i];
+			if (InputPin->LinkedTo.Num() == 1)
+			{
+				*LocallyOwnedPin = GetTracedOutputPin(InputPin->LinkedTo[0]);
+				return true;
+			}
+			else
+			{
+				*LocallyOwnedPin = InputPin;
+			}
+			return true;
+		}
+	}
+	return false;
+}
+
+UEdGraphPin* UNiagaraNodeStaticSwitch::GetTracedOutputPin(UEdGraphPin* LocallyOwnedOutputPin) const
+{
+	TArray<UEdGraphPin*> InputPins;
+	GetInputPins(InputPins);
+	TArray<UEdGraphPin*> OutputPins;
+	GetOutputPins(OutputPins);
+
+	for (int i = 0; i < OutputPins.Num(); i++)
+	{
+		UEdGraphPin* OutPin = OutputPins[i];
+		if (IsAddPin(OutPin))
+		{
+			continue;
+		}
+		int32 VarIdx;
+		if (OutPin == LocallyOwnedOutputPin && GetVarIndex(nullptr, InputPins.Num(), IsValueSet ? SwitchValue : 0, VarIdx))
+		{
+			UEdGraphPin* InputPin = InputPins[VarIdx + i];
+			if (InputPin->LinkedTo.Num() == 1)
+			{
+				return UNiagaraNode::TraceOutputPin(InputPin->LinkedTo[0]);
+			}
+		}
+	}
+	return LocallyOwnedOutputPin;
+}
+
+FText UNiagaraNodeStaticSwitch::GetTooltipText() const
+{
+	return LOCTEXT("NiagaraStaticSwitchNodeTooltip", "This is a compile-time switch that selects one branch to compile based on an input parameter.");
+}
+
+FText UNiagaraNodeStaticSwitch::GetNodeTitle(ENodeTitleType::Type TitleType) const
+{
+	return FText::FormatOrdered(LOCTEXT("StaticSwitchTitle", "Static Switch ({0})"), FText::FromName(InputParameterName));
+}
+
+FLinearColor UNiagaraNodeStaticSwitch::GetNodeTitleColor() const
+{
+	const UEdGraphSchema_Niagara* Schema = CastChecked<UEdGraphSchema_Niagara>(GetSchema());
+	return Schema->NodeTitleColor_Constant;
+}
+
+#undef LOCTEXT_NAMESPACE
diff --git a/Engine/Plugins/FX/Niagara/Source/NiagaraEditor/Private/NiagaraNodeStaticSwitch.h b/Engine/Plugins/FX/Niagara/Source/NiagaraEditor/Private/NiagaraNodeStaticSwitch.h
new file mode 100644
--- /dev/null
+++ b/Engine/Plugins/FX/Niagara/Source/NiagaraEditor/Private/NiagaraNodeStaticSwitch.h
@@ -0,0 +1,85 @@
+// Copyright 1998-2019 Epic Games, Inc. All Rights Reserved.
+
+#pragma once
+#include "EdGraphSchema_Niagara.h"
+#include "NiagaraNodeUsageSelector.h"
+#include "NiagaraNodeStaticSwitch.generated.h"
+
+UENUM()
+enum class ENiagaraStaticSwitchType : uint8
+{
+	Bool,
+	Integer,
+	Enum
+};
+
+USTRUCT()
+struct FStaticSwitchTypeData
+{
+	GENERATED_USTRUCT_BODY()
+
+	/** This determines how the switch value is interpreted */
+	UPROPERTY()
+	ENiagaraStaticSwitchType SwitchType;
+
+	/** If the type is int, this sets the upper range of input values the switch supports (so pins = MaxIntCount + 1) */
+	UPROPERTY()
+	int32 MaxIntCount;
+
+	/** If the type is enum, this is the enum being switched on, otherwise it holds no sensible value */
+	UPROPERTY()
+	UEnum* Enum;
+
+	FStaticSwitchTypeData() : SwitchType(ENiagaraStaticSwitchType::Bool), MaxIntCount(1), Enum(nullptr)
+	{ }
+};
+
+UCLASS(MinimalAPI)
+class UNiagaraNodeStaticSwitch : public UNiagaraNodeUsageSelector
+{
+	GENERATED_UCLASS_BODY()
+
+public:
+
+	UPROPERTY(EditAnywhere, Category="Static Switch")
+	FName InputParameterName;
+	
+	UPROPERTY()
+	FStaticSwitchTypeData SwitchTypeData;
+
+	/** If true then the current SwitchValue should be used for compilation, otherwise the node is not yet connected to a constant value */
+	UPROPERTY(Transient)
+	bool IsValueSet;
+
+	/** The current value that the node is evaluated with. For bool 0 is false, for enums the value is the index of the enum value. */
+	UPROPERTY(Transient)
+	int32 SwitchValue;
+
+	FNiagaraTypeDefinition GetInputType() const;
+
+	//~ Begin EdGraphNode Interface
+	virtual void AllocateDefaultPins() override;
+	virtual FText GetTooltipText() const override;
+	virtual FText GetNodeTitle(ENodeTitleType::Type TitleType) const override;
+	virtual FLinearColor GetNodeTitleColor() const override;
+	virtual bool IsCompilerRelevant() const override { return false; }
+	//~ End EdGraphNode Interface
+
+	//~ Begin UNiagaraNode Interface
+	virtual void Compile(class FHlslNiagaraTranslator* Translator, TArray<int32>& Outputs) override;
+	virtual bool SubstituteCompiledPin(FHlslNiagaraTranslator* Translator, UEdGraphPin** LocallyOwnedPin) override;
+	virtual UEdGraphPin* GetTracedOutputPin(UEdGraphPin* LocallyOwnedOutputPin) const override;
+	//~ End UNiagaraNode Interface
+
+protected:
+	//~ Begin UNiagaraNodeUsageSelector Interface
+	virtual void InsertInputPinsFor(const FNiagaraVariable& Var) override;
+	virtual bool AllowNiagaraTypeForAddPin(const FNiagaraTypeDefinition& InType) override;
+	//~ End UNiagaraNodeUsageSelector Interface
+
+private:
+	/** This finds the first valid input pin index for the current switch value, returns false if no value can be found */
+	bool GetVarIndex(class FHlslNiagaraTranslator* Translator, int32 InputPinCount, int32& VarIndexOut) const;
+
+	bool GetVarIndex(class FHlslNiagaraTranslator* Translator, int32 InputPinCount, int32 Value, int32& VarIndexOut) const;
+};
diff --git a/Engine/Plugins/FX/Niagara/Source/NiagaraEditor/Private/NiagaraNodeUsageSelector.cpp b/Engine/Plugins/FX/Niagara/Source/NiagaraEditor/Private/NiagaraNodeUsageSelector.cpp
--- a/Engine/Plugins/FX/Niagara/Source/NiagaraEditor/Private/NiagaraNodeUsageSelector.cpp
+++ b/Engine/Plugins/FX/Niagara/Source/NiagaraEditor/Private/NiagaraNodeUsageSelector.cpp
@@ -191,7 +191,7 @@ void UNiagaraNodeUsageSelector::AppendFunctionAliasForContext(const FNiagaraGrap
 	}
 }
 
-void UNiagaraNodeUsageSelector::BuildParameterMapHistory(FNiagaraParameterMapHistoryBuilder& OutHistory, bool bRecursive)
+void UNiagaraNodeUsageSelector::BuildParameterMapHistory(FNiagaraParameterMapHistoryBuilder& OutHistory, bool bRecursive) const
 {
 	const UEdGraphSchema_Niagara* Schema = CastChecked<UEdGraphSchema_Niagara>(GetSchema());
 
diff --git a/Engine/Plugins/FX/Niagara/Source/NiagaraEditor/Private/NiagaraNodeWriteDataSet.cpp b/Engine/Plugins/FX/Niagara/Source/NiagaraEditor/Private/NiagaraNodeWriteDataSet.cpp
--- a/Engine/Plugins/FX/Niagara/Source/NiagaraEditor/Private/NiagaraNodeWriteDataSet.cpp
+++ b/Engine/Plugins/FX/Niagara/Source/NiagaraEditor/Private/NiagaraNodeWriteDataSet.cpp
@@ -113,7 +113,7 @@ void UNiagaraNodeWriteDataSet::PostLoad()
 	}
 }
 
-void UNiagaraNodeWriteDataSet::BuildParameterMapHistory(FNiagaraParameterMapHistoryBuilder& OutHistory, bool bRecursive)
+void UNiagaraNodeWriteDataSet::BuildParameterMapHistory(FNiagaraParameterMapHistoryBuilder& OutHistory, bool bRecursive) const
 {
 	if (bRecursive)
 	{
diff --git a/Engine/Plugins/FX/Niagara/Source/NiagaraEditor/Private/NiagaraParameterMapHistory.cpp b/Engine/Plugins/FX/Niagara/Source/NiagaraEditor/Private/NiagaraParameterMapHistory.cpp
--- a/Engine/Plugins/FX/Niagara/Source/NiagaraEditor/Private/NiagaraParameterMapHistory.cpp
+++ b/Engine/Plugins/FX/Niagara/Source/NiagaraEditor/Private/NiagaraParameterMapHistory.cpp
@@ -124,7 +124,7 @@ int32 FNiagaraParameterMapHistory::RegisterParameterMapPin(const UEdGraphPin* Pi
 	return RetIdx;
 }
 
-uint32 FNiagaraParameterMapHistory::BeginNodeVisitation(UNiagaraNode* Node)
+uint32 FNiagaraParameterMapHistory::BeginNodeVisitation(const UNiagaraNode* Node)
 {
 	uint32 AddedIndex = MapNodeVisitations.Add(Node);
 	MapNodeVariableMetaData.Add(TTuple<uint32, uint32>(Variables.Num(), 0));
@@ -479,48 +479,6 @@ FNiagaraVariable FNiagaraParameterMapHistory::GetSourceForInitialValue(const FNi
 	return Var;
 }
 
-const FNiagaraVariableMetaData* FNiagaraParameterMapHistory::GetMetaData(int32 VarIdx) const
-{
-	if (PerVariableWriteHistory[VarIdx].Num() > 0)
-	{
-		const UEdGraphPin* Pin = PerVariableWriteHistory[VarIdx][0];
-		if (Pin != nullptr)
-		{
-			return CastChecked<UNiagaraNode>(Pin->GetOwningNode())->GetNiagaraGraph()->GetMetaData(Variables[VarIdx]);
-		}
-	}
-	else if (PerVariableReadHistory[VarIdx].Num() > 0)
-	{
-		const UEdGraphPin* Pin = PerVariableReadHistory[VarIdx][0].Key;
-		if (Pin != nullptr)
-		{
-			return CastChecked<UNiagaraNode>(Pin->GetOwningNode())->GetNiagaraGraph()->GetMetaData(Variables[VarIdx]);
-		}
-	}
-	return nullptr;
-}
-
-FNiagaraVariableMetaData* FNiagaraParameterMapHistory::GetMetaData(int32 VarIdx)
-{
-	if (PerVariableWriteHistory[VarIdx].Num() > 0)
-	{
-		const UEdGraphPin* Pin = PerVariableWriteHistory[VarIdx][0];
-		if (Pin != nullptr)
-		{
-			return CastChecked<UNiagaraNode>(Pin->GetOwningNode())->GetNiagaraGraph()->GetMetaData(Variables[VarIdx]);
-		}
-	}
-	else if (PerVariableReadHistory[VarIdx].Num() > 0)
-	{
-		const UEdGraphPin* Pin = PerVariableReadHistory[VarIdx][0].Key;
-		if (Pin != nullptr)
-		{
-			return CastChecked<UNiagaraNode>(Pin->GetOwningNode())->GetNiagaraGraph()->GetMetaData(Variables[VarIdx]);
-		}
-	}
-	return nullptr;
-}
-
 bool FNiagaraParameterMapHistory::IsPrimaryDataSetOutput(const FNiagaraVariable& InVar, const UNiagaraScript* InScript, bool bAllowDataInterfaces) const
 {
 	return IsPrimaryDataSetOutput(InVar, InScript->GetUsage(), bAllowDataInterfaces);
@@ -690,7 +648,7 @@ FNiagaraParameterMapHistoryBuilder::FNiagaraParameterMapHistoryBuilder()
 	FilterScriptType = ENiagaraScriptUsage::Function;
 }
 
-void FNiagaraParameterMapHistoryBuilder::BuildParameterMaps(UNiagaraNodeOutput* OutputNode, bool bRecursive)
+void FNiagaraParameterMapHistoryBuilder::BuildParameterMaps(const UNiagaraNodeOutput* OutputNode, bool bRecursive)
 {
 	RelevantScriptUsageContext.Emplace(OutputNode->GetUsage());
 	OutputNode->BuildParameterMapHistory(*this, bRecursive);
@@ -737,7 +695,7 @@ int32 FNiagaraParameterMapHistoryBuilder::CreateParameterMap()
 	return RetValue;
 }
 
-uint32 FNiagaraParameterMapHistoryBuilder::BeginNodeVisitation(int32 WhichParameterMap, class UNiagaraNode* Node)
+uint32 FNiagaraParameterMapHistoryBuilder::BeginNodeVisitation(int32 WhichParameterMap, const class UNiagaraNode* Node)
 {
 	if (WhichParameterMap != INDEX_NONE)
 	{
@@ -812,8 +770,8 @@ int32 FNiagaraParameterMapHistoryBuilder::FindMatchingParameterMapFromContextInp
 	{
 		return INDEX_NONE;
 	}
-	UNiagaraNode* Node = CallingContext.Last();
-	TArray<UEdGraphPin*> Inputs;
+	const UNiagaraNode* Node = CallingContext.Last();
+	TArray<const UEdGraphPin*> Inputs;
 	Node->GetInputPins(Inputs);
 	const UEdGraphSchema_Niagara* Schema = CastChecked<UEdGraphSchema_Niagara>(Node->GetSchema());
 
@@ -833,7 +791,7 @@ int32 FNiagaraParameterMapHistoryBuilder::FindMatchingParameterMapFromContextInp
 				else
 				{
 					FString ScriptUsageDisplayName;
-					UNiagaraNodeOutput* ContextOutputNode = FNiagaraStackGraphUtilities::GetEmitterOutputNodeForStackNode(*Node);
+					const UNiagaraNodeOutput* ContextOutputNode = FNiagaraStackGraphUtilities::GetEmitterOutputNodeForStackNode(*Node);
 					if (ContextOutputNode != nullptr)
 					{
 						UEnum* NiagaraScriptUsageEnum = FindObjectChecked<UEnum>(ANY_PACKAGE, TEXT("ENiagaraScriptUsage"), true);
@@ -917,7 +875,7 @@ void FNiagaraParameterMapHistoryBuilder::EndTranslation(const FString& EmitterUn
 	EmitterNameContextStack.Reset();
 }
 
-UNiagaraNode* FNiagaraParameterMapHistoryBuilder::GetCallingContext() const
+const UNiagaraNode* FNiagaraParameterMapHistoryBuilder::GetCallingContext() const
 {
 	if (CallingContext.Num() == 0)
 	{
@@ -955,7 +913,7 @@ bool FNiagaraParameterMapHistoryBuilder::InTopLevelFunctionCall(ENiagaraScriptUs
 }
 
 
-void FNiagaraParameterMapHistoryBuilder::EnterFunction(const FString& InNodeName, const UNiagaraScript* InScript, UNiagaraNode* Node)
+void FNiagaraParameterMapHistoryBuilder::EnterFunction(const FString& InNodeName, const UNiagaraScript* InScript, const UNiagaraNode* Node)
 {
 	if (InScript != nullptr )
 	{
@@ -972,7 +930,7 @@ void FNiagaraParameterMapHistoryBuilder::EnterFunction(const FString& InNodeName
 	}
 }
 
-void FNiagaraParameterMapHistoryBuilder::ExitFunction(const FString& InNodeName, const UNiagaraScript* InScript, UNiagaraNode* Node)
+void FNiagaraParameterMapHistoryBuilder::ExitFunction(const FString& InNodeName, const UNiagaraScript* InScript, const UNiagaraNode* Node)
 {
 	if (InScript != nullptr)
 	{
@@ -984,7 +942,7 @@ void FNiagaraParameterMapHistoryBuilder::ExitFunction(const FString& InNodeName,
 	}
 }
 
-void FNiagaraParameterMapHistoryBuilder::EnterEmitter(const FString& InEmitterName, UNiagaraNode* Node)
+void FNiagaraParameterMapHistoryBuilder::EnterEmitter(const FString& InEmitterName, const UNiagaraNode* Node)
 {
 	RegisterNodeVisitation(Node);
 	CallingContext.Push(Node);
@@ -994,7 +952,7 @@ void FNiagaraParameterMapHistoryBuilder::EnterEmitter(const FString& InEmitterNa
 	// Emitters must record their namespaces to their histories as well as
 	// make sure to record their current usage type is so that we can filter variables
 	// for relevance downstream.
-	UNiagaraNodeEmitter* EmitterNode = Cast<UNiagaraNodeEmitter>(Node);
+	const UNiagaraNodeEmitter* EmitterNode = Cast<UNiagaraNodeEmitter>(Node);
 	if (EmitterNode != nullptr)
 	{
 		RelevantScriptUsageContext.Emplace(EmitterNode->GetUsage());
@@ -1014,7 +972,7 @@ void FNiagaraParameterMapHistoryBuilder::EnterEmitter(const FString& InEmitterNa
 	ContextuallyVisitedNodes.Emplace();
 }
 
-void FNiagaraParameterMapHistoryBuilder::ExitEmitter(const FString& InEmitterName, UNiagaraNode* Node)
+void FNiagaraParameterMapHistoryBuilder::ExitEmitter(const FString& InEmitterName, const UNiagaraNode* Node)
 {
 	CallingContext.Pop();
 	EmitterNameContextStack.Pop();
@@ -1075,7 +1033,7 @@ const FString* FNiagaraParameterMapHistoryBuilder::GetEmitterAlias() const
 	return AliasMap.Find(TEXT("Emitter"));
 }
 
-void FNiagaraParameterMapHistoryBuilder::VisitInputPin(const UEdGraphPin* Pin, class UNiagaraNode* InNode)
+void FNiagaraParameterMapHistoryBuilder::VisitInputPin(const UEdGraphPin* Pin, const class UNiagaraNode* InNode)
 {
 	const UEdGraphSchema_Niagara* Schema = GetDefault<UEdGraphSchema_Niagara>();
 
@@ -1104,7 +1062,7 @@ void FNiagaraParameterMapHistoryBuilder::VisitInputPin(const UEdGraphPin* Pin, c
 	}
 }
 
-void FNiagaraParameterMapHistoryBuilder::VisitInputPins(class UNiagaraNode* InNode)
+void FNiagaraParameterMapHistoryBuilder::VisitInputPins(const class UNiagaraNode* InNode)
 {
 	TArray<UEdGraphPin*> InputPins;
 	InNode->GetInputPins(InputPins);
diff --git a/Engine/Plugins/FX/Niagara/Source/NiagaraEditor/Private/NiagaraScriptMergeManager.cpp b/Engine/Plugins/FX/Niagara/Source/NiagaraEditor/Private/NiagaraScriptMergeManager.cpp
--- a/Engine/Plugins/FX/Niagara/Source/NiagaraEditor/Private/NiagaraScriptMergeManager.cpp
+++ b/Engine/Plugins/FX/Niagara/Source/NiagaraEditor/Private/NiagaraScriptMergeManager.cpp
@@ -671,11 +671,16 @@ void FNiagaraScriptMergeManager::DiffChangeIds(const TMap<FGuid, FGuid>& InSourc
 		}
 		else if (MatchingSourceChangeId != nullptr && MatchingLastMergedChangeId != nullptr)
 		{
-			// If both had a copy of the node and both agree on the change id, then we should keep the change id of the instance as it will be the most accurate.
-			if (MatchingSourceChangeId->IsValid() && MatchingLastMergedChangeId->IsValid() && (*MatchingSourceChangeId) == (*MatchingLastMergedChangeId))
+			if (*MatchingSourceChangeId == *MatchingLastMergedChangeId)
 			{
+				// If both had a copy of the node and both agree on the change id, then we should keep the change id of the instance as it will be the most accurate.
 				OutChangeIdsToKeepOnInstance.Add(It.Key(), It.Value());
 			}
+			else
+			{
+				// If both had a copy of the node and they're different than the source has changed and it's change id should be used since it's the newer.
+				OutChangeIdsToKeepOnInstance.Add(It.Key(), *MatchingSourceChangeId);
+			}
 		}
 		else if (MatchingLastMergedChangeId != nullptr)
 		{
@@ -1998,6 +2003,15 @@ FNiagaraScriptMergeManager::FApplyDiffResults FNiagaraScriptMergeManager::ApplyE
 				Results.bSucceeded &= AddModuleResults.bSucceeded;
 				Results.ErrorMessages.Append(AddModuleResults.ErrorMessages);
 			}
+
+			// Force the base compile id of the new event handler to match the added instance event handler.
+			UNiagaraScriptSource* AddedEventScriptSourceFromDiff = Cast<UNiagaraScriptSource>(AddedEventHandler->GetEventScriptProperties()->Script->GetSource());
+			UNiagaraGraph* AddedEventScriptGraphFromDiff = AddedEventScriptSourceFromDiff->NodeGraph;
+			FGuid ScriptBaseIdFromDiff = AddedEventScriptGraphFromDiff->GetBaseId(ENiagaraScriptUsage::ParticleEventScript, AddedEventHandler->GetUsageId());
+			UNiagaraScriptSource* AddedEventScriptSource = Cast<UNiagaraScriptSource>(AddedEventScriptProperties.Script->GetSource());
+			UNiagaraGraph* AddedEventScriptGraph = AddedEventScriptSource->NodeGraph;
+			AddedEventScriptGraph->ForceBaseId(ENiagaraScriptUsage::ParticleEventScript, AddedEventHandler->GetUsageId(), ScriptBaseIdFromDiff);
+
 			Results.bModifiedGraph = true;
 		}
 	}
diff --git a/Engine/Plugins/FX/Niagara/Source/NiagaraEditor/Private/NiagaraScriptSource.cpp b/Engine/Plugins/FX/Niagara/Source/NiagaraEditor/Private/NiagaraScriptSource.cpp
--- a/Engine/Plugins/FX/Niagara/Source/NiagaraEditor/Private/NiagaraScriptSource.cpp
+++ b/Engine/Plugins/FX/Niagara/Source/NiagaraEditor/Private/NiagaraScriptSource.cpp
@@ -43,18 +43,30 @@ UNiagaraScriptSource::UNiagaraScriptSource(const FObjectInitializer& ObjectIniti
 {
 }
 
-void UNiagaraScriptSource::ComputeVMCompilationId(FNiagaraVMExecutableDataId& Id, ENiagaraScriptUsage InUsage, const FGuid& InUsageId) const
+void UNiagaraScriptSource::ComputeVMCompilationId(FNiagaraVMExecutableDataId& Id, ENiagaraScriptUsage InUsage, const FGuid& InUsageId, bool bForceRebuild) const
 {
 	Id.ScriptUsageType = InUsage;
 	Id.ScriptUsageTypeID = InUsageId;
 	Id.CompilerVersionID = FNiagaraCustomVersion::LatestScriptCompileVersion;
 	if (NodeGraph)
 	{
-		Id.BaseScriptID = NodeGraph->GetCompileID(InUsage, InUsageId);
-		NodeGraph->GatherExternalDependencyIDs(InUsage, InUsageId, Id.ReferencedDependencyIds, Id.ReferencedObjects);
+		NodeGraph->RebuildCachedCompileIds(bForceRebuild);
+		Id.BaseScriptID = NodeGraph->GetBaseId(InUsage, InUsageId);
+		Id.BaseScriptCompileHash = FNiagaraCompileHash(NodeGraph->GetCompileDataHash(InUsage, InUsageId));
+		NodeGraph->GatherExternalDependencyIDs(InUsage, InUsageId, Id.ReferencedCompileHashes, Id.ReferencedDependencyIds, Id.ReferencedObjects);
 	}
 }
 
+FGuid UNiagaraScriptSource::GetCompileBaseId(ENiagaraScriptUsage InUsage, const FGuid& InUsageId) const
+{
+	return NodeGraph->GetBaseId(InUsage, InUsageId);
+}
+
+FNiagaraCompileHash UNiagaraScriptSource::GetCompileHash(ENiagaraScriptUsage InUsage, const FGuid& InUsageId) const
+{
+	return NodeGraph->GetCompileDataHash(InUsage, InUsageId);
+}
+
 void UNiagaraScriptSource::InvalidateCachedCompileIds() 
 {
 	NodeGraph->InvalidateCachedCompileIds();
diff --git a/Engine/Plugins/FX/Niagara/Source/NiagaraEditor/Private/NiagaraScriptVariable.h b/Engine/Plugins/FX/Niagara/Source/NiagaraEditor/Private/NiagaraScriptVariable.h
new file mode 100644
--- /dev/null
+++ b/Engine/Plugins/FX/Niagara/Source/NiagaraEditor/Private/NiagaraScriptVariable.h
@@ -0,0 +1,19 @@
+// Copyright 1998-2019 Epic Games, Inc. All Rights Reserved.
+
+#pragma once
+
+#include "NiagaraTypes.h"
+#include "NiagaraScriptVariable.generated.h"
+
+UCLASS()
+class UNiagaraScriptVariable : public UObject
+{
+public:
+	GENERATED_UCLASS_BODY()
+
+	UPROPERTY()
+	FNiagaraVariable Variable;
+
+	UPROPERTY(EditAnywhere, Category = Metadata)
+	FNiagaraVariableMetaData Metadata;
+};
\ No newline at end of file
diff --git a/Engine/Plugins/FX/Niagara/Source/NiagaraEditor/Private/NiagaraStackFunctionInputBinder.cpp b/Engine/Plugins/FX/Niagara/Source/NiagaraEditor/Private/NiagaraStackFunctionInputBinder.cpp
--- a/Engine/Plugins/FX/Niagara/Source/NiagaraEditor/Private/NiagaraStackFunctionInputBinder.cpp
+++ b/Engine/Plugins/FX/Niagara/Source/NiagaraEditor/Private/NiagaraStackFunctionInputBinder.cpp
@@ -33,8 +33,8 @@ bool FNiagaraStackFunctionInputBinder::TryBind(
 	FInputMatchesPredicate InputMatches;
 	InputMatches.BindLambda([=](FNiagaraVariable InputVariable)
 	{
-		FNiagaraVariableMetaData* InputMetadata = FunctionGraph->GetMetaData(InputVariable);
-		if (InputMetadata != nullptr)
+		TOptional<FNiagaraVariableMetaData> InputMetadata = FunctionGraph->GetMetaData(InputVariable);
+		if (InputMetadata.IsSet())
 		{
 			FString* ValueString = InputMetadata->PropertyMetaData.Find(InMetaDataKey);
 			if (ValueString != nullptr && *ValueString == InMetaDataValue)
@@ -92,11 +92,12 @@ bool FNiagaraStackFunctionInputBinder::TryBindInternal(
 
 	TArray<const UEdGraphPin*> InputPins;
 	FNiagaraStackGraphUtilities::GetStackFunctionInputPins(*FunctionCallNode, InputPins, FNiagaraStackGraphUtilities::ENiagaraGetStackFunctionInputPinsOptions::ModuleInputsOnly);
-
+	
 	bool bInputFound = false;
+	const UEdGraphSchema_Niagara* Schema = GetDefault<UEdGraphSchema_Niagara>();
 	for (const UEdGraphPin* InputPin : InputPins)
 	{
-		FNiagaraVariable InputVariable = GetDefault<UEdGraphSchema_Niagara>()->PinToNiagaraVariable(InputPin);
+		FNiagaraVariable InputVariable = Schema->PinToNiagaraVariable(InputPin);
 		if(InputMatchesCallback.Execute(InputVariable))
 		{
 			InputName = InputVariable.GetName();
@@ -135,6 +136,24 @@ bool FNiagaraStackFunctionInputBinder::TryBindInternal(
 		}
 	}
 
+	// For the EditCondition and VisibleCondition, we want to be able to bind to static switch input variables
+	UNiagaraGraph* CalledGraph = FunctionCallNode->GetCalledGraph();
+	if (!bInputFound && CalledGraph)
+	{
+		for (FNiagaraVariable InputVar : CalledGraph->FindStaticSwitchInputs())
+		{
+			if (InputMatchesCallback.Execute(InputVar) && (!InInputType.IsSet() || InputVar.GetType() == InInputType.GetValue()))
+			{
+				InputName = InputVar.GetName();
+				InputType = InputVar.GetType();
+				AliasedParameterHandle = FNiagaraParameterHandle::CreateAliasedModuleParameterHandle(FNiagaraParameterHandle(InputVar.GetName()), FunctionCallNode.Get());
+				RefreshGraphPins();
+				bInputFound = true;
+				break;
+			}
+		}
+	}
+
 	if (bInputFound == false)
 	{
 		Reset();
@@ -274,6 +293,12 @@ void FNiagaraStackFunctionInputBinder::RefreshGraphPins() const
 {
 	OverridePin = FNiagaraStackGraphUtilities::GetStackFunctionInputOverridePin(*FunctionCallNode, AliasedParameterHandle);
 	DefaultPin = FunctionCallNode->FindParameterMapDefaultValuePin(InputName, Script->GetUsage());
+	
+	// if we bind to a static switch value the default value pin cannot be found via the parameter map
+	if (DefaultPin == nullptr)
+	{
+		DefaultPin = FunctionCallNode->FindStaticSwitchInputPin(InputName);
+	}
 	ValidScriptGraphChangeIdForOverridePin = Script->GetSource()->GetChangeID();
 	ValidScriptGraphChangeIdForDefaultPin = GetChangeIdFromFunctionScript();
 }
diff --git a/Engine/Plugins/FX/Niagara/Source/NiagaraEditor/Private/Sequencer/NiagaraSequence/MovieSceneNiagaraEmitterTrack.cpp b/Engine/Plugins/FX/Niagara/Source/NiagaraEditor/Private/Sequencer/NiagaraSequence/MovieSceneNiagaraEmitterTrack.cpp
--- a/Engine/Plugins/FX/Niagara/Source/NiagaraEditor/Private/Sequencer/NiagaraSequence/MovieSceneNiagaraEmitterTrack.cpp
+++ b/Engine/Plugins/FX/Niagara/Source/NiagaraEditor/Private/Sequencer/NiagaraSequence/MovieSceneNiagaraEmitterTrack.cpp
@@ -81,6 +81,7 @@ void UMovieSceneNiagaraEmitterTrack::Initialize(FNiagaraSystemViewModel& InSyste
 	EmitterHandleViewModel = InEmitterHandleViewModel;
 	SetDisplayName(InEmitterHandleViewModel->GetNameText());
 	EmitterHandleId = InEmitterHandleViewModel->GetId();
+	SystemPath = SystemViewModel->GetSystem().GetPathName();
 	SetColorTint(FNiagaraEditorStyle::Get().GetColor("NiagaraEditor.NiagaraSequence.DefaultTrackColor").ToFColor(true));
 	CreateSections(InFrameResolution);
 }
@@ -184,6 +185,11 @@ FGuid UMovieSceneNiagaraEmitterTrack::GetEmitterHandleId() const
 	return EmitterHandleId;
 }
 
+const FString& UMovieSceneNiagaraEmitterTrack::GetSystemPath() const
+{
+	return SystemPath;
+}
+
 const TArray<FText>& UMovieSceneNiagaraEmitterTrack::GetSectionInitializationErrors() const
 {
 	return SectionInitializationErrors;
diff --git a/Engine/Plugins/FX/Niagara/Source/NiagaraEditor/Private/Sequencer/NiagaraSequence/MovieSceneNiagaraEmitterTrack.h b/Engine/Plugins/FX/Niagara/Source/NiagaraEditor/Private/Sequencer/NiagaraSequence/MovieSceneNiagaraEmitterTrack.h
--- a/Engine/Plugins/FX/Niagara/Source/NiagaraEditor/Private/Sequencer/NiagaraSequence/MovieSceneNiagaraEmitterTrack.h
+++ b/Engine/Plugins/FX/Niagara/Source/NiagaraEditor/Private/Sequencer/NiagaraSequence/MovieSceneNiagaraEmitterTrack.h
@@ -81,6 +81,9 @@ public:
 	/** Gets the unique id for the emitter handle that was associated with this track; used for copy/paste detection */
 	FGuid GetEmitterHandleId() const;
 
+	/** Gets the string path of the system which owns the emitter associated with this track; used for copy/paste detection */
+	const FString& GetSystemPath() const;
+
 	const TArray<FText>& GetSectionInitializationErrors() const;
 
 private:
@@ -98,5 +101,9 @@ private:
 	UPROPERTY()
 	FGuid EmitterHandleId;
 
+	// Used for detecting copy/paste
+	UPROPERTY()
+	FString SystemPath;
+
 	TArray<FText> SectionInitializationErrors;
 };
\ No newline at end of file
diff --git a/Engine/Plugins/FX/Niagara/Source/NiagaraEditor/Private/Sequencer/NiagaraSequence/MovieSceneNiagaraEmitterTrackInstance.h b/Engine/Plugins/FX/Niagara/Source/NiagaraEditor/Private/Sequencer/NiagaraSequence/MovieSceneNiagaraEmitterTrackInstance.h
--- a/Engine/Plugins/FX/Niagara/Source/NiagaraEditor/Private/Sequencer/NiagaraSequence/MovieSceneNiagaraEmitterTrackInstance.h
+++ b/Engine/Plugins/FX/Niagara/Source/NiagaraEditor/Private/Sequencer/NiagaraSequence/MovieSceneNiagaraEmitterTrackInstance.h
@@ -5,7 +5,7 @@
 #include "IMovieSceneTrackInstance.h"
 
 class UMovieSceneNiagaraEmitterTrack;
-struct FNiagaraEmitterInstance;
+class FNiagaraEmitterInstance;
 
 /**
 *	One instance of a UEmitterMovieSceneTrack
diff --git a/Engine/Plugins/FX/Niagara/Source/NiagaraEditor/Private/Toolkits/NiagaraScriptToolkit.cpp b/Engine/Plugins/FX/Niagara/Source/NiagaraEditor/Private/Toolkits/NiagaraScriptToolkit.cpp
--- a/Engine/Plugins/FX/Niagara/Source/NiagaraEditor/Private/Toolkits/NiagaraScriptToolkit.cpp
+++ b/Engine/Plugins/FX/Niagara/Source/NiagaraEditor/Private/Toolkits/NiagaraScriptToolkit.cpp
@@ -50,14 +50,9 @@
 
 DECLARE_CYCLE_STAT(TEXT("Niagara - ScriptToolkit - OnApply"), STAT_NiagaraEditor_ScriptToolkit_OnApply, STATGROUP_NiagaraEditor);
 
-static TAutoConsoleVariable<int32> CVarDevDetails(
-	TEXT("fx.DevDetailsPanels"),
-	1,
-	TEXT("Whether to enable the development details panels inside Niagara."));
-
-
 const FName FNiagaraScriptToolkit::NodeGraphTabId(TEXT("NiagaraEditor_NodeGraph")); 
-const FName FNiagaraScriptToolkit::DetailsTabId(TEXT("NiagaraEditor_Details"));
+const FName FNiagaraScriptToolkit::ScriptDetailsTabId(TEXT("NiagaraEditor_ScriptDetails"));
+const FName FNiagaraScriptToolkit::SelectedDetailsTabId(TEXT("NiagaraEditor_SelectedDetails"));
 const FName FNiagaraScriptToolkit::ParametersTabId(TEXT("NiagaraEditor_Parameters"));
 const FName FNiagaraScriptToolkit::StatsTabId(TEXT("NiagaraEditor_Stats"));
 const FName FNiagaraScriptToolkit::MessageLogTabID(TEXT("NiagaraEditor_MessageLog"));
@@ -87,8 +82,13 @@ void FNiagaraScriptToolkit::RegisterTabSpawners(const TSharedRef<class FTabManag
 		.SetDisplayName( LOCTEXT("NodeGraph", "Node Graph") )
 		.SetGroup(WorkspaceMenuCategoryRef); 
 
-	InTabManager->RegisterTabSpawner(DetailsTabId, FOnSpawnTab::CreateSP(this, &FNiagaraScriptToolkit::SpawnTabNodeDetails))
-		.SetDisplayName(LOCTEXT("DetailsTab", "Details"))
+	InTabManager->RegisterTabSpawner(ScriptDetailsTabId, FOnSpawnTab::CreateSP(this, &FNiagaraScriptToolkit::SpawnTabScriptDetails))
+		.SetDisplayName(LOCTEXT("ScriptDetailsTab", "Script Details"))
+		.SetGroup(WorkspaceMenuCategoryRef)
+		.SetIcon(FSlateIcon(FEditorStyle::GetStyleSetName(), "LevelEditor.Tabs.Details"));
+
+	InTabManager->RegisterTabSpawner(SelectedDetailsTabId, FOnSpawnTab::CreateSP(this, &FNiagaraScriptToolkit::SpawnTabSelectedDetails))
+		.SetDisplayName(LOCTEXT("ScriptDetailsTab", "Selected Details"))
 		.SetGroup(WorkspaceMenuCategoryRef)
 		.SetIcon(FSlateIcon(FEditorStyle::GetStyleSetName(), "LevelEditor.Tabs.Details"));
 
@@ -111,7 +111,8 @@ void FNiagaraScriptToolkit::UnregisterTabSpawners(const TSharedRef<class FTabMan
 	FAssetEditorToolkit::UnregisterTabSpawners(InTabManager);
 
 	InTabManager->UnregisterTabSpawner(NodeGraphTabId);
-	InTabManager->UnregisterTabSpawner(DetailsTabId);
+	InTabManager->UnregisterTabSpawner(ScriptDetailsTabId);
+	InTabManager->UnregisterTabSpawner(SelectedDetailsTabId);
 	InTabManager->UnregisterTabSpawner(ParametersTabId);
 	InTabManager->UnregisterTabSpawner(StatsTabId);
 }
@@ -145,8 +146,8 @@ void FNiagaraScriptToolkit::Initialize( const EToolkitMode::Type Mode, const TSh
 	OnEditedScriptGraphChangedHandle = ScriptViewModel->GetGraphViewModel()->GetGraph()->AddOnGraphNeedsRecompileHandler(
 		FOnGraphChanged::FDelegate::CreateRaw(this, &FNiagaraScriptToolkit::OnEditedScriptGraphChanged));
 
-	DetailsSelection = MakeShareable(new FNiagaraObjectSelection());
-	DetailsSelection->SetSelectedObject(EditedNiagaraScript);
+	DetailsScriptSelection = MakeShareable(new FNiagaraObjectSelection());
+	DetailsScriptSelection->SetSelectedObject(EditedNiagaraScript);
 	
 	FMessageLogModule& MessageLogModule = FModuleManager::LoadModuleChecked<FMessageLogModule>("MessageLog");
 	FMessageLogInitializationOptions LogOptions;
@@ -190,13 +191,13 @@ void FNiagaraScriptToolkit::Initialize( const EToolkitMode::Type Mode, const TSh
 			->Split
 			(
 				FTabManager::NewSplitter()->SetOrientation(Orient_Vertical)
-				->SetSizeCoefficient(0.2f)
+				->SetSizeCoefficient(0.15f)
 				->Split
 				(
 					FTabManager::NewStack()
 					->SetSizeCoefficient(0.5f)
-					->AddTab(DetailsTabId, ETabState::OpenedTab)
-					->SetForegroundTab(DetailsTabId)
+					->AddTab(ScriptDetailsTabId, ETabState::OpenedTab)
+					->SetForegroundTab(ScriptDetailsTabId)
 				)
 				->Split
 				(
@@ -218,7 +219,7 @@ void FNiagaraScriptToolkit::Initialize( const EToolkitMode::Type Mode, const TSh
 				->Split
 				(
 				FTabManager::NewStack()
-				->SetSizeCoefficient(0.8f)
+				->SetSizeCoefficient(0.7f)
 				->AddTab(NodeGraphTabId, ETabState::OpenedTab)
 				)
 				->Split
@@ -228,6 +229,12 @@ void FNiagaraScriptToolkit::Initialize( const EToolkitMode::Type Mode, const TSh
 				->AddTab(MessageLogTabID, ETabState::OpenedTab)
 				)
 			)
+			->Split
+			(
+				FTabManager::NewStack()
+				->SetSizeCoefficient(0.15f)
+				->AddTab(SelectedDetailsTabId, ETabState::OpenedTab)
+			)
 		)
 	);
 
@@ -331,13 +338,11 @@ void FNiagaraScriptToolkit::OnVMScriptCompiled(UNiagaraScript* InScript)
 	UpdateModuleStats();
 }
 
-TSharedRef<SDockTab> FNiagaraScriptToolkit::SpawnTabNodeDetails(const FSpawnTabArgs& Args)
+TSharedRef<SDockTab> FNiagaraScriptToolkit::SpawnTabScriptDetails(const FSpawnTabArgs& Args)
 {
-	checkf(Args.GetTabId().TabType == DetailsTabId, TEXT("Wrong tab ID in NiagaraScriptToolkit"));
+	checkf(Args.GetTabId().TabType == ScriptDetailsTabId, TEXT("Wrong tab ID in NiagaraScriptToolkit"));
 	checkf(ScriptViewModel.IsValid(), TEXT("NiagaraScriptToolkit - Script editor view model is invalid"));
 
-	IConsoleVariable* DevDetailsCVar = IConsoleManager::Get().FindConsoleVariable(TEXT("fx.DevDetailsPanels"));
-
 	TWeakPtr<FNiagaraScriptViewModel> ScriptViewModelWeakPtr = ScriptViewModel;
 
 	FPropertyEditorModule& PropertyEditorModule = FModuleManager::GetModuleChecked<FPropertyEditorModule>("PropertyEditor");
@@ -350,46 +355,41 @@ TSharedRef<SDockTab> FNiagaraScriptToolkit::SpawnTabNodeDetails(const FSpawnTabA
 		FOnGetDetailCustomizationInstance::CreateStatic(&FNiagaraScriptDetails::MakeInstance, ScriptViewModelWeakPtr)
 	);
 
-	DetailsView->SetObjects(DetailsSelection->GetSelectedObjects().Array());
+	DetailsView->SetObjects(DetailsScriptSelection->GetSelectedObjects().Array());
 
-	if (DevDetailsCVar && DevDetailsCVar->GetInt() != 0)
-	{
-		return SNew(SDockTab)
-			.Label(LOCTEXT("ScriptNodeDetailsTabLabel", "Details"))
-			.TabColorScale(GetTabColorScale())
-			[
-				SNew(SScrollBox)
-				+ SScrollBox::Slot()
-				.Padding(0, 3, 0, 0)
-				[
-					DetailsView
-				]
-				+ SScrollBox::Slot()
-				.Padding(0, 3, 0, 0)
-				[
-					SNew(SNiagaraSelectedObjectsDetails, ScriptViewModel->GetGraphViewModel()->GetSelection())
-				]
-			];
-	}
-	else
-	{
-		return SNew(SDockTab)
-			.Label(LOCTEXT("ScriptNodeDetailsTabLabel", "Details"))
-			.TabColorScale(GetTabColorScale())
-			[
-				DetailsView
-			];
-	}
+	return SNew(SDockTab)
+		.Label(LOCTEXT("ScriptDetailsTabLabel", "Script Details"))
+		.TabColorScale(GetTabColorScale())
+		[
+			DetailsView
+		];
+}
+
+TSharedRef<SDockTab> FNiagaraScriptToolkit::SpawnTabSelectedDetails(const FSpawnTabArgs& Args)
+{
+	checkf(Args.GetTabId().TabType == SelectedDetailsTabId, TEXT("Wrong tab ID in NiagaraScriptToolkit"));
+	checkf(ScriptViewModel.IsValid(), TEXT("NiagaraScriptToolkit - Script editor view model is invalid"));
+
+	return SNew(SDockTab)
+		.Label(LOCTEXT("SelectedDetailsTabLabel", "Selected Details"))
+		.TabColorScale(GetTabColorScale())
+		[
+			SNew(SNiagaraSelectedObjectsDetails, ScriptViewModel->GetGraphViewModel()->GetNodeSelection(), ScriptViewModel->GetVariableSelection())
+		];
 }
 
 TSharedRef<SDockTab> FNiagaraScriptToolkit::SpawnTabScriptParameters(const FSpawnTabArgs& Args)
 {
 	checkf(Args.GetTabId().TabType == ParametersTabId, TEXT("Wrong tab ID in NiagaraScriptToolkit"));
 
+	TArray<TSharedRef<FNiagaraObjectSelection>> Array;
+	Array.Push(DetailsScriptSelection.ToSharedRef());
+	Array.Push(ScriptViewModel->GetVariableSelection());
+
 	TSharedRef<SDockTab> SpawnedTab =
 		SNew(SDockTab)
 		[
-			SNew(SNiagaraParameterMapView, DetailsSelection.ToSharedRef(), SNiagaraParameterMapView::EToolkitType::SCRIPT, GetToolkitCommands())
+			SNew(SNiagaraParameterMapView, Array, SNiagaraParameterMapView::EToolkitType::SCRIPT, GetToolkitCommands())
 		];
 
 	return SpawnedTab;
@@ -470,13 +470,12 @@ void FNiagaraScriptToolkit::ExtendToolbar()
 				TAttribute<FText>(ScriptToolkit, &FNiagaraScriptToolkit::GetCompileStatusTooltip),
 				TAttribute<FSlateIcon>(ScriptToolkit, &FNiagaraScriptToolkit::GetCompileStatusImage),
 				FName(TEXT("CompileNiagaraScript")));
-			// removed this for UE-58554 ahead of 4.20. Functionality code should also be removed if this becomes permanent.
-			/*ToolbarBuilder.AddToolBarButton(FNiagaraEditorCommands::Get().RefreshNodes,
+			ToolbarBuilder.AddToolBarButton(FNiagaraEditorCommands::Get().RefreshNodes,
 				NAME_None,
 				TAttribute<FText>(),
 				TAttribute<FText>(ScriptToolkit, &FNiagaraScriptToolkit::GetRefreshStatusTooltip),
 				TAttribute<FSlateIcon>(ScriptToolkit, &FNiagaraScriptToolkit::GetRefreshStatusImage),
-				FName(TEXT("RefreshScriptReferences")));*/
+				FName(TEXT("RefreshScriptReferences")));
 			ToolbarBuilder.EndSection();
 		}
 	};
@@ -533,7 +532,6 @@ FText FNiagaraScriptToolkit::GetRefreshStatusTooltip() const
 void FNiagaraScriptToolkit::CompileScript(bool bForce)
 {
 	ScriptViewModel->CompileStandaloneScript();
-	ScriptViewModel->RefreshMetadataCollection();
 }
 
 void FNiagaraScriptToolkit::RefreshNodes()
@@ -673,7 +671,6 @@ void FNiagaraScriptToolkit::UpdateOriginalNiagaraScript()
 
 	// Compile and then overwrite the original script in place by constructing a new one with the same name
 	ScriptViewModel->CompileStandaloneScript();
-	ScriptViewModel->RefreshMetadataCollection();
 	OriginalNiagaraScript = (UNiagaraScript*)StaticDuplicateObject(EditedNiagaraScript, OriginalNiagaraScript->GetOuter(), OriginalNiagaraScript->GetFName(),
 		RF_AllFlags,
 		OriginalNiagaraScript->GetClass());
diff --git a/Engine/Plugins/FX/Niagara/Source/NiagaraEditor/Private/Toolkits/NiagaraScriptToolkit.h b/Engine/Plugins/FX/Niagara/Source/NiagaraEditor/Private/Toolkits/NiagaraScriptToolkit.h
--- a/Engine/Plugins/FX/Niagara/Source/NiagaraEditor/Private/Toolkits/NiagaraScriptToolkit.h
+++ b/Engine/Plugins/FX/Niagara/Source/NiagaraEditor/Private/Toolkits/NiagaraScriptToolkit.h
@@ -10,7 +10,6 @@
 #include "Toolkits/AssetEditorToolkit.h"
 #include "UObject/GCObject.h"
 
-#include "NiagaraScript.h"
 #include "NiagaraEditorCommon.h"
 
 class IDetailsView;
@@ -79,7 +78,10 @@ private:
 	TSharedRef<SDockTab> SpawnTabNodeGraph(const FSpawnTabArgs& Args);
 
 	/** Spawns the tab with the script details inside. */
-	TSharedRef<SDockTab> SpawnTabNodeDetails(const FSpawnTabArgs& Args);
+	TSharedRef<SDockTab> SpawnTabScriptDetails(const FSpawnTabArgs& Args);
+
+	/** Spawns the tab with the details of the current selection. */
+	TSharedRef<SDockTab> SpawnTabSelectedDetails(const FSpawnTabArgs& Args);
 
 	/** Spawns the tab with the script details inside. */
 	TSharedRef<SDockTab> SpawnTabScriptParameters(const FSpawnTabArgs& Args);
@@ -127,7 +129,7 @@ private:
 	TSharedPtr<FNiagaraScriptViewModel> ScriptViewModel;
 
 	/** The selection displayed by the details tab. */
-	TSharedPtr<FNiagaraObjectSelection> DetailsSelection;
+	TSharedPtr<FNiagaraObjectSelection> DetailsScriptSelection;
 
 	/** Message log, with the log listing that it reflects */
 	TSharedPtr<class SWidget> NiagaraMessageLog;
@@ -135,7 +137,8 @@ private:
 
 	/**	The tab ids for the Niagara editor */
 	static const FName NodeGraphTabId; 
-	static const FName DetailsTabId;
+	static const FName ScriptDetailsTabId;
+	static const FName SelectedDetailsTabId;
 	static const FName ParametersTabId;
 	static const FName StatsTabId;
 	static const FName MessageLogTabID;
diff --git a/Engine/Plugins/FX/Niagara/Source/NiagaraEditor/Private/Toolkits/NiagaraSystemToolkit.cpp b/Engine/Plugins/FX/Niagara/Source/NiagaraEditor/Private/Toolkits/NiagaraSystemToolkit.cpp
--- a/Engine/Plugins/FX/Niagara/Source/NiagaraEditor/Private/Toolkits/NiagaraSystemToolkit.cpp
+++ b/Engine/Plugins/FX/Niagara/Source/NiagaraEditor/Private/Toolkits/NiagaraSystemToolkit.cpp
@@ -341,10 +341,10 @@ void FNiagaraSystemToolkit::InitializeWithEmitter(const EToolkitMode::Type Mode,
 	// Always remove the source information from the emitter handle when editing an emitter asset because inheritance is not valid in this case.
 	System->GetEmitterHandle(0).RemoveSource();
 
-	// We set this to the copy's change id here instead of the original emitter's change id because the copy's change id may have been
-	// updated from the original as part of post load and we use this id to detect if the editable emitter has been changed.
+	// Adding the emitter to the system has made a copy of it and we set this to the copy's change id here instead of the original emitter's change 
+	// id because the copy's change id may have been updated from the original as part of post load and we use this id to detect if the editable 
+	// emitter has been changed.
 	LastSyncedEmitterChangeId = SystemViewModel->GetEmitterHandleViewModels()[0]->GetEmitterViewModel()->GetEmitter()->GetChangeId();
-	SystemViewModel->GetSystemScriptViewModel()->RebuildEmitterNodes();
 	SystemToolkitMode = ESystemToolkitMode::Emitter;
 
 	if (GbLogNiagaraSystemChanges > 0)
@@ -615,10 +615,13 @@ TSharedRef<SDockTab> FNiagaraSystemToolkit::SpawnTab_SystemParameters(const FSpa
 		ObjectSelection->SetSelectedObject(System);
 	}
 
+	TArray<TSharedRef<FNiagaraObjectSelection>> Array;
+	Array.Push(ObjectSelection);
+
 	TSharedRef<SDockTab> SpawnedTab =
 		SNew(SDockTab)
 		[
-			SAssignNew(ParameterMapView, SNiagaraParameterMapView, ObjectSelection, SNiagaraParameterMapView::EToolkitType::SYSTEM, GetToolkitCommands())
+			SAssignNew(ParameterMapView, SNiagaraParameterMapView, Array, SNiagaraParameterMapView::EToolkitType::SYSTEM, GetToolkitCommands())
 		];
 
 	return SpawnedTab;
@@ -1095,30 +1098,6 @@ void FNiagaraSystemToolkit::OnToggleBoundsSetFixedBounds()
 
 	SystemViewModel->UpdateEmitterFixedBounds();
 
-	/*
-	// Force the component to update its bounds.
-	ParticleSystemComponent->ForceUpdateBounds();
-
-	// Grab the current bounds of the PSysComp & set it on the PSystem itself
-	ParticleSystem->Modify();
-	ParticleSystem->FixedRelativeBoundingBox.Min = ParticleSystemComponent->Bounds.GetBoxExtrema(0);
-	ParticleSystem->FixedRelativeBoundingBox.Max = ParticleSystemComponent->Bounds.GetBoxExtrema(1);
-	ParticleSystem->FixedRelativeBoundingBox.IsValid = true;
-	ParticleSystem->bUseFixedRelativeBoundingBox = true;
-
-	ParticleSystem->MarkPackageDirty();
-
-	EndTransaction(Transaction);
-
-	if ((SelectedModule == NULL) && (SelectedEmitter == NULL))
-	{
-		TArray<UObject*> NewSelection;
-		NewSelection.Add(ParticleSystem);
-		SetSelection(NewSelection);
-	}
-
-	ReassociateParticleSystem();
-	*/
 }
 
 void FNiagaraSystemToolkit::UpdateOriginalEmitter()
diff --git a/Engine/Plugins/FX/Niagara/Source/NiagaraEditor/Private/ViewModels/NiagaraEmitterHandleViewModel.cpp b/Engine/Plugins/FX/Niagara/Source/NiagaraEditor/Private/ViewModels/NiagaraEmitterHandleViewModel.cpp
--- a/Engine/Plugins/FX/Niagara/Source/NiagaraEditor/Private/ViewModels/NiagaraEmitterHandleViewModel.cpp
+++ b/Engine/Plugins/FX/Niagara/Source/NiagaraEditor/Private/ViewModels/NiagaraEmitterHandleViewModel.cpp
@@ -20,7 +20,7 @@
 
 #define LOCTEXT_NAMESPACE "EmitterHandleViewModel"
 
-FNiagaraEmitterHandleViewModel::FNiagaraEmitterHandleViewModel(FNiagaraEmitterHandle* InEmitterHandle, TWeakPtr<FNiagaraEmitterInstance> InSimulation, UNiagaraSystem& InOwningSystem)
+FNiagaraEmitterHandleViewModel::FNiagaraEmitterHandleViewModel(FNiagaraEmitterHandle* InEmitterHandle, TWeakPtr<FNiagaraEmitterInstance, ESPMode::ThreadSafe> InSimulation, UNiagaraSystem& InOwningSystem)
 	: EmitterHandle(InEmitterHandle)
 	, OwningSystem(InOwningSystem)
 	, EmitterViewModel(MakeShareable(new FNiagaraEmitterViewModel((InEmitterHandle ? InEmitterHandle->GetInstance() : nullptr), InSimulation)))
@@ -42,7 +42,7 @@ FNiagaraEmitterHandleViewModel::~FNiagaraEmitterHandleViewModel()
 	Cleanup();
 }
 
-bool FNiagaraEmitterHandleViewModel::Set(FNiagaraEmitterHandle* InEmitterHandle, TWeakPtr<FNiagaraEmitterInstance> InSimulation, UNiagaraSystem& InOwningSystem)
+bool FNiagaraEmitterHandleViewModel::Set(FNiagaraEmitterHandle* InEmitterHandle, TWeakPtr<FNiagaraEmitterInstance, ESPMode::ThreadSafe> InSimulation, UNiagaraSystem& InOwningSystem)
 {
 	if (&OwningSystem != &InOwningSystem)
 	{
@@ -67,7 +67,7 @@ void FNiagaraEmitterHandleViewModel::SetEmitterHandle(FNiagaraEmitterHandle* InE
 	EmitterHandle = InEmitterHandle;
 }
 
-void FNiagaraEmitterHandleViewModel::SetSimulation(TWeakPtr<FNiagaraEmitterInstance> InSimulation)
+void FNiagaraEmitterHandleViewModel::SetSimulation(TWeakPtr<FNiagaraEmitterInstance, ESPMode::ThreadSafe> InSimulation)
 {
 	check(EmitterViewModel.IsValid());
 	EmitterViewModel->SetSimulation(InSimulation);
diff --git a/Engine/Plugins/FX/Niagara/Source/NiagaraEditor/Private/ViewModels/NiagaraEmitterViewModel.cpp b/Engine/Plugins/FX/Niagara/Source/NiagaraEditor/Private/ViewModels/NiagaraEmitterViewModel.cpp
--- a/Engine/Plugins/FX/Niagara/Source/NiagaraEditor/Private/ViewModels/NiagaraEmitterViewModel.cpp
+++ b/Engine/Plugins/FX/Niagara/Source/NiagaraEditor/Private/ViewModels/NiagaraEmitterViewModel.cpp
@@ -31,7 +31,7 @@ namespace NiagaraCommands
 
 const float Megabyte = 1024.0f * 1024.0f;
 
-FNiagaraEmitterViewModel::FNiagaraEmitterViewModel(UNiagaraEmitter* InEmitter, TWeakPtr<FNiagaraEmitterInstance> InSimulation)
+FNiagaraEmitterViewModel::FNiagaraEmitterViewModel(UNiagaraEmitter* InEmitter, TWeakPtr<FNiagaraEmitterInstance, ESPMode::ThreadSafe> InSimulation)
 	: Emitter(InEmitter)
 	, Simulation(InSimulation)
 	, SharedScriptViewModel(MakeShareable(new FNiagaraScriptViewModel(InEmitter, LOCTEXT("SharedDisplayName", "Graph"), ENiagaraParameterEditMode::EditAll)))
@@ -51,14 +51,14 @@ void FNiagaraEmitterViewModel::Cleanup()
 
 	if (SharedScriptViewModel.IsValid())
 	{
-		SharedScriptViewModel->GetGraphViewModel()->GetSelection()->OnSelectedObjectsChanged().RemoveAll(this);
+		SharedScriptViewModel->GetGraphViewModel()->GetNodeSelection()->OnSelectedObjectsChanged().RemoveAll(this);
 		SharedScriptViewModel.Reset();
 	}
 
 	RemoveScriptEventHandlers();
 }
 
-bool FNiagaraEmitterViewModel::Set(UNiagaraEmitter* InEmitter, TWeakPtr<FNiagaraEmitterInstance> InSimulation)
+bool FNiagaraEmitterViewModel::Set(UNiagaraEmitter* InEmitter, TWeakPtr<FNiagaraEmitterInstance, ESPMode::ThreadSafe> InSimulation)
 {
 	SetEmitter(InEmitter);
 	SetSimulation(InSimulation);
@@ -105,7 +105,7 @@ void FNiagaraEmitterViewModel::SetEmitter(UNiagaraEmitter* InEmitter)
 }
 
 
-void FNiagaraEmitterViewModel::SetSimulation(TWeakPtr<FNiagaraEmitterInstance> InSimulation)
+void FNiagaraEmitterViewModel::SetSimulation(TWeakPtr<FNiagaraEmitterInstance, ESPMode::ThreadSafe> InSimulation)
 {
 	Simulation = InSimulation;
 }
@@ -121,7 +121,7 @@ FText FNiagaraEmitterViewModel::GetStatsText() const
 {
 	if (Simulation.IsValid())
 	{
-		TSharedPtr<FNiagaraEmitterInstance> SimInstance = Simulation.Pin();
+		TSharedPtr<FNiagaraEmitterInstance, ESPMode::ThreadSafe> SimInstance = Simulation.Pin();
 		if (SimInstance.IsValid())
 		{
 			static const FNumberFormattingOptions FractionalFormatOptions = FNumberFormattingOptions()
diff --git a/Engine/Plugins/FX/Niagara/Source/NiagaraEditor/Private/ViewModels/NiagaraScriptGraphViewModel.cpp b/Engine/Plugins/FX/Niagara/Source/NiagaraEditor/Private/ViewModels/NiagaraScriptGraphViewModel.cpp
--- a/Engine/Plugins/FX/Niagara/Source/NiagaraEditor/Private/ViewModels/NiagaraScriptGraphViewModel.cpp
+++ b/Engine/Plugins/FX/Niagara/Source/NiagaraEditor/Private/ViewModels/NiagaraScriptGraphViewModel.cpp
@@ -21,7 +21,7 @@ FNiagaraScriptGraphViewModel::FNiagaraScriptGraphViewModel(UNiagaraScriptSource*
 	: ScriptSource(InScriptSrc)
 	, DisplayName(InDisplayName)
 	, Commands(MakeShareable(new FUICommandList()))
-	, Selection(MakeShareable(new FNiagaraObjectSelection()))
+	, NodeSelection(MakeShareable(new FNiagaraObjectSelection()))
 {
 	SetupCommands();
 	GEditor->RegisterForUndo(this);
@@ -35,7 +35,7 @@ FNiagaraScriptGraphViewModel::~FNiagaraScriptGraphViewModel()
 
 void FNiagaraScriptGraphViewModel::SetScriptSource(UNiagaraScriptSource* InScriptSrc)
 {
-	Selection->ClearSelectedObjects();
+	NodeSelection->ClearSelectedObjects();
 	ScriptSource = InScriptSrc;
 	OnGraphChangedDelegate.Broadcast();
 }
@@ -64,9 +64,9 @@ TSharedRef<FUICommandList> FNiagaraScriptGraphViewModel::GetCommands()
 	return Commands;
 }
 
-TSharedRef<FNiagaraObjectSelection> FNiagaraScriptGraphViewModel::GetSelection()
+TSharedRef<FNiagaraObjectSelection> FNiagaraScriptGraphViewModel::GetNodeSelection()
 {
-	return Selection;
+	return NodeSelection;
 }
 
 FNiagaraScriptGraphViewModel::FOnNodesPasted& FNiagaraScriptGraphViewModel::OnNodesPasted()
@@ -156,7 +156,7 @@ void FNiagaraScriptGraphViewModel::SelectAllNodes()
 		Graph->GetNodesOfClass<UObject>(AllNodes);
 		TSet<UObject*> AllNodeSet;
 		AllNodeSet.Append(AllNodes);
-		Selection->SetSelectedObjects(AllNodeSet);
+		NodeSelection->SetSelectedObjects(AllNodeSet);
 	}
 }
 
@@ -168,8 +168,8 @@ void FNiagaraScriptGraphViewModel::DeleteSelectedNodes()
 		const FScopedTransaction Transaction(FGenericCommands::Get().Delete->GetDescription());
 		Graph->Modify();
 
-		TArray<UObject*> NodesToDelete = Selection->GetSelectedObjects().Array();
-		Selection->ClearSelectedObjects();
+		TArray<UObject*> NodesToDelete = NodeSelection->GetSelectedObjects().Array();
+		NodeSelection->ClearSelectedObjects();
 
 		for (UObject* NodeToDelete : NodesToDelete)
 		{
@@ -188,7 +188,7 @@ bool FNiagaraScriptGraphViewModel::CanDeleteNodes() const
 	UNiagaraGraph* Graph = GetGraph();
 	if (Graph != nullptr)
 	{
-		for (UObject* SelectedNode : Selection->GetSelectedObjects())
+		for (UObject* SelectedNode : NodeSelection->GetSelectedObjects())
 		{
 			UEdGraphNode* SelectedGraphNode = Cast<UEdGraphNode>(SelectedNode);
 			if (SelectedGraphNode != nullptr && SelectedGraphNode->CanUserDeleteNode())
@@ -205,7 +205,7 @@ void FNiagaraScriptGraphViewModel::CutSelectedNodes()
 	// Collect nodes which can not be delete or duplicated so they can be reselected.
 	TSet<UObject*> CanBeDuplicatedAndDeleted;
 	TSet<UObject*> CanNotBeDuplicatedAndDeleted;
-	for (UObject* SelectedNode : Selection->GetSelectedObjects())
+	for (UObject* SelectedNode : NodeSelection->GetSelectedObjects())
 	{
 		UEdGraphNode* SelectedGraphNode = Cast<UEdGraphNode>(SelectedNode);
 		if (SelectedGraphNode != nullptr)
@@ -222,10 +222,10 @@ void FNiagaraScriptGraphViewModel::CutSelectedNodes()
 	}
 
 	// Select the nodes which can be copied and deleted, copy and delete them, and then restore the ones which couldn't be copied or deleted.
-	Selection->SetSelectedObjects(CanBeDuplicatedAndDeleted);
+	NodeSelection->SetSelectedObjects(CanBeDuplicatedAndDeleted);
 	CopySelectedNodes();
 	DeleteSelectedNodes();
-	Selection->SetSelectedObjects(CanNotBeDuplicatedAndDeleted);
+	NodeSelection->SetSelectedObjects(CanNotBeDuplicatedAndDeleted);
 }
 
 bool FNiagaraScriptGraphViewModel::CanCutNodes() const
@@ -236,7 +236,7 @@ bool FNiagaraScriptGraphViewModel::CanCutNodes() const
 void FNiagaraScriptGraphViewModel::CopySelectedNodes()
 {
 	TSet<UObject*> NodesToCopy;
-	for (UObject* SelectedNode : Selection->GetSelectedObjects())
+	for (UObject* SelectedNode : NodeSelection->GetSelectedObjects())
 	{
 		UEdGraphNode* SelectedGraphNode = Cast<UEdGraphNode>(SelectedNode);
 		if (SelectedGraphNode != nullptr)
@@ -259,7 +259,7 @@ bool FNiagaraScriptGraphViewModel::CanCopyNodes() const
 	UNiagaraGraph* Graph = GetGraph();
 	if (Graph != nullptr)
 	{
-		for (UObject* SelectedNode : Selection->GetSelectedObjects())
+		for (UObject* SelectedNode : NodeSelection->GetSelectedObjects())
 		{
 			UEdGraphNode* SelectedGraphNode = Cast<UEdGraphNode>(SelectedNode);
 			if (SelectedGraphNode != nullptr && SelectedGraphNode->CanDuplicateNode())
@@ -280,7 +280,7 @@ void FNiagaraScriptGraphViewModel::PasteNodes()
 		const FScopedTransaction Transaction(FGenericCommands::Get().Paste->GetDescription());;
 		Graph->Modify();
 
-		Selection->ClearSelectedObjects();
+		NodeSelection->ClearSelectedObjects();
 
 		// Grab the text to paste from the clipboard.
 		FString TextToImport;
@@ -308,7 +308,7 @@ void FNiagaraScriptGraphViewModel::PasteNodes()
 			PastedObjects.Add(PastedNode);
 		}
 
-		Selection->SetSelectedObjects(PastedObjects);
+		NodeSelection->SetSelectedObjects(PastedObjects);
 		CastChecked<UNiagaraGraph>(Graph)->NotifyGraphNeedsRecompile();
 	}
 }
diff --git a/Engine/Plugins/FX/Niagara/Source/NiagaraEditor/Private/ViewModels/NiagaraScriptGraphViewModel.h b/Engine/Plugins/FX/Niagara/Source/NiagaraEditor/Private/ViewModels/NiagaraScriptGraphViewModel.h
--- a/Engine/Plugins/FX/Niagara/Source/NiagaraEditor/Private/ViewModels/NiagaraScriptGraphViewModel.h
+++ b/Engine/Plugins/FX/Niagara/Source/NiagaraEditor/Private/ViewModels/NiagaraScriptGraphViewModel.h
@@ -43,7 +43,7 @@ public:
 	TSharedRef<FUICommandList> GetCommands();
 
 	/** Gets the currently selected graph nodes. */
-	TSharedRef<FNiagaraObjectSelection> GetSelection();
+	TSharedRef<FNiagaraObjectSelection> GetNodeSelection();
 
 	/** Sets the currently selected graph nodes. */
 	void SetSelectedNodes(const TSet<UObject*>& InSelectedNodes);
@@ -94,8 +94,8 @@ private:
 	/** Commands for editing the graph. */
 	TSharedRef<FUICommandList> Commands;
 
-	/** The set of objects currently selected in the graph. */
-	TSharedRef<FNiagaraObjectSelection> Selection;
+	/** The set of nodes objects currently selected in the graph. */
+	TSharedRef<FNiagaraObjectSelection> NodeSelection;
 
 	/** A multicast delegate which is called whenever nodes are pasted into the graph. */
 	FOnNodesPasted OnNodesPastedDelegate;
diff --git a/Engine/Plugins/FX/Niagara/Source/NiagaraEditor/Private/ViewModels/NiagaraScriptViewModel.cpp b/Engine/Plugins/FX/Niagara/Source/NiagaraEditor/Private/ViewModels/NiagaraScriptViewModel.cpp
--- a/Engine/Plugins/FX/Niagara/Source/NiagaraEditor/Private/ViewModels/NiagaraScriptViewModel.cpp
+++ b/Engine/Plugins/FX/Niagara/Source/NiagaraEditor/Private/ViewModels/NiagaraScriptViewModel.cpp
@@ -8,7 +8,6 @@
 #include "NiagaraEmitter.h"
 #include "NiagaraScriptInputCollectionViewModel.h"
 #include "NiagaraScriptOutputCollectionViewModel.h"
-#include "NiagaraMetaDataCollectionViewModel.h"
 #include "NiagaraNodeInput.h"
 #include "INiagaraCompiler.h"
 #include "Kismet2/CompilerResultsLog.h"
@@ -23,8 +22,8 @@ template<> TMap<UNiagaraScript*, TArray<FNiagaraScriptViewModel*>> TNiagaraViewM
 FNiagaraScriptViewModel::FNiagaraScriptViewModel(UNiagaraScript* InScript, FText DisplayName, ENiagaraParameterEditMode InParameterEditMode)
 	: InputCollectionViewModel(MakeShareable(new FNiagaraScriptInputCollectionViewModel(InScript, DisplayName, InParameterEditMode)))
 	, OutputCollectionViewModel(MakeShareable(new FNiagaraScriptOutputCollectionViewModel(InScript, InParameterEditMode)))
-	, MetaDataCollectionViewModel(MakeShareable(new FNiagaraMetaDataCollectionViewModel()))
 	, GraphViewModel(MakeShareable(new FNiagaraScriptGraphViewModel(Cast<UNiagaraScriptSource>(InScript->GetSource()), DisplayName)))
+	, VariableSelection(MakeShareable(new FNiagaraObjectSelection()))
 	, bUpdatingSelectionInternally(false)
 	, LastCompileStatus(ENiagaraScriptCompileStatus::NCS_Unknown)
 {
@@ -38,15 +37,9 @@ FNiagaraScriptViewModel::FNiagaraScriptViewModel(UNiagaraScript* InScript, FText
 	InputCollectionViewModel->GetSelection().OnSelectedObjectsChanged().AddRaw(this, &FNiagaraScriptViewModel::InputViewModelSelectionChanged);
 	InputCollectionViewModel->OnParameterValueChanged().AddRaw(this, &FNiagaraScriptViewModel::InputParameterValueChanged);
 	OutputCollectionViewModel->OnParameterValueChanged().AddRaw(this, &FNiagaraScriptViewModel::OutputParameterValueChanged);
-	GraphViewModel->GetSelection()->OnSelectedObjectsChanged().AddRaw(this, &FNiagaraScriptViewModel::GraphViewModelSelectedNodesChanged);
+	GraphViewModel->GetNodeSelection()->OnSelectedObjectsChanged().AddRaw(this, &FNiagaraScriptViewModel::GraphViewModelSelectedNodesChanged);
 	GEditor->RegisterForUndo(this);
 
-	if (Source != nullptr)
-	{
-		UNiagaraGraph* Graph = Source->NodeGraph;
-		MetaDataCollectionViewModel->SetGraph(Graph);
-	}
-
 	// Guess at initial compile status
 	if (InScript != nullptr && InScript->GetVMExecutableData().IsValid() && InScript->GetVMExecutableData().ByteCode.Num() == 0) // This is either a brand new script or failed in the past. Since we create a default working script, assume invalid.
 	{
@@ -74,8 +67,8 @@ FNiagaraScriptViewModel::FNiagaraScriptViewModel(UNiagaraScript* InScript, FText
 FNiagaraScriptViewModel::FNiagaraScriptViewModel(UNiagaraEmitter* InEmitter, FText DisplayName, ENiagaraParameterEditMode InParameterEditMode)
 	: InputCollectionViewModel(MakeShareable(new FNiagaraScriptInputCollectionViewModel(InEmitter, DisplayName, InParameterEditMode)))
 	, OutputCollectionViewModel(MakeShareable(new FNiagaraScriptOutputCollectionViewModel(InEmitter, InParameterEditMode)))
-	, MetaDataCollectionViewModel(MakeShareable(new FNiagaraMetaDataCollectionViewModel()))
 	, GraphViewModel(MakeShareable(new FNiagaraScriptGraphViewModel(Cast<UNiagaraScriptSource>(InEmitter->GraphSource), DisplayName)))
+	, VariableSelection(MakeShareable(new FNiagaraObjectSelection()))
 	, bUpdatingSelectionInternally(false)
 	, LastCompileStatus(ENiagaraScriptCompileStatus::NCS_Unknown)
 {
@@ -89,17 +82,13 @@ FNiagaraScriptViewModel::FNiagaraScriptViewModel(UNiagaraEmitter* InEmitter, FTe
 		Scripts[i]->OnVMScriptCompiled().AddRaw(this, &FNiagaraScriptViewModel::OnVMScriptCompiled);
 	}
 	Source = Cast<UNiagaraScriptSource>(InEmitter->GraphSource);
-
+ 
 	InputCollectionViewModel->GetSelection().OnSelectedObjectsChanged().AddRaw(this, &FNiagaraScriptViewModel::InputViewModelSelectionChanged);
 	InputCollectionViewModel->OnParameterValueChanged().AddRaw(this, &FNiagaraScriptViewModel::InputParameterValueChanged);
 	OutputCollectionViewModel->OnParameterValueChanged().AddRaw(this, &FNiagaraScriptViewModel::OutputParameterValueChanged);
-	GraphViewModel->GetSelection()->OnSelectedObjectsChanged().AddRaw(this, &FNiagaraScriptViewModel::GraphViewModelSelectedNodesChanged);
+	GraphViewModel->GetNodeSelection()->OnSelectedObjectsChanged().AddRaw(this, &FNiagaraScriptViewModel::GraphViewModelSelectedNodesChanged);
 	GEditor->RegisterForUndo(this);
 
-	UNiagaraGraph* Graph = Source->NodeGraph;
-
-	MetaDataCollectionViewModel->SetGraph(Graph);
-
 	// Guess at initial compile status
 	LastCompileStatus = ENiagaraScriptCompileStatus::NCS_UpToDate;
 
@@ -204,7 +193,7 @@ bool FNiagaraScriptViewModel::IsGraphDirty() const
 FNiagaraScriptViewModel::~FNiagaraScriptViewModel()
 {
 	InputCollectionViewModel->GetSelection().OnSelectedObjectsChanged().RemoveAll(this);
-	GraphViewModel->GetSelection()->OnSelectedObjectsChanged().RemoveAll(this);
+	GraphViewModel->GetNodeSelection()->OnSelectedObjectsChanged().RemoveAll(this);
 
 	if (Source.IsValid() && Source != nullptr)
 	{
@@ -277,14 +266,6 @@ void FNiagaraScriptViewModel::SetScripts(UNiagaraScriptSource* InScriptSource, T
 	OutputCollectionViewModel->SetScripts(InScripts);
 	GraphViewModel->SetScriptSource(Source.Get());
 
-	UNiagaraGraph* Graph = nullptr;
-	if (Source.IsValid() && Source != nullptr)
-	{
-		// The underlying graph may have changed after the previous call.
-		Graph = Source->NodeGraph;
-	}
-	MetaDataCollectionViewModel->SetGraph(Graph);
-
 	// Guess at initial compile status
 	LastCompileStatus = ENiagaraScriptCompileStatus::NCS_UpToDate;
 
@@ -385,19 +366,14 @@ TSharedRef<FNiagaraScriptOutputCollectionViewModel> FNiagaraScriptViewModel::Get
 	return OutputCollectionViewModel;
 }
 
-TSharedRef<FNiagaraMetaDataCollectionViewModel> FNiagaraScriptViewModel::GetMetadataCollectionViewModel()
-{
-	return MetaDataCollectionViewModel;
-}
-
-void FNiagaraScriptViewModel::RefreshMetadataCollection()
+TSharedRef<FNiagaraScriptGraphViewModel> FNiagaraScriptViewModel::GetGraphViewModel()
 {
-	MetaDataCollectionViewModel->RequestRefresh();
+	return GraphViewModel;
 }
 
-TSharedRef<FNiagaraScriptGraphViewModel> FNiagaraScriptViewModel::GetGraphViewModel()
+TSharedRef<FNiagaraObjectSelection> FNiagaraScriptViewModel::GetVariableSelection()
 {
-	return GraphViewModel;
+	return VariableSelection;
 }
 
 UNiagaraScript* FNiagaraScriptViewModel::GetStandaloneScript()
@@ -563,7 +539,7 @@ void FNiagaraScriptViewModel::GraphViewModelSelectedNodesChanged()
 		bUpdatingSelectionInternally = true;
 		{
 			TSet<FName> SelectedInputIds;
-			for (UObject* SelectedObject : GraphViewModel->GetSelection()->GetSelectedObjects())
+			for (UObject* SelectedObject : GraphViewModel->GetNodeSelection()->GetSelectedObjects())
 			{
 				UNiagaraNodeInput* InputNode = Cast<UNiagaraNodeInput>(SelectedObject);
 				if (InputNode != nullptr)
@@ -613,7 +589,7 @@ void FNiagaraScriptViewModel::InputViewModelSelectionChanged()
 				}
 			}
 
-			GraphViewModel->GetSelection()->SetSelectedObjects(NodesToSelect);
+			GraphViewModel->GetNodeSelection()->SetSelectedObjects(NodesToSelect);
 		}
 		bUpdatingSelectionInternally = false;
 	}
diff --git a/Engine/Plugins/FX/Niagara/Source/NiagaraEditor/Private/ViewModels/NiagaraScriptViewModel.h b/Engine/Plugins/FX/Niagara/Source/NiagaraEditor/Private/ViewModels/NiagaraScriptViewModel.h
--- a/Engine/Plugins/FX/Niagara/Source/NiagaraEditor/Private/ViewModels/NiagaraScriptViewModel.h
+++ b/Engine/Plugins/FX/Niagara/Source/NiagaraEditor/Private/ViewModels/NiagaraScriptViewModel.h
@@ -16,6 +16,7 @@ class FNiagaraScriptInputCollectionViewModel;
 class FNiagaraScriptOutputCollectionViewModel;
 class FNiagaraMetaDataCollectionViewModel;
 class UNiagaraEmitter;
+class FNiagaraObjectSelection;
 
 
 /** A view model for Niagara scripts which manages other script related view models. */
@@ -38,18 +39,15 @@ public:
 	/** Gets the view model for the output parameter collection. */
 	TSharedRef<FNiagaraScriptOutputCollectionViewModel> GetOutputCollectionViewModel();
 
-	/** Gets the view model for the metadata collection. */
-	TSharedRef<FNiagaraMetaDataCollectionViewModel> GetMetadataCollectionViewModel();
-
-	/** Refreshes the metadata collection */
-	void RefreshMetadataCollection();
-
 	TSharedRef<INiagaraParameterCollectionViewModel> GetInputParameterMapViewModel();
 	TSharedRef<INiagaraParameterCollectionViewModel> GetOutputParameterMapViewModel();
 
 	/** Gets the view model for the graph. */
 	TSharedRef<FNiagaraScriptGraphViewModel> GetGraphViewModel();
 
+	/** Gets the currently selected script variables. */
+	TSharedRef<FNiagaraObjectSelection> GetVariableSelection();
+
 	/** Updates the script with the latest compile status. */
 	void UpdateCompileStatus(ENiagaraScriptCompileStatus InAggregateCompileStatus, const FString& InAggregateCompileErrors,
 		const TArray<ENiagaraScriptCompileStatus>& InCompileStatuses, const TArray<FString>& InCompileErrors, const TArray<FString>& InCompilePaths,
@@ -109,12 +107,12 @@ protected:
 	/** The view model for the output parameter collection .*/
 	TSharedRef<FNiagaraScriptOutputCollectionViewModel> OutputCollectionViewModel;
 
-	/** The view model for the metadata collection .*/
-	TSharedRef<FNiagaraMetaDataCollectionViewModel> MetaDataCollectionViewModel;
-
 	/** The view model for the graph. */
 	TSharedRef<FNiagaraScriptGraphViewModel> GraphViewModel;
 
+	/** The set of variables currently selected in the graph or the parameters panel. */
+	TSharedRef<FNiagaraObjectSelection> VariableSelection;
+
 	/** A flag for preventing reentrancy when synchronizing selection. */
 	bool bUpdatingSelectionInternally;
 
diff --git a/Engine/Plugins/FX/Niagara/Source/NiagaraEditor/Private/ViewModels/NiagaraSystemScriptViewModel.cpp b/Engine/Plugins/FX/Niagara/Source/NiagaraEditor/Private/ViewModels/NiagaraSystemScriptViewModel.cpp
--- a/Engine/Plugins/FX/Niagara/Source/NiagaraEditor/Private/ViewModels/NiagaraSystemScriptViewModel.cpp
+++ b/Engine/Plugins/FX/Niagara/Source/NiagaraEditor/Private/ViewModels/NiagaraSystemScriptViewModel.cpp
@@ -123,153 +123,6 @@ void FNiagaraSystemScriptViewModel::OnSystemVMCompiled(UNiagaraSystem* InSystem)
 	}
 }
 
-float EmitterNodeVerticalOffset = 150.0f;
-
-FVector2D CalculateNewEmitterNodePlacementPosition(UNiagaraGraph* Graph, UNiagaraNodeEmitter* NewEmitterNode)
-{
-	FVector2D PlacementLocation(0.0f, 0.0f);
-	TArray<UNiagaraNodeEmitter*> EmitterNodes;
-	Graph->GetNodesOfClass(EmitterNodes);
-	if (EmitterNodes.Num() > 1)
-	{
-		// If there are Emitter nodes, try to put it under the lowest one.
-		UNiagaraNodeEmitter* LowestNode = nullptr;
-		for (UNiagaraNodeEmitter* EmitterNode : EmitterNodes)
-		{
-			if (EmitterNode != NewEmitterNode && (LowestNode == nullptr || EmitterNode->NodePosY > LowestNode->NodePosY))
-			{
-				LowestNode = EmitterNode;
-			}
-		}
-		check(LowestNode);
-		PlacementLocation = FVector2D(
-			LowestNode->NodePosX,
-			LowestNode->NodePosY + EmitterNodeVerticalOffset);
-	}
-	return PlacementLocation;
-}
-
-void FNiagaraSystemScriptViewModel::RebuildEmitterNodes()
-{
-	UNiagaraGraph* SystemGraph = GetGraphViewModel()->GetGraph();
-	if (SystemGraph == nullptr)
-	{
-		return;
-	}
-	
-	TArray<UNiagaraNodeEmitter*> CurrentEmitterNodes;
-	SystemGraph->GetNodesOfClass<UNiagaraNodeEmitter>(CurrentEmitterNodes);
-
-	const UEdGraphSchema_Niagara* Schema = Cast<UEdGraphSchema_Niagara>(SystemGraph->GetSchema());
-
-	// Remove the old emitter nodes since they will be rebuilt below.
-	for (UNiagaraNodeEmitter* CurrentEmitterNode : CurrentEmitterNodes)
-	{
-		CurrentEmitterNode->Modify();
-		UEdGraphPin* InPin = CurrentEmitterNode->GetInputPin(0);
-		UEdGraphPin* OutPin = CurrentEmitterNode->GetOutputPin(0);
-		UEdGraphPin* InPinLinkedPin = InPin != nullptr && InPin->LinkedTo.Num() == 1 ? InPin->LinkedTo[0] : nullptr;
-		UEdGraphPin* OutPinLinkedPin = OutPin != nullptr && OutPin->LinkedTo.Num() == 1 ? OutPin->LinkedTo[0] : nullptr;
-		CurrentEmitterNode->DestroyNode();
-
-		if (InPinLinkedPin != nullptr &&& OutPinLinkedPin != nullptr)
-		{
-			InPinLinkedPin->MakeLinkTo(OutPinLinkedPin);
-		}
-	}
-
-	// Add output nodes if they don't exist.
-	TArray<UNiagaraNodeInput*> TempInputNodes;
-	TArray<UNiagaraNodeInput*> InputNodes;
-	TArray<UNiagaraNodeOutput*> OutputNodes;
-	OutputNodes.SetNum(2);
-	OutputNodes[0] = SystemGraph->FindOutputNode(ENiagaraScriptUsage::SystemSpawnScript);
-	OutputNodes[1] = SystemGraph->FindOutputNode(ENiagaraScriptUsage::SystemUpdateScript);
-
-	// Add input nodes if they don't exist
-	UNiagaraGraph::FFindInputNodeOptions Options;
-	Options.bFilterDuplicates = false;
-	Options.bIncludeParameters = true;
-	SystemGraph->FindInputNodes(TempInputNodes);
-	for (int32 i = 0; i < TempInputNodes.Num(); i++)
-	{
-		if (Schema->PinToTypeDefinition(TempInputNodes[i]->GetOutputPin(0)) == FNiagaraTypeDefinition::GetParameterMapDef())
-		{
-			InputNodes.Add(TempInputNodes[i]);
-		}
-	}
-
-	// Create a default id variable for the input nodes.
-	FNiagaraVariable SharedInputVar(FNiagaraTypeDefinition::GetParameterMapDef(), TEXT("InputMap"));	
-	InputNodes.SetNum(2);
-
-	// Now create the nodes if they are needed, synchronize if already created.
-	for (int32 i = 0; i < 2; i++)
-	{
-		if (OutputNodes[i] == nullptr)
-		{
-			FGraphNodeCreator<UNiagaraNodeOutput> OutputNodeCreator(*SystemGraph);
-			OutputNodes[i] = OutputNodeCreator.CreateNode();
-			OutputNodes[i]->SetUsage((ENiagaraScriptUsage)(i + (int32)ENiagaraScriptUsage::SystemSpawnScript));
-			
-			OutputNodes[i]->Outputs.Add(FNiagaraVariable(FNiagaraTypeDefinition::GetParameterMapDef(), TEXT("Out")));
-			OutputNodes[i]->NodePosX = 0;
-			OutputNodes[i]->NodePosY = 0 + i*25;
-
-			OutputNodeCreator.Finalize();
-		}
-		if (InputNodes[i] == nullptr)
-		{
-			FGraphNodeCreator<UNiagaraNodeInput> InputNodeCreator(*SystemGraph);
-			InputNodes[i] =  InputNodeCreator.CreateNode();
-			InputNodes[i]->Input = SharedInputVar;
-			InputNodes[i]->Usage = ENiagaraInputNodeUsage::Parameter;
-			InputNodes[i]->NodePosX = -50;
-			InputNodes[i]->NodePosY = 0 + i * 25;
-
-			InputNodeCreator.Finalize();
-
-			InputNodes[i]->GetOutputPin(0)->MakeLinkTo(OutputNodes[i]->GetInputPin(0));
-		}
-	}
-
-	// Add new nodes.
-	UNiagaraNode* TargetNodes[2];
-	TargetNodes[0] = OutputNodes[0];
-	TargetNodes[1] = OutputNodes[1];
-
-	for (const FNiagaraEmitterHandle& EmitterHandle : System.GetEmitterHandles())
-	{
-		for (int32 i = 0; i < 2; i++)
-		{
-			FGraphNodeCreator<UNiagaraNodeEmitter> EmitterNodeCreator(*SystemGraph);
-			UNiagaraNodeEmitter* EmitterNode = EmitterNodeCreator.CreateNode();
-			EmitterNode->SetOwnerSystem(&System);
-			EmitterNode->SetEmitterHandleId(EmitterHandle.GetId());
-			EmitterNode->SetUsage((ENiagaraScriptUsage)(i + (int32)ENiagaraScriptUsage::EmitterSpawnScript));
-
-			FVector2D NewLocation = CalculateNewEmitterNodePlacementPosition(SystemGraph, EmitterNode);
-			EmitterNode->NodePosX = NewLocation.X;
-			EmitterNode->NodePosY = NewLocation.Y;
-			EmitterNode->AllocateDefaultPins();
-			EmitterNodeCreator.Finalize();
-
-			TArray<FNiagaraStackGraphUtilities::FStackNodeGroup> StackNodeGroups;
-			FNiagaraStackGraphUtilities::GetStackNodeGroups(*OutputNodes[i], StackNodeGroups);
-
-			FNiagaraStackGraphUtilities::FStackNodeGroup EmitterGroup;
-			EmitterGroup.StartNodes.Add(EmitterNode);
-			EmitterGroup.EndNode = EmitterNode;
-
-			FNiagaraStackGraphUtilities::FStackNodeGroup& OutputGroup = StackNodeGroups[StackNodeGroups.Num() - 1];
-			FNiagaraStackGraphUtilities::FStackNodeGroup& OutputGroupPrevious = StackNodeGroups[StackNodeGroups.Num() - 2];
-			FNiagaraStackGraphUtilities::ConnectStackNodeGroup(EmitterGroup, OutputGroupPrevious, OutputGroup);
-		}
-	}
-	
-	FNiagaraStackGraphUtilities::RelayoutGraph(*SystemGraph);
-}
-
 FNiagaraSystemScriptViewModel::FOnSystemCompiled& FNiagaraSystemScriptViewModel::OnSystemCompiled()
 {
 	return OnSystemCompiledDelegate;
diff --git a/Engine/Plugins/FX/Niagara/Source/NiagaraEditor/Private/ViewModels/NiagaraSystemScriptViewModel.h b/Engine/Plugins/FX/Niagara/Source/NiagaraEditor/Private/ViewModels/NiagaraSystemScriptViewModel.h
--- a/Engine/Plugins/FX/Niagara/Source/NiagaraEditor/Private/ViewModels/NiagaraSystemScriptViewModel.h
+++ b/Engine/Plugins/FX/Niagara/Source/NiagaraEditor/Private/ViewModels/NiagaraSystemScriptViewModel.h
@@ -18,9 +18,6 @@ public:
 
 	~FNiagaraSystemScriptViewModel();
 
-	/** Rebuilds the emitter nodes in the System script due to data changes. */
-	void RebuildEmitterNodes();
-
 	FOnSystemCompiled& OnSystemCompiled();
 
 	void CompileSystem(bool bForce);
diff --git a/Engine/Plugins/FX/Niagara/Source/NiagaraEditor/Private/ViewModels/NiagaraSystemViewModel.cpp b/Engine/Plugins/FX/Niagara/Source/NiagaraEditor/Private/ViewModels/NiagaraSystemViewModel.cpp
--- a/Engine/Plugins/FX/Niagara/Source/NiagaraEditor/Private/ViewModels/NiagaraSystemViewModel.cpp
+++ b/Engine/Plugins/FX/Niagara/Source/NiagaraEditor/Private/ViewModels/NiagaraSystemViewModel.cpp
@@ -311,7 +311,7 @@ void FNiagaraSystemViewModel::AddEmitter(UNiagaraEmitter& Emitter)
 	EmitterHandle = System.AddEmitterHandle(Emitter, FNiagaraUtilities::GetUniqueName(Emitter.GetFName(), EmitterHandleNames));
 
 	check(SystemScriptViewModel.IsValid());
-	SystemScriptViewModel->RebuildEmitterNodes();
+	FNiagaraStackGraphUtilities::RebuildEmitterNodes(System);
 
 	if (EditMode == ENiagaraSystemViewModelEditMode::SystemAsset)
 	{
@@ -347,57 +347,57 @@ void FNiagaraSystemViewModel::AddEmitter(UNiagaraEmitter& Emitter)
 	bForceAutoCompileOnce = true;
 }
 
-void FNiagaraSystemViewModel::DuplicateEmitter(TSharedRef<FNiagaraEmitterHandleViewModel> EmitterHandleToDuplicate)
+void FNiagaraSystemViewModel::DuplicateEmitters(TArray<FEmitterHandleToDuplicate> EmitterHandlesToDuplicate)
 {
-	if (EmitterHandleToDuplicate->GetEmitterHandle() == nullptr)
+	if (EmitterHandlesToDuplicate.Num() <= 0)
 	{
 		return;
 	}
-	TSet<FGuid> HandlesToDuplicate;
-	HandlesToDuplicate.Add(EmitterHandleToDuplicate->GetEmitterHandle()->GetId());
-	DuplicateEmitters(HandlesToDuplicate);
-	bForceAutoCompileOnce = true;
-}
 
-void FNiagaraSystemViewModel::DuplicateEmitters(TSet<FGuid> EmitterHandleIdsToDuplicate)
-{
-	if (EmitterHandleIdsToDuplicate.Num() <= 0)
-	{
-		return;
-	}
 	// Kill all system instances before modifying the emitter handle list to prevent accessing deleted data.
 	KillSystemInstances();
-	const FScopedTransaction DeleteTransaction(EmitterHandleIdsToDuplicate.Num() == 1
+	const FScopedTransaction DeleteTransaction(EmitterHandlesToDuplicate.Num() == 1
 		? LOCTEXT("DuplicateEmitter", "Duplicate emitter")
 		: LOCTEXT("DuplicateEmitters", "Duplicate emitters"));
 
-
 	TSet<FName> EmitterHandleNames;
 	for (const FNiagaraEmitterHandle& EmitterHandle : System.GetEmitterHandles())
 	{
 		EmitterHandleNames.Add(EmitterHandle.GetName());
 	}
 
-	System.Modify();
-	for (FGuid OriginalEmitterHandleId : EmitterHandleIdsToDuplicate)
+	for (FEmitterHandleToDuplicate& EmitterHandleToDuplicate : EmitterHandlesToDuplicate)
 	{
-		FNiagaraEmitterHandle OriginalEmitterHandle;
-		for (const FNiagaraEmitterHandle& EmitterHandle : System.GetEmitterHandles())
+		FNiagaraEmitterHandle HandleToDuplicate;
+		for (TObjectIterator<UNiagaraSystem> OtherSystemIt; OtherSystemIt; ++OtherSystemIt)
 		{
-			if (EmitterHandle.GetId() == OriginalEmitterHandleId)
+			UNiagaraSystem* OtherSystem = *OtherSystemIt;
+			if (OtherSystem->GetPathName() == EmitterHandleToDuplicate.SystemPath)
+			{
+				for (const FNiagaraEmitterHandle& EmitterHandle : OtherSystem->GetEmitterHandles())
+				{
+					if (EmitterHandle.GetId() == EmitterHandleToDuplicate.EmitterHandleId)
+					{
+						HandleToDuplicate = EmitterHandle;
+						break;
+					}
+				}
+			}
+
+			if (HandleToDuplicate.IsValid())
 			{
-				OriginalEmitterHandle = EmitterHandle;
 				break;
 			}
 		}
-		if (OriginalEmitterHandle.IsValid())
+
+		if (HandleToDuplicate.IsValid())
 		{
-			FNiagaraEmitterHandle EmitterHandle = System.DuplicateEmitterHandle(OriginalEmitterHandle, FNiagaraUtilities::GetUniqueName(OriginalEmitterHandle.GetName(), EmitterHandleNames));
+			FNiagaraEmitterHandle EmitterHandle = System.DuplicateEmitterHandle(HandleToDuplicate, FNiagaraUtilities::GetUniqueName(HandleToDuplicate.GetName(), EmitterHandleNames));
+			EmitterHandleNames.Add(EmitterHandle.GetName());
 		}
 	}
 
-	check(SystemScriptViewModel.IsValid());
-	SystemScriptViewModel->RebuildEmitterNodes();
+	FNiagaraStackGraphUtilities::RebuildEmitterNodes(System);
 	RefreshAll();
 	bForceAutoCompileOnce = true;
 }
@@ -425,7 +425,7 @@ void FNiagaraSystemViewModel::DeleteEmitters(TSet<FGuid> EmitterHandleIdsToDelet
 		System.RemoveEmitterHandlesById(EmitterHandleIdsToDelete);
 
 		check(SystemScriptViewModel.IsValid());
-		SystemScriptViewModel->RebuildEmitterNodes();
+		FNiagaraStackGraphUtilities::RebuildEmitterNodes(System);
 
 		RefreshAll();
 		bForceAutoCompileOnce = true;
@@ -722,7 +722,7 @@ void FNiagaraSystemViewModel::RefreshEmitterHandleViewModels()
 	for (i = 0; i < System.GetNumEmitters(); ++i)
 	{
 		FNiagaraEmitterHandle* EmitterHandle = &System.GetEmitterHandle(i);
-		TSharedPtr<FNiagaraEmitterInstance> Simulation = SystemInstance ? SystemInstance->GetSimulationForHandle(*EmitterHandle) : nullptr;
+		TSharedPtr<FNiagaraEmitterInstance, ESPMode::ThreadSafe> Simulation = SystemInstance ? SystemInstance->GetSimulationForHandle(*EmitterHandle) : nullptr;
 		ValidEmitterHandleIds.Add(EmitterHandle->GetId());
 
 		bool bAdd = OldViewModels.Num() <= i;
@@ -1321,7 +1321,7 @@ void FNiagaraSystemViewModel::SequencerDataChanged(EMovieSceneDataChangeType Dat
 		GetOrCreateEditorData().SetPlaybackRange(TRange<float>(StartTimeSeconds, EndTimeSeconds));
 
 		TSet<FGuid> VaildTrackEmitterHandleIds;
-		TSet<FGuid> EmittersToDuplicate;
+		TArray<FEmitterHandleToDuplicate> EmittersToDuplicate;
 		TArray<TTuple<TSharedPtr<FNiagaraEmitterHandleViewModel>, FName>> EmitterHandlesToRename;
 
 		for (UMovieSceneTrack* Track : NiagaraSequence->GetMovieScene()->GetMasterTracks())
@@ -1343,7 +1343,10 @@ void FNiagaraSystemViewModel::SequencerDataChanged(EMovieSceneDataChangeType Dat
 				if (EmitterTrack->GetEmitterHandleId().IsValid())
 				{
 					// The emitter handle is invalid, but the track has a valid Id, most probably because of a copy/paste event
-					EmittersToDuplicate.Add(EmitterTrack->GetEmitterHandleId());
+					FEmitterHandleToDuplicate EmitterHandleToDuplicate;
+					EmitterHandleToDuplicate.SystemPath = EmitterTrack->GetSystemPath();
+					EmitterHandleToDuplicate.EmitterHandleId = EmitterTrack->GetEmitterHandleId();
+					EmittersToDuplicate.AddUnique(EmitterHandleToDuplicate);
 				}
 			}
 		}
@@ -1572,28 +1575,15 @@ void FNiagaraSystemViewModel::UpdateEmitterFixedBounds()
 		FNiagaraEmitterHandle* SelectedEmitterHandle = SelectedEmitterHandleViewModel->GetEmitterHandle();
 		check(SelectedEmitterHandle);
 		UNiagaraEmitter* Emitter = SelectedEmitterHandle->GetInstance();
-		for (TSharedRef<FNiagaraEmitterInstance>& EmitterInst : PreviewComponent->GetSystemInstance()->GetEmitters())
+		for (TSharedRef<FNiagaraEmitterInstance, ESPMode::ThreadSafe>& EmitterInst : PreviewComponent->GetSystemInstance()->GetEmitters())
 		{
 			if (&EmitterInst->GetEmitterHandle() == SelectedEmitterHandle && !EmitterInst->IsComplete())
 			{
-				TOptional<FBox> EmitterBounds = EmitterInst->CalculateDynamicBounds();
-				if (EmitterBounds.IsSet())
-				{
-					Emitter->Modify();
-					Emitter->bFixedBounds = true;
-					if (Emitter->bLocalSpace)
-					{
-						Emitter->FixedBounds = EmitterBounds.GetValue();
-					}
-					else
-					{
-						//Dynamic bounds are in world space. Transform back to local.
-						Emitter->FixedBounds = EmitterBounds.GetValue().TransformBy(PreviewComponent->GetComponentToWorld().Inverse());
-					}
-				}
+				EmitterInst->CalculateFixedBounds(PreviewComponent->GetComponentToWorld().Inverse());
 			}
 		}
 	}
+	PreviewComponent->MarkRenderTransformDirty();
 }
 
 void FNiagaraSystemViewModel::AddSystemEventHandlers()
diff --git a/Engine/Plugins/FX/Niagara/Source/NiagaraEditor/Private/ViewModels/Stack/NiagaraStackFunctionInput.cpp b/Engine/Plugins/FX/Niagara/Source/NiagaraEditor/Private/ViewModels/Stack/NiagaraStackFunctionInput.cpp
--- a/Engine/Plugins/FX/Niagara/Source/NiagaraEditor/Private/ViewModels/Stack/NiagaraStackFunctionInput.cpp
+++ b/Engine/Plugins/FX/Niagara/Source/NiagaraEditor/Private/ViewModels/Stack/NiagaraStackFunctionInput.cpp
@@ -40,13 +40,13 @@
 #include "ARFilter.h"
 #include "EdGraph/EdGraphPin.h"
 #include "INiagaraEditorTypeUtilities.h"
+#include "ViewModels/Stack/NiagaraStackInputCategory.h"
 
 #define LOCTEXT_NAMESPACE "NiagaraStackViewModel"
 
 UNiagaraStackFunctionInput::UNiagaraStackFunctionInput()
 	: OwningModuleNode(nullptr)
 	, OwningFunctionCallNode(nullptr)
-	, InputMetaData(nullptr)
 	, bUpdatingGraphDirectly(false)
 	, bUpdatingLocalValueDirectly(false)
 	, bShowEditConditionInline(false)
@@ -121,10 +121,12 @@ void UNiagaraStackFunctionInput::Initialize(
 	UNiagaraNodeFunctionCall& InInputFunctionCallNode,
 	FName InInputParameterHandle,
 	FNiagaraTypeDefinition InInputType,
+	EStackParameterBehavior InParameterBehavior,
 	FString InOwnerStackItemEditorDataKey)
 {
 	checkf(OwningModuleNode.IsValid() == false && OwningFunctionCallNode.IsValid() == false, TEXT("Can only initialize once."));
 	bool bInputIsAdvanced = false;
+	ParameterBehavior = InParameterBehavior;
 	FString InputStackEditorDataKey = FString::Printf(TEXT("%s-Input-%s"), *InInputFunctionCallNode.NodeGuid.ToString(EGuidFormats::DigitsWithHyphens), *InInputParameterHandle.ToString());
 	Super::Initialize(InRequiredEntryData, bInputIsAdvanced, InOwnerStackItemEditorDataKey, InputStackEditorDataKey);
 	OwningModuleNode = &InModuleNode;
@@ -280,10 +282,14 @@ FText UNiagaraStackFunctionInput::GetTooltipText(EValueMode InValueMode) const
 		}
 	}
 
-	const FNiagaraVariableMetaData* MetaData = nullptr;
+	TOptional<FNiagaraVariableMetaData> MetaData;
 	if (FNiagaraConstants::IsNiagaraConstant(ValueVariable))
 	{
-		MetaData = FNiagaraConstants::GetConstantMetaData(ValueVariable);
+		const FNiagaraVariableMetaData* FoundMetaData = FNiagaraConstants::GetConstantMetaData(ValueVariable);
+		if (FoundMetaData)
+		{
+			MetaData = *FoundMetaData;
+		}
 	}
 	else if (NodeGraph != nullptr)
 	{
@@ -291,7 +297,7 @@ FText UNiagaraStackFunctionInput::GetTooltipText(EValueMode InValueMode) const
 	}
 
 	FText Description = FText::GetEmpty();
-	if (MetaData != nullptr)
+	if (MetaData.IsSet())
 	{
 		Description = MetaData->Description;
 	}
@@ -554,7 +560,7 @@ UNiagaraDataInterface* UNiagaraStackFunctionInput::FInputValues::GetDataDefaultV
 
 void UNiagaraStackFunctionInput::RefreshValues()
 {
-	if (ensureMsgf(InputParameterHandle.IsModuleHandle(), TEXT("Function inputs can only be generated for module paramters.")) == false)
+	if (ensureMsgf(IsStaticParameter() || InputParameterHandle.IsModuleHandle(), TEXT("Function inputs can only be generated for module paramters.")) == false)
 	{
 		return;
 	}
@@ -603,7 +609,7 @@ void UNiagaraStackFunctionInput::RefreshFromMetaData()
 		FNiagaraVariable InputVariable(InputType, InputParameterHandle.GetParameterHandleString());
 		InputMetaData = FunctionGraph->GetMetaData(InputVariable);
 
-		if (InputMetaData != nullptr)
+		if (InputMetaData.IsSet())
 		{
 			SetIsAdvanced(InputMetaData->bAdvancedDisplay);
 
@@ -611,8 +617,8 @@ void UNiagaraStackFunctionInput::RefreshFromMetaData()
 			EditCondition.Refresh(InputMetaData->EditCondition, EditConditionError);
 			if (EditCondition.IsValid() && EditCondition.GetConditionInputType() == FNiagaraTypeDefinition::GetBoolDef())
 			{
-				FNiagaraVariableMetaData* EditConditionInputMetadata = EditCondition.GetConditionInputMetaData();
-				if (EditConditionInputMetadata != nullptr)
+				TOptional<FNiagaraVariableMetaData> EditConditionInputMetadata = EditCondition.GetConditionInputMetaData();
+				if (EditConditionInputMetadata.IsSet())
 				{
 					bShowEditConditionInline = EditConditionInputMetadata->bInlineEditConditionToggle;
 				}
@@ -946,16 +952,34 @@ void UNiagaraStackFunctionInput::NotifyEndLocalValueChange()
 
 bool UNiagaraStackFunctionInput::IsRapidIterationCandidate() const
 {
-	return FNiagaraStackGraphUtilities::IsRapidIterationType(InputType);
+	return !IsStaticParameter() && FNiagaraStackGraphUtilities::IsRapidIterationType(InputType);
 }
 
 void UNiagaraStackFunctionInput::SetLocalValue(TSharedRef<FStructOnScope> InLocalValue)
 {
 	TGuardValue<bool> UpdateGuard(bUpdatingLocalValueDirectly, true);
 
+	const UEdGraphSchema_Niagara* NiagaraSchema = GetDefault<UEdGraphSchema_Niagara>();
 	UEdGraphPin* DefaultPin = GetDefaultPin();
 	UEdGraphPin* OverridePin = GetOverridePin();
 	UEdGraphPin* ValuePin = DefaultPin;
+
+	// If the parameter we set is from a static switch then we don't want to use an override pin, but directly change the default pin.
+	if (IsStaticParameter())
+	{
+		FNiagaraVariable LocalValueVariable(InputType, NAME_None);
+		LocalValueVariable.SetData(InLocalValue->GetStructMemory());
+		FString PinDefaultValue;
+		if (ensureMsgf(NiagaraSchema->TryGetPinDefaultValueFromNiagaraVariable(LocalValueVariable, PinDefaultValue),
+			TEXT("Could not generate value string for static switch parameter.")))
+		{
+			DefaultPin->Modify();
+			DefaultPin->DefaultValue = PinDefaultValue;
+			Cast<UNiagaraNode>(DefaultPin->GetOwningNode())->MarkNodeRequiresSynchronization(TEXT("Default Value Changed"), true);
+		}
+		RefreshValues();
+		return;
+	}
 	
 	// If the default pin in the function graph is connected internally, rapid iteration parameters can't be used since
 	// the compilation currently won't use them.
@@ -985,7 +1009,6 @@ void UNiagaraStackFunctionInput::SetLocalValue(TSharedRef<FStructOnScope> InLoca
 		bGraphWillNeedRelayout = true;
 	}
 
-	const UEdGraphSchema_Niagara* NiagaraSchema = GetDefault<UEdGraphSchema_Niagara>();
 	if (bCanUseRapidIterationParameter)
 	{
 		for (TWeakObjectPtr<UNiagaraScript> Script : AffectedScripts)
@@ -1037,6 +1060,11 @@ void UNiagaraStackFunctionInput::SetLocalValue(TSharedRef<FStructOnScope> InLoca
 
 bool UNiagaraStackFunctionInput::CanReset() const
 {
+	if (IsStaticParameter())
+	{
+		// Static switch parameters only hold a single value in the default pin, so we disable resetting to prevent a special implementation for them
+		return false;
+	}
 	if (bCanResetCache.IsSet() == false)
 	{
 		bool bNewCanReset;
@@ -1230,6 +1258,11 @@ void UNiagaraStackFunctionInput::Reset()
 	RefreshChildren();
 }
 
+bool UNiagaraStackFunctionInput::IsStaticParameter() const
+{
+	return ParameterBehavior == EStackParameterBehavior::Static;
+}
+
 bool UNiagaraStackFunctionInput::CanResetToBase() const
 {
 	if (HasBaseEmitter())
@@ -1650,6 +1683,22 @@ UNiagaraNodeParameterMapSet& UNiagaraStackFunctionInput::GetOrCreateOverrideNode
 
 UEdGraphPin* UNiagaraStackFunctionInput::GetDefaultPin() const
 {
+	// If we have a static switch parameter, we check the pins of the function call node for a matching name,
+	// otherwise we search the parameter map node inside the function for a matching pin
+	if (IsStaticParameter())
+	{
+		const UEdGraphSchema_Niagara* NiagaraSchema = CastChecked<UEdGraphSchema_Niagara>(OwningFunctionCallNode->GetSchema());
+		FName InputName = InputParameterHandle.GetParameterHandleString();
+		for (UEdGraphPin* Pin : OwningFunctionCallNode->Pins)
+		{
+			FNiagaraVariable InputVariable = NiagaraSchema->PinToNiagaraVariable(Pin);
+			if (InputVariable.GetName() == InputName && InputVariable.GetType() == InputType)
+			{
+				return Pin;
+			}
+		}
+		return nullptr;
+	}
 	return OwningFunctionCallNode->FindParameterMapDefaultValuePin(InputParameterHandle.GetParameterHandleString(), SourceScript->GetUsage());
 }
 
diff --git a/Engine/Plugins/FX/Niagara/Source/NiagaraEditor/Private/ViewModels/Stack/NiagaraStackFunctionInputCollection.cpp b/Engine/Plugins/FX/Niagara/Source/NiagaraEditor/Private/ViewModels/Stack/NiagaraStackFunctionInputCollection.cpp
--- a/Engine/Plugins/FX/Niagara/Source/NiagaraEditor/Private/ViewModels/Stack/NiagaraStackFunctionInputCollection.cpp
+++ b/Engine/Plugins/FX/Niagara/Source/NiagaraEditor/Private/ViewModels/Stack/NiagaraStackFunctionInputCollection.cpp
@@ -98,39 +98,74 @@ void UNiagaraStackFunctionInputCollection::RefreshChildrenInternal(const TArray<
 	// Gather input data
 	for (const UEdGraphPin* InputPin : InputPins)
 	{
-		if (ProcessedInputNames.Contains(InputPin->PinName) == false)
+		if (ProcessedInputNames.Contains(InputPin->PinName))
 		{
-			ProcessedInputNames.Add(InputPin->PinName);
+			DuplicateInputNames.AddUnique(InputPin->PinName);
+			continue;
+		}
+		ProcessedInputNames.Add(InputPin->PinName);
 
-			FNiagaraVariable InputVariable = NiagaraSchema->PinToNiagaraVariable(InputPin);
-			if (InputVariable.GetType().IsValid() == false)
-			{
-				PinsWithInvalidTypes.Add(InputPin);
-				continue;
-			}
-			else
-			{
-				ValidAliasedInputNames.Add(FNiagaraParameterHandle::CreateAliasedModuleParameterHandle(
-					FNiagaraParameterHandle(InputPin->PinName), InputFunctionCallNode).GetParameterHandleString());
-			}
+		FNiagaraVariable InputVariable = NiagaraSchema->PinToNiagaraVariable(InputPin);
+		if (InputVariable.GetType().IsValid() == false)
+		{
+			PinsWithInvalidTypes.Add(InputPin);
+			continue;
+		}
+		ValidAliasedInputNames.Add(
+			FNiagaraParameterHandle::CreateAliasedModuleParameterHandle(FNiagaraParameterHandle(InputPin->PinName), InputFunctionCallNode).GetParameterHandleString());
 
-			FNiagaraVariableMetaData* InputMetaData = nullptr;
-			if (InputFunctionGraph != nullptr)
-			{
-				InputMetaData = InputFunctionGraph->GetMetaData(InputVariable);
-			}
+		TOptional<FNiagaraVariableMetaData> InputMetaData;
+		if (InputFunctionGraph != nullptr)
+		{
+			InputMetaData = InputFunctionGraph->GetMetaData(InputVariable);
+		}
 
-			FText InputCategory = InputMetaData != nullptr && InputMetaData->CategoryName.IsEmptyOrWhitespace() == false
-				? InputMetaData->CategoryName
-				: UncategorizedName;
+		FText InputCategory = InputMetaData.IsSet() && InputMetaData->CategoryName.IsEmptyOrWhitespace() == false
+			? InputMetaData->CategoryName
+			: UncategorizedName;
 
-			FInputData InputData = { InputPin, InputVariable.GetType(), InputMetaData ? InputMetaData->EditorSortPriority : 0, InputCategory };
-			InputDataCollection.Add(InputData);
+		FInputData InputData = { InputPin, InputVariable.GetType(), InputMetaData ? InputMetaData->EditorSortPriority : 0, InputCategory, false };
+		InputDataCollection.Add(InputData);
+	}
+
+	// Gather static switch parameters
+	InputPins.Reset();
+	FNiagaraStackGraphUtilities::GetStackFunctionStaticSwitchPins(*InputFunctionCallNode, InputPins);
+	for (const UEdGraphPin* InputPin : InputPins)
+	{
+		// The static switch pin names to not contain the module namespace, as they are not part of the parameter maps.
+		// We add it here only to check for name clashes with actual module parameters.
+		FString ModuleName = TEXT("Module.");
+		InputPin->PinName.AppendString(ModuleName);
+		FName SwitchPinName(*ModuleName); 
+
+		if (ProcessedInputNames.Contains(SwitchPinName))
+		{
+			DuplicateInputNames.AddUnique(SwitchPinName);
+			continue;
 		}
-		else
+		ProcessedInputNames.Add(SwitchPinName);
+
+		FNiagaraVariable InputVariable = NiagaraSchema->PinToNiagaraVariable(InputPin);
+		if (InputVariable.GetType().IsValid() == false)
 		{
-			DuplicateInputNames.AddUnique(InputPin->PinName);
+			PinsWithInvalidTypes.Add(InputPin);
+			continue;
+		}
+		ValidAliasedInputNames.Add(FNiagaraParameterHandle(*InputFunctionCallNode->GetFunctionName(), InputPin->PinName).GetParameterHandleString());
+
+		TOptional<FNiagaraVariableMetaData> InputMetaData;
+		if (InputFunctionGraph != nullptr)
+		{
+			InputMetaData = InputFunctionGraph->GetMetaData(InputVariable);
 		}
+
+		FText InputCategory = InputMetaData.IsSet() && InputMetaData->CategoryName.IsEmptyOrWhitespace() == false
+			? InputMetaData->CategoryName
+			: UncategorizedName;
+
+		FInputData InputData = { InputPin, InputVariable.GetType(), InputMetaData ? InputMetaData->EditorSortPriority : 0, InputCategory, true };
+		InputDataCollection.Add(InputData);
 	}
 
 	// Sort data and keep the uncategorized first
@@ -184,7 +219,7 @@ void UNiagaraStackFunctionInputCollection::RefreshChildrenInternal(const TArray<
 			}
 			NewChildren.Add(InputCategory);
 		}
-		InputCategory->AddInput(InputData.Pin->PinName, InputData.Type);
+		InputCategory->AddInput(InputData.Pin->PinName, InputData.Type, InputData.bIsStatic ? EStackParameterBehavior::Static : EStackParameterBehavior::Dynamic);
 	}
 	RefreshIssues(DuplicateInputNames, ValidAliasedInputNames, PinsWithInvalidTypes, NewIssues);
 }
@@ -249,7 +284,7 @@ void UNiagaraStackFunctionInputCollection::RefreshIssues(TArray<FName> Duplicate
 		FStackIssue DuplicateInputError(
 			EStackIssueSeverity::Error,
 			FText::Format(LOCTEXT("DuplicateInputSummaryFormat", "Duplicate Input: {0}"), FText::FromName(DuplicateInputName)),
-			FText::Format(LOCTEXT("DuplicateInputFormat", "There are multiple inputs with the same name {0}, but different types exposed by the function {1}.\nThis is not suppored and must be fixed in the script that defines this function."),
+			FText::Format(LOCTEXT("DuplicateInputFormat", "There are multiple inputs with the same name {0} exposed by the function {1}.\nThis is not supported and must be fixed in the script that defines this function.\nCheck for inputs with the same name and different types or static switches."),
 				FText::FromName(DuplicateInputName), FText::FromString(InputFunctionCallNode->GetFunctionName())),
 			GetStackEditorDataKey(),
 			false);
diff --git a/Engine/Plugins/FX/Niagara/Source/NiagaraEditor/Private/ViewModels/Stack/NiagaraStackFunctionInputCondition.cpp b/Engine/Plugins/FX/Niagara/Source/NiagaraEditor/Private/ViewModels/Stack/NiagaraStackFunctionInputCondition.cpp
--- a/Engine/Plugins/FX/Niagara/Source/NiagaraEditor/Private/ViewModels/Stack/NiagaraStackFunctionInputCondition.cpp
+++ b/Engine/Plugins/FX/Niagara/Source/NiagaraEditor/Private/ViewModels/Stack/NiagaraStackFunctionInputCondition.cpp
@@ -93,9 +93,9 @@ bool FNiagaraStackFunctionInputCondition::GetConditionIsEnabled() const
 	if (IsValid())
 	{
 		TArray<uint8> InputValue = InputBinder.GetData();
-		bool bMatchFound = TargetValuesData.FindByPredicate([&InputValue](const TArray<uint8>& TargetValueData) {
-			return FMemory::Memcmp(TargetValueData.GetData(), InputValue.GetData(), TargetValueData.Num()) == 0; }) != nullptr;
-		return bMatchFound;
+		return nullptr != TargetValuesData.FindByPredicate([&InputValue](const TArray<uint8>& TargetValueData) 
+			{ return FMemory::Memcmp(TargetValueData.GetData(), InputValue.GetData(), TargetValueData.Num()) == 0; }
+		);
 	}
 	return false;
 }
@@ -123,7 +123,7 @@ FNiagaraTypeDefinition FNiagaraStackFunctionInputCondition::GetConditionInputTyp
 	return InputBinder.GetInputType();
 }
 
-FNiagaraVariableMetaData* FNiagaraStackFunctionInputCondition::GetConditionInputMetaData() const
+TOptional<FNiagaraVariableMetaData> FNiagaraStackFunctionInputCondition::GetConditionInputMetaData() const
 {
 	checkf(IsValid(), TEXT("Can not get the input metadata for an invalid input condition"));
 	FNiagaraVariable InputVariable(InputBinder.GetInputType(), InputBinder.GetInputName());
