Copying //UE4/Dev-Niagara to Dev-Main (//UE4/Dev-Main)

diff --git a/Engine/Plugins/FX/Niagara/Source/NiagaraEditor/Private/ViewModels/Stack/NiagaraStackGraphUtilities.cpp b/Engine/Plugins/FX/Niagara/Source/NiagaraEditor/Private/ViewModels/Stack/NiagaraStackGraphUtilities.cpp
--- a/Engine/Plugins/FX/Niagara/Source/NiagaraEditor/Private/ViewModels/Stack/NiagaraStackGraphUtilities.cpp
+++ b/Engine/Plugins/FX/Niagara/Source/NiagaraEditor/Private/ViewModels/Stack/NiagaraStackGraphUtilities.cpp
@@ -13,6 +13,7 @@
 #include "NiagaraNodeCustomHlsl.h"
 #include "NiagaraNodeParameterMapGet.h"
 #include "NiagaraNodeParameterMapSet.h"
+#include "NiagaraNodeEmitter.h"
 #include "NiagaraScriptSource.h"
 #include "EdGraphSchema_Niagara.h"
 #include "ViewModels/Stack/NiagaraStackEntry.h"
@@ -257,14 +258,15 @@ UNiagaraNodeFunctionCall* FNiagaraStackGraphUtilities::GetNextModuleNode(UNiagar
 	return nullptr;
 }
 
-UNiagaraNodeOutput* FNiagaraStackGraphUtilities::GetEmitterOutputNodeForStackNode(UNiagaraNode& StackNode)
+template<typename OutputNodeType, typename InputNodeType>
+OutputNodeType* GetEmitterOutputNodeForStackNodeInternal(InputNodeType& StackNode)
 {
-	TArray<UNiagaraNode*> NodesToCheck;
-	TSet<UNiagaraNode*> NodesChecked;
+	TArray<InputNodeType*> NodesToCheck;
+	TSet<InputNodeType*> NodesChecked;
 	NodesToCheck.Add(&StackNode);
 	while (NodesToCheck.Num() > 0)
 	{
-		UNiagaraNode* NodeToCheck = NodesToCheck[0];
+		InputNodeType* NodeToCheck = NodesToCheck[0];
 		NodesToCheck.RemoveAt(0);
 		NodesChecked.Add(NodeToCheck);
 
@@ -273,13 +275,13 @@ UNiagaraNodeOutput* FNiagaraStackGraphUtilities::GetEmitterOutputNodeForStackNod
 			return CastChecked<UNiagaraNodeOutput>(NodeToCheck);
 		}
 		
-		TArray<UEdGraphPin*> OutputPins;
+		TArray<const UEdGraphPin*> OutputPins;
 		NodeToCheck->GetOutputPins(OutputPins);
-		for (UEdGraphPin* OutputPin : OutputPins)
+		for (const UEdGraphPin* OutputPin : OutputPins)
 		{
 			for (UEdGraphPin* LinkedPin : OutputPin->LinkedTo)
 			{
-				UNiagaraNode* LinkedNiagaraNode = Cast<UNiagaraNode>(LinkedPin->GetOwningNode());
+				InputNodeType* LinkedNiagaraNode = Cast<UNiagaraNode>(LinkedPin->GetOwningNode());
 				if (LinkedNiagaraNode != nullptr && NodesChecked.Contains(LinkedNiagaraNode) == false)
 				{
 					NodesToCheck.Add(LinkedNiagaraNode);
@@ -290,6 +292,16 @@ UNiagaraNodeOutput* FNiagaraStackGraphUtilities::GetEmitterOutputNodeForStackNod
 	return nullptr;
 }
 
+UNiagaraNodeOutput* FNiagaraStackGraphUtilities::GetEmitterOutputNodeForStackNode(UNiagaraNode& StackNode)
+{
+	return GetEmitterOutputNodeForStackNodeInternal<UNiagaraNodeOutput, UNiagaraNode>(StackNode);
+}
+
+const UNiagaraNodeOutput* FNiagaraStackGraphUtilities::GetEmitterOutputNodeForStackNode(const UNiagaraNode& StackNode)
+{
+	return GetEmitterOutputNodeForStackNodeInternal<const UNiagaraNodeOutput, const UNiagaraNode>(StackNode);
+}
+
 UNiagaraNodeInput* FNiagaraStackGraphUtilities::GetEmitterInputNodeForStackNode(UNiagaraNode& StackNode)
 {
 	// Since the stack graph can have arbitrary branches when traversing inputs, the only safe way to get the initial input
@@ -514,6 +526,31 @@ void FNiagaraStackGraphUtilities::GetStackFunctionInputPins(UNiagaraNodeFunction
 	}
 }
 
+void FNiagaraStackGraphUtilities::GetStackFunctionStaticSwitchPins(UNiagaraNodeFunctionCall& FunctionCallNode, TArray<const UEdGraphPin*>& OutInputPins)
+{
+	const UEdGraphSchema_Niagara* Schema = CastChecked<UEdGraphSchema_Niagara>(FunctionCallNode.GetSchema());
+	UNiagaraGraph* FunctionCallGraph = FunctionCallNode.GetCalledGraph();
+	if (FunctionCallGraph == nullptr)
+		return;
+
+	for (FNiagaraVariable SwitchInput : FunctionCallGraph->FindStaticSwitchInputs())
+	{
+		FEdGraphPinType PinType = Schema->TypeDefinitionToPinType(SwitchInput.GetType());
+		for (UEdGraphPin* Pin : FunctionCallNode.Pins)
+		{
+			if (Pin->Direction != EEdGraphPinDirection::EGPD_Input)
+			{
+				continue;
+			}
+			if (Pin->PinName.IsEqual(SwitchInput.GetName()) && Pin->PinType == PinType)
+			{
+				OutInputPins.Add(Pin);
+				break;
+			}
+		}
+	}
+}
+
 UNiagaraNodeParameterMapSet* FNiagaraStackGraphUtilities::GetStackFunctionOverrideNode(UNiagaraNodeFunctionCall& FunctionCallNode)
 {
 	UEdGraphPin* ParameterMapInput = FNiagaraStackGraphUtilities::GetParameterMapInputPin(FunctionCallNode);
@@ -1105,8 +1142,7 @@ UNiagaraNodeOutput* FNiagaraStackGraphUtilities::ResetGraphForOutput(UNiagaraGra
 		UNiagaraSystem* System = NiagaraGraph.GetTypedOuter<UNiagaraSystem>();
 		if (System != nullptr && System->GetEmitterHandles().Num() != 0)
 		{
-			TSharedRef<FNiagaraSystemScriptViewModel> SystemScriptViewModel = MakeShared<FNiagaraSystemScriptViewModel>(*System, nullptr);
-			SystemScriptViewModel->RebuildEmitterNodes();
+			RebuildEmitterNodes(*System);
 		}
 	}
 
@@ -1289,7 +1325,7 @@ void FNiagaraStackGraphUtilities::GetAvailableParametersForScript(UNiagaraNodeOu
 	TOptional<FName> UsageNamespace = FNiagaraStackGraphUtilities::GetNamespaceForScriptUsage(ScriptOutputNode.GetUsage());
 	if (UsageNamespace.IsSet())
 	{
-		for (const TPair<FNiagaraVariable, FNiagaraGraphParameterReferenceCollection>& Entry : ScriptOutputNode.GetNiagaraGraph()->GetParameterMap())
+		for (const TPair<FNiagaraVariable, FNiagaraGraphParameterReferenceCollection>& Entry : ScriptOutputNode.GetNiagaraGraph()->GetParameterReferenceMap())
 		{
 			// Pick up any params with 0 references from the Parameters window
 			bool bDoesParamHaveNoReferences = Entry.Value.ParameterReferences.Num() == 0;
@@ -1912,4 +1948,122 @@ bool FNiagaraStackGraphUtilities::ParameterAllowedInExecutionCategory(const FNam
 	return true;
 }
 
+void FNiagaraStackGraphUtilities::RebuildEmitterNodes(UNiagaraSystem& System)
+{
+	UNiagaraScriptSource* SystemScriptSource = Cast<UNiagaraScriptSource>(System.GetSystemSpawnScript()->GetSource());
+	UNiagaraGraph* SystemGraph = SystemScriptSource->NodeGraph;
+	if (SystemGraph == nullptr)
+	{
+		return;
+	}
+
+	TArray<UNiagaraNodeEmitter*> CurrentEmitterNodes;
+	SystemGraph->GetNodesOfClass<UNiagaraNodeEmitter>(CurrentEmitterNodes);
+
+	const UEdGraphSchema_Niagara* Schema = Cast<UEdGraphSchema_Niagara>(SystemGraph->GetSchema());
+
+	// Remove the old emitter nodes since they will be rebuilt below.
+	for (UNiagaraNodeEmitter* CurrentEmitterNode : CurrentEmitterNodes)
+	{
+		CurrentEmitterNode->Modify();
+		UEdGraphPin* InPin = CurrentEmitterNode->GetInputPin(0);
+		UEdGraphPin* OutPin = CurrentEmitterNode->GetOutputPin(0);
+		UEdGraphPin* InPinLinkedPin = InPin != nullptr && InPin->LinkedTo.Num() == 1 ? InPin->LinkedTo[0] : nullptr;
+		UEdGraphPin* OutPinLinkedPin = OutPin != nullptr && OutPin->LinkedTo.Num() == 1 ? OutPin->LinkedTo[0] : nullptr;
+		CurrentEmitterNode->DestroyNode();
+
+		if (InPinLinkedPin != nullptr &&& OutPinLinkedPin != nullptr)
+		{
+			InPinLinkedPin->MakeLinkTo(OutPinLinkedPin);
+		}
+	}
+
+	// Add output nodes if they don't exist.
+	TArray<UNiagaraNodeInput*> TempInputNodes;
+	TArray<UNiagaraNodeInput*> InputNodes;
+	TArray<UNiagaraNodeOutput*> OutputNodes;
+	OutputNodes.SetNum(2);
+	OutputNodes[0] = SystemGraph->FindOutputNode(ENiagaraScriptUsage::SystemSpawnScript);
+	OutputNodes[1] = SystemGraph->FindOutputNode(ENiagaraScriptUsage::SystemUpdateScript);
+
+	// Add input nodes if they don't exist
+	UNiagaraGraph::FFindInputNodeOptions Options;
+	Options.bFilterDuplicates = false;
+	Options.bIncludeParameters = true;
+	SystemGraph->FindInputNodes(TempInputNodes);
+	for (int32 i = 0; i < TempInputNodes.Num(); i++)
+	{
+		if (Schema->PinToTypeDefinition(TempInputNodes[i]->GetOutputPin(0)) == FNiagaraTypeDefinition::GetParameterMapDef())
+		{
+			InputNodes.Add(TempInputNodes[i]);
+		}
+	}
+
+	// Create a default id variable for the input nodes.
+	FNiagaraVariable SharedInputVar(FNiagaraTypeDefinition::GetParameterMapDef(), TEXT("InputMap"));
+	InputNodes.SetNum(2);
+
+	// Now create the nodes if they are needed, synchronize if already created.
+	for (int32 i = 0; i < 2; i++)
+	{
+		if (OutputNodes[i] == nullptr)
+		{
+			FGraphNodeCreator<UNiagaraNodeOutput> OutputNodeCreator(*SystemGraph);
+			OutputNodes[i] = OutputNodeCreator.CreateNode();
+			OutputNodes[i]->SetUsage((ENiagaraScriptUsage)(i + (int32)ENiagaraScriptUsage::SystemSpawnScript));
+
+			OutputNodes[i]->Outputs.Add(FNiagaraVariable(FNiagaraTypeDefinition::GetParameterMapDef(), TEXT("Out")));
+			OutputNodes[i]->NodePosX = 0;
+			OutputNodes[i]->NodePosY = 0 + i * 25;
+
+			OutputNodeCreator.Finalize();
+		}
+		if (InputNodes[i] == nullptr)
+		{
+			FGraphNodeCreator<UNiagaraNodeInput> InputNodeCreator(*SystemGraph);
+			InputNodes[i] = InputNodeCreator.CreateNode();
+			InputNodes[i]->Input = SharedInputVar;
+			InputNodes[i]->Usage = ENiagaraInputNodeUsage::Parameter;
+			InputNodes[i]->NodePosX = -50;
+			InputNodes[i]->NodePosY = 0 + i * 25;
+
+			InputNodeCreator.Finalize();
+
+			InputNodes[i]->GetOutputPin(0)->MakeLinkTo(OutputNodes[i]->GetInputPin(0));
+		}
+	}
+
+	// Add new nodes.
+	UNiagaraNode* TargetNodes[2];
+	TargetNodes[0] = OutputNodes[0];
+	TargetNodes[1] = OutputNodes[1];
+
+	for (const FNiagaraEmitterHandle& EmitterHandle : System.GetEmitterHandles())
+	{
+		for (int32 i = 0; i < 2; i++)
+		{
+			FGraphNodeCreator<UNiagaraNodeEmitter> EmitterNodeCreator(*SystemGraph);
+			UNiagaraNodeEmitter* EmitterNode = EmitterNodeCreator.CreateNode();
+			EmitterNode->SetOwnerSystem(&System);
+			EmitterNode->SetEmitterHandleId(EmitterHandle.GetId());
+			EmitterNode->SetUsage((ENiagaraScriptUsage)(i + (int32)ENiagaraScriptUsage::EmitterSpawnScript));
+			EmitterNode->AllocateDefaultPins();
+			EmitterNodeCreator.Finalize();
+
+			TArray<FNiagaraStackGraphUtilities::FStackNodeGroup> StackNodeGroups;
+			FNiagaraStackGraphUtilities::GetStackNodeGroups(*OutputNodes[i], StackNodeGroups);
+
+			FNiagaraStackGraphUtilities::FStackNodeGroup EmitterGroup;
+			EmitterGroup.StartNodes.Add(EmitterNode);
+			EmitterGroup.EndNode = EmitterNode;
+
+			FNiagaraStackGraphUtilities::FStackNodeGroup& OutputGroup = StackNodeGroups[StackNodeGroups.Num() - 1];
+			FNiagaraStackGraphUtilities::FStackNodeGroup& OutputGroupPrevious = StackNodeGroups[StackNodeGroups.Num() - 2];
+			FNiagaraStackGraphUtilities::ConnectStackNodeGroup(EmitterGroup, OutputGroupPrevious, OutputGroup);
+		}
+	}
+
+	RelayoutGraph(*SystemGraph);
+}
+
 #undef LOCTEXT_NAMESPACE
diff --git a/Engine/Plugins/FX/Niagara/Source/NiagaraEditor/Private/ViewModels/Stack/NiagaraStackInputCategory.cpp b/Engine/Plugins/FX/Niagara/Source/NiagaraEditor/Private/ViewModels/Stack/NiagaraStackInputCategory.cpp
--- a/Engine/Plugins/FX/Niagara/Source/NiagaraEditor/Private/ViewModels/Stack/NiagaraStackInputCategory.cpp
+++ b/Engine/Plugins/FX/Niagara/Source/NiagaraEditor/Private/ViewModels/Stack/NiagaraStackInputCategory.cpp
@@ -33,9 +33,9 @@ void UNiagaraStackInputCategory::ResetInputs()
 	Inputs.Empty();
 }
 
-void UNiagaraStackInputCategory::AddInput(FName InInputParameterHandle, FNiagaraTypeDefinition InInputType)
+void UNiagaraStackInputCategory::AddInput(FName InInputParameterHandle, FNiagaraTypeDefinition InInputType, EStackParameterBehavior InParameterBehavior)
 {
-	Inputs.Add({ InInputParameterHandle, InInputType });
+	Inputs.Add({ InInputParameterHandle, InInputType, InParameterBehavior });
 }
 
 void UNiagaraStackInputCategory::RefreshChildrenInternal(const TArray<UNiagaraStackEntry*>& CurrentChildren, TArray<UNiagaraStackEntry*>& NewChildren, TArray<FStackIssue>& NewIssues)
@@ -49,7 +49,7 @@ void UNiagaraStackInputCategory::RefreshChildrenInternal(const TArray<UNiagaraSt
 		{
 			InputChild = NewObject<UNiagaraStackFunctionInput>(this);
 			InputChild->Initialize(CreateDefaultChildRequiredData(), *ModuleNode, *InputFunctionCallNode,
-				Input.ParameterHandle, Input.Type, GetOwnerStackItemEditorDataKey());
+				Input.ParameterHandle, Input.Type, Input.ParameterBehavior, GetOwnerStackItemEditorDataKey());
 		}
 
 		NewChildren.Add(InputChild);
diff --git a/Engine/Plugins/FX/Niagara/Source/NiagaraEditor/Private/ViewModels/Stack/NiagaraStackModuleItem.cpp b/Engine/Plugins/FX/Niagara/Source/NiagaraEditor/Private/ViewModels/Stack/NiagaraStackModuleItem.cpp
--- a/Engine/Plugins/FX/Niagara/Source/NiagaraEditor/Private/ViewModels/Stack/NiagaraStackModuleItem.cpp
+++ b/Engine/Plugins/FX/Niagara/Source/NiagaraEditor/Private/ViewModels/Stack/NiagaraStackModuleItem.cpp
@@ -32,6 +32,7 @@
 #include "Widgets/SWidget.h"
 
 #include "ScopedTransaction.h"
+#include "NiagaraScriptVariable.h"
 
 // TODO: Remove these
 #include "Widgets/Notifications/SNotificationList.h"
@@ -180,11 +181,11 @@ void UNiagaraStackModuleItem::RefreshChildrenInternal(const TArray<UNiagaraStack
 		{
 			UNiagaraScriptSource* Source = CastChecked<UNiagaraScriptSource>(FunctionCallNode->FunctionScript->GetSource());
 			UNiagaraGraph* Graph = CastChecked<UNiagaraGraph>(Source->NodeGraph);
-			const TMap<FNiagaraVariable, FNiagaraVariableMetaData>& MetaDataMap = Graph->GetAllMetaData();
+			const TMap<FNiagaraVariable, UNiagaraScriptVariable*>& MetaDataMap = Graph->GetAllMetaData();
 			auto Iter = MetaDataMap.CreateConstIterator();
 			while (Iter)
 			{
-				auto PropertyIter = Iter.Value().PropertyMetaData.CreateConstIterator();
+				auto PropertyIter = Iter.Value()->Metadata.PropertyMetaData.CreateConstIterator();
 				while (PropertyIter)
 				{
 					if (PropertyIter.Key() == (TEXT("DisplayNameArg0")))
diff --git a/Engine/Plugins/FX/Niagara/Source/NiagaraEditor/Private/ViewModels/Stack/NiagaraStackModuleItemOutput.cpp b/Engine/Plugins/FX/Niagara/Source/NiagaraEditor/Private/ViewModels/Stack/NiagaraStackModuleItemOutput.cpp
--- a/Engine/Plugins/FX/Niagara/Source/NiagaraEditor/Private/ViewModels/Stack/NiagaraStackModuleItemOutput.cpp
+++ b/Engine/Plugins/FX/Niagara/Source/NiagaraEditor/Private/ViewModels/Stack/NiagaraStackModuleItemOutput.cpp
@@ -39,17 +39,17 @@ FText UNiagaraStackModuleItemOutput::GetTooltipText() const
 	{
 		UNiagaraScriptSource* Source = Cast<UNiagaraScriptSource>(FunctionCallNode->FunctionScript->GetSource());
 		const UEdGraphSchema_Niagara* NiagaraSchema = GetDefault<UEdGraphSchema_Niagara>();
-		const FNiagaraVariableMetaData* MetaData = nullptr;
+		TOptional<FNiagaraVariableMetaData> MetaData;
 		if (FNiagaraConstants::IsNiagaraConstant(ValueVariable))
 		{
-			MetaData = FNiagaraConstants::GetConstantMetaData(ValueVariable);
+			MetaData = *FNiagaraConstants::GetConstantMetaData(ValueVariable);
 		}
 		else if (Source->NodeGraph != nullptr)
 		{
 			MetaData = Source->NodeGraph->GetMetaData(ValueVariable);
 		}
 
-		if (MetaData != nullptr)
+		if (MetaData.IsSet())
 		{
 			return MetaData->Description;
 		}
diff --git a/Engine/Plugins/FX/Niagara/Source/NiagaraEditor/Private/ViewModels/Stack/NiagaraStackRendererItem.cpp b/Engine/Plugins/FX/Niagara/Source/NiagaraEditor/Private/ViewModels/Stack/NiagaraStackRendererItem.cpp
--- a/Engine/Plugins/FX/Niagara/Source/NiagaraEditor/Private/ViewModels/Stack/NiagaraStackRendererItem.cpp
+++ b/Engine/Plugins/FX/Niagara/Source/NiagaraEditor/Private/ViewModels/Stack/NiagaraStackRendererItem.cpp
@@ -279,6 +279,23 @@ void UNiagaraStackRendererItem::RefreshIssues(TArray<FStackIssue>& NewIssues)
 
 		NewIssues.Add(TargetSupportError);
 	}
+
+	// Currently we don't support multiple different renderers on GPU emitters
+	// Multiple sprite renderers may currently work but not multiple of other types or combinations of types and there is no guarantee
+	if ( GetEmitterViewModel()->GetEmitter()->SimTarget == ENiagaraSimTarget::GPUComputeSim )
+	{
+		if ( GetEmitterViewModel()->GetEmitter()->GetRenderers().Num() > 1 )
+		{
+			FStackIssue GpuSupportError(
+				EStackIssueSeverity::Error,
+				LOCTEXT("FailedRendererDueToMultipleGpuRenderers", "Multiple GPU renderers."),
+				LOCTEXT("FailedRendererDueToMultipleGpuRenderersLong", "Multiple GPU renderers are currently not supported, rendered results may not be correct."),
+				GetStackEditorDataKey(),
+				false);
+
+			NewIssues.Add(GpuSupportError);
+		}
+	}
 }
 
 void UNiagaraStackRendererItem::RendererChanged()
diff --git a/Engine/Plugins/FX/Niagara/Source/NiagaraEditor/Private/Widgets/SNiagaraParameterMapView.cpp b/Engine/Plugins/FX/Niagara/Source/NiagaraEditor/Private/Widgets/SNiagaraParameterMapView.cpp
--- a/Engine/Plugins/FX/Niagara/Source/NiagaraEditor/Private/Widgets/SNiagaraParameterMapView.cpp
+++ b/Engine/Plugins/FX/Niagara/Source/NiagaraEditor/Private/Widgets/SNiagaraParameterMapView.cpp
@@ -33,6 +33,10 @@
 #include "EdGraphSchema_Niagara.h"
 #include "ViewModels/NiagaraEmitterHandleViewModel.h"
 #include "ScopedTransaction.h"
+#include "NiagaraEditorUtilities.h"
+#include "NiagaraScriptVariable.h"
+#include "Widgets/SPanel.h"
+#include "Widgets/Layout/SBox.h"
 
 #define LOCTEXT_NAMESPACE "NiagaraParameterMapView"
 
@@ -51,6 +55,9 @@ FText NiagaraParameterMapSectionID::OnGetSectionTitle(const NiagaraParameterMapS
 	case NiagaraParameterMapSectionID::MODULE:
 		SectionTitle = NSLOCTEXT("GraphActionNode", "Module", "Module");
 		break;
+	case NiagaraParameterMapSectionID::STATIC_SWITCH:
+		SectionTitle = NSLOCTEXT("GraphActionNode", "StaticSwitch", "Static Switch");
+		break;
 	case NiagaraParameterMapSectionID::SYSTEM:
 		SectionTitle = NSLOCTEXT("GraphActionNode", "System", "System");
 		break;
@@ -73,11 +80,15 @@ FText NiagaraParameterMapSectionID::OnGetSectionTitle(const NiagaraParameterMapS
 	return SectionTitle;
 }
 
-NiagaraParameterMapSectionID::Type NiagaraParameterMapSectionID::OnGetSectionFromVariable(const FNiagaraVariable& InVar, FNiagaraParameterHandle& OutParameterHandle, const NiagaraParameterMapSectionID::Type DefaultType)
+NiagaraParameterMapSectionID::Type NiagaraParameterMapSectionID::OnGetSectionFromVariable(const FNiagaraVariable& InVar, bool IsStaticSwitchVariable, FNiagaraParameterHandle& OutParameterHandle, const NiagaraParameterMapSectionID::Type DefaultType)
 {
 	OutParameterHandle = FNiagaraParameterHandle(InVar.GetName());
 	Type SectionID = DefaultType;
-	if (OutParameterHandle.IsEmitterHandle())
+	if (IsStaticSwitchVariable)
+	{
+		SectionID = NiagaraParameterMapSectionID::STATIC_SWITCH;
+	}
+	else if (OutParameterHandle.IsEmitterHandle())
 	{
 		SectionID = NiagaraParameterMapSectionID::EMITTER;
 	}
@@ -111,7 +122,7 @@ NiagaraParameterMapSectionID::Type NiagaraParameterMapSectionID::OnGetSectionFro
 
 void FNiagaraParameterMapViewCommands::RegisterCommands()
 {
-	UI_COMMAND(DeleteEntry, "Delete", "Deletes this parameter.", EUserInterfaceActionType::Button, FInputChord(EKeys::Platform_Delete));
+	UI_COMMAND(DeleteEntry, "Delete", "Delete this parameter", EUserInterfaceActionType::Button, FInputChord(EKeys::Platform_Delete));
 }
 
 SNiagaraParameterMapView::~SNiagaraParameterMapView()
@@ -120,7 +131,7 @@ SNiagaraParameterMapView::~SNiagaraParameterMapView()
 	ToolkitCommands->UnmapAction(FNiagaraParameterMapViewCommands::Get().DeleteEntry);
 	ToolkitCommands->UnmapAction(FGenericCommands::Get().Rename);
 
-	TSet<UObject*> Objects = SelectedObjects->GetSelectedObjects();
+	TSet<UObject*> Objects = SelectedScriptObjects->GetSelectedObjects();
 	for (UObject* Object : Objects)
 	{
 		if (UNiagaraSystem* System = Cast<UNiagaraSystem>(Object))
@@ -131,18 +142,27 @@ SNiagaraParameterMapView::~SNiagaraParameterMapView()
 	}
 
 	Graphs.Empty();
-	SelectedObjects->OnSelectedObjectsChanged().RemoveAll(this);
+	SelectedScriptObjects->OnSelectedObjectsChanged().RemoveAll(this);
+	if (SelectedVariableObjects)
+	{
+		SelectedVariableObjects->OnSelectedObjectsChanged().RemoveAll(this);
+	}
 }
 
-void SNiagaraParameterMapView::Construct(const FArguments& InArgs, const TSharedRef<FNiagaraObjectSelection>& InSelectedObjects, const EToolkitType InToolkitType, const TSharedPtr<FUICommandList>& InToolkitCommands)
+void SNiagaraParameterMapView::Construct(const FArguments& InArgs, const TArray<TSharedRef<FNiagaraObjectSelection>>& InSelectedObjects, const EToolkitType InToolkitType, const TSharedPtr<FUICommandList>& InToolkitCommands)
 {
 	bNeedsRefresh = false;
 	ToolkitType = InToolkitType;
 	ToolkitCommands = InToolkitCommands;
 	AddParameterButtons.SetNum(NiagaraParameterMapSectionID::OTHER + 1);
 
-	SelectedObjects = InSelectedObjects;
-	SelectedObjects->OnSelectedObjectsChanged().AddSP(this, &SNiagaraParameterMapView::SelectedObjectsChanged);
+	SelectedScriptObjects = InSelectedObjects[0];
+	SelectedScriptObjects->OnSelectedObjectsChanged().AddSP(this, &SNiagaraParameterMapView::SelectedObjectsChanged);
+	if (InSelectedObjects.Num() == 2)
+	{
+		//SelectedVariableObjects->OnSelectedObjectsChanged().AddSP(this, &SNiagaraParameterMapView::SelectedObjectsChanged);
+		SelectedVariableObjects = InSelectedObjects[1];
+	}
 	
 	// Register all commands for right click on action node
 	{
@@ -241,7 +261,7 @@ void SNiagaraParameterMapView::AddParameter(FNiagaraVariable NewVariable)
 	for (auto& GraphWeakPtr : Graphs)
 	{
 		UNiagaraGraph* Graph = GraphWeakPtr.Get();
-		for (const auto& ParameterElement : Graph->GetParameterMap())
+		for (const auto& ParameterElement : Graph->GetParameterReferenceMap())
 		{
 			Names.Add(ParameterElement.Key.GetName());
 		}
@@ -252,9 +272,9 @@ void SNiagaraParameterMapView::AddParameter(FNiagaraVariable NewVariable)
 	bool bAddedParameter = false;
 	// Check whether we have to add this parameter to the user exposed system parameters.
 	FNiagaraParameterHandle ParameterHandle;
-	if (NiagaraParameterMapSectionID::OnGetSectionFromVariable(NewVariable, ParameterHandle) == NiagaraParameterMapSectionID::USER)
+	if (NiagaraParameterMapSectionID::OnGetSectionFromVariable(NewVariable, IsStaticSwitchParameter(NewVariable, Graphs), ParameterHandle) == NiagaraParameterMapSectionID::USER)
 	{
-		for (UObject* Object : SelectedObjects->GetSelectedObjects())
+		for (UObject* Object : SelectedScriptObjects->GetSelectedObjects())
 		{
 			if (UNiagaraSystem* System = Cast<UNiagaraSystem>(Object))
 			{
@@ -312,7 +332,8 @@ void SNiagaraParameterMapView::CollectAllActions(FGraphActionListBuilderBase& Ou
 	for (auto& GraphWeakPtr : Graphs)
 	{
 		UNiagaraGraph* Graph = GraphWeakPtr.Get();
-		for (const auto& ParameterElement : Graph->GetParameterMap())
+		for (const auto& ParameterElement : Graph->GetParameterReferenceMap())
+		// for (const auto& ParameterElement : FNiagaraEditorUtilities::GetCompiledGraphParameterMapReferences(Graph)) // TODO(mv): This makes unused module parameters not show up. See CL 6182457.
 		{
 			TArray<FNiagaraGraphParameterReferenceCollection>* Found = ParameterEntries.Find(ParameterElement.Key);
 			if (Found)
@@ -328,7 +349,7 @@ void SNiagaraParameterMapView::CollectAllActions(FGraphActionListBuilderBase& Ou
 		}
 	}
 
-	TSet<UObject*> Objects = SelectedObjects->GetSelectedObjects();
+	TSet<UObject*> Objects = SelectedScriptObjects->GetSelectedObjects();
 	for (UObject* Object : Objects)
 	{
 		if (UNiagaraSystem* System = Cast<UNiagaraSystem>(Object))
@@ -354,7 +375,7 @@ void SNiagaraParameterMapView::CollectAllActions(FGraphActionListBuilderBase& Ou
 	{
 		const FNiagaraVariable& Parameter = ParameterEntry.Key;
 		FNiagaraParameterHandle Handle;
-		const NiagaraParameterMapSectionID::Type Section = NiagaraParameterMapSectionID::OnGetSectionFromVariable(Parameter, Handle, /*Default*/ NiagaraParameterMapSectionID::OTHER);
+		const NiagaraParameterMapSectionID::Type Section = NiagaraParameterMapSectionID::OnGetSectionFromVariable(Parameter, IsStaticSwitchParameter(Parameter, Graphs), Handle, /*Default*/ NiagaraParameterMapSectionID::OTHER);
 		if (!IsSystemToolkit() || (IsSystemToolkit() && Section != NiagaraParameterMapSectionID::MODULE))
 		{
 			const FText Name = FText::FromName(Parameter.GetName());
@@ -363,6 +384,8 @@ void SNiagaraParameterMapView::CollectAllActions(FGraphActionListBuilderBase& Ou
 			OutAllActions.AddAction(ParameterAction);
 		}
 	}
+
+
 }
 
 void SNiagaraParameterMapView::CollectStaticSections(TArray<int32>& StaticSectionIDs)
@@ -371,6 +394,7 @@ void SNiagaraParameterMapView::CollectStaticSections(TArray<int32>& StaticSectio
 	{
 		StaticSectionIDs.Add(NiagaraParameterMapSectionID::MODULE);
 	}
+	StaticSectionIDs.Add(NiagaraParameterMapSectionID::STATIC_SWITCH);
 	StaticSectionIDs.Add(NiagaraParameterMapSectionID::ENGINE);
 	StaticSectionIDs.Add(NiagaraParameterMapSectionID::PARAMETERCOLLECTION);
 	StaticSectionIDs.Add(NiagaraParameterMapSectionID::USER);
@@ -408,7 +432,30 @@ FReply SNiagaraParameterMapView::OnActionDragged(const TArray<TSharedPtr<FEdGrap
 
 void SNiagaraParameterMapView::OnActionSelected(const TArray< TSharedPtr<FEdGraphSchemaAction> >& InActions, ESelectInfo::Type InSelectionType)
 {
-
+	if (!IsScriptToolkit())
+	{
+		return;
+	}
+	
+	// TODO: Can there be multiple actions and graphs? 
+	if (InActions.Num() == 1 && InActions[0].IsValid() && Graphs[0].IsValid()) 
+	{
+		FNiagaraParameterAction* Action = (FNiagaraParameterAction*)InActions[0].Get();
+		if (Action && (Action->Parameter.IsInNameSpace(TEXT("Module")) || IsStaticSwitchParameter(Action->Parameter, Graphs)))
+		{
+			// Ignore any non-Module variables, as they don't persist, i.e. they get
+			// purged when any variable is renamed due to not being in ParameterToReferencesMap
+			TMap<FNiagaraVariable, UNiagaraScriptVariable*>& AllMetaData = Graphs[0]->GetAllMetaData();
+			if (UNiagaraScriptVariable** FoundScriptVariable = AllMetaData.Find(Action->Parameter))
+			{
+				SelectedVariableObjects->SetSelectedObject(*FoundScriptVariable);
+				return;
+			}
+		}
+	} 
+	
+	// If a variable wasn't selected just clear the current selection
+	SelectedVariableObjects->ClearSelectedObjects();
 }
 
 void SNiagaraParameterMapView::OnActionDoubleClicked(const TArray< TSharedPtr<FEdGraphSchemaAction> >& InActions)
@@ -442,6 +489,13 @@ FText SNiagaraParameterMapView::OnGetSectionTitle(int32 InSectionID)
 
 TSharedRef<SWidget> SNiagaraParameterMapView::OnGetSectionWidget(TSharedRef<SWidget> RowWidget, int32 InSectionID)
 {
+	if (InSectionID == NiagaraParameterMapSectionID::STATIC_SWITCH)
+	{
+		TSharedPtr<SBox> Empty;
+		SAssignNew(Empty, SBox);
+		return Empty.ToSharedRef();
+	}
+
 	TWeakPtr<SWidget> WeakRowWidget = RowWidget;
 	FText AddNewText = LOCTEXT("AddNewParameter", "Add Parameter");
 	FName MetaDataTag = TEXT("AddNewParameter");
@@ -494,6 +548,14 @@ bool SNiagaraParameterMapView::SelectionHasContextMenu() const
 {
 	TArray<TSharedPtr<FEdGraphSchemaAction>> SelectedActions;
 	GraphActionMenu->GetSelectedActions(SelectedActions);
+	for (TSharedPtr<FEdGraphSchemaAction> Action : SelectedActions)
+	{
+		TSharedPtr<FNiagaraParameterAction> NiagaraAction = StaticCastSharedPtr<FNiagaraParameterAction>(Action);
+		if (NiagaraAction && IsStaticSwitchParameter(NiagaraAction->GetParameter(), Graphs))
+		{
+			return false;
+		}
+	}
 	return SelectedActions.Num() > 0;
 }
 
@@ -520,7 +582,7 @@ void SNiagaraParameterMapView::Refresh(bool bRefreshMenu/* = true*/)
 {
 	Graphs.Empty();
 
-	TSet<UObject*> Objects = SelectedObjects->GetSelectedObjects();
+	TSet<UObject*> Objects = SelectedScriptObjects->GetSelectedObjects();
 	for (UObject* Object : Objects)
 	{
 		if (UNiagaraScript* Script = Cast<UNiagaraScript>(Object))
@@ -596,20 +658,20 @@ void SNiagaraParameterMapView::OnDeleteEntry()
 	TArray<TSharedPtr<FEdGraphSchemaAction>> SelectedActions;
 	GraphActionMenu->GetSelectedActions(SelectedActions);
 	
-	TArray<UNiagaraGraph*> GraphsToNotifyChanged;
 	for (auto& Action : SelectedActions)
 	{
 		TSharedPtr<FNiagaraParameterAction> ParameterAction = StaticCastSharedPtr<FNiagaraParameterAction>(Action);
 		if (ParameterAction.Get())
 		{
 			FNiagaraParameterHandle ParameterHandle;
-			if (NiagaraParameterMapSectionID::OnGetSectionFromVariable(ParameterAction->GetParameter(), ParameterHandle) == NiagaraParameterMapSectionID::USER)
+			FNiagaraVariable Parameter = ParameterAction->GetParameter();
+			if (NiagaraParameterMapSectionID::OnGetSectionFromVariable(Parameter, IsStaticSwitchParameter(Parameter, Graphs), ParameterHandle) == NiagaraParameterMapSectionID::USER)
 			{
-				for (UObject* Object : SelectedObjects->GetSelectedObjects())
+				for (UObject* Object : SelectedScriptObjects->GetSelectedObjects())
 				{
 					if (UNiagaraSystem* System = Cast<UNiagaraSystem>(Object))
 					{
-						System->GetExposedParameters().RemoveParameter(ParameterAction->GetParameter());
+						System->GetExposedParameters().RemoveParameter(Parameter);
 					}
 				}
 			}
@@ -620,17 +682,11 @@ void SNiagaraParameterMapView::OnDeleteEntry()
 				if (GraphWeakPtr.IsValid())
 				{
 					UNiagaraGraph* Graph = GraphWeakPtr.Get();
-					Graph->RemoveParameter(ParameterAction->GetParameter(), false);
-					GraphsToNotifyChanged.Add(Graph);
+					Graph->RemoveParameter(Parameter);
 				}
 			}
 		}
 	}
-
-	for (UNiagaraGraph* Graph : GraphsToNotifyChanged)
-	{
-		Graph->NotifyGraphChanged();
-	}
 }
 
 bool SNiagaraParameterMapView::CanDeleteEntry() const
@@ -669,13 +725,13 @@ void SNiagaraParameterMapView::OnPostRenameActionNode(const FText& InText, FNiag
 	if (!InAction.Parameter.GetName().IsEqual(NewName, ENameCase::CaseSensitive))
 	{
 		FNiagaraParameterHandle ParameterHandle;
-		if (NiagaraParameterMapSectionID::OnGetSectionFromVariable(InAction.Parameter, ParameterHandle) == NiagaraParameterMapSectionID::USER)
+		if (NiagaraParameterMapSectionID::OnGetSectionFromVariable(InAction.Parameter, IsStaticSwitchParameter(InAction.Parameter, Graphs), ParameterHandle) == NiagaraParameterMapSectionID::USER)
 		{
 			// Check if the new name is also an user variable.
 			const FNiagaraVariable NewParameterValidTest = FNiagaraVariable(InAction.Parameter.GetType(), NewName);
-			if (NiagaraParameterMapSectionID::OnGetSectionFromVariable(NewParameterValidTest, ParameterHandle) == NiagaraParameterMapSectionID::USER)
+			if (NiagaraParameterMapSectionID::OnGetSectionFromVariable(NewParameterValidTest, IsStaticSwitchParameter(NewParameterValidTest, Graphs), ParameterHandle) == NiagaraParameterMapSectionID::USER)
 			{
-				for (UObject* Object : SelectedObjects->GetSelectedObjects())
+				for (UObject* Object : SelectedScriptObjects->GetSelectedObjects())
 				{
 					if (UNiagaraSystem* System = Cast<UNiagaraSystem>(Object))
 					{
@@ -718,6 +774,22 @@ void SNiagaraParameterMapView::RefreshActions()
 	bNeedsRefresh = true;
 }
 
+bool SNiagaraParameterMapView::IsStaticSwitchParameter(const FNiagaraVariable& Variable, const TArray<TWeakObjectPtr<UNiagaraGraph>>& Graphs)
+{
+	for (auto& GraphWeakPtr : Graphs)
+	{
+		if (UNiagaraGraph* Graph = GraphWeakPtr.Get())
+		{
+			TArray<FNiagaraVariable> SwitchInputs = Graph->FindStaticSwitchInputs();
+			if (SwitchInputs.Contains(Variable))
+			{
+				return true;
+			}
+		}
+	}
+	return false;
+}
+
 /************************************************************************/
 /* SNiagaraAddParameterMenu                                             */
 /************************************************************************/
@@ -904,14 +976,15 @@ void SNiagaraAddParameterMenu::CollectAllActions(FGraphActionListBuilderBase& Ou
 	{
 		for (TWeakObjectPtr<UNiagaraGraph>& Graph : Graphs)
 		{
-			TMap<FNiagaraVariable, FNiagaraGraphParameterReferenceCollection> ParameterEntries = Graph.Get()->GetParameterMap();
+			TMap<FNiagaraVariable, FNiagaraGraphParameterReferenceCollection> ParameterEntries = Graph.Get()->GetParameterReferenceMap();
 			ParameterEntries.KeySort([](const FNiagaraVariable& A, const FNiagaraVariable& B) { return (A.GetName() < B.GetName()); });
 
 			for (const auto& ParameterEntry : ParameterEntries)
 			{
 				const FNiagaraVariable& Parameter = ParameterEntry.Key;
 				FNiagaraParameterHandle Handle;
-				const NiagaraParameterMapSectionID::Type ParameterSectionID = NiagaraParameterMapSectionID::OnGetSectionFromVariable(Parameter, Handle, /*Default*/ NiagaraParameterMapSectionID::OTHER);
+				bool IsStaticSwitch = SNiagaraParameterMapView::IsStaticSwitchParameter(Parameter, Graphs);
+				const NiagaraParameterMapSectionID::Type ParameterSectionID = NiagaraParameterMapSectionID::OnGetSectionFromVariable(Parameter, IsStaticSwitch, Handle, /*Default*/ NiagaraParameterMapSectionID::OTHER);
 				if (CanCollectSection(ParameterSectionID))
 				{
 					if (IDsExcluded.Contains(ParameterSectionID))
diff --git a/Engine/Plugins/FX/Niagara/Source/NiagaraEditor/Private/Widgets/SNiagaraParameterMapView.h b/Engine/Plugins/FX/Niagara/Source/NiagaraEditor/Private/Widgets/SNiagaraParameterMapView.h
--- a/Engine/Plugins/FX/Niagara/Source/NiagaraEditor/Private/Widgets/SNiagaraParameterMapView.h
+++ b/Engine/Plugins/FX/Niagara/Source/NiagaraEditor/Private/Widgets/SNiagaraParameterMapView.h
@@ -31,6 +31,7 @@ namespace NiagaraParameterMapSectionID
 	{
 		NONE = 0,
 		MODULE,
+		STATIC_SWITCH,
 		ENGINE,
 		PARAMETERCOLLECTION,
 		USER,
@@ -41,7 +42,7 @@ namespace NiagaraParameterMapSectionID
 	};
 
 	static FText OnGetSectionTitle(const NiagaraParameterMapSectionID::Type InSection);
-	static NiagaraParameterMapSectionID::Type OnGetSectionFromVariable(const FNiagaraVariable& InVar, FNiagaraParameterHandle& OutParameterHandle, const NiagaraParameterMapSectionID::Type DefaultType = NiagaraParameterMapSectionID::Type::NONE);
+	static NiagaraParameterMapSectionID::Type OnGetSectionFromVariable(const FNiagaraVariable& InVar, bool IsStaticSwitchVariable, FNiagaraParameterHandle& OutParameterHandle, const NiagaraParameterMapSectionID::Type DefaultType = NiagaraParameterMapSectionID::Type::NONE);
 };
 
 class FNiagaraParameterMapViewCommands : public TCommands<FNiagaraParameterMapViewCommands>
@@ -76,11 +77,11 @@ public:
 
 	virtual ~SNiagaraParameterMapView();
 
-	void Construct(const FArguments& InArgs, const TSharedRef<FNiagaraObjectSelection>& InSelectedObjects, const EToolkitType InToolkitType, const TSharedPtr<FUICommandList>& InToolkitCommands);
+	void Construct(const FArguments& InArgs, const TArray<TSharedRef<FNiagaraObjectSelection>>& InSelectedObjects, const EToolkitType InToolkitType, const TSharedPtr<FUICommandList>& InToolkitCommands);
 
 	virtual void Tick(const FGeometry& AllottedGeometry, const double InCurrentTime, const float InDeltaTime) override;
 
-	/** Wheter the add parameter button should be enabled. */
+	/** Whether the add parameter button should be enabled. */
 	bool ParameterAddEnabled() const;
 
 	/** Adds parameter to the graph parameter store and refreshes the menu. */
@@ -92,6 +93,8 @@ public:
 
 	static TSharedRef<SExpanderArrow> CreateCustomActionExpander(const struct FCustomExpanderData& ActionMenuData);
 
+	static bool IsStaticSwitchParameter(const FNiagaraVariable& Variable, const TArray<TWeakObjectPtr<UNiagaraGraph>>& Graphs);
+
 private:
 	/** Function to bind to SNiagaraAddParameterMenus to filter types we allow creating */
 	bool AllowMakeType(const FNiagaraTypeDefinition& InType) const;
@@ -152,7 +155,8 @@ private:
 	TArray<TSharedPtr<SComboButton>> AddParameterButtons;
 
 	/** The selected objects being viewed and edited by this widget. */
-	TSharedPtr<FNiagaraObjectSelection> SelectedObjects;
+	TSharedPtr<FNiagaraObjectSelection> SelectedScriptObjects;
+	TSharedPtr<FNiagaraObjectSelection> SelectedVariableObjects;
 
 	TArray<TWeakObjectPtr<UNiagaraGraph>> Graphs;
 
diff --git a/Engine/Plugins/FX/Niagara/Source/NiagaraEditor/Private/Widgets/SNiagaraScriptGraph.cpp b/Engine/Plugins/FX/Niagara/Source/NiagaraEditor/Private/Widgets/SNiagaraScriptGraph.cpp
--- a/Engine/Plugins/FX/Niagara/Source/NiagaraEditor/Private/Widgets/SNiagaraScriptGraph.cpp
+++ b/Engine/Plugins/FX/Niagara/Source/NiagaraEditor/Private/Widgets/SNiagaraScriptGraph.cpp
@@ -25,7 +25,7 @@
 void SNiagaraScriptGraph::Construct(const FArguments& InArgs, TSharedRef<FNiagaraScriptGraphViewModel> InViewModel)
 {
 	ViewModel = InViewModel;
-	ViewModel->GetSelection()->OnSelectedObjectsChanged().AddSP(this, &SNiagaraScriptGraph::ViewModelSelectedNodesChanged);
+	ViewModel->GetNodeSelection()->OnSelectedObjectsChanged().AddSP(this, &SNiagaraScriptGraph::ViewModelSelectedNodesChanged);
 	ViewModel->OnNodesPasted().AddSP(this, &SNiagaraScriptGraph::NodesPasted);
 	ViewModel->OnGraphChanged().AddSP(this, &SNiagaraScriptGraph::GraphChanged);
 	bUpdatingGraphSelectionFromViewModel = false;
@@ -93,11 +93,11 @@ TSharedRef<SGraphEditor> SNiagaraScriptGraph::ConstructGraphEditor()
 
 void SNiagaraScriptGraph::ViewModelSelectedNodesChanged()
 {
-	if (FNiagaraEditorUtilities::SetsMatch(GraphEditor->GetSelectedNodes(), ViewModel->GetSelection()->GetSelectedObjects()) == false)
+	if (FNiagaraEditorUtilities::SetsMatch(GraphEditor->GetSelectedNodes(), ViewModel->GetNodeSelection()->GetSelectedObjects()) == false)
 	{
 		bUpdatingGraphSelectionFromViewModel = true;
 		GraphEditor->ClearSelectionSet();
-		for (UObject* SelectedNode : ViewModel->GetSelection()->GetSelectedObjects())
+		for (UObject* SelectedNode : ViewModel->GetNodeSelection()->GetSelectedObjects())
 		{
 			UEdGraphNode* GraphNode = Cast<UEdGraphNode>(SelectedNode);
 			if (GraphNode != nullptr)
@@ -113,7 +113,14 @@ void SNiagaraScriptGraph::GraphEditorSelectedNodesChanged(const TSet<UObject*>&
 {
 	if (bUpdatingGraphSelectionFromViewModel == false)
 	{
-		ViewModel->GetSelection()->SetSelectedObjects(SelectedNodes);
+		if (SelectedNodes.Num() == 0)
+		{
+			ViewModel->GetNodeSelection()->ClearSelectedObjects();
+		} 
+		else 
+		{
+			ViewModel->GetNodeSelection()->SetSelectedObjects(SelectedNodes);
+		}
 	}
 }
 
diff --git a/Engine/Plugins/FX/Niagara/Source/NiagaraEditor/Private/Widgets/SNiagaraSelectedObjectsDetails.cpp b/Engine/Plugins/FX/Niagara/Source/NiagaraEditor/Private/Widgets/SNiagaraSelectedObjectsDetails.cpp
--- a/Engine/Plugins/FX/Niagara/Source/NiagaraEditor/Private/Widgets/SNiagaraSelectedObjectsDetails.cpp
+++ b/Engine/Plugins/FX/Niagara/Source/NiagaraEditor/Private/Widgets/SNiagaraSelectedObjectsDetails.cpp
@@ -11,24 +11,49 @@
 
 void SNiagaraSelectedObjectsDetails::Construct(const FArguments& InArgs, TSharedRef<FNiagaraObjectSelection> InSelectedObjects)
 {
-	SelectedObjects = InSelectedObjects;
-	SelectedObjects->OnSelectedObjectsChanged().AddSP(this, &SNiagaraSelectedObjectsDetails::SelectedObjectsChanged);
+	SelectedObjectsArray.Push(InSelectedObjects);
+	SelectedObjectsArray[0]->OnSelectedObjectsChanged().AddSP(this, &SNiagaraSelectedObjectsDetails::SelectedObjectsChanged);
 
 	FPropertyEditorModule& PropertyEditorModule = FModuleManager::GetModuleChecked<FPropertyEditorModule>("PropertyEditor");
 	FDetailsViewArgs DetailsViewArgs(false, false, true, FDetailsViewArgs::HideNameArea, true);
 	DetailsView = PropertyEditorModule.CreateDetailView(DetailsViewArgs);
-	DetailsView->SetObjects(SelectedObjects->GetSelectedObjects().Array());
+	DetailsView->SetObjects(SelectedObjectsArray[0]->GetSelectedObjects().Array());
 	DetailsView->OnFinishedChangingProperties().AddRaw(this, &SNiagaraSelectedObjectsDetails::OnDetailsPanelFinishedChangingProperties);
 
 	ChildSlot
 	[
 		DetailsView.ToSharedRef()
 	];
-}
+} 
+
+void SNiagaraSelectedObjectsDetails::Construct(const FArguments& InArgs, TSharedRef<FNiagaraObjectSelection> InSelectedObjects, TSharedRef<FNiagaraObjectSelection> InSelectedObjects2)
+{
+	SelectedObjectsArray.Push(InSelectedObjects);
+	SelectedObjectsArray.Push(InSelectedObjects2);
+	SelectedObjectsArray[0]->OnSelectedObjectsChanged().AddSP(this, &SNiagaraSelectedObjectsDetails::SelectedObjectsChanged);
+	SelectedObjectsArray[1]->OnSelectedObjectsChanged().AddSP(this, &SNiagaraSelectedObjectsDetails::SelectedObjectsChangedSecond);
+
+	FPropertyEditorModule& PropertyEditorModule = FModuleManager::GetModuleChecked<FPropertyEditorModule>("PropertyEditor");
+	FDetailsViewArgs DetailsViewArgs(false, false, true, FDetailsViewArgs::HideNameArea, true);
+	DetailsView = PropertyEditorModule.CreateDetailView(DetailsViewArgs);
+	DetailsView->SetObjects(SelectedObjectsArray[0]->GetSelectedObjects().Array());
+	DetailsView->OnFinishedChangingProperties().AddRaw(this, &SNiagaraSelectedObjectsDetails::OnDetailsPanelFinishedChangingProperties);
+
+	ChildSlot
+	[
+		DetailsView.ToSharedRef()
+	];
+} 
 
 void SNiagaraSelectedObjectsDetails::SelectedObjectsChanged()
 {
-	DetailsView->SetObjects(SelectedObjects->GetSelectedObjects().Array());
+	DetailsView->SetObjects(SelectedObjectsArray[0]->GetSelectedObjects().Array());
+}
+
+// TODO: Instead have a delegate that takes an array argument? This seems a bit dodgy..
+void SNiagaraSelectedObjectsDetails::SelectedObjectsChangedSecond()
+{
+	DetailsView->SetObjects(SelectedObjectsArray[1]->GetSelectedObjects().Array());
 }
 
 void SNiagaraSelectedObjectsDetails::OnDetailsPanelFinishedChangingProperties(const FPropertyChangedEvent& InEvent)
diff --git a/Engine/Plugins/FX/Niagara/Source/NiagaraEditor/Private/Widgets/SNiagaraSpreadsheetView.cpp b/Engine/Plugins/FX/Niagara/Source/NiagaraEditor/Private/Widgets/SNiagaraSpreadsheetView.cpp
--- a/Engine/Plugins/FX/Niagara/Source/NiagaraEditor/Private/Widgets/SNiagaraSpreadsheetView.cpp
+++ b/Engine/Plugins/FX/Niagara/Source/NiagaraEditor/Private/Widgets/SNiagaraSpreadsheetView.cpp
@@ -97,13 +97,14 @@ TSharedRef< SWidget > SNiagaraSpreadsheetRow::GenerateWidgetForColumn(const FNam
 			float* Src = nullptr;
 			if (UseGlobalOffsets && ParameterStore->GetParameterDataArray().GetData() != nullptr)
 			{
-				uint32 CompBufferOffset = FieldInfo->GlobalStartOffset;
-				Src = (float*)(ParameterStore->GetParameterDataArray().GetData() + CompBufferOffset);
+				const uint32 CompBufferOffset = FieldInfo->GlobalStartOffset;
+				const uint8* SrcU8 = ParameterStore->GetParameterDataArray().GetData();
+				Src = SrcU8 ? (float*)(SrcU8 + CompBufferOffset) : nullptr;
 			}
-			else if (DataSet != nullptr && DataSet->PrevData().GetNumInstances() > 0)
+			else if (DataSet != nullptr && DataSet->GetCurrentDataChecked().GetNumInstances() > 0)
 			{
 				uint32 CompBufferOffset = FieldInfo->FloatStartOffset;
-				Src = DataSet->PrevData().GetInstancePtrFloat(CompBufferOffset, RealRowIdx);
+				Src = DataSet->GetCurrentData()->GetInstancePtrFloat(CompBufferOffset, RealRowIdx);
 			}
 
 			float Value = 0.0f;
@@ -121,13 +122,14 @@ TSharedRef< SWidget > SNiagaraSpreadsheetRow::GenerateWidgetForColumn(const FNam
 			int32* Src = nullptr;
 			if (UseGlobalOffsets)
 			{
-				uint32 CompBufferOffset = FieldInfo->GlobalStartOffset;
-				Src = (int32*)(ParameterStore->GetParameterDataArray().GetData() + CompBufferOffset);
+				const uint32 CompBufferOffset = FieldInfo->GlobalStartOffset;
+				const uint8* SrcU8 = ParameterStore->GetParameterDataArray().GetData();
+				Src = SrcU8 ? (int32*)(SrcU8 + CompBufferOffset) : nullptr;
 			}
-			else if (DataSet && DataSet->PrevData().GetNumInstances() > 0)
+			else if (DataSet && DataSet->GetCurrentDataChecked().GetNumInstances() > 0)
 			{
 				uint32 CompBufferOffset = FieldInfo->IntStartOffset;
-				Src = DataSet->PrevData().GetInstancePtrInt32(CompBufferOffset, RealRowIdx);
+				Src = DataSet->GetCurrentDataChecked().GetInstancePtrInt32(CompBufferOffset, RealRowIdx);
 			}
 			FText ValueText;
 			if (Src && Src[0] == 0)
@@ -151,13 +153,14 @@ TSharedRef< SWidget > SNiagaraSpreadsheetRow::GenerateWidgetForColumn(const FNam
 			int32* Src = nullptr;
 			if (UseGlobalOffsets)
 			{
-				uint32 CompBufferOffset = FieldInfo->GlobalStartOffset;
-				Src = (int32*)(ParameterStore->GetParameterDataArray().GetData() + CompBufferOffset);
+				const uint32 CompBufferOffset = FieldInfo->GlobalStartOffset;
+				const uint8* SrcU8 = ParameterStore->GetParameterDataArray().GetData();
+				Src = SrcU8 ? (int32*)(SrcU8 + CompBufferOffset) : nullptr;
 			}
-			else if (DataSet != nullptr && DataSet->PrevData().GetNumInstances() > 0)
+			else if (DataSet != nullptr && DataSet->GetCurrentDataChecked().GetNumInstances() > 0)
 			{
 				uint32 CompBufferOffset = FieldInfo->IntStartOffset;
-				Src = DataSet->PrevData().GetInstancePtrInt32(CompBufferOffset, RealRowIdx);
+				Src = DataSet->GetCurrentDataChecked().GetInstancePtrInt32(CompBufferOffset, RealRowIdx);
 			}
 
 			int32 Value = 0;
@@ -174,13 +177,14 @@ TSharedRef< SWidget > SNiagaraSpreadsheetRow::GenerateWidgetForColumn(const FNam
 			int32* Src = nullptr;
 			if (UseGlobalOffsets)
 			{
-				uint32 CompBufferOffset = FieldInfo->GlobalStartOffset;
-				Src = (int32*)(ParameterStore->GetParameterDataArray().GetData() + CompBufferOffset);
+				const uint32 CompBufferOffset = FieldInfo->GlobalStartOffset;
+				const uint8* SrcU8 = ParameterStore->GetParameterDataArray().GetData();
+				Src = SrcU8 ? (int32*)(SrcU8 + CompBufferOffset) : nullptr;
 			}
-			else if (DataSet && DataSet->PrevData().GetNumInstances() > 0)
+			else if (DataSet && DataSet->GetCurrentDataChecked().GetNumInstances() > 0)
 			{
 				uint32 CompBufferOffset = FieldInfo->IntStartOffset;
-				Src = DataSet->PrevData().GetInstancePtrInt32(CompBufferOffset, RealRowIdx);
+				Src = DataSet->GetCurrentDataChecked().GetInstancePtrInt32(CompBufferOffset, RealRowIdx);
 			}
 
 			int32 Value = 0;
@@ -703,7 +707,7 @@ void SNiagaraSpreadsheetView::SetTarget(UNiagaraComponent* InComponent)
 
 	for (int32 i = 0; i < (int32)UIMax; i++)
 	{
-		CaptureData[i].DataSet.Init(FNiagaraDataSetID(), ENiagaraSimTarget::CPUSim);
+		CaptureData[i].DataSet.Init(FNiagaraDataSetID(), ENiagaraSimTarget::CPUSim, InComponent->GetFullName() + TEXT("/Spreadsheet"));
 		CaptureData[i].InputParams.Reset();
 		CaptureData[i].CaptureData.Reset();
 	}
@@ -784,11 +788,11 @@ FText SNiagaraSpreadsheetView::LastCapturedInfoText() const
 {
 	TArray<TSharedRef<FNiagaraEmitterHandleViewModel>> SelectedEmitterHandles;
 	SystemViewModel->GetSelectedEmitterHandles(SelectedEmitterHandles);
-	if (SelectedEmitterHandles.Num() == 1 && SNiagaraSpreadsheetView::IsPausedAtRightTimeOnRightHandle())
+	if (SelectedEmitterHandles.Num() == 1 && SNiagaraSpreadsheetView::IsPausedAtRightTimeOnRightHandle() && CaptureData[(int32)TabState].DataSet.GetCurrentData())
 	{
 		return FText::Format(LOCTEXT("LastCapturedInfoName", "Captured Emitter: \"{0}\"     # Particles: {1}    Script Type: {2}"),
 			SelectedEmitterHandles[0]->GetNameText(),
-			FText::AsNumber(CaptureData[(int32)TabState].DataSet.PrevData().GetNumInstances()),
+			FText::AsNumber(CaptureData[(int32)TabState].DataSet.GetCurrentDataChecked().GetNumInstances()),
 			ScriptEnum->GetDisplayNameTextByValue((int64)CaptureData[(int32)TabState].TargetUsage));
 	}
 
@@ -895,7 +899,7 @@ void SNiagaraSpreadsheetView::SelectedEmitterHandlesChanged()
 	// Need to reset the attributes list...
 	for (int32 i = 0; i < (int32)UIMax; i++)
 	{
-		CaptureData[i].DataSet.Init(FNiagaraDataSetID(), ENiagaraSimTarget::CPUSim);
+		CaptureData[i].DataSet.Init(FNiagaraDataSetID(), ENiagaraSimTarget::CPUSim, TEXT("/Spreadsheet"));
 		CaptureData[i].SupportedInputIndices.SetNum(0);
 		CaptureData[i].SupportedOutputIndices.SetNum(0);
 		CaptureData[i].OutputsListView->RequestTreeRefresh();
@@ -936,7 +940,7 @@ FReply SNiagaraSpreadsheetView::OnCSVOutputPressed()
 
 		CSVOutput += "\r\n";
 
-		for (uint32 RowIndex = 0; RowIndex < CaptureData[(int32)TabState].DataSet.PrevData().GetNumInstances(); RowIndex++)
+		for (uint32 RowIndex = 0; RowIndex < CaptureData[(int32)TabState].DataSet.GetCurrentDataChecked().GetNumInstances(); RowIndex++)
 		{
 			NumWritten = 0;
 			for (int32 i = 0; i < CaptureData[(int32)TabState].SupportedOutputFields->Num(); i++)
@@ -952,18 +956,18 @@ FReply SNiagaraSpreadsheetView::OnCSVOutputPressed()
 				}
 				const SNiagaraSpreadsheetView::FieldInfo* FieldInfo = FieldInfos[i];
 				
-				if (FieldInfo != nullptr && CaptureData[(int32)TabState].DataSet.GetNumInstances() != 0)
+				if (FieldInfo != nullptr && CaptureData[(int32)TabState].DataSet.GetCurrentDataChecked().GetNumInstances() != 0)
 				{
 					if (FieldInfo->bFloat)
 					{
 						uint32 CompBufferOffset = FieldInfo->FloatStartOffset;
-						float* Src = CaptureData[(int32)TabState].DataSet.PrevData().GetInstancePtrFloat(CompBufferOffset, RowIndex);
+						float* Src = CaptureData[(int32)TabState].DataSet.GetCurrentDataChecked().GetInstancePtrFloat(CompBufferOffset, RowIndex);
 						CSVOutput += FString::Printf(TEXT("%3.9f"), Src[0]);
 					}
 					else
 					{
 						uint32 CompBufferOffset = FieldInfo->IntStartOffset;
-						int32* Src = CaptureData[(int32)TabState].DataSet.PrevData().GetInstancePtrInt32(CompBufferOffset, RowIndex);
+						int32* Src = CaptureData[(int32)TabState].DataSet.GetCurrentDataChecked().GetInstancePtrInt32(CompBufferOffset, RowIndex);
 						CSVOutput += FString::Printf(TEXT("%d"), Src[0]);
 					}
 				}
@@ -1018,8 +1022,7 @@ void SNiagaraSpreadsheetView::HandleTimeChange()
 					if (FoundEntry != nullptr)
 					{
 						CaptureData[i].CaptureData = *FoundEntry;
-						CaptureData[i].CaptureData->Frame.CopyCurToPrev();
-						CaptureData[i].DataSet = CaptureData[i].CaptureData->Frame;
+						CaptureData[i].CaptureData->Frame.CopyTo(CaptureData[i].DataSet);
 						CaptureData[i].InputParams = CaptureData[i].CaptureData->Parameters;
 						CaptureData[i].LastCaptureHandleId = SelectedEmitterHandles[0]->GetId();
 
@@ -1029,7 +1032,8 @@ void SNiagaraSpreadsheetView::HandleTimeChange()
 					else
 					{
 						CaptureData[i].CaptureData.Reset();
-						CaptureData[i].DataSet.Init(FNiagaraDataSetID(), ENiagaraSimTarget::CPUSim);
+						CaptureData[i].DataSet.Init(FNiagaraDataSetID(), ENiagaraSimTarget::CPUSim, TargetComponent->GetFullName() + TEXT("/Spreadsheet"));
+						CaptureData[i].DataSet.Finalize();
 
 						ResetColumns((EUITab)i);
 						ResetEntries((EUITab)i);
@@ -1084,7 +1088,7 @@ void SNiagaraSpreadsheetView::ResetEntries(EUITab Tab)
 {
 	{
 		{
-			int32 NumInstances = CaptureData[(int32)Tab].DataSet.GetPrevNumInstances();
+			int32 NumInstances = CaptureData[(int32)Tab].DataSet.GetCurrentData() ? CaptureData[(int32)Tab].DataSet.GetCurrentDataChecked().GetNumInstances() : 0;
 			if (!CaptureData[(int32)Tab].bOutputColumnsAreAttributes && CaptureData[(int32)Tab].SupportedOutputFields.IsValid())
 			{
 				NumInstances = CaptureData[(int32)Tab].SupportedOutputFields->Num();
@@ -1182,7 +1186,7 @@ void SNiagaraSpreadsheetView::ResetColumns(EUITab Tab)
 {
 	int32 i = (int32)Tab;
 
-	if (CaptureData[(int32)i].DataSet.GetNumInstances() != 0)
+	if (CaptureData[i].DataSet.GetCurrentData() && CaptureData[i].DataSet.GetCurrentData()->GetNumInstances() != 0)
 	{
 		float ManualWidth = 125.0f;
 
diff --git a/Engine/Plugins/FX/Niagara/Source/NiagaraEditor/Private/Widgets/SNiagaraSystemViewport.cpp b/Engine/Plugins/FX/Niagara/Source/NiagaraEditor/Private/Widgets/SNiagaraSystemViewport.cpp
--- a/Engine/Plugins/FX/Niagara/Source/NiagaraEditor/Private/Widgets/SNiagaraSystemViewport.cpp
+++ b/Engine/Plugins/FX/Niagara/Source/NiagaraEditor/Private/Widgets/SNiagaraSystemViewport.cpp
@@ -17,6 +17,10 @@
 #include "EditorViewportCommands.h"
 #include "AdvancedPreviewScene.h"
 #include "ImageUtils.h"
+#include "Engine/Canvas.h"
+#include "Engine/Font.h"
+#include "CanvasItem.h"
+#include "DrawDebugHelpers.h"
 
 #define LOCTEXT_NAMESPACE "SNiagaraSystemViewport"
 
@@ -42,6 +46,8 @@ public:
 
 	virtual void SetIsSimulateInEditorViewport(bool bInIsSimulateInEditorViewport)override;
 
+	void DrawInstructionCounts(UNiagaraSystem* ParticleSystem, FCanvas* Canvas, float& CurrentX, float& CurrentY, UFont* Font, const float FontHeight);
+
 	TSharedPtr<SNiagaraSystemViewport> NiagaraViewport;
 	bool bCaptureScreenShot;
 	TWeakObjectPtr<UObject> ScreenShotOwner;
@@ -91,7 +97,7 @@ void FNiagaraSystemViewportClient::Tick(float DeltaSeconds)
 	}
 }
 
-void FNiagaraSystemViewportClient::Draw(FViewport* InViewport,FCanvas* Canvas)
+void FNiagaraSystemViewportClient::Draw(FViewport* InViewport, FCanvas* Canvas)
 {
 	UNiagaraSystem* ParticleSystem = NiagaraViewport.IsValid() ? NiagaraViewport->GetPreviewComponent()->GetAsset() : nullptr;
 
@@ -108,6 +114,19 @@ void FNiagaraSystemViewportClient::Draw(FViewport* InViewport,FCanvas* Canvas)
 
 	FEditorViewportClient::Draw(InViewport, Canvas);
 
+	if (NiagaraViewport.IsValid() )
+	{
+		float CurrentX = 10.0f;
+		float CurrentY = 50.0f;
+		UFont* Font = GEngine->GetSmallFont();
+		const float FontHeight = Font->GetMaxCharHeight() * 1.1f;
+
+		if ( NiagaraViewport->GetDrawElement(SNiagaraSystemViewport::EDrawElements::InstructionCounts) )
+		{
+			DrawInstructionCounts(ParticleSystem, Canvas, CurrentX, CurrentY, Font, FontHeight);
+		}
+	}
+
 	if (bCaptureScreenShot && ScreenShotOwner.IsValid() && OnScreenShotCaptured.IsBound())
 	{
 		
@@ -139,6 +158,50 @@ void FNiagaraSystemViewportClient::Draw(FViewport* InViewport,FCanvas* Canvas)
 	}
 }
 
+void FNiagaraSystemViewportClient::DrawInstructionCounts(UNiagaraSystem* ParticleSystem, FCanvas* Canvas, float& CurrentX, float& CurrentY, UFont* Font, const float FontHeight)
+{
+	Canvas->DrawShadowedString(CurrentX, CurrentY, TEXT("Instruction Counts"), Font, FLinearColor::White);
+	CurrentY += FontHeight;
+
+	for (const FNiagaraEmitterHandle& EmitterHandle : ParticleSystem->GetEmitterHandles())
+	{
+		UNiagaraEmitter* Emitter = EmitterHandle.GetInstance();
+		if (Emitter == nullptr)
+		{
+			continue;
+		}
+
+		Canvas->DrawShadowedString(CurrentX + 10.0f, CurrentY, *FString::Printf(TEXT("Emitter %s"), *EmitterHandle.GetName().ToString()), Font, FLinearColor::White);
+		CurrentY += FontHeight;
+
+		TArray<UNiagaraScript*> EmitterScripts;
+		Emitter->GetScripts(EmitterScripts);
+
+		for (UNiagaraScript* Script : EmitterScripts)
+		{
+			uint32 NumInstructions = 0;
+			if (Script->GetUsage() == ENiagaraScriptUsage::ParticleGPUComputeScript)
+			{
+				FNiagaraShader* Shader = Script->GetRenderThreadScript()->GetShaderGameThread();
+				if (Shader != nullptr)
+				{
+					NumInstructions = Shader->GetNumInstructions();
+				}
+			}
+			else
+			{
+				NumInstructions = Script->GetVMExecutableData().LastOpCount;
+			}
+
+			if (NumInstructions > 0)
+			{
+				Canvas->DrawShadowedString(CurrentX + 20.0f, CurrentY, *FString::Printf(TEXT("%s = %u"), *Script->GetName(), NumInstructions), Font, FLinearColor::White);
+				CurrentY += FontHeight;
+			}
+		}
+	}
+}
+
 bool FNiagaraSystemViewportClient::ShouldOrbitCamera() const
 {
 	return bUsingOrbitCamera;
@@ -312,14 +375,23 @@ void SNiagaraSystemViewport::BindCommands()
 	CommandList->MapAction(
 		Commands.TogglePreviewGrid,
 		FExecuteAction::CreateSP( this, &SNiagaraSystemViewport::TogglePreviewGrid ),
-								  FCanExecuteAction(),
-								  FIsActionChecked::CreateSP( this, &SNiagaraSystemViewport::IsTogglePreviewGridChecked ) );
-	
+		FCanExecuteAction(),
+		FIsActionChecked::CreateSP( this, &SNiagaraSystemViewport::IsTogglePreviewGridChecked )
+	);
+
+	CommandList->MapAction(
+		Commands.ToggleInstructionCounts,
+		FExecuteAction::CreateLambda([Viewport=this]() { Viewport->ToggleDrawElement(EDrawElements::InstructionCounts); Viewport->RefreshViewport(); }),
+		FCanExecuteAction(),
+		FIsActionChecked::CreateLambda([Viewport=this]() -> bool { return Viewport->GetDrawElement(EDrawElements::InstructionCounts); })
+	);
+
 	CommandList->MapAction(
 		Commands.TogglePreviewBackground,
 		FExecuteAction::CreateSP( this, &SNiagaraSystemViewport::TogglePreviewBackground ),
-								  FCanExecuteAction(),
-								  FIsActionChecked::CreateSP( this, &SNiagaraSystemViewport::IsTogglePreviewBackgroundChecked ) );
+		FCanExecuteAction(),
+		FIsActionChecked::CreateSP( this, &SNiagaraSystemViewport::IsTogglePreviewBackgroundChecked )
+	);
 
 	CommandList->MapAction(
 		Commands.ToggleOrbit,
diff --git a/Engine/Plugins/FX/Niagara/Source/NiagaraEditor/Private/Widgets/SNiagaraSystemViewport.h b/Engine/Plugins/FX/Niagara/Source/NiagaraEditor/Private/Widgets/SNiagaraSystemViewport.h
--- a/Engine/Plugins/FX/Niagara/Source/NiagaraEditor/Private/Widgets/SNiagaraSystemViewport.h
+++ b/Engine/Plugins/FX/Niagara/Source/NiagaraEditor/Private/Widgets/SNiagaraSystemViewport.h
@@ -74,6 +74,7 @@ public:
 	enum EDrawElements
 	{
 		Bounds = 0x020,
+		InstructionCounts = 0x040,
 	};
 
 	bool GetDrawElement(EDrawElements Element) const;
diff --git a/Engine/Plugins/FX/Niagara/Source/NiagaraEditor/Public/NiagaraActions.h b/Engine/Plugins/FX/Niagara/Source/NiagaraEditor/Public/NiagaraActions.h
--- a/Engine/Plugins/FX/Niagara/Source/NiagaraEditor/Public/NiagaraActions.h
+++ b/Engine/Plugins/FX/Niagara/Source/NiagaraEditor/Public/NiagaraActions.h
@@ -43,11 +43,8 @@ private:
 	FCanExecuteStackAction CanPerformAction;
 };
 
-USTRUCT()
 struct NIAGARAEDITOR_API FNiagaraParameterAction : public FEdGraphSchemaAction
 {
-	GENERATED_USTRUCT_BODY()
-
 	FNiagaraParameterAction() {}
 	FNiagaraParameterAction(const FNiagaraVariable& InParameter,
 		const TArray<FNiagaraGraphParameterReferenceCollection>& InReferenceCollection,
@@ -55,10 +52,8 @@ struct NIAGARAEDITOR_API FNiagaraParameterAction : public FEdGraphSchemaAction
 
 	FNiagaraVariable& GetParameter() { return Parameter; }
 
-	UPROPERTY()
 	FNiagaraVariable Parameter;
 
-	UPROPERTY()
 	TArray<FNiagaraGraphParameterReferenceCollection> ReferenceCollection;
 };
 
diff --git a/Engine/Plugins/FX/Niagara/Source/NiagaraEditor/Public/NiagaraEditorCommands.h b/Engine/Plugins/FX/Niagara/Source/NiagaraEditor/Public/NiagaraEditorCommands.h
--- a/Engine/Plugins/FX/Niagara/Source/NiagaraEditor/Public/NiagaraEditorCommands.h
+++ b/Engine/Plugins/FX/Niagara/Source/NiagaraEditor/Public/NiagaraEditorCommands.h
@@ -29,7 +29,8 @@ public:
 	TSharedPtr<FUICommandInfo> ResetSimulation;
 
 	/** Toggles the preview pane's grid */
-	TSharedPtr< FUICommandInfo > TogglePreviewGrid;
+	TSharedPtr<FUICommandInfo> TogglePreviewGrid;
+	TSharedPtr<FUICommandInfo> ToggleInstructionCounts;
 
 	/** Toggles the preview pane's background */
 	TSharedPtr< FUICommandInfo > TogglePreviewBackground;
diff --git a/Engine/Plugins/FX/Niagara/Source/NiagaraEditor/Public/NiagaraEditorModule.h b/Engine/Plugins/FX/Niagara/Source/NiagaraEditor/Public/NiagaraEditorModule.h
--- a/Engine/Plugins/FX/Niagara/Source/NiagaraEditor/Public/NiagaraEditorModule.h
+++ b/Engine/Plugins/FX/Niagara/Source/NiagaraEditor/Public/NiagaraEditorModule.h
@@ -133,6 +133,8 @@ private:
 
 	IConsoleCommand* TestCompileScriptCommand;
 	IConsoleCommand* DumpRapidIterationParametersForAsset;
+	IConsoleCommand* PreventSystemRecompileCommand;
+	IConsoleCommand* PreventAllSystemRecompilesCommand;
 
 	FOnCheckScriptToolkitsShouldFocusGraphElement OnCheckScriptToolkitsShouldFocusGraphElement;
 };
diff --git a/Engine/Plugins/FX/Niagara/Source/NiagaraEditor/Public/NiagaraEditorUtilities.h b/Engine/Plugins/FX/Niagara/Source/NiagaraEditor/Public/NiagaraEditorUtilities.h
--- a/Engine/Plugins/FX/Niagara/Source/NiagaraEditor/Public/NiagaraEditorUtilities.h
+++ b/Engine/Plugins/FX/Niagara/Source/NiagaraEditor/Public/NiagaraEditorUtilities.h
@@ -7,6 +7,7 @@
 #include "UObject/StructOnScope.h"
 #include "Misc/Attribute.h"
 #include "AssetData.h"
+#include "NiagaraGraph.h"
 
 class UNiagaraNodeInput;
 class UNiagaraNodeOutput;
@@ -137,8 +138,15 @@ namespace FNiagaraEditorUtilities
 
 	void FixUpNumericPins(const UEdGraphSchema_Niagara* Schema, UNiagaraNode* Node);
 
+	void SetStaticSwitchConstants(UNiagaraGraph* Graph, const TArray<UEdGraphPin*>& CallInputs);
+
+	bool ResolveConstantValue(UEdGraphPin* Pin, int32& Value);
+
 	void PreprocessFunctionGraph(const UEdGraphSchema_Niagara* Schema, UNiagaraGraph* Graph, const TArray<UEdGraphPin*>& CallInputs, const TArray<UEdGraphPin*>& CallOutputs, ENiagaraScriptUsage ScriptUsage);
 
+	/** Returns the parameter map references for the given graph as if the graph was compiled, factoring in compile-time options such as static switches. */
+	const TMap<FNiagaraVariable, FNiagaraGraphParameterReferenceCollection>& GetCompiledGraphParameterMapReferences(UNiagaraGraph* Graph);
+
 	/** Options for the GetScriptsByFilter function. 
 	** @Param ScriptUsageToInclude Only return Scripts that have this usage
 	** @Param (Optional) TargetUsageToMatch Only return Scripts that have this target usage (output node) 
diff --git a/Engine/Plugins/FX/Niagara/Source/NiagaraEditor/Public/NiagaraGraph.h b/Engine/Plugins/FX/Niagara/Source/NiagaraEditor/Public/NiagaraGraph.h
--- a/Engine/Plugins/FX/Niagara/Source/NiagaraEditor/Public/NiagaraGraph.h
+++ b/Engine/Plugins/FX/Niagara/Source/NiagaraEditor/Public/NiagaraGraph.h
@@ -11,6 +11,8 @@
 #include "NiagaraParameterStore.h"
 #include "NiagaraGraph.generated.h"
 
+class UNiagaraScriptVariable;
+
 /** This is the type of action that occurred on a given Niagara graph. Note that this should follow from EEdGraphActionType, leaving some slop for growth. */
 enum ENiagaraGraphActionType
 {
@@ -18,31 +20,26 @@ enum ENiagaraGraphActionType
 };
 
 class UNiagaraNode;
+class UNiagaraGraph;
 
 typedef TPair<FGuid, TWeakObjectPtr<UNiagaraNode>> FNiagaraGraphParameterReference;
 
-USTRUCT()
 struct FNiagaraGraphParameterReferenceCollection
 {
-	GENERATED_USTRUCT_BODY()
-
 public:
 	FNiagaraGraphParameterReferenceCollection(const bool bInCreated = false);
 
 	/** All the references in the graph. */
 	TArray<FNiagaraGraphParameterReference> ParameterReferences;
 
-	UPROPERTY()
-	class UNiagaraGraph* Graph;
+	const UNiagaraGraph* Graph;
 
 	/** Returns true if this parameter was initially created by the user. */
 	bool WasCreated() const;
 
 private:
 	/** Whether this parameter was initially created by the user. */
-	UPROPERTY()
 	bool bCreated;
-
 };
 
 /** Container for UNiagaraGraph cached data for managing CompileIds and Traversals.*/
@@ -54,6 +51,10 @@ struct FNiagaraGraphScriptUsageInfo
 public:
 	FNiagaraGraphScriptUsageInfo();
 
+	/** A guid which is generated when this usage info is created.  Allows for forced recompiling when the cached ids are invalidated. */
+	UPROPERTY()
+	FGuid BaseId;
+
 	/** The context in which this sub-graph traversal will be used.*/
 	UPROPERTY()
 	ENiagaraScriptUsage UsageType;
@@ -66,18 +67,25 @@ public:
 	UPROPERTY()
 	FGuid GeneratedCompileId;
 
-	/** The SHA1 hash that we calculated last traversal. If all zeroes, it is invalid.*/
+	/** The hash that we calculated last traversal. */
 	UPROPERTY()
-	TArray<uint8> DataHash;
+	FNiagaraCompileHash CompileHash;
 
 	/** The traversal of output to input nodes for this graph. This is not a recursive traversal, it just includes nodes from this graph.*/
 	UPROPERTY()
 	TArray<UNiagaraNode*> Traversal;
+
+	void PostLoad(UObject* Owner);
+
+private:
+	UPROPERTY()
+	TArray<uint8> DataHash_DEPRECATED;
 };
 
 struct FNiagaraGraphFunctionAliasContext
 {
 	ENiagaraScriptUsage CompileUsage;
+	TArray<UEdGraphPin*> StaticSwitchValues;
 };
 
 UCLASS(MinimalAPI)
@@ -89,7 +97,7 @@ class UNiagaraGraph : public UEdGraph
 
 	//~ Begin UObject Interface
 	virtual void PostLoad() override;
-	virtual void PostEditChangeProperty(struct FPropertyChangedEvent& PropertyChangedEvent)override;
+	virtual void PostEditChangeProperty(struct FPropertyChangedEvent& PropertyChangedEvent) override;
 	//~ End UObjet Interface
 	
 	/** Get the source that owns this graph */
@@ -145,6 +153,9 @@ class UNiagaraGraph : public UEdGraph
 	/** Finds input nodes in the graph with. */
 	void FindInputNodes(TArray<class UNiagaraNodeInput*>& OutInputNodes, FFindInputNodeOptions Options = FFindInputNodeOptions()) const;
 
+	/** Returns a list of variable inputs for all static switch nodes in the graph. */
+	TArray<FNiagaraVariable> FindStaticSwitchInputs() const;
+
 	/** Get an in-order traversal of a graph by the specified target output script usage.*/
 	void BuildTraversal(TArray<class UNiagaraNode*>& OutNodesTraversed, ENiagaraScriptUsage TargetUsage, FGuid TargetUsageId) const;
 	static void BuildTraversal(TArray<class UNiagaraNode*>& OutNodesTraversed, UNiagaraNode* FinalNode);
@@ -171,7 +182,7 @@ class UNiagaraGraph : public UEdGraph
 	void GetAllReferencedGraphs(TArray<const UNiagaraGraph*>& Graphs) const;
 
 	/** Gather all the change ids of external references for this specific graph traversal.*/
-	void GatherExternalDependencyIDs(ENiagaraScriptUsage InUsage, const FGuid& InUsageId, TArray<FGuid>& InReferencedIDs, TArray<UObject*>& InReferencedObjs);
+	void GatherExternalDependencyIDs(ENiagaraScriptUsage InUsage, const FGuid& InUsageId, TArray<FNiagaraCompileHash>& InReferencedCompileHashes, TArray<FGuid>& InReferencedIDs, TArray<UObject*>& InReferencedObjs);
 
 	/** Determine if there are any external dependencies wrt to scripts and ensure that those dependencies are sucked into the existing package.*/
 	void SubsumeExternalDependencies(TMap<const UObject*, UObject*>& ExistingConversions);
@@ -190,38 +201,43 @@ class UNiagaraGraph : public UEdGraph
 	the changes that were produced during the traversal of each output node, which are referred to as the CompileID.*/
 	FGuid GetChangeID() { return ChangeId; }
 
-	/** The compile id associated with the output node traversal specified by InUsage and InUsageId. If not found, an invalid Guid is returned.*/
-	FGuid GetCompileID(ENiagaraScriptUsage InUsage, const FGuid& InUsageId);
+	/** Recomputes the current compile id associated with the output node traversal specified by InUsage and InUsageId. If the usage is not found, an invalid Guid is returned.*/
+	FGuid ComputeCompileID(ENiagaraScriptUsage InUsage, const FGuid& InUsageId);
 
-	/** Walk through the graph for an ParameterMapGet nodes and see if any of them specify a default for VariableName.*/
-	UEdGraphPin* FindParameterMapDefaultValuePin(const FName VariableName, ENiagaraScriptUsage InUsage, ENiagaraScriptUsage InParentUsage) const;
+	/** Gets the current compile data hash associated with the output node traversal specified by InUsage and InUsageId. If the usage is not found, an invalid hash is returned.*/
+	FNiagaraCompileHash GetCompileDataHash(ENiagaraScriptUsage InUsage, const FGuid& InUsageId) const;
 
-	/** Get the meta-data associated with this variable, if it exists.*/
-	FNiagaraVariableMetaData* GetMetaData(const FNiagaraVariable& InVar);
-	const FNiagaraVariableMetaData* GetMetaData(const FNiagaraVariable& InVar) const;
+	/** Gets the current base id associated with the output node traversal specified by InUsage and InUsageId. If the usage is not found, an invalid guid is returned. */
+	FGuid GetBaseId(ENiagaraScriptUsage InUsage, const FGuid& InUsageId) const;
 
-	/** Return the meta-data associated with this variable. This should only be called on variables defined explicitly for this Graph, otherwise meta-data may leak.*/
-	FNiagaraVariableMetaData& FindOrAddMetaData(const FNiagaraVariable& InVar);
+	/** Forces the base compile id for the supplied script.  This should only be used to keep things consistent after an emitter merge. */
+	void ForceBaseId(ENiagaraScriptUsage InUsage, const FGuid& InUsageId, const FGuid InForcedBaseId);
 
-	/** Remove any meta-data that is no longer being referenced within this graph.*/
-	void PurgeUnreferencedMetaData();
+	/** Walk through the graph for an ParameterMapGet nodes and see if any of them specify a default for VariableName.*/
+	UEdGraphPin* FindParameterMapDefaultValuePin(const FName VariableName, ENiagaraScriptUsage InUsage, ENiagaraScriptUsage InParentUsage) const;
+
+	/** Gets the meta-data associated with this variable, if it exists.*/
+	TOptional<FNiagaraVariableMetaData> GetMetaData(const FNiagaraVariable& InVar) const;
 
-	/** Remove any parameter that is no longer being referenced within this graph.*/
-	void PurgeUnreferencedParameters();
+	/** Sets the meta-data associated with this variable.*/
+	void SetMetaData(const FNiagaraVariable& InVar, const FNiagaraVariableMetaData& MetaData);
 
-	const TMap<FNiagaraVariable, FNiagaraVariableMetaData>& GetAllMetaData() const { return VariableToMetaData; }
-	TMap<FNiagaraVariable, FNiagaraVariableMetaData>& GetAllMetaData() { return VariableToMetaData; }
+	const TMap<FNiagaraVariable, UNiagaraScriptVariable*>& GetAllMetaData() const;
+	TMap<FNiagaraVariable, UNiagaraScriptVariable*>& GetAllMetaData();
 
-	const TMap<FNiagaraVariable, FNiagaraGraphParameterReferenceCollection>& GetParameterMap() const;
+	const TMap<FNiagaraVariable, FNiagaraGraphParameterReferenceCollection>& GetParameterReferenceMap() const;
 
 	/** Adds parameter to parameters map setting it as created by the user.*/
 	void AddParameter(const FNiagaraVariable& Parameter);
 
+	/** Adds parameter to parameters map setting it as created by the user.*/
+	void AddParameterReference(const FNiagaraVariable& Parameter, const UEdGraphPin* Pin);
+
 	/** Remove parameter from map and all the pins associated. */
-	void RemoveParameter(const FNiagaraVariable& Parameter, const bool bNotifyGraphChanged = true);
+	void RemoveParameter(const FNiagaraVariable& Parameter);
 
 	/** Rename parameter from map and all the pins associated. */
-	bool RenameParameter(const FNiagaraVariable& Parameter, FName NewName, const bool bInNotifyGraphChanged = true);
+	bool RenameParameter(const FNiagaraVariable& Parameter, FName NewName);
 
 	/** Gets a delegate which is called whenever a contained data interfaces changes. */
 	FOnDataInterfaceChanged& OnDataInterfaceChanged();
@@ -241,29 +257,31 @@ class UNiagaraGraph : public UEdGraph
 
 	void InvalidateNumericCache();
 
+	/** If this graph is the source of a function call, it can add a string to the function name to discern it from different
+	  * function calls to the same graph. For example, if the graph contains static switches and two functions call it with
+	  * different switch parameters, the final function names in the hlsl must be different.
+	  */
 	FString GetFunctionAliasByContext(const FNiagaraGraphFunctionAliasContext& FunctionAliasContext);
 
+	void RebuildCachedCompileIds(bool bForce = false);
+
 protected:
-	void RebuildCachedData(bool bForce = false);
 	void RebuildNumericCache();
 	bool bNeedNumericCacheRebuilt;
 	TMap<TPair<FGuid, UEdGraphNode*>, FNiagaraTypeDefinition> CachedNumericConversions;
 	void ResolveNumerics(TMap<UNiagaraNode*, bool>& VisitedNodes, UEdGraphNode* Node);
 
 private:
+	/** Remove any meta-data that is no longer being referenced within this graph.*/
+	void PurgeUnreferencedMetaData() const;
+
 	virtual void NotifyGraphChanged(const FEdGraphEditAction& InAction) override;
 
 	/** Find parameters in the graph. */
-	void FindParameters();
-
-	/** 
-	 Set whether finding parameters is allowed. 
-	 To block finding parameters every time NotifyGraphChanged is called whenever we know it will be called multiple times.
-	*/
-	void SetFindParametersAllowed(const bool bAllowed);
+	void RefreshParameterReferences() const;
 
-	/** Finding parameters is blocked if false.*/
-	bool bFindParametersAllowed;
+	/** Marks the found parameter collections as invalid so they're rebuilt the next time they're requested. */
+	void InvalidateCachedParameterData();
 
 	/** A delegate that broadcasts a notification whenever the graph needs recompile due to structural change. */
 	FOnGraphChanged OnGraphNeedsRecompile;
@@ -280,15 +298,22 @@ private:
 
 	/** Storage of meta-data for variables defined for use explicitly with this graph.*/
 	UPROPERTY()
-	TMap<FNiagaraVariable, FNiagaraVariableMetaData> VariableToMetaData;
-	
-	/** Parameters found within the graph. */
+	mutable TMap<FNiagaraVariable, FNiagaraVariableMetaData> VariableToMetaData_DEPRECATED;
+
+	/** Storage of variables defined for use with this graph.*/
 	UPROPERTY()
-	TMap<FNiagaraVariable, FNiagaraGraphParameterReferenceCollection> Parameters;
+	mutable TMap<FNiagaraVariable, UNiagaraScriptVariable*> VariableToScriptVariable;
+	
+	/** A map of parameters in the graph to their referencers. */
+	mutable TMap<FNiagaraVariable, FNiagaraGraphParameterReferenceCollection> ParameterToReferencesMap;
 
 	FOnDataInterfaceChanged OnDataInterfaceChangedDelegate;
 
 	/** Whether currently renaming a parameter to prevent recursion. */
 	bool bIsRenamingParameter;
+
+	mutable bool bParameterReferenceRefreshPending;
+
+	mutable bool bUnreferencedMetaDataPurgePending;
 };
 
diff --git a/Engine/Plugins/FX/Niagara/Source/NiagaraEditor/Public/NiagaraNode.h b/Engine/Plugins/FX/Niagara/Source/NiagaraEditor/Public/NiagaraNode.h
--- a/Engine/Plugins/FX/Niagara/Source/NiagaraEditor/Public/NiagaraNode.h
+++ b/Engine/Plugins/FX/Niagara/Source/NiagaraEditor/Public/NiagaraNode.h
@@ -21,7 +21,7 @@ class NIAGARAEDITOR_API UNiagaraNode : public UEdGraphNode
 	GENERATED_UCLASS_BODY()
 protected:
 
-	bool ReallocatePins();
+	bool ReallocatePins(bool bMarkNeedsResynchronizeOnChange = true);
 
 	bool CompileInputPins(class FHlslNiagaraTranslator *Translator, TArray<int32>& OutCompiledInputs);
 
@@ -63,8 +63,10 @@ public:
 
 	UEdGraphPin* GetInputPin(int32 InputIndex) const;
 	void GetInputPins(TArray<class UEdGraphPin*>& OutInputPins) const;
+	void GetInputPins(TArray<const class UEdGraphPin*>& OutInputPins) const;
 	UEdGraphPin* GetOutputPin(int32 OutputIndex) const;
 	void GetOutputPins(TArray<class UEdGraphPin*>& OutOutputPins) const;
+	void GetOutputPins(TArray<const class UEdGraphPin*>& OutOutputPins) const;
 	UEdGraphPin* GetPinByPersistentGuid(const FGuid& InGuid) const;
 	virtual void ResolveNumerics(const UEdGraphSchema_Niagara* Schema, bool bSetInline, TMap<TPair<FGuid, UEdGraphNode*>, FNiagaraTypeDefinition>* PinCache);
 
@@ -94,15 +96,18 @@ public:
 	/** Notify the rename was cancelled.*/
 	virtual bool CancelEditablePinName(const FText& InName, UEdGraphPin* InGraphPinObj) { return false; }
 
-	virtual void BuildParameterMapHistory(FNiagaraParameterMapHistoryBuilder& OutHistory, bool bRecursive = true);
+	virtual void BuildParameterMapHistory(FNiagaraParameterMapHistoryBuilder& OutHistory, bool bRecursive = true) const;
 	
 	/** Go through all the external dependencies of this node in isolation and add them to the reference id list.*/
-	virtual void GatherExternalDependencyIDs(ENiagaraScriptUsage InMasterUsage, const FGuid& InMasterUsageId, TArray<FGuid>& InReferencedIDs, TArray<UObject*>& InReferencedObjs) const {};
+	virtual void GatherExternalDependencyIDs(ENiagaraScriptUsage InMasterUsage, const FGuid& InMasterUsageId, TArray<FNiagaraCompileHash>& InReferencedCompileHashes, TArray<FGuid>& InReferencedIDs, TArray<UObject*>& InReferencedObjs) const {};
 
 	/** Traces one of this node's output pins to its source output pin if it is a reroute node output pin.*/
 	virtual UEdGraphPin* GetTracedOutputPin(UEdGraphPin* LocallyOwnedOutputPin) const {return LocallyOwnedOutputPin;}
 	static UEdGraphPin* TraceOutputPin(UEdGraphPin* LocallyOwnedOutputPin);
 
+	/** Allows a node to replace a pin that is about to be compiled with another pin. This can be used for either optimizations or features such as the static switch. Returns true if the pin was successfully replaced, false otherwise. */
+	virtual bool SubstituteCompiledPin(FHlslNiagaraTranslator* Translator, UEdGraphPin** LocallyOwnedPin);
+
 	virtual UEdGraphPin* GetPassThroughPin(const UEdGraphPin* LocallyOwnedOutputPin) const override { return nullptr; }
 	virtual UEdGraphPin* GetPassThroughPin(const UEdGraphPin* LocallyOwnedOutputPin, ENiagaraScriptUsage MasterUsage) const { return nullptr; }
 
@@ -128,7 +133,7 @@ protected:
 
 	/** Route input parameter map to output parameter map if it exists. Note that before calling this function,
 		the input pins should have been visited already.*/
-	virtual void RouteParameterMapAroundMe(FNiagaraParameterMapHistoryBuilder& OutHistory, bool bRecursive);
+	virtual void RouteParameterMapAroundMe(FNiagaraParameterMapHistoryBuilder& OutHistory, bool bRecursive) const;
 	
 	/** The current change identifier for this node. Used to sync status with UNiagaraScripts.*/
 	UPROPERTY()
diff --git a/Engine/Plugins/FX/Niagara/Source/NiagaraEditor/Public/NiagaraNodeAssignment.h b/Engine/Plugins/FX/Niagara/Source/NiagaraEditor/Public/NiagaraNodeAssignment.h
--- a/Engine/Plugins/FX/Niagara/Source/NiagaraEditor/Public/NiagaraNodeAssignment.h
+++ b/Engine/Plugins/FX/Niagara/Source/NiagaraEditor/Public/NiagaraNodeAssignment.h
@@ -44,7 +44,8 @@ public:
 
 	//~ UNiagaraNodeFunctionCall interface
 	virtual bool RefreshFromExternalChanges() override;
-	virtual void BuildParameterMapHistory(FNiagaraParameterMapHistoryBuilder& OutHistory, bool bRecursive = true) override;
+	virtual void BuildParameterMapHistory(FNiagaraParameterMapHistoryBuilder& OutHistory, bool bRecursive = true) const override;
+	virtual void GatherExternalDependencyIDs(ENiagaraScriptUsage InMasterUsage, const FGuid& InMasterUsageId, TArray<FNiagaraCompileHash>& InReferencedCompileHashes, TArray<FGuid>& InReferencedIDs, TArray<UObject*>& InReferencedObjs) const override;
 
 	void AddParameter(FNiagaraVariable InVar, FString InDefaultValue);
 	void RemoveParameter(const FNiagaraVariable& InVar);
diff --git a/Engine/Plugins/FX/Niagara/Source/NiagaraEditor/Public/NiagaraNodeCustomHlsl.h b/Engine/Plugins/FX/Niagara/Source/NiagaraEditor/Public/NiagaraNodeCustomHlsl.h
--- a/Engine/Plugins/FX/Niagara/Source/NiagaraEditor/Public/NiagaraNodeCustomHlsl.h
+++ b/Engine/Plugins/FX/Niagara/Source/NiagaraEditor/Public/NiagaraNodeCustomHlsl.h
@@ -32,7 +32,7 @@ public:
 
 	bool GetTokens(TArray<FString>& OutTokens) const;
 
-	virtual void BuildParameterMapHistory(FNiagaraParameterMapHistoryBuilder& OutHistory, bool bRecursive = true) override;
+	virtual void BuildParameterMapHistory(FNiagaraParameterMapHistoryBuilder& OutHistory, bool bRecursive = true) const override;
 
 	// Replace items in the tokens array if they start with the src string or optionally src string and a namespace delimiter
 	static void ReplaceExactMatchTokens(TArray<FString>& Tokens, const FString& SrcString, const FString& ReplaceString, bool bAllowNamespaceSeparation);
diff --git a/Engine/Plugins/FX/Niagara/Source/NiagaraEditor/Public/NiagaraNodeFunctionCall.h b/Engine/Plugins/FX/Niagara/Source/NiagaraEditor/Public/NiagaraNodeFunctionCall.h
--- a/Engine/Plugins/FX/Niagara/Source/NiagaraEditor/Public/NiagaraNodeFunctionCall.h
+++ b/Engine/Plugins/FX/Niagara/Source/NiagaraEditor/Public/NiagaraNodeFunctionCall.h
@@ -36,6 +36,10 @@ public:
 	UPROPERTY()
 	FNiagaraFunctionSignature Signature;
 
+	/** All the input values the function propagates to the next higher caller instead of forcing the user to set them directly. */
+	UPROPERTY()
+	TArray<FNiagaraVariable> PropagatedStaticSwitchParameters;
+
 	bool ScriptIsValid() const;
 
 	//Begin UObject interface
@@ -50,7 +54,7 @@ public:
 	virtual ENiagaraNumericOutputTypeSelectionMode GetNumericOutputTypeSelectionMode() const override;
 	virtual bool CanAddToGraph(UNiagaraGraph* TargetGraph, FString& OutErrorMsg) const override;
 	virtual void SubsumeExternalDependencies(TMap<const UObject*, UObject*>& ExistingConversions) override;
-	virtual void GatherExternalDependencyIDs(ENiagaraScriptUsage InMasterUsage, const FGuid& InMasterUsageId, TArray<FGuid>& InReferencedIDs, TArray<UObject*>& InReferencedObjs) const override;
+	virtual void GatherExternalDependencyIDs(ENiagaraScriptUsage InMasterUsage, const FGuid& InMasterUsageId, TArray<FNiagaraCompileHash>& InReferencedCompileHashes, TArray<FGuid>& InReferencedIDs, TArray<UObject*>& InReferencedObjs) const override;
 	//End UNiagaraNode interface
 
 	//~ Begin EdGraphNode Interface
@@ -66,7 +70,7 @@ public:
 
 	bool FindAutoBoundInput(UNiagaraNodeInput* InputNode, UEdGraphPin* PinToAutoBind, FNiagaraVariable& OutFoundVar, ENiagaraInputNodeUsage& OutNodeUsage);
 
-	virtual void BuildParameterMapHistory(FNiagaraParameterMapHistoryBuilder& OutHistory, bool bRecursive = true) override;
+	virtual void BuildParameterMapHistory(FNiagaraParameterMapHistoryBuilder& OutHistory, bool bRecursive = true) const override;
 
 	FString GetFunctionName() const { return FunctionDisplayName; }
 	UNiagaraGraph* GetCalledGraph() const;
@@ -75,6 +79,9 @@ public:
 	/** Walk through the internal script graph for an ParameterMapGet nodes and see if any of them specify a default for VariableName.*/
 	UEdGraphPin* FindParameterMapDefaultValuePin(const FName VariableName, ENiagaraScriptUsage InParentUsage) const;
 
+	/** Attempts to find the input pin for a static switch with the given name in the internal script graph. Returns nullptr if no such pin can be found. */
+	UEdGraphPin* FindStaticSwitchInputPin(const FName& VariableName) const;
+
 	/** Tries to rename this function call to a new name.  The actual name that gets applied might be different due to conflicts with existing
 		nodes with the same name. */
 	void SuggestName(FString SuggestedName);
diff --git a/Engine/Plugins/FX/Niagara/Source/NiagaraEditor/Public/NiagaraParameterMapHistory.h b/Engine/Plugins/FX/Niagara/Source/NiagaraEditor/Public/NiagaraParameterMapHistory.h
--- a/Engine/Plugins/FX/Niagara/Source/NiagaraEditor/Public/NiagaraParameterMapHistory.h
+++ b/Engine/Plugins/FX/Niagara/Source/NiagaraEditor/Public/NiagaraParameterMapHistory.h
@@ -45,7 +45,7 @@ public:
 	TArray<const UEdGraphPin*> MapPinHistory;
 
 	/** List of nodes that manipulated the parameter map from input to output.*/
-	TArray<class UNiagaraNode*> MapNodeVisitations;
+	TArray<const class UNiagaraNode*> MapNodeVisitations;
 
 	/** For each node in MapNodeVisitations, record the start index and end index of variables added within the body of the node, i.e. a Get node will record just the values it pulls out directly. A function call, however, will record all sub-graph traversals.*/
 	TArray<TTuple<uint32, uint32> > MapNodeVariableMetaData;
@@ -58,7 +58,7 @@ public:
 	*/
 	int32 RegisterParameterMapPin(const UEdGraphPin* Pin);
 
-	uint32 BeginNodeVisitation(UNiagaraNode* Node);
+	uint32 BeginNodeVisitation(const UNiagaraNode* Node);
 	void EndNodeVisitation(uint32 IndexFromBeginNode);
 
 	/**
@@ -191,9 +191,6 @@ public:
 	/** Get the default value for this variable.*/
 	const UEdGraphPin* GetDefaultValuePin(int32 VarIdx) const;
 
-	const FNiagaraVariableMetaData* GetMetaData(int32 VarIdx) const;
-	FNiagaraVariableMetaData* GetMetaData(int32 VarIdx);
-
 	static FNiagaraVariable ConvertVariableToRapidIterationConstantName(FNiagaraVariable InVar, const TCHAR* InEmitterName, ENiagaraScriptUsage InUsage);
 
 	/**
@@ -220,7 +217,7 @@ public:
 	int32 CreateParameterMap();
 
 	/** Called in order to set up the correct initial context for an Output node and invokes the UNiagaraNodeOutput's BuildParameterMapHistory method.*/
-	void BuildParameterMaps(UNiagaraNodeOutput* OutputNode, bool bRecursive = true);
+	void BuildParameterMaps(const UNiagaraNodeOutput* OutputNode, bool bRecursive = true);
 
 	/**
 	* Called first during a node's visitation during traversal to identify that a node has been visited.
@@ -230,7 +227,7 @@ public:
 	/** Important. Must be called for each routing of the parameter map. This feeds the list used by TraceParameterMapOutputPin.*/
 	int32 RegisterParameterMapPin(int32 WhichParameterMap, const UEdGraphPin* Pin);
 
-	uint32 BeginNodeVisitation(int32 WhichParameterMap, class UNiagaraNode* Node);
+	uint32 BeginNodeVisitation(int32 WhichParameterMap, const class UNiagaraNode* Node);
 	void EndNodeVisitation(int32 WhichParameterMap, uint32 IndexFromBeginNode);
 
 	/** Trace back a pin to whom it was connected to to find the current parameter map to use.*/
@@ -250,22 +247,22 @@ public:
 	/**
 	* Record that we have entered a new function scope.
 	*/
-	void EnterFunction(const FString& InNodeName, const class UNiagaraScript* InScript, class UNiagaraNode* Node);
+	void EnterFunction(const FString& InNodeName, const class UNiagaraScript* InScript, const class UNiagaraNode* Node);
 
 	/**
 	* Record that we have exited a function scope.
 	*/
-	void ExitFunction(const FString& InNodeName, const class UNiagaraScript* InScript, class UNiagaraNode* Node);
+	void ExitFunction(const FString& InNodeName, const class UNiagaraScript* InScript, const class UNiagaraNode* Node);
 
 	/**
 	* Record that we have entered an emitter scope.
 	*/
-	void EnterEmitter(const FString& InEmitterName, class UNiagaraNode* Node);
+	void EnterEmitter(const FString& InEmitterName, const class UNiagaraNode* Node);
 
 	/**
 	* Record that we have exited an emitter scope.
 	*/
-	void ExitEmitter(const FString& InEmitterName, class UNiagaraNode* Node);
+	void ExitEmitter(const FString& InEmitterName, const class UNiagaraNode* Node);
 
 	/**
 	* Use the current alias map to resolve any aliases in this input variable name.
@@ -283,10 +280,10 @@ public:
 	bool GetNodePreviouslyVisited(const class UNiagaraNode* Node) const;
 
 	/** If we haven't already visited the owning nodes, do so.*/
-	void VisitInputPins(class UNiagaraNode*);
+	void VisitInputPins(const class UNiagaraNode*);
 	
 	/** If we haven't already visited the owning node, do so.*/
-	void VisitInputPin(const UEdGraphPin* Pin, class UNiagaraNode*);
+	void VisitInputPin(const UEdGraphPin* Pin, const class UNiagaraNode*);
 
 	/**
 	* Record that a pin writes to the parameter map. The pin name is expected to be the namespaced parameter map version of the name. If any aliases are in place, they are removed.
@@ -316,7 +313,7 @@ public:
 	const FString* GetEmitterAlias() const;
 
 	/** Get the node calling this sub-graph.*/
-	UNiagaraNode* GetCallingContext() const;
+	const UNiagaraNode* GetCallingContext() const;
 
 	/** Are we currently in a top-level function call context (ie a node in the main graph or an argument into a function for the main graph.)*/
 	bool InTopLevelFunctionCall(ENiagaraScriptUsage InFilterScriptType) const;
@@ -357,7 +354,7 @@ protected:
 	bool IsNamespacedVariableRelevantToScriptType(const FNiagaraVariable& InVar, ENiagaraScriptUsage ScriptType);
 
 	/** Contains the hierarchy of nodes leading to the current graph being processed. Usually made up of FunctionCall and Emitter nodes.*/
-	TArray<UNiagaraNode*> CallingContext;
+	TArray<const UNiagaraNode*> CallingContext;
 	/** Tracker for each context level of the parameter map index associated with a given pin. Used to trace parameter maps through the graph.*/
 	TArray<TMap<const UEdGraphPin*, int32> > PinToParameterMapIndices;
 	/** List of previously visited nodes per context. Note that the same node may be visited multiple times across all graph traversals, but only one time per context level.*/
diff --git a/Engine/Plugins/FX/Niagara/Source/NiagaraEditor/Public/NiagaraScriptSource.h b/Engine/Plugins/FX/Niagara/Source/NiagaraEditor/Public/NiagaraScriptSource.h
--- a/Engine/Plugins/FX/Niagara/Source/NiagaraEditor/Public/NiagaraScriptSource.h
+++ b/Engine/Plugins/FX/Niagara/Source/NiagaraEditor/Public/NiagaraScriptSource.h
@@ -35,7 +35,11 @@ class UNiagaraScriptSource : public UNiagaraScriptSourceBase
 
 	virtual FGuid GetChangeID();
 
-	virtual void ComputeVMCompilationId(struct FNiagaraVMExecutableDataId& Id, ENiagaraScriptUsage InUsage, const FGuid& InUsageId) const override;
+	virtual void ComputeVMCompilationId(struct FNiagaraVMExecutableDataId& Id, ENiagaraScriptUsage InUsage, const FGuid& InUsageId, bool bForceRebuild = false) const override;
+
+	virtual FGuid GetCompileBaseId(ENiagaraScriptUsage InUsage, const FGuid& InUsageId) const override;
+
+	virtual FNiagaraCompileHash GetCompileHash(ENiagaraScriptUsage InUsage, const FGuid& InUsageId) const override;
 
 	virtual void PostLoadFromEmitter(UNiagaraEmitter& OwningEmitter) override;
 
diff --git a/Engine/Plugins/FX/Niagara/Source/NiagaraEditor/Public/ViewModels/NiagaraEmitterHandleViewModel.h b/Engine/Plugins/FX/Niagara/Source/NiagaraEditor/Public/ViewModels/NiagaraEmitterHandleViewModel.h
--- a/Engine/Plugins/FX/Niagara/Source/NiagaraEditor/Public/ViewModels/NiagaraEmitterHandleViewModel.h
+++ b/Engine/Plugins/FX/Niagara/Source/NiagaraEditor/Public/ViewModels/NiagaraEmitterHandleViewModel.h
@@ -10,7 +10,7 @@
 
 class UNiagaraSystem;
 struct FNiagaraEmitterHandle;
-struct FNiagaraEmitterInstance;
+class FNiagaraEmitterInstance;
 class FNiagaraEmitterViewModel;
 
 /** The view model for the FNiagaraEmitterEditorWidget. */
@@ -21,18 +21,18 @@ public:
 	DECLARE_MULTICAST_DELEGATE(FOnNameChanged);
 public:
 	/** Creates a new emitter editor view model with the supplied emitter handle and simulation. */
-	FNiagaraEmitterHandleViewModel(FNiagaraEmitterHandle* InEmitterHandle, TWeakPtr<FNiagaraEmitterInstance> InSimulation, UNiagaraSystem& InOwningSystem);
+	FNiagaraEmitterHandleViewModel(FNiagaraEmitterHandle* InEmitterHandle, TWeakPtr<FNiagaraEmitterInstance, ESPMode::ThreadSafe> InSimulation, UNiagaraSystem& InOwningSystem);
 	
 	~FNiagaraEmitterHandleViewModel();
 
 	/** Reuses a the emitter editor view model with the supplied emitter handle and simulation.*/
-	bool Set(FNiagaraEmitterHandle* InEmitterHandle, TWeakPtr<FNiagaraEmitterInstance> InSimulation, UNiagaraSystem& InOwningSystem);
+	bool Set(FNiagaraEmitterHandle* InEmitterHandle, TWeakPtr<FNiagaraEmitterInstance, ESPMode::ThreadSafe> InSimulation, UNiagaraSystem& InOwningSystem);
 
 	/** Sets the emitter handle.*/
 	void SetEmitterHandle(FNiagaraEmitterHandle* InEmitterHandle);
 
 	/** Sets the simulation for the emitter this handle references. */
-	void SetSimulation(TWeakPtr<FNiagaraEmitterInstance> InSimulation);
+	void SetSimulation(TWeakPtr<FNiagaraEmitterInstance, ESPMode::ThreadSafe> InSimulation);
 
 	/** Gets the id of the emitter handle. */
 	FGuid GetId() const;
diff --git a/Engine/Plugins/FX/Niagara/Source/NiagaraEditor/Public/ViewModels/NiagaraEmitterViewModel.h b/Engine/Plugins/FX/Niagara/Source/NiagaraEditor/Public/ViewModels/NiagaraEmitterViewModel.h
--- a/Engine/Plugins/FX/Niagara/Source/NiagaraEditor/Public/ViewModels/NiagaraEmitterViewModel.h
+++ b/Engine/Plugins/FX/Niagara/Source/NiagaraEditor/Public/ViewModels/NiagaraEmitterViewModel.h
@@ -11,7 +11,7 @@ class UNiagaraScript;
 class FNiagaraScriptViewModel;
 class FNiagaraScriptGraphViewModel;
 class UNiagaraEmitterEditorData;
-struct FNiagaraEmitterInstance;
+class FNiagaraEmitterInstance;
 struct FNiagaraVariable;
 struct FNiagaraParameterStore;
 struct FEdGraphEditAction;
@@ -28,17 +28,17 @@ public:
 
 public:
 	/** Creates a new emitter editor view model with the supplied emitter handle and simulation. */
-	FNiagaraEmitterViewModel(UNiagaraEmitter* InEmitter, TWeakPtr<FNiagaraEmitterInstance> InSimulation);
+	FNiagaraEmitterViewModel(UNiagaraEmitter* InEmitter, TWeakPtr<FNiagaraEmitterInstance, ESPMode::ThreadSafe> InSimulation);
 	virtual ~FNiagaraEmitterViewModel();
 
 	/** Reuse this view model with new parameters.*/
-	bool Set(UNiagaraEmitter* InEmitter, TWeakPtr<FNiagaraEmitterInstance> InSimulation);
+	bool Set(UNiagaraEmitter* InEmitter, TWeakPtr<FNiagaraEmitterInstance, ESPMode::ThreadSafe> InSimulation);
 
 	/** Sets this view model to a different emitter. */
 	void SetEmitter(UNiagaraEmitter* InEmitter);
 
 	/** Sets the current simulation for the emitter. */
-	void SetSimulation(TWeakPtr<FNiagaraEmitterInstance> InSimulation);
+	void SetSimulation(TWeakPtr<FNiagaraEmitterInstance, ESPMode::ThreadSafe> InSimulation);
 
 	/** Gets the emitter represented by this view model. */
 	NIAGARAEDITOR_API UNiagaraEmitter* GetEmitter();
@@ -101,7 +101,7 @@ private:
 	TWeakObjectPtr<UNiagaraEmitter> Emitter;
 
 	/** The runtime simulation for the emitter being displayed by the control */
-	TWeakPtr<FNiagaraEmitterInstance> Simulation;
+	TWeakPtr<FNiagaraEmitterInstance, ESPMode::ThreadSafe> Simulation;
 	
 	/** The view model for the update/spawn/event script. */
 	TSharedPtr<FNiagaraScriptViewModel> SharedScriptViewModel;
diff --git a/Engine/Plugins/FX/Niagara/Source/NiagaraEditor/Public/ViewModels/NiagaraSystemViewModel.h b/Engine/Plugins/FX/Niagara/Source/NiagaraEditor/Public/ViewModels/NiagaraSystemViewModel.h
--- a/Engine/Plugins/FX/Niagara/Source/NiagaraEditor/Public/ViewModels/NiagaraSystemViewModel.h
+++ b/Engine/Plugins/FX/Niagara/Source/NiagaraEditor/Public/ViewModels/NiagaraSystemViewModel.h
@@ -96,6 +96,16 @@ public:
 	DECLARE_MULTICAST_DELEGATE(FOnPinnedCurvesChanged);
 
 public:
+	struct FEmitterHandleToDuplicate
+	{
+		FString SystemPath;
+		FGuid EmitterHandleId;
+		bool operator==(const FEmitterHandleToDuplicate& Other) const
+		{
+			return SystemPath == Other.SystemPath && EmitterHandleId == Other.EmitterHandleId;
+		}
+	};
+	
 	/** Defines different multi-system reset modes for this system view model */
 	enum class EMultiResetMode
 	{
@@ -122,8 +132,6 @@ public:
 		/** Reset this system (do not pull in changes) */
 		ResetSystem,
 	};
-
-public:
 	/** Creates a new view model with the supplied System and System instance. */
 	FNiagaraSystemViewModel(UNiagaraSystem& InSystem, FNiagaraSystemViewModelOptions InOptions);
 
@@ -386,7 +394,7 @@ private:
 	void SystemInstanceReset();
 
 	/** Duplicates a set of emitters and refreshes everything.*/
-	void DuplicateEmitters(TSet<FGuid> EmitterHandleIdsToDuplicate);
+	void DuplicateEmitters(TArray<FEmitterHandleToDuplicate> EmitterHandlesToDuplicate);
 
 	/** Adds event handler for the system's scripts. */
 	void AddSystemEventHandlers();
diff --git a/Engine/Plugins/FX/Niagara/Source/NiagaraEditor/Public/ViewModels/Stack/NiagaraStackFunctionInput.h b/Engine/Plugins/FX/Niagara/Source/NiagaraEditor/Public/ViewModels/Stack/NiagaraStackFunctionInput.h
--- a/Engine/Plugins/FX/Niagara/Source/NiagaraEditor/Public/ViewModels/Stack/NiagaraStackFunctionInput.h
+++ b/Engine/Plugins/FX/Niagara/Source/NiagaraEditor/Public/ViewModels/Stack/NiagaraStackFunctionInput.h
@@ -18,6 +18,7 @@ class UNiagaraStackObject;
 class UNiagaraScript;
 class UEdGraphPin;
 class UNiagaraDataInterface;
+enum class EStackParameterBehavior;
 
 /** Represents a single module input in the module stack view model. */
 UCLASS()
@@ -59,6 +60,7 @@ public:
 	 * @param InInputFunctionCallNode The function call which this entry is an input to. NOTE: This node can be a module function call node or a dynamic input node.
 	 * @param InInputParameterHandle The input parameter handle for the function call.
 	 * @param InInputType The type of this input.
+	 * @param InParameterBehavior Determines how the parameter should behave in the stack
 	 * @param InOwnerStackItemEditorDataKey The editor data key of the item that owns this input.
 	 */
 	void Initialize(
@@ -67,6 +69,7 @@ public:
 		UNiagaraNodeFunctionCall& InInputFunctionCallNode,
 		FName InInputParameterHandle,
 		FNiagaraTypeDefinition InInputType,
+		EStackParameterBehavior InParameterBehavior,
 		FString InOwnerStackItemEditorDataKey);
 
 	/** Gets the function call node which owns this input. */
@@ -144,6 +147,12 @@ public:
 	/** Determine if this field is editable */
 	bool IsEditable() const;
 
+	/** If true the parameter can only be set to local constant values */
+	bool IsStaticParameter() const;
+
+	/** Whether or not this input has a base value.  This is true for emitter instances in systems. */
+	bool EmitterHasBase() const;
+
 	/** Whether or not this input can be reset to a base value. */
 	bool CanResetToBase() const;
 
@@ -369,7 +378,7 @@ private:
 	FNiagaraTypeDefinition InputType;
 
 	/** The meta data for this input, defined in the owning function's script. */
-	FNiagaraVariableMetaData* InputMetaData;
+	TOptional<FNiagaraVariableMetaData> InputMetaData;
 
 	/** A unique key for this input for looking up editor only UI data. */
 	FString StackEditorDataKey;
@@ -380,6 +389,9 @@ private:
 	/** The parameter handle which defined this input in the module graph. */
 	FNiagaraParameterHandle InputParameterHandle;
 
+	/** Defines the type of the stack parameter (e.g. 'static' if it should only allow constant compile-time values). */
+	EStackParameterBehavior ParameterBehavior;
+
 	/** The parameter handle which defined this input in the module graph, aliased for use in the current emitter
 	  * graph.  This only affects parameter handles which are local module handles. */
 	FNiagaraParameterHandle AliasedInputParameterHandle;
diff --git a/Engine/Plugins/FX/Niagara/Source/NiagaraEditor/Public/ViewModels/Stack/NiagaraStackFunctionInputCollection.h b/Engine/Plugins/FX/Niagara/Source/NiagaraEditor/Public/ViewModels/Stack/NiagaraStackFunctionInputCollection.h
--- a/Engine/Plugins/FX/Niagara/Source/NiagaraEditor/Public/ViewModels/Stack/NiagaraStackFunctionInputCollection.h
+++ b/Engine/Plugins/FX/Niagara/Source/NiagaraEditor/Public/ViewModels/Stack/NiagaraStackFunctionInputCollection.h
@@ -52,6 +52,7 @@ private:
 		FNiagaraTypeDefinition Type;
 		int32 SortKey;
 		FText Category;
+		bool bIsStatic;
 	};
 
 	UNiagaraNodeFunctionCall* ModuleNode;
diff --git a/Engine/Plugins/FX/Niagara/Source/NiagaraEditor/Public/ViewModels/Stack/NiagaraStackFunctionInputCondition.h b/Engine/Plugins/FX/Niagara/Source/NiagaraEditor/Public/ViewModels/Stack/NiagaraStackFunctionInputCondition.h
--- a/Engine/Plugins/FX/Niagara/Source/NiagaraEditor/Public/ViewModels/Stack/NiagaraStackFunctionInputCondition.h
+++ b/Engine/Plugins/FX/Niagara/Source/NiagaraEditor/Public/ViewModels/Stack/NiagaraStackFunctionInputCondition.h
@@ -31,7 +31,7 @@ public:
 
 	FNiagaraTypeDefinition GetConditionInputType() const;
 
-	FNiagaraVariableMetaData* GetConditionInputMetaData() const;
+	TOptional<FNiagaraVariableMetaData> GetConditionInputMetaData() const;
 
 private:
 	UNiagaraScript* Script;
diff --git a/Engine/Plugins/FX/Niagara/Source/NiagaraEditor/Public/ViewModels/Stack/NiagaraStackGraphUtilities.h b/Engine/Plugins/FX/Niagara/Source/NiagaraEditor/Public/ViewModels/Stack/NiagaraStackGraphUtilities.h
--- a/Engine/Plugins/FX/Niagara/Source/NiagaraEditor/Public/ViewModels/Stack/NiagaraStackGraphUtilities.h
+++ b/Engine/Plugins/FX/Niagara/Source/NiagaraEditor/Public/ViewModels/Stack/NiagaraStackGraphUtilities.h
@@ -43,6 +43,8 @@ namespace FNiagaraStackGraphUtilities
 
 	UNiagaraNodeOutput* GetEmitterOutputNodeForStackNode(UNiagaraNode& StackNode);
 
+	const UNiagaraNodeOutput* GetEmitterOutputNodeForStackNode(const UNiagaraNode& StackNode);
+
 	UNiagaraNodeInput* GetEmitterInputNodeForStackNode(UNiagaraNode& StackNode);
 
 	struct FStackNodeGroup
@@ -74,6 +76,8 @@ namespace FNiagaraStackGraphUtilities
 
 	void GetStackFunctionInputPins(UNiagaraNodeFunctionCall& FunctionCallNode, TArray<const UEdGraphPin*>& OutInputPins, ENiagaraGetStackFunctionInputPinsOptions Options = ENiagaraGetStackFunctionInputPinsOptions::AllInputs, bool bIgnoreDisabled = false);
 
+	void GetStackFunctionStaticSwitchPins(UNiagaraNodeFunctionCall& FunctionCallNode, TArray<const UEdGraphPin*>& OutInputPins);
+
 	UNiagaraNodeParameterMapSet* GetStackFunctionOverrideNode(UNiagaraNodeFunctionCall& FunctionCallNode);
 
 	UNiagaraNodeParameterMapSet& GetOrCreateStackFunctionOverrideNode(UNiagaraNodeFunctionCall& FunctionCallNode, const FGuid& PreferredOverrideNodeGuid = FGuid());
@@ -166,4 +170,6 @@ namespace FNiagaraStackGraphUtilities
 	/** Whether a parameter is allowed to be used in a certain execution category. 
 		Used to check if parameter can be dropped on a module or funciton stack entry. */
 	NIAGARAEDITOR_API bool ParameterAllowedInExecutionCategory(const FName InParameterName, const FName ExecutionCategory);
+
+	void RebuildEmitterNodes(UNiagaraSystem& System);
 }
diff --git a/Engine/Plugins/FX/Niagara/Source/NiagaraEditor/Public/ViewModels/Stack/NiagaraStackInputCategory.h b/Engine/Plugins/FX/Niagara/Source/NiagaraEditor/Public/ViewModels/Stack/NiagaraStackInputCategory.h
--- a/Engine/Plugins/FX/Niagara/Source/NiagaraEditor/Public/ViewModels/Stack/NiagaraStackInputCategory.h
+++ b/Engine/Plugins/FX/Niagara/Source/NiagaraEditor/Public/ViewModels/Stack/NiagaraStackInputCategory.h
@@ -9,6 +9,12 @@
 class UNiagaraStackFunctionInput;
 class UNiagaraNodeFunctionCall;
 
+UENUM()
+enum class EStackParameterBehavior
+{
+	Dynamic, Static
+};
+
 UCLASS()
 class NIAGARAEDITOR_API UNiagaraStackInputCategory : public UNiagaraStackItemContent
 {
@@ -26,7 +32,7 @@ public:
 
 	void ResetInputs();
 
-	void AddInput(FName InInputParameterHandle, FNiagaraTypeDefinition InInputType);
+	void AddInput(FName InInputParameterHandle, FNiagaraTypeDefinition InInputType, EStackParameterBehavior InParameterBehavior);
 
 	//~ UNiagaraStackEntry interface
 	virtual FText GetDisplayName() const override;
@@ -49,6 +55,7 @@ private:
 	{
 		FName ParameterHandle;
 		FNiagaraTypeDefinition Type;
+		EStackParameterBehavior ParameterBehavior;
 	};
 
 	UNiagaraNodeFunctionCall* ModuleNode;
