Copying //UE4/Dev-Niagara to Dev-Main (//UE4/Dev-Main)

diff --git a/Engine/Plugins/FX/Niagara/Source/NiagaraEditorWidgets/Private/Stack/SNiagaraStackFunctionInputValue.cpp b/Engine/Plugins/FX/Niagara/Source/NiagaraEditorWidgets/Private/Stack/SNiagaraStackFunctionInputValue.cpp
--- a/Engine/Plugins/FX/Niagara/Source/NiagaraEditorWidgets/Private/Stack/SNiagaraStackFunctionInputValue.cpp
+++ b/Engine/Plugins/FX/Niagara/Source/NiagaraEditorWidgets/Private/Stack/SNiagaraStackFunctionInputValue.cpp
@@ -178,6 +178,7 @@ void SNiagaraStackFunctionInputValue::Construct(const FArguments& InArgs, UNiaga
 				.ForegroundColor(FSlateColor::UseForeground())
 				.OnGetMenuContent(this, &SNiagaraStackFunctionInputValue::OnGetAvailableHandleMenu)
 				.ContentPadding(FMargin(2))
+				.Visibility(this, &SNiagaraStackFunctionInputValue::GetDropdownButtonVisibility)
 				.MenuPlacement(MenuPlacement_BelowRightAnchor)
 				.HAlign(HAlign_Center)
 				.VAlign(VAlign_Center)
@@ -695,6 +696,11 @@ EVisibility SNiagaraStackFunctionInputValue::GetResetButtonVisibility() const
 	return FunctionInput->CanReset() ? EVisibility::Visible : EVisibility::Hidden;
 }
 
+EVisibility SNiagaraStackFunctionInputValue::GetDropdownButtonVisibility() const
+{
+	return FunctionInput->IsStaticParameter() ? EVisibility::Hidden : EVisibility::Visible;
+}
+
 FReply SNiagaraStackFunctionInputValue::ResetButtonPressed() const
 {
 	FunctionInput->Reset();
diff --git a/Engine/Plugins/FX/Niagara/Source/NiagaraEditorWidgets/Private/Stack/SNiagaraStackFunctionInputValue.h b/Engine/Plugins/FX/Niagara/Source/NiagaraEditorWidgets/Private/Stack/SNiagaraStackFunctionInputValue.h
--- a/Engine/Plugins/FX/Niagara/Source/NiagaraEditorWidgets/Private/Stack/SNiagaraStackFunctionInputValue.h
+++ b/Engine/Plugins/FX/Niagara/Source/NiagaraEditorWidgets/Private/Stack/SNiagaraStackFunctionInputValue.h
@@ -115,6 +115,8 @@ private:
 
 	EVisibility GetResetButtonVisibility() const;
 
+	EVisibility GetDropdownButtonVisibility() const;
+
 	FReply ResetButtonPressed() const;
 
 	EVisibility GetResetToBaseButtonVisibility() const;
diff --git a/Engine/Plugins/FX/Niagara/Source/NiagaraShader/NiagaraShader.Build.cs b/Engine/Plugins/FX/Niagara/Source/NiagaraShader/NiagaraShader.Build.cs
--- a/Engine/Plugins/FX/Niagara/Source/NiagaraShader/NiagaraShader.Build.cs
+++ b/Engine/Plugins/FX/Niagara/Source/NiagaraShader/NiagaraShader.Build.cs
@@ -22,6 +22,7 @@ public class NiagaraShader : ModuleRules
                 "RenderCore",
                 "VectorVM",
                 "RHI",
+                "Projects",
                 "NiagaraCore"
             }
         );
diff --git a/Engine/Plugins/FX/Niagara/Source/NiagaraShader/Private/NiagaraShader.cpp b/Engine/Plugins/FX/Niagara/Source/NiagaraShader/Private/NiagaraShader.cpp
--- a/Engine/Plugins/FX/Niagara/Source/NiagaraShader/Private/NiagaraShader.cpp
+++ b/Engine/Plugins/FX/Niagara/Source/NiagaraShader/Private/NiagaraShader.cpp
@@ -21,8 +21,10 @@
 #include "NiagaraDataInterfaceBase.h"
 #include "UObject/UObjectGlobals.h"
 #include "NiagaraShaderModule.h"
+#include "NiagaraCustomVersion.h"
+
+IMPLEMENT_SHADER_TYPE(, FNiagaraShader, TEXT("/Plugin/FX/Niagara/Private/NiagaraEmitterInstanceShader.usf"),TEXT("SimulateMain"), SF_Compute)
 
-IMPLEMENT_SHADER_TYPE(,FNiagaraShader, TEXT("/Engine/Private/NiagaraEmitterInstanceShader.usf"),TEXT("SimulateMain"), SF_Compute)
 
 
 int32 GCreateNiagaraShadersOnLoad = 0;
@@ -90,23 +92,45 @@ void UpdateNiagaraShaderCompilingStats(const FNiagaraShaderScript* Script)
 
 void FNiagaraShaderMapId::Serialize(FArchive& Ar)
 {
-	// You must bump NIAGARASHADERMAP_DERIVEDDATA_VER if changing the serialization of FNiagaraShaderMapId.
-	/*const int32 NiagaraCustomVersion = GetLinkerCustomVersion(FNiagaraCustomVersion::GUID);
-	if (NiagaraCustomVersion <= FNiagaraCustomVersion::MovingTranslatorVersionToGuid)
+	// Handle niagara custom version and the compiler version id.
+	int32 NiagaraVer;
+	if (Ar.IsLoading())
 	{
-		int VersionId = 0;
-		Ar << VersionId;
-		CompilerVersionID.Invalidate();
+		FGuid FirstGuid;
+		Ar << FirstGuid;
+		if (FirstGuid == FNiagaraCustomVersion::GUID)
+		{
+			// If the first guid matches the niagara custom version identifier, than this was serialized with versioning so read the current version and then read the compiler version which was next.
+			Ar << NiagaraVer;
+			Ar << CompilerVersionID;
+		}
+		else
+		{
+			// If the first guid was not the niagara custom version identifier it was saved before versioning so the first guid was the compiler version.
+			NiagaraVer = FNiagaraCustomVersion::UseHashesToIdentifyCompileStateOfTopLevelScripts - 1;
+			CompilerVersionID = FirstGuid;
+		}
 	}
-	else*/
+	else // Saving
 	{
+		// When saving, the custom version is the latest version so serialize the niagara custom version identifier, the current version, and the compiler version.
+		NiagaraVer = FNiagaraCustomVersion::LatestVersion;
+		FGuid NiagaraCustomVersionGuid = FNiagaraCustomVersion::GUID;
+		Ar << NiagaraCustomVersionGuid;
+		Ar << NiagaraVer;
 		Ar << CompilerVersionID;
 	}
+
 	Ar << BaseScriptID;
 	Ar << (int32&)FeatureLevel;
-
-	//ParameterSet.Serialize(Ar);		// NIAGARATODO: at some point we'll need stuff for static switches here
 	Ar << ReferencedDependencyIds;
+	if (NiagaraVer >= FNiagaraCustomVersion::UseHashesToIdentifyCompileStateOfTopLevelScripts)
+	{
+		Ar << BaseCompileHash;
+		Ar << ReferencedCompileHashes;
+	}
+	
+	//ParameterSet.Serialize(Ar);		// NIAGARATODO: at some point we'll need stuff for static switches here
 }
 
 /** Hashes the script-specific part of this shader map Id. */
@@ -116,14 +140,20 @@ void FNiagaraShaderMapId::GetScriptHash(FSHAHash& OutHash) const
 
 	HashState.Update((const uint8*)&CompilerVersionID, sizeof(CompilerVersionID));
 	HashState.Update((const uint8*)&BaseScriptID, sizeof(BaseScriptID));
+	HashState.Update(BaseCompileHash.GetData(), FNiagaraCompileHash::HashSize);
 	HashState.Update((const uint8*)&FeatureLevel, sizeof(FeatureLevel));
-
-	//ParameterSet.UpdateHash(HashState);		// will need for static switches
 		
+	for (int32 Index = 0; Index < ReferencedCompileHashes.Num(); Index++)
+	{
+		HashState.Update(ReferencedCompileHashes[Index].GetData(), FNiagaraCompileHash::HashSize);
+	}
+
 	for (int32 Index = 0; Index < ReferencedDependencyIds.Num(); Index++)
 	{
 		HashState.Update((const uint8*)&ReferencedDependencyIds[Index], sizeof(ReferencedDependencyIds[Index]));
 	}
+
+	//ParameterSet.UpdateHash(HashState);		// will need for static switches
 	
 	HashState.Final();
 	HashState.GetHash(&OutHash.Hash[0]);
@@ -138,12 +168,26 @@ void FNiagaraShaderMapId::GetScriptHash(FSHAHash& OutHash) const
 bool FNiagaraShaderMapId::operator==(const FNiagaraShaderMapId& ReferenceSet) const
 {
 	if (BaseScriptID != ReferenceSet.BaseScriptID 
+		|| BaseCompileHash != ReferenceSet.BaseCompileHash
 		|| FeatureLevel != ReferenceSet.FeatureLevel
 		|| CompilerVersionID != ReferenceSet.CompilerVersionID)
 	{
 		return false;
 	}
 
+	if (ReferencedCompileHashes.Num() != ReferenceSet.ReferencedCompileHashes.Num())
+	{
+		return false;
+	}
+
+	for (int32 i = 0; i < ReferencedCompileHashes.Num(); i++)
+	{
+		if (ReferencedCompileHashes[i] != ReferenceSet.ReferencedCompileHashes[i])
+		{
+			return false;
+		}
+	}
+
 	if (ReferencedDependencyIds.Num() != ReferenceSet.ReferencedDependencyIds.Num())
 	{
 		return false;
@@ -181,6 +225,9 @@ void FNiagaraShaderMapId::AppendKeyString(FString& KeyString) const
 	KeyString += BaseScriptID.ToString();
 	KeyString += TEXT("_");
 
+	KeyString += BaseCompileHash.ToString();
+	KeyString += TEXT("_");
+
 	FString FeatureLevelString;
 	GetFeatureLevelName(FeatureLevel, FeatureLevelString);
 	KeyString += FeatureLevelString + TEXT("_");
@@ -188,6 +235,16 @@ void FNiagaraShaderMapId::AppendKeyString(FString& KeyString) const
 	KeyString += CompilerVersionID.ToString();
 	KeyString += TEXT("_");
 
+	// Add any referenced top level compile hashes to the key so that we will recompile when they are changed
+	for (int32 HashIndex = 0; HashIndex < ReferencedCompileHashes.Num(); HashIndex++)
+	{
+		KeyString += ReferencedCompileHashes[HashIndex].ToString();
+		if (HashIndex < ReferencedCompileHashes.Num() - 1)
+		{
+			KeyString += TEXT("_");
+		}
+	}
+
 	// Add any referenced functions to the key so that we will recompile when they are changed
 	for (int32 FunctionIndex = 0; FunctionIndex < ReferencedDependencyIds.Num(); FunctionIndex++)
 	{
@@ -238,7 +295,7 @@ FShaderCompileJob* FNiagaraShaderType::BeginCompileShader(
 	NewJob->Input.SharedEnvironment = CompilationEnvironment;
 	NewJob->Input.Target = Target;
 	NewJob->Input.ShaderFormat = LegacyShaderPlatformToShaderFormat(Platform);
-	NewJob->Input.VirtualSourceFilePath = TEXT("/Engine/Private/NiagaraEmitterInstanceShader.usf");
+	NewJob->Input.VirtualSourceFilePath = TEXT("/Plugin/FX/Niagara/Private/NiagaraEmitterInstanceShader.usf");
 	NewJob->Input.EntryPointName = TEXT("SimulateMainComputeCS");
 	NewJob->Input.Environment.SetDefine(TEXT("GPU_SIMULATION"), 1);
 	NewJob->Input.Environment.IncludeVirtualPathToContentsMap.Add(TEXT("/Engine/Generated/NiagaraEmitterInstance.ush"), Script->HlslOutput);
@@ -262,7 +319,7 @@ FShaderCompileJob* FNiagaraShaderType::BeginCompileShader(
 		nullptr,
 		this,
 		nullptr,//ShaderPipeline,
-		TEXT("/Engine/Private/NiagaraEmitterInstanceShader.usf"),
+		TEXT("/Plugin/FX/Niagara/Private/NiagaraEmitterInstanceShader.usf"),
 		TEXT("SimulateMainComputeCS"),
 		FShaderTarget(GetFrequency(), Platform),
 		NewJob,
@@ -728,6 +785,8 @@ void FNiagaraShaderMap::Compile(
 		// Since it creates a temporary ref counted pointer.
 		check(NumRefs > 0);
   
+		//All access to NiagaraShaderMapsBeingCompiled must be done on the game thread!
+		check(IsInGameThread());
 		// Add this shader map and to NiagaraShaderMapsBeingCompiled
 		TArray<FNiagaraShaderScript*>* CorrespondingScripts = NiagaraShaderMapsBeingCompiled.Find(this);
   
@@ -888,6 +947,8 @@ bool FNiagaraShaderMap::ProcessCompilationResults(const TArray<FShaderCommonComp
 bool FNiagaraShaderMap::TryToAddToExistingCompilationTask(FNiagaraShaderScript* Script)
 {
 	check(NumRefs > 0);
+	//All access to NiagaraShaderMapsBeingCompiled must be done on the game thread!
+	check(IsInGameThread());
 	TArray<FNiagaraShaderScript*>* CorrespondingScripts = FNiagaraShaderMap::NiagaraShaderMapsBeingCompiled.Find(this);
 
 	if (CorrespondingScripts)
@@ -926,6 +987,8 @@ bool FNiagaraShaderMap::IsComplete(const FNiagaraShaderScript* Script, bool bSil
 	// Make sure we are operating on a referenced shader map or the below Find will cause this shader map to be deleted,
 	// Since it creates a temporary ref counted pointer.
 	check(NumRefs > 0);
+	//All access to NiagaraShaderMapsBeingCompiled must be done on the game thread!
+	check(IsInGameThread());
 	const TArray<FNiagaraShaderScript*>* CorrespondingScripts = FNiagaraShaderMap::NiagaraShaderMapsBeingCompiled.Find(this);
 
 	if (CorrespondingScripts)
@@ -954,6 +1017,8 @@ void FNiagaraShaderMap::LoadMissingShadersFromMemory(const FNiagaraShaderScript*
 	// Since it creates a temporary ref counted pointer.
 	check(NumRefs > 0);
 
+	//All access to NiagaraShaderMapsBeingCompiled must be done on the game thread!
+	check(IsInGameThread());
 	const TArray<FNiagaraShaderScript*>* CorrespondingScripts = FNiagaraShaderMap::NiagaraShaderMapsBeingCompiled.Find(this);
 
 	if (CorrespondingScripts)
@@ -1123,6 +1188,8 @@ void FNiagaraShaderMap::DiscardSerializedShaders()
 
 void FNiagaraShaderMap::RemovePendingScript(FNiagaraShaderScript* Script)
 {
+	//All access to NiagaraShaderMapsBeingCompiled must be done on the game thread!
+	check(IsInGameThread());
 	for (TMap<TRefCountPtr<FNiagaraShaderMap>, TArray<FNiagaraShaderScript*> >::TIterator It(NiagaraShaderMapsBeingCompiled); It; ++It)
 	{
 		TArray<FNiagaraShaderScript*>& Scripts = It.Value();
@@ -1144,6 +1211,8 @@ void FNiagaraShaderMap::RemovePendingScript(FNiagaraShaderScript* Script)
 
 void FNiagaraShaderMap::RemovePendingMap(FNiagaraShaderMap* Map)
 {
+	//All access to NiagaraShaderMapsBeingCompiled must be done on the game thread!
+	check(IsInGameThread());
 	TArray<FNiagaraShaderScript*>* Scripts = NiagaraShaderMapsBeingCompiled.Find(Map);
 	if (Scripts)
 	{
@@ -1160,6 +1229,8 @@ void FNiagaraShaderMap::RemovePendingMap(FNiagaraShaderMap* Map)
 const FNiagaraShaderMap* FNiagaraShaderMap::GetShaderMapBeingCompiled(const FNiagaraShaderScript* Script)
 {
 	// Inefficient search, but only when compiling a lot of shaders
+	//All access to NiagaraShaderMapsBeingCompiled must be done on the game thread!
+	check(IsInGameThread());
 	for (TMap<TRefCountPtr<FNiagaraShaderMap>, TArray<FNiagaraShaderScript*> >::TIterator It(NiagaraShaderMapsBeingCompiled); It; ++It)
 	{
 		TArray<FNiagaraShaderScript*>& Scripts = It.Value();
diff --git a/Engine/Plugins/FX/Niagara/Source/NiagaraShader/Private/NiagaraShaderCompilationManager.cpp b/Engine/Plugins/FX/Niagara/Source/NiagaraShader/Private/NiagaraShaderCompilationManager.cpp
--- a/Engine/Plugins/FX/Niagara/Source/NiagaraShader/Private/NiagaraShaderCompilationManager.cpp
+++ b/Engine/Plugins/FX/Niagara/Source/NiagaraShader/Private/NiagaraShaderCompilationManager.cpp
@@ -20,13 +20,103 @@ static FAutoConsoleVariableRef CVarShowNiagaraShaderWarnings(
 	TEXT("When set to 1, will display all warnings from Niagara shader compiles.")
 	);
 
+//Disabled until a few issues can be resolved.
+int32 GEnableNiagaraShaderCompilerAsyncHack = 0;
+static FAutoConsoleVariableRef CVarEnableNiagaraShaderCompilerAsyncHack(
+	TEXT("fx.EnableNiagaraShaderCompilerAsyncHack"),
+	GEnableNiagaraShaderCompilerAsyncHack,
+	TEXT("If > 0 GPU Niagara shaders will be compiled off the main therad. This is a temporary solution until these are offloaded to the shader compiler worker. \n"),
+	ECVF_ReadOnly
+);
+
+//#define ASYNC_HACK_LOG(...)
+#define ASYNC_HACK_LOG(Fmt, ...) if( IsAsyncHackEnabled() ){UE_LOG(LogNiagaraShaderCompiler, Log, Fmt, ##__VA_ARGS__);}
+
 #if WITH_EDITOR
 
+
+
+//ASYNC HACK: Temporary code to run this on a worker thread. This will eventually be unnecessary as we'll move this to shadercompilerworker.
+//ASYNC HACK: There is a lot of code in here to do with this code being run on shader compiler worker.
+//ASYNC HACK: For now I'm leaving this in as these hacks will soon be removed.
+//ASYNC HACK: At that point, all this other code will make sense.
+//ASYNC HACK: During this time it's going to be very confusing to read this code.
+
+//ASYNC HACK: This async task will be kicked from the GT and will run all current work in JobQueue.
+//ASYNC HACK: JobQueue can be added to in this time meaning this worker will run for longer.
+//ASYNC HACK: Once this task completes and more Jobs come in, another task will be kicked off. 
+//ASYNC HACK: Only one of these tasks can ever be in flight at once! The shader compilation is not thread-safe. 
+//ASYNC HACK: It doesn't have to be on the game thread but it can't be parallelized.
+
+class FNiagaraShaderCompilerWorkerAsyncTask
+{
+	FNiagaraShaderCompilationManager* CompilerWorker;
+
+public:
+	FNiagaraShaderCompilerWorkerAsyncTask(FNiagaraShaderCompilationManager* InCompilerWorker)
+		: CompilerWorker(InCompilerWorker)
+	{
+
+	}
+
+	FORCEINLINE TStatId GetStatId() const
+	{
+		RETURN_QUICK_DECLARE_CYCLE_STAT(FNiagaraShaderCompilerWorkerAsyncTask, STATGROUP_TaskGraphTasks);
+	}
+
+	ENamedThreads::Type GetDesiredThread()
+	{
+		return ENamedThreads::AnyBackgroundThreadNormalTask;
+	}
+
+	static ESubsequentsMode::Type GetSubsequentsMode() { return ESubsequentsMode::TrackSubsequents; }
+
+	void DoTask(ENamedThreads::Type CurrentThread, const FGraphEventRef& MyCompletionGraphEvent)
+	{
+		check(CompilerWorker);
+		CompilerWorker->RunCompileJobs();
+	}
+};
+
+
 NIAGARASHADER_API FNiagaraShaderCompilationManager GNiagaraShaderCompilationManager;
 
-void FNiagaraShaderCompilationManager::Tick(float DeltaSeconds)
+bool FNiagaraShaderCompilationManager::IsAsyncHackEnabled()const
 {
-	RunCompileJobs();
+	return FApp::ShouldUseThreadingForPerformance() && GEnableNiagaraShaderCompilerAsyncHack;
+}
+
+void FNiagaraShaderCompilationManager::Tick(float DeltaSeconds, bool bBlock)
+{
+	if (IsAsyncHackEnabled())
+	{
+		//ASYNC HACK: Temporary code to run this on a worker thread. This will eventually be unnecessary as we'll move this to shadercompilerworker.
+		check(IsInGameThread());
+	
+		//Kick off a new task if the last one has finished processing jobs.
+		if (!AsyncWork || AsyncWork->IsComplete())
+		{
+			if (JobQueue.Num() > 0)
+			{
+				ASYNC_HACK_LOG(TEXT("AsyncHack: New Compiler Async Task"));
+				AsyncWork = TGraphTask<FNiagaraShaderCompilerWorkerAsyncTask>::CreateTask(nullptr, ENamedThreads::GameThread).ConstructAndDispatchWhenReady(this);
+			}
+		}
+
+		if (bBlock && AsyncWork && !AsyncWork->IsComplete())
+		{
+			float StartTime = FPlatformTime::Seconds();
+			FTaskGraphInterface::Get().WaitUntilTaskCompletes(AsyncWork, ENamedThreads::GameThread_Local);
+			AsyncWork = nullptr;
+			float EndTime = FPlatformTime::Seconds();
+			ASYNC_HACK_LOG(TEXT("AsyncHack: Blocked on Compiler Async Task for %5.6f Seconds"), EndTime - StartTime);
+		}
+		//ASYNC HACK END: Temporary code to run this on a worker thread. This will eventually be unnecessary as we'll move this to shadercompilerworker.
+	}
+	else
+	{
+		RunCompileJobs();
+	}
 }
 
 FNiagaraShaderCompilationManager::FNiagaraShaderCompilationManager()
@@ -39,8 +129,9 @@ FNiagaraShaderCompilationManager::FNiagaraShaderCompilationManager()
 	{
 		WorkerInfos.Add(new FNiagaraShaderCompileWorkerInfo());
 	}
-}
 
+	AsyncWork = nullptr;
+}
 
 void FNiagaraShaderCompilationManager::RunCompileJobs()
 {
@@ -49,6 +140,7 @@ void FNiagaraShaderCompilationManager::RunCompileJobs()
 	int32 NumActiveThreads = 0;
 
 
+	ASYNC_HACK_LOG(TEXT("AsyncHack: RunCompileJobs()"));
 
 	for (int32 WorkerIndex = 0; WorkerIndex < WorkerInfos.Num(); WorkerIndex++)
 	{
@@ -59,8 +151,19 @@ void FNiagaraShaderCompilationManager::RunCompileJobs()
 		{
 			check(!CurrentWorkerInfo.bComplete);
 
+			//ASYNC HACK: Temporary code to run this on a worker thread. This will eventually be unnecessary as we'll move this to shadercompilerworker.
+			if (IsAsyncHackEnabled())
+			{
+				check(!IsInGameThread());
+				//Lock access to the job queue.
+				JobQueueCriticalSection.Lock();
+			}
+			//ASYNC HACK: Temporary code to run this on a worker thread. This will eventually be unnecessary as we'll move this to shadercompilerworker.
+
 			if (JobQueue.Num() > 0)
 			{
+				ASYNC_HACK_LOG(TEXT("AsyncHack: Taking %d Jobs from the Job Queue()"), JobQueue.Num());
+
 				bool bAddedLowLatencyTask = false;
 				int32 JobIndex = 0;
 
@@ -78,6 +181,13 @@ void FNiagaraShaderCompilationManager::RunCompileJobs()
 				CurrentWorkerInfo.StartTime = FPlatformTime::Seconds();
 				JobQueue.RemoveAt(0, JobIndex);
 			}
+
+			//ASYNC HACK: Temporary code to run this on a worker thread. This will eventually be unnecessary as we'll move this to shadercompilerworker.
+			if (IsAsyncHackEnabled())
+			{
+				JobQueueCriticalSection.Unlock();
+			}
+			//ASYNC HACK: Temporary code to run this on a worker thread. This will eventually be unnecessary as we'll move this to shadercompilerworker.
 		}
 
 		if (CurrentWorkerInfo.bIssuedTasksToWorker && CurrentWorkerInfo.bLaunchedWorker)
@@ -91,6 +201,7 @@ void FNiagaraShaderCompilationManager::RunCompileJobs()
 			{
 				FShaderCompileJob& CurrentJob = *((FShaderCompileJob*)(CurrentWorkerInfo.QueuedJobs[JobIndex]));
 
+				ASYNC_HACK_LOG(TEXT("AsyncHack: Processing Compile Job %s"), *CurrentJob.Input.GenerateShaderName());
 				check(!CurrentJob.bFinalized);
 				CurrentJob.bFinalized = true;
 
@@ -150,6 +261,15 @@ void FNiagaraShaderCompilationManager::RunCompileJobs()
 		}
 	}
 
+
+	//ASYNC HACK: Temporary code to run this on a worker thread. This will eventually be unnecessary as we'll move this to shadercompilerworker.
+	if (IsAsyncHackEnabled())
+	{
+		//Lock access to the results queue.
+		ResultsQueueCriticalSection.Lock();
+	}
+	//ASYNC HACK: Temporary code to run this on a worker thread. This will eventually be unnecessary as we'll move this to shadercompilerworker.
+
 	for (int32 WorkerIndex = 0; WorkerIndex < WorkerInfos.Num(); WorkerIndex++)
 	{
 		FNiagaraShaderCompileWorkerInfo& CurrentWorkerInfo = *WorkerInfos[WorkerIndex];
@@ -167,12 +287,28 @@ void FNiagaraShaderCompilationManager::RunCompileJobs()
 		CurrentWorkerInfo.QueuedJobs.Empty();
 	}
 
+	//ASYNC HACK: Temporary code to run this on a worker thread. This will eventually be unnecessary as we'll move this to shadercompilerworker.
+	if (IsAsyncHackEnabled())
+	{
+		ResultsQueueCriticalSection.Unlock();
+	}
+	//ASYNC HACK: Temporary code to run this on a worker thread. This will eventually be unnecessary as we'll move this to shadercompilerworker.
 }
 
 
 
 NIAGARASHADER_API void FNiagaraShaderCompilationManager::AddJobs(TArray<FShaderCommonCompileJob*> InNewJobs)
 {
+	//ASYNC HACK: Temporary code to run this on a worker thread. This will eventually be unnecessary as we'll move this to shadercompilerworker.
+	if (IsAsyncHackEnabled())
+	{
+		JobQueueCriticalSection.Lock();
+		ResultsQueueCriticalSection.Lock();
+		ASYNC_HACK_LOG(TEXT("AsyncHack: Adding %d jobs"), InNewJobs.Num());
+	}
+	//ASYNC HACK: Temporary code to run this on a worker thread. This will eventually be unnecessary as we'll move this to shadercompilerworker.
+
+
 	for (FShaderCommonCompileJob *Job : InNewJobs)
 	{
 		FNiagaraShaderMapCompileResults& ShaderMapInfo = NiagaraShaderMapJobs.FindOrAdd(Job->Id);
@@ -182,14 +318,31 @@ NIAGARASHADER_API void FNiagaraShaderCompilationManager::AddJobs(TArray<FShaderC
 	}
 
 	JobQueue.Append(InNewJobs);
+
+	//ASYNC HACK: Temporary code to run this on a worker thread. This will eventually be unnecessary as we'll move this to shadercompilerworker.
+	if (IsAsyncHackEnabled())
+	{
+		ResultsQueueCriticalSection.Unlock();
+		JobQueueCriticalSection.Unlock();
+	}
+	//ASYNC HACK: Temporary code to run this on a worker thread. This will eventually be unnecessary as we'll move this to shadercompilerworker.
 }
 
 
 void FNiagaraShaderCompilationManager::ProcessAsyncResults()
 {
+	check(IsInGameThread());
+
 	int32 NumCompilingNiagaraShaderMaps = 0;
 	TArray<int32> ShaderMapsToRemove;
 
+	//ASYNC HACK: Temporary code to run this on a worker thread. This will eventually be unnecessary as we'll move this to shadercompilerworker.
+	if (IsAsyncHackEnabled())
+	{
+		ResultsQueueCriticalSection.Lock();
+	}
+	//ASYNC HACK: Temporary code to run this on a worker thread. This will eventually be unnecessary as we'll move this to shadercompilerworker.
+
 	// Get all Niagara shader maps to finalize
 	//
 	for (TMap<int32, FNiagaraShaderMapCompileResults>::TIterator It(NiagaraShaderMapJobs); It; ++It)
@@ -200,6 +353,8 @@ void FNiagaraShaderCompilationManager::ProcessAsyncResults()
 		{
 			ShaderMapsToRemove.Add(It.Key());
 			PendingFinalizeNiagaraShaderMaps.Add(It.Key(), FNiagaraShaderMapFinalizeResults(Results));
+
+			//ASYNC_HACK_LOG(TEXT("AsyncHack: Grabbing compiled shader map to process."));
 		}
 	}
 
@@ -210,6 +365,13 @@ void FNiagaraShaderCompilationManager::ProcessAsyncResults()
 
 	NumCompilingNiagaraShaderMaps = NiagaraShaderMapJobs.Num();
 
+	//ASYNC HACK: Temporary code to run this on a worker thread. This will eventually be unnecessary as we'll move this to shadercompilerworker.
+	if (IsAsyncHackEnabled())
+	{
+		ResultsQueueCriticalSection.Unlock();
+	}
+	//ASYNC HACK: Temporary code to run this on a worker thread. This will eventually be unnecessary as we'll move this to shadercompilerworker.
+
 	if (PendingFinalizeNiagaraShaderMaps.Num() > 0)
 	{
 		ProcessCompiledNiagaraShaderMaps(PendingFinalizeNiagaraShaderMaps, 0.1f);
@@ -221,6 +383,10 @@ void FNiagaraShaderCompilationManager::ProcessCompiledNiagaraShaderMaps(
 	TMap<int32, FNiagaraShaderMapFinalizeResults>& CompiledShaderMaps,
 	float TimeBudget)
 {
+	check(IsInGameThread());
+
+	ASYNC_HACK_LOG(TEXT("AsyncHack: Processing %d Compiled Shader Maps"), CompiledShaderMaps.Num());
+
 	// Keeps shader maps alive as they are passed from the shader compiler and applied to the owning Script
 	TArray<TRefCountPtr<FNiagaraShaderMap> > LocalShaderMapReferences;
 	TMap<FNiagaraShaderScript*, FNiagaraShaderMap*> ScriptsToUpdate;
@@ -251,12 +417,14 @@ void FNiagaraShaderCompilationManager::ProcessCompiledNiagaraShaderMaps(
 			// Make a copy of the array as this entry of FNiagaraShaderMap::ShaderMapsBeingCompiled will be removed below
 			TArray<FNiagaraShaderScript*> ScriptArray = *Scripts;
 			bool bSuccess = true;
-			
+
+			ASYNC_HACK_LOG(TEXT("AsyncHack: Processing %d Finished Jobs"), ResultArray.Num());
 			for (int32 JobIndex = 0; JobIndex < ResultArray.Num(); JobIndex++)
 			{
 				FShaderCompileJob& CurrentJob = *((FShaderCompileJob*)(ResultArray[JobIndex]));
 				bSuccess = bSuccess && CurrentJob.bSucceeded;
 
+				ASYNC_HACK_LOG(TEXT("AsyncHack: Processing Finished Job %s"), *CurrentJob.Input.GenerateShaderName());
 				if (bSuccess)
 				{
 					check(CurrentJob.Output.ShaderCode.GetShaderCodeSize() > 0);
@@ -271,10 +439,10 @@ void FNiagaraShaderCompilationManager::ProcessCompiledNiagaraShaderMaps(
 
 					if (CurrentJob.Output.Errors.Num())
 					{
-						UE_LOG(LogNiagaraShaderCompiler, Warning, TEXT("There were errors for job \"%s\""), *CurrentJob.Input.DebugGroupName)
+						UE_LOG(LogNiagaraShaderCompiler, Error, TEXT("There were errors for job \"%s\""), *CurrentJob.Input.DebugGroupName)
 							for (const FShaderCompilerError& Error : CurrentJob.Output.Errors)
 							{
-								UE_LOG(LogShaders, Log, TEXT("Error: %s"), *Error.GetErrorString())
+								UE_LOG(LogNiagaraShaderCompiler, Error, TEXT("Error: %s"), *Error.GetErrorString())
 							}
 					}
 				}
@@ -419,7 +587,17 @@ void FNiagaraShaderCompilationManager::FinishCompilation(const TCHAR* ScriptName
 {
 	check(!FPlatformProperties::RequiresCookedData());
 
-	RunCompileJobs();	// since we don't async compile through another process, this will run all oustanding jobs
+	//ASYNC HACK: Temporary code to run this on a worker thread. This will eventually be unnecessary as we'll move this to shadercompilerworker.
+	if (IsAsyncHackEnabled())
+	{
+		Tick(0.0f, true);//Kick current jobs and block.
+		//ASYNC HACK: Temporary code to run this on a worker thread. This will eventually be unnecessary as we'll move this to shadercompilerworker.
+	}
+	else
+	{
+		RunCompileJobs();	// since we don't async compile through another process, this will run all oustanding jobs
+	}
+	
 	ProcessAsyncResults();	// grab compiled shader maps and assign them to their resources
 
 	check(NiagaraShaderMapJobs.Num() == 0);
diff --git a/Engine/Plugins/FX/Niagara/Source/NiagaraShader/Private/NiagaraShaderDerivedDataVersion.h b/Engine/Plugins/FX/Niagara/Source/NiagaraShader/Private/NiagaraShaderDerivedDataVersion.h
--- a/Engine/Plugins/FX/Niagara/Source/NiagaraShader/Private/NiagaraShaderDerivedDataVersion.h
+++ b/Engine/Plugins/FX/Niagara/Source/NiagaraShader/Private/NiagaraShaderDerivedDataVersion.h
@@ -11,5 +11,4 @@ NiagaraShaderDerivedDataVersion.h: Shader derived data version for Niagara.
 // In case of merge conflicts with DDC versions, you MUST generate a new GUID and set this new
 // guid as version
 
-#define NIAGARASHADERMAP_DERIVEDDATA_VER		TEXT("F3D3254CE5E3409D9FB3E774D0E0C8D9")
-													  
+#define NIAGARASHADERMAP_DERIVEDDATA_VER		TEXT("DE7107EEE3F66B4EA8F73C17D107D3E6")
diff --git a/Engine/Plugins/FX/Niagara/Source/NiagaraShader/Private/NiagaraShaderModule.cpp b/Engine/Plugins/FX/Niagara/Source/NiagaraShader/Private/NiagaraShaderModule.cpp
--- a/Engine/Plugins/FX/Niagara/Source/NiagaraShader/Private/NiagaraShaderModule.cpp
+++ b/Engine/Plugins/FX/Niagara/Source/NiagaraShader/Private/NiagaraShaderModule.cpp
@@ -1,12 +1,24 @@
 // Copyright 1998-2019 Epic Games, Inc. All Rights Reserved.
 
 #include "NiagaraShaderModule.h"
+#include "Misc/Paths.h"
 #include "Modules/ModuleManager.h"
+#include "Interfaces/IPluginManager.h"
+#include "ShaderCore.h"
 
 IMPLEMENT_MODULE(INiagaraShaderModule, NiagaraShader);
 
 INiagaraShaderModule* INiagaraShaderModule::Singleton(nullptr);
 
+void INiagaraShaderModule::StartupModule()
+{
+	Singleton = this;
+
+	// Maps virtual shader source directory /Plugin/FX/Niagara to the plugin's actual Shaders directory.
+	FString PluginShaderDir = FPaths::Combine(IPluginManager::Get().FindPlugin(TEXT("Niagara"))->GetBaseDir(), TEXT("Shaders"));
+	AddShaderSourceDirectoryMapping(TEXT("/Plugin/FX/Niagara"), PluginShaderDir);
+}
+
 FDelegateHandle INiagaraShaderModule::SetOnProcessShaderCompilationQueue(FOnProcessQueue InOnProcessQueue)
 {
 	checkf(OnProcessQueue.IsBound() == false, TEXT("Shader processing queue delegate already set."));
diff --git a/Engine/Plugins/FX/Niagara/Source/NiagaraShader/Private/NiagaraShared.cpp b/Engine/Plugins/FX/Niagara/Source/NiagaraShader/Private/NiagaraShared.cpp
--- a/Engine/Plugins/FX/Niagara/Source/NiagaraShader/Private/NiagaraShared.cpp
+++ b/Engine/Plugins/FX/Niagara/Source/NiagaraShader/Private/NiagaraShared.cpp
@@ -26,10 +26,8 @@
 FNiagaraShaderScript::~FNiagaraShaderScript()
 {
 #if WITH_EDITOR
-	if (IsInGameThread())
-	{
-		FNiagaraCompilationQueue::Get()->RemovePending(this);
-	}
+	check(IsInGameThread());
+	CancelCompilation();
 #endif
 }
 
@@ -57,19 +55,18 @@ NIAGARASHADER_API void FNiagaraShaderScript::NotifyCompilationFinished()
 NIAGARASHADER_API void FNiagaraShaderScript::CancelCompilation()
 {
 #if WITH_EDITOR
-	if (IsInGameThread())
-	{
-		FNiagaraShaderMap::RemovePendingScript(this);
-		FNiagaraCompilationQueue::Get()->RemovePending(this);
+	check(IsInGameThread());
+	FNiagaraShaderMap::RemovePendingScript(this);
+	FNiagaraCompilationQueue::Get()->RemovePending(this);
 
-		UE_LOG(LogShaders, Log, TEXT("CancelCompilation %p."), this);
-		OutstandingCompileShaderMapIds.Empty();
-	}
+	UE_LOG(LogShaders, Log, TEXT("CancelCompilation %p."), this);
+	OutstandingCompileShaderMapIds.Empty();
 #endif
 }
 
 NIAGARASHADER_API void FNiagaraShaderScript::RemoveOutstandingCompileId(const int32 OldOutstandingCompileShaderMapId)
 {
+	check(IsInGameThread());
 	if (0 <= OutstandingCompileShaderMapIds.Remove(OldOutstandingCompileShaderMapId))
 	{
 		UE_LOG(LogShaders, Log, TEXT("RemoveOutstandingCompileId %p %d"), this, OldOutstandingCompileShaderMapId);
@@ -89,7 +86,13 @@ NIAGARASHADER_API void FNiagaraShaderScript::LegacySerialize(FArchive& Ar)
 
 bool FNiagaraShaderScript::IsSame(const FNiagaraShaderMapId& InId) const
 {
-	return InId.BaseScriptID == BaseScriptId && InId.ReferencedDependencyIds == ReferencedDependencyIds && InId.CompilerVersionID == CompilerVersionId;
+	return
+		InId.FeatureLevel == FeatureLevel &&
+		InId.BaseScriptID == BaseScriptId &&
+		InId.BaseCompileHash == BaseCompileHash &&
+		InId.ReferencedCompileHashes == ReferencedCompileHashes &&
+		InId.ReferencedDependencyIds == ReferencedDependencyIds &&
+		InId.CompilerVersionID == CompilerVersionId;
 }
 
 
@@ -120,6 +123,8 @@ NIAGARASHADER_API void FNiagaraShaderScript::GetShaderMapId(EShaderPlatform Plat
 		GetDependentShaderTypes(Platform, ShaderTypes);
 		OutId.FeatureLevel = GetFeatureLevel();
 		OutId.BaseScriptID = BaseScriptId;
+		OutId.BaseCompileHash = BaseCompileHash;
+		OutId.ReferencedCompileHashes = ReferencedCompileHashes;
 		OutId.ReferencedDependencyIds = ReferencedDependencyIds;
 		OutId.CompilerVersionID = FNiagaraCustomVersion::LatestScriptCompileVersion;
 	}
@@ -219,11 +224,16 @@ void FNiagaraShaderScript::SerializeShaderMap(FArchive& Ar)
 	}
 }
 
-void FNiagaraShaderScript::SetScript(UNiagaraScript *InScript, ERHIFeatureLevel::Type InFeatureLevel, const FGuid& InCompilerVersionID, const FGuid& InBaseScriptID, const TArray<FGuid>& InReferencedDependencyIds, FString InFriendlyName)
+void FNiagaraShaderScript::SetScript(UNiagaraScript *InScript, ERHIFeatureLevel::Type InFeatureLevel, const FGuid& InCompilerVersionID, const FGuid& InBaseScriptID,
+	const FNiagaraCompileHash& InBaseCompileHash, const TArray<FNiagaraCompileHash>& InReferencedCompileHashes, const TArray<FGuid>& InReferencedDependencyIds, FString InFriendlyName)
 {
+	checkf(InBaseScriptID.IsValid(), TEXT("Invalid base script id.  Script caching will fail."));
+	checkf(InBaseCompileHash.IsValid(), TEXT("Invalid base compile hash.  Script caching will fail."))
 	BaseVMScript = InScript;
-	BaseScriptId = InBaseScriptID;
 	CompilerVersionId = InCompilerVersionID;
+	BaseScriptId = InBaseScriptID;
+	BaseCompileHash = InBaseCompileHash;
+	ReferencedCompileHashes = InReferencedCompileHashes;
 	ReferencedDependencyIds = InReferencedDependencyIds;
 	FriendlyName = InFriendlyName;
 	SetFeatureLevel(InFeatureLevel);
@@ -237,6 +247,7 @@ NIAGARASHADER_API  void FNiagaraShaderScript::SetRenderingThreadShaderMap(FNiaga
 
 NIAGARASHADER_API  bool FNiagaraShaderScript::IsCompilationFinished() const
 {
+	check(IsInGameThread());
 	bool bRet = GameThreadShaderMap && GameThreadShaderMap.IsValid() && GameThreadShaderMap->IsCompilationFinalized();
 	if (OutstandingCompileShaderMapIds.Num() == 0)
 	{
@@ -264,6 +275,8 @@ bool FNiagaraShaderScript::CacheShaders(const FNiagaraShaderMapId& ShaderMapId,
 {
 	bool bSucceeded = false;
 
+	check(IsInGameThread());
+
 	{
 		// Find the script's cached shader map.
 		GameThreadShaderMap = FNiagaraShaderMap::FindId(ShaderMapId, Platform);
@@ -404,6 +417,7 @@ NIAGARASHADER_API  FNiagaraShader* FNiagaraShaderScript::GetShaderGameThread() c
 
 void FNiagaraShaderScript::GetShaderMapIDsWithUnfinishedCompilation(TArray<int32>& ShaderMapIds)
 {
+	check(IsInGameThread());
 	// Build an array of the shader map Id's are not finished compiling.
 	if (GameThreadShaderMap && GameThreadShaderMap.IsValid() && !GameThreadShaderMap->IsCompilationFinalized())
 	{
@@ -432,6 +446,7 @@ bool FNiagaraShaderScript::BeginCompileShaderMap(
 	bool bApplyCompletedShaderMapForRendering,
 	bool bSynchronous)
 {
+	check(IsInGameThread());
 #if WITH_EDITORONLY_DATA
 	bool bSuccess = false;
 
diff --git a/Engine/Plugins/FX/Niagara/Source/NiagaraShader/Public/NiagaraShaderCompilationManager.h b/Engine/Plugins/FX/Niagara/Source/NiagaraShader/Public/NiagaraShaderCompilationManager.h
--- a/Engine/Plugins/FX/Niagara/Source/NiagaraShader/Public/NiagaraShaderCompilationManager.h
+++ b/Engine/Plugins/FX/Niagara/Source/NiagaraShader/Public/NiagaraShaderCompilationManager.h
@@ -119,15 +119,17 @@ class FNiagaraShaderCompilationManager
 public:
 	FNiagaraShaderCompilationManager();
 
-	NIAGARASHADER_API void Tick(float DeltaSeconds = 0.0f);
+	NIAGARASHADER_API void Tick(float DeltaSeconds = 0.0f, bool bBlock=false);
 	NIAGARASHADER_API void AddJobs(TArray<class FShaderCommonCompileJob*> InNewJobs);
 	NIAGARASHADER_API void ProcessAsyncResults();
 
 	void FinishCompilation(const TCHAR* ScriptName, const TArray<int32>& ShaderMapIdsToFinishCompiling);
 
+	void RunCompileJobs();
+
+	FORCEINLINE bool IsAsyncHackEnabled()const;
 private:
 	void ProcessCompiledNiagaraShaderMaps(TMap<int32, FNiagaraShaderMapFinalizeResults>& CompiledShaderMaps, float TimeBudget);
-	void RunCompileJobs();
 
 	TArray<class FShaderCommonCompileJob*> JobQueue;
 
@@ -138,6 +140,21 @@ private:
 	TMap<int32, FNiagaraShaderMapFinalizeResults> PendingFinalizeNiagaraShaderMaps;
 
 	TArray<struct FNiagaraShaderCompileWorkerInfo*> WorkerInfos;
+
+
+	//////////////////////////////////////////////////////////////////////////
+	// Temp support for pushing this off onto a worker thread. 
+	// Longer term, this should be removed and this whole thing can be moved to shader compiler worker.
+	
+	//Critical section guarding interaction with the shared job queue.
+	FCriticalSection JobQueueCriticalSection;
+
+	//Critical section guarding interaction with the shared results queue.
+	FCriticalSection ResultsQueueCriticalSection;
+
+	//Handle to the last task we kicked off to compile previously.
+	//We can only have one compilation task going at a time as (apparently) the compilation is very non thread-safe
+	FGraphEventRef AsyncWork;
 };
 
 
diff --git a/Engine/Plugins/FX/Niagara/Source/NiagaraShader/Public/NiagaraShaderModule.h b/Engine/Plugins/FX/Niagara/Source/NiagaraShader/Public/NiagaraShaderModule.h
--- a/Engine/Plugins/FX/Niagara/Source/NiagaraShader/Public/NiagaraShaderModule.h
+++ b/Engine/Plugins/FX/Niagara/Source/NiagaraShader/Public/NiagaraShaderModule.h
@@ -25,11 +25,7 @@ public:
 	void NIAGARASHADER_API ResetOnRequestDefaultDataInterfaceHandler();
 	UNiagaraDataInterfaceBase* RequestDefaultDataInterface(const FString& DIClassName);
 
-	virtual void StartupModule() override
-	{
-		Singleton = this;
-	}
-
+	virtual void StartupModule() override;
 
 	virtual void ShutdownModule() override
 	{
diff --git a/Engine/Plugins/FX/Niagara/Source/NiagaraShader/Public/NiagaraShared.h b/Engine/Plugins/FX/Niagara/Source/NiagaraShader/Public/NiagaraShared.h
--- a/Engine/Plugins/FX/Niagara/Source/NiagaraShader/Public/NiagaraShared.h
+++ b/Engine/Plugins/FX/Niagara/Source/NiagaraShader/Public/NiagaraShared.h
@@ -22,6 +22,7 @@
 #include "SceneTypes.h"
 #include "StaticParameterSet.h"
 #include "Misc/Optional.h"
+#include "NiagaraCompileHash.h"
 #include "NiagaraShared.generated.h"
 
 class FNiagaraShaderScript;
@@ -156,17 +157,24 @@ public:
 	ERHIFeatureLevel::Type FeatureLevel;
 
 	/**
-	* The GUID of the subgraph this shader primarily represents.
+	* The base id of the subgraph this shader primarily represents.
 	*/
 	FGuid BaseScriptID;
 
+	/**
+	* The hash of the subgraph this shader primarily represents.
+	*/
+	FNiagaraCompileHash BaseCompileHash;
+
+	/** The compile hashes of the top level scripts the script is dependent on. */
+	TArray<FNiagaraCompileHash> ReferencedCompileHashes;
+
 	/** Guids of any functions or module scripts the script was dependent on. */
 	TArray<FGuid> ReferencedDependencyIds;
 
 	FNiagaraShaderMapId()
 		: CompilerVersionID()
-		, FeatureLevel(GMaxRHIFeatureLevel) 
-		, BaseScriptID(0, 0, 0, 0)
+		, FeatureLevel(GMaxRHIFeatureLevel)
 	{ }
 
 	~FNiagaraShaderMapId()
@@ -422,7 +430,12 @@ public:
 
 	int32 GetNumRefs() const { return NumRefs; }
 	uint32 GetCompilingId()  { return CompilingId; }
-	static TMap<TRefCountPtr<FNiagaraShaderMap>, TArray<FNiagaraShaderScript*> > &GetInFlightShaderMaps() { return NiagaraShaderMapsBeingCompiled; }
+	static TMap<TRefCountPtr<FNiagaraShaderMap>, TArray<FNiagaraShaderScript*> > &GetInFlightShaderMaps() 
+	{
+		//All access to NiagaraShaderMapsBeingCompiled must be done on the game thread!
+		check(IsInGameThread());
+		return NiagaraShaderMapsBeingCompiled; 
+	}
 
 	void SetCompiledSuccessfully(bool bSuccess) { bCompiledSuccessfully = bSuccess; }
 private:
@@ -613,6 +626,7 @@ public:
 
 	void AddCompileId(uint32 Id) 
 	{
+		check(IsInGameThread());
 		OutstandingCompileShaderMapIds.Add(Id);
 	}
 
@@ -652,7 +666,8 @@ public:
 	const FString& GetFriendlyName()	const { return FriendlyName; }
 
 
-	NIAGARASHADER_API void SetScript(UNiagaraScript *InScript, ERHIFeatureLevel::Type InFeatureLevel, const FGuid& InCompilerVersion, const FGuid& InBaseScriptID, const TArray<FGuid>& InReferencedDependencyIds, FString InFriendlyName);
+	NIAGARASHADER_API void SetScript(UNiagaraScript *InScript, ERHIFeatureLevel::Type InFeatureLevel, const FGuid& InCompilerVersion, const FGuid& InBaseScriptID,
+		const FNiagaraCompileHash& InBaseCompileHash, const TArray<FNiagaraCompileHash>& InReferencedCompileHashes, const TArray<FGuid>& InReferencedDependencyIds, FString InFriendlyName);
 
 	UNiagaraScript *GetBaseVMScript()
 	{
@@ -727,9 +742,15 @@ private:
 	/** Guid id for base script*/
 	FGuid BaseScriptId;
 
+	/** Compile hash for the base script. */
+	FNiagaraCompileHash BaseCompileHash;
+
 	/** The compiler version the script was generated with.*/
 	FGuid CompilerVersionId;
 
+	/** The compile hashes for the top level scripts referenced by the script. */
+	TArray<FNiagaraCompileHash> ReferencedCompileHashes;
+
 	/** Dependencies of the script*/
 	TArray<FGuid> ReferencedDependencyIds;
 
diff --git a/Engine/Plugins/FX/Niagara/Source/NiagaraVertexFactories/Private/NiagaraMeshVertexFactory.cpp b/Engine/Plugins/FX/Niagara/Source/NiagaraVertexFactories/Private/NiagaraMeshVertexFactory.cpp
--- a/Engine/Plugins/FX/Niagara/Source/NiagaraVertexFactories/Private/NiagaraMeshVertexFactory.cpp
+++ b/Engine/Plugins/FX/Niagara/Source/NiagaraVertexFactories/Private/NiagaraMeshVertexFactory.cpp
@@ -12,10 +12,9 @@
 
 IMPLEMENT_GLOBAL_SHADER_PARAMETER_STRUCT(FNiagaraMeshUniformParameters,"NiagaraMeshVF");
 
-class FNiagaraMeshVertexFactoryShaderParameters : public FVertexFactoryShaderParameters
+class FNiagaraMeshVertexFactoryShaderParametersVS : public FVertexFactoryShaderParameters
 {
 public:
-
 	virtual void Bind(const FShaderParameterMap& ParameterMap) override
 	{
 		//PrevTransformBuffer.Bind(ParameterMap, TEXT("PrevTransformBuffer"));
@@ -90,9 +89,34 @@ private:
 	FShaderParameter MeshFacingMode;
 	FShaderResourceParameter SortedIndices;
 	FShaderParameter SortedIndicesOffset;
-
 };
 
+class FNiagaraMeshVertexFactoryShaderParametersPS : public FVertexFactoryShaderParameters
+{
+public:
+	virtual void Bind(const FShaderParameterMap& ParameterMap) override
+	{
+	}
+
+	virtual void Serialize(FArchive& Ar) override
+	{
+	}
+
+	virtual void GetElementShaderBindings(
+		const FSceneInterface* Scene,
+		const FSceneView* View,
+		const FMeshMaterialShader* Shader,
+		bool bShaderRequiresPositionOnlyStream,
+		ERHIFeatureLevel::Type FeatureLevel,
+		const FVertexFactory* VertexFactory,
+		const FMeshBatchElement& BatchElement,
+		class FMeshDrawSingleShaderBindings& ShaderBindings,
+		FVertexInputStreamArray& VertexStreams) const override
+	{
+		FNiagaraMeshVertexFactory* NiagaraMeshVF = (FNiagaraMeshVertexFactory*)VertexFactory;
+		ShaderBindings.Add(Shader->GetUniformBufferParameter<FNiagaraMeshUniformParameters>(), NiagaraMeshVF->GetUniformBuffer());
+	}
+};
 
 void FNiagaraMeshVertexFactory::InitRHI()
 {
@@ -211,9 +235,17 @@ void FNiagaraMeshVertexFactory::SetData(const FStaticMeshDataType& InData)
 
 FVertexFactoryShaderParameters* FNiagaraMeshVertexFactory::ConstructShaderParameters(EShaderFrequency ShaderFrequency)
 {
-	return ShaderFrequency == SF_Vertex ? new FNiagaraMeshVertexFactoryShaderParameters() : NULL;
+	if (ShaderFrequency == SF_Vertex)
+	{
+		return new FNiagaraMeshVertexFactoryShaderParametersVS();
+	}
+	else if (ShaderFrequency == SF_Pixel)
+	{
+		return new FNiagaraMeshVertexFactoryShaderParametersPS();
+	}
+	return nullptr;
 }
 
-IMPLEMENT_VERTEX_FACTORY_TYPE(FNiagaraMeshVertexFactory, "/Engine/Private/NiagaraMeshVertexFactory.ush", true, false, true, false, false);
-IMPLEMENT_VERTEX_FACTORY_TYPE(FNiagaraMeshVertexFactoryEmulatedInstancing, "/Engine/Private/NiagaraMeshVertexFactory.ush", true, false, true, false, false);
+IMPLEMENT_VERTEX_FACTORY_TYPE(FNiagaraMeshVertexFactory, "/Plugin/FX/Niagara/Private/NiagaraMeshVertexFactory.ush", true, false, true, false, false);
+IMPLEMENT_VERTEX_FACTORY_TYPE(FNiagaraMeshVertexFactoryEmulatedInstancing, "/Plugin/FX/Niagara/Private/NiagaraMeshVertexFactory.ush", true, false, true, false, false);
 
diff --git a/Engine/Plugins/FX/Niagara/Source/NiagaraVertexFactories/Private/NiagaraRibbonVertexFactory.cpp b/Engine/Plugins/FX/Niagara/Source/NiagaraVertexFactories/Private/NiagaraRibbonVertexFactory.cpp
--- a/Engine/Plugins/FX/Niagara/Source/NiagaraVertexFactories/Private/NiagaraRibbonVertexFactory.cpp
+++ b/Engine/Plugins/FX/Niagara/Source/NiagaraVertexFactories/Private/NiagaraRibbonVertexFactory.cpp
@@ -239,4 +239,4 @@ void FNiagaraRibbonVertexFactory::SetDynamicParameterBuffer(const FVertexBuffer*
 
 ///////////////////////////////////////////////////////////////////////////////
 
-IMPLEMENT_VERTEX_FACTORY_TYPE(FNiagaraRibbonVertexFactory, "/Engine/Private/NiagaraRibbonVertexFactory.ush", true, false, true, false, false);
+IMPLEMENT_VERTEX_FACTORY_TYPE(FNiagaraRibbonVertexFactory, "/Plugin/FX/Niagara/Private/NiagaraRibbonVertexFactory.ush", true, false, true, false, false);
diff --git a/Engine/Plugins/FX/Niagara/Source/NiagaraVertexFactories/Private/NiagaraSortingGPU.cpp b/Engine/Plugins/FX/Niagara/Source/NiagaraVertexFactories/Private/NiagaraSortingGPU.cpp
--- a/Engine/Plugins/FX/Niagara/Source/NiagaraVertexFactories/Private/NiagaraSortingGPU.cpp
+++ b/Engine/Plugins/FX/Niagara/Source/NiagaraVertexFactories/Private/NiagaraSortingGPU.cpp
@@ -16,6 +16,14 @@ static FAutoConsoleVariableRef CVarNiagaraGPUSorting(
 	ECVF_Default
 );
 
+int32 GNiagaraGPUSortingUseMaxPrecision = 1;
+static FAutoConsoleVariableRef CVarNiagaraGPUSortinUseMaxPrecision(
+	TEXT("Niagara.GPUSorting.UseMaxPrecision"),
+	GNiagaraGPUSortingUseMaxPrecision,
+	TEXT("Wether sorting using fp32 instead of fp16. (default=1)"),
+	ECVF_Default
+);
+
 int32 GNiagaraGPUSortingCPUToGPUThreshold = -1;
 static FAutoConsoleVariableRef CVarNiagaraGPUSortinCPUToGPUThreshold(
 	TEXT("Niagara.GPUSorting.CPUToGPUThreshold"),
@@ -48,7 +56,7 @@ static FAutoConsoleVariableRef CVarNiagaraGPUSortingFrameCountBeforeBufferShrink
 	ECVF_Default
 );
 
-IMPLEMENT_SHADER_TYPE(, FNiagaraSortKeyGenCS, TEXT("/Engine/Private/NiagaraSortKeyGen.usf"), TEXT("GenerateParticleSortKeys"), SF_Compute);
+IMPLEMENT_GLOBAL_SHADER(FNiagaraSortKeyGenCS, "/Plugin/FX/Niagara/Private/NiagaraSortKeyGen.usf", "GenerateParticleSortKeys", SF_Compute);
 
 void FNiagaraSortKeyGenCS::ModifyCompilationEnvironment(const FGlobalShaderPermutationParameters& Parameters, FShaderCompilerEnvironment& OutEnvironment)
 {
@@ -69,6 +77,7 @@ FNiagaraSortKeyGenCS::FNiagaraSortKeyGenCS(const ShaderMetaType::CompiledShaderI
 	GPUParticleCountBuffer.Bind(Initializer.ParameterMap, TEXT("GPUParticleCountBuffer"));
 	ParticleCountParams.Bind(Initializer.ParameterMap, TEXT("ParticleCountParams"));
 	SortParams.Bind(Initializer.ParameterMap, TEXT("SortParams"));
+	SortKeyParams.Bind(Initializer.ParameterMap, TEXT("SortKeyParams"));
 	CameraPosition.Bind(Initializer.ParameterMap, TEXT("CameraPosition"));
 	CameraDirection.Bind(Initializer.ParameterMap, TEXT("CameraDirection"));
 
@@ -85,6 +94,7 @@ bool FNiagaraSortKeyGenCS::Serialize(FArchive& Ar)
 	Ar << GPUParticleCountBuffer;
 	Ar << ParticleCountParams;
 	Ar << SortParams;
+	Ar << SortKeyParams;
 	Ar << CameraPosition;
 	Ar << CameraDirection;
 	Ar << OutKeys;
@@ -105,7 +115,7 @@ void FNiagaraSortKeyGenCS::SetOutput(FRHICommandList& RHICmdList, FUnorderedAcce
 	}
 }
 
-void FNiagaraSortKeyGenCS::SetParameters(FRHICommandList& RHICmdList, const FNiagaraGPUSortInfo& SortInfo, int32 EmitterIndex, int32 OutputOffset)
+void FNiagaraSortKeyGenCS::SetParameters(FRHICommandList& RHICmdList, const FNiagaraGPUSortInfo& SortInfo, uint32 EmitterKey, int32 OutputOffset, const FUintVector4& SortKeyParamsValue)
 {
 	FComputeShaderRHIParamRef ComputeShaderRHI = GetComputeShader();
 
@@ -118,9 +128,10 @@ void FNiagaraSortKeyGenCS::SetParameters(FRHICommandList& RHICmdList, const FNia
 	RHICmdList.SetShaderParameter(ComputeShaderRHI, ParticleCountParams.GetBufferIndex(), ParticleCountParams.GetBaseIndex(), ParticleCountParams.GetNumBytes(), &ParticleCountParamsValue);
 
 	// (EmitterKey, OutputOffset, SortMode, SortAttributeOffset)
-	const FUintVector4 SortParamsValue((uint32)EmitterIndex << 16, OutputOffset, (uint8)SortInfo.SortMode, SortInfo.SortAttributeOffset);
+	const FUintVector4 SortParamsValue(EmitterKey, OutputOffset, (uint8)SortInfo.SortMode, SortInfo.SortAttributeOffset);
 	RHICmdList.SetShaderParameter(ComputeShaderRHI, SortParams.GetBufferIndex(), SortParams.GetBaseIndex(), SortParams.GetNumBytes(), &SortParamsValue);
 
+	RHICmdList.SetShaderParameter(ComputeShaderRHI, SortKeyParams.GetBufferIndex(), SortKeyParams.GetBaseIndex(), SortKeyParams.GetNumBytes(), &SortKeyParamsValue);
 	RHICmdList.SetShaderParameter(ComputeShaderRHI, CameraPosition.GetBufferIndex(), CameraPosition.GetBaseIndex(), CameraPosition.GetNumBytes(), &SortInfo.ViewOrigin);
 	RHICmdList.SetShaderParameter(ComputeShaderRHI, CameraDirection.GetBufferIndex(), CameraDirection.GetBaseIndex(), CameraDirection.GetNumBytes(), &SortInfo.ViewDirection);
 }
@@ -148,7 +159,7 @@ void FNiagaraSortKeyGenCS::UnbindBuffers(FRHICommandList& RHICmdList)
 
 //////////////////////////////////////////////////////////////////////////////////////////////////////////////////////
 
-IMPLEMENT_SHADER_TYPE(, FNiagaraCopyIntBufferRegionCS, TEXT("/Engine/Private/NiagaraCopyIntBuffer.usf"), TEXT("MainCS"), SF_Compute);
+IMPLEMENT_SHADER_TYPE(, FNiagaraCopyIntBufferRegionCS, TEXT("/Plugin/FX/Niagara/Private/NiagaraCopyIntBuffer.usf"), TEXT("MainCS"), SF_Compute);
 
 FNiagaraCopyIntBufferRegionCS::FNiagaraCopyIntBufferRegionCS(const ShaderMetaType::CompiledShaderInitializerType& Initializer)
 	: FGlobalShader(Initializer)
diff --git a/Engine/Plugins/FX/Niagara/Source/NiagaraVertexFactories/Private/NiagaraSpriteVertexFactory.cpp b/Engine/Plugins/FX/Niagara/Source/NiagaraVertexFactories/Private/NiagaraSpriteVertexFactory.cpp
--- a/Engine/Plugins/FX/Niagara/Source/NiagaraVertexFactories/Private/NiagaraSpriteVertexFactory.cpp
+++ b/Engine/Plugins/FX/Niagara/Source/NiagaraVertexFactories/Private/NiagaraSpriteVertexFactory.cpp
@@ -257,4 +257,4 @@ FVertexFactoryShaderParameters* FNiagaraSpriteVertexFactory::ConstructShaderPara
 	return NULL;
 }
 
-IMPLEMENT_VERTEX_FACTORY_TYPE(FNiagaraSpriteVertexFactory,"/Engine/Private/NiagaraSpriteVertexFactory.ush",true,false,true,false,false);
+IMPLEMENT_VERTEX_FACTORY_TYPE(FNiagaraSpriteVertexFactory,"/Plugin/FX/Niagara/Private/NiagaraSpriteVertexFactory.ush",true,false,true,false,false);
diff --git a/Engine/Plugins/FX/Niagara/Source/NiagaraVertexFactories/Public/NiagaraSortingGPU.h b/Engine/Plugins/FX/Niagara/Source/NiagaraVertexFactories/Public/NiagaraSortingGPU.h
--- a/Engine/Plugins/FX/Niagara/Source/NiagaraVertexFactories/Public/NiagaraSortingGPU.h
+++ b/Engine/Plugins/FX/Niagara/Source/NiagaraVertexFactories/Public/NiagaraSortingGPU.h
@@ -8,10 +8,12 @@ NiagaraSortingGPU.h: Niagara sorting shaders
 
 #include "CoreMinimal.h"
 #include "GlobalShader.h"
+#include "ShaderPermutation.h"
 
 struct FNiagaraGPUSortInfo;
 
 extern NIAGARAVERTEXFACTORIES_API int32 GNiagaraGPUSorting;
+extern NIAGARAVERTEXFACTORIES_API int32 GNiagaraGPUSortingUseMaxPrecision;
 extern NIAGARAVERTEXFACTORIES_API int32 GNiagaraGPUSortingCPUToGPUThreshold;
 extern NIAGARAVERTEXFACTORIES_API float GNiagaraGPUSortingBufferSlack;
 extern NIAGARAVERTEXFACTORIES_API int32 GNiagaraGPUSortingMinBufferSize;
@@ -26,10 +28,13 @@ extern NIAGARAVERTEXFACTORIES_API int32 GNiagaraGPUSortingFrameCountBeforeBuffer
  */
 class NIAGARAVERTEXFACTORIES_API FNiagaraSortKeyGenCS : public FGlobalShader
 {
-	DECLARE_SHADER_TYPE(FNiagaraSortKeyGenCS,Global);
+	DECLARE_GLOBAL_SHADER(FNiagaraSortKeyGenCS);
 
 public:
 
+	class FSortUsingMaxPrecision : SHADER_PERMUTATION_BOOL("SORT_MAX_PRECISION");
+	using FPermutationDomain = TShaderPermutationDomain<FSortUsingMaxPrecision>;
+
 	static bool ShouldCompilePermutation(const FGlobalShaderPermutationParameters& Parameters)
 	{
 		return RHISupportsComputeShaders(Parameters.Platform);
@@ -54,7 +59,7 @@ public:
 	/**
 	 * Set input parameters.
 	 */
-	void SetParameters(FRHICommandList& RHICmdList,	const FNiagaraGPUSortInfo& SortInfo, int32 EmitterIndex, int32 OutputOffset);
+	void SetParameters(FRHICommandList& RHICmdList,	const FNiagaraGPUSortInfo& SortInfo, uint32 EmitterKey, int32 OutputOffset, const FUintVector4& SortKeyParamsValue);
 
 	/**
 	 * Unbinds any buffers that have been bound.
@@ -69,6 +74,7 @@ private:
 	FShaderResourceParameter GPUParticleCountBuffer;
 	FShaderParameter ParticleCountParams;
 	FShaderParameter SortParams;
+	FShaderParameter SortKeyParams;
 	FShaderParameter CameraPosition;
 	FShaderParameter CameraDirection;
 
diff --git a/Engine/Shaders/Private/ParticleSpriteVertexFactory.ush b/Engine/Shaders/Private/ParticleSpriteVertexFactory.ush
--- a/Engine/Shaders/Private/ParticleSpriteVertexFactory.ush
+++ b/Engine/Shaders/Private/ParticleSpriteVertexFactory.ush
@@ -313,66 +313,63 @@ void GetTangents(FVertexFactoryInput Input, float3 WorldPosition, float3 OldWorl
 
 	float3 CameraRight = lerp(ResolvedViewRight, SpriteVF.AxisLockRight.xyz, SpriteVF.AxisLockRight.w);
 	float3 CameraUp = lerp(-ResolvedViewUp, SpriteVF.AxisLockUp.xyz, SpriteVF.AxisLockUp.w);
-	
-	// Determine the vector from the particle to the camera and the particle's movement direction.
-	float3 CameraDirection		= SafeNormalize(ResolvedView.WorldCameraOrigin - WorldPosition.xyz);
-	float3 ParticleDirection	= SafeNormalize(WorldPosition.xyz - OldWorldPosition.xyz);
-
-	// Tangent vectors for PSA_Velocity.
-	float3 Right_Velocity	= SafeNormalize(cross(CameraDirection, ParticleDirection));
-	float3 Up_Velocity		= -ParticleDirection;
 
-	// Tangent vectors for rotation locked about an axis.
-	float3 Right_AxisRot = SpriteVF.AxisLockRight.xyz;
-	float3 Up_AxisRot = -SafeNormalize(cross(Right_AxisRot,CameraDirection));
-
-	// Tangent vectors for camera facing position.
-	float3 Right_FacingCameraPosition = SafeNormalize(cross(CameraDirection,float3(0,0,1)));
-	float3 Up_FacingCameraPosition = cross(CameraDirection, Right_FacingCameraPosition);
+	// Determine the vector from the particle to the camera and the particle's movement direction.
+	float3 CameraDirection	= SafeNormalize(ResolvedView.WorldCameraOrigin - WorldPosition.xyz);
+	float3 RightVector = CameraRight.xyz;
+	float3 UpVector = CameraUp.xyz;
 
-	// Blend between PSA_FacingCamera and PSA_Square over distance
 	float4 LocalTangentSelector = SpriteVF.TangentSelector;
 
 	BRANCH
 	if (SpriteVF.CameraFacingBlend.x > 0.f)
 	{
+		// Blend between PSA_FacingCamera and PSA_Square over distance
 		float3 CameraVector = ResolvedView.WorldCameraOrigin.xyz - WorldPosition.xyz;
 		float CameraDistanceSq = dot(CameraVector, CameraVector);
 		float AlignmentMode = saturate(CameraDistanceSq * SpriteVF.CameraFacingBlend.y - SpriteVF.CameraFacingBlend.z);
 
-		LocalTangentSelector.x = 1.f - AlignmentMode;
-		LocalTangentSelector.w = AlignmentMode;
-	}
+		float3 CameraFacingRight = SafeNormalize(cross(CameraDirection, float3(0,0,1)));
+		float3 CameraFacingUp = cross(CameraDirection, CameraFacingRight);
 
-	// Select the appropriate tangents.
-	// The TangentSelector values are as follows:
-	//   x - multiplied by the standard camera right/up values
-	//   y - multiplied by the velocity right/up values
-	//   z - multiplied by the axis rotation right/up values
-	// The multiplies are summed together to give the final right/up
-	float3 Right_Square	=	(CameraRight.xyz * LocalTangentSelector.x) +
-							(Right_Velocity * LocalTangentSelector.y) +
-							(Right_AxisRot * LocalTangentSelector.z) +
-							(Right_FacingCameraPosition * LocalTangentSelector.w);
-	float3 Up_Square =		(CameraUp.xyz * LocalTangentSelector.x) +
-							(Up_Velocity * LocalTangentSelector.y) +
-							(Up_AxisRot * LocalTangentSelector.z) +
-							(Up_FacingCameraPosition * LocalTangentSelector.w);
-							
-	Right_Square = normalize(Right_Square);
-	Up_Square = normalize(Up_Square);
+		RightVector = normalize(lerp(RightVector, CameraFacingRight, AlignmentMode));
+		UpVector = normalize(lerp(UpVector, CameraFacingUp, AlignmentMode));
+	}
+	else
+	{
+		FLATTEN
+		if (LocalTangentSelector.y > 0)
+		{
+			// Tangent vectors for PSA_Velocity.
+			float3 ParticleDirection = SafeNormalize(WorldPosition.xyz - OldWorldPosition.xyz);
+			RightVector	= SafeNormalize(cross(CameraDirection, ParticleDirection));
+			UpVector = -ParticleDirection;
+		}
+		else if (LocalTangentSelector.z > 0)
+		{
+			// Tangent vectors for rotation locked about an axis.
+			RightVector = SpriteVF.AxisLockRight.xyz;
+			UpVector = -SafeNormalize(cross(RightVector,CameraDirection));
+		}
+		else if (LocalTangentSelector.w > 0)
+		{
+			// Tangent vectors for camera facing position.
+			RightVector = SafeNormalize(cross(CameraDirection,float3(0,0,1)));
+			UpVector = cross(CameraDirection, RightVector);
+		}
+	}
 
 	// Determine the angle of rotation.
-	float SinRotation;
-	float CosRotation;
+	float SinRotation; // = 0
+	float CosRotation; // = 1
 	const float RotationScale = SpriteVF.RotationScale;
 	const float RotationBias = SpriteVF.RotationBias;
 	const float SpriteRotation = Input.SizeRotSubImage.z * RotationScale + RotationBias;
 	sincos(SpriteRotation, SinRotation, CosRotation);
 
 	// Rotate the sprite to determine final tangents.
-	OutRight	= SinRotation * Up_Square + CosRotation * Right_Square;
-	OutUp		= CosRotation * Up_Square - SinRotation * Right_Square;
+	OutRight	= SinRotation * UpVector + CosRotation * RightVector;
+	OutUp		= CosRotation * UpVector - SinRotation * RightVector;
 }
 
 /** derive basis vectors */
diff --git a/Engine/Source/Editor/PropertyEditor/Private/DetailLayoutBuilderImpl.cpp b/Engine/Source/Editor/PropertyEditor/Private/DetailLayoutBuilderImpl.cpp
--- a/Engine/Source/Editor/PropertyEditor/Private/DetailLayoutBuilderImpl.cpp
+++ b/Engine/Source/Editor/PropertyEditor/Private/DetailLayoutBuilderImpl.cpp
@@ -22,6 +22,18 @@ FDetailLayoutBuilderImpl::FDetailLayoutBuilderImpl(TSharedPtr<FComplexPropertyNo
 	bLayoutForExternalRoot = bIsExternal;
 }
 
+
+FDetailLayoutBuilderImpl::~FDetailLayoutBuilderImpl()
+{
+	if (GetDetailsView())
+	{
+		for (TSharedPtr<FComplexPropertyNode> ExternalRootPropertyNode : ExternalRootPropertyNodes)
+		{
+			GetDetailsView()->SaveExpandedItems(ExternalRootPropertyNode.ToSharedRef());
+		}
+	}
+}
+
 IDetailCategoryBuilder& FDetailLayoutBuilderImpl::EditCategory( FName CategoryName, const FText& NewLocalizedDisplayName, ECategoryPriority::Type CategoryType )
 {
 	FText LocalizedDisplayName = NewLocalizedDisplayName;
@@ -563,10 +575,15 @@ void FDetailLayoutBuilderImpl::AddExternalRootPropertyNode(TSharedRef<FComplexPr
 
 void FDetailLayoutBuilderImpl::RemoveExternalRootPropertyNode(TSharedRef<FComplexPropertyNode> InExternalRootNode)
 {
-	ExternalRootPropertyNodes.RemoveAll([InExternalRootNode](TSharedPtr<FComplexPropertyNode> ExternalRootPropertyNode)
+	int32 NumRemoved = ExternalRootPropertyNodes.RemoveAll([InExternalRootNode](TSharedPtr<FComplexPropertyNode> ExternalRootPropertyNode)
 	{
 		return ExternalRootPropertyNode == InExternalRootNode;
 	});
+
+	if (NumRemoved > 0 && GetDetailsView())
+	{
+		GetDetailsView()->SaveExpandedItems(InExternalRootNode);
+	}
 }
 
 FDelegateHandle FDetailLayoutBuilderImpl::AddNodeVisibilityChangedHandler(FSimpleMulticastDelegate::FDelegate InOnNodeVisibilityChanged)
diff --git a/Engine/Source/Editor/PropertyEditor/Private/DetailLayoutBuilderImpl.h b/Engine/Source/Editor/PropertyEditor/Private/DetailLayoutBuilderImpl.h
--- a/Engine/Source/Editor/PropertyEditor/Private/DetailLayoutBuilderImpl.h
+++ b/Engine/Source/Editor/PropertyEditor/Private/DetailLayoutBuilderImpl.h
@@ -25,6 +25,8 @@ public:
 		const TSharedPtr<IDetailsViewPrivate>& InDetailsView,
 		bool bIsExternal);
 
+	~FDetailLayoutBuilderImpl();
+
 	/** IDetailLayoutBuilder Interface */
 	virtual const IDetailsView* GetDetailsView() const override;
 	virtual void GetObjectsBeingCustomized(TArray< TWeakObjectPtr<UObject> >& OutObjects) const override;
diff --git a/Engine/Source/Editor/PropertyEditor/Private/SDetailsViewBase.cpp b/Engine/Source/Editor/PropertyEditor/Private/SDetailsViewBase.cpp
--- a/Engine/Source/Editor/PropertyEditor/Private/SDetailsViewBase.cpp
+++ b/Engine/Source/Editor/PropertyEditor/Private/SDetailsViewBase.cpp
@@ -714,7 +714,8 @@ void SDetailsViewBase::Tick( const FGeometry& AllottedGeometry, const double InC
 
 	FRootPropertyNodeList& RootPropertyNodes = GetRootNodes();
 
-	for(TSharedPtr<FComplexPropertyNode>& RootPropertyNode : RootPropertyNodes)
+	bool bHadDeferredActions = DeferredActions.Num() > 0;
+	auto PreProcessRootNode = [&bHadDeferredActions, this](TSharedPtr<FComplexPropertyNode> RootPropertyNode) 
 	{
 		check(RootPropertyNode.IsValid());
 
@@ -724,12 +725,26 @@ void SDetailsViewBase::Tick( const FGeometry& AllottedGeometry, const double InC
 			ObjectRoot->PurgeKilledObjects();
 		}
 
-		if(DeferredActions.Num() > 0)
+		if(bHadDeferredActions)
 		{		
 			// Any deferred actions are likely to cause the node  tree to be at least partially rebuilt
 			// Save the expansion state of existing nodes so we can expand them later
 			SaveExpandedItems(RootPropertyNode.ToSharedRef());
 		}
+	};
+
+	for (TSharedPtr<FComplexPropertyNode>& RootPropertyNode : RootPropertyNodes)
+	{
+		PreProcessRootNode(RootPropertyNode);
+	}
+
+	for (FDetailLayoutData& LayoutData : DetailLayouts)
+	{
+		FRootPropertyNodeList& ExternalRootPropertyNodes = LayoutData.DetailLayout->GetExternalRootPropertyNodes();
+		for (TSharedPtr<FComplexPropertyNode>& ExternalRootPropertyNode : ExternalRootPropertyNodes)
+		{
+			PreProcessRootNode(ExternalRootPropertyNode);
+		}
 	}
 
 	if (DeferredActions.Num() > 0)
@@ -742,6 +757,25 @@ void SDetailsViewBase::Tick( const FGeometry& AllottedGeometry, const double InC
 		DeferredActions.Empty();
 	}
 
+	if (bHadDeferredActions)
+	{
+		for (TSharedPtr<FComplexPropertyNode>& RootPropertyNode : RootPropertyNodes)
+		{
+			check(RootPropertyNode.IsValid());
+			RestoreExpandedItems(RootPropertyNode.ToSharedRef());
+		}
+
+		for (FDetailLayoutData& LayoutData : DetailLayouts)
+		{
+			FRootPropertyNodeList& ExternalRootPropertyNodes = LayoutData.DetailLayout->GetExternalRootPropertyNodes();
+			for (TSharedPtr<FComplexPropertyNode>& ExternalRootPropertyNode : ExternalRootPropertyNodes)
+			{
+				check(ExternalRootPropertyNode.IsValid());
+				RestoreExpandedItems(ExternalRootPropertyNode.ToSharedRef());
+			}
+		}
+	}
+
 	TSharedPtr<FComplexPropertyNode> LastRootPendingKill;
 	if (RootNodesPendingKill.Num() > 0 )
 	{
@@ -764,12 +798,6 @@ void SDetailsViewBase::Tick( const FGeometry& AllottedGeometry, const double InC
 	{
 		for(TSharedPtr<FComplexPropertyNode>& RootPropertyNode : RootPropertyNodes)
 		{
-			if(RootPropertyNode == LastRootPendingKill)
-			{
-				ExpandedDetailNodes.Empty();
-				RestoreExpandedItems(RootPropertyNode.ToSharedRef());
-			}
-
 			EPropertyDataValidationResult Result = RootPropertyNode->EnsureDataIsValid();
 			if(Result == EPropertyDataValidationResult::PropertiesChanged || Result == EPropertyDataValidationResult::EditInlineNewValueChanged)
 			{
diff --git a/Engine/Source/Runtime/Engine/Classes/Engine/StaticMesh.h b/Engine/Source/Runtime/Engine/Classes/Engine/StaticMesh.h
--- a/Engine/Source/Runtime/Engine/Classes/Engine/StaticMesh.h
+++ b/Engine/Source/Runtime/Engine/Classes/Engine/StaticMesh.h
@@ -708,6 +708,13 @@ public:
 	UPROPERTY(EditAnywhere, AdvancedDisplay, Category = StaticMesh)
 	uint8 bAllowCPUAccess:1;
 
+	/**
+	 * If true, a GPU buffer containing required data for uniform mesh surface sampling will be created at load time.
+	 * It is created from the cpu data so bSupportUniformlyDistributedSampling is also required to be true.
+	 */
+	UPROPERTY(EditAnywhere, AdvancedDisplay, Category = StaticMesh)
+	uint8 bSupportGpuUniformlyDistributedSampling : 1;
+
 	/** A fence which is used to keep track of the rendering thread releasing the static mesh resources. */
 	FRenderCommandFence ReleaseResourcesFence;
 
diff --git a/Engine/Source/Runtime/Engine/Private/BatchedElements.cpp b/Engine/Source/Runtime/Engine/Private/BatchedElements.cpp
--- a/Engine/Source/Runtime/Engine/Private/BatchedElements.cpp
+++ b/Engine/Source/Runtime/Engine/Private/BatchedElements.cpp
@@ -1252,4 +1252,5 @@ void FBatchedElements::Clear()
 	Sprites.Empty();
 	MeshElements.Empty();
 	ThickLines.Empty();
+	MeshVertices.Empty();
 }
diff --git a/Engine/Source/Runtime/Engine/Private/GPUSort.cpp b/Engine/Source/Runtime/Engine/Private/GPUSort.cpp
--- a/Engine/Source/Runtime/Engine/Private/GPUSort.cpp
+++ b/Engine/Source/Runtime/Engine/Private/GPUSort.cpp
@@ -24,6 +24,7 @@ DEFINE_LOG_CATEGORY_STATIC(LogGPUSort, Log, All);
 static TAutoConsoleVariable<int32> CVarDebugOffsets(TEXT("GPUSort.DebugOffsets"),0,TEXT("Debug GPU sort offsets."));
 static TAutoConsoleVariable<int32> CVarDebugSort(TEXT("GPUSort.DebugSort"),0,TEXT("Debug GPU sorting."));
 
+#define GPUSORT_BITCOUNT 32
 #define RADIX_BITS 4
 #define DIGIT_COUNT (1 << RADIX_BITS)
 #define KEYS_PER_LOOP 8
@@ -672,6 +673,29 @@ IMPLEMENT_SHADER_TYPE(,FRadixSortDownsweepCS,TEXT("/Engine/Private/RadixSortShad
 	Public interface.
 ------------------------------------------------------------------------------*/
 
+/**
+ * Get the number of passes we will need to make in order to sort
+ */
+int32 GetGPUSortPassCount(uint32 KeyMask)
+{
+	const int32 BitCount = GPUSORT_BITCOUNT;
+	const int32 PassCount = BitCount / RADIX_BITS;
+
+	int32 PassesRequired = 0;
+
+	uint32 PassBits = DIGIT_COUNT - 1;
+	for (int32 PassIndex = 0; PassIndex < PassCount; ++PassIndex)
+	{
+		// Check to see if these key bits matter.
+		if ((PassBits & KeyMask) != 0)
+		{
+			++PassesRequired;
+		}
+		PassBits <<= RADIX_BITS;
+	}
+	return PassesRequired;
+}
+
 /**
  * Sort a buffer on the GPU.
  * @param SortBuffers - The buffer to sort including required views and a ping-
@@ -716,7 +740,7 @@ int32 SortGPUBuffers(FRHICommandListImmediate& RHICmdList, FGPUSortBuffers SortB
 	const int32 ExtraKeyCount = Count % TILE_SIZE;
 
 	// Determine how many passes are required.
-	const int32 BitCount = 32;
+	const int32 BitCount = GPUSORT_BITCOUNT;
 	const int32 PassCount = BitCount / RADIX_BITS;
 
 	// Setup sort parameters.
diff --git a/Engine/Source/Runtime/Engine/Private/GPUSort.h b/Engine/Source/Runtime/Engine/Private/GPUSort.h
--- a/Engine/Source/Runtime/Engine/Private/GPUSort.h
+++ b/Engine/Source/Runtime/Engine/Private/GPUSort.h
@@ -30,6 +30,11 @@ struct FGPUSortBuffers
 	}
 };
 
+/**
+ * Get the number of passes we will need to make in order to sort
+ */
+int32 GetGPUSortPassCount(uint32 KeyMask);
+
 /**
  * Sort a buffer on the GPU.
  * @param SortBuffers - The buffer to sort including required views and a ping-
diff --git a/Engine/Source/Runtime/Engine/Private/Particles/ParticleGpuSimulation.cpp b/Engine/Source/Runtime/Engine/Private/Particles/ParticleGpuSimulation.cpp
--- a/Engine/Source/Runtime/Engine/Private/Particles/ParticleGpuSimulation.cpp
+++ b/Engine/Source/Runtime/Engine/Private/Particles/ParticleGpuSimulation.cpp
@@ -674,6 +674,15 @@ public:
 /** Global GPU sprite vertex declaration. */
 TGlobalResource<FGPUSpriteVertexDeclaration> GGPUSpriteVertexDeclaration;
 
+/**
+ * Optional user data passed into each Mesh Batch
+ */
+struct FGPUSpriteMeshDataUserData  : public FOneFrameResource
+{
+	int32 SortedOffset = 0;
+	FShaderResourceViewRHIParamRef SortedParticleIndicesSRV = nullptr;
+};
+
 /**
  * Vertex factory for render sprites from GPU simulated particles.
  */
@@ -691,10 +700,8 @@ public:
 	FUniformBufferRHIParamRef EmitterUniformBuffer;
 	/** Emitter uniform buffer for dynamic parameters. */
 	FUniformBufferRHIRef EmitterDynamicUniformBuffer;
-	/** Buffer containing particle indices. */
-	FParticleIndicesVertexBuffer* ParticleIndicesBuffer;
-	/** Offset in to the particle indices buffer. */
-	uint32 ParticleIndicesOffset;
+	/** Buffer containing unsorted particle indices. */
+	FShaderResourceViewRHIParamRef UnsortedParticleIndicesSRV;
 	/** Texture containing positions for all particles. */
 	FTexture2DRHIParamRef PositionTextureRHI;
 	/** Texture containing velocities for all particles. */
@@ -705,8 +712,7 @@ public:
 
 	FGPUSpriteVertexFactory()
 		: FParticleVertexFactoryBase(PVFT_MAX, ERHIFeatureLevel::Num)
-		, ParticleIndicesBuffer(nullptr)
-		, ParticleIndicesOffset(0)
+		, UnsortedParticleIndicesSRV(0)
 		, PositionTextureRHI(nullptr)
 		, VelocityTextureRHI(nullptr)
 		, AttributesTextureRHI(nullptr)
@@ -737,10 +743,9 @@ public:
 	/**
 	 * Set the source vertex buffer that contains particle indices.
 	 */
-	void SetVertexBuffer( FParticleIndicesVertexBuffer* VertexBuffer, uint32 Offset )
+	void SetUnsortedParticleIndicesSRV(FShaderResourceViewRHIParamRef VertexBuffer )
 	{
-		ParticleIndicesBuffer = VertexBuffer;
-		ParticleIndicesOffset = Offset;
+		UnsortedParticleIndicesSRV = VertexBuffer;
 	}
 
 	/**
@@ -801,9 +806,18 @@ void FGPUSpriteVertexFactoryShaderParametersVS::GetElementShaderBindings(
 	FSamplerStateRHIParamRef SamplerStateLinear = TStaticSamplerState<SF_Bilinear>::GetRHI();
 	ShaderBindings.Add(Shader->GetUniformBufferParameter<FGPUSpriteEmitterUniformParameters>(), GPUVF->EmitterUniformBuffer);
 	ShaderBindings.Add(Shader->GetUniformBufferParameter<FGPUSpriteEmitterDynamicUniformParameters>(), GPUVF->EmitterDynamicUniformBuffer);
-	FShaderResourceViewRHIParamRef ParticleIndicesBuffer = GPUVF->ParticleIndicesBuffer->VertexBufferSRV;
-	ShaderBindings.Add(ParticleIndices, ParticleIndicesBuffer ? ParticleIndicesBuffer : (FShaderResourceViewRHIParamRef)GNullColorVertexBuffer.VertexBufferSRV);
-	ShaderBindings.Add(ParticleIndicesOffset, GPUVF->ParticleIndicesOffset);
+
+	const FGPUSpriteMeshDataUserData* UserData = reinterpret_cast<const FGPUSpriteMeshDataUserData*>(BatchElement.UserData);
+	if (UserData != nullptr)
+	{
+		ShaderBindings.Add(ParticleIndices, UserData->SortedParticleIndicesSRV);
+		ShaderBindings.Add(ParticleIndicesOffset, UserData->SortedOffset);
+	}
+	else
+	{
+		ShaderBindings.Add(ParticleIndices, GPUVF->UnsortedParticleIndicesSRV ? GPUVF->UnsortedParticleIndicesSRV : (FShaderResourceViewRHIParamRef)GNullColorVertexBuffer.VertexBufferSRV);
+		ShaderBindings.Add(ParticleIndicesOffset, 0);
+	}
 
 	ShaderBindings.AddTexture(PositionTexture, PositionTextureSampler, SamplerStatePoint, GPUVF->PositionTextureRHI);
 	ShaderBindings.AddTexture(VelocityTexture, VelocityTextureSampler, SamplerStatePoint, GPUVF->VelocityTextureRHI);
@@ -2868,6 +2882,7 @@ static void GetNewParticleArray(TArray<FNewParticle>& NewParticles, int32 NumPar
 		NewParticles.Reserve(NumParticlesNeeded);
 	}
 }
+
 /**
  * Dynamic emitter data for Cascade.
  */
@@ -3063,20 +3078,20 @@ public:
 				// Create per-emitter uniform buffer for dynamic parameters
 				CollectorResources.UniformBuffer = FGPUSpriteEmitterDynamicUniformBufferRef::CreateUniformBufferImmediate(PerViewDynamicParameters, UniformBuffer_SingleFrame);
 
+				FGPUSpriteMeshDataUserData* MeshBatchUserData = nullptr;
 				if (bAllowSorting && SortMode == PSORTMODE_DistanceToView)
 				{
+					MeshBatchUserData = &Collector.AllocateOneFrameResource<FGPUSpriteMeshDataUserData>();
+
 					// Extensibility TODO: This call to AddSortedGPUSimulation is very awkward. When rendering a frame we need to
 					// accumulate all GPU particle emitters that need to be sorted. That is so they can be sorted in one big radix
 					// sort for efficiency. Ideally that state is per-scene renderer but the renderer doesn't know anything about particles.
-					const int32 SortedBufferOffset = FXSystem->AddSortedGPUSimulation(Simulation, View->ViewMatrices.GetViewOrigin());
 					check(SimulationResources->SortedVertexBuffer.IsInitialized());
-					VertexFactory.SetVertexBuffer(&SimulationResources->SortedVertexBuffer, SortedBufferOffset);
-				}
-				else
-				{
-					check(Simulation->VertexBuffer.IsInitialized());
-					VertexFactory.SetVertexBuffer(&Simulation->VertexBuffer, 0);
+					MeshBatchUserData->SortedOffset = FXSystem->AddSortedGPUSimulation(Simulation, View->ViewMatrices.GetViewOrigin());
+					MeshBatchUserData->SortedParticleIndicesSRV = SimulationResources->SortedVertexBuffer.VertexBufferSRV;
 				}
+				check(Simulation->VertexBuffer.IsInitialized());
+				VertexFactory.SetUnsortedParticleIndicesSRV(Simulation->VertexBuffer.VertexBufferSRV);
 
 				const int32 ParticleCount = Simulation->VertexBuffer.ParticleCount;
 				const bool bIsWireframe = ViewFamily.EngineShowFlags.Wireframe;
@@ -3100,6 +3115,7 @@ public:
 					BatchElement.NumInstances = ParticleCount / MAX_PARTICLES_PER_INSTANCE;
 					BatchElement.FirstIndex = 0;
 					BatchElement.bIsInstancedMesh = true;
+					BatchElement.UserData = MeshBatchUserData;
 					Mesh.VertexFactory = &VertexFactory;
 					Mesh.LCI = NULL;
 					if ( bUseLocalSpace )
@@ -4604,20 +4620,13 @@ void FFXSystem::SortGPUParticles(FRHICommandListImmediate& RHICmdList)
 {
 	if (ParticleSimulationResources->SimulationsToSort.Num() > 0)
 	{
-		int32 BufferIndex = SortParticlesGPU(
+		SortParticlesGPU(
 			RHICmdList,
 			ParticleSimulationResources->ParticleSortBuffers,
 			ParticleSimulationResources->GetVisualizeStateTextures().PositionTextureRHI,
 			ParticleSimulationResources->SimulationsToSort,
 			GetFeatureLevel()
 			);
-		ParticleSimulationResources->SortedVertexBuffer.VertexBufferRHI = ParticleSimulationResources->ParticleSortBuffers.GetSortedVertexBufferRHI(BufferIndex);
-		ParticleSimulationResources->SortedVertexBuffer.VertexBufferSRV = ParticleSimulationResources->ParticleSortBuffers.GetSortedVertexBufferSRV(BufferIndex);
-	}
-	else
-	{
-		ParticleSimulationResources->SortedVertexBuffer.VertexBufferRHI = ParticleSimulationResources->ParticleSortBuffers.GetSortedVertexBufferRHI(0);
-		ParticleSimulationResources->SortedVertexBuffer.VertexBufferSRV = ParticleSimulationResources->ParticleSortBuffers.GetSortedVertexBufferSRV(0);
 	}
 }
 
@@ -5071,6 +5080,8 @@ void FFXSystem::SimulateGPUParticles(
 		RHICmdList.TransitionResources(EResourceTransitionAccess::EReadable, VisualizeStateRHIs, 2);		
 	}
 
+	UnbindRenderTargets(RHICmdList);
+
 	// Stats.
 	if (Phase == GetLastParticleSimulationPhase(GetShaderPlatform()))
 	{
diff --git a/Engine/Source/Runtime/Engine/Private/Particles/ParticleSortingGPU.cpp b/Engine/Source/Runtime/Engine/Private/Particles/ParticleSortingGPU.cpp
--- a/Engine/Source/Runtime/Engine/Private/Particles/ParticleSortingGPU.cpp
+++ b/Engine/Source/Runtime/Engine/Private/Particles/ParticleSortingGPU.cpp
@@ -329,9 +329,8 @@ FGPUSortBuffers FParticleSortBuffers::GetSortBuffers()
  * @param ParticleSortBuffers - Buffers to use while sorting GPU particles.
  * @param PositionTextureRHI - Texture containing world space position for all particles.
  * @param SimulationsToSort - A list of simulations that must be sorted.
- * @returns the buffer index in which sorting results are stored.
  */
-int32 SortParticlesGPU(
+void SortParticlesGPU(
 	FRHICommandListImmediate& RHICmdList,
 	FParticleSortBuffers& ParticleSortBuffers,
 	FTexture2DRHIParamRef PositionTextureRHI,
@@ -351,12 +350,17 @@ int32 SortParticlesGPU(
 		}
 	}
 
+	// Buffer index should start so that we always end up on buffer 0 so that callers do not need to dynmically change which buffer they are using
+	const uint32 EmitterKeyMask = (1 << FMath::CeilLogTwo(SimulationsToSort.Num())) - 1;
+	const uint32 KeyMask = (EmitterKeyMask << 16) | 0xFFFF;
+
+	const int32 BufferIndex = GetGPUSortPassCount(KeyMask) & 1;
+
 	// First generate keys for each emitter to be sorted.
-	
 	const int32 TotalParticleCount = GenerateParticleSortKeys(
 		RHICmdList,
-		ParticleSortBuffers.GetKeyBufferUAV(),
-		ParticleSortBuffers.GetVertexBufferUAV(),
+		ParticleSortBuffers.GetKeyBufferUAV(BufferIndex),
+		ParticleSortBuffers.GetVertexBufferUAV(BufferIndex),
 		PositionTextureRHI,
 		SimulationsToSort,
 		FeatureLevel
@@ -367,8 +371,8 @@ int32 SortParticlesGPU(
 	INC_DWORD_STAT_BY( STAT_SortedGPUParticles, TotalParticleCount );
 
 	// Now sort the particles based on the generated keys.
-	const uint32 EmitterKeyMask = (1 << FMath::CeilLogTwo( SimulationsToSort.Num() )) - 1;
-	const uint32 KeyMask = (EmitterKeyMask << 16) | 0xFFFF;
 	FGPUSortBuffers SortBuffers = ParticleSortBuffers.GetSortBuffers();
-	return SortGPUBuffers(RHICmdList, SortBuffers, 0, KeyMask, TotalParticleCount, FeatureLevel);
+
+	const int32 FinalBufferIndex = SortGPUBuffers(RHICmdList, SortBuffers, BufferIndex, KeyMask, TotalParticleCount, FeatureLevel);
+	ensure(FinalBufferIndex == 0);
 }
diff --git a/Engine/Source/Runtime/Engine/Private/Particles/ParticleSortingGPU.h b/Engine/Source/Runtime/Engine/Private/Particles/ParticleSortingGPU.h
--- a/Engine/Source/Runtime/Engine/Private/Particles/ParticleSortingGPU.h
+++ b/Engine/Source/Runtime/Engine/Private/Particles/ParticleSortingGPU.h
@@ -43,18 +43,20 @@ public:
 	/**
 	 * Retrieve the UAV for writing particle sort keys.
 	 */
-	FUnorderedAccessViewRHIParamRef GetKeyBufferUAV()
+	FUnorderedAccessViewRHIParamRef GetKeyBufferUAV(int32 BufferIndex)
 	{
-		return KeyBufferUAVs[0];
+		check((BufferIndex & 0xFFFFFFFE) == 0);
+		return KeyBufferUAVs[BufferIndex];
 	}
 
 	/**
 	 * Retrieve the UAV for writing particle vertices.
 	 * bAsUint : whether to return a G16R16 view or a Uint32 view.
 	 */
-	FORCEINLINE FUnorderedAccessViewRHIParamRef GetVertexBufferUAV()
+	FORCEINLINE FUnorderedAccessViewRHIParamRef GetVertexBufferUAV(int32 BufferIndex)
 	{
-		return VertexBufferUAVs[0];
+		check((BufferIndex & 0xFFFFFFFE) == 0);
+		return VertexBufferUAVs[BufferIndex];
 	}
 
 	/**
@@ -63,7 +65,7 @@ public:
 	FGPUSortBuffers GetSortBuffers();
 
 	/**
-	 * Retrieve the sorted vertex buffer at the given index.
+	 * Retrieve the sorted vertex buffer that results will always be located at.
 	 */
 	FVertexBufferRHIParamRef GetSortedVertexBufferRHI(int32 BufferIndex)
 	{
@@ -72,7 +74,7 @@ public:
 	}
 
 	/**
-	 * Retrieve the SRV for the sorted vertex buffer at the given index.
+	 * Retrieve the SRV that sort results will always be located at.
 	 */
 	FShaderResourceViewRHIParamRef GetSortedVertexBufferSRV(int32 BufferIndex)
 	{
@@ -139,9 +141,8 @@ struct FParticleSimulationSortInfo
  * @param ParticleSortBuffers - Buffers to use while sorting GPU particles.
  * @param PositionTextureRHI - Texture containing world space position for all particles.
  * @param SimulationsToSort - A list of simulations that must be sorted.
- * @returns the buffer index in which sorting results are stored.
  */
-int32 SortParticlesGPU(
+void SortParticlesGPU(
 	FRHICommandListImmediate& RHICmdList,
 	FParticleSortBuffers& ParticleSortBuffers,
 	FTexture2DRHIParamRef PositionTextureRHI,
diff --git a/Engine/Source/Runtime/Engine/Private/PositionVertexBuffer.cpp b/Engine/Source/Runtime/Engine/Private/PositionVertexBuffer.cpp
--- a/Engine/Source/Runtime/Engine/Private/PositionVertexBuffer.cpp
+++ b/Engine/Source/Runtime/Engine/Private/PositionVertexBuffer.cpp
@@ -223,9 +223,20 @@ void FPositionVertexBuffer::InitRHI()
 {
 	VertexBufferRHI = CreateRHIBuffer_RenderThread();
 	// we have decide to create the SRV based on GMaxRHIShaderPlatform because this is created once and shared between feature levels for editor preview.
-	if (VertexBufferRHI && (RHISupportsManualVertexFetch(GMaxRHIShaderPlatform) || IsGPUSkinCacheAvailable()))
+	// Also check to see whether cpu access has been activated on the vertex data
+	if (VertexBufferRHI)
 	{
-		PositionComponentSRV = RHICreateShaderResourceView(VertexData ? VertexBufferRHI : nullptr, 4, PF_R32_FLOAT);
+		// we have decide to create the SRV based on GMaxRHIShaderPlatform because this is created once and shared between feature levels for editor preview.
+		bool bSRV = RHISupportsManualVertexFetch(GMaxRHIShaderPlatform) || IsGPUSkinCacheAvailable();
+
+		// When bAllowCPUAccess is true, the meshes is likely going to be used for Niagara to spawn particles on mesh surface.
+		// And it can be the case for CPU *and* GPU access: no differenciation today. That is why we create a SRV in this case.
+		// This also avoid setting lots of states on all the members of all the different buffers used by meshes. Follow up: https://jira.it.epicgames.net/browse/UE-69376.
+		bSRV |= VertexData->GetAllowCPUAccess();
+		if(bSRV)
+		{
+			PositionComponentSRV = RHICreateShaderResourceView(VertexBufferRHI, 4, PF_R32_FLOAT);
+		}
 	}
 }
 
diff --git a/Engine/Source/Runtime/Engine/Private/RawIndexBuffer.cpp b/Engine/Source/Runtime/Engine/Private/RawIndexBuffer.cpp
--- a/Engine/Source/Runtime/Engine/Private/RawIndexBuffer.cpp
+++ b/Engine/Source/Runtime/Engine/Private/RawIndexBuffer.cpp
@@ -289,16 +289,22 @@ FIndexBufferRHIRef FRawStaticIndexBuffer::CreateRHIBuffer_Internal()
 
 	if (GetNumIndices() > 0)
 	{
+		// When bAllowCPUAccess is true, the meshes is likely going to be used for Niagara to spawn particles on mesh surface.
+		// And it can be the case for CPU *and* GPU access: no differenciation today. That is why we create a SRV in this case.
+		// This also avoid setting lots of states on all the members of all the different buffers used by meshes. Follow up: https://jira.it.epicgames.net/browse/UE-69376.
+		bool bSRV = IndexStorage.GetAllowCPUAccess();
+		uint32 BufferFlags = BUF_Static | (bSRV ? BUF_ShaderResource : BUF_None);
+
 		// Create the index buffer.
 		FRHIResourceCreateInfo CreateInfo(&IndexStorage);
 		CreateInfo.bWithoutNativeResource = !SizeInBytes;
 		if (bRenderThread)
 		{
-			return RHICreateIndexBuffer(IndexStride, SizeInBytes, BUF_Static, CreateInfo);
+			return RHICreateIndexBuffer(IndexStride, SizeInBytes, BufferFlags, CreateInfo);
 		}
 		else
 		{
-			return RHIAsyncCreateIndexBuffer(IndexStride, SizeInBytes, BUF_Static, CreateInfo);
+			return RHIAsyncCreateIndexBuffer(IndexStride, SizeInBytes, BufferFlags, CreateInfo);
 		}
 	}
 	return nullptr;
diff --git a/Engine/Source/Runtime/Engine/Private/StaticMesh.cpp b/Engine/Source/Runtime/Engine/Private/StaticMesh.cpp
--- a/Engine/Source/Runtime/Engine/Private/StaticMesh.cpp
+++ b/Engine/Source/Runtime/Engine/Private/StaticMesh.cpp
@@ -161,6 +161,67 @@ static void FillMaterialName(const TArray<FStaticMaterial>& StaticMaterials, TMa
 #endif
 
 
+/*-----------------------------------------------------------------------------
+	FStaticMeshSectionAreaWeightedTriangleSamplerBuffer
+-----------------------------------------------------------------------------*/
+
+FStaticMeshSectionAreaWeightedTriangleSamplerBuffer::FStaticMeshSectionAreaWeightedTriangleSamplerBuffer()
+{
+}
+
+FStaticMeshSectionAreaWeightedTriangleSamplerBuffer::~FStaticMeshSectionAreaWeightedTriangleSamplerBuffer()
+{
+}
+
+void FStaticMeshSectionAreaWeightedTriangleSamplerBuffer::InitRHI()
+{
+	ReleaseRHI();
+
+	if (Samplers && Samplers->Num() > 0)
+	{
+		FRHIResourceCreateInfo CreateInfo;
+		void* BufferData = nullptr;
+
+		// Count triangle count for all sections and required memory
+		const uint32 AllSectionCount = Samplers->Num();
+		uint32 TriangleCount = 0;
+		for (uint32 i = 0; i < AllSectionCount; ++i)
+		{
+			TriangleCount += (*Samplers)[i].GetNumEntries();
+		}
+		uint32 SizeByte = TriangleCount * sizeof(SectionTriangleInfo);
+
+		BufferSectionTriangleRHI = RHICreateAndLockVertexBuffer(SizeByte, BUF_Static | BUF_ShaderResource, CreateInfo, BufferData);
+
+		// Now compute the alias look up table for unifor; distribution for all section and all triangles
+		SectionTriangleInfo* SectionTriangleInfoBuffer = (SectionTriangleInfo*)BufferData;
+		for (uint32 i = 0; i < AllSectionCount; ++i)
+		{
+			FStaticMeshSectionAreaWeightedTriangleSampler& sampler = (*Samplers)[i];
+			const TArray<float>& ProbTris = sampler.GetProb();
+			const TArray<int32>& AliasTris = sampler.GetAlias();
+			const uint32 NumTriangle = sampler.GetNumEntries();
+
+			for (uint32 t = 0; t < NumTriangle; ++t)
+			{
+				SectionTriangleInfo NewTriangleInfo = { ProbTris[t], (uint32)AliasTris[t], 0, 0 };
+				*SectionTriangleInfoBuffer = NewTriangleInfo;
+				SectionTriangleInfoBuffer++;
+			}
+		}
+		RHIUnlockVertexBuffer(BufferSectionTriangleRHI);
+
+		BufferSectionTriangleSRV = RHICreateShaderResourceView(BufferSectionTriangleRHI, sizeof(SectionTriangleInfo), PF_R32G32B32A32_UINT);
+	}
+}
+
+void FStaticMeshSectionAreaWeightedTriangleSamplerBuffer::ReleaseRHI()
+{
+	BufferSectionTriangleSRV.SafeRelease();
+	BufferSectionTriangleRHI.SafeRelease();
+}
+
+
 /*-----------------------------------------------------------------------------
 	FStaticMeshLODResources
 -----------------------------------------------------------------------------*/
@@ -1078,6 +1139,11 @@ void FStaticMeshLODResources::InitResources(UStaticMesh* Parent)
 		BeginInitResource(&AdditionalIndexBuffers->AdjacencyIndexBuffer);
 	}
 
+	if (Parent->bSupportGpuUniformlyDistributedSampling && Parent->bSupportUniformlyDistributedSampling && Parent->bAllowCPUAccess)
+	{
+		BeginInitResource(&AreaWeightedSectionSamplersBuffer);
+	}
+
 #if RHI_RAYTRACING
 	if (IsRayTracingEnabled())
 	{
@@ -1144,6 +1210,7 @@ void FStaticMeshLODResources::ReleaseResources()
 	BeginReleaseResource(&VertexBuffers.PositionVertexBuffer);
 	BeginReleaseResource(&VertexBuffers.ColorVertexBuffer);
 	BeginReleaseResource(&DepthOnlyIndexBuffer);
+	BeginReleaseResource(&AreaWeightedSectionSamplersBuffer);
 
 	if (AdditionalIndexBuffers)
 	{
@@ -4576,6 +4643,15 @@ void UStaticMesh::PostLoad()
 
 	if (RenderData)
 	{
+		if (bSupportGpuUniformlyDistributedSampling)
+		{
+			// Initialise pointers to samplers
+			for (FStaticMeshLODResources &LOD : RenderData->LODResources)
+			{
+				LOD.AreaWeightedSectionSamplersBuffer.Init(&LOD.AreaWeightedSectionSamplers);
+			}
+		}
+
 		// check the MinLOD values are all within range
 		bool bFixedMinLOD = false;
 		int32 MinAvailableLOD = FMath::Max<int32>(RenderData->LODResources.Num() - 1, 0);
diff --git a/Engine/Source/Runtime/Engine/Public/Rendering/PositionVertexBuffer.h b/Engine/Source/Runtime/Engine/Public/Rendering/PositionVertexBuffer.h
--- a/Engine/Source/Runtime/Engine/Public/Rendering/PositionVertexBuffer.h
+++ b/Engine/Source/Runtime/Engine/Public/Rendering/PositionVertexBuffer.h
@@ -136,6 +136,8 @@ public:
 		return Data;
 	}
 
+	const FShaderResourceViewRHIParamRef GetSRV() const { return PositionComponentSRV; }
+
 private:
 
 	FShaderResourceViewRHIRef PositionComponentSRV;
diff --git a/Engine/Source/Runtime/Engine/Public/Rendering/StaticMeshVertexBuffer.h b/Engine/Source/Runtime/Engine/Public/Rendering/StaticMeshVertexBuffer.h
--- a/Engine/Source/Runtime/Engine/Public/Rendering/StaticMeshVertexBuffer.h
+++ b/Engine/Source/Runtime/Engine/Public/Rendering/StaticMeshVertexBuffer.h
@@ -492,6 +492,9 @@ public:
 	{
 		return IsValidRef(TangentsVertexBuffer.VertexBufferRHI) && IsValidRef(TexCoordVertexBuffer.VertexBufferRHI);
 	}
+
+	const FShaderResourceViewRHIParamRef GetTangentsSRV() const { return TangentsSRV; }
+	const FShaderResourceViewRHIParamRef GetTexCoordsSRV() const { return TextureCoordinatesSRV; }
 private:
 
 	/** The vertex data storage type */
diff --git a/Engine/Source/Runtime/Engine/Public/StaticMeshResources.h b/Engine/Source/Runtime/Engine/Public/StaticMeshResources.h
--- a/Engine/Source/Runtime/Engine/Public/StaticMeshResources.h
+++ b/Engine/Source/Runtime/Engine/Public/StaticMeshResources.h
@@ -250,6 +250,40 @@ protected:
 	FStaticMeshLODResources* Owner;
 };
 
+typedef TArray<FStaticMeshSectionAreaWeightedTriangleSampler> FStaticMeshSectionAreaWeightedTriangleSamplerArray;
+
+/** Represents GPU resource needed for area weighted uniform sampling of a mesh surface. */
+class FStaticMeshSectionAreaWeightedTriangleSamplerBuffer : public FRenderResource
+{
+public:
+
+	ENGINE_API FStaticMeshSectionAreaWeightedTriangleSamplerBuffer();
+	ENGINE_API ~FStaticMeshSectionAreaWeightedTriangleSamplerBuffer();
+
+	ENGINE_API void Init(FStaticMeshSectionAreaWeightedTriangleSamplerArray* SamplerToUpload) { Samplers = SamplerToUpload; }
+
+	// FRenderResource interface.
+	ENGINE_API virtual void InitRHI() override;
+	ENGINE_API virtual void ReleaseRHI() override;
+	virtual FString GetFriendlyName() const override { return TEXT("FStaticMeshSectionAreaWeightedTriangleSamplerBuffer"); }
+
+	ENGINE_API const FShaderResourceViewRHIRef& GetBufferSRV() const { return BufferSectionTriangleSRV; }
+
+private:
+	struct SectionTriangleInfo
+	{
+		float  Prob;
+		uint32 Alias;
+		uint32 pad0;
+		uint32 pad1;
+	};
+
+	FVertexBufferRHIRef BufferSectionTriangleRHI = nullptr;
+	FShaderResourceViewRHIRef BufferSectionTriangleSRV = nullptr;
+
+	FStaticMeshSectionAreaWeightedTriangleSamplerArray* Samplers = nullptr;
+};
+
 
 struct FDynamicMeshVertex;
 struct FModelVertex;
@@ -340,7 +374,9 @@ struct FStaticMeshLODResources
 	/**	Allows uniform random selection of mesh sections based on their area. */
 	FStaticMeshAreaWeightedSectionSampler AreaWeightedSampler;
 	/**	Allows uniform random selection of triangles on each mesh section based on triangle area. */
-	TArray<FStaticMeshSectionAreaWeightedTriangleSampler> AreaWeightedSectionSamplers;
+	FStaticMeshSectionAreaWeightedTriangleSamplerArray AreaWeightedSectionSamplers;
+	/** Allows uniform random selection of triangles on GPU. It is not cooked and serialised but created at runtime from AreaWeightedSectionSamplers when it is available and static mesh bSupportGpuUniformlyDistributedSampling=true*/
+	FStaticMeshSectionAreaWeightedTriangleSamplerBuffer AreaWeightedSectionSamplersBuffer;
 
 	uint32 DepthOnlyNumTriangles;
 
diff --git a/Engine/Source/Runtime/Engine/Public/WeightedRandomSampler.h b/Engine/Source/Runtime/Engine/Public/WeightedRandomSampler.h
--- a/Engine/Source/Runtime/Engine/Public/WeightedRandomSampler.h
+++ b/Engine/Source/Runtime/Engine/Public/WeightedRandomSampler.h
@@ -49,6 +49,10 @@ public:
 	virtual void Serialize(FArchive& Ar);
 
 	int32 GetNumEntries() const { return Prob.Num(); }
+
+	const TArray<float>& GetProb() const { return Prob; }
+	const TArray<int32>& GetAlias() const { return Alias; }
+
 protected:
 	TArray<float> Prob;
 	TArray<int32> Alias;
diff --git a/Engine/Source/Runtime/VectorVM/Private/VectorVM.cpp b/Engine/Source/Runtime/VectorVM/Private/VectorVM.cpp
--- a/Engine/Source/Runtime/VectorVM/Private/VectorVM.cpp
+++ b/Engine/Source/Runtime/VectorVM/Private/VectorVM.cpp
@@ -246,6 +246,7 @@ void FVectorVMContext::PrepareForExec(
 
 	DataSetMetaTable = &InDataSetMetaTable;
 
+	ThreadLocalTempData.Reset(DataSetMetaTable->Num());
 	ThreadLocalTempData.SetNum(DataSetMetaTable->Num());
 }
 
@@ -269,6 +270,7 @@ void FVectorVMContext::FinishExec()
 			MetaTable[DataSetIndex].LockFreeTable();
 			for (int32 IDToFree : Data.IDsToFree)
 			{
+				//UE_LOG(LogVectorVM, Warning, TEXT("AddFreeID: ID:%d | FreeTableIdx:%d."), IDToFree, NumFreeIDs);
 				FreeIDTable[NumFreeIDs++] = IDToFree;
 			}
 			//Unlock the free table.
@@ -1084,6 +1086,7 @@ struct FScalarKernelAcquireID
 			int32 AcquiredID = FreeIDTable[FreeIDTableIndex];
 			checkSlow(AcquiredID != INDEX_NONE);
 
+			//UE_LOG(LogVectorVM, Warning, TEXT("AcquireID: ID:%d | FreeTableIdx:%d."), AcquiredID, FreeIDTableIndex);
 			//Mark this entry in the FreeIDTable as invalid.
 			FreeIDTable[FreeIDTableIndex] = INDEX_NONE;
 
@@ -1127,6 +1130,8 @@ struct FScalarKernelUpdateID
 			{
 				//Add the ID to a thread local list of IDs to free which are actually added to the list safely at the end of this chunk's execution.
 				IDsToFree.Add(InstanceId);
+				IDTable[InstanceId] = INDEX_NONE;
+				//UE_LOG(LogVectorVM, Warning, TEXT("FreeingID: InstanceID:%d."), InstanceId);
 			}
 			else
 			{
diff --git a/Engine/Source/Runtime/VectorVM/Public/VectorVM.h b/Engine/Source/Runtime/VectorVM/Public/VectorVM.h
--- a/Engine/Source/Runtime/VectorVM/Public/VectorVM.h
+++ b/Engine/Source/Runtime/VectorVM/Public/VectorVM.h
@@ -156,9 +156,7 @@ struct FDataSetMeta
 {
 	uint8 **InputRegisters;
 	uint8 NumVariables;
-	uint32 DataSetSizeInBytes;
 	int32 DataSetAccessIndex;	// index for individual elements of this set
-	int32 DataSetOffset;		// offset in the register table
 
 	int32 InstanceOffset;		// offset of the first instance processed 
 	
@@ -182,14 +180,14 @@ struct FDataSetMeta
 	FORCEINLINE void LockFreeTable();
 	FORCEINLINE void UnlockFreeTable();
 
-	FDataSetMeta(uint32 DataSetSize, uint8 **Data, uint8 InNumVariables, int32 InInstanceOffset, TArray<int32>* InIDTable, TArray<int32>* InFreeIDTable, int32* InNumFreeIDs, int32* InMaxUsedID, int32 InIDAcquireTag)
-		: InputRegisters(Data), NumVariables(InNumVariables), DataSetSizeInBytes(DataSetSize), DataSetAccessIndex(INDEX_NONE), DataSetOffset(0), InstanceOffset(InInstanceOffset)
+	FDataSetMeta(uint8 **Data, uint8 InNumVariables, int32 InInstanceOffset, TArray<int32>* InIDTable, TArray<int32>* InFreeIDTable, int32* InNumFreeIDs, int32* InMaxUsedID, int32 InIDAcquireTag)
+		: InputRegisters(Data), NumVariables(InNumVariables), DataSetAccessIndex(INDEX_NONE), InstanceOffset(InInstanceOffset)
 		, IDTable(InIDTable), FreeIDTable(InFreeIDTable), NumFreeIDs(InNumFreeIDs), MaxUsedID(InMaxUsedID), IDAcquireTag(InIDAcquireTag)
 	{
 	}
 
 	FDataSetMeta() 
-		: InputRegisters(nullptr), NumVariables(0), DataSetSizeInBytes(0), DataSetAccessIndex(INDEX_NONE), DataSetOffset(0), InstanceOffset(0)
+		: InputRegisters(nullptr), NumVariables(0), DataSetAccessIndex(INDEX_NONE), InstanceOffset(0)
 		, IDTable(nullptr), FreeIDTable(nullptr), NumFreeIDs(nullptr), MaxUsedID(nullptr), IDAcquireTag(0)
 	{}
 
