- Cache player view locations in NiagaraWorldManager and use the cached locations inside the system instance tick to set the LOD
- This fixes an issue with accessing the camera / player controller during the parallel PreSimulateTick

diff --git a/Engine/Plugins/FX/Niagara/Source/Niagara/Private/NiagaraSystemInstance.cpp b/Engine/Plugins/FX/Niagara/Source/Niagara/Private/NiagaraSystemInstance.cpp
--- a/Engine/Plugins/FX/Niagara/Source/Niagara/Private/NiagaraSystemInstance.cpp
+++ b/Engine/Plugins/FX/Niagara/Source/Niagara/Private/NiagaraSystemInstance.cpp
@@ -10,7 +10,6 @@
 #include "NiagaraWorldManager.h"
 #include "NiagaraComponent.h"
 #include "NiagaraRenderer.h"
-#include "GameFramework/PlayerController.h"
 #include "Templates/AlignmentTemplates.h"
 #include "NiagaraEmitterInstanceBatcher.h"
 
@@ -1128,28 +1127,21 @@ float FNiagaraSystemInstance::GetLODDistance()
 	else
 #endif
 	{
-		UWorld* World = Component->GetWorld();
-		TArray<FVector, TInlineAllocator<8> > PlayerViewLocations;
-		//TODO: Pull this up into the world manager at least. Doing this for every instance each frame is pointless.
-		if (World->GetPlayerControllerIterator())
+		constexpr float DefaultLODDistance = 0.0f;
+
+		FNiagaraWorldManager* WorldManager = GetWorldManager();
+		if ( WorldManager == nullptr )
 		{
-			for (FConstPlayerControllerIterator Iterator = World->GetPlayerControllerIterator(); Iterator; ++Iterator)
-			{
-				APlayerController* PlayerController = Iterator->Get();
-				if (PlayerController && PlayerController->IsLocalPlayerController())
-				{
-					FVector* POVLoc = new(PlayerViewLocations) FVector;
-					FRotator POVRotation;
-					PlayerController->GetPlayerViewPoint(*POVLoc, POVRotation);
-				}
-			}
+			return DefaultLODDistance;
 		}
-		else
+
+		TArrayView<const FVector> PlayerViewLocations = WorldManager->GetCachedPlayerViewLocations();
+		if ( PlayerViewLocations.Num() == 0 )
 		{
-			PlayerViewLocations.Append(World->ViewLocationsRenderedLastFrame);
+			return DefaultLODDistance;
 		}
 
-		float LODDistanceSqr = (PlayerViewLocations.Num() ? FMath::Square(WORLD_MAX) : 0.0f);
+		float LODDistanceSqr = FMath::Square(WORLD_MAX);
 		for (const FVector& ViewLocation : PlayerViewLocations)
 		{
 			const float DistanceToEffectSqr = FVector(ViewLocation - CurrPos).SizeSquared();
diff --git a/Engine/Plugins/FX/Niagara/Source/Niagara/Private/NiagaraWorldManager.cpp b/Engine/Plugins/FX/Niagara/Source/Niagara/Private/NiagaraWorldManager.cpp
--- a/Engine/Plugins/FX/Niagara/Source/Niagara/Private/NiagaraWorldManager.cpp
+++ b/Engine/Plugins/FX/Niagara/Source/Niagara/Private/NiagaraWorldManager.cpp
@@ -11,6 +11,7 @@
 #include "Scalability.h"
 #include "Misc/ConfigCacheIni.h"
 #include "NiagaraDataInterfaceSkeletalMesh.h"
+#include "GameFramework/PlayerController.h"
 #include "EngineModule.h"
 #include "NiagaraStats.h"
 
@@ -200,6 +201,26 @@ void FNiagaraWorldManager::Tick(float DeltaSeconds)
 
 	SkeletalMeshGeneratedData.TickGeneratedData(DeltaSeconds);
 
+	// Cache player view locations for all system instances to access
+	CachedPlayerViewLocations.Reset();
+	if (World->GetPlayerControllerIterator())
+	{
+		for ( FConstPlayerControllerIterator Iterator=World->GetPlayerControllerIterator(); Iterator; ++Iterator)
+		{
+			APlayerController* PlayerController = Iterator->Get();
+			if (PlayerController && PlayerController->IsLocalPlayerController())
+			{
+				FVector* POVLoc = new(CachedPlayerViewLocations) FVector;
+				FRotator POVRotation;
+				PlayerController->GetPlayerViewPoint(*POVLoc, POVRotation);
+			}
+		}
+	}
+	else
+	{
+		CachedPlayerViewLocations.Append(World->ViewLocationsRenderedLastFrame);
+	}
+
 	//Tick our collections to push any changes to bound stores.
 	for (TPair<UNiagaraParameterCollection*, UNiagaraParameterCollectionInstance*> CollectionInstPair : ParameterCollections)
 	{
diff --git a/Engine/Plugins/FX/Niagara/Source/Niagara/Public/NiagaraWorldManager.h b/Engine/Plugins/FX/Niagara/Source/Niagara/Public/NiagaraWorldManager.h
--- a/Engine/Plugins/FX/Niagara/Source/Niagara/Public/NiagaraWorldManager.h
+++ b/Engine/Plugins/FX/Niagara/Source/Niagara/Public/NiagaraWorldManager.h
@@ -82,6 +82,9 @@ public:
 	void OnWorldCleanup(bool bSessionEnded, bool bCleanupResources);
 	
 	FORCEINLINE FNDI_SkeletalMesh_GeneratedData& GetSkeletalMeshGeneratedData() { return SkeletalMeshGeneratedData; }
+
+	TArrayView<const FVector> GetCachedPlayerViewLocations() const { return MakeArrayView(CachedPlayerViewLocations); }
+
 private:
 	UWorld* World;
 
@@ -91,6 +94,8 @@ private:
 
 	int32 CachedEffectsQuality;
 
+	TArray<FVector, TInlineAllocator<8> > CachedPlayerViewLocations;
+
 	/** Generated data used by data interfaces*/
 	FNDI_SkeletalMesh_GeneratedData SkeletalMeshGeneratedData;
 
