Fixed possible nan in AcosFast().
Reduced max niagara ribbon tessellation factor to 16
Reduced max niagara ribbon tessellation averaging to .5

diff --git a/Engine/Plugins/FX/Niagara/Source/Niagara/Private/NiagaraRendererRibbons.cpp b/Engine/Plugins/FX/Niagara/Source/Niagara/Private/NiagaraRendererRibbons.cpp
--- a/Engine/Plugins/FX/Niagara/Source/Niagara/Private/NiagaraRendererRibbons.cpp
+++ b/Engine/Plugins/FX/Niagara/Source/Niagara/Private/NiagaraRendererRibbons.cpp
@@ -28,13 +28,13 @@ static FAutoConsoleVariableRef CVarNiagaraRibbonTessellationAngle(
 	ECVF_Scalability
 );
 
-int32 GNiagaraRibbonMaxTessellation = 64;
+int32 GNiagaraRibbonMaxTessellation = 16;
 static FAutoConsoleVariableRef CVarNiagaraRibbonMaxTessellation(
 	TEXT("Niagara.Ribbon.Tessellation.MaxInterp"),
 	GNiagaraRibbonMaxTessellation,
 	TEXT("When TessellationAngle is > 0, this is the maximum tesselation factor. \n")
 	TEXT("Higher values allow more evenly divided tesselation. \n")
-	TEXT("When TessellationAngle is 0, this is the actually tesselation factor (default=60)."),
+	TEXT("When TessellationAngle is 0, this is the actually tesselation factor (default=16)."),
 	ECVF_Scalability
 );
 
@@ -87,7 +87,7 @@ FORCEINLINE float AcosFast(float InX)
 {
     float X = FMath::Abs(InX);
     float Res = -0.156583f * X + (0.5 * PI);
-    Res *= sqrt(1.0f - X);
+    Res *= sqrt(FMath::Max(0.f, 1.0f - X));
     return (InX >= 0) ? Res : PI - Res;
 }
 
