Fixing header merge error

#tests compiles

diff --git a/Engine/Plugins/FX/Niagara/Source/Niagara/Classes/NiagaraDataSet.h b/Engine/Plugins/FX/Niagara/Source/Niagara/Classes/NiagaraDataSet.h
--- a/Engine/Plugins/FX/Niagara/Source/Niagara/Classes/NiagaraDataSet.h
+++ b/Engine/Plugins/FX/Niagara/Source/Niagara/Classes/NiagaraDataSet.h
@@ -37,12 +37,7 @@ public:
 	{}
 
 	/** The owner of this object is now done with it but it may still be in use by others, possibly on other threads. Add to the deletion queue so it can be safely freed when it's no longer in use. */
-	FORCEINLINE void Destroy()
-	{
-		FScopeLock Lock(&CritSec);
-		check(DeferredDeletionList.Contains(this) == false);
-		DeferredDeletionList.Add(this);
-	}
+	void Destroy();
 	static void FlushDeletionList();
 
 	FORCEINLINE bool IsInUse()const { return ReadRefCount.Load() != 0; }
