- Fix for vertex interpolators not working with Niagara renderers
- Fix for Particle.LocalToWorld being transposed on Niagara Mesh renderer

diff --git a/Engine/Plugins/FX/Niagara/Shaders/Private/NiagaraMeshVertexFactory.ush b/Engine/Plugins/FX/Niagara/Shaders/Private/NiagaraMeshVertexFactory.ush
--- a/Engine/Plugins/FX/Niagara/Shaders/Private/NiagaraMeshVertexFactory.ush
+++ b/Engine/Plugins/FX/Niagara/Shaders/Private/NiagaraMeshVertexFactory.ush
@@ -281,10 +281,10 @@ FMaterialPixelParameters GetMaterialPixelParameters(FVertexFactoryInterpolantsVS
 #endif
 
 #if USE_PARTICLE_TRANSFORM
-	Result.Particle.LocalToWorld[0] = Interpolants.LocalToWorld[0];
-	Result.Particle.LocalToWorld[1] = Interpolants.LocalToWorld[1];
-	Result.Particle.LocalToWorld[2] = Interpolants.LocalToWorld[2];
-	Result.Particle.LocalToWorld[3] = Interpolants.LocalToWorld[3];
+	Result.Particle.LocalToWorld[0] = float4(Interpolants.LocalToWorld[0].x, Interpolants.LocalToWorld[1].x, Interpolants.LocalToWorld[2].x, Interpolants.LocalToWorld[3].x);
+	Result.Particle.LocalToWorld[1] = float4(Interpolants.LocalToWorld[0].y, Interpolants.LocalToWorld[1].y, Interpolants.LocalToWorld[2].y, Interpolants.LocalToWorld[3].y);
+	Result.Particle.LocalToWorld[2] = float4(Interpolants.LocalToWorld[0].z, Interpolants.LocalToWorld[1].z, Interpolants.LocalToWorld[2].z, Interpolants.LocalToWorld[3].z);
+	Result.Particle.LocalToWorld[3] = float4(Interpolants.LocalToWorld[0].w, Interpolants.LocalToWorld[1].w, Interpolants.LocalToWorld[2].w, Interpolants.LocalToWorld[3].w);
 #endif
 
 	Result.Particle.MotionBlurFade = 1.0f;
@@ -347,11 +347,11 @@ FMaterialVertexParameters GetMaterialVertexParameters(FVertexFactoryInput Input,
 	float4x4 Transform = GetParticleTransform(Intermediates);
 	Result.InstanceLocalToWorld = Transform;
 
-#if NEEDS_PARTICLE_TRANSFORM
-	Result.Particle.LocalToWorld[0] = Transform[0];
-	Result.Particle.LocalToWorld[1] = Transform[1];
-	Result.Particle.LocalToWorld[2] = Transform[2];
-	Result.Particle.LocalToWorld[3] = Transform[3];
+#if USE_PARTICLE_TRANSFORM
+	Result.Particle.LocalToWorld[0] = float4(Transform[0].x, Transform[1].x, Transform[2].x, Transform[3].x);
+	Result.Particle.LocalToWorld[1] = float4(Transform[0].y, Transform[1].y, Transform[2].y, Transform[3].y);
+	Result.Particle.LocalToWorld[2] = float4(Transform[0].z, Transform[1].z, Transform[2].z, Transform[3].z);
+	Result.Particle.LocalToWorld[3] = float4(Transform[0].w, Transform[1].w, Transform[2].w, Transform[3].w);
 #endif
 
 	Result.TangentToWorld = mul(TangentToLocal, (float3x3)Transform );
@@ -604,8 +604,6 @@ FVertexFactoryIntermediates GetVertexFactoryIntermediates(FVertexFactoryInput In
 
 	Intermediates.ParticleTranslatedWorldPosition = Intermediates.ParticlePosition.xyz + ResolvedView.PreViewTranslation.xyz;
 
-	//LocalToWorldMat = mul(ScaleMatrix, LocalToWorldMat);
-
 	if (MeshFacingMode == MESH_FACING_VELOCITY)
 	{
 		float4 ToAlignmentQuat = FindQuatBetween(float3(1.0f, 0.0f, 0.0f), normalize(ParticleVelocity.xyz));
@@ -740,6 +738,7 @@ FVertexFactoryInterpolantsVSToPS VertexFactoryGetInterpolantsVSToPS(FVertexFacto
 
 		float2 CustomizedUVs[ NUM_TEX_COORD_INTERPOLATORS];
 		GetMaterialCustomizedUVs(VertexParameters, CustomizedUVs);
+		GetCustomInterpolators(VertexParameters, CustomizedUVs);
 
 		UNROLL
 		for(int CoordinateIndex = 0;CoordinateIndex <  NUM_TEX_COORD_INTERPOLATORS;CoordinateIndex += 2)
diff --git a/Engine/Plugins/FX/Niagara/Shaders/Private/NiagaraRibbonVertexFactory.ush b/Engine/Plugins/FX/Niagara/Shaders/Private/NiagaraRibbonVertexFactory.ush
--- a/Engine/Plugins/FX/Niagara/Shaders/Private/NiagaraRibbonVertexFactory.ush
+++ b/Engine/Plugins/FX/Niagara/Shaders/Private/NiagaraRibbonVertexFactory.ush
@@ -297,8 +297,8 @@ struct FVertexFactoryInterpolantsVSToPS
 	float4 DynamicParameter3	: PARTICLE_DYNAMIC_PARAM3;
 #endif
 
-#if NUM_MATERIAL_TEXCOORDS
-	float4	TexCoords[(NUM_MATERIAL_TEXCOORDS + 1) / 2]	: TEXCOORD3;
+#if NUM_TEX_COORD_INTERPOLATORS
+	float4	TexCoords[(NUM_TEX_COORD_INTERPOLATORS + 1) / 2]	: TEXCOORD3;
 #endif
 
 #if NEEDS_PARTICLE_COLOR
@@ -331,7 +331,7 @@ struct FVertexFactoryInterpolantsVSToPS
 #endif
 };
 
-#if NUM_MATERIAL_TEXCOORDS
+#if NUM_TEX_COORD_INTERPOLATORS
 bool UVIndexUseZW(int UVIndex)
 {
 	#if COMPILER_GLSL_ES2
@@ -427,9 +427,9 @@ FMaterialPixelParameters GetMaterialPixelParameters(FVertexFactoryInterpolantsVS
 	// GetMaterialPixelParameters is responsible for fully initializing the result
 	FMaterialPixelParameters Result = MakeInitializedMaterialPixelParameters();
 
-#if NUM_MATERIAL_TEXCOORDS
+#if NUM_TEX_COORD_INTERPOLATORS
 	UNROLL
-	for (int CoordinateIndex = 0; CoordinateIndex < NUM_MATERIAL_TEXCOORDS; CoordinateIndex++)
+	for (int CoordinateIndex = 0; CoordinateIndex < NUM_TEX_COORD_INTERPOLATORS; CoordinateIndex++)
 	{
 		Result.TexCoords[CoordinateIndex] = GetUV(Interpolants, CoordinateIndex);
 	}
@@ -663,11 +663,16 @@ FVertexFactoryInterpolantsVSToPS VertexFactoryGetInterpolantsVSToPS(FVertexFacto
 	FVertexFactoryInterpolantsVSToPS Interpolants;
 	Interpolants = (FVertexFactoryInterpolantsVSToPS)0;
 
-#if NUM_MATERIAL_TEXCOORDS
-	float2 UV1 = Intermediates.TexCoord.xy;
-	float2 UV2 = Intermediates.TexCoord.zw;
-	SetUV(Interpolants, 0, UV1);
-	SetUV(Interpolants, 1, UV2);
+#if NUM_TEX_COORD_INTERPOLATORS
+	float2 CustomizedUVs[NUM_TEX_COORD_INTERPOLATORS];
+	GetMaterialCustomizedUVs(VertexParameters, CustomizedUVs);
+	GetCustomInterpolators(VertexParameters, CustomizedUVs);
+
+	UNROLL
+	for (int CoordinateIndex = 0; CoordinateIndex < NUM_TEX_COORD_INTERPOLATORS; CoordinateIndex++)
+	{
+		SetUV(Interpolants, CoordinateIndex, CustomizedUVs[CoordinateIndex]);
+	}
 #endif
 
 #if NEEDS_PARTICLE_COLOR
diff --git a/Engine/Plugins/FX/Niagara/Shaders/Private/NiagaraSpriteVertexFactory.ush b/Engine/Plugins/FX/Niagara/Shaders/Private/NiagaraSpriteVertexFactory.ush
--- a/Engine/Plugins/FX/Niagara/Shaders/Private/NiagaraSpriteVertexFactory.ush
+++ b/Engine/Plugins/FX/Niagara/Shaders/Private/NiagaraSpriteVertexFactory.ush
@@ -924,9 +924,9 @@ FVertexFactoryInterpolantsVSToPS VertexFactoryGetInterpolantsVSToPS(FVertexFacto
 	Interpolants = (FVertexFactoryInterpolantsVSToPS)0;
 
 #if NUM_TEX_COORD_INTERPOLATORS
-
 	float2 CustomizedUVs[NUM_TEX_COORD_INTERPOLATORS];
 	GetMaterialCustomizedUVs(VertexParameters, CustomizedUVs);
+	GetCustomInterpolators(VertexParameters, CustomizedUVs);
 
 	UNROLL
 	for (int CoordinateIndex = 0; CoordinateIndex < NUM_TEX_COORD_INTERPOLATORS; CoordinateIndex++)
