Fixed duplicate details panel LOCTEXT keys in NiagaraScriptToolkip.cpp, resulting in both the Script and and Selected Details panels having the same name in the Windows menu.

diff --git a/Engine/Plugins/FX/Niagara/Source/NiagaraEditor/Private/Toolkits/NiagaraScriptToolkit.cpp b/Engine/Plugins/FX/Niagara/Source/NiagaraEditor/Private/Toolkits/NiagaraScriptToolkit.cpp
--- a/Engine/Plugins/FX/Niagara/Source/NiagaraEditor/Private/Toolkits/NiagaraScriptToolkit.cpp
+++ b/Engine/Plugins/FX/Niagara/Source/NiagaraEditor/Private/Toolkits/NiagaraScriptToolkit.cpp
@@ -88,7 +88,7 @@ void FNiagaraScriptToolkit::RegisterTabSpawners(const TSharedRef<class FTabManag
 		.SetIcon(FSlateIcon(FEditorStyle::GetStyleSetName(), "LevelEditor.Tabs.Details"));
 
 	InTabManager->RegisterTabSpawner(SelectedDetailsTabId, FOnSpawnTab::CreateSP(this, &FNiagaraScriptToolkit::SpawnTabSelectedDetails))
-		.SetDisplayName(LOCTEXT("ScriptDetailsTab", "Selected Details"))
+		.SetDisplayName(LOCTEXT("SelectedDetailsTab", "Selected Details"))
 		.SetGroup(WorkspaceMenuCategoryRef)
 		.SetIcon(FSlateIcon(FEditorStyle::GetStyleSetName(), "LevelEditor.Tabs.Details"));
 
