Adding bone rotations to skel mesh NDI.
Created a new tesst for skel mesh sampling. TODO: Add surface sampling tests to map and gpu emitters.
Reimplemented interpolation in Bone samlping modules.
Slight refactor of module.

diff --git a/Engine/Plugins/FX/Niagara/Source/Niagara/Classes/NDISkeletalMeshCommon.h b/Engine/Plugins/FX/Niagara/Source/Niagara/Classes/NDISkeletalMeshCommon.h
--- a/Engine/Plugins/FX/Niagara/Source/Niagara/Classes/NDISkeletalMeshCommon.h
+++ b/Engine/Plugins/FX/Niagara/Source/Niagara/Classes/NDISkeletalMeshCommon.h
@@ -101,6 +101,16 @@ struct FSkinnedPositionAccessorHelper
 	}
 
 	FORCEINLINE FVector GetSkinnedBonePreviousPosition(FSkeletalMeshAccessorHelper& Accessor, int32 BoneIndex)
+	{
+		checkf(false, TEXT("Must provide a specialization for this template type"));
+	}	
+	
+	FORCEINLINE_DEBUGGABLE FQuat GetSkinnedBoneRotation(FSkeletalMeshAccessorHelper& Accessor, int32 BoneIndex)
+	{
+		checkf(false, TEXT("Must provide a specialization for this template type"));
+	}
+
+	FORCEINLINE_DEBUGGABLE FQuat GetSkinnedBonePreviousRotation(FSkeletalMeshAccessorHelper& Accessor, int32 BoneIndex)
 	{
 		checkf(false, TEXT("Must provide a specialization for this template type"));
 	}
@@ -151,6 +161,16 @@ struct FSkinnedPositionAccessorHelper<TIntegralConstant<ENDISkeletalMesh_Skinnin
 	{
 		return Accessor.Mesh->GetComposedRefPoseMatrix(BoneIndex).GetOrigin();
 	}
+
+	FORCEINLINE_DEBUGGABLE FQuat GetSkinnedBoneRotation(FSkeletalMeshAccessorHelper& Accessor, int32 BoneIndex)
+	{
+		return Accessor.Mesh->GetComposedRefPoseMatrix(BoneIndex).ToQuat();
+	}
+
+	FORCEINLINE_DEBUGGABLE FQuat GetSkinnedBonePreviousRotation(FSkeletalMeshAccessorHelper& Accessor, int32 BoneIndex)
+	{
+		return Accessor.Mesh->GetComposedRefPoseMatrix(BoneIndex).ToQuat();
+	}
 };
 
 template<>
@@ -198,6 +218,16 @@ struct FSkinnedPositionAccessorHelper<TIntegralConstant<ENDISkeletalMesh_Skinnin
 	{
 		return (*Accessor.PrevBoneComponentSpaceTransforms)[BoneIndex].GetLocation();
 	}
+
+	FORCEINLINE_DEBUGGABLE FQuat GetSkinnedBoneRotation(FSkeletalMeshAccessorHelper& Accessor, int32 BoneIndex)
+	{
+		return (*Accessor.BoneComponentSpaceTransforms)[BoneIndex].GetRotation();
+	}
+
+	FORCEINLINE_DEBUGGABLE FQuat GetSkinnedBonePreviousRotation(FSkeletalMeshAccessorHelper& Accessor, int32 BoneIndex)
+	{
+		return (*Accessor.PrevBoneComponentSpaceTransforms)[BoneIndex].GetRotation();
+	}
 };
 
 template<>
@@ -245,6 +275,16 @@ struct FSkinnedPositionAccessorHelper<TIntegralConstant<ENDISkeletalMesh_Skinnin
 	{
 		return (*Accessor.PrevBoneComponentSpaceTransforms)[BoneIndex].GetLocation();
 	}
+
+	FORCEINLINE_DEBUGGABLE FQuat GetSkinnedBoneRotation(FSkeletalMeshAccessorHelper& Accessor, int32 BoneIndex)
+	{
+		return (*Accessor.BoneComponentSpaceTransforms)[BoneIndex].GetRotation();
+	}
+
+	FORCEINLINE_DEBUGGABLE FQuat GetSkinnedBonePreviousRotation(FSkeletalMeshAccessorHelper& Accessor, int32 BoneIndex)
+	{
+		return (*Accessor.PrevBoneComponentSpaceTransforms)[BoneIndex].GetRotation();
+	}
 };
 
 //////////////////////////////////////////////////////////////////////////
diff --git a/Engine/Plugins/FX/Niagara/Source/Niagara/Private/NDISkeletalMesh_BoneSampling.cpp b/Engine/Plugins/FX/Niagara/Source/Niagara/Private/NDISkeletalMesh_BoneSampling.cpp
--- a/Engine/Plugins/FX/Niagara/Source/Niagara/Private/NDISkeletalMesh_BoneSampling.cpp
+++ b/Engine/Plugins/FX/Niagara/Source/Niagara/Private/NDISkeletalMesh_BoneSampling.cpp
@@ -70,6 +70,7 @@ void UNiagaraDataInterfaceSkeletalMesh::GetSkeletonSamplingFunctions(TArray<FNia
 		Sig.Inputs.Add(FNiagaraVariable(FNiagaraTypeDefinition(GetClass()), TEXT("SkeletalMesh")));
 		Sig.Inputs.Add(FNiagaraVariable(FNiagaraTypeDefinition::GetIntDef(), TEXT("Bone")));
 		Sig.Outputs.Add(FNiagaraVariable(FNiagaraTypeDefinition::GetVec3Def(), TEXT("Position")));
+		Sig.Outputs.Add(FNiagaraVariable(FNiagaraTypeDefinition::GetQuatDef(), TEXT("Rotation")));
 		Sig.Outputs.Add(FNiagaraVariable(FNiagaraTypeDefinition::GetVec3Def(), TEXT("Velocity")));
 		Sig.bMemberFunction = true;
 		Sig.bRequiresContext = false;
@@ -85,6 +86,7 @@ void UNiagaraDataInterfaceSkeletalMesh::GetSkeletonSamplingFunctions(TArray<FNia
 		Sig.Inputs.Add(FNiagaraVariable(FNiagaraTypeDefinition(GetClass()), TEXT("SkeletalMesh")));
 		Sig.Inputs.Add(FNiagaraVariable(FNiagaraTypeDefinition::GetIntDef(), TEXT("Bone")));
 		Sig.Outputs.Add(FNiagaraVariable(FNiagaraTypeDefinition::GetVec3Def(), TEXT("Position")));
+		Sig.Outputs.Add(FNiagaraVariable(FNiagaraTypeDefinition::GetQuatDef(), TEXT("Rotation")));
 		Sig.Outputs.Add(FNiagaraVariable(FNiagaraTypeDefinition::GetVec3Def(), TEXT("Velocity")));
 		Sig.bMemberFunction = true;
 		Sig.bRequiresContext = false;
@@ -102,6 +104,7 @@ void UNiagaraDataInterfaceSkeletalMesh::GetSkeletonSamplingFunctions(TArray<FNia
 		Sig.Inputs.Add(FNiagaraVariable(FNiagaraTypeDefinition::GetIntDef(), TEXT("Bone")));
 		Sig.Inputs.Add(FNiagaraVariable(FNiagaraTypeDefinition::GetFloatDef(), TEXT("Interpolation")));
 		Sig.Outputs.Add(FNiagaraVariable(FNiagaraTypeDefinition::GetVec3Def(), TEXT("Position")));
+		Sig.Outputs.Add(FNiagaraVariable(FNiagaraTypeDefinition::GetQuatDef(), TEXT("Rotation")));
 		Sig.Outputs.Add(FNiagaraVariable(FNiagaraTypeDefinition::GetVec3Def(), TEXT("Velocity")));
 		Sig.bMemberFunction = true;
 		Sig.bRequiresContext = false;
@@ -118,6 +121,7 @@ void UNiagaraDataInterfaceSkeletalMesh::GetSkeletonSamplingFunctions(TArray<FNia
 		Sig.Inputs.Add(FNiagaraVariable(FNiagaraTypeDefinition::GetIntDef(), TEXT("Bone")));
 		Sig.Inputs.Add(FNiagaraVariable(FNiagaraTypeDefinition::GetFloatDef(), TEXT("Interpolation")));
 		Sig.Outputs.Add(FNiagaraVariable(FNiagaraTypeDefinition::GetVec3Def(), TEXT("Position")));
+		Sig.Outputs.Add(FNiagaraVariable(FNiagaraTypeDefinition::GetQuatDef(), TEXT("Rotation")));
 		Sig.Outputs.Add(FNiagaraVariable(FNiagaraTypeDefinition::GetVec3Def(), TEXT("Velocity")));
 		Sig.bMemberFunction = true;
 		Sig.bRequiresContext = false;
@@ -214,22 +218,22 @@ void UNiagaraDataInterfaceSkeletalMesh::BindSkeletonSamplingFunction(const FVMEx
 	}
 	else if (BindingInfo.Name == GetSkinnedBoneDataName)
 	{
-		check(BindingInfo.GetNumInputs() == 2 && BindingInfo.GetNumOutputs() == 6);
+		ensure(BindingInfo.GetNumInputs() == 2 && BindingInfo.GetNumOutputs() == 10);
 		TSkinningModeBinder<TNDIExplicitBinder<FNDITransformHandlerNoop, TNDIExplicitBinder<TIntegralConstant<bool, false>, NDI_FUNC_BINDER(UNiagaraDataInterfaceSkeletalMesh, GetSkinnedBoneData)>>>::Bind(this, BindingInfo, InstanceData, OutFunc);
 	}
 	else if (BindingInfo.Name == GetSkinnedBoneDataWSName)
 	{
-		check(BindingInfo.GetNumInputs() == 2 && BindingInfo.GetNumOutputs() == 6);
+		ensure(BindingInfo.GetNumInputs() == 2 && BindingInfo.GetNumOutputs() == 10);
 		TSkinningModeBinder<TNDIExplicitBinder<FNDITransformHandler, TNDIExplicitBinder<TIntegralConstant<bool, false>, NDI_FUNC_BINDER(UNiagaraDataInterfaceSkeletalMesh, GetSkinnedBoneData)>>>::Bind(this, BindingInfo, InstanceData, OutFunc);
 	}
 	else if (BindingInfo.Name == GetSkinnedBoneDataInterpolatedName)
 	{
-		check(BindingInfo.GetNumInputs() == 3 && BindingInfo.GetNumOutputs() == 6);
+		ensure(BindingInfo.GetNumInputs() == 3 && BindingInfo.GetNumOutputs() == 10);
 		TSkinningModeBinder<TNDIExplicitBinder<FNDITransformHandlerNoop, TNDIExplicitBinder<TIntegralConstant<bool, true>, NDI_FUNC_BINDER(UNiagaraDataInterfaceSkeletalMesh, GetSkinnedBoneData)>>>::Bind(this, BindingInfo, InstanceData, OutFunc);
 	}
 	else if (BindingInfo.Name == GetSkinnedBoneDataWSInterpolatedName)
 	{
-		check(BindingInfo.GetNumInputs() == 3 && BindingInfo.GetNumOutputs() == 6);
+		ensure(BindingInfo.GetNumInputs() == 3 && BindingInfo.GetNumOutputs() == 10);
 		TSkinningModeBinder<TNDIExplicitBinder<FNDITransformHandler, TNDIExplicitBinder<TIntegralConstant<bool, true>, NDI_FUNC_BINDER(UNiagaraDataInterfaceSkeletalMesh, GetSkinnedBoneData)>>>::Bind(this, BindingInfo, InstanceData, OutFunc);
 	}
 	else if (BindingInfo.Name == GetSpecificBoneCountName)
@@ -360,18 +364,22 @@ struct FBoneSocketSkinnedDataOutputHandler
 {
 	FBoneSocketSkinnedDataOutputHandler(FVectorVMContext& Context)
 		: PosX(Context), PosY(Context), PosZ(Context)
+		, RotX(Context), RotY(Context), RotZ(Context), RotW(Context)
 		, VelX(Context), VelY(Context), VelZ(Context)
 		, bNeedsPosition(PosX.IsValid() || PosY.IsValid() || PosZ.IsValid())
+		, bNeedsRotation(RotX.IsValid() || RotY.IsValid() || RotZ.IsValid() || RotW.IsValid())
 		, bNeedsVelocity(VelX.IsValid() || VelY.IsValid() || VelZ.IsValid())
 	{
 	}
 
 	VectorVM::FExternalFuncRegisterHandler<float> PosX; VectorVM::FExternalFuncRegisterHandler<float> PosY; VectorVM::FExternalFuncRegisterHandler<float> PosZ;
+	VectorVM::FExternalFuncRegisterHandler<float> RotX; VectorVM::FExternalFuncRegisterHandler<float> RotY; VectorVM::FExternalFuncRegisterHandler<float> RotZ; VectorVM::FExternalFuncRegisterHandler<float> RotW;
 	VectorVM::FExternalFuncRegisterHandler<float> VelX; VectorVM::FExternalFuncRegisterHandler<float> VelY; VectorVM::FExternalFuncRegisterHandler<float> VelZ;
 
 	//TODO: Rotation + Scale too? Use quats so we can get proper interpolation between bone and parent.
 
 	const bool bNeedsPosition;
+	const bool bNeedsRotation;
 	const bool bNeedsVelocity;
 
 	FORCEINLINE void SetPosition(FVector Position)
@@ -381,6 +389,14 @@ struct FBoneSocketSkinnedDataOutputHandler
 		*PosZ.GetDestAndAdvance() = Position.Z;
 	}
 
+	FORCEINLINE void SetRotation(FQuat Rotation)
+	{
+		*RotX.GetDestAndAdvance() = Rotation.X;
+		*RotY.GetDestAndAdvance() = Rotation.Y;
+		*RotZ.GetDestAndAdvance() = Rotation.Z;
+		*RotW.GetDestAndAdvance() = Rotation.W;
+	}
+
 	FORCEINLINE void SetVelocity(FVector Velocity)
 	{
 		*VelX.GetDestAndAdvance() = Velocity.X;
@@ -469,6 +485,19 @@ void UNiagaraDataInterfaceSkeletalMesh::GetSkinnedBoneData(FVectorVMContext& Con
 			Velocity = (Pos - Prev) * InvDt;
 			Output.SetVelocity(Velocity);
 		}
+
+		if (Output.bNeedsRotation)
+		{
+			FQuat Rotation = SkinningHandler.GetSkinnedBoneRotation(Accessor, Bone);
+
+			if (bInterpolated::Value)
+			{
+				FQuat PrevRotation = SkinningHandler.GetSkinnedBonePreviousRotation(Accessor, Bone);
+				Rotation = FMath::Lerp(PrevRotation, Rotation, Interp);
+			}
+
+			Output.SetRotation(Rotation);
+		}
 	}
 }
 
