Fixing static analysis warning on implicit bool
#tests compiled and ran an emitter

diff --git a/Engine/Plugins/FX/Niagara/Source/Niagara/Public/NiagaraRenderer.h b/Engine/Plugins/FX/Niagara/Source/Niagara/Public/NiagaraRenderer.h
--- a/Engine/Plugins/FX/Niagara/Source/Niagara/Public/NiagaraRenderer.h
+++ b/Engine/Plugins/FX/Niagara/Source/Niagara/Public/NiagaraRenderer.h
@@ -106,7 +106,7 @@ public:
 
 	void SetDynamicData_RenderThread(FNiagaraDynamicDataBase* NewDynamicData);
 	FORCEINLINE FNiagaraDynamicDataBase *GetDynamicData()const { return DynamicDataRender; }
-	FORCEINLINE bool HasDynamicData()const { return DynamicDataRender; }
+	FORCEINLINE bool HasDynamicData()const { return DynamicDataRender != nullptr; }
 	FORCEINLINE float GetCPUTimeMS() const { return CPUTimeMS; }
 	FORCEINLINE bool HasLights()const { return bHasLights; }
 
