Implement asynchronous message building in FNiagaraMessageManager.
Move message job sending to System and Script ViewModels instead of toolkit to clear way for messages without editor.
SystemViewModel on compile delegate now fires when compile is from a Script it uses.
Fix bug with SystemViewModel on compile delegate not firing if autocompile on change is disabled.
Autocompile Scripts when opening editor for them.
Various improvements.

diff --git a/Engine/Plugins/FX/Niagara/Source/Niagara/Classes/NiagaraScript.h b/Engine/Plugins/FX/Niagara/Source/Niagara/Classes/NiagaraScript.h
--- a/Engine/Plugins/FX/Niagara/Source/Niagara/Classes/NiagaraScript.h
+++ b/Engine/Plugins/FX/Niagara/Source/Niagara/Classes/NiagaraScript.h
@@ -514,7 +514,7 @@ public:
 	NIAGARA_API void MarkScriptAndSourceDesynchronized(FString Reason);
 	
 	/** Request a synchronous compile for the script, possibly forcing it to compile.*/
-	NIAGARA_API void RequestCompile();
+	NIAGARA_API void RequestCompile(bool bForceCompile = false);
 
 	/** Request an asynchronous compile for the script, possibly forcing it to compile. The output values are the compilation id of the data as well as the async handle to 
 		gather up the results with. bTrulyAsync tells the system whether or not the compile task must be completed on the main thread (mostly used for debugging). The
diff --git a/Engine/Plugins/FX/Niagara/Source/Niagara/Private/NiagaraScript.cpp b/Engine/Plugins/FX/Niagara/Source/Niagara/Private/NiagaraScript.cpp
--- a/Engine/Plugins/FX/Niagara/Source/Niagara/Private/NiagaraScript.cpp
+++ b/Engine/Plugins/FX/Niagara/Source/Niagara/Private/NiagaraScript.cpp
@@ -933,9 +933,9 @@ void UNiagaraScript::InvalidateCachedCompileIds()
 	GetSource()->InvalidateCachedCompileIds();
 }
 
-void UNiagaraScript::RequestCompile()
+void UNiagaraScript::RequestCompile(bool bForceCompile)
 {
-	if (!AreScriptAndSourceSynchronized())
+	if (!AreScriptAndSourceSynchronized() || bForceCompile)
 	{
 		if (IsCompilable() == false)
 		{
diff --git a/Engine/Plugins/FX/Niagara/Source/NiagaraEditor/Private/NiagaraCompiler.cpp b/Engine/Plugins/FX/Niagara/Source/NiagaraEditor/Private/NiagaraCompiler.cpp
--- a/Engine/Plugins/FX/Niagara/Source/NiagaraEditor/Private/NiagaraCompiler.cpp
+++ b/Engine/Plugins/FX/Niagara/Source/NiagaraEditor/Private/NiagaraCompiler.cpp
@@ -735,7 +735,7 @@ FNiagaraCompileResults FHlslNiagaraCompiler::CompileScript(const FNiagaraCompile
 		}
 		else if (TranslatedHLSL.Len() == 0)
 		{
-			Error(LOCTEXT("HlslTranslateErrorMessageFailed", "The HLSL Translator failed to generate HLSL!"));
+			CompileResults.NumErrors++;
 			CompileResults.bVMSucceeded = false;
 		}
 		else
@@ -760,7 +760,7 @@ FNiagaraCompileResults FHlslNiagaraCompiler::CompileScript(const FNiagaraCompile
 		}
 		else if (TranslatedHLSL.Len() == 0)
 		{
-			Error(LOCTEXT("HlslTranslateErrorMessageFailed", "The HLSL Translator failed to generate HLSL!"));
+			CompileResults.NumErrors++;
 			CompileResults.bVMSucceeded = false;
 		}
 		else
diff --git a/Engine/Plugins/FX/Niagara/Source/NiagaraEditor/Private/NiagaraHlslTranslator.cpp b/Engine/Plugins/FX/Niagara/Source/NiagaraEditor/Private/NiagaraHlslTranslator.cpp
--- a/Engine/Plugins/FX/Niagara/Source/NiagaraEditor/Private/NiagaraHlslTranslator.cpp
+++ b/Engine/Plugins/FX/Niagara/Source/NiagaraEditor/Private/NiagaraHlslTranslator.cpp
@@ -5475,7 +5475,7 @@ int32 FHlslNiagaraTranslator::CompileOutputPin(const UEdGraphPin* InPin)
 	return Ret;
 }
 
-void FHlslNiagaraTranslator::Error(FText ErrorText, const UNiagaraNode* Node, const UEdGraphPin* Pin)
+void FHlslNiagaraTranslator::Error(FText ErrorText, const UNiagaraNode* Node, const UEdGraphPin* Pin) //@todo(message manager) rearrange syntax to fit with FNiagaraMessageManager
 {
 	FString NodePinStr = TEXT("");
 	FString NodePinPrefix = TEXT(" - ");
diff --git a/Engine/Plugins/FX/Niagara/Source/NiagaraEditor/Private/NiagaraMessageManager.cpp b/Engine/Plugins/FX/Niagara/Source/NiagaraEditor/Private/NiagaraMessageManager.cpp
--- a/Engine/Plugins/FX/Niagara/Source/NiagaraEditor/Private/NiagaraMessageManager.cpp
+++ b/Engine/Plugins/FX/Niagara/Source/NiagaraEditor/Private/NiagaraMessageManager.cpp
@@ -14,7 +14,7 @@
 #define LOCTEXT_NAMESPACE "NiagaraMessageManager"
 
 FNiagaraMessageManager* FNiagaraMessageManager::Singleton = nullptr;
-TQueue<TSharedPtr<INiagaraMessageJob>> FNiagaraMessageManager::JobQueue;
+const double FNiagaraMessageManager::MaxJobWorkTime = .02f; // do message jobs at 50 fps
 
 FNiagaraMessageManager::FNiagaraMessageManager()
 {
@@ -29,22 +29,145 @@ FNiagaraMessageManager* FNiagaraMessageManager::Get()
 	return Singleton;
 }
 
-TSharedPtr<INiagaraMessage> FNiagaraMessageManager::QueueMessageJob(TSharedRef<INiagaraMessageJob> InMessageJob)
+void FNiagaraMessageManager::QueueMessageJob(TSharedPtr<const INiagaraMessageJob> InMessageJob, const FGuid& InMessageJobAssetKey)
 {
-	//JobQueue.Enqueue(InMessageJob); //@todo(message manager) stubbing the asynchronous implementation
-	return InMessageJob->GenerateNiagaraMessage();
+	TArray<TSharedPtr<const INiagaraMessageJob>> JobBatch{ InMessageJob };
+	MessageJobBatchArr.Add(FNiagaraMessageJobBatch(JobBatch, InMessageJobAssetKey));
+}
+
+void FNiagaraMessageManager::QueueMessageJobBatch(TArray<TSharedPtr<const INiagaraMessageJob>> InMessageJobBatch, const FGuid& InMessageJobBatchAssetKey)
+{
+	if (ensureMsgf(InMessageJobBatch.Num() > 0, TEXT("Queued a message job batch with no messages!")))
+	{
+		MessageJobBatchArr.Add(FNiagaraMessageJobBatch(InMessageJobBatch, InMessageJobBatchAssetKey));
+	}
+}
+
+void FNiagaraMessageManager::RefreshMessagesForAssetKey(const FGuid& InAssetKey)
+{
+	ObjectToTypeMappedMessagesMap.Remove(InAssetKey);
+	MessageJobBatchArr.RemoveAll([InAssetKey](const FNiagaraMessageJobBatch& Batch) {return Batch.MessageJobsAssetKey == InAssetKey; });
+}
+
+void FNiagaraMessageManager::RefreshMessagesForAssetKeyAndMessageJobType(const FGuid& InAssetKey, const ENiagaraMessageJobType& InMessageJobType)
+{
+	TMap<const ENiagaraMessageJobType, TArray<TSharedRef<const INiagaraMessage>>>* MessageJobTypeToMessageMap = ObjectToTypeMappedMessagesMap.Find(InAssetKey);
+	if (MessageJobTypeToMessageMap != nullptr)
+	{
+		TArray<TSharedRef<const INiagaraMessage>>* MessageArr = MessageJobTypeToMessageMap->Find(InMessageJobType);
+		if (MessageArr != nullptr)
+		{
+			MessageArr->Reset();
+		}
+	} 
+
+	for (FNiagaraMessageJobBatch& Batch : MessageJobBatchArr)
+	{
+		Batch.MessageJobs.RemoveAll([InMessageJobType](const TSharedPtr<const INiagaraMessageJob>& Job) {return Job->GetMessageJobType() == InMessageJobType; });
+	}
+}
+
+const TOptional<const FString> FNiagaraMessageManager::GetStringForScriptUsageInStack(const ENiagaraScriptUsage InScriptUsage)
+{
+	switch (InScriptUsage) {
+	case ENiagaraScriptUsage::ParticleSpawnScript:
+		return TOptional<const FString>(FString("Particle Spawn Script"));
+
+	case ENiagaraScriptUsage::ParticleSpawnScriptInterpolated:
+		return TOptional<const FString>(TEXT("Particle Spawn Script Interpolated"));
+
+	case ENiagaraScriptUsage::ParticleUpdateScript:
+		return TOptional<const FString>(TEXT("Particle Update Script"));
+		
+	case ENiagaraScriptUsage::ParticleEventScript:
+		return TOptional<const FString>(TEXT("Particle Event Script"));
+
+	case ENiagaraScriptUsage::EmitterSpawnScript:
+		return TOptional<const FString>(TEXT("Emitter Spawn Script"));
+
+	case ENiagaraScriptUsage::EmitterUpdateScript:
+		return TOptional<const FString>(TEXT("Emitter Update Script"));
+
+	case ENiagaraScriptUsage::SystemSpawnScript:
+		return TOptional<const FString>(TEXT("System Spawn Script"));
+
+	case ENiagaraScriptUsage::SystemUpdateScript:
+		return TOptional<const FString>(TEXT("System Update Script"));
+
+	//We don't expect to see these usages in the stack so do not set the toptional
+	case ENiagaraScriptUsage::DynamicInput:
+	case ENiagaraScriptUsage::Function:
+	case ENiagaraScriptUsage::Module:
+		return TOptional<const FString>();
+
+	//unhandled cases
+	case ENiagaraScriptUsage::ParticleGPUComputeScript:
+	default:
+		ensureMsgf(false, TEXT("Tried to get script usage text for usage that is not handled!"));
+		return TOptional<const FString>();
+	}
+	return TOptional<const FString>();
+}
+
+const TArray<TSharedRef<const INiagaraMessage>> FNiagaraMessageManager::GetMessagesForAssetKey(const FGuid& InAssetKey) const
+{
+ 	TArray<TSharedRef<const INiagaraMessage>> FoundMessages;
+	const TMap <const ENiagaraMessageJobType, TArray<TSharedRef<const INiagaraMessage>>>* MessageJobTypeToMessagesMapPtr = ObjectToTypeMappedMessagesMap.Find(InAssetKey);
+	if (MessageJobTypeToMessagesMapPtr != nullptr)
+	{
+		for (const TTuple<const ENiagaraMessageJobType, TArray<TSharedRef<const INiagaraMessage>>>& MessageJobTypeToMessagesPair : *MessageJobTypeToMessagesMapPtr)
+		{
+			FoundMessages.Append(MessageJobTypeToMessagesPair.Value);
+		}
+	}
+	return FoundMessages;
+}
+
+void FNiagaraMessageManager::Tick(float DeltaSeconds)
+{
+	DoMessageJobsInQueueTick();
+}
+
+void FNiagaraMessageManager::DoMessageJobsInQueueTick()
+{
+	if (MessageJobBatchArr.Num() > 0)
+	{
+		double WorkStartTime = FPlatformTime::Seconds();
+		double CurrentWorkLoopTime = WorkStartTime;
+
+		while(MessageJobBatchArr.Num() > 0)
+		{ 
+			FNiagaraMessageJobBatch& CurrentMessageJobBatch = MessageJobBatchArr[0];
+			
+			while(CurrentMessageJobBatch.MessageJobs.Num() > 0)
+			{
+				TSharedPtr<const INiagaraMessageJob> CurrentMessageJob = CurrentMessageJobBatch.MessageJobs.Pop(false);
+				TSharedRef<const INiagaraMessage> GeneratedMessage = CurrentMessageJob->GenerateNiagaraMessage();
+				GeneratedMessagesForCurrentMessageJobBatch.Add(GeneratedMessage);
+				ObjectToTypeMappedMessagesMap.FindOrAdd(CurrentMessageJobBatch.MessageJobsAssetKey).FindOrAdd(CurrentMessageJob->GetMessageJobType()).Add(GeneratedMessage);
+				CurrentWorkLoopTime = FPlatformTime::Seconds();
+
+				if (CurrentWorkLoopTime - WorkStartTime > MaxJobWorkTime)
+				{
+					return;
+				}
+			}
+			OnRequestRefresh.Broadcast(CurrentMessageJobBatch.MessageJobsAssetKey, GeneratedMessagesForCurrentMessageJobBatch);
+			GeneratedMessagesForCurrentMessageJobBatch.Reset();
+			MessageJobBatchArr.Pop();
+		} 
+	}
 }
 
 FNiagaraMessageCompileEvent::FNiagaraMessageCompileEvent(
 	  const FNiagaraCompileEvent& InCompileEvent
 	, TArray<FNiagaraScriptNameAndAssetPath>& InContextScriptNamesAndAssetPaths
-	, TOptional<const FString>& InEmitterScriptName
+	, TOptional<const FText>& InOwningScriptNameAndUsageText
 	, TOptional<const FNiagaraScriptNameAndAssetPath>& InCompiledScriptNameAndAssetPath
 	)
-	: INiagaraMessage(ENiagaraMessageType::CompileEventMessage)
-	, CompileEvent(InCompileEvent)
+	: CompileEvent(InCompileEvent)
 	, ContextScriptNamesAndAssetPaths(InContextScriptNamesAndAssetPaths)
-	, EmitterScriptName(InEmitterScriptName)
+	, OwningScriptNameAndUsageText(InOwningScriptNameAndUsageText)
 	, CompiledScriptNameAndAssetPath(InCompiledScriptNameAndAssetPath)
 {
 }
@@ -126,10 +249,10 @@ TSharedRef<FTokenizedMessage> FNiagaraMessageCompileEvent::GenerateTokenizedMess
 		}
 	}
 	
-	//Now add the emitter name if we passed through an Emitter Node
-	if (EmitterScriptName.IsSet())
+	//Now add the owning script name and usage if it is set
+	if (OwningScriptNameAndUsageText.IsSet())
 	{
-		Message->AddToken(FTextToken::Create(FText::FromString(*EmitterScriptName.GetValue())));
+		Message->AddToken(FTextToken::Create(OwningScriptNameAndUsageText.GetValue()));
 	}
 
 	//Finally add the context stack of the scripts that were passed through to get to the originating graph
@@ -140,14 +263,20 @@ TSharedRef<FTokenizedMessage> FNiagaraMessageCompileEvent::GenerateTokenizedMess
 	return Message;
 }
 
-FNiagaraMessageJobCompileEvent::FNiagaraMessageJobCompileEvent(const FNiagaraCompileEvent& InCompileEvent, const TWeakObjectPtr<UNiagaraScript>& InOriginatingScriptWeakObjPtr, const bool bInFromScriptToolkit)
+FNiagaraMessageJobCompileEvent::FNiagaraMessageJobCompileEvent(
+	  const FNiagaraCompileEvent& InCompileEvent
+	, const TWeakObjectPtr<UNiagaraScript>& InOriginatingScriptWeakObjPtr
+	, const TOptional<const FString>& InOwningScriptNameString
+	, const TOptional<const FString>& InSourceScriptAssetPath
+	)
 	: CompileEvent(InCompileEvent)
 	, OriginatingScriptWeakObjPtr(InOriginatingScriptWeakObjPtr)
-	, bFromScriptToolkit(bInFromScriptToolkit)
+	, OwningScriptNameString(InOwningScriptNameString)
+	, SourceScriptAssetPath(InSourceScriptAssetPath)
 {
 }
 
-const TSharedPtr<INiagaraMessage> FNiagaraMessageJobCompileEvent::GenerateNiagaraMessage() const
+TSharedRef<const INiagaraMessage> FNiagaraMessageJobCompileEvent::GenerateNiagaraMessage() const
 {
 	TArray<FGuid> ContextStackGuids = CompileEvent.StackGuids;
 
@@ -161,10 +290,11 @@ const TSharedPtr<INiagaraMessage> FNiagaraMessageJobCompileEvent::GenerateNiagar
 		checkf(ScriptGraph, TEXT("Function Script does not have a UNiagaraGraph!"));
 
 		TArray<FNiagaraScriptNameAndAssetPath> ContextScriptNamesAndPaths;
-		TOptional<const FString> EmitterScriptName;
 		TOptional<const FNiagaraScriptNameAndAssetPath> CompiledScriptNameAndPath;
+		TOptional<const FText> OwningScriptNameAndUsageText;
 		TOptional<const FText> ScriptNameAndPathsGetterFailureReason;
-		const bool bSuccessfullyFoundScripts = RecursiveGetScriptNamesAndAssetPathsFromContextStack(ContextStackGuids, ScriptGraph, ContextScriptNamesAndPaths, EmitterScriptName, ScriptNameAndPathsGetterFailureReason);
+		TOptional<const FString> OwningScriptNameStringCopy = OwningScriptNameString;
+		bool bSuccessfullyFoundScripts = RecursiveGetScriptNamesAndAssetPathsFromContextStack(ContextStackGuids, ScriptGraph, ContextScriptNamesAndPaths, OwningScriptNameStringCopy, ScriptNameAndPathsGetterFailureReason);
 
 		if (bSuccessfullyFoundScripts == false)
 		{
@@ -172,12 +302,25 @@ const TSharedPtr<INiagaraMessage> FNiagaraMessageJobCompileEvent::GenerateNiagar
 			return MakeShared<FNiagaraMessageNeedRecompile>(ScriptNameAndPathsGetterFailureReason.GetValue());
 		}
 
-		if (bFromScriptToolkit)
+		const ENiagaraScriptUsage ScriptUsage = OriginatingScriptWeakObjPtr->GetUsage();
+
+		if (OwningScriptNameStringCopy.IsSet())
+		{
+			FString ScriptAndUsage = OwningScriptNameStringCopy.GetValue();
+			const TOptional<const FString> ScriptUsageInStackString = FNiagaraMessageManager::GetStringForScriptUsageInStack(ScriptUsage);
+			if (ScriptUsageInStackString.IsSet())
+			{
+				ScriptAndUsage = ScriptAndUsage + FString(", ") + ScriptUsageInStackString.GetValue() + FString(", ");
+			}
+			OwningScriptNameAndUsageText = TOptional<const FText>(FText::FromString(ScriptAndUsage));
+		}
+
+		if (SourceScriptAssetPath.IsSet())
 		{
 			//If this compile event is from a script, set the compiled script name and asset path so the user can navigate to errors locally.
-			CompiledScriptNameAndPath = TOptional<const FNiagaraScriptNameAndAssetPath>(FNiagaraScriptNameAndAssetPath(OriginatingScriptWeakObjPtr->GetName(), OriginatingScriptWeakObjPtr->GetPathName()));
+			CompiledScriptNameAndPath = TOptional<const FNiagaraScriptNameAndAssetPath>(FNiagaraScriptNameAndAssetPath(OriginatingScriptWeakObjPtr->GetName(), SourceScriptAssetPath.GetValue()));
 		}
-		return MakeShared<FNiagaraMessageCompileEvent>(CompileEvent, ContextScriptNamesAndPaths, EmitterScriptName, CompiledScriptNameAndPath);
+		return MakeShared<FNiagaraMessageCompileEvent>(CompileEvent, ContextScriptNamesAndPaths, OwningScriptNameAndUsageText, CompiledScriptNameAndPath);
 	}
 	//The originating script weak ptr is no longer valid, send an error message asking for recompile.
 	FText MessageText = LOCTEXT("CompileEventMessageJobFail", "Cached info for compile event is out of date, recompile to get full info. Event: {0}");
@@ -186,11 +329,11 @@ const TSharedPtr<INiagaraMessage> FNiagaraMessageJobCompileEvent::GenerateNiagar
 	return MakeShared<FNiagaraMessageNeedRecompile>(MessageText);
 }
 
-const bool FNiagaraMessageJobCompileEvent::RecursiveGetScriptNamesAndAssetPathsFromContextStack(
+bool FNiagaraMessageJobCompileEvent::RecursiveGetScriptNamesAndAssetPathsFromContextStack(
 	  TArray<FGuid>& InContextStackNodeGuids
 	, const UNiagaraGraph* InGraphToSearch
 	, TArray<FNiagaraScriptNameAndAssetPath>& OutContextScriptNamesAndAssetPaths
-	, TOptional<const FString>& OutEmitterScriptName
+	, TOptional<const FString>& OutEmitterName
 	, TOptional<const FText>& OutFailureReason
 ) const
 {
@@ -230,14 +373,14 @@ const bool FNiagaraMessageJobCompileEvent::RecursiveGetScriptNamesAndAssetPathsF
 			UNiagaraScript* FunctionCallNodeAssignedScript = FunctionCallNode->FunctionScript;
 			if (FunctionCallNodeAssignedScript == nullptr)
 			{
-				FText FailureReason = LOCTEXT("CompileEventMessageGenerator_FunctionCallNodeScriptNotFound", "Script for Function Call Node \"{0}\" not found!");
+				FText FailureReason = LOCTEXT("GenerateCompileEventMessage_FunctionCallNodeScriptNotFound", "Script for Function Call Node \"{0}\" not found!");
 				OutFailureReason = FText::Format(FailureReason, FText::FromString(FunctionCallNode->GetFunctionName()));
 				return false;
 			}
 			UNiagaraScriptSourceBase* FunctionScriptSourceBase = FunctionCallNodeAssignedScript->GetSource();
 			if (FunctionScriptSourceBase == nullptr)
 			{
-				FText FailureReason = LOCTEXT("CompileEventMessageGenerator_FunctionCallNodeScriptSourceBaseNotFound", "Source Script for Function Call Node \"{0}\" not found!");
+				FText FailureReason = LOCTEXT("GenerateCompileEventMessage_FunctionCallNodeScriptSourceBaseNotFound", "Source Script for Function Call Node \"{0}\" not found!");
 				OutFailureReason = FText::Format(FailureReason, FText::FromString(FunctionCallNode->GetFunctionName()));
 				return false;
 			}
@@ -247,12 +390,12 @@ const bool FNiagaraMessageJobCompileEvent::RecursiveGetScriptNamesAndAssetPathsF
 			UNiagaraGraph* FunctionScriptGraph = FunctionScriptSource->NodeGraph;
 			if (FunctionScriptGraph == nullptr)
 			{
-				FText FailureReason = LOCTEXT("CompileEventMessageGenerator_FunctionCallNodeGraphNotFound", "Graph for Function Call Node \"{0}\" not found!");
+				FText FailureReason = LOCTEXT("GenerateCompileEventMessage_FunctionCallNodeGraphNotFound", "Graph for Function Call Node \"{0}\" not found!");
 				OutFailureReason = FText::Format(FailureReason, FText::FromString(FunctionCallNode->GetFunctionName()));
 				return false;
 			}
 			OutContextScriptNamesAndAssetPaths.Add(FNiagaraScriptNameAndAssetPath(FunctionCallNodeAssignedScript->GetName(), FunctionCallNodeAssignedScript->GetPathName()));
-			return RecursiveGetScriptNamesAndAssetPathsFromContextStack(InContextStackNodeGuids, FunctionScriptGraph, OutContextScriptNamesAndAssetPaths, OutEmitterScriptName, OutFailureReason);
+			return RecursiveGetScriptNamesAndAssetPathsFromContextStack(InContextStackNodeGuids, FunctionScriptGraph, OutContextScriptNamesAndAssetPaths, OutEmitterName, OutFailureReason);
 		}
 
 		UNiagaraNodeEmitter* EmitterNode = Cast<UNiagaraNodeEmitter>(ContextNode.GetValue());
@@ -264,8 +407,8 @@ const bool FNiagaraMessageJobCompileEvent::RecursiveGetScriptNamesAndAssetPathsF
 			UNiagaraGraph* EmitterScriptGraph = EmitterScriptSource->NodeGraph;
 			checkf(EmitterScriptGraph, TEXT("Emitter Script Source does not have a UNiagaraGraph!"));
 
-			OutEmitterScriptName = EmitterNode->GetEmitterUniqueName();
-			return RecursiveGetScriptNamesAndAssetPathsFromContextStack(InContextStackNodeGuids, EmitterScriptGraph, OutContextScriptNamesAndAssetPaths, OutEmitterScriptName, OutFailureReason);
+			OutEmitterName = EmitterNode->GetEmitterUniqueName();
+			return RecursiveGetScriptNamesAndAssetPathsFromContextStack(InContextStackNodeGuids, EmitterScriptGraph, OutContextScriptNamesAndAssetPaths, OutEmitterName, OutFailureReason);
 		}
 		checkf(false, TEXT("Matching node is not a function call or emitter node!"));
 	}
@@ -275,10 +418,37 @@ const bool FNiagaraMessageJobCompileEvent::RecursiveGetScriptNamesAndAssetPathsF
 	return false;
 }
 
-FNiagaraMessageNeedRecompile::FNiagaraMessageNeedRecompile(const FText& InNeedRecompileMessage)
-	: INiagaraMessage(ENiagaraMessageType::NeedRecompileMessage)
-	, NeedRecompileMessage(InNeedRecompileMessage)
+TSharedRef<const INiagaraMessage> FNiagaraMessageJobPostCompileSummary::GenerateNiagaraMessage() const
 {
+	FText MessageText = FText();
+	EMessageSeverity::Type MessageSeverity = EMessageSeverity::Warning;
+	switch (ScriptCompileStatus)
+	{
+	default:
+	case ENiagaraScriptCompileStatus::NCS_Unknown:
+	case ENiagaraScriptCompileStatus::NCS_Dirty:
+		MessageText = LOCTEXT("NiagaraCompileStatusUnknownInfo", "{0} compile status unknown with {1} warning(s) and {2} error(s).");
+		MessageText = FText::Format(MessageText, CompiledObjectNameText, FText::FromString(FString::FromInt(NumCompileWarnings)), FText::FromString(FString::FromInt(NumCompileErrors)));
+		//MessageSeverity = EMessageSeverity::Warning;
+		break;
+	case ENiagaraScriptCompileStatus::NCS_Error:
+		MessageText = LOCTEXT("NiagaraCompileStatusErrorInfo", "{0} failed to compile with {1} warning(s) and {2} error(s).");
+		MessageText =  FText::Format(MessageText, CompiledObjectNameText, FText::FromString(FString::FromInt(NumCompileWarnings)), FText::FromString(FString::FromInt(NumCompileErrors)));
+		//MessageSeverity = EMessageSeverity::Error;
+		break;
+	case ENiagaraScriptCompileStatus::NCS_UpToDate:
+		MessageText = LOCTEXT("NiagaraCompileStatusSuccessInfo", "{0} successfully compiled.");
+		MessageText = FText::Format(MessageText, CompiledObjectNameText);
+		//MessageSeverity = EMessageSeverity::Info;
+		break;
+	case ENiagaraScriptCompileStatus::NCS_UpToDateWithWarnings:
+		MessageText = LOCTEXT("NiagaraCompileStatusWarningInfo", "{0} successfully compiled with {1} warning(s).");
+		MessageText =  FText::Format(MessageText, CompiledObjectNameText, FText::FromString(FString::FromInt(NumCompileWarnings)));
+		//MessageSeverity = EMessageSeverity::Warning;
+		break;
+	}
+	MessageSeverity = EMessageSeverity::Info;
+	return MakeShared<FNiagaraMessagePostCompileSummary>(MessageText, MessageSeverity);
 }
 
 TSharedRef<FTokenizedMessage> FNiagaraMessageNeedRecompile::GenerateTokenizedMessage() const
@@ -286,6 +456,11 @@ TSharedRef<FTokenizedMessage> FNiagaraMessageNeedRecompile::GenerateTokenizedMes
 	return FTokenizedMessage::Create(EMessageSeverity::Error, NeedRecompileMessage);
 }
 
+TSharedRef<FTokenizedMessage> FNiagaraMessagePostCompileSummary::GenerateTokenizedMessage() const
+{
+	return FTokenizedMessage::Create(MessageSeverity, PostCompileSummaryText);
+}
+
 FNiagaraCompileEventToken::FNiagaraCompileEventToken(
 	  const FString& InScriptAssetPath
 	, const FText& InMessage
diff --git a/Engine/Plugins/FX/Niagara/Source/NiagaraEditor/Private/Toolkits/NiagaraScriptToolkit.cpp b/Engine/Plugins/FX/Niagara/Source/NiagaraEditor/Private/Toolkits/NiagaraScriptToolkit.cpp
--- a/Engine/Plugins/FX/Niagara/Source/NiagaraEditor/Private/Toolkits/NiagaraScriptToolkit.cpp
+++ b/Engine/Plugins/FX/Niagara/Source/NiagaraEditor/Private/Toolkits/NiagaraScriptToolkit.cpp
@@ -45,6 +45,8 @@
 
 #include "AssetTypeActions/AssetTypeActions_NiagaraScript.h"
 #include "NiagaraMessageManager.h"
+#include "NiagaraStandaloneScriptViewModel.h"
+#include "NiagaraMessageLogViewModel.h"
 
 #define LOCTEXT_NAMESPACE "NiagaraScriptToolkit"
 
@@ -148,31 +150,20 @@ void FNiagaraScriptToolkit::Initialize( const EToolkitMode::Type Mode, const TSh
 
 	DetailsScriptSelection = MakeShareable(new FNiagaraObjectSelection());
 	DetailsScriptSelection->SetSelectedObject(EditedNiagaraScript);
+
+	const FGuid MessageLogGuidKey = FGuid::NewGuid();
 	
-	FMessageLogModule& MessageLogModule = FModuleManager::LoadModuleChecked<FMessageLogModule>("MessageLog");
-	FMessageLogInitializationOptions LogOptions;
-	// Show Pages so that user is never allowed to clear log messages
+ 	FMessageLogModule& MessageLogModule = FModuleManager::LoadModuleChecked<FMessageLogModule>("MessageLog"); //@todo(message manager) remove stats listing
+ 	FMessageLogInitializationOptions LogOptions;
+ 	// Show Pages so that user is never allowed to clear log messages
 	LogOptions.bShowPages = false;
 	LogOptions.bShowFilters = false;
 	LogOptions.bAllowClear = false;
 	LogOptions.MaxPageCount = 1;
-
-	// Reuse any existing log, or create a new one (that is not held onto by the message log system)
-	auto CreateMessageLogListing = [&MessageLogModule, &LogOptions](const FName& LogName)->TSharedRef<IMessageLogListing> {
-		if (MessageLogModule.IsRegisteredLogListing(LogName))
-		{
-			return MessageLogModule.GetLogListing(LogName);
-		}
-		else
-		{
-			return  MessageLogModule.CreateLogListing(LogName, LogOptions);
-		}
-	};
-
 	StatsListing = MessageLogModule.CreateLogListing("MaterialEditorStats", LogOptions);
 	Stats = MessageLogModule.CreateLogListingWidget(StatsListing.ToSharedRef());
-	NiagaraMessageLogListing = CreateMessageLogListing(GetNiagaraScriptMessageLogName(EditedNiagaraScript));
-	NiagaraMessageLog = MessageLogModule.CreateLogListingWidget(NiagaraMessageLogListing.ToSharedRef());
+
+	NiagaraMessageLogViewModel = MakeShared<FNiagaraMessageLogViewModel>(GetNiagaraScriptMessageLogName(EditedNiagaraScript), MessageLogGuidKey, NiagaraMessageLog);
 
 	TSharedRef<FTabManager::FLayout> StandaloneDefaultLayout = FTabManager::NewLayout("Standalone_Niagara_Layout_v8")
 	->AddArea
@@ -249,7 +240,7 @@ void FNiagaraScriptToolkit::Initialize( const EToolkitMode::Type Mode, const TSh
 	SetupCommands();
 	ExtendToolbar();
 	RegenerateMenusAndToolbars();
-	UpdateMessageLog();
+	
 	UpdateModuleStats();
 
 	// @todo toolkit world centric editing
@@ -334,7 +325,6 @@ void FNiagaraScriptToolkit::OnEditedScriptPropertyFinishedChanging(const FProper
 
 void FNiagaraScriptToolkit::OnVMScriptCompiled(UNiagaraScript* InScript)
 {
-	UpdateMessageLog();
 	UpdateModuleStats();
 }
 
@@ -562,60 +552,6 @@ void FNiagaraScriptToolkit::AddReferencedObjects(FReferenceCollector& Collector)
 	Collector.AddReferencedObject(EditedNiagaraScript);
 }
 
-
-
-void FNiagaraScriptToolkit::UpdateMessageLog()
-{
-	TArray<TSharedPtr<INiagaraMessage>> NiagaraMessages;
-	uint32 ErrorCount = 0;
-	uint32 WarningCount = 0;
-	for (const FNiagaraCompileEvent CompileEvent : EditedNiagaraScript->GetVMExecutableData().LastCompileEvents)
-	{
-		NiagaraMessages.Add(FNiagaraMessageManager::Get()->QueueMessageJob(MakeShared<FNiagaraMessageJobCompileEvent>(CompileEvent, TWeakObjectPtr<UNiagaraScript>(OriginalNiagaraScript), true)));
-		if (CompileEvent.Severity == FNiagaraCompileEventSeverity::Error)
-		{
-			ErrorCount++;
-		}
-		else if (CompileEvent.Severity == FNiagaraCompileEventSeverity::Warning)
-		{
-			WarningCount++;
-		}
-	}
-
-	TArray<TSharedRef<FTokenizedMessage>> TokenizedMessages;
-	for (const TSharedPtr<INiagaraMessage> Message : NiagaraMessages)
-	{
-		TokenizedMessages.Add(Message->GenerateTokenizedMessage());
-	}
-
-	const auto GetCompileCompleteMessageText = [&ErrorCount, &WarningCount](ENiagaraScriptCompileStatus Status)->const FText {
-		FText MessageText;
-		switch (Status)
-		{
-		default:
-		case ENiagaraScriptCompileStatus::NCS_Unknown:
-		case ENiagaraScriptCompileStatus::NCS_Dirty:
-			MessageText = LOCTEXT("NiagaraScriptCompileStatusUnknownInfo", "Script compile status unknown with {0} warning(s) and {1} error(s).");
-			return FText::Format(MessageText, FText::FromString(FString::FromInt(WarningCount)), FText::FromString(FString::FromInt(ErrorCount)));
-		case ENiagaraScriptCompileStatus::NCS_Error:
-			MessageText = LOCTEXT("NiagaraScriptCompileStatusErrorInfo", "Script failed to compile with {0} warning(s) and {1} error(s).");
-			return FText::Format(MessageText, FText::FromString(FString::FromInt(WarningCount)), FText::FromString(FString::FromInt(ErrorCount)));
-		case ENiagaraScriptCompileStatus::NCS_UpToDate:
-			MessageText = LOCTEXT("NiagaraScriptCompileStatusSuccessInfo", "Script successfully compiled.");
-			return MessageText;
-		case ENiagaraScriptCompileStatus::NCS_UpToDateWithWarnings:
-			MessageText = LOCTEXT("NiagaraScriptCompileStatusWarningInfo", "Script successfully compiled with {0} warning(s).");
-			return FText::Format(MessageText, FText::FromString(FString::FromInt(WarningCount)));
-		}
-	};
-
-	const FText CompileCompleteMessageText = GetCompileCompleteMessageText(ScriptViewModel->GetLatestCompileStatus());
-	TokenizedMessages.Add(FTokenizedMessage::Create(EMessageSeverity::Info, CompileCompleteMessageText));
-
-	NiagaraMessageLogListing->ClearMessages();
- 	NiagaraMessageLogListing->AddMessages(TokenizedMessages);
-}
-
 void FNiagaraScriptToolkit::UpdateModuleStats()
 {
 	TArray< TSharedRef<class FTokenizedMessage> > Messages;
diff --git a/Engine/Plugins/FX/Niagara/Source/NiagaraEditor/Private/Toolkits/NiagaraScriptToolkit.h b/Engine/Plugins/FX/Niagara/Source/NiagaraEditor/Private/Toolkits/NiagaraScriptToolkit.h
--- a/Engine/Plugins/FX/Niagara/Source/NiagaraEditor/Private/Toolkits/NiagaraScriptToolkit.h
+++ b/Engine/Plugins/FX/Niagara/Source/NiagaraEditor/Private/Toolkits/NiagaraScriptToolkit.h
@@ -20,6 +20,7 @@ class UNiagaraScriptSource;
 class FNiagaraScriptViewModel;
 class FNiagaraObjectSelection;
 struct FEdGraphEditAction;
+class FNiagaraMessageLogViewModel;
 
 /** Viewer/editor for a DataTable */
 class FNiagaraScriptToolkit : public FAssetEditorToolkit, public FGCObject
@@ -52,11 +53,6 @@ public:
 	/** FGCObject interface */
 	virtual void AddReferencedObjects(FReferenceCollector& Collector) override;
 
-	/**
-	* Updates list of Niagara messages in message log
-	*/
-	void UpdateMessageLog();
-
 	/**
 	* Updates list of module info used to show stats
 	*/
@@ -132,8 +128,8 @@ private:
 	TSharedPtr<FNiagaraObjectSelection> DetailsScriptSelection;
 
 	/** Message log, with the log listing that it reflects */
+	TSharedPtr<FNiagaraMessageLogViewModel> NiagaraMessageLogViewModel;
 	TSharedPtr<class SWidget> NiagaraMessageLog;
-	TSharedPtr<class IMessageLogListing> NiagaraMessageLogListing;
 
 	/**	The tab ids for the Niagara editor */
 	static const FName NodeGraphTabId; 
diff --git a/Engine/Plugins/FX/Niagara/Source/NiagaraEditor/Private/Toolkits/NiagaraSystemToolkit.cpp b/Engine/Plugins/FX/Niagara/Source/NiagaraEditor/Private/Toolkits/NiagaraSystemToolkit.cpp
--- a/Engine/Plugins/FX/Niagara/Source/NiagaraEditor/Private/Toolkits/NiagaraSystemToolkit.cpp
+++ b/Engine/Plugins/FX/Niagara/Source/NiagaraEditor/Private/Toolkits/NiagaraSystemToolkit.cpp
@@ -48,9 +48,7 @@
 #include "IAssetRegistry.h"
 #include "HAL/PlatformFilemanager.h"
 #include "Misc/FileHelper.h"
-#include "MessageLog/Public/MessageLogModule.h"
-#include "MessageLog/Public/IMessageLogListing.h"
-#include "NiagaraMessageManager.h"
+#include "NiagaraMessageLogViewModel.h"
 
 #define LOCTEXT_NAMESPACE "NiagaraSystemEditor"
 
@@ -177,86 +175,6 @@ void FNiagaraSystemToolkit::AddReferencedObjects(FReferenceCollector& Collector)
 	Collector.AddReferencedObject(System);
 }
 
-void FNiagaraSystemToolkit::UpdateMessageLog()
-{	
-	TArray<UNiagaraScript*> ScriptsToGetCompileEventsFrom;
-	ScriptsToGetCompileEventsFrom.Add(System->GetSystemSpawnScript());
-	ScriptsToGetCompileEventsFrom.Add(System->GetSystemUpdateScript());
-	const TArray<FNiagaraEmitterHandle> EmitterHandles = System->GetEmitterHandles();
-	for (const FNiagaraEmitterHandle& Handle : EmitterHandles)
-	{
-		const UNiagaraEmitter* EmitterInSystem = Handle.GetInstance();
-		TArray<UNiagaraScript*> EmitterScripts;
-		EmitterInSystem->GetScripts(EmitterScripts);
-		ScriptsToGetCompileEventsFrom.Append(EmitterScripts);
-	}
-
-	TArray<TSharedPtr<INiagaraMessage>> NiagaraMessages;
-	uint32 ErrorCount = 0;
-	uint32 WarningCount = 0;
-	for (UNiagaraScript* Script : ScriptsToGetCompileEventsFrom)
-	{
-		for (const FNiagaraCompileEvent CompileEvent : Script->GetVMExecutableData().LastCompileEvents)
-		{
-			NiagaraMessages.Add(FNiagaraMessageManager::Get()->QueueMessageJob(MakeShared<FNiagaraMessageJobCompileEvent>(CompileEvent, TWeakObjectPtr<UNiagaraScript>(Script), false)));
-			if (CompileEvent.Severity == FNiagaraCompileEventSeverity::Error)
-			{
-				ErrorCount++;
-			}
-			else if (CompileEvent.Severity == FNiagaraCompileEventSeverity::Warning)
-			{
-				WarningCount++;
-			}
-		}
-	}
-
-	TArray<TSharedRef<FTokenizedMessage>> TokenizedMessages;
-	for (const TSharedPtr<INiagaraMessage> Message : NiagaraMessages)
-	{
-		if (Message->GetMessageType() == ENiagaraMessageType::CompileEventMessage)
-		{
-			const TSharedPtr<const FNiagaraMessageCompileEvent> MessageCompileEvent = StaticCastSharedPtr<const FNiagaraMessageCompileEvent>(Message);
-			TokenizedMessages.Add(MessageCompileEvent->GenerateTokenizedMessage());
-		}
-		else if (Message->GetMessageType() == ENiagaraMessageType::NeedRecompileMessage)
-		{
-			const TSharedPtr<const FNiagaraMessageNeedRecompile> MessageNeedRecompile = StaticCastSharedPtr<const FNiagaraMessageNeedRecompile>(Message);
-			TokenizedMessages.Add(MessageNeedRecompile->GenerateTokenizedMessage());
-		}
-		else
-		{
-			checkf(false, TEXT("Message is of ENiagaraMessageType that is unhandled!"));
-		}
-	}
-
-	const auto GetCompileCompleteMessageText = [&ErrorCount, &WarningCount](ENiagaraScriptCompileStatus Status)->const FText {
-		FText MessageText;
-		switch (Status)
-		{
-		default:
-		case ENiagaraScriptCompileStatus::NCS_Unknown:
-		case ENiagaraScriptCompileStatus::NCS_Dirty:
-			MessageText = LOCTEXT("NiagaraSystemCompileStatusUnknownInfo", "System compile status unknown with {0} warning(s) and {1} error(s).");
-			return FText::Format(MessageText, FText::FromString(FString::FromInt(WarningCount)), FText::FromString(FString::FromInt(ErrorCount)));
-		case ENiagaraScriptCompileStatus::NCS_Error:
-			MessageText = LOCTEXT("NiagaraSystemCompileStatusErrorInfo", "System failed to compile with {0} warning(s) and {1} error(s).");
-			return FText::Format(MessageText, FText::FromString(FString::FromInt(WarningCount)), FText::FromString(FString::FromInt(ErrorCount)));
-		case ENiagaraScriptCompileStatus::NCS_UpToDate:
-			MessageText = LOCTEXT("NiagaraSystemCompileStatusSuccessInfo", "System successfully compiled.");
-			return MessageText;
-		case ENiagaraScriptCompileStatus::NCS_UpToDateWithWarnings:
-			MessageText = LOCTEXT("NiagaraSystemCompileStatusWarningInfo", "System successfully compiled with {0} warning(s).");
-			return FText::Format(MessageText, FText::FromString(FString::FromInt(WarningCount)));
-		}
-	};
-	
-	const FText CompileCompleteMessageText = GetCompileCompleteMessageText(SystemViewModel->GetLatestCompileStatus());
-	TokenizedMessages.Add(FTokenizedMessage::Create(EMessageSeverity::Info, CompileCompleteMessageText));
-
-	NiagaraMessageLogListing->ClearMessages();
-	NiagaraMessageLogListing->AddMessages(TokenizedMessages);
-}
-
 void FNiagaraSystemToolkit::InitializeWithSystem(const EToolkitMode::Type Mode, const TSharedPtr< class IToolkitHost >& InitToolkitHost, UNiagaraSystem& InSystem)
 {
 	System = &InSystem;
@@ -266,8 +184,9 @@ void FNiagaraSystemToolkit::InitializeWithSystem(const EToolkitMode::Type Mode,
 	SystemOptions.bCanModifyEmittersFromTimeline = true;
 	SystemOptions.EditMode = ENiagaraSystemViewModelEditMode::SystemAsset;
 	SystemOptions.OnGetSequencerAddMenuContent.BindSP(this, &FNiagaraSystemToolkit::GetSequencerAddMenuContent);
+	const FGuid MessageLogGuid = FGuid::NewGuid();
 
-	SystemViewModel = MakeShareable(new FNiagaraSystemViewModel(*System, SystemOptions));
+	SystemViewModel = MakeShareable(new FNiagaraSystemViewModel(*System, SystemOptions, MessageLogGuid));
 	SystemViewModel->SetToolkitCommands(GetToolkitCommands());
 	SystemToolkitMode = ESystemToolkitMode::System;
 
@@ -282,7 +201,7 @@ void FNiagaraSystemToolkit::InitializeWithSystem(const EToolkitMode::Type Mode,
 		FNiagaraEditorUtilities::WriteTextFileToDisk(FPaths::ProjectLogDir(), FilenamePart + TEXT(".onLoad.txt"), ExportText, true);
 	}
 
-	InitializeInternal(Mode, InitToolkitHost);
+	InitializeInternal(Mode, InitToolkitHost, MessageLogGuid);
 }
 
 void FNiagaraSystemToolkit::InitializeWithEmitter(const EToolkitMode::Type Mode, const TSharedPtr< class IToolkitHost >& InitToolkitHost, UNiagaraEmitter& InEmitter)
@@ -333,8 +252,9 @@ void FNiagaraSystemToolkit::InitializeWithEmitter(const EToolkitMode::Type Mode,
 	FNiagaraSystemViewModelOptions SystemOptions;
 	SystemOptions.bCanModifyEmittersFromTimeline = false;
 	SystemOptions.EditMode = ENiagaraSystemViewModelEditMode::EmitterAsset;
+	const FGuid MessageLogGuid = FGuid::NewGuid();
 
-	SystemViewModel = MakeShareable(new FNiagaraSystemViewModel(*System, SystemOptions));
+	SystemViewModel = MakeShareable(new FNiagaraSystemViewModel(*System, SystemOptions, MessageLogGuid));
 	SystemViewModel->SetToolkitCommands(GetToolkitCommands());
 	SystemViewModel->AddEmitter(*Emitter);
 
@@ -358,41 +278,17 @@ void FNiagaraSystemToolkit::InitializeWithEmitter(const EToolkitMode::Type Mode,
 		FNiagaraEditorUtilities::WriteTextFileToDisk(FPaths::ProjectLogDir(), FilenamePart + TEXT(".onLoad.txt"), ExportText, true);
 	}
 
-	InitializeInternal(Mode, InitToolkitHost);
+	InitializeInternal(Mode, InitToolkitHost, MessageLogGuid);
 }
 
-void FNiagaraSystemToolkit::InitializeInternal(const EToolkitMode::Type Mode, const TSharedPtr<IToolkitHost>& InitToolkitHost)
+void FNiagaraSystemToolkit::InitializeInternal(const EToolkitMode::Type Mode, const TSharedPtr<IToolkitHost>& InitToolkitHost, const FGuid& MessageLogGuid)
 {
-	SystemViewModel->GetSystemScriptViewModel()->OnSystemCompiled().AddSP(this, &FNiagaraSystemToolkit::OnVMSystemCompiled);
-
 	if (SystemViewModel->GetEmitterHandleViewModels().Num() > 0)
 	{
 		SystemViewModel->SetSelectedEmitterHandleById(SystemViewModel->GetEmitterHandleViewModels()[0]->GetId());
 	}
 
-	FMessageLogModule& MessageLogModule = FModuleManager::LoadModuleChecked<FMessageLogModule>("MessageLog");
-
-	// Reuse any existing log, or create a new one (that is not held onto by the message log system)
-	auto CreateMessageLogListing = [&MessageLogModule](const FName& LogName)->TSharedRef<IMessageLogListing> {	
-		FMessageLogInitializationOptions LogOptions;
-		// Show Pages so that user is never allowed to clear log messages
-		LogOptions.bShowPages = false;
-		LogOptions.bShowFilters = false;
-		LogOptions.bAllowClear = false;
-		LogOptions.MaxPageCount = 1;
-
-		if (MessageLogModule.IsRegisteredLogListing(LogName))
-		{
-			return MessageLogModule.GetLogListing(LogName);
-		}
-		else
-		{
-			return  MessageLogModule.CreateLogListing(LogName, LogOptions);
-		}
-	};
-
-	NiagaraMessageLogListing = CreateMessageLogListing(GetNiagaraSystemMessageLogName(System));
-	NiagaraMessageLog = MessageLogModule.CreateLogListingWidget(NiagaraMessageLogListing.ToSharedRef());
+	NiagaraMessageLogViewModel = MakeShared<FNiagaraMessageLogViewModel>(GetNiagaraSystemMessageLogName(System), MessageLogGuid, NiagaraMessageLog);
 
 	SystemViewModel->OnEmitterHandleViewModelsChanged().AddSP(this, &FNiagaraSystemToolkit::OnRefresh);
 	SystemViewModel->OnSelectedEmitterHandlesChanged().AddSP(this, &FNiagaraSystemToolkit::OnRefresh);
@@ -473,7 +369,6 @@ void FNiagaraSystemToolkit::InitializeInternal(const EToolkitMode::Type Mode, co
 	SetupCommands();
 	ExtendToolbar();
 	RegenerateMenusAndToolbars();
-	UpdateMessageLog();
 
 	bChangesDiscarded = false;
 }
@@ -849,11 +744,6 @@ void FNiagaraSystemToolkit::ResetSimulation()
 	SystemViewModel->ResetSystem(FNiagaraSystemViewModel::ETimeResetMode::AllowResetTime, FNiagaraSystemViewModel::EMultiResetMode::AllowResetAllInstances, FNiagaraSystemViewModel::EReinitMode::ResetSystem);
 }
 
-void FNiagaraSystemToolkit::OnVMSystemCompiled()
-{
-	UpdateMessageLog();
-}
-
 void FNiagaraSystemToolkit::ExtendToolbar()
 {
 	struct Local
diff --git a/Engine/Plugins/FX/Niagara/Source/NiagaraEditor/Private/Toolkits/NiagaraSystemToolkit.h b/Engine/Plugins/FX/Niagara/Source/NiagaraEditor/Private/Toolkits/NiagaraSystemToolkit.h
--- a/Engine/Plugins/FX/Niagara/Source/NiagaraEditor/Private/Toolkits/NiagaraSystemToolkit.h
+++ b/Engine/Plugins/FX/Niagara/Source/NiagaraEditor/Private/Toolkits/NiagaraSystemToolkit.h
@@ -27,6 +27,7 @@ class UNiagaraSequence;
 struct FAssetData;
 class FMenuBuilder;
 class ISequencer;
+class FNiagaraMessageLogViewModel;
 
 /** Viewer/editor for a NiagaraSystem
 */
@@ -61,11 +62,6 @@ public:
 	//~ FGCObject interface
 	virtual void AddReferencedObjects(FReferenceCollector& Collector) override;
 
-	/**
-	* Updates list of Niagara messages in message log
-	*/
-	void UpdateMessageLog();
-
 	FSlateIcon GetCompileStatusImage() const;
 	FText GetCompileStatusTooltip() const;
 
@@ -89,7 +85,7 @@ protected:
 	virtual bool OnRequestClose() override;
 	
 private:
-	void InitializeInternal(const EToolkitMode::Type Mode, const TSharedPtr<class IToolkitHost>& InitToolkitHost);
+	void InitializeInternal(const EToolkitMode::Type Mode, const TSharedPtr<class IToolkitHost>& InitToolkitHost, const FGuid& MessageLogGuid);
 
 	void UpdateOriginalEmitter();
 	void UpdateExistingEmitters();
@@ -113,8 +109,6 @@ private:
 
 	void ResetSimulation();
 
-	void OnVMSystemCompiled();
-
 	void GetSequencerAddMenuContent(FMenuBuilder& MenuBuilder, TSharedRef<ISequencer> Sequencer);
 	TSharedRef<SWidget> CreateAddEmitterMenuContent();
 	TSharedRef<SWidget> GenerateCompileMenuContent();
@@ -159,8 +153,8 @@ private:
 	TSharedPtr<FNiagaraSystemViewModel> SystemViewModel;
 
 	/** Message log, with the log listing that it reflects */
+	TSharedPtr<FNiagaraMessageLogViewModel> NiagaraMessageLogViewModel;
 	TSharedPtr<class SWidget> NiagaraMessageLog;
-	TSharedPtr<class IMessageLogListing> NiagaraMessageLogListing;
 
 	/** The command list for this editor */
 	TSharedPtr<FUICommandList> EditorCommands;
diff --git a/Engine/Plugins/FX/Niagara/Source/NiagaraEditor/Private/ViewModels/NiagaraMessageLogViewModel.cpp b/Engine/Plugins/FX/Niagara/Source/NiagaraEditor/Private/ViewModels/NiagaraMessageLogViewModel.cpp
new file mode 100644
--- /dev/null
+++ b/Engine/Plugins/FX/Niagara/Source/NiagaraEditor/Private/ViewModels/NiagaraMessageLogViewModel.cpp
@@ -0,0 +1,60 @@
+// Copyright 1998-2019 Epic Games, Inc. All Rights Reserved.
+
+#include "NiagaraMessageLogViewModel.h"
+#include "NiagaraScriptSourceBase.h"
+#include "MessageLog/Public/MessageLogModule.h"
+
+FNiagaraMessageLogViewModel::FNiagaraMessageLogViewModel(const FName& InMessageLogName, const FGuid& InMessageLogGuidKey, TSharedPtr<class SWidget>& OutMessageLogWidget)
+	: MessageLogGuidKey(InMessageLogGuidKey)
+{
+	OnMessageManagerRequestRefreshHandle = FNiagaraMessageManager::Get()->GetOnRequestRefresh().AddRaw(this, &FNiagaraMessageLogViewModel::UpdateMessageLog);
+
+	FMessageLogModule& MessageLogModule = FModuleManager::LoadModuleChecked<FMessageLogModule>("MessageLog");
+
+	// Reuse any existing log, or create a new one (that is not held onto by the message log system)
+	auto CreateMessageLogListing = [&MessageLogModule](const FName& LogName)->TSharedRef<IMessageLogListing> {
+		FMessageLogInitializationOptions LogOptions;
+		// Show Pages so that user is never allowed to clear log messages
+		LogOptions.bShowPages = false;
+		LogOptions.bShowFilters = false;
+		LogOptions.bAllowClear = false;
+		LogOptions.MaxPageCount = 1;
+
+		if (MessageLogModule.IsRegisteredLogListing(LogName))
+		{
+			return MessageLogModule.GetLogListing(LogName);
+		}
+		else
+		{
+			return  MessageLogModule.CreateLogListing(LogName, LogOptions);
+		}
+	};
+
+	MessageLogListing = CreateMessageLogListing(InMessageLogName);
+	OutMessageLogWidget = MessageLogModule.CreateLogListingWidget(MessageLogListing.ToSharedRef());
+}
+
+FNiagaraMessageLogViewModel::~FNiagaraMessageLogViewModel()
+{
+	FNiagaraMessageManager::Get()->GetOnRequestRefresh().Remove(OnMessageManagerRequestRefreshHandle);
+	FNiagaraMessageManager::Get()->RefreshMessagesForAssetKey(MessageLogGuidKey);
+}
+
+void FNiagaraMessageLogViewModel::UpdateMessageLog(const FGuid& InMessageJobBatchAssetKey, const TArray<TSharedRef<const INiagaraMessage>> InNewMessages)
+{
+	if (MessageLogGuidKey == InMessageJobBatchAssetKey)
+	{
+		MessageLogListing->ClearMessages();
+		TArray<TSharedRef<FTokenizedMessage>> NewTokenizedMessages;
+		for (const TSharedRef<const INiagaraMessage> NewMessage : InNewMessages)
+		{
+			NewTokenizedMessages.Add(NewMessage->GenerateTokenizedMessage());
+		}
+		MessageLogListing->AddMessages(NewTokenizedMessages);
+	}
+}
+
+void FNiagaraMessageLogViewModel::SetMessageLogGuidKey(const FGuid& InViewedAssetObjectKey)
+{
+	MessageLogGuidKey = InViewedAssetObjectKey;
+}
diff --git a/Engine/Plugins/FX/Niagara/Source/NiagaraEditor/Private/ViewModels/NiagaraMessageLogViewModel.h b/Engine/Plugins/FX/Niagara/Source/NiagaraEditor/Private/ViewModels/NiagaraMessageLogViewModel.h
new file mode 100644
--- /dev/null
+++ b/Engine/Plugins/FX/Niagara/Source/NiagaraEditor/Private/ViewModels/NiagaraMessageLogViewModel.h
@@ -0,0 +1,27 @@
+// Copyright 1998-2019 Epic Games, Inc. All Rights Reserved.
+
+#pragma once
+
+#include "CoreMinimal.h"
+#include "Logging/TokenizedMessage.h"
+#include "Templates/SharedPointer.h"
+#include "Developer/MessageLog/Public/IMessageLogListing.h"
+#include "NiagaraMessageManager.h"
+
+class FNiagaraMessageLogViewModel : public TSharedFromThis<FNiagaraMessageLogViewModel>
+{
+public:
+	FNiagaraMessageLogViewModel(const FName& InMessageLogName, const FGuid& InMessageLogGuidKey, TSharedPtr<class SWidget>& OutMessageLogWidget );
+
+	~FNiagaraMessageLogViewModel();
+
+	void UpdateMessageLog(const FGuid& InMessageJobBatchAssetKey, const TArray<TSharedRef<const INiagaraMessage>> InNewMessages);
+
+	void SetMessageLogGuidKey(const FGuid& InMessageLogGuidKey);
+
+private:
+	TSharedPtr<IMessageLogListing> MessageLogListing;
+	FGuid MessageLogGuidKey;
+
+	FDelegateHandle OnMessageManagerRequestRefreshHandle;
+};
\ No newline at end of file
diff --git a/Engine/Plugins/FX/Niagara/Source/NiagaraEditor/Private/ViewModels/NiagaraScriptViewModel.cpp b/Engine/Plugins/FX/Niagara/Source/NiagaraEditor/Private/ViewModels/NiagaraScriptViewModel.cpp
--- a/Engine/Plugins/FX/Niagara/Source/NiagaraEditor/Private/ViewModels/NiagaraScriptViewModel.cpp
+++ b/Engine/Plugins/FX/Niagara/Source/NiagaraEditor/Private/ViewModels/NiagaraScriptViewModel.cpp
@@ -139,7 +139,6 @@ FNiagaraScriptViewModel::FNiagaraScriptViewModel(UNiagaraEmitter* InEmitter, FTe
 		CompileTypes[i].Key = Scripts[i]->GetUsage();
 		CompileTypes[i].Value = Scripts[i]->GetUsageId();
 	}
-	//UE_LOG(LogNiagaraEditor, Warning, TEXT("FNiagaraScriptViewModel %p"), this);
 }
 
 void FNiagaraScriptViewModel::OnVMScriptCompiled(UNiagaraScript* InScript)
@@ -473,13 +472,13 @@ UNiagaraScript* FNiagaraScriptViewModel::GetScript(ENiagaraScriptUsage InUsage,
 	return nullptr;
 }
 
-void FNiagaraScriptViewModel::CompileStandaloneScript()
+void FNiagaraScriptViewModel::CompileStandaloneScript(bool bForceCompile)
 {
 	if (Source.IsValid() && Source != nullptr && Scripts.Num() == 1 && Scripts[0].IsValid ())
 	{
 		if (Scripts[0]->IsStandaloneScript() && Scripts[0]->IsCompilable())
 		{
-			Scripts[0]->RequestCompile();
+			Scripts[0]->RequestCompile(bForceCompile);
 		}
 		else if (!Scripts[0]->IsCompilable())
 		{
diff --git a/Engine/Plugins/FX/Niagara/Source/NiagaraEditor/Private/ViewModels/NiagaraScriptViewModel.h b/Engine/Plugins/FX/Niagara/Source/NiagaraEditor/Private/ViewModels/NiagaraScriptViewModel.h
--- a/Engine/Plugins/FX/Niagara/Source/NiagaraEditor/Private/ViewModels/NiagaraScriptViewModel.h
+++ b/Engine/Plugins/FX/Niagara/Source/NiagaraEditor/Private/ViewModels/NiagaraScriptViewModel.h
@@ -54,7 +54,7 @@ public:
 		const TArray<UNiagaraScript*>& InCompileSources);
 
 	/** Compiles a script that isn't part of an emitter or System. */
-	void CompileStandaloneScript();
+	void CompileStandaloneScript(bool bForceCompile = false);
 
 	/** Get the latest status of this view-model's script compilation.*/
 	ENiagaraScriptCompileStatus GetLatestCompileStatus();
@@ -97,7 +97,7 @@ protected:
 	/** The script which provides the data for this view model. */
 	TArray<TWeakObjectPtr<UNiagaraScript>> Scripts;
 
-	void OnVMScriptCompiled(UNiagaraScript* InScript);
+	virtual void OnVMScriptCompiled(UNiagaraScript* InScript);
 
 	TWeakObjectPtr<UNiagaraScriptSource> Source;
 
diff --git a/Engine/Plugins/FX/Niagara/Source/NiagaraEditor/Private/ViewModels/NiagaraStandaloneScriptViewModel.cpp b/Engine/Plugins/FX/Niagara/Source/NiagaraEditor/Private/ViewModels/NiagaraStandaloneScriptViewModel.cpp
new file mode 100644
--- /dev/null
+++ b/Engine/Plugins/FX/Niagara/Source/NiagaraEditor/Private/ViewModels/NiagaraStandaloneScriptViewModel.cpp
@@ -0,0 +1,56 @@
+// Copyright 1998-2019 Epic Games, Inc. All Rights Reserved.
+
+#include "NiagaraStandaloneScriptViewModel.h"
+#include "NiagaraMessageManager.h"
+
+FNiagaraStandaloneScriptViewModel::FNiagaraStandaloneScriptViewModel(
+	UNiagaraScript* InScript
+	, FText DisplayName
+	, ENiagaraParameterEditMode InParameterEditMode
+	, TSharedPtr<FNiagaraMessageLogViewModel> InNiagaraMessageLogViewModel
+	, UNiagaraScript* InSourceScript
+	, const FGuid& InSourceScriptObjKey
+	)
+	: FNiagaraScriptViewModel(InScript, DisplayName, InParameterEditMode)
+	, NiagaraMessageLogViewModel(InNiagaraMessageLogViewModel)
+	, SourceScript(InSourceScript)
+	, ScriptMessageLogGuidKey(InSourceScriptObjKey)
+{
+	SendLastCompileMessageJobs(InSourceScript);
+}
+
+void FNiagaraStandaloneScriptViewModel::OnVMScriptCompiled(UNiagaraScript* InScript)
+{
+	FNiagaraScriptViewModel::OnVMScriptCompiled(InScript);
+	SendLastCompileMessageJobs(InScript);
+}
+
+void FNiagaraStandaloneScriptViewModel::SendLastCompileMessageJobs(UNiagaraScript* InScript)
+{
+	int32 ErrorCount = 0;
+	int32 WarningCount = 0;
+
+	TArray<TSharedPtr<const INiagaraMessageJob>> JobBatchToQueue;
+
+	TArray<FNiagaraCompileEvent>& CurrentCompileEvents = InScript->GetVMExecutableData().LastCompileEvents;
+
+	// Iterate from back to front to avoid reordering the events when they are queued
+	for (int i = CurrentCompileEvents.Num() - 1; i >= 0; --i)
+	{
+		const FNiagaraCompileEvent& CompileEvent = CurrentCompileEvents[i];
+		if (CompileEvent.Severity == FNiagaraCompileEventSeverity::Error)
+		{
+			ErrorCount++;
+		}
+		else if (CompileEvent.Severity == FNiagaraCompileEventSeverity::Warning)
+		{
+			WarningCount++;
+		}
+
+		JobBatchToQueue.Add(MakeShared<FNiagaraMessageJobCompileEvent>(CompileEvent, TWeakObjectPtr<UNiagaraScript>(InScript), TOptional<const FString>(), SourceScript->GetPathName()));
+	}
+
+	JobBatchToQueue.Insert(MakeShared<FNiagaraMessageJobPostCompileSummary>(ErrorCount, WarningCount, GetLatestCompileStatus(), FText::FromString("Script")), 0);
+	FNiagaraMessageManager::Get()->RefreshMessagesForAssetKeyAndMessageJobType(ScriptMessageLogGuidKey, ENiagaraMessageJobType::CompileEventMessageJob);
+	FNiagaraMessageManager::Get()->QueueMessageJobBatch(JobBatchToQueue, ScriptMessageLogGuidKey);
+}
diff --git a/Engine/Plugins/FX/Niagara/Source/NiagaraEditor/Private/ViewModels/NiagaraStandaloneScriptViewModel.h b/Engine/Plugins/FX/Niagara/Source/NiagaraEditor/Private/ViewModels/NiagaraStandaloneScriptViewModel.h
new file mode 100644
--- /dev/null
+++ b/Engine/Plugins/FX/Niagara/Source/NiagaraEditor/Private/ViewModels/NiagaraStandaloneScriptViewModel.h
@@ -0,0 +1,33 @@
+// Copyright 1998-2019 Epic Games, Inc. All Rights Reserved.
+
+#pragma once
+
+#include "Internationalization/Text.h"
+#include "NiagaraParameterEditMode.h"
+#include "NiagaraScriptViewModel.h"
+
+class FNiagaraMessageLogViewModel;
+class UNiagaraScript;
+
+class FNiagaraStandaloneScriptViewModel : public FNiagaraScriptViewModel
+{
+public:
+	FNiagaraStandaloneScriptViewModel(
+		UNiagaraScript* InScript
+		, FText DisplayName
+		, ENiagaraParameterEditMode InParameterEditMode
+		, TSharedPtr<FNiagaraMessageLogViewModel> InNiagaraMessageLogViewModel
+		, UNiagaraScript* InSourceScript
+		, const FGuid& InMessageLogGuidKey
+	);
+
+private:
+	virtual void OnVMScriptCompiled(UNiagaraScript* InScript) override;
+
+	/** Sends message jobs to FNiagaraMessageManager for all compile events from the last compile. */
+	void SendLastCompileMessageJobs(UNiagaraScript* InScript);
+
+	TSharedPtr<FNiagaraMessageLogViewModel> NiagaraMessageLogViewModel;
+	const UNiagaraScript* SourceScript;
+	const FGuid ScriptMessageLogGuidKey;
+};
diff --git a/Engine/Plugins/FX/Niagara/Source/NiagaraEditor/Private/ViewModels/NiagaraSystemViewModel.cpp b/Engine/Plugins/FX/Niagara/Source/NiagaraEditor/Private/ViewModels/NiagaraSystemViewModel.cpp
--- a/Engine/Plugins/FX/Niagara/Source/NiagaraEditor/Private/ViewModels/NiagaraSystemViewModel.cpp
+++ b/Engine/Plugins/FX/Niagara/Source/NiagaraEditor/Private/ViewModels/NiagaraSystemViewModel.cpp
@@ -35,6 +35,7 @@
 #include "MovieSceneFolder.h"
 #include "Channels/MovieSceneChannelProxy.h"
 #include "Framework/Commands/UICommandList.h"
+#include "NiagaraMessageManager.h"
 
 DECLARE_CYCLE_STAT(TEXT("Niagara - SystemViewModel - CompileSystem"), STAT_NiagaraEditor_SystemViewModel_CompileSystem, STATGROUP_NiagaraEditor);
 
@@ -48,7 +49,7 @@ FNiagaraSystemViewModelOptions::FNiagaraSystemViewModelOptions()
 {
 }
 
-FNiagaraSystemViewModel::FNiagaraSystemViewModel(UNiagaraSystem& InSystem, FNiagaraSystemViewModelOptions InOptions)
+FNiagaraSystemViewModel::FNiagaraSystemViewModel(UNiagaraSystem& InSystem, FNiagaraSystemViewModelOptions InOptions, TOptional<const FGuid> InMessageLogGuid)
 	: System(InSystem)
 	, PreviewComponent(nullptr)
 	, SystemInstance(nullptr)
@@ -68,6 +69,7 @@ FNiagaraSystemViewModel::FNiagaraSystemViewModel(UNiagaraSystem& InSystem, FNiag
 	, EditorSettings(GetMutableDefault<UNiagaraEditorSettings>())
 	, bResetRequestPending(false)
 	, bCompilePendingCompletion(false)
+	,SystemMessageLogGuidKey(InMessageLogGuid)
 {
 	SetupPreviewComponentAndInstance();
 	SetupSequencer();
@@ -75,6 +77,7 @@ FNiagaraSystemViewModel::FNiagaraSystemViewModel(UNiagaraSystem& InSystem, FNiag
 	GEditor->RegisterForUndo(this);
 	RegisteredHandle = RegisterViewModelWithMap(&InSystem, this);
 	AddSystemEventHandlers();
+	SendLastCompileMessageJobs();
 }
 
 void FNiagaraSystemViewModel::DumpToText(FString& ExportText)
@@ -168,7 +171,7 @@ void FNiagaraSystemViewModel::CompileSystem(bool bForce)
 	bCompilePendingCompletion = true;
 }
 
-ENiagaraScriptCompileStatus FNiagaraSystemViewModel::GetLatestCompileStatus()
+ENiagaraScriptCompileStatus FNiagaraSystemViewModel::GetLatestCompileStatus() const
 {
 	check(SystemScriptViewModel.IsValid());
 	return SystemScriptViewModel->GetLatestCompileStatus();
@@ -505,6 +508,13 @@ void FNiagaraSystemViewModel::PostUndo(bool bSuccess)
 
 void FNiagaraSystemViewModel::Tick(float DeltaTime)
 {
+	if (bCompilePendingCompletion && System.HasOutstandingCompilationRequests() == false)
+	{
+		bCompilePendingCompletion = false;
+		OnSystemCompiled().Broadcast();
+		SendLastCompileMessageJobs();
+	}
+
 	if (bForceAutoCompileOnce || (GetDefault<UNiagaraEditorSettings>()->GetAutoCompile() && bCanAutoCompile))
 	{
 		bool bRecompile = false;
@@ -529,12 +539,6 @@ void FNiagaraSystemViewModel::Tick(float DeltaTime)
 
 		if (System.HasOutstandingCompilationRequests() == false)
 		{
-			if (bCompilePendingCompletion)
-			{
-				bCompilePendingCompletion = false;
-				OnSystemCompiled().Broadcast();
-			}
-
 			if (bRecompile || bForceAutoCompileOnce)
 			{
 				CompileSystem(false);
@@ -619,6 +623,69 @@ TStatId FNiagaraSystemViewModel::GetStatId() const
 	RETURN_QUICK_DECLARE_CYCLE_STAT(FNiagaraSystemViewModel, STATGROUP_Tickables);
 }
 
+void FNiagaraSystemViewModel::SendLastCompileMessageJobs() const
+{
+	if (SystemMessageLogGuidKey.IsSet() == false)
+	{
+		return;
+	}
+
+	struct FNiagaraScriptAndOwningScriptNameString
+	{
+		FNiagaraScriptAndOwningScriptNameString(const UNiagaraScript* InScript, const FString& InOwningScriptNameString)
+			: Script(InScript)
+			, OwningScriptNameString(InOwningScriptNameString)
+		{
+		}
+
+		const UNiagaraScript* Script;
+		const FString OwningScriptNameString;
+	};
+
+	int32 ErrorCount = 0;
+	int32 WarningCount = 0;
+
+	TArray<FNiagaraScriptAndOwningScriptNameString> ScriptsToGetCompileEventsFrom;
+	ScriptsToGetCompileEventsFrom.Add(FNiagaraScriptAndOwningScriptNameString(System.GetSystemSpawnScript(), System.GetName()));
+	ScriptsToGetCompileEventsFrom.Add(FNiagaraScriptAndOwningScriptNameString(System.GetSystemUpdateScript(), System.GetName()));
+	const TArray<FNiagaraEmitterHandle> EmitterHandles = System.GetEmitterHandles();
+	for (const FNiagaraEmitterHandle& Handle : EmitterHandles)
+	{
+		const UNiagaraEmitter* EmitterInSystem = Handle.GetInstance();
+		TArray<UNiagaraScript*> EmitterScripts;
+		EmitterInSystem->GetScripts(EmitterScripts);
+		for (UNiagaraScript* EmitterScript : EmitterScripts)
+		{
+			ScriptsToGetCompileEventsFrom.Add(FNiagaraScriptAndOwningScriptNameString(EmitterScript, EmitterInSystem->GetUniqueEmitterName()));
+		}
+	}
+
+	TArray<TSharedPtr<const INiagaraMessageJob>> JobBatchToQueue;
+	// Iterate from back to front to avoid reordering the events when they are queued
+	for (int i = ScriptsToGetCompileEventsFrom.Num()-1; i >=0; --i)
+	{ 
+		const FNiagaraScriptAndOwningScriptNameString& ScriptInfo = ScriptsToGetCompileEventsFrom[i];
+		const TArray<FNiagaraCompileEvent>& CurrentCompileEvents = ScriptInfo.Script->GetVMExecutableData().LastCompileEvents;
+		for (int j = CurrentCompileEvents.Num() - 1; j >= 0; --j)
+		{
+			const FNiagaraCompileEvent& CompileEvent = CurrentCompileEvents[j];
+			if (CompileEvent.Severity == FNiagaraCompileEventSeverity::Error)
+			{
+				ErrorCount++;
+			}
+			else if (CompileEvent.Severity == FNiagaraCompileEventSeverity::Warning)
+			{
+				WarningCount++;
+			}
+
+			JobBatchToQueue.Add(MakeShared<FNiagaraMessageJobCompileEvent>(CompileEvent, TWeakObjectPtr<UNiagaraScript>(ScriptInfo.Script), ScriptInfo.OwningScriptNameString));
+		}
+	}
+	JobBatchToQueue.Insert(MakeShared<FNiagaraMessageJobPostCompileSummary>(ErrorCount, WarningCount, GetLatestCompileStatus(), FText::FromString("System")), 0);
+	FNiagaraMessageManager::Get()->RefreshMessagesForAssetKeyAndMessageJobType(SystemMessageLogGuidKey.GetValue(), ENiagaraMessageJobType::CompileEventMessageJob);
+	FNiagaraMessageManager::Get()->QueueMessageJobBatch(JobBatchToQueue, SystemMessageLogGuidKey.GetValue());
+}
+
 void FNiagaraSystemViewModel::SetupPreviewComponentAndInstance()
 {
 	if (bCanSimulate)
@@ -1181,6 +1248,7 @@ void FNiagaraSystemViewModel::EmitterPropertyChanged(TSharedRef<FNiagaraEmitterH
 
 void FNiagaraSystemViewModel::ScriptCompiled()
 {
+	bCompilePendingCompletion = true;
 	//ReInitializeSystemInstances();
 }
 
diff --git a/Engine/Plugins/FX/Niagara/Source/NiagaraEditor/Public/NiagaraMessageManager.h b/Engine/Plugins/FX/Niagara/Source/NiagaraEditor/Public/NiagaraMessageManager.h
--- a/Engine/Plugins/FX/Niagara/Source/NiagaraEditor/Public/NiagaraMessageManager.h
+++ b/Engine/Plugins/FX/Niagara/Source/NiagaraEditor/Public/NiagaraMessageManager.h
@@ -17,11 +17,16 @@ enum class ENiagaraMessageType : uint8
 	NeedRecompileMessage
 };
 
+enum class ENiagaraMessageJobType : uint8 {
+	None = 0,
+	CompileEventMessageJob
+};
+
 //Struct for passing around script asset info from compile event message job to message types
 struct FNiagaraScriptNameAndAssetPath
 {
 public:
-	FNiagaraScriptNameAndAssetPath(const FString InScriptNameString, const FString InScriptAssetPathString)
+	FNiagaraScriptNameAndAssetPath(const FString& InScriptNameString, const FString& InScriptAssetPathString)
 		: ScriptNameString(InScriptNameString)
 		, ScriptAssetPathString(InScriptAssetPathString)
 	{};
@@ -36,19 +41,11 @@ public:
 class INiagaraMessage 
 {
 public:
-	INiagaraMessage(const ENiagaraMessageType InMessageType)
-		: MessageType(InMessageType)
-	{
-	};
-
 	virtual TSharedRef<FTokenizedMessage> GenerateTokenizedMessage() const = 0;
 
-	const ENiagaraMessageType GetMessageType() {return MessageType;}
+	virtual const ENiagaraMessageType GetMessageType() const = 0;
 
 	virtual ~INiagaraMessage() {};
-
-private:
-	const ENiagaraMessageType MessageType;
 };
 
 /** Interface for job supplied to FNiagaraMessageManager to generated FNiagaraMessages.
@@ -57,9 +54,9 @@ private:
 class INiagaraMessageJob
 {
 public:
-	
+	virtual TSharedRef<const INiagaraMessage> GenerateNiagaraMessage() const = 0;
 
-	virtual const TSharedPtr<INiagaraMessage> GenerateNiagaraMessage() const = 0;
+	virtual const ENiagaraMessageJobType GetMessageJobType() const = 0;
 
 	virtual ~INiagaraMessageJob() {};
 };
@@ -67,25 +64,54 @@ public:
 class FNiagaraMessageJobCompileEvent : public INiagaraMessageJob
 {
 public:
-	FNiagaraMessageJobCompileEvent( const FNiagaraCompileEvent& InCompileEvent, const TWeakObjectPtr<UNiagaraScript>& InOriginatingScriptWeakObjPtr, const bool bInFromScriptToolkit);
+	FNiagaraMessageJobCompileEvent(
+		  const FNiagaraCompileEvent& InCompileEvent
+		, const TWeakObjectPtr<UNiagaraScript>& InOriginatingScriptWeakObjPtr
+		, const TOptional<const FString>& InOwningScriptNameString = TOptional<const FString>()
+		, const TOptional<const FString>& InSourceScriptAssetPath = TOptional<const FString>()
+	);
+
+	virtual TSharedRef<const INiagaraMessage> GenerateNiagaraMessage() const override;
 
-	virtual const TSharedPtr<INiagaraMessage> GenerateNiagaraMessage() const override;
+	virtual const ENiagaraMessageJobType GetMessageJobType() const override { return ENiagaraMessageJobType::CompileEventMessageJob; };
 
 private:
-	//hidden private ctor
-	FNiagaraMessageJobCompileEvent();
 
-	const bool RecursiveGetScriptNamesAndAssetPathsFromContextStack(
+	bool RecursiveGetScriptNamesAndAssetPathsFromContextStack(
 		  TArray<FGuid>& InContextStackNodeGuids
 		, const UNiagaraGraph* InGraphToSearch
 		, TArray<FNiagaraScriptNameAndAssetPath>& OutContextScriptNamesAndAssetPaths
-		, TOptional<const FString>& OutEmitterScriptName
+		, TOptional<const FString>& OutEmitterName
 		, TOptional<const FText>& OutFailureReason
 	) const;
 
 	const FNiagaraCompileEvent CompileEvent;
 	const TWeakObjectPtr<UNiagaraScript> OriginatingScriptWeakObjPtr;
-	const bool bFromScriptToolkit;
+	TOptional<const FString> OwningScriptNameString;
+	TOptional<const FString> SourceScriptAssetPath;
+};
+
+class FNiagaraMessageJobPostCompileSummary : public INiagaraMessageJob
+{
+public:
+	FNiagaraMessageJobPostCompileSummary(const int32& InNumCompileErrors, const int32& InNumCompileWarnings, const ENiagaraScriptCompileStatus& InScriptCompileStatus, const FText& InCompiledObjectNameText)
+		: NumCompileErrors(InNumCompileErrors)
+		, NumCompileWarnings(InNumCompileWarnings)
+		, ScriptCompileStatus(InScriptCompileStatus)
+		, CompiledObjectNameText(InCompiledObjectNameText)
+	{
+	};
+
+	virtual TSharedRef<const INiagaraMessage> GenerateNiagaraMessage() const override;
+
+	// we alias the messagejob type as post compile summary is always generated with compile events
+	virtual const ENiagaraMessageJobType GetMessageJobType() const override { return ENiagaraMessageJobType::CompileEventMessageJob; };
+
+private:
+	const int32 NumCompileErrors;
+	const int32 NumCompileWarnings;
+	const ENiagaraScriptCompileStatus ScriptCompileStatus;
+	const FText CompiledObjectNameText;
 };
 
 class FNiagaraMessageCompileEvent : public INiagaraMessage
@@ -94,31 +120,57 @@ public:
 	FNiagaraMessageCompileEvent(
 		  const FNiagaraCompileEvent& InCompileEvent
 		, TArray<FNiagaraScriptNameAndAssetPath>& InContextScriptNamesAndAssetPaths
-		, TOptional<const FString>& InEmitterScriptName
+		, TOptional<const FText>& InOwningScriptNameAndUsageText
 		, TOptional<const FNiagaraScriptNameAndAssetPath>& InCompiledScriptNameAndAssetPath
 	);
 
 	virtual TSharedRef<FTokenizedMessage> GenerateTokenizedMessage() const override;
 	//@todo(message manager) make stack specific message type generator here
 
+	virtual const ENiagaraMessageType GetMessageType() const override { return ENiagaraMessageType::CompileEventMessage; };
+
 private:
 	const FNiagaraCompileEvent CompileEvent;
 	const TArray<FNiagaraScriptNameAndAssetPath> ContextScriptNamesAndAssetPaths;
-	const TOptional<const FString> EmitterScriptName;
+	const TOptional<const FText> OwningScriptNameAndUsageText;
 	const TOptional<const FNiagaraScriptNameAndAssetPath> CompiledScriptNameAndAssetPath;
 };
 
 class FNiagaraMessageNeedRecompile : public INiagaraMessage
 {
 public:
-	FNiagaraMessageNeedRecompile(const FText& InNeedRecompileMessage);
+	FNiagaraMessageNeedRecompile(const FText& InNeedRecompileMessage)
+		: NeedRecompileMessage(InNeedRecompileMessage)
+	{
+	};
 
 	virtual TSharedRef<FTokenizedMessage> GenerateTokenizedMessage() const override;
 
+	virtual const ENiagaraMessageType GetMessageType() const override { return ENiagaraMessageType::NeedRecompileMessage; };
+
 private:
 	const FText NeedRecompileMessage;
 };
 
+class FNiagaraMessagePostCompileSummary : public INiagaraMessage
+{
+public: 
+	FNiagaraMessagePostCompileSummary(const FText& InPostCompileSummaryText, const EMessageSeverity::Type& InMessageSeverity)
+		: PostCompileSummaryText(InPostCompileSummaryText)
+		, MessageSeverity(InMessageSeverity)
+	{
+	};
+
+	virtual TSharedRef<FTokenizedMessage> GenerateTokenizedMessage() const override;
+	
+	// we alias the message type as post compile summary is always generated with compile events
+	virtual const ENiagaraMessageType GetMessageType() const override { return  ENiagaraMessageType::CompileEventMessage; };
+
+private:
+	const FText PostCompileSummaryText;
+	const EMessageSeverity::Type MessageSeverity;
+};
+
 //Extension of message token to allow opening the asset editor when clicking on the linked asset name.
 class FNiagaraCompileEventToken : public IMessageToken
 {
@@ -167,16 +219,58 @@ private:
 	const TOptional<const FGuid> PinGUID;
 };
 
-class FNiagaraMessageManager
+class FNiagaraMessageManager : FTickableEditorObject
 {
+	struct FNiagaraMessageJobBatch
+	{
+	public:
+		FNiagaraMessageJobBatch(TArray<TSharedPtr<const INiagaraMessageJob>> InMessageJobs, const FGuid& InMessageJobsAssetKey)
+			: MessageJobs(InMessageJobs)
+			, MessageJobsAssetKey(InMessageJobsAssetKey)
+		{
+		};
+
+		TArray<TSharedPtr<const INiagaraMessageJob>> MessageJobs;
+		FGuid MessageJobsAssetKey;
+
+	};
+
 public:
+	DECLARE_MULTICAST_DELEGATE_TwoParams(FOnRequestRefresh, const FGuid& /*MessageJobBatchAssetKey*/, const TArray<TSharedRef<const INiagaraMessage>> /*NewMessages*/)
+
 	static FNiagaraMessageManager* Get();
 
-	static TSharedPtr<INiagaraMessage> QueueMessageJob(TSharedRef<INiagaraMessageJob> InMessageJob);
+	void QueueMessageJob(TSharedPtr<const INiagaraMessageJob> InMessageJob, const FGuid& InMessageJobAssetKey);
+
+	void QueueMessageJobBatch(TArray<TSharedPtr<const INiagaraMessageJob>> InMessageJobBatch, const FGuid& InMessageJobBatchAssetKey);
+
+	void RefreshMessagesForAssetKey(const FGuid& InAssetKey);
+
+	void RefreshMessagesForAssetKeyAndMessageJobType(const FGuid& InAssetKey, const ENiagaraMessageJobType& InMessageJobType);
+
+	static const TOptional<const FString> GetStringForScriptUsageInStack(const ENiagaraScriptUsage InScriptUsage);
+
+	const TArray<TSharedRef<const INiagaraMessage>> GetMessagesForAssetKey(const FGuid& InAssetKey) const;
+
+	FOnRequestRefresh& GetOnRequestRefresh() { return OnRequestRefresh; };
+
+	//~ Begin FTickableEditorObject Interface.
+	virtual void Tick(float DeltaSeconds) override;
+
+	virtual ETickableTickType GetTickableTickType() const override { return ETickableTickType::Always; };
+
+	virtual TStatId GetStatId() const override { RETURN_QUICK_DECLARE_CYCLE_STAT(FNiagaraMessageManager, STATGROUP_Tickables); };
+	//~ End FTickableEditorObject Interface
 
 private:
 	FNiagaraMessageManager();
 
-	static TQueue<TSharedPtr<INiagaraMessageJob>> JobQueue;
+	FOnRequestRefresh OnRequestRefresh;
+
+	void DoMessageJobsInQueueTick();
+	static const double MaxJobWorkTime;
+	TArray<FNiagaraMessageJobBatch> MessageJobBatchArr;
+	TArray<TSharedRef<const INiagaraMessage>> GeneratedMessagesForCurrentMessageJobBatch;
+	TMap<FGuid, TMap<const ENiagaraMessageJobType, TArray<TSharedRef<const INiagaraMessage>>>> ObjectToTypeMappedMessagesMap;
 	static FNiagaraMessageManager* Singleton;
 };
diff --git a/Engine/Plugins/FX/Niagara/Source/NiagaraEditor/Public/ViewModels/NiagaraSystemViewModel.h b/Engine/Plugins/FX/Niagara/Source/NiagaraEditor/Public/ViewModels/NiagaraSystemViewModel.h
--- a/Engine/Plugins/FX/Niagara/Source/NiagaraEditor/Public/ViewModels/NiagaraSystemViewModel.h
+++ b/Engine/Plugins/FX/Niagara/Source/NiagaraEditor/Public/ViewModels/NiagaraSystemViewModel.h
@@ -133,7 +133,7 @@ public:
 		ResetSystem,
 	};
 	/** Creates a new view model with the supplied System and System instance. */
-	FNiagaraSystemViewModel(UNiagaraSystem& InSystem, FNiagaraSystemViewModelOptions InOptions);
+	FNiagaraSystemViewModel(UNiagaraSystem& InSystem, FNiagaraSystemViewModelOptions InOptions, TOptional<const FGuid> InMessageLogGuid = TOptional<const FGuid>());
 
 	~FNiagaraSystemViewModel();
 
@@ -226,7 +226,7 @@ public:
 	void CompileSystem(bool bForce);
 
 	/* Get the latest status of this view-model's script compilation.*/
-	ENiagaraScriptCompileStatus GetLatestCompileStatus();
+	ENiagaraScriptCompileStatus GetLatestCompileStatus() const;
 
 	/** Gets the ids for the currently selected emitter handles. */
 	const TArray<FGuid>& GetSelectedEmitterHandleIds();
@@ -298,6 +298,10 @@ public:
 	const TArray<FNiagaraStackModuleData>& GetStackModuleDataForEmitter(TSharedRef<FNiagaraEmitterViewModel> EmitterViewModel);
 
 private:
+
+	/** Sends message jobs to FNiagaraMessageManager for all compile events from the last compile. */
+	void SendLastCompileMessageJobs() const;
+
 	/** Sets up the preview component and System instance. */
 	void SetupPreviewComponentAndInstance();
 
@@ -536,4 +540,7 @@ private:
 
 	/** An array of emitter handle ids which need their sequencer tracks refreshed next frame. */
 	TArray<FGuid> EmitterIdsRequiringSequencerTrackUpdate;
+
+	/** GUID used when sending message jobs to FNiagaraMessageManager for notifying the FNiagaraMessageLogViewModel with the same GUID key */
+	const TOptional<const FGuid> SystemMessageLogGuidKey;
 };
\ No newline at end of file
