- Niagara fix GetTriCoordVertices VM binding and reading off the end of the array
- Niagara fix GetVertexUV reading off end of the array

diff --git a/Engine/Plugins/FX/Niagara/Source/Niagara/Private/NDISkeletalMesh_TriangleSampling.cpp b/Engine/Plugins/FX/Niagara/Source/Niagara/Private/NDISkeletalMesh_TriangleSampling.cpp
--- a/Engine/Plugins/FX/Niagara/Source/Niagara/Private/NDISkeletalMesh_TriangleSampling.cpp
+++ b/Engine/Plugins/FX/Niagara/Source/Niagara/Private/NDISkeletalMesh_TriangleSampling.cpp
@@ -261,7 +261,7 @@ void UNiagaraDataInterfaceSkeletalMesh::BindTriangleSamplingFunction(const FVMEx
 	}
 	else if (BindingInfo.Name == FSkeletalMeshInterfaceHelper::GetTriCoordVerticesName)
 	{
-		check(BindingInfo.GetNumInputs() == 5 && BindingInfo.GetNumOutputs() == 3);
+		check(BindingInfo.GetNumInputs() == 2 && BindingInfo.GetNumOutputs() == 3);
 		TSkinningModeBinder<NDI_FUNC_BINDER(UNiagaraDataInterfaceSkeletalMesh, GetTriCoordVertices)>::Bind(this, BindingInfo, InstanceData, OutFunc);
 	}
 
@@ -1069,7 +1069,7 @@ void UNiagaraDataInterfaceSkeletalMesh::GetTriCoordVertices(FVectorVMContext& Co
 	for (int32 i = 0; i < Context.NumInstances; ++i)
 	{
 		const int32 Tri = FMath::Clamp(TriParam.GetAndAdvance(), 0, TriMax);
-		SkinningHandler.GetTrianlgeIndices(Accessor, TriParam.GetAndAdvance(), Idx0, Idx1, Idx2);
+		SkinningHandler.GetTrianlgeIndices(Accessor, Tri, Idx0, Idx1, Idx2);
 		*OutV0.GetDestAndAdvance() = Idx0;
 		*OutV1.GetDestAndAdvance() = Idx1;
 		*OutV2.GetDestAndAdvance() = Idx2;
diff --git a/Engine/Plugins/FX/Niagara/Source/Niagara/Private/NDISkeletalMesh_VertexSampling.cpp b/Engine/Plugins/FX/Niagara/Source/Niagara/Private/NDISkeletalMesh_VertexSampling.cpp
--- a/Engine/Plugins/FX/Niagara/Source/Niagara/Private/NDISkeletalMesh_VertexSampling.cpp
+++ b/Engine/Plugins/FX/Niagara/Source/Niagara/Private/NDISkeletalMesh_VertexSampling.cpp
@@ -471,7 +471,7 @@ void UNiagaraDataInterfaceSkeletalMesh::GetVertexUV(FVectorVMContext& Context)
 
 	FMultiSizeIndexContainer& Indices = LODData.MultiSizeIndexContainer;
 	FRawStaticIndexBuffer16or32Interface* IndexBuffer = Indices.GetIndexBuffer();
-	int32 VertMax = LODData.GetNumVertices();
+	const int32 VertMax = LODData.GetNumVertices() - 1;
 	for (int32 i = 0; i < Context.NumInstances; ++i)
 	{
 		int32 Vert = VertParam.GetAndAdvance();
