Fix non-unity build missing includes for FNiagaraMessageManager and FNiagaraMessageLogViewModel

diff --git a/Engine/Plugins/FX/Niagara/Source/NiagaraEditor/Private/ViewModels/NiagaraMessageLogViewModel.cpp b/Engine/Plugins/FX/Niagara/Source/NiagaraEditor/Private/ViewModels/NiagaraMessageLogViewModel.cpp
--- a/Engine/Plugins/FX/Niagara/Source/NiagaraEditor/Private/ViewModels/NiagaraMessageLogViewModel.cpp
+++ b/Engine/Plugins/FX/Niagara/Source/NiagaraEditor/Private/ViewModels/NiagaraMessageLogViewModel.cpp
@@ -2,6 +2,7 @@
 
 #include "NiagaraMessageLogViewModel.h"
 #include "NiagaraScriptSourceBase.h"
+#include "Modules/ModuleManager.h"
 #include "MessageLog/Public/MessageLogModule.h"
 
 FNiagaraMessageLogViewModel::FNiagaraMessageLogViewModel(const FName& InMessageLogName, const FGuid& InMessageLogGuidKey, TSharedPtr<class SWidget>& OutMessageLogWidget)
diff --git a/Engine/Plugins/FX/Niagara/Source/NiagaraEditor/Public/NiagaraMessageManager.h b/Engine/Plugins/FX/Niagara/Source/NiagaraEditor/Public/NiagaraMessageManager.h
--- a/Engine/Plugins/FX/Niagara/Source/NiagaraEditor/Public/NiagaraMessageManager.h
+++ b/Engine/Plugins/FX/Niagara/Source/NiagaraEditor/Public/NiagaraMessageManager.h
@@ -6,6 +6,7 @@
 #include "NiagaraShared.h"
 #include "Logging/TokenizedMessage.h"
 #include "NiagaraGraph.h"
+#include "TickableEditorObject.h"
 
 class FNiagaraScriptToolkit;
 
