Fixed a bug where some locked static switch metadata attributes became editable.

diff --git a/Engine/Plugins/FX/Niagara/Source/Niagara/Public/NiagaraTypes.h b/Engine/Plugins/FX/Niagara/Source/Niagara/Public/NiagaraTypes.h
--- a/Engine/Plugins/FX/Niagara/Source/Niagara/Public/NiagaraTypes.h
+++ b/Engine/Plugins/FX/Niagara/Source/Niagara/Public/NiagaraTypes.h
@@ -354,7 +354,7 @@ public:
 	UPROPERTY(EditAnywhere, Category = "Variable", DisplayName = "Property Metadata", meta = (ToolTip = "Property Metadata"))
 	TMap<FName, FString> PropertyMetaData;
 
-	UPROPERTY()
+	UPROPERTY(AdvancedDisplay, VisibleAnywhere, Category = "Variable")
 	bool bIsStaticSwitch;
 };
 
diff --git a/Engine/Plugins/FX/Niagara/Source/NiagaraEditor/Private/NiagaraScriptVariable.h b/Engine/Plugins/FX/Niagara/Source/NiagaraEditor/Private/NiagaraScriptVariable.h
--- a/Engine/Plugins/FX/Niagara/Source/NiagaraEditor/Private/NiagaraScriptVariable.h
+++ b/Engine/Plugins/FX/Niagara/Source/NiagaraEditor/Private/NiagaraScriptVariable.h
@@ -14,6 +14,6 @@ public:
 	UPROPERTY()
 	FNiagaraVariable Variable;
 
-	UPROPERTY(EditAnywhere, Category = Metadata)
+	UPROPERTY(EditAnywhere, Category = Metadata, meta=(NoResetToDefault, ShowOnlyInnerProperties))
 	FNiagaraVariableMetaData Metadata;
 };
\ No newline at end of file
