Fixed a bug where changes to a static switch node were not properly propagated in the ui

diff --git a/Engine/Plugins/FX/Niagara/Source/NiagaraEditor/Private/Customizations/NiagaraStaticSwitchNodeDetails.cpp b/Engine/Plugins/FX/Niagara/Source/NiagaraEditor/Private/Customizations/NiagaraStaticSwitchNodeDetails.cpp
--- a/Engine/Plugins/FX/Niagara/Source/NiagaraEditor/Private/Customizations/NiagaraStaticSwitchNodeDetails.cpp
+++ b/Engine/Plugins/FX/Niagara/Source/NiagaraEditor/Private/Customizations/NiagaraStaticSwitchNodeDetails.cpp
@@ -11,6 +11,7 @@
 #include "DetailWidgetRow.h"
 #include "Widgets/Input/SComboBox.h"
 #include "Widgets/Input/SNumericEntryBox.h"
+#include "Widgets/Input/SEditableTextBox.h"
 
 #define LOCTEXT_NAMESPACE "NiagaraStaticSwitchNodeDetails"
 
@@ -46,7 +47,34 @@ void FNiagaraStaticSwitchNodeDetails::CustomizeDetails(IDetailLayoutBuilder& Det
 		UpdateSelectionFromNode();
 		
 		IDetailCategoryBuilder& CategoryBuilder = DetailBuilder.EditCategory(SwitchCategoryName);		
-		FDetailWidgetRow& DropdownWidget = CategoryBuilder.AddCustomRow(LOCTEXT("NiagaraSwitchNodeDetailsFilterText", "Input parameter type"));
+		FDetailWidgetRow& NameWidget = CategoryBuilder.AddCustomRow(LOCTEXT("NiagaraSwitchNodeNameFilterText", "Input parameter name"));
+
+		NameWidget
+		.NameContent()
+		[
+			SNew(SBox)
+			.Padding(FMargin(0.0f, 2.0f))
+			[
+				SNew(STextBlock)
+				.TextStyle(FNiagaraEditorStyle::Get(), "NiagaraEditor.ParameterText")
+				.Text(LOCTEXT("NiagaraSwitchNodeNameText", "Input parameter name"))
+			]
+		]
+		.ValueContent()
+		[
+			SNew(SBox)
+			.Padding(FMargin(0.0f, 2.0f))
+			[
+				SNew(SEditableTextBox)
+				.Text(this, &FNiagaraStaticSwitchNodeDetails::GetParameterNameText)
+				.ToolTipText(LOCTEXT("NiagaraSwitchNodeNameTooltip", "This is the name of the parameter that is exposed to the user calling this function graph."))
+				.OnTextCommitted(this, &FNiagaraStaticSwitchNodeDetails::OnParameterNameCommited)
+				.SelectAllTextWhenFocused(true)
+				.RevertTextOnEscape(true)
+			]
+		];
+
+		FDetailWidgetRow& DropdownWidget = CategoryBuilder.AddCustomRow(LOCTEXT("NiagaraSwitchNodeTypeFilterText", "Input parameter type"));
 
 		DropdownWidget
 		.NameContent()
@@ -124,6 +152,7 @@ void FNiagaraStaticSwitchNodeDetails::OnSelectionChanged(TSharedPtr<SwitchDropdo
 		return;
 	}
 	
+	FNiagaraTypeDefinition OldType = Node->GetInputType();
 	if (SelectedDropdownItem == DropdownOptions[0])
 	{
 		Node->SwitchTypeData.SwitchType = ENiagaraStaticSwitchType::Bool;
@@ -137,7 +166,7 @@ void FNiagaraStaticSwitchNodeDetails::OnSelectionChanged(TSharedPtr<SwitchDropdo
 		Node->SwitchTypeData.SwitchType = ENiagaraStaticSwitchType::Enum;
 		Node->SwitchTypeData.Enum = SelectedDropdownItem->Enum;
 	}
-	Node->RefreshFromExternalChanges();
+	Node->OnSwitchParameterTypeChanged(OldType);
 }
 
 FText FNiagaraStaticSwitchNodeDetails::GetDropdownItemLabel() const
@@ -196,4 +225,17 @@ void FNiagaraStaticSwitchNodeDetails::IntOptionValueCommitted(int32 Value, EText
 	}
 }
 
+FText FNiagaraStaticSwitchNodeDetails::GetParameterNameText() const
+{
+	return Node.IsValid() ? FText::FromName(Node->InputParameterName) : FText();
+}
+
+void FNiagaraStaticSwitchNodeDetails::OnParameterNameCommited(const FText& InText, ETextCommit::Type InCommitType)
+{
+	if (Node.IsValid())
+	{
+		Node->ChangeSwitchParameterName(FName(*InText.ToString()));
+	}
+}
+
 #undef LOCTEXT_NAMESPACE
diff --git a/Engine/Plugins/FX/Niagara/Source/NiagaraEditor/Private/Customizations/NiagaraStaticSwitchNodeDetails.h b/Engine/Plugins/FX/Niagara/Source/NiagaraEditor/Private/Customizations/NiagaraStaticSwitchNodeDetails.h
--- a/Engine/Plugins/FX/Niagara/Source/NiagaraEditor/Private/Customizations/NiagaraStaticSwitchNodeDetails.h
+++ b/Engine/Plugins/FX/Niagara/Source/NiagaraEditor/Private/Customizations/NiagaraStaticSwitchNodeDetails.h
@@ -44,6 +44,10 @@ private:
 	TOptional<int32> GetIntOptionValue() const;
 	void IntOptionValueCommitted(int32 Value, ETextCommit::Type CommitInfo);
 
+	// parameter name option function
+	FText GetParameterNameText() const;
+	void OnParameterNameCommited(const FText& InText, ETextCommit::Type InCommitType);
+
 	TWeakObjectPtr<class UNiagaraNodeStaticSwitch> Node;
 	TArray<TSharedPtr<SwitchDropdownOption>> DropdownOptions;
 	TSharedPtr<SwitchDropdownOption> SelectedDropdownItem;
diff --git a/Engine/Plugins/FX/Niagara/Source/NiagaraEditor/Private/NiagaraNodeStaticSwitch.cpp b/Engine/Plugins/FX/Niagara/Source/NiagaraEditor/Private/NiagaraNodeStaticSwitch.cpp
--- a/Engine/Plugins/FX/Niagara/Source/NiagaraEditor/Private/NiagaraNodeStaticSwitch.cpp
+++ b/Engine/Plugins/FX/Niagara/Source/NiagaraEditor/Private/NiagaraNodeStaticSwitch.cpp
@@ -28,6 +28,35 @@ FNiagaraTypeDefinition UNiagaraNodeStaticSwitch::GetInputType() const
 	return FNiagaraTypeDefinition();
 }
 
+void UNiagaraNodeStaticSwitch::ChangeSwitchParameterName(const FName& NewName)
+{
+	FNiagaraVariable OldValue(GetInputType(), InputParameterName);
+	InputParameterName = NewName;
+	VisualsChangedDelegate.Broadcast(this);
+	RemoveUnusedGraphParameter(OldValue);	
+}
+
+void UNiagaraNodeStaticSwitch::OnSwitchParameterTypeChanged(const FNiagaraTypeDefinition& OldType)
+{
+	RefreshFromExternalChanges();
+	VisualsChangedDelegate.Broadcast(this);
+	RemoveUnusedGraphParameter(FNiagaraVariable(OldType, InputParameterName));
+}
+
+void UNiagaraNodeStaticSwitch::RemoveUnusedGraphParameter(const FNiagaraVariable& OldParameter)
+{
+	TArray<FNiagaraVariable> GraphVariables = GetNiagaraGraph()->FindStaticSwitchInputs();
+	int Index = GraphVariables.Find(OldParameter);
+	if (Index == INDEX_NONE)
+	{
+		GetNiagaraGraph()->RemoveParameter(OldParameter);
+	}
+	else
+	{
+		GetNiagaraGraph()->NotifyGraphChanged();
+	}
+}
+
 void UNiagaraNodeStaticSwitch::AllocateDefaultPins()
 {
 	const UEdGraphSchema_Niagara* Schema = GetDefault<UEdGraphSchema_Niagara>();
diff --git a/Engine/Plugins/FX/Niagara/Source/NiagaraEditor/Private/NiagaraNodeStaticSwitch.h b/Engine/Plugins/FX/Niagara/Source/NiagaraEditor/Private/NiagaraNodeStaticSwitch.h
--- a/Engine/Plugins/FX/Niagara/Source/NiagaraEditor/Private/NiagaraNodeStaticSwitch.h
+++ b/Engine/Plugins/FX/Niagara/Source/NiagaraEditor/Private/NiagaraNodeStaticSwitch.h
@@ -41,7 +41,7 @@ class UNiagaraNodeStaticSwitch : public UNiagaraNodeUsageSelector
 
 public:
 
-	UPROPERTY(EditAnywhere, Category="Static Switch")
+	UPROPERTY()
 	FName InputParameterName;
 	
 	UPROPERTY()
@@ -57,6 +57,9 @@ public:
 
 	FNiagaraTypeDefinition GetInputType() const;
 
+	void ChangeSwitchParameterName(const FName& NewName);
+	void OnSwitchParameterTypeChanged(const FNiagaraTypeDefinition& OldType);
+
 	//~ Begin EdGraphNode Interface
 	virtual void AllocateDefaultPins() override;
 	virtual FText GetTooltipText() const override;
@@ -82,4 +85,6 @@ private:
 	bool GetVarIndex(class FHlslNiagaraTranslator* Translator, int32 InputPinCount, int32& VarIndexOut) const;
 
 	bool GetVarIndex(class FHlslNiagaraTranslator* Translator, int32 InputPinCount, int32 Value, int32& VarIndexOut) const;
+
+	void RemoveUnusedGraphParameter(const FNiagaraVariable& OldParameter);
 };
