Allow the user to change the name of static switch parameters that are propagated up the call chain.

diff --git a/Engine/Plugins/FX/Niagara/Source/NiagaraEditor/Private/Customizations/NiagaraFunctionCallNodeDetails.cpp b/Engine/Plugins/FX/Niagara/Source/NiagaraEditor/Private/Customizations/NiagaraFunctionCallNodeDetails.cpp
--- a/Engine/Plugins/FX/Niagara/Source/NiagaraEditor/Private/Customizations/NiagaraFunctionCallNodeDetails.cpp
+++ b/Engine/Plugins/FX/Niagara/Source/NiagaraEditor/Private/Customizations/NiagaraFunctionCallNodeDetails.cpp
@@ -12,6 +12,7 @@
 #include "Widgets/Input/SComboBox.h"
 #include "Widgets/Input/SNumericEntryBox.h"
 #include "Widgets/Input/SCheckBox.h"
+#include "Widgets/Input/SEditableTextBox.h"
 
 #define LOCTEXT_NAMESPACE "NiagaraFunctionCallNodeDetails"
 
@@ -58,27 +59,53 @@ void FNiagaraFunctionCallNodeDetails::CustomizeDetails(IDetailLayoutBuilder& Det
 		]
 		.ValueContent()
 		[
-			SNew(SBox)
+			SNew(SVerticalBox)
+			+ SVerticalBox::Slot()
 			.Padding(FMargin(0.0f, 2.0f))
 			[
 				SNew(SCheckBox)
 				.IsChecked_Lambda([SwitchInput, this]()
 				{
-					return Node->PropagatedStaticSwitchParameters.Contains(SwitchInput) ? ECheckBoxState::Checked : ECheckBoxState::Unchecked;
+					return Node->FindPropagatedVariable(SwitchInput) ? ECheckBoxState::Checked : ECheckBoxState::Unchecked;
 				})
 				.OnCheckStateChanged_Lambda([SwitchInput, this](const ECheckBoxState NewState)
 				{
 					if (NewState == ECheckBoxState::Unchecked)
 					{
-						Node->PropagatedStaticSwitchParameters.Remove(SwitchInput);
+						Node->RemovePropagatedVariable(SwitchInput);
 					}
-					else if (NewState == ECheckBoxState::Checked)
+					else if (NewState == ECheckBoxState::Checked && !Node->FindPropagatedVariable(SwitchInput))
 					{
-						Node->PropagatedStaticSwitchParameters.AddUnique(SwitchInput);
+						Node->PropagatedStaticSwitchParameters.Emplace(SwitchInput);
 					}
 					Node->RefreshFromExternalChanges();
 				})
 			]
+			+ SVerticalBox::Slot()
+			.HAlign(HAlign_Fill)
+			[
+				SNew(SEditableTextBox)
+				.ToolTipText(LOCTEXT("NiagaraOverridePropagatedValueName_TooltipText", "When set, the name of the propagated switch value is changed to a different value."))
+				.HintText(LOCTEXT("NiagaraOverridePropagatedValueName_HintText", "Override name"))
+				.IsEnabled_Lambda([SwitchInput, this]() { return Node->FindPropagatedVariable(SwitchInput) != nullptr; })
+				.OnTextCommitted_Lambda([SwitchInput, this](const FText& NewText, ETextCommit::Type CommitType)
+				{
+					FNiagaraPropagatedVariable* Propagated = Node->FindPropagatedVariable(SwitchInput);
+					if (Propagated)
+					{
+						Propagated->PropagatedName = NewText.ToString();
+					}
+				})
+				.Text_Lambda([SwitchInput, this]()
+				{
+					FNiagaraPropagatedVariable* Propagated = Node->FindPropagatedVariable(SwitchInput);
+					if (Propagated)
+					{
+						return FText::FromString(Propagated->PropagatedName);
+					}
+					return FText();
+				})
+			]
 		];
 	}
 }
diff --git a/Engine/Plugins/FX/Niagara/Source/NiagaraEditor/Private/NiagaraEditorUtilities.cpp b/Engine/Plugins/FX/Niagara/Source/NiagaraEditor/Private/NiagaraEditorUtilities.cpp
--- a/Engine/Plugins/FX/Niagara/Source/NiagaraEditor/Private/NiagaraEditorUtilities.cpp
+++ b/Engine/Plugins/FX/Niagara/Source/NiagaraEditor/Private/NiagaraEditorUtilities.cpp
@@ -649,17 +649,18 @@ void FNiagaraEditorUtilities::SetStaticSwitchConstants(UNiagaraGraph* Graph, con
 		UNiagaraNodeFunctionCall* FunctionNode = Cast<UNiagaraNodeFunctionCall>(Node);
 		if (FunctionNode && FunctionNode->PropagatedStaticSwitchParameters.Num() > 0)
 		{
-			for (const FNiagaraVariable& SwitchValue : FunctionNode->PropagatedStaticSwitchParameters)
+			for (const FNiagaraPropagatedVariable& SwitchValue : FunctionNode->PropagatedStaticSwitchParameters)
 			{
-				UEdGraphPin* ValuePin = FunctionNode->FindPin(SwitchValue.GetName(), EGPD_Input);
+				UEdGraphPin* ValuePin = FunctionNode->FindPin(SwitchValue.SwitchParameter.GetName(), EGPD_Input);
 				if (!ValuePin)
 				{
 					continue;
 				}
 				ValuePin->DefaultValue = FString();
+				FName PinName = SwitchValue.PropagatedName.IsEmpty() ? ValuePin->GetFName() : FName(*SwitchValue.PropagatedName);
 				for (UEdGraphPin* InputPin : CallInputs)
 				{
-					if (InputPin->GetFName().IsEqual(ValuePin->GetFName()) && InputPin->PinType == ValuePin->PinType)
+					if (InputPin->GetFName().IsEqual(PinName) && InputPin->PinType == ValuePin->PinType)
 					{
 						ValuePin->DefaultValue = InputPin->DefaultValue;
 						break;
diff --git a/Engine/Plugins/FX/Niagara/Source/NiagaraEditor/Private/NiagaraGraph.cpp b/Engine/Plugins/FX/Niagara/Source/NiagaraEditor/Private/NiagaraGraph.cpp
--- a/Engine/Plugins/FX/Niagara/Source/NiagaraEditor/Private/NiagaraGraph.cpp
+++ b/Engine/Plugins/FX/Niagara/Source/NiagaraEditor/Private/NiagaraGraph.cpp
@@ -629,12 +629,21 @@ TArray<FNiagaraVariable> UNiagaraGraph::FindStaticSwitchInputs() const
 		UNiagaraNodeFunctionCall* FunctionNode = Cast<UNiagaraNodeFunctionCall>(Node);
 		if (FunctionNode)
 		{
-			for (FNiagaraVariable Variable : FunctionNode->PropagatedStaticSwitchParameters)
+			for (const FNiagaraPropagatedVariable& Propagated : FunctionNode->PropagatedStaticSwitchParameters)
 			{
+				FNiagaraVariable Variable = Propagated.SwitchParameter;
+				if (!Propagated.PropagatedName.IsEmpty())
+				{
+					Variable.SetName(FName(*Propagated.PropagatedName));
+				}
 				Result.AddUnique(Variable);
 			}			
 		}
 	}
+	Result.Sort([](const FNiagaraVariable& Left, const FNiagaraVariable& Right)
+	{
+		return Left.GetName() < Right.GetName();
+	});
 	return Result;
 }
 
@@ -1538,8 +1547,13 @@ void UNiagaraGraph::RefreshParameterReferences() const
 		}
 		else if (UNiagaraNodeFunctionCall* FunctionNode = Cast<UNiagaraNodeFunctionCall>(Node))
 		{
-			for (FNiagaraVariable Variable : FunctionNode->PropagatedStaticSwitchParameters)
+			for (const FNiagaraPropagatedVariable& Propagated : FunctionNode->PropagatedStaticSwitchParameters)
 			{
+				FNiagaraVariable Variable = Propagated.SwitchParameter;
+				if (!Propagated.PropagatedName.IsEmpty())
+				{
+					Variable.SetName(FName(*Propagated.PropagatedName));
+				}
 				AddStaticParameterReference(Variable, FunctionNode);
 			}
 		}
diff --git a/Engine/Plugins/FX/Niagara/Source/NiagaraEditor/Private/NiagaraNodeFunctionCall.cpp b/Engine/Plugins/FX/Niagara/Source/NiagaraEditor/Private/NiagaraNodeFunctionCall.cpp
--- a/Engine/Plugins/FX/Niagara/Source/NiagaraEditor/Private/NiagaraNodeFunctionCall.cpp
+++ b/Engine/Plugins/FX/Niagara/Source/NiagaraEditor/Private/NiagaraNodeFunctionCall.cpp
@@ -166,7 +166,7 @@ void UNiagaraNodeFunctionCall::AllocateDefaultPins()
 			{
 				NewPin->DefaultValue = PinDefaultValue;
 			}
-			NewPin->bDefaultValueIsIgnored = PropagatedStaticSwitchParameters.Contains(Input);
+			NewPin->bDefaultValueIsIgnored = FindPropagatedVariable(Input) != nullptr;
 		}
 
 		for (FNiagaraVariable& Output : Outputs)
@@ -556,7 +556,7 @@ bool UNiagaraNodeFunctionCall::RefreshFromExternalChanges()
 			{
 				if (InputVar.GetName().IsEqual(Pin->GetFName()))
 				{
-					Pin->bDefaultValueIsIgnored = PropagatedStaticSwitchParameters.Contains(InputVar);
+					Pin->bDefaultValueIsIgnored = FindPropagatedVariable(InputVar) != nullptr;
 					break;
 				}
 			}
@@ -767,6 +767,30 @@ UNiagaraNodeFunctionCall::FOnInputsChanged& UNiagaraNodeFunctionCall::OnInputsCh
 	return OnInputsChangedDelegate;
 }
 
+FNiagaraPropagatedVariable* UNiagaraNodeFunctionCall::FindPropagatedVariable(const FNiagaraVariable& Variable)
+{
+	for (FNiagaraPropagatedVariable& Propagated : PropagatedStaticSwitchParameters)
+	{
+		if (Propagated.SwitchParameter == Variable)
+		{
+			return &Propagated;
+		}
+	}
+	return nullptr;
+}
+
+void UNiagaraNodeFunctionCall::RemovePropagatedVariable(const FNiagaraVariable& Variable)
+{
+	for (int i = 0; i < PropagatedStaticSwitchParameters.Num(); i++)
+	{
+		if (PropagatedStaticSwitchParameters[i].SwitchParameter == Variable)
+		{
+			PropagatedStaticSwitchParameters.RemoveAt(i);
+			return;
+		}
+	}
+}
+
 ENiagaraNumericOutputTypeSelectionMode UNiagaraNodeFunctionCall::GetNumericOutputTypeSelectionMode() const
 {
 	if (FunctionScript)
diff --git a/Engine/Plugins/FX/Niagara/Source/NiagaraEditor/Public/NiagaraNodeFunctionCall.h b/Engine/Plugins/FX/Niagara/Source/NiagaraEditor/Public/NiagaraNodeFunctionCall.h
--- a/Engine/Plugins/FX/Niagara/Source/NiagaraEditor/Public/NiagaraNodeFunctionCall.h
+++ b/Engine/Plugins/FX/Niagara/Source/NiagaraEditor/Public/NiagaraNodeFunctionCall.h
@@ -11,6 +11,30 @@
 
 class UNiagaraScript;
 
+USTRUCT()
+struct FNiagaraPropagatedVariable
+{
+	GENERATED_USTRUCT_BODY()
+
+public:
+
+	FNiagaraPropagatedVariable() : FNiagaraPropagatedVariable(FNiagaraVariable()) {}
+
+	FNiagaraPropagatedVariable(FNiagaraVariable SwitchParameter) : SwitchParameter(SwitchParameter), PropagatedName(FString()) {}
+
+	UPROPERTY()
+	FNiagaraVariable SwitchParameter;
+
+	/** If set, then this overrides the name of the switch parameter when propagating. */
+	UPROPERTY()
+	FString PropagatedName;
+
+	bool operator==(const FNiagaraPropagatedVariable& Other)const
+	{
+		return SwitchParameter == Other.SwitchParameter;
+	}
+};
+
 UCLASS(MinimalAPI)
 class UNiagaraNodeFunctionCall : public UNiagaraNodeWithDynamicPins
 {
@@ -38,7 +62,7 @@ public:
 
 	/** All the input values the function propagates to the next higher caller instead of forcing the user to set them directly. */
 	UPROPERTY()
-	TArray<FNiagaraVariable> PropagatedStaticSwitchParameters;
+	TArray<FNiagaraPropagatedVariable> PropagatedStaticSwitchParameters;
 
 	bool ScriptIsValid() const;
 
@@ -88,6 +112,9 @@ public:
 
 	FOnInputsChanged& OnInputsChanged();
 
+	FNiagaraPropagatedVariable* FindPropagatedVariable(const FNiagaraVariable& Variable);
+	void RemovePropagatedVariable(const FNiagaraVariable& Variable);
+
 protected:
 	virtual bool IsValidPinToCompile(UEdGraphPin* Pin) const;
 
