Fixed FName API usage

(peafour-cherrypick of 6815523 by [REDACTED_PERSON])

diff --git a/Engine/Plugins/FX/Niagara/Source/NiagaraEditor/Private/NiagaraGraph.cpp b/Engine/Plugins/FX/Niagara/Source/NiagaraEditor/Private/NiagaraGraph.cpp
--- a/Engine/Plugins/FX/Niagara/Source/NiagaraEditor/Private/NiagaraGraph.cpp
+++ b/Engine/Plugins/FX/Niagara/Source/NiagaraEditor/Private/NiagaraGraph.cpp
@@ -642,7 +642,7 @@ TArray<FNiagaraVariable> UNiagaraGraph::FindStaticSwitchInputs() const
 	}
 	Result.Sort([](const FNiagaraVariable& Left, const FNiagaraVariable& Right)
 	{
-		return Left.GetName() < Right.GetName();
+		return Left.GetName().LexicalLess(Right.GetName());
 	});
 	return Result;
 }
