Reenabled inline condition toggle for static switch parameters

diff --git a/Engine/Plugins/FX/Niagara/Source/Niagara/Public/NiagaraTypes.h b/Engine/Plugins/FX/Niagara/Source/Niagara/Public/NiagaraTypes.h
--- a/Engine/Plugins/FX/Niagara/Source/Niagara/Public/NiagaraTypes.h
+++ b/Engine/Plugins/FX/Niagara/Source/Niagara/Public/NiagaraTypes.h
@@ -340,7 +340,7 @@ public:
 
 	/** Declares the associated input is used as an inline edit condition toggle, so it should be hidden and edited as a 
 	checkbox inline with the input which was designated as its edit condition. */
-	UPROPERTY(EditAnywhere, Category = "Variable", meta = (EditCondition = "!bIsStaticSwitch"))
+	UPROPERTY(EditAnywhere, Category = "Variable")
 	bool bInlineEditConditionToggle;
 
 	/** Declares the associated input should be conditionally editable based on the value of another input. */
