- Niagara fix skel mesh data interface failing to bind when we use bone only sampling without CPU access to mesh data

#rnx

diff --git a/Engine/Plugins/FX/Niagara/Source/Niagara/Private/NiagaraDataInterfaceSkeletalMesh.cpp b/Engine/Plugins/FX/Niagara/Source/Niagara/Private/NiagaraDataInterfaceSkeletalMesh.cpp
--- a/Engine/Plugins/FX/Niagara/Source/Niagara/Private/NiagaraDataInterfaceSkeletalMesh.cpp
+++ b/Engine/Plugins/FX/Niagara/Source/Niagara/Private/NiagaraDataInterfaceSkeletalMesh.cpp
@@ -1012,40 +1012,42 @@ bool FNDISkeletalMesh_InstanceData::Init(UNiagaraDataInterfaceSkeletalMesh* Inte
 	FSkeletalMeshLODRenderData& LODData = GetLODRenderDataAndSkinWeights(SkinWeightBuffer);
 
 	//Check for the validity of the Mesh's cpu data.
-
-	bool LODDataNumVerticesCorrect = LODData.GetNumVertices() > 0;
-	bool LODDataPositonNumVerticesCorrect = LODData.StaticVertexBuffers.PositionVertexBuffer.GetNumVertices() > 0;
-	bool bSkinWeightBuffer = SkinWeightBuffer != nullptr;
-	bool SkinWeightBufferNumVerticesCorrect = bSkinWeightBuffer && (SkinWeightBuffer->GetNumVertices() > 0);
-	bool bIndexBufferValid = LODData.MultiSizeIndexContainer.IsIndexBufferValid();
-	bool bIndexBufferFound = bIndexBufferValid && (LODData.MultiSizeIndexContainer.GetIndexBuffer() != nullptr);
-	bool bIndexBufferNumCorrect = bIndexBufferFound && (LODData.MultiSizeIndexContainer.GetIndexBuffer()->Num() > 0);
-
-	bool bMeshCPUDataValid = LODDataNumVerticesCorrect &&
-		LODDataPositonNumVerticesCorrect &&
-		bSkinWeightBuffer &&
-		SkinWeightBufferNumVerticesCorrect && 
-		bIndexBufferValid &&
-		bIndexBufferFound &&
-		bIndexBufferNumCorrect;
-
-	if (!bMeshCPUDataValid)
-	{
-		UE_LOG(LogNiagara, Warning, TEXT("Skeletal Mesh Data Interface is trying to sample from a mesh with missing CPU vertex or index data.\nInterface: %s\nMesh: %s\nLOD: %d\n"
-			"LODDataNumVerticesCorrect: %d  LODDataPositonNumVerticesCorrect : %d  bSkinWeightBuffer : %d  SkinWeightBufferNumVerticesCorrect : %d bIndexBufferValid : %d  bIndexBufferFound : %d  bIndexBufferNumCorrect : %d"),
-			*Interface->GetFullName(),
-			*Mesh->GetFullName(),
-			LODIndex,
-			LODDataNumVerticesCorrect ? 1 : 0,
-			LODDataPositonNumVerticesCorrect ? 1 : 0,
-			bSkinWeightBuffer ? 1 : 0,
-			SkinWeightBufferNumVerticesCorrect ? 1 : 0,
-			bIndexBufferValid ? 1 : 0,
-			bIndexBufferFound ? 1 : 0,
-			bIndexBufferNumCorrect ? 1 : 0
+	if ( Interface->bUseTriangleSampling || Interface->bUseVertexSampling )
+	{
+		bool LODDataNumVerticesCorrect = LODData.GetNumVertices() > 0;
+		bool LODDataPositonNumVerticesCorrect = LODData.StaticVertexBuffers.PositionVertexBuffer.GetNumVertices() > 0;
+		bool bSkinWeightBuffer = SkinWeightBuffer != nullptr;
+		bool SkinWeightBufferNumVerticesCorrect = bSkinWeightBuffer && (SkinWeightBuffer->GetNumVertices() > 0);
+		bool bIndexBufferValid = LODData.MultiSizeIndexContainer.IsIndexBufferValid();
+		bool bIndexBufferFound = bIndexBufferValid && (LODData.MultiSizeIndexContainer.GetIndexBuffer() != nullptr);
+		bool bIndexBufferNumCorrect = bIndexBufferFound && (LODData.MultiSizeIndexContainer.GetIndexBuffer()->Num() > 0);
+
+		bool bMeshCPUDataValid = LODDataNumVerticesCorrect &&
+			LODDataPositonNumVerticesCorrect &&
+			bSkinWeightBuffer &&
+			SkinWeightBufferNumVerticesCorrect &&
+			bIndexBufferValid &&
+			bIndexBufferFound &&
+			bIndexBufferNumCorrect;
+
+		if (!bMeshCPUDataValid)
+		{
+			UE_LOG(LogNiagara, Warning, TEXT("Skeletal Mesh Data Interface is trying to sample from a mesh with missing CPU vertex or index data.\nInterface: %s\nMesh: %s\nLOD: %d\n"
+				"LODDataNumVerticesCorrect: %d  LODDataPositonNumVerticesCorrect : %d  bSkinWeightBuffer : %d  SkinWeightBufferNumVerticesCorrect : %d bIndexBufferValid : %d  bIndexBufferFound : %d  bIndexBufferNumCorrect : %d"),
+				*Interface->GetFullName(),
+				*Mesh->GetFullName(),
+				LODIndex,
+				LODDataNumVerticesCorrect ? 1 : 0,
+				LODDataPositonNumVerticesCorrect ? 1 : 0,
+				bSkinWeightBuffer ? 1 : 0,
+				SkinWeightBufferNumVerticesCorrect ? 1 : 0,
+				bIndexBufferValid ? 1 : 0,
+				bIndexBufferFound ? 1 : 0,
+				bIndexBufferNumCorrect ? 1 : 0
 			);
 
-		return false;
+			return false;
+		}
 	}
 
 	FReferenceSkeleton& RefSkel = Mesh->RefSkeleton;
