Fix missing FNiagaraStandaloneScriptViewModel in FNiagaraScriptToolkit.

diff --git a/Engine/Plugins/FX/Niagara/Source/NiagaraEditor/Private/Toolkits/NiagaraScriptToolkit.cpp b/Engine/Plugins/FX/Niagara/Source/NiagaraEditor/Private/Toolkits/NiagaraScriptToolkit.cpp
--- a/Engine/Plugins/FX/Niagara/Source/NiagaraEditor/Private/Toolkits/NiagaraScriptToolkit.cpp
+++ b/Engine/Plugins/FX/Niagara/Source/NiagaraEditor/Private/Toolkits/NiagaraScriptToolkit.cpp
@@ -143,15 +143,17 @@ void FNiagaraScriptToolkit::Initialize( const EToolkitMode::Type Mode, const TSh
 	{
 		DisplayName = FAssetTypeActions_NiagaraScriptDynamicInputs::GetFormattedName();
 	}
-	ScriptViewModel = MakeShareable(new FNiagaraScriptViewModel(EditedNiagaraScript, DisplayName, ENiagaraParameterEditMode::EditAll));
+
+	const FGuid MessageLogGuidKey = FGuid::NewGuid();
+	NiagaraMessageLogViewModel = MakeShared<FNiagaraMessageLogViewModel>(GetNiagaraScriptMessageLogName(EditedNiagaraScript), MessageLogGuidKey, NiagaraMessageLog);
+
+	ScriptViewModel = MakeShareable(new FNiagaraStandaloneScriptViewModel(EditedNiagaraScript, DisplayName, ENiagaraParameterEditMode::EditAll, NiagaraMessageLogViewModel, OriginalNiagaraScript, MessageLogGuidKey));
 
 	OnEditedScriptGraphChangedHandle = ScriptViewModel->GetGraphViewModel()->GetGraph()->AddOnGraphNeedsRecompileHandler(
 		FOnGraphChanged::FDelegate::CreateRaw(this, &FNiagaraScriptToolkit::OnEditedScriptGraphChanged));
 
 	DetailsScriptSelection = MakeShareable(new FNiagaraObjectSelection());
 	DetailsScriptSelection->SetSelectedObject(EditedNiagaraScript);
-
-	const FGuid MessageLogGuidKey = FGuid::NewGuid();
 	
  	FMessageLogModule& MessageLogModule = FModuleManager::LoadModuleChecked<FMessageLogModule>("MessageLog"); //@todo(message manager) remove stats listing
  	FMessageLogInitializationOptions LogOptions;
@@ -163,8 +165,6 @@ void FNiagaraScriptToolkit::Initialize( const EToolkitMode::Type Mode, const TSh
 	StatsListing = MessageLogModule.CreateLogListing("MaterialEditorStats", LogOptions);
 	Stats = MessageLogModule.CreateLogListingWidget(StatsListing.ToSharedRef());
 
-	NiagaraMessageLogViewModel = MakeShared<FNiagaraMessageLogViewModel>(GetNiagaraScriptMessageLogName(EditedNiagaraScript), MessageLogGuidKey, NiagaraMessageLog);
-
 	TSharedRef<FTabManager::FLayout> StandaloneDefaultLayout = FTabManager::NewLayout("Standalone_Niagara_Layout_v8")
 	->AddArea
 	(
diff --git a/Engine/Plugins/FX/Niagara/Source/NiagaraEditor/Private/Toolkits/NiagaraScriptToolkit.h b/Engine/Plugins/FX/Niagara/Source/NiagaraEditor/Private/Toolkits/NiagaraScriptToolkit.h
--- a/Engine/Plugins/FX/Niagara/Source/NiagaraEditor/Private/Toolkits/NiagaraScriptToolkit.h
+++ b/Engine/Plugins/FX/Niagara/Source/NiagaraEditor/Private/Toolkits/NiagaraScriptToolkit.h
@@ -21,6 +21,7 @@ class FNiagaraScriptViewModel;
 class FNiagaraObjectSelection;
 struct FEdGraphEditAction;
 class FNiagaraMessageLogViewModel;
+class FNiagaraStandaloneScriptViewModel;
 
 /** Viewer/editor for a DataTable */
 class FNiagaraScriptToolkit : public FAssetEditorToolkit, public FGCObject
@@ -122,7 +123,7 @@ private:
 private:
 
 	/** The Script being edited */
-	TSharedPtr<FNiagaraScriptViewModel> ScriptViewModel;
+	TSharedPtr<FNiagaraStandaloneScriptViewModel> ScriptViewModel;
 
 	/** The selection displayed by the details tab. */
 	TSharedPtr<FNiagaraObjectSelection> DetailsScriptSelection;
