Fixup const correctness in TSharedPtr.

diff --git a/Engine/Plugins/FX/Niagara/Source/NiagaraEditor/Private/NiagaraMessageManager.cpp b/Engine/Plugins/FX/Niagara/Source/NiagaraEditor/Private/NiagaraMessageManager.cpp
--- a/Engine/Plugins/FX/Niagara/Source/NiagaraEditor/Private/NiagaraMessageManager.cpp
+++ b/Engine/Plugins/FX/Niagara/Source/NiagaraEditor/Private/NiagaraMessageManager.cpp
@@ -265,7 +265,7 @@ TSharedRef<FTokenizedMessage> FNiagaraMessageCompileEvent::GenerateTokenizedMess
 
 FNiagaraMessageJobCompileEvent::FNiagaraMessageJobCompileEvent(
 	  const FNiagaraCompileEvent& InCompileEvent
-	, const TWeakObjectPtr<UNiagaraScript>& InOriginatingScriptWeakObjPtr
+	, const TWeakObjectPtr<const UNiagaraScript>& InOriginatingScriptWeakObjPtr
 	, const TOptional<const FString>& InOwningScriptNameString
 	, const TOptional<const FString>& InSourceScriptAssetPath
 	)
@@ -282,11 +282,11 @@ TSharedRef<const INiagaraMessage> FNiagaraMessageJobCompileEvent::GenerateNiagar
 
 	if (OriginatingScriptWeakObjPtr.IsValid())
 	{
-		UNiagaraScriptSourceBase* FunctionScriptSourceBase = OriginatingScriptWeakObjPtr->GetSource();
+		const UNiagaraScriptSourceBase* FunctionScriptSourceBase = OriginatingScriptWeakObjPtr->GetSource();
 		checkf(FunctionScriptSourceBase->IsA<UNiagaraScriptSource>(), TEXT("Script source for function call node is not assigned or is not of type UNiagaraScriptSource!"))
-		UNiagaraScriptSource* FunctionScriptSource = Cast<UNiagaraScriptSource>(FunctionScriptSourceBase);
+		const UNiagaraScriptSource* FunctionScriptSource = Cast<UNiagaraScriptSource>(FunctionScriptSourceBase);
 		checkf(FunctionScriptSource, TEXT("Script source base was somehow not a derived type!"));
-		UNiagaraGraph* ScriptGraph = FunctionScriptSource->NodeGraph;
+		const UNiagaraGraph* ScriptGraph = FunctionScriptSource->NodeGraph;
 		checkf(ScriptGraph, TEXT("Function Script does not have a UNiagaraGraph!"));
 
 		TArray<FNiagaraScriptNameAndAssetPath> ContextScriptNamesAndPaths;
diff --git a/Engine/Plugins/FX/Niagara/Source/NiagaraEditor/Public/NiagaraMessageManager.h b/Engine/Plugins/FX/Niagara/Source/NiagaraEditor/Public/NiagaraMessageManager.h
--- a/Engine/Plugins/FX/Niagara/Source/NiagaraEditor/Public/NiagaraMessageManager.h
+++ b/Engine/Plugins/FX/Niagara/Source/NiagaraEditor/Public/NiagaraMessageManager.h
@@ -67,7 +67,7 @@ class FNiagaraMessageJobCompileEvent : public INiagaraMessageJob
 public:
 	FNiagaraMessageJobCompileEvent(
 		  const FNiagaraCompileEvent& InCompileEvent
-		, const TWeakObjectPtr<UNiagaraScript>& InOriginatingScriptWeakObjPtr
+		, const TWeakObjectPtr<const UNiagaraScript>& InOriginatingScriptWeakObjPtr
 		, const TOptional<const FString>& InOwningScriptNameString = TOptional<const FString>()
 		, const TOptional<const FString>& InSourceScriptAssetPath = TOptional<const FString>()
 	);
@@ -87,7 +87,7 @@ private:
 	) const;
 
 	const FNiagaraCompileEvent CompileEvent;
-	const TWeakObjectPtr<UNiagaraScript> OriginatingScriptWeakObjPtr;
+	const TWeakObjectPtr<const UNiagaraScript> OriginatingScriptWeakObjPtr;
 	TOptional<const FString> OwningScriptNameString;
 	TOptional<const FString> SourceScriptAssetPath;
 };
