Fixed a bug where static switches could not be used for inline edit conditions

diff --git a/Engine/Plugins/FX/Niagara/Source/NiagaraEditor/Private/ViewModels/Stack/NiagaraStackGraphUtilities.cpp b/Engine/Plugins/FX/Niagara/Source/NiagaraEditor/Private/ViewModels/Stack/NiagaraStackGraphUtilities.cpp
--- a/Engine/Plugins/FX/Niagara/Source/NiagaraEditor/Private/ViewModels/Stack/NiagaraStackGraphUtilities.cpp
+++ b/Engine/Plugins/FX/Niagara/Source/NiagaraEditor/Private/ViewModels/Stack/NiagaraStackGraphUtilities.cpp
@@ -596,6 +596,11 @@ UNiagaraNodeParameterMapSet& FNiagaraStackGraphUtilities::GetOrCreateStackFuncti
 
 UEdGraphPin* FNiagaraStackGraphUtilities::GetStackFunctionInputOverridePin(UNiagaraNodeFunctionCall& StackFunctionCall, FNiagaraParameterHandle AliasedInputParameterHandle)
 {
+	if (UEdGraphPin* SwitchPin = StackFunctionCall.FindStaticSwitchInputPin(AliasedInputParameterHandle.GetName()))
+	{
+		return SwitchPin;
+	}
+
 	UNiagaraNodeParameterMapSet* OverrideNode = GetStackFunctionOverrideNode(StackFunctionCall);
 	if (OverrideNode != nullptr)
 	{
