When copying data for the attribute spreadsheet on the GPU we were unintentionally using the wrong stride. We now cache off the stride when requesting the data (as a reallocation may happen later) and use that stride when copying off the data for main thread access in the spreadsheet.

#tests confirmed that GPU attribute spreadsheet data is correct now

diff --git a/Engine/Plugins/FX/Niagara/Source/Niagara/Classes/NiagaraDataSet.h b/Engine/Plugins/FX/Niagara/Source/Niagara/Classes/NiagaraDataSet.h
--- a/Engine/Plugins/FX/Niagara/Source/Niagara/Classes/NiagaraDataSet.h
+++ b/Engine/Plugins/FX/Niagara/Source/Niagara/Classes/NiagaraDataSet.h
@@ -99,7 +99,7 @@ public:
 	void KillInstance(uint32 InstanceIdx);
 	void CopyTo(FNiagaraDataBuffer& DestBuffer, int32 StartIdx, int32 NumInstances)const;
 	void CopyTo(FNiagaraDataBuffer& DestBuffer)const;
-	void GPUCopyFrom(float* GPUReadBackFloat, int* GPUReadBackInt, int32 StartIdx, int32 NumInstances);
+	void GPUCopyFrom(float* GPUReadBackFloat, int* GPUReadBackInt, int32 StartIdx, int32 NumInstances, uint32 InSrcFloatStride, uint32 InSrcIntStride);
 	void Dump(int32 StartIndex, int32 NumInstances, const FString& Label)const;
 
 	bool AppendToRegisterTable(uint8** Registers, int32& NumRegisters, int32 StartInstance);
@@ -255,7 +255,7 @@ public:
 
 	void CopyTo(FNiagaraDataSet& Other, int32 StartIdx = 0, int32 NumInstances = INDEX_NONE)const;
 
-	void CopyFromGPUReadback(float* GPUReadBackFloat, int* GPUReadBackInt, int32 StartIdx = 0, int32 NumInstances = INDEX_NONE);
+	void CopyFromGPUReadback(float* GPUReadBackFloat, int* GPUReadBackInt, int32 StartIdx = 0, int32 NumInstances = INDEX_NONE, uint32 FloatStride = 0, uint32 IntStride = 0);
 
 	void CheckForNaNs()const;
 
diff --git a/Engine/Plugins/FX/Niagara/Source/Niagara/Private/NiagaraDataSet.cpp b/Engine/Plugins/FX/Niagara/Source/Niagara/Private/NiagaraDataSet.cpp
--- a/Engine/Plugins/FX/Niagara/Source/Niagara/Private/NiagaraDataSet.cpp
+++ b/Engine/Plugins/FX/Niagara/Source/Niagara/Private/NiagaraDataSet.cpp
@@ -432,13 +432,13 @@ void FNiagaraDataSet::CopyTo(FNiagaraDataSet& Other, int32 StartIdx, int32 NumIn
 	Other.EndSimulate();
 }
 
-void FNiagaraDataSet::CopyFromGPUReadback(float* GPUReadBackFloat, int* GPUReadBackInt, int32 StartIdx /* = 0 */, int32 NumInstances /* = INDEX_NONE */)
+void FNiagaraDataSet::CopyFromGPUReadback(float* GPUReadBackFloat, int* GPUReadBackInt, int32 StartIdx /* = 0 */, int32 NumInstances /* = INDEX_NONE */, uint32 FloatStride, uint32 IntStride)
 {
 	check(IsInRenderingThread());
 	check(bFinalized);//We should be finalized with proper layout information already.
 
 	FNiagaraDataBuffer& DestBuffer = BeginSimulate();
-	DestBuffer.GPUCopyFrom(GPUReadBackFloat, GPUReadBackInt, StartIdx, NumInstances);
+	DestBuffer.GPUCopyFrom(GPUReadBackFloat, GPUReadBackInt, StartIdx, NumInstances, FloatStride, IntStride);
 	EndSimulate();
 }
  
@@ -744,7 +744,7 @@ void FNiagaraDataBuffer::CopyTo(FNiagaraDataBuffer& DestBuffer, int32 InStartIdx
 	}
 }
 
-void FNiagaraDataBuffer::GPUCopyFrom(float* GPUReadBackFloat, int* GPUReadBackInt, int32 InStartIdx, int32 InNumInstances)
+void FNiagaraDataBuffer::GPUCopyFrom(float* GPUReadBackFloat, int* GPUReadBackInt, int32 InStartIdx, int32 InNumInstances, uint32 InSrcFloatStride, uint32 InSrcIntStride)
 {
 	//CheckUsage(false); //Have to disable this as in this specific case we write to a "CPUSim" from the RT.
 
@@ -760,8 +760,10 @@ void FNiagaraDataBuffer::GPUCopyFrom(float* GPUReadBackFloat, int* GPUReadBackIn
 	{
 		for (uint32 CompIdx = 0; CompIdx < Owner->TotalFloatComponents; ++CompIdx)
 		{
-			const float* SrcStart = GetInstancePtrFloat(GPUReadBackFloat, CompIdx, InStartIdx);
-			const float* SrcEnd = GetInstancePtrFloat(GPUReadBackFloat, CompIdx, InStartIdx + InNumInstances);
+			// We have to reimplement the logic from GetInstancePtrFloat here because the incoming stride may be different than this 
+			// data buffer's stride.
+			const float* SrcStart = (const float*)((uint8*)GPUReadBackFloat + InSrcFloatStride * CompIdx) + InStartIdx; 
+			const float* SrcEnd = (const float*)((uint8*)GPUReadBackFloat + InSrcFloatStride * CompIdx) + InStartIdx + InNumInstances;
 			float* Dst = GetInstancePtrFloat(CompIdx, 0);
 			size_t Count = SrcEnd - SrcStart;
 			FMemory::Memcpy(Dst, SrcStart, Count * sizeof(float));
@@ -779,8 +781,10 @@ void FNiagaraDataBuffer::GPUCopyFrom(float* GPUReadBackFloat, int* GPUReadBackIn
 	{
 		for (uint32 CompIdx = 0; CompIdx < Owner->TotalInt32Components; ++CompIdx)
 		{
-			const int32* SrcStart = GetInstancePtrInt32(GPUReadBackInt, CompIdx, InStartIdx);
-			const int32* SrcEnd = GetInstancePtrInt32(GPUReadBackInt, CompIdx, InStartIdx + InNumInstances);
+			// We have to reimplement the logic from GetInstancePtrInt here because the incoming stride may be different than this 
+			// data buffer's stride.
+			const int32* SrcStart = (const int32*)((uint8*)GPUReadBackInt + InSrcIntStride * CompIdx) + InStartIdx;
+			const int32* SrcEnd = (const int32*)((uint8*)GPUReadBackInt + InSrcIntStride * CompIdx) + InStartIdx + InNumInstances;
 			int32* Dst = GetInstancePtrInt32(CompIdx, 0);
 			size_t Count = SrcEnd - SrcStart;
 			FMemory::Memcpy(Dst, SrcStart, Count * sizeof(int32));
