When copying data for the attribute spreadsheet on the GPU we were unintentionally using the wrong stride. We now cache off the stride when requesting the data (as a reallocation may happen later) and use that stride when copying off the data for main thread access in the spreadsheet.

#tests confirmed that GPU attribute spreadsheet data is correct now

diff --git a/Engine/Plugins/FX/Niagara/Source/Niagara/Classes/NiagaraScriptExecutionContext.h b/Engine/Plugins/FX/Niagara/Source/Niagara/Classes/NiagaraScriptExecutionContext.h
--- a/Engine/Plugins/FX/Niagara/Source/Niagara/Classes/NiagaraScriptExecutionContext.h
+++ b/Engine/Plugins/FX/Niagara/Source/Niagara/Classes/NiagaraScriptExecutionContext.h
@@ -186,6 +186,7 @@ public:
 	//Most current buffer that can be used for rendering.
 	FNiagaraDataBuffer* DataToRender;
 
+	bool ResetSinceLastReadbackIssued = false;
 	FRHIGPUMemoryReadback *GPUDataReadback;
 	uint32 AccumulatedSpawnRate;
 	uint32 NumIndicesPerInstance;	// how many vtx indices per instance the renderer is going to have for its draw call
@@ -199,6 +200,8 @@ public:
 	mutable FRHIGPUMemoryReadback *GPUDebugDataReadbackCounts;
 	mutable uint32 GPUDebugDataFloatSize;
 	mutable uint32 GPUDebugDataIntSize;
+	mutable uint32 GPUDebugDataFloatStride;
+	mutable uint32 GPUDebugDataIntStride;
 	mutable TSharedPtr<struct FNiagaraScriptDebuggerInfo, ESPMode::ThreadSafe> DebugInfo;
 #endif
 
@@ -275,4 +278,5 @@ public:
 	FNiagaraDataInterfaceInstanceData* DIInstanceData;
 	uint8* InstanceData_ParamData_Packed;
 	bool bRequiredDistanceFieldData = false;
+	bool bNeedsReset = false;
 };
diff --git a/Engine/Plugins/FX/Niagara/Source/Niagara/Private/NiagaraEmitterInstanceBatcher.cpp b/Engine/Plugins/FX/Niagara/Source/Niagara/Private/NiagaraEmitterInstanceBatcher.cpp
--- a/Engine/Plugins/FX/Niagara/Source/Niagara/Private/NiagaraEmitterInstanceBatcher.cpp
+++ b/Engine/Plugins/FX/Niagara/Source/Niagara/Private/NiagaraEmitterInstanceBatcher.cpp
@@ -454,7 +454,7 @@ void NiagaraEmitterInstanceBatcher::SimStepClearAndSetup(const FNiagaraComputeIn
 //	SCOPED_GPU_EVENTF(RHICmdList, NiagaraIndexBufferClear, TEXT("Niagara index buffer clear"));
 	SCOPED_DRAW_EVENT(RHICmdList, NiagaraIndexBufferClear);
 	SCOPED_GPU_STAT(RHICmdList, NiagaraIndexBufferClear);
-
+	//UE_LOG(LogNiagara, Log, TEXT("Clearing UAV %p"), (FRHIVertexBuffer *)DatasetIndexBufferWrite.Buffer);
 	ClearUAV(RHICmdList, DatasetIndexBufferWrite, 0);
 #endif
 }
@@ -513,7 +513,7 @@ void NiagaraEmitterInstanceBatcher::TickSingle(const FNiagaraGPUSystemTick& Tick
 	FNiagaraDataBuffer& CurrentData = *Instance->CurrentData;
 	FNiagaraDataBuffer& DestinationData = *Instance->DestinationData;
 
-	uint32 PrevNumInstances = CurrentData.GetNumInstances();
+	uint32 PrevNumInstances = Tick.bNeedsReset ? 0 : CurrentData.GetNumInstances();
 	uint32 NewNumInstances = Instance->SpawnRateInstances + Instance->EventSpawnTotal + PrevNumInstances;
 
 	//We must assume all particles survive when allocating here. 
@@ -813,13 +813,16 @@ void NiagaraEmitterInstanceBatcher::ResolveDatasetWrites(FRHICommandList &RHICmd
 		    int32 *NumInstancesAfterSim = static_cast<int32*>(Context->GPUDataReadback->Lock(64 * sizeof(int32)));
 			if (NumInstancesAfterSim)
 			{
-				int32 NewExistingDataCount = NumInstancesAfterSim[1] + Context->AccumulatedSpawnRate; // index 1 is always the count
+				int32 NewExistingDataCount = (Context->ResetSinceLastReadbackIssued ? 0 : NumInstancesAfterSim[1]) + Context->AccumulatedSpawnRate; // index 1 is always the count
 				DestinationData.SetNumInstances(NewExistingDataCount);
-				//UE_LOG(LogNiagara, Log, TEXT("GPU Syncup %s : Was(%d) Now(%d)"), *Context->DebugSimName, ExistingDataCount, NewExistingDataCount );
+				//UE_LOG(LogNiagara, Log, TEXT("GPU Syncup %s : Was(%d) Now(%d) Accumulated(%d) NumInstancesAfterSim[1](%d)"), *Context->DebugSimName, ExistingDataCount, NewExistingDataCount , Context->AccumulatedSpawnRate, NumInstancesAfterSim[1]);
+				//UE_LOG(LogNiagara, Log, TEXT("Reading UAV %p"), (FRHIVertexBuffer *)DatasetIndexBufferWrite.Buffer);
+
 				INC_DWORD_STAT_BY(STAT_NiagaraGPUParticles, NewExistingDataCount);
 				SET_DWORD_STAT(STAT_NiagaraReadbackLatency, 0);
 
 				Context->AccumulatedSpawnRate = 0;
+				Context->ResetSinceLastReadbackIssued = false;
 				bSuccessfullyRead = true;
 			}
 			else
@@ -876,7 +879,7 @@ void NiagaraEmitterInstanceBatcher::ProcessDebugInfo(FRHICommandList &RHICmdList
 					IntDataBuffer = static_cast<int*>(Context->GPUDebugDataReadbackInt->Lock(Context->GPUDebugDataIntSize));
 				}
 
-				Context->DebugInfo->Frame.CopyFromGPUReadback(FloatDataBuffer, IntDataBuffer, 0, NewExistingDataCount);
+				Context->DebugInfo->Frame.CopyFromGPUReadback(FloatDataBuffer, IntDataBuffer, 0, NewExistingDataCount, Context->GPUDebugDataFloatStride, Context->GPUDebugDataIntStride);
 
 				Context->DebugInfo->bWritten = true;
 
@@ -907,6 +910,8 @@ void NiagaraEmitterInstanceBatcher::ProcessDebugInfo(FRHICommandList &RHICmdList
 				Context->GPUDebugDataReadbackCounts = nullptr;	
 				Context->GPUDebugDataFloatSize = 0;
 				Context->GPUDebugDataIntSize = 0;
+				Context->GPUDebugDataFloatStride = 0;
+				Context->GPUDebugDataIntStride = 0;
 			}
 
 			// We've updated the debug info directly, now we need to no longer keep asking and querying because this frame is done!
@@ -1018,6 +1023,13 @@ void NiagaraEmitterInstanceBatcher::Run(const FNiagaraGPUSystemTick& Tick, const
 		InitIndexBuffer[1] = 0;		// number of instances 
 		ReadIndexBuffer.Initialize(sizeof(int32), 64, EPixelFormat::PF_R32_UINT, BUF_DrawIndirect | BUF_Static, nullptr, &InitIndexBuffer);
 	}
+	else if (Tick.bNeedsReset)
+	{
+		//UE_LOG(LogNiagara, Log, TEXT("Clearing UAV due to reset!"));
+		ClearUAV(RHICmdList, ReadIndexBuffer, 0);
+		Context->AccumulatedSpawnRate = 0;
+		Context->ResetSinceLastReadbackIssued = true;
+	}
 
 	RHICmdList.SetComputeShader(Shader->GetComputeShader());
 
@@ -1098,12 +1110,15 @@ void NiagaraEmitterInstanceBatcher::Run(const FNiagaraGPUSystemTick& Tick, const
 
 			Context->GPUDebugDataFloatSize = 0;
 			Context->GPUDebugDataIntSize = 0;
+			Context->GPUDebugDataFloatStride = 0;
+			Context->GPUDebugDataIntStride = 0;
 
 			if (DestinationData.GetGPUBufferFloat().NumBytes > 0)
 			{
 				Context->GPUDebugDataReadbackFloat = new FRHIGPUMemoryReadback(TEXT("Niagara GPU Debug Info Float Emitter Readback"));
 				Context->GPUDebugDataReadbackFloat->EnqueueCopy(RHICmdList, DestinationData.GetGPUBufferFloat().Buffer);
 				Context->GPUDebugDataFloatSize = DestinationData.GetGPUBufferFloat().NumBytes;
+				Context->GPUDebugDataFloatStride = DestinationData.GetFloatStride();
 			}
 
 			if (DestinationData.GetGPUBufferInt().NumBytes > 0)
@@ -1111,6 +1126,7 @@ void NiagaraEmitterInstanceBatcher::Run(const FNiagaraGPUSystemTick& Tick, const
 				Context->GPUDebugDataReadbackInt = new FRHIGPUMemoryReadback(TEXT("Niagara GPU Debug Info Int Emitter Readback"));
 				Context->GPUDebugDataReadbackInt->EnqueueCopy(RHICmdList, DestinationData.GetGPUBufferInt().Buffer);
 				Context->GPUDebugDataIntSize = DestinationData.GetGPUBufferInt().NumBytes;
+				Context->GPUDebugDataIntStride = DestinationData.GetInt32Stride();
 			}
 
 			Context->GPUDebugDataReadbackCounts = new FRHIGPUMemoryReadback(TEXT("Niagara GPU Emitter Readback"));
diff --git a/Engine/Plugins/FX/Niagara/Source/Niagara/Private/NiagaraScriptExecutionContext.cpp b/Engine/Plugins/FX/Niagara/Source/Niagara/Private/NiagaraScriptExecutionContext.cpp
--- a/Engine/Plugins/FX/Niagara/Source/Niagara/Private/NiagaraScriptExecutionContext.cpp
+++ b/Engine/Plugins/FX/Niagara/Source/Niagara/Private/NiagaraScriptExecutionContext.cpp
@@ -323,6 +323,11 @@ void FNiagaraGPUSystemTick::Init(FNiagaraSystemInstance* InSystemInstance)
 	FNiagaraComputeInstanceData* Instances = (FNiagaraComputeInstanceData*)(InstanceData_ParamData_Packed);
 	uint8* ParamDataBufferPtr = InstanceData_ParamData_Packed + PackedDispatchesSizeAligned;
 
+	int32 TickCount = InSystemInstance->GetTickCount();
+	check(TickCount > 0);
+	bNeedsReset = ( TickCount == 1);
+
+
 	// Now we will generate instance data for every GPU simulation we want to run on the render thread.
 	// This is spawn rate as well as DataInterface per instance data and the ParameterData for the emitter.
 	// @todo Ideally we would only update DataInterface and ParameterData bits if they have changed.
@@ -399,6 +404,8 @@ FNiagaraComputeExecutionContext::FNiagaraComputeExecutionContext()
 	, GPUDebugDataReadbackCounts(nullptr)
 	, GPUDebugDataFloatSize(0)
 	, GPUDebugDataIntSize(0)
+	, GPUDebugDataFloatStride(0)
+	, GPUDebugDataIntStride(0)
 #endif	  
 {
 }
diff --git a/Engine/Plugins/FX/Niagara/Source/Niagara/Public/NiagaraSystemInstance.h b/Engine/Plugins/FX/Niagara/Source/Niagara/Public/NiagaraSystemInstance.h
--- a/Engine/Plugins/FX/Niagara/Source/Niagara/Public/NiagaraSystemInstance.h
+++ b/Engine/Plugins/FX/Niagara/Source/Niagara/Public/NiagaraSystemInstance.h
@@ -159,6 +159,7 @@ public:
 	FORCEINLINE void SetPendingSpawn(bool bInValue) { bPendingSpawn = bInValue; }
 
 	FORCEINLINE float GetAge()const { return Age; }
+	FORCEINLINE int32 GetTickCount() const { return TickCount; }
 	
 	FORCEINLINE TSharedPtr<FNiagaraSystemSimulation, ESPMode::ThreadSafe> GetSystemSimulation()const
 	{
