Automatically transfer parameter values set to module inputs if they are replaced by static switch parameters of the same name and type.

diff --git a/Engine/Plugins/FX/Niagara/Source/NiagaraEditor/Private/NiagaraHlslTranslator.cpp b/Engine/Plugins/FX/Niagara/Source/NiagaraEditor/Private/NiagaraHlslTranslator.cpp
--- a/Engine/Plugins/FX/Niagara/Source/NiagaraEditor/Private/NiagaraHlslTranslator.cpp
+++ b/Engine/Plugins/FX/Niagara/Source/NiagaraEditor/Private/NiagaraHlslTranslator.cpp
@@ -4307,6 +4307,28 @@ void FHlslNiagaraTranslator::FunctionCall(UNiagaraNodeFunctionCall* FunctionNode
 	// We need the generated string to generate the proper signature for now.
 	ActiveHistoryForFunctionCalls.EnterFunction(FunctionNode->GetFunctionName(), FunctionNode->FunctionScript, FunctionNode);
 
+	// Check if there are static switch parameters being set directly by a set node from the stack UI.
+	// This can happen if a module was changed and the original parameter was replaced by a static switch with the same name, but the emitter was not yet updated.
+	const FString* ModuleAlias = ActiveHistoryForFunctionCalls.GetModuleAlias();
+	if (ModuleAlias)
+	{
+		for (int32 i = 0; i < ParamMapHistories.Num(); i++)
+		{
+			for (int32 j = 0; j < ParamMapHistories[i].VariablesWithOriginalAliasesIntact.Num(); j++)
+			{
+				const FNiagaraVariable Var = ParamMapHistories[i].VariablesWithOriginalAliasesIntact[j];
+				FString VarStr = Var.GetName().ToString();
+				if (VarStr.StartsWith(*ModuleAlias))
+				{
+					VarStr = VarStr.Mid(ModuleAlias->Len() + 1);
+					if (FunctionNode->FindStaticSwitchInputPin(*VarStr))
+					{
+						Error(FText::Format(LOCTEXT("SwitchPinFoundForSetPin", "A switch node pin exists but is being set directly using Set node! Please use the stack UI to resolve the conflict. Output Pin: {0}"), FText::FromName(Var.GetName())), FunctionNode, nullptr);
+					}
+				}
+			}
+		}
+	}
 
 	// Remove input add pin if it exists
 	for (int32 i = 0; i < CallOutputs.Num(); i++)
@@ -5475,7 +5497,7 @@ int32 FHlslNiagaraTranslator::CompileOutputPin(const UEdGraphPin* InPin)
 	return Ret;
 }
 
-void FHlslNiagaraTranslator::Error(FText ErrorText, const UNiagaraNode* Node, const UEdGraphPin* Pin) //@todo(message manager) rearrange syntax to fit with FNiagaraMessageManager
+void FHlslNiagaraTranslator::Error(FText ErrorText, const UNiagaraNode* Node, const UEdGraphPin* Pin)
 {
 	FString NodePinStr = TEXT("");
 	FString NodePinPrefix = TEXT(" - ");
diff --git a/Engine/Plugins/FX/Niagara/Source/NiagaraEditor/Private/ViewModels/Stack/NiagaraStackFunctionInputCollection.cpp b/Engine/Plugins/FX/Niagara/Source/NiagaraEditor/Private/ViewModels/Stack/NiagaraStackFunctionInputCollection.cpp
--- a/Engine/Plugins/FX/Niagara/Source/NiagaraEditor/Private/ViewModels/Stack/NiagaraStackFunctionInputCollection.cpp
+++ b/Engine/Plugins/FX/Niagara/Source/NiagaraEditor/Private/ViewModels/Stack/NiagaraStackFunctionInputCollection.cpp
@@ -85,6 +85,7 @@ void UNiagaraStackFunctionInputCollection::RefreshChildrenInternal(const TArray<
 	TArray<FName> ProcessedInputNames;
 	TArray<FName> DuplicateInputNames;
 	TArray<FName> ValidAliasedInputNames;
+	TMap<FName, UEdGraphPin*> StaticSwitchInputs;
 	TArray<const UEdGraphPin*> PinsWithInvalidTypes;
 
 	FText UncategorizedName = LOCTEXT("Uncategorized", "Uncategorized");
@@ -129,9 +130,9 @@ void UNiagaraStackFunctionInputCollection::RefreshChildrenInternal(const TArray<
 	}
 
 	// Gather static switch parameters
-	InputPins.Reset();
-	FNiagaraStackGraphUtilities::GetStackFunctionStaticSwitchPins(*InputFunctionCallNode, InputPins);
-	for (const UEdGraphPin* InputPin : InputPins)
+	TArray<UEdGraphPin*> SwitchPins;
+	FNiagaraStackGraphUtilities::GetStackFunctionStaticSwitchPins(*InputFunctionCallNode, SwitchPins);
+	for (UEdGraphPin* InputPin : SwitchPins)
 	{
 		// The static switch pin names to not contain the module namespace, as they are not part of the parameter maps.
 		// We add it here only to check for name clashes with actual module parameters.
@@ -152,7 +153,9 @@ void UNiagaraStackFunctionInputCollection::RefreshChildrenInternal(const TArray<
 			PinsWithInvalidTypes.Add(InputPin);
 			continue;
 		}
-		ValidAliasedInputNames.Add(FNiagaraParameterHandle(*InputFunctionCallNode->GetFunctionName(), InputPin->PinName).GetParameterHandleString());
+
+		FName AliasedName = FNiagaraParameterHandle(*InputFunctionCallNode->GetFunctionName(), InputPin->PinName).GetParameterHandleString();
+		StaticSwitchInputs.Add(AliasedName, InputPin);
 
 		TOptional<FNiagaraVariableMetaData> InputMetaData;
 		if (InputFunctionGraph != nullptr)
@@ -221,23 +224,50 @@ void UNiagaraStackFunctionInputCollection::RefreshChildrenInternal(const TArray<
 		}
 		InputCategory->AddInput(InputData.Pin->PinName, InputData.Type, InputData.bIsStatic ? EStackParameterBehavior::Static : EStackParameterBehavior::Dynamic);
 	}
-	RefreshIssues(DuplicateInputNames, ValidAliasedInputNames, PinsWithInvalidTypes, NewIssues);
+	RefreshIssues(DuplicateInputNames, ValidAliasedInputNames, PinsWithInvalidTypes, StaticSwitchInputs, NewIssues);
 }
 
-void UNiagaraStackFunctionInputCollection::RefreshIssues(TArray<FName> DuplicateInputNames, TArray<FName> ValidAliasedInputNames, TArray<const UEdGraphPin*> PinsWithInvalidTypes, TArray<FStackIssue>& NewIssues)
+UNiagaraStackEntry::FStackIssueFix UNiagaraStackFunctionInputCollection::GetNodeRemovalFix(UEdGraphPin* PinToRemove)
+{
+	FText FixDescription = LOCTEXT("RemoveInvalidInputTransaction", "Remove input override.");
+	return FStackIssueFix(
+		FixDescription,
+		UNiagaraStackEntry::FStackIssueFixDelegate::CreateLambda([=]()
+	{
+		FScopedTransaction ScopedTransaction(FixDescription);
+		TArray<TWeakObjectPtr<UNiagaraDataInterface>> RemovedDataObjects;
+		FNiagaraStackGraphUtilities::RemoveNodesForStackFunctionInputOverridePin(*PinToRemove, RemovedDataObjects);
+		for (TWeakObjectPtr<UNiagaraDataInterface> RemovedDataObject : RemovedDataObjects)
+		{
+			if (RemovedDataObject.IsValid())
+			{
+				OnDataObjectModified().Broadcast(RemovedDataObject.Get());
+			}
+		}
+		PinToRemove->GetOwningNode()->RemovePin(PinToRemove);
+	}));
+}
+
+void UNiagaraStackFunctionInputCollection::RefreshIssues(TArray<FName> DuplicateInputNames, TArray<FName> ValidAliasedInputNames, TArray<const UEdGraphPin*> PinsWithInvalidTypes, TMap<FName, UEdGraphPin*> StaticSwitchInputs, TArray<FStackIssue>& NewIssues)
 {
 	if (!GetIsEnabled())
 	{
 		NewIssues.Empty();
 		return;
 	}
-	// Try to find function input overrides which are no longer valid so we can generate errors for them.
+
+	// Gather override nodes to find candidates that were replaced by static switches and are no longer valid
+	TArray<UEdGraphPin*> OverridePins;
 	UNiagaraNodeParameterMapSet* OverrideNode = FNiagaraStackGraphUtilities::GetStackFunctionOverrideNode(*InputFunctionCallNode);
 	if (OverrideNode != nullptr)
 	{
-		TArray<UEdGraphPin*> OverridePins;
 		OverrideNode->GetInputPins(OverridePins);
-		for (UEdGraphPin* OverridePin : OverridePins)
+	}
+	for (UEdGraphPin* OverridePin : OverridePins)
+	{
+		// Try to find function input overrides which are no longer valid so we can generate errors for them.
+		UEdGraphPin** PinReference = StaticSwitchInputs.Find(OverridePin->PinName);
+		if (PinReference == nullptr)
 		{
 			// If the pin isn't in the misc category for the add pin, and not the parameter map pin, and it's for this function call,
 			// check to see if it's in the list of valid input names, and if not generate an error.
@@ -246,12 +276,39 @@ void UNiagaraStackFunctionInputCollection::RefreshIssues(TArray<FName> Duplicate
 				FNiagaraParameterHandle(OverridePin->PinName).GetNamespace().ToString() == InputFunctionCallNode->GetFunctionName() &&
 				ValidAliasedInputNames.Contains(OverridePin->PinName) == false)
 			{
-				FText FixDescription = LOCTEXT("RemoveInvalidInputTransaction", "Remove invalid input override.");
-				FStackIssueFix RemoveInputOverrideFix(
-					FixDescription,
+				FStackIssue InvalidInputOverrideError(
+					EStackIssueSeverity::Error,
+					FText::Format(LOCTEXT("InvalidInputSummaryFormat", "Invalid Input Override: {0}"), FText::FromString(OverridePin->PinName.ToString())),
+					FText::Format(LOCTEXT("InvalidInputFormat", "The input {0} was previously overriden but is no longer exposed by the function {1}.\nPress the fix button to remove this unused override data,\nor check the function definition to see why this input is no longer exposed."),
+						FText::FromString(OverridePin->PinName.ToString()), FText::FromString(InputFunctionCallNode->GetFunctionName())),
+					GetStackEditorDataKey(),
+					false,
+					GetNodeRemovalFix(OverridePin));
+
+				NewIssues.Add(InvalidInputOverrideError);
+			}
+		}
+		else
+		{
+			// If we have an override pin that is no longer valid, but has the same name and type as a static switch parameter, then it is safe to assume
+			// that the parameter was replaced by the static switch. So we ask the user to copy over its value or remove the override.
+			UEdGraphPin* SwitchPin = *PinReference;
+			bool bIsSameType = OverridePin->PinType.PinCategory == SwitchPin->PinType.PinCategory &&
+				OverridePin->PinType.PinSubCategoryObject == SwitchPin->PinType.PinSubCategoryObject;
+			if (bIsSameType && !ValidAliasedInputNames.Contains(OverridePin->PinName))
+			{
+				TArray<FStackIssueFix> Fixes;
+
+				// first possible fix: convert the value over to the static switch
+				FText ConversionFixDescription = LOCTEXT("ConvertInputToStaticSwitchTransaction", "Copy value to static switch parameter.");
+				FStackIssueFix ConvertInputOverrideFix(
+					ConversionFixDescription,
 					UNiagaraStackEntry::FStackIssueFixDelegate::CreateLambda([=]()
 				{
-					FScopedTransaction ScopedTransaction(FixDescription);
+					FScopedTransaction ScopedTransaction(ConversionFixDescription);
+					SwitchPin->Modify();
+					SwitchPin->DefaultValue = OverridePin->DefaultValue;
+
 					TArray<TWeakObjectPtr<UNiagaraDataInterface>> RemovedDataObjects;
 					FNiagaraStackGraphUtilities::RemoveNodesForStackFunctionInputOverridePin(*OverridePin, RemovedDataObjects);
 					for (TWeakObjectPtr<UNiagaraDataInterface> RemovedDataObject : RemovedDataObjects)
@@ -263,19 +320,25 @@ void UNiagaraStackFunctionInputCollection::RefreshIssues(TArray<FName> Duplicate
 					}
 					OverridePin->GetOwningNode()->RemovePin(OverridePin);
 				}));
+				Fixes.Add(ConvertInputOverrideFix);
 
-				FStackIssue InvalidInputOverrideError(
+				// second possible fix: remove the override completely
+				Fixes.Add(GetNodeRemovalFix(OverridePin));
+
+				FStackIssue DeprecatedInputOverrideError(
 					EStackIssueSeverity::Error,
-					FText::Format(LOCTEXT("InvalidInputSummaryFormat", "Invalid Input Override: {0}"), FText::FromString(OverridePin->PinName.ToString())),
-					FText::Format(LOCTEXT("InvalidInputFormat", "The input {0} was previously overriden but is no longer exposed by the function {1}.\nPress the fix button to remove this unused override data,\nor check the function definition to see why this input is no longer exposed."),
+					FText::Format(LOCTEXT("DeprecatedInputSummaryFormat", "Deprecated Input Override: {0}"), FText::FromString(OverridePin->PinName.ToString())),
+					FText::Format(LOCTEXT("DeprecatedInputFormat", "The input {0} is no longer exposed by the function {1}, but there exists a static switch parameter with the same name instead.\nYou can choose to copy the previously entered data over to the new parameter or remove the override to discard it."),
 						FText::FromString(OverridePin->PinName.ToString()), FText::FromString(InputFunctionCallNode->GetFunctionName())),
 					GetStackEditorDataKey(),
 					false,
-					RemoveInputOverrideFix);
+					Fixes);
 
-				NewIssues.Add(InvalidInputOverrideError);
+				NewIssues.Add(DeprecatedInputOverrideError);
+				break;
 			}
 		}
+		
 	}
 
 	// Generate issues for duplicate input names.
diff --git a/Engine/Plugins/FX/Niagara/Source/NiagaraEditor/Private/ViewModels/Stack/NiagaraStackGraphUtilities.cpp b/Engine/Plugins/FX/Niagara/Source/NiagaraEditor/Private/ViewModels/Stack/NiagaraStackGraphUtilities.cpp
--- a/Engine/Plugins/FX/Niagara/Source/NiagaraEditor/Private/ViewModels/Stack/NiagaraStackGraphUtilities.cpp
+++ b/Engine/Plugins/FX/Niagara/Source/NiagaraEditor/Private/ViewModels/Stack/NiagaraStackGraphUtilities.cpp
@@ -526,7 +526,7 @@ void FNiagaraStackGraphUtilities::GetStackFunctionInputPins(UNiagaraNodeFunction
 	}
 }
 
-void FNiagaraStackGraphUtilities::GetStackFunctionStaticSwitchPins(UNiagaraNodeFunctionCall& FunctionCallNode, TArray<const UEdGraphPin*>& OutInputPins)
+void FNiagaraStackGraphUtilities::GetStackFunctionStaticSwitchPins(UNiagaraNodeFunctionCall& FunctionCallNode, TArray<UEdGraphPin*>& OutInputPins)
 {
 	const UEdGraphSchema_Niagara* Schema = CastChecked<UEdGraphSchema_Niagara>(FunctionCallNode.GetSchema());
 	UNiagaraGraph* FunctionCallGraph = FunctionCallNode.GetCalledGraph();
diff --git a/Engine/Plugins/FX/Niagara/Source/NiagaraEditor/Public/ViewModels/Stack/NiagaraStackFunctionInputCollection.h b/Engine/Plugins/FX/Niagara/Source/NiagaraEditor/Public/ViewModels/Stack/NiagaraStackFunctionInputCollection.h
--- a/Engine/Plugins/FX/Niagara/Source/NiagaraEditor/Public/ViewModels/Stack/NiagaraStackFunctionInputCollection.h
+++ b/Engine/Plugins/FX/Niagara/Source/NiagaraEditor/Public/ViewModels/Stack/NiagaraStackFunctionInputCollection.h
@@ -41,11 +41,12 @@ protected:
 	virtual void RefreshChildrenInternal(const TArray<UNiagaraStackEntry*>& CurrentChildren, TArray<UNiagaraStackEntry*>& NewChildren, TArray<FStackIssue>& NewIssues) override;
 	
 private:
-	void RefreshIssues(TArray<FName> DuplicateInputNames, TArray<FName> ValidAliasedInputNames, TArray<const UEdGraphPin*> PinsWithInvalidTypes, TArray<FStackIssue>& NewIssues);
+	void RefreshIssues(TArray<FName> DuplicateInputNames, TArray<FName> ValidAliasedInputNames, TArray<const UEdGraphPin*> PinsWithInvalidTypes, TMap<FName, UEdGraphPin*> StaticSwitchInputs, TArray<FStackIssue>& NewIssues);
 
 	void OnFunctionInputsChanged();
 
-private:
+	UNiagaraStackEntry::FStackIssueFix GetNodeRemovalFix(UEdGraphPin* PinToRemove);
+
 	struct FInputData
 	{
 		const UEdGraphPin* Pin;
diff --git a/Engine/Plugins/FX/Niagara/Source/NiagaraEditor/Public/ViewModels/Stack/NiagaraStackGraphUtilities.h b/Engine/Plugins/FX/Niagara/Source/NiagaraEditor/Public/ViewModels/Stack/NiagaraStackGraphUtilities.h
--- a/Engine/Plugins/FX/Niagara/Source/NiagaraEditor/Public/ViewModels/Stack/NiagaraStackGraphUtilities.h
+++ b/Engine/Plugins/FX/Niagara/Source/NiagaraEditor/Public/ViewModels/Stack/NiagaraStackGraphUtilities.h
@@ -76,7 +76,7 @@ namespace FNiagaraStackGraphUtilities
 
 	void GetStackFunctionInputPins(UNiagaraNodeFunctionCall& FunctionCallNode, TArray<const UEdGraphPin*>& OutInputPins, ENiagaraGetStackFunctionInputPinsOptions Options = ENiagaraGetStackFunctionInputPinsOptions::AllInputs, bool bIgnoreDisabled = false);
 
-	void GetStackFunctionStaticSwitchPins(UNiagaraNodeFunctionCall& FunctionCallNode, TArray<const UEdGraphPin*>& OutInputPins);
+	void GetStackFunctionStaticSwitchPins(UNiagaraNodeFunctionCall& FunctionCallNode, TArray<UEdGraphPin*>& OutInputPins);
 
 	UNiagaraNodeParameterMapSet* GetStackFunctionOverrideNode(UNiagaraNodeFunctionCall& FunctionCallNode);
 
