Fixed a bug where static switches did either not correctly evaluate the graph or hide unused parameters in the metadata panel

diff --git a/Engine/Plugins/FX/Niagara/Source/NiagaraEditor/Private/NiagaraEditorUtilities.cpp b/Engine/Plugins/FX/Niagara/Source/NiagaraEditor/Private/NiagaraEditorUtilities.cpp
--- a/Engine/Plugins/FX/Niagara/Source/NiagaraEditor/Private/NiagaraEditorUtilities.cpp
+++ b/Engine/Plugins/FX/Niagara/Source/NiagaraEditor/Private/NiagaraEditorUtilities.cpp
@@ -862,6 +862,11 @@ const TMap<FNiagaraVariable, FNiagaraGraphParameterReferenceCollection>& FNiagar
 
 	// We clone the graph to prevent any destructive changes to the source
 	UNiagaraGraph* NodeGraphDeepCopy = CastChecked<UNiagaraGraph>(FEdGraphUtilities::CloneGraph(Graph, GetTransientPackage(), 0, true));
+
+	// We copy some transient data over here because the copied graph loses some data. Since the result of this function is
+	// also used to build the ui, the user will otherwise lose changes to yet unused parameters (e.g. in the metadata panel).
+	CastChecked<UNiagaraGraph>(Graph)->CopyCachedReferencesMap(NodeGraphDeepCopy);
+
 	TArray<UNiagaraNodeFunctionCall*> Nodes;
 	NodeGraphDeepCopy->GetNodesOfClass(Nodes);
 	for (int32 i = 0; i < Nodes.Num(); i++)
diff --git a/Engine/Plugins/FX/Niagara/Source/NiagaraEditor/Private/NiagaraGraph.cpp b/Engine/Plugins/FX/Niagara/Source/NiagaraEditor/Private/NiagaraGraph.cpp
--- a/Engine/Plugins/FX/Niagara/Source/NiagaraEditor/Private/NiagaraGraph.cpp
+++ b/Engine/Plugins/FX/Niagara/Source/NiagaraEditor/Private/NiagaraGraph.cpp
@@ -1115,6 +1115,11 @@ void UNiagaraGraph::RebuildCachedCompileIds(bool bForce)
 	RebuildNumericCache();
 }
 
+void UNiagaraGraph::CopyCachedReferencesMap(UNiagaraGraph* TargetGraph)
+{
+	TargetGraph->ParameterToReferencesMap = ParameterToReferencesMap;
+}
+
 const class UEdGraphSchema_Niagara* UNiagaraGraph::GetNiagaraSchema() const
 {
 	return Cast<UEdGraphSchema_Niagara>(GetSchema());
diff --git a/Engine/Plugins/FX/Niagara/Source/NiagaraEditor/Private/Widgets/SNiagaraParameterMapView.cpp b/Engine/Plugins/FX/Niagara/Source/NiagaraEditor/Private/Widgets/SNiagaraParameterMapView.cpp
--- a/Engine/Plugins/FX/Niagara/Source/NiagaraEditor/Private/Widgets/SNiagaraParameterMapView.cpp
+++ b/Engine/Plugins/FX/Niagara/Source/NiagaraEditor/Private/Widgets/SNiagaraParameterMapView.cpp
@@ -332,8 +332,7 @@ void SNiagaraParameterMapView::CollectAllActions(FGraphActionListBuilderBase& Ou
 	for (auto& GraphWeakPtr : Graphs)
 	{
 		UNiagaraGraph* Graph = GraphWeakPtr.Get();
-		for (const auto& ParameterElement : Graph->GetParameterReferenceMap())
-		// for (const auto& ParameterElement : FNiagaraEditorUtilities::GetCompiledGraphParameterMapReferences(Graph)) // TODO(mv): This makes unused module parameters not show up. See CL 6182457.
+		for (const auto& ParameterElement : FNiagaraEditorUtilities::GetCompiledGraphParameterMapReferences(Graph))
 		{
 			TArray<FNiagaraGraphParameterReferenceCollection>* Found = ParameterEntries.Find(ParameterElement.Key);
 			if (Found)
diff --git a/Engine/Plugins/FX/Niagara/Source/NiagaraEditor/Public/NiagaraGraph.h b/Engine/Plugins/FX/Niagara/Source/NiagaraEditor/Public/NiagaraGraph.h
--- a/Engine/Plugins/FX/Niagara/Source/NiagaraEditor/Public/NiagaraGraph.h
+++ b/Engine/Plugins/FX/Niagara/Source/NiagaraEditor/Public/NiagaraGraph.h
@@ -265,6 +265,8 @@ class UNiagaraGraph : public UEdGraph
 
 	void RebuildCachedCompileIds(bool bForce = false);
 
+	void CopyCachedReferencesMap(UNiagaraGraph* TargetGraph);
+
 protected:
 	void RebuildNumericCache();
 	bool bNeedNumericCacheRebuilt;
