Added support for static switch default values

diff --git a/Engine/Plugins/FX/Niagara/Source/Niagara/Public/NiagaraTypes.h b/Engine/Plugins/FX/Niagara/Source/Niagara/Public/NiagaraTypes.h
--- a/Engine/Plugins/FX/Niagara/Source/Niagara/Public/NiagaraTypes.h
+++ b/Engine/Plugins/FX/Niagara/Source/Niagara/Public/NiagaraTypes.h
@@ -356,6 +356,10 @@ public:
 
 	UPROPERTY(AdvancedDisplay, VisibleAnywhere, Category = "Variable")
 	bool bIsStaticSwitch;
+
+	/** The default value to use when creating new pins or stack entries for a static switch parameter */
+	UPROPERTY()
+	int32 StaticSwitchDefaultValue;
 };
 
 USTRUCT()
diff --git a/Engine/Plugins/FX/Niagara/Source/NiagaraEditor/Private/Customizations/NiagaraStaticSwitchNodeDetails.cpp b/Engine/Plugins/FX/Niagara/Source/NiagaraEditor/Private/Customizations/NiagaraStaticSwitchNodeDetails.cpp
--- a/Engine/Plugins/FX/Niagara/Source/NiagaraEditor/Private/Customizations/NiagaraStaticSwitchNodeDetails.cpp
+++ b/Engine/Plugins/FX/Niagara/Source/NiagaraEditor/Private/Customizations/NiagaraStaticSwitchNodeDetails.cpp
@@ -12,6 +12,8 @@
 #include "Widgets/Input/SComboBox.h"
 #include "Widgets/Input/SNumericEntryBox.h"
 #include "Widgets/Input/SEditableTextBox.h"
+#include "Widgets/Input/SCheckBox.h"
+#include "Widgets/Layout/SWidgetSwitcher.h"
 
 #define LOCTEXT_NAMESPACE "NiagaraStaticSwitchNodeDetails"
 
@@ -48,6 +50,9 @@ void FNiagaraStaticSwitchNodeDetails::CustomizeDetails(IDetailLayoutBuilder& Det
 		
 		IDetailCategoryBuilder& CategoryBuilder = DetailBuilder.EditCategory(SwitchCategoryName);		
 		FDetailWidgetRow& NameWidget = CategoryBuilder.AddCustomRow(LOCTEXT("NiagaraSwitchNodeNameFilterText", "Input parameter name"));
+		FDetailWidgetRow& DropdownWidget = CategoryBuilder.AddCustomRow(LOCTEXT("NiagaraSwitchNodeTypeFilterText", "Input parameter type"));
+		FDetailWidgetRow& IntValueOption = CategoryBuilder.AddCustomRow(LOCTEXT("NiagaraSwitchNodeIntFilterText", "Max integer value"));
+		FDetailWidgetRow& DefaultValueOption = CategoryBuilder.AddCustomRow(LOCTEXT("NiagaraSwitchNodeDefaultFilterText", "Default value"), true);
 
 		NameWidget
 		.NameContent()
@@ -74,8 +79,6 @@ void FNiagaraStaticSwitchNodeDetails::CustomizeDetails(IDetailLayoutBuilder& Det
 			]
 		];
 
-		FDetailWidgetRow& DropdownWidget = CategoryBuilder.AddCustomRow(LOCTEXT("NiagaraSwitchNodeTypeFilterText", "Input parameter type"));
-
 		DropdownWidget
 		.NameContent()
 		[
@@ -108,8 +111,6 @@ void FNiagaraStaticSwitchNodeDetails::CustomizeDetails(IDetailLayoutBuilder& Det
 			]
 		];
 
-		FDetailWidgetRow& IntValueOption = CategoryBuilder.AddCustomRow(LOCTEXT("NiagaraSwitchNodeIntFilterText", "Max integer value"));
-
 		IntValueOption
 		.NameContent()
 		[
@@ -136,6 +137,96 @@ void FNiagaraStaticSwitchNodeDetails::CustomizeDetails(IDetailLayoutBuilder& Det
 				.OnValueCommitted(this, &FNiagaraStaticSwitchNodeDetails::IntOptionValueCommitted)
 			]
 		];
+
+		DefaultValueOption
+		.NameContent()
+		[
+			SNew(SBox)
+			.Padding(FMargin(0.0f, 2.0f))
+			[
+				SNew(STextBlock)
+				.TextStyle(FNiagaraEditorStyle::Get(), "NiagaraEditor.ParameterText")
+				.Text(LOCTEXT("NiagaraSwitchNodeDefaultOptionText", "Default value"))
+			]
+		]
+		.ValueContent()
+		[
+			SNew(SBox)
+			.Padding(FMargin(0.0f, 2.0f))
+			[
+				SNew(SWidgetSwitcher)
+		  		.WidgetIndex(this, &FNiagaraStaticSwitchNodeDetails::GetDefaultWidgetIndex)
+
+		  		+ SWidgetSwitcher::Slot()
+		  		[
+					SNew(SCheckBox)
+					.IsChecked_Lambda([this]() { return GetSwitchDefaultValue().Get(0) ? ECheckBoxState::Checked : ECheckBoxState::Unchecked; })
+					.OnCheckStateChanged(this, &FNiagaraStaticSwitchNodeDetails::DefaultBoolValueCommitted)
+				]
+
+		  		+ SWidgetSwitcher::Slot()
+		  		[
+					SNew(SNumericEntryBox<int32>)
+					.AllowSpin(false)
+					.MinValue(0)
+					.MaxValue(99)
+					.Value(this, &FNiagaraStaticSwitchNodeDetails::GetSwitchDefaultValue)
+					.OnValueCommitted(this, &FNiagaraStaticSwitchNodeDetails::DefaultIntValueCommitted)
+				]
+				
+				+ SWidgetSwitcher::Slot()
+		  		[
+					SNew(SComboBox<TSharedPtr<FString>>)
+					.OptionsSource(&DefaultEnumDropdownOptions)
+					.OnSelectionChanged(this, &FNiagaraStaticSwitchNodeDetails::OnSelectionChanged)
+					.OnGenerateWidget(this, &FNiagaraStaticSwitchNodeDetails::CreateWidgetForDropdownOption)
+					.InitiallySelectedItem(SelectedDefaultValue)
+					[
+						SNew(STextBlock)
+						.Margin(FMargin(0.0f, 2.0f))
+						.Text(this, &FNiagaraStaticSwitchNodeDetails::GetDefaultSelectionItemLabel)
+					]
+				]
+			]
+		];
+
+		RefreshDefaultDropdownValues();
+	}
+}
+
+int32 FNiagaraStaticSwitchNodeDetails::GetDefaultWidgetIndex() const
+{
+	if (!Node.IsValid())
+	{
+		return 0;
+	}
+	ENiagaraStaticSwitchType Type = Node->SwitchTypeData.SwitchType;
+	return Type == ENiagaraStaticSwitchType::Bool ? 0 : (Type == ENiagaraStaticSwitchType::Integer ? 1 : 2);
+}
+
+TOptional<int32> FNiagaraStaticSwitchNodeDetails::GetSwitchDefaultValue() const
+{
+	TOptional<FNiagaraVariableMetaData> MetaData = GetSwitchParameterMetadata();
+	return MetaData.IsSet() ? TOptional<int32>(MetaData->StaticSwitchDefaultValue) : TOptional<int32>();
+}
+
+void FNiagaraStaticSwitchNodeDetails::DefaultIntValueCommitted(int32 Value, ETextCommit::Type CommitInfo)
+{
+	TOptional<FNiagaraVariableMetaData> MetaData = GetSwitchParameterMetadata();
+	if (MetaData.IsSet())
+	{
+		MetaData->StaticSwitchDefaultValue = Value;
+		SetSwitchParameterMetadata(MetaData.GetValue());
+	}
+}
+
+void FNiagaraStaticSwitchNodeDetails::DefaultBoolValueCommitted(ECheckBoxState NewState)
+{
+	TOptional<FNiagaraVariableMetaData> MetaData = GetSwitchParameterMetadata();
+	if (MetaData.IsSet())
+	{
+		MetaData->StaticSwitchDefaultValue = (NewState == ECheckBoxState::Checked) ? 1 : 0;
+		SetSwitchParameterMetadata(MetaData.GetValue());
 	}
 }
 
@@ -144,6 +235,31 @@ TSharedRef<SWidget> FNiagaraStaticSwitchNodeDetails::CreateWidgetForDropdownOpti
 	return SNew(STextBlock).Text(FText::FromString(*InOption->Name));
 }
 
+TSharedRef<SWidget> FNiagaraStaticSwitchNodeDetails::CreateWidgetForDropdownOption(TSharedPtr<FString> InOption)
+{
+	return SNew(STextBlock).Text(FText::FromString(*InOption));
+}
+
+void FNiagaraStaticSwitchNodeDetails::OnSelectionChanged(TSharedPtr<FString> NewValue, ESelectInfo::Type)
+{
+	SelectedDefaultValue = NewValue;
+	TOptional<FNiagaraVariableMetaData> MetaData = GetSwitchParameterMetadata();
+	if (!SelectedDefaultValue.IsValid() || !MetaData.IsSet())
+	{
+		return;
+	}
+
+	UEnum* Enum = Node->SwitchTypeData.Enum;
+	if (!Enum)
+	{
+		return;
+	}
+
+	int32 Index = Enum->GetIndexByNameString(*SelectedDefaultValue);
+	MetaData->StaticSwitchDefaultValue = Index;
+	SetSwitchParameterMetadata(MetaData.GetValue());
+}
+
 void FNiagaraStaticSwitchNodeDetails::OnSelectionChanged(TSharedPtr<SwitchDropdownOption> NewValue, ESelectInfo::Type)
 {
 	SelectedDropdownItem = NewValue;
@@ -156,10 +272,12 @@ void FNiagaraStaticSwitchNodeDetails::OnSelectionChanged(TSharedPtr<SwitchDropdo
 	if (SelectedDropdownItem == DropdownOptions[0])
 	{
 		Node->SwitchTypeData.SwitchType = ENiagaraStaticSwitchType::Bool;
+		Node->SwitchTypeData.Enum = nullptr;
 	}
 	else if (SelectedDropdownItem == DropdownOptions[1])
 	{
 		Node->SwitchTypeData.SwitchType = ENiagaraStaticSwitchType::Integer;
+		Node->SwitchTypeData.Enum = nullptr;
 	}
 	else
 	{
@@ -167,6 +285,7 @@ void FNiagaraStaticSwitchNodeDetails::OnSelectionChanged(TSharedPtr<SwitchDropdo
 		Node->SwitchTypeData.Enum = SelectedDropdownItem->Enum;
 	}
 	Node->OnSwitchParameterTypeChanged(OldType);
+	RefreshDefaultDropdownValues();
 }
 
 FText FNiagaraStaticSwitchNodeDetails::GetDropdownItemLabel() const
@@ -179,6 +298,64 @@ FText FNiagaraStaticSwitchNodeDetails::GetDropdownItemLabel() const
 	return LOCTEXT("InvalidNiagaraStaticSwitchNodeComboEntryText", "<Invalid selection>");
 }
 
+FText FNiagaraStaticSwitchNodeDetails::GetDefaultSelectionItemLabel() const
+{
+	if (SelectedDefaultValue.IsValid())
+	{
+		return FText::FromString(*SelectedDefaultValue);
+	}
+
+	return LOCTEXT("InvalidNiagaraStaticSwitchNodeComboEntryText", "<Invalid selection>");
+}
+
+void FNiagaraStaticSwitchNodeDetails::RefreshDefaultDropdownValues()
+{
+	if (!Node.IsValid() || Node->SwitchTypeData.SwitchType != ENiagaraStaticSwitchType::Enum)
+	{
+		return;
+	}
+	TOptional<FNiagaraVariableMetaData> MetaData = GetSwitchParameterMetadata();
+
+	DefaultEnumDropdownOptions.Empty();
+	UEnum* Enum = Node->SwitchTypeData.Enum;
+	if (Enum)
+	{
+		SelectedDefaultValue.Reset();
+		for (int i = 0; i < Enum->GetMaxEnumValue() - 1; i++)
+		{
+			DefaultEnumDropdownOptions.Add(MakeShareable(new FString(Enum->GetNameStringByIndex(i))));
+
+			if (MetaData.IsSet() && i == MetaData->StaticSwitchDefaultValue)
+			{
+				SelectedDefaultValue = DefaultEnumDropdownOptions[i];
+			}
+		}
+		if (!SelectedDefaultValue.IsValid() && DefaultEnumDropdownOptions.Num() > 0)
+		{
+			SelectedDefaultValue = DefaultEnumDropdownOptions[0];
+		}
+	}
+}
+
+TOptional<FNiagaraVariableMetaData> FNiagaraStaticSwitchNodeDetails::GetSwitchParameterMetadata() const
+{
+	if (!Node.IsValid() || !Node->GetNiagaraGraph())
+	{
+		TOptional<FNiagaraVariableMetaData> Empty;
+		return Empty;
+	}
+	return Node->GetNiagaraGraph()->GetMetaData(FNiagaraVariable(Node->GetInputType(), Node->InputParameterName));
+}
+
+void FNiagaraStaticSwitchNodeDetails::SetSwitchParameterMetadata(const FNiagaraVariableMetaData& MetaData)
+{
+	if (!Node.IsValid() || !Node->GetNiagaraGraph())
+	{
+		return;
+	}
+	Node->GetNiagaraGraph()->SetMetaData(FNiagaraVariable(Node->GetInputType(), Node->InputParameterName), MetaData);
+}
+
 void FNiagaraStaticSwitchNodeDetails::UpdateSelectionFromNode()
 {
 	SelectedDropdownItem.Reset();
diff --git a/Engine/Plugins/FX/Niagara/Source/NiagaraEditor/Private/Customizations/NiagaraStaticSwitchNodeDetails.h b/Engine/Plugins/FX/Niagara/Source/NiagaraEditor/Private/Customizations/NiagaraStaticSwitchNodeDetails.h
--- a/Engine/Plugins/FX/Niagara/Source/NiagaraEditor/Private/Customizations/NiagaraStaticSwitchNodeDetails.h
+++ b/Engine/Plugins/FX/Niagara/Source/NiagaraEditor/Private/Customizations/NiagaraStaticSwitchNodeDetails.h
@@ -6,6 +6,7 @@
 #include "IDetailCustomization.h"
 #include "DetailLayoutBuilder.h"
 #include "Misc/Optional.h"
+#include "NiagaraTypes.h"
 
 // This data structure is used internally by the dropdown to keep track of the user's choice
 struct SwitchDropdownOption
@@ -34,6 +35,7 @@ public:
 	virtual void CustomizeDetails( IDetailLayoutBuilder& DetailBuilder ) override;
 	
 private:
+	// enum dropdown functions
 	TSharedRef<SWidget> CreateWidgetForDropdownOption(TSharedPtr<SwitchDropdownOption> InOption);
 	void OnSelectionChanged(TSharedPtr<SwitchDropdownOption> NewValue, ESelectInfo::Type);
 	FText GetDropdownItemLabel() const;
@@ -48,7 +50,22 @@ private:
 	FText GetParameterNameText() const;
 	void OnParameterNameCommited(const FText& InText, ETextCommit::Type InCommitType);
 
+	// default value option functions
+	int32 GetDefaultWidgetIndex() const;
+	TOptional<int32> GetSwitchDefaultValue() const;
+	void DefaultIntValueCommitted(int32 Value, ETextCommit::Type CommitInfo);
+	void DefaultBoolValueCommitted(ECheckBoxState NewState);
+	TSharedRef<SWidget> CreateWidgetForDropdownOption(TSharedPtr<FString> InOption);
+	void OnSelectionChanged(TSharedPtr<FString> NewValue, ESelectInfo::Type);
+	FText GetDefaultSelectionItemLabel() const;
+	void RefreshDefaultDropdownValues();
+
+	TOptional<FNiagaraVariableMetaData> GetSwitchParameterMetadata() const;
+	void SetSwitchParameterMetadata(const FNiagaraVariableMetaData& MetaData);
+
 	TWeakObjectPtr<class UNiagaraNodeStaticSwitch> Node;
 	TArray<TSharedPtr<SwitchDropdownOption>> DropdownOptions;
 	TSharedPtr<SwitchDropdownOption> SelectedDropdownItem;
+	TArray<TSharedPtr<FString>> DefaultEnumDropdownOptions;
+	TSharedPtr<FString> SelectedDefaultValue;
 };
diff --git a/Engine/Plugins/FX/Niagara/Source/NiagaraEditor/Private/NiagaraNodeFunctionCall.cpp b/Engine/Plugins/FX/Niagara/Source/NiagaraEditor/Private/NiagaraNodeFunctionCall.cpp
--- a/Engine/Plugins/FX/Niagara/Source/NiagaraEditor/Private/NiagaraNodeFunctionCall.cpp
+++ b/Engine/Plugins/FX/Niagara/Source/NiagaraEditor/Private/NiagaraNodeFunctionCall.cpp
@@ -155,18 +155,19 @@ void UNiagaraNodeFunctionCall::AllocateDefaultPins()
 		}
 
 		TArray<FNiagaraVariable> SwitchNodeInputs = Graph->FindStaticSwitchInputs();
-		for (const FNiagaraVariable& Input : SwitchNodeInputs)
+		for (FNiagaraVariable& Input : SwitchNodeInputs)
 		{
 			UEdGraphPin* NewPin = CreatePin(EGPD_Input, Schema->TypeDefinitionToPinType(Input.GetType()), Input.GetName());
-			NewPin->bDefaultValueIsIgnored = false;
 			NewPin->bNotConnectable = true;
+			NewPin->bDefaultValueIsIgnored = FindPropagatedVariable(Input) != nullptr;
 
+			int32 DefaultValue = Graph->GetMetaData(Input).Get(FNiagaraVariableMetaData()).StaticSwitchDefaultValue;
+			Input.SetValue<FNiagaraInt32>({ DefaultValue });
 			FString PinDefaultValue;
 			if (Schema->TryGetPinDefaultValueFromNiagaraVariable(Input, PinDefaultValue))
 			{
 				NewPin->DefaultValue = PinDefaultValue;
 			}
-			NewPin->bDefaultValueIsIgnored = FindPropagatedVariable(Input) != nullptr;
 		}
 
 		for (FNiagaraVariable& Output : Outputs)
