Correctly initialize variable (previous cl did not save correctly)

diff --git a/Engine/Plugins/FX/Niagara/Source/Niagara/Public/NiagaraTypes.h b/Engine/Plugins/FX/Niagara/Source/Niagara/Public/NiagaraTypes.h
--- a/Engine/Plugins/FX/Niagara/Source/Niagara/Public/NiagaraTypes.h
+++ b/Engine/Plugins/FX/Niagara/Source/Niagara/Public/NiagaraTypes.h
@@ -322,6 +322,7 @@ public:
 		, EditorSortPriority(0)
 		, bInlineEditConditionToggle(false)
 		, bIsStaticSwitch(false)
+		, StaticSwitchDefaultValue(0)
 	{
 	}
 public:
