- Niagara fix for GetSpawnInterpolation always returning 1 on a GPU simulation when interpolated spawning is enabled

diff --git a/Engine/Plugins/FX/Niagara/Source/NiagaraEditor/Private/NiagaraHlslTranslator.cpp b/Engine/Plugins/FX/Niagara/Source/NiagaraEditor/Private/NiagaraHlslTranslator.cpp
--- a/Engine/Plugins/FX/Niagara/Source/NiagaraEditor/Private/NiagaraHlslTranslator.cpp
+++ b/Engine/Plugins/FX/Niagara/Source/NiagaraEditor/Private/NiagaraHlslTranslator.cpp
@@ -409,10 +409,17 @@ FString FHlslNiagaraTranslator::GetFunctionDefinitions()
 		// data interface functions that should be defined differently.
 	}
 
+	// Check to see if we have interpolated spawn enabled, for the GPU we need to look for the additional defines
+	bool bHasInterpolatedSpawn = CompileOptions.TargetUsage == ENiagaraScriptUsage::ParticleSpawnScriptInterpolated;
+	if ( CompileOptions.TargetUsage == ENiagaraScriptUsage::ParticleGPUComputeScript )
+	{
+		bHasInterpolatedSpawn = CompileOptions.AdditionalDefines.Contains(TEXT("InterpolatedSpawn"));
+	}
+
 	//Add a few hard coded helper functions in.
 	FwdDeclString += TEXT("float GetSpawnInterpolation();");
 	//Add helper function to get the interpolation factor.
-	if (CompileOptions.TargetUsage == ENiagaraScriptUsage::ParticleSpawnScriptInterpolated)
+	if ( bHasInterpolatedSpawn )
 	{
 		DefinitionsString += TEXT("float GetSpawnInterpolation()\n{\n");
 		DefinitionsString += TEXT("\treturn HackSpawnInterp;\n");
