- Move Niagara skeletal mesh HLSL from code to a separate shader file
- We currently generate a small function to call into the shader file code with a struct to pass the data, this allows for faster iteration on shader code as it can be done inside the editor
- Implement missing GPU functionality for triangle / bone / vertex skeletal mesh sampling, specific bone / specific sockets

diff --git a/Engine/Plugins/FX/Niagara/Shaders/Private/NiagaraDataInterfaceSkeletalMesh.ush b/Engine/Plugins/FX/Niagara/Shaders/Private/NiagaraDataInterfaceSkeletalMesh.ush
new file mode 100644
--- /dev/null
+++ b/Engine/Plugins/FX/Niagara/Shaders/Private/NiagaraDataInterfaceSkeletalMesh.ush
@@ -0,0 +1,501 @@
+// Copyright 1998-2019 Epic Games, Inc. All Rights Reserved.
+
+#pragma once
+
+#ifndef DISKELMESH_USE_EXTRA_INFLUENCES
+#define DISKELMESH_USE_EXTRA_INFLUENCES	0			// Default is 4 bone, this enables support for 4 & 8 bones
+#endif
+
+struct FDISkelMeshContext
+{
+	Buffer<uint>	MeshIndexBuffer;
+	Buffer<float>	MeshVertexBuffer;
+	Buffer<uint>	MeshSkinWeightBuffer;
+	Buffer<float4>	MeshCurrBonesBuffer;
+	Buffer<float4>	MeshPrevBonesBuffer;
+	Buffer<float4>	MeshCurrSamplingBonesBuffer;
+    Buffer<float4>	MeshPrevSamplingBonesBuffer;
+	Buffer<float4>	MeshTangentBuffer;
+	Buffer<float2>	MeshTexCoordBuffer;
+    Buffer<float4>	MeshColorBuffer;
+	Buffer<float>	MeshTriangleSamplerProbaBuffer;
+	Buffer<uint>	MeshTriangleSamplerAliasBuffer;
+	Buffer<uint>	MeshTriangleMatricesOffsetBuffer;
+	uint			MeshTriangleCount;
+	uint			MeshVertexCount;
+    uint			MeshWeightStride;
+	uint			MeshNumTexCoord;
+	uint			MeshNumWeights;
+	uint			NumSpecificBones;
+	Buffer<uint>	SpecificBones;
+	uint			NumSpecificSocketBones;
+	Buffer<uint>	SpecificSocketBones;
+	float4x4		InstanceTransform;
+	float4x4		InstancePrevTransform;
+	float			InstanceInvDeltaTime;
+	bool			UniformTriangleSamplingEnable;
+};
+
+////////////////////////////////////////////////////////////////////////////////////////////////////////////////////////////////////////////////////////////////
+// Triangle Sampling
+
+struct FDISkelMeshSkinnedVertex
+{
+	float3	Position;
+	float3	PrevPosition;
+	float3	TangentX;
+	float3	TangentY;
+	float3	TangentZ;
+};
+
+float3x4 DISkelMesh_GetPrevBoneSkinningMatrix(in FDISkelMeshContext DIContext, uint Bone)
+{
+    return float3x4(DIContext.MeshPrevBonesBuffer[Bone * 3], DIContext.MeshPrevBonesBuffer[Bone * 3 + 1], DIContext.MeshPrevBonesBuffer[Bone * 3 + 2]);
+}
+
+float3x4 DISkelMesh_GetPrevSkinningMatrix(in FDISkelMeshContext DIContext, uint VertexIndex, int4 BlendIndices, float4 BlendWeights)
+{
+	// Get the matrix offset for each vertex because BlendIndices are stored relatively to each section start vertex.
+    uint MatrixOffset = DIContext.MeshTriangleMatricesOffsetBuffer[VertexIndex];
+
+	float3x4 Result;
+    Result = DISkelMesh_GetPrevBoneSkinningMatrix(DIContext, MatrixOffset + BlendIndices.x) * BlendWeights.x;
+    Result += DISkelMesh_GetPrevBoneSkinningMatrix(DIContext, MatrixOffset + BlendIndices.y) * BlendWeights.y;
+    Result += DISkelMesh_GetPrevBoneSkinningMatrix(DIContext, MatrixOffset + BlendIndices.z) * BlendWeights.z;
+    Result += DISkelMesh_GetPrevBoneSkinningMatrix(DIContext, MatrixOffset + BlendIndices.w) * BlendWeights.w;
+	return Result;
+}
+
+float3x4 DISkelMesh_GetCurrBoneSkinningMatrix(in FDISkelMeshContext DIContext, uint Bone)
+{
+	return float3x4(DIContext.MeshCurrBonesBuffer[Bone * 3], DIContext.MeshCurrBonesBuffer[Bone * 3 + 1], DIContext.MeshCurrBonesBuffer[Bone * 3 + 2]);
+}
+
+float3x4 DISkelMesh_GetCurrSkinningMatrix(in FDISkelMeshContext DIContext, uint VertexIndex, int4 BlendIndices, float4 BlendWeights)
+{
+	// Get the matrix offset for each vertex because BlendIndices are stored relatively to each section start vertex.
+	uint MatrixOffset = DIContext.MeshTriangleMatricesOffsetBuffer[VertexIndex];
+
+	float3x4 Result;
+    Result = DISkelMesh_GetCurrBoneSkinningMatrix(DIContext, MatrixOffset + BlendIndices.x) * BlendWeights.x;
+    Result += DISkelMesh_GetCurrBoneSkinningMatrix(DIContext, MatrixOffset + BlendIndices.y) * BlendWeights.y;
+    Result += DISkelMesh_GetCurrBoneSkinningMatrix(DIContext, MatrixOffset + BlendIndices.z) * BlendWeights.z;
+    Result += DISkelMesh_GetCurrBoneSkinningMatrix(DIContext, MatrixOffset + BlendIndices.w) * BlendWeights.w;
+	return Result;
+}
+
+int4 DISkelMesh_UnpackIndices4(uint Packed)
+{
+	return int4(Packed & 0xff, Packed >> 8 & 0xff, Packed >> 16 & 0xff, Packed >> 24 & 0xff);
+}
+
+float4 DISkelMesh_UnpackWeights4(uint Packed)
+{
+	return float4(Packed & 0xff, Packed >> 8 & 0xff, Packed >> 16 & 0xff, Packed >> 24 & 0xff) / 255.0f;
+}
+
+FDISkelMeshSkinnedVertex DISkelMesh_GetSkinnedVertex(in FDISkelMeshContext DIContext, uint VertexIndex)
+{
+	//-TODO: Use a R32G32B32 format
+	float3 CurrPosition = float3(DIContext.MeshVertexBuffer[VertexIndex * 3], DIContext.MeshVertexBuffer[VertexIndex * 3 + 1], DIContext.MeshVertexBuffer[VertexIndex * 3 + 2]);
+	float3 PrevPosition = CurrPosition;
+
+	float3 TangentX = TangentBias(DIContext.MeshTangentBuffer[VertexIndex * 2    ].xyz);
+	float4 TangentZ = TangentBias(DIContext.MeshTangentBuffer[VertexIndex * 2 + 1].xyzw);
+	
+	if ( DIContext.MeshNumWeights > 0 )
+	{
+		float3x4 PrevBoneMatrix;
+		float3x4 CurrBoneMatrix;
+
+#if DISKELMESH_USE_EXTRA_INFLUENCES
+        uint WeightOffset = DIContext.MeshNumWeights == 8 ? 2 : 1;
+
+        int4 BlendIndices = DISkelMesh_UnpackIndices4(DIContext.MeshSkinWeightBuffer[VertexIndex * DIContext.MeshWeightStride]);
+        float4 BlendWeights = DISkelMesh_UnpackWeights4(DIContext.MeshSkinWeightBuffer[VertexIndex * DIContext.MeshWeightStride + WeightOffset]);
+        CurrBoneMatrix = DISkelMesh_GetCurrSkinningMatrix(DIContext, VertexIndex, BlendIndices, BlendWeights);
+        PrevBoneMatrix = DISkelMesh_GetPrevSkinningMatrix(DIContext, VertexIndex, BlendIndices, BlendWeights);
+		if (DIContext.MeshNumWeights == 8)
+        {
+            BlendIndices = DISkelMesh_UnpackIndices4(DIContext.MeshSkinWeightBuffer[VertexIndex * DIContext.MeshWeightStride + 2]);
+            BlendWeights = DISkelMesh_UnpackWeights4(DIContext.MeshSkinWeightBuffer[VertexIndex * DIContext.MeshWeightStride + WeightOffset + 1]);
+			CurrBoneMatrix += DISkelMesh_GetCurrSkinningMatrix(DIContext, VertexIndex, BlendIndices, BlendWeights);
+			PrevBoneMatrix += DISkelMesh_GetPrevSkinningMatrix(DIContext, VertexIndex, BlendIndices, BlendWeights);
+        }
+#else
+        int4 BlendIndices = DISkelMesh_UnpackIndices4(DIContext.MeshSkinWeightBuffer[VertexIndex * DIContext.MeshWeightStride]);
+        float4 BlendWeights = DISkelMesh_UnpackWeights4(DIContext.MeshSkinWeightBuffer[VertexIndex * DIContext.MeshWeightStride + 1]);
+        CurrBoneMatrix = DISkelMesh_GetCurrSkinningMatrix(DIContext, VertexIndex, BlendIndices, BlendWeights);
+        PrevBoneMatrix = DISkelMesh_GetPrevSkinningMatrix(DIContext, VertexIndex, BlendIndices, BlendWeights);
+#endif
+
+		CurrPosition = mul(CurrBoneMatrix, float4(CurrPosition, 1.0f)).xyz;
+		PrevPosition = mul(PrevBoneMatrix, float4(PrevPosition, 1.0f)).xyz;
+
+		// Not using InverseTranspose of matrices so assuming uniform scaling only (same as SkinCache)
+		TangentX.xyz = mul(CurrBoneMatrix, float4(TangentX.xyz, 0.0f)).xyz;
+		TangentZ.xyz = mul(CurrBoneMatrix, float4(TangentZ.xyz, 0.0f)).xyz;
+	}
+
+	FDISkelMeshSkinnedVertex SkinnedVertex;
+	SkinnedVertex.Position = CurrPosition;
+	SkinnedVertex.PrevPosition = PrevPosition;
+	SkinnedVertex.TangentX = TangentX;
+	SkinnedVertex.TangentY = cross(TangentZ.xyz, TangentX.xyz);
+	SkinnedVertex.TangentZ = TangentZ.xyz;
+	return SkinnedVertex;
+}
+
+FDISkelMeshSkinnedVertex DISkelMesh_GetSkinnedPointOnTriangle(in FDISkelMeshContext DIContext, uint TriangleIndex, float3 BaryCoord)
+{
+	uint IndexBufferOffset = TriangleIndex * 3;
+	uint VertexIndex0 = DIContext.MeshIndexBuffer[IndexBufferOffset];
+	uint VertexIndex1 = DIContext.MeshIndexBuffer[IndexBufferOffset + 1];
+	uint VertexIndex2 = DIContext.MeshIndexBuffer[IndexBufferOffset + 2];
+
+	FDISkelMeshSkinnedVertex SkinnedVertex0 = DISkelMesh_GetSkinnedVertex(DIContext, VertexIndex0);
+	FDISkelMeshSkinnedVertex SkinnedVertex1 = DISkelMesh_GetSkinnedVertex(DIContext, VertexIndex1);
+	FDISkelMeshSkinnedVertex SkinnedVertex2 = DISkelMesh_GetSkinnedVertex(DIContext, VertexIndex2);
+
+	FDISkelMeshSkinnedVertex FinalVertex;
+	FinalVertex.Position		= (SkinnedVertex0.Position * BaryCoord.x) + (SkinnedVertex1.Position * BaryCoord.y) + (SkinnedVertex2.Position * BaryCoord.z);
+	FinalVertex.PrevPosition	= (SkinnedVertex0.PrevPosition * BaryCoord.x) + (SkinnedVertex1.PrevPosition * BaryCoord.y) + (SkinnedVertex2.PrevPosition * BaryCoord.z);
+	FinalVertex.TangentX		= (SkinnedVertex0.TangentX * BaryCoord.x) + (SkinnedVertex1.TangentX * BaryCoord.y) + (SkinnedVertex2.TangentX * BaryCoord.z);
+	FinalVertex.TangentY		= (SkinnedVertex0.TangentY * BaryCoord.x) + (SkinnedVertex1.TangentY * BaryCoord.y) + (SkinnedVertex2.TangentY * BaryCoord.z);
+	FinalVertex.TangentZ		= (SkinnedVertex0.TangentZ * BaryCoord.x) + (SkinnedVertex1.TangentZ * BaryCoord.y) + (SkinnedVertex2.TangentZ * BaryCoord.z);
+
+	return FinalVertex;
+}
+
+FDISkelMeshSkinnedVertex DISkelMesh_GetSkinnedPointOnTriangleWS(in FDISkelMeshContext DIContext, uint TriangleIndex, float3 BaryCoord)
+{
+	FDISkelMeshSkinnedVertex FinalVertex = DISkelMesh_GetSkinnedPointOnTriangle(DIContext, TriangleIndex, BaryCoord);
+	FinalVertex.Position		= mul(float4(FinalVertex.Position, 1.0f), DIContext.InstanceTransform).xyz;
+	FinalVertex.PrevPosition	= mul(float4(FinalVertex.PrevPosition, 1.0f), DIContext.InstancePrevTransform).xyz;
+	FinalVertex.TangentX		= mul(float4(FinalVertex.TangentX, 0.0f), DIContext.InstanceTransform).xyz;
+	FinalVertex.TangentY		= mul(float4(FinalVertex.TangentY, 0.0f), DIContext.InstanceTransform).xyz;
+	FinalVertex.TangentZ		= mul(float4(FinalVertex.TangentZ, 0.0f), DIContext.InstanceTransform).xyz;
+
+	return FinalVertex;
+}
+
+void DISKelMesh_RandomTriCoord(in FDISkelMeshContext DIContext, out uint OutTriangle, out float3 OutBaryCoord)
+{
+	float RandT0 = NiagaraInternalNoise(1, 2, 3);
+
+	[branch]
+	if ( !DIContext.UniformTriangleSamplingEnable )
+	{
+		// Uniform triangle id selection
+		OutTriangle = min(uint(RandT0 * float(DIContext.MeshTriangleCount)), DIContext.MeshTriangleCount - 1); // avoid % by using mul/min to Tri = MeshTriangleCountName
+	}
+	else
+	{
+		// Uniform area weighted position selection (using alias method from Alias method from FWeightedRandomSampler)
+		uint TriangleIndex = min(uint(RandT0*float(DIContext.MeshTriangleCount)), DIContext.MeshTriangleCount - 1);
+		float TriangleProbability = DIContext.MeshTriangleSamplerProbaBuffer[TriangleIndex];
+
+		// Alias check
+		float RandT1 = NiagaraInternalNoise(1, 2, 3);
+		if( RandT1 > TriangleProbability )
+		{
+			TriangleIndex = DIContext.MeshTriangleSamplerAliasBuffer[TriangleIndex];
+		}
+		OutTriangle = TriangleIndex;
+	}
+
+	float r0 = NiagaraInternalNoise(1, 2, 3);
+	float r1 = NiagaraInternalNoise(1, 2, 3);
+	float sqrt0 = sqrt(r0);
+	float sqrt1 = sqrt(r1);
+	OutBaryCoord = float3(1.0f - sqrt0, sqrt0 * (1.0 - r1), r1 * sqrt0);
+}
+
+void DISKelMesh_GetSkinnedTriangleDataWS(in FDISkelMeshContext DIContext, in uint TriangleIndex, in float3 BaryCoord, out float3 OutPosition, out float3 OutVelocity, out float3 OutNormal, out float3 OutBinormal, out float3 OutTangent)
+{
+    FDISkelMeshSkinnedVertex SkinnedVertex = DISkelMesh_GetSkinnedPointOnTriangleWS(DIContext, TriangleIndex, BaryCoord);
+	OutPosition	= SkinnedVertex.Position;
+	OutVelocity	= (SkinnedVertex.Position - SkinnedVertex.PrevPosition) * DIContext.InstanceInvDeltaTime;
+	OutNormal	= normalize(SkinnedVertex.TangentZ);
+	OutBinormal	= normalize(SkinnedVertex.TangentY);
+	OutTangent	= normalize(SkinnedVertex.TangentX);
+}
+
+void DISKelMesh_GetSkinnedTriangleData(in FDISkelMeshContext DIContext, in uint TriangleIndex, in float3 BaryCoord, out float3 OutPosition, out float3 OutVelocity, out float3 OutNormal, out float3 OutBinormal, out float3 OutTangent)
+{
+    FDISkelMeshSkinnedVertex SkinnedVertex = DISkelMesh_GetSkinnedPointOnTriangle(DIContext, TriangleIndex, BaryCoord);
+	OutPosition	= SkinnedVertex.Position;
+	OutVelocity	= (SkinnedVertex.Position - SkinnedVertex.PrevPosition) * DIContext.InstanceInvDeltaTime;
+	OutNormal	= normalize(SkinnedVertex.TangentZ);
+	OutBinormal	= normalize(SkinnedVertex.TangentY);
+	OutTangent	= normalize(SkinnedVertex.TangentX);
+}
+
+void DISKelMesh_GetSkinnedTriangleDataInterpolatedWS(in FDISkelMeshContext DIContext, in uint TriangleIndex, in float3 BaryCoord, in float Interp, out float3 OutPosition, out float3 OutVelocity, out float3 OutNormal, out float3 OutBinormal, out float3 OutTangent)
+{
+    FDISkelMeshSkinnedVertex SkinnedVertex = DISkelMesh_GetSkinnedPointOnTriangleWS(DIContext, TriangleIndex, BaryCoord);
+	OutPosition	= lerp(SkinnedVertex.PrevPosition, SkinnedVertex.Position, Interp);
+	OutVelocity	= (SkinnedVertex.Position - SkinnedVertex.PrevPosition) * DIContext.InstanceInvDeltaTime;
+	OutNormal	= normalize(SkinnedVertex.TangentZ);
+	OutBinormal	= normalize(SkinnedVertex.TangentY);
+	OutTangent	= normalize(SkinnedVertex.TangentX);
+}
+
+void DISKelMesh_GetSkinnedTriangleDataInterpolated(in FDISkelMeshContext DIContext, in uint TriangleIndex, in float3 BaryCoord, in float Interp, out float3 OutPosition, out float3 OutVelocity, out float3 OutNormal, out float3 OutBinormal, out float3 OutTangent)
+{
+    FDISkelMeshSkinnedVertex SkinnedVertex = DISkelMesh_GetSkinnedPointOnTriangle(DIContext, TriangleIndex, BaryCoord);
+	OutPosition	= lerp(SkinnedVertex.PrevPosition, SkinnedVertex.Position, Interp);
+	OutVelocity	= (SkinnedVertex.Position - SkinnedVertex.PrevPosition) * DIContext.InstanceInvDeltaTime;
+	OutNormal	= normalize(SkinnedVertex.TangentZ);
+	OutBinormal	= normalize(SkinnedVertex.TangentY);
+	OutTangent	= normalize(SkinnedVertex.TangentX);
+}
+
+void DISKelMesh_GetTriUV(in FDISkelMeshContext DIContext, in uint TriangleIndex, in float3 BaryCoord, in int UVSet, out float2 OutUV)
+{
+    if (DIContext.MeshNumTexCoord > 0)
+	{
+        uint IndexBufferOffset = TriangleIndex * 3;
+		uint VertexIndex0 = DIContext.MeshIndexBuffer[IndexBufferOffset  ];
+		uint VertexIndex1 = DIContext.MeshIndexBuffer[IndexBufferOffset+1];
+		uint VertexIndex2 = DIContext.MeshIndexBuffer[IndexBufferOffset+2];
+
+        uint Stride = DIContext.MeshNumTexCoord;
+        uint SelectedUVSet = clamp((uint) UVSet, 0, DIContext.MeshNumTexCoord - 1);
+		float2 UV0 = DIContext.MeshTexCoordBuffer[VertexIndex0 * Stride + SelectedUVSet];
+		float2 UV1 = DIContext.MeshTexCoordBuffer[VertexIndex1 * Stride + SelectedUVSet];
+		float2 UV2 = DIContext.MeshTexCoordBuffer[VertexIndex2 * Stride + SelectedUVSet];
+
+		OutUV = UV0 * BaryCoord.x + UV1 * BaryCoord.y + UV2 * BaryCoord.z;
+	}
+	else	
+	{
+		OutUV = 0.0f;
+	}
+}
+
+////////////////////////////////////////////////////////////////////////////////////////////////////////////////////////////////////////////////////////////////
+// Bone Sampling
+
+float4 QuatSlerp(float4 Quat1, float4 Quat2, float Slerp)
+{
+    const float RawCosom = dot(Quat1, Quat2);
+    const float Cosom = abs(RawCosom);
+	
+    float Scale0, Scale1;
+    if (Cosom < 0.9999f)
+    {
+        const float Omega = acos(Cosom);
+        const float InvSin = 1.f / sin(Omega);
+        Scale0 = sin((1.f - Slerp) * Omega) * InvSin;
+        Scale1 = sin(Slerp * Omega) * InvSin;
+    }
+    else
+    {
+        Scale0 = 1.0f - Slerp;
+        Scale1 = Slerp;
+    }
+
+	// In keeping with our flipped Cosom:
+    Scale1 = RawCosom >= 0.0f ? Scale1 : -Scale1;
+
+    return (Scale0 * Quat1) + (Scale1 * Quat2);
+}
+
+void DISkelMesh_GetSkinnedBoneCommon(in FDISkelMeshContext DIContext, in int Bone, in float Interp, in bool bWorldSpace, out float3 OutPosition, out float4 OutRotation, out float3 OutVelocity)
+{
+    float3 PrevPosition = DIContext.MeshPrevSamplingBonesBuffer[Bone * 2].xyz;
+    float4 PrevRotation = DIContext.MeshPrevSamplingBonesBuffer[Bone * 2 + 1];
+    float3 CurrPosition = DIContext.MeshCurrSamplingBonesBuffer[Bone * 2].xyz;
+    float4 CurrRotation = DIContext.MeshCurrSamplingBonesBuffer[Bone * 2 + 1];
+
+    if (bWorldSpace)
+    {
+        PrevPosition = mul(float4(PrevPosition, 1), DIContext.InstancePrevTransform).xyz;
+        CurrPosition = mul(float4(CurrPosition, 1), DIContext.InstanceTransform).xyz;
+    }
+    CurrPosition = lerp(PrevPosition, CurrPosition, Interp);
+
+    OutPosition = CurrPosition;
+    OutRotation = QuatSlerp(PrevRotation, CurrRotation, Interp);
+    OutVelocity = (CurrPosition - PrevPosition) * DIContext.InstanceInvDeltaTime;
+}
+
+void DISkelMesh_GetSkinnedBone(in FDISkelMeshContext DIContext, in int Bone, out float3 OutPosition, out float4 OutRotation, out float3 OutVelocity)
+{
+    DISkelMesh_GetSkinnedBoneCommon(DIContext, Bone, 1.0f, false, OutPosition, OutRotation, OutVelocity);
+}
+
+void DISkelMesh_GetSkinnedBoneInterpolated(in FDISkelMeshContext DIContext, in int Bone, in float Interp, out float3 OutPosition, out float4 OutRotation, out float3 OutVelocity)
+{
+    DISkelMesh_GetSkinnedBoneCommon(DIContext, Bone, Interp, false, OutPosition, OutRotation, OutVelocity);
+}
+
+void DISkelMesh_GetSkinnedBoneWS(in FDISkelMeshContext DIContext, in int Bone, out float3 OutPosition, out float4 OutRotation, out float3 OutVelocity)
+{
+    DISkelMesh_GetSkinnedBoneCommon(DIContext, Bone, 1.0f, true, OutPosition, OutRotation, OutVelocity);
+}
+
+void DISkelMesh_GetSkinnedBoneInterpolatedWS(in FDISkelMeshContext DIContext, in int Bone, in float Interp, out float3 OutPosition, out float4 OutRotation, out float3 OutVelocity)
+{
+    DISkelMesh_GetSkinnedBoneCommon(DIContext, Bone, Interp, true, OutPosition, OutRotation, OutVelocity);
+}
+
+////////////////////////////////////////////////////////////////////////////////////////////////////////////////////////////////////////////////////////////////
+// Vertex Sampling
+
+void DISkelMesh_GetTriVertices(in FDISkelMeshContext DIContext, in uint TriangleIndex, out int VertexIndex0, out int VertexIndex1, out int VertexIndex2)
+{
+    uint IndexBufferOffset = TriangleIndex * 3;
+    VertexIndex0 = DIContext.MeshIndexBuffer[IndexBufferOffset];
+    VertexIndex1 = DIContext.MeshIndexBuffer[IndexBufferOffset + 1];
+    VertexIndex2 = DIContext.MeshIndexBuffer[IndexBufferOffset + 2];
+}
+
+void DISkelMesh_GetTriColor(in FDISkelMeshContext DIContext, in uint TriangleIndex, in float3 BaryCoord, out float4 OutColor)
+{
+    uint IndexBufferOffset = TriangleIndex * 3;
+    uint VertexIndex0 = DIContext.MeshIndexBuffer[IndexBufferOffset];
+    uint VertexIndex1 = DIContext.MeshIndexBuffer[IndexBufferOffset + 1];
+    uint VertexIndex2 = DIContext.MeshIndexBuffer[IndexBufferOffset + 2];
+    OutColor = (DIContext.MeshColorBuffer[VertexIndex0] * BaryCoord.x) + (DIContext.MeshColorBuffer[VertexIndex1] * BaryCoord.y) + (DIContext.MeshColorBuffer[VertexIndex2] * BaryCoord.z);
+}
+
+void DISkelMesh_GetRandomVertex(in FDISkelMeshContext DIContext, out int OutVertex)
+{
+    float RandT0 = NiagaraInternalNoise(1, 2, 3);
+    OutVertex = min(uint(RandT0 * float(DIContext.MeshVertexCount)), DIContext.MeshVertexCount - 1); // avoid % by using mul/min to Tri = MeshVertexCount
+}
+
+void DISkelMesh_IsValidVertex(in FDISkelMeshContext DIContext, in int Vertex, out bool IsValid)
+{
+    IsValid = Vertex < (int)DIContext.MeshVertexCount;
+}
+
+void DISkelMesh_GetSkinnedVertex(in FDISkelMeshContext DIContext, in int VertexIndex, out float3 OutPosition, out float3 OutVelocity)
+{
+    FDISkelMeshSkinnedVertex SkinnedVertex = DISkelMesh_GetSkinnedVertex(DIContext, VertexIndex);
+    OutPosition = SkinnedVertex.Position;
+    OutVelocity = (SkinnedVertex.Position - SkinnedVertex.PrevPosition) * DIContext.InstanceInvDeltaTime;
+}
+
+void DISkelMesh_GetSkinnedVertexWS(in FDISkelMeshContext DIContext, in int VertexIndex, out float3 OutPosition, out float3 OutVelocity)
+{
+    FDISkelMeshSkinnedVertex SkinnedVertex = DISkelMesh_GetSkinnedVertex(DIContext, VertexIndex);
+    OutPosition = mul(float4(SkinnedVertex.Position, 1.0f), DIContext.InstanceTransform).xyz;
+    float3 PrevPosition = mul(float4(SkinnedVertex.PrevPosition, 1.0f), DIContext.InstancePrevTransform).xyz;
+    OutVelocity = (OutPosition - PrevPosition) * DIContext.InstanceInvDeltaTime;
+}
+
+void DISkelMesh_GetVertexColor(in FDISkelMeshContext DIContext, in int VertexIndex, out float4 OutColor)
+{
+    OutColor = DIContext.MeshColorBuffer[VertexIndex];
+}
+
+void DISkelMesh_GetVertexUV(in FDISkelMeshContext DIContext, in int VertexIndex, in int UVSet, out float2 OutUV)
+{
+    if (DIContext.MeshNumTexCoord > 0)
+    {
+        uint Stride = DIContext.MeshNumTexCoord;
+        uint SelectedUVSet = clamp((uint) UVSet, 0, DIContext.MeshNumTexCoord - 1);
+        OutUV = DIContext.MeshTexCoordBuffer[VertexIndex * Stride + SelectedUVSet];
+    }
+    else
+    {
+        OutUV = 0.0f;
+    }
+}
+
+////////////////////////////////////////////////////////////////////////////////////////////////////////////////////////////////////////////////////////////////
+// Specific Bones / Sockets
+void DISkelMesh_GetSpecificBoneCount(in FDISkelMeshContext DIContext, out int Count)
+{
+    Count = DIContext.NumSpecificBones;
+}
+
+void DISkelMesh_GetSpecificBoneAt(in FDISkelMeshContext DIContext, in int InBoneIndex, out int Bone)
+{
+    Bone = (int)DIContext.SpecificBones[InBoneIndex];
+}
+
+void DISkelMesh_RandomSpecificBone(in FDISkelMeshContext DIContext, out int Bone)
+{
+    float RandT0 = NiagaraInternalNoise(1, 2, 3);
+    Bone = min(uint(RandT0 * float(DIContext.NumSpecificBones)), DIContext.NumSpecificBones - 1); // avoid % by using mul/min to Tri = MeshVertexCount
+}
+
+void DISkelMesh_GetSpecificSocketCount(in FDISkelMeshContext DIContext, out int Count)
+{
+    Count = DIContext.NumSpecificSocketBones;
+}
+
+void DISkelMesh_GetSpecificSocketBoneAt(in FDISkelMeshContext DIContext, in int InSocketIndex, out int Bone)
+{
+    Bone = (int)DIContext.SpecificSocketBones[InSocketIndex];
+}
+
+void DISkelMesh_RandomSpecificSocketBone(in FDISkelMeshContext DIContext, out int SocketBone)
+{
+    float RandT0 = NiagaraInternalNoise(1, 2, 3);
+    SocketBone = min(uint(RandT0 * float(DIContext.NumSpecificSocketBones)), DIContext.NumSpecificSocketBones - 1); // avoid % by using mul/min to Tri = MeshVertexCount
+}
+
+////////////////////////////////////////////////////////////////////////////////////////////////////////////////////////////////////////////////////////////////
+
+#define DISKELMESH_DECLARE_CONSTANTS(NAME) \
+	Buffer<uint>	MeshIndexBuffer_ ## NAME; \
+	Buffer<float>	MeshVertexBuffer_ ## NAME; \
+	Buffer<uint>	MeshSkinWeightBuffer_ ## NAME; \
+	Buffer<float4>	MeshCurrBonesBuffer_ ## NAME; \
+	Buffer<float4>	MeshPrevBonesBuffer_ ## NAME; \
+	Buffer<float4>	MeshCurrSamplingBonesBuffer_ ## NAME; \
+	Buffer<float4>	MeshPrevSamplingBonesBuffer_ ## NAME; \
+	Buffer<float4>	MeshTangentBuffer_ ## NAME; \
+	Buffer<float2>	MeshTexCoordBuffer_ ## NAME; \
+	Buffer<float4>	MeshColorBuffer_ ## NAME; \
+	Buffer<float>	MeshTriangleSamplerProbaBuffer_ ## NAME; \
+	Buffer<uint>	MeshTriangleSamplerAliasBuffer_ ## NAME; \
+	Buffer<uint>	MeshTriangleMatricesOffsetBuffer_ ## NAME; \
+	uint			MeshTriangleCount_ ## NAME; \
+	uint			MeshVertexCount_ ## NAME; \
+	uint			MeshWeightStride_ ## NAME; \
+	uint			MeshNumTexCoord_ ## NAME; \
+	uint			MeshNumWeights_ ## NAME; \
+	uint			NumSpecificBones_ ## NAME; \
+	Buffer<uint>	SpecificBones_ ## NAME; \
+	uint			NumSpecificSocketBones_ ## NAME; \
+	Buffer<uint>	SpecificSocketBones_ ## NAME; \
+	float4x4		InstanceTransform_ ## NAME; \
+	float4x4		InstancePrevTransform_ ## NAME; \
+	float			InstanceInvDeltaTime_ ## NAME; \
+	uint			EnabledFeatures_ ## NAME; \
+
+#define DISKELMESH_MAKE_CONTEXT(NAME) \
+	FDISkelMeshContext	DIContext; \
+	DIContext.MeshIndexBuffer = MeshIndexBuffer_ ## NAME; \
+	DIContext.MeshVertexBuffer = MeshVertexBuffer_ ## NAME; \
+	DIContext.MeshSkinWeightBuffer = MeshSkinWeightBuffer_ ## NAME; \
+	DIContext.MeshCurrBonesBuffer = MeshCurrBonesBuffer_ ## NAME; \
+	DIContext.MeshPrevBonesBuffer = MeshPrevBonesBuffer_ ## NAME; \
+	DIContext.MeshCurrSamplingBonesBuffer = MeshCurrSamplingBonesBuffer_ ## NAME; \
+	DIContext.MeshPrevSamplingBonesBuffer = MeshPrevSamplingBonesBuffer_ ## NAME; \
+	DIContext.MeshTangentBuffer = MeshTangentBuffer_ ## NAME; \
+	DIContext.MeshTexCoordBuffer = MeshTexCoordBuffer_ ## NAME; \
+	DIContext.MeshColorBuffer = MeshColorBuffer_ ## NAME; \
+	DIContext.MeshTriangleSamplerProbaBuffer = MeshTriangleSamplerProbaBuffer_ ## NAME; \
+	DIContext.MeshTriangleSamplerAliasBuffer = MeshTriangleSamplerAliasBuffer_ ## NAME; \
+	DIContext.MeshTriangleMatricesOffsetBuffer = MeshTriangleMatricesOffsetBuffer_ ## NAME; \
+	DIContext.MeshTriangleCount = MeshTriangleCount_ ## NAME; \
+	DIContext.MeshVertexCount = MeshVertexCount_ ## NAME; \
+	DIContext.MeshNumTexCoord = MeshNumTexCoord_ ## NAME; \
+	DIContext.MeshNumWeights = MeshNumWeights_ ## NAME; \
+	DIContext.NumSpecificBones	= NumSpecificBones_ ## NAME; \
+	DIContext.SpecificBones	= SpecificBones_ ## NAME; \
+	DIContext.NumSpecificSocketBones = NumSpecificSocketBones_ ## NAME; \
+	DIContext.SpecificSocketBones = SpecificSocketBones_ ## NAME; \
+	DIContext.MeshWeightStride = MeshWeightStride_ ## NAME; \
+	DIContext.InstanceTransform = InstanceTransform_ ## NAME; \
+	DIContext.InstancePrevTransform = InstancePrevTransform_ ## NAME; \
+	DIContext.InstanceInvDeltaTime = InstanceInvDeltaTime_ ## NAME; \
+	DIContext.UniformTriangleSamplingEnable = EnabledFeatures_ ## NAME & 0x1; \
+
diff --git a/Engine/Plugins/FX/Niagara/Shaders/Private/NiagaraEmitterInstanceShader.usf b/Engine/Plugins/FX/Niagara/Shaders/Private/NiagaraEmitterInstanceShader.usf
--- a/Engine/Plugins/FX/Niagara/Shaders/Private/NiagaraEmitterInstanceShader.usf
+++ b/Engine/Plugins/FX/Niagara/Shaders/Private/NiagaraEmitterInstanceShader.usf
@@ -56,6 +56,18 @@ float Reciprocal(float x){ return 1.0f/x; }
 		float4 HomogeneousWorldPosition = mul(float4(ScreenPosition * SceneDepth, SceneDepth, 1), View.ScreenToWorld);
 		return HomogeneousWorldPosition.xyz / HomogeneousWorldPosition.w;
 	}
+
+	// MASSIVE HACK - Tracked in JIRA UE-69298
+	// Hardcoded random function accessible from inner part of node implementation.
+	// It works for now at least and avoids exposing every random needed in the UI. 
+	// Temporary solution, it will be replaced when a design is validated.
+	float NiagaraInternalNoise(uint u, uint v, uint s)
+	{
+		static uint RandomSeedOffset = 0;
+		uint Seed = (u * 1664525u + v) + s + RandomSeedOffset;
+		RandomSeedOffset += Seed;
+		return float(Rand3DPCG32(int3(u,v,Seed)).x) / 4294967296.0f;
+	}
 #endif
 
 /* ----------------------------------------------------------------------------
diff --git a/Engine/Plugins/FX/Niagara/Source/Niagara/Classes/NiagaraDataInterface.h b/Engine/Plugins/FX/Niagara/Source/Niagara/Classes/NiagaraDataInterface.h
--- a/Engine/Plugins/FX/Niagara/Source/Niagara/Classes/NiagaraDataInterface.h
+++ b/Engine/Plugins/FX/Niagara/Source/Niagara/Classes/NiagaraDataInterface.h
@@ -257,6 +257,11 @@ public:
 	/** Determines if this type definition matches to a known data interface type.*/
 	static bool IsDataInterfaceType(const FNiagaraTypeDefinition& TypeDef);
 
+	/** Allows data interfaces to provide common functionality that will be shared across interfaces on that type. */
+	virtual void GetCommonHLSL(FString& OutHLSL)
+	{
+	}
+
 	virtual bool GetFunctionHLSL(const FName& DefinitionFunctionName, FString InstanceFunctionName, FNiagaraDataInterfaceGPUParamInfo& ParamInfo, FString& OutHLSL)
 	{
 //		checkf(false, TEXT("Unefined HLSL in data interface. Interfaces need to be able to return HLSL for each function they define in GetFunctions."));
diff --git a/Engine/Plugins/FX/Niagara/Source/Niagara/Classes/NiagaraDataInterfaceSkeletalMesh.h b/Engine/Plugins/FX/Niagara/Source/Niagara/Classes/NiagaraDataInterfaceSkeletalMesh.h
--- a/Engine/Plugins/FX/Niagara/Source/Niagara/Classes/NiagaraDataInterfaceSkeletalMesh.h
+++ b/Engine/Plugins/FX/Niagara/Source/Niagara/Classes/NiagaraDataInterfaceSkeletalMesh.h
@@ -238,7 +238,7 @@ public:
 	virtual ~FSkeletalMeshGpuSpawnStaticBuffers();
 
 	FORCEINLINE_DEBUGGABLE void Initialise(const FSkeletalMeshLODRenderData& SkeletalMeshLODRenderData, 
-		bool bIsGpuUniformlyDistributedSampling, const FSkeletalMeshSamplingLODBuiltData& SkeletalMeshSamplingLODBuiltData);
+		bool bIsGpuUniformlyDistributedSampling, const FSkeletalMeshSamplingLODBuiltData& SkeletalMeshSamplingLODBuiltData, const TArray<int32>& SpecificBones, const TArray<int32>& SpecificSocketBones);
 
 	virtual void InitRHI() override;
 	virtual void ReleaseRHI() override;
@@ -249,13 +249,22 @@ public:
 	FShaderResourceViewRHIRef GetBufferTriangleUniformSamplerAliasSRV() const { return BufferTriangleUniformSamplerAliasSRV; }
 	FShaderResourceViewRHIRef GetBufferTriangleMatricesOffsetSRV() const { return BufferTriangleMatricesOffsetSRV; }
 	uint32 GetTriangleCount() const { return TriangleCount; }
+	uint32 GetVertexCount() const { return VertexCount; }
 
 	FShaderResourceViewRHIParamRef GetBufferPositionSRV() const { return MeshVertexBufferSrv; }
 	FShaderResourceViewRHIParamRef GetBufferIndexSRV() const { return MeshIndexBufferSrv; }
 	FShaderResourceViewRHIParamRef GetBufferTangentSRV() const { return MeshTangentBufferSRV; }
 	FShaderResourceViewRHIParamRef GetBufferTexCoordSRV() const { return MeshTexCoordBufferSrv; }
+	FShaderResourceViewRHIParamRef GetBufferColorSRV() const { return MeshColorBufferSrv; }
 
 	uint32 GetNumTexCoord() const { return NumTexCoord; }
+	uint32 GetNumWeights() const { return NumWeights; }
+
+	uint32 GetNumSpecificBones() const { return NumSpecificBones; }
+	FShaderResourceViewRHIParamRef GetSpecificBonesSRV() const { return SpecificBonesSRV; }
+
+	uint32 GetNumSpecificSocketBones() const { return NumSpecificSocketBones; }
+	FShaderResourceViewRHIParamRef GetSpecificSocketBonesSRV() const { return SpecificSocketBonesSRV; }
 
 protected:
 
@@ -266,20 +275,33 @@ protected:
 	FVertexBufferRHIRef BufferTriangleMatricesOffsetRHI = nullptr;
 	FShaderResourceViewRHIRef BufferTriangleMatricesOffsetSRV = nullptr;
 
+	uint32 NumSpecificBones = 0;
+	TResourceArray<uint16> SpecificBonesArray;
+	FVertexBufferRHIRef SpecificBonesBuffer;
+	FShaderResourceViewRHIRef SpecificBonesSRV;
+
+	uint32 NumSpecificSocketBones = 0;
+	TResourceArray<uint16> SpecificSocketBonesArray;
+	FVertexBufferRHIRef SpecificSocketBonesBuffer;
+	FShaderResourceViewRHIRef SpecificSocketBonesSRV;
+
 	/** Cached SRV to gpu buffers of the mesh we spawn from */
-	FShaderResourceViewRHIParamRef MeshVertexBufferSrv;
-	FShaderResourceViewRHIParamRef MeshIndexBufferSrv;
-	FShaderResourceViewRHIParamRef MeshTangentBufferSRV;
-	FShaderResourceViewRHIParamRef MeshTexCoordBufferSrv;
+	FShaderResourceViewRHIRef MeshVertexBufferSrv;
+	FShaderResourceViewRHIRef MeshIndexBufferSrv;
+	FShaderResourceViewRHIRef MeshTangentBufferSRV;
+	FShaderResourceViewRHIRef MeshTexCoordBufferSrv;
+	FShaderResourceViewRHIRef MeshColorBufferSrv;
 
-	uint32 NumTexCoord;
+	uint32 NumTexCoord = 0;
+	uint32 NumWeights = 0;
 
 	// Cached data for resource creation on RenderThread
 	const FSkeletalMeshLODRenderData* LODRenderData = nullptr;
 	const FSkeletalMeshSamplingLODBuiltData* SkeletalMeshSamplingLODBuiltData = nullptr;
-	uint32 TriangleCount;
-	uint32 InputWeightStride;
-	bool bUseGpuUniformlyDistributedSampling;
+	uint32 TriangleCount = 0;
+	uint32 VertexCount = 0;
+	uint32 InputWeightStride = 0;
+	bool bUseGpuUniformlyDistributedSampling = false;
 };
 
 /**
@@ -294,7 +316,7 @@ public:
 	FSkeletalMeshGpuDynamicBufferProxy();
 	virtual ~FSkeletalMeshGpuDynamicBufferProxy();
 
-	void Initialise(const FSkeletalMeshLODRenderData& SkeletalMeshLODRenderData);
+	void Initialise(const FReferenceSkeleton& RefSkel, const FSkeletalMeshLODRenderData& SkeletalMeshLODRenderData);
 
 	virtual void InitRHI() override;
 	virtual void ReleaseRHI() override;
@@ -306,16 +328,19 @@ public:
 	/** Encapsulates a GPU read / CPU write buffer for bone data */
 	struct FSkeletalBuffer
 	{
-		FVertexBufferRHIRef Buffer;
-		FShaderResourceViewRHIRef SRV;
+		FVertexBufferRHIRef SectionBuffer;
+		FShaderResourceViewRHIRef SectionSRV;
+
+		FVertexBufferRHIRef SamplingBuffer;
+		FShaderResourceViewRHIRef SamplingSRV;
 	};
 
 	FSkeletalBuffer& GetRWBufferBone() { return RWBufferBones[CurrentBoneBufferId % 2]; }
 	FSkeletalBuffer& GetRWBufferPrevBone() { return bPrevBoneGpuBufferValid ? RWBufferBones[(CurrentBoneBufferId + 1) % 2] : GetRWBufferBone(); }
 
 private:
-
-	uint32 BoneCount = 0;
+	uint32 SamplingBoneCount = 0;
+	uint32 SectionBoneCount = 0;
 
 	enum { BufferBoneCount = 2 };
 	FSkeletalBuffer RWBufferBones[BufferBoneCount];
@@ -482,6 +507,7 @@ public:
 
 	static USkeletalMesh* GetSkeletalMeshHelper(UNiagaraDataInterfaceSkeletalMesh* Interface, class UNiagaraComponent* OwningComponent, TWeakObjectPtr<USceneComponent>& SceneComponent, USkeletalMeshComponent*& FoundSkelComp);
 
+	virtual void GetCommonHLSL(FString& OutHLSL) override;
 	virtual bool GetFunctionHLSL(const FName&  DefinitionFunctionName, FString InstanceFunctionName, FNiagaraDataInterfaceGPUParamInfo& ParamInfo, FString& OutHLSL) override;
 	virtual void GetParameterDefinitionHLSL(FNiagaraDataInterfaceGPUParamInfo& ParamInfo, FString& OutHLSL) override;
 	virtual FNiagaraDataInterfaceParametersCS* ConstructComputeParameters() const override;
@@ -493,20 +519,27 @@ public:
 	static const FString MeshSkinWeightBufferName;
 	static const FString MeshCurrBonesBufferName;
 	static const FString MeshPrevBonesBufferName;
+	static const FString MeshCurrSamplingBonesBufferName;
+	static const FString MeshPrevSamplingBonesBufferName;
 	static const FString MeshTangentBufferName;
 	static const FString MeshTexCoordBufferName;
+	static const FString MeshColorBufferName;
 	static const FString MeshTriangleSamplerProbaBufferName;
 	static const FString MeshTriangleSamplerAliasBufferName;
 	static const FString MeshTriangleMatricesOffsetBufferName;
 	static const FString MeshTriangleCountName;
-	static const FString MeshWeightStrideByteName;
+	static const FString MeshVertexCountName;
+	static const FString MeshWeightStrideName;
+	static const FString MeshNumTexCoordName;
+	static const FString MeshNumWeightsName;
+	static const FString NumSpecificBonesName;
+	static const FString SpecificBonesName;
+	static const FString NumSpecificSocketBonesName;
+	static const FString SpecificSocketBonesName;
 	static const FString InstanceTransformName;
 	static const FString InstancePrevTransformName;
 	static const FString InstanceInvDeltaTimeName;
 	static const FString EnabledFeaturesName;
-	static const FString InputWeightStrideName;
-	static const FString NumTexCoordName;
-
 
 protected:
 	virtual bool CopyToInternal(UNiagaraDataInterface* Destination) const override;
@@ -633,6 +666,7 @@ public:
 class FSkeletalMeshInterfaceHelper
 {
 public:
+	// Triangle Sampling
 	static const FName RandomTriCoordName;
 	static const FName IsValidTriCoordName;
 	static const FName GetSkinnedTriangleDataName;
@@ -644,6 +678,29 @@ public:
 	static const FName GetTriangleCountName;
 	static const FName GetTriangleAtName;
 	static const FName GetTriCoordVerticesName;
+
+	// Bone Sampling
+	static const FName GetSkinnedBoneDataName;
+	static const FName GetSkinnedBoneDataWSName;
+	static const FName GetSkinnedBoneDataInterpolatedName;
+	static const FName GetSkinnedBoneDataWSInterpolatedName;
+	static const FName RandomSpecificBoneName;
+	static const FName IsValidBoneName;
+	static const FName GetSpecificBoneCountName;
+	static const FName GetSpecificBoneAtName;
+	static const FName RandomSpecificSocketBoneName;
+	static const FName GetSpecificSocketCountName;
+	static const FName GetSpecificSocketBoneAtName;
+
+	// Vertex Sampling
+	static const FName IsValidVertexName;
+	static const FName RandomVertexName;
+	static const FName GetSkinnedVertexDataName;
+	static const FName GetSkinnedVertexDataWSName;
+	static const FName GetVertexColorName;
+	static const FName GetVertexUVName;
+	static const FName GetVertexCountName;
+	static const FName GetVertexAtName;
 };
 
 struct FNiagaraDISkeletalMeshPassedDataToRT
diff --git a/Engine/Plugins/FX/Niagara/Source/Niagara/Private/NDISkeletalMesh_BoneSampling.cpp b/Engine/Plugins/FX/Niagara/Source/Niagara/Private/NDISkeletalMesh_BoneSampling.cpp
--- a/Engine/Plugins/FX/Niagara/Source/Niagara/Private/NDISkeletalMesh_BoneSampling.cpp
+++ b/Engine/Plugins/FX/Niagara/Source/Niagara/Private/NDISkeletalMesh_BoneSampling.cpp
@@ -20,20 +20,17 @@ DEFINE_NDI_DIRECT_FUNC_BINDER(UNiagaraDataInterfaceSkeletalMesh, GetSpecificBone
 
 DEFINE_NDI_DIRECT_FUNC_BINDER(UNiagaraDataInterfaceSkeletalMesh, GetSpecificSocketBoneAt)
 
-static const FName RandomSpecificBoneName("RandomSpecificBone");
-static const FName IsValidBoneName("IsValidBoneName");
-static const FName GetSkinnedBoneDataName("GetSkinnedBoneData");
-static const FName GetSkinnedBoneDataWSName("GetSkinnedBoneDataWS");
-
-static const FName GetSkinnedBoneDataInterpolatedName("GetSkinnedBoneDataInterpolated");
-static const FName GetSkinnedBoneDataWSInterpolatedName("GetSkinnedBoneDataWSInterpolated");
-
-static const FName GetSpecificBoneCountName("GetSpecificBoneCount");
-static const FName GetSpecificBoneAtName("GetSpecificBone");
-
-static const FName RandomSpecificSocketBoneName("RandomSpecificSocketBone");
-static const FName GetSpecificSocketCountName("GetSpecificSocketCount");
-static const FName GetSpecificSocketBoneAtName("GetSpecificSocketBone");
+const FName FSkeletalMeshInterfaceHelper::GetSkinnedBoneDataName("GetSkinnedBoneData");
+const FName FSkeletalMeshInterfaceHelper::GetSkinnedBoneDataWSName("GetSkinnedBoneDataWS");
+const FName FSkeletalMeshInterfaceHelper::GetSkinnedBoneDataInterpolatedName("GetSkinnedBoneDataInterpolated");
+const FName FSkeletalMeshInterfaceHelper::GetSkinnedBoneDataWSInterpolatedName("GetSkinnedBoneDataWSInterpolated");
+const FName FSkeletalMeshInterfaceHelper::RandomSpecificBoneName("RandomSpecificBone");
+const FName FSkeletalMeshInterfaceHelper::IsValidBoneName("IsValidBoneName");
+const FName FSkeletalMeshInterfaceHelper::GetSpecificBoneCountName("GetSpecificBoneCount");
+const FName FSkeletalMeshInterfaceHelper::GetSpecificBoneAtName("GetSpecificBone");
+const FName FSkeletalMeshInterfaceHelper::RandomSpecificSocketBoneName("RandomSpecificSocketBone");
+const FName FSkeletalMeshInterfaceHelper::GetSpecificSocketCountName("GetSpecificSocketCount");
+const FName FSkeletalMeshInterfaceHelper::GetSpecificSocketBoneAtName("GetSpecificSocketBone");
 
 void UNiagaraDataInterfaceSkeletalMesh::GetSkeletonSamplingFunctions(TArray<FNiagaraFunctionSignature>& OutFunctions)
 {
@@ -42,7 +39,7 @@ void UNiagaraDataInterfaceSkeletalMesh::GetSkeletonSamplingFunctions(TArray<FNia
 
 	{
 		FNiagaraFunctionSignature Sig;
-		Sig.Name = RandomSpecificBoneName;
+		Sig.Name = FSkeletalMeshInterfaceHelper::RandomSpecificBoneName;
 		Sig.Inputs.Add(FNiagaraVariable(FNiagaraTypeDefinition(GetClass()), TEXT("SkeletalMesh")));
 		Sig.Outputs.Add(FNiagaraVariable(FNiagaraTypeDefinition::GetIntDef() , TEXT("Bone")));
 		Sig.bMemberFunction = true;
@@ -52,7 +49,7 @@ void UNiagaraDataInterfaceSkeletalMesh::GetSkeletonSamplingFunctions(TArray<FNia
 
 	{
 		FNiagaraFunctionSignature Sig;
-		Sig.Name = IsValidBoneName;
+		Sig.Name = FSkeletalMeshInterfaceHelper::IsValidBoneName;
 		Sig.Inputs.Add(FNiagaraVariable(FNiagaraTypeDefinition(GetClass()), TEXT("SkeletalMesh")));
 		Sig.Inputs.Add(FNiagaraVariable(FNiagaraTypeDefinition::GetIntDef(), TEXT("Bone")));
 		Sig.Outputs.Add(FNiagaraVariable(FNiagaraTypeDefinition::GetBoolDef(), TEXT("IsValid")));
@@ -66,7 +63,7 @@ void UNiagaraDataInterfaceSkeletalMesh::GetSkeletonSamplingFunctions(TArray<FNia
 
 	{
 		FNiagaraFunctionSignature Sig;
-		Sig.Name = GetSkinnedBoneDataName;
+		Sig.Name = FSkeletalMeshInterfaceHelper::GetSkinnedBoneDataName;
 		Sig.Inputs.Add(FNiagaraVariable(FNiagaraTypeDefinition(GetClass()), TEXT("SkeletalMesh")));
 		Sig.Inputs.Add(FNiagaraVariable(FNiagaraTypeDefinition::GetIntDef(), TEXT("Bone")));
 		Sig.Outputs.Add(FNiagaraVariable(FNiagaraTypeDefinition::GetVec3Def(), TEXT("Position")));
@@ -82,7 +79,7 @@ void UNiagaraDataInterfaceSkeletalMesh::GetSkeletonSamplingFunctions(TArray<FNia
 
 	{
 		FNiagaraFunctionSignature Sig;
-		Sig.Name = GetSkinnedBoneDataWSName;
+		Sig.Name = FSkeletalMeshInterfaceHelper::GetSkinnedBoneDataWSName;
 		Sig.Inputs.Add(FNiagaraVariable(FNiagaraTypeDefinition(GetClass()), TEXT("SkeletalMesh")));
 		Sig.Inputs.Add(FNiagaraVariable(FNiagaraTypeDefinition::GetIntDef(), TEXT("Bone")));
 		Sig.Outputs.Add(FNiagaraVariable(FNiagaraTypeDefinition::GetVec3Def(), TEXT("Position")));
@@ -99,7 +96,7 @@ void UNiagaraDataInterfaceSkeletalMesh::GetSkeletonSamplingFunctions(TArray<FNia
 
 	{
 		FNiagaraFunctionSignature Sig;
-		Sig.Name = GetSkinnedBoneDataInterpolatedName;
+		Sig.Name = FSkeletalMeshInterfaceHelper::GetSkinnedBoneDataInterpolatedName;
 		Sig.Inputs.Add(FNiagaraVariable(FNiagaraTypeDefinition(GetClass()), TEXT("SkeletalMesh")));
 		Sig.Inputs.Add(FNiagaraVariable(FNiagaraTypeDefinition::GetIntDef(), TEXT("Bone")));
 		Sig.Inputs.Add(FNiagaraVariable(FNiagaraTypeDefinition::GetFloatDef(), TEXT("Interpolation")));
@@ -116,7 +113,7 @@ void UNiagaraDataInterfaceSkeletalMesh::GetSkeletonSamplingFunctions(TArray<FNia
 
 	{
 		FNiagaraFunctionSignature Sig;
-		Sig.Name = GetSkinnedBoneDataWSInterpolatedName;
+		Sig.Name = FSkeletalMeshInterfaceHelper::GetSkinnedBoneDataWSInterpolatedName;
 		Sig.Inputs.Add(FNiagaraVariable(FNiagaraTypeDefinition(GetClass()), TEXT("SkeletalMesh")));
 		Sig.Inputs.Add(FNiagaraVariable(FNiagaraTypeDefinition::GetIntDef(), TEXT("Bone")));
 		Sig.Inputs.Add(FNiagaraVariable(FNiagaraTypeDefinition::GetFloatDef(), TEXT("Interpolation")));
@@ -133,7 +130,7 @@ void UNiagaraDataInterfaceSkeletalMesh::GetSkeletonSamplingFunctions(TArray<FNia
 
 	{
 		FNiagaraFunctionSignature Sig;
-		Sig.Name = GetSpecificBoneCountName;
+		Sig.Name = FSkeletalMeshInterfaceHelper::GetSpecificBoneCountName;
 		Sig.Inputs.Add(FNiagaraVariable(FNiagaraTypeDefinition(GetClass()), TEXT("SkeletalMesh")));
 		Sig.Outputs.Add(FNiagaraVariable(FNiagaraTypeDefinition::GetIntDef(), TEXT("Count")));
 		Sig.bMemberFunction = true;
@@ -146,7 +143,7 @@ void UNiagaraDataInterfaceSkeletalMesh::GetSkeletonSamplingFunctions(TArray<FNia
 
 	{
 		FNiagaraFunctionSignature Sig;
-		Sig.Name = GetSpecificBoneAtName;
+		Sig.Name = FSkeletalMeshInterfaceHelper::GetSpecificBoneAtName;
 		Sig.Inputs.Add(FNiagaraVariable(FNiagaraTypeDefinition(GetClass()), TEXT("SkeletalMesh")));
 		Sig.Inputs.Add(FNiagaraVariable(FNiagaraTypeDefinition::GetIntDef(), TEXT("Bone Index")));
 		Sig.Outputs.Add(FNiagaraVariable(FNiagaraTypeDefinition::GetIntDef(), TEXT("Bone")));
@@ -163,7 +160,7 @@ void UNiagaraDataInterfaceSkeletalMesh::GetSkeletonSamplingFunctions(TArray<FNia
 
 	{
 		FNiagaraFunctionSignature Sig;
-		Sig.Name = RandomSpecificSocketBoneName;
+		Sig.Name = FSkeletalMeshInterfaceHelper::RandomSpecificSocketBoneName;
 		Sig.Inputs.Add(FNiagaraVariable(FNiagaraTypeDefinition(GetClass()), TEXT("SkeletalMesh")));
 		Sig.Outputs.Add(FNiagaraVariable(FNiagaraTypeDefinition::GetIntDef(), TEXT("Socket Bone")));
 		Sig.bMemberFunction = true;
@@ -176,7 +173,7 @@ void UNiagaraDataInterfaceSkeletalMesh::GetSkeletonSamplingFunctions(TArray<FNia
 	
 	{
 		FNiagaraFunctionSignature Sig;
-		Sig.Name = GetSpecificSocketCountName;
+		Sig.Name = FSkeletalMeshInterfaceHelper::GetSpecificSocketCountName;
 		Sig.Inputs.Add(FNiagaraVariable(FNiagaraTypeDefinition(GetClass()), TEXT("SkeletalMesh")));
 		Sig.Outputs.Add(FNiagaraVariable(FNiagaraTypeDefinition::GetIntDef(), TEXT("Count")));
 		Sig.bMemberFunction = true;
@@ -189,7 +186,7 @@ void UNiagaraDataInterfaceSkeletalMesh::GetSkeletonSamplingFunctions(TArray<FNia
 
 	{
 		FNiagaraFunctionSignature Sig;
-		Sig.Name = GetSpecificSocketBoneAtName;
+		Sig.Name = FSkeletalMeshInterfaceHelper::GetSpecificSocketBoneAtName;
 		Sig.Inputs.Add(FNiagaraVariable(FNiagaraTypeDefinition(GetClass()), TEXT("SkeletalMesh")));
 		Sig.Inputs.Add(FNiagaraVariable(FNiagaraTypeDefinition::GetIntDef(), TEXT("Socket Index")));
 		Sig.Outputs.Add(FNiagaraVariable(FNiagaraTypeDefinition::GetIntDef(), TEXT("Socket Bone")));
@@ -205,62 +202,62 @@ void UNiagaraDataInterfaceSkeletalMesh::GetSkeletonSamplingFunctions(TArray<FNia
 void UNiagaraDataInterfaceSkeletalMesh::BindSkeletonSamplingFunction(const FVMExternalFunctionBindingInfo& BindingInfo, FNDISkeletalMesh_InstanceData* InstanceData, FVMExternalFunction &OutFunc)
 {
 	//Bone Functions
-	if (BindingInfo.Name == RandomSpecificBoneName)
+	if (BindingInfo.Name == FSkeletalMeshInterfaceHelper::RandomSpecificBoneName)
 	{
 		check(BindingInfo.GetNumInputs() == 1 && BindingInfo.GetNumOutputs() == 1);
 		auto Lambda = [this](FVectorVMContext& Context) { this->RandomSpecificBone(Context); }; 
 		OutFunc = FVMExternalFunction::CreateLambda(Lambda); 
 	}
-	else if (BindingInfo.Name == IsValidBoneName)
+	else if (BindingInfo.Name == FSkeletalMeshInterfaceHelper::IsValidBoneName)
 	{
 		check(BindingInfo.GetNumInputs() == 2 && BindingInfo.GetNumOutputs() == 1);
 		NDI_FUNC_BINDER(UNiagaraDataInterfaceSkeletalMesh, IsValidBone)::Bind(this, OutFunc);
 	}
-	else if (BindingInfo.Name == GetSkinnedBoneDataName)
+	else if (BindingInfo.Name == FSkeletalMeshInterfaceHelper::GetSkinnedBoneDataName)
 	{
 		ensure(BindingInfo.GetNumInputs() == 2 && BindingInfo.GetNumOutputs() == 10);
 		TSkinningModeBinder<TNDIExplicitBinder<FNDITransformHandlerNoop, TNDIExplicitBinder<TIntegralConstant<bool, false>, NDI_FUNC_BINDER(UNiagaraDataInterfaceSkeletalMesh, GetSkinnedBoneData)>>>::Bind(this, BindingInfo, InstanceData, OutFunc);
 	}
-	else if (BindingInfo.Name == GetSkinnedBoneDataWSName)
+	else if (BindingInfo.Name == FSkeletalMeshInterfaceHelper::GetSkinnedBoneDataWSName)
 	{
 		ensure(BindingInfo.GetNumInputs() == 2 && BindingInfo.GetNumOutputs() == 10);
 		TSkinningModeBinder<TNDIExplicitBinder<FNDITransformHandler, TNDIExplicitBinder<TIntegralConstant<bool, false>, NDI_FUNC_BINDER(UNiagaraDataInterfaceSkeletalMesh, GetSkinnedBoneData)>>>::Bind(this, BindingInfo, InstanceData, OutFunc);
 	}
-	else if (BindingInfo.Name == GetSkinnedBoneDataInterpolatedName)
+	else if (BindingInfo.Name == FSkeletalMeshInterfaceHelper::GetSkinnedBoneDataInterpolatedName)
 	{
 		ensure(BindingInfo.GetNumInputs() == 3 && BindingInfo.GetNumOutputs() == 10);
 		TSkinningModeBinder<TNDIExplicitBinder<FNDITransformHandlerNoop, TNDIExplicitBinder<TIntegralConstant<bool, true>, NDI_FUNC_BINDER(UNiagaraDataInterfaceSkeletalMesh, GetSkinnedBoneData)>>>::Bind(this, BindingInfo, InstanceData, OutFunc);
 	}
-	else if (BindingInfo.Name == GetSkinnedBoneDataWSInterpolatedName)
+	else if (BindingInfo.Name == FSkeletalMeshInterfaceHelper::GetSkinnedBoneDataWSInterpolatedName)
 	{
 		ensure(BindingInfo.GetNumInputs() == 3 && BindingInfo.GetNumOutputs() == 10);
 		TSkinningModeBinder<TNDIExplicitBinder<FNDITransformHandler, TNDIExplicitBinder<TIntegralConstant<bool, true>, NDI_FUNC_BINDER(UNiagaraDataInterfaceSkeletalMesh, GetSkinnedBoneData)>>>::Bind(this, BindingInfo, InstanceData, OutFunc);
 	}
-	else if (BindingInfo.Name == GetSpecificBoneCountName)
+	else if (BindingInfo.Name == FSkeletalMeshInterfaceHelper::GetSpecificBoneCountName)
 	{
 		check(BindingInfo.GetNumInputs() == 1 && BindingInfo.GetNumOutputs() == 1);
 		auto Lambda = [this](FVectorVMContext& Context) { this->GetSpecificBoneCount(Context); };
 		OutFunc = FVMExternalFunction::CreateLambda(Lambda);
 	}
-	else if (BindingInfo.Name == GetSpecificBoneAtName)
+	else if (BindingInfo.Name == FSkeletalMeshInterfaceHelper::GetSpecificBoneAtName)
 	{
 		check(BindingInfo.GetNumInputs() == 2 && BindingInfo.GetNumOutputs() == 1);
 		NDI_FUNC_BINDER(UNiagaraDataInterfaceSkeletalMesh, GetSpecificBoneAt)::Bind(this, OutFunc);
 	}
 	//Socket Functions
-	else if (BindingInfo.Name == RandomSpecificSocketBoneName)
+	else if (BindingInfo.Name == FSkeletalMeshInterfaceHelper::RandomSpecificSocketBoneName)
 	{
 		check(BindingInfo.GetNumInputs() == 1 && BindingInfo.GetNumOutputs() == 1);
 		auto Lambda = [this](FVectorVMContext& Context) { this->RandomSpecificSocketBone(Context); };
 		OutFunc = FVMExternalFunction::CreateLambda(Lambda);
 	}
-	else if (BindingInfo.Name == GetSpecificSocketCountName)
+	else if (BindingInfo.Name == FSkeletalMeshInterfaceHelper::GetSpecificSocketCountName)
 	{
 		check(BindingInfo.GetNumInputs() == 1 && BindingInfo.GetNumOutputs() == 1);
 		auto Lambda = [this](FVectorVMContext& Context) { this->GetSpecificSocketCount(Context); };
 		OutFunc = FVMExternalFunction::CreateLambda(Lambda);
 	}
-	else if (BindingInfo.Name == GetSpecificSocketBoneAtName)
+	else if (BindingInfo.Name == FSkeletalMeshInterfaceHelper::GetSpecificSocketBoneAtName)
 	{
 		check(BindingInfo.GetNumInputs() == 2 && BindingInfo.GetNumOutputs() == 1);
 		NDI_FUNC_BINDER(UNiagaraDataInterfaceSkeletalMesh, GetSpecificSocketBoneAt)::Bind(this, OutFunc);
diff --git a/Engine/Plugins/FX/Niagara/Source/Niagara/Private/NDISkeletalMesh_VertexSampling.cpp b/Engine/Plugins/FX/Niagara/Source/Niagara/Private/NDISkeletalMesh_VertexSampling.cpp
--- a/Engine/Plugins/FX/Niagara/Source/Niagara/Private/NDISkeletalMesh_VertexSampling.cpp
+++ b/Engine/Plugins/FX/Niagara/Source/Niagara/Private/NDISkeletalMesh_VertexSampling.cpp
@@ -23,20 +23,20 @@ DEFINE_NDI_FUNC_BINDER(UNiagaraDataInterfaceSkeletalMesh, IsValidVertex);
 DEFINE_NDI_FUNC_BINDER(UNiagaraDataInterfaceSkeletalMesh, GetFilteredVertexCount);
 DEFINE_NDI_FUNC_BINDER(UNiagaraDataInterfaceSkeletalMesh, GetFilteredVertexAt);
 
-static const FName RandomVertexName("RandomVertex");
-static const FName IsValidVertexName("IsValidVertex");
-static const FName GetSkinnedVertexDataName("GetSkinnedVertexData");
-static const FName GetSkinnedVertexDataWSName("GetSkinnedVertexDataWS");
-static const FName GetVertexColorName("GetVertexColor");
-static const FName GetVertexUVName("GetVertexUV");
-static const FName GetVertexCountName("GetFilteredVertexCount");
-static const FName GetVertexAtName("GetFilteredVertex");
+const FName FSkeletalMeshInterfaceHelper::IsValidVertexName("IsValidVertex");
+const FName FSkeletalMeshInterfaceHelper::RandomVertexName("RandomVertex");
+const FName FSkeletalMeshInterfaceHelper::GetSkinnedVertexDataName("GetSkinnedVertexData");
+const FName FSkeletalMeshInterfaceHelper::GetSkinnedVertexDataWSName("GetSkinnedVertexDataWS");
+const FName FSkeletalMeshInterfaceHelper::GetVertexColorName("GetVertexColor");
+const FName FSkeletalMeshInterfaceHelper::GetVertexUVName("GetVertexUV");
+const FName FSkeletalMeshInterfaceHelper::GetVertexCountName("GetFilteredVertexCount");
+const FName FSkeletalMeshInterfaceHelper::GetVertexAtName("GetFilteredVertex");
 
 void UNiagaraDataInterfaceSkeletalMesh::GetVertexSamplingFunctions(TArray<FNiagaraFunctionSignature>& OutFunctions)
 {
 	{
 		FNiagaraFunctionSignature Sig;
-		Sig.Name = RandomVertexName;
+		Sig.Name = FSkeletalMeshInterfaceHelper::RandomVertexName;
 		Sig.Inputs.Add(FNiagaraVariable(FNiagaraTypeDefinition(GetClass()), TEXT("SkeletalMesh")));
 		Sig.Outputs.Add(FNiagaraVariable(FNiagaraTypeDefinition::GetIntDef(), TEXT("Vertex")));
 		Sig.bMemberFunction = true;
@@ -46,7 +46,7 @@ void UNiagaraDataInterfaceSkeletalMesh::GetVertexSamplingFunctions(TArray<FNiaga
 
 	{
 		FNiagaraFunctionSignature Sig;
-		Sig.Name = IsValidVertexName;
+		Sig.Name = FSkeletalMeshInterfaceHelper::IsValidVertexName;
 		Sig.Inputs.Add(FNiagaraVariable(FNiagaraTypeDefinition(GetClass()), TEXT("SkeletalMesh")));
 		Sig.Inputs.Add(FNiagaraVariable(FNiagaraTypeDefinition::GetIntDef(), TEXT("Vertex")));
 		Sig.Outputs.Add(FNiagaraVariable(FNiagaraTypeDefinition::GetBoolDef(), TEXT("IsValid")));
@@ -60,7 +60,7 @@ void UNiagaraDataInterfaceSkeletalMesh::GetVertexSamplingFunctions(TArray<FNiaga
 
 	{
 		FNiagaraFunctionSignature Sig;
-		Sig.Name = GetSkinnedVertexDataName;
+		Sig.Name = FSkeletalMeshInterfaceHelper::GetSkinnedVertexDataName;
 		Sig.Inputs.Add(FNiagaraVariable(FNiagaraTypeDefinition(GetClass()), TEXT("SkeletalMesh")));
 		Sig.Inputs.Add(FNiagaraVariable(FNiagaraTypeDefinition::GetIntDef(), TEXT("Vertex")));
 		Sig.Outputs.Add(FNiagaraVariable(FNiagaraTypeDefinition::GetVec3Def(), TEXT("Position")));
@@ -75,7 +75,7 @@ void UNiagaraDataInterfaceSkeletalMesh::GetVertexSamplingFunctions(TArray<FNiaga
 
 	{
 		FNiagaraFunctionSignature Sig;
-		Sig.Name = GetSkinnedVertexDataWSName;
+		Sig.Name = FSkeletalMeshInterfaceHelper::GetSkinnedVertexDataWSName;
 		Sig.Inputs.Add(FNiagaraVariable(FNiagaraTypeDefinition(GetClass()), TEXT("SkeletalMesh")));
 		Sig.Inputs.Add(FNiagaraVariable(FNiagaraTypeDefinition::GetIntDef(), TEXT("Vertex")));
 		Sig.Outputs.Add(FNiagaraVariable(FNiagaraTypeDefinition::GetVec3Def(), TEXT("Position")));
@@ -90,7 +90,7 @@ void UNiagaraDataInterfaceSkeletalMesh::GetVertexSamplingFunctions(TArray<FNiaga
 
 	{
 		FNiagaraFunctionSignature Sig;
-		Sig.Name = GetVertexColorName;
+		Sig.Name = FSkeletalMeshInterfaceHelper::GetVertexColorName;
 		Sig.Inputs.Add(FNiagaraVariable(FNiagaraTypeDefinition(GetClass()), TEXT("SkeletalMesh")));
 		Sig.Inputs.Add(FNiagaraVariable(FNiagaraTypeDefinition::GetIntDef(), TEXT("Vertex")));
 		Sig.Outputs.Add(FNiagaraVariable(FNiagaraTypeDefinition::GetColorDef(), TEXT("Color")));
@@ -101,7 +101,7 @@ void UNiagaraDataInterfaceSkeletalMesh::GetVertexSamplingFunctions(TArray<FNiaga
 
 	{
 		FNiagaraFunctionSignature Sig;
-		Sig.Name = GetVertexUVName;
+		Sig.Name = FSkeletalMeshInterfaceHelper::GetVertexUVName;
 		Sig.Inputs.Add(FNiagaraVariable(FNiagaraTypeDefinition(GetClass()), TEXT("SkeletalMesh")));
 		Sig.Inputs.Add(FNiagaraVariable(FNiagaraTypeDefinition::GetIntDef(), TEXT("Vertex")));
 		Sig.Inputs.Add(FNiagaraVariable(FNiagaraTypeDefinition::GetIntDef(), TEXT("UV Set")));
@@ -114,7 +114,7 @@ void UNiagaraDataInterfaceSkeletalMesh::GetVertexSamplingFunctions(TArray<FNiaga
 
 	{
 		FNiagaraFunctionSignature Sig;
-		Sig.Name = GetVertexCountName;
+		Sig.Name = FSkeletalMeshInterfaceHelper::GetVertexCountName;
 		Sig.Inputs.Add(FNiagaraVariable(FNiagaraTypeDefinition(GetClass()), TEXT("SkeletalMesh")));
 		Sig.Outputs.Add(FNiagaraVariable(FNiagaraTypeDefinition::GetIntDef(), TEXT("Count")));
 		Sig.bMemberFunction = true;
@@ -124,7 +124,7 @@ void UNiagaraDataInterfaceSkeletalMesh::GetVertexSamplingFunctions(TArray<FNiaga
 
 	{
 		FNiagaraFunctionSignature Sig;
-		Sig.Name = GetVertexAtName;
+		Sig.Name = FSkeletalMeshInterfaceHelper::GetVertexAtName;
 		Sig.Inputs.Add(FNiagaraVariable(FNiagaraTypeDefinition(GetClass()), TEXT("SkeletalMesh")));
 		Sig.Inputs.Add(FNiagaraVariable(FNiagaraTypeDefinition::GetIntDef(), TEXT("Filtered Index")));
 		Sig.Outputs.Add(FNiagaraVariable(FNiagaraTypeDefinition::GetIntDef(), TEXT("Vertex")));
@@ -137,27 +137,27 @@ void UNiagaraDataInterfaceSkeletalMesh::GetVertexSamplingFunctions(TArray<FNiaga
 void UNiagaraDataInterfaceSkeletalMesh::BindVertexSamplingFunction(const FVMExternalFunctionBindingInfo& BindingInfo, FNDISkeletalMesh_InstanceData* InstanceData, FVMExternalFunction &OutFunc)
 {
 
-	if (BindingInfo.Name == RandomVertexName)
+	if (BindingInfo.Name == FSkeletalMeshInterfaceHelper::RandomVertexName)
 	{
 		check(BindingInfo.GetNumInputs() == 1 && BindingInfo.GetNumOutputs() == 1);
 		TFilterModeBinder<NDI_FUNC_BINDER(UNiagaraDataInterfaceSkeletalMesh, RandomVertex)>::Bind(this, BindingInfo, InstanceData, OutFunc);
 	}
-	else if (BindingInfo.Name == IsValidVertexName)
+	else if (BindingInfo.Name == FSkeletalMeshInterfaceHelper::IsValidVertexName)
 	{
 		check(BindingInfo.GetNumInputs() == 2 && BindingInfo.GetNumOutputs() == 1);
 		TFilterModeBinder<NDI_FUNC_BINDER(UNiagaraDataInterfaceSkeletalMesh, IsValidVertex)>::Bind(this, BindingInfo, InstanceData, OutFunc);
 	}
-	else if (BindingInfo.Name == GetSkinnedVertexDataName)
+	else if (BindingInfo.Name == FSkeletalMeshInterfaceHelper::GetSkinnedVertexDataName)
 	{
 		check(BindingInfo.GetNumInputs() == 2 && BindingInfo.GetNumOutputs() == 6);
 		TSkinningModeBinder<TNDIExplicitBinder<FNDITransformHandlerNoop, TVertexAccessorBinder<NDI_FUNC_BINDER(UNiagaraDataInterfaceSkeletalMesh, GetVertexSkinnedData)>>>::Bind(this, BindingInfo, InstanceData, OutFunc);
 	}
-	else if (BindingInfo.Name == GetSkinnedVertexDataWSName)
+	else if (BindingInfo.Name == FSkeletalMeshInterfaceHelper::GetSkinnedVertexDataWSName)
 	{
 		check(BindingInfo.GetNumInputs() == 2 && BindingInfo.GetNumOutputs() == 6);
 		TSkinningModeBinder<TNDIExplicitBinder<FNDITransformHandler, TVertexAccessorBinder<NDI_FUNC_BINDER(UNiagaraDataInterfaceSkeletalMesh, GetVertexSkinnedData)>>>::Bind(this, BindingInfo, InstanceData, OutFunc);
 	}
-	else if (BindingInfo.Name == GetVertexColorName)
+	else if (BindingInfo.Name == FSkeletalMeshInterfaceHelper::GetVertexColorName)
 	{
 		check(BindingInfo.GetNumInputs() == 2 && BindingInfo.GetNumOutputs() == 4);
 		if (InstanceData->HasColorData())
@@ -169,17 +169,17 @@ void UNiagaraDataInterfaceSkeletalMesh::BindVertexSamplingFunction(const FVMExte
 			NDI_FUNC_BINDER(UNiagaraDataInterfaceSkeletalMesh, GetVertexColorFallback)::Bind(this, OutFunc);
 		}
 	}
-	else if (BindingInfo.Name == GetVertexUVName)
+	else if (BindingInfo.Name == FSkeletalMeshInterfaceHelper::GetVertexUVName)
 	{
 		check(BindingInfo.GetNumInputs() == 3 && BindingInfo.GetNumOutputs() == 2);
 		TVertexAccessorBinder<NDI_FUNC_BINDER(UNiagaraDataInterfaceSkeletalMesh, GetVertexUV)>::Bind(this, BindingInfo, InstanceData, OutFunc);
 	}
-	else if (BindingInfo.Name == GetVertexCountName)
+	else if (BindingInfo.Name == FSkeletalMeshInterfaceHelper::GetVertexCountName)
 	{
 		check(BindingInfo.GetNumInputs() == 1 && BindingInfo.GetNumOutputs() == 1);
 		TFilterModeBinder<NDI_FUNC_BINDER(UNiagaraDataInterfaceSkeletalMesh, GetFilteredVertexCount)>::Bind(this, BindingInfo, InstanceData, OutFunc);
 	}
-	else if (BindingInfo.Name == GetVertexAtName)
+	else if (BindingInfo.Name == FSkeletalMeshInterfaceHelper::GetVertexAtName)
 	{
 		check(BindingInfo.GetNumInputs() == 2 && BindingInfo.GetNumOutputs() == 1);
 		TFilterModeBinder<NDI_FUNC_BINDER(UNiagaraDataInterfaceSkeletalMesh, GetFilteredVertexAt)>::Bind(this, BindingInfo, InstanceData, OutFunc);
diff --git a/Engine/Plugins/FX/Niagara/Source/Niagara/Private/NiagaraDataInterfaceSkeletalMesh.cpp b/Engine/Plugins/FX/Niagara/Source/Niagara/Private/NiagaraDataInterfaceSkeletalMesh.cpp
--- a/Engine/Plugins/FX/Niagara/Source/Niagara/Private/NiagaraDataInterfaceSkeletalMesh.cpp
+++ b/Engine/Plugins/FX/Niagara/Source/Niagara/Private/NiagaraDataInterfaceSkeletalMesh.cpp
@@ -289,15 +289,40 @@ FSkeletalMeshGpuSpawnStaticBuffers::~FSkeletalMeshGpuSpawnStaticBuffers()
 	//ValidSections.Empty();
 }
 
-void FSkeletalMeshGpuSpawnStaticBuffers::Initialise(const FSkeletalMeshLODRenderData& SkeletalMeshLODRenderData,
-	bool bIsGpuUniformlyDistributedSampling, const FSkeletalMeshSamplingLODBuiltData& MeshSamplingLODBuiltData)
+void FSkeletalMeshGpuSpawnStaticBuffers::Initialise(const FSkeletalMeshLODRenderData& SkeletalMeshLODRenderData, bool bIsGpuUniformlyDistributedSampling, const FSkeletalMeshSamplingLODBuiltData& MeshSamplingLODBuiltData, const TArray<int32>& InSpecificBones, const TArray<int32>& InSpecificSocketBones)
 {
 	SkeletalMeshSamplingLODBuiltData = &MeshSamplingLODBuiltData;
 	bUseGpuUniformlyDistributedSampling = bIsGpuUniformlyDistributedSampling;
 
 	LODRenderData = &SkeletalMeshLODRenderData;
 	TriangleCount = SkeletalMeshLODRenderData.MultiSizeIndexContainer.GetIndexBuffer()->Num() / 3;
+	VertexCount = SkeletalMeshLODRenderData.GetNumVertices();
 	check(TriangleCount > 0);
+	check(VertexCount > 0);
+
+	// Copy Specific Bones / Socket data into arrays that the renderer will use to create read buffers
+	//-TODO: Exclude setting up these arrays if we don't sample from them
+	NumSpecificBones = InSpecificBones.Num();
+	if (NumSpecificBones > 0)
+	{
+		SpecificBonesArray.Reserve(NumSpecificBones);
+		for ( int32 v : InSpecificBones )
+		{
+			check(v <= 65535);
+			SpecificBonesArray.Add(v);
+		}
+	}
+
+	NumSpecificSocketBones = InSpecificSocketBones.Num();
+	if (NumSpecificSocketBones > 0)
+	{
+		SpecificSocketBonesArray.Reserve(NumSpecificSocketBones);
+		for (int32 v : InSpecificSocketBones)
+		{
+			check(v <= 65535);
+			SpecificSocketBonesArray.Add(v);
+		}
+	}
 }
 
 void FSkeletalMeshGpuSpawnStaticBuffers::InitRHI()
@@ -308,6 +333,8 @@ void FSkeletalMeshGpuSpawnStaticBuffers::InitRHI()
 
 	const FMultiSizeIndexContainer& IndexBuffer = LODRenderData->MultiSizeIndexContainer;
 	MeshIndexBufferSrv = IndexBuffer.GetIndexBuffer()->GetSRV();
+	//ensure(MeshIndexBufferSrv != nullptr);
+
 	MeshVertexBufferSrv = LODRenderData->StaticVertexBuffers.PositionVertexBuffer.GetSRV();
 
 	MeshTangentBufferSRV = LODRenderData->StaticVertexBuffers.StaticMeshVertexBuffer.GetTangentsSRV();
@@ -316,7 +343,10 @@ void FSkeletalMeshGpuSpawnStaticBuffers::InitRHI()
 	MeshTexCoordBufferSrv = LODRenderData->StaticVertexBuffers.StaticMeshVertexBuffer.GetTexCoordsSRV();
 	NumTexCoord = LODRenderData->StaticVertexBuffers.StaticMeshVertexBuffer.GetNumTexCoords();
 
-	uint32 VertexCount = LODRenderData->StaticVertexBuffers.StaticMeshVertexBuffer.GetNumVertices();
+	MeshColorBufferSrv = LODRenderData->StaticVertexBuffers.ColorVertexBuffer.GetColorComponentsSRV();
+
+	NumWeights = LODRenderData->SkinWeightVertexBuffer.HasExtraBoneInfluences() ? 8 : 4;
+
 	uint32 SectionCount = LODRenderData->RenderSections.Num();
 
 	if (bUseGpuUniformlyDistributedSampling)
@@ -343,32 +373,65 @@ void FSkeletalMeshGpuSpawnStaticBuffers::InitRHI()
 	// Prepare the vertex matrix lookup offset for each of the sections. This is needed because per vertex BlendIndicies are stored relatively to each Section used matrices.
 	// And these offset per section need to point to the correct matrix according to each section BoneMap.
 	// There is not section selection/culling in the interface so technically we could compute that array in the pipeline.
-	FRHIResourceCreateInfo CreateInfo;
-	void* BufferData = nullptr;
-	BufferTriangleMatricesOffsetRHI = RHICreateAndLockVertexBuffer(VertexCount * sizeof(uint32), BUF_Static | BUF_ShaderResource, CreateInfo, BufferData);
-	uint32* MatricesOffsets = (uint32*)BufferData;
-	uint32 AccumulatedMatrixOffset = 0;
-	for (uint32 s = 0; s < SectionCount; ++s)
-	{
-		const FSkelMeshRenderSection& Section = LODRenderData->RenderSections[s];
-		const uint32 SectionBaseVertexIndex = Section.BaseVertexIndex;
-		const uint32 SectionNumVertices = Section.NumVertices;
-		for (uint32 SectionVertex = 0; SectionVertex < SectionNumVertices; ++SectionVertex)
-		{
-			MatricesOffsets[SectionBaseVertexIndex + SectionVertex] = AccumulatedMatrixOffset;
+	{
+		FRHIResourceCreateInfo CreateInfo;
+		void* BufferData = nullptr;
+		BufferTriangleMatricesOffsetRHI = RHICreateAndLockVertexBuffer(VertexCount * sizeof(uint32), BUF_Static | BUF_ShaderResource, CreateInfo, BufferData);
+		uint32* MatricesOffsets = (uint32*)BufferData;
+		uint32 AccumulatedMatrixOffset = 0;
+		for (uint32 s = 0; s < SectionCount; ++s)
+		{
+			const FSkelMeshRenderSection& Section = LODRenderData->RenderSections[s];
+			const uint32 SectionBaseVertexIndex = Section.BaseVertexIndex;
+			const uint32 SectionNumVertices = Section.NumVertices;
+			for (uint32 SectionVertex = 0; SectionVertex < SectionNumVertices; ++SectionVertex)
+			{
+				MatricesOffsets[SectionBaseVertexIndex + SectionVertex] = AccumulatedMatrixOffset;
+			}
+			AccumulatedMatrixOffset += Section.BoneMap.Num();
 		}
-		AccumulatedMatrixOffset += Section.BoneMap.Num();
+		RHIUnlockVertexBuffer(BufferTriangleMatricesOffsetRHI);
+		BufferTriangleMatricesOffsetSRV = RHICreateShaderResourceView(BufferTriangleMatricesOffsetRHI, sizeof(uint32), PF_R32_UINT);
+	}
+
+	// Create arrays for specific bones / sockets
+	if ( NumSpecificBones > 0 )
+	{
+		FRHIResourceCreateInfo CreateInfo;
+		CreateInfo.ResourceArray = &SpecificBonesArray;
+
+		SpecificBonesBuffer = RHICreateVertexBuffer(NumSpecificBones * sizeof(uint16), BUF_Static | BUF_ShaderResource, CreateInfo);
+		SpecificBonesSRV = RHICreateShaderResourceView(SpecificBonesBuffer, sizeof(uint16), PF_R16_UINT);
+	}
+
+	if (NumSpecificSocketBones > 0)
+	{
+		FRHIResourceCreateInfo CreateInfo;
+		CreateInfo.ResourceArray = &SpecificSocketBonesArray;
+
+		SpecificSocketBonesBuffer = RHICreateVertexBuffer(NumSpecificSocketBones * sizeof(uint16), BUF_Static | BUF_ShaderResource, CreateInfo);
+		SpecificSocketBonesSRV = RHICreateShaderResourceView(SpecificSocketBonesBuffer, sizeof(uint16), PF_R16_UINT);
 	}
-	RHIUnlockVertexBuffer(BufferTriangleMatricesOffsetRHI);
-	BufferTriangleMatricesOffsetSRV = RHICreateShaderResourceView(BufferTriangleMatricesOffsetRHI, sizeof(uint32), PF_R32_UINT);
 }
 
 void FSkeletalMeshGpuSpawnStaticBuffers::ReleaseRHI()
 {
+	SpecificBonesBuffer.SafeRelease();
+	SpecificBonesSRV.SafeRelease();
+
+	SpecificSocketBonesBuffer.SafeRelease();
+	SpecificSocketBonesSRV.SafeRelease();
+
 	BufferTriangleUniformSamplerProbaRHI.SafeRelease();
 	BufferTriangleUniformSamplerProbaSRV.SafeRelease();
 	BufferTriangleUniformSamplerAliasRHI.SafeRelease();
 	BufferTriangleUniformSamplerAliasSRV.SafeRelease();
+
+	MeshVertexBufferSrv.SafeRelease();
+	MeshIndexBufferSrv.SafeRelease();
+	MeshTangentBufferSRV.SafeRelease();
+	MeshTexCoordBufferSrv.SafeRelease();
+	MeshColorBufferSrv.SafeRelease();
 }
 
 //////////////////////////////////////////////////////////////////////////
@@ -384,13 +447,15 @@ FSkeletalMeshGpuDynamicBufferProxy::~FSkeletalMeshGpuDynamicBufferProxy()
 	//
 }
 
-void FSkeletalMeshGpuDynamicBufferProxy::Initialise(const FSkeletalMeshLODRenderData& SkeletalMeshLODRenderData)
+void FSkeletalMeshGpuDynamicBufferProxy::Initialise(const FReferenceSkeleton& RefSkel, const FSkeletalMeshLODRenderData& SkeletalMeshLODRenderData)
 {
-	BoneCount = 0;
-	for (const FSkelMeshRenderSection& section : SkeletalMeshLODRenderData.RenderSections)
+	SectionBoneCount = 0;
+	for (const FSkelMeshRenderSection& Section : SkeletalMeshLODRenderData.RenderSections)
 	{
-		BoneCount += section.BoneMap.Num();
+		SectionBoneCount += Section.BoneMap.Num();
 	}
+
+	SamplingBoneCount = RefSkel.GetNum();
 }
 
 void FSkeletalMeshGpuDynamicBufferProxy::InitRHI()
@@ -399,8 +464,11 @@ void FSkeletalMeshGpuDynamicBufferProxy::InitRHI()
 	{
 		FRHIResourceCreateInfo CreateInfo;
 		CreateInfo.DebugName = TEXT("SkeletalMeshGpuDynamicBuffer");
-		Buffer.Buffer = RHICreateVertexBuffer(sizeof(FVector4) * 3 * BoneCount, BUF_ShaderResource | BUF_Dynamic, CreateInfo);
-		Buffer.SRV = RHICreateShaderResourceView(Buffer.Buffer, sizeof(FVector4), PF_A32B32G32R32F);
+		Buffer.SectionBuffer = RHICreateVertexBuffer(sizeof(FVector4) * 3 * SectionBoneCount, BUF_ShaderResource | BUF_Dynamic, CreateInfo);
+		Buffer.SectionSRV = RHICreateShaderResourceView(Buffer.SectionBuffer, sizeof(FVector4), PF_A32B32G32R32F);
+
+		Buffer.SamplingBuffer = RHICreateVertexBuffer(sizeof(FVector4) * 2 * SamplingBoneCount, BUF_ShaderResource | BUF_Dynamic, CreateInfo);
+		Buffer.SamplingSRV = RHICreateShaderResourceView(Buffer.SamplingBuffer, sizeof(FVector4), PF_A32B32G32R32F);
 	}
 }
 
@@ -408,8 +476,11 @@ void FSkeletalMeshGpuDynamicBufferProxy::ReleaseRHI()
 {
 	for (FSkeletalBuffer& Buffer : RWBufferBones)
 	{
-		Buffer.Buffer.SafeRelease();
-		Buffer.SRV.SafeRelease();
+		Buffer.SectionBuffer.SafeRelease();
+		Buffer.SectionSRV.SafeRelease();
+
+		Buffer.SamplingBuffer.SafeRelease();
+		Buffer.SamplingSRV.SafeRelease();
 	}
 }
 
@@ -432,44 +503,70 @@ void FSkeletalMeshGpuDynamicBufferProxy::NewFrame(const FNDISkeletalMesh_Instanc
 
 		// Count number of matrices we want before appending all of them according to the per section mapping from BoneMap
 		uint32 Float4Count = 0;
-		for ( const FSkelMeshRenderSection& section : Sections )
+		for ( const FSkelMeshRenderSection& Section : Sections )
 		{
-			Float4Count += section.BoneMap.Num() * 3;
+			Float4Count += Section.BoneMap.Num() * 3;
 		}
-		check(Float4Count == 3 * BoneCount);
+		check(Float4Count == 3 * SectionBoneCount);
 		AllSectionsRefToLocalMatrices.AddUninitialized(Float4Count);
 
 		Float4Count = 0;
-		for ( const FSkelMeshRenderSection& section : Sections )
+		for ( const FSkelMeshRenderSection& Section : Sections )
 		{
-			const uint32 MatrixCount = section.BoneMap.Num();
+			const uint32 MatrixCount = Section.BoneMap.Num();
 			for (uint32 m = 0; m < MatrixCount; ++m)
 			{
-				RefToLocalMatrices[section.BoneMap[m]].To3x4MatrixTranspose(&AllSectionsRefToLocalMatrices[Float4Count].X);
+				RefToLocalMatrices[Section.BoneMap[m]].To3x4MatrixTranspose(&AllSectionsRefToLocalMatrices[Float4Count].X);
 				Float4Count += 3;
 			}
 		}
 
+		// Generate information for bone sampling
+		TArray<FVector4> BoneSamplingData;
+		{
+			const TArray<FTransform>& ComponentTransforms = SkelComp->GetComponentSpaceTransforms();
+			check(ComponentTransforms.Num() == SamplingBoneCount);
+
+			BoneSamplingData.Reserve(ComponentTransforms.Num() * 2);
+
+			for (const FTransform& boneMat : ComponentTransforms)
+			{
+				const FQuat Rotation = boneMat.GetRotation();
+				BoneSamplingData.Add(boneMat.GetLocation());
+				BoneSamplingData.Add(FVector4(Rotation.X, Rotation.Y, Rotation.Z, Rotation.W));
+			}
+		}
+
 		FSkeletalMeshGpuDynamicBufferProxy* ThisProxy = this;
 		ENQUEUE_RENDER_COMMAND(UpdateSpawnInfoForSkinnedMesh)(
-			[ThisProxy, AllSectionsRefToLocalMatrices = MoveTemp(AllSectionsRefToLocalMatrices)](FRHICommandListImmediate& RHICmdList) mutable
-		{
-			ThisProxy->CurrentBoneBufferId = (ThisProxy->CurrentBoneBufferId + 1) % BufferBoneCount;
-			ThisProxy->bPrevBoneGpuBufferValid = ThisProxy->bBoneGpuBufferValid;
-			ThisProxy->bBoneGpuBufferValid = true;
+			[ThisProxy, AllSectionsRefToLocalMatrices = MoveTemp(AllSectionsRefToLocalMatrices), BoneSamplingData = MoveTemp(BoneSamplingData)](FRHICommandListImmediate& RHICmdList) mutable
+			{
+				ThisProxy->CurrentBoneBufferId = (ThisProxy->CurrentBoneBufferId + 1) % BufferBoneCount;
+				ThisProxy->bPrevBoneGpuBufferValid = ThisProxy->bBoneGpuBufferValid;
+				ThisProxy->bBoneGpuBufferValid = true;
 
-			const uint32 NumBytes = AllSectionsRefToLocalMatrices.Num() * sizeof(FVector4);
+				// Copy bone remap data matrices
+				{
+					const uint32 NumBytes = AllSectionsRefToLocalMatrices.Num() * sizeof(FVector4);
+					void* DstData = RHILockVertexBuffer(ThisProxy->GetRWBufferBone().SectionBuffer, 0, NumBytes, RLM_WriteOnly);
+					FMemory::Memcpy(DstData, AllSectionsRefToLocalMatrices.GetData(), NumBytes);
+					RHIUnlockVertexBuffer(ThisProxy->GetRWBufferBone().SectionBuffer);
+				}
 
-			void* DstData = RHILockVertexBuffer(ThisProxy->GetRWBufferBone().Buffer, 0, NumBytes, RLM_WriteOnly);
-			FMemory::Memcpy(DstData, AllSectionsRefToLocalMatrices.GetData(), NumBytes);
-			RHIUnlockVertexBuffer(ThisProxy->GetRWBufferBone().Buffer);
-		});
+				// Copy bone sampling data
+				{
+					const uint32 NumBytes = BoneSamplingData.Num() * sizeof(FVector4);
+					FVector4* DstData = reinterpret_cast<FVector4*>(RHILockVertexBuffer(ThisProxy->GetRWBufferBone().SamplingBuffer, 0, NumBytes, RLM_WriteOnly));
+					FMemory::Memcpy(DstData, BoneSamplingData.GetData(), NumBytes);
+					RHIUnlockVertexBuffer(ThisProxy->GetRWBufferBone().SamplingBuffer);
+				}
+			}
+		);
 	}
 }
 
 //////////////////////////////////////////////////////////////////////////
 //FNiagaraDataInterfaceParametersCS_SkeletalMesh
-
 struct FNDISkeletalMeshParametersName
 {
 	FString MeshIndexBufferName;
@@ -477,20 +574,29 @@ struct FNDISkeletalMeshParametersName
 	FString MeshSkinWeightBufferName;
 	FString MeshCurrBonesBufferName;
 	FString MeshPrevBonesBufferName;
+	FString MeshCurrSamplingBonesBufferName;
+	FString MeshPrevSamplingBonesBufferName;
 	FString MeshTangentBufferName;
 	FString MeshTexCoordBufferName;
+	FString MeshColorBufferName;
 	FString MeshTriangleSamplerProbaBufferName;
 	FString MeshTriangleSamplerAliasBufferName;
 	FString MeshTriangleMatricesOffsetBufferName;
 	FString MeshTriangleCountName;
-	FString MeshWeightStrideByteName;
+	FString MeshVertexCountName;
+	FString MeshWeightStrideName;
+	FString MeshNumTexCoordName;
+	FString MeshNumWeightsName;
+	FString NumSpecificBonesName;
+	FString SpecificBonesName;
+	FString NumSpecificSocketBonesName;
+	FString SpecificSocketBonesName;
 	FString InstanceTransformName;
 	FString InstancePrevTransformName;
 	FString InstanceInvDeltaTimeName;
 	FString EnabledFeaturesName;
-	FString InputWeightStrideName;
-	FString NumTexCoordName;
 };
+
 static void GetNiagaraDataInterfaceParametersName(FNDISkeletalMeshParametersName& Names, const FString& Suffix)
 {
 	Names.MeshIndexBufferName = UNiagaraDataInterfaceSkeletalMesh::MeshIndexBufferName + Suffix;
@@ -498,19 +604,27 @@ static void GetNiagaraDataInterfaceParametersName(FNDISkeletalMeshParametersName
 	Names.MeshSkinWeightBufferName = UNiagaraDataInterfaceSkeletalMesh::MeshSkinWeightBufferName + Suffix;
 	Names.MeshCurrBonesBufferName = UNiagaraDataInterfaceSkeletalMesh::MeshCurrBonesBufferName + Suffix;
 	Names.MeshPrevBonesBufferName = UNiagaraDataInterfaceSkeletalMesh::MeshPrevBonesBufferName + Suffix;
+	Names.MeshCurrSamplingBonesBufferName = UNiagaraDataInterfaceSkeletalMesh::MeshCurrSamplingBonesBufferName + Suffix;
+	Names.MeshPrevSamplingBonesBufferName = UNiagaraDataInterfaceSkeletalMesh::MeshPrevSamplingBonesBufferName + Suffix;
 	Names.MeshTangentBufferName = UNiagaraDataInterfaceSkeletalMesh::MeshTangentBufferName + Suffix;
 	Names.MeshTexCoordBufferName = UNiagaraDataInterfaceSkeletalMesh::MeshTexCoordBufferName + Suffix;
+	Names.MeshColorBufferName = UNiagaraDataInterfaceSkeletalMesh::MeshColorBufferName + Suffix;
 	Names.MeshTriangleSamplerProbaBufferName = UNiagaraDataInterfaceSkeletalMesh::MeshTriangleSamplerProbaBufferName + Suffix;
 	Names.MeshTriangleSamplerAliasBufferName = UNiagaraDataInterfaceSkeletalMesh::MeshTriangleSamplerAliasBufferName + Suffix;
 	Names.MeshTriangleMatricesOffsetBufferName = UNiagaraDataInterfaceSkeletalMesh::MeshTriangleMatricesOffsetBufferName + Suffix;
 	Names.MeshTriangleCountName = UNiagaraDataInterfaceSkeletalMesh::MeshTriangleCountName + Suffix;
-	Names.MeshWeightStrideByteName = UNiagaraDataInterfaceSkeletalMesh::MeshWeightStrideByteName + Suffix;
+	Names.MeshVertexCountName = UNiagaraDataInterfaceSkeletalMesh::MeshVertexCountName + Suffix;
+	Names.MeshWeightStrideName = UNiagaraDataInterfaceSkeletalMesh::MeshWeightStrideName + Suffix;
+	Names.MeshNumTexCoordName = UNiagaraDataInterfaceSkeletalMesh::MeshNumTexCoordName + Suffix;
+	Names.MeshNumWeightsName = UNiagaraDataInterfaceSkeletalMesh::MeshNumWeightsName + Suffix;
+	Names.NumSpecificBonesName = UNiagaraDataInterfaceSkeletalMesh::NumSpecificBonesName + Suffix;
+	Names.SpecificBonesName = UNiagaraDataInterfaceSkeletalMesh::SpecificBonesName + Suffix;
+	Names.NumSpecificSocketBonesName = UNiagaraDataInterfaceSkeletalMesh::NumSpecificSocketBonesName + Suffix;
+	Names.SpecificSocketBonesName = UNiagaraDataInterfaceSkeletalMesh::SpecificSocketBonesName + Suffix;
 	Names.InstanceTransformName = UNiagaraDataInterfaceSkeletalMesh::InstanceTransformName + Suffix;
 	Names.InstancePrevTransformName = UNiagaraDataInterfaceSkeletalMesh::InstancePrevTransformName + Suffix;
 	Names.InstanceInvDeltaTimeName = UNiagaraDataInterfaceSkeletalMesh::InstanceInvDeltaTimeName + Suffix;
 	Names.EnabledFeaturesName = UNiagaraDataInterfaceSkeletalMesh::EnabledFeaturesName + Suffix;
-	Names.InputWeightStrideName = UNiagaraDataInterfaceSkeletalMesh::InputWeightStrideName + Suffix;
-	Names.NumTexCoordName = UNiagaraDataInterfaceSkeletalMesh::NumTexCoordName + Suffix;
 }
 
 struct FNiagaraDataInterfaceParametersCS_SkeletalMesh : public FNiagaraDataInterfaceParametersCS
@@ -525,44 +639,27 @@ struct FNiagaraDataInterfaceParametersCS_SkeletalMesh : public FNiagaraDataInter
 		MeshSkinWeightBuffer.Bind(ParameterMap, *ParamNames.MeshSkinWeightBufferName);
 		MeshCurrBonesBuffer.Bind(ParameterMap, *ParamNames.MeshCurrBonesBufferName);
 		MeshPrevBonesBuffer.Bind(ParameterMap, *ParamNames.MeshPrevBonesBufferName);
+		MeshCurrSamplingBonesBuffer.Bind(ParameterMap, *ParamNames.MeshCurrSamplingBonesBufferName);
+		MeshPrevSamplingBonesBuffer.Bind(ParameterMap, *ParamNames.MeshPrevSamplingBonesBufferName);
 		MeshTangentBuffer.Bind(ParameterMap, *ParamNames.MeshTangentBufferName);
 		MeshTexCoordBuffer.Bind(ParameterMap, *ParamNames.MeshTexCoordBufferName);
+		MeshColorBuffer.Bind(ParameterMap, *ParamNames.MeshColorBufferName);
 		MeshTriangleSamplerProbaBuffer.Bind(ParameterMap, *ParamNames.MeshTriangleSamplerProbaBufferName);
 		MeshTriangleSamplerAliasBuffer.Bind(ParameterMap, *ParamNames.MeshTriangleSamplerAliasBufferName);
 		MeshTriangleMatricesOffsetBuffer.Bind(ParameterMap, *ParamNames.MeshTriangleMatricesOffsetBufferName);
 		MeshTriangleCount.Bind(ParameterMap, *ParamNames.MeshTriangleCountName);
-		MeshWeightStrideByte.Bind(ParameterMap, *ParamNames.MeshWeightStrideByteName);
+		MeshVertexCount.Bind(ParameterMap, *ParamNames.MeshVertexCountName);
+		MeshWeightStride.Bind(ParameterMap, *ParamNames.MeshWeightStrideName);
+		MeshNumTexCoord.Bind(ParameterMap, *ParamNames.MeshNumTexCoordName);
+		MeshNumWeights.Bind(ParameterMap, *ParamNames.MeshNumWeightsName);
+		NumSpecificBones.Bind(ParameterMap, *ParamNames.NumSpecificBonesName);
+		SpecificBones.Bind(ParameterMap, *ParamNames.SpecificBonesName);
+		NumSpecificSocketBones.Bind(ParameterMap, *ParamNames.NumSpecificSocketBonesName);
+		SpecificSocketBones.Bind(ParameterMap, *ParamNames.SpecificSocketBonesName);
 		InstanceTransform.Bind(ParameterMap, *ParamNames.InstanceTransformName);
 		InstancePrevTransform.Bind(ParameterMap, *ParamNames.InstancePrevTransformName);
 		InstanceInvDeltaTime.Bind(ParameterMap, *ParamNames.InstanceInvDeltaTimeName);
 		EnabledFeatures.Bind(ParameterMap, *ParamNames.EnabledFeaturesName);
-		InputWeightStride.Bind(ParameterMap, *ParamNames.InputWeightStrideName);
-		NumTexCoord.Bind(ParameterMap, *ParamNames.NumTexCoordName);
-
-		if (!MeshIndexBuffer.IsBound())
-		{
-			UE_LOG(LogNiagara, Warning, TEXT("Binding failed for FNiagaraDataInterfaceParametersCS_StaticMesh Texture %s. Was it optimized out?"), *ParamNames.MeshIndexBufferName)
-		}
-		if (!MeshVertexBuffer.IsBound())
-		{
-			UE_LOG(LogNiagara, Warning, TEXT("Binding failed for FNiagaraDataInterfaceParametersCS_StaticMesh Sampler %s. Was it optimized out?"), *ParamNames.MeshVertexBufferName)
-		}
-		if (!MeshSkinWeightBuffer.IsBound())
-		{
-			UE_LOG(LogNiagara, Warning, TEXT("Binding failed for FNiagaraDataInterfaceParametersCS_StaticMesh Sampler %s. Was it optimized out?"), *ParamNames.MeshSkinWeightBufferName)
-		}
-		if (!MeshCurrBonesBuffer.IsBound())
-		{
-			UE_LOG(LogNiagara, Warning, TEXT("Binding failed for FNiagaraDataInterfaceParametersCS_StaticMesh Sampler %s. Was it optimized out?"), *ParamNames.MeshCurrBonesBufferName)
-		}
-		if (!MeshTangentBuffer.IsBound())
-		{
-			UE_LOG(LogNiagara, Warning, TEXT("Binding failed for FNiagaraDataInterfaceParametersCS_StaticMesh Sampler %s. Was it optimized out?"), *ParamNames.MeshTangentBufferName)
-		}
-		if (!MeshTriangleMatricesOffsetBuffer.IsBound())
-		{
-			UE_LOG(LogNiagara, Warning, TEXT("Binding failed for FNiagaraDataInterfaceParametersCS_StaticMesh Sampler %s. Was it optimized out?"), *ParamNames.MeshTriangleMatricesOffsetBufferName)
-		}
 	}
 
 	virtual void Serialize(FArchive& Ar)override
@@ -572,19 +669,27 @@ struct FNiagaraDataInterfaceParametersCS_SkeletalMesh : public FNiagaraDataInter
 		Ar << MeshSkinWeightBuffer;
 		Ar << MeshCurrBonesBuffer;
 		Ar << MeshPrevBonesBuffer;
+		Ar << MeshCurrSamplingBonesBuffer;
+		Ar << MeshPrevSamplingBonesBuffer;
 		Ar << MeshTangentBuffer;
 		Ar << MeshTexCoordBuffer;
+		Ar << MeshColorBuffer;
 		Ar << MeshTriangleSamplerProbaBuffer;
 		Ar << MeshTriangleSamplerAliasBuffer;
 		Ar << MeshTriangleMatricesOffsetBuffer;
 		Ar << MeshTriangleCount;
-		Ar << MeshWeightStrideByte;
+		Ar << MeshVertexCount;
+		Ar << MeshWeightStride;
+		Ar << MeshNumTexCoord;
+		Ar << MeshNumWeights;
+		Ar << NumSpecificBones;
+		Ar << SpecificBones;
+		Ar << NumSpecificSocketBones;
+		Ar << SpecificSocketBones;
 		Ar << InstanceTransform;
 		Ar << InstancePrevTransform;
 		Ar << InstanceInvDeltaTime;
 		Ar << EnabledFeatures;
-		Ar << InputWeightStride;
-		Ar << NumTexCoord;
 	}
 
 	virtual void Set(FRHICommandList& RHICmdList, const FNiagaraDataInterfaceSetArgs& Context) const override
@@ -602,7 +707,7 @@ struct FNiagaraDataInterfaceParametersCS_SkeletalMesh : public FNiagaraDataInter
 			SetSRVParameter(RHICmdList, ComputeShaderRHI, MeshIndexBuffer, StaticBuffers->GetBufferIndexSRV());
 			SetSRVParameter(RHICmdList, ComputeShaderRHI, MeshTangentBuffer, StaticBuffers->GetBufferTangentSRV());
 
-			SetShaderValue(RHICmdList, ComputeShaderRHI, NumTexCoord, StaticBuffers->GetNumTexCoord());
+			SetShaderValue(RHICmdList, ComputeShaderRHI, MeshNumTexCoord, StaticBuffers->GetNumTexCoord());
 			if (StaticBuffers->GetNumTexCoord() > 0)
 			{
 				SetSRVParameter(RHICmdList, ComputeShaderRHI, MeshTexCoordBuffer, StaticBuffers->GetBufferTexCoordSRV());
@@ -611,7 +716,9 @@ struct FNiagaraDataInterfaceParametersCS_SkeletalMesh : public FNiagaraDataInter
 			{
 				SetSRVParameter(RHICmdList, ComputeShaderRHI, MeshTexCoordBuffer, FNiagaraRenderer::GetDummyFloatBuffer().SRV);
 			}
+			SetSRVParameter(RHICmdList, ComputeShaderRHI, MeshColorBuffer, StaticBuffers->GetBufferColorSRV());
 			SetShaderValue(RHICmdList, ComputeShaderRHI, MeshTriangleCount, StaticBuffers->GetTriangleCount());
+			SetShaderValue(RHICmdList, ComputeShaderRHI, MeshVertexCount, StaticBuffers->GetVertexCount());
 			if (InstanceData->bIsGpuUniformlyDistributedSampling)
 			{
 				SetSRVParameter(RHICmdList, ComputeShaderRHI, MeshTriangleSamplerProbaBuffer, StaticBuffers->GetBufferTriangleUniformSamplerProbaSRV().GetReference());
@@ -625,11 +732,7 @@ struct FNiagaraDataInterfaceParametersCS_SkeletalMesh : public FNiagaraDataInter
 
 			SetSRVParameter(RHICmdList, ComputeShaderRHI, MeshSkinWeightBuffer, InstanceData->MeshSkinWeightBufferSrv);
 
-			SetShaderValue(RHICmdList, ComputeShaderRHI, MeshWeightStrideByte, InstanceData->MeshWeightStrideByte);
-			SetShaderValue(RHICmdList, ComputeShaderRHI, InstanceTransform, InstanceData->Transform);
-			SetShaderValue(RHICmdList, ComputeShaderRHI, InstancePrevTransform, InstanceData->PrevTransform);
-			SetShaderValue(RHICmdList, ComputeShaderRHI, InstanceInvDeltaTime, 1.0f / InstanceData->DeltaSeconds);
-			SetShaderValue(RHICmdList, ComputeShaderRHI, InputWeightStride, InstanceData->MeshWeightStrideByte/4);
+			SetShaderValue(RHICmdList, ComputeShaderRHI, MeshWeightStride, InstanceData->MeshWeightStrideByte/4);
 
 			uint32 EnabledFeaturesBits = InstanceData->bIsGpuUniformlyDistributedSampling ? 1 : 0;
 
@@ -637,19 +740,33 @@ struct FNiagaraDataInterfaceParametersCS_SkeletalMesh : public FNiagaraDataInter
 			check(DynamicBuffers);
 			if(DynamicBuffers->DoesBoneDataExist())
 			{
-				EnabledFeaturesBits |= 2; // Enable the skinning feature
-				SetSRVParameter(RHICmdList, ComputeShaderRHI, MeshCurrBonesBuffer, DynamicBuffers->GetRWBufferBone().SRV);
-				SetSRVParameter(RHICmdList, ComputeShaderRHI, MeshPrevBonesBuffer, DynamicBuffers->GetRWBufferPrevBone().SRV);
+				SetShaderValue(RHICmdList, ComputeShaderRHI, MeshNumWeights, StaticBuffers->GetNumWeights());
+				SetSRVParameter(RHICmdList, ComputeShaderRHI, MeshCurrBonesBuffer, DynamicBuffers->GetRWBufferBone().SectionSRV);
+				SetSRVParameter(RHICmdList, ComputeShaderRHI, MeshPrevBonesBuffer, DynamicBuffers->GetRWBufferPrevBone().SectionSRV);
+				SetSRVParameter(RHICmdList, ComputeShaderRHI, MeshCurrSamplingBonesBuffer, DynamicBuffers->GetRWBufferBone().SamplingSRV);
+				SetSRVParameter(RHICmdList, ComputeShaderRHI, MeshPrevSamplingBonesBuffer, DynamicBuffers->GetRWBufferPrevBone().SamplingSRV);
 				SetSRVParameter(RHICmdList, ComputeShaderRHI, MeshTriangleMatricesOffsetBuffer, StaticBuffers->GetBufferTriangleMatricesOffsetSRV());
 			}
 			// Bind dummy data for validation purposes only.  Code will not execute due to "EnabledFeatures" bits but validation can not determine that.
 			else
 			{
+				SetShaderValue(RHICmdList, ComputeShaderRHI, MeshNumWeights, 0);
 				SetSRVParameter(RHICmdList, ComputeShaderRHI, MeshCurrBonesBuffer, FNiagaraRenderer::GetDummyFloat4Buffer().SRV);
 				SetSRVParameter(RHICmdList, ComputeShaderRHI, MeshPrevBonesBuffer, FNiagaraRenderer::GetDummyFloat4Buffer().SRV);
+				SetSRVParameter(RHICmdList, ComputeShaderRHI, MeshCurrSamplingBonesBuffer, FNiagaraRenderer::GetDummyFloat4Buffer().SRV);
+				SetSRVParameter(RHICmdList, ComputeShaderRHI, MeshPrevSamplingBonesBuffer, FNiagaraRenderer::GetDummyFloat4Buffer().SRV);
 				SetSRVParameter(RHICmdList, ComputeShaderRHI, MeshTriangleMatricesOffsetBuffer, FNiagaraRenderer::GetDummyUIntBuffer().SRV);
 			}
 
+			SetShaderValue(RHICmdList, ComputeShaderRHI, NumSpecificBones, StaticBuffers->GetNumSpecificBones());
+			SetSRVParameter(RHICmdList, ComputeShaderRHI, SpecificBones, StaticBuffers->GetNumSpecificBones() > 0 ? StaticBuffers->GetSpecificBonesSRV() : FNiagaraRenderer::GetDummyUIntBuffer().SRV);
+			SetShaderValue(RHICmdList, ComputeShaderRHI, NumSpecificSocketBones, StaticBuffers->GetNumSpecificSocketBones());
+			SetSRVParameter(RHICmdList, ComputeShaderRHI, SpecificSocketBones, StaticBuffers->GetNumSpecificSocketBones() > 0 ? StaticBuffers->GetSpecificSocketBonesSRV() : FNiagaraRenderer::GetDummyUIntBuffer().SRV);
+
+			SetShaderValue(RHICmdList, ComputeShaderRHI, InstanceTransform, InstanceData->Transform);
+			SetShaderValue(RHICmdList, ComputeShaderRHI, InstancePrevTransform, InstanceData->PrevTransform);
+			SetShaderValue(RHICmdList, ComputeShaderRHI, InstanceInvDeltaTime, 1.0f / InstanceData->DeltaSeconds);
+
 			SetShaderValue(RHICmdList, ComputeShaderRHI, EnabledFeatures, EnabledFeaturesBits);
 		}
 		else
@@ -661,24 +778,34 @@ struct FNiagaraDataInterfaceParametersCS_SkeletalMesh : public FNiagaraDataInter
 			SetSRVParameter(RHICmdList, ComputeShaderRHI, MeshIndexBuffer, FNiagaraRenderer::GetDummyUIntBuffer().SRV);
 			SetSRVParameter(RHICmdList, ComputeShaderRHI, MeshTangentBuffer, FNiagaraRenderer::GetDummyFloatBuffer().SRV);
 
-			SetShaderValue(RHICmdList, ComputeShaderRHI, NumTexCoord, 0);
+			SetShaderValue(RHICmdList, ComputeShaderRHI, MeshNumTexCoord, 0);
 			SetSRVParameter(RHICmdList, ComputeShaderRHI, MeshTexCoordBuffer, FNiagaraRenderer::GetDummyFloatBuffer().SRV);
+			SetSRVParameter(RHICmdList, ComputeShaderRHI, MeshColorBuffer, FNiagaraRenderer::GetDummyFloatBuffer().SRV);
 			SetShaderValue(RHICmdList, ComputeShaderRHI, MeshTriangleCount, 0);
+			SetShaderValue(RHICmdList, ComputeShaderRHI, MeshVertexCount, 0);
 			SetSRVParameter(RHICmdList, ComputeShaderRHI, MeshTriangleSamplerProbaBuffer, FNiagaraRenderer::GetDummyUIntBuffer().SRV);
 			SetSRVParameter(RHICmdList, ComputeShaderRHI, MeshTriangleSamplerAliasBuffer, FNiagaraRenderer::GetDummyUIntBuffer().SRV);
 
 			SetSRVParameter(RHICmdList, ComputeShaderRHI, MeshSkinWeightBuffer, FNiagaraRenderer::GetDummyUIntBuffer().SRV);
 
-			SetShaderValue(RHICmdList, ComputeShaderRHI, MeshWeightStrideByte, 0);
-			SetShaderValue(RHICmdList, ComputeShaderRHI, InstanceTransform, FMatrix::Identity);
-			SetShaderValue(RHICmdList, ComputeShaderRHI, InstancePrevTransform, FMatrix::Identity);
-			SetShaderValue(RHICmdList, ComputeShaderRHI, InstanceInvDeltaTime, 0.0f);
-			SetShaderValue(RHICmdList, ComputeShaderRHI, InputWeightStride, 0);
+			SetShaderValue(RHICmdList, ComputeShaderRHI, MeshWeightStride, 0);
+			SetShaderValue(RHICmdList, ComputeShaderRHI, MeshNumWeights, 0);
 
 			SetSRVParameter(RHICmdList, ComputeShaderRHI, MeshCurrBonesBuffer, FNiagaraRenderer::GetDummyFloat4Buffer().SRV);
 			SetSRVParameter(RHICmdList, ComputeShaderRHI, MeshPrevBonesBuffer, FNiagaraRenderer::GetDummyFloat4Buffer().SRV);
+			SetSRVParameter(RHICmdList, ComputeShaderRHI, MeshCurrSamplingBonesBuffer, FNiagaraRenderer::GetDummyFloat4Buffer().SRV);
+			SetSRVParameter(RHICmdList, ComputeShaderRHI, MeshPrevSamplingBonesBuffer, FNiagaraRenderer::GetDummyFloat4Buffer().SRV);
 			SetSRVParameter(RHICmdList, ComputeShaderRHI, MeshTriangleMatricesOffsetBuffer, FNiagaraRenderer::GetDummyUIntBuffer().SRV);
 
+			SetShaderValue(RHICmdList, ComputeShaderRHI, NumSpecificBones, 0);
+			SetSRVParameter(RHICmdList, ComputeShaderRHI, SpecificBones, FNiagaraRenderer::GetDummyUIntBuffer().SRV);
+			SetShaderValue(RHICmdList, ComputeShaderRHI, NumSpecificSocketBones, 0);
+			SetSRVParameter(RHICmdList, ComputeShaderRHI, SpecificBones, FNiagaraRenderer::GetDummyUIntBuffer().SRV);
+
+			SetShaderValue(RHICmdList, ComputeShaderRHI, InstanceTransform, FMatrix::Identity);
+			SetShaderValue(RHICmdList, ComputeShaderRHI, InstancePrevTransform, FMatrix::Identity);
+			SetShaderValue(RHICmdList, ComputeShaderRHI, InstanceInvDeltaTime, 0.0f);
+
 			SetShaderValue(RHICmdList, ComputeShaderRHI, EnabledFeatures, 0);
 		}
 	}
@@ -691,19 +818,27 @@ private:
 	FShaderResourceParameter MeshSkinWeightBuffer;
 	FShaderResourceParameter MeshCurrBonesBuffer;
 	FShaderResourceParameter MeshPrevBonesBuffer;
+	FShaderResourceParameter MeshCurrSamplingBonesBuffer;
+	FShaderResourceParameter MeshPrevSamplingBonesBuffer;
 	FShaderResourceParameter MeshTangentBuffer;
 	FShaderResourceParameter MeshTexCoordBuffer;
+	FShaderResourceParameter MeshColorBuffer;
 	FShaderResourceParameter MeshTriangleSamplerProbaBuffer;
 	FShaderResourceParameter MeshTriangleSamplerAliasBuffer;
 	FShaderResourceParameter MeshTriangleMatricesOffsetBuffer;
 	FShaderParameter MeshTriangleCount;
-	FShaderParameter MeshWeightStrideByte;
+	FShaderParameter MeshVertexCount;
+	FShaderParameter MeshWeightStride;
+	FShaderParameter MeshNumTexCoord;
+	FShaderParameter MeshNumWeights;
+	FShaderParameter NumSpecificBones;
+	FShaderResourceParameter SpecificBones;
+	FShaderParameter NumSpecificSocketBones;
+	FShaderResourceParameter SpecificSocketBones;
 	FShaderParameter InstanceTransform;
 	FShaderParameter InstancePrevTransform;
 	FShaderParameter InstanceInvDeltaTime;
 	FShaderParameter EnabledFeatures;
-	FShaderParameter InputWeightStride;
-	FShaderParameter NumTexCoord;
 };
 
 //////////////////////////////////////////////////////////////////////////
@@ -1126,11 +1261,11 @@ bool FNDISkeletalMesh_InstanceData::Init(UNiagaraDataInterfaceSkeletalMesh* Inte
 		}
 
 		MeshGpuSpawnStaticBuffers = new FSkeletalMeshGpuSpawnStaticBuffers();
-		MeshGpuSpawnStaticBuffers->Initialise(LODData, bIsGpuUniformlyDistributedSampling, SamplingInfo.GetBuiltData().WholeMeshBuiltData[LODIndex]);
+		MeshGpuSpawnStaticBuffers->Initialise(LODData, bIsGpuUniformlyDistributedSampling, SamplingInfo.GetBuiltData().WholeMeshBuiltData[LODIndex], SpecificBones, SpecificSocketBones);
 		BeginInitResource(MeshGpuSpawnStaticBuffers);
 
 		MeshGpuSpawnDynamicBuffers = new FSkeletalMeshGpuDynamicBufferProxy();
-		MeshGpuSpawnDynamicBuffers->Initialise(LODData);
+		MeshGpuSpawnDynamicBuffers->Initialise(RefSkel, LODData);
 		BeginInitResource(MeshGpuSpawnDynamicBuffers);
 	}
 
@@ -1613,313 +1748,186 @@ void UNiagaraDataInterfaceSkeletalMesh::ValidateFunction(const FNiagaraFunctionS
 
 #endif
 
-const FString UNiagaraDataInterfaceSkeletalMesh::MeshIndexBufferName(TEXT("IndexBuffer_"));
-const FString UNiagaraDataInterfaceSkeletalMesh::MeshVertexBufferName(TEXT("VertexBuffer_"));
-const FString UNiagaraDataInterfaceSkeletalMesh::MeshSkinWeightBufferName(TEXT("VertexSkinWeightBuffer_"));
+const FString UNiagaraDataInterfaceSkeletalMesh::MeshIndexBufferName(TEXT("MeshIndexBuffer_"));
+const FString UNiagaraDataInterfaceSkeletalMesh::MeshVertexBufferName(TEXT("MeshVertexBuffer_"));
+const FString UNiagaraDataInterfaceSkeletalMesh::MeshSkinWeightBufferName(TEXT("MeshSkinWeightBuffer_"));
 const FString UNiagaraDataInterfaceSkeletalMesh::MeshCurrBonesBufferName(TEXT("MeshCurrBonesBuffer_"));
 const FString UNiagaraDataInterfaceSkeletalMesh::MeshPrevBonesBufferName(TEXT("MeshPrevBonesBuffer_"));
+const FString UNiagaraDataInterfaceSkeletalMesh::MeshCurrSamplingBonesBufferName(TEXT("MeshCurrSamplingBonesBuffer_"));
+const FString UNiagaraDataInterfaceSkeletalMesh::MeshPrevSamplingBonesBufferName(TEXT("MeshPrevSamplingBonesBuffer_"));
 const FString UNiagaraDataInterfaceSkeletalMesh::MeshTangentBufferName(TEXT("MeshTangentBuffer_"));
-const FString UNiagaraDataInterfaceSkeletalMesh::MeshTexCoordBufferName(TEXT("TexCoordBuffer_"));
+const FString UNiagaraDataInterfaceSkeletalMesh::MeshTexCoordBufferName(TEXT("MeshTexCoordBuffer_"));
+const FString UNiagaraDataInterfaceSkeletalMesh::MeshColorBufferName(TEXT("MeshColorBuffer_"));
 const FString UNiagaraDataInterfaceSkeletalMesh::MeshTriangleSamplerProbaBufferName(TEXT("MeshTriangleSamplerProbaBuffer_"));
 const FString UNiagaraDataInterfaceSkeletalMesh::MeshTriangleSamplerAliasBufferName(TEXT("MeshTriangleSamplerAliasBuffer_"));
 const FString UNiagaraDataInterfaceSkeletalMesh::MeshTriangleMatricesOffsetBufferName(TEXT("MeshTriangleMatricesOffsetBuffer_"));
+const FString UNiagaraDataInterfaceSkeletalMesh::MeshTriangleCountName(TEXT("MeshTriangleCount_"));
+const FString UNiagaraDataInterfaceSkeletalMesh::MeshVertexCountName(TEXT("MeshVertexCount_"));
+const FString UNiagaraDataInterfaceSkeletalMesh::MeshWeightStrideName(TEXT("MeshWeightStride_"));
+const FString UNiagaraDataInterfaceSkeletalMesh::MeshNumTexCoordName(TEXT("MeshNumTexCoord_"));
+const FString UNiagaraDataInterfaceSkeletalMesh::MeshNumWeightsName(TEXT("MeshNumWeights_"));
+const FString UNiagaraDataInterfaceSkeletalMesh::NumSpecificBonesName(TEXT("NumSpecificBones_"));
+const FString UNiagaraDataInterfaceSkeletalMesh::SpecificBonesName(TEXT("SpecificBones_"));
+const FString UNiagaraDataInterfaceSkeletalMesh::NumSpecificSocketBonesName(TEXT("NumSpecificSocketBones_"));
+const FString UNiagaraDataInterfaceSkeletalMesh::SpecificSocketBonesName(TEXT("SpecificSocketBones_"));
 const FString UNiagaraDataInterfaceSkeletalMesh::InstanceTransformName(TEXT("InstanceTransform_"));
 const FString UNiagaraDataInterfaceSkeletalMesh::InstancePrevTransformName(TEXT("InstancePrevTransform_"));
 const FString UNiagaraDataInterfaceSkeletalMesh::InstanceInvDeltaTimeName(TEXT("InstanceInvDeltaTime_"));
-const FString UNiagaraDataInterfaceSkeletalMesh::MeshWeightStrideByteName(TEXT("MeshWeightStrideByte_"));
-const FString UNiagaraDataInterfaceSkeletalMesh::MeshTriangleCountName(TEXT("MeshTriangleCount_"));
 const FString UNiagaraDataInterfaceSkeletalMesh::EnabledFeaturesName(TEXT("EnabledFeatures_"));
-const FString UNiagaraDataInterfaceSkeletalMesh::InputWeightStrideName(TEXT("InputWeightStride_"));
-const FString UNiagaraDataInterfaceSkeletalMesh::NumTexCoordName(TEXT("NumTexCoordName_"));
 
-bool UNiagaraDataInterfaceSkeletalMesh::GetFunctionHLSL(const FName&  DefinitionFunctionName, FString InstanceFunctionName, FNiagaraDataInterfaceGPUParamInfo& ParamInfo, FString& OutHLSL) 
+void UNiagaraDataInterfaceSkeletalMesh::GetCommonHLSL(FString& OutHLSL)
+{
+	OutHLSL += TEXT("#include \"/Plugin/FX/Niagara/Private/NiagaraDataInterfaceSkeletalMesh.ush\"\n");
+}
+
+bool UNiagaraDataInterfaceSkeletalMesh::GetFunctionHLSL(const FName& DefinitionFunctionName, FString InstanceFunctionName, FNiagaraDataInterfaceGPUParamInfo& ParamInfo, FString& OutHLSL) 
 {
 	FNDISkeletalMeshParametersName ParamNames;
 	GetNiagaraDataInterfaceParametersName(ParamNames, ParamInfo.DataInterfaceHLSLSymbol);
-	FString MeshTriCoordinateStructName = "MeshTriCoordinate";
-
-	static const TCHAR* FormatCommonFunctions = TEXT(R"(
-		void {InstanceFunctionName}_GetIndicesAndWeights(uint VertexIndex, out int4 BlendIndices, out float4 BlendWeights)
-		{
-			uint PackedBlendIndices = {MeshSkinWeightBufferName}[VertexIndex * ({InputWeightStrideName})    ];
-			uint PackedBlendWeights = {MeshSkinWeightBufferName}[VertexIndex * ({InputWeightStrideName}) + 1];
-			BlendIndices.x = PackedBlendIndices & 0xff;
-			BlendIndices.y = PackedBlendIndices >> 8 & 0xff;
-			BlendIndices.z = PackedBlendIndices >> 16 & 0xff;
-			BlendIndices.w = PackedBlendIndices >> 24 & 0xff;
-			BlendWeights.x = float(PackedBlendWeights & 0xff) / 255.0f;
-			BlendWeights.y = float(PackedBlendWeights >> 8 & 0xff) / 255.0f;
-			BlendWeights.z = float(PackedBlendWeights >> 16 & 0xff) / 255.0f;
-			BlendWeights.w = float(PackedBlendWeights >> 24 & 0xff) / 255.0f;
-		}
-
-		float3x4 {InstanceFunctionName}_GetPrevBoneMatrix(uint Bone)
-		{
-			return float3x4({MeshPrevBonesBufferName}[Bone * 3], {MeshPrevBonesBufferName}[Bone * 3 + 1], {MeshPrevBonesBufferName}[Bone * 3 + 2]);
-		}
-
-		float3x4 {InstanceFunctionName}_GetPrevSkinningMatrix(uint VertexIndex, int4 BlendIndices, float4 BlendWeights)
-		{
-			// Get the matrix offset for each vertex because BlendIndices are stored relatively to each section start vertex.
-			uint MatrixOffset = {MeshTriangleMatricesOffsetBufferName}[VertexIndex];
-
-			float3x4 Result;
-			Result  = {InstanceFunctionName}_GetPrevBoneMatrix(MatrixOffset + BlendIndices.x) * BlendWeights.x;
-			Result += {InstanceFunctionName}_GetPrevBoneMatrix(MatrixOffset + BlendIndices.y) * BlendWeights.y;
-			Result += {InstanceFunctionName}_GetPrevBoneMatrix(MatrixOffset + BlendIndices.z) * BlendWeights.z;
-			Result += {InstanceFunctionName}_GetPrevBoneMatrix(MatrixOffset + BlendIndices.w) * BlendWeights.w;
-			return Result;
-		}
-
-		float3x4 {InstanceFunctionName}_GetCurrBoneMatrix(uint Bone)
-		{
-			return float3x4({MeshCurrBonesBufferName}[Bone * 3], {MeshCurrBonesBufferName}[Bone * 3 + 1], {MeshCurrBonesBufferName}[Bone * 3 + 2]);
-		}
-
-		float3x4 {InstanceFunctionName}_GetCurrSkinningMatrix(uint VertexIndex, int4 BlendIndices, float4 BlendWeights)
-		{
-			// Get the matrix offset for each vertex because BlendIndices are stored relatively to each section start vertex.
-			uint MatrixOffset = {MeshTriangleMatricesOffsetBufferName}[VertexIndex];
-
-			float3x4 Result;
-			Result  = {InstanceFunctionName}_GetCurrBoneMatrix(MatrixOffset + BlendIndices.x) * BlendWeights.x;
-			Result += {InstanceFunctionName}_GetCurrBoneMatrix(MatrixOffset + BlendIndices.y) * BlendWeights.y;
-			Result += {InstanceFunctionName}_GetCurrBoneMatrix(MatrixOffset + BlendIndices.z) * BlendWeights.z;
-			Result += {InstanceFunctionName}_GetCurrBoneMatrix(MatrixOffset + BlendIndices.w) * BlendWeights.w;
-			return Result;
-		}
-	)");
-
-	static const TCHAR* FormatSampleSkinnedTriangleDataWSHeader = TEXT(R"(
-		void {InstanceFunctionName} (in {MeshTriCoordinateStructName} In_Coord, out float3 Out_Position, out float3 Out_Velocity, out float3 Out_Normal, out float3 Out_Binormal, out float3 Out_Tangent)
-		{
-			const float In_Interp = 1.0f;
-		)");
-
-	static const TCHAR* FormatSampleSkinnedTriangleDataWSInterpolatedHeader = TEXT(R"(
-		void {InstanceFunctionName} (in {MeshTriCoordinateStructName} In_Coord, float In_Interp, out float3 Out_Position, out float3 Out_Velocity, out float3 Out_Normal, out float3 Out_Binormal, out float3 Out_Tangent)
-		{
-		)");
-
-	static const TCHAR* FormatSampleSkinnedTriangleDataWSPart0 = TEXT(R"(
-			const bool SkinningEnabled = {EnabledFeaturesName} & 0x0002;
-
-			uint TriangleIndex = In_Coord.Tri * 3;
-			uint VertexIndex0 = {MeshIndexBufferName}[TriangleIndex  ];
-			uint VertexIndex1 = {MeshIndexBufferName}[TriangleIndex+1];
-			uint VertexIndex2 = {MeshIndexBufferName}[TriangleIndex+2];
-			
-			// I could not find a R32G32B32f format to create an SRV on that buffer. So float load it is for now...
-			float3 Vertex0 = float3({MeshVertexBufferName}[VertexIndex0*3], {MeshVertexBufferName}[VertexIndex0*3+1], {MeshVertexBufferName}[VertexIndex0*3+2]);
-			float3 Vertex1 = float3({MeshVertexBufferName}[VertexIndex1*3], {MeshVertexBufferName}[VertexIndex1*3+1], {MeshVertexBufferName}[VertexIndex1*3+2]);
-			float3 Vertex2 = float3({MeshVertexBufferName}[VertexIndex2*3], {MeshVertexBufferName}[VertexIndex2*3+1], {MeshVertexBufferName}[VertexIndex2*3+2]);
-			float3 PrevVertex0 = Vertex0;
-			float3 PrevVertex1 = Vertex1;
-			float3 PrevVertex2 = Vertex2;
-
-			float3 TangentX0 = TangentBias({MeshTangentBufferName}[VertexIndex0*2  ].xyz);
-			float4 TangentZ0 = TangentBias({MeshTangentBufferName}[VertexIndex0*2+1].xyzw);
-			float3 TangentX1 = TangentBias({MeshTangentBufferName}[VertexIndex1*2  ].xyz);
-			float4 TangentZ1 = TangentBias({MeshTangentBufferName}[VertexIndex1*2+1].xyzw);
-			float3 TangentX2 = TangentBias({MeshTangentBufferName}[VertexIndex2*2  ].xyz);
-			float4 TangentZ2 = TangentBias({MeshTangentBufferName}[VertexIndex2*2+1].xyzw);
-
-			if(SkinningEnabled)
-			{
-				int4 BlendIndices0;
-				int4 BlendIndices1;
-				int4 BlendIndices2;
-				float4 BlendWeights0;
-				float4 BlendWeights1;
-				float4 BlendWeights2;
-
-				{InstanceFunctionName}_GetIndicesAndWeights(VertexIndex0, BlendIndices0, BlendWeights0);
-				{InstanceFunctionName}_GetIndicesAndWeights(VertexIndex1, BlendIndices1, BlendWeights1);
-				{InstanceFunctionName}_GetIndicesAndWeights(VertexIndex2, BlendIndices2, BlendWeights2);
-
-				float3x4 PrevBoneMatrix0 = {InstanceFunctionName}_GetPrevSkinningMatrix(VertexIndex0, BlendIndices0, BlendWeights0);
-				float3x4 PrevBoneMatrix1 = {InstanceFunctionName}_GetPrevSkinningMatrix(VertexIndex1, BlendIndices1, BlendWeights1);
-				float3x4 PrevBoneMatrix2 = {InstanceFunctionName}_GetPrevSkinningMatrix(VertexIndex2, BlendIndices2, BlendWeights2);
-				PrevVertex0 = mul( PrevBoneMatrix0, float4(Vertex0, 1.0f) ).xyz;
-				PrevVertex1 = mul( PrevBoneMatrix1, float4(Vertex1, 1.0f) ).xyz;
-				PrevVertex2 = mul( PrevBoneMatrix2, float4(Vertex2, 1.0f) ).xyz;
-
-				float3x4 CurrBoneMatrix0 = {InstanceFunctionName}_GetCurrSkinningMatrix(VertexIndex0, BlendIndices0, BlendWeights0);
-				float3x4 CurrBoneMatrix1 = {InstanceFunctionName}_GetCurrSkinningMatrix(VertexIndex1, BlendIndices1, BlendWeights1);
-				float3x4 CurrBoneMatrix2 = {InstanceFunctionName}_GetCurrSkinningMatrix(VertexIndex2, BlendIndices2, BlendWeights2);
-				Vertex0 = mul( CurrBoneMatrix0, float4(Vertex0, 1.0f) ).xyz;
-				Vertex1 = mul( CurrBoneMatrix1, float4(Vertex1, 1.0f) ).xyz;
-				Vertex2 = mul( CurrBoneMatrix2, float4(Vertex2, 1.0f) ).xyz;
-
-				// Not using InverseTranspose of matrices so assuming uniform scaling only (same as SkinCache)
-				TangentX0.xyz = mul( CurrBoneMatrix0, float4(TangentX0.xyz, 0.0f) ).xyz;
-				TangentZ0.xyz = mul( CurrBoneMatrix0, float4(TangentZ0.xyz, 0.0f) ).xyz;
-				TangentX1.xyz = mul( CurrBoneMatrix1, float4(TangentX1.xyz, 0.0f) ).xyz;
-				TangentZ1.xyz = mul( CurrBoneMatrix1, float4(TangentZ1.xyz, 0.0f) ).xyz;
-				TangentX2.xyz = mul( CurrBoneMatrix2, float4(TangentX2.xyz, 0.0f) ).xyz;
-				TangentZ2.xyz = mul( CurrBoneMatrix2, float4(TangentZ2.xyz, 0.0f) ).xyz;
-			}
-
-			// Evaluate current and previous world position
-			float3 WSPos = Vertex0 * In_Coord.BaryCoord.x + Vertex1 * In_Coord.BaryCoord.y + Vertex2 * In_Coord.BaryCoord.z;
-			WSPos = mul(float4(WSPos,1.0), {InstanceTransformName}).xyz;
-			float3 PrevWSPos = PrevVertex0 * In_Coord.BaryCoord.x + PrevVertex1 * In_Coord.BaryCoord.y + PrevVertex2 * In_Coord.BaryCoord.z;
-			PrevWSPos = mul(float4(PrevWSPos,1.0), {InstancePrevTransformName}).xyz;
-
-			// Not using InverseTranspose of matrices so assuming uniform scaling only (same as SkinCache)
-			float3 Binormal0 = cross(TangentZ0.xyz, TangentX0.xyz) * TangentZ0.w;
-			float3 Binormal1 = cross(TangentZ1.xyz, TangentX1.xyz) * TangentZ1.w;
-			float3 Binormal2 = cross(TangentZ2.xyz, TangentX2.xyz) * TangentZ2.w;
-			float3 Normal   = TangentZ0.xyz * In_Coord.BaryCoord.x + TangentZ1.xyz * In_Coord.BaryCoord.y + TangentZ2.xyz * In_Coord.BaryCoord.z; // Normal is TangentZ
-			float3 Tangent  = TangentX0.xyz * In_Coord.BaryCoord.x + TangentX1.xyz * In_Coord.BaryCoord.y + TangentX2.xyz * In_Coord.BaryCoord.z;
-			float3 Binormal = Binormal0.xyz * In_Coord.BaryCoord.x + Binormal1.xyz * In_Coord.BaryCoord.y + Binormal2.xyz * In_Coord.BaryCoord.z;
-			float3 NormalWorld   = mul(float4(Normal  , 0.0), {InstanceTransformName}).xyz;
-			float3 TangentWorld  = mul(float4(Tangent , 0.0), {InstanceTransformName}).xyz;
-			float3 BinormalWorld = mul(float4(Binormal, 0.0), {InstanceTransformName}).xyz;
-			
-			Out_Position = lerp(PrevWSPos, WSPos, float3(In_Interp,In_Interp,In_Interp));
-			Out_Velocity = (WSPos - PrevWSPos) * {InstanceInvDeltaTimeName};				// Velocity is unafected by spawn interpolation. That would require another set of previous data.
-			Out_Normal   = normalize(NormalWorld);
-			Out_Tangent  = normalize(TangentWorld);
-			Out_Binormal = normalize(BinormalWorld);
-		}
-		)");
-
-
 	TMap<FString, FStringFormatArg> ArgsSample = {
 		{TEXT("InstanceFunctionName"), InstanceFunctionName},
-		{TEXT("MeshTriCoordinateStructName"), MeshTriCoordinateStructName},
-		{TEXT("MeshIndexBufferName"), ParamNames.MeshIndexBufferName},
-		{TEXT("MeshVertexBufferName"), ParamNames.MeshVertexBufferName},
-		{TEXT("MeshSkinWeightBufferName"), ParamNames.MeshSkinWeightBufferName},
-		{TEXT("MeshCurrBonesBufferName"), ParamNames.MeshCurrBonesBufferName},
-		{TEXT("MeshPrevBonesBufferName"), ParamNames.MeshPrevBonesBufferName},
-		{TEXT("MeshTangentBufferName"), ParamNames.MeshTangentBufferName},
-		{TEXT("MeshTexCoordBufferName"), ParamNames.MeshTexCoordBufferName},
-		{TEXT("MeshTriangleSamplerProbaBufferName"), ParamNames.MeshTriangleSamplerProbaBufferName},
-		{TEXT("MeshTriangleSamplerAliasBufferName"), ParamNames.MeshTriangleSamplerAliasBufferName},
-		{TEXT("MeshTriangleMatricesOffsetBufferName"), ParamNames.MeshTriangleMatricesOffsetBufferName},
-		{TEXT("MeshTriangleCountName"), ParamNames.MeshTriangleCountName},
-		{TEXT("InstanceTransformName"), ParamNames.InstanceTransformName},
-		{TEXT("InstancePrevTransformName"), ParamNames.InstancePrevTransformName},
-		{TEXT("InstanceInvDeltaTimeName"), ParamNames.InstanceInvDeltaTimeName},
-		{TEXT("EnabledFeaturesName"), ParamNames.EnabledFeaturesName},
-		{TEXT("InputWeightStrideName"), ParamNames.InputWeightStrideName},
-		{TEXT("NumTexCoordName"), ParamNames.NumTexCoordName},
+		{TEXT("MeshTriCoordinateStructName"), TEXT("MeshTriCoordinate")},
+		{TEXT("MeshTriangleCount"), ParamNames.MeshTriangleCountName},
+		{TEXT("MeshVertexCount"), ParamNames.MeshVertexCountName},
+		{TEXT("GetDISkelMeshContextName"), TEXT("DISKELMESH_MAKE_CONTEXT(") + ParamInfo.DataInterfaceHLSLSymbol + TEXT(")")},
 	};
 
+	// Triangle Sampling
 	if (DefinitionFunctionName == FSkeletalMeshInterfaceHelper::RandomTriCoordName)
 	{
-		static const TCHAR *FormatSample = TEXT(R"(
-			void {InstanceFunctionName} (out {MeshTriCoordinateStructName} Out_Coord)
-			{
-				const bool UniformTriangleSamplingEnable = {EnabledFeaturesName} & 0x0001;
-
-				float RandT0 = NiagaraInternalNoise(1, 2, 3);
-				[branch]
-				if (!UniformTriangleSamplingEnable)
-				{
-					// Uniform triangle id selection
-					Out_Coord.Tri = min(uint(RandT0*float({MeshTriangleCountName})), {MeshTriangleCountName}-1); // avoid % by using mul/min to Tri = MeshTriangleCountName
-				}
-				else
-				{
-					// Uniform area weighted position selection (using alias method from Alias method from FWeightedRandomSampler)
-					uint TriangleIndex = min(uint(RandT0*float({MeshTriangleCountName})), {MeshTriangleCountName}-1);
-					float TriangleProbability = {MeshTriangleSamplerProbaBufferName}[TriangleIndex];
-
-					// Alias check
-					float RandT1 = NiagaraInternalNoise(1, 2, 3);
-					if( RandT1 > TriangleProbability )
-					{
-						TriangleIndex = {MeshTriangleSamplerAliasBufferName}[TriangleIndex];
-					}
-					Out_Coord.Tri = TriangleIndex;
-				}
-
-				float r0 = NiagaraInternalNoise(1, 2, 3);
-				float r1 = NiagaraInternalNoise(1, 2, 3);
-				float sqrt0 = sqrt(r0);
-				float sqrt1 = sqrt(r1);
-				Out_Coord.BaryCoord = float3(1.0f - sqrt0, sqrt0 * (1.0 - r1), r1 * sqrt0);
-			//	Out_Coord.BaryCoord = float3(1.0f, 0.0f, 0.0f);
-			}
-			)");
+		static const TCHAR* FormatSample = TEXT("void {InstanceFunctionName} (out {MeshTriCoordinateStructName} OutCoord) { {GetDISkelMeshContextName} DISKelMesh_RandomTriCoord(DIContext, OutCoord.Tri, OutCoord.BaryCoord); }");
 		OutHLSL += FString::Format(FormatSample, ArgsSample);
 	}
 	else if (DefinitionFunctionName == FSkeletalMeshInterfaceHelper::GetSkinnedTriangleDataWSName)
 	{
-		OutHLSL += FString::Format(FormatCommonFunctions, ArgsSample);
-		OutHLSL += FString::Format(FormatSampleSkinnedTriangleDataWSHeader, ArgsSample);
-		OutHLSL += FString::Format(FormatSampleSkinnedTriangleDataWSPart0, ArgsSample);
+		static const TCHAR* FormatSample = TEXT("void {InstanceFunctionName} (in {MeshTriCoordinateStructName} InCoord, out float3 OutPosition, out float3 OutVelocity, out float3 OutNormal, out float3 OutBinormal, out float3 OutTangent) { {GetDISkelMeshContextName} DISKelMesh_GetSkinnedTriangleDataWS(DIContext, InCoord.Tri, InCoord.BaryCoord, OutPosition, OutVelocity, OutNormal, OutBinormal, OutTangent); }");
+		OutHLSL += FString::Format(FormatSample, ArgsSample);
 	}
 	else if (DefinitionFunctionName == FSkeletalMeshInterfaceHelper::GetSkinnedTriangleDataWSInterpName)
 	{
-		OutHLSL += FString::Format(FormatCommonFunctions, ArgsSample);
-		OutHLSL += FString::Format(FormatSampleSkinnedTriangleDataWSInterpolatedHeader, ArgsSample);
-		OutHLSL += FString::Format(FormatSampleSkinnedTriangleDataWSPart0, ArgsSample);
+		static const TCHAR* FormatSample = TEXT("void {InstanceFunctionName} (in {MeshTriCoordinateStructName} InCoord, in float InInterp, out float3 OutPosition, out float3 OutVelocity, out float3 OutNormal, out float3 OutBinormal, out float3 OutTangent) { {GetDISkelMeshContextName} DISKelMesh_GetSkinnedTriangleDataInterpolatedWS(DIContext, InCoord.Tri, InCoord.BaryCoord, InInterp, OutPosition, OutVelocity, OutNormal, OutBinormal, OutTangent); }");
+		OutHLSL += FString::Format(FormatSample, ArgsSample);
 	}
-	else if (DefinitionFunctionName == FSkeletalMeshInterfaceHelper::GetTriColorName)
+	else if (DefinitionFunctionName == FSkeletalMeshInterfaceHelper::GetSkinnedTriangleDataName)
 	{
-		static const TCHAR *FormatSample = TEXT(R"(
-				void {InstanceFunctionName} (in {MeshTriCoordinateStructName} In_Coord, out float4 Out_Color)
-				{
-					Out_Color = 0.0f;
-				}
-				)");
+		static const TCHAR* FormatSample = TEXT("void {InstanceFunctionName} (in {MeshTriCoordinateStructName} InCoord, out float3 OutPosition, out float3 OutVelocity, out float3 OutNormal, out float3 OutBinormal, out float3 OutTangent) { {GetDISkelMeshContextName} DISKelMesh_GetSkinnedTriangleData(DIContext, InCoord.Tri, InCoord.BaryCoord, OutPosition, OutVelocity, OutNormal, OutBinormal, OutTangent); }");
+		OutHLSL += FString::Format(FormatSample, ArgsSample);
+	}
+	else if (DefinitionFunctionName == FSkeletalMeshInterfaceHelper::GetSkinnedTriangleDataInterpName)
+	{
+		static const TCHAR* FormatSample = TEXT("void {InstanceFunctionName} (in {MeshTriCoordinateStructName} InCoord, in float InInterp, out float3 OutPosition, out float3 OutVelocity, out float3 OutNormal, out float3 OutBinormal, out float3 OutTangent) { {GetDISkelMeshContextName} DISKelMesh_GetSkinnedTriangleDataInterpolated(DIContext, InCoord.Tri, InCoord.BaryCoord, InInterp, OutPosition, OutVelocity, OutNormal, OutBinormal, OutTangent); }");
 		OutHLSL += FString::Format(FormatSample, ArgsSample);
 	}
 	else if (DefinitionFunctionName == FSkeletalMeshInterfaceHelper::GetTriUVName)
 	{
-		static const TCHAR *FormatSample = TEXT(R"(
-			void {InstanceFunctionName} (in {MeshTriCoordinateStructName} In_Coord, in int In_UVSet, out float2 Out_UV)
-			{
-				if({NumTexCoordName}>0)
-				{
-					uint TriangleIndex = In_Coord.Tri * 3;
-					uint VertexIndex0 = {MeshIndexBufferName}[TriangleIndex  ];
-					uint VertexIndex1 = {MeshIndexBufferName}[TriangleIndex+1];
-					uint VertexIndex2 = {MeshIndexBufferName}[TriangleIndex+2];
-
-					uint stride = {NumTexCoordName};
-					uint SelectedUVSet = clamp(In_UVSet, 0, {NumTexCoordName}-1);
-					float2 UV0 = {MeshTexCoordBufferName}[VertexIndex0 * stride + SelectedUVSet];
-					float2 UV1 = {MeshTexCoordBufferName}[VertexIndex1 * stride + SelectedUVSet];
-					float2 UV2 = {MeshTexCoordBufferName}[VertexIndex2 * stride + SelectedUVSet];
-
-					Out_UV = UV0 * In_Coord.BaryCoord.x + UV1 * In_Coord.BaryCoord.y + UV2 * In_Coord.BaryCoord.z;
-				}
-				else	
-				{
-					Out_UV = 0.0f;
-				}
-			}
-			)");
+		static const TCHAR* FormatSample = TEXT("void {InstanceFunctionName} (in {MeshTriCoordinateStructName} InCoord, in int InUVSet, out float2 OutUV) { {GetDISkelMeshContextName} DISKelMesh_GetTriUV(DIContext, InCoord.Tri, InCoord.BaryCoord, InUVSet, OutUV); }");
 		OutHLSL += FString::Format(FormatSample, ArgsSample);
 	}
-	//else if (DefinitionFunctionName == FSkeletalMeshInterfaceHelper::IsValidTriCoordName)
-	//{
-	//	// unimplemented();
-	//}
-	//else if (DefinitionFunctionName == FSkeletalMeshInterfaceHelper::GetSkinnedTriangleDataName)
-	//{
-	//	// unimplemented();
-	//}
-	//else if (DefinitionFunctionName == FSkeletalMeshInterfaceHelper::GetSkinnedTriangleDataInterpName)
-	//{
-	//	// unimplemented();
-	//}
-	//else if (DefinitionFunctionName == FSkeletalMeshInterfaceHelper::GetTriangleCountName)
-	//{
-	//	// unimplemented();
-	//}
-	//else if (DefinitionFunctionName == FSkeletalMeshInterfaceHelper::GetTriangleAtName)
-	//{
-	//	// unimplemented();
-	//}
-	//else if (DefinitionFunctionName == FSkeletalMeshInterfaceHelper::GetTriCoordVerticesName)
-	//{
-	//}
+	else if (DefinitionFunctionName == FSkeletalMeshInterfaceHelper::GetTriColorName)
+	{
+		static const TCHAR* FormatSample = TEXT("void {InstanceFunctionName} (in {MeshTriCoordinateStructName} InCoord, out float4 OutColor) { {GetDISkelMeshContextName} DISkelMesh_GetTriColor(DIContext, InCoord.Tri, InCoord.BaryCoord, OutColor); }");
+		OutHLSL += FString::Format(FormatSample, ArgsSample);
+	}
+	else if (DefinitionFunctionName == FSkeletalMeshInterfaceHelper::IsValidTriCoordName)
+	{
+		static const TCHAR* FormatSample = TEXT("void {InstanceFunctionName} (in {MeshTriCoordinateStructName} InCoord, out bool IsValid) { {GetDISkelMeshContextName} IsValid = InCoord.Tri < DIContext.MeshTriangleCount; }");
+		OutHLSL += FString::Format(FormatSample, ArgsSample);
+	}
+	else if (DefinitionFunctionName == FSkeletalMeshInterfaceHelper::GetTriCoordVerticesName)
+	{
+		static const TCHAR* FormatSample = TEXT("void {InstanceFunctionName} (in int TriangleIndex, out int OutVertexIndex0, out int OutVertexIndex1, out int OutVertexIndex2) { {GetDISkelMeshContextName} DISkelMesh_GetTriVertices(DIContext, TriangleIndex, OutVertexIndex0, OutVertexIndex1, OutVertexIndex2); }");
+		OutHLSL += FString::Format(FormatSample, ArgsSample);
+	}
+	// Bone Sampling
+	else if (DefinitionFunctionName == FSkeletalMeshInterfaceHelper::GetSkinnedBoneDataName)
+	{
+		static const TCHAR* FormatSample = TEXT("void {InstanceFunctionName} (in int InBone, out float3 OutPosition, out float4 OutRotation, out float3 OutVelocity) { {GetDISkelMeshContextName} DISkelMesh_GetSkinnedBone(DIContext, InBone, OutPosition, OutRotation, OutVelocity); }");
+		OutHLSL += FString::Format(FormatSample, ArgsSample);
+	}
+	else if (DefinitionFunctionName == FSkeletalMeshInterfaceHelper::GetSkinnedBoneDataInterpolatedName)
+	{
+		static const TCHAR* FormatSample = TEXT("void {InstanceFunctionName} (in int InBone, in float Interp, out float3 OutPosition, out float4 OutRotation, out float3 OutVelocity) { {GetDISkelMeshContextName} DISkelMesh_GetSkinnedBoneInterpolated(DIContext, InBone, Interp, OutPosition, OutRotation, OutVelocity); }");
+		OutHLSL += FString::Format(FormatSample, ArgsSample);
+	}
+	else if (DefinitionFunctionName == FSkeletalMeshInterfaceHelper::GetSkinnedBoneDataWSName)
+	{
+		static const TCHAR* FormatSample = TEXT("void {InstanceFunctionName} (in int InBone, out float3 OutPosition, out float4 OutRotation, out float3 OutVelocity) { {GetDISkelMeshContextName} DISkelMesh_GetSkinnedBoneWS(DIContext, InBone, OutPosition, OutRotation, OutVelocity); }");
+		OutHLSL += FString::Format(FormatSample, ArgsSample);
+	}
+	else if (DefinitionFunctionName == FSkeletalMeshInterfaceHelper::GetSkinnedBoneDataWSInterpolatedName)
+	{
+		static const TCHAR* FormatSample = TEXT("void {InstanceFunctionName} (in int InBone, in float Interp, out float3 OutPosition, out float4 OutRotation, out float3 OutVelocity) { {GetDISkelMeshContextName} DISkelMesh_GetSkinnedBoneInterpolatedWS(DIContext, InBone, Interp, OutPosition, OutRotation, OutVelocity); }");
+		OutHLSL += FString::Format(FormatSample, ArgsSample);
+	}
+	// Vertex Sampling
+	else if (DefinitionFunctionName == FSkeletalMeshInterfaceHelper::RandomVertexName)
+	{
+		static const TCHAR* FormatSample = TEXT("void {InstanceFunctionName}(out int OutVertex) { {GetDISkelMeshContextName} DISkelMesh_GetRandomVertex(DIContext, OutVertex); }");
+		OutHLSL += FString::Format(FormatSample, ArgsSample);
+	}
+	else if (DefinitionFunctionName == FSkeletalMeshInterfaceHelper::IsValidVertexName)
+	{
+		static const TCHAR* FormatSample = TEXT("void {InstanceFunctionName} (in int Vertex, out bool IsValid) { {GetDISkelMeshContextName} DISkelMesh_IsValidVertex(DIContext, Vertex, IsValid); }");
+		OutHLSL += FString::Format(FormatSample, ArgsSample);
+	}
+	else if (DefinitionFunctionName == FSkeletalMeshInterfaceHelper::GetSkinnedVertexDataName)
+	{
+		static const TCHAR* FormatSample = TEXT("void {InstanceFunctionName} (in int Vertex, out float3 OutPosition, out float3 OutVelocity) { {GetDISkelMeshContextName} DISkelMesh_GetSkinnedVertex(DIContext, Vertex, OutPosition, OutVelocity); }");
+		OutHLSL += FString::Format(FormatSample, ArgsSample);
+	}
+	else if (DefinitionFunctionName == FSkeletalMeshInterfaceHelper::GetSkinnedVertexDataWSName)
+	{
+		static const TCHAR* FormatSample = TEXT("void {InstanceFunctionName} (in int Vertex, out float3 OutPosition, out float3 OutVelocity) { {GetDISkelMeshContextName} DISkelMesh_GetSkinnedVertexWS(DIContext, Vertex, OutPosition, OutVelocity); }");
+		OutHLSL += FString::Format(FormatSample, ArgsSample);
+	}
+	else if (DefinitionFunctionName == FSkeletalMeshInterfaceHelper::GetVertexColorName)
+	{
+		static const TCHAR* FormatSample = TEXT("void {InstanceFunctionName} (in int Vertex, out float4 OutColor) { {GetDISkelMeshContextName} DISkelMesh_GetVertexColor(DIContext, Vertex, OutColor); }");
+		OutHLSL += FString::Format(FormatSample, ArgsSample);
+	}
+	else if (DefinitionFunctionName == FSkeletalMeshInterfaceHelper::GetVertexUVName)
+	{
+		static const TCHAR* FormatSample = TEXT("void {InstanceFunctionName} (in int Vertex, in int UVSet, out float2 OutUV) { {GetDISkelMeshContextName} DISkelMesh_GetVertexUV(DIContext, Vertex, UVSet, OutUV); }");
+		OutHLSL += FString::Format(FormatSample, ArgsSample);
+	}
+	// Specific Bone
+	else if (DefinitionFunctionName == FSkeletalMeshInterfaceHelper::GetSpecificBoneCountName)
+	{
+		static const TCHAR* FormatSample = TEXT("void {InstanceFunctionName} (out int Count) { {GetDISkelMeshContextName} DISkelMesh_GetSpecificBoneCount(DIContext, Count); }");
+		OutHLSL += FString::Format(FormatSample, ArgsSample);
+	}
+	else if (DefinitionFunctionName == FSkeletalMeshInterfaceHelper::GetSpecificBoneAtName)
+	{
+		static const TCHAR* FormatSample = TEXT("void {InstanceFunctionName} (in int BoneIndex, out int Bone) { {GetDISkelMeshContextName} DISkelMesh_GetSpecificBoneAt(DIContext, BoneIndex, Bone); }");
+		OutHLSL += FString::Format(FormatSample, ArgsSample);
+	}
+	else if (DefinitionFunctionName == FSkeletalMeshInterfaceHelper::RandomSpecificBoneName)
+	{
+		static const TCHAR* FormatSample = TEXT("void {InstanceFunctionName} (out int Bone) { {GetDISkelMeshContextName} DISkelMesh_RandomSpecificBone(DIContext, Bone); }");
+		OutHLSL += FString::Format(FormatSample, ArgsSample);
+	}
+	// Specific Socket
+	else if (DefinitionFunctionName == FSkeletalMeshInterfaceHelper::GetSpecificSocketCountName)
+	{
+		static const TCHAR* FormatSample = TEXT("void {InstanceFunctionName} (out int Count) { {GetDISkelMeshContextName} DISkelMesh_GetSpecificSocketCount(DIContext, Count); }");
+		OutHLSL += FString::Format(FormatSample, ArgsSample);
+	}
+	else if (DefinitionFunctionName == FSkeletalMeshInterfaceHelper::GetSpecificSocketBoneAtName)
+	{
+		static const TCHAR* FormatSample = TEXT("void {InstanceFunctionName} (in int SocketIndex, out int Bone) { {GetDISkelMeshContextName} DISkelMesh_GetSpecificSocketBoneAt(DIContext, SocketIndex, Bone); }");
+		OutHLSL += FString::Format(FormatSample, ArgsSample);
+	}
+	else if (DefinitionFunctionName == FSkeletalMeshInterfaceHelper::RandomSpecificSocketBoneName)
+	{
+		static const TCHAR* FormatSample = TEXT("void {InstanceFunctionName} (out int SocketBone) { {GetDISkelMeshContextName} DISkelMesh_RandomSpecificSocketBone(DIContext, SocketBone); }");
+		OutHLSL += FString::Format(FormatSample, ArgsSample);
+	}
+	// Unsupported functionality
+	//else if (DefinitionFunctionName == FSkeletalMeshInterfaceHelper::GetVertexCountName)				// void GetFilteredVertexCount(out int VertexCount)
+	//else if (DefinitionFunctionName == FSkeletalMeshInterfaceHelper::GetVertexAtName)					// void GetFilteredVertex(in int FilterdIndex, out int VertexIndex)
+	//else if (DefinitionFunctionName == FSkeletalMeshInterfaceHelper::GetTriangleCountName)			// void GetFilteredTriangleCount(out int OutTriangleCount)
+	//else if (DefinitionFunctionName == FSkeletalMeshInterfaceHelper::GetTriangleAtName)				// void GetFilteredTriangle(in int Triangle, out {MeshTriCoordinateStructName} OutCoord)
+	//else if (DefinitionFunctionName == FSkeletalMeshInterfaceHelper::IsValidBoneName)					// void IsValidBoneName(in int BoneIndex, out bool IsValid)
 	else
 	{
 		// This function is not support
@@ -1928,31 +1936,13 @@ bool UNiagaraDataInterfaceSkeletalMesh::GetFunctionHLSL(const FName&  Definition
 
 	OutHLSL += TEXT("\n");
 	return true;
-
 }
+
 void UNiagaraDataInterfaceSkeletalMesh::GetParameterDefinitionHLSL(FNiagaraDataInterfaceGPUParamInfo& ParamInfo, FString& OutHLSL)
 {
-	FNDISkeletalMeshParametersName ParamNames;
-	GetNiagaraDataInterfaceParametersName(ParamNames, ParamInfo.DataInterfaceHLSLSymbol);
-
-	OutHLSL += TEXT("Buffer<uint> ") + ParamNames.MeshIndexBufferName + TEXT(";\n");
-	OutHLSL += TEXT("Buffer<float> ") + ParamNames.MeshVertexBufferName + TEXT(";\n");
-	OutHLSL += TEXT("Buffer<uint> ") + ParamNames.MeshSkinWeightBufferName + TEXT(";\n");
-	OutHLSL += TEXT("Buffer<float4> ") + ParamNames.MeshCurrBonesBufferName + TEXT(";\n");
-	OutHLSL += TEXT("Buffer<float4> ") + ParamNames.MeshPrevBonesBufferName + TEXT(";\n");
-	OutHLSL += TEXT("Buffer<float4> ") + ParamNames.MeshTangentBufferName + TEXT(";\n");
-	OutHLSL += TEXT("Buffer<float2> ") + ParamNames.MeshTexCoordBufferName + TEXT(";\n");
-	OutHLSL += TEXT("Buffer<float> ") + ParamNames.MeshTriangleSamplerProbaBufferName + TEXT(";\n");
-	OutHLSL += TEXT("Buffer<uint> ") + ParamNames.MeshTriangleSamplerAliasBufferName + TEXT(";\n");
-	OutHLSL += TEXT("Buffer<uint> ") + ParamNames.MeshTriangleMatricesOffsetBufferName + TEXT(";\n");
-	OutHLSL += TEXT("uint ") + ParamNames.MeshTriangleCountName + TEXT(";\n");
-	OutHLSL += TEXT("float4x4 ") + ParamNames.InstanceTransformName + TEXT(";\n");
-	OutHLSL += TEXT("float4x4 ") + ParamNames.InstancePrevTransformName + TEXT(";\n");
-	OutHLSL += TEXT("float ") + ParamNames.InstanceInvDeltaTimeName + TEXT(";\n");
-	OutHLSL += TEXT("uint ") + ParamNames.EnabledFeaturesName + TEXT(";\n");
-	OutHLSL += TEXT("uint ") + ParamNames.InputWeightStrideName + TEXT(";\n");
-	OutHLSL += TEXT("uint ") + ParamNames.NumTexCoordName + TEXT(";\n");
+	OutHLSL += TEXT("DISKELMESH_DECLARE_CONSTANTS(") + ParamInfo.DataInterfaceHLSLSymbol + TEXT(")\n");
 }
+
 FNiagaraDataInterfaceParametersCS* UNiagaraDataInterfaceSkeletalMesh::ConstructComputeParameters() const
 {
 	return new FNiagaraDataInterfaceParametersCS_SkeletalMesh();
@@ -2020,5 +2010,4 @@ void FSkeletalMeshAccessorHelper::Init<
 	SamplingRegionBuiltData = &SamplingInfo.GetRegionBuiltData(InstData->SamplingRegionIndices[0]);
 }
 
-
 #undef LOCTEXT_NAMESPACE
diff --git a/Engine/Plugins/FX/Niagara/Source/Niagara/Private/NiagaraDataInterfaceStaticMesh.cpp b/Engine/Plugins/FX/Niagara/Source/Niagara/Private/NiagaraDataInterfaceStaticMesh.cpp
--- a/Engine/Plugins/FX/Niagara/Source/Niagara/Private/NiagaraDataInterfaceStaticMesh.cpp
+++ b/Engine/Plugins/FX/Niagara/Source/Niagara/Private/NiagaraDataInterfaceStaticMesh.cpp
@@ -465,15 +465,6 @@ struct FNiagaraDataInterfaceParametersCS_StaticMesh : public FNiagaraDataInterfa
 		InstanceWorldVelocity.Bind(ParameterMap, *ParamNames.InstanceWorldVelocityName);
 		AreaWeightedSampling.Bind(ParameterMap, *ParamNames.AreaWeightedSamplingName);
 		NumTexCoord.Bind(ParameterMap, *ParamNames.NumTexCoordName);
-
-		if (!MeshIndexBuffer.IsBound())
-		{
-			UE_LOG(LogNiagara, Warning, TEXT("Binding failed for FNiagaraDataInterfaceParametersCS_StaticMesh Texture %s. Was it optimized out?"), *ParamNames.MeshIndexBufferName)
-		}
-		if (!MeshVertexBuffer.IsBound())
-		{
-			UE_LOG(LogNiagara, Warning, TEXT("Binding failed for FNiagaraDataInterfaceParametersCS_StaticMesh Sampler %s. Was it optimized out?"), *ParamNames.MeshVertexBufferName)
-		}
 	}
 
 	virtual void Serialize(FArchive& Ar)override
@@ -2030,7 +2021,7 @@ bool UNiagaraDataInterfaceStaticMesh::GetFunctionHLSL(const FName&  DefinitionFu
 		static const TCHAR *FormatSample = TEXT(R"(
 			void {InstanceFunctionName} (in int In_Section, out {MeshTriCoordinateStructName} Out_Coord)
 			{
-				int Section = clamp(In_Section, 0, {SectionCountName} - 1);
+				int Section = clamp(In_Section, 0, (int)({SectionCountName} - 1));
 
 				uint4 SectionData = {MeshSectionBufferName}[Section];
 				uint SectionFirstTriangle = SectionData.x;
@@ -2245,7 +2236,7 @@ bool UNiagaraDataInterfaceStaticMesh::GetFunctionHLSL(const FName&  DefinitionFu
 					uint VertexIndex2 = {MeshIndexBufferName}[TriangleIndex+2];
 
 					uint stride = {NumTexCoordName};
-					uint SelectedUVSet = clamp(In_UVSet, 0, {NumTexCoordName}-1);
+					uint SelectedUVSet = clamp((uint)In_UVSet, 0, {NumTexCoordName}-1);
 					float2 UV0 = {MeshTexCoordBufferName}[VertexIndex0 * stride + SelectedUVSet];
 					float2 UV1 = {MeshTexCoordBufferName}[VertexIndex1 * stride + SelectedUVSet];
 					float2 UV2 = {MeshTexCoordBufferName}[VertexIndex2 * stride + SelectedUVSet];
diff --git a/Engine/Plugins/FX/Niagara/Source/NiagaraCore/Private/NiagaraCustomVersion.cpp b/Engine/Plugins/FX/Niagara/Source/NiagaraCore/Private/NiagaraCustomVersion.cpp
--- a/Engine/Plugins/FX/Niagara/Source/NiagaraCore/Private/NiagaraCustomVersion.cpp
+++ b/Engine/Plugins/FX/Niagara/Source/NiagaraCore/Private/NiagaraCustomVersion.cpp
@@ -8,4 +8,4 @@ const FGuid FNiagaraCustomVersion::GUID(0xFCF57AFA, 0x50764283, 0xB9A9E658, 0xFF
 // Register the custom version with core
 FCustomVersionRegistration GRegisterNiagaraCustomVersion(FNiagaraCustomVersion::GUID, FNiagaraCustomVersion::LatestVersion, TEXT("NiagaraVer"));
 
-const FGuid FNiagaraCustomVersion::LatestScriptCompileVersion(0x4182078E, 0x4D784761, 0xB4205380, 0x84A48D4E);
+const FGuid FNiagaraCustomVersion::LatestScriptCompileVersion(0x8D38659E, 0xFF5CC945, 0xA718F14F, 0x52E92EBA);
diff --git a/Engine/Plugins/FX/Niagara/Source/NiagaraEditor/Private/NiagaraHlslTranslator.cpp b/Engine/Plugins/FX/Niagara/Source/NiagaraEditor/Private/NiagaraHlslTranslator.cpp
--- a/Engine/Plugins/FX/Niagara/Source/NiagaraEditor/Private/NiagaraHlslTranslator.cpp
+++ b/Engine/Plugins/FX/Niagara/Source/NiagaraEditor/Private/NiagaraHlslTranslator.cpp
@@ -1003,23 +1003,6 @@ const FNiagaraTranslateResults &FHlslNiagaraTranslator::Translate(const FNiagara
 	if (FunctionCtx())
 		return TranslateResults;
 
-	// MASSIVE HACK - Tracked in JIRA UE-69298
-	// Hardcoded random function accessible from inner part of node implementation.
-	// It works for now at least and avoids exposing every random needed in the UI. 
-	// Temporary solution, it will be replaced when a design is validated.
-	if (CompilationTarget == ENiagaraSimTarget::GPUComputeSim)
-	{
-		HlslOutput += TEXT(R"(
-		float NiagaraInternalNoise(uint u, uint v, uint s)
-		{
-			static uint RandomSeedOffset = 0;
-			uint Seed = (u * 1664525u + v) + s + RandomSeedOffset;
-			RandomSeedOffset += Seed;
-			return float(Rand3DPCG32(int3(u,v,Seed)).x) / 4294967296.0f;
-		}
-		)");
-	}
-
 	//Now evaluate all the code chunks to generate the shader code.
 	//FString HlslOutput;
 	if (TranslateResults.bHLSLGenSucceeded)
@@ -1661,13 +1644,13 @@ void FHlslNiagaraTranslator::DefineDataSetWriteFunction(FString &HlslOutputStrin
 	HlslOutput += TEXT("}\n\n");
 }
 
-
-
 void FHlslNiagaraTranslator::DefineDataInterfaceHLSL(FString &InHlslOutput)
 {
+	FString InterfaceCommonHLSL;
 	FString InterfaceUniformHLSL;
 	FString InterfaceFunctionHLSL;
 	TArray<FString> BufferParamNames;
+	TSet<FName> InterfaceClasses;
 	for (uint32 i = 0; i < 32; i++)
 	{
 		BufferParamNames.Add(TEXT("DataInterfaceBuffer_") + FString::FromInt(i));
@@ -1683,6 +1666,12 @@ void FHlslNiagaraTranslator::DefineDataInterfaceHLSL(FString &InHlslOutput)
 		UNiagaraDataInterface* CDO = Cast<UNiagaraDataInterface>(*FoundCDO);
 		if (CDO && CDO->CanExecuteOnTarget(ENiagaraSimTarget::GPUComputeSim))
 		{
+			if ( !InterfaceClasses.Contains(Info.Type.GetFName()) )
+			{
+				CDO->GetCommonHLSL(InterfaceCommonHLSL);
+				InterfaceClasses.Add(Info.Type.GetFName());
+			}
+
 			FString OwnerIDString = Info.Name.ToString();
 			FString SanitizedOwnerIDString = GetSanitizedSymbolName(OwnerIDString, true);
 
@@ -1719,7 +1708,7 @@ void FHlslNiagaraTranslator::DefineDataInterfaceHLSL(FString &InHlslOutput)
 				, nullptr, nullptr);
 		}
 	}
-	InHlslOutput += InterfaceUniformHLSL + InterfaceFunctionHLSL;
+	InHlslOutput += InterfaceCommonHLSL + InterfaceUniformHLSL + InterfaceFunctionHLSL;
 }
 
 
diff --git a/Engine/Plugins/FX/Niagara/Source/NiagaraShader/Private/NiagaraShaderDerivedDataVersion.h b/Engine/Plugins/FX/Niagara/Source/NiagaraShader/Private/NiagaraShaderDerivedDataVersion.h
--- a/Engine/Plugins/FX/Niagara/Source/NiagaraShader/Private/NiagaraShaderDerivedDataVersion.h
+++ b/Engine/Plugins/FX/Niagara/Source/NiagaraShader/Private/NiagaraShaderDerivedDataVersion.h
@@ -11,4 +11,4 @@ NiagaraShaderDerivedDataVersion.h: Shader derived data version for Niagara.
 // In case of merge conflicts with DDC versions, you MUST generate a new GUID and set this new
 // guid as version
 
-#define NIAGARASHADERMAP_DERIVEDDATA_VER		TEXT("DE7107EEE3F66B4EA8F73C17D107D3E6")
+#define NIAGARASHADERMAP_DERIVEDDATA_VER		TEXT("068D842C69C776419BF530A17E1A388B")
