- Fix GCC compile error for FShaderResourceViewRHIParamRef vs FShaderResourceViewRHIRef

#rnx

diff --git a/Engine/Plugins/FX/Niagara/Source/Niagara/Private/NiagaraDataInterfaceSkeletalMesh.cpp b/Engine/Plugins/FX/Niagara/Source/Niagara/Private/NiagaraDataInterfaceSkeletalMesh.cpp
--- a/Engine/Plugins/FX/Niagara/Source/Niagara/Private/NiagaraDataInterfaceSkeletalMesh.cpp
+++ b/Engine/Plugins/FX/Niagara/Source/Niagara/Private/NiagaraDataInterfaceSkeletalMesh.cpp
@@ -758,10 +758,13 @@ struct FNiagaraDataInterfaceParametersCS_SkeletalMesh : public FNiagaraDataInter
 				SetSRVParameter(RHICmdList, ComputeShaderRHI, MeshTriangleMatricesOffsetBuffer, FNiagaraRenderer::GetDummyUIntBuffer().SRV);
 			}
 
+			FShaderResourceViewRHIParamRef SpecificBonesSRV = StaticBuffers->GetNumSpecificBones() > 0 ? StaticBuffers->GetSpecificBonesSRV() : FNiagaraRenderer::GetDummyUIntBuffer().SRV.GetReference();
 			SetShaderValue(RHICmdList, ComputeShaderRHI, NumSpecificBones, StaticBuffers->GetNumSpecificBones());
-			SetSRVParameter(RHICmdList, ComputeShaderRHI, SpecificBones, StaticBuffers->GetNumSpecificBones() > 0 ? StaticBuffers->GetSpecificBonesSRV() : FNiagaraRenderer::GetDummyUIntBuffer().SRV);
+			SetSRVParameter(RHICmdList, ComputeShaderRHI, SpecificBones, SpecificBonesSRV);
+
+			FShaderResourceViewRHIParamRef SpecificSocketBonesSRV = StaticBuffers->GetNumSpecificSocketBones() > 0 ? StaticBuffers->GetSpecificSocketBonesSRV() : FNiagaraRenderer::GetDummyUIntBuffer().SRV.GetReference();
 			SetShaderValue(RHICmdList, ComputeShaderRHI, NumSpecificSocketBones, StaticBuffers->GetNumSpecificSocketBones());
-			SetSRVParameter(RHICmdList, ComputeShaderRHI, SpecificSocketBones, StaticBuffers->GetNumSpecificSocketBones() > 0 ? StaticBuffers->GetSpecificSocketBonesSRV() : FNiagaraRenderer::GetDummyUIntBuffer().SRV);
+			SetSRVParameter(RHICmdList, ComputeShaderRHI, SpecificSocketBones, SpecificSocketBonesSRV);
 
 			SetShaderValue(RHICmdList, ComputeShaderRHI, InstanceTransform, InstanceData->Transform);
 			SetShaderValue(RHICmdList, ComputeShaderRHI, InstancePrevTransform, InstanceData->PrevTransform);
