Adding AnimNotify support for Niagara.

#tests created auto-test that will go in a different CL
#5176

diff --git a/Engine/Plugins/FX/Niagara/Niagara.uplugin b/Engine/Plugins/FX/Niagara/Niagara.uplugin
--- a/Engine/Plugins/FX/Niagara/Niagara.uplugin
+++ b/Engine/Plugins/FX/Niagara/Niagara.uplugin
@@ -26,12 +26,16 @@
 		  "Type": "Runtime",
 		  "LoadingPhase": "PreDefault"
 		},
+		{
+			"Name" : "NiagaraAnimNotifies",
+			"Type" : "Runtime",
+			"LoadingPhase" : "PreDefault"
+		},
 		{
 			"Name" : "NiagaraShader",
 			"Type" : "Runtime",
 			"LoadingPhase" : "PostConfigInit"
 		},
-
 		{
 			"Name" : "NiagaraVertexFactories",
 			"Type" : "Runtime",
diff --git a/Engine/Plugins/FX/Niagara/Source/Niagara/Public/NiagaraComponent.h b/Engine/Plugins/FX/Niagara/Source/Niagara/Public/NiagaraComponent.h
--- a/Engine/Plugins/FX/Niagara/Source/Niagara/Public/NiagaraComponent.h
+++ b/Engine/Plugins/FX/Niagara/Source/Niagara/Public/NiagaraComponent.h
@@ -47,6 +47,7 @@ public:
 	void SetVectorParameter(FName ParameterName, FVector Param) override;
 	void SetColorParameter(FName ParameterName, FLinearColor Param) override;
 	void SetActorParameter(FName ParameterName, class AActor* Param) override;
+	virtual UFXSystemAsset* GetFXSystemAsset() const override { return Asset; };
 	/********* UFXSystemComponent *********/
 
 private:
diff --git a/Engine/Plugins/FX/Niagara/Source/NiagaraAnimNotifies/NiagaraAnimNotifies.Build.cs b/Engine/Plugins/FX/Niagara/Source/NiagaraAnimNotifies/NiagaraAnimNotifies.Build.cs
new file mode 100644
--- /dev/null
+++ b/Engine/Plugins/FX/Niagara/Source/NiagaraAnimNotifies/NiagaraAnimNotifies.Build.cs
@@ -0,0 +1,16 @@
+// Copyright 1998-2018 Epic Games, Inc. All Rights Reserved.
+
+using UnrealBuildTool;
+
+public class NiagaraAnimNotifies : ModuleRules
+{
+    public NiagaraAnimNotifies(ReadOnlyTargetRules Target) : base(Target)
+    {
+        PrivateIncludePaths.AddRange(
+            new string[] {
+                "NiagaraAnimNotifies/Private",
+            });
+
+        PublicDependencyModuleNames.AddRange(new string[] { "Core", "CoreUObject", "Engine", "Niagara" });
+    }
+}
diff --git a/Engine/Plugins/FX/Niagara/Source/NiagaraAnimNotifies/Private/AnimNotifyState_TimedNiagaraEffect.cpp b/Engine/Plugins/FX/Niagara/Source/NiagaraAnimNotifies/Private/AnimNotifyState_TimedNiagaraEffect.cpp
new file mode 100644
--- /dev/null
+++ b/Engine/Plugins/FX/Niagara/Source/NiagaraAnimNotifies/Private/AnimNotifyState_TimedNiagaraEffect.cpp
@@ -0,0 +1,131 @@
+// Copyright 1998-2018 Epic Games, Inc. All Rights Reserved.
+
+#include "AnimNotifyState_TimedNiagaraEffect.h"
+
+#include "Components/SkeletalMeshComponent.h"
+#include "Kismet/GameplayStatics.h"
+
+#include "NiagaraComponent.h"
+#include "NiagaraSystem.h"
+#include "NiagaraFunctionLibrary.h"
+
+UAnimNotifyState_TimedNiagaraEffect::UAnimNotifyState_TimedNiagaraEffect(const FObjectInitializer& ObjectInitializer)
+	: Super(ObjectInitializer)
+{
+	Template = nullptr;
+	LocationOffset.Set(0.0f, 0.0f, 0.0f);
+	RotationOffset = FRotator(0.0f, 0.0f, 0.0f);
+}
+
+UFXSystemComponent* UAnimNotifyState_TimedNiagaraEffect::SpawnEffect(USkeletalMeshComponent* MeshComp, UAnimSequenceBase* Animation)
+{
+	// Only spawn if we've got valid params
+	if (ValidateParameters(MeshComp))
+	{
+		UFXSystemComponent* NewComponent = UNiagaraFunctionLibrary::SpawnSystemAttached(Template, MeshComp, SocketName, LocationOffset, RotationOffset, EAttachLocation::KeepRelativeOffset, !bDestroyAtEnd);
+		return NewComponent;
+	}
+	return nullptr;
+}
+
+void UAnimNotifyState_TimedNiagaraEffect::NotifyBegin(USkeletalMeshComponent * MeshComp, class UAnimSequenceBase * Animation, float TotalDuration)
+{
+	SpawnEffect(MeshComp, Animation);
+	Received_NotifyBegin(MeshComp, Animation, TotalDuration);
+}
+
+void UAnimNotifyState_TimedNiagaraEffect::NotifyTick(USkeletalMeshComponent * MeshComp, class UAnimSequenceBase * Animation, float FrameDeltaTime)
+{
+	Received_NotifyTick(MeshComp, Animation, FrameDeltaTime);
+}
+
+void UAnimNotifyState_TimedNiagaraEffect::NotifyEnd(USkeletalMeshComponent * MeshComp, class UAnimSequenceBase * Animation)
+{
+	TArray<USceneComponent*> Children;
+	MeshComp->GetChildrenComponents(false, Children);
+
+	for (USceneComponent* Component : Children)
+	{
+		if (UFXSystemComponent* FXSystemComponent = Cast<UFXSystemComponent>(Component))
+		{
+			bool bSocketMatch = FXSystemComponent->GetAttachSocketName() == SocketName;
+			bool bTemplateMatch = FXSystemComponent->GetFXSystemAsset() == Template;
+
+#if WITH_EDITORONLY_DATA
+			// In editor someone might have changed our parameters while we're ticking; so check 
+			// previous known parameters too.
+			bSocketMatch |= PreviousSocketNames.Contains(FXSystemComponent->GetAttachSocketName());
+			bTemplateMatch |= PreviousTemplates.Contains(FXSystemComponent->GetFXSystemAsset());
+#endif
+
+			if (bSocketMatch && bTemplateMatch && FXSystemComponent->bIsActive)
+			{
+				// Either destroy the component or deactivate it to have it's active FXSystems finish.
+				// The component will auto destroy once all FXSystem are gone.
+				if (bDestroyAtEnd)
+				{
+					FXSystemComponent->DestroyComponent();
+				}
+				else
+				{
+					FXSystemComponent->Deactivate();
+				}
+
+#if WITH_EDITORONLY_DATA
+				// No longer need to track previous values as we've found our component
+				// and removed it.
+				PreviousTemplates.Empty();
+				PreviousSocketNames.Empty();
+#endif
+				// Removed a component, no need to continue
+				break;
+			}
+		}
+	}
+
+	Received_NotifyEnd(MeshComp, Animation);
+}
+
+bool UAnimNotifyState_TimedNiagaraEffect::ValidateParameters(USkeletalMeshComponent* MeshComp)
+{
+	bool bValid = true;
+
+	if (!Template)
+	{
+		bValid = false;
+	}
+	else if (!MeshComp->DoesSocketExist(SocketName) && MeshComp->GetBoneIndex(SocketName) == INDEX_NONE)
+	{
+		bValid = false;
+	}
+
+	return bValid;
+}
+
+FString UAnimNotifyState_TimedNiagaraEffect::GetNotifyName_Implementation() const
+{
+	if (Template)
+	{
+		return Template->GetName();
+	}
+
+	return UAnimNotifyState::GetNotifyName_Implementation();
+}
+
+#if WITH_EDITOR
+void UAnimNotifyState_TimedNiagaraEffect::PreEditChange(UProperty* PropertyAboutToChange)
+{
+	if (PropertyAboutToChange)
+	{
+		if (PropertyAboutToChange->GetName() == GET_MEMBER_NAME_STRING_CHECKED(UAnimNotifyState_TimedNiagaraEffect, Template) && Template != NULL)
+		{
+			PreviousTemplates.Add(Template);
+		}
+
+		if (PropertyAboutToChange->GetName() == GET_MEMBER_NAME_STRING_CHECKED(UAnimNotifyState_TimedNiagaraEffect, SocketName) && SocketName != FName(TEXT("None")))
+		{
+			PreviousSocketNames.Add(SocketName);
+		}
+	}
+}
+#endif
diff --git a/Engine/Plugins/FX/Niagara/Source/NiagaraAnimNotifies/Private/AnimNotify_PlayNiagaraEffect.cpp b/Engine/Plugins/FX/Niagara/Source/NiagaraAnimNotifies/Private/AnimNotify_PlayNiagaraEffect.cpp
new file mode 100644
--- /dev/null
+++ b/Engine/Plugins/FX/Niagara/Source/NiagaraAnimNotifies/Private/AnimNotify_PlayNiagaraEffect.cpp
@@ -0,0 +1,119 @@
+// Copyright 1998-2018 Epic Games, Inc. All Rights Reserved.
+
+#include "AnimNotify_PlayNiagaraEffect.h"
+
+#include "Components/SkeletalMeshComponent.h"
+#include "Kismet/GameplayStatics.h"
+#include "Animation/AnimSequenceBase.h"
+
+#if WITH_EDITOR
+#include "Logging/MessageLog.h"
+#include "Misc/UObjectToken.h"
+#endif
+
+#include "NiagaraComponent.h"
+#include "NiagaraSystem.h"
+#include "NiagaraFunctionLibrary.h"
+
+/////////////////////////////////////////////////////
+// UAnimNotify_PlayNiagaraEffect
+
+UAnimNotify_PlayNiagaraEffect::UAnimNotify_PlayNiagaraEffect()
+	: Super()
+{
+	Attached = true;
+	Scale = FVector(1.f);
+
+#if WITH_EDITORONLY_DATA
+	NotifyColor = FColor(192, 255, 99, 255);
+#endif // WITH_EDITORONLY_DATA
+}
+
+void UAnimNotify_PlayNiagaraEffect::PostLoad()
+{
+	Super::PostLoad();
+
+	RotationOffsetQuat = FQuat(RotationOffset);
+}
+
+#if WITH_EDITOR
+void UAnimNotify_PlayNiagaraEffect::PostEditChangeProperty(FPropertyChangedEvent& PropertyChangedEvent)
+{
+	Super::PostEditChangeProperty(PropertyChangedEvent);
+
+	if (PropertyChangedEvent.MemberProperty && PropertyChangedEvent.MemberProperty->GetFName() == GET_MEMBER_NAME_CHECKED(UAnimNotify_PlayNiagaraEffect, RotationOffset))
+	{
+		RotationOffsetQuat = FQuat(RotationOffset);
+	}
+}
+
+void UAnimNotify_PlayNiagaraEffect::ValidateAssociatedAssets()
+{
+	static const FName NAME_AssetCheck("AssetCheck");
+
+	if ((Template != nullptr) && (Template->IsLooping()))
+	{
+		UObject* ContainingAsset = GetContainingAsset();
+
+		FMessageLog AssetCheckLog(NAME_AssetCheck);
+
+		const FText MessageLooping = FText::Format(
+			NSLOCTEXT("AnimNotify", "NiagaraSystem_ShouldNotLoop", "Niagara system {0} used in anim notify for asset {1} is set to looping, but the slot is a one-shot (it won't be played to avoid leaking a component per notify)."),
+			FText::AsCultureInvariant(Template->GetPathName()),
+			FText::AsCultureInvariant(ContainingAsset->GetPathName()));
+		AssetCheckLog.Warning()
+			->AddToken(FUObjectToken::Create(ContainingAsset))
+			->AddToken(FTextToken::Create(MessageLooping));
+
+		if (GIsEditor)
+		{
+			AssetCheckLog.Notify(MessageLooping, EMessageSeverity::Warning, /*bForce=*/ true);
+		}
+	}
+}
+#endif
+
+void UAnimNotify_PlayNiagaraEffect::Notify(class USkeletalMeshComponent* MeshComp, class UAnimSequenceBase* Animation)
+{
+	// Don't call super to avoid unnecessary call in to blueprints
+	SpawnEffect(MeshComp, Animation);
+}
+
+FString UAnimNotify_PlayNiagaraEffect::GetNotifyName_Implementation() const
+{
+	if (Template)
+	{
+		return Template->GetName();
+	}
+	else
+	{
+		return Super::GetNotifyName_Implementation();
+	}
+}
+
+UFXSystemComponent* UAnimNotify_PlayNiagaraEffect::SpawnEffect(USkeletalMeshComponent* MeshComp, UAnimSequenceBase* Animation)
+{
+	UFXSystemComponent* ReturnComp = nullptr;
+
+	if (Template)
+	{
+		if (Template->IsLooping())
+		{
+			return ReturnComp;
+		}
+
+		if (Attached)
+		{
+			ReturnComp = UNiagaraFunctionLibrary::SpawnSystemAttached(Template, MeshComp, SocketName, LocationOffset, RotationOffset, EAttachLocation::KeepRelativeOffset, true);
+		}
+		else
+		{
+			const FTransform MeshTransform = MeshComp->GetSocketTransform(SocketName);
+			ReturnComp = UNiagaraFunctionLibrary::SpawnSystemAtLocation(MeshComp->GetWorld(), Template, MeshTransform.TransformPosition(LocationOffset), (MeshTransform.GetRotation() * RotationOffsetQuat).Rotator(), true);
+		}
+
+		ReturnComp->RelativeScale3D = Scale;
+	}
+
+	return ReturnComp;
+}
diff --git a/Engine/Plugins/FX/Niagara/Source/NiagaraAnimNotifies/Private/NiagaraAnimNotifiesModule.cpp b/Engine/Plugins/FX/Niagara/Source/NiagaraAnimNotifies/Private/NiagaraAnimNotifiesModule.cpp
new file mode 100644
--- /dev/null
+++ b/Engine/Plugins/FX/Niagara/Source/NiagaraAnimNotifies/Private/NiagaraAnimNotifiesModule.cpp
@@ -0,0 +1,6 @@
+// Copyright 1998-2018 Epic Games, Inc. All Rights Reserved.
+
+#include "NiagaraAnimNotifiesModule.h"
+#include "Modules/ModuleManager.h"
+
+IMPLEMENT_MODULE(INiagaraAnimNotifiesModule, NiagaraAnimNotifies);
\ No newline at end of file
diff --git a/Engine/Plugins/FX/Niagara/Source/NiagaraAnimNotifies/Public/AnimNotifyState_TimedNiagaraEffect.h b/Engine/Plugins/FX/Niagara/Source/NiagaraAnimNotifies/Public/AnimNotifyState_TimedNiagaraEffect.h
new file mode 100644
--- /dev/null
+++ b/Engine/Plugins/FX/Niagara/Source/NiagaraAnimNotifies/Public/AnimNotifyState_TimedNiagaraEffect.h
@@ -0,0 +1,74 @@
+// Copyright 1998-2018 Epic Games, Inc. All Rights Reserved.
+
+#pragma once
+
+#include "CoreMinimal.h"
+#include "UObject/ObjectMacros.h"
+#include "Animation/AnimNotifies/AnimNotifyState.h"
+#include "AnimNotifyState_TimedNiagaraEffect.generated.h"
+
+class UNiagaraSystem;
+class USkeletalMeshComponent;
+class UFXSystemAsset;
+class UFXSystemComponent;
+
+// Timed Niagara Effect Notify
+// Allows a looping Niagara effect to be played in an animation that will activate
+// at the beginning of the notify and deactivate at the end.
+UCLASS(Blueprintable, meta = (DisplayName = "Timed Niagara Effect"))
+class NIAGARAANIMNOTIFIES_API UAnimNotifyState_TimedNiagaraEffect : public UAnimNotifyState
+{
+	GENERATED_UCLASS_BODY()
+
+	// The niagara system template to use when spawning the niagara component
+	UPROPERTY(EditAnywhere, Category = NiagaraSystem, meta = (ToolTip = "The niagara system to spawn for the notify state"))
+	UNiagaraSystem* Template;
+
+	// The socket within our mesh component to attach to when we spawn the Niagara component
+	UPROPERTY(EditAnywhere, Category = NiagaraSystem, meta = (ToolTip = "The socket or bone to attach the system to", AnimNotifyBoneName = "true"))
+	FName SocketName;
+
+	// Offset from the socket / bone location
+	UPROPERTY(EditAnywhere, Category = NiagaraSystem, meta = (ToolTip = "Offset from the socket or bone to place the Niagara system"))
+	FVector LocationOffset;
+
+	// Offset from the socket / bone rotation
+	UPROPERTY(EditAnywhere, Category = NiagaraSystem, meta = (ToolTip = "Rotation offset from the socket or bone for the Niagara system"))
+	FRotator RotationOffset;
+
+	// Whether or not we destroy the component at the end of the notify or instead just stop
+	// the emitters.
+	UPROPERTY(EditAnywhere, Category = NiagaraSystem, meta = (DisplayName = "Destroy Immediately", ToolTip = "Whether the Niagara system should be immediately destroyed at the end of the notify state or be allowed to finish"))
+	bool bDestroyAtEnd;
+
+#if WITH_EDITORONLY_DATA
+	// The following arrays are used to handle property changes during a state. Because we can't
+	// store any stateful data here we can't know which emitter is ours. The best metric we have
+	// is an emitter on our Mesh Component with the same template and socket name we have defined.
+	// Because these can change at any time we need to track previous versions when we are in an
+	// editor build. Refactor when stateful data is possible, tracking our component instead.
+	UPROPERTY(transient)
+	TArray<UFXSystemAsset*> PreviousTemplates;
+
+	UPROPERTY(transient)
+	TArray<FName> PreviousSocketNames;
+
+#endif
+
+#if WITH_EDITOR
+	virtual void PreEditChange(UProperty* PropertyAboutToChange) override;
+#endif
+
+	virtual void NotifyBegin(class USkeletalMeshComponent * MeshComp, class UAnimSequenceBase * Animation, float TotalDuration) override;
+	virtual void NotifyTick(class USkeletalMeshComponent * MeshComp, class UAnimSequenceBase * Animation, float FrameDeltaTime) override;
+	virtual void NotifyEnd(class USkeletalMeshComponent * MeshComp, class UAnimSequenceBase * Animation) override;
+
+	// Overridden from UAnimNotifyState to provide custom notify name.
+	FString GetNotifyName_Implementation() const override;
+
+protected:
+	// Spawns the NiagaraSystemComponent. Called from Notify.
+	virtual UFXSystemComponent* SpawnEffect(USkeletalMeshComponent* MeshComp, UAnimSequenceBase* Animation);
+private:
+	bool ValidateParameters(USkeletalMeshComponent* MeshComp);
+};
diff --git a/Engine/Plugins/FX/Niagara/Source/NiagaraAnimNotifies/Public/AnimNotify_PlayNiagaraEffect.h b/Engine/Plugins/FX/Niagara/Source/NiagaraAnimNotifies/Public/AnimNotify_PlayNiagaraEffect.h
new file mode 100644
--- /dev/null
+++ b/Engine/Plugins/FX/Niagara/Source/NiagaraAnimNotifies/Public/AnimNotify_PlayNiagaraEffect.h
@@ -0,0 +1,77 @@
+// Copyright 1998-2018 Epic Games, Inc. All Rights Reserved.
+
+#pragma once
+
+#include "CoreMinimal.h"
+#include "UObject/ObjectMacros.h"
+#include "Animation/AnimNotifies/AnimNotify.h"
+#include "AnimNotify_PlayNiagaraEffect.generated.h"
+
+class UAnimSequenceBase;
+class UNiagaraSystem;
+class USkeletalMeshComponent;
+class UNiagaraSystemComponent;
+
+class UNiagaraSystem;
+class UFXSystemComponent;
+
+UCLASS(const, hidecategories = Object, collapsecategories, meta = (DisplayName = "Play Niagara Particle Effect"))
+class NIAGARAANIMNOTIFIES_API UAnimNotify_PlayNiagaraEffect : public UAnimNotify
+{
+	GENERATED_BODY()
+
+public:
+
+	UAnimNotify_PlayNiagaraEffect();
+
+	// Begin UObject interface
+	virtual void PostLoad() override;
+#if WITH_EDITOR
+	virtual void PostEditChangeProperty(struct FPropertyChangedEvent& PropertyChangedEvent) override;
+#endif
+	// End UObject interface
+
+	// Begin UAnimNotify interface
+	virtual FString GetNotifyName_Implementation() const override;
+	virtual void Notify(USkeletalMeshComponent* MeshComp, UAnimSequenceBase* Animation) override;
+#if WITH_EDITOR
+	virtual void ValidateAssociatedAssets() override;
+#endif
+	// End UAnimNotify interface
+
+	// Niagara System to Spawn
+	UPROPERTY(EditAnywhere, BlueprintReadWrite, Category = "AnimNotify", meta = (DisplayName = "Niagara System"))
+	UNiagaraSystem* Template;
+
+	// Location offset from the socket
+	UPROPERTY(EditAnywhere, BlueprintReadWrite, Category = "AnimNotify")
+	FVector LocationOffset;
+
+	// Rotation offset from socket
+	UPROPERTY(EditAnywhere, BlueprintReadWrite, Category = "AnimNotify")
+	FRotator RotationOffset;
+
+	// Scale to spawn the Niagara system at
+	UPROPERTY(EditAnywhere, Category = "AnimNotify")
+	FVector Scale;
+
+protected:
+	// Cached version of the Rotation Offset already in Quat form
+	FQuat RotationOffsetQuat;
+
+	// Spawns the NiagaraSystemComponent. Called from Notify.
+	virtual UFXSystemComponent* SpawnEffect(USkeletalMeshComponent* MeshComp, UAnimSequenceBase* Animation);
+
+public:
+
+	// Should attach to the bone/socket
+	UPROPERTY(EditAnywhere, BlueprintReadWrite, Category = "AnimNotify")
+	uint32 Attached : 1; 	//~ Does not follow coding standard due to redirection from BP
+
+		// SocketName to attach to
+	UPROPERTY(EditAnywhere, BlueprintReadWrite, Category = "AnimNotify", meta = (AnimNotifyBoneName = "true"))
+	FName SocketName;
+};
+
+
+
diff --git a/Engine/Plugins/FX/Niagara/Source/NiagaraAnimNotifies/Public/NiagaraAnimNotifiesModule.h b/Engine/Plugins/FX/Niagara/Source/NiagaraAnimNotifies/Public/NiagaraAnimNotifiesModule.h
new file mode 100644
--- /dev/null
+++ b/Engine/Plugins/FX/Niagara/Source/NiagaraAnimNotifies/Public/NiagaraAnimNotifiesModule.h
@@ -0,0 +1,14 @@
+// Copyright 1998-2018 Epic Games, Inc. All Rights Reserved.
+
+#pragma once
+
+#include "CoreMinimal.h"
+#include "Modules/ModuleInterface.h"
+
+/**
+* Niagara AnimNotifies module interface
+*/
+class INiagaraAnimNotifiesModule : public IModuleInterface
+{
+};
+
diff --git a/Engine/Source/Runtime/Engine/Classes/Particles/ParticleSystemComponent.h b/Engine/Source/Runtime/Engine/Classes/Particles/ParticleSystemComponent.h
--- a/Engine/Source/Runtime/Engine/Classes/Particles/ParticleSystemComponent.h
+++ b/Engine/Source/Runtime/Engine/Classes/Particles/ParticleSystemComponent.h
@@ -34,6 +34,7 @@ class UAnimNotifyState;
 class FParticleDynamicData;
 class FParticleSystemSceneProxy;
 class UAnimNotifyState;
+class UFXSystemAsset;
 struct FDynamicEmitterDataBase;
 struct FDynamicEmitterReplayDataBase;
 struct FParticleAnimTrailEmitterInstance;
@@ -354,6 +355,11 @@ public:
 	 */
 	UFUNCTION(BlueprintCallable, Category="Effects|Components|ParticleSystem")
 	virtual void SetActorParameter(FName ParameterName, class AActor* Param) {}
+
+	/** 
+	 * Get the referenced FXSystem asset.
+	*/
+	virtual UFXSystemAsset* GetFXSystemAsset() const { return nullptr; };
 };
 
 
@@ -1020,6 +1026,11 @@ public:
 	 */
 	void SetActorParameter(FName ParameterName, class AActor* Param) override;
 
+	/**
+	 * Get the referenced FXSystem asset.
+	*/
+	virtual UFXSystemAsset* GetFXSystemAsset() const { return Template; };
+
 	/** 
 	 *	Set a named material instance parameter on this ParticleSystemComponent. 
 	 *	Updates the parameter if it already exists, or creates a new entry if not. 
