Fixing copyright notice

#tests compiled

diff --git a/Engine/Plugins/FX/Niagara/Source/NiagaraAnimNotifies/NiagaraAnimNotifies.Build.cs b/Engine/Plugins/FX/Niagara/Source/NiagaraAnimNotifies/NiagaraAnimNotifies.Build.cs
--- a/Engine/Plugins/FX/Niagara/Source/NiagaraAnimNotifies/NiagaraAnimNotifies.Build.cs
+++ b/Engine/Plugins/FX/Niagara/Source/NiagaraAnimNotifies/NiagaraAnimNotifies.Build.cs
@@ -1,4 +1,4 @@
-// Copyright 1998-2018 Epic Games, Inc. All Rights Reserved.
+// Copyright 1998-2019 Epic Games, Inc. All Rights Reserved.
 
 using UnrealBuildTool;
 
diff --git a/Engine/Plugins/FX/Niagara/Source/NiagaraAnimNotifies/Private/AnimNotifyState_TimedNiagaraEffect.cpp b/Engine/Plugins/FX/Niagara/Source/NiagaraAnimNotifies/Private/AnimNotifyState_TimedNiagaraEffect.cpp
--- a/Engine/Plugins/FX/Niagara/Source/NiagaraAnimNotifies/Private/AnimNotifyState_TimedNiagaraEffect.cpp
+++ b/Engine/Plugins/FX/Niagara/Source/NiagaraAnimNotifies/Private/AnimNotifyState_TimedNiagaraEffect.cpp
@@ -1,4 +1,4 @@
-// Copyright 1998-2018 Epic Games, Inc. All Rights Reserved.
+// Copyright 1998-2019 Epic Games, Inc. All Rights Reserved.
 
 #include "AnimNotifyState_TimedNiagaraEffect.h"
 
diff --git a/Engine/Plugins/FX/Niagara/Source/NiagaraAnimNotifies/Public/AnimNotifyState_TimedNiagaraEffect.h b/Engine/Plugins/FX/Niagara/Source/NiagaraAnimNotifies/Public/AnimNotifyState_TimedNiagaraEffect.h
--- a/Engine/Plugins/FX/Niagara/Source/NiagaraAnimNotifies/Public/AnimNotifyState_TimedNiagaraEffect.h
+++ b/Engine/Plugins/FX/Niagara/Source/NiagaraAnimNotifies/Public/AnimNotifyState_TimedNiagaraEffect.h
@@ -1,4 +1,4 @@
-// Copyright 1998-2018 Epic Games, Inc. All Rights Reserved.
+// Copyright 1998-2019 Epic Games, Inc. All Rights Reserved.
 
 #pragma once
 
diff --git a/Engine/Plugins/FX/Niagara/Source/NiagaraAnimNotifies/Public/AnimNotify_PlayNiagaraEffect.h b/Engine/Plugins/FX/Niagara/Source/NiagaraAnimNotifies/Public/AnimNotify_PlayNiagaraEffect.h
--- a/Engine/Plugins/FX/Niagara/Source/NiagaraAnimNotifies/Public/AnimNotify_PlayNiagaraEffect.h
+++ b/Engine/Plugins/FX/Niagara/Source/NiagaraAnimNotifies/Public/AnimNotify_PlayNiagaraEffect.h
@@ -1,4 +1,4 @@
-// Copyright 1998-2018 Epic Games, Inc. All Rights Reserved.
+// Copyright 1998-2019 Epic Games, Inc. All Rights Reserved.
 
 #pragma once
 
diff --git a/Engine/Plugins/FX/Niagara/Source/NiagaraAnimNotifies/Public/NiagaraAnimNotifiesModule.h b/Engine/Plugins/FX/Niagara/Source/NiagaraAnimNotifies/Public/NiagaraAnimNotifiesModule.h
--- a/Engine/Plugins/FX/Niagara/Source/NiagaraAnimNotifies/Public/NiagaraAnimNotifiesModule.h
+++ b/Engine/Plugins/FX/Niagara/Source/NiagaraAnimNotifies/Public/NiagaraAnimNotifiesModule.h
@@ -1,4 +1,4 @@
-// Copyright 1998-2018 Epic Games, Inc. All Rights Reserved.
+// Copyright 1998-2019 Epic Games, Inc. All Rights Reserved.
 
 #pragma once
 
