Fixing copyright

#tests compiled

diff --git a/Engine/Plugins/FX/Niagara/Source/NiagaraAnimNotifies/Private/AnimNotify_PlayNiagaraEffect.cpp b/Engine/Plugins/FX/Niagara/Source/NiagaraAnimNotifies/Private/AnimNotify_PlayNiagaraEffect.cpp
--- a/Engine/Plugins/FX/Niagara/Source/NiagaraAnimNotifies/Private/AnimNotify_PlayNiagaraEffect.cpp
+++ b/Engine/Plugins/FX/Niagara/Source/NiagaraAnimNotifies/Private/AnimNotify_PlayNiagaraEffect.cpp
@@ -1,4 +1,4 @@
-// Copyright 1998-2018 Epic Games, Inc. All Rights Reserved.
+// Copyright 1998-2019 Epic Games, Inc. All Rights Reserved.
 
 #include "AnimNotify_PlayNiagaraEffect.h"
 
diff --git a/Engine/Plugins/FX/Niagara/Source/NiagaraAnimNotifies/Private/NiagaraAnimNotifiesModule.cpp b/Engine/Plugins/FX/Niagara/Source/NiagaraAnimNotifies/Private/NiagaraAnimNotifiesModule.cpp
--- a/Engine/Plugins/FX/Niagara/Source/NiagaraAnimNotifies/Private/NiagaraAnimNotifiesModule.cpp
+++ b/Engine/Plugins/FX/Niagara/Source/NiagaraAnimNotifies/Private/NiagaraAnimNotifiesModule.cpp
@@ -1,4 +1,4 @@
-// Copyright 1998-2018 Epic Games, Inc. All Rights Reserved.
+// Copyright 1998-2019 Epic Games, Inc. All Rights Reserved.
 
 #include "NiagaraAnimNotifiesModule.h"
 #include "Modules/ModuleManager.h"
