Fix engine vars not working in system spawn/update scripts

diff --git a/Engine/Plugins/FX/Niagara/Source/Niagara/Private/NiagaraSystemInstance.cpp b/Engine/Plugins/FX/Niagara/Source/Niagara/Private/NiagaraSystemInstance.cpp
--- a/Engine/Plugins/FX/Niagara/Source/Niagara/Private/NiagaraSystemInstance.cpp
+++ b/Engine/Plugins/FX/Niagara/Source/Niagara/Private/NiagaraSystemInstance.cpp
@@ -864,6 +864,8 @@ void FNiagaraSystemInstance::BindParameters()
 	}
 
 	Component->GetOverrideParameters().Bind(&InstanceParameters);
+	InstanceParameters.Bind(&SystemSimulation->GetSpawnExecutionContext().Parameters);
+	InstanceParameters.Bind(&SystemSimulation->GetUpdateExecutionContext().Parameters);
 
 	if (SystemSimulation->GetIsSolo())
 	{
@@ -896,6 +898,8 @@ void FNiagaraSystemInstance::UnbindParameters()
 
 	if (SystemSimulation.IsValid())
 	{
+		InstanceParameters.Unbind(&SystemSimulation->GetSpawnExecutionContext().Parameters);
+		InstanceParameters.Unbind(&SystemSimulation->GetUpdateExecutionContext().Parameters);
 		if (SystemSimulation->GetIsSolo())
 		{
 			if (Component)
