- Niagara fix check to handle running  in -OneThread mode

diff --git a/Engine/Plugins/FX/Niagara/Source/Niagara/Private/NiagaraDataSet.cpp b/Engine/Plugins/FX/Niagara/Source/Niagara/Private/NiagaraDataSet.cpp
--- a/Engine/Plugins/FX/Niagara/Source/Niagara/Private/NiagaraDataSet.cpp
+++ b/Engine/Plugins/FX/Niagara/Source/Niagara/Private/NiagaraDataSet.cpp
@@ -466,7 +466,7 @@ void FNiagaraDataBuffer::CheckUsage(bool bReadOnly)const
 	if (Owner->SimTarget == ENiagaraSimTarget::CPUSim)
 	{
 		//We can read on the RT but any modifications must be GT (or GT Task).
-		check(bReadOnly || !IsInRenderingThread());
+		check(IsInGameThread() || (bReadOnly || !IsInRenderingThread()));
 	}
 	else
 	{
