- Niagara only create skel mesh GPU data if we have a GPU emitter in the system

#rnx

diff --git a/Engine/Plugins/FX/Niagara/Source/Niagara/Private/NiagaraDataInterfaceSkeletalMesh.cpp b/Engine/Plugins/FX/Niagara/Source/Niagara/Private/NiagaraDataInterfaceSkeletalMesh.cpp
--- a/Engine/Plugins/FX/Niagara/Source/Niagara/Private/NiagaraDataInterfaceSkeletalMesh.cpp
+++ b/Engine/Plugins/FX/Niagara/Source/Niagara/Private/NiagaraDataInterfaceSkeletalMesh.cpp
@@ -333,7 +333,10 @@ void FSkeletalMeshGpuSpawnStaticBuffers::InitRHI()
 
 	const FMultiSizeIndexContainer& IndexBuffer = LODRenderData->MultiSizeIndexContainer;
 	MeshIndexBufferSrv = IndexBuffer.GetIndexBuffer()->GetSRV();
-	//ensure(MeshIndexBufferSrv != nullptr);
+	if (!MeshIndexBufferSrv.IsValid())
+	{
+		UE_LOG(LogNiagara, Warning, TEXT("Skeletal Mesh does not have an SRV for the index buffer, if you are using triangle sampling it will not work."));
+	}
 
 	MeshVertexBufferSrv = LODRenderData->StaticVertexBuffers.PositionVertexBuffer.GetSRV();
 
@@ -972,29 +975,27 @@ USkeletalMesh* UNiagaraDataInterfaceSkeletalMesh::GetSkeletalMeshHelper(UNiagara
 bool FNDISkeletalMesh_InstanceData::Init(UNiagaraDataInterfaceSkeletalMesh* Interface, FNiagaraSystemInstance* SystemInstance)
 {
 	check(SystemInstance);
-	ChangeId = Interface->ChangeId;
-	USkeletalMesh* PrevMesh = Mesh;
+
+	// Initialize members
 	Component = nullptr;
 	Mesh = nullptr;
+	MeshSafe = nullptr;
 	Transform = FMatrix::Identity;
 	TransformInverseTransposed = FMatrix::Identity;
 	PrevTransform = FMatrix::Identity;
 	PrevTransformInverseTransposed = FMatrix::Identity;
 	DeltaSeconds = 0.0f;
+	ChangeId = Interface->ChangeId;
+	bIsGpuUniformlyDistributedSampling = false;
+	MeshWeightStrideByte = 0;
+	MeshGpuSpawnStaticBuffers = nullptr;
+	MeshGpuSpawnDynamicBuffers = nullptr;
 
+	// Get skel mesh and confirm have valid data
 	USkeletalMeshComponent* NewSkelComp = nullptr;
 	Mesh = UNiagaraDataInterfaceSkeletalMesh::GetSkeletalMeshHelper(Interface, SystemInstance->GetComponent(), Component, NewSkelComp);
-	
 	MeshSafe = Mesh;
 
-	if (Component.IsValid() && Mesh)
-	{
-		PrevTransform = Transform;
-		PrevTransformInverseTransposed = TransformInverseTransposed;
-		Transform = Component->GetComponentToWorld().ToMatrixWithScale();
-		TransformInverseTransposed = Transform.InverseFast().GetTransposed();
-	}
-
 	if (!Mesh)
 	{
 		/*USceneComponent* Comp = Component.Get();
@@ -1005,23 +1006,27 @@ bool FNDISkeletalMesh_InstanceData::Init(UNiagaraDataInterfaceSkeletalMesh* Inte
 		return false;
 	}
 
+	if (!Component.IsValid())
+	{
+		UE_LOG(LogNiagara, Log, TEXT("SkeletalMesh data interface has no valid component. Failed InitPerInstanceData - %s"), *Interface->GetFullName());
+		return false;
+	}
+
+	Transform = Component->GetComponentToWorld().ToMatrixWithScale();
+	TransformInverseTransposed = Transform.InverseFast().GetTransposed();
+	PrevTransform = Transform;
+	PrevTransformInverseTransposed = TransformInverseTransposed;
+
 #if WITH_EDITOR
 	MeshSafe->GetOnMeshChanged().AddUObject(SystemInstance->GetComponent(), &UNiagaraComponent::ReinitializeSystem);
 #endif
 
-
 // 	if (!Mesh->bAllowCPUAccess)
 // 	{
 // 		UE_LOG(LogNiagara, Log, TEXT("SkeletalMesh data interface using a mesh that does not allow CPU access. Failed InitPerInstanceData - Mesh: %s"), *Mesh->GetFullName());
 // 		return false;
 // 	}
 
-	if (!Component.IsValid())
-	{
-		UE_LOG(LogNiagara, Log, TEXT("SkeletalMesh data interface has no valid component. Failed InitPerInstanceData - %s"), *Interface->GetFullName());
-		return false;
-	}
-
 	//Setup where to spawn from
 	SamplingRegionIndices.Empty();
 	bool bAllRegionsAreAreaWeighting = true;
@@ -1246,6 +1251,8 @@ bool FNDISkeletalMesh_InstanceData::Init(UNiagaraDataInterfaceSkeletalMesh* Inte
 		}
 	}
 
+	//-TODO: We should find out if this DI is connected to a GPU emitter or not rather than a blanket accross the system
+	if ( SystemInstance->HasGPUEmitters() )
 	{
 		MeshWeightStrideByte = SkinWeightBuffer->GetStride();
 		MeshSkinWeightBufferSrv = SkinWeightBuffer->GetSRV();
diff --git a/Engine/Plugins/FX/Niagara/Source/Niagara/Private/NiagaraSystemInstance.cpp b/Engine/Plugins/FX/Niagara/Source/Niagara/Private/NiagaraSystemInstance.cpp
--- a/Engine/Plugins/FX/Niagara/Source/Niagara/Private/NiagaraSystemInstance.cpp
+++ b/Engine/Plugins/FX/Niagara/Source/Niagara/Private/NiagaraSystemInstance.cpp
@@ -1297,6 +1297,8 @@ void FNiagaraSystemInstance::InitEmitters()
 		Component->MarkRenderStateDirty();
 	}
 
+	bHasGPUEmitters = false;
+
 	Emitters.Empty();
 	UNiagaraSystem* System = GetSystem();
 	if (System != nullptr)
diff --git a/Engine/Plugins/FX/Niagara/Source/Niagara/Public/NiagaraSystemInstance.h b/Engine/Plugins/FX/Niagara/Source/Niagara/Public/NiagaraSystemInstance.h
--- a/Engine/Plugins/FX/Niagara/Source/Niagara/Public/NiagaraSystemInstance.h
+++ b/Engine/Plugins/FX/Niagara/Source/Niagara/Public/NiagaraSystemInstance.h
@@ -211,8 +211,8 @@ public:
 
 	static bool AllocateSystemInstance(class UNiagaraComponent* InComponent, TUniquePtr< FNiagaraSystemInstance >& OutSystemInstanceAllocation);
 	static bool DeallocateSystemInstance(TUniquePtr< FNiagaraSystemInstance >& SystemInstanceAllocation);
-	/*void SetHasGPUEmitters(bool bInHasGPUEmitters) { bHasGPUEmitters = bInHasGPUEmitters; }
-	bool HasGPUEmitters() { return bHasGPUEmitters;  }*/
+	/*void SetHasGPUEmitters(bool bInHasGPUEmitters) { bHasGPUEmitters = bInHasGPUEmitters; }*/
+	bool HasGPUEmitters() { return bHasGPUEmitters;  }
 
 	int32 GetDetailLevel()const;
 private:
