Static switch bug fixes and small ui improvements

diff --git a/Engine/Plugins/FX/Niagara/Source/NiagaraEditor/Private/Customizations/NiagaraStaticSwitchNodeDetails.cpp b/Engine/Plugins/FX/Niagara/Source/NiagaraEditor/Private/Customizations/NiagaraStaticSwitchNodeDetails.cpp
--- a/Engine/Plugins/FX/Niagara/Source/NiagaraEditor/Private/Customizations/NiagaraStaticSwitchNodeDetails.cpp
+++ b/Engine/Plugins/FX/Niagara/Source/NiagaraEditor/Private/Customizations/NiagaraStaticSwitchNodeDetails.cpp
@@ -52,7 +52,7 @@ void FNiagaraStaticSwitchNodeDetails::CustomizeDetails(IDetailLayoutBuilder& Det
 		FDetailWidgetRow& NameWidget = CategoryBuilder.AddCustomRow(LOCTEXT("NiagaraSwitchNodeNameFilterText", "Input parameter name"));
 		FDetailWidgetRow& DropdownWidget = CategoryBuilder.AddCustomRow(LOCTEXT("NiagaraSwitchNodeTypeFilterText", "Input parameter type"));
 		FDetailWidgetRow& IntValueOption = CategoryBuilder.AddCustomRow(LOCTEXT("NiagaraSwitchNodeIntFilterText", "Max integer value"));
-		FDetailWidgetRow& DefaultValueOption = CategoryBuilder.AddCustomRow(LOCTEXT("NiagaraSwitchNodeDefaultFilterText", "Default value"), true);
+		FDetailWidgetRow& DefaultValueOption = CategoryBuilder.AddCustomRow(LOCTEXT("NiagaraSwitchNodeDefaultFilterText", "Default value"));
 
 		NameWidget
 		.NameContent()
diff --git a/Engine/Plugins/FX/Niagara/Source/NiagaraEditor/Private/NiagaraNodeStaticSwitch.cpp b/Engine/Plugins/FX/Niagara/Source/NiagaraEditor/Private/NiagaraNodeStaticSwitch.cpp
--- a/Engine/Plugins/FX/Niagara/Source/NiagaraEditor/Private/NiagaraNodeStaticSwitch.cpp
+++ b/Engine/Plugins/FX/Niagara/Source/NiagaraEditor/Private/NiagaraNodeStaticSwitch.cpp
@@ -55,6 +55,9 @@ void UNiagaraNodeStaticSwitch::RemoveUnusedGraphParameter(const FNiagaraVariable
 	{
 		GetNiagaraGraph()->NotifyGraphChanged();
 	}
+
+	// force the graph to refresh the metadata
+	GetNiagaraGraph()->GetParameterReferenceMap();
 }
 
 void UNiagaraNodeStaticSwitch::AllocateDefaultPins()
@@ -111,6 +114,9 @@ void UNiagaraNodeStaticSwitch::AllocateDefaultPins()
 	}
 
 	CreateAddPin(EGPD_Output);
+
+	// force the graph to refresh the metadata
+	GetNiagaraGraph()->GetParameterReferenceMap();
 }
 
 void UNiagaraNodeStaticSwitch::InsertInputPinsFor(const FNiagaraVariable& Var)
diff --git a/Engine/Plugins/FX/Niagara/Source/NiagaraEditor/Private/ViewModels/Stack/NiagaraStackFunctionInputCollection.cpp b/Engine/Plugins/FX/Niagara/Source/NiagaraEditor/Private/ViewModels/Stack/NiagaraStackFunctionInputCollection.cpp
--- a/Engine/Plugins/FX/Niagara/Source/NiagaraEditor/Private/ViewModels/Stack/NiagaraStackFunctionInputCollection.cpp
+++ b/Engine/Plugins/FX/Niagara/Source/NiagaraEditor/Private/ViewModels/Stack/NiagaraStackFunctionInputCollection.cpp
@@ -227,9 +227,8 @@ void UNiagaraStackFunctionInputCollection::RefreshChildrenInternal(const TArray<
 	RefreshIssues(DuplicateInputNames, ValidAliasedInputNames, PinsWithInvalidTypes, StaticSwitchInputs, NewIssues);
 }
 
-UNiagaraStackEntry::FStackIssueFix UNiagaraStackFunctionInputCollection::GetNodeRemovalFix(UEdGraphPin* PinToRemove)
+UNiagaraStackEntry::FStackIssueFix UNiagaraStackFunctionInputCollection::GetNodeRemovalFix(UEdGraphPin* PinToRemove, FText FixDescription)
 {
-	FText FixDescription = LOCTEXT("RemoveInvalidInputTransaction", "Remove input override.");
 	return FStackIssueFix(
 		FixDescription,
 		UNiagaraStackEntry::FStackIssueFixDelegate::CreateLambda([=]()
@@ -283,7 +282,7 @@ void UNiagaraStackFunctionInputCollection::RefreshIssues(TArray<FName> Duplicate
 						FText::FromString(OverridePin->PinName.ToString()), FText::FromString(InputFunctionCallNode->GetFunctionName())),
 					GetStackEditorDataKey(),
 					false,
-					GetNodeRemovalFix(OverridePin));
+					GetNodeRemovalFix(OverridePin, LOCTEXT("RemoveInvalidInputTransaction", "Remove input override")));
 
 				NewIssues.Add(InvalidInputOverrideError);
 			}
@@ -300,7 +299,7 @@ void UNiagaraStackFunctionInputCollection::RefreshIssues(TArray<FName> Duplicate
 				TArray<FStackIssueFix> Fixes;
 
 				// first possible fix: convert the value over to the static switch
-				FText ConversionFixDescription = LOCTEXT("ConvertInputToStaticSwitchTransaction", "Copy value to static switch parameter.");
+				FText ConversionFixDescription = LOCTEXT("ConvertInputToStaticSwitchTransaction", "Copy value to static switch parameter");
 				FStackIssueFix ConvertInputOverrideFix(
 					ConversionFixDescription,
 					UNiagaraStackEntry::FStackIssueFixDelegate::CreateLambda([=]()
@@ -323,7 +322,7 @@ void UNiagaraStackFunctionInputCollection::RefreshIssues(TArray<FName> Duplicate
 				Fixes.Add(ConvertInputOverrideFix);
 
 				// second possible fix: remove the override completely
-				Fixes.Add(GetNodeRemovalFix(OverridePin));
+				Fixes.Add(GetNodeRemovalFix(OverridePin, LOCTEXT("RemoveInvalidInputTransaction", "Remove input override (WARNING: this could result in different behavior!)")));
 
 				FStackIssue DeprecatedInputOverrideError(
 					EStackIssueSeverity::Error,
diff --git a/Engine/Plugins/FX/Niagara/Source/NiagaraEditor/Public/ViewModels/Stack/NiagaraStackFunctionInputCollection.h b/Engine/Plugins/FX/Niagara/Source/NiagaraEditor/Public/ViewModels/Stack/NiagaraStackFunctionInputCollection.h
--- a/Engine/Plugins/FX/Niagara/Source/NiagaraEditor/Public/ViewModels/Stack/NiagaraStackFunctionInputCollection.h
+++ b/Engine/Plugins/FX/Niagara/Source/NiagaraEditor/Public/ViewModels/Stack/NiagaraStackFunctionInputCollection.h
@@ -45,7 +45,7 @@ private:
 
 	void OnFunctionInputsChanged();
 
-	UNiagaraStackEntry::FStackIssueFix GetNodeRemovalFix(UEdGraphPin* PinToRemove);
+	UNiagaraStackEntry::FStackIssueFix GetNodeRemovalFix(UEdGraphPin* PinToRemove, FText FixDescription);
 
 	struct FInputData
 	{
