More static switch bug fixes:
* prevent infinite loop with recursive function calls in a graph
* Fix default enum values display text
* Fix crash on function node initialization when metadata is missing

diff --git a/Engine/Plugins/FX/Niagara/Source/NiagaraEditor/Private/Customizations/NiagaraStaticSwitchNodeDetails.cpp b/Engine/Plugins/FX/Niagara/Source/NiagaraEditor/Private/Customizations/NiagaraStaticSwitchNodeDetails.cpp
--- a/Engine/Plugins/FX/Niagara/Source/NiagaraEditor/Private/Customizations/NiagaraStaticSwitchNodeDetails.cpp
+++ b/Engine/Plugins/FX/Niagara/Source/NiagaraEditor/Private/Customizations/NiagaraStaticSwitchNodeDetails.cpp
@@ -176,7 +176,7 @@ void FNiagaraStaticSwitchNodeDetails::CustomizeDetails(IDetailLayoutBuilder& Det
 				
 				+ SWidgetSwitcher::Slot()
 		  		[
-					SNew(SComboBox<TSharedPtr<FString>>)
+					SNew(SComboBox<TSharedPtr<DefaultEnumOption>>)
 					.OptionsSource(&DefaultEnumDropdownOptions)
 					.OnSelectionChanged(this, &FNiagaraStaticSwitchNodeDetails::OnSelectionChanged)
 					.OnGenerateWidget(this, &FNiagaraStaticSwitchNodeDetails::CreateWidgetForDropdownOption)
@@ -235,12 +235,12 @@ TSharedRef<SWidget> FNiagaraStaticSwitchNodeDetails::CreateWidgetForDropdownOpti
 	return SNew(STextBlock).Text(FText::FromString(*InOption->Name));
 }
 
-TSharedRef<SWidget> FNiagaraStaticSwitchNodeDetails::CreateWidgetForDropdownOption(TSharedPtr<FString> InOption)
+TSharedRef<SWidget> FNiagaraStaticSwitchNodeDetails::CreateWidgetForDropdownOption(TSharedPtr<DefaultEnumOption> InOption)
 {
-	return SNew(STextBlock).Text(FText::FromString(*InOption));
+	return SNew(STextBlock).Text(InOption->DisplayName);
 }
 
-void FNiagaraStaticSwitchNodeDetails::OnSelectionChanged(TSharedPtr<FString> NewValue, ESelectInfo::Type)
+void FNiagaraStaticSwitchNodeDetails::OnSelectionChanged(TSharedPtr<DefaultEnumOption> NewValue, ESelectInfo::Type)
 {
 	SelectedDefaultValue = NewValue;
 	TOptional<FNiagaraVariableMetaData> MetaData = GetSwitchParameterMetadata();
@@ -255,8 +255,7 @@ void FNiagaraStaticSwitchNodeDetails::OnSelectionChanged(TSharedPtr<FString> New
 		return;
 	}
 
-	int32 Index = Enum->GetIndexByNameString(*SelectedDefaultValue);
-	MetaData->StaticSwitchDefaultValue = Index;
+	MetaData->StaticSwitchDefaultValue = SelectedDefaultValue->EnumIndex;
 	SetSwitchParameterMetadata(MetaData.GetValue());
 }
 
@@ -302,7 +301,7 @@ FText FNiagaraStaticSwitchNodeDetails::GetDefaultSelectionItemLabel() const
 {
 	if (SelectedDefaultValue.IsValid())
 	{
-		return FText::FromString(*SelectedDefaultValue);
+		return SelectedDefaultValue->DisplayName;
 	}
 
 	return LOCTEXT("InvalidNiagaraStaticSwitchNodeComboEntryText", "<Invalid selection>");
@@ -323,7 +322,8 @@ void FNiagaraStaticSwitchNodeDetails::RefreshDefaultDropdownValues()
 		SelectedDefaultValue.Reset();
 		for (int i = 0; i < Enum->GetMaxEnumValue() - 1; i++)
 		{
-			DefaultEnumDropdownOptions.Add(MakeShareable(new FString(Enum->GetNameStringByIndex(i))));
+			FText DisplayName = Enum->GetDisplayNameTextByIndex(i);
+			DefaultEnumDropdownOptions.Add(MakeShared<DefaultEnumOption>(DisplayName, i));
 
 			if (MetaData.IsSet() && i == MetaData->StaticSwitchDefaultValue)
 			{
diff --git a/Engine/Plugins/FX/Niagara/Source/NiagaraEditor/Private/Customizations/NiagaraStaticSwitchNodeDetails.h b/Engine/Plugins/FX/Niagara/Source/NiagaraEditor/Private/Customizations/NiagaraStaticSwitchNodeDetails.h
--- a/Engine/Plugins/FX/Niagara/Source/NiagaraEditor/Private/Customizations/NiagaraStaticSwitchNodeDetails.h
+++ b/Engine/Plugins/FX/Niagara/Source/NiagaraEditor/Private/Customizations/NiagaraStaticSwitchNodeDetails.h
@@ -21,6 +21,21 @@ struct SwitchDropdownOption
 	{}
 };
 
+// This data structure is used internally by the default enum dropdown to keep track of the user's choice
+struct DefaultEnumOption
+{
+	FText DisplayName;
+	int32 EnumIndex;	
+
+	DefaultEnumOption(FText DisplayName) : DisplayName(DisplayName), EnumIndex(0)
+	{
+	}
+
+	DefaultEnumOption(FText DisplayName, int32 EnumIndex) : DisplayName(DisplayName), EnumIndex(EnumIndex)
+	{
+	}
+};
+
 /** This customization sets up a custom details panel for the static switch node in the niagara module graph. */
 class FNiagaraStaticSwitchNodeDetails : public IDetailCustomization
 {
@@ -55,8 +70,8 @@ private:
 	TOptional<int32> GetSwitchDefaultValue() const;
 	void DefaultIntValueCommitted(int32 Value, ETextCommit::Type CommitInfo);
 	void DefaultBoolValueCommitted(ECheckBoxState NewState);
-	TSharedRef<SWidget> CreateWidgetForDropdownOption(TSharedPtr<FString> InOption);
-	void OnSelectionChanged(TSharedPtr<FString> NewValue, ESelectInfo::Type);
+	TSharedRef<SWidget> CreateWidgetForDropdownOption(TSharedPtr<DefaultEnumOption> InOption);
+	void OnSelectionChanged(TSharedPtr<DefaultEnumOption> NewValue, ESelectInfo::Type);
 	FText GetDefaultSelectionItemLabel() const;
 	void RefreshDefaultDropdownValues();
 
@@ -66,6 +81,6 @@ private:
 	TWeakObjectPtr<class UNiagaraNodeStaticSwitch> Node;
 	TArray<TSharedPtr<SwitchDropdownOption>> DropdownOptions;
 	TSharedPtr<SwitchDropdownOption> SelectedDropdownItem;
-	TArray<TSharedPtr<FString>> DefaultEnumDropdownOptions;
-	TSharedPtr<FString> SelectedDefaultValue;
+	TArray<TSharedPtr<DefaultEnumOption>> DefaultEnumDropdownOptions;
+	TSharedPtr<DefaultEnumOption> SelectedDefaultValue;
 };
diff --git a/Engine/Plugins/FX/Niagara/Source/NiagaraEditor/Private/NiagaraEditorUtilities.cpp b/Engine/Plugins/FX/Niagara/Source/NiagaraEditor/Private/NiagaraEditorUtilities.cpp
--- a/Engine/Plugins/FX/Niagara/Source/NiagaraEditor/Private/NiagaraEditorUtilities.cpp
+++ b/Engine/Plugins/FX/Niagara/Source/NiagaraEditor/Private/NiagaraEditorUtilities.cpp
@@ -879,7 +879,10 @@ const TMap<FNiagaraVariable, FNiagaraGraphParameterReferenceCollection>& FNiagar
 			// We append all function nodes we find to the end of the queue, basically traversing the graph breadth-first
 			TArray<UNiagaraNodeFunctionCall*> SubFunctionNodes;
 			FunctionGraph->GetNodesOfClass(SubFunctionNodes);
-			Nodes.Append(SubFunctionNodes);
+			for (UNiagaraNodeFunctionCall* SubNode : SubFunctionNodes)
+			{
+				Nodes.AddUnique(SubNode);
+			}
 
 			// We set the static switch values based on the function call inputs to prepare the parameter map history generation
 			TArray<UEdGraphPin*> CallInputs;
diff --git a/Engine/Plugins/FX/Niagara/Source/NiagaraEditor/Private/NiagaraNodeFunctionCall.cpp b/Engine/Plugins/FX/Niagara/Source/NiagaraEditor/Private/NiagaraNodeFunctionCall.cpp
--- a/Engine/Plugins/FX/Niagara/Source/NiagaraEditor/Private/NiagaraNodeFunctionCall.cpp
+++ b/Engine/Plugins/FX/Niagara/Source/NiagaraEditor/Private/NiagaraNodeFunctionCall.cpp
@@ -161,12 +161,25 @@ void UNiagaraNodeFunctionCall::AllocateDefaultPins()
 			NewPin->bNotConnectable = true;
 			NewPin->bDefaultValueIsIgnored = FindPropagatedVariable(Input) != nullptr;
 
-			int32 DefaultValue = Graph->GetMetaData(Input).Get(FNiagaraVariableMetaData()).StaticSwitchDefaultValue;
-			Input.SetValue<FNiagaraInt32>({ DefaultValue });
 			FString PinDefaultValue;
-			if (Schema->TryGetPinDefaultValueFromNiagaraVariable(Input, PinDefaultValue))
+			TOptional<FNiagaraVariableMetaData> MetaData = Graph->GetMetaData(Input);
+			if (MetaData.IsSet())
 			{
-				NewPin->DefaultValue = PinDefaultValue;
+				int32 DefaultValue = MetaData->StaticSwitchDefaultValue;
+				Input.AllocateData();
+				Input.SetValue<FNiagaraInt32>({ DefaultValue });
+				
+				if (Schema->TryGetPinDefaultValueFromNiagaraVariable(Input, PinDefaultValue))
+				{
+					NewPin->DefaultValue = PinDefaultValue;
+				}
+			}
+			else
+			{
+				if (Schema->TryGetPinDefaultValueFromNiagaraVariable(Input, PinDefaultValue))
+				{
+					NewPin->DefaultValue = PinDefaultValue;
+				}
 			}
 		}
 
