Niagara - Enable chained emitter inheritance.

diff --git a/Engine/Plugins/FX/Niagara/Source/Niagara/Classes/NiagaraEmitter.h b/Engine/Plugins/FX/Niagara/Source/Niagara/Classes/NiagaraEmitter.h
--- a/Engine/Plugins/FX/Niagara/Source/Niagara/Classes/NiagaraEmitter.h
+++ b/Engine/Plugins/FX/Niagara/Source/Niagara/Classes/NiagaraEmitter.h
@@ -8,6 +8,7 @@
 #include "UObject/Object.h"
 #include "NiagaraScript.h"
 #include "NiagaraCollision.h"
+#include "INiagaraMergeManager.h"
 #include "NiagaraEmitter.generated.h"
 
 class UMaterial;
@@ -172,6 +173,8 @@ class UNiagaraEmitter : public UObject
 {
 	GENERATED_UCLASS_BODY()
 
+	friend struct FNiagaraEmitterHandle;
+
 public:
 #if WITH_EDITOR
 	DECLARE_MULTICAST_DELEGATE(FOnPropertiesChanged);
@@ -184,10 +187,15 @@ public:
 #endif
 
 public:
-	//Begin UObject Interface
 #if WITH_EDITOR
-	virtual void PostEditChangeProperty(struct FPropertyChangedEvent& PropertyChangedEvent) override;
+	/** Creates a new emitter with the supplied emitter as a parent emitter and the supplied system as it's owner. */
+	NIAGARA_API static UNiagaraEmitter* CreateWithParentAndOwner(UNiagaraEmitter& InParentEmitter, UObject* InOwner, FName InName, EObjectFlags FlagMask);
 
+	/** Creates a new emitter by duplicating an existing emitter.  The new emitter  will reference the same parent emitter if one is available. */
+	static UNiagaraEmitter* CreateAsDuplicate(const UNiagaraEmitter& InEmitterToDuplicate, FName InDuplicateName, UNiagaraSystem& InDuplicateOwnerSystem);
+
+	//Begin UObject Interface
+	virtual void PostEditChangeProperty(struct FPropertyChangedEvent& PropertyChangedEvent) override;
 	NIAGARA_API FOnPropertiesChanged& OnPropertiesChanged();
 #endif
 	void Serialize(FArchive& Ar)override;
@@ -207,15 +215,6 @@ public:
 	UPROPERTY(EditAnywhere, Category = "Emitter", meta = (EditCondition = "bDeterminism"))
 	int32 RandomSeed;
 
-	//UPROPERTY(EditAnywhere, Category = "Emitter")
-	//float StartTime;
-	//UPROPERTY(EditAnywhere, Category = "Emitter")
-	//float EndTime;
-	//UPROPERTY(EditAnywhere, Category = "Emitter")
-	//int32 NumLoops
-	//UPROPERTY(EditAnywhere, Category = "Emitter")
-	//ENiagaraCollisionMode CollisionMode;
-
 	UPROPERTY()
 	FNiagaraEmitterScriptProperties UpdateScriptProps;
 
@@ -314,6 +313,15 @@ public:
 	NIAGARA_API FOnEmitterCompiled& OnEmitterVMCompiled();
 
 	NIAGARA_API static bool GetForceCompileOnLoad();
+
+	/** Whether or not this emitter is synchronized with its parent emitter. */
+	NIAGARA_API bool IsSynchronizedWithParent() const;
+
+	/** Merges in any changes from the parent emitter into this emitter. */
+	NIAGARA_API INiagaraMergeManager::FMergeEmitterResults MergeChangesFromParent();
+
+	/** Whether or not this emitter uses the supplied emitter */
+	bool UsesEmitter(const UNiagaraEmitter& InEmitter) const;
 #endif
 
 	/** Is this emitter allowed to be enabled by the current system detail level. */
@@ -353,11 +361,17 @@ public:
 
 	TStatId GetStatID(bool bGameThread, bool bConcurrent)const;
 
+	NIAGARA_API UNiagaraEmitter* GetParent() const;
+
+	NIAGARA_API void RemoveParent();
+
 protected:
 	virtual void BeginDestroy() override;
 
 #if WITH_EDITORONLY_DATA
 private:
+	void UpdateFromMergedCopy(const INiagaraMergeManager& MergeManager, UNiagaraEmitter* MergedEmitter);
+
 	void SyncEmitterAlias(const FString& InOldName, const FString& InNewName);
 
 	void UpdateChangeId();
@@ -392,6 +406,12 @@ private:
 	UPROPERTY()
 	TArray<FName> SharedEventGeneratorIds;
 
+	UPROPERTY()
+	UNiagaraEmitter* Parent;
+
+	UPROPERTY()
+	UNiagaraEmitter* ParentAtLastMerge;
+
 #if WITH_EDITOR
 	FOnPropertiesChanged OnPropertiesChangedDelegate;
 #endif
diff --git a/Engine/Plugins/FX/Niagara/Source/Niagara/Classes/NiagaraEmitterHandle.h b/Engine/Plugins/FX/Niagara/Source/Niagara/Classes/NiagaraEmitterHandle.h
--- a/Engine/Plugins/FX/Niagara/Source/Niagara/Classes/NiagaraEmitterHandle.h
+++ b/Engine/Plugins/FX/Niagara/Source/Niagara/Classes/NiagaraEmitterHandle.h
@@ -13,8 +13,7 @@ class UNiagaraSystem;
 class UNiagaraEmitter;
 
 /** 
- * Stores a references to a source emitter asset, and a copy of that emitter for editing within an System.  Also
- * stores whether or not this emitter is enabled, and it's name within the editor.
+ * Stores emitter information within the context of a System.
  */
 USTRUCT()
 struct NIAGARA_API FNiagaraEmitterHandle
@@ -25,13 +24,10 @@ public:
 	FNiagaraEmitterHandle();
 
 #if WITH_EDITORONLY_DATA
-	/** Creates a new emitter handle from an emitter and an owning System. */
-	FNiagaraEmitterHandle(UNiagaraEmitter& InSourceEmitter, FName InName, UNiagaraSystem& InOuterSystem);
-
-	/** Creates a new emitter handle by duplicating an existing handle.  The new emitter handle will reference the same source emitter
-		but will have it's own copy of the emitter made from the one in the supplied handle and will have it's own Id. */
-	FNiagaraEmitterHandle(const FNiagaraEmitterHandle& InHandleToDuplicate, FName InDuplicateName, UNiagaraSystem& InDuplicateOwnerSystem);
+	/** Creates a new emitter handle from an emitter. */
+	FNiagaraEmitterHandle(UNiagaraEmitter& InEmitter);
 #endif
+
 	/** Whether or not this is a valid emitter handle. */
 	bool IsValid() const;
 
@@ -55,9 +51,6 @@ public:
 	void SetIsEnabled(bool bInIsEnabled);
 
 #if WITH_EDITORONLY_DATA
-	/** Gets the source emitter this emitter handle was built from. */
-	const UNiagaraEmitter* GetSource() const;
-
 	bool IsIsolated() const {	return bIsolated; }
 	void SetIsolated(bool bInIsolated) { bIsolated = bInIsolated; }
 #endif
@@ -69,26 +62,14 @@ public:
 	FString GetUniqueInstanceName()const;
 
 #if WITH_EDITORONLY_DATA
-	/** Determine whether or not the Source and Instance refer to the same Emitter ChangeId.*/
-	bool IsSynchronizedWithSource() const;
-
 	/** Determine whether or not the Instance script is in synch with its graph.*/
 	bool NeedsRecompile() const;
 
 	/** Calls conditional post load on all sub-objects this handle references. */
-	void ConditionalPostLoad();
-
-	/** Merges in any changes from the source emitter into the instanced emitter. */
-	INiagaraModule::FMergeEmitterResults MergeSourceChanges();
-
-	/** Gets an object which contains the history of modifications to the emitter instance. */
-	UObject* GetEmitterModificationHistory();
-
-	/** Sets an object which will contain the history of modifications to the emitter instance. */
-	void SetEmitterModificationHistory();
+	void ConditionalPostLoad(int32 NiagaraCustomVersion);
 
-	/** Removes the source system from the this emitter handle which will prevent inheriting any further changes. */
-	void RemoveSource();
+	/** Whether or not this handle uses the supplied emitter. */
+	bool UsesEmitter(const UNiagaraEmitter& InEmitter) const;
 
 #endif
 public:
@@ -116,11 +97,11 @@ private:
 #if WITH_EDITORONLY_DATA
 	/** The source emitter this emitter handle was built from. */
 	UPROPERTY()
-	UNiagaraEmitter* Source;
+	UNiagaraEmitter* Source_DEPRECATED;
 
 	/** An unmodified copy of the emitter this handle references for use when merging change from the source emitter. */
 	UPROPERTY()
-	UNiagaraEmitter* LastMergedSource;
+	UNiagaraEmitter* LastMergedSource_DEPRECATED;
 
 	bool bIsolated;
 #endif
diff --git a/Engine/Plugins/FX/Niagara/Source/Niagara/Classes/NiagaraSystem.h b/Engine/Plugins/FX/Niagara/Source/Niagara/Classes/NiagaraSystem.h
--- a/Engine/Plugins/FX/Niagara/Source/Niagara/Classes/NiagaraSystem.h
+++ b/Engine/Plugins/FX/Niagara/Source/Niagara/Classes/NiagaraSystem.h
@@ -139,15 +139,9 @@ public:
 	FORCEINLINE float GetWarmupTickDelta()const { return WarmupTickDelta; }
 
 #if WITH_EDITORONLY_DATA
-	/** Called to query whether or not this emitter is referenced as the source to any emitter handles for this System.*/
-	bool ReferencesSourceEmitter(UNiagaraEmitter& Emitter);
-
 	/** Determines if this system has the supplied emitter as an editable and simulating emitter instance. */
 	bool ReferencesInstanceEmitter(UNiagaraEmitter& Emitter);
 
-	/** Updates all handles which use this emitter as their source. */
-	void UpdateFromEmitterChanges(UNiagaraEmitter& ChangedSourceEmitter, bool bRecompileOnChange = true);
-
 	/** Updates the system's rapid iteration parameters from a specific emitter. */
 	void RefreshSystemParametersFromEmitter(const FNiagaraEmitterHandle& EmitterHandle);
 
@@ -238,7 +232,6 @@ public:
 
 private:
 #if WITH_EDITORONLY_DATA
-	INiagaraModule::FMergeEmitterResults MergeChangesForEmitterHandle(FNiagaraEmitterHandle& EmitterHandle);
 	bool QueryCompileComplete(bool bWait, bool bDoPost, bool bDoNotApply = false);
 #endif
 
diff --git a/Engine/Plugins/FX/Niagara/Source/Niagara/Private/NiagaraEmitter.cpp b/Engine/Plugins/FX/Niagara/Source/Niagara/Private/NiagaraEmitter.cpp
--- a/Engine/Plugins/FX/Niagara/Source/Niagara/Private/NiagaraEmitter.cpp
+++ b/Engine/Plugins/FX/Niagara/Source/Niagara/Private/NiagaraEmitter.cpp
@@ -14,6 +14,8 @@
 
 #if WITH_EDITOR
 const FName UNiagaraEmitter::PrivateMemberNames::EventHandlerScriptProps = GET_MEMBER_NAME_CHECKED(UNiagaraEmitter, EventHandlerScriptProps);
+
+const FString InitialNotSynchronizedReason("Emitter created");
 #endif
 
 static int32 GbForceNiagaraCompileOnLoad = 0;
@@ -24,6 +26,14 @@ static FAutoConsoleVariableRef CVarForceNiagaraCompileOnLoad(
 	ECVF_Default
 	);
 
+static int32 GbForceNiagaraMergeOnLoad = 0;
+static FAutoConsoleVariableRef CVarForceNiagaraMergeOnLoad(
+	TEXT("fx.ForceMergeOnLoad"),
+	GbForceNiagaraMergeOnLoad,
+	TEXT("If > 0 emitters will be forced to merge on load. \n"),
+	ECVF_Default
+);
+
 static int32 GbForceNiagaraFailToCompile = 0;
 static FAutoConsoleVariableRef CVarForceNiagaraCompileToFail(
 	TEXT("fx.ForceNiagaraCompileToFail"),
@@ -126,6 +136,84 @@ bool UNiagaraEmitter::GetForceCompileOnLoad()
 {
 	return GbForceNiagaraCompileOnLoad > 0;
 }
+
+bool UNiagaraEmitter::IsSynchronizedWithParent() const
+{
+	if (Parent == nullptr)
+	{
+		// If the emitter has no parent than it is synchronized by default.
+		return true;
+	}
+
+	if (ParentAtLastMerge == nullptr)
+	{
+		// If the parent was valid but the parent at last merge isn't, they we don't know if it's up to date so we say it's not, and let 
+		// the actual merge code print an appropriate message to the log.
+		return false;
+	}
+
+	if (Parent->GetChangeId().IsValid() == false ||
+		ParentAtLastMerge->GetChangeId().IsValid() == false)
+	{
+		// If any of the change Ids aren't valid then we assume we're out of sync.
+		return false;
+	}
+
+	// Otherwise check the change ids, and the force flag.
+	return Parent->GetChangeId() == ParentAtLastMerge->GetChangeId() && GbForceNiagaraMergeOnLoad <= 0;
+}
+
+INiagaraMergeManager::FMergeEmitterResults UNiagaraEmitter::MergeChangesFromParent()
+{
+	UE_LOG(LogNiagara, Log, TEXT("Emitter %s-%s is merging changes from parent %s because its Change ID was updated."), *GetPathName(), *GetName(),
+		Parent != nullptr ? *Parent->GetPathName() : TEXT("(null)"));
+
+	if (Parent == nullptr)
+	{
+		// If we don't have a copy of the parent emitter, this emitter can't safely be merged.
+		INiagaraMergeManager::FMergeEmitterResults MergeResults;
+		MergeResults.bSucceeded = false;
+		MergeResults.bModifiedGraph = false;
+		MergeResults.ErrorMessages.Add(NSLOCTEXT("NiagaraEmitter", "NoParentErrorMessage", "This emitter has no 'Parent' so changes can't be merged in."));
+		return MergeResults;
+	}
+
+	if (ParentAtLastMerge == nullptr)
+	{
+		// If we don't have a copy of the last merged parent emitter, this emitter can't safely be
+		// merged.
+		INiagaraMergeManager::FMergeEmitterResults MergeResults;
+		MergeResults.bSucceeded = false;
+		MergeResults.bModifiedGraph = false;
+		MergeResults.ErrorMessages.Add(NSLOCTEXT("NiagaraEmitter", "NoLastMergedParentErrorMessage", "This emitter has no 'ParentAtLastMerge' so changes can't be merged in."));
+		return MergeResults;
+	}
+
+	INiagaraModule& NiagaraModule = FModuleManager::Get().GetModuleChecked<INiagaraModule>("Niagara");
+	const INiagaraMergeManager& MergeManager = NiagaraModule.GetMergeManager();
+	INiagaraMergeManager::FMergeEmitterResults MergeResults = MergeManager.MergeEmitter(*Parent, *ParentAtLastMerge, *this);
+	if (MergeResults.bSucceeded)
+	{
+		UpdateFromMergedCopy(MergeManager, MergeResults.MergedInstance);
+
+		// Update the last merged source and clear it's stand alone and public flags since it's not an asset.
+		ParentAtLastMerge = CastChecked<UNiagaraEmitter>(StaticDuplicateObject(Parent, this));
+		ParentAtLastMerge->ClearFlags(RF_Standalone | RF_Public);
+	}
+	else
+	{
+		UE_LOG(LogNiagara, Warning, TEXT("Failed to merge changes for parent emitter.  Emitter: %s  Parent Emitter: %s  Error Message: %s"),
+			*GetPathName(), Parent != nullptr ? *Parent->GetPathName() : TEXT("(null)"), *MergeResults.GetErrorMessagesString());
+	}
+
+	return MergeResults;
+}
+
+bool UNiagaraEmitter::UsesEmitter(const UNiagaraEmitter& InEmitter) const
+{
+	return Parent == &InEmitter || (Parent != nullptr && Parent->UsesEmitter(InEmitter));
+}
+
 #endif
 
 void UNiagaraEmitter::Serialize(FArchive& Ar)
@@ -203,8 +291,14 @@ void UNiagaraEmitter::PostLoad()
 		Script->ConditionalPostLoad();
 	}
 
-	// Reset scripts if recompile is forced.
 #if WITH_EDITORONLY_DATA
+
+	if (IsSynchronizedWithParent() == false)
+	{
+		MergeChangesFromParent();
+	}
+
+	// Reset scripts if recompile is forced.
 	bool bGenerateNewChangeId = false;
 	if (GetForceCompileOnLoad())
 	{
@@ -283,6 +377,35 @@ void UNiagaraEmitter::PostLoad()
 }
 
 #if WITH_EDITOR
+/** Creates a new emitter with the supplied emitter as a parent emitter and the supplied system as it's owner. */
+UNiagaraEmitter* UNiagaraEmitter::CreateWithParentAndOwner(UNiagaraEmitter& InParentEmitter, UObject* InOwner, FName InName, EObjectFlags FlagMask)
+{
+	UNiagaraEmitter* NewEmitter = Cast<UNiagaraEmitter>(StaticDuplicateObject(&InParentEmitter, InOwner, InName, FlagMask));
+	NewEmitter->Parent = &InParentEmitter;
+	NewEmitter->ParentAtLastMerge = Cast<UNiagaraEmitter>(StaticDuplicateObject(&InParentEmitter, NewEmitter));
+	NewEmitter->ParentAtLastMerge->ClearFlags(RF_Standalone | RF_Public);
+	NewEmitter->SetUniqueEmitterName(InName.ToString());
+	NewEmitter->GraphSource->MarkNotSynchronized(InitialNotSynchronizedReason);
+	return NewEmitter;
+}
+
+/** Creates a new emitter by duplicating an existing emitter.  The new emitter  will reference the same parent emitter if one is available. */
+UNiagaraEmitter* UNiagaraEmitter::CreateAsDuplicate(const UNiagaraEmitter& InEmitterToDuplicate, FName InDuplicateName, UNiagaraSystem& InDuplicateOwnerSystem)
+{
+	UNiagaraEmitter* NewEmitter = Cast<UNiagaraEmitter>(StaticDuplicateObject(&InEmitterToDuplicate, &InDuplicateOwnerSystem));
+	NewEmitter->ClearFlags(RF_Standalone | RF_Public);
+	NewEmitter->Parent = InEmitterToDuplicate.Parent;
+	if (InEmitterToDuplicate.ParentAtLastMerge != nullptr)
+	{
+		NewEmitter->ParentAtLastMerge = InEmitterToDuplicate.ParentAtLastMerge;
+		NewEmitter->ParentAtLastMerge->ClearFlags(RF_Standalone | RF_Public);
+	}
+	NewEmitter->SetUniqueEmitterName(InDuplicateName.ToString());
+	NewEmitter->GraphSource->MarkNotSynchronized(InitialNotSynchronizedReason);
+
+	return NewEmitter;
+}
+
 void UNiagaraEmitter::PostEditChangeProperty(struct FPropertyChangedEvent& PropertyChangedEvent)
 {
 	Super::PostEditChangeProperty(PropertyChangedEvent);
@@ -668,11 +791,6 @@ bool UNiagaraEmitter::UsesScript(const UNiagaraScript* Script)const
 	return false;
 }
 
-//TODO
-// bool UNiagaraEmitter::UsesDataInterface(UNiagaraDataInterface* Interface)
-//{
-//}
-
 bool UNiagaraEmitter::UsesCollection(const class UNiagaraParameterCollection* Collection)const
 {
 	if (SpawnScriptProps.Script && SpawnScriptProps.Script->UsesCollection(Collection))
@@ -699,6 +817,90 @@ FString UNiagaraEmitter::GetUniqueEmitterName()const
 }
 
 #if WITH_EDITORONLY_DATA
+
+void UNiagaraEmitter::UpdateFromMergedCopy(const INiagaraMergeManager& MergeManager, UNiagaraEmitter* MergedEmitter)
+{
+	auto ReouterMergedObject = [](UObject* NewOuter, UObject* TargetObject)
+	{
+		FName MergedObjectUniqueName = MakeUniqueObjectName(NewOuter, TargetObject->GetClass(), TargetObject->GetFName());
+		TargetObject->Rename(*MergedObjectUniqueName.ToString(), NewOuter, REN_ForceNoResetLoaders);
+	};
+
+	// The merged copy was based on the parent emitter so its name might be wrong, check and fix that first,
+	// otherwise the rapid iteration parameter names will be wrong from the copied scripts.
+	if (MergedEmitter->GetUniqueEmitterName() != UniqueEmitterName)
+	{
+		MergedEmitter->SetUniqueEmitterName(UniqueEmitterName);
+	}
+
+	// Copy base editable emitter properties.
+	TArray<UProperty*> DifferentProperties;
+	MergeManager.DiffEditableProperties(this, MergedEmitter, *UNiagaraEmitter::StaticClass(), DifferentProperties);
+	MergeManager.CopyPropertiesToBase(this, MergedEmitter, DifferentProperties);
+
+	// Copy source and scripts
+	ReouterMergedObject(this, MergedEmitter->GraphSource);
+	GraphSource->OnChanged().RemoveAll(this);
+	GraphSource = MergedEmitter->GraphSource;
+	GraphSource->OnChanged().AddUObject(this, &UNiagaraEmitter::GraphSourceChanged);
+
+	ReouterMergedObject(this, MergedEmitter->SpawnScriptProps.Script);
+	SpawnScriptProps.Script->RapidIterationParameters.RemoveAllOnChangedHandlers(this);
+	SpawnScriptProps.Script = MergedEmitter->SpawnScriptProps.Script;
+	SpawnScriptProps.Script->RapidIterationParameters.AddOnChangedHandler(
+		FNiagaraParameterStore::FOnChanged::FDelegate::CreateUObject(this, &UNiagaraEmitter::ScriptRapidIterationParameterChanged));
+
+	ReouterMergedObject(this, MergedEmitter->UpdateScriptProps.Script);
+	UpdateScriptProps.Script->RapidIterationParameters.RemoveAllOnChangedHandlers(this);
+	UpdateScriptProps.Script = MergedEmitter->UpdateScriptProps.Script;
+	UpdateScriptProps.Script->RapidIterationParameters.AddOnChangedHandler(
+		FNiagaraParameterStore::FOnChanged::FDelegate::CreateUObject(this, &UNiagaraEmitter::ScriptRapidIterationParameterChanged));
+
+	ReouterMergedObject(this, MergedEmitter->EmitterSpawnScriptProps.Script);
+	EmitterSpawnScriptProps.Script->RapidIterationParameters.RemoveAllOnChangedHandlers(this);
+	EmitterSpawnScriptProps.Script = MergedEmitter->EmitterSpawnScriptProps.Script;
+	EmitterSpawnScriptProps.Script->RapidIterationParameters.AddOnChangedHandler(
+		FNiagaraParameterStore::FOnChanged::FDelegate::CreateUObject(this, &UNiagaraEmitter::ScriptRapidIterationParameterChanged));
+
+	ReouterMergedObject(this, MergedEmitter->EmitterUpdateScriptProps.Script);
+	EmitterUpdateScriptProps.Script->RapidIterationParameters.RemoveAllOnChangedHandlers(this);
+	EmitterUpdateScriptProps.Script = MergedEmitter->EmitterUpdateScriptProps.Script;
+	EmitterUpdateScriptProps.Script->RapidIterationParameters.AddOnChangedHandler(
+		FNiagaraParameterStore::FOnChanged::FDelegate::CreateUObject(this, &UNiagaraEmitter::ScriptRapidIterationParameterChanged));
+
+	// Copy event handlers
+	for (FNiagaraEventScriptProperties& EventScriptProperties : EventHandlerScriptProps)
+	{
+		EventScriptProperties.Script->RapidIterationParameters.RemoveAllOnChangedHandlers(this);
+	}
+	EventHandlerScriptProps.Empty();
+
+	for (FNiagaraEventScriptProperties& MergedEventScriptProperties : MergedEmitter->EventHandlerScriptProps)
+	{
+		EventHandlerScriptProps.Add(MergedEventScriptProperties);
+		ReouterMergedObject(this, MergedEventScriptProperties.Script);
+		MergedEventScriptProperties.Script->RapidIterationParameters.AddOnChangedHandler(
+			FNiagaraParameterStore::FOnChanged::FDelegate::CreateUObject(this, &UNiagaraEmitter::ScriptRapidIterationParameterChanged));
+	}
+
+	// Copy renderers
+	for (UNiagaraRendererProperties* Renderer : RendererProperties)
+	{
+		Renderer->OnChanged().RemoveAll(this);
+	}
+	RendererProperties.Empty();
+
+	for (UNiagaraRendererProperties* MergedRenderer : MergedEmitter->RendererProperties)
+	{
+		ReouterMergedObject(this, MergedRenderer);
+		RendererProperties.Add(MergedRenderer);
+		MergedRenderer->OnChanged().AddUObject(this, &UNiagaraEmitter::RendererChanged);
+	}
+
+	// Update the change id since we don't know what's changed.
+	ChangeId = FGuid::NewGuid();
+}
+
 void UNiagaraEmitter::SyncEmitterAlias(const FString& InOldName, const FString& InNewName)
 {
 	TMap<FString, FString> RenameMap;
@@ -724,6 +926,13 @@ bool UNiagaraEmitter::SetUniqueEmitterName(const FString& InName)
 		FString OldName = UniqueEmitterName;
 		UniqueEmitterName = InName;
 
+		if (GetName() != InName)
+		{
+			// Also rename the underlying uobject to keep things consistent.
+			FName UniqueObjectName = MakeUniqueObjectName(GetOuter(), UNiagaraEmitter::StaticClass(), *InName);
+			Rename(*UniqueObjectName.ToString());
+		}
+
 #if WITH_EDITORONLY_DATA
 		SyncEmitterAlias(OldName, UniqueEmitterName);
 #endif
@@ -880,3 +1089,14 @@ void UNiagaraEmitter::GenerateStatID()
 	StatID_RT_CNC = FDynamicStats::CreateStatId<FStatGroup_STATGROUP_NiagaraEmitters>(GetName() + TEXT("[RT_CNC]"));
 #endif
 }
+
+UNiagaraEmitter* UNiagaraEmitter::GetParent() const
+{
+	return Parent;
+}
+
+void UNiagaraEmitter::RemoveParent()
+{
+	Parent = nullptr;
+	ParentAtLastMerge = nullptr;
+}
diff --git a/Engine/Plugins/FX/Niagara/Source/Niagara/Private/NiagaraEmitterHandle.cpp b/Engine/Plugins/FX/Niagara/Source/Niagara/Private/NiagaraEmitterHandle.cpp
--- a/Engine/Plugins/FX/Niagara/Source/Niagara/Private/NiagaraEmitterHandle.cpp
+++ b/Engine/Plugins/FX/Niagara/Source/Niagara/Private/NiagaraEmitterHandle.cpp
@@ -11,54 +11,29 @@
 #include "Modules/ModuleManager.h"
 
 const FNiagaraEmitterHandle FNiagaraEmitterHandle::InvalidHandle;
-const FString InitialNotSynchronizedReason("Emitter handle constructed");
 
-FNiagaraEmitterHandle::FNiagaraEmitterHandle() :
-	bIsEnabled(true),
+FNiagaraEmitterHandle::FNiagaraEmitterHandle() 
+	: bIsEnabled(true)
 #if WITH_EDITORONLY_DATA
-	Source(nullptr),
-	LastMergedSource(nullptr),
-	bIsolated(false),
+	, Source_DEPRECATED(nullptr)
+	, LastMergedSource_DEPRECATED(nullptr)
+	, bIsolated(false)
 #endif
-	Instance(nullptr)
+	, Instance(nullptr)
 {
 }
 
 #if WITH_EDITORONLY_DATA
-FNiagaraEmitterHandle::FNiagaraEmitterHandle(UNiagaraEmitter& InSourceEmitter, FName InName, UNiagaraSystem& InOuterSystem)
+FNiagaraEmitterHandle::FNiagaraEmitterHandle(UNiagaraEmitter& InEmitter)
 	: Id(FGuid::NewGuid())
 	, IdName(*Id.ToString())
 	, bIsEnabled(true)
-	, Name(InName)
-	, Source(&InSourceEmitter)
-	, LastMergedSource(Cast<UNiagaraEmitter>(StaticDuplicateObject(Source, &InOuterSystem)))
-	, bIsolated(false)
-	, Instance(Cast<UNiagaraEmitter>(StaticDuplicateObject(Source, &InOuterSystem)))
-{
-	Instance->ClearFlags(RF_Standalone | RF_Public);
-	Instance->SetUniqueEmitterName(Name.ToString());
-	Instance->GraphSource->MarkNotSynchronized(InitialNotSynchronizedReason);
-	LastMergedSource->ClearFlags(RF_Standalone | RF_Public);
-}
-
-FNiagaraEmitterHandle::FNiagaraEmitterHandle(const FNiagaraEmitterHandle& InHandleToDuplicate, FName InDuplicateName, UNiagaraSystem& InDuplicateOwnerSystem)
-	: Id(FGuid::NewGuid())
-	, IdName(*Id.ToString())
-	, bIsEnabled(InHandleToDuplicate.bIsEnabled)
-	, Name(InDuplicateName)
-	, Source(InHandleToDuplicate.Source)
-	, LastMergedSource(InHandleToDuplicate.LastMergedSource != nullptr ? Cast<UNiagaraEmitter>(StaticDuplicateObject(InHandleToDuplicate.LastMergedSource, &InDuplicateOwnerSystem)) : nullptr)
+	, Name(*InEmitter.GetUniqueEmitterName())
+	, Source_DEPRECATED(nullptr)
+	, LastMergedSource_DEPRECATED(nullptr)
 	, bIsolated(false)
-	, Instance(Cast<UNiagaraEmitter>(StaticDuplicateObject(InHandleToDuplicate.Instance, &InDuplicateOwnerSystem)))
+	, Instance(&InEmitter)
 {
-	// Clear stand alone and public flags from the referenced emitters since they are not assets.
-	Instance->ClearFlags(RF_Standalone | RF_Public);
-	Instance->SetUniqueEmitterName(Name.ToString());
-	Instance->GraphSource->MarkNotSynchronized(InitialNotSynchronizedReason);
-	if (LastMergedSource != nullptr)
-	{
-		LastMergedSource->ClearFlags(RF_Standalone | RF_Public);
-	}
 }
 #endif
 
@@ -129,14 +104,6 @@ void FNiagaraEmitterHandle::SetIsEnabled(bool bInIsEnabled)
 	bIsEnabled = bInIsEnabled;
 }
 
-#if WITH_EDITORONLY_DATA
-const UNiagaraEmitter* FNiagaraEmitterHandle::GetSource() const
-{
-	return Source;
-}
-
-#endif
-
 UNiagaraEmitter* FNiagaraEmitterHandle::GetInstance() const
 {
 	return Instance;
@@ -150,31 +117,6 @@ FString FNiagaraEmitterHandle::GetUniqueInstanceName()const
 
 #if WITH_EDITORONLY_DATA
 
-bool FNiagaraEmitterHandle::IsSynchronizedWithSource() const
-{
-	if (Source == nullptr && LastMergedSource == nullptr)
-	{
-		// If the emitter has no source and no last merged source than it is synchronized by default.
-		return true;
-	}
-
-	if (Source == nullptr || LastMergedSource == nullptr)
-	{
-		// If either only the source or only the last merged sources is missing, then we're not synchronized.  The
-		// merge logic will detect this and print an appropriate message to the log.
-		return false;
-	}
-
-	if (Source->GetChangeId().IsValid() == false ||
-		LastMergedSource->GetChangeId().IsValid() == false)
-	{
-		// If any of the change Ids aren't valid then we assume we're out of sync.
-		return false;
-	}
-
-	return Source->GetChangeId() == LastMergedSource->GetChangeId();
-}
-
 bool FNiagaraEmitterHandle::NeedsRecompile() const
 {
 	TArray<UNiagaraScript*> Scripts;
@@ -190,69 +132,41 @@ bool FNiagaraEmitterHandle::NeedsRecompile() const
 	return false;
 }
 
-void FNiagaraEmitterHandle::ConditionalPostLoad()
-{
-	if (Source != nullptr)
-	{
-		Source->ConditionalPostLoad();
-	}
-	if (LastMergedSource != nullptr)
-	{
-		LastMergedSource->ConditionalPostLoad();
-	}
-	Instance->ConditionalPostLoad();
-}
-
-INiagaraModule::FMergeEmitterResults FNiagaraEmitterHandle::MergeSourceChanges()
+void FNiagaraEmitterHandle::ConditionalPostLoad(int32 NiagaraCustomVersion)
 {
-	UE_LOG(LogNiagara, Log, TEXT("Emitter %s-%s is merging changes from source %s because its Change ID was updated."), *Instance->GetPathName(), *Name.ToString(),
-		Source != nullptr ? *Source->GetPathName() : TEXT("(null)"));
-
-	if (Source == nullptr)
-	{
-		// If we don't have a copy of the source emitter, this emitter can't safely be merged.
-		INiagaraModule::FMergeEmitterResults MergeResults;
-		MergeResults.bSucceeded = false;
-		MergeResults.bModifiedGraph = false;
-		MergeResults.ErrorMessages.Add(NSLOCTEXT("NiagaraEmitterHandle", "NoSourceErrorMessage", "This emitter has no 'Source' so changes can't be merged in."));
-		return MergeResults;
-	}
-
-	if (LastMergedSource == nullptr)
+	if (Instance != nullptr)
 	{
-		// If we don't have a copy of the last merged source emitter, this emitter can't safely be
-		// merged.
-		INiagaraModule::FMergeEmitterResults MergeResults;
-		MergeResults.bSucceeded = false;
-		MergeResults.bModifiedGraph = false;
-		MergeResults.ErrorMessages.Add(NSLOCTEXT("NiagaraEmitterHandle", "NoLastMergedSourceErrorMessage", "This emitter has no 'LastMergedSource' so changes can't be merged in."));
-		return MergeResults;
-	}
-
-	INiagaraModule& NiagaraModule = FModuleManager::Get().GetModuleChecked<INiagaraModule>("Niagara");
-	INiagaraModule::FMergeEmitterResults MergeResults = NiagaraModule.MergeEmitter(*Source, *LastMergedSource, *Instance);
-	if (MergeResults.bSucceeded)
-	{
-		UObject* Outer = Instance->GetOuter();
-		Instance = MergeResults.MergedInstance;
-
-		// Rename the merged instance into this package with the correct handle name and then clear it's stand alone and public flags since it's not a root asset.
-		FName NewInstanceName = MakeUniqueObjectName(Outer, UNiagaraEmitter::StaticClass(), Name);
-		Instance->Rename(*NewInstanceName.ToString(), Outer, REN_ForceNoResetLoaders);
-		Instance->ClearFlags(RF_Standalone | RF_Public);
-		Instance->SetUniqueEmitterName(Name.ToString());
-
-		// Update the last merged source and clear it's stand alone and public flags since it's not an asset.
-		LastMergedSource = CastChecked<UNiagaraEmitter>(StaticDuplicateObject(Source, Outer));
-		LastMergedSource->ClearFlags(RF_Standalone | RF_Public);
+		Instance->ConditionalPostLoad();
+		if (NiagaraCustomVersion < FNiagaraCustomVersion::MoveInheritanceDataFromTheEmitterHandleToTheEmitter)
+		{
+			if (Source_DEPRECATED != nullptr)
+			{
+				Source_DEPRECATED->ConditionalPostLoad();
+				Instance->Parent = Source_DEPRECATED;
+				Instance->Parent->Rename(nullptr, Instance, REN_ForceNoResetLoaders);
+				Source_DEPRECATED = nullptr;
+			}
+			if (LastMergedSource_DEPRECATED != nullptr)
+			{
+				LastMergedSource_DEPRECATED->ConditionalPostLoad();
+				Instance->ParentAtLastMerge = LastMergedSource_DEPRECATED;
+				Instance->ParentAtLastMerge->Rename(nullptr, Instance, REN_ForceNoResetLoaders);
+				LastMergedSource_DEPRECATED = nullptr;
+			}
+
+			// Since we've previously post loaded the emitter it wouldn't have merged on load since it didn't have
+			// the parent information so we check this again here, now that the parent information has been set.
+			if (Instance->IsSynchronizedWithParent() == false)
+			{
+				Instance->MergeChangesFromParent();
+			}
+		}
 	}
-	return MergeResults;
 }
 
-void FNiagaraEmitterHandle::RemoveSource()
+bool FNiagaraEmitterHandle::UsesEmitter(const UNiagaraEmitter& InEmitter) const
 {
-	Source = nullptr;
-	LastMergedSource = nullptr;
+	return Instance == &InEmitter || (Instance != nullptr && Instance->UsesEmitter(InEmitter));
 }
 
 #endif
\ No newline at end of file
diff --git a/Engine/Plugins/FX/Niagara/Source/Niagara/Private/NiagaraModule.cpp b/Engine/Plugins/FX/Niagara/Source/Niagara/Private/NiagaraModule.cpp
--- a/Engine/Plugins/FX/Niagara/Source/Niagara/Private/NiagaraModule.cpp
+++ b/Engine/Plugins/FX/Niagara/Source/Niagara/Private/NiagaraModule.cpp
@@ -352,30 +352,23 @@ void INiagaraModule::TickWorld(UWorld* World, ELevelTick TickType, float DeltaSe
 }
 
 #if WITH_EDITOR
-INiagaraModule::FMergeEmitterResults INiagaraModule::MergeEmitter(UNiagaraEmitter& Source, UNiagaraEmitter& LastMergedSource, UNiagaraEmitter& Instance)
+const INiagaraMergeManager& INiagaraModule::GetMergeManager()
 {
-	if (OnMergeEmitterDelegate.IsBound())
-	{
-		return OnMergeEmitterDelegate.Execute(Source, LastMergedSource, Instance);
-	}
-	FMergeEmitterResults Results;
-	Results.bSucceeded = false;
-	Results.ErrorMessages.Add(FText::Format(LOCTEXT("MergeDelegateNotRegisteredFormat", "Failed to merge emitter {0}.  Merge delegate not registered."), FText::FromString(Instance.GetPathName())));
-	return Results;
+	checkf(MergeManager.IsValid(), TEXT("Merge manager was never registered, or was unregistered."));
+	return *MergeManager.Get();
 }
 
-FDelegateHandle INiagaraModule::RegisterOnMergeEmitter(FOnMergeEmitter OnMergeEmitter)
+void INiagaraModule::RegisterMergeManager(TSharedRef<INiagaraMergeManager> InMergeManager)
 {
-	checkf(OnMergeEmitterDelegate.IsBound() == false, TEXT("Only one handler is allowed for the OnMergeEmitter delegate"));
-	OnMergeEmitterDelegate = OnMergeEmitter;
-	return OnMergeEmitterDelegate.GetHandle();
+	checkf(MergeManager.IsValid() == false, TEXT("Only one merge manager can be registered at a time."));
+	MergeManager = InMergeManager;
 }
 
-void INiagaraModule::UnregisterOnMergeEmitter(FDelegateHandle DelegateHandle)
+void INiagaraModule::UnregisterMergeManager(TSharedRef<INiagaraMergeManager> InMergeManager)
 {
-	checkf(OnMergeEmitterDelegate.IsBound(), TEXT("OnMergeEmitter is not registered"));
-	checkf(OnMergeEmitterDelegate.GetHandle() == DelegateHandle, TEXT("Can only unregister the OnMergeEmitter delegate with the handle it was registered with."));
-	OnMergeEmitterDelegate.Unbind();
+	checkf(MergeManager.IsValid(), TEXT("MergeManager is not registered"));
+	checkf(MergeManager == InMergeManager, TEXT("Can only unregister the merge manager which was previously registered."));
+	MergeManager.Reset();
 }
 
 UNiagaraScriptSourceBase* INiagaraModule::CreateDefaultScriptSource(UObject* Outer)
diff --git a/Engine/Plugins/FX/Niagara/Source/Niagara/Private/NiagaraSystem.cpp b/Engine/Plugins/FX/Niagara/Source/Niagara/Private/NiagaraSystem.cpp
--- a/Engine/Plugins/FX/Niagara/Source/Niagara/Private/NiagaraSystem.cpp
+++ b/Engine/Plugins/FX/Niagara/Source/Niagara/Private/NiagaraSystem.cpp
@@ -130,7 +130,7 @@ bool UNiagaraSystem::UsesScript(const UNiagaraScript* Script)const
 
 	for (FNiagaraEmitterHandle EmitterHandle : GetEmitterHandles())
 	{
-		if ((EmitterHandle.GetSource() && EmitterHandle.GetSource()->UsesScript(Script)) || (EmitterHandle.GetInstance() && EmitterHandle.GetInstance()->UsesScript(Script)))
+		if (EmitterHandle.GetInstance() && EmitterHandle.GetInstance()->UsesScript(Script))
 		{
 			return true;
 		}
@@ -141,11 +141,14 @@ bool UNiagaraSystem::UsesScript(const UNiagaraScript* Script)const
 
 bool UNiagaraSystem::UsesEmitter(const UNiagaraEmitter* Emitter) const
 {
-	for (FNiagaraEmitterHandle EmitterHandle : GetEmitterHandles())
+	if (Emitter != nullptr)
 	{
-		if (Emitter == EmitterHandle.GetSource() || Emitter == EmitterHandle.GetInstance())
+		for (const FNiagaraEmitterHandle& EmitterHandle : GetEmitterHandles())
 		{
-			return true;
+			if (EmitterHandle.UsesEmitter(*Emitter))
+			{
+				return true;
+			}
 		}
 	}
 	return false;
@@ -290,7 +293,6 @@ void UNiagaraSystem::PostLoad()
 		bSystemScriptsAreSynchronized &= SystemScript->AreScriptAndSourceSynchronized();
 	}
 
-	bool bEmitterGraphChangedFromMerge = false;
 	bool bEmitterScriptsAreSynchronized = true;
 
 #if 0
@@ -313,21 +315,11 @@ void UNiagaraSystem::PostLoad()
 
 	for (FNiagaraEmitterHandle& EmitterHandle : EmitterHandles)
 	{
-		EmitterHandle.ConditionalPostLoad();
-		if (EmitterHandle.IsSynchronizedWithSource() == false)
+		EmitterHandle.ConditionalPostLoad(NiagaraVer);
+		if (!EmitterHandle.GetInstance()->AreAllScriptAndSourcesSynchronized())
 		{
-			INiagaraModule::FMergeEmitterResults Results = MergeChangesForEmitterHandle(EmitterHandle);
-			if (Results.bSucceeded)
-			{
-				bEmitterGraphChangedFromMerge |= Results.bModifiedGraph;
-			}
-		}
-		if (bEmitterScriptsAreSynchronized)
-		{
-			if (!EmitterHandle.GetInstance()->AreAllScriptAndSourcesSynchronized())
-			{
-				bEmitterScriptsAreSynchronized = false;
-			}
+			bEmitterScriptsAreSynchronized = false;
+			break;
 		}
 	}
 
@@ -352,11 +344,6 @@ void UNiagaraSystem::PostLoad()
 		UE_LOG(LogNiagara, Log, TEXT("System %s being compiled because there were changes to an emitter script Change ID."), *GetPathName());
 	}
 
-	if (bEmitterGraphChangedFromMerge)
-	{
-		UE_LOG(LogNiagara, Log, TEXT("System %s being compiled because graph changes were merged for a base emitter."), *GetPathName());
-	}
-
 #if 0
 	UE_LOG(LogNiagara, Log, TEXT("Before"));
 	for (FNiagaraEmitterHandle& EmitterHandle : EmitterHandles)
@@ -375,7 +362,7 @@ void UNiagaraSystem::PostLoad()
 	}
 #endif
 
-	if (bSystemScriptsAreSynchronized == false || bEmitterScriptsAreSynchronized == false || bEmitterGraphChangedFromMerge)
+	if (bSystemScriptsAreSynchronized == false || bEmitterScriptsAreSynchronized == false)
 	{
 		RequestCompile(false);
 	}
@@ -417,39 +404,6 @@ void UNiagaraSystem::SetEditorData(UNiagaraEditorDataBase* InEditorData)
 	EditorData = InEditorData;
 }
 
-INiagaraModule::FMergeEmitterResults UNiagaraSystem::MergeChangesForEmitterHandle(FNiagaraEmitterHandle& EmitterHandle)
-{
-	INiagaraModule::FMergeEmitterResults Results = EmitterHandle.MergeSourceChanges();
-	if (Results.bSucceeded)
-	{
-		UNiagaraEmitter* Instance = EmitterHandle.GetInstance();
-		RefreshSystemParametersFromEmitter(EmitterHandle);
-		if (Instance->bInterpolatedSpawning)
-		{
-			Instance->UpdateScriptProps.Script->RapidIterationParameters.CopyParametersTo(
-				Instance->SpawnScriptProps.Script->RapidIterationParameters, false, FNiagaraParameterStore::EDataInterfaceCopyMethod::None);
-		}
-	}
-	else
-	{
-		UE_LOG(LogNiagara, Warning, TEXT("Failed to merge changes for base emitter.  System: %s  Emitter: %s  Error Message: %s"), 
-			*GetPathName(), *EmitterHandle.GetName().ToString(), *Results.GetErrorMessagesString());
-	}
-	return Results;
-}
-
-bool UNiagaraSystem::ReferencesSourceEmitter(UNiagaraEmitter& Emitter)
-{
-	for (FNiagaraEmitterHandle& Handle : EmitterHandles)
-	{
-		if (&Emitter == Handle.GetSource())
-		{
-			return true;
-		}
-	}
-	return false;
-}
-
 bool UNiagaraSystem::ReferencesInstanceEmitter(UNiagaraEmitter& Emitter)
 {
 	for (FNiagaraEmitterHandle& Handle : EmitterHandles)
@@ -462,24 +416,6 @@ bool UNiagaraSystem::ReferencesInstanceEmitter(UNiagaraEmitter& Emitter)
 	return false;
 }
 
-void UNiagaraSystem::UpdateFromEmitterChanges(UNiagaraEmitter& ChangedSourceEmitter, bool bRecompileOnChange)
-{
-	bool bNeedsCompile = false;
-	for(FNiagaraEmitterHandle& EmitterHandle : EmitterHandles)
-	{
-		if (EmitterHandle.GetSource() == &ChangedSourceEmitter)
-		{
-			INiagaraModule::FMergeEmitterResults Results = MergeChangesForEmitterHandle(EmitterHandle);
-			bNeedsCompile |= Results.bSucceeded && Results.bModifiedGraph;
-		}
-	}
-
-	if (bNeedsCompile && bRecompileOnChange)
-	{
-		RequestCompile(false);
-	}
-}
-
 void UNiagaraSystem::RefreshSystemParametersFromEmitter(const FNiagaraEmitterHandle& EmitterHandle)
 {
 	InitEmitterSpawnAttributes();
@@ -618,12 +554,15 @@ bool UNiagaraSystem::IsValid()const
 }
 #if WITH_EDITORONLY_DATA
 
-FNiagaraEmitterHandle UNiagaraSystem::AddEmitterHandle(UNiagaraEmitter& SourceEmitter, FName EmitterName)
+FNiagaraEmitterHandle UNiagaraSystem::AddEmitterHandle(UNiagaraEmitter& InEmitter, FName EmitterName)
 {
-	FNiagaraEmitterHandle EmitterHandle(SourceEmitter, EmitterName, *this);
-	if (SourceEmitter.bIsTemplateAsset)
+	UNiagaraEmitter* NewEmitter = UNiagaraEmitter::CreateWithParentAndOwner(InEmitter, this, EmitterName, ~(RF_Public | RF_Standalone));
+	FNiagaraEmitterHandle EmitterHandle(*NewEmitter);
+	if (InEmitter.bIsTemplateAsset)
 	{
-		EmitterHandle.RemoveSource();
+		NewEmitter->bIsTemplateAsset = false;
+		NewEmitter->TemplateAssetDescription = FText();
+		NewEmitter->RemoveParent();
 	}
 	EmitterHandles.Add(EmitterHandle);
 	RefreshSystemParametersFromEmitter(EmitterHandle);
@@ -632,7 +571,9 @@ FNiagaraEmitterHandle UNiagaraSystem::AddEmitterHandle(UNiagaraEmitter& SourceEm
 
 FNiagaraEmitterHandle UNiagaraSystem::DuplicateEmitterHandle(const FNiagaraEmitterHandle& EmitterHandleToDuplicate, FName EmitterName)
 {
-	FNiagaraEmitterHandle EmitterHandle(EmitterHandleToDuplicate, EmitterName, *this);
+	UNiagaraEmitter* DuplicateEmitter = UNiagaraEmitter::CreateAsDuplicate(*EmitterHandleToDuplicate.GetInstance(), EmitterName, *this);
+	FNiagaraEmitterHandle EmitterHandle(*DuplicateEmitter);
+	EmitterHandle.SetIsEnabled(EmitterHandleToDuplicate.GetIsEnabled());
 	EmitterHandles.Add(EmitterHandle);
 	RefreshSystemParametersFromEmitter(EmitterHandle);
 	return EmitterHandle;
diff --git a/Engine/Plugins/FX/Niagara/Source/Niagara/Private/NiagaraSystemInstance.cpp b/Engine/Plugins/FX/Niagara/Source/Niagara/Private/NiagaraSystemInstance.cpp
--- a/Engine/Plugins/FX/Niagara/Source/Niagara/Private/NiagaraSystemInstance.cpp
+++ b/Engine/Plugins/FX/Niagara/Source/Niagara/Private/NiagaraSystemInstance.cpp
@@ -1261,7 +1261,7 @@ bool FNiagaraSystemInstance::UsesScript(const UNiagaraScript* Script)const
 	{
 		for (FNiagaraEmitterHandle EmitterHandle : GetSystem()->GetEmitterHandles())
 		{
-			if ((EmitterHandle.GetSource() && EmitterHandle.GetSource()->UsesScript(Script)) || (EmitterHandle.GetInstance() && EmitterHandle.GetInstance()->UsesScript(Script)))
+			if (EmitterHandle.GetInstance() && EmitterHandle.GetInstance()->UsesScript(Script))
 			{
 				return true;
 			}
diff --git a/Engine/Plugins/FX/Niagara/Source/Niagara/Public/INiagaraMergeManager.h b/Engine/Plugins/FX/Niagara/Source/Niagara/Public/INiagaraMergeManager.h
new file mode 100644
--- /dev/null
+++ b/Engine/Plugins/FX/Niagara/Source/Niagara/Public/INiagaraMergeManager.h
@@ -0,0 +1,39 @@
+// Copyright 1998-2019 Epic Games, Inc. All Rights Reserved.
+
+#pragma once
+
+class INiagaraMergeManager
+{
+public:
+	struct FMergeEmitterResults
+	{
+		FMergeEmitterResults()
+			: bSucceeded(true)
+			, bModifiedGraph(false)
+		{
+		}
+
+		bool bSucceeded;
+		TArray<FText> ErrorMessages;
+		bool bModifiedGraph;
+		UNiagaraEmitter* MergedInstance;
+
+		FString GetErrorMessagesString() const
+		{
+			TArray<FString> ErrorMessageStrings;
+			for (FText ErrorMessage : ErrorMessages)
+			{
+				ErrorMessageStrings.Add(ErrorMessage.ToString());
+			}
+			return FString::Join(ErrorMessageStrings, TEXT("\n"));
+		}
+	};
+
+	virtual FMergeEmitterResults MergeEmitter(UNiagaraEmitter& Parent, UNiagaraEmitter& ParentAtLastMerge, UNiagaraEmitter& Instance) const = 0;
+
+	virtual void DiffEditableProperties(const void* BaseDataAddress, const void* OtherDataAddress, UStruct& Struct, TArray<UProperty*>& OutDifferentProperties) const = 0;
+
+	virtual void CopyPropertiesToBase(void* BaseDataAddress, const void* OtherDataAddress, TArray<UProperty*> PropertiesToCopy) const = 0;
+
+	virtual ~INiagaraMergeManager() { }
+};
\ No newline at end of file
diff --git a/Engine/Plugins/FX/Niagara/Source/Niagara/Public/NiagaraModule.h b/Engine/Plugins/FX/Niagara/Source/Niagara/Public/NiagaraModule.h
--- a/Engine/Plugins/FX/Niagara/Source/Niagara/Public/NiagaraModule.h
+++ b/Engine/Plugins/FX/Niagara/Source/Niagara/Public/NiagaraModule.h
@@ -15,6 +15,7 @@ struct FNiagaraVMExecutableData;
 class UNiagaraScript;
 class FNiagaraCompileOptions;
 class FNiagaraCompileRequestDataBase;
+class INiagaraMergeManager;
 
 extern NIAGARA_API int32 GEnableVerboseNiagaraChangeIdLogging;
 
@@ -23,38 +24,9 @@ extern NIAGARA_API int32 GEnableVerboseNiagaraChangeIdLogging;
 */
 class NIAGARA_API INiagaraModule : public IModuleInterface
 {
-public:
-#if WITH_EDITOR
-	struct FMergeEmitterResults
-	{
-		FMergeEmitterResults()
-			: bSucceeded(true)
-			, bModifiedGraph(false)
-		{
-		}
-
-		bool bSucceeded;
-		TArray<FText> ErrorMessages;
-		bool bModifiedGraph;
-		UNiagaraEmitter* MergedInstance;
-
-		FString GetErrorMessagesString() const
-		{
-			TArray<FString> ErrorMessageStrings;
-			for (FText ErrorMessage : ErrorMessages)
-			{
-				ErrorMessageStrings.Add(ErrorMessage.ToString());
-			}
-			return FString::Join(ErrorMessageStrings, TEXT("\n"));
-		}
-	};
-	
-#endif
-
 public:
 #if WITH_EDITOR
 	typedef TSharedPtr<FNiagaraCompileRequestDataBase, ESPMode::ThreadSafe> CompileRequestPtr;
-	DECLARE_DELEGATE_RetVal_ThreeParams(FMergeEmitterResults, FOnMergeEmitter, UNiagaraEmitter&, UNiagaraEmitter&, UNiagaraEmitter&);
 	DECLARE_DELEGATE_RetVal_OneParam(class UNiagaraScriptSourceBase*, FOnCreateDefaultScriptSource, UObject*);
 	DECLARE_DELEGATE_RetVal_TwoParams(TSharedPtr<FNiagaraVMExecutableData>, FScriptCompiler,const FNiagaraCompileRequestDataBase*, const FNiagaraCompileOptions&);
 	DECLARE_DELEGATE_RetVal_OneParam(CompileRequestPtr, FOnPrecompile, UObject*);
@@ -87,11 +59,11 @@ public:
 	void TickWorld(UWorld* World, ELevelTick TickType, float DeltaSeconds);
 
 #if WITH_EDITOR
-	FMergeEmitterResults MergeEmitter(UNiagaraEmitter& Source, UNiagaraEmitter& LastMergedSource, UNiagaraEmitter& Instance);
+	const INiagaraMergeManager& GetMergeManager();
 
-	FDelegateHandle RegisterOnMergeEmitter(FOnMergeEmitter OnMergeEmitter);
+	void RegisterMergeManager(TSharedRef<INiagaraMergeManager> InMergeManager);
 
-	void UnregisterOnMergeEmitter(FDelegateHandle DelegateHandle);
+	void UnregisterMergeManager(TSharedRef<INiagaraMergeManager> InMergeManager);
 
 	UNiagaraScriptSourceBase* CreateDefaultScriptSource(UObject* Outer);
 
@@ -200,7 +172,7 @@ private:
 	FOnProcessQueue OnProcessQueue;
 
 #if WITH_EDITORONLY_DATA
-	FOnMergeEmitter OnMergeEmitterDelegate;
+	TSharedPtr<INiagaraMergeManager> MergeManager;
 	FOnCreateDefaultScriptSource OnCreateDefaultScriptSourceDelegate;
 	FScriptCompiler ScriptCompilerDelegate;
 	FOnPrecompile ObjectPrecompilerDelegate;
diff --git a/Engine/Plugins/FX/Niagara/Source/NiagaraCore/Public/NiagaraCustomVersion.h b/Engine/Plugins/FX/Niagara/Source/NiagaraCore/Public/NiagaraCustomVersion.h
--- a/Engine/Plugins/FX/Niagara/Source/NiagaraCore/Public/NiagaraCustomVersion.h
+++ b/Engine/Plugins/FX/Niagara/Source/NiagaraCore/Public/NiagaraCustomVersion.h
@@ -119,6 +119,8 @@ struct FNiagaraCustomVersion
 		UseHashesToIdentifyCompileStateOfTopLevelScripts, // Move to using the traversed graph hash and the base script id for the FNiagaraVMExecutableDataId instead of the change id guid to prevent invalidating the DDC.
 
 		MetaDataAndParametersUpdate, // Reworked how the metadata is stored in NiagaraGraph from storing a Map of FNiagaraVariableMetaData to storing a map of UNiagaraScriptVariable* to be used with the Details panel.
+
+		MoveInheritanceDataFromTheEmitterHandleToTheEmitter, // Moved the emitter inheritance data from the emitter handle to the emitter to allow for chained emitter inheritance.
 		
 		// -----<new versions can be added above this line>-------------------------------------------------
 		VersionPlusOne,
diff --git a/Engine/Plugins/FX/Niagara/Source/NiagaraEditor/Private/NiagaraEditorModule.cpp b/Engine/Plugins/FX/Niagara/Source/NiagaraEditor/Private/NiagaraEditorModule.cpp
--- a/Engine/Plugins/FX/Niagara/Source/NiagaraEditor/Private/NiagaraEditorModule.cpp
+++ b/Engine/Plugins/FX/Niagara/Source/NiagaraEditor/Private/NiagaraEditorModule.cpp
@@ -277,6 +277,34 @@ void DumpRapidIterationParamersForAsset(const TArray<FString>& Arguments)
 	}
 }
 
+void CompileEmitterStandAlone(UNiagaraEmitter* Emitter, TSet<UNiagaraEmitter*>& InOutCompiledEmitters)
+{
+	if (InOutCompiledEmitters.Contains(Emitter) == false)
+	{
+		if (Emitter->GetParent() != nullptr)
+		{
+			// If the emitter has a parent emitter make sure to compile that one first.
+			CompileEmitterStandAlone(Emitter->GetParent(), InOutCompiledEmitters);
+
+			if (Emitter->IsSynchronizedWithParent() == false)
+			{
+				// If compiling the parent caused it to become out of sync with the current emitter merge in changes before compiling.
+				Emitter->MergeChangesFromParent();
+			}
+		}
+
+		Emitter->MarkPackageDirty();
+		UNiagaraSystem* TransientSystem = NewObject<UNiagaraSystem>(GetTransientPackage(), NAME_None, RF_Transient);
+		UNiagaraSystemFactoryNew::InitializeSystem(TransientSystem, true);
+		TransientSystem->AddEmitterHandle(*Emitter, TEXT("Emitter"));
+		FNiagaraStackGraphUtilities::RebuildEmitterNodes(*TransientSystem);
+		TransientSystem->RequestCompile(false);
+		TransientSystem->WaitForCompilationComplete();
+
+		InOutCompiledEmitters.Add(Emitter);
+	}
+}
+
 void PreventSystemRecompile(FAssetData SystemAsset, TSet<UNiagaraEmitter*>& InOutCompiledEmitters)
 {
 	UNiagaraSystem* System = Cast<UNiagaraSystem>(SystemAsset.GetAsset());
@@ -284,22 +312,7 @@ void PreventSystemRecompile(FAssetData SystemAsset, TSet<UNiagaraEmitter*>& InOu
 	{
 		for (const FNiagaraEmitterHandle& EmitterHandle : System->GetEmitterHandles())
 		{
-			UNiagaraEmitter* SourceEmitter = const_cast<UNiagaraEmitter*>(EmitterHandle.GetSource());
-			if (SourceEmitter != nullptr)
-			{
-				if (InOutCompiledEmitters.Contains(SourceEmitter) == false)
-				{
-					SourceEmitter->MarkPackageDirty();
-					UNiagaraSystem* TransientSystem = NewObject<UNiagaraSystem>(GetTransientPackage(), NAME_None, RF_Transient);
-					UNiagaraSystemFactoryNew::InitializeSystem(TransientSystem, true);
-					TransientSystem->AddEmitterHandle(*SourceEmitter, TEXT("Emitter")); // TODO Frank.Fella, resolve this properly...
-					FNiagaraStackGraphUtilities::RebuildEmitterNodes(*TransientSystem);
-					TransientSystem->RequestCompile(false);
-					TransientSystem->WaitForCompilationComplete();
-					InOutCompiledEmitters.Add(SourceEmitter);
-				}
-				System->UpdateFromEmitterChanges(*SourceEmitter, false);
-			}
+			CompileEmitterStandAlone(EmitterHandle.GetInstance(), InOutCompiledEmitters);
 		}
 
 		System->MarkPackageDirty();
@@ -544,7 +557,7 @@ void FNiagaraEditorModule::StartupModule()
 
 	// Register the emitter merge handler.
 	ScriptMergeManager = MakeShared<FNiagaraScriptMergeManager>();
-	MergeEmitterHandle = NiagaraModule.RegisterOnMergeEmitter(INiagaraModule::FOnMergeEmitter::CreateSP(ScriptMergeManager.ToSharedRef(), &FNiagaraScriptMergeManager::MergeEmitter));
+	NiagaraModule.RegisterMergeManager(ScriptMergeManager.ToSharedRef());
 
 	// Register the script compiler
 	ScriptCompilerHandle = NiagaraModule.RegisterScriptCompiler(INiagaraModule::FScriptCompiler::CreateLambda([this](const FNiagaraCompileRequestDataBase* CompileRequest, const FNiagaraCompileOptions& Options)
@@ -645,7 +658,7 @@ void FNiagaraEditorModule::ShutdownModule()
 	INiagaraModule* NiagaraModule = FModuleManager::GetModulePtr<INiagaraModule>("Niagara");
 	if (NiagaraModule != nullptr)
 	{
-		NiagaraModule->UnregisterOnMergeEmitter(MergeEmitterHandle);
+		NiagaraModule->UnregisterMergeManager(ScriptMergeManager.ToSharedRef());
 		NiagaraModule->UnregisterOnCreateDefaultScriptSource(CreateDefaultScriptSourceHandle);
 		NiagaraModule->UnregisterScriptCompiler(ScriptCompilerHandle);
 		NiagaraModule->UnregisterPrecompiler(PrecompilerHandle);
diff --git a/Engine/Plugins/FX/Niagara/Source/NiagaraEditor/Private/NiagaraEditorUtilities.cpp b/Engine/Plugins/FX/Niagara/Source/NiagaraEditor/Private/NiagaraEditorUtilities.cpp
--- a/Engine/Plugins/FX/Niagara/Source/NiagaraEditor/Private/NiagaraEditorUtilities.cpp
+++ b/Engine/Plugins/FX/Niagara/Source/NiagaraEditor/Private/NiagaraEditorUtilities.cpp
@@ -1076,7 +1076,7 @@ TArray<UNiagaraComponent*> FNiagaraEditorUtilities::GetComponentsThatReferenceSy
 		{
 			for (auto EmitterHandle : ReferencedSystemViewModel.GetSystem().GetEmitterHandles())
 			{
-				if (Component->GetAsset()->UsesEmitter(EmitterHandle.GetSource()))
+				if (Component->GetAsset()->UsesEmitter(EmitterHandle.GetInstance()->GetParent()))
 				{
 					ReferencingComponents.Add(Component);
 				}
diff --git a/Engine/Plugins/FX/Niagara/Source/NiagaraEditor/Private/NiagaraEmitterFactoryNew.cpp b/Engine/Plugins/FX/Niagara/Source/NiagaraEditor/Private/NiagaraEmitterFactoryNew.cpp
--- a/Engine/Plugins/FX/Niagara/Source/NiagaraEditor/Private/NiagaraEmitterFactoryNew.cpp
+++ b/Engine/Plugins/FX/Niagara/Source/NiagaraEditor/Private/NiagaraEmitterFactoryNew.cpp
@@ -31,6 +31,7 @@ UNiagaraEmitterFactoryNew::UNiagaraEmitterFactoryNew(const FObjectInitializer& O
 	bEditAfterNew = true;
 	bCreateNew = true;
 	EmitterToCopy = nullptr;
+	bUseInheritance = false;
 }
 
 bool UNiagaraEmitterFactoryNew::ConfigureProperties()
@@ -51,6 +52,7 @@ bool UNiagaraEmitterFactoryNew::ConfigureProperties()
 	if (SelectedEmitterAsset.IsSet())
 	{
 		EmitterToCopy = Cast<UNiagaraEmitter>(SelectedEmitterAsset->GetAsset());
+		bUseInheritance = NewEmitterDialog->GetUseInheritance();
 		if (EmitterToCopy == nullptr)
 		{
 			FText Title = LOCTEXT("FailedToLoadTitle", "Create Default?");
@@ -120,7 +122,16 @@ UObject* UNiagaraEmitterFactoryNew::FactoryCreateNew(UClass* Class, UObject* InP
 
 	if (EmitterToCopy != nullptr)
 	{
-		NewEmitter = Cast<UNiagaraEmitter>(StaticDuplicateObject(EmitterToCopy, InParent, Name, Flags, Class));
+		if (bUseInheritance && EmitterToCopy->bIsTemplateAsset == false)
+		{
+			NewEmitter = UNiagaraEmitter::CreateWithParentAndOwner(*EmitterToCopy, InParent, Name, Flags);
+		}
+		else
+		{
+			NewEmitter = Cast<UNiagaraEmitter>(StaticDuplicateObject(EmitterToCopy, InParent, Name, Flags, Class));
+			NewEmitter->SetUniqueEmitterName(Name.ToString());
+		}
+
 		NewEmitter->bIsTemplateAsset = false;
 		NewEmitter->TemplateAssetDescription = FText();
 	}
diff --git a/Engine/Plugins/FX/Niagara/Source/NiagaraEditor/Private/NiagaraScriptMergeManager.cpp b/Engine/Plugins/FX/Niagara/Source/NiagaraEditor/Private/NiagaraScriptMergeManager.cpp
--- a/Engine/Plugins/FX/Niagara/Source/NiagaraEditor/Private/NiagaraScriptMergeManager.cpp
+++ b/Engine/Plugins/FX/Niagara/Source/NiagaraEditor/Private/NiagaraScriptMergeManager.cpp
@@ -656,7 +656,7 @@ TSharedRef<FNiagaraScriptMergeManager> FNiagaraScriptMergeManager::Get()
 	return NiagaraEditorModule.GetScriptMergeManager();
 }
 
-void FNiagaraScriptMergeManager::DiffChangeIds(const TMap<FGuid, FGuid>& InSourceChangeIds, const TMap<FGuid, FGuid>& InLastMergedChangeIds, const TMap<FGuid, FGuid>& InInstanceChangeIds, TMap<FGuid, FGuid>& OutChangeIdsToKeepOnInstance)
+void FNiagaraScriptMergeManager::DiffChangeIds(const TMap<FGuid, FGuid>& InSourceChangeIds, const TMap<FGuid, FGuid>& InLastMergedChangeIds, const TMap<FGuid, FGuid>& InInstanceChangeIds, TMap<FGuid, FGuid>& OutChangeIdsToKeepOnInstance) const
 {
 	auto It = InInstanceChangeIds.CreateConstIterator();
 	while (It)
@@ -698,7 +698,7 @@ void FNiagaraScriptMergeManager::DiffChangeIds(const TMap<FGuid, FGuid>& InSourc
 	}
 }
 
-FNiagaraScriptMergeManager::FApplyDiffResults FNiagaraScriptMergeManager::ResolveChangeIds(TSharedRef<FNiagaraEmitterMergeAdapter> MergedInstanceAdapter, UNiagaraEmitter& OriginalEmitterInstance, const TMap<FGuid, FGuid>& ChangeIdsThatNeedToBeReset)
+FNiagaraScriptMergeManager::FApplyDiffResults FNiagaraScriptMergeManager::ResolveChangeIds(TSharedRef<FNiagaraEmitterMergeAdapter> MergedInstanceAdapter, UNiagaraEmitter& OriginalEmitterInstance, const TMap<FGuid, FGuid>& ChangeIdsThatNeedToBeReset) const
 {
 	FApplyDiffResults DiffResults;
 
@@ -807,24 +807,24 @@ FNiagaraScriptMergeManager::FApplyDiffResults FNiagaraScriptMergeManager::Resolv
 	return DiffResults;
 }
 
-INiagaraModule::FMergeEmitterResults FNiagaraScriptMergeManager::MergeEmitter(UNiagaraEmitter& Source, UNiagaraEmitter& LastMergedSource, UNiagaraEmitter& Instance)
+INiagaraMergeManager::FMergeEmitterResults FNiagaraScriptMergeManager::MergeEmitter(UNiagaraEmitter& Parent, UNiagaraEmitter& ParentAtLastMerge, UNiagaraEmitter& Instance) const
 {
 	SCOPE_CYCLE_COUNTER(STAT_NiagaraEditor_ScriptMergeManager_MergeEmitter);
-	INiagaraModule::FMergeEmitterResults MergeResults;
+	INiagaraMergeManager::FMergeEmitterResults MergeResults;
 
-	FNiagaraEmitterDiffResults DiffResults = DiffEmitters(LastMergedSource, Instance);
+	FNiagaraEmitterDiffResults DiffResults = DiffEmitters(ParentAtLastMerge, Instance);
 
 	if (DiffResults.IsValid())
 	{
-		UNiagaraEmitter* MergedInstance = CastChecked<UNiagaraEmitter>(StaticDuplicateObject(&Source, (UObject*)GetTransientPackage()));
+		UNiagaraEmitter* MergedInstance = CastChecked<UNiagaraEmitter>(StaticDuplicateObject(&Parent, (UObject*)GetTransientPackage()));
 		TSharedRef<FNiagaraEmitterMergeAdapter> MergedInstanceAdapter = MakeShared<FNiagaraEmitterMergeAdapter>(*MergedInstance);
 
 		TMap<FGuid, FGuid> SourceChangeIds;
 		TMap<FGuid, FGuid> PreviousSourceChangeIds;
 		TMap<FGuid, FGuid> LastChangeIds;
 		TMap<FGuid, FGuid> ChangeIdsThatNeedToBeReset;
-		FNiagaraEditorUtilities::GatherChangeIds(Source, SourceChangeIds, TEXT("Source"));
-		FNiagaraEditorUtilities::GatherChangeIds(LastMergedSource, PreviousSourceChangeIds, TEXT("MergeLast"));
+		FNiagaraEditorUtilities::GatherChangeIds(Parent, SourceChangeIds, TEXT("Source"));
+		FNiagaraEditorUtilities::GatherChangeIds(ParentAtLastMerge, PreviousSourceChangeIds, TEXT("MergeLast"));
 		FNiagaraEditorUtilities::GatherChangeIds(Instance, LastChangeIds, TEXT("Instance"));
 		DiffChangeIds(SourceChangeIds, PreviousSourceChangeIds, LastChangeIds, ChangeIdsThatNeedToBeReset);
 
@@ -1154,7 +1154,7 @@ void FNiagaraScriptMergeManager::ResetEmitterEditablePropertySetToBase(UNiagaraE
 	Emitter.PostEditChange();
 }
 
-FNiagaraEmitterDiffResults FNiagaraScriptMergeManager::DiffEmitters(UNiagaraEmitter& BaseEmitter, UNiagaraEmitter& OtherEmitter)
+FNiagaraEmitterDiffResults FNiagaraScriptMergeManager::DiffEmitters(UNiagaraEmitter& BaseEmitter, UNiagaraEmitter& OtherEmitter) const
 {
 	SCOPE_CYCLE_COUNTER(STAT_NiagaraEditor_ScriptMergeManager_DiffEmitters);
 
@@ -1283,7 +1283,7 @@ FListDiffResults<ValueType> DiffLists(const TArray<ValueType>& BaseList, const T
 	return DiffResults;
 }
 
-void FNiagaraScriptMergeManager::DiffEventHandlers(const TArray<TSharedRef<FNiagaraEventHandlerMergeAdapter>>& BaseEventHandlers, const TArray<TSharedRef<FNiagaraEventHandlerMergeAdapter>>& OtherEventHandlers, FNiagaraEmitterDiffResults& DiffResults)
+void FNiagaraScriptMergeManager::DiffEventHandlers(const TArray<TSharedRef<FNiagaraEventHandlerMergeAdapter>>& BaseEventHandlers, const TArray<TSharedRef<FNiagaraEventHandlerMergeAdapter>>& OtherEventHandlers, FNiagaraEmitterDiffResults& DiffResults) const
 {
 	FListDiffResults<TSharedRef<FNiagaraEventHandlerMergeAdapter>> EventHandlerListDiffResults = DiffLists<TSharedRef<FNiagaraEventHandlerMergeAdapter>, FGuid>(
 		BaseEventHandlers,
@@ -1334,7 +1334,7 @@ void FNiagaraScriptMergeManager::DiffEventHandlers(const TArray<TSharedRef<FNiag
 	}
 }
 
-void FNiagaraScriptMergeManager::DiffRenderers(const TArray<TSharedRef<FNiagaraRendererMergeAdapter>>& BaseRenderers, const TArray<TSharedRef<FNiagaraRendererMergeAdapter>>& OtherRenderers, FNiagaraEmitterDiffResults& DiffResults)
+void FNiagaraScriptMergeManager::DiffRenderers(const TArray<TSharedRef<FNiagaraRendererMergeAdapter>>& BaseRenderers, const TArray<TSharedRef<FNiagaraRendererMergeAdapter>>& OtherRenderers, FNiagaraEmitterDiffResults& DiffResults) const
 {
 	FListDiffResults<TSharedRef<FNiagaraRendererMergeAdapter>> RendererListDiffResults = DiffLists<TSharedRef<FNiagaraRendererMergeAdapter>, FGuid>(
 		BaseRenderers,
@@ -1354,7 +1354,7 @@ void FNiagaraScriptMergeManager::DiffRenderers(const TArray<TSharedRef<FNiagaraR
 	}
 }
 
-void FNiagaraScriptMergeManager::DiffScriptStacks(TSharedRef<FNiagaraScriptStackMergeAdapter> BaseScriptStackAdapter, TSharedRef<FNiagaraScriptStackMergeAdapter> OtherScriptStackAdapter, FNiagaraScriptStackDiffResults& DiffResults)
+void FNiagaraScriptMergeManager::DiffScriptStacks(TSharedRef<FNiagaraScriptStackMergeAdapter> BaseScriptStackAdapter, TSharedRef<FNiagaraScriptStackMergeAdapter> OtherScriptStackAdapter, FNiagaraScriptStackDiffResults& DiffResults) const
 {
 	// Diff the module lists.
 	FListDiffResults<TSharedRef<FNiagaraStackFunctionMergeAdapter>> ModuleListDiffResults = DiffLists<TSharedRef<FNiagaraStackFunctionMergeAdapter>, FGuid>(
@@ -1408,7 +1408,7 @@ void FNiagaraScriptMergeManager::DiffScriptStacks(TSharedRef<FNiagaraScriptStack
 	}
 }
 
-void FNiagaraScriptMergeManager::DiffFunctionInputs(TSharedRef<FNiagaraStackFunctionMergeAdapter> BaseFunctionAdapter, TSharedRef<FNiagaraStackFunctionMergeAdapter> OtherFunctionAdapter, FNiagaraScriptStackDiffResults& DiffResults)
+void FNiagaraScriptMergeManager::DiffFunctionInputs(TSharedRef<FNiagaraStackFunctionMergeAdapter> BaseFunctionAdapter, TSharedRef<FNiagaraStackFunctionMergeAdapter> OtherFunctionAdapter, FNiagaraScriptStackDiffResults& DiffResults) const
 {
 	FListDiffResults<TSharedRef<FNiagaraStackFunctionInputOverrideMergeAdapter>> ListDiffResults = DiffLists<TSharedRef<FNiagaraStackFunctionInputOverrideMergeAdapter>, FString>(
 		BaseFunctionAdapter->GetInputOverrides(),
@@ -1438,7 +1438,7 @@ void FNiagaraScriptMergeManager::DiffFunctionInputs(TSharedRef<FNiagaraStackFunc
 	}
 }
 
-void FNiagaraScriptMergeManager::DiffEditableProperties(const void* BaseDataAddress, const void* OtherDataAddress, UStruct& Struct, TArray<UProperty*>& OutDifferentProperties)
+void FNiagaraScriptMergeManager::DiffEditableProperties(const void* BaseDataAddress, const void* OtherDataAddress, UStruct& Struct, TArray<UProperty*>& OutDifferentProperties) const
 {
 	for (TFieldIterator<UProperty> PropertyIterator(&Struct); PropertyIterator; ++PropertyIterator)
 	{
@@ -1454,7 +1454,7 @@ void FNiagaraScriptMergeManager::DiffEditableProperties(const void* BaseDataAddr
 	}
 }
 
-TOptional<bool> FNiagaraScriptMergeManager::DoFunctionInputOverridesMatch(TSharedRef<FNiagaraStackFunctionInputOverrideMergeAdapter> BaseFunctionInputAdapter, TSharedRef<FNiagaraStackFunctionInputOverrideMergeAdapter> OtherFunctionInputAdapter)
+TOptional<bool> FNiagaraScriptMergeManager::DoFunctionInputOverridesMatch(TSharedRef<FNiagaraStackFunctionInputOverrideMergeAdapter> BaseFunctionInputAdapter, TSharedRef<FNiagaraStackFunctionInputOverrideMergeAdapter> OtherFunctionInputAdapter) const
 {
 	// Local String Value.
 	if ((BaseFunctionInputAdapter->GetLocalValueString().IsSet() && OtherFunctionInputAdapter->GetLocalValueString().IsSet() == false) ||
@@ -1546,7 +1546,7 @@ TOptional<bool> FNiagaraScriptMergeManager::DoFunctionInputOverridesMatch(TShare
 	return TOptional<bool>();
 }
 
-FNiagaraScriptMergeManager::FApplyDiffResults FNiagaraScriptMergeManager::AddModule(FString UniqueEmitterName, UNiagaraScript& OwningScript, UNiagaraNodeOutput& TargetOutputNode, TSharedRef<FNiagaraStackFunctionMergeAdapter> AddModule)
+FNiagaraScriptMergeManager::FApplyDiffResults FNiagaraScriptMergeManager::AddModule(FString UniqueEmitterName, UNiagaraScript& OwningScript, UNiagaraNodeOutput& TargetOutputNode, TSharedRef<FNiagaraStackFunctionMergeAdapter> AddModule) const
 {
 	FApplyDiffResults Results;
 
@@ -1600,7 +1600,7 @@ FNiagaraScriptMergeManager::FApplyDiffResults FNiagaraScriptMergeManager::AddMod
 	return Results;
 }
 
-FNiagaraScriptMergeManager::FApplyDiffResults FNiagaraScriptMergeManager::RemoveInputOverride(UNiagaraScript& OwningScript, TSharedRef<FNiagaraStackFunctionInputOverrideMergeAdapter> OverrideToRemove)
+FNiagaraScriptMergeManager::FApplyDiffResults FNiagaraScriptMergeManager::RemoveInputOverride(UNiagaraScript& OwningScript, TSharedRef<FNiagaraStackFunctionInputOverrideMergeAdapter> OverrideToRemove) const
 {
 	FApplyDiffResults Results;
 	if (OverrideToRemove->GetOverridePin() != nullptr && OverrideToRemove->GetOverrideNode() != nullptr)
@@ -1626,7 +1626,7 @@ FNiagaraScriptMergeManager::FApplyDiffResults FNiagaraScriptMergeManager::Remove
 	return Results;
 }
 
-FNiagaraScriptMergeManager::FApplyDiffResults FNiagaraScriptMergeManager::AddInputOverride(FString UniqueEmitterName, UNiagaraScript& OwningScript, UNiagaraNodeFunctionCall& TargetFunctionCall, TSharedRef<FNiagaraStackFunctionInputOverrideMergeAdapter> OverrideToAdd)
+FNiagaraScriptMergeManager::FApplyDiffResults FNiagaraScriptMergeManager::AddInputOverride(FString UniqueEmitterName, UNiagaraScript& OwningScript, UNiagaraNodeFunctionCall& TargetFunctionCall, TSharedRef<FNiagaraStackFunctionInputOverrideMergeAdapter> OverrideToAdd) const
 {
 	FApplyDiffResults Results;
 
@@ -1758,7 +1758,7 @@ FNiagaraScriptMergeManager::FApplyDiffResults FNiagaraScriptMergeManager::AddInp
 	return Results;
 }
 
-void FNiagaraScriptMergeManager::CopyPropertiesToBase(void* BaseDataAddress, const void* OtherDataAddress, TArray<UProperty*> PropertiesToCopy)
+void FNiagaraScriptMergeManager::CopyPropertiesToBase(void* BaseDataAddress, const void* OtherDataAddress, TArray<UProperty*> PropertiesToCopy) const
 {
 	for (UProperty* PropertyToCopy : PropertiesToCopy)
 	{
@@ -1766,7 +1766,7 @@ void FNiagaraScriptMergeManager::CopyPropertiesToBase(void* BaseDataAddress, con
 	}
 }
 
-FNiagaraScriptMergeManager::FApplyDiffResults FNiagaraScriptMergeManager::ApplyScriptStackDiff(TSharedRef<FNiagaraScriptStackMergeAdapter> BaseScriptStackAdapter, const FNiagaraScriptStackDiffResults& DiffResults)
+FNiagaraScriptMergeManager::FApplyDiffResults FNiagaraScriptMergeManager::ApplyScriptStackDiff(TSharedRef<FNiagaraScriptStackMergeAdapter> BaseScriptStackAdapter, const FNiagaraScriptStackDiffResults& DiffResults) const
 {
 	FApplyDiffResults Results;
 
@@ -1926,7 +1926,7 @@ FNiagaraScriptMergeManager::FApplyDiffResults FNiagaraScriptMergeManager::ApplyS
 	return Results;
 }
 
-FNiagaraScriptMergeManager::FApplyDiffResults FNiagaraScriptMergeManager::ApplyEventHandlerDiff(TSharedRef<FNiagaraEmitterMergeAdapter> BaseEmitterAdapter, const FNiagaraEmitterDiffResults& DiffResults)
+FNiagaraScriptMergeManager::FApplyDiffResults FNiagaraScriptMergeManager::ApplyEventHandlerDiff(TSharedRef<FNiagaraEmitterMergeAdapter> BaseEmitterAdapter, const FNiagaraEmitterDiffResults& DiffResults) const
 {
 	FApplyDiffResults Results;
 	if (DiffResults.RemovedBaseEventHandlers.Num() > 0)
@@ -2018,7 +2018,7 @@ FNiagaraScriptMergeManager::FApplyDiffResults FNiagaraScriptMergeManager::ApplyE
 	return Results;
 }
 
-FNiagaraScriptMergeManager::FApplyDiffResults FNiagaraScriptMergeManager::ApplyRendererDiff(UNiagaraEmitter& BaseEmitter, const FNiagaraEmitterDiffResults& DiffResults)
+FNiagaraScriptMergeManager::FApplyDiffResults FNiagaraScriptMergeManager::ApplyRendererDiff(UNiagaraEmitter& BaseEmitter, const FNiagaraEmitterDiffResults& DiffResults) const
 {
 	TArray<UNiagaraRendererProperties*> RenderersToRemove;
 	TArray<UNiagaraRendererProperties*> RenderersToAdd;
diff --git a/Engine/Plugins/FX/Niagara/Source/NiagaraEditor/Private/NiagaraScriptMergeManager.h b/Engine/Plugins/FX/Niagara/Source/NiagaraEditor/Private/NiagaraScriptMergeManager.h
--- a/Engine/Plugins/FX/Niagara/Source/NiagaraEditor/Private/NiagaraScriptMergeManager.h
+++ b/Engine/Plugins/FX/Niagara/Source/NiagaraEditor/Private/NiagaraScriptMergeManager.h
@@ -4,6 +4,7 @@
 
 #include "ViewModels/Stack/NiagaraParameterHandle.h"
 #include "NiagaraModule.h"
+#include "INiagaraMergeManager.h"
 #include "NiagaraTypes.h"
 #include "NiagaraCommon.h"
 
@@ -272,7 +273,7 @@ private:
 	TArray<FText> ErrorMessages;
 };
 
-class FNiagaraScriptMergeManager
+class FNiagaraScriptMergeManager : public INiagaraMergeManager
 {
 public:
 	struct FApplyDiffResults
@@ -289,11 +290,11 @@ public:
 	};
 
 public:
-	INiagaraModule::FMergeEmitterResults MergeEmitter(UNiagaraEmitter& Source, UNiagaraEmitter& LastMergedSource, UNiagaraEmitter& Instance);
+	virtual INiagaraMergeManager::FMergeEmitterResults MergeEmitter(UNiagaraEmitter& Parent, UNiagaraEmitter& ParentAtLastMerge, UNiagaraEmitter& Instance) const override;
 
 	static TSharedRef<FNiagaraScriptMergeManager> Get();
 
-	FNiagaraEmitterDiffResults DiffEmitters(UNiagaraEmitter& BaseEmitter, UNiagaraEmitter& OtherEmitter);
+	FNiagaraEmitterDiffResults DiffEmitters(UNiagaraEmitter& BaseEmitter, UNiagaraEmitter& OtherEmitter) const;
 
 	bool IsMergeableScriptUsage(ENiagaraScriptUsage ScriptUsage) const;
 
@@ -319,40 +320,40 @@ public:
 
 	void ResetEmitterEditablePropertySetToBase(UNiagaraEmitter& Emitter, const UNiagaraEmitter& BaseEmitter);
 
-	void DiffEventHandlers(const TArray<TSharedRef<FNiagaraEventHandlerMergeAdapter>>& BaseEventHandlers, const TArray<TSharedRef<FNiagaraEventHandlerMergeAdapter>>& OtherEventHandlers, FNiagaraEmitterDiffResults& DiffResults);
+	void DiffEventHandlers(const TArray<TSharedRef<FNiagaraEventHandlerMergeAdapter>>& BaseEventHandlers, const TArray<TSharedRef<FNiagaraEventHandlerMergeAdapter>>& OtherEventHandlers, FNiagaraEmitterDiffResults& DiffResults) const;
 
-	void DiffRenderers(const TArray<TSharedRef<FNiagaraRendererMergeAdapter>>& BaseRenderers, const TArray<TSharedRef<FNiagaraRendererMergeAdapter>>& OtherRenderers, FNiagaraEmitterDiffResults& DiffResults);
+	void DiffRenderers(const TArray<TSharedRef<FNiagaraRendererMergeAdapter>>& BaseRenderers, const TArray<TSharedRef<FNiagaraRendererMergeAdapter>>& OtherRenderers, FNiagaraEmitterDiffResults& DiffResults) const;
 
-	void DiffScriptStacks(TSharedRef<FNiagaraScriptStackMergeAdapter> BaseScriptStackAdapter, TSharedRef<FNiagaraScriptStackMergeAdapter> OtherScriptStackAdapter, FNiagaraScriptStackDiffResults& DiffResults);
+	void DiffScriptStacks(TSharedRef<FNiagaraScriptStackMergeAdapter> BaseScriptStackAdapter, TSharedRef<FNiagaraScriptStackMergeAdapter> OtherScriptStackAdapter, FNiagaraScriptStackDiffResults& DiffResults) const;
 
-	void DiffFunctionInputs(TSharedRef<FNiagaraStackFunctionMergeAdapter> BaseFunctionAdapter, TSharedRef<FNiagaraStackFunctionMergeAdapter> OtherFunctionAdapter, FNiagaraScriptStackDiffResults& DiffResults);
+	void DiffFunctionInputs(TSharedRef<FNiagaraStackFunctionMergeAdapter> BaseFunctionAdapter, TSharedRef<FNiagaraStackFunctionMergeAdapter> OtherFunctionAdapter, FNiagaraScriptStackDiffResults& DiffResults) const;
 
-	void DiffEditableProperties(const void* BaseDataAddress, const void* OtherDataAddress, UStruct& Struct, TArray<UProperty*>& OutDifferentProperties);
+	virtual void DiffEditableProperties(const void* BaseDataAddress, const void* OtherDataAddress, UStruct& Struct, TArray<UProperty*>& OutDifferentProperties) const override;
 
-private:
-	TOptional<bool> DoFunctionInputOverridesMatch(TSharedRef<FNiagaraStackFunctionInputOverrideMergeAdapter> BaseFunctionInputAdapter, TSharedRef<FNiagaraStackFunctionInputOverrideMergeAdapter> OtherFunctionInputAdapter);
+	virtual void CopyPropertiesToBase(void* BaseDataAddress, const void* OtherDataAddress, TArray<UProperty*> PropertiesToCopy) const override;
 
-	FApplyDiffResults ApplyScriptStackDiff(TSharedRef<FNiagaraScriptStackMergeAdapter> BaseScriptStackAdapter, const FNiagaraScriptStackDiffResults& DiffResults);
+private:
+	TOptional<bool> DoFunctionInputOverridesMatch(TSharedRef<FNiagaraStackFunctionInputOverrideMergeAdapter> BaseFunctionInputAdapter, TSharedRef<FNiagaraStackFunctionInputOverrideMergeAdapter> OtherFunctionInputAdapter) const;
 
-	FApplyDiffResults ApplyEventHandlerDiff(TSharedRef<FNiagaraEmitterMergeAdapter> BaseEmitterAdapter, const FNiagaraEmitterDiffResults& DiffResults);
+	FApplyDiffResults ApplyScriptStackDiff(TSharedRef<FNiagaraScriptStackMergeAdapter> BaseScriptStackAdapter, const FNiagaraScriptStackDiffResults& DiffResults) const;
 
-	FApplyDiffResults ApplyRendererDiff(UNiagaraEmitter& BaseEmitter, const FNiagaraEmitterDiffResults& DiffResults);
+	FApplyDiffResults ApplyEventHandlerDiff(TSharedRef<FNiagaraEmitterMergeAdapter> BaseEmitterAdapter, const FNiagaraEmitterDiffResults& DiffResults) const;
 
-	FApplyDiffResults AddModule(FString UniqueEmitterName, UNiagaraScript& OwningScript, UNiagaraNodeOutput& TargetOutputNode, TSharedRef<FNiagaraStackFunctionMergeAdapter> AddModule);
+	FApplyDiffResults ApplyRendererDiff(UNiagaraEmitter& BaseEmitter, const FNiagaraEmitterDiffResults& DiffResults) const;
 
-	FApplyDiffResults RemoveInputOverride(UNiagaraScript& OwningScript, TSharedRef<FNiagaraStackFunctionInputOverrideMergeAdapter> OverrideToRemove);
+	FApplyDiffResults AddModule(FString UniqueEmitterName, UNiagaraScript& OwningScript, UNiagaraNodeOutput& TargetOutputNode, TSharedRef<FNiagaraStackFunctionMergeAdapter> AddModule) const;
 
-	FApplyDiffResults AddInputOverride(FString UniqueEmitterName, UNiagaraScript& OwningScript, UNiagaraNodeFunctionCall& TargetFunctionCall, TSharedRef<FNiagaraStackFunctionInputOverrideMergeAdapter> OverrideToAdd);
+	FApplyDiffResults RemoveInputOverride(UNiagaraScript& OwningScript, TSharedRef<FNiagaraStackFunctionInputOverrideMergeAdapter> OverrideToRemove) const;
 
-	void CopyPropertiesToBase(void* BaseDataAddress, const void* OtherDataAddress, TArray<UProperty*> PropertiesToCopy);
+	FApplyDiffResults AddInputOverride(FString UniqueEmitterName, UNiagaraScript& OwningScript, UNiagaraNodeFunctionCall& TargetFunctionCall, TSharedRef<FNiagaraStackFunctionInputOverrideMergeAdapter> OverrideToAdd) const;
 
 private:
 	TSharedRef<FNiagaraEmitterMergeAdapter> GetEmitterMergeAdapterUsingCache(const UNiagaraEmitter& Emitter);
 
 	TSharedRef<FNiagaraEmitterMergeAdapter> GetEmitterMergeAdapterUsingCache(UNiagaraEmitter& Emitter);
 
-	void DiffChangeIds(const TMap<FGuid, FGuid>& InSourceChangeIds, const TMap<FGuid, FGuid>& InLastMergedChangeIds, const TMap<FGuid, FGuid>& InInstanceChangeIds, TMap<FGuid, FGuid>& OutChangeIdsToKeepOnInstance);
-	FApplyDiffResults ResolveChangeIds(TSharedRef<FNiagaraEmitterMergeAdapter> MergedInstanceAdapter, UNiagaraEmitter& OriginalEmitterInstance, const TMap<FGuid, FGuid>& ChangeIdsThatNeedToBeReset);
+	void DiffChangeIds(const TMap<FGuid, FGuid>& InSourceChangeIds, const TMap<FGuid, FGuid>& InLastMergedChangeIds, const TMap<FGuid, FGuid>& InInstanceChangeIds, TMap<FGuid, FGuid>& OutChangeIdsToKeepOnInstance) const;
+	FApplyDiffResults ResolveChangeIds(TSharedRef<FNiagaraEmitterMergeAdapter> MergedInstanceAdapter, UNiagaraEmitter& OriginalEmitterInstance, const TMap<FGuid, FGuid>& ChangeIdsThatNeedToBeReset) const;
 
 
 private:
diff --git a/Engine/Plugins/FX/Niagara/Source/NiagaraEditor/Private/SNewEmitterDialog.cpp b/Engine/Plugins/FX/Niagara/Source/NiagaraEditor/Private/SNewEmitterDialog.cpp
--- a/Engine/Plugins/FX/Niagara/Source/NiagaraEditor/Private/SNewEmitterDialog.cpp
+++ b/Engine/Plugins/FX/Niagara/Source/NiagaraEditor/Private/SNewEmitterDialog.cpp
@@ -33,17 +33,26 @@ void SNewEmitterDialog::Construct(const FArguments& InArgs)
 				LOCTEXT("CreateFromTemplateLabel", "Create a new emitter from an emitter template"),
 				LOCTEXT("TemplatesPickerHeader", "Select a Template Emitter"),
 				SNiagaraNewAssetDialog::FOnGetSelectedAssetsFromPicker::CreateSP(this, &SNewEmitterDialog::GetSelectedEmitterTemplateAssets),
+				SNiagaraNewAssetDialog::FOnSelectionConfirmed(),
 				SAssignNew(TemplateAssetPicker, SNiagaraTemplateAssetPicker, UNiagaraEmitter::StaticClass())
 				.OnTemplateAssetActivated(this, &SNewEmitterDialog::OnTemplateAssetActivated)),
 			SNiagaraNewAssetDialog::FNiagaraNewAssetDialogOption(
 				LOCTEXT("CreateFromOtherEmitterLabel", "Copy an existing emitter from your project content"),
 				LOCTEXT("ProjectEmitterPickerHeader", "Select a Project Emitter"),
 				SNiagaraNewAssetDialog::FOnGetSelectedAssetsFromPicker::CreateSP(this, &SNewEmitterDialog::GetSelectedProjectEmiterAssets),
+				SNiagaraNewAssetDialog::FOnSelectionConfirmed(),
+				AssetPicker),
+			SNiagaraNewAssetDialog::FNiagaraNewAssetDialogOption(
+				LOCTEXT("InheritFromOtherEmitterLabel", "Inherit from an existing emitter in your project content"),
+				LOCTEXT("InheritProjectEmitterPickerHeader", "Select a Parent Project Emitter"),
+				SNiagaraNewAssetDialog::FOnGetSelectedAssetsFromPicker::CreateSP(this, &SNewEmitterDialog::GetSelectedProjectEmiterAssets),
+				SNiagaraNewAssetDialog::FOnSelectionConfirmed::CreateSP(this, &SNewEmitterDialog::InheritanceOptionConfirmed),
 				AssetPicker),
 			SNiagaraNewAssetDialog::FNiagaraNewAssetDialogOption(
 				LOCTEXT("CreateEmptyLabel", "Create an empty emitter with no modules or renderers (Advanced)"),
 				LOCTEXT("EmptyLabel", "Empty Emitter"),
 				SNiagaraNewAssetDialog::FOnGetSelectedAssetsFromPicker(),
+				SNiagaraNewAssetDialog::FOnSelectionConfirmed(),
 				SNew(SBox)
 				.HAlign(HAlign_Center)
 				.VAlign(VAlign_Center)
@@ -64,6 +73,11 @@ TOptional<FAssetData> SNewEmitterDialog::GetSelectedEmitterAsset()
 	return TOptional<FAssetData>();
 }
 
+bool SNewEmitterDialog::GetUseInheritance() const
+{
+	return bUseInheritance;
+}
+
 void SNewEmitterDialog::GetSelectedEmitterTemplateAssets(TArray<FAssetData>& OutSelectedAssets)
 {
 	OutSelectedAssets.Append(TemplateAssetPicker->GetSelectedAssets());
@@ -101,4 +115,9 @@ void SNewEmitterDialog::OnEmitterAssetsActivated(const TArray<FAssetData>& Activ
 	}
 }
 
+void SNewEmitterDialog::InheritanceOptionConfirmed()
+{
+	bUseInheritance = true;
+}
+
 #undef LOCTEXT_NAMESPACE
\ No newline at end of file
diff --git a/Engine/Plugins/FX/Niagara/Source/NiagaraEditor/Private/SNewEmitterDialog.h b/Engine/Plugins/FX/Niagara/Source/NiagaraEditor/Private/SNewEmitterDialog.h
--- a/Engine/Plugins/FX/Niagara/Source/NiagaraEditor/Private/SNewEmitterDialog.h
+++ b/Engine/Plugins/FX/Niagara/Source/NiagaraEditor/Private/SNewEmitterDialog.h
@@ -23,6 +23,8 @@ public:
 
 	TOptional<FAssetData> GetSelectedEmitterAsset();
 
+	bool GetUseInheritance() const;
+
 private:
 	void GetSelectedEmitterTemplateAssets(TArray<FAssetData>& OutSelectedAssets);
 
@@ -32,6 +34,8 @@ private:
 
 	void OnEmitterAssetsActivated(const TArray<FAssetData>& ActivatedAssets, EAssetTypeActivationMethod::Type ActivationMethod);
 
+	void InheritanceOptionConfirmed();
+
 private:
 	TSharedPtr<SNiagaraTemplateAssetPicker> TemplateAssetPicker;
 
@@ -40,4 +44,6 @@ private:
 	FAssetData ActivatedTemplateAsset;
 
 	FAssetData ActivatedProjectAsset;
+
+	bool bUseInheritance;
 };
\ No newline at end of file
diff --git a/Engine/Plugins/FX/Niagara/Source/NiagaraEditor/Private/SNewSystemDialog.cpp b/Engine/Plugins/FX/Niagara/Source/NiagaraEditor/Private/SNewSystemDialog.cpp
--- a/Engine/Plugins/FX/Niagara/Source/NiagaraEditor/Private/SNewSystemDialog.cpp
+++ b/Engine/Plugins/FX/Niagara/Source/NiagaraEditor/Private/SNewSystemDialog.cpp
@@ -48,12 +48,14 @@ void SNewSystemDialog::Construct(const FArguments& InArgs)
 				LOCTEXT("CreateFromTemplateLabel", "Create a new system from a system template"),
 				LOCTEXT("TemplateLabel", "Select a System Template"),
 				SNiagaraNewAssetDialog::FOnGetSelectedAssetsFromPicker::CreateSP(this, &SNewSystemDialog::GetSelectedSystemTemplateAssets),
+				SNiagaraNewAssetDialog::FOnSelectionConfirmed(),
 				SAssignNew(TemplateAssetPicker, SNiagaraTemplateAssetPicker, UNiagaraSystem::StaticClass())
 				.OnTemplateAssetActivated(this, &SNewSystemDialog::OnTemplateAssetActivated)),
 			SNiagaraNewAssetDialog::FNiagaraNewAssetDialogOption(
 				LOCTEXT("CreateFromSelectedEmittersLabel", "Create a new system from a set of selected emitters"),
 				LOCTEXT("ProjectEmittersLabel", "Select Emitters to Add"),
 				SNiagaraNewAssetDialog::FOnGetSelectedAssetsFromPicker::CreateSP(this, &SNewSystemDialog::GetSelectedProjectEmiterAssets),
+				SNiagaraNewAssetDialog::FOnSelectionConfirmed(),
 				SNew(SVerticalBox)
 				+ SVerticalBox::Slot()
 				.Padding(0, 0, 0, 10)
@@ -109,11 +111,13 @@ void SNewSystemDialog::Construct(const FArguments& InArgs)
 				LOCTEXT("CreateFromOtherSystemLabel", "Copy an existing system from your project content"),
 				LOCTEXT("ProjectSystemsLabel", "Select a Project System"),
 				SNiagaraNewAssetDialog::FOnGetSelectedAssetsFromPicker::CreateSP(this, &SNewSystemDialog::GetSelectedProjectSystemAssets),
+				SNiagaraNewAssetDialog::FOnSelectionConfirmed(),
 				SystemAssetPicker),
 			SNiagaraNewAssetDialog::FNiagaraNewAssetDialogOption(
 				LOCTEXT("CreateEmptyLabel", "Create an empty system with no emitters"),
 				LOCTEXT("EmptyLabel", "Empty System"),
 				SNiagaraNewAssetDialog::FOnGetSelectedAssetsFromPicker(),
+				SNiagaraNewAssetDialog::FOnSelectionConfirmed(),
 				SNew(SBox)
 				.HAlign(HAlign_Center)
 				.VAlign(VAlign_Center)
diff --git a/Engine/Plugins/FX/Niagara/Source/NiagaraEditor/Private/SNiagaraNewAssetDialog.cpp b/Engine/Plugins/FX/Niagara/Source/NiagaraEditor/Private/SNiagaraNewAssetDialog.cpp
--- a/Engine/Plugins/FX/Niagara/Source/NiagaraEditor/Private/SNiagaraNewAssetDialog.cpp
+++ b/Engine/Plugins/FX/Niagara/Source/NiagaraEditor/Private/SNiagaraNewAssetDialog.cpp
@@ -197,6 +197,7 @@ void SNiagaraNewAssetDialog::ConfirmSelection()
 		SelectedOption.OnGetSelectedAssetsFromPicker.Execute(SelectedAssets);
 		ensureMsgf(SelectedAssets.Num() > 0, TEXT("No assets selected when dialog was confirmed."));
 	}
+	SelectedOption.OnSelectionConfirmed.ExecuteIfBound();
 	bUserConfirmedSelection = true;
 	RequestDestroyWindow();
 }
diff --git a/Engine/Plugins/FX/Niagara/Source/NiagaraEditor/Private/SNiagaraNewAssetDialog.h b/Engine/Plugins/FX/Niagara/Source/NiagaraEditor/Private/SNiagaraNewAssetDialog.h
--- a/Engine/Plugins/FX/Niagara/Source/NiagaraEditor/Private/SNiagaraNewAssetDialog.h
+++ b/Engine/Plugins/FX/Niagara/Source/NiagaraEditor/Private/SNiagaraNewAssetDialog.h
@@ -14,6 +14,7 @@ class SNiagaraNewAssetDialog : public SWindow
 {
 public:
 	DECLARE_DELEGATE_OneParam(FOnGetSelectedAssetsFromPicker, TArray<FAssetData>& /* OutSelectedAssets */);
+	DECLARE_DELEGATE(FOnSelectionConfirmed);
 
 public:
 	class FNiagaraNewAssetDialogOption
@@ -23,12 +24,14 @@ public:
 		FText AssetPickerHeader;
 		TSharedRef<SWidget> AssetPicker;
 		FOnGetSelectedAssetsFromPicker OnGetSelectedAssetsFromPicker;
+		FOnSelectionConfirmed OnSelectionConfirmed;
 
-		FNiagaraNewAssetDialogOption(FText InOptionText, FText InAssetPickerHeader, FOnGetSelectedAssetsFromPicker InOnGetSelectedAssetsFromPicker, TSharedRef<SWidget> InAssetPicker)
+		FNiagaraNewAssetDialogOption(FText InOptionText, FText InAssetPickerHeader, FOnGetSelectedAssetsFromPicker InOnGetSelectedAssetsFromPicker, FOnSelectionConfirmed InOnSelecitonConfirmed, TSharedRef<SWidget> InAssetPicker)
 			: OptionText(InOptionText)
 			, AssetPickerHeader(InAssetPickerHeader)
 			, AssetPicker(InAssetPicker)
 			, OnGetSelectedAssetsFromPicker(InOnGetSelectedAssetsFromPicker)
+			, OnSelectionConfirmed(InOnSelecitonConfirmed)
 		{
 		}
 	};
diff --git a/Engine/Plugins/FX/Niagara/Source/NiagaraEditor/Private/Toolkits/NiagaraScriptToolkit.cpp b/Engine/Plugins/FX/Niagara/Source/NiagaraEditor/Private/Toolkits/NiagaraScriptToolkit.cpp
--- a/Engine/Plugins/FX/Niagara/Source/NiagaraEditor/Private/Toolkits/NiagaraScriptToolkit.cpp
+++ b/Engine/Plugins/FX/Niagara/Source/NiagaraEditor/Private/Toolkits/NiagaraScriptToolkit.cpp
@@ -685,7 +685,7 @@ void FNiagaraScriptToolkit::UpdateOriginalNiagaraScript()
 				for (int32 i = 0; i < System->GetNumEmitters(); i++)
 				{
 					AffectedEmitters.AddUnique(System->GetEmitterHandle(i).GetInstance());
-					AffectedEmitters.AddUnique((UNiagaraEmitter*)System->GetEmitterHandle(i).GetSource());
+					AffectedEmitters.AddUnique((UNiagaraEmitter*)System->GetEmitterHandle(i).GetInstance()->GetParent());
 				}
 			}
 		}
diff --git a/Engine/Plugins/FX/Niagara/Source/NiagaraEditor/Private/Toolkits/NiagaraSystemToolkit.cpp b/Engine/Plugins/FX/Niagara/Source/NiagaraEditor/Private/Toolkits/NiagaraSystemToolkit.cpp
--- a/Engine/Plugins/FX/Niagara/Source/NiagaraEditor/Private/Toolkits/NiagaraSystemToolkit.cpp
+++ b/Engine/Plugins/FX/Niagara/Source/NiagaraEditor/Private/Toolkits/NiagaraSystemToolkit.cpp
@@ -258,9 +258,6 @@ void FNiagaraSystemToolkit::InitializeWithEmitter(const EToolkitMode::Type Mode,
 	SystemViewModel->SetToolkitCommands(GetToolkitCommands());
 	SystemViewModel->AddEmitter(*Emitter);
 
-	// Always remove the source information from the emitter handle when editing an emitter asset because inheritance is not valid in this case.
-	System->GetEmitterHandle(0).RemoveSource();
-
 	// Adding the emitter to the system has made a copy of it and we set this to the copy's change id here instead of the original emitter's change 
 	// id because the copy's change id may have been updated from the original as part of post load and we use this id to detect if the editable 
 	// emitter has been changed.
@@ -1053,10 +1050,7 @@ void FNiagaraSystemToolkit::UpdateOriginalEmitter()
 		LastSyncedEmitterChangeId = EditableEmitter->GetChangeId();
 		bEmitterThumbnailUpdated = false;
 
-		TArray<UNiagaraEmitter*> AffectedEmitters;
-		AffectedEmitters.Add(Emitter);
 		UpdateExistingEmitters();
-
 		GWarn->EndSlowTask();
 	}
 	else if(bEmitterThumbnailUpdated)
@@ -1068,28 +1062,76 @@ void FNiagaraSystemToolkit::UpdateOriginalEmitter()
 	}
 }
 
+void MergeEmittersRecursively(UNiagaraEmitter* ChangedEmitter, const TMap<UNiagaraEmitter*, TArray<UNiagaraEmitter*>>& EmitterToReferencingEmittersMap, TSet<UNiagaraEmitter*>& OutMergedEmitters)
+{
+	const TArray<UNiagaraEmitter*>* ReferencingEmitters = EmitterToReferencingEmittersMap.Find(ChangedEmitter);
+	if (ReferencingEmitters != nullptr)
+	{
+		for (UNiagaraEmitter* ReferencingEmitter : (*ReferencingEmitters))
+		{
+			if (ReferencingEmitter->IsSynchronizedWithParent() == false)
+			{
+				ReferencingEmitter->MergeChangesFromParent();
+				OutMergedEmitters.Add(ReferencingEmitter);
+				MergeEmittersRecursively(ReferencingEmitter, EmitterToReferencingEmittersMap, OutMergedEmitters);
+			}
+		}
+	}
+}
+
 void FNiagaraSystemToolkit::UpdateExistingEmitters()
 {
+	// Build a tree of references from the currently loaded emitters so that we can efficiently find all emitters that reference the modified emitter.
+	TMap<UNiagaraEmitter*, TArray<UNiagaraEmitter*>> EmitterToReferencingEmittersMap;
+	UNiagaraEmitter* EditableCopy = System->GetEmitterHandles()[0].GetInstance();
+	for (TObjectIterator<UNiagaraEmitter> EmitterIterator; EmitterIterator; ++EmitterIterator)
+	{
+		UNiagaraEmitter* LoadedEmitter = *EmitterIterator;
+		if (LoadedEmitter != EditableCopy && LoadedEmitter->GetParent() != nullptr)
+		{
+			TArray<UNiagaraEmitter*>& ReferencingEmitters = EmitterToReferencingEmittersMap.FindOrAdd(LoadedEmitter->GetParent());
+			ReferencingEmitters.Add(LoadedEmitter);
+		}
+	}
+
+	// Recursively merge emitters by traversing the reference chains.
+	TSet<UNiagaraEmitter*> MergedEmitters;
+	MergeEmittersRecursively(Emitter, EmitterToReferencingEmittersMap, MergedEmitters);
+
+	// find referencing systems, aside from the system being edited by this toolkit and request that they recompile,
+	// also refresh their view models, and reinitialize their components.
 	for (TObjectIterator<UNiagaraSystem> SystemIterator; SystemIterator; ++SystemIterator)
 	{
 		UNiagaraSystem* LoadedSystem = *SystemIterator;
-
 		if (LoadedSystem != System &&
 			LoadedSystem->IsPendingKill() == false && 
-			LoadedSystem->HasAnyFlags(RF_ClassDefaultObject) == false &&
-			LoadedSystem->ReferencesSourceEmitter(*Emitter))
+			LoadedSystem->HasAnyFlags(RF_ClassDefaultObject) == false)
 		{
-			LoadedSystem->UpdateFromEmitterChanges(*Emitter);
-			TArray<TSharedPtr<FNiagaraSystemViewModel>> ReferencingSystemViewModels;
-			FNiagaraSystemViewModel::GetAllViewModelsForObject(LoadedSystem, ReferencingSystemViewModels);
-
-			for (TSharedPtr<FNiagaraSystemViewModel> ReferencingSystemViewModel : ReferencingSystemViewModels)
+			bool bUsesMergedEmitterDirectly = false;
+			for (const FNiagaraEmitterHandle& EmitterHandle : LoadedSystem->GetEmitterHandles())
 			{
-				ReferencingSystemViewModel->RefreshAll();
+				if (MergedEmitters.Contains(EmitterHandle.GetInstance()))
+				{
+					bUsesMergedEmitterDirectly = true;
+					break;
+				}
 			}
 
-			if (ReferencingSystemViewModels.Num() == 0)
+			if (bUsesMergedEmitterDirectly)
 			{
+				// Request that the system recompile.
+				bool bForce = false;
+				LoadedSystem->RequestCompile(bForce);
+
+				// Invalidate any view models.
+				TArray<TSharedPtr<FNiagaraSystemViewModel>> ReferencingSystemViewModels;
+				FNiagaraSystemViewModel::GetAllViewModelsForObject(LoadedSystem, ReferencingSystemViewModels);
+				for (TSharedPtr<FNiagaraSystemViewModel> ReferencingSystemViewModel : ReferencingSystemViewModels)
+				{
+					ReferencingSystemViewModel->RefreshAll();
+				}
+
+				// Reinit any running components
 				for (TObjectIterator<UNiagaraComponent> ComponentIterator; ComponentIterator; ++ComponentIterator)
 				{
 					UNiagaraComponent* Component = *ComponentIterator;
diff --git a/Engine/Plugins/FX/Niagara/Source/NiagaraEditor/Private/ViewModels/NiagaraEmitterHandleViewModel.cpp b/Engine/Plugins/FX/Niagara/Source/NiagaraEditor/Private/ViewModels/NiagaraEmitterHandleViewModel.cpp
--- a/Engine/Plugins/FX/Niagara/Source/NiagaraEditor/Private/ViewModels/NiagaraEmitterHandleViewModel.cpp
+++ b/Engine/Plugins/FX/Niagara/Source/NiagaraEditor/Private/ViewModels/NiagaraEmitterHandleViewModel.cpp
@@ -162,24 +162,6 @@ FText FNiagaraEmitterHandleViewModel::GetNameText() const
 	return FText();
 }
 
-FText FNiagaraEmitterHandleViewModel::GetSourceNameText() const
-{
-	if (EmitterHandle && EmitterHandle->GetSource())
-	{
-		return FText::FromString(EmitterHandle->GetSource()->GetName()); 
-	}
-	return FText();
-}
-
-FText FNiagaraEmitterHandleViewModel::GetSourcePathNameText() const
-{
-	if (EmitterHandle && EmitterHandle->GetSource())
-	{
-		return FText::FromString(EmitterHandle->GetSource()->GetPathName());
-	}
-	return FText();
-}
-
 void FNiagaraEmitterHandleViewModel::OnNameTextComitted(const FText& InText, ETextCommit::Type CommitInfo)
 {
 	SetName(*InText.ToString());
@@ -240,19 +222,6 @@ TSharedPtr<FNiagaraEmitterViewModel> FNiagaraEmitterHandleViewModel::GetEmitterV
 	return EmitterViewModel;
 }
 
-
-void FNiagaraEmitterHandleViewModel::OpenSourceEmitter()
-{
-	if (EmitterHandle)
-	{
-		UNiagaraEmitter* EmitterSource = const_cast<UNiagaraEmitter*>(EmitterHandle->GetSource());
-		if (EmitterSource)
-		{
-			FAssetEditorManager::Get().OpenEditorForAsset(EmitterSource);
-		}
-	}
-}
-
 FNiagaraEmitterHandleViewModel::FOnPropertyChanged& FNiagaraEmitterHandleViewModel::OnPropertyChanged()
 {
 	return OnPropertyChangedDelegate;
diff --git a/Engine/Plugins/FX/Niagara/Source/NiagaraEditor/Private/ViewModels/NiagaraEmitterViewModel.cpp b/Engine/Plugins/FX/Niagara/Source/NiagaraEditor/Private/ViewModels/NiagaraEmitterViewModel.cpp
--- a/Engine/Plugins/FX/Niagara/Source/NiagaraEditor/Private/ViewModels/NiagaraEmitterViewModel.cpp
+++ b/Engine/Plugins/FX/Niagara/Source/NiagaraEditor/Private/ViewModels/NiagaraEmitterViewModel.cpp
@@ -116,6 +116,28 @@ UNiagaraEmitter* FNiagaraEmitterViewModel::GetEmitter()
 	return Emitter.Get();
 }
 
+const UNiagaraEmitter* FNiagaraEmitterViewModel::GetParentEmitter() const
+{
+	return Emitter.IsValid() ? Emitter->GetParent() : nullptr;
+}
+
+FText FNiagaraEmitterViewModel::GetParentNameText() const
+{
+	if (Emitter.IsValid() && Emitter->GetParent() != nullptr)
+	{
+		return FText::FromString(Emitter->GetParent()->GetName());
+	}
+	return FText();
+}
+
+FText FNiagaraEmitterViewModel::GetParentPathNameText() const
+{
+	if (Emitter.IsValid() && Emitter->GetParent() != nullptr)
+	{
+		return FText::FromString(Emitter->GetParent()->GetPathName());
+	}
+	return FText();
+}
 
 FText FNiagaraEmitterViewModel::GetStatsText() const
 {
diff --git a/Engine/Plugins/FX/Niagara/Source/NiagaraEditor/Private/ViewModels/NiagaraSystemViewModel.cpp b/Engine/Plugins/FX/Niagara/Source/NiagaraEditor/Private/ViewModels/NiagaraSystemViewModel.cpp
--- a/Engine/Plugins/FX/Niagara/Source/NiagaraEditor/Private/ViewModels/NiagaraSystemViewModel.cpp
+++ b/Engine/Plugins/FX/Niagara/Source/NiagaraEditor/Private/ViewModels/NiagaraSystemViewModel.cpp
@@ -310,9 +310,17 @@ void FNiagaraSystemViewModel::AddEmitter(UNiagaraEmitter& Emitter)
 	if (EditMode == ENiagaraSystemViewModelEditMode::SystemAsset)
 	{
 		System.Modify();
+		EmitterHandle = System.AddEmitterHandle(Emitter, FNiagaraUtilities::GetUniqueName(Emitter.GetFName(), EmitterHandleNames));
 	}
-	EmitterHandle = System.AddEmitterHandle(Emitter, FNiagaraUtilities::GetUniqueName(Emitter.GetFName(), EmitterHandleNames));
-
+	else if (EditMode == ENiagaraSystemViewModelEditMode::EmitterAsset)
+	{
+		// When editing an emitter asset we add the emitter as a duplicate so that the parent emitter is duplicated, but it's parent emitter
+		// information is maintained.
+		checkf(System.GetNumEmitters() == 0, TEXT("Can not add multiple emitters to a system being edited in emitter asset mode."));
+		FNiagaraEmitterHandle TemporaryEmitterHandle(Emitter);
+		EmitterHandle = System.DuplicateEmitterHandle(TemporaryEmitterHandle, *Emitter.GetUniqueEmitterName());
+	}
+	
 	check(SystemScriptViewModel.IsValid());
 	FNiagaraStackGraphUtilities::RebuildEmitterNodes(System);
 
diff --git a/Engine/Plugins/FX/Niagara/Source/NiagaraEditor/Private/ViewModels/Stack/NiagaraStackEmitterSpawnScriptItemGroup.cpp b/Engine/Plugins/FX/Niagara/Source/NiagaraEditor/Private/ViewModels/Stack/NiagaraStackEmitterSpawnScriptItemGroup.cpp
--- a/Engine/Plugins/FX/Niagara/Source/NiagaraEditor/Private/ViewModels/Stack/NiagaraStackEmitterSpawnScriptItemGroup.cpp
+++ b/Engine/Plugins/FX/Niagara/Source/NiagaraEditor/Private/ViewModels/Stack/NiagaraStackEmitterSpawnScriptItemGroup.cpp
@@ -38,7 +38,7 @@ bool UNiagaraStackEmitterPropertiesItem::CanResetToBase() const
 {
 	if (bCanResetToBaseCache.IsSet() == false)
 	{
-		const UNiagaraEmitter* BaseEmitter = FNiagaraStackGraphUtilities::GetBaseEmitter(*Emitter.Get(), GetSystemViewModel()->GetSystem());
+		const UNiagaraEmitter* BaseEmitter = GetEmitterViewModel()->GetEmitter()->GetParent();
 		if (BaseEmitter != nullptr && Emitter != BaseEmitter)
 		{
 			TSharedRef<FNiagaraScriptMergeManager> MergeManager = FNiagaraScriptMergeManager::Get();
@@ -56,7 +56,7 @@ void UNiagaraStackEmitterPropertiesItem::ResetToBase()
 {
 	if (CanResetToBase())
 	{
-		const UNiagaraEmitter* BaseEmitter = FNiagaraStackGraphUtilities::GetBaseEmitter(*Emitter.Get(), GetSystemViewModel()->GetSystem());
+		const UNiagaraEmitter* BaseEmitter = GetEmitterViewModel()->GetEmitter()->GetParent();
 		TSharedRef<FNiagaraScriptMergeManager> MergeManager = FNiagaraScriptMergeManager::Get();
 		MergeManager->ResetEmitterEditablePropertySetToBase(*Emitter, *BaseEmitter);
 	}
diff --git a/Engine/Plugins/FX/Niagara/Source/NiagaraEditor/Private/ViewModels/Stack/NiagaraStackEntry.cpp b/Engine/Plugins/FX/Niagara/Source/NiagaraEditor/Private/ViewModels/Stack/NiagaraStackEntry.cpp
--- a/Engine/Plugins/FX/Niagara/Source/NiagaraEditor/Private/ViewModels/Stack/NiagaraStackEntry.cpp
+++ b/Engine/Plugins/FX/Niagara/Source/NiagaraEditor/Private/ViewModels/Stack/NiagaraStackEntry.cpp
@@ -391,7 +391,7 @@ bool UNiagaraStackEntry::HasBaseEmitter() const
 	if (bHasBaseEmitterCache.IsSet() == false)
 	{
 		TSharedRef<FNiagaraScriptMergeManager> MergeManager = FNiagaraScriptMergeManager::Get();
-		const UNiagaraEmitter* BaseEmitter = FNiagaraStackGraphUtilities::GetBaseEmitter(*GetEmitterViewModel()->GetEmitter(), GetSystemViewModel()->GetSystem());
+		const UNiagaraEmitter* BaseEmitter = GetEmitterViewModel()->GetParentEmitter();
 		bHasBaseEmitterCache = BaseEmitter != nullptr;
 	}
 	return bHasBaseEmitterCache.GetValue();
diff --git a/Engine/Plugins/FX/Niagara/Source/NiagaraEditor/Private/ViewModels/Stack/NiagaraStackEventScriptItemGroup.cpp b/Engine/Plugins/FX/Niagara/Source/NiagaraEditor/Private/ViewModels/Stack/NiagaraStackEventScriptItemGroup.cpp
--- a/Engine/Plugins/FX/Niagara/Source/NiagaraEditor/Private/ViewModels/Stack/NiagaraStackEventScriptItemGroup.cpp
+++ b/Engine/Plugins/FX/Niagara/Source/NiagaraEditor/Private/ViewModels/Stack/NiagaraStackEventScriptItemGroup.cpp
@@ -57,7 +57,7 @@ bool UNiagaraStackEventHandlerPropertiesItem::CanResetToBase() const
 	{
 		if (HasBaseEventHandler())
 		{
-			const UNiagaraEmitter* BaseEmitter = FNiagaraStackGraphUtilities::GetBaseEmitter(*Emitter.Get(), GetSystemViewModel()->GetSystem());
+			const UNiagaraEmitter* BaseEmitter = GetEmitterViewModel()->GetEmitter()->GetParent();
 			if (BaseEmitter != nullptr && Emitter != BaseEmitter)
 			{
 				TSharedRef<FNiagaraScriptMergeManager> MergeManager = FNiagaraScriptMergeManager::Get();
@@ -80,7 +80,7 @@ void UNiagaraStackEventHandlerPropertiesItem::ResetToBase()
 {
 	if (CanResetToBase())
 	{
-		const UNiagaraEmitter* BaseEmitter = FNiagaraStackGraphUtilities::GetBaseEmitter(*Emitter.Get(), GetSystemViewModel()->GetSystem());
+		const UNiagaraEmitter* BaseEmitter = GetEmitterViewModel()->GetEmitter()->GetParent();
 		TSharedRef<FNiagaraScriptMergeManager> MergeManager = FNiagaraScriptMergeManager::Get();
 		MergeManager->ResetEventHandlerPropertySetToBase(*Emitter, *BaseEmitter, EventScriptUsageId);
 		RefreshChildren();
@@ -173,7 +173,7 @@ bool UNiagaraStackEventHandlerPropertiesItem::HasBaseEventHandler() const
 {
 	if (bHasBaseEventHandlerCache.IsSet() == false)
 	{
-		const UNiagaraEmitter* BaseEmitter = FNiagaraStackGraphUtilities::GetBaseEmitter(*Emitter.Get(), GetSystemViewModel()->GetSystem());
+		const UNiagaraEmitter* BaseEmitter = GetEmitterViewModel()->GetEmitter()->GetParent();
 		if (BaseEmitter != nullptr && Emitter != BaseEmitter)
 		{
 			TSharedRef<FNiagaraScriptMergeManager> MergeManager = FNiagaraScriptMergeManager::Get();
@@ -287,7 +287,7 @@ bool UNiagaraStackEventScriptItemGroup::HasBaseEventHandler() const
 {
 	if (bHasBaseEventHandlerCache.IsSet() == false)
 	{
-		const UNiagaraEmitter* BaseEmitter = FNiagaraStackGraphUtilities::GetBaseEmitter(*GetEmitterViewModel()->GetEmitter(), GetSystemViewModel()->GetSystem());
+		const UNiagaraEmitter* BaseEmitter = GetEmitterViewModel()->GetEmitter()->GetParent();
 		bHasBaseEventHandlerCache = BaseEmitter != nullptr && FNiagaraScriptMergeManager::Get()->HasBaseEventHandler(*BaseEmitter, GetScriptUsageId());
 	}
 	return bHasBaseEventHandlerCache.GetValue();
diff --git a/Engine/Plugins/FX/Niagara/Source/NiagaraEditor/Private/ViewModels/Stack/NiagaraStackFunctionInput.cpp b/Engine/Plugins/FX/Niagara/Source/NiagaraEditor/Private/ViewModels/Stack/NiagaraStackFunctionInput.cpp
--- a/Engine/Plugins/FX/Niagara/Source/NiagaraEditor/Private/ViewModels/Stack/NiagaraStackFunctionInput.cpp
+++ b/Engine/Plugins/FX/Niagara/Source/NiagaraEditor/Private/ViewModels/Stack/NiagaraStackFunctionInput.cpp
@@ -1277,7 +1277,7 @@ bool UNiagaraStackFunctionInput::CanResetToBase() const
 				UNiagaraNodeOutput* OutputNode = FNiagaraStackGraphUtilities::GetEmitterOutputNodeForStackNode(*OwningFunctionCallNode.Get());
 				if(MergeManager->IsMergeableScriptUsage(OutputNode->GetUsage()))
 				{
-					const UNiagaraEmitter* BaseEmitter = FNiagaraStackGraphUtilities::GetBaseEmitter(*GetEmitterViewModel()->GetEmitter(), GetSystemViewModel()->GetSystem());
+					const UNiagaraEmitter* BaseEmitter = GetEmitterViewModel()->GetParentEmitter();
 
 					bCanResetToBaseCache = BaseEmitter != nullptr && MergeManager->IsModuleInputDifferentFromBase(
 						*GetEmitterViewModel()->GetEmitter(),
@@ -1308,17 +1308,7 @@ void UNiagaraStackFunctionInput::ResetToBase()
 	{
 		TSharedRef<FNiagaraScriptMergeManager> MergeManager = FNiagaraScriptMergeManager::Get();
 
-		TSharedPtr<FNiagaraEmitterHandleViewModel> ThisEmitterHandleViewModel;
-		for (TSharedRef<FNiagaraEmitterHandleViewModel> EmitterHandleViewModel : GetSystemViewModel()->GetEmitterHandleViewModels())
-		{
-			if (EmitterHandleViewModel->GetEmitterViewModel() == GetEmitterViewModel())
-			{
-				ThisEmitterHandleViewModel = EmitterHandleViewModel;
-				break;
-			}
-		}
-
-		const UNiagaraEmitter* BaseEmitter = ThisEmitterHandleViewModel->GetEmitterHandle()->GetSource();
+		const UNiagaraEmitter* BaseEmitter = GetEmitterViewModel()->GetEmitter()->GetParent();
 		UNiagaraNodeOutput* OutputNode = FNiagaraStackGraphUtilities::GetEmitterOutputNodeForStackNode(*OwningFunctionCallNode.Get());
 
 		FScopedTransaction ScopedTransaction(LOCTEXT("ResetInputToBaseTransaction", "Reset this input to match the parent emitter."));
diff --git a/Engine/Plugins/FX/Niagara/Source/NiagaraEditor/Private/ViewModels/Stack/NiagaraStackGraphUtilities.cpp b/Engine/Plugins/FX/Niagara/Source/NiagaraEditor/Private/ViewModels/Stack/NiagaraStackGraphUtilities.cpp
--- a/Engine/Plugins/FX/Niagara/Source/NiagaraEditor/Private/ViewModels/Stack/NiagaraStackGraphUtilities.cpp
+++ b/Engine/Plugins/FX/Niagara/Source/NiagaraEditor/Private/ViewModels/Stack/NiagaraStackGraphUtilities.cpp
@@ -1154,27 +1154,6 @@ UNiagaraNodeOutput* FNiagaraStackGraphUtilities::ResetGraphForOutput(UNiagaraGra
 	return OutputNode;
 }
 
-const UNiagaraEmitter* FNiagaraStackGraphUtilities::GetBaseEmitter(UNiagaraEmitter& Emitter, UNiagaraSystem& OwningSystem)
-{
-	for(const FNiagaraEmitterHandle& Handle : OwningSystem.GetEmitterHandles())
-	{ 
-		if (Handle.GetInstance() == &Emitter)
-		{
-			if (Handle.GetSource() != nullptr && Handle.GetSource() != &Emitter)
-			{
-				return Handle.GetSource();
-			}
-			else
-			{
-				// If the source is null then it was deleted and if the source is the same as the emitter the owning
-				// system is transient and the emitter doesn't have base.
-				return nullptr;
-			}
-		}
-	}
-	return nullptr;
-}
-
 void GetFunctionNamesRecursive(UNiagaraNode* CurrentNode, TArray<UNiagaraNode*>& VisitedNodes, TArray<FString>& FunctionNames)
 {
 	if (VisitedNodes.Contains(CurrentNode) == false)
diff --git a/Engine/Plugins/FX/Niagara/Source/NiagaraEditor/Private/ViewModels/Stack/NiagaraStackModuleItem.cpp b/Engine/Plugins/FX/Niagara/Source/NiagaraEditor/Private/ViewModels/Stack/NiagaraStackModuleItem.cpp
--- a/Engine/Plugins/FX/Niagara/Source/NiagaraEditor/Private/ViewModels/Stack/NiagaraStackModuleItem.cpp
+++ b/Engine/Plugins/FX/Niagara/Source/NiagaraEditor/Private/ViewModels/Stack/NiagaraStackModuleItem.cpp
@@ -699,7 +699,7 @@ bool UNiagaraStackModuleItem::CanMoveAndDelete() const
 			// When editing systems only non-base modules can be moved and deleted.
 			TSharedRef<FNiagaraScriptMergeManager> MergeManager = FNiagaraScriptMergeManager::Get();
 
-			const UNiagaraEmitter* BaseEmitter = FNiagaraStackGraphUtilities::GetBaseEmitter(*GetEmitterViewModel()->GetEmitter(), GetSystemViewModel()->GetSystem());
+			const UNiagaraEmitter* BaseEmitter = GetEmitterViewModel()->GetEmitter()->GetParent();
 
 			bool bIsMergeable = MergeManager->IsMergeableScriptUsage(OutputNode->GetUsage());
 			bool bHasBaseModule = bIsMergeable && BaseEmitter != nullptr && MergeManager->HasBaseModule(*BaseEmitter, OutputNode->GetUsage(), OutputNode->GetUsageId(), FunctionCallNode->NodeGuid);
diff --git a/Engine/Plugins/FX/Niagara/Source/NiagaraEditor/Private/ViewModels/Stack/NiagaraStackPropertyRow.cpp b/Engine/Plugins/FX/Niagara/Source/NiagaraEditor/Private/ViewModels/Stack/NiagaraStackPropertyRow.cpp
--- a/Engine/Plugins/FX/Niagara/Source/NiagaraEditor/Private/ViewModels/Stack/NiagaraStackPropertyRow.cpp
+++ b/Engine/Plugins/FX/Niagara/Source/NiagaraEditor/Private/ViewModels/Stack/NiagaraStackPropertyRow.cpp
@@ -30,6 +30,12 @@ bool UNiagaraStackPropertyRow::GetIsEnabled() const
 	return OwningNiagaraNode == nullptr || OwningNiagaraNode->GetDesiredEnabledState() == ENodeEnabledState::Enabled;
 }
 
+void UNiagaraStackPropertyRow::FinalizeInternal()
+{
+	Super::FinalizeInternal();
+	DetailTreeNode.Reset();
+}
+
 void UNiagaraStackPropertyRow::RefreshChildrenInternal(const TArray<UNiagaraStackEntry*>& CurrentChildren, TArray<UNiagaraStackEntry*>& NewChildren, TArray<FStackIssue>& NewIssues)
 {
 	TArray<TSharedRef<IDetailTreeNode>> NodeChildren;
diff --git a/Engine/Plugins/FX/Niagara/Source/NiagaraEditor/Private/ViewModels/Stack/NiagaraStackRendererItem.cpp b/Engine/Plugins/FX/Niagara/Source/NiagaraEditor/Private/ViewModels/Stack/NiagaraStackRendererItem.cpp
--- a/Engine/Plugins/FX/Niagara/Source/NiagaraEditor/Private/ViewModels/Stack/NiagaraStackRendererItem.cpp
+++ b/Engine/Plugins/FX/Niagara/Source/NiagaraEditor/Private/ViewModels/Stack/NiagaraStackRendererItem.cpp
@@ -168,7 +168,7 @@ bool UNiagaraStackRendererItem::HasBaseRenderer() const
 		if (bHasBaseRendererCache.IsSet() == false)
 		{
 			TSharedRef<FNiagaraScriptMergeManager> MergeManager = FNiagaraScriptMergeManager::Get();
-			const UNiagaraEmitter* BaseEmitter = FNiagaraStackGraphUtilities::GetBaseEmitter(*GetEmitterViewModel()->GetEmitter(), GetSystemViewModel()->GetSystem());
+			const UNiagaraEmitter* BaseEmitter = GetEmitterViewModel()->GetEmitter()->GetParent();
 			bHasBaseRendererCache = BaseEmitter != nullptr && MergeManager->HasBaseRenderer(*BaseEmitter, RendererProperties->GetMergeId());
 		}
 		return bHasBaseRendererCache.GetValue();
@@ -183,7 +183,7 @@ bool UNiagaraStackRendererItem::CanResetToBase() const
 		if (bCanResetToBaseCache.IsSet() == false)
 		{
 			TSharedRef<FNiagaraScriptMergeManager> MergeManager = FNiagaraScriptMergeManager::Get();
-			const UNiagaraEmitter* BaseEmitter = FNiagaraStackGraphUtilities::GetBaseEmitter(*GetEmitterViewModel()->GetEmitter(), GetSystemViewModel()->GetSystem());
+			const UNiagaraEmitter* BaseEmitter = GetEmitterViewModel()->GetEmitter()->GetParent();
 			bCanResetToBaseCache = BaseEmitter != nullptr && MergeManager->IsRendererDifferentFromBase(*GetEmitterViewModel()->GetEmitter(), *BaseEmitter, RendererProperties->GetMergeId());
 		}
 		return bCanResetToBaseCache.GetValue();
@@ -196,7 +196,7 @@ void UNiagaraStackRendererItem::ResetToBase()
 	if (CanResetToBase())
 	{
 		TSharedRef<FNiagaraScriptMergeManager> MergeManager = FNiagaraScriptMergeManager::Get();
-		const UNiagaraEmitter* BaseEmitter = FNiagaraStackGraphUtilities::GetBaseEmitter(*GetEmitterViewModel()->GetEmitter(), GetSystemViewModel()->GetSystem());
+		const UNiagaraEmitter* BaseEmitter = GetEmitterViewModel()->GetEmitter()->GetParent();
 		MergeManager->ResetRendererToBase(*GetEmitterViewModel()->GetEmitter(), *BaseEmitter, RendererProperties->GetMergeId());
 		ModifiedGroupItemsDelegate.ExecuteIfBound();
 	}
diff --git a/Engine/Plugins/FX/Niagara/Source/NiagaraEditor/Private/ViewModels/Stack/NiagaraStackViewModel.cpp b/Engine/Plugins/FX/Niagara/Source/NiagaraEditor/Private/ViewModels/Stack/NiagaraStackViewModel.cpp
--- a/Engine/Plugins/FX/Niagara/Source/NiagaraEditor/Private/ViewModels/Stack/NiagaraStackViewModel.cpp
+++ b/Engine/Plugins/FX/Niagara/Source/NiagaraEditor/Private/ViewModels/Stack/NiagaraStackViewModel.cpp
@@ -182,17 +182,18 @@ bool UNiagaraStackViewModel::HasDismissedStackIssues()
 		|| GetEmitterHandleViewModel()->GetEmitterViewModel()->GetEditorData().GetStackEditorData().GetDismissedStackIssueIds().Num() > 0;
 }
 
-bool UNiagaraStackViewModel::HasEmitterSource() const
+bool UNiagaraStackViewModel::HasParentEmitter() const
 {
-	return EmitterHandleViewModel->GetEmitterHandle()->GetSource() != nullptr;
+	return EmitterHandleViewModel->GetEmitterViewModel()->GetParentEmitter() != nullptr;
 }
 
 void UNiagaraStackViewModel::RemoveEmitterSource()
 {
 	{
-		FScopedTransaction ScopedTransaction(LOCTEXT("RemoveEmitterSourceTransaction", "Remove Emitter Source"));
-		SystemViewModel->GetSystem().Modify();
-		EmitterHandleViewModel->GetEmitterHandle()->RemoveSource();
+		FScopedTransaction ScopedTransaction(LOCTEXT("RemoveParentEmitterTransaction", "Remove Parent Emitter"));
+		UNiagaraEmitter* Emitter = EmitterHandleViewModel->GetEmitterViewModel()->GetEmitter();
+		Emitter->Modify();
+		Emitter->RemoveParent();
 	}
 	RootEntry->RefreshChildren();
 }
diff --git a/Engine/Plugins/FX/Niagara/Source/NiagaraEditor/Public/NiagaraEditorModule.h b/Engine/Plugins/FX/Niagara/Source/NiagaraEditor/Public/NiagaraEditorModule.h
--- a/Engine/Plugins/FX/Niagara/Source/NiagaraEditor/Public/NiagaraEditorModule.h
+++ b/Engine/Plugins/FX/Niagara/Source/NiagaraEditor/Public/NiagaraEditorModule.h
@@ -119,7 +119,6 @@ private:
 	FDelegateHandle CreateVectorParameterTrackEditorHandle;
 	FDelegateHandle CreateColorParameterTrackEditorHandle;
 
-	FDelegateHandle MergeEmitterHandle;
 	FDelegateHandle CreateDefaultScriptSourceHandle;
 	FDelegateHandle ScriptCompilerHandle;
 	FDelegateHandle PrecompilerHandle;
diff --git a/Engine/Plugins/FX/Niagara/Source/NiagaraEditor/Public/NiagaraEmitterFactoryNew.h b/Engine/Plugins/FX/Niagara/Source/NiagaraEditor/Public/NiagaraEmitterFactoryNew.h
--- a/Engine/Plugins/FX/Niagara/Source/NiagaraEditor/Public/NiagaraEmitterFactoryNew.h
+++ b/Engine/Plugins/FX/Niagara/Source/NiagaraEditor/Public/NiagaraEmitterFactoryNew.h
@@ -14,6 +14,7 @@ class UNiagaraEmitterFactoryNew : public UFactory
 	GENERATED_UCLASS_BODY()
 
 	UNiagaraEmitter* EmitterToCopy;
+	bool bUseInheritance;
 	bool bAddDefaultModulesAndRenderersToEmptyEmitter;
 
 	//~ Begin UFactory Interface
diff --git a/Engine/Plugins/FX/Niagara/Source/NiagaraEditor/Public/ViewModels/NiagaraEmitterHandleViewModel.h b/Engine/Plugins/FX/Niagara/Source/NiagaraEditor/Public/ViewModels/NiagaraEmitterHandleViewModel.h
--- a/Engine/Plugins/FX/Niagara/Source/NiagaraEditor/Public/ViewModels/NiagaraEmitterHandleViewModel.h
+++ b/Engine/Plugins/FX/Niagara/Source/NiagaraEditor/Public/ViewModels/NiagaraEmitterHandleViewModel.h
@@ -47,12 +47,6 @@ public:
 	/** Gets the text representation of the emitter handle name. */
 	NIAGARAEDITOR_API FText GetNameText() const;
 
-	/** Gets the text representation of the source emitter handle name. */
-	NIAGARAEDITOR_API FText GetSourceNameText() const;
-
-	/** Gets the text representation of the emitter handle name. */
-	NIAGARAEDITOR_API FText GetSourcePathNameText() const;
-
 	/** Called when the contents of the name text control is committed. */
 	NIAGARAEDITOR_API void OnNameTextComitted(const FText& InText, ETextCommit::Type CommitInfo);
 
@@ -85,10 +79,6 @@ public:
 	/** Gets the view model for the emitter this handle references. */
 	NIAGARAEDITOR_API TSharedPtr<FNiagaraEmitterViewModel> GetEmitterViewModel();
 
-
-	/** Opens the source emitter in a stand alone asset editor. */
-	void OpenSourceEmitter();
-
 	/** Gets a multicast delegate which is called any time a property on the handle changes. */
 	FOnPropertyChanged& OnPropertyChanged();
 
diff --git a/Engine/Plugins/FX/Niagara/Source/NiagaraEditor/Public/ViewModels/NiagaraEmitterViewModel.h b/Engine/Plugins/FX/Niagara/Source/NiagaraEditor/Public/ViewModels/NiagaraEmitterViewModel.h
--- a/Engine/Plugins/FX/Niagara/Source/NiagaraEditor/Public/ViewModels/NiagaraEmitterViewModel.h
+++ b/Engine/Plugins/FX/Niagara/Source/NiagaraEditor/Public/ViewModels/NiagaraEmitterViewModel.h
@@ -43,6 +43,15 @@ public:
 	/** Gets the emitter represented by this view model. */
 	NIAGARAEDITOR_API UNiagaraEmitter* GetEmitter();
 
+	/** Gets the parent emitter for the emitter represented by this view model, if it has one. */
+	NIAGARAEDITOR_API const UNiagaraEmitter* GetParentEmitter() const;
+
+	/** Gets the text representation of the parent emitter name. */
+	NIAGARAEDITOR_API FText GetParentNameText() const;
+
+	/** Gets the text representation of the parent emitter path. */
+	NIAGARAEDITOR_API FText GetParentPathNameText() const;
+
 	/** Gets text representing stats for the emitter. */
 	//~ TODO: Instead of a single string here, we should probably have separate controls with tooltips etc.
 	NIAGARAEDITOR_API FText GetStatsText() const;
diff --git a/Engine/Plugins/FX/Niagara/Source/NiagaraEditor/Public/ViewModels/Stack/NiagaraStackGraphUtilities.h b/Engine/Plugins/FX/Niagara/Source/NiagaraEditor/Public/ViewModels/Stack/NiagaraStackGraphUtilities.h
--- a/Engine/Plugins/FX/Niagara/Source/NiagaraEditor/Public/ViewModels/Stack/NiagaraStackGraphUtilities.h
+++ b/Engine/Plugins/FX/Niagara/Source/NiagaraEditor/Public/ViewModels/Stack/NiagaraStackGraphUtilities.h
@@ -120,8 +120,6 @@ namespace FNiagaraStackGraphUtilities
 
 	UNiagaraNodeOutput* ResetGraphForOutput(UNiagaraGraph& NiagaraGraph, ENiagaraScriptUsage ScriptUsage, FGuid ScriptUsageId, const FGuid& PreferredOutputNodeGuid = FGuid(), const FGuid& PreferredInputNodeGuid = FGuid());
 
-	const UNiagaraEmitter* GetBaseEmitter(UNiagaraEmitter& Emitter, UNiagaraSystem& OwningSystem);
-
 	bool IsRapidIterationType(const FNiagaraTypeDefinition& InputType);
 
 	FNiagaraVariable CreateRapidIterationParameter(const FString& UniqueEmitterName, ENiagaraScriptUsage ScriptUsage, const FName& AliasedInputName, const FNiagaraTypeDefinition& InputType);
diff --git a/Engine/Plugins/FX/Niagara/Source/NiagaraEditor/Public/ViewModels/Stack/NiagaraStackPropertyRow.h b/Engine/Plugins/FX/Niagara/Source/NiagaraEditor/Public/ViewModels/Stack/NiagaraStackPropertyRow.h
--- a/Engine/Plugins/FX/Niagara/Source/NiagaraEditor/Public/ViewModels/Stack/NiagaraStackPropertyRow.h
+++ b/Engine/Plugins/FX/Niagara/Source/NiagaraEditor/Public/ViewModels/Stack/NiagaraStackPropertyRow.h
@@ -21,6 +21,8 @@ public:
 	virtual bool GetIsEnabled() const override;
 
 protected:
+	virtual void FinalizeInternal() override;
+
 	virtual void RefreshChildrenInternal(const TArray<UNiagaraStackEntry*>& CurrentChildren, TArray<UNiagaraStackEntry*>& NewChildren, TArray<FStackIssue>& NewIssues) override;
 
 	virtual void GetSearchItems(TArray<FStackSearchItem>& SearchItems) const override;
diff --git a/Engine/Plugins/FX/Niagara/Source/NiagaraEditor/Public/ViewModels/Stack/NiagaraStackViewModel.h b/Engine/Plugins/FX/Niagara/Source/NiagaraEditor/Public/ViewModels/Stack/NiagaraStackViewModel.h
--- a/Engine/Plugins/FX/Niagara/Source/NiagaraEditor/Public/ViewModels/Stack/NiagaraStackViewModel.h
+++ b/Engine/Plugins/FX/Niagara/Source/NiagaraEditor/Public/ViewModels/Stack/NiagaraStackViewModel.h
@@ -83,7 +83,7 @@ public:
 
 	bool HasDismissedStackIssues();
 
-	bool HasEmitterSource() const;
+	bool HasParentEmitter() const;
 	void RemoveEmitterSource();
 
 private:
diff --git a/Engine/Plugins/FX/Niagara/Source/NiagaraEditorWidgets/Private/SNiagaraStack.cpp b/Engine/Plugins/FX/Niagara/Source/NiagaraEditorWidgets/Private/SNiagaraStack.cpp
--- a/Engine/Plugins/FX/Niagara/Source/NiagaraEditorWidgets/Private/SNiagaraStack.cpp
+++ b/Engine/Plugins/FX/Niagara/Source/NiagaraEditorWidgets/Private/SNiagaraStack.cpp
@@ -162,7 +162,7 @@ void SNiagaraStack::Construct(const FArguments& InArgs, UNiagaraStackViewModel*
 					FSlateIcon(),
 					FUIAction(
 						FExecuteAction::CreateUObject(StackViewModel, &UNiagaraStackViewModel::RemoveEmitterSource),
-						FCanExecuteAction::CreateUObject(StackViewModel, &UNiagaraStackViewModel::HasEmitterSource)));
+						FCanExecuteAction::CreateUObject(StackViewModel, &UNiagaraStackViewModel::HasParentEmitter)));
 
 				MenuBuilder.AddMenuEntry(
 					LOCTEXT("ShowEmitterInContentBrowser", "Show in Content Browser"),
@@ -339,11 +339,11 @@ void SNiagaraStack::ConstructHeaderWidget()
 				[
 					SNew(SButton)
 					.IsFocusable(false)
-					.ToolTipText(LOCTEXT("OpenAndFocusSourceEmitterToolTip", "Open and Focus Source Emitter"))
+					.ToolTipText(LOCTEXT("OpenAndFocusParentEmitterToolTip", "Open and Focus Parent Emitter"))
 					.ButtonStyle(FEditorStyle::Get(), "HoverHintOnly")
 					.ForegroundColor(this, &SNiagaraStack::GetPinColor)
 					.ContentPadding(2)
-					.OnClicked(this, &SNiagaraStack::OpenSourceEmitter)
+					.OnClicked(this, &SNiagaraStack::OpenParentEmitter)
 					.Visibility(this, &SNiagaraStack::GetOpenSourceEmitterVisibility)
 					// GoToSource icon is 30x30px so we scale it down to stay in line with other 12x12px UI
 					.DesiredSizeScale(FVector2D(0.55f, 0.55f))
@@ -530,14 +530,14 @@ bool SNiagaraStack::IsEntryFocusedInSearch(UNiagaraStackEntry* Entry) const
 	return false;
 }
 
-FReply SNiagaraStack::OpenSourceEmitter() 
+FReply SNiagaraStack::OpenParentEmitter() 
 {
-	if (StackViewModel && StackViewModel->GetEmitterHandleViewModel().IsValid() && StackViewModel->GetEmitterHandleViewModel()->GetEmitterHandle())
+	if (StackViewModel && StackViewModel->GetEmitterHandleViewModel().IsValid() && StackViewModel->GetEmitterHandleViewModel()->GetEmitterViewModel().IsValid())
 	{
-		UNiagaraEmitter* Emitter = const_cast<UNiagaraEmitter*>(StackViewModel->GetEmitterHandleViewModel()->GetEmitterHandle()->GetSource());
-		if (Emitter != nullptr)
+		UNiagaraEmitter* ParentEmitter = const_cast<UNiagaraEmitter*>(StackViewModel->GetEmitterHandleViewModel()->GetEmitterViewModel()->GetParentEmitter());
+		if (ParentEmitter != nullptr)
 		{
-			FAssetEditorManager::Get().OpenEditorForAsset(Emitter);
+			FAssetEditorManager::Get().OpenEditorForAsset(ParentEmitter);
 		}
 	}
 	return FReply::Handled();
@@ -555,7 +555,7 @@ EVisibility SNiagaraStack::GetPinEmitterVisibility() const
 
 EVisibility SNiagaraStack::GetOpenSourceEmitterVisibility() const
 {
-	return CanOpenSourceEmitter() ? EVisibility::Visible : EVisibility::Collapsed;
+	return StackViewModel->HasParentEmitter() ? EVisibility::Visible : EVisibility::Collapsed;
 }
 
 bool SNiagaraStack::GetEmitterNameIsReadOnly() const
@@ -567,17 +567,6 @@ bool SNiagaraStack::GetEmitterNameIsReadOnly() const
 	return true;
 }
 
-bool SNiagaraStack::CanOpenSourceEmitter() const
-{
-	if (StackViewModel && StackViewModel->GetEmitterHandleViewModel().IsValid() && StackViewModel->GetEmitterHandleViewModel()->GetEmitterHandle()
-		&& StackViewModel->GetEmitterHandleViewModel()->GetEmitterHandle()->GetSource())
-	{
-		return true;
-	}
-
-	return false;
-}
-
 void SNiagaraStack::SetEmitterEnabled(bool bIsEnabled)
 {
 	StackViewModel->GetEmitterHandleViewModel()->SetIsEnabled(bIsEnabled);
@@ -592,8 +581,11 @@ void SNiagaraStack::ShowEmitterInContentBrowser()
 {
 	FContentBrowserModule& ContentBrowserModule = FModuleManager::Get().LoadModuleChecked<FContentBrowserModule>(TEXT("ContentBrowser"));
 	TArray<FAssetData> Assets;
-	Assets.Add(FAssetData(StackViewModel->GetEmitterHandleViewModel()->GetEmitterHandle()->GetSource()));
-	ContentBrowserModule.Get().SyncBrowserToAssets(Assets);
+	if (StackViewModel->HasParentEmitter())
+	{
+		Assets.Add(FAssetData(StackViewModel->GetEmitterHandleViewModel()->GetEmitterViewModel()->GetParentEmitter()));
+		ContentBrowserModule.Get().SyncBrowserToAssets(Assets);
+	}
 }
 
 void SNiagaraStack::NavigateTo(UNiagaraStackEntry* Item)
@@ -1067,16 +1059,16 @@ void SNiagaraStack::StackStructureChanged()
 
 FText SNiagaraStack::GetSourceEmitterNameText() const 
 {
-	return StackViewModel->GetEmitterHandleViewModel()->GetSourceNameText();
+	return StackViewModel->GetEmitterHandleViewModel()->GetEmitterViewModel()->GetParentNameText();
 }
 
 FText SNiagaraStack::GetEmitterNameToolTip() const
 {
-	if (CanOpenSourceEmitter())
+	if (StackViewModel->HasParentEmitter())
 	{
 		// We are looking at this Emitter in a System Asset and it has a valid parent Emitter
 		TSharedPtr<FNiagaraEmitterHandleViewModel> ThisViewModel = StackViewModel->GetEmitterHandleViewModel();
-		return FText::Format(LOCTEXT("EmitterNameAndPath", "{0}\nParent: {1}"), ThisViewModel->GetNameText(), ThisViewModel->GetSourcePathNameText());
+		return FText::Format(LOCTEXT("EmitterNameAndPath", "{0}\nParent: {1}"), ThisViewModel->GetNameText(), ThisViewModel->GetEmitterViewModel()->GetParentPathNameText());
 	}
 	else
 	{
@@ -1092,13 +1084,13 @@ void SNiagaraStack::OnStackViewNameTextCommitted(const FText& InText, ETextCommi
 
 EVisibility SNiagaraStack::GetSourceEmitterNameVisibility() const
 {
-	return CanOpenSourceEmitter() && GetIsEmitterRenamed() ? EVisibility::Visible : EVisibility::Collapsed;
+	return StackViewModel->HasParentEmitter() && GetIsEmitterRenamed() ? EVisibility::Visible : EVisibility::Collapsed;
 }
 
 bool SNiagaraStack::GetIsEmitterRenamed() const
 {
 	const FText CurrentNameText = StackViewModel->GetEmitterHandleViewModel()->GetNameText();
-	const FText SourceNameText = StackViewModel->GetEmitterHandleViewModel()->GetSourceNameText();
+	const FText SourceNameText = StackViewModel->GetEmitterHandleViewModel()->GetEmitterViewModel()->GetParentNameText();
 	return !CurrentNameText.EqualTo(SourceNameText);
 }
 
diff --git a/Engine/Plugins/FX/Niagara/Source/NiagaraEditorWidgets/Private/SNiagaraStack.h b/Engine/Plugins/FX/Niagara/Source/NiagaraEditorWidgets/Private/SNiagaraStack.h
--- a/Engine/Plugins/FX/Niagara/Source/NiagaraEditorWidgets/Private/SNiagaraStack.h
+++ b/Engine/Plugins/FX/Niagara/Source/NiagaraEditorWidgets/Private/SNiagaraStack.h
@@ -72,7 +72,7 @@ private:
 	void ConstructHeaderWidget();
 	FSlateColor GetPinColor() const;
 	FReply PinButtonPressed();
-	FReply OpenSourceEmitter();
+	FReply OpenParentEmitter();
 	EVisibility GetEnableCheckboxVisibility() const;
 	EVisibility GetPinEmitterVisibility() const;
 	EVisibility GetOpenSourceEmitterVisibility() const;
@@ -100,7 +100,6 @@ private:
 	bool IsEntryFocusedInSearch(UNiagaraStackEntry* Entry) const;
 	
 	// Inline menu commands
-	bool CanOpenSourceEmitter() const;
 	void SetEmitterEnabled(bool bIsEnabled);
 	bool CheckEmitterEnabledStatus(bool bIsEnabled);
 	void ShowEmitterInContentBrowser();
