Niagara - Add missing include to fix non-unity build.

diff --git a/Engine/Plugins/FX/Niagara/Source/Niagara/Private/NiagaraEmitter.cpp b/Engine/Plugins/FX/Niagara/Source/Niagara/Private/NiagaraEmitter.cpp
--- a/Engine/Plugins/FX/Niagara/Source/Niagara/Private/NiagaraEmitter.cpp
+++ b/Engine/Plugins/FX/Niagara/Source/Niagara/Private/NiagaraEmitter.cpp
@@ -11,6 +11,7 @@
 #include "NiagaraModule.h"
 #include "NiagaraSystem.h"
 #include "NiagaraStats.h"
+#include "Modules/ModuleManager.h"
 
 #if WITH_EDITOR
 const FName UNiagaraEmitter::PrivateMemberNames::EventHandlerScriptProps = GET_MEMBER_NAME_CHECKED(UNiagaraEmitter, EventHandlerScriptProps);
