- Invalidate the CachedScriptVM if it's out of date to fix some cook errors, a further investigation is required in how to handle this correctly

# fyi frank.fella,shaun.kime
#rnx

diff --git a/Engine/Plugins/FX/Niagara/Source/Niagara/Private/NiagaraScript.cpp b/Engine/Plugins/FX/Niagara/Source/Niagara/Private/NiagaraScript.cpp
--- a/Engine/Plugins/FX/Niagara/Source/Niagara/Private/NiagaraScript.cpp
+++ b/Engine/Plugins/FX/Niagara/Source/Niagara/Private/NiagaraScript.cpp
@@ -492,6 +492,12 @@ void UNiagaraScript::PostLoad()
 			}
 		}
 	}
+
+	// Invalidate the CachedScriptVM if it's out of date to fix some cook errors, a further investigation is required in how to handle this correctly
+	if (CachedScriptVMId.CompilerVersionID != FNiagaraCustomVersion::LatestScriptCompileVersion)
+	{
+		CachedScriptVM.LastCompileStatus = ENiagaraScriptCompileStatus::NCS_Unknown;
+	}
 #endif
 	
 	// Resources can be processed / registered now that we're back on the main thread
