Right clicingk in NiagaraStack and then selecting Insert Module Above/Below now auto selects the generated search box.

diff --git a/Engine/Plugins/FX/Niagara/Source/NiagaraEditorWidgets/Private/Stack/SNiagaraStackItemGroupAddMenu.cpp b/Engine/Plugins/FX/Niagara/Source/NiagaraEditorWidgets/Private/Stack/SNiagaraStackItemGroupAddMenu.cpp
--- a/Engine/Plugins/FX/Niagara/Source/NiagaraEditorWidgets/Private/Stack/SNiagaraStackItemGroupAddMenu.cpp
+++ b/Engine/Plugins/FX/Niagara/Source/NiagaraEditorWidgets/Private/Stack/SNiagaraStackItemGroupAddMenu.cpp
@@ -12,6 +12,7 @@ void SNiagaraStackItemGroupAddMenu::Construct(const FArguments& InArgs, INiagara
 {
 	AddUtilities = InAddUtilities;
 	InsertIndex = InInsertIndex;
+	bSetFocusOnNextTick = true;
 	ChildSlot
 	[
 		SNew(SBorder)
@@ -32,6 +33,15 @@ void SNiagaraStackItemGroupAddMenu::Construct(const FArguments& InArgs, INiagara
 	];
 }
 
+void SNiagaraStackItemGroupAddMenu::Tick(const FGeometry& AllottedGeometry, const double InCurrentTime, const float InDeltaTime)
+{
+	if (bSetFocusOnNextTick)
+	{
+		bSetFocusOnNextTick = false;
+		FSlateApplication::Get().SetKeyboardFocus(GetFilterTextBox(), EFocusCause::SetDirectly);
+	}
+}
+
 TSharedPtr<SEditableTextBox> SNiagaraStackItemGroupAddMenu::GetFilterTextBox()
 {
 	return AddMenu->GetFilterTextBox();
diff --git a/Engine/Plugins/FX/Niagara/Source/NiagaraEditorWidgets/Private/Stack/SNiagaraStackItemGroupAddMenu.h b/Engine/Plugins/FX/Niagara/Source/NiagaraEditorWidgets/Private/Stack/SNiagaraStackItemGroupAddMenu.h
--- a/Engine/Plugins/FX/Niagara/Source/NiagaraEditorWidgets/Private/Stack/SNiagaraStackItemGroupAddMenu.h
+++ b/Engine/Plugins/FX/Niagara/Source/NiagaraEditorWidgets/Private/Stack/SNiagaraStackItemGroupAddMenu.h
@@ -16,6 +16,8 @@ public:
 
 	void Construct(const FArguments& InArgs, INiagaraStackItemGroupAddUtilities* InAddUtilities, int32 InInsertIndex);
 
+	virtual void Tick(const FGeometry& AllottedGeometry, const double InCurrentTime, const float InDeltaTime) override;
+
 	TSharedPtr<class SEditableTextBox> GetFilterTextBox();
 
 private:
@@ -29,4 +31,6 @@ private:
 	int32 InsertIndex;
 
 	TSharedPtr<class SGraphActionMenu> AddMenu;
+
+	bool bSetFocusOnNextTick;
 };
\ No newline at end of file
