Niagara - Fix an issue where emitters weren't being coppied correctly for inheritance which resulted in cross asset references to the same object.

diff --git a/Engine/Plugins/FX/Niagara/Source/Niagara/Private/NiagaraEmitter.cpp b/Engine/Plugins/FX/Niagara/Source/Niagara/Private/NiagaraEmitter.cpp
--- a/Engine/Plugins/FX/Niagara/Source/Niagara/Private/NiagaraEmitter.cpp
+++ b/Engine/Plugins/FX/Niagara/Source/Niagara/Private/NiagaraEmitter.cpp
@@ -398,7 +398,7 @@ UNiagaraEmitter* UNiagaraEmitter::CreateAsDuplicate(const UNiagaraEmitter& InEmi
 	NewEmitter->Parent = InEmitterToDuplicate.Parent;
 	if (InEmitterToDuplicate.ParentAtLastMerge != nullptr)
 	{
-		NewEmitter->ParentAtLastMerge = InEmitterToDuplicate.ParentAtLastMerge;
+		NewEmitter->ParentAtLastMerge = Cast<UNiagaraEmitter>StaticDuplicateObject(InEmitterToDuplicate.ParentAtLastMerge, NewEmitter);
 		NewEmitter->ParentAtLastMerge->ClearFlags(RF_Standalone | RF_Public);
 	}
 	NewEmitter->SetUniqueEmitterName(InDuplicateName.ToString());
