Niagara - Fix last change, wrong filed unshelved.

diff --git a/Engine/Plugins/FX/Niagara/Source/Niagara/Private/NiagaraEmitter.cpp b/Engine/Plugins/FX/Niagara/Source/Niagara/Private/NiagaraEmitter.cpp
--- a/Engine/Plugins/FX/Niagara/Source/Niagara/Private/NiagaraEmitter.cpp
+++ b/Engine/Plugins/FX/Niagara/Source/Niagara/Private/NiagaraEmitter.cpp
@@ -398,7 +398,7 @@ UNiagaraEmitter* UNiagaraEmitter::CreateAsDuplicate(const UNiagaraEmitter& InEmi
 	NewEmitter->Parent = InEmitterToDuplicate.Parent;
 	if (InEmitterToDuplicate.ParentAtLastMerge != nullptr)
 	{
-		NewEmitter->ParentAtLastMerge = Cast<UNiagaraEmitter>StaticDuplicateObject(InEmitterToDuplicate.ParentAtLastMerge, NewEmitter);
+		NewEmitter->ParentAtLastMerge = Cast<UNiagaraEmitter>(StaticDuplicateObject(InEmitterToDuplicate.ParentAtLastMerge, NewEmitter));
 		NewEmitter->ParentAtLastMerge->ClearFlags(RF_Standalone | RF_Public);
 	}
 	NewEmitter->SetUniqueEmitterName(InDuplicateName.ToString());
