Niagara - Fix chained emitter inheritance cook issues.
+ When copying data from the merged emitter, copy the gpu script as well.
+ When renaming the emitter object as a result of setting it's unique name, don't reset loaders.
+ When upgrading old data to move the inheritance information from the handle to the emitter, don't outer the parent emitter to the instance because that makes no sense.

diff --git a/Engine/Plugins/FX/Niagara/Source/Niagara/Private/NiagaraEmitter.cpp b/Engine/Plugins/FX/Niagara/Source/Niagara/Private/NiagaraEmitter.cpp
--- a/Engine/Plugins/FX/Niagara/Source/Niagara/Private/NiagaraEmitter.cpp
+++ b/Engine/Plugins/FX/Niagara/Source/Niagara/Private/NiagaraEmitter.cpp
@@ -869,6 +869,12 @@ void UNiagaraEmitter::UpdateFromMergedCopy(const INiagaraMergeManager& MergeMana
 	EmitterUpdateScriptProps.Script->RapidIterationParameters.AddOnChangedHandler(
 		FNiagaraParameterStore::FOnChanged::FDelegate::CreateUObject(this, &UNiagaraEmitter::ScriptRapidIterationParameterChanged));
 
+	ReouterMergedObject(this, MergedEmitter->GPUComputeScript);
+	GPUComputeScript->RapidIterationParameters.RemoveAllOnChangedHandlers(this);
+	GPUComputeScript = MergedEmitter->GPUComputeScript;
+	GPUComputeScript->RapidIterationParameters.AddOnChangedHandler(
+		FNiagaraParameterStore::FOnChanged::FDelegate::CreateUObject(this, &UNiagaraEmitter::ScriptRapidIterationParameterChanged));
+
 	// Copy event handlers
 	for (FNiagaraEventScriptProperties& EventScriptProperties : EventHandlerScriptProps)
 	{
@@ -931,7 +937,7 @@ bool UNiagaraEmitter::SetUniqueEmitterName(const FString& InName)
 		{
 			// Also rename the underlying uobject to keep things consistent.
 			FName UniqueObjectName = MakeUniqueObjectName(GetOuter(), UNiagaraEmitter::StaticClass(), *InName);
-			Rename(*UniqueObjectName.ToString());
+			Rename(*UniqueObjectName.ToString(), GetOuter(), REN_ForceNoResetLoaders);
 		}
 
 #if WITH_EDITORONLY_DATA
diff --git a/Engine/Plugins/FX/Niagara/Source/Niagara/Private/NiagaraEmitterHandle.cpp b/Engine/Plugins/FX/Niagara/Source/Niagara/Private/NiagaraEmitterHandle.cpp
--- a/Engine/Plugins/FX/Niagara/Source/Niagara/Private/NiagaraEmitterHandle.cpp
+++ b/Engine/Plugins/FX/Niagara/Source/Niagara/Private/NiagaraEmitterHandle.cpp
@@ -143,7 +143,6 @@ void FNiagaraEmitterHandle::ConditionalPostLoad(int32 NiagaraCustomVersion)
 			{
 				Source_DEPRECATED->ConditionalPostLoad();
 				Instance->Parent = Source_DEPRECATED;
-				Instance->Parent->Rename(nullptr, Instance, REN_ForceNoResetLoaders);
 				Source_DEPRECATED = nullptr;
 			}
 			if (LastMergedSource_DEPRECATED != nullptr)
