Fixed crash when static mesh LODs are stripped, or not yet streamed in the niaraga renderer.

diff --git a/Engine/Plugins/FX/Niagara/Source/Niagara/Private/NiagaraRendererMeshes.cpp b/Engine/Plugins/FX/Niagara/Source/Niagara/Private/NiagaraRendererMeshes.cpp
--- a/Engine/Plugins/FX/Niagara/Source/Niagara/Private/NiagaraRendererMeshes.cpp
+++ b/Engine/Plugins/FX/Niagara/Source/Niagara/Private/NiagaraRendererMeshes.cpp
@@ -103,6 +103,8 @@ FNiagaraRendererMeshes::FNiagaraRendererMeshes(ERHIFeatureLevel::Type FeatureLev
 	MaterialParamValidMask |= MaterialParamOffset1 == -1 ? 0 : 0x2;
 	MaterialParamValidMask |= MaterialParamOffset2 == -1 ? 0 : 0x4;
 	MaterialParamValidMask |= MaterialParamOffset3 == -1 ? 0 : 0x8;
+
+	MeshMinimumLOD = Properties->ParticleMesh->MinLOD.GetValueForFeatureLevel(FeatureLevel);
 }
 
 void FNiagaraRendererMeshes::ReleaseRenderThreadResources()
@@ -204,7 +206,8 @@ void FNiagaraRendererMeshes::GetDynamicMeshElements(const TArray<const FSceneVie
 			{
 				const FSceneView* View = Views[ViewIndex];
 				//TODO: Work out LOD and pass through in dynamic data.
-				const FStaticMeshLODResources& LODModel = MeshRenderData->LODResources[0];
+				const int32 LODIndex = FMath::Max<int32>((int32)MeshRenderData->CurrentFirstLODIdx, MeshMinimumLOD);
+				const FStaticMeshLODResources& LODModel = MeshRenderData->LODResources[LODIndex];
 
 				FNiagaraMeshCollectorResourcesMesh& CollectorResources = Collector.AllocateOneFrameResource<FNiagaraMeshCollectorResourcesMesh>();
 				SetupVertexFactory(&CollectorResources.VertexFactory, LODModel);
@@ -435,7 +438,8 @@ FNiagaraDynamicDataBase *FNiagaraRendererMeshes::GenerateDynamicData(const FNiag
 		DynamicData = new FNiagaraDynamicDataMesh(Emitter);
 
 		//TODO: Mesh LODs.
-		const FStaticMeshLODResources& LODModel = Properties->ParticleMesh->RenderData->LODResources[0];
+		const int32 LODIndex = FMath::Max<int32>((int32)MeshRenderData->CurrentFirstLODIdx, MeshMinimumLOD);
+		const FStaticMeshLODResources& LODModel = Properties->ParticleMesh->RenderData->LODResources[LODIndex];
 
 		check(BaseMaterials_GT.Num() == LODModel.Sections.Num());
 
diff --git a/Engine/Plugins/FX/Niagara/Source/Niagara/Public/NiagaraRendererMeshes.h b/Engine/Plugins/FX/Niagara/Source/Niagara/Public/NiagaraRendererMeshes.h
--- a/Engine/Plugins/FX/Niagara/Source/Niagara/Public/NiagaraRendererMeshes.h
+++ b/Engine/Plugins/FX/Niagara/Source/Niagara/Public/NiagaraRendererMeshes.h
@@ -58,4 +58,6 @@ private:
 	int32 NormalizedAgeOffset;
 	int32 MaterialRandomOffset;
 	int32 CustomSortingOffset;
+
+	int32 MeshMinimumLOD = 0;
 };
