Changing default Niagara GPU sort precision to 16 bits (through "Niagara.GPUSorting.UseMaxPrecision 0").
Fixed Niagara GPU sort using next frame particle count instead of current frame count (creating render corruption).

diff --git a/Engine/Plugins/FX/Niagara/Source/Niagara/Private/NiagaraEmitterInstanceBatcher.cpp b/Engine/Plugins/FX/Niagara/Source/Niagara/Private/NiagaraEmitterInstanceBatcher.cpp
--- a/Engine/Plugins/FX/Niagara/Source/Niagara/Private/NiagaraEmitterInstanceBatcher.cpp
+++ b/Engine/Plugins/FX/Niagara/Source/Niagara/Private/NiagaraEmitterInstanceBatcher.cpp
@@ -342,8 +342,6 @@ void NiagaraEmitterInstanceBatcher::DispatchAllOnCompute(FOverlappableTicks& Ove
 void NiagaraEmitterInstanceBatcher::PostRenderOpaque(FRHICommandListImmediate& RHICmdList, const FUniformBufferRHIParamRef ViewUniformBuffer, const class FShaderParametersMetadata* SceneTexturesUniformBufferStruct, FUniformBufferRHIParamRef SceneTexturesUniformBuffer)
 {
 	ExecuteAll(RHICmdList, ViewUniformBuffer);
-	// Sort buffer only after the tick, so that the sorting and rendering stay coherent.
-	SortGPUParticles(RHICmdList);
 }
 
 void NiagaraEmitterInstanceBatcher::ExecuteAll(FRHICommandList &RHICmdList, FUniformBufferRHIParamRef ViewUniformBuffer)
@@ -590,6 +588,9 @@ bool NiagaraEmitterInstanceBatcher::UsesGlobalDistanceField() const
 void NiagaraEmitterInstanceBatcher::PreRender(FRHICommandListImmediate& RHICmdList, const class FGlobalDistanceFieldParameterData* GlobalDistanceFieldParameterData)
 {
 	GlobalDistanceFieldParams = GlobalDistanceFieldParameterData ? *GlobalDistanceFieldParameterData : FGlobalDistanceFieldParameterData();
+
+	// Sort buffer after mesh batches are issued, before tick (which will change the GPU instance count).
+	SortGPUParticles(RHICmdList);
 }
 
 int32 NiagaraEmitterInstanceBatcher::AddSortedGPUSimulation(const FNiagaraGPUSortInfo& SortInfo)
diff --git a/Engine/Plugins/FX/Niagara/Source/NiagaraVertexFactories/Private/NiagaraSortingGPU.cpp b/Engine/Plugins/FX/Niagara/Source/NiagaraVertexFactories/Private/NiagaraSortingGPU.cpp
--- a/Engine/Plugins/FX/Niagara/Source/NiagaraVertexFactories/Private/NiagaraSortingGPU.cpp
+++ b/Engine/Plugins/FX/Niagara/Source/NiagaraVertexFactories/Private/NiagaraSortingGPU.cpp
@@ -16,11 +16,11 @@ static FAutoConsoleVariableRef CVarNiagaraGPUSorting(
 	ECVF_Default
 );
 
-int32 GNiagaraGPUSortingUseMaxPrecision = 1;
+int32 GNiagaraGPUSortingUseMaxPrecision = 0;
 static FAutoConsoleVariableRef CVarNiagaraGPUSortinUseMaxPrecision(
 	TEXT("Niagara.GPUSorting.UseMaxPrecision"),
 	GNiagaraGPUSortingUseMaxPrecision,
-	TEXT("Wether sorting using fp32 instead of fp16. (default=1)"),
+	TEXT("Wether sorting using fp32 instead of fp16. (default=0)"),
 	ECVF_Default
 );
 
