Fix CIS

#rnx

diff --git a/Engine/Plugins/FX/Niagara/Source/NiagaraEditorWidgets/Private/Stack/SNiagaraStackItemGroupAddMenu.cpp b/Engine/Plugins/FX/Niagara/Source/NiagaraEditorWidgets/Private/Stack/SNiagaraStackItemGroupAddMenu.cpp
--- a/Engine/Plugins/FX/Niagara/Source/NiagaraEditorWidgets/Private/Stack/SNiagaraStackItemGroupAddMenu.cpp
+++ b/Engine/Plugins/FX/Niagara/Source/NiagaraEditorWidgets/Private/Stack/SNiagaraStackItemGroupAddMenu.cpp
@@ -135,4 +135,6 @@ void SNiagaraStackItemGroupAddMenu::OnLibraryToggleChanged(ECheckBoxState CheckS
 ECheckBoxState SNiagaraStackItemGroupAddMenu::LibraryToggleIsChecked() const
 {
 	return SNiagaraStackItemGroupAddMenu::bIncludeNonLibraryScripts ? ECheckBoxState::Unchecked : ECheckBoxState::Checked;
-}
\ No newline at end of file
+}
+
+#undef LOCTEXT_NAMESPACE
\ No newline at end of file
