Added script usage enum (e.g. update script, spawn script) as constant for static switches

diff --git a/Engine/Plugins/FX/Niagara/Source/Niagara/Classes/NiagaraConstants.h b/Engine/Plugins/FX/Niagara/Source/Niagara/Classes/NiagaraConstants.h
--- a/Engine/Plugins/FX/Niagara/Source/Niagara/Classes/NiagaraConstants.h
+++ b/Engine/Plugins/FX/Niagara/Source/Niagara/Classes/NiagaraConstants.h
@@ -102,7 +102,7 @@
 #define SYS_PARAM_PARTICLES_RIBBONLINKORDER              INiagaraModule::GetVar_Particles_RibbonLinkOrder()
 
 #define SYS_PARAM_INSTANCE_ALIVE                         INiagaraModule::GetVar_DataInstance_Alive()
-
+#define SYS_PARAM_SCRIPT_USAGE                           INiagaraModule::GetVar_ScriptUsage()
 #define TRANSLATOR_PARAM_BEGIN_DEFAULTS                  INiagaraModule::GetVar_BeginDefaults()
 
 struct NIAGARA_API FNiagaraConstants
diff --git a/Engine/Plugins/FX/Niagara/Source/Niagara/Private/NiagaraConstants.cpp b/Engine/Plugins/FX/Niagara/Source/Niagara/Private/NiagaraConstants.cpp
--- a/Engine/Plugins/FX/Niagara/Source/Niagara/Private/NiagaraConstants.cpp
+++ b/Engine/Plugins/FX/Niagara/Source/Niagara/Private/NiagaraConstants.cpp
@@ -77,6 +77,7 @@ void FNiagaraConstants::Init()
 		SwitchParameters.Add(SYS_PARAM_EMITTER_LOCALSPACE);
 		SwitchParameters.Add(SYS_PARAM_EMITTER_DETERMINISM);
 		SwitchParameters.Add(SYS_PARAM_EMITTER_SIMULATION_TARGET);
+		SwitchParameters.Add(SYS_PARAM_SCRIPT_USAGE);
 	}
 
 	if (UpdatedSystemParameters.Num() == 0)
diff --git a/Engine/Plugins/FX/Niagara/Source/Niagara/Private/NiagaraModule.cpp b/Engine/Plugins/FX/Niagara/Source/Niagara/Private/NiagaraModule.cpp
--- a/Engine/Plugins/FX/Niagara/Source/Niagara/Private/NiagaraModule.cpp
+++ b/Engine/Plugins/FX/Niagara/Source/Niagara/Private/NiagaraModule.cpp
@@ -145,6 +145,7 @@ FNiagaraVariable INiagaraModule::Particles_RibbonWidth;
 FNiagaraVariable INiagaraModule::Particles_RibbonTwist;
 FNiagaraVariable INiagaraModule::Particles_RibbonFacing;
 FNiagaraVariable INiagaraModule::Particles_RibbonLinkOrder;
+FNiagaraVariable INiagaraModule::ScriptUsage;
 FNiagaraVariable INiagaraModule::DataInstance_Alive;
 FNiagaraVariable INiagaraModule::Translator_BeginDefaults;
 
@@ -253,6 +254,7 @@ void INiagaraModule::StartupModule()
 	Particles_RibbonFacing = FNiagaraVariable(FNiagaraTypeDefinition::GetVec3Def(), TEXT("Particles.RibbonFacing"));
 	Particles_RibbonLinkOrder = FNiagaraVariable(FNiagaraTypeDefinition::GetFloatDef(), TEXT("Particles.RibbonLinkOrder"));
 
+	ScriptUsage = FNiagaraVariable(FNiagaraTypeDefinition::GetScriptUsageEnum(), TEXT("Script.Usage"));
 	DataInstance_Alive = FNiagaraVariable(FNiagaraTypeDefinition::GetBoolDef(), TEXT("DataInstance.Alive"));
 
 	Translator_BeginDefaults = FNiagaraVariable(FNiagaraTypeDefinition::GetParameterMapDef(), TEXT("Begin Defaults"));
@@ -469,6 +471,7 @@ UScriptStruct* FNiagaraTypeDefinition::QuatStruct;
 UEnum* FNiagaraTypeDefinition::ExecutionStateEnum;
 UEnum* FNiagaraTypeDefinition::SimulationTargetEnum;
 UEnum* FNiagaraTypeDefinition::ExecutionStateSourceEnum;
+UEnum* FNiagaraTypeDefinition::ScriptUsageEnum;
 
 FNiagaraTypeDefinition FNiagaraTypeDefinition::ParameterMapDef;
 FNiagaraTypeDefinition FNiagaraTypeDefinition::IDDef;
@@ -576,6 +579,7 @@ void FNiagaraTypeDefinition::Init()
 	ExecutionStateEnum = StaticEnum<ENiagaraExecutionState>();
 	ExecutionStateSourceEnum = StaticEnum<ENiagaraExecutionStateSource>();
 	SimulationTargetEnum = StaticEnum<ENiagaraSimTarget>();
+	ScriptUsageEnum = StaticEnum<ENiagaraScriptUsage>();
 	
 	RecreateUserDefinedTypeRegistry();
 }
diff --git a/Engine/Plugins/FX/Niagara/Source/Niagara/Public/NiagaraModule.h b/Engine/Plugins/FX/Niagara/Source/Niagara/Public/NiagaraModule.h
--- a/Engine/Plugins/FX/Niagara/Source/Niagara/Public/NiagaraModule.h
+++ b/Engine/Plugins/FX/Niagara/Source/Niagara/Public/NiagaraModule.h
@@ -132,6 +132,7 @@ public:
 	FORCEINLINE static const FNiagaraVariable&  GetVar_Emitter_SpawnRate() { return Emitter_SpawnRate; }
 	FORCEINLINE static const FNiagaraVariable&  GetVar_Emitter_SpawnInterval() { return Emitter_SpawnInterval; }
 	FORCEINLINE static const FNiagaraVariable&  GetVar_Emitter_SimulationTarget() { return Emitter_SimulationTarget; }
+	FORCEINLINE static const FNiagaraVariable&  GetVar_ScriptUsage() { return ScriptUsage; }
 	FORCEINLINE static const FNiagaraVariable&  GetVar_Emitter_InterpSpawnStartDt() { return Emitter_InterpSpawnStartDt; }
 	FORCEINLINE static const FNiagaraVariable&  GetVar_Emitter_SpawnGroup() { return Emitter_SpawnGroup; }
 
@@ -263,6 +264,7 @@ private:
 	static FNiagaraVariable Particles_RibbonFacing;
 	static FNiagaraVariable Particles_RibbonLinkOrder;
 
+	static FNiagaraVariable ScriptUsage;
 	static FNiagaraVariable DataInstance_Alive;
 	static FNiagaraVariable Translator_BeginDefaults;
 };
diff --git a/Engine/Plugins/FX/Niagara/Source/Niagara/Public/NiagaraTypes.h b/Engine/Plugins/FX/Niagara/Source/Niagara/Public/NiagaraTypes.h
--- a/Engine/Plugins/FX/Niagara/Source/Niagara/Public/NiagaraTypes.h
+++ b/Engine/Plugins/FX/Niagara/Source/Niagara/Public/NiagaraTypes.h
@@ -571,6 +571,7 @@ public:
 
 	static UEnum* GetExecutionStateEnum() { return ExecutionStateEnum; }
 	static UEnum* GetSimulationTargetEnum() { return SimulationTargetEnum; }
+	static UEnum* GetScriptUsageEnum() { return ScriptUsageEnum; }
 
 	static const FNiagaraTypeDefinition& GetCollisionEventDef() { return CollisionEventDef; }
 
@@ -619,6 +620,7 @@ private:
 	static UScriptStruct* NumericStruct;
 
 	static UEnum* SimulationTargetEnum;
+	static UEnum* ScriptUsageEnum;
 	static UEnum* ExecutionStateEnum;
 	static UEnum* ExecutionStateSourceEnum;
 
diff --git a/Engine/Plugins/FX/Niagara/Source/NiagaraEditor/Private/NiagaraHlslTranslator.cpp b/Engine/Plugins/FX/Niagara/Source/NiagaraEditor/Private/NiagaraHlslTranslator.cpp
--- a/Engine/Plugins/FX/Niagara/Source/NiagaraEditor/Private/NiagaraHlslTranslator.cpp
+++ b/Engine/Plugins/FX/Niagara/Source/NiagaraEditor/Private/NiagaraHlslTranslator.cpp
@@ -3370,6 +3370,13 @@ bool FHlslNiagaraTranslator::GetLiteralConstantVariable(FNiagaraVariable& OutVar
 			return true;
 		}
 	}
+	else if (OutVar == FNiagaraVariable(FNiagaraTypeDefinition::GetScriptUsageEnum(), TEXT("Script.Usage")))
+	{
+		FNiagaraInt32 EnumValue;
+		EnumValue.Value = (uint8)GetCurrentUsage();
+		OutVar.SetValue(EnumValue);
+		return true;
+	}
 	return false;
 }
 
