Fixed a bug where nested function calls with static switches were not always correctly compiled

diff --git a/Engine/Plugins/FX/Niagara/Source/NiagaraEditor/Private/NiagaraNodeFunctionCall.cpp b/Engine/Plugins/FX/Niagara/Source/NiagaraEditor/Private/NiagaraNodeFunctionCall.cpp
--- a/Engine/Plugins/FX/Niagara/Source/NiagaraEditor/Private/NiagaraNodeFunctionCall.cpp
+++ b/Engine/Plugins/FX/Niagara/Source/NiagaraEditor/Private/NiagaraNodeFunctionCall.cpp
@@ -381,6 +381,7 @@ void UNiagaraNodeFunctionCall::Compile(class FHlslNiagaraTranslator* Translator,
 
 		TArray<UEdGraphPin*> CallerInputPins;
 		GetInputPins(CallerInputPins);
+		FNiagaraEditorUtilities::SetStaticSwitchConstants(GetCalledGraph(), CallerInputPins);
 
 		UNiagaraScriptSource* Source = CastChecked<UNiagaraScriptSource>(FunctionScript->GetSource());
 		UNiagaraGraph* FunctionGraph = Source->NodeGraph;
@@ -679,6 +680,10 @@ void UNiagaraNodeFunctionCall::BuildParameterMapHistory(FNiagaraParameterMapHist
 			OutputNode = FunctionGraph->FindOutputNode(ENiagaraScriptUsage::DynamicInput);
 		}
 
+		TArray<UEdGraphPin*> InputPins;
+		GetInputPins(InputPins);
+		FNiagaraEditorUtilities::SetStaticSwitchConstants(FunctionGraph, InputPins);
+
 		int32 ParamMapIdx = INDEX_NONE;
 		uint32 NodeIdx = INDEX_NONE;
 		const UEdGraphPin* CandidateParamMapPin = GetInputPin(0);
