UE4 - Fix deprecation warnings
#rnx

diff --git a/Engine/Plugins/FX/Niagara/Source/NiagaraVertexFactories/Private/NiagaraDrawIndirect.cpp b/Engine/Plugins/FX/Niagara/Source/NiagaraVertexFactories/Private/NiagaraDrawIndirect.cpp
--- a/Engine/Plugins/FX/Niagara/Source/NiagaraVertexFactories/Private/NiagaraDrawIndirect.cpp
+++ b/Engine/Plugins/FX/Niagara/Source/NiagaraVertexFactories/Private/NiagaraDrawIndirect.cpp
@@ -36,9 +36,9 @@ bool FNiagaraDrawIndirectArgsGenCS::Serialize(FArchive& Ar)
 	return bShaderHasOutdatedParameters;
 }
 
-void FNiagaraDrawIndirectArgsGenCS::SetOutput(FRHICommandList& RHICmdList, FUnorderedAccessViewRHIParamRef DrawIndirectArgsUAV, FUnorderedAccessViewRHIParamRef InstanceCountsUAV)
+void FNiagaraDrawIndirectArgsGenCS::SetOutput(FRHICommandList& RHICmdList, FRHIUnorderedAccessView* DrawIndirectArgsUAV, FRHIUnorderedAccessView* InstanceCountsUAV)
 {
-	FComputeShaderRHIParamRef ComputeShaderRHI = GetComputeShader();
+	FRHIComputeShader* ComputeShaderRHI = GetComputeShader();
 	if (DrawIndirectArgsParam.IsBound())
 	{
 		RHICmdList.SetUAVParameter(ComputeShaderRHI, DrawIndirectArgsParam.GetUAVIndex(), DrawIndirectArgsUAV);
@@ -51,7 +51,7 @@ void FNiagaraDrawIndirectArgsGenCS::SetOutput(FRHICommandList& RHICmdList, FUnor
 
 void FNiagaraDrawIndirectArgsGenCS::SetParameters(FRHICommandList& RHICmdList, FShaderResourceViewRHIParamRef TaskInfosBuffer, int32 NumArgGenTasks, int32 NumInstanceCountClearTasks)
 {
-	FComputeShaderRHIParamRef ComputeShaderRHI = GetComputeShader();
+	FRHIComputeShader* ComputeShaderRHI = GetComputeShader();
 
 	RHICmdList.SetShaderResourceViewParameter(ComputeShaderRHI, TaskInfosParam.GetBaseIndex(), TaskInfosBuffer);
 
@@ -61,18 +61,17 @@ void FNiagaraDrawIndirectArgsGenCS::SetParameters(FRHICommandList& RHICmdList, F
 
 void FNiagaraDrawIndirectArgsGenCS::UnbindBuffers(FRHICommandList& RHICmdList)
 {
-	FComputeShaderRHIParamRef ComputeShaderRHI = GetComputeShader();
+	FRHIComputeShader* ComputeShaderRHI = GetComputeShader();
 	if (TaskInfosParam.IsBound())
 	{
 		RHICmdList.SetShaderResourceViewParameter(ComputeShaderRHI, TaskInfosParam.GetBaseIndex(), FShaderResourceViewRHIParamRef());
 	}
 	if (DrawIndirectArgsParam.IsBound())
 	{
-		RHICmdList.SetUAVParameter(ComputeShaderRHI, DrawIndirectArgsParam.GetUAVIndex(), FUnorderedAccessViewRHIParamRef());
+		RHICmdList.SetUAVParameter(ComputeShaderRHI, DrawIndirectArgsParam.GetUAVIndex(), nullptr);
 	}
 	if (InstanceCountsParam.IsBound())
 	{
-		RHICmdList.SetUAVParameter(ComputeShaderRHI, InstanceCountsParam.GetUAVIndex(), FUnorderedAccessViewRHIParamRef());
+		RHICmdList.SetUAVParameter(ComputeShaderRHI, InstanceCountsParam.GetUAVIndex(), nullptr);
 	}
 }
-
diff --git a/Engine/Source/Runtime/RHI/Public/RHIValidation.h b/Engine/Source/Runtime/RHI/Public/RHIValidation.h
--- a/Engine/Source/Runtime/RHI/Public/RHIValidation.h
+++ b/Engine/Source/Runtime/RHI/Public/RHIValidation.h
@@ -643,13 +643,6 @@ public:
 	* @param Flags - ETextureCreateFlags creation flags
 	*/
 	// FlushType: Wait RHI Thread
-	PRAGMA_DISABLE_DEPRECATION_WARNINGS
-	virtual FTexture2DArrayRHIRef RHICreateTexture2DArray(uint32 SizeX, uint32 SizeY, uint32 SizeZ, uint8 Format, uint32 NumMips, uint32 Flags, FRHIResourceCreateInfo& CreateInfo) override final
-	{
-		return RHI->RHICreateTexture2DArray(SizeX, SizeY, SizeZ, Format, NumMips, Flags, CreateInfo);
-	}
-	PRAGMA_ENABLE_DEPRECATION_WARNINGS
-
 	virtual FTexture2DArrayRHIRef RHICreateTexture2DArray(uint32 SizeX, uint32 SizeY, uint32 SizeZ, uint8 Format, uint32 NumMips, uint32 NumSamples, uint32 Flags, FRHIResourceCreateInfo& CreateInfo)
 	{
 		return RHI->RHICreateTexture2DArray(SizeX, SizeY, SizeZ, Format, NumMips, NumSamples, Flags, CreateInfo);
@@ -1419,13 +1412,6 @@ public:
 		return RHI->RHICreateTextureExternal2D_RenderThread(RHICmdList, SizeX, SizeY, Format, NumMips, NumSamples, Flags, CreateInfo);
 	}
 
-	PRAGMA_DISABLE_DEPRECATION_WARNINGS
-	virtual FTexture2DArrayRHIRef RHICreateTexture2DArray_RenderThread(class FRHICommandListImmediate& RHICmdList, uint32 SizeX, uint32 SizeY, uint32 SizeZ, uint8 Format, uint32 NumMips, uint32 Flags, FRHIResourceCreateInfo& CreateInfo) override final
-	{
-		return RHI->RHICreateTexture2DArray_RenderThread(RHICmdList, SizeX, SizeY, SizeZ, Format, NumMips, Flags, CreateInfo);
-	}
-	PRAGMA_ENABLE_DEPRECATION_WARNINGS
-
 	virtual FTexture3DRHIRef RHICreateTexture3D_RenderThread(class FRHICommandListImmediate& RHICmdList, uint32 SizeX, uint32 SizeY, uint32 SizeZ, uint8 Format, uint32 NumMips, uint32 Flags, FRHIResourceCreateInfo& CreateInfo) override final
 	{
 		return RHI->RHICreateTexture3D_RenderThread(RHICmdList, SizeX, SizeY, SizeZ, Format, NumMips, Flags, CreateInfo);
