UE4 - Deprecation warning fix
#rnx

diff --git a/Engine/Plugins/FX/Niagara/Source/NiagaraVertexFactories/Private/NiagaraDrawIndirect.cpp b/Engine/Plugins/FX/Niagara/Source/NiagaraVertexFactories/Private/NiagaraDrawIndirect.cpp
--- a/Engine/Plugins/FX/Niagara/Source/NiagaraVertexFactories/Private/NiagaraDrawIndirect.cpp
+++ b/Engine/Plugins/FX/Niagara/Source/NiagaraVertexFactories/Private/NiagaraDrawIndirect.cpp
@@ -49,7 +49,7 @@ void FNiagaraDrawIndirectArgsGenCS::SetOutput(FRHICommandList& RHICmdList, FRHIU
 	}
 }
 
-void FNiagaraDrawIndirectArgsGenCS::SetParameters(FRHICommandList& RHICmdList, FShaderResourceViewRHIParamRef TaskInfosBuffer, int32 NumArgGenTasks, int32 NumInstanceCountClearTasks)
+void FNiagaraDrawIndirectArgsGenCS::SetParameters(FRHICommandList& RHICmdList, FRHIShaderResourceView* TaskInfosBuffer, int32 NumArgGenTasks, int32 NumInstanceCountClearTasks)
 {
 	FRHIComputeShader* ComputeShaderRHI = GetComputeShader();
 
@@ -64,7 +64,7 @@ void FNiagaraDrawIndirectArgsGenCS::UnbindBuffers(FRHICommandList& RHICmdList)
 	FRHIComputeShader* ComputeShaderRHI = GetComputeShader();
 	if (TaskInfosParam.IsBound())
 	{
-		RHICmdList.SetShaderResourceViewParameter(ComputeShaderRHI, TaskInfosParam.GetBaseIndex(), FShaderResourceViewRHIParamRef());
+		RHICmdList.SetShaderResourceViewParameter(ComputeShaderRHI, TaskInfosParam.GetBaseIndex(), nullptr);
 	}
 	if (DrawIndirectArgsParam.IsBound())
 	{
