Enabled the stack ui to resolve constants for static switches to correctly hide unreachable module parameters. Also removed some unused code.

#tests niagara automation tests

diff --git a/Engine/Plugins/FX/Niagara/Source/NiagaraEditor/Private/NiagaraEditorUtilities.cpp b/Engine/Plugins/FX/Niagara/Source/NiagaraEditor/Private/NiagaraEditorUtilities.cpp
--- a/Engine/Plugins/FX/Niagara/Source/NiagaraEditor/Private/NiagaraEditorUtilities.cpp
+++ b/Engine/Plugins/FX/Niagara/Source/NiagaraEditor/Private/NiagaraEditorUtilities.cpp
@@ -34,6 +34,7 @@
 #include "NiagaraNodeWriteDataSet.h"
 #include "NiagaraNodeStaticSwitch.h"
 #include "NiagaraNodeFunctionCall.h"
+#include "NiagaraParameterMapHistory.h"
 
 #define LOCTEXT_NAMESPACE "FNiagaraEditorUtilities"
 
@@ -622,7 +623,7 @@ void FNiagaraEditorUtilities::FixUpNumericPins(const UEdGraphSchema_Niagara* Sch
 	TraverseGraphFromOutputDepthFirst(Schema, Node, FixUpVisitor);
 }
 
-void FNiagaraEditorUtilities::SetStaticSwitchConstants(UNiagaraGraph* Graph, const TArray<UEdGraphPin*>& CallInputs)
+void FNiagaraEditorUtilities::SetStaticSwitchConstants(UNiagaraGraph* Graph, const TArray<UEdGraphPin*>& CallInputs, FCompileConstantResolver* ConstantResolver)
 {
 	for (UEdGraphNode* Node : Graph->Nodes)
 	{
@@ -631,19 +632,25 @@ void FNiagaraEditorUtilities::SetStaticSwitchConstants(UNiagaraGraph* Graph, con
 		if (SwitchNode)
 		{
 			SwitchNode->ClearSwitchValue();
-			for (UEdGraphPin* InputPin : CallInputs)
+			if (SwitchNode->IsSetByCompiler() && ConstantResolver)
 			{
-				if (InputPin->GetFName().IsEqual(SwitchNode->InputParameterName))
+				SwitchNode->SetSwitchValue(*ConstantResolver);
+			}
+			else
+			{
+				for (UEdGraphPin* InputPin : CallInputs)
 				{
-					int32 SwitchValue = 0;
-					if (ResolveConstantValue(InputPin, SwitchValue))
+					if (InputPin->GetFName().IsEqual(SwitchNode->InputParameterName))
 					{
-						SwitchNode->SetSwitchValue(SwitchValue);
+						int32 SwitchValue = 0;
+						if (ResolveConstantValue(InputPin, SwitchValue))
+						{
+							SwitchNode->SetSwitchValue(SwitchValue);
+						}
+						break;
 					}
-					break;
 				}
 			}
-			
 		}
 
 		// if there is a function node, it might have delegated some of the static switch values inside its script graph
@@ -858,43 +865,6 @@ void FNiagaraEditorUtilities::PreprocessFunctionGraph(const UEdGraphSchema_Niaga
 	FNiagaraEditorUtilities::SetStaticSwitchConstants(Graph, CallInputs);
 }
 
-const TMap<FNiagaraVariable, FNiagaraGraphParameterReferenceCollection>& FNiagaraEditorUtilities::GetCompiledGraphParameterMapReferences(UNiagaraGraph* Graph)
-{
-	const UEdGraphSchema_Niagara* Schema = GetDefault<UEdGraphSchema_Niagara>();
-
-	// We clone the graph to prevent any destructive changes to the source
-	UNiagaraGraph* NodeGraphDeepCopy = CastChecked<UNiagaraGraph>(FEdGraphUtilities::CloneGraph(Graph, GetTransientPackage(), 0, true));
-
-	// We copy some transient data over here because the copied graph loses some data. Since the result of this function is
-	// also used to build the ui, the user will otherwise lose changes to yet unused parameters (e.g. in the metadata panel).
-	CastChecked<UNiagaraGraph>(Graph)->CopyCachedReferencesMap(NodeGraphDeepCopy);
-
-	TArray<UNiagaraNodeFunctionCall*> Nodes;
-	NodeGraphDeepCopy->GetNodesOfClass(Nodes);
-	for (int32 i = 0; i < Nodes.Num(); i++)
-	{
-		UNiagaraNodeFunctionCall* Node = Nodes[i];
-		ENiagaraScriptUsage ScriptUsage = Node->GetCalledUsage();
-		UNiagaraGraph* FunctionGraph = Node->GetCalledGraph();
-		if (FunctionGraph)
-		{
-			// We append all function nodes we find to the end of the queue, basically traversing the graph breadth-first
-			TArray<UNiagaraNodeFunctionCall*> SubFunctionNodes;
-			FunctionGraph->GetNodesOfClass(SubFunctionNodes);
-			for (UNiagaraNodeFunctionCall* SubNode : SubFunctionNodes)
-			{
-				Nodes.AddUnique(SubNode);
-			}
-
-			// We set the static switch values based on the function call inputs to prepare the parameter map history generation
-			TArray<UEdGraphPin*> CallInputs;
-			Node->GetInputPins(CallInputs);
-			SetStaticSwitchConstants(FunctionGraph, CallInputs);
-		}
-	}
-	return NodeGraphDeepCopy->GetParameterReferenceMap();
-}
-
 void FNiagaraEditorUtilities::GetFilteredScriptAssets(FGetFilteredScriptAssetsOptions InFilter, TArray<FAssetData>& OutFilteredScriptAssets)
 {
 	FARFilter ScriptFilter;
diff --git a/Engine/Plugins/FX/Niagara/Source/NiagaraEditor/Private/NiagaraGraph.cpp b/Engine/Plugins/FX/Niagara/Source/NiagaraEditor/Private/NiagaraGraph.cpp
--- a/Engine/Plugins/FX/Niagara/Source/NiagaraEditor/Private/NiagaraGraph.cpp
+++ b/Engine/Plugins/FX/Niagara/Source/NiagaraEditor/Private/NiagaraGraph.cpp
@@ -659,7 +659,7 @@ TArray<UEdGraphNode*> UNiagaraGraph::FindReachbleNodes() const
 		}
 		
 		UNiagaraNodeStaticSwitch* SwitchNode = Cast<UNiagaraNodeStaticSwitch>(Node);
-		if (SwitchNode && !SwitchNode->IsSetByCompiler())
+		if (SwitchNode)
 		{
 			TArray<UEdGraphPin*> OutPins;
 			SwitchNode->GetOutputPins(OutPins);
diff --git a/Engine/Plugins/FX/Niagara/Source/NiagaraEditor/Private/NiagaraNodeEmitter.cpp b/Engine/Plugins/FX/Niagara/Source/NiagaraEditor/Private/NiagaraNodeEmitter.cpp
--- a/Engine/Plugins/FX/Niagara/Source/NiagaraEditor/Private/NiagaraNodeEmitter.cpp
+++ b/Engine/Plugins/FX/Niagara/Source/NiagaraEditor/Private/NiagaraNodeEmitter.cpp
@@ -308,6 +308,7 @@ void UNiagaraNodeEmitter::BuildParameterMapHistory(FNiagaraParameterMapHistoryBu
 
 			// Build up a new parameter map history with all the child graph nodes..
 			FNiagaraParameterMapHistoryBuilder ChildBuilder;
+			ChildBuilder.ConstantResolver = OutHistory.ConstantResolver;
 			ChildBuilder.RegisterEncounterableVariables(OutHistory.GetEncounterableVariables());
 			ChildBuilder.EnableScriptWhitelist(true, GetUsage());
 			FString LocalEmitterName = TEXT("Emitter");
diff --git a/Engine/Plugins/FX/Niagara/Source/NiagaraEditor/Private/NiagaraNodeFunctionCall.cpp b/Engine/Plugins/FX/Niagara/Source/NiagaraEditor/Private/NiagaraNodeFunctionCall.cpp
--- a/Engine/Plugins/FX/Niagara/Source/NiagaraEditor/Private/NiagaraNodeFunctionCall.cpp
+++ b/Engine/Plugins/FX/Niagara/Source/NiagaraEditor/Private/NiagaraNodeFunctionCall.cpp
@@ -682,7 +682,7 @@ void UNiagaraNodeFunctionCall::BuildParameterMapHistory(FNiagaraParameterMapHist
 
 		TArray<UEdGraphPin*> InputPins;
 		GetInputPins(InputPins);
-		FNiagaraEditorUtilities::SetStaticSwitchConstants(FunctionGraph, InputPins);
+		FNiagaraEditorUtilities::SetStaticSwitchConstants(FunctionGraph, InputPins, &OutHistory.ConstantResolver);
 
 		int32 ParamMapIdx = INDEX_NONE;
 		uint32 NodeIdx = INDEX_NONE;
diff --git a/Engine/Plugins/FX/Niagara/Source/NiagaraEditor/Private/NiagaraNodeStaticSwitch.cpp b/Engine/Plugins/FX/Niagara/Source/NiagaraEditor/Private/NiagaraNodeStaticSwitch.cpp
--- a/Engine/Plugins/FX/Niagara/Source/NiagaraEditor/Private/NiagaraNodeStaticSwitch.cpp
+++ b/Engine/Plugins/FX/Niagara/Source/NiagaraEditor/Private/NiagaraNodeStaticSwitch.cpp
@@ -50,6 +50,31 @@ void UNiagaraNodeStaticSwitch::SetSwitchValue(int Value)
 	SwitchValue = Value;
 }
 
+void UNiagaraNodeStaticSwitch::SetSwitchValue(const FCompileConstantResolver& ConstantResolver)
+{
+	if (!IsSetByCompiler())
+	{
+		return;
+	}
+	IsValueSet = false;
+
+	const FNiagaraVariable* Found = FNiagaraConstants::FindStaticSwitchConstant(SwitchTypeData.SwitchConstant);
+	FNiagaraVariable Constant = Found ? *Found : FNiagaraVariable();
+	if (Found && ConstantResolver.ResolveConstant(Constant))
+	{
+		if (SwitchTypeData.SwitchType == ENiagaraStaticSwitchType::Bool)
+		{
+			SwitchValue = Constant.GetValue<bool>();
+			IsValueSet = true;
+		}
+		else if (SwitchTypeData.SwitchType == ENiagaraStaticSwitchType::Integer || SwitchTypeData.SwitchType == ENiagaraStaticSwitchType::Enum)
+		{
+			SwitchValue = Constant.GetValue<int32>();
+			IsValueSet = true;
+		}
+	}
+}
+
 void UNiagaraNodeStaticSwitch::ClearSwitchValue()
 {
 	IsValueSet = false;
diff --git a/Engine/Plugins/FX/Niagara/Source/NiagaraEditor/Private/NiagaraNodeStaticSwitch.h b/Engine/Plugins/FX/Niagara/Source/NiagaraEditor/Private/NiagaraNodeStaticSwitch.h
--- a/Engine/Plugins/FX/Niagara/Source/NiagaraEditor/Private/NiagaraNodeStaticSwitch.h
+++ b/Engine/Plugins/FX/Niagara/Source/NiagaraEditor/Private/NiagaraNodeStaticSwitch.h
@@ -57,6 +57,7 @@ public:
 	void OnSwitchParameterTypeChanged(const FNiagaraTypeDefinition& OldType);
 
 	void SetSwitchValue(int Value);
+	void SetSwitchValue(const FCompileConstantResolver& ConstantResolver);
 	void ClearSwitchValue();
 	/** If true then the value of this static switch is not set by the user but directly by the compiler via one of the engine constants (e.g. Emitter.Determinism). */
 	bool IsSetByCompiler() const;
diff --git a/Engine/Plugins/FX/Niagara/Source/NiagaraEditor/Private/NiagaraParameterMapHistory.cpp b/Engine/Plugins/FX/Niagara/Source/NiagaraEditor/Private/NiagaraParameterMapHistory.cpp
--- a/Engine/Plugins/FX/Niagara/Source/NiagaraEditor/Private/NiagaraParameterMapHistory.cpp
+++ b/Engine/Plugins/FX/Niagara/Source/NiagaraEditor/Private/NiagaraParameterMapHistory.cpp
@@ -1375,4 +1375,26 @@ void FNiagaraParameterMapHistoryBuilder::BuildCurrentAliases()
 	}
 }
 
+bool FCompileConstantResolver::ResolveConstant(FNiagaraVariable& OutConstant) const
+{
+	if (Emitter && OutConstant == FNiagaraVariable(FNiagaraTypeDefinition::GetBoolDef(), TEXT("Emitter.Localspace")))
+	{
+		OutConstant.SetValue(Emitter->bLocalSpace ? FNiagaraBool(true) : FNiagaraBool(false));
+		return true;
+	}
+	if (Emitter && OutConstant == FNiagaraVariable(FNiagaraTypeDefinition::GetBoolDef(), TEXT("Emitter.Determinism")))
+	{
+		OutConstant.SetValue(Emitter->bDeterminism ? FNiagaraBool(true) : FNiagaraBool(false));
+		return true;
+	}
+	if (Emitter && OutConstant == FNiagaraVariable(FNiagaraTypeDefinition::GetSimulationTargetEnum(), TEXT("Emitter.SimulationTarget")))
+	{
+		FNiagaraInt32 EnumValue;
+		EnumValue.Value = (uint8)Emitter->SimTarget;
+		OutConstant.SetValue(EnumValue);
+		return true;
+	}
+	return false;
+}
+
 #undef LOCTEXT_NAMESPACE
diff --git a/Engine/Plugins/FX/Niagara/Source/NiagaraEditor/Private/NiagaraScriptSource.cpp b/Engine/Plugins/FX/Niagara/Source/NiagaraEditor/Private/NiagaraScriptSource.cpp
--- a/Engine/Plugins/FX/Niagara/Source/NiagaraEditor/Private/NiagaraScriptSource.cpp
+++ b/Engine/Plugins/FX/Niagara/Source/NiagaraEditor/Private/NiagaraScriptSource.cpp
@@ -193,8 +193,7 @@ void InitializeNewRapidIterationParametersForNode(const UEdGraphSchema_Niagara*
 	if (FunctionCallNode != nullptr)
 	{
 		TArray<const UEdGraphPin*> FunctionInputPins;
-		TSet<const UEdGraphPin*> HiddenInputPins;
-		FNiagaraStackGraphUtilities::GetStackFunctionInputPins(*FunctionCallNode, FunctionInputPins, HiddenInputPins, FNiagaraStackGraphUtilities::ENiagaraGetStackFunctionInputPinsOptions::ModuleInputsOnly, false);
+		FNiagaraStackGraphUtilities::GetStackFunctionInputPins(*FunctionCallNode, FunctionInputPins, FNiagaraStackGraphUtilities::ENiagaraGetStackFunctionInputPinsOptions::ModuleInputsOnly, false);
 		for (const UEdGraphPin* FunctionInputPin : FunctionInputPins)
 		{
 			FNiagaraTypeDefinition InputType = Schema->PinToTypeDefinition(FunctionInputPin);
diff --git a/Engine/Plugins/FX/Niagara/Source/NiagaraEditor/Private/NiagaraStackFunctionInputBinder.cpp b/Engine/Plugins/FX/Niagara/Source/NiagaraEditor/Private/NiagaraStackFunctionInputBinder.cpp
--- a/Engine/Plugins/FX/Niagara/Source/NiagaraEditor/Private/NiagaraStackFunctionInputBinder.cpp
+++ b/Engine/Plugins/FX/Niagara/Source/NiagaraEditor/Private/NiagaraStackFunctionInputBinder.cpp
@@ -91,8 +91,7 @@ bool FNiagaraStackFunctionInputBinder::TryBindInternal(
 	FunctionCallNode = InFunctionCallNode;
 
 	TArray<const UEdGraphPin*> InputPins;
-	TSet<const UEdGraphPin*> HiddenPins;
-	FNiagaraStackGraphUtilities::GetStackFunctionInputPins(*FunctionCallNode, InputPins, HiddenPins, FNiagaraStackGraphUtilities::ENiagaraGetStackFunctionInputPinsOptions::ModuleInputsOnly);
+	FNiagaraStackGraphUtilities::GetStackFunctionInputPins(*FunctionCallNode, InputPins, FNiagaraStackGraphUtilities::ENiagaraGetStackFunctionInputPinsOptions::ModuleInputsOnly);
 	
 	bool bInputFound = false;
 	const UEdGraphSchema_Niagara* Schema = GetDefault<UEdGraphSchema_Niagara>();
diff --git a/Engine/Plugins/FX/Niagara/Source/NiagaraEditor/Private/ViewModels/Stack/NiagaraStackFunctionInputCollection.cpp b/Engine/Plugins/FX/Niagara/Source/NiagaraEditor/Private/ViewModels/Stack/NiagaraStackFunctionInputCollection.cpp
--- a/Engine/Plugins/FX/Niagara/Source/NiagaraEditor/Private/ViewModels/Stack/NiagaraStackFunctionInputCollection.cpp
+++ b/Engine/Plugins/FX/Niagara/Source/NiagaraEditor/Private/ViewModels/Stack/NiagaraStackFunctionInputCollection.cpp
@@ -6,6 +6,7 @@
 #include "NiagaraNodeFunctionCall.h"
 #include "NiagaraEmitterEditorData.h"
 #include "EdGraphSchema_Niagara.h"
+#include "ViewModels/NiagaraEmitterViewModel.h"
 #include "ViewModels/Stack/NiagaraStackGraphUtilities.h"
 #include "NiagaraNodeParameterMapSet.h"
 #include "NiagaraDataInterface.h"
@@ -80,7 +81,8 @@ void UNiagaraStackFunctionInputCollection::RefreshChildrenInternal(const TArray<
 {
 	TSet<const UEdGraphPin*> HiddenPins;
 	TArray<const UEdGraphPin*> InputPins;
-	FNiagaraStackGraphUtilities::GetStackFunctionInputPins(*InputFunctionCallNode, InputPins, HiddenPins, FNiagaraStackGraphUtilities::ENiagaraGetStackFunctionInputPinsOptions::ModuleInputsOnly);
+	FCompileConstantResolver ConstantResolver(GetEmitterViewModel()->GetEmitter());
+	FNiagaraStackGraphUtilities::GetStackFunctionInputPins(*InputFunctionCallNode, InputPins, HiddenPins, ConstantResolver, FNiagaraStackGraphUtilities::ENiagaraGetStackFunctionInputPinsOptions::ModuleInputsOnly);
 
 	const UEdGraphSchema_Niagara* NiagaraSchema = GetDefault<UEdGraphSchema_Niagara>();
 
diff --git a/Engine/Plugins/FX/Niagara/Source/NiagaraEditor/Private/ViewModels/Stack/NiagaraStackGraphUtilities.cpp b/Engine/Plugins/FX/Niagara/Source/NiagaraEditor/Private/ViewModels/Stack/NiagaraStackGraphUtilities.cpp
--- a/Engine/Plugins/FX/Niagara/Source/NiagaraEditor/Private/ViewModels/Stack/NiagaraStackGraphUtilities.cpp
+++ b/Engine/Plugins/FX/Niagara/Source/NiagaraEditor/Private/ViewModels/Stack/NiagaraStackGraphUtilities.cpp
@@ -519,10 +519,11 @@ void ExtractInputPinsFromHistory(FNiagaraParameterMapHistory& History, UEdGraph*
 	}
 }
 
-void FNiagaraStackGraphUtilities::GetStackFunctionInputPins(UNiagaraNodeFunctionCall& FunctionCallNode, TArray<const UEdGraphPin*>& OutInputPins, TSet<const UEdGraphPin*>& OutHiddenPins, ENiagaraGetStackFunctionInputPinsOptions Options, bool bIgnoreDisabled)
+void FNiagaraStackGraphUtilities::GetStackFunctionInputPins(UNiagaraNodeFunctionCall& FunctionCallNode, TArray<const UEdGraphPin*>& OutInputPins, TSet<const UEdGraphPin*>& OutHiddenPins, FCompileConstantResolver ConstantResolver, ENiagaraGetStackFunctionInputPinsOptions Options, bool bIgnoreDisabled)
 {
 	FNiagaraParameterMapHistoryBuilder Builder;
 	Builder.SetIgnoreDisabled(bIgnoreDisabled);
+	Builder.ConstantResolver = ConstantResolver;
 	FunctionCallNode.BuildParameterMapHistory(Builder, false, false);
 	
 	if (Builder.Histories.Num() == 1)
@@ -530,6 +531,8 @@ void FNiagaraStackGraphUtilities::GetStackFunctionInputPins(UNiagaraNodeFunction
 		ExtractInputPinsFromHistory(Builder.Histories[0], FunctionCallNode.GetCalledGraph(), Options, OutInputPins);
 
 		FNiagaraParameterMapHistoryBuilder BuilderCompiled;
+		BuilderCompiled.ConstantResolver = ConstantResolver;
+		BuilderCompiled.SetIgnoreDisabled(bIgnoreDisabled);
 		FunctionCallNode.BuildParameterMapHistory(BuilderCompiled, false, true);
 		TArray<const UEdGraphPin*> CompilationPins;
 		if (BuilderCompiled.Histories.Num() == 1)
@@ -557,6 +560,13 @@ void FNiagaraStackGraphUtilities::GetStackFunctionInputPins(UNiagaraNodeFunction
 	}
 }
 
+void FNiagaraStackGraphUtilities::GetStackFunctionInputPins(UNiagaraNodeFunctionCall& FunctionCallNode, TArray<const UEdGraphPin*>& OutInputPins, ENiagaraGetStackFunctionInputPinsOptions Options /*= ENiagaraGetStackFunctionInputPinsOptions::AllInputs*/, bool bIgnoreDisabled /*= false*/)
+{
+	TSet<const UEdGraphPin*> HiddenPins;
+	FCompileConstantResolver EmptyResolver;
+	GetStackFunctionInputPins(FunctionCallNode, OutInputPins, HiddenPins, EmptyResolver, Options, bIgnoreDisabled);
+}
+
 void FNiagaraStackGraphUtilities::GetStackFunctionStaticSwitchPins(UNiagaraNodeFunctionCall& FunctionCallNode, TArray<UEdGraphPin*>& OutInputPins, TSet<UEdGraphPin*>& OutHiddenPins)
 {
 	const UEdGraphSchema_Niagara* Schema = CastChecked<UEdGraphSchema_Niagara>(FunctionCallNode.GetSchema());
@@ -1552,9 +1562,7 @@ bool TryGetStackFunctionInputValue(UNiagaraScript& OwningScript, const UEdGraphP
 			UNiagaraNodeFunctionCall* DynamicInputFunctionCall = CastChecked<UNiagaraNodeFunctionCall>(InputPin.LinkedTo[0]->GetOwningNode());
 			OutStackFunctionInputValue.DynamicValue = DynamicInputFunctionCall;
 			TArray<const UEdGraphPin*> DynamicValueInputPins;
-			TSet<const UEdGraphPin*> HiddenInputPins;
-			FNiagaraStackGraphUtilities::GetStackFunctionInputPins(*DynamicInputFunctionCall, DynamicValueInputPins, HiddenInputPins,
-				FNiagaraStackGraphUtilities::ENiagaraGetStackFunctionInputPinsOptions::ModuleInputsOnly);
+			FNiagaraStackGraphUtilities::GetStackFunctionInputPins(*DynamicInputFunctionCall, DynamicValueInputPins, FNiagaraStackGraphUtilities::ENiagaraGetStackFunctionInputPinsOptions::ModuleInputsOnly);
 
 			FRapidIterationParameterContext InputRapidIterationParameterContext = RapidIterationParameterContext.IsValid() 
 				? FRapidIterationParameterContext(RapidIterationParameterContext.UniqueEmitterName, *DynamicInputFunctionCall)
diff --git a/Engine/Plugins/FX/Niagara/Source/NiagaraEditor/Private/ViewModels/Stack/NiagaraStackModuleItemLinkedInputCollection.cpp b/Engine/Plugins/FX/Niagara/Source/NiagaraEditor/Private/ViewModels/Stack/NiagaraStackModuleItemLinkedInputCollection.cpp
--- a/Engine/Plugins/FX/Niagara/Source/NiagaraEditor/Private/ViewModels/Stack/NiagaraStackModuleItemLinkedInputCollection.cpp
+++ b/Engine/Plugins/FX/Niagara/Source/NiagaraEditor/Private/ViewModels/Stack/NiagaraStackModuleItemLinkedInputCollection.cpp
@@ -5,6 +5,7 @@
 #include "NiagaraNodeFunctionCall.h"
 #include "NiagaraNodeParameterMapGet.h"
 #include "NiagaraEmitterEditorData.h"
+#include "ViewModels/NiagaraEmitterViewModel.h"
 #include "ViewModels/Stack/NiagaraStackGraphUtilities.h"
 
 #include "EdGraph/EdGraphPin.h"
@@ -60,6 +61,7 @@ void UNiagaraStackModuleItemLinkedInputCollection::RefreshChildrenInternal(const
 	{
 		FNiagaraParameterMapHistoryBuilder Builder;
 		Builder.SetIgnoreDisabled(false);
+		Builder.ConstantResolver = FCompileConstantResolver(GetEmitterViewModel()->GetEmitter());
 		FunctionCallNode->BuildParameterMapHistory(Builder, false);
 
 		if (ensureMsgf(Builder.Histories.Num() == 1, TEXT("Invalid Stack Graph - Function call node has invalid history count!")))
diff --git a/Engine/Plugins/FX/Niagara/Source/NiagaraEditor/Private/ViewModels/Stack/NiagaraStackModuleItemOutputCollection.cpp b/Engine/Plugins/FX/Niagara/Source/NiagaraEditor/Private/ViewModels/Stack/NiagaraStackModuleItemOutputCollection.cpp
--- a/Engine/Plugins/FX/Niagara/Source/NiagaraEditor/Private/ViewModels/Stack/NiagaraStackModuleItemOutputCollection.cpp
+++ b/Engine/Plugins/FX/Niagara/Source/NiagaraEditor/Private/ViewModels/Stack/NiagaraStackModuleItemOutputCollection.cpp
@@ -5,6 +5,7 @@
 #include "NiagaraNodeFunctionCall.h"
 #include "NiagaraNodeParameterMapSet.h"
 #include "NiagaraEmitterEditorData.h"
+#include "ViewModels/NiagaraEmitterViewModel.h"
 #include "ViewModels/Stack/NiagaraStackGraphUtilities.h"
 
 #include "EdGraph/EdGraphPin.h"
@@ -49,6 +50,7 @@ void UNiagaraStackModuleItemOutputCollection::RefreshChildrenInternal(const TArr
 	{
 		FNiagaraParameterMapHistoryBuilder Builder;
 		Builder.SetIgnoreDisabled(false);
+		Builder.ConstantResolver = FCompileConstantResolver(GetEmitterViewModel()->GetEmitter());
 		FunctionCallNode->BuildParameterMapHistory(Builder, false);
 
 		if (ensureMsgf(Builder.Histories.Num() == 1, TEXT("Invalid Stack Graph - Function call node has invalid history count!")))
diff --git a/Engine/Plugins/FX/Niagara/Source/NiagaraEditor/Private/ViewModels/Stack/NiagaraStackScriptItemGroup.cpp b/Engine/Plugins/FX/Niagara/Source/NiagaraEditor/Private/ViewModels/Stack/NiagaraStackScriptItemGroup.cpp
--- a/Engine/Plugins/FX/Niagara/Source/NiagaraEditor/Private/ViewModels/Stack/NiagaraStackScriptItemGroup.cpp
+++ b/Engine/Plugins/FX/Niagara/Source/NiagaraEditor/Private/ViewModels/Stack/NiagaraStackScriptItemGroup.cpp
@@ -258,8 +258,7 @@ private:
 		UNiagaraNodeAssignment* NewAssignmentModule = FNiagaraStackGraphUtilities::AddParameterModuleToStack(Vars, *OutputNode, TargetIndex,DefaultVals );
 		
 		TArray<const UEdGraphPin*> InputPins;
-		TSet<const UEdGraphPin*> HiddenPins;
-		FNiagaraStackGraphUtilities::GetStackFunctionInputPins(*NewAssignmentModule, InputPins, HiddenPins);
+		FNiagaraStackGraphUtilities::GetStackFunctionInputPins(*NewAssignmentModule, InputPins);
 		if (InputPins.Num() == 1)
 		{
 			FString FunctionInputEditorDataKey = FNiagaraStackGraphUtilities::GenerateStackFunctionInputEditorDataKey(*NewAssignmentModule, InputPins[0]->PinName);
diff --git a/Engine/Plugins/FX/Niagara/Source/NiagaraEditor/Private/Widgets/SNiagaraParameterMapView.cpp b/Engine/Plugins/FX/Niagara/Source/NiagaraEditor/Private/Widgets/SNiagaraParameterMapView.cpp
--- a/Engine/Plugins/FX/Niagara/Source/NiagaraEditor/Private/Widgets/SNiagaraParameterMapView.cpp
+++ b/Engine/Plugins/FX/Niagara/Source/NiagaraEditor/Private/Widgets/SNiagaraParameterMapView.cpp
@@ -336,7 +336,7 @@ void SNiagaraParameterMapView::CollectAllActions(FGraphActionListBuilderBase& Ou
 			continue;
 		}
 		UNiagaraGraph* Graph = GraphWeakPtr.Get();
-		for (const auto& ParameterElement : FNiagaraEditorUtilities::GetCompiledGraphParameterMapReferences(Graph))
+		for (const auto& ParameterElement : Graph->GetParameterReferenceMap())
 		{
 			TArray<FNiagaraGraphParameterReferenceCollection>* Found = ParameterEntries.Find(ParameterElement.Key);
 			if (Found)
diff --git a/Engine/Plugins/FX/Niagara/Source/NiagaraEditor/Public/NiagaraEditorUtilities.h b/Engine/Plugins/FX/Niagara/Source/NiagaraEditor/Public/NiagaraEditorUtilities.h
--- a/Engine/Plugins/FX/Niagara/Source/NiagaraEditor/Public/NiagaraEditorUtilities.h
+++ b/Engine/Plugins/FX/Niagara/Source/NiagaraEditor/Public/NiagaraEditorUtilities.h
@@ -26,6 +26,7 @@ class SWidget;
 class UNiagaraNode;
 class UEdGraphSchema_Niagara;
 class UEdGraphPin;
+class FCompileConstantResolver;
 
 namespace FNiagaraEditorUtilities
 {
@@ -138,15 +139,12 @@ namespace FNiagaraEditorUtilities
 
 	void FixUpNumericPins(const UEdGraphSchema_Niagara* Schema, UNiagaraNode* Node);
 
-	void SetStaticSwitchConstants(UNiagaraGraph* Graph, const TArray<UEdGraphPin*>& CallInputs);
+	void SetStaticSwitchConstants(UNiagaraGraph* Graph, const TArray<UEdGraphPin*>& CallInputs, FCompileConstantResolver* ConstantResolver = nullptr);
 
 	bool ResolveConstantValue(UEdGraphPin* Pin, int32& Value);
 
 	void PreprocessFunctionGraph(const UEdGraphSchema_Niagara* Schema, UNiagaraGraph* Graph, const TArray<UEdGraphPin*>& CallInputs, const TArray<UEdGraphPin*>& CallOutputs, ENiagaraScriptUsage ScriptUsage);
 
-	/** Returns the parameter map references for the given graph as if the graph was compiled, factoring in compile-time options such as static switches. */
-	const TMap<FNiagaraVariable, FNiagaraGraphParameterReferenceCollection>& GetCompiledGraphParameterMapReferences(UNiagaraGraph* Graph);
-
 	/** Options for the GetScriptsByFilter function. 
 	** @Param ScriptUsageToInclude Only return Scripts that have this usage
 	** @Param (Optional) TargetUsageToMatch Only return Scripts that have this target usage (output node) 
diff --git a/Engine/Plugins/FX/Niagara/Source/NiagaraEditor/Public/NiagaraParameterMapHistory.h b/Engine/Plugins/FX/Niagara/Source/NiagaraEditor/Public/NiagaraParameterMapHistory.h
--- a/Engine/Plugins/FX/Niagara/Source/NiagaraEditor/Public/NiagaraParameterMapHistory.h
+++ b/Engine/Plugins/FX/Niagara/Source/NiagaraEditor/Public/NiagaraParameterMapHistory.h
@@ -5,12 +5,24 @@
 #include "CoreMinimal.h"
 #include "NiagaraCommon.h"
 #include "Templates/Tuple.h"
+#include "ViewModels/Stack/NiagaraStackGraphUtilities.h"
 
 class UNiagaraNodeOutput;
 class UEdGraphPin;
 class UEdGraphNode;
 class UNiagaraParameterCollection;
 
+class FCompileConstantResolver
+{
+public:
+	FCompileConstantResolver() : Emitter(nullptr) {}
+	FCompileConstantResolver(UNiagaraEmitter* Emitter) : Emitter(Emitter) {}
+
+	bool ResolveConstant(FNiagaraVariable& OutConstant) const;
+private:
+	UNiagaraEmitter* Emitter;
+};
+
 /** Traverses a Niagara node graph to identify the variables that have been written and read from a parameter map. 
 * 	This class is meant to aid in UI and compilation of the graph. There are several main script types and each one interacts
 *	slightly differently with the history depending on context.
@@ -341,6 +353,8 @@ public:
 	void RegisterEncounterableVariables(const TArray<FNiagaraVariable>& Variables);
 	const TArray<FNiagaraVariable>& GetEncounterableVariables() const {	return EncounterableExternalVariables;}
 
+	FCompileConstantResolver ConstantResolver;
+
 protected:
 	/**
 	* Generate the internal alias map from the current traversal state.
diff --git a/Engine/Plugins/FX/Niagara/Source/NiagaraEditor/Public/ViewModels/Stack/NiagaraStackGraphUtilities.h b/Engine/Plugins/FX/Niagara/Source/NiagaraEditor/Public/ViewModels/Stack/NiagaraStackGraphUtilities.h
--- a/Engine/Plugins/FX/Niagara/Source/NiagaraEditor/Public/ViewModels/Stack/NiagaraStackGraphUtilities.h
+++ b/Engine/Plugins/FX/Niagara/Source/NiagaraEditor/Public/ViewModels/Stack/NiagaraStackGraphUtilities.h
@@ -18,10 +18,12 @@ class UNiagaraNodeCustomHlsl;
 class UNiagaraNodeAssignment;
 class UNiagaraNodeParameterMapSet;
 class FNiagaraSystemViewModel;
+class UNiagaraEmitter;
 class FNiagaraEmitterViewModel;
 class UNiagaraStackEditorData;
 class UNiagaraStackEntry;
 class UNiagaraStackErrorItem;
+class FCompileConstantResolver;
 
 namespace FNiagaraStackGraphUtilities
 {
@@ -74,7 +76,9 @@ namespace FNiagaraStackGraphUtilities
 		ModuleInputsOnly
 	};
 
-	void GetStackFunctionInputPins(UNiagaraNodeFunctionCall& FunctionCallNode, TArray<const UEdGraphPin*>& OutInputPins, TSet<const UEdGraphPin*>& OutHiddenPins, ENiagaraGetStackFunctionInputPinsOptions Options = ENiagaraGetStackFunctionInputPinsOptions::AllInputs, bool bIgnoreDisabled = false);
+	void GetStackFunctionInputPins(UNiagaraNodeFunctionCall& FunctionCallNode, TArray<const UEdGraphPin*>& OutInputPins, ENiagaraGetStackFunctionInputPinsOptions Options = ENiagaraGetStackFunctionInputPinsOptions::AllInputs, bool bIgnoreDisabled = false);
+
+	void GetStackFunctionInputPins(UNiagaraNodeFunctionCall& FunctionCallNode, TArray<const UEdGraphPin*>& OutInputPins, TSet<const UEdGraphPin*>& OutHiddenPins, FCompileConstantResolver ConstantResolver, ENiagaraGetStackFunctionInputPinsOptions Options = ENiagaraGetStackFunctionInputPinsOptions::AllInputs, bool bIgnoreDisabled = false);
 
 	void GetStackFunctionStaticSwitchPins(UNiagaraNodeFunctionCall& FunctionCallNode, TArray<UEdGraphPin*>& OutInputPins, TSet<UEdGraphPin*>& OutHiddenPins);
 
