Clean up old propagated static switch values

diff --git a/Engine/Plugins/FX/Niagara/Source/NiagaraEditor/Private/NiagaraNodeFunctionCall.cpp b/Engine/Plugins/FX/Niagara/Source/NiagaraEditor/Private/NiagaraNodeFunctionCall.cpp
--- a/Engine/Plugins/FX/Niagara/Source/NiagaraEditor/Private/NiagaraNodeFunctionCall.cpp
+++ b/Engine/Plugins/FX/Niagara/Source/NiagaraEditor/Private/NiagaraNodeFunctionCall.cpp
@@ -70,15 +70,8 @@ void UNiagaraNodeFunctionCall::PostLoad()
 	}
 
 	// Clean up invalid old references to propagated parameters
-	for (int i = PropagatedStaticSwitchParameters.Num() - 1; i >= 0; i--)
-	{
-		FNiagaraPropagatedVariable& Propagated = PropagatedStaticSwitchParameters[i];
-		if (Propagated.SwitchParameter.GetName().IsNone())
-		{
-			PropagatedStaticSwitchParameters.RemoveAt(i);
-		}
-	}
-
+	CleanupPropagatedSwitchValues();
+	
 	if (FunctionDisplayName.IsEmpty())
 	{
 		ComputeNodeName();
@@ -573,6 +566,7 @@ bool UNiagaraNodeFunctionCall::RefreshFromExternalChanges()
 	UNiagaraGraph* CalledGraph = GetCalledGraph();
 	if (CalledGraph)
 	{
+		CleanupPropagatedSwitchValues();
 		TArray<UEdGraphPin*> InputPins;
 		GetInputPins(InputPins);
 		for (FNiagaraVariable InputVar : CalledGraph->FindStaticSwitchInputs())
@@ -906,6 +900,35 @@ void UNiagaraNodeFunctionCall::SetPinAutoGeneratedDefaultValue(UEdGraphPin& Func
 	}
 }
 
+void UNiagaraNodeFunctionCall::CleanupPropagatedSwitchValues()
+{
+	for (int i = PropagatedStaticSwitchParameters.Num() - 1; i >= 0; i--)
+	{
+		FNiagaraPropagatedVariable& Propagated = PropagatedStaticSwitchParameters[i];
+		if (Propagated.SwitchParameter.GetName().IsNone() || !IsValidPropagatedVariable(Propagated.SwitchParameter))
+		{
+			PropagatedStaticSwitchParameters.RemoveAt(i);
+		}
+	}
+}
+
+bool UNiagaraNodeFunctionCall::IsValidPropagatedVariable(const FNiagaraVariable& Variable) const
+{
+	UNiagaraGraph* Graph = GetCalledGraph();
+	if (!Graph)
+	{
+		return false;
+	}
+	for (const FNiagaraVariable& Var : Graph->FindStaticSwitchInputs(false))
+	{
+		if (Var == Variable)
+		{
+			return true;
+		}
+	}
+	return false;
+}
+
 bool UNiagaraNodeFunctionCall::FindAutoBoundInput(UNiagaraNodeInput* InputNode, UEdGraphPin* PinToAutoBind, FNiagaraVariable& OutFoundVar, ENiagaraInputNodeUsage& OutNodeUsage)
 {
 	check(InputNode && InputNode->IsExposed());
diff --git a/Engine/Plugins/FX/Niagara/Source/NiagaraEditor/Public/NiagaraNodeFunctionCall.h b/Engine/Plugins/FX/Niagara/Source/NiagaraEditor/Public/NiagaraNodeFunctionCall.h
--- a/Engine/Plugins/FX/Niagara/Source/NiagaraEditor/Public/NiagaraNodeFunctionCall.h
+++ b/Engine/Plugins/FX/Niagara/Source/NiagaraEditor/Public/NiagaraNodeFunctionCall.h
@@ -124,6 +124,7 @@ public:
 
 	FNiagaraPropagatedVariable* FindPropagatedVariable(const FNiagaraVariable& Variable);
 	void RemovePropagatedVariable(const FNiagaraVariable& Variable);
+	void CleanupPropagatedSwitchValues();
 
 protected:
 	virtual bool IsValidPinToCompile(UEdGraphPin* Pin) const;
@@ -138,6 +139,8 @@ protected:
 
 	void SetPinAutoGeneratedDefaultValue(UEdGraphPin& FunctionInputPin, UNiagaraNodeInput& FunctionScriptInputNode);
 
+	bool IsValidPropagatedVariable(const FNiagaraVariable& Variable) const;
+
 	/** Adjusted every time that we compile this script. Lets us know that we might differ from any cached versions.*/
 	UPROPERTY()
 	FGuid CachedChangeId;
