Prevent a crash in the static switch logic when the underlying enum definition changed

diff --git a/Engine/Plugins/FX/Niagara/Source/NiagaraEditor/Private/NiagaraNodeStaticSwitch.cpp b/Engine/Plugins/FX/Niagara/Source/NiagaraEditor/Private/NiagaraNodeStaticSwitch.cpp
--- a/Engine/Plugins/FX/Niagara/Source/NiagaraEditor/Private/NiagaraNodeStaticSwitch.cpp
+++ b/Engine/Plugins/FX/Niagara/Source/NiagaraEditor/Private/NiagaraNodeStaticSwitch.cpp
@@ -291,6 +291,20 @@ bool UNiagaraNodeStaticSwitch::GetVarIndex(FHlslNiagaraTranslator* Translator, i
 		int32 MaxValue = SwitchTypeData.Enum->NumEnums() - 1;
 		if (MaxValue > 0)
 		{
+			// do a sanity check here if the number of pins actually matches the enum count (which might have changed in the meantime without us noticing)
+			TArray<UEdGraphPin*> Pins;
+			GetOutputPins(Pins);
+			int32 OutputPinCount = Pins.Num() - 1;
+			int32 ReservedValues = (InputPinCount / OutputPinCount);
+			if (OutputPinCount > 0 && (MaxValue > ReservedValues || MaxValue < ReservedValues))
+			{
+				MaxValue = ReservedValues;
+				if (Translator)
+				{
+					Translator->Error(FText::Format(LOCTEXT("InvalidSwitchEnumDefinition", "The number of pins on the static switch does not match the number of values defined in the enum."), FText::FromString(FString::FromInt(SwitchValue))), this, nullptr);
+				}
+			}
+			
 			if (Value <= MaxValue && Value >= 0)
 			{
 				VarIndexOut = Value * (InputPinCount / MaxValue);
