- Add ERWNoBarrier transition for count buffer as multiple dispatches will write to different regions of the buffer
- This fixes a validation issue about the UAV being dirty

#rnx

diff --git a/Engine/Plugins/FX/Niagara/Source/Niagara/Private/NiagaraEmitterInstanceBatcher.cpp b/Engine/Plugins/FX/Niagara/Source/Niagara/Private/NiagaraEmitterInstanceBatcher.cpp
--- a/Engine/Plugins/FX/Niagara/Source/Niagara/Private/NiagaraEmitterInstanceBatcher.cpp
+++ b/Engine/Plugins/FX/Niagara/Source/Niagara/Private/NiagaraEmitterInstanceBatcher.cpp
@@ -952,6 +952,7 @@ void NiagaraEmitterInstanceBatcher::Run(const FNiagaraGPUSystemTick& Tick, const
 	//
 	if (Shader->InstanceCountsParam.IsBound())
 	{
+		RHICmdList.TransitionResource(EResourceTransitionAccess::ERWNoBarrier, EResourceTransitionPipeline::EComputeToCompute, GPUInstanceCounterManager.GetInstanceCountBuffer().UAV);
 		RHICmdList.SetUAVParameter(Shader->GetComputeShader(), Shader->InstanceCountsParam.GetUAVIndex(), GPUInstanceCounterManager.GetInstanceCountBuffer().UAV);
 		const uint32 ReadOffset = CurrentData.GetGPUInstanceCountBufferOffset();
 		const uint32 WriteOffset = DestinationData.GetGPUInstanceCountBufferOffset();
