Fix new Niagara System assets not compiling on creation.

diff --git a/Engine/Plugins/FX/Niagara/Source/NiagaraEditor/Private/AssetTypeActions/AssetTypeActions_NiagaraEmitter.cpp b/Engine/Plugins/FX/Niagara/Source/NiagaraEditor/Private/AssetTypeActions/AssetTypeActions_NiagaraEmitter.cpp
--- a/Engine/Plugins/FX/Niagara/Source/NiagaraEditor/Private/AssetTypeActions/AssetTypeActions_NiagaraEmitter.cpp
+++ b/Engine/Plugins/FX/Niagara/Source/NiagaraEditor/Private/AssetTypeActions/AssetTypeActions_NiagaraEmitter.cpp
@@ -102,21 +102,13 @@ void FAssetTypeActions_NiagaraEmitter::ExecuteNewNiagaraSystem(TArray<TWeakObjec
 
 			// Create the factory used to generate the asset
 			UNiagaraSystemFactoryNew* Factory = NewObject<UNiagaraSystemFactoryNew>();
+			Factory->EmittersToAddToNewSystem.Add(Emitter);
 			FAssetToolsModule& AssetToolsModule = FModuleManager::GetModuleChecked<FAssetToolsModule>("AssetTools");
 			UObject* NewAsset = AssetToolsModule.Get().CreateAsset(Name, FPackageName::GetLongPackagePath(PackageName), UNiagaraSystem::StaticClass(), Factory);
 			
 			UNiagaraSystem* System = Cast<UNiagaraSystem>(NewAsset);
 			if (System != nullptr)
 			{
-				System->AddEmitterHandle(*Emitter, FName(*Emitter->GetName()));
-				FNiagaraStackGraphUtilities::RebuildEmitterNodes(*System);
-
-				// Ensure the new System is compiled
-				if (!Emitter->AreAllScriptAndSourcesSynchronized())
-				{
-					System->RequestCompile(true);
-				}
-
 				ObjectsToSync.Add(NewAsset);
 			}
 		}
diff --git a/Engine/Plugins/FX/Niagara/Source/NiagaraEditor/Private/NiagaraSystemFactoryNew.cpp b/Engine/Plugins/FX/Niagara/Source/NiagaraEditor/Private/NiagaraSystemFactoryNew.cpp
--- a/Engine/Plugins/FX/Niagara/Source/NiagaraEditor/Private/NiagaraSystemFactoryNew.cpp
+++ b/Engine/Plugins/FX/Niagara/Source/NiagaraEditor/Private/NiagaraSystemFactoryNew.cpp
@@ -146,6 +146,8 @@ UObject* UNiagaraSystemFactoryNew::FactoryCreateNew(UClass* Class, UObject* InPa
 		InitializeSystem(NewSystem, true);
 	}
 
+	NewSystem->RequestCompile(false);
+
 	return NewSystem;
 }
 
