No longer warning about multiple GPU renderers since that now works.

#tests n/a

diff --git a/Engine/Plugins/FX/Niagara/Source/NiagaraEditor/Private/ViewModels/Stack/NiagaraStackRendererItem.cpp b/Engine/Plugins/FX/Niagara/Source/NiagaraEditor/Private/ViewModels/Stack/NiagaraStackRendererItem.cpp
--- a/Engine/Plugins/FX/Niagara/Source/NiagaraEditor/Private/ViewModels/Stack/NiagaraStackRendererItem.cpp
+++ b/Engine/Plugins/FX/Niagara/Source/NiagaraEditor/Private/ViewModels/Stack/NiagaraStackRendererItem.cpp
@@ -280,22 +280,6 @@ void UNiagaraStackRendererItem::RefreshIssues(TArray<FStackIssue>& NewIssues)
 		NewIssues.Add(TargetSupportError);
 	}
 
-	// Currently we don't support multiple different renderers on GPU emitters
-	// Multiple sprite renderers may currently work but not multiple of other types or combinations of types and there is no guarantee
-	if ( GetEmitterViewModel()->GetEmitter()->SimTarget == ENiagaraSimTarget::GPUComputeSim )
-	{
-		if ( GetEmitterViewModel()->GetEmitter()->GetRenderers().Num() > 1 )
-		{
-			FStackIssue GpuSupportError(
-				EStackIssueSeverity::Error,
-				LOCTEXT("FailedRendererDueToMultipleGpuRenderers", "Multiple GPU renderers."),
-				LOCTEXT("FailedRendererDueToMultipleGpuRenderersLong", "Multiple GPU renderers are currently not supported, rendered results may not be correct."),
-				GetStackEditorDataKey(),
-				false);
-
-			NewIssues.Add(GpuSupportError);
-		}
-	}
 }
 
 void UNiagaraStackRendererItem::RendererChanged()
