- Niagara socket sampling includes the local transform

#rnx

diff --git a/Engine/Plugins/FX/Niagara/Shaders/Private/NiagaraDataInterfaceSkeletalMesh.ush b/Engine/Plugins/FX/Niagara/Shaders/Private/NiagaraDataInterfaceSkeletalMesh.ush
--- a/Engine/Plugins/FX/Niagara/Shaders/Private/NiagaraDataInterfaceSkeletalMesh.ush
+++ b/Engine/Plugins/FX/Niagara/Shaders/Private/NiagaraDataInterfaceSkeletalMesh.ush
@@ -28,8 +28,8 @@ struct FDISkelMeshContext
 	uint			MeshNumWeights;
 	uint			NumSpecificBones;
 	Buffer<uint>	SpecificBones;
-	uint			NumSpecificSocketBones;
-	Buffer<uint>	SpecificSocketBones;
+	uint			NumSpecificSockets;
+	uint			SpecificSocketBoneOffset;
 	float4x4		InstanceTransform;
 	float4x4		InstancePrevTransform;
 	float			InstanceInvDeltaTime;
@@ -426,18 +426,18 @@ void DISkelMesh_RandomSpecificBone(in FDISkelMeshContext DIContext, out int Bone
 
 void DISkelMesh_GetSpecificSocketCount(in FDISkelMeshContext DIContext, out int Count)
 {
-    Count = DIContext.NumSpecificSocketBones;
+    Count = DIContext.NumSpecificSockets;
 }
 
 void DISkelMesh_GetSpecificSocketBoneAt(in FDISkelMeshContext DIContext, in int InSocketIndex, out int Bone)
 {
-    Bone = (int)DIContext.SpecificSocketBones[InSocketIndex];
+    Bone = (int)DIContext.SpecificSocketBoneOffset + clamp(InSocketIndex, 0, (int)DIContext.NumSpecificSockets);
 }
 
 void DISkelMesh_RandomSpecificSocketBone(in FDISkelMeshContext DIContext, out int SocketBone)
 {
     float RandT0 = NiagaraInternalNoise(1, 2, 3);
-    SocketBone = min(uint(RandT0 * float(DIContext.NumSpecificSocketBones)), DIContext.NumSpecificSocketBones - 1); // avoid % by using mul/min to Tri = MeshVertexCount
+    SocketBone = min(uint(RandT0 * float(DIContext.NumSpecificSockets)), DIContext.NumSpecificSockets - 1); // avoid % by using mul/min to Tri = MeshVertexCount
 }
 
 ////////////////////////////////////////////////////////////////////////////////////////////////////////////////////////////////////////////////////////////////
@@ -463,8 +463,8 @@ void DISkelMesh_RandomSpecificSocketBone(in FDISkelMeshContext DIContext, out in
 	uint			MeshNumWeights_ ## NAME; \
 	uint			NumSpecificBones_ ## NAME; \
 	Buffer<uint>	SpecificBones_ ## NAME; \
-	uint			NumSpecificSocketBones_ ## NAME; \
-	Buffer<uint>	SpecificSocketBones_ ## NAME; \
+	uint			NumSpecificSockets_ ## NAME; \
+	uint			SpecificSocketBoneOffset_ ## NAME; \
 	float4x4		InstanceTransform_ ## NAME; \
 	float4x4		InstancePrevTransform_ ## NAME; \
 	float			InstanceInvDeltaTime_ ## NAME; \
@@ -491,8 +491,8 @@ void DISkelMesh_RandomSpecificSocketBone(in FDISkelMeshContext DIContext, out in
 	DIContext.MeshNumWeights = MeshNumWeights_ ## NAME; \
 	DIContext.NumSpecificBones	= NumSpecificBones_ ## NAME; \
 	DIContext.SpecificBones	= SpecificBones_ ## NAME; \
-	DIContext.NumSpecificSocketBones = NumSpecificSocketBones_ ## NAME; \
-	DIContext.SpecificSocketBones = SpecificSocketBones_ ## NAME; \
+	DIContext.NumSpecificSockets = NumSpecificSockets_ ## NAME; \
+	DIContext.SpecificSocketBoneOffset = SpecificSocketBoneOffset_ ## NAME; \
 	DIContext.MeshWeightStride = MeshWeightStride_ ## NAME; \
 	DIContext.InstanceTransform = InstanceTransform_ ## NAME; \
 	DIContext.InstancePrevTransform = InstancePrevTransform_ ## NAME; \
diff --git a/Engine/Plugins/FX/Niagara/Source/Niagara/Classes/NDISkeletalMeshCommon.h b/Engine/Plugins/FX/Niagara/Source/Niagara/Classes/NDISkeletalMeshCommon.h
--- a/Engine/Plugins/FX/Niagara/Source/Niagara/Classes/NDISkeletalMeshCommon.h
+++ b/Engine/Plugins/FX/Niagara/Source/Niagara/Classes/NDISkeletalMeshCommon.h
@@ -6,19 +6,6 @@
 
 struct FSkeletalMeshAccessorHelper
 {
-	FSkeletalMeshAccessorHelper()
-		: Comp(nullptr)
-		, Mesh(nullptr)
-		, LODData(nullptr)
-		, SkinWeightBuffer(nullptr)
-		, IndexBuffer(nullptr)
-		, SamplingRegion(nullptr)
-		, SamplingRegionBuiltData(nullptr)
-		, SkinningData(nullptr)
-	{
-
-	}
-
 	template<typename FilterMode, typename AreaWeightingMode>
 	FORCEINLINE void Init(FNDISkeletalMesh_InstanceData* InstData)
 	{
@@ -37,18 +24,18 @@ struct FSkeletalMeshAccessorHelper
 		}
 	}
 
-	USkeletalMeshComponent* Comp;
-	USkeletalMesh* Mesh;
+	USkeletalMeshComponent* Comp = nullptr;
+	USkeletalMesh* Mesh = nullptr;
 	TWeakObjectPtr<USkeletalMesh> MeshSafe;
-	FSkeletalMeshLODRenderData* LODData;
-	FSkinWeightVertexBuffer* SkinWeightBuffer;
-	FRawStaticIndexBuffer16or32Interface* IndexBuffer;
-	const FSkeletalMeshSamplingRegion* SamplingRegion;
-	const FSkeletalMeshSamplingRegionBuiltData* SamplingRegionBuiltData;
-	FSkeletalMeshSkinningData* SkinningData;
+	FSkeletalMeshLODRenderData* LODData = nullptr;
+	FSkinWeightVertexBuffer* SkinWeightBuffer = nullptr;
+	FRawStaticIndexBuffer16or32Interface* IndexBuffer = nullptr;
+	const FSkeletalMeshSamplingRegion* SamplingRegion = nullptr;
+	const FSkeletalMeshSamplingRegionBuiltData* SamplingRegionBuiltData = nullptr;
+	FSkeletalMeshSkinningData* SkinningData = nullptr;
 	FSkeletalMeshSkinningDataUsage Usage;
-	const TArray<FTransform>* BoneComponentSpaceTransforms;
-	const TArray<FTransform>* PrevBoneComponentSpaceTransforms;
+	const TArray<FTransform>* BoneComponentSpaceTransforms = nullptr;
+	const TArray<FTransform>* PrevBoneComponentSpaceTransforms = nullptr;
 };
 
 template<>
diff --git a/Engine/Plugins/FX/Niagara/Source/Niagara/Classes/NiagaraDataInterfaceSkeletalMesh.h b/Engine/Plugins/FX/Niagara/Source/Niagara/Classes/NiagaraDataInterfaceSkeletalMesh.h
--- a/Engine/Plugins/FX/Niagara/Source/Niagara/Classes/NiagaraDataInterfaceSkeletalMesh.h
+++ b/Engine/Plugins/FX/Niagara/Source/Niagara/Classes/NiagaraDataInterfaceSkeletalMesh.h
@@ -237,8 +237,7 @@ public:
 
 	virtual ~FSkeletalMeshGpuSpawnStaticBuffers();
 
-	FORCEINLINE_DEBUGGABLE void Initialise(const FSkeletalMeshLODRenderData& SkeletalMeshLODRenderData, 
-		bool bIsGpuUniformlyDistributedSampling, const FSkeletalMeshSamplingLODBuiltData& SkeletalMeshSamplingLODBuiltData, const TArray<int32>& SpecificBones, const TArray<int32>& SpecificSocketBones);
+	FORCEINLINE_DEBUGGABLE void Initialise(struct FNDISkeletalMesh_InstanceData* InstData, const FSkeletalMeshLODRenderData& SkeletalMeshLODRenderData,const FSkeletalMeshSamplingLODBuiltData& SkeletalMeshSamplingLODBuiltData);
 
 	virtual void InitRHI() override;
 	virtual void ReleaseRHI() override;
@@ -263,8 +262,8 @@ public:
 	uint32 GetNumSpecificBones() const { return NumSpecificBones; }
 	FRHIShaderResourceView* GetSpecificBonesSRV() const { return SpecificBonesSRV; }
 
-	uint32 GetNumSpecificSocketBones() const { return NumSpecificSocketBones; }
-	FRHIShaderResourceView* GetSpecificSocketBonesSRV() const { return SpecificSocketBonesSRV; }
+	uint32 GetNumSpecificSockets() const { return NumSpecificSockets; }
+	uint32 GetSpecificSocketBoneOffset() const { return SpecificSocketBoneOffset; }
 
 protected:
 
@@ -280,10 +279,8 @@ protected:
 	FVertexBufferRHIRef SpecificBonesBuffer;
 	FShaderResourceViewRHIRef SpecificBonesSRV;
 
-	uint32 NumSpecificSocketBones = 0;
-	TResourceArray<uint16> SpecificSocketBonesArray;
-	FVertexBufferRHIRef SpecificSocketBonesBuffer;
-	FShaderResourceViewRHIRef SpecificSocketBonesSRV;
+	uint32 NumSpecificSockets = 0;
+	uint32 SpecificSocketBoneOffset = 0;
 
 	/** Cached SRV to gpu buffers of the mesh we spawn from */
 	FRHIShaderResourceView* MeshVertexBufferSrv;
@@ -316,7 +313,7 @@ public:
 	FSkeletalMeshGpuDynamicBufferProxy();
 	virtual ~FSkeletalMeshGpuDynamicBufferProxy();
 
-	void Initialise(const FReferenceSkeleton& RefSkel, const FSkeletalMeshLODRenderData& SkeletalMeshLODRenderData);
+	void Initialise(const FReferenceSkeleton& RefSkel, const FSkeletalMeshLODRenderData& SkeletalMeshLODRenderData, uint32 InSamplingSocketCount);
 
 	virtual void InitRHI() override;
 	virtual void ReleaseRHI() override;
@@ -340,6 +337,7 @@ public:
 
 private:
 	uint32 SamplingBoneCount = 0;
+	uint32 SamplingSocketCount = 0;
 	uint32 SectionBoneCount = 0;
 
 	enum { BufferBoneCount = 2 };
@@ -384,10 +382,15 @@ struct FNDISkeletalMesh_InstanceData
 	/** Indices of the bones specifically referenced by the interface. */
 	TArray<int32> SpecificBones;
 
-	/** Indices of the sockets specifically referenced by the interface. */
-	TArray<int32> SpecificSockets;
-	/** The bone indices for the specific sockets. */
-	TArray<int32> SpecificSocketBones;
+	/** Name of all the sockets we use. */
+	TArray<FName> SpecificSockets;
+	/** Bone index of the first socket, sockets are appended to the end of the bone array */
+	int32 SpecificSocketBoneOffset = 0;
+
+	/** Index into which socket transforms to use.  */
+	uint32 SpecificSocketTransformsIndex = 0;
+	/** Transforms for sockets. */
+	TStaticArray<TArray<FTransform>, 2> SpecificSocketTransforms;
 
 	uint32 ChangeId;
 
@@ -423,6 +426,11 @@ struct FNDISkeletalMesh_InstanceData
 		return Ret;
 	}
 
+	void UpdateSpecificSocketTransforms();
+	TArray<FTransform>& GetSpecificSocketsWriteBuffer() { return SpecificSocketTransforms[SpecificSocketTransformsIndex]; }
+	const TArray<FTransform>& GetSpecificSocketsCurrBuffer() const { return SpecificSocketTransforms[SpecificSocketTransformsIndex]; }
+	const TArray<FTransform>& GetSpecificSocketsPrevBuffer() const { return SpecificSocketTransforms[(SpecificSocketTransformsIndex + 1) % SpecificSocketTransforms.Num()]; }
+
 	bool HasColorData();
 };
 
@@ -464,7 +472,6 @@ public:
 	/** Set of specific sockets that can be used for sampling. Select from these with GetSpecificSocketAt and RandomSpecificSocket. */
 	UPROPERTY(EditAnywhere, Category = "Skeleton")
 	TArray<FName> SpecificSockets;
-	
 
 	/** Whether any triangle sampling function is bound. Only used in game. */
 	UPROPERTY()
@@ -534,8 +541,8 @@ public:
 	static const FString MeshNumWeightsName;
 	static const FString NumSpecificBonesName;
 	static const FString SpecificBonesName;
-	static const FString NumSpecificSocketBonesName;
-	static const FString SpecificSocketBonesName;
+	static const FString NumSpecificSocketsName;
+	static const FString SpecificSocketBoneOffsetName;
 	static const FString InstanceTransformName;
 	static const FString InstancePrevTransformName;
 	static const FString InstanceInvDeltaTimeName;
diff --git a/Engine/Plugins/FX/Niagara/Source/Niagara/Private/NDISkeletalMesh_BoneSampling.cpp b/Engine/Plugins/FX/Niagara/Source/Niagara/Private/NDISkeletalMesh_BoneSampling.cpp
--- a/Engine/Plugins/FX/Niagara/Source/Niagara/Private/NDISkeletalMesh_BoneSampling.cpp
+++ b/Engine/Plugins/FX/Niagara/Source/Niagara/Private/NDISkeletalMesh_BoneSampling.cpp
@@ -436,9 +436,14 @@ void UNiagaraDataInterfaceSkeletalMesh::GetSkinnedBoneData(FVectorVMContext& Con
 
 	const FReferenceSkeleton& RefSkel = Accessor.Mesh->RefSkeleton;
 
-	int32 BoneMax = RefSkel.GetNum() - 1;
+	const int32 BoneMax = RefSkel.GetNum() - 1;
+	const int32 BoneAndSocketMax = BoneMax + InstData->SpecificSockets.Num();
 	float InvDt = 1.0f / InstData->DeltaSeconds;
 
+	const int32 SpecificSocketBoneOffset = InstData->SpecificSocketBoneOffset;
+	const TArray<FTransform>& SpecificSocketCurrTransforms = InstData->GetSpecificSocketsCurrBuffer();
+	const TArray<FTransform>& SpecificSocketPrevTransforms = InstData->GetSpecificSocketsPrevBuffer();
+
 	FVector BonePos;
 	FVector BonePrev;
 
@@ -448,31 +453,70 @@ void UNiagaraDataInterfaceSkeletalMesh::GetSkinnedBoneData(FVectorVMContext& Con
 
 	for (int32 i = 0; i < Context.NumInstances; ++i)
 	{
-		int32 Bone = FMath::Clamp(BoneParam.GetAndAdvance(), 0, BoneMax);
+		const float Interp = bInterpolated::Value ? InterpParam.GetAndAdvance() : 1.0f;
 
-		float Interp = 1.0f;
-		if (bInterpolated::Value)
+		// Determine bone or socket
+		int32 Bone = FMath::Clamp(BoneParam.GetAndAdvance(), 0, BoneAndSocketMax);
+		const bool bIsSocket = Bone > BoneMax;
+		if ( bIsSocket )
 		{
-			Interp = InterpParam.GetAndAdvance();
+			const int32 Socket = Bone - SpecificSocketBoneOffset;
+			FTransform CurrSocketTransform = SpecificSocketCurrTransforms[Socket];
+			FTransform PrevSocketTransform = SpecificSocketPrevTransforms[Socket];
+
+			Pos = CurrSocketTransform.GetLocation();
+			TransformHandler.TransformPosition(Pos, Transform);
+
+			if (Output.bNeedsVelocity || bInterpolated::Value)
+			{
+				Prev = PrevSocketTransform.GetLocation();
+				TransformHandler.TransformPosition(Prev, PrevTransform);
+			}
+
+			if (Output.bNeedsRotation)
+			{
+				FQuat Rotation = CurrSocketTransform.GetRotation();
+				if (bInterpolated::Value)
+				{
+					FQuat PrevRotation = PrevSocketTransform.GetRotation();
+					Rotation = FMath::Lerp(PrevRotation, Rotation, Interp);
+				}
+
+				Output.SetRotation(Rotation);
+			}
 		}
+		// Bone
+		else
+		{
+			Pos = SkinningHandler.GetSkinnedBonePosition(Accessor, Bone);
+			TransformHandler.TransformPosition(Pos, Transform);
 
+			if (Output.bNeedsVelocity || bInterpolated::Value)
+			{
+				Prev = SkinningHandler.GetSkinnedBonePreviousPosition(Accessor, Bone);
+				TransformHandler.TransformPosition(Prev, PrevTransform);
+			}
 
-		Pos = SkinningHandler.GetSkinnedBonePosition(Accessor, Bone);
-		TransformHandler.TransformPosition(Pos, Transform);
+			if (Output.bNeedsRotation)
+			{
+				FQuat Rotation = SkinningHandler.GetSkinnedBoneRotation(Accessor, Bone);
+				if (bInterpolated::Value)
+				{
+					FQuat PrevRotation = SkinningHandler.GetSkinnedBonePreviousRotation(Accessor, Bone);
+					Rotation = FMath::Lerp(PrevRotation, Rotation, Interp);
+				}
+
+				Output.SetRotation(Rotation);
+			}
+		}
 
 		if (Output.bNeedsVelocity || bInterpolated::Value)
 		{
-			Prev = SkinningHandler.GetSkinnedBonePreviousPosition(Accessor, Bone);
-			TransformHandler.TransformPosition(Prev, PrevTransform);
+			Pos = FMath::Lerp(Prev, Pos, Interp);
 		}
 
 		if (Output.bNeedsPosition)
 		{
-			if (bInterpolated::Value)
-			{
-				Pos = FMath::Lerp(Prev, Pos, Interp);
-			}
-
 			Output.SetPosition(Pos);
 		}
 
@@ -482,19 +526,6 @@ void UNiagaraDataInterfaceSkeletalMesh::GetSkinnedBoneData(FVectorVMContext& Con
 			Velocity = (Pos - Prev) * InvDt;
 			Output.SetVelocity(Velocity);
 		}
-
-		if (Output.bNeedsRotation)
-		{
-			FQuat Rotation = SkinningHandler.GetSkinnedBoneRotation(Accessor, Bone);
-
-			if (bInterpolated::Value)
-			{
-				FQuat PrevRotation = SkinningHandler.GetSkinnedBonePreviousRotation(Accessor, Bone);
-				Rotation = FMath::Lerp(PrevRotation, Rotation, Interp);
-			}
-
-			Output.SetRotation(Rotation);
-		}
 	}
 }
 
@@ -508,7 +539,7 @@ void UNiagaraDataInterfaceSkeletalMesh::GetSpecificSocketCount(FVectorVMContext&
 
 	VectorVM::FExternalFuncRegisterHandler<int32> OutCount(Context);
 
-	int32 Num = SpecificSockets.Num();
+	const int32 Num = InstData->SpecificSockets.Num();
 	for (int32 i = 0; i < Context.NumInstances; ++i)
 	{
 		*OutCount.GetDestAndAdvance() = Num;
@@ -523,15 +554,16 @@ void UNiagaraDataInterfaceSkeletalMesh::GetSpecificSocketBoneAt(FVectorVMContext
 	VectorVM::FUserPtrHandler<FNDISkeletalMesh_InstanceData> InstData(Context);
 
 	VectorVM::FExternalFuncRegisterHandler<int32> OutSocketBone(Context);
-	const TArray<int32>& SpecificSocketsArray = InstData->SpecificSocketBones;
+	const TArray<FName>& SpecificSocketsArray = InstData->SpecificSockets;
+	const int32 SpecificSocketBoneOffset = InstData->SpecificSocketBoneOffset;
 
 	int32 Max = SpecificSockets.Num() - 1;
 	if (Max != INDEX_NONE)
 	{
 		for (int32 i = 0; i < Context.NumInstances; ++i)
 		{
-			int32 SocketIndex = FMath::Clamp(SocketParam.GetAndAdvance(), 0, Max);
-			*OutSocketBone.GetDestAndAdvance() = SpecificSocketsArray[SocketIndex];
+			const int32 SocketIndex = FMath::Clamp(SocketParam.GetAndAdvance(), 0, Max);
+			*OutSocketBone.GetDestAndAdvance() = SpecificSocketBoneOffset + SocketIndex;
 		}
 	}
 	else
@@ -547,15 +579,16 @@ void UNiagaraDataInterfaceSkeletalMesh::RandomSpecificSocketBone(FVectorVMContex
 	VectorVM::FUserPtrHandler<FNDISkeletalMesh_InstanceData> InstData(Context);
 
 	VectorVM::FExternalFuncRegisterHandler<int32> OutSocketBone(Context);
-	const TArray<int32>& SpecificSocketsArray = InstData->SpecificSocketBones;
+	const TArray<FName>& SpecificSocketsArray = InstData->SpecificSockets;
+	const int32 SpecificSocketBoneOffset = InstData->SpecificSocketBoneOffset;
 
 	int32 Max = SpecificSockets.Num() - 1;
 	if (Max != INDEX_NONE)
 	{
 		for (int32 i = 0; i < Context.NumInstances; ++i)
 		{
-			int32 SocketIndex = Context.RandStream.RandRange(0, Max);
-			*OutSocketBone.GetDestAndAdvance() = SpecificSocketsArray[SocketIndex];
+			const int32 SocketIndex = Context.RandStream.RandRange(0, Max);
+			*OutSocketBone.GetDestAndAdvance() = SpecificSocketBoneOffset + SocketIndex;
 		}
 	}
 	else
diff --git a/Engine/Plugins/FX/Niagara/Source/Niagara/Private/NiagaraDataInterfaceSkeletalMesh.cpp b/Engine/Plugins/FX/Niagara/Source/Niagara/Private/NiagaraDataInterfaceSkeletalMesh.cpp
--- a/Engine/Plugins/FX/Niagara/Source/Niagara/Private/NiagaraDataInterfaceSkeletalMesh.cpp
+++ b/Engine/Plugins/FX/Niagara/Source/Niagara/Private/NiagaraDataInterfaceSkeletalMesh.cpp
@@ -289,10 +289,10 @@ FSkeletalMeshGpuSpawnStaticBuffers::~FSkeletalMeshGpuSpawnStaticBuffers()
 	//ValidSections.Empty();
 }
 
-void FSkeletalMeshGpuSpawnStaticBuffers::Initialise(const FSkeletalMeshLODRenderData& SkeletalMeshLODRenderData, bool bIsGpuUniformlyDistributedSampling, const FSkeletalMeshSamplingLODBuiltData& MeshSamplingLODBuiltData, const TArray<int32>& InSpecificBones, const TArray<int32>& InSpecificSocketBones)
+void FSkeletalMeshGpuSpawnStaticBuffers::Initialise(FNDISkeletalMesh_InstanceData* InstData, const FSkeletalMeshLODRenderData& SkeletalMeshLODRenderData, const FSkeletalMeshSamplingLODBuiltData& MeshSamplingLODBuiltData)
 {
 	SkeletalMeshSamplingLODBuiltData = &MeshSamplingLODBuiltData;
-	bUseGpuUniformlyDistributedSampling = bIsGpuUniformlyDistributedSampling;
+	bUseGpuUniformlyDistributedSampling = InstData->bIsGpuUniformlyDistributedSampling;
 
 	LODRenderData = &SkeletalMeshLODRenderData;
 	TriangleCount = SkeletalMeshLODRenderData.MultiSizeIndexContainer.GetIndexBuffer()->Num() / 3;
@@ -302,27 +302,19 @@ void FSkeletalMeshGpuSpawnStaticBuffers::Initialise(const FSkeletalMeshLODRender
 
 	// Copy Specific Bones / Socket data into arrays that the renderer will use to create read buffers
 	//-TODO: Exclude setting up these arrays if we don't sample from them
-	NumSpecificBones = InSpecificBones.Num();
+	NumSpecificBones = InstData->SpecificBones.Num();
 	if (NumSpecificBones > 0)
 	{
 		SpecificBonesArray.Reserve(NumSpecificBones);
-		for ( int32 v : InSpecificBones )
+		for ( int32 v : InstData->SpecificBones )
 		{
 			check(v <= 65535);
 			SpecificBonesArray.Add(v);
 		}
 	}
 
-	NumSpecificSocketBones = InSpecificSocketBones.Num();
-	if (NumSpecificSocketBones > 0)
-	{
-		SpecificSocketBonesArray.Reserve(NumSpecificSocketBones);
-		for (int32 v : InSpecificSocketBones)
-		{
-			check(v <= 65535);
-			SpecificSocketBonesArray.Add(v);
-		}
-	}
+	NumSpecificSockets = InstData->SpecificSockets.Num();
+	SpecificSocketBoneOffset = InstData->SpecificSocketBoneOffset;
 }
 
 void FSkeletalMeshGpuSpawnStaticBuffers::InitRHI()
@@ -406,15 +398,6 @@ void FSkeletalMeshGpuSpawnStaticBuffers::InitRHI()
 		SpecificBonesBuffer = RHICreateVertexBuffer(NumSpecificBones * sizeof(uint16), BUF_Static | BUF_ShaderResource, CreateInfo);
 		SpecificBonesSRV = RHICreateShaderResourceView(SpecificBonesBuffer, sizeof(uint16), PF_R16_UINT);
 	}
-
-	if (NumSpecificSocketBones > 0)
-	{
-		FRHIResourceCreateInfo CreateInfo;
-		CreateInfo.ResourceArray = &SpecificSocketBonesArray;
-
-		SpecificSocketBonesBuffer = RHICreateVertexBuffer(NumSpecificSocketBones * sizeof(uint16), BUF_Static | BUF_ShaderResource, CreateInfo);
-		SpecificSocketBonesSRV = RHICreateShaderResourceView(SpecificSocketBonesBuffer, sizeof(uint16), PF_R16_UINT);
-	}
 }
 
 void FSkeletalMeshGpuSpawnStaticBuffers::ReleaseRHI()
@@ -422,9 +405,6 @@ void FSkeletalMeshGpuSpawnStaticBuffers::ReleaseRHI()
 	SpecificBonesBuffer.SafeRelease();
 	SpecificBonesSRV.SafeRelease();
 
-	SpecificSocketBonesBuffer.SafeRelease();
-	SpecificSocketBonesSRV.SafeRelease();
-
 	BufferTriangleUniformSamplerProbaRHI.SafeRelease();
 	BufferTriangleUniformSamplerProbaSRV.SafeRelease();
 	BufferTriangleUniformSamplerAliasRHI.SafeRelease();
@@ -450,7 +430,7 @@ FSkeletalMeshGpuDynamicBufferProxy::~FSkeletalMeshGpuDynamicBufferProxy()
 	//
 }
 
-void FSkeletalMeshGpuDynamicBufferProxy::Initialise(const FReferenceSkeleton& RefSkel, const FSkeletalMeshLODRenderData& SkeletalMeshLODRenderData)
+void FSkeletalMeshGpuDynamicBufferProxy::Initialise(const FReferenceSkeleton& RefSkel, const FSkeletalMeshLODRenderData& SkeletalMeshLODRenderData, uint32 InSamplingSocketCount)
 {
 	SectionBoneCount = 0;
 	for (const FSkelMeshRenderSection& Section : SkeletalMeshLODRenderData.RenderSections)
@@ -459,6 +439,7 @@ void FSkeletalMeshGpuDynamicBufferProxy::Initialise(const FReferenceSkeleton& Re
 	}
 
 	SamplingBoneCount = RefSkel.GetNum();
+	SamplingSocketCount = InSamplingSocketCount;
 }
 
 void FSkeletalMeshGpuDynamicBufferProxy::InitRHI()
@@ -470,7 +451,7 @@ void FSkeletalMeshGpuDynamicBufferProxy::InitRHI()
 		Buffer.SectionBuffer = RHICreateVertexBuffer(sizeof(FVector4) * 3 * SectionBoneCount, BUF_ShaderResource | BUF_Dynamic, CreateInfo);
 		Buffer.SectionSRV = RHICreateShaderResourceView(Buffer.SectionBuffer, sizeof(FVector4), PF_A32B32G32R32F);
 
-		Buffer.SamplingBuffer = RHICreateVertexBuffer(sizeof(FVector4) * 2 * SamplingBoneCount, BUF_ShaderResource | BUF_Dynamic, CreateInfo);
+		Buffer.SamplingBuffer = RHICreateVertexBuffer(sizeof(FVector4) * 2 * (SamplingBoneCount + SamplingSocketCount), BUF_ShaderResource | BUF_Dynamic, CreateInfo);
 		Buffer.SamplingSRV = RHICreateShaderResourceView(Buffer.SamplingBuffer, sizeof(FVector4), PF_A32B32G32R32F);
 	}
 }
@@ -530,12 +511,21 @@ void FSkeletalMeshGpuDynamicBufferProxy::NewFrame(const FNDISkeletalMesh_Instanc
 			const TArray<FTransform>& ComponentTransforms = SkelComp->GetComponentSpaceTransforms();
 			check(ComponentTransforms.Num() == SamplingBoneCount);
 
-			BoneSamplingData.Reserve(ComponentTransforms.Num() * 2);
+			BoneSamplingData.Reserve((SamplingBoneCount + SamplingSocketCount) * 2);
+
+			// Append bones
+			for (const FTransform& BoneTransform : ComponentTransforms)
+			{
+				const FQuat Rotation = BoneTransform.GetRotation();
+				BoneSamplingData.Add(BoneTransform.GetLocation());
+				BoneSamplingData.Add(FVector4(Rotation.X, Rotation.Y, Rotation.Z, Rotation.W));
+			}
 
-			for (const FTransform& boneMat : ComponentTransforms)
+			// Append sockets
+			for (const FTransform& SocketTransform : InstanceData->GetSpecificSocketsCurrBuffer())
 			{
-				const FQuat Rotation = boneMat.GetRotation();
-				BoneSamplingData.Add(boneMat.GetLocation());
+				const FQuat Rotation = SocketTransform.GetRotation();
+				BoneSamplingData.Add(SocketTransform.GetLocation());
 				BoneSamplingData.Add(FVector4(Rotation.X, Rotation.Y, Rotation.Z, Rotation.W));
 			}
 		}
@@ -592,8 +582,8 @@ struct FNDISkeletalMeshParametersName
 	FString MeshNumWeightsName;
 	FString NumSpecificBonesName;
 	FString SpecificBonesName;
-	FString NumSpecificSocketBonesName;
-	FString SpecificSocketBonesName;
+	FString NumSpecificSocketsName;
+	FString SpecificSocketBoneOffsetName;
 	FString InstanceTransformName;
 	FString InstancePrevTransformName;
 	FString InstanceInvDeltaTimeName;
@@ -622,8 +612,8 @@ static void GetNiagaraDataInterfaceParametersName(FNDISkeletalMeshParametersName
 	Names.MeshNumWeightsName = UNiagaraDataInterfaceSkeletalMesh::MeshNumWeightsName + Suffix;
 	Names.NumSpecificBonesName = UNiagaraDataInterfaceSkeletalMesh::NumSpecificBonesName + Suffix;
 	Names.SpecificBonesName = UNiagaraDataInterfaceSkeletalMesh::SpecificBonesName + Suffix;
-	Names.NumSpecificSocketBonesName = UNiagaraDataInterfaceSkeletalMesh::NumSpecificSocketBonesName + Suffix;
-	Names.SpecificSocketBonesName = UNiagaraDataInterfaceSkeletalMesh::SpecificSocketBonesName + Suffix;
+	Names.NumSpecificSocketsName = UNiagaraDataInterfaceSkeletalMesh::NumSpecificSocketsName + Suffix;
+	Names.SpecificSocketBoneOffsetName = UNiagaraDataInterfaceSkeletalMesh::SpecificSocketBoneOffsetName + Suffix;
 	Names.InstanceTransformName = UNiagaraDataInterfaceSkeletalMesh::InstanceTransformName + Suffix;
 	Names.InstancePrevTransformName = UNiagaraDataInterfaceSkeletalMesh::InstancePrevTransformName + Suffix;
 	Names.InstanceInvDeltaTimeName = UNiagaraDataInterfaceSkeletalMesh::InstanceInvDeltaTimeName + Suffix;
@@ -657,8 +647,8 @@ struct FNiagaraDataInterfaceParametersCS_SkeletalMesh : public FNiagaraDataInter
 		MeshNumWeights.Bind(ParameterMap, *ParamNames.MeshNumWeightsName);
 		NumSpecificBones.Bind(ParameterMap, *ParamNames.NumSpecificBonesName);
 		SpecificBones.Bind(ParameterMap, *ParamNames.SpecificBonesName);
-		NumSpecificSocketBones.Bind(ParameterMap, *ParamNames.NumSpecificSocketBonesName);
-		SpecificSocketBones.Bind(ParameterMap, *ParamNames.SpecificSocketBonesName);
+		NumSpecificSockets.Bind(ParameterMap, *ParamNames.NumSpecificSocketsName);
+		SpecificSocketBoneOffset.Bind(ParameterMap, *ParamNames.SpecificSocketBoneOffsetName);
 		InstanceTransform.Bind(ParameterMap, *ParamNames.InstanceTransformName);
 		InstancePrevTransform.Bind(ParameterMap, *ParamNames.InstancePrevTransformName);
 		InstanceInvDeltaTime.Bind(ParameterMap, *ParamNames.InstanceInvDeltaTimeName);
@@ -687,8 +677,8 @@ struct FNiagaraDataInterfaceParametersCS_SkeletalMesh : public FNiagaraDataInter
 		Ar << MeshNumWeights;
 		Ar << NumSpecificBones;
 		Ar << SpecificBones;
-		Ar << NumSpecificSocketBones;
-		Ar << SpecificSocketBones;
+		Ar << NumSpecificSockets;
+		Ar << SpecificSocketBoneOffset;
 		Ar << InstanceTransform;
 		Ar << InstancePrevTransform;
 		Ar << InstanceInvDeltaTime;
@@ -765,9 +755,9 @@ struct FNiagaraDataInterfaceParametersCS_SkeletalMesh : public FNiagaraDataInter
 			SetShaderValue(RHICmdList, ComputeShaderRHI, NumSpecificBones, StaticBuffers->GetNumSpecificBones());
 			SetSRVParameter(RHICmdList, ComputeShaderRHI, SpecificBones, SpecificBonesSRV);
 
-			FRHIShaderResourceView* SpecificSocketBonesSRV = StaticBuffers->GetNumSpecificSocketBones() > 0 ? StaticBuffers->GetSpecificSocketBonesSRV() : FNiagaraRenderer::GetDummyUIntBuffer().SRV.GetReference();
-			SetShaderValue(RHICmdList, ComputeShaderRHI, NumSpecificSocketBones, StaticBuffers->GetNumSpecificSocketBones());
-			SetSRVParameter(RHICmdList, ComputeShaderRHI, SpecificSocketBones, SpecificSocketBonesSRV);
+			FShaderResourceViewRHIParamRef SpecificSocketBonesSRV = StaticBuffers->GetNumSpecificSocketBones() > 0 ? StaticBuffers->GetSpecificSocketBonesSRV() : FNiagaraRenderer::GetDummyUIntBuffer().SRV.GetReference();
+			SetShaderValue(RHICmdList, ComputeShaderRHI, NumSpecificSockets, StaticBuffers->GetNumSpecificSockets());
+			SetShaderValue(RHICmdList, ComputeShaderRHI, SpecificSocketBoneOffset, StaticBuffers->GetSpecificSocketBoneOffset());
 
 			SetShaderValue(RHICmdList, ComputeShaderRHI, InstanceTransform, InstanceData->Transform);
 			SetShaderValue(RHICmdList, ComputeShaderRHI, InstancePrevTransform, InstanceData->PrevTransform);
@@ -805,8 +795,8 @@ struct FNiagaraDataInterfaceParametersCS_SkeletalMesh : public FNiagaraDataInter
 
 			SetShaderValue(RHICmdList, ComputeShaderRHI, NumSpecificBones, 0);
 			SetSRVParameter(RHICmdList, ComputeShaderRHI, SpecificBones, FNiagaraRenderer::GetDummyUIntBuffer().SRV);
-			SetShaderValue(RHICmdList, ComputeShaderRHI, NumSpecificSocketBones, 0);
-			SetSRVParameter(RHICmdList, ComputeShaderRHI, SpecificBones, FNiagaraRenderer::GetDummyUIntBuffer().SRV);
+			SetShaderValue(RHICmdList, ComputeShaderRHI, NumSpecificSockets, 0);
+			SetShaderValue(RHICmdList, ComputeShaderRHI, SpecificSocketBoneOffset, 0);
 
 			SetShaderValue(RHICmdList, ComputeShaderRHI, InstanceTransform, FMatrix::Identity);
 			SetShaderValue(RHICmdList, ComputeShaderRHI, InstancePrevTransform, FMatrix::Identity);
@@ -839,8 +829,8 @@ private:
 	FShaderParameter MeshNumWeights;
 	FShaderParameter NumSpecificBones;
 	FShaderResourceParameter SpecificBones;
-	FShaderParameter NumSpecificSocketBones;
-	FShaderResourceParameter SpecificSocketBones;
+	FShaderParameter NumSpecificSockets;
+	FShaderParameter SpecificSocketBoneOffset;
 	FShaderParameter InstanceTransform;
 	FShaderParameter InstancePrevTransform;
 	FShaderParameter InstanceInvDeltaTime;
@@ -1193,61 +1183,67 @@ bool FNDISkeletalMesh_InstanceData::Init(UNiagaraDataInterfaceSkeletalMesh* Inte
 		}
 	}
 
+	// Gather specific bones information
 	FReferenceSkeleton& RefSkel = Mesh->RefSkeleton;
-	SpecificBones.SetNumUninitialized(Interface->SpecificBones.Num());
-	TArray<FName, TInlineAllocator<16>> MissingBones;
-	for (int32 BoneIdx = 0; BoneIdx < SpecificBones.Num(); ++BoneIdx)
 	{
-		FName BoneName = Interface->SpecificBones[BoneIdx];
-		int32 Bone = RefSkel.FindBoneIndex(BoneName);
-		if (Bone == INDEX_NONE)
+		SpecificBones.SetNumUninitialized(Interface->SpecificBones.Num());
+		TArray<FName, TInlineAllocator<16>> MissingBones;
+		for (int32 BoneIdx = 0; BoneIdx < SpecificBones.Num(); ++BoneIdx)
 		{
-			MissingBones.Add(BoneName);
-			SpecificBones[BoneIdx] = 0;
-		}
-		else
-		{
-			SpecificBones[BoneIdx] = Bone;
+			FName BoneName = Interface->SpecificBones[BoneIdx];
+			int32 Bone = RefSkel.FindBoneIndex(BoneName);
+			if (Bone == INDEX_NONE)
+			{
+				MissingBones.Add(BoneName);
+				SpecificBones[BoneIdx] = 0;
+			}
+			else
+			{
+				SpecificBones[BoneIdx] = Bone;
+			}
 		}
-	}
 
-	if (MissingBones.Num() > 0)
-	{
-		UE_LOG(LogNiagara, Warning, TEXT("Skeletal Mesh Data Interface is trying to sample from bones that don't exist in it's skeleton.\nMesh: %s\nBones: "), *Mesh->GetName());
-		for (FName BoneName : MissingBones)
+		if (MissingBones.Num() > 0)
 		{
-			UE_LOG(LogNiagara, Warning, TEXT("%s\n"), *BoneName.ToString());
+			UE_LOG(LogNiagara, Warning, TEXT("Skeletal Mesh Data Interface is trying to sample from bones that don't exist in it's skeleton.\nMesh: %s\nBones: "), *Mesh->GetName());
+			for (FName BoneName : MissingBones)
+			{
+				UE_LOG(LogNiagara, Warning, TEXT("%s\n"), *BoneName.ToString());
+			}
 		}
 	}
 
-	SpecificSockets.SetNumUninitialized(Interface->SpecificSockets.Num());
-	SpecificSocketBones.SetNumUninitialized(Interface->SpecificSockets.Num());
-	TArray<FName, TInlineAllocator<16>> MissingSockets;
-	for (int32 SocketIdx = 0; SocketIdx < SpecificSockets.Num(); ++SocketIdx)
+	// Gather specific socket information
 	{
-		FName SocketName = Interface->SpecificSockets[SocketIdx];
-		int32 SocketIndex = INDEX_NONE;
-		USkeletalMeshSocket* Socket = Mesh->FindSocketAndIndex(SocketName, SocketIndex);
-		if (SocketIndex == INDEX_NONE)
+		SpecificSockets = Interface->SpecificSockets;
+		SpecificSocketBoneOffset = Mesh->RefSkeleton.GetNum();
+
+		SpecificSocketTransformsIndex = 0;
+		SpecificSocketTransforms[0].Reset(SpecificSockets.Num());
+		SpecificSocketTransforms[0].AddDefaulted(SpecificSockets.Num());
+		UpdateSpecificSocketTransforms();
+		for (int32 i=1; i < SpecificSocketTransforms.Num(); ++i)
 		{
-			MissingSockets.Add(SocketName);
-			SpecificSockets[SocketIdx] = 0;
-			SpecificSocketBones[SocketIdx] = 0;
+			SpecificSocketTransforms[i].Reset(SpecificSockets.Num());
+			SpecificSocketTransforms[i].Append(SpecificSocketTransforms[0]);
 		}
-		else
+
+		TArray<FName, TInlineAllocator<16>> MissingSockets;
+		for (FName SocketName : SpecificSockets)
 		{
-			check(Socket);
-			SpecificSockets[SocketIdx] = SocketIndex;
-			SpecificSocketBones[SocketIdx] = RefSkel.FindBoneIndex(Socket->BoneName);
+			if (Mesh->FindSocket(SocketName) == nullptr)
+			{
+				MissingSockets.Add(SocketName);
+			}
 		}
-	}
 
-	if (MissingSockets.Num() > 0)
-	{
-		UE_LOG(LogNiagara, Warning, TEXT("Skeletal Mesh Data Interface is trying to sample from sockets that don't exist in it's skeleton.\nMesh: %s\nSockets: "), *Mesh->GetName());
-		for (FName SocketName : MissingSockets)
+		if (MissingSockets.Num() > 0)
 		{
-			UE_LOG(LogNiagara, Warning, TEXT("%s\n"), *SocketName.ToString());
+			UE_LOG(LogNiagara, Warning, TEXT("Skeletal Mesh Data Interface is trying to sample from sockets that don't exist in it's skeleton.\nMesh: %s\nSockets: "), *Mesh->GetName());
+			for (FName SocketName : MissingSockets)
+			{
+				UE_LOG(LogNiagara, Warning, TEXT("%s\n"), *SocketName.ToString());
+			}
 		}
 	}
 
@@ -1271,11 +1267,11 @@ bool FNDISkeletalMesh_InstanceData::Init(UNiagaraDataInterfaceSkeletalMesh* Inte
 		}
 
 		MeshGpuSpawnStaticBuffers = new FSkeletalMeshGpuSpawnStaticBuffers();
-		MeshGpuSpawnStaticBuffers->Initialise(LODData, bIsGpuUniformlyDistributedSampling, SamplingInfo.GetBuiltData().WholeMeshBuiltData[LODIndex], SpecificBones, SpecificSocketBones);
+		MeshGpuSpawnStaticBuffers->Initialise(this, LODData, SamplingInfo.GetBuiltData().WholeMeshBuiltData[LODIndex]);
 		BeginInitResource(MeshGpuSpawnStaticBuffers);
 
 		MeshGpuSpawnDynamicBuffers = new FSkeletalMeshGpuDynamicBufferProxy();
-		MeshGpuSpawnDynamicBuffers->Initialise(RefSkel, LODData);
+		MeshGpuSpawnDynamicBuffers->Initialise(RefSkel, LODData, SpecificSockets.Num());
 		BeginInitResource(MeshGpuSpawnDynamicBuffers);
 	}
 
@@ -1335,6 +1331,7 @@ bool FNDISkeletalMesh_InstanceData::Tick(UNiagaraDataInterfaceSkeletalMesh* Inte
 	else
 	{
 		DeltaSeconds = InDeltaSeconds;
+
 		if (Component.IsValid() && Mesh)
 		{
 			PrevTransform = Transform;
@@ -1350,6 +1347,10 @@ bool FNDISkeletalMesh_InstanceData::Tick(UNiagaraDataInterfaceSkeletalMesh* Inte
 			TransformInverseTransposed = FMatrix::Identity;
 		}
 
+		// Cache socket transforms to avoid potentially calculating them multiple times during the VM
+		SpecificSocketTransformsIndex = (SpecificSocketTransformsIndex + 1) % SpecificSocketTransforms.Num();
+		UpdateSpecificSocketTransforms();
+
 		if (MeshGpuSpawnDynamicBuffers)
 		{
 			USkeletalMeshComponent* Comp = Cast<USkeletalMeshComponent>(Component.Get());
@@ -1366,6 +1367,35 @@ bool FNDISkeletalMesh_InstanceData::Tick(UNiagaraDataInterfaceSkeletalMesh* Inte
 	}
 }
 
+void FNDISkeletalMesh_InstanceData::UpdateSpecificSocketTransforms()
+{
+	USkeletalMeshComponent* SkelComp = Cast<USkeletalMeshComponent>(Component.Get());
+	TArray<FTransform>& WriteBuffer = GetSpecificSocketsWriteBuffer();
+
+	//-TODO: We may need to handle skinning mode changes here
+	if (SkelComp != nullptr)
+	{
+		for (int32 i = 0; i < SpecificSockets.Num(); ++i)
+		{
+			WriteBuffer[i] = SkelComp->GetSocketTransform(SpecificSockets[i], RTS_Component);
+		}
+	}
+	else if ( Mesh != nullptr )
+	{
+		for (int32 i = 0; i < SpecificSockets.Num(); ++i)
+		{
+			WriteBuffer[i] = FTransform(Mesh->GetComposedRefPoseMatrix(SpecificSockets[i]));
+		}
+	}
+	else
+	{
+		for (int32 i = 0; i < SpecificSockets.Num(); ++i)
+		{
+			WriteBuffer[i] = FTransform::Identity;
+		}
+	}
+}
+
 bool FNDISkeletalMesh_InstanceData::HasColorData()
 {
 	check(Mesh);
@@ -1778,8 +1808,8 @@ const FString UNiagaraDataInterfaceSkeletalMesh::MeshNumTexCoordName(TEXT("MeshN
 const FString UNiagaraDataInterfaceSkeletalMesh::MeshNumWeightsName(TEXT("MeshNumWeights_"));
 const FString UNiagaraDataInterfaceSkeletalMesh::NumSpecificBonesName(TEXT("NumSpecificBones_"));
 const FString UNiagaraDataInterfaceSkeletalMesh::SpecificBonesName(TEXT("SpecificBones_"));
-const FString UNiagaraDataInterfaceSkeletalMesh::NumSpecificSocketBonesName(TEXT("NumSpecificSocketBones_"));
-const FString UNiagaraDataInterfaceSkeletalMesh::SpecificSocketBonesName(TEXT("SpecificSocketBones_"));
+const FString UNiagaraDataInterfaceSkeletalMesh::NumSpecificSocketsName(TEXT("NumSpecificSockets_"));
+const FString UNiagaraDataInterfaceSkeletalMesh::SpecificSocketBoneOffsetName(TEXT("SpecificSocketBoneOffset_"));
 const FString UNiagaraDataInterfaceSkeletalMesh::InstanceTransformName(TEXT("InstanceTransform_"));
 const FString UNiagaraDataInterfaceSkeletalMesh::InstancePrevTransformName(TEXT("InstancePrevTransform_"));
 const FString UNiagaraDataInterfaceSkeletalMesh::InstanceInvDeltaTimeName(TEXT("InstanceInvDeltaTime_"));
diff --git a/Engine/Plugins/FX/Niagara/Source/NiagaraCore/Private/NiagaraCustomVersion.cpp b/Engine/Plugins/FX/Niagara/Source/NiagaraCore/Private/NiagaraCustomVersion.cpp
--- a/Engine/Plugins/FX/Niagara/Source/NiagaraCore/Private/NiagaraCustomVersion.cpp
+++ b/Engine/Plugins/FX/Niagara/Source/NiagaraCore/Private/NiagaraCustomVersion.cpp
@@ -8,4 +8,4 @@ const FGuid FNiagaraCustomVersion::GUID(0xFCF57AFA, 0x50764283, 0xB9A9E658, 0xFF
 // Register the custom version with core
 FCustomVersionRegistration GRegisterNiagaraCustomVersion(FNiagaraCustomVersion::GUID, FNiagaraCustomVersion::LatestVersion, TEXT("NiagaraVer"));
 
-const FGuid FNiagaraCustomVersion::LatestScriptCompileVersion(0x397D4E10, 0x87CC4DF8, 0x9EDAD0CA, 0xE837CA88);
\ No newline at end of file
+const FGuid FNiagaraCustomVersion::LatestScriptCompileVersion(0xC543A0CC, 0xD850F04D, 0xBBD65664, 0xF181DFD1);
\ No newline at end of file
