#UE4 Fix a warning about a shadow variable

diff --git a/Engine/Plugins/FX/Niagara/Source/NiagaraEditor/Private/NiagaraNodeStaticSwitch.cpp b/Engine/Plugins/FX/Niagara/Source/NiagaraEditor/Private/NiagaraNodeStaticSwitch.cpp
--- a/Engine/Plugins/FX/Niagara/Source/NiagaraEditor/Private/NiagaraNodeStaticSwitch.cpp
+++ b/Engine/Plugins/FX/Niagara/Source/NiagaraEditor/Private/NiagaraNodeStaticSwitch.cpp
@@ -292,9 +292,9 @@ bool UNiagaraNodeStaticSwitch::GetVarIndex(FHlslNiagaraTranslator* Translator, i
 		if (MaxValue > 0)
 		{
 			// do a sanity check here if the number of pins actually matches the enum count (which might have changed in the meantime without us noticing)
-			TArray<UEdGraphPin*> Pins;
-			GetOutputPins(Pins);
-			int32 OutputPinCount = Pins.Num() - 1;
+			TArray<UEdGraphPin*> LocalOutputPins;
+			GetOutputPins(LocalOutputPins);
+			int32 OutputPinCount = LocalOutputPins.Num() - 1;
 			int32 ReservedValues = (InputPinCount / OutputPinCount);
 			if (OutputPinCount > 0 && (MaxValue > ReservedValues || MaxValue < ReservedValues))
 			{
