- NiagaraMeshVertexFactory add missing function VertexFactoryGetPrimitiveId

#rnx

diff --git a/Engine/Plugins/FX/Niagara/Shaders/Private/NiagaraMeshVertexFactory.ush b/Engine/Plugins/FX/Niagara/Shaders/Private/NiagaraMeshVertexFactory.ush
--- a/Engine/Plugins/FX/Niagara/Shaders/Private/NiagaraMeshVertexFactory.ush
+++ b/Engine/Plugins/FX/Niagara/Shaders/Private/NiagaraMeshVertexFactory.ush
@@ -822,6 +822,11 @@ float4 VertexFactoryGetPreviousWorldPosition(FVertexFactoryInput Input, FVertexF
 	return CalcPrevWorldPosition(Input, Intermediates);
 }
 
+uint VertexFactoryGetPrimitiveId(FVertexFactoryInterpolantsVSToPS Interpolants)
+{
+    return 0;
+}
+
 #if USING_TESSELLATION
 	struct FVertexFactoryInterpolantsVSToDS
 	{
