Fixed a bug where compiling with constant static switches could sometimes result in the wrong switch path being taken.

diff --git a/Engine/Plugins/FX/Niagara/Source/NiagaraEditor/Private/NiagaraCompiler.cpp b/Engine/Plugins/FX/Niagara/Source/NiagaraEditor/Private/NiagaraCompiler.cpp
--- a/Engine/Plugins/FX/Niagara/Source/NiagaraEditor/Private/NiagaraCompiler.cpp
+++ b/Engine/Plugins/FX/Niagara/Source/NiagaraEditor/Private/NiagaraCompiler.cpp
@@ -342,7 +342,7 @@ void FNiagaraCompileRequestData::DeepCopyGraphs(UNiagaraScriptSource* ScriptSour
 }
 
 
-void FNiagaraCompileRequestData::FinishPrecompile(UNiagaraScriptSource* ScriptSource, const TArray<FNiagaraVariable>& EncounterableVariables, ENiagaraScriptUsage InUsage)
+void FNiagaraCompileRequestData::FinishPrecompile(UNiagaraScriptSource* ScriptSource, const TArray<FNiagaraVariable>& EncounterableVariables, ENiagaraScriptUsage InUsage, FCompileConstantResolver* ConstantResolver)
 {
 	{
 		ENiagaraScriptCompileStatusEnum = StaticEnum<ENiagaraScriptCompileStatus>();
@@ -358,12 +358,17 @@ void FNiagaraCompileRequestData::FinishPrecompile(UNiagaraScriptSource* ScriptSo
 		{
 			// Map all for this output node
 			FNiagaraParameterMapHistoryBuilder Builder;
+			if (ConstantResolver)
+			{
+				Builder.ConstantResolver = *ConstantResolver;
+			}
 			Builder.RegisterEncounterableVariables(EncounterableVariables);
 
 			FString TranslationName = TEXT("Emitter");
 			Builder.BeginTranslation(TranslationName);
 			Builder.EnableScriptWhitelist(true, FoundOutputNode->GetUsage());
 			Builder.BuildParameterMaps(FoundOutputNode, true);
+			
 			TArray<FNiagaraParameterMapHistory> Histories = Builder.Histories;
 			ensure(Histories.Num() <= 1);
 
@@ -445,7 +450,7 @@ TSharedPtr<FNiagaraCompileRequestDataBase, ESPMode::ThreadSafe> FNiagaraEditorMo
 		UNiagaraScriptSource* Source = Cast<UNiagaraScriptSource>(Script->GetSource());
 		BasePtr->DeepCopyGraphs(Source, Script->GetUsage());
 		const TArray<FNiagaraVariable> EncounterableVariables;
-		BasePtr->FinishPrecompile(Source, EncounterableVariables, Script->GetUsage());
+		BasePtr->FinishPrecompile(Source, EncounterableVariables, Script->GetUsage(), nullptr);
 	}
 	else if (System)
 	{
@@ -475,7 +480,7 @@ TSharedPtr<FNiagaraCompileRequestDataBase, ESPMode::ThreadSafe> FNiagaraEditorMo
 		// Now deep copy the system graphs, skipping traversal into any emitter references.
 		UNiagaraScriptSource* Source = Cast<UNiagaraScriptSource>(System->GetSystemSpawnScript()->GetSource());
 		BasePtr->DeepCopyGraphs(Source, ENiagaraScriptUsage::SystemSpawnScript);
-		BasePtr->FinishPrecompile(Source, EncounterableVars, ENiagaraScriptUsage::SystemSpawnScript);
+		BasePtr->FinishPrecompile(Source, EncounterableVars, ENiagaraScriptUsage::SystemSpawnScript, nullptr);
 
 		// Add the User and System variables that we did encounter to the list that emitters might also encounter.
 		BasePtr->GatherPreCompiledVariables(TEXT("User"), EncounterableVars);
@@ -485,7 +490,8 @@ TSharedPtr<FNiagaraCompileRequestDataBase, ESPMode::ThreadSafe> FNiagaraEditorMo
 		for (int32 i = 0; i < System->GetEmitterHandles().Num(); i++)
 		{
 			const FNiagaraEmitterHandle& Handle = System->GetEmitterHandle(i);
-			BasePtr->EmitterData[i]->FinishPrecompile(Cast<UNiagaraScriptSource>(Handle.GetInstance()->GraphSource), EncounterableVars, ENiagaraScriptUsage::EmitterSpawnScript);
+			FCompileConstantResolver ConstantResolver(Handle.GetInstance());
+			BasePtr->EmitterData[i]->FinishPrecompile(Cast<UNiagaraScriptSource>(Handle.GetInstance()->GraphSource), EncounterableVars, ENiagaraScriptUsage::EmitterSpawnScript, &ConstantResolver);
 			BasePtr->MergeInEmitterPrecompiledData(BasePtr->EmitterData[i].Get());
 		}
 
diff --git a/Engine/Plugins/FX/Niagara/Source/NiagaraEditor/Private/NiagaraHlslTranslator.cpp b/Engine/Plugins/FX/Niagara/Source/NiagaraEditor/Private/NiagaraHlslTranslator.cpp
--- a/Engine/Plugins/FX/Niagara/Source/NiagaraEditor/Private/NiagaraHlslTranslator.cpp
+++ b/Engine/Plugins/FX/Niagara/Source/NiagaraEditor/Private/NiagaraHlslTranslator.cpp
@@ -45,6 +45,7 @@
 #include "NiagaraShaderCompilationManager.h"
 
 #include "NiagaraEditorSettings.h"
+#include "NiagaraNodeStaticSwitch.h"
 
 #define LOCTEXT_NAMESPACE "NiagaraCompiler"
 
@@ -5431,6 +5432,14 @@ int32 FHlslNiagaraTranslator::CompileOutputPin(const UEdGraphPin* InPin)
 {
 	SCOPE_CYCLE_COUNTER(STAT_NiagaraEditor_HlslTranslator_CompileOutputPin);
 
+	if (InPin)
+	{
+		if (UNiagaraNodeStaticSwitch* SwitchNode = Cast<UNiagaraNodeStaticSwitch>(InPin->GetOwningNode()))
+		{
+			SwitchNode->UpdateCompilerConstantValue(this);
+		}
+	}
+
 	// The incoming pin to compile may be pointing to a reroute node. If so, we just jump over it
 	// to where it really came from.
 	UEdGraphPin* Pin = UNiagaraNode::TraceOutputPin(const_cast<UEdGraphPin*>(InPin));
diff --git a/Engine/Plugins/FX/Niagara/Source/NiagaraEditor/Private/NiagaraHlslTranslator.h b/Engine/Plugins/FX/Niagara/Source/NiagaraEditor/Private/NiagaraHlslTranslator.h
--- a/Engine/Plugins/FX/Niagara/Source/NiagaraEditor/Private/NiagaraHlslTranslator.h
+++ b/Engine/Plugins/FX/Niagara/Source/NiagaraEditor/Private/NiagaraHlslTranslator.h
@@ -96,7 +96,7 @@ public:
 	const FString& GetUniqueEmitterName() const { return EmitterUniqueName; }
 	void VisitReferencedGraphs(UNiagaraGraph* InSrcGraph, UNiagaraGraph* InDupeGraph, ENiagaraScriptUsage InUsage);
 	void DeepCopyGraphs(UNiagaraScriptSource* ScriptSource, ENiagaraScriptUsage InUsage);
-	void FinishPrecompile(UNiagaraScriptSource* ScriptSource, const TArray<FNiagaraVariable>& EncounterableVariables, ENiagaraScriptUsage InUsage);
+	void FinishPrecompile(UNiagaraScriptSource* ScriptSource, const TArray<FNiagaraVariable>& EncounterableVariables, ENiagaraScriptUsage InUsage, FCompileConstantResolver* ConstantResolver);
 	virtual int32 GetDependentRequestCount() const override {
 		return EmitterData.Num();
 	};
