Speculative fix for CIS errors.

#rnx

diff --git a/Engine/Plugins/FX/Niagara/Source/Niagara/Private/NiagaraDataInterfaceSkeletalMesh.cpp b/Engine/Plugins/FX/Niagara/Source/Niagara/Private/NiagaraDataInterfaceSkeletalMesh.cpp
--- a/Engine/Plugins/FX/Niagara/Source/Niagara/Private/NiagaraDataInterfaceSkeletalMesh.cpp
+++ b/Engine/Plugins/FX/Niagara/Source/Niagara/Private/NiagaraDataInterfaceSkeletalMesh.cpp
@@ -755,7 +755,7 @@ struct FNiagaraDataInterfaceParametersCS_SkeletalMesh : public FNiagaraDataInter
 			SetShaderValue(RHICmdList, ComputeShaderRHI, NumSpecificBones, StaticBuffers->GetNumSpecificBones());
 			SetSRVParameter(RHICmdList, ComputeShaderRHI, SpecificBones, SpecificBonesSRV);
 
-			FShaderResourceViewRHIParamRef SpecificSocketBonesSRV = StaticBuffers->GetNumSpecificSocketBones() > 0 ? StaticBuffers->GetSpecificSocketBonesSRV() : FNiagaraRenderer::GetDummyUIntBuffer().SRV.GetReference();
+			FShaderResourceViewRHIParamRef SpecificSocketBonesSRV = StaticBuffers->GetNumSpecificSockets() > 0 ? StaticBuffers->GetSpecificSocketBonesSRV() : FNiagaraRenderer::GetDummyUIntBuffer().SRV.GetReference();
 			SetShaderValue(RHICmdList, ComputeShaderRHI, NumSpecificSockets, StaticBuffers->GetNumSpecificSockets());
 			SetShaderValue(RHICmdList, ComputeShaderRHI, SpecificSocketBoneOffset, StaticBuffers->GetSpecificSocketBoneOffset());
 
