Fix failure in gpu sort test.
The Niagara gpu sort shader has two permutations and one does not have the SortKeyParams input.

diff --git a/Engine/Plugins/FX/Niagara/Source/NiagaraVertexFactories/Private/NiagaraSortingGPU.cpp b/Engine/Plugins/FX/Niagara/Source/NiagaraVertexFactories/Private/NiagaraSortingGPU.cpp
--- a/Engine/Plugins/FX/Niagara/Source/NiagaraVertexFactories/Private/NiagaraSortingGPU.cpp
+++ b/Engine/Plugins/FX/Niagara/Source/NiagaraVertexFactories/Private/NiagaraSortingGPU.cpp
@@ -131,7 +131,11 @@ void FNiagaraSortKeyGenCS::SetParameters(FRHICommandList& RHICmdList, const FNia
 	const FUintVector4 SortParamsValue(EmitterKey, OutputOffset, (uint8)SortInfo.SortMode, SortInfo.SortAttributeOffset);
 	RHICmdList.SetShaderParameter(ComputeShaderRHI, SortParams.GetBufferIndex(), SortParams.GetBaseIndex(), SortParams.GetNumBytes(), &SortParamsValue);
 
-	RHICmdList.SetShaderParameter(ComputeShaderRHI, SortKeyParams.GetBufferIndex(), SortKeyParams.GetBaseIndex(), SortKeyParams.GetNumBytes(), &SortKeyParamsValue);
+	// SortKeyParams only exists in the permutation with SORT_MAX_PRECISION set.
+	if(SortKeyParams.IsBound())
+	{
+		RHICmdList.SetShaderParameter(ComputeShaderRHI, SortKeyParams.GetBufferIndex(), SortKeyParams.GetBaseIndex(), SortKeyParams.GetNumBytes(), &SortKeyParamsValue);
+	}
 	RHICmdList.SetShaderParameter(ComputeShaderRHI, CameraPosition.GetBufferIndex(), CameraPosition.GetBaseIndex(), CameraPosition.GetNumBytes(), &SortInfo.ViewOrigin);
 	RHICmdList.SetShaderParameter(ComputeShaderRHI, CameraDirection.GetBufferIndex(), CameraDirection.GetBaseIndex(), CameraDirection.GetNumBytes(), &SortInfo.ViewDirection);
 }
