* Fixed a bug where static switches with compiled constants were not correctly compiled across complex graphs
* Adjusted automation test asset

diff --git a/Engine/Plugins/FX/Niagara/Source/NiagaraEditor/Private/NiagaraCompiler.cpp b/Engine/Plugins/FX/Niagara/Source/NiagaraEditor/Private/NiagaraCompiler.cpp
--- a/Engine/Plugins/FX/Niagara/Source/NiagaraEditor/Private/NiagaraCompiler.cpp
+++ b/Engine/Plugins/FX/Niagara/Source/NiagaraEditor/Private/NiagaraCompiler.cpp
@@ -52,7 +52,7 @@ static FAutoConsoleVariableRef CVarForceNiagaraTranslatorSingleThreaded(
 
 static FCriticalSection TranslationCritSec;
 
-void FNiagaraCompileRequestData::VisitReferencedGraphs(UNiagaraGraph* InSrcGraph, UNiagaraGraph* InDupeGraph, ENiagaraScriptUsage InUsage)
+void FNiagaraCompileRequestData::VisitReferencedGraphs(UNiagaraGraph* InSrcGraph, UNiagaraGraph* InDupeGraph, ENiagaraScriptUsage InUsage, FCompileConstantResolver ConstantResolver)
 {
 	if (!InDupeGraph && !InSrcGraph)
 	{
@@ -78,10 +78,10 @@ void FNiagaraCompileRequestData::VisitReferencedGraphs(UNiagaraGraph* InSrcGraph
 	FNiagaraEditorUtilities::ResolveNumerics(InDupeGraph, bStandaloneScript, ChangedFromNumericVars);
 	ClonedGraphs.AddUnique(InDupeGraph);
 
-	VisitReferencedGraphsRecursive(InDupeGraph);
+	VisitReferencedGraphsRecursive(InDupeGraph, ConstantResolver);
 }
 
-void FNiagaraCompileRequestData::VisitReferencedGraphsRecursive(UNiagaraGraph* InGraph)
+void FNiagaraCompileRequestData::VisitReferencedGraphsRecursive(UNiagaraGraph* InGraph, const FCompileConstantResolver& ConstantResolver)
 {
 	if (!InGraph)
 	{
@@ -169,7 +169,7 @@ void FNiagaraCompileRequestData::VisitReferencedGraphsRecursive(UNiagaraGraph* I
 								DupeScript = DuplicateObject<UNiagaraScript>(FunctionScript, InNode, FunctionScript->GetFName());
 								ProcessedGraph = Cast<UNiagaraScriptSource>(DupeScript->GetSource())->NodeGraph;
 								FEdGraphUtilities::MergeChildrenGraphsIn(ProcessedGraph, ProcessedGraph, /*bRequireSchemaMatch=*/ true);
-								FNiagaraEditorUtilities::PreprocessFunctionGraph(Schema, ProcessedGraph, CallInputs, CallOutputs, ScriptUsage);
+								FNiagaraEditorUtilities::PreprocessFunctionGraph(Schema, ProcessedGraph, CallInputs, CallOutputs, ScriptUsage, ConstantResolver);
 								FunctionCallNode->FunctionScript = DupeScript;
 
 								/*UE_LOG(LogNiagaraEditor, Log, TEXT("Duplicating script %s OriginalGraph %p NewGraph %p"), *FunctionScript->GetFName().ToString(),
@@ -187,7 +187,7 @@ void FNiagaraCompileRequestData::VisitReferencedGraphsRecursive(UNiagaraGraph* I
 							TArray<FunctionData> MadeArray;
 							MadeArray.Add(Data);
 							PreprocessedFunctions.Add(FunctionGraph, MadeArray);
-							VisitReferencedGraphsRecursive(ProcessedGraph);		
+							VisitReferencedGraphsRecursive(ProcessedGraph, ConstantResolver);		
 							ClonedGraphs.AddUnique(ProcessedGraph);
 
 						}
@@ -196,7 +196,7 @@ void FNiagaraCompileRequestData::VisitReferencedGraphsRecursive(UNiagaraGraph* I
 							UNiagaraScript* DupeScript = DuplicateObject<UNiagaraScript>(FunctionScript, InNode, FunctionScript->GetFName());
 							ProcessedGraph = Cast<UNiagaraScriptSource>(DupeScript->GetSource())->NodeGraph;
 							FEdGraphUtilities::MergeChildrenGraphsIn(ProcessedGraph, ProcessedGraph, /*bRequireSchemaMatch=*/ true);
-							FNiagaraEditorUtilities::PreprocessFunctionGraph(Schema, ProcessedGraph, CallInputs, CallOutputs, ScriptUsage);
+							FNiagaraEditorUtilities::PreprocessFunctionGraph(Schema, ProcessedGraph, CallInputs, CallOutputs, ScriptUsage, ConstantResolver);
 							FunctionCallNode->FunctionScript = DupeScript;
 
 							/*UE_LOG(LogNiagaraEditor, Log, TEXT("Duplicating script %s OriginalGraph %p NewGraph %p"), *FunctionScript->GetFName().ToString(),
@@ -215,7 +215,7 @@ void FNiagaraCompileRequestData::VisitReferencedGraphsRecursive(UNiagaraGraph* I
 							Data.bHasNumericInputs = bHasNumericInputs;
 
 							FoundArray->Add(Data);
-							VisitReferencedGraphsRecursive(ProcessedGraph);
+							VisitReferencedGraphsRecursive(ProcessedGraph, ConstantResolver);
 						}
 						else if (bFromDifferentPackage)
 						{
@@ -331,18 +331,18 @@ bool FNiagaraCompileRequestData::GatherPreCompiledVariables(const FString& InNam
 	return true;
 }
 
-void FNiagaraCompileRequestData::DeepCopyGraphs(UNiagaraScriptSource* ScriptSource, ENiagaraScriptUsage InUsage)
+void FNiagaraCompileRequestData::DeepCopyGraphs(UNiagaraScriptSource* ScriptSource, ENiagaraScriptUsage InUsage, FCompileConstantResolver ConstantResolver)
 {
 	// Clone the source graph so we can modify it as needed; merging in the child graphs
 	//NodeGraphDeepCopy = CastChecked<UNiagaraGraph>(FEdGraphUtilities::CloneGraph(ScriptSource->NodeGraph, GetTransientPackage(), 0));
 	Source = DuplicateObject<UNiagaraScriptSource>(ScriptSource, GetTransientPackage());
 	NodeGraphDeepCopy = Source->NodeGraph;
 	FEdGraphUtilities::MergeChildrenGraphsIn(NodeGraphDeepCopy, NodeGraphDeepCopy, /*bRequireSchemaMatch=*/ true);
-	VisitReferencedGraphs(ScriptSource->NodeGraph, NodeGraphDeepCopy, InUsage /*InEmitter ? ENiagaraScriptUsage::EmitterSpawnScript : ENiagaraScriptUsage::SystemSpawnScript*/);
+	VisitReferencedGraphs(ScriptSource->NodeGraph, NodeGraphDeepCopy, InUsage, ConstantResolver);
 }
 
 
-void FNiagaraCompileRequestData::FinishPrecompile(UNiagaraScriptSource* ScriptSource, const TArray<FNiagaraVariable>& EncounterableVariables, ENiagaraScriptUsage InUsage, FCompileConstantResolver* ConstantResolver)
+void FNiagaraCompileRequestData::FinishPrecompile(UNiagaraScriptSource* ScriptSource, const TArray<FNiagaraVariable>& EncounterableVariables, ENiagaraScriptUsage InUsage, FCompileConstantResolver ConstantResolver)
 {
 	{
 		ENiagaraScriptCompileStatusEnum = StaticEnum<ENiagaraScriptCompileStatus>();
@@ -358,10 +358,7 @@ void FNiagaraCompileRequestData::FinishPrecompile(UNiagaraScriptSource* ScriptSo
 		{
 			// Map all for this output node
 			FNiagaraParameterMapHistoryBuilder Builder;
-			if (ConstantResolver)
-			{
-				Builder.ConstantResolver = *ConstantResolver;
-			}
+			Builder.ConstantResolver = ConstantResolver;
 			Builder.RegisterEncounterableVariables(EncounterableVariables);
 
 			FString TranslationName = TEXT("Emitter");
@@ -442,15 +439,16 @@ TSharedPtr<FNiagaraCompileRequestDataBase, ESPMode::ThreadSafe> FNiagaraEditorMo
 
 	TSharedPtr<FNiagaraCompileRequestData, ESPMode::ThreadSafe> BasePtr = MakeShared<FNiagaraCompileRequestData, ESPMode::ThreadSafe>();
 	TArray<TSharedPtr<FNiagaraCompileRequestData, ESPMode::ThreadSafe>> DependentRequests;
+	FCompileConstantResolver EmptyResolver;
 
 	BasePtr->SourceName = InObj->GetName();
 
 	if (Script)
 	{
 		UNiagaraScriptSource* Source = Cast<UNiagaraScriptSource>(Script->GetSource());
-		BasePtr->DeepCopyGraphs(Source, Script->GetUsage());
+		BasePtr->DeepCopyGraphs(Source, Script->GetUsage(), EmptyResolver);
 		const TArray<FNiagaraVariable> EncounterableVariables;
-		BasePtr->FinishPrecompile(Source, EncounterableVariables, Script->GetUsage(), nullptr);
+		BasePtr->FinishPrecompile(Source, EncounterableVariables, Script->GetUsage(), EmptyResolver);
 	}
 	else if (System)
 	{
@@ -470,8 +468,9 @@ TSharedPtr<FNiagaraCompileRequestDataBase, ESPMode::ThreadSafe> FNiagaraEditorMo
 		for (int32 i = 0; i < System->GetEmitterHandles().Num(); i++)
 		{
 			const FNiagaraEmitterHandle& Handle = System->GetEmitterHandle(i);
+			FCompileConstantResolver ConstantResolver(Handle.GetInstance());
 			TSharedPtr<FNiagaraCompileRequestData, ESPMode::ThreadSafe> EmitterPtr = MakeShared<FNiagaraCompileRequestData, ESPMode::ThreadSafe>();
-			EmitterPtr->DeepCopyGraphs(Cast<UNiagaraScriptSource>(Handle.GetInstance()->GraphSource), ENiagaraScriptUsage::EmitterSpawnScript);
+			EmitterPtr->DeepCopyGraphs(Cast<UNiagaraScriptSource>(Handle.GetInstance()->GraphSource), ENiagaraScriptUsage::EmitterSpawnScript, ConstantResolver);
 			EmitterPtr->EmitterUniqueName = Handle.GetInstance()->GetUniqueEmitterName();
 			EmitterPtr->SourceName = BasePtr->SourceName;
 			BasePtr->EmitterData.Add(EmitterPtr);
@@ -479,8 +478,8 @@ TSharedPtr<FNiagaraCompileRequestDataBase, ESPMode::ThreadSafe> FNiagaraEditorMo
 
 		// Now deep copy the system graphs, skipping traversal into any emitter references.
 		UNiagaraScriptSource* Source = Cast<UNiagaraScriptSource>(System->GetSystemSpawnScript()->GetSource());
-		BasePtr->DeepCopyGraphs(Source, ENiagaraScriptUsage::SystemSpawnScript);
-		BasePtr->FinishPrecompile(Source, EncounterableVars, ENiagaraScriptUsage::SystemSpawnScript, nullptr);
+		BasePtr->DeepCopyGraphs(Source, ENiagaraScriptUsage::SystemSpawnScript, FCompileConstantResolver());
+		BasePtr->FinishPrecompile(Source, EncounterableVars, ENiagaraScriptUsage::SystemSpawnScript, EmptyResolver);
 
 		// Add the User and System variables that we did encounter to the list that emitters might also encounter.
 		BasePtr->GatherPreCompiledVariables(TEXT("User"), EncounterableVars);
@@ -491,7 +490,7 @@ TSharedPtr<FNiagaraCompileRequestDataBase, ESPMode::ThreadSafe> FNiagaraEditorMo
 		{
 			const FNiagaraEmitterHandle& Handle = System->GetEmitterHandle(i);
 			FCompileConstantResolver ConstantResolver(Handle.GetInstance());
-			BasePtr->EmitterData[i]->FinishPrecompile(Cast<UNiagaraScriptSource>(Handle.GetInstance()->GraphSource), EncounterableVars, ENiagaraScriptUsage::EmitterSpawnScript, &ConstantResolver);
+			BasePtr->EmitterData[i]->FinishPrecompile(Cast<UNiagaraScriptSource>(Handle.GetInstance()->GraphSource), EncounterableVars, ENiagaraScriptUsage::EmitterSpawnScript, ConstantResolver);
 			BasePtr->MergeInEmitterPrecompiledData(BasePtr->EmitterData[i].Get());
 		}
 
diff --git a/Engine/Plugins/FX/Niagara/Source/NiagaraEditor/Private/NiagaraEditorUtilities.cpp b/Engine/Plugins/FX/Niagara/Source/NiagaraEditor/Private/NiagaraEditorUtilities.cpp
--- a/Engine/Plugins/FX/Niagara/Source/NiagaraEditor/Private/NiagaraEditorUtilities.cpp
+++ b/Engine/Plugins/FX/Niagara/Source/NiagaraEditor/Private/NiagaraEditorUtilities.cpp
@@ -623,7 +623,7 @@ void FNiagaraEditorUtilities::FixUpNumericPins(const UEdGraphSchema_Niagara* Sch
 	TraverseGraphFromOutputDepthFirst(Schema, Node, FixUpVisitor);
 }
 
-void FNiagaraEditorUtilities::SetStaticSwitchConstants(UNiagaraGraph* Graph, const TArray<UEdGraphPin*>& CallInputs, FCompileConstantResolver* ConstantResolver)
+void FNiagaraEditorUtilities::SetStaticSwitchConstants(UNiagaraGraph* Graph, const TArray<UEdGraphPin*>& CallInputs, const FCompileConstantResolver& ConstantResolver)
 {
 	for (UEdGraphNode* Node : Graph->Nodes)
 	{
@@ -631,13 +631,13 @@ void FNiagaraEditorUtilities::SetStaticSwitchConstants(UNiagaraGraph* Graph, con
 		UNiagaraNodeStaticSwitch* SwitchNode = Cast<UNiagaraNodeStaticSwitch>(Node);
 		if (SwitchNode)
 		{
-			SwitchNode->ClearSwitchValue();
-			if (SwitchNode->IsSetByCompiler() && ConstantResolver)
+			if (SwitchNode->IsSetByCompiler())
 			{
-				SwitchNode->SetSwitchValue(*ConstantResolver);
+				SwitchNode->SetSwitchValue(ConstantResolver);
 			}
 			else
 			{
+				SwitchNode->ClearSwitchValue();
 				for (UEdGraphPin* InputPin : CallInputs)
 				{
 					if (InputPin->GetFName().IsEqual(SwitchNode->InputParameterName))
@@ -810,7 +810,7 @@ void FNiagaraEditorUtilities::ResolveNumerics(UNiagaraGraph* SourceGraph, bool b
 	}
 }
 
-void FNiagaraEditorUtilities::PreprocessFunctionGraph(const UEdGraphSchema_Niagara* Schema, UNiagaraGraph* Graph, const TArray<UEdGraphPin*>& CallInputs, const TArray<UEdGraphPin*>& CallOutputs, ENiagaraScriptUsage ScriptUsage)
+void FNiagaraEditorUtilities::PreprocessFunctionGraph(const UEdGraphSchema_Niagara* Schema, UNiagaraGraph* Graph, const TArray<UEdGraphPin*>& CallInputs, const TArray<UEdGraphPin*>& CallOutputs, ENiagaraScriptUsage ScriptUsage, const FCompileConstantResolver& ConstantResolver)
 {
 	// Change any numeric inputs or outputs to match the types from the call node.
 	TArray<UNiagaraNodeInput*> InputNodes;
@@ -860,9 +860,9 @@ void FNiagaraEditorUtilities::PreprocessFunctionGraph(const UEdGraphSchema_Niaga
 			}
 		}
 	}
-
+	
 	FNiagaraEditorUtilities::FixUpNumericPins(Schema, OutputNode);
-	FNiagaraEditorUtilities::SetStaticSwitchConstants(Graph, CallInputs);
+	FNiagaraEditorUtilities::SetStaticSwitchConstants(Graph, CallInputs, ConstantResolver);
 }
 
 void FNiagaraEditorUtilities::GetFilteredScriptAssets(FGetFilteredScriptAssetsOptions InFilter, TArray<FAssetData>& OutFilteredScriptAssets)
diff --git a/Engine/Plugins/FX/Niagara/Source/NiagaraEditor/Private/NiagaraHlslTranslator.h b/Engine/Plugins/FX/Niagara/Source/NiagaraEditor/Private/NiagaraHlslTranslator.h
--- a/Engine/Plugins/FX/Niagara/Source/NiagaraEditor/Private/NiagaraHlslTranslator.h
+++ b/Engine/Plugins/FX/Niagara/Source/NiagaraEditor/Private/NiagaraHlslTranslator.h
@@ -94,9 +94,9 @@ public:
 	class UNiagaraGraph* GetPrecomputedNodeGraph() { return NodeGraphDeepCopy; }
 	const class UNiagaraGraph* GetPrecomputedNodeGraph() const { return NodeGraphDeepCopy; }
 	const FString& GetUniqueEmitterName() const { return EmitterUniqueName; }
-	void VisitReferencedGraphs(UNiagaraGraph* InSrcGraph, UNiagaraGraph* InDupeGraph, ENiagaraScriptUsage InUsage);
-	void DeepCopyGraphs(UNiagaraScriptSource* ScriptSource, ENiagaraScriptUsage InUsage);
-	void FinishPrecompile(UNiagaraScriptSource* ScriptSource, const TArray<FNiagaraVariable>& EncounterableVariables, ENiagaraScriptUsage InUsage, FCompileConstantResolver* ConstantResolver);
+	void VisitReferencedGraphs(UNiagaraGraph* InSrcGraph, UNiagaraGraph* InDupeGraph, ENiagaraScriptUsage InUsage, FCompileConstantResolver ConstantResolver);
+	void DeepCopyGraphs(UNiagaraScriptSource* ScriptSource, ENiagaraScriptUsage InUsage, FCompileConstantResolver ConstantResolver);
+	void FinishPrecompile(UNiagaraScriptSource* ScriptSource, const TArray<FNiagaraVariable>& EncounterableVariables, ENiagaraScriptUsage InUsage, FCompileConstantResolver ConstantResolver);
 	virtual int32 GetDependentRequestCount() const override {
 		return EmitterData.Num();
 	};
@@ -132,7 +132,7 @@ public:
 	TMap<const UNiagaraGraph*, TArray<FunctionData>> PreprocessedFunctions;
 	TArray<UNiagaraGraph*> ClonedGraphs;
 protected:
-	void VisitReferencedGraphsRecursive(UNiagaraGraph* InGraph);
+	void VisitReferencedGraphsRecursive(UNiagaraGraph* InGraph, const FCompileConstantResolver& ConstantResolver);
 };
 
 
diff --git a/Engine/Plugins/FX/Niagara/Source/NiagaraEditor/Private/NiagaraNodeFunctionCall.cpp b/Engine/Plugins/FX/Niagara/Source/NiagaraEditor/Private/NiagaraNodeFunctionCall.cpp
--- a/Engine/Plugins/FX/Niagara/Source/NiagaraEditor/Private/NiagaraNodeFunctionCall.cpp
+++ b/Engine/Plugins/FX/Niagara/Source/NiagaraEditor/Private/NiagaraNodeFunctionCall.cpp
@@ -374,7 +374,8 @@ void UNiagaraNodeFunctionCall::Compile(class FHlslNiagaraTranslator* Translator,
 
 		TArray<UEdGraphPin*> CallerInputPins;
 		GetInputPins(CallerInputPins);
-		FNiagaraEditorUtilities::SetStaticSwitchConstants(GetCalledGraph(), CallerInputPins);
+		FCompileConstantResolver ConstantResolver(Translator);
+		FNiagaraEditorUtilities::SetStaticSwitchConstants(GetCalledGraph(), CallerInputPins, ConstantResolver);
 
 		UNiagaraScriptSource* Source = CastChecked<UNiagaraScriptSource>(FunctionScript->GetSource());
 		UNiagaraGraph* FunctionGraph = Source->NodeGraph;
@@ -676,7 +677,7 @@ void UNiagaraNodeFunctionCall::BuildParameterMapHistory(FNiagaraParameterMapHist
 
 		TArray<UEdGraphPin*> InputPins;
 		GetInputPins(InputPins);
-		FNiagaraEditorUtilities::SetStaticSwitchConstants(FunctionGraph, InputPins, &OutHistory.ConstantResolver);
+		FNiagaraEditorUtilities::SetStaticSwitchConstants(FunctionGraph, InputPins, OutHistory.ConstantResolver);
 
 		int32 ParamMapIdx = INDEX_NONE;
 		uint32 NodeIdx = INDEX_NONE;
diff --git a/Engine/Plugins/FX/Niagara/Source/NiagaraEditor/Private/NiagaraParameterMapHistory.cpp b/Engine/Plugins/FX/Niagara/Source/NiagaraEditor/Private/NiagaraParameterMapHistory.cpp
--- a/Engine/Plugins/FX/Niagara/Source/NiagaraEditor/Private/NiagaraParameterMapHistory.cpp
+++ b/Engine/Plugins/FX/Niagara/Source/NiagaraEditor/Private/NiagaraParameterMapHistory.cpp
@@ -1377,6 +1377,11 @@ void FNiagaraParameterMapHistoryBuilder::BuildCurrentAliases()
 
 bool FCompileConstantResolver::ResolveConstant(FNiagaraVariable& OutConstant) const
 {
+	if (Translator)
+	{
+		return Translator->GetLiteralConstantVariable(OutConstant);
+	}
+
 	if (Emitter && OutConstant == FNiagaraVariable(FNiagaraTypeDefinition::GetBoolDef(), TEXT("Emitter.Localspace")))
 	{
 		OutConstant.SetValue(Emitter->bLocalSpace ? FNiagaraBool(true) : FNiagaraBool(false));
diff --git a/Engine/Plugins/FX/Niagara/Source/NiagaraEditor/Public/NiagaraEditorUtilities.h b/Engine/Plugins/FX/Niagara/Source/NiagaraEditor/Public/NiagaraEditorUtilities.h
--- a/Engine/Plugins/FX/Niagara/Source/NiagaraEditor/Public/NiagaraEditorUtilities.h
+++ b/Engine/Plugins/FX/Niagara/Source/NiagaraEditor/Public/NiagaraEditorUtilities.h
@@ -139,11 +139,11 @@ namespace FNiagaraEditorUtilities
 
 	void FixUpNumericPins(const UEdGraphSchema_Niagara* Schema, UNiagaraNode* Node);
 
-	void SetStaticSwitchConstants(UNiagaraGraph* Graph, const TArray<UEdGraphPin*>& CallInputs, FCompileConstantResolver* ConstantResolver = nullptr);
+	void SetStaticSwitchConstants(UNiagaraGraph* Graph, const TArray<UEdGraphPin*>& CallInputs, const FCompileConstantResolver& ConstantResolver);
 
 	bool ResolveConstantValue(UEdGraphPin* Pin, int32& Value);
 
-	void PreprocessFunctionGraph(const UEdGraphSchema_Niagara* Schema, UNiagaraGraph* Graph, const TArray<UEdGraphPin*>& CallInputs, const TArray<UEdGraphPin*>& CallOutputs, ENiagaraScriptUsage ScriptUsage);
+	void PreprocessFunctionGraph(const UEdGraphSchema_Niagara* Schema, UNiagaraGraph* Graph, const TArray<UEdGraphPin*>& CallInputs, const TArray<UEdGraphPin*>& CallOutputs, ENiagaraScriptUsage ScriptUsage, const FCompileConstantResolver& ConstantResolver);
 
 	/** Options for the GetScriptsByFilter function. 
 	** @Param ScriptUsageToInclude Only return Scripts that have this usage
diff --git a/Engine/Plugins/FX/Niagara/Source/NiagaraEditor/Public/NiagaraParameterMapHistory.h b/Engine/Plugins/FX/Niagara/Source/NiagaraEditor/Public/NiagaraParameterMapHistory.h
--- a/Engine/Plugins/FX/Niagara/Source/NiagaraEditor/Public/NiagaraParameterMapHistory.h
+++ b/Engine/Plugins/FX/Niagara/Source/NiagaraEditor/Public/NiagaraParameterMapHistory.h
@@ -11,16 +11,19 @@ class UNiagaraNodeOutput;
 class UEdGraphPin;
 class UEdGraphNode;
 class UNiagaraParameterCollection;
+class FHlslNiagaraTranslator;
 
 class FCompileConstantResolver
 {
 public:
-	FCompileConstantResolver() : Emitter(nullptr) {}
-	FCompileConstantResolver(UNiagaraEmitter* Emitter) : Emitter(Emitter) {}
+	FCompileConstantResolver() : Emitter(nullptr), Translator(nullptr) {}
+	FCompileConstantResolver(UNiagaraEmitter* Emitter) : Emitter(Emitter), Translator(nullptr) {}
+	FCompileConstantResolver(FHlslNiagaraTranslator* Translator) : Emitter(nullptr), Translator(Translator) {}
 
 	bool ResolveConstant(FNiagaraVariable& OutConstant) const;
 private:
 	UNiagaraEmitter* Emitter;
+	FHlslNiagaraTranslator* Translator;
 };
 
 /** Traverses a Niagara node graph to identify the variables that have been written and read from a parameter map. 
