Fix LogSlateStyle error from incorrect search box style get.

diff --git a/Engine/Plugins/FX/Niagara/Source/NiagaraEditor/Private/Widgets/SNiagaraScriptGraph.cpp b/Engine/Plugins/FX/Niagara/Source/NiagaraEditor/Private/Widgets/SNiagaraScriptGraph.cpp
--- a/Engine/Plugins/FX/Niagara/Source/NiagaraEditor/Private/Widgets/SNiagaraScriptGraph.cpp
+++ b/Engine/Plugins/FX/Niagara/Source/NiagaraEditor/Private/Widgets/SNiagaraScriptGraph.cpp
@@ -53,7 +53,7 @@ TSharedRef<SGraphEditor> SNiagaraScriptGraph::ConstructGraphEditor()
 	FGraphAppearanceInfo AppearanceInfo;
 	AppearanceInfo.CornerText = LOCTEXT("AppearanceCornerText", "NIAGARA");
 
-	const FEditableTextBoxStyle& Style = FCoreStyle::Get().GetWidgetStyle<FEditableTextBoxStyle>("SearchBox");
+	const FSearchBoxStyle& Style = FCoreStyle::Get().GetWidgetStyle<FSearchBoxStyle>("SearchBox");
 
 	TSharedRef<SWidget> TitleBarWidget =
 		SNew(SBorder)
@@ -112,9 +112,9 @@ TSharedRef<SGraphEditor> SNiagaraScriptGraph::ConstructGraphEditor()
 					SNew(SBorder)
 					.HAlign(HAlign_Center)
 					.VAlign(VAlign_Fill)
-					.BorderImage(&Style.BackgroundImageHovered)
-					.BorderBackgroundColor(Style.BackgroundColor)
-					.ForegroundColor(Style.ForegroundColor)
+					.BorderImage(&Style.TextBoxStyle.BackgroundImageHovered)
+					.BorderBackgroundColor(Style.TextBoxStyle.BackgroundColor)
+					.ForegroundColor(Style.TextBoxStyle.ForegroundColor)
 					.Padding(0)
 					[
 						SNew(SButton)
