Fixed check in Niagara gpu particles when switching preview feature level.

diff --git a/Engine/Plugins/FX/Niagara/Source/Niagara/Private/NiagaraGPUInstanceCountManager.cpp b/Engine/Plugins/FX/Niagara/Source/Niagara/Private/NiagaraGPUInstanceCountManager.cpp
--- a/Engine/Plugins/FX/Niagara/Source/Niagara/Private/NiagaraGPUInstanceCountManager.cpp
+++ b/Engine/Plugins/FX/Niagara/Source/Niagara/Private/NiagaraGPUInstanceCountManager.cpp
@@ -168,36 +168,38 @@ void FNiagaraGPUInstanceCountManager::UpdateDrawIndirectBuffer(FRHICommandList&
 {
 	if (DrawIndirectArgGenTasks.Num() || InstanceCountClearTasks.Num())
 	{
-		RHICmdList.TransitionResource(EResourceTransitionAccess::ERWBarrier, EResourceTransitionPipeline::EComputeToCompute, CountBuffer.UAV);
-
-		FReadBuffer TaskInfosBuffer;
+		if (FeatureLevel == ERHIFeatureLevel::SM5 /*|| FeatureLevel == ERHIFeatureLevel::ES3_1*/) // Must match with RHISupportsComputeShaders()
 		{
-			// All draw indirect args task are run first because of the binding between the task ID and arg write offset.
-			const uint32 ArgGenSize = DrawIndirectArgGenTasks.Num() * sizeof(FArgGenTaskInfo);
-			const uint32 InstanceCountClearSize = InstanceCountClearTasks.Num() * sizeof(uint32);
-			const uint32 TaskBufferSize = ArgGenSize + InstanceCountClearSize;
-			TaskInfosBuffer.Initialize(sizeof(uint32), TaskBufferSize / sizeof(uint32), EPixelFormat::PF_R32_UINT, BUF_Volatile);
-			uint8* TaskBufferData = (uint8*)RHILockVertexBuffer(TaskInfosBuffer.Buffer, 0, TaskBufferSize, RLM_WriteOnly);
-			FMemory::Memcpy(TaskBufferData, DrawIndirectArgGenTasks.GetData(), ArgGenSize);
-			FMemory::Memcpy(TaskBufferData + ArgGenSize, InstanceCountClearTasks.GetData(), InstanceCountClearSize);
-			RHIUnlockVertexBuffer(TaskInfosBuffer.Buffer);
+			RHICmdList.TransitionResource(EResourceTransitionAccess::ERWBarrier, EResourceTransitionPipeline::EComputeToCompute, CountBuffer.UAV);
+
+			FReadBuffer TaskInfosBuffer;
+			{
+				// All draw indirect args task are run first because of the binding between the task ID and arg write offset.
+				const uint32 ArgGenSize = DrawIndirectArgGenTasks.Num() * sizeof(FArgGenTaskInfo);
+				const uint32 InstanceCountClearSize = InstanceCountClearTasks.Num() * sizeof(uint32);
+				const uint32 TaskBufferSize = ArgGenSize + InstanceCountClearSize;
+				TaskInfosBuffer.Initialize(sizeof(uint32), TaskBufferSize / sizeof(uint32), EPixelFormat::PF_R32_UINT, BUF_Volatile);
+				uint8* TaskBufferData = (uint8*)RHILockVertexBuffer(TaskInfosBuffer.Buffer, 0, TaskBufferSize, RLM_WriteOnly);
+				FMemory::Memcpy(TaskBufferData, DrawIndirectArgGenTasks.GetData(), ArgGenSize);
+				FMemory::Memcpy(TaskBufferData + ArgGenSize, InstanceCountClearTasks.GetData(), InstanceCountClearSize);
+				RHIUnlockVertexBuffer(TaskInfosBuffer.Buffer);
+			}
+
+			RHICmdList.TransitionResource(EResourceTransitionAccess::EWritable, EResourceTransitionPipeline::EGfxToCompute, DrawIndirectBuffer.UAV);
+
+			FNiagaraDrawIndirectArgsGenCS::FPermutationDomain PermutationVector;
+			TShaderMapRef<FNiagaraDrawIndirectArgsGenCS> DrawIndirectArgsGenCS(GetGlobalShaderMap(FeatureLevel), PermutationVector);
+			RHICmdList.SetComputeShader(DrawIndirectArgsGenCS->GetComputeShader());
+			DrawIndirectArgsGenCS->SetOutput(RHICmdList, DrawIndirectBuffer.UAV, CountBuffer.UAV);
+			DrawIndirectArgsGenCS->SetParameters(RHICmdList, TaskInfosBuffer.SRV, DrawIndirectArgGenTasks.Num(), InstanceCountClearTasks.Num());
+			DispatchComputeShader(RHICmdList, *DrawIndirectArgsGenCS, FMath::DivideAndRoundUp(DrawIndirectArgGenTasks.Num() + InstanceCountClearTasks.Num(), NIAGARA_DRAW_INDIRECT_ARGS_GEN_THREAD_COUNT), 1, 1);
+			DrawIndirectArgsGenCS->UnbindBuffers(RHICmdList);
+
+			// Sync after clear.
+			RHICmdList.TransitionResource(EResourceTransitionAccess::ERWBarrier, EResourceTransitionPipeline::EComputeToCompute, CountBuffer.UAV);
+			// Transition draw indirect to readable for gfx draw indirect.
+			RHICmdList.TransitionResource(EResourceTransitionAccess::EReadable, EResourceTransitionPipeline::EComputeToGfx, DrawIndirectBuffer.UAV);
 		}
-
-		RHICmdList.TransitionResource(EResourceTransitionAccess::EWritable, EResourceTransitionPipeline::EGfxToCompute, DrawIndirectBuffer.UAV);
-
-		FNiagaraDrawIndirectArgsGenCS::FPermutationDomain PermutationVector;
-		TShaderMapRef<FNiagaraDrawIndirectArgsGenCS> DrawIndirectArgsGenCS(GetGlobalShaderMap(FeatureLevel), PermutationVector);
-		RHICmdList.SetComputeShader(DrawIndirectArgsGenCS->GetComputeShader());
-		DrawIndirectArgsGenCS->SetOutput(RHICmdList, DrawIndirectBuffer.UAV, CountBuffer.UAV);
-		DrawIndirectArgsGenCS->SetParameters(RHICmdList, TaskInfosBuffer.SRV, DrawIndirectArgGenTasks.Num(), InstanceCountClearTasks.Num());
-		DispatchComputeShader(RHICmdList, *DrawIndirectArgsGenCS, FMath::DivideAndRoundUp(DrawIndirectArgGenTasks.Num() + InstanceCountClearTasks.Num(), NIAGARA_DRAW_INDIRECT_ARGS_GEN_THREAD_COUNT), 1, 1);
-		DrawIndirectArgsGenCS->UnbindBuffers(RHICmdList);
-
-		// Sync after clear.
-		RHICmdList.TransitionResource(EResourceTransitionAccess::ERWBarrier, EResourceTransitionPipeline::EComputeToCompute, CountBuffer.UAV);
-		// Transition draw indirect to readable for gfx draw indirect.
-		RHICmdList.TransitionResource(EResourceTransitionAccess::EReadable, EResourceTransitionPipeline::EComputeToGfx, DrawIndirectBuffer.UAV);
-
 		// Once cleared to 0, the count are reusable.
 		FreeEntries.Append(InstanceCountClearTasks);
 
diff --git a/Engine/Plugins/FX/Niagara/Source/Niagara/Private/NiagaraScriptExecutionContext.cpp b/Engine/Plugins/FX/Niagara/Source/Niagara/Private/NiagaraScriptExecutionContext.cpp
--- a/Engine/Plugins/FX/Niagara/Source/Niagara/Private/NiagaraScriptExecutionContext.cpp
+++ b/Engine/Plugins/FX/Niagara/Source/Niagara/Private/NiagaraScriptExecutionContext.cpp
@@ -411,8 +411,9 @@ FNiagaraComputeExecutionContext::FNiagaraComputeExecutionContext()
 
 FNiagaraComputeExecutionContext::~FNiagaraComputeExecutionContext()
 {
-	checkf(IsInRenderingThread(), TEXT("Can only delete the gpu readback from the render thread"));
-	check(EmitterInstanceReadback.GPUCountOffset == INDEX_NONE);
+	// EmitterInstanceReadback.GPUCountOffset should be INDEX_NONE at this point to ensure the index is reused.
+	// When the batcher is being destroyed though, we don't free the index, but this would not be leaking.
+	// check(EmitterInstanceReadback.GPUCountOffset == INDEX_NONE);
 
 #if WITH_EDITORONLY_DATA
 	if (GPUDebugDataReadbackFloat)
