Add adjustable tessellation parameters

diff --git a/Engine/Plugins/FX/Niagara/Source/Niagara/Private/NiagaraRendererRibbons.cpp b/Engine/Plugins/FX/Niagara/Source/Niagara/Private/NiagaraRendererRibbons.cpp
--- a/Engine/Plugins/FX/Niagara/Source/Niagara/Private/NiagaraRendererRibbons.cpp
+++ b/Engine/Plugins/FX/Niagara/Source/Niagara/Private/NiagaraRendererRibbons.cpp
@@ -23,8 +23,7 @@ float GNiagaraRibbonTessellationAngle = 15.f * (2.f * PI) / 360.f; // Every 15 d
 static FAutoConsoleVariableRef CVarNiagaraRibbonTessellationAngle(
 	TEXT("Niagara.Ribbon.Tessellation.MinAngle"),
 	GNiagaraRibbonTessellationAngle,
-	TEXT("Ribbon segment angle to tesselate in radian. \n")
-	TEXT("Set 0 to use fixed tesselation. (default=15 degrees)"),
+	TEXT("Ribbon segment angle to tesselate in radian. (default=15 degrees)"),
 	ECVF_Scalability
 );
 
@@ -148,6 +147,12 @@ FNiagaraRendererRibbons::FNiagaraRendererRibbons(ERHIFeatureLevel::Type FeatureL
 	, UV1TilingDistance(0.0f)
 	, UV1Scale(FVector2D(1.0f, 1.0f))
 	, UV1AgeOffsetMode(ENiagaraRibbonAgeOffsetMode::Scale)
+	, TessellationMode(ENiagaraRibbonTessellationMode::Automatic)
+	, CustomCurveTension(0.f)
+	, CustomTessellationFactor(16)
+	, bCustomUseConstantFactor(false)
+	, CustomTessellationMinAngle(15.f * PI / 180.f)
+	, bCustomUseScreenSpace(true)
 	, PositionDataOffset(INDEX_NONE)
 	, VelocityDataOffset(INDEX_NONE)
 	, WidthDataOffset(INDEX_NONE)
@@ -172,6 +177,13 @@ FNiagaraRendererRibbons::FNiagaraRendererRibbons(ERHIFeatureLevel::Type FeatureL
 	UV1Offset = Properties->UV1Offset;
 	UV1AgeOffsetMode = Properties->UV1AgeOffsetMode;
 	DrawDirection = Properties->DrawDirection;
+	TessellationMode = Properties->TessellationMode;
+	CustomCurveTension = FMath::Clamp<float>(Properties->CurveTension, 0.f, 0.9999f);
+	CustomTessellationFactor = Properties->TessellationFactor;
+	bCustomUseConstantFactor = Properties->bUseConstantFactor;
+	CustomTessellationMinAngle = Properties->TessellationAngle > 0.f && Properties->TessellationAngle < 1.f ? 1.f : Properties->TessellationAngle;
+	CustomTessellationMinAngle *= PI / 180.f;
+	bCustomUseScreenSpace = Properties->bScreenSpaceTessellation;
 
 	// required attributes
 	int32 IntDummy;
@@ -333,6 +345,29 @@ void FNiagaraRendererRibbons::GetDynamicMeshElements(const TArray<const FSceneVi
 		WorldSpacePrimitiveUniformBuffer.InitResource();
 	}
 
+	// Modify tessellation parameters based on tessellation mode
+	// Set auto defaults here:
+	int32 TessellationFactor = GNiagaraRibbonMaxTessellation;
+	bool bUseConstantFactor = false;
+	float TessellationMinAngle = GNiagaraRibbonTessellationAngle;
+	float ScreenPercentage = GNiagaraRibbonTessellationScreenPercentage;
+	switch (TessellationMode)
+	{
+	case ENiagaraRibbonTessellationMode::Automatic:
+		break;
+	case ENiagaraRibbonTessellationMode::Custom:
+		TessellationFactor = CustomTessellationFactor;
+		bUseConstantFactor = bCustomUseConstantFactor;
+		TessellationMinAngle = CustomTessellationMinAngle;
+		ScreenPercentage = bCustomUseScreenSpace ? GNiagaraRibbonTessellationScreenPercentage : 0.f;
+		break;
+	case ENiagaraRibbonTessellationMode::Disabled:
+		TessellationFactor = 1;
+		break;
+	default:
+		break;
+	}
+
 	// Compute the per-view uniform buffers.
 	for (int32 ViewIndex = 0; ViewIndex < Views.Num(); ViewIndex++)
 	{
@@ -345,12 +380,22 @@ void FNiagaraRendererRibbons::GetDynamicMeshElements(const TArray<const FSceneVi
 
 			int32 SegmentTessellation = 1;
 			int32 NumSegments = DynamicDataRibbon->SegmentData.Num();
-			if (GNiagaraRibbonMaxTessellation > 1 && TessellationCurvature > SMALL_NUMBER && ViewFamily.GetFeatureLevel() == ERHIFeatureLevel::SM5)
+			if (TessellationFactor > 1 && TessellationCurvature > SMALL_NUMBER && ViewFamily.GetFeatureLevel() == ERHIFeatureLevel::SM5)
 			{
-				const float MinTesselation = FMath::Max<float>(1.f, FMath::Max(TessellationTwistAngle, TessellationAngle) / FMath::Max<float>(SMALL_NUMBER, GNiagaraRibbonTessellationAngle));
-
+				const float MinTesselation = [&]
+				{
+					if (TessellationMinAngle == 0.f) { return 0.f; }
+					if (!bUseConstantFactor)
+					{
+						return FMath::Max<float>(1.f, FMath::Max(TessellationTwistAngle, TessellationAngle) / FMath::Max<float>(SMALL_NUMBER, TessellationMinAngle));
+					}
+					else
+					{
+						return static_cast<float>(TessellationFactor);
+					}
+				}(); 
 				const float ViewDistance = SceneProxy->GetBounds().ComputeSquaredDistanceFromBoxToPoint(ViewOriginForDistanceCulling);
-				const float MaxDisplacementError = FMath::Max(GNiagaraRibbonTessellationMinDisplacementError, GNiagaraRibbonTessellationScreenPercentage * FMath::Sqrt(ViewDistance) / View->LODDistanceFactor);
+				const float MaxDisplacementError = FMath::Max(GNiagaraRibbonTessellationMinDisplacementError, ScreenPercentage * FMath::Sqrt(ViewDistance) / View->LODDistanceFactor);
 				float Tess = TessellationAngle / FMath::Max(SMALL_NUMBER, AcosFast(TessellationCurvature / (TessellationCurvature + MaxDisplacementError)));
 				// FMath::RoundUpToPowerOfTwo ? This could avoid vertices moving around as tesselation increases
 
@@ -359,7 +404,7 @@ void FNiagaraRendererRibbons::GetDynamicMeshElements(const TArray<const FSceneVi
 					const float TwistTess  = TessellationTwistAngle / FMath::Max(SMALL_NUMBER, AcosFast(TessellationTwistCurvature / (TessellationTwistCurvature + MaxDisplacementError)));
 					Tess = FMath::Max(TwistTess, Tess);
 				}
-				SegmentTessellation = FMath::Clamp<int32>(FMath::RoundToInt(Tess), FMath::RoundToInt(MinTesselation), GNiagaraRibbonMaxTessellation);
+				SegmentTessellation = FMath::Clamp<int32>(FMath::RoundToInt(Tess), FMath::RoundToInt(MinTesselation), TessellationFactor);
 				NumSegments *= SegmentTessellation;
 			}
 
@@ -707,7 +752,7 @@ FNiagaraDynamicDataBase* FNiagaraRendererRibbons::GenerateDynamicData(const FNia
 			{
 				// Normalize CurrToNextVec
 				CurrToNextVec *= 1.f / FMath::Max(GNiagaraRibbonMinSegmentLength, CurrToNextSize);
-				const FVector Tangent = (LastToCurrVec + CurrToNextVec).GetSafeNormal();
+				const FVector Tangent = (1.f - CustomCurveTension) * (LastToCurrVec + CurrToNextVec).GetSafeNormal();
 
 				// Update the distance for CurrentIndex.
 				TotalDistance += LastToCurrSize;
diff --git a/Engine/Plugins/FX/Niagara/Source/Niagara/Private/NiagaraRibbonRendererProperties.cpp b/Engine/Plugins/FX/Niagara/Source/Niagara/Private/NiagaraRibbonRendererProperties.cpp
--- a/Engine/Plugins/FX/Niagara/Source/Niagara/Private/NiagaraRibbonRendererProperties.cpp
+++ b/Engine/Plugins/FX/Niagara/Source/Niagara/Private/NiagaraRibbonRendererProperties.cpp
@@ -14,6 +14,12 @@ UNiagaraRibbonRendererProperties::UNiagaraRibbonRendererProperties()
 	, UV1TilingDistance(0.0f)
 	, UV1Scale(FVector2D(1.0f, 1.0f))
 	, UV1AgeOffsetMode(ENiagaraRibbonAgeOffsetMode::Scale)
+	, CurveTension(0.f)
+	, TessellationMode(ENiagaraRibbonTessellationMode::Automatic)
+	, TessellationFactor(16)
+	, bUseConstantFactor(false)
+	, TessellationAngle(15)
+	, bScreenSpaceTessellation(true)
 {
 }
 
@@ -75,6 +81,14 @@ void UNiagaraRibbonRendererProperties::InitBindings()
 void UNiagaraRibbonRendererProperties::PostEditChangeProperty(struct FPropertyChangedEvent& PropertyChangedEvent)
 {
 	Super::PostEditChangeProperty(PropertyChangedEvent);
+	FName PropertyName = PropertyChangedEvent.GetPropertyName();
+	if (PropertyName == GET_MEMBER_NAME_CHECKED(UNiagaraRibbonRendererProperties, TessellationAngle))
+	{
+		if (TessellationAngle > 0.f && TessellationAngle < 1.f)
+		{
+			TessellationAngle = 1.f;
+		}
+	}
 }
 
 const TArray<FNiagaraVariable>& UNiagaraRibbonRendererProperties::GetRequiredAttributes()
@@ -128,4 +142,27 @@ void UNiagaraRibbonRendererProperties::FixMaterial(UMaterial* InMaterial)
 	InMaterial->ForceRecompileForRendering();
 }
 
+bool UNiagaraRibbonRendererProperties::CanEditChange(const UProperty* InProperty) const
+{
+
+	if (InProperty->HasMetaData(TEXT("Category")) && InProperty->GetMetaData(TEXT("Category")).Contains("Tessellation"))
+	{
+		FName PropertyName = InProperty->GetFName();
+		if (PropertyName == GET_MEMBER_NAME_CHECKED(UNiagaraRibbonRendererProperties, CurveTension))
+		{
+			return TessellationMode != ENiagaraRibbonTessellationMode::Disabled;
+		}
+		if (PropertyName == GET_MEMBER_NAME_CHECKED(UNiagaraRibbonRendererProperties, TessellationFactor))
+		{
+			return TessellationMode == ENiagaraRibbonTessellationMode::Custom;
+		}
+		if (PropertyName == GET_MEMBER_NAME_CHECKED(UNiagaraRibbonRendererProperties, TessellationMode))
+		{
+			return Super::CanEditChange(InProperty);
+		}
+		return TessellationMode == ENiagaraRibbonTessellationMode::Custom;
+	}
+	return Super::CanEditChange(InProperty);
+}
+
 #endif // WITH_EDITORONLY_DATA
diff --git a/Engine/Plugins/FX/Niagara/Source/Niagara/Public/NiagaraRendererRibbons.h b/Engine/Plugins/FX/Niagara/Source/Niagara/Public/NiagaraRendererRibbons.h
--- a/Engine/Plugins/FX/Niagara/Source/Niagara/Public/NiagaraRendererRibbons.h
+++ b/Engine/Plugins/FX/Niagara/Source/Niagara/Public/NiagaraRendererRibbons.h
@@ -48,6 +48,12 @@ private:
 	FVector2D UV1Offset;
 	ENiagaraRibbonAgeOffsetMode UV1AgeOffsetMode;
 	ENiagaraRibbonDrawDirection DrawDirection;
+	ENiagaraRibbonTessellationMode TessellationMode;
+	float CustomCurveTension;
+	int32 CustomTessellationFactor;
+	bool bCustomUseConstantFactor;
+	float CustomTessellationMinAngle;
+	bool bCustomUseScreenSpace;
 
 	int32 PositionDataOffset;
 	int32 VelocityDataOffset;
diff --git a/Engine/Plugins/FX/Niagara/Source/Niagara/Public/NiagaraRibbonRendererProperties.h b/Engine/Plugins/FX/Niagara/Source/Niagara/Public/NiagaraRibbonRendererProperties.h
--- a/Engine/Plugins/FX/Niagara/Source/Niagara/Public/NiagaraRibbonRendererProperties.h
+++ b/Engine/Plugins/FX/Niagara/Source/Niagara/Public/NiagaraRibbonRendererProperties.h
@@ -35,6 +35,17 @@ enum class ENiagaraRibbonDrawDirection : uint8
 	BackToFront
 };
 
+UENUM()
+enum class ENiagaraRibbonTessellationMode : uint8
+{
+	/** Default tessellation parameters. */
+	Automatic,
+	/** Custom tessellation parameters. */
+	Custom,
+	/** Disable tessellation entirely. */
+	Disabled
+};
+
 UCLASS(editinlinenew)
 class UNiagaraRibbonRendererProperties : public UNiagaraRendererProperties
 {
@@ -47,6 +58,7 @@ public:
 	virtual void PostInitProperties() override;
 #if WITH_EDITOR
 	virtual void PostEditChangeProperty(struct FPropertyChangedEvent& PropertyChangedEvent) override;
+	virtual bool CanEditChange(const UProperty* InProperty) const override;
 #endif
 	//UObject Interface END
 
@@ -100,6 +112,36 @@ public:
 	UPROPERTY(EditAnywhere, Category = "Ribbon Rendering")
 	ENiagaraRibbonDrawDirection DrawDirection;
 
+	/** Defines the curve tension, or how long the curve's tangents are.
+	  * Ranges from 0 to 1. The higher the value, the sharper the curve becomes.
+	  */
+	UPROPERTY(EditAnywhere, Category = "Tessellation", meta = (ClampMin = "0", ClampMax = "1"))
+	float CurveTension;
+
+	/** Defines the tessellation mode allowing custom tessellation parameters or disabling tessellation entirely. */
+	UPROPERTY(EditAnywhere, Category = "Tessellation", meta = (DisplayName = "Mode"))
+	ENiagaraRibbonTessellationMode TessellationMode;
+
+	/** Custom tessellation factor.
+	  * Ranges from 1 to 16. Greater values increase amount of tessellation.
+	  */
+	UPROPERTY(EditAnywhere, Category = "Tessellation", meta = (DisplayName = "Max Tessellation Factor", ClampMin = "1", ClampMax = "16"))
+	int32 TessellationFactor;
+
+	/** If checked, use the above constant factor. Otherwise, adaptively select the tessellation factor based on the below parameters. */
+	UPROPERTY(EditAnywhere, Category = "Tessellation")
+	bool bUseConstantFactor;
+
+	/** Defines the angle in degrees at which tessellation occurs.
+	  * Ranges from 1 to 180. Smaller values increase amount of tessellation.
+	  * If set to 0, use the maximum tessellation set above.
+	  */
+	UPROPERTY(EditAnywhere, Category = "Tessellation", meta = (EditCondition = "!bUseConstantFactor", ClampMin = "0", ClampMax = "180", UIMin = "1", UIMax = "180"))
+	float TessellationAngle;
+
+	/** If checked, use the ribbon's screen space percentage to adaptively adjust the tessellation factor. */
+	UPROPERTY(EditAnywhere, Category = "Tessellation", meta = (DisplayName = "Screen Space", EditCondition = "!bUseConstantFactor"))
+	bool bScreenSpaceTessellation;
 
 	/** Which attribute should we use for position when generating ribbons?*/
 	UPROPERTY(EditAnywhere, AdvancedDisplay, Category = "Bindings")
