- Niagara calcuate instance count correctly when using dispatch overlapping if a tick contains bNeedsReset

#rnx

diff --git a/Engine/Plugins/FX/Niagara/Source/Niagara/Private/NiagaraEmitterInstanceBatcher.cpp b/Engine/Plugins/FX/Niagara/Source/Niagara/Private/NiagaraEmitterInstanceBatcher.cpp
--- a/Engine/Plugins/FX/Niagara/Source/Niagara/Private/NiagaraEmitterInstanceBatcher.cpp
+++ b/Engine/Plugins/FX/Niagara/Source/Niagara/Private/NiagaraEmitterInstanceBatcher.cpp
@@ -161,6 +161,7 @@ void NiagaraEmitterInstanceBatcher::ResizeBuffersAndGatherResources(FOverlappabl
 	{
 		const uint32 DispatchCount = Tick->Count;
 		const bool bIsFinalTick = Tick->bIsFinalTick;
+		const bool bNeedsReset = Tick->bNeedsReset;
 
 		FNiagaraComputeInstanceData* Instances = Tick->GetInstanceData();
 		for (uint32 Index = 0; Index < DispatchCount; Index++)
@@ -187,7 +188,7 @@ void NiagaraEmitterInstanceBatcher::ResizeBuffersAndGatherResources(FOverlappabl
 			FNiagaraDataBuffer& CurrentData = *Instance.CurrentData;
 			FNiagaraDataBuffer& DestinationData = *Instance.DestinationData;
 
-			const uint32 PrevNumInstances = CurrentData.GetNumInstances();
+			const uint32 PrevNumInstances = bNeedsReset ? 0 : CurrentData.GetNumInstances();
 			const uint32 NewNumInstances = Instance.SpawnRateInstances + Instance.EventSpawnTotal + PrevNumInstances;
 
 			//We must assume all particles survive when allocating here. 
