- Temporary fix for hitting a check when using skel mesh sampling on the CPU
- We can't use bUseTriangleSampling / bUseVertexSampling / bUseSkeletonSampling in the editor as they are bound late

#rnx

diff --git a/Engine/Plugins/FX/Niagara/Source/Niagara/Private/NiagaraDataInterfaceSkeletalMesh.cpp b/Engine/Plugins/FX/Niagara/Source/Niagara/Private/NiagaraDataInterfaceSkeletalMesh.cpp
--- a/Engine/Plugins/FX/Niagara/Source/Niagara/Private/NiagaraDataInterfaceSkeletalMesh.cpp
+++ b/Engine/Plugins/FX/Niagara/Source/Niagara/Private/NiagaraDataInterfaceSkeletalMesh.cpp
@@ -1144,7 +1144,8 @@ bool FNDISkeletalMesh_InstanceData::Init(UNiagaraDataInterfaceSkeletalMesh* Inte
 	FSkeletalMeshLODRenderData& LODData = GetLODRenderDataAndSkinWeights(SkinWeightBuffer);
 
 	//Check for the validity of the Mesh's cpu data.
-	if ( Interface->bUseTriangleSampling || Interface->bUseVertexSampling )
+	//-TODO: These values are currently created too late in the editor so we can't use them to validate the DI correctly
+	//if ( Interface->bUseTriangleSampling || Interface->bUseVertexSampling )
 	{
 		bool LODDataNumVerticesCorrect = LODData.GetNumVertices() > 0;
 		bool LODDataPositonNumVerticesCorrect = LODData.StaticVertexBuffers.PositionVertexBuffer.GetNumVertices() > 0;
