Niagara - Fix issues with emitter inheritance.
+ Conditionally post load the Parent and ParentAtLastMerge before trying to merge changes.
+ Remove parent information from emitters who's outers are other emitters since they are being used as ParentAtLastMerge and they should never merge changes.
+ Add a utility function to copy an emitter without merging or parent information on the duplicate to prevent uneccessary merging during post load, and the merge process itself.
+ Add more information to the merge results so that in the case where there were no differences we don't update the emitter.

diff --git a/Engine/Plugins/FX/Niagara/Source/Niagara/Classes/NiagaraEmitter.h b/Engine/Plugins/FX/Niagara/Source/Niagara/Classes/NiagaraEmitter.h
--- a/Engine/Plugins/FX/Niagara/Source/Niagara/Classes/NiagaraEmitter.h
+++ b/Engine/Plugins/FX/Niagara/Source/Niagara/Classes/NiagaraEmitter.h
@@ -322,6 +322,9 @@ public:
 
 	/** Whether or not this emitter uses the supplied emitter */
 	bool UsesEmitter(const UNiagaraEmitter& InEmitter) const;
+
+	/** Duplicates this emitter, but prevents the duplicate from merging in changes from the parent emitter.  The resulting duplicate will have no parent information. */
+	NIAGARA_API UNiagaraEmitter* DuplicateWithoutMerging(UObject* InOuter);
 #endif
 
 	/** Is this emitter allowed to be enabled by the current system detail level. */
@@ -374,7 +377,7 @@ private:
 
 	void SyncEmitterAlias(const FString& InOldName, const FString& InNewName);
 
-	void UpdateChangeId();
+	void UpdateChangeId(const FString& Reason);
 
 	void ScriptRapidIterationParameterChanged();
 
diff --git a/Engine/Plugins/FX/Niagara/Source/Niagara/Private/NiagaraEmitter.cpp b/Engine/Plugins/FX/Niagara/Source/Niagara/Private/NiagaraEmitter.cpp
--- a/Engine/Plugins/FX/Niagara/Source/Niagara/Private/NiagaraEmitter.cpp
+++ b/Engine/Plugins/FX/Niagara/Source/Niagara/Private/NiagaraEmitter.cpp
@@ -43,6 +43,14 @@ static FAutoConsoleVariableRef CVarForceNiagaraCompileToFail(
 	ECVF_Default
 );
 
+static int32 GbEnableEmitterChangeIdMergeLogging = 0;
+static FAutoConsoleVariableRef CVarEnableEmitterChangeIdMergeLogging(
+	TEXT("fx.EnableEmitterMergeChangeIdLogging"),
+	GbEnableEmitterChangeIdMergeLogging,
+	TEXT("If > 0 verbose change id information will be logged to help with debuggin merge issues. \n"),
+	ECVF_Default
+);
+
 void FNiagaraEmitterScriptProperties::InitDataSetAccess()
 {
 	EventReceivers.Empty();
@@ -166,14 +174,22 @@ bool UNiagaraEmitter::IsSynchronizedWithParent() const
 
 INiagaraMergeManager::FMergeEmitterResults UNiagaraEmitter::MergeChangesFromParent()
 {
-	UE_LOG(LogNiagara, Log, TEXT("Emitter %s-%s is merging changes from parent %s because its Change ID was updated."), *GetPathName(), *GetName(),
+	UE_LOG(LogNiagara, Log, TEXT("Emitter %s is merging changes from parent %s because its Change ID was updated."), *GetPathName(),
 		Parent != nullptr ? *Parent->GetPathName() : TEXT("(null)"));
 
+	if (GbEnableEmitterChangeIdMergeLogging)
+	{
+		UE_LOG(LogNiagara, Log, TEXT("\nEmitter %s Id=%s \nParentAtLastMerge %s id=%s \nParent %s Id=%s."), 
+			*GetPathName(), *ChangeId.ToString(),
+			ParentAtLastMerge != nullptr ? *ParentAtLastMerge->GetPathName() : TEXT("(null)"), ParentAtLastMerge != nullptr ? *ParentAtLastMerge->GetChangeId().ToString() : TEXT("(null)"),
+			Parent != nullptr ? *Parent->GetPathName() : TEXT("(null)"), Parent != nullptr ? *Parent->GetChangeId().ToString() : TEXT("(null)"));
+	}
+
 	if (Parent == nullptr)
 	{
 		// If we don't have a copy of the parent emitter, this emitter can't safely be merged.
 		INiagaraMergeManager::FMergeEmitterResults MergeResults;
-		MergeResults.bSucceeded = false;
+		MergeResults.MergeResult = INiagaraMergeManager::EMergeEmitterResult::FailedToDiff;
 		MergeResults.bModifiedGraph = false;
 		MergeResults.ErrorMessages.Add(NSLOCTEXT("NiagaraEmitter", "NoParentErrorMessage", "This emitter has no 'Parent' so changes can't be merged in."));
 		return MergeResults;
@@ -184,7 +200,7 @@ INiagaraMergeManager::FMergeEmitterResults UNiagaraEmitter::MergeChangesFromPare
 		// If we don't have a copy of the last merged parent emitter, this emitter can't safely be
 		// merged.
 		INiagaraMergeManager::FMergeEmitterResults MergeResults;
-		MergeResults.bSucceeded = false;
+		MergeResults.MergeResult = INiagaraMergeManager::EMergeEmitterResult::FailedToDiff;
 		MergeResults.bModifiedGraph = false;
 		MergeResults.ErrorMessages.Add(NSLOCTEXT("NiagaraEmitter", "NoLastMergedParentErrorMessage", "This emitter has no 'ParentAtLastMerge' so changes can't be merged in."));
 		return MergeResults;
@@ -193,12 +209,15 @@ INiagaraMergeManager::FMergeEmitterResults UNiagaraEmitter::MergeChangesFromPare
 	INiagaraModule& NiagaraModule = FModuleManager::Get().GetModuleChecked<INiagaraModule>("Niagara");
 	const INiagaraMergeManager& MergeManager = NiagaraModule.GetMergeManager();
 	INiagaraMergeManager::FMergeEmitterResults MergeResults = MergeManager.MergeEmitter(*Parent, *ParentAtLastMerge, *this);
-	if (MergeResults.bSucceeded)
+	if (MergeResults.MergeResult == INiagaraMergeManager::EMergeEmitterResult::SucceededDifferencesApplied || MergeResults.MergeResult == INiagaraMergeManager::EMergeEmitterResult::SucceededNoDifferences)
 	{
-		UpdateFromMergedCopy(MergeManager, MergeResults.MergedInstance);
+		if (MergeResults.MergeResult == INiagaraMergeManager::EMergeEmitterResult::SucceededDifferencesApplied)
+		{
+			UpdateFromMergedCopy(MergeManager, MergeResults.MergedInstance);
+		}
 
 		// Update the last merged source and clear it's stand alone and public flags since it's not an asset.
-		ParentAtLastMerge = CastChecked<UNiagaraEmitter>(StaticDuplicateObject(Parent, this));
+		ParentAtLastMerge = Parent->DuplicateWithoutMerging(this);
 		ParentAtLastMerge->ClearFlags(RF_Standalone | RF_Public);
 	}
 	else
@@ -215,6 +234,17 @@ bool UNiagaraEmitter::UsesEmitter(const UNiagaraEmitter& InEmitter) const
 	return Parent == &InEmitter || (Parent != nullptr && Parent->UsesEmitter(InEmitter));
 }
 
+UNiagaraEmitter* UNiagaraEmitter::DuplicateWithoutMerging(UObject* InOuter)
+{
+	UNiagaraEmitter* Duplicate;
+	{
+		TGuardValue<UNiagaraEmitter*> ParentGuard(Parent, nullptr);
+		TGuardValue<UNiagaraEmitter*> ParentAtLastMergeGuard(ParentAtLastMerge, nullptr);
+		Duplicate = Cast<UNiagaraEmitter>(StaticDuplicateObject(this, InOuter));
+	}
+	return Duplicate;
+}
+
 #endif
 
 void UNiagaraEmitter::Serialize(FArchive& Ar)
@@ -279,6 +309,13 @@ void UNiagaraEmitter::PostLoad()
 	}
 
 #if WITH_EDITORONLY_DATA
+	if (GetOuter()->IsA<UNiagaraEmitter>())
+	{
+		// If this emitter is owned by another emitter, remove it's inheritance information so that it doesn't try to merge changes.
+		Parent = nullptr;
+		ParentAtLastMerge = nullptr;
+	}
+
 	GraphSource->ConditionalPostLoad();
 	GraphSource->PostLoadFromEmitter(*this);
 #endif
@@ -294,6 +331,15 @@ void UNiagaraEmitter::PostLoad()
 
 #if WITH_EDITORONLY_DATA
 
+	// Handle emitter inheritance.
+	if (Parent != nullptr)
+	{
+		Parent->ConditionalPostLoad();
+	}
+	if (ParentAtLastMerge != nullptr)
+	{
+		ParentAtLastMerge->ConditionalPostLoad();
+	}
 	if (IsSynchronizedWithParent() == false)
 	{
 		MergeChangesFromParent();
@@ -301,6 +347,7 @@ void UNiagaraEmitter::PostLoad()
 
 	// Reset scripts if recompile is forced.
 	bool bGenerateNewChangeId = false;
+	FString GenerateNewChangeIdReason;
 	if (GetForceCompileOnLoad())
 	{
 		// If we are a standalone emitter, then we invalidate id's, which should cause systems dependent on us to regenerate.
@@ -309,6 +356,7 @@ void UNiagaraEmitter::PostLoad()
 		{
 			GraphSource->InvalidateCachedCompileIds();
 			bGenerateNewChangeId = true;
+			GenerateNewChangeIdReason = TEXT("PostLoad - Force compile on load");
 			if (GEnableVerboseNiagaraChangeIdLogging)
 			{
 				UE_LOG(LogNiagara, Log, TEXT("InvalidateCachedCompileIds for %s because GbForceNiagaraCompileOnLoad = %d"), *GetPathName(), GbForceNiagaraCompileOnLoad);
@@ -320,6 +368,7 @@ void UNiagaraEmitter::PostLoad()
 	{
 		// If the change id is already invalid we need to generate a new one, and can skip checking the owned scripts.
 		bGenerateNewChangeId = true;
+		GenerateNewChangeIdReason = TEXT("PostLoad - Change id was invalid.");
 		if (GEnableVerboseNiagaraChangeIdLogging)
 		{
 			UE_LOG(LogNiagara, Log, TEXT("Change ID updated for emitter %s because the ID was invalid."), *GetPathName());
@@ -332,6 +381,7 @@ void UNiagaraEmitter::PostLoad()
 			if (Script->AreScriptAndSourceSynchronized() == false)
 			{
 				bGenerateNewChangeId = true;
+				GenerateNewChangeIdReason = TEXT("PostLoad - Script out of sync");
 				if (GEnableVerboseNiagaraChangeIdLogging)
 				{
 					UE_LOG(LogNiagara, Log, TEXT("Change ID updated for emitter %s because of a change to its script %s"), *GetPathName(), *Script->GetPathName());
@@ -342,7 +392,7 @@ void UNiagaraEmitter::PostLoad()
 
 	if (bGenerateNewChangeId)
 	{
-		UpdateChangeId();
+		UpdateChangeId(GenerateNewChangeIdReason);
 	}
 
 	GraphSource->OnChanged().AddUObject(this, &UNiagaraEmitter::GraphSourceChanged);
@@ -480,7 +530,7 @@ void UNiagaraEmitter::PostEditChangeProperty(struct FPropertyChangedEvent& Prope
 	}
 
 	ThumbnailImageOutOfDate = true;
-	UpdateChangeId();
+	UpdateChangeId(TEXT("PostEditChangeProperty"));
 	OnPropertiesChangedDelegate.Broadcast();
 }
 
@@ -905,7 +955,7 @@ void UNiagaraEmitter::UpdateFromMergedCopy(const INiagaraMergeManager& MergeMana
 	}
 
 	// Update the change id since we don't know what's changed.
-	ChangeId = FGuid::NewGuid();
+	UpdateChangeId(TEXT("Updated from merged copy"));
 }
 
 void UNiagaraEmitter::SyncEmitterAlias(const FString& InOldName, const FString& InNewName)
@@ -963,7 +1013,7 @@ void UNiagaraEmitter::AddRenderer(UNiagaraRendererProperties* Renderer)
 	RendererProperties.Add(Renderer);
 #if WITH_EDITOR
 	Renderer->OnChanged().AddUObject(this, &UNiagaraEmitter::RendererChanged);
-	UpdateChangeId();
+	UpdateChangeId(TEXT("Renderer added"));
 #endif
 }
 
@@ -973,7 +1023,7 @@ void UNiagaraEmitter::RemoveRenderer(UNiagaraRendererProperties* Renderer)
 	RendererProperties.Remove(Renderer);
 #if WITH_EDITOR
 	Renderer->OnChanged().RemoveAll(this);
-	UpdateChangeId();
+	UpdateChangeId(TEXT("Renderer removed"));
 #endif
 }
 
@@ -996,7 +1046,7 @@ void UNiagaraEmitter::AddEventHandler(FNiagaraEventScriptProperties EventHandler
 #if WITH_EDITOR
 	EventHandler.Script->RapidIterationParameters.AddOnChangedHandler(
 		FNiagaraParameterStore::FOnChanged::FDelegate::CreateUObject(this, &UNiagaraEmitter::ScriptRapidIterationParameterChanged));
-	UpdateChangeId();
+	UpdateChangeId(TEXT("Event handler added"));
 #endif
 }
 
@@ -1013,7 +1063,7 @@ void UNiagaraEmitter::RemoveEventHandlerByUsageId(FGuid EventHandlerUsageId)
 #endif
 	EventHandlerScriptProps.RemoveAll(FindEventHandlerById);
 #if WITH_EDITOR
-	UpdateChangeId();
+	UpdateChangeId(TEXT("Event handler removed"));
 #endif
 }
 
@@ -1035,27 +1085,33 @@ void UNiagaraEmitter::BeginDestroy()
 
 #if WITH_EDITORONLY_DATA
 
-void UNiagaraEmitter::UpdateChangeId()
+void UNiagaraEmitter::UpdateChangeId(const FString& Reason)
 {
 	// We don't mark the package dirty here because this can happen as a result of a compile and we don't want to dirty files
 	// due to compilation, in cases where the package should be marked dirty an previous modify would have already done this.
 	Modify(false);
+	FGuid OldId = ChangeId;
 	ChangeId = FGuid::NewGuid();
+	if (GbEnableEmitterChangeIdMergeLogging)
+	{
+		UE_LOG(LogNiagara, Log, TEXT("Emitter %s change id updated. Reason: %s OldId: %s NewId: %s"),
+			*GetPathName(), *Reason, *OldId.ToString(), *ChangeId.ToString());
+	}
 }
 
 void UNiagaraEmitter::ScriptRapidIterationParameterChanged()
 {
-	UpdateChangeId();
+	UpdateChangeId(TEXT("Script rapid iteration parameter changed."));
 }
 
 void UNiagaraEmitter::RendererChanged()
 {
-	UpdateChangeId();
+	UpdateChangeId(TEXT("Renderer changed."));
 }
 
 void UNiagaraEmitter::GraphSourceChanged()
 {
-	UpdateChangeId();
+	UpdateChangeId(TEXT("Graph source changed."));
 }
 #endif
 
diff --git a/Engine/Plugins/FX/Niagara/Source/Niagara/Public/INiagaraMergeManager.h b/Engine/Plugins/FX/Niagara/Source/Niagara/Public/INiagaraMergeManager.h
--- a/Engine/Plugins/FX/Niagara/Source/Niagara/Public/INiagaraMergeManager.h
+++ b/Engine/Plugins/FX/Niagara/Source/Niagara/Public/INiagaraMergeManager.h
@@ -5,15 +5,24 @@
 class INiagaraMergeManager
 {
 public:
+	enum class EMergeEmitterResult
+	{
+		SucceededNoDifferences,
+		SucceededDifferencesApplied,
+		FailedToDiff,
+		FailedToMerge,
+		None
+	};
+
 	struct FMergeEmitterResults
 	{
 		FMergeEmitterResults()
-			: bSucceeded(true)
+			: MergeResult(EMergeEmitterResult::None)
 			, bModifiedGraph(false)
 		{
 		}
 
-		bool bSucceeded;
+		EMergeEmitterResult MergeResult;
 		TArray<FText> ErrorMessages;
 		bool bModifiedGraph;
 		UNiagaraEmitter* MergedInstance;
diff --git a/Engine/Plugins/FX/Niagara/Source/NiagaraEditor/Private/NiagaraScriptMergeManager.cpp b/Engine/Plugins/FX/Niagara/Source/NiagaraEditor/Private/NiagaraScriptMergeManager.cpp
--- a/Engine/Plugins/FX/Niagara/Source/NiagaraEditor/Private/NiagaraScriptMergeManager.cpp
+++ b/Engine/Plugins/FX/Niagara/Source/NiagaraEditor/Private/NiagaraScriptMergeManager.cpp
@@ -814,9 +814,18 @@ INiagaraMergeManager::FMergeEmitterResults FNiagaraScriptMergeManager::MergeEmit
 
 	FNiagaraEmitterDiffResults DiffResults = DiffEmitters(ParentAtLastMerge, Instance);
 
-	if (DiffResults.IsValid())
+	if (DiffResults.IsValid() == false)
+	{
+		MergeResults.MergeResult = INiagaraMergeManager::EMergeEmitterResult::FailedToDiff;
+		MergeResults.ErrorMessages = DiffResults.GetErrorMessages();
+	}
+	else if (DiffResults.IsEmpty())
+	{
+		MergeResults.MergeResult = INiagaraMergeManager::EMergeEmitterResult::SucceededNoDifferences;
+	}
+	else
 	{
-		UNiagaraEmitter* MergedInstance = CastChecked<UNiagaraEmitter>(StaticDuplicateObject(&Parent, (UObject*)GetTransientPackage()));
+		UNiagaraEmitter* MergedInstance = Parent.DuplicateWithoutMerging((UObject*)GetTransientPackage());
 		TSharedRef<FNiagaraEmitterMergeAdapter> MergedInstanceAdapter = MakeShared<FNiagaraEmitterMergeAdapter>(*MergedInstance);
 
 		TMap<FGuid, FGuid> SourceChangeIds;
@@ -828,33 +837,53 @@ INiagaraMergeManager::FMergeEmitterResults FNiagaraScriptMergeManager::MergeEmit
 		FNiagaraEditorUtilities::GatherChangeIds(Instance, LastChangeIds, TEXT("Instance"));
 		DiffChangeIds(SourceChangeIds, PreviousSourceChangeIds, LastChangeIds, ChangeIdsThatNeedToBeReset);
 
+		MergeResults.MergeResult = INiagaraMergeManager::EMergeEmitterResult::SucceededDifferencesApplied;
+
 		FApplyDiffResults EmitterSpawnResults = ApplyScriptStackDiff(MergedInstanceAdapter->GetEmitterSpawnStack().ToSharedRef(), DiffResults.EmitterSpawnDiffResults);
-		MergeResults.bSucceeded &= EmitterSpawnResults.bSucceeded;
+		if (EmitterSpawnResults.bSucceeded == false)
+		{
+			MergeResults.MergeResult = INiagaraMergeManager::EMergeEmitterResult::FailedToMerge;
+		}
 		MergeResults.bModifiedGraph |= EmitterSpawnResults.bModifiedGraph;
 		MergeResults.ErrorMessages.Append(EmitterSpawnResults.ErrorMessages);
 
 		FApplyDiffResults EmitterUpdateResults = ApplyScriptStackDiff(MergedInstanceAdapter->GetEmitterUpdateStack().ToSharedRef(), DiffResults.EmitterUpdateDiffResults);
-		MergeResults.bSucceeded &= EmitterUpdateResults.bSucceeded;
+		if (EmitterUpdateResults.bSucceeded == false)
+		{
+			MergeResults.MergeResult = INiagaraMergeManager::EMergeEmitterResult::FailedToMerge;
+		}
 		MergeResults.bModifiedGraph |= EmitterUpdateResults.bModifiedGraph;
 		MergeResults.ErrorMessages.Append(EmitterUpdateResults.ErrorMessages);
 
 		FApplyDiffResults ParticleSpawnResults = ApplyScriptStackDiff(MergedInstanceAdapter->GetParticleSpawnStack().ToSharedRef(), DiffResults.ParticleSpawnDiffResults);
-		MergeResults.bSucceeded &= ParticleSpawnResults.bSucceeded;
+		if (ParticleSpawnResults.bSucceeded == false)
+		{
+			MergeResults.MergeResult = INiagaraMergeManager::EMergeEmitterResult::FailedToMerge;
+		}
 		MergeResults.bModifiedGraph |= ParticleSpawnResults.bModifiedGraph;
 		MergeResults.ErrorMessages.Append(ParticleSpawnResults.ErrorMessages);
 
 		FApplyDiffResults ParticleUpdateResults = ApplyScriptStackDiff(MergedInstanceAdapter->GetParticleUpdateStack().ToSharedRef(), DiffResults.ParticleUpdateDiffResults);
-		MergeResults.bSucceeded &= ParticleUpdateResults.bSucceeded;
+		if (ParticleUpdateResults.bSucceeded == false)
+		{
+			MergeResults.MergeResult = INiagaraMergeManager::EMergeEmitterResult::FailedToMerge;
+		}
 		MergeResults.bModifiedGraph |= ParticleUpdateResults.bModifiedGraph;
 		MergeResults.ErrorMessages.Append(ParticleUpdateResults.ErrorMessages);
 
 		FApplyDiffResults EventHandlerResults = ApplyEventHandlerDiff(MergedInstanceAdapter, DiffResults);
-		MergeResults.bSucceeded &= EventHandlerResults.bSucceeded;
+		if (EventHandlerResults.bSucceeded == false)
+		{
+			MergeResults.MergeResult = INiagaraMergeManager::EMergeEmitterResult::FailedToMerge;
+		}
 		MergeResults.bModifiedGraph |= EventHandlerResults.bModifiedGraph;
 		MergeResults.ErrorMessages.Append(EventHandlerResults.ErrorMessages);
 
 		FApplyDiffResults RendererResults = ApplyRendererDiff(*MergedInstance, DiffResults);
-		MergeResults.bSucceeded &= RendererResults.bSucceeded;
+		if (RendererResults.bSucceeded == false)
+		{
+			MergeResults.MergeResult = INiagaraMergeManager::EMergeEmitterResult::FailedToMerge;
+		}
 		MergeResults.bModifiedGraph |= RendererResults.bModifiedGraph;
 		MergeResults.ErrorMessages.Append(RendererResults.ErrorMessages);
 
@@ -884,7 +913,10 @@ INiagaraMergeManager::FMergeEmitterResults FNiagaraScriptMergeManager::MergeEmit
 #endif
 		
 		FApplyDiffResults ChangeIdResults = ResolveChangeIds(MergedInstanceAdapter, Instance, ChangeIdsThatNeedToBeReset);
-		MergeResults.bSucceeded &= ChangeIdResults.bSucceeded;
+		if (ChangeIdResults.bSucceeded == false)
+		{
+			MergeResults.MergeResult = INiagaraMergeManager::EMergeEmitterResult::FailedToMerge;
+		}
 		MergeResults.bModifiedGraph |= ChangeIdResults.bModifiedGraph;
 		MergeResults.ErrorMessages.Append(ChangeIdResults.ErrorMessages);
 		
@@ -925,7 +957,7 @@ INiagaraMergeManager::FMergeEmitterResults FNiagaraScriptMergeManager::MergeEmit
 			UpdateScript->RapidIterationParameters.DumpParameters();
 		}
 #endif
-		if (MergeResults.bSucceeded)
+		if (MergeResults.MergeResult == INiagaraMergeManager::EMergeEmitterResult::SucceededDifferencesApplied)
 		{
 			UNiagaraScriptSource* ScriptSource = Cast<UNiagaraScriptSource>(MergedInstance->GraphSource);
 			FNiagaraStackGraphUtilities::RelayoutGraph(*ScriptSource->NodeGraph);
@@ -935,11 +967,6 @@ INiagaraMergeManager::FMergeEmitterResults FNiagaraScriptMergeManager::MergeEmit
 		TMap<FGuid, FGuid> FinalChangeIds;
 		FNiagaraEditorUtilities::GatherChangeIds(*MergedInstance, FinalChangeIds, TEXT("Final"));
 	}
-	else
-	{
-		MergeResults.bSucceeded = false;
-		MergeResults.ErrorMessages = DiffResults.GetErrorMessages();
-	}
 	return MergeResults;
 }
 
