Small Niagara refactor to better reuse draw indirect args with hit proxies, ray tracing and reflection captures.

diff --git a/Engine/Plugins/FX/Niagara/Source/Niagara/Classes/NiagaraGPUInstanceCountManager.h b/Engine/Plugins/FX/Niagara/Source/Niagara/Classes/NiagaraGPUInstanceCountManager.h
--- a/Engine/Plugins/FX/Niagara/Source/Niagara/Classes/NiagaraGPUInstanceCountManager.h
+++ b/Engine/Plugins/FX/Niagara/Source/Niagara/Classes/NiagaraGPUInstanceCountManager.h
@@ -13,6 +13,11 @@ NiagaraGPUInstanceCountManager.h: GPU particle count handling
 
 class FRHIGPUMemoryReadback;
 
+FORCEINLINE uint32 GetTypeHash(const FNiagaraDrawIndirectArgsGenCS::FArgGenTaskInfo& Info)
+{
+	return (Info.InstanceCountBufferOffset << 8) | Info.NumIndicesPerInstance;
+}
+
 /** 
  * A manager that handles the buffer containing the GPU particle count. 
  * Also provides related functionalities like the generation of the draw indirect buffer.
@@ -53,18 +58,8 @@ public:
 	 * Called on the renderthread from FNiagaraRenderer::CreateRenderThreadResources()
 	 * or FNiagaraRenderer::ReleaseRenderThreadResources()
 	 */
-	void IncrementMaxDrawIndirectCount() 
-	{ 
-		checkSlow(IsInRenderingThread()); 
-		MaxDrawIndirectArgs += DrawIndirectArgsIncrement; 
-	}
-
-	void DecrementMaxDrawIndirectCount() 
-	{ 
-		checkSlow(IsInRenderingThread()); 
-		MaxDrawIndirectArgs -= DrawIndirectArgsIncrement; 
-
-	}
+	void IncrementMaxDrawIndirectCount() { checkSlow(IsInRenderingThread()); ++MaxDrawIndirectArgs; }
+	void DecrementMaxDrawIndirectCount() { checkSlow(IsInRenderingThread()); --MaxDrawIndirectArgs; }
 
 	// Resize instance count and draw indirect buffers to ensure it is big enough to hold all draw indirect args.
 	void ResizeBuffers(int32 ReservedInstanceCounts);
@@ -95,12 +90,12 @@ protected:
 
 	/** The list of all draw indirected tasks that are to be run in UpdateDrawIndirectBuffer() */
 	typedef FNiagaraDrawIndirectArgsGenCS::FArgGenTaskInfo FArgGenTaskInfo;
+
 	TArray<FArgGenTaskInfo> DrawIndirectArgGenTasks;
+	/** The map between each task FArgGenTaskInfo and entry offset from DrawIndirectArgGenTasks. Used to reuse entries. */
+	TMap<FArgGenTaskInfo, uint32> DrawIndirectArgMap;
 	/** The list of all instance count clear tasks that are to be run in UpdateDrawIndirectBuffer() */
 	TArray<uint32> InstanceCountClearTasks;
 	/** A buffer holding drawindirect data to render GPU emitter renderers. */
 	FRWBuffer DrawIndirectBuffer;
-
-	// A multiplier applied to each increment/decrement in IncrementMaxDrawIndirectCount()/DecrementMaxDrawIndirectCount()
-	int32 DrawIndirectArgsIncrement = 1;
 };
\ No newline at end of file
diff --git a/Engine/Plugins/FX/Niagara/Source/Niagara/Private/NiagaraGPUInstanceCountManager.cpp b/Engine/Plugins/FX/Niagara/Source/Niagara/Private/NiagaraGPUInstanceCountManager.cpp
--- a/Engine/Plugins/FX/Niagara/Source/Niagara/Private/NiagaraGPUInstanceCountManager.cpp
+++ b/Engine/Plugins/FX/Niagara/Source/Niagara/Private/NiagaraGPUInstanceCountManager.cpp
@@ -27,17 +27,7 @@ DECLARE_DWORD_COUNTER_STAT(TEXT("Max Num GPU Renderers"), STAT_NiagaraMaxNumGPUR
 //*****************************************************************************
 
 FNiagaraGPUInstanceCountManager::FNiagaraGPUInstanceCountManager() 
-{
-#if RHI_RAYTRACING
-	// GetDynamicRayTracingInstances()
-	++DrawIndirectArgsIncrement;
-#endif
-	if (GIsEditor)
-	{
-		// Hit proxies
-		++DrawIndirectArgsIncrement;
-	}
-}
+{}
 
 FNiagaraGPUInstanceCountManager::~FNiagaraGPUInstanceCountManager()
 {
@@ -145,7 +135,15 @@ void FNiagaraGPUInstanceCountManager::ResizeBuffers(int32 ReservedInstanceCounts
 uint32 FNiagaraGPUInstanceCountManager::AddDrawIndirect(uint32 InstanceCountBufferOffset, uint32 NumIndicesPerInstance)
 {
 	checkSlow(IsInRenderingThread());
-	if (DrawIndirectArgGenTasks.Num() < AllocatedDrawIndirectArgs)
+
+	const FArgGenTaskInfo Info(InstanceCountBufferOffset, NumIndicesPerInstance);
+
+	uint32& CachedOffset = DrawIndirectArgMap.FindOrAdd(Info, INDEX_NONE);
+	if (CachedOffset != INDEX_NONE)
+	{
+		return CachedOffset;
+	}
+	else if (DrawIndirectArgGenTasks.Num() < AllocatedDrawIndirectArgs)
 	{
 #if !UE_BUILD_SHIPPING
 		if (DrawIndirectArgGenTasks.Num() >= MaxDrawIndirectArgs)
@@ -153,9 +151,9 @@ uint32 FNiagaraGPUInstanceCountManager::AddDrawIndirect(uint32 InstanceCountBuff
 			UE_LOG(LogNiagara, Warning, TEXT("More draw indirect args then expected (%d / %d)"), DrawIndirectArgGenTasks.Num() + 1, MaxDrawIndirectArgs);
 		}
 #endif
-
-		DrawIndirectArgGenTasks.Add(FArgGenTaskInfo(InstanceCountBufferOffset, NumIndicesPerInstance));
-		return (DrawIndirectArgGenTasks.Num() - 1) * NIAGARA_DRAW_INDIRECT_ARGS_SIZE * sizeof(uint32);
+		DrawIndirectArgGenTasks.Add(Info);
+		CachedOffset = (DrawIndirectArgGenTasks.Num() - 1) * NIAGARA_DRAW_INDIRECT_ARGS_SIZE * sizeof(uint32);
+		return CachedOffset;
 	}
 	else
 	{
@@ -204,6 +202,7 @@ void FNiagaraGPUInstanceCountManager::UpdateDrawIndirectBuffer(FRHICommandList&
 		FreeEntries.Append(InstanceCountClearTasks);
 
 		DrawIndirectArgGenTasks.Empty();
+		DrawIndirectArgMap.Empty();
 		InstanceCountClearTasks.Empty();
 	}
 }
diff --git a/Engine/Plugins/FX/Niagara/Source/NiagaraVertexFactories/Public/NiagaraDrawIndirect.h b/Engine/Plugins/FX/Niagara/Source/NiagaraVertexFactories/Public/NiagaraDrawIndirect.h
--- a/Engine/Plugins/FX/Niagara/Source/NiagaraVertexFactories/Public/NiagaraDrawIndirect.h
+++ b/Engine/Plugins/FX/Niagara/Source/NiagaraVertexFactories/Public/NiagaraDrawIndirect.h
@@ -40,6 +40,11 @@ public:
 
 		uint32 InstanceCountBufferOffset;
 		uint32 NumIndicesPerInstance; // When -1 the counter needs to be reset to 0.
+
+		bool operator==(const FArgGenTaskInfo& Rhs) const
+		{
+			return InstanceCountBufferOffset == Rhs.InstanceCountBufferOffset && NumIndicesPerInstance == Rhs.NumIndicesPerInstance;
+		}
 	};
 
 	using FPermutationDomain = TShaderPermutationDomain<>;
