Add per particle ribbon side vector parameter

diff --git a/Engine/Plugins/FX/Niagara/Shaders/Private/NiagaraRibbonVertexFactory.ush b/Engine/Plugins/FX/Niagara/Shaders/Private/NiagaraRibbonVertexFactory.ush
--- a/Engine/Plugins/FX/Niagara/Shaders/Private/NiagaraRibbonVertexFactory.ush
+++ b/Engine/Plugins/FX/Niagara/Shaders/Private/NiagaraRibbonVertexFactory.ush
@@ -89,6 +89,7 @@ struct FVertexFactoryIntermediates
 	float3 Right;
 	float3 Up;
 	float3 Facing;
+	float3 SideVector;
 	float4 TexCoord;
 	
 #if USE_PARTICLE_POSITION
@@ -341,7 +342,6 @@ bool UVIndexUseZW(int UVIndex)
 	#endif
 }
 
-
 float2 GetUV(FVertexFactoryInterpolantsVSToPS Interpolants, int UVIndex)
 {
 	float4 UVVector = Interpolants.TexCoords[UVIndex / 2];
@@ -556,7 +556,7 @@ float3 SafeNormalize(float3 V)
 #endif
 
 void GetNiagaraParticleTangents(in FInterpVtxData InterpVtxData, in out FVertexFactoryIntermediates Intermediates)
-{
+{	
 	BRANCH
 	if(NiagaraRibbonVF.FacingDataOffset != -1)
 	{
@@ -567,14 +567,32 @@ void GetNiagaraParticleTangents(in FInterpVtxData InterpVtxData, in out FVertexF
 	else
 	{
 		Intermediates.Facing = normalize(ResolvedView.WorldCameraOrigin - Intermediates.Position);
-	}	
+	}
 
 	ISOLATE
 	{
 		Intermediates.Right = normalize(cross(Intermediates.Facing, Intermediates.Direction));
-		Intermediates.Up = cross(Intermediates.Right, Intermediates.Facing);
 	}
 
+	BRANCH
+	if (NiagaraRibbonVF.SideVectorDataOffset != -1)
+	{
+		const float3 Side0 = GetVec3(NiagaraRibbonVF.SideVectorDataOffset, InterpVtxData.ParticleId0);
+		const float3 Side1 = GetVec3(NiagaraRibbonVF.SideVectorDataOffset, InterpVtxData.ParticleId1);
+		Intermediates.SideVector = -normalize(TransformVector(lerp(Side0, Side1, InterpVtxData.Alpha)));
+		Intermediates.Right = normalize(Intermediates.SideVector - dot(Intermediates.SideVector, Intermediates.Direction * Intermediates.Direction));
+		Intermediates.Facing = normalize(cross(Intermediates.Direction, Intermediates.Right));
+	}
+	else
+	{
+		Intermediates.SideVector = Intermediates.Right;
+	}
+
+	ISOLATE
+	{
+		Intermediates.Up = cross(Intermediates.Right, Intermediates.Facing);
+	}
+	
 	BRANCH
 	if (NiagaraRibbonVF.TwistDataOffset != -1)
 	{
@@ -596,7 +614,7 @@ float4 CalcWorldPosition(FVertexFactoryInput Input, FVertexFactoryIntermediates
 	{
 		// expand and align
 		int Sgn = (((int)Input.InterpVtxId & 0x1) * 2) - 1;
-		float3 Pos = Intermediates.Position.xyz + Intermediates.Right * Sgn * (Intermediates.RibbonWidth.xxx / 2);
+		float3 Pos = Intermediates.Position.xyz + Intermediates.SideVector * Sgn * (Intermediates.RibbonWidth.xxx / 2);
 		// Intermediates are already in world space, only need to account for the translation around view origin.
 		return float4(Pos + ResolvedView.PreViewTranslation, 1);
 	}
diff --git a/Engine/Plugins/FX/Niagara/Source/Niagara/Classes/NiagaraConstants.h b/Engine/Plugins/FX/Niagara/Source/Niagara/Classes/NiagaraConstants.h
--- a/Engine/Plugins/FX/Niagara/Source/Niagara/Classes/NiagaraConstants.h
+++ b/Engine/Plugins/FX/Niagara/Source/Niagara/Classes/NiagaraConstants.h
@@ -99,6 +99,7 @@
 #define SYS_PARAM_PARTICLES_RIBBONWIDTH                  INiagaraModule::GetVar_Particles_RibbonWidth()
 #define SYS_PARAM_PARTICLES_RIBBONTWIST                  INiagaraModule::GetVar_Particles_RibbonTwist()
 #define SYS_PARAM_PARTICLES_RIBBONFACING                 INiagaraModule::GetVar_Particles_RibbonFacing()
+#define SYS_PARAM_PARTICLES_RIBBONSIDEVECTOR              INiagaraModule::GetVar_Particles_RibbonSideVector()
 #define SYS_PARAM_PARTICLES_RIBBONLINKORDER              INiagaraModule::GetVar_Particles_RibbonLinkOrder()
 
 #define SYS_PARAM_INSTANCE_ALIVE                         INiagaraModule::GetVar_DataInstance_Alive()
diff --git a/Engine/Plugins/FX/Niagara/Source/Niagara/Private/NiagaraConstants.cpp b/Engine/Plugins/FX/Niagara/Source/Niagara/Private/NiagaraConstants.cpp
--- a/Engine/Plugins/FX/Niagara/Source/Niagara/Private/NiagaraConstants.cpp
+++ b/Engine/Plugins/FX/Niagara/Source/Niagara/Private/NiagaraConstants.cpp
@@ -196,6 +196,7 @@ void FNiagaraConstants::Init()
 		Attributes.Add(SYS_PARAM_PARTICLES_RIBBONWIDTH);
 		Attributes.Add(SYS_PARAM_PARTICLES_RIBBONTWIST);
 		Attributes.Add(SYS_PARAM_PARTICLES_RIBBONFACING);
+		Attributes.Add(SYS_PARAM_PARTICLES_RIBBONSIDEVECTOR);
 		Attributes.Add(SYS_PARAM_PARTICLES_RIBBONLINKORDER);
 
 
@@ -230,6 +231,7 @@ void FNiagaraConstants::Init()
 		AttrDataSetKeyMap.Add(SYS_PARAM_PARTICLES_RIBBONWIDTH, GetAttributeAsDataSetKey(SYS_PARAM_PARTICLES_RIBBONWIDTH));
 		AttrDataSetKeyMap.Add(SYS_PARAM_PARTICLES_RIBBONTWIST, GetAttributeAsDataSetKey(SYS_PARAM_PARTICLES_RIBBONTWIST));
 		AttrDataSetKeyMap.Add(SYS_PARAM_PARTICLES_RIBBONFACING, GetAttributeAsDataSetKey(SYS_PARAM_PARTICLES_RIBBONFACING));
+		AttrDataSetKeyMap.Add(SYS_PARAM_PARTICLES_RIBBONSIDEVECTOR, GetAttributeAsDataSetKey(SYS_PARAM_PARTICLES_RIBBONSIDEVECTOR));
 		AttrDataSetKeyMap.Add(SYS_PARAM_PARTICLES_RIBBONLINKORDER, GetAttributeAsDataSetKey(SYS_PARAM_PARTICLES_RIBBONLINKORDER));
 	}
 
@@ -364,6 +366,11 @@ void FNiagaraConstants::Init()
 		Var.SetValue<FVector>(FVector(0.0f, 0.0f, 1.0f));
 		AttrDefaultsValueMap.Add(SYS_PARAM_PARTICLES_RIBBONFACING, Var);
 
+		AttrDefaultsStrMap.Add(SYS_PARAM_PARTICLES_RIBBONSIDEVECTOR, TEXT("0.0, 0.0, 1.0"));
+		Var = SYS_PARAM_PARTICLES_RIBBONSIDEVECTOR;
+		Var.SetValue<FVector>(FVector(0.0f, 0.0f, 1.0f));
+		AttrDefaultsValueMap.Add(SYS_PARAM_PARTICLES_RIBBONSIDEVECTOR, Var);
+
 		AttrDefaultsStrMap.Add(SYS_PARAM_PARTICLES_RIBBONLINKORDER, TEXT("0"));
 		Var = SYS_PARAM_PARTICLES_RIBBONLINKORDER;
 		Var.SetValue<float>(0.0f);
@@ -400,6 +407,7 @@ void FNiagaraConstants::Init()
 		AttrDescStrMap.Add(SYS_PARAM_PARTICLES_RIBBONWIDTH, LOCTEXT("RibbonWidthDesc", "Sets the ribbon width for a particle, in UE4 units."));
 		AttrDescStrMap.Add(SYS_PARAM_PARTICLES_RIBBONTWIST, LOCTEXT("RibbonTwistDesc", "Sets the ribbon twist for a particle, in degrees."));
 		AttrDescStrMap.Add(SYS_PARAM_PARTICLES_RIBBONFACING, LOCTEXT("RibbonFacingDesc", "Sets the facing vector of the ribbon at the particle position."));
+		AttrDescStrMap.Add(SYS_PARAM_PARTICLES_RIBBONSIDEVECTOR, LOCTEXT("RibbonSideVectorDesc", "Sets the vector that the ribbon is extended width-wise at the particle position."));
 		AttrDescStrMap.Add(SYS_PARAM_PARTICLES_RIBBONLINKORDER, LOCTEXT("RibbonLinkOrderDesc", "Explicit order for linking particles within a ribbon. Particles of the same ribbon id will be connected into a ribbon in incrementing order of this attribute value."));
 		AttrDescStrMap.Add(SYS_PARAM_PARTICLES_ID, LOCTEXT("IDDesc", "Engine managed particle attribute that is a persistent ID for each particle."));
 		AttrDescStrMap.Add(SYS_PARAM_PARTICLES_UNIQUE_ID, LOCTEXT("UniqueIDDesc", "Engine managed particle attribute that is a unique ID for each particle. The ID is incremented for each new particle spawned.")); 
diff --git a/Engine/Plugins/FX/Niagara/Source/Niagara/Private/NiagaraModule.cpp b/Engine/Plugins/FX/Niagara/Source/Niagara/Private/NiagaraModule.cpp
--- a/Engine/Plugins/FX/Niagara/Source/Niagara/Private/NiagaraModule.cpp
+++ b/Engine/Plugins/FX/Niagara/Source/Niagara/Private/NiagaraModule.cpp
@@ -144,6 +144,7 @@ FNiagaraVariable INiagaraModule::Particles_RibbonID;
 FNiagaraVariable INiagaraModule::Particles_RibbonWidth;
 FNiagaraVariable INiagaraModule::Particles_RibbonTwist;
 FNiagaraVariable INiagaraModule::Particles_RibbonFacing;
+FNiagaraVariable INiagaraModule::Particles_RibbonSideVector;
 FNiagaraVariable INiagaraModule::Particles_RibbonLinkOrder;
 FNiagaraVariable INiagaraModule::ScriptUsage;
 FNiagaraVariable INiagaraModule::DataInstance_Alive;
@@ -252,6 +253,7 @@ void INiagaraModule::StartupModule()
 	Particles_RibbonWidth = FNiagaraVariable(FNiagaraTypeDefinition::GetFloatDef(), TEXT("Particles.RibbonWidth"));
 	Particles_RibbonTwist = FNiagaraVariable(FNiagaraTypeDefinition::GetFloatDef(), TEXT("Particles.RibbonTwist"));
 	Particles_RibbonFacing = FNiagaraVariable(FNiagaraTypeDefinition::GetVec3Def(), TEXT("Particles.RibbonFacing"));
+	Particles_RibbonSideVector = FNiagaraVariable(FNiagaraTypeDefinition::GetVec3Def(), TEXT("Particles.RibbonSideVector"));
 	Particles_RibbonLinkOrder = FNiagaraVariable(FNiagaraTypeDefinition::GetFloatDef(), TEXT("Particles.RibbonLinkOrder"));
 
 	ScriptUsage = FNiagaraVariable(FNiagaraTypeDefinition::GetScriptUsageEnum(), TEXT("Script.Usage"));
diff --git a/Engine/Plugins/FX/Niagara/Source/Niagara/Private/NiagaraRendererRibbons.cpp b/Engine/Plugins/FX/Niagara/Source/Niagara/Private/NiagaraRendererRibbons.cpp
--- a/Engine/Plugins/FX/Niagara/Source/Niagara/Private/NiagaraRendererRibbons.cpp
+++ b/Engine/Plugins/FX/Niagara/Source/Niagara/Private/NiagaraRendererRibbons.cpp
@@ -158,6 +158,7 @@ FNiagaraRendererRibbons::FNiagaraRendererRibbons(ERHIFeatureLevel::Type FeatureL
 	, WidthDataOffset(INDEX_NONE)
 	, TwistDataOffset(INDEX_NONE)
 	, FacingDataOffset(INDEX_NONE)
+	, SideVectorDataOffset(INDEX_NONE)
 	, ColorDataOffset(INDEX_NONE)
 	, NormalizedAgeDataOffset(INDEX_NONE)
 	, MaterialRandomDataOffset(INDEX_NONE)
@@ -195,6 +196,7 @@ FNiagaraRendererRibbons::FNiagaraRendererRibbons(ERHIFeatureLevel::Type FeatureL
 	Data.GetVariableComponentOffsets(Properties->RibbonWidthBinding.DataSetVariable, WidthDataOffset, IntDummy);
 	Data.GetVariableComponentOffsets(Properties->RibbonTwistBinding.DataSetVariable, TwistDataOffset, IntDummy);
 	Data.GetVariableComponentOffsets(Properties->RibbonFacingBinding.DataSetVariable, FacingDataOffset, IntDummy);
+	Data.GetVariableComponentOffsets(Properties->SideVectorBinding.DataSetVariable, SideVectorDataOffset, IntDummy);
 	Data.GetVariableComponentOffsets(Properties->NormalizedAgeBinding.DataSetVariable, NormalizedAgeDataOffset, IntDummy);
 	Data.GetVariableComponentOffsets(Properties->MaterialRandomBinding.DataSetVariable, MaterialRandomDataOffset, IntDummy);
 
@@ -439,6 +441,7 @@ void FNiagaraRendererRibbons::GetDynamicMeshElements(const TArray<const FSceneVi
 			PerViewUniformParameters.WidthDataOffset = WidthDataOffset;
 			PerViewUniformParameters.TwistDataOffset = TwistDataOffset;
 			PerViewUniformParameters.FacingDataOffset = FacingMode == ENiagaraRibbonFacingMode::Custom ? FacingDataOffset : -1;
+			PerViewUniformParameters.SideVectorDataOffset = FacingMode == ENiagaraRibbonFacingMode::FromSideVector ? SideVectorDataOffset : -1;
 			PerViewUniformParameters.NormalizedAgeDataOffset = NormalizedAgeDataOffset;
 			PerViewUniformParameters.MaterialRandomDataOffset = MaterialRandomDataOffset;
 			PerViewUniformParameters.MaterialParamValidMask = MaterialParamValidMask;
@@ -634,7 +637,8 @@ FNiagaraDynamicDataBase* FNiagaraRendererRibbons::GenerateDynamicData(const FNia
 	FNiagaraDataSetAccessor<FVector> PosData(Data, Properties->PositionBinding.DataSetVariable);
 	FNiagaraDataSetAccessor<float> SizeData(Data, Properties->RibbonWidthBinding.DataSetVariable);
 	FNiagaraDataSetAccessor<float> TwistData(Data, Properties->RibbonTwistBinding.DataSetVariable);
-	FNiagaraDataSetAccessor<FVector> AlignData(Data, Properties->RibbonFacingBinding.DataSetVariable);
+	FNiagaraDataSetAccessor<FVector> FacingData(Data, Properties->RibbonFacingBinding.DataSetVariable);
+	FNiagaraDataSetAccessor<FVector> AlignmentData(Data, Properties->SideVectorBinding.DataSetVariable);
 	FNiagaraDataSetAccessor<FVector4> MaterialParamData(Data, Properties->DynamicMaterialBinding.DataSetVariable);
 	FNiagaraDataSetAccessor<FVector4> MaterialParam1Data(Data, Properties->DynamicMaterial1Binding.DataSetVariable);
 	FNiagaraDataSetAccessor<FVector4> MaterialParam2Data(Data, Properties->DynamicMaterial2Binding.DataSetVariable);
diff --git a/Engine/Plugins/FX/Niagara/Source/Niagara/Private/NiagaraRibbonRendererProperties.cpp b/Engine/Plugins/FX/Niagara/Source/Niagara/Private/NiagaraRibbonRendererProperties.cpp
--- a/Engine/Plugins/FX/Niagara/Source/Niagara/Private/NiagaraRibbonRendererProperties.cpp
+++ b/Engine/Plugins/FX/Niagara/Source/Niagara/Private/NiagaraRibbonRendererProperties.cpp
@@ -71,6 +71,7 @@ void UNiagaraRibbonRendererProperties::InitBindings()
 		RibbonTwistBinding = FNiagaraConstants::GetAttributeDefaultBinding(SYS_PARAM_PARTICLES_RIBBONTWIST);
 		RibbonWidthBinding = FNiagaraConstants::GetAttributeDefaultBinding(SYS_PARAM_PARTICLES_RIBBONWIDTH);
 		RibbonFacingBinding = FNiagaraConstants::GetAttributeDefaultBinding(SYS_PARAM_PARTICLES_RIBBONFACING);
+		SideVectorBinding = FNiagaraConstants::GetAttributeDefaultBinding(SYS_PARAM_PARTICLES_RIBBONSIDEVECTOR);
 		RibbonIdBinding = FNiagaraConstants::GetAttributeDefaultBinding(SYS_PARAM_PARTICLES_RIBBONID);
 		RibbonLinkOrderBinding = FNiagaraConstants::GetAttributeDefaultBinding(SYS_PARAM_PARTICLES_RIBBONLINKORDER);
 		MaterialRandomBinding = FNiagaraConstants::GetAttributeDefaultBinding(SYS_PARAM_PARTICLES_MATERIAL_RANDOM);
@@ -116,6 +117,7 @@ const TArray<FNiagaraVariable>& UNiagaraRibbonRendererProperties::GetOptionalAtt
 		Attrs.Add(SYS_PARAM_PARTICLES_RIBBONTWIST);
 		Attrs.Add(SYS_PARAM_PARTICLES_RIBBONWIDTH);
 		Attrs.Add(SYS_PARAM_PARTICLES_RIBBONFACING);
+		Attrs.Add(SYS_PARAM_PARTICLES_RIBBONSIDEVECTOR);
 		Attrs.Add(SYS_PARAM_PARTICLES_RIBBONLINKORDER);
 	}
 
diff --git a/Engine/Plugins/FX/Niagara/Source/Niagara/Public/NiagaraModule.h b/Engine/Plugins/FX/Niagara/Source/Niagara/Public/NiagaraModule.h
--- a/Engine/Plugins/FX/Niagara/Source/Niagara/Public/NiagaraModule.h
+++ b/Engine/Plugins/FX/Niagara/Source/Niagara/Public/NiagaraModule.h
@@ -165,6 +165,7 @@ public:
 	FORCEINLINE static const FNiagaraVariable&  GetVar_Particles_RibbonWidth() { return Particles_RibbonWidth; }
 	FORCEINLINE static const FNiagaraVariable&  GetVar_Particles_RibbonTwist() { return Particles_RibbonTwist; }
 	FORCEINLINE static const FNiagaraVariable&  GetVar_Particles_RibbonFacing() { return Particles_RibbonFacing; }
+	FORCEINLINE static const FNiagaraVariable&  GetVar_Particles_RibbonSideVector() { return Particles_RibbonSideVector; }
 	FORCEINLINE static const FNiagaraVariable&  GetVar_Particles_RibbonLinkOrder() { return Particles_RibbonLinkOrder; }
 	
 	FORCEINLINE static const FNiagaraVariable&  GetVar_DataInstance_Alive() { return DataInstance_Alive; }
@@ -262,6 +263,7 @@ private:
 	static FNiagaraVariable Particles_RibbonWidth;
 	static FNiagaraVariable Particles_RibbonTwist;
 	static FNiagaraVariable Particles_RibbonFacing;
+	static FNiagaraVariable Particles_RibbonSideVector;
 	static FNiagaraVariable Particles_RibbonLinkOrder;
 
 	static FNiagaraVariable ScriptUsage;
diff --git a/Engine/Plugins/FX/Niagara/Source/Niagara/Public/NiagaraRendererRibbons.h b/Engine/Plugins/FX/Niagara/Source/Niagara/Public/NiagaraRendererRibbons.h
--- a/Engine/Plugins/FX/Niagara/Source/Niagara/Public/NiagaraRendererRibbons.h
+++ b/Engine/Plugins/FX/Niagara/Source/Niagara/Public/NiagaraRendererRibbons.h
@@ -60,6 +60,7 @@ private:
 	int32 WidthDataOffset;
 	int32 TwistDataOffset;
 	int32 FacingDataOffset;
+	int32 SideVectorDataOffset;
 	int32 ColorDataOffset;
 	int32 NormalizedAgeDataOffset;
 	int32 MaterialRandomDataOffset;
diff --git a/Engine/Plugins/FX/Niagara/Source/Niagara/Public/NiagaraRibbonRendererProperties.h b/Engine/Plugins/FX/Niagara/Source/Niagara/Public/NiagaraRibbonRendererProperties.h
--- a/Engine/Plugins/FX/Niagara/Source/Niagara/Public/NiagaraRibbonRendererProperties.h
+++ b/Engine/Plugins/FX/Niagara/Source/Niagara/Public/NiagaraRibbonRendererProperties.h
@@ -14,7 +14,12 @@ UENUM()
 enum class ENiagaraRibbonFacingMode : uint8
 {
 	Screen,
-	Custom
+
+	/** Use Particles.RibbonFacing as the facing vector. */
+	Custom,
+
+	/** Use Particles.RibbonSideVector to determine facing vector. */
+	FromSideVector
 };
 
 /** Defines different modes for offsetting UVs by age when ordering ribbon particles using normalized age. */
@@ -170,6 +175,10 @@ public:
 	/** Which attribute should we use for ribbon facing when generating ribbons?*/
 	UPROPERTY(EditAnywhere, AdvancedDisplay, Category = "Bindings")
 	FNiagaraVariableAttributeBinding RibbonFacingBinding;
+
+	/** Which attribute should we use for ribbon side vector when generating ribbons?*/
+	UPROPERTY(EditAnywhere, AdvancedDisplay, Category = "Bindings")
+	FNiagaraVariableAttributeBinding SideVectorBinding;
 	
 	/** Which attribute should we use for ribbon id when generating ribbons?*/
 	UPROPERTY(EditAnywhere, AdvancedDisplay, Category = "Bindings")
diff --git a/Engine/Plugins/FX/Niagara/Source/NiagaraVertexFactories/Public/NiagaraRibbonVertexFactory.h b/Engine/Plugins/FX/Niagara/Source/NiagaraVertexFactories/Public/NiagaraRibbonVertexFactory.h
--- a/Engine/Plugins/FX/Niagara/Source/NiagaraVertexFactories/Public/NiagaraRibbonVertexFactory.h
+++ b/Engine/Plugins/FX/Niagara/Source/NiagaraVertexFactories/Public/NiagaraRibbonVertexFactory.h
@@ -35,6 +35,7 @@ BEGIN_GLOBAL_SHADER_PARAMETER_STRUCT(FNiagaraRibbonUniformParameters, NIAGARAVER
 	SHADER_PARAMETER(int, TwistDataOffset)
 	SHADER_PARAMETER(int, ColorDataOffset)
 	SHADER_PARAMETER(int, FacingDataOffset)
+	SHADER_PARAMETER(int, SideVectorDataOffset)
 	SHADER_PARAMETER(int, NormalizedAgeDataOffset)
 	SHADER_PARAMETER(int, MaterialRandomDataOffset)
 	SHADER_PARAMETER(uint32, MaterialParamValidMask)
