Small improvement to the compilation of static switches

diff --git a/Engine/Plugins/FX/Niagara/Source/NiagaraEditor/Private/NiagaraEditorUtilities.cpp b/Engine/Plugins/FX/Niagara/Source/NiagaraEditor/Private/NiagaraEditorUtilities.cpp
--- a/Engine/Plugins/FX/Niagara/Source/NiagaraEditor/Private/NiagaraEditorUtilities.cpp
+++ b/Engine/Plugins/FX/Niagara/Source/NiagaraEditor/Private/NiagaraEditorUtilities.cpp
@@ -625,6 +625,8 @@ void FNiagaraEditorUtilities::FixUpNumericPins(const UEdGraphSchema_Niagara* Sch
 
 void FNiagaraEditorUtilities::SetStaticSwitchConstants(UNiagaraGraph* Graph, const TArray<UEdGraphPin*>& CallInputs, const FCompileConstantResolver& ConstantResolver)
 {
+	const UEdGraphSchema_Niagara* Schema = GetDefault<UEdGraphSchema_Niagara>();
+
 	for (UEdGraphNode* Node : Graph->Nodes)
 	{
 		// if there is a static switch node its value must be set by the caller
@@ -637,17 +639,18 @@ void FNiagaraEditorUtilities::SetStaticSwitchConstants(UNiagaraGraph* Graph, con
 			}
 			else
 			{
+				FEdGraphPinType VarType = Schema->TypeDefinitionToPinType(SwitchNode->GetInputType());
 				SwitchNode->ClearSwitchValue();
 				for (UEdGraphPin* InputPin : CallInputs)
 				{
-					if (InputPin->GetFName().IsEqual(SwitchNode->InputParameterName))
+					if (InputPin->GetFName().IsEqual(SwitchNode->InputParameterName) && InputPin->PinType == VarType)
 					{
 						int32 SwitchValue = 0;
 						if (ResolveConstantValue(InputPin, SwitchValue))
 						{
 							SwitchNode->SetSwitchValue(SwitchValue);
+							break;
 						}
-						break;
 					}
 				}
 			}
diff --git a/Engine/Plugins/FX/Niagara/Source/NiagaraEditor/Private/NiagaraHlslTranslator.cpp b/Engine/Plugins/FX/Niagara/Source/NiagaraEditor/Private/NiagaraHlslTranslator.cpp
--- a/Engine/Plugins/FX/Niagara/Source/NiagaraEditor/Private/NiagaraHlslTranslator.cpp
+++ b/Engine/Plugins/FX/Niagara/Source/NiagaraEditor/Private/NiagaraHlslTranslator.cpp
@@ -2643,6 +2643,32 @@ bool FHlslNiagaraTranslator::ShouldInterpolateParameter(const FNiagaraVariable&
 	return true;
 }
 
+void FHlslNiagaraTranslator::UpdateStaticSwitchConstants(UEdGraphNode* Node)
+{
+	if (UNiagaraNodeStaticSwitch* SwitchNode = Cast<UNiagaraNodeStaticSwitch>(Node))
+	{
+		TArray<UNiagaraNodeStaticSwitch*> NodesToUpdate;
+		NodesToUpdate.Add(SwitchNode);
+
+		for (int i = 0; i < NodesToUpdate.Num(); i++)
+		{
+			SwitchNode->UpdateCompilerConstantValue(this);
+			
+			// also check direct upstream static switches, because they are otherwise skipped during the compilation and
+			// might be evaluated without their values set correctly.
+			TArray<UEdGraphPin*> InPins;
+			SwitchNode->GetInputPins(InPins);
+			for (UEdGraphPin* Pin : InPins)
+			{
+				if (UNiagaraNodeStaticSwitch* ConnectedNode = Cast<UNiagaraNodeStaticSwitch>(Pin->GetOwningNode()))
+				{
+					NodesToUpdate.AddUnique(ConnectedNode);
+				}
+			}
+		}
+	}
+}
+
 int32 FHlslNiagaraTranslator::GetRapidIterationParameter(const FNiagaraVariable& Parameter)
 {
 	if (!AddStructToDefinitionSet(Parameter.GetType()))
@@ -5434,10 +5460,7 @@ int32 FHlslNiagaraTranslator::CompileOutputPin(const UEdGraphPin* InPin)
 
 	if (InPin)
 	{
-		if (UNiagaraNodeStaticSwitch* SwitchNode = Cast<UNiagaraNodeStaticSwitch>(InPin->GetOwningNode()))
-		{
-			SwitchNode->UpdateCompilerConstantValue(this);
-		}
+		UpdateStaticSwitchConstants(InPin->GetOwningNode());
 	}
 
 	// The incoming pin to compile may be pointing to a reroute node. If so, we just jump over it
diff --git a/Engine/Plugins/FX/Niagara/Source/NiagaraEditor/Private/NiagaraHlslTranslator.h b/Engine/Plugins/FX/Niagara/Source/NiagaraEditor/Private/NiagaraHlslTranslator.h
--- a/Engine/Plugins/FX/Niagara/Source/NiagaraEditor/Private/NiagaraHlslTranslator.h
+++ b/Engine/Plugins/FX/Niagara/Source/NiagaraEditor/Private/NiagaraHlslTranslator.h
@@ -543,6 +543,8 @@ private:
 
 	bool ShouldInterpolateParameter(const FNiagaraVariable& Parameter);
 
+	void UpdateStaticSwitchConstants(UEdGraphNode* Node);
+
 	bool IsBulkSystemScript() const;
 	bool IsSpawnScript() const;
 	bool RequiresInterpolation() const;
