Removed the "Dynamic load balanced sim" enum value for simulation targets

#tests passed the niagara automation tests, assets still seem to work fine as well

diff --git a/Engine/Plugins/FX/Niagara/Source/Niagara/Private/NiagaraEmitter.cpp b/Engine/Plugins/FX/Niagara/Source/Niagara/Private/NiagaraEmitter.cpp
--- a/Engine/Plugins/FX/Niagara/Source/Niagara/Private/NiagaraEmitter.cpp
+++ b/Engine/Plugins/FX/Niagara/Source/Niagara/Private/NiagaraEmitter.cpp
@@ -549,7 +549,7 @@ bool UNiagaraEmitter::IsValid()const
 		return false;
 	}
 
-	if (SimTarget == ENiagaraSimTarget::CPUSim || SimTarget == ENiagaraSimTarget::DynamicLoadBalancedSim)
+	if (SimTarget == ENiagaraSimTarget::CPUSim)
 	{
 		if (!SpawnScriptProps.Script->IsScriptCompilationPending(false) && !SpawnScriptProps.Script->DidScriptCompilationSucceed(false))
 		{
@@ -572,7 +572,7 @@ bool UNiagaraEmitter::IsValid()const
 		}
 	}
 
-	if (SimTarget == ENiagaraSimTarget::GPUComputeSim || SimTarget == ENiagaraSimTarget::DynamicLoadBalancedSim)
+	if (SimTarget == ENiagaraSimTarget::GPUComputeSim)
 	{
 		if (!GPUComputeScript->IsScriptCompilationPending(true) && 
 			!GPUComputeScript->DidScriptCompilationSucceed(true))
@@ -591,7 +591,7 @@ bool UNiagaraEmitter::IsReadyToRun() const
 		return false;
 	}
 
-	if (SimTarget == ENiagaraSimTarget::CPUSim || SimTarget == ENiagaraSimTarget::DynamicLoadBalancedSim)
+	if (SimTarget == ENiagaraSimTarget::CPUSim)
 	{
 		if (SpawnScriptProps.Script->IsScriptCompilationPending(false))
 		{
@@ -613,7 +613,7 @@ bool UNiagaraEmitter::IsReadyToRun() const
 		}
 	}
 
-	if (SimTarget == ENiagaraSimTarget::GPUComputeSim || SimTarget == ENiagaraSimTarget::DynamicLoadBalancedSim)
+	if (SimTarget == ENiagaraSimTarget::GPUComputeSim)
 	{
 		if (GPUComputeScript->IsScriptCompilationPending(true))
 		{
@@ -642,7 +642,7 @@ void UNiagaraEmitter::GetScripts(TArray<UNiagaraScript*>& OutScripts, bool bComp
 		}
 	}
 
-	if (SimTarget == ENiagaraSimTarget::DynamicLoadBalancedSim || SimTarget == ENiagaraSimTarget::GPUComputeSim)
+	if (SimTarget == ENiagaraSimTarget::GPUComputeSim)
 	{
 		OutScripts.Add(GPUComputeScript);
 	}
diff --git a/Engine/Plugins/FX/Niagara/Source/Niagara/Private/NiagaraEmitterInstance.cpp b/Engine/Plugins/FX/Niagara/Source/Niagara/Private/NiagaraEmitterInstance.cpp
--- a/Engine/Plugins/FX/Niagara/Source/Niagara/Private/NiagaraEmitterInstance.cpp
+++ b/Engine/Plugins/FX/Niagara/Source/Niagara/Private/NiagaraEmitterInstance.cpp
@@ -428,7 +428,7 @@ void FNiagaraEmitterInstance::CheckForErrors()
 		return;
 	}
 
-	if (CachedEmitter->SimTarget == ENiagaraSimTarget::CPUSim || CachedEmitter->SimTarget == ENiagaraSimTarget::DynamicLoadBalancedSim)
+	if (CachedEmitter->SimTarget == ENiagaraSimTarget::CPUSim)
 	{
 		bool bFailed = false;
 		if (!CachedEmitter->SpawnScriptProps.Script->DidScriptCompilationSucceed(false))
@@ -462,7 +462,7 @@ void FNiagaraEmitterInstance::CheckForErrors()
 		}
 	}
 
-	if (CachedEmitter->SimTarget == ENiagaraSimTarget::GPUComputeSim || CachedEmitter->SimTarget == ENiagaraSimTarget::DynamicLoadBalancedSim)
+	if (CachedEmitter->SimTarget == ENiagaraSimTarget::GPUComputeSim)
 	{
 		if (CachedEmitter->GetGPUComputeScript()->IsScriptCompilationPending(true))
 		{
diff --git a/Engine/Plugins/FX/Niagara/Source/Niagara/Public/NiagaraCommon.h b/Engine/Plugins/FX/Niagara/Source/Niagara/Public/NiagaraCommon.h
--- a/Engine/Plugins/FX/Niagara/Source/Niagara/Public/NiagaraCommon.h
+++ b/Engine/Plugins/FX/Niagara/Source/Niagara/Public/NiagaraCommon.h
@@ -37,8 +37,7 @@ UENUM()
 enum class ENiagaraSimTarget : uint8
 {
 	CPUSim,
-	GPUComputeSim,
-	DynamicLoadBalancedSim UMETA(Hidden)
+	GPUComputeSim
 };
 
 
diff --git a/Engine/Plugins/FX/Niagara/Source/NiagaraEditor/Private/ViewModels/NiagaraEmitterViewModel.cpp b/Engine/Plugins/FX/Niagara/Source/NiagaraEditor/Private/ViewModels/NiagaraEmitterViewModel.cpp
--- a/Engine/Plugins/FX/Niagara/Source/NiagaraEditor/Private/ViewModels/NiagaraEmitterViewModel.cpp
+++ b/Engine/Plugins/FX/Niagara/Source/NiagaraEditor/Private/ViewModels/NiagaraEmitterViewModel.cpp
@@ -294,7 +294,7 @@ ENiagaraScriptCompileStatus FNiagaraEmitterViewModel::GetLatestCompileStatus()
 	check(SharedScriptViewModel.IsValid());
 	ENiagaraScriptCompileStatus UnionStatus = SharedScriptViewModel->GetLatestCompileStatus();
 
-	if (Emitter->SimTarget == ENiagaraSimTarget::GPUComputeSim || Emitter->SimTarget == ENiagaraSimTarget::DynamicLoadBalancedSim)
+	if (Emitter->SimTarget == ENiagaraSimTarget::GPUComputeSim)
 	{
 		if (UnionStatus != ENiagaraScriptCompileStatus::NCS_Dirty)
 		{
