Niagara - Fix an issue with emitter inheritance where changes wouldn't be merged from parent to child if the child had never been modified.

diff --git a/Engine/Plugins/FX/Niagara/Source/NiagaraEditor/Private/NiagaraScriptMergeManager.cpp b/Engine/Plugins/FX/Niagara/Source/NiagaraEditor/Private/NiagaraScriptMergeManager.cpp
--- a/Engine/Plugins/FX/Niagara/Source/NiagaraEditor/Private/NiagaraScriptMergeManager.cpp
+++ b/Engine/Plugins/FX/Niagara/Source/NiagaraEditor/Private/NiagaraScriptMergeManager.cpp
@@ -821,7 +821,19 @@ INiagaraMergeManager::FMergeEmitterResults FNiagaraScriptMergeManager::MergeEmit
 	}
 	else if (DiffResults.IsEmpty())
 	{
-		MergeResults.MergeResult = INiagaraMergeManager::EMergeEmitterResult::SucceededNoDifferences;
+		// If there were no changes made on the instance, check if the instance matches the parent.
+		FNiagaraEmitterDiffResults DiffResultsFromParent = DiffEmitters(Parent, Instance);
+		if (DiffResultsFromParent.IsEmpty())
+		{
+			MergeResults.MergeResult = INiagaraMergeManager::EMergeEmitterResult::SucceededNoDifferences;
+		}
+		else
+		{
+			// If there were differences from the parent we can just return a copy of the parent as the merged instance since there
+			// were no changes in the instance which need to be applied.
+			MergeResults.MergeResult = INiagaraMergeManager::EMergeEmitterResult::SucceededDifferencesApplied;
+			MergeResults.MergedInstance = Parent.DuplicateWithoutMerging((UObject*)GetTransientPackage());
+		}
 	}
 	else
 	{
