Niagara - System Overview tab only.

diff --git a/Engine/Plugins/FX/Niagara/Source/NiagaraEditor/Private/NiagaraEditorModule.cpp b/Engine/Plugins/FX/Niagara/Source/NiagaraEditor/Private/NiagaraEditorModule.cpp
--- a/Engine/Plugins/FX/Niagara/Source/NiagaraEditor/Private/NiagaraEditorModule.cpp
+++ b/Engine/Plugins/FX/Niagara/Source/NiagaraEditor/Private/NiagaraEditorModule.cpp
@@ -726,23 +726,22 @@ TSharedPtr<INiagaraEditorTypeUtilities, ESPMode::ThreadSafe> FNiagaraEditorModul
 	return TSharedPtr<INiagaraEditorTypeUtilities, ESPMode::ThreadSafe>();
 }
 
-TSharedRef<SWidget> FNiagaraEditorModule::CreateStackWidget(UNiagaraStackViewModel* StackViewModel) const
+
+void FNiagaraEditorModule::RegisterWidgetProvider(TSharedRef<INiagaraEditorWidgetProvider> InWidgetProvider)
 {
-	checkf(OnCreateStackWidget.IsBound(), TEXT("Can not create stack widget.  Stack creation delegate was never set."));
-	return OnCreateStackWidget.Execute(StackViewModel);
+	checkf(WidgetProvider.IsValid() == false, TEXT("Widget provider has already been set."));
+	WidgetProvider = InWidgetProvider;
 }
 
-FDelegateHandle FNiagaraEditorModule::SetOnCreateStackWidget(FOnCreateStackWidget InOnCreateStackWidget)
-{
-	checkf(OnCreateStackWidget.IsBound() == false, TEXT("Stack creation delegate already set."));
-	OnCreateStackWidget = InOnCreateStackWidget;
-	return OnCreateStackWidget.GetHandle();
+void FNiagaraEditorModule::UnregisterWidgetProvider(TSharedRef<INiagaraEditorWidgetProvider> InWidgetProvider)
+{	
+	checkf(WidgetProvider.IsValid() && WidgetProvider == InWidgetProvider, TEXT("Can only unregister the widget provider that was originally registered."));
+	WidgetProvider.Reset();
 }
 
-void FNiagaraEditorModule::ResetOnCreateStackWidget(FDelegateHandle Handle)
+TSharedRef<INiagaraEditorWidgetProvider> FNiagaraEditorModule::GetWidgetProvider() const
 {
-	checkf(OnCreateStackWidget.GetHandle() == Handle, TEXT("Can only reset the stack creation module with the handle it was created with."));
-	OnCreateStackWidget.Unbind();
+	return WidgetProvider.ToSharedRef();
 }
 
 TSharedRef<FNiagaraScriptMergeManager> FNiagaraEditorModule::GetScriptMergeManager() const
diff --git a/Engine/Plugins/FX/Niagara/Source/NiagaraEditor/Private/Toolkits/NiagaraSystemToolkit.cpp b/Engine/Plugins/FX/Niagara/Source/NiagaraEditor/Private/Toolkits/NiagaraSystemToolkit.cpp
--- a/Engine/Plugins/FX/Niagara/Source/NiagaraEditor/Private/Toolkits/NiagaraSystemToolkit.cpp
+++ b/Engine/Plugins/FX/Niagara/Source/NiagaraEditor/Private/Toolkits/NiagaraSystemToolkit.cpp
@@ -66,6 +66,7 @@ const FName FNiagaraSystemToolkit::DebugSpreadsheetTabID(TEXT("NiagaraSystemEdit
 const FName FNiagaraSystemToolkit::PreviewSettingsTabId(TEXT("NiagaraSystemEditor_PreviewSettings"));
 const FName FNiagaraSystemToolkit::GeneratedCodeTabID(TEXT("NiagaraSystemEditor_GeneratedCode"));
 const FName FNiagaraSystemToolkit::MessageLogTabID(TEXT("NiagaraSystemEditor_MessageLog"));
+const FName FNiagaraSystemToolkit::SystemOverviewTabID(TEXT("NiagaraSystemEditor_SystemOverview"));
 
 static int32 GbLogNiagaraSystemChanges = 0;
 static FAutoConsoleVariableRef CVarSuppressNiagaraSystems(
@@ -140,6 +141,10 @@ void FNiagaraSystemToolkit::RegisterTabSpawners(const TSharedRef<class FTabManag
 	InTabManager->RegisterTabSpawner(MessageLogTabID, FOnSpawnTab::CreateSP(this, &FNiagaraSystemToolkit::SpawnTab_MessageLog))
 		.SetDisplayName(LOCTEXT("NiagaraMessageLog", "Niagara Log"))
 		.SetGroup(WorkspaceMenuCategory.ToSharedRef());
+
+	InTabManager->RegisterTabSpawner(SystemOverviewTabID, FOnSpawnTab::CreateSP(this, &FNiagaraSystemToolkit::SpawnTab_SystemOverview))
+		.SetDisplayName(LOCTEXT("SystemOverviewTabName", "System Overview"))
+		.SetGroup(WorkspaceMenuCategory.ToSharedRef());
 }
 
 void FNiagaraSystemToolkit::UnregisterTabSpawners(const TSharedRef<class FTabManager>& InTabManager)
@@ -157,6 +162,7 @@ void FNiagaraSystemToolkit::UnregisterTabSpawners(const TSharedRef<class FTabMan
 	InTabManager->UnregisterTabSpawner(DebugSpreadsheetTabID);
 	InTabManager->UnregisterTabSpawner(PreviewSettingsTabId);
 	InTabManager->UnregisterTabSpawner(GeneratedCodeTabID);
+	InTabManager->UnregisterTabSpawner(SystemOverviewTabID);
 }
 
 
@@ -295,7 +301,7 @@ void FNiagaraSystemToolkit::InitializeInternal(const EToolkitMode::Type Mode, co
 	const float InTime = -0.02f;
 	const float OutTime = 3.2f;
 
-	TSharedRef<FTabManager::FLayout> StandaloneDefaultLayout = FTabManager::NewLayout("Standalone_Niagara_System_Layout_v18")
+	TSharedRef<FTabManager::FLayout> StandaloneDefaultLayout = FTabManager::NewLayout("Standalone_Niagara_System_Layout_v19")
 		->AddArea
 		(
 			FTabManager::NewPrimaryArea()->SetOrientation(Orient_Vertical)
@@ -322,6 +328,7 @@ void FNiagaraSystemToolkit::InitializeInternal(const EToolkitMode::Type Mode, co
 							FTabManager::NewStack()
 							->SetSizeCoefficient(.80f)
 							->AddTab(ViewportTabID, ETabState::OpenedTab)
+							->AddTab(SystemOverviewTabID, ETabState::ClosedTab)
 						)
 						->Split
 						(
@@ -642,6 +649,18 @@ TSharedRef<SDockTab> FNiagaraSystemToolkit::SpawnTab_MessageLog(const FSpawnTabA
 	return SpawnedTab;
 }
 
+TSharedRef<SDockTab> FNiagaraSystemToolkit::SpawnTab_SystemOverview(const FSpawnTabArgs& Args)
+{
+	FNiagaraEditorModule& NiagaraEditorModule = FModuleManager::LoadModuleChecked<FNiagaraEditorModule>("NiagaraEditor");
+	TSharedRef<SDockTab> SpawnedTab = SNew(SDockTab)
+		.Label(LOCTEXT("SystemOverviewTabLabel", "System Overview"))
+		[
+			NiagaraEditorModule.GetWidgetProvider()->CreateSystemOverview(SystemViewModel.ToSharedRef())
+		];
+
+	return SpawnedTab;
+}
+
 void FNiagaraSystemToolkit::SetupCommands()
 {
 	GetToolkitCommands()->MapAction(
diff --git a/Engine/Plugins/FX/Niagara/Source/NiagaraEditor/Private/Toolkits/NiagaraSystemToolkit.h b/Engine/Plugins/FX/Niagara/Source/NiagaraEditor/Private/Toolkits/NiagaraSystemToolkit.h
--- a/Engine/Plugins/FX/Niagara/Source/NiagaraEditor/Private/Toolkits/NiagaraSystemToolkit.h
+++ b/Engine/Plugins/FX/Niagara/Source/NiagaraEditor/Private/Toolkits/NiagaraSystemToolkit.h
@@ -102,6 +102,7 @@ private:
 	TSharedRef<SDockTab> SpawnTab_PreviewSettings(const FSpawnTabArgs& Args);
 	TSharedRef<SDockTab> SpawnTab_GeneratedCode(const FSpawnTabArgs& Args);
 	TSharedRef<SDockTab> SpawnTab_MessageLog(const FSpawnTabArgs& Args);
+	TSharedRef<SDockTab> SpawnTab_SystemOverview(const FSpawnTabArgs& Args);
 
 	/** Builds the toolbar widget */
 	void ExtendToolbar();	
@@ -175,4 +176,5 @@ private:
 	static const FName PreviewSettingsTabId;
 	static const FName GeneratedCodeTabID;
 	static const FName MessageLogTabID;
+	static const FName SystemOverviewTabID;
 };
diff --git a/Engine/Plugins/FX/Niagara/Source/NiagaraEditor/Private/Widgets/SNiagaraSelectedEmitterHandles.cpp b/Engine/Plugins/FX/Niagara/Source/NiagaraEditor/Private/Widgets/SNiagaraSelectedEmitterHandles.cpp
--- a/Engine/Plugins/FX/Niagara/Source/NiagaraEditor/Private/Widgets/SNiagaraSelectedEmitterHandles.cpp
+++ b/Engine/Plugins/FX/Niagara/Source/NiagaraEditor/Private/Widgets/SNiagaraSelectedEmitterHandles.cpp
@@ -90,7 +90,7 @@ void SNiagaraSelectedEmitterHandles::RefreshEmitterWidgets()
 		StackViewModels.Add(StackModel);
 		EmitterSplitter->AddSlot()
 		[
-			NiagaraEditorModule.CreateStackWidget(StackModel)
+			NiagaraEditorModule.GetWidgetProvider()->CreateStackView(*StackModel)
 		];
 	}
 	EmitterHandlesToDisplay.Empty();
diff --git a/Engine/Plugins/FX/Niagara/Source/NiagaraEditor/Public/NiagaraEditorModule.h b/Engine/Plugins/FX/Niagara/Source/NiagaraEditor/Public/NiagaraEditorModule.h
--- a/Engine/Plugins/FX/Niagara/Source/NiagaraEditor/Public/NiagaraEditorModule.h
+++ b/Engine/Plugins/FX/Niagara/Source/NiagaraEditor/Public/NiagaraEditorModule.h
@@ -15,6 +15,7 @@ class INiagaraEditorTypeUtilities;
 class UNiagaraSettings;
 class USequencerSettings;
 class UNiagaraStackViewModel;
+class FNiagaraSystemViewModel;
 class FNiagaraScriptMergeManager;
 class FNiagaraCompileOptions;
 class FNiagaraCompileRequestDataBase;
@@ -23,12 +24,19 @@ struct IConsoleCommand;
 
 DECLARE_STATS_GROUP(TEXT("Niagara Editor"), STATGROUP_NiagaraEditor, STATCAT_Advanced);
 
+/* Defines methods for allowing external modules to supply widgets to the core editor module. */
+NIAGARAEDITOR_API class INiagaraEditorWidgetProvider
+{
+public:
+	virtual TSharedRef<SWidget> CreateStackView(UNiagaraStackViewModel& StackViewModel) = 0;
+	virtual TSharedRef<SWidget> CreateSystemOverview(TSharedRef<FNiagaraSystemViewModel> SystemViewModel) = 0;
+};
+
 /** Niagara Editor module */
 class FNiagaraEditorModule : public IModuleInterface,
 	public IHasMenuExtensibility, public IHasToolBarExtensibility, public FGCObject
 {
 public:
-	DECLARE_DELEGATE_RetVal_OneParam(TSharedRef<SWidget>, FOnCreateStackWidget, UNiagaraStackViewModel*);
 	DECLARE_DELEGATE_RetVal_OneParam(UMovieSceneNiagaraParameterTrack*, FOnCreateMovieSceneTrackForParameter, FNiagaraVariable);
 	DECLARE_MULTICAST_DELEGATE_OneParam(FOnCheckScriptToolkitsShouldFocusGraphElement, const FNiagaraScriptIDAndGraphFocusInfo*);
 
@@ -63,10 +71,10 @@ public:
 
 	static EAssetTypeCategories::Type GetAssetCategory() { return NiagaraAssetCategory; }
 
-	TSharedRef<SWidget> CreateStackWidget(UNiagaraStackViewModel* StackViewModel) const;
+	NIAGARAEDITOR_API void RegisterWidgetProvider(TSharedRef<INiagaraEditorWidgetProvider> InWidgetProvider);
+	NIAGARAEDITOR_API void UnregisterWidgetProvider(TSharedRef<INiagaraEditorWidgetProvider> InWidgetProvider);
 
-	FDelegateHandle NIAGARAEDITOR_API SetOnCreateStackWidget(FOnCreateStackWidget InOnCreateStackWidget);
-	void NIAGARAEDITOR_API ResetOnCreateStackWidget(FDelegateHandle DelegateHandle);
+	TSharedRef<INiagaraEditorWidgetProvider> GetWidgetProvider() const;
 
 	TSharedRef<FNiagaraScriptMergeManager> GetScriptMergeManager() const;
 
@@ -124,7 +132,8 @@ private:
 	FDelegateHandle PrecompilerHandle;
 
 	USequencerSettings* SequencerSettings;
-	FOnCreateStackWidget OnCreateStackWidget;
+	
+	TSharedPtr<INiagaraEditorWidgetProvider> WidgetProvider;
 
 	TSharedPtr<FNiagaraScriptMergeManager> ScriptMergeManager;
 
diff --git a/Engine/Plugins/FX/Niagara/Source/NiagaraEditorWidgets/Private/NiagaraEditorWidgetsModule.cpp b/Engine/Plugins/FX/Niagara/Source/NiagaraEditorWidgets/Private/NiagaraEditorWidgetsModule.cpp
--- a/Engine/Plugins/FX/Niagara/Source/NiagaraEditorWidgets/Private/NiagaraEditorWidgetsModule.cpp
+++ b/Engine/Plugins/FX/Niagara/Source/NiagaraEditorWidgets/Private/NiagaraEditorWidgetsModule.cpp
@@ -2,7 +2,6 @@
 
 #include "NiagaraEditorWidgetsModule.h"
 #include "NiagaraEditorWidgetsStyle.h"
-#include "NiagaraEditorModule.h"
 #include "SNiagaraStack.h"
 #include "DetailCustomizations/NiagaraDataInterfaceCurveDetails.h"
 #include "DetailCustomizations/NiagaraDataInterfaceDetails.h"
@@ -83,10 +82,9 @@ void FNiagaraStackCurveEditorOptions::SetAreCurvesVisible(bool bInAreCurvesVisib
 void FNiagaraEditorWidgetsModule::StartupModule()
 {
 	FNiagaraEditorModule& NiagaraEditorModule = FModuleManager::LoadModuleChecked<FNiagaraEditorModule>("NiagaraEditor");
-	OnCreateStackWidgetHandle = NiagaraEditorModule.SetOnCreateStackWidget(FNiagaraEditorModule::FOnCreateStackWidget::CreateLambda([](UNiagaraStackViewModel* ViewModel)
-	{
-		return SNew(SNiagaraStack, ViewModel);
-	}));
+	WidgetProvider = MakeShared<FNiagaraEditorWidgetProvider>();
+	NiagaraEditorModule.RegisterWidgetProvider(WidgetProvider.ToSharedRef());
+
 	FNiagaraEditorWidgetsStyle::Initialize();
 
 	FPropertyEditorModule& PropertyModule = FModuleManager::LoadModuleChecked<FPropertyEditorModule>("PropertyEditor");
@@ -104,7 +102,7 @@ void FNiagaraEditorWidgetsModule::ShutdownModule()
 	FNiagaraEditorModule* NiagaraEditorModule = FModuleManager::GetModulePtr<FNiagaraEditorModule>("NiagaraEditor");
 	if (NiagaraEditorModule != nullptr)
 	{
-		NiagaraEditorModule->ResetOnCreateStackWidget(OnCreateStackWidgetHandle);
+		NiagaraEditorModule->UnregisterWidgetProvider(WidgetProvider.ToSharedRef());
 	}
 
 	FPropertyEditorModule* PropertyModule = FModuleManager::GetModulePtr<FPropertyEditorModule>("PropertyEditor");
@@ -132,4 +130,14 @@ TSharedRef<FNiagaraStackCurveEditorOptions> FNiagaraEditorWidgetsModule::GetOrCr
 		(*StackCurveEditorOptions)->SetHeight(DefaultHeight);
 	}
 	return *StackCurveEditorOptions;
+}
+
+TSharedRef<SWidget> FNiagaraEditorWidgetsModule::FNiagaraEditorWidgetProvider::CreateStackView(UNiagaraStackViewModel& StackViewModel)
+{
+	return SNew(SNiagaraStack, &StackViewModel);
+}
+
+TSharedRef<SWidget> FNiagaraEditorWidgetsModule::FNiagaraEditorWidgetProvider::CreateSystemOverview(TSharedRef<FNiagaraSystemViewModel> SystemViewModel)
+{
+	return SNullWidget::NullWidget;
 }
\ No newline at end of file
diff --git a/Engine/Plugins/FX/Niagara/Source/NiagaraEditorWidgets/Public/NiagaraEditorWidgetsModule.h b/Engine/Plugins/FX/Niagara/Source/NiagaraEditorWidgets/Public/NiagaraEditorWidgetsModule.h
--- a/Engine/Plugins/FX/Niagara/Source/NiagaraEditorWidgets/Public/NiagaraEditorWidgetsModule.h
+++ b/Engine/Plugins/FX/Niagara/Source/NiagaraEditorWidgets/Public/NiagaraEditorWidgetsModule.h
@@ -2,6 +2,7 @@
 
 #pragma once
 
+#include "NiagaraEditorModule.h"
 #include "Modules/ModuleInterface.h"
 #include "Delegates/IDelegateInstance.h"
 #include "UObject/ObjectKey.h"
@@ -42,6 +43,13 @@ class UObject;
 /** A module containing widgets for editing niagara data. */
 class FNiagaraEditorWidgetsModule : public IModuleInterface
 {
+private:
+	class FNiagaraEditorWidgetProvider : public INiagaraEditorWidgetProvider
+	{
+	public:
+		virtual TSharedRef<SWidget> CreateStackView(UNiagaraStackViewModel& StackViewModel) override;
+		virtual TSharedRef<SWidget> CreateSystemOverview(TSharedRef<FNiagaraSystemViewModel> SystemViewModel) override;
+	};
 
 public:
 	// IModuleInterface
@@ -51,7 +59,9 @@ public:
 	TSharedRef<FNiagaraStackCurveEditorOptions> GetOrCreateStackCurveEditorOptionsForObject(UObject* Object, bool bDefaultAreCurvesVisible, float DefaultHeight);
 
 private:
-	FDelegateHandle OnCreateStackWidgetHandle;
 
+private:
 	TMap<FObjectKey, TSharedRef<FNiagaraStackCurveEditorOptions>> ObjectToStackCurveEditorOptionsMap;
+
+	TSharedPtr<FNiagaraEditorWidgetProvider> WidgetProvider;
 };
\ No newline at end of file
