Changed the LOD selection logic in FNiagaraRendererMeshes to use the same fallback as in FDynamicMeshEmitterData::GetMeshLODIndexFromProxy().

diff --git a/Engine/Plugins/FX/Niagara/Source/Niagara/Private/NiagaraRendererMeshes.cpp b/Engine/Plugins/FX/Niagara/Source/Niagara/Private/NiagaraRendererMeshes.cpp
--- a/Engine/Plugins/FX/Niagara/Source/Niagara/Private/NiagaraRendererMeshes.cpp
+++ b/Engine/Plugins/FX/Niagara/Source/Niagara/Private/NiagaraRendererMeshes.cpp
@@ -213,8 +213,14 @@ void FNiagaraRendererMeshes::GetDynamicMeshElements(const TArray<const FSceneVie
 			if (VisibilityMap & (1 << ViewIndex))
 			{
 				const FSceneView* View = Views[ViewIndex];
-				//TODO: Work out LOD and pass through in dynamic data.
-				const int32 LODIndex = FMath::Max<int32>((int32)MeshRenderData->CurrentFirstLODIdx, MeshMinimumLOD);
+
+				// Doesn't seem to work for some reason. See comment in FDynamicMeshEmitterData::GetMeshLODIndexFromProxy()
+				// const int32 LODIndex = FMath::Max<int32>((int32)MeshRenderData->CurrentFirstLODIdx, MeshMinimumLOD);
+				int32 LODIndex = (int32)MeshRenderData->CurrentFirstLODIdx;
+				while (LODIndex < MeshRenderData->LODResources.Num() - 1 && !MeshRenderData->LODResources[LODIndex].GetNumVertices())
+				{
+					++LODIndex;
+				}
 				const FStaticMeshLODResources& LODModel = MeshRenderData->LODResources[LODIndex];
 
 				FNiagaraMeshCollectorResourcesMesh& CollectorResources = Collector.AllocateOneFrameResource<FNiagaraMeshCollectorResourcesMesh>();
@@ -430,9 +436,14 @@ FNiagaraDynamicDataBase *FNiagaraRendererMeshes::GenerateDynamicData(const FNiag
 	{
 		DynamicData = new FNiagaraDynamicDataMesh(Emitter);
 
-		//TODO: Mesh LODs.
-		const int32 LODIndex = FMath::Max<int32>((int32)MeshRenderData->CurrentFirstLODIdx, MeshMinimumLOD);
-		const FStaticMeshLODResources& LODModel = Properties->ParticleMesh->RenderData->LODResources[LODIndex];
+		// Doesn't seem to work for some reason. See comment in FDynamicMeshEmitterData::GetMeshLODIndexFromProxy()
+		// const int32 LODIndex = FMath::Max<int32>((int32)MeshRenderData->CurrentFirstLODIdx, MeshMinimumLOD);
+		int32 LODIndex = (int32)MeshRenderData->CurrentFirstLODIdx;
+		while (LODIndex < MeshRenderData->LODResources.Num() - 1 && !MeshRenderData->LODResources[LODIndex].GetNumVertices())
+		{
+			++LODIndex;
+		}
+		const FStaticMeshLODResources& LODModel = MeshRenderData->LODResources[LODIndex];
 
 		check(BaseMaterials_GT.Num() == LODModel.Sections.Num());
 
