Fix _API placement
#rnx

diff --git a/Engine/Plugins/FX/Niagara/Source/NiagaraEditor/Public/NiagaraEditorModule.h b/Engine/Plugins/FX/Niagara/Source/NiagaraEditor/Public/NiagaraEditorModule.h
--- a/Engine/Plugins/FX/Niagara/Source/NiagaraEditor/Public/NiagaraEditorModule.h
+++ b/Engine/Plugins/FX/Niagara/Source/NiagaraEditor/Public/NiagaraEditorModule.h
@@ -25,7 +25,7 @@ struct IConsoleCommand;
 DECLARE_STATS_GROUP(TEXT("Niagara Editor"), STATGROUP_NiagaraEditor, STATCAT_Advanced);
 
 /* Defines methods for allowing external modules to supply widgets to the core editor module. */
-NIAGARAEDITOR_API class INiagaraEditorWidgetProvider
+class NIAGARAEDITOR_API INiagaraEditorWidgetProvider
 {
 public:
 	virtual TSharedRef<SWidget> CreateStackView(UNiagaraStackViewModel& StackViewModel) = 0;
