Now handling reset of GPU systems and ensuring that we don't propagate old data.

#tests engine tests pass except for known issue in static switch tests

diff --git a/Engine/Plugins/FX/Niagara/Shaders/Private/NiagaraEmitterInstanceShader.usf b/Engine/Plugins/FX/Niagara/Shaders/Private/NiagaraEmitterInstanceShader.usf
--- a/Engine/Plugins/FX/Niagara/Shaders/Private/NiagaraEmitterInstanceShader.usf
+++ b/Engine/Plugins/FX/Niagara/Shaders/Private/NiagaraEmitterInstanceShader.usf
@@ -429,7 +429,14 @@ void SimulateMainComputeCS(
 	*/
 
 	uint InstanceID = UpdateStartInstance + DispatchThreadId.x;
-	GSpawnStartInstance = RWInstanceCounts[ReadInstanceCountOffset];				// needed by ExecIndex()
+	if (ReadInstanceCountOffset == 0xFFFFFFFF)
+	{
+		GSpawnStartInstance = 0;
+	}
+	else
+	{
+		GSpawnStartInstance = RWInstanceCounts[ReadInstanceCountOffset];				// needed by ExecIndex()
+	}
 	int MaxInstances = GSpawnStartInstance + SpawnedInstances;
 
 	if(InstanceID < UpdateStartInstance + MaxInstances )
diff --git a/Engine/Plugins/FX/Niagara/Source/Niagara/Private/NiagaraEmitterInstanceBatcher.cpp b/Engine/Plugins/FX/Niagara/Source/Niagara/Private/NiagaraEmitterInstanceBatcher.cpp
--- a/Engine/Plugins/FX/Niagara/Source/Niagara/Private/NiagaraEmitterInstanceBatcher.cpp
+++ b/Engine/Plugins/FX/Niagara/Source/Niagara/Private/NiagaraEmitterInstanceBatcher.cpp
@@ -955,7 +955,7 @@ void NiagaraEmitterInstanceBatcher::Run(const FNiagaraGPUSystemTick& Tick, const
 	{
 		RHICmdList.TransitionResource(EResourceTransitionAccess::ERWNoBarrier, EResourceTransitionPipeline::EComputeToCompute, GPUInstanceCounterManager.GetInstanceCountBuffer().UAV);
 		RHICmdList.SetUAVParameter(Shader->GetComputeShader(), Shader->InstanceCountsParam.GetUAVIndex(), GPUInstanceCounterManager.GetInstanceCountBuffer().UAV);
-		const uint32 ReadOffset = CurrentData.GetGPUInstanceCountBufferOffset();
+		const uint32 ReadOffset = Tick.bNeedsReset ? INDEX_NONE : CurrentData.GetGPUInstanceCountBufferOffset();
 		const uint32 WriteOffset = DestinationData.GetGPUInstanceCountBufferOffset();
 		RHICmdList.SetShaderParameter(Shader->GetComputeShader(), Shader->ReadInstanceCountOffsetParam.GetBufferIndex(), Shader->ReadInstanceCountOffsetParam.GetBaseIndex(), Shader->ReadInstanceCountOffsetParam.GetNumBytes(), &ReadOffset);
 		RHICmdList.SetShaderParameter(Shader->GetComputeShader(), Shader->WriteInstanceCountOffsetParam.GetBufferIndex(), Shader->WriteInstanceCountOffsetParam.GetBaseIndex(), Shader->WriteInstanceCountOffsetParam.GetNumBytes(), &WriteOffset);
diff --git a/Engine/Plugins/FX/Niagara/Source/NiagaraCore/Private/NiagaraCustomVersion.cpp b/Engine/Plugins/FX/Niagara/Source/NiagaraCore/Private/NiagaraCustomVersion.cpp
--- a/Engine/Plugins/FX/Niagara/Source/NiagaraCore/Private/NiagaraCustomVersion.cpp
+++ b/Engine/Plugins/FX/Niagara/Source/NiagaraCore/Private/NiagaraCustomVersion.cpp
@@ -8,4 +8,4 @@ const FGuid FNiagaraCustomVersion::GUID(0xFCF57AFA, 0x50764283, 0xB9A9E658, 0xFF
 // Register the custom version with core
 FCustomVersionRegistration GRegisterNiagaraCustomVersion(FNiagaraCustomVersion::GUID, FNiagaraCustomVersion::LatestVersion, TEXT("NiagaraVer"));
 
-const FGuid FNiagaraCustomVersion::LatestScriptCompileVersion(0xC543A0CC, 0xD850F04D, 0xBBD65664, 0xF181DFD1);
\ No newline at end of file
+const FGuid FNiagaraCustomVersion::LatestScriptCompileVersion(0x18B0ACF1, 0x208646CD, 0x91C04B94, 0x3717A7A6 );
\ No newline at end of file
