Niagara - Add static switch support to emitter inheritance.

diff --git a/Engine/Plugins/FX/Niagara/Source/NiagaraEditor/Private/NiagaraScriptMergeManager.cpp b/Engine/Plugins/FX/Niagara/Source/NiagaraEditor/Private/NiagaraScriptMergeManager.cpp
--- a/Engine/Plugins/FX/Niagara/Source/NiagaraEditor/Private/NiagaraScriptMergeManager.cpp
+++ b/Engine/Plugins/FX/Niagara/Source/NiagaraEditor/Private/NiagaraScriptMergeManager.cpp
@@ -100,6 +100,18 @@ FNiagaraStackFunctionInputOverrideMergeAdapter::FNiagaraStackFunctionInputOverri
 {
 }
 
+FNiagaraStackFunctionInputOverrideMergeAdapter::FNiagaraStackFunctionInputOverrideMergeAdapter(UEdGraphPin* InStaticSwitchPin)
+	: OwningScript(nullptr)
+	, OwningFunctionCallNode(CastChecked<UNiagaraNodeFunctionCall>(InStaticSwitchPin->GetOwningNode()))
+	, InputName(InStaticSwitchPin->PinName.ToString())
+	, OverridePin(nullptr)
+	, DataValueObject(nullptr)
+	, StaticSwitchValue(InStaticSwitchPin->DefaultValue)
+{
+	const UEdGraphSchema_Niagara* NiagaraSchema = GetDefault<UEdGraphSchema_Niagara>();
+	Type = NiagaraSchema->PinToTypeDefinition(InStaticSwitchPin);
+}
+
 UNiagaraScript* FNiagaraStackFunctionInputOverrideMergeAdapter::GetOwningScript() const
 {
 	return OwningScript.Get();
@@ -160,6 +172,11 @@ TSharedPtr<FNiagaraStackFunctionMergeAdapter> FNiagaraStackFunctionInputOverride
 	return DynamicValueFunction;
 }
 
+TOptional<FString> FNiagaraStackFunctionInputOverrideMergeAdapter::GetStaticSwitchValue() const
+{
+	return StaticSwitchValue;
+}
+
 FNiagaraStackFunctionMergeAdapter::FNiagaraStackFunctionMergeAdapter(FString InEmitterUniqueName, UNiagaraScript& InOwningScript, UNiagaraNodeFunctionCall& InFunctionCallNode, int32 InStackIndex)
 {
 	UniqueEmitterName = InEmitterUniqueName;
@@ -215,6 +232,16 @@ FNiagaraStackFunctionMergeAdapter::FNiagaraStackFunctionMergeAdapter(FString InE
 			}
 		}
 	}
+
+	TArray<UEdGraphPin*> StaticSwitchPins;
+	TSet<UEdGraphPin*> StaticSwitchPinsHidden;
+	FNiagaraStackGraphUtilities::GetStackFunctionStaticSwitchPins(*FunctionCallNode.Get(), StaticSwitchPins, StaticSwitchPinsHidden);
+	for (UEdGraphPin* StaticSwitchPin : StaticSwitchPins)
+	{
+		// TODO: Only add static switch overrides when the current value is different from the default.  This requires a 
+		// refactor of the static switch default storage to use the same data format as FNiagaraVariables.
+		InputOverrides.Add(MakeShared<FNiagaraStackFunctionInputOverrideMergeAdapter>(StaticSwitchPin));
+	}
 }
 
 UNiagaraNodeFunctionCall* FNiagaraStackFunctionMergeAdapter::GetFunctionCallNode() const
@@ -1582,6 +1609,12 @@ TOptional<bool> FNiagaraScriptMergeManager::DoFunctionInputOverridesMatch(TShare
 			FunctionDiffResults.ModifiedOtherInputOverrides.Num() == 0;
 	}
 
+	// Static switch
+	if (BaseFunctionInputAdapter->GetStaticSwitchValue().IsSet() && OtherFunctionInputAdapter->GetStaticSwitchValue().IsSet())
+	{
+		return BaseFunctionInputAdapter->GetStaticSwitchValue().GetValue() == OtherFunctionInputAdapter->GetStaticSwitchValue().GetValue();
+	}
+
 	return TOptional<bool>();
 }
 
@@ -1656,6 +1689,12 @@ FNiagaraScriptMergeManager::FApplyDiffResults FNiagaraScriptMergeManager::Remove
 		Results.bSucceeded = true;
 		Results.bModifiedGraph = false;
 	}
+	else if (OverrideToRemove->GetStaticSwitchValue().IsSet())
+	{
+		// TODO: Static switches are always treated as overrides right now so removing them is a no-op.  This code should updated
+		// so that removing a static switch override sets the value back to the module default.
+		Results.bSucceeded = true;
+	}
 	else
 	{
 		Results.bSucceeded = false;
@@ -1787,6 +1826,34 @@ FNiagaraScriptMergeManager::FApplyDiffResults FNiagaraScriptMergeManager::AddInp
 			OwningScript.RapidIterationParameters.SetParameterData(SourceData, RapidIterationParameter, bAddParameterIfMissing);
 			Results.bSucceeded = true;
 		}
+		else if (OverrideToAdd->GetStaticSwitchValue().IsSet())
+		{
+			TArray<UEdGraphPin*> StaticSwitchPins;
+			TSet<UEdGraphPin*> StaticSwitchPinsHidden;
+			FNiagaraStackGraphUtilities::GetStackFunctionStaticSwitchPins(TargetFunctionCall, StaticSwitchPins, StaticSwitchPinsHidden);
+			UEdGraphPin** MatchingStaticSwitchPinPtr = StaticSwitchPins.FindByPredicate([&OverrideToAdd](UEdGraphPin* StaticSwitchPin) { return StaticSwitchPin->PinName == *OverrideToAdd->GetInputName(); });
+			if (MatchingStaticSwitchPinPtr != nullptr)
+			{
+				UEdGraphPin* MatchingStaticSwitchPin = *MatchingStaticSwitchPinPtr;
+				const UEdGraphSchema_Niagara* NiagaraSchema = GetDefault<UEdGraphSchema_Niagara>();
+				FNiagaraTypeDefinition SwitchType = NiagaraSchema->PinToTypeDefinition(MatchingStaticSwitchPin);
+				if (SwitchType == OverrideToAdd->GetType())
+				{
+					MatchingStaticSwitchPin->DefaultValue = OverrideToAdd->GetStaticSwitchValue().GetValue();
+					Results.bSucceeded = true;
+				}
+				else
+				{
+					Results.bSucceeded = false;
+					Results.ErrorMessages.Add(LOCTEXT("AddStaticInputOverrideFailedWrongType", "Failed to add static switch input override because a the type of the pin matched by name did not match."));
+				}
+			}
+			else
+			{
+				Results.bSucceeded = false;
+				Results.ErrorMessages.Add(LOCTEXT("AddStaticInputOverrideFailedNotFound", "Failed to add static switch input override because a matching pin could not be found."));
+			}
+		}
 		else
 		{
 			Results.bSucceeded = false;
diff --git a/Engine/Plugins/FX/Niagara/Source/NiagaraEditor/Private/NiagaraScriptMergeManager.h b/Engine/Plugins/FX/Niagara/Source/NiagaraEditor/Private/NiagaraScriptMergeManager.h
--- a/Engine/Plugins/FX/Niagara/Source/NiagaraEditor/Private/NiagaraScriptMergeManager.h
+++ b/Engine/Plugins/FX/Niagara/Source/NiagaraEditor/Private/NiagaraScriptMergeManager.h
@@ -43,6 +43,8 @@ public:
 		FString InInputName,
 		FNiagaraVariable InRapidIterationParameter);
 
+	FNiagaraStackFunctionInputOverrideMergeAdapter(UEdGraphPin* InStaticSwitchPin);
+
 	UNiagaraScript* GetOwningScript() const;
 	UNiagaraNodeFunctionCall* GetOwningFunctionCall() const;
 	FString GetInputName() const;
@@ -56,6 +58,7 @@ public:
 	TOptional<FNiagaraParameterHandle> GetLinkedValueHandle() const;
 	UNiagaraDataInterface* GetDataValueObject() const;
 	TSharedPtr<FNiagaraStackFunctionMergeAdapter> GetDynamicValueFunction() const;
+	TOptional<FString> GetStaticSwitchValue() const;
 
 private:
 	FString UniqueEmitterName;
@@ -72,6 +75,7 @@ private:
 	TOptional<FNiagaraParameterHandle> LinkedValueHandle;
 	UNiagaraDataInterface* DataValueObject;
 	TSharedPtr<FNiagaraStackFunctionMergeAdapter> DynamicValueFunction;
+	TOptional<FString> StaticSwitchValue;
 
 	FGuid OverrideValueNodePersistentId;
 };
