Fix ribbon side not aligning with input side vector

diff --git a/Engine/Plugins/FX/Niagara/Shaders/Private/NiagaraRibbonVertexFactory.ush b/Engine/Plugins/FX/Niagara/Shaders/Private/NiagaraRibbonVertexFactory.ush
--- a/Engine/Plugins/FX/Niagara/Shaders/Private/NiagaraRibbonVertexFactory.ush
+++ b/Engine/Plugins/FX/Niagara/Shaders/Private/NiagaraRibbonVertexFactory.ush
@@ -571,10 +571,10 @@ void GetNiagaraParticleTangents(in FInterpVtxData InterpVtxData, in out FVertexF
 		Intermediates.Facing = normalize(ResolvedView.WorldCameraOrigin - Intermediates.Position);
 	}
 
+	BRANCH
     if (FacingMode == 2 /* CustomSide */)
 	{
-		Intermediates.SideVector = Intermediates.Facing;
-		Intermediates.SideVector.z = -Intermediates.Facing.z;
+		Intermediates.SideVector = -Intermediates.Facing;
 		Intermediates.Right = normalize(Intermediates.SideVector - dot(Intermediates.SideVector, Intermediates.Direction * Intermediates.Direction));
 		Intermediates.Facing = normalize(cross(Intermediates.Direction, Intermediates.Right));
 	}
