Fix constant factor tessellation parameter

diff --git a/Engine/Plugins/FX/Niagara/Source/Niagara/Private/NiagaraRendererRibbons.cpp b/Engine/Plugins/FX/Niagara/Source/Niagara/Private/NiagaraRendererRibbons.cpp
--- a/Engine/Plugins/FX/Niagara/Source/Niagara/Private/NiagaraRendererRibbons.cpp
+++ b/Engine/Plugins/FX/Niagara/Source/Niagara/Private/NiagaraRendererRibbons.cpp
@@ -359,7 +359,7 @@ void FNiagaraRendererRibbons::GetDynamicMeshElements(const TArray<const FSceneVi
 		TessellationFactor = CustomTessellationFactor;
 		bUseConstantFactor = bCustomUseConstantFactor;
 		TessellationMinAngle = CustomTessellationMinAngle;
-		ScreenPercentage = bCustomUseScreenSpace ? GNiagaraRibbonTessellationScreenPercentage : 0.f;
+		ScreenPercentage = bCustomUseScreenSpace && !bUseConstantFactor ? GNiagaraRibbonTessellationScreenPercentage : 0.f;
 		break;
 	case ENiagaraRibbonTessellationMode::Disabled:
 		TessellationFactor = 1;
@@ -384,14 +384,12 @@ void FNiagaraRendererRibbons::GetDynamicMeshElements(const TArray<const FSceneVi
 			{
 				const float MinTesselation = [&]
 				{
-					if (TessellationMinAngle == 0.f) { return 0.f; }
-					if (!bUseConstantFactor)
-					{
-						return FMath::Max<float>(1.f, FMath::Max(TessellationTwistAngle, TessellationAngle) / FMath::Max<float>(SMALL_NUMBER, TessellationMinAngle));
+					if (TessellationMinAngle == 0.f || bUseConstantFactor) {
+						return static_cast<float>(TessellationFactor);
 					}
 					else
 					{
-						return static_cast<float>(TessellationFactor);
+						return FMath::Max<float>(1.f, FMath::Max(TessellationTwistAngle, TessellationAngle) / FMath::Max<float>(SMALL_NUMBER, TessellationMinAngle));
 					}
 				}(); 
 				const float ViewDistance = SceneProxy->GetBounds().ComputeSquaredDistanceFromBoxToPoint(ViewOriginForDistanceCulling);
