Update facing mode tooltip with warning about using twist and side together

diff --git a/Engine/Plugins/FX/Niagara/Source/Niagara/Public/NiagaraRibbonRendererProperties.h b/Engine/Plugins/FX/Niagara/Source/Niagara/Public/NiagaraRibbonRendererProperties.h
--- a/Engine/Plugins/FX/Niagara/Source/Niagara/Public/NiagaraRibbonRendererProperties.h
+++ b/Engine/Plugins/FX/Niagara/Source/Niagara/Public/NiagaraRibbonRendererProperties.h
@@ -19,7 +19,9 @@ enum class ENiagaraRibbonFacingMode : uint8
 	/** Use Particles.RibbonFacing as the facing vector. */
 	Custom,
 
-	/** Use Particles.RibbonFacing as the side vector, and calculate the facing vector from that. */
+	/** Use Particles.RibbonFacing as the side vector, and calculate the facing vector from that.
+	 *  Using ribbon twist with this mode is NOT supported.
+	 */
 	CustomSideVector
 };
 
