Fix for badly sorted Niagara GPU particles

diff --git a/Engine/Plugins/FX/Niagara/Source/Niagara/Private/NiagaraEmitterInstanceBatcher.cpp b/Engine/Plugins/FX/Niagara/Source/Niagara/Private/NiagaraEmitterInstanceBatcher.cpp
--- a/Engine/Plugins/FX/Niagara/Source/Niagara/Private/NiagaraEmitterInstanceBatcher.cpp
+++ b/Engine/Plugins/FX/Niagara/Source/Niagara/Private/NiagaraEmitterInstanceBatcher.cpp
@@ -290,7 +290,7 @@ void NiagaraEmitterInstanceBatcher::DispatchAllOnCompute(FOverlappableTicks& Ove
 					Run<false>(*Tick, &Instance, UpdateStartInstance, Instance.DestinationData->GetNumInstances(), Context->GPUScript_RT->GetShader(), RHICmdList, ViewUniformBuffer);
 
 					FNiagaraDataBuffer* CurrentData = Instance.CurrentData;
-					if (bSetReadback && Tick->bIsFinalTick && Index == DispatchCount - 1)
+					if (bSetReadback && Tick->bIsFinalTick)
 					{
 						// Now that the current data is not required anymore, stage it for readback.
 						if (CurrentData->GetNumInstances() && Context->EmitterInstanceReadback.GPUCountOffset == INDEX_NONE && CurrentData->GetGPUInstanceCountBufferOffset() != INDEX_NONE)
@@ -411,7 +411,7 @@ void NiagaraEmitterInstanceBatcher::ExecuteAll(FRHICommandList &RHICmdList, FRHI
 			for (uint32 i = 0; i < DispatchCount; i++)
 			{
 				FNiagaraComputeInstanceData& Instance = Instances[i];
-				TickSingle(Tick, &Instance, RHICmdList, ViewUniformBuffer, bSetReadback && i == DispatchCount - 1); // Only readback last count.
+				TickSingle(Tick, &Instance, RHICmdList, ViewUniformBuffer, bSetReadback);
 			}
 		}
 	}
@@ -510,10 +510,14 @@ void NiagaraEmitterInstanceBatcher::PreInitViews(FRHICommandListImmediate& RHICm
 		// Cancel any pending readback if the emitter is resetting.
 		if (Tick.bNeedsReset)
 		{
-			FNiagaraComputeExecutionContext* Context = Tick.GetInstanceData()->Context;
-			if (Context)
+			FNiagaraComputeInstanceData* Instances = Tick.GetInstanceData();
+			for (uint32 InstanceIndex = 0; InstanceIndex < Tick.Count; ++InstanceIndex)
 			{
-				GPUInstanceCounterManager.FreeEntry(Context->EmitterInstanceReadback.GPUCountOffset);
+				FNiagaraComputeExecutionContext* Context = Instances[InstanceIndex].Context;
+				if (Context)
+				{
+					GPUInstanceCounterManager.FreeEntry(Context->EmitterInstanceReadback.GPUCountOffset);
+				}
 			}
 		}
 	}
@@ -527,25 +531,29 @@ void NiagaraEmitterInstanceBatcher::PreInitViews(FRHICommandListImmediate& RHICm
 		{
 			for (FNiagaraGPUSystemTick& Tick : Ticks_RT)
 			{
-				FNiagaraComputeExecutionContext* Context = Tick.GetInstanceData()->Context;
-				if (Context && Context->EmitterInstanceReadback.GPUCountOffset != INDEX_NONE)
+				FNiagaraComputeInstanceData* Instances = Tick.GetInstanceData();
+				for (uint32 InstanceIndex = 0; InstanceIndex < Tick.Count; ++InstanceIndex)
 				{
-					check(Context->MainDataSet);
-					FNiagaraDataBuffer* CurrentData = Context->MainDataSet->GetCurrentData();
-					if (CurrentData)
+					FNiagaraComputeExecutionContext* Context = Instances[InstanceIndex].Context;
+					if (Context && Context->EmitterInstanceReadback.GPUCountOffset != INDEX_NONE)
 					{
-						const uint32 DeadInstanceCount = Context->EmitterInstanceReadback.CPUCount - Counts[Context->EmitterInstanceReadback.GPUCountOffset];
-
-						// This will communicate the particle counts to the game thread. If DeadInstanceCount equals CurrentData->GetNumInstances() the game thread will know that the emitter has completed.
-						if (DeadInstanceCount <= CurrentData->GetNumInstances())
+						check(Context->MainDataSet);
+						FNiagaraDataBuffer* CurrentData = Context->MainDataSet->GetCurrentData();
+						if (CurrentData)
 						{
-							CurrentData->SetNumInstances(CurrentData->GetNumInstances() - DeadInstanceCount);
+							const uint32 DeadInstanceCount = Context->EmitterInstanceReadback.CPUCount - Counts[Context->EmitterInstanceReadback.GPUCountOffset];
+
+							// This will communicate the particle counts to the game thread. If DeadInstanceCount equals CurrentData->GetNumInstances() the game thread will know that the emitter has completed.
+							if (DeadInstanceCount <= CurrentData->GetNumInstances())
+							{
+								CurrentData->SetNumInstances(CurrentData->GetNumInstances() - DeadInstanceCount); 
+							}
 						}
-					}
 
-					// Now release the readback since another one will be enqueued in the tick.
-					// Also prevents processing the same data again.
-					GPUInstanceCounterManager.FreeEntry(Context->EmitterInstanceReadback.GPUCountOffset);
+						// Now release the readback since another one will be enqueued in the tick.
+						// Also prevents processing the same data again.
+						GPUInstanceCounterManager.FreeEntry(Context->EmitterInstanceReadback.GPUCountOffset);
+					}
 				}
 			}
 			// Readback is only valid for one frame, so that any newly allocated instance count
@@ -708,6 +716,11 @@ void NiagaraEmitterInstanceBatcher::SortGPUParticles(FRHICommandListImmediate& R
 				++NumFramesRequiringShrinking;
 			}
 		}
+		else // Reset counter since we are not in a shrinking situation anymore.
+		{
+			NumFramesRequiringShrinking = 0;
+		}
+
 	}
 	else // If the are no sort task, we don't need any of the sort buffers.
 	{
diff --git a/Engine/Plugins/FX/Niagara/Source/NiagaraVertexFactories/Private/NiagaraSortingGPU.cpp b/Engine/Plugins/FX/Niagara/Source/NiagaraVertexFactories/Private/NiagaraSortingGPU.cpp
--- a/Engine/Plugins/FX/Niagara/Source/NiagaraVertexFactories/Private/NiagaraSortingGPU.cpp
+++ b/Engine/Plugins/FX/Niagara/Source/NiagaraVertexFactories/Private/NiagaraSortingGPU.cpp
@@ -56,6 +56,24 @@ static FAutoConsoleVariableRef CVarNiagaraGPUSortingFrameCountBeforeBufferShrink
 	ECVF_Default
 );
 
+class FNiagaraSortingDummyUAV : public FRenderResource
+{
+public:
+	FRWBuffer Buffer;
+
+	virtual void InitRHI() override
+	{
+		Buffer.Initialize(sizeof(int32), 1, EPixelFormat::PF_R32_SINT, BUF_Static, TEXT("FNiagaraSortingDummyUAV"));
+	}
+
+	virtual void ReleaseRHI() override
+	{
+		Buffer.Release();
+	}
+}; 
+static TGlobalResource<FNiagaraSortingDummyUAV> NiagaraSortingDummyUAV[NIAGARA_COPY_BUFFER_BUFFER_COUNT];
+
+
 IMPLEMENT_GLOBAL_SHADER(FNiagaraSortKeyGenCS, "/Plugin/FX/Niagara/Private/NiagaraSortKeyGen.usf", "GenerateParticleSortKeys", SF_Compute);
 
 void FNiagaraSortKeyGenCS::ModifyCompilationEnvironment(const FGlobalShaderPermutationParameters& Parameters, FShaderCompilerEnvironment& OutEnvironment)
@@ -206,10 +224,9 @@ void FNiagaraCopyIntBufferRegionCS::SetParameters(
 		CopyParamsValue[Index + 1] = InUsedIndexCounts[Index];
 	}
 
-	// Workaround for Vulkan driver crashes: replicate the last UAV over
 	for (int32 Index = DestCount; Index < NIAGARA_COPY_BUFFER_BUFFER_COUNT; ++Index)
 	{
-		RHICmdList.SetUAVParameter(ComputeShaderRHI, DestData[Index].GetBaseIndex(), InDestDatas[DestCount - 1]);
+		RHICmdList.SetUAVParameter(ComputeShaderRHI, DestData[Index].GetBaseIndex(), NiagaraSortingDummyUAV[Index].Buffer.UAV);
 	}
 
 	RHICmdList.SetShaderParameter(ComputeShaderRHI, CopyParams.GetBufferIndex(), CopyParams.GetBaseIndex(), CopyParams.GetNumBytes(), &CopyParamsValue);
