Foundation work for Niagara System Overview Graph.

diff --git a/Engine/Plugins/FX/Niagara/Source/Niagara/Classes/NiagaraSystem.h b/Engine/Plugins/FX/Niagara/Source/Niagara/Classes/NiagaraSystem.h
--- a/Engine/Plugins/FX/Niagara/Source/Niagara/Classes/NiagaraSystem.h
+++ b/Engine/Plugins/FX/Niagara/Source/Niagara/Classes/NiagaraSystem.h
@@ -230,6 +230,10 @@ public:
 
 	TStatId GetStatID(bool bGameThread, bool bConcurrent)const;
 
+#if WITH_EDITORONLY_DATA
+	UEdGraph* GetSystemOverviewGraph() const;
+#endif 
+
 private:
 #if WITH_EDITORONLY_DATA
 	bool QueryCompileComplete(bool bWait, bool bDoPost, bool bDoNotApply = false);
@@ -310,6 +314,19 @@ protected:
 	TArray<FName> UserDINamesReadInSystemScripts;
 
 	void GenerateStatID();
+
+private:
+#if WITH_EDITORONLY_DATA
+	/** Graph presenting overview of the current system and its emitters. */
+	UPROPERTY()
+	UEdGraph* SystemOverviewGraph;
+
+	void InitSystemOverviewGraph() const;
+	void SystemOverviewHandleAdded(const FGuid AddedHandleGuid) const;
+	void SystemOverviewHandlesRemoved() const;
+#endif
+
+protected:
 #if STATS
 	TStatId StatID_GT;
 	TStatId StatID_GT_CNC;
diff --git a/Engine/Plugins/FX/Niagara/Source/Niagara/Private/NiagaraSystem.cpp b/Engine/Plugins/FX/Niagara/Source/Niagara/Private/NiagaraSystem.cpp
--- a/Engine/Plugins/FX/Niagara/Source/Niagara/Private/NiagaraSystem.cpp
+++ b/Engine/Plugins/FX/Niagara/Source/Niagara/Private/NiagaraSystem.cpp
@@ -21,6 +21,12 @@
 #include "DerivedDataCacheInterface.h"
 #endif
 
+#if WITH_EDITORONLY_DATA
+#include "Widgets/NiagaraOverviewNodeStackItem.h"
+
+static const float SystemOverviewNodePadding = 300.0f;
+#endif
+
 DECLARE_CYCLE_STAT(TEXT("Niagara - System - Precompile"), STAT_Niagara_System_Precompile, STATGROUP_Niagara);
 DECLARE_CYCLE_STAT(TEXT("Niagara - System - CompileScript"), STAT_Niagara_System_CompileScript, STATGROUP_Niagara);
 DECLARE_CYCLE_STAT(TEXT("Niagara - System - CompileScript_ResetAfter"), STAT_Niagara_System_CompileScriptResetAfter, STATGROUP_Niagara);
@@ -90,6 +96,11 @@ void UNiagaraSystem::PostInitProperties()
 
 		SystemUpdateScript = NewObject<UNiagaraScript>(this, "SystemUpdateScript", RF_Transactional);
 		SystemUpdateScript->SetUsage(ENiagaraScriptUsage::SystemUpdateScript);
+
+#if WITH_EDITORONLY_DATA
+		SystemOverviewGraph = NewObject<UEdGraph>(this, "NiagaraOverview", RF_Transactional);
+		SystemOverviewGraph->Schema = UEdGraphSchema::StaticClass();
+#endif
 	}
 
 	GenerateStatID();
@@ -384,6 +395,13 @@ void UNiagaraSystem::PostLoad()
 		UpdateScript->RapidIterationParameters.DumpParameters();
 	}
 #endif
+
+	if (SystemOverviewGraph == nullptr)
+	{
+		SystemOverviewGraph = NewObject<UEdGraph>(this, "NiagaraOverview", RF_Transactional);
+		SystemOverviewGraph->Schema = UEdGraphSchema::StaticClass();
+		InitSystemOverviewGraph();
+	}
 #endif
 }
 
@@ -566,6 +584,7 @@ FNiagaraEmitterHandle UNiagaraSystem::AddEmitterHandle(UNiagaraEmitter& InEmitte
 	}
 	EmitterHandles.Add(EmitterHandle);
 	RefreshSystemParametersFromEmitter(EmitterHandle);
+	SystemOverviewHandleAdded(EmitterHandle.GetId());
 	return EmitterHandle;
 }
 
@@ -576,6 +595,7 @@ FNiagaraEmitterHandle UNiagaraSystem::DuplicateEmitterHandle(const FNiagaraEmitt
 	EmitterHandle.SetIsEnabled(EmitterHandleToDuplicate.GetIsEnabled());
 	EmitterHandles.Add(EmitterHandle);
 	RefreshSystemParametersFromEmitter(EmitterHandle);
+	SystemOverviewHandleAdded(EmitterHandle.GetId());
 	return EmitterHandle;
 }
 
@@ -585,6 +605,7 @@ void UNiagaraSystem::RemoveEmitterHandle(const FNiagaraEmitterHandle& EmitterHan
 	RemoveSystemParametersForEmitter(EmitterHandleToDelete);
 	auto RemovePredicate = [&](const FNiagaraEmitterHandle& EmitterHandle) { return EmitterHandle.GetId() == EmitterHandleToDelete.GetId(); };
 	EmitterHandles.RemoveAll(RemovePredicate);
+	SystemOverviewHandlesRemoved();
 }
 
 void UNiagaraSystem::RemoveEmitterHandlesById(const TSet<FGuid>& HandlesToRemove)
@@ -596,6 +617,7 @@ void UNiagaraSystem::RemoveEmitterHandlesById(const TSet<FGuid>& HandlesToRemove
 	EmitterHandles.RemoveAll(RemovePredicate);
 
 	InitEmitterSpawnAttributes();
+	SystemOverviewHandlesRemoved();
 }
 #endif
 
@@ -656,6 +678,13 @@ bool UNiagaraSystem::PollForCompilationComplete()
 	return true;
 }
 
+#if WITH_EDITORONLY_DATA
+UEdGraph* UNiagaraSystem::GetSystemOverviewGraph() const
+{
+	return SystemOverviewGraph;
+}
+#endif
+
 bool UNiagaraSystem::QueryCompileComplete(bool bWait, bool bDoPost, bool bDoNotApply)
 {
 	if (ActiveCompilations.Num() > 0)
@@ -1054,4 +1083,82 @@ void UNiagaraSystem::GenerateStatID()
 	StatID_RT = FDynamicStats::CreateStatId<FStatGroup_STATGROUP_NiagaraSystems>(GetPathName() + TEXT("[RT]"));
 	StatID_RT_CNC = FDynamicStats::CreateStatId<FStatGroup_STATGROUP_NiagaraSystems>(GetPathName() + TEXT("[RT_CNC]"));
 #endif
-}
\ No newline at end of file
+}
+
+#if WITH_EDITORONLY_DATA
+void UNiagaraSystem::InitSystemOverviewGraph() const
+{
+	float NextNodePosX = 0.0;
+
+	FGraphNodeCreator<UNiagaraOverviewNodeStackItem> SystemOverviewNodeCreator(*SystemOverviewGraph);
+	UNiagaraOverviewNodeStackItem* SystemOverviewNode = SystemOverviewNodeCreator.CreateNode(false);
+	SystemOverviewNode->Initialize(this);
+
+	SystemOverviewNode->NodePosX = NextNodePosX;
+	SystemOverviewNode->NodePosY = 0.0;
+	NextNodePosX += SystemOverviewNodePadding;
+
+	SystemOverviewNodeCreator.Finalize();
+
+	for (const FNiagaraEmitterHandle& Handle : EmitterHandles)
+	{
+		FGraphNodeCreator<UNiagaraOverviewNodeStackItem> EmitterOverviewNodeCreator(*SystemOverviewGraph);
+		UNiagaraOverviewNodeStackItem* EmitterOverviewNode = EmitterOverviewNodeCreator.CreateNode(false);
+		EmitterOverviewNode->Initialize(this, Handle.GetId());
+
+		EmitterOverviewNode->NodePosX = NextNodePosX;
+		EmitterOverviewNode->NodePosY = 0.0;
+		NextNodePosX += SystemOverviewNodePadding;
+
+		EmitterOverviewNodeCreator.Finalize();
+	}
+}
+
+void UNiagaraSystem::SystemOverviewHandleAdded(const FGuid AddedHandleGuid) const
+{
+	auto GetGoodPlaceForNewOverviewNode = [](UEdGraph* SystemOverviewGraph, const TArray<FNiagaraEmitterHandle>& EmitterHandles)->const FVector2D {
+		TArray<UNiagaraOverviewNodeStackItem*> OverviewNodes;
+		SystemOverviewGraph->GetNodesOfClass<UNiagaraOverviewNodeStackItem>(OverviewNodes);
+
+		// Find the last created Emitter overview node location and place next to that. 
+		if (EmitterHandles.Num() > 1)
+		{
+			const FGuid PreviousEmitterHandleGuid = EmitterHandles.Last(1).GetId();
+			for (const UNiagaraOverviewNodeStackItem* PerOverviewNode : OverviewNodes)
+			{
+				if (PreviousEmitterHandleGuid == PerOverviewNode->GetEmitterHandleGuid())
+				{
+					return FVector2D(PerOverviewNode->NodePosX + SystemOverviewNodePadding, PerOverviewNode->NodePosY);
+				}
+			}
+		}
+		return SystemOverviewGraph->GetGoodPlaceForNewNode();
+	};
+
+	const FVector2D NewNodeLocation = GetGoodPlaceForNewOverviewNode(SystemOverviewGraph, EmitterHandles);
+
+	FGraphNodeCreator<UNiagaraOverviewNodeStackItem> EmitterOverviewNodeCreator(*SystemOverviewGraph);
+	UNiagaraOverviewNodeStackItem* EmitterOverviewNode = EmitterOverviewNodeCreator.CreateNode(false);
+	EmitterOverviewNode->Initialize(this, AddedHandleGuid);
+	
+	EmitterOverviewNode->NodePosX = NewNodeLocation.X;
+	EmitterOverviewNode->NodePosY = NewNodeLocation.Y;
+
+	EmitterOverviewNodeCreator.Finalize();
+}
+
+void UNiagaraSystem::SystemOverviewHandlesRemoved() const
+{
+	TArray<UNiagaraOverviewNodeStackItem*> CurrentOverviewNodes;
+	SystemOverviewGraph->GetNodesOfClass<UNiagaraOverviewNodeStackItem>(CurrentOverviewNodes);
+	for (UNiagaraOverviewNodeStackItem* OverviewNode : CurrentOverviewNodes)
+	{
+		const FGuid OverviewNodeGuid = OverviewNode->GetEmitterHandleGuid();
+		// If the OverviewNode's EmitterHandleGuid is valid (not representing a system) and is not in the current array of this System's EmitterHandles, that node can be removed.
+		if (OverviewNodeGuid.IsValid() && false == EmitterHandles.ContainsByPredicate([&OverviewNodeGuid](const FNiagaraEmitterHandle& Handle) {return Handle.GetId() == OverviewNodeGuid;}))
+		{
+			SystemOverviewGraph->RemoveNode(OverviewNode);
+		}
+	}
+}
+#endif
diff --git a/Engine/Plugins/FX/Niagara/Source/Niagara/Private/Widgets/NiagaraOverviewNodeStackItem.cpp b/Engine/Plugins/FX/Niagara/Source/Niagara/Private/Widgets/NiagaraOverviewNodeStackItem.cpp
new file mode 100644
--- /dev/null
+++ b/Engine/Plugins/FX/Niagara/Source/Niagara/Private/Widgets/NiagaraOverviewNodeStackItem.cpp
@@ -0,0 +1,49 @@
+// Copyright 1998-2019 Epic Games, Inc. All Rights Reserved.
+
+#include "NiagaraOverviewNodeStackItem.h"
+
+#define LOCTEXT_NAMESPACE "NiagaraOverviewNodeStackItem"
+
+UNiagaraOverviewNodeStackItem::UNiagaraOverviewNodeStackItem()
+	: OwningSystem(nullptr)
+	, EmitterHandleGuid(FGuid())
+{
+};
+
+void UNiagaraOverviewNodeStackItem::Initialize(const UNiagaraSystem* InOwningSystem)
+{
+	OwningSystem = const_cast<UNiagaraSystem*>(InOwningSystem);
+}
+
+void UNiagaraOverviewNodeStackItem::Initialize(const UNiagaraSystem* InOwningSystem, FGuid InEmitterHandleGuid)
+{
+	OwningSystem = const_cast<UNiagaraSystem*>(InOwningSystem);
+	EmitterHandleGuid = InEmitterHandleGuid;
+}
+
+const FGuid UNiagaraOverviewNodeStackItem::GetEmitterHandleGuid() const
+{
+	return EmitterHandleGuid;
+}
+
+FText UNiagaraOverviewNodeStackItem::GetNodeTitle(ENodeTitleType::Type TitleType) const
+{
+	if (EmitterHandleGuid.IsValid())
+	{
+		for (const FNiagaraEmitterHandle& Handle : OwningSystem->GetEmitterHandles())
+		{
+			if (Handle.GetId() == EmitterHandleGuid)
+			{
+				return FText::FromName(Handle.GetName());
+			}
+		}
+		ensureMsgf(false, TEXT("Failed to find matching emitter handle for existing overview node!"));
+		return LOCTEXT("UnknownEmitterName", "Unknown Emitter");
+	}
+	else
+	{
+		return FText::FromString(OwningSystem->GetName());
+	}
+}
+
+#undef LOCTEXT_NAMESPACE
diff --git a/Engine/Plugins/FX/Niagara/Source/Niagara/Private/Widgets/NiagaraOverviewNodeStackItem.h b/Engine/Plugins/FX/Niagara/Source/Niagara/Private/Widgets/NiagaraOverviewNodeStackItem.h
new file mode 100644
--- /dev/null
+++ b/Engine/Plugins/FX/Niagara/Source/Niagara/Private/Widgets/NiagaraOverviewNodeStackItem.h
@@ -0,0 +1,37 @@
+// Copyright 1998-2019 Epic Games, Inc. All Rights Reserved.
+#pragma once
+
+#include "CoreMinimal.h"
+#include "EdGraph/EdGraphNode.h"
+#include "Misc/Guid.h"
+#include "NiagaraSystem.h"
+#include "NiagaraOverviewNodeStackItem.generated.h"
+
+class UEdGraphPin;
+class INiagaraCompiler;
+struct FNiagaraGraphFunctionAliasContext;
+
+#if WITH_EDITORONLY_DATA
+UCLASS()
+class UNiagaraOverviewNodeStackItem : public UEdGraphNode
+{
+	GENERATED_BODY()
+public:
+	UNiagaraOverviewNodeStackItem();
+	void Initialize(const UNiagaraSystem* InOwningSystem);
+	void Initialize(const UNiagaraSystem* InOwningSystem, FGuid InEmitterHandleGuid);
+	const FGuid GetEmitterHandleGuid() const;
+
+	//~ Begin UEdGraphNode Interface
+	/** Gets the name of this node, shown in title bar */
+	virtual FText GetNodeTitle(ENodeTitleType::Type TitleType) const override;
+	//~ End UEdGraphNode Interface
+
+private:
+	UPROPERTY()
+	UNiagaraSystem* OwningSystem;
+
+	UPROPERTY()
+	FGuid EmitterHandleGuid;
+};
+#endif
diff --git a/Engine/Plugins/FX/Niagara/Source/NiagaraEditorWidgets/Private/NiagaraEditorWidgetsModule.cpp b/Engine/Plugins/FX/Niagara/Source/NiagaraEditorWidgets/Private/NiagaraEditorWidgetsModule.cpp
--- a/Engine/Plugins/FX/Niagara/Source/NiagaraEditorWidgets/Private/NiagaraEditorWidgetsModule.cpp
+++ b/Engine/Plugins/FX/Niagara/Source/NiagaraEditorWidgets/Private/NiagaraEditorWidgetsModule.cpp
@@ -10,6 +10,11 @@
 #include "Modules/ModuleManager.h"
 #include "PropertyEditorModule.h"
 
+#include "ViewModels/NiagaraSystemViewModel.h"
+#include "NiagaraSystem.h"
+#include "NodeFactory.h"
+#include "GraphEditor.h"
+
 IMPLEMENT_MODULE(FNiagaraEditorWidgetsModule, NiagaraEditorWidgets);
 
 FNiagaraStackCurveEditorOptions::FNiagaraStackCurveEditorOptions()
@@ -139,5 +144,9 @@ TSharedRef<SWidget> FNiagaraEditorWidgetsModule::FNiagaraEditorWidgetProvider::C
 
 TSharedRef<SWidget> FNiagaraEditorWidgetsModule::FNiagaraEditorWidgetProvider::CreateSystemOverview(TSharedRef<FNiagaraSystemViewModel> SystemViewModel)
 {
-	return SNullWidget::NullWidget;
-}
\ No newline at end of file
+	TSharedRef<SGraphEditor> GraphPanel = SNew(SGraphEditor)
+		.GraphToEdit(SystemViewModel->GetSystem().GetSystemOverviewGraph());
+
+	GraphPanel->SetNodeFactory(MakeShareable(new FGraphNodeFactory));
+	return GraphPanel;
+}
