Update Engine.Time in FNiagaraSystemSimulation.
Remove SystemInstance bindings to Simulation parameter stores.

diff --git a/Engine/Plugins/FX/Niagara/Source/Niagara/Private/NiagaraSystemInstance.cpp b/Engine/Plugins/FX/Niagara/Source/Niagara/Private/NiagaraSystemInstance.cpp
--- a/Engine/Plugins/FX/Niagara/Source/Niagara/Private/NiagaraSystemInstance.cpp
+++ b/Engine/Plugins/FX/Niagara/Source/Niagara/Private/NiagaraSystemInstance.cpp
@@ -864,8 +864,6 @@ void FNiagaraSystemInstance::BindParameters()
 	}
 
 	Component->GetOverrideParameters().Bind(&InstanceParameters);
-	InstanceParameters.Bind(&SystemSimulation->GetSpawnExecutionContext().Parameters);
-	InstanceParameters.Bind(&SystemSimulation->GetUpdateExecutionContext().Parameters);
 
 	if (SystemSimulation->GetIsSolo())
 	{
@@ -898,8 +896,6 @@ void FNiagaraSystemInstance::UnbindParameters()
 
 	if (SystemSimulation.IsValid())
 	{
-		InstanceParameters.Unbind(&SystemSimulation->GetSpawnExecutionContext().Parameters);
-		InstanceParameters.Unbind(&SystemSimulation->GetUpdateExecutionContext().Parameters);
 		if (SystemSimulation->GetIsSolo())
 		{
 			if (Component)
diff --git a/Engine/Plugins/FX/Niagara/Source/Niagara/Private/NiagaraSystemSimulation.cpp b/Engine/Plugins/FX/Niagara/Source/Niagara/Private/NiagaraSystemSimulation.cpp
--- a/Engine/Plugins/FX/Niagara/Source/Niagara/Private/NiagaraSystemSimulation.cpp
+++ b/Engine/Plugins/FX/Niagara/Source/Niagara/Private/NiagaraSystemSimulation.cpp
@@ -171,6 +171,8 @@ bool FNiagaraSystemSimulation::Init(UNiagaraSystem* InSystem, UWorld* InWorld, b
 			}
 		}
 
+		SpawnTimeParam.Init(SpawnExecContext.Parameters, SYS_PARAM_ENGINE_TIME);
+		UpdateTimeParam.Init(UpdateExecContext.Parameters, SYS_PARAM_ENGINE_TIME);
 		SpawnDeltaTimeParam.Init(SpawnExecContext.Parameters, SYS_PARAM_ENGINE_DELTA_TIME);
 		UpdateDeltaTimeParam.Init(UpdateExecContext.Parameters, SYS_PARAM_ENGINE_DELTA_TIME);
 		SpawnInvDeltaTimeParam.Init(SpawnExecContext.Parameters, SYS_PARAM_ENGINE_INV_DELTA_TIME);
@@ -284,6 +286,11 @@ bool FNiagaraSystemSimulation::Tick(float DeltaSeconds)
 		// TODO: evaluate whether or not we should have removed this from the world manager instead?
 		return false;
 	}
+	else if (World == nullptr)
+	{
+		ensureMsgf(false, TEXT("World was null while ticking System Simulation!"));
+		return false;
+	}
 
 	FScopeCycleCounter SystemStatCounter(System->GetStatID(true, false));
 
@@ -444,6 +451,9 @@ bool FNiagaraSystemSimulation::Tick(float DeltaSeconds)
 			float InvDt = 1.0f / DeltaSeconds;
 			float GlobalSpawnCountScale = INiagaraModule::GetGlobalSpawnCountScale();
 			float GlobalSystemCountScale = INiagaraModule::GetGlobalSystemCountScale();
+
+			SpawnTimeParam.SetValue(World->TimeSeconds);
+			UpdateTimeParam.SetValue(World->TimeSeconds);
 			SpawnDeltaTimeParam.SetValue(DeltaSeconds);
 			UpdateDeltaTimeParam.SetValue(DeltaSeconds);
 			SpawnInvDeltaTimeParam.SetValue(InvDt);
diff --git a/Engine/Plugins/FX/Niagara/Source/Niagara/Public/NiagaraSystemSimulation.h b/Engine/Plugins/FX/Niagara/Source/Niagara/Public/NiagaraSystemSimulation.h
--- a/Engine/Plugins/FX/Niagara/Source/Niagara/Public/NiagaraSystemSimulation.h
+++ b/Engine/Plugins/FX/Niagara/Source/Niagara/Public/NiagaraSystemSimulation.h
@@ -184,6 +184,9 @@ protected:
 	TArray<TArray<FNiagaraParameterStoreToDataSetBinding>> DataSetToEmitterEventParameters;
 
 	/** Direct bindings for Engine variables in System Spawn and Update scripts. */
+	FNiagaraParameterDirectBinding<float> SpawnTimeParam;
+	FNiagaraParameterDirectBinding<float> UpdateTimeParam;
+
 	FNiagaraParameterDirectBinding<float> SpawnDeltaTimeParam;
 	FNiagaraParameterDirectBinding<float> UpdateDeltaTimeParam;
 
diff --git a/Engine/Plugins/FX/Niagara/Source/NiagaraEditor/Private/NiagaraHlslTranslator.cpp b/Engine/Plugins/FX/Niagara/Source/NiagaraEditor/Private/NiagaraHlslTranslator.cpp
--- a/Engine/Plugins/FX/Niagara/Source/NiagaraEditor/Private/NiagaraHlslTranslator.cpp
+++ b/Engine/Plugins/FX/Niagara/Source/NiagaraEditor/Private/NiagaraHlslTranslator.cpp
@@ -1020,6 +1020,7 @@ const FNiagaraTranslateResults &FHlslNiagaraTranslator::Translate(const FNiagara
 			{
 				int32 OutputIdx = 0;
 				//ensure the interpolated spawn constants are part of the parameter set.
+				ParameterMapRegisterExternalConstantNamespaceVariable(SYS_PARAM_ENGINE_TIME, nullptr, 0, OutputIdx, nullptr);
 				ParameterMapRegisterExternalConstantNamespaceVariable(SYS_PARAM_ENGINE_DELTA_TIME, nullptr, 0, OutputIdx, nullptr);
 				ParameterMapRegisterExternalConstantNamespaceVariable(SYS_PARAM_ENGINE_INV_DELTA_TIME, nullptr, 0, OutputIdx, nullptr);
 				ParameterMapRegisterExternalConstantNamespaceVariable(SYS_PARAM_ENGINE_EXEC_COUNT, nullptr, 0, OutputIdx, nullptr);
