- Fix for data interfaces not being initialized if you deactivate a system and activate after the system has completed

diff --git a/Engine/Plugins/FX/Niagara/Source/Niagara/Private/NiagaraSystemInstance.cpp b/Engine/Plugins/FX/Niagara/Source/Niagara/Private/NiagaraSystemInstance.cpp
--- a/Engine/Plugins/FX/Niagara/Source/Niagara/Private/NiagaraSystemInstance.cpp
+++ b/Engine/Plugins/FX/Niagara/Source/Niagara/Private/NiagaraSystemInstance.cpp
@@ -490,11 +490,14 @@ void FNiagaraSystemInstance::Reset(FNiagaraSystemInstance::EResetMode Mode, bool
 	{
 		//UE_LOG(LogNiagara, Log, TEXT("FNiagaraSystemInstance::Reset true"));
 		ResetInternal(true);
+
+		bBindParams |= !IsComplete() == false;
 	}
 	else if (Mode == EResetMode::ReInit)
 	{
 		//UE_LOG(LogNiagara, Log, TEXT("FNiagaraSystemInstance::ReInit"));
 		ReInitInternal();
+
 		// If the system was reinitialized successfully than we need to force a rebind of the parameters.
 		bBindParams = IsComplete() == false;
 	}
@@ -508,7 +511,7 @@ void FNiagaraSystemInstance::Reset(FNiagaraSystemInstance::EResetMode Mode, bool
 	SetActualExecutionState(ENiagaraExecutionState::Active);
 
 	// We avoid calling InitDataInterfaces in the ResetSystem path so as to not clear out the System's DI on this frame.
-	if (Mode == EResetMode::ResetAll || Mode == EResetMode::ReInit)
+	if (bBindParams)
 	{
 		InitDataInterfaces();
 	}
