Improved tooltips for custom sorting

diff --git a/Engine/Plugins/FX/Niagara/Source/Niagara/Classes/NiagaraGPUSortInfo.h b/Engine/Plugins/FX/Niagara/Source/Niagara/Classes/NiagaraGPUSortInfo.h
--- a/Engine/Plugins/FX/Niagara/Source/Niagara/Classes/NiagaraGPUSortInfo.h
+++ b/Engine/Plugins/FX/Niagara/Source/Niagara/Classes/NiagaraGPUSortInfo.h
@@ -19,9 +19,9 @@ enum class ENiagaraSortMode : uint8
 	ViewDepth,
 	/** Sort by distance to the camera's origin.*/
 	ViewDistance,
-	/** Custom sorting according to a per particle attribute. Lower values are rendered before higher values. */
+	/** Custom sorting according to a per particle attribute. Which attribute is defined by the renderer's CustomSortingBinding which defaults to Particles.NormalizedAge. Lower values are rendered before higher values. */
 	CustomAscending,
-	/** Custom sorting according to a per particle attribute. Higher values are rendered before lower values. */
+	/** Custom sorting according to a per particle attribute. Which attribute is defined by the renderer's CustomSortingBinding which defaults to Particles.NormalizedAge. Higher values are rendered before lower values. */
 	CustomDecending,
 };
 
diff --git a/Engine/Plugins/FX/Niagara/Source/Niagara/Public/NiagaraSpriteRendererProperties.h b/Engine/Plugins/FX/Niagara/Source/Niagara/Public/NiagaraSpriteRendererProperties.h
--- a/Engine/Plugins/FX/Niagara/Source/Niagara/Public/NiagaraSpriteRendererProperties.h
+++ b/Engine/Plugins/FX/Niagara/Source/Niagara/Public/NiagaraSpriteRendererProperties.h
@@ -182,7 +182,7 @@ public:
 	UPROPERTY(EditAnywhere, Category = "Bindings")
 	FNiagaraVariableAttributeBinding MaterialRandomBinding;
 
-	/** Which attribute should we use for custom sorting? */
+	/** Which attribute should we use for custom sorting? Defaults to Particles.NormalizedAge. */
 	UPROPERTY(EditAnywhere, Category = "Bindings")
 	FNiagaraVariableAttributeBinding CustomSortingBinding;
 
