Adds the option to convert all Numeric Niagara graph nodes simultaneously when right clicking a node in the script editor (as opposed to just right clicking a single pin).

diff --git a/Engine/Plugins/FX/Niagara/Source/NiagaraEditor/Private/EdGraphSchema_Niagara.cpp b/Engine/Plugins/FX/Niagara/Source/NiagaraEditor/Private/EdGraphSchema_Niagara.cpp
--- a/Engine/Plugins/FX/Niagara/Source/NiagaraEditor/Private/EdGraphSchema_Niagara.cpp
+++ b/Engine/Plugins/FX/Niagara/Source/NiagaraEditor/Private/EdGraphSchema_Niagara.cpp
@@ -1440,6 +1440,24 @@ void UEdGraphSchema_Niagara::GetBreakLinkToSubMenuActions(class FMenuBuilder& Me
 	}
 }
 
+void UEdGraphSchema_Niagara::ConvertNumericPinToTypeAll(UNiagaraNode* InNode, FNiagaraTypeDefinition TypeDef)
+{
+	if (InNode)
+	{
+		for (auto Pin : InNode->Pins)
+		{
+			if (PinToTypeDefinition(Pin) == FNiagaraTypeDefinition::GetGenericNumericDef())
+			{
+				FScopedTransaction Transaction(NSLOCTEXT("UnrealEd", "NiagaraEditorChangeNumericPinType", "Change Pin Type"));
+				if (false == InNode->ConvertNumericPinToType(Pin, TypeDef))
+				{
+					Transaction.Cancel();
+				}
+			}
+		}
+	}
+}
+
 void UEdGraphSchema_Niagara::ConvertNumericPinToType(UEdGraphPin* InGraphPin, FNiagaraTypeDefinition TypeDef)
 {
 	if (PinToTypeDefinition(InGraphPin) != TypeDef)
@@ -1510,6 +1528,22 @@ void UEdGraphSchema_Niagara::GetNumericConversionToSubMenuActions(class FMenuBui
 	}
 }
 
+void UEdGraphSchema_Niagara::GetNumericConversionToSubMenuActionsAll(class FMenuBuilder& MenuBuilder, UNiagaraNode* InNode)
+{
+	// Add all the types we could convert to
+	for (const FNiagaraTypeDefinition& TypeDef : FNiagaraTypeRegistry::GetNumericTypes())
+	{
+		FText Title = TypeDef.GetNameText();
+
+		FText Description;
+		FFormatNamedArguments Args;
+		Args.Add(TEXT("TypeTitle"), Title);
+		Description = FText::Format(LOCTEXT("NumericConversionText", "{TypeTitle}"), Args);
+		MenuBuilder.AddMenuEntry(Description, Description, FSlateIcon(), FUIAction(
+			FExecuteAction::CreateUObject((UEdGraphSchema_Niagara*const)this, &UEdGraphSchema_Niagara::ConvertNumericPinToTypeAll, InNode, FNiagaraTypeDefinition(TypeDef))));
+	}
+}
+
 void UEdGraphSchema_Niagara::ToggleNodeEnabledState(UNiagaraNode* InNode) const
 {
 	if (InNode != nullptr)
@@ -1590,7 +1624,7 @@ void UEdGraphSchema_Niagara::GetContextMenuActions(const UEdGraph* CurrentGraph,
 			{
 				MenuBuilder->AddSubMenu(
 					LOCTEXT("ConvertNumericSpecific", "Convert Numeric To..."),
-					LOCTEXT("ConvertNumericSpecificToolTip", "Convert Numeric pin to specific typed pin."),
+					LOCTEXT("ConvertNumericSpecificToolTip", "Convert Numeric pin to the specific typed pin."),
 				FNewMenuDelegate::CreateUObject((UEdGraphSchema_Niagara*const)this, &UEdGraphSchema_Niagara::GetNumericConversionToSubMenuActions, const_cast<UEdGraphPin*>(InGraphPin)));
 			}
 
@@ -1621,7 +1655,7 @@ void UEdGraphSchema_Niagara::GetContextMenuActions(const UEdGraph* CurrentGraph,
 				if (InGraphPin->LinkedTo.Num() == 0 && InGraphPin->bDefaultValueIsIgnored == false)
 				{
 					MenuBuilder->AddMenuEntry(
-						LOCTEXT("ResetInputToDefault", "Reset to Default"), 
+						LOCTEXT("ResetInputToDefault", "Reset to Default"),  // TODO(mv): This is currently broken
 						LOCTEXT("ResetInputToDefaultToolTip", "Reset this input to its default value."),
 						FSlateIcon(),
 						FUIAction(
@@ -1635,12 +1669,32 @@ void UEdGraphSchema_Niagara::GetContextMenuActions(const UEdGraph* CurrentGraph,
 	else if (InGraphNode)
 	{
 		const UNiagaraNode* Node = Cast<UNiagaraNode>(InGraphNode);
+
+		bool bHasNumerics = false;
+		for (auto Pin : Node->Pins)
+		{
+			if (PinToTypeDefinition(Pin) == FNiagaraTypeDefinition::GetGenericNumericDef())
+			{
+				bHasNumerics = true;
+				break;
+			}
+		}
+		if (bHasNumerics)
+		{
+			MenuBuilder->BeginSection("EdGraphSchema_NiagaraNodeActions", LOCTEXT("PinConversionMenuHeader", "Convert Pins"));
+			MenuBuilder->AddSubMenu(
+				LOCTEXT("ConvertNumericSpecific", "Convert All Numerics To..."),
+				LOCTEXT("ConvertNumericSpecificToolTip", "Convert all Numeric pins to the specific typed pin."),
+				FNewMenuDelegate::CreateUObject((UEdGraphSchema_Niagara*const)this, &UEdGraphSchema_Niagara::GetNumericConversionToSubMenuActionsAll, const_cast<UNiagaraNode*>(Node)));
+			MenuBuilder->EndSection();
+		}
+		
 		MenuBuilder->BeginSection("EdGraphSchema_NiagaraNodeActions", LOCTEXT("NodeActionsMenuHeader", "Node Actions"));
 		MenuBuilder->AddMenuEntry(LOCTEXT("ToggleEnabledState", "Toggle Enabled State"), LOCTEXT("ToggleEnabledStateTooltip", "Toggle this node between Enbled (default) and Disabled (skipped from compilation)."), FSlateIcon(),
 			FUIAction(FExecuteAction::CreateUObject((UEdGraphSchema_Niagara*const)this, &UEdGraphSchema_Niagara::ToggleNodeEnabledState, const_cast<UNiagaraNode*>(Node))));
 		MenuBuilder->AddMenuEntry(LOCTEXT("RefreshNode", "Refresh Node"), LOCTEXT("RefreshNodeTooltip", "Refresh this node."), FSlateIcon(),
 			FUIAction(FExecuteAction::CreateUObject((UEdGraphSchema_Niagara*const)this, &UEdGraphSchema_Niagara::RefreshNode, const_cast<UNiagaraNode*>(Node))));
-
+		
 		MenuBuilder->EndSection();
 	}
 
diff --git a/Engine/Plugins/FX/Niagara/Source/NiagaraEditor/Public/EdGraphSchema_Niagara.h b/Engine/Plugins/FX/Niagara/Source/NiagaraEditor/Public/EdGraphSchema_Niagara.h
--- a/Engine/Plugins/FX/Niagara/Source/NiagaraEditor/Public/EdGraphSchema_Niagara.h
+++ b/Engine/Plugins/FX/Niagara/Source/NiagaraEditor/Public/EdGraphSchema_Niagara.h
@@ -126,7 +126,9 @@ class NIAGARAEDITOR_API UEdGraphSchema_Niagara : public UEdGraphSchema
 private:
 	void GetBreakLinkToSubMenuActions(class FMenuBuilder& MenuBuilder, UEdGraphPin* InGraphPin);
 	void GetNumericConversionToSubMenuActions(class FMenuBuilder& MenuBuilder, UEdGraphPin* InGraphPin);
+	void GetNumericConversionToSubMenuActionsAll(class FMenuBuilder& MenuBuilder, UNiagaraNode* InGraphPin);
 	void ConvertNumericPinToType(UEdGraphPin* InPin, FNiagaraTypeDefinition TypeDef);
+	void ConvertNumericPinToTypeAll(UNiagaraNode* InPin, FNiagaraTypeDefinition TypeDef);
 	static bool CheckCircularConnection(const UEdGraphNode* InRootNode, const EEdGraphPinDirection InRootPinDirection, const UEdGraphPin* InPin, int32& OutDepth);
 };
 
