Parameter default rework checkpoint before break.

 - Adds a details customization panel for UNiagaraScriptVariables selected in the Parameters panel.
    - Supports static switch variable as well.
    - Renaming or deleting a static switch *node* will properly delete it from the Parameters panel, and vice versa. Changing the type of the static switch variable in the node correctly reflects in the Parameters panel.
    - Adds per-graph description metadata for non-module parameters, you can now explain the purpose of each introduced attribute.
    - Selectively hides irrelevant metadata for static switch parameters and non-module parameters.
    - Adds a proper value widget leveraging NiagaraTypeUtilities to the details panel which mirrors the value in the default pin in the graph.

- Bugfix: Adding a new variable from an Add pin in the graph will now generate a proper unique name.
- Verified several metadata related JIRAs as being fixed by this rewor:
    - UE-74744
    - UE-75168
    - UE-76184

diff --git a/Engine/Plugins/FX/Niagara/Source/Niagara/Public/NiagaraTypes.h b/Engine/Plugins/FX/Niagara/Source/Niagara/Public/NiagaraTypes.h
--- a/Engine/Plugins/FX/Niagara/Source/Niagara/Public/NiagaraTypes.h
+++ b/Engine/Plugins/FX/Niagara/Source/Niagara/Public/NiagaraTypes.h
@@ -345,22 +345,24 @@ public:
 	bool bInlineEditConditionToggle;
 
 	/** Declares the associated input should be conditionally editable based on the value of another input. */
-	UPROPERTY(EditAnywhere, Category = "Input Conditions", meta = (EditCondition = "!bIsStaticSwitch"))
+	UPROPERTY(EditAnywhere, Category = "Input Conditions")
 	FNiagaraInputConditionMetadata EditCondition;
 
 	/** Declares the associated input should be conditionally visible based on the value of another input. */
-	UPROPERTY(EditAnywhere, Category = "Input Conditions", meta = (EditCondition = "!bIsStaticSwitch"))
+	UPROPERTY(EditAnywhere, Category = "Input Conditions")
 	FNiagaraInputConditionMetadata VisibleCondition;
 
 	UPROPERTY(EditAnywhere, Category = "Variable", DisplayName = "Property Metadata", meta = (ToolTip = "Property Metadata"))
 	TMap<FName, FString> PropertyMetaData;
-
-	UPROPERTY(AdvancedDisplay, VisibleAnywhere, Category = "Variable", meta = (ToolTip = "This is a read-only variable that designates if the metadata is tied to a static switch or not."))
-	bool bIsStaticSwitch;
+	
+	/** This is a read-only variable that designates if the metadata is tied to a static switch or not. */
+	UPROPERTY()
+	bool bIsStaticSwitch; // TODO: This should be moved to the UNiagaraScriptVariable in the future
 
 	/** The default value to use when creating new pins or stack entries for a static switch parameter */
 	UPROPERTY()
-	int32 StaticSwitchDefaultValue;
+	int32 StaticSwitchDefaultValue; // TODO: This should be moved to the UNiagaraScriptVariable in the future
+
 };
 
 USTRUCT()
diff --git a/Engine/Plugins/FX/Niagara/Source/NiagaraEditor/Private/Customizations/NiagaraScriptVariableDetails.cpp b/Engine/Plugins/FX/Niagara/Source/NiagaraEditor/Private/Customizations/NiagaraScriptVariableDetails.cpp
new file mode 100644
--- /dev/null
+++ b/Engine/Plugins/FX/Niagara/Source/NiagaraEditor/Private/Customizations/NiagaraScriptVariableDetails.cpp
@@ -0,0 +1,203 @@
+// Copyright 1998-2019 Epic Games, Inc. All Rights Reserved.
+
+#include "NiagaraScriptVariableDetails.h"
+
+#include "NiagaraScriptVariable.h"
+#include "NiagaraGraph.h"
+#include "NiagaraEditorModule.h"
+
+#include "SNiagaraParameterEditor.h"
+
+#include "DetailWidgetRow.h"
+#include "DetailCategoryBuilder.h"
+#include "DetailLayoutBuilder.h"
+
+#include "TypeEditorUtilities/NiagaraFloatTypeEditorUtilities.h"
+#include "TypeEditorUtilities/NiagaraIntegerTypeEditorUtilities.h"
+#include "TypeEditorUtilities/NiagaraEnumTypeEditorUtilities.h"
+#include "TypeEditorUtilities/NiagaraBoolTypeEditorUtilities.h"
+#include "TypeEditorUtilities/NiagaraFloatTypeEditorUtilities.h"
+#include "TypeEditorUtilities/NiagaraVectorTypeEditorUtilities.h"
+#include "TypeEditorUtilities/NiagaraColorTypeEditorUtilities.h"
+#include "NiagaraEditorStyle.h"
+
+#define LOCTEXT_NAMESPACE "NiagaraScriptVariableVariableDetails"
+
+TSharedRef<IDetailCustomization> FNiagaraScriptVariableDetails::MakeInstance()
+{
+	return MakeShareable(new FNiagaraScriptVariableDetails);
+}
+
+FNiagaraScriptVariableDetails::FNiagaraScriptVariableDetails()
+{
+
+}
+
+void FNiagaraScriptVariableDetails::CustomizeDetails(IDetailLayoutBuilder& DetailBuilder)
+{
+	static const FName CategoryName = TEXT("Script Variable Data");
+
+	TArray<TWeakObjectPtr<UObject>> ObjectsCustomized;
+	DetailBuilder.GetObjectsBeingCustomized(ObjectsCustomized);
+
+	if (ObjectsCustomized.Num() != 1)
+	{
+		// TODO: Could we allow selecting multiple items in the future?
+		return;
+	}
+	if (!ObjectsCustomized[0]->IsA<UNiagaraScriptVariable>())
+	{
+		return;
+	}
+
+	Variable = Cast<UNiagaraScriptVariable>(ObjectsCustomized[0].Get());
+	if (Variable == nullptr)
+	{
+		return;
+	}
+
+	IDetailCategoryBuilder& CategoryBuilder = DetailBuilder.EditCategory(CategoryName);		
+	FNiagaraEditorModule& EditorModule = FNiagaraEditorModule::Get();
+
+	if (Variable->Metadata.bIsStaticSwitch)
+	{
+		// Hide metadata UProperties that isn't useful for static switches
+		DetailBuilder.HideProperty(GET_MEMBER_NAME_CHECKED(UNiagaraScriptVariable, Metadata.EditCondition));
+		DetailBuilder.HideProperty(GET_MEMBER_NAME_CHECKED(UNiagaraScriptVariable, Metadata.VisibleCondition));
+
+		// define a custom parameter editor widget for the static switch data
+		TypeUtilityStaticSwitchValue = EditorModule.GetTypeUtilities(Variable->Variable.GetType());
+		if (TypeUtilityStaticSwitchValue && TypeUtilityStaticSwitchValue->CanCreateParameterEditor())
+		{
+			ParameterEditorStaticSwitchValue = TypeUtilityStaticSwitchValue->CreateParameterEditor(Variable->Variable.GetType());
+			
+			TSharedPtr<FStructOnScope> ParameterValue = MakeShareable(new FStructOnScope(Variable->Variable.GetType().GetStruct()));
+			Variable->Variable.SetValue(Variable->Metadata.StaticSwitchDefaultValue);
+			Variable->Variable.CopyTo(ParameterValue->GetStructMemory());
+			ParameterEditorStaticSwitchValue->UpdateInternalValueFromStruct(ParameterValue.ToSharedRef());
+			ParameterEditorStaticSwitchValue->SetOnValueChanged(SNiagaraParameterEditor::FOnValueChange::CreateSP(this, &FNiagaraScriptVariableDetails::OnStaticSwitchValueChanged));
+
+			FDetailWidgetRow& DefaultValueWidget = CategoryBuilder.AddCustomRow(LOCTEXT("DefaultValueFilterText", "Default Value"));
+			DefaultValueWidget
+			.NameContent()
+			[
+				SNew(STextBlock)
+				.Text(FText::FromString(TEXT("Default Value")))
+			]
+			.ValueContent()
+			.HAlign(HAlign_Fill)
+			[
+				ParameterEditorStaticSwitchValue.ToSharedRef()
+			];
+		}
+		else
+		{
+			TypeUtilityStaticSwitchValue = nullptr;
+		}
+	}
+	else if (Variable->Variable.GetName().ToString().StartsWith("Module.")) 
+	{
+		if (UNiagaraGraph* Graph = Cast<UNiagaraGraph>(Variable->GetOuter())) {
+			// TODO: Probably make a custom function for this...
+			UEdGraphPin* Pin = Graph->FindParameterMapDefaultValuePin(Variable->Variable.GetName(), ENiagaraScriptUsage::Module, ENiagaraScriptUsage::Module);
+			if (Pin == nullptr) 
+			{
+				Pin = Graph->FindParameterMapDefaultValuePin(Variable->Variable.GetName(), ENiagaraScriptUsage::DynamicInput, ENiagaraScriptUsage::Module);
+			}
+			if (Pin == nullptr) 
+			{
+				Pin = Graph->FindParameterMapDefaultValuePin(Variable->Variable.GetName(), ENiagaraScriptUsage::Function, ENiagaraScriptUsage::Module);
+			}
+			if (Pin)
+			{
+				// define a custom parameter editor widget for module variables
+				TypeUtilityValue = EditorModule.GetTypeUtilities(Variable->Variable.GetType());
+				if (TypeUtilityValue && TypeUtilityValue->CanCreateParameterEditor())
+				{
+					ParameterEditorValue = TypeUtilityValue->CreateParameterEditor(Variable->Variable.GetType());
+
+					TypeUtilityValue->SetValueFromPinDefaultString(Pin->DefaultValue, Variable->Variable);
+					TSharedPtr<FStructOnScope> ParameterValue = MakeShareable(new FStructOnScope(Variable->Variable.GetType().GetStruct()));
+					Variable->Variable.CopyTo(ParameterValue->GetStructMemory());
+					ParameterEditorValue->UpdateInternalValueFromStruct(ParameterValue.ToSharedRef());
+					ParameterEditorValue->SetOnValueChanged(SNiagaraParameterEditor::FOnValueChange::CreateSP(this, &FNiagaraScriptVariableDetails::OnValueChanged));
+					
+					FDetailWidgetRow& DefaultValueWidget = CategoryBuilder.AddCustomRow(LOCTEXT("DefaultValueFilterText", "Default Value"));
+					DefaultValueWidget
+					.NameContent()
+					[
+						SNew(STextBlock)
+						.Font(FNiagaraEditorStyle::Get().GetFontStyle("NiagaraEditor.ParameterFont"))
+						.Text(FText::FromString(TEXT("Default Value")))
+					]
+					.ValueContent()
+					.HAlign(HAlign_Fill)
+					[
+						ParameterEditorValue.ToSharedRef()
+					];
+				}
+				else
+				{
+					TypeUtilityValue = nullptr;
+				}
+			}
+		}
+	}
+	else
+	{
+		// Hide metadata UProperties that isn't useful for non module and non static switch variables
+		DetailBuilder.HideProperty(GET_MEMBER_NAME_CHECKED(UNiagaraScriptVariable, Metadata.CategoryName));
+		DetailBuilder.HideProperty(GET_MEMBER_NAME_CHECKED(UNiagaraScriptVariable, Metadata.bAdvancedDisplay));
+		DetailBuilder.HideProperty(GET_MEMBER_NAME_CHECKED(UNiagaraScriptVariable, Metadata.EditorSortPriority));
+		DetailBuilder.HideProperty(GET_MEMBER_NAME_CHECKED(UNiagaraScriptVariable, Metadata.bInlineEditConditionToggle));
+		DetailBuilder.HideProperty(GET_MEMBER_NAME_CHECKED(UNiagaraScriptVariable, Metadata.EditCondition));
+		DetailBuilder.HideProperty(GET_MEMBER_NAME_CHECKED(UNiagaraScriptVariable, Metadata.VisibleCondition));
+		DetailBuilder.HideProperty(GET_MEMBER_NAME_CHECKED(UNiagaraScriptVariable, Metadata.PropertyMetaData));
+	}
+
+	// bInlineEditConditionToggle is only useful for bool types
+	if (Variable->Variable.GetType() != FNiagaraTypeDefinition::GetBoolDef())
+	{
+		DetailBuilder.HideProperty(GET_MEMBER_NAME_CHECKED(UNiagaraScriptVariable, Metadata.bInlineEditConditionToggle));
+	}
+}
+
+void FNiagaraScriptVariableDetails::OnValueChanged()
+{
+	if (TypeUtilityValue && ParameterEditorValue)
+	{
+		if (UNiagaraGraph* Graph = Cast<UNiagaraGraph>(Variable->GetOuter()))
+		{
+			UEdGraphPin* Pin = Graph->FindParameterMapDefaultValuePin(Variable->Variable.GetName(), ENiagaraScriptUsage::Module, ENiagaraScriptUsage::Module);
+			if (Pin == nullptr) 
+			{
+				Pin = Graph->FindParameterMapDefaultValuePin(Variable->Variable.GetName(), ENiagaraScriptUsage::DynamicInput, ENiagaraScriptUsage::Module);
+			}
+			if (Pin == nullptr) 
+			{
+				Pin = Graph->FindParameterMapDefaultValuePin(Variable->Variable.GetName(), ENiagaraScriptUsage::Function, ENiagaraScriptUsage::Module);
+			}
+			if (Pin)
+			{
+				TSharedPtr<FStructOnScope> ParameterValue = MakeShareable(new FStructOnScope(Variable->Variable.GetType().GetStruct()));
+				ParameterEditorValue->UpdateStructFromInternalValue(ParameterValue.ToSharedRef());
+				Variable->Variable.SetData(ParameterValue->GetStructMemory());
+
+				Pin->DefaultValue = TypeUtilityValue->GetPinDefaultStringFromValue(Variable->Variable);
+			}
+		}
+	} 
+}
+
+void FNiagaraScriptVariableDetails::OnStaticSwitchValueChanged()
+{
+	if (TypeUtilityStaticSwitchValue && ParameterEditorStaticSwitchValue)
+	{
+		TSharedPtr<FStructOnScope> ParameterValue = MakeShareable(new FStructOnScope(Variable->Variable.GetType().GetStruct()));
+		ParameterEditorStaticSwitchValue->UpdateStructFromInternalValue(ParameterValue.ToSharedRef());
+		Variable->Variable.SetData(ParameterValue->GetStructMemory());
+		Variable->Metadata.StaticSwitchDefaultValue = Variable->Variable.GetValue<int>();
+	} 
+}
+
+#undef LOCTEXT_NAMESPACE
diff --git a/Engine/Plugins/FX/Niagara/Source/NiagaraEditor/Private/Customizations/NiagaraScriptVariableDetails.h b/Engine/Plugins/FX/Niagara/Source/NiagaraEditor/Private/Customizations/NiagaraScriptVariableDetails.h
new file mode 100644
--- /dev/null
+++ b/Engine/Plugins/FX/Niagara/Source/NiagaraEditor/Private/Customizations/NiagaraScriptVariableDetails.h
@@ -0,0 +1,37 @@
+// Copyright 1998-2019 Epic Games, Inc. All Rights Reserved.
+
+#pragma once
+
+#include "CoreMinimal.h"
+#include "IDetailCustomization.h"
+#include "IPropertyUtilities.h"
+
+/** This customization sets up a custom details panel for the static switch Variable in the niagara module graph. */
+class FNiagaraScriptVariableDetails : public IDetailCustomization, public FEditorUndoClient
+{
+public:
+
+	/** Makes a new instance of this detail layout class for a specific detail view requesting it */
+	static TSharedRef<IDetailCustomization> MakeInstance();
+
+	~FNiagaraScriptVariableDetails();
+	FNiagaraScriptVariableDetails();
+
+	/** IDetailCustomization interface */
+	virtual void CustomizeDetails( IDetailLayoutBuilder& DetailBuilder ) override;
+	
+	//~ Begin FEditorUndoClient Interface
+	virtual void PostUndo(bool bSuccess) override;
+	virtual void PostRedo(bool bSuccess) override;
+	// End of FEditorUndoClient
+
+private:
+	void OnValueChanged();
+	void OnStaticSwitchValueChanged();
+
+	TWeakObjectPtr<class UNiagaraScriptVariable> Variable;
+	TSharedPtr<class INiagaraEditorTypeUtilities, ESPMode::ThreadSafe> TypeUtilityValue;
+	TSharedPtr<class INiagaraEditorTypeUtilities, ESPMode::ThreadSafe> TypeUtilityStaticSwitchValue;
+	TSharedPtr<class SNiagaraParameterEditor> ParameterEditorValue;
+	TSharedPtr<class SNiagaraParameterEditor> ParameterEditorStaticSwitchValue;
+};
diff --git a/Engine/Plugins/FX/Niagara/Source/NiagaraEditor/Private/NiagaraEditorModule.cpp b/Engine/Plugins/FX/Niagara/Source/NiagaraEditor/Private/NiagaraEditorModule.cpp
--- a/Engine/Plugins/FX/Niagara/Source/NiagaraEditor/Private/NiagaraEditorModule.cpp
+++ b/Engine/Plugins/FX/Niagara/Source/NiagaraEditor/Private/NiagaraEditorModule.cpp
@@ -94,6 +94,7 @@
 #include "Misc/FeedbackContext.h"
 #include "Customizations/NiagaraStaticSwitchNodeDetails.h"
 #include "Customizations/NiagaraFunctionCallNodeDetails.h"
+#include "Customizations/NiagaraScriptVariableDetails.h"
 
 IMPLEMENT_MODULE( FNiagaraEditorModule, NiagaraEditor );
 
@@ -409,6 +410,8 @@ void FNiagaraEditorModule::StartupModule()
 	PropertyModule.RegisterCustomClassLayout("NiagaraComponent", FOnGetDetailCustomizationInstance::CreateStatic(&FNiagaraComponentDetails::MakeInstance));
 
 	PropertyModule.RegisterCustomClassLayout("NiagaraNodeStaticSwitch", FOnGetDetailCustomizationInstance::CreateStatic(&FNiagaraStaticSwitchNodeDetails::MakeInstance));
+	
+	PropertyModule.RegisterCustomClassLayout("NiagaraScriptVariable", FOnGetDetailCustomizationInstance::CreateStatic(&FNiagaraScriptVariableDetails::MakeInstance));
 
 	PropertyModule.RegisterCustomClassLayout("NiagaraNodeFunctionCall", FOnGetDetailCustomizationInstance::CreateStatic(&FNiagaraFunctionCallNodeDetails::MakeInstance));
 	
diff --git a/Engine/Plugins/FX/Niagara/Source/NiagaraEditor/Private/NiagaraGraph.cpp b/Engine/Plugins/FX/Niagara/Source/NiagaraEditor/Private/NiagaraGraph.cpp
--- a/Engine/Plugins/FX/Niagara/Source/NiagaraEditor/Private/NiagaraGraph.cpp
+++ b/Engine/Plugins/FX/Niagara/Source/NiagaraEditor/Private/NiagaraGraph.cpp
@@ -22,6 +22,7 @@
 #include "EdGraphSchema_Niagara.h"
 #include "NiagaraNodeParameterMapBase.h"
 #include "NiagaraNodeParameterMapGet.h"
+#include "NiagaraNodeParameterMapSet.h"
 #include "INiagaraEditorTypeUtilities.h"
 #include "NiagaraConstants.h"
 #include "NiagaraEditorModule.h"
@@ -80,7 +81,6 @@ UNiagaraGraph::UNiagaraGraph(const FObjectInitializer& ObjectInitializer)
 	, bNeedNumericCacheRebuilt(true)
 	, bIsRenamingParameter(false)
 	, bParameterReferenceRefreshPending(true)
-	, bUnreferencedMetaDataPurgePending(true)
 {
 	Schema = UEdGraphSchema_Niagara::StaticClass();
 	ChangeId = FGuid::NewGuid();
@@ -209,27 +209,63 @@ void UNiagaraGraph::PostLoad()
 	// Migrate input condition metadata
 	const int32 NiagaraVer = GetLinkerCustomVersion(FNiagaraCustomVersion::GUID);
 
+	FString FullPathName = GetPathName();
+	FString PathName;
+	int ColonPos;
+	if (FullPathName.FindChar(TCHAR('.'), ColonPos))
+	{
+		// GetPathName() returns something similar to "/Path/To/ScriptName.ScriptName:NiagaraScriptSource_N.NiagaraGraph_N"
+		// so this will extract "/Path/To/ScriptName"
+		PathName = FullPathName.Left(ColonPos);
+	}
+	// UE_LOG(LogNiagaraEditor, Log, TEXT("PostLoad %s"), *FullPathName);
+
 	if (NiagaraVer < FNiagaraCustomVersion::MetaDataAndParametersUpdate)
 	{
+		int NumMigrated = 0;
 		// If the version of the asset is older than FNiagaraCustomVersion::MetaDataAndParametersUpdate 
 		// we need to migrate the old metadata by looping through VariableToMetaData_DEPRECATED
 		// and create new entries in VariableToScriptVariable
-		for (auto It = VariableToMetaData_DEPRECATED.CreateConstIterator(); It; ++It)
+		for (auto& Pair : VariableToMetaData_DEPRECATED)
 		{
-			SetMetaData(It.Key(), It.Value());
+			FNiagaraVariable Var(Pair.Key.GetType(), Pair.Key.GetName());
+			SetMetaData(Var, Pair.Value);
+			NumMigrated++;
+		}
+		VariableToMetaData_DEPRECATED.Empty();
 
-			FString PathName = GetPathName();
-			int ColonPos;
-			if (PathName.FindChar(TCHAR('.'), ColonPos))
+		if (NumMigrated > 0)
+		{
+			// TODO: Should we even spam this? It will spam it until it's resaved
+			//UE_LOG(LogNiagaraEditor, Log, TEXT("Migrated %d old metadata entries for \"%s\""), NumMigrated, *FullPathName);
+		}
+	}
+	
+	{
+		// Create a UNiagaraScriptVariable instance for every entry in the parameter map for which there is no existing script variable. 
+
+		TArray<FNiagaraVariable> VarsToAdd;
+		for (auto& ParameterToReferences : GetParameterReferenceMap())
+		{
+			UNiagaraScriptVariable* Variable = GetScriptVariable(ParameterToReferences.Key.GetName());
+			if (Variable == nullptr)
 			{
-				// GetPathName() returns something similar to "/Path/To/ScriptName.ScriptName:NiagaraScriptSource_N.NiagaraGraph_N"
-				// so this will extract "/Path/To/ScriptName"
-				PathName = PathName.Left(ColonPos);
+				VarsToAdd.Add(ParameterToReferences.Key);
 			}
-			UE_LOG(LogNiagaraEditor, Log, TEXT("Migrated old metadata entry for variable \"%s\" in \"%s\""), *It.Key().GetName().ToString(), *PathName);
 		}
-		VariableToMetaData_DEPRECATED.Empty();
+
+		for (FNiagaraVariable& Var : VarsToAdd)
+		{
+			AddParameter(Var);
+		}
+
+		if (VarsToAdd.Num() > 0)
+		{
+			// TODO: Should we even spam this? It will spam it until it's resaved
+			//UE_LOG(LogNiagaraEditor, Log, TEXT("Added %d new script variables for \"%s\""), VarsToAdd.Num(), *FullPathName);
+		}
 	}
+
 	if (NiagaraVer < FNiagaraCustomVersion::MoveCommonInputMetadataToProperties)
 	{
 		auto MigrateInputCondition = [](TMap<FName, FString>& PropertyMetaData, const FName& InputConditionKey, FNiagaraInputConditionMetadata& InOutInputCondition)
@@ -303,7 +339,6 @@ FGuid UNiagaraGraph::ComputeCompileID(ENiagaraScriptUsage InUsage, const FGuid&
 	}
 
 	return FGuid();
-
 }
 
 FNiagaraCompileHash UNiagaraGraph::GetCompileDataHash(ENiagaraScriptUsage InUsage, const FGuid& InUsageId) const
@@ -356,9 +391,7 @@ UEdGraphPin* UNiagaraGraph::FindParameterMapDefaultValuePin(const FName Variable
 	UEdGraphPin* DefaultInputPin = nullptr;
 	for (UNiagaraNode* Node : NodesTraversed)
 	{
-		UNiagaraNodeParameterMapGet* GetNode = Cast<UNiagaraNodeParameterMapGet>(Node);
-	
-		if (GetNode)
+		if (UNiagaraNodeParameterMapGet* GetNode = Cast<UNiagaraNodeParameterMapGet>(Node))
 		{
 			TArray<UEdGraphPin*> OutputPins;
 			GetNode->GetOutputPins(OutputPins);
@@ -374,15 +407,33 @@ UEdGraphPin* UNiagaraGraph::FindParameterMapDefaultValuePin(const FName Variable
 					}
 				}
 			}
+			
+			if (DefaultInputPin != nullptr)
+			{
+				break;
+			}
 		}
-
-		if (DefaultInputPin != nullptr)
+		
+		if (UNiagaraNodeParameterMapSet* SetNode = Cast<UNiagaraNodeParameterMapSet>(Node))
 		{
-			break;
+			TArray<UEdGraphPin*> InputPins;
+			SetNode->GetInputPins(InputPins);
+			for (UEdGraphPin* InputPin : InputPins)
+			{
+				if (VariableName == InputPin->PinName)
+				{
+					DefaultInputPin = InputPin;
+					break;
+				}
+			}
+
+			if (DefaultInputPin != nullptr)
+			{
+				break;
+			}
 		}
 	}
 
-
 	// There are some pins 
 	if (DefaultInputPin && DefaultInputPin->LinkedTo.Num() != 0 && DefaultInputPin->LinkedTo[0] != nullptr)
 	{
@@ -428,7 +479,6 @@ void UNiagaraGraph::FindOutputNodes(TArray<UNiagaraNodeOutput*>& OutputNodes) co
 	}
 }
 
-
 void UNiagaraGraph::FindOutputNodes(ENiagaraScriptUsage TargetUsageType, TArray<UNiagaraNodeOutput*>& OutputNodes) const
 {
 	TArray<UNiagaraNodeOutput*> NodesFound;
@@ -491,7 +541,6 @@ UNiagaraNodeOutput* UNiagaraGraph::FindEquivalentOutputNode(ENiagaraScriptUsage
 	return nullptr;
 }
 
-
 void BuildTraversalHelper(TArray<class UNiagaraNode*>& OutNodesTraversed, UNiagaraNode* CurrentNode)
 {
 	if (CurrentNode == nullptr)
@@ -535,7 +584,6 @@ void UNiagaraGraph::BuildTraversal(TArray<class UNiagaraNode*>& OutNodesTraverse
 	}
 }
 
-
 void UNiagaraGraph::FindInputNodes(TArray<UNiagaraNodeInput*>& OutInputNodes, UNiagaraGraph::FFindInputNodeOptions Options) const
 {
 	SCOPE_CYCLE_COUNTER(STAT_NiagaraEditor_Graph_FindInputNodes);
@@ -695,55 +743,13 @@ TArray<UEdGraphNode*> UNiagaraGraph::FindReachbleNodes() const
 	return ResultNodes;
 }
 
-void UNiagaraGraph::GetParameters(TArray<FNiagaraVariable>& Inputs, TArray<FNiagaraVariable>& Outputs)const
-{
-	Inputs.Empty();
-	Outputs.Empty();
-
-	TArray<UNiagaraNodeInput*> InputsNodes;
-	FFindInputNodeOptions Options;
-	Options.bSort = true;
-	FindInputNodes(InputsNodes, Options);
-	for (UNiagaraNodeInput* Input : InputsNodes)
-	{
-		Inputs.Add(Input->Input);
-	}
-
-	TArray<UNiagaraNodeOutput*> OutputNodes;
-	FindOutputNodes(OutputNodes);
-	for (UNiagaraNodeOutput* OutputNode : OutputNodes)
-	{
-		for (FNiagaraVariable& Var : OutputNode->Outputs)
-		{
-			Outputs.AddUnique(Var);
-		}
-	}
-
-	//Do we need to sort outputs?
-	//Should leave them as they're defined in the output node?
-// 	auto SortVars = [](const FNiagaraVariable& A, const FNiagaraVariable& B)
-// 	{
-// 		//Case insensitive lexicographical comparisons of names for sorting.
-// 		return A.GetName().ToString() < B.GetName().ToString();
-// 	};
-// 	Outputs.Sort(SortVars);
-}
-
 const TMap<FNiagaraVariable, UNiagaraScriptVariable*>& UNiagaraGraph::GetAllMetaData() const
 {
-	if (bUnreferencedMetaDataPurgePending)
-	{
-		PurgeUnreferencedMetaData();
-	}
 	return VariableToScriptVariable;
 }
 
 TMap<FNiagaraVariable, UNiagaraScriptVariable*>& UNiagaraGraph::GetAllMetaData()
 {
-	if (bUnreferencedMetaDataPurgePending)
-	{
-		PurgeUnreferencedMetaData();
-	}
 	return VariableToScriptVariable;
 }
 
@@ -756,7 +762,28 @@ const TMap<FNiagaraVariable, FNiagaraGraphParameterReferenceCollection>& UNiagar
 	return ParameterToReferencesMap;
 }
 
-void UNiagaraGraph::AddParameter(const FNiagaraVariable& Parameter)
+UNiagaraScriptVariable* UNiagaraGraph::GetScriptVariable(FNiagaraVariable Parameter) const
+{
+	if (UNiagaraScriptVariable** FoundScriptVariable = VariableToScriptVariable.Find(Parameter))
+	{
+		return *FoundScriptVariable;
+	}
+	return nullptr;
+}
+
+UNiagaraScriptVariable* UNiagaraGraph::GetScriptVariable(FName ParameterName) const
+{
+	for (auto& VariableToScriptVariableItem : VariableToScriptVariable)
+	{
+		if (VariableToScriptVariableItem.Key.GetName() == ParameterName)
+		{
+			return VariableToScriptVariableItem.Value;
+		}
+	}
+	return nullptr;
+}
+
+void UNiagaraGraph::AddParameter(const FNiagaraVariable& Parameter, bool bIsStaticSwitch)
 {
 	FNiagaraGraphParameterReferenceCollection* FoundParameterReferenceCollection = ParameterToReferencesMap.Find(Parameter);
 	if (!FoundParameterReferenceCollection)
@@ -771,11 +798,12 @@ void UNiagaraGraph::AddParameter(const FNiagaraVariable& Parameter)
 	{
 		UNiagaraScriptVariable* NewScriptVariable = NewObject<UNiagaraScriptVariable>(this);
 		NewScriptVariable->Variable = Parameter;
+		NewScriptVariable->Metadata.bIsStaticSwitch = bIsStaticSwitch;
 		VariableToScriptVariable.Add(Parameter, NewScriptVariable);
 	}
 }
 
-void UNiagaraGraph::RemoveParameter(const FNiagaraVariable& Parameter)
+void UNiagaraGraph::RemoveParameter(const FNiagaraVariable& Parameter, bool bFromStaticSwitch)
 {
 	FNiagaraGraphParameterReferenceCollection* ReferenceCollection = ParameterToReferencesMap.Find(Parameter);
 	if (ReferenceCollection)
@@ -799,9 +827,33 @@ void UNiagaraGraph::RemoveParameter(const FNiagaraVariable& Parameter)
 		ParameterToReferencesMap.Remove(Parameter);
 		NotifyGraphChanged();
 	}
+
+	TArray<FNiagaraVariable> VarsToRemove;
+	for (auto It : VariableToScriptVariable)
+	{
+		if (bFromStaticSwitch)
+		{
+			if (It.Key.GetName() == Parameter.GetName() && It.Key.GetType() == Parameter.GetType())
+			{
+				VarsToRemove.Add(It.Key);
+			}
+		}
+		else
+		{
+			if (It.Key.GetName() == Parameter.GetName())
+			{
+				VarsToRemove.Add(It.Key);
+			}
+		}
+		
+	}
+	for (FNiagaraVariable& Var : VarsToRemove)
+	{
+		VariableToScriptVariable.Remove(Var);
+	}
 }
 
-bool UNiagaraGraph::RenameParameter(const FNiagaraVariable& Parameter, FName NewName)
+bool UNiagaraGraph::RenameParameter(const FNiagaraVariable& Parameter, FName NewName, bool bFromStaticSwitch)
 {
 	// Block rename when already renaming. This prevents recursion when CommitEditablePinName is called on referenced nodes. 
 	if (bIsRenamingParameter)
@@ -818,13 +870,20 @@ bool UNiagaraGraph::RenameParameter(const FNiagaraVariable& Parameter, FName New
 	FNiagaraVariableMetaData OldMetaData;
 	if (OldScriptVariable)
 	{
+		if (!bFromStaticSwitch && (*OldScriptVariable)->Metadata.bIsStaticSwitch)
+		{
+			// We current disallow renaming static switch variables in the Parameters panel. 
+			bIsRenamingParameter = false;
+			return false;
+		}
 		OldMetaData = (*OldScriptVariable)->Metadata;
 	}
-		
 
 	FNiagaraGraphParameterReferenceCollection* ReferenceCollection = ParameterToReferencesMap.Find(Parameter);
+	int NumReferences = 0;
 	if (ReferenceCollection)
 	{
+		NumReferences = ReferenceCollection->ParameterReferences.Num();
 		const FText NewNameText = FText::FromName(NewName);
 		FNiagaraGraphParameterReferenceCollection NewReferences = *ReferenceCollection;
 		for (FNiagaraGraphParameterReference& Reference : NewReferences.ParameterReferences)
@@ -847,7 +906,10 @@ bool UNiagaraGraph::RenameParameter(const FNiagaraVariable& Parameter, FName New
 	// Swap metadata to the new parameter
 	if (UNiagaraScriptVariable** Metadata = VariableToScriptVariable.Find(Parameter))
 	{
-		VariableToScriptVariable.Remove(Parameter);
+		if (!bFromStaticSwitch || (bFromStaticSwitch && NumReferences <= 1))
+		{
+			VariableToScriptVariable.Remove(Parameter);
+		}
 	}
 	SetMetaData(NewParameter, OldMetaData);
 
@@ -890,6 +952,40 @@ void UNiagaraGraph::GetOutputNodeVariables(ENiagaraScriptUsage InScriptUsage, TA
 	}
 }
 
+void UNiagaraGraph::GetParameters(TArray<FNiagaraVariable>& Inputs, TArray<FNiagaraVariable>& Outputs)const
+{
+	Inputs.Empty();
+	Outputs.Empty();
+
+	TArray<UNiagaraNodeInput*> InputsNodes;
+	FFindInputNodeOptions Options;
+	Options.bSort = true;
+	FindInputNodes(InputsNodes, Options);
+	for (UNiagaraNodeInput* Input : InputsNodes)
+	{
+		Inputs.Add(Input->Input);
+	}
+
+	TArray<UNiagaraNodeOutput*> OutputNodes;
+	FindOutputNodes(OutputNodes);
+	for (UNiagaraNodeOutput* OutputNode : OutputNodes)
+	{
+		for (FNiagaraVariable& Var : OutputNode->Outputs)
+		{
+			Outputs.AddUnique(Var);
+		}
+	}
+
+	//Do we need to sort outputs?
+	//Should leave them as they're defined in the output node?
+	// 	auto SortVars = [](const FNiagaraVariable& A, const FNiagaraVariable& B)
+	// 	{
+	// 		//Case insensitive lexicographical comparisons of names for sorting.
+	// 		return A.GetName().ToString() < B.GetName().ToString();
+	// 	};
+	// 	Outputs.Sort(SortVars);
+}
+
 bool UNiagaraGraph::HasParameterMapParameters()const
 {
 	TArray<FNiagaraVariable> Inputs;
@@ -947,7 +1043,6 @@ void UNiagaraGraph::NotifyGraphNeedsRecompile()
 	NotifyGraphChanged(Action);
 }
 
-
 void UNiagaraGraph::NotifyGraphDataInterfaceChanged()
 {
 	OnDataInterfaceChangedDelegate.Broadcast();
@@ -1239,7 +1334,6 @@ void UNiagaraGraph::ResolveNumerics(TMap<UNiagaraNode*, bool>& VisitedNodes, UEd
 	}
 }
 
-
 void UNiagaraGraph::SynchronizeInternalCacheWithGraph(UNiagaraGraph* Other)
 {
 	// Force us to rebuild the cache, note that this builds traversals and everything else, keeping it in sync if nothing changed from the current version.
@@ -1295,7 +1389,6 @@ void UNiagaraGraph::SynchronizeInternalCacheWithGraph(UNiagaraGraph* Other)
 	}
 }
 
-
 void UNiagaraGraph::InvalidateCachedCompileIds()
 {
 	Modify();
@@ -1354,7 +1447,6 @@ void UNiagaraGraph::GatherExternalDependencyIDs(ENiagaraScriptUsage InUsage, con
 	}
 }
 
-
 void UNiagaraGraph::GetAllReferencedGraphs(TArray<const UNiagaraGraph*>& Graphs) const
 {
 	Graphs.AddUnique(this);
@@ -1418,11 +1510,6 @@ void UNiagaraGraph::MarkGraphRequiresSynchronization(FString Reason)
 
 TOptional<FNiagaraVariableMetaData> UNiagaraGraph::GetMetaData(const FNiagaraVariable& InVar) const
 {
-	if (bUnreferencedMetaDataPurgePending)
-	{
-		PurgeUnreferencedMetaData();
-	}
-	
 	if (UNiagaraScriptVariable** MetaData = VariableToScriptVariable.Find(InVar))
 	{
 		if (*MetaData)
@@ -1453,44 +1540,6 @@ void UNiagaraGraph::SetMetaData(const FNiagaraVariable& InVar, const FNiagaraVar
 	}
 }
 
-void UNiagaraGraph::PurgeUnreferencedMetaData() const
-{
-	TSet<FNiagaraVariable> ReferencedParameters;
-	ReferencedParameters.Append(FindStaticSwitchInputs());
-	const UEdGraphSchema_Niagara* NiagaraSchema = Cast<UEdGraphSchema_Niagara>(Schema);
-	for (UEdGraphNode* Node : Nodes)
-	{
-		for (UEdGraphPin* Pin : Node->Pins)
-		{
-			if (Pin->PinType.PinSubCategory == UNiagaraNodeParameterMapBase::ParameterPinSubCategory)
-			{
-				const FNiagaraVariable Parameter = NiagaraSchema->PinToNiagaraVariable(Pin, false);
-				const FNiagaraParameterHandle Handle = FNiagaraParameterHandle(Parameter.GetName());
-				if (Handle.IsModuleHandle() && !FNiagaraConstants::IsNiagaraConstant(Parameter))
-				{
-					ReferencedParameters.Add(Parameter);
-				}
-			}
-		}
-	}
-
-	TArray<FNiagaraVariable> VarsToRemove;
-	for (auto It = VariableToScriptVariable.CreateConstIterator(); It; ++It)
-	{
-		if (!ReferencedParameters.Contains(It.Key()))
-		{
-			VarsToRemove.Add(It.Key());
-		}
-	}
-
-	for (FNiagaraVariable& Var : VarsToRemove)
-	{
-		VariableToScriptVariable.Remove(Var);
-	}
-
-	bUnreferencedMetaDataPurgePending = false;
-}
-
 UNiagaraGraph::FOnDataInterfaceChanged& UNiagaraGraph::OnDataInterfaceChanged()
 {
 	return OnDataInterfaceChangedDelegate;
@@ -1525,18 +1574,6 @@ void UNiagaraGraph::RefreshParameterReferences() const
 				FNiagaraGraphParameterReferenceCollection& NewReferenceCollection = ParameterToReferencesMap.Add(Parameter);
 				NewReferenceCollection.Graph = this;
 				ReferenceCollection = &NewReferenceCollection;
-
-				// When a variable is created or added from the graph it won't call AddParameter, 
-				// but instead call this method
-				UNiagaraScriptVariable** FoundScriptVariable = VariableToScriptVariable.Find(Parameter);
-				if (!FoundScriptVariable)
-				{
-					// Warning! This method (RefreshParameterReferences) isn't const at all!
-					// Need to cast away the const to be able to create a serializable UObject here...
-					UNiagaraScriptVariable* NewScriptVariable = NewObject<UNiagaraScriptVariable>(const_cast<UNiagaraGraph*>(this));
-					NewScriptVariable->Variable = Parameter;
-					VariableToScriptVariable.Add(Parameter, NewScriptVariable);
-				}
 			}
 			ReferenceCollection->ParameterReferences.AddUnique(FNiagaraGraphParameterReference(Pin->PersistentGuid, Cast<UNiagaraNode>(Pin->GetOwningNode())));
 
@@ -1556,14 +1593,6 @@ void UNiagaraGraph::RefreshParameterReferences() const
 			NewReferenceCollection.Graph = this;
 			ReferenceCollection = &ParameterToReferencesMap.Add(Variable, NewReferenceCollection);
 		}
-		UNiagaraScriptVariable** FoundScriptVariable = VariableToScriptVariable.Find(Variable);
-		if (!FoundScriptVariable)
-		{
-			UNiagaraScriptVariable* NewScriptVariable = NewObject<UNiagaraScriptVariable>(const_cast<UNiagaraGraph*>(this));
-			NewScriptVariable->Variable = Variable;
-			NewScriptVariable->Metadata.bIsStaticSwitch = true;
-			VariableToScriptVariable.Add(Variable, NewScriptVariable);
-		}
 		ReferenceCollection->ParameterReferences.AddUnique(FNiagaraGraphParameterReference(Node->NodeGuid, Node));
 		CandidateUnreferencedParametersToRemove.Remove(Variable);
 	};
@@ -1621,6 +1650,31 @@ void UNiagaraGraph::RefreshParameterReferences() const
 	for (const FNiagaraVariable& UnreferencedParameterToRemove : CandidateUnreferencedParametersToRemove)
 	{
 		ParameterToReferencesMap.Remove(UnreferencedParameterToRemove);
+		VariableToScriptVariable.Remove(UnreferencedParameterToRemove);
+	}
+
+
+	auto AddMissingScriptVariable = [&](const FNiagaraVariable& Variable)
+	{
+		UNiagaraScriptVariable** FoundScriptVariable = VariableToScriptVariable.Find(Variable);
+		if (!FoundScriptVariable)
+		{
+			UNiagaraScriptVariable* NewScriptVariable = NewObject<UNiagaraScriptVariable>(const_cast<UNiagaraGraph*>(this));
+			NewScriptVariable->Variable = Variable;
+			NewScriptVariable->Metadata.bIsStaticSwitch = false;
+			VariableToScriptVariable.Add(Variable, NewScriptVariable);
+		}
+	};
+
+	for (auto& ParameterToReferences : ParameterToReferencesMap)
+	{
+		UNiagaraScriptVariable* Variable = GetScriptVariable(ParameterToReferences.Key);
+		if (Variable)
+		{
+			continue;
+		}
+
+		AddMissingScriptVariable(ParameterToReferences.Key);
 	}
 
 	bParameterReferenceRefreshPending = false;
@@ -1629,7 +1683,6 @@ void UNiagaraGraph::RefreshParameterReferences() const
 void UNiagaraGraph::InvalidateCachedParameterData()
 {
 	bParameterReferenceRefreshPending = true;
-	bUnreferencedMetaDataPurgePending = true;
 }
 
 #undef LOCTEXT_NAMESPACE
\ No newline at end of file
diff --git a/Engine/Plugins/FX/Niagara/Source/NiagaraEditor/Private/NiagaraNodeStaticSwitch.cpp b/Engine/Plugins/FX/Niagara/Source/NiagaraEditor/Private/NiagaraNodeStaticSwitch.cpp
--- a/Engine/Plugins/FX/Niagara/Source/NiagaraEditor/Private/NiagaraNodeStaticSwitch.cpp
+++ b/Engine/Plugins/FX/Niagara/Source/NiagaraEditor/Private/NiagaraNodeStaticSwitch.cpp
@@ -12,6 +12,12 @@ UNiagaraNodeStaticSwitch::UNiagaraNodeStaticSwitch(const FObjectInitializer& Obj
 {
 }
 
+void UNiagaraNodeStaticSwitch::DestroyNode()
+{
+	GetNiagaraGraph()->RemoveParameter(FNiagaraVariable(GetInputType(), InputParameterName), true);
+	Super::DestroyNode();
+}
+
 FNiagaraTypeDefinition UNiagaraNodeStaticSwitch::GetInputType() const
 {
 	if (SwitchTypeData.SwitchType == ENiagaraStaticSwitchType::Bool)
@@ -33,13 +39,16 @@ void UNiagaraNodeStaticSwitch::ChangeSwitchParameterName(const FName& NewName)
 {
 	FNiagaraVariable OldValue(GetInputType(), InputParameterName);
 	InputParameterName = NewName;
+	GetNiagaraGraph()->RenameParameter(OldValue, NewName, true);
 	VisualsChangedDelegate.Broadcast(this);
 	RemoveUnusedGraphParameter(OldValue);	
 }
 
 void UNiagaraNodeStaticSwitch::OnSwitchParameterTypeChanged(const FNiagaraTypeDefinition& OldType)
 {
-	RefreshFromExternalChanges();
+	TOptional<FNiagaraVariableMetaData> OldMetaData = GetNiagaraGraph()->GetMetaData(FNiagaraVariable(OldType, InputParameterName));
+	RefreshFromExternalChanges(); // Magick happens here: The old pins are destroyed and new ones are created.
+	GetNiagaraGraph()->SetMetaData(FNiagaraVariable(GetInputType(), InputParameterName), OldMetaData.GetValue());
 	VisualsChangedDelegate.Broadcast(this);
 	RemoveUnusedGraphParameter(FNiagaraVariable(OldType, InputParameterName));
 }
@@ -92,7 +101,7 @@ void UNiagaraNodeStaticSwitch::RemoveUnusedGraphParameter(const FNiagaraVariable
 	int Index = GraphVariables.Find(OldParameter);
 	if (Index == INDEX_NONE)
 	{
-		GetNiagaraGraph()->RemoveParameter(OldParameter);
+		GetNiagaraGraph()->RemoveParameter(OldParameter, true);
 	}
 	else
 	{
@@ -120,6 +129,7 @@ void UNiagaraNodeStaticSwitch::AllocateDefaultPins()
 			FEdGraphPinType PinType = Schema->TypeDefinitionToPinType(Var.GetType());
 			CreatePin(EGPD_Input, PinType, *(Var.GetName().ToString() + TEXT(" if false")));
 		}
+		GetNiagaraGraph()->AddParameter(FNiagaraVariable(GetInputType(), InputParameterName), true);
 	}
 	else if (SwitchTypeData.SwitchType == ENiagaraStaticSwitchType::Integer)
 	{
@@ -132,6 +142,7 @@ void UNiagaraNodeStaticSwitch::AllocateDefaultPins()
 				CreatePin(EGPD_Input, PinType, *(Var.GetName().ToString() + TEXT(" if ") + FString::FromInt(i)));
 			}
 		}
+		GetNiagaraGraph()->AddParameter(FNiagaraVariable(GetInputType(), InputParameterName), true);
 	}
 	else if (SwitchTypeData.SwitchType == ENiagaraStaticSwitchType::Enum && SwitchTypeData.Enum)
 	{
@@ -145,6 +156,7 @@ void UNiagaraNodeStaticSwitch::AllocateDefaultPins()
 				CreatePin(EGPD_Input, PinType, *(Var.GetName().ToString() + TEXT(" if ") + EnumName.ToString()));
 			}
 		}
+		GetNiagaraGraph()->AddParameter(FNiagaraVariable(GetInputType(), InputParameterName), true);
 	}
 	
 
@@ -158,7 +170,7 @@ void UNiagaraNodeStaticSwitch::AllocateDefaultPins()
 
 	CreateAddPin(EGPD_Output);
 
-	// force the graph to refresh the metadata
+	// force the graph to refresh the metadata 
 	GetNiagaraGraph()->GetParameterReferenceMap();
 }
 
diff --git a/Engine/Plugins/FX/Niagara/Source/NiagaraEditor/Private/NiagaraNodeStaticSwitch.h b/Engine/Plugins/FX/Niagara/Source/NiagaraEditor/Private/NiagaraNodeStaticSwitch.h
--- a/Engine/Plugins/FX/Niagara/Source/NiagaraEditor/Private/NiagaraNodeStaticSwitch.h
+++ b/Engine/Plugins/FX/Niagara/Source/NiagaraEditor/Private/NiagaraNodeStaticSwitch.h
@@ -80,6 +80,8 @@ public:
 	virtual void BuildParameterMapHistory(FNiagaraParameterMapHistoryBuilder& OutHistory, bool bRecursive = true, bool bFilterForCompilation = true) const override;
 	//~ End UNiagaraNode Interface
 
+	virtual void DestroyNode() override;
+
 protected:
 	//~ Begin UNiagaraNodeUsageSelector Interface
 	virtual void InsertInputPinsFor(const FNiagaraVariable& Var) override;
diff --git a/Engine/Plugins/FX/Niagara/Source/NiagaraEditor/Private/NiagaraNodeWithDynamicPins.cpp b/Engine/Plugins/FX/Niagara/Source/NiagaraEditor/Private/NiagaraNodeWithDynamicPins.cpp
--- a/Engine/Plugins/FX/Niagara/Source/NiagaraEditor/Private/NiagaraNodeWithDynamicPins.cpp
+++ b/Engine/Plugins/FX/Niagara/Source/NiagaraEditor/Private/NiagaraNodeWithDynamicPins.cpp
@@ -14,6 +14,7 @@
 #include "NiagaraEditorUtilities.h"
 #include "SNiagaraParameterMapView.h"
 #include "Framework/Application/SlateApplication.h"
+#include "NiagaraNodeParameterMapBase.h"
 
 #define LOCTEXT_NAMESPACE "NiagaraNodeWithDynamicPins"
 
@@ -255,7 +256,27 @@ void UNiagaraNodeWithDynamicPins::CollectAddPinActions(FGraphActionListBuilderBa
 void UNiagaraNodeWithDynamicPins::AddParameter(FNiagaraVariable Parameter, UEdGraphPin* AddPin)
 {
 	FScopedTransaction AddNewPinTransaction(LOCTEXT("AddNewPinTransaction", "Add pin to node"));
-	this->RequestNewTypedPin(AddPin->Direction, Parameter.GetType(), Parameter.GetName());
+	{
+		// We need to create new UNiagaraScriptVariable instances only when creating a new variable
+		if (UNiagaraNodeParameterMapBase* Map = Cast<UNiagaraNodeParameterMapBase>(this))
+		{
+			if (UNiagaraGraph* Graph = GetNiagaraGraph())
+			{	
+				UNiagaraScriptVariable* Variable = Graph->GetScriptVariable(Parameter);
+				if (Variable == nullptr)
+				{
+					TSet<FName> Names;
+					for (const auto& ParameterElement : Graph->GetParameterReferenceMap())
+					{
+						Names.Add(ParameterElement.Key.GetName());
+					}
+					Parameter.SetName(FNiagaraUtilities::GetUniqueName(Parameter.GetName(), Names));
+					Graph->AddParameter(Parameter);
+				}
+			}
+		}
+	}
+	UEdGraphPin* Pin = this->RequestNewTypedPin(AddPin->Direction, Parameter.GetType(), Parameter.GetName());
 }
 
 void UNiagaraNodeWithDynamicPins::RemoveDynamicPin(UEdGraphPin* Pin)
diff --git a/Engine/Plugins/FX/Niagara/Source/NiagaraEditor/Private/NiagaraScriptVariable.h b/Engine/Plugins/FX/Niagara/Source/NiagaraEditor/Private/NiagaraScriptVariable.h
--- a/Engine/Plugins/FX/Niagara/Source/NiagaraEditor/Private/NiagaraScriptVariable.h
+++ b/Engine/Plugins/FX/Niagara/Source/NiagaraEditor/Private/NiagaraScriptVariable.h
@@ -5,15 +5,20 @@
 #include "NiagaraTypes.h"
 #include "NiagaraScriptVariable.generated.h"
 
+/*
+* Used to store variable data and metadata per graph. 
+*/
 UCLASS()
 class UNiagaraScriptVariable : public UObject
 {
 public:
 	GENERATED_UCLASS_BODY()
 
+	/** Variable type, name and data. The data is not persistent, but used as a buffer when interfacing elsewhere. */
 	UPROPERTY()
 	FNiagaraVariable Variable;
 
-	UPROPERTY(EditAnywhere, Category = Metadata, meta=(NoResetToDefault, ShowOnlyInnerProperties))
+	/** The metadata associated with this script variable. */
+	UPROPERTY(EditAnywhere, Category = "Script Variable Data")
 	FNiagaraVariableMetaData Metadata;
 };
\ No newline at end of file
diff --git a/Engine/Plugins/FX/Niagara/Source/NiagaraEditor/Private/Widgets/SNiagaraParameterMapView.cpp b/Engine/Plugins/FX/Niagara/Source/NiagaraEditor/Private/Widgets/SNiagaraParameterMapView.cpp
--- a/Engine/Plugins/FX/Niagara/Source/NiagaraEditor/Private/Widgets/SNiagaraParameterMapView.cpp
+++ b/Engine/Plugins/FX/Niagara/Source/NiagaraEditor/Private/Widgets/SNiagaraParameterMapView.cpp
@@ -437,27 +437,25 @@ void SNiagaraParameterMapView::OnActionSelected(const TArray< TSharedPtr<FEdGrap
 {
 	if (!IsScriptToolkit())
 	{
+		// Don't accept any input for SystemToolkits, as there's no parameters panel there
 		return;
 	}
 	
 	// TODO: Can there be multiple actions and graphs? 
-	if (InActions.Num() == 1 && InActions[0].IsValid() && Graphs[0].IsValid()) 
+	if (InActions.Num() == 1 && InActions[0].IsValid() && Graphs.Num() > 0 && Graphs[0].IsValid()) 
 	{
-		FNiagaraParameterAction* Action = (FNiagaraParameterAction*)InActions[0].Get();
-		if (Action && (Action->Parameter.IsInNameSpace(TEXT("Module")) || IsStaticSwitchParameter(Action->Parameter, Graphs)))
+		if (FNiagaraParameterAction* Action = (FNiagaraParameterAction*)InActions[0].Get())
 		{
-			// Ignore any non-Module variables, as they don't persist, i.e. they get
-			// purged when any variable is renamed due to not being in ParameterToReferencesMap
-			TMap<FNiagaraVariable, UNiagaraScriptVariable*>& AllMetaData = Graphs[0]->GetAllMetaData();
-			if (UNiagaraScriptVariable** FoundScriptVariable = AllMetaData.Find(Action->Parameter))
+			if (UNiagaraScriptVariable* Variable = Graphs[0]->GetScriptVariable(Action->Parameter))
 			{
-				SelectedVariableObjects->SetSelectedObject(*FoundScriptVariable);
+				SelectedVariableObjects->SetSelectedObject(Variable);
 				return;
 			}
 		}
 	} 
 	
 	// If a variable wasn't selected just clear the current selection
+	// TODO: Get proper clearing to work. Current there's no way to clear while clicking on an empty location in the graph area
 	SelectedVariableObjects->ClearSelectedObjects();
 }
 
diff --git a/Engine/Plugins/FX/Niagara/Source/NiagaraEditor/Public/NiagaraGraph.h b/Engine/Plugins/FX/Niagara/Source/NiagaraEditor/Public/NiagaraGraph.h
--- a/Engine/Plugins/FX/Niagara/Source/NiagaraEditor/Public/NiagaraGraph.h
+++ b/Engine/Plugins/FX/Niagara/Source/NiagaraEditor/Public/NiagaraGraph.h
@@ -225,19 +225,22 @@ class UNiagaraGraph : public UEdGraph
 	const TMap<FNiagaraVariable, UNiagaraScriptVariable*>& GetAllMetaData() const;
 	TMap<FNiagaraVariable, UNiagaraScriptVariable*>& GetAllMetaData();
 
-	const TMap<FNiagaraVariable, FNiagaraGraphParameterReferenceCollection>& GetParameterReferenceMap() const;
+	const TMap<FNiagaraVariable, FNiagaraGraphParameterReferenceCollection>& GetParameterReferenceMap() const; // NOTE: The const is a lie! (This indirectly calls RefreshParameterReferences, which can recreate the entire map)
+
+	UNiagaraScriptVariable* GetScriptVariable(FNiagaraVariable Parameter) const;
+	UNiagaraScriptVariable* GetScriptVariable(FName ParameterName) const;
 
 	/** Adds parameter to parameters map setting it as created by the user.*/
-	void AddParameter(const FNiagaraVariable& Parameter);
+	void AddParameter(const FNiagaraVariable& Parameter, bool bIsStaticSwitch = false);
 
 	/** Adds parameter to parameters map setting it as created by the user.*/
 	void AddParameterReference(const FNiagaraVariable& Parameter, const UEdGraphPin* Pin);
 
 	/** Remove parameter from map and all the pins associated. */
-	void RemoveParameter(const FNiagaraVariable& Parameter);
+	void RemoveParameter(const FNiagaraVariable& Parameter, bool bFromStaticSwitch = false);
 
 	/** Rename parameter from map and all the pins associated. */
-	bool RenameParameter(const FNiagaraVariable& Parameter, FName NewName);
+	bool RenameParameter(const FNiagaraVariable& Parameter, FName NewName, bool bFromStaticSwitch = false);
 
 	/** Gets a delegate which is called whenever a contained data interfaces changes. */
 	FOnDataInterfaceChanged& OnDataInterfaceChanged();
@@ -274,9 +277,6 @@ protected:
 	void ResolveNumerics(TMap<UNiagaraNode*, bool>& VisitedNodes, UEdGraphNode* Node);
 
 private:
-	/** Remove any meta-data that is no longer being referenced within this graph.*/
-	void PurgeUnreferencedMetaData() const;
-
 	virtual void NotifyGraphChanged(const FEdGraphEditAction& InAction) override;
 
 	/** Find parameters in the graph. */
@@ -318,7 +318,5 @@ private:
 	bool bIsRenamingParameter;
 
 	mutable bool bParameterReferenceRefreshPending;
-
-	mutable bool bUnreferencedMetaDataPurgePending;
 };
 
