Bugfix: Removed unimplemented Undo/Redo declarations in FNiagaraScriptVariableDetails

diff --git a/Engine/Plugins/FX/Niagara/Source/NiagaraEditor/Private/Customizations/NiagaraScriptVariableDetails.h b/Engine/Plugins/FX/Niagara/Source/NiagaraEditor/Private/Customizations/NiagaraScriptVariableDetails.h
--- a/Engine/Plugins/FX/Niagara/Source/NiagaraEditor/Private/Customizations/NiagaraScriptVariableDetails.h
+++ b/Engine/Plugins/FX/Niagara/Source/NiagaraEditor/Private/Customizations/NiagaraScriptVariableDetails.h
@@ -14,16 +14,17 @@ public:
 	/** Makes a new instance of this detail layout class for a specific detail view requesting it */
 	static TSharedRef<IDetailCustomization> MakeInstance();
 
-	~FNiagaraScriptVariableDetails();
+	//~FNiagaraScriptVariableDetails();
 	FNiagaraScriptVariableDetails();
 
 	/** IDetailCustomization interface */
 	virtual void CustomizeDetails( IDetailLayoutBuilder& DetailBuilder ) override;
-	
+	/*
 	//~ Begin FEditorUndoClient Interface
 	virtual void PostUndo(bool bSuccess) override;
 	virtual void PostRedo(bool bSuccess) override;
 	// End of FEditorUndoClient
+	*/
 
 private:
 	void OnValueChanged();
