- Clean-up confusion around rebinding parameters / data interfaces
- Reset method no longer takes in the bBindParams as this is determined based on the EResetMode & the state of the system

#tests Ran EngineTests, deactivate then activating a system with a data interface bound (both complete and incomplete)

diff --git a/Engine/Plugins/FX/Niagara/Source/Niagara/Private/NiagaraComponent.cpp b/Engine/Plugins/FX/Niagara/Source/Niagara/Private/NiagaraComponent.cpp
--- a/Engine/Plugins/FX/Niagara/Source/Niagara/Private/NiagaraComponent.cpp
+++ b/Engine/Plugins/FX/Niagara/Source/Niagara/Private/NiagaraComponent.cpp
@@ -1407,7 +1407,7 @@ void UNiagaraComponent::SetParameterValueOverriddenLocally(const FNiagaraVariabl
 	
 	if (bRequiresSystemInstanceReset && SystemInstance)
 	{
-		SystemInstance->Reset(FNiagaraSystemInstance::EResetMode::ResetAll, true);
+		SystemInstance->Reset(FNiagaraSystemInstance::EResetMode::ResetAll);
 	}
 	
 }
diff --git a/Engine/Plugins/FX/Niagara/Source/Niagara/Private/NiagaraSystemInstance.cpp b/Engine/Plugins/FX/Niagara/Source/Niagara/Private/NiagaraSystemInstance.cpp
--- a/Engine/Plugins/FX/Niagara/Source/Niagara/Private/NiagaraSystemInstance.cpp
+++ b/Engine/Plugins/FX/Niagara/Source/Niagara/Private/NiagaraSystemInstance.cpp
@@ -79,7 +79,7 @@ void FNiagaraSystemInstance::Init(bool bInForceSolo)
 	// In order to get user data interface parameters in the component to work properly,
 	// we need to bind here, otherwise the instances when we init data interfaces during reset will potentially
 	// be the defaults (i.e. null) for things like static mesh data interfaces.
-	Reset(EResetMode::ReInit, true);
+	Reset(EResetMode::ReInit);
 
 #if WITH_EDITORONLY_DATA
 	InstanceParameters.DebugName = *FString::Printf(TEXT("SystemInstance %p"), this);
@@ -308,7 +308,7 @@ void FNiagaraSystemInstance::Activate(EResetMode InResetMode)
 	UNiagaraSystem* System = GetSystem();
 	if (System && System->IsValid() && IsReadyToRun())
 	{
-		Reset(InResetMode, true);
+		Reset(InResetMode);
 	}
 	else
 	{
@@ -445,7 +445,7 @@ void FNiagaraSystemInstance::SetPaused(bool bInPaused)
 	bPaused = bInPaused;
 }
 
-void FNiagaraSystemInstance::Reset(FNiagaraSystemInstance::EResetMode Mode, bool bBindParams)
+void FNiagaraSystemInstance::Reset(FNiagaraSystemInstance::EResetMode Mode)
 {
 	SCOPE_CYCLE_COUNTER(STAT_NiagaraSystemReset);
 
@@ -481,6 +481,9 @@ void FNiagaraSystemInstance::Reset(FNiagaraSystemInstance::EResetMode Mode, bool
 		Mode = EResetMode::ReInit;
 	}
 		
+	// Depending on the rest mode we may need to bind or can possibly skip it
+	// We must bind if we were previously complete as unbind will have been called, we can not get here if the system was disabled
+	bool bBindParams = IsComplete();
 	if (Mode == EResetMode::ResetSystem)
 	{
 		//UE_LOG(LogNiagara, Log, TEXT("FNiagaraSystemInstance::Reset false"));
@@ -491,15 +494,16 @@ void FNiagaraSystemInstance::Reset(FNiagaraSystemInstance::EResetMode Mode, bool
 		//UE_LOG(LogNiagara, Log, TEXT("FNiagaraSystemInstance::Reset true"));
 		ResetInternal(true);
 
-		bBindParams |= !IsComplete() == false;
+		// ResetInternal will not set the execution state but if we are neither complete or disabled we require a rebind
+		bBindParams = !IsComplete();
 	}
 	else if (Mode == EResetMode::ReInit)
 	{
 		//UE_LOG(LogNiagara, Log, TEXT("FNiagaraSystemInstance::ReInit"));
 		ReInitInternal();
 
-		// If the system was reinitialized successfully than we need to force a rebind of the parameters.
-		bBindParams = IsComplete() == false;
+		// ReInitInternal will set the execution state so if we are neither complete or disabled we require a rebind
+		bBindParams = !IsComplete();
 	}
 	
 	if (bBindParams)
@@ -510,7 +514,6 @@ void FNiagaraSystemInstance::Reset(FNiagaraSystemInstance::EResetMode Mode, bool
 	SetRequestedExecutionState(ENiagaraExecutionState::Active);
 	SetActualExecutionState(ENiagaraExecutionState::Active);
 
-	// We avoid calling InitDataInterfaces in the ResetSystem path so as to not clear out the System's DI on this frame.
 	if (bBindParams)
 	{
 		InitDataInterfaces();
diff --git a/Engine/Plugins/FX/Niagara/Source/Niagara/Public/NiagaraSystemInstance.h b/Engine/Plugins/FX/Niagara/Source/Niagara/Public/NiagaraSystemInstance.h
--- a/Engine/Plugins/FX/Niagara/Source/Niagara/Public/NiagaraSystemInstance.h
+++ b/Engine/Plugins/FX/Niagara/Source/Niagara/Public/NiagaraSystemInstance.h
@@ -78,7 +78,7 @@ public:
 	bool RequiresDistanceFieldData() const;
 
 	/** Requests the the simulation be reset on the next tick. */
-	void Reset(EResetMode Mode, bool bBindParams = false);
+	void Reset(EResetMode Mode);
 
 	void ComponentTick(float DeltaSeconds);
 	void PreSimulateTick(float DeltaSeconds);
