- Use bUseTriangleSampling & bUseVertexSampling to validate the interface requires CPU access or not
- Move the check for CPU access inside of here to ensure we only disable the DI when access is required

#tests Ran Niagara EngineTest and contrail system on a cooked platform & editor

diff --git a/Engine/Plugins/FX/Niagara/Source/Niagara/Classes/NDISkeletalMeshCommon.h b/Engine/Plugins/FX/Niagara/Source/Niagara/Classes/NDISkeletalMeshCommon.h
--- a/Engine/Plugins/FX/Niagara/Source/Niagara/Classes/NDISkeletalMeshCommon.h
+++ b/Engine/Plugins/FX/Niagara/Source/Niagara/Classes/NDISkeletalMeshCommon.h
@@ -325,7 +325,6 @@ struct TAreaWeightingModeBinder
 		else
 		{
 			int32 LODIndex = InstData->GetLODIndex();
-			check(InstData->Mesh->GetLODInfo(LODIndex)->bAllowCPUAccess);
 			bAreaWeighting = InstData->Mesh->GetLODInfo(LODIndex)->bSupportUniformlyDistributedSampling;
 		}
 
diff --git a/Engine/Plugins/FX/Niagara/Source/Niagara/Private/NiagaraDataInterfaceSkeletalMesh.cpp b/Engine/Plugins/FX/Niagara/Source/Niagara/Private/NiagaraDataInterfaceSkeletalMesh.cpp
--- a/Engine/Plugins/FX/Niagara/Source/Niagara/Private/NiagaraDataInterfaceSkeletalMesh.cpp
+++ b/Engine/Plugins/FX/Niagara/Source/Niagara/Private/NiagaraDataInterfaceSkeletalMesh.cpp
@@ -1033,16 +1033,6 @@ bool FNDISkeletalMesh_InstanceData::Init(UNiagaraDataInterfaceSkeletalMesh* Inte
 		{
 			LODIndex = FMath::Clamp(Interface->WholeMeshLOD, 0, Mesh->GetLODNum() - 1);
 		}
-
-		if (!Mesh->GetLODInfo(LODIndex)->bAllowCPUAccess && (Interface->bUseTriangleSampling || Interface->bUseVertexSampling))
-		{
-			UE_LOG(LogNiagara, Warning, TEXT("Skeletal Mesh Data Interface is trying to spawn from a whole mesh that does not allow CPU Access.\nInterface: %s\nMesh: %s\nLOD: %d"),
-				*Interface->GetFullName(),
-				*Mesh->GetFullName(),
-				LODIndex);
-
-			return false;
-		}
 	}
 	else
 	{
@@ -1144,9 +1134,18 @@ bool FNDISkeletalMesh_InstanceData::Init(UNiagaraDataInterfaceSkeletalMesh* Inte
 	FSkeletalMeshLODRenderData& LODData = GetLODRenderDataAndSkinWeights(SkinWeightBuffer);
 
 	//Check for the validity of the Mesh's cpu data.
-	//-TODO: These values are currently created too late in the editor so we can't use them to validate the DI correctly
-	//if ( Interface->bUseTriangleSampling || Interface->bUseVertexSampling )
+	if ( Interface->bUseTriangleSampling || Interface->bUseVertexSampling )
 	{
+		if ( !Mesh->GetLODInfo(LODIndex)->bAllowCPUAccess )
+		{
+			UE_LOG(LogNiagara, Warning, TEXT("Skeletal Mesh Data Interface is trying to spawn from a whole mesh that does not allow CPU Access.\nInterface: %s\nMesh: %s\nLOD: %d"),
+				*Interface->GetFullName(),
+				*Mesh->GetFullName(),
+				LODIndex);
+
+			return false;
+		}
+
 		bool LODDataNumVerticesCorrect = LODData.GetNumVertices() > 0;
 		bool LODDataPositonNumVerticesCorrect = LODData.StaticVertexBuffers.PositionVertexBuffer.GetNumVertices() > 0;
 		bool bSkinWeightBuffer = SkinWeightBuffer != nullptr;
