Back out changelist 7112046. Nonunity errors

diff --git a/Engine/Plugins/FX/Niagara/Source/Niagara/Public/NiagaraTypes.h b/Engine/Plugins/FX/Niagara/Source/Niagara/Public/NiagaraTypes.h
--- a/Engine/Plugins/FX/Niagara/Source/Niagara/Public/NiagaraTypes.h
+++ b/Engine/Plugins/FX/Niagara/Source/Niagara/Public/NiagaraTypes.h
@@ -345,24 +345,22 @@ public:
 	bool bInlineEditConditionToggle;
 
 	/** Declares the associated input should be conditionally editable based on the value of another input. */
-	UPROPERTY(EditAnywhere, Category = "Input Conditions")
+	UPROPERTY(EditAnywhere, Category = "Input Conditions", meta = (EditCondition = "!bIsStaticSwitch"))
 	FNiagaraInputConditionMetadata EditCondition;
 
 	/** Declares the associated input should be conditionally visible based on the value of another input. */
-	UPROPERTY(EditAnywhere, Category = "Input Conditions")
+	UPROPERTY(EditAnywhere, Category = "Input Conditions", meta = (EditCondition = "!bIsStaticSwitch"))
 	FNiagaraInputConditionMetadata VisibleCondition;
 
 	UPROPERTY(EditAnywhere, Category = "Variable", DisplayName = "Property Metadata", meta = (ToolTip = "Property Metadata"))
 	TMap<FName, FString> PropertyMetaData;
-	
-	/** This is a read-only variable that designates if the metadata is tied to a static switch or not. */
-	UPROPERTY()
-	bool bIsStaticSwitch; // TODO: This should be moved to the UNiagaraScriptVariable in the future
+
+	UPROPERTY(AdvancedDisplay, VisibleAnywhere, Category = "Variable", meta = (ToolTip = "This is a read-only variable that designates if the metadata is tied to a static switch or not."))
+	bool bIsStaticSwitch;
 
 	/** The default value to use when creating new pins or stack entries for a static switch parameter */
 	UPROPERTY()
-	int32 StaticSwitchDefaultValue; // TODO: This should be moved to the UNiagaraScriptVariable in the future
-
+	int32 StaticSwitchDefaultValue;
 };
 
 USTRUCT()
diff --git a/Engine/Plugins/FX/Niagara/Source/NiagaraEditor/Private/Customizations/NiagaraScriptVariableDetails.cpp b/Engine/Plugins/FX/Niagara/Source/NiagaraEditor/Private/Customizations/NiagaraScriptVariableDetails.cpp
deleted file mode 100644
diff --git a/Engine/Plugins/FX/Niagara/Source/NiagaraEditor/Private/Customizations/NiagaraScriptVariableDetails.h b/Engine/Plugins/FX/Niagara/Source/NiagaraEditor/Private/Customizations/NiagaraScriptVariableDetails.h
deleted file mode 100644
diff --git a/Engine/Plugins/FX/Niagara/Source/NiagaraEditor/Private/NiagaraEditorModule.cpp b/Engine/Plugins/FX/Niagara/Source/NiagaraEditor/Private/NiagaraEditorModule.cpp
--- a/Engine/Plugins/FX/Niagara/Source/NiagaraEditor/Private/NiagaraEditorModule.cpp
+++ b/Engine/Plugins/FX/Niagara/Source/NiagaraEditor/Private/NiagaraEditorModule.cpp
@@ -94,7 +94,6 @@
 #include "Misc/FeedbackContext.h"
 #include "Customizations/NiagaraStaticSwitchNodeDetails.h"
 #include "Customizations/NiagaraFunctionCallNodeDetails.h"
-#include "Customizations/NiagaraScriptVariableDetails.h"
 
 IMPLEMENT_MODULE( FNiagaraEditorModule, NiagaraEditor );
 
@@ -410,8 +409,6 @@ void FNiagaraEditorModule::StartupModule()
 	PropertyModule.RegisterCustomClassLayout("NiagaraComponent", FOnGetDetailCustomizationInstance::CreateStatic(&FNiagaraComponentDetails::MakeInstance));
 
 	PropertyModule.RegisterCustomClassLayout("NiagaraNodeStaticSwitch", FOnGetDetailCustomizationInstance::CreateStatic(&FNiagaraStaticSwitchNodeDetails::MakeInstance));
-	
-	PropertyModule.RegisterCustomClassLayout("NiagaraScriptVariable", FOnGetDetailCustomizationInstance::CreateStatic(&FNiagaraScriptVariableDetails::MakeInstance));
 
 	PropertyModule.RegisterCustomClassLayout("NiagaraNodeFunctionCall", FOnGetDetailCustomizationInstance::CreateStatic(&FNiagaraFunctionCallNodeDetails::MakeInstance));
 	
diff --git a/Engine/Plugins/FX/Niagara/Source/NiagaraEditor/Private/NiagaraGraph.cpp b/Engine/Plugins/FX/Niagara/Source/NiagaraEditor/Private/NiagaraGraph.cpp
--- a/Engine/Plugins/FX/Niagara/Source/NiagaraEditor/Private/NiagaraGraph.cpp
+++ b/Engine/Plugins/FX/Niagara/Source/NiagaraEditor/Private/NiagaraGraph.cpp
@@ -22,7 +22,6 @@
 #include "EdGraphSchema_Niagara.h"
 #include "NiagaraNodeParameterMapBase.h"
 #include "NiagaraNodeParameterMapGet.h"
-#include "NiagaraNodeParameterMapSet.h"
 #include "INiagaraEditorTypeUtilities.h"
 #include "NiagaraConstants.h"
 #include "NiagaraEditorModule.h"
@@ -81,6 +80,7 @@ UNiagaraGraph::UNiagaraGraph(const FObjectInitializer& ObjectInitializer)
 	, bNeedNumericCacheRebuilt(true)
 	, bIsRenamingParameter(false)
 	, bParameterReferenceRefreshPending(true)
+	, bUnreferencedMetaDataPurgePending(true)
 {
 	Schema = UEdGraphSchema_Niagara::StaticClass();
 	ChangeId = FGuid::NewGuid();
@@ -209,63 +209,27 @@ void UNiagaraGraph::PostLoad()
 	// Migrate input condition metadata
 	const int32 NiagaraVer = GetLinkerCustomVersion(FNiagaraCustomVersion::GUID);
 
-	FString FullPathName = GetPathName();
-	FString PathName;
-	int ColonPos;
-	if (FullPathName.FindChar(TCHAR('.'), ColonPos))
-	{
-		// GetPathName() returns something similar to "/Path/To/ScriptName.ScriptName:NiagaraScriptSource_N.NiagaraGraph_N"
-		// so this will extract "/Path/To/ScriptName"
-		PathName = FullPathName.Left(ColonPos);
-	}
-	// UE_LOG(LogNiagaraEditor, Log, TEXT("PostLoad %s"), *FullPathName);
-
 	if (NiagaraVer < FNiagaraCustomVersion::MetaDataAndParametersUpdate)
 	{
-		int NumMigrated = 0;
 		// If the version of the asset is older than FNiagaraCustomVersion::MetaDataAndParametersUpdate 
 		// we need to migrate the old metadata by looping through VariableToMetaData_DEPRECATED
 		// and create new entries in VariableToScriptVariable
-		for (auto& Pair : VariableToMetaData_DEPRECATED)
+		for (auto It = VariableToMetaData_DEPRECATED.CreateConstIterator(); It; ++It)
 		{
-			FNiagaraVariable Var(Pair.Key.GetType(), Pair.Key.GetName());
-			SetMetaData(Var, Pair.Value);
-			NumMigrated++;
-		}
-		VariableToMetaData_DEPRECATED.Empty();
+			SetMetaData(It.Key(), It.Value());
 
-		if (NumMigrated > 0)
-		{
-			// TODO: Should we even spam this? It will spam it until it's resaved
-			//UE_LOG(LogNiagaraEditor, Log, TEXT("Migrated %d old metadata entries for \"%s\""), NumMigrated, *FullPathName);
-		}
-	}
-	
-	{
-		// Create a UNiagaraScriptVariable instance for every entry in the parameter map for which there is no existing script variable. 
-
-		TArray<FNiagaraVariable> VarsToAdd;
-		for (auto& ParameterToReferences : GetParameterReferenceMap())
-		{
-			UNiagaraScriptVariable* Variable = GetScriptVariable(ParameterToReferences.Key.GetName());
-			if (Variable == nullptr)
+			FString PathName = GetPathName();
+			int ColonPos;
+			if (PathName.FindChar(TCHAR('.'), ColonPos))
 			{
-				VarsToAdd.Add(ParameterToReferences.Key);
+				// GetPathName() returns something similar to "/Path/To/ScriptName.ScriptName:NiagaraScriptSource_N.NiagaraGraph_N"
+				// so this will extract "/Path/To/ScriptName"
+				PathName = PathName.Left(ColonPos);
 			}
+			UE_LOG(LogNiagaraEditor, Log, TEXT("Migrated old metadata entry for variable \"%s\" in \"%s\""), *It.Key().GetName().ToString(), *PathName);
 		}
-
-		for (FNiagaraVariable& Var : VarsToAdd)
-		{
-			AddParameter(Var);
-		}
-
-		if (VarsToAdd.Num() > 0)
-		{
-			// TODO: Should we even spam this? It will spam it until it's resaved
-			//UE_LOG(LogNiagaraEditor, Log, TEXT("Added %d new script variables for \"%s\""), VarsToAdd.Num(), *FullPathName);
-		}
+		VariableToMetaData_DEPRECATED.Empty();
 	}
-
 	if (NiagaraVer < FNiagaraCustomVersion::MoveCommonInputMetadataToProperties)
 	{
 		auto MigrateInputCondition = [](TMap<FName, FString>& PropertyMetaData, const FName& InputConditionKey, FNiagaraInputConditionMetadata& InOutInputCondition)
@@ -339,6 +303,7 @@ FGuid UNiagaraGraph::ComputeCompileID(ENiagaraScriptUsage InUsage, const FGuid&
 	}
 
 	return FGuid();
+
 }
 
 FNiagaraCompileHash UNiagaraGraph::GetCompileDataHash(ENiagaraScriptUsage InUsage, const FGuid& InUsageId) const
@@ -391,7 +356,9 @@ UEdGraphPin* UNiagaraGraph::FindParameterMapDefaultValuePin(const FName Variable
 	UEdGraphPin* DefaultInputPin = nullptr;
 	for (UNiagaraNode* Node : NodesTraversed)
 	{
-		if (UNiagaraNodeParameterMapGet* GetNode = Cast<UNiagaraNodeParameterMapGet>(Node))
+		UNiagaraNodeParameterMapGet* GetNode = Cast<UNiagaraNodeParameterMapGet>(Node);
+	
+		if (GetNode)
 		{
 			TArray<UEdGraphPin*> OutputPins;
 			GetNode->GetOutputPins(OutputPins);
@@ -407,33 +374,15 @@ UEdGraphPin* UNiagaraGraph::FindParameterMapDefaultValuePin(const FName Variable
 					}
 				}
 			}
-			
-			if (DefaultInputPin != nullptr)
-			{
-				break;
-			}
 		}
-		
-		if (UNiagaraNodeParameterMapSet* SetNode = Cast<UNiagaraNodeParameterMapSet>(Node))
-		{
-			TArray<UEdGraphPin*> InputPins;
-			SetNode->GetInputPins(InputPins);
-			for (UEdGraphPin* InputPin : InputPins)
-			{
-				if (VariableName == InputPin->PinName)
-				{
-					DefaultInputPin = InputPin;
-					break;
-				}
-			}
 
-			if (DefaultInputPin != nullptr)
-			{
-				break;
-			}
+		if (DefaultInputPin != nullptr)
+		{
+			break;
 		}
 	}
 
+
 	// There are some pins 
 	if (DefaultInputPin && DefaultInputPin->LinkedTo.Num() != 0 && DefaultInputPin->LinkedTo[0] != nullptr)
 	{
@@ -479,6 +428,7 @@ void UNiagaraGraph::FindOutputNodes(TArray<UNiagaraNodeOutput*>& OutputNodes) co
 	}
 }
 
+
 void UNiagaraGraph::FindOutputNodes(ENiagaraScriptUsage TargetUsageType, TArray<UNiagaraNodeOutput*>& OutputNodes) const
 {
 	TArray<UNiagaraNodeOutput*> NodesFound;
@@ -541,6 +491,7 @@ UNiagaraNodeOutput* UNiagaraGraph::FindEquivalentOutputNode(ENiagaraScriptUsage
 	return nullptr;
 }
 
+
 void BuildTraversalHelper(TArray<class UNiagaraNode*>& OutNodesTraversed, UNiagaraNode* CurrentNode)
 {
 	if (CurrentNode == nullptr)
@@ -584,6 +535,7 @@ void UNiagaraGraph::BuildTraversal(TArray<class UNiagaraNode*>& OutNodesTraverse
 	}
 }
 
+
 void UNiagaraGraph::FindInputNodes(TArray<UNiagaraNodeInput*>& OutInputNodes, UNiagaraGraph::FFindInputNodeOptions Options) const
 {
 	SCOPE_CYCLE_COUNTER(STAT_NiagaraEditor_Graph_FindInputNodes);
@@ -743,47 +695,68 @@ TArray<UEdGraphNode*> UNiagaraGraph::FindReachbleNodes() const
 	return ResultNodes;
 }
 
-const TMap<FNiagaraVariable, UNiagaraScriptVariable*>& UNiagaraGraph::GetAllMetaData() const
+void UNiagaraGraph::GetParameters(TArray<FNiagaraVariable>& Inputs, TArray<FNiagaraVariable>& Outputs)const
 {
-	return VariableToScriptVariable;
-}
+	Inputs.Empty();
+	Outputs.Empty();
 
-TMap<FNiagaraVariable, UNiagaraScriptVariable*>& UNiagaraGraph::GetAllMetaData()
-{
-	return VariableToScriptVariable;
+	TArray<UNiagaraNodeInput*> InputsNodes;
+	FFindInputNodeOptions Options;
+	Options.bSort = true;
+	FindInputNodes(InputsNodes, Options);
+	for (UNiagaraNodeInput* Input : InputsNodes)
+	{
+		Inputs.Add(Input->Input);
+	}
+
+	TArray<UNiagaraNodeOutput*> OutputNodes;
+	FindOutputNodes(OutputNodes);
+	for (UNiagaraNodeOutput* OutputNode : OutputNodes)
+	{
+		for (FNiagaraVariable& Var : OutputNode->Outputs)
+		{
+			Outputs.AddUnique(Var);
+		}
+	}
+
+	//Do we need to sort outputs?
+	//Should leave them as they're defined in the output node?
+// 	auto SortVars = [](const FNiagaraVariable& A, const FNiagaraVariable& B)
+// 	{
+// 		//Case insensitive lexicographical comparisons of names for sorting.
+// 		return A.GetName().ToString() < B.GetName().ToString();
+// 	};
+// 	Outputs.Sort(SortVars);
 }
 
-const TMap<FNiagaraVariable, FNiagaraGraphParameterReferenceCollection>& UNiagaraGraph::GetParameterReferenceMap() const
+const TMap<FNiagaraVariable, UNiagaraScriptVariable*>& UNiagaraGraph::GetAllMetaData() const
 {
-	if (bParameterReferenceRefreshPending)
+	if (bUnreferencedMetaDataPurgePending)
 	{
-		RefreshParameterReferences();
+		PurgeUnreferencedMetaData();
 	}
-	return ParameterToReferencesMap;
+	return VariableToScriptVariable;
 }
 
-UNiagaraScriptVariable* UNiagaraGraph::GetScriptVariable(FNiagaraVariable Parameter) const
+TMap<FNiagaraVariable, UNiagaraScriptVariable*>& UNiagaraGraph::GetAllMetaData()
 {
-	if (UNiagaraScriptVariable** FoundScriptVariable = VariableToScriptVariable.Find(Parameter))
+	if (bUnreferencedMetaDataPurgePending)
 	{
-		return *FoundScriptVariable;
+		PurgeUnreferencedMetaData();
 	}
-	return nullptr;
+	return VariableToScriptVariable;
 }
 
-UNiagaraScriptVariable* UNiagaraGraph::GetScriptVariable(FName ParameterName) const
+const TMap<FNiagaraVariable, FNiagaraGraphParameterReferenceCollection>& UNiagaraGraph::GetParameterReferenceMap() const
 {
-	for (auto& VariableToScriptVariableItem : VariableToScriptVariable)
+	if (bParameterReferenceRefreshPending)
 	{
-		if (VariableToScriptVariableItem.Key.GetName() == ParameterName)
-		{
-			return VariableToScriptVariableItem.Value;
-		}
+		RefreshParameterReferences();
 	}
-	return nullptr;
+	return ParameterToReferencesMap;
 }
 
-void UNiagaraGraph::AddParameter(const FNiagaraVariable& Parameter, bool bIsStaticSwitch)
+void UNiagaraGraph::AddParameter(const FNiagaraVariable& Parameter)
 {
 	FNiagaraGraphParameterReferenceCollection* FoundParameterReferenceCollection = ParameterToReferencesMap.Find(Parameter);
 	if (!FoundParameterReferenceCollection)
@@ -798,12 +771,11 @@ void UNiagaraGraph::AddParameter(const FNiagaraVariable& Parameter, bool bIsStat
 	{
 		UNiagaraScriptVariable* NewScriptVariable = NewObject<UNiagaraScriptVariable>(this);
 		NewScriptVariable->Variable = Parameter;
-		NewScriptVariable->Metadata.bIsStaticSwitch = bIsStaticSwitch;
 		VariableToScriptVariable.Add(Parameter, NewScriptVariable);
 	}
 }
 
-void UNiagaraGraph::RemoveParameter(const FNiagaraVariable& Parameter, bool bFromStaticSwitch)
+void UNiagaraGraph::RemoveParameter(const FNiagaraVariable& Parameter)
 {
 	FNiagaraGraphParameterReferenceCollection* ReferenceCollection = ParameterToReferencesMap.Find(Parameter);
 	if (ReferenceCollection)
@@ -827,33 +799,9 @@ void UNiagaraGraph::RemoveParameter(const FNiagaraVariable& Parameter, bool bFro
 		ParameterToReferencesMap.Remove(Parameter);
 		NotifyGraphChanged();
 	}
-
-	TArray<FNiagaraVariable> VarsToRemove;
-	for (auto It : VariableToScriptVariable)
-	{
-		if (bFromStaticSwitch)
-		{
-			if (It.Key.GetName() == Parameter.GetName() && It.Key.GetType() == Parameter.GetType())
-			{
-				VarsToRemove.Add(It.Key);
-			}
-		}
-		else
-		{
-			if (It.Key.GetName() == Parameter.GetName())
-			{
-				VarsToRemove.Add(It.Key);
-			}
-		}
-		
-	}
-	for (FNiagaraVariable& Var : VarsToRemove)
-	{
-		VariableToScriptVariable.Remove(Var);
-	}
 }
 
-bool UNiagaraGraph::RenameParameter(const FNiagaraVariable& Parameter, FName NewName, bool bFromStaticSwitch)
+bool UNiagaraGraph::RenameParameter(const FNiagaraVariable& Parameter, FName NewName)
 {
 	// Block rename when already renaming. This prevents recursion when CommitEditablePinName is called on referenced nodes. 
 	if (bIsRenamingParameter)
@@ -870,20 +818,13 @@ bool UNiagaraGraph::RenameParameter(const FNiagaraVariable& Parameter, FName New
 	FNiagaraVariableMetaData OldMetaData;
 	if (OldScriptVariable)
 	{
-		if (!bFromStaticSwitch && (*OldScriptVariable)->Metadata.bIsStaticSwitch)
-		{
-			// We current disallow renaming static switch variables in the Parameters panel. 
-			bIsRenamingParameter = false;
-			return false;
-		}
 		OldMetaData = (*OldScriptVariable)->Metadata;
 	}
+		
 
 	FNiagaraGraphParameterReferenceCollection* ReferenceCollection = ParameterToReferencesMap.Find(Parameter);
-	int NumReferences = 0;
 	if (ReferenceCollection)
 	{
-		NumReferences = ReferenceCollection->ParameterReferences.Num();
 		const FText NewNameText = FText::FromName(NewName);
 		FNiagaraGraphParameterReferenceCollection NewReferences = *ReferenceCollection;
 		for (FNiagaraGraphParameterReference& Reference : NewReferences.ParameterReferences)
@@ -906,10 +847,7 @@ bool UNiagaraGraph::RenameParameter(const FNiagaraVariable& Parameter, FName New
 	// Swap metadata to the new parameter
 	if (UNiagaraScriptVariable** Metadata = VariableToScriptVariable.Find(Parameter))
 	{
-		if (!bFromStaticSwitch || (bFromStaticSwitch && NumReferences <= 1))
-		{
-			VariableToScriptVariable.Remove(Parameter);
-		}
+		VariableToScriptVariable.Remove(Parameter);
 	}
 	SetMetaData(NewParameter, OldMetaData);
 
@@ -952,40 +890,6 @@ void UNiagaraGraph::GetOutputNodeVariables(ENiagaraScriptUsage InScriptUsage, TA
 	}
 }
 
-void UNiagaraGraph::GetParameters(TArray<FNiagaraVariable>& Inputs, TArray<FNiagaraVariable>& Outputs)const
-{
-	Inputs.Empty();
-	Outputs.Empty();
-
-	TArray<UNiagaraNodeInput*> InputsNodes;
-	FFindInputNodeOptions Options;
-	Options.bSort = true;
-	FindInputNodes(InputsNodes, Options);
-	for (UNiagaraNodeInput* Input : InputsNodes)
-	{
-		Inputs.Add(Input->Input);
-	}
-
-	TArray<UNiagaraNodeOutput*> OutputNodes;
-	FindOutputNodes(OutputNodes);
-	for (UNiagaraNodeOutput* OutputNode : OutputNodes)
-	{
-		for (FNiagaraVariable& Var : OutputNode->Outputs)
-		{
-			Outputs.AddUnique(Var);
-		}
-	}
-
-	//Do we need to sort outputs?
-	//Should leave them as they're defined in the output node?
-	// 	auto SortVars = [](const FNiagaraVariable& A, const FNiagaraVariable& B)
-	// 	{
-	// 		//Case insensitive lexicographical comparisons of names for sorting.
-	// 		return A.GetName().ToString() < B.GetName().ToString();
-	// 	};
-	// 	Outputs.Sort(SortVars);
-}
-
 bool UNiagaraGraph::HasParameterMapParameters()const
 {
 	TArray<FNiagaraVariable> Inputs;
@@ -1043,6 +947,7 @@ void UNiagaraGraph::NotifyGraphNeedsRecompile()
 	NotifyGraphChanged(Action);
 }
 
+
 void UNiagaraGraph::NotifyGraphDataInterfaceChanged()
 {
 	OnDataInterfaceChangedDelegate.Broadcast();
@@ -1334,6 +1239,7 @@ void UNiagaraGraph::ResolveNumerics(TMap<UNiagaraNode*, bool>& VisitedNodes, UEd
 	}
 }
 
+
 void UNiagaraGraph::SynchronizeInternalCacheWithGraph(UNiagaraGraph* Other)
 {
 	// Force us to rebuild the cache, note that this builds traversals and everything else, keeping it in sync if nothing changed from the current version.
@@ -1389,6 +1295,7 @@ void UNiagaraGraph::SynchronizeInternalCacheWithGraph(UNiagaraGraph* Other)
 	}
 }
 
+
 void UNiagaraGraph::InvalidateCachedCompileIds()
 {
 	Modify();
@@ -1447,6 +1354,7 @@ void UNiagaraGraph::GatherExternalDependencyIDs(ENiagaraScriptUsage InUsage, con
 	}
 }
 
+
 void UNiagaraGraph::GetAllReferencedGraphs(TArray<const UNiagaraGraph*>& Graphs) const
 {
 	Graphs.AddUnique(this);
@@ -1510,6 +1418,11 @@ void UNiagaraGraph::MarkGraphRequiresSynchronization(FString Reason)
 
 TOptional<FNiagaraVariableMetaData> UNiagaraGraph::GetMetaData(const FNiagaraVariable& InVar) const
 {
+	if (bUnreferencedMetaDataPurgePending)
+	{
+		PurgeUnreferencedMetaData();
+	}
+	
 	if (UNiagaraScriptVariable** MetaData = VariableToScriptVariable.Find(InVar))
 	{
 		if (*MetaData)
@@ -1540,6 +1453,44 @@ void UNiagaraGraph::SetMetaData(const FNiagaraVariable& InVar, const FNiagaraVar
 	}
 }
 
+void UNiagaraGraph::PurgeUnreferencedMetaData() const
+{
+	TSet<FNiagaraVariable> ReferencedParameters;
+	ReferencedParameters.Append(FindStaticSwitchInputs());
+	const UEdGraphSchema_Niagara* NiagaraSchema = Cast<UEdGraphSchema_Niagara>(Schema);
+	for (UEdGraphNode* Node : Nodes)
+	{
+		for (UEdGraphPin* Pin : Node->Pins)
+		{
+			if (Pin->PinType.PinSubCategory == UNiagaraNodeParameterMapBase::ParameterPinSubCategory)
+			{
+				const FNiagaraVariable Parameter = NiagaraSchema->PinToNiagaraVariable(Pin, false);
+				const FNiagaraParameterHandle Handle = FNiagaraParameterHandle(Parameter.GetName());
+				if (Handle.IsModuleHandle() && !FNiagaraConstants::IsNiagaraConstant(Parameter))
+				{
+					ReferencedParameters.Add(Parameter);
+				}
+			}
+		}
+	}
+
+	TArray<FNiagaraVariable> VarsToRemove;
+	for (auto It = VariableToScriptVariable.CreateConstIterator(); It; ++It)
+	{
+		if (!ReferencedParameters.Contains(It.Key()))
+		{
+			VarsToRemove.Add(It.Key());
+		}
+	}
+
+	for (FNiagaraVariable& Var : VarsToRemove)
+	{
+		VariableToScriptVariable.Remove(Var);
+	}
+
+	bUnreferencedMetaDataPurgePending = false;
+}
+
 UNiagaraGraph::FOnDataInterfaceChanged& UNiagaraGraph::OnDataInterfaceChanged()
 {
 	return OnDataInterfaceChangedDelegate;
@@ -1574,6 +1525,18 @@ void UNiagaraGraph::RefreshParameterReferences() const
 				FNiagaraGraphParameterReferenceCollection& NewReferenceCollection = ParameterToReferencesMap.Add(Parameter);
 				NewReferenceCollection.Graph = this;
 				ReferenceCollection = &NewReferenceCollection;
+
+				// When a variable is created or added from the graph it won't call AddParameter, 
+				// but instead call this method
+				UNiagaraScriptVariable** FoundScriptVariable = VariableToScriptVariable.Find(Parameter);
+				if (!FoundScriptVariable)
+				{
+					// Warning! This method (RefreshParameterReferences) isn't const at all!
+					// Need to cast away the const to be able to create a serializable UObject here...
+					UNiagaraScriptVariable* NewScriptVariable = NewObject<UNiagaraScriptVariable>(const_cast<UNiagaraGraph*>(this));
+					NewScriptVariable->Variable = Parameter;
+					VariableToScriptVariable.Add(Parameter, NewScriptVariable);
+				}
 			}
 			ReferenceCollection->ParameterReferences.AddUnique(FNiagaraGraphParameterReference(Pin->PersistentGuid, Cast<UNiagaraNode>(Pin->GetOwningNode())));
 
@@ -1593,6 +1556,14 @@ void UNiagaraGraph::RefreshParameterReferences() const
 			NewReferenceCollection.Graph = this;
 			ReferenceCollection = &ParameterToReferencesMap.Add(Variable, NewReferenceCollection);
 		}
+		UNiagaraScriptVariable** FoundScriptVariable = VariableToScriptVariable.Find(Variable);
+		if (!FoundScriptVariable)
+		{
+			UNiagaraScriptVariable* NewScriptVariable = NewObject<UNiagaraScriptVariable>(const_cast<UNiagaraGraph*>(this));
+			NewScriptVariable->Variable = Variable;
+			NewScriptVariable->Metadata.bIsStaticSwitch = true;
+			VariableToScriptVariable.Add(Variable, NewScriptVariable);
+		}
 		ReferenceCollection->ParameterReferences.AddUnique(FNiagaraGraphParameterReference(Node->NodeGuid, Node));
 		CandidateUnreferencedParametersToRemove.Remove(Variable);
 	};
@@ -1650,31 +1621,6 @@ void UNiagaraGraph::RefreshParameterReferences() const
 	for (const FNiagaraVariable& UnreferencedParameterToRemove : CandidateUnreferencedParametersToRemove)
 	{
 		ParameterToReferencesMap.Remove(UnreferencedParameterToRemove);
-		VariableToScriptVariable.Remove(UnreferencedParameterToRemove);
-	}
-
-
-	auto AddMissingScriptVariable = [&](const FNiagaraVariable& Variable)
-	{
-		UNiagaraScriptVariable** FoundScriptVariable = VariableToScriptVariable.Find(Variable);
-		if (!FoundScriptVariable)
-		{
-			UNiagaraScriptVariable* NewScriptVariable = NewObject<UNiagaraScriptVariable>(const_cast<UNiagaraGraph*>(this));
-			NewScriptVariable->Variable = Variable;
-			NewScriptVariable->Metadata.bIsStaticSwitch = false;
-			VariableToScriptVariable.Add(Variable, NewScriptVariable);
-		}
-	};
-
-	for (auto& ParameterToReferences : ParameterToReferencesMap)
-	{
-		UNiagaraScriptVariable* Variable = GetScriptVariable(ParameterToReferences.Key);
-		if (Variable)
-		{
-			continue;
-		}
-
-		AddMissingScriptVariable(ParameterToReferences.Key);
 	}
 
 	bParameterReferenceRefreshPending = false;
@@ -1683,6 +1629,7 @@ void UNiagaraGraph::RefreshParameterReferences() const
 void UNiagaraGraph::InvalidateCachedParameterData()
 {
 	bParameterReferenceRefreshPending = true;
+	bUnreferencedMetaDataPurgePending = true;
 }
 
 #undef LOCTEXT_NAMESPACE
\ No newline at end of file
diff --git a/Engine/Plugins/FX/Niagara/Source/NiagaraEditor/Private/NiagaraNodeStaticSwitch.cpp b/Engine/Plugins/FX/Niagara/Source/NiagaraEditor/Private/NiagaraNodeStaticSwitch.cpp
--- a/Engine/Plugins/FX/Niagara/Source/NiagaraEditor/Private/NiagaraNodeStaticSwitch.cpp
+++ b/Engine/Plugins/FX/Niagara/Source/NiagaraEditor/Private/NiagaraNodeStaticSwitch.cpp
@@ -12,12 +12,6 @@ UNiagaraNodeStaticSwitch::UNiagaraNodeStaticSwitch(const FObjectInitializer& Obj
 {
 }
 
-void UNiagaraNodeStaticSwitch::DestroyNode()
-{
-	GetNiagaraGraph()->RemoveParameter(FNiagaraVariable(GetInputType(), InputParameterName), true);
-	Super::DestroyNode();
-}
-
 FNiagaraTypeDefinition UNiagaraNodeStaticSwitch::GetInputType() const
 {
 	if (SwitchTypeData.SwitchType == ENiagaraStaticSwitchType::Bool)
@@ -39,16 +33,13 @@ void UNiagaraNodeStaticSwitch::ChangeSwitchParameterName(const FName& NewName)
 {
 	FNiagaraVariable OldValue(GetInputType(), InputParameterName);
 	InputParameterName = NewName;
-	GetNiagaraGraph()->RenameParameter(OldValue, NewName, true);
 	VisualsChangedDelegate.Broadcast(this);
 	RemoveUnusedGraphParameter(OldValue);	
 }
 
 void UNiagaraNodeStaticSwitch::OnSwitchParameterTypeChanged(const FNiagaraTypeDefinition& OldType)
 {
-	TOptional<FNiagaraVariableMetaData> OldMetaData = GetNiagaraGraph()->GetMetaData(FNiagaraVariable(OldType, InputParameterName));
-	RefreshFromExternalChanges(); // Magick happens here: The old pins are destroyed and new ones are created.
-	GetNiagaraGraph()->SetMetaData(FNiagaraVariable(GetInputType(), InputParameterName), OldMetaData.GetValue());
+	RefreshFromExternalChanges();
 	VisualsChangedDelegate.Broadcast(this);
 	RemoveUnusedGraphParameter(FNiagaraVariable(OldType, InputParameterName));
 }
@@ -101,7 +92,7 @@ void UNiagaraNodeStaticSwitch::RemoveUnusedGraphParameter(const FNiagaraVariable
 	int Index = GraphVariables.Find(OldParameter);
 	if (Index == INDEX_NONE)
 	{
-		GetNiagaraGraph()->RemoveParameter(OldParameter, true);
+		GetNiagaraGraph()->RemoveParameter(OldParameter);
 	}
 	else
 	{
@@ -129,7 +120,6 @@ void UNiagaraNodeStaticSwitch::AllocateDefaultPins()
 			FEdGraphPinType PinType = Schema->TypeDefinitionToPinType(Var.GetType());
 			CreatePin(EGPD_Input, PinType, *(Var.GetName().ToString() + TEXT(" if false")));
 		}
-		GetNiagaraGraph()->AddParameter(FNiagaraVariable(GetInputType(), InputParameterName), true);
 	}
 	else if (SwitchTypeData.SwitchType == ENiagaraStaticSwitchType::Integer)
 	{
@@ -142,7 +132,6 @@ void UNiagaraNodeStaticSwitch::AllocateDefaultPins()
 				CreatePin(EGPD_Input, PinType, *(Var.GetName().ToString() + TEXT(" if ") + FString::FromInt(i)));
 			}
 		}
-		GetNiagaraGraph()->AddParameter(FNiagaraVariable(GetInputType(), InputParameterName), true);
 	}
 	else if (SwitchTypeData.SwitchType == ENiagaraStaticSwitchType::Enum && SwitchTypeData.Enum)
 	{
@@ -156,7 +145,6 @@ void UNiagaraNodeStaticSwitch::AllocateDefaultPins()
 				CreatePin(EGPD_Input, PinType, *(Var.GetName().ToString() + TEXT(" if ") + EnumName.ToString()));
 			}
 		}
-		GetNiagaraGraph()->AddParameter(FNiagaraVariable(GetInputType(), InputParameterName), true);
 	}
 	
 
@@ -170,7 +158,7 @@ void UNiagaraNodeStaticSwitch::AllocateDefaultPins()
 
 	CreateAddPin(EGPD_Output);
 
-	// force the graph to refresh the metadata 
+	// force the graph to refresh the metadata
 	GetNiagaraGraph()->GetParameterReferenceMap();
 }
 
diff --git a/Engine/Plugins/FX/Niagara/Source/NiagaraEditor/Private/NiagaraNodeStaticSwitch.h b/Engine/Plugins/FX/Niagara/Source/NiagaraEditor/Private/NiagaraNodeStaticSwitch.h
--- a/Engine/Plugins/FX/Niagara/Source/NiagaraEditor/Private/NiagaraNodeStaticSwitch.h
+++ b/Engine/Plugins/FX/Niagara/Source/NiagaraEditor/Private/NiagaraNodeStaticSwitch.h
@@ -80,8 +80,6 @@ public:
 	virtual void BuildParameterMapHistory(FNiagaraParameterMapHistoryBuilder& OutHistory, bool bRecursive = true, bool bFilterForCompilation = true) const override;
 	//~ End UNiagaraNode Interface
 
-	virtual void DestroyNode() override;
-
 protected:
 	//~ Begin UNiagaraNodeUsageSelector Interface
 	virtual void InsertInputPinsFor(const FNiagaraVariable& Var) override;
diff --git a/Engine/Plugins/FX/Niagara/Source/NiagaraEditor/Private/NiagaraNodeWithDynamicPins.cpp b/Engine/Plugins/FX/Niagara/Source/NiagaraEditor/Private/NiagaraNodeWithDynamicPins.cpp
--- a/Engine/Plugins/FX/Niagara/Source/NiagaraEditor/Private/NiagaraNodeWithDynamicPins.cpp
+++ b/Engine/Plugins/FX/Niagara/Source/NiagaraEditor/Private/NiagaraNodeWithDynamicPins.cpp
@@ -14,7 +14,6 @@
 #include "NiagaraEditorUtilities.h"
 #include "SNiagaraParameterMapView.h"
 #include "Framework/Application/SlateApplication.h"
-#include "NiagaraNodeParameterMapBase.h"
 
 #define LOCTEXT_NAMESPACE "NiagaraNodeWithDynamicPins"
 
@@ -256,27 +255,7 @@ void UNiagaraNodeWithDynamicPins::CollectAddPinActions(FGraphActionListBuilderBa
 void UNiagaraNodeWithDynamicPins::AddParameter(FNiagaraVariable Parameter, UEdGraphPin* AddPin)
 {
 	FScopedTransaction AddNewPinTransaction(LOCTEXT("AddNewPinTransaction", "Add pin to node"));
-	{
-		// We need to create new UNiagaraScriptVariable instances only when creating a new variable
-		if (UNiagaraNodeParameterMapBase* Map = Cast<UNiagaraNodeParameterMapBase>(this))
-		{
-			if (UNiagaraGraph* Graph = GetNiagaraGraph())
-			{	
-				UNiagaraScriptVariable* Variable = Graph->GetScriptVariable(Parameter);
-				if (Variable == nullptr)
-				{
-					TSet<FName> Names;
-					for (const auto& ParameterElement : Graph->GetParameterReferenceMap())
-					{
-						Names.Add(ParameterElement.Key.GetName());
-					}
-					Parameter.SetName(FNiagaraUtilities::GetUniqueName(Parameter.GetName(), Names));
-					Graph->AddParameter(Parameter);
-				}
-			}
-		}
-	}
-	UEdGraphPin* Pin = this->RequestNewTypedPin(AddPin->Direction, Parameter.GetType(), Parameter.GetName());
+	this->RequestNewTypedPin(AddPin->Direction, Parameter.GetType(), Parameter.GetName());
 }
 
 void UNiagaraNodeWithDynamicPins::RemoveDynamicPin(UEdGraphPin* Pin)
diff --git a/Engine/Plugins/FX/Niagara/Source/NiagaraEditor/Private/NiagaraScriptVariable.h b/Engine/Plugins/FX/Niagara/Source/NiagaraEditor/Private/NiagaraScriptVariable.h
--- a/Engine/Plugins/FX/Niagara/Source/NiagaraEditor/Private/NiagaraScriptVariable.h
+++ b/Engine/Plugins/FX/Niagara/Source/NiagaraEditor/Private/NiagaraScriptVariable.h
@@ -5,20 +5,15 @@
 #include "NiagaraTypes.h"
 #include "NiagaraScriptVariable.generated.h"
 
-/*
-* Used to store variable data and metadata per graph. 
-*/
 UCLASS()
 class UNiagaraScriptVariable : public UObject
 {
 public:
 	GENERATED_UCLASS_BODY()
 
-	/** Variable type, name and data. The data is not persistent, but used as a buffer when interfacing elsewhere. */
 	UPROPERTY()
 	FNiagaraVariable Variable;
 
-	/** The metadata associated with this script variable. */
-	UPROPERTY(EditAnywhere, Category = "Script Variable Data")
+	UPROPERTY(EditAnywhere, Category = Metadata, meta=(NoResetToDefault, ShowOnlyInnerProperties))
 	FNiagaraVariableMetaData Metadata;
 };
\ No newline at end of file
diff --git a/Engine/Plugins/FX/Niagara/Source/NiagaraEditor/Private/Widgets/SNiagaraParameterMapView.cpp b/Engine/Plugins/FX/Niagara/Source/NiagaraEditor/Private/Widgets/SNiagaraParameterMapView.cpp
--- a/Engine/Plugins/FX/Niagara/Source/NiagaraEditor/Private/Widgets/SNiagaraParameterMapView.cpp
+++ b/Engine/Plugins/FX/Niagara/Source/NiagaraEditor/Private/Widgets/SNiagaraParameterMapView.cpp
@@ -437,25 +437,27 @@ void SNiagaraParameterMapView::OnActionSelected(const TArray< TSharedPtr<FEdGrap
 {
 	if (!IsScriptToolkit())
 	{
-		// Don't accept any input for SystemToolkits, as there's no parameters panel there
 		return;
 	}
 	
 	// TODO: Can there be multiple actions and graphs? 
-	if (InActions.Num() == 1 && InActions[0].IsValid() && Graphs.Num() > 0 && Graphs[0].IsValid()) 
+	if (InActions.Num() == 1 && InActions[0].IsValid() && Graphs[0].IsValid()) 
 	{
-		if (FNiagaraParameterAction* Action = (FNiagaraParameterAction*)InActions[0].Get())
+		FNiagaraParameterAction* Action = (FNiagaraParameterAction*)InActions[0].Get();
+		if (Action && (Action->Parameter.IsInNameSpace(TEXT("Module")) || IsStaticSwitchParameter(Action->Parameter, Graphs)))
 		{
-			if (UNiagaraScriptVariable* Variable = Graphs[0]->GetScriptVariable(Action->Parameter))
+			// Ignore any non-Module variables, as they don't persist, i.e. they get
+			// purged when any variable is renamed due to not being in ParameterToReferencesMap
+			TMap<FNiagaraVariable, UNiagaraScriptVariable*>& AllMetaData = Graphs[0]->GetAllMetaData();
+			if (UNiagaraScriptVariable** FoundScriptVariable = AllMetaData.Find(Action->Parameter))
 			{
-				SelectedVariableObjects->SetSelectedObject(Variable);
+				SelectedVariableObjects->SetSelectedObject(*FoundScriptVariable);
 				return;
 			}
 		}
 	} 
 	
 	// If a variable wasn't selected just clear the current selection
-	// TODO: Get proper clearing to work. Current there's no way to clear while clicking on an empty location in the graph area
 	SelectedVariableObjects->ClearSelectedObjects();
 }
 
diff --git a/Engine/Plugins/FX/Niagara/Source/NiagaraEditor/Public/NiagaraGraph.h b/Engine/Plugins/FX/Niagara/Source/NiagaraEditor/Public/NiagaraGraph.h
--- a/Engine/Plugins/FX/Niagara/Source/NiagaraEditor/Public/NiagaraGraph.h
+++ b/Engine/Plugins/FX/Niagara/Source/NiagaraEditor/Public/NiagaraGraph.h
@@ -225,22 +225,19 @@ class UNiagaraGraph : public UEdGraph
 	const TMap<FNiagaraVariable, UNiagaraScriptVariable*>& GetAllMetaData() const;
 	TMap<FNiagaraVariable, UNiagaraScriptVariable*>& GetAllMetaData();
 
-	const TMap<FNiagaraVariable, FNiagaraGraphParameterReferenceCollection>& GetParameterReferenceMap() const; // NOTE: The const is a lie! (This indirectly calls RefreshParameterReferences, which can recreate the entire map)
-
-	UNiagaraScriptVariable* GetScriptVariable(FNiagaraVariable Parameter) const;
-	UNiagaraScriptVariable* GetScriptVariable(FName ParameterName) const;
+	const TMap<FNiagaraVariable, FNiagaraGraphParameterReferenceCollection>& GetParameterReferenceMap() const;
 
 	/** Adds parameter to parameters map setting it as created by the user.*/
-	void AddParameter(const FNiagaraVariable& Parameter, bool bIsStaticSwitch = false);
+	void AddParameter(const FNiagaraVariable& Parameter);
 
 	/** Adds parameter to parameters map setting it as created by the user.*/
 	void AddParameterReference(const FNiagaraVariable& Parameter, const UEdGraphPin* Pin);
 
 	/** Remove parameter from map and all the pins associated. */
-	void RemoveParameter(const FNiagaraVariable& Parameter, bool bFromStaticSwitch = false);
+	void RemoveParameter(const FNiagaraVariable& Parameter);
 
 	/** Rename parameter from map and all the pins associated. */
-	bool RenameParameter(const FNiagaraVariable& Parameter, FName NewName, bool bFromStaticSwitch = false);
+	bool RenameParameter(const FNiagaraVariable& Parameter, FName NewName);
 
 	/** Gets a delegate which is called whenever a contained data interfaces changes. */
 	FOnDataInterfaceChanged& OnDataInterfaceChanged();
@@ -277,6 +274,9 @@ protected:
 	void ResolveNumerics(TMap<UNiagaraNode*, bool>& VisitedNodes, UEdGraphNode* Node);
 
 private:
+	/** Remove any meta-data that is no longer being referenced within this graph.*/
+	void PurgeUnreferencedMetaData() const;
+
 	virtual void NotifyGraphChanged(const FEdGraphEditAction& InAction) override;
 
 	/** Find parameters in the graph. */
@@ -318,5 +318,7 @@ private:
 	bool bIsRenamingParameter;
 
 	mutable bool bParameterReferenceRefreshPending;
+
+	mutable bool bUnreferencedMetaDataPurgePending;
 };
 
