Fixed some issues with Niagara mesh renderer in ES31.
Some particle properties like particle position were not available for VS.

diff --git a/Engine/Plugins/FX/Niagara/Shaders/Private/NiagaraMeshVertexFactory.ush b/Engine/Plugins/FX/Niagara/Shaders/Private/NiagaraMeshVertexFactory.ush
--- a/Engine/Plugins/FX/Niagara/Shaders/Private/NiagaraMeshVertexFactory.ush
+++ b/Engine/Plugins/FX/Niagara/Shaders/Private/NiagaraMeshVertexFactory.ush
@@ -158,12 +158,12 @@ struct FVertexFactoryIntermediates
 	/** Particle position in camera-centered translated world space. */
 	float3 ParticleTranslatedWorldPosition;
 
-#if USE_PARTICLE_TIME
+#if NEEDS_PARTICLE_TIME
 	/** Relative time. */
 	float RelativeTime;
 #endif
 
-#if USE_PARTICLE_RANDOM
+#if NEEDS_PARTICLE_RANDOM
 	/** Particle Random value */
 	float ParticleRandom;
 #endif
@@ -307,20 +307,20 @@ FMaterialVertexParameters GetMaterialVertexParameters(FVertexFactoryInput Input,
 	Result.PreSkinnedPosition = Input.Position.xyz;
 	Result.PreSkinnedNormal = TangentBias(Input.TangentZ.xyz);
 
-#if USE_PARTICLE_POSITION
+#if NEEDS_PARTICLE_POSITION
 	Result.Particle.TranslatedWorldPositionAndSize.xyz = Intermediates.ParticleTranslatedWorldPosition.xyz;
 	Result.Particle.TranslatedWorldPositionAndSize.w = 1;
 #endif
 
-#if USE_PARTICLE_VELOCITY
+#if NEEDS_PARTICLE_VELOCITY
 	Result.Particle.Velocity = Intermediates.ParticleVelocity;
 #endif
 
-#if USE_PARTICLE_TIME
+#if NEEDS_PARTICLE_TIME
 	Result.Particle.RelativeTime = Intermediates.RelativeTime;
 #endif
 
-#if USE_PARTICLE_RANDOM
+#if NEEDS_PARTICLE_RANDOM
 	Result.Particle.Random = Intermediates.ParticleRandom;
 #else
 	Result.Particle.Random = 0.0f;
@@ -347,7 +347,7 @@ FMaterialVertexParameters GetMaterialVertexParameters(FVertexFactoryInput Input,
 	float4x4 Transform = GetParticleTransform(Intermediates);
 	Result.InstanceLocalToWorld = Transform;
 
-#if USE_PARTICLE_TRANSFORM
+#if NEEDS_PARTICLE_TRANSFORM
 	Result.Particle.LocalToWorld[0] = float4(Transform[0].x, Transform[1].x, Transform[2].x, Transform[3].x);
 	Result.Particle.LocalToWorld[1] = float4(Transform[0].y, Transform[1].y, Transform[2].y, Transform[3].y);
 	Result.Particle.LocalToWorld[2] = float4(Transform[0].z, Transform[1].z, Transform[2].z, Transform[3].z);
@@ -542,7 +542,7 @@ FVertexFactoryIntermediates GetVertexFactoryIntermediates(FVertexFactoryInput In
 	Intermediates.DynamicParameter3 = NiagaraMeshVF.MaterialParam3DataOffset != -1 ? GetVec4(NiagaraMeshVF.MaterialParam3DataOffset, Input.InstanceId) : float4(1, 1, 1, 1);
 #endif
 
-#if USE_PARTICLE_TIME
+#if NEEDS_PARTICLE_TIME
 	if (NiagaraMeshVF.NormalizedAgeDataOffset == -1)
 	{
 		Intermediates.RelativeTime = 0;
@@ -553,7 +553,7 @@ FVertexFactoryIntermediates GetVertexFactoryIntermediates(FVertexFactoryInput In
 	}
 #endif
 
-#if USE_PARTICLE_RANDOM
+#if NEEDS_PARTICLE_RANDOM
 	if (NiagaraMeshVF.MaterialRandomDataOffset == -1)
 	{
 		Intermediates.ParticleRandom = 0;
@@ -564,7 +564,7 @@ FVertexFactoryIntermediates GetVertexFactoryIntermediates(FVertexFactoryInput In
 	}
 #endif
 
-#if USE_PARTICLE_VELOCITY
+#if NEEDS_PARTICLE_VELOCITY
 	float ParticleVelocityMagnitude = length(ParticleVelocity);
 
 	Intermediates.ParticleVelocity.xyz = (ParticleVelocityMagnitude > 0) ? normalize(ParticleVelocity) : float3(0.0f, 0.0f, 0.0f);
