Fix ensure on Message Manager getting usage text of GPU Script.

diff --git a/Engine/Plugins/FX/Niagara/Source/NiagaraEditor/Private/NiagaraMessageManager.cpp b/Engine/Plugins/FX/Niagara/Source/NiagaraEditor/Private/NiagaraMessageManager.cpp
--- a/Engine/Plugins/FX/Niagara/Source/NiagaraEditor/Private/NiagaraMessageManager.cpp
+++ b/Engine/Plugins/FX/Niagara/Source/NiagaraEditor/Private/NiagaraMessageManager.cpp
@@ -76,6 +76,9 @@ const TOptional<const FString> FNiagaraMessageManager::GetStringForScriptUsageIn
 	case ENiagaraScriptUsage::ParticleSpawnScriptInterpolated:
 		return TOptional<const FString>(TEXT("Particle Spawn Script Interpolated"));
 
+	case ENiagaraScriptUsage::ParticleGPUComputeScript:
+		return TOptional<const FString>(TEXT("Particle GPU Compute Script"));
+
 	case ENiagaraScriptUsage::ParticleUpdateScript:
 		return TOptional<const FString>(TEXT("Particle Update Script"));
 		
@@ -101,7 +104,6 @@ const TOptional<const FString> FNiagaraMessageManager::GetStringForScriptUsageIn
 		return TOptional<const FString>();
 
 	//unhandled cases
-	case ENiagaraScriptUsage::ParticleGPUComputeScript:
 	default:
 		ensureMsgf(false, TEXT("Tried to get script usage text for usage that is not handled!"));
 		return TOptional<const FString>();
