Fixed issue where DI instance data was not being recreated after a reset leading to a crash.

diff --git a/Engine/Plugins/FX/Niagara/Source/Niagara/Private/NiagaraSystemInstance.cpp b/Engine/Plugins/FX/Niagara/Source/Niagara/Private/NiagaraSystemInstance.cpp
--- a/Engine/Plugins/FX/Niagara/Source/Niagara/Private/NiagaraSystemInstance.cpp
+++ b/Engine/Plugins/FX/Niagara/Source/Niagara/Private/NiagaraSystemInstance.cpp
@@ -493,17 +493,13 @@ void FNiagaraSystemInstance::Reset(FNiagaraSystemInstance::EResetMode Mode)
 	{
 		//UE_LOG(LogNiagara, Log, TEXT("FNiagaraSystemInstance::Reset true"));
 		ResetInternal(true);
-
-		// ResetInternal will not set the execution state but if we are neither complete or disabled we require a rebind
-		bBindParams = !IsComplete();
+		bBindParams = true;
 	}
 	else if (Mode == EResetMode::ReInit)
 	{
 		//UE_LOG(LogNiagara, Log, TEXT("FNiagaraSystemInstance::ReInit"));
 		ReInitInternal();
-
-		// ReInitInternal will set the execution state so if we are neither complete or disabled we require a rebind
-		bBindParams = !IsComplete();
+		bBindParams = true;
 	}
 	
 	if (bBindParams)
