Fixed possible crash when accessing deleted NiagaraEmitterInstanceBatcher following a change to the preview feature level.
The batcher now has a gamethread callback that allows it to cleanup the deferred deletion queue from the niagara world manager.

Fixed niagara GPU particles being ran in non SM5 preview.

diff --git a/Engine/Plugins/FX/Niagara/Source/Niagara/Classes/NiagaraEmitterInstanceBatcher.h b/Engine/Plugins/FX/Niagara/Source/Niagara/Classes/NiagaraEmitterInstanceBatcher.h
--- a/Engine/Plugins/FX/Niagara/Source/Niagara/Classes/NiagaraEmitterInstanceBatcher.h
+++ b/Engine/Plugins/FX/Niagara/Source/Niagara/Classes/NiagaraEmitterInstanceBatcher.h
@@ -81,6 +81,7 @@ public:
 	virtual void PreInitViews(FRHICommandListImmediate& RHICmdList) override;
 	virtual bool UsesGlobalDistanceField() const override;
 	virtual void PreRender(FRHICommandListImmediate& RHICmdList, const class FGlobalDistanceFieldParameterData* GlobalDistanceFieldParameterData, bool bAllowGPUParticleSceneUpdate) override;
+	virtual void Destroy() override; // Called on the gamethread to delete the batcher on the renderthread.
 
 	virtual void Tick(float DeltaTime) override
 	{
@@ -130,6 +131,9 @@ public:
 
 	FORCEINLINE FNiagaraGPUInstanceCountManager& GetGPUInstanceCounterManager() { check(IsInRenderingThread()); return GPUInstanceCounterManager; }
 
+	FORCEINLINE EShaderPlatform GetShaderPlatform() const { return ShaderPlatform; }
+	FORCEINLINE ERHIFeatureLevel::Type GetFeatureLevel() const { return FeatureLevel; }
+
 private:
 
 	void ExecuteAll(FRHICommandList &RHICmdList, FRHIUniformBuffer* ViewUniformBuffer, bool bSetReadback);
diff --git a/Engine/Plugins/FX/Niagara/Source/Niagara/Private/NiagaraComponent.cpp b/Engine/Plugins/FX/Niagara/Source/Niagara/Private/NiagaraComponent.cpp
--- a/Engine/Plugins/FX/Niagara/Source/Niagara/Private/NiagaraComponent.cpp
+++ b/Engine/Plugins/FX/Niagara/Source/Niagara/Private/NiagaraComponent.cpp
@@ -299,7 +299,7 @@ void FNiagaraSceneProxy::GetDynamicMeshElements(const TArray<const FSceneView*>&
 	for (int32 RendererIdx : RendererDrawOrder)
 	{
 		FNiagaraRenderer* Renderer = EmitterRenderers[RendererIdx];
-		if (Renderer)
+		if (Renderer && (Renderer->GetSimTarget() == ENiagaraSimTarget::CPUSim || ViewFamily.GetFeatureLevel() == ERHIFeatureLevel::SM5))
 		{
 			Renderer->GetDynamicMeshElements(Views, ViewFamily, VisibilityMap, Collector, this);
 		}
diff --git a/Engine/Plugins/FX/Niagara/Source/Niagara/Private/NiagaraEmitterInstanceBatcher.cpp b/Engine/Plugins/FX/Niagara/Source/Niagara/Private/NiagaraEmitterInstanceBatcher.cpp
--- a/Engine/Plugins/FX/Niagara/Source/Niagara/Private/NiagaraEmitterInstanceBatcher.cpp
+++ b/Engine/Plugins/FX/Niagara/Source/Niagara/Private/NiagaraEmitterInstanceBatcher.cpp
@@ -588,6 +588,12 @@ void NiagaraEmitterInstanceBatcher::PreRender(FRHICommandListImmediate& RHICmdLi
 
 }
 
+void NiagaraEmitterInstanceBatcher::Destroy()
+{
+	INiagaraModule::OnBatcherDestroyed(this);
+	FFXSystemInterface::Destroy();
+}
+
 int32 NiagaraEmitterInstanceBatcher::AddSortedGPUSimulation(const FNiagaraGPUSortInfo& SortInfo)
 {
 	const int32 ResultOffset = SortedParticleCount;
diff --git a/Engine/Plugins/FX/Niagara/Source/Niagara/Private/NiagaraModule.cpp b/Engine/Plugins/FX/Niagara/Source/Niagara/Private/NiagaraModule.cpp
--- a/Engine/Plugins/FX/Niagara/Source/Niagara/Private/NiagaraModule.cpp
+++ b/Engine/Plugins/FX/Niagara/Source/Niagara/Private/NiagaraModule.cpp
@@ -316,6 +316,14 @@ FNiagaraWorldManager* INiagaraModule::GetWorldManager(UWorld* World)
 	return *OutWorld;
 }
 
+void INiagaraModule::OnBatcherDestroyed(NiagaraEmitterInstanceBatcher* InBatcher)
+{
+	for (TPair<UWorld*, FNiagaraWorldManager*>& Pair : WorldManagers)
+	{
+		Pair.Value->OnBatcherDestroyed(InBatcher);
+	}
+}
+
 void INiagaraModule::DestroyAllSystemSimulations(class UNiagaraSystem* System)
 {
 	for (TPair<UWorld*, FNiagaraWorldManager*>& Pair : WorldManagers)
diff --git a/Engine/Plugins/FX/Niagara/Source/Niagara/Private/NiagaraSystemSimulation.cpp b/Engine/Plugins/FX/Niagara/Source/Niagara/Private/NiagaraSystemSimulation.cpp
--- a/Engine/Plugins/FX/Niagara/Source/Niagara/Private/NiagaraSystemSimulation.cpp
+++ b/Engine/Plugins/FX/Niagara/Source/Niagara/Private/NiagaraSystemSimulation.cpp
@@ -687,7 +687,7 @@ bool FNiagaraSystemSimulation::Tick(float DeltaSeconds)
 			SystemInstance->FinalizeTick(DeltaSeconds);
 			// @todo If we buffer up the Tick object we create here we can push all the data for every system in this simulation at once.
 			// @todo It may also be worth considering hoisting the push out to FNiagaraWorldManager::Tick to push every desired simulation at once.
-			if (SystemInstance->ActiveGPUEmitterCount > 0)
+			if (SystemInstance->ActiveGPUEmitterCount > 0 && RHISupportsComputeShaders(SystemInstance->GetBatcher()->GetShaderPlatform()))
 			{
 				ensure(!SystemInstance->IsComplete());
 				FNiagaraGPUSystemTick GPUTick;
diff --git a/Engine/Plugins/FX/Niagara/Source/Niagara/Private/NiagaraWorldManager.cpp b/Engine/Plugins/FX/Niagara/Source/Niagara/Private/NiagaraWorldManager.cpp
--- a/Engine/Plugins/FX/Niagara/Source/Niagara/Private/NiagaraWorldManager.cpp
+++ b/Engine/Plugins/FX/Niagara/Source/Niagara/Private/NiagaraWorldManager.cpp
@@ -14,6 +14,7 @@
 #include "GameFramework/PlayerController.h"
 #include "EngineModule.h"
 #include "NiagaraStats.h"
+#include "NiagaraEmitterInstanceBatcher.h"
 
 DECLARE_CYCLE_STAT(TEXT("Niagara Manager Tick [GT]"), STAT_NiagaraWorldManTick, STATGROUP_Niagara);
 DECLARE_CYCLE_STAT(TEXT("Niagara Manager Wait On Render [GT]"), STAT_NiagaraWorldManWaitOnRender, STATGROUP_Niagara);
@@ -181,6 +182,21 @@ void FNiagaraWorldManager::DestroySystemInstance(TUniquePtr<FNiagaraSystemInstan
 	DeferredDeletionQueue[DeferredDeletionQueueIndex].Queue.Emplace(MoveTemp(InPtr));
 }
 
+void FNiagaraWorldManager::OnBatcherDestroyed(NiagaraEmitterInstanceBatcher* InBatcher)
+{
+	// Process the deferred deletion queue before deleting the batcher of this world.
+	// This is required because the batcher is accessed in FNiagaraEmitterInstance::~FNiagaraEmitterInstance
+	if (World && World->FXSystem && World->FXSystem->GetInterface(NiagaraEmitterInstanceBatcher::Name) == InBatcher)
+	{
+		for ( int32 i=0; i < NumDeferredQueues; ++i)
+		{
+			DeferredDeletionQueue[DeferredDeletionQueueIndex].Fence.Wait();
+			DeferredDeletionQueue[i].Queue.Empty();
+		}
+	}
+}
+
+
 void FNiagaraWorldManager::OnWorldCleanup(bool bSessionEnded, bool bCleanupResources)
 {
 	for (TPair<UNiagaraSystem*, TSharedRef<FNiagaraSystemSimulation, ESPMode::ThreadSafe>>& SimPair : SystemSimulations)
diff --git a/Engine/Plugins/FX/Niagara/Source/Niagara/Public/NiagaraModule.h b/Engine/Plugins/FX/Niagara/Source/Niagara/Public/NiagaraModule.h
--- a/Engine/Plugins/FX/Niagara/Source/Niagara/Public/NiagaraModule.h
+++ b/Engine/Plugins/FX/Niagara/Source/Niagara/Public/NiagaraModule.h
@@ -45,6 +45,9 @@ public:
 
 	static FNiagaraWorldManager* GetWorldManager(UWorld* World);
 
+	// Gamethread callback to cleanup references to the given batcher before it gets deleted on the renderthread.
+	static void OnBatcherDestroyed(class NiagaraEmitterInstanceBatcher* InBatcher);
+
 	void DestroyAllSystemSimulations(class UNiagaraSystem* System);
 
 	// Callback function registered with global world delegates to instantiate world manager when a game world is created
diff --git a/Engine/Plugins/FX/Niagara/Source/Niagara/Public/NiagaraRenderer.h b/Engine/Plugins/FX/Niagara/Source/Niagara/Public/NiagaraRenderer.h
--- a/Engine/Plugins/FX/Niagara/Source/Niagara/Public/NiagaraRenderer.h
+++ b/Engine/Plugins/FX/Niagara/Source/Niagara/Public/NiagaraRenderer.h
@@ -119,6 +119,8 @@ public:
 	static FRWBuffer& GetDummyIntBuffer();
 	static FRWBuffer& GetDummyUIntBuffer();
 	
+	FORCEINLINE ENiagaraSimTarget GetSimTarget() const { return SimTarget; }
+
 protected:
 
 	struct FNiagaraDynamicDataBase *DynamicDataRender;
diff --git a/Engine/Plugins/FX/Niagara/Source/Niagara/Public/NiagaraWorldManager.h b/Engine/Plugins/FX/Niagara/Source/Niagara/Public/NiagaraWorldManager.h
--- a/Engine/Plugins/FX/Niagara/Source/Niagara/Public/NiagaraWorldManager.h
+++ b/Engine/Plugins/FX/Niagara/Source/Niagara/Public/NiagaraWorldManager.h
@@ -78,6 +78,9 @@ public:
 	void DestroySystemSimulation(UNiagaraSystem* System);
 	void DestroySystemInstance(TUniquePtr<FNiagaraSystemInstance>& InPtr);
 
+	// Gamethread callback to cleanup references to the given batcher before it gets deleted on the renderthread.
+	void OnBatcherDestroyed(NiagaraEmitterInstanceBatcher* InBatcher);
+
 	void Tick(float DeltaSeconds);
 
 	void OnWorldCleanup(bool bSessionEnded, bool bCleanupResources);
diff --git a/Engine/Source/Editor/UnrealEd/Private/ThumbnailHelpers.cpp b/Engine/Source/Editor/UnrealEd/Private/ThumbnailHelpers.cpp
--- a/Engine/Source/Editor/UnrealEd/Private/ThumbnailHelpers.cpp
+++ b/Engine/Source/Editor/UnrealEd/Private/ThumbnailHelpers.cpp
@@ -172,7 +172,11 @@ FParticleSystemThumbnailScene::FParticleSystemThumbnailScene()
 
 FParticleSystemThumbnailScene::~FParticleSystemThumbnailScene()
 {
-	FFXSystemInterface::Destroy( ThumbnailFXSystem );
+	if (ThumbnailFXSystem)
+	{
+		ThumbnailFXSystem->Destroy();
+		ThumbnailFXSystem = nullptr;
+	}
 }
 
 void FParticleSystemThumbnailScene::SetParticleSystem(UParticleSystem* ParticleSystem)
diff --git a/Engine/Source/Runtime/Engine/Private/Particles/FXSystem.cpp b/Engine/Source/Runtime/Engine/Private/Particles/FXSystem.cpp
--- a/Engine/Source/Runtime/Engine/Private/Particles/FXSystem.cpp
+++ b/Engine/Source/Runtime/Engine/Private/Particles/FXSystem.cpp
@@ -46,12 +46,13 @@ FFXSystemInterface* FFXSystemInterface::Create(ERHIFeatureLevel::Type InFeatureL
 	}
 }
 
-void FFXSystemInterface::Destroy( FFXSystemInterface* FXSystem )
+void FFXSystemInterface::Destroy()
 {
-	check(FXSystem && !FXSystem->bIsPendingKill);
+	check(!bIsPendingKill);
 
-	// Notify that the delete command is on its way. Preventing any future render commands from accessing the FFXSystemInterface.
-	FXSystem->bIsPendingKill = true;
+	bIsPendingKill = true;
+
+	FFXSystemInterface* FXSystem = this;
 	ENQUEUE_RENDER_COMMAND(FDestroyFXSystemCommand)(
 		[FXSystem](FRHICommandList& RHICmdList)
 		{
@@ -69,10 +70,6 @@ void FFXSystemInterface::UnregisterCustomFXSystem(const FName& InterfaceName)
 	CreateCustomFXDelegates.Remove(InterfaceName);
 }
 
-FFXSystemInterface::~FFXSystemInterface()
-{
-}
-
 /*------------------------------------------------------------------------------
 	FX system console variables.
 ------------------------------------------------------------------------------*/
diff --git a/Engine/Source/Runtime/Engine/Private/Particles/FXSystemSet.cpp b/Engine/Source/Runtime/Engine/Private/Particles/FXSystemSet.cpp
--- a/Engine/Source/Runtime/Engine/Private/Particles/FXSystemSet.cpp
+++ b/Engine/Source/Runtime/Engine/Private/Particles/FXSystemSet.cpp
@@ -10,11 +10,13 @@ FFXSystemInterface* FFXSystemSet::GetInterface(const FName& InName)
 {
 	for (FFXSystemInterface* FXSystem : FXSystems)
 	{
-		check(FXSystem);
-		FFXSystemInterface* Res = FXSystem->GetInterface(InName);
-		if (Res)
+		if (FXSystem) // Can be null when called from ::Destroy()
 		{
-			return Res;
+			FXSystem = FXSystem->GetInterface(InName);
+			if (FXSystem)
+			{
+				return FXSystem;
+			}
 		}
 	}
 	return nullptr;
@@ -136,12 +138,16 @@ void FFXSystemSet::PostRenderOpaque(
 	}
 }
 
-FFXSystemSet::~FFXSystemSet()
+void FFXSystemSet::Destroy()
 {
-	for (FFXSystemInterface* FXSystem : FXSystems)
+	for (FFXSystemInterface*& FXSystem : FXSystems)
 	{
+		// Here it is important that the FXSystem stays in the FXSystems array while destroy is being called.
+		// This is used to figure out from which world the destroyed FXSystem is owned (from UWorld::FXSystem->GetInterface())
 		check(FXSystem);
-		FFXSystemInterface::Destroy(FXSystem);
+		FXSystem->Destroy();
+		FXSystem = nullptr;
 	}
-	FXSystems.Empty();
-}
\ No newline at end of file
+
+	FFXSystemInterface::Destroy();
+}
diff --git a/Engine/Source/Runtime/Engine/Private/Particles/FXSystemSet.h b/Engine/Source/Runtime/Engine/Private/Particles/FXSystemSet.h
--- a/Engine/Source/Runtime/Engine/Private/Particles/FXSystemSet.h
+++ b/Engine/Source/Runtime/Engine/Private/Particles/FXSystemSet.h
@@ -39,8 +39,10 @@ public:
 		const class FShaderParametersMetadata* SceneTexturesUniformBufferStruct,
 		FRHIUniformBuffer* SceneTexturesUniformBuffer) override;
 
+	virtual void Destroy() override;
+
 protected:
 
 	/** By making the destructor protected, an instance must be destroyed via FFXSystemInterface::Destroy. */
-	virtual ~FFXSystemSet();
+	virtual ~FFXSystemSet(){}
 };
diff --git a/Engine/Source/Runtime/Engine/Private/RendererSupport.cpp b/Engine/Source/Runtime/Engine/Private/RendererSupport.cpp
--- a/Engine/Source/Runtime/Engine/Private/RendererSupport.cpp
+++ b/Engine/Source/Runtime/Engine/Private/RendererSupport.cpp
@@ -58,9 +58,9 @@ static void ClearReferencesToRendererModuleClasses(
 				Level->ReleaseRenderingResources();
 			}
 
-			if (World->FXSystem != NULL)
+			if (World->FXSystem)
 			{
-				FFXSystemInterface::Destroy(World->FXSystem);
+				World->FXSystem->Destroy();
 				World->FXSystem = NULL;
 			}
 
diff --git a/Engine/Source/Runtime/Engine/Private/World.cpp b/Engine/Source/Runtime/Engine/Private/World.cpp
--- a/Engine/Source/Runtime/Engine/Private/World.cpp
+++ b/Engine/Source/Runtime/Engine/Private/World.cpp
@@ -854,7 +854,7 @@ void UWorld::FinishDestroy()
 
 		if (FXSystem)
 		{
-			FFXSystemInterface::Destroy( FXSystem );
+			FXSystem->Destroy();
 			FXSystem = NULL;
 		}
 
diff --git a/Engine/Source/Runtime/Engine/Public/FXSystem.h b/Engine/Source/Runtime/Engine/Public/FXSystem.h
--- a/Engine/Source/Runtime/Engine/Public/FXSystem.h
+++ b/Engine/Source/Runtime/Engine/Public/FXSystem.h
@@ -129,11 +129,6 @@ public:
 	 */
 	ENGINE_API static FFXSystemInterface* Create(ERHIFeatureLevel::Type InFeatureLevel, EShaderPlatform InShaderPlatform);
 
-	/**
-	 * Destroy an effects system instance.
-	 */
-	ENGINE_API static void Destroy(FFXSystemInterface* FXSystem);
-
 	/**
 	 * Register a custom FX system implementation.
 	 */
@@ -149,6 +144,12 @@ public:
 	 */
 	virtual FFXSystemInterface* GetInterface(const FName& InName) { return nullptr; };
 
+	/**
+	 * Destroy an effects system instance.
+	 */
+	ENGINE_API virtual void Destroy();
+
+
 	/**
 	 * Tick the effects system.
 	 * @param DeltaSeconds The number of seconds by which to step simulations forward.
@@ -226,7 +227,7 @@ public:
 protected:
 
 	/** By making the destructor protected, an instance must be destroyed via FFXSystemInterface::Destroy. */
-	ENGINE_API virtual ~FFXSystemInterface();
+	ENGINE_API virtual ~FFXSystemInterface() {}
 
 private:
 
diff --git a/Engine/Source/Runtime/Renderer/Private/RendererScene.cpp b/Engine/Source/Runtime/Renderer/Private/RendererScene.cpp
--- a/Engine/Source/Runtime/Renderer/Private/RendererScene.cpp
+++ b/Engine/Source/Runtime/Renderer/Private/RendererScene.cpp
@@ -1097,7 +1097,7 @@ FScene::FScene(UWorld* InWorld, bool bInRequiresHitProxies, bool bInIsEditorScen
 
 	if (World->FXSystem)
 	{
-		FFXSystemInterface::Destroy(World->FXSystem);
+		World->FXSystem->Destroy();
 	}
 
 	if (bCreateFXSystem)
