- Fix binding parameters when the system is disabled

diff --git a/Engine/Plugins/FX/Niagara/Source/Niagara/Private/NiagaraSystemInstance.cpp b/Engine/Plugins/FX/Niagara/Source/Niagara/Private/NiagaraSystemInstance.cpp
--- a/Engine/Plugins/FX/Niagara/Source/Niagara/Private/NiagaraSystemInstance.cpp
+++ b/Engine/Plugins/FX/Niagara/Source/Niagara/Private/NiagaraSystemInstance.cpp
@@ -493,13 +493,13 @@ void FNiagaraSystemInstance::Reset(FNiagaraSystemInstance::EResetMode Mode)
 	{
 		//UE_LOG(LogNiagara, Log, TEXT("FNiagaraSystemInstance::Reset true"));
 		ResetInternal(true);
-		bBindParams = true;
+		bBindParams = !IsDisabled();
 	}
 	else if (Mode == EResetMode::ReInit)
 	{
 		//UE_LOG(LogNiagara, Log, TEXT("FNiagaraSystemInstance::ReInit"));
 		ReInitInternal();
-		bBindParams = true;
+		bBindParams = !IsDisabled();
 	}
 	
 	if (bBindParams)
