- Fix null float buffer SRV on GPU simulation when Niagara emitter has no float parameters

diff --git a/Engine/Plugins/FX/Niagara/Source/Niagara/Private/NiagaraRendererMeshes.cpp b/Engine/Plugins/FX/Niagara/Source/Niagara/Private/NiagaraRendererMeshes.cpp
--- a/Engine/Plugins/FX/Niagara/Source/Niagara/Private/NiagaraRendererMeshes.cpp
+++ b/Engine/Plugins/FX/Niagara/Source/Niagara/Private/NiagaraRendererMeshes.cpp
@@ -321,7 +321,14 @@ void FNiagaraRendererMeshes::GetDynamicMeshElements(const TArray<const FSceneVie
 							CollectorResources.VertexFactory.SetSortedIndices(Batcher->GetGPUSortedBuffer().VertexBufferSRV, IndexBufferOffset);
 						}
 					}
-					CollectorResources.VertexFactory.SetParticleData(SourceParticleData->GetGPUBufferFloat().SRV, 0, SourceParticleData->GetFloatStride() / sizeof(float));
+					if (SourceParticleData->GetGPUBufferFloat().SRV.IsValid())
+					{
+						CollectorResources.VertexFactory.SetParticleData(SourceParticleData->GetGPUBufferFloat().SRV, 0, SourceParticleData->GetFloatStride() / sizeof(float));
+					}
+					else
+					{
+						CollectorResources.VertexFactory.SetParticleData(FNiagaraRenderer::GetDummyFloatBuffer().SRV, 0, 0);
+					}
 				}
 
 				// Collector.AllocateOneFrameResource uses default ctor, initialize the vertex factory
diff --git a/Engine/Plugins/FX/Niagara/Source/Niagara/Private/NiagaraRendererSprites.cpp b/Engine/Plugins/FX/Niagara/Source/Niagara/Private/NiagaraRendererSprites.cpp
--- a/Engine/Plugins/FX/Niagara/Source/Niagara/Private/NiagaraRendererSprites.cpp
+++ b/Engine/Plugins/FX/Niagara/Source/Niagara/Private/NiagaraRendererSprites.cpp
@@ -415,7 +415,14 @@ void FNiagaraRendererSprites::SetVertexFactoryParticleData(
 				}
 			}
 
-			OutVertexFactory.SetParticleData(SourceParticleData->GetGPUBufferFloat().SRV, 0, SourceParticleData->GetFloatStride() / sizeof(float));
+			if ( SourceParticleData->GetGPUBufferFloat().SRV.IsValid() )
+			{
+				OutVertexFactory.SetParticleData(SourceParticleData->GetGPUBufferFloat().SRV, 0, SourceParticleData->GetFloatStride() / sizeof(float));
+			}
+			else
+			{
+				OutVertexFactory.SetParticleData(FNiagaraRenderer::GetDummyFloatBuffer().SRV, 0, 0);
+			}
 		}
 	}
 }
