- Enable multiple spawn infos for the GPU
- Currently capped at 8 spawn infos which can be modified changing NIAGARA_MAX_GPU_SPAWN_INFOS and the script version.  The cost is 2 float4's and a dot product per 4 spawn infos.

#tests Ran EngineTests and have a test level which uses many spawn infos (tested for overflowing capacity, etc)

diff --git a/Engine/Plugins/FX/Niagara/Shaders/Private/NiagaraEmitterInstanceShader.usf b/Engine/Plugins/FX/Niagara/Shaders/Private/NiagaraEmitterInstanceShader.usf
--- a/Engine/Plugins/FX/Niagara/Shaders/Private/NiagaraEmitterInstanceShader.usf
+++ b/Engine/Plugins/FX/Niagara/Shaders/Private/NiagaraEmitterInstanceShader.usf
@@ -68,6 +68,33 @@ float Reciprocal(float x){ return 1.0f/x; }
 		RandomSeedOffset += Seed;
 		return float(Rand3DPCG32(int3(u,v,Seed)).x) / 4294967296.0f;
 	}
+
+	// NIAGARA_MAX_GPU_SPAWN_INFOS is set from the shader compiler
+	#define NIAGARA_MAX_GPU_SPAWN_INFOS_V4	((NIAGARA_MAX_GPU_SPAWN_INFOS + 3) / 4)
+
+	float4	EmitterSpawnInfoOffsets[NIAGARA_MAX_GPU_SPAWN_INFOS_V4];
+	float4	EmitterSpawnInfoParams[NIAGARA_MAX_GPU_SPAWN_INFOS];			// Packed data where x = IntervalDt, y = InterpStartDt, z = Group, w = Start Particle Index
+
+	static int GInterpSpawnIndex;
+	static float Emitter_SpawnInterval;
+	static float Emitter_InterpSpawnStartDt;
+	static int Emitter_SpawnGroup;
+
+	void GetEmitterSpawnInfoForParticle(int ParticleIndex)
+	{
+		int SpawnInfoIndex = 0;
+		UNROLL
+		for ( int i=0; i < NIAGARA_MAX_GPU_SPAWN_INFOS_V4; ++i )
+		{
+			SpawnInfoIndex += (int)dot(1.0f, ParticleIndex >= EmitterSpawnInfoOffsets[i]);
+		}
+
+		Emitter_SpawnInterval		= EmitterSpawnInfoParams[SpawnInfoIndex].x;
+		Emitter_InterpSpawnStartDt	= EmitterSpawnInfoParams[SpawnInfoIndex].y;
+		Emitter_SpawnGroup			= asint(EmitterSpawnInfoParams[SpawnInfoIndex].z);
+
+		GInterpSpawnIndex		= ParticleIndex - asint(EmitterSpawnInfoParams[SpawnInfoIndex].w);
+	}
 #endif
 
 /* ----------------------------------------------------------------------------
diff --git a/Engine/Plugins/FX/Niagara/Source/Niagara/Classes/NiagaraEmitterInstance.h b/Engine/Plugins/FX/Niagara/Source/Niagara/Classes/NiagaraEmitterInstance.h
--- a/Engine/Plugins/FX/Niagara/Source/Niagara/Classes/NiagaraEmitterInstance.h
+++ b/Engine/Plugins/FX/Niagara/Source/Niagara/Classes/NiagaraEmitterInstance.h
@@ -155,10 +155,6 @@ private:
 	FNiagaraParameterDirectBinding<float> InterpSpawnStartBinding;
 	FNiagaraParameterDirectBinding<int32> SpawnGroupBinding;
 
-	FNiagaraParameterDirectBinding<float> SpawnIntervalBindingGPU;
-	FNiagaraParameterDirectBinding<float> InterpSpawnStartBindingGPU;
-	FNiagaraParameterDirectBinding<int32> SpawnGroupBindingGPU;
-
 	FNiagaraParameterDirectBinding<float> EmitterAgeBindingGPU;
 
 	FNiagaraParameterDirectBinding<float> SpawnEmitterAgeBinding;
diff --git a/Engine/Plugins/FX/Niagara/Source/Niagara/Classes/NiagaraEmitterInstanceBatcher.h b/Engine/Plugins/FX/Niagara/Source/Niagara/Classes/NiagaraEmitterInstanceBatcher.h
--- a/Engine/Plugins/FX/Niagara/Source/Niagara/Classes/NiagaraEmitterInstanceBatcher.h
+++ b/Engine/Plugins/FX/Niagara/Source/Niagara/Classes/NiagaraEmitterInstanceBatcher.h
@@ -124,6 +124,7 @@ public:
 				FNiagaraShader* Shader,
 				FRHICommandList &RHICmdList, 
 				FRHIUniformBuffer* ViewUniformBuffer, 
+				const FNiagaraGpuSpawnInfo& SpawnInfo,
 				bool bCopyBeforeStart = false
 			) const;
 
diff --git a/Engine/Plugins/FX/Niagara/Source/Niagara/Classes/NiagaraScriptExecutionContext.h b/Engine/Plugins/FX/Niagara/Source/Niagara/Classes/NiagaraScriptExecutionContext.h
--- a/Engine/Plugins/FX/Niagara/Source/Niagara/Classes/NiagaraScriptExecutionContext.h
+++ b/Engine/Plugins/FX/Niagara/Source/Niagara/Classes/NiagaraScriptExecutionContext.h
@@ -147,6 +147,14 @@ struct FNiagaraScriptExecutionContext
 	bool CanExecute()const;
 };
 
+struct FNiagaraGpuSpawnInfo
+{
+	uint32		EventSpawnTotal = 0;
+	uint32		SpawnRateInstances = 0;
+	FVector4	SpawnInfoStartOffsets[NIAGARA_MAX_GPU_SPAWN_INFOS_V4];
+	FVector4	SpawnInfoParams[NIAGARA_MAX_GPU_SPAWN_INFOS];
+};
+
 struct FNiagaraComputeExecutionContext
 {
 	FNiagaraComputeExecutionContext();
@@ -196,8 +204,8 @@ public:
 	//Most current buffer that can be used for rendering.
 	FNiagaraDataBuffer* DataToRender;
 
-	uint32 EventSpawnTotal_GT;
-	uint32 SpawnRateInstances_GT;
+	// Game thread spawn info will be sent to the render thread inside FNiagaraComputeInstanceData
+	FNiagaraGpuSpawnInfo GpuSpawnInfo_GT;
 
 #if WITH_EDITORONLY_DATA
 	mutable FRHIGPUMemoryReadback *GPUDebugDataReadbackFloat;
@@ -228,8 +236,7 @@ struct FNiagaraDataInterfaceInstanceData
 
 struct FNiagaraComputeInstanceData
 {
-	uint32 EventSpawnTotal;
-	uint32 SpawnRateInstances;
+	FNiagaraGpuSpawnInfo SpawnInfo;
 	uint8* ParamData;
 	FNiagaraComputeExecutionContext* Context;
 	TArray<FNiagaraDataInterfaceProxy*> DataInterfaceProxies;
@@ -240,9 +247,7 @@ struct FNiagaraComputeInstanceData
 	FNiagaraDataBuffer* DestinationData;
 
 	FNiagaraComputeInstanceData()
-		: EventSpawnTotal(0)
-		, SpawnRateInstances(0)
-		, ParamData(nullptr)
+		: ParamData(nullptr)
 		, Context(nullptr)
 		, CurrentData(nullptr)
 		, DestinationData(nullptr)
diff --git a/Engine/Plugins/FX/Niagara/Source/Niagara/Private/NiagaraEmitterInstance.cpp b/Engine/Plugins/FX/Niagara/Source/Niagara/Private/NiagaraEmitterInstance.cpp
--- a/Engine/Plugins/FX/Niagara/Source/Niagara/Private/NiagaraEmitterInstance.cpp
+++ b/Engine/Plugins/FX/Niagara/Source/Niagara/Private/NiagaraEmitterInstance.cpp
@@ -297,13 +297,6 @@ void FNiagaraEmitterInstance::Init(int32 InEmitterIdx, FName InSystemInstanceNam
 	InterpSpawnStartBinding.Init(SpawnExecContext.Parameters, CachedEmitter->ToEmitterParameter(SYS_PARAM_EMITTER_INTERP_SPAWN_START_DT));
 	SpawnGroupBinding.Init(SpawnExecContext.Parameters, CachedEmitter->ToEmitterParameter(SYS_PARAM_EMITTER_SPAWN_GROUP));
 
-	if (CachedEmitter->SimTarget == ENiagaraSimTarget::GPUComputeSim && GPUExecContext != nullptr)
-	{
-		SpawnIntervalBindingGPU.Init(GPUExecContext->CombinedParamStore, CachedEmitter->ToEmitterParameter(SYS_PARAM_EMITTER_SPAWN_INTERVAL));
-		InterpSpawnStartBindingGPU.Init(GPUExecContext->CombinedParamStore, CachedEmitter->ToEmitterParameter(SYS_PARAM_EMITTER_INTERP_SPAWN_START_DT));
-		SpawnGroupBindingGPU.Init(GPUExecContext->CombinedParamStore, CachedEmitter->ToEmitterParameter(SYS_PARAM_EMITTER_SPAWN_GROUP));
-	}
-
 	FNiagaraVariable EmitterAgeParam = CachedEmitter->ToEmitterParameter(SYS_PARAM_EMITTER_AGE);
 	SpawnEmitterAgeBinding.Init(SpawnExecContext.Parameters, EmitterAgeParam);
 	UpdateEmitterAgeBinding.Init(UpdateExecContext.Parameters, EmitterAgeParam);
@@ -1018,9 +1011,6 @@ void FNiagaraEmitterInstance::Tick(float DeltaSeconds)
 		check(GPUExecContext->GPUScript_RT == CachedEmitter->GetGPUComputeScript()->GetRenderThreadScript());
 		GPUExecContext->GPUScript_RT = CachedEmitter->GetGPUComputeScript()->GetRenderThreadScript();
 
-		GPUExecContext->EventSpawnTotal_GT = EventSpawnTotal;
-		GPUExecContext->SpawnRateInstances_GT = SpawnTotal;
-		
 #if WITH_EDITORONLY_DATA
 		if (ParentSystemInstance->ShouldCaptureThisFrame())
 		{
@@ -1040,29 +1030,54 @@ void FNiagaraEmitterInstance::Tick(float DeltaSeconds)
 			}
 		}
 #endif
-		// If this is not correct we will not propagate all data correctly.
-		// @todo-threadsafety we keep a counter of this, so 
-		//check(ParentSystemInstance->HasGPUEmitters());
 
-		bool bOnlySetOnce = false;
-		for (FNiagaraSpawnInfo& Info : SpawnInfos)
+		// Calculate spawn information to pass to the RT
 		{
-			if (Info.Count > 0 && !bOnlySetOnce)
+			static_assert((NIAGARA_MAX_GPU_SPAWN_INFOS == NIAGARA_MAX_GPU_SPAWN_INFOS_V4 * 4) && (NIAGARA_MAX_GPU_SPAWN_INFOS > 0), "NIAGARA_MAX_GPU_SPAWN_INFOS should be greater than zero and a multiple of 4");
+
+			FNiagaraGpuSpawnInfo& GpuSpawnInfo = GPUExecContext->GpuSpawnInfo_GT;
+			GpuSpawnInfo.EventSpawnTotal = EventSpawnTotal;
+			GpuSpawnInfo.SpawnRateInstances = 0;
+
+			int NumSpawnInfos = 0;
+			if (ExecutionState == ENiagaraExecutionState::Active)
 			{
-				// @todo-threadsafety do these need to propagate to the RT?
-				SpawnIntervalBindingGPU.SetValue(Info.IntervalDt);
-				InterpSpawnStartBindingGPU.SetValue(Info.InterpStartDt);
-				SpawnGroupBindingGPU.SetValue(Info.SpawnGroup);
-				bOnlySetOnce = true;
+				for (FNiagaraSpawnInfo& Info : SpawnInfos)
+				{
+					if (Info.Count > 0 && (NumSpawnInfos < NIAGARA_MAX_GPU_SPAWN_INFOS))
+					{
+						GpuSpawnInfo.SpawnInfoParams[NumSpawnInfos].X = Info.IntervalDt;
+						GpuSpawnInfo.SpawnInfoParams[NumSpawnInfos].Y = Info.InterpStartDt;
+						reinterpret_cast<int32&>(GpuSpawnInfo.SpawnInfoParams[NumSpawnInfos].Z) = Info.SpawnGroup;
+						reinterpret_cast<int32&>(GpuSpawnInfo.SpawnInfoParams[NumSpawnInfos].W) = GpuSpawnInfo.SpawnRateInstances;
+
+						GpuSpawnInfo.SpawnRateInstances += Info.Count;
+						reinterpret_cast<float*>(GpuSpawnInfo.SpawnInfoStartOffsets)[NumSpawnInfos] = (float)GpuSpawnInfo.SpawnRateInstances;
+
+						++NumSpawnInfos;
+					}
+					else if (Info.Count > 0)
+					{
+						UE_LOG(LogNiagara, Log, TEXT("Exceeded Gpu spawn info count, see NIAGARA_MAX_GPU_SPAWN_INFOS for more information!"));
+						break;
+					}
+
+					// NOTE(mv): Separate particle count path for GPU emitters, as they early out..
+					TotalSpawnedParticles += Info.Count;
+				}
 			}
-			else if (Info.Count > 0)
+
+			// If we have spawning make sure we clear out the remaining data and leave the end slot as MAX to avoid reading off end of the array on the GPU
+			if ( GpuSpawnInfo.EventSpawnTotal + GpuSpawnInfo.SpawnRateInstances >  0 )
 			{
-				UE_LOG(LogNiagara, Log, TEXT("Multiple spawns are happening this frame. Only doing the first!"));
-				break;
+				while (NumSpawnInfos < NIAGARA_MAX_GPU_SPAWN_INFOS - 1)
+				{
+					reinterpret_cast<float*>(GpuSpawnInfo.SpawnInfoStartOffsets)[NumSpawnInfos] = MAX_FLT;
+					GpuSpawnInfo.SpawnInfoParams[NumSpawnInfos] = FVector4(ForceInitToZero);
+					++NumSpawnInfos;
+				}
+				reinterpret_cast<float*>(GpuSpawnInfo.SpawnInfoStartOffsets)[NIAGARA_MAX_GPU_SPAWN_INFOS - 1] = MAX_FLT;
 			}
-
-			// NOTE(mv): Separate particle count path for GPU emitters, as they early out..
-			TotalSpawnedParticles += Info.Count;
 		}
 
 		//GPUExecContext.UpdateInterfaces = CachedEmitter->UpdateScriptProps.Script->GetCachedDefaultDataInterfaces();
@@ -1595,7 +1610,7 @@ void FNiagaraEmitterInstance::SetExecutionState(ENiagaraExecutionState InState)
 	else
 	{
 		//Try to gracefully fail in this case.
-		InState = ENiagaraExecutionState::Inactive;
+		ExecutionState = ENiagaraExecutionState::Inactive;
 	}
 
 }
diff --git a/Engine/Plugins/FX/Niagara/Source/Niagara/Private/NiagaraEmitterInstanceBatcher.cpp b/Engine/Plugins/FX/Niagara/Source/Niagara/Private/NiagaraEmitterInstanceBatcher.cpp
--- a/Engine/Plugins/FX/Niagara/Source/Niagara/Private/NiagaraEmitterInstanceBatcher.cpp
+++ b/Engine/Plugins/FX/Niagara/Source/Niagara/Private/NiagaraEmitterInstanceBatcher.cpp
@@ -189,7 +189,7 @@ void NiagaraEmitterInstanceBatcher::ResizeBuffersAndGatherResources(FOverlappabl
 			FNiagaraDataBuffer& DestinationData = *Instance.DestinationData;
 
 			const uint32 PrevNumInstances = bNeedsReset ? 0 : CurrentData.GetNumInstances();
-			const uint32 NewNumInstances = Instance.SpawnRateInstances + Instance.EventSpawnTotal + PrevNumInstances;
+			const uint32 NewNumInstances = Instance.SpawnInfo.SpawnRateInstances + Instance.SpawnInfo.EventSpawnTotal + PrevNumInstances;
 
 			//We must assume all particles survive when allocating here. 
 			//If this is not true, the read back in ResolveDatasetWrites will shrink the buffers.
@@ -287,7 +287,7 @@ void NiagaraEmitterInstanceBatcher::DispatchAllOnCompute(FOverlappableTicks& Ove
 
 					// run shader, sim and spawn in a single dispatch
 					uint32 UpdateStartInstance = 0;
-					Run<false>(*Tick, &Instance, UpdateStartInstance, Instance.DestinationData->GetNumInstances(), Context->GPUScript_RT->GetShader(), RHICmdList, ViewUniformBuffer);
+					Run<false>(*Tick, &Instance, UpdateStartInstance, Instance.DestinationData->GetNumInstances(), Context->GPUScript_RT->GetShader(), RHICmdList, ViewUniformBuffer, Instance.SpawnInfo);
 
 					FNiagaraDataBuffer* CurrentData = Instance.CurrentData;
 					if (bSetReadback && Tick->bIsFinalTick)
@@ -452,14 +452,14 @@ void NiagaraEmitterInstanceBatcher::TickSingle(const FNiagaraGPUSystemTick& Tick
 	FNiagaraDataBuffer& DestinationData = *Instance->DestinationData;
 
 	const uint32 PrevNumInstances = Tick.bNeedsReset ? 0 : CurrentData.GetNumInstances();
-	const uint32 NewNumInstances = Instance->SpawnRateInstances + Instance->EventSpawnTotal + PrevNumInstances;
+	const uint32 NewNumInstances = Instance->SpawnInfo.SpawnRateInstances + Instance->SpawnInfo.EventSpawnTotal + PrevNumInstances;
 
 	DestinationData.AllocateGPU(NewNumInstances + 1, GPUInstanceCounterManager, RHICmdList);
 	DestinationData.SetNumInstances(NewNumInstances);
 
 	// run shader, sim and spawn in a single dispatch
 	uint32 UpdateStartInstance = 0;
-	Run<true>(Tick, Instance, UpdateStartInstance, NewNumInstances, ComputeShader, RHICmdList, ViewUniformBuffer);
+	Run<true>(Tick, Instance, UpdateStartInstance, NewNumInstances, ComputeShader, RHICmdList, ViewUniformBuffer, Instance->SpawnInfo);
 
 	Context->MainDataSet->EndSimulate();
 	Context->SetDataToRender(Instance->DestinationData);
@@ -934,7 +934,7 @@ void NiagaraEmitterInstanceBatcher::UnsetDataInterfaceParameters(const TArray<FN
  */
  template<bool bDoResourceTransitions>
 void NiagaraEmitterInstanceBatcher::Run(const FNiagaraGPUSystemTick& Tick, const FNiagaraComputeInstanceData* Instance, uint32 UpdateStartInstance, const uint32 TotalNumInstances, FNiagaraShader* Shader,
-	FRHICommandList &RHICmdList, FRHIUniformBuffer* ViewUniformBuffer, bool bCopyBeforeStart) const
+	FRHICommandList &RHICmdList, FRHIUniformBuffer* ViewUniformBuffer, const FNiagaraGpuSpawnInfo& SpawnInfo, bool bCopyBeforeStart) const
 {
 	FNiagaraComputeExecutionContext* Context = Instance->Context;
 	if (TotalNumInstances == 0 && !bDoResourceTransitions)
@@ -987,8 +987,15 @@ void NiagaraEmitterInstanceBatcher::Run(const FNiagaraGPUSystemTick& Tick, const
 		RHICmdList.SetShaderParameter(Shader->GetComputeShader(), Shader->EmitterTickCounterParam.GetBufferIndex(), Shader->EmitterTickCounterParam.GetBaseIndex(), Shader->EmitterTickCounterParam.GetNumBytes(), &FNiagaraComputeExecutionContext::TickCounter);
 	}
 
+	// set spawn info
+	//
+	static_assert((sizeof(SpawnInfo.SpawnInfoStartOffsets) % SHADER_PARAMETER_ARRAY_ELEMENT_ALIGNMENT) == 0, "sizeof SpawnInfoStartOffsets should be a multiple of SHADER_PARAMETER_ARRAY_ELEMENT_ALIGNMENT");
+	static_assert((sizeof(SpawnInfo.SpawnInfoParams) % SHADER_PARAMETER_ARRAY_ELEMENT_ALIGNMENT) == 0, "sizeof SpawnInfoParams should be a multiple of SHADER_PARAMETER_ARRAY_ELEMENT_ALIGNMENT");
+	SetShaderValueArray(RHICmdList, Shader->GetComputeShader(), Shader->EmitterSpawnInfoOffsetsParam, SpawnInfo.SpawnInfoStartOffsets, NIAGARA_MAX_GPU_SPAWN_INFOS_V4);
+	SetShaderValueArray(RHICmdList, Shader->GetComputeShader(), Shader->EmitterSpawnInfoParamsParam, SpawnInfo.SpawnInfoParams, NIAGARA_MAX_GPU_SPAWN_INFOS);
+
 	RHICmdList.SetShaderParameter(Shader->GetComputeShader(), Shader->UpdateStartInstanceParam.GetBufferIndex(), Shader->UpdateStartInstanceParam.GetBaseIndex(), Shader->UpdateStartInstanceParam.GetNumBytes(), &UpdateStartInstance);					// 0, except for event handler runs
-	int32 InstancesToSpawnThisFrame = Instance->SpawnRateInstances + Instance->EventSpawnTotal;
+	int32 InstancesToSpawnThisFrame = Instance->SpawnInfo.SpawnRateInstances + Instance->SpawnInfo.EventSpawnTotal;
 	RHICmdList.SetShaderParameter(Shader->GetComputeShader(), Shader->NumSpawnedInstancesParam.GetBufferIndex(), Shader->NumSpawnedInstancesParam.GetBaseIndex(), Shader->NumSpawnedInstancesParam.GetNumBytes(), &InstancesToSpawnThisFrame);				// number of instances in the spawn run
 
 	uint32 NumThreadGroups = 1;
diff --git a/Engine/Plugins/FX/Niagara/Source/Niagara/Private/NiagaraScriptExecutionContext.cpp b/Engine/Plugins/FX/Niagara/Source/Niagara/Private/NiagaraScriptExecutionContext.cpp
--- a/Engine/Plugins/FX/Niagara/Source/Niagara/Private/NiagaraScriptExecutionContext.cpp
+++ b/Engine/Plugins/FX/Niagara/Source/Niagara/Private/NiagaraScriptExecutionContext.cpp
@@ -344,8 +344,8 @@ void FNiagaraGPUSystemTick::Init(FNiagaraSystemInstance* InSystemInstance)
 
 			InstanceData->Context = Emitter->GetGPUContext();
 			check(InstanceData->Context->MainDataSet);
-			InstanceData->SpawnRateInstances = Emitter->GetGPUContext()->SpawnRateInstances_GT;
-			InstanceData->EventSpawnTotal = Emitter->GetGPUContext()->EventSpawnTotal_GT;
+
+			InstanceData->SpawnInfo = Emitter->GetGPUContext()->GpuSpawnInfo_GT;
 
 			int32 ParmSize = Emitter->GetGPUContext()->CombinedParamStore.GetPaddedParameterSizeInBytes();
 
diff --git a/Engine/Plugins/FX/Niagara/Source/Niagara/Public/NiagaraCommon.h b/Engine/Plugins/FX/Niagara/Source/Niagara/Public/NiagaraCommon.h
--- a/Engine/Plugins/FX/Niagara/Source/Niagara/Public/NiagaraCommon.h
+++ b/Engine/Plugins/FX/Niagara/Source/Niagara/Public/NiagaraCommon.h
@@ -25,6 +25,10 @@ const uint32 NIAGARA_MAX_COMPUTE_THREADGROUPS = 65535;
 
 const FString INTERPOLATED_PARAMETER_PREFIX = TEXT("PREV_");
 
+/** The maximum number of spawn infos we can run on the GPU, modifying this will require a version update as it is used in the shader compiler  */
+constexpr uint32 NIAGARA_MAX_GPU_SPAWN_INFOS = 8;
+constexpr uint32 NIAGARA_MAX_GPU_SPAWN_INFOS_V4 = (NIAGARA_MAX_GPU_SPAWN_INFOS + 3) / 4;
+
 enum ENiagaraBaseTypes
 {
 	NBT_Float,
diff --git a/Engine/Plugins/FX/Niagara/Source/NiagaraCore/Private/NiagaraCustomVersion.cpp b/Engine/Plugins/FX/Niagara/Source/NiagaraCore/Private/NiagaraCustomVersion.cpp
--- a/Engine/Plugins/FX/Niagara/Source/NiagaraCore/Private/NiagaraCustomVersion.cpp
+++ b/Engine/Plugins/FX/Niagara/Source/NiagaraCore/Private/NiagaraCustomVersion.cpp
@@ -8,4 +8,4 @@ const FGuid FNiagaraCustomVersion::GUID(0xFCF57AFA, 0x50764283, 0xB9A9E658, 0xFF
 // Register the custom version with core
 FCustomVersionRegistration GRegisterNiagaraCustomVersion(FNiagaraCustomVersion::GUID, FNiagaraCustomVersion::LatestVersion, TEXT("NiagaraVer"));
 
-const FGuid FNiagaraCustomVersion::LatestScriptCompileVersion(0x18B0ACF1, 0x208646CD, 0x91C04B94, 0x3717A7A6 );
\ No newline at end of file
+const FGuid FNiagaraCustomVersion::LatestScriptCompileVersion(0x76F1E483, 0x558E5F4D, 0xBF7A18BD, 0x591D49F1);
diff --git a/Engine/Plugins/FX/Niagara/Source/NiagaraEditor/Private/NiagaraHlslTranslator.cpp b/Engine/Plugins/FX/Niagara/Source/NiagaraEditor/Private/NiagaraHlslTranslator.cpp
--- a/Engine/Plugins/FX/Niagara/Source/NiagaraEditor/Private/NiagaraHlslTranslator.cpp
+++ b/Engine/Plugins/FX/Niagara/Source/NiagaraEditor/Private/NiagaraHlslTranslator.cpp
@@ -712,6 +712,31 @@ void FHlslNiagaraTranslator::HandleNamespacedExternalVariablesToDataSetRead(TArr
 	}
 }
 
+bool FHlslNiagaraTranslator::IsVariableInUniformBuffer(const FNiagaraVariable& Variable) const
+{
+	static FNiagaraVariable ExcludeVariables[] =
+	{
+		FNiagaraVariable(FNiagaraTypeDefinition::GetFloatDef(), TEXT("Emitter_SpawnInterval")),
+		FNiagaraVariable(FNiagaraTypeDefinition::GetFloatDef(), TEXT("Emitter_InterpSpawnStartDt")),
+		FNiagaraVariable(FNiagaraTypeDefinition::GetFloatDef(), TEXT("Emitter.SpawnInterval")),
+		FNiagaraVariable(FNiagaraTypeDefinition::GetFloatDef(), TEXT("Emitter.InterpSpawnStartDt")),
+		FNiagaraVariable(FNiagaraTypeDefinition::GetIntDef(),   TEXT("Emitter_SpawnGroup")),
+		FNiagaraVariable(FNiagaraTypeDefinition::GetIntDef(),   TEXT("Emitter.SpawnGroup")),
+	};
+
+	if (CompilationTarget == ENiagaraSimTarget::GPUComputeSim)
+	{
+		for ( const FNiagaraVariable& ExcludeVar : ExcludeVariables )
+		{
+			if ( Variable == ExcludeVar )
+			{
+				return false;
+			}
+		}
+	}
+	return true;
+}
+
 const FNiagaraTranslateResults &FHlslNiagaraTranslator::Translate(const FNiagaraCompileRequestData* InCompileData, const FNiagaraCompileOptions& InCompileOptions, FHlslNiagaraTranslatorOptions InTranslateOptions)
 {
 	SCOPE_CYCLE_COUNTER(STAT_NiagaraEditor_HlslTranslator_Translate);
@@ -1025,9 +1050,12 @@ const FNiagaraTranslateResults &FHlslNiagaraTranslator::Translate(const FNiagara
 				ParameterMapRegisterExternalConstantNamespaceVariable(SYS_PARAM_ENGINE_INV_DELTA_TIME, nullptr, 0, OutputIdx, nullptr);
 				ParameterMapRegisterExternalConstantNamespaceVariable(SYS_PARAM_ENGINE_EXEC_COUNT, nullptr, 0, OutputIdx, nullptr);
 				ParameterMapRegisterExternalConstantNamespaceVariable(SYS_PARAM_EMITTER_SPAWNRATE, nullptr, 0, OutputIdx, nullptr);
-				ParameterMapRegisterExternalConstantNamespaceVariable(SYS_PARAM_EMITTER_SPAWN_INTERVAL, nullptr, 0, OutputIdx, nullptr);
-				ParameterMapRegisterExternalConstantNamespaceVariable(SYS_PARAM_EMITTER_INTERP_SPAWN_START_DT, nullptr, 0, OutputIdx, nullptr);
-				ParameterMapRegisterExternalConstantNamespaceVariable(SYS_PARAM_EMITTER_SPAWN_GROUP, nullptr, 0, OutputIdx, nullptr);
+				if (CompilationTarget != ENiagaraSimTarget::GPUComputeSim)
+				{
+					ParameterMapRegisterExternalConstantNamespaceVariable(SYS_PARAM_EMITTER_SPAWN_INTERVAL, nullptr, 0, OutputIdx, nullptr);
+					ParameterMapRegisterExternalConstantNamespaceVariable(SYS_PARAM_EMITTER_INTERP_SPAWN_START_DT, nullptr, 0, OutputIdx, nullptr);
+					ParameterMapRegisterExternalConstantNamespaceVariable(SYS_PARAM_EMITTER_SPAWN_GROUP, nullptr, 0, OutputIdx, nullptr);
+				}
 			}
 
 			if (TranslationStages.Num() > 0)
@@ -1062,8 +1090,12 @@ const FNiagaraTranslateResults &FHlslNiagaraTranslator::Translate(const FNiagara
 
 			for (int32 i = 0; i < ChunksByMode[(int32)ENiagaraCodeChunkMode::Uniform].Num(); ++i)
 			{
-				FNiagaraCodeChunk& Chunk = CodeChunks[ChunksByMode[(int32)ENiagaraCodeChunkMode::Uniform][i]];
-				HlslOutput += TEXT("\t") + GetCode(ChunksByMode[(int32)ENiagaraCodeChunkMode::Uniform][i]);
+				FNiagaraVariable BufferVariable(CodeChunks[ChunksByMode[(int32)ENiagaraCodeChunkMode::Uniform][i]].Type, FName(*CodeChunks[ChunksByMode[(int32)ENiagaraCodeChunkMode::Uniform][i]].SymbolName));
+				if ( IsVariableInUniformBuffer(BufferVariable) )
+				{
+					FString Chunk = GetCode(ChunksByMode[(int32)ENiagaraCodeChunkMode::Uniform][i]);
+					HlslOutput += TEXT("\t") + Chunk;
+				}
 			}
 
 			if (bInterpolateParams)
@@ -1457,7 +1489,16 @@ void FHlslNiagaraTranslator::DefineInterpolatedParametersFunction(FString &HlslO
 		FString PrevMap = TranslationStages[i - 1].PassNamespace;
 		FString CurMap = TranslationStages[i].PassNamespace;
 		{
-			HlslOutputString += TEXT("\tint InterpSpawn_Index = ExecIndex();\n");
+			// GPU simulation is slightly different as we run all spawns at once rather than separate invocations so we can not ExecIndex() as it has to be biased into the correct spawn group's index
+			if ( CompilationTarget == ENiagaraSimTarget::GPUComputeSim )
+			{
+				HlslOutputString += TEXT("\tint InterpSpawn_Index = GInterpSpawnIndex;\n");
+			}
+			else
+			{
+				HlslOutputString += TEXT("\tint InterpSpawn_Index = ExecIndex();\n");
+			}
+
 			HlslOutputString += TEXT("\tfloat InterpSpawn_SpawnTime = ") + Emitter_InterpSpawnStartDt + TEXT(" + (") + Emitter_SpawnInterval + TEXT(" * InterpSpawn_Index);\n");
 			HlslOutputString += TEXT("\tfloat InterpSpawn_UpdateTime = Engine_DeltaTime - InterpSpawn_SpawnTime;\n");
 			HlslOutputString += TEXT("\tfloat InterpSpawn_InvSpawnTime = 1.0 / InterpSpawn_SpawnTime;\n");
@@ -1746,6 +1787,15 @@ void FHlslNiagaraTranslator::DefineMain(FString &OutHlslOutput,
 	//TODO: Grab indices for reading data sets and do the read.
 	//read input.
 
+	// If we are GPU we need to set the spawn info global variables
+	if (CompilationTarget == ENiagaraSimTarget::GPUComputeSim)
+	{
+		OutHlslOutput += TEXT("\tif (Phase == 0)\n");
+		OutHlslOutput += TEXT("\t{\n");
+		OutHlslOutput += TEXT("\t\tGetEmitterSpawnInfoForParticle(ExecIndex());\n");
+		OutHlslOutput += TEXT("\t}\n");
+	}
+
 	//The VM register binding assumes the same inputs as outputs which is obviously not always the case.
 	for (int32 VarArrayIdx = 0; VarArrayIdx < InstanceReadVars.Num(); VarArrayIdx++)
 	{
@@ -3500,7 +3550,11 @@ bool FHlslNiagaraTranslator::ParameterMapRegisterExternalConstantNamespaceVariab
 					Error(FText::Format(LOCTEXT("GetParameterUnsetParam", "Variable {0} hasn't had its default value set. Required Bytes: {1} vs Allocated Bytes: {2}"), FText::FromName(InVariable.GetName()), FText::AsNumber(InVariable.GetType().GetSize()), FText::AsNumber(InVariable.GetSizeInBytes())), nullptr, nullptr);
 				}
 
-				CompilationOutput.ScriptData.Parameters.SetOrAdd(InVariable);
+				if ( IsVariableInUniformBuffer(InVariable) )
+				{
+					CompilationOutput.ScriptData.Parameters.SetOrAdd(InVariable);
+				}
+
 				UniformIdx = ChunksByMode[(int32)ENiagaraCodeChunkMode::Uniform].Num();
 
 				UniformChunk = AddUniformChunk(SymbolNameDefined, InVariable.GetType());
diff --git a/Engine/Plugins/FX/Niagara/Source/NiagaraEditor/Private/NiagaraHlslTranslator.h b/Engine/Plugins/FX/Niagara/Source/NiagaraEditor/Private/NiagaraHlslTranslator.h
--- a/Engine/Plugins/FX/Niagara/Source/NiagaraEditor/Private/NiagaraHlslTranslator.h
+++ b/Engine/Plugins/FX/Niagara/Source/NiagaraEditor/Private/NiagaraHlslTranslator.h
@@ -499,6 +499,9 @@ private:
 
 	void HandleNamespacedExternalVariablesToDataSetRead(TArray<FNiagaraVariable>& InDataSetVars, FString InNamespaceStr);
 
+	// For GPU simulations we have to special case some variables and pass them view shader parameters rather than the uniform buffer as they vary from CPU simulations
+	bool IsVariableInUniformBuffer(const FNiagaraVariable& Variable) const;
+
 	FString ComputeMatrixColumnAccess(const FString& Name);
 	FString ComputeMatrixRowAccess(const FString& Name);
 
diff --git a/Engine/Plugins/FX/Niagara/Source/NiagaraShader/Private/NiagaraShader.cpp b/Engine/Plugins/FX/Niagara/Source/NiagaraShader/Private/NiagaraShader.cpp
--- a/Engine/Plugins/FX/Niagara/Source/NiagaraShader/Private/NiagaraShader.cpp
+++ b/Engine/Plugins/FX/Niagara/Source/NiagaraShader/Private/NiagaraShader.cpp
@@ -297,6 +297,7 @@ FShaderCompileJob* FNiagaraShaderType::BeginCompileShader(
 	NewJob->Input.VirtualSourceFilePath = TEXT("/Plugin/FX/Niagara/Private/NiagaraEmitterInstanceShader.usf");
 	NewJob->Input.EntryPointName = TEXT("SimulateMainComputeCS");
 	NewJob->Input.Environment.SetDefine(TEXT("GPU_SIMULATION"), 1);
+	NewJob->Input.Environment.SetDefine(TEXT("NIAGARA_MAX_GPU_SPAWN_INFOS"), NIAGARA_MAX_GPU_SPAWN_INFOS);
 	NewJob->Input.Environment.IncludeVirtualPathToContentsMap.Add(TEXT("/Engine/Generated/NiagaraEmitterInstance.ush"), Script->HlslOutput);
 	
 	AddReferencedUniformBufferIncludes(NewJob->Input.Environment, NewJob->Input.SourceFilePrefix, (EShaderPlatform)Target.Platform);
@@ -1280,6 +1281,9 @@ void FNiagaraShader::BindParams(const FShaderParameterMap &ParameterMap)
 	WriteInstanceCountOffsetParam.Bind(ParameterMap, TEXT("WriteInstanceCountOffset"));
 
 	EmitterTickCounterParam.Bind(ParameterMap, TEXT("EmitterTickCounter"));
+	EmitterSpawnInfoOffsetsParam.Bind(ParameterMap, TEXT("EmitterSpawnInfoOffsets"));
+	EmitterSpawnInfoParamsParam.Bind(ParameterMap, TEXT("EmitterSpawnInfoParams"));
+
 	NumEventsPerParticleParam.Bind(ParameterMap, TEXT("NumEventsPerParticle"));
 	NumParticlesPerEventParam.Bind(ParameterMap, TEXT("NumParticlesPerEvent"));
 	CopyInstancesBeforeStartParam.Bind(ParameterMap, TEXT("CopyInstancesBeforeStart"));
@@ -1323,7 +1327,6 @@ void FNiagaraShader::BindParams(const FShaderParameterMap &ParameterMap)
 	}
 
 	ensure(FloatOutputBufferParam.IsBound() || IntOutputBufferParam.IsBound());	// we should have at least one output buffer we're writing to
-	ensure(ComponentBufferSizeWriteParam.IsBound());
 	ensure(InstanceCountsParam.IsBound());
 	ensure(UpdateStartInstanceParam.IsBound());
 	ensure(NumSpawnedInstancesParam.IsBound());
@@ -1346,6 +1349,8 @@ bool FNiagaraShader::Serialize(FArchive& Ar)
 	Ar << WriteInstanceCountOffsetParam;
 
 	Ar << EmitterTickCounterParam;
+	Ar << EmitterSpawnInfoOffsetsParam;
+	Ar << EmitterSpawnInfoParamsParam;
 
 	Ar << NumSpawnedInstancesParam;
 	Ar << UpdateStartInstanceParam;
diff --git a/Engine/Plugins/FX/Niagara/Source/NiagaraShader/Public/NiagaraShader.h b/Engine/Plugins/FX/Niagara/Source/NiagaraShader/Public/NiagaraShader.h
--- a/Engine/Plugins/FX/Niagara/Source/NiagaraShader/Public/NiagaraShader.h
+++ b/Engine/Plugins/FX/Niagara/Source/NiagaraShader/Public/NiagaraShader.h
@@ -109,6 +109,8 @@ public:
 	FShaderUniformBufferParameter DataInterfaceUniformBufferParam;
 	FShaderUniformBufferParameter ViewUniformBufferParam;
 	FShaderParameter EmitterTickCounterParam;
+	FShaderParameter EmitterSpawnInfoOffsetsParam;
+	FShaderParameter EmitterSpawnInfoParamsParam;
 	FShaderParameter NumEventsPerParticleParam;
 	FShaderParameter NumParticlesPerEventParam;
 	FShaderParameter CopyInstancesBeforeStartParam;
