- Niagara components that are in solo mode will not have their tick disabled until they have compeleted
- This fixes an issue with Niagara Systems being left around when in solo mode after having deactivate called on them

diff --git a/Engine/Plugins/FX/Niagara/Source/Niagara/Private/NiagaraComponent.cpp b/Engine/Plugins/FX/Niagara/Source/Niagara/Private/NiagaraComponent.cpp
--- a/Engine/Plugins/FX/Niagara/Source/Niagara/Private/NiagaraComponent.cpp
+++ b/Engine/Plugins/FX/Niagara/Source/Niagara/Private/NiagaraComponent.cpp
@@ -724,15 +724,27 @@ void UNiagaraComponent::Activate(bool bReset /* = false */)
 void UNiagaraComponent::Deactivate()
 {
 	SCOPE_CYCLE_COUNTER(STAT_NiagaraComponentDeactivate);
-	Super::Deactivate();
 
 	//UE_LOG(LogNiagara, Log, TEXT("Deactivate: %u - %s"), this, *Asset->GetName());
 
-	bIsActive = false;
-
 	if (SystemInstance)
 	{
+		// Don't deactivate in solo mode as we are not ticked by the world but rather the component
+		// Deactivating will cause the system to never Complete
+		if (SystemInstance->IsSolo() == false)
+		{
+			Super::Deactivate();
+		}
+
 		SystemInstance->Deactivate();
+
+		// We are considered active until we are complete
+		bIsActive = !SystemInstance->IsComplete();
+	}
+	else
+	{
+		Super::Deactivate();
+		bIsActive = false;
 	}
 }
 
