Enable skipping HLSLCC backend optimization passes by default when compiling Niagara shaders through VectorVM.

diff --git a/Engine/Plugins/FX/Niagara/Source/NiagaraShader/Private/NiagaraShader.cpp b/Engine/Plugins/FX/Niagara/Source/NiagaraShader/Private/NiagaraShader.cpp
--- a/Engine/Plugins/FX/Niagara/Source/NiagaraShader/Private/NiagaraShader.cpp
+++ b/Engine/Plugins/FX/Niagara/Source/NiagaraShader/Private/NiagaraShader.cpp
@@ -33,7 +33,7 @@ static FAutoConsoleVariableRef CVarCreateNiagaraShadersOnLoad(
 	TEXT("Whether to create Niagara's simulation shaders on load, which can reduce hitching, but use more memory.  Otherwise they will be created as needed.")
 );
 
-int32 GNiagaraSkipVectorVMBackendOptimizations = 0;
+int32 GNiagaraSkipVectorVMBackendOptimizations = 1;
 static FAutoConsoleVariableRef CVarNiagaraSkipVectorVMBackendOptimizations(
 	TEXT("fx.SkipVectorVMBackendOptimizations"),
 	GNiagaraSkipVectorVMBackendOptimizations,
