- Fix assert inside FNiagaraDataSetExecutionInfo when events are used
- If NumInstances was equal to 0 inside FNiagaraScriptExecutionContext::Execute we left stale DataSetInfo bound leading to the assert on destruction

#tests Test repro which was to add location event sender / event spawn, then ran engine tests

diff --git a/Engine/Plugins/FX/Niagara/Source/Niagara/Private/NiagaraScriptExecutionContext.cpp b/Engine/Plugins/FX/Niagara/Source/Niagara/Private/NiagaraScriptExecutionContext.cpp
--- a/Engine/Plugins/FX/Niagara/Source/Niagara/Private/NiagaraScriptExecutionContext.cpp
+++ b/Engine/Plugins/FX/Niagara/Source/Niagara/Private/NiagaraScriptExecutionContext.cpp
@@ -156,6 +156,7 @@ bool FNiagaraScriptExecutionContext::Execute(uint32 NumInstances)
 {
 	if (NumInstances == 0)
 	{
+		DataSetInfo.Reset();
 		return true;
 	}
 
