- Remove editor only Niagara data from cooked builds
- The result is over 60% reduction in file size costs on cooked platforms, one example system went from 2.38mb to 0.87mb.  This will also reduce in-memory usage.

diff --git a/Engine/Plugins/FX/Niagara/Source/Niagara/Classes/NiagaraScript.h b/Engine/Plugins/FX/Niagara/Source/Niagara/Classes/NiagaraScript.h
--- a/Engine/Plugins/FX/Niagara/Source/Niagara/Classes/NiagaraScript.h
+++ b/Engine/Plugins/FX/Niagara/Source/Niagara/Classes/NiagaraScript.h
@@ -139,6 +139,7 @@ public:
 	UPROPERTY()
 	FGuid BaseScriptID;
 
+#if WITH_EDITORONLY_DATA
 	/**
 	* The hash of the subgraph this shader primarily represents.
 	*/
@@ -155,6 +156,7 @@ public:
 
 	/** Temp storage while generating the Id. This is NOT serialized and shouldn't be used in any comparisons*/
 	TArray<UObject*> ReferencedObjects;
+#endif
 
 	FNiagaraVMExecutableDataId()
 		: CompilerVersionID()
@@ -186,6 +188,7 @@ public:
 	void GetScriptHash(FSHAHash& OutHash) const;
 #endif
 
+#if WITH_EDITORONLY_DATA
 	/**
 	* Tests this set against another for equality, disregarding override settings.
 	*
@@ -201,6 +204,7 @@ public:
 
 	/** Appends string representations of this Id to a key string. */
 	void AppendKeyString(FString& KeyString) const;
+#endif
 };
 
 /** Struct containing all of the data needed to run a Niagara VM executable script.*/
@@ -255,6 +259,7 @@ public:
 	UPROPERTY()
 	TArray<FNiagaraStatScope> StatScopes;
 
+#if WITH_EDITORONLY_DATA
 	UPROPERTY()
 	FString LastHlslTranslation;
 
@@ -266,18 +271,22 @@ public:
 
 	UPROPERTY()
 	uint32 LastOpCount;
+#endif
 
 	UPROPERTY()
 	TArray<FNiagaraDataInterfaceGPUParamInfo> DIParamInfo; //TODO: GPU Param info should not be in the "VM executable data"
 
+#if WITH_EDITORONLY_DATA
 	/** The parameter collections used by this script. */
 	UPROPERTY()
 	TArray<FString> ParameterCollectionPaths;
-	
+#endif
+
 	/** Last known compile status. Lets us determine the latest state of the script byte buffer.*/
 	UPROPERTY()
 	ENiagaraScriptCompileStatus LastCompileStatus;
 
+#if WITH_EDITORONLY_DATA
 	UPROPERTY()
 	bool bReadsAttributeData;
 
@@ -290,11 +299,12 @@ public:
 	/** Array of all compile events generated last time the script was compiled.*/
 	UPROPERTY()
 	TArray<FNiagaraCompileEvent> LastCompileEvents;
+#endif
 
 	void SerializeData(FArchive& Ar, bool bDDCData);
 	
 	bool IsValid() const;
-	
+
 	void Reset();
 };
 
@@ -322,11 +332,11 @@ private:
 	FGuid UsageId;
 
 public:
+#if WITH_EDITORONLY_DATA
 	/** When used as a module, what are the appropriate script types for referencing this module?*/
 	UPROPERTY(AssetRegistrySearchable, EditAnywhere, Category = Script, meta = (Bitmask, BitmaskEnum = ENiagaraScriptUsage))
 	int32 ModuleUsageBitmask;
 
-#if WITH_EDITORONLY_DATA
 	/** Used to break up scripts of the same Usage type in UI display.*/
 	UPROPERTY(AssetRegistrySearchable, EditAnywhere, Category = Script)
 	FText Category;
@@ -350,19 +360,17 @@ public:
 	/* If this script is exposed to the library. */
 	UPROPERTY(AssetRegistrySearchable, EditAnywhere, Category = Script)
 	uint32 bExposeToLibrary : 1;
-
 #endif
 
-
 	/** Contains all of the top-level values that are iterated on in the UI. These are usually "Module" variables in the graph. They don't necessarily have to be in the order that they are expected in the uniform table.*/
 	UPROPERTY()
 	FNiagaraParameterStore RapidIterationParameters;
 
+#if WITH_EDITORONLY_DATA
 	/** The mode to use when deducing the type of numeric output pins from the types of the input pins. */
 	UPROPERTY(EditAnywhere, Category=Script)
 	ENiagaraNumericOutputTypeSelectionMode NumericOutputTypeSelectionMode;
 
-#if WITH_EDITORONLY_DATA
 	UPROPERTY(AssetRegistrySearchable, EditAnywhere, Category = Script, meta = (MultiLine = true))
 	FText Description;
 
@@ -436,8 +444,10 @@ public:
 
 	static bool NIAGARA_API ConvertUsageToGroup(ENiagaraScriptUsage InUsage, ENiagaraScriptGroup& OutGroup);
 
+#if WITH_EDITORONLY_DATA
 	NIAGARA_API TArray<ENiagaraScriptUsage> GetSupportedUsageContexts() const;
 	static NIAGARA_API TArray<ENiagaraScriptUsage> GetSupportedUsageContextsForBitmask(int32 InModuleUsageBitmask);
+#endif
 
 	NIAGARA_API bool CanBeRunOnGpu() const;
 	NIAGARA_API bool IsReadyToRun(ENiagaraSimTarget SimTarget) const;
diff --git a/Engine/Plugins/FX/Niagara/Source/Niagara/Classes/NiagaraSystem.h b/Engine/Plugins/FX/Niagara/Source/Niagara/Classes/NiagaraSystem.h
--- a/Engine/Plugins/FX/Niagara/Source/Niagara/Classes/NiagaraSystem.h
+++ b/Engine/Plugins/FX/Niagara/Source/Niagara/Classes/NiagaraSystem.h
@@ -130,15 +130,15 @@ public:
 
 	bool IsReadyToRun() const;
 
-	/** Are there any pending compile requests?*/
-	bool HasOutstandingCompilationRequests() const;
-
 	FORCEINLINE bool NeedsWarmup()const { return WarmupTickCount > 0 && WarmupTickDelta > SMALL_NUMBER; }
 	FORCEINLINE float GetWarmupTime()const { return WarmupTime; }
 	FORCEINLINE int32 GetWarmupTickCount()const { return WarmupTickCount; }
 	FORCEINLINE float GetWarmupTickDelta()const { return WarmupTickDelta; }
 
 #if WITH_EDITORONLY_DATA
+	/** Are there any pending compile requests?*/
+	bool HasOutstandingCompilationRequests() const;
+
 	/** Determines if this system has the supplied emitter as an editable and simulating emitter instance. */
 	bool ReferencesInstanceEmitter(UNiagaraEmitter& Emitter);
 
@@ -245,8 +245,10 @@ protected:
 	UPROPERTY(EditAnywhere, Category="System")
 	TArray<UNiagaraParameterCollectionInstance*> ParameterCollectionOverrides;
 
+#if WITH_EDITORONLY_DATA
 	UPROPERTY(Transient)
 	TArray<FNiagaraSystemCompileRequest> ActiveCompilations;
+#endif
 
 // 	/** Category of this system. */
 // 	UPROPERTY(EditAnywhere, Category = System)
@@ -270,8 +272,7 @@ protected:
 	UPROPERTY()
 	FNiagaraUserRedirectionParameterStore ExposedParameters;
 
-#if WITH_EDITORONLY_DATA	
-
+#if WITH_EDITORONLY_DATA
 	/** Data used by the editor to maintain UI state etc.. */
 	UPROPERTY()
 	UNiagaraEditorDataBase* EditorData;
diff --git a/Engine/Plugins/FX/Niagara/Source/Niagara/Private/NiagaraScript.cpp b/Engine/Plugins/FX/Niagara/Source/Niagara/Private/NiagaraScript.cpp
--- a/Engine/Plugins/FX/Niagara/Source/Niagara/Private/NiagaraScript.cpp
+++ b/Engine/Plugins/FX/Niagara/Source/Niagara/Private/NiagaraScript.cpp
@@ -59,14 +59,17 @@ UNiagaraScriptSourceBase::UNiagaraScriptSourceBase(const FObjectInitializer& Obj
 
 FNiagaraVMExecutableData::FNiagaraVMExecutableData() 
 	: NumUserPtrs(0)
+#if WITH_EDITORONLY_DATA
 	, LastOpCount(0)
+#endif
 	, LastCompileStatus(ENiagaraScriptCompileStatus::NCS_Unknown)
+#if WITH_EDITORONLY_DATA
 	, bReadsAttributeData(false)
 	, CompileTime(0.0f)
+#endif
 {
 }
 
-
 bool FNiagaraVMExecutableData::IsValid() const
 {
 	return LastCompileStatus != ENiagaraScriptCompileStatus::NCS_Unknown;
@@ -103,6 +106,7 @@ bool FNiagaraVMExecutableDataId::RequiresPersistentIDs() const
 	return AdditionalDefines.Contains("RequiresPersistentIDs");
 }
 
+#if WITH_EDITORONLY_DATA
 /**
 * Tests this set against another for equality, disregarding override settings.
 *
@@ -211,15 +215,16 @@ void FNiagaraVMExecutableDataId::AppendKeyString(FString& KeyString) const
 		}
 	}	
 }
+#endif
 
 UNiagaraScript::UNiagaraScript(const FObjectInitializer& ObjectInitializer)
 	: Super(ObjectInitializer)
 	, Usage(ENiagaraScriptUsage::Function)
 #if WITH_EDITORONLY_DATA
 	, UsageIndex_DEPRECATED(0)
-#endif
 	, ModuleUsageBitmask( (1 << (int32)ENiagaraScriptUsage::ParticleSpawnScript) | (1 << (int32)ENiagaraScriptUsage::ParticleSpawnScriptInterpolated) | (1 << (int32)ENiagaraScriptUsage::ParticleUpdateScript) | (1 << (int32)ENiagaraScriptUsage::ParticleEventScript) )
 	, NumericOutputTypeSelectionMode(ENiagaraNumericOutputTypeSelectionMode::Largest)
+#endif
 {
 #if WITH_EDITORONLY_DATA
 	ScriptResource.OnCompilationComplete().AddUniqueDynamic(this, &UNiagaraScript::OnCompilationComplete);
@@ -1452,6 +1457,7 @@ FNiagaraShaderScript* UNiagaraScript::AllocateResource()
 	return new FNiagaraShaderScript();
 }
 
+#if WITH_EDITORONLY_DATA
 TArray<ENiagaraScriptUsage> UNiagaraScript::GetSupportedUsageContexts() const
 {
 	return GetSupportedUsageContextsForBitmask(ModuleUsageBitmask);
@@ -1470,6 +1476,7 @@ TArray<ENiagaraScriptUsage> UNiagaraScript::GetSupportedUsageContextsForBitmask(
 	}
 	return Supported;
 }
+#endif
 
 bool UNiagaraScript::CanBeRunOnGpu()const
 {
diff --git a/Engine/Plugins/FX/Niagara/Source/Niagara/Private/NiagaraSystem.cpp b/Engine/Plugins/FX/Niagara/Source/Niagara/Private/NiagaraSystem.cpp
--- a/Engine/Plugins/FX/Niagara/Source/Niagara/Private/NiagaraSystem.cpp
+++ b/Engine/Plugins/FX/Niagara/Source/Niagara/Private/NiagaraSystem.cpp
@@ -482,10 +482,12 @@ bool UNiagaraSystem::IsReadyToRun() const
 	return true;
 }
 
+#if WITH_EDITORONLY_DATA
 bool UNiagaraSystem::HasOutstandingCompilationRequests() const
 {
 	return ActiveCompilations.Num() > 0;
 }
+#endif
 
 bool UNiagaraSystem::HasSystemScriptDIsWithPerInstanceData() const
 {
