Object Parameters
-------------------

New "Object" parameter type for Niagara.
These allow you to pass in Raw UObject pointers for use in data interfaces or renderers.
Current use case is to pass in Skeletal Mesh references from the outside.
This avoids having to create a whole new data interface when passing in external object references.

Other use cases could be passing in;
- material overrides
- Static mesh references to override the mesh a renderer uses.
- Basically any UObject a DI or renderer may use.

These are weakly typed, meaning that they have to be UObjects but beyond that, they can be anything.
So in theory you could pass in a Material to these Skeletal mesh references.
It's on the using code to validate and act appropriately.
This weak typing does allow a gread deal of flexibility though.
E.g. you can now pass in a Skeletal Mesh Actor OR a SkeletalMeshComponent to Niagara via the same parameter and the DI will interpret it accordingly.
The ability to pass in direct sub component refs is something cascade could never do.

Deterministic Randoms for Data Interfaces
-------------------------------------------

To support functional tests for the object parameters, this CL also includes deterministic randoms for DIs.
Random seeds for deterministic randoms are now packaged up into a NiagaraRandomInfo that can be passed to DI functions.
There is a helper class that will handle the parameter and let you generate randoms that are deterministic or not based on the passed NiagaraRandomInfo.
Non deterministic is indicated using a special case value of 0xFFFFFFFF for Seed3.

Misc
-------------
Some utility additions to int vector class.
Slight refactor of mesh helper for the skel mesh DI.
Fixed issue with incorrect delta time on first tick of skel mesh DI.

Work Needed
------------------------

UI customization for these parameters in the stack and the component override panel.
As these are typeless probably the best thing here is to have the stack UI be an asset picker and the component UI be an actor picker?
Is there such a thing as a combined asset/actor picker? IDK.

Deterministic Randoms need to be utilized in other DIs/Functions.

diff --git a/Engine/Plugins/FX/Niagara/Source/Niagara/Classes/NiagaraDataInterface.h b/Engine/Plugins/FX/Niagara/Source/Niagara/Classes/NiagaraDataInterface.h
--- a/Engine/Plugins/FX/Niagara/Source/Niagara/Classes/NiagaraDataInterface.h
+++ b/Engine/Plugins/FX/Niagara/Source/Niagara/Classes/NiagaraDataInterface.h
@@ -432,3 +432,192 @@ struct TCurveUseLUTBinder
 		}
 	}
 };
+
+
+/** Helper class for decoding NDI parameters into a usable struct type. */
+template<typename T>
+struct FNDIParameter
+{
+	FNDIParameter(FVectorVMContext& Context)
+	{
+		static_assert(false, "You must specialize FNDIParameter for all types it's used with.");
+	}
+	FORCEINLINE void GetAndAdvance(T& OutValue)
+	{
+		static_assert(false, "You must specialize FNDIParameter for all types it's used with.");
+	}
+	FORCEINLINE bool IsConstant()const
+	{
+		static_assert(false, "You must specialize FNDIParameter for all types it's used with.");
+	}
+};
+
+template<>
+struct FNDIParameter<FNiagaraRandInfo>
+{
+	VectorVM::FExternalFuncInputHandler<int32> Seed1Param;
+	VectorVM::FExternalFuncInputHandler<int32> Seed2Param;
+	VectorVM::FExternalFuncInputHandler<int32> Seed3Param;
+	
+	FVectorVMContext& Context;
+
+	FNDIParameter(FVectorVMContext& InContext)
+		: Context(InContext)
+	{
+		Seed1Param.Init(Context);
+		Seed2Param.Init(Context);
+		Seed3Param.Init(Context);
+	}
+
+	FORCEINLINE void GetAndAdvance(FNiagaraRandInfo& OutValue)
+	{
+		OutValue.Seed1 = Seed1Param.GetAndAdvance();
+		OutValue.Seed2 = Seed2Param.GetAndAdvance();
+		OutValue.Seed3 = Seed3Param.GetAndAdvance();
+	}
+
+
+	FORCEINLINE bool IsConstant()const
+	{
+		return Seed1Param.IsConstant() && Seed2Param.IsConstant() && Seed3Param.IsConstant();
+	}
+};
+
+struct FNDIRandomHelper
+{
+	FNDIRandomHelper(FVectorVMContext& InContext)
+		: Context(InContext)
+		, RandParam(Context)
+	{
+
+	}
+
+	FORCEINLINE void GetAndAdvance()
+	{
+		RandParam.GetAndAdvance(RandInfo);
+	}
+
+	FORCEINLINE bool IsDeterministic()
+	{
+		return RandInfo.Seed3 != INDEX_NONE;
+	}
+
+	//////////////////////////////////////////////////////////////////////////
+	
+	FORCEINLINE_DEBUGGABLE FVector4 Rand4(int32 InstanceIndex)
+	{
+		if (IsDeterministic())
+		{
+			int32 RandomCounter = Context.RandCounters[InstanceIndex]++;
+
+			FIntVector4 v = FIntVector4(RandomCounter, RandInfo.Seed1, RandInfo.Seed2, RandInfo.Seed3) * 1664525 + FIntVector4(1013904223);
+
+			v.X += v.Y*v.W;
+			v.Y += v.Z*v.X;
+			v.Z += v.X*v.Y;
+			v.W += v.Y*v.Z;
+			v.X += v.Y*v.W;
+			v.Y += v.Z*v.X;
+			v.Z += v.X*v.Y;
+			v.W += v.Y*v.Z;
+
+			// NOTE(mv): We can use 24 bits of randomness, as all integers in [0, 2^24] 
+			//           are exactly representable in single precision floats.
+			//           We use the upper 24 bits as they tend to be higher quality.
+
+			// NOTE(mv): The divide can often be folded with the range scale in the rand functions
+			return FVector4((v >> 8) & 0x00ffffff) / 16777216.0; // 0x01000000 == 16777216
+			// return float4((v >> 8) & 0x00ffffff) * (1.0/16777216.0); // bugged, see UE-67738
+		}
+		else
+		{
+			return FVector4(Context.RandStream.GetFraction(), Context.RandStream.GetFraction(), Context.RandStream.GetFraction(), Context.RandStream.GetFraction());
+		}
+	}
+
+	FORCEINLINE_DEBUGGABLE FVector Rand3(int32 InstanceIndex)
+	{
+		if (IsDeterministic())
+		{
+			int32 RandomCounter = Context.RandCounters[InstanceIndex]++;
+
+			FIntVector v = FIntVector(RandInfo.Seed1, RandInfo.Seed2, RandomCounter | (RandInfo.Seed3 << 16)) * 1664525 + FIntVector(1013904223);
+
+			v.X += v.Y*v.Z;
+			v.Y += v.Z*v.X;
+			v.Z += v.X*v.Y;
+			v.X += v.Y*v.Z;
+			v.Y += v.Z*v.X;
+			v.Z += v.X*v.Y;
+
+			return FVector((v >> 8) & 0x00ffffff) / 16777216.0; // 0x01000000 == 16777216
+		}
+		else
+		{
+			return FVector(Context.RandStream.GetFraction(), Context.RandStream.GetFraction(), Context.RandStream.GetFraction());
+		}
+	}
+
+	FORCEINLINE_DEBUGGABLE FVector2D Rand2(int32 InstanceIndex)
+	{
+		if (IsDeterministic())
+		{
+			FVector Rand3D = Rand3(InstanceIndex);
+			return FVector2D(Rand3D.X, Rand3D.Y);
+		}
+		else
+		{
+			return FVector2D(Context.RandStream.GetFraction(), Context.RandStream.GetFraction());
+		}
+	}
+
+	FORCEINLINE_DEBUGGABLE float Rand(int32 InstanceIndex)
+	{
+		if (IsDeterministic())
+		{
+			return Rand3(InstanceIndex).X;
+		}
+		else
+		{
+			return Context.RandStream.GetFraction();
+		}
+	}
+
+	FORCEINLINE_DEBUGGABLE FVector4 RandRange(int32 InstanceIndex, FVector4 Min, FVector4 Max)
+	{
+		FVector4 Range = Max - Min;
+		return Min + (Rand(InstanceIndex) * Range);
+	}
+
+	FORCEINLINE_DEBUGGABLE FVector RandRange(int32 InstanceIndex, FVector Min, FVector Max)
+	{
+		FVector Range = Max - Min;
+		return Min + (Rand(InstanceIndex) * Range);
+	}
+
+	FORCEINLINE_DEBUGGABLE FVector2D RandRange(int32 InstanceIndex, FVector2D Min, FVector2D Max)
+	{
+		FVector2D Range = Max - Min;
+		return Min + (Rand(InstanceIndex) * Range);
+	}
+
+	FORCEINLINE_DEBUGGABLE float RandRange(int32 InstanceIndex, float Min, float Max)
+	{
+		float Range = Max - Min;
+		return Min + (Rand(InstanceIndex) * Range);
+	}
+
+	FORCEINLINE_DEBUGGABLE int32 RandRange(int32 InstanceIndex, int32 Min, int32 Max)
+	{
+		// NOTE: Scaling a uniform float range provides better distribution of 
+		//       numbers than using %.
+		// NOTE: Inclusive! So [0, x] instead of [0, x)
+		int32 Range = Max - Min;
+		return Min + (int(Rand(InstanceIndex) * (Range + 1)));
+	}
+
+	FVectorVMContext& Context;
+	FNDIParameter<FNiagaraRandInfo> RandParam;
+
+	FNiagaraRandInfo RandInfo;
+};
\ No newline at end of file
diff --git a/Engine/Plugins/FX/Niagara/Source/Niagara/Classes/NiagaraDataInterfaceSkeletalMesh.h b/Engine/Plugins/FX/Niagara/Source/Niagara/Classes/NiagaraDataInterfaceSkeletalMesh.h
--- a/Engine/Plugins/FX/Niagara/Source/Niagara/Classes/NiagaraDataInterfaceSkeletalMesh.h
+++ b/Engine/Plugins/FX/Niagara/Source/Niagara/Classes/NiagaraDataInterfaceSkeletalMesh.h
@@ -7,6 +7,7 @@
 #include "Rendering/SkeletalMeshRenderData.h"
 #include "Components/SkeletalMeshComponent.h"
 #include "Containers/Array.h"
+#include "NiagaraParameterStore.h"
 #include "NiagaraDataInterfaceSkeletalMesh.generated.h"
 
 class UNiagaraDataInterfaceSkeletalMesh;
@@ -353,6 +354,11 @@ struct FNDISkeletalMesh_InstanceData
 	//Cached ptr to component we sample from. 
 	TWeakObjectPtr<USceneComponent> Component;
 
+	/** A binding to the user ptr we're reading the mesh from (if we are). */
+	FNiagaraParameterDirectBinding<UObject*> UserParamBinding;
+
+	UObject* CachedUserParam;
+
 	USkeletalMesh* Mesh;
 
 	TWeakObjectPtr<USkeletalMesh> MeshSafe;
@@ -449,6 +455,10 @@ public:
 	/** The source actor from which to sample. Takes precedence over the direct mesh. */
 	UPROPERTY(EditAnywhere, Category = "Mesh")
 	AActor* Source;
+
+	/** Reference to a user parameter if we're reading one. */
+	UPROPERTY(EditAnywhere, Category = "Mesh")
+	FNiagaraUserParameterBinding MeshUserParameter;
 	
 	/** The source component from which to sample. Takes precedence over the direct mesh. Not exposed to the user, only indirectly accessible from blueprints. */
 	UPROPERTY(Transient)
@@ -512,7 +522,7 @@ public:
 #endif
 	//~ UNiagaraDataInterface interface END
 
-	static USkeletalMesh* GetSkeletalMeshHelper(UNiagaraDataInterfaceSkeletalMesh* Interface, class UNiagaraComponent* OwningComponent, TWeakObjectPtr<USceneComponent>& SceneComponent, USkeletalMeshComponent*& FoundSkelComp);
+	USkeletalMesh* GetSkeletalMesh(class UNiagaraComponent* OwningComponent, TWeakObjectPtr<USceneComponent>& SceneComponent, USkeletalMeshComponent*& FoundSkelComp, FNDISkeletalMesh_InstanceData* InstData = nullptr);
 
 	virtual void GetCommonHLSL(FString& OutHLSL) override;
 	virtual bool GetFunctionHLSL(const FName&  DefinitionFunctionName, FString InstanceFunctionName, FNiagaraDataInterfaceGPUParamInfo& ParamInfo, FString& OutHLSL) override;
@@ -586,7 +596,7 @@ public:
 
 private:
 	template<typename FilterMode, typename AreaWeightingMode>
-	FORCEINLINE int32 RandomTriIndex(FRandomStream& RandStream, FSkeletalMeshAccessorHelper& Accessor, FNDISkeletalMesh_InstanceData* InstData);
+	FORCEINLINE int32 RandomTriIndex(FNDIRandomHelper& RandHelper, FSkeletalMeshAccessorHelper& Accessor, FNDISkeletalMesh_InstanceData* InstData, int32 InstanceIndex);
 
 	template<typename FilterMode, typename AreaWeightingMode>
 	FORCEINLINE int32 GetSpecificTriangleCount(FSkeletalMeshAccessorHelper& Accessor, FNDISkeletalMesh_InstanceData* InstData);
diff --git a/Engine/Plugins/FX/Niagara/Source/Niagara/Private/NDISkeletalMesh_TriangleSampling.cpp b/Engine/Plugins/FX/Niagara/Source/Niagara/Private/NDISkeletalMesh_TriangleSampling.cpp
--- a/Engine/Plugins/FX/Niagara/Source/Niagara/Private/NDISkeletalMesh_TriangleSampling.cpp
+++ b/Engine/Plugins/FX/Niagara/Source/Niagara/Private/NDISkeletalMesh_TriangleSampling.cpp
@@ -43,6 +43,7 @@ void UNiagaraDataInterfaceSkeletalMesh::GetTriangleSamplingFunctions(TArray<FNia
 		FNiagaraFunctionSignature Sig;
 		Sig.Name = FSkeletalMeshInterfaceHelper::RandomTriCoordName;
 		Sig.Inputs.Add(FNiagaraVariable(FNiagaraTypeDefinition(GetClass()), TEXT("SkeletalMesh")));
+		Sig.Inputs.Add(FNiagaraVariable(FNiagaraTypeDefinition(FNiagaraRandInfo::StaticStruct()), TEXT("RandomInfo")));
 		Sig.Outputs.Add(FNiagaraVariable(FNiagaraTypeDefinition(FMeshTriCoordinate::StaticStruct()), TEXT("Coord")));
 		Sig.bMemberFunction = true;
 		Sig.bRequiresContext = false;
@@ -204,7 +205,7 @@ void UNiagaraDataInterfaceSkeletalMesh::BindTriangleSamplingFunction(const FVMEx
 
 	if (BindingInfo.Name == FSkeletalMeshInterfaceHelper::RandomTriCoordName)
 	{
-		check(BindingInfo.GetNumInputs() == 1 && BindingInfo.GetNumOutputs() == 4);
+		check(BindingInfo.GetNumInputs() == 4 && BindingInfo.GetNumOutputs() == 4);
 		TFilterModeBinder<TAreaWeightingModeBinder<NDI_FUNC_BINDER(UNiagaraDataInterfaceSkeletalMesh, RandomTriCoord)>>::Bind(this, BindingInfo, InstanceData, OutFunc);
 	}
 	else if (BindingInfo.Name == FSkeletalMeshInterfaceHelper::IsValidTriCoordName)
@@ -268,7 +269,7 @@ void UNiagaraDataInterfaceSkeletalMesh::BindTriangleSamplingFunction(const FVMEx
 }
 
 template<typename FilterMode, typename AreaWeightingMode>
-FORCEINLINE int32 UNiagaraDataInterfaceSkeletalMesh::RandomTriIndex(FRandomStream& RandStream, FSkeletalMeshAccessorHelper& Accessor, FNDISkeletalMesh_InstanceData* InstData)
+FORCEINLINE int32 UNiagaraDataInterfaceSkeletalMesh::RandomTriIndex(FNDIRandomHelper& RandHelper, FSkeletalMeshAccessorHelper& Accessor, FNDISkeletalMesh_InstanceData* InstData, int32 InstanceIndex)
 {
 	checkf(false, TEXT("Invalid template call for RandomTriIndex. Bug in Filter binding or Area Weighting binding. Contact code team."));
 	return 0;
@@ -278,11 +279,12 @@ template<>
 FORCEINLINE int32 UNiagaraDataInterfaceSkeletalMesh::RandomTriIndex<
 	TIntegralConstant<ENDISkeletalMesh_FilterMode, ENDISkeletalMesh_FilterMode::None>,
 	TIntegralConstant<ENDISkelMesh_AreaWeightingMode, ENDISkelMesh_AreaWeightingMode::None>>
-	(FRandomStream& RandStream, FSkeletalMeshAccessorHelper& Accessor, FNDISkeletalMesh_InstanceData* InstData)
+	(FNDIRandomHelper& RandHelper, FSkeletalMeshAccessorHelper& Accessor, FNDISkeletalMesh_InstanceData* InstData, int32 InstanceIndex)
 {
-	int32 SecIdx = RandStream.RandRange(0, Accessor.LODData->RenderSections.Num() - 1);
+	int32 SecIdx = RandHelper.RandRange(InstanceIndex, 0, Accessor.LODData->RenderSections.Num() - 1);
+	//int32 SecIdx = RandStream.RandRange(0, Accessor.LODData->RenderSections.Num() - 1);
 	FSkelMeshRenderSection& Sec = Accessor.LODData->RenderSections[SecIdx];
-	int32 Tri = RandStream.RandRange(0, Sec.NumTriangles - 1);
+	int32 Tri = RandHelper.RandRange(InstanceIndex, 0, Sec.NumTriangles - 1);
 	return (Sec.BaseIndex / 3) + Tri;
 }
 
@@ -290,11 +292,11 @@ template<>
 FORCEINLINE int32 UNiagaraDataInterfaceSkeletalMesh::RandomTriIndex<
 	TIntegralConstant<ENDISkeletalMesh_FilterMode, ENDISkeletalMesh_FilterMode::None>,
 	TIntegralConstant<ENDISkelMesh_AreaWeightingMode, ENDISkelMesh_AreaWeightingMode::AreaWeighted>>
-	(FRandomStream& RandStream, FSkeletalMeshAccessorHelper& Accessor, FNDISkeletalMesh_InstanceData* InstData)
+	(FNDIRandomHelper& RandHelper, FSkeletalMeshAccessorHelper& Accessor, FNDISkeletalMesh_InstanceData* InstData, int32 InstanceIndex)
 {
 	const FSkeletalMeshSamplingInfo& SamplingInfo = InstData->Mesh->GetSamplingInfo();
 	const FSkeletalMeshSamplingLODBuiltData& WholeMeshBuiltData = SamplingInfo.GetWholeMeshLODBuiltData(InstData->GetLODIndex());
-	int32 TriIdx = WholeMeshBuiltData.AreaWeightedTriangleSampler.GetEntryIndex(RandStream.GetFraction(), RandStream.GetFraction());
+	int32 TriIdx = WholeMeshBuiltData.AreaWeightedTriangleSampler.GetEntryIndex(RandHelper.Rand(InstanceIndex), RandHelper.Rand(InstanceIndex));
 	return TriIdx;
 }
 
@@ -302,9 +304,9 @@ template<>
 FORCEINLINE int32 UNiagaraDataInterfaceSkeletalMesh::RandomTriIndex<
 	TIntegralConstant<ENDISkeletalMesh_FilterMode, ENDISkeletalMesh_FilterMode::SingleRegion>,
 	TIntegralConstant<ENDISkelMesh_AreaWeightingMode, ENDISkelMesh_AreaWeightingMode::None>>
-	(FRandomStream& RandStream, FSkeletalMeshAccessorHelper& Accessor, FNDISkeletalMesh_InstanceData* InstData)
+	(FNDIRandomHelper& RandHelper, FSkeletalMeshAccessorHelper& Accessor, FNDISkeletalMesh_InstanceData* InstData, int32 InstanceIndex)
 {
-	int32 Idx = RandStream.RandRange(0, Accessor.SamplingRegionBuiltData->TriangleIndices.Num() - 1);
+	int32 Idx = RandHelper.RandRange(InstanceIndex, 0, Accessor.SamplingRegionBuiltData->TriangleIndices.Num() - 1);
 	return Accessor.SamplingRegionBuiltData->TriangleIndices[Idx] / 3;
 }
 
@@ -312,9 +314,9 @@ template<>
 FORCEINLINE int32 UNiagaraDataInterfaceSkeletalMesh::RandomTriIndex<
 	TIntegralConstant<ENDISkeletalMesh_FilterMode, ENDISkeletalMesh_FilterMode::SingleRegion>,
 	TIntegralConstant<ENDISkelMesh_AreaWeightingMode, ENDISkelMesh_AreaWeightingMode::AreaWeighted>>
-	(FRandomStream& RandStream, FSkeletalMeshAccessorHelper& Accessor, FNDISkeletalMesh_InstanceData* InstData)
+	(FNDIRandomHelper& RandHelper, FSkeletalMeshAccessorHelper& Accessor, FNDISkeletalMesh_InstanceData* InstData, int32 InstanceIndex)
 {
-	int32 Idx = Accessor.SamplingRegionBuiltData->AreaWeightedSampler.GetEntryIndex(RandStream.GetFraction(), RandStream.GetFraction());
+	int32 Idx = Accessor.SamplingRegionBuiltData->AreaWeightedSampler.GetEntryIndex(RandHelper.Rand(InstanceIndex), RandHelper.Rand(InstanceIndex));
 	return Accessor.SamplingRegionBuiltData->TriangleIndices[Idx] / 3;
 }
 
@@ -322,13 +324,13 @@ template<>
 FORCEINLINE int32 UNiagaraDataInterfaceSkeletalMesh::RandomTriIndex<
 	TIntegralConstant<ENDISkeletalMesh_FilterMode, ENDISkeletalMesh_FilterMode::MultiRegion>,
 	TIntegralConstant<ENDISkelMesh_AreaWeightingMode, ENDISkelMesh_AreaWeightingMode::None>>
-	(FRandomStream& RandStream, FSkeletalMeshAccessorHelper& Accessor, FNDISkeletalMesh_InstanceData* InstData)
+	(FNDIRandomHelper& RandHelper, FSkeletalMeshAccessorHelper& Accessor, FNDISkeletalMesh_InstanceData* InstData, int32 InstanceIndex)
 {
-	int32 RegionIdx = RandStream.RandRange(0, InstData->SamplingRegionIndices.Num() - 1);
+	int32 RegionIdx = RandHelper.RandRange(InstanceIndex, 0, InstData->SamplingRegionIndices.Num() - 1);
 	const FSkeletalMeshSamplingInfo& SamplingInfo = InstData->Mesh->GetSamplingInfo();
 	const FSkeletalMeshSamplingRegion& Region = SamplingInfo.GetRegion(InstData->SamplingRegionIndices[RegionIdx]);
 	const FSkeletalMeshSamplingRegionBuiltData& RegionBuiltData = SamplingInfo.GetRegionBuiltData(InstData->SamplingRegionIndices[RegionIdx]);
-	int32 Idx = RandStream.RandRange(0, RegionBuiltData.TriangleIndices.Num() - 1);
+	int32 Idx = RandHelper.RandRange(InstanceIndex, 0, RegionBuiltData.TriangleIndices.Num() - 1);
 	return RegionBuiltData.TriangleIndices[Idx] / 3;
 }
 
@@ -336,13 +338,13 @@ template<>
 FORCEINLINE int32 UNiagaraDataInterfaceSkeletalMesh::RandomTriIndex<
 	TIntegralConstant<ENDISkeletalMesh_FilterMode, ENDISkeletalMesh_FilterMode::MultiRegion>,
 	TIntegralConstant<ENDISkelMesh_AreaWeightingMode, ENDISkelMesh_AreaWeightingMode::AreaWeighted>>
-	(FRandomStream& RandStream, FSkeletalMeshAccessorHelper& Accessor, FNDISkeletalMesh_InstanceData* InstData)
+	(FNDIRandomHelper& RandHelper, FSkeletalMeshAccessorHelper& Accessor, FNDISkeletalMesh_InstanceData* InstData, int32 InstanceIndex)
 {
-	int32 RegionIdx = InstData->SamplingRegionAreaWeightedSampler.GetEntryIndex(RandStream.GetFraction(), RandStream.GetFraction());
+	int32 RegionIdx = InstData->SamplingRegionAreaWeightedSampler.GetEntryIndex(RandHelper.Rand(InstanceIndex), RandHelper.Rand(InstanceIndex));
 	const FSkeletalMeshSamplingInfo& SamplingInfo = InstData->Mesh->GetSamplingInfo();
 	const FSkeletalMeshSamplingRegion& Region = SamplingInfo.GetRegion(InstData->SamplingRegionIndices[RegionIdx]);
 	const FSkeletalMeshSamplingRegionBuiltData& RegionBuiltData = SamplingInfo.GetRegionBuiltData(InstData->SamplingRegionIndices[RegionIdx]);
-	int32 Idx = RegionBuiltData.AreaWeightedSampler.GetEntryIndex(RandStream.GetFraction(), RandStream.GetFraction());
+	int32 Idx = RegionBuiltData.AreaWeightedSampler.GetEntryIndex(RandHelper.Rand(InstanceIndex), RandHelper.Rand(InstanceIndex));
 	return RegionBuiltData.TriangleIndices[Idx] / 3;
 }
 
@@ -350,6 +352,10 @@ template<typename FilterMode, typename AreaWeightingMode>
 void UNiagaraDataInterfaceSkeletalMesh::RandomTriCoord(FVectorVMContext& Context)
 {
 	SCOPE_CYCLE_COUNTER(STAT_NiagaraSkel_Sample);
+
+	//FNDIParameter<FNiagaraRandInfo> RandInfoParam(Context);
+	FNDIRandomHelper RandHelper(Context);
+
 	VectorVM::FUserPtrHandler<FNDISkeletalMesh_InstanceData> InstData(Context);
 	checkfSlow(InstData.Get(), TEXT("Skeletal Mesh Interface has invalid instance data. %s"), *GetPathName());
 	checkfSlow(InstData->Mesh, TEXT("Skeletal Mesh Interface has invalid mesh. %s"), *GetPathName());
@@ -359,10 +365,11 @@ void UNiagaraDataInterfaceSkeletalMesh::RandomTriCoord(FVectorVMContext& Context
 
 	FSkeletalMeshAccessorHelper MeshAccessor;
 	MeshAccessor.Init<FilterMode, AreaWeightingMode>(InstData);
-
 	for (int32 i = 0; i < Context.NumInstances; ++i)
 	{
-		*OutTri.GetDest() = RandomTriIndex<FilterMode, AreaWeightingMode>(Context.RandStream, MeshAccessor, InstData);
+		RandHelper.GetAndAdvance();//We grab the rand info to a local value first so it can be used for multiple rand calls from the helper.
+
+		*OutTri.GetDest() = RandomTriIndex<FilterMode, AreaWeightingMode>(RandHelper, MeshAccessor, InstData, i);
 		FVector Bary = RandomBarycentricCoord(Context.RandStream);
 		*OutBaryX.GetDest() = Bary.X;		*OutBaryY.GetDest() = Bary.Y;		*OutBaryZ.GetDest() = Bary.Z;
 
diff --git a/Engine/Plugins/FX/Niagara/Source/Niagara/Private/NiagaraComponent.cpp b/Engine/Plugins/FX/Niagara/Source/Niagara/Private/NiagaraComponent.cpp
--- a/Engine/Plugins/FX/Niagara/Source/Niagara/Private/NiagaraComponent.cpp
+++ b/Engine/Plugins/FX/Niagara/Source/Niagara/Private/NiagaraComponent.cpp
@@ -1131,6 +1131,18 @@ void UNiagaraComponent::SetNiagaraVariableBool(const FString& InVariableName, bo
 	OverrideParameters.SetParameterValue(InValue ? FNiagaraBool::True : FNiagaraBool::False, FNiagaraVariable(FNiagaraTypeDefinition::GetBoolDef(), VarName), true);
 }
 
+void UNiagaraComponent::SetNiagaraVariableActor(const FString& InVariableName, AActor* InValue)
+{
+	SetNiagaraVariableObject(InVariableName, InValue);
+}
+
+void UNiagaraComponent::SetNiagaraVariableObject(const FString& InVariableName, UObject* InValue)
+{
+	FName VarName = FName(*InVariableName);
+
+	OverrideParameters.SetUObject(InValue, FNiagaraVariable(FNiagaraTypeDefinition::GetUObjectDef(), VarName));
+}
+
 TArray<FVector> UNiagaraComponent::GetNiagaraParticlePositions_DebugOnly(const FString& InEmitterName)
 {
 	return GetNiagaraParticleValueVec3_DebugOnly(InEmitterName, TEXT("Position"));
diff --git a/Engine/Plugins/FX/Niagara/Source/Niagara/Private/NiagaraDataInterfaceSkeletalMesh.cpp b/Engine/Plugins/FX/Niagara/Source/Niagara/Private/NiagaraDataInterfaceSkeletalMesh.cpp
--- a/Engine/Plugins/FX/Niagara/Source/Niagara/Private/NiagaraDataInterfaceSkeletalMesh.cpp
+++ b/Engine/Plugins/FX/Niagara/Source/Niagara/Private/NiagaraDataInterfaceSkeletalMesh.cpp
@@ -876,17 +876,65 @@ void UNiagaraDataInterfaceSkeletalMesh::ProvidePerInstanceDataForRenderThread(vo
 	Data->MeshSkinWeightBufferSrv = SourceData->MeshSkinWeightBufferSrv;
 }
 
-USkeletalMesh* UNiagaraDataInterfaceSkeletalMesh::GetSkeletalMeshHelper(UNiagaraDataInterfaceSkeletalMesh* Interface, UNiagaraComponent* OwningComponent, TWeakObjectPtr<USceneComponent>& SceneComponent, USkeletalMeshComponent*& FoundSkelComp)
+USkeletalMesh* UNiagaraDataInterfaceSkeletalMesh::GetSkeletalMesh(UNiagaraComponent* OwningComponent, TWeakObjectPtr<USceneComponent>& SceneComponent, USkeletalMeshComponent*& FoundSkelComp, FNDISkeletalMesh_InstanceData* InstData)
 {
 	USkeletalMesh* Mesh = nullptr;
-	if (Interface->SourceComponent)
+	if (MeshUserParameter.Parameter.IsValid() && InstData)
 	{
-		Mesh = Interface->SourceComponent->SkeletalMesh;
-		FoundSkelComp = Interface->SourceComponent;
+		FNiagaraSystemInstance* SystemInstance = OwningComponent->GetSystemInstance();
+		if (UObject* UserParamObject = InstData->UserParamBinding.Init(SystemInstance->GetInstanceParameters(), MeshUserParameter.Parameter))
+		{
+			InstData->CachedUserParam = UserParamObject;
+			if (USkeletalMeshComponent* UserSkelMeshComp = Cast<USkeletalMeshComponent>(UserParamObject))
+			{
+				FoundSkelComp = UserSkelMeshComp;
+				Mesh = FoundSkelComp->SkeletalMesh;
+			}
+			else if (ASkeletalMeshActor* UserSkelMeshActor = Cast<ASkeletalMeshActor>(UserParamObject))
+			{
+				FoundSkelComp = UserSkelMeshActor->GetSkeletalMeshComponent();
+				Mesh = FoundSkelComp->SkeletalMesh;
+			}
+			else if (AActor* Actor = Cast<AActor>(UserParamObject))
+			{
+				TArray<UActorComponent*> SourceComps = Actor->GetComponentsByClass(USkeletalMeshComponent::StaticClass());
+				for (UActorComponent* ActorComp : SourceComps)
+				{
+					USkeletalMeshComponent* SourceComp = Cast<USkeletalMeshComponent>(ActorComp);
+					if (SourceComp)
+					{
+						USkeletalMesh* PossibleMesh = SourceComp->SkeletalMesh;
+						if (PossibleMesh != nullptr/* && PossibleMesh->bAllowCPUAccess*/)
+						{
+							Mesh = PossibleMesh;
+							FoundSkelComp = SourceComp;
+							break;
+						}
+					}
+				}
+			}
+			else
+			{
+				//We have a valid, non-null UObject parameter type but it is not a type we can use to get a skeletal mesh from. 
+				UE_LOG(LogNiagara, Warning, TEXT("SkeletalMesh data interface using object parameter with invalid type. Skeletal Mesh Data Interfaces can only get a valid mesh from SkeletalMeshComponents, SkeletalMeshActors or Actors."));
+				UE_LOG(LogNiagara, Warning, TEXT("Invalid Parameter : %s"), *UserParamObject->GetFullName());
+				UE_LOG(LogNiagara, Warning, TEXT("Niagara Component : %s"), *OwningComponent->GetFullName());
+				UE_LOG(LogNiagara, Warning, TEXT("System : %s"), *OwningComponent->GetAsset()->GetFullName());
+			}
+		}
+		else
+		{
+			//WARNING - We have a valid user param but the object set is null.
+		}
 	}
-	else if (Interface->Source)
+	else if (SourceComponent)
 	{
-		ASkeletalMeshActor* MeshActor = Cast<ASkeletalMeshActor>(Interface->Source);
+		Mesh = SourceComponent->SkeletalMesh;
+		FoundSkelComp = SourceComponent;
+	}
+	else if (Source)
+	{
+		ASkeletalMeshActor* MeshActor = Cast<ASkeletalMeshActor>(Source);
 		USkeletalMeshComponent* SourceComp = nullptr;
 		if (MeshActor != nullptr)
 		{
@@ -894,7 +942,7 @@ USkeletalMesh* UNiagaraDataInterfaceSkeletalMesh::GetSkeletalMeshHelper(UNiagara
 		}
 		else
 		{
-			SourceComp = Interface->Source->FindComponentByClass<USkeletalMeshComponent>();
+			SourceComp = Source->FindComponentByClass<USkeletalMeshComponent>();
 		}
 
 		if (SourceComp)
@@ -904,7 +952,7 @@ USkeletalMesh* UNiagaraDataInterfaceSkeletalMesh::GetSkeletalMeshHelper(UNiagara
 		}
 		else
 		{
-			SceneComponent = Interface->Source->GetRootComponent();
+			SceneComponent = Source->GetRootComponent();
 		}
 	}
 	else
@@ -952,9 +1000,9 @@ USkeletalMesh* UNiagaraDataInterfaceSkeletalMesh::GetSkeletalMeshHelper(UNiagara
 		SceneComponent = FoundSkelComp;
 	}
 	
-	if (!Mesh && Interface->DefaultMesh)
+	if (!Mesh && DefaultMesh)
 	{
-		Mesh = Interface->DefaultMesh;
+		Mesh = DefaultMesh;
 	}
 
 	return Mesh;
@@ -973,7 +1021,7 @@ bool FNDISkeletalMesh_InstanceData::Init(UNiagaraDataInterfaceSkeletalMesh* Inte
 	TransformInverseTransposed = FMatrix::Identity;
 	PrevTransform = FMatrix::Identity;
 	PrevTransformInverseTransposed = FMatrix::Identity;
-	DeltaSeconds = 0.0f;
+	DeltaSeconds = SystemInstance->GetComponent()->GetWorld()->GetDeltaSeconds();
 	ChangeId = Interface->ChangeId;
 	bIsGpuUniformlyDistributedSampling = false;
 	MeshWeightStrideByte = 0;
@@ -982,7 +1030,7 @@ bool FNDISkeletalMesh_InstanceData::Init(UNiagaraDataInterfaceSkeletalMesh* Inte
 
 	// Get skel mesh and confirm have valid data
 	USkeletalMeshComponent* NewSkelComp = nullptr;
-	Mesh = UNiagaraDataInterfaceSkeletalMesh::GetSkeletalMeshHelper(Interface, SystemInstance->GetComponent(), Component, NewSkelComp);
+	Mesh = Interface->GetSkeletalMesh(SystemInstance->GetComponent(), Component, NewSkelComp, this);
 	MeshSafe = Mesh;
 
 	if (!Mesh)
@@ -1285,13 +1333,23 @@ bool FNDISkeletalMesh_InstanceData::ResetRequired(UNiagaraDataInterfaceSkeletalM
 		//The component we were bound to is no longer valid so we have to trigger a reset.
 		return true;
 	}
-		
+	
 	if (USkeletalMeshComponent* SkelComp = Cast<USkeletalMeshComponent>(Comp))
 	{
-		if (!SkelComp->SkeletalMesh)
+		if (!SkelComp->SkeletalMesh)//TODO: Handle clearing the mesh gracefully.
 		{
 			return true;
 		}
+
+		//If the user ptr has been changed to look at a new mesh component. TODO: Handle more gracefully.
+		if (Interface->MeshUserParameter.Parameter.IsValid())
+		{
+			UObject* NewUserParam = UserParamBinding.GetValue();
+			if (CachedUserParam != NewUserParam)
+			{
+				return true;
+			}
+		}
 		
 		// Handle the case where they've procedurally swapped out the skeletal mesh from
 		// the one we previously cached data for.
@@ -1553,6 +1611,7 @@ bool UNiagaraDataInterfaceSkeletalMesh::CopyToInternal(UNiagaraDataInterface* De
 
 	UNiagaraDataInterfaceSkeletalMesh* OtherTyped = CastChecked<UNiagaraDataInterfaceSkeletalMesh>(Destination);
 	OtherTyped->Source = Source;
+	OtherTyped->MeshUserParameter = MeshUserParameter;
 	OtherTyped->DefaultMesh = DefaultMesh;
 	OtherTyped->SkinningMode = SkinningMode;
 	OtherTyped->SamplingRegions = SamplingRegions;
@@ -1574,6 +1633,7 @@ bool UNiagaraDataInterfaceSkeletalMesh::Equals(const UNiagaraDataInterface* Othe
 	const UNiagaraDataInterfaceSkeletalMesh* OtherTyped = CastChecked<const UNiagaraDataInterfaceSkeletalMesh>(Other);
 	return OtherTyped->Source == Source &&
 		OtherTyped->DefaultMesh == DefaultMesh &&
+		OtherTyped->MeshUserParameter == MeshUserParameter &&
 		OtherTyped->SkinningMode == SkinningMode &&
 		OtherTyped->SamplingRegions == SamplingRegions &&
 		OtherTyped->WholeMeshLOD == WholeMeshLOD &&
diff --git a/Engine/Plugins/FX/Niagara/Source/Niagara/Private/NiagaraModule.cpp b/Engine/Plugins/FX/Niagara/Source/Niagara/Private/NiagaraModule.cpp
--- a/Engine/Plugins/FX/Niagara/Source/Niagara/Private/NiagaraModule.cpp
+++ b/Engine/Plugins/FX/Niagara/Source/Niagara/Private/NiagaraModule.cpp
@@ -476,6 +476,8 @@ UScriptStruct* FNiagaraTypeDefinition::Vec2Struct;
 UScriptStruct* FNiagaraTypeDefinition::ColorStruct;
 UScriptStruct* FNiagaraTypeDefinition::QuatStruct;
 
+UClass* FNiagaraTypeDefinition::UObjectClass;
+
 UEnum* FNiagaraTypeDefinition::ExecutionStateEnum;
 UEnum* FNiagaraTypeDefinition::SimulationTargetEnum;
 UEnum* FNiagaraTypeDefinition::ExecutionStateSourceEnum;
@@ -494,6 +496,8 @@ FNiagaraTypeDefinition FNiagaraTypeDefinition::Vec2Def;
 FNiagaraTypeDefinition FNiagaraTypeDefinition::ColorDef;
 FNiagaraTypeDefinition FNiagaraTypeDefinition::QuatDef;
 
+FNiagaraTypeDefinition FNiagaraTypeDefinition::UObjectDef;
+
 TSet<UScriptStruct*> FNiagaraTypeDefinition::NumericStructs;
 TArray<FNiagaraTypeDefinition> FNiagaraTypeDefinition::OrderedNumericTypes;
 
@@ -535,6 +539,8 @@ void FNiagaraTypeDefinition::Init()
 	FNiagaraTypeDefinition::Vec4Struct = FindObjectChecked<UScriptStruct>(CoreUObjectPkg, TEXT("Vector4"));
 	FNiagaraTypeDefinition::ColorStruct = FindObjectChecked<UScriptStruct>(CoreUObjectPkg, TEXT("LinearColor"));
 	FNiagaraTypeDefinition::QuatStruct = FindObjectChecked<UScriptStruct>(CoreUObjectPkg, TEXT("Quat"));
+
+	FNiagaraTypeDefinition::UObjectClass = UObject::StaticClass();
 	
 	ParameterMapDef = FNiagaraTypeDefinition(ParameterMapStruct);
 	IDDef = FNiagaraTypeDefinition(IDStruct);
@@ -549,6 +555,8 @@ void FNiagaraTypeDefinition::Init()
 	QuatDef = FNiagaraTypeDefinition(QuatStruct);
 	Matrix4Def = FNiagaraTypeDefinition(Matrix4Struct);
 
+	UObjectDef = FNiagaraTypeDefinition(UObjectClass);
+
 	CollisionEventDef = FNiagaraTypeDefinition(FNiagaraCollisionEventPayload::StaticStruct());
 	NumericStructs.Add(NumericStruct);
 	NumericStructs.Add(FloatStruct);
@@ -633,6 +641,8 @@ void FNiagaraTypeDefinition::RecreateUserDefinedTypeRegistry()
 	UScriptStruct* SpawnInfoStruct = FindObjectChecked<UScriptStruct>(NiagaraPkg, TEXT("NiagaraSpawnInfo"));
 	FNiagaraTypeRegistry::Register(FNiagaraTypeDefinition(SpawnInfoStruct), true, false, false);
 
+	FNiagaraTypeRegistry::Register(UObjectDef, true, false, false);
+
 	const UNiagaraSettings* Settings = GetDefault<UNiagaraSettings>();
 	check(Settings);
 	TArray<FSoftObjectPath> TotalStructAssets;
@@ -708,6 +718,7 @@ void FNiagaraTypeDefinition::RecreateUserDefinedTypeRegistry()
 		}
 	}
 
+	FNiagaraTypeRegistry::Register(FNiagaraRandInfo::StaticStruct(), true, true, true);
 }
 
 bool FNiagaraTypeDefinition::IsScalarDefinition(const FNiagaraTypeDefinition& Type)
diff --git a/Engine/Plugins/FX/Niagara/Source/Niagara/Private/NiagaraParameterStore.cpp b/Engine/Plugins/FX/Niagara/Source/Niagara/Private/NiagaraParameterStore.cpp
--- a/Engine/Plugins/FX/Niagara/Source/Niagara/Private/NiagaraParameterStore.cpp
+++ b/Engine/Plugins/FX/Niagara/Source/Niagara/Private/NiagaraParameterStore.cpp
@@ -38,6 +38,7 @@ FNiagaraParameterStore::FNiagaraParameterStore(UObject* InOwner)
 	: Owner(InOwner)
 	, bParametersDirty(true)
 	, bInterfacesDirty(true)
+	, bUObjectsDirty(true)
 	, LayoutVersion(0)
 {
 #if WITH_EDITORONLY_DATA
@@ -61,6 +62,7 @@ FNiagaraParameterStore& FNiagaraParameterStore::operator=(const FNiagaraParamete
 	ParameterData = Other.ParameterData;
 	INC_MEMORY_STAT_BY(STAT_NiagaraParamStoreMemory, ParameterData.Num());
 	DataInterfaces = Other.DataInterfaces;
+	UObjects = Other.UObjects;
 	++LayoutVersion;
 #if WITH_EDITOR
 	OnChangedDelegate.Broadcast();
@@ -187,7 +189,7 @@ void FNiagaraParameterStore::Tick()
 #if NIAGARA_NAN_CHECKING
 	CheckForNaNs();
 #endif
-	if (bParametersDirty || bInterfacesDirty)
+	if (bParametersDirty || bInterfacesDirty || bUObjectsDirty)
 	{
 		for (TPair<FNiagaraParameterStore*, FNiagaraParameterStoreBinding>& Binding : Bindings)
 		{
@@ -199,6 +201,7 @@ void FNiagaraParameterStore::Tick()
 		//We have to have ticked all our source stores before now.
 		bParametersDirty = false;
 		bInterfacesDirty = false;
+		bUObjectsDirty = false;
 	}
 }
 
@@ -276,6 +279,12 @@ bool FNiagaraParameterStore::AddParameter(const FNiagaraVariable& Param, bool bI
 		ParameterOffsets.Add(Param) = Offset;
 		DataInterfaces[Offset] = bInitInterfaces ? NewObject<UNiagaraDataInterface>(Owner, const_cast<UClass*>(Param.GetType().GetClass()), NAME_None, RF_Transactional | RF_Public) : nullptr;
 	}
+	else if (Param.GetType().IsUObject())
+	{
+		int32 Offset = UObjects.AddDefaulted();
+		ParameterOffsets.Add(Param) = Offset;
+		//UObjects[Offset] = nullptr;
+	}
 	else
 	{
 		int32 ParamSize = Param.GetSizeInBytes();
@@ -321,6 +330,7 @@ bool FNiagaraParameterStore::RemoveParameter(const FNiagaraVariable& ToRemove)
 		TMap<FNiagaraVariable, int32> NewOffsets;
 		TArray<uint8> NewData;
 		TArray<UNiagaraDataInterface*> NewInterfaces;
+		TArray<UObject*> NewUObjects;
 		for (TPair<FNiagaraVariable, int32>& Existing : ParameterOffsets)
 		{
 			FNiagaraVariable& ExistingVar = Existing.Key;
@@ -335,6 +345,12 @@ bool FNiagaraParameterStore::RemoveParameter(const FNiagaraVariable& ToRemove)
 					NewOffsets.Add(ExistingVar) = Offset;
 					NewInterfaces[Offset] = DataInterfaces[ExistingOffset];
 				}
+				else if (ExistingVar.IsUObject())
+				{
+					int32 Offset = NewUObjects.AddDefaulted();
+					NewOffsets.Add(ExistingVar) = Offset;
+					NewUObjects[Offset] = UObjects[ExistingOffset];
+				}
 				else
 				{
 					int32 Offset = NewData.Num();
@@ -353,6 +369,7 @@ bool FNiagaraParameterStore::RemoveParameter(const FNiagaraVariable& ToRemove)
 		ParameterOffsets = NewOffsets;
 		ParameterData = NewData;
 		DataInterfaces = NewInterfaces;
+		UObjects = NewUObjects;
 
 		OnLayoutChange();
 		return true;
@@ -385,6 +402,10 @@ void FNiagaraParameterStore::RenameParameter(const FNiagaraVariable& Param, FNam
 		{
 			SetDataInterface(GetDataInterface(Idx), NewIdx);
 		}
+		else if (Param.IsUObject())
+		{
+			SetUObject(GetUObject(Idx), NewIdx);
+		}
 		else
 		{
 			SetParameterData(GetParameterData_Internal(Idx), NewIdx, Param.GetSizeInBytes());
@@ -460,6 +481,10 @@ void FNiagaraParameterStore::CopyParametersTo(FNiagaraParameterStore& DestStore,
 					checkf(false, TEXT("A data interface copy method must be specified if the parameter store has data interfaces."));
 				}
 			}
+			else if (Parameter.IsUObject())
+			{
+				DestStore.SetUObject(GetUObject(SrcIndex), DestIndex);//UObjects are just refs to external objects. They never need to be deep copied.
+			}
 			else
 			{
 				if (ParameterData.Num() != 0)
@@ -498,10 +523,13 @@ void FNiagaraParameterStore::InitFromSource(const FNiagaraParameterStore* SrcSto
 
 	DataInterfaces = SrcStore->DataInterfaces;
 
+	UObjects = SrcStore->UObjects;
+
 	if (bNotifyAsDirty)
 	{
 		MarkParametersDirty();
 		MarkInterfacesDirty();
+		MarkUObjectsDirty();
 		OnLayoutChange();
 	}
 }
@@ -526,6 +554,9 @@ void FNiagaraParameterStore::Empty(bool bClearBindings)
 	DEC_MEMORY_STAT_BY(STAT_NiagaraParamStoreMemory, ParameterData.Num());
 
 	DataInterfaces.Empty();
+
+	UObjects.Empty();
+
 	if (bClearBindings)
 	{
 		UnbindFromSourceStores();
@@ -541,6 +572,9 @@ void FNiagaraParameterStore::Reset(bool bClearBindings)
 	DEC_MEMORY_STAT_BY(STAT_NiagaraParamStoreMemory, ParameterData.Num());
 
 	DataInterfaces.Reset();
+
+	UObjects.Reset();
+
 	if (bClearBindings)
 	{
 		UnbindFromSourceStores();
diff --git a/Engine/Plugins/FX/Niagara/Source/Niagara/Public/NiagaraCommon.h b/Engine/Plugins/FX/Niagara/Source/Niagara/Public/NiagaraCommon.h
--- a/Engine/Plugins/FX/Niagara/Source/Niagara/Public/NiagaraCommon.h
+++ b/Engine/Plugins/FX/Niagara/Source/Niagara/Public/NiagaraCommon.h
@@ -573,4 +573,35 @@ namespace FNiagaraUtilities
 	 */
 	void NIAGARA_API PrepareRapidIterationParameters(const TArray<UNiagaraScript*>& Scripts, const TMap<UNiagaraScript*, UNiagaraScript*>& ScriptDependencyMap, const TMap<UNiagaraScript*, FString>& ScriptToEmitterNameMap);
 #endif
-};
\ No newline at end of file
+};
+
+USTRUCT()
+struct FNiagaraUserParameterBinding
+{
+	GENERATED_USTRUCT_BODY()
+
+	FNiagaraUserParameterBinding();
+
+	UPROPERTY(EditAnywhere, Category = "User Parameter")
+	FNiagaraVariable Parameter;
+
+	FORCEINLINE bool operator==(const FNiagaraUserParameterBinding& Other)const
+	{
+		return Other.Parameter == Parameter;
+	}
+};
+
+USTRUCT()
+struct FNiagaraRandInfo
+{
+	GENERATED_USTRUCT_BODY()
+
+	UPROPERTY(EditAnywhere, Category = "Random")
+	int32 Seed1;
+	
+	UPROPERTY(EditAnywhere, Category = "Random")
+	int32 Seed2;
+
+	UPROPERTY(EditAnywhere, Category = "Random")
+	int32 Seed3;
+};
diff --git a/Engine/Plugins/FX/Niagara/Source/Niagara/Public/NiagaraComponent.h b/Engine/Plugins/FX/Niagara/Source/Niagara/Public/NiagaraComponent.h
--- a/Engine/Plugins/FX/Niagara/Source/Niagara/Public/NiagaraComponent.h
+++ b/Engine/Plugins/FX/Niagara/Source/Niagara/Public/NiagaraComponent.h
@@ -241,7 +241,10 @@ public:
 	void SetNiagaraVariableBool(const FString& InVariableName, bool InValue);
 
 	UFUNCTION(BlueprintCallable, Category = Niagara, meta = (DisplayName = "Set Niagara Variable (Actor)"))
-	void SetNiagaraVariableActor(const FString& InVariableName, AActor* Actor) {} // TODO sckime ?????? fix this!!!!
+	void SetNiagaraVariableActor(const FString& InVariableName, AActor* Actor);
+
+	UFUNCTION(BlueprintCallable, Category = Niagara, meta = (DisplayName = "Set Niagara Variable (Object)"))
+	void SetNiagaraVariableObject(const FString& InVariableName, UObject* Object);
 
 	/** Debug accessors for getting positions in blueprints. */
 	UFUNCTION(BlueprintCallable, Category = Niagara, meta = (DisplayName = "Get Niagara Emitter Positions"))
diff --git a/Engine/Plugins/FX/Niagara/Source/Niagara/Public/NiagaraParameterStore.h b/Engine/Plugins/FX/Niagara/Source/Niagara/Public/NiagaraParameterStore.h
--- a/Engine/Plugins/FX/Niagara/Source/Niagara/Public/NiagaraParameterStore.h
+++ b/Engine/Plugins/FX/Niagara/Source/Niagara/Public/NiagaraParameterStore.h
@@ -52,6 +52,25 @@ struct FNiagaraParameterStoreBinding
 	/** Bindings of data interfaces. Src and Dest offsets.*/
 	TArray<FInterfaceBinding> InterfaceBindings;
 
+	struct FUObjectBinding
+	{
+		FUObjectBinding(int32 InSrcOffset, int32 InDestOffset)
+			: SrcOffset((uint16)InSrcOffset), DestOffset((uint16)InDestOffset)
+		{
+			//No way these offsets could ever be higher but check just in case.
+			check(InSrcOffset < (int32)0xFFFF);
+			check(InDestOffset < (int32)0xFFFF);
+		}
+		FORCEINLINE bool operator==(const FUObjectBinding& Other)const
+		{
+			return SrcOffset == Other.SrcOffset && DestOffset == Other.DestOffset;
+		}
+		uint16 SrcOffset;
+		uint16 DestOffset;
+	};
+	/** Bindings of UObject params. Src and Dest offsets.*/
+	TArray<FUObjectBinding> UObjectBindings;
+
 	FNiagaraParameterStoreBinding()
 	{
 	}
@@ -95,6 +114,10 @@ private:
 	UPROPERTY()
 	TArray<UNiagaraDataInterface*> DataInterfaces;
 
+	/** UObjects referenced by this store. Also indexed by ParameterOffsets.*/
+	UPROPERTY()
+	TArray<UObject*> UObjects;
+
 	/** Bindings between this parameter store and others we push data into when we tick. */
 	TMap<FNiagaraParameterStore*, FNiagaraParameterStoreBinding> Bindings;
 
@@ -105,6 +128,8 @@ private:
 	uint32 bParametersDirty : 1;
 	/** Marks our interfaces as dirty. They will be pushed to any bound stores on tick if true. */
 	uint32 bInterfacesDirty : 1;
+	/** Marks our UObjects as dirty. They will be pushed to any bound stores on tick if true. */
+	uint32 bUObjectsDirty : 1;
 
 	/** Uniquely identifies the current layout of this parameter store for detecting layout changes. */
 	uint32 LayoutVersion;
@@ -133,9 +158,11 @@ public:
 
 	FORCEINLINE uint32 GetParametersDirty() const { return bParametersDirty; }
 	FORCEINLINE uint32 GetInterfacesDirty() const { return bInterfacesDirty; }
+	FORCEINLINE uint32 GetUObjectsDirty() const { return bUObjectsDirty; }
 
 	FORCEINLINE void MarkParametersDirty() { bParametersDirty = true; }
 	FORCEINLINE void MarkInterfacesDirty() { bInterfacesDirty = true; }
+	FORCEINLINE void MarkUObjectsDirty() { bUObjectsDirty = true; }
 
 	FORCEINLINE uint32 GetLayoutVersion() const { return LayoutVersion; }
 
@@ -181,6 +208,7 @@ public:
 	FORCEINLINE void GetParameters(TArray<FNiagaraVariable>& OutParameters) const { return ParameterOffsets.GenerateKeyArray(OutParameters); }
 	FORCEINLINE int32 GetNumParameters() const { return ParameterOffsets.Num(); }
 
+	FORCEINLINE const TArray<UObject*>& GetUObjects()const { return UObjects; }
 	FORCEINLINE const TArray<UNiagaraDataInterface*>& GetDataInterfaces()const { return DataInterfaces; }
 	FORCEINLINE const TArray<uint8>& GetParameterDataArray()const { return ParameterData; }
 
@@ -268,6 +296,25 @@ public:
 		return ParameterOffsets.Find(Parameter);
 	}
 
+	/** Returns the UObject at the passed offset. */
+	FORCEINLINE UObject* GetUObject(int32 Offset)const
+	{
+		if (UObjects.IsValidIndex(Offset))
+		{
+			return UObjects[Offset];
+		}
+
+		return nullptr;
+	}
+
+	FORCEINLINE UObject* GetUObject(const FNiagaraVariable& Parameter)const
+	{
+		int32 Offset = IndexOf(Parameter);
+		UObject* Obj = GetUObject(Offset);
+		checkSlow(!Obj || Parameter.GetType().GetClass() == Obj->GetClass());
+		return Obj;
+	}
+
 	/** Copies the passed parameter from this parameter store into another. */
 	FORCEINLINE_DEBUGGABLE void CopyParameterData(FNiagaraParameterStore& DestStore, const FNiagaraVariable& Parameter) const
 	{
@@ -280,6 +327,10 @@ public:
 				DataInterfaces[SrcIndex]->CopyTo(DestStore.DataInterfaces[DestIndex]);
 				DestStore.OnInterfaceChange();
 			}
+			else if (Parameter.IsUObject())
+			{
+				DestStore.SetUObject(GetUObject(SrcIndex), DestIndex);
+			}
 			else
 			{
 				DestStore.SetParameterData(GetParameterData_Internal(SrcIndex), DestIndex, Parameter.GetSizeInBytes());
@@ -421,6 +472,22 @@ public:
 		}
 	}
 
+	FORCEINLINE_DEBUGGABLE void SetUObject(UObject* InObject, int32 Offset)
+	{
+		UObjects[Offset] = InObject;
+		OnUObjectChange();
+	}
+
+	FORCEINLINE_DEBUGGABLE void SetUObject(UObject* InObject, const FNiagaraVariable& Parameter)
+	{
+		int32 Offset = IndexOf(Parameter);
+		if (Offset != INDEX_NONE)
+		{
+			UObjects[Offset] = InObject;
+			OnUObjectChange();
+		}
+	}
+
 	FORCEINLINE void OnParameterChange() 
 	{ 
 		bParametersDirty = true;
@@ -437,6 +504,14 @@ public:
 #endif
 	}
 
+	FORCEINLINE void OnUObjectChange()
+	{
+		bUObjectsDirty = true;
+#if WITH_EDITOR
+		OnChangedDelegate.Broadcast();
+#endif
+	}
+
 #if WITH_EDITOR
 	FDelegateHandle AddOnChangedHandler(FOnChanged::FDelegate InOnChanged);
 	void RemoveOnChangedHandler(FDelegateHandle DelegateHandle);
@@ -480,12 +555,14 @@ FORCEINLINE_DEBUGGABLE void FNiagaraParameterStoreBinding::Empty(FNiagaraParamet
 	DestStore = nullptr;
 	ParameterBindings.Reset();
 	InterfaceBindings.Reset();
+	UObjectBindings.Reset();
 }
 
 FORCEINLINE_DEBUGGABLE void FNiagaraParameterStoreBinding::BindParameters(FNiagaraParameterStore* DestStore, FNiagaraParameterStore* SrcStore)
 {
 	InterfaceBindings.Reset();
 	ParameterBindings.Reset();
+	UObjectBindings.Reset();
 	for (const TPair<FNiagaraVariable, int32>& ParamOffsetPair : DestStore->GetParameterOffsets())
 	{
 		const FNiagaraVariable& Parameter = ParamOffsetPair.Key;
@@ -498,6 +575,10 @@ FORCEINLINE_DEBUGGABLE void FNiagaraParameterStoreBinding::BindParameters(FNiaga
 			{
 				InterfaceBindings.Add(FInterfaceBinding(SrcOffset, DestOffset));
 			}
+			else if (Parameter.IsUObject())
+			{
+				UObjectBindings.Add(FUObjectBinding(SrcOffset, DestOffset));
+			}
 			else
 			{
 				ParameterBindings.Add(FParameterBinding(SrcOffset, DestOffset, Parameter.GetSizeInBytes()));
@@ -537,6 +618,14 @@ FORCEINLINE_DEBUGGABLE bool FNiagaraParameterStoreBinding::VerifyBinding(const F
 				bBindingValid = false;
 			}
 		}
+		else if (Parameter.IsUObject())
+		{
+			if (UObjectBindings.Find(FUObjectBinding(SrcOffset, DestOffset)) == INDEX_NONE)
+			{
+				MissingParameterNames.Add(Parameter.GetName());
+				bBindingValid = false;
+			}
+		}
 		else
 		{
 			if (ParameterBindings.Find(FParameterBinding(SrcOffset, DestOffset, Parameter.GetSizeInBytes())) == INDEX_NONE)
@@ -580,6 +669,14 @@ FORCEINLINE_DEBUGGABLE void FNiagaraParameterStoreBinding::Tick(FNiagaraParamete
 		}
 	}
 
+	if (SrcStore->GetUObjectsDirty() || bForce)
+	{
+		for (FUObjectBinding& Binding : UObjectBindings)
+		{
+			DestStore->SetUObject(SrcStore->GetUObject(Binding.SrcOffset), Binding.DestOffset);
+		}
+	}
+
 #if NIAGARA_NAN_CHECKING
 	DestStore->CheckForNaNs();
 #endif
@@ -653,6 +750,38 @@ FORCEINLINE_DEBUGGABLE void FNiagaraParameterStoreBinding::Dump(const FNiagaraPa
 			UE_LOG(LogNiagara, Log, TEXT("Failed to find matching data interface param in bound store!\n"));
 		}
 	}
+
+	for (const FUObjectBinding& Binding : UObjectBindings)
+	{
+		ensure(Binding.SrcOffset != -1);
+		ensure(Binding.DestOffset != -1);
+		FNiagaraVariable Param;
+		bool bFound = false;
+		for (const TPair<FNiagaraVariable, int32>& ParamOffsetPair : DestStore->GetParameterOffsets())
+		{
+			if (ParamOffsetPair.Value == Binding.DestOffset && ParamOffsetPair.Key.IsUObject())
+			{
+				bFound = true;
+				Param = ParamOffsetPair.Key;
+			}
+		}
+		if (ensure(bFound))
+		{
+			UE_LOG(LogNiagara, Log, TEXT("| UObject | %s %s: Src:%d - Dest:%d\n"), *Param.GetType().GetName(), *Param.GetName().ToString(), Binding.SrcOffset, Binding.DestOffset);
+
+			//Also ensure the param has been pushed correctly.
+			const UObject* SrcData = SrcStore->GetUObjects()[Binding.SrcOffset];
+			const UObject* DestData = DestStore->GetUObjects()[Binding.DestOffset];
+			if (!ensure(SrcData == DestData))
+			{
+				UE_LOG(LogNiagara, Log, TEXT("Data interface parameter in dest store is incorrect!\n"));
+			}
+		}
+		else
+		{
+			UE_LOG(LogNiagara, Log, TEXT("Failed to find matching data interface param in bound store!\n"));
+		}
+	}
 #endif
 }
 
@@ -918,4 +1047,56 @@ struct FNiagaraParameterDirectBinding<FNiagaraBool>
 		}
 		return Ret;
 	}
+};
+
+template<>
+struct FNiagaraParameterDirectBinding<UObject*>
+{
+	int32 UObjectOffset;
+	FNiagaraParameterStore* BoundStore;
+	FNiagaraVariable BoundVariable;
+	uint32 LayoutVersion;
+
+	FNiagaraParameterDirectBinding()
+		: UObjectOffset(INDEX_NONE), BoundStore(nullptr)
+	{}
+
+	UObject* Init(FNiagaraParameterStore& InStore, const FNiagaraVariable& DestVariable)
+	{
+		if (DestVariable.IsValid())
+		{
+			BoundStore = &InStore;
+			BoundVariable = DestVariable;
+			LayoutVersion = BoundStore->GetLayoutVersion();
+
+			check(BoundVariable.GetType() == FNiagaraTypeDefinition::GetUObjectDef());
+			UObjectOffset = BoundStore->IndexOf(DestVariable);
+			UObject* Ret = BoundStore->GetUObject(UObjectOffset);
+			return Ret;
+		}
+		return nullptr;
+	}
+
+	FORCEINLINE void SetValue(UObject* InValue)
+	{
+		if (UObjectOffset != INDEX_NONE)
+		{
+			check(BoundVariable.GetType() == FNiagaraTypeDefinition::GetUObjectDef());
+			checkf(LayoutVersion == BoundStore->GetLayoutVersion(), TEXT("This binding is invalid, its bound parameter store's layout was changed since it was created"));
+
+			BoundStore->SetUObject(InValue, UObjectOffset);
+		}
+	}
+
+	FORCEINLINE UObject* GetValue()const
+	{
+		if (UObjectOffset != INDEX_NONE)
+		{
+			check(BoundVariable.GetType() == FNiagaraTypeDefinition::GetUObjectDef());
+			checkf(LayoutVersion == BoundStore->GetLayoutVersion(), TEXT("This binding is invalid, its bound parameter store's layout was changed since it was created"));
+
+			return BoundStore->GetUObject(UObjectOffset);
+		}
+		return nullptr;
+	}
 };
\ No newline at end of file
diff --git a/Engine/Plugins/FX/Niagara/Source/Niagara/Public/NiagaraTypes.h b/Engine/Plugins/FX/Niagara/Source/Niagara/Public/NiagaraTypes.h
--- a/Engine/Plugins/FX/Niagara/Source/Niagara/Public/NiagaraTypes.h
+++ b/Engine/Plugins/FX/Niagara/Source/Niagara/Public/NiagaraTypes.h
@@ -477,6 +477,11 @@ public:
 
 	bool IsDataInterface()const;
 
+	FORCEINLINE bool IsUObject()const
+	{
+		return Struct->IsChildOf<UObject>();
+	}
+
 	bool IsEnum() const { return Enum != nullptr; }
 	
 	int32 GetSize()const
@@ -555,6 +560,7 @@ public:
 	static const FNiagaraTypeDefinition& GetGenericNumericDef() { return NumericDef; }
 	static const FNiagaraTypeDefinition& GetParameterMapDef() { return ParameterMapDef; }
 	static const FNiagaraTypeDefinition& GetIDDef() { return IDDef; }
+	static const FNiagaraTypeDefinition& GetUObjectDef() { return UObjectDef; }
 
 	static UScriptStruct* GetFloatStruct() { return FloatStruct; }
 	static UScriptStruct* GetBoolStruct() { return BoolStruct; }
@@ -607,6 +613,7 @@ private:
 	static FNiagaraTypeDefinition NumericDef;
 	static FNiagaraTypeDefinition ParameterMapDef;
 	static FNiagaraTypeDefinition IDDef;
+	static FNiagaraTypeDefinition UObjectDef;
 
 	static UScriptStruct* FloatStruct;
 	static UScriptStruct* BoolStruct;
@@ -619,6 +626,8 @@ private:
 	static UScriptStruct* Matrix4Struct;
 	static UScriptStruct* NumericStruct;
 
+	static UClass* UObjectClass;
+
 	static UEnum* SimulationTargetEnum;
 	static UEnum* ScriptUsageEnum;
 	static UEnum* ExecutionStateEnum;
@@ -793,6 +802,7 @@ struct FNiagaraVariable
 	const FNiagaraTypeDefinition& GetType()const { return TypeDef; }
 
 	FORCEINLINE bool IsDataInterface()const { return GetType().IsDataInterface(); }
+	FORCEINLINE bool IsUObject()const { return GetType().IsUObject(); }
 
 	void AllocateData()
 	{
diff --git a/Engine/Plugins/FX/Niagara/Source/NiagaraCore/Private/NiagaraCustomVersion.cpp b/Engine/Plugins/FX/Niagara/Source/NiagaraCore/Private/NiagaraCustomVersion.cpp
--- a/Engine/Plugins/FX/Niagara/Source/NiagaraCore/Private/NiagaraCustomVersion.cpp
+++ b/Engine/Plugins/FX/Niagara/Source/NiagaraCore/Private/NiagaraCustomVersion.cpp
@@ -8,4 +8,4 @@ const FGuid FNiagaraCustomVersion::GUID(0xFCF57AFA, 0x50764283, 0xB9A9E658, 0xFF
 // Register the custom version with core
 FCustomVersionRegistration GRegisterNiagaraCustomVersion(FNiagaraCustomVersion::GUID, FNiagaraCustomVersion::LatestVersion, TEXT("NiagaraVer"));
 
-const FGuid FNiagaraCustomVersion::LatestScriptCompileVersion(0x76F1E483, 0x558E5F4D, 0xBF7A18BD, 0x591D49F1);
+const FGuid FNiagaraCustomVersion::LatestScriptCompileVersion(0x784CB4D0, 0x1529422B, 0xB6880B12, 0xFB5A1E3F);
\ No newline at end of file
diff --git a/Engine/Plugins/FX/Niagara/Source/NiagaraEditor/Private/Customizations/NiagaraComponentDetails.cpp b/Engine/Plugins/FX/Niagara/Source/NiagaraEditor/Private/Customizations/NiagaraComponentDetails.cpp
--- a/Engine/Plugins/FX/Niagara/Source/NiagaraEditor/Private/Customizations/NiagaraComponentDetails.cpp
+++ b/Engine/Plugins/FX/Niagara/Source/NiagaraEditor/Private/Customizations/NiagaraComponentDetails.cpp
@@ -109,13 +109,7 @@ public:
 
 			TSharedPtr<SWidget> CustomValueWidget;
 
-			if (!Parameter.IsDataInterface())
-			{
-				TSharedPtr<FStructOnScope> StructOnScope = MakeShareable(new FStructOnScope(Parameter.GetType().GetStruct(), (uint8*)ParamStore.GetParameterData(Parameter)));
-				Row = ChildrenBuilder.AddExternalStructureProperty(StructOnScope.ToSharedRef(), NAME_None, Parameter.GetName());
-
-			}
-			else 
+			if (Parameter.IsDataInterface())
 			{
 				int32 DataInterfaceOffset = ParamStore.IndexOf(Parameter);
 				UObject* DefaultValueObject = ParamStore.GetDataInterfaces()[DataInterfaceOffset];
@@ -132,6 +126,45 @@ public:
 					.TextStyle(FNiagaraEditorStyle::Get(), "NiagaraEditor.ParameterText")
 					.Text(FText::FromString(FName::NameToDisplayString(DefaultValueObject->GetClass()->GetName(), false)));
 			}
+			else if (Parameter.IsUObject())
+			{
+				int32 ObjectOffst = ParamStore.IndexOf(Parameter);
+				UObject* DefaultValueObject = ParamStore.GetUObjects()[ObjectOffst];
+
+				TArray<UObject*> Objects;
+				Objects.Add(DefaultValueObject);
+
+				//How do I set this up so I can have this pick actors from the level?
+
+				TOptional<bool> bAllowChildrenOverride(true);
+				TOptional<bool> bCreateCategoryNodesOverride(false);
+				Row = ChildrenBuilder.AddExternalObjectProperty(Objects, NAME_None, Parameter.GetName(), bAllowChildrenOverride, bCreateCategoryNodesOverride);
+
+				//How do I make this an object picker from the level editor?
+				//Neither of these seem to work well.
+// 				CustomValueWidget = PropertyCustomizationHelpers::MakeActorPickerWithMenu(nullptr,
+// 					true,
+// 					FOnShouldFilterActor::CreateRaw(this, &FNiagaraComponentNodeBuilder::IsFilteredActor),
+// 					FOnActorSelected::CreateRaw(this, &FNiagaraComponentNodeBuilder::OnActorSelected),
+// 					FSimpleDelegate::CreateRaw(this, &FNiagaraComponentNodeBuilder::CloseComboButton),
+// 					FSimpleDelegate::CreateRaw(this, &FNiagaraComponentNodeBuilder::OnUse));
+				
+// 				CustomValueWidget = PropertyCustomizationHelpers::MakeActorPickerAnchorButton(
+// 					FOnGetActorFilters::CreateRaw(this, &FNiagaraComponentNodeBuilder::OnGetActorFiltersForSceneOutliner),
+// 					FOnActorSelected::CreateRaw(this, &FNiagaraComponentNodeBuilder::OnActorSelected));
+
+				FString ValueName = DefaultValueObject ? FName::NameToDisplayString(DefaultValueObject->GetClass()->GetName(), false) : TEXT("null");
+				CustomValueWidget =
+					SNew(STextBlock)
+					.TextStyle(FNiagaraEditorStyle::Get(), "NiagaraEditor.ParameterText")
+					.Text(FText::FromString(ValueName));
+			}
+			else
+			{
+				TSharedPtr<FStructOnScope> StructOnScope = MakeShareable(new FStructOnScope(Parameter.GetType().GetStruct(), (uint8*)ParamStore.GetParameterData(Parameter)));
+				Row = ChildrenBuilder.AddExternalStructureProperty(StructOnScope.ToSharedRef(), NAME_None, Parameter.GetName());
+
+			}
 
 			if (Row)
 			{
@@ -158,6 +191,17 @@ public:
 					PropertyHandle->SetOnChildPropertyValueChanged(
 						FSimpleDelegate::CreateRaw(this, &FNiagaraComponentNodeBuilder::OnDataInterfaceChanged, Parameter));
 				}
+				else if (Parameter.IsUObject())
+				{
+					PropertyHandle->SetOnPropertyValuePreChange(
+						FSimpleDelegate::CreateRaw(this, &FNiagaraComponentNodeBuilder::OnUObjectPreChange, Parameter));
+					PropertyHandle->SetOnChildPropertyValuePreChange(
+						FSimpleDelegate::CreateRaw(this, &FNiagaraComponentNodeBuilder::OnUObjectPreChange, Parameter));
+					PropertyHandle->SetOnPropertyValueChanged(
+						FSimpleDelegate::CreateRaw(this, &FNiagaraComponentNodeBuilder::OnUObjectChanged, Parameter));
+					PropertyHandle->SetOnChildPropertyValueChanged(
+						FSimpleDelegate::CreateRaw(this, &FNiagaraComponentNodeBuilder::OnUObjectChanged, Parameter));
+				}
 				else
 				{
 					PropertyHandle->SetOnPropertyValuePreChange(
@@ -231,6 +275,12 @@ private:
 		check(Component.IsValid());
 		Component->Modify();
 	}
+	
+	void OnUObjectPreChange(FNiagaraVariable Var)
+	{
+		check(Component.IsValid());
+		Component->Modify();
+	}
 
 	void OnParameterChanged(FNiagaraVariable Var)
 	{
@@ -246,6 +296,13 @@ private:
 		Component->SetParameterValueOverriddenLocally(Var, true, true);
 	}
 
+	void OnUObjectChanged(FNiagaraVariable Var)
+	{
+		Component->GetOverrideParameters().OnUObjectChange();
+		check(Component.IsValid());
+		Component->SetParameterValueOverriddenLocally(Var, true, true);
+	}
+
 	bool DoesParameterDifferFromDefault(const FNiagaraVariable& Var)
 	{
 		check(Component.IsValid());
diff --git a/Engine/Plugins/FX/Niagara/Source/NiagaraEditor/Private/Customizations/NiagaraTypeCustomizations.cpp b/Engine/Plugins/FX/Niagara/Source/NiagaraEditor/Private/Customizations/NiagaraTypeCustomizations.cpp
--- a/Engine/Plugins/FX/Niagara/Source/NiagaraEditor/Private/Customizations/NiagaraTypeCustomizations.cpp
+++ b/Engine/Plugins/FX/Niagara/Source/NiagaraEditor/Private/Customizations/NiagaraTypeCustomizations.cpp
@@ -20,6 +20,7 @@
 #include "NiagaraConstants.h"
 #include "NiagaraNodeParameterMapBase.h"
 #include "Widgets/Input/STextComboBox.h"
+#include "NiagaraSystem.h"
 
 #define LOCTEXT_NAMESPACE "FNiagaraVariableAttributeBindingCustomization"
 
@@ -298,4 +299,185 @@ void FNiagaraVariableAttributeBindingCustomization::CustomizeHeader(TSharedRef<I
 }
 
 
+//////////////////////////////////////////////////////////////////////////
+
+FText FNiagaraUserParameterBindingCustomization::GetCurrentText() const
+{
+	if (BaseSystem && TargetUserParameterBinding)
+	{
+		return FText::FromName(TargetUserParameterBinding->Parameter.GetName());
+	}
+	return FText::FromString(TEXT("Missing"));
+}
+
+FText FNiagaraUserParameterBindingCustomization::GetTooltipText() const
+{
+	if (BaseSystem && TargetUserParameterBinding && TargetUserParameterBinding->Parameter.IsValid())
+	{
+		FText TooltipDesc = FText::Format(LOCTEXT("BindingTooltip", "Bound to the user parameter \"{0}\""), FText::FromName(TargetUserParameterBinding->Parameter.GetName()));
+		return TooltipDesc;
+	}
+	return FText::FromString(TEXT("Missing"));
+}
+
+TSharedRef<SWidget> FNiagaraUserParameterBindingCustomization::OnGetMenuContent() const
+{
+	FGraphActionMenuBuilder MenuBuilder;
+
+	return SNew(SBorder)
+		.BorderImage(FEditorStyle::GetBrush("Menu.Background"))
+		.Padding(5)
+		[
+			SNew(SBox)
+			[
+				SNew(SGraphActionMenu)
+				.OnActionSelected(this, &FNiagaraUserParameterBindingCustomization::OnActionSelected)
+		.OnCreateWidgetForAction(SGraphActionMenu::FOnCreateWidgetForAction::CreateSP(this, &FNiagaraUserParameterBindingCustomization::OnCreateWidgetForAction))
+		.OnCollectAllActions(this, &FNiagaraUserParameterBindingCustomization::CollectAllActions)
+		.AutoExpandActionMenu(false)
+		.ShowFilterTextBox(true)
+			]
+		];
+}
+
+TArray<FName> FNiagaraUserParameterBindingCustomization::GetNames() const
+{
+	TArray<FName> Names;
+
+	if (BaseSystem && TargetUserParameterBinding)
+	{
+		for (auto& Pair : BaseSystem->GetExposedParameters().GetParameterOffsets())
+		{
+			const FNiagaraVariable& Var = Pair.Key;
+			if (FNiagaraParameterMapHistory::IsUserParameter(Var) && Var.GetType() == TargetUserParameterBinding->Parameter.GetType())
+			{
+				Names.AddUnique(Var.GetName());
+			}
+		}
+	}
+	
+	return Names;
+}
+
+void FNiagaraUserParameterBindingCustomization::CollectAllActions(FGraphActionListBuilderBase& OutAllActions)
+{
+	TArray<FName> UserParamNames = GetNames();
+	for (FName UserParamName : UserParamNames)
+	{
+		FText CategoryName = FText();
+		FString DisplayNameString = FName::NameToDisplayString(UserParamName.ToString(), false);
+		const FText NameText = FText::FromString(DisplayNameString);
+		const FText TooltipDesc = FText::Format(LOCTEXT("SetFunctionPopupTooltip", "Bind to the User Parameter \"{0}\" "), FText::FromString(DisplayNameString));
+		TSharedPtr<FNiagaraStackAssetAction_VarBind> NewNodeAction(new FNiagaraStackAssetAction_VarBind(UserParamName, CategoryName, NameText,
+			TooltipDesc, 0, FText()));
+		OutAllActions.AddAction(NewNodeAction);
+	}
+}
+
+TSharedRef<SWidget> FNiagaraUserParameterBindingCustomization::OnCreateWidgetForAction(struct FCreateWidgetForActionData* const InCreateData)
+{
+	return SNew(SVerticalBox)
+		+ SVerticalBox::Slot()
+		.AutoHeight()
+		[
+			SNew(STextBlock)
+			.Text(InCreateData->Action->GetMenuDescription())
+		.ToolTipText(InCreateData->Action->GetTooltipDescription())
+		];
+}
+
+
+void FNiagaraUserParameterBindingCustomization::OnActionSelected(const TArray< TSharedPtr<FEdGraphSchemaAction> >& SelectedActions, ESelectInfo::Type InSelectionType)
+{
+	if (InSelectionType == ESelectInfo::OnMouseClick || InSelectionType == ESelectInfo::OnKeyPress || SelectedActions.Num() == 0)
+	{
+		for (int32 ActionIndex = 0; ActionIndex < SelectedActions.Num(); ActionIndex++)
+		{
+			TSharedPtr<FEdGraphSchemaAction> CurrentAction = SelectedActions[ActionIndex];
+
+			if (CurrentAction.IsValid())
+			{
+				FSlateApplication::Get().DismissAllMenus();
+				FNiagaraStackAssetAction_VarBind* EventSourceAction = (FNiagaraStackAssetAction_VarBind*)CurrentAction.Get();
+				ChangeSource(EventSourceAction->VarName);
+			}
+		}
+	}
+}
+
+void FNiagaraUserParameterBindingCustomization::ChangeSource(FName InVarName)
+{
+	FScopedTransaction Transaction(FText::Format(LOCTEXT("ChangeSource", " Change User Parameter Source to \"{0}\" "), FText::FromName(InVarName)));
+	TArray<UObject*> Objects;
+	PropertyHandle->GetOuterObjects(Objects);
+	for (UObject* Obj : Objects)
+	{
+		Obj->Modify();
+	}
+
+	PropertyHandle->NotifyPreChange();
+	TargetUserParameterBinding->Parameter.SetName(InVarName);
+	TargetUserParameterBinding->Parameter.SetType(FNiagaraTypeDefinition::GetUObjectDef());
+	//TargetVariableBinding->DataSetVariable = FNiagaraConstants::GetAttributeAsDataSetKey(TargetVariableBinding->BoundVariable);
+	PropertyHandle->NotifyPostChange();
+	PropertyHandle->NotifyFinishedChangingProperties();
+}
+
+void FNiagaraUserParameterBindingCustomization::CustomizeHeader(TSharedRef<IPropertyHandle> InPropertyHandle, FDetailWidgetRow& HeaderRow, IPropertyTypeCustomizationUtils& CustomizationUtils)
+{
+	PropertyHandle = InPropertyHandle;
+	TArray<UObject*> Objects;
+	PropertyHandle->GetOuterObjects(Objects);
+	bool bAddDefault = true;
+	if (Objects.Num() == 1)
+	{
+		BaseSystem = Objects[0]->GetTypedOuter<UNiagaraSystem>();
+		if (BaseSystem)
+		{
+			TargetUserParameterBinding = (FNiagaraUserParameterBinding*)PropertyHandle->GetValueBaseAddress((uint8*)Objects[0]);
+
+			HeaderRow
+				.NameContent()
+				[
+					PropertyHandle->CreatePropertyNameWidget()
+				]
+			.ValueContent()
+				.MaxDesiredWidth(200.f)
+				[
+					SNew(SComboButton)
+					.OnGetMenuContent(this, &FNiagaraUserParameterBindingCustomization::OnGetMenuContent)
+				.ContentPadding(1)
+				.ToolTipText(this, &FNiagaraUserParameterBindingCustomization::GetTooltipText)
+				.ButtonContent()
+				[
+					SNew(STextBlock)
+					.Text(this, &FNiagaraUserParameterBindingCustomization::GetCurrentText)
+				.Font(IDetailLayoutBuilder::GetDetailFont())
+				]
+				];
+			bAddDefault = false;
+		}
+	}
+
+
+	if (bAddDefault)
+	{
+		HeaderRow
+			.NameContent()
+			[
+				PropertyHandle->CreatePropertyNameWidget()
+			]
+		.ValueContent()
+			.MaxDesiredWidth(200.f)
+			[
+				SNew(STextBlock)
+				.Text(FText::FromString(FName::NameToDisplayString(Cast<UStructProperty>(PropertyHandle->GetProperty())->Struct->GetName(), false)))
+			.Font(IDetailLayoutBuilder::GetDetailFont())
+			];
+	}
+}
+
+
+
+
 #undef LOCTEXT_NAMESPACE
diff --git a/Engine/Plugins/FX/Niagara/Source/NiagaraEditor/Private/Customizations/NiagaraTypeCustomizations.h b/Engine/Plugins/FX/Niagara/Source/NiagaraEditor/Private/Customizations/NiagaraTypeCustomizations.h
--- a/Engine/Plugins/FX/Niagara/Source/NiagaraEditor/Private/Customizations/NiagaraTypeCustomizations.h
+++ b/Engine/Plugins/FX/Niagara/Source/NiagaraEditor/Private/Customizations/NiagaraTypeCustomizations.h
@@ -126,4 +126,33 @@ private:
 	class UNiagaraEmitter* BaseEmitter;
 	struct FNiagaraVariableAttributeBinding* TargetVariableBinding;
 
+};
+
+class FNiagaraUserParameterBindingCustomization : public IPropertyTypeCustomization
+{
+public:
+	/** @return A new instance of this class */
+	static TSharedRef<IPropertyTypeCustomization> MakeInstance()
+	{
+		return MakeShared<FNiagaraUserParameterBindingCustomization>();
+	}
+
+	/** IPropertyTypeCustomization interface begin */
+	virtual void CustomizeHeader(TSharedRef<IPropertyHandle> StructPropertyHandle, class FDetailWidgetRow& HeaderRow, IPropertyTypeCustomizationUtils& StructCustomizationUtils) override;
+	virtual void CustomizeChildren(TSharedRef<IPropertyHandle> StructPropertyHandle, class IDetailChildrenBuilder& ChildBuilder, IPropertyTypeCustomizationUtils& StructCustomizationUtils) override {};
+	/** IPropertyTypeCustomization interface end */
+private:
+	FText GetCurrentText() const;
+	FText GetTooltipText() const;
+	TSharedRef<SWidget> OnGetMenuContent() const;
+	TArray<FName> GetNames() const;
+	void ChangeSource(FName InVarName);
+	void CollectAllActions(FGraphActionListBuilderBase& OutAllActions);
+	TSharedRef<SWidget> OnCreateWidgetForAction(struct FCreateWidgetForActionData* const InCreateData);
+	void OnActionSelected(const TArray< TSharedPtr<FEdGraphSchemaAction> >& SelectedActions, ESelectInfo::Type InSelectionType);
+
+	TSharedPtr<IPropertyHandle> PropertyHandle;
+	class UNiagaraSystem* BaseSystem;
+	struct FNiagaraUserParameterBinding* TargetUserParameterBinding;
+
 };
\ No newline at end of file
diff --git a/Engine/Plugins/FX/Niagara/Source/NiagaraEditor/Private/NiagaraEditorModule.cpp b/Engine/Plugins/FX/Niagara/Source/NiagaraEditor/Private/NiagaraEditorModule.cpp
--- a/Engine/Plugins/FX/Niagara/Source/NiagaraEditor/Private/NiagaraEditorModule.cpp
+++ b/Engine/Plugins/FX/Niagara/Source/NiagaraEditor/Private/NiagaraEditorModule.cpp
@@ -95,8 +95,12 @@
 #include "Customizations/NiagaraStaticSwitchNodeDetails.h"
 #include "Customizations/NiagaraFunctionCallNodeDetails.h"
 
+#include "NiagaraNodeFunctionCall.h"
+
 IMPLEMENT_MODULE( FNiagaraEditorModule, NiagaraEditor );
 
+PRAGMA_DISABLE_OPTIMIZATION
+
 #define LOCTEXT_NAMESPACE "NiagaraEditorModule"
 
 const FName FNiagaraEditorModule::NiagaraEditorAppIdentifier( TEXT( "NiagaraEditorApp" ) );
@@ -314,7 +318,7 @@ void PreventSystemRecompile(FAssetData SystemAsset, TSet<UNiagaraEmitter*>& InOu
 		{
 			CompileEmitterStandAlone(EmitterHandle.GetInstance(), InOutCompiledEmitters);
 		}
-
+		
 		System->MarkPackageDirty();
 		System->RequestCompile(false);
 		System->WaitForCompilationComplete();
@@ -367,6 +371,66 @@ void PreventAllSystemRecompiles()
 	GWarn->EndSlowTask();
 }
 
+void UpgradeAllNiagaraAssets()
+{
+	//First Load All Niagara Assets.
+	const FText SlowTaskText_Load = NSLOCTEXT("NiagaraEditor", "UpgradeAllNiagaraAssets_Load", "Loading all Niagara Assets ready to upgrade.");
+	GWarn->BeginSlowTask(SlowTaskText_Load, true, true);
+
+	FAssetRegistryModule& AssetRegistryModule = FModuleManager::LoadModuleChecked<FAssetRegistryModule>("AssetRegistry");
+
+	TArray<FAssetData> SystemAssets;
+	AssetRegistryModule.Get().GetAssetsByClass(UNiagaraSystem::StaticClass()->GetFName(), SystemAssets);
+
+	TArray<UNiagaraSystem*> Systems;
+	Systems.Reserve(SystemAssets.Num());
+	TSet<UNiagaraEmitter*> CompiledEmitters;
+	int32 ItemIndex = 0;
+	for (FAssetData& SystemAsset : SystemAssets)
+	{
+		if (GWarn->ReceivedUserCancel())
+		{
+			return;
+		}
+		GWarn->UpdateProgress(ItemIndex++, SystemAssets.Num());
+
+		UNiagaraSystem* System = Cast<UNiagaraSystem>(SystemAsset.GetAsset());
+		if (System != nullptr)
+		{
+			Systems.Add(System);
+		}
+	}
+
+	GWarn->EndSlowTask();
+
+	//////////////////////////////////////////////////////////////////////////
+
+	//Now process any data that needs to be updated.
+	const FText SlowTaskText_Upgrade = NSLOCTEXT("NiagaraEditor", "UpgradeAllNiagaraAssets_Upgrade", "Upgrading All Niagara Assets.");
+	GWarn->BeginSlowTask(SlowTaskText_Upgrade, true, true);
+
+	//Upgrade any data interface function call nodes.
+	TArray<UObject*> FunctionCallNodes;
+	GetObjectsOfClass(UNiagaraNodeFunctionCall::StaticClass(), FunctionCallNodes);
+	ItemIndex = 0;
+	for (UObject* Object : FunctionCallNodes)
+	{
+		if (GWarn->ReceivedUserCancel())
+		{
+			return;
+		}
+
+		if (UNiagaraNodeFunctionCall* FuncCallNode = Cast<UNiagaraNodeFunctionCall>(Object))
+		{
+			FuncCallNode->UpgradeDIFunctionCalls();
+		}
+
+		GWarn->UpdateProgress(ItemIndex++, FunctionCallNodes.Num());
+	}
+
+	GWarn->EndSlowTask();
+}
+
 class FNiagaraSystemBoolParameterTrackEditor : public FNiagaraSystemParameterTrackEditor<UMovieSceneNiagaraBoolParameterTrack, UMovieSceneBoolSection>
 {
 	virtual TSharedRef<ISequencerSection> MakeSectionInterface(UMovieSceneSection& SectionObject, UMovieSceneTrack& Track, FGuid ObjectBinding) override
@@ -441,6 +505,10 @@ void FNiagaraEditorModule::StartupModule()
 		FOnGetPropertyTypeCustomizationInstance::CreateStatic(&FNiagaraVariableAttributeBindingCustomization::MakeInstance)
 	);
 
+	PropertyModule.RegisterCustomPropertyTypeLayout("NiagaraUserParameterBinding",
+		FOnGetPropertyTypeCustomizationInstance::CreateStatic(&FNiagaraUserParameterBindingCustomization::MakeInstance)
+	);
+
 	FNiagaraEditorStyle::Initialize();
 	FNiagaraEditorCommands::Register();
 
@@ -595,6 +663,11 @@ void FNiagaraEditorModule::StartupModule()
 		TEXT("Loads all of the systems in the project and forces each system to refresh all it's dependencies so it won't recompile on load.  This may mark multiple assets dirty for re-saving."),
 		FConsoleCommandDelegate::CreateStatic(&PreventAllSystemRecompiles));
 
+	UpgradeAllNiagaraAssetsCommand = IConsoleManager::Get().RegisterConsoleCommand(
+		TEXT("fx.UpgradeAllNiagaraAssets"),
+		TEXT("Loads all Niagara assets and preforms any data upgrade processes required. This may mark multiple assets dirty for re-saving."),
+		FConsoleCommandDelegate::CreateStatic(&UpgradeAllNiagaraAssets));
+
 	if (GIsEditor)
 	{
 		UThumbnailManager::Get().RegisterCustomRenderer(UNiagaraEmitter::StaticClass(), UNiagaraEmitterThumbnailRenderer::StaticClass());
@@ -836,4 +909,6 @@ void FNiagaraEditorModule::OnPreGarbageCollection()
 	}
 }
 
+PRAGMA_ENABLE_OPTIMIZATION
+
 #undef LOCTEXT_NAMESPACE
diff --git a/Engine/Plugins/FX/Niagara/Source/NiagaraEditor/Private/NiagaraEditorUtilities.cpp b/Engine/Plugins/FX/Niagara/Source/NiagaraEditor/Private/NiagaraEditorUtilities.cpp
--- a/Engine/Plugins/FX/Niagara/Source/NiagaraEditor/Private/NiagaraEditorUtilities.cpp
+++ b/Engine/Plugins/FX/Niagara/Source/NiagaraEditor/Private/NiagaraEditorUtilities.cpp
@@ -97,7 +97,7 @@ void FNiagaraEditorUtilities::InitializeParameterInputNode(UNiagaraNodeInput& In
 			InputNode.SetDataInterface(nullptr);
 		}
 	}
-	else
+	else if(Type.IsDataInterface())
 	{
 		InputNode.Input.AllocateData(); // Frees previously used memory if we're switching from a struct to a class type.
 		InputNode.SetDataInterface(NewObject<UNiagaraDataInterface>(&InputNode, const_cast<UClass*>(Type.GetClass()), NAME_None, RF_Transactional));
diff --git a/Engine/Plugins/FX/Niagara/Source/NiagaraEditor/Private/NiagaraNodeAssignment.cpp b/Engine/Plugins/FX/Niagara/Source/NiagaraEditor/Private/NiagaraNodeAssignment.cpp
--- a/Engine/Plugins/FX/Niagara/Source/NiagaraEditor/Private/NiagaraNodeAssignment.cpp
+++ b/Engine/Plugins/FX/Niagara/Source/NiagaraEditor/Private/NiagaraNodeAssignment.cpp
@@ -246,7 +246,7 @@ void UNiagaraNodeAssignment::BuildCreateParameterMenu(FMenuBuilder& MenuBuilder,
 		}
 
 		TArray<FNiagaraTypeDefinition> AvailableTypes;
-		FNiagaraStackGraphUtilities::GetNewParameterAvailableTypes(AvailableTypes);
+		FNiagaraStackGraphUtilities::GetNewParameterAvailableTypes(AvailableTypes, NewParameterNamespace.GetValue());
 		for (const FNiagaraTypeDefinition& AvailableType : AvailableTypes)
 		{
 			// Make generic new parameter name
diff --git a/Engine/Plugins/FX/Niagara/Source/NiagaraEditor/Private/NiagaraNodeFunctionCall.cpp b/Engine/Plugins/FX/Niagara/Source/NiagaraEditor/Private/NiagaraNodeFunctionCall.cpp
--- a/Engine/Plugins/FX/Niagara/Source/NiagaraEditor/Private/NiagaraNodeFunctionCall.cpp
+++ b/Engine/Plugins/FX/Niagara/Source/NiagaraEditor/Private/NiagaraNodeFunctionCall.cpp
@@ -16,6 +16,7 @@
 #include "NiagaraNodeParameterMapGet.h"
 #include "NiagaraConstants.h"
 #include "NiagaraCustomVersion.h"
+#include "NiagaraDataInterfaceSkeletalMesh.h"
 
 #define LOCTEXT_NAMESPACE "NiagaraNodeFunctionCall"
 
@@ -78,6 +79,52 @@ void UNiagaraNodeFunctionCall::PostLoad()
 	}
 }
 
+void UNiagaraNodeFunctionCall::UpgradeDIFunctionCalls()
+{
+	UClass* InterfaceClass = nullptr;
+	UNiagaraDataInterface* InterfaceCDO = nullptr;
+	if (Signature.IsValid() && FunctionScript == nullptr)
+	{
+		if (Signature.Inputs.Num() > 0)
+		{
+			if (Signature.Inputs[0].GetType().IsDataInterface())
+			{
+				InterfaceClass = Signature.Inputs[0].GetType().GetClass();
+				InterfaceCDO = Cast<UNiagaraDataInterface>(InterfaceClass->GetDefaultObject());
+			}
+		}
+	}
+
+	FString UpgradeNote;
+
+	//TODO: Move this out into DI specific functions or helper classes?
+	if (InterfaceClass && InterfaceCDO)
+	{		
+		if (InterfaceClass == UNiagaraDataInterfaceSkeletalMesh::StaticClass())
+		{
+			if (Signature.Name == TEXT("RandomTriCoord"))
+			{
+				if (Signature.Inputs.Num() == 1)//If this is before we added the additional seed inputs. TODO: Add a per DI version number to make this simpler and clearer. In this case we can detect old data easy enough but a version number is better.
+				{
+					UpgradeNote = TEXT("Adding RandomInfo parameter to support deterministic random sampling.");
+
+					Signature.Inputs.Add(FNiagaraVariable(FNiagaraTypeDefinition(FNiagaraRandInfo::StaticStruct()), TEXT("RandomInfo")));
+					ReallocatePins();
+				}
+			}
+		}
+	}
+
+	if (!UpgradeNote.IsEmpty())
+	{
+		UE_LOG(LogNiagaraEditor, Log, TEXT("Upgradeing Niagara Data Interface fuction call node. This may cause unnessessary recompiles. Please resave these assets if this occurs. Or use fx.UpgradeAllNiagaraAssets."));
+		UE_LOG(LogNiagaraEditor, Log, TEXT("Node: %s"), *GetFullName());
+		UE_LOG(LogNiagaraEditor, Log, TEXT("Interface: %s"), *InterfaceCDO->GetFullName());
+		UE_LOG(LogNiagaraEditor, Log, TEXT("Function: %s"), *Signature.GetName());
+		UE_LOG(LogNiagaraEditor, Log, TEXT("Upgrade Note: %s"),* UpgradeNote);
+	}
+}
+
 void UNiagaraNodeFunctionCall::PostEditChangeProperty(struct FPropertyChangedEvent& PropertyChangedEvent)
 {
 	if (PropertyChangedEvent.Property != nullptr)
diff --git a/Engine/Plugins/FX/Niagara/Source/NiagaraEditor/Private/ViewModels/Stack/NiagaraStackFunctionInput.cpp b/Engine/Plugins/FX/Niagara/Source/NiagaraEditor/Private/ViewModels/Stack/NiagaraStackFunctionInput.cpp
--- a/Engine/Plugins/FX/Niagara/Source/NiagaraEditor/Private/ViewModels/Stack/NiagaraStackFunctionInput.cpp
+++ b/Engine/Plugins/FX/Niagara/Source/NiagaraEditor/Private/ViewModels/Stack/NiagaraStackFunctionInput.cpp
@@ -1743,7 +1743,7 @@ UEdGraphPin& UNiagaraStackFunctionInput::GetOrCreateOverridePin()
 
 bool UNiagaraStackFunctionInput::TryGetCurrentLocalValue(TSharedPtr<FStructOnScope>& LocalValue, UEdGraphPin& DefaultPin, UEdGraphPin& ValuePin, TSharedPtr<FStructOnScope> OldValueToReuse)
 {
-	if (InputType.IsDataInterface() == false && ValuePin.LinkedTo.Num() == 0)
+	if (InputType.IsUObject() == false && ValuePin.LinkedTo.Num() == 0)
 	{
 		const UEdGraphSchema_Niagara* NiagaraSchema = GetDefault<UEdGraphSchema_Niagara>();
 		FNiagaraVariable ValueVariable = NiagaraSchema->PinToNiagaraVariable(&ValuePin, true);
@@ -1781,7 +1781,7 @@ bool UNiagaraStackFunctionInput::TryGetCurrentLocalValue(TSharedPtr<FStructOnSco
 
 bool UNiagaraStackFunctionInput::TryGetCurrentDataValue(FDataValues& DataValues, UEdGraphPin* OverrideValuePin, UEdGraphPin& DefaultValuePin, UNiagaraDataInterface* LocallyOwnedDefaultDataValueObjectToReuse)
 {
-	if (InputType.GetClass() != nullptr)
+	if (InputType.IsDataInterface())
 	{
 		UNiagaraDataInterface* DataValueObject = nullptr;
 		if (OverrideValuePin != nullptr)
diff --git a/Engine/Plugins/FX/Niagara/Source/NiagaraEditor/Private/ViewModels/Stack/NiagaraStackGraphUtilities.cpp b/Engine/Plugins/FX/Niagara/Source/NiagaraEditor/Private/ViewModels/Stack/NiagaraStackGraphUtilities.cpp
--- a/Engine/Plugins/FX/Niagara/Source/NiagaraEditor/Private/ViewModels/Stack/NiagaraStackGraphUtilities.cpp
+++ b/Engine/Plugins/FX/Niagara/Source/NiagaraEditor/Private/ViewModels/Stack/NiagaraStackGraphUtilities.cpp
@@ -1238,7 +1238,7 @@ bool FNiagaraStackGraphUtilities::IsRapidIterationType(const FNiagaraTypeDefinit
 {
 	checkf(InputType.IsValid(), TEXT("Type is invalid."));
 	return InputType != FNiagaraTypeDefinition::GetBoolDef() && !InputType.IsEnum() &&
-		InputType != FNiagaraTypeDefinition::GetParameterMapDef() && !InputType.IsDataInterface();
+		InputType != FNiagaraTypeDefinition::GetParameterMapDef() && !InputType.IsUObject();
 }
 
 FNiagaraVariable FNiagaraStackGraphUtilities::CreateRapidIterationParameter(const FString& UniqueEmitterName, ENiagaraScriptUsage ScriptUsage, const FName& AliasedInputName, const FNiagaraTypeDefinition& InputType)
@@ -1293,10 +1293,16 @@ void FNiagaraStackGraphUtilities::CleanUpStaleRapidIterationParameters(UNiagaraE
 	}
 }
 
-void FNiagaraStackGraphUtilities::GetNewParameterAvailableTypes(TArray<FNiagaraTypeDefinition>& OutAvailableTypes)
+void FNiagaraStackGraphUtilities::GetNewParameterAvailableTypes(TArray<FNiagaraTypeDefinition>& OutAvailableTypes, FName Namespace)
 {
 	for (const FNiagaraTypeDefinition& RegisteredParameterType : FNiagaraTypeRegistry::GetRegisteredParameterTypes())
 	{
+		//Object types only allowed in user namespace at the moment.
+		if (RegisteredParameterType == FNiagaraTypeDefinition::GetUObjectDef() && Namespace != FNiagaraParameterHandle::UserNamespace)
+		{
+			continue;
+		}
+
 		if (RegisteredParameterType != FNiagaraTypeDefinition::GetGenericNumericDef() && RegisteredParameterType != FNiagaraTypeDefinition::GetParameterMapDef())
 		{
 			OutAvailableTypes.Add(RegisteredParameterType);
diff --git a/Engine/Plugins/FX/Niagara/Source/NiagaraEditor/Private/ViewModels/Stack/NiagaraStackParameterStoreGroup.cpp b/Engine/Plugins/FX/Niagara/Source/NiagaraEditor/Private/ViewModels/Stack/NiagaraStackParameterStoreGroup.cpp
--- a/Engine/Plugins/FX/Niagara/Source/NiagaraEditor/Private/ViewModels/Stack/NiagaraStackParameterStoreGroup.cpp
+++ b/Engine/Plugins/FX/Niagara/Source/NiagaraEditor/Private/ViewModels/Stack/NiagaraStackParameterStoreGroup.cpp
@@ -72,7 +72,7 @@ public:
 	virtual void GenerateAddActions(TArray<TSharedRef<INiagaraStackItemGroupAddAction>>& OutAddActions, const FNiagaraStackItemGroupAddOptions& AddProperties) const override
 	{
 		TArray<FNiagaraTypeDefinition> AvailableTypes;
-		FNiagaraStackGraphUtilities::GetNewParameterAvailableTypes(AvailableTypes);
+		FNiagaraStackGraphUtilities::GetNewParameterAvailableTypes(AvailableTypes, FNiagaraParameterHandle::UserNamespace);
 		for (const FNiagaraTypeDefinition& AvailableType : AvailableTypes)
 		{
 			FNiagaraParameterHandle NewParameterHandle(FNiagaraParameterHandle::UserNamespace, *(TEXT("New") + AvailableType.GetName()));
diff --git a/Engine/Plugins/FX/Niagara/Source/NiagaraEditor/Private/ViewModels/Stack/NiagaraStackScriptItemGroup.cpp b/Engine/Plugins/FX/Niagara/Source/NiagaraEditor/Private/ViewModels/Stack/NiagaraStackScriptItemGroup.cpp
--- a/Engine/Plugins/FX/Niagara/Source/NiagaraEditor/Private/ViewModels/Stack/NiagaraStackScriptItemGroup.cpp
+++ b/Engine/Plugins/FX/Niagara/Source/NiagaraEditor/Private/ViewModels/Stack/NiagaraStackScriptItemGroup.cpp
@@ -216,7 +216,7 @@ public:
 		if (NewParameterNamespace.IsSet())
 		{
 			TArray<FNiagaraTypeDefinition> AvailableTypes;
-			FNiagaraStackGraphUtilities::GetNewParameterAvailableTypes(AvailableTypes);
+			FNiagaraStackGraphUtilities::GetNewParameterAvailableTypes(AvailableTypes, NewParameterNamespace.Get(TEXT("Unknown")));
 			for (const FNiagaraTypeDefinition& AvailableType : AvailableTypes)
 			{
 				OutAddActions.Add(FScriptGroupAddAction::CreateNewParameterModuleAction(NewParameterNamespace.GetValue(), AvailableType));
diff --git a/Engine/Plugins/FX/Niagara/Source/NiagaraEditor/Public/NiagaraEditorModule.h b/Engine/Plugins/FX/Niagara/Source/NiagaraEditor/Public/NiagaraEditorModule.h
--- a/Engine/Plugins/FX/Niagara/Source/NiagaraEditor/Public/NiagaraEditorModule.h
+++ b/Engine/Plugins/FX/Niagara/Source/NiagaraEditor/Public/NiagaraEditorModule.h
@@ -143,6 +143,7 @@ private:
 	IConsoleCommand* DumpRapidIterationParametersForAsset;
 	IConsoleCommand* PreventSystemRecompileCommand;
 	IConsoleCommand* PreventAllSystemRecompilesCommand;
+	IConsoleCommand* UpgradeAllNiagaraAssetsCommand;
 
 	FOnCheckScriptToolkitsShouldFocusGraphElement OnCheckScriptToolkitsShouldFocusGraphElement;
 };
diff --git a/Engine/Plugins/FX/Niagara/Source/NiagaraEditor/Public/NiagaraNodeFunctionCall.h b/Engine/Plugins/FX/Niagara/Source/NiagaraEditor/Public/NiagaraNodeFunctionCall.h
--- a/Engine/Plugins/FX/Niagara/Source/NiagaraEditor/Public/NiagaraNodeFunctionCall.h
+++ b/Engine/Plugins/FX/Niagara/Source/NiagaraEditor/Public/NiagaraNodeFunctionCall.h
@@ -126,7 +126,10 @@ public:
 	void RemovePropagatedVariable(const FNiagaraVariable& Variable);
 	void CleanupPropagatedSwitchValues();
 
+	/** Does any automated data manipulated required to update DI function call nodes to the current version. */
+	void UpgradeDIFunctionCalls();
 protected:
+
 	virtual bool IsValidPinToCompile(UEdGraphPin* Pin) const;
 
 	virtual bool AllowDynamicPins() const override { return false; }
diff --git a/Engine/Plugins/FX/Niagara/Source/NiagaraEditor/Public/ViewModels/Stack/NiagaraStackGraphUtilities.h b/Engine/Plugins/FX/Niagara/Source/NiagaraEditor/Public/ViewModels/Stack/NiagaraStackGraphUtilities.h
--- a/Engine/Plugins/FX/Niagara/Source/NiagaraEditor/Public/ViewModels/Stack/NiagaraStackGraphUtilities.h
+++ b/Engine/Plugins/FX/Niagara/Source/NiagaraEditor/Public/ViewModels/Stack/NiagaraStackGraphUtilities.h
@@ -132,7 +132,7 @@ namespace FNiagaraStackGraphUtilities
 
 	void CleanUpStaleRapidIterationParameters(UNiagaraEmitter& Emitter);
 
-	void GetNewParameterAvailableTypes(TArray<FNiagaraTypeDefinition>& OutAvailableTypes);
+	void GetNewParameterAvailableTypes(TArray<FNiagaraTypeDefinition>& OutAvailableTypes, FName Namespace);
 
 	void GetScriptAssetsByDependencyProvided(ENiagaraScriptUsage AssetUsage, FName DependencyName, TArray<FAssetData>& OutAssets);
 
diff --git a/Engine/Plugins/FX/Niagara/Source/NiagaraEditorWidgets/Private/DetailCustomizations/NiagaraDataInterfaceSkeletalMeshDetails.cpp b/Engine/Plugins/FX/Niagara/Source/NiagaraEditorWidgets/Private/DetailCustomizations/NiagaraDataInterfaceSkeletalMeshDetails.cpp
--- a/Engine/Plugins/FX/Niagara/Source/NiagaraEditorWidgets/Private/DetailCustomizations/NiagaraDataInterfaceSkeletalMeshDetails.cpp
+++ b/Engine/Plugins/FX/Niagara/Source/NiagaraEditorWidgets/Private/DetailCustomizations/NiagaraDataInterfaceSkeletalMeshDetails.cpp
@@ -31,7 +31,8 @@ void FNiagaraDataInterfaceSkeletalMeshDetails::CustomizeDetails(IDetailLayoutBui
 
 	 TWeakObjectPtr<USceneComponent> SceneComponent;
 	 USkeletalMeshComponent* FoundSkelComp = nullptr;
-	 MeshObject = UNiagaraDataInterfaceSkeletalMesh::GetSkeletalMeshHelper(MeshInterface.Get(), Cast<UNiagaraComponent>(MeshInterface->GetOuter()), SceneComponent, FoundSkelComp);
+	 UNiagaraDataInterfaceSkeletalMesh* Interface = MeshInterface.Get();
+	 MeshObject = Interface->GetSkeletalMesh(Cast<UNiagaraComponent>(MeshInterface->GetOuter()), SceneComponent, FoundSkelComp);
 	 if (MeshObject.IsValid())
 	 {
 		 MeshObject->GetOnMeshChanged().RemoveAll(this);
@@ -109,7 +110,8 @@ void FNiagaraDataInterfaceSkeletalMeshDetails::CustomizeDetails(IDetailLayoutBui
 	 {
 		 MeshObject->GetOnMeshChanged().RemoveAll(this);
 	 }
-	 MeshObject = UNiagaraDataInterfaceSkeletalMesh::GetSkeletalMeshHelper(MeshInterface.Get(), Cast<UNiagaraComponent>(MeshInterface->GetOuter()), SceneComponent, FoundSkelComp);
+	 UNiagaraDataInterfaceSkeletalMesh* Interface = MeshInterface.Get();
+	 MeshObject = Interface->GetSkeletalMesh(Cast<UNiagaraComponent>(MeshInterface->GetOuter()), SceneComponent, FoundSkelComp);
 	 if (MeshObject.IsValid())
 	 {
 		 MeshObject->GetOnMeshChanged().AddSP(this, &FNiagaraDataInterfaceSkeletalMeshDetails::OnDataChanged);
@@ -137,7 +139,8 @@ void FNiagaraDataInterfaceSkeletalMeshDetails::GenerateRegionsArray(TArray<TShar
 	 {
 		 TWeakObjectPtr<USceneComponent> SceneComponent;
 		 USkeletalMeshComponent* FoundSkelComp = nullptr;
-		 USkeletalMesh* Mesh = UNiagaraDataInterfaceSkeletalMesh::GetSkeletalMeshHelper(MeshInterface.Get(), Cast<UNiagaraComponent>(MeshInterface->GetOuter()), SceneComponent, FoundSkelComp);
+		 UNiagaraDataInterfaceSkeletalMesh* Interface = MeshInterface.Get();
+		 USkeletalMesh* Mesh = Interface->GetSkeletalMesh(Cast<UNiagaraComponent>(MeshInterface->GetOuter()), SceneComponent, FoundSkelComp);
 
 		 if (Mesh != nullptr)
 		 {
@@ -156,7 +159,8 @@ void FNiagaraDataInterfaceSkeletalMeshDetails::GenerateRegionsArray(TArray<TShar
 	 {
 		 TWeakObjectPtr<USceneComponent> SceneComponent;
 		 USkeletalMeshComponent* FoundSkelComp = nullptr;
-		 USkeletalMesh* Mesh = UNiagaraDataInterfaceSkeletalMesh::GetSkeletalMeshHelper(MeshInterface.Get(), Cast<UNiagaraComponent>(MeshInterface->GetOuter()), SceneComponent, FoundSkelComp);
+		 UNiagaraDataInterfaceSkeletalMesh* Interface = MeshInterface.Get();
+		 USkeletalMesh* Mesh = Interface->GetSkeletalMesh(Cast<UNiagaraComponent>(MeshInterface->GetOuter()), SceneComponent, FoundSkelComp);
 
 		 if (Mesh != nullptr)
 		 {
@@ -175,7 +179,8 @@ void FNiagaraDataInterfaceSkeletalMeshDetails::GenerateSocketsArray(TArray<TShar
 	 {
 		 TWeakObjectPtr<USceneComponent> SceneComponent;
 		 USkeletalMeshComponent* FoundSkelComp = nullptr;
-		 USkeletalMesh* Mesh = UNiagaraDataInterfaceSkeletalMesh::GetSkeletalMeshHelper(MeshInterface.Get(), Cast<UNiagaraComponent>(MeshInterface->GetOuter()), SceneComponent, FoundSkelComp);
+		 UNiagaraDataInterfaceSkeletalMesh* Interface = MeshInterface.Get();
+		 USkeletalMesh* Mesh = Interface->GetSkeletalMesh(Cast<UNiagaraComponent>(MeshInterface->GetOuter()), SceneComponent, FoundSkelComp);
 
 		 if (Mesh != nullptr)
 		 {
diff --git a/Engine/Source/Runtime/Core/Public/Math/IntVector.h b/Engine/Source/Runtime/Core/Public/Math/IntVector.h
--- a/Engine/Source/Runtime/Core/Public/Math/IntVector.h
+++ b/Engine/Source/Runtime/Core/Public/Math/IntVector.h
@@ -189,6 +189,46 @@ public:
 	 */
 	FIntVector operator-( const FIntVector& Other ) const;
 
+	/**
+	 * Shifts all components to the right.
+	 *
+	 * @param Shift The number of bits to shift.
+	 * @return A new shifted int point.
+	 */
+	FIntVector operator>>(int32 Shift) const;
+
+	/**
+	 * Shifts all components to the left.
+	 *
+	 * @param Shift The number of bits to shift.
+	 * @return A new shifted int point.
+	 */
+	FIntVector operator<<(int32 Shift) const;
+
+	/**
+	 * Component-wise AND.
+	 *
+	 * @param Value Number to AND with the each component.
+	 * @return A new shifted int point.
+	 */
+	FIntVector operator&(int32 Val) const;
+
+	/**
+	 * Component-wise OR.
+	 *
+	 * @param Value Number to OR with the each component.
+	 * @return A new shifted int point.
+	 */
+	FIntVector operator|(int32 Value) const;
+
+	/**
+	 * Component-wise XOR.
+	 *
+	 * @param Value Number to XOR with the each component.
+	 * @return A new shifted int point.
+	 */
+	FIntVector operator^(int32 Value) const;
+
 	/**
 	 * Is vector equal to zero.
 	 * @return is zero
@@ -408,13 +448,36 @@ FORCEINLINE FIntVector FIntVector::operator-( const FIntVector& Other ) const
 	return FIntVector(*this) -= Other;
 }
 
+FORCEINLINE FIntVector FIntVector::operator>>(int32 Shift) const
+{
+	return FIntVector(X >> Shift, Y >> Shift, Z >> Shift);
+}
+
+FORCEINLINE FIntVector FIntVector::operator<<(int32 Shift) const
+{
+	return FIntVector(X << Shift, Y << Shift, Z << Shift);
+}
+
+FORCEINLINE FIntVector FIntVector::operator&(int32 Value) const
+{
+	return FIntVector(X & Value, Y & Value, Z & Value);
+}
+
+FORCEINLINE FIntVector FIntVector::operator|(int32 Value) const
+{
+	return FIntVector(X | Value, Y | Value, Z | Value);
+}
+
+FORCEINLINE FIntVector FIntVector::operator^(int32 Value) const
+{
+	return FIntVector(X ^ Value, Y ^ Value, Z ^ Value);
+}
 
 FORCEINLINE FIntVector FIntVector::DivideAndRoundUp( FIntVector lhs, int32 Divisor )
 {
 	return FIntVector(FMath::DivideAndRoundUp(lhs.X, Divisor), FMath::DivideAndRoundUp(lhs.Y, Divisor), FMath::DivideAndRoundUp(lhs.Z, Divisor));
 }
 
-
 FORCEINLINE int32 FIntVector::GetMax() const
 {
 	return FMath::Max(FMath::Max(X, Y), Z);
@@ -489,29 +552,305 @@ struct FIntVector4
 	{
 	}
 
-	FORCEINLINE const int32& operator[](int32 ComponentIndex) const
-	{
-		return (&X)[ComponentIndex];
-	}
+	/**
+	 * Gets specific component of a point.
+	 *
+	 * @param ComponentIndex Index of point component.
+	 * @return const reference to component.
+	 */
+	const int32& operator()(int32 ComponentIndex) const;
 
+	/**
+	 * Gets specific component of a point.
+	 *
+	 * @param ComponentIndex Index of point component.
+	 * @return reference to component.
+	 */
+	int32& operator()(int32 ComponentIndex);
 
-	FORCEINLINE int32& operator[](int32 ComponentIndex)
-	{
-		return (&X)[ComponentIndex];
-	}
+	/**
+	 * Gets specific component of a point.
+	 *
+	 * @param ComponentIndex Index of point component.
+	 * @return const reference to component.
+	 */
+	const int32& operator[](int32 ComponentIndex) const;
 
-	FORCEINLINE bool operator==(const FIntVector4& Other) const
-	{
-		return X==Other.X && Y==Other.Y && Z==Other.Z && W==Other.W;
-	}
+	/**
+	 * Gets specific component of a point.
+	 *
+	 * @param ComponentIndex Index of point component.
+	 * @return reference to component.
+	 */
+	int32& operator[](int32 ComponentIndex);
 
+	/**
+	 * Compares points for equality.
+	 *
+	 * @param Other The other int point being compared.
+	 * @return true if the points are equal, false otherwise..
+	 */
+	bool operator==(const FIntVector4& Other) const;
 
-	FORCEINLINE bool operator!=(const FIntVector4& Other) const
-	{
-		return X!=Other.X || Y!=Other.Y || Z!=Other.Z || W!=Other.W;
-	}
+	/**
+	 * Compares points for inequality.
+	 *
+	 * @param Other The other int point being compared.
+	 * @return true if the points are not equal, false otherwise..
+	 */
+	bool operator!=(const FIntVector4& Other) const;
+
+	/**
+	 * Scales this point.
+	 *
+	 * @param Scale What to multiply the point by.
+	 * @return Reference to this point after multiplication.
+	 */
+	FIntVector4& operator*=(int32 Scale);
+
+	/**
+	 * Divides this point.
+	 *
+	 * @param Divisor What to divide the point by.
+	 * @return Reference to this point after division.
+	 */
+	FIntVector4& operator/=(int32 Divisor);
+
+	/**
+	 * Adds to this point.
+	 *
+	 * @param Other The point to add to this point.
+	 * @return Reference to this point after addition.
+	 */
+	FIntVector4& operator+=(const FIntVector4& Other);
+
+	/**
+	 * Subtracts from this point.
+	 *
+	 * @param Other The point to subtract from this point.
+	 * @return Reference to this point after subtraction.
+	 */
+	FIntVector4& operator-=(const FIntVector4& Other);
+
+	/**
+	 * Assigns another point to this one.
+	 *
+	 * @param Other The point to assign this point from.
+	 * @return Reference to this point after assignment.
+	 */
+	FIntVector4& operator=(const FIntVector4& Other);
+
+	/**
+	 * Gets the result of scaling on this point.
+	 *
+	 * @param Scale What to multiply the point by.
+	 * @return A new scaled int point.
+	 */
+	FIntVector4 operator*(int32 Scale) const;
+
+	/**
+	 * Gets the result of division on this point.
+	 *
+	 * @param Divisor What to divide the point by.
+	 * @return A new divided int point.
+	 */
+	FIntVector4 operator/(int32 Divisor) const;
+
+	/**
+	 * Gets the result of addition on this point.
+	 *
+	 * @param Other The other point to add to this.
+	 * @return A new combined int point.
+	 */
+	FIntVector4 operator+(const FIntVector4& Other) const;
+
+	/**
+	 * Gets the result of subtraction from this point.
+	 *
+	 * @param Other The other point to subtract from this.
+	 * @return A new subtracted int point.
+	 */
+	FIntVector4 operator-(const FIntVector4& Other) const;
+
+	/**
+	 * Shifts all components to the right.
+	 *
+	 * @param Shift The number of bits to shift.
+	 * @return A new shifted int point.
+	 */
+	FIntVector4 operator>>(int32 Shift) const;
+
+	/**
+	 * Shifts all components to the left.
+	 *
+	 * @param Shift The number of bits to shift.
+	 * @return A new shifted int point.
+	 */
+	FIntVector4 operator<<(int32 Shift) const;
+
+	/**
+	 * Component-wise AND.
+	 *
+	 * @param Value Number to AND with the each component.
+	 * @return A new shifted int point.
+	 */
+	FIntVector4 operator&(int32 Val) const;
+
+	/**
+	 * Component-wise OR.
+	 *
+	 * @param Value Number to OR with the each component.
+	 * @return A new shifted int point.
+	 */
+	FIntVector4 operator|(int32 Value) const;
+
+	/**
+	 * Component-wise XOR.
+	 *
+	 * @param Value Number to XOR with the each component.
+	 * @return A new shifted int point.
+	 */
+	FIntVector4 operator^(int32 Value) const;
 };
 
+FORCEINLINE const int32& FIntVector4::operator()(int32 ComponentIndex) const
+{
+	return (&X)[ComponentIndex];
+}
+
+
+FORCEINLINE int32& FIntVector4::operator()(int32 ComponentIndex)
+{
+	return (&X)[ComponentIndex];
+}
+
+
+FORCEINLINE const int32& FIntVector4::operator[](int32 ComponentIndex) const
+{
+	return (&X)[ComponentIndex];
+}
+
+
+FORCEINLINE int32& FIntVector4::operator[](int32 ComponentIndex)
+{
+	return (&X)[ComponentIndex];
+}
+
+FORCEINLINE bool FIntVector4::operator==(const FIntVector4& Other) const
+{
+	return X == Other.X && Y == Other.Y && Z == Other.Z;
+}
+
+
+FORCEINLINE bool FIntVector4::operator!=(const FIntVector4& Other) const
+{
+	return X != Other.X || Y != Other.Y || Z != Other.Z || W != Other.W;
+}
+
+
+FORCEINLINE FIntVector4& FIntVector4::operator*=(int32 Scale)
+{
+	X *= Scale;
+	Y *= Scale;
+	Z *= Scale;
+	W *= Scale;
+
+	return *this;
+}
+
+
+FORCEINLINE FIntVector4& FIntVector4::operator/=(int32 Divisor)
+{
+	X /= Divisor;
+	Y /= Divisor;
+	Z /= Divisor;
+	W /= Divisor;
+
+	return *this;
+}
+
+
+FORCEINLINE FIntVector4& FIntVector4::operator+=(const FIntVector4& Other)
+{
+	X += Other.X;
+	Y += Other.Y;
+	Z += Other.Z;
+	W += Other.W;
+
+	return *this;
+}
+
+
+FORCEINLINE FIntVector4& FIntVector4::operator-=(const FIntVector4& Other)
+{
+	X -= Other.X;
+	Y -= Other.Y;
+	Z -= Other.Z;
+	W -= Other.W;
+
+	return *this;
+}
+
+
+FORCEINLINE FIntVector4& FIntVector4::operator=(const FIntVector4& Other)
+{
+	X = Other.X;
+	Y = Other.Y;
+	Z = Other.Z;
+	W = Other.W;
+
+	return *this;
+}
+
+
+FORCEINLINE FIntVector4 FIntVector4::operator*(int32 Scale) const
+{
+	return FIntVector4(*this) *= Scale;
+}
+
+
+FORCEINLINE FIntVector4 FIntVector4::operator/(int32 Divisor) const
+{
+	return FIntVector4(*this) /= Divisor;
+}
+
+
+FORCEINLINE FIntVector4 FIntVector4::operator+(const FIntVector4& Other) const
+{
+	return FIntVector4(*this) += Other;
+}
+
+FORCEINLINE FIntVector4 FIntVector4::operator-(const FIntVector4& Other) const
+{
+	return FIntVector4(*this) -= Other;
+}
+
+FORCEINLINE FIntVector4 FIntVector4::operator>>(int32 Shift) const
+{
+	return FIntVector4(X >> Shift, Y >> Shift, Z >> Shift, W >> Shift);
+}
+
+FORCEINLINE FIntVector4 FIntVector4::operator<<(int32 Shift) const
+{
+	return FIntVector4(X << Shift, Y << Shift, Z << Shift, W << Shift);
+}
+
+FORCEINLINE FIntVector4 FIntVector4::operator&(int32 Value) const
+{
+	return FIntVector4(X & Value, Y & Value, Z & Value, W & Value);
+}
+
+FORCEINLINE FIntVector4 FIntVector4::operator|(int32 Value) const
+{
+	return FIntVector4(X | Value, Y | Value, Z | Value, W | Value);
+}
+
+FORCEINLINE FIntVector4 FIntVector4::operator^(int32 Value) const
+{
+	return FIntVector4(X ^ Value, Y ^ Value, Z ^ Value, W ^ Value);
+}
+
+
 struct FUintVector4
 {
 	uint32 X, Y, Z, W;
diff --git a/Engine/Source/Runtime/Core/Public/Math/Vector4.h b/Engine/Source/Runtime/Core/Public/Math/Vector4.h
--- a/Engine/Source/Runtime/Core/Public/Math/Vector4.h
+++ b/Engine/Source/Runtime/Core/Public/Math/Vector4.h
@@ -65,6 +65,13 @@ public:
 	 */
 	explicit FVector4(FVector2D InXY, FVector2D InZW);
 
+	/**
+	 * Creates and initializes a new vector from an int vector value.
+	 *
+	 * @param InVector IntVector used to set vector.
+	 */
+	FVector4(const FIntVector4& InVector);
+
 	/**
 	 * Creates and initializes a new vector to zero.
 	 *
@@ -493,6 +500,13 @@ FORCEINLINE FVector4::FVector4(FVector2D InXY, FVector2D InZW)
 	DiagnosticCheckNaN();
 }
 
+FORCEINLINE FVector4::FVector4(const FIntVector4& InVector)
+	: X(InVector.X)
+	, Y(InVector.Y)
+	, Z(InVector.Z)
+	, W(InVector.W)
+{
+}
 
 FORCEINLINE float& FVector4::operator[](int32 ComponentIndex)
 {
diff --git a/Engine/Source/Runtime/VectorVM/Public/VectorVM.h b/Engine/Source/Runtime/VectorVM/Public/VectorVM.h
--- a/Engine/Source/Runtime/VectorVM/Public/VectorVM.h
+++ b/Engine/Source/Runtime/VectorVM/Public/VectorVM.h
@@ -266,8 +266,12 @@ struct FVectorVMContext : TThreadSingleton<FVectorVMContext>
 	TArray<uint8, TAlignedHeapAllocator<VECTOR_WIDTH_BYTES>> TempRegTable;
 	uint8 *RESTRICT RegisterTable[VectorVM::MaxRegisters];
 
+	/** Thread local random stream for use in external functions needing non-deterministic randoms. */
 	FRandomStream RandStream;
 
+	/** Thread local per instance random counters for use in external functions needing deterministic randoms. */
+	TArray<int32> RandCounters;
+
 	FVectorVMContext();
 
 	void PrepareForExec(
@@ -294,6 +298,9 @@ struct FVectorVMContext : TThreadSingleton<FVectorVMContext>
 		Code = InCode;
 		NumInstances = InNumInstances;
 		StartInstance = InStartInstance;
+
+		RandCounters.Reset();
+		RandCounters.SetNumZeroed(InNumInstances);
 	}
 };
 
