Compile fix: Missing file

diff --git a/Engine/Plugins/FX/Niagara/Source/Niagara/Private/NiagaraCommon.cpp b/Engine/Plugins/FX/Niagara/Source/Niagara/Private/NiagaraCommon.cpp
--- a/Engine/Plugins/FX/Niagara/Source/Niagara/Private/NiagaraCommon.cpp
+++ b/Engine/Plugins/FX/Niagara/Source/Niagara/Private/NiagaraCommon.cpp
@@ -367,4 +367,12 @@ void FNiagaraUtilities::PrepareRapidIterationParameters(const TArray<UNiagaraScr
 		}
 	}
 }
-#endif
\ No newline at end of file
+#endif
+
+//////////////////////////////////////////////////////////////////////////
+
+FNiagaraUserParameterBinding::FNiagaraUserParameterBinding()
+	: Parameter(FNiagaraTypeDefinition::GetUObjectDef(), NAME_None)
+	{
+
+	}
\ No newline at end of file
