- Change static_assert to = delete to comply with GCC

#rnx

diff --git a/Engine/Plugins/FX/Niagara/Source/Niagara/Classes/NiagaraDataInterface.h b/Engine/Plugins/FX/Niagara/Source/Niagara/Classes/NiagaraDataInterface.h
--- a/Engine/Plugins/FX/Niagara/Source/Niagara/Classes/NiagaraDataInterface.h
+++ b/Engine/Plugins/FX/Niagara/Source/Niagara/Classes/NiagaraDataInterface.h
@@ -438,18 +438,9 @@ struct TCurveUseLUTBinder
 template<typename T>
 struct FNDIParameter
 {
-	FNDIParameter(FVectorVMContext& Context)
-	{
-		static_assert(false, "You must specialize FNDIParameter for all types it's used with.");
-	}
-	FORCEINLINE void GetAndAdvance(T& OutValue)
-	{
-		static_assert(false, "You must specialize FNDIParameter for all types it's used with.");
-	}
-	FORCEINLINE bool IsConstant()const
-	{
-		static_assert(false, "You must specialize FNDIParameter for all types it's used with.");
-	}
+	FNDIParameter(FVectorVMContext& Context) = delete;
+	FORCEINLINE void GetAndAdvance(T& OutValue) = delete;
+	FORCEINLINE bool IsConstant() const = delete;
 };
 
 template<>
