- Compile time check for unspecialized FSkinnedPositionAccessorHelper being used

#rnx

diff --git a/Engine/Plugins/FX/Niagara/Source/Niagara/Classes/NDISkeletalMeshCommon.h b/Engine/Plugins/FX/Niagara/Source/Niagara/Classes/NDISkeletalMeshCommon.h
--- a/Engine/Plugins/FX/Niagara/Source/Niagara/Classes/NDISkeletalMeshCommon.h
+++ b/Engine/Plugins/FX/Niagara/Source/Niagara/Classes/NDISkeletalMeshCommon.h
@@ -55,58 +55,21 @@ void FSkeletalMeshAccessorHelper::Init<
 template<typename SkinningMode>
 struct FSkinnedPositionAccessorHelper
 {
-	FORCEINLINE void GetTrianlgeIndices(int32 Tri, int32& Idx0, int32& Idx1, int32& Idx2)
-	{
-		checkf(false, TEXT("Must provide a specialization for this template type"));
-	}
-
-	FORCEINLINE void GetSkinnedTrianglePositions(FSkeletalMeshAccessorHelper& Accessor,
-		int32 Tri, FVector& OutPos0, FVector& OutPos1, FVector& OutPos2, FVector& OutPrev0, FVector& OutPrev1, FVector& OutPrev2, int32& Idx0, int32& Idx1, int32& Idx2)
-	{
-		checkf(false, TEXT("Must provide a specialization for this template type"));
-	}
-
-	FORCEINLINE void GetSkinnedTrianglePositions(FSkeletalMeshAccessorHelper& Accessor,
-		int32 Tri, FVector& OutPos0, FVector& OutPos1, FVector& OutPos2, int32& Idx0, int32& Idx1, int32& Idx2)
-	{
-		checkf(false, TEXT("Must provide a specialization for this template type"));
-	}
-
-	FORCEINLINE FVector GetSkinnedVertexPosition(FSkeletalMeshAccessorHelper& Accessor, int32 VertexIndex)
-	{
-		checkf(false, TEXT("Must provide a specialization for this template type"));
-	}
-
-	FORCEINLINE FVector GetSkinnedVertexPreviousPosition(FSkeletalMeshAccessorHelper& Accessor, int32 VertexIndex)
-	{
-		checkf(false, TEXT("Must provide a specialization for this template type"));
-	}
-
-	FORCEINLINE FVector GetSkinnedBonePosition(FSkeletalMeshAccessorHelper& Accessor, int32 BoneIndex)
-	{
-		checkf(false, TEXT("Must provide a specialization for this template type"));
-	}
-
-	FORCEINLINE FVector GetSkinnedBonePreviousPosition(FSkeletalMeshAccessorHelper& Accessor, int32 BoneIndex)
-	{
-		checkf(false, TEXT("Must provide a specialization for this template type"));
-	}	
-	
-	FORCEINLINE_DEBUGGABLE FQuat GetSkinnedBoneRotation(FSkeletalMeshAccessorHelper& Accessor, int32 BoneIndex)
-	{
-		checkf(false, TEXT("Must provide a specialization for this template type"));
-	}
-
-	FORCEINLINE_DEBUGGABLE FQuat GetSkinnedBonePreviousRotation(FSkeletalMeshAccessorHelper& Accessor, int32 BoneIndex)
-	{
-		checkf(false, TEXT("Must provide a specialization for this template type"));
-	}
+	FORCEINLINE void GetTriangleIndices(int32 Tri, int32& Idx0, int32& Idx1, int32& Idx2) = delete;
+	FORCEINLINE void GetSkinnedTrianglePositions(FSkeletalMeshAccessorHelper& Accessor, int32 Tri, FVector& OutPos0, FVector& OutPos1, FVector& OutPos2, FVector& OutPrev0, FVector& OutPrev1, FVector& OutPrev2, int32& Idx0, int32& Idx1, int32& Idx2) = delete;
+	FORCEINLINE void GetSkinnedTrianglePositions(FSkeletalMeshAccessorHelper& Accessor, int32 Tri, FVector& OutPos0, FVector& OutPos1, FVector& OutPos2, int32& Idx0, int32& Idx1, int32& Idx2) = delete;
+	FORCEINLINE FVector GetSkinnedVertexPosition(FSkeletalMeshAccessorHelper& Accessor, int32 VertexIndex) = delete;
+	FORCEINLINE FVector GetSkinnedVertexPreviousPosition(FSkeletalMeshAccessorHelper& Accessor, int32 VertexIndex) = delete;
+	FORCEINLINE FVector GetSkinnedBonePosition(FSkeletalMeshAccessorHelper& Accessor, int32 BoneIndex) = delete;
+	FORCEINLINE FVector GetSkinnedBonePreviousPosition(FSkeletalMeshAccessorHelper& Accessor, int32 BoneIndex) = delete;
+	FORCEINLINE_DEBUGGABLE FQuat GetSkinnedBoneRotation(FSkeletalMeshAccessorHelper& Accessor, int32 BoneIndex) = delete;
+	FORCEINLINE_DEBUGGABLE FQuat GetSkinnedBonePreviousRotation(FSkeletalMeshAccessorHelper& Accessor, int32 BoneIndex) = delete;
 };
 
 template<>
 struct FSkinnedPositionAccessorHelper<TIntegralConstant<ENDISkeletalMesh_SkinningMode, ENDISkeletalMesh_SkinningMode::None>>
 {
-	FORCEINLINE void GetTrianlgeIndices(FSkeletalMeshAccessorHelper& Accessor, int32 Tri, int32& Idx0, int32& Idx1, int32& Idx2)
+	FORCEINLINE void GetTriangleIndices(FSkeletalMeshAccessorHelper& Accessor, int32 Tri, int32& Idx0, int32& Idx1, int32& Idx2)
 	{
 		const int32 BaseIndex = Tri * 3;
 		checkSlow(BaseIndex + 2 < Accessor.IndexBuffer->Num());
@@ -163,7 +126,7 @@ struct FSkinnedPositionAccessorHelper<TIntegralConstant<ENDISkeletalMesh_Skinnin
 template<>
 struct FSkinnedPositionAccessorHelper<TIntegralConstant<ENDISkeletalMesh_SkinningMode, ENDISkeletalMesh_SkinningMode::SkinOnTheFly>>
 {
-	FORCEINLINE void GetTrianlgeIndices(FSkeletalMeshAccessorHelper& Accessor, int32 Tri, int32& Idx0, int32& Idx1, int32& Idx2)
+	FORCEINLINE void GetTriangleIndices(FSkeletalMeshAccessorHelper& Accessor, int32 Tri, int32& Idx0, int32& Idx1, int32& Idx2)
 	{
 		const int32 BaseIndex = Tri * 3;
 		checkSlow(BaseIndex + 2 < Accessor.IndexBuffer->Num());
@@ -220,7 +183,7 @@ struct FSkinnedPositionAccessorHelper<TIntegralConstant<ENDISkeletalMesh_Skinnin
 template<>
 struct FSkinnedPositionAccessorHelper<TIntegralConstant<ENDISkeletalMesh_SkinningMode, ENDISkeletalMesh_SkinningMode::PreSkin>>
 {
-	FORCEINLINE void GetTrianlgeIndices(FSkeletalMeshAccessorHelper& Accessor, int32 Tri, int32& Idx0, int32& Idx1, int32& Idx2)
+	FORCEINLINE void GetTriangleIndices(FSkeletalMeshAccessorHelper& Accessor, int32 Tri, int32& Idx0, int32& Idx1, int32& Idx2)
 	{
 		const int32 BaseIndex = Tri * 3;
 		checkSlow(BaseIndex + 2 < Accessor.IndexBuffer->Num());
diff --git a/Engine/Plugins/FX/Niagara/Source/Niagara/Private/NDISkeletalMesh_TriangleSampling.cpp b/Engine/Plugins/FX/Niagara/Source/Niagara/Private/NDISkeletalMesh_TriangleSampling.cpp
--- a/Engine/Plugins/FX/Niagara/Source/Niagara/Private/NDISkeletalMesh_TriangleSampling.cpp
+++ b/Engine/Plugins/FX/Niagara/Source/Niagara/Private/NDISkeletalMesh_TriangleSampling.cpp
@@ -911,7 +911,7 @@ void UNiagaraDataInterfaceSkeletalMesh::GetTriCoordSkinnedData(FVectorVMContext&
 			MeshTriCoord = FMeshTriCoordinate(0, FVector(1.0f, 0.0f, 0.0f));
 		}
 
-		SkinningHandler.GetTrianlgeIndices(Accessor, MeshTriCoord.Tri, Idx0, Idx1, Idx2);
+		SkinningHandler.GetTriangleIndices(Accessor, MeshTriCoord.Tri, Idx0, Idx1, Idx2);
 
 		if (bNeedsCurr)
 		{
@@ -1076,7 +1076,7 @@ void UNiagaraDataInterfaceSkeletalMesh::GetTriCoordVertices(FVectorVMContext& Co
 	for (int32 i = 0; i < Context.NumInstances; ++i)
 	{
 		const int32 Tri = FMath::Clamp(TriParam.GetAndAdvance(), 0, TriMax);
-		SkinningHandler.GetTrianlgeIndices(Accessor, Tri, Idx0, Idx1, Idx2);
+		SkinningHandler.GetTriangleIndices(Accessor, Tri, Idx0, Idx1, Idx2);
 		*OutV0.GetDestAndAdvance() = Idx0;
 		*OutV1.GetDestAndAdvance() = Idx1;
 		*OutV2.GetDestAndAdvance() = Idx2;
