Add create comments keybind to c key in niagara graphs.

diff --git a/Engine/Plugins/FX/Niagara/Source/NiagaraEditor/Private/EdGraphSchema_Niagara.cpp b/Engine/Plugins/FX/Niagara/Source/NiagaraEditor/Private/EdGraphSchema_Niagara.cpp
--- a/Engine/Plugins/FX/Niagara/Source/NiagaraEditor/Private/EdGraphSchema_Niagara.cpp
+++ b/Engine/Plugins/FX/Niagara/Source/NiagaraEditor/Private/EdGraphSchema_Niagara.cpp
@@ -162,6 +162,25 @@ void FNiagaraSchemaAction_NewNode::AddReferencedObjects( FReferenceCollector& Co
 	Collector.AddReferencedObject( NodeTemplate );
 }
 
+UEdGraphNode* FNiagaraSchemaAction_NewComment::PerformAction(class UEdGraph* ParentGraph, UEdGraphPin* FromPin, const FVector2D Location, bool bSelectNewNode /*= true*/)
+{
+	// Add menu item for creating comment boxes
+	UEdGraphNode_Comment* CommentTemplate = NewObject<UEdGraphNode_Comment>();
+
+	FVector2D SpawnLocation = Location;
+	FSlateRect Bounds;
+	
+	if (GraphEditor->GetBoundsForSelectedNodes(Bounds, 50.0f))
+	{
+		CommentTemplate->SetBounds(Bounds);
+		SpawnLocation.X = CommentTemplate->NodePosX;
+		SpawnLocation.Y = CommentTemplate->NodePosY;
+	}
+
+	UEdGraphNode* NewNode = FNiagaraSchemaAction_NewNode::SpawnNodeFromTemplate<UEdGraphNode_Comment>(ParentGraph, CommentTemplate, SpawnLocation, bSelectNewNode);
+	return NewNode;
+}
+
 //////////////////////////////////////////////////////////////////////////
 
 static int32 GbAllowAllNiagaraNodesInEmitterGraphs = 1;
@@ -1668,7 +1687,18 @@ void UEdGraphSchema_Niagara::GetContextMenuActions(const UEdGraph* CurrentGraph,
 	}
 	else if (InGraphNode)
 	{
+		if (InGraphNode->IsA<UEdGraphNode_Comment>())
+		{
+			//Comment boxes do not support enable/disable or pin handling, so exit out now
+			return;
+		}
+
 		const UNiagaraNode* Node = Cast<UNiagaraNode>(InGraphNode);
+		if (Node == nullptr)
+		{
+			ensureMsgf(false, TEXT("Encountered unexpected node type when creating context menu actions for Niagara Script Graph!"));
+			return;
+		}
 
 		bool bHasNumerics = false;
 		for (auto Pin : Node->Pins)
diff --git a/Engine/Plugins/FX/Niagara/Source/NiagaraEditor/Private/Widgets/SNiagaraScriptGraph.cpp b/Engine/Plugins/FX/Niagara/Source/NiagaraEditor/Private/Widgets/SNiagaraScriptGraph.cpp
--- a/Engine/Plugins/FX/Niagara/Source/NiagaraEditor/Private/Widgets/SNiagaraScriptGraph.cpp
+++ b/Engine/Plugins/FX/Niagara/Source/NiagaraEditor/Private/Widgets/SNiagaraScriptGraph.cpp
@@ -25,6 +25,7 @@
 #include "Widgets/Images/SImage.h"
 #include "Layout/WidgetPath.h"
 #include "Framework/Application/SlateApplication.h"
+#include "GraphEditorActions.h"
 
 #define LOCTEXT_NAMESPACE "NiagaraScriptGraph"
 
@@ -148,6 +149,9 @@ TSharedRef<SGraphEditor> SNiagaraScriptGraph::ConstructGraphEditor()
 	Commands->MapAction(
 		FNiagaraEditorCommands::Get().FindInCurrentView,
 		FExecuteAction::CreateRaw(this, &SNiagaraScriptGraph::FocusGraphSearchBox));
+	Commands->MapAction(
+		FGraphEditorCommands::Get().CreateComment,
+		FExecuteAction::CreateRaw(this, &SNiagaraScriptGraph::OnCreateComment));
 	
 	TSharedRef<SGraphEditor> CreatedGraphEditor = SNew(SGraphEditor)
 		.AdditionalCommands(ViewModel->GetCommands())
@@ -505,6 +509,12 @@ void SNiagaraScriptGraph::FocusGraphSearchBox()
 	}
 }
 
+void SNiagaraScriptGraph::OnCreateComment()
+{
+	FNiagaraSchemaAction_NewComment CommentAction = FNiagaraSchemaAction_NewComment(GraphEditor);
+	CommentAction.PerformAction(ViewModel->GetGraph(), nullptr, GraphEditor->GetPasteLocation(), false);
+}
+
 void SNiagaraScriptGraph::OnSearchBoxSearch(SSearchBox::SearchDirection Direction)
 {
 	if (CurrentSearchResults.Num() > 0)
diff --git a/Engine/Plugins/FX/Niagara/Source/NiagaraEditor/Private/Widgets/SNiagaraScriptGraph.h b/Engine/Plugins/FX/Niagara/Source/NiagaraEditor/Private/Widgets/SNiagaraScriptGraph.h
--- a/Engine/Plugins/FX/Niagara/Source/NiagaraEditor/Private/Widgets/SNiagaraScriptGraph.h
+++ b/Engine/Plugins/FX/Niagara/Source/NiagaraEditor/Private/Widgets/SNiagaraScriptGraph.h
@@ -30,6 +30,8 @@ public:
 
 	void FocusGraphSearchBox();
 
+	void OnCreateComment();
+
 private:
 	/** Constructs the graph editor widget for the current graph. */
 	TSharedRef<SGraphEditor> ConstructGraphEditor();
diff --git a/Engine/Plugins/FX/Niagara/Source/NiagaraEditor/Public/EdGraphSchema_Niagara.h b/Engine/Plugins/FX/Niagara/Source/NiagaraEditor/Public/EdGraphSchema_Niagara.h
--- a/Engine/Plugins/FX/Niagara/Source/NiagaraEditor/Public/EdGraphSchema_Niagara.h
+++ b/Engine/Plugins/FX/Niagara/Source/NiagaraEditor/Public/EdGraphSchema_Niagara.h
@@ -53,6 +53,31 @@ struct NIAGARAEDITOR_API FNiagaraSchemaAction_NewNode : public FEdGraphSchemaAct
 	}
 };
 
+USTRUCT()
+struct NIAGARAEDITOR_API FNiagaraSchemaAction_NewComment : public FEdGraphSchemaAction
+{
+public:
+	GENERATED_USTRUCT_BODY();
+
+	FNiagaraSchemaAction_NewComment()
+		: FEdGraphSchemaAction()
+		, GraphEditor(nullptr)
+	{}
+
+	FNiagaraSchemaAction_NewComment(const TSharedPtr<SGraphEditor>& InGraphEditor)
+		: FEdGraphSchemaAction()
+		, GraphEditor(InGraphEditor)
+	{}
+
+	//~ Begin FEdGraphSchemaAction Interface
+	virtual UEdGraphNode* PerformAction(class UEdGraph* ParentGraph, UEdGraphPin* FromPin, const FVector2D Location, bool bSelectNewNode = true) override;
+	//~ End FEdGraphSchemaAction Interface
+
+private:
+	/** SGraphEditor to use for getting new comment boxes dimensions */
+	TSharedPtr<SGraphEditor> GraphEditor;
+};
+
 UCLASS()
 class NIAGARAEDITOR_API UEdGraphSchema_Niagara : public UEdGraphSchema
 {
