Added categories for enabling/disabling effects. This gives us more high-level control than just Niagara on/off. All existing parts of the Niagara integration have been modified to contain this.
fx.FavorNiagaraOverCascadeXXXXX where XXXX is the category name is the CVar that controls access.
Also added FName based sets for Niagara parameter values. This keeps us from going to FString to Fname over and over againl. Unfortunately, BP's won't safely up-convert this change, so I had to duplicate the FString/FName functions. The old functions now have the ByString name in BP's.
Stubbed in the Emitter enabled/disabled flag on FXComponent, although it is currently not working.

diff --git a/Engine/Plugins/FX/Niagara/Source/Niagara/Classes/NiagaraEmitterInstance.h b/Engine/Plugins/FX/Niagara/Source/Niagara/Classes/NiagaraEmitterInstance.h
--- a/Engine/Plugins/FX/Niagara/Source/Niagara/Classes/NiagaraEmitterInstance.h
+++ b/Engine/Plugins/FX/Niagara/Source/Niagara/Classes/NiagaraEmitterInstance.h
@@ -47,6 +47,8 @@ public:
 	void BindParameters();
 	void UnbindParameters();
 
+	bool IsAllowedToExecute() const;
+
 	void PreTick();
 	void Tick(float DeltaSeconds);
 	void PostTick();
diff --git a/Engine/Plugins/FX/Niagara/Source/Niagara/Private/NiagaraComponent.cpp b/Engine/Plugins/FX/Niagara/Source/Niagara/Private/NiagaraComponent.cpp
--- a/Engine/Plugins/FX/Niagara/Source/Niagara/Private/NiagaraComponent.cpp
+++ b/Engine/Plugins/FX/Niagara/Source/Niagara/Private/NiagaraComponent.cpp
@@ -384,22 +384,35 @@ UNiagaraComponent::UNiagaraComponent(const FObjectInitializer& ObjectInitializer
 /********* UFXSystemComponent *********/
 void UNiagaraComponent::SetFloatParameter(FName ParameterName, float Param)
 {
-	SetNiagaraVariableFloat(ParameterName.ToString(), Param);
+	SetVariableFloat(ParameterName, Param);
 }
 
 void UNiagaraComponent::SetVectorParameter(FName ParameterName, FVector Param)
 {
-	SetNiagaraVariableVec3(ParameterName.ToString(), Param);
+	SetVariableVec3(ParameterName, Param);
 }
 
 void UNiagaraComponent::SetColorParameter(FName ParameterName, FLinearColor Param)
 {
-	SetNiagaraVariableLinearColor(ParameterName.ToString(), Param);
+	SetVariableLinearColor(ParameterName, Param);
 }
 
 void UNiagaraComponent::SetActorParameter(FName ParameterName, class AActor* Param)
 {
-	SetNiagaraVariableActor(ParameterName.ToString(), Param);
+	SetVariableActor(ParameterName, Param);
+}
+
+
+void UNiagaraComponent::SetEmitterEnable(FName EmitterName, bool bNewEnableState)
+{
+	if (!SystemInstance)
+	{
+		return;
+	}
+	if (SystemInstance.Get() && !SystemInstance->IsComplete())
+	{
+		SystemInstance->SetEmitterEnable(EmitterName, bNewEnableState);
+	}
 }
 /********* UFXSystemComponent *********/
 
@@ -1075,6 +1088,11 @@ FNiagaraSystemInstance* UNiagaraComponent::GetSystemInstance() const
 	return SystemInstance.Get();
 }
 
+void UNiagaraComponent::SetVariableLinearColor(FName InVariableName, const FLinearColor& InValue)
+{
+	OverrideParameters.SetParameterValue(InValue, FNiagaraVariable(FNiagaraTypeDefinition::GetColorDef(), InVariableName), true);
+}
+
 void UNiagaraComponent::SetNiagaraVariableLinearColor(const FString& InVariableName, const FLinearColor& InValue)
 {
 	FName VarName = FName(*InVariableName);
@@ -1082,6 +1100,12 @@ void UNiagaraComponent::SetNiagaraVariableLinearColor(const FString& InVariableN
 	OverrideParameters.SetParameterValue(InValue, FNiagaraVariable(FNiagaraTypeDefinition::GetColorDef(), VarName), true);
 }
 
+
+void UNiagaraComponent::SetVariableQuat(FName InVariableName, const FQuat& InValue)
+{
+	OverrideParameters.SetParameterValue(InValue, FNiagaraVariable(FNiagaraTypeDefinition::GetQuatDef(), InVariableName), true);
+}
+
 void UNiagaraComponent::SetNiagaraVariableQuat(const FString& InVariableName, const FQuat& InValue)
 {
 	FName VarName = FName(*InVariableName);
@@ -1089,6 +1113,11 @@ void UNiagaraComponent::SetNiagaraVariableQuat(const FString& InVariableName, co
 	OverrideParameters.SetParameterValue(InValue, FNiagaraVariable(FNiagaraTypeDefinition::GetQuatDef(), VarName), true);
 }
 
+void UNiagaraComponent::SetVariableVec4(FName InVariableName, const FVector4& InValue)
+{
+	OverrideParameters.SetParameterValue(InValue, FNiagaraVariable(FNiagaraTypeDefinition::GetVec4Def(), InVariableName), true);
+}
+
 void UNiagaraComponent::SetNiagaraVariableVec4(const FString& InVariableName, const FVector4& InValue)
 {
 	FName VarName = FName(*InVariableName);
@@ -1096,6 +1125,11 @@ void UNiagaraComponent::SetNiagaraVariableVec4(const FString& InVariableName, co
 	OverrideParameters.SetParameterValue(InValue, FNiagaraVariable(FNiagaraTypeDefinition::GetVec4Def(), VarName), true);
 }
 
+void UNiagaraComponent::SetVariableVec3(FName InVariableName, FVector InValue)
+{
+	OverrideParameters.SetParameterValue(InValue, FNiagaraVariable(FNiagaraTypeDefinition::GetVec3Def(), InVariableName), true);
+}
+
 void UNiagaraComponent::SetNiagaraVariableVec3(const FString& InVariableName, FVector InValue)
 {
 	FName VarName = FName(*InVariableName);
@@ -1103,6 +1137,11 @@ void UNiagaraComponent::SetNiagaraVariableVec3(const FString& InVariableName, FV
 	OverrideParameters.SetParameterValue(InValue, FNiagaraVariable(FNiagaraTypeDefinition::GetVec3Def(), VarName), true);
 }
 
+void UNiagaraComponent::SetVariableVec2(FName InVariableName, FVector2D InValue)
+{
+	OverrideParameters.SetParameterValue(InValue, FNiagaraVariable(FNiagaraTypeDefinition::GetVec2Def(),InVariableName), true);
+}
+
 void UNiagaraComponent::SetNiagaraVariableVec2(const FString& InVariableName, FVector2D InValue)
 {
 	FName VarName = FName(*InVariableName);
@@ -1110,6 +1149,11 @@ void UNiagaraComponent::SetNiagaraVariableVec2(const FString& InVariableName, FV
 	OverrideParameters.SetParameterValue(InValue, FNiagaraVariable(FNiagaraTypeDefinition::GetVec2Def(),VarName), true);
 }
 
+void UNiagaraComponent::SetVariableFloat(FName InVariableName, float InValue)
+{
+	OverrideParameters.SetParameterValue(InValue, FNiagaraVariable(FNiagaraTypeDefinition::GetFloatDef(), InVariableName), true);
+}
+
 void UNiagaraComponent::SetNiagaraVariableFloat(const FString& InVariableName, float InValue)
 {
 	FName VarName = FName(*InVariableName);
@@ -1117,6 +1161,11 @@ void UNiagaraComponent::SetNiagaraVariableFloat(const FString& InVariableName, f
 	OverrideParameters.SetParameterValue(InValue, FNiagaraVariable(FNiagaraTypeDefinition::GetFloatDef(), VarName), true);
 }
 
+void UNiagaraComponent::SetVariableInt(FName InVariableName, int32 InValue)
+{
+	OverrideParameters.SetParameterValue(InValue, FNiagaraVariable(FNiagaraTypeDefinition::GetIntDef(), InVariableName), true);
+}
+
 void UNiagaraComponent::SetNiagaraVariableInt(const FString& InVariableName, int32 InValue)
 {
 	FName VarName = FName(*InVariableName);
@@ -1124,6 +1173,11 @@ void UNiagaraComponent::SetNiagaraVariableInt(const FString& InVariableName, int
 	OverrideParameters.SetParameterValue(InValue, FNiagaraVariable(FNiagaraTypeDefinition::GetIntDef(), VarName), true);
 }
 
+void UNiagaraComponent::SetVariableBool(FName InVariableName, bool InValue)
+{
+	OverrideParameters.SetParameterValue(InValue ? FNiagaraBool::True : FNiagaraBool::False, FNiagaraVariable(FNiagaraTypeDefinition::GetBoolDef(), InVariableName), true);
+}
+
 void UNiagaraComponent::SetNiagaraVariableBool(const FString& InVariableName, bool InValue)
 {
 	FName VarName = FName(*InVariableName);
@@ -1136,6 +1190,11 @@ void UNiagaraComponent::SetNiagaraVariableActor(const FString& InVariableName, A
 	SetNiagaraVariableObject(InVariableName, InValue);
 }
 
+void UNiagaraComponent::SetVariableActor(FName InVariableName, AActor* InValue)
+{
+	SetVariableObject(InVariableName, InValue);
+}
+
 void UNiagaraComponent::SetNiagaraVariableObject(const FString& InVariableName, UObject* InValue)
 {
 	FName VarName = FName(*InVariableName);
@@ -1143,6 +1202,11 @@ void UNiagaraComponent::SetNiagaraVariableObject(const FString& InVariableName,
 	OverrideParameters.SetUObject(InValue, FNiagaraVariable(FNiagaraTypeDefinition::GetUObjectDef(), VarName));
 }
 
+void UNiagaraComponent::SetVariableObject(FName InVariableName, UObject* InValue)
+{
+	OverrideParameters.SetUObject(InValue, FNiagaraVariable(FNiagaraTypeDefinition::GetUObjectDef(), InVariableName));
+}
+
 TArray<FVector> UNiagaraComponent::GetNiagaraParticlePositions_DebugOnly(const FString& InEmitterName)
 {
 	return GetNiagaraParticleValueVec3_DebugOnly(InEmitterName, TEXT("Position"));
diff --git a/Engine/Plugins/FX/Niagara/Source/Niagara/Private/NiagaraEmitterInstance.cpp b/Engine/Plugins/FX/Niagara/Source/Niagara/Private/NiagaraEmitterInstance.cpp
--- a/Engine/Plugins/FX/Niagara/Source/Niagara/Private/NiagaraEmitterInstance.cpp
+++ b/Engine/Plugins/FX/Niagara/Source/Niagara/Private/NiagaraEmitterInstance.cpp
@@ -153,6 +153,20 @@ void FNiagaraEmitterInstance::Dump()const
 	ParticleDataSet->Dump(0, INDEX_NONE, TEXT("Particle Data"));
 }
 
+bool FNiagaraEmitterInstance::IsAllowedToExecute() const
+{
+	int32 DetailLevel = ParentSystemInstance->GetDetailLevel();
+	const FNiagaraEmitterHandle& EmitterHandle = GetEmitterHandle();
+	if (!EmitterHandle.GetIsEnabled()
+		|| !CachedEmitter->IsAllowedByDetailLevel(DetailLevel)
+		|| (!FNiagaraUtilities::SupportsGPUParticles(GMaxRHIFeatureLevel) && CachedEmitter->SimTarget == ENiagaraSimTarget::GPUComputeSim)  // skip if GPU sim and <SM5. TODO: fall back to CPU sim instead once we have scalability functionality to do so
+		)
+	{
+		return false;
+	}
+	return true;
+}
+
 void FNiagaraEmitterInstance::Init(int32 InEmitterIdx, FName InSystemInstanceName)
 {
 	check(ParticleDataSet);
@@ -164,11 +178,8 @@ void FNiagaraEmitterInstance::Init(int32 InEmitterIdx, FName InSystemInstanceNam
 	checkSlow(CachedEmitter);
 	CachedIDName = EmitterHandle.GetIdName();
 
-	int32 DetailLevel = ParentSystemInstance->GetDetailLevel();
-	if (!EmitterHandle.GetIsEnabled()
-		|| !CachedEmitter->IsAllowedByDetailLevel(DetailLevel)
-		|| (!FNiagaraUtilities::SupportsGPUParticles(GMaxRHIFeatureLevel) && CachedEmitter->SimTarget == ENiagaraSimTarget::GPUComputeSim)  // skip if GPU sim and <SM5. TODO: fall back to CPU sim instead once we have scalability functionality to do so
-		)
+	if (!IsAllowedToExecute())
+
 	{
 		ExecutionState = ENiagaraExecutionState::Disabled;
 		return;
diff --git a/Engine/Plugins/FX/Niagara/Source/Niagara/Private/NiagaraFunctionLibrary.cpp b/Engine/Plugins/FX/Niagara/Source/Niagara/Private/NiagaraFunctionLibrary.cpp
--- a/Engine/Plugins/FX/Niagara/Source/Niagara/Private/NiagaraFunctionLibrary.cpp
+++ b/Engine/Plugins/FX/Niagara/Source/Niagara/Private/NiagaraFunctionLibrary.cpp
@@ -43,7 +43,7 @@ UNiagaraComponent* CreateNiagaraSystem(UNiagaraSystem* SystemTemplate, UWorld* W
 * Spawns a Niagara System at the specified world location/rotation
 * @return			The spawned UNiagaraComponent
 */
-UNiagaraComponent* UNiagaraFunctionLibrary::SpawnSystemAtLocation(UObject* WorldContextObject, UNiagaraSystem* SystemTemplate, FVector SpawnLocation, FRotator SpawnRotation, bool bAutoDestroy)
+UNiagaraComponent* UNiagaraFunctionLibrary::SpawnSystemAtLocation(const UObject* WorldContextObject, UNiagaraSystem* SystemTemplate, FVector SpawnLocation, FRotator SpawnRotation, FVector Scale, bool bAutoDestroy)
 {
 	UNiagaraComponent* PSC = NULL;
 	if (SystemTemplate)
@@ -60,7 +60,7 @@ UNiagaraComponent* UNiagaraFunctionLibrary::SpawnSystemAtLocation(UObject* World
 
 			PSC->SetAbsolute(true, true, true);
 			PSC->SetWorldLocationAndRotation(SpawnLocation, SpawnRotation);
-			PSC->SetRelativeScale3D(FVector(1.f));
+			PSC->SetRelativeScale3D(Scale);
 			PSC->Activate(true);
 		}
 	}
diff --git a/Engine/Plugins/FX/Niagara/Source/Niagara/Private/NiagaraPreviewGrid.cpp b/Engine/Plugins/FX/Niagara/Source/Niagara/Private/NiagaraPreviewGrid.cpp
--- a/Engine/Plugins/FX/Niagara/Source/Niagara/Private/NiagaraPreviewGrid.cpp
+++ b/Engine/Plugins/FX/Niagara/Source/Niagara/Private/NiagaraPreviewGrid.cpp
@@ -324,8 +324,8 @@ void UNiagaraPreviewAxis_InterpParamInt32::ApplyToPreview_Implementation(UNiagar
 	check(PreviewComponent);
 	float Interp = Count > 1 ? (float)PreviewIndex / (Count - 1) : 1.0f;
 	int32 Val = FMath::Lerp(Min, Max, Interp);
-	PreviewComponent->SetNiagaraVariableFloat(Param, Val);
-	OutLabelText = FString::Printf(TEXT("%s = %d"), *Param, Val);
+	PreviewComponent->SetVariableFloat(Param, Val);
+	OutLabelText = FString::Printf(TEXT("%s = %d"), *Param.ToString(), Val);
 }
 void UNiagaraPreviewAxis_InterpParamFloat::ApplyToPreview_Implementation(UNiagaraComponent* PreviewComponent, int32 PreviewIndex, bool bIsXAxis, FString& OutLabelText)
 {
@@ -333,8 +333,8 @@ void UNiagaraPreviewAxis_InterpParamFloat::ApplyToPreview_Implementation(UNiagar
 	check(PreviewComponent);
 	float Interp = Count > 1 ? (float)PreviewIndex / (Count - 1) : 1.0f;
 	float Val = FMath::Lerp(Min, Max, Interp);
-	PreviewComponent->SetNiagaraVariableFloat(Param, Val);
-	OutLabelText = FString::Printf(TEXT("%s = %g"), *Param, Val);
+	PreviewComponent->SetVariableFloat(Param, Val);
+	OutLabelText = FString::Printf(TEXT("%s = %g"), *Param.ToString(), Val);
 }
 void UNiagaraPreviewAxis_InterpParamVector2D::ApplyToPreview_Implementation(UNiagaraComponent* PreviewComponent, int32 PreviewIndex, bool bIsXAxis, FString& OutLabelText)
 {
@@ -342,8 +342,8 @@ void UNiagaraPreviewAxis_InterpParamVector2D::ApplyToPreview_Implementation(UNia
 	check(PreviewComponent);
 	float Interp = Count > 1 ? (float)PreviewIndex / (Count - 1) : 1.0f;
 	FVector2D Val = FMath::Lerp(Min, Max, Interp);
-	PreviewComponent->SetNiagaraVariableVec2(Param, Val);
-	OutLabelText = FString::Printf(TEXT("%s = {%g, %g}"), *Param, Val.X, Val.Y);
+	PreviewComponent->SetVariableVec2(Param, Val);
+	OutLabelText = FString::Printf(TEXT("%s = {%g, %g}"), *Param.ToString(), Val.X, Val.Y);
 }
 void UNiagaraPreviewAxis_InterpParamVector::ApplyToPreview_Implementation(UNiagaraComponent* PreviewComponent, int32 PreviewIndex, bool bIsXAxis, FString& OutLabelText)
 {
@@ -351,8 +351,8 @@ void UNiagaraPreviewAxis_InterpParamVector::ApplyToPreview_Implementation(UNiaga
 	check(PreviewComponent);
 	float Interp = Count > 1 ? (float)PreviewIndex / (Count - 1) : 1.0f;
 	FVector Val = FMath::Lerp(Min, Max, Interp);
-	PreviewComponent->SetNiagaraVariableVec3(Param, Val);
-	OutLabelText = FString::Printf(TEXT("%s = {%g, %g, %g}"), *Param, Val.X, Val.Y, Val.Z);
+	PreviewComponent->SetVariableVec3(Param, Val);
+	OutLabelText = FString::Printf(TEXT("%s = {%g, %g, %g}"), *Param.ToString(), Val.X, Val.Y, Val.Z);
 }
 void UNiagaraPreviewAxis_InterpParamVector4::ApplyToPreview_Implementation(UNiagaraComponent* PreviewComponent, int32 PreviewIndex, bool bIsXAxis, FString& OutLabelText)
 {
@@ -360,8 +360,8 @@ void UNiagaraPreviewAxis_InterpParamVector4::ApplyToPreview_Implementation(UNiag
 	check(PreviewComponent);
 	float Interp = Count > 1 ? (float)PreviewIndex / (Count - 1) : 1.0f;
 	FVector4 Val = FMath::Lerp(Min, Max, Interp);
-	PreviewComponent->SetNiagaraVariableVec4(Param, Val);
-	OutLabelText = FString::Printf(TEXT("%s = {%g, %g, %g, %g}"), *Param, Val.X, Val.Y, Val.Z, Val.W);
+	PreviewComponent->SetVariableVec4(Param, Val);
+	OutLabelText = FString::Printf(TEXT("%s = {%g, %g, %g, %g}"), *Param.ToString(), Val.X, Val.Y, Val.Z, Val.W);
 }
 void UNiagaraPreviewAxis_InterpParamLinearColor::ApplyToPreview_Implementation(UNiagaraComponent* PreviewComponent, int32 PreviewIndex, bool bIsXAxis, FString& OutLabelText)
 {
@@ -369,6 +369,6 @@ void UNiagaraPreviewAxis_InterpParamLinearColor::ApplyToPreview_Implementation(U
 	check(PreviewComponent);
 	float Interp = Count > 1 ? (float)PreviewIndex / (Count - 1) : 1.0f;
 	FLinearColor Val = FMath::Lerp(Min, Max, Interp);
-	PreviewComponent->SetNiagaraVariableLinearColor(Param, Val);
-	OutLabelText = FString::Printf(TEXT("%s = {%g, %g, %g, %g}"), *Param, Val.R, Val.G, Val.B, Val.A);
+	PreviewComponent->SetVariableLinearColor(Param, Val);
+	OutLabelText = FString::Printf(TEXT("%s = {%g, %g, %g, %g}"), *Param.ToString(), Val.R, Val.G, Val.B, Val.A);
 }
\ No newline at end of file
diff --git a/Engine/Plugins/FX/Niagara/Source/Niagara/Private/NiagaraSystemInstance.cpp b/Engine/Plugins/FX/Niagara/Source/Niagara/Private/NiagaraSystemInstance.cpp
--- a/Engine/Plugins/FX/Niagara/Source/Niagara/Private/NiagaraSystemInstance.cpp
+++ b/Engine/Plugins/FX/Niagara/Source/Niagara/Private/NiagaraSystemInstance.cpp
@@ -69,6 +69,53 @@ FNiagaraSystemInstance::FNiagaraSystemInstance(UNiagaraComponent* InComponent)
 	}
 }
 
+
+void FNiagaraSystemInstance::SetEmitterEnable(FName EmitterName, bool bNewEnableState)
+{
+	UNiagaraSystem* System = GetSystem();
+	if (System != nullptr)
+	{
+		const TArray<FNiagaraEmitterHandle>& EmitterHandles = GetSystem()->GetEmitterHandles();
+		int32 FoundIdx = INDEX_NONE;
+		for (int32 EmitterIdx = 0; EmitterIdx < GetSystem()->GetEmitterHandles().Num(); ++EmitterIdx)
+		{
+			const FNiagaraEmitterHandle& EmitterHandle = EmitterHandles[EmitterIdx];
+			if (EmitterName == EmitterHandle.GetName())
+			{
+				FoundIdx = EmitterIdx;
+				break;
+			}
+		}
+
+		if (FoundIdx != INDEX_NONE && Emitters.IsValidIndex(FoundIdx))
+		{
+			if (Emitters[FoundIdx]->IsAllowedToExecute())
+			{
+				
+				{
+					if (bNewEnableState)
+					{
+						Emitters[FoundIdx]->SetExecutionState(ENiagaraExecutionState::Active);
+					}
+					else
+					{
+						Emitters[FoundIdx]->SetExecutionState(ENiagaraExecutionState::Inactive);
+					}
+				}
+			}
+			else
+			{
+				UE_LOG(LogNiagara, Log, TEXT("SetEmitterEnable: Emitter \"%s\" was found in the system's list of emitters, but it does not pass FNiagaraEmitterInstance::IsAllowedToExecute() and therefore cannot be manually enabled!"), *EmitterName.ToString());
+			}
+		}
+		else
+		{
+			UE_LOG(LogNiagara, Log, TEXT("SetEmitterEnable: Emitter \"%s\" was not found in the system's list of emitters!"), *EmitterName.ToString());
+		}
+	}
+}
+
+
 void FNiagaraSystemInstance::Init(bool bInForceSolo)
 {
 	bForceSolo = bInForceSolo;
diff --git a/Engine/Plugins/FX/Niagara/Source/Niagara/Public/NiagaraComponent.h b/Engine/Plugins/FX/Niagara/Source/Niagara/Public/NiagaraComponent.h
--- a/Engine/Plugins/FX/Niagara/Source/Niagara/Public/NiagaraComponent.h
+++ b/Engine/Plugins/FX/Niagara/Source/Niagara/Public/NiagaraComponent.h
@@ -48,6 +48,7 @@ public:
 	void SetColorParameter(FName ParameterName, FLinearColor Param) override;
 	void SetActorParameter(FName ParameterName, class AActor* Param) override;
 	virtual UFXSystemAsset* GetFXSystemAsset() const override { return Asset; };
+	void SetEmitterEnable(FName EmitterName, bool bNewEnableState) override;
 	/********* UFXSystemComponent *********/
 
 private:
@@ -209,42 +210,82 @@ public:
 	bool ForcesSolo()const;
 
 	/** Sets a Niagara FLinearColor parameter by name, overriding locally if necessary.*/
-	UFUNCTION(BlueprintCallable, Category = Niagara, meta = (DisplayName = "Set Niagara Variable (LinearColor)"))
+	UFUNCTION(BlueprintCallable, Category = Niagara, meta = (DisplayName = "Set Niagara Variable By String (LinearColor)"))
 	void SetNiagaraVariableLinearColor(const FString& InVariableName, const FLinearColor& InValue);
+	
+	/** Sets a Niagara FLinearColor parameter by name, overriding locally if necessary.*/
+	UFUNCTION(BlueprintCallable, Category = Niagara, meta = (DisplayName = "Set Niagara Variable (LinearColor)"))
+	void SetVariableLinearColor(FName InVariableName, const FLinearColor& InValue);
+
 
 	/** Sets a Niagara Vector4 parameter by name, overriding locally if necessary.*/
-	UFUNCTION(BlueprintCallable, Category = Niagara, meta = (DisplayName = "Set Niagara Variable (Vector4)"))
+	UFUNCTION(BlueprintCallable, Category = Niagara, meta = (DisplayName = "Set Niagara Variable By String (Vector4)"))
 	void SetNiagaraVariableVec4(const FString& InVariableName, const FVector4& InValue);
+	
+	/** Sets a Niagara Vector4 parameter by name, overriding locally if necessary.*/
+	UFUNCTION(BlueprintCallable, Category = Niagara, meta = (DisplayName = "Set Niagara Variable (Vector4)"))
+	void SetVariableVec4(FName InVariableName, const FVector4& InValue);
+
 
 	/** Sets a Niagara Vector3 parameter by name, overriding locally if necessary.*/
-	UFUNCTION(BlueprintCallable, Category = Niagara, meta = (DisplayName = "Set Niagara Variable (Quaternion)"))
+	UFUNCTION(BlueprintCallable, Category = Niagara, meta = (DisplayName = "Set Niagara Variable By String (Quaternion)"))
 	void SetNiagaraVariableQuat(const FString& InVariableName, const FQuat& InValue);
 
 	/** Sets a Niagara Vector3 parameter by name, overriding locally if necessary.*/
-	UFUNCTION(BlueprintCallable, Category = Niagara, meta = (DisplayName = "Set Niagara Variable (Vector3)"))
+	UFUNCTION(BlueprintCallable, Category = Niagara, meta = (DisplayName = "Set Niagara Variable (Quaternion)"))
+	void SetVariableQuat(FName InVariableName, const FQuat& InValue);
+
+	/** Sets a Niagara Vector3 parameter by name, overriding locally if necessary.*/
+	UFUNCTION(BlueprintCallable, Category = Niagara, meta = (DisplayName = "Set Niagara Variable By String (Vector3)"))
 	void SetNiagaraVariableVec3(const FString& InVariableName, FVector InValue);
 
 	/** Sets a Niagara Vector3 parameter by name, overriding locally if necessary.*/
-	UFUNCTION(BlueprintCallable, Category = Niagara, meta = (DisplayName = "Set Niagara Variable (Vector2)"))
+	UFUNCTION(BlueprintCallable, Category = Niagara, meta = (DisplayName = "Set Niagara Variable (Vector3)"))
+	void SetVariableVec3(FName InVariableName, FVector InValue);
+
+	/** Sets a Niagara Vector3 parameter by name, overriding locally if necessary.*/
+	UFUNCTION(BlueprintCallable, Category = Niagara, meta = (DisplayName = "Set Niagara Variable By String (Vector2)"))
 	void SetNiagaraVariableVec2(const FString& InVariableName, FVector2D InValue);
 
+	/** Sets a Niagara Vector3 parameter by name, overriding locally if necessary.*/
+	UFUNCTION(BlueprintCallable, Category = Niagara, meta = (DisplayName = "Set Niagara Variable (Vector2)"))
+	void SetVariableVec2(FName InVariableName, FVector2D InValue);
+
 	/** Sets a Niagara float parameter by name, overriding locally if necessary.*/
-	UFUNCTION(BlueprintCallable, Category = Niagara, meta = (DisplayName = "Set Niagara Variable (Float)"))
+	UFUNCTION(BlueprintCallable, Category = Niagara, meta = (DisplayName = "Set Niagara Variable By String (Float)"))
 	void SetNiagaraVariableFloat(const FString& InVariableName, float InValue);
 
+	/** Sets a Niagara float parameter by name, overriding locally if necessary.*/
+	UFUNCTION(BlueprintCallable, Category = Niagara, meta = (DisplayName = "Set Niagara Variable (Float)"))
+	void SetVariableFloat(FName InVariableName, float InValue);
+
 	/** Sets a Niagara int parameter by name, overriding locally if necessary.*/
-	UFUNCTION(BlueprintCallable, Category = Niagara, meta = (DisplayName = "Set Niagara Variable (Int32)"))
+	UFUNCTION(BlueprintCallable, Category = Niagara, meta = (DisplayName = "Set Niagara Variable By String (Int32)"))
 	void SetNiagaraVariableInt(const FString& InVariableName, int32 InValue);
 
+	/** Sets a Niagara int parameter by name, overriding locally if necessary.*/
+	UFUNCTION(BlueprintCallable, Category = Niagara, meta = (DisplayName = "Set Niagara Variable (Int32)"))
+	void SetVariableInt(FName InVariableName, int32 InValue);
+
 	/** Sets a Niagara float parameter by name, overriding locally if necessary.*/
-	UFUNCTION(BlueprintCallable, Category = Niagara, meta = (DisplayName = "Set Niagara Variable (Bool)"))
+	UFUNCTION(BlueprintCallable, Category = Niagara, meta = (DisplayName = "Set Niagara Variable By String (Bool)"))
 	void SetNiagaraVariableBool(const FString& InVariableName, bool InValue);
 
-	UFUNCTION(BlueprintCallable, Category = Niagara, meta = (DisplayName = "Set Niagara Variable (Actor)"))
+	/** Sets a Niagara float parameter by name, overriding locally if necessary.*/
+	UFUNCTION(BlueprintCallable, Category = Niagara, meta = (DisplayName = "Set Niagara Variable (Bool)"))
+	void SetVariableBool(FName InVariableName, bool InValue);
+
+	UFUNCTION(BlueprintCallable, Category = Niagara, meta = (DisplayName = "Set Niagara Variable By String (Actor)"))
 	void SetNiagaraVariableActor(const FString& InVariableName, AActor* Actor);
+	
+	UFUNCTION(BlueprintCallable, Category = Niagara, meta = (DisplayName = "Set Niagara Variable (Actor)"))
+	void SetVariableActor(FName InVariableName, AActor* Actor);
+	
+	UFUNCTION(BlueprintCallable, Category = Niagara, meta = (DisplayName = "Set Niagara Variable By String (Object)"))
+	void SetNiagaraVariableObject(const FString& InVariableName, UObject* Object);
 
 	UFUNCTION(BlueprintCallable, Category = Niagara, meta = (DisplayName = "Set Niagara Variable (Object)"))
-	void SetNiagaraVariableObject(const FString& InVariableName, UObject* Object);
+	void SetVariableObject(FName InVariableName, UObject* Object);
 
 	/** Debug accessors for getting positions in blueprints. */
 	UFUNCTION(BlueprintCallable, Category = Niagara, meta = (DisplayName = "Get Niagara Emitter Positions"))
diff --git a/Engine/Plugins/FX/Niagara/Source/Niagara/Public/NiagaraFunctionLibrary.h b/Engine/Plugins/FX/Niagara/Source/Niagara/Public/NiagaraFunctionLibrary.h
--- a/Engine/Plugins/FX/Niagara/Source/Niagara/Public/NiagaraFunctionLibrary.h
+++ b/Engine/Plugins/FX/Niagara/Source/Niagara/Public/NiagaraFunctionLibrary.h
@@ -27,7 +27,7 @@ class NIAGARA_API UNiagaraFunctionLibrary : public UBlueprintFunctionLibrary
 	* @return			The spawned UNiagaraComponent
 	*/
 	UFUNCTION(BlueprintCallable, Category = Niagara, meta = (Keywords = "niagara System", WorldContext = "WorldContextObject", UnsafeDuringActorConstruction = "true"))
-	static UNiagaraComponent* SpawnSystemAtLocation(UObject* WorldContextObject, class UNiagaraSystem* SystemTemplate, FVector Location, FRotator Rotation = FRotator::ZeroRotator, bool bAutoDestroy = true);
+	static UNiagaraComponent* SpawnSystemAtLocation(const UObject* WorldContextObject, class UNiagaraSystem* SystemTemplate, FVector Location, FRotator Rotation = FRotator::ZeroRotator, FVector Scale = FVector(1.f), bool bAutoDestroy = true);
 
 	UFUNCTION(BlueprintCallable, Category = Niagara, meta = (Keywords = "niagara System", UnsafeDuringActorConstruction = "true"))
 	static UNiagaraComponent* SpawnSystemAttached(UNiagaraSystem* SystemTemplate, USceneComponent* AttachToComponent, FName AttachPointName, FVector Location, FRotator Rotation, EAttachLocation::Type LocationType, bool bAutoDestroy);
diff --git a/Engine/Plugins/FX/Niagara/Source/Niagara/Public/NiagaraPreviewGrid.h b/Engine/Plugins/FX/Niagara/Source/Niagara/Public/NiagaraPreviewGrid.h
--- a/Engine/Plugins/FX/Niagara/Source/Niagara/Public/NiagaraPreviewGrid.h
+++ b/Engine/Plugins/FX/Niagara/Source/Niagara/Public/NiagaraPreviewGrid.h
@@ -68,7 +68,7 @@ class UNiagaraPreviewAxis_InterpParamBase : public UNiagaraPreviewAxis
 
 protected:
 	UPROPERTY(EditAnywhere, Category = Axis)
-	FString Param;
+	FName Param;
 	UPROPERTY(EditAnywhere, Category = Axis)
 	int32 Count;
 
diff --git a/Engine/Plugins/FX/Niagara/Source/Niagara/Public/NiagaraSystemInstance.h b/Engine/Plugins/FX/Niagara/Source/Niagara/Public/NiagaraSystemInstance.h
--- a/Engine/Plugins/FX/Niagara/Source/Niagara/Public/NiagaraSystemInstance.h
+++ b/Engine/Plugins/FX/Niagara/Source/Niagara/Public/NiagaraSystemInstance.h
@@ -87,6 +87,8 @@ public:
 	/** Handles completion of the system and returns true if the system is complete. */
 	bool HandleCompletion();
 
+	void SetEmitterEnable(FName EmitterName, bool bNewEnableState);
+
 	/** Perform per-tick updates on data interfaces that need it. This can cause systems to complete so cannot be parallelized. */
 	void TickDataInterfaces(float DeltaSeconds, bool bPostSimulate);
 
diff --git a/Engine/Plugins/FX/Niagara/Source/NiagaraAnimNotifies/Private/AnimNotify_PlayNiagaraEffect.cpp b/Engine/Plugins/FX/Niagara/Source/NiagaraAnimNotifies/Private/AnimNotify_PlayNiagaraEffect.cpp
--- a/Engine/Plugins/FX/Niagara/Source/NiagaraAnimNotifies/Private/AnimNotify_PlayNiagaraEffect.cpp
+++ b/Engine/Plugins/FX/Niagara/Source/NiagaraAnimNotifies/Private/AnimNotify_PlayNiagaraEffect.cpp
@@ -109,7 +109,7 @@ UFXSystemComponent* UAnimNotify_PlayNiagaraEffect::SpawnEffect(USkeletalMeshComp
 		else
 		{
 			const FTransform MeshTransform = MeshComp->GetSocketTransform(SocketName);
-			ReturnComp = UNiagaraFunctionLibrary::SpawnSystemAtLocation(MeshComp->GetWorld(), Template, MeshTransform.TransformPosition(LocationOffset), (MeshTransform.GetRotation() * RotationOffsetQuat).Rotator(), true);
+			ReturnComp = UNiagaraFunctionLibrary::SpawnSystemAtLocation(MeshComp->GetWorld(), Template, MeshTransform.TransformPosition(LocationOffset), (MeshTransform.GetRotation() * RotationOffsetQuat).Rotator(), FVector(1.0f),true);
 		}
 
 		ReturnComp->RelativeScale3D = Scale;
diff --git a/Engine/Source/Runtime/Engine/Classes/Particles/ParticleSystemComponent.h b/Engine/Source/Runtime/Engine/Classes/Particles/ParticleSystemComponent.h
--- a/Engine/Source/Runtime/Engine/Classes/Particles/ParticleSystemComponent.h
+++ b/Engine/Source/Runtime/Engine/Classes/Particles/ParticleSystemComponent.h
@@ -360,6 +360,15 @@ public:
 	 * Get the referenced FXSystem asset.
 	*/
 	virtual UFXSystemAsset* GetFXSystemAsset() const { return nullptr; };
+
+	/**
+	 *	Enables/Disables a sub-emitter
+	 *
+	 *	@param	EmitterName			The name of the sub-emitter to set it on
+	 *	@param	bNewEnableState		The value to set it to
+	 */
+	UFUNCTION(BlueprintCallable, Category = "Effects|Components|ParticleSystem")
+	virtual void SetEmitterEnable(FName EmitterName, bool bNewEnableState) {};
 };
 
 
@@ -1002,8 +1011,7 @@ public:
 	 *	@param	EmitterName			The name of the sub-emitter to set it on
 	 *	@param	bNewEnableState		The value to set it to
 	 */
-	UFUNCTION(BlueprintCallable, Category="Effects|Components|ParticleSystem")
-	virtual void SetEmitterEnable(FName EmitterName, bool bNewEnableState);
+	void SetEmitterEnable(FName EmitterName, bool bNewEnableState) override;
 
 	/** Change a named float parameter */
 	void SetFloatParameter(FName ParameterName, float Param) override;
