Niagara - Fix issues with replacing deprecated and missing modules and dynamic inputs.
+ Update the stack function input binder to become invalid if it can't find a default pin for an input.  This fixes a crash in the case where you're replacing a deprecated module with a suggested module and it has an input with the same name, but the original had an inline edit condition, and the new input doesn't, this allows the edit condition binder to go from valid to invalid when it's refreshed.
+ Update the stack module item to refresh the module function call node when reassigning the function script so that static switch input pins are added and removed as needed.
+ Update the stack function input to refresh the dynamic input function call node when reassigning the function script so that static switch input pins are added and removed as needed.

diff --git a/Engine/Plugins/FX/Niagara/Source/NiagaraEditor/Private/NiagaraStackFunctionInputBinder.cpp b/Engine/Plugins/FX/Niagara/Source/NiagaraEditor/Private/NiagaraStackFunctionInputBinder.cpp
--- a/Engine/Plugins/FX/Niagara/Source/NiagaraEditor/Private/NiagaraStackFunctionInputBinder.cpp
+++ b/Engine/Plugins/FX/Niagara/Source/NiagaraEditor/Private/NiagaraStackFunctionInputBinder.cpp
@@ -176,7 +176,7 @@ bool FNiagaraStackFunctionInputBinder::IsValid() const
 		{
 			RefreshGraphPins();
 		}
-		return OverridePin == nullptr || OverridePin->LinkedTo.Num() == 0;
+		return (OverridePin == nullptr || OverridePin->LinkedTo.Num() == 0) && DefaultPin != nullptr;
 	}
 	return false;
 }
diff --git a/Engine/Plugins/FX/Niagara/Source/NiagaraEditor/Private/ViewModels/Stack/NiagaraStackFunctionInput.cpp b/Engine/Plugins/FX/Niagara/Source/NiagaraEditor/Private/ViewModels/Stack/NiagaraStackFunctionInput.cpp
--- a/Engine/Plugins/FX/Niagara/Source/NiagaraEditor/Private/ViewModels/Stack/NiagaraStackFunctionInput.cpp
+++ b/Engine/Plugins/FX/Niagara/Source/NiagaraEditor/Private/ViewModels/Stack/NiagaraStackFunctionInput.cpp
@@ -1589,6 +1589,7 @@ void UNiagaraStackFunctionInput::ReassignDynamicInputScript(UNiagaraScript* Dyna
 		FScopedTransaction ScopedTransaction(LOCTEXT("ReassignDynamicInputTransaction", "Reassign dynamic input script"));
 		InputValues.DynamicNode->Modify();
 		InputValues.DynamicNode->FunctionScript = DynamicInputScript;
+		InputValues.DynamicNode->RefreshFromExternalChanges();
 		InputValues.DynamicNode->MarkNodeRequiresSynchronization(TEXT("Dynamic input script reassigned."), true);
 		RefreshChildren();
 	}
diff --git a/Engine/Plugins/FX/Niagara/Source/NiagaraEditor/Private/ViewModels/Stack/NiagaraStackModuleItem.cpp b/Engine/Plugins/FX/Niagara/Source/NiagaraEditor/Private/ViewModels/Stack/NiagaraStackModuleItem.cpp
--- a/Engine/Plugins/FX/Niagara/Source/NiagaraEditor/Private/ViewModels/Stack/NiagaraStackModuleItem.cpp
+++ b/Engine/Plugins/FX/Niagara/Source/NiagaraEditor/Private/ViewModels/Stack/NiagaraStackModuleItem.cpp
@@ -832,6 +832,7 @@ void UNiagaraStackModuleItem::ReassignModuleScript(UNiagaraScript* ModuleScript)
 		FScopedTransaction ScopedTransaction(LOCTEXT("ReassignModuleTransaction", "Reassign module script"));
 		FunctionCallNode->Modify();
 		FunctionCallNode->FunctionScript = ModuleScript;
+		FunctionCallNode->RefreshFromExternalChanges();
 		FunctionCallNode->MarkNodeRequiresSynchronization(TEXT("Module script reassigned."), true);
 		RefreshChildren();
 	}
