- Fix Niagara memory leak not deleting the VertexFactory object inside the renderer

diff --git a/Engine/Plugins/FX/Niagara/Source/Niagara/Private/NiagaraRendererMeshes.cpp b/Engine/Plugins/FX/Niagara/Source/Niagara/Private/NiagaraRendererMeshes.cpp
--- a/Engine/Plugins/FX/Niagara/Source/Niagara/Private/NiagaraRendererMeshes.cpp
+++ b/Engine/Plugins/FX/Niagara/Source/Niagara/Private/NiagaraRendererMeshes.cpp
@@ -107,6 +107,15 @@ FNiagaraRendererMeshes::FNiagaraRendererMeshes(ERHIFeatureLevel::Type FeatureLev
 	MeshMinimumLOD = Properties->ParticleMesh->MinLOD.GetValueForFeatureLevel(FeatureLevel);
 }
 
+FNiagaraRendererMeshes::~FNiagaraRendererMeshes()
+{
+	if (VertexFactory != nullptr)
+	{
+		delete VertexFactory;
+		VertexFactory = nullptr;
+	}
+}
+
 void FNiagaraRendererMeshes::ReleaseRenderThreadResources(NiagaraEmitterInstanceBatcher* Batcher)
 {
 	FNiagaraRenderer::ReleaseRenderThreadResources(Batcher);
diff --git a/Engine/Plugins/FX/Niagara/Source/Niagara/Private/NiagaraRendererRibbons.cpp b/Engine/Plugins/FX/Niagara/Source/Niagara/Private/NiagaraRendererRibbons.cpp
--- a/Engine/Plugins/FX/Niagara/Source/Niagara/Private/NiagaraRendererRibbons.cpp
+++ b/Engine/Plugins/FX/Niagara/Source/Niagara/Private/NiagaraRendererRibbons.cpp
@@ -209,6 +209,14 @@ FNiagaraRendererRibbons::FNiagaraRendererRibbons(ERHIFeatureLevel::Type FeatureL
 	MaterialParamValidMask |= MaterialParamOffset3 != -1 ? 8 : 0;
 }
 
+FNiagaraRendererRibbons::~FNiagaraRendererRibbons()
+{
+	if (VertexFactory != nullptr)
+	{
+		delete VertexFactory;
+		VertexFactory = nullptr;
+	}
+}
 
 void FNiagaraRendererRibbons::ReleaseRenderThreadResources(NiagaraEmitterInstanceBatcher* Batcher)
 {
diff --git a/Engine/Plugins/FX/Niagara/Source/Niagara/Private/NiagaraRendererSprites.cpp b/Engine/Plugins/FX/Niagara/Source/Niagara/Private/NiagaraRendererSprites.cpp
--- a/Engine/Plugins/FX/Niagara/Source/Niagara/Private/NiagaraRendererSprites.cpp
+++ b/Engine/Plugins/FX/Niagara/Source/Niagara/Private/NiagaraRendererSprites.cpp
@@ -138,6 +138,15 @@ FNiagaraRendererSprites::FNiagaraRendererSprites(ERHIFeatureLevel::Type FeatureL
 	MaterialParamValidMask |= MaterialParamOffset3 != -1 ? 8 : 0;
 }
 
+FNiagaraRendererSprites::~FNiagaraRendererSprites()
+{
+	if ( VertexFactory != nullptr )
+	{
+		delete VertexFactory;
+		VertexFactory = nullptr;
+	}
+}
+
 void FNiagaraRendererSprites::ReleaseRenderThreadResources(NiagaraEmitterInstanceBatcher* Batcher)
 {
 	FNiagaraRenderer::ReleaseRenderThreadResources(Batcher);
diff --git a/Engine/Plugins/FX/Niagara/Source/Niagara/Public/NiagaraRendererMeshes.h b/Engine/Plugins/FX/Niagara/Source/Niagara/Public/NiagaraRendererMeshes.h
--- a/Engine/Plugins/FX/Niagara/Source/Niagara/Public/NiagaraRendererMeshes.h
+++ b/Engine/Plugins/FX/Niagara/Source/Niagara/Public/NiagaraRendererMeshes.h
@@ -16,8 +16,8 @@ class FNiagaraDataSet;
 class NIAGARA_API FNiagaraRendererMeshes : public FNiagaraRenderer
 {
 public:
-
 	FNiagaraRendererMeshes(ERHIFeatureLevel::Type FeatureLevel, const UNiagaraRendererProperties *InProps, const FNiagaraEmitterInstance* Emitter);
+	~FNiagaraRendererMeshes();
 	
 	//FNiagaraRenderer Interface
 	virtual void CreateRenderThreadResources(NiagaraEmitterInstanceBatcher* Batcher) override;
diff --git a/Engine/Plugins/FX/Niagara/Source/Niagara/Public/NiagaraRendererRibbons.h b/Engine/Plugins/FX/Niagara/Source/Niagara/Public/NiagaraRendererRibbons.h
--- a/Engine/Plugins/FX/Niagara/Source/Niagara/Public/NiagaraRendererRibbons.h
+++ b/Engine/Plugins/FX/Niagara/Source/Niagara/Public/NiagaraRendererRibbons.h
@@ -18,6 +18,7 @@ class NIAGARA_API FNiagaraRendererRibbons : public FNiagaraRenderer
 {
 public:
 	FNiagaraRendererRibbons(ERHIFeatureLevel::Type FeatureLevel, const UNiagaraRendererProperties *InProps, const FNiagaraEmitterInstance* Emitter);	// FNiagaraRenderer Interface 
+	~FNiagaraRendererRibbons();
 
 	// FNiagaraRenderer Interface 
 	virtual void CreateRenderThreadResources(NiagaraEmitterInstanceBatcher* Batcher) override;
diff --git a/Engine/Plugins/FX/Niagara/Source/Niagara/Public/NiagaraRendererSprites.h b/Engine/Plugins/FX/Niagara/Source/Niagara/Public/NiagaraRendererSprites.h
--- a/Engine/Plugins/FX/Niagara/Source/Niagara/Public/NiagaraRendererSprites.h
+++ b/Engine/Plugins/FX/Niagara/Source/Niagara/Public/NiagaraRendererSprites.h
@@ -16,8 +16,8 @@ struct FNiagaraDynamicDataSprites;
 class NIAGARA_API FNiagaraRendererSprites : public FNiagaraRenderer
 {
 public:
-
 	FNiagaraRendererSprites(ERHIFeatureLevel::Type FeatureLevel, const UNiagaraRendererProperties *InProps, const FNiagaraEmitterInstance* Emitter);
+	~FNiagaraRendererSprites();
 
 	//FNiagaraRenderer interface
 	virtual void CreateRenderThreadResources(NiagaraEmitterInstanceBatcher* Batcher)override;
