- Add Niagara category to LLM

diff --git a/Engine/Plugins/FX/Niagara/Source/Niagara/Private/NiagaraComponent.cpp b/Engine/Plugins/FX/Niagara/Source/Niagara/Private/NiagaraComponent.cpp
--- a/Engine/Plugins/FX/Niagara/Source/Niagara/Private/NiagaraComponent.cpp
+++ b/Engine/Plugins/FX/Niagara/Source/Niagara/Private/NiagaraComponent.cpp
@@ -142,6 +142,8 @@ void FNiagaraSceneProxy::ReleaseRenderers()
 
 void FNiagaraSceneProxy::CreateRenderers(const UNiagaraComponent* Component)
 {
+	LLM_SCOPE(ELLMTag::Niagara);
+
 	check(Component);
 	check(Component->GetSystemInstance());
 
@@ -219,6 +221,7 @@ void FNiagaraSceneProxy::ReleaseRenderThreadResources()
 // FPrimitiveSceneProxy interface.
 void FNiagaraSceneProxy::CreateRenderThreadResources()
 {
+	LLM_SCOPE(ELLMTag::Niagara);
 	for (FNiagaraRenderer* Renderer : EmitterRenderers)
 	{
 		if (Renderer)
@@ -419,6 +422,7 @@ void UNiagaraComponent::SetEmitterEnable(FName EmitterName, bool bNewEnableState
 
 void UNiagaraComponent::TickComponent(float DeltaSeconds, enum ELevelTick TickType, FActorComponentTickFunction* ThisTickFunction)
 {
+	LLM_SCOPE(ELLMTag::Niagara);
 	SCOPE_CYCLE_COUNTER(STAT_NiagaraComponentTick);
 	CSV_SCOPED_TIMING_STAT_EXCLUSIVE(Particles);
 
@@ -594,6 +598,7 @@ bool UNiagaraComponent::IsWorldReadyToRun() const
 
 bool UNiagaraComponent::InitializeSystem()
 {
+	LLM_SCOPE(ELLMTag::Niagara);
 	if (SystemInstance.IsValid() == false)
 	{
 		FNiagaraSystemInstance::AllocateSystemInstance(this, SystemInstance);
@@ -931,6 +936,7 @@ void UNiagaraComponent::CreateRenderState_Concurrent()
 
 void UNiagaraComponent::SendRenderDynamicData_Concurrent()
 {
+	LLM_SCOPE(ELLMTag::Niagara);
 	SCOPE_CYCLE_COUNTER(STAT_NiagaraComponentSendRenderData);
 	if (SystemInstance.IsValid() && SceneProxy)
 	{
@@ -1050,6 +1056,7 @@ FBoxSphereBounds UNiagaraComponent::CalcBounds(const FTransform& LocalToWorld) c
 
 FPrimitiveSceneProxy* UNiagaraComponent::CreateSceneProxy()
 {
+	LLM_SCOPE(ELLMTag::Niagara);
 	SCOPE_CYCLE_COUNTER(STAT_NiagaraCreateSceneProxy);
 	SCOPE_CYCLE_COUNTER(STAT_NiagaraOverview_GT);
 	
diff --git a/Engine/Plugins/FX/Niagara/Source/Niagara/Private/NiagaraEmitterInstanceBatcher.cpp b/Engine/Plugins/FX/Niagara/Source/Niagara/Private/NiagaraEmitterInstanceBatcher.cpp
--- a/Engine/Plugins/FX/Niagara/Source/Niagara/Private/NiagaraEmitterInstanceBatcher.cpp
+++ b/Engine/Plugins/FX/Niagara/Source/Niagara/Private/NiagaraEmitterInstanceBatcher.cpp
@@ -103,11 +103,13 @@ void NiagaraEmitterInstanceBatcher::GiveSystemTick_RenderThread(FNiagaraGPUSyste
 
 void NiagaraEmitterInstanceBatcher::GiveEmitterContextToDestroy_RenderThread(FNiagaraComputeExecutionContext* Context)
 {
+	LLM_SCOPE(ELLMTag::Niagara);
 	ContextsToDestroy_RT.Add(Context);
 }
 
 void NiagaraEmitterInstanceBatcher::GiveDataSetToDestroy_RenderThread(FNiagaraDataSet* DataSet)
 {
+	LLM_SCOPE(ELLMTag::Niagara);
 	DataSetsToDestroy_RT.Add(DataSet);
 }
 
@@ -330,6 +332,8 @@ void NiagaraEmitterInstanceBatcher::DispatchAllOnCompute(FOverlappableTicks& Ove
 
 void NiagaraEmitterInstanceBatcher::PostRenderOpaque(FRHICommandListImmediate& RHICmdList, FRHIUniformBuffer* ViewUniformBuffer, const class FShaderParametersMetadata* SceneTexturesUniformBufferStruct, FRHIUniformBuffer* SceneTexturesUniformBuffer)
 {
+	LLM_SCOPE(ELLMTag::Niagara);
+
 	// Setup new readback since if there is no pending request, there is no risk of having invalid data read (offset being allocated after the readback was sent).
 	ExecuteAll(RHICmdList, ViewUniformBuffer, !GPUInstanceCounterManager.HasPendingGPUReadback());
 
@@ -493,6 +497,8 @@ void NiagaraEmitterInstanceBatcher::TickSingle(const FNiagaraGPUSystemTick& Tick
 
 void NiagaraEmitterInstanceBatcher::PreInitViews(FRHICommandListImmediate& RHICmdList)
 {
+	LLM_SCOPE(ELLMTag::Niagara);
+
 	SortedParticleCount = 0;
 	SimulationsToSort.Reset();
 
@@ -578,6 +584,8 @@ bool NiagaraEmitterInstanceBatcher::UsesGlobalDistanceField() const
 
 void NiagaraEmitterInstanceBatcher::PreRender(FRHICommandListImmediate& RHICmdList, const class FGlobalDistanceFieldParameterData* GlobalDistanceFieldParameterData, bool bAllowGPUParticleSceneUpdate)
 {
+	LLM_SCOPE(ELLMTag::Niagara);
+
 	GlobalDistanceFieldParams = GlobalDistanceFieldParameterData ? *GlobalDistanceFieldParameterData : FGlobalDistanceFieldParameterData();
 
 	// Sort buffer after mesh batches are issued, before tick (which will change the GPU instance count).
diff --git a/Engine/Plugins/FX/Niagara/Source/Niagara/Private/NiagaraWorldManager.cpp b/Engine/Plugins/FX/Niagara/Source/Niagara/Private/NiagaraWorldManager.cpp
--- a/Engine/Plugins/FX/Niagara/Source/Niagara/Private/NiagaraWorldManager.cpp
+++ b/Engine/Plugins/FX/Niagara/Source/Niagara/Private/NiagaraWorldManager.cpp
@@ -153,6 +153,7 @@ void FNiagaraWorldManager::CleanupParameterCollections()
 
 TSharedRef<FNiagaraSystemSimulation, ESPMode::ThreadSafe> FNiagaraWorldManager::GetSystemSimulation(UNiagaraSystem* System)
 {
+	LLM_SCOPE(ELLMTag::Niagara);
 	TSharedRef<FNiagaraSystemSimulation, ESPMode::ThreadSafe>* SimPtr = SystemSimulations.Find(System);
 	if (SimPtr != nullptr)
 	{
@@ -208,13 +209,14 @@ void FNiagaraWorldManager::OnWorldCleanup(bool bSessionEnded, bool bCleanupResou
 
 	for ( int32 i=0; i < NumDeferredQueues; ++i)
 	{
-		DeferredDeletionQueue[DeferredDeletionQueueIndex].Fence.Wait();
+		DeferredDeletionQueue[i].Fence.Wait();
 		DeferredDeletionQueue[i].Queue.Empty();
 	}
 }
 
 void FNiagaraWorldManager::Tick(float DeltaSeconds)
 {
+	LLM_SCOPE(ELLMTag::Niagara);
 	SCOPE_CYCLE_COUNTER(STAT_NiagaraWorldManTick);
 	SCOPE_CYCLE_COUNTER(STAT_NiagaraOverview_GT);
 
diff --git a/Engine/Source/Runtime/Core/Private/HAL/LowLevelMemTracker.cpp b/Engine/Source/Runtime/Core/Private/HAL/LowLevelMemTracker.cpp
--- a/Engine/Source/Runtime/Core/Private/HAL/LowLevelMemTracker.cpp
+++ b/Engine/Source/Runtime/Core/Private/HAL/LowLevelMemTracker.cpp
@@ -81,6 +81,7 @@ DECLARE_LLM_MEMORY_STAT(TEXT("Animation"), STAT_AnimationLLM, STATGROUP_LLMFULL)
 DECLARE_LLM_MEMORY_STAT(TEXT("StaticMesh"), STAT_StaticMeshLLM, STATGROUP_LLMFULL);
 DECLARE_LLM_MEMORY_STAT(TEXT("Materials"), STAT_MaterialsLLM, STATGROUP_LLMFULL);
 DECLARE_LLM_MEMORY_STAT(TEXT("Particles"), STAT_ParticlesLLM, STATGROUP_LLMFULL);
+DECLARE_LLM_MEMORY_STAT(TEXT("Niagara"), STAT_NiagaraLLM, STATGROUP_LLMFULL);
 DECLARE_LLM_MEMORY_STAT(TEXT("GC"), STAT_GCLLM, STATGROUP_LLMFULL);
 DECLARE_LLM_MEMORY_STAT(TEXT("UI"), STAT_UILLM, STATGROUP_LLMFULL);
 DECLARE_LLM_MEMORY_STAT(TEXT("PhysX"), STAT_PhysXLLM, STATGROUP_LLMFULL);
@@ -116,6 +117,7 @@ DECLARE_LLM_MEMORY_STAT(TEXT("Animation"), STAT_AnimationSummaryLLM, STATGROUP_L
 DECLARE_LLM_MEMORY_STAT(TEXT("StaticMesh"), STAT_StaticMeshSummaryLLM, STATGROUP_LLM);
 DECLARE_LLM_MEMORY_STAT(TEXT("Materials"), STAT_MaterialsSummaryLLM, STATGROUP_LLM);
 DECLARE_LLM_MEMORY_STAT(TEXT("Particles"), STAT_ParticlesSummaryLLM, STATGROUP_LLM);
+DECLARE_LLM_MEMORY_STAT(TEXT("Niagara"), STAT_NiagaraSummaryLLM, STATGROUP_LLM);
 DECLARE_LLM_MEMORY_STAT(TEXT("UI"), STAT_UISummaryLLM, STATGROUP_LLM);
 DECLARE_LLM_MEMORY_STAT(TEXT("Textures"), STAT_TexturesSummaryLLM, STATGROUP_LLM);
 
diff --git a/Engine/Source/Runtime/Core/Public/HAL/LowLevelMemTracker.h b/Engine/Source/Runtime/Core/Public/HAL/LowLevelMemTracker.h
--- a/Engine/Source/Runtime/Core/Public/HAL/LowLevelMemTracker.h
+++ b/Engine/Source/Runtime/Core/Public/HAL/LowLevelMemTracker.h
@@ -195,6 +195,7 @@ enum class ELLMTagSet : uint8
 	macro(StaticMesh,							"StaticMesh",					GET_STATFNAME(STAT_StaticMeshLLM),							GET_STATFNAME(STAT_StaticMeshSummaryLLM),		ELLMTag::Meshes)\
 	macro(Materials,							"Materials",					GET_STATFNAME(STAT_MaterialsLLM),							GET_STATFNAME(STAT_MaterialsSummaryLLM),		-1)\
 	macro(Particles,							"Particles",					GET_STATFNAME(STAT_ParticlesLLM),							GET_STATFNAME(STAT_ParticlesSummaryLLM),		-1)\
+	macro(Niagara,								"Niagara",						GET_STATFNAME(STAT_NiagaraLLM),								GET_STATFNAME(STAT_NiagaraSummaryLLM),			-1)\
 	macro(GC,									"GC",							GET_STATFNAME(STAT_GCLLM),									GET_STATFNAME(STAT_EngineSummaryLLM),			-1)\
 	macro(UI,									"UI",							GET_STATFNAME(STAT_UILLM),									GET_STATFNAME(STAT_UISummaryLLM),				-1)\
 	macro(PhysX,								"PhysX",						GET_STATFNAME(STAT_PhysXLLM),								GET_STATFNAME(STAT_PhysXSummaryLLM),			-1)\
