- Use GetAllocatedSize rather than Num() or Max() to ensure we always get the correct size even if the underlying structure changes element type

diff --git a/Engine/Plugins/FX/Niagara/Source/Niagara/Private/NiagaraDataSet.cpp b/Engine/Plugins/FX/Niagara/Source/Niagara/Private/NiagaraDataSet.cpp
--- a/Engine/Plugins/FX/Niagara/Source/Niagara/Private/NiagaraDataSet.cpp
+++ b/Engine/Plugins/FX/Niagara/Source/Niagara/Private/NiagaraDataSet.cpp
@@ -472,7 +472,7 @@ FNiagaraDataBuffer::~FNiagaraDataBuffer()
 	// The only exception is if the batcher was pending kill and we couldn't enqueue a rendering command, 
 	// in which case this would have been released on the game thread and not from the batcher DataSetsToDestroy_RT.
 	check(!IsInRenderingThread() || GPUInstanceCountBufferOffset == INDEX_NONE);
-	DEC_MEMORY_STAT_BY(STAT_NiagaraParticleMemory, FloatData.Max() + Int32Data.Max());
+	DEC_MEMORY_STAT_BY(STAT_NiagaraParticleMemory, FloatData.GetAllocatedSize() + Int32Data.GetAllocatedSize());
 }
 
 void FNiagaraDataBuffer::CheckUsage(bool bReadOnly)const
@@ -550,7 +550,7 @@ void FNiagaraDataBuffer::Allocate(uint32 InNumInstances, bool bMaintainExisting)
 	NumInstancesAllocated = InNumInstances;
 	NumInstances = 0;
 
-	DEC_MEMORY_STAT_BY(STAT_NiagaraParticleMemory, FloatData.Max() + Int32Data.Max());
+	DEC_MEMORY_STAT_BY(STAT_NiagaraParticleMemory, FloatData.GetAllocatedSize() + Int32Data.GetAllocatedSize());
 
 	const uint32 OldFloatStride = FloatStride;
 	TArray<uint8> OldFloatData;
@@ -570,7 +570,7 @@ void FNiagaraDataBuffer::Allocate(uint32 InNumInstances, bool bMaintainExisting)
 	Int32Stride = GetSafeComponentBufferSize(NumInstancesAllocated * sizeof(int32));
 	Int32Data.SetNum(Int32Stride * Owner->GetNumInt32Components(), false);
 
-	INC_MEMORY_STAT_BY(STAT_NiagaraParticleMemory, FloatData.Max() + Int32Data.Max());
+	INC_MEMORY_STAT_BY(STAT_NiagaraParticleMemory, FloatData.GetAllocatedSize() + Int32Data.GetAllocatedSize());
 
 	//In some cases we want the existing data in the buffer to be maintained which due to the data layout requires some fix up.
 	if (bMaintainExisting)
diff --git a/Engine/Plugins/FX/Niagara/Source/Niagara/Private/NiagaraParameterStore.cpp b/Engine/Plugins/FX/Niagara/Source/Niagara/Private/NiagaraParameterStore.cpp
--- a/Engine/Plugins/FX/Niagara/Source/Niagara/Private/NiagaraParameterStore.cpp
+++ b/Engine/Plugins/FX/Niagara/Source/Niagara/Private/NiagaraParameterStore.cpp
@@ -51,16 +51,18 @@ FNiagaraParameterStore::FNiagaraParameterStore(UObject* InOwner)
 
 FNiagaraParameterStore::FNiagaraParameterStore(const FNiagaraParameterStore& Other)
 {
+	DEC_MEMORY_STAT_BY(STAT_NiagaraParamStoreMemory, ParameterData.GetAllocatedSize());
 	*this = Other;
+	INC_MEMORY_STAT_BY(STAT_NiagaraParamStoreMemory, ParameterData.GetAllocatedSize());
 }
 
 FNiagaraParameterStore& FNiagaraParameterStore::operator=(const FNiagaraParameterStore& Other)
 {
 	Owner = Other.Owner;
 	ParameterOffsets = Other.ParameterOffsets;
-	DEC_MEMORY_STAT_BY(STAT_NiagaraParamStoreMemory, ParameterData.Num());
+	DEC_MEMORY_STAT_BY(STAT_NiagaraParamStoreMemory, ParameterData.GetAllocatedSize());
 	ParameterData = Other.ParameterData;
-	INC_MEMORY_STAT_BY(STAT_NiagaraParamStoreMemory, ParameterData.Num());
+	INC_MEMORY_STAT_BY(STAT_NiagaraParamStoreMemory, ParameterData.GetAllocatedSize());
 	DataInterfaces = Other.DataInterfaces;
 	UObjects = Other.UObjects;
 	++LayoutVersion;
@@ -75,7 +77,7 @@ FNiagaraParameterStore::~FNiagaraParameterStore()
 {
 	//Ensure that any stores bound to drive this one are unbound.
 	UnbindFromSourceStores();
-	DEC_MEMORY_STAT_BY(STAT_NiagaraParamStoreMemory, ParameterData.Num());
+	DEC_MEMORY_STAT_BY(STAT_NiagaraParamStoreMemory, ParameterData.GetAllocatedSize());
 
 	//Also unbind from any stores we're feeding.
 	for (TPair<FNiagaraParameterStore*, FNiagaraParameterStoreBinding>& Binding : Bindings)
@@ -287,6 +289,8 @@ bool FNiagaraParameterStore::AddParameter(const FNiagaraVariable& Param, bool bI
 	}
 	else
 	{
+		DEC_MEMORY_STAT_BY(STAT_NiagaraParamStoreMemory, ParameterData.GetAllocatedSize());
+
 		int32 ParamSize = Param.GetSizeInBytes();
 		int32 ParamAlignment = Param.GetAlignment();
 		//int32 Offset = AlignArbitrary(ParameterData.Num(), ParamAlignment);//TODO: We need to handle alignment better here. Need to both satisfy CPU and GPU alignment concerns. VM doesn't care but the VM complier needs to be aware. Probably best to have everything adhere to GPU alignment rules.
@@ -295,7 +299,7 @@ bool FNiagaraParameterStore::AddParameter(const FNiagaraVariable& Param, bool bI
 		ParameterOffsets.Add(ParamWithNoAllocatedData) = Offset; // We don't need the default value saved in here.
 		ParameterData.AddUninitialized(ParamSize);
 				
-		INC_MEMORY_STAT_BY(STAT_NiagaraParamStoreMemory, ParamSize);
+		INC_MEMORY_STAT_BY(STAT_NiagaraParamStoreMemory, ParameterData.GetAllocatedSize());
 
 		//Temporary to init param data from FNiagaraVariable storage. This will be removed when we change the UNiagaraScript to use a parameter store too.
 		if (Param.IsDataAllocated())
@@ -360,17 +364,17 @@ bool FNiagaraParameterStore::RemoveParameter(const FNiagaraVariable& ToRemove)
 					FMemory::Memcpy(NewData.GetData() + Offset, ParameterData.GetData() + ExistingOffset, ParamSize);
 				}
 			}
-			else
-			{
-				DEC_MEMORY_STAT_BY(STAT_NiagaraParamStoreMemory, ExistingVar.GetSizeInBytes());
-			}
 		}
 
+		DEC_MEMORY_STAT_BY(STAT_NiagaraParamStoreMemory, ParameterData.GetAllocatedSize());
+
 		ParameterOffsets = NewOffsets;
 		ParameterData = NewData;
 		DataInterfaces = NewInterfaces;
 		UObjects = NewUObjects;
 
+		INC_MEMORY_STAT_BY(STAT_NiagaraParamStoreMemory, ParameterData.GetAllocatedSize());
+
 		OnLayoutChange();
 		return true;
 	}
@@ -500,9 +504,9 @@ void FNiagaraParameterStore::CopyParametersTo(FNiagaraParameterStore& DestStore,
 
 FORCEINLINE void FNiagaraParameterStore::SetParameterDataArray(const TArray<uint8>& InParameterDataArray)
 {
-	DEC_MEMORY_STAT_BY(STAT_NiagaraParamStoreMemory, ParameterData.Num());
+	DEC_MEMORY_STAT_BY(STAT_NiagaraParamStoreMemory, ParameterData.GetAllocatedSize());
 	ParameterData = InParameterDataArray;
-	INC_MEMORY_STAT_BY(STAT_NiagaraParamStoreMemory, ParameterData.Num());
+	INC_MEMORY_STAT_BY(STAT_NiagaraParamStoreMemory, ParameterData.GetAllocatedSize());
 
 	OnParameterChange();
 }
@@ -517,9 +521,9 @@ void FNiagaraParameterStore::InitFromSource(const FNiagaraParameterStore* SrcSto
 
 	ParameterOffsets = SrcStore->ParameterOffsets;
 
-	DEC_MEMORY_STAT_BY(STAT_NiagaraParamStoreMemory, ParameterData.Num());
+	DEC_MEMORY_STAT_BY(STAT_NiagaraParamStoreMemory, ParameterData.GetAllocatedSize());
 	ParameterData = SrcStore->ParameterData;
-	INC_MEMORY_STAT_BY(STAT_NiagaraParamStoreMemory, ParameterData.Num());
+	INC_MEMORY_STAT_BY(STAT_NiagaraParamStoreMemory, ParameterData.GetAllocatedSize());
 
 	DataInterfaces = SrcStore->DataInterfaces;
 
@@ -541,7 +545,6 @@ void FNiagaraParameterStore::RemoveParameters(FNiagaraParameterStore& DestStore)
 	{
 		FNiagaraVariable Parameter = It->Key;
 		DestStore.RemoveParameter(Parameter);
-		DEC_MEMORY_STAT_BY(STAT_NiagaraParamStoreMemory, Parameter.GetSizeInBytes());
 		++It;
 	}
 }
@@ -549,9 +552,9 @@ void FNiagaraParameterStore::RemoveParameters(FNiagaraParameterStore& DestStore)
 void FNiagaraParameterStore::Empty(bool bClearBindings)
 {
 	ParameterOffsets.Empty();
-	DEC_MEMORY_STAT_BY(STAT_NiagaraParamStoreMemory, ParameterData.Num());
+	DEC_MEMORY_STAT_BY(STAT_NiagaraParamStoreMemory, ParameterData.GetAllocatedSize());
 	ParameterData.Empty();
-	DEC_MEMORY_STAT_BY(STAT_NiagaraParamStoreMemory, ParameterData.Num());
+	DEC_MEMORY_STAT_BY(STAT_NiagaraParamStoreMemory, ParameterData.GetAllocatedSize());
 
 	DataInterfaces.Empty();
 
@@ -567,9 +570,9 @@ void FNiagaraParameterStore::Empty(bool bClearBindings)
 void FNiagaraParameterStore::Reset(bool bClearBindings)
 {
 	ParameterOffsets.Reset();
-	DEC_MEMORY_STAT_BY(STAT_NiagaraParamStoreMemory, ParameterData.Num());
+	DEC_MEMORY_STAT_BY(STAT_NiagaraParamStoreMemory, ParameterData.GetAllocatedSize());
 	ParameterData.Reset();
-	DEC_MEMORY_STAT_BY(STAT_NiagaraParamStoreMemory, ParameterData.Num());
+	DEC_MEMORY_STAT_BY(STAT_NiagaraParamStoreMemory, ParameterData.GetAllocatedSize());
 
 	DataInterfaces.Reset();
 
