- Add inline allocator for FNiagaraDataSet -> Data to reduce memory churn

diff --git a/Engine/Plugins/FX/Niagara/Source/Niagara/Classes/NiagaraDataSet.h b/Engine/Plugins/FX/Niagara/Source/Niagara/Classes/NiagaraDataSet.h
--- a/Engine/Plugins/FX/Niagara/Source/Niagara/Classes/NiagaraDataSet.h
+++ b/Engine/Plugins/FX/Niagara/Source/Niagara/Classes/NiagaraDataSet.h
@@ -352,7 +352,7 @@ private:
 	We keep track of the Current and Previous buffers which move with each simulate.
 	Additional buffers may be in here if they are currently being used by the render thread.
 	*/
-	TArray<FNiagaraDataBuffer*> Data;
+	TArray<FNiagaraDataBuffer*, TInlineAllocator<2>> Data;
 
 	FString DebugName;
 };
