Fix for deprecated const-incorrect delegate binding.

diff --git a/Engine/Plugins/FX/Niagara/Source/NiagaraEditor/Private/Customizations/NiagaraTypeCustomizations.cpp b/Engine/Plugins/FX/Niagara/Source/NiagaraEditor/Private/Customizations/NiagaraTypeCustomizations.cpp
--- a/Engine/Plugins/FX/Niagara/Source/NiagaraEditor/Private/Customizations/NiagaraTypeCustomizations.cpp
+++ b/Engine/Plugins/FX/Niagara/Source/NiagaraEditor/Private/Customizations/NiagaraTypeCustomizations.cpp
@@ -331,9 +331,9 @@ TSharedRef<SWidget> FNiagaraUserParameterBindingCustomization::OnGetMenuContent(
 			SNew(SBox)
 			[
 				SNew(SGraphActionMenu)
-				.OnActionSelected(this, &FNiagaraUserParameterBindingCustomization::OnActionSelected)
-		.OnCreateWidgetForAction(SGraphActionMenu::FOnCreateWidgetForAction::CreateSP(this, &FNiagaraUserParameterBindingCustomization::OnCreateWidgetForAction))
-		.OnCollectAllActions(this, &FNiagaraUserParameterBindingCustomization::CollectAllActions)
+				.OnActionSelected(const_cast<FNiagaraUserParameterBindingCustomization*>(this), &FNiagaraUserParameterBindingCustomization::OnActionSelected)
+		.OnCreateWidgetForAction(SGraphActionMenu::FOnCreateWidgetForAction::CreateSP(const_cast<FNiagaraUserParameterBindingCustomization*>(this), &FNiagaraUserParameterBindingCustomization::OnCreateWidgetForAction))
+		.OnCollectAllActions(const_cast<FNiagaraUserParameterBindingCustomization*>(this), &FNiagaraUserParameterBindingCustomization::CollectAllActions)
 		.AutoExpandActionMenu(false)
 		.ShowFilterTextBox(true)
 			]
