- Fix Niagara mesh particle camera facing not using the correct world position

diff --git a/Engine/Plugins/FX/Niagara/Shaders/Private/NiagaraMeshVertexFactory.ush b/Engine/Plugins/FX/Niagara/Shaders/Private/NiagaraMeshVertexFactory.ush
--- a/Engine/Plugins/FX/Niagara/Shaders/Private/NiagaraMeshVertexFactory.ush
+++ b/Engine/Plugins/FX/Niagara/Shaders/Private/NiagaraMeshVertexFactory.ush
@@ -612,7 +612,7 @@ FVertexFactoryIntermediates GetVertexFactoryIntermediates(FVertexFactoryInput In
 	}
 	else if (MeshFacingMode == MESH_FACING_CAMERA_POSITION)
 	{
-		float3 WorldPosition = Intermediates.ParticleTranslatedWorldPosition;
+		float3 WorldPosition = Intermediates.ParticlePosition.xyz;
 		float3 CameraVector = ResolvedView.WorldCameraOrigin.xyz - WorldPosition.xyz;
 		float3 CameraDirection = normalize(CameraVector);
 		float4 ToAlignmentQuat = FindQuatBetween(float3(1.0f, 0.0f, 0.0f), CameraDirection);
