Niagara - Fixes for emitter merging and deprecation.
+ Fail the merge if the scripts for the same function call don't match between parent and child.
+ Fix error logging for merge diff failures.
+ Don't allow fixing missing and deprecated scripts in child emitters, these issues must be fixed in parent emitters.

diff --git a/Engine/Plugins/FX/Niagara/Source/NiagaraEditor/Private/NiagaraScriptMergeManager.cpp b/Engine/Plugins/FX/Niagara/Source/NiagaraEditor/Private/NiagaraScriptMergeManager.cpp
--- a/Engine/Plugins/FX/Niagara/Source/NiagaraEditor/Private/NiagaraScriptMergeManager.cpp
+++ b/Engine/Plugins/FX/Niagara/Source/NiagaraEditor/Private/NiagaraScriptMergeManager.cpp
@@ -633,7 +633,16 @@ FNiagaraEmitterDiffResults::FNiagaraEmitterDiffResults()
 
 bool FNiagaraEmitterDiffResults::IsValid() const
 {
-	return bIsValid &&
+	bool bEventHandlerDiffsAreValid = true;
+	for (const FNiagaraModifiedEventHandlerDiffResults& EventHandlerDiffResults : ModifiedEventHandlers)
+	{
+		if (EventHandlerDiffResults.ScriptDiffResults.IsValid() == false)
+		{
+			bEventHandlerDiffsAreValid = false;
+			break;
+		}
+	}
+	return bIsValid && bEventHandlerDiffsAreValid &&
 		EmitterSpawnDiffResults.IsValid() &&
 		EmitterUpdateDiffResults.IsValid() &&
 		ParticleSpawnDiffResults.IsValid() &&
@@ -845,18 +854,42 @@ INiagaraMergeManager::FMergeEmitterResults FNiagaraScriptMergeManager::MergeEmit
 	{
 		MergeResults.MergeResult = INiagaraMergeManager::EMergeEmitterResult::FailedToDiff;
 		MergeResults.ErrorMessages = DiffResults.GetErrorMessages();
+
+		auto ReportScriptStackDiffErrors = [](FMergeEmitterResults& EmitterMergeResults, const FNiagaraScriptStackDiffResults& ScriptStackDiffResults, FText ScriptName)
+		{
+			FText ScriptStackDiffInvalidFormat = LOCTEXT("ScriptStackDiffInvalidFormat", "Failed to diff {0} script stack.  {1} Errors:");
+			if (ScriptStackDiffResults.IsValid() == false)
+			{
+				EmitterMergeResults.ErrorMessages.Add(FText::Format(ScriptStackDiffInvalidFormat, ScriptName, ScriptStackDiffResults.GetErrorMessages().Num()));
+				for (const FText& ErrorMessage : ScriptStackDiffResults.GetErrorMessages())
+				{
+					EmitterMergeResults.ErrorMessages.Add(ErrorMessage);
+				}
+			}
+		};
+
+		ReportScriptStackDiffErrors(MergeResults, DiffResults.EmitterSpawnDiffResults, LOCTEXT("EmitterSpawnScriptName", "Emitter Spawn"));
+		ReportScriptStackDiffErrors(MergeResults, DiffResults.EmitterUpdateDiffResults, LOCTEXT("EmitterUpdateScriptName", "Emitter Update"));
+		ReportScriptStackDiffErrors(MergeResults, DiffResults.ParticleSpawnDiffResults, LOCTEXT("ParticleSpawnScriptName", "Particle Spawn"));
+		ReportScriptStackDiffErrors(MergeResults, DiffResults.ParticleUpdateDiffResults, LOCTEXT("ParticleUpdateScriptName", "Particle Update"));
+
+		for (const FNiagaraModifiedEventHandlerDiffResults& EventHandlerDiffResults : DiffResults.ModifiedEventHandlers)
+		{
+			FText EventHandlerName = FText::Format(LOCTEXT("EventHandlerScriptNameFormat", "Event Handler - {0}"), FText::FromName(EventHandlerDiffResults.BaseAdapter->GetEventScriptProperties()->SourceEventName));
+			ReportScriptStackDiffErrors(MergeResults, EventHandlerDiffResults.ScriptDiffResults, EventHandlerName);
+		}
 	}
 	else if (DiffResults.IsEmpty())
 	{
 		// If there were no changes made on the instance, check if the instance matches the parent.
 		FNiagaraEmitterDiffResults DiffResultsFromParent = DiffEmitters(Parent, Instance);
-		if (DiffResultsFromParent.IsEmpty())
+		if (DiffResultsFromParent.IsValid() && DiffResultsFromParent.IsEmpty())
 		{
 			MergeResults.MergeResult = INiagaraMergeManager::EMergeEmitterResult::SucceededNoDifferences;
 		}
 		else
 		{
-			// If there were differences from the parent we can just return a copy of the parent as the merged instance since there
+			// If there were differences from the parent or the parent diff failed we can just return a copy of the parent as the merged instance since there
 			// were no changes in the instance which need to be applied.
 			MergeResults.MergeResult = INiagaraMergeManager::EMergeEmitterResult::SucceededDifferencesApplied;
 			MergeResults.MergedInstance = Parent.DuplicateWithoutMerging((UObject*)GetTransientPackage());
@@ -1231,10 +1264,6 @@ FNiagaraEmitterDiffResults FNiagaraScriptMergeManager::DiffEmitters(UNiagaraEmit
 	if (BaseEmitterAdapter->GetEmitterSpawnStack().IsValid() && OtherEmitterAdapter->GetEmitterSpawnStack().IsValid())
 	{
 		DiffScriptStacks(BaseEmitterAdapter->GetEmitterSpawnStack().ToSharedRef(), OtherEmitterAdapter->GetEmitterSpawnStack().ToSharedRef(), EmitterDiffResults.EmitterSpawnDiffResults);
-		if (EmitterDiffResults.EmitterSpawnDiffResults.IsValid() == false)
-		{
-			EmitterDiffResults.AddError(LOCTEXT("EmitterSpawnStacksDiffInvalidMessage", "Emitter spawn diff is invalid."));
-		}
 	}
 	else
 	{
@@ -1244,10 +1273,6 @@ FNiagaraEmitterDiffResults FNiagaraScriptMergeManager::DiffEmitters(UNiagaraEmit
 	if (BaseEmitterAdapter->GetEmitterUpdateStack().IsValid() && OtherEmitterAdapter->GetEmitterUpdateStack().IsValid())
 	{
 		DiffScriptStacks(BaseEmitterAdapter->GetEmitterUpdateStack().ToSharedRef(), OtherEmitterAdapter->GetEmitterUpdateStack().ToSharedRef(), EmitterDiffResults.EmitterUpdateDiffResults);
-		if (EmitterDiffResults.EmitterUpdateDiffResults.IsValid() == false)
-		{
-			EmitterDiffResults.AddError(LOCTEXT("EmitterUpdateStacksDiffInvalidMessage", "Emitter update diff is invalid."));
-		}
 	}
 	else
 	{
@@ -1257,10 +1282,6 @@ FNiagaraEmitterDiffResults FNiagaraScriptMergeManager::DiffEmitters(UNiagaraEmit
 	if (BaseEmitterAdapter->GetParticleSpawnStack().IsValid() && OtherEmitterAdapter->GetParticleSpawnStack().IsValid())
 	{
 		DiffScriptStacks(BaseEmitterAdapter->GetParticleSpawnStack().ToSharedRef(), OtherEmitterAdapter->GetParticleSpawnStack().ToSharedRef(), EmitterDiffResults.ParticleSpawnDiffResults);
-		if (EmitterDiffResults.ParticleSpawnDiffResults.IsValid() == false)
-		{
-			EmitterDiffResults.AddError(LOCTEXT("ParticleSpawnStacksDiffInvalidMessage", "Particle spawn diff is invalid."));
-		}
 	}
 	else
 	{
@@ -1270,10 +1291,6 @@ FNiagaraEmitterDiffResults FNiagaraScriptMergeManager::DiffEmitters(UNiagaraEmit
 	if (BaseEmitterAdapter->GetParticleUpdateStack().IsValid() && OtherEmitterAdapter->GetParticleUpdateStack().IsValid())
 	{
 		DiffScriptStacks(BaseEmitterAdapter->GetParticleUpdateStack().ToSharedRef(), OtherEmitterAdapter->GetParticleUpdateStack().ToSharedRef(), EmitterDiffResults.ParticleUpdateDiffResults);
-		if (EmitterDiffResults.ParticleUpdateDiffResults.IsValid() == false)
-		{
-			EmitterDiffResults.AddError(LOCTEXT("ParticleUpdateStacksDiffInvalidMessage", "Particle update diff is invalid."));
-		}
 	}
 	else
 	{
@@ -1464,7 +1481,19 @@ void FNiagaraScriptMergeManager::DiffScriptStacks(TSharedRef<FNiagaraScriptStack
 			DiffResults.EnabledChangedOtherModules.Add(CommonValuePair.OtherValue);
 		}
 
-		DiffFunctionInputs(CommonValuePair.BaseValue, CommonValuePair.OtherValue, DiffResults);
+		if (CommonValuePair.BaseValue->GetFunctionCallNode()->FunctionScript == CommonValuePair.OtherValue->GetFunctionCallNode()->FunctionScript ||
+			CommonValuePair.BaseValue->GetFunctionCallNode()->IsA<UNiagaraNodeAssignment>())
+		{
+			DiffFunctionInputs(CommonValuePair.BaseValue, CommonValuePair.OtherValue, DiffResults);
+		}
+		else
+		{
+			FText ErrorMessage = FText::Format(LOCTEXT("FunctionScriptMismatchFormat", "Function scripts for function {0} did not match.  Parent: {1} Child: {2}.  This can be fixed by removing the module from the parent, merging the removal to the child, then removing it from the child, and then re-adding it to the parent and merging again."),
+				FText::FromString(CommonValuePair.BaseValue->GetFunctionCallNode()->GetFunctionName()), 
+				FText::FromString(CommonValuePair.BaseValue->GetFunctionCallNode()->FunctionScript != nullptr ? CommonValuePair.BaseValue->GetFunctionCallNode()->FunctionScript->GetPathName() : TEXT("(null)")),
+				FText::FromString(CommonValuePair.OtherValue->GetFunctionCallNode()->FunctionScript != nullptr ? CommonValuePair.OtherValue->GetFunctionCallNode()->FunctionScript->GetPathName() : TEXT("(null)")));
+			DiffResults.AddError(ErrorMessage);
+		}
 	}
 
 	if (BaseScriptStackAdapter->GetScript()->GetUsage() != OtherScriptStackAdapter->GetScript()->GetUsage())
diff --git a/Engine/Plugins/FX/Niagara/Source/NiagaraEditor/Private/ViewModels/Stack/NiagaraStackModuleItem.cpp b/Engine/Plugins/FX/Niagara/Source/NiagaraEditor/Private/ViewModels/Stack/NiagaraStackModuleItem.cpp
--- a/Engine/Plugins/FX/Niagara/Source/NiagaraEditor/Private/ViewModels/Stack/NiagaraStackModuleItem.cpp
+++ b/Engine/Plugins/FX/Niagara/Source/NiagaraEditor/Private/ViewModels/Stack/NiagaraStackModuleItem.cpp
@@ -255,49 +255,78 @@ void UNiagaraStackModuleItem::RefreshIssues(TArray<FStackIssue>& NewIssues)
 	{
 		if (FunctionCallNode->FunctionScript != nullptr && FunctionCallNode->FunctionScript->bDeprecated)
 		{
-			FText LongMessage = FunctionCallNode->FunctionScript->DeprecationRecommendation != nullptr ? 
-				FText::Format(LOCTEXT("ModuleScriptDeprecationLong", "The script asset for the assigned module {0} has been deprecated. Suggested replacement: {1}"), FText::FromString(FunctionCallNode->GetFunctionName()),FText::FromString(FunctionCallNode->FunctionScript->DeprecationRecommendation->GetPathName())) :
-				FText::Format(LOCTEXT("ModuleScriptDeprecationUnknownLong", "The script asset for the assigned module {0} has been deprecated."), FText::FromString(FunctionCallNode->GetFunctionName()));
-
-			int32 AddIdx = NewIssues.Add(FStackIssue(
-				EStackIssueSeverity::Warning,
-				LOCTEXT("ModuleScriptDeprecationShort", "Deprecated module"),
-				LongMessage,
-				GetStackEditorDataKey(),
-				false,
+			FText ModuleScriptDeprecationShort = LOCTEXT("ModuleScriptDeprecationShort", "Deprecated module");
+			if (CanMoveAndDelete())
+			{
+				FText LongMessage = FunctionCallNode->FunctionScript->DeprecationRecommendation != nullptr ?
+					FText::Format(LOCTEXT("ModuleScriptDeprecationLong", "The script asset for the assigned module {0} has been deprecated. Suggested replacement: {1}"), FText::FromString(FunctionCallNode->GetFunctionName()), FText::FromString(FunctionCallNode->FunctionScript->DeprecationRecommendation->GetPathName())) :
+					FText::Format(LOCTEXT("ModuleScriptDeprecationUnknownLong", "The script asset for the assigned module {0} has been deprecated."), FText::FromString(FunctionCallNode->GetFunctionName()));
+
+				int32 AddIdx = NewIssues.Add(FStackIssue(
+					EStackIssueSeverity::Warning,
+					ModuleScriptDeprecationShort,
+					LongMessage,
+					GetStackEditorDataKey(),
+					false,
+					{
+						FStackIssueFix(
+							LOCTEXT("SelectNewModuleScriptFix", "Select a new module script"),
+							FStackIssueFixDelegate::CreateLambda([this]() { this->bIsModuleScriptReassignmentPending = true; })),
+						FStackIssueFix(
+							LOCTEXT("DeleteFix", "Delete this module"),
+							FStackIssueFixDelegate::CreateLambda([this]() { this->Delete(); }))
+					}));
+
+				if (FunctionCallNode->FunctionScript->DeprecationRecommendation != nullptr)
 				{
-					FStackIssueFix(
-						LOCTEXT("SelectNewModuleScriptFix", "Select a new module script"),
-						FStackIssueFixDelegate::CreateLambda([this]() { this->bIsModuleScriptReassignmentPending = true; })),
-					FStackIssueFix(
-						LOCTEXT("DeleteFix", "Delete this module"),
-						FStackIssueFixDelegate::CreateLambda([this]() { this->Delete(); }))
-				}));
-			
-			if (FunctionCallNode->FunctionScript->DeprecationRecommendation != nullptr)
+					NewIssues[AddIdx].InsertFix(0,
+						FStackIssueFix(
+							LOCTEXT("SelectNewModuleScriptFixUseRecommended", "Use recommended replacement"),
+							FStackIssueFixDelegate::CreateLambda([this]() { ReassignModuleScript(FunctionCallNode->FunctionScript->DeprecationRecommendation); })));
+				}
+			}
+			else
 			{
-				NewIssues[AddIdx].InsertFix(0,
-					FStackIssueFix(
-						LOCTEXT("SelectNewModuleScriptFixUseRecommended", "Use recommended replacement"),
-						FStackIssueFixDelegate::CreateLambda([this]() { ReassignModuleScript(FunctionCallNode->FunctionScript->DeprecationRecommendation); })));
+				NewIssues.Add(FStackIssue(
+					EStackIssueSeverity::Warning,
+					ModuleScriptDeprecationShort,
+					FText::Format(LOCTEXT("ModuleScriptDeprecationFixParentLong", "The script asset for the assigned module {0} has been deprecated.  This module is inherited and this issue must be fixed in the parent emitter."),
+						FText::FromString(FunctionCallNode->GetFunctionName())),
+					GetStackEditorDataKey(),
+					false));
 			}
 		}
 		if(FunctionCallNode->FunctionScript == nullptr && FunctionCallNode->GetClass() == UNiagaraNodeFunctionCall::StaticClass())
 		{
-			NewIssues.Add(FStackIssue(
-				EStackIssueSeverity::Error,
-				LOCTEXT("ModuleScriptMissingShort", "Missing module script"),
-				FText::Format(LOCTEXT("ModuleScriptMissingLong", "The script asset for the assigned module {0} is missing."), FText::FromString(FunctionCallNode->GetFunctionName())),
-				GetStackEditorDataKey(),
-				false,
-				{
-					FStackIssueFix(
-						LOCTEXT("SelectNewModuleScriptFix", "Select a new module script"),
-						FStackIssueFixDelegate::CreateLambda([this]() { this->bIsModuleScriptReassignmentPending = true; })),
-				FStackIssueFix(
-					LOCTEXT("DeleteFix", "Delete this module"),
-					FStackIssueFixDelegate::CreateLambda([this]() { this->Delete(); }))
-				}));
+			FText ModuleScriptMissingShort = LOCTEXT("ModuleScriptMissingShort", "Missing module script");
+			if (CanMoveAndDelete())
+			{
+				NewIssues.Add(FStackIssue(
+					EStackIssueSeverity::Error,
+					ModuleScriptMissingShort,
+					FText::Format(LOCTEXT("ModuleScriptMissingLong", "The script asset for the assigned module {0} is missing."), FText::FromString(FunctionCallNode->GetFunctionName())),
+					GetStackEditorDataKey(),
+					false,
+					{
+						FStackIssueFix(
+							LOCTEXT("SelectNewModuleScriptFix", "Select a new module script"),
+							FStackIssueFixDelegate::CreateLambda([this]() { this->bIsModuleScriptReassignmentPending = true; })),
+						FStackIssueFix(
+							LOCTEXT("DeleteFix", "Delete this module"),
+							FStackIssueFixDelegate::CreateLambda([this]() { this->Delete(); }))
+					}));
+			}
+			else
+			{
+				// If the module can't be moved or deleted it's inherited and it's not valid to reassign scripts in child emitters because it breaks merging.
+				NewIssues.Add(FStackIssue(
+					EStackIssueSeverity::Error,
+					ModuleScriptMissingShort,
+					FText::Format(LOCTEXT("ModuleScriptMissingFixParentLong", "The script asset for the assigned module {0} is missing.  This module is inherited and this issue must be fixed in the parent emitter."), 
+						FText::FromString(FunctionCallNode->GetFunctionName())),
+					GetStackEditorDataKey(),
+					false));
+			}
 		}
 		else if (!FunctionCallNode->ScriptIsValid())
 		{
