Merging 4340719 from Dev-Niagara for nicholas.goldstein
Disable Widget focus of expander arrows and pulldown arrows in Niagara Stack to prevent unexpected tab selection

diff --git a/Engine/Plugins/FX/Niagara/Source/NiagaraEditorWidgets/Private/Stack/SNiagaraStackItemExpander.cpp b/Engine/Plugins/FX/Niagara/Source/NiagaraEditorWidgets/Private/Stack/SNiagaraStackItemExpander.cpp
--- a/Engine/Plugins/FX/Niagara/Source/NiagaraEditorWidgets/Private/Stack/SNiagaraStackItemExpander.cpp
+++ b/Engine/Plugins/FX/Niagara/Source/NiagaraEditorWidgets/Private/Stack/SNiagaraStackItemExpander.cpp
@@ -22,6 +22,7 @@ void SNiagaraStackItemExpander::Construct(const FArguments& InArgs, UNiagaraStac
 		.ContentPadding(2)
 		.ToolTipText(this, &SNiagaraStackItemExpander::GetToolTipText)
 		.OnClicked(this, &SNiagaraStackItemExpander::ExpandButtonClicked)
+		.IsFocusable(false)
 		.Content()
 		[
 			SNew(SImage)
diff --git a/Engine/Plugins/FX/Niagara/Source/NiagaraEditorWidgets/Private/Stack/SNiagaraStackTableRow.cpp b/Engine/Plugins/FX/Niagara/Source/NiagaraEditorWidgets/Private/Stack/SNiagaraStackTableRow.cpp
--- a/Engine/Plugins/FX/Niagara/Source/NiagaraEditorWidgets/Private/Stack/SNiagaraStackTableRow.cpp
+++ b/Engine/Plugins/FX/Niagara/Source/NiagaraEditorWidgets/Private/Stack/SNiagaraStackTableRow.cpp
@@ -118,6 +118,7 @@ void SNiagaraStackTableRow::SetNameAndValueContent(TSharedRef<SWidget> InNameWid
 			.ForegroundColor(FSlateColor::UseForeground())
 			.ContentPadding(2)
 			.HAlign(HAlign_Center)
+			.IsFocusable(false)
 			[
 				SNew(SImage)
 				.Image(this, &SNiagaraStackTableRow::GetExpandButtonImage)
