Merging 4340773 from Dev-Niagara for nicholas.goldstein
Add UProperty Specifiers to PropertyMetaData to fix displayed name and tooltip

diff --git a/Engine/Plugins/FX/Niagara/Source/Niagara/Public/NiagaraTypes.h b/Engine/Plugins/FX/Niagara/Source/Niagara/Public/NiagaraTypes.h
--- a/Engine/Plugins/FX/Niagara/Source/Niagara/Public/NiagaraTypes.h
+++ b/Engine/Plugins/FX/Niagara/Source/Niagara/Public/NiagaraTypes.h
@@ -307,7 +307,7 @@ public:
 	{
 	}
 public:
-	UPROPERTY(EditAnywhere, Category = "Variable")
+	UPROPERTY(EditAnywhere, Category = "Variable", DisplayName = "Property Metadata", meta = (ToolTip = "Property Metadata"))
 	TMap<FName, FString> PropertyMetaData;
 
 	UPROPERTY(EditAnywhere, Category = "Variable", meta = (MultiLine = true))
