Implemented D3D12 staging buffer and GPU fence to allow Niagara GPU emitters.
Fixed command list timing issue in Niagara GPU simulation. Now runs in the command list immediate as workaround.
Removed GRHISupportsGPUFence as it is no longer required.

diff --git a/Engine/Plugins/FX/Niagara/Source/Niagara/Classes/NiagaraEmitterInstanceBatcher.h b/Engine/Plugins/FX/Niagara/Source/Niagara/Classes/NiagaraEmitterInstanceBatcher.h
--- a/Engine/Plugins/FX/Niagara/Source/Niagara/Classes/NiagaraEmitterInstanceBatcher.h
+++ b/Engine/Plugins/FX/Niagara/Source/Niagara/Classes/NiagaraEmitterInstanceBatcher.h
@@ -87,11 +87,14 @@ public:
 	virtual void Execute(FRHICommandList &RHICmdList, FUniformBufferRHIParamRef ViewUniformBuffer)
 	{
 		CurQueueIndex ^= 0x1;
-		ExecuteAll(RHICmdList, ViewUniformBuffer);
+
+		// @TODO : we can't use RHICmdList here because we call EnqueueStagingRead which 
+		// require to flush the simulation dispatch before sending the copy to staging buffer command.
+		ExecuteAll(FRHICommandListExecutor::GetImmediateCommandList(), ViewUniformBuffer);
 	}
 
-	void ExecuteAll(FRHICommandList &RHICmdList, FUniformBufferRHIParamRef ViewUniformBuffer);
-	void TickSingle(FNiagaraComputeExecutionContext *Context, FRHICommandList &RHICmdList, FUniformBufferRHIParamRef ViewUniformBuffer) const;
+	void ExecuteAll(FRHICommandListImmediate &RHICmdList, FUniformBufferRHIParamRef ViewUniformBuffer);
+	void TickSingle(FNiagaraComputeExecutionContext *Context, FRHICommandListImmediate &RHICmdList, FUniformBufferRHIParamRef ViewUniformBuffer) const;
 
 	void SetPrevDataStrideParams(const FNiagaraDataSet *Set, FNiagaraShader *Shader, FRHICommandList &RHICmdList) const;
 
@@ -111,7 +114,7 @@ public:
 	void RunEventHandlers(const FNiagaraComputeExecutionContext *Context, uint32 NumInstancesAfterSim, uint32 NumInstancesAfterSpawn, uint32 NumInstancesAfterNonEventSpawn, FRHICommandList &RhiCmdList) const;
 
 	void ClearIndexBufferCur(FRHICommandList &RHICmdList, FNiagaraComputeExecutionContext *Context) const;
-	void ResolveDatasetWrites(FNiagaraComputeExecutionContext *Context) const;
+	void ResolveDatasetWrites(FRHICommandListImmediate &RHICmdList, FNiagaraComputeExecutionContext *Context) const;
 	void ResizeCurrentBuffer(FRHICommandList &RHICmdList, FNiagaraComputeExecutionContext *Context, uint32 NewNumInstances, uint32 PrevNumInstances) const;
 private:
 	static NiagaraEmitterInstanceBatcher* BatcherSingleton;
diff --git a/Engine/Plugins/FX/Niagara/Source/Niagara/Private/NiagaraEmitterInstanceBatcher.cpp b/Engine/Plugins/FX/Niagara/Source/Niagara/Private/NiagaraEmitterInstanceBatcher.cpp
--- a/Engine/Plugins/FX/Niagara/Source/Niagara/Private/NiagaraEmitterInstanceBatcher.cpp
+++ b/Engine/Plugins/FX/Niagara/Source/Niagara/Private/NiagaraEmitterInstanceBatcher.cpp
@@ -11,8 +11,9 @@
 #include "ShaderParameterUtils.h"
 #include "ClearQuad.h"
 
-DECLARE_CYCLE_STAT(TEXT("Batching"), STAT_NiagaraGPUSimTick_RT, STATGROUP_Niagara);
-DECLARE_CYCLE_STAT(TEXT("GPU Data Readback"), STAT_NiagaraGPUReadback_RT, STATGROUP_Niagara);
+DECLARE_CYCLE_STAT(TEXT("GPU Emitter Dispatch [RT]"), STAT_NiagaraGPUSimTick_RT, STATGROUP_Niagara);
+DECLARE_CYCLE_STAT(TEXT("GPU Data Readback [RT]"), STAT_NiagaraGPUReadback_RT, STATGROUP_Niagara);
+DECLARE_CYCLE_STAT(TEXT("Allocate GPU Readback Data [RT]"), STAT_NiagaraAllocateGPUReadback_RT, STATGROUP_Niagara);
 DECLARE_FLOAT_COUNTER_STAT(TEXT("Niagara GPU Sim"), STAT_GPU_NiagaraSim, STATGROUP_GPU);
 DECLARE_DWORD_COUNTER_STAT(TEXT("# GPU Particles"), STAT_NiagaraGPUParticles, STATGROUP_Niagara);
 DECLARE_DWORD_COUNTER_STAT(TEXT("Readback latency (frames)"), STAT_NiagaraReadbackLatency, STATGROUP_Niagara);
@@ -20,7 +21,6 @@ DECLARE_DWORD_COUNTER_STAT(TEXT("Readback latency (frames)"), STAT_NiagaraReadba
 DECLARE_GPU_STAT_NAMED(NiagaraGPUSimulation, TEXT("Niagara GPU Simulation"));
 DECLARE_GPU_STAT_NAMED(NiagaraIndexBufferClear, TEXT("Niagara index buffer clear"));
 
-
 NiagaraEmitterInstanceBatcher* NiagaraEmitterInstanceBatcher::BatcherSingleton = nullptr;
 uint32 FNiagaraComputeExecutionContext::TickCounter = 0;
 
@@ -54,7 +54,7 @@ void NiagaraEmitterInstanceBatcher::Remove(FNiagaraComputeExecutionContext *InCo
 		});
 }
 
-void NiagaraEmitterInstanceBatcher::ExecuteAll(FRHICommandList &RHICmdList, FUniformBufferRHIParamRef ViewUniformBuffer)
+void NiagaraEmitterInstanceBatcher::ExecuteAll(FRHICommandListImmediate &RHICmdList, FUniformBufferRHIParamRef ViewUniformBuffer)
 {
 	TArray<FNiagaraComputeExecutionContext*> &WorkQueue = SimulationQueue[CurQueueIndex ^ 0x1];
 	for (FNiagaraComputeExecutionContext *Context : WorkQueue)
@@ -66,7 +66,7 @@ void NiagaraEmitterInstanceBatcher::ExecuteAll(FRHICommandList &RHICmdList, FUni
 	WorkQueue.Empty();
 }
 
-void NiagaraEmitterInstanceBatcher::TickSingle(FNiagaraComputeExecutionContext *Context, FRHICommandList &RHICmdList, FUniformBufferRHIParamRef ViewUniformBuffer) const
+void NiagaraEmitterInstanceBatcher::TickSingle(FNiagaraComputeExecutionContext *Context, FRHICommandListImmediate &RHICmdList, FUniformBufferRHIParamRef ViewUniformBuffer) const
 {
 	SCOPE_CYCLE_COUNTER(STAT_NiagaraGPUSimTick_RT);
 
@@ -106,8 +106,13 @@ void NiagaraEmitterInstanceBatcher::TickSingle(FNiagaraComputeExecutionContext *
 	// assume all instances survived; ResolveDataSetWrites will change this if the deferred readback was successful; that data may be several frames old
 	Context->MainDataSet->CurrData().SetNumInstances(NewNumInstances);
 
-	// resolve data set writes - grabs the number of instances written from the index set during the simulation run
-	ResolveDatasetWrites(Context);
+	// Don't resolve if the data if there are no instances (prevents a transition issue warning).
+	if (NewNumInstances > 0)
+	{
+		// resolve data set writes - grabs the number of instances written from the index set during the simulation run
+		ResolveDatasetWrites(RHICmdList, Context);
+	}
+
 	RHICmdList.TransitionResource(EResourceTransitionAccess::EReadable, EResourceTransitionPipeline::EComputeToGfx, Context->MainDataSet->GetCurDataSetIndices().UAV);		// transition to readable; we'll be using this next frame
 
 	/*
@@ -164,31 +169,43 @@ void NiagaraEmitterInstanceBatcher::ClearIndexBufferCur(FRHICommandList &RHICmdL
  *	overallocating a bit, and the CPU might think we have more particles alive than we actually do;
  *	since we use DrawIndirect with the GPU determining draw call parameters, that's not an issue
  */
-void NiagaraEmitterInstanceBatcher::ResolveDatasetWrites(FNiagaraComputeExecutionContext *Context) const
+void NiagaraEmitterInstanceBatcher::ResolveDatasetWrites(FRHICommandListImmediate &RHICmdList, FNiagaraComputeExecutionContext *Context) const
 {
 	FRWBuffer &DatasetIndexBufferWrite = Context->MainDataSet->GetCurDataSetIndices();
 	uint32 SpawnedThisFrame = Context->SpawnRateInstances + Context->EventSpawnTotal;
 	Context->AccumulatedSpawnRate += SpawnedThisFrame;
 	if (!Context->GPUDataReadback)
 	{
+		SCOPE_CYCLE_COUNTER(STAT_NiagaraAllocateGPUReadback_RT);
+
 		Context->GPUDataReadback = new FRHIGPUMemoryReadback(DatasetIndexBufferWrite.Buffer, TEXT("Niagara GPU Emitter Readback"));
 		INC_DWORD_STAT(STAT_NiagaraReadbackLatency);
-		Context->GPUDataReadback->Insert();
+		Context->GPUDataReadback->Insert(RHICmdList);
 	}
 	else if (Context->GPUDataReadback->IsReady())
 	{
-		SCOPE_CYCLE_COUNTER(STAT_NiagaraGPUReadback_RT);
-		int32 *NumInstancesAfterSim = static_cast<int32*>(Context->GPUDataReadback->RetrieveData(64 * sizeof(int32)));
-		Context->MainDataSet->CurrData().SetNumInstances(NumInstancesAfterSim[1] + Context->AccumulatedSpawnRate);	// index 1 is always the count
-		SET_DWORD_STAT(STAT_NiagaraGPUParticles, NumInstancesAfterSim[1] + Context->AccumulatedSpawnRate);
-		SET_DWORD_STAT(STAT_NiagaraReadbackLatency, 0);
-
-		Context->GPUDataReadback->Finish();
-
-		Context->AccumulatedSpawnRate = 0;
-		delete Context->GPUDataReadback;
-		Context->GPUDataReadback = new FRHIGPUMemoryReadback(DatasetIndexBufferWrite.Buffer, TEXT("Niagara GPU Emitter Readback"));
-		Context->GPUDataReadback->Insert();
+		{
+		    SCOPE_CYCLE_COUNTER(STAT_NiagaraGPUReadback_RT);
+		    int32 *NumInstancesAfterSim = static_cast<int32*>(Context->GPUDataReadback->RetrieveData(64 * sizeof(int32)));
+		    int32 ExistingDataCount = Context->MainDataSet->CurrData().GetNumInstances();// index 1 is always the count
+		    int32 NewExistingDataCount = NumInstancesAfterSim[1] + Context->AccumulatedSpawnRate;
+		    Context->MainDataSet->CurrData().SetNumInstances(NewExistingDataCount);	
+		    FString PathName = Context->GPUScript->GetOutermost()->GetPathName();
+		    // UE_LOG(LogNiagara, Log, TEXT("GPU Syncup %s : Was(%d) Now(%d)"), *PathName, ExistingDataCount, NewExistingDataCount );
+		    SET_DWORD_STAT(STAT_NiagaraGPUParticles, NewExistingDataCount);
+		    SET_DWORD_STAT(STAT_NiagaraReadbackLatency, 0);
+    
+		    Context->GPUDataReadback->Finish();
+    
+		    Context->AccumulatedSpawnRate = 0;
+		}
+		{
+			SCOPE_CYCLE_COUNTER(STAT_NiagaraAllocateGPUReadback_RT);
+			// The following code seems to take significant time on d3d12
+			delete Context->GPUDataReadback;
+			Context->GPUDataReadback = new FRHIGPUMemoryReadback(DatasetIndexBufferWrite.Buffer, TEXT("Niagara GPU Emitter Readback"));
+			Context->GPUDataReadback->Insert(RHICmdList);
+		}
 	}
 }
 
diff --git a/Engine/Source/Runtime/Apple/MetalRHI/Private/MetalRHI.cpp b/Engine/Source/Runtime/Apple/MetalRHI/Private/MetalRHI.cpp
--- a/Engine/Source/Runtime/Apple/MetalRHI/Private/MetalRHI.cpp
+++ b/Engine/Source/Runtime/Apple/MetalRHI/Private/MetalRHI.cpp
@@ -419,8 +419,6 @@ FMetalDynamicRHI::FMetalDynamicRHI(ERHIFeatureLevel::Type RequestedFeatureLevel)
 		GSupportsParallelOcclusionQueries = false;
 	}
 
-	GRHISupportsGPUFence = true;
-
 	if (FPlatformMisc::IsDebuggerPresent() && UE_BUILD_DEBUG)
 	{
 #if PLATFORM_IOS // @todo zebra : needs a RENDER_API or whatever
diff --git a/Engine/Source/Runtime/D3D12RHI/Private/D3D12Adapter.cpp b/Engine/Source/Runtime/D3D12RHI/Private/D3D12Adapter.cpp
--- a/Engine/Source/Runtime/D3D12RHI/Private/D3D12Adapter.cpp
+++ b/Engine/Source/Runtime/D3D12RHI/Private/D3D12Adapter.cpp
@@ -23,7 +23,7 @@ struct FRHICommandSignalFrameFence final : public FRHICommand<FRHICommandSignalF
 		: QueueType(InQueueType)
 		, Fence(InFence)
 		, Value(InValue)
-	{
+	{ 
 	}
 
 	void Execute(FRHICommandListBase& CmdList)
@@ -41,7 +41,6 @@ FD3D12Adapter::FD3D12Adapter(FD3D12AdapterDesc& DescIn)
 	, RootSignatureManager(this)
 	, PipelineStateCache(this)
 	, FenceCorePool(this)
-	, FrameFence(nullptr)
 	, DeferredDeletionQueue(this)
 	, DefaultContextRedirector(this, true, false)
 	, DefaultAsyncComputeContextRedirector(this, false, true)
@@ -343,6 +342,9 @@ void FD3D12Adapter::InitializeDevices()
 		FrameFence = new FD3D12ManualFence(this, FRHIGPUMask::All(), L"Adapter Frame Fence");
 		FrameFence->CreateFence();
 
+		StagingFence = new FD3D12Fence(this, FRHIGPUMask::All(), L"Staging Fence");
+		StagingFence->CreateFence();
+
 		CreateSignatures();
 
 		//Create all of the FD3D12Devices
@@ -507,6 +509,13 @@ void FD3D12Adapter::Cleanup()
 		FrameFence.SafeRelease();
 	}
 
+	if (StagingFence)
+	{
+		StagingFence->Destroy();
+		StagingFence.SafeRelease();
+	}
+
+
 	PipelineStateCache.Close();
 	FenceCorePool.Destroy();
 }
diff --git a/Engine/Source/Runtime/D3D12RHI/Private/D3D12Adapter.h b/Engine/Source/Runtime/D3D12RHI/Private/D3D12Adapter.h
--- a/Engine/Source/Runtime/D3D12RHI/Private/D3D12Adapter.h
+++ b/Engine/Source/Runtime/D3D12RHI/Private/D3D12Adapter.h
@@ -144,6 +144,8 @@ public:
 
 	FORCEINLINE FD3D12ManualFence& GetFrameFence()  { check(FrameFence); return *FrameFence; }
 
+	FORCEINLINE FD3D12Fence* GetStagingFence()  { return StagingFence.GetReference(); }
+
 	FORCEINLINE FD3D12Device* GetDevice(uint32 GPUIndex)
 	{
 		check(GPUIndex < GNumExplicitGPUsForRendering);
@@ -330,6 +332,9 @@ protected:
 	/** A Fence whos value increases every frame*/
 	TRefCountPtr<FD3D12ManualFence> FrameFence;
 
+	/** A Fence used to syncrhonize FD3D12GPUFence and FD3D12StagingBuffer */
+	TRefCountPtr<FD3D12Fence> StagingFence;
+
 	FD3D12DeferredDeletionQueue DeferredDeletionQueue;
 
 	FD3D12CommandContextRedirector DefaultContextRedirector;
diff --git a/Engine/Source/Runtime/D3D12RHI/Private/D3D12DirectCommandListManager.cpp b/Engine/Source/Runtime/D3D12RHI/Private/D3D12DirectCommandListManager.cpp
--- a/Engine/Source/Runtime/D3D12RHI/Private/D3D12DirectCommandListManager.cpp
+++ b/Engine/Source/Runtime/D3D12RHI/Private/D3D12DirectCommandListManager.cpp
@@ -14,6 +14,159 @@ FComputeFenceRHIRef FD3D12DynamicRHI::RHICreateComputeFence(const FName& Name)
 	return Fence;
 }
 
+// =============================================================================
+
+class FD3D12GPUFence : public FRHIGPUFence
+{
+public:
+
+	FD3D12GPUFence(FName InName, FD3D12Fence* InFence) 
+		: FRHIGPUFence(InName)
+		, Fence(InFence)
+		, Value(0)
+	{}
+
+	virtual void Write() final override
+	{
+		if (Fence)
+		{
+			Value = Fence->Signal(ED3D12CommandQueueType::Default);
+		}
+	}
+
+	virtual bool Poll() const final override
+	{
+		return !Value || (Fence && Fence->IsFenceComplete(Value));
+	}
+
+	virtual bool Wait(float TimeoutMs) const final override
+	{
+		if (Value && Fence)
+		{
+			Fence->WaitForFence(Value);
+		}
+		return Value != 0;
+	}
+
+protected:
+
+	TRefCountPtr<FD3D12Fence> Fence;
+	uint64 Value;
+};
+
+template<>
+struct TD3D12ResourceTraits<FRHIGPUFence>
+{
+	typedef FD3D12GPUFence TConcreteType;
+};
+
+
+FGPUFenceRHIRef FD3D12DynamicRHI::RHICreateGPUFence(const FName& Name)
+{
+	return new FD3D12GPUFence(Name, GetAdapter().GetStagingFence());
+}
+
+// =============================================================================
+
+class FD3D12StagingBuffer : public FRHIStagingBuffer
+{
+public:
+
+	FD3D12StagingBuffer(FVertexBufferRHIRef InBuffer)
+		: FRHIStagingBuffer(InBuffer)
+	{
+	}
+
+	TRefCountPtr<FD3D12Resource> StagedRead;
+};
+
+template<>
+struct TD3D12ResourceTraits<FRHIStagingBuffer>
+{
+	typedef FD3D12StagingBuffer TConcreteType;
+};
+
+
+FStagingBufferRHIRef FD3D12DynamicRHI::RHICreateStagingBuffer(FVertexBufferRHIParamRef VertexBufferRHI)
+{
+	return new FD3D12StagingBuffer(VertexBufferRHI);
+}
+
+void FD3D12DynamicRHI::RHIEnqueueStagedRead(FStagingBufferRHIParamRef StagingBufferRHI, FGPUFenceRHIParamRef FenceRHI, uint32 Offset, uint32 NumBytes)
+{
+	FD3D12StagingBuffer* StagingBuffer = FD3D12DynamicRHI::ResourceCast(StagingBufferRHI);
+	check(StagingBuffer);
+
+	FD3D12VertexBuffer* VertexBuffer = FD3D12DynamicRHI::ResourceCast(StagingBuffer->GetBackingBuffer());
+	check(VertexBuffer);
+
+	// Only get data from the first gpu for now.
+	FD3D12Device* StagingDevice = VertexBuffer->GetParentDevice();
+	FRHIGPUMask StagingGPUMask = StagingDevice->GetGPUMask();
+	StagingBuffer->StagedRead.SafeRelease();
+	VERIFYD3D12RESULT(GetAdapter().CreateBuffer(D3D12_HEAP_TYPE_READBACK, StagingGPUMask, StagingGPUMask, Offset + NumBytes, StagingBuffer->StagedRead.GetInitReference()));
+	
+	FD3D12CommandContext& Context = StagingDevice->GetDefaultCommandContext();
+
+	{
+		FD3D12Resource* pSourceResource = VertexBuffer->ResourceLocation.GetResource();
+		D3D12_RESOURCE_DESC const& SourceBufferDesc = pSourceResource->GetDesc();
+
+		FD3D12Resource* pDestResource = StagingBuffer->StagedRead.GetReference();
+		D3D12_RESOURCE_DESC const& DestBufferDesc = pDestResource->GetDesc();
+
+		TransitionResource(Context.CommandListHandle, pSourceResource, D3D12_RESOURCE_STATE_COPY_SOURCE, 0);
+		// TransitionResource(Context.CommandListHandle, pDestResource, D3D12_RESOURCE_STATE_COPY_DEST, 0);
+		Context.CommandListHandle.FlushResourceBarriers();	// Must flush so the desired state is actually set.
+
+		Context.numCopies++;
+
+		Context.CommandListHandle->CopyBufferRegion(pDestResource->GetResource(), Offset, pSourceResource->GetResource(), Offset, NumBytes);
+		Context.CommandListHandle.UpdateResidency(pDestResource);
+		Context.CommandListHandle.UpdateResidency(pSourceResource);
+
+		Context.FlushCommands();
+
+}
+
+	// Signal the fence
+	FenceRHI->Write();
+}
+
+
+void* FD3D12DynamicRHI::RHILockStagingBuffer(FStagingBufferRHIParamRef StagingBufferRHI, uint32 Offset, uint32 SizeRHI)
+{
+	FD3D12StagingBuffer* StagingBuffer = FD3D12DynamicRHI::ResourceCast(StagingBufferRHI);
+	check(StagingBuffer);
+
+	FD3D12Resource* pResource = StagingBuffer->StagedRead.GetReference();
+	if (pResource)
+	{
+		D3D12_RANGE ReadRange;
+		ReadRange.Begin = Offset;
+		ReadRange.End = Offset + SizeRHI;
+		return reinterpret_cast<uint8*>(pResource->Map(&ReadRange)) + Offset;
+	}
+	else
+	{
+		return nullptr;
+	}
+}
+
+void FD3D12DynamicRHI::RHIUnlockStagingBuffer(FStagingBufferRHIParamRef StagingBufferRHI)
+{
+	FD3D12StagingBuffer* StagingBuffer = FD3D12DynamicRHI::ResourceCast(StagingBufferRHI);
+	check(StagingBuffer);
+
+	FD3D12Resource* pResource = StagingBuffer->StagedRead.GetReference();
+	if (pResource)
+	{
+		pResource->Unmap();
+	}
+}
+
+// =============================================================================
+
 FD3D12FenceCore::FD3D12FenceCore(FD3D12Adapter* Parent, uint64 InitialValue, uint32 InGPUIndex)
 	: FD3D12AdapterChild(Parent)
 	, FenceValueAvailableAt(0)
diff --git a/Engine/Source/Runtime/D3D12RHI/Private/D3D12RHIPrivate.h b/Engine/Source/Runtime/D3D12RHI/Private/D3D12RHIPrivate.h
--- a/Engine/Source/Runtime/D3D12RHI/Private/D3D12RHIPrivate.h
+++ b/Engine/Source/Runtime/D3D12RHI/Private/D3D12RHIPrivate.h
@@ -234,6 +234,11 @@ public:
 	virtual FGeometryShaderRHIRef RHICreateGeometryShaderWithStreamOutput(const TArray<uint8>& Code, const FStreamOutElementList& ElementList, uint32 NumStrides, const uint32* Strides, int32 RasterizedStream) final override;
 	virtual FComputeShaderRHIRef RHICreateComputeShader(const TArray<uint8>& Code) override;
 	virtual FComputeFenceRHIRef RHICreateComputeFence(const FName& Name) final override;
+	virtual FGPUFenceRHIRef RHICreateGPUFence(const FName& Name) final override;
+	virtual FStagingBufferRHIRef RHICreateStagingBuffer(FVertexBufferRHIParamRef VertexBuffer) final override;
+    virtual void RHIEnqueueStagedRead(FStagingBufferRHIParamRef StagingBuffer, FGPUFenceRHIParamRef Fence, uint32 Offset, uint32 NumBytes) final override;
+	virtual void* RHILockStagingBuffer(FStagingBufferRHIParamRef StagingBuffer, uint32 Offset, uint32 SizeRHI) final override;
+    virtual void RHIUnlockStagingBuffer(FStagingBufferRHIParamRef StagingBuffer) final override;
 	virtual FBoundShaderStateRHIRef RHICreateBoundShaderState(FVertexDeclarationRHIParamRef VertexDeclaration, FVertexShaderRHIParamRef VertexShader, FHullShaderRHIParamRef HullShader, FDomainShaderRHIParamRef DomainShader, FPixelShaderRHIParamRef PixelShader, FGeometryShaderRHIParamRef GeometryShader) final override;
 	virtual FGraphicsPipelineStateRHIRef RHICreateGraphicsPipelineState(const FGraphicsPipelineStateInitializer& Initializer) final override;
 	virtual TRefCountPtr<FRHIComputePipelineState> RHICreateComputePipelineState(FRHIComputeShader* ComputeShader) final override;
diff --git a/Engine/Source/Runtime/D3D12RHI/Public/D3D12Resources.h b/Engine/Source/Runtime/D3D12RHI/Public/D3D12Resources.h
--- a/Engine/Source/Runtime/D3D12RHI/Public/D3D12Resources.h
+++ b/Engine/Source/Runtime/D3D12RHI/Public/D3D12Resources.h
@@ -131,10 +131,10 @@ public:
 	operator ID3D12Resource&() { return *Resource; }
 	ID3D12Resource* GetResource() const { return Resource.GetReference(); }
 
-	inline void* Map()
+	inline void* Map(const D3D12_RANGE* ReadRange = nullptr)
 	{
 		check(Resource);
-		VERIFYD3D12RESULT(Resource->Map(0, nullptr, &ResourceBaseAddress));
+		VERIFYD3D12RESULT(Resource->Map(0, ReadRange, &ResourceBaseAddress));
 
 		return ResourceBaseAddress;
 	}
diff --git a/Engine/Source/Runtime/RHI/Private/RHI.cpp b/Engine/Source/Runtime/RHI/Private/RHI.cpp
--- a/Engine/Source/Runtime/RHI/Private/RHI.cpp
+++ b/Engine/Source/Runtime/RHI/Private/RHI.cpp
@@ -515,8 +515,6 @@ EPixelFormat GRHIHDRDisplayOutputFormat = PF_FloatRGBA;
 
 uint64 GRHIPresentCounter = 1;
 
-bool GRHISupportsGPUFence = false;
-
 /** Whether we are profiling GPU hitches. */
 bool GTriggerGPUHitchProfile = false;
 
diff --git a/Engine/Source/Runtime/RHI/Private/RHICommandList.cpp b/Engine/Source/Runtime/RHI/Private/RHICommandList.cpp
--- a/Engine/Source/Runtime/RHI/Private/RHICommandList.cpp
+++ b/Engine/Source/Runtime/RHI/Private/RHICommandList.cpp
@@ -2479,43 +2479,6 @@ void FDynamicRHI::RHICopySubTextureRegion_RenderThread(class FRHICommandListImme
 	return GDynamicRHI->RHICopySubTextureRegion(SourceTexture, DestinationTexture, SourceBox, DestinationBox);
 }
 
-struct FDynamicRHICommandEnqueueStagedRead : public FRHICommand<FDynamicRHICommandEnqueueStagedRead>
-{
-	FStagingBufferRHIParamRef Buffer;
-	FGPUFenceRHIParamRef Fence;
-	uint32 Offset;
-	uint32 NumBytes;
-	
-	FORCEINLINE_DEBUGGABLE FDynamicRHICommandEnqueueStagedRead(FStagingBufferRHIParamRef StagingBuffer, FGPUFenceRHIParamRef InFence, uint32 InOffset, uint32 InNumBytes)
-	: Buffer(StagingBuffer)
-	, Fence(InFence)
-	, Offset(InOffset)
-	, NumBytes(InNumBytes)
-	{
-	}
-	
-	virtual ~FDynamicRHICommandEnqueueStagedRead()
-	{
-	}
-	
-	void Execute(FRHICommandListBase& CmdList)
-	{
-		GDynamicRHI->RHIEnqueueStagedRead(Buffer, Fence, Offset, NumBytes);
-	}
-};
-
-void FDynamicRHI::RHIEnqueueStagedRead_RenderThread(class FRHICommandListImmediate& RHICmdList, FStagingBufferRHIParamRef StagingBuffer, FGPUFenceRHIParamRef Fence, uint32 Offset, uint32 NumBytes)
-{
-	// When the RHI does not support GRHISupportsGPUFence then we assume the default implementation of FRHIGPUFence.
-	// The default FRHIGPUFence uses GFrameNumberRenderThread and thus does not need to be dispatched to the RHI thread.
-	if (!GRHISupportsGPUFence || RHICmdList.Bypass() || !IsRunningRHIInSeparateThread())
-	{
-		GDynamicRHI->RHIEnqueueStagedRead(StagingBuffer, Fence, Offset, NumBytes);
-		return;
-	}
-	new (RHICmdList.AllocCommand<FDynamicRHICommandEnqueueStagedRead>()) FDynamicRHICommandEnqueueStagedRead(StagingBuffer, Fence, Offset, NumBytes);
-}
-
 void FDynamicRHI::RHIEnqueueStagedRead(FStagingBufferRHIParamRef StagingBuffer, FGPUFenceRHIParamRef Fence, uint32 Offset, uint32 NumBytes)
 {
 	IRHICommandContext* Context = RHIGetDefaultContext();
diff --git a/Engine/Source/Runtime/RHI/Private/RHIGPUReadback.cpp b/Engine/Source/Runtime/RHI/Private/RHIGPUReadback.cpp
--- a/Engine/Source/Runtime/RHI/Private/RHIGPUReadback.cpp
+++ b/Engine/Source/Runtime/RHI/Private/RHIGPUReadback.cpp
@@ -8,3 +8,36 @@
 #include "RHIGPUReadback.h"
 
 
+void FRHIGPUFence::Write()
+{
+	UE_LOG(LogRHI, Warning, TEXT("FRHIGPUFence::Write is not implemented"));
+}
+
+FGenericRHIGPUFence::FGenericRHIGPUFence(FName InName)
+    : FRHIGPUFence(InName)
+    , InsertedFrameNumber(0)
+{}
+
+void FGenericRHIGPUFence::Write()
+{
+	// GPU generally overlap the game. This overlap increases when using AFR.
+	InsertedFrameNumber = GFrameNumberRenderThread + GNumAlternateFrameRenderingGroups;
+}
+
+bool FGenericRHIGPUFence::Poll() const
+{
+	if (GFrameNumberRenderThread > InsertedFrameNumber)
+	{
+		return true;
+	}
+	return false;
+}
+
+bool FGenericRHIGPUFence::Wait(float TimeoutMs) const
+{
+	if (GFrameNumberRenderThread > InsertedFrameNumber)
+	{
+		return true;
+	}
+	return false;
+}
diff --git a/Engine/Source/Runtime/RHI/Public/DynamicRHI.h b/Engine/Source/Runtime/RHI/Public/DynamicRHI.h
--- a/Engine/Source/Runtime/RHI/Public/DynamicRHI.h
+++ b/Engine/Source/Runtime/RHI/Public/DynamicRHI.h
@@ -212,7 +212,7 @@ public:
 
 	virtual FGPUFenceRHIRef RHICreateGPUFence(const FName &Name)
 	{
-		return new FRHIGPUFence(Name);
+		return new FGenericRHIGPUFence(Name);
 	}
 
 	/**
@@ -1055,7 +1055,6 @@ public:
      * @param NumBytes The lenght of data in 'StagingBuffer' to make available.
      */
     virtual void RHIEnqueueStagedRead(FStagingBufferRHIParamRef StagingBuffer, FGPUFenceRHIParamRef Fence, uint32 Offset, uint32 NumBytes);
-    virtual void RHIEnqueueStagedRead_RenderThread(class FRHICommandListImmediate& RHICmdList, FStagingBufferRHIParamRef StagingBuffer, FGPUFenceRHIParamRef Fence, uint32 Offset, uint32 NumBytes);
 	
 	virtual FRHIFlipDetails RHIWaitForFlip(double TimeoutInSeconds) { return FRHIFlipDetails(); }
 	virtual void RHISignalFlipEvent() { }
diff --git a/Engine/Source/Runtime/RHI/Public/RHI.h b/Engine/Source/Runtime/RHI/Public/RHI.h
--- a/Engine/Source/Runtime/RHI/Public/RHI.h
+++ b/Engine/Source/Runtime/RHI/Public/RHI.h
@@ -463,14 +463,6 @@ extern RHI_API EPixelFormat GRHIHDRDisplayOutputFormat;
 /** Counter incremented once on each frame present. Used to support game thread synchronization with swap chain frame flips. */
 extern RHI_API uint64 GRHIPresentCounter;
 
-/** Whether or not the RHI can support a real GPU-to-CPU fence or not. 
- * True if and only if the RHI implements:
- *  - RHICreateGPUFence
- *  - RHIInsertGPUFence
- * such that the implementation is *not* dependent upon GFrameNumber or GFrameNumberRenderThread.
- */
-extern RHI_API bool GRHISupportsGPUFence;
-
 /** Called once per frame only from within an RHI. */
 extern RHI_API void RHIPrivateBeginFrame();
 
diff --git a/Engine/Source/Runtime/RHI/Public/RHICommandList.h b/Engine/Source/Runtime/RHI/Public/RHICommandList.h
--- a/Engine/Source/Runtime/RHI/Public/RHICommandList.h
+++ b/Engine/Source/Runtime/RHI/Public/RHICommandList.h
@@ -3991,7 +3991,8 @@ public:
 
 	FORCEINLINE void EnqueueStagedRead(FStagingBufferRHIParamRef StagingBuffer, FGPUFenceRHIParamRef Fence, uint32 Offset, uint32 NumBytes)
 	{
-		GDynamicRHI->RHIEnqueueStagedRead_RenderThread(*this, StagingBuffer, Fence, Offset, NumBytes);
+		ImmediateFlush(EImmediateFlushType::FlushRHIThread); 
+		GDynamicRHI->RHIEnqueueStagedRead(StagingBuffer, Fence, Offset, NumBytes);
 	}
 	
 	FORCEINLINE void ExecuteCommandList(FRHICommandList* CmdList)
diff --git a/Engine/Source/Runtime/RHI/Public/RHIGPUReadback.h b/Engine/Source/Runtime/RHI/Public/RHIGPUReadback.h
--- a/Engine/Source/Runtime/RHI/Public/RHIGPUReadback.h
+++ b/Engine/Source/Runtime/RHI/Public/RHIGPUReadback.h
@@ -29,9 +29,8 @@ public:
 	{
 	}
 
-	void Insert()
+	void Insert(FRHICommandListImmediate& RHICmdList)
 	{
-		FRHICommandListImmediate& RHICmdList = FRHICommandListExecutor::GetImmediateCommandList();
 		RHICmdList.EnqueueStagedRead(StagingBuffer, Fence, 0, GPUVertexBuffer->GetSize());
 	}
 
diff --git a/Engine/Source/Runtime/RHI/Public/RHIResources.h b/Engine/Source/Runtime/RHI/Public/RHIResources.h
--- a/Engine/Source/Runtime/RHI/Public/RHIResources.h
+++ b/Engine/Source/Runtime/RHI/Public/RHIResources.h
@@ -770,64 +770,59 @@ public:
 //
 
 
-
 /* 
 * Generic GPU fence class used by FRHIGPUMemoryReadback
 * RHI specific fences derive from this to implement real GPU->CPU fencing.
 * The default implementation always returns false for Poll until the next frame from the frame the fence was inserted
 * because not all APIs have a GPU/CPU sync object, we need to fake it.
 */
-class FRHIGPUFence : public FRHIResource
+class RHI_API FRHIGPUFence : public FRHIResource
 {
 public:
-	FRHIGPUFence(FName InName)
-		: FenceName(InName)
-		, InsertedFrameNumber(0)
-	{}
-
-	virtual ~FRHIGPUFence()
-	{}
+	FRHIGPUFence(FName InName) : FenceName(InName) {}
+	virtual ~FRHIGPUFence() {}
 
-    /* Convenience function to write InsertedFrameNumber, used to emulate true GPU fences, RHI's should override RHIInsertGPUFence and implement their own fence mechanism. */
-	void Write()
-	{
-		InsertedFrameNumber = GFrameNumberRenderThread;
-	}
+    /**
+     * Signal this fence now, which completion can be tested with Poll or Wait.
+	 * RHI implementation might not need to implement this depending on how RHIInsertGPUFence is implemented.
+	 * Note also that RHIInsertGPUFence is currently only called by RHIEnqueueStagedRead.
+     */
+	virtual void Write();
 
     /**
-     * Poll the fence to see if the GPU has signalled it.
-     * @discussion RHI implementations must be thread-safe and must correctly handle being called before RHIInsertFence if an RHI thread is active.
+     * Poll the fence to see if the GPU has signaled it.
      * @returns True if and only if the GPU fence has been inserted and the GPU has signalled the fence.
      */
-	virtual bool Poll() const
-	{
-		if (GFrameNumberRenderThread > InsertedFrameNumber)
-		{
-			return true;
-		}
-		return false;
-	}
+	virtual bool Poll() const = 0;
 
     /**
      * Wait for the GPU to pass and signal the fence.
-     * @discussion RHI implementations must be thread-safe and must correctly handle being called before RHIInsertFence if an RHI thread is active.
      * @param TimeoutMs The maximum time to wait for the fence in milliseconds.
      * @returns True if and only if the GPU fence has been inserted and the GPU has signalled the fence.
      */
-	virtual bool Wait(float TimeoutMs) const
-	{
-		if (GFrameNumberRenderThread > InsertedFrameNumber)
-		{
-			return true;
-		}
-		return false;
-	}
+	virtual bool Wait(float TimeoutMs) const = 0;
 
 private:
 	FName FenceName;
-	uint32 InsertedFrameNumber;
 };
 
+// Generic implementation of FRHIGPUFence
+class RHI_API FGenericRHIGPUFence : public FRHIGPUFence
+{
+public:
+	FGenericRHIGPUFence(FName InName);
+
+	virtual void Write() final override;
+
+    /** @discussion RHI implementations must be thread-safe and must correctly handle being called before RHIInsertFence if an RHI thread is active. */
+	virtual bool Poll() const final override;
+
+    /** @discussion RHI implementations must be thread-safe and must correctly handle being called before RHIInsertFence if an RHI thread is active. */
+	virtual bool Wait(float TimeoutMs) const final override;
+
+private:
+	uint32 InsertedFrameNumber;
+};
 
 class FRHIRenderQuery : public FRHIResource {};
 
diff --git a/Engine/Source/Runtime/VulkanRHI/Private/VulkanRHI.cpp b/Engine/Source/Runtime/VulkanRHI/Private/VulkanRHI.cpp
--- a/Engine/Source/Runtime/VulkanRHI/Private/VulkanRHI.cpp
+++ b/Engine/Source/Runtime/VulkanRHI/Private/VulkanRHI.cpp
@@ -593,7 +593,6 @@ void FVulkanDynamicRHI::InitInstance()
 		GSupportsRenderTargetFormat_PF_G8 = false;	// #todo-rco
 		GRHISupportsTextureStreaming = true;
 		GSupportsTimestampRenderQueries = FVulkanPlatform::SupportsTimestampRenderQueries();
-		GRHISupportsGPUFence = true;
 #if VULKAN_ENABLE_DUMP_LAYER
 		// Disable RHI thread by default if the dump layer is enabled
 		GRHISupportsRHIThread = false;
