Merging from Dev-Niagara 4346597, 4346602, 4346709, 4353333, 4353344, 4353346, 4353355, 4355466, 4355948, 4356192

Niagara - New emitter and new system dialogs

diff --git a/Engine/Plugins/FX/Niagara/Source/Niagara/Classes/NiagaraEmitter.h b/Engine/Plugins/FX/Niagara/Source/Niagara/Classes/NiagaraEmitter.h
--- a/Engine/Plugins/FX/Niagara/Source/Niagara/Classes/NiagaraEmitter.h
+++ b/Engine/Plugins/FX/Niagara/Source/Niagara/Classes/NiagaraEmitter.h
@@ -287,6 +287,12 @@ public:
 	/** Internal: Indicates the thumbnail image is out of date.*/
 	UPROPERTY()
 	uint32 ThumbnailImageOutOfDate : 1;
+
+	UPROPERTY(EditAnywhere, AdvancedDisplay, Category = "Asset Options", AssetRegistrySearchable)
+	bool bIsTemplateAsset;
+
+	UPROPERTY(EditAnywhere, AdvancedDisplay, Category = "Asset Options", AssetRegistrySearchable)
+	FText TemplateAssetDescription;
 	
 	/** Callback issued whenever a VM compilation successfully happened (even if the results are a script that cannot be executed due to errors)*/
 	NIAGARA_API FOnEmitterCompiled& OnEmitterVMCompiled();
diff --git a/Engine/Plugins/FX/Niagara/Source/Niagara/Classes/NiagaraSystem.h b/Engine/Plugins/FX/Niagara/Source/Niagara/Classes/NiagaraSystem.h
--- a/Engine/Plugins/FX/Niagara/Source/Niagara/Classes/NiagaraSystem.h
+++ b/Engine/Plugins/FX/Niagara/Source/Niagara/Classes/NiagaraSystem.h
@@ -190,6 +190,12 @@ public:
 	UPROPERTY()
 	uint32 ThumbnailImageOutOfDate : 1;
 
+	UPROPERTY(EditAnywhere, AdvancedDisplay, Category = "Asset Options", AssetRegistrySearchable)
+	bool bIsTemplateAsset;
+
+	UPROPERTY(EditAnywhere, AdvancedDisplay, Category = "Asset Options", AssetRegistrySearchable)
+	FText TemplateAssetDescription;
+
 	bool GetIsolateEnabled() const;
 	void SetIsolateEnabled(bool bIsolate);
 #endif
diff --git a/Engine/Plugins/FX/Niagara/Source/NiagaraEditor/NiagaraEditor.Build.cs b/Engine/Plugins/FX/Niagara/Source/NiagaraEditor/NiagaraEditor.Build.cs
--- a/Engine/Plugins/FX/Niagara/Source/NiagaraEditor/NiagaraEditor.Build.cs
+++ b/Engine/Plugins/FX/Niagara/Source/NiagaraEditor/NiagaraEditor.Build.cs
@@ -46,6 +46,7 @@ public class NiagaraEditor : ModuleRules
                 "MovieSceneTracks",
                 "AdvancedPreviewScene",
 				"Projects",
+                "MainFrame",
 			}
         );
 
diff --git a/Engine/Plugins/FX/Niagara/Source/NiagaraEditor/Private/NiagaraEditorSettings.cpp b/Engine/Plugins/FX/Niagara/Source/NiagaraEditor/Private/NiagaraEditorSettings.cpp
--- a/Engine/Plugins/FX/Niagara/Source/NiagaraEditor/Private/NiagaraEditorSettings.cpp
+++ b/Engine/Plugins/FX/Niagara/Source/NiagaraEditor/Private/NiagaraEditorSettings.cpp
@@ -67,6 +67,22 @@ void UNiagaraEditorSettings::SetResimulateOnChangeWhilePaused(bool bInResimulate
 	}
 }
 
+FNiagaraNewAssetDialogConfig UNiagaraEditorSettings::GetNewAssetDailogConfig(FName InDialogConfigKey) const
+{
+	const FNiagaraNewAssetDialogConfig* Config = NewAssetDialogConfigMap.Find(InDialogConfigKey);
+	if (Config != nullptr)
+	{
+		return *Config;
+	}
+	return FNiagaraNewAssetDialogConfig();
+}
+
+void UNiagaraEditorSettings::SetNewAssetDialogConfig(FName InDialogConfigKey, const FNiagaraNewAssetDialogConfig& InNewAssetDialogConfig)
+{
+	NewAssetDialogConfigMap.Add(InDialogConfigKey, InNewAssetDialogConfig);
+	SaveConfig();
+}
+
 FName UNiagaraEditorSettings::GetCategoryName() const
 {
 	return TEXT("Plugins");
diff --git a/Engine/Plugins/FX/Niagara/Source/NiagaraEditor/Private/NiagaraEditorStyle.cpp b/Engine/Plugins/FX/Niagara/Source/NiagaraEditor/Private/NiagaraEditorStyle.cpp
--- a/Engine/Plugins/FX/Niagara/Source/NiagaraEditor/Private/NiagaraEditorStyle.cpp
+++ b/Engine/Plugins/FX/Niagara/Source/NiagaraEditor/Private/NiagaraEditorStyle.cpp
@@ -69,6 +69,53 @@ TSharedRef< FSlateStyleSet > FNiagaraEditorStyle::Create()
 
 	Style->Set("NiagaraEditor.StatsText", CategoryText);
 
+	// Asset picker
+	FTextBlockStyle AssetPickerAssetNameText = FTextBlockStyle(NormalText)
+		.SetColorAndOpacity(FLinearColor::White)
+		.SetFont(DEFAULT_FONT("Bold", 14))
+		.SetShadowOffset(FVector2D(0, 1))
+		.SetShadowColorAndOpacity(FLinearColor(0, 0, 0, 0.9f));
+
+	Style->Set("NiagaraEditor.AssetPickerAssetNameText", AssetPickerAssetNameText);
+
+	FTextBlockStyle AssetPickerAssetCategoryText = FTextBlockStyle(NormalText)
+		.SetFont(DEFAULT_FONT("Regular", 10));
+
+	Style->Set("NiagaraEditor.AssetPickerAssetCategoryText", AssetPickerAssetCategoryText);
+
+	// New Asset Dialog
+	FTextBlockStyle NewAssetDialogOptionText = FTextBlockStyle(NormalText)
+		.SetFont(DEFAULT_FONT("Bold", 11));
+
+	Style->Set("NiagaraEditor.NewAssetDialog.OptionText", NewAssetDialogOptionText);
+
+	FTextBlockStyle NewAssetDialogHeaderText = FTextBlockStyle(NormalText)
+		.SetColorAndOpacity(FLinearColor::White)
+		.SetFont(DEFAULT_FONT("Bold", 10))
+		.SetShadowOffset(FVector2D(0, 1))
+		.SetShadowColorAndOpacity(FLinearColor(0, 0, 0, 0.9f));
+
+	Style->Set("NiagaraEditor.NewAssetDialog.HeaderText", NewAssetDialogHeaderText);
+
+	FTextBlockStyle NewAssetDialogSubHeaderText = FTextBlockStyle(NormalText)
+		.SetColorAndOpacity(FLinearColor::White)
+		.SetFont(DEFAULT_FONT("Bold", 10))
+		.SetShadowOffset(FVector2D(0, 1))
+		.SetShadowColorAndOpacity(FLinearColor(0, 0, 0, 0.9f));
+
+	Style->Set("NiagaraEditor.NewAssetDialog.SubHeaderText", NewAssetDialogSubHeaderText);
+
+	Style->Set("NiagaraEditor.NewAssetDialog.AddButton", FButtonStyle()
+		.SetNormal(BOX_CORE_BRUSH("Common/FlatButton", 2.0f / 8.0f, FLinearColor(0, 0, 0, .25f)))
+		.SetHovered(BOX_CORE_BRUSH("Common/FlatButton", 2.0f / 8.0f, FEditorStyle::GetSlateColor("SelectionColor")))
+		.SetPressed(BOX_CORE_BRUSH("Common/FlatButton", 2.0f / 8.0f, FEditorStyle::GetSlateColor("SelectionColor_Pressed")))
+	);
+
+	Style->Set("NiagaraEditor.NewAssetDialog.SubBorderColor", FLinearColor(FColor(48, 48, 48)));
+	Style->Set("NiagaraEditor.NewAssetDialog.ActiveOptionBorderColor", FLinearColor(FColor(96, 96, 96)));
+
+	Style->Set("NiagaraEditor.NewAssetDialog.SubBorder", new BOX_CORE_BRUSH("Common/GroupBorderLight", FMargin(4.0f / 16.0f)));
+
 	// Emitter Header
 	FTextBlockStyle HeadingText = FTextBlockStyle(NormalText)
 		.SetFont(DEFAULT_FONT("Regular", 14))
diff --git a/Engine/Plugins/FX/Niagara/Source/NiagaraEditor/Private/NiagaraEmitterFactoryNew.cpp b/Engine/Plugins/FX/Niagara/Source/NiagaraEditor/Private/NiagaraEmitterFactoryNew.cpp
--- a/Engine/Plugins/FX/Niagara/Source/NiagaraEditor/Private/NiagaraEmitterFactoryNew.cpp
+++ b/Engine/Plugins/FX/Niagara/Source/NiagaraEditor/Private/NiagaraEmitterFactoryNew.cpp
@@ -13,6 +13,11 @@
 #include "ViewModels/Stack/NiagaraStackGraphUtilities.h"
 #include "NiagaraConstants.h"
 #include "NiagaraNodeAssignment.h"
+#include "SNewEmitterDialog.h"
+#include "Misc/MessageDialog.h"
+#include "Modules/ModuleManager.h"
+#include "Interfaces/IMainFrameModule.h"
+#include "Framework/Application/SlateApplication.h"
 
 #include "Misc/ConfigCacheIni.h"
 
@@ -25,6 +30,83 @@ UNiagaraEmitterFactoryNew::UNiagaraEmitterFactoryNew(const FObjectInitializer& O
 	bCreateNew = false;
 	bEditAfterNew = true;
 	bCreateNew = true;
+	EmitterToCopy = nullptr;
+}
+
+bool UNiagaraEmitterFactoryNew::ConfigureProperties()
+{
+	IMainFrameModule& MainFrame = FModuleManager::LoadModuleChecked<IMainFrameModule>("MainFrame");
+	TSharedPtr<SWindow>	ParentWindow = MainFrame.GetParentWindow();
+
+	TSharedRef<SNewEmitterDialog> NewEmitterDialog = SNew(SNewEmitterDialog);
+	FSlateApplication::Get().AddModalWindow(NewEmitterDialog, ParentWindow);
+
+	if (NewEmitterDialog->GetUserConfirmedSelection() == false)
+	{
+		// User cancelled or closed the dialog so abort asset creation.
+		return false;
+	}
+
+	TOptional<FAssetData> SelectedEmitterAsset = NewEmitterDialog->GetSelectedEmitterAsset();
+	if (SelectedEmitterAsset.IsSet())
+	{
+		EmitterToCopy = Cast<UNiagaraEmitter>(SelectedEmitterAsset->GetAsset());
+		if (EmitterToCopy == nullptr)
+		{
+			FText Title = LOCTEXT("FailedToLoadTitle", "Create Default?");
+			EAppReturnType::Type DialogResult = FMessageDialog::Open(EAppMsgType::OkCancel, EAppReturnType::Cancel,
+				LOCTEXT("FailedToLoadMessage", "The selected emitter failed to load\nWould you like to create a default emitter?"),
+				&Title);
+			if (DialogResult == EAppReturnType::Cancel)
+			{
+				return false;
+			}
+			else
+			{
+				// The selected emitter couldn't be loaded but the user still wants to create a default emitter so 
+				// well need to create a new empty emitter and add some default modules to it.
+				EmitterToCopy = nullptr;
+				bAddDefaultModulesAndRenderersToEmptyEmitter = true;
+			}
+		}
+	}
+	else
+	{
+		// User selected an empty emitter so set the emitter to copy to null, and disable creationg of default modules.
+		EmitterToCopy = nullptr;
+		bAddDefaultModulesAndRenderersToEmptyEmitter = false;
+	}
+
+	return true;
+}
+
+UNiagaraNodeFunctionCall* AddModuleFromAssetPath(FString AssetPath, UNiagaraNodeOutput& TargetOutputNode)
+{
+	FStringAssetReference AssetRef(AssetPath);
+	UNiagaraScript* AssetScript = Cast<UNiagaraScript>(AssetRef.TryLoad());
+	FAssetData ScriptAssetData(AssetScript);
+	if (ScriptAssetData.IsValid())
+	{
+		return FNiagaraStackGraphUtilities::AddScriptModuleToStack(ScriptAssetData, TargetOutputNode);
+	}
+	else
+	{
+		UE_LOG(LogNiagaraEditor, Error, TEXT("Failed to create default modules for emitter.  Missing %s"), *AssetRef.ToString());
+		return nullptr;
+	}
+}
+
+template<typename ValueType>
+void SetRapidIterationParameter(FString UniqueEmitterName, UNiagaraScript& TargetScript, UNiagaraNodeFunctionCall& TargetFunctionCallNode,
+	FName InputName, FNiagaraTypeDefinition InputType, ValueType Value)
+{
+	FNiagaraParameterHandle InputHandle = FNiagaraParameterHandle::CreateModuleParameterHandle(InputName);
+	FNiagaraParameterHandle AliasedInputHandle = FNiagaraParameterHandle::CreateAliasedModuleParameterHandle(InputHandle, &TargetFunctionCallNode);
+	FNiagaraVariable RapidIterationParameter = FNiagaraStackGraphUtilities::CreateRapidIterationParameter(UniqueEmitterName, TargetScript.GetUsage(),
+		AliasedInputHandle.GetParameterHandleString(), InputType);
+	RapidIterationParameter.SetValue(Value);
+	bool bAddParameterIfMissing = true;
+	TargetScript.RapidIterationParameters.SetParameterData(RapidIterationParameter.GetData(), RapidIterationParameter, bAddParameterIfMissing);
 }
 
 UObject* UNiagaraEmitterFactoryNew::FactoryCreateNew(UClass* Class, UObject* InParent, FName Name, EObjectFlags Flags, UObject* Context, FFeedbackContext* Warn)
@@ -36,185 +118,85 @@ UObject* UNiagaraEmitterFactoryNew::FactoryCreateNew(UClass* Class, UObject* InP
 
 	UNiagaraEmitter* NewEmitter;
 
-	if (UNiagaraEmitter* Default = Cast<UNiagaraEmitter>(Settings->DefaultEmitter.TryLoad()))
+	if (EmitterToCopy != nullptr)
 	{
-		NewEmitter = Cast<UNiagaraEmitter>(StaticDuplicateObject(Default, InParent, Name, Flags, Class));
+		NewEmitter = Cast<UNiagaraEmitter>(StaticDuplicateObject(EmitterToCopy, InParent, Name, Flags, Class));
+		NewEmitter->bIsTemplateAsset = false;
+		NewEmitter->TemplateAssetDescription = FText();
 	}
 	else
 	{
-		UE_LOG(LogNiagaraEditor, Log, TEXT("Default Emitter \"%s\" could not be loaded. Creating graph procedurally."), *Settings->DefaultEmitter.ToString());
-
+		// Create an empty emitter, source, and graph.
 		NewEmitter = NewObject<UNiagaraEmitter>(InParent, Class, Name, Flags | RF_Transactional);
-
-		NewEmitter->AddRenderer(NewObject<UNiagaraSpriteRendererProperties>(NewEmitter, "Renderer"));
+		NewEmitter->SimTarget = ENiagaraSimTarget::CPUSim;
 
 		UNiagaraScriptSource* Source = NewObject<UNiagaraScriptSource>(NewEmitter, NAME_None, RF_Transactional);
-		if (Source)
+		UNiagaraGraph* CreatedGraph = NewObject<UNiagaraGraph>(Source, NAME_None, RF_Transactional);
+		Source->NodeGraph = CreatedGraph;
+
+		// Fix up source pointers.
+		NewEmitter->GraphSource = Source;
+		NewEmitter->SpawnScriptProps.Script->SetSource(Source);
+		NewEmitter->UpdateScriptProps.Script->SetSource(Source);
+		NewEmitter->EmitterSpawnScriptProps.Script->SetSource(Source);
+		NewEmitter->EmitterUpdateScriptProps.Script->SetSource(Source);
+		NewEmitter->GetGPUComputeScript()->SetSource(Source);
+
+		// Initialize the scripts for output.
+		UNiagaraNodeOutput* EmitterSpawnOutputNode = FNiagaraStackGraphUtilities::ResetGraphForOutput(*Source->NodeGraph, ENiagaraScriptUsage::EmitterSpawnScript, NewEmitter->EmitterSpawnScriptProps.Script->GetUsageId());
+		UNiagaraNodeOutput* EmitterUpdateOutputNode = FNiagaraStackGraphUtilities::ResetGraphForOutput(*Source->NodeGraph, ENiagaraScriptUsage::EmitterUpdateScript, NewEmitter->EmitterUpdateScriptProps.Script->GetUsageId());
+		UNiagaraNodeOutput* ParticleSpawnOutputNode = FNiagaraStackGraphUtilities::ResetGraphForOutput(*Source->NodeGraph, ENiagaraScriptUsage::ParticleSpawnScript, NewEmitter->SpawnScriptProps.Script->GetUsageId());
+		UNiagaraNodeOutput* ParticleUpdateOutputNode = FNiagaraStackGraphUtilities::ResetGraphForOutput(*Source->NodeGraph, ENiagaraScriptUsage::ParticleUpdateScript, NewEmitter->UpdateScriptProps.Script->GetUsageId());
+
+		checkf(EmitterSpawnOutputNode != nullptr && EmitterUpdateOutputNode != nullptr && ParticleSpawnOutputNode != nullptr && ParticleUpdateOutputNode != nullptr,
+			TEXT("Failed to create output nodes for emitter scripts."));
+
+		if (bAddDefaultModulesAndRenderersToEmptyEmitter)
 		{
-			UNiagaraGraph* CreatedGraph = NewObject<UNiagaraGraph>(Source, NAME_None, RF_Transactional);
-			Source->NodeGraph = CreatedGraph;
-
-			// Set pointer in script to source
-			NewEmitter->GraphSource = Source;
-			NewEmitter->SpawnScriptProps.Script->SetSource(Source);
-			NewEmitter->UpdateScriptProps.Script->SetSource(Source);
-			NewEmitter->EmitterSpawnScriptProps.Script->SetSource(Source);
-			NewEmitter->EmitterUpdateScriptProps.Script->SetSource(Source);
-			NewEmitter->GetGPUComputeScript()->SetSource(Source);
-			NewEmitter->SimTarget = ENiagaraSimTarget::CPUSim;
-
-			bool bCreateDefaultNodes = true;
-			if (bCreateDefaultNodes)
+			NewEmitter->AddRenderer(NewObject<UNiagaraSpriteRendererProperties>(NewEmitter, "Renderer"));
+
+			AddModuleFromAssetPath(TEXT("/Niagara/Modules/Emitter/EmitterLifeCycle.EmitterLifeCycle"), *EmitterUpdateOutputNode);
+			
+			UNiagaraNodeFunctionCall* SpawnRateNode = AddModuleFromAssetPath(TEXT("/Niagara/Modules/Emitter/SpawnRate.SpawnRate"), *EmitterUpdateOutputNode);
+			if (SpawnRateNode != nullptr)
 			{
-				if (Source)
-				{
-					UNiagaraNodeOutput* EmitterSpawnOutputNode = FNiagaraStackGraphUtilities::ResetGraphForOutput(*Source->NodeGraph, ENiagaraScriptUsage::EmitterSpawnScript, NewEmitter->EmitterSpawnScriptProps.Script->GetUsageId());
-					UNiagaraNodeOutput* EmitterUpdateOutputNode = FNiagaraStackGraphUtilities::ResetGraphForOutput(*Source->NodeGraph, ENiagaraScriptUsage::EmitterUpdateScript, NewEmitter->EmitterUpdateScriptProps.Script->GetUsageId());
-					UNiagaraNodeOutput* ParticleSpawnOutputNode = FNiagaraStackGraphUtilities::ResetGraphForOutput(*Source->NodeGraph, ENiagaraScriptUsage::ParticleSpawnScript, NewEmitter->SpawnScriptProps.Script->GetUsageId());
-					UNiagaraNodeOutput* ParticleUpdateOutputNode = FNiagaraStackGraphUtilities::ResetGraphForOutput(*Source->NodeGraph, ENiagaraScriptUsage::ParticleUpdateScript, NewEmitter->UpdateScriptProps.Script->GetUsageId());
-
-					{
-						FStringAssetReference EmitterUpdateScriptRef(TEXT("/Niagara/Modules/Emitter/EmitterLifeCycle.EmitterLifeCycle"));
-						UNiagaraScript* EmitterUpdateScript = Cast<UNiagaraScript>(EmitterUpdateScriptRef.TryLoad());
-						FAssetData EmitterUpdateModuleScriptAsset(EmitterUpdateScript);
-						if (EmitterUpdateOutputNode && EmitterUpdateModuleScriptAsset.IsValid())
-						{
-							FNiagaraStackGraphUtilities::AddScriptModuleToStack(EmitterUpdateModuleScriptAsset, *EmitterUpdateOutputNode);
-						}
-						else
-						{
-							UE_LOG(LogNiagaraEditor, Error, TEXT("Missing %s"), *EmitterUpdateScriptRef.ToString());
-						}
-					}
-
-					{
-						FStringAssetReference EmitterUpdateScriptRef(TEXT("/Niagara/Modules/Emitter/SpawnRate.SpawnRate"));
-						UNiagaraScript* EmitterUpdateScript = Cast<UNiagaraScript>(EmitterUpdateScriptRef.TryLoad());
-						FAssetData EmitterUpdateModuleScriptAsset(EmitterUpdateScript);
-						if (EmitterUpdateOutputNode && EmitterUpdateModuleScriptAsset.IsValid())
-						{
-							FNiagaraStackGraphUtilities::AddScriptModuleToStack(EmitterUpdateModuleScriptAsset, *EmitterUpdateOutputNode);
-						}
-						else
-						{
-							UE_LOG(LogNiagaraEditor, Error, TEXT("Missing %s"), *EmitterUpdateScriptRef.ToString());
-						}
-					}
-
-					{
-						FStringAssetReference SpawnScriptRef(TEXT("/Niagara/Modules/Spawn/Location/SystemLocation.SystemLocation"));
-						UNiagaraScript* SpawnScript = Cast<UNiagaraScript>(SpawnScriptRef.TryLoad());
-						FAssetData SpawnModuleScriptAsset(SpawnScript);
-						if (ParticleSpawnOutputNode && SpawnModuleScriptAsset.IsValid())
-						{
-							FNiagaraStackGraphUtilities::AddScriptModuleToStack(SpawnModuleScriptAsset, *ParticleSpawnOutputNode);
-						}
-						else
-						{
-							UE_LOG(LogNiagaraEditor, Error, TEXT("Missing %s"), *SpawnScriptRef.ToString());
-						}
-					}
-
-					{
-						FStringAssetReference SpawnScriptRef(TEXT("/Niagara/Modules/Spawn/Velocity/AddVelocity.AddVelocity"));
-						UNiagaraScript* SpawnScript = Cast<UNiagaraScript>(SpawnScriptRef.TryLoad());
-						FAssetData SpawnModuleScriptAsset(SpawnScript);
-						if (ParticleSpawnOutputNode && SpawnModuleScriptAsset.IsValid())
-						{
-							UNiagaraNodeFunctionCall* CallNode = FNiagaraStackGraphUtilities::AddScriptModuleToStack(SpawnModuleScriptAsset, *ParticleSpawnOutputNode);
-							if (CallNode)
-							{
-								FNiagaraVariable VelocityVar(FNiagaraTypeDefinition::GetVec3Def(), *(TEXT("Constants.Emitter.") + CallNode->GetFunctionName() + TEXT(".Velocity")));
-								VelocityVar.SetValue(FVector(0.0f, 0.0f, 100.0f));
-								bool bAddParameterIfMissing = true;
-								NewEmitter->SpawnScriptProps.Script->RapidIterationParameters.SetParameterData(VelocityVar.GetData(), VelocityVar, bAddParameterIfMissing);
-							}
-						}
-						else
-						{
-							UE_LOG(LogNiagaraEditor, Error, TEXT("Missing %s"), *SpawnScriptRef.ToString());
-						}
-
-					}
-
-					{
-						if (ParticleSpawnOutputNode)
-						{
-							TArray<FNiagaraVariable> Vars;
-							TArray<FString> Defaults;
-							{
-								FNiagaraVariable Var = SYS_PARAM_PARTICLES_SPRITE_SIZE;
-								FString DefaultValue = FNiagaraConstants::GetAttributeDefaultValue(Var);
-								Vars.Add(Var);
-								Defaults.Add(DefaultValue);
-							}
-
-							{
-								FNiagaraVariable Var = SYS_PARAM_PARTICLES_SPRITE_ROTATION;
-								FString DefaultValue = FNiagaraConstants::GetAttributeDefaultValue(Var);
-								Vars.Add(Var);
-								Defaults.Add(DefaultValue);
-							}
-
-							{
-								FNiagaraVariable Var = SYS_PARAM_PARTICLES_LIFETIME;
-								FString DefaultValue = FNiagaraConstants::GetAttributeDefaultValue(Var);
-								Vars.Add(Var);
-								Defaults.Add(DefaultValue);
-							}
-							FNiagaraStackGraphUtilities::AddParameterModuleToStack(Vars, *ParticleSpawnOutputNode, INDEX_NONE, Defaults);
-
-						}
-					}
-					
-					{
-						FStringAssetReference UpdateScriptRef(TEXT("/Niagara/Modules/Update/Lifetime/UpdateAge.UpdateAge"));
-						UNiagaraScript* UpdateScript = Cast<UNiagaraScript>(UpdateScriptRef.TryLoad());
-						FAssetData UpdateModuleScriptAsset(UpdateScript);
-						if (ParticleUpdateOutputNode && UpdateModuleScriptAsset.IsValid())
-						{
-							FNiagaraStackGraphUtilities::AddScriptModuleToStack(UpdateModuleScriptAsset, *ParticleUpdateOutputNode);
-						}
-						else
-						{
-							UE_LOG(LogNiagaraEditor, Error, TEXT("Missing %s"), *UpdateScriptRef.ToString());
-						}
-					}
-
-					{
-						FStringAssetReference UpdateScriptRef(TEXT("/Niagara/Modules/Update/Color/Color.Color"));
-						UNiagaraScript* UpdateScript = Cast<UNiagaraScript>(UpdateScriptRef.TryLoad());
-						FAssetData UpdateModuleScriptAsset(UpdateScript);
-						if (ParticleUpdateOutputNode && UpdateModuleScriptAsset.IsValid())
-						{
-							FNiagaraStackGraphUtilities::AddScriptModuleToStack(UpdateModuleScriptAsset, *ParticleUpdateOutputNode);
-						}
-						else
-						{
-							UE_LOG(LogNiagaraEditor, Error, TEXT("Missing %s"), *UpdateScriptRef.ToString());
-						}
-					}
-
-					{
-						FStringAssetReference UpdateScriptRef(TEXT("/Niagara/Modules/Solvers/SolveForcesAndVelocity.SolveForcesAndVelocity"));
-						UNiagaraScript* UpdateScript = Cast<UNiagaraScript>(UpdateScriptRef.TryLoad());
-						FAssetData UpdateModuleScriptAsset(UpdateScript);
-						if (ParticleUpdateOutputNode && UpdateModuleScriptAsset.IsValid())
-						{
-							FNiagaraStackGraphUtilities::AddScriptModuleToStack(UpdateModuleScriptAsset, *ParticleUpdateOutputNode);
-						}
-						else
-						{
-							UE_LOG(LogNiagaraEditor, Error, TEXT("Missing %s"), *UpdateScriptRef.ToString());
-						}
-					}
-
-					FNiagaraStackGraphUtilities::RelayoutGraph(*Source->NodeGraph);
-					NewEmitter->bInterpolatedSpawning = true;
-					NewEmitter->SpawnScriptProps.Script->SetUsage(ENiagaraScriptUsage::ParticleSpawnScriptInterpolated);
-				}
+				SetRapidIterationParameter(NewEmitter->GetUniqueEmitterName(), *NewEmitter->EmitterUpdateScriptProps.Script, *SpawnRateNode,
+					"SpawnRate", FNiagaraTypeDefinition::GetFloatDef(), 10.0f);
 			}
+
+			AddModuleFromAssetPath(TEXT("/Niagara/Modules/Spawn/Location/SystemLocation.SystemLocation"), *ParticleSpawnOutputNode);
+
+			UNiagaraNodeFunctionCall* AddVelocityNode = AddModuleFromAssetPath(TEXT("/Niagara/Modules/Spawn/Velocity/AddVelocity.AddVelocity"), *ParticleSpawnOutputNode);
+			if (AddVelocityNode != nullptr)
+			{
+				SetRapidIterationParameter(NewEmitter->GetUniqueEmitterName(), *NewEmitter->SpawnScriptProps.Script, *AddVelocityNode,
+					"Velocity", FNiagaraTypeDefinition::GetVec3Def(), FVector(0.0f, 0.0f, 100.0f));
+			}
+
+			TArray<FNiagaraVariable> Vars =
+			{
+				SYS_PARAM_PARTICLES_SPRITE_SIZE,
+				SYS_PARAM_PARTICLES_SPRITE_ROTATION,
+				SYS_PARAM_PARTICLES_LIFETIME
+			};
+
+			TArray<FString> Defaults = 
+			{
+				FNiagaraConstants::GetAttributeDefaultValue(SYS_PARAM_PARTICLES_SPRITE_SIZE),
+				FNiagaraConstants::GetAttributeDefaultValue(SYS_PARAM_PARTICLES_SPRITE_ROTATION),
+				FNiagaraConstants::GetAttributeDefaultValue(SYS_PARAM_PARTICLES_LIFETIME)
+			};
+
+			FNiagaraStackGraphUtilities::AddParameterModuleToStack(Vars, *ParticleSpawnOutputNode, INDEX_NONE, Defaults);
+
+			AddModuleFromAssetPath(TEXT("/Niagara/Modules/Update/Lifetime/UpdateAge.UpdateAge"), *ParticleUpdateOutputNode);
+			AddModuleFromAssetPath(TEXT("/Niagara/Modules/Update/Color/Color.Color"), *ParticleUpdateOutputNode);
+			AddModuleFromAssetPath(TEXT("/Niagara/Modules/Solvers/SolveForcesAndVelocity.SolveForcesAndVelocity"), *ParticleUpdateOutputNode);
 		}
+
+		FNiagaraStackGraphUtilities::RelayoutGraph(*Source->NodeGraph);
+		NewEmitter->bInterpolatedSpawning = true;
+		NewEmitter->SpawnScriptProps.Script->SetUsage(ENiagaraScriptUsage::ParticleSpawnScriptInterpolated);
 	}
 	
 	return NewEmitter;
diff --git a/Engine/Plugins/FX/Niagara/Source/NiagaraEditor/Private/NiagaraSystemFactoryNew.cpp b/Engine/Plugins/FX/Niagara/Source/NiagaraEditor/Private/NiagaraSystemFactoryNew.cpp
--- a/Engine/Plugins/FX/Niagara/Source/NiagaraEditor/Private/NiagaraSystemFactoryNew.cpp
+++ b/Engine/Plugins/FX/Niagara/Source/NiagaraEditor/Private/NiagaraSystemFactoryNew.cpp
@@ -9,19 +9,102 @@
 #include "NiagaraEditorSettings.h"
 #include "AssetData.h"
 #include "ViewModels/Stack/NiagaraStackGraphUtilities.h"
+#include "SNewSystemDialog.h"
+#include "Misc/MessageDialog.h"
+#include "ViewModels/NiagaraSystemViewModel.h"
+#include "Modules/ModuleManager.h"
+#include "Interfaces/IMainFrameModule.h"
+#include "Framework/Application/SlateApplication.h"
 
 #define LOCTEXT_NAMESPACE "NiagaraSystemFactory"
 
 UNiagaraSystemFactoryNew::UNiagaraSystemFactoryNew(const FObjectInitializer& ObjectInitializer)
 : Super(ObjectInitializer)
 {
-
 	SupportedClass = UNiagaraSystem::StaticClass();
 	bCreateNew = false;
 	bEditAfterNew = true;
 	bCreateNew = true;
 }
 
+bool UNiagaraSystemFactoryNew::ConfigureProperties()
+{
+	IMainFrameModule& MainFrame = FModuleManager::LoadModuleChecked<IMainFrameModule>("MainFrame");
+	TSharedPtr<SWindow>	ParentWindow = MainFrame.GetParentWindow();
+
+	TSharedRef<SNewSystemDialog> NewSystemDialog = SNew(SNewSystemDialog);
+	FSlateApplication::Get().AddModalWindow(NewSystemDialog, ParentWindow);
+
+	if (NewSystemDialog->GetUserConfirmedSelection() == false)
+	{
+		// User cancelled or closed the dialog so abort asset creation.
+		return false;
+	}
+
+	TOptional<FAssetData> SelectedSystemAsset = NewSystemDialog->GetSelectedSystemAsset();
+	TArray<FAssetData> EmitterAssetsToAddToNewSystem = NewSystemDialog->GetSelectedEmitterAssets();
+	if (SelectedSystemAsset.IsSet())
+	{
+		SystemToCopy = Cast<UNiagaraSystem>(SelectedSystemAsset.GetValue().GetAsset());
+		if (SystemToCopy == nullptr)
+		{
+			FText Title = LOCTEXT("FailedToLoadSystemTitle", "Create Default?");
+			EAppReturnType::Type DialogResult = FMessageDialog::Open(EAppMsgType::OkCancel, EAppReturnType::Cancel,
+				LOCTEXT("FailedToLoadEmitterMessage", "The selected system failed to load.\nWould you like to create an empty system?"),
+				&Title);
+			if (DialogResult == EAppReturnType::Cancel)
+			{
+				return false;
+			}
+			else
+			{
+				SystemToCopy = nullptr;
+				EmitterAssetsToAddToNewSystem.Empty();
+			}
+		}
+	}
+	else if (EmitterAssetsToAddToNewSystem.Num() > 0)
+	{
+		bool bAllEmittersLoaded = true;
+		for (const FAssetData& EmitterAssetToAdd : EmitterAssetsToAddToNewSystem)
+		{
+			UNiagaraEmitter* EmitterToAdd = Cast<UNiagaraEmitter>(EmitterAssetToAdd.GetAsset());
+			if (EmitterToAdd != nullptr)
+			{
+				EmittersToAddToNewSystem.Add(EmitterToAdd);
+			}
+			else
+			{
+				bAllEmittersLoaded = false;
+				break;
+			}
+		}
+
+		if(bAllEmittersLoaded == false)
+		{
+			FText Title = LOCTEXT("FailedToLoadEmitterTitle", "Create Default?");
+			EAppReturnType::Type DialogResult = FMessageDialog::Open(EAppMsgType::OkCancel, EAppReturnType::Cancel,
+				LOCTEXT("FailedToLoadMessage", "A selected emitter failed to load.\nWould you like to create an empty system system?"),
+				&Title);
+			if (DialogResult == EAppReturnType::Cancel)
+			{
+				return false;
+			}
+			else
+			{
+				EmittersToAddToNewSystem.Empty();
+			}
+		}
+	}
+	else
+	{
+		SystemToCopy = nullptr;
+		EmittersToAddToNewSystem.Empty();
+	}
+
+	return true;
+}
+
 UObject* UNiagaraSystemFactoryNew::FactoryCreateNew(UClass* Class, UObject* InParent, FName Name, EObjectFlags Flags, UObject* Context, FFeedbackContext* Warn)
 {
 	check(Class->IsChildOf(UNiagaraSystem::StaticClass()));
@@ -31,10 +114,27 @@ UObject* UNiagaraSystemFactoryNew::FactoryCreateNew(UClass* Class, UObject* InPa
 
 	UNiagaraSystem* NewSystem;
 	
-	if (UNiagaraSystem* Default = Cast<UNiagaraSystem>(Settings->DefaultSystem.TryLoad()))
+	if (SystemToCopy != nullptr)
 	{
-		NewSystem = Cast<UNiagaraSystem>(StaticDuplicateObject(Default, InParent, Name, Flags, Class));
-		InitializeSystem(NewSystem, false);
+		NewSystem = Cast<UNiagaraSystem>(StaticDuplicateObject(SystemToCopy, InParent, Name, Flags, Class));
+		NewSystem->bIsTemplateAsset = false;
+		NewSystem->TemplateAssetDescription = FText();
+	}
+	else if (EmittersToAddToNewSystem.Num() > 0)
+	{
+		NewSystem = NewObject<UNiagaraSystem>(InParent, Class, Name, Flags | RF_Transactional);
+		InitializeSystem(NewSystem, true);
+
+		FNiagaraSystemViewModelOptions SystemViewModelOptions;
+		SystemViewModelOptions.bCanAutoCompile = false;
+		SystemViewModelOptions.bCanSimulate = false;
+		SystemViewModelOptions.EditMode = ENiagaraSystemViewModelEditMode::SystemAsset;
+
+		TSharedRef<FNiagaraSystemViewModel> NewSystemViewModel = MakeShared<FNiagaraSystemViewModel>(*NewSystem, SystemViewModelOptions);
+		for (UNiagaraEmitter* EmitterToAddToNewSystem : EmittersToAddToNewSystem)
+		{
+			NewSystemViewModel->AddEmitter(*EmitterToAddToNewSystem);
+		}
 	}
 	else
 	{
@@ -42,7 +142,6 @@ UObject* UNiagaraSystemFactoryNew::FactoryCreateNew(UClass* Class, UObject* InPa
 		InitializeSystem(NewSystem, true);
 	}
 
-
 	return NewSystem;
 }
 
diff --git a/Engine/Plugins/FX/Niagara/Source/NiagaraEditor/Private/SItemSelector.h b/Engine/Plugins/FX/Niagara/Source/NiagaraEditor/Private/SItemSelector.h
new file mode 100644
--- /dev/null
+++ b/Engine/Plugins/FX/Niagara/Source/NiagaraEditor/Private/SItemSelector.h
@@ -0,0 +1,518 @@
+// Copyright 1998-2018 Epic Games, Inc. All Rights Reserved.
+
+#pragma once
+
+#include "Widgets/Views/STableRow.h"
+#include "Widgets/Input/SSearchBox.h"
+#include "Widgets/Views/STreeView.h"
+
+/** A generic widget for selecting an item from an array of items including optional filtering and categorization. */
+template<typename CategoryType, typename ItemType>
+class SItemSelector : public SCompoundWidget
+{
+public:
+	DECLARE_DELEGATE_RetVal_OneParam(TArray<CategoryType>, FOnGetCategoriesForItem, const ItemType& /* Item */);
+	DECLARE_DELEGATE_RetVal_TwoParams(bool, FOnCompareCategoriesForEquality, const CategoryType& /* CategoryA */, const CategoryType& /* CategoryB */);
+	DECLARE_DELEGATE_RetVal_TwoParams(bool, FOnCompareCategoriesForSorting, const CategoryType& /* CategoryA */, const CategoryType& /* CateogoryB */);
+	DECLARE_DELEGATE_RetVal_TwoParams(bool, FOnCompareItemsForSorting, const ItemType& /* ItemA */, const ItemType& /* ItemB */);
+	DECLARE_DELEGATE_RetVal_TwoParams(bool, FOnDoesItemMatchFilterText, const FText& /* Filter text */, const ItemType& /* Item */);
+	DECLARE_DELEGATE_RetVal_OneParam(TSharedRef<SWidget>, FOnGenerateWidgetForCategory, const CategoryType& /* Category */);
+	DECLARE_DELEGATE_RetVal_OneParam(TSharedRef<SWidget>, FOnGenerateWidgetForItem, const ItemType& /* Item */);
+
+public:
+	SLATE_BEGIN_ARGS(SItemSelector)
+		: _AllowMultiselect(false)
+	{}
+
+	SLATE_ARGUMENT(TArray<ItemType>, Items)
+		/** Whether or not this item selector should allow multiple items to be selected. */
+		SLATE_ARGUMENT(bool, AllowMultiselect)
+
+		/** An optional delegate to get an array of categories for the specified item. Each category in the returned array represents one level of nested categories. 
+		NOTE: The OnCompareCategoriesForEquality, and OnGenerateWidgetForCategory delegates must be bound if this delegate is bound. */
+		SLATE_EVENT(FOnGetCategoriesForItem, OnGetCategoriesForItem)
+
+		/** An optional delegate to compare two categories for equality which must be supplied when generating categories for items.  This equality comparer will be used to collate items into matching categories. */
+		SLATE_EVENT(FOnCompareCategoriesForEquality, OnCompareCategoriesForEquality)
+
+		/** An optional delegate which determines the sorting for categories.  If not bound categories will be ordered by the order they're encountered while processing items. */
+		SLATE_EVENT(FOnCompareCategoriesForSorting, OnCompareCategoriesForSorting)
+
+		/** An optional delegate which determines the sorting for items within each category. */
+		SLATE_EVENT(FOnCompareItemsForSorting, OnCompareItemsForSorting)
+
+		/** An optional delegate which can be used to filter items available for selection.  If not bound the search box will not be shown. */
+		SLATE_EVENT(FOnDoesItemMatchFilterText, OnDoesItemMatchFilterText)
+
+		/** An optional delegate which generates widgets for categories which must be bound when generating categories for items. */
+		SLATE_EVENT(FOnGenerateWidgetForCategory, OnGenerateWidgetForCategory)
+
+		/** The delegate which is used to generate widgets for the items to be selected. */
+		SLATE_EVENT(FOnGenerateWidgetForItem, OnGenerateWidgetForItem)
+	SLATE_END_ARGS();
+
+private:
+	enum class EItemSelectorItemViewModelType
+	{
+		Category,
+		Item
+	};
+
+	class IItemSelectorItemViewModelUtilities
+	{
+	public:
+		virtual ~IItemSelectorItemViewModelUtilities()
+		{
+		}
+
+		virtual bool DoesItemMatchFilterText(const ItemType& InItem) const = 0;
+		virtual bool CompareCategoriesForEquality(const CategoryType& CategoryA, const CategoryType& CategoryB) const = 0;
+		virtual const FOnCompareCategoriesForSorting& GetOnCompareCategoriesForSorting() const = 0;
+		virtual const FOnCompareItemsForSorting& GetOnCompareItemsForSorting() const = 0;
+	};
+
+	class FItemSelectorItemViewModel
+	{
+	public:
+		FItemSelectorItemViewModel(TSharedRef<IItemSelectorItemViewModelUtilities> InItemUtilities, EItemSelectorItemViewModelType InType)
+			: ItemUtilities(InItemUtilities)
+			, Type(InType)
+		{
+		}
+
+		virtual ~FItemSelectorItemViewModel()
+		{
+		}
+
+		TSharedRef<IItemSelectorItemViewModelUtilities> GetItemUtilities() const
+		{
+			TSharedPtr<IItemSelectorItemViewModelUtilities> ItemUtilitiesPinned = ItemUtilities.Pin();
+			checkf(ItemUtilities.IsValid(), TEXT("Item utilities deleted before item."));
+			return ItemUtilitiesPinned.ToSharedRef();
+		}
+
+		EItemSelectorItemViewModelType GetType() const
+		{
+			return Type;
+		}
+
+		void GetChildren(TArray<TSharedRef<FItemSelectorItemViewModel>>& OutChildren) const
+		{
+			TArray<TSharedRef<FItemSelectorItemViewModel>> Children;
+			GetChildrenInternal(Children);
+			for (TSharedRef<FItemSelectorItemViewModel> Child : Children)
+			{
+				if (Child->PassesFilter())
+				{
+					OutChildren.Add(Child);
+				}
+			}
+		}
+
+		virtual bool PassesFilter() const
+		{
+			return true;
+		}
+
+	protected:
+		virtual void GetChildrenInternal(TArray<TSharedRef<FItemSelectorItemViewModel>>& OutChildren) const
+		{
+		}
+
+	private:
+		TWeakPtr<IItemSelectorItemViewModelUtilities> ItemUtilities;
+		EItemSelectorItemViewModelType Type;
+	};
+
+	class FItemSelectorItemContainerViewModel : public FItemSelectorItemViewModel
+	{
+	public:
+		FItemSelectorItemContainerViewModel(TSharedRef<IItemSelectorItemViewModelUtilities> InItemUtilities, const ItemType& InItem)
+			: FItemSelectorItemViewModel(InItemUtilities, EItemSelectorItemViewModelType::Item)
+			, Item(InItem)
+		{
+		}
+
+		const ItemType& GetItem() const
+		{
+			return Item;
+		}
+
+		virtual bool PassesFilter() const
+		{
+			return this->GetItemUtilities()->DoesItemMatchFilterText(Item);
+		}
+
+	private:
+		const ItemType& Item;
+	};
+
+	class FItemSelectorItemCategoryViewModel : public FItemSelectorItemViewModel
+	{
+	public:
+		FItemSelectorItemCategoryViewModel(TSharedRef<IItemSelectorItemViewModelUtilities> InItemUtilities, const CategoryType& InCategory)
+			: FItemSelectorItemViewModel(InItemUtilities, EItemSelectorItemViewModelType::Category)
+			, Category(InCategory)
+		{
+		}
+
+		const CategoryType& GetCategory() const
+		{
+			return Category;
+		}
+
+		TSharedRef<FItemSelectorItemCategoryViewModel> AddCategory(const CategoryType& InCategory)
+		{
+			TSharedRef<FItemSelectorItemCategoryViewModel> NewCategoryViewModel = MakeShared<FItemSelectorItemCategoryViewModel>(this->GetItemUtilities(), InCategory);
+			ChildCategoryViewModels.Add(NewCategoryViewModel);
+			return NewCategoryViewModel;
+		}
+
+		void AddItem(const ItemType& InItem)
+		{
+			ChildItemViewModels.Add(MakeShared<FItemSelectorItemContainerViewModel>(this->GetItemUtilities(), InItem));
+		}
+
+		TSharedPtr<FItemSelectorItemCategoryViewModel> FindChildCategory(const CategoryType& InCategory)
+		{
+			for (TSharedRef<FItemSelectorItemCategoryViewModel> ChildCategoryViewModel : ChildCategoryViewModels)
+			{
+				if (this->GetItemUtilities()->CompareCategoriesForEquality(ChildCategoryViewModel->GetCategory(), InCategory))
+				{
+					return ChildCategoryViewModel;
+				}
+			}
+			return TSharedPtr<FItemSelectorItemCategoryViewModel>();
+		}
+
+		void SortChildren()
+		{
+			TSharedRef<IItemSelectorItemViewModelUtilities> Utilities = this->GetItemUtilities();
+			if (ChildCategoryViewModels.Num() > 0 && Utilities->GetOnCompareCategoriesForSorting().IsBound())
+			{
+				ChildCategoryViewModels.Sort([Utilities](const TSharedRef<FItemSelectorItemCategoryViewModel>& CategoryViewModelA, const TSharedRef<FItemSelectorItemCategoryViewModel>& CategoryViewModelB)
+				{
+					return Utilities->GetOnCompareCategoriesForSorting().Execute(CategoryViewModelA->GetCategory(), CategoryViewModelB->GetCategory());
+				});
+
+			}
+			if (ChildItemViewModels.Num() > 0 && Utilities->GetOnCompareItemsForSorting().IsBound())
+			{
+				ChildItemViewModels.Sort([Utilities](const TSharedRef<FItemSelectorItemContainerViewModel>& ItemViewModelA, const TSharedRef<FItemSelectorItemContainerViewModel>& ItemViewModelB)
+				{
+					return Utilities->GetOnCompareItemsForSorting().Execute(ItemViewModelA->GetItem(), ItemViewModelB->GetItem());
+				});
+			}
+
+			for (TSharedRef<FItemSelectorItemCategoryViewModel> ChildCategoryViewModel : ChildCategoryViewModels)
+			{
+				ChildCategoryViewModel->SortChildren();
+			}
+		}
+
+		virtual bool PassesFilter() const
+		{
+			for (const TSharedRef<FItemSelectorItemContainerViewModel> ChildItemViewModel : ChildItemViewModels)
+			{
+				if (ChildItemViewModel->PassesFilter())
+				{
+					return true;
+				}
+			}
+			for (const TSharedRef<FItemSelectorItemCategoryViewModel> ChildCategoryViewModel : ChildCategoryViewModels)
+			{
+				if (ChildCategoryViewModel->PassesFilter())
+				{
+					return true;
+				}
+			}
+			return false;
+		}
+
+	protected:
+		virtual void GetChildrenInternal(TArray<TSharedRef<FItemSelectorItemViewModel>>& OutChildren) const
+		{
+			OutChildren.Append(ChildCategoryViewModels);
+			OutChildren.Append(ChildItemViewModels);
+		}
+
+	private:
+		const CategoryType& Category;
+		TArray<TSharedRef<FItemSelectorItemCategoryViewModel>> ChildCategoryViewModels;
+		TArray<TSharedRef<FItemSelectorItemContainerViewModel>> ChildItemViewModels;
+	};
+
+	class FItemSelectorViewModel : public IItemSelectorItemViewModelUtilities, public TSharedFromThis<FItemSelectorViewModel>
+	{
+	public:
+		FItemSelectorViewModel(TArray<ItemType> InItems, FOnGetCategoriesForItem InOnGetCategoriesForItem, FOnCompareCategoriesForEquality InOnCompareCategoriesForEquality, FOnCompareCategoriesForSorting InOnCompareCategoriesForSorting, FOnCompareItemsForSorting InOnCompareItemsForSorting, FOnDoesItemMatchFilterText InOnDoesItemMatchFilterText)
+			: Items(InItems)
+			, OnGetCategoriesForItem(InOnGetCategoriesForItem)
+			, OnCompareCategoriesForEquality(InOnCompareCategoriesForEquality)
+			, OnCompareCategoriesForSorting(InOnCompareCategoriesForSorting)
+			, OnCompareItemsForSorting(InOnCompareItemsForSorting)
+			, OnDoesItemMatchFilterText(InOnDoesItemMatchFilterText)
+		{
+		}
+
+		const TArray<TSharedRef<FItemSelectorItemViewModel>>* GetRootItems()
+		{
+			if (RootCategoryViewModel.IsValid() == false)
+			{
+				RootCategoryViewModel = MakeShared<FItemSelectorItemCategoryViewModel>(this->AsShared(), RootCategory);
+				for (const ItemType& Item : Items)
+				{
+					// Cache the category array since the category view models hold a const reference to the generated categories and without this they
+					// are deleted.
+					TSharedRef<TArray<CategoryType>> ItemCategories = MakeShared<TArray<CategoryType>>();
+					ItemCategoriesCache.Add(ItemCategories);
+
+					if (OnGetCategoriesForItem.IsBound())
+					{
+						ItemCategories->Append(OnGetCategoriesForItem.Execute(Item));
+					}
+
+					TSharedPtr<FItemSelectorItemCategoryViewModel> CurrentCategoryViewModel = RootCategoryViewModel;
+					if (ItemCategories->Num() > 0)
+					{
+
+						for (const CategoryType& ItemCategory : (*ItemCategories))
+						{
+							TSharedPtr<FItemSelectorItemCategoryViewModel> ExistingCategoryViewModel = CurrentCategoryViewModel->FindChildCategory(ItemCategory);
+							if (ExistingCategoryViewModel.IsValid())
+							{
+								CurrentCategoryViewModel = ExistingCategoryViewModel;
+							}
+							else
+							{
+								TSharedRef<FItemSelectorItemCategoryViewModel> NewItemCategoryViewModel = CurrentCategoryViewModel->AddCategory(ItemCategory);
+								CurrentCategoryViewModel = NewItemCategoryViewModel;
+							}
+						}
+					}
+					CurrentCategoryViewModel->AddItem(Item);
+				}
+				RootCategoryViewModel->SortChildren();
+				RootCategoryViewModel->GetChildren(RootTreeCategories);
+			}
+			return &RootTreeCategories;
+		}
+
+		const FText& GetFilterText() const
+		{
+			return FilterText;
+		}
+
+		void SetFilterText(FText InFilterText)
+		{
+			FilterText = InFilterText;
+			RootTreeCategories.Empty();
+			RootCategoryViewModel->GetChildren(RootTreeCategories);
+		}
+
+		virtual bool DoesItemMatchFilterText(const ItemType& InItem) const override
+		{
+			return FilterText.IsEmpty() || OnDoesItemMatchFilterText.IsBound() == false || OnDoesItemMatchFilterText.Execute(FilterText, InItem);
+		}
+
+		virtual bool CompareCategoriesForEquality(const CategoryType& CategoryA, const CategoryType& CategoryB) const override
+		{
+			return OnCompareCategoriesForEquality.Execute(CategoryA, CategoryB);
+		}
+
+		virtual const FOnCompareCategoriesForSorting& GetOnCompareCategoriesForSorting() const override
+		{
+			return OnCompareCategoriesForSorting;
+		}
+
+		virtual const FOnCompareItemsForSorting& GetOnCompareItemsForSorting() const override
+		{
+			return OnCompareItemsForSorting;
+		}
+
+		bool CanCompareItems() const
+		{
+			return OnCompareItemsForSorting.IsBound();
+		}
+
+		int32 CompareItems(const ItemType& ItemA, const ItemType& ItemB)
+		{
+			return OnCompareItemsForSorting.Execute(ItemA, ItemB);
+		}
+
+	private:
+		CategoryType RootCategory;
+
+		TArray<ItemType> Items;
+		TArray<TSharedRef<TArray<CategoryType>>> ItemCategoriesCache;
+
+		FOnGetCategoriesForItem OnGetCategoriesForItem;
+		FOnCompareCategoriesForEquality OnCompareCategoriesForEquality;
+		FOnCompareCategoriesForSorting OnCompareCategoriesForSorting;
+		FOnCompareItemsForSorting OnCompareItemsForSorting;
+		FOnDoesItemMatchFilterText OnDoesItemMatchFilterText;
+
+		TSharedPtr<FItemSelectorItemCategoryViewModel> RootCategoryViewModel;
+		TArray<TSharedRef<FItemSelectorItemViewModel>> RootTreeCategories;
+		FText FilterText;
+	};
+
+	typedef STableRow<TSharedRef<FItemSelectorItemViewModel>> SItemSelectorTableRow;
+
+	class SItemSelectorItemContainerTableRow : public SItemSelectorTableRow
+	{
+	public:
+		SLATE_BEGIN_ARGS(SItemSelectorItemContainerTableRow)
+		{}
+			SLATE_DEFAULT_SLOT(typename SItemSelectorItemContainerTableRow::FArguments, Content)
+		SLATE_END_ARGS();
+
+		void Construct(const FArguments& InArgs, const TSharedRef<STableViewBase>& OwnerTree)
+		{
+			typename SItemSelectorTableRow::FArguments Arguments;
+			Arguments[InArgs._Content.Widget];
+			SItemSelectorTableRow::Construct(Arguments,	OwnerTree);
+		}
+
+		virtual int32 GetIndentLevel() const override
+		{
+			return 0;
+		}
+	};
+
+public:
+	void Construct(const FArguments& InArgs)
+	{
+		Items = InArgs._Items;
+		OnGetCategoriesForItem = InArgs._OnGetCategoriesForItem;
+		OnCompareCategoriesForEquality = InArgs._OnCompareCategoriesForEquality;
+		OnCompareCategoriesForSorting = InArgs._OnCompareCategoriesForSorting;
+		OnCompareItemsForSorting = InArgs._OnCompareItemsForSorting;
+		OnDoesItemMatchFilterText = InArgs._OnDoesItemMatchFilterText;
+		OnGenerateWidgetForCategory = InArgs._OnGenerateWidgetForCategory;
+		OnGenerateWidgetForItem = InArgs._OnGenerateWidgetForItem;
+
+		checkf(OnGetCategoriesForItem.IsBound() == false || OnCompareCategoriesForEquality.IsBound(), TEXT("OnCompareCategoriesForEquality must be bound if OnGenerateCategoriesForItem is bound."));
+		checkf(OnGetCategoriesForItem.IsBound() == false || OnGenerateWidgetForCategory.IsBound(), TEXT("OnGenerateWidgetForCategory must be bound if OnGenerateCategoriesForItem is bound."));
+		checkf(OnGenerateWidgetForItem.IsBound(), TEXT("OnGenerateWidgetForItem must be bound"));
+
+		ViewModel = MakeShared<FItemSelectorViewModel>(Items, OnGetCategoriesForItem, OnCompareCategoriesForEquality, OnCompareCategoriesForSorting, OnCompareItemsForSorting, OnDoesItemMatchFilterText);
+
+		ChildSlot
+		[
+			SNew(SVerticalBox)
+			+ SVerticalBox::Slot()
+			.AutoHeight()
+			.Padding(0, 0, 0, 5)
+			[
+				SAssignNew(SearchBox, SSearchBox)
+				.Visibility(this, &SItemSelector::GetSearchBoxVisibility)
+				.OnTextChanged(this, &SItemSelector::OnSearchTextChanged)
+			]
+			+ SVerticalBox::Slot()
+			.Padding(0)
+			[
+				SAssignNew(ItemTree, STreeView<TSharedRef<FItemSelectorItemViewModel>>)
+				.SelectionMode(InArgs._AllowMultiselect ? ESelectionMode::Multi : ESelectionMode::SingleToggle)
+				.OnGenerateRow(this, &SItemSelector::OnGenerateRow)
+				.OnGetChildren(this, &SItemSelector::OnGetChildren)
+				.TreeItemsSource(ViewModel->GetRootItems())
+			]
+		];
+
+		ExpandTree();
+	}
+
+	TArray<ItemType> GetSelectedItems()
+	{
+		TArray<TSharedRef<FItemSelectorItemViewModel>> SelectedItemViewModels;
+		ItemTree->GetSelectedItems(SelectedItemViewModels);
+
+		TArray<ItemType> SelectedItems;
+		for (TSharedRef<FItemSelectorItemViewModel> SelectedItemViewModel : SelectedItemViewModels)
+		{
+			if (SelectedItemViewModel->GetType() == EItemSelectorItemViewModelType::Item)
+			{
+				TSharedRef<FItemSelectorItemContainerViewModel> SelectedItemContainerViewModel = StaticCastSharedRef<FItemSelectorItemContainerViewModel>(SelectedItemViewModel);
+				SelectedItems.Add(SelectedItemContainerViewModel->GetItem());
+			}
+		}
+		return SelectedItems;
+	}
+
+private:
+	EVisibility GetSearchBoxVisibility() const
+	{
+		return OnDoesItemMatchFilterText.IsBound() ? EVisibility::Visible : EVisibility::Collapsed;
+	}
+
+	void OnSearchTextChanged(const FText& SearchText)
+	{
+		if (ViewModel->GetFilterText().CompareTo(SearchText) != 0)
+		{
+			ViewModel->SetFilterText(SearchText);
+			ExpandTree();
+			ItemTree->RequestTreeRefresh();
+		}
+	}
+
+	TSharedRef<ITableRow> OnGenerateRow(TSharedRef<FItemSelectorItemViewModel> Item, const TSharedRef<STableViewBase>& OwnerTable)
+	{
+		switch (Item->GetType())
+		{
+		case EItemSelectorItemViewModelType::Category:
+		{
+			TSharedRef<FItemSelectorItemCategoryViewModel> ItemCategoryViewModel = StaticCastSharedRef<FItemSelectorItemCategoryViewModel>(Item);
+			return SNew(STableRow<TSharedRef<FItemSelectorItemViewModel>>, OwnerTable)
+				.ShowSelection(false)
+				[
+					OnGenerateWidgetForCategory.Execute(ItemCategoryViewModel->GetCategory())
+				];
+		}
+		case EItemSelectorItemViewModelType::Item:
+		{
+			TSharedRef<FItemSelectorItemContainerViewModel> ItemContainerViewModel = StaticCastSharedRef<FItemSelectorItemContainerViewModel>(Item);
+			return SNew(SItemSelectorItemContainerTableRow, OwnerTable)
+				[
+					OnGenerateWidgetForItem.Execute(ItemContainerViewModel->GetItem())
+				];
+		}
+		default:
+			checkf(false, TEXT("Unsupported type"));
+			return SNew(STableRow<TSharedRef<FItemSelectorItemViewModel>>, OwnerTable);
+		}
+	}
+
+	void OnGetChildren(TSharedRef<FItemSelectorItemViewModel> Item, TArray<TSharedRef<FItemSelectorItemViewModel>>& OutChildren)
+	{
+		Item->GetChildren(OutChildren);
+	}
+
+	void ExpandTree()
+	{
+		TArray<TSharedRef<FItemSelectorItemViewModel>> ItemsToProcess;
+		ItemsToProcess.Append(*ViewModel->GetRootItems());
+		while (ItemsToProcess.Num() > 0)
+		{
+			TSharedRef<FItemSelectorItemViewModel> ItemToProcess = ItemsToProcess[0];
+			ItemsToProcess.RemoveAtSwap(0);
+			ItemTree->SetItemExpansion(ItemToProcess, true);
+			ItemToProcess->GetChildren(ItemsToProcess);
+		}
+	}
+
+private:
+	TArray<ItemType> Items;
+
+	FOnGetCategoriesForItem OnGetCategoriesForItem;
+	FOnCompareCategoriesForEquality OnCompareCategoriesForEquality;
+	FOnCompareCategoriesForSorting OnCompareCategoriesForSorting;
+	FOnCompareItemsForSorting OnCompareItemsForSorting;
+	FOnDoesItemMatchFilterText OnDoesItemMatchFilterText;
+	FOnGenerateWidgetForCategory OnGenerateWidgetForCategory;
+	FOnGenerateWidgetForItem OnGenerateWidgetForItem;
+
+	TSharedPtr<FItemSelectorViewModel> ViewModel;
+	TSharedPtr<SSearchBox> SearchBox;
+	TSharedPtr<STreeView<TSharedRef<FItemSelectorItemViewModel>>> ItemTree;
+};
\ No newline at end of file
diff --git a/Engine/Plugins/FX/Niagara/Source/NiagaraEditor/Private/SNewEmitterDialog.cpp b/Engine/Plugins/FX/Niagara/Source/NiagaraEditor/Private/SNewEmitterDialog.cpp
new file mode 100644
--- /dev/null
+++ b/Engine/Plugins/FX/Niagara/Source/NiagaraEditor/Private/SNewEmitterDialog.cpp
@@ -0,0 +1,75 @@
+// Copyright 1998-2018 Epic Games, Inc. All Rights Reserved.
+
+#include "SNewEmitterDialog.h"
+#include "NiagaraEmitter.h"
+#include "NiagaraEditorStyle.h"
+#include "SNiagaraTemplateAssetPicker.h"
+#include "SNiagaraNewAssetDialog.h"
+
+#include "AssetData.h"
+
+#include "ContentBrowserModule.h"
+#include "IContentBrowserSingleton.h"
+
+typedef SItemSelector<FText, FAssetData> SNiagaraAssetItemSelector;
+
+#define LOCTEXT_NAMESPACE "SNewEmitterDialog"
+
+void SNewEmitterDialog::Construct(const FArguments& InArgs)
+{
+	FAssetPickerConfig AssetPickerConfig;
+	AssetPickerConfig.SelectionMode = ESelectionMode::SingleToggle;
+	AssetPickerConfig.InitialAssetViewType = EAssetViewType::List;
+	AssetPickerConfig.Filter.ClassNames.Add(UNiagaraEmitter::StaticClass()->GetFName());
+	AssetPickerConfig.GetCurrentSelectionDelegates.Add(&GetSelectedEmitterAssetsFromPicker);
+
+	FContentBrowserModule& ContentBrowserModule = FModuleManager::Get().LoadModuleChecked<FContentBrowserModule>(TEXT("ContentBrowser"));
+	TSharedRef<SWidget> AssetPicker = ContentBrowserModule.Get().CreateAssetPicker(AssetPickerConfig);
+
+	SNiagaraNewAssetDialog::Construct(SNiagaraNewAssetDialog::FArguments(), UNiagaraEmitter::StaticClass()->GetFName(), LOCTEXT("AssetTypeName", "emitter"),
+		{
+			SNiagaraNewAssetDialog::FNiagaraNewAssetDialogOption(
+				LOCTEXT("CreateFromTemplateLabel", "Create a new emitter from an emitter template"),
+				LOCTEXT("TemplatesPickerHeader", "Select a Template Emitter"),
+				SNiagaraNewAssetDialog::FOnGetSelectedAssetsFromPicker::CreateSP(this, &SNewEmitterDialog::GetSelectedEmitterTemplateAssets),
+				SAssignNew(TemplateAssetPicker, SNiagaraTemplateAssetPicker, UNiagaraEmitter::StaticClass())),
+			SNiagaraNewAssetDialog::FNiagaraNewAssetDialogOption(
+				LOCTEXT("CreateFromOtherEmitterLabel", "Copy an existing emitter from your project content"),
+				LOCTEXT("ProjectEmitterPickerHeader", "Select a Project Emitter"),
+				SNiagaraNewAssetDialog::FOnGetSelectedAssetsFromPicker::CreateSP(this, &SNewEmitterDialog::GetSelectedProjectEmiterAssets),
+				AssetPicker),
+			SNiagaraNewAssetDialog::FNiagaraNewAssetDialogOption(
+				LOCTEXT("CreateEmptyLabel", "Create an empty emitter with no modules or renderers (Advanced)"),
+				LOCTEXT("EmptyLabel", "Empty Emitter"),
+				SNiagaraNewAssetDialog::FOnGetSelectedAssetsFromPicker(),
+				SNew(SBox)
+				.HAlign(HAlign_Center)
+				.VAlign(VAlign_Center)
+				[
+					SNew(STextBlock)
+					.Text(LOCTEXT("NoOptionsLabel", "No Options"))
+				])
+		});
+}
+
+TOptional<FAssetData> SNewEmitterDialog::GetSelectedEmitterAsset()
+{
+	const TArray<FAssetData>& SelectedEmitterAssets = GetSelectedAssets();
+	if (SelectedEmitterAssets.Num() > 0)
+	{
+		return TOptional<FAssetData>(SelectedEmitterAssets[0]);
+	}
+	return TOptional<FAssetData>();
+}
+
+TArray<FAssetData> SNewEmitterDialog::GetSelectedEmitterTemplateAssets()
+{
+	return TemplateAssetPicker->GetSelectedAssets();
+}
+
+TArray<FAssetData> SNewEmitterDialog::GetSelectedProjectEmiterAssets()
+{
+	return GetSelectedEmitterAssetsFromPicker.Execute();
+}
+
+#undef LOCTEXT_NAMESPACE
\ No newline at end of file
diff --git a/Engine/Plugins/FX/Niagara/Source/NiagaraEditor/Private/SNewEmitterDialog.h b/Engine/Plugins/FX/Niagara/Source/NiagaraEditor/Private/SNewEmitterDialog.h
new file mode 100644
--- /dev/null
+++ b/Engine/Plugins/FX/Niagara/Source/NiagaraEditor/Private/SNewEmitterDialog.h
@@ -0,0 +1,35 @@
+// Copyright 1998-2018 Epic Games, Inc. All Rights Reserved.
+
+#pragma once
+
+#include "SNiagaraNewAssetDialog.h"
+
+#include "Widgets/DeclarativeSyntaxSupport.h"
+#include "Widgets/SWindow.h"
+#include "AssetData.h"
+#include "ContentBrowserDelegates.h"
+
+class SNiagaraTemplateAssetPicker;
+
+/** A modal dialog to collect information needed to create a new niagara system. */
+class SNewEmitterDialog : public SNiagaraNewAssetDialog
+{
+public:
+	SLATE_BEGIN_ARGS(SNewEmitterDialog)
+	{}
+	SLATE_END_ARGS()
+
+	void Construct(const FArguments& InArgs);
+
+	TOptional<FAssetData> GetSelectedEmitterAsset();
+
+private:
+	TArray<FAssetData> GetSelectedEmitterTemplateAssets();
+
+	TArray<FAssetData> GetSelectedProjectEmiterAssets();
+
+private:
+	TSharedPtr<SNiagaraTemplateAssetPicker> TemplateAssetPicker;
+
+	FGetCurrentSelectionDelegate GetSelectedEmitterAssetsFromPicker;
+};
\ No newline at end of file
diff --git a/Engine/Plugins/FX/Niagara/Source/NiagaraEditor/Private/SNewSystemDialog.cpp b/Engine/Plugins/FX/Niagara/Source/NiagaraEditor/Private/SNewSystemDialog.cpp
new file mode 100644
--- /dev/null
+++ b/Engine/Plugins/FX/Niagara/Source/NiagaraEditor/Private/SNewSystemDialog.cpp
@@ -0,0 +1,238 @@
+// Copyright 1998-2018 Epic Games, Inc. All Rights Reserved.
+
+#include "SNewSystemDialog.h"
+#include "NiagaraSystem.h"
+#include "NiagaraEditorStyle.h"
+#include "SNiagaraTemplateAssetPicker.h"
+
+#include "AssetData.h"
+
+#include "Widgets/Input/SCheckBox.h"
+#include "Widgets/Input/SButton.h"
+#include "Widgets/Layout/SUniformGridPanel.h"
+#include "Widgets/Layout/SWrapBox.h"
+#include "EditorStyleSet.h"
+
+#include "ContentBrowserModule.h"
+#include "IContentBrowserSingleton.h"
+
+typedef SItemSelector<FText, FAssetData> SNiagaraAssetItemSelector;
+
+#define LOCTEXT_NAMESPACE "SNewSystemDialog"
+
+void SNewSystemDialog::Construct(const FArguments& InArgs)
+{
+	FContentBrowserModule& ContentBrowserModule = FModuleManager::Get().LoadModuleChecked<FContentBrowserModule>(TEXT("ContentBrowser"));
+
+	FAssetPickerConfig SystemAssetPickerConfig;
+	SystemAssetPickerConfig.SelectionMode = ESelectionMode::SingleToggle;
+	SystemAssetPickerConfig.InitialAssetViewType = EAssetViewType::List;
+	SystemAssetPickerConfig.Filter.ClassNames.Add(UNiagaraSystem::StaticClass()->GetFName());
+	SystemAssetPickerConfig.GetCurrentSelectionDelegates.Add(&GetSelectedSystemAssetsFromPicker);
+
+	TSharedRef<SWidget> SystemAssetPicker = ContentBrowserModule.Get().CreateAssetPicker(SystemAssetPickerConfig);
+
+	FAssetPickerConfig EmitterAssetPickerConfig;
+	EmitterAssetPickerConfig.SelectionMode = ESelectionMode::Multi;
+	EmitterAssetPickerConfig.InitialAssetViewType = EAssetViewType::List;
+	EmitterAssetPickerConfig.Filter.ClassNames.Add(UNiagaraEmitter::StaticClass()->GetFName());
+	EmitterAssetPickerConfig.GetCurrentSelectionDelegates.Add(&GetSelectedEmitterAssetsFromPicker);
+
+	TSharedRef<SWidget> EmitterAssetPicker = ContentBrowserModule.Get().CreateAssetPicker(EmitterAssetPickerConfig);
+
+	SNiagaraNewAssetDialog::Construct(SNiagaraNewAssetDialog::FArguments(), UNiagaraSystem::StaticClass()->GetFName(), LOCTEXT("AssetTypeName", "system"),
+		{
+			SNiagaraNewAssetDialog::FNiagaraNewAssetDialogOption(
+				LOCTEXT("CreateFromTemplateLabel", "Create a new system from a system template"),
+				LOCTEXT("TemplateLabel", "Select a System Template"),
+				SNiagaraNewAssetDialog::FOnGetSelectedAssetsFromPicker::CreateSP(this, &SNewSystemDialog::GetSelectedSystemTemplateAssets),
+				SAssignNew(TemplateAssetPicker, SNiagaraTemplateAssetPicker, UNiagaraSystem::StaticClass())),
+			SNiagaraNewAssetDialog::FNiagaraNewAssetDialogOption(
+				LOCTEXT("CreateFromSelectedEmittersLabel", "Create a new system from a set of selected emitters"),
+				LOCTEXT("ProjectEmittersLabel", "Select Emitters to Add"),
+				SNiagaraNewAssetDialog::FOnGetSelectedAssetsFromPicker::CreateSP(this, &SNewSystemDialog::GetSelectedProjectEmiterAssets),
+				SNew(SVerticalBox)
+				+ SVerticalBox::Slot()
+				.Padding(0, 0, 0, 10)
+				[
+					EmitterAssetPicker
+				]
+				+ SVerticalBox::Slot()
+				.AutoHeight()
+				.Padding(0, 0, 0, 2)
+				[
+					SNew(SHorizontalBox)
+					+ SHorizontalBox::Slot()
+					.AutoWidth()
+					.VAlign(VAlign_Center)
+					[
+						SNew(STextBlock)
+						.TextStyle(FNiagaraEditorStyle::Get(), "NiagaraEditor.NewAssetDialog.SubHeaderText")
+						.Text(LOCTEXT("SelectedEmittersLabel", "Emitters to Add:"))
+					]
+					+ SHorizontalBox::Slot()
+					.HAlign(HAlign_Right)
+					.VAlign(VAlign_Center)
+					[
+						SNew(SButton)
+						//.ButtonStyle(FNiagaraEditorStyle::Get(), "NiagaraEditor.NewAssetDialog.AddButton")
+						.ButtonStyle(FEditorStyle::Get(), "FlatButton.Success")
+						.IsEnabled(this, &SNewSystemDialog::IsAddEmittersToSelectionButtonEnabled)
+						.OnClicked(this, &SNewSystemDialog::AddEmittersToSelectionButtonClicked)
+						.ToolTipText(LOCTEXT("AddSelectedEmitterToolTip", "Add the selected emitter to the collection\n of emitters to be added to the new system."))
+						.Content()
+						[
+							SNew(SBox)
+							.WidthOverride(32.0f)
+							.HeightOverride(16.0f)
+							.HAlign(HAlign_Center)
+							.VAlign(VAlign_Center)
+							[
+								SNew(STextBlock)
+								.TextStyle(FEditorStyle::Get(), "NormalText.Important")
+								.Font(FEditorStyle::Get().GetFontStyle("FontAwesome.10"))
+								.Text(FText::FromString(FString(TEXT("\xf067"))) /*fa-plus*/)
+							]
+						]
+					]
+				]
+				+ SVerticalBox::Slot()
+				.AutoHeight()
+				[
+					SAssignNew(SelectedEmitterBox, SWrapBox)
+					.UseAllottedWidth(true)
+				]),
+			SNiagaraNewAssetDialog::FNiagaraNewAssetDialogOption(
+				LOCTEXT("CreateFromOtherSystemLabel", "Copy an existing system from your project content"),
+				LOCTEXT("ProjectSystemsLabel", "Select a Project System"),
+				SNiagaraNewAssetDialog::FOnGetSelectedAssetsFromPicker::CreateSP(this, &SNewSystemDialog::GetSelectedProjectSystemAssets),
+				SystemAssetPicker),
+			SNiagaraNewAssetDialog::FNiagaraNewAssetDialogOption(
+				LOCTEXT("CreateEmptyLabel", "Create an empty system with no emitters"),
+				LOCTEXT("EmptyLabel", "Empty System"),
+				SNiagaraNewAssetDialog::FOnGetSelectedAssetsFromPicker(),
+				SNew(SBox)
+				.HAlign(HAlign_Center)
+				.VAlign(VAlign_Center)
+				[
+					SNew(STextBlock)
+					.Text(LOCTEXT("NoOptionsLabel", "No Options"))
+				])
+		});
+}
+
+TOptional<FAssetData> SNewSystemDialog::GetSelectedSystemAsset() const
+{
+	const TArray<FAssetData>& AllSelectedAssets = GetSelectedAssets();
+	TArray<FAssetData> SelectedSystemAssets;
+	for (const FAssetData& SelectedAsset : AllSelectedAssets)
+	{
+		if (SelectedAsset.AssetClass == UNiagaraSystem::StaticClass()->GetFName())
+		{
+			SelectedSystemAssets.Add(SelectedAsset);
+		}
+	}
+	if (SelectedSystemAssets.Num() == 1)
+	{
+		return TOptional<FAssetData>(SelectedSystemAssets[0]);
+	}
+	return TOptional<FAssetData>();
+}
+
+TArray<FAssetData> SNewSystemDialog::GetSelectedEmitterAssets() const
+{
+	const TArray<FAssetData>& AllSelectedAssets = GetSelectedAssets();
+	TArray<FAssetData> ConfirmedSelectedEmitterAssets;
+	for (const FAssetData& SelectedAsset : AllSelectedAssets)
+	{
+		if (SelectedAsset.AssetClass == UNiagaraEmitter::StaticClass()->GetFName())
+		{
+			ConfirmedSelectedEmitterAssets.Add(SelectedAsset);
+		}
+	}
+	return ConfirmedSelectedEmitterAssets;
+}
+
+TArray<FAssetData> SNewSystemDialog::GetSelectedSystemTemplateAssets()
+{
+	return TemplateAssetPicker->GetSelectedAssets();
+}
+
+TArray<FAssetData> SNewSystemDialog::GetSelectedProjectSystemAssets()
+{
+	return GetSelectedSystemAssetsFromPicker.Execute();
+}
+
+TArray<FAssetData> SNewSystemDialog::GetSelectedProjectEmiterAssets()
+{
+	return SelectedEmitterAssets;
+}
+
+bool SNewSystemDialog::IsAddEmittersToSelectionButtonEnabled() const
+{
+	return GetSelectedEmitterAssetsFromPicker.Execute().Num() > 0;
+}
+
+FReply SNewSystemDialog::AddEmittersToSelectionButtonClicked()
+{
+	TArray<FAssetData> SelectedEmitterAssetsFromPicker = GetSelectedEmitterAssetsFromPicker.Execute();
+	for (const FAssetData& SelectedEmitterAsset : SelectedEmitterAssetsFromPicker)
+	{
+		if (SelectedEmitterAssets.Contains(SelectedEmitterAsset) == false)
+		{
+			TSharedPtr<SWidget> SelectedEmitterWidget;
+			SelectedEmitterBox->AddSlot()
+				.Padding(FMargin(0, 0, 5, 0))
+				[
+					SNew(SBorder)
+					.BorderImage(FNiagaraEditorStyle::Get().GetBrush("NiagaraEditor.NewAssetDialog.SubBorder"))
+					.BorderBackgroundColor(FNiagaraEditorStyle::Get().GetColor("NiagaraEditor.NewAssetDialog.SubBorderColor"))
+					[
+						SAssignNew(SelectedEmitterWidget, SHorizontalBox)
+						+ SHorizontalBox::Slot()
+						.AutoWidth()
+						.VAlign(VAlign_Center)
+						.Padding(5, 0, 0, 0)
+						[
+							SNew(STextBlock)
+							.Text(FText::FromName(SelectedEmitterAsset.AssetName))
+						]
+						+ SHorizontalBox::Slot()
+						.AutoWidth()
+						.VAlign(VAlign_Center)
+						.Padding(2, 0, 0, 0)
+						[
+							SNew(SButton)
+							.ButtonStyle(FEditorStyle::Get(), "HoverHintOnly")
+							.OnClicked(this, &SNewSystemDialog::RemoveEmitterFromSelectionButtonClicked, SelectedEmitterAsset)
+							.ToolTipText(LOCTEXT("RemoveSelectedEmitterToolTip", "Remove the selected emitter from the collection\n of emitters to be added to the new system."))
+							[
+								SNew(STextBlock)
+								//.TextStyle(FEditorStyle::Get(), "NormalText.Important")
+								.Font(FEditorStyle::Get().GetFontStyle("FontAwesome.10"))
+								.Text(FText::FromString(FString(TEXT("\xf057"))) /*times-circle*/)
+								.ColorAndOpacity(FLinearColor(.8f, .2f, .2f, 1.0f))
+							]
+						]
+					]
+				];
+			SelectedEmitterAssets.Add(SelectedEmitterAsset);
+			SelectedEmitterAssetWidgets.Add(SelectedEmitterWidget);
+		}
+	}
+	return FReply::Handled();
+}
+
+FReply SNewSystemDialog::RemoveEmitterFromSelectionButtonClicked(FAssetData EmitterAsset)
+{
+	int32 SelectionIndex = SelectedEmitterAssets.IndexOfByKey(EmitterAsset);
+	if (SelectionIndex != INDEX_NONE)
+	{
+		SelectedEmitterAssets.RemoveAt(SelectionIndex);
+		SelectedEmitterBox->RemoveSlot(SelectedEmitterAssetWidgets[SelectionIndex].ToSharedRef());
+		SelectedEmitterAssetWidgets.RemoveAt(SelectionIndex);
+	}
+	return FReply::Handled();
+}
+
+#undef LOCTEXT_NAMESPACE
\ No newline at end of file
diff --git a/Engine/Plugins/FX/Niagara/Source/NiagaraEditor/Private/SNewSystemDialog.h b/Engine/Plugins/FX/Niagara/Source/NiagaraEditor/Private/SNewSystemDialog.h
new file mode 100644
--- /dev/null
+++ b/Engine/Plugins/FX/Niagara/Source/NiagaraEditor/Private/SNewSystemDialog.h
@@ -0,0 +1,53 @@
+// Copyright 1998-2018 Epic Games, Inc. All Rights Reserved.
+
+#pragma once
+
+#include "Widgets/DeclarativeSyntaxSupport.h"
+#include "Widgets/SWindow.h"
+#include "AssetData.h"
+#include "IAssetTypeActions.h"
+#include "ContentBrowserDelegates.h"
+#include "SNiagaraNewAssetDialog.h"
+
+class SNiagaraTemplateAssetPicker;
+class SWrapBox;
+
+/** A modal dialog to collect information needed to create a new niagara system. */
+class SNewSystemDialog : public SNiagaraNewAssetDialog
+{
+public:
+	SLATE_BEGIN_ARGS(SNewSystemDialog)
+	{}
+	SLATE_END_ARGS()
+
+	void Construct(const FArguments& InArgs);
+
+	TOptional<FAssetData> GetSelectedSystemAsset() const;
+
+	TArray<FAssetData> GetSelectedEmitterAssets() const;
+
+private:
+	TArray<FAssetData> GetSelectedSystemTemplateAssets();
+
+	TArray<FAssetData> GetSelectedProjectSystemAssets();
+
+	TArray<FAssetData> GetSelectedProjectEmiterAssets();
+	bool IsAddEmittersToSelectionButtonEnabled() const;
+	
+	FReply AddEmittersToSelectionButtonClicked();
+
+	FReply RemoveEmitterFromSelectionButtonClicked(FAssetData EmitterAsset);
+
+private:
+	FGetCurrentSelectionDelegate GetSelectedSystemAssetsFromPicker;
+
+	FGetCurrentSelectionDelegate GetSelectedEmitterAssetsFromPicker;
+
+	TArray<FAssetData> SelectedEmitterAssets;
+
+	TArray<TSharedPtr<SWidget>> SelectedEmitterAssetWidgets;
+
+	TSharedPtr<SWrapBox> SelectedEmitterBox;
+
+	TSharedPtr<SNiagaraTemplateAssetPicker> TemplateAssetPicker;
+};
\ No newline at end of file
diff --git a/Engine/Plugins/FX/Niagara/Source/NiagaraEditor/Private/SNiagaraNewAssetDialog.cpp b/Engine/Plugins/FX/Niagara/Source/NiagaraEditor/Private/SNiagaraNewAssetDialog.cpp
new file mode 100644
--- /dev/null
+++ b/Engine/Plugins/FX/Niagara/Source/NiagaraEditor/Private/SNiagaraNewAssetDialog.cpp
@@ -0,0 +1,271 @@
+// Copyright 1998-2018 Epic Games, Inc. All Rights Reserved.
+
+#include "SNiagaraNewAssetDialog.h"
+#include "NiagaraEmitter.h"
+#include "NiagaraEditorStyle.h"
+#include "SNiagaraTemplateAssetPicker.h"
+#include "NiagaraEditorSettings.h"
+
+#include "AssetData.h"
+
+#include "Widgets/Input/SCheckBox.h"
+#include "Widgets/Input/SButton.h"
+#include "Widgets/Layout/SUniformGridPanel.h"
+#include "EditorStyleSet.h"
+
+#include "ContentBrowserModule.h"
+#include "IContentBrowserSingleton.h"
+
+typedef SItemSelector<FText, FAssetData> SNiagaraAssetItemSelector;
+
+#define LOCTEXT_NAMESPACE "SNiagaraNewAssetDialog"
+
+void SNiagaraNewAssetDialog::Construct(const FArguments& InArgs, FName InSaveConfigKey, FText AssetTypeDisplayName, TArray<FNiagaraNewAssetDialogOption> InOptions)
+{
+	bUserConfirmedSelection = false;
+
+	SaveConfigKey = InSaveConfigKey;
+	FNiagaraNewAssetDialogConfig DialogConfig = GetDefault<UNiagaraEditorSettings>()->GetNewAssetDailogConfig(SaveConfigKey);
+	SelectedOptionIndex = DialogConfig.SelectedOptionIndex;
+
+	Options = InOptions;
+
+	SetOnWindowClosed(FOnWindowClosed::CreateSP(this, &SNiagaraNewAssetDialog::OnWindowClosed));
+	
+	TSharedPtr<SVerticalBox> OptionsBox;
+	TSharedPtr<SOverlay> AssetPickerOverlay;
+
+	TSharedRef<SVerticalBox> RootBox =
+		SNew(SVerticalBox)
+
+		// Options label
+		+ SVerticalBox::Slot()
+		.Padding(0, 7, 0, 0)
+		.AutoHeight()
+		[
+			SNew(SBox)
+			.Padding(FEditorStyle::GetMargin("StandardDialog.SlotPadding"))
+			[
+				SNew(STextBlock)
+				.TextStyle(FNiagaraEditorStyle::Get(), "NiagaraEditor.NewAssetDialog.HeaderText")
+				.Text(LOCTEXT("OptionsLabel", "Select an Option"))
+			]
+		]
+
+		// Creation mode radio buttons.
+		+ SVerticalBox::Slot()
+		.Padding(0, 5, 0, 5)
+		.AutoHeight()
+		[
+			SNew(SBox)
+			.Padding(FEditorStyle::GetMargin("StandardDialog.SlotPadding"))
+			[
+					
+				SNew(SBorder)
+				.BorderImage(FEditorStyle::GetBrush("ToolPanel.GroupBorder"))
+				.Padding(FMargin(7))
+				[
+					SAssignNew(OptionsBox, SVerticalBox)
+				]
+			]
+		]
+
+		// Asset pickers label
+		+ SVerticalBox::Slot()
+		.Padding(0, 5, 0, 0)
+		.AutoHeight()
+		[
+			SNew(SBox)
+			.Padding(FEditorStyle::GetMargin("StandardDialog.SlotPadding"))
+			[
+				SNew(STextBlock)
+				.TextStyle(FNiagaraEditorStyle::Get(), "NiagaraEditor.NewAssetDialog.HeaderText")
+				.Text(this, &SNiagaraNewAssetDialog::GetAssetPickersLabelText)
+			]
+		]
+
+		// Asset pickers
+		+ SVerticalBox::Slot()
+		.Padding(0, 5, 0, 5)
+		[
+			SNew(SBox)
+			.Padding(FEditorStyle::GetMargin("StandardDialog.SlotPadding"))
+			[
+				SNew(SBorder)
+				.BorderImage(FEditorStyle::GetBrush("ToolPanel.GroupBorder"))
+				.Padding(FMargin(7))
+				[
+					SAssignNew(AssetPickerOverlay, SOverlay)
+				]
+			]
+		]
+
+		// OK/Cancel buttons
+		+SVerticalBox::Slot()
+		.AutoHeight()
+		.HAlign(HAlign_Right)
+		.Padding(0, 5, 0, 5)
+		[
+			SNew(SUniformGridPanel)
+			.SlotPadding(FEditorStyle::GetMargin("StandardDialog.SlotPadding"))
+			.MinDesiredSlotWidth(FEditorStyle::GetFloat("StandardDialog.MinDesiredSlotWidth"))
+			.MinDesiredSlotHeight(FEditorStyle::GetFloat("StandardDialog.MinDesiredSlotHeight"))
+			+SUniformGridPanel::Slot(0, 0)
+			[
+				SNew(SButton)
+				.HAlign(HAlign_Center)
+				.ContentPadding(FEditorStyle::GetMargin("StandardDialog.ContentPadding"))
+				.Text(LOCTEXT("OK", "OK"))
+				.OnClicked(this, &SNiagaraNewAssetDialog::OnOkButtonClicked)
+				.IsEnabled(this, &SNiagaraNewAssetDialog::IsOkButtonEnabled)
+			]
+			+SUniformGridPanel::Slot(1, 0)
+			[
+				SNew(SButton)
+				.HAlign(HAlign_Center)
+				.ContentPadding(FEditorStyle::GetMargin("StandardDialog.ContentPadding"))
+				.Text(LOCTEXT("Cancel", "Cancel"))
+				.OnClicked(this, &SNiagaraNewAssetDialog::OnCancelButtonClicked)
+			]
+		];
+
+	int32 OptionIndex = 0;
+	for (FNiagaraNewAssetDialogOption& Option : Options)
+	{
+		OptionsBox->AddSlot()
+			.Padding(0, 0, 0, OptionIndex < Options.Num() - 1 ? 7 : 0)
+			.AutoHeight()
+			[
+				SNew(SBorder)
+				.BorderImage(FNiagaraEditorStyle::Get().GetBrush("NiagaraEditor.NewAssetDialog.SubBorder"))
+				.BorderBackgroundColor(this, &SNiagaraNewAssetDialog::GetOptionBorderColor, OptionIndex)
+				[
+					SNew(SCheckBox)
+					.Style(FCoreStyle::Get(), "RadioButton")
+					.CheckBoxContentUsesAutoWidth(false)
+					.IsChecked(this, &SNiagaraNewAssetDialog::GetOptionCheckBoxState, OptionIndex)
+					.OnCheckStateChanged(this, &SNiagaraNewAssetDialog::OptionCheckBoxStateChanged, OptionIndex)
+					.Content()
+					[
+						SNew(STextBlock)
+						.TextStyle(FNiagaraEditorStyle::Get(), "NiagaraEditor.NewAssetDialog.OptionText")
+						.ColorAndOpacity(this, &SNiagaraNewAssetDialog::GetOptionTextColor, OptionIndex)
+						.Text(Option.OptionText)
+						.AutoWrapText(true)
+					]
+				]
+			];
+
+		AssetPickerOverlay->AddSlot()
+			[
+				SNew(SBox)
+				.Visibility(this, &SNiagaraNewAssetDialog::GetAssetPickerVisibility, OptionIndex)
+				[
+					Option.AssetPicker
+				]
+			];
+
+		OptionIndex++;
+	}
+
+	SWindow::Construct(SWindow::FArguments()
+		.Title(FText::Format(LOCTEXT("NewEmitterDialogTitle", "Pick a starting point for your {0}"), AssetTypeDisplayName))
+		.SizingRule(ESizingRule::UserSized)
+		.ClientSize(DialogConfig.WindowSize)
+		.SupportsMaximize(false)
+		.SupportsMinimize(false)
+		[
+			RootBox
+		]);
+}
+
+bool SNiagaraNewAssetDialog::GetUserConfirmedSelection() const
+{
+	return bUserConfirmedSelection;
+}
+
+const TArray<FAssetData>& SNiagaraNewAssetDialog::GetSelectedAssets() const
+{
+	return SelectedAssets;
+}
+
+void SNiagaraNewAssetDialog::OnWindowClosed(const TSharedRef<SWindow>& Window)
+{
+	SaveConfig();
+}
+
+FSlateColor SNiagaraNewAssetDialog::GetOptionBorderColor(int32 OptionIndex) const
+{
+	return SelectedOptionIndex == OptionIndex
+		? FNiagaraEditorStyle::Get().GetColor("NiagaraEditor.NewAssetDialog.ActiveOptionBorderColor")
+		: FSlateColor(FLinearColor::Transparent);
+}
+
+FSlateColor SNiagaraNewAssetDialog::GetOptionTextColor(int32 OptionIndex) const
+{
+	return SelectedOptionIndex == OptionIndex
+		? FSlateColor(FLinearColor::White)
+		: FSlateColor::UseForeground();
+}
+
+ECheckBoxState SNiagaraNewAssetDialog::GetOptionCheckBoxState(int32 OptionIndex) const
+{
+	return SelectedOptionIndex == OptionIndex ? ECheckBoxState::Checked : ECheckBoxState::Unchecked;
+}
+
+void SNiagaraNewAssetDialog::OptionCheckBoxStateChanged(ECheckBoxState InCheckBoxState, int32 OptionIndex)
+{
+	if (InCheckBoxState == ECheckBoxState::Checked)
+	{
+		SelectedOptionIndex = OptionIndex;
+	}
+}
+
+FText SNiagaraNewAssetDialog::GetAssetPickersLabelText() const
+{
+	return Options[SelectedOptionIndex].AssetPickerHeader;
+}
+
+EVisibility SNiagaraNewAssetDialog::GetAssetPickerVisibility(int32 OptionIndex) const
+{
+	return  SelectedOptionIndex == OptionIndex ? EVisibility::Visible : EVisibility::Collapsed;
+}
+
+bool SNiagaraNewAssetDialog::IsOkButtonEnabled() const
+{
+	const FNiagaraNewAssetDialogOption& SelectedOption = Options[SelectedOptionIndex];
+	return SelectedOption.OnGetSelectedAssetsFromPicker.IsBound() == false || SelectedOption.OnGetSelectedAssetsFromPicker.Execute().Num() > 0;
+}
+
+FReply SNiagaraNewAssetDialog::OnOkButtonClicked()
+{
+	const FNiagaraNewAssetDialogOption& SelectedOption = Options[SelectedOptionIndex];
+	if (SelectedOption.OnGetSelectedAssetsFromPicker.IsBound())
+	{
+		SelectedAssets.Append(SelectedOption.OnGetSelectedAssetsFromPicker.Execute());
+	}
+	bUserConfirmedSelection = true;
+
+	RequestDestroyWindow();
+	return FReply::Handled();
+}
+
+FReply SNiagaraNewAssetDialog::OnCancelButtonClicked()
+{
+	bUserConfirmedSelection = false;
+	SelectedAssets.Empty();
+
+	RequestDestroyWindow();
+	return FReply::Handled();
+}
+
+void SNiagaraNewAssetDialog::SaveConfig()
+{
+	FNiagaraNewAssetDialogConfig Config;
+	Config.SelectedOptionIndex = SelectedOptionIndex;
+	Config.WindowSize = GetClientSizeInScreen() / GetDPIScaleFactor();
+
+	GetMutableDefault<UNiagaraEditorSettings>()->SetNewAssetDialogConfig(SaveConfigKey, Config);
+}
+
+#undef LOCTEXT_NAMESPACE
\ No newline at end of file
diff --git a/Engine/Plugins/FX/Niagara/Source/NiagaraEditor/Private/SNiagaraNewAssetDialog.h b/Engine/Plugins/FX/Niagara/Source/NiagaraEditor/Private/SNiagaraNewAssetDialog.h
new file mode 100644
--- /dev/null
+++ b/Engine/Plugins/FX/Niagara/Source/NiagaraEditor/Private/SNiagaraNewAssetDialog.h
@@ -0,0 +1,80 @@
+// Copyright 1998-2018 Epic Games, Inc. All Rights Reserved.
+
+#pragma once
+
+#include "Widgets/DeclarativeSyntaxSupport.h"
+#include "Widgets/SWindow.h"
+#include "AssetData.h"
+#include "IAssetTypeActions.h"
+
+class SNiagaraTemplateAssetPicker;
+
+/** A modal dialog to collect information needed to create a new niagara system. */
+class SNiagaraNewAssetDialog : public SWindow
+{
+public:
+	DECLARE_DELEGATE_RetVal(TArray<FAssetData>, FOnGetSelectedAssetsFromPicker);
+
+public:
+	class FNiagaraNewAssetDialogOption
+	{
+	public:
+		FText OptionText;
+		FText AssetPickerHeader;
+		TSharedRef<SWidget> AssetPicker;
+		FOnGetSelectedAssetsFromPicker OnGetSelectedAssetsFromPicker;
+
+		FNiagaraNewAssetDialogOption(FText InOptionText, FText InAssetPickerHeader, FOnGetSelectedAssetsFromPicker InOnGetSelectedAssetsFromPicker, TSharedRef<SWidget> InAssetPicker)
+			: OptionText(InOptionText)
+			, AssetPickerHeader(InAssetPickerHeader)
+			, AssetPicker(InAssetPicker)
+			, OnGetSelectedAssetsFromPicker(InOnGetSelectedAssetsFromPicker)
+		{
+		}
+	};
+
+public:
+	SLATE_BEGIN_ARGS(SNiagaraNewAssetDialog)
+	{}
+	SLATE_END_ARGS()
+
+	void Construct(const FArguments& InArgs, FName InSaveConfigKey, FText AssetTypeDisplayName, TArray<FNiagaraNewAssetDialogOption> InOptions);
+
+	bool GetUserConfirmedSelection() const;
+
+protected:
+	const TArray<FAssetData>& GetSelectedAssets() const;
+
+private:
+	void OnWindowClosed(const TSharedRef<SWindow>& Window);
+
+	FSlateColor GetOptionBorderColor(int32 OptionIndex) const;
+
+	ECheckBoxState GetOptionCheckBoxState(int32 OptionIndex) const;
+
+	void OptionCheckBoxStateChanged(ECheckBoxState InCheckBoxState, int32 OptionIndex);
+
+	FSlateColor GetOptionTextColor(int32 OptionIndex) const;
+
+	FText GetAssetPickersLabelText() const;
+
+	EVisibility GetAssetPickerVisibility(int32 OptionIndex) const;
+
+	bool IsOkButtonEnabled() const;
+
+	FReply OnOkButtonClicked();
+
+	FReply OnCancelButtonClicked();
+
+	void SaveConfig();
+
+private:
+	FName SaveConfigKey;
+
+	TArray<FNiagaraNewAssetDialogOption> Options;
+
+	int32 SelectedOptionIndex;
+
+	bool bUserConfirmedSelection;
+	TArray<FAssetData> SelectedAssets;
+};
\ No newline at end of file
diff --git a/Engine/Plugins/FX/Niagara/Source/NiagaraEditor/Private/SNiagaraTemplateAssetPicker.cpp b/Engine/Plugins/FX/Niagara/Source/NiagaraEditor/Private/SNiagaraTemplateAssetPicker.cpp
new file mode 100644
--- /dev/null
+++ b/Engine/Plugins/FX/Niagara/Source/NiagaraEditor/Private/SNiagaraTemplateAssetPicker.cpp
@@ -0,0 +1,187 @@
+// Copyright 1998-2018 Epic Games, Inc. All Rights Reserved.
+
+#include "SNiagaraTemplateAssetPicker.h"
+#include "NiagaraEmitter.h"
+#include "NiagaraEditorStyle.h"
+
+#include "AssetData.h"
+#include "AssetRegistryModule.h"
+#include "Modules/ModuleManager.h"
+#include "AssetThumbnail.h"
+
+#include "Widgets/Input/SButton.h"
+#include "Widgets/Layout/SUniformGridPanel.h"
+#include "EditorStyleSet.h"
+
+#define LOCTEXT_NAMESPACE "SNiagaraAssetSelector"
+
+void SNiagaraTemplateAssetPicker::Construct(const FArguments& InArgs, UClass* AssetClass)
+{
+	FAssetRegistryModule& AssetRegistryModule = FModuleManager::LoadModuleChecked<FAssetRegistryModule>(TEXT("AssetRegistry"));
+	TArray<FAssetData> EmitterAssets;
+	AssetRegistryModule.Get().GetAssetsByClass(AssetClass->GetFName(), EmitterAssets);
+
+	NiagaraPluginCategory = LOCTEXT("NiagaraCategory", "Engine (Niagara Plugin)");
+	ProjectCategory = LOCTEXT("ProjectCategory", "Project");
+
+	AssetThumbnailPool = MakeShareable(new FAssetThumbnailPool(24));
+		
+	TArray<FAssetData> EmittersToShow;
+	for (const FAssetData& EmitterAsset : EmitterAssets)
+	{
+		bool bShowEmitter = false;
+		EmitterAsset.GetTagValue("bIsTemplateAsset", bShowEmitter);
+		if (bShowEmitter)
+		{
+			EmittersToShow.Add(EmitterAsset);
+		}
+	}
+
+	ChildSlot
+	[
+		SAssignNew(ItemSelector, SNiagaraAssetItemSelector)
+		.Items(EmittersToShow)
+		.OnGetCategoriesForItem(this, &SNiagaraTemplateAssetPicker::OnGetCategoriesForItem)
+		.OnCompareCategoriesForEquality(this, &SNiagaraTemplateAssetPicker::OnCompareCategoriesForEquality)
+		.OnCompareCategoriesForSorting(this, &SNiagaraTemplateAssetPicker::OnCompareCategoriesForSorting)
+		.OnCompareItemsForSorting(this, &SNiagaraTemplateAssetPicker::OnCompareItemsForSorting)
+		.OnDoesItemMatchFilterText(this, &SNiagaraTemplateAssetPicker::OnDoesItemMatchFilterText)
+		.OnGenerateWidgetForCategory(this, &SNiagaraTemplateAssetPicker::OnGenerateWidgetForCategory)
+		.OnGenerateWidgetForItem(this, &SNiagaraTemplateAssetPicker::OnGenerateWidgetForItem)
+	];
+}
+
+void SNiagaraTemplateAssetPicker::Tick(const FGeometry& AllottedGeometry, const double InCurrentTime, const float InDeltaTime)
+{
+	AssetThumbnailPool->Tick(InDeltaTime);
+}
+
+TArray<FAssetData> SNiagaraTemplateAssetPicker::GetSelectedAssets() const
+{
+	return ItemSelector->GetSelectedItems();
+}
+
+TArray<FText> SNiagaraTemplateAssetPicker::OnGetCategoriesForItem(const FAssetData& Item)
+{
+	TArray<FText> Categories;
+	TArray<FString> AssetPathParts;
+	Item.ObjectPath.ToString().ParseIntoArray(AssetPathParts, TEXT("/"));
+	if (AssetPathParts.Num() > 0)
+	{
+		if (AssetPathParts[0] == TEXT("Niagara"))
+		{
+			Categories.Add(LOCTEXT("NiagaraCategory", "Engine (Niagara Plugin)"));
+		}
+		else if(AssetPathParts[0] == TEXT("Game"))
+		{
+			Categories.Add(LOCTEXT("ProjectCategory", "Project"));
+		}
+		else
+		{
+			Categories.Add(FText::Format(LOCTEXT("OtherPluginFormat", "Plugin - {0}"), FText::FromString(AssetPathParts[0])));
+		}
+	}
+	return Categories;
+}
+
+bool SNiagaraTemplateAssetPicker::OnCompareCategoriesForEquality(const FText& CategoryA, const FText& CategoryB) const
+{
+	return CategoryA.CompareTo(CategoryB) == 0;
+}
+
+bool SNiagaraTemplateAssetPicker::OnCompareCategoriesForSorting(const FText& CategoryA, const FText& CategoryB) const
+{
+	int32 CompareResult = CategoryA.CompareTo(CategoryB);
+	if (CompareResult != 0)
+	{
+		// Project first
+		if (CategoryA.CompareTo(ProjectCategory) == 0)
+		{
+			return true;
+		}
+		if (CategoryB.CompareTo(ProjectCategory) == 0)
+		{
+			return false;
+		}
+
+		// Niagara plugin second.
+		if (CategoryA.CompareTo(NiagaraPluginCategory) == 0)
+		{
+			return true;
+		}
+		if (CategoryB.CompareTo(NiagaraPluginCategory) == 0)
+		{
+			return false;
+		}
+	}
+	// Otherwise just return the actual result.
+	return CompareResult == -1;
+}
+
+bool SNiagaraTemplateAssetPicker::OnCompareItemsForSorting(const FAssetData& ItemA, const FAssetData& ItemB) const
+{
+	return ItemA.AssetName.ToString().Compare(ItemB.AssetName.ToString()) == -1;
+}
+
+bool SNiagaraTemplateAssetPicker::OnDoesItemMatchFilterText(const FText& FilterText, const FAssetData& Item)
+{
+	return Item.AssetName.ToString().Contains(FilterText.ToString());
+}
+
+TSharedRef<SWidget> SNiagaraTemplateAssetPicker::OnGenerateWidgetForCategory(const FText& Category)
+{
+	return SNew(SBox)
+		.Padding(FMargin(5, 5, 5, 3))
+		[
+			SNew(STextBlock)
+			.TextStyle(FNiagaraEditorStyle::Get(), "NiagaraEditor.AssetPickerAssetCategoryText")
+			.Text(Category)
+		];
+}
+
+const int32 ThumbnailSize = 64;
+
+TSharedRef<SWidget> SNiagaraTemplateAssetPicker::OnGenerateWidgetForItem(const FAssetData& Item)
+{
+	TSharedRef<FAssetThumbnail> AssetThumbnail = MakeShared<FAssetThumbnail>(Item, ThumbnailSize, ThumbnailSize, AssetThumbnailPool);
+	FAssetThumbnailConfig ThumbnailConfig;
+	ThumbnailConfig.bAllowFadeIn = true;
+
+	FText AssetDescription;
+	Item.GetTagValue("TemplateAssetDescription", AssetDescription);
+
+	return
+		SNew(SVerticalBox)
+		+ SVerticalBox::Slot()
+		.AutoHeight()
+		.Padding(5, 3, 5, 5)
+		[
+			SNew(STextBlock)
+			.TextStyle(FNiagaraEditorStyle::Get(), "NiagaraEditor.AssetPickerAssetNameText")
+			.Text(FText::FromString(FName::NameToDisplayString(Item.AssetName.ToString(), false)))
+		]
+		+ SVerticalBox::Slot()
+		.Padding(5, 0, 5, 5)
+		[
+			SNew(SHorizontalBox)
+			+ SHorizontalBox::Slot()
+			.AutoWidth()
+			.Padding(0, 0, 10, 0)
+			[
+				SNew(SBox)
+				.WidthOverride(ThumbnailSize)
+				.HeightOverride(ThumbnailSize)
+				[
+					AssetThumbnail->MakeThumbnailWidget(ThumbnailConfig)
+				]
+			]
+			+ SHorizontalBox::Slot()
+			[
+				SNew(STextBlock)
+				.Text(AssetDescription)
+				.AutoWrapText(true)
+			]
+		];
+}
+
+#undef LOCTEXT_NAMESPACE
diff --git a/Engine/Plugins/FX/Niagara/Source/NiagaraEditor/Private/SNiagaraTemplateAssetPicker.h b/Engine/Plugins/FX/Niagara/Source/NiagaraEditor/Private/SNiagaraTemplateAssetPicker.h
new file mode 100644
--- /dev/null
+++ b/Engine/Plugins/FX/Niagara/Source/NiagaraEditor/Private/SNiagaraTemplateAssetPicker.h
@@ -0,0 +1,44 @@
+// Copyright 1998-2018 Epic Games, Inc. All Rights Reserved.
+
+#pragma once
+
+#include "SItemSelector.h"
+#include "AssetData.h"
+
+typedef SItemSelector<FText, FAssetData> SNiagaraAssetItemSelector;
+
+class FAssetThumbnailPool;
+
+class SNiagaraTemplateAssetPicker : public SCompoundWidget
+{
+public:
+	SLATE_BEGIN_ARGS(SNiagaraTemplateAssetPicker) {}
+	SLATE_END_ARGS();
+
+	void Construct(const FArguments& InArgs, UClass* AssetClass);
+
+	virtual void Tick(const FGeometry& AllottedGeometry, const double InCurrentTime, const float InDeltaTime) override;
+
+	TArray<FAssetData> GetSelectedAssets() const;
+
+private:
+	TArray<FText> OnGetCategoriesForItem(const FAssetData& Item);
+
+	bool OnCompareCategoriesForEquality(const FText& CategoryA, const FText& CategoryB) const;
+
+	bool OnCompareCategoriesForSorting(const FText& CategoryA, const FText& CategoryB) const;
+
+	bool OnCompareItemsForSorting(const FAssetData& ItemA, const FAssetData& ItemB) const;
+
+	bool OnDoesItemMatchFilterText(const FText& FilterText, const FAssetData& Item);
+
+	TSharedRef<SWidget> OnGenerateWidgetForCategory(const FText& Category);
+
+	TSharedRef<SWidget> OnGenerateWidgetForItem(const FAssetData& Item);
+
+private:
+	TSharedPtr<SNiagaraAssetItemSelector> ItemSelector;
+	FText NiagaraPluginCategory;
+	FText ProjectCategory;
+	TSharedPtr<FAssetThumbnailPool> AssetThumbnailPool;
+};
\ No newline at end of file
diff --git a/Engine/Plugins/FX/Niagara/Source/NiagaraEditor/Public/NiagaraEditorSettings.h b/Engine/Plugins/FX/Niagara/Source/NiagaraEditor/Public/NiagaraEditorSettings.h
--- a/Engine/Plugins/FX/Niagara/Source/NiagaraEditor/Public/NiagaraEditorSettings.h
+++ b/Engine/Plugins/FX/Niagara/Source/NiagaraEditor/Public/NiagaraEditorSettings.h
@@ -9,20 +9,29 @@
 #include "NiagaraSpawnShortcut.h"
 #include "NiagaraEditorSettings.generated.h"
 
+USTRUCT()
+struct FNiagaraNewAssetDialogConfig
+{
+	GENERATED_BODY()
+
+	UPROPERTY()
+	int32 SelectedOptionIndex;
+
+	UPROPERTY()
+	FVector2D WindowSize;
+
+	FNiagaraNewAssetDialogConfig()
+	{
+		SelectedOptionIndex = 0;
+		WindowSize = FVector2D(450, 600);
+	}
+};
 
 UCLASS(config = Niagara, defaultconfig, meta=(DisplayName="Niagara"))
 class UNiagaraEditorSettings : public UDeveloperSettings
 {
 public:
 	GENERATED_UCLASS_BODY()
-		
-	/** System to duplicate as the base of all new System assets created. */
-	UPROPERTY(config, EditAnywhere, Category = Niagara)
-	FSoftObjectPath DefaultSystem;
-
-	/** Emitter to duplicate as the base of all new emitter assets created. */
-	UPROPERTY(config, EditAnywhere, Category = Niagara)
-	FSoftObjectPath DefaultEmitter;
 
 	/** Niagara script to duplicate as the base of all new script assets created. */
 	UPROPERTY(config, EditAnywhere, Category = Niagara)
@@ -67,6 +76,10 @@ public:
 
 	/** Sets whether or not to rerun the simulation to the current time when making modifications while paused. */
 	void SetResimulateOnChangeWhilePaused(bool bInResimulateOnChangeWhilePaused);
+
+	FNiagaraNewAssetDialogConfig GetNewAssetDailogConfig(FName InDialogConfigKey) const;
+
+	void SetNewAssetDialogConfig(FName InDialogConfigKey, const FNiagaraNewAssetDialogConfig& InNewAssetDialogConfig);
 	
 	// Begin UDeveloperSettings Interface
 	virtual FName GetCategoryName() const override;
@@ -99,4 +112,7 @@ private:
 	/** Whether or not to rerun the simulation to the current time when making modifications while paused. */
 	UPROPERTY(config, EditAnywhere, Category = SimulationOptions)
 	bool bResimulateOnChangeWhilePaused;
+
+	UPROPERTY(config)
+	TMap<FName, FNiagaraNewAssetDialogConfig> NewAssetDialogConfigMap;
 };
diff --git a/Engine/Plugins/FX/Niagara/Source/NiagaraEditor/Public/NiagaraEmitterFactoryNew.h b/Engine/Plugins/FX/Niagara/Source/NiagaraEditor/Public/NiagaraEmitterFactoryNew.h
--- a/Engine/Plugins/FX/Niagara/Source/NiagaraEditor/Public/NiagaraEmitterFactoryNew.h
+++ b/Engine/Plugins/FX/Niagara/Source/NiagaraEditor/Public/NiagaraEmitterFactoryNew.h
@@ -13,7 +13,11 @@ class UNiagaraEmitterFactoryNew : public UFactory
 {
 	GENERATED_UCLASS_BODY()
 
+	UNiagaraEmitter* EmitterToCopy;
+	bool bAddDefaultModulesAndRenderersToEmptyEmitter;
+
 	//~ Begin UFactory Interface
+	virtual bool ConfigureProperties() override;
 	virtual UObject* FactoryCreateNew(UClass* Class, UObject* InParent, FName Name, EObjectFlags Flags, UObject* Context, FFeedbackContext* Warn) override;
 	//~ Begin UFactory Interface	
 };
diff --git a/Engine/Plugins/FX/Niagara/Source/NiagaraEditor/Public/NiagaraSystemFactoryNew.h b/Engine/Plugins/FX/Niagara/Source/NiagaraEditor/Public/NiagaraSystemFactoryNew.h
--- a/Engine/Plugins/FX/Niagara/Source/NiagaraEditor/Public/NiagaraSystemFactoryNew.h
+++ b/Engine/Plugins/FX/Niagara/Source/NiagaraEditor/Public/NiagaraSystemFactoryNew.h
@@ -8,13 +8,18 @@
 #include "NiagaraSystemFactoryNew.generated.h"
 
 class UNiagaraSystem;
+class UNiagaraEmitter;
 
 UCLASS(hidecategories = Object)
 class UNiagaraSystemFactoryNew : public UFactory
 {
 	GENERATED_UCLASS_BODY()
 
+	UNiagaraSystem* SystemToCopy;
+	TArray<UNiagaraEmitter*> EmittersToAddToNewSystem;
+
 	//~ Begin UFactory Interface
+	virtual bool ConfigureProperties() override;
 	virtual UObject* FactoryCreateNew(UClass* Class, UObject* InParent, FName Name, EObjectFlags Flags, UObject* Context, FFeedbackContext* Warn) override;
 	//~ Begin UFactory Interface	
 
diff --git a/Engine/Source/Editor/ContentBrowser/Private/SAssetView.cpp b/Engine/Source/Editor/ContentBrowser/Private/SAssetView.cpp
--- a/Engine/Source/Editor/ContentBrowser/Private/SAssetView.cpp
+++ b/Engine/Source/Editor/ContentBrowser/Private/SAssetView.cpp
@@ -946,6 +946,12 @@ void SAssetView::Tick( const FGeometry& AllottedGeometry, const double InCurrent
 
 	CalculateThumbnailHintColorAndOpacity();
 
+	if (FSlateApplication::Get().GetActiveModalWindow().IsValid())
+	{
+		// If we're in a model window then we need to tick the thumbnail pool in order for thumbnails to render correctly.
+		AssetThumbnailPool->Tick(InDeltaTime);
+	}
+
 	if ( bPendingUpdateThumbnails )
 	{
 		UpdateThumbnails();
diff --git a/Engine/Source/Runtime/Slate/Private/Widgets/Input/SCheckBox.cpp b/Engine/Source/Runtime/Slate/Private/Widgets/Input/SCheckBox.cpp
--- a/Engine/Source/Runtime/Slate/Private/Widgets/Input/SCheckBox.cpp
+++ b/Engine/Source/Runtime/Slate/Private/Widgets/Input/SCheckBox.cpp
@@ -37,6 +37,7 @@ void SCheckBox::Construct( const SCheckBox::FArguments& InArgs )
 	CheckBoxTypeOverride = InArgs._Type;
 
 	HorizontalAlignment = InArgs._HAlign;
+	bCheckBoxContentUsesAutoWidth = InArgs._CheckBoxContentUsesAutoWidth;
 
 	bIsPressed = false;
 	bIsFocusable = InArgs._IsFocusable;
@@ -395,6 +396,7 @@ void SCheckBox::BuildCheckBox(TSharedRef<SWidget> InContent)
 	if (CheckBoxType == ESlateCheckBoxType::CheckBox)
 	{
 		// Check boxes use a separate check button to the side of the user's content (often, a text label or icon.)
+		SHorizontalBox::FSlot* ContentSlot;
 		this->ChildSlot
 		[
 			SNew(SHorizontalBox)
@@ -408,9 +410,9 @@ void SCheckBox::BuildCheckBox(TSharedRef<SWidget> InContent)
 				.ColorAndOpacity(this, &SCheckBox::OnGetForegroundColor)
 			]
 			+ SHorizontalBox::Slot()
-			.AutoWidth()
 			.Padding(TAttribute<FMargin>(this, &SCheckBox::OnGetPadding))
 			.VAlign(VAlign_Center)
+			.Expose(ContentSlot)
 			[
 				SAssignNew(ContentContainer, SBorder)
 				.BorderImage(FStyleDefaults::GetNoBrush())
@@ -420,6 +422,10 @@ void SCheckBox::BuildCheckBox(TSharedRef<SWidget> InContent)
 				]
 			]
 		];
+		if (bCheckBoxContentUsesAutoWidth)
+		{
+			ContentSlot->AutoWidth();
+		}
 	}
 	else if (ensure(CheckBoxType == ESlateCheckBoxType::ToggleButton))
 	{
diff --git a/Engine/Source/Runtime/Slate/Public/Widgets/Input/SCheckBox.h b/Engine/Source/Runtime/Slate/Public/Widgets/Input/SCheckBox.h
--- a/Engine/Source/Runtime/Slate/Public/Widgets/Input/SCheckBox.h
+++ b/Engine/Source/Runtime/Slate/Public/Widgets/Input/SCheckBox.h
@@ -37,6 +37,7 @@ public:
 		, _OnCheckStateChanged()
 		, _IsChecked( ECheckBoxState::Unchecked )
 		, _HAlign( HAlign_Fill )
+		, _CheckBoxContentUsesAutoWidth(true)
 		, _Padding()
 		, _ClickMethod( EButtonClickMethod::DownAndUp )
 		, _ForegroundColor()
@@ -71,6 +72,9 @@ public:
 		/** How the content of the toggle button should align within the given space*/
 		SLATE_ARGUMENT( EHorizontalAlignment, HAlign )
 
+		/** Whether or not the content portion of the checkbox should layout using auto-width. When true the content will always be arranged at its desired size as opposed to resizing to the available space. */
+		SLATE_ARGUMENT(bool, CheckBoxContentUsesAutoWidth)
+
 		/** Spacing between the check box image and its content (set by the Style arg but the Style can be overridden with this) */
 		SLATE_ATTRIBUTE( FMargin, Padding )
 
@@ -278,6 +282,9 @@ protected:
 	/** Horiz align setting if in togglebox mode */
 	EHorizontalAlignment HorizontalAlignment;
 
+	/** Whether or not the checkbox content is arranged using auto-width when in checkbox mode. */
+	bool bCheckBoxContentUsesAutoWidth;
+
 	/** Sets whether a click should be triggered on mouse down, mouse up, or that both a mouse down and up are required. */
 	EButtonClickMethod::Type ClickMethod;
 
