Merging from Dev-Niagara 4356318 for nicholas.goldstein

Adjust DisplayName generation of linked inputs in Niagara Stack input options to show the entire namespace and to not pass through FName::NameToDisplayString

diff --git a/Engine/Plugins/FX/Niagara/Source/NiagaraEditorWidgets/Private/Stack/SNiagaraStackFunctionInputValue.cpp b/Engine/Plugins/FX/Niagara/Source/NiagaraEditorWidgets/Private/Stack/SNiagaraStackFunctionInputValue.cpp
--- a/Engine/Plugins/FX/Niagara/Source/NiagaraEditorWidgets/Private/Stack/SNiagaraStackFunctionInputValue.cpp
+++ b/Engine/Plugins/FX/Niagara/Source/NiagaraEditorWidgets/Private/Stack/SNiagaraStackFunctionInputValue.cpp
@@ -533,7 +533,7 @@ void SNiagaraStackFunctionInputValue::CollectAllActions(FGraphActionListBuilderB
 			const FText MapInputFormat = LOCTEXT("LinkInputFormat", "Link this input to {0}");
 			for (const FNiagaraParameterHandle& Handle : Handles)
 			{
-				const FText DisplayName = FText::FromString(FName::NameToDisplayString(Handle.GetName().ToString(), false));
+				const FText DisplayName = FText::FromName(Handle.GetParameterHandleString());
 				const FText Tooltip = FText::Format(MapInputFormat, FText::FromName(Handle.GetParameterHandleString()));
 				TSharedPtr<FNiagaraMenuAction> LinkAction(new FNiagaraMenuAction(SectionDisplay, DisplayName, Tooltip, 0, FText(),
 					FNiagaraMenuAction::FOnExecuteStackAction::CreateSP(this, &SNiagaraStackFunctionInputValue::ParameterHandleSelected, Handle)));
