Merging from Dev-Niagara 4356331 for nicholas.goldstein

Add Spacer between Event Handler header row and added Event Handler properties

diff --git a/Engine/Plugins/FX/Niagara/Source/NiagaraEditor/Private/ViewModels/Stack/NiagaraStackEventScriptItemGroup.cpp b/Engine/Plugins/FX/Niagara/Source/NiagaraEditor/Private/ViewModels/Stack/NiagaraStackEventScriptItemGroup.cpp
--- a/Engine/Plugins/FX/Niagara/Source/NiagaraEditor/Private/ViewModels/Stack/NiagaraStackEventScriptItemGroup.cpp
+++ b/Engine/Plugins/FX/Niagara/Source/NiagaraEditor/Private/ViewModels/Stack/NiagaraStackEventScriptItemGroup.cpp
@@ -187,6 +187,17 @@ void UNiagaraStackEventScriptItemGroup::Initialize(
 
 void UNiagaraStackEventScriptItemGroup::RefreshChildrenInternal(const TArray<UNiagaraStackEntry*>& CurrentChildren, TArray<UNiagaraStackEntry*>& NewChildren, TArray<FStackIssue>& NewIssues)
 {
+	FName EventSpacerKey = *FString::Printf(TEXT("EventSpacer"));
+	UNiagaraStackSpacer* SeparatorSpacer = FindCurrentChildOfTypeByPredicate<UNiagaraStackSpacer>(CurrentChildren,
+		[=](UNiagaraStackSpacer* CurrentSpacer) { return CurrentSpacer->GetSpacerKey() == EventSpacerKey; });
+	if (SeparatorSpacer == nullptr)
+	{
+		SeparatorSpacer = NewObject<UNiagaraStackSpacer>(this);
+		FRequiredEntryData RequiredEntryData(GetSystemViewModel(), GetEmitterViewModel(), GetExecutionCategoryName(), NAME_None, GetStackEditorData());
+		SeparatorSpacer->Initialize(RequiredEntryData, EventSpacerKey);
+	}
+	NewChildren.Add(SeparatorSpacer);
+
 	UNiagaraEmitter* Emitter = GetEmitterViewModel()->GetEmitter();
 
 	const FNiagaraEventScriptProperties* EventScriptProperties = Emitter->GetEventHandlers().FindByPredicate(
