Merging from Dev-Niagara 4357655 for nicholas.goldstein

Fix various capitalization inconsistencies in niagara stack menus

diff --git a/Engine/Plugins/FX/Niagara/Source/NiagaraEditorWidgets/Private/SNiagaraStack.cpp b/Engine/Plugins/FX/Niagara/Source/NiagaraEditorWidgets/Private/SNiagaraStack.cpp
--- a/Engine/Plugins/FX/Niagara/Source/NiagaraEditorWidgets/Private/SNiagaraStack.cpp
+++ b/Engine/Plugins/FX/Niagara/Source/NiagaraEditorWidgets/Private/SNiagaraStack.cpp
@@ -133,12 +133,12 @@ void SNiagaraStack::Construct(const FArguments& InArgs, UNiagaraStackViewModel*
 			MenuBuilder.BeginSection("EmitterInlineMenuActions", LOCTEXT("EmitterActions", "Emitter Actions"));
 			{
 				MenuBuilder.AddMenuEntry(
-					LOCTEXT("OpenAndFocusEmitter", "Open and focus source emitter"),
+					LOCTEXT("OpenAndFocusEmitter", "Open and Focus Source Emitter"),
 					LOCTEXT("OpenAndFocusEmitterToolTip", "Open the source emitter in emitter mode"),
 					FSlateIcon(),
 					FUIAction(FExecuteAction::CreateSP(this, &SNiagaraStack::OpenSourceEmitter), FCanExecuteAction::CreateSP(this, &SNiagaraStack::CanOpenSourceEmitter)));
 				MenuBuilder.AddMenuEntry(
-					LOCTEXT("ShowEmitterInContentBrowser", "Show in content browser"),
+					LOCTEXT("ShowEmitterInContentBrowser", "Show in Content Browser"),
 					LOCTEXT("ShowEmitterInContentBrowserToolTip", "Show the emitter in this stack in the Content Browser"),
 					FSlateIcon(),
 					FUIAction(FExecuteAction::CreateSP(this, &SNiagaraStack::ShowEmitterInContentBrowser))); 
@@ -198,7 +198,7 @@ void SNiagaraStack::Construct(const FArguments& InArgs, UNiagaraStackViewModel*
 			if (StackViewModel->HasDismissedStackIssues())
 			{
 				MenuBuilder.AddMenuEntry(
-					LOCTEXT("UndismissIssues", "Undismiss all stack issues"),
+					LOCTEXT("UndismissIssues", "Undismiss All Stack Issues"),
 					LOCTEXT("ShowAssetInContentBrowserToolTip", "Undismiss all issues that were previously dismissed for this stack, if any"),
 					FSlateIcon(),
 					FUIAction(FExecuteAction::CreateUObject(StackViewModel, &UNiagaraStackViewModel::UndismissAllIssues)));
@@ -507,7 +507,7 @@ TSharedRef<SWidget> SNiagaraStack::GetViewOptionsMenu() const
 	FMenuBuilder MenuBuilder(false, nullptr);
 
 	MenuBuilder.AddMenuEntry(
-		LOCTEXT("ShowAllAdvancedLabel", "Show all advanced"),
+		LOCTEXT("ShowAllAdvancedLabel", "Show All Advanced"),
 		LOCTEXT("ShowAllAdvancedToolTip", "Forces all advanced items to be showing in the stack."),
 		FSlateIcon(),
 		FUIAction(
@@ -517,7 +517,7 @@ TSharedRef<SWidget> SNiagaraStack::GetViewOptionsMenu() const
 		NAME_None, EUserInterfaceActionType::Check);
 
 	MenuBuilder.AddMenuEntry(
-		LOCTEXT("ShowOutputsLabel", "Show outputs"),
+		LOCTEXT("ShowOutputsLabel", "Show Outputs"),
 		LOCTEXT("ShowOutputsToolTip", "Whether or now to show module outputs in the stack."),
 		FSlateIcon(),
 		FUIAction(
diff --git a/Engine/Plugins/FX/Niagara/Source/NiagaraEditorWidgets/Private/Stack/SNiagaraStackModuleItem.cpp b/Engine/Plugins/FX/Niagara/Source/NiagaraEditorWidgets/Private/Stack/SNiagaraStackModuleItem.cpp
--- a/Engine/Plugins/FX/Niagara/Source/NiagaraEditorWidgets/Private/Stack/SNiagaraStackModuleItem.cpp
+++ b/Engine/Plugins/FX/Niagara/Source/NiagaraEditorWidgets/Private/Stack/SNiagaraStackModuleItem.cpp
@@ -143,7 +143,7 @@ bool SNiagaraStackModuleItem::CheckEnabledStatus(bool bIsEnabled)
 void SNiagaraStackModuleItem::FillRowContextMenu(FMenuBuilder& MenuBuilder)
 {
 	MenuBuilder.AddMenuEntry(
-		LOCTEXT("InsertModuleAbove", "Insert above"),
+		LOCTEXT("InsertModuleAbove", "Insert Above"),
 		LOCTEXT("InsertModuleAboveToolTip", "Insert a new module above this module in the stack."),
 		FSlateIcon(),
 		FUIAction(FExecuteAction::CreateSP(this, &SNiagaraStackModuleItem::InsertModuleAbove)));
diff --git a/Engine/Plugins/FX/Niagara/Source/NiagaraEditorWidgets/Private/Stack/SNiagaraStackTableRow.cpp b/Engine/Plugins/FX/Niagara/Source/NiagaraEditorWidgets/Private/Stack/SNiagaraStackTableRow.cpp
--- a/Engine/Plugins/FX/Niagara/Source/NiagaraEditorWidgets/Private/Stack/SNiagaraStackTableRow.cpp
+++ b/Engine/Plugins/FX/Niagara/Source/NiagaraEditorWidgets/Private/Stack/SNiagaraStackTableRow.cpp
@@ -118,7 +118,6 @@ void SNiagaraStackTableRow::SetNameAndValueContent(TSharedRef<SWidget> InNameWid
 			.ForegroundColor(FSlateColor::UseForeground())
 			.ContentPadding(2)
 			.HAlign(HAlign_Center)
-			.IsFocusable(false)
 			[
 				SNew(SImage)
 				.Image(this, &SNiagaraStackTableRow::GetExpandButtonImage)
@@ -270,7 +269,7 @@ FReply SNiagaraStackTableRow::OnMouseButtonUp(const FGeometry& MyGeometry, const
 		{
 			MenuBuilder.BeginSection("AssetActions", LOCTEXT("AssetActions", "Asset Actions"));
 			MenuBuilder.AddMenuEntry(
-				LOCTEXT("OpenAndFocusAsset", "Open and focus Asset"),
+				LOCTEXT("OpenAndFocusAsset", "Open and Focus Asset"),
 				FText::Format(LOCTEXT("OpenAndFocusAssetTooltip", "Open {0} in separate editor"), StackEntry->GetDisplayName()),
 				FSlateIcon(),
 				FUIAction(FExecuteAction::CreateSP(this, &SNiagaraStackTableRow::OpenSourceAsset)));
@@ -292,7 +291,7 @@ FReply SNiagaraStackTableRow::OnMouseButtonUp(const FGeometry& MyGeometry, const
 			{
 				MenuBuilder.BeginSection("StackRowNavigateTo", LOCTEXT("NavigateToSection", "Navigate to:"));
 				MenuBuilder.AddMenuEntry(
-					LOCTEXT("TopOfSection", "Top of section"),
+					LOCTEXT("TopOfSection", "Top of Section"),
 					FText::Format(LOCTEXT("NavigateToFormatted", "Navigate to {0}"), Parent->GetDisplayName()),
 					FSlateIcon(),
 					FUIAction(FExecuteAction::CreateSP(this, &SNiagaraStackTableRow::NavigateTo, Parent)));
@@ -300,7 +299,7 @@ FReply SNiagaraStackTableRow::OnMouseButtonUp(const FGeometry& MyGeometry, const
 			if (ItemParent != nullptr)
 			{
 				MenuBuilder.AddMenuEntry(
-					LOCTEXT("TopOfModule", "Top of module"),
+					LOCTEXT("TopOfModule", "Top of Module"),
 					FText::Format(LOCTEXT("NavigateToFormatted", "Navigate to {0}"), Parent->GetDisplayName()),
 					FSlateIcon(),
 					FUIAction(FExecuteAction::CreateSP(this, &SNiagaraStackTableRow::NavigateTo, Parent)));
@@ -313,12 +312,12 @@ FReply SNiagaraStackTableRow::OnMouseButtonUp(const FGeometry& MyGeometry, const
 
 		MenuBuilder.BeginSection("StackActions", LOCTEXT("StackActions", "Stack Actions"));
 		MenuBuilder.AddMenuEntry(
-			LOCTEXT("ExpandAllItems", "Expand all"),
+			LOCTEXT("ExpandAllItems", "Expand All"),
 			LOCTEXT("ExpandAllItemsToolTip", "Expand all items under this header."),
 			FSlateIcon(),
 			FUIAction(FExecuteAction::CreateSP(this, &SNiagaraStackTableRow::ExpandChildren)));
 		MenuBuilder.AddMenuEntry(
-			LOCTEXT("CollapseAllItems", "Collapse all"),
+			LOCTEXT("CollapseAllItems", "Collapse All"),
 			LOCTEXT("CollapseAllItemsToolTip", "Collapse all items under this header."),
 			FSlateIcon(),
 			FUIAction(FExecuteAction::CreateSP(this, &SNiagaraStackTableRow::CollapseChildren)));
