Merging from Dev-Niagara 4397394 for nicholas.goldstein

Add missing entries and param creation subcategory to Set Variables module.

diff --git a/Engine/Plugins/FX/Niagara/Source/Niagara/Public/NiagaraTypes.h b/Engine/Plugins/FX/Niagara/Source/Niagara/Public/NiagaraTypes.h
--- a/Engine/Plugins/FX/Niagara/Source/Niagara/Public/NiagaraTypes.h
+++ b/Engine/Plugins/FX/Niagara/Source/Niagara/Public/NiagaraTypes.h
@@ -818,7 +818,7 @@ struct FNiagaraVariable
 		return Name != NAME_None && TypeDef.IsValid();
 	}
 
-	FORCEINLINE bool IsInNameSpace(FString Namespace) 
+	FORCEINLINE bool IsInNameSpace(FString Namespace) const
 	{
 		return Name.ToString().StartsWith(Namespace + TEXT("."));
 	}
diff --git a/Engine/Plugins/FX/Niagara/Source/NiagaraEditor/Private/NiagaraGraph.cpp b/Engine/Plugins/FX/Niagara/Source/NiagaraEditor/Private/NiagaraGraph.cpp
--- a/Engine/Plugins/FX/Niagara/Source/NiagaraEditor/Private/NiagaraGraph.cpp
+++ b/Engine/Plugins/FX/Niagara/Source/NiagaraEditor/Private/NiagaraGraph.cpp
@@ -512,7 +512,7 @@ void UNiagaraGraph::GetParameters(TArray<FNiagaraVariable>& Inputs, TArray<FNiag
 // 	Outputs.Sort(SortVars);
 }
 
-const TMap<FNiagaraVariable, FNiagaraGraphParameterReferenceCollection>& UNiagaraGraph::GetParameterMap()
+const TMap<FNiagaraVariable, FNiagaraGraphParameterReferenceCollection>& UNiagaraGraph::GetParameterMap() const
 {
 	return Parameters;
 }
diff --git a/Engine/Plugins/FX/Niagara/Source/NiagaraEditor/Private/NiagaraNodeAssignment.cpp b/Engine/Plugins/FX/Niagara/Source/NiagaraEditor/Private/NiagaraNodeAssignment.cpp
--- a/Engine/Plugins/FX/Niagara/Source/NiagaraEditor/Private/NiagaraNodeAssignment.cpp
+++ b/Engine/Plugins/FX/Niagara/Source/NiagaraEditor/Private/NiagaraNodeAssignment.cpp
@@ -19,6 +19,7 @@
 #include "ViewModels/Stack/NiagaraParameterHandle.h"
 #include "Framework/MultiBox/MultiBoxBuilder.h"
 #include "ScopedTransaction.h"
+#include "ViewModels/Stack/INiagaraStackItemGroupAddUtilities.h"
 #include "ViewModels/Stack/NiagaraStackGraphUtilities.h"
 
 #define LOCTEXT_NAMESPACE "NiagaraNodeAssigment"
@@ -194,29 +195,10 @@ void UNiagaraNodeAssignment::MergeUp()
 	//NiagaraStackUtilities::
 }
 
-void UNiagaraNodeAssignment::BuildParameterMenu(FMenuBuilder& MenuBuilder, ENiagaraScriptUsage InUsage, UNiagaraNodeOutput* InGraphOutputNode)
+void UNiagaraNodeAssignment::BuildAddParameterMenu(FMenuBuilder& MenuBuilder, ENiagaraScriptUsage InUsage, UNiagaraNodeOutput* InGraphOutputNode)
 {
-	TArray<FNiagaraParameterMapHistory> Histories = UNiagaraNodeParameterMapBase::GetParameterMaps(InGraphOutputNode->GetNiagaraGraph());
 	TArray<FNiagaraVariable> AvailableParameters;
-
-	if (InGraphOutputNode->GetUsage() == ENiagaraScriptUsage::ParticleSpawnScript ||
-		InGraphOutputNode->GetUsage() == ENiagaraScriptUsage::ParticleSpawnScriptInterpolated ||
-		InGraphOutputNode->GetUsage() == ENiagaraScriptUsage::ParticleUpdateScript ||
-		InGraphOutputNode->GetUsage() == ENiagaraScriptUsage::ParticleEventScript)
-	{
-		AvailableParameters.Append(FNiagaraConstants::GetCommonParticleAttributes());
-	}
-
-	for (FNiagaraParameterMapHistory& History : Histories)
-	{
-		for (FNiagaraVariable& Variable : History.Variables)
-		{
-			if (History.IsPrimaryDataSetOutput(Variable, InGraphOutputNode->GetUsage()))
-			{
-				AvailableParameters.AddUnique(Variable);
-			}
-		}
-	}
+	FNiagaraStackGraphUtilities::GetAvailableParametersForScript(*InGraphOutputNode, AvailableParameters);
 
 	for (const FNiagaraVariable& AvailableParameter : AvailableParameters)
 	{
@@ -226,7 +208,7 @@ void UNiagaraNodeAssignment::BuildParameterMenu(FMenuBuilder& MenuBuilder, ENiag
 		FString VarDefaultValue = FNiagaraConstants::GetAttributeDefaultValue(AvailableParameter);
 		const FText TooltipDesc = FText::Format(LOCTEXT("SetFunctionPopupTooltip", "Description: Set the parameter {0}. {1}"), FText::FromName(AvailableParameter.GetName()), VarDesc);
 		FText CategoryName = LOCTEXT("ModuleSetCategory", "Set Specific Parameters");
-		bool bCanExecute = AssignmentTargets.Contains(AvailableParameter) == false;
+		bool bCanExecute = AssignmentTargets.Contains(AvailableParameter) == false; 
 
 		MenuBuilder.AddMenuEntry(
 			NameText,
@@ -239,6 +221,55 @@ void UNiagaraNodeAssignment::BuildParameterMenu(FMenuBuilder& MenuBuilder, ENiag
 
 }
 
+void UNiagaraNodeAssignment::BuildCreateParameterMenu(FMenuBuilder& MenuBuilder, ENiagaraScriptUsage InUsage, UNiagaraNodeOutput* InGraphOutputNode)
+{
+	// Generate actions for creating new typed parameters.
+	TOptional<FName> NewParameterNamespace = FNiagaraStackGraphUtilities::GetNamespaceForScriptUsage(InGraphOutputNode->GetUsage());
+	if (NewParameterNamespace.IsSet())
+	{
+		// Collect all parameter names for ensuring new param has unique name
+		TArray<const UNiagaraGraph*> Graphs;
+		InGraphOutputNode->GetNiagaraGraph()->GetAllReferencedGraphs(Graphs);
+		TSet<FName> Names;
+		for (const UNiagaraGraph* Graph : Graphs)
+		{
+			for (const TPair<FNiagaraVariable, FNiagaraGraphParameterReferenceCollection>& ParameterElement : Graph->GetParameterMap())
+			{
+				Names.Add(ParameterElement.Key.GetName());
+			}
+		}
+
+		TArray<FNiagaraTypeDefinition> AvailableTypes;
+		FNiagaraStackGraphUtilities::GetNewParameterAvailableTypes(AvailableTypes);
+		for (const FNiagaraTypeDefinition& AvailableType : AvailableTypes)
+		{
+			// Make generic new parameter name
+			const FString NewParameterNameString = NewParameterNamespace.GetValue().ToString() + ".New" + AvailableType.GetName();
+			const FName NewParameterName = FName(*NewParameterNameString);
+
+			// Make NewParameterName unique  
+			const FName UniqueNewParameterName = FNiagaraUtilities::GetUniqueName(NewParameterName, Names);
+
+			// Create the new param
+			FNiagaraVariable NewParameter = FNiagaraVariable(AvailableType, UniqueNewParameterName);
+			FString VarDefaultValue = FNiagaraConstants::GetAttributeDefaultValue(NewParameter);
+			
+			// Tooltip and menu entry Text
+			FText VarDesc = FNiagaraConstants::GetAttributeDescription(NewParameter);
+			const FText TypeText = AvailableType.GetNameText();
+			const FText TooltipDesc = FText::Format(LOCTEXT("NewParameterModuleDescriptionFormat", "Description: Create a new {0} parameter. {1}"), TypeText, VarDesc);
+			FText Category = LOCTEXT("NewParameterModuleCategory", "Create New Parameter");
+
+			MenuBuilder.AddMenuEntry(
+				TypeText,
+				TooltipDesc,
+				FSlateIcon(),
+				FUIAction(
+					FExecuteAction::CreateUObject(this, &UNiagaraNodeAssignment::AddParameter, NewParameter, VarDefaultValue)));
+		}
+	}
+}
+
 void UNiagaraNodeAssignment::AddParameter(FNiagaraVariable InVar, FString InDefaultValue)
 {
 	const FText TransactionDesc = FText::Format(LOCTEXT("SetFunctionTransactionDesc", "Add the parameter {0}."), FText::FromName(InVar.GetName()));
diff --git a/Engine/Plugins/FX/Niagara/Source/NiagaraEditor/Private/ViewModels/Stack/NiagaraStackGraphUtilities.cpp b/Engine/Plugins/FX/Niagara/Source/NiagaraEditor/Private/ViewModels/Stack/NiagaraStackGraphUtilities.cpp
--- a/Engine/Plugins/FX/Niagara/Source/NiagaraEditor/Private/ViewModels/Stack/NiagaraStackGraphUtilities.cpp
+++ b/Engine/Plugins/FX/Niagara/Source/NiagaraEditor/Private/ViewModels/Stack/NiagaraStackGraphUtilities.cpp
@@ -1315,6 +1315,22 @@ void FNiagaraStackGraphUtilities::GetAvailableParametersForScript(UNiagaraNodeOu
 			}
 		}
 	}
+
+	TOptional<FName> UsageNamespace = FNiagaraStackGraphUtilities::GetNamespaceForScriptUsage(ScriptOutputNode.GetUsage());
+	if (UsageNamespace.IsSet())
+	{
+		for (const TPair<FNiagaraVariable, FNiagaraGraphParameterReferenceCollection>& Entry : ScriptOutputNode.GetNiagaraGraph()->GetParameterMap())
+		{
+			// Pick up any params with 0 references from the Parameters window
+			bool bDoesParamHaveNoReferences = Entry.Value.ParameterReferences.Num() == 0;
+			bool bIsParamInUsageNamespace = Entry.Key.IsInNameSpace(UsageNamespace.GetValue().ToString());
+
+			if (bDoesParamHaveNoReferences && bIsParamInUsageNamespace)
+			{
+				OutAvailableParameters.AddUnique(Entry.Key);
+			}
+		}
+	}
 }
 
 TOptional<FName> FNiagaraStackGraphUtilities::GetNamespaceForScriptUsage(ENiagaraScriptUsage ScriptUsage)
diff --git a/Engine/Plugins/FX/Niagara/Source/NiagaraEditor/Public/NiagaraGraph.h b/Engine/Plugins/FX/Niagara/Source/NiagaraEditor/Public/NiagaraGraph.h
--- a/Engine/Plugins/FX/Niagara/Source/NiagaraEditor/Public/NiagaraGraph.h
+++ b/Engine/Plugins/FX/Niagara/Source/NiagaraEditor/Public/NiagaraGraph.h
@@ -208,7 +208,7 @@ class UNiagaraGraph : public UEdGraph
 	const TMap<FNiagaraVariable, FNiagaraVariableMetaData>& GetAllMetaData() const { return VariableToMetaData; }
 	TMap<FNiagaraVariable, FNiagaraVariableMetaData>& GetAllMetaData() { return VariableToMetaData; }
 
-	const TMap<FNiagaraVariable, FNiagaraGraphParameterReferenceCollection>& GetParameterMap();
+	const TMap<FNiagaraVariable, FNiagaraGraphParameterReferenceCollection>& GetParameterMap() const;
 
 	/** Adds parameter to parameters map setting it as created by the user.*/
 	void AddParameter(const FNiagaraVariable& Parameter);
diff --git a/Engine/Plugins/FX/Niagara/Source/NiagaraEditor/Public/NiagaraNodeAssignment.h b/Engine/Plugins/FX/Niagara/Source/NiagaraEditor/Public/NiagaraNodeAssignment.h
--- a/Engine/Plugins/FX/Niagara/Source/NiagaraEditor/Public/NiagaraNodeAssignment.h
+++ b/Engine/Plugins/FX/Niagara/Source/NiagaraEditor/Public/NiagaraNodeAssignment.h
@@ -34,7 +34,8 @@ public:
 	const TArray<FString>& GetAssignmentDefaults() const { return AssignmentDefaultValues; }
 
 	void MergeUp();
-	void BuildParameterMenu(FMenuBuilder& MenuBuilder, ENiagaraScriptUsage InUsage, UNiagaraNodeOutput* InGraphOutputNode);
+	void BuildAddParameterMenu(FMenuBuilder& MenuBuilder, ENiagaraScriptUsage InUsage, UNiagaraNodeOutput* InGraphOutputNode);
+	void BuildCreateParameterMenu(FMenuBuilder& MenuBuilder, ENiagaraScriptUsage InUsage, UNiagaraNodeOutput* InGraphOutputNode);
 
 	//~ Begin EdGraphNode Interface
 	virtual void PostLoad() override;
diff --git a/Engine/Plugins/FX/Niagara/Source/NiagaraEditorWidgets/Private/Stack/SNiagaraStackModuleItem.cpp b/Engine/Plugins/FX/Niagara/Source/NiagaraEditorWidgets/Private/Stack/SNiagaraStackModuleItem.cpp
--- a/Engine/Plugins/FX/Niagara/Source/NiagaraEditorWidgets/Private/Stack/SNiagaraStackModuleItem.cpp
+++ b/Engine/Plugins/FX/Niagara/Source/NiagaraEditorWidgets/Private/Stack/SNiagaraStackModuleItem.cpp
@@ -223,11 +223,17 @@ TSharedRef<SWidget> SNiagaraStackModuleItem::RaiseActionMenuClicked()
 					LOCTEXT("MergeToolTip", "If a Set Variables node precedes this one in the stack, merge this node (and all variable binding logic) into that stack."),
 					FSlateIcon(),
 					FUIAction(FExecuteAction::CreateUObject(AssignmentNode, &UNiagaraNodeAssignment::MergeUp)));*/
-				MenuBuilder.AddSubMenu(LOCTEXT("AddVariables", "Add Variable"),
+				MenuBuilder.AddSubMenu(LOCTEXT("AddVariables", "Add Variable"), 
 					LOCTEXT("AddVariablesTooltip", "Add another variable to the end of the list"),
 					FNewMenuDelegate::CreateLambda([OutputNode, AssignmentNode](FMenuBuilder& SubMenuBuilder)
 				{
-					AssignmentNode->BuildParameterMenu(SubMenuBuilder, OutputNode->GetUsage(), OutputNode);
+					AssignmentNode->BuildAddParameterMenu(SubMenuBuilder, OutputNode->GetUsage(), OutputNode);
+				}));
+				MenuBuilder.AddSubMenu(LOCTEXT("CreateVariables", "Create New Variable"),
+					LOCTEXT("CreateVariablesTooltip", "Create a new variable and set its value"),
+					FNewMenuDelegate::CreateLambda([OutputNode, AssignmentNode](FMenuBuilder& SubMenuBuilder)
+				{
+					AssignmentNode->BuildCreateParameterMenu(SubMenuBuilder, OutputNode->GetUsage(), OutputNode);
 				}));
 
 				return MenuBuilder.MakeWidget();
