Merging from Dev-Niagara 4399353 for nicholas.goldstein

Improve discoverability of and navigability to parent Emitters in Systems

diff --git a/Engine/Plugins/FX/Niagara/Source/NiagaraEditor/Private/NiagaraEditorStyle.cpp b/Engine/Plugins/FX/Niagara/Source/NiagaraEditor/Private/NiagaraEditorStyle.cpp
--- a/Engine/Plugins/FX/Niagara/Source/NiagaraEditor/Private/NiagaraEditorStyle.cpp
+++ b/Engine/Plugins/FX/Niagara/Source/NiagaraEditor/Private/NiagaraEditorStyle.cpp
@@ -137,7 +137,15 @@ TSharedRef< FSlateStyleSet > FNiagaraEditorStyle::Create()
 		.SetShadowColorAndOpacity(FLinearColor(0, 0, 0, 0.9f));
 	
 	Style->Set("NiagaraEditor.AttributeSpreadsheetTabText", TabText);
+
+	FTextBlockStyle SubduedHeadingText = FTextBlockStyle(NormalText)
+		.SetFont(DEFAULT_FONT("Regular", 14))
+		.SetShadowOffset(FVector2D(0, 1))
+		.SetColorAndOpacity(FSlateColor(FLinearColor(0.5f, 0.5f, 0.5f, 1.0f)))
+		.SetShadowColorAndOpacity(FLinearColor(0, 0, 0, 0.9f));
 	
+	Style->Set("NiagaraEditor.SubduedHeadingTextBox", SubduedHeadingText);
+
 	// Parameters
 	FSlateFontInfo ParameterFont = DEFAULT_FONT("Regular", 8);
 
diff --git a/Engine/Plugins/FX/Niagara/Source/NiagaraEditor/Private/ViewModels/NiagaraEmitterHandleViewModel.cpp b/Engine/Plugins/FX/Niagara/Source/NiagaraEditor/Private/ViewModels/NiagaraEmitterHandleViewModel.cpp
--- a/Engine/Plugins/FX/Niagara/Source/NiagaraEditor/Private/ViewModels/NiagaraEmitterHandleViewModel.cpp
+++ b/Engine/Plugins/FX/Niagara/Source/NiagaraEditor/Private/ViewModels/NiagaraEmitterHandleViewModel.cpp
@@ -162,6 +162,24 @@ FText FNiagaraEmitterHandleViewModel::GetNameText() const
 	return FText();
 }
 
+FText FNiagaraEmitterHandleViewModel::GetSourceNameText() const
+{
+	if (EmitterHandle && EmitterHandle->GetSource())
+	{
+		return FText::FromString(EmitterHandle->GetSource()->GetName()); 
+	}
+	return FText();
+}
+
+FText FNiagaraEmitterHandleViewModel::GetSourcePathNameText() const
+{
+	if (EmitterHandle && EmitterHandle->GetSource())
+	{
+		return FText::FromString(EmitterHandle->GetSource()->GetPathName());
+	}
+	return FText();
+}
+
 void FNiagaraEmitterHandleViewModel::OnNameTextComitted(const FText& InText, ETextCommit::Type CommitInfo)
 {
 	SetName(*InText.ToString());
diff --git a/Engine/Plugins/FX/Niagara/Source/NiagaraEditor/Public/ViewModels/NiagaraEmitterHandleViewModel.h b/Engine/Plugins/FX/Niagara/Source/NiagaraEditor/Public/ViewModels/NiagaraEmitterHandleViewModel.h
--- a/Engine/Plugins/FX/Niagara/Source/NiagaraEditor/Public/ViewModels/NiagaraEmitterHandleViewModel.h
+++ b/Engine/Plugins/FX/Niagara/Source/NiagaraEditor/Public/ViewModels/NiagaraEmitterHandleViewModel.h
@@ -47,6 +47,12 @@ public:
 	/** Gets the text representation of the emitter handle name. */
 	NIAGARAEDITOR_API FText GetNameText() const;
 
+	/** Gets the text representation of the source emitter handle name. */
+	NIAGARAEDITOR_API FText GetSourceNameText() const;
+
+	/** Gets the text representation of the emitter handle name. */
+	NIAGARAEDITOR_API FText GetSourcePathNameText() const;
+
 	/** Called when the contents of the name text control is committed. */
 	NIAGARAEDITOR_API void OnNameTextComitted(const FText& InText, ETextCommit::Type CommitInfo);
 
diff --git a/Engine/Plugins/FX/Niagara/Source/NiagaraEditor/Public/ViewModels/NiagaraSystemViewModel.h b/Engine/Plugins/FX/Niagara/Source/NiagaraEditor/Public/ViewModels/NiagaraSystemViewModel.h
--- a/Engine/Plugins/FX/Niagara/Source/NiagaraEditor/Public/ViewModels/NiagaraSystemViewModel.h
+++ b/Engine/Plugins/FX/Niagara/Source/NiagaraEditor/Public/ViewModels/NiagaraSystemViewModel.h
@@ -34,7 +34,7 @@ class UNiagaraNodeFunctionCall;
 class FNiagaraEmitterViewModel;
 
 /** Defines different editing modes for this system view model. */
-enum class ENiagaraSystemViewModelEditMode
+NIAGARAEDITOR_API enum class ENiagaraSystemViewModelEditMode
 {
 	/** A system asset is being edited.  This assumes that emitters should be inheriting from a base version and that emitter editing will be restricted. */
 	SystemAsset,
@@ -126,7 +126,7 @@ public:
 	bool GetCanModifyEmittersFromTimeline() const;
 
 	/** Gets the current editing mode for this system view model. */
-	ENiagaraSystemViewModelEditMode GetEditMode() const;
+	NIAGARAEDITOR_API ENiagaraSystemViewModelEditMode GetEditMode() const;
 
 	/** Adds a new emitter to the System from an emitter asset data. */
 	void AddEmitterFromAssetData(const FAssetData& AssetData);
diff --git a/Engine/Plugins/FX/Niagara/Source/NiagaraEditorWidgets/Private/NiagaraEditorWidgetsStyle.cpp b/Engine/Plugins/FX/Niagara/Source/NiagaraEditorWidgets/Private/NiagaraEditorWidgetsStyle.cpp
--- a/Engine/Plugins/FX/Niagara/Source/NiagaraEditorWidgets/Private/NiagaraEditorWidgetsStyle.cpp
+++ b/Engine/Plugins/FX/Niagara/Source/NiagaraEditorWidgets/Private/NiagaraEditorWidgetsStyle.cpp
@@ -50,6 +50,7 @@ const FVector2D Icon8x16(8.0f, 16.0f);
 const FVector2D Icon12x12(12.0f, 12.0f);
 const FVector2D Icon16x16(16.0f, 16.0f);
 const FVector2D Icon20x20(20.0f, 20.0f);
+const FVector2D Icon30x30(30.0f, 30.0f);
 const FVector2D Icon40x40(40.0f, 40.0f);
 
 TSharedRef< FSlateStyleSet > FNiagaraEditorWidgetsStyle::Create()
@@ -121,6 +122,7 @@ TSharedRef< FSlateStyleSet > FNiagaraEditorWidgetsStyle::Create()
 	Style->Set("NiagaraEditor.Stack.IconColor.Particle", FLinearColor(FColor(131, 228, 9)));
 	Style->Set("NiagaraEditor.Stack.IconColor.Render", FLinearColor(FColor(230, 102, 102)));
 
+	Style->Set("NiagaraEditor.Stack.GoToSourceIcon", new IMAGE_CORE_BRUSH("Common/GoToSource", Icon30x30, FLinearColor::White));
 	Style->Set("NiagaraEditor.Stack.ParametersIcon", new IMAGE_PLUGIN_BRUSH("Icons/SystemParams", Icon12x12, FLinearColor::White));
 	Style->Set("NiagaraEditor.Stack.SpawnIcon", new IMAGE_PLUGIN_BRUSH("Icons/Spawn", Icon12x12, FLinearColor::White));
 	Style->Set("NiagaraEditor.Stack.UpdateIcon", new IMAGE_PLUGIN_BRUSH("Icons/Update", Icon12x12, FLinearColor::White));
diff --git a/Engine/Plugins/FX/Niagara/Source/NiagaraEditorWidgets/Private/SNiagaraStack.cpp b/Engine/Plugins/FX/Niagara/Source/NiagaraEditorWidgets/Private/SNiagaraStack.cpp
--- a/Engine/Plugins/FX/Niagara/Source/NiagaraEditorWidgets/Private/SNiagaraStack.cpp
+++ b/Engine/Plugins/FX/Niagara/Source/NiagaraEditorWidgets/Private/SNiagaraStack.cpp
@@ -51,6 +51,7 @@
 #include "Stack/SNiagaraStackErrorItem.h"
 #include "NiagaraStackEditorData.h"
 #include "ScopedTransaction.h"
+#include "Widgets/Layout/SWrapBox.h"
 
 /** Contains data for a socket drag and drop operation in the StackEntry node. */
 class FNiagaraStackEntryDragDropOp : public FDecoratedDragDropOp
@@ -133,10 +134,11 @@ void SNiagaraStack::Construct(const FArguments& InArgs, UNiagaraStackViewModel*
 			MenuBuilder.BeginSection("EmitterInlineMenuActions", LOCTEXT("EmitterActions", "Emitter Actions"));
 			{
 				MenuBuilder.AddMenuEntry(
-					LOCTEXT("OpenAndFocusEmitter", "Open and Focus Source Emitter"),
-					LOCTEXT("OpenAndFocusEmitterToolTip", "Open the source emitter in emitter mode"),
+					LOCTEXT("RenameEmitter", "Rename Emitter"),
+					LOCTEXT("RenameEmitterToolTip", "Rename this emitter"), 
 					FSlateIcon(),
-					FUIAction(FExecuteAction::CreateSP(this, &SNiagaraStack::OpenSourceEmitter), FCanExecuteAction::CreateSP(this, &SNiagaraStack::CanOpenSourceEmitter)));
+					FUIAction(FExecuteAction::CreateSP(InlineEditableTextBlock.Get(), &SInlineEditableTextBlock::EnterEditingMode)));
+
 				MenuBuilder.AddMenuEntry(
 					LOCTEXT("ShowEmitterInContentBrowser", "Show in Content Browser"),
 					LOCTEXT("ShowEmitterInContentBrowserToolTip", "Show the emitter in this stack in the Content Browser"),
@@ -168,7 +170,8 @@ void SNiagaraStack::Construct(const FArguments& InArgs, UNiagaraStackViewModel*
 					FUIAction(FExecuteAction::CreateUObject(StackViewModel, &UNiagaraStackViewModel::CollapseToHeaders)));
 			}
 			MenuBuilder.EndSection();
-			MenuBuilder.BeginSection("EmitterNavigateTo", LOCTEXT("EmitterNavigateTo", "Navigate to Section:"));
+			MenuBuilder.BeginSection("EmitterNavigateTo", 
+				FText::Format(LOCTEXT("EmitterNavigateTo", "Navigate to {0} Section:"), StackViewModel->GetEmitterHandleViewModel()->GetNameText()));
 			{
 				// Parse all children of root entries and if they are UNiagaraStackItemGroup then add a navigate menu entry
 				TArray<UNiagaraStackEntry*> EntriesToProcess(StackViewModel->GetRootEntries());
@@ -245,7 +248,7 @@ void SNiagaraStack::ConstructHeaderWidget()
 	else
 	{
 		SAssignNew(HeaderWidget, SVerticalBox)
-			//~ Enabled check box rename text box, and external header controls.
+			//~ Enable check box, pin text box, view source emitter button, and external header controls.
 			+ SVerticalBox::Slot()
 			.AutoHeight()
 			.HAlign(HAlign_Fill)
@@ -253,7 +256,6 @@ void SNiagaraStack::ConstructHeaderWidget()
 				//~ Enabled
 				SNew(SHorizontalBox)
 				+ SHorizontalBox::Slot()
-				.HAlign(HAlign_Left)
 				.AutoWidth()
 				.Padding(2)
 				[
@@ -282,19 +284,58 @@ void SNiagaraStack::ConstructHeaderWidget()
 						.RenderTransformPivot(FVector2D(.5f, .5f))
 					]
 				]
-				//~ Name
+				// Name and Source Emitter Name
+				+ SHorizontalBox::Slot()
+				.Padding(2)
+				[
+					SNew(SWrapBox)
+					.Clipping(EWidgetClipping::ClipToBoundsAlways) 
+					.UseAllottedWidth(true)
+					+ SWrapBox::Slot()
+					.Padding(3, 0)
+					[						
+						SAssignNew(InlineEditableTextBlock, SInlineEditableTextBlock)
+						.ToolTipText(this, &SNiagaraStack::GetEmitterNameToolTip)
+						.Style(FNiagaraEditorStyle::Get(), "NiagaraEditor.HeadingInlineEditableText") 
+						.Clipping(EWidgetClipping::ClipToBoundsAlways)
+						.Text(StackViewModel->GetEmitterHandleViewModel()->AsShared(), &FNiagaraEmitterHandleViewModel::GetNameText)
+						.OnTextCommitted(this, &SNiagaraStack::OnStackViewNameTextCommitted)
+						.OnVerifyTextChanged(StackViewModel->GetEmitterHandleViewModel()->AsShared(), &FNiagaraEmitterHandleViewModel::VerifyNameTextChanged)
+					]
+					+ SWrapBox::Slot()
+					.Padding(3, 0)
+					[
+						SNew(STextBlock)
+						.ToolTipText(this, &SNiagaraStack::GetEmitterNameToolTip)
+						.TextStyle(FNiagaraEditorStyle::Get(), "NiagaraEditor.SubduedHeadingTextBox")
+						.Clipping(EWidgetClipping::ClipToBoundsAlways)
+						.Text(this, &SNiagaraStack::GetSourceEmitterNameText)
+						.Visibility(this, &SNiagaraStack::GetSourceEmitterNameVisibility) 
+					]
+ 				]
+				// Open and Focus Source Emitter
 				+ SHorizontalBox::Slot()
+				.HAlign(HAlign_Fill)
+				.VAlign(VAlign_Center)
 				.AutoWidth()
 				.Padding(2)
-				.HAlign(HAlign_Left)
 				[
-					SNew(SInlineEditableTextBlock)
-					.ToolTipText(NSLOCTEXT("NiagaraEmitterEditor", "NameTextToolTip", "Click to edit the emitter name."))
-					.Style(FNiagaraEditorStyle::Get(), "NiagaraEditor.HeadingInlineEditableText")
-					.WrapTextAt(150.0f)
-					.Text(StackViewModel->GetEmitterHandleViewModel()->AsShared(), &FNiagaraEmitterHandleViewModel::GetNameText)
-					.OnTextCommitted(StackViewModel->GetEmitterHandleViewModel()->AsShared(), &FNiagaraEmitterHandleViewModel::OnNameTextComitted)
-					.OnVerifyTextChanged(StackViewModel->GetEmitterHandleViewModel()->AsShared(), &FNiagaraEmitterHandleViewModel::VerifyNameTextChanged)
+					SNew(SButton)
+					.IsFocusable(false)
+					.ToolTipText(LOCTEXT("OpenAndFocusSourceEmitterToolTip", "Open and Focus Source Emitter"))
+					.ButtonStyle(FEditorStyle::Get(), "HoverHintOnly")
+					.ForegroundColor(this, &SNiagaraStack::GetPinColor)
+					.ContentPadding(2)
+					.OnClicked(this, &SNiagaraStack::OpenSourceEmitter)
+					.Visibility(CanOpenSourceEmitter() ? EVisibility::Visible : EVisibility::Collapsed)
+					// GoToSource icon is 30x30px so we scale it down to stay in line with other 12x12px UI
+					.DesiredSizeScale(FVector2D(0.55f, 0.55f))
+					.Content()
+					[
+						SNew(SImage)
+						.Image(FNiagaraEditorWidgetsStyle::Get().GetBrush("NiagaraEditor.Stack.GoToSourceIcon"))
+						.ColorAndOpacity(FNiagaraEditorWidgetsStyle::Get().GetColor("NiagaraEditor.Stack.FlatButtonColor"))
+					]
 				]
 			]
 
@@ -307,7 +348,7 @@ void SNiagaraStack::ConstructHeaderWidget()
 				SNew(STextBlock)
 				.Text(StackViewModel->GetEmitterHandleViewModel()->GetEmitterViewModel()->AsShared(), &FNiagaraEmitterViewModel::GetStatsText)
 			]
-
+			
 			//~ Search, view options
 			+ SVerticalBox::Slot()
 			.AutoHeight()
@@ -432,7 +473,7 @@ bool SNiagaraStack::IsEntryFocusedInSearch(UNiagaraStackEntry* Entry) const
 	return false;
 }
 
-void SNiagaraStack::OpenSourceEmitter()
+FReply SNiagaraStack::OpenSourceEmitter() 
 {
 	if (StackViewModel && StackViewModel->GetEmitterHandleViewModel().IsValid() && StackViewModel->GetEmitterHandleViewModel()->GetEmitterHandle())
 	{
@@ -442,14 +483,15 @@ void SNiagaraStack::OpenSourceEmitter()
 			FAssetEditorManager::Get().OpenEditorForAsset(Emitter);
 		}
 	}
+	return FReply::Handled();
 }
 
 bool SNiagaraStack::CanOpenSourceEmitter() const
 {
-	if (StackViewModel && StackViewModel->GetEmitterHandleViewModel().IsValid() && StackViewModel->GetEmitterHandleViewModel()->GetEmitterHandle())
+	if (StackViewModel && StackViewModel->GetEmitterHandleViewModel().IsValid() && StackViewModel->GetEmitterHandleViewModel()->GetEmitterHandle()
+		&& StackViewModel->GetEmitterHandleViewModel()->GetEmitterHandle()->GetSource())
 	{
-		UNiagaraEmitter* Emitter = const_cast<UNiagaraEmitter*>(StackViewModel->GetEmitterHandleViewModel()->GetEmitterHandle()->GetSource());
-		if (Emitter != nullptr)
+		if (StackViewModel->GetSystemViewModel()->GetEditMode() == ENiagaraSystemViewModelEditMode::SystemAsset)
 		{
 			return true;
 		}
@@ -489,7 +531,7 @@ void CollapseEntriesRecursive(TArray<UNiagaraStackEntry*> Entries)
 		{
 			Entry->SetIsExpanded(false);
 		}
-
+		
 		TArray<UNiagaraStackEntry*> Children;
 		Entry->GetUnfilteredChildren(Children);
 		CollapseEntriesRecursive(Children);
@@ -942,9 +984,45 @@ void SNiagaraStack::OnContentColumnWidthChanged(float Width)
 
 void SNiagaraStack::StackStructureChanged()
 {
-	ConstructHeaderWidget();
 	PrimeTreeExpansion();
 	StackTree->RequestTreeRefresh();
 }
 
+FText SNiagaraStack::GetSourceEmitterNameText() const 
+{
+	return StackViewModel->GetEmitterHandleViewModel()->GetSourceNameText();
+}
+
+FText SNiagaraStack::GetEmitterNameToolTip() const
+{
+	if (CanOpenSourceEmitter())
+	{
+		// We are looking at this Emitter in a System Asset and it has a valid parent Emitter
+		TSharedPtr<FNiagaraEmitterHandleViewModel> ThisViewModel = StackViewModel->GetEmitterHandleViewModel();
+		return FText::Format(LOCTEXT("EmitterNameAndPath", "{0}\nParent: {1}"), ThisViewModel->GetNameText(), ThisViewModel->GetSourcePathNameText());
+	}
+	else
+	{
+		// We are looking at this Emitter in an Emitter Asset or we are looking at this Emitter in a System Asset and it does not have a valid parent Emitter
+		return StackViewModel->GetEmitterHandleViewModel()->GetNameText();
+	}
+}
+
+void SNiagaraStack::OnStackViewNameTextCommitted(const FText& InText, ETextCommit::Type CommitInfo) const
+{
+	StackViewModel->GetEmitterHandleViewModel()->OnNameTextComitted(InText, CommitInfo);
+}
+
+EVisibility SNiagaraStack::GetSourceEmitterNameVisibility() const
+{
+	return CanOpenSourceEmitter() && GetIsEmitterRenamed() ? EVisibility::Visible : EVisibility::Collapsed;
+}
+
+bool SNiagaraStack::GetIsEmitterRenamed() const
+{
+	const FText CurrentNameText = StackViewModel->GetEmitterHandleViewModel()->GetNameText();
+	const FText SourceNameText = StackViewModel->GetEmitterHandleViewModel()->GetSourceNameText();
+	return !CurrentNameText.EqualTo(SourceNameText);
+}
+
 #undef LOCTEXT_NAMESPACE
diff --git a/Engine/Plugins/FX/Niagara/Source/NiagaraEditorWidgets/Private/SNiagaraStack.h b/Engine/Plugins/FX/Niagara/Source/NiagaraEditorWidgets/Private/SNiagaraStack.h
--- a/Engine/Plugins/FX/Niagara/Source/NiagaraEditorWidgets/Private/SNiagaraStack.h
+++ b/Engine/Plugins/FX/Niagara/Source/NiagaraEditorWidgets/Private/SNiagaraStack.h
@@ -10,6 +10,7 @@
 #include "Widgets/Views/STreeView.h"
 #include "ViewModels/Stack/NiagaraStackEntry.h"
 #include "Widgets/Input/SSearchBox.h"
+#include "Widgets/Text/SInlineEditableTextBlock.h"
 
 
 class UNiagaraStackViewModel;
@@ -71,6 +72,14 @@ private:
 	void ConstructHeaderWidget();
 	FSlateColor GetPinColor() const;
 	FReply PinButtonPressed();
+	FReply OpenSourceEmitter();
+
+	// source name handling
+	FText GetSourceEmitterNameText() const;
+	FText GetEmitterNameToolTip() const;
+	void OnStackViewNameTextCommitted(const FText& InText, ETextCommit::Type CommitInfo) const;
+	EVisibility GetSourceEmitterNameVisibility() const; 
+	bool GetIsEmitterRenamed() const;
 
 	// ~stack search stuff
 	void OnSearchTextChanged(const FText& SearchText);
@@ -88,7 +97,6 @@ private:
 	bool IsEntryFocusedInSearch(UNiagaraStackEntry* Entry) const;
 	
 	// Inline menu commands
-	void OpenSourceEmitter();
 	bool CanOpenSourceEmitter() const;
 	void SetEmitterEnabled(bool bIsEnabled);
 	bool CheckEmitterEnabledStatus(bool bIsEnabled);
@@ -122,9 +130,13 @@ private:
 
 	FLinearColor CurrentPinColor;
 	
+	// emitter name textblock
+	TSharedPtr<SInlineEditableTextBlock> InlineEditableTextBlock;
+
 	// ~ search stuff
 	TSharedPtr<SSearchBox> SearchBox;
 	TSharedPtr<FActiveTimerHandle> SearchExpandTimer;
 	static const FText OccurencesFormat;
 	bool bNeedsJumpToNextOccurence;
+	
 };
