Fixed GPU emitters not being ticked correctly on the renderthread.

diff --git a/Engine/Plugins/FX/Niagara/Source/Niagara/Classes/NiagaraScriptExecutionContext.h b/Engine/Plugins/FX/Niagara/Source/Niagara/Classes/NiagaraScriptExecutionContext.h
--- a/Engine/Plugins/FX/Niagara/Source/Niagara/Classes/NiagaraScriptExecutionContext.h
+++ b/Engine/Plugins/FX/Niagara/Source/Niagara/Classes/NiagaraScriptExecutionContext.h
@@ -96,7 +96,7 @@ struct FNiagaraComputeExecutionContext
 		, GPUDataReadback(nullptr)
 		, AccumulatedSpawnRate(0)
 		, NumIndicesPerInstance(0)
-		, bPendingExecution(0)
+		, PendingExecutionQueueMask(0)
 	{
 	}
 
@@ -111,7 +111,7 @@ struct FNiagaraComputeExecutionContext
 	void Reset()
 	{
 		AccumulatedSpawnRate = 0;
-		bPendingExecution = 0;
+		PendingExecutionQueueMask = 0;
 		if (GPUDataReadback)
 		{
 			delete GPUDataReadback;
@@ -207,6 +207,6 @@ struct FNiagaraComputeExecutionContext
 	uint32 AccumulatedSpawnRate;
 	uint32 NumIndicesPerInstance;	// how many vtx indices per instance the renderer is going to have for its draw call
 
-	/** Ensures we only enqueue each context once before they're dispatched. */
-	uint32 bPendingExecution : 1;
+	/** Ensures we only enqueue each context once per queue before they're dispatched. See SIMULATION_QUEUE_COUNT */
+	uint32 PendingExecutionQueueMask;
 };
diff --git a/Engine/Plugins/FX/Niagara/Source/Niagara/Private/NiagaraEmitterInstanceBatcher.cpp b/Engine/Plugins/FX/Niagara/Source/Niagara/Private/NiagaraEmitterInstanceBatcher.cpp
--- a/Engine/Plugins/FX/Niagara/Source/Niagara/Private/NiagaraEmitterInstanceBatcher.cpp
+++ b/Engine/Plugins/FX/Niagara/Source/Niagara/Private/NiagaraEmitterInstanceBatcher.cpp
@@ -32,11 +32,12 @@ void NiagaraEmitterInstanceBatcher::Queue(FNiagaraComputeExecutionContext *InCon
 			uint32, QueueIndex, CurQueueIndex,
 			FNiagaraComputeExecutionContext*, ExecContext, InContext,
 			{
+				const uint32 QueueIndexMask = (1 << QueueIndex);
 				//Don't queue the same context for execution multiple times. TODO: possibly try to combine/accumulate the tick info if we happen to have > 1 before it's executed.
-				if (!ExecContext->bPendingExecution)
+				if (!(ExecContext->PendingExecutionQueueMask & QueueIndexMask))
 				{
 					Queue[QueueIndex].Add(ExecContext);
-					ExecContext->bPendingExecution = true;
+					ExecContext->PendingExecutionQueueMask |= QueueIndexMask;
 				}
 			});
 }
@@ -51,17 +52,26 @@ void NiagaraEmitterInstanceBatcher::Remove(FNiagaraComputeExecutionContext *InCo
 			{
 				Queue[i].Remove(ExecContext);
 			}
+			ExecContext->PendingExecutionQueueMask = 0;
 		});
 }
 
 void NiagaraEmitterInstanceBatcher::ExecuteAll(FRHICommandListImmediate &RHICmdList, FUniformBufferRHIParamRef ViewUniformBuffer)
 {
-	TArray<FNiagaraComputeExecutionContext*> &WorkQueue = SimulationQueue[CurQueueIndex ^ 0x1];
+	const uint32 TickedQueueIndex = (CurQueueIndex ^ 0x1);
+	const uint32 TickedQueueIndexMask = (1 << TickedQueueIndex);
+
+	TArray<FNiagaraComputeExecutionContext*> &WorkQueue = SimulationQueue[TickedQueueIndex];
 	for (FNiagaraComputeExecutionContext *Context : WorkQueue)
 	{
-		// need to call RenderThreadInit on data interfaces
-		//ExecuteSingle(Context, RHICmdList);
-		TickSingle(Context, RHICmdList, ViewUniformBuffer);
+		if (Context)
+		{
+			// need to call RenderThreadInit on data interfaces
+			//ExecuteSingle(Context, RHICmdList);
+			TickSingle(Context, RHICmdList, ViewUniformBuffer);
+			Context->PendingExecutionQueueMask &= ~TickedQueueIndexMask;
+		}
+
 	}
 	WorkQueue.Empty();
 }
@@ -72,7 +82,6 @@ void NiagaraEmitterInstanceBatcher::TickSingle(FNiagaraComputeExecutionContext *
 
 	check(IsInRenderingThread());
 	Context->MainDataSet->Tick();
-	Context->bPendingExecution = false;
 
 	FNiagaraComputeExecutionContext::TickCounter++;
 
