Niagara - Fix linux compile error from CL 4418815.

diff --git a/Engine/Plugins/FX/Niagara/Source/NiagaraEditor/Public/ViewModels/NiagaraSystemViewModel.h b/Engine/Plugins/FX/Niagara/Source/NiagaraEditor/Public/ViewModels/NiagaraSystemViewModel.h
--- a/Engine/Plugins/FX/Niagara/Source/NiagaraEditor/Public/ViewModels/NiagaraSystemViewModel.h
+++ b/Engine/Plugins/FX/Niagara/Source/NiagaraEditor/Public/ViewModels/NiagaraSystemViewModel.h
@@ -34,7 +34,7 @@ class UNiagaraNodeFunctionCall;
 class FNiagaraEmitterViewModel;
 
 /** Defines different editing modes for this system view model. */
-NIAGARAEDITOR_API enum class ENiagaraSystemViewModelEditMode
+enum class NIAGARAEDITOR_API ENiagaraSystemViewModelEditMode
 {
 	/** A system asset is being edited.  This assumes that emitters should be inheriting from a base version and that emitter editing will be restricted. */
 	SystemAsset,
