Fixing localization warnings

diff --git a/Engine/Plugins/FX/Niagara/Source/Niagara/Private/NDISkeletalMesh_BoneSampling.cpp b/Engine/Plugins/FX/Niagara/Source/Niagara/Private/NDISkeletalMesh_BoneSampling.cpp
--- a/Engine/Plugins/FX/Niagara/Source/Niagara/Private/NDISkeletalMesh_BoneSampling.cpp
+++ b/Engine/Plugins/FX/Niagara/Source/Niagara/Private/NDISkeletalMesh_BoneSampling.cpp
@@ -85,7 +85,7 @@ void UNiagaraDataInterfaceSkeletalMesh::GetSkeletonSamplingFunctions(TArray<FNia
 		Sig.bMemberFunction = true;
 		Sig.bRequiresContext = false;
 #if WITH_EDITORONLY_DATA
-		Sig.Description = LOCTEXT("GetSkinnedBoneDataDesc", "Returns skinning dependant data for the pased bone in world space. All outputs are optional and you will incur zero to minimal cost if they are not connected.");
+		Sig.Description = LOCTEXT("GetSkinnedBoneDataWSDesc", "Returns skinning dependant data for the pased bone in world space. All outputs are optional and you will incur zero to minimal cost if they are not connected.");
 #endif
 		OutFunctions.Add(Sig);
 	}
diff --git a/Engine/Plugins/FX/Niagara/Source/Niagara/Private/NDISkeletalMesh_TriangleSampling.cpp b/Engine/Plugins/FX/Niagara/Source/Niagara/Private/NDISkeletalMesh_TriangleSampling.cpp
--- a/Engine/Plugins/FX/Niagara/Source/Niagara/Private/NDISkeletalMesh_TriangleSampling.cpp
+++ b/Engine/Plugins/FX/Niagara/Source/Niagara/Private/NDISkeletalMesh_TriangleSampling.cpp
@@ -91,7 +91,7 @@ void UNiagaraDataInterfaceSkeletalMesh::GetTriangleSamplingFunctions(TArray<FNia
 		Sig.bMemberFunction = true;
 		Sig.bRequiresContext = false;
 #if WITH_EDITORONLY_DATA
-		Sig.Description = LOCTEXT("GetSkinnedDataDesc", "Returns skinning dependant data for the pased MeshTriCoord in world space. All outputs are optional and you will incur zerp minimal cost if they are not connected.");
+		Sig.Description = LOCTEXT("GetSkinnedDataWSDesc", "Returns skinning dependant data for the pased MeshTriCoord in world space. All outputs are optional and you will incur zerp minimal cost if they are not connected.");
 #endif
 		OutFunctions.Add(Sig);
 	}
diff --git a/Engine/Plugins/FX/Niagara/Source/Niagara/Private/NDISkeletalMesh_VertexSampling.cpp b/Engine/Plugins/FX/Niagara/Source/Niagara/Private/NDISkeletalMesh_VertexSampling.cpp
--- a/Engine/Plugins/FX/Niagara/Source/Niagara/Private/NDISkeletalMesh_VertexSampling.cpp
+++ b/Engine/Plugins/FX/Niagara/Source/Niagara/Private/NDISkeletalMesh_VertexSampling.cpp
@@ -82,7 +82,7 @@ void UNiagaraDataInterfaceSkeletalMesh::GetVertexSamplingFunctions(TArray<FNiaga
 		Sig.bMemberFunction = true;
 		Sig.bRequiresContext = false;
 #if WITH_EDITORONLY_DATA
-		Sig.Description = LOCTEXT("GetSkinnedDataDesc", "Returns skinning dependant data for the pased vertex in world space. All outputs are optional and you will incur zero to minimal cost if they are not connected.");
+		Sig.Description = LOCTEXT("GetSkinnedDataWSDesc", "Returns skinning dependant data for the pased vertex in world space. All outputs are optional and you will incur zero to minimal cost if they are not connected.");
 #endif
 		OutFunctions.Add(Sig);
 	}
diff --git a/Engine/Plugins/FX/Niagara/Source/Niagara/Private/NiagaraDataInterfaceTexture.cpp b/Engine/Plugins/FX/Niagara/Source/Niagara/Private/NiagaraDataInterfaceTexture.cpp
--- a/Engine/Plugins/FX/Niagara/Source/Niagara/Private/NiagaraDataInterfaceTexture.cpp
+++ b/Engine/Plugins/FX/Niagara/Source/Niagara/Private/NiagaraDataInterfaceTexture.cpp
@@ -101,7 +101,7 @@ void UNiagaraDataInterfaceTexture::GetFunctions(TArray<FNiagaraFunctionSignature
 		Sig.bRequiresContext = false;
 		Sig.Inputs.Add(FNiagaraVariable(FNiagaraTypeDefinition(GetClass()), TEXT("Texture")));
 		Sig.Inputs.Add(FNiagaraVariable(FNiagaraTypeDefinition::GetVec2Def(), TEXT("UV")));
-		Sig.SetDescription(LOCTEXT("TextureSampleDesc", "Sample mip level 0 of the input 2d texture at the specified UV coordinates. The UV origin (0,0) is in the upper left hand corner of the image."));
+		Sig.SetDescription(LOCTEXT("TextureSampleTexture2DDesc", "Sample mip level 0 of the input 2d texture at the specified UV coordinates. The UV origin (0,0) is in the upper left hand corner of the image."));
 		Sig.Outputs.Add(FNiagaraVariable(FNiagaraTypeDefinition::GetVec4Def(), TEXT("Value")));
 		//Sig.Owner = *GetName();
 
@@ -115,7 +115,7 @@ void UNiagaraDataInterfaceTexture::GetFunctions(TArray<FNiagaraFunctionSignature
 	//	Sig.bRequiresContext = false;
 	//	Sig.Inputs.Add(FNiagaraVariable(FNiagaraTypeDefinition(GetClass()), TEXT("Texture")));
 	//	Sig.Inputs.Add(FNiagaraVariable(FNiagaraTypeDefinition::GetVec3Def(), TEXT("UVW")));
-	//	Sig.SetDescription(LOCTEXT("TextureSampleDesc", "Sample mip level 0 of the input 3d texture at the specified UVW coordinates. The UVW origin (0,0) is in the bottom left hand corner of the volume."));
+	//	Sig.SetDescription(LOCTEXT("TextureSampleVolumeTextureDesc", "Sample mip level 0 of the input 3d texture at the specified UVW coordinates. The UVW origin (0,0) is in the bottom left hand corner of the volume."));
 	//	Sig.Outputs.Add(FNiagaraVariable(FNiagaraTypeDefinition::GetVec4Def(), TEXT("Value")));
 	//	//Sig.Owner = *GetName();
 
@@ -136,7 +136,7 @@ void UNiagaraDataInterfaceTexture::GetFunctions(TArray<FNiagaraFunctionSignature
 		Sig.Inputs.Add(FNiagaraVariable(FNiagaraTypeDefinition::GetVec2Def(), TEXT("DDX")));
 		Sig.Inputs.Add(FNiagaraVariable(FNiagaraTypeDefinition::GetVec2Def(), TEXT("DDY")));
 		
-		Sig.SetDescription(LOCTEXT("TextureSampleDesc", "Return a pseudovolume texture sample.\n"
+		Sig.SetDescription(LOCTEXT("TextureSamplePseudoVolumeTextureDesc", "Return a pseudovolume texture sample.\n"
 			"Useful for simulating 3D texturing with a 2D texture or as a texture flipbook with lerped transitions.\n"
 			"Treats 2d layout of frames as a 3d texture and performs bilinear filtering by blending with an offset Z frame.\n"
 			"Texture = Input Texture Object storing Volume Data\n"
diff --git a/Engine/Plugins/FX/Niagara/Source/NiagaraEditor/Private/Customizations/NiagaraEventScriptPropertiesCustomization.cpp b/Engine/Plugins/FX/Niagara/Source/NiagaraEditor/Private/Customizations/NiagaraEventScriptPropertiesCustomization.cpp
--- a/Engine/Plugins/FX/Niagara/Source/NiagaraEditor/Private/Customizations/NiagaraEventScriptPropertiesCustomization.cpp
+++ b/Engine/Plugins/FX/Niagara/Source/NiagaraEditor/Private/Customizations/NiagaraEventScriptPropertiesCustomization.cpp
@@ -156,7 +156,7 @@ void FNiagaraEventScriptPropertiesCustomization::CustomizeChildren(TSharedRef<IP
 
 	// Add Spawn Number property
 	{
-		FText EventSrcTxt = LOCTEXT("SpawnNumber", "Minimum SpawnNumber");
+		FText EventSrcTxt = LOCTEXT("MinimumSpawnNumber", "Minimum SpawnNumber");
 		FDetailWidgetRow& SpawnRow = ChildBuilder.AddCustomRow(EventSrcTxt);
 		TSharedRef<SWidget> MinSpawnNumberValueWidget = HandleMinSpawnNumber->CreatePropertyValueWidget();
 		TAttribute<bool> EnabledAttribute(this, &FNiagaraEventScriptPropertiesCustomization::GetUseRandomSpawnNumber);
diff --git a/Engine/Plugins/FX/Niagara/Source/NiagaraEditor/Private/NiagaraHlslTranslator.cpp b/Engine/Plugins/FX/Niagara/Source/NiagaraEditor/Private/NiagaraHlslTranslator.cpp
--- a/Engine/Plugins/FX/Niagara/Source/NiagaraEditor/Private/NiagaraHlslTranslator.cpp
+++ b/Engine/Plugins/FX/Niagara/Source/NiagaraEditor/Private/NiagaraHlslTranslator.cpp
@@ -2473,7 +2473,7 @@ int32 FHlslNiagaraTranslator::GetRapidIterationParameter(const FNiagaraVariable&
 {
 	if (!AddStructToDefinitionSet(Parameter.GetType()))
 	{
-		Error(FText::Format(LOCTEXT("GetRapidIterationParameterTypeFail", "Cannot handle type {0}! Variable: {1}"), Parameter.GetType().GetNameText(), FText::FromName(Parameter.GetName())), nullptr, nullptr);
+		Error(FText::Format(LOCTEXT("GetRapidIterationParameterTypeFail_InvalidType", "Cannot handle type {0}! Variable: {1}"), Parameter.GetType().GetNameText(), FText::FromName(Parameter.GetName())), nullptr, nullptr);
 		return INDEX_NONE;
 	}
 
@@ -2493,7 +2493,7 @@ int32 FHlslNiagaraTranslator::GetRapidIterationParameter(const FNiagaraVariable&
 		}
 		else
 		{
-			Error(FText::Format(LOCTEXT("GetRapidIterationParameterTypeFail", "Variable: {0} cannot be a RapidIterationParameter input node because it isn't a supported type {1}"), FText::FromName(Parameter.GetName()), Parameter.GetType().GetNameText()), nullptr, nullptr);
+			Error(FText::Format(LOCTEXT("GetRapidIterationParameterTypeFail_UnsupportedInput", "Variable: {0} cannot be a RapidIterationParameter input node because it isn't a supported type {1}"), FText::FromName(Parameter.GetName()), Parameter.GetType().GetNameText()), nullptr, nullptr);
 			return INDEX_NONE;
 		}
 	}
@@ -2997,7 +2997,7 @@ void FHlslNiagaraTranslator::ParameterMapSet(UNiagaraNodeParameterMapSet* SetNod
 
 			if (ParamMapHistoryIdx == -1)
 			{
-				Error(LOCTEXT("NoParamMapIdx", "Cannot find parameter map for input!"), SetNode, nullptr);
+				Error(LOCTEXT("NoParamMapIdxForInput", "Cannot find parameter map for input!"), SetNode, nullptr);
 				for (int32 j = 0; j < Outputs.Num(); j++)
 				{
 					Outputs[j] = INDEX_NONE;
@@ -3447,7 +3447,7 @@ void FHlslNiagaraTranslator::ParameterMapGet(UNiagaraNodeParameterMapGet* GetNod
 
 	if (ParamMapHistoryIdx == -1)
 	{
-		Error(LOCTEXT("NoParamMapIdx", "Cannot find parameter map for input!"), GetNode, nullptr);
+		Error(LOCTEXT("NoParamMapIdxForInput", "Cannot find parameter map for input!"), GetNode, nullptr);
 		for (int32 i = 0; i < Outputs.Num(); i++)
 		{
 			Outputs[i] = INDEX_NONE;
@@ -3456,7 +3456,7 @@ void FHlslNiagaraTranslator::ParameterMapGet(UNiagaraNodeParameterMapGet* GetNod
 	}
 	else if (ParamMapHistoryIdx >= ParamMapHistories.Num())
 	{
-		Error(FText::Format(LOCTEXT("InvalidParamMapIdx", "Invalid parameter map index for input {0} of {1}!"), ParamMapHistoryIdx, ParamMapHistories.Num()), GetNode, nullptr);
+		Error(FText::Format(LOCTEXT("InvalidParamMapIdxForInput", "Invalid parameter map index for input {0} of {1}!"), ParamMapHistoryIdx, ParamMapHistories.Num()), GetNode, nullptr);
 		for (int32 i = 0; i < Outputs.Num(); i++)
 		{
 			Outputs[i] = INDEX_NONE;
@@ -3787,7 +3787,7 @@ void FHlslNiagaraTranslator::ReadDataSet(const FNiagaraDataSetID DataSet, const
 
 	if (ParamMapHistoryIdx == -1)
 	{
-		Error(LOCTEXT("NoParamMapIdx", "Cannot find parameter map for input to ReadDataSet!"), nullptr, nullptr);
+		Error(LOCTEXT("NoParamMapIdxToReadDataSet", "Cannot find parameter map for input to ReadDataSet!"), nullptr, nullptr);
 		for (int32 i = 0; i < Outputs.Num(); i++)
 		{
 			Outputs[i] = INDEX_NONE;
@@ -3796,7 +3796,7 @@ void FHlslNiagaraTranslator::ReadDataSet(const FNiagaraDataSetID DataSet, const
 	}
 	else if (ParamMapHistoryIdx >= ParamMapHistories.Num())
 	{
-		Error(FText::Format(LOCTEXT("InvalidParamMapIdx", "Invalid parameter map index for ReadDataSet input {0} of {1}!"), ParamMapHistoryIdx, ParamMapHistories.Num()), nullptr, nullptr);
+		Error(FText::Format(LOCTEXT("InvalidParamMapIdxToReadDataSet", "Invalid parameter map index for ReadDataSet input {0} of {1}!"), ParamMapHistoryIdx, ParamMapHistories.Num()), nullptr, nullptr);
 		for (int32 i = 0; i < Outputs.Num(); i++)
 		{
 			Outputs[i] = INDEX_NONE;
@@ -3846,7 +3846,7 @@ void FHlslNiagaraTranslator::WriteDataSet(const FNiagaraDataSetID DataSet, const
 
 	if (ParamMapHistoryIdx == -1)
 	{
-		Error(LOCTEXT("NoParamMapIdx", "Cannot find parameter map for input to WriteDataSet!"), nullptr, nullptr);
+		Error(LOCTEXT("NoParamMapIdxToWriteDataSet", "Cannot find parameter map for input to WriteDataSet!"), nullptr, nullptr);
 		for (int32 i = 0; i < Outputs.Num(); i++)
 		{
 			Outputs[i] = INDEX_NONE;
@@ -3855,7 +3855,7 @@ void FHlslNiagaraTranslator::WriteDataSet(const FNiagaraDataSetID DataSet, const
 	}
 	else if (ParamMapHistoryIdx >= ParamMapHistories.Num())
 	{
-		Error(FText::Format(LOCTEXT("InvalidParamMapIdx", "Invalid parameter map index for WriteDataSet input {0} of {1}!"), ParamMapHistoryIdx, ParamMapHistories.Num()), nullptr, nullptr);
+		Error(FText::Format(LOCTEXT("InvalidParamMapIdxToWriteDataSet", "Invalid parameter map index for WriteDataSet input {0} of {1}!"), ParamMapHistoryIdx, ParamMapHistories.Num()), nullptr, nullptr);
 		for (int32 i = 0; i < Outputs.Num(); i++)
 		{
 			Outputs[i] = INDEX_NONE;
diff --git a/Engine/Plugins/FX/Niagara/Source/NiagaraEditor/Private/NiagaraNodeParameterMapGet.cpp b/Engine/Plugins/FX/Niagara/Source/NiagaraEditor/Private/NiagaraNodeParameterMapGet.cpp
--- a/Engine/Plugins/FX/Niagara/Source/NiagaraEditor/Private/NiagaraNodeParameterMapGet.cpp
+++ b/Engine/Plugins/FX/Niagara/Source/NiagaraEditor/Private/NiagaraNodeParameterMapGet.cpp
@@ -291,14 +291,14 @@ void UNiagaraNodeParameterMapGet::SynchronizeDefaultInputPin(UEdGraphPin* Defaul
 		DefaultPin->bDefaultValueIsIgnored = true;
 		DefaultPin->bNotConnectable = true;
 		DefaultPin->bHidden = true;
-		DefaultPin->PinToolTip = FText::Format(LOCTEXT("DefaultValueTooltip", "Default value for {0}. Disabled for Engine Parameters."), FText::FromName(OutputPin->PinName)).ToString();
+		DefaultPin->PinToolTip = FText::Format(LOCTEXT("DefaultValueTooltip_DisabledForEngineParameters", "Default value for {0}. Disabled for Engine Parameters."), FText::FromName(OutputPin->PinName)).ToString();
 	}
 	else
 	{
 		DefaultPin->bDefaultValueIsIgnored = false;
 		DefaultPin->bNotConnectable = false;
 		DefaultPin->bHidden = false;
-		DefaultPin->PinToolTip = FText::Format(LOCTEXT("DefaultValueTooltip", "Default value for {0} if no other module has set it previously in the stack."), FText::FromName(OutputPin->PinName)).ToString();
+		DefaultPin->PinToolTip = FText::Format(LOCTEXT("DefaultValueTooltip_UnlessOverridden", "Default value for {0} if no other module has set it previously in the stack."), FText::FromName(OutputPin->PinName)).ToString();
 	}
 }
 
@@ -508,7 +508,7 @@ void UNiagaraNodeParameterMapGet::GetPinHoverText(const UEdGraphPin& Pin, FStrin
 				}
 				else
 				{
-					FText Desc = FText::Format(LOCTEXT("GetVarTooltip", "Name: \"{0}\"\nType: {1}\nDesc: None"), FText::FromName(Pin.PinName),
+					FText Desc = FText::Format(LOCTEXT("GetVarTooltip_NoDesc", "Name: \"{0}\"\nType: {1}\nDesc: None"), FText::FromName(Pin.PinName),
 						TypeDef.GetNameText());
 					HoverTextOut = Desc.ToString();
 				}
diff --git a/Engine/Plugins/FX/Niagara/Source/NiagaraEditor/Private/Sequencer/LevelSequence/NiagaraSystemTrackEditor.cpp b/Engine/Plugins/FX/Niagara/Source/NiagaraEditor/Private/Sequencer/LevelSequence/NiagaraSystemTrackEditor.cpp
--- a/Engine/Plugins/FX/Niagara/Source/NiagaraEditor/Private/Sequencer/LevelSequence/NiagaraSystemTrackEditor.cpp
+++ b/Engine/Plugins/FX/Niagara/Source/NiagaraEditor/Private/Sequencer/LevelSequence/NiagaraSystemTrackEditor.cpp
@@ -128,7 +128,7 @@ void FNiagaraSystemTrackEditor::AddNiagaraSystemTrack(FGuid ObjectBinding)
 			return;
 		}
 
-		FScopedTransaction AddTrackTransaction(LOCTEXT("AddNiagaraSystemTrackTransaction", "Add Niagara System Life Cycle Track"));
+		FScopedTransaction AddTrackTransaction(LOCTEXT("AddNiagaraSystemLifeCycleTrackTransaction", "Add Niagara System Life Cycle Track"));
 		MovieScene->Modify();
 
 		UMovieSceneNiagaraSystemTrack* NiagaraSystemTrack = MovieScene->AddTrack<UMovieSceneNiagaraSystemTrack>(ObjectBinding);
@@ -169,7 +169,7 @@ void FNiagaraSystemTrackEditor::AddNiagaraParameterTrack(FGuid ObjectBinding, FN
 			return;
 		}
 
-		FScopedTransaction AddTrackTransaction(LOCTEXT("AddNiagaraSystemTrackTransaction", "Add Niagara Parameter Track"));
+		FScopedTransaction AddTrackTransaction(LOCTEXT("AddNiagaraParameterTrackTransaction", "Add Niagara Parameter Track"));
 		MovieScene->Modify();
 
 		UMovieSceneNiagaraParameterTrack* ParameterTrack = NiagaraEditorModule.CreateParameterTrackForType(*Parameter.GetType().GetScriptStruct(), Parameter);
diff --git a/Engine/Plugins/FX/Niagara/Source/NiagaraEditor/Private/Sequencer/NiagaraSequence/NiagaraEmitterTrackEditor.cpp b/Engine/Plugins/FX/Niagara/Source/NiagaraEditor/Private/Sequencer/NiagaraSequence/NiagaraEmitterTrackEditor.cpp
--- a/Engine/Plugins/FX/Niagara/Source/NiagaraEditor/Private/Sequencer/NiagaraSequence/NiagaraEmitterTrackEditor.cpp
+++ b/Engine/Plugins/FX/Niagara/Source/NiagaraEditor/Private/Sequencer/NiagaraSequence/NiagaraEmitterTrackEditor.cpp
@@ -225,8 +225,8 @@ void FNiagaraEmitterTrackEditor::BuildTrackContextMenu( FMenuBuilder& MenuBuilde
 					? LOCTEXT("RemoveFromIsolation", "Remove this from isolation.")
 					: LOCTEXT("AddToIsolation", "Add this to isolation"),
 				SystemViewModel.IsEmitterIsolated(EmitterTrack->GetEmitterHandleViewModel().ToSharedRef())
-					? LOCTEXT("RemoveFromIsolation", "Remove this emitter from isolation, without changing other emitters.")
-					: LOCTEXT("AddToIsolation", "Add this emitter to isolation, without changing other emitters."),
+					? LOCTEXT("RemoveFromIsolation_NoChangeOthers", "Remove this emitter from isolation, without changing other emitters.")
+					: LOCTEXT("AddToIsolation_NoChangeOthers", "Add this emitter to isolation, without changing other emitters."),
 				FSlateIcon(),
 				FUIAction(FExecuteAction::CreateRaw(&SystemViewModel, &FNiagaraSystemViewModel::ToggleEmitterIsolation, EmitterTrack->GetEmitterHandleViewModel().ToSharedRef())));
 			
diff --git a/Engine/Plugins/FX/Niagara/Source/NiagaraEditor/Private/ViewModels/Stack/NiagaraStackModuleItem.cpp b/Engine/Plugins/FX/Niagara/Source/NiagaraEditor/Private/ViewModels/Stack/NiagaraStackModuleItem.cpp
--- a/Engine/Plugins/FX/Niagara/Source/NiagaraEditor/Private/ViewModels/Stack/NiagaraStackModuleItem.cpp
+++ b/Engine/Plugins/FX/Niagara/Source/NiagaraEditor/Private/ViewModels/Stack/NiagaraStackModuleItem.cpp
@@ -260,8 +260,8 @@ void UNiagaraStackModuleItem::RefreshIssues(TArray<FStackIssue>& NewIssues)
 		{
 			FStackIssue InvalidScriptError(
 				EStackIssueSeverity::Error,
-				LOCTEXT("MissingModuleShortDescription", "Invalid module script."),
-				LOCTEXT("MissingModuleShortDescription", "The script this module is supposed to execute is missing or invalid for other reasons.  If it depends on an external script that no longer exists there will be load errors in the log."),
+				LOCTEXT("InvalidModuleScriptErrorSummary", "Invalid module script."),
+				LOCTEXT("InvalidModuleScriptError", "The script this module is supposed to execute is missing or invalid for other reasons.  If it depends on an external script that no longer exists there will be load errors in the log."),
 				GetStackEditorDataKey(),
 				false);
 
@@ -306,7 +306,7 @@ void UNiagaraStackModuleItem::RefreshIssues(TArray<FStackIssue>& NewIssues)
 					FStackIssue DuplicateAssignmentTargetError(
 						EStackIssueSeverity::Error,
 						LOCTEXT("DuplicateAssignmentTargetErrorSummary", "Duplicate variables detected."),
-						LOCTEXT("InconsistentEnabledError", "This 'Set Variables' module is attempting to set the same variable more than once, which is unsupported."),
+						LOCTEXT("DuplicateAssignmentTargetError", "This 'Set Variables' module is attempting to set the same variable more than once, which is unsupported."),
 						GetStackEditorDataKey(),
 						false,
 						RemoveDuplicateFix);
