Fixing localization warnings

diff --git a/Engine/Plugins/FX/Niagara/Source/NiagaraEditor/Private/NiagaraNodeParameterMapBase.cpp b/Engine/Plugins/FX/Niagara/Source/NiagaraEditor/Private/NiagaraNodeParameterMapBase.cpp
--- a/Engine/Plugins/FX/Niagara/Source/NiagaraEditor/Private/NiagaraNodeParameterMapBase.cpp
+++ b/Engine/Plugins/FX/Niagara/Source/NiagaraEditor/Private/NiagaraNodeParameterMapBase.cpp
@@ -194,7 +194,7 @@ void UNiagaraNodeParameterMapBase::GetPinHoverText(const UEdGraphPin& Pin, FStri
 			}
 			else
 			{
-				FText Desc = FText::Format(LOCTEXT("GetVarTooltip", "Name: \"{0}\"\nType: {1}\nDesc: None"), FText::FromName(Pin.PinName),
+				FText Desc = FText::Format(LOCTEXT("GetVarTooltip_NoDesc", "Name: \"{0}\"\nType: {1}\nDesc: None"), FText::FromName(Pin.PinName),
 					TypeDef.GetNameText());
 				HoverTextOut = Desc.ToString();
 			}
