Niagara - Fix static analysis for SItemSelector.h - Preflight passed.

diff --git a/Engine/Plugins/FX/Niagara/Source/NiagaraEditor/Private/SItemSelector.h b/Engine/Plugins/FX/Niagara/Source/NiagaraEditor/Private/SItemSelector.h
--- a/Engine/Plugins/FX/Niagara/Source/NiagaraEditor/Private/SItemSelector.h
+++ b/Engine/Plugins/FX/Niagara/Source/NiagaraEditor/Private/SItemSelector.h
@@ -370,7 +370,7 @@ private:
 		void Construct(const FArguments& InArgs, const TSharedRef<STableViewBase>& OwnerTree)
 		{
 			typename SItemSelectorTableRow::FArguments Arguments;
-			Arguments[InArgs._Content.Widget];
+			Arguments = Arguments[InArgs._Content.Widget];
 			SItemSelectorTableRow::Construct(Arguments,	OwnerTree);
 		}
 
