Fixed FGnmMultiBufferResource::Unlock free the resource buffer after a read-only lock

diff --git a/Engine/Plugins/FX/Niagara/Source/Niagara/Classes/NiagaraDataSet.h b/Engine/Plugins/FX/Niagara/Source/Niagara/Classes/NiagaraDataSet.h
--- a/Engine/Plugins/FX/Niagara/Source/Niagara/Classes/NiagaraDataSet.h
+++ b/Engine/Plugins/FX/Niagara/Source/Niagara/Classes/NiagaraDataSet.h
@@ -578,7 +578,8 @@ public:
 		{
 			DataSetIndices[CurrBuffer].Release();
 		}
-		DataSetIndices[CurrBuffer].Initialize(sizeof(int32), 64 /*Context->NumDataSets*/, EPixelFormat::PF_R32_UINT, BUF_DrawIndirect | BUF_Static);	// always allocate for up to 64 data sets
+		// Use BUF_KeepCPUAccessible here since some platforms will lock it for readonly (depending on the implementation of RHIEnqueueStagedRead) after GPU simulation.
+		DataSetIndices[CurrBuffer].Initialize(sizeof(int32), 64 /*Context->NumDataSets*/, EPixelFormat::PF_R32_UINT, BUF_DrawIndirect | BUF_Static | BUF_KeepCPUAccessible);	// always allocate for up to 64 data sets
 	}
 
 	FORCEINLINE uint32 GetNumFloatComponents()const { return TotalFloatComponents; }
diff --git a/Engine/Source/Runtime/RHI/Public/RHIGPUReadback.h b/Engine/Source/Runtime/RHI/Public/RHIGPUReadback.h
--- a/Engine/Source/Runtime/RHI/Public/RHIGPUReadback.h
+++ b/Engine/Source/Runtime/RHI/Public/RHIGPUReadback.h
@@ -41,7 +41,7 @@ public:
 
 	bool IsReady()
 	{
-		return Fence->Poll();
+		return !Fence || Fence->Poll();
 	}
 
 	void *RetrieveData(uint32 NumBytes)
