Fixed how UNiagaraDataInterfaceCurlNoise release its render resource by adding a release command in BeginDestroy.

diff --git a/Engine/Plugins/FX/Niagara/Source/Niagara/Classes/NiagaraDataInterfaceCurlNoise.h b/Engine/Plugins/FX/Niagara/Source/Niagara/Classes/NiagaraDataInterfaceCurlNoise.h
--- a/Engine/Plugins/FX/Niagara/Source/Niagara/Classes/NiagaraDataInterfaceCurlNoise.h
+++ b/Engine/Plugins/FX/Niagara/Source/Niagara/Classes/NiagaraDataInterfaceCurlNoise.h
@@ -21,6 +21,8 @@ public:
 	//UObject Interface
 	virtual void PostInitProperties()override;
 	virtual void PostLoad() override;
+	virtual void BeginDestroy() override;
+	virtual bool IsReadyForFinishDestroy() override;
 #if WITH_EDITOR
 	virtual void PostEditChangeProperty(struct FPropertyChangedEvent& PropertyChangedEvent) override;
 #endif
@@ -41,6 +43,7 @@ public:
 	virtual void GetParameterDefinitionHLSL(FNiagaraDataInterfaceGPUParamInfo& ParamInfo, FString& OutHLSL) override;
 	virtual FNiagaraDataInterfaceParametersCS* ConstructComputeParameters()const override;
 
+	void ReleaseResource();
 	FRWBuffer& GetGPUBuffer();
 	static const FString CurlNoiseBufferName;
 protected:
@@ -52,7 +55,9 @@ protected:
 	template<typename T>
 	void ReplicateBorder(T* DestBuffer);
 
-	FRWBuffer GPUBuffer;
+	TUniquePtr<FRWBuffer> GPUBuffer;
+	/** A fence which is used to keep track of the rendering thread releasing RHI resources. */
+	FRenderCommandFence ReleaseResourcesFence;
 
 	static const FName SampleNoiseFieldName;
 };
diff --git a/Engine/Plugins/FX/Niagara/Source/Niagara/Private/NiagaraDataInterfaceCurlNoise.cpp b/Engine/Plugins/FX/Niagara/Source/Niagara/Private/NiagaraDataInterfaceCurlNoise.cpp
--- a/Engine/Plugins/FX/Niagara/Source/Niagara/Private/NiagaraDataInterfaceCurlNoise.cpp
+++ b/Engine/Plugins/FX/Niagara/Source/Niagara/Private/NiagaraDataInterfaceCurlNoise.cpp
@@ -11,6 +11,7 @@ UNiagaraDataInterfaceCurlNoise::UNiagaraDataInterfaceCurlNoise(FObjectInitialize
 	: Super(ObjectInitializer)
 	, bGPUBufferDirty(true)
 	, Seed(0)
+	, GPUBuffer(new FRWBuffer)
 {
 
 }
@@ -31,6 +32,21 @@ void UNiagaraDataInterfaceCurlNoise::PostLoad()
 	InitNoiseLUT();
 }
 
+void UNiagaraDataInterfaceCurlNoise::BeginDestroy()
+{
+	Super::BeginDestroy();
+
+	if (FApp::CanEverRender() && !HasAnyFlags(RF_ClassDefaultObject))
+	{
+		ReleaseResource();
+	}
+}
+
+bool UNiagaraDataInterfaceCurlNoise::IsReadyForFinishDestroy()
+{
+	return ReleaseResourcesFence.IsFenceComplete();
+}
+
 #if WITH_EDITOR
 
 void UNiagaraDataInterfaceCurlNoise::PostEditChangeProperty(struct FPropertyChangedEvent& PropertyChangedEvent)
@@ -188,13 +204,24 @@ void UNiagaraDataInterfaceCurlNoise::GetParameterDefinitionHLSL(FNiagaraDataInte
 	OutHLSL += TEXT("Buffer<float4> ") + BufferName + TEXT(";\n");
 }
 
+void UNiagaraDataInterfaceCurlNoise::ReleaseResource()
+{
+	ENQUEUE_UNIQUE_RENDER_COMMAND_ONEPARAMETER(NiagaraCurlNoiseReleaseCommand,
+		FRWBuffer*,GPUBuffer,GPUBuffer.Get(),
+	{
+		GPUBuffer->Release();
+	});
+	// Insert a fence to signal when these commands completed
+	ReleaseResourcesFence.BeginFence();
+}
+
 FRWBuffer& UNiagaraDataInterfaceCurlNoise::GetGPUBuffer()
 {
 	check(IsInRenderingThread());
 
 	if (bGPUBufferDirty)
 	{
-		GPUBuffer.Release();
+		GPUBuffer->Release();
 		uint32 BufferSize = 17 * 17 * 17 * sizeof(float) * 4;
 		//int32 *BufferData = static_cast<int32*>(RHILockVertexBuffer(GPUBuffer.Buffer.Buffer, 0, BufferSize, EResourceLockMode::RLM_WriteOnly));
 		TResourceArray<FVector4> TempTable;
@@ -210,13 +237,13 @@ FRWBuffer& UNiagaraDataInterfaceCurlNoise::GetGPUBuffer()
 				}
 			}
 		}
-		GPUBuffer.Initialize(sizeof(float) * 4, 17 * 17 * 17, EPixelFormat::PF_A32B32G32R32F, BUF_Static, TEXT("CurlnoiseTable"), &TempTable);
+		GPUBuffer->Initialize(sizeof(float) * 4, 17 * 17 * 17, EPixelFormat::PF_A32B32G32R32F, BUF_Static, TEXT("CurlnoiseTable"), &TempTable);
 		//FPlatformMemory::Memcpy(BufferData, TempTable, BufferSize);
 		//RHIUnlockVertexBuffer(GPUBuffer.Buffer.Buffer);
 		bGPUBufferDirty = false;
 	}
 
-	return GPUBuffer;
+	return *GPUBuffer;
 }
 
 
