Fixed integration issue from Niagara with "Collision Query" module

diff --git a/Engine/Plugins/FX/Niagara/Source/Niagara/Private/NiagaraDataInterfaceCollisionQuery.cpp b/Engine/Plugins/FX/Niagara/Source/Niagara/Private/NiagaraDataInterfaceCollisionQuery.cpp
--- a/Engine/Plugins/FX/Niagara/Source/Niagara/Private/NiagaraDataInterfaceCollisionQuery.cpp
+++ b/Engine/Plugins/FX/Niagara/Source/Niagara/Private/NiagaraDataInterfaceCollisionQuery.cpp
@@ -136,7 +136,7 @@ bool UNiagaraDataInterfaceCollisionQuery::GetFunctionHLSL(const FName& Definitio
 	{
 		OutHLSL += TEXT("float3 WorldPositionFromSceneDepth(float2 ScreenPosition, float SceneDepth)\n{\n\tfloat4 HomogeneousWorldPosition = mul(float4(ScreenPosition * SceneDepth, SceneDepth, 1), View.ScreenToWorld);\n\treturn HomogeneousWorldPosition.xyz / HomogeneousWorldPosition.w;\n}\n");
 		OutHLSL += TEXT("\n");
-		OutHLSL += TEXT("void ") + InstanceFunctionName + TEXT("(in int InQueryID, in float3 In_ParticlePos, in float3 In_ParticleVel, in float In_DeltaSeconds, float CollisionRadius, in float CollisionDepthBounds, \
+		OutHLSL += TEXT("void ") + InstanceFunctionName + TEXT("(in int InQueryID, in float3 In_ParticlePos, in float3 In_ParticleVel, in float In_DeltaSeconds, float CollisionRadius, in float CollisionDepthBounds, int In_InstanceData, \
 			out int Out_QueryID, out bool OutCollisionValid, out float3 Out_CollisionPos, out float3 Out_CollisionNormal, out float Out_Friction, out float Out_Restitution) \n{\n");
 		// get the screen position
 		OutHLSL += TEXT("\
