Fixed Niagara proxy rendering in ES2 while the shaders are not available.

diff --git a/Engine/Plugins/FX/Niagara/Source/Niagara/Private/NiagaraComponent.cpp b/Engine/Plugins/FX/Niagara/Source/Niagara/Private/NiagaraComponent.cpp
--- a/Engine/Plugins/FX/Niagara/Source/Niagara/Private/NiagaraComponent.cpp
+++ b/Engine/Plugins/FX/Niagara/Source/Niagara/Private/NiagaraComponent.cpp
@@ -203,7 +203,7 @@ FPrimitiveViewRelevance FNiagaraSceneProxy::GetViewRelevance(const FSceneView* V
 {
 	FPrimitiveViewRelevance Relevance;
 
-	if (bRenderingEnabled == false)
+	if (bRenderingEnabled == false || (View->GetFeatureLevel() != ERHIFeatureLevel::SM5 && View->GetFeatureLevel() != ERHIFeatureLevel::ES3_1))
 	{
 		return Relevance;
 	}
