Unshelved from pending changelist '4425352':

Fix Skeletal mesh datainterface refusing to run if sampled skel mesh does not have vertex color

diff --git a/Engine/Plugins/FX/Niagara/Source/Niagara/Classes/NiagaraDataInterfaceSkeletalMesh.h b/Engine/Plugins/FX/Niagara/Source/Niagara/Classes/NiagaraDataInterfaceSkeletalMesh.h
--- a/Engine/Plugins/FX/Niagara/Source/Niagara/Classes/NiagaraDataInterfaceSkeletalMesh.h
+++ b/Engine/Plugins/FX/Niagara/Source/Niagara/Classes/NiagaraDataInterfaceSkeletalMesh.h
@@ -286,6 +286,8 @@ struct FNDISkeletalMesh_InstanceData
 		}
 		return Ret;
 	}
+
+	bool HasColorData();
 };
 
 /** Data Interface allowing sampling of static meshes. */
@@ -382,6 +384,8 @@ public:
 
 	void GetTriCoordColor(FVectorVMContext& Context);
 
+	void GetTriCoordColorFallback(FVectorVMContext& Context);
+
 	template<typename VertexAccessorType>
 	void GetTriCoordUV(FVectorVMContext& Context);
 
@@ -425,6 +429,8 @@ public:
 
 	void GetVertexColor(FVectorVMContext& Context);
 
+	void GetVertexColorFallback(FVectorVMContext& Context);
+
 	template<typename VertexAccessorType>
 	void GetVertexUV(FVectorVMContext& Context);
 
diff --git a/Engine/Plugins/FX/Niagara/Source/Niagara/Private/NDISkeletalMesh_TriangleSampling.cpp b/Engine/Plugins/FX/Niagara/Source/Niagara/Private/NDISkeletalMesh_TriangleSampling.cpp
--- a/Engine/Plugins/FX/Niagara/Source/Niagara/Private/NDISkeletalMesh_TriangleSampling.cpp
+++ b/Engine/Plugins/FX/Niagara/Source/Niagara/Private/NDISkeletalMesh_TriangleSampling.cpp
@@ -17,6 +17,7 @@ DECLARE_CYCLE_STAT(TEXT("Skel Mesh Sampling"), STAT_NiagaraSkel_Sample, STATGROU
 DEFINE_NDI_FUNC_BINDER(UNiagaraDataInterfaceSkeletalMesh, RandomTriCoord);
 DEFINE_NDI_FUNC_BINDER(UNiagaraDataInterfaceSkeletalMesh, GetTriCoordSkinnedData);
 DEFINE_NDI_DIRECT_FUNC_BINDER(UNiagaraDataInterfaceSkeletalMesh, GetTriCoordColor);
+DEFINE_NDI_DIRECT_FUNC_BINDER(UNiagaraDataInterfaceSkeletalMesh, GetTriCoordColorFallback);
 DEFINE_NDI_FUNC_BINDER(UNiagaraDataInterfaceSkeletalMesh, GetTriCoordUV);
 DEFINE_NDI_FUNC_BINDER(UNiagaraDataInterfaceSkeletalMesh, IsValidTriCoord);
 DEFINE_NDI_FUNC_BINDER(UNiagaraDataInterfaceSkeletalMesh, GetFilteredTriangleCount);
@@ -160,7 +161,6 @@ void UNiagaraDataInterfaceSkeletalMesh::GetTriangleSamplingFunctions(TArray<FNia
 
 void UNiagaraDataInterfaceSkeletalMesh::BindTriangleSamplingFunction(const FVMExternalFunctionBindingInfo& BindingInfo, FNDISkeletalMesh_InstanceData* InstanceData, FVMExternalFunction &OutFunc)
 {
-	bool bNeedsVertexColors = false;
 
 	if (BindingInfo.Name == RandomTriCoordName)
 	{
@@ -185,8 +185,14 @@ void UNiagaraDataInterfaceSkeletalMesh::BindTriangleSamplingFunction(const FVMEx
 	else if (BindingInfo.Name == GetTriColorName)
 	{
 		check(BindingInfo.GetNumInputs() == 5 && BindingInfo.GetNumOutputs() == 4);
-		bNeedsVertexColors = true;
-		NDI_FUNC_BINDER(UNiagaraDataInterfaceSkeletalMesh, GetTriCoordColor)::Bind(this, OutFunc);
+		if (InstanceData->HasColorData())
+		{
+			NDI_FUNC_BINDER(UNiagaraDataInterfaceSkeletalMesh, GetTriCoordColor)::Bind(this, OutFunc);
+		}
+		else
+		{
+			NDI_FUNC_BINDER(UNiagaraDataInterfaceSkeletalMesh, GetTriCoordColorFallback)::Bind(this, OutFunc);
+		}
 	}
 	else if (BindingInfo.Name == GetTriUVName)
 	{
@@ -209,15 +215,6 @@ void UNiagaraDataInterfaceSkeletalMesh::BindTriangleSamplingFunction(const FVMEx
 		TSkinningModeBinder<NDI_FUNC_BINDER(UNiagaraDataInterfaceSkeletalMesh, GetTriCoordVertices)>::Bind(this, BindingInfo, InstanceData, OutFunc);
 	}
 
-	check(InstanceData->Mesh);
-	FSkinWeightVertexBuffer* SkinWeightBuffer;
-	FSkeletalMeshLODRenderData& LODData = InstanceData->GetLODRenderDataAndSkinWeights(SkinWeightBuffer);
-
-	if (bNeedsVertexColors && LODData.StaticVertexBuffers.ColorVertexBuffer.GetNumVertices() == 0)
-	{
-		UE_LOG(LogNiagara, Log, TEXT("Skeletal Mesh data interface is cannot run as it's reading color data on a mesh that does not provide it. - Mesh:%s  "), *InstanceData->Mesh->GetFullName());
-		OutFunc = FVMExternalFunction();
-	}
 }
 
 template<typename FilterMode, typename AreaWeightingMode>
@@ -652,6 +649,29 @@ void UNiagaraDataInterfaceSkeletalMesh::GetTriCoordColor(FVectorVMContext& Conte
 	}
 }
 
+// Where we determine we are sampling a skeletal mesh without tri color we bind to this fallback method 
+void UNiagaraDataInterfaceSkeletalMesh::GetTriCoordColorFallback(FVectorVMContext& Context)
+{
+	VectorVM::FExternalFuncInputHandler<int32> TriParam(Context);
+	VectorVM::FExternalFuncInputHandler<float> BaryXParam(Context);
+	VectorVM::FExternalFuncInputHandler<float> BaryYParam(Context);
+	VectorVM::FExternalFuncInputHandler<float> BaryZParam(Context);
+	VectorVM::FUserPtrHandler<FNDISkeletalMesh_InstanceData> InstData(Context);
+
+	VectorVM::FExternalFuncRegisterHandler<float> OutColorR(Context);
+	VectorVM::FExternalFuncRegisterHandler<float> OutColorG(Context);
+	VectorVM::FExternalFuncRegisterHandler<float> OutColorB(Context);
+	VectorVM::FExternalFuncRegisterHandler<float> OutColorA(Context);
+
+	for (int32 i = 0; i < Context.NumInstances; ++i)
+	{
+		*OutColorR.GetDestAndAdvance() = 1.0f;
+		*OutColorG.GetDestAndAdvance() = 1.0f;
+		*OutColorB.GetDestAndAdvance() = 1.0f;
+		*OutColorA.GetDestAndAdvance() = 1.0f;
+	}
+}
+
 template<typename VertexAccessorType>
 void UNiagaraDataInterfaceSkeletalMesh::GetTriCoordUV(FVectorVMContext& Context)
 {
diff --git a/Engine/Plugins/FX/Niagara/Source/Niagara/Private/NDISkeletalMesh_VertexSampling.cpp b/Engine/Plugins/FX/Niagara/Source/Niagara/Private/NDISkeletalMesh_VertexSampling.cpp
--- a/Engine/Plugins/FX/Niagara/Source/Niagara/Private/NDISkeletalMesh_VertexSampling.cpp
+++ b/Engine/Plugins/FX/Niagara/Source/Niagara/Private/NDISkeletalMesh_VertexSampling.cpp
@@ -17,6 +17,7 @@ DECLARE_CYCLE_STAT(TEXT("Skel Mesh Vertex Sampling"), STAT_NiagaraSkel_Vertex_Sa
 DEFINE_NDI_FUNC_BINDER(UNiagaraDataInterfaceSkeletalMesh, RandomVertex);
 DEFINE_NDI_FUNC_BINDER(UNiagaraDataInterfaceSkeletalMesh, GetVertexSkinnedData);
 DEFINE_NDI_DIRECT_FUNC_BINDER(UNiagaraDataInterfaceSkeletalMesh, GetVertexColor);
+DEFINE_NDI_DIRECT_FUNC_BINDER(UNiagaraDataInterfaceSkeletalMesh, GetVertexColorFallback);
 DEFINE_NDI_FUNC_BINDER(UNiagaraDataInterfaceSkeletalMesh, GetVertexUV);
 DEFINE_NDI_FUNC_BINDER(UNiagaraDataInterfaceSkeletalMesh, IsValidVertex);
 DEFINE_NDI_FUNC_BINDER(UNiagaraDataInterfaceSkeletalMesh, GetFilteredVertexCount);
@@ -135,7 +136,6 @@ void UNiagaraDataInterfaceSkeletalMesh::GetVertexSamplingFunctions(TArray<FNiaga
 
 void UNiagaraDataInterfaceSkeletalMesh::BindVertexSamplingFunction(const FVMExternalFunctionBindingInfo& BindingInfo, FNDISkeletalMesh_InstanceData* InstanceData, FVMExternalFunction &OutFunc)
 {
-	bool bNeedsVertexColors = false;
 
 	if (BindingInfo.Name == RandomVertexName)
 	{
@@ -160,8 +160,14 @@ void UNiagaraDataInterfaceSkeletalMesh::BindVertexSamplingFunction(const FVMExte
 	else if (BindingInfo.Name == GetVertexColorName)
 	{
 		check(BindingInfo.GetNumInputs() == 2 && BindingInfo.GetNumOutputs() == 4);
-		bNeedsVertexColors = true;
-		NDI_FUNC_BINDER(UNiagaraDataInterfaceSkeletalMesh, GetVertexColor)::Bind(this, OutFunc);
+		if (InstanceData->HasColorData())
+		{
+			NDI_FUNC_BINDER(UNiagaraDataInterfaceSkeletalMesh, GetVertexColor)::Bind(this, OutFunc);
+		}
+		else
+		{
+			NDI_FUNC_BINDER(UNiagaraDataInterfaceSkeletalMesh, GetVertexColorFallback)::Bind(this, OutFunc);
+		}
 	}
 	else if (BindingInfo.Name == GetVertexUVName)
 	{
@@ -179,15 +185,6 @@ void UNiagaraDataInterfaceSkeletalMesh::BindVertexSamplingFunction(const FVMExte
 		TFilterModeBinder<NDI_FUNC_BINDER(UNiagaraDataInterfaceSkeletalMesh, GetFilteredVertexAt)>::Bind(this, BindingInfo, InstanceData, OutFunc);
 	}
 
-	check(InstanceData->Mesh);
-	FSkinWeightVertexBuffer* SkinWeightBuffer;
-	FSkeletalMeshLODRenderData& LODData = InstanceData->GetLODRenderDataAndSkinWeights(SkinWeightBuffer);
-
-	if (bNeedsVertexColors && LODData.StaticVertexBuffers.ColorVertexBuffer.GetNumVertices() == 0)
-	{
-		UE_LOG(LogNiagara, Log, TEXT("Skeletal Mesh data interface is cannot run as it's reading color data on a mesh that does not provide it. - Mesh:%s  "), *InstanceData->Mesh->GetFullName());
-		OutFunc = FVMExternalFunction();
-	}
 }
 
 template<typename FilterMode>
@@ -435,6 +432,25 @@ void UNiagaraDataInterfaceSkeletalMesh::GetVertexColor(FVectorVMContext& Context
 	}
 }
 
+void UNiagaraDataInterfaceSkeletalMesh::GetVertexColorFallback(FVectorVMContext& Context)
+{
+	VectorVM::FExternalFuncInputHandler<int32> VertParam(Context);
+	VectorVM::FUserPtrHandler<FNDISkeletalMesh_InstanceData> InstData(Context);
+
+	VectorVM::FExternalFuncRegisterHandler<float> OutColorR(Context);
+	VectorVM::FExternalFuncRegisterHandler<float> OutColorG(Context);
+	VectorVM::FExternalFuncRegisterHandler<float> OutColorB(Context);
+	VectorVM::FExternalFuncRegisterHandler<float> OutColorA(Context);
+
+	for (int32 i = 0; i < Context.NumInstances; ++i)
+	{
+		*OutColorR.GetDestAndAdvance() = 1.0f;
+		*OutColorG.GetDestAndAdvance() = 1.0f;
+		*OutColorB.GetDestAndAdvance() = 1.0f;
+		*OutColorA.GetDestAndAdvance() = 1.0f;
+	}
+}
+
 template<typename VertexAccessorType>
 void UNiagaraDataInterfaceSkeletalMesh::GetVertexUV(FVectorVMContext& Context)
 {
diff --git a/Engine/Plugins/FX/Niagara/Source/Niagara/Private/NiagaraDataInterfaceSkeletalMesh.cpp b/Engine/Plugins/FX/Niagara/Source/Niagara/Private/NiagaraDataInterfaceSkeletalMesh.cpp
--- a/Engine/Plugins/FX/Niagara/Source/Niagara/Private/NiagaraDataInterfaceSkeletalMesh.cpp
+++ b/Engine/Plugins/FX/Niagara/Source/Niagara/Private/NiagaraDataInterfaceSkeletalMesh.cpp
@@ -710,6 +710,15 @@ bool FNDISkeletalMesh_InstanceData::Tick(UNiagaraDataInterfaceSkeletalMesh* Inte
 	}
 }
 
+bool FNDISkeletalMesh_InstanceData::HasColorData()
+{
+	check(Mesh);
+	FSkinWeightVertexBuffer* SkinWeightBuffer;
+	FSkeletalMeshLODRenderData& LODData = GetLODRenderDataAndSkinWeights(SkinWeightBuffer);
+
+	return LODData.StaticVertexBuffers.ColorVertexBuffer.GetNumVertices() != 0;
+}
+
 //Instance Data END
 //////////////////////////////////////////////////////////////////////////
 
@@ -987,7 +996,7 @@ void UNiagaraDataInterfaceSkeletalMesh::ValidateFunction(const FNiagaraFunctionS
 
 		if (SkinnedDataDeprecatedFunctions.Contains(Function))
 		{
-			OutValidationErrors.Add(FText::Format(LOCTEXT("SkinnedDataFunctionDeprecationMsgFmt", "Skeletal Mesh DI Function {0} has beend deprecated. Use GetSinnedTriangleData or GetSkinnedTriangleDataWS instead.\n"), FText::FromString(Function.GetName())));
+			OutValidationErrors.Add(FText::Format(LOCTEXT("SkinnedDataFunctionDeprecationMsgFmt", "Skeletal Mesh DI Function {0} has been deprecated. Use GetSinnedTriangleData or GetSkinnedTriangleDataWS instead.\n"), FText::FromString(Function.GetName())));
 		}
 		else
 		{
