Unshelved from dev-niagara:

Allow dragging and dropping variables from Parameters panel into Stack to create new Set Variable block
Moved ParameterIsCompatibleWithScriptUsage() impl from NiagaraStackModuleItem to FNiagaraStackGraphUtils

diff --git a/Engine/Plugins/FX/Niagara/Source/NiagaraEditor/Private/ViewModels/Stack/NiagaraStackGraphUtilities.cpp b/Engine/Plugins/FX/Niagara/Source/NiagaraEditor/Private/ViewModels/Stack/NiagaraStackGraphUtilities.cpp
--- a/Engine/Plugins/FX/Niagara/Source/NiagaraEditor/Private/ViewModels/Stack/NiagaraStackGraphUtilities.cpp
+++ b/Engine/Plugins/FX/Niagara/Source/NiagaraEditor/Private/ViewModels/Stack/NiagaraStackGraphUtilities.cpp
@@ -1562,6 +1562,27 @@ bool FNiagaraStackGraphUtilities::IsValidDefaultDynamicInput(UNiagaraScript& Own
 	return TryGetStackFunctionInputValue(OwningScript, nullptr, DefaultPin, NAME_None, FRapidIterationParameterContext(), InputValue) && InputValue.DynamicValue.IsSet();
 }
 
+bool FNiagaraStackGraphUtilities::ParameterIsCompatibleWithScriptUsage(FNiagaraVariable Parameter, ENiagaraScriptUsage Usage)
+{
+	const FNiagaraParameterHandle ParameterHandle(Parameter.GetName());
+	switch (Usage)
+	{
+	case ENiagaraScriptUsage::SystemSpawnScript:
+	case ENiagaraScriptUsage::SystemUpdateScript:
+		return ParameterHandle.IsSystemHandle();
+	case ENiagaraScriptUsage::EmitterSpawnScript:
+	case ENiagaraScriptUsage::EmitterUpdateScript:
+		return ParameterHandle.IsEmitterHandle();
+	case ENiagaraScriptUsage::ParticleSpawnScript:
+	case ENiagaraScriptUsage::ParticleSpawnScriptInterpolated:
+	case ENiagaraScriptUsage::ParticleUpdateScript:
+	case ENiagaraScriptUsage::ParticleEventScript:
+		return ParameterHandle.IsParticleAttributeHandle();
+	default:
+		return false;
+	}
+}
+
 bool FNiagaraStackGraphUtilities::DoesDynamicInputMatchDefault(
 	FString EmitterUniqueName, 
 	UNiagaraScript& OwningScript,
diff --git a/Engine/Plugins/FX/Niagara/Source/NiagaraEditor/Private/ViewModels/Stack/NiagaraStackModuleItem.cpp b/Engine/Plugins/FX/Niagara/Source/NiagaraEditor/Private/ViewModels/Stack/NiagaraStackModuleItem.cpp
--- a/Engine/Plugins/FX/Niagara/Source/NiagaraEditor/Private/ViewModels/Stack/NiagaraStackModuleItem.cpp
+++ b/Engine/Plugins/FX/Niagara/Source/NiagaraEditor/Private/ViewModels/Stack/NiagaraStackModuleItem.cpp
@@ -665,33 +665,12 @@ void UNiagaraStackModuleItem::NotifyModuleMoved()
 	ModifiedGroupItemsDelegate.ExecuteIfBound();
 }
 
-bool ParameterIsCompatibleWithScriptUsage(FNiagaraVariable Parameter, ENiagaraScriptUsage Usage)
-{
-	FNiagaraParameterHandle ParameterHandle(Parameter.GetName());
-	switch (Usage)
-	{
-	case ENiagaraScriptUsage::SystemSpawnScript:
-	case ENiagaraScriptUsage::SystemUpdateScript:
-		return ParameterHandle.IsSystemHandle();
-	case ENiagaraScriptUsage::EmitterSpawnScript:
-	case ENiagaraScriptUsage::EmitterUpdateScript:
-		return ParameterHandle.IsEmitterHandle();
-	case ENiagaraScriptUsage::ParticleSpawnScript:
-	case ENiagaraScriptUsage::ParticleSpawnScriptInterpolated:
-	case ENiagaraScriptUsage::ParticleUpdateScript:
-	case ENiagaraScriptUsage::ParticleEventScript:
-		return ParameterHandle.IsParticleAttributeHandle();
-	default:
-		return false;
-	}
-}
-
 bool UNiagaraStackModuleItem::CanAddInput(FNiagaraVariable InputParameter) const
 {
 	UNiagaraNodeAssignment* AssignmentModule = Cast<UNiagaraNodeAssignment>(FunctionCallNode);
 	return AssignmentModule != nullptr &&
 		AssignmentModule->GetAssignmentTargets().Contains(InputParameter) == false &&
-		ParameterIsCompatibleWithScriptUsage(InputParameter, OutputNode->GetUsage());
+		FNiagaraStackGraphUtilities::ParameterIsCompatibleWithScriptUsage(InputParameter, OutputNode->GetUsage());
 }
 
 void UNiagaraStackModuleItem::AddInput(FNiagaraVariable InputParameter)
diff --git a/Engine/Plugins/FX/Niagara/Source/NiagaraEditor/Private/ViewModels/Stack/NiagaraStackModuleSpacer.cpp b/Engine/Plugins/FX/Niagara/Source/NiagaraEditor/Private/ViewModels/Stack/NiagaraStackModuleSpacer.cpp
new file mode 100644
--- /dev/null
+++ b/Engine/Plugins/FX/Niagara/Source/NiagaraEditor/Private/ViewModels/Stack/NiagaraStackModuleSpacer.cpp
@@ -0,0 +1,40 @@
+// Copyright 1998-2018 Epic Games, Inc. All Rights Reserved.
+
+#include "ViewModels/Stack/NiagaraStackModuleSpacer.h"
+#include "NiagaraActions.h"
+#include "NiagaraTypes.h"
+#include "ViewModels/Stack/NiagaraParameterHandle.h"
+#include "ViewModels/Stack/NiagaraStackGraphUtilities.h"
+
+void UNiagaraStackModuleSpacer::Initialize(FRequiredEntryData InRequiredEntryData, ENiagaraScriptUsage InScriptUsage, FName InSpacerKey, float InSpacerScale, EStackRowStyle InRowStyle)
+{
+	Super::Initialize(InRequiredEntryData, InSpacerKey, InSpacerScale, InRowStyle);
+	ItemGroupScriptUsage = InScriptUsage;
+}
+
+FReply UNiagaraStackModuleSpacer::OnStackSpacerDrop(TSharedPtr<FDragDropOperation> DragDropOperation)
+{
+	if (DragDropOperation->IsOfType<FNiagaraStackDragOperation>())
+	{
+		TSharedPtr<FNiagaraStackDragOperation> InputDragDropOperation = StaticCastSharedPtr<FNiagaraStackDragOperation>(DragDropOperation);
+		TSharedPtr<FNiagaraParameterAction> Action = StaticCastSharedPtr<FNiagaraParameterAction>(InputDragDropOperation->GetAction());
+		
+		OnStackSpacerAcceptDrop.ExecuteIfBound(this, Action->GetParameter());
+		return FReply::Handled();
+	}
+	return FReply::Unhandled();
+}
+
+bool UNiagaraStackModuleSpacer::OnStackSpacerAllowDrop(TSharedPtr<FDragDropOperation> DragDropOperation)
+{
+	if (DragDropOperation->IsOfType<FNiagaraStackDragOperation>())
+	{
+		TSharedPtr<FNiagaraStackDragOperation> InputDragDropOperation = StaticCastSharedPtr<FNiagaraStackDragOperation>(DragDropOperation);
+		TSharedPtr<FNiagaraParameterAction> Action = StaticCastSharedPtr<FNiagaraParameterAction>(InputDragDropOperation->GetAction());
+		if (Action.IsValid())
+		{
+			return FNiagaraStackGraphUtilities::ParameterIsCompatibleWithScriptUsage(Action->GetParameter(), ItemGroupScriptUsage);
+		}	
+	}
+	return false;
+}
diff --git a/Engine/Plugins/FX/Niagara/Source/NiagaraEditor/Private/ViewModels/Stack/NiagaraStackScriptItemGroup.cpp b/Engine/Plugins/FX/Niagara/Source/NiagaraEditor/Private/ViewModels/Stack/NiagaraStackScriptItemGroup.cpp
--- a/Engine/Plugins/FX/Niagara/Source/NiagaraEditor/Private/ViewModels/Stack/NiagaraStackScriptItemGroup.cpp
+++ b/Engine/Plugins/FX/Niagara/Source/NiagaraEditor/Private/ViewModels/Stack/NiagaraStackScriptItemGroup.cpp
@@ -24,6 +24,7 @@
 #include "Widgets/Notifications/SNotificationList.h"
 #include "Framework/Notifications/NotificationManager.h"
 #include "Internationalization/Internationalization.h"
+#include "ViewModels/Stack/NiagaraStackModuleSpacer.h"
 
 #define LOCTEXT_NAMESPACE "UNiagaraStackScriptItemGroup"
 
@@ -269,6 +270,15 @@ void UNiagaraStackScriptItemGroup::Initialize(
 		FOnGraphChanged::FDelegate::CreateUObject(this, &UNiagaraStackScriptItemGroup::OnScriptGraphChanged));
 }
 
+UNiagaraNodeOutput* UNiagaraStackScriptItemGroup::GetScriptOutputNode() const
+{
+	TSharedPtr<FNiagaraScriptViewModel> ScriptViewModelPinned = ScriptViewModel.Pin();
+	checkf(ScriptViewModelPinned.IsValid(), TEXT("Can not get script output node when the script view model has been deleted."));
+
+	UNiagaraGraph* Graph = ScriptViewModelPinned->GetGraphViewModel()->GetGraph();
+	return Graph->FindEquivalentOutputNode(ScriptUsage, ScriptUsageId);
+}
+
 void UNiagaraStackScriptItemGroup::FinalizeInternal()
 {
 	if (ScriptGraph.IsValid())
@@ -301,13 +311,14 @@ void UNiagaraStackScriptItemGroup::RefreshChildrenInternal(const TArray<UNiagara
 		for (UNiagaraNodeFunctionCall* ModuleNode : ModuleNodes)
 		{
 			FName ModuleSpacerKey = *FString::Printf(TEXT("Module%i"), ModuleIndex);
-			UNiagaraStackSpacer* ModuleSpacer = FindCurrentChildOfTypeByPredicate<UNiagaraStackSpacer>(CurrentChildren,
-				[=](UNiagaraStackSpacer* CurrentModuleSpacer) { return CurrentModuleSpacer->GetSpacerKey() == ModuleSpacerKey; });
+			UNiagaraStackModuleSpacer* ModuleSpacer = FindCurrentChildOfTypeByPredicate<UNiagaraStackModuleSpacer>(CurrentChildren,
+				[=](UNiagaraStackModuleSpacer* CurrentModuleSpacer) { return CurrentModuleSpacer->GetSpacerKey() == ModuleSpacerKey; });
 
 			if (ModuleSpacer == nullptr)
 			{
-				ModuleSpacer = NewObject<UNiagaraStackSpacer>(this);
-				ModuleSpacer->Initialize(CreateDefaultChildRequiredData(), ModuleSpacerKey, 1.4f);
+				ModuleSpacer = NewObject<UNiagaraStackModuleSpacer>(this);
+				ModuleSpacer->Initialize(CreateDefaultChildRequiredData(), GetScriptUsage(), ModuleSpacerKey, 1.4f, UNiagaraStackEntry::EStackRowStyle::None);
+				ModuleSpacer->OnStackSpacerAcceptDrop.BindUObject(this, &UNiagaraStackScriptItemGroup::AddParameterModuleToStack);
 			}
 
 			NewChildren.Add(ModuleSpacer);
@@ -330,13 +341,14 @@ void UNiagaraStackScriptItemGroup::RefreshChildrenInternal(const TArray<UNiagara
 
 		// Add the post items spacer.
 		FName ModuleSpacerKey = *FString::Printf(TEXT("Module%i"), ModuleIndex);
-		UNiagaraStackSpacer* ModuleSpacer = FindCurrentChildOfTypeByPredicate<UNiagaraStackSpacer>(CurrentChildren,
-			[=](UNiagaraStackSpacer* CurrentModuleSpacer) { return CurrentModuleSpacer->GetSpacerKey() == ModuleSpacerKey; });
+		UNiagaraStackModuleSpacer* ModuleSpacer = FindCurrentChildOfTypeByPredicate<UNiagaraStackModuleSpacer>(CurrentChildren,
+			[=](UNiagaraStackModuleSpacer* CurrentModuleSpacer) { return CurrentModuleSpacer->GetSpacerKey() == ModuleSpacerKey; });
 
 		if (ModuleSpacer == nullptr)
 		{
-			ModuleSpacer = NewObject<UNiagaraStackSpacer>(this);
-			ModuleSpacer->Initialize(CreateDefaultChildRequiredData(), ModuleSpacerKey, 1.4f);
+			ModuleSpacer = NewObject<UNiagaraStackModuleSpacer>(this);
+			ModuleSpacer->Initialize(CreateDefaultChildRequiredData(), GetScriptUsage(), ModuleSpacerKey, 1.4f, UNiagaraStackEntry::EStackRowStyle::None);
+			ModuleSpacer->OnStackSpacerAcceptDrop.BindUObject(this, &UNiagaraStackScriptItemGroup::AddParameterModuleToStack);
 		}
 
 		NewChildren.Add(ModuleSpacer);
@@ -629,5 +641,25 @@ void UNiagaraStackScriptItemGroup::OnScriptGraphChanged(const struct FEdGraphEdi
 	}
 }
 
+void UNiagaraStackScriptItemGroup::AddParameterModuleToStack(const UNiagaraStackModuleSpacer* InModuleSpacer, const FNiagaraVariable &InVariable)
+{
+	TArray<FNiagaraVariable> Vars;
+	Vars.Add(InVariable);
+	TArray<FString> DefaultVals;
+	DefaultVals.Add(FNiagaraConstants::GetAttributeDefaultValue(InVariable));
+
+	UNiagaraNodeOutput* OutputNode = GetScriptOutputNode();
+
+	int32 TargetIndex = INDEX_NONE;
+	UNiagaraStackModuleItem** TargetModuleItemPtr = StackSpacerToModuleItemMap.Find(FObjectKey(InModuleSpacer));
+	if (*TargetModuleItemPtr != nullptr)
+	{
+		UNiagaraStackModuleItem* TargetModuleItem = *TargetModuleItemPtr;
+		TargetIndex = TargetModuleItem->GetModuleIndex();
+	}
+
+	FNiagaraStackGraphUtilities::AddParameterModuleToStack(Vars, *OutputNode, TargetIndex, DefaultVals);
+	RefreshChildren();
+}
 
 #undef LOCTEXT_NAMESPACE
diff --git a/Engine/Plugins/FX/Niagara/Source/NiagaraEditor/Private/ViewModels/Stack/NiagaraStackSpacer.cpp b/Engine/Plugins/FX/Niagara/Source/NiagaraEditor/Private/ViewModels/Stack/NiagaraStackSpacer.cpp
--- a/Engine/Plugins/FX/Niagara/Source/NiagaraEditor/Private/ViewModels/Stack/NiagaraStackSpacer.cpp
+++ b/Engine/Plugins/FX/Niagara/Source/NiagaraEditor/Private/ViewModels/Stack/NiagaraStackSpacer.cpp
@@ -1,6 +1,8 @@
 // Copyright 1998-2018 Epic Games, Inc. All Rights Reserved.
 
 #include "ViewModels/Stack/NiagaraStackSpacer.h"
+#include "NiagaraActions.h"
+#include "ViewModels/Stack/NiagaraStackGraphUtilities.h"
 
 void UNiagaraStackSpacer::Initialize(FRequiredEntryData InRequiredEntryData, FName InSpacerKey, float InSpacerScale, EStackRowStyle InRowStyle)
 {
@@ -33,4 +35,14 @@ FName UNiagaraStackSpacer::GetSpacerKey() const
 float UNiagaraStackSpacer::GetSpacerScale() const
 {
 	return SpacerScale;
-}
\ No newline at end of file
+}
+
+FReply UNiagaraStackSpacer::OnStackSpacerDrop(TSharedPtr<FDragDropOperation> DragDropOperation)
+{
+	return FReply::Unhandled(); 
+}
+
+bool UNiagaraStackSpacer::OnStackSpacerAllowDrop(TSharedPtr<FDragDropOperation> DragDropOperation)
+{
+	return false;
+}
diff --git a/Engine/Plugins/FX/Niagara/Source/NiagaraEditor/Public/NiagaraEditorUtilities.h b/Engine/Plugins/FX/Niagara/Source/NiagaraEditor/Public/NiagaraEditorUtilities.h
--- a/Engine/Plugins/FX/Niagara/Source/NiagaraEditor/Public/NiagaraEditorUtilities.h
+++ b/Engine/Plugins/FX/Niagara/Source/NiagaraEditor/Public/NiagaraEditorUtilities.h
@@ -137,7 +137,7 @@ namespace FNiagaraEditorUtilities
 
 	void PreprocessFunctionGraph(const UEdGraphSchema_Niagara* Schema, UNiagaraGraph* Graph, const TArray<UEdGraphPin*>& CallInputs, const TArray<UEdGraphPin*>& CallOutputs, ENiagaraScriptUsage ScriptUsage);
 
-	UNiagaraNodeOutput* GetScriptOutputNode(UNiagaraScript& Script);
+	NIAGARAEDITOR_API UNiagaraNodeOutput* GetScriptOutputNode(UNiagaraScript& Script);
 
 	UNiagaraScript* GetScriptFromSystem(UNiagaraSystem& System, FGuid EmitterHandleId, ENiagaraScriptUsage Usage, FGuid UsageId);
 
diff --git a/Engine/Plugins/FX/Niagara/Source/NiagaraEditor/Public/ViewModels/NiagaraEmitterViewModel.h b/Engine/Plugins/FX/Niagara/Source/NiagaraEditor/Public/ViewModels/NiagaraEmitterViewModel.h
--- a/Engine/Plugins/FX/Niagara/Source/NiagaraEditor/Public/ViewModels/NiagaraEmitterViewModel.h
+++ b/Engine/Plugins/FX/Niagara/Source/NiagaraEditor/Public/ViewModels/NiagaraEmitterViewModel.h
@@ -41,7 +41,7 @@ public:
 	void SetSimulation(TWeakPtr<FNiagaraEmitterInstance> InSimulation);
 
 	/** Gets the emitter represented by this view model. */
-	UNiagaraEmitter* GetEmitter();
+	NIAGARAEDITOR_API UNiagaraEmitter* GetEmitter();
 
 	/** Gets text representing stats for the emitter. */
 	//~ TODO: Instead of a single string here, we should probably have separate controls with tooltips etc.
diff --git a/Engine/Plugins/FX/Niagara/Source/NiagaraEditor/Public/ViewModels/Stack/NiagaraStackGraphUtilities.h b/Engine/Plugins/FX/Niagara/Source/NiagaraEditor/Public/ViewModels/Stack/NiagaraStackGraphUtilities.h
--- a/Engine/Plugins/FX/Niagara/Source/NiagaraEditor/Public/ViewModels/Stack/NiagaraStackGraphUtilities.h
+++ b/Engine/Plugins/FX/Niagara/Source/NiagaraEditor/Public/ViewModels/Stack/NiagaraStackGraphUtilities.h
@@ -106,7 +106,7 @@ namespace FNiagaraStackGraphUtilities
 	
 	bool FindScriptModulesInStack(FAssetData ModuleScriptAsset, UNiagaraNodeOutput& TargetOutputNode, TArray<UNiagaraNodeFunctionCall*> OutFunctionCalls);
 
-	UNiagaraNodeAssignment* AddParameterModuleToStack(const TArray<FNiagaraVariable>& ParameterVariables, UNiagaraNodeOutput& TargetOutputNode, int32 TargetIndex, const TArray<FString>& InDefaultValues);
+	NIAGARAEDITOR_API UNiagaraNodeAssignment* AddParameterModuleToStack(const TArray<FNiagaraVariable>& ParameterVariables, UNiagaraNodeOutput& TargetOutputNode, int32 TargetIndex, const TArray<FString>& InDefaultValues);
 		
 	TOptional<bool> GetModuleIsEnabled(UNiagaraNodeFunctionCall& FunctionCallNode);
 
@@ -140,6 +140,8 @@ namespace FNiagaraStackGraphUtilities
 
 	bool IsValidDefaultDynamicInput(UNiagaraScript& OwningScript, UEdGraphPin& DefaultPin);
 
+	bool ParameterIsCompatibleWithScriptUsage(FNiagaraVariable Parameter, ENiagaraScriptUsage Usage);
+
 	bool DoesDynamicInputMatchDefault(
 		FString EmitterUniqueName,
 		UNiagaraScript& OwningScript,
diff --git a/Engine/Plugins/FX/Niagara/Source/NiagaraEditor/Public/ViewModels/Stack/NiagaraStackModuleSpacer.h b/Engine/Plugins/FX/Niagara/Source/NiagaraEditor/Public/ViewModels/Stack/NiagaraStackModuleSpacer.h
new file mode 100644
--- /dev/null
+++ b/Engine/Plugins/FX/Niagara/Source/NiagaraEditor/Public/ViewModels/Stack/NiagaraStackModuleSpacer.h
@@ -0,0 +1,29 @@
+// Copyright 1998-2018 Epic Games, Inc. All Rights Reserved.
+
+#pragma once
+
+#include "ViewModels/Stack/NiagaraStackSpacer.h"
+#include "NiagaraStackScriptItemGroup.h"
+#include "NiagaraStackModuleSpacer.generated.h"
+
+
+UCLASS()
+class NIAGARAEDITOR_API UNiagaraStackModuleSpacer : public UNiagaraStackSpacer
+{
+	GENERATED_BODY()
+
+public:
+
+	DECLARE_DELEGATE_TwoParams(FOnStackSpacerAcceptDrop, const UNiagaraStackModuleSpacer*, const FNiagaraVariable&)
+
+	void Initialize(FRequiredEntryData InRequiredEntryData, ENiagaraScriptUsage InScriptUsage, FName SpacerKey = NAME_None, float InSpacerScale = 1.0f, EStackRowStyle InRowStyle = UNiagaraStackEntry::EStackRowStyle::None);
+
+	//~ DragDropHandling
+	virtual FReply OnStackSpacerDrop(TSharedPtr<FDragDropOperation> DragDropOperation) override; 
+	virtual bool OnStackSpacerAllowDrop(TSharedPtr<FDragDropOperation> DragDropOperation) override;
+
+	FOnStackSpacerAcceptDrop OnStackSpacerAcceptDrop;
+
+private:
+	ENiagaraScriptUsage ItemGroupScriptUsage;
+}; 
diff --git a/Engine/Plugins/FX/Niagara/Source/NiagaraEditor/Public/ViewModels/Stack/NiagaraStackScriptItemGroup.h b/Engine/Plugins/FX/Niagara/Source/NiagaraEditor/Public/ViewModels/Stack/NiagaraStackScriptItemGroup.h
--- a/Engine/Plugins/FX/Niagara/Source/NiagaraEditor/Public/ViewModels/Stack/NiagaraStackScriptItemGroup.h
+++ b/Engine/Plugins/FX/Niagara/Source/NiagaraEditor/Public/ViewModels/Stack/NiagaraStackScriptItemGroup.h
@@ -14,6 +14,7 @@ class UNiagaraNodeFunctionCall;
 class FScriptItemGroupAddUtilities;
 class UNiagaraStackModuleItem;
 class UEdGraph;
+class UNiagaraStackModuleSpacer;
 
 UCLASS()
 class NIAGARAEDITOR_API UNiagaraStackScriptItemGroup : public UNiagaraStackItemGroup
@@ -31,6 +32,14 @@ public:
 
 	ENiagaraScriptUsage GetScriptUsage() const { return ScriptUsage; }
 	FGuid GetScriptUsageId() const { return ScriptUsageId; }
+	UNiagaraNodeOutput* GetScriptOutputNode() const;
+
+	/**
+	 * Add a Set Variables module to the stack.
+	 * @InModuleSpacer: Target UNiagaraStackModuleSpacer which we use to derive the insertion index.
+	 * @InVariable: Initial FNiagaraVariable to populate the new Set Variables module.
+	 */
+	void AddParameterModuleToStack(const UNiagaraStackModuleSpacer* InModuleSpacer, const FNiagaraVariable& InVariable);
 
 protected:
 	virtual void RefreshChildrenInternal(const TArray<UNiagaraStackEntry*>& CurrentChildren, TArray<UNiagaraStackEntry*>& NewChildren, TArray<FStackIssue>& NewIssues) override;
diff --git a/Engine/Plugins/FX/Niagara/Source/NiagaraEditor/Public/ViewModels/Stack/NiagaraStackSpacer.h b/Engine/Plugins/FX/Niagara/Source/NiagaraEditor/Public/ViewModels/Stack/NiagaraStackSpacer.h
--- a/Engine/Plugins/FX/Niagara/Source/NiagaraEditor/Public/ViewModels/Stack/NiagaraStackSpacer.h
+++ b/Engine/Plugins/FX/Niagara/Source/NiagaraEditor/Public/ViewModels/Stack/NiagaraStackSpacer.h
@@ -3,6 +3,8 @@
 #pragma once
 
 #include "ViewModels/Stack/NiagaraStackEntry.h"
+#include "Input/DragAndDrop.h"
+#include "Input/Reply.h"
 #include "NiagaraStackSpacer.generated.h"
 
 UCLASS()
@@ -22,6 +24,10 @@ public:
 
 	float GetSpacerScale() const;
 
+	//~ DragDropHandling
+	virtual FReply OnStackSpacerDrop(TSharedPtr<FDragDropOperation> DragDropOperation);
+	virtual bool OnStackSpacerAllowDrop(TSharedPtr<FDragDropOperation> DragDropOperation);
+
 private:
 	FName SpacerKey;
 	float SpacerScale;
diff --git a/Engine/Plugins/FX/Niagara/Source/NiagaraEditorWidgets/Private/SNiagaraStack.cpp b/Engine/Plugins/FX/Niagara/Source/NiagaraEditorWidgets/Private/SNiagaraStack.cpp
--- a/Engine/Plugins/FX/Niagara/Source/NiagaraEditorWidgets/Private/SNiagaraStack.cpp
+++ b/Engine/Plugins/FX/Niagara/Source/NiagaraEditorWidgets/Private/SNiagaraStack.cpp
@@ -52,6 +52,7 @@
 #include "NiagaraStackEditorData.h"
 #include "ScopedTransaction.h"
 #include "Widgets/Layout/SWrapBox.h"
+#include "Stack/SNiagaraStackSpacer.h"
 
 /** Contains data for a socket drag and drop operation in the StackEntry node. */
 class FNiagaraStackEntryDragDropOp : public FDecoratedDragDropOp
@@ -624,7 +625,6 @@ FReply SNiagaraStack::OnRowAcceptDrop(const FDragDropEvent& InDragDropEvent, EIt
 	TSharedPtr<FNiagaraStackEntryDragDropOp> DragDropOp = InDragDropEvent.GetOperationAs<FNiagaraStackEntryDragDropOp>();
 	if (DragDropOp.IsValid())
 	{
-
 		if (ensureMsgf(InTargetEntry->Drop(DragDropOp->GetDraggedEntries()).IsSet(), TEXT("Failed to drop stack entry when can drop returned true")))
 		{
 			return FReply::Handled();
@@ -810,8 +810,9 @@ SNiagaraStack::FRowWidgets SNiagaraStack::ConstructNameAndValueWidgetsForItem(UN
 	{
 		FMargin ContentPadding = Container->GetContentPadding();
 		Container->SetContentPadding(FMargin(ContentPadding.Left, 0, ContentPadding.Right, 0));
-		return FRowWidgets(SNew(SBox)
-			.HeightOverride(SpacerHeight * CastChecked<UNiagaraStackSpacer>(Item)->GetSpacerScale()));
+		UNiagaraStackSpacer* SpacerItem = CastChecked<UNiagaraStackSpacer>(Item); 
+		return FRowWidgets(SNew(SNiagaraStackSpacer, *SpacerItem)
+			.HeightOverride(SpacerHeight * SpacerItem->GetSpacerScale()));
 	}
 	else if (Item->IsA<UNiagaraStackItemGroup>())
 	{
diff --git a/Engine/Plugins/FX/Niagara/Source/NiagaraEditorWidgets/Private/Stack/SNiagaraStackSpacer.cpp b/Engine/Plugins/FX/Niagara/Source/NiagaraEditorWidgets/Private/Stack/SNiagaraStackSpacer.cpp
new file mode 100644
--- /dev/null
+++ b/Engine/Plugins/FX/Niagara/Source/NiagaraEditorWidgets/Private/Stack/SNiagaraStackSpacer.cpp
@@ -0,0 +1,49 @@
+// Copyright 1998-2018 Epic Games, Inc. All Rights Reserved.
+
+#include "SNiagaraStackSpacer.h"
+#include "SDropTarget.h"
+#include "NiagaraTypes.h"
+#include "NiagaraActions.h"
+#include "NiagaraConstants.h"
+#include "ViewModels/Stack/NiagaraStackGraphUtilities.h"
+#include "ViewModels/Stack/NiagaraStackScriptItemGroup.h"
+
+void SNiagaraStackSpacer::Construct(const FArguments& InArgs, UNiagaraStackSpacer& InStackSpacer) 
+{
+	HeightOverride = InArgs._HeightOverride;
+
+	StackSpacer = &InStackSpacer;
+
+	ChildSlot
+	[
+		SNew(SDropTarget)
+		.OnAllowDrop(this, &SNiagaraStackSpacer::OnSStackSpacerAllowDrop)
+		.OnDrop(this, &SNiagaraStackSpacer::OnSStackSpacerDrop)
+		.Content()
+		[
+			SNew(SBox)
+			.HeightOverride(HeightOverride)
+		]
+	];
+	
+}
+
+bool SNiagaraStackSpacer::OnSStackSpacerAllowDrop(TSharedPtr<FDragDropOperation> DragDropOperation) const
+{
+	if (StackSpacer)
+	{
+		return StackSpacer->OnStackSpacerAllowDrop(DragDropOperation);
+	}
+	return false;
+}
+
+FReply SNiagaraStackSpacer::OnSStackSpacerDrop(TSharedPtr<FDragDropOperation> DragDropOperation) const
+{
+	if (StackSpacer)
+	{
+		return StackSpacer->OnStackSpacerDrop(DragDropOperation);
+	}
+	return FReply::Unhandled();
+}
+
+
diff --git a/Engine/Plugins/FX/Niagara/Source/NiagaraEditorWidgets/Private/Stack/SNiagaraStackSpacer.h b/Engine/Plugins/FX/Niagara/Source/NiagaraEditorWidgets/Private/Stack/SNiagaraStackSpacer.h
new file mode 100644
--- /dev/null
+++ b/Engine/Plugins/FX/Niagara/Source/NiagaraEditorWidgets/Private/Stack/SNiagaraStackSpacer.h
@@ -0,0 +1,32 @@
+// Copyright 1998-2018 Epic Games, Inc. All Rights Reserved.
+
+#pragma once
+
+#include "Widgets/Layout/SBox.h"
+#include "Widgets/SCompoundWidget.h"
+#include "ViewModels/Stack/NiagaraStackSpacer.h"
+
+class SNiagaraStackSpacer : public SCompoundWidget 
+{
+public:
+	SLATE_BEGIN_ARGS(SNiagaraStackSpacer)
+
+	// Ported HeightOverride from SBox to expose setting box size when instantiating in SNiagaraStack
+	: _HeightOverride(FOptionalSize())
+	{ }
+	/** When specified, ignore the content's desired size and report the HeightOverride as the Box's desired height. */
+	SLATE_ATTRIBUTE(FOptionalSize, HeightOverride)
+
+	SLATE_END_ARGS(); 
+
+	void Construct(const FArguments& InArgs, UNiagaraStackSpacer& InStackSpacer);
+
+private:
+	UNiagaraStackSpacer* StackSpacer;
+
+	bool OnSStackSpacerAllowDrop(TSharedPtr<FDragDropOperation> DragDropOperation) const;
+	FReply OnSStackSpacerDrop(TSharedPtr<FDragDropOperation> DragDropOperation) const;
+
+	/** When specified, ignore the content's desired size and report the.HeightOverride as the Box's desired height. */
+	TAttribute<FOptionalSize> HeightOverride;
+};
\ No newline at end of file
