Unshelved from dev-niagara:

Collapse irrelevant UI options for Emitters

diff --git a/Engine/Plugins/FX/Niagara/Source/NiagaraEditor/Private/Sequencer/NiagaraSequence/NiagaraEmitterTrackEditor.cpp b/Engine/Plugins/FX/Niagara/Source/NiagaraEditor/Private/Sequencer/NiagaraSequence/NiagaraEmitterTrackEditor.cpp
--- a/Engine/Plugins/FX/Niagara/Source/NiagaraEditor/Private/Sequencer/NiagaraSequence/NiagaraEmitterTrackEditor.cpp
+++ b/Engine/Plugins/FX/Niagara/Source/NiagaraEditor/Private/Sequencer/NiagaraSequence/NiagaraEmitterTrackEditor.cpp
@@ -51,6 +51,7 @@ public:
 				.ToolTipText(LOCTEXT("EnabledTooltip", "Toggle whether or not this emitter is enabled."))
 				.IsChecked(this, &SEmitterTrackWidget::GetEnabledCheckState)
 				.OnCheckStateChanged(this, &SEmitterTrackWidget::OnEnabledCheckStateChanged)
+				.Visibility(this, &SEmitterTrackWidget::GetEnableCheckboxVisibility)
 			]
 			// Isolate toggle
 			+ SHorizontalBox::Slot()
@@ -64,6 +65,7 @@ public:
 				.ContentPadding(1)
 				.ToolTipText(this, &SEmitterTrackWidget::GetToggleIsolateToolTip)
 				.OnClicked(this, &SEmitterTrackWidget::OnToggleIsolateButtonClicked)
+				.Visibility(this, &SEmitterTrackWidget::GetIsolateToggleVisibility)
 				.Content()
 				[
 					SNew(SImage)
@@ -170,6 +172,16 @@ private:
 		return FReply::Handled();
 	}
 
+	EVisibility GetEnableCheckboxVisibility() const
+	{
+		return EmitterTrack.IsValid() && EmitterTrack->GetSystemViewModel().GetEditMode() == ENiagaraSystemViewModelEditMode::SystemAsset ? EVisibility::Visible : EVisibility::Collapsed;
+	}
+
+	EVisibility GetIsolateToggleVisibility() const
+	{
+		return EmitterTrack.IsValid() && EmitterTrack->GetSystemViewModel().GetEditMode() == ENiagaraSystemViewModelEditMode::SystemAsset ? EVisibility::Visible : EVisibility::Collapsed;
+	}
+
 private:
 	TWeakObjectPtr<UMovieSceneNiagaraEmitterTrack> EmitterTrack;
 	mutable TOptional<FText> TrackErrorIconToolTip;
diff --git a/Engine/Plugins/FX/Niagara/Source/NiagaraEditorWidgets/Private/SNiagaraStack.cpp b/Engine/Plugins/FX/Niagara/Source/NiagaraEditorWidgets/Private/SNiagaraStack.cpp
--- a/Engine/Plugins/FX/Niagara/Source/NiagaraEditorWidgets/Private/SNiagaraStack.cpp
+++ b/Engine/Plugins/FX/Niagara/Source/NiagaraEditorWidgets/Private/SNiagaraStack.cpp
@@ -264,6 +264,7 @@ void SNiagaraStack::ConstructHeaderWidget()
 					.ToolTipText(LOCTEXT("EnabledToolTip", "Toggles whether this emitter is enabled. Disabled emitters don't simulate or render."))
 					.IsChecked(StackViewModel->GetEmitterHandleViewModel()->AsShared(), &FNiagaraEmitterHandleViewModel::GetIsEnabledCheckState)
 					.OnCheckStateChanged(StackViewModel->GetEmitterHandleViewModel()->AsShared(), &FNiagaraEmitterHandleViewModel::OnIsEnabledCheckStateChanged)
+					.Visibility(this, &SNiagaraStack::GetEnableCheckboxVisibility) 
 				]
 				// Pin
 				+ SHorizontalBox::Slot()
@@ -277,6 +278,7 @@ void SNiagaraStack::ConstructHeaderWidget()
 					.ForegroundColor(this, &SNiagaraStack::GetPinColor)
 					.ContentPadding(2)
 					.OnClicked(this, &SNiagaraStack::PinButtonPressed)
+					.Visibility(this, &SNiagaraStack::GetPinEmitterVisibility)
 					.Content()
 					[
 						SNew(STextBlock)
@@ -328,7 +330,7 @@ void SNiagaraStack::ConstructHeaderWidget()
 					.ForegroundColor(this, &SNiagaraStack::GetPinColor)
 					.ContentPadding(2)
 					.OnClicked(this, &SNiagaraStack::OpenSourceEmitter)
-					.Visibility(CanOpenSourceEmitter() ? EVisibility::Visible : EVisibility::Collapsed)
+					.Visibility(this, &SNiagaraStack::GetOpenSourceEmitterVisibility)
 					// GoToSource icon is 30x30px so we scale it down to stay in line with other 12x12px UI
 					.DesiredSizeScale(FVector2D(0.55f, 0.55f))
 					.Content()
@@ -487,6 +489,21 @@ FReply SNiagaraStack::OpenSourceEmitter()
 	return FReply::Handled();
 }
 
+EVisibility SNiagaraStack::GetEnableCheckboxVisibility() const
+{
+	return StackViewModel->GetSystemViewModel()->GetEditMode() == ENiagaraSystemViewModelEditMode::SystemAsset ? EVisibility::Visible : EVisibility::Collapsed;
+}
+
+EVisibility SNiagaraStack::GetPinEmitterVisibility() const
+{
+	return StackViewModel->GetSystemViewModel()->GetEditMode() == ENiagaraSystemViewModelEditMode::SystemAsset ? EVisibility::Visible : EVisibility::Collapsed;
+}
+
+EVisibility SNiagaraStack::GetOpenSourceEmitterVisibility() const
+{
+	return CanOpenSourceEmitter() ? EVisibility::Visible : EVisibility::Collapsed;
+}
+
 bool SNiagaraStack::CanOpenSourceEmitter() const
 {
 	if (StackViewModel && StackViewModel->GetEmitterHandleViewModel().IsValid() && StackViewModel->GetEmitterHandleViewModel()->GetEmitterHandle()
diff --git a/Engine/Plugins/FX/Niagara/Source/NiagaraEditorWidgets/Private/SNiagaraStack.h b/Engine/Plugins/FX/Niagara/Source/NiagaraEditorWidgets/Private/SNiagaraStack.h
--- a/Engine/Plugins/FX/Niagara/Source/NiagaraEditorWidgets/Private/SNiagaraStack.h
+++ b/Engine/Plugins/FX/Niagara/Source/NiagaraEditorWidgets/Private/SNiagaraStack.h
@@ -73,6 +73,9 @@ private:
 	FSlateColor GetPinColor() const;
 	FReply PinButtonPressed();
 	FReply OpenSourceEmitter();
+	EVisibility GetEnableCheckboxVisibility() const;
+	EVisibility GetPinEmitterVisibility() const;
+	EVisibility GetOpenSourceEmitterVisibility() const;
 
 	// source name handling
 	FText GetSourceEmitterNameText() const;
