Niagara - Don't delete the gpu readback from the gamethread when resetting the gpu context.

diff --git a/Engine/Plugins/FX/Niagara/Source/Niagara/Classes/NiagaraScriptExecutionContext.h b/Engine/Plugins/FX/Niagara/Source/Niagara/Classes/NiagaraScriptExecutionContext.h
--- a/Engine/Plugins/FX/Niagara/Source/Niagara/Classes/NiagaraScriptExecutionContext.h
+++ b/Engine/Plugins/FX/Niagara/Source/Niagara/Classes/NiagaraScriptExecutionContext.h
@@ -110,13 +110,13 @@ struct FNiagaraComputeExecutionContext
 
 	void Reset()
 	{
-		AccumulatedSpawnRate = 0;
-		PendingExecutionQueueMask = 0;
-		if (GPUDataReadback)
-		{
-			delete GPUDataReadback;
-		}
-		GPUDataReadback = nullptr;
+		ENQUEUE_UNIQUE_RENDER_COMMAND_ONEPARAMETER(
+			ResetRT,
+			FNiagaraComputeExecutionContext*, Context, this,
+			{
+				Context->ResetInternal();
+			}
+		);
 	}
 
 	void InitParams(UNiagaraScript* InGPUComputeScript, UNiagaraScript* InSpawnScript, UNiagaraScript *InUpdateScript, ENiagaraSimTarget SimTarget)
@@ -179,6 +179,20 @@ struct FNiagaraComputeExecutionContext
 		return true;
 	}
 
+private:
+	void ResetInternal()
+	{
+		checkf(IsInRenderingThread(), TEXT("Can only reset the gpu context from the render thread"));
+		AccumulatedSpawnRate = 0;
+		PendingExecutionQueueMask = 0;
+		if (GPUDataReadback)
+		{
+			delete GPUDataReadback;
+		}
+		GPUDataReadback = nullptr;
+	}
+
+public:
 	const TArray<FNiagaraEventScriptProperties> &GetEventHandlers() const { return EventHandlerScriptProps; }
 
 	class FNiagaraDataSet *MainDataSet;
