Niagara - Fix multiple issues with the new emitter and system dialogs.
+ Double click and "Enter" now confirm the selection in both dialogs.
+ Emitters are now removed from the selection properly in the system dialog.
+ Duplicate emitters can now be added in the system dialog.

diff --git a/Engine/Plugins/FX/Niagara/Source/NiagaraEditor/Private/SItemSelector.h b/Engine/Plugins/FX/Niagara/Source/NiagaraEditor/Private/SItemSelector.h
--- a/Engine/Plugins/FX/Niagara/Source/NiagaraEditor/Private/SItemSelector.h
+++ b/Engine/Plugins/FX/Niagara/Source/NiagaraEditor/Private/SItemSelector.h
@@ -18,6 +18,7 @@ public:
 	DECLARE_DELEGATE_RetVal_TwoParams(bool, FOnDoesItemMatchFilterText, const FText& /* Filter text */, const ItemType& /* Item */);
 	DECLARE_DELEGATE_RetVal_OneParam(TSharedRef<SWidget>, FOnGenerateWidgetForCategory, const CategoryType& /* Category */);
 	DECLARE_DELEGATE_RetVal_OneParam(TSharedRef<SWidget>, FOnGenerateWidgetForItem, const ItemType& /* Item */);
+	DECLARE_DELEGATE_OneParam(FOnItemActivated, const ItemType& /* Item */);
 
 public:
 	SLATE_BEGIN_ARGS(SItemSelector)
@@ -49,6 +50,9 @@ public:
 
 		/** The delegate which is used to generate widgets for the items to be selected. */
 		SLATE_EVENT(FOnGenerateWidgetForItem, OnGenerateWidgetForItem)
+
+		/** A delegate which is called when an item is activated by either double clicking on it or by pressing enter while it's selected. */
+		SLATE_EVENT(FOnItemActivated, OnItemActivated)
 	SLATE_END_ARGS();
 
 private:
@@ -391,6 +395,7 @@ public:
 		OnDoesItemMatchFilterText = InArgs._OnDoesItemMatchFilterText;
 		OnGenerateWidgetForCategory = InArgs._OnGenerateWidgetForCategory;
 		OnGenerateWidgetForItem = InArgs._OnGenerateWidgetForItem;
+		OnItemActivated = InArgs._OnItemActivated;
 
 		checkf(OnGetCategoriesForItem.IsBound() == false || OnCompareCategoriesForEquality.IsBound(), TEXT("OnCompareCategoriesForEquality must be bound if OnGenerateCategoriesForItem is bound."));
 		checkf(OnGetCategoriesForItem.IsBound() == false || OnGenerateWidgetForCategory.IsBound(), TEXT("OnGenerateWidgetForCategory must be bound if OnGenerateCategoriesForItem is bound."));
@@ -416,6 +421,7 @@ public:
 				.SelectionMode(InArgs._AllowMultiselect ? ESelectionMode::Multi : ESelectionMode::SingleToggle)
 				.OnGenerateRow(this, &SItemSelector::OnGenerateRow)
 				.OnGetChildren(this, &SItemSelector::OnGetChildren)
+				.OnMouseButtonDoubleClick(this, &SItemSelector::OnMouseDoubleClick)
 				.TreeItemsSource(ViewModel->GetRootItems())
 			]
 		];
@@ -440,6 +446,22 @@ public:
 		return SelectedItems;
 	}
 
+	virtual FReply OnKeyDown(const FGeometry& MyGeometry, const FKeyEvent& InKeyEvent)
+	{
+		if (InKeyEvent.GetKey() == EKeys::Enter && OnItemActivated.IsBound())
+		{
+			TArray<TSharedRef<FItemSelectorItemViewModel>> SelectedItemViewModels;
+			ItemTree->GetSelectedItems(SelectedItemViewModels);
+			if (SelectedItemViewModels.Num() == 1 && SelectedItemViewModels[0]->GetType() == EItemSelectorItemViewModelType::Item)
+			{
+				TSharedRef<FItemSelectorItemContainerViewModel> ItemContainer = StaticCastSharedRef<FItemSelectorItemContainerViewModel>(SelectedItemViewModels[0]);
+				OnItemActivated.Execute(ItemContainer->GetItem());
+				return FReply::Handled();
+			}
+		}
+		return SCompoundWidget::OnKeyDown(MyGeometry, InKeyEvent);
+	}
+
 private:
 	EVisibility GetSearchBoxVisibility() const
 	{
@@ -488,6 +510,15 @@ private:
 		Item->GetChildren(OutChildren);
 	}
 
+	void OnMouseDoubleClick(TSharedRef<FItemSelectorItemViewModel> ItemDoubleClicked)
+	{
+		if (OnItemActivated.IsBound() && ItemDoubleClicked->GetType() == EItemSelectorItemViewModelType::Item)
+		{
+			TSharedRef<FItemSelectorItemContainerViewModel> ItemContainer = StaticCastSharedRef<FItemSelectorItemContainerViewModel>(ItemDoubleClicked);
+			OnItemActivated.Execute(ItemContainer->GetItem());
+		}
+	}
+
 	void ExpandTree()
 	{
 		TArray<TSharedRef<FItemSelectorItemViewModel>> ItemsToProcess;
@@ -511,6 +542,7 @@ private:
 	FOnDoesItemMatchFilterText OnDoesItemMatchFilterText;
 	FOnGenerateWidgetForCategory OnGenerateWidgetForCategory;
 	FOnGenerateWidgetForItem OnGenerateWidgetForItem;
+	FOnItemActivated OnItemActivated;
 
 	TSharedPtr<FItemSelectorViewModel> ViewModel;
 	TSharedPtr<SSearchBox> SearchBox;
diff --git a/Engine/Plugins/FX/Niagara/Source/NiagaraEditor/Private/SNewEmitterDialog.cpp b/Engine/Plugins/FX/Niagara/Source/NiagaraEditor/Private/SNewEmitterDialog.cpp
--- a/Engine/Plugins/FX/Niagara/Source/NiagaraEditor/Private/SNewEmitterDialog.cpp
+++ b/Engine/Plugins/FX/Niagara/Source/NiagaraEditor/Private/SNewEmitterDialog.cpp
@@ -20,8 +20,10 @@ void SNewEmitterDialog::Construct(const FArguments& InArgs)
 	FAssetPickerConfig AssetPickerConfig;
 	AssetPickerConfig.SelectionMode = ESelectionMode::SingleToggle;
 	AssetPickerConfig.InitialAssetViewType = EAssetViewType::List;
+	AssetPickerConfig.ThumbnailScale = .4f;
 	AssetPickerConfig.Filter.ClassNames.Add(UNiagaraEmitter::StaticClass()->GetFName());
 	AssetPickerConfig.GetCurrentSelectionDelegates.Add(&GetSelectedEmitterAssetsFromPicker);
+	AssetPickerConfig.OnAssetsActivated.BindSP(this, &SNewEmitterDialog::OnEmitterAssetsActivated);
 
 	FContentBrowserModule& ContentBrowserModule = FModuleManager::Get().LoadModuleChecked<FContentBrowserModule>(TEXT("ContentBrowser"));
 	TSharedRef<SWidget> AssetPicker = ContentBrowserModule.Get().CreateAssetPicker(AssetPickerConfig);
@@ -32,7 +34,8 @@ void SNewEmitterDialog::Construct(const FArguments& InArgs)
 				LOCTEXT("CreateFromTemplateLabel", "Create a new emitter from an emitter template"),
 				LOCTEXT("TemplatesPickerHeader", "Select a Template Emitter"),
 				SNiagaraNewAssetDialog::FOnGetSelectedAssetsFromPicker::CreateSP(this, &SNewEmitterDialog::GetSelectedEmitterTemplateAssets),
-				SAssignNew(TemplateAssetPicker, SNiagaraTemplateAssetPicker, UNiagaraEmitter::StaticClass())),
+				SAssignNew(TemplateAssetPicker, SNiagaraTemplateAssetPicker, UNiagaraEmitter::StaticClass())
+				.OnTemplateAssetActivated(this, &SNewEmitterDialog::OnTemplateAssetActivated)),
 			SNiagaraNewAssetDialog::FNiagaraNewAssetDialogOption(
 				LOCTEXT("CreateFromOtherEmitterLabel", "Copy an existing emitter from your project content"),
 				LOCTEXT("ProjectEmitterPickerHeader", "Select a Project Emitter"),
@@ -72,4 +75,17 @@ TArray<FAssetData> SNewEmitterDialog::GetSelectedProjectEmiterAssets()
 	return GetSelectedEmitterAssetsFromPicker.Execute();
 }
 
+void SNewEmitterDialog::OnTemplateAssetActivated(const FAssetData& ActivatedTemplateAsset)
+{
+	ConfirmSelection();
+}
+
+void SNewEmitterDialog::OnEmitterAssetsActivated(const TArray<FAssetData>& ActivatedAssets, EAssetTypeActivationMethod::Type ActivationMethod)
+{
+	if ((ActivationMethod == EAssetTypeActivationMethod::DoubleClicked || ActivationMethod == EAssetTypeActivationMethod::Opened) && ActivatedAssets.Num() == 1)
+	{
+		ConfirmSelection();
+	}
+}
+
 #undef LOCTEXT_NAMESPACE
\ No newline at end of file
diff --git a/Engine/Plugins/FX/Niagara/Source/NiagaraEditor/Private/SNewEmitterDialog.h b/Engine/Plugins/FX/Niagara/Source/NiagaraEditor/Private/SNewEmitterDialog.h
--- a/Engine/Plugins/FX/Niagara/Source/NiagaraEditor/Private/SNewEmitterDialog.h
+++ b/Engine/Plugins/FX/Niagara/Source/NiagaraEditor/Private/SNewEmitterDialog.h
@@ -28,6 +28,10 @@ private:
 
 	TArray<FAssetData> GetSelectedProjectEmiterAssets();
 
+	void OnTemplateAssetActivated(const FAssetData& ActivatedTemplateAsset);
+
+	void OnEmitterAssetsActivated(const TArray<FAssetData>& ActivatedAssets, EAssetTypeActivationMethod::Type ActivationMethod);
+
 private:
 	TSharedPtr<SNiagaraTemplateAssetPicker> TemplateAssetPicker;
 
diff --git a/Engine/Plugins/FX/Niagara/Source/NiagaraEditor/Private/SNewSystemDialog.cpp b/Engine/Plugins/FX/Niagara/Source/NiagaraEditor/Private/SNewSystemDialog.cpp
--- a/Engine/Plugins/FX/Niagara/Source/NiagaraEditor/Private/SNewSystemDialog.cpp
+++ b/Engine/Plugins/FX/Niagara/Source/NiagaraEditor/Private/SNewSystemDialog.cpp
@@ -27,16 +27,20 @@ void SNewSystemDialog::Construct(const FArguments& InArgs)
 	FAssetPickerConfig SystemAssetPickerConfig;
 	SystemAssetPickerConfig.SelectionMode = ESelectionMode::SingleToggle;
 	SystemAssetPickerConfig.InitialAssetViewType = EAssetViewType::List;
+	SystemAssetPickerConfig.ThumbnailScale = .4f;
 	SystemAssetPickerConfig.Filter.ClassNames.Add(UNiagaraSystem::StaticClass()->GetFName());
 	SystemAssetPickerConfig.GetCurrentSelectionDelegates.Add(&GetSelectedSystemAssetsFromPicker);
+	SystemAssetPickerConfig.OnAssetsActivated.BindSP(this, &SNewSystemDialog::OnSystemAssetsActivated);
 
 	TSharedRef<SWidget> SystemAssetPicker = ContentBrowserModule.Get().CreateAssetPicker(SystemAssetPickerConfig);
 
 	FAssetPickerConfig EmitterAssetPickerConfig;
 	EmitterAssetPickerConfig.SelectionMode = ESelectionMode::Multi;
 	EmitterAssetPickerConfig.InitialAssetViewType = EAssetViewType::List;
+	EmitterAssetPickerConfig.ThumbnailScale = .4f;
 	EmitterAssetPickerConfig.Filter.ClassNames.Add(UNiagaraEmitter::StaticClass()->GetFName());
 	EmitterAssetPickerConfig.GetCurrentSelectionDelegates.Add(&GetSelectedEmitterAssetsFromPicker);
+	EmitterAssetPickerConfig.OnAssetsActivated.BindSP(this, &SNewSystemDialog::OnEmitterAssetsActivated);
 
 	TSharedRef<SWidget> EmitterAssetPicker = ContentBrowserModule.Get().CreateAssetPicker(EmitterAssetPickerConfig);
 
@@ -46,7 +50,8 @@ void SNewSystemDialog::Construct(const FArguments& InArgs)
 				LOCTEXT("CreateFromTemplateLabel", "Create a new system from a system template"),
 				LOCTEXT("TemplateLabel", "Select a System Template"),
 				SNiagaraNewAssetDialog::FOnGetSelectedAssetsFromPicker::CreateSP(this, &SNewSystemDialog::GetSelectedSystemTemplateAssets),
-				SAssignNew(TemplateAssetPicker, SNiagaraTemplateAssetPicker, UNiagaraSystem::StaticClass())),
+				SAssignNew(TemplateAssetPicker, SNiagaraTemplateAssetPicker, UNiagaraSystem::StaticClass())
+				.OnTemplateAssetActivated(this, &SNewSystemDialog::OnTemplateAssetActivated)),
 			SNiagaraNewAssetDialog::FNiagaraNewAssetDialogOption(
 				LOCTEXT("CreateFromSelectedEmittersLabel", "Create a new system from a set of selected emitters"),
 				LOCTEXT("ProjectEmittersLabel", "Select Emitters to Add"),
@@ -168,6 +173,27 @@ TArray<FAssetData> SNewSystemDialog::GetSelectedProjectEmiterAssets()
 	return SelectedEmitterAssets;
 }
 
+void SNewSystemDialog::OnTemplateAssetActivated(const FAssetData& ActivatedTemplateAsset)
+{
+	ConfirmSelection();
+}
+
+void SNewSystemDialog::OnSystemAssetsActivated(const TArray<FAssetData>& ActivatedAssets, EAssetTypeActivationMethod::Type ActivationMethod)
+{
+	if ((ActivationMethod == EAssetTypeActivationMethod::DoubleClicked || ActivationMethod == EAssetTypeActivationMethod::Opened) && ActivatedAssets.Num() == 1)
+	{
+		ConfirmSelection();
+	}
+}
+
+void SNewSystemDialog::OnEmitterAssetsActivated(const TArray<FAssetData>& ActivatedAssets, EAssetTypeActivationMethod::Type ActivationMethod)
+{
+	if (ActivationMethod == EAssetTypeActivationMethod::DoubleClicked || ActivationMethod == EAssetTypeActivationMethod::Opened)
+	{
+		AddEmitterAssetsToSelection(ActivatedAssets);
+	}
+}
+
 bool SNewSystemDialog::IsAddEmittersToSelectionButtonEnabled() const
 {
 	return GetSelectedEmitterAssetsFromPicker.Execute().Num() > 0;
@@ -176,51 +202,53 @@ bool SNewSystemDialog::IsAddEmittersToSelectionButtonEnabled() const
 FReply SNewSystemDialog::AddEmittersToSelectionButtonClicked()
 {
 	TArray<FAssetData> SelectedEmitterAssetsFromPicker = GetSelectedEmitterAssetsFromPicker.Execute();
-	for (const FAssetData& SelectedEmitterAsset : SelectedEmitterAssetsFromPicker)
+	AddEmitterAssetsToSelection(SelectedEmitterAssetsFromPicker);
+	return FReply::Handled();
+}
+
+void SNewSystemDialog::AddEmitterAssetsToSelection(const TArray<FAssetData>& EmitterAssets)
+{
+	for (const FAssetData& SelectedEmitterAsset : EmitterAssets)
 	{
-		if (SelectedEmitterAssets.Contains(SelectedEmitterAsset) == false)
-		{
-			TSharedPtr<SWidget> SelectedEmitterWidget;
-			SelectedEmitterBox->AddSlot()
-				.Padding(FMargin(0, 0, 5, 0))
+		TSharedPtr<SWidget> SelectedEmitterWidget;
+		SelectedEmitterBox->AddSlot()
+			.Padding(FMargin(0, 0, 5, 0))
+			[
+				SAssignNew(SelectedEmitterWidget, SBorder)
+				.BorderImage(FNiagaraEditorStyle::Get().GetBrush("NiagaraEditor.NewAssetDialog.SubBorder"))
+				.BorderBackgroundColor(FNiagaraEditorStyle::Get().GetColor("NiagaraEditor.NewAssetDialog.SubBorderColor"))
 				[
-					SNew(SBorder)
-					.BorderImage(FNiagaraEditorStyle::Get().GetBrush("NiagaraEditor.NewAssetDialog.SubBorder"))
-					.BorderBackgroundColor(FNiagaraEditorStyle::Get().GetColor("NiagaraEditor.NewAssetDialog.SubBorderColor"))
+					SNew(SHorizontalBox)
+					+ SHorizontalBox::Slot()
+					.AutoWidth()
+					.VAlign(VAlign_Center)
+					.Padding(5, 0, 0, 0)
+					[
+						SNew(STextBlock)
+						.Text(FText::FromName(SelectedEmitterAsset.AssetName))
+					]
+					+ SHorizontalBox::Slot()
+					.AutoWidth()
+					.VAlign(VAlign_Center)
+					.Padding(2, 0, 0, 0)
 					[
-						SAssignNew(SelectedEmitterWidget, SHorizontalBox)
-						+ SHorizontalBox::Slot()
-						.AutoWidth()
-						.VAlign(VAlign_Center)
-						.Padding(5, 0, 0, 0)
+						SNew(SButton)
+						.ButtonStyle(FEditorStyle::Get(), "HoverHintOnly")
+						.OnClicked(this, &SNewSystemDialog::RemoveEmitterFromSelectionButtonClicked, SelectedEmitterAsset)
+						.ToolTipText(LOCTEXT("RemoveSelectedEmitterToolTip", "Remove the selected emitter from the collection\n of emitters to be added to the new system."))
 						[
 							SNew(STextBlock)
-							.Text(FText::FromName(SelectedEmitterAsset.AssetName))
-						]
-						+ SHorizontalBox::Slot()
-						.AutoWidth()
-						.VAlign(VAlign_Center)
-						.Padding(2, 0, 0, 0)
-						[
-							SNew(SButton)
-							.ButtonStyle(FEditorStyle::Get(), "HoverHintOnly")
-							.OnClicked(this, &SNewSystemDialog::RemoveEmitterFromSelectionButtonClicked, SelectedEmitterAsset)
-							.ToolTipText(LOCTEXT("RemoveSelectedEmitterToolTip", "Remove the selected emitter from the collection\n of emitters to be added to the new system."))
-							[
-								SNew(STextBlock)
-								//.TextStyle(FEditorStyle::Get(), "NormalText.Important")
-								.Font(FEditorStyle::Get().GetFontStyle("FontAwesome.10"))
-								.Text(FText::FromString(FString(TEXT("\xf057"))) /*times-circle*/)
-								.ColorAndOpacity(FLinearColor(.8f, .2f, .2f, 1.0f))
-							]
+							//.TextStyle(FEditorStyle::Get(), "NormalText.Important")
+							.Font(FEditorStyle::Get().GetFontStyle("FontAwesome.10"))
+							.Text(FText::FromString(FString(TEXT("\xf057"))) /*times-circle*/)
+							.ColorAndOpacity(FLinearColor(.8f, .2f, .2f, 1.0f))
 						]
 					]
-				];
-			SelectedEmitterAssets.Add(SelectedEmitterAsset);
-			SelectedEmitterAssetWidgets.Add(SelectedEmitterWidget);
-		}
+				]
+			];
+		SelectedEmitterAssets.Add(SelectedEmitterAsset);
+		SelectedEmitterAssetWidgets.Add(SelectedEmitterWidget);
 	}
-	return FReply::Handled();
 }
 
 FReply SNewSystemDialog::RemoveEmitterFromSelectionButtonClicked(FAssetData EmitterAsset)
diff --git a/Engine/Plugins/FX/Niagara/Source/NiagaraEditor/Private/SNewSystemDialog.h b/Engine/Plugins/FX/Niagara/Source/NiagaraEditor/Private/SNewSystemDialog.h
--- a/Engine/Plugins/FX/Niagara/Source/NiagaraEditor/Private/SNewSystemDialog.h
+++ b/Engine/Plugins/FX/Niagara/Source/NiagaraEditor/Private/SNewSystemDialog.h
@@ -32,10 +32,19 @@ private:
 	TArray<FAssetData> GetSelectedProjectSystemAssets();
 
 	TArray<FAssetData> GetSelectedProjectEmiterAssets();
+
+	void OnTemplateAssetActivated(const FAssetData& ActivatedTemplateAsset);
+
+	void OnSystemAssetsActivated(const TArray<FAssetData>& ActivatedAssets, EAssetTypeActivationMethod::Type ActivationMethod);
+
+	void OnEmitterAssetsActivated(const TArray<FAssetData>& ActivatedAssets, EAssetTypeActivationMethod::Type ActivationMethod);
+
 	bool IsAddEmittersToSelectionButtonEnabled() const;
 	
 	FReply AddEmittersToSelectionButtonClicked();
 
+	void AddEmitterAssetsToSelection(const TArray<FAssetData>& EmitterAssets);
+
 	FReply RemoveEmitterFromSelectionButtonClicked(FAssetData EmitterAsset);
 
 private:
diff --git a/Engine/Plugins/FX/Niagara/Source/NiagaraEditor/Private/SNiagaraNewAssetDialog.cpp b/Engine/Plugins/FX/Niagara/Source/NiagaraEditor/Private/SNiagaraNewAssetDialog.cpp
--- a/Engine/Plugins/FX/Niagara/Source/NiagaraEditor/Private/SNiagaraNewAssetDialog.cpp
+++ b/Engine/Plugins/FX/Niagara/Source/NiagaraEditor/Private/SNiagaraNewAssetDialog.cpp
@@ -189,6 +189,20 @@ const TArray<FAssetData>& SNiagaraNewAssetDialog::GetSelectedAssets() const
 	return SelectedAssets;
 }
 
+void SNiagaraNewAssetDialog::ConfirmSelection()
+{
+	const FNiagaraNewAssetDialogOption& SelectedOption = Options[SelectedOptionIndex];
+	if (SelectedOption.OnGetSelectedAssetsFromPicker.IsBound())
+	{
+		SelectedAssets.Append(SelectedOption.OnGetSelectedAssetsFromPicker.Execute());
+	}
+	if (ensureMsgf(SelectedAssets.Num() > 0, TEXT("No assets selected when dialog was confirmed.")))
+	{
+		bUserConfirmedSelection = true;
+	}
+	RequestDestroyWindow();
+}
+
 void SNiagaraNewAssetDialog::OnWindowClosed(const TSharedRef<SWindow>& Window)
 {
 	SaveConfig();
@@ -239,14 +253,7 @@ bool SNiagaraNewAssetDialog::IsOkButtonEnabled() const
 
 FReply SNiagaraNewAssetDialog::OnOkButtonClicked()
 {
-	const FNiagaraNewAssetDialogOption& SelectedOption = Options[SelectedOptionIndex];
-	if (SelectedOption.OnGetSelectedAssetsFromPicker.IsBound())
-	{
-		SelectedAssets.Append(SelectedOption.OnGetSelectedAssetsFromPicker.Execute());
-	}
-	bUserConfirmedSelection = true;
-
-	RequestDestroyWindow();
+	ConfirmSelection();
 	return FReply::Handled();
 }
 
diff --git a/Engine/Plugins/FX/Niagara/Source/NiagaraEditor/Private/SNiagaraNewAssetDialog.h b/Engine/Plugins/FX/Niagara/Source/NiagaraEditor/Private/SNiagaraNewAssetDialog.h
--- a/Engine/Plugins/FX/Niagara/Source/NiagaraEditor/Private/SNiagaraNewAssetDialog.h
+++ b/Engine/Plugins/FX/Niagara/Source/NiagaraEditor/Private/SNiagaraNewAssetDialog.h
@@ -45,6 +45,8 @@ public:
 protected:
 	const TArray<FAssetData>& GetSelectedAssets() const;
 
+	void ConfirmSelection();
+
 private:
 	void OnWindowClosed(const TSharedRef<SWindow>& Window);
 
diff --git a/Engine/Plugins/FX/Niagara/Source/NiagaraEditor/Private/SNiagaraTemplateAssetPicker.cpp b/Engine/Plugins/FX/Niagara/Source/NiagaraEditor/Private/SNiagaraTemplateAssetPicker.cpp
--- a/Engine/Plugins/FX/Niagara/Source/NiagaraEditor/Private/SNiagaraTemplateAssetPicker.cpp
+++ b/Engine/Plugins/FX/Niagara/Source/NiagaraEditor/Private/SNiagaraTemplateAssetPicker.cpp
@@ -25,6 +25,8 @@ void SNiagaraTemplateAssetPicker::Construct(const FArguments& InArgs, UClass* As
 	ProjectCategory = LOCTEXT("ProjectCategory", "Project");
 
 	AssetThumbnailPool = MakeShareable(new FAssetThumbnailPool(24));
+
+	OnTemplateAssetActivated = InArgs._OnTemplateAssetActivated;
 		
 	TArray<FAssetData> EmittersToShow;
 	for (const FAssetData& EmitterAsset : EmitterAssets)
@@ -48,6 +50,7 @@ void SNiagaraTemplateAssetPicker::Construct(const FArguments& InArgs, UClass* As
 		.OnDoesItemMatchFilterText(this, &SNiagaraTemplateAssetPicker::OnDoesItemMatchFilterText)
 		.OnGenerateWidgetForCategory(this, &SNiagaraTemplateAssetPicker::OnGenerateWidgetForCategory)
 		.OnGenerateWidgetForItem(this, &SNiagaraTemplateAssetPicker::OnGenerateWidgetForItem)
+		.OnItemActivated(this, &SNiagaraTemplateAssetPicker::OnItemActivated)
 	];
 }
 
@@ -139,13 +142,13 @@ TSharedRef<SWidget> SNiagaraTemplateAssetPicker::OnGenerateWidgetForCategory(con
 		];
 }
 
-const int32 ThumbnailSize = 64;
+const int32 ThumbnailSize = 72;
 
 TSharedRef<SWidget> SNiagaraTemplateAssetPicker::OnGenerateWidgetForItem(const FAssetData& Item)
 {
 	TSharedRef<FAssetThumbnail> AssetThumbnail = MakeShared<FAssetThumbnail>(Item, ThumbnailSize, ThumbnailSize, AssetThumbnailPool);
 	FAssetThumbnailConfig ThumbnailConfig;
-	ThumbnailConfig.bAllowFadeIn = true;
+	ThumbnailConfig.bAllowFadeIn = false;
 
 	FText AssetDescription;
 	Item.GetTagValue("TemplateAssetDescription", AssetDescription);
@@ -184,4 +187,9 @@ TSharedRef<SWidget> SNiagaraTemplateAssetPicker::OnGenerateWidgetForItem(const F
 		];
 }
 
+void SNiagaraTemplateAssetPicker::OnItemActivated(const FAssetData& Item)
+{
+	OnTemplateAssetActivated.ExecuteIfBound(Item);
+}
+
 #undef LOCTEXT_NAMESPACE
diff --git a/Engine/Plugins/FX/Niagara/Source/NiagaraEditor/Private/SNiagaraTemplateAssetPicker.h b/Engine/Plugins/FX/Niagara/Source/NiagaraEditor/Private/SNiagaraTemplateAssetPicker.h
--- a/Engine/Plugins/FX/Niagara/Source/NiagaraEditor/Private/SNiagaraTemplateAssetPicker.h
+++ b/Engine/Plugins/FX/Niagara/Source/NiagaraEditor/Private/SNiagaraTemplateAssetPicker.h
@@ -11,8 +11,12 @@ class FAssetThumbnailPool;
 
 class SNiagaraTemplateAssetPicker : public SCompoundWidget
 {
+public:
+	DECLARE_DELEGATE_OneParam(FOnTemplateAssetActivated, const FAssetData&);
+
 public:
 	SLATE_BEGIN_ARGS(SNiagaraTemplateAssetPicker) {}
+		SLATE_EVENT(FOnTemplateAssetActivated, OnTemplateAssetActivated);
 	SLATE_END_ARGS();
 
 	void Construct(const FArguments& InArgs, UClass* AssetClass);
@@ -36,9 +40,12 @@ private:
 
 	TSharedRef<SWidget> OnGenerateWidgetForItem(const FAssetData& Item);
 
+	void OnItemActivated(const FAssetData& Item);
+
 private:
 	TSharedPtr<SNiagaraAssetItemSelector> ItemSelector;
 	FText NiagaraPluginCategory;
 	FText ProjectCategory;
 	TSharedPtr<FAssetThumbnailPool> AssetThumbnailPool;
+	FOnTemplateAssetActivated OnTemplateAssetActivated;
 };
\ No newline at end of file
