Niagara - Fix a crash when saving an emitter in a debug build.

diff --git a/Engine/Plugins/FX/Niagara/Source/NiagaraEditor/Private/Toolkits/NiagaraSystemToolkit.cpp b/Engine/Plugins/FX/Niagara/Source/NiagaraEditor/Private/Toolkits/NiagaraSystemToolkit.cpp
--- a/Engine/Plugins/FX/Niagara/Source/NiagaraEditor/Private/Toolkits/NiagaraSystemToolkit.cpp
+++ b/Engine/Plugins/FX/Niagara/Source/NiagaraEditor/Private/Toolkits/NiagaraSystemToolkit.cpp
@@ -999,12 +999,28 @@ void FNiagaraSystemToolkit::UpdateOriginalEmitter()
 		Emitter = (UNiagaraEmitter*)StaticDuplicateObject(EditableEmitter, Emitter->GetOuter(),
 			Emitter->GetFName(), RF_AllFlags, Emitter->GetClass());
 
+		TArray<UNiagaraScript*> EmitterScripts;
+		Emitter->GetScripts(EmitterScripts, false);
+
+		TArray<UNiagaraScript*> EditableEmitterScripts;
+		EditableEmitter->GetScripts(EditableEmitterScripts, false);
+
+		// Validate that the change ids on the original emitters match the editable emitters ids to ensure the DDC contents are up to data without having to recompile.
+		if (ensureMsgf(EmitterScripts.Num() == EditableEmitterScripts.Num(), TEXT("Script count mismatch after copying from editable emitter to original emitter.")))
+		{
+			for (UNiagaraScript* EmitterScript : EmitterScripts)
+			{
+				UNiagaraScript** MatchingEditableEmitterScriptPtr = EditableEmitterScripts.FindByPredicate([EmitterScript](UNiagaraScript* EditableEmitterScript) { 
+					return EditableEmitterScript->GetUsage() == EmitterScript->GetUsage() && EditableEmitterScript->GetUsageId() == EmitterScript->GetUsageId(); });
+				ensureMsgf(MatchingEditableEmitterScriptPtr != nullptr, TEXT("Matching script could not be found in editable emitter after copying to original emitter."));
+				ensureMsgf((*MatchingEditableEmitterScriptPtr)->GetBaseChangeID() == EmitterScript->GetBaseChangeID(), TEXT("Script change ids didn't match after copying from editable emitter to original emitter."));
+			}
+		}
+
 		// Record the last synced change id to detect future changes.
 		LastSyncedEmitterChangeId = EditableEmitter->GetChangeId();
 		bEmitterThumbnailUpdated = false;
 
-		checkSlow(UNiagaraEmitter::GetForceCompileOnLoad() || Emitter->GetChangeId() == EditableEmitter->GetChangeId());
-
 		// Restore RF_Standalone on the original emitter, as it had been removed from the preview emitter so that it could be GC'd.
 		Emitter->SetFlags(RF_Standalone);
 
@@ -1012,8 +1028,6 @@ void FNiagaraSystemToolkit::UpdateOriginalEmitter()
 		AffectedEmitters.Add(Emitter);
 		UpdateExistingEmitters();
 
-		checkSlow(UNiagaraEmitter::GetForceCompileOnLoad() || Emitter->GetChangeId() == EditableEmitter->GetChangeId());
-
 		GWarn->EndSlowTask();
 	}
 	else if(bEmitterThumbnailUpdated)
