Niagara - Fix a crash when trying to sample a psuedo volume texture on CPU.

diff --git a/Engine/Plugins/FX/Niagara/Source/Niagara/Classes/NiagaraDataInterfaceTexture.h b/Engine/Plugins/FX/Niagara/Source/Niagara/Classes/NiagaraDataInterfaceTexture.h
--- a/Engine/Plugins/FX/Niagara/Source/Niagara/Classes/NiagaraDataInterfaceTexture.h
+++ b/Engine/Plugins/FX/Niagara/Source/Niagara/Classes/NiagaraDataInterfaceTexture.h
@@ -35,7 +35,7 @@ public:
 
 	void SampleTexture(FVectorVMContext& Context);
 	void GetTextureDimensions(FVectorVMContext& Context);
-	void NoOp(FVectorVMContext& Context);
+	void SamplePseudoVolumeTexture(FVectorVMContext& Context);
 
 	virtual bool Equals(const UNiagaraDataInterface* Other) const override;
 
diff --git a/Engine/Plugins/FX/Niagara/Source/Niagara/Private/NiagaraDataInterfaceTexture.cpp b/Engine/Plugins/FX/Niagara/Source/Niagara/Private/NiagaraDataInterfaceTexture.cpp
--- a/Engine/Plugins/FX/Niagara/Source/Niagara/Private/NiagaraDataInterfaceTexture.cpp
+++ b/Engine/Plugins/FX/Niagara/Source/Niagara/Private/NiagaraDataInterfaceTexture.cpp
@@ -165,7 +165,9 @@ void UNiagaraDataInterfaceTexture::GetFunctions(TArray<FNiagaraFunctionSignature
 		OutFunctions.Add(Sig);
 	}
 }
+
 DEFINE_NDI_DIRECT_FUNC_BINDER(UNiagaraDataInterfaceTexture, SampleTexture);
+DEFINE_NDI_DIRECT_FUNC_BINDER(UNiagaraDataInterfaceTexture, SamplePseudoVolumeTexture)
 void UNiagaraDataInterfaceTexture::GetVMExternalFunction(const FVMExternalFunctionBindingInfo& BindingInfo, void* InstanceData, FVMExternalFunction &OutFunc)
 {
 	if (BindingInfo.Name == SampleTexture2DName)
@@ -175,9 +177,8 @@ void UNiagaraDataInterfaceTexture::GetVMExternalFunction(const FVMExternalFuncti
 	}
 	else if (BindingInfo.Name == SamplePseudoVolumeTextureName)
 	{
-		//check(BindingInfo.GetNumInputs() == 2 && BindingInfo.GetNumOutputs() == 4);
-		OutFunc = FVMExternalFunction::CreateUObject(this, &UNiagaraDataInterfaceTexture::NoOp);
-
+		check(BindingInfo.GetNumInputs() == 12 && BindingInfo.GetNumOutputs() == 4);
+		NDI_FUNC_BINDER(UNiagaraDataInterfaceTexture, SamplePseudoVolumeTexture)::Bind(this, OutFunc);
 	}
 	else if (BindingInfo.Name == TextureDimsName)
 	{
@@ -186,12 +187,6 @@ void UNiagaraDataInterfaceTexture::GetVMExternalFunction(const FVMExternalFuncti
 	}
 }
 
-
-void UNiagaraDataInterfaceTexture::NoOp(FVectorVMContext& Context)
-{
-
-}
-
 void UNiagaraDataInterfaceTexture::GetTextureDimensions(FVectorVMContext& Context)
 {
 	VectorVM::FExternalFuncRegisterHandler<float> OutWidth(Context);
@@ -238,6 +233,61 @@ void UNiagaraDataInterfaceTexture::SampleTexture(FVectorVMContext& Context)
 
 }
 
+void UNiagaraDataInterfaceTexture::SamplePseudoVolumeTexture(FVectorVMContext& Context)
+{
+	// Noop handler which just returns magenta since this doesn't run on CPU.
+	VectorVM::FExternalFuncInputHandler<float> UVW_UParam(Context);
+	VectorVM::FExternalFuncInputHandler<float> UVW_VParam(Context);
+	VectorVM::FExternalFuncInputHandler<float> UVW_WParam(Context);
+
+	VectorVM::FExternalFuncInputHandler<float> XYNumFrames_XParam(Context);
+	VectorVM::FExternalFuncInputHandler<float> XYNumFrames_YParam(Context);
+	
+	VectorVM::FExternalFuncInputHandler<float> TotalNumFramesParam(Context);
+
+	VectorVM::FExternalFuncInputHandler<int32> MipModeParam(Context);
+
+	VectorVM::FExternalFuncInputHandler<float> MipLevelParam(Context);
+
+	VectorVM::FExternalFuncInputHandler<float> DDX_XParam(Context);
+	VectorVM::FExternalFuncInputHandler<float> DDX_YParam(Context);
+
+	VectorVM::FExternalFuncInputHandler<float> DDY_XParam(Context);
+	VectorVM::FExternalFuncInputHandler<float> DDY_YParam(Context);
+
+	VectorVM::FExternalFuncRegisterHandler<float> OutSampleR(Context);
+	VectorVM::FExternalFuncRegisterHandler<float> OutSampleG(Context);
+	VectorVM::FExternalFuncRegisterHandler<float> OutSampleB(Context);
+	VectorVM::FExternalFuncRegisterHandler<float> OutSampleA(Context);
+
+	for (int32 i = 0; i < Context.NumInstances; ++i)
+	{
+		UVW_UParam.Advance();
+		UVW_VParam.Advance();
+		UVW_WParam.Advance();
+
+		XYNumFrames_XParam.Advance();
+		XYNumFrames_YParam.Advance();
+
+		TotalNumFramesParam.Advance();
+
+		MipModeParam.Advance();
+
+		MipLevelParam.Advance();
+
+		DDX_XParam.Advance();
+		DDX_YParam.Advance();
+
+		DDY_XParam.Advance();
+		DDY_YParam.Advance();
+
+		*OutSampleR.GetDestAndAdvance() = 1.0;
+		*OutSampleG.GetDestAndAdvance() = 0.0;
+		*OutSampleB.GetDestAndAdvance() = 1.0;
+		*OutSampleA.GetDestAndAdvance() = 1.0;
+	}
+}
+
 bool UNiagaraDataInterfaceTexture::GetFunctionHLSL(const FName& DefinitionFunctionName, FString InstanceFunctionName, FNiagaraDataInterfaceGPUParamInfo& ParamInfo, FString& OutHLSL)
 {
 	if (DefinitionFunctionName == SampleTexture2DName)
