Niagara - Fix static analysis warning about null ptr.

diff --git a/Engine/Plugins/FX/Niagara/Source/NiagaraEditor/Private/Toolkits/NiagaraSystemToolkit.cpp b/Engine/Plugins/FX/Niagara/Source/NiagaraEditor/Private/Toolkits/NiagaraSystemToolkit.cpp
--- a/Engine/Plugins/FX/Niagara/Source/NiagaraEditor/Private/Toolkits/NiagaraSystemToolkit.cpp
+++ b/Engine/Plugins/FX/Niagara/Source/NiagaraEditor/Private/Toolkits/NiagaraSystemToolkit.cpp
@@ -1012,8 +1012,10 @@ void FNiagaraSystemToolkit::UpdateOriginalEmitter()
 			{
 				UNiagaraScript** MatchingEditableEmitterScriptPtr = EditableEmitterScripts.FindByPredicate([EmitterScript](UNiagaraScript* EditableEmitterScript) { 
 					return EditableEmitterScript->GetUsage() == EmitterScript->GetUsage() && EditableEmitterScript->GetUsageId() == EmitterScript->GetUsageId(); });
-				ensureMsgf(MatchingEditableEmitterScriptPtr != nullptr, TEXT("Matching script could not be found in editable emitter after copying to original emitter."));
-				ensureMsgf((*MatchingEditableEmitterScriptPtr)->GetBaseChangeID() == EmitterScript->GetBaseChangeID(), TEXT("Script change ids didn't match after copying from editable emitter to original emitter."));
+				if (ensureMsgf(MatchingEditableEmitterScriptPtr != nullptr, TEXT("Matching script could not be found in editable emitter after copying to original emitter.")))
+				{
+					ensureMsgf((*MatchingEditableEmitterScriptPtr)->GetBaseChangeID() == EmitterScript->GetBaseChangeID(), TEXT("Script change ids didn't match after copying from editable emitter to original emitter."));
+				}
 			}
 		}
 
