Copying using //UE4/Dev-Niagara_to_//UE4/Release-4.21:
Fix for UE-65001

Incorrect temlpate arguments caused binding to wrong function instance.

diff --git a/Engine/Plugins/FX/Niagara/Source/Niagara/Private/NDISkeletalMesh_VertexSampling.cpp b/Engine/Plugins/FX/Niagara/Source/Niagara/Private/NDISkeletalMesh_VertexSampling.cpp
--- a/Engine/Plugins/FX/Niagara/Source/Niagara/Private/NDISkeletalMesh_VertexSampling.cpp
+++ b/Engine/Plugins/FX/Niagara/Source/Niagara/Private/NDISkeletalMesh_VertexSampling.cpp
@@ -17,7 +17,6 @@ DECLARE_CYCLE_STAT(TEXT("Skel Mesh Vertex Sampling"), STAT_NiagaraSkel_Vertex_Sa
 DEFINE_NDI_FUNC_BINDER(UNiagaraDataInterfaceSkeletalMesh, RandomVertex);
 DEFINE_NDI_FUNC_BINDER(UNiagaraDataInterfaceSkeletalMesh, GetVertexSkinnedData);
 DEFINE_NDI_DIRECT_FUNC_BINDER(UNiagaraDataInterfaceSkeletalMesh, GetVertexColor);
-DEFINE_NDI_DIRECT_FUNC_BINDER(UNiagaraDataInterfaceSkeletalMesh, GetVertexColorFallback);
 DEFINE_NDI_FUNC_BINDER(UNiagaraDataInterfaceSkeletalMesh, GetVertexUV);
 DEFINE_NDI_FUNC_BINDER(UNiagaraDataInterfaceSkeletalMesh, IsValidVertex);
 DEFINE_NDI_FUNC_BINDER(UNiagaraDataInterfaceSkeletalMesh, GetFilteredVertexCount);
@@ -83,7 +82,7 @@ void UNiagaraDataInterfaceSkeletalMesh::GetVertexSamplingFunctions(TArray<FNiaga
 		Sig.bMemberFunction = true;
 		Sig.bRequiresContext = false;
 #if WITH_EDITORONLY_DATA
-		Sig.Description = LOCTEXT("GetSkinnedDataWSDesc", "Returns skinning dependant data for the pased vertex in world space. All outputs are optional and you will incur zero to minimal cost if they are not connected.");
+		Sig.Description = LOCTEXT("GetSkinnedDataDesc", "Returns skinning dependant data for the pased vertex in world space. All outputs are optional and you will incur zero to minimal cost if they are not connected.");
 #endif
 		OutFunctions.Add(Sig);
 	}
@@ -136,6 +135,7 @@ void UNiagaraDataInterfaceSkeletalMesh::GetVertexSamplingFunctions(TArray<FNiaga
 
 void UNiagaraDataInterfaceSkeletalMesh::BindVertexSamplingFunction(const FVMExternalFunctionBindingInfo& BindingInfo, FNDISkeletalMesh_InstanceData* InstanceData, FVMExternalFunction &OutFunc)
 {
+	bool bNeedsVertexColors = false;
 
 	if (BindingInfo.Name == RandomVertexName)
 	{
@@ -160,14 +160,8 @@ void UNiagaraDataInterfaceSkeletalMesh::BindVertexSamplingFunction(const FVMExte
 	else if (BindingInfo.Name == GetVertexColorName)
 	{
 		check(BindingInfo.GetNumInputs() == 2 && BindingInfo.GetNumOutputs() == 4);
-		if (InstanceData->HasColorData())
-		{
-			NDI_FUNC_BINDER(UNiagaraDataInterfaceSkeletalMesh, GetVertexColor)::Bind(this, OutFunc);
-		}
-		else
-		{
-			NDI_FUNC_BINDER(UNiagaraDataInterfaceSkeletalMesh, GetVertexColorFallback)::Bind(this, OutFunc);
-		}
+		bNeedsVertexColors = true;
+		NDI_FUNC_BINDER(UNiagaraDataInterfaceSkeletalMesh, GetVertexColor)::Bind(this, OutFunc);
 	}
 	else if (BindingInfo.Name == GetVertexUVName)
 	{
@@ -185,6 +179,15 @@ void UNiagaraDataInterfaceSkeletalMesh::BindVertexSamplingFunction(const FVMExte
 		TFilterModeBinder<NDI_FUNC_BINDER(UNiagaraDataInterfaceSkeletalMesh, GetFilteredVertexAt)>::Bind(this, BindingInfo, InstanceData, OutFunc);
 	}
 
+	check(InstanceData->Mesh);
+	FSkinWeightVertexBuffer* SkinWeightBuffer;
+	FSkeletalMeshLODRenderData& LODData = InstanceData->GetLODRenderDataAndSkinWeights(SkinWeightBuffer);
+
+	if (bNeedsVertexColors && LODData.StaticVertexBuffers.ColorVertexBuffer.GetNumVertices() == 0)
+	{
+		UE_LOG(LogNiagara, Log, TEXT("Skeletal Mesh data interface is cannot run as it's reading color data on a mesh that does not provide it. - Mesh:%s  "), *InstanceData->Mesh->GetFullName());
+		OutFunc = FVMExternalFunction();
+	}
 }
 
 template<typename FilterMode>
@@ -232,7 +235,7 @@ void UNiagaraDataInterfaceSkeletalMesh::RandomVertex(FVectorVMContext& Context)
 	VectorVM::FExternalFuncRegisterHandler<int32> OutVert(Context);
 
 	FSkeletalMeshAccessorHelper MeshAccessor;
-	MeshAccessor.Init<FilterMode, TIntegralConstant<int32,0>>(InstData);
+	MeshAccessor.Init<FilterMode, TIntegralConstant<ENDISkelMesh_AreaWeightingMode, ENDISkelMesh_AreaWeightingMode::None>>(InstData);
 
 	for (int32 i = 0; i < Context.NumInstances; ++i)
 	{
@@ -432,25 +435,6 @@ void UNiagaraDataInterfaceSkeletalMesh::GetVertexColor(FVectorVMContext& Context
 	}
 }
 
-void UNiagaraDataInterfaceSkeletalMesh::GetVertexColorFallback(FVectorVMContext& Context)
-{
-	VectorVM::FExternalFuncInputHandler<int32> VertParam(Context);
-	VectorVM::FUserPtrHandler<FNDISkeletalMesh_InstanceData> InstData(Context);
-
-	VectorVM::FExternalFuncRegisterHandler<float> OutColorR(Context);
-	VectorVM::FExternalFuncRegisterHandler<float> OutColorG(Context);
-	VectorVM::FExternalFuncRegisterHandler<float> OutColorB(Context);
-	VectorVM::FExternalFuncRegisterHandler<float> OutColorA(Context);
-
-	for (int32 i = 0; i < Context.NumInstances; ++i)
-	{
-		*OutColorR.GetDestAndAdvance() = 1.0f;
-		*OutColorG.GetDestAndAdvance() = 1.0f;
-		*OutColorB.GetDestAndAdvance() = 1.0f;
-		*OutColorA.GetDestAndAdvance() = 1.0f;
-	}
-}
-
 template<typename VertexAccessorType>
 void UNiagaraDataInterfaceSkeletalMesh::GetVertexUV(FVectorVMContext& Context)
 {
