Copying using //UE4/Dev-Niagara_to_//UE4/Release-4.21 for Simon.Tovey:
Made some Triangle sampling functions more robust to bad input.

diff --git a/Engine/Plugins/FX/Niagara/Source/Niagara/Private/NDISkeletalMesh_TriangleSampling.cpp b/Engine/Plugins/FX/Niagara/Source/Niagara/Private/NDISkeletalMesh_TriangleSampling.cpp
--- a/Engine/Plugins/FX/Niagara/Source/Niagara/Private/NDISkeletalMesh_TriangleSampling.cpp
+++ b/Engine/Plugins/FX/Niagara/Source/Niagara/Private/NDISkeletalMesh_TriangleSampling.cpp
@@ -17,7 +17,6 @@ DECLARE_CYCLE_STAT(TEXT("Skel Mesh Sampling"), STAT_NiagaraSkel_Sample, STATGROU
 DEFINE_NDI_FUNC_BINDER(UNiagaraDataInterfaceSkeletalMesh, RandomTriCoord);
 DEFINE_NDI_FUNC_BINDER(UNiagaraDataInterfaceSkeletalMesh, GetTriCoordSkinnedData);
 DEFINE_NDI_DIRECT_FUNC_BINDER(UNiagaraDataInterfaceSkeletalMesh, GetTriCoordColor);
-DEFINE_NDI_DIRECT_FUNC_BINDER(UNiagaraDataInterfaceSkeletalMesh, GetTriCoordColorFallback);
 DEFINE_NDI_FUNC_BINDER(UNiagaraDataInterfaceSkeletalMesh, GetTriCoordUV);
 DEFINE_NDI_FUNC_BINDER(UNiagaraDataInterfaceSkeletalMesh, IsValidTriCoord);
 DEFINE_NDI_FUNC_BINDER(UNiagaraDataInterfaceSkeletalMesh, GetFilteredTriangleCount);
@@ -92,7 +91,7 @@ void UNiagaraDataInterfaceSkeletalMesh::GetTriangleSamplingFunctions(TArray<FNia
 		Sig.bMemberFunction = true;
 		Sig.bRequiresContext = false;
 #if WITH_EDITORONLY_DATA
-		Sig.Description = LOCTEXT("GetSkinnedDataWSDesc", "Returns skinning dependant data for the pased MeshTriCoord in world space. All outputs are optional and you will incur zerp minimal cost if they are not connected.");
+		Sig.Description = LOCTEXT("GetSkinnedDataDesc", "Returns skinning dependant data for the pased MeshTriCoord in world space. All outputs are optional and you will incur zerp minimal cost if they are not connected.");
 #endif
 		OutFunctions.Add(Sig);
 	}
@@ -161,6 +160,7 @@ void UNiagaraDataInterfaceSkeletalMesh::GetTriangleSamplingFunctions(TArray<FNia
 
 void UNiagaraDataInterfaceSkeletalMesh::BindTriangleSamplingFunction(const FVMExternalFunctionBindingInfo& BindingInfo, FNDISkeletalMesh_InstanceData* InstanceData, FVMExternalFunction &OutFunc)
 {
+	bool bNeedsVertexColors = false;
 
 	if (BindingInfo.Name == RandomTriCoordName)
 	{
@@ -185,14 +185,8 @@ void UNiagaraDataInterfaceSkeletalMesh::BindTriangleSamplingFunction(const FVMEx
 	else if (BindingInfo.Name == GetTriColorName)
 	{
 		check(BindingInfo.GetNumInputs() == 5 && BindingInfo.GetNumOutputs() == 4);
-		if (InstanceData->HasColorData())
-		{
-			NDI_FUNC_BINDER(UNiagaraDataInterfaceSkeletalMesh, GetTriCoordColor)::Bind(this, OutFunc);
-		}
-		else
-		{
-			NDI_FUNC_BINDER(UNiagaraDataInterfaceSkeletalMesh, GetTriCoordColorFallback)::Bind(this, OutFunc);
-		}
+		bNeedsVertexColors = true;
+		NDI_FUNC_BINDER(UNiagaraDataInterfaceSkeletalMesh, GetTriCoordColor)::Bind(this, OutFunc);
 	}
 	else if (BindingInfo.Name == GetTriUVName)
 	{
@@ -215,6 +209,15 @@ void UNiagaraDataInterfaceSkeletalMesh::BindTriangleSamplingFunction(const FVMEx
 		TSkinningModeBinder<NDI_FUNC_BINDER(UNiagaraDataInterfaceSkeletalMesh, GetTriCoordVertices)>::Bind(this, BindingInfo, InstanceData, OutFunc);
 	}
 
+	check(InstanceData->Mesh);
+	FSkinWeightVertexBuffer* SkinWeightBuffer;
+	FSkeletalMeshLODRenderData& LODData = InstanceData->GetLODRenderDataAndSkinWeights(SkinWeightBuffer);
+
+	if (bNeedsVertexColors && LODData.StaticVertexBuffers.ColorVertexBuffer.GetNumVertices() == 0)
+	{
+		UE_LOG(LogNiagara, Log, TEXT("Skeletal Mesh data interface is cannot run as it's reading color data on a mesh that does not provide it. - Mesh:%s  "), *InstanceData->Mesh->GetFullName());
+		OutFunc = FVMExternalFunction();
+	}
 }
 
 template<typename FilterMode, typename AreaWeightingMode>
@@ -587,7 +590,7 @@ void UNiagaraDataInterfaceSkeletalMesh::GetFilteredTriangleAt(FVectorVMContext&
 		RealIdx = GetSpecificTriangleAt<FilterMode, AreaWeightingMode>(Accessor, InstData, Tri);
 
 		int32 TriMax = Accessor.IndexBuffer->Num() - 3;
-		RealIdx = FMath::Min(RealIdx, TriMax);
+		RealIdx = FMath::Clamp(RealIdx, 0, TriMax);
 
 		*OutTri.GetDest() = RealIdx;
 		float Coord = 1.0f / 3.0f;
@@ -649,29 +652,6 @@ void UNiagaraDataInterfaceSkeletalMesh::GetTriCoordColor(FVectorVMContext& Conte
 	}
 }
 
-// Where we determine we are sampling a skeletal mesh without tri color we bind to this fallback method 
-void UNiagaraDataInterfaceSkeletalMesh::GetTriCoordColorFallback(FVectorVMContext& Context)
-{
-	VectorVM::FExternalFuncInputHandler<int32> TriParam(Context);
-	VectorVM::FExternalFuncInputHandler<float> BaryXParam(Context);
-	VectorVM::FExternalFuncInputHandler<float> BaryYParam(Context);
-	VectorVM::FExternalFuncInputHandler<float> BaryZParam(Context);
-	VectorVM::FUserPtrHandler<FNDISkeletalMesh_InstanceData> InstData(Context);
-
-	VectorVM::FExternalFuncRegisterHandler<float> OutColorR(Context);
-	VectorVM::FExternalFuncRegisterHandler<float> OutColorG(Context);
-	VectorVM::FExternalFuncRegisterHandler<float> OutColorB(Context);
-	VectorVM::FExternalFuncRegisterHandler<float> OutColorA(Context);
-
-	for (int32 i = 0; i < Context.NumInstances; ++i)
-	{
-		*OutColorR.GetDestAndAdvance() = 1.0f;
-		*OutColorG.GetDestAndAdvance() = 1.0f;
-		*OutColorB.GetDestAndAdvance() = 1.0f;
-		*OutColorA.GetDestAndAdvance() = 1.0f;
-	}
-}
-
 template<typename VertexAccessorType>
 void UNiagaraDataInterfaceSkeletalMesh::GetTriCoordUV(FVectorVMContext& Context)
 {
@@ -701,7 +681,7 @@ void UNiagaraDataInterfaceSkeletalMesh::GetTriCoordUV(FVectorVMContext& Context)
 	for (int32 i = 0; i < Context.NumInstances; ++i)
 	{
 		int32 Tri = TriParam.Get();
-		Tri = FMath::Min(Tri, TriMax);
+		Tri = FMath::Clamp(Tri, 0, TriMax);
 
 		int32 Idx0 = IndexBuffer->Get(Tri);
 		int32 Idx1 = IndexBuffer->Get(Tri + 1);
@@ -819,7 +799,7 @@ void UNiagaraDataInterfaceSkeletalMesh::GetTriCoordSkinnedData(FVectorVMContext&
 	FSkeletalMeshLODRenderData& LODData = InstData->GetLODRenderDataAndSkinWeights(SkinWeightBuffer);
 
 	FSkeletalMeshAccessorHelper Accessor;
-	Accessor.Init<TIntegralConstant<int32, 0>, TIntegralConstant<int32, 0>>(InstData);
+	Accessor.Init<TIntegralConstant<ENDISkeletalMesh_FilterMode, ENDISkeletalMesh_FilterMode::None>, TIntegralConstant<ENDISkelMesh_AreaWeightingMode, ENDISkelMesh_AreaWeightingMode::None>>(InstData);
 	int32 TriMax = Accessor.IndexBuffer->Num() - 3;
 	float InvDt = 1.0f / InstData->DeltaSeconds;
 
@@ -837,6 +817,11 @@ void UNiagaraDataInterfaceSkeletalMesh::GetTriCoordSkinnedData(FVectorVMContext&
 	{
 		FMeshTriCoordinate MeshTriCoord(TriParam.GetAndAdvance(), FVector(BaryXParam.GetAndAdvance(), BaryYParam.GetAndAdvance(), BaryZParam.GetAndAdvance()));
 
+		if(MeshTriCoord.Tri < 0 || MeshTriCoord.Tri > TriMax)
+		{
+			MeshTriCoord = FMeshTriCoordinate(0, FVector(1.0f, 0.0f, 0.0f));
+		}
+
 		SkinningHandler.GetTrianlgeIndices(Accessor, MeshTriCoord.Tri, Idx0, Idx1, Idx2);
 
 		if (Output.bNeedsPosition || Output.bNeedsVelocity)
@@ -926,8 +911,11 @@ void UNiagaraDataInterfaceSkeletalMesh::GetTriCoordVertices(FVectorVMContext& Co
 	FSkeletalMeshAccessorHelper Accessor;
 	Accessor.Init<TIntegralConstant<int32, 0>, TIntegralConstant<int32, 0>>(InstData);
 
+	int32 TriMax = Accessor.IndexBuffer->Num() - 3;
+
 	for (int32 i = 0; i < Context.NumInstances; ++i)
 	{
+		int32 Tri = FMath::Clamp(TriParam.GetAndAdvance(), 0, TriMax);
 		SkinningHandler.GetTrianlgeIndices(Accessor, TriParam.GetAndAdvance(), Idx0, Idx1, Idx2);
 		*OutV0.GetDestAndAdvance() = Idx0;
 		*OutV1.GetDestAndAdvance() = Idx1;
