Fixed Niagara GPU emitter not being rendering when running non cooked game.

diff --git a/Engine/Plugins/FX/Niagara/Source/Niagara/Private/NiagaraScript.cpp b/Engine/Plugins/FX/Niagara/Source/Niagara/Private/NiagaraScript.cpp
--- a/Engine/Plugins/FX/Niagara/Source/Niagara/Private/NiagaraScript.cpp
+++ b/Engine/Plugins/FX/Niagara/Source/Niagara/Private/NiagaraScript.cpp
@@ -994,7 +994,8 @@ void UNiagaraScript::CacheShadersForResources(EShaderPlatform ShaderPlatform, FN
 {
 	if (CanBeRunOnGpu())
 	{
-		const bool bSuccess = ResourceToCache->CacheShaders(ShaderPlatform, bApplyCompletedShaderMapForRendering, bForceRecompile, bCooking);
+		// When not running in the editor, the shaders are created in-sync (in the postload) to avoid update issues.
+		const bool bSuccess = ResourceToCache->CacheShaders(ShaderPlatform, bApplyCompletedShaderMapForRendering, bForceRecompile, bCooking || !GIsEditor);
 
 #if defined(NIAGARA_SCRIPT_COMPILE_LOGGING_MEDIUM)
 		if (!bSuccess)
diff --git a/Engine/Plugins/FX/Niagara/Source/NiagaraEditor/Private/NiagaraHlslTranslator.cpp b/Engine/Plugins/FX/Niagara/Source/NiagaraEditor/Private/NiagaraHlslTranslator.cpp
--- a/Engine/Plugins/FX/Niagara/Source/NiagaraEditor/Private/NiagaraHlslTranslator.cpp
+++ b/Engine/Plugins/FX/Niagara/Source/NiagaraEditor/Private/NiagaraHlslTranslator.cpp
@@ -128,7 +128,8 @@ void FNiagaraShaderQueueTickable::ProcessQueue()
 			TRefCountPtr<FShaderCompilerEnvironment> CompilerEnvironment = new FShaderCompilerEnvironment();
 
 			FString ShaderCode = CompilableScript->GetVMExecutableData().LastHlslTranslationGPU;
-			const bool bSynchronousCompile = false;
+			// When not running in the editor, the shaders are created in-sync in the postload.
+			const bool bSynchronousCompile = !GIsEditor;
 
 			// Compile the shaders for the script.
 			NewShaderMap->Compile(ShaderScript, Item.ShaderMapId, CompilerEnvironment, NewCompilationOutput, Item.Platform, bSynchronousCompile, Item.bApply);
diff --git a/Engine/Plugins/FX/Niagara/Source/NiagaraShader/Private/NiagaraShader.cpp b/Engine/Plugins/FX/Niagara/Source/NiagaraShader/Private/NiagaraShader.cpp
--- a/Engine/Plugins/FX/Niagara/Source/NiagaraShader/Private/NiagaraShader.cpp
+++ b/Engine/Plugins/FX/Niagara/Source/NiagaraShader/Private/NiagaraShader.cpp
@@ -814,7 +814,7 @@ void FNiagaraShaderMap::Compile(
 			{
 				TArray<int32> CurrentShaderMapId;
 				CurrentShaderMapId.Add(CompilingId);
-				//GNiagaraShaderCompilationManager->FinishCompilation(*FriendlyName, CurrentShaderMapId);
+				GNiagaraShaderCompilationManager.FinishCompilation(*FriendlyName, CurrentShaderMapId);
 			}
 		}
 	}
diff --git a/Engine/Plugins/FX/Niagara/Source/NiagaraShader/Private/NiagaraShaderCompilationManager.cpp b/Engine/Plugins/FX/Niagara/Source/NiagaraShader/Private/NiagaraShaderCompilationManager.cpp
--- a/Engine/Plugins/FX/Niagara/Source/NiagaraShader/Private/NiagaraShaderCompilationManager.cpp
+++ b/Engine/Plugins/FX/Niagara/Source/NiagaraShader/Private/NiagaraShaderCompilationManager.cpp
@@ -9,7 +9,7 @@
 #include "ShaderCompiler.h"
 #include "HAL/FileManager.h"
 #include "Misc/Paths.h"
-
+#include "UObject/UObjectThreadContext.h"
 
 DEFINE_LOG_CATEGORY_STATIC(LogNiagaraShaderCompiler, All, All);
 
@@ -357,7 +357,9 @@ void FNiagaraShaderCompilationManager::ProcessCompiledNiagaraShaderMaps(
 					}
 					else
 					{
-						if (CompletedShaderMap->IsComplete(Script, true))
+						// Can't call NotifyCompilationFinished() when post-loading. 
+						// This normally happens when compiled in-sync for which the callback is not required.
+						if (CompletedShaderMap->IsComplete(Script, true) && !FUObjectThreadContext::Get().IsRoutingPostLoad)
 						{
 							Script->NotifyCompilationFinished();
 						}
diff --git a/Engine/Plugins/FX/Niagara/Source/NiagaraShader/Private/NiagaraShared.cpp b/Engine/Plugins/FX/Niagara/Source/NiagaraShader/Private/NiagaraShared.cpp
--- a/Engine/Plugins/FX/Niagara/Source/NiagaraShader/Private/NiagaraShared.cpp
+++ b/Engine/Plugins/FX/Niagara/Source/NiagaraShader/Private/NiagaraShared.cpp
@@ -441,9 +441,11 @@ bool FNiagaraShaderScript::BeginCompileShaderMap(
 		INiagaraShaderModule NiagaraShaderModule = FModuleManager::GetModuleChecked<INiagaraShaderModule>(TEXT("NiagaraShader"));
 		NiagaraShaderModule.ProcessShaderCompilationQueue();
 	}
-
-	// For async compile, set to nullptr so that we fall back to CPU side simulation until shader compile is finished
-	OutShaderMap = nullptr;
+	else
+	{
+		// For async compile, set to nullptr so that we fall back to CPU side simulation until shader compile is finished
+		OutShaderMap = nullptr;
+	}
 
 	INC_FLOAT_STAT_BY(STAT_ShaderCompiling_NiagaraShaders, (float)NiagaraCompileTime);
 
