Fixed check failing when calling niagara script NotifyCompilationFinished in non async compilation.

diff --git a/Engine/Plugins/FX/Niagara/Source/NiagaraShader/Private/NiagaraShaderCompilationManager.cpp b/Engine/Plugins/FX/Niagara/Source/NiagaraShader/Private/NiagaraShaderCompilationManager.cpp
--- a/Engine/Plugins/FX/Niagara/Source/NiagaraShader/Private/NiagaraShaderCompilationManager.cpp
+++ b/Engine/Plugins/FX/Niagara/Source/NiagaraShader/Private/NiagaraShaderCompilationManager.cpp
@@ -400,8 +400,12 @@ void FNiagaraShaderCompilationManager::ProcessCompiledNiagaraShaderMaps(
 					Script->SetRenderingThreadShaderMap(CompiledShaderMap);
 				});
 
-
-			Script->NotifyCompilationFinished();
+			// Can't call NotifyCompilationFinished() when post-loading. 
+			// This normally happens when compiled in-sync for which the callback is not required.
+			if (!FUObjectThreadContext::Get().IsRoutingPostLoad)
+			{
+				Script->NotifyCompilationFinished();
+			}
 		}
 	}
 }
