Niagara - Fix random crashes with GPU simulation due to threading and simulation context lifetime issues.

diff --git a/Engine/Plugins/FX/Niagara/Source/Niagara/Classes/NiagaraEmitterInstance.h b/Engine/Plugins/FX/Niagara/Source/Niagara/Classes/NiagaraEmitterInstance.h
--- a/Engine/Plugins/FX/Niagara/Source/Niagara/Classes/NiagaraEmitterInstance.h
+++ b/Engine/Plugins/FX/Niagara/Source/Niagara/Classes/NiagaraEmitterInstance.h
@@ -132,7 +132,7 @@ private:
 
 	FNiagaraScriptExecutionContext SpawnExecContext;
 	FNiagaraScriptExecutionContext UpdateExecContext;
-	FNiagaraComputeExecutionContext GPUExecContext;
+	FNiagaraComputeExecutionContext* GPUExecContext;
 	TArray<FNiagaraScriptExecutionContext> EventExecContexts;
 
 	FNiagaraParameterDirectBinding<float> SpawnIntervalBinding;
diff --git a/Engine/Plugins/FX/Niagara/Source/Niagara/Classes/NiagaraScriptExecutionContext.h b/Engine/Plugins/FX/Niagara/Source/Niagara/Classes/NiagaraScriptExecutionContext.h
--- a/Engine/Plugins/FX/Niagara/Source/Niagara/Classes/NiagaraScriptExecutionContext.h
+++ b/Engine/Plugins/FX/Niagara/Source/Niagara/Classes/NiagaraScriptExecutionContext.h
@@ -102,6 +102,7 @@ struct FNiagaraComputeExecutionContext
 
 	~FNiagaraComputeExecutionContext()
 	{
+		checkf(IsInRenderingThread(), TEXT("Can only delete the gpu readback from the render thread"));
 		if (GPUDataReadback)
 		{
 			delete GPUDataReadback;
diff --git a/Engine/Plugins/FX/Niagara/Source/Niagara/Private/NiagaraEmitterInstance.cpp b/Engine/Plugins/FX/Niagara/Source/Niagara/Private/NiagaraEmitterInstance.cpp
--- a/Engine/Plugins/FX/Niagara/Source/Niagara/Private/NiagaraEmitterInstance.cpp
+++ b/Engine/Plugins/FX/Niagara/Source/Niagara/Private/NiagaraEmitterInstance.cpp
@@ -53,6 +53,7 @@ FNiagaraEmitterInstance::FNiagaraEmitterInstance(FNiagaraSystemInstance* InParen
 : CPUTimeMS(0.0f)
 , ExecutionState(ENiagaraExecutionState::Inactive)
 , CachedBounds(ForceInit)
+, GPUExecContext(nullptr)
 , ParentSystemInstance(InParentSystemInstance)
 , CachedEmitter(nullptr)
 #if !UE_BUILD_SHIPPING
@@ -72,15 +73,31 @@ FNiagaraEmitterInstance::~FNiagaraEmitterInstance()
 
 	if (CachedEmitter != nullptr && CachedEmitter->SimTarget == ENiagaraSimTarget::GPUComputeSim)
 	{
-		NiagaraEmitterInstanceBatcher::Get()->Remove(&GPUExecContext);
-	}
+		/** We defer the deletion of the particle dataset and the compute context to the RT to be sure all in-flight RT commands have finished using it.*/
+		if (GPUExecContext != nullptr)
+		{
+			NiagaraEmitterInstanceBatcher::Get()->Remove(GPUExecContext);
+
+			ENQUEUE_UNIQUE_RENDER_COMMAND_ONEPARAMETER(FDeleteComputeContextCommand,
+				FNiagaraComputeExecutionContext*, ExecContext, GPUExecContext,
+				{
+					delete ExecContext;
+				});
+			
+			GPUExecContext = nullptr;
+		}
 
-	/** We defer the deletion of the particle dataset to the RT to be sure all in-flight RT commands have finished using it.*/
-	ENQUEUE_UNIQUE_RENDER_COMMAND_ONEPARAMETER(FDeleteParticleDataSetCommand,
-		FNiagaraDataSet*, DataSet, ParticleDataSet,
+		if (ParticleDataSet != nullptr)
 		{
-			delete DataSet;
-		});
+			ENQUEUE_UNIQUE_RENDER_COMMAND_ONEPARAMETER(FDeleteParticleDataSetCommand,
+				FNiagaraDataSet*, DataSet, ParticleDataSet,
+				{
+					delete DataSet;
+				});
+
+			ParticleDataSet = nullptr;
+		}
+	}
 }
 
 void FNiagaraEmitterInstance::ClearRenderer()
@@ -119,8 +136,11 @@ void FNiagaraEmitterInstance::Dump()const
 	SpawnExecContext.Parameters.DumpParameters(true);
 	UE_LOG(LogNiagara, Log, TEXT(".................Update................."));
 	UpdateExecContext.Parameters.DumpParameters(true);
-	UE_LOG(LogNiagara, Log, TEXT("................. %s Combined Parameters ................."), TEXT("GPU Script"));
-	GPUExecContext.CombinedParamStore.DumpParameters();
+	if (CachedEmitter->SimTarget == ENiagaraSimTarget::GPUComputeSim && GPUExecContext != nullptr)
+	{
+		UE_LOG(LogNiagara, Log, TEXT("................. %s Combined Parameters ................."), TEXT("GPU Script"));
+		GPUExecContext->CombinedParamStore.DumpParameters();
+	}
 	UE_LOG(LogNiagara, Log, TEXT("................. Particles ................."));
 	ParticleDataSet->Dump(false);
 	ParticleDataSet->Dump(true);
@@ -235,9 +255,15 @@ void FNiagaraEmitterInstance::Init(int32 InEmitterIdx, FName InSystemInstanceNam
 	// setup the parameer store for the GPU execution context; since spawn and update are combined here, we build one with params from both script props
 	if (CachedEmitter->SimTarget == ENiagaraSimTarget::GPUComputeSim)
 	{
-		GPUExecContext.InitParams(CachedEmitter->GetGPUComputeScript(), CachedEmitter->SpawnScriptProps.Script, CachedEmitter->UpdateScriptProps.Script, CachedEmitter->SimTarget);
-		SpawnExecContext.Parameters.Bind(&GPUExecContext.CombinedParamStore);
-		UpdateExecContext.Parameters.Bind(&GPUExecContext.CombinedParamStore);
+		GPUExecContext = new FNiagaraComputeExecutionContext();
+		GPUExecContext->InitParams(CachedEmitter->GetGPUComputeScript(), CachedEmitter->SpawnScriptProps.Script, CachedEmitter->UpdateScriptProps.Script, CachedEmitter->SimTarget);
+		GPUExecContext->MainDataSet = &Data;
+		GPUExecContext->RTGPUScript = CachedEmitter->GetGPUComputeScript()->GetRenderThreadScript();
+		GPUExecContext->RTSpawnScript = CachedEmitter->SpawnScriptProps.Script->GetRenderThreadScript();
+		GPUExecContext->RTUpdateScript = CachedEmitter->UpdateScriptProps.Script->GetRenderThreadScript();
+
+		SpawnExecContext.Parameters.Bind(&GPUExecContext->CombinedParamStore);
+		UpdateExecContext.Parameters.Bind(&GPUExecContext->CombinedParamStore);
 	}
 
 	EventExecContexts.SetNum(CachedEmitter->GetEventHandlers().Num());
@@ -257,11 +283,11 @@ void FNiagaraEmitterInstance::Init(int32 InEmitterIdx, FName InSystemInstanceNam
 	InterpSpawnStartBinding.Init(SpawnExecContext.Parameters, CachedEmitter->ToEmitterParameter(SYS_PARAM_EMITTER_INTERP_SPAWN_START_DT));
 	SpawnGroupBinding.Init(SpawnExecContext.Parameters, CachedEmitter->ToEmitterParameter(SYS_PARAM_EMITTER_SPAWN_GROUP));
 
-	if (CachedEmitter->SimTarget == ENiagaraSimTarget::GPUComputeSim)
+	if (CachedEmitter->SimTarget == ENiagaraSimTarget::GPUComputeSim && GPUExecContext != nullptr)
 	{
-		SpawnIntervalBindingGPU.Init(GPUExecContext.CombinedParamStore, CachedEmitter->ToEmitterParameter(SYS_PARAM_EMITTER_SPAWN_INTERVAL));
-		InterpSpawnStartBindingGPU.Init(GPUExecContext.CombinedParamStore, CachedEmitter->ToEmitterParameter(SYS_PARAM_EMITTER_INTERP_SPAWN_START_DT));
-		SpawnGroupBindingGPU.Init(GPUExecContext.CombinedParamStore, CachedEmitter->ToEmitterParameter(SYS_PARAM_EMITTER_SPAWN_GROUP));
+		SpawnIntervalBindingGPU.Init(GPUExecContext->CombinedParamStore, CachedEmitter->ToEmitterParameter(SYS_PARAM_EMITTER_SPAWN_INTERVAL));
+		InterpSpawnStartBindingGPU.Init(GPUExecContext->CombinedParamStore, CachedEmitter->ToEmitterParameter(SYS_PARAM_EMITTER_INTERP_SPAWN_START_DT));
+		SpawnGroupBindingGPU.Init(GPUExecContext->CombinedParamStore, CachedEmitter->ToEmitterParameter(SYS_PARAM_EMITTER_SPAWN_GROUP));
 	}
 
 	FNiagaraVariable EmitterAgeParam = CachedEmitter->ToEmitterParameter(SYS_PARAM_EMITTER_AGE);
@@ -273,9 +299,9 @@ void FNiagaraEmitterInstance::Init(int32 InEmitterIdx, FName InSystemInstanceNam
 		EventEmitterAgeBindings[i].Init(EventExecContexts[i].Parameters, EmitterAgeParam);
 	}
 
-	if (CachedEmitter->SimTarget == ENiagaraSimTarget::GPUComputeSim)
+	if (CachedEmitter->SimTarget == ENiagaraSimTarget::GPUComputeSim && GPUExecContext != nullptr)
 	{
-		EmitterAgeBindingGPU.Init(GPUExecContext.CombinedParamStore, EmitterAgeParam);
+		EmitterAgeBindingGPU.Init(GPUExecContext->CombinedParamStore, EmitterAgeParam);
 	}
 
 	SpawnExecCountBinding.Init(SpawnExecContext.Parameters, SYS_PARAM_ENGINE_EXEC_COUNT);
@@ -286,7 +312,7 @@ void FNiagaraEmitterInstance::Init(int32 InEmitterIdx, FName InSystemInstanceNam
 		EventExecCountBindings[i].Init(EventExecContexts[i].Parameters, SYS_PARAM_ENGINE_EXEC_COUNT);
 	}
 
-	if (CachedEmitter->SimTarget == ENiagaraSimTarget::GPUComputeSim)
+	if (CachedEmitter->SimTarget == ENiagaraSimTarget::GPUComputeSim && GPUExecContext != nullptr)
 	{
 		//Just ensure we've generated the singleton here on the GT as it throws a wobbler if we do this later in parallel.
 		NiagaraEmitterInstanceBatcher::Get();
@@ -328,7 +354,10 @@ void FNiagaraEmitterInstance::ResetSimulation()
 		UpdateScriptEventDataSet->ResetBuffers();
 	}
 	
-	GPUExecContext.Reset();
+	if (CachedEmitter->SimTarget == ENiagaraSimTarget::GPUComputeSim && GPUExecContext != nullptr)
+	{
+		GPUExecContext->Reset();
+	}
 
 	SetExecutionState(ENiagaraExecutionState::Active);
 }
@@ -425,7 +454,11 @@ void FNiagaraEmitterInstance::DirtyDataInterfaces()
 	// Make sure that our function tables need to be regenerated...
 	SpawnExecContext.DirtyDataInterfaces();
 	UpdateExecContext.DirtyDataInterfaces();
-	GPUExecContext.DirtyDataInterfaces();
+
+	if (CachedEmitter->SimTarget == ENiagaraSimTarget::GPUComputeSim && GPUExecContext != nullptr)
+	{
+		GPUExecContext->DirtyDataInterfaces();
+	}
 
 	for (FNiagaraScriptExecutionContext& EventContext : EventExecContexts)
 	{
@@ -453,6 +486,10 @@ void FNiagaraEmitterInstance::UnbindParameters()
 {
 	SpawnExecContext.Parameters.UnbindFromSourceStores();
 	UpdateExecContext.Parameters.UnbindFromSourceStores();
+	if (GPUExecContext != nullptr)
+	{
+		GPUExecContext->CombinedParamStore.UnbindFromSourceStores();
+	}
 
 	for (int32 EventIdx = 0; EventIdx < EventExecContexts.Num(); ++EventIdx)
 	{
@@ -763,9 +800,9 @@ void FNiagaraEmitterInstance::PreTick()
 	bool bOk = true;
 	bOk &= SpawnExecContext.Tick(ParentSystemInstance);
 	bOk &= UpdateExecContext.Tick(ParentSystemInstance);
-	if (CachedEmitter->SimTarget == ENiagaraSimTarget::GPUComputeSim)
+	if (CachedEmitter->SimTarget == ENiagaraSimTarget::GPUComputeSim && GPUExecContext != nullptr)
 	{
-		bOk &= GPUExecContext.Tick(ParentSystemInstance);
+		bOk &= GPUExecContext->Tick(ParentSystemInstance);
 	}
 	for (FNiagaraScriptExecutionContext& EventContext : EventExecContexts)
 	{
@@ -888,7 +925,7 @@ void FNiagaraEmitterInstance::Tick(float DeltaSeconds)
 			Binding.SetValue(Age);
 		}
 		
-		if (CachedEmitter->SimTarget == ENiagaraSimTarget::GPUComputeSim)
+		if (CachedEmitter->SimTarget == ENiagaraSimTarget::GPUComputeSim && GPUExecContext != nullptr)
 		{
 			EmitterAgeBindingGPU.SetValue(Age);
 		}
@@ -940,16 +977,11 @@ void FNiagaraEmitterInstance::Tick(float DeltaSeconds)
 
 	/* GPU simulation -  we just create an FNiagaraComputeExecutionContext, queue it, and let the batcher take care of the rest
 	 */
-	if (CachedEmitter->SimTarget == ENiagaraSimTarget::GPUComputeSim)	
+	if (CachedEmitter->SimTarget == ENiagaraSimTarget::GPUComputeSim && GPUExecContext != nullptr)
 	{
-		//FNiagaraComputeExecutionContext *ComputeContext = new FNiagaraComputeExecutionContext();
-		GPUExecContext.MainDataSet = &Data;
-		GPUExecContext.RTGPUScript = CachedEmitter->GetGPUComputeScript()->GetRenderThreadScript();
-		GPUExecContext.RTSpawnScript = CachedEmitter->SpawnScriptProps.Script->GetRenderThreadScript();
-		GPUExecContext.RTUpdateScript = CachedEmitter->UpdateScriptProps.Script->GetRenderThreadScript();
-		GPUExecContext.SpawnRateInstances = SpawnTotal;
-		GPUExecContext.EventSpawnTotal = EventSpawnTotal;
-		GPUExecContext.NumIndicesPerInstance = CachedEmitter->GetRenderers()[0]->GetNumIndicesPerInstance();
+		GPUExecContext->SpawnRateInstances = SpawnTotal;
+		GPUExecContext->EventSpawnTotal = EventSpawnTotal;
+		GPUExecContext->NumIndicesPerInstance = CachedEmitter->GetRenderers()[0]->GetNumIndicesPerInstance();
 
 		bool bOnlySetOnce = false;
 		for (FNiagaraSpawnInfo& Info : SpawnInfos)
@@ -986,27 +1018,27 @@ void FNiagaraEmitterInstance::Tick(float DeltaSeconds)
 			UE_LOG(LogNiagara, Log, TEXT(".................Update................."));
 			UpdateExecContext.Parameters.DumpParameters(true);
 			UE_LOG(LogNiagara, Log, TEXT("................. %s Combined Parameters (%d Spawned )................."), TEXT("GPU Script"), SpawnTotal);
-			GPUExecContext.CombinedParamStore.DumpParameters();
+			GPUExecContext->CombinedParamStore.DumpParameters();
 		}
 
-		int32 ParmSize = GPUExecContext.CombinedParamStore.GetPaddedParameterSizeInBytes();
+		int32 ParmSize = GPUExecContext->CombinedParamStore.GetPaddedParameterSizeInBytes();
 		
-		GPUExecContext.ParamData_RT.SetNumZeroed(ParmSize);
-		GPUExecContext.CombinedParamStore.CopyParameterDataToPaddedBuffer(GPUExecContext.ParamData_RT.GetData(), ParmSize);
+		GPUExecContext->ParamData_RT.SetNumZeroed(ParmSize);
+		GPUExecContext->CombinedParamStore.CopyParameterDataToPaddedBuffer(GPUExecContext->ParamData_RT.GetData(), ParmSize);
 		// Because each context is only ran once each frame, the CBuffer layout stays constant for the lifetime duration of the CBuffer (one frame).
-		GPUExecContext.CBufferLayout.ConstantBufferSize = ParmSize;
-		GPUExecContext.CBufferLayout.ComputeHash();
+		GPUExecContext->CBufferLayout.ConstantBufferSize = ParmSize;
+		GPUExecContext->CBufferLayout.ComputeHash();
 
 		// push event data sets to the context
 		for (FNiagaraDataSet *Set : UpdateScriptEventDataSets)
 		{
-			GPUExecContext.UpdateEventWriteDataSets.Add(Set);
+			GPUExecContext->UpdateEventWriteDataSets.Add(Set);
 		}
 
-		GPUExecContext.EventHandlerScriptProps = CachedEmitter->GetEventHandlers();
-		GPUExecContext.EventSets = EventSet;
-		GPUExecContext.EventSpawnCounts = EventHandlerSpawnCounts;
-		NiagaraEmitterInstanceBatcher::Get()->Queue(&GPUExecContext);
+		GPUExecContext->EventHandlerScriptProps = CachedEmitter->GetEventHandlers();
+		GPUExecContext->EventSets = EventSet;
+		GPUExecContext->EventSpawnCounts = EventHandlerSpawnCounts;
+		NiagaraEmitterInstanceBatcher::Get()->Queue(GPUExecContext);
 
 		// Need to call post-tick, which calls the copy to previous for interpolated spawning
 		SpawnExecContext.PostTick();
