Prevent changing internal name of Niagara Emitter assets that are not local copies in Systems.

diff --git a/Engine/Plugins/FX/Niagara/Source/NiagaraEditor/Private/Sequencer/NiagaraSequence/MovieSceneNiagaraEmitterTrack.cpp b/Engine/Plugins/FX/Niagara/Source/NiagaraEditor/Private/Sequencer/NiagaraSequence/MovieSceneNiagaraEmitterTrack.cpp
--- a/Engine/Plugins/FX/Niagara/Source/NiagaraEditor/Private/Sequencer/NiagaraSequence/MovieSceneNiagaraEmitterTrack.cpp
+++ b/Engine/Plugins/FX/Niagara/Source/NiagaraEditor/Private/Sequencer/NiagaraSequence/MovieSceneNiagaraEmitterTrack.cpp
@@ -85,6 +85,15 @@ void UMovieSceneNiagaraEmitterTrack::Initialize(FNiagaraSystemViewModel& InSyste
 	CreateSections(InFrameResolution);
 }
 
+bool UMovieSceneNiagaraEmitterTrack::CanRename() const
+{
+	if (SystemViewModel && SystemViewModel->GetEditMode() == ENiagaraSystemViewModelEditMode::SystemAsset)
+	{
+		return true;
+	}
+	return false;
+}
+
 FNiagaraSystemViewModel& UMovieSceneNiagaraEmitterTrack::GetSystemViewModel() const
 {
 	return *SystemViewModel;
diff --git a/Engine/Plugins/FX/Niagara/Source/NiagaraEditor/Private/Sequencer/NiagaraSequence/MovieSceneNiagaraEmitterTrack.h b/Engine/Plugins/FX/Niagara/Source/NiagaraEditor/Private/Sequencer/NiagaraSequence/MovieSceneNiagaraEmitterTrack.h
--- a/Engine/Plugins/FX/Niagara/Source/NiagaraEditor/Private/Sequencer/NiagaraSequence/MovieSceneNiagaraEmitterTrack.h
+++ b/Engine/Plugins/FX/Niagara/Source/NiagaraEditor/Private/Sequencer/NiagaraSequence/MovieSceneNiagaraEmitterTrack.h
@@ -58,6 +58,8 @@ class UMovieSceneNiagaraEmitterTrack
 public:
 	void Initialize(FNiagaraSystemViewModel& SystemViewModel, TSharedRef<FNiagaraEmitterHandleViewModel> InEmitterHandleViewModel, const FFrameRate& InFrameResolution);
 
+	virtual bool CanRename() const override;
+
 	FNiagaraSystemViewModel& GetSystemViewModel() const;
 
 	TSharedPtr<FNiagaraEmitterHandleViewModel> GetEmitterHandleViewModel() const;
diff --git a/Engine/Plugins/FX/Niagara/Source/NiagaraEditorWidgets/Private/SNiagaraStack.cpp b/Engine/Plugins/FX/Niagara/Source/NiagaraEditorWidgets/Private/SNiagaraStack.cpp
--- a/Engine/Plugins/FX/Niagara/Source/NiagaraEditorWidgets/Private/SNiagaraStack.cpp
+++ b/Engine/Plugins/FX/Niagara/Source/NiagaraEditorWidgets/Private/SNiagaraStack.cpp
@@ -134,11 +134,14 @@ void SNiagaraStack::Construct(const FArguments& InArgs, UNiagaraStackViewModel*
 			FMenuBuilder MenuBuilder(true, nullptr);
 			MenuBuilder.BeginSection("EmitterInlineMenuActions", LOCTEXT("EmitterActions", "Emitter Actions"));
 			{
-				MenuBuilder.AddMenuEntry(
-					LOCTEXT("RenameEmitter", "Rename Emitter"),
-					LOCTEXT("RenameEmitterToolTip", "Rename this emitter"), 
-					FSlateIcon(),
-					FUIAction(FExecuteAction::CreateSP(InlineEditableTextBlock.Get(), &SInlineEditableTextBlock::EnterEditingMode)));
+				if (!GetEmitterNameIsReadOnly()) // Only allow renaming local copies of Emitters in Systems
+				{
+					MenuBuilder.AddMenuEntry(
+						LOCTEXT("RenameEmitter", "Rename Emitter"),
+						LOCTEXT("RenameEmitterToolTip", "Rename this local emitter copy"),
+						FSlateIcon(),
+						FUIAction(FExecuteAction::CreateSP(InlineEditableTextBlock.Get(), &SInlineEditableTextBlock::EnterEditingMode)));
+				}
 
 				MenuBuilder.AddMenuEntry(
 					LOCTEXT("ShowEmitterInContentBrowser", "Show in Content Browser"),
@@ -304,6 +307,7 @@ void SNiagaraStack::ConstructHeaderWidget()
 						.Text(StackViewModel->GetEmitterHandleViewModel()->AsShared(), &FNiagaraEmitterHandleViewModel::GetNameText)
 						.OnTextCommitted(this, &SNiagaraStack::OnStackViewNameTextCommitted)
 						.OnVerifyTextChanged(StackViewModel->GetEmitterHandleViewModel()->AsShared(), &FNiagaraEmitterHandleViewModel::VerifyNameTextChanged)
+						.IsReadOnly(this, &SNiagaraStack::GetEmitterNameIsReadOnly)
 					]
 					+ SWrapBox::Slot()
 					.Padding(3, 0)
@@ -504,6 +508,15 @@ EVisibility SNiagaraStack::GetOpenSourceEmitterVisibility() const
 	return CanOpenSourceEmitter() ? EVisibility::Visible : EVisibility::Collapsed;
 }
 
+bool SNiagaraStack::GetEmitterNameIsReadOnly() const
+{
+	if (StackViewModel && StackViewModel->GetSystemViewModel().IsValid())
+	{
+		return StackViewModel->GetSystemViewModel()->GetEditMode() == ENiagaraSystemViewModelEditMode::EmitterAsset;
+	}
+	return true;
+}
+
 bool SNiagaraStack::CanOpenSourceEmitter() const
 {
 	if (StackViewModel && StackViewModel->GetEmitterHandleViewModel().IsValid() && StackViewModel->GetEmitterHandleViewModel()->GetEmitterHandle()
diff --git a/Engine/Plugins/FX/Niagara/Source/NiagaraEditorWidgets/Private/SNiagaraStack.h b/Engine/Plugins/FX/Niagara/Source/NiagaraEditorWidgets/Private/SNiagaraStack.h
--- a/Engine/Plugins/FX/Niagara/Source/NiagaraEditorWidgets/Private/SNiagaraStack.h
+++ b/Engine/Plugins/FX/Niagara/Source/NiagaraEditorWidgets/Private/SNiagaraStack.h
@@ -78,6 +78,7 @@ private:
 	EVisibility GetOpenSourceEmitterVisibility() const;
 
 	// source name handling
+	bool GetEmitterNameIsReadOnly() const;
 	FText GetSourceEmitterNameText() const;
 	FText GetEmitterNameToolTip() const;
 	void OnStackViewNameTextCommitted(const FText& InText, ETextCommit::Type CommitInfo) const;
