Niagara - Fix cooking of compute shaders for target platforms with more than one shader platform.

diff --git a/Engine/Plugins/FX/Niagara/Source/Niagara/Classes/NiagaraScript.h b/Engine/Plugins/FX/Niagara/Source/Niagara/Classes/NiagaraScript.h
--- a/Engine/Plugins/FX/Niagara/Source/Niagara/Classes/NiagaraScript.h
+++ b/Engine/Plugins/FX/Niagara/Source/Niagara/Classes/NiagaraScript.h
@@ -19,8 +19,8 @@
 class UNiagaraDataInterface;
 class FNiagaraCompileRequestDataBase;
 
-void SerializeNiagaraShaderMaps(const TMap<const ITargetPlatform*, FNiagaraShaderScript*>* PlatformScriptResourcesToSavePtr, FArchive& Ar, FNiagaraShaderScript& OutLoadedResources);
-void ProcessSerializedShaderMaps(UNiagaraScript* Owner, FNiagaraShaderScript LoadedResource, FNiagaraShaderScript* (&OutScriptResourcesLoaded)[ERHIFeatureLevel::Num]);
+void SerializeNiagaraShaderMaps(const TMap<const ITargetPlatform*, TArray<FNiagaraShaderScript*>>* PlatformScriptResourcesToSave, FArchive& Ar, TArray<FNiagaraShaderScript>& OutLoadedResources);
+void ProcessSerializedShaderMaps(UNiagaraScript* Owner, TArray<FNiagaraShaderScript>& LoadedResources, FNiagaraShaderScript& OutResourceForCurrentPlatform, FNiagaraShaderScript* (&OutScriptResourcesLoaded)[ERHIFeatureLevel::Num]);
 
 #define NIAGARA_INVALID_MEMORY (0xBA)
 
@@ -413,7 +413,8 @@ public:
 	//~ End UObject interface
 
 	// Infrastructure for GPU compute Shaders
-	NIAGARA_API void CacheResourceShadersForCooking(EShaderPlatform ShaderPlatform, FNiagaraShaderScript* &OutCachedResource);
+	NIAGARA_API void CacheResourceShadersForCooking(EShaderPlatform ShaderPlatform, TArray<FNiagaraShaderScript*>& InOutCachedResources);
+
 	NIAGARA_API void CacheResourceShadersForRendering(bool bRegenerateId, bool bForceRecompile=false);
 	void BeginCacheForCookedPlatformData(const ITargetPlatform *TargetPlatform);
 	void CacheShadersForResources(EShaderPlatform ShaderPlatform, FNiagaraShaderScript *ResourceToCache, bool bApplyCompletedShaderMapForRendering, bool bForceRecompile = false, bool bCooking=false);
@@ -536,6 +537,8 @@ private:
 	UPROPERTY()
 	mutable FNiagaraVMExecutableDataId LastGeneratedVMId;
 
+	TArray<FNiagaraShaderScript> LoadedScriptResources;
+
 	FNiagaraShaderScript ScriptResource;
 
 	FNiagaraShaderScript* ScriptResourcesByFeatureLevel[ERHIFeatureLevel::Num];
@@ -553,7 +556,7 @@ private:
 
 #if WITH_EDITORONLY_DATA
 	/* script resources being cached for cooking. */
-	TMap<const class ITargetPlatform*, FNiagaraShaderScript*> CachedScriptResourcesForCooking;
+	TMap<const class ITargetPlatform*, TArray<FNiagaraShaderScript*>> CachedScriptResourcesForCooking;
 
 	UPROPERTY(Transient)
 	TArray<UObject*> ActiveCompileRoots;
diff --git a/Engine/Plugins/FX/Niagara/Source/Niagara/Private/NiagaraScript.cpp b/Engine/Plugins/FX/Niagara/Source/Niagara/Private/NiagaraScript.cpp
--- a/Engine/Plugins/FX/Niagara/Source/Niagara/Private/NiagaraScript.cpp
+++ b/Engine/Plugins/FX/Niagara/Source/Niagara/Private/NiagaraScript.cpp
@@ -336,9 +336,9 @@ void UNiagaraScript::Serialize(FArchive& Ar)
 		)
 	{
 #if WITH_EDITOR
-		SerializeNiagaraShaderMaps(&CachedScriptResourcesForCooking, Ar, ScriptResource);
+		SerializeNiagaraShaderMaps(&CachedScriptResourcesForCooking, Ar, LoadedScriptResources);
 #else
-		SerializeNiagaraShaderMaps(NULL, Ar, ScriptResource);
+		SerializeNiagaraShaderMaps(nullptr, Ar, LoadedScriptResources);
 #endif
 	}
 }
@@ -411,7 +411,7 @@ void UNiagaraScript::PostLoad()
 	const int32 NiagaraVer = GetLinkerCustomVersion(FNiagaraCustomVersion::GUID);
 	
 	// Resources can be processed / registered now that we're back on the main thread
-	ProcessSerializedShaderMaps(this, ScriptResource, ScriptResourcesByFeatureLevel);
+	ProcessSerializedShaderMaps(this, LoadedScriptResources, ScriptResource, ScriptResourcesByFeatureLevel);
 
 	if (GIsEditor)
 	{
@@ -920,6 +920,7 @@ void UNiagaraScript::GetAssetRegistryTags(TArray<FAssetRegistryTag>& OutTags) co
 }
 
 #if WITH_EDITOR
+
 void UNiagaraScript::BeginCacheForCookedPlatformData(const ITargetPlatform *TargetPlatform)
 {
 	if (CanBeRunOnGpu())
@@ -936,30 +937,21 @@ void UNiagaraScript::BeginCacheForCookedPlatformData(const ITargetPlatform *Targ
 		TArray<FName> DesiredShaderFormats;
 		TargetPlatform->GetAllTargetedShaderFormats(DesiredShaderFormats);
 
-		FNiagaraShaderScript** CachedScriptResourceForPlatformPtr = CachedScriptResourcesForCooking.Find(TargetPlatform);
+		TArray<FNiagaraShaderScript*>& CachedScriptResourcesForPlatform = CachedScriptResourcesForCooking.FindOrAdd(TargetPlatform);
 
-		if (DesiredShaderFormats.Num() > 0 && CachedScriptResourceForPlatformPtr == nullptr)
+		// Cache for all the shader formats that the cooking target requires
+		for (int32 FormatIndex = 0; FormatIndex < DesiredShaderFormats.Num(); FormatIndex++)
 		{
-			CachedScriptResourcesForCooking.Add(TargetPlatform);
-			FNiagaraShaderScript* &CachedScriptResourceForPlatform = *CachedScriptResourcesForCooking.Find(TargetPlatform);
-
-			// Cache for all the shader formats that the cooking target requires
-			for (int32 FormatIndex = 0; FormatIndex < DesiredShaderFormats.Num(); FormatIndex++)
+			const EShaderPlatform LegacyShaderPlatform = ShaderFormatToLegacyShaderPlatform(DesiredShaderFormats[FormatIndex]);
+			if (RHISupportsComputeShaders(LegacyShaderPlatform))
 			{
-				if (DesiredShaderFormats[FormatIndex] != TEXT("PCD3D_SM4"))	// TODO: remove once we get rid of SM4 globally
-				{
-					const EShaderPlatform LegacyShaderPlatform = ShaderFormatToLegacyShaderPlatform(DesiredShaderFormats[FormatIndex]);
-
-					// Begin caching shaders for the target platform and store the FNiagaraShaderScript being compiled into CachedScriptResourcesForCooking
-					CacheResourceShadersForCooking(LegacyShaderPlatform, CachedScriptResourceForPlatform);
-				}
+				CacheResourceShadersForCooking(LegacyShaderPlatform, CachedScriptResourcesForPlatform);
 			}
-			check(CachedScriptResourceForPlatform);
 		}
 	}
 }
 
-void UNiagaraScript::CacheResourceShadersForCooking(EShaderPlatform ShaderPlatform, FNiagaraShaderScript* &OutCachedResource)
+void UNiagaraScript::CacheResourceShadersForCooking(EShaderPlatform ShaderPlatform, TArray<FNiagaraShaderScript*>& InOutCachedResources)
 {
 	if (CanBeRunOnGpu())
 	{
@@ -983,7 +975,7 @@ void UNiagaraScript::CacheResourceShadersForCooking(EShaderPlatform ShaderPlatfo
 			INiagaraModule NiagaraModule = FModuleManager::GetModuleChecked<INiagaraModule>(TEXT("Niagara"));
 			NiagaraModule.ProcessShaderCompilationQueue();
 
-			OutCachedResource = ResourceToCache;
+			InOutCachedResources.Add(ResourceToCache);
 		}
 	}
 }
@@ -1205,67 +1197,77 @@ NIAGARA_API bool UNiagaraScript::DidScriptCompilationSucceed(bool bGPUScript) co
 	return false;
 }
 
-void SerializeNiagaraShaderMaps(const TMap<const ITargetPlatform*, FNiagaraShaderScript*>* PlatformScriptResourcesToSavePtr, FArchive& Ar, FNiagaraShaderScript &OutLoadedResource)
+void SerializeNiagaraShaderMaps(const TMap<const ITargetPlatform*, TArray<FNiagaraShaderScript*>>* PlatformScriptResourcesToSave, FArchive& Ar, TArray<FNiagaraShaderScript>& OutLoadedResources)
 {
 //	SCOPED_LOADTIMER(SerializeInlineShaderMaps);
 	if (Ar.IsSaving())
 	{
 		int32 NumResourcesToSave = 0;
-		FNiagaraShaderScript *ScriptResourceToSave = nullptr;
-		FNiagaraShaderScript* const * ScriptResourceToSavePtr = nullptr;
+		const TArray<FNiagaraShaderScript*>* ScriptResourcesToSavePtr = nullptr;
+
 		if (Ar.IsCooking())
 		{
-			check(PlatformScriptResourcesToSavePtr);
-			auto& PlatformScriptResourcesToSave = *PlatformScriptResourcesToSavePtr;
-
-			ScriptResourceToSavePtr = PlatformScriptResourcesToSave.Find(Ar.CookingTarget());
-			//check(ScriptResourceToSavePtr != NULL || (Ar.GetLinker() == NULL));
-			if (ScriptResourceToSavePtr != NULL)
+			checkf(PlatformScriptResourcesToSave != nullptr, TEXT("PlatformScriptResourcesToSave must be supplied when cooking"));
+			ScriptResourcesToSavePtr = PlatformScriptResourcesToSave->Find(Ar.CookingTarget());
+			if (ScriptResourcesToSavePtr != nullptr)
 			{
-				NumResourcesToSave++;
+				NumResourcesToSave = ScriptResourcesToSavePtr->Num();
 			}
 		}
 
 		Ar << NumResourcesToSave;
 
-		if (ScriptResourceToSavePtr)
+		if (ScriptResourcesToSavePtr != nullptr)
 		{
-			ScriptResourceToSave = *ScriptResourceToSavePtr;
-			check(ScriptResourceToSave);
-			ScriptResourceToSave->SerializeShaderMap(Ar);
+			for (FNiagaraShaderScript* ScriptResourceToSave : (*ScriptResourcesToSavePtr))
+			{
+				checkf(ScriptResourceToSave != nullptr, TEXT("Invalid script resource was cached"));
+				ScriptResourceToSave->SerializeShaderMap(Ar);
+			}
 		}
-
 	}
 	else if (Ar.IsLoading())
 	{
 		int32 NumLoadedResources = 0;
 		Ar << NumLoadedResources;
-		if (NumLoadedResources)
+		for (int32 i = 0; i < NumLoadedResources; i++)
 		{
-			OutLoadedResource.SerializeShaderMap(Ar);
+			FNiagaraShaderScript LoadedResource;
+			LoadedResource.SerializeShaderMap(Ar);
+			OutLoadedResources.Add(LoadedResource);
 		}
 	}
 }
 
-void ProcessSerializedShaderMaps(UNiagaraScript* Owner, FNiagaraShaderScript LoadedResource, FNiagaraShaderScript* (&OutScriptResourcesLoaded)[ERHIFeatureLevel::Num])
+void ProcessSerializedShaderMaps(UNiagaraScript* Owner, TArray<FNiagaraShaderScript>& LoadedResources, FNiagaraShaderScript& OutResourceForCurrentPlatform, FNiagaraShaderScript* (&OutScriptResourcesLoaded)[ERHIFeatureLevel::Num])
 {
 	check(IsInGameThread());
 
-	LoadedResource.RegisterShaderMap();
-
-	FNiagaraShaderMap* LoadedShaderMap = LoadedResource.GetGameThreadShaderMap();
-	if (LoadedShaderMap && LoadedShaderMap->GetShaderPlatform() == GMaxRHIShaderPlatform)
+	for (FNiagaraShaderScript& LoadedResource : LoadedResources)
 	{
-		ERHIFeatureLevel::Type LoadedFeatureLevel = LoadedShaderMap->GetShaderMapId().FeatureLevel;
-		if (!OutScriptResourcesLoaded[LoadedFeatureLevel])
+		LoadedResource.RegisterShaderMap();
+
+		FNiagaraShaderMap* LoadedShaderMap = LoadedResource.GetGameThreadShaderMap();
+		if (LoadedShaderMap && LoadedShaderMap->GetShaderPlatform() == GMaxRHIShaderPlatform)
 		{
-			OutScriptResourcesLoaded[LoadedFeatureLevel] = Owner->AllocateResource();
-		}
+			OutResourceForCurrentPlatform = LoadedResource;
 
-		OutScriptResourcesLoaded[LoadedFeatureLevel]->SetShaderMap(LoadedShaderMap);
-		LoadedResource.SetDataInterfaceParamInfo(LoadedResource.GetShaderGameThread()->GetDIParameters());
-	}
+			ERHIFeatureLevel::Type LoadedFeatureLevel = LoadedShaderMap->GetShaderMapId().FeatureLevel;
+			if (!OutScriptResourcesLoaded[LoadedFeatureLevel])
+			{
+				OutScriptResourcesLoaded[LoadedFeatureLevel] = Owner->AllocateResource();
+			}
+
+			OutScriptResourcesLoaded[LoadedFeatureLevel]->SetShaderMap(LoadedShaderMap);
+			OutResourceForCurrentPlatform.SetDataInterfaceParamInfo(LoadedResource.GetShaderGameThread()->GetDIParameters());
 
+			break;
+		}
+		else
+		{
+			LoadedResource.DiscardShaderMap();
+		}
+	}
 }
 
 FNiagaraShaderScript* UNiagaraScript::AllocateResource()
diff --git a/Engine/Plugins/FX/Niagara/Source/NiagaraShader/Private/NiagaraShared.cpp b/Engine/Plugins/FX/Niagara/Source/NiagaraShader/Private/NiagaraShared.cpp
--- a/Engine/Plugins/FX/Niagara/Source/NiagaraShader/Private/NiagaraShared.cpp
+++ b/Engine/Plugins/FX/Niagara/Source/NiagaraShader/Private/NiagaraShared.cpp
@@ -140,6 +140,14 @@ void FNiagaraShaderScript::RegisterShaderMap()
 	}
 }
 
+void  FNiagaraShaderScript::DiscardShaderMap()
+{
+	if (GameThreadShaderMap)
+	{
+		GameThreadShaderMap->DiscardSerializedShaders();
+	}
+}
+
 void FNiagaraShaderScript::ReleaseShaderMap()
 {
 	if (GameThreadShaderMap)
diff --git a/Engine/Plugins/FX/Niagara/Source/NiagaraShader/Public/NiagaraShared.h b/Engine/Plugins/FX/Niagara/Source/NiagaraShader/Public/NiagaraShared.h
--- a/Engine/Plugins/FX/Niagara/Source/NiagaraShader/Public/NiagaraShared.h
+++ b/Engine/Plugins/FX/Niagara/Source/NiagaraShader/Public/NiagaraShared.h
@@ -555,6 +555,8 @@ public:
 
 	NIAGARASHADER_API void RegisterShaderMap();
 
+	NIAGARASHADER_API void DiscardShaderMap();
+
 	/** Note: SetRenderingThreadShaderMap must also be called with the same value, but from the rendering thread. */
 	void SetGameThreadShaderMap(FNiagaraShaderMap* InShaderMap)
 	{
