Fix custom hlsl node from crashing when dragging in new parameter

diff --git a/Engine/Plugins/FX/Niagara/Source/NiagaraEditor/Private/NiagaraNodeCustomHlsl.cpp b/Engine/Plugins/FX/Niagara/Source/NiagaraEditor/Private/NiagaraNodeCustomHlsl.cpp
--- a/Engine/Plugins/FX/Niagara/Source/NiagaraEditor/Private/NiagaraNodeCustomHlsl.cpp
+++ b/Engine/Plugins/FX/Niagara/Source/NiagaraEditor/Private/NiagaraNodeCustomHlsl.cpp
@@ -212,7 +212,7 @@ void UNiagaraNodeCustomHlsl::BuildParameterMapHistory(FNiagaraParameterMapHistor
 	int32 ParamMapIdx = INDEX_NONE;
 	TArray<FNiagaraVariable> LocalVars;
 	// This only works currently if the input pins are in the same order as the signature pins.
-	if (InputPins.Num() == Signature.Inputs.Num() + 1)// the add pin is extra
+	if (InputPins.Num() == Signature.Inputs.Num() + 1 && OutputPins.Num() == Signature.Outputs.Num() + 1)// the add pin is extra
 	{
 		bool bHasParamMapInput = false;
 		bool bHasParamMapOutput = false;
@@ -340,8 +340,6 @@ void UNiagaraNodeCustomHlsl::RebuildSignatureFromPins()
 
 	Signature = Sig;
 
-	RefreshFromExternalChanges();
-	MarkNodeRequiresSynchronization(__FUNCTION__, true);
 }
 
 #undef LOCTEXT_NAMESPACE
diff --git a/Engine/Plugins/FX/Niagara/Source/NiagaraEditor/Private/NiagaraNodeWithDynamicPins.cpp b/Engine/Plugins/FX/Niagara/Source/NiagaraEditor/Private/NiagaraNodeWithDynamicPins.cpp
--- a/Engine/Plugins/FX/Niagara/Source/NiagaraEditor/Private/NiagaraNodeWithDynamicPins.cpp
+++ b/Engine/Plugins/FX/Niagara/Source/NiagaraEditor/Private/NiagaraNodeWithDynamicPins.cpp
@@ -279,6 +279,7 @@ void UNiagaraNodeWithDynamicPins::PinNameTextCommitted(const FText& Text, ETextC
 		FString PinOldName = Pin->PinName.ToString();
 		Pin->PinName = *Text.ToString();
 		OnPinRenamed(Pin, PinOldName);
+		MarkNodeRequiresSynchronization(__FUNCTION__, true);
 	}
 }
 
