Fix options drop-down arrow for stack inputs being a tab-stop.

diff --git a/Engine/Plugins/FX/Niagara/Source/NiagaraEditorWidgets/Private/Stack/SNiagaraStackFunctionInputValue.cpp b/Engine/Plugins/FX/Niagara/Source/NiagaraEditorWidgets/Private/Stack/SNiagaraStackFunctionInputValue.cpp
--- a/Engine/Plugins/FX/Niagara/Source/NiagaraEditorWidgets/Private/Stack/SNiagaraStackFunctionInputValue.cpp
+++ b/Engine/Plugins/FX/Niagara/Source/NiagaraEditorWidgets/Private/Stack/SNiagaraStackFunctionInputValue.cpp
@@ -164,6 +164,7 @@ void SNiagaraStackFunctionInputValue::Construct(const FArguments& InArgs, UNiaga
 			.Padding(3, 0, 0, 0)
 			[
 				SAssignNew(SetFunctionInputButton, SComboButton)
+				.IsFocusable(false)
 				.ButtonStyle(FEditorStyle::Get(), "HoverHintOnly")
 				.ForegroundColor(FSlateColor::UseForeground())
 				.OnGetMenuContent(this, &SNiagaraStackFunctionInputValue::OnGetAvailableHandleMenu)
