Niagara - Fix an assert when creating a new empty system or emitter.  Also adjust default icon scale.

diff --git a/Engine/Plugins/FX/Niagara/Source/NiagaraEditor/Private/SNewEmitterDialog.cpp b/Engine/Plugins/FX/Niagara/Source/NiagaraEditor/Private/SNewEmitterDialog.cpp
--- a/Engine/Plugins/FX/Niagara/Source/NiagaraEditor/Private/SNewEmitterDialog.cpp
+++ b/Engine/Plugins/FX/Niagara/Source/NiagaraEditor/Private/SNewEmitterDialog.cpp
@@ -20,7 +20,6 @@ void SNewEmitterDialog::Construct(const FArguments& InArgs)
 	FAssetPickerConfig AssetPickerConfig;
 	AssetPickerConfig.SelectionMode = ESelectionMode::SingleToggle;
 	AssetPickerConfig.InitialAssetViewType = EAssetViewType::List;
-	AssetPickerConfig.ThumbnailScale = .4f;
 	AssetPickerConfig.Filter.ClassNames.Add(UNiagaraEmitter::StaticClass()->GetFName());
 	AssetPickerConfig.GetCurrentSelectionDelegates.Add(&GetSelectedEmitterAssetsFromPicker);
 	AssetPickerConfig.OnAssetsActivated.BindSP(this, &SNewEmitterDialog::OnEmitterAssetsActivated);
diff --git a/Engine/Plugins/FX/Niagara/Source/NiagaraEditor/Private/SNewSystemDialog.cpp b/Engine/Plugins/FX/Niagara/Source/NiagaraEditor/Private/SNewSystemDialog.cpp
--- a/Engine/Plugins/FX/Niagara/Source/NiagaraEditor/Private/SNewSystemDialog.cpp
+++ b/Engine/Plugins/FX/Niagara/Source/NiagaraEditor/Private/SNewSystemDialog.cpp
@@ -27,7 +27,6 @@ void SNewSystemDialog::Construct(const FArguments& InArgs)
 	FAssetPickerConfig SystemAssetPickerConfig;
 	SystemAssetPickerConfig.SelectionMode = ESelectionMode::SingleToggle;
 	SystemAssetPickerConfig.InitialAssetViewType = EAssetViewType::List;
-	SystemAssetPickerConfig.ThumbnailScale = .4f;
 	SystemAssetPickerConfig.Filter.ClassNames.Add(UNiagaraSystem::StaticClass()->GetFName());
 	SystemAssetPickerConfig.GetCurrentSelectionDelegates.Add(&GetSelectedSystemAssetsFromPicker);
 	SystemAssetPickerConfig.OnAssetsActivated.BindSP(this, &SNewSystemDialog::OnSystemAssetsActivated);
@@ -37,7 +36,6 @@ void SNewSystemDialog::Construct(const FArguments& InArgs)
 	FAssetPickerConfig EmitterAssetPickerConfig;
 	EmitterAssetPickerConfig.SelectionMode = ESelectionMode::Multi;
 	EmitterAssetPickerConfig.InitialAssetViewType = EAssetViewType::List;
-	EmitterAssetPickerConfig.ThumbnailScale = .4f;
 	EmitterAssetPickerConfig.Filter.ClassNames.Add(UNiagaraEmitter::StaticClass()->GetFName());
 	EmitterAssetPickerConfig.GetCurrentSelectionDelegates.Add(&GetSelectedEmitterAssetsFromPicker);
 	EmitterAssetPickerConfig.OnAssetsActivated.BindSP(this, &SNewSystemDialog::OnEmitterAssetsActivated);
diff --git a/Engine/Plugins/FX/Niagara/Source/NiagaraEditor/Private/SNiagaraNewAssetDialog.cpp b/Engine/Plugins/FX/Niagara/Source/NiagaraEditor/Private/SNiagaraNewAssetDialog.cpp
--- a/Engine/Plugins/FX/Niagara/Source/NiagaraEditor/Private/SNiagaraNewAssetDialog.cpp
+++ b/Engine/Plugins/FX/Niagara/Source/NiagaraEditor/Private/SNiagaraNewAssetDialog.cpp
@@ -195,11 +195,9 @@ void SNiagaraNewAssetDialog::ConfirmSelection()
 	if (SelectedOption.OnGetSelectedAssetsFromPicker.IsBound())
 	{
 		SelectedAssets.Append(SelectedOption.OnGetSelectedAssetsFromPicker.Execute());
+		ensureMsgf(SelectedAssets.Num() > 0, TEXT("No assets selected when dialog was confirmed."));
 	}
-	if (ensureMsgf(SelectedAssets.Num() > 0, TEXT("No assets selected when dialog was confirmed.")))
-	{
-		bUserConfirmedSelection = true;
-	}
+	bUserConfirmedSelection = true;
 	RequestDestroyWindow();
 }
 
