Fix Niagara Systems in world not reinitializing when parameters are changed.
Fix log spam from Niagara Emitters/Systems when DIs fail to compile.
Improve Skel Mesh DI error reporting.

diff --git a/Engine/Plugins/FX/Niagara/Source/Niagara/Private/NiagaraComponent.cpp b/Engine/Plugins/FX/Niagara/Source/Niagara/Private/NiagaraComponent.cpp
--- a/Engine/Plugins/FX/Niagara/Source/Niagara/Private/NiagaraComponent.cpp
+++ b/Engine/Plugins/FX/Niagara/Source/Niagara/Private/NiagaraComponent.cpp
@@ -353,14 +353,14 @@ void UNiagaraComponent::TickComponent(float DeltaSeconds, enum ELevelTick TickTy
 		return;
 	}
 
-	if (!bIsActive && bAutoActivate)
+	if (!SystemInstance)
 	{
-		Activate();
+		return;
 	}
 
-	if (!SystemInstance)
+	if (!bIsActive && bAutoActivate && SystemInstance.Get() && SystemInstance->GetAreDataInterfacesInitialized())
 	{
-		return;
+		Activate();
 	}
 
 	check(SystemInstance->IsSolo());
@@ -1034,7 +1034,6 @@ void UNiagaraComponent::PreEditChange(UProperty* PropertyAboutToChange)
 
 void UNiagaraComponent::PostEditChangeProperty(FPropertyChangedEvent& PropertyChangedEvent)
 {
-
 	FName PropertyName;
 	if (PropertyChangedEvent.Property)
 	{
@@ -1183,7 +1182,7 @@ bool UNiagaraComponent::IsParameterValueOverriddenLocally(const FName& InParamNa
 	return false;
 }
 
-void UNiagaraComponent::SetParameterValueOverriddenLocally(const FNiagaraVariable& InParam, bool bInOverriden)
+void UNiagaraComponent::SetParameterValueOverriddenLocally(const FNiagaraVariable& InParam, bool bInOverriden, bool bRequiresSystemInstanceReset)
 {
 	bool* FoundVar = EditorOverridesValue.Find(InParam.GetName());
 
@@ -1200,6 +1199,12 @@ void UNiagaraComponent::SetParameterValueOverriddenLocally(const FNiagaraVariabl
 		EditorOverridesValue.Remove(InParam.GetName());
 		Asset->GetExposedParameters().CopyParameterData(OverrideParameters, InParam);
 	}
+	
+	if (bRequiresSystemInstanceReset)
+	{
+		SystemInstance->Reset(FNiagaraSystemInstance::EResetMode::ResetSystem, true);
+	}
+	
 }
 
 
diff --git a/Engine/Plugins/FX/Niagara/Source/Niagara/Private/NiagaraDataInterfaceSkeletalMesh.cpp b/Engine/Plugins/FX/Niagara/Source/Niagara/Private/NiagaraDataInterfaceSkeletalMesh.cpp
--- a/Engine/Plugins/FX/Niagara/Source/Niagara/Private/NiagaraDataInterfaceSkeletalMesh.cpp
+++ b/Engine/Plugins/FX/Niagara/Source/Niagara/Private/NiagaraDataInterfaceSkeletalMesh.cpp
@@ -862,17 +862,25 @@ bool UNiagaraDataInterfaceSkeletalMesh::PerInstanceTick(void* PerInstanceData, F
 TArray<FNiagaraDataInterfaceError> UNiagaraDataInterfaceSkeletalMesh::GetErrors()
 {
 	TArray<FNiagaraDataInterfaceError> Errors;
-	bool bHasError= false;
+	bool bHasCPUAccessError= false;
+	bool bHasNoMeshAssignedError = false;
+	
+	// Collect Errors
 	if (DefaultMesh != nullptr)
 	{
 		for (auto info : DefaultMesh->GetLODInfoArray())
 		{
 			if (!info.bAllowCPUAccess)
-				bHasError = true;
+				bHasCPUAccessError = true;
 		}
 	}
+	else
+	{
+		bHasNoMeshAssignedError = true;
+	}
 
-	if (Source == nullptr && bHasError)
+	// Report Errors
+	if (Source == nullptr && bHasCPUAccessError)
 	{
 		FNiagaraDataInterfaceError CPUAccessNotAllowedError(FText::Format(LOCTEXT("CPUAccessNotAllowedError", "This mesh needs CPU access in order to be used properly.({0})"), FText::FromString(DefaultMesh->GetName())),
 			LOCTEXT("CPUAccessNotAllowedErrorSummary", "CPU access error"),
@@ -890,6 +898,16 @@ TArray<FNiagaraDataInterfaceError> UNiagaraDataInterfaceSkeletalMesh::GetErrors(
 
 		Errors.Add(CPUAccessNotAllowedError);
 	}
+
+	if (Source == nullptr && bHasNoMeshAssignedError)
+	{
+		FNiagaraDataInterfaceError NoMeshAssignedError(LOCTEXT("NoMeshAssignedError", "This Data Interface must be assigned a skeletal mesh to operate."),
+			LOCTEXT("NoMeshAssignedErrorSummary", "No mesh assigned error"),
+			FNiagaraDataInterfaceFix());
+
+		Errors.Add(NoMeshAssignedError);
+	}
+
 	return Errors;
 }
 
diff --git a/Engine/Plugins/FX/Niagara/Source/Niagara/Private/NiagaraSystemInstance.cpp b/Engine/Plugins/FX/Niagara/Source/Niagara/Private/NiagaraSystemInstance.cpp
--- a/Engine/Plugins/FX/Niagara/Source/Niagara/Private/NiagaraSystemInstance.cpp
+++ b/Engine/Plugins/FX/Niagara/Source/Niagara/Private/NiagaraSystemInstance.cpp
@@ -49,6 +49,7 @@ FNiagaraSystemInstance::FNiagaraSystemInstance(UNiagaraComponent* InComponent)
 	, bPaused(false)
 	, RequestedExecutionState(ENiagaraExecutionState::Complete)
 	, ActualExecutionState(ENiagaraExecutionState::Complete)
+	, bDataInterfacesInitialized(false)
 {
 	SystemBounds.Init();
 }
@@ -829,7 +830,7 @@ void FNiagaraSystemInstance::InitDataInterfaces()
 
 	DataInterfaceInstanceData.SetNumUninitialized(InstanceDataSize);
 
-	bool bOk = true;
+	bDataInterfacesInitialized = true;
 	for (TPair<TWeakObjectPtr<UNiagaraDataInterface>, int32>& Pair : DataInterfaceInstanceDataOffsets)
 	{
 		if (UNiagaraDataInterface* Interface = Pair.Key.Get())
@@ -838,7 +839,7 @@ void FNiagaraSystemInstance::InitDataInterfaces()
 
 			//Ideally when we make the batching changes, we can keep the instance data in big single type blocks that can all be updated together with a single virtual call.
 			bool bResult = Pair.Key->InitPerInstanceData(&DataInterfaceInstanceData[Pair.Value], this);
-			bOk &= bResult;
+			bDataInterfacesInitialized &= bResult;
 			if (!bResult)
 			{
 				UE_LOG(LogNiagara, Error, TEXT("Error initializing data interface \"%s\" for system. %u | %s"), *Interface->GetPathName(), Component, *Component->GetAsset()->GetName());
@@ -847,11 +848,11 @@ void FNiagaraSystemInstance::InitDataInterfaces()
 		else
 		{
 			UE_LOG(LogNiagara, Error, TEXT("A data interface currently in use by an System has been destroyed."));
-			bOk = false;
+			bDataInterfacesInitialized = false;
 		}
 	}
 
-	if (!bOk && (!IsComplete() && !IsPendingSpawn()))
+	if (!bDataInterfacesInitialized && (!IsComplete() && !IsPendingSpawn()))
 	{
 		//Some error initializing the data interfaces so disable until we're explicitly reinitialized.
 		UE_LOG(LogNiagara, Error, TEXT("Error initializing data interfaces. Completing system. %u | %s"), Component, *Component->GetAsset()->GetName());
diff --git a/Engine/Plugins/FX/Niagara/Source/Niagara/Public/NiagaraComponent.h b/Engine/Plugins/FX/Niagara/Source/Niagara/Public/NiagaraComponent.h
--- a/Engine/Plugins/FX/Niagara/Source/Niagara/Public/NiagaraComponent.h
+++ b/Engine/Plugins/FX/Niagara/Source/Niagara/Public/NiagaraComponent.h
@@ -280,7 +280,7 @@ public:
 	virtual void PostEditChangeProperty(FPropertyChangedEvent& PropertyChangedEvent) override;
 
 	bool IsParameterValueOverriddenLocally(const FName& InParamName);
-	void SetParameterValueOverriddenLocally(const FNiagaraVariable& InParam, bool bInOverridden);
+	void SetParameterValueOverriddenLocally(const FNiagaraVariable& InParam, bool bInOverridden, bool bRequiresSystemInstanceReset);
 	
 	FOnSystemInstanceChanged& OnSystemInstanceChanged() { return OnSystemInstanceChangedDelegate; }
 
diff --git a/Engine/Plugins/FX/Niagara/Source/Niagara/Public/NiagaraSystemInstance.h b/Engine/Plugins/FX/Niagara/Source/Niagara/Public/NiagaraSystemInstance.h
--- a/Engine/Plugins/FX/Niagara/Source/Niagara/Public/NiagaraSystemInstance.h
+++ b/Engine/Plugins/FX/Niagara/Source/Niagara/Public/NiagaraSystemInstance.h
@@ -43,6 +43,7 @@ public:
 		None
 	};
 
+	FORCEINLINE bool GetAreDataInterfacesInitialized() const { return bDataInterfacesInitialized; }
 
 	/** Creates a new niagara System instance with the supplied component. */
 	explicit FNiagaraSystemInstance(UNiagaraComponent* InComponent);
@@ -199,6 +200,8 @@ public:
 	void Dump()const;
 
 private:
+	bool bDataInterfacesInitialized;
+
 	/** Builds the emitter simulations. */
 	void InitEmitters();
 
diff --git a/Engine/Plugins/FX/Niagara/Source/NiagaraEditor/Private/Customizations/NiagaraComponentDetails.cpp b/Engine/Plugins/FX/Niagara/Source/NiagaraEditor/Private/Customizations/NiagaraComponentDetails.cpp
--- a/Engine/Plugins/FX/Niagara/Source/NiagaraEditor/Private/Customizations/NiagaraComponentDetails.cpp
+++ b/Engine/Plugins/FX/Niagara/Source/NiagaraEditor/Private/Customizations/NiagaraComponentDetails.cpp
@@ -233,14 +233,14 @@ private:
 	{
 		check(Component.IsValid());
 		Component->GetOverrideParameters().OnParameterChange();
-		Component->SetParameterValueOverriddenLocally(Var, true);
+		Component->SetParameterValueOverriddenLocally(Var, true, false);
 	}
 
 	void OnDataInterfaceChanged(FNiagaraVariable Var)
 	{
 		Component->GetOverrideParameters().OnInterfaceChange();
 		check(Component.IsValid());
-		Component->SetParameterValueOverriddenLocally(Var, true);
+		Component->SetParameterValueOverriddenLocally(Var, true, true);
 	}
 
 	bool DoesParameterDifferFromDefault(const FNiagaraVariable& Var)
@@ -254,7 +254,7 @@ private:
 		check(Component.IsValid());
 		FScopedTransaction ScopedTransaction(LOCTEXT("ResetParameterValue", "Reset parameter value to system defaults."));
 		Component->Modify();
-		Component->SetParameterValueOverriddenLocally(Parameter, false);
+		Component->SetParameterValueOverriddenLocally(Parameter, false, false);
 		return FReply::Handled();
 	}
 
diff --git a/Engine/Plugins/FX/Niagara/Source/NiagaraEditor/Private/ViewModels/NiagaraSystemViewModel.cpp b/Engine/Plugins/FX/Niagara/Source/NiagaraEditor/Private/ViewModels/NiagaraSystemViewModel.cpp
--- a/Engine/Plugins/FX/Niagara/Source/NiagaraEditor/Private/ViewModels/NiagaraSystemViewModel.cpp
+++ b/Engine/Plugins/FX/Niagara/Source/NiagaraEditor/Private/ViewModels/NiagaraSystemViewModel.cpp
@@ -1420,7 +1420,7 @@ void FNiagaraSystemViewModel::SequencerDataChanged(EMovieSceneDataChangeType Dat
 
 void FNiagaraSystemViewModel::SequencerTimeChanged()
 {
-	if (!PreviewComponent)
+	if (!PreviewComponent || !PreviewComponent->GetSystemInstance() || !PreviewComponent->GetSystemInstance()->GetAreDataInterfacesInitialized())
 	{
 		return;
 	}
