Niagara - Clean up the system instance pointer on the component when it's unregistered from the world to prevent crashes when components are reused when transitioning levels.

diff --git a/Engine/Plugins/FX/Niagara/Source/Niagara/Private/NiagaraComponent.cpp b/Engine/Plugins/FX/Niagara/Source/Niagara/Private/NiagaraComponent.cpp
--- a/Engine/Plugins/FX/Niagara/Source/Niagara/Private/NiagaraComponent.cpp
+++ b/Engine/Plugins/FX/Niagara/Source/Niagara/Private/NiagaraComponent.cpp
@@ -738,6 +738,7 @@ void UNiagaraComponent::OnUnregister()
 	if (SystemInstance)
 	{
 		SystemInstance->Deactivate(true);
+		SystemInstance = nullptr;
 	}
 }
 
