Niagara - Update styling for stack drop targets.

diff --git a/Engine/Plugins/FX/Niagara/Source/NiagaraEditorWidgets/Private/NiagaraEditorWidgetsStyle.cpp b/Engine/Plugins/FX/Niagara/Source/NiagaraEditorWidgets/Private/NiagaraEditorWidgetsStyle.cpp
--- a/Engine/Plugins/FX/Niagara/Source/NiagaraEditorWidgets/Private/NiagaraEditorWidgetsStyle.cpp
+++ b/Engine/Plugins/FX/Niagara/Source/NiagaraEditorWidgets/Private/NiagaraEditorWidgetsStyle.cpp
@@ -122,6 +122,11 @@ TSharedRef< FSlateStyleSet > FNiagaraEditorWidgetsStyle::Create()
 	Style->Set("NiagaraEditor.Stack.IconColor.Particle", FLinearColor(FColor(131, 228, 9)));
 	Style->Set("NiagaraEditor.Stack.IconColor.Render", FLinearColor(FColor(230, 102, 102)));
 
+ 	Style->Set("NiagaraEditor.Stack.DropTarget.BackgroundColor", FLinearColor(1.0f, 1.0f, 1.0f, 0.25f));
+ 	Style->Set("NiagaraEditor.Stack.DropTarget.BackgroundColorHover", FLinearColor(1.0f, 1.0f, 1.0f, 0.1f));
+	Style->Set("NiagaraEditor.Stack.DropTarget.BorderVertical", new IMAGE_PLUGIN_BRUSH("Icons/StackDropTargetBorder_Vertical", FVector2D(2, 8), FLinearColor::White, ESlateBrushTileType::Vertical));
+	Style->Set("NiagaraEditor.Stack.DropTarget.BorderHorizontal", new IMAGE_PLUGIN_BRUSH("Icons/StackDropTargetBorder_Horizontal", FVector2D(8, 2), FLinearColor::White, ESlateBrushTileType::Horizontal));
+
 	Style->Set("NiagaraEditor.Stack.GoToSourceIcon", new IMAGE_CORE_BRUSH("Common/GoToSource", Icon30x30, FLinearColor::White));
 	Style->Set("NiagaraEditor.Stack.ParametersIcon", new IMAGE_PLUGIN_BRUSH("Icons/SystemParams", Icon12x12, FLinearColor::White));
 	Style->Set("NiagaraEditor.Stack.SpawnIcon", new IMAGE_PLUGIN_BRUSH("Icons/Spawn", Icon12x12, FLinearColor::White));
diff --git a/Engine/Plugins/FX/Niagara/Source/NiagaraEditorWidgets/Private/Stack/SNiagaraStackFunctionInputValue.cpp b/Engine/Plugins/FX/Niagara/Source/NiagaraEditorWidgets/Private/Stack/SNiagaraStackFunctionInputValue.cpp
--- a/Engine/Plugins/FX/Niagara/Source/NiagaraEditorWidgets/Private/Stack/SNiagaraStackFunctionInputValue.cpp
+++ b/Engine/Plugins/FX/Niagara/Source/NiagaraEditorWidgets/Private/Stack/SNiagaraStackFunctionInputValue.cpp
@@ -41,6 +41,10 @@ void SNiagaraStackFunctionInputValue::Construct(const FArguments& InArgs, UNiaga
 		SNew(SDropTarget)
 		.OnAllowDrop(this, &SNiagaraStackFunctionInputValue::OnFunctionInputAllowDrop)
 		.OnDrop(this, &SNiagaraStackFunctionInputValue::OnFunctionInputDrop)
+		.HorizontalImage(FNiagaraEditorWidgetsStyle::Get().GetBrush("NiagaraEditor.Stack.DropTarget.BorderHorizontal"))
+		.VerticalImage(FNiagaraEditorWidgetsStyle::Get().GetBrush("NiagaraEditor.Stack.DropTarget.BorderVertical"))
+		.BackgroundColor(FNiagaraEditorWidgetsStyle::Get().GetColor("NiagaraEditor.Stack.DropTarget.BackgroundColor"))
+		.BackgroundColorHover(FNiagaraEditorWidgetsStyle::Get().GetColor("NiagaraEditor.Stack.DropTarget.BackgroundColorHover"))
 		.Content()
 		[
 			// Values
diff --git a/Engine/Plugins/FX/Niagara/Source/NiagaraEditorWidgets/Private/Stack/SNiagaraStackModuleItem.cpp b/Engine/Plugins/FX/Niagara/Source/NiagaraEditorWidgets/Private/Stack/SNiagaraStackModuleItem.cpp
--- a/Engine/Plugins/FX/Niagara/Source/NiagaraEditorWidgets/Private/Stack/SNiagaraStackModuleItem.cpp
+++ b/Engine/Plugins/FX/Niagara/Source/NiagaraEditorWidgets/Private/Stack/SNiagaraStackModuleItem.cpp
@@ -37,6 +37,10 @@ void SNiagaraStackModuleItem::Construct(const FArguments& InArgs, UNiagaraStackM
 		SNew(SDropTarget)
 		.OnAllowDrop(this, &SNiagaraStackModuleItem::OnModuleItemAllowDrop)
 		.OnDrop(this, &SNiagaraStackModuleItem::OnModuleItemDrop)
+		.HorizontalImage(FNiagaraEditorWidgetsStyle::Get().GetBrush("NiagaraEditor.Stack.DropTarget.BorderHorizontal"))
+		.VerticalImage(FNiagaraEditorWidgetsStyle::Get().GetBrush("NiagaraEditor.Stack.DropTarget.BorderVertical"))
+		.BackgroundColor(FNiagaraEditorWidgetsStyle::Get().GetColor("NiagaraEditor.Stack.DropTarget.BackgroundColor"))
+		.BackgroundColorHover(FNiagaraEditorWidgetsStyle::Get().GetColor("NiagaraEditor.Stack.DropTarget.BackgroundColorHover"))
 		.Content()
 		[
 			SNew(SHorizontalBox)
diff --git a/Engine/Plugins/FX/Niagara/Source/NiagaraEditorWidgets/Private/Stack/SNiagaraStackSpacer.cpp b/Engine/Plugins/FX/Niagara/Source/NiagaraEditorWidgets/Private/Stack/SNiagaraStackSpacer.cpp
--- a/Engine/Plugins/FX/Niagara/Source/NiagaraEditorWidgets/Private/Stack/SNiagaraStackSpacer.cpp
+++ b/Engine/Plugins/FX/Niagara/Source/NiagaraEditorWidgets/Private/Stack/SNiagaraStackSpacer.cpp
@@ -7,6 +7,7 @@
 #include "NiagaraConstants.h"
 #include "ViewModels/Stack/NiagaraStackGraphUtilities.h"
 #include "ViewModels/Stack/NiagaraStackScriptItemGroup.h"
+#include "NiagaraEditorWidgetsStyle.h"
 
 void SNiagaraStackSpacer::Construct(const FArguments& InArgs, UNiagaraStackSpacer& InStackSpacer) 
 {
@@ -19,6 +20,10 @@ void SNiagaraStackSpacer::Construct(const FArguments& InArgs, UNiagaraStackSpace
 		SNew(SDropTarget)
 		.OnAllowDrop(this, &SNiagaraStackSpacer::OnSStackSpacerAllowDrop)
 		.OnDrop(this, &SNiagaraStackSpacer::OnSStackSpacerDrop)
+		.HorizontalImage(FNiagaraEditorWidgetsStyle::Get().GetBrush("NiagaraEditor.Stack.DropTarget.BorderHorizontal"))
+		.VerticalImage(FNiagaraEditorWidgetsStyle::Get().GetBrush("NiagaraEditor.Stack.DropTarget.BorderVertical"))
+		.BackgroundColor(FNiagaraEditorWidgetsStyle::Get().GetColor("NiagaraEditor.Stack.DropTarget.BackgroundColor"))
+		.BackgroundColorHover(FNiagaraEditorWidgetsStyle::Get().GetColor("NiagaraEditor.Stack.DropTarget.BackgroundColorHover"))
 		.Content()
 		[
 			SNew(SBox)
