Niagara - Fix assert when double clicking selected entries in the new emitter and system dialogs.

diff --git a/Engine/Plugins/FX/Niagara/Source/NiagaraEditor/Private/SNewEmitterDialog.cpp b/Engine/Plugins/FX/Niagara/Source/NiagaraEditor/Private/SNewEmitterDialog.cpp
--- a/Engine/Plugins/FX/Niagara/Source/NiagaraEditor/Private/SNewEmitterDialog.cpp
+++ b/Engine/Plugins/FX/Niagara/Source/NiagaraEditor/Private/SNewEmitterDialog.cpp
@@ -64,18 +64,29 @@ TOptional<FAssetData> SNewEmitterDialog::GetSelectedEmitterAsset()
 	return TOptional<FAssetData>();
 }
 
-TArray<FAssetData> SNewEmitterDialog::GetSelectedEmitterTemplateAssets()
+void SNewEmitterDialog::GetSelectedEmitterTemplateAssets(TArray<FAssetData>& OutSelectedAssets)
 {
-	return TemplateAssetPicker->GetSelectedAssets();
+	OutSelectedAssets.Append(TemplateAssetPicker->GetSelectedAssets());
+	if (ActivatedTemplateAsset.IsValid())
+	{
+		OutSelectedAssets.AddUnique(ActivatedTemplateAsset);
+	}
 }
 
-TArray<FAssetData> SNewEmitterDialog::GetSelectedProjectEmiterAssets()
+void SNewEmitterDialog::GetSelectedProjectEmiterAssets(TArray<FAssetData>& OutSelectedAssets)
 {
-	return GetSelectedEmitterAssetsFromPicker.Execute();
+	OutSelectedAssets.Append(GetSelectedEmitterAssetsFromPicker.Execute());
+	if (ActivatedProjectAsset.IsValid())
+	{
+		OutSelectedAssets.AddUnique(ActivatedProjectAsset);
+	}
 }
 
-void SNewEmitterDialog::OnTemplateAssetActivated(const FAssetData& ActivatedTemplateAsset)
+void SNewEmitterDialog::OnTemplateAssetActivated(const FAssetData& InActivatedTemplateAsset)
 {
+	// Input handling issues with the list view widget can allow items to be activated but not added to the selection so cache this here
+	// so it can be included in the selection set.
+	ActivatedTemplateAsset = InActivatedTemplateAsset;
 	ConfirmSelection();
 }
 
@@ -83,6 +94,9 @@ void SNewEmitterDialog::OnEmitterAssetsActivated(const TArray<FAssetData>& Activ
 {
 	if ((ActivationMethod == EAssetTypeActivationMethod::DoubleClicked || ActivationMethod == EAssetTypeActivationMethod::Opened) && ActivatedAssets.Num() == 1)
 	{
+		// Input handling issues with the list view widget can allow items to be activated but not added to the selection so cache this here
+		// so it can be included in the selection set.
+		ActivatedProjectAsset = ActivatedAssets[0];
 		ConfirmSelection();
 	}
 }
diff --git a/Engine/Plugins/FX/Niagara/Source/NiagaraEditor/Private/SNewEmitterDialog.h b/Engine/Plugins/FX/Niagara/Source/NiagaraEditor/Private/SNewEmitterDialog.h
--- a/Engine/Plugins/FX/Niagara/Source/NiagaraEditor/Private/SNewEmitterDialog.h
+++ b/Engine/Plugins/FX/Niagara/Source/NiagaraEditor/Private/SNewEmitterDialog.h
@@ -24,11 +24,11 @@ public:
 	TOptional<FAssetData> GetSelectedEmitterAsset();
 
 private:
-	TArray<FAssetData> GetSelectedEmitterTemplateAssets();
+	void GetSelectedEmitterTemplateAssets(TArray<FAssetData>& OutSelectedAssets);
 
-	TArray<FAssetData> GetSelectedProjectEmiterAssets();
+	void GetSelectedProjectEmiterAssets(TArray<FAssetData>& OutSelectedAssets);
 
-	void OnTemplateAssetActivated(const FAssetData& ActivatedTemplateAsset);
+	void OnTemplateAssetActivated(const FAssetData& InActivatedTemplateAsset);
 
 	void OnEmitterAssetsActivated(const TArray<FAssetData>& ActivatedAssets, EAssetTypeActivationMethod::Type ActivationMethod);
 
@@ -36,4 +36,8 @@ private:
 	TSharedPtr<SNiagaraTemplateAssetPicker> TemplateAssetPicker;
 
 	FGetCurrentSelectionDelegate GetSelectedEmitterAssetsFromPicker;
+
+	FAssetData ActivatedTemplateAsset;
+
+	FAssetData ActivatedProjectAsset;
 };
\ No newline at end of file
diff --git a/Engine/Plugins/FX/Niagara/Source/NiagaraEditor/Private/SNewSystemDialog.cpp b/Engine/Plugins/FX/Niagara/Source/NiagaraEditor/Private/SNewSystemDialog.cpp
--- a/Engine/Plugins/FX/Niagara/Source/NiagaraEditor/Private/SNewSystemDialog.cpp
+++ b/Engine/Plugins/FX/Niagara/Source/NiagaraEditor/Private/SNewSystemDialog.cpp
@@ -156,23 +156,34 @@ TArray<FAssetData> SNewSystemDialog::GetSelectedEmitterAssets() const
 	return ConfirmedSelectedEmitterAssets;
 }
 
-TArray<FAssetData> SNewSystemDialog::GetSelectedSystemTemplateAssets()
+void SNewSystemDialog::GetSelectedSystemTemplateAssets(TArray<FAssetData>& OutSelectedAssets)
 {
-	return TemplateAssetPicker->GetSelectedAssets();
+	OutSelectedAssets.Append(TemplateAssetPicker->GetSelectedAssets());
+	if (ActivatedTemplateSystemAsset.IsValid())
+	{
+		OutSelectedAssets.AddUnique(ActivatedTemplateSystemAsset);
+	}
 }
 
-TArray<FAssetData> SNewSystemDialog::GetSelectedProjectSystemAssets()
+void SNewSystemDialog::GetSelectedProjectSystemAssets(TArray<FAssetData>& OutSelectedAssets)
 {
-	return GetSelectedSystemAssetsFromPicker.Execute();
+	OutSelectedAssets.Append(GetSelectedSystemAssetsFromPicker.Execute());
+	if (ActivatedProjectSystemAsset.IsValid())
+	{
+		OutSelectedAssets.AddUnique(ActivatedProjectSystemAsset);
+	}
 }
 
-TArray<FAssetData> SNewSystemDialog::GetSelectedProjectEmiterAssets()
+void SNewSystemDialog::GetSelectedProjectEmiterAssets(TArray<FAssetData>& OutSelectedAssets)
 {
-	return SelectedEmitterAssets;
+	OutSelectedAssets.Append(SelectedEmitterAssets);
 }
 
 void SNewSystemDialog::OnTemplateAssetActivated(const FAssetData& ActivatedTemplateAsset)
 {
+	// Input handling issues with the list view widget can allow items to be activated but not added to the selection so cache this here
+	// so it can be included in the selection set.
+	ActivatedTemplateSystemAsset = ActivatedTemplateAsset;
 	ConfirmSelection();
 }
 
@@ -180,6 +191,9 @@ void SNewSystemDialog::OnSystemAssetsActivated(const TArray<FAssetData>& Activat
 {
 	if ((ActivationMethod == EAssetTypeActivationMethod::DoubleClicked || ActivationMethod == EAssetTypeActivationMethod::Opened) && ActivatedAssets.Num() == 1)
 	{
+		// Input handling issues with the list view widget can allow items to be activated but not added to the selection so cache this here
+		// so it can be included in the selection set.
+		ActivatedProjectSystemAsset = ActivatedAssets[0];
 		ConfirmSelection();
 	}
 }
diff --git a/Engine/Plugins/FX/Niagara/Source/NiagaraEditor/Private/SNewSystemDialog.h b/Engine/Plugins/FX/Niagara/Source/NiagaraEditor/Private/SNewSystemDialog.h
--- a/Engine/Plugins/FX/Niagara/Source/NiagaraEditor/Private/SNewSystemDialog.h
+++ b/Engine/Plugins/FX/Niagara/Source/NiagaraEditor/Private/SNewSystemDialog.h
@@ -27,11 +27,11 @@ public:
 	TArray<FAssetData> GetSelectedEmitterAssets() const;
 
 private:
-	TArray<FAssetData> GetSelectedSystemTemplateAssets();
+	void GetSelectedSystemTemplateAssets(TArray<FAssetData>& OutSelectedAssets);
 
-	TArray<FAssetData> GetSelectedProjectSystemAssets();
+	void GetSelectedProjectSystemAssets(TArray<FAssetData>& OutSelectedAssets);
 
-	TArray<FAssetData> GetSelectedProjectEmiterAssets();
+	void GetSelectedProjectEmiterAssets(TArray<FAssetData>& OutSelectedAssets);
 
 	void OnTemplateAssetActivated(const FAssetData& ActivatedTemplateAsset);
 
@@ -52,6 +52,10 @@ private:
 
 	FGetCurrentSelectionDelegate GetSelectedEmitterAssetsFromPicker;
 
+	FAssetData ActivatedTemplateSystemAsset;
+
+	FAssetData ActivatedProjectSystemAsset;
+
 	TArray<FAssetData> SelectedEmitterAssets;
 
 	TArray<TSharedPtr<SWidget>> SelectedEmitterAssetWidgets;
diff --git a/Engine/Plugins/FX/Niagara/Source/NiagaraEditor/Private/SNiagaraNewAssetDialog.cpp b/Engine/Plugins/FX/Niagara/Source/NiagaraEditor/Private/SNiagaraNewAssetDialog.cpp
--- a/Engine/Plugins/FX/Niagara/Source/NiagaraEditor/Private/SNiagaraNewAssetDialog.cpp
+++ b/Engine/Plugins/FX/Niagara/Source/NiagaraEditor/Private/SNiagaraNewAssetDialog.cpp
@@ -194,7 +194,7 @@ void SNiagaraNewAssetDialog::ConfirmSelection()
 	const FNiagaraNewAssetDialogOption& SelectedOption = Options[SelectedOptionIndex];
 	if (SelectedOption.OnGetSelectedAssetsFromPicker.IsBound())
 	{
-		SelectedAssets.Append(SelectedOption.OnGetSelectedAssetsFromPicker.Execute());
+		SelectedOption.OnGetSelectedAssetsFromPicker.Execute(SelectedAssets);
 		ensureMsgf(SelectedAssets.Num() > 0, TEXT("No assets selected when dialog was confirmed."));
 	}
 	bUserConfirmedSelection = true;
@@ -246,7 +246,16 @@ EVisibility SNiagaraNewAssetDialog::GetAssetPickerVisibility(int32 OptionIndex)
 bool SNiagaraNewAssetDialog::IsOkButtonEnabled() const
 {
 	const FNiagaraNewAssetDialogOption& SelectedOption = Options[SelectedOptionIndex];
-	return SelectedOption.OnGetSelectedAssetsFromPicker.IsBound() == false || SelectedOption.OnGetSelectedAssetsFromPicker.Execute().Num() > 0;
+	if (SelectedOption.OnGetSelectedAssetsFromPicker.IsBound())
+	{
+		TArray<FAssetData> TempSelectedAssets;
+		SelectedOption.OnGetSelectedAssetsFromPicker.Execute(TempSelectedAssets);
+		return TempSelectedAssets.Num() != 0;
+	}
+	else
+	{
+		return true;
+	}
 }
 
 FReply SNiagaraNewAssetDialog::OnOkButtonClicked()
diff --git a/Engine/Plugins/FX/Niagara/Source/NiagaraEditor/Private/SNiagaraNewAssetDialog.h b/Engine/Plugins/FX/Niagara/Source/NiagaraEditor/Private/SNiagaraNewAssetDialog.h
--- a/Engine/Plugins/FX/Niagara/Source/NiagaraEditor/Private/SNiagaraNewAssetDialog.h
+++ b/Engine/Plugins/FX/Niagara/Source/NiagaraEditor/Private/SNiagaraNewAssetDialog.h
@@ -13,7 +13,7 @@ class SNiagaraTemplateAssetPicker;
 class SNiagaraNewAssetDialog : public SWindow
 {
 public:
-	DECLARE_DELEGATE_RetVal(TArray<FAssetData>, FOnGetSelectedAssetsFromPicker);
+	DECLARE_DELEGATE_OneParam(FOnGetSelectedAssetsFromPicker, TArray<FAssetData>& /* OutSelectedAssets */);
 
 public:
 	class FNiagaraNewAssetDialogOption
