Niagara: fix for using texture references

diff --git a/Engine/Plugins/FX/Niagara/Source/Niagara/Private/NiagaraDataInterfaceTexture.cpp b/Engine/Plugins/FX/Niagara/Source/Niagara/Private/NiagaraDataInterfaceTexture.cpp
--- a/Engine/Plugins/FX/Niagara/Source/Niagara/Private/NiagaraDataInterfaceTexture.cpp
+++ b/Engine/Plugins/FX/Niagara/Source/Niagara/Private/NiagaraDataInterfaceTexture.cpp
@@ -388,7 +388,7 @@ struct FNiagaraDataInterfaceParametersCS_Texture : public FNiagaraDataInterfaceP
 			SetShaderValue(RHICmdList, ComputeShaderRHI, Dimensions, TexDims);
 			return;
 		}
-		FTextureRHIParamRef TextureRHI = Texture->TextureReference.TextureReferenceRHI;
+		FTextureRHIParamRef TextureRHI = Texture->TextureReference.TextureReferenceRHI->GetReferencedTexture();
 		SetTextureParameter(
 			RHICmdList,
 			ComputeShaderRHI,
