From Dev-Niagara..
Fixed issue where we were incorrectly returning a float instead of an int for InputDataInt.
Tweaked some conversion warnings in Niagara from float to int. Updating the Nigara custom version to force shader regeneration.

#tests autotests passed

diff --git a/Engine/Plugins/FX/Niagara/Source/NiagaraCore/Private/NiagaraCustomVersion.cpp b/Engine/Plugins/FX/Niagara/Source/NiagaraCore/Private/NiagaraCustomVersion.cpp
--- a/Engine/Plugins/FX/Niagara/Source/NiagaraCore/Private/NiagaraCustomVersion.cpp
+++ b/Engine/Plugins/FX/Niagara/Source/NiagaraCore/Private/NiagaraCustomVersion.cpp
@@ -8,4 +8,4 @@ const FGuid FNiagaraCustomVersion::GUID(0xFCF57AFA, 0x50764283, 0xB9A9E658, 0xFF
 // Register the custom version with core
 FCustomVersionRegistration GRegisterNiagaraCustomVersion(FNiagaraCustomVersion::GUID, FNiagaraCustomVersion::LatestVersion, TEXT("NiagaraVer"));
 
-const FGuid FNiagaraCustomVersion::LatestScriptCompileVersion(0xF109EF07, 0xE00C47AC, 0xB1DF217F, 0xAC0F1789);
\ No newline at end of file
+const FGuid FNiagaraCustomVersion::LatestScriptCompileVersion(0x044AB7CE, 0x10B040DD, 0x9C859C8B, 0xC0E47E09);
\ No newline at end of file
diff --git a/Engine/Shaders/Private/Common.ush b/Engine/Shaders/Private/Common.ush
--- a/Engine/Shaders/Private/Common.ush
+++ b/Engine/Shaders/Private/Common.ush
@@ -1338,7 +1338,9 @@ float AntialiasedTextureMask( Texture2D Tex, SamplerState Sampler, float2 UV, fl
 	return Result;
 }
 
-float Noise3D_Multiplexer(int Function, float3 Position, int Quality, bool bTiling, uint RepeatSize)
+// While RepeatSize is a float here, the expectation is that it would be largely integer values coming in from the UI. The downstream logic assumes
+// floats for all called functions (NoiseTileWrap) and this prevents any float-to-int conversion errors from automatic type conversion.
+float Noise3D_Multiplexer(int Function, float3 Position, int Quality, bool bTiling, float RepeatSize)
 {
 	// verified, HLSL compiled out the switch if Function is a constant
 	switch(Function)
diff --git a/Engine/Shaders/Private/NiagaraEmitterInstanceShader.usf b/Engine/Shaders/Private/NiagaraEmitterInstanceShader.usf
--- a/Engine/Shaders/Private/NiagaraEmitterInstanceShader.usf
+++ b/Engine/Shaders/Private/NiagaraEmitterInstanceShader.usf
@@ -245,7 +245,7 @@ NiagaraSimulationShader.usf:
 		return InputFloat[RegisterIdx*ComponentBufferSizeRead + InstanceIdx];
 	}
 
-	float InputDataInt(int DataSetIndex, int RegisterIdx, int InstanceIdx)
+	int InputDataInt(int DataSetIndex, int RegisterIdx, int InstanceIdx)
 	{
 		return InputInt[RegisterIdx*ComponentBufferSizeRead + InstanceIdx];
 	}
diff --git a/Engine/Shaders/Private/Random.ush b/Engine/Shaders/Private/Random.ush
--- a/Engine/Shaders/Private/Random.ush
+++ b/Engine/Shaders/Private/Random.ush
@@ -783,7 +783,7 @@ float2 GetPerlinNoiseGradientTextureAt(float2 v)
 
 float3 GetPerlinNoiseGradientTextureAt(float3 v)
 {
-	const float2 ZShear = int2(17, 89);
+	const float2 ZShear = float2(17.0f, 89.0f);
 
 	float2 OffsetA = v.z * ZShear;
 	float2 TexA = (v.xy + OffsetA + 0.5f) / 128.0f;
