From Dev-Niagara...
No longer using StaticFind in the middle of serializing shader maps as that is a no-no. Seems like this pathway changed on the last main integration. Now we just use the Type registry's list of data interface classes and resolve the default object from there. This should be properly initialized.

diff --git a/Engine/Plugins/FX/Niagara/Source/Niagara/Private/NiagaraModule.cpp b/Engine/Plugins/FX/Niagara/Source/Niagara/Private/NiagaraModule.cpp
--- a/Engine/Plugins/FX/Niagara/Source/Niagara/Private/NiagaraModule.cpp
+++ b/Engine/Plugins/FX/Niagara/Source/Niagara/Private/NiagaraModule.cpp
@@ -17,6 +17,8 @@
 #include "NiagaraRibbonRendererProperties.h"
 #include "NiagaraRenderer.h"
 #include "Misc/CoreDelegates.h"
+#include "NiagaraShaderModule.h"
+#include "UObject/CoreRedirects.h"
 
 IMPLEMENT_MODULE(INiagaraModule, Niagara);
 
@@ -229,6 +231,13 @@ void INiagaraModule::StartupModule()
 	UNiagaraRibbonRendererProperties::InitCDOPropertiesAfterModuleStartup();
 	UNiagaraMeshRendererProperties::InitCDOPropertiesAfterModuleStartup();
 
+	// Register the data interface CDO finder with teh shader module..
+	INiagaraShaderModule& NiagaraShaderModule = FModuleManager::LoadModuleChecked<INiagaraShaderModule>("NiagaraShader");
+	NiagaraShaderModule.SetOnRequestDefaultDataInterfaceHandler(INiagaraShaderModule::FOnRequestDefaultDataInterface::CreateLambda([](const FString& DIClassName) -> UNiagaraDataInterfaceBase*
+	{
+		return FNiagaraTypeRegistry::GetDefaultDataInterfaceByName(DIClassName);
+	}));
+
 }
 
 void INiagaraModule::ShutdownRenderingResources()
@@ -246,6 +255,10 @@ void INiagaraModule::ShutdownModule()
 		Pair.Value = nullptr;
 	}
 
+	// Clear out the handler when shutting down..
+	INiagaraShaderModule& NiagaraShaderModule = FModuleManager::LoadModuleChecked<INiagaraShaderModule>("NiagaraShader");
+	NiagaraShaderModule.ResetOnRequestDefaultDataInterfaceHandler();
+
 	CVarDetailLevel.AsVariable()->SetOnChangedCallback(FConsoleVariableDelegate());
 	ShutdownRenderingResources();
 }
@@ -765,6 +778,41 @@ bool FNiagaraVariable::GetValue() const
 
 //////////////////////////////////////////////////////////////////////////
 
+UNiagaraDataInterfaceBase* FNiagaraTypeRegistry::GetDefaultDataInterfaceByName(const FString& DIClassName)
+{
+	UClass* DIClass = nullptr;
+	for (const FNiagaraTypeDefinition& Def : RegisteredTypes)
+	{
+		if (Def.IsDataInterface())
+		{
+			UClass* FoundDIClass = Def.GetClass();
+			if (FoundDIClass && (FoundDIClass->GetName() == DIClassName || FoundDIClass->GetFullName() == DIClassName))
+			{
+				DIClass = FoundDIClass;
+				break;
+			}
+		}
+	}
+	// Consider the possibility of a redirector pointing to a new location..
+	if (DIClass == nullptr)
+	{
+		FCoreRedirectObjectName OldObjName;
+		OldObjName.ObjectName = *DIClassName;
+		FCoreRedirectObjectName NewObjName = FCoreRedirects::GetRedirectedName(ECoreRedirectFlags::Type_Class, OldObjName);
+		if (NewObjName.IsValid() && OldObjName != NewObjName)
+		{
+			return GetDefaultDataInterfaceByName(NewObjName.ObjectName.ToString());
+		}
+	}
+
+	if (DIClass)
+	{
+		return CastChecked<UNiagaraDataInterfaceBase>(DIClass->GetDefaultObject(false)); // We wouldn't be registered if the CDO had not already been created...
+	}
+
+	return nullptr;
+}
+
 
 FDelegateHandle INiagaraModule::SetOnProcessShaderCompilationQueue(FOnProcessQueue InOnProcessQueue)
 {
diff --git a/Engine/Plugins/FX/Niagara/Source/Niagara/Public/NiagaraTypes.h b/Engine/Plugins/FX/Niagara/Source/Niagara/Public/NiagaraTypes.h
--- a/Engine/Plugins/FX/Niagara/Source/Niagara/Public/NiagaraTypes.h
+++ b/Engine/Plugins/FX/Niagara/Source/Niagara/Public/NiagaraTypes.h
@@ -8,6 +8,7 @@
 #include "Templates/SharedPointer.h"
 #include "NiagaraTypes.generated.h"
 
+class UNiagaraDataInterfaceBase;
 DECLARE_LOG_CATEGORY_EXTERN(LogNiagara, Log, Verbose);
 
 // basic type struct definitions
@@ -618,6 +619,8 @@ public:
 		return RegisteredNumericTypes;
 	}
 
+	static UNiagaraDataInterfaceBase* GetDefaultDataInterfaceByName(const FString& DIClassName);
+
 	static void ClearUserDefinedRegistry()
 	{
 		for (const FNiagaraTypeDefinition& Def : RegisteredUserDefinedTypes)
diff --git a/Engine/Plugins/FX/Niagara/Source/NiagaraShader/Private/NiagaraShader.cpp b/Engine/Plugins/FX/Niagara/Source/NiagaraShader/Private/NiagaraShader.cpp
--- a/Engine/Plugins/FX/Niagara/Source/NiagaraShader/Private/NiagaraShader.cpp
+++ b/Engine/Plugins/FX/Niagara/Source/NiagaraShader/Private/NiagaraShader.cpp
@@ -20,6 +20,7 @@
 
 #include "NiagaraDataInterfaceBase.h"
 #include "UObject/UObjectGlobals.h"
+#include "NiagaraShaderModule.h"
 
 IMPLEMENT_SHADER_TYPE(,FNiagaraShader, TEXT("/Engine/Private/NiagaraEmitterInstanceShader.usf"),TEXT("SimulateMain"), SF_Compute)
 
@@ -1335,17 +1336,15 @@ void FNiagaraDataInterfaceParamRef::ConstructParameters()
 
 void FNiagaraDataInterfaceParamRef::InitDIClass()
 {
-	DIClass = Cast<UClass>(StaticFindObject(UClass::StaticClass(), ANY_PACKAGE, *ParameterInfo.DIClassName, true));
-	if (DIClass == nullptr)
+	INiagaraShaderModule* Module = INiagaraShaderModule::Get();
+	check(Module != nullptr);
+	
+	// Getting the base here in hopes that in the future we would just reference it directly rather than going through
+	// the class intermediary.
+	UNiagaraDataInterfaceBase* Base = Module->RequestDefaultDataInterface(ParameterInfo.DIClassName);
+	if (Base)
 	{
-		FCoreRedirectObjectName OldObjName;
-		OldObjName.ObjectName = *ParameterInfo.DIClassName;
-		FCoreRedirectObjectName NewObjName = FCoreRedirects::GetRedirectedName(ECoreRedirectFlags::Type_Class, OldObjName);
-		if (NewObjName.IsValid())
-		{
-			DIClass = Cast<UClass>(StaticFindObject(UClass::StaticClass(), ANY_PACKAGE, *NewObjName.ObjectName.ToString(), true));
-		}
-
+		DIClass = Base->GetClass();
 	}
 	ensureMsgf(DIClass, TEXT("Failed to load class for FNiagaraDataInterfaceParamRef. %s"), *ParameterInfo.DIClassName);
 }
diff --git a/Engine/Plugins/FX/Niagara/Source/NiagaraShader/Private/NiagaraShaderModule.cpp b/Engine/Plugins/FX/Niagara/Source/NiagaraShader/Private/NiagaraShaderModule.cpp
--- a/Engine/Plugins/FX/Niagara/Source/NiagaraShader/Private/NiagaraShaderModule.cpp
+++ b/Engine/Plugins/FX/Niagara/Source/NiagaraShader/Private/NiagaraShaderModule.cpp
@@ -5,9 +5,7 @@
 
 IMPLEMENT_MODULE(INiagaraShaderModule, NiagaraShader);
 
-
-
-
+INiagaraShaderModule* INiagaraShaderModule::Singleton(nullptr);
 
 FDelegateHandle INiagaraShaderModule::SetOnProcessShaderCompilationQueue(FOnProcessQueue InOnProcessQueue)
 {
@@ -27,3 +25,21 @@ void INiagaraShaderModule::ProcessShaderCompilationQueue()
 	checkf(OnProcessQueue.IsBound(), TEXT("Can not process shader queue.  Delegate was never set."));
 	return OnProcessQueue.Execute();
 }
+
+FDelegateHandle INiagaraShaderModule::SetOnRequestDefaultDataInterfaceHandler(FOnRequestDefaultDataInterface InHandler)
+{
+	checkf(OnRequestDefaultDataInterface.IsBound() == false, TEXT("Shader OnRequestDefaultDataInterface delegate already set."));
+	OnRequestDefaultDataInterface = InHandler;
+	return OnRequestDefaultDataInterface.GetHandle();
+}
+
+void  INiagaraShaderModule::ResetOnRequestDefaultDataInterfaceHandler()
+{
+	OnRequestDefaultDataInterface.Unbind();
+}
+
+UNiagaraDataInterfaceBase* INiagaraShaderModule::RequestDefaultDataInterface(const FString& DIClassName)
+{
+	checkf(OnRequestDefaultDataInterface.IsBound(), TEXT("Can not invoke OnRequestDefaultDataInterface.  Delegate was never set."));
+	return OnRequestDefaultDataInterface.Execute(DIClassName);
+}
\ No newline at end of file
diff --git a/Engine/Plugins/FX/Niagara/Source/NiagaraShader/Public/NiagaraShaderModule.h b/Engine/Plugins/FX/Niagara/Source/NiagaraShader/Public/NiagaraShaderModule.h
--- a/Engine/Plugins/FX/Niagara/Source/NiagaraShader/Public/NiagaraShaderModule.h
+++ b/Engine/Plugins/FX/Niagara/Source/NiagaraShader/Public/NiagaraShaderModule.h
@@ -5,6 +5,8 @@
 #include "CoreMinimal.h"
 #include "Modules/ModuleInterface.h"
 
+class UNiagaraDataInterfaceBase;
+
 /**
 * Niagara shader module interface
 */
@@ -12,11 +14,37 @@ class INiagaraShaderModule : public IModuleInterface
 {
 public:
 	DECLARE_DELEGATE_RetVal(void, FOnProcessQueue);
+	DECLARE_DELEGATE_RetVal_OneParam(UNiagaraDataInterfaceBase*, FOnRequestDefaultDataInterface, const FString&); 
+	
 	FDelegateHandle NIAGARASHADER_API SetOnProcessShaderCompilationQueue(FOnProcessQueue InOnProcessQueue);
 	void NIAGARASHADER_API ResetOnProcessShaderCompilationQueue(FDelegateHandle DelegateHandle);
 	void ProcessShaderCompilationQueue();
 
+	// Handles converting string name to data interface CDO pointer
+	FDelegateHandle NIAGARASHADER_API SetOnRequestDefaultDataInterfaceHandler(FOnRequestDefaultDataInterface InHandler); 
+	void NIAGARASHADER_API ResetOnRequestDefaultDataInterfaceHandler();
+	UNiagaraDataInterfaceBase* RequestDefaultDataInterface(const FString& DIClassName);
+
+	virtual void StartupModule() override
+	{
+		Singleton = this;
+	}
+
+
+	virtual void ShutdownModule() override
+	{
+		Singleton = nullptr;
+	}
+
+	// If you know that you wnat the shader module and have access to the header you don't have to find the module.
+	static INiagaraShaderModule* Get()
+	{
+		return Singleton;
+	}
+
 private:
 	FOnProcessQueue OnProcessQueue;
+	FOnRequestDefaultDataInterface OnRequestDefaultDataInterface;
+	static INiagaraShaderModule* Singleton;
 };
 
