Fixed crash when trying to compile Niagara scripts when running on PC with SM2 emulation.

diff --git a/Engine/Plugins/FX/Niagara/Source/Niagara/Private/NiagaraScript.cpp b/Engine/Plugins/FX/Niagara/Source/Niagara/Private/NiagaraScript.cpp
--- a/Engine/Plugins/FX/Niagara/Source/Niagara/Private/NiagaraScript.cpp
+++ b/Engine/Plugins/FX/Niagara/Source/Niagara/Private/NiagaraScript.cpp
@@ -1069,9 +1069,12 @@ void UNiagaraScript::CacheResourceShadersForRendering(bool bRegenerateId, bool b
 			//if (ScriptResourcesByFeatureLevel[FeatureLevel])
 			{
 				EShaderPlatform ShaderPlatform = GShaderPlatformForFeatureLevel[CacheFeatureLevel];
-				ResourceToCache = ScriptResourcesByFeatureLevel[CacheFeatureLevel];
-				CacheShadersForResources(ShaderPlatform, &ScriptResource, true);
-				ScriptResourcesByFeatureLevel[CacheFeatureLevel] = &ScriptResource;
+				if (IsFeatureLevelSupported(ShaderPlatform, ERHIFeatureLevel::SM5) || IsFeatureLevelSupported(ShaderPlatform, ERHIFeatureLevel::ES3_1))
+				{
+					ResourceToCache = ScriptResourcesByFeatureLevel[CacheFeatureLevel];
+					CacheShadersForResources(ShaderPlatform, &ScriptResource, true);
+					ScriptResourcesByFeatureLevel[CacheFeatureLevel] = &ScriptResource;
+				}
 			}
 		}
 	}
