Merging //UE4/Dev-Main to Dev-Framework (//UE4/Dev-Framework) @ 4736060
#rnx

diff --git a/Engine/Plugins/FX/Niagara/Source/Niagara/Private/NiagaraModule.cpp b/Engine/Plugins/FX/Niagara/Source/Niagara/Private/NiagaraModule.cpp
--- a/Engine/Plugins/FX/Niagara/Source/Niagara/Private/NiagaraModule.cpp
+++ b/Engine/Plugins/FX/Niagara/Source/Niagara/Private/NiagaraModule.cpp
@@ -69,6 +69,7 @@ FNiagaraVariable INiagaraModule::Engine_Owner_XAxis;
 FNiagaraVariable INiagaraModule::Engine_Owner_YAxis;
 FNiagaraVariable INiagaraModule::Engine_Owner_ZAxis;
 FNiagaraVariable INiagaraModule::Engine_Owner_Scale;
+FNiagaraVariable INiagaraModule::Engine_Owner_Rotation;
 
 FNiagaraVariable INiagaraModule::Engine_Owner_SystemLocalToWorld;
 FNiagaraVariable INiagaraModule::Engine_Owner_SystemWorldToLocal;
@@ -84,6 +85,8 @@ FNiagaraVariable INiagaraModule::Engine_Owner_ExecutionState;
 
 FNiagaraVariable INiagaraModule::Engine_ExecutionCount;
 FNiagaraVariable INiagaraModule::Engine_Emitter_NumParticles;
+FNiagaraVariable INiagaraModule::Engine_Emitter_TotalSpawnedParticles;
+FNiagaraVariable INiagaraModule::Engine_System_TickCount;
 FNiagaraVariable INiagaraModule::Engine_System_NumEmittersAlive;
 FNiagaraVariable INiagaraModule::Engine_System_NumEmitters;
 FNiagaraVariable INiagaraModule::Engine_NumSystemInstances;
@@ -95,11 +98,14 @@ FNiagaraVariable INiagaraModule::Engine_System_Age;
 
 FNiagaraVariable INiagaraModule::Emitter_Age;
 FNiagaraVariable INiagaraModule::Emitter_LocalSpace;
+FNiagaraVariable INiagaraModule::Emitter_Determinism;
+FNiagaraVariable INiagaraModule::Emitter_RandomSeed;
 FNiagaraVariable INiagaraModule::Emitter_SpawnRate;
 FNiagaraVariable INiagaraModule::Emitter_SpawnInterval;
 FNiagaraVariable INiagaraModule::Emitter_InterpSpawnStartDt;
 FNiagaraVariable INiagaraModule::Emitter_SpawnGroup;
 
+FNiagaraVariable INiagaraModule::Particles_UniqueID;
 FNiagaraVariable INiagaraModule::Particles_ID;
 FNiagaraVariable INiagaraModule::Particles_Position;
 FNiagaraVariable INiagaraModule::Particles_Velocity;
@@ -121,6 +127,9 @@ FNiagaraVariable INiagaraModule::Particles_UVScale;
 FNiagaraVariable INiagaraModule::Particles_CameraOffset;
 FNiagaraVariable INiagaraModule::Particles_MaterialRandom;
 FNiagaraVariable INiagaraModule::Particles_LightRadius;
+FNiagaraVariable INiagaraModule::Particles_LightExponent;
+FNiagaraVariable INiagaraModule::Particles_LightEnabled;
+FNiagaraVariable INiagaraModule::Particles_LightVolumetricScattering;
 FNiagaraVariable INiagaraModule::Particles_RibbonID;
 FNiagaraVariable INiagaraModule::Particles_RibbonWidth;
 FNiagaraVariable INiagaraModule::Particles_RibbonTwist;
@@ -167,6 +176,7 @@ void INiagaraModule::StartupModule()
 	Engine_Owner_YAxis = FNiagaraVariable(FNiagaraTypeDefinition::GetVec3Def(), TEXT("Engine.Owner.SystemYAxis"));
 	Engine_Owner_ZAxis = FNiagaraVariable(FNiagaraTypeDefinition::GetVec3Def(), TEXT("Engine.Owner.SystemZAxis"));
 	Engine_Owner_Scale = FNiagaraVariable(FNiagaraTypeDefinition::GetVec3Def(), TEXT("Engine.Owner.Scale"));
+	Engine_Owner_Rotation = FNiagaraVariable(FNiagaraTypeDefinition::GetQuatDef(), TEXT("Engine.Owner.Rotation"));
 
 	Engine_Owner_SystemLocalToWorld = FNiagaraVariable(FNiagaraTypeDefinition::GetMatrix4Def(), TEXT("Engine.Owner.SystemLocalToWorld"));
 	Engine_Owner_SystemWorldToLocal = FNiagaraVariable(FNiagaraTypeDefinition::GetMatrix4Def(), TEXT("Engine.Owner.SystemWorldToLocal"));
@@ -182,6 +192,8 @@ void INiagaraModule::StartupModule()
 	
 	Engine_ExecutionCount = FNiagaraVariable(FNiagaraTypeDefinition::GetIntDef(), TEXT("Engine.ExecutionCount"));
 	Engine_Emitter_NumParticles = FNiagaraVariable(FNiagaraTypeDefinition::GetIntDef(), TEXT("Engine.Emitter.NumParticles"));
+	Engine_Emitter_TotalSpawnedParticles = FNiagaraVariable(FNiagaraTypeDefinition::GetIntDef(), TEXT("Engine.Emitter.TotalSpawnedParticles"));
+	Engine_System_TickCount = FNiagaraVariable(FNiagaraTypeDefinition::GetIntDef(), TEXT("Engine.System.TickCount"));
 	Engine_System_NumEmittersAlive = FNiagaraVariable(FNiagaraTypeDefinition::GetIntDef(), TEXT("Engine.System.NumEmittersAlive"));
 	Engine_System_NumEmitters = FNiagaraVariable(FNiagaraTypeDefinition::GetIntDef(), TEXT("Engine.System.NumEmitters"));
 	Engine_NumSystemInstances = FNiagaraVariable(FNiagaraTypeDefinition::GetIntDef(), TEXT("Engine.NumSystemInstances"));
@@ -192,11 +204,14 @@ void INiagaraModule::StartupModule()
 	Engine_System_Age = FNiagaraVariable(FNiagaraTypeDefinition::GetFloatDef(), TEXT("Engine.System.Age"));
 	Emitter_Age = FNiagaraVariable(FNiagaraTypeDefinition::GetFloatDef(), TEXT("Emitter.Age"));
 	Emitter_LocalSpace = FNiagaraVariable(FNiagaraTypeDefinition::GetBoolDef(), TEXT("Emitter.LocalSpace"));
+	Emitter_RandomSeed = FNiagaraVariable(FNiagaraTypeDefinition::GetIntDef(), TEXT("Emitter.RandomSeed"));
+	Emitter_Determinism = FNiagaraVariable(FNiagaraTypeDefinition::GetBoolDef(), TEXT("Emitter.Determinism"));
 	Emitter_SpawnRate = FNiagaraVariable(FNiagaraTypeDefinition::GetFloatDef(), TEXT("Emitter.SpawnRate"));
 	Emitter_SpawnInterval = FNiagaraVariable(FNiagaraTypeDefinition::GetFloatDef(), TEXT("Emitter.SpawnInterval"));
 	Emitter_InterpSpawnStartDt = FNiagaraVariable(FNiagaraTypeDefinition::GetFloatDef(), TEXT("Emitter.InterpSpawnStartDt"));
 	Emitter_SpawnGroup = FNiagaraVariable(FNiagaraTypeDefinition::GetIntDef(), TEXT("Emitter.SpawnGroup"));
 
+	Particles_UniqueID = FNiagaraVariable(FNiagaraTypeDefinition::GetIntDef(), TEXT("Particles.UniqueID"));
 	Particles_ID = FNiagaraVariable(FNiagaraTypeDefinition::GetIDDef(), TEXT("Particles.ID"));
 	Particles_Position = FNiagaraVariable(FNiagaraTypeDefinition::GetVec3Def(), TEXT("Particles.Position"));
 	Particles_Velocity = FNiagaraVariable(FNiagaraTypeDefinition::GetVec3Def(), TEXT("Particles.Velocity"));
@@ -218,6 +233,9 @@ void INiagaraModule::StartupModule()
 	Particles_CameraOffset = FNiagaraVariable(FNiagaraTypeDefinition::GetFloatDef(), TEXT("Particles.CameraOffset"));
 	Particles_MaterialRandom = FNiagaraVariable(FNiagaraTypeDefinition::GetFloatDef(), TEXT("Particles.MaterialRandom"));
 	Particles_LightRadius = FNiagaraVariable(FNiagaraTypeDefinition::GetFloatDef(), TEXT("Particles.LightRadius"));
+	Particles_LightExponent = FNiagaraVariable(FNiagaraTypeDefinition::GetFloatDef(), TEXT("Particles.LightExponent"));
+	Particles_LightEnabled = FNiagaraVariable(FNiagaraTypeDefinition::GetBoolDef(), TEXT("Particles.LightEnabled"));
+	Particles_LightVolumetricScattering = FNiagaraVariable(FNiagaraTypeDefinition::GetFloatDef(), TEXT("Particles.LightVolumetricScattering"));
 	Particles_RibbonID = FNiagaraVariable(FNiagaraTypeDefinition::GetIDDef(), TEXT("Particles.RibbonID"));
 	Particles_RibbonWidth = FNiagaraVariable(FNiagaraTypeDefinition::GetFloatDef(), TEXT("Particles.RibbonWidth"));
 	Particles_RibbonTwist = FNiagaraVariable(FNiagaraTypeDefinition::GetFloatDef(), TEXT("Particles.RibbonTwist"));
@@ -598,6 +616,8 @@ void FNiagaraTypeDefinition::RecreateUserDefinedTypeRegistry()
 
 	for (FSoftObjectPath AssetRef : TotalStructAssets)
 	{
+		FName AssetRefPathNamePreResolve = AssetRef.GetAssetPathName();
+
 		UObject* Obj = AssetRef.ResolveObject();
 		if (Obj == nullptr)
 		{
@@ -613,6 +633,11 @@ void FNiagaraTypeDefinition::RecreateUserDefinedTypeRegistry()
 			{
 				FNiagaraTypeRegistry::Register(ScriptStruct, ParamRefFound != nullptr, PayloadRefFound != nullptr, true);
 			}
+			if (Obj->GetPathName() != AssetRefPathNamePreResolve.ToString())
+			{
+				UE_LOG(LogNiagara, Warning, TEXT("Additional parameter/payload enum has moved from where it was in settings (this may cause errors at runtime): Was: \"%s\" Now: \"%s\""), *AssetRefPathNamePreResolve.ToString(), *Obj->GetPathName());
+			}
+
 		}
 		else
 		{
@@ -623,6 +648,7 @@ void FNiagaraTypeDefinition::RecreateUserDefinedTypeRegistry()
 
 	for (FSoftObjectPath AssetRef : Settings->AdditionalParameterEnums)
 	{
+		FName AssetRefPathNamePreResolve = AssetRef.GetAssetPathName();
 		UObject* Obj = AssetRef.ResolveObject();
 		if (Obj == nullptr)
 		{
@@ -638,6 +664,11 @@ void FNiagaraTypeDefinition::RecreateUserDefinedTypeRegistry()
 			{
 				FNiagaraTypeRegistry::Register(Enum, ParamRefFound != nullptr, PayloadRefFound != nullptr, true);
 			}
+
+			if (Obj->GetPathName() != AssetRefPathNamePreResolve.ToString())
+			{
+				UE_LOG(LogNiagara, Warning, TEXT("Additional parameter/payload enum has moved from where it was in settings (this may cause errors at runtime): Was: \"%s\" Now: \"%s\""), *AssetRefPathNamePreResolve.ToString(), *Obj->GetPathName());
+			}
 		}
 		else
 		{
@@ -757,27 +788,10 @@ FNiagaraTypeDefinition FNiagaraTypeDefinition::GetNumericOutputType(const TArray
 	{
 		return SortedTypeDefinitions[0];
 	}
-}
-
-//////////////////////////////////////////////////////////////////////////
 
-template<>
-void FNiagaraVariable::SetValue(const bool& Data)
-{
-	check(TypeDef == FNiagaraTypeDefinition::GetBoolDef());
-	AllocateData();
-	FNiagaraBool* BoolStruct = (FNiagaraBool*)GetData();
-	BoolStruct->SetValue(Data);
+	return FNiagaraTypeDefinition::GetGenericNumericDef();
 }
 
-template<>
-bool FNiagaraVariable::GetValue() const
-{
-	check(TypeDef == FNiagaraTypeDefinition::GetBoolDef());
-	check(IsDataAllocated());
-	FNiagaraBool* BoolStruct = (FNiagaraBool*)GetData();
-	return BoolStruct->GetValue();
-}
 
 //////////////////////////////////////////////////////////////////////////
 
diff --git a/Engine/Plugins/FX/Niagara/Source/Niagara/Private/NiagaraScript.cpp b/Engine/Plugins/FX/Niagara/Source/Niagara/Private/NiagaraScript.cpp
--- a/Engine/Plugins/FX/Niagara/Source/Niagara/Private/NiagaraScript.cpp
+++ b/Engine/Plugins/FX/Niagara/Source/Niagara/Private/NiagaraScript.cpp
@@ -16,6 +16,9 @@
 #include "UObject/Linker.h"
 #include "HAL/PlatformFilemanager.h"
 #include "Misc/FileHelper.h"
+#include "UObject/EditorObjectVersion.h"
+#include "UObject/ReleaseObjectVersion.h"
+
 
 #if WITH_EDITOR
 	#include "NiagaraScriptDerivedData.h"
@@ -26,13 +29,21 @@
 
 DECLARE_STATS_GROUP(TEXT("Niagara Detailed"), STATGROUP_NiagaraDetailed, STATCAT_Advanced);
 
-FNiagaraScriptDebuggerInfo::FNiagaraScriptDebuggerInfo() : FrameLastWriteId(-1)
+FNiagaraScriptDebuggerInfo::FNiagaraScriptDebuggerInfo() : bWaitForGPU(false), FrameLastWriteId(-1), bWritten(false)
 {
 }
 
 
-FNiagaraScriptDebuggerInfo::FNiagaraScriptDebuggerInfo(FName InName, ENiagaraScriptUsage InUsage, const FGuid& InUsageId) : HandleName(InName), Usage(InUsage), UsageId(InUsageId), FrameLastWriteId(-1)
+FNiagaraScriptDebuggerInfo::FNiagaraScriptDebuggerInfo(FName InName, ENiagaraScriptUsage InUsage, const FGuid& InUsageId) : HandleName(InName), Usage(InUsage), UsageId(InUsageId), FrameLastWriteId(-1), bWritten(false)
 {
+	if (InUsage == ENiagaraScriptUsage::ParticleGPUComputeScript)
+	{
+		bWaitForGPU = true;
+	}
+	else
+	{
+		bWaitForGPU = false;
+	}
 }
 
 
@@ -230,6 +241,10 @@ void UNiagaraScript::ComputeVMCompilationId(FNiagaraVMExecutableDataId& Id) cons
 		{
 			Id.AdditionalDefines.Add(TEXT("Emitter.Localspace"));
 		}
+		if (Emitter->bDeterminism)
+		{
+			Id.AdditionalDefines.Add(TEXT("Emitter.Determinism"));
+		}
 	}
 
 	if (UNiagaraSystem* System = Cast<UNiagaraSystem>(Obj))
@@ -243,6 +258,10 @@ void UNiagaraScript::ComputeVMCompilationId(FNiagaraVMExecutableDataId& Id) cons
 				{
 					Id.AdditionalDefines.Add(Emitter->GetUniqueEmitterName() + TEXT(".Localspace"));
 				}
+				if (Emitter->bDeterminism)
+				{
+					Id.AdditionalDefines.Add(Emitter->GetUniqueEmitterName() + TEXT(".Determinism"));
+				}
 			}
 		}
 	}
@@ -476,7 +495,23 @@ void UNiagaraScript::GenerateStatScopeIDs()
 
 void UNiagaraScript::PostEditChangeProperty(FPropertyChangedEvent& PropertyChangedEvent)
 {
+	Super::PostEditChangeProperty(PropertyChangedEvent);
+
+	FName PropertyName;
+	if (PropertyChangedEvent.Property)
+	{
+		PropertyName = PropertyChangedEvent.Property->GetFName();
+	}
+
 	CacheResourceShadersForRendering(true);
+
+	if (PropertyName == GET_MEMBER_NAME_CHECKED(UNiagaraScript, bDeprecated) || PropertyName == GET_MEMBER_NAME_CHECKED(UNiagaraScript, DeprecationRecommendation))
+	{
+		if (Source)
+		{
+			Source->MarkNotSynchronized(TEXT("Deprecation changed."));
+		}
+	}
 }
 
 #endif
@@ -1034,9 +1069,12 @@ void UNiagaraScript::CacheResourceShadersForRendering(bool bRegenerateId, bool b
 			//if (ScriptResourcesByFeatureLevel[FeatureLevel])
 			{
 				EShaderPlatform ShaderPlatform = GShaderPlatformForFeatureLevel[CacheFeatureLevel];
-				ResourceToCache = ScriptResourcesByFeatureLevel[CacheFeatureLevel];
-				CacheShadersForResources(ShaderPlatform, &ScriptResource, true);
-				ScriptResourcesByFeatureLevel[CacheFeatureLevel] = &ScriptResource;
+				if (IsFeatureLevelSupported(ShaderPlatform, ERHIFeatureLevel::SM5) || IsFeatureLevelSupported(ShaderPlatform, ERHIFeatureLevel::ES3_1))
+				{
+					ResourceToCache = ScriptResourcesByFeatureLevel[CacheFeatureLevel];
+					CacheShadersForResources(ShaderPlatform, &ScriptResource, true);
+					ScriptResourcesByFeatureLevel[CacheFeatureLevel] = &ScriptResource;
+				}
 			}
 		}
 	}
@@ -1063,6 +1101,14 @@ void UNiagaraScript::SyncAliases(const TMap<FString, FString>& RenameMap)
 	// Now handle any Parameters overall..
 	for (int32 i = 0; i < GetVMExecutableData().Parameters.Parameters.Num(); i++)
 	{
+		if (GetVMExecutableData().Parameters.Parameters[i].IsValid() == false)
+		{
+			const FNiagaraVariable& InvalidParameter = GetVMExecutableData().Parameters.Parameters[i];
+			UE_LOG(LogNiagara, Error, TEXT("Invalid parameter found while syncing script aliases.  Script: %s Parameter Name: %s Parameter Type: %s"),
+				*GetPathName(), *InvalidParameter.GetName().ToString(), InvalidParameter.GetType().IsValid() ? *InvalidParameter.GetType().GetName() : TEXT("Unknown"));
+			continue;
+		}
+
 		FNiagaraVariable Var = GetVMExecutableData().Parameters.Parameters[i];
 		FNiagaraVariable NewVar = FNiagaraVariable::ResolveAliases(Var, RenameMap);
 		if (NewVar.GetName() != Var.GetName())
@@ -1199,6 +1245,11 @@ NIAGARA_API bool UNiagaraScript::DidScriptCompilationSucceed(bool bGPUScript) co
 
 void SerializeNiagaraShaderMaps(const TMap<const ITargetPlatform*, TArray<FNiagaraShaderScript*>>* PlatformScriptResourcesToSave, FArchive& Ar, TArray<FNiagaraShaderScript>& OutLoadedResources)
 {
+	Ar.UsingCustomVersion(FFortniteMainBranchObjectVersion::GUID);
+	Ar.UsingCustomVersion(FRenderingObjectVersion::GUID);
+	Ar.UsingCustomVersion(FEditorObjectVersion::GUID);
+	Ar.UsingCustomVersion(FReleaseObjectVersion::GUID);
+
 //	SCOPED_LOADTIMER(SerializeInlineShaderMaps);
 	if (Ar.IsSaving())
 	{
diff --git a/Engine/Plugins/FX/Niagara/Source/NiagaraEditor/Private/NiagaraEmitterFactoryNew.cpp b/Engine/Plugins/FX/Niagara/Source/NiagaraEditor/Private/NiagaraEmitterFactoryNew.cpp
--- a/Engine/Plugins/FX/Niagara/Source/NiagaraEditor/Private/NiagaraEmitterFactoryNew.cpp
+++ b/Engine/Plugins/FX/Niagara/Source/NiagaraEditor/Private/NiagaraEmitterFactoryNew.cpp
@@ -196,6 +196,7 @@ UObject* UNiagaraEmitterFactoryNew::FactoryCreateNew(UClass* Class, UObject* InP
 
 		FNiagaraStackGraphUtilities::RelayoutGraph(*Source->NodeGraph);
 		NewEmitter->bInterpolatedSpawning = true;
+		NewEmitter->bDeterminism = false; // NOTE: Default to non-determinism
 		NewEmitter->SpawnScriptProps.Script->SetUsage(ENiagaraScriptUsage::ParticleSpawnScriptInterpolated);
 	}
 	
diff --git a/Engine/Plugins/FX/Niagara/Source/NiagaraEditor/Private/NiagaraScriptSource.cpp b/Engine/Plugins/FX/Niagara/Source/NiagaraEditor/Private/NiagaraScriptSource.cpp
--- a/Engine/Plugins/FX/Niagara/Source/NiagaraEditor/Private/NiagaraScriptSource.cpp
+++ b/Engine/Plugins/FX/Niagara/Source/NiagaraEditor/Private/NiagaraScriptSource.cpp
@@ -185,6 +185,13 @@ void InitializeNewRapidIterationParametersForNode(const UEdGraphSchema_Niagara*
 		for (const UEdGraphPin* FunctionInputPin : FunctionInputPins)
 		{
 			FNiagaraTypeDefinition InputType = Schema->PinToTypeDefinition(FunctionInputPin);
+			if (InputType.IsValid() == false)
+			{
+				UE_LOG(LogNiagaraEditor, Error, TEXT("Invalid input type found while attempting initialize new rapid iteration parameters. Function Node: %s %s Input Name: %s"),
+					*FunctionCallNode->GetPathName(), *FunctionCallNode->GetFunctionName(), *FunctionInputPin->GetName());
+				continue;
+			}
+
 			if (FNiagaraStackGraphUtilities::IsRapidIterationType(InputType))
 			{
 				FNiagaraParameterHandle AliasedFunctionInputHandle = FNiagaraParameterHandle::CreateAliasedModuleParameterHandle(FNiagaraParameterHandle(FunctionInputPin->PinName), FunctionCallNode);
diff --git a/Engine/Plugins/FX/Niagara/Source/NiagaraEditor/Private/ViewModels/Stack/NiagaraStackScriptItemGroup.cpp b/Engine/Plugins/FX/Niagara/Source/NiagaraEditor/Private/ViewModels/Stack/NiagaraStackScriptItemGroup.cpp
--- a/Engine/Plugins/FX/Niagara/Source/NiagaraEditor/Private/ViewModels/Stack/NiagaraStackScriptItemGroup.cpp
+++ b/Engine/Plugins/FX/Niagara/Source/NiagaraEditor/Private/ViewModels/Stack/NiagaraStackScriptItemGroup.cpp
@@ -643,13 +643,6 @@ void UNiagaraStackScriptItemGroup::OnScriptGraphChanged(const struct FEdGraphEdi
 
 void UNiagaraStackScriptItemGroup::AddParameterModuleToStack(const UNiagaraStackModuleSpacer* InModuleSpacer, const FNiagaraVariable &InVariable)
 {
-	TArray<FNiagaraVariable> Vars;
-	Vars.Add(InVariable);
-	TArray<FString> DefaultVals;
-	DefaultVals.Add(FNiagaraConstants::GetAttributeDefaultValue(InVariable));
-
-	UNiagaraNodeOutput* OutputNode = GetScriptOutputNode();
-
 	int32 TargetIndex = INDEX_NONE;
 	UNiagaraStackModuleItem** TargetModuleItemPtr = StackSpacerToModuleItemMap.Find(FObjectKey(InModuleSpacer));
 	if (*TargetModuleItemPtr != nullptr)
@@ -658,8 +651,8 @@ void UNiagaraStackScriptItemGroup::AddParameterModuleToStack(const UNiagaraStack
 		TargetIndex = TargetModuleItem->GetModuleIndex();
 	}
 
-	FNiagaraStackGraphUtilities::AddParameterModuleToStack(Vars, *OutputNode, TargetIndex, DefaultVals);
-	RefreshChildren();
+	TSharedRef<FScriptGroupAddAction> AddAction = FScriptGroupAddAction::CreateExistingParameterModuleAction(InVariable);
+	AddUtilities->ExecuteAddAction(AddAction, TargetIndex);
 }
 
 #undef LOCTEXT_NAMESPACE
