Remove the "User." prefix for niagara variable names in the sequencer

diff --git a/Engine/Plugins/FX/Niagara/Source/NiagaraEditor/Private/Sequencer/LevelSequence/NiagaraSystemTrackEditor.cpp b/Engine/Plugins/FX/Niagara/Source/NiagaraEditor/Private/Sequencer/LevelSequence/NiagaraSystemTrackEditor.cpp
--- a/Engine/Plugins/FX/Niagara/Source/NiagaraEditor/Private/Sequencer/LevelSequence/NiagaraSystemTrackEditor.cpp
+++ b/Engine/Plugins/FX/Niagara/Source/NiagaraEditor/Private/Sequencer/LevelSequence/NiagaraSystemTrackEditor.cpp
@@ -83,7 +83,7 @@ void FNiagaraSystemTrackEditor::BuildObjectBindingTrackMenu(FMenuBuilder& MenuBu
 	if (System != nullptr)
 	{
 		TArray<FNiagaraVariable> ParameterVariables;
-		System->GetExposedParameters().GetParameters(ParameterVariables);
+		System->GetExposedParameters().GetUserParameters(ParameterVariables);
 
 		TSet<FNiagaraVariable> AnimatedParameters;
 		GetAnimatedParameters(*GetSequencer()->GetFocusedMovieSceneSequence()->GetMovieScene(), ObjectBinding, AnimatedParameters);
