Updates readback to not create/destroy stuff per emitter every frame.
Updated RHI api to support this.

diff --git a/Engine/Plugins/FX/Niagara/Source/Niagara/Private/NiagaraEmitterInstanceBatcher.cpp b/Engine/Plugins/FX/Niagara/Source/Niagara/Private/NiagaraEmitterInstanceBatcher.cpp
--- a/Engine/Plugins/FX/Niagara/Source/Niagara/Private/NiagaraEmitterInstanceBatcher.cpp
+++ b/Engine/Plugins/FX/Niagara/Source/Niagara/Private/NiagaraEmitterInstanceBatcher.cpp
@@ -125,6 +125,9 @@ void NiagaraEmitterInstanceBatcher::TickSingle(FNiagaraComputeExecutionContext *
 	// assume all instances survived; ResolveDataSetWrites will change this if the deferred readback was successful; that data may be several frames old
 	Context->MainDataSet->CurrData().SetNumInstances(NewNumInstances);
 
+	// ResolveDatasetWrites may read this, so we must transition it here.
+	RHICmdList.TransitionResource(EResourceTransitionAccess::EReadable, EResourceTransitionPipeline::EComputeToGfx, Context->MainDataSet->GetCurDataSetIndices().UAV);		// transition to readable; we'll be using this next frame
+
 	// Don't resolve if the data if there are no instances (prevents a transition issue warning).
 	if (NewNumInstances > 0)
 	{
@@ -132,8 +135,6 @@ void NiagaraEmitterInstanceBatcher::TickSingle(FNiagaraComputeExecutionContext *
 		ResolveDatasetWrites(RHICmdList, Context);
 	}
 
-	RHICmdList.TransitionResource(EResourceTransitionAccess::EReadable, EResourceTransitionPipeline::EComputeToGfx, Context->MainDataSet->GetCurDataSetIndices().UAV);		// transition to readable; we'll be using this next frame
-
 	/*
 	// TODO: hack - only updating event set 0 on update scripts now; need to match them to their indices and update them all
 	if (Context->UpdateEventWriteDataSets.Num())
@@ -197,28 +198,26 @@ void NiagaraEmitterInstanceBatcher::ResolveDatasetWrites(FRHICommandList &RHICmd
 	{
 		SCOPE_CYCLE_COUNTER(STAT_NiagaraAllocateGPUReadback_RT);
 
-		Context->GPUDataReadback = new FRHIGPUMemoryReadback(DatasetIndexBufferWrite.Buffer, TEXT("Niagara GPU Emitter Readback"));
+		Context->GPUDataReadback = new FRHIGPUMemoryReadback(TEXT("Niagara GPU Emitter Readback"));
 		INC_DWORD_STAT(STAT_NiagaraReadbackLatency);
-		Context->GPUDataReadback->Insert(RHICmdList);
+		Context->GPUDataReadback->EnqueueCopy(RHICmdList, DatasetIndexBufferWrite.Buffer);
 	}
 	else if (Context->GPUDataReadback->IsReady())
 	{
 		bool bSuccessfullyRead = false;
 		{
 		    SCOPE_CYCLE_COUNTER(STAT_NiagaraGPUReadback_RT);
-		    int32 *NumInstancesAfterSim = static_cast<int32*>(Context->GPUDataReadback->RetrieveData(64 * sizeof(int32)));
+		    int32 *NumInstancesAfterSim = static_cast<int32*>(Context->GPUDataReadback->Lock(64 * sizeof(int32)));
 			if (NumInstancesAfterSim)
 			{
 				int32 ExistingDataCount = Context->MainDataSet->CurrData().GetNumInstances();// index 1 is always the count
 				int32 NewExistingDataCount = NumInstancesAfterSim[1] + Context->AccumulatedSpawnRate;
 				Context->MainDataSet->CurrData().SetNumInstances(NewExistingDataCount);
 				FString PathName = Context->GPUScript->GetOutermost()->GetPathName();
-				// UE_LOG(LogNiagara, Log, TEXT("GPU Syncup %s : Was(%d) Now(%d)"), *PathName, ExistingDataCount, NewExistingDataCount );
+				//UE_LOG(LogNiagara, Log, TEXT("GPU Syncup %s : Was(%d) Now(%d)"), *PathName, ExistingDataCount, NewExistingDataCount );
 				SET_DWORD_STAT(STAT_NiagaraGPUParticles, NewExistingDataCount);
 				SET_DWORD_STAT(STAT_NiagaraReadbackLatency, 0);
 
-				Context->GPUDataReadback->Finish();
-
 				Context->AccumulatedSpawnRate = 0;
 				bSuccessfullyRead = true;
 			}
@@ -226,17 +225,15 @@ void NiagaraEmitterInstanceBatcher::ResolveDatasetWrites(FRHICommandList &RHICmd
 			{
 				UE_LOG(LogNiagara, Warning, TEXT("GPUDataReadback said it was ready, but returned an invalid buffer. Skipping this time.."));
 			}
+			Context->GPUDataReadback->Unlock();
 		}
 		if (bSuccessfullyRead)
 		{
 			SCOPE_CYCLE_COUNTER(STAT_NiagaraAllocateGPUReadback_RT);
 			// The following code seems to take significant time on d3d12
-			delete Context->GPUDataReadback;
-			Context->GPUDataReadback = new FRHIGPUMemoryReadback(DatasetIndexBufferWrite.Buffer, TEXT("Niagara GPU Emitter Readback"));
-			Context->GPUDataReadback->Insert(RHICmdList);
+			Context->GPUDataReadback->EnqueueCopy(RHICmdList, DatasetIndexBufferWrite.Buffer);
 		}
 	}
-
 }
 
 void NiagaraEmitterInstanceBatcher::ProcessDebugInfo(FRHICommandList &RHICmdList, const FNiagaraComputeExecutionContext *Context) const
@@ -260,31 +257,31 @@ void NiagaraEmitterInstanceBatcher::ProcessDebugInfo(FRHICommandList &RHICmdList
 		{
 			//UE_LOG(LogNiagara, Warning, TEXT("Read back!"));
 
-			int32 *NumInstancesAfterSim = static_cast<int32*>(Context->GPUDebugDataReadbackCounts->RetrieveData(64 * sizeof(int32)));
+			int32 *NumInstancesAfterSim = static_cast<int32*>(Context->GPUDebugDataReadbackCounts->Lock(64 * sizeof(int32)));
 			int32 NewExistingDataCount = NumInstancesAfterSim[1];
 			{
 				float* FloatDataBuffer = nullptr;
 				if (Context->GPUDebugDataReadbackFloat)
 				{
-					FloatDataBuffer = static_cast<float*>(Context->GPUDebugDataReadbackFloat->RetrieveData(Context->GPUDebugDataFloatSize));
+					FloatDataBuffer = static_cast<float*>(Context->GPUDebugDataReadbackFloat->Lock(Context->GPUDebugDataFloatSize));
 				}
 				int* IntDataBuffer = nullptr;
 				if (Context->GPUDebugDataReadbackInt)
 				{
-					IntDataBuffer = static_cast<int*>(Context->GPUDebugDataReadbackInt->RetrieveData(Context->GPUDebugDataIntSize));
+					IntDataBuffer = static_cast<int*>(Context->GPUDebugDataReadbackInt->Lock(Context->GPUDebugDataIntSize));
 				}
 				Context->MainDataSet->DumpGPU(Context->DebugInfo->Frame, FloatDataBuffer, IntDataBuffer, 0, NewExistingDataCount);
 				Context->DebugInfo->bWritten = true;
 
 				if (Context->GPUDebugDataReadbackFloat)
 				{
-					Context->GPUDebugDataReadbackFloat->Finish();
+					Context->GPUDebugDataReadbackFloat->Unlock();
 				}
 				if (Context->GPUDebugDataReadbackInt)
 				{
-					Context->GPUDebugDataReadbackInt->Finish();
+					Context->GPUDebugDataReadbackInt->Unlock();
 				}
-				Context->GPUDebugDataReadbackCounts->Finish();
+				Context->GPUDebugDataReadbackCounts->Unlock();
 			}
 			{
 				// The following code seems to take significant time on d3d12
@@ -515,20 +512,20 @@ void NiagaraEmitterInstanceBatcher::Run(const FNiagaraComputeExecutionContext *C
 
 			if (Context->MainDataSet->GetNumFloatComponents() > 0)
 			{
-				Context->GPUDebugDataReadbackFloat = new FRHIGPUMemoryReadback(Context->MainDataSet->GetDataByIndex(WriteBufferIdx).GetGPUBufferFloat()->Buffer, TEXT("Niagara GPU Debug Info Float Emitter Readback"));
-				Context->GPUDebugDataReadbackFloat->Insert(RHICmdList);
+				Context->GPUDebugDataReadbackFloat = new FRHIGPUMemoryReadback(TEXT("Niagara GPU Debug Info Float Emitter Readback"));
+				Context->GPUDebugDataReadbackFloat->EnqueueCopy(RHICmdList, Context->MainDataSet->GetDataByIndex(WriteBufferIdx).GetGPUBufferFloat()->Buffer);
 				Context->GPUDebugDataFloatSize = Context->MainDataSet->GetDataByIndex(WriteBufferIdx).GetGPUBufferFloat()->NumBytes;
 			}
 
 			if (Context->MainDataSet->GetNumInt32Components() > 0)
 			{
-				Context->GPUDebugDataReadbackInt = new FRHIGPUMemoryReadback(Context->MainDataSet->GetDataByIndex(WriteBufferIdx).GetGPUBufferInt()->Buffer, TEXT("Niagara GPU Debug Info Int Emitter Readback"));
-				Context->GPUDebugDataReadbackInt->Insert(RHICmdList);
+				Context->GPUDebugDataReadbackInt = new FRHIGPUMemoryReadback(TEXT("Niagara GPU Debug Info Int Emitter Readback"));
+				Context->GPUDebugDataReadbackInt->EnqueueCopy(RHICmdList, Context->MainDataSet->GetDataByIndex(WriteBufferIdx).GetGPUBufferInt()->Buffer);
 				Context->GPUDebugDataIntSize = Context->MainDataSet->GetDataByIndex(WriteBufferIdx).GetGPUBufferInt()->NumBytes;
 			}
 
-			Context->GPUDebugDataReadbackCounts = new FRHIGPUMemoryReadback(DatasetIndexBufferWrite.Buffer, TEXT("Niagara GPU Emitter Readback"));
-			Context->GPUDebugDataReadbackCounts->Insert(RHICmdList);
+			Context->GPUDebugDataReadbackCounts = new FRHIGPUMemoryReadback(TEXT("Niagara GPU Emitter Readback"));
+			Context->GPUDebugDataReadbackCounts->EnqueueCopy(RHICmdList, DatasetIndexBufferWrite.Buffer);
 		}
 	}
 #endif // WITH_EDITORONLY_DATA
diff --git a/Engine/Source/Runtime/Apple/MetalRHI/Private/MetalUAV.cpp b/Engine/Source/Runtime/Apple/MetalRHI/Private/MetalUAV.cpp
--- a/Engine/Source/Runtime/Apple/MetalRHI/Private/MetalUAV.cpp
+++ b/Engine/Source/Runtime/Apple/MetalRHI/Private/MetalUAV.cpp
@@ -848,39 +848,39 @@ void FMetalGPUFence::WriteInternal(mtlpp::CommandBuffer& CmdBuffer)
 	check(Fence);
 }
 
-void FMetalRHICommandContext::RHIEnqueueStagedRead(FStagingBufferRHIParamRef StagingBuffer, FGPUFenceRHIParamRef InFence, uint32 Offset, uint32 NumBytes)
+void FMetalRHICommandContext::RHICopyToStagingBuffer(FVertexBufferRHIParamRef SourceBufferRHI, FStagingBufferRHIParamRef DestinationStagingBufferRHI, uint32 Offset, uint32 NumBytes, FGPUFenceRHIParamRef FenceRHI)
 {
 	@autoreleasepool {
-		check(StagingBuffer);
-		
-		FMetalStagingBuffer* MetalStagingBuffer = ResourceCast(StagingBuffer);
-		FMetalVertexBuffer* BackingBuffer = ResourceCast(MetalStagingBuffer->GetBackingBuffer());
-		FMetalBuffer& ReadbackBuffer = MetalStagingBuffer->ReadbackStagingBuffer;
-		
+		check(DestinationStagingBufferRHI);
+
+		FMetalStagingBuffer* MetalStagingBuffer = ResourceCast(DestinationStagingBufferRHI);
+		ensureMsgf(!MetalStagingBuffer->bIsLocked, TEXT("Attempting to Copy to a locked staging buffer. This may have undefined behavior"));
+		FMetalVertexBuffer* SourceBuffer = ResourceCast(SourceBufferRHI);
+		FMetalBuffer& ReadbackBuffer = MetalStagingBuffer->ShadowBuffer;
+
 		// Need a shadow buffer for this read. If it hasn't been allocated in our FStagingBuffer or if
 		// it's not big enough to hold our readback we need to allocate.
-		if(!ReadbackBuffer || ReadbackBuffer.GetLength() < NumBytes)
+		if (!ReadbackBuffer || ReadbackBuffer.GetLength() < NumBytes)
 		{
-			if(ReadbackBuffer)
+			if (ReadbackBuffer)
 			{
 				SafeReleaseMetalBuffer(ReadbackBuffer);
 			}
 			FMetalPooledBufferArgs ArgsCPU(GetMetalDeviceContext().GetDevice(), NumBytes, mtlpp::StorageMode::Shared);
 			ReadbackBuffer = GetMetalDeviceContext().CreatePooledBuffer(ArgsCPU);
 		}
-		
+
 		// Inline copy from the actual buffer to the shadow
-		GetMetalDeviceContext().CopyFromBufferToBuffer(BackingBuffer->Buffer, Offset, ReadbackBuffer, 0, NumBytes);
-		
-		if (InFence)
+		GetMetalDeviceContext().CopyFromBufferToBuffer(SourceBuffer->Buffer, Offset, ReadbackBuffer, 0, NumBytes);
+
+		if (FenceRHI)
 		{
-			FMetalGPUFence* Fence = ResourceCast(InFence);
+			FMetalGPUFence* Fence = ResourceCast(FenceRHI);
 			Fence->WriteInternal(Context->GetCurrentCommandBuffer());
 		}
 	}
 }
 
-
 FGPUFenceRHIRef FMetalDynamicRHI::RHICreateGPUFence(const FName &Name)
 {
 	@autoreleasepool {
@@ -888,6 +888,11 @@ FGPUFenceRHIRef FMetalDynamicRHI::RHICreateGPUFence(const FName &Name)
 	}
 }
 
+void FMetalGPUFence::Clear()
+{
+	Fence = mtlpp::CommandBufferFence();
+}
+
 bool FMetalGPUFence::Poll() const
 {
 	if (Fence)
diff --git a/Engine/Source/Runtime/Apple/MetalRHI/Private/MetalVertexBuffer.cpp b/Engine/Source/Runtime/Apple/MetalRHI/Private/MetalVertexBuffer.cpp
--- a/Engine/Source/Runtime/Apple/MetalRHI/Private/MetalVertexBuffer.cpp
+++ b/Engine/Source/Runtime/Apple/MetalRHI/Private/MetalVertexBuffer.cpp
@@ -769,16 +769,16 @@ void FMetalDynamicRHI::UnlockStagingBuffer_RenderThread(class FRHICommandListImm
 	return RHIUnlockStagingBuffer(StagingBuffer);
 }
 
-FStagingBufferRHIRef FMetalDynamicRHI::RHICreateStagingBuffer(FVertexBufferRHIParamRef VertexBuffer)
+FStagingBufferRHIRef FMetalDynamicRHI::RHICreateStagingBuffer()
 {
-	return new FMetalStagingBuffer(VertexBuffer);
+	return new FMetalStagingBuffer();
 }
 
 FMetalStagingBuffer::~FMetalStagingBuffer()
 {
-	if(ReadbackStagingBuffer)
+	if (ShadowBuffer)
 	{
-		SafeReleaseMetalBuffer(ReadbackStagingBuffer);
+		SafeReleaseMetalBuffer(ShadowBuffer);
 	}
 }
 
@@ -786,13 +786,16 @@ FMetalStagingBuffer::~FMetalStagingBuffer()
 // If this was not fenced correctly it will not have the expected data.
 void *FMetalStagingBuffer::Lock(uint32 Offset, uint32 NumBytes)
 {
-	check(ReadbackStagingBuffer);
-	
-	uint8* BackingPtr = (uint8*) ReadbackStagingBuffer.GetContents();
-	return BackingPtr+Offset;
+	check(ShadowBuffer);
+	check(!bIsLocked);
+	bIsLocked = true;
+	uint8* BackingPtr = (uint8*)ShadowBuffer.GetContents();
+	return BackingPtr + Offset;
 }
 
 void FMetalStagingBuffer::Unlock()
 {
 	// does nothing in metal.
+	check(bIsLocked);
+	bIsLocked = false;
 }
diff --git a/Engine/Source/Runtime/Apple/MetalRHI/Public/MetalRHI.h b/Engine/Source/Runtime/Apple/MetalRHI/Public/MetalRHI.h
--- a/Engine/Source/Runtime/Apple/MetalRHI/Public/MetalRHI.h
+++ b/Engine/Source/Runtime/Apple/MetalRHI/Public/MetalRHI.h
@@ -189,7 +189,7 @@ public:
 	virtual FComputeFenceRHIRef RHICreateComputeFence(const FName& Name) final override;
 	virtual FGPUFenceRHIRef RHICreateGPUFence(const FName &Name) final override;
 
-	virtual FStagingBufferRHIRef RHICreateStagingBuffer(FVertexBufferRHIParamRef VertexBuffer) final override;
+	virtual FStagingBufferRHIRef RHICreateStagingBuffer() final override;
 	virtual void* RHILockStagingBuffer(FStagingBufferRHIParamRef StagingBuffer, uint32 Offset, uint32 SizeRHI) final override;
 	virtual void RHIUnlockStagingBuffer(FStagingBufferRHIParamRef StagingBuffer) final override;
 	virtual void* LockStagingBuffer_RenderThread(class FRHICommandListImmediate& RHICmdList, FStagingBufferRHIParamRef StagingBuffer, uint32 Offset, uint32 SizeRHI) final override;
diff --git a/Engine/Source/Runtime/Apple/MetalRHI/Public/MetalRHIContext.h b/Engine/Source/Runtime/Apple/MetalRHI/Public/MetalRHIContext.h
--- a/Engine/Source/Runtime/Apple/MetalRHI/Public/MetalRHIContext.h
+++ b/Engine/Source/Runtime/Apple/MetalRHI/Public/MetalRHIContext.h
@@ -318,7 +318,7 @@ public:
 	 */
 	virtual void RHITransitionResources(EResourceTransitionAccess TransitionType, EResourceTransitionPipeline TransitionPipeline, FUnorderedAccessViewRHIParamRef* InUAVs, int32 NumUAVs, FComputeFenceRHIParamRef WriteComputeFence) final override;
 	
-	virtual void RHIEnqueueStagedRead(FStagingBufferRHIParamRef StagingBuffer, FGPUFenceRHIParamRef InFence, uint32 Offset, uint32 NumBytes) final override;
+	virtual void RHICopyToStagingBuffer(FVertexBufferRHIParamRef SourceBufferRHI, FStagingBufferRHIParamRef DestinationStagingBufferRHI, uint32 Offset, uint32 NumBytes, FGPUFenceRHIParamRef FenceRHI) final override;
 
 	/**
 	 * Explicitly transition a texture resource from readable -> writable by the GPU or vice versa.
diff --git a/Engine/Source/Runtime/Apple/MetalRHI/Public/MetalResources.h b/Engine/Source/Runtime/Apple/MetalRHI/Public/MetalResources.h
--- a/Engine/Source/Runtime/Apple/MetalRHI/Public/MetalResources.h
+++ b/Engine/Source/Runtime/Apple/MetalRHI/Public/MetalResources.h
@@ -1033,6 +1033,8 @@ public:
 	{
 	}
 
+	virtual void Clear() override final;
+
 	void WriteInternal(mtlpp::CommandBuffer& CmdBuffer);
 
 	virtual bool Poll() const override final;
@@ -1043,19 +1045,20 @@ private:
 
 class FMetalStagingBuffer final : public FRHIStagingBuffer
 {
+	friend class FMetalRHICommandContext;
 public:
-	FMetalStagingBuffer(FVertexBufferRHIRef InBuffer)
-	: FRHIStagingBuffer(InBuffer)
+	FMetalStagingBuffer()
+		: FRHIStagingBuffer()
 	{
 	}
-	
-	~FMetalStagingBuffer();
 
-	void *Lock(uint32 Offset, uint32 NumBytes);
+	virtual ~FMetalStagingBuffer() final override;
 
-	void Unlock();
-	
-	FMetalBuffer ReadbackStagingBuffer;
+	virtual void* Lock(uint32 Offset, uint32 NumBytes) final override;
+
+	virtual void Unlock() final override;
+private:
+	FMetalBuffer ShadowBuffer;
 };
 
 class FMetalShaderLibrary final : public FRHIShaderLibrary
diff --git a/Engine/Source/Runtime/D3D12RHI/Private/D3D12CommandContext.h b/Engine/Source/Runtime/D3D12RHI/Private/D3D12CommandContext.h
--- a/Engine/Source/Runtime/D3D12RHI/Private/D3D12CommandContext.h
+++ b/Engine/Source/Runtime/D3D12RHI/Private/D3D12CommandContext.h
@@ -262,7 +262,7 @@ public:
 	virtual void RHIClearTinyUAV(FUnorderedAccessViewRHIParamRef UnorderedAccessViewRHI, const uint32* Values) final override;
 	virtual void RHICopyToResolveTarget(FTextureRHIParamRef SourceTexture, FTextureRHIParamRef DestTexture, const FResolveParams& ResolveParams) final override;
 	virtual void RHITransitionResources(EResourceTransitionAccess TransitionType, FTextureRHIParamRef* InTextures, int32 NumTextures) final override;
-    virtual void RHIEnqueueStagedRead(FStagingBufferRHIParamRef StagingBuffer, FGPUFenceRHIParamRef Fence, uint32 Offset, uint32 NumBytes) final override;
+	virtual void RHICopyToStagingBuffer(FVertexBufferRHIParamRef SourceBuffer, FStagingBufferRHIParamRef DestinationStagingBuffer, uint32 Offset, uint32 NumBytes, FGPUFenceRHIParamRef Fence = nullptr) final override;
 	virtual void RHIBeginRenderQuery(FRenderQueryRHIParamRef RenderQuery) final override;
 	virtual void RHIEndRenderQuery(FRenderQueryRHIParamRef RenderQuery) final override;
 	void RHIBeginOcclusionQueryBatch(uint32 NumQueriesInBatch);
@@ -449,9 +449,9 @@ public:
 	// Special implementation that only signal the fence once.
 	virtual void RHITransitionResources(EResourceTransitionAccess TransitionType, EResourceTransitionPipeline TransitionPipeline, FUnorderedAccessViewRHIParamRef* InUAVs, int32 NumUAVs, FComputeFenceRHIParamRef WriteComputeFenceRHI) final override;
 
-	FORCEINLINE virtual void RHIEnqueueStagedRead(FStagingBufferRHIParamRef StagingBuffer, FGPUFenceRHIParamRef Fence, uint32 Offset, uint32 NumBytes) final override
+	FORCEINLINE virtual void RHICopyToStagingBuffer(FVertexBufferRHIParamRef SourceBuffer, FStagingBufferRHIParamRef DestinationStagingBuffer, uint32 Offset, uint32 NumBytes, FGPUFenceRHIParamRef Fence = nullptr) final override
 	{
-		ContextRedirect(RHIEnqueueStagedRead(StagingBuffer, Fence, Offset, NumBytes));
+		ContextRedirect(RHICopyToStagingBuffer(SourceBuffer, DestinationStagingBuffer, Offset, NumBytes, Fence));
 	}
 	FORCEINLINE virtual void RHISetShaderTexture(FComputeShaderRHIParamRef PixelShader, uint32 TextureIndex, FTextureRHIParamRef NewTexture) final override
 	{
diff --git a/Engine/Source/Runtime/D3D12RHI/Private/D3D12Commands.cpp b/Engine/Source/Runtime/D3D12RHI/Private/D3D12Commands.cpp
--- a/Engine/Source/Runtime/D3D12RHI/Private/D3D12Commands.cpp
+++ b/Engine/Source/Runtime/D3D12RHI/Private/D3D12Commands.cpp
@@ -395,19 +395,32 @@ void FD3D12CommandContext::RHITransitionResources(EResourceTransitionAccess Tran
 	}
 }
 
-void FD3D12CommandContext::RHIEnqueueStagedRead(FStagingBufferRHIParamRef StagingBufferRHI, FGPUFenceRHIParamRef FenceRHI, uint32 Offset, uint32 NumBytes)
+
+void FD3D12CommandContext::RHICopyToStagingBuffer(FVertexBufferRHIParamRef SourceBufferRHI, FStagingBufferRHIParamRef StagingBufferRHI, uint32 Offset, uint32 NumBytes, FGPUFenceRHIParamRef FenceRHI)
 {
 	FD3D12StagingBuffer* StagingBuffer = FD3D12DynamicRHI::ResourceCast(StagingBufferRHI);
 	check(StagingBuffer);
+	ensureMsgf(!StagingBuffer->bIsLocked, TEXT("Attempting to Copy to a locked staging buffer. This may have undefined behavior"));
 
-	FD3D12VertexBuffer* VertexBuffer = FD3D12DynamicRHI::ResourceCast(StagingBuffer->GetBackingBuffer());
+	FD3D12VertexBuffer* VertexBuffer = FD3D12DynamicRHI::ResourceCast(SourceBufferRHI);
 	check(VertexBuffer);
 
 	// Only get data from the first gpu for now.
 	FD3D12Device* StagingDevice = VertexBuffer->GetParentDevice();
-	StagingBuffer->StagedRead.SafeRelease();
-	VERIFYD3D12RESULT(GetParentDevice()->GetParentAdapter()->CreateBuffer(D3D12_HEAP_TYPE_READBACK, GetGPUMask(), GetGPUMask(), Offset + NumBytes, StagingBuffer->StagedRead.GetInitReference()));
-	
+
+	// Ensure our shadow buffer is large enough to hold the readback.
+	if (!StagingBuffer->StagedRead || StagingBuffer->ShadowBufferSize < NumBytes)
+	{
+		// @todo-mattc I feel like we should allocate more than NumBytes to handle small reads without blowing tons of space. Need to pool this.
+		// Hopefully d3d12 will do smart pooling out of an internal heap.
+		if (StagingBuffer->StagedRead)
+		{
+			StagingBuffer->StagedRead.SafeRelease();
+		}
+		VERIFYD3D12RESULT(GetParentDevice()->GetParentAdapter()->CreateBuffer(D3D12_HEAP_TYPE_READBACK, GetGPUMask(), GetGPUMask(), NumBytes, StagingBuffer->StagedRead.GetInitReference()));
+		StagingBuffer->ShadowBufferSize = NumBytes;
+	}
+
 	{
 		FD3D12Resource* pSourceResource = VertexBuffer->ResourceLocation.GetResource();
 		D3D12_RESOURCE_DESC const& SourceBufferDesc = pSourceResource->GetDesc();
@@ -427,6 +440,7 @@ void FD3D12CommandContext::RHIEnqueueStagedRead(FStagingBufferRHIParamRef Stagin
 
 	if (FenceRHI)
 	{
+		// @todo-mattc we don't want to flush here. That should be the caller's responsibility.
 		RHISubmitCommandsHint();
 		FD3D12GPUFence* Fence = FD3D12DynamicRHI::ResourceCast(FenceRHI);
 		Fence->WriteInternal(ED3D12CommandQueueType::Default);
diff --git a/Engine/Source/Runtime/D3D12RHI/Private/D3D12DirectCommandListManager.cpp b/Engine/Source/Runtime/D3D12RHI/Private/D3D12DirectCommandListManager.cpp
--- a/Engine/Source/Runtime/D3D12RHI/Private/D3D12DirectCommandListManager.cpp
+++ b/Engine/Source/Runtime/D3D12RHI/Private/D3D12DirectCommandListManager.cpp
@@ -24,7 +24,13 @@ void FD3D12GPUFence::WriteInternal(ED3D12CommandQueueType QueueType)
 
 bool FD3D12GPUFence::Poll() const
 {
-	return !Value || (Fence && Fence->IsFenceComplete(Value));
+	// @todo-mattc Value of 0 means signaled? Revisit this...
+	return !Value || (Fence && Fence->PeekLastCompletedFence() >= Value);
+}
+
+void FD3D12GPUFence::Clear()
+{
+	Value = MAX_uint64;
 }
 
 
@@ -33,9 +39,9 @@ FGPUFenceRHIRef FD3D12DynamicRHI::RHICreateGPUFence(const FName& Name)
 	return new FD3D12GPUFence(Name, GetAdapter().GetStagingFence());
 }
 
-FStagingBufferRHIRef FD3D12DynamicRHI::RHICreateStagingBuffer(FVertexBufferRHIParamRef VertexBufferRHI)
+FStagingBufferRHIRef FD3D12DynamicRHI::RHICreateStagingBuffer()
 {
-	return new FD3D12StagingBuffer(VertexBufferRHI);
+	return new FD3D12StagingBuffer();
 }
 
 void* FD3D12DynamicRHI::RHILockStagingBuffer(FStagingBufferRHIParamRef StagingBufferRHI, uint32 Offset, uint32 SizeRHI)
@@ -43,30 +49,14 @@ void* FD3D12DynamicRHI::RHILockStagingBuffer(FStagingBufferRHIParamRef StagingBu
 	FD3D12StagingBuffer* StagingBuffer = FD3D12DynamicRHI::ResourceCast(StagingBufferRHI);
 	check(StagingBuffer);
 
-	FD3D12Resource* pResource = StagingBuffer->StagedRead.GetReference();
-	if (pResource)
-	{
-		D3D12_RANGE ReadRange;
-		ReadRange.Begin = Offset;
-		ReadRange.End = Offset + SizeRHI;
-		return reinterpret_cast<uint8*>(pResource->Map(&ReadRange)) + Offset;
-	}
-	else
-	{
-		return nullptr;
-	}
+	return StagingBuffer->Lock(Offset, SizeRHI);
 }
 
 void FD3D12DynamicRHI::RHIUnlockStagingBuffer(FStagingBufferRHIParamRef StagingBufferRHI)
 {
 	FD3D12StagingBuffer* StagingBuffer = FD3D12DynamicRHI::ResourceCast(StagingBufferRHI);
 	check(StagingBuffer);
-
-	FD3D12Resource* pResource = StagingBuffer->StagedRead.GetReference();
-	if (pResource)
-	{
-		pResource->Unmap();
-	}
+	StagingBuffer->Unlock();
 }
 
 // =============================================================================
diff --git a/Engine/Source/Runtime/D3D12RHI/Private/D3D12RHIPrivate.h b/Engine/Source/Runtime/D3D12RHI/Private/D3D12RHIPrivate.h
--- a/Engine/Source/Runtime/D3D12RHI/Private/D3D12RHIPrivate.h
+++ b/Engine/Source/Runtime/D3D12RHI/Private/D3D12RHIPrivate.h
@@ -235,7 +235,7 @@ public:
 	virtual FComputeShaderRHIRef RHICreateComputeShader(const TArray<uint8>& Code) override;
 	virtual FComputeFenceRHIRef RHICreateComputeFence(const FName& Name) final override;
 	virtual FGPUFenceRHIRef RHICreateGPUFence(const FName& Name) final override;
-	virtual FStagingBufferRHIRef RHICreateStagingBuffer(FVertexBufferRHIParamRef VertexBuffer) final override;
+	virtual FStagingBufferRHIRef RHICreateStagingBuffer() final override;
 	virtual void* RHILockStagingBuffer(FStagingBufferRHIParamRef StagingBuffer, uint32 Offset, uint32 SizeRHI) final override;
     virtual void RHIUnlockStagingBuffer(FStagingBufferRHIParamRef StagingBuffer) final override;
 	virtual FBoundShaderStateRHIRef RHICreateBoundShaderState(FVertexDeclarationRHIParamRef VertexDeclaration, FVertexShaderRHIParamRef VertexShader, FHullShaderRHIParamRef HullShader, FDomainShaderRHIParamRef DomainShader, FPixelShaderRHIParamRef PixelShader, FGeometryShaderRHIParamRef GeometryShader) final override;
diff --git a/Engine/Source/Runtime/D3D12RHI/Private/D3D12UAV.cpp b/Engine/Source/Runtime/D3D12RHI/Private/D3D12UAV.cpp
--- a/Engine/Source/Runtime/D3D12RHI/Private/D3D12UAV.cpp
+++ b/Engine/Source/Runtime/D3D12RHI/Private/D3D12UAV.cpp
@@ -260,4 +260,41 @@ FUnorderedAccessViewRHIRef FD3D12DynamicRHI::RHICreateUnorderedAccessView_Render
 		return RHICreateUnorderedAccessView(VertexBufferRHI, Format);
 	}
 	return RHICreateUnorderedAccessView(VertexBufferRHI, Format);
+}
+
+FD3D12StagingBuffer::~FD3D12StagingBuffer()
+{
+	if (StagedRead)
+	{
+		StagedRead.SafeRelease();
+	}
+}
+
+void* FD3D12StagingBuffer::Lock(uint32 Offset, uint32 NumBytes)
+{
+	check(!bIsLocked);
+	bIsLocked = true;
+	FD3D12Resource* pResource = StagedRead.GetReference();
+	if (pResource)
+	{
+		D3D12_RANGE ReadRange;
+		ReadRange.Begin = Offset;
+		ReadRange.End = Offset + NumBytes;
+		return reinterpret_cast<uint8*>(pResource->Map(&ReadRange)) + Offset;
+	}
+	else
+	{
+		return nullptr;
+	}
+}
+
+void FD3D12StagingBuffer::Unlock()
+{
+	check(bIsLocked);
+	bIsLocked = false;
+	FD3D12Resource* pResource = StagedRead.GetReference();
+	if (pResource)
+	{
+		pResource->Unmap();
+	}
 }
\ No newline at end of file
diff --git a/Engine/Source/Runtime/D3D12RHI/Public/D3D12Resources.h b/Engine/Source/Runtime/D3D12RHI/Public/D3D12Resources.h
--- a/Engine/Source/Runtime/D3D12RHI/Public/D3D12Resources.h
+++ b/Engine/Source/Runtime/D3D12RHI/Public/D3D12Resources.h
@@ -880,24 +880,36 @@ private:
 
 class FD3D12StagingBuffer : public FRHIStagingBuffer
 {
+	friend class FD3D12CommandContext;
+	friend class FD3D12DynamicRHI;
+
 public:
-	FD3D12StagingBuffer(FVertexBufferRHIRef InBuffer)
-		: FRHIStagingBuffer(InBuffer)
+	FD3D12StagingBuffer()
+		: FRHIStagingBuffer()
+		, ShadowBufferSize(0)
 	{}
 
+	virtual ~FD3D12StagingBuffer() final override;
+
+	virtual void* Lock(uint32 Offset, uint32 NumBytes) final override;
+	virtual void Unlock() final override;
+
+private:
 	TRefCountPtr<FD3D12Resource> StagedRead;
+	uint32 ShadowBufferSize;
 };
 
 class FD3D12GPUFence : public FRHIGPUFence
 {
 public:
-	FD3D12GPUFence(FName InName, FD3D12Fence* InFence) 
+	FD3D12GPUFence(FName InName, FD3D12Fence* InFence)
 		: FRHIGPUFence(InName)
 		, Fence(InFence)
 		, Value(0)
 	{}
 
 	void WriteInternal(ED3D12CommandQueueType QueueType);
+	virtual void Clear() final override;
 	virtual bool Poll() const final override;
 
 protected:
diff --git a/Engine/Source/Runtime/RHI/Private/RHICommandList.cpp b/Engine/Source/Runtime/RHI/Private/RHICommandList.cpp
--- a/Engine/Source/Runtime/RHI/Private/RHICommandList.cpp
+++ b/Engine/Source/Runtime/RHI/Private/RHICommandList.cpp
@@ -2069,13 +2069,17 @@ void FDynamicRHI::UnlockIndexBuffer_RenderThread(class FRHICommandListImmediate&
 	}
 }
 
+// @todo-mattc-staging Default implementation
 void* FDynamicRHI::RHILockStagingBuffer(FStagingBufferRHIParamRef StagingBuffer, uint32 Offset, uint32 SizeRHI)
 {
-	return GDynamicRHI->RHILockVertexBuffer(StagingBuffer->GetBackingBuffer(), Offset, SizeRHI, RLM_ReadOnly);
+	check(false);
+	return nullptr;
+	//return GDynamicRHI->RHILockVertexBuffer(StagingBuffer->GetSourceBuffer(), Offset, SizeRHI, RLM_ReadOnly);
 }
 void FDynamicRHI::RHIUnlockStagingBuffer(FStagingBufferRHIParamRef StagingBuffer)
 {
-	GDynamicRHI->RHIUnlockVertexBuffer(StagingBuffer->GetBackingBuffer());
+	check(false);
+	//GDynamicRHI->RHIUnlockVertexBuffer(StagingBuffer->GetSourceBuffer());
 }
 
 void* FDynamicRHI::LockStagingBuffer_RenderThread(class FRHICommandListImmediate& RHICmdList, FStagingBufferRHIParamRef StagingBuffer, uint32 Offset, uint32 SizeRHI)
@@ -2083,7 +2087,6 @@ void* FDynamicRHI::LockStagingBuffer_RenderThread(class FRHICommandListImmediate
 	QUICK_SCOPE_CYCLE_COUNTER(STAT_FDynamicRHI_LockStagingBuffer_RenderThread);
 	check(IsInRenderingThread());
 
-	FScopedRHIThreadStaller StallRHIThread(RHICmdList);
 	return GDynamicRHI->RHILockStagingBuffer(StagingBuffer, Offset, SizeRHI);
 }
 void FDynamicRHI::UnlockStagingBuffer_RenderThread(class FRHICommandListImmediate& RHICmdList, FStagingBufferRHIParamRef StagingBuffer)
@@ -2091,7 +2094,6 @@ void FDynamicRHI::UnlockStagingBuffer_RenderThread(class FRHICommandListImmediat
 	QUICK_SCOPE_CYCLE_COUNTER(STAT_FDynamicRHI_UnlockStagingBuffer_RenderThread);
 	check(IsInRenderingThread());
 	
-	FScopedRHIThreadStaller StallRHIThread(RHICmdList);
 	GDynamicRHI->RHIUnlockStagingBuffer(StagingBuffer);
 }
 
diff --git a/Engine/Source/Runtime/RHI/Private/RHIGPUReadback.cpp b/Engine/Source/Runtime/RHI/Private/RHIGPUReadback.cpp
--- a/Engine/Source/Runtime/RHI/Private/RHIGPUReadback.cpp
+++ b/Engine/Source/Runtime/RHI/Private/RHIGPUReadback.cpp
@@ -7,20 +7,19 @@
 
 #include "RHIGPUReadback.h"
 
+FGenericRHIGPUFence::FGenericRHIGPUFence(FName InName)
+	: FRHIGPUFence(InName)
+	, InsertedFrameNumber(MAX_uint32)
+{}
 
-void FRHIGPUFence::Write()
+void FGenericRHIGPUFence::Clear()
 {
-	UE_LOG(LogRHI, Warning, TEXT("FRHIGPUFence::Write is not implemented"));
+	InsertedFrameNumber = MAX_uint32;
 }
 
-FGenericRHIGPUFence::FGenericRHIGPUFence(FName InName)
-    : FRHIGPUFence(InName)
-    , InsertedFrameNumber(0)
-{}
-
-void FGenericRHIGPUFence::Write()
+void FGenericRHIGPUFence::WriteInternal()
 {
-	// GPU generally overlap the game. This overlap increases when using AFR.
+	// GPU generally overlap the game. This overlap increases when using AFR. In normal mode this can make us appear to be further behind the gpu than we actually are.
 	InsertedFrameNumber = GFrameNumberRenderThread + GNumAlternateFrameRenderingGroups;
 }
 
@@ -32,3 +31,18 @@ bool FGenericRHIGPUFence::Poll() const
 	}
 	return false;
 }
+
+void* FGenericRHIStagingBuffer::Lock(uint32 InOffset, uint32 NumBytes)
+{
+	check(ShadowBuffer);
+	check(!bIsLocked);
+	bIsLocked = true;
+	return reinterpret_cast<void*>(reinterpret_cast<uint8*>(RHILockVertexBuffer(ShadowBuffer, InOffset, NumBytes, RLM_ReadOnly)) + Offset);
+}
+
+void FGenericRHIStagingBuffer::Unlock()
+{
+	check(bIsLocked);
+	RHIUnlockVertexBuffer(ShadowBuffer);
+	bIsLocked = false;
+}
\ No newline at end of file
diff --git a/Engine/Source/Runtime/RHI/Public/DynamicRHI.h b/Engine/Source/Runtime/RHI/Public/DynamicRHI.h
--- a/Engine/Source/Runtime/RHI/Public/DynamicRHI.h
+++ b/Engine/Source/Runtime/RHI/Public/DynamicRHI.h
@@ -216,19 +216,18 @@ public:
 	}
 
 	/**
-	* Creates a staging buffer that uses 'BackingBuffer' as the backing store whose contents may be accessed using Lock/Unlock.
-	* @param BackingBuffer - The buffer whose contents may be accessed using Lock/Unlock.
+	* Creates a staging buffer, which is memory visible to the cpu without any locking.
 	* @return The new staging-buffer.
 	*/
 	// FlushType: Thread safe.	
-	virtual FStagingBufferRHIRef RHICreateStagingBuffer(FVertexBufferRHIParamRef BackingBuffer)
+	virtual FStagingBufferRHIRef RHICreateStagingBuffer()
 	{
-		return new FRHIStagingBuffer(BackingBuffer);
+		return new FGenericRHIStagingBuffer();
 	}
 
     /**
      * Lock a staging buffer to read contents on the CPU that were written by the GPU, without having to stall.
-     * @discussion This function requires that you have issued an EnqueueStagedRead invocation and verified that the FRHIGPUFence has been signaled before calling.
+     * @discussion This function requires that you have issued an CopyToStagingBuffer invocation and verified that the FRHIGPUFence has been signaled before calling.
      * @param StagingBuffer The buffer to lock.
      * @param Offset The offset in the buffer to return.
      * @param SizeRHI The length of the region in the buffer to lock.
@@ -244,7 +243,7 @@ public:
 
     /**
      * Lock a staging buffer to read contents on the CPU that were written by the GPU, without having to stall.
-     * @discussion This function requires that you have issued an EnqueueStagedRead invocation and verified that the FRHIGPUFence has been signaled before calling.
+     * @discussion This function requires that you have issued an CopyToStagingBuffer invocation and verified that the FRHIGPUFence has been signaled before calling.
      * @param RHICmdList The command-list to execute on or synchronize with.
      * @param StagingBuffer The buffer to lock.
      * @param Offset The offset in the buffer to return.
diff --git a/Engine/Source/Runtime/RHI/Public/RHICommandList.h b/Engine/Source/Runtime/RHI/Public/RHICommandList.h
--- a/Engine/Source/Runtime/RHI/Public/RHICommandList.h
+++ b/Engine/Source/Runtime/RHI/Public/RHICommandList.h
@@ -1447,19 +1447,21 @@ struct FRHICommandWaitComputeFence final : public FRHICommand<FRHICommandWaitCom
 };
 
 template<ECmdList CmdListType>
-struct FRHICommandEnqueueStagedRead final : public FRHICommand<FRHICommandEnqueueStagedRead<CmdListType>>
+struct FRHICommandCopyToStagingBuffer final : public FRHICommand<FRHICommandCopyToStagingBuffer<CmdListType>>
 {
-	FStagingBufferRHIParamRef StagingBuffer;
+	FVertexBufferRHIParamRef SourceBuffer;
+	FStagingBufferRHIParamRef DestinationStagingBuffer;
 	FGPUFenceRHIParamRef Fence;
 	uint32 Offset;
 	uint32 NumBytes;
 
-	FORCEINLINE_DEBUGGABLE FRHICommandEnqueueStagedRead(FStagingBufferRHIParamRef InStagingBuffer, FGPUFenceRHIParamRef InFence, uint32 InOffset, uint32 InNumBytes)
-		: StagingBuffer(InStagingBuffer)
+	FORCEINLINE_DEBUGGABLE FRHICommandCopyToStagingBuffer(FVertexBufferRHIParamRef InSourceBuffer, FStagingBufferRHIParamRef InDestinationStagingBuffer, uint32 InOffset, uint32 InNumBytes, FGPUFenceRHIParamRef InFence = nullptr)
+		: SourceBuffer(InSourceBuffer)
+		, DestinationStagingBuffer(InDestinationStagingBuffer)
 		, Fence(InFence)
 		, Offset(InOffset)
 		, NumBytes(InNumBytes)
-	{		
+	{
 	}
 	RHI_API void Execute(FRHICommandListBase& CmdList);
 };
@@ -2730,14 +2732,14 @@ public:
 		ALLOC_COMMAND(FRHICommandWaitComputeFence<ECmdList::EGfx>)(WaitFence);
 	}
 
-	FORCEINLINE_DEBUGGABLE void EnqueueStagedRead(FStagingBufferRHIParamRef StagingBuffer, FGPUFenceRHIParamRef Fence, uint32 Offset, uint32 NumBytes)
+	FORCEINLINE_DEBUGGABLE void CopyToStagingBuffer(FVertexBufferRHIParamRef SourceBuffer, FStagingBufferRHIParamRef DestinationStagingBuffer, uint32 Offset, uint32 NumBytes, FGPUFenceRHIParamRef Fence)
 	{
 		if (Bypass())
 		{
-			CMD_CONTEXT(RHIEnqueueStagedRead)(StagingBuffer, Fence, Offset, NumBytes);
+			CMD_CONTEXT(RHICopyToStagingBuffer)(SourceBuffer, DestinationStagingBuffer, Offset, NumBytes, Fence);
 			return;
 		}
-		ALLOC_COMMAND(FRHICommandEnqueueStagedRead<ECmdList::EGfx>)(StagingBuffer, Fence, Offset, NumBytes);
+		ALLOC_COMMAND(FRHICommandCopyToStagingBuffer<ECmdList::EGfx>)(SourceBuffer, DestinationStagingBuffer, Offset, NumBytes, Fence);
 	}
 
 	FORCEINLINE_DEBUGGABLE void BeginRenderPass(const FRHIRenderPassInfo& InInfo, const TCHAR* Name)
@@ -3127,14 +3129,14 @@ public:
 		ALLOC_COMMAND(FRHICommandWaitComputeFence<ECmdList::ECompute>)(WaitFence);
 	}
 
-	FORCEINLINE_DEBUGGABLE void EnqueueStagedRead(FStagingBufferRHIParamRef StagingBuffer, FGPUFenceRHIParamRef Fence, uint32 Offset, uint32 NumBytes)
+	FORCEINLINE_DEBUGGABLE void CopyToStagingBuffer(FVertexBufferRHIParamRef SourceBuffer, FStagingBufferRHIParamRef DestinationStagingBuffer, uint32 Offset, uint32 NumBytes, FGPUFenceRHIParamRef Fence = nullptr)
 	{
 		if (Bypass())
 		{
-			CMD_CONTEXT(RHIEnqueueStagedRead)(StagingBuffer, Fence, Offset, NumBytes);
+			CMD_CONTEXT(RHICopyToStagingBuffer)(SourceBuffer, DestinationStagingBuffer, Offset, NumBytes, Fence);
 			return;
 		}
-		ALLOC_COMMAND(FRHICommandEnqueueStagedRead<ECmdList::ECompute>)(StagingBuffer, Fence, Offset, NumBytes);
+		ALLOC_COMMAND(FRHICommandCopyToStagingBuffer<ECmdList::ECompute>)(SourceBuffer, DestinationStagingBuffer, Offset, NumBytes, Fence);
 	}
 };
 
@@ -3353,9 +3355,9 @@ public:
 		return GDynamicRHI->RHICreateGPUFence(Name);
 	}
 
-	FORCEINLINE FStagingBufferRHIRef CreateStagingBuffer(FVertexBufferRHIParamRef VertexBuffer)
+	FORCEINLINE FStagingBufferRHIRef CreateStagingBuffer()
 	{
-		return GDynamicRHI->RHICreateStagingBuffer(VertexBuffer);
+		return GDynamicRHI->RHICreateStagingBuffer();
 	}
 
 	FORCEINLINE FBoundShaderStateRHIRef CreateBoundShaderState(FVertexDeclarationRHIParamRef VertexDeclaration, FVertexShaderRHIParamRef VertexShader, FHullShaderRHIParamRef HullShader, FDomainShaderRHIParamRef DomainShader, FPixelShaderRHIParamRef PixelShader, FGeometryShaderRHIParamRef GeometryShader)
@@ -4291,9 +4293,9 @@ FORCEINLINE FGPUFenceRHIRef RHICreateGPUFence(const FName& Name)
 	return FRHICommandListExecutor::GetImmediateCommandList().CreateGPUFence(Name);
 }
 
-FORCEINLINE FStagingBufferRHIRef RHICreateStagingBuffer(FVertexBufferRHIParamRef VertexBuffer)
+FORCEINLINE FStagingBufferRHIRef RHICreateStagingBuffer()
 {
-	return FRHICommandListExecutor::GetImmediateCommandList().CreateStagingBuffer(VertexBuffer);
+	return FRHICommandListExecutor::GetImmediateCommandList().CreateStagingBuffer();
 }
 
 FORCEINLINE FIndexBufferRHIRef RHICreateAndLockIndexBuffer(uint32 Stride, uint32 Size, uint32 InUsage, FRHIResourceCreateInfo& CreateInfo, void*& OutDataBuffer)
diff --git a/Engine/Source/Runtime/RHI/Public/RHICommandListCommandExecutes.inl b/Engine/Source/Runtime/RHI/Public/RHICommandListCommandExecutes.inl
--- a/Engine/Source/Runtime/RHI/Public/RHICommandListCommandExecutes.inl
+++ b/Engine/Source/Runtime/RHI/Public/RHICommandListCommandExecutes.inl
@@ -493,13 +493,13 @@ template struct FRHICommandWaitComputeFence<ECmdList::EGfx>;
 template struct FRHICommandWaitComputeFence<ECmdList::ECompute>;
 
 template<ECmdList CmdListType>
-void FRHICommandEnqueueStagedRead<CmdListType>::Execute(FRHICommandListBase& CmdList)
+void FRHICommandCopyToStagingBuffer<CmdListType>::Execute(FRHICommandListBase& CmdList)
 {
 	RHISTAT(EnqueueStagedRead);
-	INTERNAL_DECORATOR_CONTEXT(RHIEnqueueStagedRead)(StagingBuffer, Fence, Offset, NumBytes);
+	INTERNAL_DECORATOR_CONTEXT(RHICopyToStagingBuffer)(SourceBuffer, DestinationStagingBuffer, Offset, NumBytes, Fence);
 }
-template struct FRHICommandEnqueueStagedRead<ECmdList::EGfx>;
-template struct FRHICommandEnqueueStagedRead<ECmdList::ECompute>;
+template struct FRHICommandCopyToStagingBuffer<ECmdList::EGfx>;
+template struct FRHICommandCopyToStagingBuffer<ECmdList::ECompute>;
 
 void FRHICommandBuildLocalGraphicsPipelineState::Execute(FRHICommandListBase& CmdList)
 {
diff --git a/Engine/Source/Runtime/RHI/Public/RHIContext.h b/Engine/Source/Runtime/RHI/Public/RHIContext.h
--- a/Engine/Source/Runtime/RHI/Public/RHIContext.h
+++ b/Engine/Source/Runtime/RHI/Public/RHIContext.h
@@ -127,19 +127,17 @@ public:
 	 */
 	virtual void RHIInvalidateCachedState() {};
 
-    /**
-     * Enqueues on the GPU timeline any necessary operations to make the contents of 'StagingBuffer' accessible to the CPU, flushing outstanding GPU writes and/or transferring from inaccessible non-unified GPU memory to local CPU memory.
-     * @param StagingBuffer The buffer to stage. Must not be null.
-     * @param Fence A GPU fence that will be inserted into the GPU timeline and which must then be tested on the CPU to know when the StagedRead was completed. Must not be null.
-     * @param Offset The start of the data in 'StagingBuffer' to make available.
-     * @param NumBytes The lenght of data in 'StagingBuffer' to make available.
-     */
-	virtual void RHIEnqueueStagedRead(FStagingBufferRHIParamRef StagingBuffer, FGPUFenceRHIParamRef Fence, uint32 InOffset, uint32 InNumBytes)
+	/**
+	 * Performs a copy of the data in 'SourceBuffer' to 'DestinationStagingBuffer.' This will occur inline on the GPU timeline. This is a mechanism to perform nonblocking readback of a buffer at a point in time.
+	 * @param SourceBuffer The source vertex buffer that will be inlined copied.
+	 * @param DestinationStagingBuffer The the host-visible destination buffer
+	 * @param Offset The start of the data in 'SourceBuffer'
+	 * @param NumBytes The number of bytes to copy out of 'SourceBuffer'
+	 * @param Fence (optional) A GPU fence that will be inserted into the GPU timeline and which must then be tested on the CPU to know when the Copy was completed. Can be NULL, though that implies you will not have any guarantees as to when it is safe to read from 'DestinationStagingBuffer'
+	 */
+	virtual void RHICopyToStagingBuffer(FVertexBufferRHIParamRef SourceBufferRHI, FStagingBufferRHIParamRef DestinationStagingBufferRHI, uint32 InOffset, uint32 InNumBytes, FGPUFenceRHIParamRef FenceRHI = nullptr)
 	{
-		if (Fence)
-		{
-			Fence->Write();
-		}
+		check(false);
 	}
 };
 
diff --git a/Engine/Source/Runtime/RHI/Public/RHIGPUReadback.h b/Engine/Source/Runtime/RHI/Public/RHIGPUReadback.h
--- a/Engine/Source/Runtime/RHI/Public/RHIGPUReadback.h
+++ b/Engine/Source/Runtime/RHI/Public/RHIGPUReadback.h
@@ -11,17 +11,16 @@
 #include "CoreMinimal.h"
 #include "RHI.h"
 
-/* FRHIGPUMemoryReadback: represents a memory readback request scheduled with RHIScheduleGPUMemoryUpdate
-*
+/* FRHIGPUMemoryReadback: Represents a memory readback request scheduled with CopyToStagingBuffer
+** Wraps a FRHIStagingBuffer with a FRHIGPUFence for synchronization.
 */
 class FRHIGPUMemoryReadback
 {
 public:
 
-	FRHIGPUMemoryReadback(FVertexBufferRHIRef InBuffer, FName RequestName)
+	FRHIGPUMemoryReadback(FName RequestName)
 	{
-		GPUVertexBuffer = InBuffer;
-		StagingBuffer = RHICreateStagingBuffer(InBuffer);
+		DestinationStagingBuffer = RHICreateStagingBuffer();
 		Fence = RHICreateGPUFence(RequestName);
 	}
 
@@ -29,33 +28,47 @@ public:
 	{
 	}
 
-	void Insert(FRHICommandList& RHICmdList, uint32 NumBytes = 0)
+	/**
+	 * Returns the CPU accessible pointer that backs this staging buffer.
+	 * @param RHICmdList The command list to enqueue the copy request on.
+	 * @param SourceBuffer The buffer holding the source data.
+	 * @param NumBytes The number of bytes to copy. If 0, this will copy the entire buffer.
+	 */
+
+	void EnqueueCopy(FRHICommandList& RHICmdList, FVertexBufferRHIParamRef SourceBuffer, uint32 NumBytes = 0)
 	{
-		RHICmdList.EnqueueStagedRead(StagingBuffer, Fence, 0, NumBytes ? NumBytes : GPUVertexBuffer->GetSize());
+		Fence->Clear();
+		RHICmdList.CopyToStagingBuffer(SourceBuffer, DestinationStagingBuffer, 0, NumBytes ? NumBytes : SourceBuffer->GetSize(), Fence);
 	}
 
+	/**
+	 * Indicates if the data is in place and ready to be read.
+	 */
 	bool IsReady()
 	{
 		return !Fence || Fence->Poll();
 	}
 
-	void *RetrieveData(uint32 NumBytes)
+	/**
+	 * Returns the CPU accessible pointer that backs this staging buffer.
+	 * @param NumBytes The maximum number of bytes the host will read from this pointer.
+	 * @returns A CPU accessible pointer to the backing buffer.
+	 */
+	void* Lock(uint32 NumBytes)
 	{
 		ensure(Fence->Poll());
-		return RHILockStagingBuffer(StagingBuffer, 0, NumBytes);
+		return RHILockStagingBuffer(DestinationStagingBuffer, 0, NumBytes);
 	}
 
-	void Finish()
+	/**
+	 * Signals that the host is finished reading from the backing buffer.
+	 */
+	void Unlock()
 	{
-		RHIUnlockStagingBuffer(StagingBuffer);
+		RHIUnlockStagingBuffer(DestinationStagingBuffer);
 	}
 
 private:
-	FStagingBufferRHIRef StagingBuffer;
-	FVertexBufferRHIRef  GPUVertexBuffer;
+	FStagingBufferRHIRef DestinationStagingBuffer;
 	FGPUFenceRHIRef Fence;
 };
-
-
-
-
diff --git a/Engine/Source/Runtime/RHI/Public/RHIResources.h b/Engine/Source/Runtime/RHI/Public/RHIResources.h
--- a/Engine/Source/Runtime/RHI/Public/RHIResources.h
+++ b/Engine/Source/Runtime/RHI/Public/RHIResources.h
@@ -774,8 +774,9 @@ public:
 //
 
 
-/* 
-* Generic GPU fence class used by FRHIGPUMemoryReadback
+/*
+* Generic GPU fence class.
+* Granularity differs depending on backing RHI - ie it may only represent command buffer granularity.
 * RHI specific fences derive from this to implement real GPU->CPU fencing.
 * The default implementation always returns false for Poll until the next frame from the frame the fence was inserted
 * because not all APIs have a GPU/CPU sync object, we need to fake it.
@@ -786,20 +787,15 @@ public:
 	FRHIGPUFence(FName InName) : FenceName(InName) {}
 	virtual ~FRHIGPUFence() {}
 
-    /**
-     * Signal this fence now, which completion can be tested with Poll or Wait.
-	 * RHI implementation might not need to implement this depending on how RHIInsertGPUFence is implemented.
-	 * Note also that RHIInsertGPUFence is currently only called by RHIEnqueueStagedRead.
-     */
-	virtual void Write();
-
-    /**
-     * Poll the fence to see if the GPU has signaled it.
-     * @returns True if and only if the GPU fence has been inserted and the GPU has signalled the fence.
-     */
+	virtual void Clear() = 0;
+
+	/**
+	 * Poll the fence to see if the GPU has signaled it.
+	 * @returns True if and only if the GPU fence has been inserted and the GPU has signaled the fence.
+	 */
 	virtual bool Poll() const = 0;
 
-private:
+protected:
 	FName FenceName;
 };
 
@@ -809,11 +805,13 @@ class RHI_API FGenericRHIGPUFence : public FRHIGPUFence
 public:
 	FGenericRHIGPUFence(FName InName);
 
-	virtual void Write() final override;
+	virtual void Clear() final override;
 
-    /** @discussion RHI implementations must be thread-safe and must correctly handle being called before RHIInsertFence if an RHI thread is active. */
+	/** @discussion RHI implementations must be thread-safe and must correctly handle being called before RHIInsertFence if an RHI thread is active. */
 	virtual bool Poll() const final override;
 
+	void WriteInternal();
+
 private:
 	uint32 InsertedFrameNumber;
 };
@@ -1010,19 +1008,34 @@ typedef TRefCountPtr<FRHIGraphicsPipelineState> FGraphicsPipelineStateRHIRef;
 /* Generic staging buffer class used by FRHIGPUMemoryReadback
 * RHI specific staging buffers derive from this
 */
-class FRHIStagingBuffer : public FRHIResource
+class RHI_API FRHIStagingBuffer : public FRHIResource
 {
 public:
-	FRHIStagingBuffer(FVertexBufferRHIParamRef InBuffer)
-		: BackingBuffer(InBuffer)
-	{
-	}
+	FRHIStagingBuffer()
+		: bIsLocked(false)
+	{}
 
-    /** Convenience function to access the vertex-buffer that acts as the backing-store. */
-	FVertexBufferRHIParamRef GetBackingBuffer() const { return BackingBuffer.GetReference(); }
+	virtual ~FRHIStagingBuffer() {}
 
+	virtual void *Lock(uint32 Offset, uint32 NumBytes) = 0;
+	virtual void Unlock() = 0;
 protected:
-	FVertexBufferRHIRef BackingBuffer;
+	bool bIsLocked;
+};
+
+class RHI_API FGenericRHIStagingBuffer : public FRHIStagingBuffer
+{
+public:
+	FGenericRHIStagingBuffer()
+		: FRHIStagingBuffer()
+	{}
+
+	~FGenericRHIStagingBuffer() {}
+
+	virtual void* Lock(uint32 Offset, uint32 NumBytes) final override;
+	virtual void Unlock() final override;
+	FVertexBufferRHIRef ShadowBuffer;
+	uint32 Offset;
 };
 
 typedef FRHIStagingBuffer*				FStagingBufferRHIParamRef;
diff --git a/Engine/Source/Runtime/VulkanRHI/Private/VulkanCommands.cpp b/Engine/Source/Runtime/VulkanRHI/Private/VulkanCommands.cpp
--- a/Engine/Source/Runtime/VulkanRHI/Private/VulkanCommands.cpp
+++ b/Engine/Source/Runtime/VulkanRHI/Private/VulkanCommands.cpp
@@ -1199,28 +1199,38 @@ void FVulkanCommandListContext::PrepareParallelFromBase(const FVulkanCommandList
 	TransitionAndLayoutManager.TempCopy(BaseContext.TransitionAndLayoutManager);
 }
 
-void FVulkanCommandListContext::RHIEnqueueStagedRead(FStagingBufferRHIParamRef StagingBufferRHI, FGPUFenceRHIParamRef FenceRHI, uint32 Offset, uint32 NumBytes)
+void FVulkanCommandListContext::RHICopyToStagingBuffer(FVertexBufferRHIParamRef SourceBufferRHI, FStagingBufferRHIParamRef StagingBufferRHI, uint32 Offset, uint32 NumBytes, FGPUFenceRHIParamRef FenceRHI)
 {
 	FVulkanCmdBuffer* CmdBuffer = CommandBufferManager->GetActiveCmdBuffer();
 
+	FVulkanVertexBuffer* VertexBuffer = ResourceCast(SourceBufferRHI);
 	FVulkanGPUFence* Fence = ResourceCast(FenceRHI);
 	Fence->CmdBuffer = CmdBuffer;
 	Fence->FenceSignaledCounter = CmdBuffer->GetFenceSignaledCounter();
 
 	ensure(CmdBuffer->IsOutsideRenderPass());
-	VulkanRHI::FStagingBuffer* ReadbackStagingBuffer = Device->GetStagingManager().AcquireBuffer(NumBytes);
 	FVulkanStagingBuffer* StagingBuffer = ResourceCast(StagingBufferRHI);
-	StagingBuffer->StagingBuffer = ReadbackStagingBuffer;
+	if (!StagingBuffer->StagingBuffer || StagingBuffer->StagingBuffer->GetSize() < NumBytes)
+	{
+		if (StagingBuffer->StagingBuffer)
+		{
+			Device->GetStagingManager().ReleaseBuffer(nullptr, StagingBuffer->StagingBuffer);
+		}
+
+		VulkanRHI::FStagingBuffer* ReadbackStagingBuffer = Device->GetStagingManager().AcquireBuffer(NumBytes);
+		StagingBuffer->StagingBuffer = ReadbackStagingBuffer;
+		StagingBuffer->Device = Device;
+	}
+
 	StagingBuffer->QueuedOffset = Offset;
 	StagingBuffer->QueuedNumBytes = NumBytes;
 
 	VkBufferCopy Region;
 	FMemory::Memzero(Region);
 	Region.size = NumBytes;
-	FVulkanVertexBuffer* VertexBuffer = ResourceCast(StagingBufferRHI->GetBackingBuffer());
 	Region.srcOffset = Offset + VertexBuffer->GetOffset();
 	//Region.dstOffset = 0;
-	VulkanRHI::vkCmdCopyBuffer(CmdBuffer->GetHandle(), VertexBuffer->GetHandle(), ReadbackStagingBuffer->GetHandle(), 1, &Region);
+	VulkanRHI::vkCmdCopyBuffer(CmdBuffer->GetHandle(), VertexBuffer->GetHandle(), StagingBuffer->StagingBuffer->GetHandle(), 1, &Region);
 }
 
 
diff --git a/Engine/Source/Runtime/VulkanRHI/Private/VulkanUAV.cpp b/Engine/Source/Runtime/VulkanRHI/Private/VulkanUAV.cpp
--- a/Engine/Source/Runtime/VulkanRHI/Private/VulkanUAV.cpp
+++ b/Engine/Source/Runtime/VulkanRHI/Private/VulkanUAV.cpp
@@ -343,6 +343,11 @@ FComputeFenceRHIRef FVulkanDynamicRHI::RHICreateComputeFence(const FName& Name)
 	return new FVulkanComputeFence(Device, Name);
 }
 
+void FVulkanGPUFence::Clear()
+{
+	CmdBuffer = nullptr;
+	FenceSignaledCounter = MAX_uint64;
+}
 /*
 FRenderQueryPoolRHIRef FVulkanDynamicRHI::RHICreateRenderQueryPool(ERenderQueryType QueryType, uint32 NumQueries)
 {
@@ -360,7 +365,7 @@ FRenderQueryPoolRHIRef FVulkanDynamicRHI::RHICreateRenderQueryPool(ERenderQueryT
 
 bool FVulkanGPUFence::Poll() const
 {
-	return FenceSignaledCounter < CmdBuffer->GetFenceSignaledCounter();
+	return (CmdBuffer && (FenceSignaledCounter < CmdBuffer->GetFenceSignaledCounter()));
 }
 
 FGPUFenceRHIRef FVulkanDynamicRHI::RHICreateGPUFence(const FName& Name)
diff --git a/Engine/Source/Runtime/VulkanRHI/Private/VulkanUtil.cpp b/Engine/Source/Runtime/VulkanRHI/Private/VulkanUtil.cpp
--- a/Engine/Source/Runtime/VulkanRHI/Private/VulkanUtil.cpp
+++ b/Engine/Source/Runtime/VulkanRHI/Private/VulkanUtil.cpp
@@ -118,30 +118,43 @@ FVulkanStagingBuffer::~FVulkanStagingBuffer()
 {
 	if (StagingBuffer)
 	{
-		FVulkanVertexBuffer* VertexBuffer = ResourceCast(GetBackingBuffer());
-		VertexBuffer->GetParent()->GetStagingManager().ReleaseBuffer(nullptr, StagingBuffer);
+		check(Device);
+		Device->GetStagingManager().ReleaseBuffer(nullptr, StagingBuffer);
 	}
 }
 
-FStagingBufferRHIRef FVulkanDynamicRHI::RHICreateStagingBuffer(FVertexBufferRHIParamRef BackingBuffer)
+void* FVulkanStagingBuffer::Lock(uint32 Offset, uint32 NumBytes)
 {
-	return new FVulkanStagingBuffer(BackingBuffer);
+	check(!bIsLocked);
+	bIsLocked = true;
+	uint32 QueuedEndOffset = QueuedNumBytes + QueuedOffset;
+	uint32 EndOffset = Offset + NumBytes;
+	check(Offset < QueuedNumBytes && EndOffset <= QueuedEndOffset);
+	//#todo-rco: Apply the offset in case it doesn't match
+	return (void*)((uint8*)StagingBuffer->GetMappedPointer() + Offset);
+}
+
+void FVulkanStagingBuffer::Unlock()
+{
+	check(bIsLocked);
+	bIsLocked = false;
 }
 
-void* FVulkanDynamicRHI::RHILockStagingBuffer(FStagingBufferRHIParamRef StagingBufferRHI, uint32 Offset, uint32 SizeRHI)
+FStagingBufferRHIRef FVulkanDynamicRHI::RHICreateStagingBuffer()
+{
+	return new FVulkanStagingBuffer();
+}
+
+void* FVulkanDynamicRHI::RHILockStagingBuffer(FStagingBufferRHIParamRef StagingBufferRHI, uint32 Offset, uint32 NumBytes)
 {
 	FVulkanStagingBuffer* StagingBuffer = ResourceCast(StagingBufferRHI);
-	uint32 QueuedEndOffset = StagingBuffer->QueuedNumBytes + StagingBuffer->QueuedOffset;
-	uint32 EndOffset = Offset + SizeRHI;
-	check(Offset < StagingBuffer->QueuedNumBytes && EndOffset <= QueuedEndOffset);
-	//#todo-rco: Apply the offset in case it doesn't match
-	return StagingBuffer->StagingBuffer->GetMappedPointer();
+	return StagingBuffer->Lock(Offset, NumBytes);
 }
 
 void FVulkanDynamicRHI::RHIUnlockStagingBuffer(FStagingBufferRHIParamRef StagingBufferRHI)
 {
 	FVulkanStagingBuffer* StagingBuffer = ResourceCast(StagingBufferRHI);
-	Device->GetStagingManager().ReleaseBuffer(nullptr, StagingBuffer->StagingBuffer);
+	StagingBuffer->Unlock();
 }
 
 /**
diff --git a/Engine/Source/Runtime/VulkanRHI/Public/VulkanContext.h b/Engine/Source/Runtime/VulkanRHI/Public/VulkanContext.h
--- a/Engine/Source/Runtime/VulkanRHI/Public/VulkanContext.h
+++ b/Engine/Source/Runtime/VulkanRHI/Public/VulkanContext.h
@@ -247,7 +247,7 @@ public:
 	virtual void RHICopyTexture(FTextureRHIParamRef SourceTexture, FTextureRHIParamRef DestTexture, const FRHICopyTextureInfo& CopyInfo) final override;
 	virtual void RHITransitionResources(EResourceTransitionAccess TransitionType, FTextureRHIParamRef* InRenderTargets, int32 NumTextures) final override;
 	virtual void RHITransitionResources(EResourceTransitionAccess TransitionType, EResourceTransitionPipeline TransitionPipeline, FUnorderedAccessViewRHIParamRef* InUAVs, int32 NumUAVs, FComputeFenceRHIParamRef WriteComputeFence) final override;
-	virtual void RHIEnqueueStagedRead(FStagingBufferRHIParamRef StagingBuffer, FGPUFenceRHIParamRef Fence, uint32 Offset, uint32 NumBytes) final override;
+	virtual void RHICopyToStagingBuffer(FVertexBufferRHIParamRef SourceBuffer, FStagingBufferRHIParamRef DestinationStagingBuffer, uint32 Offset, uint32 NumBytes, FGPUFenceRHIParamRef Fence) final override;
 
 	// Render time measurement
 	virtual void RHIBeginRenderQuery(FRenderQueryRHIParamRef RenderQuery) final override;
diff --git a/Engine/Source/Runtime/VulkanRHI/Public/VulkanDynamicRHI.h b/Engine/Source/Runtime/VulkanRHI/Public/VulkanDynamicRHI.h
--- a/Engine/Source/Runtime/VulkanRHI/Public/VulkanDynamicRHI.h
+++ b/Engine/Source/Runtime/VulkanRHI/Public/VulkanDynamicRHI.h
@@ -54,7 +54,7 @@ public:
 	virtual FComputeShaderRHIRef RHICreateComputeShader(const TArray<uint8>& Code) final override;
 	virtual FComputeFenceRHIRef RHICreateComputeFence(const FName& Name) final override;
 	virtual FGPUFenceRHIRef RHICreateGPUFence(const FName &Name) final override;
-	virtual FStagingBufferRHIRef RHICreateStagingBuffer(FVertexBufferRHIParamRef BackingBuffer) final override;
+	virtual FStagingBufferRHIRef RHICreateStagingBuffer() final override;
 	virtual FBoundShaderStateRHIRef RHICreateBoundShaderState(FVertexDeclarationRHIParamRef VertexDeclaration, FVertexShaderRHIParamRef VertexShader, FHullShaderRHIParamRef HullShader, FDomainShaderRHIParamRef DomainShader, FPixelShaderRHIParamRef PixelShader, FGeometryShaderRHIParamRef GeometryShader) final override;
 	virtual FGraphicsPipelineStateRHIRef RHICreateGraphicsPipelineState(const FGraphicsPipelineStateInitializer& Initializer) final override;
 	virtual FUniformBufferRHIRef RHICreateUniformBuffer(const void* Contents, const FRHIUniformBufferLayout& Layout, EUniformBufferUsage Usage) final override;
diff --git a/Engine/Source/Runtime/VulkanRHI/Public/VulkanResources.h b/Engine/Source/Runtime/VulkanRHI/Public/VulkanResources.h
--- a/Engine/Source/Runtime/VulkanRHI/Public/VulkanResources.h
+++ b/Engine/Source/Runtime/VulkanRHI/Public/VulkanResources.h
@@ -1520,17 +1520,25 @@ protected:
 
 class FVulkanStagingBuffer : public FRHIStagingBuffer
 {
+	friend class FVulkanCommandListContext;
 public:
-	FVulkanStagingBuffer(FVertexBufferRHIRef InBuffer)
-		: FRHIStagingBuffer(InBuffer)
+	FVulkanStagingBuffer()
+		: FRHIStagingBuffer()
 	{
+		check(!bIsLocked);
 	}
 
 	virtual ~FVulkanStagingBuffer();
 
+	virtual void* Lock(uint32 Offset, uint32 NumBytes) final override;
+	virtual void Unlock() final override;
+
+private:
 	VulkanRHI::FStagingBuffer* StagingBuffer = nullptr;
 	uint32 QueuedOffset = 0;
 	uint32 QueuedNumBytes = 0;
+	// The staging buffer was allocated from this device.
+	FVulkanDevice* Device;
 };
 
 class FVulkanGPUFence : public FRHIGPUFence
@@ -1541,6 +1549,7 @@ public:
 	{
 	}
 
+	virtual void Clear() final override;
 	virtual bool Poll() const final override;
 
 protected:
diff --git a/Engine/Source/Runtime/Windows/D3D11RHI/Private/D3D11Commands.cpp b/Engine/Source/Runtime/Windows/D3D11RHI/Private/D3D11Commands.cpp
--- a/Engine/Source/Runtime/Windows/D3D11RHI/Private/D3D11Commands.cpp
+++ b/Engine/Source/Runtime/Windows/D3D11RHI/Private/D3D11Commands.cpp
@@ -2143,32 +2143,42 @@ void FD3D11DynamicRHI::RHIFlushComputeShaderCache()
 
 //*********************** StagingBuffer Implementation ***********************//
 
-FStagingBufferRHIRef FD3D11DynamicRHI::RHICreateStagingBuffer(FVertexBufferRHIParamRef VertexBufferRHI)
+FStagingBufferRHIRef FD3D11DynamicRHI::RHICreateStagingBuffer()
 {
-	return new FD3D11StagingBuffer(VertexBufferRHI);
+	return new FD3D11StagingBuffer();
 }
 
-void FD3D11DynamicRHI::RHIEnqueueStagedRead(FStagingBufferRHIParamRef StagingBufferRHI, FGPUFenceRHIParamRef Fence, uint32 Offset, uint32 NumBytes)
+void FD3D11DynamicRHI::RHICopyToStagingBuffer(FVertexBufferRHIParamRef SourceBufferRHI, FStagingBufferRHIParamRef StagingBufferRHI, uint32 Offset, uint32 NumBytes, FGPUFenceRHIParamRef FenceRHI)
 {
+	FD3D11VertexBuffer* SourceBuffer = ResourceCast(SourceBufferRHI);
 	FD3D11StagingBuffer* StagingBuffer = ResourceCast(StagingBufferRHI);
+	ensureMsgf(!StagingBuffer->bIsLocked, TEXT("Attempting to Copy to a locked staging buffer. This may have undefined behavior"));
 	if (StagingBuffer)
 	{
-		FD3D11VertexBuffer* VertexBuffer = ResourceCast(StagingBuffer->GetBackingBuffer());
-		if (VertexBuffer)
+		if (SourceBuffer)
 		{
-			// Free previously allocated buffer.
-			StagingBuffer->StagedRead.SafeRelease();
-
-			// If the static buffer is being locked for reading, create a staging buffer.
-			D3D11_BUFFER_DESC StagedReadDesc;
-			ZeroMemory(&StagedReadDesc, sizeof(D3D11_BUFFER_DESC));
-			StagedReadDesc.ByteWidth = NumBytes;
-			StagedReadDesc.Usage = D3D11_USAGE_STAGING;
-			StagedReadDesc.BindFlags = 0;
-			StagedReadDesc.CPUAccessFlags = D3D11_CPU_ACCESS_READ;
-			StagedReadDesc.MiscFlags = 0;
-			TRefCountPtr<ID3D11Buffer> StagingVertexBuffer;
-			VERIFYD3D11RESULT_EX(Direct3DDevice->CreateBuffer(&StagedReadDesc, NULL, StagingBuffer->StagedRead.GetInitReference()), Direct3DDevice);
+			if (!StagingBuffer->StagedRead || StagingBuffer->ShadowBufferSize < NumBytes)
+			{
+				// Free previously allocated buffer.
+				if (StagingBuffer->StagedRead)
+				{
+					StagingBuffer->StagedRead.SafeRelease();
+				}
+
+				// Allocate a new one with enough space.
+				// @todo-mattc I feel like we should allocate more than NumBytes to handle small reads without blowing tons of space. Need to pool this.
+				D3D11_BUFFER_DESC StagedReadDesc;
+				ZeroMemory(&StagedReadDesc, sizeof(D3D11_BUFFER_DESC));
+				StagedReadDesc.ByteWidth = NumBytes;
+				StagedReadDesc.Usage = D3D11_USAGE_STAGING;
+				StagedReadDesc.BindFlags = 0;
+				StagedReadDesc.CPUAccessFlags = D3D11_CPU_ACCESS_READ;
+				StagedReadDesc.MiscFlags = 0;
+				TRefCountPtr<ID3D11Buffer> StagingVertexBuffer;
+				VERIFYD3D11RESULT_EX(Direct3DDevice->CreateBuffer(&StagedReadDesc, NULL, StagingBuffer->StagedRead.GetInitReference()), Direct3DDevice);
+				StagingBuffer->ShadowBufferSize = NumBytes;
+				StagingBuffer->Context = Direct3DDeviceIMContext;
+			}
 
 			// Copy the contents of the vertex buffer to the staging buffer.
 			D3D11_BOX SourceBox;
@@ -2176,38 +2186,29 @@ void FD3D11DynamicRHI::RHIEnqueueStagedRead(FStagingBufferRHIParamRef StagingBuf
 			SourceBox.right = NumBytes;
 			SourceBox.top = SourceBox.front = 0;
 			SourceBox.bottom = SourceBox.back = 1;
-			Direct3DDeviceIMContext->CopySubresourceRegion(StagingBuffer->StagedRead, 0, 0, 0, 0, VertexBuffer->Resource, 0, &SourceBox);
+			Direct3DDeviceIMContext->CopySubresourceRegion(StagingBuffer->StagedRead, 0, 0, 0, 0, SourceBuffer->Resource, 0, &SourceBox);
 		}
 	}
 
+	// @todo-staging Implement real fences for D3D11
+	// D3D11 only has the generic fence for now.
+	FGenericRHIGPUFence* Fence = ResourceCast(FenceRHI);
+
 	if (Fence)
 	{
-		Fence->Write();
+		Fence->WriteInternal();
 	}
 }
 
 void* FD3D11DynamicRHI::RHILockStagingBuffer(FStagingBufferRHIParamRef StagingBufferRHI, uint32 Offset, uint32 SizeRHI)
 {
+	check(StagingBufferRHI);
 	FD3D11StagingBuffer* StagingBuffer = ResourceCast(StagingBufferRHI);
-	if (StagingBuffer && StagingBuffer->StagedRead)
-	{
-		// Map the staging buffer's memory for reading.
-		D3D11_MAPPED_SUBRESOURCE MappedSubresource;
-		VERIFYD3D11RESULT_EX(Direct3DDeviceIMContext->Map(StagingBuffer->StagedRead ,0,D3D11_MAP_READ,0,&MappedSubresource), Direct3DDevice);
-
-		return (void*)((uint8*)MappedSubresource.pData + Offset);
-	}
-	else
-	{
-		return nullptr;
-	}
+	return StagingBuffer->Lock(Offset, SizeRHI);
 }
 
 void FD3D11DynamicRHI::RHIUnlockStagingBuffer(FStagingBufferRHIParamRef StagingBufferRHI)
 {
 	FD3D11StagingBuffer* StagingBuffer = ResourceCast(StagingBufferRHI);
-	if (StagingBuffer && StagingBuffer->StagedRead)
-	{
-		Direct3DDeviceIMContext->Unmap(StagingBuffer->StagedRead,0);
-	}
+	StagingBuffer->Unlock();
 }
diff --git a/Engine/Source/Runtime/Windows/D3D11RHI/Private/D3D11RHIPrivate.h b/Engine/Source/Runtime/Windows/D3D11RHI/Private/D3D11RHIPrivate.h
--- a/Engine/Source/Runtime/Windows/D3D11RHI/Private/D3D11RHIPrivate.h
+++ b/Engine/Source/Runtime/Windows/D3D11RHI/Private/D3D11RHIPrivate.h
@@ -395,8 +395,8 @@ public:
 	virtual FGeometryShaderRHIRef RHICreateGeometryShader(const TArray<uint8>& Code) final override;
 	virtual FGeometryShaderRHIRef RHICreateGeometryShaderWithStreamOutput(const TArray<uint8>& Code, const FStreamOutElementList& ElementList, uint32 NumStrides, const uint32* Strides, int32 RasterizedStream) final override;
 	virtual FComputeShaderRHIRef RHICreateComputeShader(const TArray<uint8>& Code) final override;
-	virtual FStagingBufferRHIRef RHICreateStagingBuffer(FVertexBufferRHIParamRef VertexBuffer) final override;
-    virtual void RHIEnqueueStagedRead(FStagingBufferRHIParamRef StagingBuffer, FGPUFenceRHIParamRef Fence, uint32 Offset, uint32 NumBytes) final override;
+	virtual FStagingBufferRHIRef RHICreateStagingBuffer() final override;
+	virtual void RHICopyToStagingBuffer(FVertexBufferRHIParamRef SourceBuffer, FStagingBufferRHIParamRef DestinationStagingBuffer, uint32 Offset, uint32 NumBytes, FGPUFenceRHIParamRef Fence = nullptr) final override;
 	virtual void* RHILockStagingBuffer(FStagingBufferRHIParamRef StagingBuffer, uint32 Offset, uint32 SizeRHI) final override;
     virtual void RHIUnlockStagingBuffer(FStagingBufferRHIParamRef StagingBuffer) final override;
 	virtual FBoundShaderStateRHIRef RHICreateBoundShaderState(FVertexDeclarationRHIParamRef VertexDeclaration, FVertexShaderRHIParamRef VertexShader, FHullShaderRHIParamRef HullShader, FDomainShaderRHIParamRef DomainShader, FPixelShaderRHIParamRef PixelShader, FGeometryShaderRHIParamRef GeometryShader) final override;
diff --git a/Engine/Source/Runtime/Windows/D3D11RHI/Private/D3D11UAV.cpp b/Engine/Source/Runtime/Windows/D3D11RHI/Private/D3D11UAV.cpp
--- a/Engine/Source/Runtime/Windows/D3D11RHI/Private/D3D11UAV.cpp
+++ b/Engine/Source/Runtime/Windows/D3D11RHI/Private/D3D11UAV.cpp
@@ -285,3 +285,39 @@ void FD3D11DynamicRHI::RHIBindDebugLabelName(FUnorderedAccessViewRHIParamRef Uno
 	UAV->View->SetPrivateData(WKPDID_D3DDebugObjectName, FCString::Strlen(Name) + 1, TCHAR_TO_ANSI(Name));
 #endif
 }
+
+FD3D11StagingBuffer::~FD3D11StagingBuffer()
+{
+	if (StagedRead)
+	{
+		StagedRead.SafeRelease();
+	}
+}
+
+void* FD3D11StagingBuffer::Lock(uint32 Offset, uint32 NumBytes)
+{
+	check(!bIsLocked);
+	bIsLocked = true;
+	if (StagedRead)
+	{
+		// Map the staging buffer's memory for reading.
+		D3D11_MAPPED_SUBRESOURCE MappedSubresource;
+		VERIFYD3D11RESULT(Context->Map(StagedRead, 0, D3D11_MAP_READ, 0, &MappedSubresource));
+
+		return (void*)((uint8*)MappedSubresource.pData + Offset);
+	}
+	else
+	{
+		return nullptr;
+	}
+}
+
+void FD3D11StagingBuffer::Unlock()
+{
+	check(bIsLocked);
+	bIsLocked = false;
+	if (StagedRead)
+	{
+		Context->Unmap(StagedRead, 0);
+	}
+}
diff --git a/Engine/Source/Runtime/Windows/D3D11RHI/Public/D3D11Resources.h b/Engine/Source/Runtime/Windows/D3D11RHI/Public/D3D11Resources.h
--- a/Engine/Source/Runtime/Windows/D3D11RHI/Public/D3D11Resources.h
+++ b/Engine/Source/Runtime/Windows/D3D11RHI/Public/D3D11Resources.h
@@ -784,12 +784,21 @@ public:
 
 class FD3D11StagingBuffer : public FRHIStagingBuffer
 {
+	friend class FD3D11DynamicRHI;
 public:
-	FD3D11StagingBuffer(FVertexBufferRHIRef InBuffer)
-		: FRHIStagingBuffer(InBuffer)
+	FD3D11StagingBuffer()
+		: FRHIStagingBuffer()
 	{}
 
+	virtual ~FD3D11StagingBuffer() final override;
+
+	virtual void* Lock(uint32 Offset, uint32 NumBytes) final override;
+	virtual void Unlock() final override;
+
+private:
+	FD3D11DeviceContext* Context;
 	TRefCountPtr<ID3D11Buffer> StagedRead;
+	uint32 ShadowBufferSize;
 };
 
 /** Shader resource view class. */
@@ -924,6 +933,12 @@ struct TD3D11ResourceTraits<FRHIStagingBuffer>
 {
 	typedef FD3D11StagingBuffer TConcreteType;
 };
+// @todo-staging Implement D3D11 fences.
+template<>
+struct TD3D11ResourceTraits<FRHIGPUFence>
+{
+	typedef FGenericRHIGPUFence TConcreteType;
+};
 template<>
 struct TD3D11ResourceTraits<FRHIShaderResourceView>
 {
