Fixed division by 0

diff --git a/Engine/Plugins/FX/Niagara/Source/Niagara/Private/NiagaraDataInterfaceVectorField.cpp b/Engine/Plugins/FX/Niagara/Source/Niagara/Private/NiagaraDataInterfaceVectorField.cpp
--- a/Engine/Plugins/FX/Niagara/Source/Niagara/Private/NiagaraDataInterfaceVectorField.cpp
+++ b/Engine/Plugins/FX/Niagara/Source/Niagara/Private/NiagaraDataInterfaceVectorField.cpp
@@ -474,6 +474,8 @@ void UNiagaraDataInterfaceVectorField::SampleVectorField(FVectorVMContext& Conte
 	UVectorFieldStatic* StaticVectorField = Cast<UVectorFieldStatic>(Field);
 	UVectorFieldAnimated* AnimatedVectorField = Cast<UVectorFieldAnimated>(Field);
 
+	bool bSuccess = false;
+
 	if (StaticVectorField != nullptr && StaticVectorField->bAllowCPUAccess)
 	{
 		const FVector4 TilingAxes = FVector4(bTileX ? 1.0f : 0.0f, bTileY ? 1.0f : 0.0f, bTileZ ? 1.0f : 0.0f, 0.0);
@@ -483,96 +485,103 @@ void UNiagaraDataInterfaceVectorField::SampleVectorField(FVectorVMContext& Conte
 		const uint32 SizeZ = (uint32)StaticVectorField->SizeZ;
 		const FVector4 Size(SizeX, SizeY, SizeZ, 1.0f);
 
-		const FVector4 MinBounds(StaticVectorField->Bounds.Min.X, StaticVectorField->Bounds.Min.Y, StaticVectorField->Bounds.Min.Z, 1.0f);
-		const FVector4 MaxBounds(StaticVectorField->Bounds.Max.X, StaticVectorField->Bounds.Max.Y, StaticVectorField->Bounds.Max.Z, 1.0f);
+		const FVector4 MinBounds(StaticVectorField->Bounds.Min.X, StaticVectorField->Bounds.Min.Y, StaticVectorField->Bounds.Min.Z, 0.f);
+		const FVector BoundSize = StaticVectorField->Bounds.GetSize();
 
 		const FVector4 *Data = StaticVectorField->CPUData.GetData();
 
-		check(Data != nullptr); // TODO(mv): Guard better? This shouldn't happen 
-
-		// Math helper
-		static auto FVector4Clamp = [](FVector4 v, FVector4 a, FVector4 b) {
-			return FVector4(FMath::Clamp(v.X, a.X, b.X),
-				FMath::Clamp(v.Y, a.Y, b.Y),
-				FMath::Clamp(v.Z, a.Z, b.Z),
-				FMath::Clamp(v.W, a.W, b.W));
-		};
-
-		static auto FVector4Floor = [](FVector4 v) {
-			return FVector4(FGenericPlatformMath::FloorToFloat(v.X),
-				FGenericPlatformMath::FloorToFloat(v.Y),
-				FGenericPlatformMath::FloorToFloat(v.Z),
-				FGenericPlatformMath::FloorToFloat(v.W));
-		};
-
-		for (int32 InstanceIdx = 0; InstanceIdx < Context.NumInstances; ++InstanceIdx)
+		if (ensure(Data && FMath::Min3(SizeX, SizeY, SizeZ) > 0 && BoundSize.GetMin() > SMALL_NUMBER))
 		{
-			// Position in Volume Space
-			FVector4 Pos(XParam.Get(), YParam.Get(), ZParam.Get(), 0.0f);
-
-			// Normalize position
-			Pos = (Pos - MinBounds) / (MaxBounds - MinBounds);
-
-			// Scaled position
-			Pos = Pos * Size;
-
-			// Offset by half a cell size due to sample being in the center of its cell
-			Pos = Pos - FVector4(0.5f, 0.5f, 0.5f, 0.0f);
-
-			// 
-			FVector4 Index0 = FVector4Floor(Pos);
-			FVector4 Index1 = Index0 + FVector4(1.0f, 1.0f, 1.0f, 0.0f);
-
-			// 
-			FVector4 Fraction = Pos - Index0;
-
-			Index0 = Index0 - TilingAxes*FVector4Floor(Index0 / Size)*Size;
-			Index1 = Index1 - TilingAxes*FVector4Floor(Index1 / Size)*Size;
-
-			Index0 = FVector4Clamp(Index0, FVector4(0.0f), Size - FVector4(1.0f, 1.0f, 1.0f, 0.0f));
-			Index1 = FVector4Clamp(Index1, FVector4(0.0f), Size - FVector4(1.0f, 1.0f, 1.0f, 0.0f));
-
-			// Sample by regular trilinear interpolation:
-
-			// TODO(mv): Optimize indexing for cache? Periodicity is problematic...
-			// TODO(mv): Vectorize?
-			// Fetch corners
-			FVector4 V000 = Data[int(Index0.X + SizeX * Index0.Y + SizeX * SizeY * Index0.Z)];
-			FVector4 V100 = Data[int(Index1.X + SizeX * Index0.Y + SizeX * SizeY * Index0.Z)];
-			FVector4 V010 = Data[int(Index0.X + SizeX * Index1.Y + SizeX * SizeY * Index0.Z)];
-			FVector4 V110 = Data[int(Index1.X + SizeX * Index1.Y + SizeX * SizeY * Index0.Z)];
-			FVector4 V001 = Data[int(Index0.X + SizeX * Index0.Y + SizeX * SizeY * Index1.Z)];
-			FVector4 V101 = Data[int(Index1.X + SizeX * Index0.Y + SizeX * SizeY * Index1.Z)];
-			FVector4 V011 = Data[int(Index0.X + SizeX * Index1.Y + SizeX * SizeY * Index1.Z)];
-			FVector4 V111 = Data[int(Index1.X + SizeX * Index1.Y + SizeX * SizeY * Index1.Z)];
-
-			// Blend x-axis
-			FVector4 V00 = FMath::Lerp(V000, V100, Fraction.X);
-			FVector4 V01 = FMath::Lerp(V001, V101, Fraction.X);
-			FVector4 V10 = FMath::Lerp(V010, V110, Fraction.X);
-			FVector4 V11 = FMath::Lerp(V011, V111, Fraction.X);
-
-			// Blend y-axis
-			FVector4 V0 = FMath::Lerp(V00, V10, Fraction.Y);
-			FVector4 V1 = FMath::Lerp(V01, V11, Fraction.Y);
-
-			// Blend z-axis
-			FVector4 V = FMath::Lerp(V0, V1, Fraction.Z);
-
-			// Write final output...
-			*OutSampleX.GetDest() = V.X;
-			*OutSampleY.GetDest() = V.Y;
-			*OutSampleZ.GetDest() = V.Z;
+			const FVector4 OneOverBoundSize(FVector::OneVector / BoundSize, 1.0f);
+
+			// Math helper
+			static auto FVector4Clamp = [](FVector4 v, FVector4 a, FVector4 b) {
+				return FVector4(FMath::Clamp(v.X, a.X, b.X),
+					FMath::Clamp(v.Y, a.Y, b.Y),
+					FMath::Clamp(v.Z, a.Z, b.Z),
+					FMath::Clamp(v.W, a.W, b.W));
+			};
+
+			static auto FVector4Floor = [](FVector4 v) {
+				return FVector4(FGenericPlatformMath::FloorToFloat(v.X),
+					FGenericPlatformMath::FloorToFloat(v.Y),
+					FGenericPlatformMath::FloorToFloat(v.Z),
+					FGenericPlatformMath::FloorToFloat(v.W));
+			};
+
+			for (int32 InstanceIdx = 0; InstanceIdx < Context.NumInstances; ++InstanceIdx)
+			{
+				// Position in Volume Space
+				FVector4 Pos(XParam.Get(), YParam.Get(), ZParam.Get(), 0.0f);
+
+				// Normalize position
+
+				Pos = (Pos - MinBounds) * OneOverBoundSize;
+
+				// Scaled position
+				Pos = Pos * Size;
+
+				// Offset by half a cell size due to sample being in the center of its cell
+				Pos = Pos - FVector4(0.5f, 0.5f, 0.5f, 0.0f);
+
+				// 
+				FVector4 Index0 = FVector4Floor(Pos);
+				FVector4 Index1 = Index0 + FVector4(1.0f, 1.0f, 1.0f, 0.0f);
+
+				// 
+				FVector4 Fraction = Pos - Index0;
+
+				Index0 = Index0 - TilingAxes*FVector4Floor(Index0 / Size)*Size;
+				Index1 = Index1 - TilingAxes*FVector4Floor(Index1 / Size)*Size;
+
+				Index0 = FVector4Clamp(Index0, FVector4(0.0f), Size - FVector4(1.0f, 1.0f, 1.0f, 0.0f));
+				Index1 = FVector4Clamp(Index1, FVector4(0.0f), Size - FVector4(1.0f, 1.0f, 1.0f, 0.0f));
+
+				// Sample by regular trilinear interpolation:
+
+				// TODO(mv): Optimize indexing for cache? Periodicity is problematic...
+				// TODO(mv): Vectorize?
+				// Fetch corners
+				FVector4 V000 = Data[int(Index0.X + SizeX * Index0.Y + SizeX * SizeY * Index0.Z)];
+				FVector4 V100 = Data[int(Index1.X + SizeX * Index0.Y + SizeX * SizeY * Index0.Z)];
+				FVector4 V010 = Data[int(Index0.X + SizeX * Index1.Y + SizeX * SizeY * Index0.Z)];
+				FVector4 V110 = Data[int(Index1.X + SizeX * Index1.Y + SizeX * SizeY * Index0.Z)];
+				FVector4 V001 = Data[int(Index0.X + SizeX * Index0.Y + SizeX * SizeY * Index1.Z)];
+				FVector4 V101 = Data[int(Index1.X + SizeX * Index0.Y + SizeX * SizeY * Index1.Z)];
+				FVector4 V011 = Data[int(Index0.X + SizeX * Index1.Y + SizeX * SizeY * Index1.Z)];
+				FVector4 V111 = Data[int(Index1.X + SizeX * Index1.Y + SizeX * SizeY * Index1.Z)];
+
+				// Blend x-axis
+				FVector4 V00 = FMath::Lerp(V000, V100, Fraction.X);
+				FVector4 V01 = FMath::Lerp(V001, V101, Fraction.X);
+				FVector4 V10 = FMath::Lerp(V010, V110, Fraction.X);
+				FVector4 V11 = FMath::Lerp(V011, V111, Fraction.X);
+
+				// Blend y-axis
+				FVector4 V0 = FMath::Lerp(V00, V10, Fraction.Y);
+				FVector4 V1 = FMath::Lerp(V01, V11, Fraction.Y);
+
+				// Blend z-axis
+				FVector4 V = FMath::Lerp(V0, V1, Fraction.Z);
+
+				// Write final output...
+				*OutSampleX.GetDest() = V.X;
+				*OutSampleY.GetDest() = V.Y;
+				*OutSampleZ.GetDest() = V.Z;
+
+				XParam.Advance();
+				YParam.Advance();
+				ZParam.Advance();
+				OutSampleX.Advance();
+				OutSampleY.Advance();
+				OutSampleZ.Advance();
+			}
 
-			XParam.Advance();
-			YParam.Advance();
-			ZParam.Advance();
-			OutSampleX.Advance();
-			OutSampleY.Advance();
-			OutSampleZ.Advance();
+			bSuccess = true;
 		}
 	}
-	else
+
+	if (!bSuccess)
 	{
 		// TODO(mv): Add warnings?
 		if (StaticVectorField != nullptr && !StaticVectorField->bAllowCPUAccess)
