Fixing CIS

diff --git a/Engine/Plugins/FX/Niagara/Source/NiagaraShader/Private/NiagaraShader.cpp b/Engine/Plugins/FX/Niagara/Source/NiagaraShader/Private/NiagaraShader.cpp
--- a/Engine/Plugins/FX/Niagara/Source/NiagaraShader/Private/NiagaraShader.cpp
+++ b/Engine/Plugins/FX/Niagara/Source/NiagaraShader/Private/NiagaraShader.cpp
@@ -497,7 +497,7 @@ TArray<uint8>* FNiagaraShaderMap::BackupShadersToMemory()
 	TArray<uint8>* SavedShaderData = new TArray<uint8>();
 	FMemoryWriter Ar(*SavedShaderData);
 
-	SerializeInline(Ar, true, true);
+	SerializeInline(Ar, true, true, false);
 	RegisterSerializedShaders(false);
 	Empty();
 
@@ -507,7 +507,7 @@ TArray<uint8>* FNiagaraShaderMap::BackupShadersToMemory()
 void FNiagaraShaderMap::RestoreShadersFromMemory(const TArray<uint8>& ShaderData)
 {
 	FMemoryReader Ar(ShaderData);
-	SerializeInline(Ar, true, true);
+	SerializeInline(Ar, true, true, false);
 	RegisterSerializedShaders(false);
 }
 
@@ -564,7 +564,7 @@ void FNiagaraShaderMap::SaveForRemoteRecompile(FArchive& Ar, const TMap<FString,
 
 	for (int32 Index = 0; Index < NumUniqueResources; Index++)
 	{
-		UniqueResources[Index]->Serialize(Ar);
+		UniqueResources[Index]->Serialize(Ar, false);
 	}
 
 	// now we serialize a map (for each script)
@@ -1092,13 +1092,13 @@ void FNiagaraShaderMap::Serialize(FArchive& Ar, bool bInlineShaderResources)
 
 	if (Ar.IsSaving())
 	{
-		TShaderMap<FNiagaraShaderType>::SerializeInline(Ar, bInlineShaderResources, false);
+		TShaderMap<FNiagaraShaderType>::SerializeInline(Ar, bInlineShaderResources, false, false);
 		RegisterSerializedShaders(false);
 	}
 
 	if (Ar.IsLoading())
 	{
-		TShaderMap<FNiagaraShaderType>::SerializeInline(Ar, bInlineShaderResources, false);
+		TShaderMap<FNiagaraShaderType>::SerializeInline(Ar, bInlineShaderResources, false, false);
 	}
 }
 
