Copying //UE4/Dev-Niagara to Dev-Main (//UE4/Dev-Main) as of Niagara CL 4733936
#tests Niagara auto-tests for DX11 pass

diff --git a/Engine/Plugins/FX/Niagara/Source/Niagara/Private/NiagaraScript.cpp b/Engine/Plugins/FX/Niagara/Source/Niagara/Private/NiagaraScript.cpp
--- a/Engine/Plugins/FX/Niagara/Source/Niagara/Private/NiagaraScript.cpp
+++ b/Engine/Plugins/FX/Niagara/Source/Niagara/Private/NiagaraScript.cpp
@@ -16,6 +16,9 @@
 #include "UObject/Linker.h"
 #include "HAL/PlatformFilemanager.h"
 #include "Misc/FileHelper.h"
+#include "UObject/EditorObjectVersion.h"
+#include "UObject/ReleaseObjectVersion.h"
+
 
 #if WITH_EDITOR
 	#include "NiagaraScriptDerivedData.h"
@@ -28,13 +31,21 @@
 
 DECLARE_STATS_GROUP(TEXT("Niagara Detailed"), STATGROUP_NiagaraDetailed, STATCAT_Advanced);
 
-FNiagaraScriptDebuggerInfo::FNiagaraScriptDebuggerInfo() : FrameLastWriteId(-1)
+FNiagaraScriptDebuggerInfo::FNiagaraScriptDebuggerInfo() : bWaitForGPU(false), FrameLastWriteId(-1), bWritten(false)
 {
 }
 
 
-FNiagaraScriptDebuggerInfo::FNiagaraScriptDebuggerInfo(FName InName, ENiagaraScriptUsage InUsage, const FGuid& InUsageId) : HandleName(InName), Usage(InUsage), UsageId(InUsageId), FrameLastWriteId(-1)
+FNiagaraScriptDebuggerInfo::FNiagaraScriptDebuggerInfo(FName InName, ENiagaraScriptUsage InUsage, const FGuid& InUsageId) : HandleName(InName), Usage(InUsage), UsageId(InUsageId), FrameLastWriteId(-1), bWritten(false)
 {
+	if (InUsage == ENiagaraScriptUsage::ParticleGPUComputeScript)
+	{
+		bWaitForGPU = true;
+	}
+	else
+	{
+		bWaitForGPU = false;
+	}
 }
 
 
@@ -232,6 +243,10 @@ void UNiagaraScript::ComputeVMCompilationId(FNiagaraVMExecutableDataId& Id) cons
 		{
 			Id.AdditionalDefines.Add(TEXT("Emitter.Localspace"));
 		}
+		if (Emitter->bDeterminism)
+		{
+			Id.AdditionalDefines.Add(TEXT("Emitter.Determinism"));
+		}
 	}
 
 	if (UNiagaraSystem* System = Cast<UNiagaraSystem>(Obj))
@@ -245,6 +260,10 @@ void UNiagaraScript::ComputeVMCompilationId(FNiagaraVMExecutableDataId& Id) cons
 				{
 					Id.AdditionalDefines.Add(Emitter->GetUniqueEmitterName() + TEXT(".Localspace"));
 				}
+				if (Emitter->bDeterminism)
+				{
+					Id.AdditionalDefines.Add(Emitter->GetUniqueEmitterName() + TEXT(".Determinism"));
+				}
 			}
 		}
 	}
@@ -478,7 +497,23 @@ void UNiagaraScript::GenerateStatScopeIDs()
 
 void UNiagaraScript::PostEditChangeProperty(FPropertyChangedEvent& PropertyChangedEvent)
 {
+	Super::PostEditChangeProperty(PropertyChangedEvent);
+
+	FName PropertyName;
+	if (PropertyChangedEvent.Property)
+	{
+		PropertyName = PropertyChangedEvent.Property->GetFName();
+	}
+
 	CacheResourceShadersForRendering(true);
+
+	if (PropertyName == GET_MEMBER_NAME_CHECKED(UNiagaraScript, bDeprecated) || PropertyName == GET_MEMBER_NAME_CHECKED(UNiagaraScript, DeprecationRecommendation))
+	{
+		if (Source)
+		{
+			Source->MarkNotSynchronized(TEXT("Deprecation changed."));
+		}
+	}
 }
 
 #endif
@@ -1065,6 +1100,14 @@ void UNiagaraScript::SyncAliases(const TMap<FString, FString>& RenameMap)
 	// Now handle any Parameters overall..
 	for (int32 i = 0; i < GetVMExecutableData().Parameters.Parameters.Num(); i++)
 	{
+		if (GetVMExecutableData().Parameters.Parameters[i].IsValid() == false)
+		{
+			const FNiagaraVariable& InvalidParameter = GetVMExecutableData().Parameters.Parameters[i];
+			UE_LOG(LogNiagara, Error, TEXT("Invalid parameter found while syncing script aliases.  Script: %s Parameter Name: %s Parameter Type: %s"),
+				*GetPathName(), *InvalidParameter.GetName().ToString(), InvalidParameter.GetType().IsValid() ? *InvalidParameter.GetType().GetName() : TEXT("Unknown"));
+			continue;
+		}
+
 		FNiagaraVariable Var = GetVMExecutableData().Parameters.Parameters[i];
 		FNiagaraVariable NewVar = FNiagaraVariable::ResolveAliases(Var, RenameMap);
 		if (NewVar.GetName() != Var.GetName())
@@ -1203,6 +1246,8 @@ void SerializeNiagaraShaderMaps(const TMap<const ITargetPlatform*, TArray<FNiaga
 {
 	Ar.UsingCustomVersion(FFortniteMainBranchObjectVersion::GUID);
 	Ar.UsingCustomVersion(FRenderingObjectVersion::GUID);
+	Ar.UsingCustomVersion(FEditorObjectVersion::GUID);
+	Ar.UsingCustomVersion(FReleaseObjectVersion::GUID);
 
 //	SCOPED_LOADTIMER(SerializeInlineShaderMaps);
 	if (Ar.IsSaving())
diff --git a/Engine/Plugins/FX/Niagara/Source/NiagaraEditor/Private/NiagaraNodeUsageSelector.cpp b/Engine/Plugins/FX/Niagara/Source/NiagaraEditor/Private/NiagaraNodeUsageSelector.cpp
--- a/Engine/Plugins/FX/Niagara/Source/NiagaraEditor/Private/NiagaraNodeUsageSelector.cpp
+++ b/Engine/Plugins/FX/Niagara/Source/NiagaraEditor/Private/NiagaraNodeUsageSelector.cpp
@@ -32,10 +32,37 @@ bool UNiagaraNodeUsageSelector::AllowNiagaraTypeForAddPin(const FNiagaraTypeDefi
 	return Super::AllowNiagaraTypeForAddPin(InType) && InType != FNiagaraTypeDefinition::GetParameterMapDef();
 }
 
-void UNiagaraNodeUsageSelector::AllocateDefaultPins()
+void UNiagaraNodeUsageSelector::InsertInputPinsFor(const FNiagaraVariable& Var)
 {
 	const UEdGraphSchema_Niagara* Schema = GetDefault<UEdGraphSchema_Niagara>();
+	UEnum* ENiagaraScriptGroupEnum = FindObject<UEnum>(ANY_PACKAGE, TEXT("ENiagaraScriptGroup"), true);
+	int64 GroupCount = (int64)ENiagaraScriptGroup::Max;
+
+	TArray<UEdGraphPin*> OldPins(Pins);
+	Pins.Reset(Pins.Num() + GroupCount);
+
+	// Create the inputs for each path.
+	for (int64 i = 0; i < GroupCount; i++)
+	{
+		// Add the previous input pins
+		for (int32 k = 0; k < OutputVars.Num() - 1; k++)
+		{
+			Pins.Add(OldPins[k]);
+		}
+		OldPins.RemoveAt(0, OutputVars.Num() - 1);
 
+		// Add the new input pin
+		const FString PathSuffix = ENiagaraScriptGroupEnum ? (FString::Printf(TEXT(" if %s"), *ENiagaraScriptGroupEnum->GetNameStringByValue((int64)i))) : TEXT("Error Unknown!");
+		CreatePin(EGPD_Input, Schema->TypeDefinitionToPinType(Var.GetType()), *(Var.GetName().ToString() + PathSuffix));
+	}
+
+	// Move the rest of the old pins over
+	Pins.Append(OldPins);
+}
+
+void UNiagaraNodeUsageSelector::AllocateDefaultPins()
+{
+	const UEdGraphSchema_Niagara* Schema = GetDefault<UEdGraphSchema_Niagara>();
 	UEnum* ENiagaraScriptGroupEnum = StaticEnum<ENiagaraScriptGroup>();
 
 	//Create the inputs for each path.
@@ -137,6 +164,33 @@ UEdGraphPin* UNiagaraNodeUsageSelector::GetPassThroughPin(const UEdGraphPin* Loc
 	return nullptr;
 }
 
+void UNiagaraNodeUsageSelector::AppendFunctionAliasForContext(const FNiagaraGraphFunctionAliasContext& InFunctionAliasContext, FString& InOutFunctionAlias)
+{
+	FString UsageString;
+	switch (InFunctionAliasContext.CompileUsage)
+	{
+	case ENiagaraScriptUsage::SystemSpawnScript:
+	case ENiagaraScriptUsage::SystemUpdateScript:
+		UsageString = "System";
+		break;
+	case ENiagaraScriptUsage::EmitterSpawnScript:
+	case ENiagaraScriptUsage::EmitterUpdateScript:
+		UsageString = "Emitter";
+		break;
+	case ENiagaraScriptUsage::ParticleSpawnScript:
+	case ENiagaraScriptUsage::ParticleUpdateScript:
+	case ENiagaraScriptUsage::ParticleEventScript:
+	case ENiagaraScriptUsage::ParticleGPUComputeScript:
+		UsageString = "Particle";
+		break;
+	}
+
+	if (UsageString.IsEmpty() == false)
+	{
+		InOutFunctionAlias += "_" + UsageString;
+	}
+}
+
 void UNiagaraNodeUsageSelector::BuildParameterMapHistory(FNiagaraParameterMapHistoryBuilder& OutHistory, bool bRecursive)
 {
 	const UEdGraphSchema_Niagara* Schema = CastChecked<UEdGraphSchema_Niagara>(GetSchema());
@@ -208,14 +262,17 @@ void UNiagaraNodeUsageSelector::OnNewTypedPinAdded(UEdGraphPin* NewPin)
 	{
 		OutputNames.Add(Output.GetName());
 	}
-	FName OutputName = FNiagaraUtilities::GetUniqueName(*OutputType.GetNameText().ToString(), OutputNames);
-
+	FName OutputName = FNiagaraUtilities::GetUniqueName(*NewPin->GetName(), OutputNames);
+	NewPin->PinName = OutputName;
 	FGuid Guid = AddOutput(OutputType, OutputName);
 
 	// Update the pin's data too so that it's connection is maintained after reallocating.
 	NewPin->PersistentGuid = Guid;
 
-	ReallocatePins();
+	// We cannot just reallocate the pins here, because that invalidates all pins of this node (including
+	// the NewPin parameter). If the calling method tries to access the provided new pin afterwards, it
+	// runs into a nullptr error (e.g. when called by drag and drop).
+	InsertInputPinsFor(OutputVars.Last());
 }
 
 void UNiagaraNodeUsageSelector::OnPinRenamed(UEdGraphPin* RenamedPin, const FString& OldName)
