Merging //UE4/Dev-Main to Dev-Sequencer (//UE4/Dev-Sequencer)

diff --git a/Engine/Plugins/Animation/LiveLink/Source/LiveLink/Private/MovieScene/MovieSceneLiveLinkSection.cpp b/Engine/Plugins/Animation/LiveLink/Source/LiveLink/Private/MovieScene/MovieSceneLiveLinkSection.cpp
--- a/Engine/Plugins/Animation/LiveLink/Source/LiveLink/Private/MovieScene/MovieSceneLiveLinkSection.cpp
+++ b/Engine/Plugins/Animation/LiveLink/Source/LiveLink/Private/MovieScene/MovieSceneLiveLinkSection.cpp
@@ -154,19 +154,19 @@ void UMovieSceneLiveLinkSection::Serialize(FArchive& Ar)
 	}
 }
 
-void UMovieSceneLiveLinkSection::SetMask(const TArray<bool>& InChannelMask)
+void UMovieSceneLiveLinkSection::PostEditImport()
 {
-	ChannelMask = InChannelMask;
-	UpdateChannelProxy();
-	//MattH todo set priorities or whatever based upon mask
+	Super::PostEditImport();
 
+	UpdateChannelProxy();
 }
 
-void UMovieSceneLiveLinkSection::PostEditImport()
+void UMovieSceneLiveLinkSection::SetMask(const TArray<bool>& InChannelMask)
 {
-	Super::PostEditImport();
-
+	ChannelMask = InChannelMask;
 	UpdateChannelProxy();
+	//MattH todo set priorities or whatever based upon mask
+
 }
 
 #undef LOCTEXT_NAMESPACE // MovieSceneNiagaraEmitterTimedSection
\ No newline at end of file
diff --git a/Engine/Plugins/FX/Niagara/Source/Niagara/Private/NiagaraScript.cpp b/Engine/Plugins/FX/Niagara/Source/Niagara/Private/NiagaraScript.cpp
--- a/Engine/Plugins/FX/Niagara/Source/Niagara/Private/NiagaraScript.cpp
+++ b/Engine/Plugins/FX/Niagara/Source/Niagara/Private/NiagaraScript.cpp
@@ -16,6 +16,9 @@
 #include "UObject/Linker.h"
 #include "HAL/PlatformFilemanager.h"
 #include "Misc/FileHelper.h"
+#include "UObject/EditorObjectVersion.h"
+#include "UObject/ReleaseObjectVersion.h"
+
 
 #if WITH_EDITOR
 	#include "NiagaraScriptDerivedData.h"
@@ -26,13 +29,21 @@
 
 DECLARE_STATS_GROUP(TEXT("Niagara Detailed"), STATGROUP_NiagaraDetailed, STATCAT_Advanced);
 
-FNiagaraScriptDebuggerInfo::FNiagaraScriptDebuggerInfo() : FrameLastWriteId(-1)
+FNiagaraScriptDebuggerInfo::FNiagaraScriptDebuggerInfo() : bWaitForGPU(false), FrameLastWriteId(-1), bWritten(false)
 {
 }
 
 
-FNiagaraScriptDebuggerInfo::FNiagaraScriptDebuggerInfo(FName InName, ENiagaraScriptUsage InUsage, const FGuid& InUsageId) : HandleName(InName), Usage(InUsage), UsageId(InUsageId), FrameLastWriteId(-1)
+FNiagaraScriptDebuggerInfo::FNiagaraScriptDebuggerInfo(FName InName, ENiagaraScriptUsage InUsage, const FGuid& InUsageId) : HandleName(InName), Usage(InUsage), UsageId(InUsageId), FrameLastWriteId(-1), bWritten(false)
 {
+	if (InUsage == ENiagaraScriptUsage::ParticleGPUComputeScript)
+	{
+		bWaitForGPU = true;
+	}
+	else
+	{
+		bWaitForGPU = false;
+	}
 }
 
 
@@ -230,6 +241,10 @@ void UNiagaraScript::ComputeVMCompilationId(FNiagaraVMExecutableDataId& Id) cons
 		{
 			Id.AdditionalDefines.Add(TEXT("Emitter.Localspace"));
 		}
+		if (Emitter->bDeterminism)
+		{
+			Id.AdditionalDefines.Add(TEXT("Emitter.Determinism"));
+		}
 	}
 
 	if (UNiagaraSystem* System = Cast<UNiagaraSystem>(Obj))
@@ -243,6 +258,10 @@ void UNiagaraScript::ComputeVMCompilationId(FNiagaraVMExecutableDataId& Id) cons
 				{
 					Id.AdditionalDefines.Add(Emitter->GetUniqueEmitterName() + TEXT(".Localspace"));
 				}
+				if (Emitter->bDeterminism)
+				{
+					Id.AdditionalDefines.Add(Emitter->GetUniqueEmitterName() + TEXT(".Determinism"));
+				}
 			}
 		}
 	}
@@ -476,7 +495,23 @@ void UNiagaraScript::GenerateStatScopeIDs()
 
 void UNiagaraScript::PostEditChangeProperty(FPropertyChangedEvent& PropertyChangedEvent)
 {
+	Super::PostEditChangeProperty(PropertyChangedEvent);
+
+	FName PropertyName;
+	if (PropertyChangedEvent.Property)
+	{
+		PropertyName = PropertyChangedEvent.Property->GetFName();
+	}
+
 	CacheResourceShadersForRendering(true);
+
+	if (PropertyName == GET_MEMBER_NAME_CHECKED(UNiagaraScript, bDeprecated) || PropertyName == GET_MEMBER_NAME_CHECKED(UNiagaraScript, DeprecationRecommendation))
+	{
+		if (Source)
+		{
+			Source->MarkNotSynchronized(TEXT("Deprecation changed."));
+		}
+	}
 }
 
 #endif
@@ -1063,6 +1098,14 @@ void UNiagaraScript::SyncAliases(const TMap<FString, FString>& RenameMap)
 	// Now handle any Parameters overall..
 	for (int32 i = 0; i < GetVMExecutableData().Parameters.Parameters.Num(); i++)
 	{
+		if (GetVMExecutableData().Parameters.Parameters[i].IsValid() == false)
+		{
+			const FNiagaraVariable& InvalidParameter = GetVMExecutableData().Parameters.Parameters[i];
+			UE_LOG(LogNiagara, Error, TEXT("Invalid parameter found while syncing script aliases.  Script: %s Parameter Name: %s Parameter Type: %s"),
+				*GetPathName(), *InvalidParameter.GetName().ToString(), InvalidParameter.GetType().IsValid() ? *InvalidParameter.GetType().GetName() : TEXT("Unknown"));
+			continue;
+		}
+
 		FNiagaraVariable Var = GetVMExecutableData().Parameters.Parameters[i];
 		FNiagaraVariable NewVar = FNiagaraVariable::ResolveAliases(Var, RenameMap);
 		if (NewVar.GetName() != Var.GetName())
@@ -1199,6 +1242,11 @@ NIAGARA_API bool UNiagaraScript::DidScriptCompilationSucceed(bool bGPUScript) co
 
 void SerializeNiagaraShaderMaps(const TMap<const ITargetPlatform*, TArray<FNiagaraShaderScript*>>* PlatformScriptResourcesToSave, FArchive& Ar, TArray<FNiagaraShaderScript>& OutLoadedResources)
 {
+	Ar.UsingCustomVersion(FFortniteMainBranchObjectVersion::GUID);
+	Ar.UsingCustomVersion(FRenderingObjectVersion::GUID);
+	Ar.UsingCustomVersion(FEditorObjectVersion::GUID);
+	Ar.UsingCustomVersion(FReleaseObjectVersion::GUID);
+
 //	SCOPED_LOADTIMER(SerializeInlineShaderMaps);
 	if (Ar.IsSaving())
 	{
