Merging //UE4/Dev-Main to Dev-Sequencer (//UE4/Dev-Sequencer)

diff --git a/Engine/Plugins/Editor/PluginBrowser/Source/PluginBrowser/Private/SNewPluginWizard.cpp b/Engine/Plugins/Editor/PluginBrowser/Source/PluginBrowser/Private/SNewPluginWizard.cpp
--- a/Engine/Plugins/Editor/PluginBrowser/Source/PluginBrowser/Private/SNewPluginWizard.cpp
+++ b/Engine/Plugins/Editor/PluginBrowser/Source/PluginBrowser/Private/SNewPluginWizard.cpp
@@ -743,7 +743,7 @@ FReply SNewPluginWizard::OnCreatePluginClicked()
 	if (bSucceeded && bHasModules)
 	{
 		FString ProjectFileName = FPaths::GetProjectFilePath();
-		FString Arguments = FString::Printf(TEXT("%s %s -TargetType=Editor -Module=%s -EnablePlugin=\"%s\" -Project=\"%s\" -Progress -NoHotReloadFromIDE"), FModuleManager::Get().GetUBTConfiguration(), FPlatformMisc::GetUBTPlatform(), *AutoPluginName, *PluginModuleName, *ProjectFileName);
+		FString Arguments = FString::Printf(TEXT("%s %s -TargetType=Editor -Plugin=\"%s\" -Project=\"%s\" -Progress -NoHotReloadFromIDE"), FModuleManager::Get().GetUBTConfiguration(), FPlatformMisc::GetUBTPlatform(), *UPluginFilePath, *ProjectFileName);
 		if (!FDesktopPlatformModule::Get()->RunUnrealBuildTool(LOCTEXT("Compiling", "Compiling..."), FPaths::RootDir(), Arguments, GWarn))
 		{
 			PopErrorNotification(LOCTEXT("FailedToCompile", "Failed to compile source code."));
diff --git a/Engine/Plugins/Experimental/PixelStreaming/Source/PixelStreaming/PixelStreaming.Build.cs b/Engine/Plugins/Experimental/PixelStreaming/Source/PixelStreaming/PixelStreaming.Build.cs
--- a/Engine/Plugins/Experimental/PixelStreaming/Source/PixelStreaming/PixelStreaming.Build.cs
+++ b/Engine/Plugins/Experimental/PixelStreaming/Source/PixelStreaming/PixelStreaming.Build.cs
@@ -144,9 +144,16 @@ namespace UnrealBuildTool.Rules
             if (Target.Platform == UnrealTargetPlatform.Win32 || Target.Platform == UnrealTargetPlatform.Win64)
             {
                 PrivateDependencyModuleNames.AddRange(new string[] { "D3D11RHI" });
-            }
-
-            AddWebRTCProxy();
+				PrivateIncludePaths.AddRange(
+					new string[] {
+					"../../../../Source/Runtime/Windows/D3D11RHI/Private",
+					"../../../../Source/Runtime/Windows/D3D11RHI/Private/Windows",
+					});
+				AddEngineThirdPartyPrivateStaticDependencies(Target, "IntelMetricsDiscovery");
+				AddEngineThirdPartyPrivateStaticDependencies(Target, "NVAftermath");
+			}
+
+			AddWebRTCProxy();
             AddSignallingServer();
             AddMatchmakingServer();
             AddWebRTCServers();
diff --git a/Engine/Plugins/Experimental/PixelStreaming/Source/PixelStreaming/Private/NvVideoEncoder.cpp b/Engine/Plugins/Experimental/PixelStreaming/Source/PixelStreaming/Private/NvVideoEncoder.cpp
--- a/Engine/Plugins/Experimental/PixelStreaming/Source/PixelStreaming/Private/NvVideoEncoder.cpp
+++ b/Engine/Plugins/Experimental/PixelStreaming/Source/PixelStreaming/Private/NvVideoEncoder.cpp
@@ -24,9 +24,8 @@
 #		include "Windows/AllowWindowsPlatformTypes.h"
 #			include "NvEncoder/nvEncodeAPI.h"
 #			include <d3d11.h>
-#			include "D3D11State.h"
-#			include "D3D11Resources.h"
 #		include "Windows/HideWindowsPlatformTypes.h"
+#		include "D3D11RHIPrivate.h"
 #	pragma warning(pop)
 #endif
 
diff --git a/Engine/Plugins/Experimental/SkeletalReduction/Source/Private/SkeletalMeshReductionPlugin.cpp b/Engine/Plugins/Experimental/SkeletalReduction/Source/Private/SkeletalMeshReductionPlugin.cpp
--- a/Engine/Plugins/Experimental/SkeletalReduction/Source/Private/SkeletalMeshReductionPlugin.cpp
+++ b/Engine/Plugins/Experimental/SkeletalReduction/Source/Private/SkeletalMeshReductionPlugin.cpp
@@ -826,21 +826,26 @@ float FQuadricSkeletalMeshReduction::SimplifyMesh( const FSkeletalMeshOptimizati
 
 	// Determine the stop criteria used
 
-	const bool bUseVertexCriterion   = Settings.TerminationCriterion == SkeletalMeshTerminationCriterion::SMTC_NumOfVerts     || Settings.TerminationCriterion == SkeletalMeshTerminationCriterion::SMTC_TriangleOrVert;
-	const bool bUseTriangleCriterion = Settings.TerminationCriterion == SkeletalMeshTerminationCriterion::SMTC_NumOfTriangles || Settings.TerminationCriterion == SkeletalMeshTerminationCriterion::SMTC_TriangleOrVert;
+	const bool bUseVertexPercentCriterion = Settings.TerminationCriterion == SkeletalMeshTerminationCriterion::SMTC_NumOfVerts || Settings.TerminationCriterion == SkeletalMeshTerminationCriterion::SMTC_TriangleOrVert;
+	const bool bUseTrianglePercentCriterion = Settings.TerminationCriterion == SkeletalMeshTerminationCriterion::SMTC_NumOfTriangles || Settings.TerminationCriterion == SkeletalMeshTerminationCriterion::SMTC_TriangleOrVert;
 
+	const bool bUseMaxVertNumCriterion = Settings.TerminationCriterion == SkeletalMeshTerminationCriterion::SMTC_AbsNumOfVerts || Settings.TerminationCriterion == SkeletalMeshTerminationCriterion::SMTC_AbsTriangleOrVert;
+	const bool bUseMaxTrisNumCriterion = Settings.TerminationCriterion == SkeletalMeshTerminationCriterion::SMTC_AbsNumOfTriangles || Settings.TerminationCriterion == SkeletalMeshTerminationCriterion::SMTC_AbsTriangleOrVert;
 
 	// We can support a stopping criteria based on the MaxDistance the new vertex is from the plans of the source triangles.
 	// but there seems to be no good use for this.  We are better off just using triangle count.
-	const float MaxDist              = FLT_MAX; // (Settings.ReductionMethod != SkeletalMeshOptimizationType::SMOT_NumOfTriangles) ? Settings.MaxDeviationPercentage * Bounds.SphereRadius : FLT_MAX;
-	const int32 MaxTriNumToRetain    = Mesh.NumIndices() / 3;
-	const float TriangleRetainRatio  = FMath::Clamp(Settings.NumOfTrianglesPercentage, 0.f, 1.f);
-	const int32 MinTriNumToRetain    = (bUseTriangleCriterion) ? FMath::Max(4, FMath::CeilToInt(TriangleRetainRatio * MaxTriNumToRetain)) : 4 ; 	
-	const float MaxCollapseCost      = FLT_MAX;
+	const float MaxDist = FLT_MAX; // (Settings.ReductionMethod != SkeletalMeshOptimizationType::SMOT_NumOfTriangles) ? Settings.MaxDeviationPercentage * Bounds.SphereRadius : FLT_MAX;
+	const int32 SrcTriNum = Mesh.NumIndices() / 3;
+	const float TriangleRetainRatio = FMath::Clamp(Settings.NumOfTrianglesPercentage, 0.f, 1.f);
+	const int32 TargetTriNum = (bUseTrianglePercentCriterion) ? FMath::CeilToInt(TriangleRetainRatio * SrcTriNum) : Settings.MaxNumOfTriangles;
 
-	const int32 MaxVerNumToRetain = Mesh.NumVertices();
-	const float VertRetainRatio   = FMath::Clamp(Settings.NumOfVertPercentage, 0.f, 1.f);
-	const int32 MinVerNumToRetain = (bUseVertexCriterion) ? FMath::Max(3, FMath::CeilToInt(VertRetainRatio * MaxVerNumToRetain)) : 0;
+	const int32 MinTriNumToRetain = (bUseTrianglePercentCriterion || bUseMaxTrisNumCriterion) ? FMath::Max(4, TargetTriNum) : 4;
+	const float MaxCollapseCost = FLT_MAX;
+
+	const int32 SrcVertNum = Mesh.NumVertices();
+	const float VertRetainRatio = FMath::Clamp(Settings.NumOfVertPercentage, 0.f, 1.f);
+	const int32 TargetVertNum = (bUseVertexPercentCriterion) ? FMath::CeilToInt(VertRetainRatio * SrcVertNum) : Settings.MaxNumOfVerts + 1;
+	const int32 MinVerNumToRetain = (bUseVertexPercentCriterion || bUseMaxVertNumCriterion) ? FMath::Max(6, TargetVertNum) : 6;
 
 	const float VolumeImportance      = FMath::Clamp(Settings.VolumeImportance, 0.f, 2.f);
 	const bool bLockEdges             = Settings.bLockEdges;
@@ -848,7 +853,7 @@ float FQuadricSkeletalMeshReduction::SimplifyMesh( const FSkeletalMeshOptimizati
 	const bool bEnforceBoneBoundaries = Settings.bEnforceBoneBoundaries;
 
 	// Terminator tells the simplifier when to stop
-	SkeletalSimplifier::FSimplifierTerminator Terminator(MinTriNumToRetain, MaxTriNumToRetain, MinVerNumToRetain, MaxVerNumToRetain, MaxCollapseCost, MaxDist);
+	SkeletalSimplifier::FSimplifierTerminator Terminator(MinTriNumToRetain, SrcTriNum, MinVerNumToRetain, SrcVertNum, MaxCollapseCost, MaxDist);
 
 	double NormalWeight    =16.00;
 	double TangentWeight   = 0.10;
diff --git a/Engine/Plugins/FX/Niagara/Source/Niagara/Private/NiagaraEmitterInstanceBatcher.cpp b/Engine/Plugins/FX/Niagara/Source/Niagara/Private/NiagaraEmitterInstanceBatcher.cpp
--- a/Engine/Plugins/FX/Niagara/Source/Niagara/Private/NiagaraEmitterInstanceBatcher.cpp
+++ b/Engine/Plugins/FX/Niagara/Source/Niagara/Private/NiagaraEmitterInstanceBatcher.cpp
@@ -125,6 +125,9 @@ void NiagaraEmitterInstanceBatcher::TickSingle(FNiagaraComputeExecutionContext *
 	// assume all instances survived; ResolveDataSetWrites will change this if the deferred readback was successful; that data may be several frames old
 	Context->MainDataSet->CurrData().SetNumInstances(NewNumInstances);
 
+	// ResolveDatasetWrites may read this, so we must transition it here.
+	RHICmdList.TransitionResource(EResourceTransitionAccess::EReadable, EResourceTransitionPipeline::EComputeToGfx, Context->MainDataSet->GetCurDataSetIndices().UAV);		// transition to readable; we'll be using this next frame
+
 	// Don't resolve if the data if there are no instances (prevents a transition issue warning).
 	if (NewNumInstances > 0)
 	{
@@ -132,8 +135,6 @@ void NiagaraEmitterInstanceBatcher::TickSingle(FNiagaraComputeExecutionContext *
 		ResolveDatasetWrites(RHICmdList, Context);
 	}
 
-	RHICmdList.TransitionResource(EResourceTransitionAccess::EReadable, EResourceTransitionPipeline::EComputeToGfx, Context->MainDataSet->GetCurDataSetIndices().UAV);		// transition to readable; we'll be using this next frame
-
 	/*
 	// TODO: hack - only updating event set 0 on update scripts now; need to match them to their indices and update them all
 	if (Context->UpdateEventWriteDataSets.Num())
@@ -197,28 +198,26 @@ void NiagaraEmitterInstanceBatcher::ResolveDatasetWrites(FRHICommandList &RHICmd
 	{
 		SCOPE_CYCLE_COUNTER(STAT_NiagaraAllocateGPUReadback_RT);
 
-		Context->GPUDataReadback = new FRHIGPUMemoryReadback(DatasetIndexBufferWrite.Buffer, TEXT("Niagara GPU Emitter Readback"));
+		Context->GPUDataReadback = new FRHIGPUMemoryReadback(TEXT("Niagara GPU Emitter Readback"));
 		INC_DWORD_STAT(STAT_NiagaraReadbackLatency);
-		Context->GPUDataReadback->Insert(RHICmdList);
+		Context->GPUDataReadback->EnqueueCopy(RHICmdList, DatasetIndexBufferWrite.Buffer);
 	}
 	else if (Context->GPUDataReadback->IsReady())
 	{
 		bool bSuccessfullyRead = false;
 		{
 		    SCOPE_CYCLE_COUNTER(STAT_NiagaraGPUReadback_RT);
-		    int32 *NumInstancesAfterSim = static_cast<int32*>(Context->GPUDataReadback->RetrieveData(64 * sizeof(int32)));
+		    int32 *NumInstancesAfterSim = static_cast<int32*>(Context->GPUDataReadback->Lock(64 * sizeof(int32)));
 			if (NumInstancesAfterSim)
 			{
 				int32 ExistingDataCount = Context->MainDataSet->CurrData().GetNumInstances();// index 1 is always the count
 				int32 NewExistingDataCount = NumInstancesAfterSim[1] + Context->AccumulatedSpawnRate;
 				Context->MainDataSet->CurrData().SetNumInstances(NewExistingDataCount);
 				FString PathName = Context->GPUScript->GetOutermost()->GetPathName();
-				// UE_LOG(LogNiagara, Log, TEXT("GPU Syncup %s : Was(%d) Now(%d)"), *PathName, ExistingDataCount, NewExistingDataCount );
+				//UE_LOG(LogNiagara, Log, TEXT("GPU Syncup %s : Was(%d) Now(%d)"), *PathName, ExistingDataCount, NewExistingDataCount );
 				SET_DWORD_STAT(STAT_NiagaraGPUParticles, NewExistingDataCount);
 				SET_DWORD_STAT(STAT_NiagaraReadbackLatency, 0);
 
-				Context->GPUDataReadback->Finish();
-
 				Context->AccumulatedSpawnRate = 0;
 				bSuccessfullyRead = true;
 			}
@@ -226,17 +225,15 @@ void NiagaraEmitterInstanceBatcher::ResolveDatasetWrites(FRHICommandList &RHICmd
 			{
 				UE_LOG(LogNiagara, Warning, TEXT("GPUDataReadback said it was ready, but returned an invalid buffer. Skipping this time.."));
 			}
+			Context->GPUDataReadback->Unlock();
 		}
 		if (bSuccessfullyRead)
 		{
 			SCOPE_CYCLE_COUNTER(STAT_NiagaraAllocateGPUReadback_RT);
 			// The following code seems to take significant time on d3d12
-			delete Context->GPUDataReadback;
-			Context->GPUDataReadback = new FRHIGPUMemoryReadback(DatasetIndexBufferWrite.Buffer, TEXT("Niagara GPU Emitter Readback"));
-			Context->GPUDataReadback->Insert(RHICmdList);
+			Context->GPUDataReadback->EnqueueCopy(RHICmdList, DatasetIndexBufferWrite.Buffer);
 		}
 	}
-
 }
 
 void NiagaraEmitterInstanceBatcher::ProcessDebugInfo(FRHICommandList &RHICmdList, const FNiagaraComputeExecutionContext *Context) const
@@ -260,31 +257,31 @@ void NiagaraEmitterInstanceBatcher::ProcessDebugInfo(FRHICommandList &RHICmdList
 		{
 			//UE_LOG(LogNiagara, Warning, TEXT("Read back!"));
 
-			int32 *NumInstancesAfterSim = static_cast<int32*>(Context->GPUDebugDataReadbackCounts->RetrieveData(64 * sizeof(int32)));
+			int32 *NumInstancesAfterSim = static_cast<int32*>(Context->GPUDebugDataReadbackCounts->Lock(64 * sizeof(int32)));
 			int32 NewExistingDataCount = NumInstancesAfterSim[1];
 			{
 				float* FloatDataBuffer = nullptr;
 				if (Context->GPUDebugDataReadbackFloat)
 				{
-					FloatDataBuffer = static_cast<float*>(Context->GPUDebugDataReadbackFloat->RetrieveData(Context->GPUDebugDataFloatSize));
+					FloatDataBuffer = static_cast<float*>(Context->GPUDebugDataReadbackFloat->Lock(Context->GPUDebugDataFloatSize));
 				}
 				int* IntDataBuffer = nullptr;
 				if (Context->GPUDebugDataReadbackInt)
 				{
-					IntDataBuffer = static_cast<int*>(Context->GPUDebugDataReadbackInt->RetrieveData(Context->GPUDebugDataIntSize));
+					IntDataBuffer = static_cast<int*>(Context->GPUDebugDataReadbackInt->Lock(Context->GPUDebugDataIntSize));
 				}
 				Context->MainDataSet->DumpGPU(Context->DebugInfo->Frame, FloatDataBuffer, IntDataBuffer, 0, NewExistingDataCount);
 				Context->DebugInfo->bWritten = true;
 
 				if (Context->GPUDebugDataReadbackFloat)
 				{
-					Context->GPUDebugDataReadbackFloat->Finish();
+					Context->GPUDebugDataReadbackFloat->Unlock();
 				}
 				if (Context->GPUDebugDataReadbackInt)
 				{
-					Context->GPUDebugDataReadbackInt->Finish();
+					Context->GPUDebugDataReadbackInt->Unlock();
 				}
-				Context->GPUDebugDataReadbackCounts->Finish();
+				Context->GPUDebugDataReadbackCounts->Unlock();
 			}
 			{
 				// The following code seems to take significant time on d3d12
@@ -515,20 +512,20 @@ void NiagaraEmitterInstanceBatcher::Run(const FNiagaraComputeExecutionContext *C
 
 			if (Context->MainDataSet->GetNumFloatComponents() > 0)
 			{
-				Context->GPUDebugDataReadbackFloat = new FRHIGPUMemoryReadback(Context->MainDataSet->GetDataByIndex(WriteBufferIdx).GetGPUBufferFloat()->Buffer, TEXT("Niagara GPU Debug Info Float Emitter Readback"));
-				Context->GPUDebugDataReadbackFloat->Insert(RHICmdList);
+				Context->GPUDebugDataReadbackFloat = new FRHIGPUMemoryReadback(TEXT("Niagara GPU Debug Info Float Emitter Readback"));
+				Context->GPUDebugDataReadbackFloat->EnqueueCopy(RHICmdList, Context->MainDataSet->GetDataByIndex(WriteBufferIdx).GetGPUBufferFloat()->Buffer);
 				Context->GPUDebugDataFloatSize = Context->MainDataSet->GetDataByIndex(WriteBufferIdx).GetGPUBufferFloat()->NumBytes;
 			}
 
 			if (Context->MainDataSet->GetNumInt32Components() > 0)
 			{
-				Context->GPUDebugDataReadbackInt = new FRHIGPUMemoryReadback(Context->MainDataSet->GetDataByIndex(WriteBufferIdx).GetGPUBufferInt()->Buffer, TEXT("Niagara GPU Debug Info Int Emitter Readback"));
-				Context->GPUDebugDataReadbackInt->Insert(RHICmdList);
+				Context->GPUDebugDataReadbackInt = new FRHIGPUMemoryReadback(TEXT("Niagara GPU Debug Info Int Emitter Readback"));
+				Context->GPUDebugDataReadbackInt->EnqueueCopy(RHICmdList, Context->MainDataSet->GetDataByIndex(WriteBufferIdx).GetGPUBufferInt()->Buffer);
 				Context->GPUDebugDataIntSize = Context->MainDataSet->GetDataByIndex(WriteBufferIdx).GetGPUBufferInt()->NumBytes;
 			}
 
-			Context->GPUDebugDataReadbackCounts = new FRHIGPUMemoryReadback(DatasetIndexBufferWrite.Buffer, TEXT("Niagara GPU Emitter Readback"));
-			Context->GPUDebugDataReadbackCounts->Insert(RHICmdList);
+			Context->GPUDebugDataReadbackCounts = new FRHIGPUMemoryReadback(TEXT("Niagara GPU Emitter Readback"));
+			Context->GPUDebugDataReadbackCounts->EnqueueCopy(RHICmdList, DatasetIndexBufferWrite.Buffer);
 		}
 	}
 #endif // WITH_EDITORONLY_DATA
diff --git a/Engine/Plugins/Runtime/CableComponent/Source/CableComponent/Classes/CableComponent.h b/Engine/Plugins/Runtime/CableComponent/Source/CableComponent/Classes/CableComponent.h
--- a/Engine/Plugins/Runtime/CableComponent/Source/CableComponent/Classes/CableComponent.h
+++ b/Engine/Plugins/Runtime/CableComponent/Source/CableComponent/Classes/CableComponent.h
@@ -87,7 +87,8 @@ public:
 	FVector EndLocation;
 
 	/** Attaches the end of the cable to a specific Component **/
-	void SetAttachEndToComponent(TWeakObjectPtr<class USceneComponent> Component, FName SocketName = NAME_None);
+	UFUNCTION(BlueprintCallable, Category = "Cable")
+	void SetAttachEndToComponent(USceneComponent* Component, FName SocketName = NAME_None);
 
 	/** Attaches the end of the cable to a specific Component within an Actor **/
 	UFUNCTION(BlueprintCallable, Category = "Cable")
diff --git a/Engine/Plugins/Runtime/CableComponent/Source/CableComponent/Private/CableComponent.cpp b/Engine/Plugins/Runtime/CableComponent/Source/CableComponent/Private/CableComponent.cpp
--- a/Engine/Plugins/Runtime/CableComponent/Source/CableComponent/Private/CableComponent.cpp
+++ b/Engine/Plugins/Runtime/CableComponent/Source/CableComponent/Private/CableComponent.cpp
@@ -555,9 +555,10 @@ void UCableComponent::PerformSubstep(float InSubstepTime, const FVector& Gravity
 	}
 }
 
-void UCableComponent::SetAttachEndToComponent(TWeakObjectPtr<class USceneComponent> Component, FName SocketName)
+void UCableComponent::SetAttachEndToComponent(USceneComponent* Component, FName SocketName)
 {
 	AttachEndTo.OtherActor = Component->GetOwner();
+	AttachEndTo.ComponentProperty = NAME_None;
 	AttachEndTo.OverrideComponent = Component;
 	AttachEndToSocketName = SocketName;
 }
diff --git a/Engine/Source/Developer/DesktopPlatform/Private/DesktopPlatformBase.cpp b/Engine/Source/Developer/DesktopPlatform/Private/DesktopPlatformBase.cpp
--- a/Engine/Source/Developer/DesktopPlatform/Private/DesktopPlatformBase.cpp
+++ b/Engine/Source/Developer/DesktopPlatform/Private/DesktopPlatformBase.cpp
@@ -491,8 +491,8 @@ bool FDesktopPlatformBase::GetEngineIdentifierForProject(const FString& ProjectF
 	}
 
 	// Otherwise check the engine version string for 4.0, in case this project existed before the engine association stuff went in
-	FString EngineVersionString = ProjectFile->GetStringField(TEXT("EngineVersion"));
-	if(EngineVersionString.Len() > 0)
+	FString EngineVersionString;
+	if(ProjectFile->TryGetStringField(TEXT("EngineVersion"), EngineVersionString) && EngineVersionString.Len() > 0)
 	{
 		FEngineVersion EngineVersion;
 		if(FEngineVersion::Parse(EngineVersionString, EngineVersion) && EngineVersion.HasChangelist() && EngineVersion.ToString(EVersionComponent::Minor) == TEXT("4.0"))
diff --git a/Engine/Source/Editor/DataTableEditor/Private/SDataTableListViewRowName.cpp b/Engine/Source/Editor/DataTableEditor/Private/SDataTableListViewRowName.cpp
--- a/Engine/Source/Editor/DataTableEditor/Private/SDataTableListViewRowName.cpp
+++ b/Engine/Source/Editor/DataTableEditor/Private/SDataTableListViewRowName.cpp
@@ -24,9 +24,9 @@ void SDataTableListViewRowName::Construct(const FArguments& InArgs, const TShare
 				.HeightOverride(RowDataPtr->DesiredRowHeight)
 				[
 					SNew(STextBlock)
-					.ColorAndOpacity(DataTableEditor.Get(), &FDataTableEditor::GetRowTextColor, RowDataPtr->RowId)
+					.ColorAndOpacity(DataTableEditor.Pin().Get(), &FDataTableEditor::GetRowTextColor, RowDataPtr->RowId)
 					.Text(RowDataPtr->DisplayName)
-					.HighlightText(DataTableEditor.Get(), &FDataTableEditor::GetFilterText)
+					.HighlightText(DataTableEditor.Pin().Get(), &FDataTableEditor::GetFilterText)
 				]
 			]
 		],
@@ -54,11 +54,14 @@ void SDataTableListViewRowName::OnSearchForReferences()
 {
 	if (DataTableEditor.IsValid() && RowDataPtr.IsValid())
 	{
-		UDataTable* SourceDataTable = const_cast<UDataTable*>(DataTableEditor->GetDataTable());
-		
-		TArray<FAssetIdentifier> AssetIdentifiers;
-		AssetIdentifiers.Add(FAssetIdentifier(SourceDataTable, RowDataPtr->RowId));
+		if (FDataTableEditor* DataTableEditorPtr = DataTableEditor.Pin().Get())
+		{
+			UDataTable* SourceDataTable = const_cast<UDataTable*>(DataTableEditorPtr->GetDataTable());
+
+			TArray<FAssetIdentifier> AssetIdentifiers;
+			AssetIdentifiers.Add(FAssetIdentifier(SourceDataTable, RowDataPtr->RowId));
 
-		FEditorDelegates::OnOpenReferenceViewer.Broadcast(AssetIdentifiers);
+			FEditorDelegates::OnOpenReferenceViewer.Broadcast(AssetIdentifiers);
+		}
 	}
 }
diff --git a/Engine/Source/Editor/DataTableEditor/Private/SDataTableListViewRowName.h b/Engine/Source/Editor/DataTableEditor/Private/SDataTableListViewRowName.h
--- a/Engine/Source/Editor/DataTableEditor/Private/SDataTableListViewRowName.h
+++ b/Engine/Source/Editor/DataTableEditor/Private/SDataTableListViewRowName.h
@@ -35,5 +35,5 @@ private:
 
 	void OnSearchForReferences();
 	FDataTableEditorRowListViewDataPtr RowDataPtr;
-	TSharedPtr<FDataTableEditor> DataTableEditor;
+	TWeakPtr<FDataTableEditor> DataTableEditor;
 };
\ No newline at end of file
diff --git a/Engine/Source/Editor/DetailCustomizations/Private/SkeletalMeshReductionSettingsDetails.cpp b/Engine/Source/Editor/DetailCustomizations/Private/SkeletalMeshReductionSettingsDetails.cpp
--- a/Engine/Source/Editor/DetailCustomizations/Private/SkeletalMeshReductionSettingsDetails.cpp
+++ b/Engine/Source/Editor/DetailCustomizations/Private/SkeletalMeshReductionSettingsDetails.cpp
@@ -107,6 +107,8 @@ void FSkeletalMeshReductionSettingsDetails::CustomizeChildren(TSharedRef<IProper
 		// Parameters not used by simplygon
 		const TArray<FName> CustomSimplifierOnlyProperties = {
 			GET_MEMBER_NAME_CHECKED(FSkeletalMeshOptimizationSettings, NumOfVertPercentage),
+			GET_MEMBER_NAME_CHECKED(FSkeletalMeshOptimizationSettings, MaxNumOfVerts),
+			GET_MEMBER_NAME_CHECKED(FSkeletalMeshOptimizationSettings, MaxNumOfTriangles),
 			GET_MEMBER_NAME_CHECKED(FSkeletalMeshOptimizationSettings, TerminationCriterion),
 			GET_MEMBER_NAME_CHECKED(FSkeletalMeshOptimizationSettings, bLockEdges),
 			GET_MEMBER_NAME_CHECKED(FSkeletalMeshOptimizationSettings, bEnforceBoneBoundaries),
@@ -176,6 +178,8 @@ void FSkeletalMeshReductionSettingsDetails::CustomizeChildren(TSharedRef<IProper
 		TSharedPtr<IPropertyHandle> VertPercentPropertyHandle = PropertyHandles.FindChecked(GET_MEMBER_NAME_CHECKED(FSkeletalMeshOptimizationSettings, NumOfVertPercentage));
 		TSharedPtr<IPropertyHandle> TriPercentPropertyHandle  = PropertyHandles.FindChecked(GET_MEMBER_NAME_CHECKED(FSkeletalMeshOptimizationSettings, NumOfTrianglesPercentage));
 
+		TSharedPtr<IPropertyHandle> MaxNumOfVertsPropertyHandle = PropertyHandles.FindChecked(GET_MEMBER_NAME_CHECKED(FSkeletalMeshOptimizationSettings, MaxNumOfVerts));
+		TSharedPtr<IPropertyHandle> MaxNumOfTrisPropertyHandle = PropertyHandles.FindChecked(GET_MEMBER_NAME_CHECKED(FSkeletalMeshOptimizationSettings, MaxNumOfTriangles));
 
 		for (auto Iter(PropertyHandles.CreateConstIterator()); Iter; ++Iter)
 		{
@@ -191,13 +195,43 @@ void FSkeletalMeshReductionSettingsDetails::CustomizeChildren(TSharedRef<IProper
 				// depending on the value of the pull down, optionally hide at most one of these.
 				if (Iter.Value() == VertPercentPropertyHandle)
 				{
+					const TArray< SkeletalMeshTerminationCriterion > VizList = {
+					SMTC_NumOfVerts,
+					SMTC_TriangleOrVert
+					};
+
 					// Hide property if using triangle percentage
-					SettingsRow.Visibility(TAttribute<EVisibility>::Create(TAttribute<EVisibility>::FGetter::CreateSP(this, &FSkeletalMeshReductionSettingsDetails::HideIfCurrentCriterionIs, SMTC_NumOfTriangles)));
+					SettingsRow.Visibility(TAttribute<EVisibility>::Create(TAttribute<EVisibility>::FGetter::CreateSP(this, &FSkeletalMeshReductionSettingsDetails::ShowIfCurrentCriterionIs, VizList)));
 				}
 				else if (Iter.Value() == TriPercentPropertyHandle)
 				{
+					const TArray< SkeletalMeshTerminationCriterion > VizList = 
+					{
+						SMTC_NumOfTriangles,
+						SMTC_TriangleOrVert
+					};
 					// Hide property if using vert percentage
-					SettingsRow.Visibility(TAttribute<EVisibility>::Create(TAttribute<EVisibility>::FGetter::CreateSP(this, &FSkeletalMeshReductionSettingsDetails::HideIfCurrentCriterionIs, SMTC_NumOfVerts)));
+					SettingsRow.Visibility(TAttribute<EVisibility>::Create(TAttribute<EVisibility>::FGetter::CreateSP(this, &FSkeletalMeshReductionSettingsDetails::ShowIfCurrentCriterionIs, VizList)));
+				}
+				else if (Iter.Value() == MaxNumOfVertsPropertyHandle)
+				{
+					const TArray< SkeletalMeshTerminationCriterion > VizList = 
+					{
+						SMTC_AbsNumOfVerts,
+						SMTC_AbsTriangleOrVert
+					};
+					// Hide property if using triangle percentage
+					SettingsRow.Visibility(TAttribute<EVisibility>::Create(TAttribute<EVisibility>::FGetter::CreateSP(this, &FSkeletalMeshReductionSettingsDetails::ShowIfCurrentCriterionIs, VizList)));
+				}
+				else if (Iter.Value() == MaxNumOfTrisPropertyHandle)
+				{
+					const TArray< SkeletalMeshTerminationCriterion > VizList = 
+					{
+						SMTC_AbsNumOfTriangles,
+						SMTC_AbsTriangleOrVert
+					};
+					// Hide property if using triangle percentage
+					SettingsRow.Visibility(TAttribute<EVisibility>::Create(TAttribute<EVisibility>::FGetter::CreateSP(this, &FSkeletalMeshReductionSettingsDetails::ShowIfCurrentCriterionIs, VizList)));
 				}
 
 			}
@@ -266,19 +300,19 @@ EVisibility FSkeletalMeshReductionSettingsDetails::GetVisibiltyIfCurrentReductio
 	return EVisibility::Hidden;
 }
 
-EVisibility FSkeletalMeshReductionSettingsDetails::HideIfCurrentCriterionIs(SkeletalMeshTerminationCriterion TerminationCriterion) const
+EVisibility FSkeletalMeshReductionSettingsDetails::ShowIfCurrentCriterionIs(TArray<SkeletalMeshTerminationCriterion> TerminationCriterionArray) const
 {
 	uint8 CurrentEnum;
 	if (TerminationCriterionPopertyHandle->GetValue(CurrentEnum) != FPropertyAccess::Fail)
 	{
 		enum SkeletalMeshTerminationCriterion CurrentReductionType = (SkeletalMeshTerminationCriterion)CurrentEnum;
-		if (CurrentReductionType == TerminationCriterion)
+		if (TerminationCriterionArray.Contains(CurrentReductionType))
 		{
-			return EVisibility::Hidden; 
+			return EVisibility::Visible;
 		}
 	}
 
-	return EVisibility::Visible;
+	return EVisibility::Hidden;
 }
 
 bool FSkeletalMeshReductionSettingsDetails::UseNativeLODTool() const
diff --git a/Engine/Source/Editor/DetailCustomizations/Private/SkeletalMeshReductionSettingsDetails.h b/Engine/Source/Editor/DetailCustomizations/Private/SkeletalMeshReductionSettingsDetails.h
--- a/Engine/Source/Editor/DetailCustomizations/Private/SkeletalMeshReductionSettingsDetails.h
+++ b/Engine/Source/Editor/DetailCustomizations/Private/SkeletalMeshReductionSettingsDetails.h
@@ -39,7 +39,7 @@ private:
 	EVisibility GetVisibiltyIfCurrentReductionMethodIsNot(SkeletalMeshOptimizationType ReductionType) const;
 
 	// Used by the native tool UI.
-	EVisibility HideIfCurrentCriterionIs(SkeletalMeshTerminationCriterion TerminationCriterion) const;
+	EVisibility ShowIfCurrentCriterionIs(TArray<SkeletalMeshTerminationCriterion> TerminationCriterionArray) const;
 
 	/** Detect usage of thirdparty vs native tool */
 	bool UseNativeLODTool() const;
diff --git a/Engine/Source/Editor/UnrealEd/Private/Commandlets/FixConflictingLocalizationKeys.cpp b/Engine/Source/Editor/UnrealEd/Private/Commandlets/FixConflictingLocalizationKeys.cpp
--- a/Engine/Source/Editor/UnrealEd/Private/Commandlets/FixConflictingLocalizationKeys.cpp
+++ b/Engine/Source/Editor/UnrealEd/Private/Commandlets/FixConflictingLocalizationKeys.cpp
@@ -178,7 +178,7 @@ bool ReKeyTextProperty(UStruct* InOuterType, void* InAddrToUpdate, const TArray<
 						// ContainerIndex is the element index, but we need the sparse index
 						int32 SparseIndex = 0;
 						{
-							const int32 ElementCount = ScriptMapHelper.GetMaxIndex();
+							const int32 ElementCount = ScriptMapHelper.Num();
 							for (int32 ElementIndex = 0; ElementIndex < ElementCount; ++SparseIndex)
 							{
 								if (ScriptMapHelper.IsValidIndex(SparseIndex))
@@ -215,7 +215,7 @@ bool ReKeyTextProperty(UStruct* InOuterType, void* InAddrToUpdate, const TArray<
 						// ContainerIndex is the element index, but we need the sparse index
 						int32 SparseIndex = 0;
 						{
-							const int32 ElementCount = ScriptSetHelper.GetMaxIndex();
+							const int32 ElementCount = ScriptSetHelper.Num();
 							for (int32 ElementIndex = 0; ElementIndex < ElementCount; ++SparseIndex)
 							{
 								if (ScriptSetHelper.IsValidIndex(SparseIndex))
diff --git a/Engine/Source/Programs/AutomationTool/AutomationUtils/P4Environment.cs b/Engine/Source/Programs/AutomationTool/AutomationUtils/P4Environment.cs
--- a/Engine/Source/Programs/AutomationTool/AutomationUtils/P4Environment.cs
+++ b/Engine/Source/Programs/AutomationTool/AutomationUtils/P4Environment.cs
@@ -178,7 +178,8 @@ namespace AutomationTool
 				P4ClientInfo ThisClient;
 				if(String.IsNullOrEmpty(Client))
 				{
-					ThisClient = DetectClient(DefaultConnection, User, Environment.MachineName.ToLower(), CmdEnv.UATExe);
+					string HostName = System.Net.Dns.GetHostName();
+					ThisClient = DetectClient(DefaultConnection, User, HostName, CmdEnv.UATExe);
 					Log.TraceInformation("Using user {0} clientspec {1} {2}", User, ThisClient.Name, ThisClient.RootPath);
 					Client = ThisClient.Name;
 					CommandUtils.SetEnvVar(EnvVarNames.Client, Client);
diff --git a/Engine/Source/Programs/CrashReportClient/Private/CrashDescription.cpp b/Engine/Source/Programs/CrashReportClient/Private/CrashDescription.cpp
--- a/Engine/Source/Programs/CrashReportClient/Private/CrashDescription.cpp
+++ b/Engine/Source/Programs/CrashReportClient/Private/CrashDescription.cpp
@@ -138,7 +138,7 @@ void FPrimaryCrashProperties::Shutdown()
 
 void FPrimaryCrashProperties::UpdateIDs()
 {
-	const bool bAddPersonalData = FCrashReportClientConfig::Get().GetAllowToBeContacted() || FEngineBuildSettings::IsInternalBuild();
+	const bool bAddPersonalData = FCrashReportClientConfig::Get().GetAllowToBeContacted();
 	bAllowToBeContacted = bAddPersonalData;
 	if (bAddPersonalData)
 	{
diff --git a/Engine/Source/Programs/CrashReportClient/Private/CrashReportClientConfig.cpp b/Engine/Source/Programs/CrashReportClient/Private/CrashReportClientConfig.cpp
--- a/Engine/Source/Programs/CrashReportClient/Private/CrashReportClientConfig.cpp
+++ b/Engine/Source/Programs/CrashReportClient/Private/CrashReportClientConfig.cpp
@@ -4,6 +4,7 @@
 #include "Logging/LogMacros.h"
 #include "Misc/ConfigCacheIni.h"
 #include "Misc/App.h"
+#include "Misc/EngineBuildSettings.h"
 #include "CrashReportClientApp.h"
 #include "GenericPlatform/GenericPlatformCrashContext.h"
 
@@ -52,7 +53,11 @@ FCrashReportClientConfig::FCrashReportClientConfig()
 		UE_LOG(CrashReportClientLog, Log, TEXT("DataRouterUrl: %s"), *DataRouterUrl);
 	}
 
-	if (!GConfig->GetBool( TEXT( "CrashReportClient" ), TEXT( "bAllowToBeContacted" ), bAllowToBeContacted, GEngineIni ))
+	if (FEngineBuildSettings::IsInternalBuild())
+	{
+		bAllowToBeContacted = true;
+	}
+	else if (!GConfig->GetBool( TEXT( "CrashReportClient" ), TEXT( "bAllowToBeContacted" ), bAllowToBeContacted, GEngineIni ))
 	{
 		// Default to true when unattended when config is missing. This is mostly for dedicated servers that do not have config files for CRC.
 		if (bUnattended)
diff --git a/Engine/Source/Programs/CrashReportClient/Private/SCrashReportClient.cpp b/Engine/Source/Programs/CrashReportClient/Private/SCrashReportClient.cpp
--- a/Engine/Source/Programs/CrashReportClient/Private/SCrashReportClient.cpp
+++ b/Engine/Source/Programs/CrashReportClient/Private/SCrashReportClient.cpp
@@ -20,6 +20,7 @@
 #include "Widgets/Input/SButton.h"
 #include "Widgets/Layout/SSpacer.h"
 #include "Framework/Application/SlateApplication.h"
+#include "Misc/EngineBuildSettings.h"
 
 #define LOCTEXT_NAMESPACE "CrashReportClient"
 
@@ -236,6 +237,7 @@ void SCrashReportClient::Construct(const FArguments& InArgs, TSharedRef<FCrashRe
 				[
 					SNew(SCheckBox)
 					.IsChecked( FCrashReportClientConfig::Get().GetAllowToBeContacted() ? ECheckBoxState::Checked : ECheckBoxState::Unchecked )
+					.IsEnabled( !FEngineBuildSettings::IsInternalBuild() )
 					.OnCheckStateChanged(CrashReportClient.ToSharedRef(), &FCrashReportClient::AllowToBeContacted_OnCheckStateChanged)
 				]
 
@@ -246,6 +248,7 @@ void SCrashReportClient::Construct(const FArguments& InArgs, TSharedRef<FCrashRe
 				[
 					SNew(STextBlock)
 					.AutoWrapText(true)
+					.IsEnabled( !FEngineBuildSettings::IsInternalBuild() )
 					.Text(LOCTEXT("IAgree", "I agree to be contacted by Epic Games via email if additional information about this crash would help fix it."))
 				]
 			]
diff --git a/Engine/Source/Programs/UnrealBuildTool/Platform/Windows/UEBuildWindows.cs b/Engine/Source/Programs/UnrealBuildTool/Platform/Windows/UEBuildWindows.cs
--- a/Engine/Source/Programs/UnrealBuildTool/Platform/Windows/UEBuildWindows.cs
+++ b/Engine/Source/Programs/UnrealBuildTool/Platform/Windows/UEBuildWindows.cs
@@ -203,6 +203,13 @@ namespace UnrealBuildTool
 		/// </summary>
 		public bool bOptimizeGlobalData = true;
 
+		/// <summary>
+		/// (Experimental) Appends the -ftime-trace argument to the command line for Clang, to output a JSON file containing a timeline for the compile. 
+		/// See http://aras-p.info/blog/2019/01/16/time-trace-timeline-flame-chart-profiler-for-Clang/ for more info.
+		/// </summary>
+		[XmlConfigFile(Category = "WindowsPlatform")]
+		public bool bClangTimeTrace = false;
+
 		/// <summary>
 		/// The Visual C++ environment to use for this target. Only initialized after all the target settings are finalized, in ValidateTarget().
 		/// </summary>
@@ -375,6 +382,11 @@ namespace UnrealBuildTool
 			get { return Inner.bOptimizeGlobalData; }
 		}
 
+		public bool bClangTimeTrace
+		{
+			get { return Inner.bClangTimeTrace; }
+		}
+
 		public string GetVisualStudioCompilerVersionName()
 		{
 			return Inner.GetVisualStudioCompilerVersionName();
diff --git a/Engine/Source/Programs/UnrealBuildTool/Platform/Windows/VCEnvironment.cs b/Engine/Source/Programs/UnrealBuildTool/Platform/Windows/VCEnvironment.cs
--- a/Engine/Source/Programs/UnrealBuildTool/Platform/Windows/VCEnvironment.cs
+++ b/Engine/Source/Programs/UnrealBuildTool/Platform/Windows/VCEnvironment.cs
@@ -502,7 +502,7 @@ namespace UnrealBuildTool
 			if(Compiler == WindowsCompiler.Clang || Compiler == WindowsCompiler.Intel)
 			{
 				ToolChain = WindowsCompiler.VisualStudio2017;
-				if(!WindowsPlatform.TryGetToolChainDir(ToolChain, CompilerVersion, out SelectedToolChainVersion, out SelectedToolChainDir))
+				if(!WindowsPlatform.TryGetToolChainDir(ToolChain, null, out SelectedToolChainVersion, out SelectedToolChainDir))
 				{
 					throw new BuildException("{0}{1} must be installed in order to build this target.", WindowsPlatform.GetCompilerName(Compiler), String.IsNullOrEmpty(CompilerVersion)? "" : String.Format(" ({0})", CompilerVersion));
 				}
diff --git a/Engine/Source/Programs/UnrealBuildTool/Platform/Windows/VCToolChain.cs b/Engine/Source/Programs/UnrealBuildTool/Platform/Windows/VCToolChain.cs
--- a/Engine/Source/Programs/UnrealBuildTool/Platform/Windows/VCToolChain.cs
+++ b/Engine/Source/Programs/UnrealBuildTool/Platform/Windows/VCToolChain.cs
@@ -1020,6 +1020,13 @@ namespace UnrealBuildTool
 					CompileAction.ProducedItems.Add(ObjectFile);
 					Result.ObjectFiles.Add(ObjectFile);
 					FileArguments.Add(String.Format("/Fo\"{0}\"", ObjectFile.AbsolutePath));
+
+					// Experimental: support for JSON output of timing data
+					if(Target.WindowsPlatform.Compiler == WindowsCompiler.Clang && Target.WindowsPlatform.bClangTimeTrace)
+					{
+						SharedArguments.Add("-Xclang -ftime-trace");
+						CompileAction.ProducedItems.Add(FileItem.GetItemByFileReference(ObjectFile.Location.ChangeExtension(".json")));
+					}
 				}
 
 				// Don't farm out creation of precompiled headers as it is the critical path task.
diff --git a/Engine/Source/Programs/UnrealBuildTool/System/HotReload.cs b/Engine/Source/Programs/UnrealBuildTool/System/HotReload.cs
--- a/Engine/Source/Programs/UnrealBuildTool/System/HotReload.cs
+++ b/Engine/Source/Programs/UnrealBuildTool/System/HotReload.cs
@@ -144,7 +144,7 @@ namespace UnrealBuildTool
 			if (TargetDesc.Name.EndsWith("Editor", StringComparison.OrdinalIgnoreCase))
 			{
 				string EditorBaseFileName = "UE4Editor";
-				if (TargetDesc.Configuration != UnrealTargetConfiguration.Development && TargetDesc.Configuration != UnrealTargetConfiguration.DebugGame)
+				if (TargetDesc.Configuration != UnrealTargetConfiguration.Development)
 				{
 					EditorBaseFileName = String.Format("{0}-{1}-{2}", EditorBaseFileName, TargetDesc.Platform, TargetDesc.Configuration);
 				}
diff --git a/Engine/Source/Runtime/Advertising/IOS/IOSAdvertising/Private/IOSAdvertising.cpp b/Engine/Source/Runtime/Advertising/IOS/IOSAdvertising/Private/IOSAdvertising.cpp
--- a/Engine/Source/Runtime/Advertising/IOS/IOSAdvertising/Private/IOSAdvertising.cpp
+++ b/Engine/Source/Runtime/Advertising/IOS/IOSAdvertising/Private/IOSAdvertising.cpp
@@ -10,7 +10,7 @@ DEFINE_LOG_CATEGORY_STATIC( LogAdvertising, Display, All );
 
 IMPLEMENT_MODULE( FIOSAdvertisingProvider, IOSAdvertising );
 
-#if __IPHONE_OS_VERSION_MIN_REQUIRED < __IPHONE_10_0
+#if __IPHONE_OS_VERSION_MIN_REQUIRED < __IPHONE_10_0 && !PLATFORM_TVOS
 #import <iAd/ADBannerView.h>
 
 @interface IOSAdvertising : UIResponder <ADBannerViewDelegate>
diff --git a/Engine/Source/Runtime/Apple/MetalRHI/Private/MetalUAV.cpp b/Engine/Source/Runtime/Apple/MetalRHI/Private/MetalUAV.cpp
--- a/Engine/Source/Runtime/Apple/MetalRHI/Private/MetalUAV.cpp
+++ b/Engine/Source/Runtime/Apple/MetalRHI/Private/MetalUAV.cpp
@@ -848,39 +848,39 @@ void FMetalGPUFence::WriteInternal(mtlpp::CommandBuffer& CmdBuffer)
 	check(Fence);
 }
 
-void FMetalRHICommandContext::RHIEnqueueStagedRead(FStagingBufferRHIParamRef StagingBuffer, FGPUFenceRHIParamRef InFence, uint32 Offset, uint32 NumBytes)
+void FMetalRHICommandContext::RHICopyToStagingBuffer(FVertexBufferRHIParamRef SourceBufferRHI, FStagingBufferRHIParamRef DestinationStagingBufferRHI, uint32 Offset, uint32 NumBytes, FGPUFenceRHIParamRef FenceRHI)
 {
 	@autoreleasepool {
-		check(StagingBuffer);
-		
-		FMetalStagingBuffer* MetalStagingBuffer = ResourceCast(StagingBuffer);
-		FMetalVertexBuffer* BackingBuffer = ResourceCast(MetalStagingBuffer->GetBackingBuffer());
-		FMetalBuffer& ReadbackBuffer = MetalStagingBuffer->ReadbackStagingBuffer;
-		
+		check(DestinationStagingBufferRHI);
+
+		FMetalStagingBuffer* MetalStagingBuffer = ResourceCast(DestinationStagingBufferRHI);
+		ensureMsgf(!MetalStagingBuffer->bIsLocked, TEXT("Attempting to Copy to a locked staging buffer. This may have undefined behavior"));
+		FMetalVertexBuffer* SourceBuffer = ResourceCast(SourceBufferRHI);
+		FMetalBuffer& ReadbackBuffer = MetalStagingBuffer->ShadowBuffer;
+
 		// Need a shadow buffer for this read. If it hasn't been allocated in our FStagingBuffer or if
 		// it's not big enough to hold our readback we need to allocate.
-		if(!ReadbackBuffer || ReadbackBuffer.GetLength() < NumBytes)
+		if (!ReadbackBuffer || ReadbackBuffer.GetLength() < NumBytes)
 		{
-			if(ReadbackBuffer)
+			if (ReadbackBuffer)
 			{
 				SafeReleaseMetalBuffer(ReadbackBuffer);
 			}
 			FMetalPooledBufferArgs ArgsCPU(GetMetalDeviceContext().GetDevice(), NumBytes, mtlpp::StorageMode::Shared);
 			ReadbackBuffer = GetMetalDeviceContext().CreatePooledBuffer(ArgsCPU);
 		}
-		
+
 		// Inline copy from the actual buffer to the shadow
-		GetMetalDeviceContext().CopyFromBufferToBuffer(BackingBuffer->Buffer, Offset, ReadbackBuffer, 0, NumBytes);
-		
-		if (InFence)
+		GetMetalDeviceContext().CopyFromBufferToBuffer(SourceBuffer->Buffer, Offset, ReadbackBuffer, 0, NumBytes);
+
+		if (FenceRHI)
 		{
-			FMetalGPUFence* Fence = ResourceCast(InFence);
+			FMetalGPUFence* Fence = ResourceCast(FenceRHI);
 			Fence->WriteInternal(Context->GetCurrentCommandBuffer());
 		}
 	}
 }
 
-
 FGPUFenceRHIRef FMetalDynamicRHI::RHICreateGPUFence(const FName &Name)
 {
 	@autoreleasepool {
@@ -888,6 +888,11 @@ FGPUFenceRHIRef FMetalDynamicRHI::RHICreateGPUFence(const FName &Name)
 	}
 }
 
+void FMetalGPUFence::Clear()
+{
+	Fence = mtlpp::CommandBufferFence();
+}
+
 bool FMetalGPUFence::Poll() const
 {
 	if (Fence)
diff --git a/Engine/Source/Runtime/Apple/MetalRHI/Private/MetalVertexBuffer.cpp b/Engine/Source/Runtime/Apple/MetalRHI/Private/MetalVertexBuffer.cpp
--- a/Engine/Source/Runtime/Apple/MetalRHI/Private/MetalVertexBuffer.cpp
+++ b/Engine/Source/Runtime/Apple/MetalRHI/Private/MetalVertexBuffer.cpp
@@ -769,16 +769,16 @@ void FMetalDynamicRHI::UnlockStagingBuffer_RenderThread(class FRHICommandListImm
 	return RHIUnlockStagingBuffer(StagingBuffer);
 }
 
-FStagingBufferRHIRef FMetalDynamicRHI::RHICreateStagingBuffer(FVertexBufferRHIParamRef VertexBuffer)
+FStagingBufferRHIRef FMetalDynamicRHI::RHICreateStagingBuffer()
 {
-	return new FMetalStagingBuffer(VertexBuffer);
+	return new FMetalStagingBuffer();
 }
 
 FMetalStagingBuffer::~FMetalStagingBuffer()
 {
-	if(ReadbackStagingBuffer)
+	if (ShadowBuffer)
 	{
-		SafeReleaseMetalBuffer(ReadbackStagingBuffer);
+		SafeReleaseMetalBuffer(ShadowBuffer);
 	}
 }
 
@@ -786,13 +786,16 @@ FMetalStagingBuffer::~FMetalStagingBuffer()
 // If this was not fenced correctly it will not have the expected data.
 void *FMetalStagingBuffer::Lock(uint32 Offset, uint32 NumBytes)
 {
-	check(ReadbackStagingBuffer);
-	
-	uint8* BackingPtr = (uint8*) ReadbackStagingBuffer.GetContents();
-	return BackingPtr+Offset;
+	check(ShadowBuffer);
+	check(!bIsLocked);
+	bIsLocked = true;
+	uint8* BackingPtr = (uint8*)ShadowBuffer.GetContents();
+	return BackingPtr + Offset;
 }
 
 void FMetalStagingBuffer::Unlock()
 {
 	// does nothing in metal.
+	check(bIsLocked);
+	bIsLocked = false;
 }
diff --git a/Engine/Source/Runtime/Apple/MetalRHI/Public/MetalRHI.h b/Engine/Source/Runtime/Apple/MetalRHI/Public/MetalRHI.h
--- a/Engine/Source/Runtime/Apple/MetalRHI/Public/MetalRHI.h
+++ b/Engine/Source/Runtime/Apple/MetalRHI/Public/MetalRHI.h
@@ -189,7 +189,7 @@ public:
 	virtual FComputeFenceRHIRef RHICreateComputeFence(const FName& Name) final override;
 	virtual FGPUFenceRHIRef RHICreateGPUFence(const FName &Name) final override;
 
-	virtual FStagingBufferRHIRef RHICreateStagingBuffer(FVertexBufferRHIParamRef VertexBuffer) final override;
+	virtual FStagingBufferRHIRef RHICreateStagingBuffer() final override;
 	virtual void* RHILockStagingBuffer(FStagingBufferRHIParamRef StagingBuffer, uint32 Offset, uint32 SizeRHI) final override;
 	virtual void RHIUnlockStagingBuffer(FStagingBufferRHIParamRef StagingBuffer) final override;
 	virtual void* LockStagingBuffer_RenderThread(class FRHICommandListImmediate& RHICmdList, FStagingBufferRHIParamRef StagingBuffer, uint32 Offset, uint32 SizeRHI) final override;
diff --git a/Engine/Source/Runtime/Apple/MetalRHI/Public/MetalRHIContext.h b/Engine/Source/Runtime/Apple/MetalRHI/Public/MetalRHIContext.h
--- a/Engine/Source/Runtime/Apple/MetalRHI/Public/MetalRHIContext.h
+++ b/Engine/Source/Runtime/Apple/MetalRHI/Public/MetalRHIContext.h
@@ -318,7 +318,7 @@ public:
 	 */
 	virtual void RHITransitionResources(EResourceTransitionAccess TransitionType, EResourceTransitionPipeline TransitionPipeline, FUnorderedAccessViewRHIParamRef* InUAVs, int32 NumUAVs, FComputeFenceRHIParamRef WriteComputeFence) final override;
 	
-	virtual void RHIEnqueueStagedRead(FStagingBufferRHIParamRef StagingBuffer, FGPUFenceRHIParamRef InFence, uint32 Offset, uint32 NumBytes) final override;
+	virtual void RHICopyToStagingBuffer(FVertexBufferRHIParamRef SourceBufferRHI, FStagingBufferRHIParamRef DestinationStagingBufferRHI, uint32 Offset, uint32 NumBytes, FGPUFenceRHIParamRef FenceRHI) final override;
 
 	/**
 	 * Explicitly transition a texture resource from readable -> writable by the GPU or vice versa.
diff --git a/Engine/Source/Runtime/Apple/MetalRHI/Public/MetalResources.h b/Engine/Source/Runtime/Apple/MetalRHI/Public/MetalResources.h
--- a/Engine/Source/Runtime/Apple/MetalRHI/Public/MetalResources.h
+++ b/Engine/Source/Runtime/Apple/MetalRHI/Public/MetalResources.h
@@ -1033,6 +1033,8 @@ public:
 	{
 	}
 
+	virtual void Clear() override final;
+
 	void WriteInternal(mtlpp::CommandBuffer& CmdBuffer);
 
 	virtual bool Poll() const override final;
@@ -1043,19 +1045,20 @@ private:
 
 class FMetalStagingBuffer final : public FRHIStagingBuffer
 {
+	friend class FMetalRHICommandContext;
 public:
-	FMetalStagingBuffer(FVertexBufferRHIRef InBuffer)
-	: FRHIStagingBuffer(InBuffer)
+	FMetalStagingBuffer()
+		: FRHIStagingBuffer()
 	{
 	}
-	
-	~FMetalStagingBuffer();
 
-	void *Lock(uint32 Offset, uint32 NumBytes);
+	virtual ~FMetalStagingBuffer() final override;
 
-	void Unlock();
-	
-	FMetalBuffer ReadbackStagingBuffer;
+	virtual void* Lock(uint32 Offset, uint32 NumBytes) final override;
+
+	virtual void Unlock() final override;
+private:
+	FMetalBuffer ShadowBuffer;
 };
 
 class FMetalShaderLibrary final : public FRHIShaderLibrary
diff --git a/Engine/Source/Runtime/ApplicationCore/Private/IOS/IOSInputInterface.cpp b/Engine/Source/Runtime/ApplicationCore/Private/IOS/IOSInputInterface.cpp
--- a/Engine/Source/Runtime/ApplicationCore/Private/IOS/IOSInputInterface.cpp
+++ b/Engine/Source/Runtime/ApplicationCore/Private/IOS/IOSInputInterface.cpp
@@ -351,7 +351,7 @@ else if(bIsPressed) \
     double* NextRepeatTime = NextKeyRepeatTime.Find(UEButton); \
     if(NextRepeatTime && *NextRepeatTime <= CurrentTime) \
     { \
-        MessageHandler->OnControllerButtonPressed(UEButton, Cont.playerIndex, false); \
+        MessageHandler->OnControllerButtonPressed(UEButton, Cont.playerIndex, true); \
 *NextRepeatTime = CurrentTime + RepeatDelay; \
     } \
 } \
diff --git a/Engine/Source/Runtime/CoreUObject/Private/Serialization/PropertyLocalizationDataGathering.cpp b/Engine/Source/Runtime/CoreUObject/Private/Serialization/PropertyLocalizationDataGathering.cpp
--- a/Engine/Source/Runtime/CoreUObject/Private/Serialization/PropertyLocalizationDataGathering.cpp
+++ b/Engine/Source/Runtime/CoreUObject/Private/Serialization/PropertyLocalizationDataGathering.cpp
@@ -295,7 +295,7 @@ void FPropertyLocalizationDataGatherer::GatherLocalizationDataFromChildTextPrope
 
 			// Iterate over all elements of the map.
 			FScriptMapHelper ScriptMapHelper(MapProperty, ElementValueAddress);
-			const int32 ElementCount = ScriptMapHelper.GetMaxIndex();
+			const int32 ElementCount = ScriptMapHelper.Num();
 			for(int32 j = 0, ElementIndex = 0; ElementIndex < ElementCount; ++j)
 			{
 				if (!ScriptMapHelper.IsValidIndex(j))
@@ -338,7 +338,7 @@ void FPropertyLocalizationDataGatherer::GatherLocalizationDataFromChildTextPrope
 
 			// Iterate over all elements of the Set.
 			FScriptSetHelper ScriptSetHelper(SetProperty, ElementValueAddress);
-			const int32 ElementCount = ScriptSetHelper.GetMaxIndex();
+			const int32 ElementCount = ScriptSetHelper.Num();
 			for(int32 j = 0, ElementIndex = 0; ElementIndex < ElementCount; ++j)
 			{
 				if (!ScriptSetHelper.IsValidIndex(j))
diff --git a/Engine/Source/Runtime/D3D12RHI/Private/D3D12CommandContext.h b/Engine/Source/Runtime/D3D12RHI/Private/D3D12CommandContext.h
--- a/Engine/Source/Runtime/D3D12RHI/Private/D3D12CommandContext.h
+++ b/Engine/Source/Runtime/D3D12RHI/Private/D3D12CommandContext.h
@@ -262,7 +262,7 @@ public:
 	virtual void RHIClearTinyUAV(FUnorderedAccessViewRHIParamRef UnorderedAccessViewRHI, const uint32* Values) final override;
 	virtual void RHICopyToResolveTarget(FTextureRHIParamRef SourceTexture, FTextureRHIParamRef DestTexture, const FResolveParams& ResolveParams) final override;
 	virtual void RHITransitionResources(EResourceTransitionAccess TransitionType, FTextureRHIParamRef* InTextures, int32 NumTextures) final override;
-    virtual void RHIEnqueueStagedRead(FStagingBufferRHIParamRef StagingBuffer, FGPUFenceRHIParamRef Fence, uint32 Offset, uint32 NumBytes) final override;
+	virtual void RHICopyToStagingBuffer(FVertexBufferRHIParamRef SourceBuffer, FStagingBufferRHIParamRef DestinationStagingBuffer, uint32 Offset, uint32 NumBytes, FGPUFenceRHIParamRef Fence = nullptr) final override;
 	virtual void RHIBeginRenderQuery(FRenderQueryRHIParamRef RenderQuery) final override;
 	virtual void RHIEndRenderQuery(FRenderQueryRHIParamRef RenderQuery) final override;
 	void RHIBeginOcclusionQueryBatch(uint32 NumQueriesInBatch);
@@ -449,9 +449,9 @@ public:
 	// Special implementation that only signal the fence once.
 	virtual void RHITransitionResources(EResourceTransitionAccess TransitionType, EResourceTransitionPipeline TransitionPipeline, FUnorderedAccessViewRHIParamRef* InUAVs, int32 NumUAVs, FComputeFenceRHIParamRef WriteComputeFenceRHI) final override;
 
-	FORCEINLINE virtual void RHIEnqueueStagedRead(FStagingBufferRHIParamRef StagingBuffer, FGPUFenceRHIParamRef Fence, uint32 Offset, uint32 NumBytes) final override
+	FORCEINLINE virtual void RHICopyToStagingBuffer(FVertexBufferRHIParamRef SourceBuffer, FStagingBufferRHIParamRef DestinationStagingBuffer, uint32 Offset, uint32 NumBytes, FGPUFenceRHIParamRef Fence = nullptr) final override
 	{
-		ContextRedirect(RHIEnqueueStagedRead(StagingBuffer, Fence, Offset, NumBytes));
+		ContextRedirect(RHICopyToStagingBuffer(SourceBuffer, DestinationStagingBuffer, Offset, NumBytes, Fence));
 	}
 	FORCEINLINE virtual void RHISetShaderTexture(FComputeShaderRHIParamRef PixelShader, uint32 TextureIndex, FTextureRHIParamRef NewTexture) final override
 	{
diff --git a/Engine/Source/Runtime/D3D12RHI/Private/D3D12Commands.cpp b/Engine/Source/Runtime/D3D12RHI/Private/D3D12Commands.cpp
--- a/Engine/Source/Runtime/D3D12RHI/Private/D3D12Commands.cpp
+++ b/Engine/Source/Runtime/D3D12RHI/Private/D3D12Commands.cpp
@@ -395,19 +395,32 @@ void FD3D12CommandContext::RHITransitionResources(EResourceTransitionAccess Tran
 	}
 }
 
-void FD3D12CommandContext::RHIEnqueueStagedRead(FStagingBufferRHIParamRef StagingBufferRHI, FGPUFenceRHIParamRef FenceRHI, uint32 Offset, uint32 NumBytes)
+
+void FD3D12CommandContext::RHICopyToStagingBuffer(FVertexBufferRHIParamRef SourceBufferRHI, FStagingBufferRHIParamRef StagingBufferRHI, uint32 Offset, uint32 NumBytes, FGPUFenceRHIParamRef FenceRHI)
 {
 	FD3D12StagingBuffer* StagingBuffer = FD3D12DynamicRHI::ResourceCast(StagingBufferRHI);
 	check(StagingBuffer);
+	ensureMsgf(!StagingBuffer->bIsLocked, TEXT("Attempting to Copy to a locked staging buffer. This may have undefined behavior"));
 
-	FD3D12VertexBuffer* VertexBuffer = FD3D12DynamicRHI::ResourceCast(StagingBuffer->GetBackingBuffer());
+	FD3D12VertexBuffer* VertexBuffer = FD3D12DynamicRHI::ResourceCast(SourceBufferRHI);
 	check(VertexBuffer);
 
 	// Only get data from the first gpu for now.
 	FD3D12Device* StagingDevice = VertexBuffer->GetParentDevice();
-	StagingBuffer->StagedRead.SafeRelease();
-	VERIFYD3D12RESULT(GetParentDevice()->GetParentAdapter()->CreateBuffer(D3D12_HEAP_TYPE_READBACK, GetGPUMask(), GetGPUMask(), Offset + NumBytes, StagingBuffer->StagedRead.GetInitReference()));
-	
+
+	// Ensure our shadow buffer is large enough to hold the readback.
+	if (!StagingBuffer->StagedRead || StagingBuffer->ShadowBufferSize < NumBytes)
+	{
+		// @todo-mattc I feel like we should allocate more than NumBytes to handle small reads without blowing tons of space. Need to pool this.
+		// Hopefully d3d12 will do smart pooling out of an internal heap.
+		if (StagingBuffer->StagedRead)
+		{
+			StagingBuffer->StagedRead.SafeRelease();
+		}
+		VERIFYD3D12RESULT(GetParentDevice()->GetParentAdapter()->CreateBuffer(D3D12_HEAP_TYPE_READBACK, GetGPUMask(), GetGPUMask(), NumBytes, StagingBuffer->StagedRead.GetInitReference()));
+		StagingBuffer->ShadowBufferSize = NumBytes;
+	}
+
 	{
 		FD3D12Resource* pSourceResource = VertexBuffer->ResourceLocation.GetResource();
 		D3D12_RESOURCE_DESC const& SourceBufferDesc = pSourceResource->GetDesc();
@@ -427,6 +440,7 @@ void FD3D12CommandContext::RHIEnqueueStagedRead(FStagingBufferRHIParamRef Stagin
 
 	if (FenceRHI)
 	{
+		// @todo-mattc we don't want to flush here. That should be the caller's responsibility.
 		RHISubmitCommandsHint();
 		FD3D12GPUFence* Fence = FD3D12DynamicRHI::ResourceCast(FenceRHI);
 		Fence->WriteInternal(ED3D12CommandQueueType::Default);
diff --git a/Engine/Source/Runtime/D3D12RHI/Private/D3D12DirectCommandListManager.cpp b/Engine/Source/Runtime/D3D12RHI/Private/D3D12DirectCommandListManager.cpp
--- a/Engine/Source/Runtime/D3D12RHI/Private/D3D12DirectCommandListManager.cpp
+++ b/Engine/Source/Runtime/D3D12RHI/Private/D3D12DirectCommandListManager.cpp
@@ -24,7 +24,13 @@ void FD3D12GPUFence::WriteInternal(ED3D12CommandQueueType QueueType)
 
 bool FD3D12GPUFence::Poll() const
 {
-	return !Value || (Fence && Fence->IsFenceComplete(Value));
+	// @todo-mattc Value of 0 means signaled? Revisit this...
+	return !Value || (Fence && Fence->PeekLastCompletedFence() >= Value);
+}
+
+void FD3D12GPUFence::Clear()
+{
+	Value = MAX_uint64;
 }
 
 
@@ -33,9 +39,9 @@ FGPUFenceRHIRef FD3D12DynamicRHI::RHICreateGPUFence(const FName& Name)
 	return new FD3D12GPUFence(Name, GetAdapter().GetStagingFence());
 }
 
-FStagingBufferRHIRef FD3D12DynamicRHI::RHICreateStagingBuffer(FVertexBufferRHIParamRef VertexBufferRHI)
+FStagingBufferRHIRef FD3D12DynamicRHI::RHICreateStagingBuffer()
 {
-	return new FD3D12StagingBuffer(VertexBufferRHI);
+	return new FD3D12StagingBuffer();
 }
 
 void* FD3D12DynamicRHI::RHILockStagingBuffer(FStagingBufferRHIParamRef StagingBufferRHI, uint32 Offset, uint32 SizeRHI)
@@ -43,30 +49,14 @@ void* FD3D12DynamicRHI::RHILockStagingBuffer(FStagingBufferRHIParamRef StagingBu
 	FD3D12StagingBuffer* StagingBuffer = FD3D12DynamicRHI::ResourceCast(StagingBufferRHI);
 	check(StagingBuffer);
 
-	FD3D12Resource* pResource = StagingBuffer->StagedRead.GetReference();
-	if (pResource)
-	{
-		D3D12_RANGE ReadRange;
-		ReadRange.Begin = Offset;
-		ReadRange.End = Offset + SizeRHI;
-		return reinterpret_cast<uint8*>(pResource->Map(&ReadRange)) + Offset;
-	}
-	else
-	{
-		return nullptr;
-	}
+	return StagingBuffer->Lock(Offset, SizeRHI);
 }
 
 void FD3D12DynamicRHI::RHIUnlockStagingBuffer(FStagingBufferRHIParamRef StagingBufferRHI)
 {
 	FD3D12StagingBuffer* StagingBuffer = FD3D12DynamicRHI::ResourceCast(StagingBufferRHI);
 	check(StagingBuffer);
-
-	FD3D12Resource* pResource = StagingBuffer->StagedRead.GetReference();
-	if (pResource)
-	{
-		pResource->Unmap();
-	}
+	StagingBuffer->Unlock();
 }
 
 // =============================================================================
diff --git a/Engine/Source/Runtime/D3D12RHI/Private/D3D12RHIPrivate.h b/Engine/Source/Runtime/D3D12RHI/Private/D3D12RHIPrivate.h
--- a/Engine/Source/Runtime/D3D12RHI/Private/D3D12RHIPrivate.h
+++ b/Engine/Source/Runtime/D3D12RHI/Private/D3D12RHIPrivate.h
@@ -235,7 +235,7 @@ public:
 	virtual FComputeShaderRHIRef RHICreateComputeShader(const TArray<uint8>& Code) override;
 	virtual FComputeFenceRHIRef RHICreateComputeFence(const FName& Name) final override;
 	virtual FGPUFenceRHIRef RHICreateGPUFence(const FName& Name) final override;
-	virtual FStagingBufferRHIRef RHICreateStagingBuffer(FVertexBufferRHIParamRef VertexBuffer) final override;
+	virtual FStagingBufferRHIRef RHICreateStagingBuffer() final override;
 	virtual void* RHILockStagingBuffer(FStagingBufferRHIParamRef StagingBuffer, uint32 Offset, uint32 SizeRHI) final override;
     virtual void RHIUnlockStagingBuffer(FStagingBufferRHIParamRef StagingBuffer) final override;
 	virtual FBoundShaderStateRHIRef RHICreateBoundShaderState(FVertexDeclarationRHIParamRef VertexDeclaration, FVertexShaderRHIParamRef VertexShader, FHullShaderRHIParamRef HullShader, FDomainShaderRHIParamRef DomainShader, FPixelShaderRHIParamRef PixelShader, FGeometryShaderRHIParamRef GeometryShader) final override;
diff --git a/Engine/Source/Runtime/D3D12RHI/Private/D3D12UAV.cpp b/Engine/Source/Runtime/D3D12RHI/Private/D3D12UAV.cpp
--- a/Engine/Source/Runtime/D3D12RHI/Private/D3D12UAV.cpp
+++ b/Engine/Source/Runtime/D3D12RHI/Private/D3D12UAV.cpp
@@ -260,4 +260,41 @@ FUnorderedAccessViewRHIRef FD3D12DynamicRHI::RHICreateUnorderedAccessView_Render
 		return RHICreateUnorderedAccessView(VertexBufferRHI, Format);
 	}
 	return RHICreateUnorderedAccessView(VertexBufferRHI, Format);
+}
+
+FD3D12StagingBuffer::~FD3D12StagingBuffer()
+{
+	if (StagedRead)
+	{
+		StagedRead.SafeRelease();
+	}
+}
+
+void* FD3D12StagingBuffer::Lock(uint32 Offset, uint32 NumBytes)
+{
+	check(!bIsLocked);
+	bIsLocked = true;
+	FD3D12Resource* pResource = StagedRead.GetReference();
+	if (pResource)
+	{
+		D3D12_RANGE ReadRange;
+		ReadRange.Begin = Offset;
+		ReadRange.End = Offset + NumBytes;
+		return reinterpret_cast<uint8*>(pResource->Map(&ReadRange)) + Offset;
+	}
+	else
+	{
+		return nullptr;
+	}
+}
+
+void FD3D12StagingBuffer::Unlock()
+{
+	check(bIsLocked);
+	bIsLocked = false;
+	FD3D12Resource* pResource = StagedRead.GetReference();
+	if (pResource)
+	{
+		pResource->Unmap();
+	}
 }
\ No newline at end of file
diff --git a/Engine/Source/Runtime/D3D12RHI/Public/D3D12Resources.h b/Engine/Source/Runtime/D3D12RHI/Public/D3D12Resources.h
--- a/Engine/Source/Runtime/D3D12RHI/Public/D3D12Resources.h
+++ b/Engine/Source/Runtime/D3D12RHI/Public/D3D12Resources.h
@@ -880,24 +880,36 @@ private:
 
 class FD3D12StagingBuffer : public FRHIStagingBuffer
 {
+	friend class FD3D12CommandContext;
+	friend class FD3D12DynamicRHI;
+
 public:
-	FD3D12StagingBuffer(FVertexBufferRHIRef InBuffer)
-		: FRHIStagingBuffer(InBuffer)
+	FD3D12StagingBuffer()
+		: FRHIStagingBuffer()
+		, ShadowBufferSize(0)
 	{}
 
+	virtual ~FD3D12StagingBuffer() final override;
+
+	virtual void* Lock(uint32 Offset, uint32 NumBytes) final override;
+	virtual void Unlock() final override;
+
+private:
 	TRefCountPtr<FD3D12Resource> StagedRead;
+	uint32 ShadowBufferSize;
 };
 
 class FD3D12GPUFence : public FRHIGPUFence
 {
 public:
-	FD3D12GPUFence(FName InName, FD3D12Fence* InFence) 
+	FD3D12GPUFence(FName InName, FD3D12Fence* InFence)
 		: FRHIGPUFence(InName)
 		, Fence(InFence)
 		, Value(0)
 	{}
 
 	void WriteInternal(ED3D12CommandQueueType QueueType);
+	virtual void Clear() final override;
 	virtual bool Poll() const final override;
 
 protected:
diff --git a/Engine/Source/Runtime/Engine/Classes/Kismet/GameplayStatics.h b/Engine/Source/Runtime/Engine/Classes/Kismet/GameplayStatics.h
--- a/Engine/Source/Runtime/Engine/Classes/Kismet/GameplayStatics.h
+++ b/Engine/Source/Runtime/Engine/Classes/Kismet/GameplayStatics.h
@@ -705,6 +705,19 @@ public:
 	UFUNCTION(BlueprintCallable, Category = "Audio", meta = (WorldContext = "WorldContextObject"))
 	static class UReverbEffect* GetCurrentReverbEffect(const UObject* WorldContextObject);
 
+	/**
+	 * Sets the max number of voices (also known as "channels") dynamically by percentage. E.g. if you want to temporarily reduce voice count by 50%, use 0.50. Later, you can return to the original max voice count by using 1.0.
+	 */
+	UFUNCTION(BlueprintCallable, Category = "Audio", meta = (WorldContext = "WorldContextObject"))
+	static void SetMaxAudioChannelsScaled(const UObject* WorldContextObject, float MaxChannelCountScale);
+
+	/**
+	 * Retrieves the max voice count currently used by the audio engine.
+	 */
+	UFUNCTION(BlueprintCallable, Category = "Audio", meta = (WorldContext = "WorldContextObject"))
+	static int32 GetMaxAudioChannelCount(const UObject* WorldContextObject);
+
+
 	// --- Decal functions ------------------------------
 
 	/** Spawns a decal at the given location and rotation, fire and forget. Does not replicate.
diff --git a/Engine/Source/Runtime/Engine/Classes/Sound/SoundBase.h b/Engine/Source/Runtime/Engine/Classes/Sound/SoundBase.h
--- a/Engine/Source/Runtime/Engine/Classes/Sound/SoundBase.h
+++ b/Engine/Source/Runtime/Engine/Classes/Sound/SoundBase.h
@@ -221,5 +221,8 @@ public:
 	/** Returns whether the sound has cooked analysis data (e.g. FFT or envelope following data) and returns sound waves which have cooked data. */
 	virtual bool GetSoundWavesWithCookedAnalysisData(TArray<USoundWave*>& OutSoundWaves);
 
+	/** Queries if the sound has cooked FFT or envelope data. */
+	virtual bool HasCookedFFTData() const { return false; }
+	virtual bool HasCookedAmplitudeEnvelopeData() const { return false; }
 };
 
diff --git a/Engine/Source/Runtime/Engine/Classes/Sound/SoundCue.h b/Engine/Source/Runtime/Engine/Classes/Sound/SoundCue.h
--- a/Engine/Source/Runtime/Engine/Classes/Sound/SoundCue.h
+++ b/Engine/Source/Runtime/Engine/Classes/Sound/SoundCue.h
@@ -150,6 +150,8 @@ public:
 	virtual const FSoundAttenuationSettings* GetAttenuationSettingsToApply() const override;
 	virtual float GetSubtitlePriority() const override;
 	virtual bool GetSoundWavesWithCookedAnalysisData(TArray<USoundWave*>& OutSoundWaves) override;
+	virtual bool HasCookedFFTData() const override;
+	virtual bool HasCookedAmplitudeEnvelopeData() const override;
 	//~ End USoundBase Interface.
 
 	/** Construct and initialize a node within this Cue */
@@ -180,7 +182,8 @@ public:
 	/**
 	 * Recursively finds sound nodes of type T
 	 */
-	template<typename T> void RecursiveFindNode(USoundNode* Node, TArray<T*>& OutNodes)
+	template<typename T> 
+	void RecursiveFindNode(USoundNode* Node, TArray<T*>& OutNodes)
 	{
 		if (Node)
 		{
@@ -199,6 +202,27 @@ public:
 		}
 	}
 
+	template<typename T>
+	void RecursiveFindNode(const USoundNode* Node, TArray<const T*>& OutNodes) const
+	{
+		if (Node)
+		{
+			// Record the node if it is the desired type
+			if (const T* FoundNode = Cast<T>(Node))
+			{
+				OutNodes.AddUnique(FoundNode);
+			}
+
+			// Recurse.
+			const int32 MaxChildNodes = Node->GetMaxChildNodes();
+			for (int32 ChildIndex = 0; ChildIndex < Node->ChildNodes.Num() && ChildIndex < MaxChildNodes; ++ChildIndex)
+			{
+				RecursiveFindNode<T>(Node->ChildNodes[ChildIndex], OutNodes);
+			}
+		}
+	}
+
+
 	/** Find the path through the sound cue to a node identified by its hash */
 	ENGINE_API bool FindPathToNode(const UPTRINT NodeHashToFind, TArray<USoundNode*>& OutPath) const;
 
diff --git a/Engine/Source/Runtime/Engine/Classes/Sound/SoundWave.h b/Engine/Source/Runtime/Engine/Classes/Sound/SoundWave.h
--- a/Engine/Source/Runtime/Engine/Classes/Sound/SoundWave.h
+++ b/Engine/Source/Runtime/Engine/Classes/Sound/SoundWave.h
@@ -540,6 +540,8 @@ public:
 	virtual float GetSubtitlePriority() const override;
 	virtual bool IsAllowedVirtual() const override;
 	virtual bool GetSoundWavesWithCookedAnalysisData(TArray<USoundWave*>& OutSoundWaves) override;
+	virtual bool HasCookedFFTData() const override;
+	virtual bool HasCookedAmplitudeEnvelopeData() const override;
 	//~ End USoundBase Interface.
 
 	// Called  when the procedural sound wave begins on the render thread. Only used in the audio mixer and when bProcedural is true.
diff --git a/Engine/Source/Runtime/Engine/Private/AudioDevice.cpp b/Engine/Source/Runtime/Engine/Private/AudioDevice.cpp
--- a/Engine/Source/Runtime/Engine/Private/AudioDevice.cpp
+++ b/Engine/Source/Runtime/Engine/Private/AudioDevice.cpp
@@ -49,6 +49,13 @@ FAutoConsoleVariableRef CVarSetAudioChannelCount(
 	TEXT("0: Disable, >0: Enable"),
 	ECVF_Default);
 
+static float AudioChannelCountScaleCVar = 1.0f;
+FAutoConsoleVariableRef CVarSetAudioChannelScaleCount(
+	TEXT("au.SetAudioChannelScaleCount"),
+	AudioChannelCountScaleCVar,
+	TEXT("Changes the audio channel count by percentage.\n"),
+	ECVF_Default);
+
 static int32 DisableStoppingVoicesCvar = 0;
 FAutoConsoleVariableRef CVarDisableStoppingVoices(
 	TEXT("au.DisableStoppingVoices"),
@@ -159,6 +166,9 @@ void FDynamicParameter::Update(float DeltaTime)
 
 FAudioDevice::FAudioDevice()
 	: MaxChannels(0)
+	, MaxChannels_GameThread(0)
+	, MaxChannelsScale(1.0f)
+	, MaxChannelsScale_GameThread(1.0f)
 	, NumStoppingVoices(32)
 	, MaxWaveInstances(0)
 	, SampleRate(0)
@@ -247,6 +257,7 @@ bool FAudioDevice::Init(int32 InMaxChannels)
 
 	// MaxChannels is the min of the platform-specific value and the max value in the quality settings (InMaxChannels)
 	MaxChannels = PlatformSettings.MaxChannels > 0 ? FMath::Min(PlatformSettings.MaxChannels, InMaxChannels) : InMaxChannels;
+	MaxChannels_GameThread = MaxChannels;
 
 	// Mixed sample rate is set by the platform
 	SampleRate = PlatformSettings.SampleRate;
@@ -428,6 +439,21 @@ void FAudioDevice::PrecacheStartupSounds()
 
 void FAudioDevice::SetMaxChannels(int32 InMaxChannels)
 {
+	if (!IsInAudioThread())
+	{
+		MaxChannels_GameThread = InMaxChannels;
+
+		DECLARE_CYCLE_STAT(TEXT("FAudioThreadTask.SetMaxChannels"), STAT_AudioSetMaxChannels, STATGROUP_AudioThreadCommands);
+
+		FAudioThread::RunCommandOnAudioThread([this, InMaxChannels]()
+		{
+			this->SetMaxChannels(InMaxChannels);
+
+		}, GET_STATID(STAT_AudioSetMaxChannels));
+
+		return;
+	}
+
 	if (InMaxChannels > Sources.Num())
 	{
 		UE_LOG(LogAudio, Warning, TEXT("Can't increase channels past starting number!"));
@@ -437,14 +463,56 @@ void FAudioDevice::SetMaxChannels(int32 InMaxChannels)
 	MaxChannels = InMaxChannels;
 }
 
+void FAudioDevice::SetMaxChannelsScaled(float InScaledChannelCount)
+{
+	if (!IsInAudioThread())
+	{
+		MaxChannelsScale_GameThread = InScaledChannelCount;
+
+		DECLARE_CYCLE_STAT(TEXT("FAudioThreadTask.SetMaxChannelsScaled"), STAT_AudioSetMaxChannelsScaled, STATGROUP_AudioThreadCommands);
+
+		FAudioThread::RunCommandOnAudioThread([this, InScaledChannelCount]()
+		{
+			MaxChannelsScale = FMath::Clamp(InScaledChannelCount, 0.0f, 1.0f);
+
+		}, GET_STATID(STAT_AudioSetMaxChannelsScaled));
+
+		return;
+	}
+	else
+	{
+		MaxChannelsScale = FMath::Clamp(InScaledChannelCount, 0.0f, 1.0f);
+
+		DECLARE_CYCLE_STAT(TEXT("FAudioThreadTask.SetMaxChannelsScaled"), STAT_AudioSetMaxChannelsScaled, STATGROUP_AudioThreadCommands);
+
+		FAudioThread::RunCommandOnGameThread([this, InScaledChannelCount]()
+		{
+			MaxChannelsScale_GameThread = InScaledChannelCount;
+
+		}, GET_STATID(STAT_AudioSetMaxChannelsScaled));
+	}
+}
+
 int32 FAudioDevice::GetMaxChannels() const
 {
-	if (AudioChannelCountCVar > 0 && AudioChannelCountCVar < Sources.Num())
+	if (IsInAudioThread())
 	{
-		return AudioChannelCountCVar;
+		if (AudioChannelCountCVar > 0 && AudioChannelCountCVar < Sources.Num())
+		{
+			return FMath::Max(int32(AudioChannelCountCVar * MaxChannelsScale * AudioChannelCountScaleCVar), 1);
+		}
+
+		return FMath::Max(int32(MaxChannels * MaxChannelsScale * AudioChannelCountScaleCVar), 1);
 	}
+	else
+	{
+		if (AudioChannelCountCVar > 0 && AudioChannelCountCVar < Sources.Num())
+		{
+			return FMath::Max(int32(AudioChannelCountCVar * MaxChannelsScale_GameThread * AudioChannelCountScaleCVar), 1);
+		}
 
-	return MaxChannels;
+		return FMath::Max(int32(MaxChannels_GameThread * MaxChannelsScale_GameThread * AudioChannelCountScaleCVar), 1);
+	}
 }
 
 void FAudioDevice::Teardown()
diff --git a/Engine/Source/Runtime/Engine/Private/Components/AudioComponent.cpp b/Engine/Source/Runtime/Engine/Private/Components/AudioComponent.cpp
--- a/Engine/Source/Runtime/Engine/Private/Components/AudioComponent.cpp
+++ b/Engine/Source/Runtime/Engine/Private/Components/AudioComponent.cpp
@@ -1133,11 +1133,7 @@ bool UAudioComponent::HasCookedFFTData() const
 {
 	if (Sound)
 	{
-		USoundWave* SoundWave = Cast<USoundWave>(Sound);
-		if (SoundWave)
-		{
-			return SoundWave->CookedSpectralTimeData.Num() > 0;
-		}
+		return Sound->HasCookedFFTData();
 	}
 	return false;
 }
@@ -1146,11 +1142,7 @@ bool UAudioComponent::HasCookedAmplitudeEnvelopeData() const
 {
 	if (Sound)
 	{
-		USoundWave* SoundWave = Cast<USoundWave>(Sound);
-		if (SoundWave)
-		{
-			return SoundWave->CookedEnvelopeTimeData.Num() > 0;
-		}
+		return Sound->HasCookedAmplitudeEnvelopeData();
 	}
 	return false;
 }
diff --git a/Engine/Source/Runtime/Engine/Private/Curves/RichCurve.cpp b/Engine/Source/Runtime/Engine/Private/Curves/RichCurve.cpp
--- a/Engine/Source/Runtime/Engine/Private/Curves/RichCurve.cpp
+++ b/Engine/Source/Runtime/Engine/Private/Curves/RichCurve.cpp
@@ -1378,10 +1378,7 @@ void FRichCurve::CompressCurve(FCompressedRichCurve& OutCurve, float ErrorThresh
 
 			const int32 KeySize = sizeof(uint16);
 
-			if (KeyFormat == RCKTCF_uint16)
-			{
-				WriteOffset = Align(WriteOffset, sizeof(uint16));
-			}
+			WriteOffset = Align(WriteOffset, sizeof(uint16));
 
 			uint8* KeyTimes8 = BasePtr + WriteOffset;
 			uint16* KeyTimes16 = reinterpret_cast<uint16*>(KeyTimes8);
diff --git a/Engine/Source/Runtime/Engine/Private/GameplayStatics.cpp b/Engine/Source/Runtime/Engine/Private/GameplayStatics.cpp
--- a/Engine/Source/Runtime/Engine/Private/GameplayStatics.cpp
+++ b/Engine/Source/Runtime/Engine/Private/GameplayStatics.cpp
@@ -1600,6 +1600,48 @@ class UReverbEffect* UGameplayStatics::GetCurrentReverbEffect(const UObject* Wor
 	return nullptr;
 }
 
+void UGameplayStatics::SetMaxAudioChannelsScaled(const UObject* WorldContextObject, float MaxChannelCountScale)
+{
+	if (GEngine == nullptr || !GEngine->UseSound())
+	{
+		return;
+	}
+
+	UWorld* ThisWorld = GEngine->GetWorldFromContextObject(WorldContextObject, EGetWorldErrorMode::LogAndReturnNull);
+	if (!ThisWorld || !ThisWorld->bAllowAudioPlayback)
+	{
+		return;
+	}
+
+	if (FAudioDevice* AudioDevice = ThisWorld->GetAudioDevice())
+	{
+		AudioDevice->SetMaxChannelsScaled(MaxChannelCountScale);
+	}
+}
+
+int32 UGameplayStatics::GetMaxAudioChannelCount(const UObject* WorldContextObject)
+{
+	if (GEngine == nullptr || !GEngine->UseSound())
+	{
+		return 0;
+	}
+
+
+	UWorld* ThisWorld = GEngine->GetWorldFromContextObject(WorldContextObject, EGetWorldErrorMode::LogAndReturnNull);
+	if (!ThisWorld || !ThisWorld->bAllowAudioPlayback)
+	{
+		return 0;
+	}
+
+	if (FAudioDevice* AudioDevice = ThisWorld->GetAudioDevice())
+	{
+		return AudioDevice->GetMaxChannels();
+	}
+
+	return 0;
+}
+
+
 UDecalComponent* CreateDecalComponent(class UMaterialInterface* DecalMaterial, FVector DecalSize, UWorld* World, AActor* Actor, float LifeSpan)
 {
 	UDecalComponent* DecalComp = NewObject<UDecalComponent>((Actor ? Actor : (UObject*)World));
diff --git a/Engine/Source/Runtime/Engine/Private/Materials/MaterialShader.cpp b/Engine/Source/Runtime/Engine/Private/Materials/MaterialShader.cpp
--- a/Engine/Source/Runtime/Engine/Private/Materials/MaterialShader.cpp
+++ b/Engine/Source/Runtime/Engine/Private/Materials/MaterialShader.cpp
@@ -35,6 +35,7 @@ namespace MaterialShaderCookStats
 //
 // Globals
 //
+FCriticalSection FMaterialShaderMap::GIdToMaterialShaderMapCS;
 TMap<FMaterialShaderMapId,FMaterialShaderMap*> FMaterialShaderMap::GIdToMaterialShaderMap[SP_NumPlatforms];
 #if ALLOW_SHADERMAP_DEBUG_DATA
 TArray<FMaterialShaderMap*> FMaterialShaderMap::AllMaterialShaderMaps;
@@ -963,10 +964,13 @@ FShader* FMaterialShaderType::FinishCompileShader(
 * @param Platform - The platform to lookup for
 * @return NULL if no cached shader map was found.
 */
-FMaterialShaderMap* FMaterialShaderMap::FindId(const FMaterialShaderMapId& ShaderMapId, EShaderPlatform InPlatform)
+TRefCountPtr<FMaterialShaderMap> FMaterialShaderMap::FindId(const FMaterialShaderMapId& ShaderMapId, EShaderPlatform InPlatform)
 {
+	FScopeLock ScopeLock(&GIdToMaterialShaderMapCS);
 	check(ShaderMapId.BaseMaterialId != FGuid());
-	return GIdToMaterialShaderMap[InPlatform].FindRef(ShaderMapId);
+	TRefCountPtr<FMaterialShaderMap> Result = GIdToMaterialShaderMap[InPlatform].FindRef(ShaderMapId);
+	check(Result == nullptr || !Result->bDeletedThroughDeferredCleanup);
+	return Result;
 }
 
 #if ALLOW_SHADERMAP_DEBUG_DATA
@@ -995,6 +999,8 @@ void FMaterialShaderMap::FlushShaderTypes(TArray<FShaderType*>& ShaderTypesToFlu
 
 void FMaterialShaderMap::FixupShaderTypes(EShaderPlatform Platform, const TMap<FShaderType*, FString>& ShaderTypeNames, const TMap<const FShaderPipelineType*, FString>& ShaderPipelineTypeNames, const TMap<FVertexFactoryType*, FString>& VertexFactoryTypeNames)
 {
+	FScopeLock ScopeLock(&GIdToMaterialShaderMapCS);
+
 	TArray<FMaterialShaderMapId> Keys;
 	FMaterialShaderMap::GIdToMaterialShaderMap[Platform].GenerateKeyArray(Keys);
 
@@ -2134,14 +2140,18 @@ void FMaterialShaderMap::Register(EShaderPlatform InShaderPlatform)
 		INC_DWORD_STAT_BY(STAT_Shaders_ShaderMapMemory, GetSizeBytes());
 	}
 
-	GIdToMaterialShaderMap[GetShaderPlatform()].Add(ShaderMapId,this);
-	bRegistered = true;
+	{
+		FScopeLock ScopeLock(&GIdToMaterialShaderMapCS);
+		GIdToMaterialShaderMap[GetShaderPlatform()].Add(ShaderMapId,this);
+		bRegistered = true;
+	}
 }
 
 void FMaterialShaderMap::AddRef()
 {
 	//#todo-mw: re-enable to try to find potential corruption of the global shader map ID array
 	//check(IsInGameThread());
+	FScopeLock ScopeLock(&GIdToMaterialShaderMapCS);
 	check(!bDeletedThroughDeferredCleanup);
 	++NumRefs;
 }
@@ -2150,20 +2160,27 @@ void FMaterialShaderMap::Release()
 {
 	//#todo-mw: re-enable to try to find potential corruption of the global shader map ID array
 	//check(IsInGameThread());
-	check(NumRefs > 0);
-	if(--NumRefs == 0)
+
 	{
-		if (bRegistered)
+		FScopeLock ScopeLock(&GIdToMaterialShaderMapCS);
+
+		check(NumRefs > 0);
+		if (--NumRefs == 0)
 		{
-			DEC_DWORD_STAT(STAT_Shaders_NumShaderMaps);
-			DEC_DWORD_STAT_BY(STAT_Shaders_ShaderMapMemory, GetSizeBytes());
+			if (bRegistered)
+			{
+				bRegistered = false;
+				DEC_DWORD_STAT(STAT_Shaders_NumShaderMaps);
+				DEC_DWORD_STAT_BY(STAT_Shaders_ShaderMapMemory, GetSizeBytes());
+				GIdToMaterialShaderMap[GetShaderPlatform()].Remove(ShaderMapId);
+			}
 
-			GIdToMaterialShaderMap[GetShaderPlatform()].Remove(ShaderMapId);
-			bRegistered = false;
+			check(!bDeletedThroughDeferredCleanup);
+			bDeletedThroughDeferredCleanup = true;
 		}
-
-		check(!bDeletedThroughDeferredCleanup);
-		bDeletedThroughDeferredCleanup = true;
+	}
+	if (bDeletedThroughDeferredCleanup)
+	{
 		BeginCleanup(this);
 	}
 }
diff --git a/Engine/Source/Runtime/Engine/Private/Materials/MaterialShared.cpp b/Engine/Source/Runtime/Engine/Private/Materials/MaterialShared.cpp
--- a/Engine/Source/Runtime/Engine/Private/Materials/MaterialShared.cpp
+++ b/Engine/Source/Runtime/Engine/Private/Materials/MaterialShared.cpp
@@ -1610,7 +1610,7 @@ bool FMaterial::CacheShaders(const FMaterialShaderMapId& ShaderMapId, EShaderPla
 	// If we loaded this material with inline shaders, use what was loaded (GameThreadShaderMap) instead of looking in the DDC
 	if (bContainsInlineShaders)
 	{
-		FMaterialShaderMap* ExistingShaderMap = nullptr;
+		TRefCountPtr<FMaterialShaderMap> ExistingShaderMap = nullptr;
 		
 		if (GameThreadShaderMap)
 		{
diff --git a/Engine/Source/Runtime/Engine/Private/SoundCue.cpp b/Engine/Source/Runtime/Engine/Private/SoundCue.cpp
--- a/Engine/Source/Runtime/Engine/Private/SoundCue.cpp
+++ b/Engine/Source/Runtime/Engine/Private/SoundCue.cpp
@@ -505,6 +505,40 @@ bool USoundCue::GetSoundWavesWithCookedAnalysisData(TArray<USoundWave*>& OutSoun
 	return bHasAnalysisData;
 }
 
+bool USoundCue::HasCookedFFTData() const
+{
+	// Check this sound cue's wave players to see if any of their soundwaves have cooked analysis data
+	TArray<const USoundNodeWavePlayer*> WavePlayers;
+	RecursiveFindNode<USoundNodeWavePlayer>(FirstNode, WavePlayers);
+
+	for (const USoundNodeWavePlayer* Player : WavePlayers)
+	{
+		const USoundWave* SoundWave = Player->GetSoundWave();
+		if (SoundWave && SoundWave->HasCookedFFTData())
+		{
+			return true;
+		}
+	}
+	return false;
+}
+
+bool USoundCue::HasCookedAmplitudeEnvelopeData() const
+{
+	// Check this sound cue's wave players to see if any of their soundwaves have cooked analysis data
+	TArray<const USoundNodeWavePlayer*> WavePlayers;
+	RecursiveFindNode<USoundNodeWavePlayer>(FirstNode, WavePlayers);
+
+	for (const USoundNodeWavePlayer* Player : WavePlayers)
+	{
+		const USoundWave* SoundWave = Player->GetSoundWave();
+		if (SoundWave && SoundWave->HasCookedAmplitudeEnvelopeData())
+		{
+			return true;
+		}
+	}
+	return false;
+}
+
 #if WITH_EDITOR
 UEdGraph* USoundCue::GetGraph()
 { 
diff --git a/Engine/Source/Runtime/Engine/Private/SoundWave.cpp b/Engine/Source/Runtime/Engine/Private/SoundWave.cpp
--- a/Engine/Source/Runtime/Engine/Private/SoundWave.cpp
+++ b/Engine/Source/Runtime/Engine/Private/SoundWave.cpp
@@ -1487,6 +1487,16 @@ bool USoundWave::GetSoundWavesWithCookedAnalysisData(TArray<USoundWave*>& OutSou
 	return false;
 }
 
+bool USoundWave::HasCookedFFTData() const
+{
+	return CookedSpectralTimeData.Num() > 0;
+}
+
+bool USoundWave::HasCookedAmplitudeEnvelopeData() const
+{
+	return CookedEnvelopeTimeData.Num() > 0;
+}
+
 void USoundWave::UpdatePlatformData()
 {
 	if (IsStreaming())
diff --git a/Engine/Source/Runtime/Engine/Public/AudioDevice.h b/Engine/Source/Runtime/Engine/Public/AudioDevice.h
--- a/Engine/Source/Runtime/Engine/Public/AudioDevice.h
+++ b/Engine/Source/Runtime/Engine/Public/AudioDevice.h
@@ -624,7 +624,12 @@ public:
 	/** 
 	 * Sets the maximum number of channels dynamically. Can't raise the cap over the initial value but can lower it 
 	 */
-	virtual void SetMaxChannels(int32 InMaxChannels);
+	void SetMaxChannels(int32 InMaxChannels);
+
+	/**
+	 * Sets the maximum number of channels dynamically by scaled percentage.
+	 */
+	void SetMaxChannelsScaled(float InScaledChannelCount);
 
 	/** Returns the max channels used by the audio device. */
 	int32 GetMaxChannels() const;
@@ -1550,6 +1555,11 @@ public:
 
 	/** The maximum number of concurrent audible sounds */
 	int32 MaxChannels;
+	int32 MaxChannels_GameThread;
+
+	/** A scaler on the max channels. */
+	float MaxChannelsScale;
+	float MaxChannelsScale_GameThread;
 
 	/** The number of sources to reserve for stopping sounds. */
 	int32 NumStoppingVoices;
diff --git a/Engine/Source/Runtime/Engine/Public/MaterialShared.h b/Engine/Source/Runtime/Engine/Public/MaterialShared.h
--- a/Engine/Source/Runtime/Engine/Public/MaterialShared.h
+++ b/Engine/Source/Runtime/Engine/Public/MaterialShared.h
@@ -705,7 +705,7 @@ public:
 	 * @param Platform - The platform to lookup for
 	 * @return NULL if no cached shader map was found.
 	 */
-	static FMaterialShaderMap* FindId(const FMaterialShaderMapId& ShaderMapId, EShaderPlatform Platform);
+	static TRefCountPtr<FMaterialShaderMap> FindId(const FMaterialShaderMapId& ShaderMapId, EShaderPlatform Platform);
 
 #if ALLOW_SHADERMAP_DEBUG_DATA
 	/** Flushes the given shader types from any loaded FMaterialShaderMap's. */
@@ -878,7 +878,7 @@ public:
 
 	bool IsValidForRendering(bool bFailOnInvalid = false) const
 	{
-		const bool bValid = bCompilationFinalized && bCompiledSuccessfully && !bDeletedThroughDeferredCleanup;
+		const bool bValid = bCompilationFinalized && bCompiledSuccessfully;// && !bDeletedThroughDeferredCleanup; //deferred actually deletion will prevent the material to go away before we finish rendering
 		checkf(bValid || !bFailOnInvalid, TEXT("FMaterialShaderMap %s invalid for rendering: bCompilationFinalized: %i, bCompiledSuccessfully: %i, bDeletedThroughDeferredCleanup: %i"), *GetFriendlyName(), bCompilationFinalized, bCompiledSuccessfully, bDeletedThroughDeferredCleanup ? 1 : 0);
 		return bValid;
 	}
@@ -912,6 +912,7 @@ private:
 	 * No ref counting needed as these are removed on destruction of the shader map.
 	 */
 	static TMap<FMaterialShaderMapId,FMaterialShaderMap*> GIdToMaterialShaderMap[SP_NumPlatforms];
+	static FCriticalSection GIdToMaterialShaderMapCS;
 
 #if ALLOW_SHADERMAP_DEBUG_DATA
 	/** 
diff --git a/Engine/Source/Runtime/Engine/Public/SkeletalMeshReductionSettings.h b/Engine/Source/Runtime/Engine/Public/SkeletalMeshReductionSettings.h
--- a/Engine/Source/Runtime/Engine/Public/SkeletalMeshReductionSettings.h
+++ b/Engine/Source/Runtime/Engine/Public/SkeletalMeshReductionSettings.h
@@ -18,7 +18,10 @@ enum SkeletalMeshTerminationCriterion
 {
 	SMTC_NumOfTriangles UMETA(DisplayName = "Triangles", ToolTip = "Triangle count criterion will be used for simplification."),
 	SMTC_NumOfVerts UMETA(DisplayName = "Vertices", ToolTip = "Vertex cont criterion will be used for simplification."),
-	SMTC_TriangleOrVert UMETA(DisplayName = "First Satisfied", ToolTip = "Simplification will continue until either Triangle or Vertex count criteria is met."),
+	SMTC_TriangleOrVert UMETA(DisplayName = "First Percent Satisfied", ToolTip = "Simplification will continue until either Triangle or Vertex count criteria is met."),
+	SMTC_AbsNumOfTriangles UMETA(DisplayName = "Max Triangles", ToolTip = "Triangle count criterion will be used for simplification."),
+	SMTC_AbsNumOfVerts UMETA(DisplayName = "Max Vertices", ToolTip = "Vertex cont criterion will be used for simplification."),
+	SMTC_AbsTriangleOrVert UMETA(DisplayName = "First Max Satisfied", ToolTip = "Simplification will continue until either Triangle or Vertex count criteria is met."),
 	SMTC_MAX UMETA(Hidden),
 };
 
@@ -64,6 +67,14 @@ struct FSkeletalMeshOptimizationSettings
 	/** The percentage of vertices to retain as a ratio, e.g. 0.1 indicates 10 percent */
 	UPROPERTY(EditAnywhere, Category = ReductionMethod, meta = (DisplayName = "Percent of Vertices"))
 	float NumOfVertPercentage;
+	
+	/** The maximum number of triangles to retain */
+	UPROPERTY(EditAnywhere, Category = ReductionMethod, meta = (DisplayName = "Max Triangle Count", ClampMin = 4))
+	uint32 MaxNumOfTriangles;
+
+	/** The maximum number of vertices to retain */
+	UPROPERTY(EditAnywhere, Category = ReductionMethod, meta = (DisplayName = "Max Vertex Count", ClampMin = 6))
+	uint32 MaxNumOfVerts;
 
 	/**If ReductionMethod equals MaxDeviation this value is the maximum deviation from the base mesh as a percentage of the bounding sphere. 
 	 * In code, it ranges from [0, 1]. In the editor UI, it ranges from [0, 100]
@@ -140,6 +151,8 @@ struct FSkeletalMeshOptimizationSettings
 		: TerminationCriterion(SMTC_NumOfTriangles)
 		, NumOfTrianglesPercentage(0.5f)
 		, NumOfVertPercentage(0.5f)
+		, MaxNumOfTriangles(4)
+		, MaxNumOfVerts(6)
 		, MaxDeviationPercentage(0.5f)
 		, ReductionMethod(SMOT_NumOfTriangles)
 		, SilhouetteImportance(SMOI_Normal)
diff --git a/Engine/Source/Runtime/IOS/IOSAudio/Private/IOSAudioSource.cpp b/Engine/Source/Runtime/IOS/IOSAudio/Private/IOSAudioSource.cpp
--- a/Engine/Source/Runtime/IOS/IOSAudio/Private/IOSAudioSource.cpp
+++ b/Engine/Source/Runtime/IOS/IOSAudio/Private/IOSAudioSource.cpp
@@ -226,7 +226,7 @@ void FIOSAudioSoundSource::Update(void)
 
     SetFilterFrequency();
     
-    SourceLPFFrequency = LPFFrequency;
+    SourceLPFFrequency = LPFFrequency / (((float) SampleRate) * 0.5f);
     
 	// Factor in the xaudio2 attenuation that happens to stereo assets.
 	if (WaveInstance->WaveData->NumChannels == 2 && WaveInstance->bUseSpatialization)
@@ -292,10 +292,9 @@ void FIOSAudioSoundSource::Play(void)
 		// Updates the source which sets the pitch and volume
 		Update();
         
-        for (Audio::FBiquadFilter& Filter : LowpassFilterBank)
+        for (Audio::FOnePoleLPF& Filter : LowpassFilterBank)
         {
             Filter.Reset();
-            Filter.Init(SampleRate, 1, Audio::EBiquadFilter::Lowpass, LPFFrequency);
         }
         
         for (Audio::FParam& Param : LPFParamBank)
@@ -509,12 +508,11 @@ OSStatus FIOSAudioSoundSource::IOSAudioRenderCallback(void* RefCon, AudioUnitRen
         
         ChannelIndex = Source->LowpassFilterBank.AddDefaulted(1);
         Source->LowpassFilterBank[ChannelIndex].Reset();
-        Source->LowpassFilterBank[ChannelIndex].Init(Source->SampleRate, 1, Audio::EBiquadFilter::Lowpass, Source->SourceLPFFrequency);
     }
     
     // Set up LPF filter for this channel:
     Audio::FParam& LPFParam = Source->LPFParamBank[Channel];
-    Audio::FBiquadFilter& LowpassFilter = Source->LowpassFilterBank[Channel];
+    Audio::FOnePoleLPF& LowpassFilter = Source->LowpassFilterBank[Channel];
     
     LPFParam.Reset();
     const bool bShouldUpdateCutoffFrequency = !FMath::IsNearlyEqual(Source->SourceLPFFrequency, LPFParam.GetValue());
@@ -541,7 +539,7 @@ OSStatus FIOSAudioSoundSource::IOSAudioRenderCallback(void* RefCon, AudioUnitRen
         float FloatSample = ((float) IntSample) / 32767.0f;
         if(LPFParam.GetValue() != MAX_FILTER_FREQUENCY)
         {
-            LowpassFilter.ProcessAudioFrame(&FloatSample, &FloatSample);
+            FloatSample = LowpassFilter.ProcessAudioSample(FloatSample);
         }
         *OutData++ = (SInt16) (FloatSample * 32767.0f);
 	}
diff --git a/Engine/Source/Runtime/IOS/IOSAudio/Public/IOSAudioDevice.h b/Engine/Source/Runtime/IOS/IOSAudio/Public/IOSAudioDevice.h
--- a/Engine/Source/Runtime/IOS/IOSAudio/Public/IOSAudioDevice.h
+++ b/Engine/Source/Runtime/IOS/IOSAudio/Public/IOSAudioDevice.h
@@ -11,7 +11,7 @@
 #include "AudioDevice.h"
 #include "Sound/SoundWave.h"
 #include "DSP/ParamInterpolator.h"
-#include "DSP/Filter.h"
+#include "DSP/OnePole.h"
 
 /*------------------------------------------------------------------------------------
 	Audio Framework system headers
@@ -150,7 +150,7 @@ public:
 	/** Calculates the audio unit element of the input channel relative to the base bus number */
 	AudioUnitElement GetAudioUnitElement(int32 Channel);
 
-    TArray<Audio::FBiquadFilter> LowpassFilterBank;
+    TArray<Audio::FOnePoleLPF> LowpassFilterBank;
     TArray<Audio::FParam> LPFParamBank;
     
     int32 SampleRate;
diff --git a/Engine/Source/Runtime/RHI/Private/RHICommandList.cpp b/Engine/Source/Runtime/RHI/Private/RHICommandList.cpp
--- a/Engine/Source/Runtime/RHI/Private/RHICommandList.cpp
+++ b/Engine/Source/Runtime/RHI/Private/RHICommandList.cpp
@@ -2069,13 +2069,17 @@ void FDynamicRHI::UnlockIndexBuffer_RenderThread(class FRHICommandListImmediate&
 	}
 }
 
+// @todo-mattc-staging Default implementation
 void* FDynamicRHI::RHILockStagingBuffer(FStagingBufferRHIParamRef StagingBuffer, uint32 Offset, uint32 SizeRHI)
 {
-	return GDynamicRHI->RHILockVertexBuffer(StagingBuffer->GetBackingBuffer(), Offset, SizeRHI, RLM_ReadOnly);
+	check(false);
+	return nullptr;
+	//return GDynamicRHI->RHILockVertexBuffer(StagingBuffer->GetSourceBuffer(), Offset, SizeRHI, RLM_ReadOnly);
 }
 void FDynamicRHI::RHIUnlockStagingBuffer(FStagingBufferRHIParamRef StagingBuffer)
 {
-	GDynamicRHI->RHIUnlockVertexBuffer(StagingBuffer->GetBackingBuffer());
+	check(false);
+	//GDynamicRHI->RHIUnlockVertexBuffer(StagingBuffer->GetSourceBuffer());
 }
 
 void* FDynamicRHI::LockStagingBuffer_RenderThread(class FRHICommandListImmediate& RHICmdList, FStagingBufferRHIParamRef StagingBuffer, uint32 Offset, uint32 SizeRHI)
@@ -2083,7 +2087,6 @@ void* FDynamicRHI::LockStagingBuffer_RenderThread(class FRHICommandListImmediate
 	QUICK_SCOPE_CYCLE_COUNTER(STAT_FDynamicRHI_LockStagingBuffer_RenderThread);
 	check(IsInRenderingThread());
 
-	FScopedRHIThreadStaller StallRHIThread(RHICmdList);
 	return GDynamicRHI->RHILockStagingBuffer(StagingBuffer, Offset, SizeRHI);
 }
 void FDynamicRHI::UnlockStagingBuffer_RenderThread(class FRHICommandListImmediate& RHICmdList, FStagingBufferRHIParamRef StagingBuffer)
@@ -2091,7 +2094,6 @@ void FDynamicRHI::UnlockStagingBuffer_RenderThread(class FRHICommandListImmediat
 	QUICK_SCOPE_CYCLE_COUNTER(STAT_FDynamicRHI_UnlockStagingBuffer_RenderThread);
 	check(IsInRenderingThread());
 	
-	FScopedRHIThreadStaller StallRHIThread(RHICmdList);
 	GDynamicRHI->RHIUnlockStagingBuffer(StagingBuffer);
 }
 
diff --git a/Engine/Source/Runtime/RHI/Private/RHIGPUReadback.cpp b/Engine/Source/Runtime/RHI/Private/RHIGPUReadback.cpp
--- a/Engine/Source/Runtime/RHI/Private/RHIGPUReadback.cpp
+++ b/Engine/Source/Runtime/RHI/Private/RHIGPUReadback.cpp
@@ -7,20 +7,19 @@
 
 #include "RHIGPUReadback.h"
 
+FGenericRHIGPUFence::FGenericRHIGPUFence(FName InName)
+	: FRHIGPUFence(InName)
+	, InsertedFrameNumber(MAX_uint32)
+{}
 
-void FRHIGPUFence::Write()
+void FGenericRHIGPUFence::Clear()
 {
-	UE_LOG(LogRHI, Warning, TEXT("FRHIGPUFence::Write is not implemented"));
+	InsertedFrameNumber = MAX_uint32;
 }
 
-FGenericRHIGPUFence::FGenericRHIGPUFence(FName InName)
-    : FRHIGPUFence(InName)
-    , InsertedFrameNumber(0)
-{}
-
-void FGenericRHIGPUFence::Write()
+void FGenericRHIGPUFence::WriteInternal()
 {
-	// GPU generally overlap the game. This overlap increases when using AFR.
+	// GPU generally overlap the game. This overlap increases when using AFR. In normal mode this can make us appear to be further behind the gpu than we actually are.
 	InsertedFrameNumber = GFrameNumberRenderThread + GNumAlternateFrameRenderingGroups;
 }
 
@@ -32,3 +31,18 @@ bool FGenericRHIGPUFence::Poll() const
 	}
 	return false;
 }
+
+void* FGenericRHIStagingBuffer::Lock(uint32 InOffset, uint32 NumBytes)
+{
+	check(ShadowBuffer);
+	check(!bIsLocked);
+	bIsLocked = true;
+	return reinterpret_cast<void*>(reinterpret_cast<uint8*>(RHILockVertexBuffer(ShadowBuffer, InOffset, NumBytes, RLM_ReadOnly)) + Offset);
+}
+
+void FGenericRHIStagingBuffer::Unlock()
+{
+	check(bIsLocked);
+	RHIUnlockVertexBuffer(ShadowBuffer);
+	bIsLocked = false;
+}
\ No newline at end of file
diff --git a/Engine/Source/Runtime/RHI/Public/DynamicRHI.h b/Engine/Source/Runtime/RHI/Public/DynamicRHI.h
--- a/Engine/Source/Runtime/RHI/Public/DynamicRHI.h
+++ b/Engine/Source/Runtime/RHI/Public/DynamicRHI.h
@@ -216,19 +216,18 @@ public:
 	}
 
 	/**
-	* Creates a staging buffer that uses 'BackingBuffer' as the backing store whose contents may be accessed using Lock/Unlock.
-	* @param BackingBuffer - The buffer whose contents may be accessed using Lock/Unlock.
+	* Creates a staging buffer, which is memory visible to the cpu without any locking.
 	* @return The new staging-buffer.
 	*/
 	// FlushType: Thread safe.	
-	virtual FStagingBufferRHIRef RHICreateStagingBuffer(FVertexBufferRHIParamRef BackingBuffer)
+	virtual FStagingBufferRHIRef RHICreateStagingBuffer()
 	{
-		return new FRHIStagingBuffer(BackingBuffer);
+		return new FGenericRHIStagingBuffer();
 	}
 
     /**
      * Lock a staging buffer to read contents on the CPU that were written by the GPU, without having to stall.
-     * @discussion This function requires that you have issued an EnqueueStagedRead invocation and verified that the FRHIGPUFence has been signaled before calling.
+     * @discussion This function requires that you have issued an CopyToStagingBuffer invocation and verified that the FRHIGPUFence has been signaled before calling.
      * @param StagingBuffer The buffer to lock.
      * @param Offset The offset in the buffer to return.
      * @param SizeRHI The length of the region in the buffer to lock.
@@ -244,7 +243,7 @@ public:
 
     /**
      * Lock a staging buffer to read contents on the CPU that were written by the GPU, without having to stall.
-     * @discussion This function requires that you have issued an EnqueueStagedRead invocation and verified that the FRHIGPUFence has been signaled before calling.
+     * @discussion This function requires that you have issued an CopyToStagingBuffer invocation and verified that the FRHIGPUFence has been signaled before calling.
      * @param RHICmdList The command-list to execute on or synchronize with.
      * @param StagingBuffer The buffer to lock.
      * @param Offset The offset in the buffer to return.
diff --git a/Engine/Source/Runtime/RHI/Public/RHICommandList.h b/Engine/Source/Runtime/RHI/Public/RHICommandList.h
--- a/Engine/Source/Runtime/RHI/Public/RHICommandList.h
+++ b/Engine/Source/Runtime/RHI/Public/RHICommandList.h
@@ -1447,19 +1447,21 @@ struct FRHICommandWaitComputeFence final : public FRHICommand<FRHICommandWaitCom
 };
 
 template<ECmdList CmdListType>
-struct FRHICommandEnqueueStagedRead final : public FRHICommand<FRHICommandEnqueueStagedRead<CmdListType>>
+struct FRHICommandCopyToStagingBuffer final : public FRHICommand<FRHICommandCopyToStagingBuffer<CmdListType>>
 {
-	FStagingBufferRHIParamRef StagingBuffer;
+	FVertexBufferRHIParamRef SourceBuffer;
+	FStagingBufferRHIParamRef DestinationStagingBuffer;
 	FGPUFenceRHIParamRef Fence;
 	uint32 Offset;
 	uint32 NumBytes;
 
-	FORCEINLINE_DEBUGGABLE FRHICommandEnqueueStagedRead(FStagingBufferRHIParamRef InStagingBuffer, FGPUFenceRHIParamRef InFence, uint32 InOffset, uint32 InNumBytes)
-		: StagingBuffer(InStagingBuffer)
+	FORCEINLINE_DEBUGGABLE FRHICommandCopyToStagingBuffer(FVertexBufferRHIParamRef InSourceBuffer, FStagingBufferRHIParamRef InDestinationStagingBuffer, uint32 InOffset, uint32 InNumBytes, FGPUFenceRHIParamRef InFence = nullptr)
+		: SourceBuffer(InSourceBuffer)
+		, DestinationStagingBuffer(InDestinationStagingBuffer)
 		, Fence(InFence)
 		, Offset(InOffset)
 		, NumBytes(InNumBytes)
-	{		
+	{
 	}
 	RHI_API void Execute(FRHICommandListBase& CmdList);
 };
@@ -2730,14 +2732,14 @@ public:
 		ALLOC_COMMAND(FRHICommandWaitComputeFence<ECmdList::EGfx>)(WaitFence);
 	}
 
-	FORCEINLINE_DEBUGGABLE void EnqueueStagedRead(FStagingBufferRHIParamRef StagingBuffer, FGPUFenceRHIParamRef Fence, uint32 Offset, uint32 NumBytes)
+	FORCEINLINE_DEBUGGABLE void CopyToStagingBuffer(FVertexBufferRHIParamRef SourceBuffer, FStagingBufferRHIParamRef DestinationStagingBuffer, uint32 Offset, uint32 NumBytes, FGPUFenceRHIParamRef Fence)
 	{
 		if (Bypass())
 		{
-			CMD_CONTEXT(RHIEnqueueStagedRead)(StagingBuffer, Fence, Offset, NumBytes);
+			CMD_CONTEXT(RHICopyToStagingBuffer)(SourceBuffer, DestinationStagingBuffer, Offset, NumBytes, Fence);
 			return;
 		}
-		ALLOC_COMMAND(FRHICommandEnqueueStagedRead<ECmdList::EGfx>)(StagingBuffer, Fence, Offset, NumBytes);
+		ALLOC_COMMAND(FRHICommandCopyToStagingBuffer<ECmdList::EGfx>)(SourceBuffer, DestinationStagingBuffer, Offset, NumBytes, Fence);
 	}
 
 	FORCEINLINE_DEBUGGABLE void BeginRenderPass(const FRHIRenderPassInfo& InInfo, const TCHAR* Name)
@@ -3127,14 +3129,14 @@ public:
 		ALLOC_COMMAND(FRHICommandWaitComputeFence<ECmdList::ECompute>)(WaitFence);
 	}
 
-	FORCEINLINE_DEBUGGABLE void EnqueueStagedRead(FStagingBufferRHIParamRef StagingBuffer, FGPUFenceRHIParamRef Fence, uint32 Offset, uint32 NumBytes)
+	FORCEINLINE_DEBUGGABLE void CopyToStagingBuffer(FVertexBufferRHIParamRef SourceBuffer, FStagingBufferRHIParamRef DestinationStagingBuffer, uint32 Offset, uint32 NumBytes, FGPUFenceRHIParamRef Fence = nullptr)
 	{
 		if (Bypass())
 		{
-			CMD_CONTEXT(RHIEnqueueStagedRead)(StagingBuffer, Fence, Offset, NumBytes);
+			CMD_CONTEXT(RHICopyToStagingBuffer)(SourceBuffer, DestinationStagingBuffer, Offset, NumBytes, Fence);
 			return;
 		}
-		ALLOC_COMMAND(FRHICommandEnqueueStagedRead<ECmdList::ECompute>)(StagingBuffer, Fence, Offset, NumBytes);
+		ALLOC_COMMAND(FRHICommandCopyToStagingBuffer<ECmdList::ECompute>)(SourceBuffer, DestinationStagingBuffer, Offset, NumBytes, Fence);
 	}
 };
 
@@ -3353,9 +3355,9 @@ public:
 		return GDynamicRHI->RHICreateGPUFence(Name);
 	}
 
-	FORCEINLINE FStagingBufferRHIRef CreateStagingBuffer(FVertexBufferRHIParamRef VertexBuffer)
+	FORCEINLINE FStagingBufferRHIRef CreateStagingBuffer()
 	{
-		return GDynamicRHI->RHICreateStagingBuffer(VertexBuffer);
+		return GDynamicRHI->RHICreateStagingBuffer();
 	}
 
 	FORCEINLINE FBoundShaderStateRHIRef CreateBoundShaderState(FVertexDeclarationRHIParamRef VertexDeclaration, FVertexShaderRHIParamRef VertexShader, FHullShaderRHIParamRef HullShader, FDomainShaderRHIParamRef DomainShader, FPixelShaderRHIParamRef PixelShader, FGeometryShaderRHIParamRef GeometryShader)
@@ -4291,9 +4293,9 @@ FORCEINLINE FGPUFenceRHIRef RHICreateGPUFence(const FName& Name)
 	return FRHICommandListExecutor::GetImmediateCommandList().CreateGPUFence(Name);
 }
 
-FORCEINLINE FStagingBufferRHIRef RHICreateStagingBuffer(FVertexBufferRHIParamRef VertexBuffer)
+FORCEINLINE FStagingBufferRHIRef RHICreateStagingBuffer()
 {
-	return FRHICommandListExecutor::GetImmediateCommandList().CreateStagingBuffer(VertexBuffer);
+	return FRHICommandListExecutor::GetImmediateCommandList().CreateStagingBuffer();
 }
 
 FORCEINLINE FIndexBufferRHIRef RHICreateAndLockIndexBuffer(uint32 Stride, uint32 Size, uint32 InUsage, FRHIResourceCreateInfo& CreateInfo, void*& OutDataBuffer)
diff --git a/Engine/Source/Runtime/RHI/Public/RHICommandListCommandExecutes.inl b/Engine/Source/Runtime/RHI/Public/RHICommandListCommandExecutes.inl
--- a/Engine/Source/Runtime/RHI/Public/RHICommandListCommandExecutes.inl
+++ b/Engine/Source/Runtime/RHI/Public/RHICommandListCommandExecutes.inl
@@ -493,13 +493,13 @@ template struct FRHICommandWaitComputeFence<ECmdList::EGfx>;
 template struct FRHICommandWaitComputeFence<ECmdList::ECompute>;
 
 template<ECmdList CmdListType>
-void FRHICommandEnqueueStagedRead<CmdListType>::Execute(FRHICommandListBase& CmdList)
+void FRHICommandCopyToStagingBuffer<CmdListType>::Execute(FRHICommandListBase& CmdList)
 {
 	RHISTAT(EnqueueStagedRead);
-	INTERNAL_DECORATOR_CONTEXT(RHIEnqueueStagedRead)(StagingBuffer, Fence, Offset, NumBytes);
+	INTERNAL_DECORATOR_CONTEXT(RHICopyToStagingBuffer)(SourceBuffer, DestinationStagingBuffer, Offset, NumBytes, Fence);
 }
-template struct FRHICommandEnqueueStagedRead<ECmdList::EGfx>;
-template struct FRHICommandEnqueueStagedRead<ECmdList::ECompute>;
+template struct FRHICommandCopyToStagingBuffer<ECmdList::EGfx>;
+template struct FRHICommandCopyToStagingBuffer<ECmdList::ECompute>;
 
 void FRHICommandBuildLocalGraphicsPipelineState::Execute(FRHICommandListBase& CmdList)
 {
diff --git a/Engine/Source/Runtime/RHI/Public/RHIContext.h b/Engine/Source/Runtime/RHI/Public/RHIContext.h
--- a/Engine/Source/Runtime/RHI/Public/RHIContext.h
+++ b/Engine/Source/Runtime/RHI/Public/RHIContext.h
@@ -127,19 +127,17 @@ public:
 	 */
 	virtual void RHIInvalidateCachedState() {};
 
-    /**
-     * Enqueues on the GPU timeline any necessary operations to make the contents of 'StagingBuffer' accessible to the CPU, flushing outstanding GPU writes and/or transferring from inaccessible non-unified GPU memory to local CPU memory.
-     * @param StagingBuffer The buffer to stage. Must not be null.
-     * @param Fence A GPU fence that will be inserted into the GPU timeline and which must then be tested on the CPU to know when the StagedRead was completed. Must not be null.
-     * @param Offset The start of the data in 'StagingBuffer' to make available.
-     * @param NumBytes The lenght of data in 'StagingBuffer' to make available.
-     */
-	virtual void RHIEnqueueStagedRead(FStagingBufferRHIParamRef StagingBuffer, FGPUFenceRHIParamRef Fence, uint32 InOffset, uint32 InNumBytes)
+	/**
+	 * Performs a copy of the data in 'SourceBuffer' to 'DestinationStagingBuffer.' This will occur inline on the GPU timeline. This is a mechanism to perform nonblocking readback of a buffer at a point in time.
+	 * @param SourceBuffer The source vertex buffer that will be inlined copied.
+	 * @param DestinationStagingBuffer The the host-visible destination buffer
+	 * @param Offset The start of the data in 'SourceBuffer'
+	 * @param NumBytes The number of bytes to copy out of 'SourceBuffer'
+	 * @param Fence (optional) A GPU fence that will be inserted into the GPU timeline and which must then be tested on the CPU to know when the Copy was completed. Can be NULL, though that implies you will not have any guarantees as to when it is safe to read from 'DestinationStagingBuffer'
+	 */
+	virtual void RHICopyToStagingBuffer(FVertexBufferRHIParamRef SourceBufferRHI, FStagingBufferRHIParamRef DestinationStagingBufferRHI, uint32 InOffset, uint32 InNumBytes, FGPUFenceRHIParamRef FenceRHI = nullptr)
 	{
-		if (Fence)
-		{
-			Fence->Write();
-		}
+		check(false);
 	}
 };
 
diff --git a/Engine/Source/Runtime/RHI/Public/RHIGPUReadback.h b/Engine/Source/Runtime/RHI/Public/RHIGPUReadback.h
--- a/Engine/Source/Runtime/RHI/Public/RHIGPUReadback.h
+++ b/Engine/Source/Runtime/RHI/Public/RHIGPUReadback.h
@@ -11,17 +11,16 @@
 #include "CoreMinimal.h"
 #include "RHI.h"
 
-/* FRHIGPUMemoryReadback: represents a memory readback request scheduled with RHIScheduleGPUMemoryUpdate
-*
+/* FRHIGPUMemoryReadback: Represents a memory readback request scheduled with CopyToStagingBuffer
+** Wraps a FRHIStagingBuffer with a FRHIGPUFence for synchronization.
 */
 class FRHIGPUMemoryReadback
 {
 public:
 
-	FRHIGPUMemoryReadback(FVertexBufferRHIRef InBuffer, FName RequestName)
+	FRHIGPUMemoryReadback(FName RequestName)
 	{
-		GPUVertexBuffer = InBuffer;
-		StagingBuffer = RHICreateStagingBuffer(InBuffer);
+		DestinationStagingBuffer = RHICreateStagingBuffer();
 		Fence = RHICreateGPUFence(RequestName);
 	}
 
@@ -29,33 +28,47 @@ public:
 	{
 	}
 
-	void Insert(FRHICommandList& RHICmdList, uint32 NumBytes = 0)
+	/**
+	 * Returns the CPU accessible pointer that backs this staging buffer.
+	 * @param RHICmdList The command list to enqueue the copy request on.
+	 * @param SourceBuffer The buffer holding the source data.
+	 * @param NumBytes The number of bytes to copy. If 0, this will copy the entire buffer.
+	 */
+
+	void EnqueueCopy(FRHICommandList& RHICmdList, FVertexBufferRHIParamRef SourceBuffer, uint32 NumBytes = 0)
 	{
-		RHICmdList.EnqueueStagedRead(StagingBuffer, Fence, 0, NumBytes ? NumBytes : GPUVertexBuffer->GetSize());
+		Fence->Clear();
+		RHICmdList.CopyToStagingBuffer(SourceBuffer, DestinationStagingBuffer, 0, NumBytes ? NumBytes : SourceBuffer->GetSize(), Fence);
 	}
 
+	/**
+	 * Indicates if the data is in place and ready to be read.
+	 */
 	bool IsReady()
 	{
 		return !Fence || Fence->Poll();
 	}
 
-	void *RetrieveData(uint32 NumBytes)
+	/**
+	 * Returns the CPU accessible pointer that backs this staging buffer.
+	 * @param NumBytes The maximum number of bytes the host will read from this pointer.
+	 * @returns A CPU accessible pointer to the backing buffer.
+	 */
+	void* Lock(uint32 NumBytes)
 	{
 		ensure(Fence->Poll());
-		return RHILockStagingBuffer(StagingBuffer, 0, NumBytes);
+		return RHILockStagingBuffer(DestinationStagingBuffer, 0, NumBytes);
 	}
 
-	void Finish()
+	/**
+	 * Signals that the host is finished reading from the backing buffer.
+	 */
+	void Unlock()
 	{
-		RHIUnlockStagingBuffer(StagingBuffer);
+		RHIUnlockStagingBuffer(DestinationStagingBuffer);
 	}
 
 private:
-	FStagingBufferRHIRef StagingBuffer;
-	FVertexBufferRHIRef  GPUVertexBuffer;
+	FStagingBufferRHIRef DestinationStagingBuffer;
 	FGPUFenceRHIRef Fence;
 };
-
-
-
-
diff --git a/Engine/Source/Runtime/RHI/Public/RHIResources.h b/Engine/Source/Runtime/RHI/Public/RHIResources.h
--- a/Engine/Source/Runtime/RHI/Public/RHIResources.h
+++ b/Engine/Source/Runtime/RHI/Public/RHIResources.h
@@ -774,8 +774,9 @@ public:
 //
 
 
-/* 
-* Generic GPU fence class used by FRHIGPUMemoryReadback
+/*
+* Generic GPU fence class.
+* Granularity differs depending on backing RHI - ie it may only represent command buffer granularity.
 * RHI specific fences derive from this to implement real GPU->CPU fencing.
 * The default implementation always returns false for Poll until the next frame from the frame the fence was inserted
 * because not all APIs have a GPU/CPU sync object, we need to fake it.
@@ -786,20 +787,15 @@ public:
 	FRHIGPUFence(FName InName) : FenceName(InName) {}
 	virtual ~FRHIGPUFence() {}
 
-    /**
-     * Signal this fence now, which completion can be tested with Poll or Wait.
-	 * RHI implementation might not need to implement this depending on how RHIInsertGPUFence is implemented.
-	 * Note also that RHIInsertGPUFence is currently only called by RHIEnqueueStagedRead.
-     */
-	virtual void Write();
-
-    /**
-     * Poll the fence to see if the GPU has signaled it.
-     * @returns True if and only if the GPU fence has been inserted and the GPU has signalled the fence.
-     */
+	virtual void Clear() = 0;
+
+	/**
+	 * Poll the fence to see if the GPU has signaled it.
+	 * @returns True if and only if the GPU fence has been inserted and the GPU has signaled the fence.
+	 */
 	virtual bool Poll() const = 0;
 
-private:
+protected:
 	FName FenceName;
 };
 
@@ -809,11 +805,13 @@ class RHI_API FGenericRHIGPUFence : public FRHIGPUFence
 public:
 	FGenericRHIGPUFence(FName InName);
 
-	virtual void Write() final override;
+	virtual void Clear() final override;
 
-    /** @discussion RHI implementations must be thread-safe and must correctly handle being called before RHIInsertFence if an RHI thread is active. */
+	/** @discussion RHI implementations must be thread-safe and must correctly handle being called before RHIInsertFence if an RHI thread is active. */
 	virtual bool Poll() const final override;
 
+	void WriteInternal();
+
 private:
 	uint32 InsertedFrameNumber;
 };
@@ -1010,19 +1008,34 @@ typedef TRefCountPtr<FRHIGraphicsPipelineState> FGraphicsPipelineStateRHIRef;
 /* Generic staging buffer class used by FRHIGPUMemoryReadback
 * RHI specific staging buffers derive from this
 */
-class FRHIStagingBuffer : public FRHIResource
+class RHI_API FRHIStagingBuffer : public FRHIResource
 {
 public:
-	FRHIStagingBuffer(FVertexBufferRHIParamRef InBuffer)
-		: BackingBuffer(InBuffer)
-	{
-	}
+	FRHIStagingBuffer()
+		: bIsLocked(false)
+	{}
 
-    /** Convenience function to access the vertex-buffer that acts as the backing-store. */
-	FVertexBufferRHIParamRef GetBackingBuffer() const { return BackingBuffer.GetReference(); }
+	virtual ~FRHIStagingBuffer() {}
 
+	virtual void *Lock(uint32 Offset, uint32 NumBytes) = 0;
+	virtual void Unlock() = 0;
 protected:
-	FVertexBufferRHIRef BackingBuffer;
+	bool bIsLocked;
+};
+
+class RHI_API FGenericRHIStagingBuffer : public FRHIStagingBuffer
+{
+public:
+	FGenericRHIStagingBuffer()
+		: FRHIStagingBuffer()
+	{}
+
+	~FGenericRHIStagingBuffer() {}
+
+	virtual void* Lock(uint32 Offset, uint32 NumBytes) final override;
+	virtual void Unlock() final override;
+	FVertexBufferRHIRef ShadowBuffer;
+	uint32 Offset;
 };
 
 typedef FRHIStagingBuffer*				FStagingBufferRHIParamRef;
diff --git a/Engine/Source/Runtime/SlateCore/Private/Rendering/ElementBatcher.cpp b/Engine/Source/Runtime/SlateCore/Private/Rendering/ElementBatcher.cpp
--- a/Engine/Source/Runtime/SlateCore/Private/Rendering/ElementBatcher.cpp
+++ b/Engine/Source/Runtime/SlateCore/Private/Rendering/ElementBatcher.cpp
@@ -1930,8 +1930,18 @@ void FSlateElementBatcher::AddBorderElement( const FSlateDrawElement& DrawElemen
 	RightMarginU += HalfTexel.X;
 
 	// Determine the amount of tiling needed for the texture in this element.  The formula is number of pixels covered by the tiling portion of the texture / the number number of texels corresponding to the tiled portion of the texture.
-	float TopTiling = (RightMarginX-LeftMarginX)/(TextureSizeLocalSpace.X * ( 1 - Margin.GetTotalSpaceAlong<Orient_Horizontal>() ));
-	float LeftTiling = (BottomMarginY-TopMarginY)/(TextureSizeLocalSpace.Y * ( 1 - Margin.GetTotalSpaceAlong<Orient_Vertical>() ));
+	float TopTiling = 1.0f;
+	float LeftTiling = 1.0f;
+	float Denom = TextureSizeLocalSpace.X * (1.0f - Margin.GetTotalSpaceAlong<Orient_Horizontal>());
+	if (!FMath::IsNearlyZero(Denom))
+	{
+		TopTiling = (RightMarginX - LeftMarginX) / Denom;
+	}
+	Denom = TextureSizeLocalSpace.Y * (1.0f - Margin.GetTotalSpaceAlong<Orient_Vertical>());
+	if (!FMath::IsNearlyZero(Denom))
+	{
+		LeftTiling = (BottomMarginY - TopMarginY) / Denom;
+	}
 	
 	FShaderParams ShaderParams = FShaderParams::MakePixelShaderParams( FVector4(LeftMarginU,RightMarginU,TopMarginV,BottomMarginV) );
 
diff --git a/Engine/Source/Runtime/SlateCore/Private/Widgets/SBoxPanel.cpp b/Engine/Source/Runtime/SlateCore/Private/Widgets/SBoxPanel.cpp
--- a/Engine/Source/Runtime/SlateCore/Private/Widgets/SBoxPanel.cpp
+++ b/Engine/Source/Runtime/SlateCore/Private/Widgets/SBoxPanel.cpp
@@ -1,4 +1,4 @@
-// Copyright 1998-2018 Epic Games, Inc. All Rights Reserved.
+// Copyright 1998-2019 Epic Games, Inc. All Rights Reserved.
 
 #include "Widgets/SBoxPanel.h"
 #include "Layout/LayoutUtils.h"
diff --git a/Engine/Source/Runtime/SlateCore/Public/Layout/FlowDirection.h b/Engine/Source/Runtime/SlateCore/Public/Layout/FlowDirection.h
--- a/Engine/Source/Runtime/SlateCore/Public/Layout/FlowDirection.h
+++ b/Engine/Source/Runtime/SlateCore/Public/Layout/FlowDirection.h
@@ -5,6 +5,7 @@
 #include "CoreMinimal.h"
 #include "UObject/ObjectMacros.h"
 #include "UObject/Class.h"
+
 #include "FlowDirection.generated.h"
 
 /**
@@ -26,13 +27,13 @@ enum class EFlowDirection : uint8
 UENUM()
 enum class EFlowDirectionPreference : uint8
 {
-	/** */
+	/** Inherits the flow direction set by the parent widget. */
 	Inherit,
-	/** */
+	/** Begins laying out widgets using the current cultures layout direction preference, flipping the directionality of flows. */
 	Culture,
-	/** */
+	/** Forces a Left to Right layout flow. */
 	LeftToRight,
-	/**  */
+	/** Forces a Right to Left layout flow. */
 	RightToLeft
 };
 
diff --git a/Engine/Source/Runtime/SlateCore/Public/Layout/LayoutUtils.h b/Engine/Source/Runtime/SlateCore/Public/Layout/LayoutUtils.h
--- a/Engine/Source/Runtime/SlateCore/Public/Layout/LayoutUtils.h
+++ b/Engine/Source/Runtime/SlateCore/Public/Layout/LayoutUtils.h
@@ -3,11 +3,11 @@
 #pragma once
 
 #include "CoreMinimal.h"
-#include "Layout/Margin.h"
-#include "Layout/Visibility.h"
-#include "Layout/SlateRect.h"
-#include "Layout/ArrangedChildren.h"
-#include "Layout/FlowDirection.h"
+#include "Margin.h"
+#include "Visibility.h"
+#include "SlateRect.h"
+#include "ArrangedChildren.h"
+#include "FlowDirection.h"
 
 struct AlignmentArrangeResult
 {
diff --git a/Engine/Source/Runtime/UMG/Private/Components/PanelWidget.cpp b/Engine/Source/Runtime/UMG/Private/Components/PanelWidget.cpp
--- a/Engine/Source/Runtime/UMG/Private/Components/PanelWidget.cpp
+++ b/Engine/Source/Runtime/UMG/Private/Components/PanelWidget.cpp
@@ -62,7 +62,7 @@ TArray<UWidget*> UPanelWidget::GetAllChildren() const
 	return Result;
 }
 
-int32 UPanelWidget::GetChildIndex(UWidget* Content) const
+int32 UPanelWidget::GetChildIndex(const UWidget* Content) const
 {
 	const int32 ChildCount = GetChildrenCount();
 	for ( int32 ChildIndex = 0; ChildIndex < ChildCount; ChildIndex++ )
diff --git a/Engine/Source/Runtime/UMG/Public/Components/Border.h b/Engine/Source/Runtime/UMG/Public/Components/Border.h
--- a/Engine/Source/Runtime/UMG/Public/Components/Border.h
+++ b/Engine/Source/Runtime/UMG/Public/Components/Border.h
@@ -1,4 +1,4 @@
-// Copyright 1998-2018 Epic Games, Inc. All Rights Reserved.
+// Copyright 1998-2019 Epic Games, Inc. All Rights Reserved.
 
 #pragma once
 
diff --git a/Engine/Source/Runtime/UMG/Public/Components/PanelWidget.h b/Engine/Source/Runtime/UMG/Public/Components/PanelWidget.h
--- a/Engine/Source/Runtime/UMG/Public/Components/PanelWidget.h
+++ b/Engine/Source/Runtime/UMG/Public/Components/PanelWidget.h
@@ -41,7 +41,7 @@ public:
 
 	/** Gets the index of a specific child widget */
 	UFUNCTION(BlueprintCallable, Category="Widget|Panel")
-	int32 GetChildIndex(UWidget* Content) const;
+	int32 GetChildIndex(const UWidget* Content) const;
 
 	/** @return true if panel contains this widget */
 	UFUNCTION(BlueprintCallable, Category="Widget|Panel")
diff --git a/Engine/Source/Runtime/VulkanRHI/Private/VulkanCommands.cpp b/Engine/Source/Runtime/VulkanRHI/Private/VulkanCommands.cpp
--- a/Engine/Source/Runtime/VulkanRHI/Private/VulkanCommands.cpp
+++ b/Engine/Source/Runtime/VulkanRHI/Private/VulkanCommands.cpp
@@ -1199,28 +1199,38 @@ void FVulkanCommandListContext::PrepareParallelFromBase(const FVulkanCommandList
 	TransitionAndLayoutManager.TempCopy(BaseContext.TransitionAndLayoutManager);
 }
 
-void FVulkanCommandListContext::RHIEnqueueStagedRead(FStagingBufferRHIParamRef StagingBufferRHI, FGPUFenceRHIParamRef FenceRHI, uint32 Offset, uint32 NumBytes)
+void FVulkanCommandListContext::RHICopyToStagingBuffer(FVertexBufferRHIParamRef SourceBufferRHI, FStagingBufferRHIParamRef StagingBufferRHI, uint32 Offset, uint32 NumBytes, FGPUFenceRHIParamRef FenceRHI)
 {
 	FVulkanCmdBuffer* CmdBuffer = CommandBufferManager->GetActiveCmdBuffer();
 
+	FVulkanVertexBuffer* VertexBuffer = ResourceCast(SourceBufferRHI);
 	FVulkanGPUFence* Fence = ResourceCast(FenceRHI);
 	Fence->CmdBuffer = CmdBuffer;
 	Fence->FenceSignaledCounter = CmdBuffer->GetFenceSignaledCounter();
 
 	ensure(CmdBuffer->IsOutsideRenderPass());
-	VulkanRHI::FStagingBuffer* ReadbackStagingBuffer = Device->GetStagingManager().AcquireBuffer(NumBytes);
 	FVulkanStagingBuffer* StagingBuffer = ResourceCast(StagingBufferRHI);
-	StagingBuffer->StagingBuffer = ReadbackStagingBuffer;
+	if (!StagingBuffer->StagingBuffer || StagingBuffer->StagingBuffer->GetSize() < NumBytes)
+	{
+		if (StagingBuffer->StagingBuffer)
+		{
+			Device->GetStagingManager().ReleaseBuffer(nullptr, StagingBuffer->StagingBuffer);
+		}
+
+		VulkanRHI::FStagingBuffer* ReadbackStagingBuffer = Device->GetStagingManager().AcquireBuffer(NumBytes);
+		StagingBuffer->StagingBuffer = ReadbackStagingBuffer;
+		StagingBuffer->Device = Device;
+	}
+
 	StagingBuffer->QueuedOffset = Offset;
 	StagingBuffer->QueuedNumBytes = NumBytes;
 
 	VkBufferCopy Region;
 	FMemory::Memzero(Region);
 	Region.size = NumBytes;
-	FVulkanVertexBuffer* VertexBuffer = ResourceCast(StagingBufferRHI->GetBackingBuffer());
 	Region.srcOffset = Offset + VertexBuffer->GetOffset();
 	//Region.dstOffset = 0;
-	VulkanRHI::vkCmdCopyBuffer(CmdBuffer->GetHandle(), VertexBuffer->GetHandle(), ReadbackStagingBuffer->GetHandle(), 1, &Region);
+	VulkanRHI::vkCmdCopyBuffer(CmdBuffer->GetHandle(), VertexBuffer->GetHandle(), StagingBuffer->StagingBuffer->GetHandle(), 1, &Region);
 }
 
 
diff --git a/Engine/Source/Runtime/VulkanRHI/Private/VulkanUAV.cpp b/Engine/Source/Runtime/VulkanRHI/Private/VulkanUAV.cpp
--- a/Engine/Source/Runtime/VulkanRHI/Private/VulkanUAV.cpp
+++ b/Engine/Source/Runtime/VulkanRHI/Private/VulkanUAV.cpp
@@ -343,6 +343,11 @@ FComputeFenceRHIRef FVulkanDynamicRHI::RHICreateComputeFence(const FName& Name)
 	return new FVulkanComputeFence(Device, Name);
 }
 
+void FVulkanGPUFence::Clear()
+{
+	CmdBuffer = nullptr;
+	FenceSignaledCounter = MAX_uint64;
+}
 /*
 FRenderQueryPoolRHIRef FVulkanDynamicRHI::RHICreateRenderQueryPool(ERenderQueryType QueryType, uint32 NumQueries)
 {
@@ -360,7 +365,7 @@ FRenderQueryPoolRHIRef FVulkanDynamicRHI::RHICreateRenderQueryPool(ERenderQueryT
 
 bool FVulkanGPUFence::Poll() const
 {
-	return FenceSignaledCounter < CmdBuffer->GetFenceSignaledCounter();
+	return (CmdBuffer && (FenceSignaledCounter < CmdBuffer->GetFenceSignaledCounter()));
 }
 
 FGPUFenceRHIRef FVulkanDynamicRHI::RHICreateGPUFence(const FName& Name)
diff --git a/Engine/Source/Runtime/VulkanRHI/Private/VulkanUtil.cpp b/Engine/Source/Runtime/VulkanRHI/Private/VulkanUtil.cpp
--- a/Engine/Source/Runtime/VulkanRHI/Private/VulkanUtil.cpp
+++ b/Engine/Source/Runtime/VulkanRHI/Private/VulkanUtil.cpp
@@ -118,30 +118,43 @@ FVulkanStagingBuffer::~FVulkanStagingBuffer()
 {
 	if (StagingBuffer)
 	{
-		FVulkanVertexBuffer* VertexBuffer = ResourceCast(GetBackingBuffer());
-		VertexBuffer->GetParent()->GetStagingManager().ReleaseBuffer(nullptr, StagingBuffer);
+		check(Device);
+		Device->GetStagingManager().ReleaseBuffer(nullptr, StagingBuffer);
 	}
 }
 
-FStagingBufferRHIRef FVulkanDynamicRHI::RHICreateStagingBuffer(FVertexBufferRHIParamRef BackingBuffer)
+void* FVulkanStagingBuffer::Lock(uint32 Offset, uint32 NumBytes)
 {
-	return new FVulkanStagingBuffer(BackingBuffer);
+	check(!bIsLocked);
+	bIsLocked = true;
+	uint32 QueuedEndOffset = QueuedNumBytes + QueuedOffset;
+	uint32 EndOffset = Offset + NumBytes;
+	check(Offset < QueuedNumBytes && EndOffset <= QueuedEndOffset);
+	//#todo-rco: Apply the offset in case it doesn't match
+	return (void*)((uint8*)StagingBuffer->GetMappedPointer() + Offset);
+}
+
+void FVulkanStagingBuffer::Unlock()
+{
+	check(bIsLocked);
+	bIsLocked = false;
 }
 
-void* FVulkanDynamicRHI::RHILockStagingBuffer(FStagingBufferRHIParamRef StagingBufferRHI, uint32 Offset, uint32 SizeRHI)
+FStagingBufferRHIRef FVulkanDynamicRHI::RHICreateStagingBuffer()
+{
+	return new FVulkanStagingBuffer();
+}
+
+void* FVulkanDynamicRHI::RHILockStagingBuffer(FStagingBufferRHIParamRef StagingBufferRHI, uint32 Offset, uint32 NumBytes)
 {
 	FVulkanStagingBuffer* StagingBuffer = ResourceCast(StagingBufferRHI);
-	uint32 QueuedEndOffset = StagingBuffer->QueuedNumBytes + StagingBuffer->QueuedOffset;
-	uint32 EndOffset = Offset + SizeRHI;
-	check(Offset < StagingBuffer->QueuedNumBytes && EndOffset <= QueuedEndOffset);
-	//#todo-rco: Apply the offset in case it doesn't match
-	return StagingBuffer->StagingBuffer->GetMappedPointer();
+	return StagingBuffer->Lock(Offset, NumBytes);
 }
 
 void FVulkanDynamicRHI::RHIUnlockStagingBuffer(FStagingBufferRHIParamRef StagingBufferRHI)
 {
 	FVulkanStagingBuffer* StagingBuffer = ResourceCast(StagingBufferRHI);
-	Device->GetStagingManager().ReleaseBuffer(nullptr, StagingBuffer->StagingBuffer);
+	StagingBuffer->Unlock();
 }
 
 /**
diff --git a/Engine/Source/Runtime/VulkanRHI/Public/VulkanContext.h b/Engine/Source/Runtime/VulkanRHI/Public/VulkanContext.h
--- a/Engine/Source/Runtime/VulkanRHI/Public/VulkanContext.h
+++ b/Engine/Source/Runtime/VulkanRHI/Public/VulkanContext.h
@@ -247,7 +247,7 @@ public:
 	virtual void RHICopyTexture(FTextureRHIParamRef SourceTexture, FTextureRHIParamRef DestTexture, const FRHICopyTextureInfo& CopyInfo) final override;
 	virtual void RHITransitionResources(EResourceTransitionAccess TransitionType, FTextureRHIParamRef* InRenderTargets, int32 NumTextures) final override;
 	virtual void RHITransitionResources(EResourceTransitionAccess TransitionType, EResourceTransitionPipeline TransitionPipeline, FUnorderedAccessViewRHIParamRef* InUAVs, int32 NumUAVs, FComputeFenceRHIParamRef WriteComputeFence) final override;
-	virtual void RHIEnqueueStagedRead(FStagingBufferRHIParamRef StagingBuffer, FGPUFenceRHIParamRef Fence, uint32 Offset, uint32 NumBytes) final override;
+	virtual void RHICopyToStagingBuffer(FVertexBufferRHIParamRef SourceBuffer, FStagingBufferRHIParamRef DestinationStagingBuffer, uint32 Offset, uint32 NumBytes, FGPUFenceRHIParamRef Fence) final override;
 
 	// Render time measurement
 	virtual void RHIBeginRenderQuery(FRenderQueryRHIParamRef RenderQuery) final override;
diff --git a/Engine/Source/Runtime/VulkanRHI/Public/VulkanDynamicRHI.h b/Engine/Source/Runtime/VulkanRHI/Public/VulkanDynamicRHI.h
--- a/Engine/Source/Runtime/VulkanRHI/Public/VulkanDynamicRHI.h
+++ b/Engine/Source/Runtime/VulkanRHI/Public/VulkanDynamicRHI.h
@@ -54,7 +54,7 @@ public:
 	virtual FComputeShaderRHIRef RHICreateComputeShader(const TArray<uint8>& Code) final override;
 	virtual FComputeFenceRHIRef RHICreateComputeFence(const FName& Name) final override;
 	virtual FGPUFenceRHIRef RHICreateGPUFence(const FName &Name) final override;
-	virtual FStagingBufferRHIRef RHICreateStagingBuffer(FVertexBufferRHIParamRef BackingBuffer) final override;
+	virtual FStagingBufferRHIRef RHICreateStagingBuffer() final override;
 	virtual FBoundShaderStateRHIRef RHICreateBoundShaderState(FVertexDeclarationRHIParamRef VertexDeclaration, FVertexShaderRHIParamRef VertexShader, FHullShaderRHIParamRef HullShader, FDomainShaderRHIParamRef DomainShader, FPixelShaderRHIParamRef PixelShader, FGeometryShaderRHIParamRef GeometryShader) final override;
 	virtual FGraphicsPipelineStateRHIRef RHICreateGraphicsPipelineState(const FGraphicsPipelineStateInitializer& Initializer) final override;
 	virtual FUniformBufferRHIRef RHICreateUniformBuffer(const void* Contents, const FRHIUniformBufferLayout& Layout, EUniformBufferUsage Usage) final override;
diff --git a/Engine/Source/Runtime/VulkanRHI/Public/VulkanResources.h b/Engine/Source/Runtime/VulkanRHI/Public/VulkanResources.h
--- a/Engine/Source/Runtime/VulkanRHI/Public/VulkanResources.h
+++ b/Engine/Source/Runtime/VulkanRHI/Public/VulkanResources.h
@@ -1520,17 +1520,25 @@ protected:
 
 class FVulkanStagingBuffer : public FRHIStagingBuffer
 {
+	friend class FVulkanCommandListContext;
 public:
-	FVulkanStagingBuffer(FVertexBufferRHIRef InBuffer)
-		: FRHIStagingBuffer(InBuffer)
+	FVulkanStagingBuffer()
+		: FRHIStagingBuffer()
 	{
+		check(!bIsLocked);
 	}
 
 	virtual ~FVulkanStagingBuffer();
 
+	virtual void* Lock(uint32 Offset, uint32 NumBytes) final override;
+	virtual void Unlock() final override;
+
+private:
 	VulkanRHI::FStagingBuffer* StagingBuffer = nullptr;
 	uint32 QueuedOffset = 0;
 	uint32 QueuedNumBytes = 0;
+	// The staging buffer was allocated from this device.
+	FVulkanDevice* Device;
 };
 
 class FVulkanGPUFence : public FRHIGPUFence
@@ -1541,6 +1549,7 @@ public:
 	{
 	}
 
+	virtual void Clear() final override;
 	virtual bool Poll() const final override;
 
 protected:
diff --git a/Engine/Source/Runtime/Windows/D3D11RHI/Private/D3D11Commands.cpp b/Engine/Source/Runtime/Windows/D3D11RHI/Private/D3D11Commands.cpp
--- a/Engine/Source/Runtime/Windows/D3D11RHI/Private/D3D11Commands.cpp
+++ b/Engine/Source/Runtime/Windows/D3D11RHI/Private/D3D11Commands.cpp
@@ -2143,32 +2143,42 @@ void FD3D11DynamicRHI::RHIFlushComputeShaderCache()
 
 //*********************** StagingBuffer Implementation ***********************//
 
-FStagingBufferRHIRef FD3D11DynamicRHI::RHICreateStagingBuffer(FVertexBufferRHIParamRef VertexBufferRHI)
+FStagingBufferRHIRef FD3D11DynamicRHI::RHICreateStagingBuffer()
 {
-	return new FD3D11StagingBuffer(VertexBufferRHI);
+	return new FD3D11StagingBuffer();
 }
 
-void FD3D11DynamicRHI::RHIEnqueueStagedRead(FStagingBufferRHIParamRef StagingBufferRHI, FGPUFenceRHIParamRef Fence, uint32 Offset, uint32 NumBytes)
+void FD3D11DynamicRHI::RHICopyToStagingBuffer(FVertexBufferRHIParamRef SourceBufferRHI, FStagingBufferRHIParamRef StagingBufferRHI, uint32 Offset, uint32 NumBytes, FGPUFenceRHIParamRef FenceRHI)
 {
+	FD3D11VertexBuffer* SourceBuffer = ResourceCast(SourceBufferRHI);
 	FD3D11StagingBuffer* StagingBuffer = ResourceCast(StagingBufferRHI);
 	if (StagingBuffer)
 	{
-		FD3D11VertexBuffer* VertexBuffer = ResourceCast(StagingBuffer->GetBackingBuffer());
-		if (VertexBuffer)
+		ensureMsgf(!StagingBuffer->bIsLocked, TEXT("Attempting to Copy to a locked staging buffer. This may have undefined behavior"));
+		if (SourceBuffer)
 		{
-			// Free previously allocated buffer.
-			StagingBuffer->StagedRead.SafeRelease();
-
-			// If the static buffer is being locked for reading, create a staging buffer.
-			D3D11_BUFFER_DESC StagedReadDesc;
-			ZeroMemory(&StagedReadDesc, sizeof(D3D11_BUFFER_DESC));
-			StagedReadDesc.ByteWidth = NumBytes;
-			StagedReadDesc.Usage = D3D11_USAGE_STAGING;
-			StagedReadDesc.BindFlags = 0;
-			StagedReadDesc.CPUAccessFlags = D3D11_CPU_ACCESS_READ;
-			StagedReadDesc.MiscFlags = 0;
-			TRefCountPtr<ID3D11Buffer> StagingVertexBuffer;
-			VERIFYD3D11RESULT_EX(Direct3DDevice->CreateBuffer(&StagedReadDesc, NULL, StagingBuffer->StagedRead.GetInitReference()), Direct3DDevice);
+			if (!StagingBuffer->StagedRead || StagingBuffer->ShadowBufferSize < NumBytes)
+			{
+				// Free previously allocated buffer.
+				if (StagingBuffer->StagedRead)
+				{
+					StagingBuffer->StagedRead.SafeRelease();
+				}
+
+				// Allocate a new one with enough space.
+				// @todo-mattc I feel like we should allocate more than NumBytes to handle small reads without blowing tons of space. Need to pool this.
+				D3D11_BUFFER_DESC StagedReadDesc;
+				ZeroMemory(&StagedReadDesc, sizeof(D3D11_BUFFER_DESC));
+				StagedReadDesc.ByteWidth = NumBytes;
+				StagedReadDesc.Usage = D3D11_USAGE_STAGING;
+				StagedReadDesc.BindFlags = 0;
+				StagedReadDesc.CPUAccessFlags = D3D11_CPU_ACCESS_READ;
+				StagedReadDesc.MiscFlags = 0;
+				TRefCountPtr<ID3D11Buffer> StagingVertexBuffer;
+				VERIFYD3D11RESULT_EX(Direct3DDevice->CreateBuffer(&StagedReadDesc, NULL, StagingBuffer->StagedRead.GetInitReference()), Direct3DDevice);
+				StagingBuffer->ShadowBufferSize = NumBytes;
+				StagingBuffer->Context = Direct3DDeviceIMContext;
+			}
 
 			// Copy the contents of the vertex buffer to the staging buffer.
 			D3D11_BOX SourceBox;
@@ -2176,38 +2186,29 @@ void FD3D11DynamicRHI::RHIEnqueueStagedRead(FStagingBufferRHIParamRef StagingBuf
 			SourceBox.right = NumBytes;
 			SourceBox.top = SourceBox.front = 0;
 			SourceBox.bottom = SourceBox.back = 1;
-			Direct3DDeviceIMContext->CopySubresourceRegion(StagingBuffer->StagedRead, 0, 0, 0, 0, VertexBuffer->Resource, 0, &SourceBox);
+			Direct3DDeviceIMContext->CopySubresourceRegion(StagingBuffer->StagedRead, 0, 0, 0, 0, SourceBuffer->Resource, 0, &SourceBox);
 		}
 	}
 
+	// @todo-staging Implement real fences for D3D11
+	// D3D11 only has the generic fence for now.
+	FGenericRHIGPUFence* Fence = ResourceCast(FenceRHI);
+
 	if (Fence)
 	{
-		Fence->Write();
+		Fence->WriteInternal();
 	}
 }
 
 void* FD3D11DynamicRHI::RHILockStagingBuffer(FStagingBufferRHIParamRef StagingBufferRHI, uint32 Offset, uint32 SizeRHI)
 {
+	check(StagingBufferRHI);
 	FD3D11StagingBuffer* StagingBuffer = ResourceCast(StagingBufferRHI);
-	if (StagingBuffer && StagingBuffer->StagedRead)
-	{
-		// Map the staging buffer's memory for reading.
-		D3D11_MAPPED_SUBRESOURCE MappedSubresource;
-		VERIFYD3D11RESULT_EX(Direct3DDeviceIMContext->Map(StagingBuffer->StagedRead ,0,D3D11_MAP_READ,0,&MappedSubresource), Direct3DDevice);
-
-		return (void*)((uint8*)MappedSubresource.pData + Offset);
-	}
-	else
-	{
-		return nullptr;
-	}
+	return StagingBuffer->Lock(Offset, SizeRHI);
 }
 
 void FD3D11DynamicRHI::RHIUnlockStagingBuffer(FStagingBufferRHIParamRef StagingBufferRHI)
 {
 	FD3D11StagingBuffer* StagingBuffer = ResourceCast(StagingBufferRHI);
-	if (StagingBuffer && StagingBuffer->StagedRead)
-	{
-		Direct3DDeviceIMContext->Unmap(StagingBuffer->StagedRead,0);
-	}
+	StagingBuffer->Unlock();
 }
diff --git a/Engine/Source/Runtime/Windows/D3D11RHI/Private/D3D11RHIPrivate.h b/Engine/Source/Runtime/Windows/D3D11RHI/Private/D3D11RHIPrivate.h
--- a/Engine/Source/Runtime/Windows/D3D11RHI/Private/D3D11RHIPrivate.h
+++ b/Engine/Source/Runtime/Windows/D3D11RHI/Private/D3D11RHIPrivate.h
@@ -395,8 +395,8 @@ public:
 	virtual FGeometryShaderRHIRef RHICreateGeometryShader(const TArray<uint8>& Code) final override;
 	virtual FGeometryShaderRHIRef RHICreateGeometryShaderWithStreamOutput(const TArray<uint8>& Code, const FStreamOutElementList& ElementList, uint32 NumStrides, const uint32* Strides, int32 RasterizedStream) final override;
 	virtual FComputeShaderRHIRef RHICreateComputeShader(const TArray<uint8>& Code) final override;
-	virtual FStagingBufferRHIRef RHICreateStagingBuffer(FVertexBufferRHIParamRef VertexBuffer) final override;
-    virtual void RHIEnqueueStagedRead(FStagingBufferRHIParamRef StagingBuffer, FGPUFenceRHIParamRef Fence, uint32 Offset, uint32 NumBytes) final override;
+	virtual FStagingBufferRHIRef RHICreateStagingBuffer() final override;
+	virtual void RHICopyToStagingBuffer(FVertexBufferRHIParamRef SourceBuffer, FStagingBufferRHIParamRef DestinationStagingBuffer, uint32 Offset, uint32 NumBytes, FGPUFenceRHIParamRef Fence = nullptr) final override;
 	virtual void* RHILockStagingBuffer(FStagingBufferRHIParamRef StagingBuffer, uint32 Offset, uint32 SizeRHI) final override;
     virtual void RHIUnlockStagingBuffer(FStagingBufferRHIParamRef StagingBuffer) final override;
 	virtual FBoundShaderStateRHIRef RHICreateBoundShaderState(FVertexDeclarationRHIParamRef VertexDeclaration, FVertexShaderRHIParamRef VertexShader, FHullShaderRHIParamRef HullShader, FDomainShaderRHIParamRef DomainShader, FPixelShaderRHIParamRef PixelShader, FGeometryShaderRHIParamRef GeometryShader) final override;
diff --git a/Engine/Source/Runtime/Windows/D3D11RHI/Private/D3D11UAV.cpp b/Engine/Source/Runtime/Windows/D3D11RHI/Private/D3D11UAV.cpp
--- a/Engine/Source/Runtime/Windows/D3D11RHI/Private/D3D11UAV.cpp
+++ b/Engine/Source/Runtime/Windows/D3D11RHI/Private/D3D11UAV.cpp
@@ -285,3 +285,39 @@ void FD3D11DynamicRHI::RHIBindDebugLabelName(FUnorderedAccessViewRHIParamRef Uno
 	UAV->View->SetPrivateData(WKPDID_D3DDebugObjectName, FCString::Strlen(Name) + 1, TCHAR_TO_ANSI(Name));
 #endif
 }
+
+FD3D11StagingBuffer::~FD3D11StagingBuffer()
+{
+	if (StagedRead)
+	{
+		StagedRead.SafeRelease();
+	}
+}
+
+void* FD3D11StagingBuffer::Lock(uint32 Offset, uint32 NumBytes)
+{
+	check(!bIsLocked);
+	bIsLocked = true;
+	if (StagedRead)
+	{
+		// Map the staging buffer's memory for reading.
+		D3D11_MAPPED_SUBRESOURCE MappedSubresource;
+		VERIFYD3D11RESULT(Context->Map(StagedRead, 0, D3D11_MAP_READ, 0, &MappedSubresource));
+
+		return (void*)((uint8*)MappedSubresource.pData + Offset);
+	}
+	else
+	{
+		return nullptr;
+	}
+}
+
+void FD3D11StagingBuffer::Unlock()
+{
+	check(bIsLocked);
+	bIsLocked = false;
+	if (StagedRead)
+	{
+		Context->Unmap(StagedRead, 0);
+	}
+}
diff --git a/Engine/Source/Runtime/Windows/D3D11RHI/Public/D3D11Resources.h b/Engine/Source/Runtime/Windows/D3D11RHI/Public/D3D11Resources.h
--- a/Engine/Source/Runtime/Windows/D3D11RHI/Public/D3D11Resources.h
+++ b/Engine/Source/Runtime/Windows/D3D11RHI/Public/D3D11Resources.h
@@ -784,12 +784,21 @@ public:
 
 class FD3D11StagingBuffer : public FRHIStagingBuffer
 {
+	friend class FD3D11DynamicRHI;
 public:
-	FD3D11StagingBuffer(FVertexBufferRHIRef InBuffer)
-		: FRHIStagingBuffer(InBuffer)
+	FD3D11StagingBuffer()
+		: FRHIStagingBuffer()
 	{}
 
+	virtual ~FD3D11StagingBuffer() final override;
+
+	virtual void* Lock(uint32 Offset, uint32 NumBytes) final override;
+	virtual void Unlock() final override;
+
+private:
+	FD3D11DeviceContext* Context;
 	TRefCountPtr<ID3D11Buffer> StagedRead;
+	uint32 ShadowBufferSize;
 };
 
 /** Shader resource view class. */
@@ -924,6 +933,12 @@ struct TD3D11ResourceTraits<FRHIStagingBuffer>
 {
 	typedef FD3D11StagingBuffer TConcreteType;
 };
+// @todo-staging Implement D3D11 fences.
+template<>
+struct TD3D11ResourceTraits<FRHIGPUFence>
+{
+	typedef FGenericRHIGPUFence TConcreteType;
+};
 template<>
 struct TD3D11ResourceTraits<FRHIShaderResourceView>
 {
