Fixed Niagra null SRV crash. NiagraMeshVertexFactory HLR binding path was missing some of code from the draw policy path.

diff --git a/Engine/Plugins/FX/Niagara/Source/NiagaraVertexFactories/Private/NiagaraMeshVertexFactory.cpp b/Engine/Plugins/FX/Niagara/Source/NiagaraVertexFactories/Private/NiagaraMeshVertexFactory.cpp
--- a/Engine/Plugins/FX/Niagara/Source/NiagaraVertexFactories/Private/NiagaraMeshVertexFactory.cpp
+++ b/Engine/Plugins/FX/Niagara/Source/NiagaraVertexFactories/Private/NiagaraMeshVertexFactory.cpp
@@ -92,7 +92,7 @@ public:
 		ShaderBindings.Add(FloatDataOffset, NiagaraMeshVF->GetFloatDataOffset());
 		ShaderBindings.Add(FloatDataStride, NiagaraMeshVF->GetFloatDataStride());
 
-		ShaderBindings.Add(SortedIndices, NiagaraMeshVF->GetSortedIndicesSRV());
+		ShaderBindings.Add(SortedIndices, NiagaraMeshVF->GetSortedIndicesSRV() ? NiagaraMeshVF->GetSortedIndicesSRV() : GFNiagaraNullSortedIndicesVertexBuffer.VertexBufferSRV);
 		ShaderBindings.Add(SortedIndicesOffset, NiagaraMeshVF->GetSortedIndicesOffset());
 	}
 
