A bad entry ended up in the ddc so here's a check to catch that early.

diff --git a/Engine/Plugins/FX/Niagara/Source/NiagaraShader/Private/NiagaraShader.cpp b/Engine/Plugins/FX/Niagara/Source/NiagaraShader/Private/NiagaraShader.cpp
--- a/Engine/Plugins/FX/Niagara/Source/NiagaraShader/Private/NiagaraShader.cpp
+++ b/Engine/Plugins/FX/Niagara/Source/NiagaraShader/Private/NiagaraShader.cpp
@@ -835,7 +835,8 @@ FShader* FNiagaraShaderMap::ProcessCompilationResultsForSingleJob(FShaderCommonC
 	bCompiledSuccessfully = CurrentJob.bSucceeded;
 
 	FNiagaraShader *NiagaraShader = static_cast<FNiagaraShader*>(Shader);
-	check(Shader);
+	// UE-67395 - we had a case where we polluted the DDC with a shader containing no bytecode.
+	check(Shader && Shader->GetCode().Num() > 0);
 	check(!HasShader(NiagaraShaderType, /* PermutationId = */ 0));
 	AddShader(NiagaraShaderType, /* PermutationId = */ 0, Shader);
 
