Fix many clang errors

diff --git a/Engine/Plugins/FX/Niagara/Source/Niagara/Private/NiagaraRendererSprites.cpp b/Engine/Plugins/FX/Niagara/Source/Niagara/Private/NiagaraRendererSprites.cpp
--- a/Engine/Plugins/FX/Niagara/Source/Niagara/Private/NiagaraRendererSprites.cpp
+++ b/Engine/Plugins/FX/Niagara/Source/Niagara/Private/NiagaraRendererSprites.cpp
@@ -282,13 +282,13 @@ void NiagaraRendererSprites::GetDynamicMeshElements(const TArray<const FSceneVie
 
 				FNiagaraSpriteVFLooseParameters VFLooseParams;
 				VFLooseParams.NumCutoutVerticesPerFrame = CollectorResources.VertexFactory.GetNumCutoutVerticesPerFrame();
-				VFLooseParams.CutoutGeometry = CollectorResources.VertexFactory.GetCutoutGeometrySRV() ? CollectorResources.VertexFactory.GetCutoutGeometrySRV() : GFNiagaraNullSubUVCutoutVertexBuffer.VertexBufferSRV;
+				VFLooseParams.CutoutGeometry = CollectorResources.VertexFactory.GetCutoutGeometrySRV() ? CollectorResources.VertexFactory.GetCutoutGeometrySRV() : GFNiagaraNullSubUVCutoutVertexBuffer.VertexBufferSRV.GetReference();
 				VFLooseParams.NiagaraParticleDataFloat = CollectorResources.VertexFactory.GetParticleDataFloatSRV();
 				VFLooseParams.NiagaraFloatDataOffset = CollectorResources.VertexFactory.GetFloatDataOffset();
 				VFLooseParams.NiagaraFloatDataStride = CollectorResources.VertexFactory.GetFloatDataStride();
 				VFLooseParams.ParticleAlignmentMode = CollectorResources.VertexFactory.GetAlignmentMode();
 				VFLooseParams.ParticleFacingMode = CollectorResources.VertexFactory.GetFacingMode();
-				VFLooseParams.SortedIndices = CollectorResources.VertexFactory.GetSortedIndicesSRV() ? CollectorResources.VertexFactory.GetSortedIndicesSRV() : GFNiagaraNullSortedIndicesVertexBuffer.VertexBufferSRV;
+				VFLooseParams.SortedIndices = CollectorResources.VertexFactory.GetSortedIndicesSRV() ? CollectorResources.VertexFactory.GetSortedIndicesSRV() : GFNiagaraNullSortedIndicesVertexBuffer.VertexBufferSRV.GetReference();
 				VFLooseParams.SortedIndicesOffset = CollectorResources.VertexFactory.GetSortedIndicesOffset();
 
 				CollectorResources.VertexFactory.LooseParameterUniformBuffer = FNiagaraSpriteVFLooseParametersRef::CreateUniformBufferImmediate(VFLooseParams, UniformBuffer_SingleFrame);
diff --git a/Engine/Plugins/FX/Niagara/Source/NiagaraVertexFactories/Public/NiagaraSpriteVertexFactory.h b/Engine/Plugins/FX/Niagara/Source/NiagaraVertexFactories/Public/NiagaraSpriteVertexFactory.h
--- a/Engine/Plugins/FX/Niagara/Source/NiagaraVertexFactories/Public/NiagaraSpriteVertexFactory.h
+++ b/Engine/Plugins/FX/Niagara/Source/NiagaraVertexFactories/Public/NiagaraSpriteVertexFactory.h
@@ -119,9 +119,8 @@ public:
 		FacingMode(0),
 		FloatDataOffset(0),
 		FloatDataStride(0),
-		SortedIndicesOffset(0),
-		LooseParameterUniformBuffer(nullptr)
-
+		LooseParameterUniformBuffer(nullptr),
+		SortedIndicesOffset(0)
 	{}
 
 	FNiagaraSpriteVertexFactory()
diff --git a/Engine/Source/Runtime/OpenGLDrv/Private/OpenGLUniformBuffer.cpp b/Engine/Source/Runtime/OpenGLDrv/Private/OpenGLUniformBuffer.cpp
--- a/Engine/Source/Runtime/OpenGLDrv/Private/OpenGLUniformBuffer.cpp
+++ b/Engine/Source/Runtime/OpenGLDrv/Private/OpenGLUniformBuffer.cpp
@@ -739,10 +739,13 @@ void FOpenGLDynamicRHI::RHIUpdateUniformBuffer(FUniformBufferRHIParamRef Uniform
 		{
 			FRHIResource* Resource = *(FRHIResource**)((uint8*)Contents + Layout.Resources[ResourceIndex].MemberOffset);
 
+#if 0
 			checkf(Resource, TEXT("Invalid resource entry creating uniform buffer, %s.Resources[%u], ResourceType 0x%x."),
 				*Layout.GetDebugName().ToString(),
 				ResourceIndex,
 				Layout.Resources[ResourceIndex].MemberType);
+#endif
+			check(Resource);
 
 			UniformBuffer->ResourceTable[ResourceIndex] = Resource;
 		}
diff --git a/Engine/Source/Runtime/Renderer/Private/PathTracing.cpp b/Engine/Source/Runtime/Renderer/Private/PathTracing.cpp
--- a/Engine/Source/Runtime/Renderer/Private/PathTracing.cpp
+++ b/Engine/Source/Runtime/Renderer/Private/PathTracing.cpp
@@ -669,5 +669,4 @@ void FDeferredShadingSceneRenderer::RenderPathTracing(FRHICommandListImmediate&
 	ResolveSceneColor(RHICmdList);
 
 }
-
-#endif RHI_RAYTRACING
+#endif
\ No newline at end of file
diff --git a/Engine/Source/Runtime/Renderer/Private/RayTracing/RayTracingBarycentrics.cpp b/Engine/Source/Runtime/Renderer/Private/RayTracing/RayTracingBarycentrics.cpp
--- a/Engine/Source/Runtime/Renderer/Private/RayTracing/RayTracingBarycentrics.cpp
+++ b/Engine/Source/Runtime/Renderer/Private/RayTracing/RayTracingBarycentrics.cpp
@@ -115,5 +115,4 @@ void FDeferredShadingSceneRenderer::RenderRayTracedBarycentrics(FRHICommandListI
 
 	GraphBuilder.Execute();
 }
-
-#endif RHI_RAYTRACING
+#endif
\ No newline at end of file
diff --git a/Engine/Source/Runtime/Renderer/Private/RayTracing/RayTracingDebug.cpp b/Engine/Source/Runtime/Renderer/Private/RayTracing/RayTracingDebug.cpp
--- a/Engine/Source/Runtime/Renderer/Private/RayTracing/RayTracingDebug.cpp
+++ b/Engine/Source/Runtime/Renderer/Private/RayTracing/RayTracingDebug.cpp
@@ -151,5 +151,4 @@ void FDeferredShadingSceneRenderer::RenderRayTracedDebug(FRHICommandListImmediat
 }
 
 #undef LOCTEXT_NAMESPACE
-
-#endif RHI_RAYTRACING
+#endif
\ No newline at end of file
diff --git a/Engine/Source/Runtime/Renderer/Private/ScreenSpaceDenoise.h b/Engine/Source/Runtime/Renderer/Private/ScreenSpaceDenoise.h
--- a/Engine/Source/Runtime/Renderer/Private/ScreenSpaceDenoise.h
+++ b/Engine/Source/Runtime/Renderer/Private/ScreenSpaceDenoise.h
@@ -40,6 +40,9 @@ public:
 		SHADER_PARAMETER_RDG_TEXTURE(Texture2D, SpecularPenumbra)
 	END_SHADER_PARAMETER_STRUCT()
 
+
+	virtual ~IScreenSpaceDenoiser() {};
+
 	/** Debug name of the denoiser for draw event. */
 	virtual const TCHAR* GetDebugName() const = 0;
 
diff --git a/Engine/Source/Runtime/Renderer/Public/MeshPassProcessor.inl b/Engine/Source/Runtime/Renderer/Public/MeshPassProcessor.inl
--- a/Engine/Source/Runtime/Renderer/Public/MeshPassProcessor.inl
+++ b/Engine/Source/Runtime/Renderer/Public/MeshPassProcessor.inl
@@ -186,11 +186,13 @@ void FMeshPassProcessor::BuildRayTracingDrawCommands(
 
 	SharedMeshDrawCommand.SetDebugData(PrimitiveSceneProxy, &MaterialResource, &MaterialRenderProxy, PassShaders.GetUntypedShaders());
 
+#if RHI_RAYTRACING
 	if (PassShaders.RayHitGroupShader)
 	{
 		FMeshDrawSingleShaderBindings ShaderBindings = SharedMeshDrawCommand.ShaderBindings.GetSingleShaderBindings(SF_RayHitGroup);
 		PassShaders.RayHitGroupShader->GetShaderBindings(Scene, FeatureLevel, PrimitiveSceneProxy, MaterialRenderProxy, MaterialResource, DrawRenderState.GetViewUniformBuffer(), DrawRenderState.GetPassUniformBuffer(), ShaderElementData, ShaderBindings);
 	}
+#endif
 
 	const int32 NumElements = MeshBatch.Elements.Num();
 
@@ -201,11 +203,13 @@ void FMeshPassProcessor::BuildRayTracingDrawCommands(
 			const FMeshBatchElement& BatchElement = MeshBatch.Elements[BatchElementIndex];
 			FMeshDrawCommand& MeshDrawCommand = DrawListContext.AddCommand(SharedMeshDrawCommand);
 
+#if RHI_RAYTRACING
 			if (PassShaders.RayHitGroupShader)
 			{
 				FMeshDrawSingleShaderBindings RayHitGroupShaderBindings = MeshDrawCommand.ShaderBindings.GetSingleShaderBindings(SF_RayHitGroup);
 				PassShaders.RayHitGroupShader->GetElementShaderBindings(Scene, ViewIfDynamicMeshCommand, VertexFactory, false, FeatureLevel, PrimitiveSceneProxy, MeshBatch, BatchElement, ShaderElementData, RayHitGroupShaderBindings, MeshDrawCommand.VertexStreams);
 			}
+#endif
 
 			SetDrawCommandEvent(PrimitiveSceneProxy, MaterialResource, MeshDrawCommand);
 
