Fix niagara uniformbuffer lifetime crashes

diff --git a/Engine/Plugins/FX/Niagara/Source/NiagaraVertexFactories/Public/NiagaraSpriteVertexFactory.h b/Engine/Plugins/FX/Niagara/Source/NiagaraVertexFactories/Public/NiagaraSpriteVertexFactory.h
--- a/Engine/Plugins/FX/Niagara/Source/NiagaraVertexFactories/Public/NiagaraSpriteVertexFactory.h
+++ b/Engine/Plugins/FX/Niagara/Source/NiagaraVertexFactories/Public/NiagaraSpriteVertexFactory.h
@@ -247,7 +247,7 @@ public:
 	 */
 	static FVertexFactoryShaderParameters* ConstructShaderParameters(EShaderFrequency ShaderFrequency);
 
-	FUniformBufferRHIParamRef LooseParameterUniformBuffer;
+	FUniformBufferRHIRef LooseParameterUniformBuffer;
 
 protected:
 	/** Initialize streams for this vertex factory. */
