Fixed FNiagaraSceneProxy velocity relevance.

diff --git a/Engine/Plugins/FX/Niagara/Source/Niagara/Private/NiagaraComponent.cpp b/Engine/Plugins/FX/Niagara/Source/Niagara/Private/NiagaraComponent.cpp
--- a/Engine/Plugins/FX/Niagara/Source/Niagara/Private/NiagaraComponent.cpp
+++ b/Engine/Plugins/FX/Niagara/Source/Niagara/Private/NiagaraComponent.cpp
@@ -217,6 +217,9 @@ FPrimitiveViewRelevance FNiagaraSceneProxy::GetViewRelevance(const FSceneView* V
 			Relevance |= Renderer->GetViewRelevance(View, this);
 		}
 	}
+
+	Relevance.bVelocityRelevance = IsMovable() && Relevance.bOpaqueRelevance && Relevance.bRenderInMainPass;
+
 	return Relevance;
 }
 
