We don't support gpu particles on SM4. So we have to be a little smarter about our check. Should probably be rolled into RHISupportsComputeShaders.

diff --git a/Engine/Plugins/FX/Niagara/Source/Niagara/Private/NiagaraScript.cpp b/Engine/Plugins/FX/Niagara/Source/Niagara/Private/NiagaraScript.cpp
--- a/Engine/Plugins/FX/Niagara/Source/Niagara/Private/NiagaraScript.cpp
+++ b/Engine/Plugins/FX/Niagara/Source/Niagara/Private/NiagaraScript.cpp
@@ -1071,7 +1071,10 @@ void UNiagaraScript::CacheResourceShadersForRendering(bool bRegenerateId, bool b
 			//if (ScriptResourcesByFeatureLevel[FeatureLevel])
 			{
 				EShaderPlatform ShaderPlatform = GShaderPlatformForFeatureLevel[CacheFeatureLevel];
-				if (IsFeatureLevelSupported(ShaderPlatform, ERHIFeatureLevel::SM5) || IsFeatureLevelSupported(ShaderPlatform, ERHIFeatureLevel::ES3_1))
+				// SM4 is "in between" these feature levels but we do NOT support GPU particles in SM4. 
+				// So we can't use IsFeatureLevelSupported since SM4 will be seen as past ES3.1 and we do NOT support SM4 GPU particles.
+				// @todo-mattc This check should be rolled into RHISupportsComputeShaders.
+				if (GetMaxSupportedFeatureLevel(ShaderPlatform) == ERHIFeatureLevel::SM5 || GetMaxSupportedFeatureLevel(ShaderPlatform) == ERHIFeatureLevel::ES3_1)
 				{
 					ResourceToCache = ScriptResourcesByFeatureLevel[CacheFeatureLevel];
 					CacheShadersForResources(ShaderPlatform, &ScriptResource, true);
