Check if raytracing is enabled before initializing all its stuff. Fixes failure at launch on D3D10.

diff --git a/Engine/Plugins/FX/Niagara/Source/Niagara/Private/NiagaraRendererSprites.cpp b/Engine/Plugins/FX/Niagara/Source/Niagara/Private/NiagaraRendererSprites.cpp
--- a/Engine/Plugins/FX/Niagara/Source/Niagara/Private/NiagaraRendererSprites.cpp
+++ b/Engine/Plugins/FX/Niagara/Source/Niagara/Private/NiagaraRendererSprites.cpp
@@ -69,8 +69,11 @@ void NiagaraRendererSprites::ReleaseRenderThreadResources()
 	WorldSpacePrimitiveUniformBuffer.ReleaseResource();
 
 #if RHI_RAYTRACING
-	RayTracingGeometry.ReleaseResource();
-	RayTracingDynamicVertexBuffer.Release();
+	if (IsRayTracingEnabled())
+	{
+		RayTracingGeometry.ReleaseResource();
+		RayTracingDynamicVertexBuffer.Release();
+	}
 #endif
 }
 
@@ -80,21 +83,24 @@ void NiagaraRendererSprites::CreateRenderThreadResources()
 	VertexFactory->InitResource();
 
 #if RHI_RAYTRACING
-	RayTracingDynamicVertexBuffer.Initialize(4, 256, PF_R32_FLOAT, BUF_UnorderedAccess | BUF_ShaderResource, TEXT("RayTracingDynamicVertexBuffer"));
-
-	FRayTracingGeometryInitializer Initializer;
-	Initializer.PositionVertexBuffer = nullptr;
-	Initializer.IndexBuffer = nullptr;
-	Initializer.BaseVertexIndex = 0;
-	Initializer.VertexBufferStride = 12;
-	Initializer.VertexBufferByteOffset = 0;
-	Initializer.TotalPrimitiveCount = 0;
-	Initializer.VertexBufferElementType = VET_Float3;
-	Initializer.PrimitiveType = PT_TriangleList;
-	Initializer.bFastBuild = true;
-	Initializer.bAllowUpdate = false;
-	RayTracingGeometry.SetInitializer(Initializer);
-	RayTracingGeometry.InitResource();
+	if (IsRayTracingEnabled())
+	{
+		RayTracingDynamicVertexBuffer.Initialize(4, 256, PF_R32_FLOAT, BUF_UnorderedAccess | BUF_ShaderResource, TEXT("RayTracingDynamicVertexBuffer"));
+
+		FRayTracingGeometryInitializer Initializer;
+		Initializer.PositionVertexBuffer = nullptr;
+		Initializer.IndexBuffer = nullptr;
+		Initializer.BaseVertexIndex = 0;
+		Initializer.VertexBufferStride = 12;
+		Initializer.VertexBufferByteOffset = 0;
+		Initializer.TotalPrimitiveCount = 0;
+		Initializer.VertexBufferElementType = VET_Float3;
+		Initializer.PrimitiveType = PT_TriangleList;
+		Initializer.bFastBuild = true;
+		Initializer.bAllowUpdate = false;
+		RayTracingGeometry.SetInitializer(Initializer);
+		RayTracingGeometry.InitResource();
+	}
 #endif
 }
 
