Adds a cvar to allow niagara VM shader compilation to skip some of HLSLCC's backend optimization passes.
Depends on 4861258!

diff --git a/Engine/Plugins/FX/Niagara/Source/Niagara/Private/NiagaraScriptDerivedData.cpp b/Engine/Plugins/FX/Niagara/Source/Niagara/Private/NiagaraScriptDerivedData.cpp
--- a/Engine/Plugins/FX/Niagara/Source/Niagara/Private/NiagaraScriptDerivedData.cpp
+++ b/Engine/Plugins/FX/Niagara/Source/Niagara/Private/NiagaraScriptDerivedData.cpp
@@ -32,8 +32,9 @@ FString FNiagaraScriptDerivedData::GetPluginSpecificCacheKeySuffix() const
 {
 	enum { UE_NIAGARA_COMPILATION_DERIVEDDATA_VER = 1 };
 
-	FString Ret = FString::Printf(TEXT("%i"),
-		(int32)UE_NIAGARA_COMPILATION_DERIVEDDATA_VER);
+	FString Ret = FString::Printf(TEXT("%i_%i"),
+		(int32)UE_NIAGARA_COMPILATION_DERIVEDDATA_VER, GNiagaraSkipVectorVMBackendOptimizations);
+
 	GeneratedVMId.AppendKeyString(Ret);
 
 	//UE_LOG(LogNiagara, Log, TEXT("Cache Key Gen Niagara VM Script DDC data for %s key %s"), *ScriptFullName, *Ret);
diff --git a/Engine/Plugins/FX/Niagara/Source/NiagaraEditor/Private/NiagaraCompiler.cpp b/Engine/Plugins/FX/Niagara/Source/NiagaraEditor/Private/NiagaraCompiler.cpp
--- a/Engine/Plugins/FX/Niagara/Source/NiagaraEditor/Private/NiagaraCompiler.cpp
+++ b/Engine/Plugins/FX/Niagara/Source/NiagaraEditor/Private/NiagaraCompiler.cpp
@@ -28,6 +28,7 @@
 #include "EdGraphSchema_Niagara.h"
 #include "Misc/FileHelper.h"
 #include "ShaderCompiler.h"
+#include "NiagaraShader.h"
 
 #define LOCTEXT_NAMESPACE "NiagaraCompiler"
 
@@ -616,7 +617,7 @@ void FNiagaraEditorModule::TestCompileScriptFromConsole(const TArray<FString>& A
 			FShaderCompilerOutput Output;
 			FVectorVMCompilationOutput CompilationOutput;
 			double StartTime = FPlatformTime::Seconds();
-			bool bSucceeded = CompileShader_VectorVM(Input, Output, FString(FPlatformProcess::ShaderDir()), 0, CompilationOutput);
+			bool bSucceeded = CompileShader_VectorVM(Input, Output, FString(FPlatformProcess::ShaderDir()), 0, CompilationOutput, GNiagaraSkipVectorVMBackendOptimizations);
 			float DeltaTime = (float)(FPlatformTime::Seconds() - StartTime);
 
 			if (bSucceeded)
@@ -767,7 +768,7 @@ FNiagaraCompileResults FHlslNiagaraCompiler::CompileScript(const FNiagaraCompile
 			
 			CritSec.Lock();
 			double StartTime = FPlatformTime::Seconds();
-			CompileResults.bVMSucceeded = CompileShader_VectorVM(Input, Output, FString(FPlatformProcess::ShaderDir()), 0, CompilationOutput);
+			CompileResults.bVMSucceeded = CompileShader_VectorVM(Input, Output, FString(FPlatformProcess::ShaderDir()), 0, CompilationOutput, GNiagaraSkipVectorVMBackendOptimizations);
 			CompileResults.CompileTime = (float)(FPlatformTime::Seconds() - StartTime);
 			CritSec.Unlock();
 		}
diff --git a/Engine/Plugins/FX/Niagara/Source/NiagaraShader/Private/NiagaraShader.cpp b/Engine/Plugins/FX/Niagara/Source/NiagaraShader/Private/NiagaraShader.cpp
--- a/Engine/Plugins/FX/Niagara/Source/NiagaraShader/Private/NiagaraShader.cpp
+++ b/Engine/Plugins/FX/Niagara/Source/NiagaraShader/Private/NiagaraShader.cpp
@@ -32,6 +32,13 @@ static FAutoConsoleVariableRef CVarCreateNiagaraShadersOnLoad(
 	TEXT("Whether to create Niagara's simulation shaders on load, which can reduce hitching, but use more memory.  Otherwise they will be created as needed.")
 	);
 
+int32 GNiagaraSkipVectorVMBackendOptimizations = 0;
+static FAutoConsoleVariableRef CVarNiagaraSkipVectorVMBackendOptimizations(
+	TEXT("fx.SkipVectorVMBackendOptimizations"),
+	GNiagaraSkipVectorVMBackendOptimizations,
+	TEXT("If 1, skip HLSLCC's backend optimization passes during VectorVM compilation. \n"),
+	ECVF_Default
+);
 
 #if ENABLE_COOK_STATS
 namespace NiagaraShaderCookStats
diff --git a/Engine/Plugins/FX/Niagara/Source/NiagaraShader/Public/NiagaraShader.h b/Engine/Plugins/FX/Niagara/Source/NiagaraShader/Public/NiagaraShader.h
--- a/Engine/Plugins/FX/Niagara/Source/NiagaraShader/Public/NiagaraShader.h
+++ b/Engine/Plugins/FX/Niagara/Source/NiagaraShader/Public/NiagaraShader.h
@@ -163,3 +163,5 @@ class FNiagaraEmitterInstanceShader : public FNiagaraShader
 {
 
 };
+
+extern NIAGARASHADER_API int32 GNiagaraSkipVectorVMBackendOptimizations;
diff --git a/Engine/Source/Developer/ShaderFormatVectorVM/Private/VectorVMShaderCompiler.cpp b/Engine/Source/Developer/ShaderFormatVectorVM/Private/VectorVMShaderCompiler.cpp
--- a/Engine/Source/Developer/ShaderFormatVectorVM/Private/VectorVMShaderCompiler.cpp
+++ b/Engine/Source/Developer/ShaderFormatVectorVM/Private/VectorVMShaderCompiler.cpp
@@ -31,7 +31,7 @@ bool CompileShader_VectorVM(const FShaderCompilerInput& Input, FShaderCompilerOu
 }
 
 //TODO: Move to this output living in the shader eco-system with the compute shaders too but for now just do things more directly.
-bool CompileShader_VectorVM(const FShaderCompilerInput& Input, FShaderCompilerOutput& Output,const FString& WorkingDirectory, uint8 Version, FVectorVMCompilationOutput& VMCompilationOutput)
+bool CompileShader_VectorVM(const FShaderCompilerInput& Input, FShaderCompilerOutput& Output, const FString& WorkingDirectory, uint8 Version, FVectorVMCompilationOutput& VMCompilationOutput, bool bSkipBackendOptimizations)
 {
 	SCOPE_CYCLE_COUNTER(STAT_VectorVM_Compiler_CompileShader_VectorVM);
 
@@ -102,6 +102,11 @@ bool CompileShader_VectorVM(const FShaderCompilerInput& Input, FShaderCompilerOu
 		//CCFlags |= HLSLCC_PrintAST;
 	//CCFlags |= HLSLCC_UseFullPrecisionInPS;
 
+	if (bSkipBackendOptimizations)
+	{
+		CCFlags |= HLSLCC_DisableBackendOptimizations;
+	}
+
 	//TODO: Do this later when we implement the rest of the shader plumbing stuff.
 // 		if (bDumpDebugInfo)
 // 		{
diff --git a/Engine/Source/Developer/ShaderFormatVectorVM/Public/ShaderFormatVectorVM.h b/Engine/Source/Developer/ShaderFormatVectorVM/Public/ShaderFormatVectorVM.h
--- a/Engine/Source/Developer/ShaderFormatVectorVM/Public/ShaderFormatVectorVM.h
+++ b/Engine/Source/Developer/ShaderFormatVectorVM/Public/ShaderFormatVectorVM.h
@@ -34,4 +34,4 @@ struct FVectorVMCompilationOutput
 extern bool SHADERFORMATVECTORVM_API CompileShader_VectorVM(const struct FShaderCompilerInput& Input, struct FShaderCompilerOutput& Output, const class FString& WorkingDirectory, uint8 Version);
 
 //Cheating hack version. To be removed when we add all the plumbing for VVM scripts to be treat like real shaders.
-extern bool SHADERFORMATVECTORVM_API CompileShader_VectorVM(const struct FShaderCompilerInput& Input,struct FShaderCompilerOutput& Output,const class FString& WorkingDirectory, uint8 Version, struct FVectorVMCompilationOutput& VMCompilationOutput);
+extern bool SHADERFORMATVECTORVM_API CompileShader_VectorVM(const struct FShaderCompilerInput& Input,struct FShaderCompilerOutput& Output,const class FString& WorkingDirectory, uint8 Version, struct FVectorVMCompilationOutput& VMCompilationOutput, bool bSkipBackendOptimizations = false);
