Niagara - Fix op post load so that it doesn't invalide compile ids when adding the 'add' pin.  This prevents constant recompiles.

diff --git a/Engine/Plugins/FX/Niagara/Source/NiagaraCore/Public/NiagaraCustomVersion.h b/Engine/Plugins/FX/Niagara/Source/NiagaraCore/Public/NiagaraCustomVersion.h
--- a/Engine/Plugins/FX/Niagara/Source/NiagaraCore/Public/NiagaraCustomVersion.h
+++ b/Engine/Plugins/FX/Niagara/Source/NiagaraCore/Public/NiagaraCustomVersion.h
@@ -111,6 +111,9 @@ struct FNiagaraCustomVersion
 		TextureDataInterfaceSizeSerialize, // The texture data interface now streams size info
 
 		SkelMeshInterfaceAPIImprovements, //API to skeletal mesh interface was improved but requires a recompile and some graph fixup.
+
+		ImproveLoadTimeFixupOfOpAddPins, // Only do op add pin fixup on existing nodes which are before this version
+
 		// -----<new versions can be added above this line>-------------------------------------------------
 		VersionPlusOne,
 		LatestVersion = VersionPlusOne - 1,
diff --git a/Engine/Plugins/FX/Niagara/Source/NiagaraEditor/Private/NiagaraNodeOp.cpp b/Engine/Plugins/FX/Niagara/Source/NiagaraEditor/Private/NiagaraNodeOp.cpp
--- a/Engine/Plugins/FX/Niagara/Source/NiagaraEditor/Private/NiagaraNodeOp.cpp
+++ b/Engine/Plugins/FX/Niagara/Source/NiagaraEditor/Private/NiagaraNodeOp.cpp
@@ -162,24 +162,14 @@ void UNiagaraNodeOp::PostLoad()
 		UE_LOG(LogNiagaraEditor, Log, TEXT("OpNode: Converted %s to %s, Package: %s"), *OriginalOpName.ToString(), *OpName.ToString(), *GetOutermost()->GetName());
 	}
 
-
-	if (AllowDynamicPins())
+	const int32 NiagaraVer = GetLinkerCustomVersion(FNiagaraCustomVersion::GUID);
+	if (NiagaraVer < FNiagaraCustomVersion::ImproveLoadTimeFixupOfOpAddPins &&
+		AllowDynamicPins() &&
+		Pins.FindByPredicate([this](UEdGraphPin* Pin) { return IsAddPin(Pin); }) == nullptr)
 	{
-		bool HasAddPin = false;
-		for (UEdGraphPin* Pin : Pins)
-		{
-			if (IsAddPin(Pin))
-			{
-				HasAddPin = true;
-				break;
-			}
-		}
-		if (!HasAddPin)
-		{
-			// This adds the "plus" pin to nodes that were created before the feature was added
-			ReallocatePins();
-			UE_LOG(LogNiagaraEditor, Log, TEXT("OpNode %s, Package %s: reallocated existing pins to add new extension pin"), *OpName.ToString(), *GetOutermost()->GetName());
-		}
+		// Add the pin directly here rather than calling allocate default pins to prevent the graph id from being invalidated
+		// since adding the add pin doesn't change the compile behavior.
+		CreateAddPin(EGPD_Input);
 	}
 }
 
