Niagara - Enable multiple accepted values for visible and edit conditions on script inputs and add properties for common meta-data rather than using the name-string map to define them.

diff --git a/Engine/Plugins/FX/Niagara/Source/Niagara/Public/NiagaraTypes.h b/Engine/Plugins/FX/Niagara/Source/Niagara/Public/NiagaraTypes.h
--- a/Engine/Plugins/FX/Niagara/Source/Niagara/Public/NiagaraTypes.h
+++ b/Engine/Plugins/FX/Niagara/Source/Niagara/Public/NiagaraTypes.h
@@ -297,6 +297,21 @@ enum class ENiagaraExecutionState : uint32
 	Num
 };
 
+/** Defines options for conditionally editing and showing script inputs in the UI. */
+USTRUCT()
+struct NIAGARA_API FNiagaraInputConditionMetadata
+{
+	GENERATED_USTRUCT_BODY()
+public:
+	/** The name of the input to use for matching the target values. */
+	UPROPERTY(EditAnywhere, Category="Input Condition")
+	FName InputName;
+
+	/** The list of target values which will satisfy the input condition.  If this is empty it's assumed to be a single value of "true" for matching bool inputs. */
+	UPROPERTY(EditAnywhere, Category="Input Condition")
+	TArray<FString> TargetValues;
+};
+
 USTRUCT()
 struct NIAGARA_API FNiagaraVariableMetaData
 {
@@ -304,6 +319,8 @@ struct NIAGARA_API FNiagaraVariableMetaData
 public:
 	FNiagaraVariableMetaData()
 		:EditorSortPriority(0)
+		, bAdvancedDisplay(false)
+		, bInlineEditConditionToggle(false)
 		, CallSortPriority(0)
 	{
 	}
@@ -320,6 +337,23 @@ public:
 	UPROPERTY(EditAnywhere, Category = "Variable", meta = (ToolTip = "Affects the sort order in the editor stacks. Use a smaller number to push it to the top. Defaults to zero."))
 	int32 EditorSortPriority;
 
+	/** Declares that this input is advanced and should only be visible if expanded inputs have been expanded. */
+	UPROPERTY(EditAnywhere, Category = "Variable")
+	bool bAdvancedDisplay;
+
+	/** Declares the associated input is used as an inline edit condition toggle, so it should should be hidden and it should be edited as a 
+	checkbox inline with the input which has designated as its edit condition. */
+	UPROPERTY(EditAnywhere, Category = "Variable")
+	bool bInlineEditConditionToggle;
+
+	/** Declares the associated input should be conditionally editable based on the value of another input. */
+	UPROPERTY(EditAnywhere, Category = "Input Conditions")
+	FNiagaraInputConditionMetadata EditCondition;
+
+	/** Declares the associated input should be conditionally visible based on the value of another input. */
+	UPROPERTY(EditAnywhere, Category = "Input Conditions")
+	FNiagaraInputConditionMetadata VisibleCondition;
+
 	UPROPERTY()
 	int32 CallSortPriority;
 
diff --git a/Engine/Plugins/FX/Niagara/Source/NiagaraCore/Public/NiagaraCustomVersion.h b/Engine/Plugins/FX/Niagara/Source/NiagaraCore/Public/NiagaraCustomVersion.h
--- a/Engine/Plugins/FX/Niagara/Source/NiagaraCore/Public/NiagaraCustomVersion.h
+++ b/Engine/Plugins/FX/Niagara/Source/NiagaraCore/Public/NiagaraCustomVersion.h
@@ -114,6 +114,8 @@ struct FNiagaraCustomVersion
 
 		ImproveLoadTimeFixupOfOpAddPins, // Only do op add pin fixup on existing nodes which are before this version
 
+		MoveCommonInputMetadataToProperties, // Moved commonly used input metadata out of the strin/string property metadata map to actual properties on the metadata struct.
+
 		// -----<new versions can be added above this line>-------------------------------------------------
 		VersionPlusOne,
 		LatestVersion = VersionPlusOne - 1,
diff --git a/Engine/Plugins/FX/Niagara/Source/NiagaraEditor/Private/Customizations/NiagaraMetaDataCustomNodeBuilder.h b/Engine/Plugins/FX/Niagara/Source/NiagaraEditor/Private/Customizations/NiagaraMetaDataCustomNodeBuilder.h
--- a/Engine/Plugins/FX/Niagara/Source/NiagaraEditor/Private/Customizations/NiagaraMetaDataCustomNodeBuilder.h
+++ b/Engine/Plugins/FX/Niagara/Source/NiagaraEditor/Private/Customizations/NiagaraMetaDataCustomNodeBuilder.h
@@ -55,50 +55,35 @@ public:
 			MetaDataGroup.ToggleExpansion(true);
 			TSharedRef<FStructOnScope> StructData = MetadataViewModel->GetValueStruct();
 			
-			IDetailPropertyRow* DescriptionRow = ChildrenBuilder.AddExternalStructureProperty(StructData, FName(TEXT("Description")), "Description");
-
-			if (DescriptionRow)
-			{
-				DescriptionRow->Visibility(EVisibility::Hidden); // hide it here, show it in groups only
-				MetaDataGroup.AddPropertyRow(DescriptionRow->GetPropertyHandle()->AsShared());
-
-				TSharedPtr<IPropertyHandle> PropertyHandle = DescriptionRow->GetPropertyHandle();
-				PropertyHandle->SetOnPropertyValueChanged(FSimpleDelegate::CreateSP(MetadataViewModel, &FNiagaraMetaDataViewModel::NotifyMetaDataChanged));
-			}
-			IDetailPropertyRow* CategoryRow = ChildrenBuilder.AddExternalStructureProperty(StructData, FName(TEXT("CategoryName")), "CategoryName");
-
-			if (CategoryRow)
+			auto AddGroupedPropertyRow = [&ChildrenBuilder, &MetadataViewModel, &MetaDataGroup, &StructData](FName PropertyName, bool bShouldAutoExpand, bool bShouldWatchChildProperties)
 			{
-				CategoryRow->Visibility(EVisibility::Hidden); // hide it here, show it in groups only
-				MetaDataGroup.AddPropertyRow(CategoryRow->GetPropertyHandle()->AsShared());
-
-				TSharedPtr<IPropertyHandle> PropertyHandle = CategoryRow->GetPropertyHandle();
-				PropertyHandle->SetOnPropertyValueChanged(FSimpleDelegate::CreateSP(MetadataViewModel, &FNiagaraMetaDataViewModel::NotifyMetaDataChanged));
-			}
-
-			IDetailPropertyRow* SortingRow = ChildrenBuilder.AddExternalStructureProperty(StructData, FName(TEXT("EditorSortPriority")), "EditorSortPriority");
-
-			if (SortingRow)
-			{
-				SortingRow->Visibility(EVisibility::Hidden); // hide it here, show it in groups only
-				MetaDataGroup.AddPropertyRow(SortingRow->GetPropertyHandle()->AsShared());
-
-				TSharedPtr<IPropertyHandle> PropertyHandle = SortingRow->GetPropertyHandle();
-				PropertyHandle->SetOnPropertyValueChanged(FSimpleDelegate::CreateSP(MetadataViewModel, &FNiagaraMetaDataViewModel::NotifyMetaDataChanged));
-			}
-
-			IDetailPropertyRow* MetaDataRow = ChildrenBuilder.AddExternalStructureProperty(StructData, FName(TEXT("PropertyMetaData")), "PropertyMetaData");
-			
-			if (MetaDataRow)
-			{
-				MetaDataRow->ShouldAutoExpand(true);
-				MetaDataRow->Visibility(EVisibility::Hidden);
-				MetaDataGroup.AddPropertyRow(MetaDataRow->GetPropertyHandle()->AsShared());
-
-				TSharedPtr<IPropertyHandle> PropertyHandle = MetaDataRow->GetPropertyHandle();
-				PropertyHandle->SetOnPropertyValueChanged(FSimpleDelegate::CreateSP(MetadataViewModel, &FNiagaraMetaDataViewModel::NotifyMetaDataChanged));
-				PropertyHandle->SetOnChildPropertyValueChanged(FSimpleDelegate::CreateSP(MetadataViewModel, &FNiagaraMetaDataViewModel::NotifyMetaDataChanged));
-			}
+				IDetailPropertyRow* PropertyRow = ChildrenBuilder.AddExternalStructureProperty(StructData, PropertyName, PropertyName);
+				if (PropertyRow)
+				{
+					if (bShouldAutoExpand)
+					{
+						PropertyRow->ShouldAutoExpand(true);
+					}
+					PropertyRow->Visibility(EVisibility::Hidden);
+					MetaDataGroup.AddPropertyRow(PropertyRow->GetPropertyHandle()->AsShared());
+
+					TSharedPtr<IPropertyHandle> PropertyHandle = PropertyRow->GetPropertyHandle();
+					PropertyHandle->SetOnPropertyValueChanged(FSimpleDelegate::CreateSP(MetadataViewModel, &FNiagaraMetaDataViewModel::NotifyMetaDataChanged));
+					if (bShouldWatchChildProperties)
+					{
+						PropertyHandle->SetOnChildPropertyValueChanged(FSimpleDelegate::CreateSP(MetadataViewModel, &FNiagaraMetaDataViewModel::NotifyMetaDataChanged));
+					}
+				}
+			};
+
+			AddGroupedPropertyRow(GET_MEMBER_NAME_CHECKED(FNiagaraVariableMetaData, Description), false, false);
+			AddGroupedPropertyRow(GET_MEMBER_NAME_CHECKED(FNiagaraVariableMetaData, CategoryName), false, false);
+			AddGroupedPropertyRow(GET_MEMBER_NAME_CHECKED(FNiagaraVariableMetaData, bAdvancedDisplay), false, false);
+			AddGroupedPropertyRow(GET_MEMBER_NAME_CHECKED(FNiagaraVariableMetaData, bInlineEditConditionToggle), false, false);
+			AddGroupedPropertyRow(GET_MEMBER_NAME_CHECKED(FNiagaraVariableMetaData, EditorSortPriority), false, false);
+			AddGroupedPropertyRow(GET_MEMBER_NAME_CHECKED(FNiagaraVariableMetaData, EditCondition), false, true);
+			AddGroupedPropertyRow(GET_MEMBER_NAME_CHECKED(FNiagaraVariableMetaData, VisibleCondition), false, true);
+			AddGroupedPropertyRow(GET_MEMBER_NAME_CHECKED(FNiagaraVariableMetaData, PropertyMetaData), true, true);
 		}
 	}
 
diff --git a/Engine/Plugins/FX/Niagara/Source/NiagaraEditor/Private/NiagaraEditorModule.cpp b/Engine/Plugins/FX/Niagara/Source/NiagaraEditor/Private/NiagaraEditorModule.cpp
--- a/Engine/Plugins/FX/Niagara/Source/NiagaraEditor/Private/NiagaraEditorModule.cpp
+++ b/Engine/Plugins/FX/Niagara/Source/NiagaraEditor/Private/NiagaraEditorModule.cpp
@@ -99,11 +99,6 @@ IMPLEMENT_MODULE( FNiagaraEditorModule, NiagaraEditor );
 const FName FNiagaraEditorModule::NiagaraEditorAppIdentifier( TEXT( "NiagaraEditorApp" ) );
 const FLinearColor FNiagaraEditorModule::WorldCentricTabColorScale(0.0f, 0.0f, 0.2f, 0.5f);
 
-const FName FNiagaraEditorModule::FInputMetaDataKeys::AdvancedDisplay = "AdvancedDisplay";
-const FName FNiagaraEditorModule::FInputMetaDataKeys::EditCondition = "EditCondition";
-const FName FNiagaraEditorModule::FInputMetaDataKeys::VisibleCondition = "VisibleCondition";
-const FName FNiagaraEditorModule::FInputMetaDataKeys::InlineEditConditionToggle = "InlineEditConditionToggle";
-
 EAssetTypeCategories::Type FNiagaraEditorModule::NiagaraAssetCategory;
 
 //////////////////////////////////////////////////////////////////////////
diff --git a/Engine/Plugins/FX/Niagara/Source/NiagaraEditor/Private/NiagaraGraph.cpp b/Engine/Plugins/FX/Niagara/Source/NiagaraEditor/Private/NiagaraGraph.cpp
--- a/Engine/Plugins/FX/Niagara/Source/NiagaraEditor/Private/NiagaraGraph.cpp
+++ b/Engine/Plugins/FX/Niagara/Source/NiagaraEditor/Private/NiagaraGraph.cpp
@@ -184,6 +184,55 @@ void UNiagaraGraph::PostLoad()
 
 	Parameters.Empty();
 	FindParameters();
+
+	// Migrate input condition metadata
+	const int32 NiagaraVer = GetLinkerCustomVersion(FNiagaraCustomVersion::GUID);
+	if (NiagaraVer < FNiagaraCustomVersion::MoveCommonInputMetadataToProperties)
+	{
+		auto MigrateInputCondition = [](TMap<FName,FString>& PropertyMetaData, const FName& InputConditionKey, FNiagaraInputConditionMetadata& InOutInputCondition)
+		{
+			FString* InputCondition = PropertyMetaData.Find(InputConditionKey);
+			if (InputCondition != nullptr)
+			{
+				FString InputName;
+				FString TargetValue;
+				int32 EqualsIndex = InputCondition->Find("=");
+				if (EqualsIndex == INDEX_NONE)
+				{
+					InOutInputCondition.InputName = **InputCondition;
+				}
+				else
+				{
+					InOutInputCondition.InputName = *InputCondition->Left(EqualsIndex);
+					InOutInputCondition.TargetValues.Add(InputCondition->RightChop(EqualsIndex + 1));
+				}
+				PropertyMetaData.Remove(InputConditionKey);
+			}
+		};
+
+		for (auto& VariableToMetaDataItem : VariableToMetaData)
+		{
+			FNiagaraVariableMetaData& MetaData = VariableToMetaDataItem.Value;
+
+			// Migrate advanced display.
+			if (MetaData.PropertyMetaData.Contains("AdvancedDisplay"))
+			{
+				MetaData.bAdvancedDisplay = true;
+				MetaData.PropertyMetaData.Remove("AdvancedDisplay");
+			}
+
+			// Migrate inline edit condition toggle
+			if (MetaData.PropertyMetaData.Contains("InlineEditConditionToggle"))
+			{
+				MetaData.bInlineEditConditionToggle = true;
+				MetaData.PropertyMetaData.Remove("InlineEditConditionToggle");
+			}
+
+			// Migrate edit and visible conditions
+			MigrateInputCondition(MetaData.PropertyMetaData, TEXT("EditCondition"), MetaData.EditCondition);
+			MigrateInputCondition(MetaData.PropertyMetaData, TEXT("VisibleCondition"), MetaData.VisibleCondition);
+		}
+	}
 }
 
 void UNiagaraGraph::PostEditChangeProperty(FPropertyChangedEvent& PropertyChangedEvent)
diff --git a/Engine/Plugins/FX/Niagara/Source/NiagaraEditor/Private/ViewModels/Stack/NiagaraStackFunctionInput.cpp b/Engine/Plugins/FX/Niagara/Source/NiagaraEditor/Private/ViewModels/Stack/NiagaraStackFunctionInput.cpp
--- a/Engine/Plugins/FX/Niagara/Source/NiagaraEditor/Private/ViewModels/Stack/NiagaraStackFunctionInput.cpp
+++ b/Engine/Plugins/FX/Niagara/Source/NiagaraEditor/Private/ViewModels/Stack/NiagaraStackFunctionInput.cpp
@@ -604,15 +604,16 @@ void UNiagaraStackFunctionInput::RefreshFromMetaData()
 
 		if (InputMetaData != nullptr)
 		{
-			SetIsAdvanced(InputMetaData->PropertyMetaData.Contains(FNiagaraEditorModule::FInputMetaDataKeys::AdvancedDisplay));
+			SetIsAdvanced(InputMetaData->bAdvancedDisplay);
 
-			EditCondition.Refresh(InputMetaData->PropertyMetaData.Find(FNiagaraEditorModule::FInputMetaDataKeys::EditCondition));
+			FText EditConditionError;
+			EditCondition.Refresh(InputMetaData->EditCondition, EditConditionError);
 			if (EditCondition.IsValid() && EditCondition.GetConditionInputType() == FNiagaraTypeDefinition::GetBoolDef())
 			{
 				FNiagaraVariableMetaData* EditConditionInputMetadata = EditCondition.GetConditionInputMetaData();
 				if (EditConditionInputMetadata != nullptr)
 				{
-					bShowEditConditionInline = EditConditionInputMetadata->PropertyMetaData.Contains(FNiagaraEditorModule::FInputMetaDataKeys::InlineEditConditionToggle);
+					bShowEditConditionInline = EditConditionInputMetadata->bInlineEditConditionToggle;
 				}
 			}
 			else
@@ -620,10 +621,23 @@ void UNiagaraStackFunctionInput::RefreshFromMetaData()
 				bShowEditConditionInline = false;
 			}
 
-			VisibleCondition.Refresh(InputMetaData->PropertyMetaData.Find(FNiagaraEditorModule::FInputMetaDataKeys::VisibleCondition));
+			if (EditConditionError.IsEmpty() == false)
+			{
+				UE_LOG(LogNiagaraEditor, Warning, TEXT("Edit condition failed to bind.  Function: %s Input: %s Message: %s"), 
+					*OwningFunctionCallNode->GetFunctionName(), *InputParameterHandle.GetName().ToString(), *EditConditionError.ToString());
+			}
+
+			FText VisibleConditionError;
+			VisibleCondition.Refresh(InputMetaData->VisibleCondition, VisibleConditionError);
+
+			if (VisibleConditionError.IsEmpty() == false)
+			{
+				UE_LOG(LogNiagaraEditor, Warning, TEXT("Visible condition failed to bind.  Function: %s Input: %s Message: %s"),
+					*OwningFunctionCallNode->GetFunctionName(), *InputParameterHandle.GetName().ToString(), *VisibleConditionError.ToString());
+			}
 
 			bIsInlineEditConditionToggle = InputType == FNiagaraTypeDefinition::GetBoolDef() && 
-				InputMetaData->PropertyMetaData.Contains(FNiagaraEditorModule::FInputMetaDataKeys::InlineEditConditionToggle);
+				InputMetaData->bInlineEditConditionToggle;
 		}
 	}
 }
diff --git a/Engine/Plugins/FX/Niagara/Source/NiagaraEditor/Private/ViewModels/Stack/NiagaraStackFunctionInputCondition.cpp b/Engine/Plugins/FX/Niagara/Source/NiagaraEditor/Private/ViewModels/Stack/NiagaraStackFunctionInputCondition.cpp
--- a/Engine/Plugins/FX/Niagara/Source/NiagaraEditor/Private/ViewModels/Stack/NiagaraStackFunctionInputCondition.cpp
+++ b/Engine/Plugins/FX/Niagara/Source/NiagaraEditor/Private/ViewModels/Stack/NiagaraStackFunctionInputCondition.cpp
@@ -23,77 +23,81 @@ void FNiagaraStackFunctionInputCondition::Initialize(UNiagaraScript* InScript,
 	FunctionCallNode = InFunctionCallNode;
 }
 
-void FNiagaraStackFunctionInputCondition::Refresh(const FString* Condition)
+void FNiagaraStackFunctionInputCondition::Refresh(const FNiagaraInputConditionMetadata& InputCondition, FText& OutErrorMessage)
 {
-	TargetValueData.Empty();
+	TargetValuesData.Empty();
 	InputBinder.Reset();
 
-	if (Condition == nullptr || Condition->IsEmpty())
+	if (InputCondition.InputName == "")
 	{
 		return;
 	}
 
-	FString InputName;
-	FString TargetValue;
-	int32 EqualsIndex = Condition->Find("=");
-	if (EqualsIndex == INDEX_NONE)
+	TArray<FString> TargetValues(InputCondition.TargetValues);
+	if (TargetValues.Num() == 0)
 	{
-		// If the condition doesn't have an equals sign the target value is assumed to be a bool value of true.
-		InputName = *Condition;
-		TargetValue = "true";
-	}
-	else
-	{
-		InputName = Condition->Left(EqualsIndex);
-		TargetValue = Condition->RightChop(EqualsIndex + 1);
+		// If no target values were specified it's assumed that the input is a bool and that the target value is "true".
+		TargetValues.Add("true");
 	}
 
-	FText ErrorMessage;
-	if (InputBinder.TryBind(Script, DependentScripts, OwningEmitterUniqueName, FunctionCallNode, *InputName, TOptional<FNiagaraTypeDefinition>(), true, ErrorMessage))
+	if (InputBinder.TryBind(Script, DependentScripts, OwningEmitterUniqueName, FunctionCallNode, InputCondition.InputName, TOptional<FNiagaraTypeDefinition>(), true, OutErrorMessage))
 	{
 		FNiagaraEditorModule& NiagaraEditorModule = FModuleManager::GetModuleChecked<FNiagaraEditorModule>("NiagaraEditor");
 		TSharedPtr<INiagaraEditorTypeUtilities, ESPMode::ThreadSafe> TypeEditorUtilities = NiagaraEditorModule.GetTypeUtilities(InputBinder.GetInputType());
 		if (TypeEditorUtilities.IsValid())
 		{
 			FNiagaraVariable TempVariable(InputBinder.GetInputType(), "Temp");
-			bool bValueParsed = false;
-
-			if (TypeEditorUtilities->CanHandlePinDefaults())
-			{
-				bValueParsed = TypeEditorUtilities->SetValueFromPinDefaultString(TargetValue, TempVariable);
-			}
-			if (bValueParsed == false && TypeEditorUtilities->CanSetValueFromDisplayName())
+			for (FString& TargetValue : TargetValues)
 			{
-				bValueParsed = TypeEditorUtilities->SetValueFromDisplayName(FText::FromString(TargetValue), TempVariable);
-			}
-
-			if (bValueParsed)
-			{
-				TargetValueData.AddUninitialized(InputBinder.GetInputType().GetSize());
-				TempVariable.CopyTo(TargetValueData.GetData());
-			}
-			else
-			{
-				ErrorMessage = FText::Format(LOCTEXT("ParseValueError", "Target value {0} is not a valid for type {1}"), FText::FromString(TargetValue), InputBinder.GetInputType().GetNameText());
-				InputBinder.Reset();
+				bool bValueParsed = true;
+				if (TypeEditorUtilities->CanHandlePinDefaults())
+				{
+					bValueParsed = TypeEditorUtilities->SetValueFromPinDefaultString(TargetValue, TempVariable);
+				}
+				if (bValueParsed == false && TypeEditorUtilities->CanSetValueFromDisplayName())
+				{
+					bValueParsed = TypeEditorUtilities->SetValueFromDisplayName(FText::FromString(TargetValue), TempVariable);
+				}
+
+				if (bValueParsed)
+				{
+					TArray<uint8>& TargetValueData = TargetValuesData.AddDefaulted_GetRef();
+					TargetValueData.AddUninitialized(InputBinder.GetInputType().GetSize());
+					TempVariable.CopyTo(TargetValueData.GetData());
+				}
+				else
+				{
+					OutErrorMessage = FText::Format(LOCTEXT("ParseValueError", "Target value {0} is not a valid for type {1}"),
+						FText::FromString(TargetValue), InputBinder.GetInputType().GetNameText());
+					InputBinder.Reset();
+					break;
+				}
 			}
 		}
-	}
-
-	if(ErrorMessage.IsEmpty() == false)
-	{
-		UE_LOG(LogNiagaraEditor, Warning, TEXT("Input condition failed to bind %s=%s.  Message: %s"), *InputName, *TargetValue, *ErrorMessage.ToString());
+		else
+		{
+			OutErrorMessage = FText::Format(LOCTEXT("TypeEditorUtilitiesNotFoundError", "No type editor utilites registered for type {0}.  Can not parse target input condition {1}."),
+				InputBinder.GetInputType().GetNameText(), FText::FromName(InputCondition.InputName));
+			InputBinder.Reset();
+		}
 	}
 }
 
 bool FNiagaraStackFunctionInputCondition::IsValid() const
 {
-	return InputBinder.IsValid() && (TargetValueData.Num() > 0);
+	return InputBinder.IsValid() && TargetValuesData.Num() > 0;
 }
 
 bool FNiagaraStackFunctionInputCondition::GetConditionIsEnabled() const
 {
-	return IsValid() && FMemory::Memcmp(TargetValueData.GetData(), InputBinder.GetData().GetData(), TargetValueData.Num()) == 0;
+	if (IsValid())
+	{
+		TArray<uint8> InputValue = InputBinder.GetData();
+		bool bMatchFound = TargetValuesData.FindByPredicate([&InputValue](const TArray<uint8>& TargetValueData) {
+			return FMemory::Memcmp(TargetValueData.GetData(), InputValue.GetData(), TargetValueData.Num()) == 0; }) != nullptr;
+		return bMatchFound;
+	}
+	return false;
 }
 
 bool FNiagaraStackFunctionInputCondition::CanSetConditionIsEnabled() const
diff --git a/Engine/Plugins/FX/Niagara/Source/NiagaraEditor/Public/NiagaraEditorModule.h b/Engine/Plugins/FX/Niagara/Source/NiagaraEditor/Public/NiagaraEditorModule.h
--- a/Engine/Plugins/FX/Niagara/Source/NiagaraEditor/Public/NiagaraEditorModule.h
+++ b/Engine/Plugins/FX/Niagara/Source/NiagaraEditor/Public/NiagaraEditorModule.h
@@ -32,14 +32,6 @@ public:
 	DECLARE_DELEGATE_RetVal_OneParam(TSharedRef<SWidget>, FOnCreateStackWidget, UNiagaraStackViewModel*);
 	DECLARE_DELEGATE_RetVal_OneParam(UMovieSceneNiagaraParameterTrack*, FOnCreateMovieSceneTrackForParameter, FNiagaraVariable);
 
-	struct FInputMetaDataKeys
-	{
-		static const FName AdvancedDisplay;
-		static const FName EditCondition;
-		static const FName VisibleCondition;
-		static const FName InlineEditConditionToggle;
-	};
-
 public:
 	FNiagaraEditorModule();
 
diff --git a/Engine/Plugins/FX/Niagara/Source/NiagaraEditor/Public/ViewModels/Stack/NiagaraStackFunctionInputCondition.h b/Engine/Plugins/FX/Niagara/Source/NiagaraEditor/Public/ViewModels/Stack/NiagaraStackFunctionInputCondition.h
--- a/Engine/Plugins/FX/Niagara/Source/NiagaraEditor/Public/ViewModels/Stack/NiagaraStackFunctionInputCondition.h
+++ b/Engine/Plugins/FX/Niagara/Source/NiagaraEditor/Public/ViewModels/Stack/NiagaraStackFunctionInputCondition.h
@@ -17,7 +17,7 @@ public:
 		FString InOwningEmitterUniqueName,
 		UNiagaraNodeFunctionCall* InFunctionCallNode);
 
-	void Refresh(const FString* Condition);
+	void Refresh(const FNiagaraInputConditionMetadata& InputCondition, FText& OutErrorMessage);
 
 	bool IsValid() const;
 
@@ -44,5 +44,5 @@ private:
 
 	FNiagaraStackFunctionInputBinder InputBinder;
 
-	TArray<uint8> TargetValueData;
+	TArray<TArray<uint8>> TargetValuesData;
 };
\ No newline at end of file
