UE4 - Resolve conflict
#rnx

diff --git a/Engine/Plugins/FX/Niagara/Source/Niagara/Private/NiagaraEmitterInstanceBatcher.cpp b/Engine/Plugins/FX/Niagara/Source/Niagara/Private/NiagaraEmitterInstanceBatcher.cpp
--- a/Engine/Plugins/FX/Niagara/Source/Niagara/Private/NiagaraEmitterInstanceBatcher.cpp
+++ b/Engine/Plugins/FX/Niagara/Source/Niagara/Private/NiagaraEmitterInstanceBatcher.cpp
@@ -30,30 +30,30 @@ FFXSystemInterface* NiagaraEmitterInstanceBatcher::GetInterface(const FName& InN
 	return InName == Name ? this : nullptr;
 }
 
-void NiagaraEmitterInstanceBatcher::Queue(FNiagaraComputeExecutionContext *InContext)
+void NiagaraEmitterInstanceBatcher::Queue(FNiagaraComputeExecutionContext *ExecContext)
 {
 	//UE_LOG(LogNiagara, Warning, TEXT("Submitted!"));
 	//SimulationQueue[CurQueueIndex]->Add(InContext);
-		ENQUEUE_UNIQUE_RENDER_COMMAND_THREEPARAMETER(QueueNiagaraDispatch,
-			TArray<FNiagaraComputeExecutionContext*>*, Queue, &SimulationQueue[0],
-			uint32, QueueIndex, CurQueueIndex,
-			FNiagaraComputeExecutionContext*, ExecContext, InContext,
+	TArray<FNiagaraComputeExecutionContext*>* Queue = &SimulationQueue[0];
+	uint32 QueueIndex = CurQueueIndex;
+	ENQUEUE_RENDER_COMMAND(QueueNiagaraDispatch)(
+		[Queue, QueueIndex, ExecContext](FRHICommandListImmediate& RHICmdList)
+		{
+			const uint32 QueueIndexMask = (1 << QueueIndex);
+			//Don't queue the same context for execution multiple times. TODO: possibly try to combine/accumulate the tick info if we happen to have > 1 before it's executed.
+			if (!(ExecContext->PendingExecutionQueueMask & QueueIndexMask))
 			{
-				const uint32 QueueIndexMask = (1 << QueueIndex);
-				//Don't queue the same context for execution multiple times. TODO: possibly try to combine/accumulate the tick info if we happen to have > 1 before it's executed.
-				if (!(ExecContext->PendingExecutionQueueMask & QueueIndexMask))
-				{
-					Queue[QueueIndex].Add(ExecContext);
-					ExecContext->PendingExecutionQueueMask |= QueueIndexMask;
-				}
-			});
+				Queue[QueueIndex].Add(ExecContext);
+				ExecContext->PendingExecutionQueueMask |= QueueIndexMask;
+			}
+		});
 }
 
-void NiagaraEmitterInstanceBatcher::Remove(FNiagaraComputeExecutionContext *InContext)
+void NiagaraEmitterInstanceBatcher::Remove(FNiagaraComputeExecutionContext *ExecContext)
 {
-	ENQUEUE_UNIQUE_RENDER_COMMAND_TWOPARAMETER(RemoveNiagaraDispatch,
-		TArray<FNiagaraComputeExecutionContext*>*, Queue, &SimulationQueue[0],
-		FNiagaraComputeExecutionContext*, ExecContext, InContext,
+	TArray<FNiagaraComputeExecutionContext*>* Queue = &SimulationQueue[0];
+	ENQUEUE_RENDER_COMMAND(RemoveNiagaraDispatch)(
+		[Queue, ExecContext](FRHICommandListImmediate& RHICmdList)
 		{
 			for (int32 i = 0; i < SIMULATION_QUEUE_COUNT; i++)
 			{
