Scalability Changes

> Changed "MinDistToCamera" to "LODDistance"

> Editor only override for detail mode.

> Override for LODDistance

> Can now open ParameterCollection by double clicking it's usage in the stack.

> Now keeping view relevance even when dynamic data is not present so that visbility culling can function.

> Safety check bypassing crash in GPU sims debug info

> PreviewGrid helper actor to improve previwing Niagara FX and how parameter or scalability changes will affect them. See ScalabilityTests in EngineTests for examples.

diff --git a/Engine/Plugins/FX/Niagara/Source/Niagara/Classes/NiagaraConstants.h b/Engine/Plugins/FX/Niagara/Source/Niagara/Classes/NiagaraConstants.h
--- a/Engine/Plugins/FX/Niagara/Source/Niagara/Classes/NiagaraConstants.h
+++ b/Engine/Plugins/FX/Niagara/Source/Niagara/Classes/NiagaraConstants.h
@@ -39,8 +39,10 @@
 #define SYS_PARAM_ENGINE_LOCAL_TO_WORLD_NO_SCALE         INiagaraModule::GetVar_Engine_Owner_SystemLocalToWorldNoScale()
 #define SYS_PARAM_ENGINE_WORLD_TO_LOCAL_NO_SCALE         INiagaraModule::GetVar_Engine_Owner_SystemWorldToLocalNoScale()
 
-#define SYS_PARAM_ENGINE_TIME_SINCE_RENDERED             INiagaraModule::GetVar_Engine_Owner_TimeSinceRendered()
-#define SYS_PARAM_ENGINE_MIN_DIST_TO_CAMERA              INiagaraModule::GetVar_Engine_Owner_MinDistanceToCamera()
+
+#define SYS_PARAM_ENGINE_TIME_SINCE_RENDERED		INiagaraModule::GetVar_Engine_Owner_TimeSinceRendered()
+#define SYS_PARAM_ENGINE_LOD_DISTANCE				INiagaraModule::GetVar_Engine_Owner_LODDistance()
+
 
 #define SYS_PARAM_ENGINE_EXECUTION_STATE                 INiagaraModule::GetVar_Engine_Owner_ExecutionState()
 
diff --git a/Engine/Plugins/FX/Niagara/Source/Niagara/Classes/NiagaraEmitter.h b/Engine/Plugins/FX/Niagara/Source/Niagara/Classes/NiagaraEmitter.h
--- a/Engine/Plugins/FX/Niagara/Source/Niagara/Classes/NiagaraEmitter.h
+++ b/Engine/Plugins/FX/Niagara/Source/Niagara/Classes/NiagaraEmitter.h
@@ -317,7 +317,7 @@ public:
 #endif
 
 	/** Is this emitter allowed to be enabled by the current system detail level. */
-	bool IsAllowedByDetailLevel()const;
+	bool IsAllowedByDetailLevel(int32 DetailLevel)const;
 	NIAGARA_API bool RequiresPersistantIDs()const;
 
 	NIAGARA_API bool IsValid()const;
diff --git a/Engine/Plugins/FX/Niagara/Source/Niagara/Private/NiagaraComponent.cpp b/Engine/Plugins/FX/Niagara/Source/Niagara/Private/NiagaraComponent.cpp
--- a/Engine/Plugins/FX/Niagara/Source/Niagara/Private/NiagaraComponent.cpp
+++ b/Engine/Plugins/FX/Niagara/Source/Niagara/Private/NiagaraComponent.cpp
@@ -342,6 +342,10 @@ UNiagaraComponent::UNiagaraComponent(const FObjectInitializer& ObjectInitializer
 	, bIsSeeking(false)
 	, bAutoDestroy(false)
 #if WITH_EDITOR
+	, PreviewDetailLevel(INDEX_NONE)
+	, PreviewLODDistance(0.0f)
+	, bEnablePreviewDetailLevel(false)
+	, bEnablePreviewLODDistance(false)
 	, bWaitForCompilationOnActivate(false)
 #endif
 	, bAwaitingActivationDueToNotReady(false)
@@ -453,7 +457,10 @@ void UNiagaraComponent::ResetSystem()
 void UNiagaraComponent::ReinitializeSystem()
 {
 	DestroyInstance();
-	Activate();
+	if (bIsActive)
+	{
+		Activate(true);
+	}
 }
 
 bool UNiagaraComponent::GetRenderingEnabled() const
@@ -1206,6 +1213,24 @@ void UNiagaraComponent::PostLoadNormalizeOverrideNames()
 	EditorOverridesValue = ValueMap;
 }
 
+void UNiagaraComponent::SetPreviewDetailLevel(bool bInEnablePreviewDetailLevel, int32 InPreviewDetailLevel)
+{
+	bool bReInit = bEnablePreviewDetailLevel != bInEnablePreviewDetailLevel || (bEnablePreviewDetailLevel && PreviewDetailLevel != InPreviewDetailLevel);
+
+	bEnablePreviewDetailLevel = bInEnablePreviewDetailLevel;
+	PreviewDetailLevel = InPreviewDetailLevel;
+	if (bReInit)
+	{
+		ReinitializeSystem();
+	}
+}
+
+void UNiagaraComponent::SetPreviewLODDistance(bool bInEnablePreviewLODDistance, float InPreviewLODDistance)
+{
+	bEnablePreviewLODDistance = bInEnablePreviewLODDistance;
+	PreviewLODDistance = InPreviewLODDistance;
+}
+
 bool UNiagaraComponent::IsParameterValueOverriddenLocally(const FName& InParamName)
 {
 	bool* FoundVar = EditorOverridesValue.Find(InParamName);
diff --git a/Engine/Plugins/FX/Niagara/Source/Niagara/Private/NiagaraConstants.cpp b/Engine/Plugins/FX/Niagara/Source/Niagara/Private/NiagaraConstants.cpp
--- a/Engine/Plugins/FX/Niagara/Source/Niagara/Private/NiagaraConstants.cpp
+++ b/Engine/Plugins/FX/Niagara/Source/Niagara/Private/NiagaraConstants.cpp
@@ -44,7 +44,7 @@ void FNiagaraConstants::Init()
 		SystemParameters.Add(SYS_PARAM_ENGINE_LOCAL_TO_WORLD_NO_SCALE);
 		SystemParameters.Add(SYS_PARAM_ENGINE_WORLD_TO_LOCAL_NO_SCALE);
 
-		SystemParameters.Add(SYS_PARAM_ENGINE_MIN_DIST_TO_CAMERA);
+		SystemParameters.Add(SYS_PARAM_ENGINE_LOD_DISTANCE);
 		SystemParameters.Add(SYS_PARAM_ENGINE_TIME_SINCE_RENDERED);
 
 		SystemParameters.Add(SYS_PARAM_ENGINE_EXECUTION_STATE);
@@ -146,7 +146,7 @@ void FNiagaraConstants::Init()
 		SystemStrMap.Add(SYS_PARAM_ENGINE_WORLD_TO_LOCAL_NO_SCALE, LOCTEXT("WorldToLocalNoScaleDesc", "Owning component's world space to local space transform matrix with scaling removed."));
 
 		SystemStrMap.Add(SYS_PARAM_ENGINE_TIME_SINCE_RENDERED, LOCTEXT("TimeSinceRendered", "The time in seconds that have passed since this system was last rendered."));
-		SystemStrMap.Add(SYS_PARAM_ENGINE_MIN_DIST_TO_CAMERA, LOCTEXT("MinDistanceToCamera", "The distance from the owner component to the nearest local player viewpoint."));
+		SystemStrMap.Add(SYS_PARAM_ENGINE_LOD_DISTANCE, LOCTEXT("LODDistance", "The distance to use in LOD scaling/culling. Typically, the distance to the nearest camera."));
 		SystemStrMap.Add(SYS_PARAM_ENGINE_EXECUTION_STATE, LOCTEXT("ExecutionState", "The execution state of the systems owner. Takes precedence over the systems internal execution state."));
 
 		SystemStrMap.Add(SYS_PARAM_ENGINE_EXEC_COUNT, LOCTEXT("ExecCountDesc", "The index of this particle in the read buffer."));
diff --git a/Engine/Plugins/FX/Niagara/Source/Niagara/Private/NiagaraEmitter.cpp b/Engine/Plugins/FX/Niagara/Source/Niagara/Private/NiagaraEmitter.cpp
--- a/Engine/Plugins/FX/Niagara/Source/Niagara/Private/NiagaraEmitter.cpp
+++ b/Engine/Plugins/FX/Niagara/Source/Niagara/Private/NiagaraEmitter.cpp
@@ -476,9 +476,8 @@ UNiagaraScript* UNiagaraEmitter::GetScript(ENiagaraScriptUsage Usage, FGuid Usag
 	return nullptr;
 }
 
-bool UNiagaraEmitter::IsAllowedByDetailLevel()const
+bool UNiagaraEmitter::IsAllowedByDetailLevel(int32 DetailLevel)const
 {
-	int32 DetailLevel = INiagaraModule::GetDetailLevel();
 	if ((bUseMinDetailLevel && DetailLevel < MinDetailLevel) || (bUseMaxDetailLevel && DetailLevel > MaxDetailLevel))
 	{
 		return false;
diff --git a/Engine/Plugins/FX/Niagara/Source/Niagara/Private/NiagaraEmitterInstance.cpp b/Engine/Plugins/FX/Niagara/Source/Niagara/Private/NiagaraEmitterInstance.cpp
--- a/Engine/Plugins/FX/Niagara/Source/Niagara/Private/NiagaraEmitterInstance.cpp
+++ b/Engine/Plugins/FX/Niagara/Source/Niagara/Private/NiagaraEmitterInstance.cpp
@@ -161,8 +161,9 @@ void FNiagaraEmitterInstance::Init(int32 InEmitterIdx, FName InSystemInstanceNam
 	checkSlow(CachedEmitter);
 	CachedIDName = EmitterHandle.GetIdName();
 
+	int32 DetailLevel = ParentSystemInstance->GetDetailLevel();
 	if (!EmitterHandle.GetIsEnabled()
-		|| !CachedEmitter->IsAllowedByDetailLevel()
+		|| !CachedEmitter->IsAllowedByDetailLevel(DetailLevel)
 		|| (GMaxRHIFeatureLevel != ERHIFeatureLevel::SM5 && GMaxRHIFeatureLevel != ERHIFeatureLevel::ES3_1 && CachedEmitter->SimTarget == ENiagaraSimTarget::GPUComputeSim)  // skip if GPU sim and <SM5. TODO: fall back to CPU sim instead once we have scalability functionality to do so
 		)
 	{
@@ -891,10 +892,9 @@ void FNiagaraEmitterInstance::PreTick()
 
 bool FNiagaraEmitterInstance::WaitForDebugInfo()
 {
-	if (CachedEmitter->SimTarget == ENiagaraSimTarget::GPUComputeSim)
+	FNiagaraComputeExecutionContext* DebugContext = GPUExecContext;
+	if (CachedEmitter->SimTarget == ENiagaraSimTarget::GPUComputeSim && DebugContext)
 	{
-		FNiagaraComputeExecutionContext* DebugContext = GPUExecContext;
-
 		ENQUEUE_RENDER_COMMAND(CaptureCommand)([=](FRHICommandListImmediate& RHICmdList)
 		{
 			Batcher->ProcessDebugInfo(RHICmdList, GPUExecContext);
diff --git a/Engine/Plugins/FX/Niagara/Source/Niagara/Private/NiagaraModule.cpp b/Engine/Plugins/FX/Niagara/Source/Niagara/Private/NiagaraModule.cpp
--- a/Engine/Plugins/FX/Niagara/Source/Niagara/Private/NiagaraModule.cpp
+++ b/Engine/Plugins/FX/Niagara/Source/Niagara/Private/NiagaraModule.cpp
@@ -80,7 +80,7 @@ FNiagaraVariable INiagaraModule::Engine_Owner_SystemLocalToWorldNoScale;
 FNiagaraVariable INiagaraModule::Engine_Owner_SystemWorldToLocalNoScale;
 
 FNiagaraVariable INiagaraModule::Engine_Owner_TimeSinceRendered;
-FNiagaraVariable INiagaraModule::Engine_Owner_MinDistanceToCamera;
+FNiagaraVariable INiagaraModule::Engine_Owner_LODDistance;
 
 FNiagaraVariable INiagaraModule::Engine_Owner_ExecutionState;
 
@@ -187,7 +187,7 @@ void INiagaraModule::StartupModule()
 	Engine_Owner_SystemWorldToLocalNoScale = FNiagaraVariable(FNiagaraTypeDefinition::GetMatrix4Def(), TEXT("Engine.Owner.SystemWorldToLocalNoScale"));
 
 	Engine_Owner_TimeSinceRendered = FNiagaraVariable(FNiagaraTypeDefinition::GetFloatDef(), TEXT("Engine.Owner.TimeSinceRendered"));
-	Engine_Owner_MinDistanceToCamera = FNiagaraVariable(FNiagaraTypeDefinition::GetFloatDef(), TEXT("Engine.Owner.MinDistanceToCamera"));
+	Engine_Owner_LODDistance = FNiagaraVariable(FNiagaraTypeDefinition::GetFloatDef(), TEXT("Engine.Owner.LODDistance"));
 
 	Engine_Owner_ExecutionState = FNiagaraVariable(FNiagaraTypeDefinition::GetExecutionStateEnum(), TEXT("Engine.Owner.ExecutionState"));
 	
diff --git a/Engine/Plugins/FX/Niagara/Source/Niagara/Private/NiagaraPreviewGrid.cpp b/Engine/Plugins/FX/Niagara/Source/Niagara/Private/NiagaraPreviewGrid.cpp
new file mode 100644
--- /dev/null
+++ b/Engine/Plugins/FX/Niagara/Source/Niagara/Private/NiagaraPreviewGrid.cpp
@@ -0,0 +1,362 @@
+// Copyright 1998-2019 Epic Games, Inc. All Rights Reserved.
+
+#include "NiagaraPreviewGrid.h"
+#include "UObject/ConstructorHelpers.h"
+#include "Engine/BlueprintGeneratedClass.h"
+#include "NiagaraComponent.h"
+
+#include "Components/TextRenderComponent.h"
+#include "Components/StaticMeshComponent.h"
+#if WITH_EDITOR
+#include "Components/ArrowComponent.h"
+#include "Components/BillboardComponent.h"
+#endif
+
+#if WITH_EDITOR
+void UNiagaraPreviewAxis::PostEditChangeProperty(struct FPropertyChangedEvent& PropertyChangedEvent)
+{
+	Super::PostEditChangeProperty(PropertyChangedEvent);
+
+	//Trigger a regeneration of preview grid.
+	if (ANiagaraPreviewGrid* PreviewGrid = GetTypedOuter<ANiagaraPreviewGrid>())
+	{
+		PreviewGrid->PostEditChange();
+	}
+}
+#endif
+
+//////////////////////////////////////////////////////////////////////////
+
+ANiagaraPreviewGrid::ANiagaraPreviewGrid(const FObjectInitializer& ObjectInitializer)
+	: Super(ObjectInitializer)
+	, System(nullptr)
+	, ResetMode(ENiagaraPreviewGridResetMode::Never)
+	, PreviewClass(nullptr)
+	, SpacingX(250.0f)
+	, SpacingY(250.0f)
+	, NumX(0)
+	, NumY(0)
+	, bPreviewDirty(true)
+{
+	PrimaryActorTick.bCanEverTick = true;
+	PrimaryActorTick.bStartWithTickEnabled = true;
+
+	static ConstructorHelpers::FClassFinder<ANiagaraPreviewBase> DefaultPreviewClassBP(TEXT("/Niagara/Blueprints/NiagaraPreview"));
+	if (DefaultPreviewClassBP.Class != NULL)
+	{
+		PreviewClass = DefaultPreviewClassBP.Class;
+	}
+
+#if WITH_EDITORONLY_DATA
+	RootComponent = CreateEditorOnlyDefaultSubobject<USceneComponent>(TEXT("SceneComp"));
+	SpriteComponent = CreateEditorOnlyDefaultSubobject<UBillboardComponent>(TEXT("Sprite"));
+	ArrowComponent = CreateEditorOnlyDefaultSubobject<UArrowComponent>(TEXT("ArrowComponent0"));
+
+	if (!IsRunningCommandlet())
+	{
+		// Structure to hold one-time initialization
+		struct FConstructorStatics
+		{
+			ConstructorHelpers::FObjectFinderOptional<UTexture2D> SpriteTextureObject;
+			FName ID_Effects;
+			FText NAME_Effects;
+			FConstructorStatics()
+				: SpriteTextureObject(TEXT("/Niagara/Icons/S_ParticleSystem"))
+				, ID_Effects(TEXT("Effects"))
+				, NAME_Effects(NSLOCTEXT("SpriteCategory", "Effects", "Effects"))
+			{
+			}
+		};
+		static FConstructorStatics ConstructorStatics;
+
+
+		if (SpriteComponent)
+		{
+			SpriteComponent->Sprite = ConstructorStatics.SpriteTextureObject.Get();
+			SpriteComponent->RelativeScale3D = FVector(0.5f, 0.5f, 0.5f);
+			SpriteComponent->bHiddenInGame = false;
+			SpriteComponent->bIsScreenSizeScaled = true;
+			SpriteComponent->SpriteInfo.Category = ConstructorStatics.ID_Effects;
+			SpriteComponent->SpriteInfo.DisplayName = ConstructorStatics.NAME_Effects;
+			SpriteComponent->SetupAttachment(RootComponent);
+			SpriteComponent->bReceivesDecals = false;
+		}
+
+		if (ArrowComponent)
+		{
+			ArrowComponent->ArrowColor = FColor(0, 255, 128);
+
+			ArrowComponent->ArrowSize = 1.5f;
+			ArrowComponent->bHiddenInGame = false;
+			ArrowComponent->bTreatAsASprite = true;
+			ArrowComponent->bIsScreenSizeScaled = true;
+			ArrowComponent->SpriteInfo.Category = ConstructorStatics.ID_Effects;
+			ArrowComponent->SpriteInfo.DisplayName = ConstructorStatics.NAME_Effects;
+			ArrowComponent->SetupAttachment(RootComponent);
+			ArrowComponent->bAbsoluteScale = true;
+		}
+	}
+#endif // WITH_EDITORONLY_DATA
+}
+
+void ANiagaraPreviewGrid::PostLoad()
+{
+	Super::PostLoad();
+
+}
+
+void ANiagaraPreviewGrid::BeginDestroy()
+{
+	Super::BeginDestroy();
+	DestroyPreviews();
+}
+
+#if WITH_EDITOR
+void ANiagaraPreviewGrid::PostEditChangeProperty(struct FPropertyChangedEvent& PropertyChangedEvent)
+{
+	Super::PostEditChangeProperty(PropertyChangedEvent);
+
+	bPreviewDirty = true;
+}
+#endif
+
+void ANiagaraPreviewGrid::ActivatePreviews(bool bReset)
+{
+	if (bPreviewDirty)
+	{
+		bPreviewDirty = false;
+		GeneratePreviews();
+	}
+
+	for (int32 X = 0; X < NumX; ++X)
+	{
+		float XLocation = X / NumX - 1;
+
+		for (int32 Y = 0; Y < NumY; ++Y)
+		{
+			float YLocation = Y / NumY - 1;
+
+			int32 PreviewIdx = PreviewIndex(X, Y);
+			UChildActorComponent* PreviewComp = PreviewComponents[PreviewIdx];
+			ANiagaraPreviewBase* PreviewActor = CastChecked<ANiagaraPreviewBase>(PreviewComp->GetChildActor());
+
+			TArray<UNiagaraComponent*, TInlineAllocator<4>> NiagaraComponents;
+			PreviewActor->GetComponents<UNiagaraComponent>(NiagaraComponents);
+			for (UNiagaraComponent* Component : NiagaraComponents)
+			{
+				Component->Activate(true);
+			}
+		}
+	}
+}
+
+void ANiagaraPreviewGrid::DeactivatePreviews()
+{
+	for (int32 X = 0; X < NumX; ++X)
+	{
+		float XLocation = X / NumX - 1;
+
+		for (int32 Y = 0; Y < NumY; ++Y)
+		{
+			float YLocation = Y / NumY - 1;
+
+			int32 PreviewIdx = PreviewIndex(X, Y);
+			UChildActorComponent* PreviewComp = PreviewComponents[PreviewIdx];
+			ANiagaraPreviewBase* PreviewActor = CastChecked<ANiagaraPreviewBase>(PreviewComp->GetChildActor());
+
+			TArray<UNiagaraComponent*, TInlineAllocator<4>> NiagaraComponents;
+			PreviewActor->GetComponents<UNiagaraComponent>(NiagaraComponents);
+			for (UNiagaraComponent* Component : NiagaraComponents)
+			{
+				Component->Deactivate();
+			}
+		}
+	}
+}
+
+void ANiagaraPreviewGrid::TickActor(float DeltaTime, enum ELevelTick TickType, FActorTickFunction& ThisTickFunction)
+{
+	Super::TickActor(DeltaTime, TickType, ThisTickFunction);
+
+	if (bPreviewDirty)
+	{
+		bPreviewDirty = false;
+		GeneratePreviews();
+	}
+	TickPreviews();
+}
+
+void ANiagaraPreviewGrid::DestroyPreviews()
+{
+	for (UChildActorComponent* Preview : PreviewComponents)
+	{
+		if (Preview)
+		{
+			Preview->DestroyComponent();
+		}
+	}
+	PreviewComponents.Reset();
+	NumX = 0;
+	NumY = 0;
+}
+
+void ANiagaraPreviewGrid::GeneratePreviews()
+{
+	DestroyPreviews();
+
+	if (System && PreviewAxisX && PreviewAxisY && PreviewClass)
+	{
+		System->ConditionalPostLoad();
+		PreviewClass->ConditionalPostLoad();
+
+		NumX = PreviewAxisX->Num();
+		NumY = PreviewAxisY->Num();
+
+		PreviewComponents.SetNum(NumX * NumY);
+
+		if (NumX > 0 && NumY > 0)
+		{
+			for (int32 X = 0; X < NumX; ++X)
+			{
+				float XLocation = X / NumX - 1;
+
+				for (int32 Y = 0; Y < NumY; ++Y)
+				{
+					int32 Index = PreviewIndex(X, Y);
+
+					UChildActorComponent* PreviewComp = NewObject<UChildActorComponent>(this, *FString::Printf(TEXT("%s[%d][%d]"), *GetName(), X, Y), RF_Transient);
+					PreviewComponents[Index] = PreviewComp;
+					PreviewComp->SetChildActorClass(PreviewClass);
+					PreviewComp->bAutoActivate = false;
+					PreviewComp->AttachToComponent(RootComponent, FAttachmentTransformRules::KeepRelativeTransform);
+					PreviewComp->RegisterComponentWithWorld(RootComponent->GetWorld());
+					ANiagaraPreviewBase* PreviewActor = CastChecked<ANiagaraPreviewBase>(PreviewComp->GetChildActor());
+
+					FAttachmentTransformRules AttachRules(EAttachmentRule::SnapToTarget, EAttachmentRule::SnapToTarget, EAttachmentRule::KeepRelative, false);
+					PreviewActor->AttachToComponent(PreviewComp, AttachRules);
+					PreviewActor->SetActorRelativeLocation(FVector(SpacingX * X, SpacingY * Y, 0.0f));
+					PreviewActor->PrimaryActorTick.bCanEverTick = true;
+					PreviewActor->PrimaryActorTick.bStartWithTickEnabled = true;
+					PreviewActor->SetSystem(System);
+				}
+			}
+		}
+	}
+}
+
+void ANiagaraPreviewGrid::TickPreviews()
+{
+	if (System && PreviewAxisX && PreviewAxisY)
+	{
+		bool bAllInactive = true;
+		for (int32 X = 0; X < NumX; ++X)
+		{
+			float XLocation = X / NumX - 1;
+
+			for (int32 Y = 0; Y < NumY; ++Y)
+			{
+				float YLocation = Y / NumY - 1;
+
+				int32 PreviewIdx = PreviewIndex(X, Y);
+
+				UChildActorComponent* PreviewComp = PreviewComponents[PreviewIdx];
+				ANiagaraPreviewBase* PreviewActor = CastChecked<ANiagaraPreviewBase>(PreviewComp->GetChildActor());
+
+				TArray<FStringFormatArg> Args;
+				Args.Add(FStringFormatArg(PreviewAxisX->GetClass()->GetName()));
+				Args.Add(FStringFormatArg(X));
+				FString XLabel = FString::Format(TEXT("{1} | X = {0}"), Args);
+
+				Args.Reset();
+				Args.Add(FStringFormatArg(PreviewAxisY->GetClass()->GetName()));
+				Args.Add(FStringFormatArg(Y));
+				FString YLabel = FString::Format(TEXT("{1} | Y = {0}"), Args);
+
+				TArray<UNiagaraComponent*, TInlineAllocator<4>> NiagaraComponents;
+				PreviewActor->GetComponents<UNiagaraComponent>(NiagaraComponents);
+				for (UNiagaraComponent* Component : NiagaraComponents)
+				{
+					//Ensure we're active. Will re-trigger one shot effects.
+					if (!Component->IsActive())
+					{
+						if (ResetMode == ENiagaraPreviewGridResetMode::Individual)
+						{
+							Component->Activate(true);
+						}
+					}
+					else
+					{
+						bAllInactive = false;
+					}
+
+					PreviewAxisX->ApplyToPreview(Component, X, true, XLabel);
+					PreviewAxisY->ApplyToPreview(Component, Y, false, YLabel);
+				}
+
+				PreviewActor->SetLabelText(FText::FromString(XLabel), FText::FromString(YLabel));
+			}
+		}
+
+		if (bAllInactive && ResetMode == ENiagaraPreviewGridResetMode::All)
+		{
+			ActivatePreviews(true);
+		}
+	}
+}
+
+//////////////////////////////////////////////////////////////////////////
+
+void UNiagaraPreviewAxis_InterpParamInt32::ApplyToPreview_Implementation(UNiagaraComponent* PreviewComponent, int32 PreviewIndex, bool bIsXAxis, FString& OutLabelText)
+{
+	check(Count > 0);
+	check(PreviewComponent);
+	float Interp = Count > 1 ? (float)PreviewIndex / (Count - 1) : 1.0f;
+	int32 Val = FMath::Lerp(Min, Max, Interp);
+	PreviewComponent->SetNiagaraVariableFloat(Param, Val);
+	OutLabelText = FString::Printf(TEXT("%s = %d"), *Param, Val);
+}
+void UNiagaraPreviewAxis_InterpParamFloat::ApplyToPreview_Implementation(UNiagaraComponent* PreviewComponent, int32 PreviewIndex, bool bIsXAxis, FString& OutLabelText)
+{
+	check(Count > 0);
+	check(PreviewComponent);
+	float Interp = Count > 1 ? (float)PreviewIndex / (Count - 1) : 1.0f;
+	float Val = FMath::Lerp(Min, Max, Interp);
+	PreviewComponent->SetNiagaraVariableFloat(Param, Val);
+	OutLabelText = FString::Printf(TEXT("%s = %g"), *Param, Val);
+}
+void UNiagaraPreviewAxis_InterpParamVector2D::ApplyToPreview_Implementation(UNiagaraComponent* PreviewComponent, int32 PreviewIndex, bool bIsXAxis, FString& OutLabelText)
+{
+	check(Count > 0);
+	check(PreviewComponent);
+	float Interp = Count > 1 ? (float)PreviewIndex / (Count - 1) : 1.0f;
+	FVector2D Val = FMath::Lerp(Min, Max, Interp);
+	PreviewComponent->SetNiagaraVariableVec2(Param, Val);
+	OutLabelText = FString::Printf(TEXT("%s = {%g, %g}"), *Param, Val.X, Val.Y);
+}
+void UNiagaraPreviewAxis_InterpParamVector::ApplyToPreview_Implementation(UNiagaraComponent* PreviewComponent, int32 PreviewIndex, bool bIsXAxis, FString& OutLabelText)
+{
+	check(Count > 0);
+	check(PreviewComponent);
+	float Interp = Count > 1 ? (float)PreviewIndex / (Count - 1) : 1.0f;
+	FVector Val = FMath::Lerp(Min, Max, Interp);
+	PreviewComponent->SetNiagaraVariableVec3(Param, Val);
+	OutLabelText = FString::Printf(TEXT("%s = {%g, %g, %g}"), *Param, Val.X, Val.Y, Val.Z);
+}
+void UNiagaraPreviewAxis_InterpParamVector4::ApplyToPreview_Implementation(UNiagaraComponent* PreviewComponent, int32 PreviewIndex, bool bIsXAxis, FString& OutLabelText)
+{
+	check(Count > 0);
+	check(PreviewComponent);
+	float Interp = Count > 1 ? (float)PreviewIndex / (Count - 1) : 1.0f;
+	FVector4 Val = FMath::Lerp(Min, Max, Interp);
+	PreviewComponent->SetNiagaraVariableVec4(Param, Val);
+	OutLabelText = FString::Printf(TEXT("%s = {%g, %g, %g, %g}"), *Param, Val.X, Val.Y, Val.Z, Val.W);
+}
+void UNiagaraPreviewAxis_InterpParamLinearColor::ApplyToPreview_Implementation(UNiagaraComponent* PreviewComponent, int32 PreviewIndex, bool bIsXAxis, FString& OutLabelText)
+{
+	check(Count > 0);
+	check(PreviewComponent);
+	float Interp = Count > 1 ? (float)PreviewIndex / (Count - 1) : 1.0f;
+	FLinearColor Val = FMath::Lerp(Min, Max, Interp);
+	PreviewComponent->SetNiagaraVariableLinearColor(Param, Val);
+	OutLabelText = FString::Printf(TEXT("%s = {%g, %g, %g, %g}"), *Param, Val.R, Val.G, Val.B, Val.A);
+}
\ No newline at end of file
diff --git a/Engine/Plugins/FX/Niagara/Source/Niagara/Private/NiagaraSystemInstance.cpp b/Engine/Plugins/FX/Niagara/Source/Niagara/Private/NiagaraSystemInstance.cpp
--- a/Engine/Plugins/FX/Niagara/Source/Niagara/Private/NiagaraSystemInstance.cpp
+++ b/Engine/Plugins/FX/Niagara/Source/Niagara/Private/NiagaraSystemInstance.cpp
@@ -28,7 +28,7 @@ DECLARE_CYCLE_STAT(TEXT("System PreSimulateTick [CNC]"), STAT_NiagaraSystemPreSi
 DECLARE_CYCLE_STAT(TEXT("System Instance Tick [CNC]"), STAT_NiagaraSystemInstanceTick, STATGROUP_Niagara);
 
  
-/** Safety time to allow for the LastRenderTime coming back from the RT. */
+/** Safety time to allow for the LastRenderTime coming back from the RT. This is overkill but that's ok.*/
 static float GLastRenderTimeSafetyBias = 0.1f;
 static FAutoConsoleVariableRef CVarLastRenderTimeSafetyBias(
 	TEXT("fx.LastRenderTimeSafetyBias"),
@@ -666,7 +666,7 @@ void FNiagaraSystemInstance::ReInitInternal()
 	InstanceParameters.AddParameter(SYS_PARAM_ENGINE_INV_DELTA_TIME, true, false);
 	InstanceParameters.AddParameter(SYS_PARAM_ENGINE_TIME_SINCE_RENDERED, true, false);
 	InstanceParameters.AddParameter(SYS_PARAM_ENGINE_EXECUTION_STATE, true, false);
-	InstanceParameters.AddParameter(SYS_PARAM_ENGINE_MIN_DIST_TO_CAMERA, true, false);
+	InstanceParameters.AddParameter(SYS_PARAM_ENGINE_LOD_DISTANCE, true, false);
 	InstanceParameters.AddParameter(SYS_PARAM_ENGINE_SYSTEM_NUM_EMITTERS, true, false);
 	InstanceParameters.AddParameter(SYS_PARAM_ENGINE_SYSTEM_NUM_EMITTERS_ALIVE, true, false);
 	InstanceParameters.AddParameter(SYS_PARAM_ENGINE_SYSTEM_AGE);
@@ -725,7 +725,7 @@ void FNiagaraSystemInstance::ReInitInternal()
 	OwnerEngineTimeParam.Init(InstanceParameters, SYS_PARAM_ENGINE_TIME);
 	OwnerEngineRealtimeParam.Init(InstanceParameters, SYS_PARAM_ENGINE_REAL_TIME);
 
-	OwnerMinDistanceToCameraParam.Init(InstanceParameters, SYS_PARAM_ENGINE_MIN_DIST_TO_CAMERA);
+	OwnerLODDistanceParam.Init(InstanceParameters, SYS_PARAM_ENGINE_LOD_DISTANCE);
 	SystemNumEmittersParam.Init(InstanceParameters, SYS_PARAM_ENGINE_SYSTEM_NUM_EMITTERS);
 	SystemNumEmittersAliveParam.Init(InstanceParameters, SYS_PARAM_ENGINE_SYSTEM_NUM_EMITTERS_ALIVE);
 
@@ -1003,6 +1003,17 @@ bool FNiagaraSystemInstance::GetPerInstanceDataAndOffsets(void*& OutData, uint32
 	return DataInterfaceInstanceDataOffsets.Num() != 0;
 }
 
+int32 FNiagaraSystemInstance::GetDetailLevel()const
+{
+	int32 DetailLevel = INiagaraModule::GetDetailLevel();
+#if WITH_EDITOR
+	if (Component && Component->bEnablePreviewDetailLevel)
+	{
+		DetailLevel = Component->PreviewDetailLevel;
+	}
+#endif
+	return DetailLevel;
+}
 
 void FNiagaraSystemInstance::TickDataInterfaces(float DeltaSeconds, bool bPostSimulate)
 {
@@ -1041,6 +1052,52 @@ void FNiagaraSystemInstance::TickDataInterfaces(float DeltaSeconds, bool bPostSi
 	}
 }
 
+float FNiagaraSystemInstance::GetLODDistance()
+{
+	check(Component);
+	FVector CurrPos = Component->GetComponentLocation();
+#if WITH_EDITOR
+	if (Component->bEnablePreviewLODDistance)
+	{
+		return Component->PreviewLODDistance;
+	}
+	else
+#endif
+	{
+		UWorld* World = Component->GetWorld();
+		TArray<FVector, TInlineAllocator<8> > PlayerViewLocations;
+		//TODO: Pull this up into the world manager at least. Doing this for every instance each frame is pointless.
+		if (World->GetPlayerControllerIterator())
+		{
+			for (FConstPlayerControllerIterator Iterator = World->GetPlayerControllerIterator(); Iterator; ++Iterator)
+			{
+				APlayerController* PlayerController = Iterator->Get();
+				if (PlayerController && PlayerController->IsLocalPlayerController())
+				{
+					FVector* POVLoc = new(PlayerViewLocations) FVector;
+					FRotator POVRotation;
+					PlayerController->GetPlayerViewPoint(*POVLoc, POVRotation);
+				}
+			}
+		}
+		else
+		{
+			PlayerViewLocations.Append(World->ViewLocationsRenderedLastFrame);
+		}
+
+		float LODDistanceSqr = (PlayerViewLocations.Num() ? FMath::Square(WORLD_MAX) : 0.0f);
+		for (const FVector& ViewLocation : PlayerViewLocations)
+		{
+			const float DistanceToEffectSqr = FVector(ViewLocation - CurrPos).SizeSquared();
+			if (DistanceToEffectSqr < LODDistanceSqr)
+			{
+				LODDistanceSqr = DistanceToEffectSqr;
+			}
+		}
+		return FMath::Sqrt(LODDistanceSqr);
+	}
+}
+
 void FNiagaraSystemInstance::TickInstanceParameters(float DeltaSeconds)
 {
 	//TODO: Create helper binding objects to avoid the search in set parameter value.
@@ -1078,36 +1135,8 @@ void FNiagaraSystemInstance::TickInstanceParameters(float DeltaSeconds)
 	UWorld* World = Component->GetWorld();
 	if (World != NULL)
 	{
-		TArray<FVector, TInlineAllocator<8> > PlayerViewLocations;
-		if (World->GetPlayerControllerIterator())
-		{
-			for (FConstPlayerControllerIterator Iterator = World->GetPlayerControllerIterator(); Iterator; ++Iterator)
-			{
-				APlayerController* PlayerController = Iterator->Get();
-				if (PlayerController && PlayerController->IsLocalPlayerController())
-				{
-					FVector* POVLoc = new(PlayerViewLocations) FVector;
-					FRotator POVRotation;
-					PlayerController->GetPlayerViewPoint(*POVLoc, POVRotation);
-				}
-			}
-		}
-		else
-		{
-			PlayerViewLocations.Append(World->ViewLocationsRenderedLastFrame);
-		}
-
-		float LODDistanceSqr = (PlayerViewLocations.Num() ? FMath::Square(WORLD_MAX) : 0.0f);
-		for (const FVector& ViewLocation : PlayerViewLocations)
-		{
-			const float DistanceToEffectSqr = FVector(ViewLocation - CurrPos).SizeSquared();
-			if (DistanceToEffectSqr < LODDistanceSqr)
-			{
-				LODDistanceSqr = DistanceToEffectSqr;
-			}
-		}
-		OwnerMinDistanceToCameraParam.SetValue(FMath::Sqrt(LODDistanceSqr));
-
+		float LODDistance = GetLODDistance();
+		OwnerLODDistanceParam.SetValue(LODDistance);
 		OwnerEngineTimeParam.SetValue(World->TimeSeconds);
 		OwnerEngineRealtimeParam.SetValue(World->RealTimeSeconds);
 	}
@@ -1134,7 +1163,9 @@ void FNiagaraSystemInstance::TickInstanceParameters(float DeltaSeconds)
 	SystemNumEmittersAliveParam.SetValue(NumAlive);
 
 	check(World);
-	float SafeTimeSinceRendererd = FMath::Max(0.0f, World->GetTimeSeconds() - Component->LastRenderTime - GLastRenderTimeSafetyBias);
+	float WorldTime = World->GetTimeSeconds();
+	//Bias the LastRenderTime slightly to account for any delay as it's written by the RT.
+	float SafeTimeSinceRendererd = FMath::Max(0.0f, WorldTime - Component->LastRenderTime - GLastRenderTimeSafetyBias);
 	SystemTimeSinceRenderedParam.SetValue(SafeTimeSinceRendererd);
 	
 	OwnerExecutionStateParam.SetValue((int32)RequestedExecutionState);
@@ -1315,6 +1346,9 @@ void FNiagaraSystemInstance::PreSimulateTick(float DeltaSeconds)
 {
 	SCOPE_CYCLE_COUNTER(STAT_NiagaraSystemPreSimulateTick);
 	TickInstanceParameters(DeltaSeconds);
+
+	Age += DeltaSeconds;
+	TickCount += 1;
 }
 
 void FNiagaraSystemInstance::PostSimulateTick(float DeltaSeconds)
@@ -1346,9 +1380,6 @@ void FNiagaraSystemInstance::PostSimulateTick(float DeltaSeconds)
 		FNiagaraEmitterInstance& Inst = Emitters[EmitterIdx].Get();
 		Inst.Tick(DeltaSeconds);
 	}
-
-	Age += DeltaSeconds;
-	TickCount += 1;
 }
 
 #if WITH_EDITORONLY_DATA
diff --git a/Engine/Plugins/FX/Niagara/Source/Niagara/Public/NiagaraComponent.h b/Engine/Plugins/FX/Niagara/Source/Niagara/Public/NiagaraComponent.h
--- a/Engine/Plugins/FX/Niagara/Source/Niagara/Public/NiagaraComponent.h
+++ b/Engine/Plugins/FX/Niagara/Source/Niagara/Public/NiagaraComponent.h
@@ -277,8 +277,17 @@ public:
 	virtual void PostLoad();
 #if WITH_EDITOR
 	virtual void PreEditChange(UProperty* PropertyAboutToChange) override;
-
 	virtual void PostEditChangeProperty(FPropertyChangedEvent& PropertyChangedEvent) override;
+#endif
+	//~ End UObject Interface
+
+#if WITH_EDITOR
+
+	UFUNCTION(BlueprintCallable, Category = Preview, meta = (Keywords = "preview detail level scalability"))
+	void SetPreviewDetailLevel(bool bEnablePreviewDetailLevel, int32 PreviewDetailLevel);
+
+	UFUNCTION(BlueprintCallable, Category = Preview, meta = (Keywords = "preview LOD Distance scalability"))
+	void SetPreviewLODDistance(bool bEnablePreviewLODDistance, float PreviewLODDistance);
 
 	void PostLoadNormalizeOverrideNames();
 	bool IsParameterValueOverriddenLocally(const FName& InParamName);
@@ -356,6 +365,19 @@ public:
 	void SetAutoAttachmentParameters(USceneComponent* Parent, FName SocketName, EAttachmentRule LocationRule, EAttachmentRule RotationRule, EAttachmentRule ScaleRule);
 
 #if WITH_EDITORONLY_DATA
+
+	UPROPERTY(EditAnywhere, Category = Preview, Transient, meta=(EditCondition=bEnablePreviewDetailLevel))
+	int32 PreviewDetailLevel;
+
+	UPROPERTY(EditAnywhere, Category = Preview, Transient, meta=(EditCondition= bEnablePreviewLODDistance))
+	float PreviewLODDistance;
+
+	UPROPERTY(EditAnywhere, Category = Preview, Transient)
+	uint32 bEnablePreviewDetailLevel : 1;
+
+	UPROPERTY(EditAnywhere, Category = Preview, Transient)
+	uint32 bEnablePreviewLODDistance : 1;
+
 	UPROPERTY(EditAnywhere, Category = Compilation)
 	uint32 bWaitForCompilationOnActivate : 1;
 #endif
diff --git a/Engine/Plugins/FX/Niagara/Source/Niagara/Public/NiagaraModule.h b/Engine/Plugins/FX/Niagara/Source/Niagara/Public/NiagaraModule.h
--- a/Engine/Plugins/FX/Niagara/Source/Niagara/Public/NiagaraModule.h
+++ b/Engine/Plugins/FX/Niagara/Source/Niagara/Public/NiagaraModule.h
@@ -135,7 +135,7 @@ public:
 	FORCEINLINE static const FNiagaraVariable&  GetVar_Engine_Owner_SystemWorldToLocalNoScale() { return Engine_Owner_SystemWorldToLocalNoScale; }
 
 	FORCEINLINE static const FNiagaraVariable&  GetVar_Engine_Owner_TimeSinceRendered() { return Engine_Owner_TimeSinceRendered; }
-	FORCEINLINE static const FNiagaraVariable&  GetVar_Engine_Owner_MinDistanceToCamera() { return Engine_Owner_MinDistanceToCamera; }
+	FORCEINLINE static const FNiagaraVariable&  GetVar_Engine_Owner_LODDistance() { return Engine_Owner_LODDistance; }
 
 	FORCEINLINE static const FNiagaraVariable&  GetVar_Engine_Owner_ExecutionState() { return Engine_Owner_ExecutionState; }
 
@@ -231,7 +231,7 @@ private:
 	static FNiagaraVariable Engine_Owner_SystemWorldToLocalNoScale;
 
 	static FNiagaraVariable Engine_Owner_TimeSinceRendered;
-	static FNiagaraVariable Engine_Owner_MinDistanceToCamera;
+	static FNiagaraVariable Engine_Owner_LODDistance;
 	
 	static FNiagaraVariable Engine_Owner_ExecutionState;
 
diff --git a/Engine/Plugins/FX/Niagara/Source/Niagara/Public/NiagaraPreviewGrid.h b/Engine/Plugins/FX/Niagara/Source/Niagara/Public/NiagaraPreviewGrid.h
new file mode 100644
--- /dev/null
+++ b/Engine/Plugins/FX/Niagara/Source/Niagara/Public/NiagaraPreviewGrid.h
@@ -0,0 +1,269 @@
+// Copyright 1998-2019 Epic Games, Inc. All Rights Reserved.
+
+#pragma once
+
+#include "CoreMinimal.h"
+#include "UObject/ObjectMacros.h"
+#include "NiagaraCommon.h"
+#include "GameFramework/Actor.h"
+#include "UObject/TextProperty.h"
+
+
+#include "NiagaraPreviewGrid.generated.h"
+
+//////////////////////////////////////////////////////////////////////////
+
+/** Base actor for preview actors used in UNiagaraPreviewAxis. */
+UCLASS(abstract, Blueprintable, Transient)
+class ANiagaraPreviewBase : public AActor
+{
+	GENERATED_BODY()
+
+public:
+	//AActor Interface
+	virtual bool ShouldTickIfViewportsOnly() const final { return true; }
+	//AActor Interface End
+
+	UFUNCTION(BlueprintCallable, BlueprintImplementableEvent, CallInEditor, Category = "Niagara|Preview")
+	void SetSystem(UNiagaraSystem* InSystem);
+
+	UFUNCTION(BlueprintCallable, BlueprintImplementableEvent, CallInEditor, Category = "Niagara|Preview")
+	void SetLabelText(const FText& InXAxisText, const FText& InYAxisText);
+};
+
+/**
+Base class for all preview axis types.
+NiagaraPreviewGrid uses these to control how many systems to spawn in each axis and how each system varies on that axis.
+C++ Examples are show below. You can also create these as Blueprint classes as show in the Plugin content folder.
+*/
+UCLASS(abstract, EditInlineNew, BlueprintType, Blueprintable)
+class UNiagaraPreviewAxis : public UObject
+{
+	GENERATED_BODY()
+
+private:
+	virtual int32 Num_Implementation() { return 1; }
+	virtual void ApplyToPreview_Implementation(UNiagaraComponent* PreviewComponent, int32 PreviewIndex, bool bIsXAxis, FString& OutLabelText){ }
+public:
+
+	//UObject Interface
+#if WITH_EDITOR
+	virtual void PostEditChangeProperty(struct FPropertyChangedEvent& PropertyChangedEvent)override;
+#endif
+	//UObject Interface END
+	
+	/** Returns the number of previews for this axis. */
+	UFUNCTION(BlueprintNativeEvent)
+	int32 Num();
+	
+	/** Applies this axis for the preview at PreviewIndex on this axis. */
+	UFUNCTION(BlueprintNativeEvent)
+	void ApplyToPreview(UNiagaraComponent* PreviewComponent, int32 PreviewIndex, bool bIsXAxis, FString& OutLabelText);
+};
+
+UCLASS(abstract, EditInlineNew, BlueprintType, Blueprintable)
+class UNiagaraPreviewAxis_InterpParamBase : public UNiagaraPreviewAxis
+{
+	GENERATED_BODY()
+
+protected:
+	UPROPERTY(EditAnywhere, Category = Axis)
+	FString Param;
+	UPROPERTY(EditAnywhere, Category = Axis)
+	int32 Count;
+
+	virtual int32 Num_Implementation() { return Count; }
+};
+
+UCLASS(EditInlineNew, BlueprintType, Blueprintable)
+class UNiagaraPreviewAxis_InterpParamInt32 : public UNiagaraPreviewAxis_InterpParamBase
+{
+	GENERATED_BODY()
+
+private:
+	UPROPERTY(EditAnywhere, Category = Axis)
+	int32 Min;
+	UPROPERTY(EditAnywhere, Category = Axis)
+	int32 Max;
+
+	//UNiagaraPreviewAxis Interface
+	virtual void ApplyToPreview_Implementation(UNiagaraComponent* PreviewComponent, int32 PreviewIndex, bool bIsXAxis, FString& OutLabelText);
+	//UNiagaraPreviewAxis Interface END
+};
+
+UCLASS(EditInlineNew, BlueprintType, Blueprintable)
+class UNiagaraPreviewAxis_InterpParamFloat : public UNiagaraPreviewAxis_InterpParamBase
+{
+	GENERATED_BODY()
+
+private:
+	UPROPERTY(EditAnywhere, Category=Axis)
+	float Min;	
+	UPROPERTY(EditAnywhere, Category=Axis)
+	float Max;
+
+	//UNiagaraPreviewAxis Interface
+	virtual void ApplyToPreview_Implementation(UNiagaraComponent* PreviewComponent, int32 PreviewIndex, bool bIsXAxis, FString& OutLabelText);
+	//UNiagaraPreviewAxis Interface END
+};
+
+UCLASS(EditInlineNew, BlueprintType, Blueprintable)
+class UNiagaraPreviewAxis_InterpParamVector2D : public UNiagaraPreviewAxis_InterpParamBase
+{
+	GENERATED_BODY()
+
+private:
+	UPROPERTY(EditAnywhere, Category = Axis)
+	FVector2D Min;
+	UPROPERTY(EditAnywhere, Category = Axis)
+	FVector2D Max;
+
+	//UNiagaraPreviewAxis Interface
+	virtual void ApplyToPreview_Implementation(UNiagaraComponent* PreviewComponent, int32 PreviewIndex, bool bIsXAxis, FString& OutLabelText);
+	//UNiagaraPreviewAxis Interface END
+};
+
+UCLASS(EditInlineNew, BlueprintType, Blueprintable)
+class UNiagaraPreviewAxis_InterpParamVector : public UNiagaraPreviewAxis_InterpParamBase
+{
+	GENERATED_BODY()
+
+private:
+	UPROPERTY(EditAnywhere, Category = Axis)
+	FVector Min;
+	UPROPERTY(EditAnywhere, Category = Axis)
+	FVector Max;
+
+	//UNiagaraPreviewAxis Interface
+	virtual void ApplyToPreview_Implementation(UNiagaraComponent* PreviewComponent, int32 PreviewIndex, bool bIsXAxis, FString& OutLabelText);
+	//UNiagaraPreviewAxis Interface END
+};
+
+UCLASS(EditInlineNew, BlueprintType, Blueprintable)
+class UNiagaraPreviewAxis_InterpParamVector4 : public UNiagaraPreviewAxis_InterpParamBase
+{
+	GENERATED_BODY()
+
+private:
+	UPROPERTY(EditAnywhere, Category = Axis)
+	FVector4 Min;
+	UPROPERTY(EditAnywhere, Category = Axis)
+	FVector4 Max;
+
+	//UNiagaraPreviewAxis Interface
+	virtual void ApplyToPreview_Implementation(UNiagaraComponent* PreviewComponent, int32 PreviewIndex, bool bIsXAxis, FString& OutLabelText);
+	//UNiagaraPreviewAxis Interface END
+};
+
+UCLASS(EditInlineNew, BlueprintType, Blueprintable)
+class UNiagaraPreviewAxis_InterpParamLinearColor : public UNiagaraPreviewAxis_InterpParamBase
+{
+	GENERATED_BODY()
+
+private:
+	UPROPERTY(EditAnywhere, Category = Axis)
+	FLinearColor Min;
+	UPROPERTY(EditAnywhere, Category = Axis)
+	FLinearColor Max;
+
+	//UNiagaraPreviewAxis Interface
+	virtual void ApplyToPreview_Implementation(UNiagaraComponent* PreviewComponent, int32 PreviewIndex, bool bIsXAxis, FString& OutLabelText);
+	//UNiagaraPreviewAxis Interface END
+};
+
+UENUM()
+enum class ENiagaraPreviewGridResetMode : uint8
+{
+	Never,/** Never resets the previews. */
+	Individual,/** Resets each preview as it completes. */
+	All,/** Resets all previews when all have completed. */
+};
+
+UCLASS()
+class ANiagaraPreviewGrid : public AActor
+{
+	GENERATED_UCLASS_BODY()
+
+public:
+	//UObject Interface
+	virtual void PostLoad()override;
+	virtual void BeginDestroy()override;
+#if WITH_EDITOR
+	virtual void PostEditChangeProperty(struct FPropertyChangedEvent& PropertyChangedEvent)override;
+#endif
+	//UObject Interface End
+
+	//AActor Interface
+	virtual void TickActor(float DeltaTime, enum ELevelTick TickType, FActorTickFunction& ThisTickFunction)override;
+	virtual bool ShouldTickIfViewportsOnly() const override { return true; }
+	//AActor Interface End
+
+	UFUNCTION(BlueprintCallable, Category=Preview)
+	void ActivatePreviews(bool bReset);
+
+	UFUNCTION(BlueprintCallable, Category=Preview)
+	void DeactivatePreviews();
+private:
+
+	void DestroyPreviews();
+	void GeneratePreviews();
+	void TickPreviews();
+
+	FORCEINLINE int32 PreviewIndex(int32 InX, int32 InY)const { return InX * NumY + InY; }
+
+public:
+
+	UPROPERTY(EditAnywhere, Category = Preview)
+	UNiagaraSystem* System;
+	
+	UPROPERTY(EditAnywhere, Category = Preview)
+	ENiagaraPreviewGridResetMode ResetMode;
+
+	/** Object controlling behavior varying on the X axis. */
+	UPROPERTY(EditAnywhere, Category = Preview, Instanced)
+	UNiagaraPreviewAxis* PreviewAxisX;
+
+	/** Object controlling behavior varying on the Y axis. */
+	UPROPERTY(EditAnywhere, Category = Preview, Instanced)
+	UNiagaraPreviewAxis* PreviewAxisY;
+
+	/** Class used to for previews in this grid. */
+	UPROPERTY(EditAnywhere, Category = Preview)
+	TSubclassOf<ANiagaraPreviewBase> PreviewClass;
+
+	//TODO: Have the preview actor/class define the size of the preview to make spacing eaiser.
+	/** The default spacing between previews in X if the axis does not override it. */
+	UPROPERTY(EditAnywhere, Category = Preview)
+	float SpacingX;
+	
+	/** The default spacing between previews if the axis does not override it. */
+	UPROPERTY(EditAnywhere, Category = Preview)
+	float SpacingY;
+	
+private:
+
+	UPROPERTY(transient)
+	int32 NumX;
+	
+	UPROPERTY(transient)
+	int32 NumY;
+
+	UPROPERTY(transient)
+	TArray<UChildActorComponent*> PreviewComponents;
+
+#if WITH_EDITORONLY_DATA
+	// Reference to sprite visualization component
+	UPROPERTY()
+	class UBillboardComponent* SpriteComponent;
+
+	// Reference to arrow visualization component
+	UPROPERTY()
+	class UArrowComponent* ArrowComponent;
+#endif
+
+
+	uint32 bPreviewDirty : 1;
+	uint32 bPreviewActive : 1;
+};
+
+
diff --git a/Engine/Plugins/FX/Niagara/Source/Niagara/Public/NiagaraRenderer.h b/Engine/Plugins/FX/Niagara/Source/Niagara/Public/NiagaraRenderer.h
--- a/Engine/Plugins/FX/Niagara/Source/Niagara/Public/NiagaraRenderer.h
+++ b/Engine/Plugins/FX/Niagara/Source/Niagara/Public/NiagaraRenderer.h
@@ -109,7 +109,7 @@ public:
 		bool bHasDynamicData = HasDynamicData();
 
 		//Always draw so our LastRenderTime is updated. We may not have dynamic data if we're disabled from visibility culling.
-		Result.bDrawRelevance = bHasDynamicData && SceneProxy->IsShown(View) && View->Family->EngineShowFlags.Particles;
+		Result.bDrawRelevance =/* bHasDynamicData && */SceneProxy->IsShown(View) && View->Family->EngineShowFlags.Particles;
 		Result.bShadowRelevance = bHasDynamicData && SceneProxy->IsShadowCast(View);
 		Result.bDynamicRelevance = bHasDynamicData;
 		if (bHasDynamicData && View->Family->EngineShowFlags.Bounds)
diff --git a/Engine/Plugins/FX/Niagara/Source/Niagara/Public/NiagaraRendererMeshes.h b/Engine/Plugins/FX/Niagara/Source/Niagara/Public/NiagaraRendererMeshes.h
--- a/Engine/Plugins/FX/Niagara/Source/Niagara/Public/NiagaraRendererMeshes.h
+++ b/Engine/Plugins/FX/Niagara/Source/Niagara/Public/NiagaraRendererMeshes.h
@@ -54,7 +54,8 @@ public:
 	{
 		FPrimitiveViewRelevance Result;
 		bool bHasDynamicData = HasDynamicData();
-		Result.bDrawRelevance = bHasDynamicData && SceneProxy->IsShown(View) && View->Family->EngineShowFlags.Particles;
+		//Always draw so our LastRenderTime is updated. We may not have dynamic data if we're disabled from visibility culling.
+		Result.bDrawRelevance = /*bHasDynamicData && */SceneProxy->IsShown(View) && View->Family->EngineShowFlags.Particles;
 		Result.bShadowRelevance = bHasDynamicData && SceneProxy->IsShadowCast(View);
 		Result.bDynamicRelevance = bHasDynamicData;
 
diff --git a/Engine/Plugins/FX/Niagara/Source/Niagara/Public/NiagaraSystemInstance.h b/Engine/Plugins/FX/Niagara/Source/Niagara/Public/NiagaraSystemInstance.h
--- a/Engine/Plugins/FX/Niagara/Source/Niagara/Public/NiagaraSystemInstance.h
+++ b/Engine/Plugins/FX/Niagara/Source/Niagara/Public/NiagaraSystemInstance.h
@@ -204,6 +204,7 @@ public:
 
 	NiagaraEmitterInstanceBatcher* GetBatcher() const { return Batcher; }
 
+	int32 GetDetailLevel()const;
 private:
 
 	/** Builds the emitter simulations. */
@@ -227,6 +228,9 @@ private:
 
 	void BindParameterCollections(FNiagaraScriptExecutionContext& ExecContext);
 	
+	/** Calculates the distance to use for distance based LODing / culling. */
+	float GetLODDistance();
+
 	UNiagaraComponent* Component;
 	TSharedPtr<class FNiagaraSystemSimulation, ESPMode::ThreadSafe> SystemSimulation;
 	FBox SystemBounds;
@@ -291,7 +295,7 @@ private:
 	FNiagaraParameterDirectBinding<float> SystemAgeParam;
 	FNiagaraParameterDirectBinding<int32> SystemTickCountParam;
 
-	FNiagaraParameterDirectBinding<float> OwnerMinDistanceToCameraParam;
+	FNiagaraParameterDirectBinding<float> OwnerLODDistanceParam;
 	FNiagaraParameterDirectBinding<int32> SystemNumEmittersParam;
 	FNiagaraParameterDirectBinding<int32> SystemNumEmittersAliveParam;
 
diff --git a/Engine/Plugins/FX/Niagara/Source/NiagaraEditorWidgets/Private/Stack/SNiagaraStackFunctionInputValue.cpp b/Engine/Plugins/FX/Niagara/Source/NiagaraEditorWidgets/Private/Stack/SNiagaraStackFunctionInputValue.cpp
--- a/Engine/Plugins/FX/Niagara/Source/NiagaraEditorWidgets/Private/Stack/SNiagaraStackFunctionInputValue.cpp
+++ b/Engine/Plugins/FX/Niagara/Source/NiagaraEditorWidgets/Private/Stack/SNiagaraStackFunctionInputValue.cpp
@@ -25,6 +25,8 @@
 #include "Modules/ModuleManager.h"
 #include "AssetRegistryModule.h"
 #include "NiagaraParameterCollection.h"
+#include "ViewModels/NiagaraSystemViewModel.h"
+#include "NiagaraSystem.h"
 
 #define LOCTEXT_NAMESPACE "NiagaraStackFunctionInputValue"
 
@@ -440,7 +442,16 @@ FReply SNiagaraStackFunctionInputValue::OnLinkedInputDoubleClicked()
 		UNiagaraParameterCollection* Collection = CastChecked<UNiagaraParameterCollection>(CollectionAsset.GetAsset());
 		if (Collection && Collection->GetNamespace() == *ParamCollection)
 		{
-			FAssetEditorManager::Get().OpenEditorForAsset(Collection);
+			if (UNiagaraParameterCollectionInstance* NPCInst = FunctionInput->GetSystemViewModel()->GetSystem().GetParameterCollectionOverride(Collection))
+			{
+				//If we override this NPC then open the instance.
+				FAssetEditorManager::Get().OpenEditorForAsset(NPCInst);
+			}
+			else
+			{
+				FAssetEditorManager::Get().OpenEditorForAsset(Collection); 
+			}
+			
 			return FReply::Handled();
 		}
 	}
