Fixed up NiagaraPreview BP and code to avoid load/cook issues.

diff --git a/Engine/Plugins/FX/Niagara/Source/Niagara/Private/NiagaraPreviewGrid.cpp b/Engine/Plugins/FX/Niagara/Source/Niagara/Private/NiagaraPreviewGrid.cpp
--- a/Engine/Plugins/FX/Niagara/Source/Niagara/Private/NiagaraPreviewGrid.cpp
+++ b/Engine/Plugins/FX/Niagara/Source/Niagara/Private/NiagaraPreviewGrid.cpp
@@ -41,11 +41,11 @@ ANiagaraPreviewGrid::ANiagaraPreviewGrid(const FObjectInitializer& ObjectInitial
 	PrimaryActorTick.bCanEverTick = true;
 	PrimaryActorTick.bStartWithTickEnabled = true;
 
-	static ConstructorHelpers::FClassFinder<ANiagaraPreviewBase> DefaultPreviewClassBP(TEXT("/Niagara/Blueprints/NiagaraPreview"));
-	if (DefaultPreviewClassBP.Class != NULL)
-	{
-		PreviewClass = DefaultPreviewClassBP.Class;
-	}
+// 	static ConstructorHelpers::FClassFinder<ANiagaraPreviewBase> DefaultPreviewClassBP(TEXT("/Niagara/Blueprints/NiagaraPreview"));
+// 	if (DefaultPreviewClassBP.Class != NULL)
+// 	{
+// 		PreviewClass = DefaultPreviewClassBP.Class;
+// 	}
 
 #if WITH_EDITORONLY_DATA
 	RootComponent = CreateEditorOnlyDefaultSubobject<USceneComponent>(TEXT("SceneComp"));
