Reverting activation change that seemed to break initialization of systems after compilation.

diff --git a/Engine/Plugins/FX/Niagara/Source/Niagara/Private/NiagaraComponent.cpp b/Engine/Plugins/FX/Niagara/Source/Niagara/Private/NiagaraComponent.cpp
--- a/Engine/Plugins/FX/Niagara/Source/Niagara/Private/NiagaraComponent.cpp
+++ b/Engine/Plugins/FX/Niagara/Source/Niagara/Private/NiagaraComponent.cpp
@@ -457,10 +457,7 @@ void UNiagaraComponent::ResetSystem()
 void UNiagaraComponent::ReinitializeSystem()
 {
 	DestroyInstance();
-	if (bIsActive)
-	{
-		Activate(true);
-	}
+	Activate(true);
 }
 
 bool UNiagaraComponent::GetRenderingEnabled() const
