Skeletal Mesh DI interpolation

> Skeletal mesh interface can now optionally interpolate in it's GetSkinnedData function.
> Updated SampleSkeletalMeshSurface module.
> New utility op GetSpawnInterpolation(). Could be useful for other things.

diff --git a/Engine/Plugins/FX/Niagara/Source/Niagara/Classes/NiagaraDataInterfaceSkeletalMesh.h b/Engine/Plugins/FX/Niagara/Source/Niagara/Classes/NiagaraDataInterfaceSkeletalMesh.h
--- a/Engine/Plugins/FX/Niagara/Source/Niagara/Classes/NiagaraDataInterfaceSkeletalMesh.h
+++ b/Engine/Plugins/FX/Niagara/Source/Niagara/Classes/NiagaraDataInterfaceSkeletalMesh.h
@@ -389,10 +389,7 @@ public:
 	template<typename FilterMode, typename AreaWeightingMode>
 	void IsValidTriCoord(FVectorVMContext& Context);
 
-	template<typename SkinningHandlerType, typename TransformHandlerType, typename VertexAccessorType>
-	void GetTriCoordSkinnedData_Opt(FVectorVMContext& Context);
-
-	template<typename SkinningHandlerType, typename TransformHandlerType, typename VertexAccessorType>
+	template<typename SkinningHandlerType, typename TransformHandlerType, typename VertexAccessorType, typename bInterpolated>
 	void GetTriCoordSkinnedData(FVectorVMContext& Context);
 
 	void GetTriCoordColor(FVectorVMContext& Context);
diff --git a/Engine/Plugins/FX/Niagara/Source/Niagara/Private/NDISkeletalMesh_TriangleSampling.cpp b/Engine/Plugins/FX/Niagara/Source/Niagara/Private/NDISkeletalMesh_TriangleSampling.cpp
--- a/Engine/Plugins/FX/Niagara/Source/Niagara/Private/NDISkeletalMesh_TriangleSampling.cpp
+++ b/Engine/Plugins/FX/Niagara/Source/Niagara/Private/NDISkeletalMesh_TriangleSampling.cpp
@@ -28,6 +28,8 @@ static const FName RandomTriCoordName("RandomTriCoord");
 static const FName IsValidTriCoordName("IsValidTriCoord");
 static const FName GetSkinnedTriangleDataName("GetSkinnedTriangleData");
 static const FName GetSkinnedTriangleDataWSName("GetSkinnedTriangleDataWS");
+static const FName GetSkinnedTriangleDataInterpName("GetSkinnedTriangleDataInterpolated");
+static const FName GetSkinnedTriangleDataWSInterpName("GetSkinnedTriangleDataWSInterpolated");
 static const FName GetTriColorName("GetTriColor");
 static const FName GetTriUVName("GetTriUV");
 static const FName GetTriangleCountName("GetFilteredTriangleCount");
@@ -97,6 +99,44 @@ void UNiagaraDataInterfaceSkeletalMesh::GetTriangleSamplingFunctions(TArray<FNia
 		OutFunctions.Add(Sig);
 	}
 
+	{
+		FNiagaraFunctionSignature Sig;
+		Sig.Name = GetSkinnedTriangleDataInterpName;
+		Sig.Inputs.Add(FNiagaraVariable(FNiagaraTypeDefinition(GetClass()), TEXT("SkeletalMesh")));
+		Sig.Inputs.Add(FNiagaraVariable(FNiagaraTypeDefinition(FMeshTriCoordinate::StaticStruct()), TEXT("Coord")));
+		Sig.Inputs.Add(FNiagaraVariable(FNiagaraTypeDefinition::GetFloatDef(), TEXT("Interp")));
+		Sig.Outputs.Add(FNiagaraVariable(FNiagaraTypeDefinition::GetVec3Def(), TEXT("Position")));
+		Sig.Outputs.Add(FNiagaraVariable(FNiagaraTypeDefinition::GetVec3Def(), TEXT("Velocity")));
+		Sig.Outputs.Add(FNiagaraVariable(FNiagaraTypeDefinition::GetVec3Def(), TEXT("Normal")));
+		Sig.Outputs.Add(FNiagaraVariable(FNiagaraTypeDefinition::GetVec3Def(), TEXT("Binormal")));
+		Sig.Outputs.Add(FNiagaraVariable(FNiagaraTypeDefinition::GetVec3Def(), TEXT("Tangent")));
+		Sig.bMemberFunction = true;
+		Sig.bRequiresContext = false;
+#if WITH_EDITORONLY_DATA
+		Sig.Description = LOCTEXT("GetSkinnedDataDesc", "Returns skinning dependant data for the pased MeshTriCoord in local space. Interpolates between previous and current frame. All outputs are optional and you will incur zerp minimal cost if they are not connected.");
+#endif
+		OutFunctions.Add(Sig);
+	}
+
+	{
+		FNiagaraFunctionSignature Sig;
+		Sig.Name = GetSkinnedTriangleDataWSInterpName;
+		Sig.Inputs.Add(FNiagaraVariable(FNiagaraTypeDefinition(GetClass()), TEXT("SkeletalMesh")));
+		Sig.Inputs.Add(FNiagaraVariable(FNiagaraTypeDefinition(FMeshTriCoordinate::StaticStruct()), TEXT("Coord")));
+		Sig.Inputs.Add(FNiagaraVariable(FNiagaraTypeDefinition::GetFloatDef(), TEXT("Interp")));
+		Sig.Outputs.Add(FNiagaraVariable(FNiagaraTypeDefinition::GetVec3Def(), TEXT("Position")));
+		Sig.Outputs.Add(FNiagaraVariable(FNiagaraTypeDefinition::GetVec3Def(), TEXT("Velocity")));
+		Sig.Outputs.Add(FNiagaraVariable(FNiagaraTypeDefinition::GetVec3Def(), TEXT("Normal")));
+		Sig.Outputs.Add(FNiagaraVariable(FNiagaraTypeDefinition::GetVec3Def(), TEXT("Binormal")));
+		Sig.Outputs.Add(FNiagaraVariable(FNiagaraTypeDefinition::GetVec3Def(), TEXT("Tangent")));
+		Sig.bMemberFunction = true;
+		Sig.bRequiresContext = false;
+#if WITH_EDITORONLY_DATA
+		Sig.Description = LOCTEXT("GetSkinnedDataWSDesc", "Returns skinning dependant data for the pased MeshTriCoord in world space. Interpolates between previous and current frame. All outputs are optional and you will incur zerp minimal cost if they are not connected.");
+#endif
+		OutFunctions.Add(Sig);
+	}
+
 	{
 		FNiagaraFunctionSignature Sig;
 		Sig.Name = GetTriColorName;
@@ -175,13 +215,23 @@ void UNiagaraDataInterfaceSkeletalMesh::BindTriangleSamplingFunction(const FVMEx
 	else if (BindingInfo.Name == GetSkinnedTriangleDataName)
 	{
 		check(BindingInfo.GetNumInputs() == 5 && BindingInfo.GetNumOutputs() == 15);
-		TSkinningModeBinder<TNDIExplicitBinder<FNDITransformHandlerNoop, TVertexAccessorBinder<NDI_FUNC_BINDER(UNiagaraDataInterfaceSkeletalMesh, GetTriCoordSkinnedData)>>>::Bind(this, BindingInfo, InstanceData, OutFunc);
+		TSkinningModeBinder<TNDIExplicitBinder<FNDITransformHandlerNoop, TVertexAccessorBinder<TNDIExplicitBinder<TIntegralConstant<bool, false>, NDI_FUNC_BINDER(UNiagaraDataInterfaceSkeletalMesh, GetTriCoordSkinnedData)>>>>::Bind(this, BindingInfo, InstanceData, OutFunc);
 	}
 	else if (BindingInfo.Name == GetSkinnedTriangleDataWSName)
 	{
 		check(BindingInfo.GetNumInputs() == 5 && BindingInfo.GetNumOutputs() == 15);
-		TSkinningModeBinder<TNDIExplicitBinder<FNDITransformHandler, TVertexAccessorBinder<NDI_FUNC_BINDER(UNiagaraDataInterfaceSkeletalMesh, GetTriCoordSkinnedData)>>>::Bind(this, BindingInfo, InstanceData, OutFunc);
- 	}
+		TSkinningModeBinder<TNDIExplicitBinder<FNDITransformHandler, TVertexAccessorBinder<TNDIExplicitBinder<TIntegralConstant<bool, false>, NDI_FUNC_BINDER(UNiagaraDataInterfaceSkeletalMesh, GetTriCoordSkinnedData)>>>>::Bind(this, BindingInfo, InstanceData, OutFunc);
+	}
+	else if (BindingInfo.Name == GetSkinnedTriangleDataInterpName)
+	{
+		check(BindingInfo.GetNumInputs() == 6 && BindingInfo.GetNumOutputs() == 15);
+		TSkinningModeBinder<TNDIExplicitBinder<FNDITransformHandlerNoop, TVertexAccessorBinder<TNDIExplicitBinder<TIntegralConstant<bool, true>, NDI_FUNC_BINDER(UNiagaraDataInterfaceSkeletalMesh, GetTriCoordSkinnedData)>>>>::Bind(this, BindingInfo, InstanceData, OutFunc);
+	}
+	else if (BindingInfo.Name == GetSkinnedTriangleDataWSInterpName)
+	{
+		check(BindingInfo.GetNumInputs() == 6 && BindingInfo.GetNumOutputs() == 15);
+		TSkinningModeBinder<TNDIExplicitBinder< FNDITransformHandler, TVertexAccessorBinder<TNDIExplicitBinder<TIntegralConstant<bool, true>,NDI_FUNC_BINDER(UNiagaraDataInterfaceSkeletalMesh, GetTriCoordSkinnedData)>>>>::Bind(this, BindingInfo, InstanceData, OutFunc);
+	}
 	else if (BindingInfo.Name == GetTriColorName)
 	{
 		check(BindingInfo.GetNumInputs() == 5 && BindingInfo.GetNumOutputs() == 4);
@@ -791,7 +841,7 @@ struct FGetTriCoodSkinnedDataOutputHandler
 	}
 };
 
-template<typename SkinningHandlerType, typename TransformHandlerType, typename VertexAccessorType>
+template<typename SkinningHandlerType, typename TransformHandlerType, typename VertexAccessorType, typename bInterpolated>
 void UNiagaraDataInterfaceSkeletalMesh::GetTriCoordSkinnedData(FVectorVMContext& Context)
 {
 	SCOPE_CYCLE_COUNTER(STAT_NiagaraSkel_Sample);
@@ -802,6 +852,12 @@ void UNiagaraDataInterfaceSkeletalMesh::GetTriCoordSkinnedData(FVectorVMContext&
 	VectorVM::FExternalFuncInputHandler<float> BaryXParam(Context);
 	VectorVM::FExternalFuncInputHandler<float> BaryYParam(Context);
 	VectorVM::FExternalFuncInputHandler<float> BaryZParam(Context);
+	VectorVM::FExternalFuncInputHandler<float> InterpParam;
+
+ 	if(bInterpolated::Value)
+ 	{
+ 		InterpParam.Init(Context);
+ 	}	
 
 	VectorVM::FUserPtrHandler<FNDISkeletalMesh_InstanceData> InstData(Context);
 
@@ -833,10 +889,19 @@ void UNiagaraDataInterfaceSkeletalMesh::GetTriCoordSkinnedData(FVectorVMContext&
 	FVector Prev;
 	FVector Velocity;
 
+	bool bNeedsCurr = bInterpolated::Value || Output.bNeedsPosition || Output.bNeedsVelocity || Output.bNeedsNorm || Output.bNeedsBinorm || Output.bNeedsTangent;
+	bool bNeedsPrev = bInterpolated::Value || Output.bNeedsVelocity;
+
 	for (int32 i = 0; i < Context.NumInstances; ++i)
 	{
 		FMeshTriCoordinate MeshTriCoord(TriParam.GetAndAdvance(), FVector(BaryXParam.GetAndAdvance(), BaryYParam.GetAndAdvance(), BaryZParam.GetAndAdvance()));
 
+		float Interp = 1.0f;
+		if (bInterpolated::Value)
+		{
+			Interp = InterpParam.GetAndAdvance();
+		}
+
 		if(MeshTriCoord.Tri < 0 || MeshTriCoord.Tri > TriMax)
 		{
 			MeshTriCoord = FMeshTriCoordinate(0, FVector(1.0f, 0.0f, 0.0f));
@@ -844,27 +909,54 @@ void UNiagaraDataInterfaceSkeletalMesh::GetTriCoordSkinnedData(FVectorVMContext&
 
 		SkinningHandler.GetTrianlgeIndices(Accessor, MeshTriCoord.Tri, Idx0, Idx1, Idx2);
 
-		if (Output.bNeedsPosition || Output.bNeedsVelocity)
+		if (bNeedsCurr)
 		{
 			SkinningHandler.GetSkinnedTrianglePositions(Accessor, Idx0, Idx1, Idx2, Pos0, Pos1, Pos2);
+		}
+
+		if (bNeedsPrev)
+		{
+			SkinningHandler.GetSkinnedTrianglePreviousPositions(Accessor, Idx0, Idx1, Idx2, Prev0, Prev1, Prev2);
+			Prev = BarycentricInterpolate(MeshTriCoord.BaryCoord, Prev0, Prev1, Prev2);
+			TransformHandler.TransformPosition(Prev, PrevTransform);
+		}
+
+		if (Output.bNeedsPosition || Output.bNeedsVelocity)
+		{
 			Pos = BarycentricInterpolate(MeshTriCoord.BaryCoord, Pos0, Pos1, Pos2);
 			TransformHandler.TransformPosition(Pos, Transform);
+			
+			if (bInterpolated::Value)
+			{
+				Pos = FMath::Lerp(Prev, Pos, Interp);
+			}
+
 			Output.SetPosition(Pos);
 		}
 
 		if (Output.bNeedsVelocity)
 		{
-			SkinningHandler.GetSkinnedTrianglePreviousPositions(Accessor, Idx0, Idx1, Idx2, Prev0, Prev1, Prev2);
-			Prev = BarycentricInterpolate(MeshTriCoord.BaryCoord, Prev0, Prev1, Prev2);
-			TransformHandler.TransformPosition(Prev, PrevTransform);
 			Velocity = (Pos - Prev) * InvDt;
+
+			//No need to handle velocity wrt interpolation as it's based on the prev position anyway
+
 			Output.SetVelocity(Velocity);
 		}
 
+		//TODO: For preskin we should be able to calculate this stuff on the mesh for a perf win in most cases.
 		if (Output.bNeedsNorm)
 		{
 			Normal = ((Pos1 - Pos2) ^ (Pos0 - Pos2)).GetSafeNormal();
 			TransformHandler.TransformVector(Normal, Transform);
+
+			if (bInterpolated::Value)
+			{
+				FVector PrevNormal = ((Prev1 - Prev2) ^ (Prev0 - Prev2)).GetSafeNormal();
+				TransformHandler.TransformVector(PrevNormal, PrevTransform);
+
+				Normal = FMath::VInterpNormalRotationTo(PrevNormal, Normal, Interp, 1.0f);
+			}
+
 			Output.SetNormal(Normal);
 		}
 
@@ -894,18 +986,62 @@ void UNiagaraDataInterfaceSkeletalMesh::GetTriCoordSkinnedData(FVectorVMContext&
 			// Use InverseSlow to catch singular matrices.  Inverse can miss this sometimes.
 			const FMatrix TextureToLocal = ParameterToTexture.Inverse() * ParameterToLocal;
 
-			if (Output.bNeedsBinorm)
+			if (bInterpolated::Value)
 			{
-				Binormal = (TextureToLocal.TransformVector(FVector(1, 0, 0)).GetSafeNormal());
-				TransformHandler.TransformVector(Binormal, Transform);
-				Output.SetBinormal(Binormal);
-			}
+				FMatrix	PrevParameterToLocal(
+					FPlane(Prev1.X - Prev0.X, Prev1.Y - Prev0.Y, Prev1.Z - Prev0.Z, 0),
+					FPlane(Prev2.X - Prev0.X, Prev2.Y - Prev0.Y, Prev2.Z - Prev0.Z, 0),
+					FPlane(Prev0.X, Prev0.Y, Prev0.Z, 0),
+					FPlane(0, 0, 0, 1)
+				);
 
-			if (Output.bNeedsTangent)
+				// Use InverseSlow to catch singular matrices.  Inverse can miss this sometimes.
+				const FMatrix PrevTextureToLocal = ParameterToTexture.Inverse() * PrevParameterToLocal;
+
+				//TODO: For preskin we should be able to calculate this stuff on the mesh for a perf win in most cases.
+				if (Output.bNeedsBinorm)
+				{
+					Binormal = (TextureToLocal.TransformVector(FVector(1, 0, 0)).GetSafeNormal());
+					TransformHandler.TransformVector(Binormal, Transform);
+
+					FVector PrevBinormal = (PrevTextureToLocal.TransformVector(FVector(1, 0, 0)).GetSafeNormal());
+					TransformHandler.TransformVector(PrevBinormal, PrevTransform);
+
+					Binormal = FMath::VInterpNormalRotationTo(Binormal, Binormal, Interp, 1.0f);
+
+					Output.SetBinormal(Binormal);
+				}
+
+				//TODO: For preskin we should be able to calculate this stuff on the mesh for a perf win in most cases.
+				if (Output.bNeedsTangent)
+				{
+					Tangent = (TextureToLocal.TransformVector(FVector(0, 1, 0)).GetSafeNormal());
+					TransformHandler.TransformVector(Tangent, Transform);
+
+					FVector PrevTangent = (TextureToLocal.TransformVector(FVector(0, 1, 0)).GetSafeNormal());
+					TransformHandler.TransformVector(PrevTangent, PrevTransform);
+
+					Tangent = FMath::VInterpNormalRotationTo(PrevTangent, Tangent, Interp, 1.0f);
+
+					Output.SetTangent(Tangent);
+				}
+			}
+			else
 			{
-				Tangent = (TextureToLocal.TransformVector(FVector(0, 1, 0)).GetSafeNormal());
-				TransformHandler.TransformVector(Tangent, Transform);
-				Output.SetTangent(Tangent);
+				if (Output.bNeedsBinorm)
+				{
+					Binormal = (TextureToLocal.TransformVector(FVector(1, 0, 0)).GetSafeNormal());
+					TransformHandler.TransformVector(Binormal, Transform);
+
+					Output.SetBinormal(Binormal);
+				}
+
+				if (Output.bNeedsTangent)
+				{
+					Tangent = (TextureToLocal.TransformVector(FVector(0, 1, 0)).GetSafeNormal());
+					TransformHandler.TransformVector(Tangent, Transform);
+					Output.SetTangent(Tangent);
+				}
 			}
 		}
 	}
diff --git a/Engine/Plugins/FX/Niagara/Source/NiagaraCore/Private/NiagaraCustomVersion.cpp b/Engine/Plugins/FX/Niagara/Source/NiagaraCore/Private/NiagaraCustomVersion.cpp
--- a/Engine/Plugins/FX/Niagara/Source/NiagaraCore/Private/NiagaraCustomVersion.cpp
+++ b/Engine/Plugins/FX/Niagara/Source/NiagaraCore/Private/NiagaraCustomVersion.cpp
@@ -8,4 +8,4 @@ const FGuid FNiagaraCustomVersion::GUID(0xFCF57AFA, 0x50764283, 0xB9A9E658, 0xFF
 // Register the custom version with core
 FCustomVersionRegistration GRegisterNiagaraCustomVersion(FNiagaraCustomVersion::GUID, FNiagaraCustomVersion::LatestVersion, TEXT("NiagaraVer"));
 
-const FGuid FNiagaraCustomVersion::LatestScriptCompileVersion(0x0BE17832, 0xB1DE4421, 0x95243884, 0xBD77C7B3);
+const FGuid FNiagaraCustomVersion::LatestScriptCompileVersion(0xC210475C, 0x562E42A7, 0xADC2E907, 0x745E567F);
diff --git a/Engine/Plugins/FX/Niagara/Source/NiagaraEditor/Private/NiagaraEditorCommon.cpp b/Engine/Plugins/FX/Niagara/Source/NiagaraEditor/Private/NiagaraEditorCommon.cpp
--- a/Engine/Plugins/FX/Niagara/Source/NiagaraEditor/Private/NiagaraEditorCommon.cpp
+++ b/Engine/Plugins/FX/Niagara/Source/NiagaraEditor/Private/NiagaraEditorCommon.cpp
@@ -1191,6 +1191,15 @@ void FNiagaraOpInfo::Init()
 	Op->Outputs.Add(FNiagaraOpInOutInfo(Result, IntType, ResultText, ResultText, Default_IntZero, TEXT("ExecIndex()")));
 	Op->BuildName(TEXT("ExecIndex"), UtilCategoryName);
 	OpInfoMap.Add(Op->Name) = Idx;
+
+	Idx = OpInfos.AddDefaulted();
+	Op = &OpInfos[Idx];
+	Op->Category = UtilCategory;
+	Op->FriendlyName = NSLOCTEXT("NiagaraOpInfo", "SpawnInterp Name", "Spawn Interpolation");
+	Op->Description = NSLOCTEXT("NiagaraOpInfo", "SpawnInterp Desc", "Returns the fraction used for interpolated spawning. i.e. A fraction defining where this particle was spawned between this frame and the last.");
+	Op->Outputs.Add(FNiagaraOpInOutInfo(Result, FNiagaraTypeDefinition::GetFloatDef(), ResultText, ResultText, Default_IntZero, TEXT("GetSpawnInterpolation()")));
+	Op->BuildName(TEXT("SpawnInterpolation"), UtilCategoryName);
+	OpInfoMap.Add(Op->Name) = Idx;
 }
 END_FUNCTION_BUILD_OPTIMIZATION
 //////////////////////////////////////////////////////////////////////////
diff --git a/Engine/Plugins/FX/Niagara/Source/NiagaraEditor/Private/NiagaraHlslTranslator.cpp b/Engine/Plugins/FX/Niagara/Source/NiagaraEditor/Private/NiagaraHlslTranslator.cpp
--- a/Engine/Plugins/FX/Niagara/Source/NiagaraEditor/Private/NiagaraHlslTranslator.cpp
+++ b/Engine/Plugins/FX/Niagara/Source/NiagaraEditor/Private/NiagaraHlslTranslator.cpp
@@ -408,6 +408,22 @@ FString FHlslNiagaraTranslator::GetFunctionDefinitions()
 		// data interface functions that should be defined differently.
 	}
 
+	//Add a few hard coded helper functions in.
+	FwdDeclString += TEXT("float GetSpawnInterpolation();");
+	//Add helper function to get the interpolation factor.
+	if (CompileOptions.TargetUsage == ENiagaraScriptUsage::ParticleSpawnScriptInterpolated)
+	{
+		DefinitionsString += TEXT("float GetSpawnInterpolation()\n{\n");
+		DefinitionsString += TEXT("\treturn HackSpawnInterp;\n");
+		DefinitionsString += TEXT("}\n\n");
+	}
+	else
+	{
+		DefinitionsString += TEXT("float GetSpawnInterpolation()\n{\n");
+		DefinitionsString += TEXT("\treturn 1.0f;");
+		DefinitionsString += TEXT("}\n\n");
+	}
+
 	return FwdDeclString + TEXT("\n") + DefinitionsString;
 }
 
@@ -1171,6 +1187,8 @@ const FNiagaraTranslateResults &FHlslNiagaraTranslator::Translate(const FNiagara
 
 			HlslOutput += TEXT("};\n\n");
 		}
+		
+		HlslOutput += TEXT("static float HackSpawnInterp = 1.0;\n");
 
 		HlslOutput += FunctionDefinitionString;
 
@@ -1437,6 +1455,7 @@ void FHlslNiagaraTranslator::DefineInterpolatedParametersFunction(FString &HlslO
 			HlslOutputString += TEXT("\tfloat InterpSpawn_InvSpawnTime = 1.0 / InterpSpawn_SpawnTime;\n");
 			HlslOutputString += TEXT("\tfloat InterpSpawn_InvUpdateTime = 1.0 / InterpSpawn_UpdateTime;\n");
 			HlslOutputString += TEXT("\tfloat SpawnInterp = InterpSpawn_SpawnTime * Engine_InverseDeltaTime ;\n");
+			HlslOutputString += TEXT("\tHackSpawnInterp = SpawnInterp;\n");
 
 			HlslOutputString += TEXT("\tContext.") + PrevMap + TEXT(".Interpolation.InterpSpawn_Index = InterpSpawn_Index;\n");
 			HlslOutputString += TEXT("\tContext.") + PrevMap + TEXT(".Interpolation.InterpSpawn_SpawnTime = InterpSpawn_SpawnTime;\n");
diff --git a/Engine/Source/Runtime/VectorVM/Public/VectorVM.h b/Engine/Source/Runtime/VectorVM/Public/VectorVM.h
--- a/Engine/Source/Runtime/VectorVM/Public/VectorVM.h
+++ b/Engine/Source/Runtime/VectorVM/Public/VectorVM.h
@@ -388,12 +388,24 @@ namespace VectorVM
 		int32 AdvanceOffset;
 
 	public:
-		FORCEINLINE FExternalFuncInputHandler(FVectorVMContext& Context)
-			: InputOffset(DecodeU16(Context))
-			, InputPtr(IsConstant() ? (T*)(Context.ConstantTable + GetOffset()) : (T*)Context.RegisterTable[GetOffset()])
-			, AdvanceOffset(IsConstant() ? 0 : 1)
+		FExternalFuncInputHandler()
+			: InputOffset(INDEX_NONE)
+			, InputPtr(nullptr)
+			, AdvanceOffset(0)
 		{}
 
+		FORCEINLINE FExternalFuncInputHandler(FVectorVMContext& Context)
+		{
+			Init(Context);
+		}
+
+		void Init(FVectorVMContext& Context)
+		{
+			InputOffset = DecodeU16(Context);
+			InputPtr = IsConstant() ? (T*)(Context.ConstantTable + GetOffset()) : (T*)Context.RegisterTable[GetOffset()];
+			AdvanceOffset = IsConstant() ? 0 : 1;
+		}
+
 		FORCEINLINE bool IsConstant()const { return !IsRegister(); }
 		FORCEINLINE bool IsRegister()const { return (InputOffset & VVM_EXT_FUNC_INPUT_LOC_BIT) != 0; }
 		FORCEINLINE int32 GetOffset()const { return InputOffset & VVM_EXT_FUNC_INPUT_LOC_MASK; }
