Engine Test fixes

diff --git a/Engine/Plugins/FX/Niagara/Source/Niagara/Private/NiagaraPreviewGrid.cpp b/Engine/Plugins/FX/Niagara/Source/Niagara/Private/NiagaraPreviewGrid.cpp
--- a/Engine/Plugins/FX/Niagara/Source/Niagara/Private/NiagaraPreviewGrid.cpp
+++ b/Engine/Plugins/FX/Niagara/Source/Niagara/Private/NiagaraPreviewGrid.cpp
@@ -28,15 +28,21 @@ void UNiagaraPreviewAxis::PostEditChangeProperty(struct FPropertyChangedEvent& P
 //////////////////////////////////////////////////////////////////////////
 
 ANiagaraPreviewGrid::ANiagaraPreviewGrid(const FObjectInitializer& ObjectInitializer)
-	: Super(ObjectInitializer)
-	, System(nullptr)
+	: System(nullptr)
 	, ResetMode(ENiagaraPreviewGridResetMode::Never)
+	, PreviewAxisX(nullptr)
+	, PreviewAxisY(nullptr)
 	, PreviewClass(nullptr)
 	, SpacingX(250.0f)
 	, SpacingY(250.0f)
 	, NumX(0)
 	, NumY(0)
+#if WITH_EDITORONLY_DATA
+	, SpriteComponent(nullptr)
+	, ArrowComponent(nullptr)
+#endif
 	, bPreviewDirty(true)
+	, bPreviewActive(false)
 {
 	PrimaryActorTick.bCanEverTick = true;
 	PrimaryActorTick.bStartWithTickEnabled = true;
