Unifies how we check for for GPU particle capability. Fixes UE-69005.

diff --git a/Engine/Plugins/FX/Niagara/Source/Niagara/Private/NiagaraScript.cpp b/Engine/Plugins/FX/Niagara/Source/Niagara/Private/NiagaraScript.cpp
--- a/Engine/Plugins/FX/Niagara/Source/Niagara/Private/NiagaraScript.cpp
+++ b/Engine/Plugins/FX/Niagara/Source/Niagara/Private/NiagaraScript.cpp
@@ -980,7 +980,7 @@ void UNiagaraScript::BeginCacheForCookedPlatformData(const ITargetPlatform *Targ
 		for (int32 FormatIndex = 0; FormatIndex < DesiredShaderFormats.Num(); FormatIndex++)
 		{
 			const EShaderPlatform LegacyShaderPlatform = ShaderFormatToLegacyShaderPlatform(DesiredShaderFormats[FormatIndex]);
-			if (RHISupportsComputeShaders(LegacyShaderPlatform))
+			if (FNiagaraUtilities::SupportsGPUParticles(LegacyShaderPlatform))
 			{
 				CacheResourceShadersForCooking(LegacyShaderPlatform, CachedScriptResourcesForPlatform);
 			}
@@ -1070,12 +1070,9 @@ void UNiagaraScript::CacheResourceShadersForRendering(bool bRegenerateId, bool b
 
 			//if (ScriptResourcesByFeatureLevel[FeatureLevel])
 			{
-				EShaderPlatform ShaderPlatform = GShaderPlatformForFeatureLevel[CacheFeatureLevel];
-				// SM4 is "in between" these feature levels but we do NOT support GPU particles in SM4. 
-				// So we can't use IsFeatureLevelSupported since SM4 will be seen as past ES3.1 and we do NOT support SM4 GPU particles.
-				// @todo-mattc This check should be rolled into RHISupportsComputeShaders.
-				if (GetMaxSupportedFeatureLevel(ShaderPlatform) == ERHIFeatureLevel::SM5 || GetMaxSupportedFeatureLevel(ShaderPlatform) == ERHIFeatureLevel::ES3_1)
+				if (FNiagaraUtilities::SupportsGPUParticles(CacheFeatureLevel))
 				{
+					EShaderPlatform ShaderPlatform = GShaderPlatformForFeatureLevel[CacheFeatureLevel];
 					ResourceToCache = ScriptResourcesByFeatureLevel[CacheFeatureLevel];
 					CacheShadersForResources(ShaderPlatform, &ScriptResource, true);
 					ScriptResourcesByFeatureLevel[CacheFeatureLevel] = &ScriptResource;
diff --git a/Engine/Plugins/FX/Niagara/Source/Niagara/Public/NiagaraCommon.h b/Engine/Plugins/FX/Niagara/Source/Niagara/Public/NiagaraCommon.h
--- a/Engine/Plugins/FX/Niagara/Source/Niagara/Public/NiagaraCommon.h
+++ b/Engine/Plugins/FX/Niagara/Source/Niagara/Public/NiagaraCommon.h
@@ -6,6 +6,7 @@
 #include "UObject/ObjectMacros.h"
 #include "NiagaraTypes.h"
 #include "UObject/SoftObjectPath.h"
+#include "RHI.h"
 #include "NiagaraCommon.generated.h"
 
 class UNiagaraSystem;
@@ -534,6 +535,29 @@ namespace FNiagaraUtilities
 
 	void CollectScriptDataInterfaceParameters(const UObject& Owner, const TArray<UNiagaraScript*>& Scripts, FNiagaraParameterStore& OutDataInterfaceParameters);
 
+	inline bool SupportsNiagaraRendering(ERHIFeatureLevel::Type FeatureLevel)
+	{
+		return FeatureLevel == ERHIFeatureLevel::SM4 || FeatureLevel == ERHIFeatureLevel::SM5 || FeatureLevel == ERHIFeatureLevel::ES3_1;
+	}
+
+	inline bool SupportsNiagaraRendering(EShaderPlatform ShaderPlatform)
+	{
+		// Note:
+		// IsFeatureLevelSupported does a FeatureLevel < MaxFeatureLevel(ShaderPlatform) so checking ES3.1 support will return true for SM4. I added it explicitly to be clear what we are doing.
+		return IsFeatureLevelSupported(ShaderPlatform, ERHIFeatureLevel::SM5) || IsFeatureLevelSupported(ShaderPlatform, ERHIFeatureLevel::ES3_1) || IsFeatureLevelSupported(ShaderPlatform, ERHIFeatureLevel::SM4);
+	}
+
+	inline bool SupportsGPUParticles(ERHIFeatureLevel::Type FeatureLevel)
+	{
+		EShaderPlatform ShaderPlatform = GShaderPlatformForFeatureLevel[FeatureLevel];
+		return RHISupportsComputeShaders(ShaderPlatform);
+	}
+
+	inline bool SupportsGPUParticles(EShaderPlatform ShaderPlatform)
+	{
+		return RHISupportsComputeShaders(ShaderPlatform);
+	}
+
 #if WITH_EDITORONLY_DATA
 	/**
 	 * Prepares rapid iteration parameter stores for simulation by removing old parameters no longer used by functions, by initializing new parameters
diff --git a/Engine/Plugins/FX/Niagara/Source/NiagaraShader/Public/NiagaraShader.h b/Engine/Plugins/FX/Niagara/Source/NiagaraShader/Public/NiagaraShader.h
--- a/Engine/Plugins/FX/Niagara/Source/NiagaraShader/Public/NiagaraShader.h
+++ b/Engine/Plugins/FX/Niagara/Source/NiagaraShader/Public/NiagaraShader.h
@@ -12,6 +12,7 @@
 #include "SceneView.h"
 #include "Shader.h"
 #include "GlobalShader.h"
+#include "NiagaraCommon.h"
 #include "NiagaraShared.h"
 #include "NiagaraShaderType.h"
 #include "SceneRenderTargetParameters.h"
@@ -48,7 +49,7 @@ public:
 	static bool ShouldCompilePermutation(EShaderPlatform Platform, const FNiagaraShaderScript*  Script)
 	{
 		//@todo - lit materials only 
-		return RHISupportsComputeShaders(Platform);
+		return FNiagaraUtilities::SupportsGPUParticles(Platform);
 	}
 
 
