Added particle GPU sorting for GPU niagara simulation when using sprites and meshes.
Console variables:
- Niagara.GPUSorting : controls whether GPU sorting is enabled.
- Niagara.GPUSorting.CPUToGPUThreshold : controls if CPU particles should be sorted through the GPU path.
- Niagara.GPUSorting.BufferSlack : controls the slack when increasing/decreasing the GPU sort buffers.
- Niagara.GPUSorting.MinBufferSize : minimum size for the GPU sort buffers.
- Niagara.GPUSorting.FrameCountBeforeShrinking : controls how many frame before shrinking GPU sort buffers.

diff --git a/Engine/Plugins/FX/Niagara/Source/Niagara/Classes/NiagaraEmitterInstanceBatcher.h b/Engine/Plugins/FX/Niagara/Source/Niagara/Classes/NiagaraEmitterInstanceBatcher.h
--- a/Engine/Plugins/FX/Niagara/Source/Niagara/Classes/NiagaraEmitterInstanceBatcher.h
+++ b/Engine/Plugins/FX/Niagara/Source/Niagara/Classes/NiagaraEmitterInstanceBatcher.h
@@ -16,24 +16,45 @@ the same VectorVM byte code / compute shader code
 #include "Modules/ModuleManager.h"
 #include "RHIResources.h"
 #include "FXSystem.h"
+#include "NiagaraRendererProperties.h"
+#include "ParticleResources.h"
+#include "Runtime/Engine/Private/Particles/ParticleSortingGPU.h"
+#include "NiagaraGPUSortInfo.h"
 
 struct FNiagaraScriptExecutionContext;
 struct FNiagaraComputeExecutionContext;
 
 #define SIMULATION_QUEUE_COUNT 2
 
+class FNiagaraIndicesVertexBuffer : public FParticleIndicesVertexBuffer
+{
+public:
+
+	FNiagaraIndicesVertexBuffer(int32 InIndexCount);
+
+	FUnorderedAccessViewRHIRef VertexBufferUAV;
+
+	// The allocation count.
+	const int32 IndexCount;
+
+	// Currently used count.
+	int32 UsedIndexCount = 0;
+};
+
 class NiagaraEmitterInstanceBatcher : public FFXSystemInterface
 {
 public:
-	NiagaraEmitterInstanceBatcher()
-		: CurQueueIndex(0)
-	{
-	}
 
-	~NiagaraEmitterInstanceBatcher()
+	NiagaraEmitterInstanceBatcher(ERHIFeatureLevel::Type InFeatureLevel, EShaderPlatform InShaderPlatform)
+		: FeatureLevel(InFeatureLevel)
+		, ShaderPlatform(InShaderPlatform)
+		, CurQueueIndex(0)
+		, ParticleSortBuffers(true)
 	{
 	}
 
+	~NiagaraEmitterInstanceBatcher();
+
 	static const FName Name;
 	virtual FFXSystemInterface* GetInterface(const FName& InName) override;
 
@@ -50,7 +71,7 @@ public:
 	virtual void AddVectorField(UVectorFieldComponent* VectorFieldComponent) override {}
 	virtual void RemoveVectorField(UVectorFieldComponent* VectorFieldComponent) override {}
 	virtual void UpdateVectorField(UVectorFieldComponent* VectorFieldComponent) override {}
-	virtual void PreInitViews() override {}
+	virtual void PreInitViews() override;
 	virtual bool UsesGlobalDistanceField() const override { return false; }
 	virtual void PreRender(FRHICommandListImmediate& RHICmdList, const class FGlobalDistanceFieldParameterData* GlobalDistanceFieldParameterData) override {}
 
@@ -78,11 +99,18 @@ public:
 		CurQueueIndex ^= 0x1;
 
 		ExecuteAll(RHICmdList, ViewUniformBuffer);
+		SortGPUParticles(RHICmdList);
 	}
 
 	void ExecuteAll(FRHICommandList &RHICmdList, FUniformBufferRHIParamRef ViewUniformBuffer);
 	void TickSingle(FNiagaraComputeExecutionContext *Context, FRHICommandList &RHICmdList, FUniformBufferRHIParamRef ViewUniformBuffer) const;
 
+	int32 AddSortedGPUSimulation(const FNiagaraGPUSortInfo& SortInfo);
+	void SortGPUParticles(FRHICommandListImmediate& RHICmdList);
+	void ResolveParticleSortBuffers(FRHICommandListImmediate& RHICmdList, int32 ResultBufferIndex);
+
+	const FParticleIndicesVertexBuffer& GetGPUSortedBuffer() const { return SortedVertexBuffers.Last(); }
+
 	void ProcessDebugInfo(FRHICommandList &RHICmdList, const FNiagaraComputeExecutionContext *Context) const;
 
 	void SetPrevDataStrideParams(const FNiagaraDataSet *Set, FNiagaraShader *Shader, FRHICommandList &RHICmdList) const;
@@ -106,8 +134,25 @@ public:
 	void ClearIndexBufferCur(FRHICommandList &RHICmdList, FNiagaraComputeExecutionContext *Context) const;
 	void ResolveDatasetWrites(FRHICommandList &RHICmdList, FNiagaraComputeExecutionContext *Context) const;
 	void ResizeCurrentBuffer(FRHICommandList &RHICmdList, FNiagaraComputeExecutionContext *Context, uint32 NewNumInstances, uint32 PrevNumInstances) const;
+
 private:
 
+	/** Feature level of this effects system */
+	ERHIFeatureLevel::Type FeatureLevel;
+	/** Shader platform that will be rendering this effects system */
+	EShaderPlatform ShaderPlatform;
+
 	uint32 CurQueueIndex;
+
 	TArray<FNiagaraComputeExecutionContext*> SimulationQueue[SIMULATION_QUEUE_COUNT];
+
+	// Number of particle to sort this frame.
+	int32 SortedParticleCount = 0;
+	int32 NumFramesRequiringShrinking = 0;
+	TArray<FNiagaraGPUSortInfo> SimulationsToSort;
+	FParticleSortBuffers ParticleSortBuffers;
+
+	// The result of the GPU sort. Each next element replace the previous.
+	// The last entry is used to transfer the result of the ParticleSortBuffers.
+	TIndirectArray<FNiagaraIndicesVertexBuffer> SortedVertexBuffers;
 };
diff --git a/Engine/Plugins/FX/Niagara/Source/Niagara/Classes/NiagaraGPUSortInfo.h b/Engine/Plugins/FX/Niagara/Source/Niagara/Classes/NiagaraGPUSortInfo.h
new file mode 100644
--- /dev/null
+++ b/Engine/Plugins/FX/Niagara/Source/Niagara/Classes/NiagaraGPUSortInfo.h
@@ -0,0 +1,45 @@
+// Copyright 1998-2019 Epic Games, Inc. All Rights Reserved.
+
+/*==============================================================================
+NiagaraGPUSortInfo.h: GPU particle sorting helper
+==============================================================================*/
+#pragma once
+
+#include "CoreMinimal.h"
+#include "RHIResources.h"
+#include "NiagaraGPUSortInfo.generated.h"
+
+UENUM()
+enum class ENiagaraSortMode : uint8
+{
+	/** Perform no additional sorting prior to rendering.*/
+	None,
+	/** Sort by depth to the camera's near plane.*/
+	ViewDepth,
+	/** Sort by distance to the camera's origin.*/
+	ViewDistance,
+	/** Custom sorting according to a per particle attribute. Lower values are rendered before higher values. */
+	CustomAscending,
+	/** Custom sorting according to a per particle attribute. Higher values are rendered before lower values. */
+	CustomDecending,
+};
+
+struct FNiagaraGPUSortInfo
+{
+	// The number of particles in the system.
+	int32 ParticleCount = 0;
+	// How the particles should be sorted.
+	ENiagaraSortMode SortMode = ENiagaraSortMode::None;
+	// On which attribute to base the sorting
+	int32 SortAttributeOffset = INDEX_NONE;
+	// The data buffer that holds the particle attributes.
+	FShaderResourceViewRHIRef ParticleDataFloatSRV;
+	uint32 FloatDataOffset = 0;
+	uint32 FloatDataStride = 0;
+	// The actual GPU sim particle count. Needed to get an exact match on the index list.
+	FShaderResourceViewRHIRef GPUParticleCountSRV;
+	uint32 GPUParticleCountOffset = INDEX_NONE;
+	// View data.
+	FVector ViewOrigin;
+	FVector ViewDirection;
+};
\ No newline at end of file
diff --git a/Engine/Plugins/FX/Niagara/Source/Niagara/Private/NiagaraComponent.cpp b/Engine/Plugins/FX/Niagara/Source/Niagara/Private/NiagaraComponent.cpp
--- a/Engine/Plugins/FX/Niagara/Source/Niagara/Private/NiagaraComponent.cpp
+++ b/Engine/Plugins/FX/Niagara/Source/Niagara/Private/NiagaraComponent.cpp
@@ -114,6 +114,7 @@ FNiagaraSceneProxy::FNiagaraSceneProxy(const UNiagaraComponent* InComponent)
 		//UE_LOG(LogNiagara, Warning, TEXT("FNiagaraSceneProxy %p"), this);
 
 		bAlwaysHasVelocity = true;
+		Batcher = SystemInst->GetBatcher();
 	}
 }
 
diff --git a/Engine/Plugins/FX/Niagara/Source/Niagara/Private/NiagaraEmitterInstanceBatcher.cpp b/Engine/Plugins/FX/Niagara/Source/Niagara/Private/NiagaraEmitterInstanceBatcher.cpp
--- a/Engine/Plugins/FX/Niagara/Source/Niagara/Private/NiagaraEmitterInstanceBatcher.cpp
+++ b/Engine/Plugins/FX/Niagara/Source/Niagara/Private/NiagaraEmitterInstanceBatcher.cpp
@@ -6,23 +6,37 @@
 #include "RHIGPUReadback.h"
 #include "NiagaraStats.h"
 #include "NiagaraShader.h"
+#include "NiagaraSortingGPU.h"
 #include "ShaderParameterUtils.h"
 #include "SceneUtils.h"
 #include "ShaderParameterUtils.h"
 #include "ClearQuad.h"
+#include "Runtime/Engine/Private/GPUSort.h"
 
 DECLARE_CYCLE_STAT(TEXT("GPU Emitter Dispatch [RT]"), STAT_NiagaraGPUSimTick_RT, STATGROUP_Niagara);
 DECLARE_CYCLE_STAT(TEXT("GPU Data Readback [RT]"), STAT_NiagaraGPUReadback_RT, STATGROUP_Niagara);
 DECLARE_CYCLE_STAT(TEXT("Allocate GPU Readback Data [RT]"), STAT_NiagaraAllocateGPUReadback_RT, STATGROUP_Niagara);
 DECLARE_FLOAT_COUNTER_STAT(TEXT("Niagara GPU Sim"), STAT_GPU_NiagaraSim, STATGROUP_GPU);
 DECLARE_DWORD_COUNTER_STAT(TEXT("# GPU Particles"), STAT_NiagaraGPUParticles, STATGROUP_Niagara);
+DECLARE_DWORD_COUNTER_STAT(TEXT("# GPU Sorted Particles"), STAT_NiagaraGPUSortedParticles, STATGROUP_Niagara);
+DECLARE_DWORD_COUNTER_STAT(TEXT("# GPU Sorted Buffers"), STAT_NiagaraGPUSortedBuffers, STATGROUP_Niagara);
 DECLARE_DWORD_COUNTER_STAT(TEXT("Readback latency (frames)"), STAT_NiagaraReadbackLatency, STATGROUP_Niagara);
 
 DECLARE_GPU_STAT_NAMED(NiagaraGPUSimulation, TEXT("Niagara GPU Simulation"));
-DECLARE_GPU_STAT_NAMED(NiagaraIndexBufferClear, TEXT("Niagara index buffer clear"));
+DECLARE_GPU_STAT_NAMED(NiagaraGPUSorting, TEXT("Niagara GPU sorting"));
 
 uint32 FNiagaraComputeExecutionContext::TickCounter = 0;
 
+
+FNiagaraIndicesVertexBuffer::FNiagaraIndicesVertexBuffer(int32 InIndexCount)
+	: IndexCount(InIndexCount)
+{
+	FRHIResourceCreateInfo CreateInfo;
+	VertexBufferRHI = RHICreateVertexBuffer((uint32)IndexCount * sizeof(int32), BUF_Static | BUF_ShaderResource | BUF_UnorderedAccess, CreateInfo);
+	VertexBufferSRV = RHICreateShaderResourceView(VertexBufferRHI, sizeof(int32), PF_R32_SINT);
+	VertexBufferUAV = RHICreateUnorderedAccessView(VertexBufferRHI, PF_R32_SINT);
+}
+
 const FName NiagaraEmitterInstanceBatcher::Name(TEXT("NiagaraEmitterInstanceBatcher"));
 
 FFXSystemInterface* NiagaraEmitterInstanceBatcher::GetInterface(const FName& InName)
@@ -30,6 +44,11 @@ FFXSystemInterface* NiagaraEmitterInstanceBatcher::GetInterface(const FName& InN
 	return InName == Name ? this : nullptr;
 }
 
+NiagaraEmitterInstanceBatcher::~NiagaraEmitterInstanceBatcher()
+{
+	ParticleSortBuffers.ReleaseRHI();
+}
+
 void NiagaraEmitterInstanceBatcher::Queue(FNiagaraComputeExecutionContext *ExecContext)
 {
 	//UE_LOG(LogNiagara, Warning, TEXT("Submitted!"));
@@ -157,6 +176,175 @@ void NiagaraEmitterInstanceBatcher::TickSingle(FNiagaraComputeExecutionContext *
 	check(Context->MainDataSet->HasDatasetIndices());
 }
 
+void NiagaraEmitterInstanceBatcher::PreInitViews()
+{
+	SortedParticleCount = 0;
+	SimulationsToSort.Reset();
+
+	for (FNiagaraIndicesVertexBuffer& SortedVertexBuffer : SortedVertexBuffers)
+	{
+		SortedVertexBuffer.UsedIndexCount = 0;
+	}
+}
+
+int32 NiagaraEmitterInstanceBatcher::AddSortedGPUSimulation(const FNiagaraGPUSortInfo& SortInfo)
+{
+	const int32 ResultOffset = SortedParticleCount;
+	SimulationsToSort.Add(SortInfo);
+
+	SortedParticleCount += SortInfo.ParticleCount;
+
+	if (!SortedVertexBuffers.Num())
+	{
+		SortedVertexBuffers.Add(new FNiagaraIndicesVertexBuffer(FMath::Max(GNiagaraGPUSortingMinBufferSize, (int32)(SortedParticleCount * GNiagaraGPUSortingBufferSlack))));
+	}
+	// If we don't fit anymore, reallocate to a bigger size.
+	else if (SortedParticleCount > SortedVertexBuffers.Last().IndexCount)
+	{
+		SortedVertexBuffers.Add(new FNiagaraIndicesVertexBuffer((int32)(SortedParticleCount * GNiagaraGPUSortingBufferSlack)));
+	}
+
+	// Keep track of the last used index, which is also the first used index of next entry
+	// if we need to increase the size of SortedVertexBuffers. Used in FNiagaraCopyIntBufferRegionCS
+	SortedVertexBuffers.Last().UsedIndexCount = SortedParticleCount;
+
+	return ResultOffset;
+}
+
+void NiagaraEmitterInstanceBatcher::SortGPUParticles(FRHICommandListImmediate& RHICmdList)
+{
+	if (SortedParticleCount > 0 && SortedVertexBuffers.Num() > 0 && SimulationsToSort.Num() && GNiagaraGPUSortingBufferSlack > 1.f)
+	{
+		SCOPED_GPU_STAT(RHICmdList, NiagaraGPUSorting);
+
+		ensure(SortedVertexBuffers.Last().IndexCount >= SortedParticleCount);
+
+		// The particle sort buffer must be able to hold all the particles.
+		if (SortedVertexBuffers.Last().IndexCount != ParticleSortBuffers.GetSize())
+		{
+			ParticleSortBuffers.ReleaseRHI();
+			ParticleSortBuffers.SetBufferSize(SortedVertexBuffers.Last().IndexCount);
+			ParticleSortBuffers.InitRHI();
+		}
+
+		SET_DWORD_STAT(STAT_NiagaraGPUSortedParticles, SortedParticleCount);
+		SET_DWORD_STAT(STAT_NiagaraGPUSortedBuffers, ParticleSortBuffers.GetSize());
+
+		// Make sure our outputs are safe to write to.
+		FUnorderedAccessViewRHIParamRef OutputUAVs[2] = { ParticleSortBuffers.GetKeyBufferUAV(), ParticleSortBuffers.GetVertexBufferUAV() };
+		RHICmdList.TransitionResources(EResourceTransitionAccess::ERWBarrier, EResourceTransitionPipeline::EGfxToCompute, OutputUAVs, 2);
+
+		{
+			SCOPED_DRAW_EVENT(RHICmdList, NiagaraSortKeyGen);
+
+			// Bind the shader
+			TShaderMapRef<FNiagaraSortKeyGenCS> KeyGenCS(GetGlobalShaderMap(FeatureLevel));
+			RHICmdList.SetComputeShader(KeyGenCS->GetComputeShader());
+			KeyGenCS->SetOutput(RHICmdList, ParticleSortBuffers.GetKeyBufferUAV(), ParticleSortBuffers.GetVertexBufferUAV());
+
+			int32 OutputOffset = 0;
+			for (int32 EmitterIndex = 0; EmitterIndex < SimulationsToSort.Num(); ++EmitterIndex)
+			{
+				const FNiagaraGPUSortInfo& SortInfo = SimulationsToSort[EmitterIndex];
+				KeyGenCS->SetParameters(RHICmdList, SortInfo, EmitterIndex, OutputOffset);
+				DispatchComputeShader(RHICmdList, *KeyGenCS, FMath::DivideAndRoundUp(SortInfo.ParticleCount, NIAGARA_KEY_GEN_THREAD_COUNT), 1, 1);
+
+				OutputOffset += SortInfo.ParticleCount;
+			}
+			KeyGenCS->UnbindBuffers(RHICmdList);
+		}
+
+		// We may be able to remove this transition if each step isn't dependent on the previous one.
+		RHICmdList.TransitionResources(EResourceTransitionAccess::ERWBarrier, EResourceTransitionPipeline::EComputeToCompute, OutputUAVs, 2);
+
+		// Sort buffers and copy results to index buffers.
+		{
+			const uint32 EmitterKeyMask = (1 << FMath::CeilLogTwo(SimulationsToSort.Num())) - 1;
+			const uint32 KeyMask = (EmitterKeyMask << 16) | 0xFFFF;
+			const int32 ResultBufferIndex = SortGPUBuffers(RHICmdList, ParticleSortBuffers.GetSortBuffers(), 0, KeyMask, SortedParticleCount, FeatureLevel);
+			ResolveParticleSortBuffers(RHICmdList, ResultBufferIndex);
+		}
+
+		// Only keep the last sorted index buffer, which is of the same size as ParticleSortBuffers.GetSize().
+		SortedVertexBuffers.RemoveAt(0, SortedVertexBuffers.Num() - 1);
+
+		// Resize the buffer to maximize next frame.
+		// Those ratio must take into consideration the slack ratio to be stable.
+
+		const int32 RecommandedSize = FMath::Max(GNiagaraGPUSortingMinBufferSize, (int32)(SortedParticleCount * GNiagaraGPUSortingBufferSlack));
+		const float BufferUsage = (float)SortedParticleCount /	(float)ParticleSortBuffers.GetSize();
+
+		if (RecommandedSize < ParticleSortBuffers.GetSize() / GNiagaraGPUSortingBufferSlack)
+		{
+			if (NumFramesRequiringShrinking >= GNiagaraGPUSortingFrameCountBeforeBufferShrinking)
+			{
+				NumFramesRequiringShrinking = 0;
+				ParticleSortBuffers.ReleaseRHI();
+				ParticleSortBuffers.SetBufferSize(0);
+
+				// Add an entry that should fit well for next frame.
+				SortedVertexBuffers.Empty();
+				SortedVertexBuffers.Add(new FNiagaraIndicesVertexBuffer(RecommandedSize));
+			}
+			else
+			{
+				++NumFramesRequiringShrinking;
+			}
+		}
+	}
+	else // If the are no sort task, we don't need any of the sort buffers.
+	{
+		if (NumFramesRequiringShrinking >= GNiagaraGPUSortingFrameCountBeforeBufferShrinking)
+		{
+			NumFramesRequiringShrinking = 0;
+			ParticleSortBuffers.ReleaseRHI();
+			ParticleSortBuffers.SetBufferSize(0);
+			SortedVertexBuffers.Empty();
+		}
+		else
+		{
+			++NumFramesRequiringShrinking;
+		}
+	}
+}
+
+void NiagaraEmitterInstanceBatcher::ResolveParticleSortBuffers(FRHICommandListImmediate& RHICmdList, int32 ResultBufferIndex)
+{
+	SCOPED_DRAW_EVENT(RHICmdList, NiagaraResolveParticleSortBuffers);
+
+#if 0 // TODO use this once working properly!
+	if (SortedVertexBuffers.Num() == 1)
+	{
+		RHICmdList.CopyVertexBuffer(ParticleSortBuffers.GetSortedVertexBufferRHI(ResultBufferIndex), SortedVertexBuffers.Last().VertexBufferRHI);
+		return;
+	}
+#endif
+
+	TShaderMapRef<FNiagaraCopyIntBufferRegionCS> CopyBufferCS(GetGlobalShaderMap(FeatureLevel));
+	RHICmdList.SetComputeShader(CopyBufferCS->GetComputeShader());
+
+	int32 StartingIndex = 0;
+
+	for (int32 Index = 0; Index < SortedVertexBuffers.Num(); Index += NIAGARA_COPY_BUFFER_BUFFER_COUNT)
+	{
+		FUnorderedAccessViewRHIParamRef UAVs[NIAGARA_COPY_BUFFER_BUFFER_COUNT] = {};
+		int32 UsedIndexCounts[NIAGARA_COPY_BUFFER_BUFFER_COUNT] = {};
+
+		const int32 NumBuffers = FMath::Min<int32>(NIAGARA_COPY_BUFFER_BUFFER_COUNT, SortedVertexBuffers.Num() - Index);
+		for (int32 SubIndex = 0; SubIndex < NumBuffers; ++SubIndex)
+		{
+			UAVs[SubIndex] = SortedVertexBuffers[Index + SubIndex].VertexBufferUAV;
+			UsedIndexCounts[SubIndex] = SortedVertexBuffers[Index + SubIndex].UsedIndexCount;
+		}
+
+		CopyBufferCS->SetParameters(RHICmdList, ParticleSortBuffers.GetSortedVertexBufferSRV(ResultBufferIndex), UAVs, UsedIndexCounts, StartingIndex, NumBuffers);
+		DispatchComputeShader(RHICmdList, *CopyBufferCS, FMath::DivideAndRoundUp(UsedIndexCounts[NumBuffers - 1] - StartingIndex, NIAGARA_COPY_BUFFER_THREAD_COUNT), 1, 1);
+		RHICmdList.TransitionResources(EResourceTransitionAccess::ERWBarrier, EResourceTransitionPipeline::EComputeToGfx, UAVs, NumBuffers);
+
+		StartingIndex = UsedIndexCounts[NumBuffers - 1];
+	}
+	CopyBufferCS->UnbindBuffers(RHICmdList);
+}
 
 /* Calculate total number of spawned instances from all spawn events
 */
@@ -183,7 +371,6 @@ void NiagaraEmitterInstanceBatcher::ClearIndexBufferCur(FRHICommandList &RHICmdL
 {
 	FRWBuffer &DatasetIndexBufferWrite = Context->MainDataSet->GetCurDataSetIndices();
 	SCOPED_DRAW_EVENTF(RHICmdList, NiagaraIndexBufferClear, TEXT("Niagara index buffer clear"));
-	SCOPED_GPU_STAT(RHICmdList, NiagaraIndexBufferClear);
 
 	ClearUAV(RHICmdList, DatasetIndexBufferWrite, 0);
 	RHICmdList.TransitionResource(EResourceTransitionAccess::EWritable, EResourceTransitionPipeline::EComputeToCompute, DatasetIndexBufferWrite.UAV);
diff --git a/Engine/Plugins/FX/Niagara/Source/Niagara/Private/NiagaraModule.cpp b/Engine/Plugins/FX/Niagara/Source/Niagara/Private/NiagaraModule.cpp
--- a/Engine/Plugins/FX/Niagara/Source/Niagara/Private/NiagaraModule.cpp
+++ b/Engine/Plugins/FX/Niagara/Source/Niagara/Private/NiagaraModule.cpp
@@ -262,7 +262,7 @@ void INiagaraModule::StartupModule()
 
 	FFXSystemInterface::RegisterCustomFXSystem(NiagaraEmitterInstanceBatcher::Name, FCreateCustomFXSystemDelegate::CreateLambda([](ERHIFeatureLevel::Type InFeatureLevel, EShaderPlatform InShaderPlatform) -> FFXSystemInterface*
 	{
-		return new NiagaraEmitterInstanceBatcher;
+		return new NiagaraEmitterInstanceBatcher(InFeatureLevel, InShaderPlatform);
 	}));
 }
 
diff --git a/Engine/Plugins/FX/Niagara/Source/Niagara/Private/NiagaraRendererMeshes.cpp b/Engine/Plugins/FX/Niagara/Source/Niagara/Private/NiagaraRendererMeshes.cpp
--- a/Engine/Plugins/FX/Niagara/Source/Niagara/Private/NiagaraRendererMeshes.cpp
+++ b/Engine/Plugins/FX/Niagara/Source/Niagara/Private/NiagaraRendererMeshes.cpp
@@ -7,6 +7,8 @@
 #include "NiagaraStats.h"
 #include "Async/ParallelFor.h"
 #include "Engine/StaticMesh.h"
+#include "NiagaraEmitterInstanceBatcher.h"
+#include "NiagaraSortingGPU.h"
 
 DECLARE_CYCLE_STAT(TEXT("Generate Mesh Vertex Data [GT]"), STAT_NiagaraGenMeshVertexData, STATGROUP_Niagara);
 DECLARE_CYCLE_STAT(TEXT("Render Meshes [RT]"), STAT_NiagaraRenderMeshes, STATGROUP_Niagara);
@@ -242,20 +244,48 @@ void NiagaraRendererMeshes::GetDynamicMeshElements(const TArray<const FSceneView
 				}
 
 				//Sort particles if needed.
-				FGlobalDynamicReadBuffer::FAllocation SortedIndices;
 				CollectorResources.VertexFactory.SetSortedIndices(nullptr, 0xFFFFFFFF);
+
+				FNiagaraGPUSortInfo SortInfo;
+				if (View && Properties->SortMode != ENiagaraSortMode::None && (bHasTranslucentMaterials || !Properties->bSortOnlyWhenTranslucent))
+				{
+					SortInfo.ParticleCount = NumInstances;
+					SortInfo.SortMode = Properties->SortMode;
+					SortInfo.SortAttributeOffset = (SortInfo.SortMode == ENiagaraSortMode::CustomAscending || SortInfo.SortMode == ENiagaraSortMode::CustomDecending) ? CustomSortingOffset : PositionOffset;
+					SortInfo.ViewOrigin = View->ViewMatrices.GetViewOrigin();
+					SortInfo.ViewDirection = View->GetViewDirection();
+					if (bLocalSpace)
+					{
+						FMatrix InvTransform = SceneProxy->GetLocalToWorld().InverseFast();
+						SortInfo.ViewOrigin = InvTransform.TransformPosition(SortInfo.ViewOrigin);
+						SortInfo.ViewDirection = InvTransform.TransformVector(SortInfo.ViewDirection);
+					}
+				};
+
 				if (DynamicDataMesh->DataSet->GetSimTarget() == ENiagaraSimTarget::CPUSim)//TODO: Compute shader for sorting gpu sims and larger cpu sims.
 				{
 					check(ParticleData.IsValid());
-					if (bHasTranslucentMaterials || !Properties->bSortOnlyWhenTranslucent)
+					if (SortInfo.SortMode != ENiagaraSortMode::None && SortInfo.SortAttributeOffset != INDEX_NONE)
 					{
-						ENiagaraSortMode SortMode = Properties->SortMode;
-						bool bCustomSortMode = SortMode == ENiagaraSortMode::CustomAscending || SortMode == ENiagaraSortMode::CustomDecending;
-						int32 SortAttributeOffest = bCustomSortMode ? CustomSortingOffset : PositionOffset;
-						if (SortMode != ENiagaraSortMode::None && SortAttributeOffest != INDEX_NONE)
+						if (GNiagaraGPUSorting &&
+							GNiagaraGPUSortingCPUToGPUThreshold != INDEX_NONE &&
+							SortInfo.ParticleCount >= GNiagaraGPUSortingCPUToGPUThreshold &&
+							SceneProxy->GetBatcher())
+						{
+							SortInfo.ParticleDataFloatSRV = ParticleData.ReadBuffer->SRV;
+							SortInfo.FloatDataOffset = ParticleData.FirstIndex / sizeof(float);
+							SortInfo.FloatDataStride = DynamicDataMesh->RTParticleData.GetFloatStride() / sizeof(float);
+							const int32 IndexBufferOffset = SceneProxy->GetBatcher()->AddSortedGPUSimulation(SortInfo);
+							if (IndexBufferOffset != INDEX_NONE)
+							{
+								CollectorResources.VertexFactory.SetSortedIndices(SceneProxy->GetBatcher()->GetGPUSortedBuffer().VertexBufferSRV, IndexBufferOffset);
+							}
+						}
+						else
 						{
+							FGlobalDynamicReadBuffer::FAllocation SortedIndices;
 							SortedIndices = DynamicReadBuffer.AllocateInt32(NumInstances);
-							SortIndices(SortMode, SortAttributeOffest, DynamicDataMesh->RTParticleData, SceneProxy->GetLocalToWorld(), View, SortedIndices);
+							SortIndices(SortInfo.SortMode, SortInfo.SortAttributeOffset, DynamicDataMesh->RTParticleData, SceneProxy->GetLocalToWorld(), View, SortedIndices);
 							CollectorResources.VertexFactory.SetSortedIndices(SortedIndices.ReadBuffer->SRV, SortedIndices.FirstIndex / sizeof(float));
 						}
 					}
@@ -263,6 +293,19 @@ void NiagaraRendererMeshes::GetDynamicMeshElements(const TArray<const FSceneView
 				}
 				else
 				{
+					if (SortInfo.SortMode != ENiagaraSortMode::None && SortInfo.SortAttributeOffset != INDEX_NONE && GNiagaraGPUSorting && SceneProxy->GetBatcher())
+					{
+						SortInfo.ParticleDataFloatSRV = DynamicDataMesh->DataSet->CurrData().GetGPUBufferFloat()->SRV;
+						SortInfo.FloatDataOffset = 0;
+						SortInfo.FloatDataStride = DynamicDataMesh->RTParticleData.GetFloatStride() / sizeof(float);
+						SortInfo.GPUParticleCountSRV = DynamicDataMesh->DataSet->GetCurDataSetIndices().SRV;
+						SortInfo.GPUParticleCountOffset = 1;
+						const int32 IndexBufferOffset = SceneProxy->GetBatcher()->AddSortedGPUSimulation(SortInfo);
+						if (IndexBufferOffset != INDEX_NONE)
+						{
+							CollectorResources.VertexFactory.SetSortedIndices(SceneProxy->GetBatcher()->GetGPUSortedBuffer().VertexBufferSRV, IndexBufferOffset);
+						}
+					}
 					CollectorResources.VertexFactory.SetParticleData(DynamicDataMesh->DataSet->CurrData().GetGPUBufferFloat()->SRV, 0, DynamicDataMesh->DataSet->CurrData().GetFloatStride() / sizeof(float));
 				}
 
diff --git a/Engine/Plugins/FX/Niagara/Source/Niagara/Private/NiagaraRendererSprites.cpp b/Engine/Plugins/FX/Niagara/Source/Niagara/Private/NiagaraRendererSprites.cpp
--- a/Engine/Plugins/FX/Niagara/Source/Niagara/Private/NiagaraRendererSprites.cpp
+++ b/Engine/Plugins/FX/Niagara/Source/Niagara/Private/NiagaraRendererSprites.cpp
@@ -5,6 +5,8 @@
 #include "NiagaraSpriteVertexFactory.h"
 #include "NiagaraDataSet.h"
 #include "NiagaraStats.h"
+#include "NiagaraEmitterInstanceBatcher.h"
+#include "NiagaraSortingGPU.h"
 
 DECLARE_CYCLE_STAT(TEXT("Generate Sprite Vertex Data [GT]"), STAT_NiagaraGenSpriteVertexData, STATGROUP_Niagara);
 DECLARE_CYCLE_STAT(TEXT("Render Sprites [RT]"), STAT_NiagaraRenderSprites, STATGROUP_Niagara);
@@ -33,8 +35,6 @@ public:
 	}
 };
 
-
-
 NiagaraRendererSprites::NiagaraRendererSprites(ERHIFeatureLevel::Type FeatureLevel, UNiagaraRendererProperties *InProps) 
 	: NiagaraRenderer()
 	, PositionOffset(INDEX_NONE)
@@ -252,20 +252,48 @@ void NiagaraRendererSprites::GetDynamicMeshElements(const TArray<const FSceneVie
 
 				//Sort particles if needed.
 				EBlendMode BlendMode = MaterialRenderProxy->GetMaterial(VertexFactory->GetFeatureLevel())->GetBlendMode();
-				FGlobalDynamicReadBuffer::FAllocation SortedIndices;
 				CollectorResources.VertexFactory.SetSortedIndices(nullptr, 0xFFFFFFFF);
+
+				FNiagaraGPUSortInfo SortInfo;
+				if (View && Properties->SortMode != ENiagaraSortMode::None && (BlendMode == BLEND_AlphaComposite || BlendMode == BLEND_Translucent || !Properties->bSortOnlyWhenTranslucent))
+				{
+					SortInfo.ParticleCount = NumInstances;
+					SortInfo.SortMode = Properties->SortMode;
+					SortInfo.SortAttributeOffset = (SortInfo.SortMode == ENiagaraSortMode::CustomAscending || SortInfo.SortMode == ENiagaraSortMode::CustomDecending) ? CustomSortingOffset : PositionOffset;
+					SortInfo.ViewOrigin = View->ViewMatrices.GetViewOrigin();
+					SortInfo.ViewDirection = View->GetViewDirection();
+					if (bLocalSpace)
+					{
+						FMatrix InvTransform = SceneProxy->GetLocalToWorld().InverseFast();
+						SortInfo.ViewOrigin = InvTransform.TransformPosition(SortInfo.ViewOrigin);
+						SortInfo.ViewDirection = InvTransform.TransformVector(SortInfo.ViewDirection);
+					}
+				};
+
 				if (DynamicDataSprites->DataSet->GetSimTarget() == ENiagaraSimTarget::CPUSim)//TODO: Compute shader for sorting gpu sims and larger cpu sims.
 				{
 					check(ParticleData.IsValid());
-					if (BlendMode == BLEND_AlphaComposite || BlendMode == BLEND_Translucent || !Properties->bSortOnlyWhenTranslucent)
+					if (SortInfo.SortMode != ENiagaraSortMode::None && SortInfo.SortAttributeOffset != INDEX_NONE)
 					{
-						ENiagaraSortMode SortMode = Properties->SortMode;
-						bool bCustomSortMode = SortMode == ENiagaraSortMode::CustomAscending || SortMode == ENiagaraSortMode::CustomDecending;
-						int32 SortAttributeOffest = bCustomSortMode ? CustomSortingOffset : PositionOffset;
-						if (SortMode != ENiagaraSortMode::None && SortAttributeOffest != INDEX_NONE)
+						if (GNiagaraGPUSorting && 
+							GNiagaraGPUSortingCPUToGPUThreshold != INDEX_NONE && 
+							SortInfo.ParticleCount >= GNiagaraGPUSortingCPUToGPUThreshold && 
+							SceneProxy->GetBatcher())
 						{
+							SortInfo.ParticleDataFloatSRV = ParticleData.ReadBuffer->SRV;
+							SortInfo.FloatDataOffset = ParticleData.FirstIndex / sizeof(float);
+							SortInfo.FloatDataStride = DynamicDataSprites->RTParticleData.GetFloatStride() / sizeof(float);
+							const int32 IndexBufferOffset = SceneProxy->GetBatcher()->AddSortedGPUSimulation(SortInfo);
+							if (IndexBufferOffset != INDEX_NONE)
+							{
+								CollectorResources.VertexFactory.SetSortedIndices(SceneProxy->GetBatcher()->GetGPUSortedBuffer().VertexBufferSRV, IndexBufferOffset);
+							}
+						}
+						else
+						{
+							FGlobalDynamicReadBuffer::FAllocation SortedIndices;
 							SortedIndices = DynamicReadBuffer.AllocateInt32(NumInstances);
-							SortIndices(SortMode, SortAttributeOffest, DynamicDataSprites->RTParticleData, SceneProxy->GetLocalToWorld(), View, SortedIndices);
+							SortIndices(SortInfo.SortMode, SortInfo.SortAttributeOffset, DynamicDataSprites->RTParticleData, SceneProxy->GetLocalToWorld(), View, SortedIndices);
 							CollectorResources.VertexFactory.SetSortedIndices(SortedIndices.ReadBuffer->SRV, SortedIndices.FirstIndex / sizeof(float));
 						}
 					}
@@ -273,6 +301,21 @@ void NiagaraRendererSprites::GetDynamicMeshElements(const TArray<const FSceneVie
 				}
 				else
 				{
+					if (SortInfo.SortMode != ENiagaraSortMode::None && SortInfo.SortAttributeOffset != INDEX_NONE && GNiagaraGPUSorting && SceneProxy->GetBatcher())
+					{
+						SortInfo.ParticleDataFloatSRV = DynamicDataSprites->DataSet->CurrData().GetGPUBufferFloat()->SRV;
+						SortInfo.FloatDataOffset = 0;
+						SortInfo.FloatDataStride = DynamicDataSprites->RTParticleData.GetFloatStride() / sizeof(float);
+						SortInfo.GPUParticleCountSRV = DynamicDataSprites->DataSet->GetCurDataSetIndices().SRV;
+						SortInfo.GPUParticleCountOffset = 1;
+						const int32 IndexBufferOffset = SceneProxy->GetBatcher()->AddSortedGPUSimulation(SortInfo);
+						if (IndexBufferOffset != INDEX_NONE)
+						{
+							CollectorResources.VertexFactory.SetSortedIndices(SceneProxy->GetBatcher()->GetGPUSortedBuffer().VertexBufferSRV, IndexBufferOffset);
+						}
+
+					}
+
 					CollectorResources.VertexFactory.SetParticleData(DynamicDataSprites->DataSet->CurrData().GetGPUBufferFloat()->SRV, 0, DynamicDataSprites->DataSet->CurrData().GetFloatStride() / sizeof(float));
 				}
 
diff --git a/Engine/Plugins/FX/Niagara/Source/Niagara/Public/NiagaraComponent.h b/Engine/Plugins/FX/Niagara/Source/Niagara/Public/NiagaraComponent.h
--- a/Engine/Plugins/FX/Niagara/Source/Niagara/Public/NiagaraComponent.h
+++ b/Engine/Plugins/FX/Niagara/Source/Niagara/Public/NiagaraComponent.h
@@ -19,6 +19,7 @@ class UNiagaraSystem;
 class UNiagaraParameterCollection;
 class UNiagaraParameterCollectionInstance;
 class FNiagaraSystemSimulation;
+class NiagaraEmitterInstanceBatcher;
 
 // Called when the particle system is done
 DECLARE_DYNAMIC_MULTICAST_DELEGATE_OneParam(FOnNiagaraSystemFinished, class UNiagaraComponent*, PSystem);
@@ -432,6 +433,8 @@ public:
 	/** Sets whether or not this scene proxy should be rendered. */
 	void SetRenderingEnabled(bool bInRenderingEnabled);
 
+	NiagaraEmitterInstanceBatcher* GetBatcher() const { return Batcher; }
+
 #if RHI_RAYTRACING
 	virtual void GetRayTracingGeometryInstances(TArray<FRayTracingGeometryInstanceCollection>& OutInstanceCollections) override;
 	virtual bool IsRayTracingRelevant() const override { return true; }
@@ -468,6 +471,6 @@ private:
 private:
 	//class NiagaraRenderer* EmitterRenderer;
 	TArray<class NiagaraRenderer*>EmitterRenderers;
-
 	bool bRenderingEnabled;
+	NiagaraEmitterInstanceBatcher* Batcher = nullptr;
 };
diff --git a/Engine/Plugins/FX/Niagara/Source/Niagara/Public/NiagaraRendererProperties.h b/Engine/Plugins/FX/Niagara/Source/Niagara/Public/NiagaraRendererProperties.h
--- a/Engine/Plugins/FX/Niagara/Source/Niagara/Public/NiagaraRendererProperties.h
+++ b/Engine/Plugins/FX/Niagara/Source/Niagara/Public/NiagaraRendererProperties.h
@@ -9,25 +9,9 @@
 #include "NiagaraTypes.h"
 #include "NiagaraCommon.h"
 #include "NiagaraMergeable.h"
+#include "NiagaraGPUSortInfo.h"
 #include "NiagaraRendererProperties.generated.h"
 
-
-UENUM()
-enum class ENiagaraSortMode : uint8
-{
-	/** Perform no additional sorting prior to rendering.*/
-	None,
-	/** Sort by depth to the camera's near plane.*/
-	ViewDepth,
-	/** Sort by distance to the camera's origin.*/
-	ViewDistance,
-	/** Custom sorting according to a per particle attribute. Lower values are rendered before higher values. */
-	CustomAscending,
-	/** Custom sorting according to a per particle attribute. Higher values are rendered before lower values. */
-	CustomDecending,
-};
-
-
 /**
 * Emitter properties base class
 * Each EmitterRenderer derives from this with its own class, and returns it in GetProperties; a copy
diff --git a/Engine/Plugins/FX/Niagara/Source/NiagaraVertexFactories/Private/NiagaraSortingGPU.cpp b/Engine/Plugins/FX/Niagara/Source/NiagaraVertexFactories/Private/NiagaraSortingGPU.cpp
new file mode 100644
--- /dev/null
+++ b/Engine/Plugins/FX/Niagara/Source/NiagaraVertexFactories/Private/NiagaraSortingGPU.cpp
@@ -0,0 +1,212 @@
+// Copyright 1998-2019 Epic Games, Inc. All Rights Reserved.
+
+/*==============================================================================
+NiagaraSortingGPU.cpp: Niagara sorting shaders
+==============================================================================*/
+
+#include "NiagaraSortingGPU.h"
+#include "NiagaraGPUSortInfo.h"
+
+
+int32 GNiagaraGPUSorting = 1;
+static FAutoConsoleVariableRef CVarNiagaraGPUSorting(
+	TEXT("Niagara.GPUSorting"),
+	GNiagaraGPUSorting,
+	TEXT("Whether to sort particles on the GPU"),
+	ECVF_Default
+);
+
+int32 GNiagaraGPUSortingCPUToGPUThreshold = -1;
+static FAutoConsoleVariableRef CVarNiagaraGPUSortinCPUToGPUThreshold(
+	TEXT("Niagara.GPUSorting.CPUToGPUThreshold"),
+	GNiagaraGPUSortingCPUToGPUThreshold,
+	TEXT("Particle count to move from a CPU sort to a GPU sort. -1 disables. (default=-1)"),
+	ECVF_Default
+);
+
+float GNiagaraGPUSortingBufferSlack = 2.f;
+static FAutoConsoleVariableRef CVarNiagaraGPUSortingBufferSlack(
+	TEXT("Niagara.GPUSorting.BufferSlack"),
+	GNiagaraGPUSortingBufferSlack,
+	TEXT("Slack ratio when allocating GPU sort buffer (default=2)"),
+	ECVF_Default
+);
+
+int32 GNiagaraGPUSortingMinBufferSize = 8192;
+static FAutoConsoleVariableRef CVarNiagaraGPUSortingMinBufferSize(
+	TEXT("Niagara.GPUSorting.MinBufferSize"),
+	GNiagaraGPUSortingMinBufferSize,
+	TEXT("Minimum GPU sort buffer size, in particles (default=8192)"),
+	ECVF_Default
+);
+
+int32 GNiagaraGPUSortingFrameCountBeforeBufferShrinking = 100;
+static FAutoConsoleVariableRef CVarNiagaraGPUSortingFrameCountBeforeBufferShrinking(
+	TEXT("Niagara.GPUSorting.FrameCountBeforeShrinking"),
+	GNiagaraGPUSortingFrameCountBeforeBufferShrinking,
+	TEXT("Number of consecutive frames where the GPU sort buffer is considered oversized before allowing shrinking. (default=100)"),
+	ECVF_Default
+);
+
+IMPLEMENT_SHADER_TYPE(, FNiagaraSortKeyGenCS, TEXT("/Engine/Private/NiagaraSortKeyGen.usf"), TEXT("GenerateParticleSortKeys"), SF_Compute);
+
+void FNiagaraSortKeyGenCS::ModifyCompilationEnvironment(const FGlobalShaderPermutationParameters& Parameters, FShaderCompilerEnvironment& OutEnvironment)
+{
+	FGlobalShader::ModifyCompilationEnvironment(Parameters, OutEnvironment);
+	OutEnvironment.SetDefine(TEXT("THREAD_COUNT"), NIAGARA_KEY_GEN_THREAD_COUNT);
+	OutEnvironment.SetDefine(TEXT("SORT_VIEW_DEPTH"), (uint8)ENiagaraSortMode::ViewDepth);
+	OutEnvironment.SetDefine(TEXT("SORT_VIEW_DISTANCE"), (uint8)ENiagaraSortMode::ViewDistance);
+	OutEnvironment.SetDefine(TEXT("SORT_CUSTOM_ASCENDING"), (uint8)ENiagaraSortMode::CustomAscending);
+	OutEnvironment.SetDefine(TEXT("SORT_CUSTOM_DESCENDING"), (uint8)ENiagaraSortMode::CustomDecending);
+}
+
+FNiagaraSortKeyGenCS::FNiagaraSortKeyGenCS(const ShaderMetaType::CompiledShaderInitializerType& Initializer)
+	: FGlobalShader(Initializer)
+{
+	NiagaraParticleDataFloat.Bind(Initializer.ParameterMap, TEXT("NiagaraParticleDataFloat"));
+	FloatDataOffset.Bind(Initializer.ParameterMap, TEXT("NiagaraFloatDataOffset"));
+	FloatDataStride.Bind(Initializer.ParameterMap, TEXT("NiagaraFloatDataStride"));
+	GPUParticleCountBuffer.Bind(Initializer.ParameterMap, TEXT("GPUParticleCountBuffer"));
+	ParticleCountParams.Bind(Initializer.ParameterMap, TEXT("ParticleCountParams"));
+	SortParams.Bind(Initializer.ParameterMap, TEXT("SortParams"));
+	CameraPosition.Bind(Initializer.ParameterMap, TEXT("CameraPosition"));
+	CameraDirection.Bind(Initializer.ParameterMap, TEXT("CameraDirection"));
+
+	OutKeys.Bind(Initializer.ParameterMap, TEXT("OutKeys"));
+	OutParticleIndices.Bind(Initializer.ParameterMap, TEXT("OutParticleIndices"));
+}
+
+bool FNiagaraSortKeyGenCS::Serialize(FArchive& Ar)
+{
+	bool bShaderHasOutdatedParameters = FGlobalShader::Serialize(Ar);
+	Ar << NiagaraParticleDataFloat;
+	Ar << FloatDataOffset;
+	Ar << FloatDataStride;
+	Ar << GPUParticleCountBuffer;
+	Ar << ParticleCountParams;
+	Ar << SortParams;
+	Ar << CameraPosition;
+	Ar << CameraDirection;
+	Ar << OutKeys;
+	Ar << OutParticleIndices;
+	return bShaderHasOutdatedParameters;
+}
+
+void FNiagaraSortKeyGenCS::SetOutput(FRHICommandList& RHICmdList, FUnorderedAccessViewRHIParamRef OutKeysUAV, FUnorderedAccessViewRHIParamRef OutIndicesUAV)
+{
+	FComputeShaderRHIParamRef ComputeShaderRHI = GetComputeShader();
+	if (OutKeys.IsBound())
+	{
+		RHICmdList.SetUAVParameter(ComputeShaderRHI, OutKeys.GetBaseIndex(), OutKeysUAV);
+	}
+	if (OutParticleIndices.IsBound())
+	{
+		RHICmdList.SetUAVParameter(ComputeShaderRHI, OutParticleIndices.GetBaseIndex(), OutIndicesUAV);
+	}
+}
+
+void FNiagaraSortKeyGenCS::SetParameters(FRHICommandList& RHICmdList, const FNiagaraGPUSortInfo& SortInfo, int32 EmitterIndex, int32 OutputOffset)
+{
+	FComputeShaderRHIParamRef ComputeShaderRHI = GetComputeShader();
+
+	RHICmdList.SetShaderResourceViewParameter(ComputeShaderRHI, NiagaraParticleDataFloat.GetBaseIndex(), SortInfo.ParticleDataFloatSRV);
+	RHICmdList.SetShaderParameter(ComputeShaderRHI, FloatDataOffset.GetBufferIndex(), FloatDataOffset.GetBaseIndex(), FloatDataOffset.GetNumBytes(), &SortInfo.FloatDataOffset);
+	RHICmdList.SetShaderParameter(ComputeShaderRHI, FloatDataStride.GetBufferIndex(), FloatDataStride.GetBaseIndex(), FloatDataStride.GetNumBytes(), &SortInfo.FloatDataStride);
+
+	RHICmdList.SetShaderResourceViewParameter(ComputeShaderRHI, GPUParticleCountBuffer.GetBaseIndex(), SortInfo.GPUParticleCountSRV);
+	const FUintVector4 ParticleCountParamsValue(SortInfo.ParticleCount, SortInfo.GPUParticleCountOffset, 0, 0);
+	RHICmdList.SetShaderParameter(ComputeShaderRHI, ParticleCountParams.GetBufferIndex(), ParticleCountParams.GetBaseIndex(), ParticleCountParams.GetNumBytes(), &ParticleCountParamsValue);
+
+	// (EmitterKey, OutputOffset, SortMode, SortAttributeOffset)
+	const FUintVector4 SortParamsValue((uint32)EmitterIndex << 16, OutputOffset, (uint8)SortInfo.SortMode, SortInfo.SortAttributeOffset);
+	RHICmdList.SetShaderParameter(ComputeShaderRHI, SortParams.GetBufferIndex(), SortParams.GetBaseIndex(), SortParams.GetNumBytes(), &SortParamsValue);
+
+	RHICmdList.SetShaderParameter(ComputeShaderRHI, CameraPosition.GetBufferIndex(), CameraPosition.GetBaseIndex(), CameraPosition.GetNumBytes(), &SortInfo.ViewOrigin);
+	RHICmdList.SetShaderParameter(ComputeShaderRHI, CameraDirection.GetBufferIndex(), CameraDirection.GetBaseIndex(), CameraDirection.GetNumBytes(), &SortInfo.ViewDirection);
+}
+
+void FNiagaraSortKeyGenCS::UnbindBuffers(FRHICommandList& RHICmdList)
+{
+	FComputeShaderRHIParamRef ComputeShaderRHI = GetComputeShader();
+	if (NiagaraParticleDataFloat.IsBound())
+	{
+		RHICmdList.SetShaderResourceViewParameter(ComputeShaderRHI, NiagaraParticleDataFloat.GetBaseIndex(), FShaderResourceViewRHIParamRef());
+	}
+	if (GPUParticleCountBuffer.IsBound())
+	{
+		RHICmdList.SetShaderResourceViewParameter(ComputeShaderRHI, GPUParticleCountBuffer.GetBaseIndex(), FShaderResourceViewRHIParamRef());
+	}
+	if (OutKeys.IsBound())
+	{
+		RHICmdList.SetUAVParameter(ComputeShaderRHI, OutKeys.GetBaseIndex(), FUnorderedAccessViewRHIParamRef());
+	}
+	if (OutParticleIndices.IsBound())
+	{
+		RHICmdList.SetUAVParameter(ComputeShaderRHI, OutParticleIndices.GetBaseIndex(), FUnorderedAccessViewRHIParamRef());
+	}
+}
+
+//////////////////////////////////////////////////////////////////////////////////////////////////////////////////////
+
+IMPLEMENT_SHADER_TYPE(, FNiagaraCopyIntBufferRegionCS, TEXT("/Engine/Private/NiagaraCopyIntBuffer.usf"), TEXT("MainCS"), SF_Compute);
+
+FNiagaraCopyIntBufferRegionCS::FNiagaraCopyIntBufferRegionCS(const ShaderMetaType::CompiledShaderInitializerType& Initializer)
+	: FGlobalShader(Initializer)
+{
+	CopyParams.Bind(Initializer.ParameterMap, TEXT("CopyParams"));
+	SourceData.Bind(Initializer.ParameterMap, TEXT("SourceData"));
+	DestData[0].Bind(Initializer.ParameterMap, TEXT("DestData0"));
+	DestData[1].Bind(Initializer.ParameterMap, TEXT("DestData1"));
+	DestData[2].Bind(Initializer.ParameterMap, TEXT("DestData2"));
+}
+
+bool FNiagaraCopyIntBufferRegionCS::Serialize(FArchive& Ar)
+{
+	bool bShaderHasOutdatedParameters = FGlobalShader::Serialize(Ar);
+	Ar << CopyParams;
+	Ar << SourceData;
+	Ar << DestData[0];
+	Ar << DestData[1];
+	Ar << DestData[2];
+	return bShaderHasOutdatedParameters;
+}
+
+void FNiagaraCopyIntBufferRegionCS::SetParameters(
+	FRHICommandList& RHICmdList,
+	const FShaderResourceViewRHIParamRef InSourceData,
+	const FUnorderedAccessViewRHIParamRef* InDestDatas,
+	const int32* InUsedIndexCounts,
+	int32 StartingIndex,
+	int32 DestCount)
+{
+	FComputeShaderRHIParamRef ComputeShaderRHI = GetComputeShader();
+	check(DestCount > 0);
+
+	RHICmdList.SetShaderResourceViewParameter(ComputeShaderRHI, SourceData.GetBaseIndex(), InSourceData);
+
+	FUintVector4 CopyParamsValue(StartingIndex, 0, 0, 0);
+	for (int32 Index = 0; Index < DestCount; ++Index)
+	{
+		RHICmdList.SetUAVParameter(ComputeShaderRHI, DestData[Index].GetBaseIndex(), InDestDatas[Index]);
+		CopyParamsValue[Index + 1] = InUsedIndexCounts[Index];
+	}
+	RHICmdList.SetShaderParameter(ComputeShaderRHI, CopyParams.GetBufferIndex(), CopyParams.GetBaseIndex(), CopyParams.GetNumBytes(), &CopyParamsValue);
+}
+
+void FNiagaraCopyIntBufferRegionCS::UnbindBuffers(FRHICommandList& RHICmdList)
+{
+	FComputeShaderRHIParamRef ComputeShaderRHI = GetComputeShader();
+	if (SourceData.IsBound())
+	{
+		RHICmdList.SetShaderResourceViewParameter(ComputeShaderRHI, SourceData.GetBaseIndex(), FShaderResourceViewRHIParamRef());
+	}
+	for (int32 Index = 0; Index < NIAGARA_COPY_BUFFER_BUFFER_COUNT; ++Index)
+	{
+		if (DestData[Index].IsBound())
+		{
+			RHICmdList.SetUAVParameter(ComputeShaderRHI, DestData[Index].GetBaseIndex(), FUnorderedAccessViewRHIParamRef());
+		}
+	}
+}
+
+
diff --git a/Engine/Plugins/FX/Niagara/Source/NiagaraVertexFactories/Public/NiagaraSortingGPU.h b/Engine/Plugins/FX/Niagara/Source/NiagaraVertexFactories/Public/NiagaraSortingGPU.h
new file mode 100644
--- /dev/null
+++ b/Engine/Plugins/FX/Niagara/Source/NiagaraVertexFactories/Public/NiagaraSortingGPU.h
@@ -0,0 +1,134 @@
+// Copyright 1998-2019 Epic Games, Inc. All Rights Reserved.
+
+/*==============================================================================
+NiagaraSortingGPU.h: Niagara sorting shaders
+==============================================================================*/
+
+#pragma once
+
+#include "CoreMinimal.h"
+#include "GlobalShader.h"
+
+struct FNiagaraGPUSortInfo;
+
+extern NIAGARAVERTEXFACTORIES_API int32 GNiagaraGPUSorting;
+extern NIAGARAVERTEXFACTORIES_API int32 GNiagaraGPUSortingCPUToGPUThreshold;
+extern NIAGARAVERTEXFACTORIES_API float GNiagaraGPUSortingBufferSlack;
+extern NIAGARAVERTEXFACTORIES_API int32 GNiagaraGPUSortingMinBufferSize;
+extern NIAGARAVERTEXFACTORIES_API int32 GNiagaraGPUSortingFrameCountBeforeBufferShrinking;
+
+#define NIAGARA_KEY_GEN_THREAD_COUNT 64
+#define NIAGARA_COPY_BUFFER_THREAD_COUNT 64
+#define NIAGARA_COPY_BUFFER_BUFFER_COUNT 3
+
+/**
+ * Compute shader used to generate particle sort keys.
+ */
+class NIAGARAVERTEXFACTORIES_API FNiagaraSortKeyGenCS : public FGlobalShader
+{
+	DECLARE_SHADER_TYPE(FNiagaraSortKeyGenCS,Global);
+
+public:
+
+	static bool ShouldCompilePermutation(const FGlobalShaderPermutationParameters& Parameters)
+	{
+		return RHISupportsComputeShaders(Parameters.Platform);
+	}
+
+	static void ModifyCompilationEnvironment(const FGlobalShaderPermutationParameters& Parameters, FShaderCompilerEnvironment& OutEnvironment);
+
+	/** Default constructor. */
+	FNiagaraSortKeyGenCS() {}
+
+	/** Initialization constructor. */
+	explicit FNiagaraSortKeyGenCS(const ShaderMetaType::CompiledShaderInitializerType& Initializer);
+
+	/** Serialization. */
+	virtual bool Serialize( FArchive& Ar ) override;
+
+	/**
+	 * Set output buffers for this shader.
+	 */
+	void SetOutput(FRHICommandList& RHICmdList, FUnorderedAccessViewRHIParamRef OutKeysUAV, FUnorderedAccessViewRHIParamRef OutIndicesUAV);
+
+	/**
+	 * Set input parameters.
+	 */
+	void SetParameters(FRHICommandList& RHICmdList,	const FNiagaraGPUSortInfo& SortInfo, int32 EmitterIndex, int32 OutputOffset);
+
+	/**
+	 * Unbinds any buffers that have been bound.
+	 */
+	void UnbindBuffers(FRHICommandList& RHICmdList);
+
+private:
+
+	FShaderResourceParameter NiagaraParticleDataFloat;
+	FShaderParameter FloatDataOffset;
+	FShaderParameter FloatDataStride;
+	FShaderResourceParameter GPUParticleCountBuffer;
+	FShaderParameter ParticleCountParams;
+	FShaderParameter SortParams;
+	FShaderParameter CameraPosition;
+	FShaderParameter CameraDirection;
+
+	/** Output key buffer. */
+	FShaderResourceParameter OutKeys;
+	/** Output indices buffer. */
+	FShaderResourceParameter OutParticleIndices;
+};
+
+/**
+ * Compute shader used to copy a reference buffer split in several buffers.
+ * Each buffer contains a segment of the buffer, and each buffer is increasingly bigger to hold the next part.
+ * Could alternatively use DMA copy, if the RHI provided a way to copy buffer regions.
+ */
+class NIAGARAVERTEXFACTORIES_API FNiagaraCopyIntBufferRegionCS : public FGlobalShader
+{
+	DECLARE_SHADER_TYPE(FNiagaraCopyIntBufferRegionCS,Global);
+
+public:
+
+	static bool ShouldCompilePermutation(const FGlobalShaderPermutationParameters& Parameters)
+	{
+		return RHISupportsComputeShaders(Parameters.Platform);
+	}
+
+	static void ModifyCompilationEnvironment(const FGlobalShaderPermutationParameters& Parameters, FShaderCompilerEnvironment& OutEnvironment)
+	{
+		FGlobalShader::ModifyCompilationEnvironment(Parameters, OutEnvironment);
+		OutEnvironment.SetDefine(TEXT("THREAD_COUNT"), NIAGARA_COPY_BUFFER_THREAD_COUNT);
+		OutEnvironment.SetDefine(TEXT("BUFFER_COUNT"), NIAGARA_COPY_BUFFER_BUFFER_COUNT);
+	}
+
+	/** Default constructor. */
+	FNiagaraCopyIntBufferRegionCS() {}
+
+	/** Initialization constructor. */
+	explicit FNiagaraCopyIntBufferRegionCS(const ShaderMetaType::CompiledShaderInitializerType& Initializer);
+
+	/** Serialization. */
+	virtual bool Serialize(FArchive& Ar) override;
+
+	/**
+	 * Set parameters.
+	 */
+	void SetParameters(
+		FRHICommandList& RHICmdList,	
+		const FShaderResourceViewRHIParamRef InSourceData, 
+		const FUnorderedAccessViewRHIParamRef* InDestDatas, 
+		const int32* InUsedIndexCounts, 
+		int32 StartingIndex,
+		int32 DestCount);
+
+	/**
+	 * Unbinds any buffers that have been bound.
+	 */
+	void UnbindBuffers(FRHICommandList& RHICmdList);
+
+private:
+
+	FShaderParameter CopyParams;
+	FShaderResourceParameter SourceData;
+	FShaderResourceParameter DestData[NIAGARA_COPY_BUFFER_BUFFER_COUNT];
+};
diff --git a/Engine/Shaders/Private/NiagaraCopyIntBuffer.usf b/Engine/Shaders/Private/NiagaraCopyIntBuffer.usf
new file mode 100644
--- /dev/null
+++ b/Engine/Shaders/Private/NiagaraCopyIntBuffer.usf
@@ -0,0 +1,53 @@
+// Copyright 1998-2019 Epic Games, Inc. All Rights Reserved.
+
+/*==============================================================================
+	NiagaraCopyIntBuffer.usf: Shader to copy sorted index buffers.
+==============================================================================*/
+
+/*------------------------------------------------------------------------------
+	Compile time parameters:
+		THREAD_COUNT - The number of threads to launch per workgroup.
+		SORT_VIEW_DEPTH
+		SORT_VIEW_DISTANCE
+		SORT_CUSTOM_ASCENDING
+		SORT_CUSTOM_DESCENDING
+------------------------------------------------------------------------------*/
+
+#include "Common.ush"
+
+
+// (StartingIndex, UsedIndexCount0, UsedIndexCount1, UsedIndexCount2)
+uint4 CopyParams;
+
+
+Buffer<int> SourceData;
+
+RWBuffer<int> DestData0;
+RWBuffer<int> DestData1;
+RWBuffer<int> DestData2;
+
+[numthreads(THREAD_COUNT,1,1)]
+void MainCS(uint IndexOffset : SV_DispatchThreadID)
+{
+	const uint StartingIndex = CopyParams.x;
+	const uint UsedIndexCount0 = CopyParams.y;
+	const uint UsedIndexCount1 = CopyParams.z;
+	const uint UsedIndexCount2 = CopyParams.w;
+
+	uint Index = StartingIndex + IndexOffset;
+
+	// The source buffer copy is splitted between the destination buffers.
+	if (Index < UsedIndexCount0)
+	{
+		DestData0[Index] = SourceData[Index];
+	}
+	else if (Index < UsedIndexCount1)
+	{
+		DestData1[Index] = SourceData[Index];
+	}
+	else if (Index < UsedIndexCount2)
+	{
+		DestData2[Index] = SourceData[Index];
+	}
+}
+ 
\ No newline at end of file
diff --git a/Engine/Shaders/Private/NiagaraSortKeyGen.usf b/Engine/Shaders/Private/NiagaraSortKeyGen.usf
new file mode 100644
--- /dev/null
+++ b/Engine/Shaders/Private/NiagaraSortKeyGen.usf
@@ -0,0 +1,86 @@
+// Copyright 1998-2019 Epic Games, Inc. All Rights Reserved.
+
+/*==============================================================================
+	NiagaraSortKeyGen.usf: Shader to generate particle sorting keys.
+==============================================================================*/
+
+/*------------------------------------------------------------------------------
+	Compile time parameters:
+		THREAD_COUNT - The number of threads to launch per workgroup.
+		SORT_VIEW_DEPTH
+		SORT_VIEW_DISTANCE
+		SORT_CUSTOM_ASCENDING
+		SORT_CUSTOM_DESCENDING
+------------------------------------------------------------------------------*/
+
+#include "Common.ush"
+#include "NiagaraVFParticleAccess.usf"
+
+Buffer<uint> GPUParticleCountBuffer;
+
+// (ParticleCount, GPUParticleCountOffset)
+uint2 ParticleCountParams;
+
+// (EmitterKey, OutputOffset, SortMode, SortAttributeOffset)
+uint4 SortParams;
+
+float3 CameraPosition;
+float3 CameraDirection;
+
+/** Output key buffer. */
+RWBuffer<uint> OutKeys;
+/** Output indices buffer. */
+RWBuffer<int> OutParticleIndices;
+
+[numthreads(THREAD_COUNT,1,1)]
+void GenerateParticleSortKeys(uint ParticleIndex : SV_DispatchThreadID)
+{
+	const uint ParticleCount = ParticleCountParams.x;
+	const uint GPUParticleCountOffset = ParticleCountParams.y;
+
+	const uint GPUParticleCount = GPUParticleCountOffset != -1 ? GPUParticleCountBuffer[GPUParticleCountOffset] : ParticleCount;
+
+	const uint EmitterKey = SortParams.x;
+	const uint OutputOffset = SortParams.y;
+	const uint SortMode = SortParams.z;
+	const uint SortAttributeOffset = SortParams.w;
+
+	if (ParticleIndex < ParticleCount)
+	{
+		float SortValue = 0;
+		uint SortKey = 0xffff; 
+
+		if (ParticleIndex < GPUParticleCount)
+		{
+			if (SortMode == SORT_VIEW_DEPTH || SortMode == SORT_VIEW_DISTANCE)
+			{ 
+				float3 ParticlePosition = GetVec3(SortAttributeOffset, ParticleIndex); 
+				float3 CameraVector = ParticlePosition - CameraPosition;
+				SortValue = SortMode == SORT_VIEW_DEPTH ? dot(CameraVector, CameraDirection) : length(CameraVector);
+				
+				// Render far away first, knowing that SortValue > 0
+				SortKey = (~f32tof16(SortValue)) & 0xffff;
+			}
+			else
+			{
+				SortValue = GetFloat(SortAttributeOffset, ParticleIndex);
+
+				// Convert f32 to uint16. Special handling for negative values.
+				SortKey = f32tof16(SortValue);
+				SortKey = SortValue >= 0 ? (SortKey | 0x8000) : ~SortKey;
+
+				SortKey = (SortMode == SORT_CUSTOM_ASCENDING ? SortKey : ~SortKey) & 0xffff;
+			}
+
+
+			// Clamp to max-1 to make sure that out-of-bound indices are at the end.
+			SortKey = min(SortKey, (uint)0xfffe); // Note that some compiler thread 0xfffe as signed.
+		}
+
+		SortKey |= EmitterKey;
+
+		OutKeys[OutputOffset + ParticleIndex] = SortKey;
+		OutParticleIndices[OutputOffset + ParticleIndex] = ParticleIndex;
+	}
+}
+ 
\ No newline at end of file
diff --git a/Engine/Source/Runtime/Engine/Private/GPUSort.h b/Engine/Source/Runtime/Engine/Private/GPUSort.h
--- a/Engine/Source/Runtime/Engine/Private/GPUSort.h
+++ b/Engine/Source/Runtime/Engine/Private/GPUSort.h
@@ -39,7 +39,7 @@ struct FGPUSortBuffers
  * @param Count - How many items in the buffer need to be sorted.
  * @returns The index of the buffer containing sorted results.
  */
-int32 SortGPUBuffers(FRHICommandListImmediate& RHICmdList, FGPUSortBuffers SortBuffers, int32 BufferIndex, uint32 KeyMask, int32 Count, ERHIFeatureLevel::Type FeatureLevel);
+ENGINE_API int32 SortGPUBuffers(FRHICommandListImmediate& RHICmdList, FGPUSortBuffers SortBuffers, int32 BufferIndex, uint32 KeyMask, int32 Count, ERHIFeatureLevel::Type FeatureLevel);
 
 /**
  * GPU sorting tests.
diff --git a/Engine/Source/Runtime/Engine/Private/Particles/ParticleGpuSimulation.cpp b/Engine/Source/Runtime/Engine/Private/Particles/ParticleGpuSimulation.cpp
--- a/Engine/Source/Runtime/Engine/Private/Particles/ParticleGpuSimulation.cpp
+++ b/Engine/Source/Runtime/Engine/Private/Particles/ParticleGpuSimulation.cpp
@@ -358,24 +358,6 @@ public:
 	}
 };
 
-/**
- * Vertex buffer used to hold particle indices.
- */
-class FParticleIndicesVertexBuffer : public FVertexBuffer
-{
-public:
-
-	/** Shader resource view of the vertex buffer. */
-	FShaderResourceViewRHIRef VertexBufferSRV;
-
-	/** Release RHI resources. */
-	virtual void ReleaseRHI() override
-	{
-		VertexBufferSRV.SafeRelease();
-		FVertexBuffer::ReleaseRHI();
-	}
-};
-
 /**
  * Resources required for GPU particle simulation.
  */
diff --git a/Engine/Source/Runtime/Engine/Private/Particles/ParticleSortingGPU.cpp b/Engine/Source/Runtime/Engine/Private/Particles/ParticleSortingGPU.cpp
--- a/Engine/Source/Runtime/Engine/Private/Particles/ParticleSortingGPU.cpp
+++ b/Engine/Source/Runtime/Engine/Private/Particles/ParticleSortingGPU.cpp
@@ -264,10 +264,21 @@ void FParticleSortBuffers::InitRHI()
 			KeyBufferUAVs[BufferIndex] = RHICreateUnorderedAccessView( KeyBuffers[BufferIndex], PF_R32_UINT );
 
 			VertexBuffers[BufferIndex] = RHICreateVertexBuffer( BufferSize * sizeof(uint32), BUF_Static | BUF_ShaderResource | BUF_UnorderedAccess, CreateInfo);
-			VertexBufferSRVs[BufferIndex] = RHICreateShaderResourceView( VertexBuffers[BufferIndex], /*Stride=*/ sizeof(FFloat16)*2, PF_G16R16F );
-			VertexBufferUAVs[BufferIndex] = RHICreateUnorderedAccessView( VertexBuffers[BufferIndex], PF_G16R16F );
-			VertexBufferSortSRVs[BufferIndex] = RHICreateShaderResourceView( VertexBuffers[BufferIndex], /*Stride=*/ sizeof(uint32), PF_R32_UINT );
-			VertexBufferSortUAVs[BufferIndex] = RHICreateUnorderedAccessView( VertexBuffers[BufferIndex], PF_R32_UINT );
+
+			VertexBufferSortSRVs[BufferIndex] = RHICreateShaderResourceView(VertexBuffers[BufferIndex], /*Stride=*/ sizeof(uint32), PF_R32_UINT);
+			VertexBufferSortUAVs[BufferIndex] = RHICreateUnorderedAccessView(VertexBuffers[BufferIndex], PF_R32_UINT);
+
+			if (bAsInt32)
+			{
+				VertexBufferSRVs[BufferIndex] = RHICreateShaderResourceView(VertexBuffers[BufferIndex], /*Stride=*/ sizeof(int32), PF_R32_SINT);;
+				VertexBufferUAVs[BufferIndex] = RHICreateUnorderedAccessView(VertexBuffers[BufferIndex], PF_R32_SINT);
+			}
+			else
+			{
+				VertexBufferSRVs[BufferIndex] = RHICreateShaderResourceView(VertexBuffers[BufferIndex], /*Stride=*/ sizeof(FFloat16) * 2, PF_G16R16F);
+				VertexBufferUAVs[BufferIndex] = RHICreateUnorderedAccessView(VertexBuffers[BufferIndex], PF_G16R16F);
+			}
+
 		}
 	}
 }
diff --git a/Engine/Source/Runtime/Engine/Private/Particles/ParticleSortingGPU.h b/Engine/Source/Runtime/Engine/Private/Particles/ParticleSortingGPU.h
--- a/Engine/Source/Runtime/Engine/Private/Particles/ParticleSortingGPU.h
+++ b/Engine/Source/Runtime/Engine/Private/Particles/ParticleSortingGPU.h
@@ -14,13 +14,14 @@ struct FGPUSortBuffers;
 /**
  * Buffers in GPU memory used to sort particles.
  */
-class FParticleSortBuffers : public FRenderResource
+class ENGINE_API FParticleSortBuffers : public FRenderResource
 {
 public:
 
 	/** Initialization constructor. */
-	explicit FParticleSortBuffers()
-		:BufferSize(0)
+	explicit FParticleSortBuffers(bool InAsInt32 = false)
+		: BufferSize(0)
+		, bAsInt32(InAsInt32)
 	{
 	}
 
@@ -49,8 +50,9 @@ public:
 
 	/**
 	 * Retrieve the UAV for writing particle vertices.
+	 * bAsUint : whether to return a G16R16 view or a Uint32 view.
 	 */
-	FUnorderedAccessViewRHIParamRef GetVertexBufferUAV()
+	FORCEINLINE FUnorderedAccessViewRHIParamRef GetVertexBufferUAV()
 	{
 		return VertexBufferUAVs[0];
 	}
@@ -77,6 +79,16 @@ public:
 		check((BufferIndex & 0xFFFFFFFE) == 0);
 		return VertexBufferSRVs[BufferIndex];
 	}
+	   
+	/**
+	 * Retrieve the UAV for the sorted vertex buffer at the given index.
+	 */
+	FUnorderedAccessViewRHIParamRef GetSortedVertexBufferUAV(int32 BufferIndex)
+	{
+		check((BufferIndex & 0xFFFFFFFE) == 0);
+		return VertexBufferUAVs[BufferIndex];
+	}
+
 
 	/**
 	 * Get the size allocated for sorted vertex buffers.
@@ -105,6 +117,8 @@ private:
 
 	/** Size allocated for buffers. */
 	int32 BufferSize;
+	/** Whether to allocate UAV and SRV as Int32 instead of G16R16F. */
+	bool bAsInt32;
 };
 
 /**
diff --git a/Engine/Source/Runtime/Engine/Public/ParticleResources.h b/Engine/Source/Runtime/Engine/Public/ParticleResources.h
--- a/Engine/Source/Runtime/Engine/Public/ParticleResources.h
+++ b/Engine/Source/Runtime/Engine/Public/ParticleResources.h
@@ -80,3 +80,22 @@ private:
 
 /** The global scratch vertex buffer. */
 extern TGlobalResource<FParticleScratchVertexBuffer> GParticleScratchVertexBuffer;
+
+
+/**
+ * Vertex buffer used to hold particle indices.
+ */
+class FParticleIndicesVertexBuffer : public FVertexBuffer
+{
+public:
+
+	/** Shader resource view of the vertex buffer. */
+	FShaderResourceViewRHIRef VertexBufferSRV;
+
+	/** Release RHI resources. */
+	virtual void ReleaseRHI() override
+	{
+		VertexBufferSRV.SafeRelease();
+		FVertexBuffer::ReleaseRHI();
+	}
+};
\ No newline at end of file
