Strip Additional IndexBuffers from resources when not in use, which should save about 1MB when used on mobile.

diff --git a/Engine/Plugins/FX/Niagara/Source/Niagara/Private/NiagaraRendererMeshes.cpp b/Engine/Plugins/FX/Niagara/Source/Niagara/Private/NiagaraRendererMeshes.cpp
--- a/Engine/Plugins/FX/Niagara/Source/Niagara/Private/NiagaraRendererMeshes.cpp
+++ b/Engine/Plugins/FX/Niagara/Source/Niagara/Private/NiagaraRendererMeshes.cpp
@@ -345,13 +345,13 @@ void NiagaraRendererMeshes::GetDynamicMeshElements(const TArray<const FSceneView
 
 					if (bIsWireframe)
 					{
-						if (LODModel.WireframeIndexBuffer.IsInitialized())
+						if (LODModel.AdditionalIndexBuffers && LODModel.AdditionalIndexBuffers->WireframeIndexBuffer.IsInitialized())
 						{
 							Mesh.Type = PT_LineList;
 							Mesh.MaterialRenderProxy = UMaterial::GetDefaultMaterial(MD_Surface)->GetRenderProxy();
 							BatchElement.FirstIndex = 0;
-							BatchElement.IndexBuffer = &LODModel.WireframeIndexBuffer;
-							BatchElement.NumPrimitives = LODModel.WireframeIndexBuffer.GetNumIndices() / 2;
+							BatchElement.IndexBuffer = &LODModel.AdditionalIndexBuffers->WireframeIndexBuffer;
+							BatchElement.NumPrimitives = LODModel.AdditionalIndexBuffers->WireframeIndexBuffer.GetNumIndices() / 2;
 
 						}
 						else
