Fixed missing dependency

diff --git a/Engine/Plugins/FX/Niagara/Source/Niagara/Classes/NiagaraGPUSortInfo.h b/Engine/Plugins/FX/Niagara/Source/Niagara/Classes/NiagaraGPUSortInfo.h
--- a/Engine/Plugins/FX/Niagara/Source/Niagara/Classes/NiagaraGPUSortInfo.h
+++ b/Engine/Plugins/FX/Niagara/Source/Niagara/Classes/NiagaraGPUSortInfo.h
@@ -6,6 +6,7 @@ NiagaraGPUSortInfo.h: GPU particle sorting helper
 #pragma once
 
 #include "CoreMinimal.h"
+#include "RHI.h"
 #include "RHIResources.h"
 #include "NiagaraGPUSortInfo.generated.h"
 
