Implementing CL5061473 in main manually to resolve merge issue.

Fix for data interface calls not being stripped correctly.

diff --git a/Engine/Plugins/FX/Niagara/Source/NiagaraCore/Private/NiagaraCustomVersion.cpp b/Engine/Plugins/FX/Niagara/Source/NiagaraCore/Private/NiagaraCustomVersion.cpp
--- a/Engine/Plugins/FX/Niagara/Source/NiagaraCore/Private/NiagaraCustomVersion.cpp
+++ b/Engine/Plugins/FX/Niagara/Source/NiagaraCore/Private/NiagaraCustomVersion.cpp
@@ -8,4 +8,4 @@ const FGuid FNiagaraCustomVersion::GUID(0xFCF57AFA, 0x50764283, 0xB9A9E658, 0xFF
 // Register the custom version with core
 FCustomVersionRegistration GRegisterNiagaraCustomVersion(FNiagaraCustomVersion::GUID, FNiagaraCustomVersion::LatestVersion, TEXT("NiagaraVer"));
 
-const FGuid FNiagaraCustomVersion::LatestScriptCompileVersion(0xC210475C, 0x562E42A7, 0xADC2E907, 0x745E567F);
+const FGuid FNiagaraCustomVersion::LatestScriptCompileVersion(0xC9A20384, 0x928D48B9, 0x9C80D8BA, 0x3E95F269);
\ No newline at end of file
diff --git a/Engine/Source/Developer/ShaderFormatVectorVM/Private/ir_vm_gen_bytecode_visitor.cpp b/Engine/Source/Developer/ShaderFormatVectorVM/Private/ir_vm_gen_bytecode_visitor.cpp
--- a/Engine/Source/Developer/ShaderFormatVectorVM/Private/ir_vm_gen_bytecode_visitor.cpp
+++ b/Engine/Source/Developer/ShaderFormatVectorVM/Private/ir_vm_gen_bytecode_visitor.cpp
@@ -1133,7 +1133,7 @@ struct op_external_func : public op_base
 		bool func_is_used = false;
 		for (variable_info_node* output : outputs)
 		{
-			if (output->last_read > op_idx)
+			if (output->last_read > op_idx || output->last_read != INDEX_NONE)
 			{
 				func_is_used = true;
 				break;
