Add system bounds override setting in system details panel

diff --git a/Engine/Plugins/FX/Niagara/Source/Niagara/Classes/NiagaraSystem.h b/Engine/Plugins/FX/Niagara/Source/Niagara/Classes/NiagaraSystem.h
--- a/Engine/Plugins/FX/Niagara/Source/Niagara/Classes/NiagaraSystem.h
+++ b/Engine/Plugins/FX/Niagara/Source/Niagara/Classes/NiagaraSystem.h
@@ -232,6 +232,12 @@ public:
 
 	const TArray<FName>& GetUserDINamesReadInSystemScripts() const;
 
+	FBox GetFixedBounds() const;
+
+	/** Whether or not fixed bounds are enabled. */
+	UPROPERTY(EditAnywhere, Category = "System", meta = (InlineEditConditionToggle))
+		uint32 bFixedBounds : 1;
+
 private:
 #if WITH_EDITORONLY_DATA
 	INiagaraModule::FMergeEmitterResults MergeChangesForEmitterHandle(FNiagaraEmitterHandle& EmitterHandle);
@@ -285,6 +291,10 @@ protected:
 	FOnSystemCompiled OnSystemCompiledDelegate;
 #endif
 
+	/** The fixed bounding box value. bFixedBounds is the condition whether the fixed bounds can be edited. */
+	UPROPERTY(EditAnywhere, Category = "System", meta = (EditCondition = "bFixedBounds"))
+	FBox FixedBounds;
+
 	UPROPERTY(EditAnywhere, Category = Performance, meta = (ToolTip = "Auto-deactivate system if all emitters are determined to not spawn particles again, regardless of lifetime."))
 	bool bAutoDeactivate;
 
diff --git a/Engine/Plugins/FX/Niagara/Source/Niagara/Private/NiagaraComponent.cpp b/Engine/Plugins/FX/Niagara/Source/Niagara/Private/NiagaraComponent.cpp
--- a/Engine/Plugins/FX/Niagara/Source/Niagara/Private/NiagaraComponent.cpp
+++ b/Engine/Plugins/FX/Niagara/Source/Niagara/Private/NiagaraComponent.cpp
@@ -856,15 +856,23 @@ FBoxSphereBounds UNiagaraComponent::CalcBounds(const FTransform& LocalToWorld) c
 	FBoxSphereBounds SystemBounds;
 	if (SystemInstance.IsValid())
 	{
-		SystemInstance->GetSystemBounds().Init();
-		for (int32 i = 0; i < SystemInstance->GetEmitters().Num(); i++)
+		UNiagaraSystem* System = SystemInstance->GetSystem();
+		if (System->bFixedBounds)
 		{
-			FNiagaraEmitterInstance &Sim = *(SystemInstance->GetEmitters()[i]);
-			SystemInstance->GetSystemBounds() += Sim.GetBounds();
+			SystemBounds = System->GetFixedBounds();
 		}
-		FBox BoundingBox = SystemInstance->GetSystemBounds();
+		else
+		{
+			SystemInstance->GetSystemBounds().Init();
+			for (int32 i = 0; i < SystemInstance->GetEmitters().Num(); i++)
+			{
+				FNiagaraEmitterInstance &Sim = *(SystemInstance->GetEmitters()[i]);
+				SystemInstance->GetSystemBounds() += Sim.GetBounds();
+			}
+			FBox BoundingBox = SystemInstance->GetSystemBounds();
 
-		SystemBounds = FBoxSphereBounds(BoundingBox);
+			SystemBounds = FBoxSphereBounds(BoundingBox);
+		}
 	}
 	else
 	{
diff --git a/Engine/Plugins/FX/Niagara/Source/Niagara/Private/NiagaraSystem.cpp b/Engine/Plugins/FX/Niagara/Source/Niagara/Private/NiagaraSystem.cpp
--- a/Engine/Plugins/FX/Niagara/Source/Niagara/Private/NiagaraSystem.cpp
+++ b/Engine/Plugins/FX/Niagara/Source/Niagara/Private/NiagaraSystem.cpp
@@ -34,10 +34,12 @@ DECLARE_CYCLE_STAT(TEXT("Niagara - System - CompileScript_ResetAfter"), STAT_Nia
 
 UNiagaraSystem::UNiagaraSystem(const FObjectInitializer& ObjectInitializer)
 : Super(ObjectInitializer)
+, bFixedBounds(false)
 , ExposedParameters(this)
 #if WITH_EDITORONLY_DATA
 , bIsolateEnabled(false)
 #endif
+, FixedBounds(FBox(FVector(-100), FVector(100)))
 , bAutoDeactivate(true)
 , WarmupTime(0.0f)
 , WarmupTickCount(0)
@@ -557,6 +559,11 @@ const TArray<FName>& UNiagaraSystem::GetUserDINamesReadInSystemScripts() const
 	return UserDINamesReadInSystemScripts;
 }
 
+FBox UNiagaraSystem::GetFixedBounds() const
+{
+	return FixedBounds;
+}
+
 void CheckDICompileInfo(const TArray<FNiagaraScriptDataInterfaceCompileInfo>& ScriptDICompileInfos, bool& bOutbHasSystemDIsWithPerInstanceData, TArray<FName>& OutUserDINamesReadInSystemScripts)
 {
 	for (const FNiagaraScriptDataInterfaceCompileInfo& ScriptDICompileInfo : ScriptDICompileInfos)
