Fixed some static analysis warning

diff --git a/Engine/Plugins/FX/Niagara/Source/Niagara/Private/NiagaraEmitterInstanceBatcher.cpp b/Engine/Plugins/FX/Niagara/Source/Niagara/Private/NiagaraEmitterInstanceBatcher.cpp
--- a/Engine/Plugins/FX/Niagara/Source/Niagara/Private/NiagaraEmitterInstanceBatcher.cpp
+++ b/Engine/Plugins/FX/Niagara/Source/Niagara/Private/NiagaraEmitterInstanceBatcher.cpp
@@ -331,17 +331,22 @@ void NiagaraEmitterInstanceBatcher::ResolveParticleSortBuffers(FRHICommandListIm
 		int32 UsedIndexCounts[NIAGARA_COPY_BUFFER_BUFFER_COUNT] = {};
 
 		const int32 NumBuffers = FMath::Min<int32>(NIAGARA_COPY_BUFFER_BUFFER_COUNT, SortedVertexBuffers.Num() - Index);
+
+		int32 LastCount = StartingIndex;
 		for (int32 SubIndex = 0; SubIndex < NumBuffers; ++SubIndex)
 		{
-			UAVs[SubIndex] = SortedVertexBuffers[Index + SubIndex].VertexBufferUAV;
-			UsedIndexCounts[SubIndex] = SortedVertexBuffers[Index + SubIndex].UsedIndexCount;
+			const FNiagaraIndicesVertexBuffer& SortBuffer = SortedVertexBuffers[Index + SubIndex];
+			UAVs[SubIndex] = SortBuffer.VertexBufferUAV;
+			UsedIndexCounts[SubIndex] = SortBuffer.UsedIndexCount;
+
+			LastCount = SortBuffer.UsedIndexCount;
 		}
 
 		CopyBufferCS->SetParameters(RHICmdList, ParticleSortBuffers.GetSortedVertexBufferSRV(ResultBufferIndex), UAVs, UsedIndexCounts, StartingIndex, NumBuffers);
-		DispatchComputeShader(RHICmdList, *CopyBufferCS, FMath::DivideAndRoundUp(UsedIndexCounts[NumBuffers - 1] - StartingIndex, NIAGARA_COPY_BUFFER_THREAD_COUNT), 1, 1);
+		DispatchComputeShader(RHICmdList, *CopyBufferCS, FMath::DivideAndRoundUp(LastCount - StartingIndex, NIAGARA_COPY_BUFFER_THREAD_COUNT), 1, 1);
 		RHICmdList.TransitionResources(EResourceTransitionAccess::ERWBarrier, EResourceTransitionPipeline::EComputeToGfx, UAVs, NumBuffers);
 
-		StartingIndex = UsedIndexCounts[NumBuffers - 1];
+		StartingIndex = LastCount;
 	}
 	CopyBufferCS->UnbindBuffers(RHICmdList);
 }
diff --git a/Engine/Plugins/FX/Niagara/Source/Niagara/Private/NiagaraRendererMeshes.cpp b/Engine/Plugins/FX/Niagara/Source/Niagara/Private/NiagaraRendererMeshes.cpp
--- a/Engine/Plugins/FX/Niagara/Source/Niagara/Private/NiagaraRendererMeshes.cpp
+++ b/Engine/Plugins/FX/Niagara/Source/Niagara/Private/NiagaraRendererMeshes.cpp
@@ -317,7 +317,7 @@ void NiagaraRendererMeshes::GetDynamicMeshElements(const TArray<const FSceneView
 				CollectorResources.VertexFactory.InitResource();
 				CollectorResources.VertexFactory.SetUniformBuffer(CollectorResources.UniformBuffer);
 			
-				const bool bIsWireframe = AllowDebugViewmodes() && View->Family->EngineShowFlags.Wireframe;
+				const bool bIsWireframe = AllowDebugViewmodes() && View && View->Family->EngineShowFlags.Wireframe;
 
 				for (int32 SectionIndex = 0; SectionIndex < LODModel.Sections.Num(); SectionIndex++)
 				{
