Fix convert nodes not injecting struct definitions in FHLSLNiagaraTranslator::Convert().

diff --git a/Engine/Plugins/FX/Niagara/Source/NiagaraEditor/Private/NiagaraHlslTranslator.cpp b/Engine/Plugins/FX/Niagara/Source/NiagaraEditor/Private/NiagaraHlslTranslator.cpp
--- a/Engine/Plugins/FX/Niagara/Source/NiagaraEditor/Private/NiagaraHlslTranslator.cpp
+++ b/Engine/Plugins/FX/Niagara/Source/NiagaraEditor/Private/NiagaraHlslTranslator.cpp
@@ -77,11 +77,9 @@ DECLARE_CYCLE_STAT(TEXT("Niagara - HlslTranslator - GenerateFunctionSignature_Ou
 DECLARE_CYCLE_STAT(TEXT("Niagara - HlslTranslator - GenerateFunctionSignature_Inputs"), STAT_NiagaraEditor_Module_NiagaraHLSLTranslator_GenerateFunctionSignature_Inputs, STATGROUP_NiagaraEditor);
 DECLARE_CYCLE_STAT(TEXT("Niagara - HlslTranslator - GenerateFunctionSignature_FindInputNodes"), STAT_NiagaraEditor_Module_NiagaraHLSLTranslator_GenerateFunctionSignature_FindInputNodes, STATGROUP_NiagaraEditor);
 
-
 static FNiagaraShaderQueueTickable NiagaraShaderQueueProcessor;
 FNiagaraShaderProcessorTickable NiagaraShaderProcessor;
 
-
 // not pretty. TODO: refactor
 // this will be called via a delegate from UNiagaraScript's cache for cook function,
 // because editor tickables aren't ticked during cooking
@@ -123,7 +121,7 @@ void FNiagaraShaderQueueTickable::ProcessQueue()
 			ShaderScript->SourceName = Emitter->GetUniqueEmitterName();
 		}
 		ShaderScript->HlslOutput = CompilableScript->GetVMExecutableData().LastHlslTranslationGPU;
-		
+
 		{
 			// Create a shader compiler environment for the script that will be shared by all jobs from this script
 			TRefCountPtr<FShaderCompilerEnvironment> CompilerEnvironment = new FShaderCompilerEnvironment();
@@ -224,8 +222,8 @@ FString FHlslNiagaraTranslator::GetCode(FNiagaraCodeChunk& Chunk)
 
 FString FHlslNiagaraTranslator::GetCodeAsSource(int32 ChunkIdx)
 {
-	if(ChunkIdx>=0 && ChunkIdx<CodeChunks.Num())
-	{ 
+	if (ChunkIdx >= 0 && ChunkIdx < CodeChunks.Num())
+	{
 		FNiagaraCodeChunk& Chunk = CodeChunks[ChunkIdx];
 		return Chunk.SymbolName + Chunk.ComponentMask;
 	}
@@ -256,7 +254,7 @@ bool FHlslNiagaraTranslator::ValidateTypePins(UNiagaraNode* NodeToValidate)
 }
 
 
-void FHlslNiagaraTranslator::GenerateFunctionSignature(ENiagaraScriptUsage ScriptUsage, FString InName, const FString& InFullName, UNiagaraGraph* FuncGraph, TArray<int32>& Inputs, 
+void FHlslNiagaraTranslator::GenerateFunctionSignature(ENiagaraScriptUsage ScriptUsage, FString InName, const FString& InFullName, UNiagaraGraph* FuncGraph, TArray<int32>& Inputs,
 	bool bHadNumericInputs, bool bHasParameterMapParameters, FNiagaraFunctionSignature& OutSig)const
 {
 	SCOPE_CYCLE_COUNTER(STAT_NiagaraEditor_Module_NiagaraHLSLTranslator_GenerateFunctionSignature);
@@ -280,7 +278,7 @@ void FHlslNiagaraTranslator::GenerateFunctionSignature(ENiagaraScriptUsage Scrip
 
 		if (Inputs.Num() != InputsNodes.Num())
 		{
-			const_cast<FHlslNiagaraTranslator*>(this)->Error(FText::Format(LOCTEXT("GenerateFunctionSignatureFail", "Generating function signature for {0} failed.  The function call is providing a different number of inputs than the function graph supplies."), 
+			const_cast<FHlslNiagaraTranslator*>(this)->Error(FText::Format(LOCTEXT("GenerateFunctionSignatureFail", "Generating function signature for {0} failed.  The function call is providing a different number of inputs than the function graph supplies."),
 				FText::FromString(InFullName)), nullptr, nullptr);
 			return;
 		}
@@ -427,8 +425,6 @@ FString FHlslNiagaraTranslator::GetFunctionDefinitions()
 	return FwdDeclString + TEXT("\n") + DefinitionsString;
 }
 
-	
-
 void FHlslNiagaraTranslator::BuildMissingDefaults()
 {
 	AddBodyComment(TEXT("// Begin HandleMissingDefaultValues"));
@@ -479,7 +475,7 @@ FString FHlslNiagaraTranslator::BuildParameterMapHlslDefinitions(TArray<FNiagara
 		const UEdGraphPin* OriginalPin = ParamMapHistories[ParamMapIdx].GetOriginalPin();
 		UniqueParamMapStartingPins.AddUnique(OriginalPin);
 	}
-	
+
 
 	TArray<FNiagaraVariable> UniqueVariables;
 
@@ -639,7 +635,7 @@ FString FHlslNiagaraTranslator::BuildParameterMapHlslDefinitions(TArray<FNiagara
 				TPair<FString, FString> Pair(TypeName, VarName);
 				ParamStructNameToMembers.FindOrAdd(StructType).AddUnique(Pair);
 				ParamStructNames.AddUnique(StructType);
-			}		
+			}
 		}
 	}
 
@@ -738,11 +734,11 @@ const FNiagaraTranslateResults &FHlslNiagaraTranslator::Translate(const FNiagara
 		return TranslateResults;
 	}
 
-    bool bNeedsPersistentIDs = CompileOptions.AdditionalDefines.Contains(TEXT("RequiresPersistentIDs"));
-    if (bNeedsPersistentIDs && CompilationTarget == ENiagaraSimTarget::GPUComputeSim) {
-        Error(LOCTEXT("GPUPersistentIDFail", "GPU particles do not support persistent IDs. Change to a CPU simulation or disable persistent IDs."), nullptr, nullptr);
-        return TranslateResults;
-    }
+	bool bNeedsPersistentIDs = CompileOptions.AdditionalDefines.Contains(TEXT("RequiresPersistentIDs"));
+	if (bNeedsPersistentIDs && CompilationTarget == ENiagaraSimTarget::GPUComputeSim) {
+		Error(LOCTEXT("GPUPersistentIDFail", "GPU particles do not support persistent IDs. Change to a CPU simulation or disable persistent IDs."), nullptr, nullptr);
+		return TranslateResults;
+	}
 
 	TranslationStages.Empty();
 	ActiveStageIdx = 0;
@@ -781,7 +777,7 @@ const FNiagaraTranslateResults &FHlslNiagaraTranslator::Translate(const FNiagara
 		break;
 	default:
 		TranslationStages.Add(FHlslNiagaraTranslationStage(CompileOptions.TargetUsage, CompileOptions.TargetUsageId));
-		TranslationStages[0].PassNamespace = TEXT("Map");		
+		TranslationStages[0].PassNamespace = TEXT("Map");
 		TranslationStages[0].OutputNode = SourceGraph->FindEquivalentOutputNode(TranslationStages[0].ScriptUsage, TranslationStages[0].UsageId);
 		TranslationStages[0].ChunkModeIndex = ENiagaraCodeChunkMode::Body;
 		ParamMapHistories.AddDefaulted(1);
@@ -819,7 +815,7 @@ const FNiagaraTranslateResults &FHlslNiagaraTranslator::Translate(const FNiagara
 	// Get all the parameter map histories traced to this graph from output nodes. We'll revisit this shortly in order to build out just the ones we care about for this translation.
 	OtherOutputParamMapHistories = CompileData->GetPrecomputedHistories();
 
-	if (ParamMapHistories.Num() == 1 && OtherOutputParamMapHistories.Num() == 1 && ( CompileOptions.TargetUsage == ENiagaraScriptUsage::Function || CompileOptions.TargetUsage == ENiagaraScriptUsage::DynamicInput ) )
+	if (ParamMapHistories.Num() == 1 && OtherOutputParamMapHistories.Num() == 1 && (CompileOptions.TargetUsage == ENiagaraScriptUsage::Function || CompileOptions.TargetUsage == ENiagaraScriptUsage::DynamicInput))
 	{
 		ParamMapHistories[0] = (OtherOutputParamMapHistories[0]);
 
@@ -847,7 +843,7 @@ const FNiagaraTranslateResults &FHlslNiagaraTranslator::Translate(const FNiagara
 				UNiagaraNodeOutput* TargetOutputNode = TranslationStages[ParamMapIdx].OutputNode;
 				if (FoundHistory.GetFinalOutputNode() == TargetOutputNode)
 				{
-					if(bNeedsPersistentIDs)
+					if (bNeedsPersistentIDs)
 					{
 						//TODO: Setup alias for current level to decouple from "Particles". Would we ever want emitter or system persistent IDs?
 						FNiagaraVariable Var = FNiagaraVariable(FNiagaraTypeDefinition::GetIDDef(), TEXT("Particles.ID"));
@@ -917,7 +913,7 @@ const FNiagaraTranslateResults &FHlslNiagaraTranslator::Translate(const FNiagara
 	{
 		ValidateParticleIDUsage();
 	}
-	
+
 	//Create main scope pin cache.
 	PinToCodeChunks.AddDefaulted(1);
 
@@ -947,7 +943,7 @@ const FNiagaraTranslateResults &FHlslNiagaraTranslator::Translate(const FNiagara
 				AddBodyComment(TEXT("//End Spawn Script!\n\n"));
 				BuildMissingDefaults();
 			}
-			
+
 			if (TranslationStages[i].bInterpolatePreviousParams)
 			{
 				bInterpolateParams = true;
@@ -1046,7 +1042,7 @@ const FNiagaraTranslateResults &FHlslNiagaraTranslator::Translate(const FNiagara
 				Error(ErrorMessage, nullptr, nullptr);
 			}
 		}
-				
+
 		//Declare parameters.
 		//TODO: Separate Cbuffer for Global, System and Emitter parameters.
 		{
@@ -1092,7 +1088,7 @@ const FNiagaraTranslateResults &FHlslNiagaraTranslator::Translate(const FNiagara
 		DataSetReadIds.Add(GetInstanceDataSetID());
 		DataSetWriteVars.Add(&InstanceWriteVars);
 		DataSetWriteIds.Add(GetInstanceDataSetID());
-				
+
 		if (IsBulkSystemScript())
 		{
 			// We have two sets of data that can change independently.. The engine data set are variables
@@ -1152,10 +1148,10 @@ const FNiagaraTranslateResults &FHlslNiagaraTranslator::Translate(const FNiagara
 				}
 			}
 		}
-		
+
 		InstanceReadVars = BasicAttributes;
 		InstanceWriteVars = BasicAttributes;
-	
+
 
 		// We sort the variables so that they end up in the same ordering between Spawn & Update...
 		InstanceReadVars.Sort([&](const FNiagaraVariable& A, const FNiagaraVariable& B)
@@ -1223,7 +1219,7 @@ const FNiagaraTranslateResults &FHlslNiagaraTranslator::Translate(const FNiagara
 		}
 
 		DefineInterpolatedParametersFunction(HlslOutput);
-		
+
 		// define functions for reading and writing all secondary data sets
 		DefineDataSetReadFunction(HlslOutput, CompilationOutput.ScriptData.ReadDataSets);
 		DefineDataSetWriteFunction(HlslOutput, CompilationOutput.ScriptData.WriteDataSets, WriteConditionVars);
@@ -1367,7 +1363,7 @@ void FHlslNiagaraTranslator::GatherComponentsForDataSetAccess(UScriptStruct* Str
 	{
 		bMatrixRoot = true;
 	}
-	
+
 	//Bools are an awkward special case. TODO: make neater.
 	if (FNiagaraTypeDefinition(Struct) == FNiagaraTypeDefinition::GetBoolDef())
 	{
@@ -1443,10 +1439,10 @@ void FHlslNiagaraTranslator::DefineInterpolatedParametersFunction(FString &HlslO
 		Emitter_InterpSpawnStartDt = Emitter_InterpSpawnStartDt.Replace(TEXT("."), TEXT("_"));//TODO: This should be rolled into GetSanitisedSymbolName but currently some usages require the '.' intact. Fix those up!.
 		FString Emitter_SpawnInterval = GetSanitizedSymbolName(ActiveHistoryForFunctionCalls.ResolveAliases(SYS_PARAM_EMITTER_SPAWN_INTERVAL).GetName().ToString());
 		Emitter_SpawnInterval = Emitter_SpawnInterval.Replace(TEXT("."), TEXT("_"));//TODO: This should be rolled into GetSanitisedSymbolName but currently some usages require the '.' intact. Fix those up!.
-		
+
 		HlslOutputString += TEXT("void InterpolateParameters(inout FSimulationContext Context)\n{\n");
 
-		FString PrevMap = TranslationStages[i-1].PassNamespace;
+		FString PrevMap = TranslationStages[i - 1].PassNamespace;
 		FString CurMap = TranslationStages[i].PassNamespace;
 		{
 			HlslOutputString += TEXT("\tint InterpSpawn_Index = ExecIndex();\n");
@@ -1601,9 +1597,9 @@ void FHlslNiagaraTranslator::DefineDataSetWriteFunction(FString &HlslOutputStrin
 
 		HlslOutputString += "\tTmpWriteIndex = AcquireIndex(";
 		HlslOutputString += FString::FromInt(DataSetIndex);
-		HlslOutputString +=", bValid);\n";
+		HlslOutputString += ", bValid);\n";
 
-		HlslOutputString += CompilationTarget==ENiagaraSimTarget::GPUComputeSim ? "\tif(TmpWriteIndex>=0)\n\t{\n" : "";
+		HlslOutputString += CompilationTarget == ENiagaraSimTarget::GPUComputeSim ? "\tif(TmpWriteIndex>=0)\n\t{\n" : "";
 
 		for (TPair<int32, FDataSetAccessInfo>& IndexInfoPair : DataSetInfoPair.Value)
 		{
@@ -1647,13 +1643,13 @@ void FHlslNiagaraTranslator::DefineDataInterfaceHLSL(FString &InHlslOutput)
 	TArray<FString> BufferParamNames;
 	for (uint32 i = 0; i < 32; i++)
 	{
-		BufferParamNames.Add( TEXT("DataInterfaceBuffer_") + FString::FromInt(i) );
+		BufferParamNames.Add(TEXT("DataInterfaceBuffer_") + FString::FromInt(i));
 	}
 
 	uint32 CurBufferIndex = 0;
 	for (int32 i = 0; i < CompilationOutput.ScriptData.DataInterfaceInfo.Num(); i++)
 	{
-		FNiagaraScriptDataInterfaceCompileInfo& Info = CompilationOutput.ScriptData.DataInterfaceInfo[i];		
+		FNiagaraScriptDataInterfaceCompileInfo& Info = CompilationOutput.ScriptData.DataInterfaceInfo[i];
 
 		UObject* const* FoundCDO = CompileData->CDOs.Find(Info.Type.GetClass());
 		check(FoundCDO != nullptr);
@@ -1745,12 +1741,12 @@ void FHlslNiagaraTranslator::DefineMain(FString &OutHlslOutput,
 
 		// NOTE(mv): This relies on Particles.UniqueID being excluded from being default initialized. 
 		//           This happens in FNiagaraParameterMapHistory::ShouldIgnoreVariableDefault
-		if (UNiagaraScript::IsParticleSpawnScript(CompileOptions.TargetUsage)) 
+		if (UNiagaraScript::IsParticleSpawnScript(CompileOptions.TargetUsage))
 		{
 			FString MapName = UNiagaraScript::IsInterpolatedParticleSpawnScript(CompileOptions.TargetUsage) ? TEXT("Context.MapSpawn") : TEXT("Context.Map");
 			OutHlslOutput += FString::Printf(TEXT("\t%s.Particles.UniqueID = Engine_Emitter_TotalSpawnedParticles + ExecIndex();\n"), *MapName);
 		}
-		else if (UNiagaraScript::IsGPUScript(CompileOptions.TargetUsage)) 
+		else if (UNiagaraScript::IsGPUScript(CompileOptions.TargetUsage))
 		{
 			// NOTE(mv): The GPU script only have one file, so we need to make sure we only apply this in the spawn phase. 
 			//           
@@ -1768,7 +1764,7 @@ void FHlslNiagaraTranslator::DefineMain(FString &OutHlslOutput,
 	{
 		// call the read data set function
 		OutHlslOutput += TEXT("\tReadDataSets(Context);\n");
-		
+
 		// branch between spawn and update
 		if (CompilationTarget == ENiagaraSimTarget::GPUComputeSim)
 		{
@@ -1846,7 +1842,7 @@ void FHlslNiagaraTranslator::DefineMain(FString &OutHlslOutput,
 					OutHlslOutput += FString::Printf(TEXT("\t\tPhase = %d;\n"), StageIdx + 1);
 					OutHlslOutput += FString::Printf(TEXT("\t\tGCurrentPhase = %d;\n"), StageIdx + 1);
 				}
-				
+
 				OutHlslOutput += TEXT("\t}\n");
 			}
 		}
@@ -1974,9 +1970,9 @@ void FHlslNiagaraTranslator::DefineDataSetVariableWrites(FString &OutHlslOutput,
 				OutHlslOutput += TEXT("\tbool bValid = Context.MapUpdate.") + DataSetName + TEXT(".Alive;\n");
 			}
 		}
-		else if ((UNiagaraScript::IsParticleSpawnScript(CompileOptions.TargetUsage) && bHasPerParticleAliveSpawn) 
+		else if ((UNiagaraScript::IsParticleSpawnScript(CompileOptions.TargetUsage) && bHasPerParticleAliveSpawn)
 			|| (UNiagaraScript::IsGPUScript(CompileOptions.TargetUsage) && bHasPerParticleAliveSpawn)
-			|| (UNiagaraScript::IsParticleUpdateScript(CompileOptions.TargetUsage) && bHasPerParticleAliveUpdate) 
+			|| (UNiagaraScript::IsParticleUpdateScript(CompileOptions.TargetUsage) && bHasPerParticleAliveUpdate)
 			|| (UNiagaraScript::IsParticleEventScript(CompileOptions.TargetUsage) && bHasPerParticleAliveEvent)
 			|| (UNiagaraScript::IsGPUScript(CompileOptions.TargetUsage) && bHasPerParticleAliveUpdate && CompilationTarget == ENiagaraSimTarget::GPUComputeSim))
 		{
@@ -2040,15 +2036,15 @@ void FHlslNiagaraTranslator::DefineDataSetVariableReads(FString &OutHlslOutput,
 	}
 
 	bool bIsGPUScript = UNiagaraScript::IsGPUScript(CompileOptions.TargetUsage);
-	bool bIsSpawnScript = 
-		UNiagaraScript::IsParticleSpawnScript(CompileOptions.TargetUsage) ||UNiagaraScript::IsInterpolatedParticleSpawnScript(CompileOptions.TargetUsage) ||
-		UNiagaraScript::IsEmitterSpawnScript(CompileOptions.TargetUsage) ||	UNiagaraScript::IsSystemSpawnScript(CompileOptions.TargetUsage);
-	bool bIsUpdateScript = 
+	bool bIsSpawnScript =
+		UNiagaraScript::IsParticleSpawnScript(CompileOptions.TargetUsage) || UNiagaraScript::IsInterpolatedParticleSpawnScript(CompileOptions.TargetUsage) ||
+		UNiagaraScript::IsEmitterSpawnScript(CompileOptions.TargetUsage) || UNiagaraScript::IsSystemSpawnScript(CompileOptions.TargetUsage);
+	bool bIsUpdateScript =
 		UNiagaraScript::IsParticleUpdateScript(CompileOptions.TargetUsage) || UNiagaraScript::IsEmitterUpdateScript(CompileOptions.TargetUsage) ||
 		UNiagaraScript::IsSystemUpdateScript(CompileOptions.TargetUsage);
 	bool bIsEventScript = UNiagaraScript::IsParticleEventScript(CompileOptions.TargetUsage);
-	bool bIsSystemOrEmitterScript = 
-		UNiagaraScript::IsEmitterSpawnScript(CompileOptions.TargetUsage) ||	UNiagaraScript::IsSystemSpawnScript(CompileOptions.TargetUsage) ||
+	bool bIsSystemOrEmitterScript =
+		UNiagaraScript::IsEmitterSpawnScript(CompileOptions.TargetUsage) || UNiagaraScript::IsSystemSpawnScript(CompileOptions.TargetUsage) ||
 		UNiagaraScript::IsEmitterUpdateScript(CompileOptions.TargetUsage) || UNiagaraScript::IsSystemUpdateScript(CompileOptions.TargetUsage);
 	bool bIsPrimaryDataSet = DataSetIndex == 0;
 
@@ -2060,7 +2056,7 @@ void FHlslNiagaraTranslator::DefineDataSetVariableReads(FString &OutHlslOutput,
 		FString ContextName = TEXT("\tContext.Map.");
 		if (TranslationStages.Num() > 1)
 		{
-			ContextName = FString::Printf(TEXT("\tContext.%s."),*TranslationStages[0].PassNamespace);
+			ContextName = FString::Printf(TEXT("\tContext.%s."), *TranslationStages[0].PassNamespace);
 		}
 
 		FString VarReads;
@@ -2168,7 +2164,7 @@ void FHlslNiagaraTranslator::WriteDataSetStructDeclarations(TMap<FNiagaraDataSet
 				OutHlslOutput += "int DSComponentBufferSizeReadFloat" + IndexString + ";\n";
 				OutHlslOutput += "int DSComponentBufferSizeReadInt" + IndexString + ";\n";
 			}
-			else 
+			else
 			{
 				OutHlslOutput += FString(TEXT("RWBuffer<float> RWWriteDataSetFloat")) + IndexString + ";\n";
 				OutHlslOutput += FString(TEXT("RWBuffer<int> RWWriteDataSetInt")) + IndexString + ";\n";
@@ -2254,19 +2250,19 @@ FString FHlslNiagaraTranslator::GetSanitizedSymbolName(FString SymbolName, bool
 		{
 			SplitName[i] = TEXT("INTEGER_") + SplitName[i];
 		}
-		
+
 		const FString* FoundReplacementStr = ReplacementsForInvalid.Find(SplitName[i]); // Look for the string in the keyword protections array.
 		if (FoundReplacementStr)
 		{
 			SplitName[i] = *FoundReplacementStr;
-		}		
+		}
 
 		SplitName[i].ReplaceInline(TEXT("\t"), TEXT(""));
 		SplitName[i].ReplaceInline(TEXT(" "), TEXT(""));
 
 		// Handle internationalization of characters..
 		FString ChangedSplitName;
-		ChangedSplitName.Reserve(SplitName[i].Len()*6); // Assign room for every current char to be 'ASCXXX'
+		ChangedSplitName.Reserve(SplitName[i].Len() * 6); // Assign room for every current char to be 'ASCXXX'
 		for (int32 j = 0; j < SplitName[i].Len(); j++)
 		{
 			if ((SplitName[i][j] >= TCHAR('0') && SplitName[i][j] <= TCHAR('9')) ||
@@ -2318,7 +2314,7 @@ FString FHlslNiagaraTranslator::GetUniqueSymbolName(FName BaseName)
 		return RetString;
 	}
 
-	if(*NameCount > 0)
+	if (*NameCount > 0)
 	{
 		RetString += LexToString(*NameCount);
 	}
@@ -2578,12 +2574,12 @@ bool FHlslNiagaraTranslator::ShouldInterpolateParameter(const FNiagaraVariable&
 	//Skip interpolation for some system constants.
 	if (Parameter == SYS_PARAM_ENGINE_DELTA_TIME ||
 		Parameter == SYS_PARAM_ENGINE_INV_DELTA_TIME ||
-		Parameter == SYS_PARAM_ENGINE_EXEC_COUNT || 
+		Parameter == SYS_PARAM_ENGINE_EXEC_COUNT ||
 		Parameter == SYS_PARAM_EMITTER_SPAWNRATE ||
 		Parameter == SYS_PARAM_EMITTER_SPAWN_INTERVAL ||
 		Parameter == SYS_PARAM_EMITTER_INTERP_SPAWN_START_DT ||
-		Parameter == SYS_PARAM_ENGINE_EMITTER_TOTAL_SPAWNED_PARTICLES || 
-		Parameter == SYS_PARAM_EMITTER_RANDOM_SEED || 
+		Parameter == SYS_PARAM_ENGINE_EMITTER_TOTAL_SPAWNED_PARTICLES ||
+		Parameter == SYS_PARAM_EMITTER_RANDOM_SEED ||
 		Parameter == SYS_PARAM_ENGINE_SYSTEM_TICK_COUNT)
 	{
 		return false;
@@ -2632,7 +2628,7 @@ int32 FHlslNiagaraTranslator::GetRapidIterationParameter(const FNiagaraVariable&
 	// Check to see if this is the first time we've encountered this node and it is a viable candidate for rapid iteration
 	if (bIsCandidateForRapidIteration && TranslationOptions.bParameterRapidIteration)
 	{
-		
+
 		// go ahead and make it into a constant variable..
 		int32 OutputChunkId = INDEX_NONE;
 		if (ParameterMapRegisterExternalConstantNamespaceVariable(Parameter, nullptr, INDEX_NONE, OutputChunkId, nullptr))
@@ -2738,7 +2734,7 @@ int32 FHlslNiagaraTranslator::GetParameter(const FNiagaraVariable& Parameter)
 								InitializeParameterMapDefaults(i);
 							}
 						}
-						
+
 						return i;
 					}
 				}
@@ -3015,7 +3011,7 @@ void FHlslNiagaraTranslator::Output(UNiagaraNodeOutput* OutputNode, const TArray
 	}
 	else
 	{
-		
+
 		{
 			check(InstanceWrite.CodeChunks.Num() == 0);//Should only hit one output node.
 
@@ -3060,7 +3056,7 @@ void FHlslNiagaraTranslator::Output(UNiagaraNodeOutput* OutputNode, const TArray
 }
 
 int32 FHlslNiagaraTranslator::GetAttribute(const FNiagaraVariable& Attribute)
-{	
+{
 	if (!AddStructToDefinitionSet(Attribute.GetType()))
 	{
 		Error(FText::Format(LOCTEXT("GetConstantFail", "Cannot handle type {0}! Variable: {1}"), Attribute.GetType().GetNameText(), FText::FromName(Attribute.GetName())), nullptr, nullptr);
@@ -3080,7 +3076,7 @@ int32 FHlslNiagaraTranslator::GetAttribute(const FNiagaraVariable& Attribute)
 		}
 		else
 		{
-			Error(LOCTEXT("AttrReadInSpawnError","Cannot read attribute in a spawn script as it's value is not yet initialized."), nullptr, nullptr);
+			Error(LOCTEXT("AttrReadInSpawnError", "Cannot read attribute in a spawn script as it's value is not yet initialized."), nullptr, nullptr);
 			return INDEX_NONE;
 		}
 	}
@@ -3090,7 +3086,7 @@ int32 FHlslNiagaraTranslator::GetAttribute(const FNiagaraVariable& Attribute)
 		//           is written to before Simulate() in the SpawnScript...
 		// TODO(mv): Also allow Particles.ID for the same reasons?
 		CompilationOutput.ScriptData.DataUsage.bReadsAttributeData |= (Attribute.GetName() != TEXT("Particles.UniqueID"));
-		
+
 		int32 Chunk = INDEX_NONE;
 		if (!ParameterMapRegisterNamespaceAttributeVariable(Attribute, nullptr, 0, Chunk))
 		{
@@ -3212,7 +3208,7 @@ void FHlslNiagaraTranslator::ParameterMapSet(UNiagaraNodeParameterMapSet* SetNod
 				}
 
 				bool bAllowDataInterfaces = true;
-				if(ParamMapHistoryIdx < ParamMapHistories.Num() && ParamMapHistories[ParamMapHistoryIdx].IsPrimaryDataSetOutput(Var, CompileOptions.TargetUsage, bAllowDataInterfaces))
+				if (ParamMapHistoryIdx < ParamMapHistories.Num() && ParamMapHistories[ParamMapHistoryIdx].IsPrimaryDataSetOutput(Var, CompileOptions.TargetUsage, bAllowDataInterfaces))
 				{
 					if (Input < 0 || Input >= CompilationOutput.ScriptData.DataInterfaceInfo.Num())
 					{
@@ -3452,22 +3448,22 @@ bool FHlslNiagaraTranslator::ParameterMapRegisterUniformAttributeVariable(const
 
 void FHlslNiagaraTranslator::ValidateParticleIDUsage()
 {
-    if (CompileOptions.AdditionalDefines.Contains(TEXT("RequiresPersistentIDs"))) 
-	{
-        // persistent IDs are active and can be safely used as inputs
-        return;
-    }
-    FName particleIDName(TEXT("Particles.ID"));
-    for (FNiagaraParameterMapHistory& History : ParamMapHistories)
-    {
-        for (const FNiagaraVariable& Variable : History.Variables)
-        {
-            if (Variable.GetName() == particleIDName) 
+	if (CompileOptions.AdditionalDefines.Contains(TEXT("RequiresPersistentIDs")))
+	{
+		// persistent IDs are active and can be safely used as inputs
+		return;
+	}
+	FName particleIDName(TEXT("Particles.ID"));
+	for (FNiagaraParameterMapHistory& History : ParamMapHistories)
+	{
+		for (const FNiagaraVariable& Variable : History.Variables)
+		{
+			if (Variable.GetName() == particleIDName)
 			{
 				Error(LOCTEXT("PersistentIDActivationFail", "Before the Particles.ID parameter can be used, the 'Requires persistent IDs' option has to be activated in the emitter properties. Note that this comes with additional memory and CPU costs."), nullptr, nullptr);
-            }
-        }
-    }
+			}
+		}
+	}
 }
 
 bool FHlslNiagaraTranslator::ParameterMapRegisterNamespaceAttributeVariable(const FNiagaraVariable& InVariable, UNiagaraNode* InNode, int32 InParamMapHistoryIdx, int32& Output)
@@ -3616,7 +3612,7 @@ void FHlslNiagaraTranslator::ParameterMapGet(UNiagaraNodeParameterMapGet* GetNod
 
 	TArray<UEdGraphPin*> InputPins;
 	GetNode->GetInputPins(InputPins);
-	
+
 	int32 ParamMapHistoryIdx = Inputs[0];
 
 	Outputs.SetNum(OutputPins.Num());
@@ -3641,7 +3637,7 @@ void FHlslNiagaraTranslator::ParameterMapGet(UNiagaraNodeParameterMapGet* GetNod
 	}
 
 	FString ParameterMapInstanceName = GetParameterMapInstanceName(ParamMapHistoryIdx);
-	
+
 	for (int32 i = 0; i < OutputPins.Num(); i++)
 	{
 		if (GetNode->IsAddPin(OutputPins[i]))
@@ -3872,7 +3868,7 @@ void FHlslNiagaraTranslator::HandleParameterRead(int32 ParamMapHistoryIdx, const
 						{
 							if (InputNode->Usage != ENiagaraInputNodeUsage::TranslatorConstant)
 							{
-								Error(FText::Format(LOCTEXT("InvalidParamMapStartForDefaultPin", "Default found for {0}, but the parameter map source for default pins needs to be a {1} node, not a generic input node."), 
+								Error(FText::Format(LOCTEXT("InvalidParamMapStartForDefaultPin", "Default found for {0}, but the parameter map source for default pins needs to be a {1} node, not a generic input node."),
 									FText::FromName(Var.GetName()),
 									FText::FromName(TRANSLATOR_PARAM_BEGIN_DEFAULTS.GetName())), ErrorNode, nullptr);
 							}
@@ -3882,7 +3878,7 @@ void FHlslNiagaraTranslator::HandleParameterRead(int32 ParamMapHistoryIdx, const
 				LastSetChunkIdx = CompilePin(InputPin);
 				CurrentDefaultPinTraversal.Pop();
 			}
-			else 
+			else
 			{
 				LastSetChunkIdx = GetConstant(Var);
 			}
@@ -3985,14 +3981,14 @@ void FHlslNiagaraTranslator::ReadDataSet(const FNiagaraDataSetID DataSet, const
 	if (!DataSetReadForInput)
 	{
 		DataSetReadForInput = &Reads.Add(InputChunk);
-		
+
 		DataSetReadForInput->Variables = Variables;
-		DataSetReadForInput->CodeChunks.Reserve(Variables.Num()+1);
+		DataSetReadForInput->CodeChunks.Reserve(Variables.Num() + 1);
 
 		FString DataSetAccessSymbol = GetDataSetAccessSymbol(DataSet, InputChunk, true);
 		//Add extra output to indicate if event read is valid data.
 		//DataSetReadForInput->CodeChunks.Add(AddSourceChunk(DataSetAccessSymbol + TEXT("_Valid"), FNiagaraTypeDefinition::GetIntDef()));
-		for (int32 i=0; i<Variables.Num(); ++i)
+		for (int32 i = 0; i < Variables.Num(); ++i)
 		{
 			const FNiagaraVariable& Var = Variables[i];
 			if (!AddStructToDefinitionSet(Var.GetType()))
@@ -4043,7 +4039,7 @@ void FHlslNiagaraTranslator::WriteDataSet(const FNiagaraDataSetID DataSet, const
 	TMap<int32, FDataSetAccessInfo>& Writes = DataSetWriteInfo[(int32)AccessMode].FindOrAdd(DataSet);
 	FDataSetAccessInfo* DataSetWriteForInput = Writes.Find(InputChunk);
 	int32& DataSetWriteConditionalIndex = DataSetWriteConditionalInfo[(int32)AccessMode].FindOrAdd(DataSet);
-	
+
 	//We should never try to write to the exact same dataset at the same index/condition twice.
 	//This is still possible but we can catch easy cases here.
 	if (DataSetWriteForInput)
@@ -4063,7 +4059,7 @@ void FHlslNiagaraTranslator::WriteDataSet(const FNiagaraDataSetID DataSet, const
 
 	//First chunk for a write is always the condition pin.
 	//We always write the event payload into the temp storage but we can access this condition to pass to the final actual write to the buffer.
-	
+
 	DataSetWriteForInput->CodeChunks.Add(AddBodyChunk(DataSetAccessName + TEXT("_Valid"), TEXT("{0}"), FNiagaraTypeDefinition::GetBoolDef(), Inputs[1], false));
 	for (int32 i = 0; i < Variables.Num(); ++i)
 	{
@@ -4142,7 +4138,7 @@ int32 FHlslNiagaraTranslator::RegisterDataInterface(FNiagaraVariable& Var, UNiag
 		}
 		CompilationOutput.ScriptData.DataInterfaceInfo[Idx].RegisteredParameterMapRead = UsageName;
 	}
-	
+
 	return Idx;
 }
 
@@ -4152,7 +4148,7 @@ void FHlslNiagaraTranslator::Operation(class UNiagaraNodeOp* Operation, TArray<i
 
 	// Use the pins to determine the output type here since they may have been changed due to numeric pin fix up.
 	const FNiagaraOpInfo* OpInfo = FNiagaraOpInfo::GetOpInfo(Operation->OpName);
-	
+
 	TArray<UEdGraphPin*> OutputPins;
 	Operation->GetOutputPins(OutputPins);
 	for (int32 OutputIndex = 0; OutputIndex < OutputPins.Num(); OutputIndex++)
@@ -4193,7 +4189,7 @@ void FHlslNiagaraTranslator::FunctionCall(UNiagaraNodeFunctionCall* FunctionNode
 	TArray<UEdGraphPin*> CallInputs;
 	FunctionNode->GetOutputPins(CallOutputs);
 	FunctionNode->GetInputPins(CallInputs);
-	
+
 	// If the function call is disabled, we 
 	// need to route the input parameter map pin to the output parameter map pin.
 	// Any other outputs become invalid.
@@ -4260,7 +4256,7 @@ void FHlslNiagaraTranslator::FunctionCall(UNiagaraNodeFunctionCall* FunctionNode
 
 	ENiagaraScriptUsage ScriptUsage = ENiagaraScriptUsage::Function;
 	FString Name;
-	FString FullName; 
+	FString FullName;
 	UNiagaraScriptSource* Source = nullptr;
 	bool bCustomHlsl = false;
 	FString CustomHlsl;
@@ -4280,8 +4276,8 @@ void FHlslNiagaraTranslator::FunctionCall(UNiagaraNodeFunctionCall* FunctionNode
 		// All of the arguments here are resolved withing the HandleCustomHlsl function..
 		HandleCustomHlslNode(CustomFunctionHlsl, ScriptUsage, Name, FullName, bCustomHlsl, CustomHlsl, Signature, Inputs);
 	}
-	
-	RegisterFunctionCall(ScriptUsage, Name, FullName,  FunctionNode->NodeGuid, Source, Signature, bCustomHlsl, CustomHlsl, Inputs, CallInputs, CallOutputs, OutputSignature);
+
+	RegisterFunctionCall(ScriptUsage, Name, FullName, FunctionNode->NodeGuid, Source, Signature, bCustomHlsl, CustomHlsl, Inputs, CallInputs, CallOutputs, OutputSignature);
 
 	if (OutputSignature.IsValid() == false)
 	{
@@ -4357,7 +4353,7 @@ void FHlslNiagaraTranslator::FinalResolveNamespacedTokens(const FString& Paramet
 
 					}
 
-					
+
 					Tokens[i] = ParameterMapInstanceNamespace + GetSanitizedSymbolName(Var.GetName().ToString());
 					break;
 				}
@@ -4366,7 +4362,7 @@ void FHlslNiagaraTranslator::FinalResolveNamespacedTokens(const FString& Paramet
 	}
 }
 
-void FHlslNiagaraTranslator::HandleCustomHlslNode(UNiagaraNodeCustomHlsl* CustomFunctionHlsl, ENiagaraScriptUsage& OutScriptUsage, FString& OutName, FString& OutFullName, bool& bOutCustomHlsl, FString& OutCustomHlsl, 
+void FHlslNiagaraTranslator::HandleCustomHlslNode(UNiagaraNodeCustomHlsl* CustomFunctionHlsl, ENiagaraScriptUsage& OutScriptUsage, FString& OutName, FString& OutFullName, bool& bOutCustomHlsl, FString& OutCustomHlsl,
 	FNiagaraFunctionSignature& OutSignature, TArray<int32>& Inputs)
 {
 	SCOPE_CYCLE_COUNTER(STAT_NiagaraEditor_Module_NiagaraHLSLTranslator_CustomHLSL);
@@ -4508,7 +4504,7 @@ void FHlslNiagaraTranslator::HandleCustomHlslNode(UNiagaraNodeCustomHlsl* Custom
 	OutCustomHlsl = TEXT("\n") + OutCustomHlsl + TEXT("\n");
 }
 
-void FHlslNiagaraTranslator::RegisterFunctionCall(ENiagaraScriptUsage ScriptUsage, const FString& InName, const FString& InFullName, const FGuid& CallNodeId, UNiagaraScriptSource* Source,  FNiagaraFunctionSignature& InSignature, bool bIsCustomHlsl, const FString& InCustomHlsl, TArray<int32>& Inputs, const TArray<UEdGraphPin*>& CallInputs, const TArray<UEdGraphPin*>& CallOutputs,
+void FHlslNiagaraTranslator::RegisterFunctionCall(ENiagaraScriptUsage ScriptUsage, const FString& InName, const FString& InFullName, const FGuid& CallNodeId, UNiagaraScriptSource* Source, FNiagaraFunctionSignature& InSignature, bool bIsCustomHlsl, const FString& InCustomHlsl, TArray<int32>& Inputs, const TArray<UEdGraphPin*>& CallInputs, const TArray<UEdGraphPin*>& CallOutputs,
 	FNiagaraFunctionSignature& OutSignature)
 {
 	SCOPE_CYCLE_COUNTER(STAT_NiagaraEditor_Module_NiagaraHLSLTranslator_RegisterFunctionCall);
@@ -4531,28 +4527,28 @@ void FHlslNiagaraTranslator::RegisterFunctionCall(ENiagaraScriptUsage ScriptUsag
 			}
 		}
 
-		
+
 		bool bHasParameterMapParameters = SourceGraph->HasParameterMapParameters();
 
 		GenerateFunctionSignature(ScriptUsage, InName, InFullName, SourceGraph, Inputs, bHasNumericInputs, bHasParameterMapParameters, OutSignature);
 
-// 		//Sort the input and outputs to match the sorted parameters. They may be different.
-// 		TArray<FNiagaraVariable> OrderedInputs;
-// 		TArray<FNiagaraVariable> OrderedOutputs;
-// 		SourceGraph->GetParameters(OrderedInputs, OrderedOutputs);
-// 		TArray<UEdGraphPin*> InPins;
-// 		FunctionNode->GetInputPins(InPins);
-// 
-// 		TArray<int32> OrderedInputChunks;
-// 		OrderedInputChunks.SetNumUninitialized(Inputs.Num());
-// 		for (int32 i = 0; i < InPins.Num(); ++i)
-// 		{
-// 			FNiagaraVariable PinVar(Schema->PinToTypeDefinition(InPins[i]), *InPins[i]->PinName);
-// 			int32 InputIdx = OrderedInputs.IndexOfByKey(PinVar);
-// 			check(InputIdx != INDEX_NONE);
-// 			OrderedInputChunks[i] = Inputs[InputIdx];
-// 		}
-// 		Inputs = OrderedInputChunks;
+		// 		//Sort the input and outputs to match the sorted parameters. They may be different.
+		// 		TArray<FNiagaraVariable> OrderedInputs;
+		// 		TArray<FNiagaraVariable> OrderedOutputs;
+		// 		SourceGraph->GetParameters(OrderedInputs, OrderedOutputs);
+		// 		TArray<UEdGraphPin*> InPins;
+		// 		FunctionNode->GetInputPins(InPins);
+		// 
+		// 		TArray<int32> OrderedInputChunks;
+		// 		OrderedInputChunks.SetNumUninitialized(Inputs.Num());
+		// 		for (int32 i = 0; i < InPins.Num(); ++i)
+		// 		{
+		// 			FNiagaraVariable PinVar(Schema->PinToTypeDefinition(InPins[i]), *InPins[i]->PinName);
+		// 			int32 InputIdx = OrderedInputs.IndexOfByKey(PinVar);
+		// 			check(InputIdx != INDEX_NONE);
+		// 			OrderedInputChunks[i] = Inputs[InputIdx];
+		// 		}
+		// 		Inputs = OrderedInputChunks;
 
 		FString* FuncBody = Functions.Find(OutSignature);
 		if (!FuncBody)
@@ -4714,7 +4710,7 @@ void FHlslNiagaraTranslator::RegisterFunctionCall(ENiagaraScriptUsage ScriptUsag
 		check(InSignature.bMemberFunction || bIsCustomHlsl);
 		check(Inputs.Num() > 0);
 
-		OutSignature = InSignature;		
+		OutSignature = InSignature;
 
 		//First input for these is the owner of the function.
 		if (bIsCustomHlsl)
@@ -4817,7 +4813,7 @@ void FHlslNiagaraTranslator::GenerateFunctionCall(FNiagaraFunctionSignature& Fun
 		{
 			if (!AddStructToDefinitionSet(Type))
 			{
-				Error(FText::Format(LOCTEXT("GetConstantFailTypeVar2", "Cannot handle type {0}! Variable: {1}"),  Type.GetNameText(), FText::FromName(FunctionSignature.Inputs[i].GetName())), nullptr, nullptr);
+				Error(FText::Format(LOCTEXT("GetConstantFailTypeVar2", "Cannot handle type {0}! Variable: {1}"), Type.GetNameText(), FText::FromName(FunctionSignature.Inputs[i].GetName())), nullptr, nullptr);
 			}
 
 			int32 Input = Inputs[i];
@@ -4862,7 +4858,7 @@ void FHlslNiagaraTranslator::GenerateFunctionCall(FNiagaraFunctionSignature& Fun
 			{
 				Error(FText::Format(LOCTEXT("GetConstantFailTypeVar3", "Cannot handle type {0}! Variable: {1}"), Type.GetNameText(), FText::FromName(FunctionSignature.Outputs[i].GetName())), nullptr, nullptr);
 			}
-			
+
 			int32 Output = INDEX_NONE;
 			int32 ParamOutput = INDEX_NONE;
 			bool bSkip = false;
@@ -4956,7 +4952,7 @@ FString FHlslNiagaraTranslator::GetFunctionSignature(const FNiagaraFunctionSigna
 		const FNiagaraVariable& Input = Sig.Inputs[i];
 		//We don't write class types as real params in the hlsl
 		if (Input.GetType().GetClass() == nullptr)
-		{	
+		{
 			if (Input.GetType() == FNiagaraTypeDefinition::GetParameterMapDef())
 			{
 				// Skip parameter maps.
@@ -5168,13 +5164,31 @@ void FHlslNiagaraTranslator::Convert(class UNiagaraNodeConvert* Convert, TArray
 	TArray<UEdGraphPin*> OutputPins;
 	Convert->GetOutputPins(OutputPins);
 
+	// Add input struct definitions if necessary.
+	for (UEdGraphPin* InputPin : InputPins)
+	{
+		if (InputPin->PinType.PinCategory == UEdGraphSchema_Niagara::PinCategoryType ||
+			InputPin->PinType.PinCategory == UEdGraphSchema_Niagara::PinCategoryEnum)
+		{
+			FNiagaraTypeDefinition Type = Schema->PinToTypeDefinition(InputPin);
+			if (!AddStructToDefinitionSet(Type))
+			{
+				Error(FText::Format(LOCTEXT("ConvertTypeError", "Cannot handle input pin type {0}! Pin: {1}"), Type.GetNameText(), InputPin->GetDisplayName()), nullptr, nullptr);
+			}
+		}
+	}
+
 	// Generate outputs.
 	for (UEdGraphPin* OutputPin : OutputPins)
 	{
-		if (OutputPin->PinType.PinCategory == UEdGraphSchema_Niagara::PinCategoryType || 
+		if (OutputPin->PinType.PinCategory == UEdGraphSchema_Niagara::PinCategoryType ||
 			OutputPin->PinType.PinCategory == UEdGraphSchema_Niagara::PinCategoryEnum)
 		{
 			FNiagaraTypeDefinition Type = Schema->PinToTypeDefinition(OutputPin);
+			if (!AddStructToDefinitionSet(Type))
+			{
+				Error(FText::Format(LOCTEXT("ConvertTypeError", "Cannot handle output pin type {0}! Pin: {1}"), Type.GetNameText(), OutputPin->GetDisplayName()), nullptr, nullptr);
+			}
 			int32 OutChunk = AddBodyChunk(GetUniqueSymbolName(OutputPin->PinName), TEXT(""), Type);
 			Outputs.Add(OutChunk);
 		}
@@ -5191,13 +5205,21 @@ void FHlslNiagaraTranslator::Convert(class UNiagaraNodeConvert* Convert, TArray
 		if (SourceIndex != INDEX_NONE && SourceIndex < Inputs.Num() && DestinationIndex != INDEX_NONE && DestinationIndex < Outputs.Num())
 		{
 			FNiagaraTypeDefinition SrcPinType = Schema->PinToTypeDefinition(InputPins[SourceIndex]);
+			if (!AddStructToDefinitionSet(SrcPinType))
+			{
+				Error(FText::Format(LOCTEXT("ConvertTypeError", "Cannot handle input subpin type {0}! Subpin: {1}"), SrcPinType.GetNameText(), InputPins[SourceIndex]->GetDisplayName()), nullptr, nullptr);
+			}
 			TArray<FName> ConditionedSourcePath = ConditionPropertyPath(SrcPinType, Connection.SourcePath);
-			
+
 			FNiagaraTypeDefinition DestPinType = Schema->PinToTypeDefinition(OutputPins[DestinationIndex]);
+			if (!AddStructToDefinitionSet(DestPinType))
+			{
+				Error(FText::Format(LOCTEXT("ConvertTypeError", "Cannot handle output subpin type type {0}! Subpin: {1}"), DestPinType.GetNameText(), OutputPins[SourceIndex]->GetDisplayName()), nullptr, nullptr);
+			}
 			TArray<FName> ConditionedDestinationPath = ConditionPropertyPath(DestPinType, Connection.DestinationPath);
 
 			FString ConvertDefinition = GenerateAssignment(SrcPinType, ConditionedSourcePath, DestPinType, ConditionedDestinationPath);
-			
+
 			TArray<int32> SourceChunks;
 			SourceChunks.Add(Outputs[DestinationIndex]);
 			SourceChunks.Add(Inputs[SourceIndex]);
@@ -5236,7 +5258,7 @@ void FHlslNiagaraTranslator::If(UNiagaraNodeIf* IfNode, TArray<FNiagaraVariable>
 		FNiagaraTypeDefinition OutChunkType = CodeChunks[Outputs[i]].Type;
 		FNiagaraCodeChunk& BranchChunk = CodeChunks[AddBodyChunk(OutSymbols[i], TEXT("{0}"), OutChunkType, false)];
 		BranchChunk.AddSourceChunk(PathA[i]);
-	}	
+	}
 	AddBodyChunk(TEXT(""), TEXT("}\n\telse\n\t{"), FNiagaraTypeDefinition::GetFloatDef(), false, false);
 	for (int32 i = 0; i < NumVars; ++i)
 	{
@@ -5295,7 +5317,7 @@ int32 FHlslNiagaraTranslator::CompilePin(const UEdGraphPin* Pin)
 int32 FHlslNiagaraTranslator::CompileOutputPin(const UEdGraphPin* InPin)
 {
 	SCOPE_CYCLE_COUNTER(STAT_NiagaraEditor_HlslTranslator_CompileOutputPin);
-	
+
 	// The incoming pin to compile may be pointing to a reroute node. If so, we just jump over it
 	// to where it really came from.
 	UEdGraphPin* Pin = UNiagaraNode::TraceOutputPin(const_cast<UEdGraphPin*>(InPin));
@@ -5373,7 +5395,7 @@ void FHlslNiagaraTranslator::Error(FText ErrorText, const UNiagaraNode* Node, co
 		NodePinStr += TEXT(" Pin: ") + (Pin->PinFriendlyName.ToString().Len() > 0 ? Pin->PinFriendlyName.ToString() : Pin->GetName());
 		NodePinSuffix = TEXT(" - ");
 	}
-	
+
 	FString ErrorString = FString::Printf(TEXT("%s%s%s%sCallstack: %s"), *ErrorText.ToString(), *NodePinPrefix, *NodePinStr, *NodePinSuffix, *GetCallstack());
 	TranslateResults.CompileEvents.Add(FNiagaraCompileEvent(FNiagaraCompileEventType::Error, ErrorString, Node ? Node->NodeGuid : FGuid(), Pin ? Pin->PersistentGuid : FGuid(), GetCallstackGuids()));
 	TranslateResults.NumErrors++;
@@ -5409,7 +5431,7 @@ bool FHlslNiagaraTranslator::GetFunctionParameter(const FNiagaraVariable& Parame
 		int32 ParamIdx = FunctionContext->Signature.Inputs.IndexOfByPredicate([&](const FNiagaraVariable& InVar) { return InVar.IsEquivalent(Parameter); });
 		if (ParamIdx != INDEX_NONE)
 		{
-			OutParam = FunctionContext->Inputs[ParamIdx];			
+			OutParam = FunctionContext->Inputs[ParamIdx];
 		}
 		return true;
 	}
@@ -5516,7 +5538,7 @@ FString FHlslNiagaraTranslator::GetHlslDefaultForType(FNiagaraTypeDefinition Typ
 
 bool FHlslNiagaraTranslator::IsBuiltInHlslType(FNiagaraTypeDefinition Type)
 {
-	return	
+	return
 		Type == FNiagaraTypeDefinition::GetFloatDef() ||
 		Type == FNiagaraTypeDefinition::GetVec2Def() ||
 		Type == FNiagaraTypeDefinition::GetVec3Def() ||
@@ -5620,7 +5642,7 @@ FString FHlslNiagaraTranslator::BuildHLSLStructDecl(FNiagaraTypeDefinition Type,
 			FString PropertyTypeName = GetPropertyHlslTypeName(Property);
 			if (PropertyTypeName.IsEmpty())
 			{
-				OutErrorMessage = FText::Format(LOCTEXT("UnknownPropertyTypeErrorFormat", "Failed to build hlsl struct declaration for type {0}.  Property {1} has an unsuported type {2}."), 
+				OutErrorMessage = FText::Format(LOCTEXT("UnknownPropertyTypeErrorFormat", "Failed to build hlsl struct declaration for type {0}.  Property {1} has an unsuported type {2}."),
 					FText::FromString(Type.GetName()), Property->GetDisplayNameText(), FText::FromString(Property->GetClass()->GetName()));
 				return TEXT("");
 			}
@@ -5685,7 +5707,7 @@ bool FHlslNiagaraTranslator::AddStructToDefinitionSet(const FNiagaraTypeDefiniti
 				}
 			}
 		}
-		
+
 		// Add the new type def
 		StructsToDefine.AddUnique(TypeDef);
 	}
@@ -5795,7 +5817,7 @@ FString FHlslNiagaraTranslator::CompileDataInterfaceFunction(UNiagaraDataInterfa
 		return TEXT("");
 	}
 	else
-	{		
+	{
 		return TEXT("");
 		check(0);
 	}
