Fix deprecated modules showing up in script context menu, stack context menu and other search windows.

diff --git a/Engine/Plugins/FX/Niagara/Source/Niagara/Classes/NiagaraScript.h b/Engine/Plugins/FX/Niagara/Source/Niagara/Classes/NiagaraScript.h
--- a/Engine/Plugins/FX/Niagara/Source/Niagara/Classes/NiagaraScript.h
+++ b/Engine/Plugins/FX/Niagara/Source/Niagara/Classes/NiagaraScript.h
@@ -327,11 +327,11 @@ public:
 	TArray<FNiagaraModuleDependency> RequiredDependencies;
 
 	/* If this script is no longer meant to be used, this option should be set.*/
-	UPROPERTY(EditAnywhere, Category = "Script")
+	UPROPERTY(AssetRegistrySearchable, EditAnywhere, Category = Script)
 	uint32 bDeprecated : 1;
 
 	/* Which script to use if this is deprecated.*/
-	UPROPERTY(EditAnywhere, Category = "Script", meta = (EditCondition = "bDeprecated"))
+	UPROPERTY(EditAnywhere, Category = Script, meta = (EditCondition = "bDeprecated"))
 	UNiagaraScript* DeprecationRecommendation;
 
 #endif
diff --git a/Engine/Plugins/FX/Niagara/Source/NiagaraEditor/Private/EdGraphSchema_Niagara.cpp b/Engine/Plugins/FX/Niagara/Source/NiagaraEditor/Private/EdGraphSchema_Niagara.cpp
--- a/Engine/Plugins/FX/Niagara/Source/NiagaraEditor/Private/EdGraphSchema_Niagara.cpp
+++ b/Engine/Plugins/FX/Niagara/Source/NiagaraEditor/Private/EdGraphSchema_Niagara.cpp
@@ -393,82 +393,44 @@ TArray<TSharedPtr<FNiagaraSchemaAction_NewNode> > UEdGraphSchema_Niagara::GetGra
 	};
 
 	//Add functions
-	FAssetRegistryModule& AssetRegistryModule = FModuleManager::LoadModuleChecked<FAssetRegistryModule>(TEXT("AssetRegistry"));
-	TArray<FAssetData> ScriptAssets;
-	AssetRegistryModule.Get().GetAssetsByClass(UNiagaraScript::StaticClass()->GetFName(), ScriptAssets);
-	UEnum* NiagaraScriptUsageEnum = FindObjectChecked<UEnum>(ANY_PACKAGE, TEXT("ENiagaraScriptUsage"), true);
 	if (GbAllowAllNiagaraNodesInEmitterGraphs || bModuleGraph || bFunctionGraph || bDynamicInputGraph)
 	{
-		for (const FAssetData& ScriptAsset : ScriptAssets)
-		{
-			FName UsageName;
-			ScriptAsset.GetTagValue(GET_MEMBER_NAME_CHECKED(UNiagaraScript, Usage), UsageName);
+		TArray<FAssetData> FunctionScriptAssets;
+		FNiagaraEditorUtilities::FGetFilteredScriptAssetsOptions FunctionScriptFilterOptions;
+		FunctionScriptFilterOptions.ScriptUsageToInclude = ENiagaraScriptUsage::Function;
+		FNiagaraEditorUtilities::GetFilteredScriptAssets(FunctionScriptFilterOptions, FunctionScriptAssets);
 
-			FString QualifiedUsageName = "ENiagaraScriptUsage::" + UsageName.ToString();
-			int32 UsageIndex = NiagaraScriptUsageEnum->GetIndexByNameString(QualifiedUsageName);
-			if (UsageIndex != INDEX_NONE)
-			{
-				ENiagaraScriptUsage Usage = static_cast<ENiagaraScriptUsage>(NiagaraScriptUsageEnum->GetValueByIndex(UsageIndex));
-				if (Usage == ENiagaraScriptUsage::Function)
-				{
-					AddScriptFunctionAction(LOCTEXT("Function Menu Title", "Functions"), ScriptAsset);
-				}
-			}
+		for (const FAssetData& FunctionScriptAsset : FunctionScriptAssets)
+		{
+			AddScriptFunctionAction(LOCTEXT("Function Menu Title", "Functions"), FunctionScriptAsset);
 		}
 	}
 
 	//Add modules
 	if (!bFunctionGraph)
 	{
-		for (const FAssetData& ScriptAsset : ScriptAssets)
-		{
-			FName UsageName;
-			ScriptAsset.GetTagValue(GET_MEMBER_NAME_CHECKED(UNiagaraScript, Usage), UsageName);
-			const FString BitfieldTagValue = ScriptAsset.GetTagValueRef<FString>(GET_MEMBER_NAME_CHECKED(UNiagaraScript, ModuleUsageBitmask));
-			int32 BitfieldValue = FCString::Atoi(*BitfieldTagValue);
-
-			ENiagaraScriptUsage TargetUsage = ENiagaraScriptUsage::Module;
-			if (bParticleUpdateGraph)
-			{
-				TargetUsage = ENiagaraScriptUsage::ParticleUpdateScript;
-			}
-			else if (bSystemGraph)
-			{
-				TargetUsage = ENiagaraScriptUsage::SystemSpawnScript;
-			}
+		TArray<FAssetData> ModuleScriptAssets;
+		FNiagaraEditorUtilities::FGetFilteredScriptAssetsOptions ModuleScriptFilterOptions;
+		ModuleScriptFilterOptions.ScriptUsageToInclude = ENiagaraScriptUsage::Module;
+		FNiagaraEditorUtilities::GetFilteredScriptAssets(ModuleScriptFilterOptions, ModuleScriptAssets);
 
-			int32 TargetBit = (BitfieldValue >> (int32)TargetUsage) & 1;
-
-			FString QualifiedUsageName = "ENiagaraScriptUsage::" + UsageName.ToString();
-			int32 UsageIndex = NiagaraScriptUsageEnum->GetIndexByNameString(QualifiedUsageName);
-			if (UsageIndex != INDEX_NONE && TargetBit == 1)
-			{
-				ENiagaraScriptUsage Usage = static_cast<ENiagaraScriptUsage>(NiagaraScriptUsageEnum->GetValueByIndex(UsageIndex));
-				if (Usage == ENiagaraScriptUsage::Module)
-				{
-					AddScriptFunctionAction(LOCTEXT("Module Menu Title", "Modules"), ScriptAsset);
-				}
-			}
+		for (const FAssetData& ModuleScriptAsset : ModuleScriptAssets)
+		{
+			AddScriptFunctionAction(LOCTEXT("Module Menu Title", "Modules"), ModuleScriptAsset);
 		}
 	}
 
 	// Add dynamic inputs for default usage in module and dynamic input graphs
 	if (bModuleGraph || bDynamicInputGraph)
 	{
-		for (const FAssetData& ScriptAsset : ScriptAssets)
+		TArray<FAssetData> DynamicInputScriptAssets;
+		FNiagaraEditorUtilities::FGetFilteredScriptAssetsOptions DynamicInputScriptFilterOptions;
+		DynamicInputScriptFilterOptions.ScriptUsageToInclude = ENiagaraScriptUsage::DynamicInput;
+		FNiagaraEditorUtilities::GetFilteredScriptAssets(DynamicInputScriptFilterOptions, DynamicInputScriptAssets);
+
+		for (const FAssetData& DynamicInputScriptAsset : DynamicInputScriptAssets)
 		{
-			FName UsageName;
-			ScriptAsset.GetTagValue(GET_MEMBER_NAME_CHECKED(UNiagaraScript, Usage), UsageName);
-			FString QualifiedUsageName = "ENiagaraScriptUsage::" + UsageName.ToString();
-			int32 UsageIndex = NiagaraScriptUsageEnum->GetIndexByNameString(QualifiedUsageName);
-			if (UsageIndex != INDEX_NONE)
-			{
-				ENiagaraScriptUsage Usage = static_cast<ENiagaraScriptUsage>(NiagaraScriptUsageEnum->GetValueByIndex(UsageIndex));
-				if (Usage == ENiagaraScriptUsage::DynamicInput)
-				{
-					AddScriptFunctionAction(LOCTEXT("Dynamic Input Menu Title", "Dynamic Inputs"), ScriptAsset);
-				}
-			}
+			AddScriptFunctionAction(LOCTEXT("Dynamic Input Menu Title", "Dynamic Inputs"), DynamicInputScriptAsset);
 		}
 	}
 
diff --git a/Engine/Plugins/FX/Niagara/Source/NiagaraEditor/Private/NiagaraEditorUtilities.cpp b/Engine/Plugins/FX/Niagara/Source/NiagaraEditor/Private/NiagaraEditorUtilities.cpp
--- a/Engine/Plugins/FX/Niagara/Source/NiagaraEditor/Private/NiagaraEditorUtilities.cpp
+++ b/Engine/Plugins/FX/Niagara/Source/NiagaraEditor/Private/NiagaraEditorUtilities.cpp
@@ -772,6 +772,66 @@ void FNiagaraEditorUtilities::PreprocessFunctionGraph(const UEdGraphSchema_Niaga
 
 }
 
+void FNiagaraEditorUtilities::GetFilteredScriptAssets(FGetFilteredScriptAssetsOptions InFilter, TArray<FAssetData>& OutFilteredScriptAssets)
+{
+	FARFilter ScriptFilter;
+	ScriptFilter.ClassNames.Add(UNiagaraScript::StaticClass()->GetFName());
+
+	const UEnum* NiagaraScriptUsageEnum = FindObjectChecked<UEnum>(ANY_PACKAGE, TEXT("ENiagaraScriptUsage"), true);
+	const FString QualifiedScriptUsageString = NiagaraScriptUsageEnum->GetNameStringByValue(static_cast<uint8>(InFilter.ScriptUsageToInclude));
+	int32 LastColonIndex;
+	QualifiedScriptUsageString.FindLastChar(TEXT(':'), LastColonIndex);
+	const FString UnqualifiedScriptUsageString = QualifiedScriptUsageString.RightChop(LastColonIndex + 1);
+	ScriptFilter.TagsAndValues.Add(GET_MEMBER_NAME_CHECKED(UNiagaraScript, Usage), UnqualifiedScriptUsageString);
+
+	const FAssetRegistryModule& AssetRegistryModule = FModuleManager::LoadModuleChecked<FAssetRegistryModule>(TEXT("AssetRegistry"));
+	AssetRegistryModule.Get().GetAssets(ScriptFilter, OutFilteredScriptAssets);
+
+	// We remove deprecated scripts separately as FARFilter does not support filtering by non-string tags.
+	if (InFilter.bIncludeDeprecatedScripts == false)
+	{
+		bool bScriptIsDeprecated = false;
+		bool bFoundDeprecatedTag = false;
+		for (int i = OutFilteredScriptAssets.Num() - 1; i >= 0; --i)
+		{
+			bFoundDeprecatedTag = OutFilteredScriptAssets[i].GetTagValue(GET_MEMBER_NAME_CHECKED(UNiagaraScript, bDeprecated), bScriptIsDeprecated);
+			// If the asset does not have the metadata tag, check if it is loaded and if so check the bDeprecated value directly
+			if (bFoundDeprecatedTag == false)
+			{
+				if (OutFilteredScriptAssets[i].IsAssetLoaded())
+				{
+					UNiagaraScript* Script = static_cast<UNiagaraScript*>(OutFilteredScriptAssets[i].GetAsset());
+					if (Script != nullptr)
+					{
+						bScriptIsDeprecated = Script->bDeprecated;
+					}
+				}
+			}
+			if (bScriptIsDeprecated)
+			{
+				OutFilteredScriptAssets.RemoveAt(i);
+			}
+		}
+	}
+	// We remove scripts with non matching usage bitmasks separately as FARFilter does not support filtering by non-string tags.
+	if (InFilter.TargetUsageToMatch.IsSet())
+	{
+		FString BitfieldTagValue;
+		int32 BitfieldValue, TargetBit;
+		for (int i = OutFilteredScriptAssets.Num() - 1; i >= 0; --i)
+		{
+			BitfieldTagValue = OutFilteredScriptAssets[i].GetTagValueRef<FString>(GET_MEMBER_NAME_CHECKED(UNiagaraScript, ModuleUsageBitmask));
+			BitfieldValue = FCString::Atoi(*BitfieldTagValue);
+
+			TargetBit = (BitfieldValue >> (int32)InFilter.TargetUsageToMatch.GetValue()) & 1;
+			if (TargetBit != 1)
+			{
+				OutFilteredScriptAssets.RemoveAt(i);
+			}
+		}
+	}
+}
+
 UNiagaraNodeOutput* FNiagaraEditorUtilities::GetScriptOutputNode(UNiagaraScript& Script)
 {
 	UNiagaraScriptSource* Source = CastChecked<UNiagaraScriptSource>(Script.GetSource());
diff --git a/Engine/Plugins/FX/Niagara/Source/NiagaraEditor/Private/ViewModels/Stack/NiagaraStackFunctionInput.cpp b/Engine/Plugins/FX/Niagara/Source/NiagaraEditor/Private/ViewModels/Stack/NiagaraStackFunctionInput.cpp
--- a/Engine/Plugins/FX/Niagara/Source/NiagaraEditor/Private/ViewModels/Stack/NiagaraStackFunctionInput.cpp
+++ b/Engine/Plugins/FX/Niagara/Source/NiagaraEditor/Private/ViewModels/Stack/NiagaraStackFunctionInput.cpp
@@ -850,19 +850,10 @@ UNiagaraNodeCustomHlsl* UNiagaraStackFunctionInput::GetExpressionNode() const
 
 void UNiagaraStackFunctionInput::GetAvailableDynamicInputs(TArray<UNiagaraScript*>& AvailableDynamicInputs)
 {
-	UEnum* NiagaraScriptUsageEnum = FindObjectChecked<UEnum>(ANY_PACKAGE, TEXT("ENiagaraScriptUsage"), true);
-	FString QualifiedDynamicInputUsageString = NiagaraScriptUsageEnum->GetNameStringByValue(static_cast<uint8>(ENiagaraScriptUsage::DynamicInput));
-	int32 LastColonIndex;
-	QualifiedDynamicInputUsageString.FindLastChar(TEXT(':'), LastColonIndex);
-	FString UnqualifiedDynamicInputUsageString = QualifiedDynamicInputUsageString.RightChop(LastColonIndex + 1);
-
-	FARFilter DynamicInputFilter;
-	DynamicInputFilter.ClassNames.Add(UNiagaraScript::StaticClass()->GetFName());
-	DynamicInputFilter.TagsAndValues.Add(GET_MEMBER_NAME_CHECKED(UNiagaraScript, Usage), UnqualifiedDynamicInputUsageString);
-
-	FAssetRegistryModule& AssetRegistryModule = FModuleManager::LoadModuleChecked<FAssetRegistryModule>(TEXT("AssetRegistry"));
 	TArray<FAssetData> DynamicInputAssets;
-	AssetRegistryModule.Get().GetAssets(DynamicInputFilter, DynamicInputAssets);
+	FNiagaraEditorUtilities::FGetFilteredScriptAssetsOptions DynamicInputScriptFilterOptions;
+	DynamicInputScriptFilterOptions.ScriptUsageToInclude = ENiagaraScriptUsage::DynamicInput;
+	FNiagaraEditorUtilities::GetFilteredScriptAssets(DynamicInputScriptFilterOptions, DynamicInputAssets);
 
 	for (const FAssetData& DynamicInputAsset : DynamicInputAssets)
 	{
diff --git a/Engine/Plugins/FX/Niagara/Source/NiagaraEditor/Private/ViewModels/Stack/NiagaraStackGraphUtilities.cpp b/Engine/Plugins/FX/Niagara/Source/NiagaraEditor/Private/ViewModels/Stack/NiagaraStackGraphUtilities.cpp
--- a/Engine/Plugins/FX/Niagara/Source/NiagaraEditor/Private/ViewModels/Stack/NiagaraStackGraphUtilities.cpp
+++ b/Engine/Plugins/FX/Niagara/Source/NiagaraEditor/Private/ViewModels/Stack/NiagaraStackGraphUtilities.cpp
@@ -1236,37 +1236,6 @@ void FNiagaraStackGraphUtilities::GetNewParameterAvailableTypes(TArray<FNiagaraT
 	}
 }
 
-void FNiagaraStackGraphUtilities::GetScriptAssetsByUsage(ENiagaraScriptUsage AssetUsage, ENiagaraScriptUsage TargetUsage, TArray<FAssetData>& OutAssets)
-{
-	FAssetRegistryModule& AssetRegistryModule = FModuleManager::LoadModuleChecked<FAssetRegistryModule>(TEXT("AssetRegistry"));
-	TArray<FAssetData> ScriptAssets;
-	AssetRegistryModule.Get().GetAssetsByClass(UNiagaraScript::StaticClass()->GetFName(), ScriptAssets);
-
-	UEnum* NiagaraScriptUsageEnum = FindObjectChecked<UEnum>(ANY_PACKAGE, TEXT("ENiagaraScriptUsage"), true);
-
-	for (const FAssetData& ScriptAsset : ScriptAssets)
-	{
-		FName CandidateUsageName;
-		ScriptAsset.GetTagValue(GET_MEMBER_NAME_CHECKED(UNiagaraScript, Usage), CandidateUsageName);
-		FString QualifiedCandidateUsageName = "ENiagaraScriptUsage::" + CandidateUsageName.ToString();
-		int32 CandidateUsageIndex = NiagaraScriptUsageEnum->GetIndexByNameString(QualifiedCandidateUsageName);
-		if (CandidateUsageIndex != INDEX_NONE)
-		{
-			ENiagaraScriptUsage CandidateAssetUsage = static_cast<ENiagaraScriptUsage>(NiagaraScriptUsageEnum->GetValueByIndex(CandidateUsageIndex));
-			if (CandidateAssetUsage == AssetUsage)
-			{
-				const FString TargetUsageBitfieldTagValue = ScriptAsset.GetTagValueRef<FString>(GET_MEMBER_NAME_CHECKED(UNiagaraScript, ModuleUsageBitmask));
-				int32 TargetUsageBitfieldValue = FCString::Atoi(*TargetUsageBitfieldTagValue);
-				bool bTargetUsageBitfieldHasTargetUsage = ((TargetUsageBitfieldValue >> (int32)TargetUsage) & 1) == 1;
-				if (bTargetUsageBitfieldHasTargetUsage)
-				{
-					OutAssets.Add(ScriptAsset);
-				}
-			}
-		}
-	}
-}
-
 void FNiagaraStackGraphUtilities::GetScriptAssetsByDependencyProvided(ENiagaraScriptUsage AssetUsage, FName DependencyName, TArray<FAssetData>& OutAssets)
 {
 	FAssetRegistryModule& AssetRegistryModule = FModuleManager::LoadModuleChecked<FAssetRegistryModule>(TEXT("AssetRegistry"));
diff --git a/Engine/Plugins/FX/Niagara/Source/NiagaraEditor/Private/ViewModels/Stack/NiagaraStackScriptItemGroup.cpp b/Engine/Plugins/FX/Niagara/Source/NiagaraEditor/Private/ViewModels/Stack/NiagaraStackScriptItemGroup.cpp
--- a/Engine/Plugins/FX/Niagara/Source/NiagaraEditor/Private/ViewModels/Stack/NiagaraStackScriptItemGroup.cpp
+++ b/Engine/Plugins/FX/Niagara/Source/NiagaraEditor/Private/ViewModels/Stack/NiagaraStackScriptItemGroup.cpp
@@ -169,8 +169,11 @@ public:
 
 		// Generate actions for adding script asset modules.
 		TArray<FAssetData> ModuleAssets;
-		FNiagaraStackGraphUtilities::GetScriptAssetsByUsage(ENiagaraScriptUsage::Module, OutputNode->GetUsage(), ModuleAssets);
-		for(const FAssetData& ModuleAsset : ModuleAssets)
+		FNiagaraEditorUtilities::FGetFilteredScriptAssetsOptions ModuleScriptFilterOptions;
+		ModuleScriptFilterOptions.ScriptUsageToInclude = ENiagaraScriptUsage::Module;
+		ModuleScriptFilterOptions.TargetUsageToMatch = OutputNode->GetUsage();
+		FNiagaraEditorUtilities::GetFilteredScriptAssets(ModuleScriptFilterOptions, ModuleAssets);
+		for (const FAssetData& ModuleAsset : ModuleAssets)
 		{
 			OutAddActions.Add(FScriptGroupAddAction::CreateAssetModuleAction(ModuleAsset));
 		}
diff --git a/Engine/Plugins/FX/Niagara/Source/NiagaraEditor/Public/NiagaraEditorUtilities.h b/Engine/Plugins/FX/Niagara/Source/NiagaraEditor/Public/NiagaraEditorUtilities.h
--- a/Engine/Plugins/FX/Niagara/Source/NiagaraEditor/Public/NiagaraEditorUtilities.h
+++ b/Engine/Plugins/FX/Niagara/Source/NiagaraEditor/Public/NiagaraEditorUtilities.h
@@ -6,6 +6,7 @@
 #include "NiagaraCommon.h"
 #include "UObject/StructOnScope.h"
 #include "Misc/Attribute.h"
+#include "AssetData.h"
 
 class UNiagaraNodeInput;
 class UNiagaraNodeOutput;
@@ -137,6 +138,27 @@ namespace FNiagaraEditorUtilities
 
 	void PreprocessFunctionGraph(const UEdGraphSchema_Niagara* Schema, UNiagaraGraph* Graph, const TArray<UEdGraphPin*>& CallInputs, const TArray<UEdGraphPin*>& CallOutputs, ENiagaraScriptUsage ScriptUsage);
 
+	/** Options for the GetScriptsByFilter function. 
+	** @Param ScriptUsageToInclude Only return Scripts that have this usage
+	** @Param (Optional) TargetUsageToMatch Only return Scripts that have this target usage (output node) 
+	** @Param bIncludeDeprecatedScripts Whether or not to return Scripts that are deprecated (defaults to false) 
+	*/
+	struct FGetFilteredScriptAssetsOptions
+	{
+		FGetFilteredScriptAssetsOptions()
+			: ScriptUsageToInclude(ENiagaraScriptUsage::Module)
+			, TargetUsageToMatch()
+			, bIncludeDeprecatedScripts(false)
+		{
+		}
+
+		ENiagaraScriptUsage ScriptUsageToInclude;
+		TOptional<ENiagaraScriptUsage> TargetUsageToMatch;
+		bool bIncludeDeprecatedScripts;
+	};
+
+	NIAGARAEDITOR_API void GetFilteredScriptAssets(FGetFilteredScriptAssetsOptions InFilter, TArray<FAssetData>& OutFilteredScriptAssets); 
+
 	NIAGARAEDITOR_API UNiagaraNodeOutput* GetScriptOutputNode(UNiagaraScript& Script);
 
 	UNiagaraScript* GetScriptFromSystem(UNiagaraSystem& System, FGuid EmitterHandleId, ENiagaraScriptUsage Usage, FGuid UsageId);
diff --git a/Engine/Plugins/FX/Niagara/Source/NiagaraEditor/Public/ViewModels/Stack/NiagaraStackGraphUtilities.h b/Engine/Plugins/FX/Niagara/Source/NiagaraEditor/Public/ViewModels/Stack/NiagaraStackGraphUtilities.h
--- a/Engine/Plugins/FX/Niagara/Source/NiagaraEditor/Public/ViewModels/Stack/NiagaraStackGraphUtilities.h
+++ b/Engine/Plugins/FX/Niagara/Source/NiagaraEditor/Public/ViewModels/Stack/NiagaraStackGraphUtilities.h
@@ -128,8 +128,6 @@ namespace FNiagaraStackGraphUtilities
 
 	void GetNewParameterAvailableTypes(TArray<FNiagaraTypeDefinition>& OutAvailableTypes);
 
-	NIAGARAEDITOR_API void GetScriptAssetsByUsage(ENiagaraScriptUsage AssetUsage, ENiagaraScriptUsage TargetUsage, TArray<FAssetData>& OutAssets);
-
 	void GetScriptAssetsByDependencyProvided(ENiagaraScriptUsage AssetUsage, FName DependencyName, TArray<FAssetData>& OutAssets);
 
 	void GetAvailableParametersForScript(UNiagaraNodeOutput& ScriptOutputNode, TArray<FNiagaraVariable>& OutAvailableParameters);
diff --git a/Engine/Plugins/FX/Niagara/Source/NiagaraEditorWidgets/Private/Stack/SNiagaraStackModuleItem.cpp b/Engine/Plugins/FX/Niagara/Source/NiagaraEditorWidgets/Private/Stack/SNiagaraStackModuleItem.cpp
--- a/Engine/Plugins/FX/Niagara/Source/NiagaraEditorWidgets/Private/Stack/SNiagaraStackModuleItem.cpp
+++ b/Engine/Plugins/FX/Niagara/Source/NiagaraEditorWidgets/Private/Stack/SNiagaraStackModuleItem.cpp
@@ -352,7 +352,10 @@ void OnActionSelected(const TArray<TSharedPtr<FEdGraphSchemaAction>>& SelectedAc
 void CollectModuleActions(FGraphActionListBuilderBase& ModuleActions, UNiagaraStackModuleItem* ModuleItem)
 {
 	TArray<FAssetData> ModuleAssets;
-	FNiagaraStackGraphUtilities::GetScriptAssetsByUsage(ENiagaraScriptUsage::Module, ModuleItem->GetOutputNode()->GetUsage(), ModuleAssets);
+	FNiagaraEditorUtilities::FGetFilteredScriptAssetsOptions ModuleScriptFilterOptions;
+	ModuleScriptFilterOptions.ScriptUsageToInclude = ENiagaraScriptUsage::Module;
+	ModuleScriptFilterOptions.TargetUsageToMatch = ModuleItem->GetOutputNode()->GetUsage();
+	FNiagaraEditorUtilities::GetFilteredScriptAssets(ModuleScriptFilterOptions, ModuleAssets);
 	for (const FAssetData& ModuleAsset : ModuleAssets)
 	{
 		FText Category;
@@ -362,14 +365,6 @@ void CollectModuleActions(FGraphActionListBuilderBase& ModuleActions, UNiagaraSt
 			Category = LOCTEXT("ModuleNotCategorized", "Uncategorized Modules");
 		}
 
-		uint32 bDeprecatedModule;
-		ModuleAsset.GetTagValue(GET_MEMBER_NAME_CHECKED(UNiagaraScript, bDeprecated), bDeprecatedModule);
-
-		if (bDeprecatedModule)
-		{
-			continue;
-		}
-
 		FString DisplayNameString = FName::NameToDisplayString(ModuleAsset.AssetName.ToString(), false);
 		FText DisplayName = FText::FromString(DisplayNameString);
 
