Before we duplicate the system script we want to copy we first make sure that the scripts have been compiled. This allows the GPU script to have a proper compile id.
#tests ran through the repro steps from bug, works after change

diff --git a/Engine/Plugins/FX/Niagara/Source/NiagaraEditor/Private/NiagaraSystemFactoryNew.cpp b/Engine/Plugins/FX/Niagara/Source/NiagaraEditor/Private/NiagaraSystemFactoryNew.cpp
--- a/Engine/Plugins/FX/Niagara/Source/NiagaraEditor/Private/NiagaraSystemFactoryNew.cpp
+++ b/Engine/Plugins/FX/Niagara/Source/NiagaraEditor/Private/NiagaraSystemFactoryNew.cpp
@@ -116,6 +116,10 @@ UObject* UNiagaraSystemFactoryNew::FactoryCreateNew(UClass* Class, UObject* InPa
 	
 	if (SystemToCopy != nullptr)
 	{
+		if (SystemToCopy->IsReadyToRun() == false)
+		{
+			SystemToCopy->WaitForCompilationComplete();
+		}
 		NewSystem = Cast<UNiagaraSystem>(StaticDuplicateObject(SystemToCopy, InParent, Name, Flags, Class));
 		NewSystem->bIsTemplateAsset = false;
 		NewSystem->TemplateAssetDescription = FText();
