Fix light rendering for multiple light renderers in one component

#tests niagara automation tests passed

diff --git a/Engine/Plugins/FX/Niagara/Source/Niagara/Private/NiagaraComponent.cpp b/Engine/Plugins/FX/Niagara/Source/Niagara/Private/NiagaraComponent.cpp
--- a/Engine/Plugins/FX/Niagara/Source/Niagara/Private/NiagaraComponent.cpp
+++ b/Engine/Plugins/FX/Niagara/Source/Niagara/Private/NiagaraComponent.cpp
@@ -296,37 +296,32 @@ void FNiagaraSceneProxy::GetDynamicRayTracingInstances(FRayTracingMaterialGather
 
 void FNiagaraSceneProxy::GatherSimpleLights(const FSceneViewFamily& ViewFamily, FSimpleLightArray& OutParticleLights) const
 {
-	NiagaraRendererLights *LightRenderer = nullptr;
-	FNiagaraDynamicDataLights *DynamicData = nullptr;
 	for (int32 Idx = 0; Idx < EmitterRenderers.Num(); Idx++)
 	{
 		NiagaraRenderer *Renderer = EmitterRenderers[Idx];
 		if (Renderer && Renderer->GetPropertiesClass() == UNiagaraLightRendererProperties::StaticClass())
 		{
-			LightRenderer = static_cast<NiagaraRendererLights*>(Renderer);
-			DynamicData = static_cast<FNiagaraDynamicDataLights*>(Renderer->GetDynamicData());
-			break;
-		}
-	}
+			NiagaraRendererLights* LightRenderer = static_cast<NiagaraRendererLights*>(Renderer);
+			FNiagaraDynamicDataLights* DynamicData = static_cast<FNiagaraDynamicDataLights*>(Renderer->GetDynamicData());
 
+			if (DynamicData)
+			{
+				int32 LightCount = DynamicData->LightArray.Num();
 
-	if (DynamicData)
-	{
-		int32 LightCount = DynamicData->LightArray.Num();
-		
-		OutParticleLights.InstanceData.Reserve(LightCount);
-		OutParticleLights.PerViewData.Reserve(LightCount);
+				OutParticleLights.InstanceData.Reserve(LightCount + OutParticleLights.InstanceData.Num());
+				OutParticleLights.PerViewData.Reserve(LightCount + OutParticleLights.PerViewData.Num());
 
-		for (NiagaraRendererLights::SimpleLightData &LightData : DynamicData->LightArray)
-		{
-			// When not using camera-offset, output one position for all views to share. 
-			OutParticleLights.PerViewData.Add(LightData.PerViewEntry);
+				for (NiagaraRendererLights::SimpleLightData &LightData : DynamicData->LightArray)
+				{
+					// When not using camera-offset, output one position for all views to share. 
+					OutParticleLights.PerViewData.Add(LightData.PerViewEntry);
 
-			// Add an entry for the light instance.
-			OutParticleLights.InstanceData.Add(LightData.LightEntry);
+					// Add an entry for the light instance.
+					OutParticleLights.InstanceData.Add(LightData.LightEntry);
+				}
+			}
 		}
 	}
-
 }
 
 
