Taken into account "Niagara.Ribbon.Tessellation.MinAngle" when computing tesselation based on Niagara ribbon  twist.

diff --git a/Engine/Plugins/FX/Niagara/Source/Niagara/Private/NiagaraRendererRibbons.cpp b/Engine/Plugins/FX/Niagara/Source/Niagara/Private/NiagaraRendererRibbons.cpp
--- a/Engine/Plugins/FX/Niagara/Source/Niagara/Private/NiagaraRendererRibbons.cpp
+++ b/Engine/Plugins/FX/Niagara/Source/Niagara/Private/NiagaraRendererRibbons.cpp
@@ -281,7 +281,7 @@ void NiagaraRendererRibbons::GetDynamicMeshElements(const TArray<const FSceneVie
 			int32 NumSegments = DynamicDataRibbon->SegmentData.Num();
 			if (GNiagaraRibbonMaxTessellation > 1 && TessellationCurvature > SMALL_NUMBER && ViewFamily.GetFeatureLevel() == ERHIFeatureLevel::SM5)
 			{
-				const float MinTesselation = FMath::Max<float>(1.f, TessellationAngle / FMath::Max<float>(SMALL_NUMBER, GNiagaraRibbonTessellationAngle));
+				const float MinTesselation = FMath::Max<float>(1.f, FMath::Max(TessellationTwistAngle, TessellationAngle) / FMath::Max<float>(SMALL_NUMBER, GNiagaraRibbonTessellationAngle));
 
 				const float ViewDistance = SceneProxy->GetBounds().ComputeSquaredDistanceFromBoxToPoint(ViewOriginForDistanceCulling);
 				const float MaxDisplacementError = FMath::Max(GNiagaraRibbonTessellationMinDisplacementError, GNiagaraRibbonTessellationScreenPercentage * FMath::Sqrt(ViewDistance) / View->LODDistanceFactor);
