Sequencer: Add support for track editors to add tracks for all the selected object bindings

diff --git a/Engine/Plugins/Compositing/Composure/Source/ComposureEditor/Private/Sequencer/ComposureExportTrackEditor.cpp b/Engine/Plugins/Compositing/Composure/Source/ComposureEditor/Private/Sequencer/ComposureExportTrackEditor.cpp
--- a/Engine/Plugins/Compositing/Composure/Source/ComposureEditor/Private/Sequencer/ComposureExportTrackEditor.cpp
+++ b/Engine/Plugins/Compositing/Composure/Source/ComposureEditor/Private/Sequencer/ComposureExportTrackEditor.cpp
@@ -70,14 +70,14 @@ void FComposureExportTrackEditor::BuildTrackContextMenu(FMenuBuilder& MenuBuilde
 	MenuBuilder.AddSubMenu(LOCTEXT("Properties_MenuText", "Properties"), FText(), FNewMenuDelegate::CreateLambda(PopulateSubMenu));
 }
 
-void FComposureExportTrackEditor::BuildObjectBindingTrackMenu(FMenuBuilder& MenuBuilder, const FGuid& ObjectBinding, const UClass* ObjectClass)
+void FComposureExportTrackEditor::BuildObjectBindingTrackMenu(FMenuBuilder& MenuBuilder, const TArray<FGuid>& ObjectBindings, const UClass* ObjectClass)
 {
 	TSharedPtr<ISequencer> SequencerPtr = GetSequencer();
 	UMovieScene* MovieScene = SequencerPtr->GetFocusedMovieSceneSequence()->GetMovieScene();
 
 	TSet<FName> ExistingPasses;
 	{
-		const FMovieSceneBinding* Binding = MovieScene->FindBinding(ObjectBinding);
+		const FMovieSceneBinding* Binding = MovieScene->FindBinding(ObjectBindings[0]);
 		if (Binding)
 		{
 			for (UMovieSceneTrack* Track : Binding->GetTracks())
@@ -91,7 +91,7 @@ void FComposureExportTrackEditor::BuildObjectBindingTrackMenu(FMenuBuilder& Menu
 		}
 	}
 
-	for (TWeakObjectPtr<> WeakObject : SequencerPtr->FindObjectsInCurrentSequence(ObjectBinding))
+	for (TWeakObjectPtr<> WeakObject : SequencerPtr->FindObjectsInCurrentSequence(ObjectBindings[0]))
 	{
 		ACompositingElement* CompShotElement = Cast<ACompositingElement>(WeakObject.Get());
 		if (CompShotElement)
@@ -103,7 +103,7 @@ void FComposureExportTrackEditor::BuildObjectBindingTrackMenu(FMenuBuilder& Menu
 					LOCTEXT("AddOutputTrack_Tooltip", "Adds a new track that controls whether this composure element's output should be captured as part of a Sequencer capture."),
 					FSlateIcon(),
 					FUIAction(
-						FExecuteAction::CreateRaw(this, &FComposureExportTrackEditor::AddNewTrack, ObjectBinding, FName(), true, FName("Output"))
+						FExecuteAction::CreateRaw(this, &FComposureExportTrackEditor::AddNewTrack, ObjectBindings, FName(), true, FName("Output"))
 					)
 				);
 			}
@@ -121,7 +121,7 @@ void FComposureExportTrackEditor::BuildObjectBindingTrackMenu(FMenuBuilder& Menu
 						LOCTEXT("AddTrack_Tooltip", "Adds a new track that controls whether this composure element pass should be captured as part of a Sequencer capture."),
 						FSlateIcon(),
 						FUIAction(
-							FExecuteAction::CreateRaw(this, &FComposureExportTrackEditor::AddNewTrack, ObjectBinding, TransformPass->PassName, false, FName())
+							FExecuteAction::CreateRaw(this, &FComposureExportTrackEditor::AddNewTrack, ObjectBindings, TransformPass->PassName, false, FName())
 						)
 					);
 				}
@@ -132,22 +132,25 @@ void FComposureExportTrackEditor::BuildObjectBindingTrackMenu(FMenuBuilder& Menu
 	}
 }
 
-void FComposureExportTrackEditor::AddNewTrack(FGuid ObjectBinding, FName InPassName, bool bRenamePass, FName InExportAs)
+void FComposureExportTrackEditor::AddNewTrack(TArray<FGuid> ObjectBindings, FName InPassName, bool bRenamePass, FName InExportAs)
 {
 	TSharedPtr<ISequencer> SequencerPtr = GetSequencer();
 	if (SequencerPtr.IsValid())
 	{
 		FScopedTransaction Transaction(LOCTEXT("AddNewTrack_Transaction", "Add Composure Export Track"));
 
-		UMovieScene*                       MovieScene = SequencerPtr->GetFocusedMovieSceneSequence()->GetMovieScene();
-		UMovieSceneComposureExportTrack*   NewTrack   = MovieScene->AddTrack<UMovieSceneComposureExportTrack>(ObjectBinding);
-		UMovieSceneComposureExportSection* NewSection = Cast<UMovieSceneComposureExportSection>(NewTrack->CreateNewSection());
+		for (FGuid ObjectBinding : ObjectBindings)
+		{
+			UMovieScene*                       MovieScene = SequencerPtr->GetFocusedMovieSceneSequence()->GetMovieScene();
+			UMovieSceneComposureExportTrack*   NewTrack = MovieScene->AddTrack<UMovieSceneComposureExportTrack>(ObjectBinding);
+			UMovieSceneComposureExportSection* NewSection = Cast<UMovieSceneComposureExportSection>(NewTrack->CreateNewSection());
 
-		NewTrack->Pass.TransformPassName = InPassName;
-		NewTrack->Pass.bRenamePass       = bRenamePass;
-		NewTrack->Pass.ExportedAs        = InExportAs;
+			NewTrack->Pass.TransformPassName = InPassName;
+			NewTrack->Pass.bRenamePass = bRenamePass;
+			NewTrack->Pass.ExportedAs = InExportAs;
 
-		NewTrack->AddSection(*NewSection);
+			NewTrack->AddSection(*NewSection);
+		}
 
 		SequencerPtr->NotifyMovieSceneDataChanged(EMovieSceneDataChangeType::MovieSceneStructureItemAdded);
 	}
diff --git a/Engine/Plugins/Compositing/Composure/Source/ComposureEditor/Private/Sequencer/ComposureExportTrackEditor.h b/Engine/Plugins/Compositing/Composure/Source/ComposureEditor/Private/Sequencer/ComposureExportTrackEditor.h
--- a/Engine/Plugins/Compositing/Composure/Source/ComposureEditor/Private/Sequencer/ComposureExportTrackEditor.h
+++ b/Engine/Plugins/Compositing/Composure/Source/ComposureEditor/Private/Sequencer/ComposureExportTrackEditor.h
@@ -14,13 +14,13 @@ public:
 
 private:
 
-	virtual void BuildObjectBindingTrackMenu(FMenuBuilder& MenuBuilder, const FGuid& ObjectBinding, const UClass* ObjectClass) override;
+	virtual void BuildObjectBindingTrackMenu(FMenuBuilder& MenuBuilder, const TArray<FGuid>& ObjectBindings, const UClass* ObjectClass) override;
 	virtual void BuildTrackContextMenu(FMenuBuilder& MenuBuilder, UMovieSceneTrack* Track) override;
 	virtual bool SupportsType(TSubclassOf<UMovieSceneTrack> TrackClass) const override;
 	virtual bool SupportsSequence(UMovieSceneSequence* InSequence) const override;
 
 	void HandleActorAdded(AActor* Actor, FGuid TargetObjectGuid);
-	void AddNewTrack(FGuid ObjectBinding, FName InPassName, bool bRenamePass, FName ExportAs);
+	void AddNewTrack(TArray<FGuid> ObjectBindings, FName InPassName, bool bRenamePass, FName ExportAs);
 
 	FDelegateHandle OnActorAddedToSequencerHandle;
 };
\ No newline at end of file
diff --git a/Engine/Plugins/Experimental/ControlRig/Source/ControlRigEditor/Private/Sequencer/ControlRigBindingTrackEditor.cpp b/Engine/Plugins/Experimental/ControlRig/Source/ControlRigEditor/Private/Sequencer/ControlRigBindingTrackEditor.cpp
--- a/Engine/Plugins/Experimental/ControlRig/Source/ControlRigEditor/Private/Sequencer/ControlRigBindingTrackEditor.cpp
+++ b/Engine/Plugins/Experimental/ControlRig/Source/ControlRigEditor/Private/Sequencer/ControlRigBindingTrackEditor.cpp
@@ -20,11 +20,11 @@ FControlRigBindingTrackEditor::FControlRigBindingTrackEditor(TSharedRef<ISequenc
 }
 
 
-void FControlRigBindingTrackEditor::BuildObjectBindingTrackMenu(FMenuBuilder& MenuBuilder, const FGuid& ObjectBinding, const UClass* ObjectClass)
+void FControlRigBindingTrackEditor::BuildObjectBindingTrackMenu(FMenuBuilder& MenuBuilder, const TArray<FGuid>& ObjectBindings, const UClass* ObjectClass)
 {
 	UMovieSceneSequence* MovieSequence = GetSequencer()->GetFocusedMovieSceneSequence();
 
-	if (!MovieSequence || MovieSequence->GetClass()->GetName() != TEXT("LevelSequence") || !MovieSequence->GetMovieScene()->FindSpawnable(ObjectBinding))
+	if (!MovieSequence || MovieSequence->GetClass()->GetName() != TEXT("LevelSequence") || !MovieSequence->GetMovieScene()->FindSpawnable(ObjectBindings[0]))
 	{
 		return;
 	}
@@ -34,8 +34,8 @@ void FControlRigBindingTrackEditor::BuildObjectBindingTrackMenu(FMenuBuilder& Me
 		LOCTEXT("AddBindingTrackTooltip", "Adds a new track that controls the lifetime and binding of the animation controller."),
 		FSlateIcon(),
 		FUIAction(
-			FExecuteAction::CreateRaw(this, &FControlRigBindingTrackEditor::HandleAddBindingTrackMenuEntryExecute, ObjectBinding),
-			FCanExecuteAction::CreateSP(this, &FControlRigBindingTrackEditor::CanAddBindingTrack, ObjectBinding)
+			FExecuteAction::CreateRaw(this, &FControlRigBindingTrackEditor::HandleAddBindingTrackMenuEntryExecute, ObjectBindings),
+			FCanExecuteAction::CreateSP(this, &FControlRigBindingTrackEditor::CanAddBindingTrack, ObjectBindings[0])
 		)
 	);
 }
@@ -46,11 +46,16 @@ bool FControlRigBindingTrackEditor::SupportsType(TSubclassOf<UMovieSceneTrack> T
 	return (Type == UControlRigBindingTrack::StaticClass());
 }
 
-void FControlRigBindingTrackEditor::HandleAddBindingTrackMenuEntryExecute(FGuid ObjectBinding)
+void FControlRigBindingTrackEditor::HandleAddBindingTrackMenuEntryExecute(TArray<FGuid> ObjectBindings)
 {
 	FScopedTransaction AddSpawnTrackTransaction(LOCTEXT("AddBindingTrack_Transaction", "Add Binding Track"));
-	AddTrack(GetSequencer()->GetFocusedMovieSceneSequence()->GetMovieScene(), ObjectBinding, UControlRigBindingTrack::StaticClass(), NAME_None);
-	GetSequencer()->NotifyMovieSceneDataChanged( EMovieSceneDataChangeType::MovieSceneStructureItemAdded );
+
+	for (FGuid ObjectBinding : ObjectBindings)
+	{
+		AddTrack(GetSequencer()->GetFocusedMovieSceneSequence()->GetMovieScene(), ObjectBinding, UControlRigBindingTrack::StaticClass(), NAME_None);
+	}
+	
+	GetSequencer()->NotifyMovieSceneDataChanged(EMovieSceneDataChangeType::MovieSceneStructureItemAdded);
 
 	if (FControlRigEditMode* ControlRigEditMode = static_cast<FControlRigEditMode*>(GLevelEditorModeTools().GetActiveMode(FControlRigEditMode::ModeName)))
 	{
diff --git a/Engine/Plugins/Experimental/ControlRig/Source/ControlRigEditor/Private/Sequencer/ControlRigBindingTrackEditor.h b/Engine/Plugins/Experimental/ControlRig/Source/ControlRigEditor/Private/Sequencer/ControlRigBindingTrackEditor.h
--- a/Engine/Plugins/Experimental/ControlRig/Source/ControlRigEditor/Private/Sequencer/ControlRigBindingTrackEditor.h
+++ b/Engine/Plugins/Experimental/ControlRig/Source/ControlRigEditor/Private/Sequencer/ControlRigBindingTrackEditor.h
@@ -33,7 +33,7 @@ public:
 
 	// ISequencerTrackEditor interface
 
-	virtual void BuildObjectBindingTrackMenu(FMenuBuilder& MenuBuilder, const FGuid& ObjectBinding, const UClass* ObjectClass) override;
+	virtual void BuildObjectBindingTrackMenu(FMenuBuilder& MenuBuilder, const TArray<FGuid>& ObjectBindings, const UClass* ObjectClass) override;
 	virtual TSharedPtr<SWidget> BuildOutlinerEditWidget(const FGuid& ObjectBinding, UMovieSceneTrack* Track, const FBuildEditWidgetParams& Params) override { return TSharedPtr<SWidget>(); }
 	virtual bool HandleAssetAdded(UObject* Asset, const FGuid& TargetObjectGuid) override { return false; }
 	virtual bool SupportsType(TSubclassOf<UMovieSceneTrack> Type) const override;
@@ -41,6 +41,6 @@ public:
 private:
 
 	/** Callback for executing the "Binding Track" menu entry. */
-	void HandleAddBindingTrackMenuEntryExecute(FGuid ObjectBinding);
+	void HandleAddBindingTrackMenuEntryExecute(TArray<FGuid> ObjectBindings);
 	bool CanAddBindingTrack(FGuid ObjectBinding) const;
 };
diff --git a/Engine/Plugins/Experimental/ControlRig/Source/ControlRigEditor/Private/Sequencer/ControlRigTrackEditor.cpp b/Engine/Plugins/Experimental/ControlRig/Source/ControlRigEditor/Private/Sequencer/ControlRigTrackEditor.cpp
--- a/Engine/Plugins/Experimental/ControlRig/Source/ControlRigEditor/Private/Sequencer/ControlRigTrackEditor.cpp
+++ b/Engine/Plugins/Experimental/ControlRig/Source/ControlRigEditor/Private/Sequencer/ControlRigTrackEditor.cpp
@@ -215,7 +215,7 @@ void FControlRigTrackEditor::BuildAddTrackMenu(FMenuBuilder& MenuBuilder)
 	// do nothing
 }
 
-void FControlRigTrackEditor::BuildObjectBindingTrackMenu(FMenuBuilder& MenuBuilder, const FGuid& ObjectBinding, const UClass* ObjectClass)
+void FControlRigTrackEditor::BuildObjectBindingTrackMenu(FMenuBuilder& MenuBuilder, const TArray<FGuid>& ObjectBindings, const UClass* ObjectClass)
 {
 	if (ObjectClass->IsChildOf(USkeletalMeshComponent::StaticClass()) || ObjectClass->IsChildOf(AActor::StaticClass()))
 	{
@@ -225,16 +225,16 @@ void FControlRigTrackEditor::BuildObjectBindingTrackMenu(FMenuBuilder& MenuBuild
 
 		MenuBuilder.AddSubMenu(
 			LOCTEXT("AddControlRig", "Animation ControlRig"), NSLOCTEXT("Sequencer", "AddControlRigTooltip", "Adds an animation ControlRig track."),
-			FNewMenuDelegate::CreateRaw(this, &FControlRigTrackEditor::AddControlRigSubMenu, ObjectBinding, Track)
+			FNewMenuDelegate::CreateRaw(this, &FControlRigTrackEditor::AddControlRigSubMenu, ObjectBindings, Track)
 		);
 	}
 }
 
-TSharedRef<SWidget> FControlRigTrackEditor::BuildControlRigSubMenu(FGuid ObjectBinding, UMovieSceneTrack* Track)
+TSharedRef<SWidget> FControlRigTrackEditor::BuildControlRigSubMenu(const TArray<FGuid>& ObjectBindings, UMovieSceneTrack* Track)
 {
 	FMenuBuilder MenuBuilder(true, nullptr);
 	
-	AddControlRigSubMenu(MenuBuilder, ObjectBinding, Track);
+	AddControlRigSubMenu(MenuBuilder, ObjectBindings, Track);
 
 	return MenuBuilder.MakeWidget();
 }
@@ -265,14 +265,14 @@ void FControlRigTrackEditor::BuildTrackContextMenu( FMenuBuilder& MenuBuilder, U
 	
 }
 
-void FControlRigTrackEditor::AddControlRigSubMenu(FMenuBuilder& MenuBuilder, FGuid ObjectBinding, UMovieSceneTrack* Track)
+void FControlRigTrackEditor::AddControlRigSubMenu(FMenuBuilder& MenuBuilder, TArray<FGuid> ObjectBindings, UMovieSceneTrack* Track)
 {
 	MenuBuilder.BeginSection(TEXT("ChooseSequence"), LOCTEXT("ChooseSequence", "Choose Sequence"));
 	{
 		FAssetPickerConfig AssetPickerConfig;
 		{
-			AssetPickerConfig.OnAssetSelected = FOnAssetSelected::CreateRaw(this, &FControlRigTrackEditor::OnSequencerAssetSelected, ObjectBinding, Track);
-			AssetPickerConfig.OnAssetEnterPressed = FOnAssetEnterPressed::CreateRaw(this, &FControlRigTrackEditor::OnSequencerAssetEnterPressed, ObjectBinding, Track);
+			AssetPickerConfig.OnAssetSelected = FOnAssetSelected::CreateRaw(this, &FControlRigTrackEditor::OnSequencerAssetSelected, ObjectBindings, Track);
+			AssetPickerConfig.OnAssetEnterPressed = FOnAssetEnterPressed::CreateRaw(this, &FControlRigTrackEditor::OnSequencerAssetEnterPressed, ObjectBindings, Track);
 			AssetPickerConfig.bAllowNullSelection = false;
 			AssetPickerConfig.InitialAssetViewType = EAssetViewType::List;
 			AssetPickerConfig.Filter.bRecursiveClasses = true;
@@ -293,7 +293,7 @@ void FControlRigTrackEditor::AddControlRigSubMenu(FMenuBuilder& MenuBuilder, FGu
 	MenuBuilder.EndSection();
 }
 
-void FControlRigTrackEditor::OnSequencerAssetSelected(const FAssetData& AssetData, FGuid ObjectBinding, UMovieSceneTrack* Track)
+void FControlRigTrackEditor::OnSequencerAssetSelected(const FAssetData& AssetData, TArray<FGuid> ObjectBindings, UMovieSceneTrack* Track)
 {
 	FSlateApplication::Get().DismissAllMenus();
 
@@ -302,15 +302,21 @@ void FControlRigTrackEditor::OnSequencerAssetSelected(const FAssetData& AssetDat
 	if (SelectedObject && SelectedObject->IsA(UControlRigSequence::StaticClass()))
 	{
 		UControlRigSequence* Sequence = CastChecked<UControlRigSequence>(AssetData.GetAsset());
-		AnimatablePropertyChanged(FOnKeyProperty::CreateRaw(this, &FControlRigTrackEditor::AddKeyInternal, ObjectBinding, Sequence, Track));
+
+		const FScopedTransaction Transaction(LOCTEXT("AddControlRig_Transaction", "Add Control Rig"));
+
+		for (FGuid ObjectBinding : ObjectBindings)
+		{
+			AnimatablePropertyChanged(FOnKeyProperty::CreateRaw(this, &FControlRigTrackEditor::AddKeyInternal, ObjectBinding, Sequence, Track));
+		}
 	}
 }
 
-void FControlRigTrackEditor::OnSequencerAssetEnterPressed(const TArray<FAssetData>& AssetData, FGuid ObjectBinding, UMovieSceneTrack* Track)
+void FControlRigTrackEditor::OnSequencerAssetEnterPressed(const TArray<FAssetData>& AssetData, TArray<FGuid> ObjectBindings, UMovieSceneTrack* Track)
 {
 	if (AssetData.Num() > 0)
 	{
-		OnSequencerAssetSelected(AssetData[0].GetAsset(), ObjectBinding, Track);
+		OnSequencerAssetSelected(AssetData[0].GetAsset(), ObjectBindings, Track);
 	}
 }
 
@@ -353,7 +359,10 @@ TSharedRef<SWidget> FControlRigTrackEditor::HandleAddSubSequenceComboButtonGetMe
 {
 	FMenuBuilder MenuBuilder(true, nullptr);
 
-	AddControlRigSubMenu(MenuBuilder, ObjectBinding, InTrack);
+	TArray<FGuid> ObjectBindings;
+	ObjectBindings.Add(ObjectBinding);
+
+	AddControlRigSubMenu(MenuBuilder, ObjectBindings, InTrack);
 
 	return MenuBuilder.MakeWidget();
 }
diff --git a/Engine/Plugins/Experimental/ControlRig/Source/ControlRigEditor/Private/Sequencer/ControlRigTrackEditor.h b/Engine/Plugins/Experimental/ControlRig/Source/ControlRigEditor/Private/Sequencer/ControlRigTrackEditor.h
--- a/Engine/Plugins/Experimental/ControlRig/Source/ControlRigEditor/Private/Sequencer/ControlRigTrackEditor.h
+++ b/Engine/Plugins/Experimental/ControlRig/Source/ControlRigEditor/Private/Sequencer/ControlRigTrackEditor.h
@@ -49,7 +49,7 @@ public:
 
 	virtual TSharedPtr<SWidget> BuildOutlinerEditWidget(const FGuid& ObjectBinding, UMovieSceneTrack* Track, const FBuildEditWidgetParams& Params) override;
 	virtual void BuildAddTrackMenu(FMenuBuilder& MenuBuilder) override;
-	virtual void BuildObjectBindingTrackMenu(FMenuBuilder& MenuBuilder, const FGuid& ObjectBinding, const UClass* ObjectClass) override;
+	virtual void BuildObjectBindingTrackMenu(FMenuBuilder& MenuBuilder, const TArray<FGuid>& ObjectBindings, const UClass* ObjectClass) override;
 	virtual TSharedRef<ISequencerSection> MakeSectionInterface( UMovieSceneSection& SectionObject, UMovieSceneTrack& Track, FGuid ObjectBinding ) override;
 	virtual bool SupportsSequence(UMovieSceneSequence* InSequence) const override;
 	virtual bool SupportsType( TSubclassOf<UMovieSceneTrack> Type ) const override;
@@ -59,14 +59,14 @@ public:
 private:
 
 	/** ControlRig sub menu */
-	TSharedRef<SWidget> BuildControlRigSubMenu(FGuid ObjectBinding, UMovieSceneTrack* Track);
-	void AddControlRigSubMenu(FMenuBuilder& MenuBuilder, FGuid ObjectBinding, UMovieSceneTrack* Track);
+	TSharedRef<SWidget> BuildControlRigSubMenu(const TArray<FGuid>& ObjectBindings, UMovieSceneTrack* Track);
+	void AddControlRigSubMenu(FMenuBuilder& MenuBuilder, TArray<FGuid> ObjectBindings, UMovieSceneTrack* Track);
 
 	/** Animation asset selected */
-	void OnSequencerAssetSelected(const FAssetData& AssetData, FGuid ObjectBinding, UMovieSceneTrack* Track);
+	void OnSequencerAssetSelected(const FAssetData& AssetData, TArray<FGuid> ObjectBindings, UMovieSceneTrack* Track);
 
 	/** Animation asset enter pressed */
-	void OnSequencerAssetEnterPressed(const TArray<FAssetData>& AssetData, FGuid ObjectBinding, UMovieSceneTrack* Track);
+	void OnSequencerAssetEnterPressed(const TArray<FAssetData>& AssetData, TArray<FGuid> ObjectBindings, UMovieSceneTrack* Track);
 
 	/** Delegate for AnimatablePropertyChanged in AddKey */
 	FKeyPropertyResult AddKeyInternal(FFrameNumber KeyTime, FGuid ObjectBinding, UControlRigSequence* Sequence, UMovieSceneTrack* Track);
diff --git a/Engine/Plugins/Experimental/GeometryCache/Source/GeometryCacheSequencer/Classes/GeometryCacheTrackEditor.cpp b/Engine/Plugins/Experimental/GeometryCache/Source/GeometryCacheSequencer/Classes/GeometryCacheTrackEditor.cpp
--- a/Engine/Plugins/Experimental/GeometryCache/Source/GeometryCacheSequencer/Classes/GeometryCacheTrackEditor.cpp
+++ b/Engine/Plugins/Experimental/GeometryCache/Source/GeometryCacheSequencer/Classes/GeometryCacheTrackEditor.cpp
@@ -280,11 +280,11 @@ TSharedRef<ISequencerSection> FGeometryCacheTrackEditor::MakeSectionInterface(UM
 	return MakeShareable(new FGeometryCacheSection(SectionObject, GetSequencer()));
 }
 
-void FGeometryCacheTrackEditor::BuildObjectBindingTrackMenu(FMenuBuilder& MenuBuilder, const FGuid& ObjectBinding, const UClass* ObjectClass)
+void FGeometryCacheTrackEditor::BuildObjectBindingTrackMenu(FMenuBuilder& MenuBuilder, const TArray<FGuid>& ObjectBindings, const UClass* ObjectClass)
 {
 	if (ObjectClass->IsChildOf(UGeometryCacheComponent::StaticClass()) || ObjectClass->IsChildOf(AActor::StaticClass()))
 	{
-		UGeometryCacheComponent* GeomMeshComp = AcquireGeometryCacheFromObjectGuid(ObjectBinding, GetSequencer());
+		UGeometryCacheComponent* GeomMeshComp = AcquireGeometryCacheFromObjectGuid(ObjectBindings[0], GetSequencer());
 
 		if (GeomMeshComp)
 		{
@@ -295,24 +295,34 @@ void FGeometryCacheTrackEditor::BuildObjectBindingTrackMenu(FMenuBuilder& MenuBu
 				NSLOCTEXT("Sequencer", "AddGeometryCacheTooltip", "Adds a Geometry Cache track."),
 				FSlateIcon(),
 				FUIAction(
-					FExecuteAction::CreateSP(this, &FGeometryCacheTrackEditor::BuildGeometryCacheTrack, ObjectBinding, GeomMeshComp, Track)
+					FExecuteAction::CreateSP(this, &FGeometryCacheTrackEditor::BuildGeometryCacheTrack, ObjectBindings, Track)
 				)
 			);
 		}
 	}
 }
 
-void FGeometryCacheTrackEditor::BuildGeometryCacheTrack(FGuid ObjectBinding, UGeometryCacheComponent *GeomCacheComp, UMovieSceneTrack* Track)
+void FGeometryCacheTrackEditor::BuildGeometryCacheTrack(TArray<FGuid> ObjectBindings, UMovieSceneTrack* Track)
 {
 	TSharedPtr<ISequencer> SequencerPtr = GetSequencer();
 
-	if (SequencerPtr.IsValid() && ObjectBinding.IsValid())
+	if (SequencerPtr.IsValid())
 	{
-		UObject* Object = SequencerPtr->FindSpawnedObjectOrTemplate(ObjectBinding);
-		if (Object)
+		const FScopedTransaction Transaction(LOCTEXT("AddGeometryCache_Transaction", "Add Geometry Cache"));
+
+		for (FGuid ObjectBinding : ObjectBindings)
 		{
-			const FScopedTransaction Transaction(LOCTEXT("AddGeometryCache_Transaction", "Add Geometry Cache"));
-			AnimatablePropertyChanged(FOnKeyProperty::CreateRaw(this, &FGeometryCacheTrackEditor::AddKeyInternal, Object, GeomCacheComp, Track));
+			if (ObjectBinding.IsValid())
+			{
+				UObject* Object = SequencerPtr->FindSpawnedObjectOrTemplate(ObjectBinding);
+
+				UGeometryCacheComponent* GeomMeshComp = AcquireGeometryCacheFromObjectGuid(ObjectBinding, GetSequencer());
+
+				if (Object && GeomMeshComp)
+				{
+					AnimatablePropertyChanged(FOnKeyProperty::CreateRaw(this, &FGeometryCacheTrackEditor::AddKeyInternal, Object, GeomMeshComp, Track));
+				}
+			}
 		}
 	}
 }
@@ -361,8 +371,11 @@ TSharedPtr<SWidget> FGeometryCacheTrackEditor::BuildOutlinerEditWidget(const FGu
 		{
 			FMenuBuilder MenuBuilder(true, nullptr);
 
-			BuildGeometryCacheTrack(ObjectBinding, GeomMeshComp, Track);
+			TArray<FGuid> ObjectBindings;
+			ObjectBindings.Add(ObjectBinding);
 
+			BuildGeometryCacheTrack(ObjectBindings, Track);
+			
 			return MenuBuilder.MakeWidget();
 		};
 
diff --git a/Engine/Plugins/Experimental/GeometryCache/Source/GeometryCacheSequencer/Classes/GeometryCacheTrackEditor.h b/Engine/Plugins/Experimental/GeometryCache/Source/GeometryCacheSequencer/Classes/GeometryCacheTrackEditor.h
--- a/Engine/Plugins/Experimental/GeometryCache/Source/GeometryCacheSequencer/Classes/GeometryCacheTrackEditor.h
+++ b/Engine/Plugins/Experimental/GeometryCache/Source/GeometryCacheSequencer/Classes/GeometryCacheTrackEditor.h
@@ -45,7 +45,7 @@ public:
 public:
 
 	// ISequencerTrackEditor interface
-	virtual void BuildObjectBindingTrackMenu(FMenuBuilder& MenuBuilder, const FGuid& ObjectBinding, const UClass* ObjectClass) override;
+	virtual void BuildObjectBindingTrackMenu(FMenuBuilder& MenuBuilder, const TArray<FGuid>& ObjectBindings, const UClass* ObjectClass) override;
 	virtual TSharedRef<ISequencerSection> MakeSectionInterface( UMovieSceneSection& SectionObject, UMovieSceneTrack& Track, FGuid ObjectBinding ) override;
 	virtual bool SupportsType( TSubclassOf<UMovieSceneTrack> Type ) const override;
 	virtual TSharedPtr<SWidget> BuildOutlinerEditWidget(const FGuid& ObjectBinding, UMovieSceneTrack* Track, const FBuildEditWidgetParams& Params) override;
@@ -53,7 +53,7 @@ public:
 
 private:
 
-	void BuildGeometryCacheTrack(FGuid ObjectBinding, UGeometryCacheComponent *GeomCacheComp, UMovieSceneTrack* Track);
+	void BuildGeometryCacheTrack(TArray<FGuid> ObjectBindings, UMovieSceneTrack* Track);
 
 	/** Delegate for AnimatablePropertyChanged in AddKey */
 	FKeyPropertyResult AddKeyInternal(FFrameNumber KeyTime, UObject* Object, UGeometryCacheComponent* GeomCacheComp, UMovieSceneTrack* Track);
diff --git a/Engine/Plugins/Experimental/ImagePlate/Source/ImagePlateEditor/Private/ImagePlateTrackEditor.cpp b/Engine/Plugins/Experimental/ImagePlate/Source/ImagePlateEditor/Private/ImagePlateTrackEditor.cpp
--- a/Engine/Plugins/Experimental/ImagePlate/Source/ImagePlateEditor/Private/ImagePlateTrackEditor.cpp
+++ b/Engine/Plugins/Experimental/ImagePlate/Source/ImagePlateEditor/Private/ImagePlateTrackEditor.cpp
@@ -268,7 +268,7 @@ void FImagePlateTrackEditor::OnAnimatedPropertyChanged(const FPropertyChangedPar
 	}
 }
 
-void FImagePlateTrackEditor::BuildObjectBindingTrackMenu(FMenuBuilder& MenuBuilder, const FGuid& ObjectBinding, const UClass* ObjectClass)
+void FImagePlateTrackEditor::BuildObjectBindingTrackMenu(FMenuBuilder& MenuBuilder, const TArray<FGuid>& ObjectBindings, const UClass* ObjectClass)
 {
 	// We only know how to add specific properties for image plates and components. Anything else must be keyed through the generic MediaPlayer property
 	if (!ObjectClass->IsChildOf(AImagePlate::StaticClass()) && !ObjectClass->IsChildOf(UImagePlateComponent::StaticClass()))
@@ -277,7 +277,7 @@ void FImagePlateTrackEditor::BuildObjectBindingTrackMenu(FMenuBuilder& MenuBuild
 	}
 
 	// Find the spawned object or its template
-	UObject* Object = GetSequencer()->FindSpawnedObjectOrTemplate(ObjectBinding);
+	UObject* Object = GetSequencer()->FindSpawnedObjectOrTemplate(ObjectBindings[0]);
 	if (!Object)
 	{
 		return;
@@ -323,7 +323,7 @@ void FImagePlateTrackEditor::BuildObjectBindingTrackMenu(FMenuBuilder& MenuBuild
 
 		bool bCanAddTrack = true;
 
-		const FMovieSceneBinding* Binding = MovieScene->GetBindings().FindByPredicate([=](const FMovieSceneBinding& InBinding){ return InBinding.GetObjectGuid() == ObjectBinding; });
+		const FMovieSceneBinding* Binding = MovieScene->GetBindings().FindByPredicate([=](const FMovieSceneBinding& InBinding){ return InBinding.GetObjectGuid() == ObjectBindings[0]; });
 		if (Binding)
 		{
 			FString PredicatePath = Path->ToString(TEXT("."));
diff --git a/Engine/Plugins/Experimental/ImagePlate/Source/ImagePlateEditor/Private/ImagePlateTrackEditor.h b/Engine/Plugins/Experimental/ImagePlate/Source/ImagePlateEditor/Private/ImagePlateTrackEditor.h
--- a/Engine/Plugins/Experimental/ImagePlate/Source/ImagePlateEditor/Private/ImagePlateTrackEditor.h
+++ b/Engine/Plugins/Experimental/ImagePlate/Source/ImagePlateEditor/Private/ImagePlateTrackEditor.h
@@ -29,7 +29,7 @@ public:
 	virtual TSharedRef<ISequencerSection> MakeSectionInterface(UMovieSceneSection& SectionObject, UMovieSceneTrack& Track, FGuid ObjectBinding) override;
 	virtual bool SupportsType(TSubclassOf<UMovieSceneTrack> TrackClass) const override;
 	virtual void Tick(float DeltaTime) override;
-	virtual void BuildObjectBindingTrackMenu(FMenuBuilder& MenuBuilder, const FGuid& ObjectBinding, const UClass* ObjectClass) override;
+	virtual void BuildObjectBindingTrackMenu(FMenuBuilder& MenuBuilder, const TArray<FGuid>& ObjectBindings, const UClass* ObjectClass) override;
 	virtual TSharedPtr<SWidget> BuildOutlinerEditWidget(const FGuid& ObjectBinding, UMovieSceneTrack* Track, const FBuildEditWidgetParams& Params) override;
 	virtual UMovieSceneTrack* AddTrack(UMovieScene* FocusedMovieScene, const FGuid& ObjectHandle, TSubclassOf<class UMovieSceneTrack> TrackClass, FName UniqueTypeName) override;
 	
diff --git a/Engine/Plugins/FX/Niagara/Source/NiagaraEditor/Private/Sequencer/LevelSequence/NiagaraSystemTrackEditor.cpp b/Engine/Plugins/FX/Niagara/Source/NiagaraEditor/Private/Sequencer/LevelSequence/NiagaraSystemTrackEditor.cpp
--- a/Engine/Plugins/FX/Niagara/Source/NiagaraEditor/Private/Sequencer/LevelSequence/NiagaraSystemTrackEditor.cpp
+++ b/Engine/Plugins/FX/Niagara/Source/NiagaraEditor/Private/Sequencer/LevelSequence/NiagaraSystemTrackEditor.cpp
@@ -56,7 +56,7 @@ void GetAnimatedParameters(UMovieScene& MovieScene, FGuid ObjectBinding, TSet<FN
 	}
 }
 
-void FNiagaraSystemTrackEditor::BuildObjectBindingTrackMenu(FMenuBuilder& MenuBuilder, const FGuid& ObjectBinding, const UClass* ObjectClass)
+void FNiagaraSystemTrackEditor::BuildObjectBindingTrackMenu(FMenuBuilder& MenuBuilder, const TArray<FGuid>& ObjectBindings, const UClass* ObjectClass)
 {
 	if (ObjectClass->IsChildOf(UNiagaraComponent::StaticClass()))
 	{
@@ -64,11 +64,11 @@ void FNiagaraSystemTrackEditor::BuildObjectBindingTrackMenu(FMenuBuilder& MenuBu
 			LOCTEXT("AddNiagaraSystemTrack", "Niagara System Life Cycle Track"),
 			LOCTEXT("AddNiagaraSystemTrackToolTip", "Add a track for controlling niagara system life cycle behavior."),
 			FSlateIcon(),
-			FUIAction(FExecuteAction::CreateSP(this, &FNiagaraSystemTrackEditor::AddNiagaraSystemTrack, ObjectBinding))
+			FUIAction(FExecuteAction::CreateSP(this, &FNiagaraSystemTrackEditor::AddNiagaraSystemTrack, ObjectBindings))
 		);
 	}
 
-	TArrayView<TWeakObjectPtr<UObject>> BoundObjects = GetSequencer()->FindBoundObjects(ObjectBinding, GetSequencer()->GetFocusedTemplateID());
+	TArrayView<TWeakObjectPtr<UObject>> BoundObjects = GetSequencer()->FindBoundObjects(ObjectBindings[0], GetSequencer()->GetFocusedTemplateID());
 
 	UNiagaraSystem* System = nullptr;
 	for (TWeakObjectPtr<UObject> BoundObject : BoundObjects)
@@ -86,7 +86,7 @@ void FNiagaraSystemTrackEditor::BuildObjectBindingTrackMenu(FMenuBuilder& MenuBu
 		System->GetExposedParameters().GetUserParameters(ParameterVariables);
 
 		TSet<FNiagaraVariable> AnimatedParameters;
-		GetAnimatedParameters(*GetSequencer()->GetFocusedMovieSceneSequence()->GetMovieScene(), ObjectBinding, AnimatedParameters);
+		GetAnimatedParameters(*GetSequencer()->GetFocusedMovieSceneSequence()->GetMovieScene(), ObjectBindings[0], AnimatedParameters);
 
 		FNiagaraEditorModule& NiagaraEditorModule = FModuleManager::GetModuleChecked<FNiagaraEditorModule>("NiagaraEditor");
 		for (const FNiagaraVariable& ParameterVariable : ParameterVariables)
@@ -99,66 +99,69 @@ void FNiagaraSystemTrackEditor::BuildObjectBindingTrackMenu(FMenuBuilder& MenuBu
 					FText::Format(LOCTEXT("AddNiagaraParameterTrackFormat", "{0} Parameter Track"), FText::FromName(ParameterVariable.GetName())),
 					FText::Format(LOCTEXT("AddNiagaraSystemTrackToolTipFormat", "Add a track for animating the {0} parameter."), FText::FromName(ParameterVariable.GetName())),
 					FSlateIcon(),
-					FUIAction(FExecuteAction::CreateSP(this, &FNiagaraSystemTrackEditor::AddNiagaraParameterTrack, ObjectBinding, ParameterVariable))
+					FUIAction(FExecuteAction::CreateSP(this, &FNiagaraSystemTrackEditor::AddNiagaraParameterTrack, ObjectBindings, ParameterVariable))
 				);
 			}
 		}
 	}
 }
 
-void FNiagaraSystemTrackEditor::AddNiagaraSystemTrack(FGuid ObjectBinding)
+void FNiagaraSystemTrackEditor::AddNiagaraSystemTrack(TArray<FGuid> ObjectBindings)
 {
-	TArrayView<TWeakObjectPtr<UObject>> BoundObjects = GetSequencer()->FindBoundObjects(ObjectBinding, GetSequencer()->GetFocusedTemplateID());
-
-	UNiagaraSystem* System = nullptr;
-	for (TWeakObjectPtr<UObject> BoundObject : BoundObjects)
+	UMovieScene* MovieScene = GetSequencer()->GetFocusedMovieSceneSequence()->GetMovieScene();
+	if (MovieScene->IsReadOnly())
 	{
-		UNiagaraComponent* NiagaraComponent = Cast<UNiagaraComponent>(BoundObject.Get());
-		if (NiagaraComponent != nullptr)
-		{
-			System = NiagaraComponent->GetAsset();
-		}
+		return;
 	}
 
-	if (System != nullptr)
+	FScopedTransaction AddTrackTransaction(LOCTEXT("AddNiagaraSystemLifeCycleTrackTransaction", "Add Niagara System Life Cycle Track"));
+	MovieScene->Modify();
+
+	for (FGuid ObjectBinding : ObjectBindings)
 	{
-		UMovieScene* MovieScene = GetSequencer()->GetFocusedMovieSceneSequence()->GetMovieScene();
-		if (MovieScene->IsReadOnly())
+		TArrayView<TWeakObjectPtr<UObject>> BoundObjects = GetSequencer()->FindBoundObjects(ObjectBinding, GetSequencer()->GetFocusedTemplateID());
+
+		UNiagaraSystem* System = nullptr;
+		for (TWeakObjectPtr<UObject> BoundObject : BoundObjects)
 		{
-			return;
+			UNiagaraComponent* NiagaraComponent = Cast<UNiagaraComponent>(BoundObject.Get());
+			if (NiagaraComponent != nullptr)
+			{
+				System = NiagaraComponent->GetAsset();
+			}
 		}
 
-		FScopedTransaction AddTrackTransaction(LOCTEXT("AddNiagaraSystemLifeCycleTrackTransaction", "Add Niagara System Life Cycle Track"));
-		MovieScene->Modify();
+		if (System != nullptr)
+		{
+			UMovieSceneNiagaraSystemTrack* NiagaraSystemTrack = MovieScene->AddTrack<UMovieSceneNiagaraSystemTrack>(ObjectBinding);
+			NiagaraSystemTrack->SetDisplayName(LOCTEXT("SystemLifeCycleTrackName", "System Life Cycle"));
+			UMovieSceneNiagaraSystemSpawnSection* SpawnSection = NewObject<UMovieSceneNiagaraSystemSpawnSection>(NiagaraSystemTrack, NAME_None, RF_Transactional);
 
-		UMovieSceneNiagaraSystemTrack* NiagaraSystemTrack = MovieScene->AddTrack<UMovieSceneNiagaraSystemTrack>(ObjectBinding);
-		NiagaraSystemTrack->SetDisplayName(LOCTEXT("SystemLifeCycleTrackName", "System Life Cycle"));
-		UMovieSceneNiagaraSystemSpawnSection* SpawnSection = NewObject<UMovieSceneNiagaraSystemSpawnSection>(NiagaraSystemTrack, NAME_None, RF_Transactional);
+			FFrameRate FrameResolution = GetSequencer()->GetFocusedMovieSceneSequence()->GetMovieScene()->GetTickResolution();
+			FFrameTime SpawnSectionStartTime = GetSequencer()->GetLocalTime().ConvertTo(FrameResolution);
+			FFrameTime SpawnSectionDuration;
 
-		FFrameRate FrameResolution = GetSequencer()->GetFocusedMovieSceneSequence()->GetMovieScene()->GetTickResolution();
-		FFrameTime SpawnSectionStartTime = GetSequencer()->GetLocalTime().ConvertTo(FrameResolution);
-		FFrameTime SpawnSectionDuration;
+			UNiagaraSystemEditorData* SystemEditorData = Cast<UNiagaraSystemEditorData>(System->GetEditorData());
+			if (SystemEditorData != nullptr && SystemEditorData->GetPlaybackRange().HasLowerBound() && SystemEditorData->GetPlaybackRange().HasUpperBound())
+			{
+				SpawnSectionDuration = FrameResolution.AsFrameTime(SystemEditorData->GetPlaybackRange().Size<float>());
+			}
+			else
+			{
+				SpawnSectionDuration = FrameResolution.AsFrameTime(5.0);
+			}
 
-		UNiagaraSystemEditorData* SystemEditorData = Cast<UNiagaraSystemEditorData>(System->GetEditorData());
-		if (SystemEditorData != nullptr && SystemEditorData->GetPlaybackRange().HasLowerBound() && SystemEditorData->GetPlaybackRange().HasUpperBound())
-		{
-			SpawnSectionDuration = FrameResolution.AsFrameTime(SystemEditorData->GetPlaybackRange().Size<float>());
+			SpawnSection->SetRange(TRange<FFrameNumber>(
+				SpawnSectionStartTime.RoundToFrame(),
+				(SpawnSectionStartTime + SpawnSectionDuration).RoundToFrame()));
+			NiagaraSystemTrack->AddSection(*SpawnSection);
 		}
-		else
-		{
-			SpawnSectionDuration = FrameResolution.AsFrameTime(5.0);
-		}
-
-		SpawnSection->SetRange(TRange<FFrameNumber>(
-			SpawnSectionStartTime.RoundToFrame(),
-			(SpawnSectionStartTime + SpawnSectionDuration).RoundToFrame()));
-		NiagaraSystemTrack->AddSection(*SpawnSection);
-
-		GetSequencer()->NotifyMovieSceneDataChanged(EMovieSceneDataChangeType::MovieSceneStructureItemAdded);
 	}
+		
+	GetSequencer()->NotifyMovieSceneDataChanged(EMovieSceneDataChangeType::MovieSceneStructureItemAdded);
 }
 
-void FNiagaraSystemTrackEditor::AddNiagaraParameterTrack(FGuid ObjectBinding, FNiagaraVariable Parameter)
+void FNiagaraSystemTrackEditor::AddNiagaraParameterTrack(TArray<FGuid> ObjectBindings, FNiagaraVariable Parameter)
 {
 	FNiagaraEditorModule& NiagaraEditorModule = FModuleManager::GetModuleChecked<FNiagaraEditorModule>("NiagaraEditor");
 	if (NiagaraEditorModule.CanCreateParameterTrackForType(*Parameter.GetType().GetScriptStruct()))
@@ -172,14 +175,17 @@ void FNiagaraSystemTrackEditor::AddNiagaraParameterTrack(FGuid ObjectBinding, FN
 		FScopedTransaction AddTrackTransaction(LOCTEXT("AddNiagaraParameterTrackTransaction", "Add Niagara Parameter Track"));
 		MovieScene->Modify();
 
-		UMovieSceneNiagaraParameterTrack* ParameterTrack = NiagaraEditorModule.CreateParameterTrackForType(*Parameter.GetType().GetScriptStruct(), Parameter);
-		MovieScene->AddGivenTrack(ParameterTrack, ObjectBinding);
+		for (FGuid ObjectBinding : ObjectBindings)
+		{
+			UMovieSceneNiagaraParameterTrack* ParameterTrack = NiagaraEditorModule.CreateParameterTrackForType(*Parameter.GetType().GetScriptStruct(), Parameter);
+			MovieScene->AddGivenTrack(ParameterTrack, ObjectBinding);
 
-		ParameterTrack->SetParameter(Parameter);
-		ParameterTrack->SetDisplayName(FText::FromName(Parameter.GetName()));
+			ParameterTrack->SetParameter(Parameter);
+			ParameterTrack->SetDisplayName(FText::FromName(Parameter.GetName()));
 
-		UMovieSceneSection* ParameterSection = ParameterTrack->CreateNewSection();
-		ParameterTrack->AddSection(*ParameterSection);
+			UMovieSceneSection* ParameterSection = ParameterTrack->CreateNewSection();
+			ParameterTrack->AddSection(*ParameterSection);
+		}
 
 		GetSequencer()->NotifyMovieSceneDataChanged(EMovieSceneDataChangeType::MovieSceneStructureItemAdded);
 	}
diff --git a/Engine/Plugins/FX/Niagara/Source/NiagaraEditor/Private/Sequencer/LevelSequence/NiagaraSystemTrackEditor.h b/Engine/Plugins/FX/Niagara/Source/NiagaraEditor/Private/Sequencer/LevelSequence/NiagaraSystemTrackEditor.h
--- a/Engine/Plugins/FX/Niagara/Source/NiagaraEditor/Private/Sequencer/LevelSequence/NiagaraSystemTrackEditor.h
+++ b/Engine/Plugins/FX/Niagara/Source/NiagaraEditor/Private/Sequencer/LevelSequence/NiagaraSystemTrackEditor.h
@@ -39,9 +39,9 @@ public:
 	//~ ISequencerTrackEditor interface
 	virtual TSharedRef<ISequencerSection> MakeSectionInterface(UMovieSceneSection& SectionObject, UMovieSceneTrack& Track, FGuid ObjectBinding) override;
 	virtual bool SupportsType(TSubclassOf<UMovieSceneTrack> Type) const override;
-	virtual void BuildObjectBindingTrackMenu(FMenuBuilder& MenuBuilder, const FGuid& ObjectBinding, const UClass* ObjectClass) override;
+	virtual void BuildObjectBindingTrackMenu(FMenuBuilder& MenuBuilder, const TArray<FGuid>& ObjectBindings, const UClass* ObjectClass) override;
 
 private:
-	void AddNiagaraSystemTrack(FGuid ObjectBinding);
-	void AddNiagaraParameterTrack(FGuid ObjectBinding, FNiagaraVariable Parameter);
+	void AddNiagaraSystemTrack(TArray<FGuid> ObjectBindings);
+	void AddNiagaraParameterTrack(TArray<FGuid> ObjectBindings, FNiagaraVariable Parameter);
 };
\ No newline at end of file
diff --git a/Engine/Plugins/MovieScene/LevelSequenceEditor/Source/LevelSequenceEditor/Private/LevelSequenceEditorToolkit.cpp b/Engine/Plugins/MovieScene/LevelSequenceEditor/Source/LevelSequenceEditor/Private/LevelSequenceEditorToolkit.cpp
--- a/Engine/Plugins/MovieScene/LevelSequenceEditor/Source/LevelSequenceEditor/Private/LevelSequenceEditorToolkit.cpp
+++ b/Engine/Plugins/MovieScene/LevelSequenceEditor/Source/LevelSequenceEditor/Private/LevelSequenceEditorToolkit.cpp
@@ -579,24 +579,38 @@ void FLevelSequenceEditorToolkit::HandleAddComponentActionExecute(UActorComponen
 
 void FLevelSequenceEditorToolkit::HandleAddComponentMaterialActionExecute(UPrimitiveComponent* Component, int32 MaterialIndex)
 {
-	FGuid ObjectHandle = Sequencer->GetHandleToObject(Component);
 	UMovieScene* FocusedMovieScene = Sequencer->GetFocusedMovieSceneSequence()->GetMovieScene();
-	FName IndexName( *FString::FromInt(MaterialIndex) );
-	if ( FocusedMovieScene->FindTrack( UMovieSceneComponentMaterialTrack::StaticClass(), ObjectHandle, IndexName ) == nullptr )
+	if (FocusedMovieScene->IsReadOnly())
 	{
-		if (FocusedMovieScene->IsReadOnly())
-		{
-			return;
-		}
+		return;
+	}
 
-		const FScopedTransaction Transaction( LOCTEXT( "AddComponentMaterialTrack", "Add component material track" ) );
+	const FScopedTransaction Transaction(LOCTEXT("AddComponentMaterialTrack", "Add component material track"));
 
-		FocusedMovieScene->Modify();
+	FocusedMovieScene->Modify();
 
-		UMovieSceneComponentMaterialTrack* MaterialTrack = FocusedMovieScene->AddTrack<UMovieSceneComponentMaterialTrack>( ObjectHandle );
-		MaterialTrack->Modify();
-		MaterialTrack->SetMaterialIndex( MaterialIndex );
+	FString ComponentName = Component->GetName();
+	USelection* SelectedActors = GEditor->GetSelectedActors();
+	for (FSelectionIterator Iter(*SelectedActors); Iter; ++Iter)
+	{
+		AActor* Actor = CastChecked<AActor>(*Iter);
 
+		TArray<UActorComponent*> OutComponents;
+		Actor->GetComponents(OutComponents);
+		for (UActorComponent* ActorComponent : OutComponents)
+		{
+			if (ActorComponent->GetName() == ComponentName)
+			{
+				FGuid ObjectHandle = Sequencer->GetHandleToObject(ActorComponent);
+				FName IndexName(*FString::FromInt(MaterialIndex));
+				if (FocusedMovieScene->FindTrack(UMovieSceneComponentMaterialTrack::StaticClass(), ObjectHandle, IndexName) == nullptr)
+				{
+					UMovieSceneComponentMaterialTrack* MaterialTrack = FocusedMovieScene->AddTrack<UMovieSceneComponentMaterialTrack>(ObjectHandle);
+					MaterialTrack->Modify();
+					MaterialTrack->SetMaterialIndex(MaterialIndex);
+				}
+			}
+		}
 		Sequencer->NotifyMovieSceneDataChanged( EMovieSceneDataChangeType::MovieSceneStructureItemAdded );
 	}
 }
diff --git a/Engine/Source/Editor/MovieSceneTools/Private/TrackEditors/ActorPickerTrackEditor.cpp b/Engine/Source/Editor/MovieSceneTools/Private/TrackEditors/ActorPickerTrackEditor.cpp
--- a/Engine/Source/Editor/MovieSceneTools/Private/TrackEditors/ActorPickerTrackEditor.cpp
+++ b/Engine/Source/Editor/MovieSceneTools/Private/TrackEditors/ActorPickerTrackEditor.cpp
@@ -32,7 +32,7 @@ FActorPickerTrackEditor::FActorPickerTrackEditor(TSharedRef<ISequencer> InSequen
 {
 }
 
-void FActorPickerTrackEditor::PickActorInteractive(FGuid ObjectBinding, UMovieSceneSection* Section)
+void FActorPickerTrackEditor::PickActorInteractive(const TArray<FGuid>& ObjectBindings, UMovieSceneSection* Section)
 {
 	if(GUnrealEd->GetSelectedActorCount())
 	{
@@ -40,26 +40,26 @@ void FActorPickerTrackEditor::PickActorInteractive(FGuid ObjectBinding, UMovieSc
 
 		ActorPickerMode.BeginActorPickingMode(
 			FOnGetAllowedClasses(), 
-			FOnShouldFilterActor::CreateSP(this, &FActorPickerTrackEditor::IsActorPickable, ObjectBinding, Section), 
-			FOnActorSelected::CreateSP(this, &FActorPickerTrackEditor::ActorPicked, ObjectBinding, Section)
+			FOnShouldFilterActor::CreateSP(this, &FActorPickerTrackEditor::IsActorPickable, ObjectBindings[0], Section), 
+			FOnActorSelected::CreateSP(this, &FActorPickerTrackEditor::ActorPicked, ObjectBindings, Section)
 			);
 	}
 }
 
-void FActorPickerTrackEditor::ShowActorSubMenu(FMenuBuilder& MenuBuilder, FGuid ObjectBinding, UMovieSceneSection* Section)
+void FActorPickerTrackEditor::ShowActorSubMenu(FMenuBuilder& MenuBuilder, TArray<FGuid> ObjectBindings, UMovieSceneSection* Section)
 {
 	struct Local
 	{
-		static FReply OnInteractiveActorPickerClicked(FActorPickerTrackEditor* ActorPickerTrackEditor, FGuid TheObjectBinding, UMovieSceneSection* TheSection)
+		static FReply OnInteractiveActorPickerClicked(FActorPickerTrackEditor* ActorPickerTrackEditor, TArray<FGuid> TheObjectBindings, UMovieSceneSection* TheSection)
 		{
 			FSlateApplication::Get().DismissAllMenus();
-			ActorPickerTrackEditor->PickActorInteractive(TheObjectBinding, TheSection);
+			ActorPickerTrackEditor->PickActorInteractive(TheObjectBindings, TheSection);
 			return FReply::Handled();
 		}
 	};
 
 	auto CreateNewBinding = 
-		[this, ObjectBinding, Section](FMenuBuilder& SubMenuBuilder)
+		[this, ObjectBindings, Section](FMenuBuilder& SubMenuBuilder)
 	{
 		using namespace SceneOutliner;
 
@@ -74,7 +74,7 @@ void FActorPickerTrackEditor::ShowActorSubMenu(FMenuBuilder& MenuBuilder, FGuid
 			InitOptions.ColumnMap.Add(FBuiltInColumnTypes::Label(), FColumnInfo(EColumnVisibility::Visible, 0));
 
 			// Only display Actors that we can attach too
-			InitOptions.Filters->AddFilterPredicate( SceneOutliner::FActorFilterPredicate::CreateSP(this, &FActorPickerTrackEditor::IsActorPickable, ObjectBinding, Section) );
+			InitOptions.Filters->AddFilterPredicate( SceneOutliner::FActorFilterPredicate::CreateSP(this, &FActorPickerTrackEditor::IsActorPickable, ObjectBindings[0], Section) );
 		}		
 
 		// Actor selector to allow the user to choose a parent actor
@@ -92,7 +92,7 @@ void FActorPickerTrackEditor::ShowActorSubMenu(FMenuBuilder& MenuBuilder, FGuid
 				[
 					SceneOutlinerModule.CreateSceneOutliner(
 						InitOptions,
-						FOnActorPicked::CreateSP(this, &FActorPickerTrackEditor::ActorPicked, ObjectBinding, Section )
+						FOnActorPicked::CreateSP(this, &FActorPickerTrackEditor::ActorPicked, ObjectBindings, Section )
 						)
 				]
 			]
@@ -110,7 +110,7 @@ void FActorPickerTrackEditor::ShowActorSubMenu(FMenuBuilder& MenuBuilder, FGuid
 					SNew(SButton)
 					.ToolTipText( LOCTEXT( "PickButtonLabel", "Pick a parent actor to attach to") )
 					.ButtonStyle(FEditorStyle::Get(), "HoverHintOnly")
-					.OnClicked(FOnClicked::CreateStatic(&Local::OnInteractiveActorPickerClicked, this, ObjectBinding, Section))
+					.OnClicked(FOnClicked::CreateStatic(&Local::OnInteractiveActorPickerClicked, this, ObjectBindings, Section))
 					.ContentPadding(4.0f)
 					.ForegroundColor(FSlateColor::UseForeground())
 					.IsFocusable(false)
@@ -129,7 +129,7 @@ void FActorPickerTrackEditor::ShowActorSubMenu(FMenuBuilder& MenuBuilder, FGuid
 
 	// Always recreate the binding picker to ensure we have the correct sequence ID
 	BindingIDPicker = MakeShared<FTrackEditorBindingIDPicker>(SequencerPtr->GetFocusedTemplateID(), SequencerPtr);
-	BindingIDPicker->OnBindingPicked().AddRaw(this, &FActorPickerTrackEditor::ExistingBindingPicked, ObjectBinding, Section);
+	BindingIDPicker->OnBindingPicked().AddRaw(this, &FActorPickerTrackEditor::ExistingBindingPicked, ObjectBindings, Section);
 
 	FText ExistingBindingText = LOCTEXT("ExistingBinding", "Existing Binding");
 	FText NewBindingText = LOCTEXT("NewBinding", "New Binding");
@@ -254,13 +254,13 @@ public:
 	}
 };
 
-void FActorPickerTrackEditor::ActorPicked(AActor* ParentActor, FGuid ObjectGuid, UMovieSceneSection* Section)
+void FActorPickerTrackEditor::ActorPicked(AActor* ParentActor, TArray<FGuid> ObjectGuids, UMovieSceneSection* Section)
 {
-	ActorPickerIDPicked(FActorPickerID(ParentActor, FMovieSceneObjectBindingID()), ObjectGuid, Section);
+	ActorPickerIDPicked(FActorPickerID(ParentActor, FMovieSceneObjectBindingID()), ObjectGuids, Section);
 }
 
 
-void FActorPickerTrackEditor::ExistingBindingPicked(FMovieSceneObjectBindingID ExistingBindingID, FGuid ObjectBinding, UMovieSceneSection* Section)
+void FActorPickerTrackEditor::ExistingBindingPicked(FMovieSceneObjectBindingID ExistingBindingID, TArray<FGuid> ObjectBindings, UMovieSceneSection* Section)
 {
 	TSharedPtr<ISequencer> SequencerPtr = GetSequencer();
 
@@ -272,17 +272,17 @@ void FActorPickerTrackEditor::ExistingBindingPicked(FMovieSceneObjectBindingID E
 			AActor* Actor = Cast<AActor>(RuntimeObject.Get());
 			if (Actor)
 			{
-				ActorPickerIDPicked(FActorPickerID(Actor, ExistingBindingID), ObjectBinding, Section);
+				ActorPickerIDPicked(FActorPickerID(Actor, ExistingBindingID), ObjectBindings, Section);
 				return;
 			}
 		}
 	}
 
-	ActorPickerIDPicked(FActorPickerID(nullptr, ExistingBindingID), ObjectBinding, Section);
+	ActorPickerIDPicked(FActorPickerID(nullptr, ExistingBindingID), ObjectBindings, Section);
 }
 
 
-void FActorPickerTrackEditor::ActorPickerIDPicked(FActorPickerID ActorPickerID, FGuid ObjectGuid, UMovieSceneSection* Section)
+void FActorPickerTrackEditor::ActorPickerIDPicked(FActorPickerID ActorPickerID, const TArray<FGuid>& ObjectGuids, UMovieSceneSection* Section)
 {
 	TArray<USceneComponent*> ComponentsWithSockets;
 	if (ActorPickerID.ActorPicked.IsValid())
@@ -301,7 +301,7 @@ void FActorPickerTrackEditor::ActorPickerIDPicked(FActorPickerID ActorPickerID,
 	if (ComponentsWithSockets.Num() == 0)
 	{
 		FSlateApplication::Get().DismissAllMenus();
-		ActorSocketPicked( NAME_None, nullptr, ActorPickerID, ObjectGuid, Section );
+		ActorSocketPicked( NAME_None, nullptr, ActorPickerID, ObjectGuids, Section );
 		return;
 	}
 
@@ -315,7 +315,7 @@ void FActorPickerTrackEditor::ActorPickerIDPicked(FActorPickerID ActorPickerID,
 		MenuWidget = 
 			SNew(SComponentChooserPopup)
 			.Actor(ActorPickerID.ActorPicked.Get())
-			.OnComponentChosen(this, &FActorPickerTrackEditor::ActorComponentPicked, ActorPickerID, ObjectGuid, Section);		
+			.OnComponentChosen(this, &FActorPickerTrackEditor::ActorComponentPicked, ActorPickerID, ObjectGuids, Section);		
 
 		// Create as context menu
 		FSlateApplication::Get().PushMenu(
@@ -328,12 +328,12 @@ void FActorPickerTrackEditor::ActorPickerIDPicked(FActorPickerID ActorPickerID,
 	}
 	else
 	{
-		ActorComponentPicked(ComponentsWithSockets[0]->GetName(), ActorPickerID, ObjectGuid, Section);
+		ActorComponentPicked(ComponentsWithSockets[0]->GetName(), ActorPickerID, ObjectGuids, Section);
 	}
 }
 
 
-void FActorPickerTrackEditor::ActorComponentPicked(FString ComponentName, FActorPickerID ActorPickerID, FGuid ObjectGuid, UMovieSceneSection* Section)
+void FActorPickerTrackEditor::ActorComponentPicked(FString ComponentName, FActorPickerID ActorPickerID, TArray<FGuid> ObjectGuids, UMovieSceneSection* Section)
 {
 	USceneComponent* ComponentWithSockets = nullptr;
 	if (ActorPickerID.ActorPicked.IsValid())
@@ -363,7 +363,7 @@ void FActorPickerTrackEditor::ActorComponentPicked(FString ComponentName, FActor
 	MenuWidget = 
 		SNew(SSocketChooserPopup)
 		.SceneComponent(ComponentWithSockets)
-		.OnSocketChosen(this, &FActorPickerTrackEditor::ActorSocketPicked, ComponentWithSockets, ActorPickerID, ObjectGuid, Section);		
+		.OnSocketChosen(this, &FActorPickerTrackEditor::ActorSocketPicked, ComponentWithSockets, ActorPickerID, ObjectGuids, Section);		
 
 	// Create as context menu
 	FSlateApplication::Get().PushMenu(
diff --git a/Engine/Source/Editor/MovieSceneTools/Private/TrackEditors/ActorPickerTrackEditor.h b/Engine/Source/Editor/MovieSceneTools/Private/TrackEditors/ActorPickerTrackEditor.h
--- a/Engine/Source/Editor/MovieSceneTools/Private/TrackEditors/ActorPickerTrackEditor.h
+++ b/Engine/Source/Editor/MovieSceneTools/Private/TrackEditors/ActorPickerTrackEditor.h
@@ -34,27 +34,27 @@ public:
 	virtual bool IsActorPickable( const AActor* const ParentActor, FGuid ObjectBinding, UMovieSceneSection* InSection) { return false; }
 
 	/** Actor socket was picked */
-	virtual void ActorSocketPicked(const FName SocketName, USceneComponent* Component, FActorPickerID ActorPickerID, FGuid ObjectBinding, UMovieSceneSection* Section) {}
+	virtual void ActorSocketPicked(const FName SocketName, USceneComponent* Component, FActorPickerID ActorPickerID, TArray<FGuid> ObjectBinding, UMovieSceneSection* Section) {}
 
 	/** Show a sub menu of the pickable actors */
-	void ShowActorSubMenu(FMenuBuilder& MenuBuilder, FGuid ObjectBinding, UMovieSceneSection* Section);
+	void ShowActorSubMenu(FMenuBuilder& MenuBuilder, TArray<FGuid> ObjectBindings, UMovieSceneSection* Section);
 
 private: 
 
 	/** Actor was picked */
-	void ActorPicked(AActor* ParentActor, FGuid ObjectBinding, UMovieSceneSection* Section);
+	void ActorPicked(AActor* ParentActor, TArray<FGuid> ObjectBinding, UMovieSceneSection* Section);
 
 	/** Existing binding was picked */
-	void ExistingBindingPicked(FMovieSceneObjectBindingID ExistingBindingID, FGuid ObjectBinding, UMovieSceneSection* Section);
+	void ExistingBindingPicked(FMovieSceneObjectBindingID ExistingBindingID, TArray<FGuid> ObjectBindings, UMovieSceneSection* Section);
 
 	/** Actor and/or existing binding was picked */
-	void ActorPickerIDPicked(FActorPickerID, FGuid ObjectBinding, UMovieSceneSection* Section);
+	void ActorPickerIDPicked(FActorPickerID, const TArray<FGuid>& ObjectBindings, UMovieSceneSection* Section);
 
 	/** Actor component was picked */
-	void ActorComponentPicked(FString ComponentName, FActorPickerID ActorPickerID, FGuid ObjectBinding, UMovieSceneSection* Section);
+	void ActorComponentPicked(FString ComponentName, FActorPickerID ActorPickerID, TArray<FGuid> ObjectBindings, UMovieSceneSection* Section);
 
 	/** Interactively pick an actor from the viewport */
-	void PickActorInteractive(FGuid ObjectBinding, UMovieSceneSection* Section);
+	void PickActorInteractive(const TArray<FGuid>& ObjectBindings, UMovieSceneSection* Section);
 
 	/** A binding ID picker that allows us to create a new section from an existing binding */
 	TSharedPtr<FTrackEditorBindingIDPicker> BindingIDPicker;
diff --git a/Engine/Source/Editor/MovieSceneTools/Private/TrackEditors/AttachTrackEditor.cpp b/Engine/Source/Editor/MovieSceneTools/Private/TrackEditors/AttachTrackEditor.cpp
--- a/Engine/Source/Editor/MovieSceneTools/Private/TrackEditors/AttachTrackEditor.cpp
+++ b/Engine/Source/Editor/MovieSceneTools/Private/TrackEditors/AttachTrackEditor.cpp
@@ -81,9 +81,12 @@ public:
 	
 	virtual void BuildSectionContextMenu(FMenuBuilder& MenuBuilder, const FGuid& ObjectBinding) override
 	{
+		TArray<FGuid> ObjectBindings;
+		ObjectBindings.Add(ObjectBinding);
+
 		MenuBuilder.AddSubMenu(
 			LOCTEXT("SetAttach", "Attach"), LOCTEXT("SetAttachTooltip", "Set attach"),
-			FNewMenuDelegate::CreateRaw(AttachTrackEditor, &FActorPickerTrackEditor::ShowActorSubMenu, ObjectBinding, &Section));
+			FNewMenuDelegate::CreateRaw(AttachTrackEditor, &FActorPickerTrackEditor::ShowActorSubMenu, ObjectBindings, &Section));
 	}
 
 private:
@@ -128,7 +131,7 @@ TSharedRef<ISequencerSection> F3DAttachTrackEditor::MakeSectionInterface( UMovie
 }
 
 
-void F3DAttachTrackEditor::BuildObjectBindingTrackMenu(FMenuBuilder& MenuBuilder, const FGuid& ObjectBinding, const UClass* ObjectClass)
+void F3DAttachTrackEditor::BuildObjectBindingTrackMenu(FMenuBuilder& MenuBuilder, const TArray<FGuid>& ObjectBindings, const UClass* ObjectClass)
 {
 	if (ObjectClass != nullptr && ObjectClass->IsChildOf(AActor::StaticClass()))
 	{
@@ -136,7 +139,7 @@ void F3DAttachTrackEditor::BuildObjectBindingTrackMenu(FMenuBuilder& MenuBuilder
 
 		MenuBuilder.AddSubMenu(
 			LOCTEXT("AddAttach", "Attach"), LOCTEXT("AddAttachTooltip", "Adds an attach track."),
-			FNewMenuDelegate::CreateRaw(this, &FActorPickerTrackEditor::ShowActorSubMenu, ObjectBinding, DummySection));
+			FNewMenuDelegate::CreateRaw(this, &FActorPickerTrackEditor::ShowActorSubMenu, ObjectBindings, DummySection));
 	}
 }
 
@@ -179,7 +182,7 @@ bool F3DAttachTrackEditor::IsActorPickable(const AActor* const ParentActor, FGui
 }
 
 
-void F3DAttachTrackEditor::ActorSocketPicked(const FName SocketName, USceneComponent* Component, FActorPickerID ActorPickerID, FGuid ObjectGuid, UMovieSceneSection* Section)
+void F3DAttachTrackEditor::ActorSocketPicked(const FName SocketName, USceneComponent* Component, FActorPickerID ActorPickerID, TArray<FGuid> ObjectGuids, UMovieSceneSection* Section)
 {
 	if (Section != nullptr)
 	{
@@ -207,12 +210,19 @@ void F3DAttachTrackEditor::ActorSocketPicked(const FName SocketName, USceneCompo
 		AttachSection->AttachSocketName = SocketName;			
 		AttachSection->AttachComponentName = Component ? Component->GetFName() : NAME_None;
 	}
-	else if (ObjectGuid.IsValid())
+	else
 	{
 		TArray<TWeakObjectPtr<>> OutObjects;
-		for (TWeakObjectPtr<> Object : GetSequencer()->FindObjectsInCurrentSequence(ObjectGuid))
+
+		for (FGuid ObjectGuid : ObjectGuids)
 		{
-			OutObjects.Add(Object);
+			if (ObjectGuid.IsValid())
+			{
+				for (TWeakObjectPtr<> Object : GetSequencer()->FindObjectsInCurrentSequence(ObjectGuid))
+				{
+					OutObjects.Add(Object);
+				}
+			}
 		}
 
 		AnimatablePropertyChanged( FOnKeyProperty::CreateRaw( this, &F3DAttachTrackEditor::AddKeyInternal, OutObjects, SocketName, Component ? Component->GetFName() : NAME_None, ActorPickerID) );
diff --git a/Engine/Source/Editor/MovieSceneTools/Private/TrackEditors/AttachTrackEditor.h b/Engine/Source/Editor/MovieSceneTools/Private/TrackEditors/AttachTrackEditor.h
--- a/Engine/Source/Editor/MovieSceneTools/Private/TrackEditors/AttachTrackEditor.h
+++ b/Engine/Source/Editor/MovieSceneTools/Private/TrackEditors/AttachTrackEditor.h
@@ -45,13 +45,13 @@ public:
 
 	// ISequencerTrackEditor interface
 
-	virtual void BuildObjectBindingTrackMenu(FMenuBuilder& MenuBuilder, const FGuid& ObjectBinding, const UClass* ObjectClass) override;
+	virtual void BuildObjectBindingTrackMenu(FMenuBuilder& MenuBuilder, const TArray<FGuid>& ObjectBindings, const UClass* ObjectClass) override;
 	virtual TSharedRef<ISequencerSection> MakeSectionInterface( UMovieSceneSection& SectionObject, UMovieSceneTrack& Track, FGuid ObjectBinding ) override;
 	virtual bool SupportsType( TSubclassOf<UMovieSceneTrack> Type ) const override;
 
 	// FTrackEditorActorPicker
 	virtual bool IsActorPickable( const AActor* const ParentActor, FGuid ObjectBinding, UMovieSceneSection* InSection ) override;
-	virtual void ActorSocketPicked(const FName SocketName, USceneComponent* Component, FActorPickerID ActorPickerID, FGuid ObjectBinding, UMovieSceneSection* Section) override;
+	virtual void ActorSocketPicked(const FName SocketName, USceneComponent* Component, FActorPickerID ActorPickerID, TArray<FGuid> ObjectBindings, UMovieSceneSection* Section) override;
 
 private:
 
diff --git a/Engine/Source/Editor/MovieSceneTools/Private/TrackEditors/CameraAnimTrackEditor.cpp b/Engine/Source/Editor/MovieSceneTools/Private/TrackEditors/CameraAnimTrackEditor.cpp
--- a/Engine/Source/Editor/MovieSceneTools/Private/TrackEditors/CameraAnimTrackEditor.cpp
+++ b/Engine/Source/Editor/MovieSceneTools/Private/TrackEditors/CameraAnimTrackEditor.cpp
@@ -123,10 +123,10 @@ bool FCameraAnimTrackEditor::HandleAssetAdded(UObject* Asset, const FGuid& Targe
 
 
 
-void FCameraAnimTrackEditor::BuildObjectBindingTrackMenu(FMenuBuilder& MenuBuilder, const FGuid& ObjectBinding, const UClass* ObjectClass)
+void FCameraAnimTrackEditor::BuildObjectBindingTrackMenu(FMenuBuilder& MenuBuilder, const TArray<FGuid>& ObjectBindings, const UClass* ObjectClass)
 {
 	// only offer this track if we can find a camera component
-	UCameraComponent const* const CamComponent = AcquireCameraComponentFromObjectGuid(ObjectBinding);
+	UCameraComponent const* const CamComponent = AcquireCameraComponentFromObjectGuid(ObjectBindings[0]);
 	if (CamComponent)
 	{
 		const TSharedPtr<ISequencer> ParentSequencer = GetSequencer();
@@ -140,7 +140,7 @@ void FCameraAnimTrackEditor::BuildObjectBindingTrackMenu(FMenuBuilder& MenuBuild
 
 		MenuBuilder.AddSubMenu(
 			LOCTEXT("AddCameraAnim", "Camera Anim"), NSLOCTEXT("Sequencer", "AddCameraAnimTooltip", "Adds an additive camera animation track."),
-			FNewMenuDelegate::CreateRaw(this, &FCameraAnimTrackEditor::AddCameraAnimSubMenu, ObjectBinding)
+			FNewMenuDelegate::CreateRaw(this, &FCameraAnimTrackEditor::AddCameraAnimSubMenu, ObjectBindings)
 			);
 	}
 }
@@ -149,13 +149,16 @@ TSharedRef<SWidget> FCameraAnimTrackEditor::BuildCameraAnimSubMenu(FGuid ObjectB
 {
 	FMenuBuilder MenuBuilder(true, nullptr);
 
-	AddCameraAnimSubMenu(MenuBuilder, ObjectBinding);
+	TArray<FGuid> ObjectBindings;
+	ObjectBindings.Add(ObjectBinding);
+
+	AddCameraAnimSubMenu(MenuBuilder, ObjectBindings);
 
 	return MenuBuilder.MakeWidget();
 }
 
 
-void FCameraAnimTrackEditor::AddCameraAnimSubMenu(FMenuBuilder& MenuBuilder, FGuid ObjectBinding)
+void FCameraAnimTrackEditor::AddCameraAnimSubMenu(FMenuBuilder& MenuBuilder, TArray<FGuid> ObjectBinding)
 {
 	FAssetPickerConfig AssetPickerConfig;
 	{
@@ -193,7 +196,7 @@ TSharedPtr<SWidget> FCameraAnimTrackEditor::BuildOutlinerEditWidget(const FGuid&
 	];
 }
 
-void FCameraAnimTrackEditor::OnCameraAnimAssetSelected(const FAssetData& AssetData, FGuid ObjectBinding)
+void FCameraAnimTrackEditor::OnCameraAnimAssetSelected(const FAssetData& AssetData, TArray<FGuid> ObjectBindings)
 {
 	FSlateApplication::Get().DismissAllMenus();
 
@@ -204,19 +207,25 @@ void FCameraAnimTrackEditor::OnCameraAnimAssetSelected(const FAssetData& AssetDa
 		UCameraAnim* const CameraAnim = CastChecked<UCameraAnim>(AssetData.GetAsset());
 
 		TArray<TWeakObjectPtr<>> OutObjects;
-		for (TWeakObjectPtr<> Object : GetSequencer()->FindObjectsInCurrentSequence(ObjectBinding))
+		for (FGuid ObjectBinding : ObjectBindings)
 		{
-			OutObjects.Add(Object);
+			for (TWeakObjectPtr<> Object : GetSequencer()->FindObjectsInCurrentSequence(ObjectBinding))
+			{
+				OutObjects.Add(Object);
+			}
 		}
+		
+		const FScopedTransaction Transaction(LOCTEXT("AddCameraAnim_Transaction", "Add Camera Anim"));
+
 		AnimatablePropertyChanged(FOnKeyProperty::CreateRaw(this, &FCameraAnimTrackEditor::AddKeyInternal, OutObjects, CameraAnim));
 	}
 }
 
-void FCameraAnimTrackEditor::OnCameraAnimAssetEnterPressed(const TArray<FAssetData>& AssetData, FGuid ObjectBinding)
+void FCameraAnimTrackEditor::OnCameraAnimAssetEnterPressed(const TArray<FAssetData>& AssetData, TArray<FGuid> ObjectBindings)
 {
 	if (AssetData.Num() > 0)
 	{
-		OnCameraAnimAssetSelected(AssetData[0].GetAsset(), ObjectBinding);
+		OnCameraAnimAssetSelected(AssetData[0].GetAsset(), ObjectBindings);
 	}
 }
 
diff --git a/Engine/Source/Editor/MovieSceneTools/Private/TrackEditors/CameraAnimTrackEditor.h b/Engine/Source/Editor/MovieSceneTools/Private/TrackEditors/CameraAnimTrackEditor.h
--- a/Engine/Source/Editor/MovieSceneTools/Private/TrackEditors/CameraAnimTrackEditor.h
+++ b/Engine/Source/Editor/MovieSceneTools/Private/TrackEditors/CameraAnimTrackEditor.h
@@ -45,7 +45,7 @@ public:
 
 	// ISequencerTrackEditor interface
 	virtual void AddKey(const FGuid& ObjectGuid) override;
-	virtual void BuildObjectBindingTrackMenu(FMenuBuilder& MenuBuilder, const FGuid& ObjectBinding, const UClass* ObjectClass) override;
+	virtual void BuildObjectBindingTrackMenu(FMenuBuilder& MenuBuilder, const TArray<FGuid>& ObjectBindings, const UClass* ObjectClass) override;
 	virtual bool HandleAssetAdded(UObject* Asset, const FGuid& TargetObjectGuid) override;
 	virtual TSharedRef<ISequencerSection> MakeSectionInterface(UMovieSceneSection& SectionObject, UMovieSceneTrack& Track, FGuid ObjectBinding) override;
 	virtual bool SupportsType(TSubclassOf<UMovieSceneTrack> Type) const override;
@@ -58,13 +58,13 @@ private:
 
 	/** Animation sub menu */
 	TSharedRef<SWidget> BuildCameraAnimSubMenu(FGuid ObjectBinding);
-	void AddCameraAnimSubMenu(FMenuBuilder& MenuBuilder, FGuid ObjectBinding);
+	void AddCameraAnimSubMenu(FMenuBuilder& MenuBuilder, TArray<FGuid> ObjectBindings);
 
 	/** Animation asset selected */
-	void OnCameraAnimAssetSelected(const FAssetData& AssetData, FGuid ObjectBinding);
+	void OnCameraAnimAssetSelected(const FAssetData& AssetData, TArray<FGuid> ObjectBindings);
 
 	/** Animation asset enter pressed */
-	void OnCameraAnimAssetEnterPressed(const TArray<FAssetData>& AssetData, FGuid ObjectBinding);
+	void OnCameraAnimAssetEnterPressed(const TArray<FAssetData>& AssetData, TArray<FGuid> ObjectBindings);
 
 	class UCameraComponent* AcquireCameraComponentFromObjectGuid(const FGuid& Guid);
 };
diff --git a/Engine/Source/Editor/MovieSceneTools/Private/TrackEditors/CameraShakeTrackEditor.cpp b/Engine/Source/Editor/MovieSceneTools/Private/TrackEditors/CameraShakeTrackEditor.cpp
--- a/Engine/Source/Editor/MovieSceneTools/Private/TrackEditors/CameraShakeTrackEditor.cpp
+++ b/Engine/Source/Editor/MovieSceneTools/Private/TrackEditors/CameraShakeTrackEditor.cpp
@@ -124,10 +124,10 @@ bool FCameraShakeTrackEditor::HandleAssetAdded(UObject* Asset, const FGuid& Targ
 	return false;
 }
 
-void FCameraShakeTrackEditor::BuildObjectBindingTrackMenu(FMenuBuilder& MenuBuilder, const FGuid& ObjectBinding, const UClass* ObjectClass)
+void FCameraShakeTrackEditor::BuildObjectBindingTrackMenu(FMenuBuilder& MenuBuilder, const TArray<FGuid>& ObjectBindings, const UClass* ObjectClass)
 {
 	// only offer this track if we can find a camera component
-	UCameraComponent const* const CamComponent = AcquireCameraComponentFromObjectGuid(ObjectBinding);
+	UCameraComponent const* const CamComponent = AcquireCameraComponentFromObjectGuid(ObjectBindings[0]);
 	if (CamComponent)
 	{
 		const TSharedPtr<ISequencer> ParentSequencer = GetSequencer();
@@ -147,7 +147,7 @@ void FCameraShakeTrackEditor::BuildObjectBindingTrackMenu(FMenuBuilder& MenuBuil
 
 		MenuBuilder.AddSubMenu(
 			LOCTEXT("AddCameraShake", "Camera Shake"), NSLOCTEXT("Sequencer", "AddCameraShakeTooltip", "Adds an additive camera shake track."),
-			FNewMenuDelegate::CreateRaw(this, &FCameraShakeTrackEditor::AddCameraShakeSubMenu, ObjectBinding)
+			FNewMenuDelegate::CreateRaw(this, &FCameraShakeTrackEditor::AddCameraShakeSubMenu, ObjectBindings)
 			);
 	}
 }
@@ -156,18 +156,21 @@ TSharedRef<SWidget> FCameraShakeTrackEditor::BuildCameraShakeSubMenu(FGuid Objec
 {
 	FMenuBuilder MenuBuilder(true, nullptr);
 
-	AddCameraShakeSubMenu(MenuBuilder, ObjectBinding);
+	TArray<FGuid> ObjectBindings;
+	ObjectBindings.Add(ObjectBinding);
+
+	AddCameraShakeSubMenu(MenuBuilder, ObjectBindings);
 
 	return MenuBuilder.MakeWidget();
 }
 
 
-void FCameraShakeTrackEditor::AddCameraShakeSubMenu(FMenuBuilder& MenuBuilder, FGuid ObjectBinding)
+void FCameraShakeTrackEditor::AddCameraShakeSubMenu(FMenuBuilder& MenuBuilder, TArray<FGuid> ObjectBindings)
 {
 	FAssetPickerConfig AssetPickerConfig;
 	{
-		AssetPickerConfig.OnAssetSelected = FOnAssetSelected::CreateRaw(this, &FCameraShakeTrackEditor::OnCameraShakeAssetSelected, ObjectBinding);
-		AssetPickerConfig.OnAssetEnterPressed = FOnAssetEnterPressed::CreateRaw(this, &FCameraShakeTrackEditor::OnCameraShakeAssetEnterPressed, ObjectBinding);
+		AssetPickerConfig.OnAssetSelected = FOnAssetSelected::CreateRaw(this, &FCameraShakeTrackEditor::OnCameraShakeAssetSelected, ObjectBindings);
+		AssetPickerConfig.OnAssetEnterPressed = FOnAssetEnterPressed::CreateRaw(this, &FCameraShakeTrackEditor::OnCameraShakeAssetEnterPressed, ObjectBindings);
 		AssetPickerConfig.bAllowNullSelection = false;
 		AssetPickerConfig.InitialAssetViewType = EAssetViewType::List;
 		AssetPickerConfig.Filter.ClassNames.Add(UBlueprint::StaticClass()->GetFName());
@@ -201,7 +204,7 @@ TSharedPtr<SWidget> FCameraShakeTrackEditor::BuildOutlinerEditWidget(const FGuid
 }
 
 
-void FCameraShakeTrackEditor::OnCameraShakeAssetSelected(const FAssetData& AssetData, FGuid ObjectBinding)
+void FCameraShakeTrackEditor::OnCameraShakeAssetSelected(const FAssetData& AssetData, TArray<FGuid> ObjectBindings)
 {
 	FSlateApplication::Get().DismissAllMenus();
 
@@ -211,10 +214,16 @@ void FCameraShakeTrackEditor::OnCameraShakeAssetSelected(const FAssetData& Asset
 		TSubclassOf<UCameraShake> const ShakeClass = *(SelectedObject->GeneratedClass);
 
 		TArray<TWeakObjectPtr<>> OutObjects;
-		for (TWeakObjectPtr<> Object : GetSequencer()->FindObjectsInCurrentSequence(ObjectBinding))
+		for (FGuid ObjectBinding : ObjectBindings)
 		{
-			OutObjects.Add(Object);
+			for (TWeakObjectPtr<> Object : GetSequencer()->FindObjectsInCurrentSequence(ObjectBinding))
+			{
+				OutObjects.Add(Object);
+			}
 		}
+		
+		const FScopedTransaction Transaction(LOCTEXT("AddCameraShake_Transaction", "Add Camera Shake"));
+
 		AnimatablePropertyChanged(FOnKeyProperty::CreateRaw(this, &FCameraShakeTrackEditor::AddKeyInternal, OutObjects, ShakeClass));
 	}
 }
@@ -252,11 +261,11 @@ FKeyPropertyResult FCameraShakeTrackEditor::AddKeyInternal(FFrameNumber KeyTime,
 	return KeyPropertyResult;
 }
 
-void FCameraShakeTrackEditor::OnCameraShakeAssetEnterPressed(const TArray<FAssetData>& AssetData, FGuid ObjectBinding)
+void FCameraShakeTrackEditor::OnCameraShakeAssetEnterPressed(const TArray<FAssetData>& AssetData, TArray<FGuid> ObjectBindings)
 {
 	if (AssetData.Num() > 0)
 	{
-		OnCameraShakeAssetSelected(AssetData[0].GetAsset(), ObjectBinding);
+		OnCameraShakeAssetSelected(AssetData[0].GetAsset(), ObjectBindings);
 	}
 }
 
diff --git a/Engine/Source/Editor/MovieSceneTools/Private/TrackEditors/CameraShakeTrackEditor.h b/Engine/Source/Editor/MovieSceneTools/Private/TrackEditors/CameraShakeTrackEditor.h
--- a/Engine/Source/Editor/MovieSceneTools/Private/TrackEditors/CameraShakeTrackEditor.h
+++ b/Engine/Source/Editor/MovieSceneTools/Private/TrackEditors/CameraShakeTrackEditor.h
@@ -45,7 +45,7 @@ public:
 
 	// ISequencerTrackEditor interface
 	virtual void AddKey(const FGuid& ObjectGuid) override;
-	virtual void BuildObjectBindingTrackMenu(FMenuBuilder& MenuBuilder, const FGuid& ObjectBinding, const UClass* ObjectClass) override;
+	virtual void BuildObjectBindingTrackMenu(FMenuBuilder& MenuBuilder, const TArray<FGuid>& ObjectBindings, const UClass* ObjectClass) override;
 	virtual bool HandleAssetAdded(UObject* Asset, const FGuid& TargetObjectGuid) override;
 	virtual TSharedRef<ISequencerSection> MakeSectionInterface(UMovieSceneSection& SectionObject, UMovieSceneTrack& Track, FGuid ObjectBinding) override;
 	virtual bool SupportsType(TSubclassOf<UMovieSceneTrack> Type) const override;
@@ -58,13 +58,13 @@ private:
 
 	/** Animation sub menu */
 	TSharedRef<SWidget> BuildCameraShakeSubMenu(FGuid ObjectBinding);
-	void AddCameraShakeSubMenu(FMenuBuilder& MenuBuilder, FGuid ObjectBinding);
+	void AddCameraShakeSubMenu(FMenuBuilder& MenuBuilder, TArray<FGuid> ObjectBindings);
 
 	/** Animation asset selected */
-	void OnCameraShakeAssetSelected(const FAssetData& AssetData, FGuid ObjectBinding);
+	void OnCameraShakeAssetSelected(const FAssetData& AssetData, TArray<FGuid> ObjectBindings);
 
 	/** Animation asset enter pressed */
-	void OnCameraShakeAssetEnterPressed(const TArray<FAssetData>& AssetData, FGuid ObjectBinding);
+	void OnCameraShakeAssetEnterPressed(const TArray<FAssetData>& AssetData, TArray<FGuid> ObjectBindings);
 
 	class UCameraComponent* AcquireCameraComponentFromObjectGuid(const FGuid& Guid);
 };
diff --git a/Engine/Source/Editor/MovieSceneTools/Private/TrackEditors/EventTrackEditor.cpp b/Engine/Source/Editor/MovieSceneTools/Private/TrackEditors/EventTrackEditor.cpp
--- a/Engine/Source/Editor/MovieSceneTools/Private/TrackEditors/EventTrackEditor.cpp
+++ b/Engine/Source/Editor/MovieSceneTools/Private/TrackEditors/EventTrackEditor.cpp
@@ -77,13 +77,13 @@ void FEventTrackEditor::BuildAddTrackMenu(FMenuBuilder& MenuBuilder)
 			LOCTEXT("AddEventTooltip", "Adds a new event track that can trigger events on the timeline."),
 			FSlateIcon(FEditorStyle::GetStyleSetName(), "Sequencer.Tracks.Event"),
 			FUIAction(
-				FExecuteAction::CreateRaw(this, &FEventTrackEditor::HandleAddEventTrackMenuEntryExecute, FGuid())
+				FExecuteAction::CreateRaw(this, &FEventTrackEditor::HandleAddEventTrackMenuEntryExecute, TArray<FGuid>())
 			)
 		);
 	}
 }
 
-void FEventTrackEditor::BuildObjectBindingTrackMenu(FMenuBuilder& MenuBuilder, const FGuid& ObjectBinding, const UClass* ObjectClass)
+void FEventTrackEditor::BuildObjectBindingTrackMenu(FMenuBuilder& MenuBuilder, const TArray<FGuid>& ObjectBindings, const UClass* ObjectClass)
 {
 	UMovieSceneSequence*       RootMovieSceneSequence = GetSequencer()->GetRootMovieSceneSequence();
 	FMovieSceneSequenceEditor* SequenceEditor         = FMovieSceneSequenceEditor::Find(RootMovieSceneSequence);
@@ -95,7 +95,7 @@ void FEventTrackEditor::BuildObjectBindingTrackMenu(FMenuBuilder& MenuBuilder, c
 			LOCTEXT("AddEventTooltip_ObjectBinding", "Adds a new event track that will trigger events on this object binding."),
 			FSlateIcon(),
 			FUIAction( 
-				FExecuteAction::CreateSP( this, &FEventTrackEditor::HandleAddEventTrackMenuEntryExecute, ObjectBinding )
+				FExecuteAction::CreateSP( this, &FEventTrackEditor::HandleAddEventTrackMenuEntryExecute, ObjectBindings )
 			)
 		);
 	}
@@ -229,7 +229,7 @@ const FSlateBrush* FEventTrackEditor::GetIconBrush() const
 /* FEventTrackEditor callbacks
  *****************************************************************************/
 
-void FEventTrackEditor::HandleAddEventTrackMenuEntryExecute(FGuid InObjectBindingID)
+void FEventTrackEditor::HandleAddEventTrackMenuEntryExecute(TArray<FGuid> InObjectBindingIDs)
 {
 	UMovieScene* FocusedMovieScene = GetFocusedMovieScene();
 
@@ -246,27 +246,35 @@ void FEventTrackEditor::HandleAddEventTrackMenuEntryExecute(FGuid InObjectBindin
 	const FScopedTransaction Transaction(NSLOCTEXT("Sequencer", "AddEventTrack_Transaction", "Add Event Track"));
 	FocusedMovieScene->Modify();
 
-	UMovieSceneEventTrack* NewTrack = nullptr;
-	if (InObjectBindingID.IsValid())
+	TArray<UMovieSceneEventTrack*> NewTracks;
+
+	for (FGuid InObjectBindingID : InObjectBindingIDs)
 	{
-		NewTrack = FocusedMovieScene->AddTrack<UMovieSceneEventTrack>(InObjectBindingID);
+		if (InObjectBindingID.IsValid())
+		{
+			NewTracks.Add(FocusedMovieScene->AddTrack<UMovieSceneEventTrack>(InObjectBindingID));
+		}
 	}
-	else
+
+	if (!NewTracks.Num())
 	{
-		NewTrack = FocusedMovieScene->AddMasterTrack<UMovieSceneEventTrack>();
+		NewTracks.Add(FocusedMovieScene->AddMasterTrack<UMovieSceneEventTrack>());
 	}
 
-	check(NewTrack);
+	check(NewTracks.Num() != 0);
 
-	UMovieSceneSection* NewSection = NewTrack->CreateNewSection();
-	check(NewSection);
+	for (UMovieSceneEventTrack* NewTrack : NewTracks)
+	{
+		UMovieSceneSection* NewSection = NewTrack->CreateNewSection();
+		check(NewSection);
 
-	NewTrack->AddSection(*NewSection);
-	NewTrack->SetDisplayName(LOCTEXT("TrackName", "Events"));
+		NewTrack->AddSection(*NewSection);
+		NewTrack->SetDisplayName(LOCTEXT("TrackName", "Events"));
 
-	if (GetSequencer().IsValid())
-	{
-		GetSequencer()->OnAddTrack(NewTrack);
+		if (GetSequencer().IsValid())
+		{
+			GetSequencer()->OnAddTrack(NewTrack);
+		}
 	}
 
 	GetSequencer()->NotifyMovieSceneDataChanged( EMovieSceneDataChangeType::MovieSceneStructureItemAdded );
diff --git a/Engine/Source/Editor/MovieSceneTools/Private/TrackEditors/EventTrackEditor.h b/Engine/Source/Editor/MovieSceneTools/Private/TrackEditors/EventTrackEditor.h
--- a/Engine/Source/Editor/MovieSceneTools/Private/TrackEditors/EventTrackEditor.h
+++ b/Engine/Source/Editor/MovieSceneTools/Private/TrackEditors/EventTrackEditor.h
@@ -43,7 +43,7 @@ public:
 	// ISequencerTrackEditor interface
 
 	virtual void BuildAddTrackMenu(FMenuBuilder& MenuBuilder) override;
-	virtual void BuildObjectBindingTrackMenu(FMenuBuilder& MenuBuilder, const FGuid& ObjectBinding, const UClass* ObjectClass) override;
+	virtual void BuildObjectBindingTrackMenu(FMenuBuilder& MenuBuilder, const TArray<FGuid>& ObjectBindings, const UClass* ObjectClass) override;
 	virtual bool SupportsType( TSubclassOf<UMovieSceneTrack> Type ) const override;
 	virtual bool SupportsSequence(UMovieSceneSequence* InSequence) const override;
 	virtual const FSlateBrush* GetIconBrush() const override;
@@ -56,7 +56,7 @@ public:
 private:
 
 	/** Callback for executing the "Add Event Track" menu entry. */
-	void HandleAddEventTrackMenuEntryExecute(FGuid InObjectBindingID);
+	void HandleAddEventTrackMenuEntryExecute(TArray<FGuid> InObjectBindingIDs);
 
 	void CreateNewSection(UMovieSceneTrack* Track, int32 RowIndex, UClass* SectionType);
 };
diff --git a/Engine/Source/Editor/MovieSceneTools/Private/TrackEditors/ParticleParameterTrackEditor.cpp b/Engine/Source/Editor/MovieSceneTools/Private/TrackEditors/ParticleParameterTrackEditor.cpp
--- a/Engine/Source/Editor/MovieSceneTools/Private/TrackEditors/ParticleParameterTrackEditor.cpp
+++ b/Engine/Source/Editor/MovieSceneTools/Private/TrackEditors/ParticleParameterTrackEditor.cpp
@@ -46,7 +46,7 @@ TSharedPtr<SWidget> FParticleParameterTrackEditor::BuildOutlinerEditWidget( cons
 }
 
 
-void FParticleParameterTrackEditor::BuildObjectBindingTrackMenu( FMenuBuilder& MenuBuilder, const FGuid& ObjectBinding, const UClass* ObjectClass )
+void FParticleParameterTrackEditor::BuildObjectBindingTrackMenu( FMenuBuilder& MenuBuilder, const TArray<FGuid>& ObjectBindings, const UClass* ObjectClass )
 {
 	if ( ObjectClass->IsChildOf( AEmitter::StaticClass() ) || ObjectClass->IsChildOf( UParticleSystemComponent::StaticClass() ) )
 	{
@@ -58,8 +58,8 @@ void FParticleParameterTrackEditor::BuildObjectBindingTrackMenu( FMenuBuilder& M
 			FSlateIcon(),
 			FUIAction
 			(
-				FExecuteAction::CreateSP( this, &FParticleParameterTrackEditor::AddParticleParameterTrack, ObjectBinding ),
-				FCanExecuteAction::CreateSP( this, &FParticleParameterTrackEditor::CanAddParticleParameterTrack, ObjectBinding )
+				FExecuteAction::CreateSP( this, &FParticleParameterTrackEditor::AddParticleParameterTrack, ObjectBindings ),
+				FCanExecuteAction::CreateSP( this, &FParticleParameterTrackEditor::CanAddParticleParameterTrack, ObjectBindings[0] )
 			));
 	}
 }
@@ -128,10 +128,16 @@ bool FParticleParameterTrackEditor::CanAddParticleParameterTrack( FGuid ObjectBi
 }
 
 
-void FParticleParameterTrackEditor::AddParticleParameterTrack( FGuid ObjectBinding )
+void FParticleParameterTrackEditor::AddParticleParameterTrack( TArray<FGuid> ObjectBindings )
 {
-	FindOrCreateTrackForObject( ObjectBinding, UMovieSceneParticleParameterTrack::StaticClass(), TrackName, true);
-	GetSequencer()->NotifyMovieSceneDataChanged( EMovieSceneDataChangeType::MovieSceneStructureItemAdded );
+	const FScopedTransaction Transaction(LOCTEXT("AddParticleParameterTrack", "Add particle parameter track"));
+
+	for (FGuid ObjectBinding : ObjectBindings)
+	{
+		FindOrCreateTrackForObject(ObjectBinding, UMovieSceneParticleParameterTrack::StaticClass(), TrackName, true);
+	}
+
+	GetSequencer()->NotifyMovieSceneDataChanged(EMovieSceneDataChangeType::MovieSceneStructureItemAdded);
 }
 
 
diff --git a/Engine/Source/Editor/MovieSceneTools/Private/TrackEditors/ParticleParameterTrackEditor.h b/Engine/Source/Editor/MovieSceneTools/Private/TrackEditors/ParticleParameterTrackEditor.h
--- a/Engine/Source/Editor/MovieSceneTools/Private/TrackEditors/ParticleParameterTrackEditor.h
+++ b/Engine/Source/Editor/MovieSceneTools/Private/TrackEditors/ParticleParameterTrackEditor.h
@@ -68,14 +68,14 @@ private:
 		}
 	};
 
-	void BuildObjectBindingTrackMenu( FMenuBuilder& MenuBuilder, const FGuid& ObjectBinding, const UClass* ObjectClass );
+	void BuildObjectBindingTrackMenu( FMenuBuilder& MenuBuilder, const TArray<FGuid>& ObjectBindings, const UClass* ObjectClass );
 
 	/** Provides the contents of the add parameter menu. */
 	TSharedRef<SWidget> OnGetAddParameterMenuContent( FGuid ObjectBinding, UMovieSceneParticleParameterTrack* ParticleParameterTrack );
 
 	bool CanAddParticleParameterTrack (FGuid ObjectBinding );
 
-	void AddParticleParameterTrack( FGuid ObjectBinding );
+	void AddParticleParameterTrack( TArray<FGuid> ObjectBindings );
 
 	/** Adds a scalar parameter and initial key to a material track.
 	 * @param ObjectBinding The object binding which owns the material track.
diff --git a/Engine/Source/Editor/MovieSceneTools/Private/TrackEditors/ParticleTrackEditor.cpp b/Engine/Source/Editor/MovieSceneTools/Private/TrackEditors/ParticleTrackEditor.cpp
--- a/Engine/Source/Editor/MovieSceneTools/Private/TrackEditors/ParticleTrackEditor.cpp
+++ b/Engine/Source/Editor/MovieSceneTools/Private/TrackEditors/ParticleTrackEditor.cpp
@@ -237,7 +237,7 @@ TSharedRef<ISequencerSection> FParticleTrackEditor::MakeSectionInterface( UMovie
 }
 
 
-void FParticleTrackEditor::BuildObjectBindingTrackMenu(FMenuBuilder& MenuBuilder, const FGuid& ObjectBinding, const UClass* ObjectClass)
+void FParticleTrackEditor::BuildObjectBindingTrackMenu(FMenuBuilder& MenuBuilder, const TArray<FGuid>& ObjectBindings, const UClass* ObjectClass)
 {
 	if (ObjectClass->IsChildOf(AEmitter::StaticClass()) || ObjectClass->IsChildOf(UParticleSystemComponent::StaticClass()))
 	{
@@ -247,20 +247,25 @@ void FParticleTrackEditor::BuildObjectBindingTrackMenu(FMenuBuilder& MenuBuilder
 			LOCTEXT("AddParticleTrack", "Particle Toggle Track"),
 			LOCTEXT("TriggerParticlesTooltip", "Adds a track for controlling particle emitter state."),
 			FSlateIcon(),
-			FUIAction(FExecuteAction::CreateSP(this, &FParticleTrackEditor::AddParticleKey, ObjectBinding))
+			FUIAction(FExecuteAction::CreateSP(this, &FParticleTrackEditor::AddParticleKey, ObjectBindings))
 		);
 	}
 }
 
 
-void FParticleTrackEditor::AddParticleKey( const FGuid ObjectGuid )
+void FParticleTrackEditor::AddParticleKey( TArray<FGuid> ObjectGuids )
 {
-	TSharedPtr<ISequencer> SequencerPtr = GetSequencer();
-	UObject* Object = SequencerPtr.IsValid() ? SequencerPtr->FindSpawnedObjectOrTemplate(ObjectGuid) : nullptr;
+	const FScopedTransaction Transaction(NSLOCTEXT("Sequencer", "AddParticleKey", "Add Particle Key"));
 
-	if (Object)
+	TSharedPtr<ISequencer> SequencerPtr = GetSequencer();
+	for (FGuid ObjectGuid : ObjectGuids)
 	{
-		AnimatablePropertyChanged( FOnKeyProperty::CreateRaw( this, &FParticleTrackEditor::AddKeyInternal, Object ) );
+		UObject* Object = SequencerPtr.IsValid() ? SequencerPtr->FindSpawnedObjectOrTemplate(ObjectGuid) : nullptr;
+
+		if (Object)
+		{
+			AnimatablePropertyChanged(FOnKeyProperty::CreateRaw(this, &FParticleTrackEditor::AddKeyInternal, Object));
+		}
 	}
 }
 
diff --git a/Engine/Source/Editor/MovieSceneTools/Private/TrackEditors/ParticleTrackEditor.h b/Engine/Source/Editor/MovieSceneTools/Private/TrackEditors/ParticleTrackEditor.h
--- a/Engine/Source/Editor/MovieSceneTools/Private/TrackEditors/ParticleTrackEditor.h
+++ b/Engine/Source/Editor/MovieSceneTools/Private/TrackEditors/ParticleTrackEditor.h
@@ -49,12 +49,12 @@ public:
 
 	// ISequencerTrackEditor interface
 
-	virtual void BuildObjectBindingTrackMenu(FMenuBuilder& MenuBuilder, const FGuid& ObjectBinding, const UClass* ObjectClass) override;
+	virtual void BuildObjectBindingTrackMenu(FMenuBuilder& MenuBuilder, const TArray<FGuid>& ObjectBindings, const UClass* ObjectClass) override;
 	virtual TSharedRef<ISequencerSection> MakeSectionInterface( UMovieSceneSection& SectionObject, UMovieSceneTrack& Track, FGuid ObjectBinding ) override;
 	virtual bool SupportsType( TSubclassOf<UMovieSceneTrack> Type ) const override;
 	virtual void BuildTrackContextMenu( FMenuBuilder& MenuBuilder, UMovieSceneTrack* Track ) override;
 	
-	void AddParticleKey(const FGuid ObjectGuid);
+	void AddParticleKey(TArray<FGuid> ObjectGuids);
 private:
 
 	/** Delegate for AnimatablePropertyChanged in AddKey. */
diff --git a/Engine/Source/Editor/MovieSceneTools/Private/TrackEditors/PathTrackEditor.cpp b/Engine/Source/Editor/MovieSceneTools/Private/TrackEditors/PathTrackEditor.cpp
--- a/Engine/Source/Editor/MovieSceneTools/Private/TrackEditors/PathTrackEditor.cpp
+++ b/Engine/Source/Editor/MovieSceneTools/Private/TrackEditors/PathTrackEditor.cpp
@@ -71,9 +71,12 @@ public:
 	
 	virtual void BuildSectionContextMenu(FMenuBuilder& MenuBuilder, const FGuid& ObjectBinding) override
 	{
+		TArray<FGuid> ObjectBindings;
+		ObjectBindings.Add(ObjectBinding);
+
 		MenuBuilder.AddSubMenu(
 			LOCTEXT("SetPath", "Path"), LOCTEXT("SetPathTooltip", "Set path"),
-			FNewMenuDelegate::CreateRaw(PathTrackEditor, &FActorPickerTrackEditor::ShowActorSubMenu, ObjectBinding, &Section));
+			FNewMenuDelegate::CreateRaw(PathTrackEditor, &FActorPickerTrackEditor::ShowActorSubMenu, ObjectBindings, &Section));
 	}
 
 private:
@@ -113,7 +116,7 @@ TSharedRef<ISequencerSection> F3DPathTrackEditor::MakeSectionInterface( UMovieSc
 }
 
 
-void F3DPathTrackEditor::BuildObjectBindingTrackMenu(FMenuBuilder& MenuBuilder, const FGuid& ObjectBinding, const UClass* ObjectClass)
+void F3DPathTrackEditor::BuildObjectBindingTrackMenu(FMenuBuilder& MenuBuilder, const TArray<FGuid>& ObjectBindings, const UClass* ObjectClass)
 {
 	if (ObjectClass && ObjectClass->IsChildOf(AActor::StaticClass()))
 	{
@@ -121,7 +124,7 @@ void F3DPathTrackEditor::BuildObjectBindingTrackMenu(FMenuBuilder& MenuBuilder,
 
 		MenuBuilder.AddSubMenu(
 			LOCTEXT("AddPath", "Path"), LOCTEXT("AddPathTooltip", "Adds a path track."),
-			FNewMenuDelegate::CreateRaw(this, &FActorPickerTrackEditor::ShowActorSubMenu, ObjectBinding, DummySection));
+			FNewMenuDelegate::CreateRaw(this, &FActorPickerTrackEditor::ShowActorSubMenu, ObjectBindings, DummySection));
 	}
 }
 
@@ -170,7 +173,7 @@ bool F3DPathTrackEditor::IsActorPickable(const AActor* const ParentActor, FGuid
 }
 
 
-void F3DPathTrackEditor::ActorSocketPicked(const FName SocketName, USceneComponent* Component, FActorPickerID ActorPickerID, FGuid ObjectGuid, UMovieSceneSection* Section)
+void F3DPathTrackEditor::ActorSocketPicked(const FName SocketName, USceneComponent* Component, FActorPickerID ActorPickerID, TArray<FGuid> ObjectGuids, UMovieSceneSection* Section)
 {
 	if (Section != nullptr)
 	{
@@ -195,13 +198,18 @@ void F3DPathTrackEditor::ActorSocketPicked(const FName SocketName, USceneCompone
 			PathSection->SetConstraintBindingID(ConstraintBindingID);
 		}
 	}
-	else if (ObjectGuid.IsValid())
+	else
 	{
 		TArray<TWeakObjectPtr<>> OutObjects;
-		for (TWeakObjectPtr<> Object : GetSequencer()->FindObjectsInCurrentSequence(ObjectGuid))
+
+		for (FGuid ObjectGuid : ObjectGuids)
 		{
-			OutObjects.Add(Object);
+			for (TWeakObjectPtr<> Object : GetSequencer()->FindObjectsInCurrentSequence(ObjectGuid))
+			{
+				OutObjects.Add(Object);
+			}
 		}
+
 		AnimatablePropertyChanged( FOnKeyProperty::CreateRaw( this, &F3DPathTrackEditor::AddKeyInternal, OutObjects, ActorPickerID) );
 	}
 }
diff --git a/Engine/Source/Editor/MovieSceneTools/Private/TrackEditors/PathTrackEditor.h b/Engine/Source/Editor/MovieSceneTools/Private/TrackEditors/PathTrackEditor.h
--- a/Engine/Source/Editor/MovieSceneTools/Private/TrackEditors/PathTrackEditor.h
+++ b/Engine/Source/Editor/MovieSceneTools/Private/TrackEditors/PathTrackEditor.h
@@ -45,13 +45,13 @@ public:
 
 	// ISequencerTrackEditor interface
 
-	virtual void BuildObjectBindingTrackMenu(FMenuBuilder& MenuBuilder, const FGuid& ObjectBinding, const UClass* ObjectClass) override;
+	virtual void BuildObjectBindingTrackMenu(FMenuBuilder& MenuBuilder, const TArray<FGuid>& ObjectBindings, const UClass* ObjectClass) override;
 	virtual TSharedRef<ISequencerSection> MakeSectionInterface( UMovieSceneSection& SectionObject, UMovieSceneTrack& Track, FGuid ObjectBinding ) override;
 	virtual bool SupportsType( TSubclassOf<UMovieSceneTrack> Type ) const override;
 
 	// FTrackEditorActorPicker
 	virtual bool IsActorPickable( const AActor* const ParentActor, FGuid ObjectBinding, UMovieSceneSection* InSection ) override;
-	virtual void ActorSocketPicked(const FName SocketName, USceneComponent* Component, FActorPickerID ActorPickerID, FGuid ObjectBinding, UMovieSceneSection* Section) override;
+	virtual void ActorSocketPicked(const FName SocketName, USceneComponent* Component, FActorPickerID ActorPickerID, TArray<FGuid> ObjectBindings, UMovieSceneSection* Section) override;
 
 private:
 
diff --git a/Engine/Source/Editor/MovieSceneTools/Private/TrackEditors/PrimitiveMaterialTrackEditor.cpp b/Engine/Source/Editor/MovieSceneTools/Private/TrackEditors/PrimitiveMaterialTrackEditor.cpp
--- a/Engine/Source/Editor/MovieSceneTools/Private/TrackEditors/PrimitiveMaterialTrackEditor.cpp
+++ b/Engine/Source/Editor/MovieSceneTools/Private/TrackEditors/PrimitiveMaterialTrackEditor.cpp
@@ -21,15 +21,15 @@ TSharedRef<ISequencerTrackEditor> FPrimitiveMaterialTrackEditor::CreateTrackEdit
 	return MakeShared<FPrimitiveMaterialTrackEditor>(OwningSequencer);
 }
 
-void FPrimitiveMaterialTrackEditor::ExtendObjectBindingTrackMenu(TSharedRef<FExtender> Extender, const FGuid& ObjectBinding, const UClass* ObjectClass)
+void FPrimitiveMaterialTrackEditor::ExtendObjectBindingTrackMenu(TSharedRef<FExtender> Extender, const TArray<FGuid>& ObjectBindings, const UClass* ObjectClass)
 {
 	if (ObjectClass->IsChildOf(UPrimitiveComponent::StaticClass()))
 	{
-		Extender->AddMenuExtension(SequencerMenuExtensionPoints::AddTrackMenu_PropertiesSection, EExtensionHook::Before, nullptr, FMenuExtensionDelegate::CreateSP(this, &FPrimitiveMaterialTrackEditor::ConstructObjectBindingTrackMenu, ObjectBinding));
+		Extender->AddMenuExtension(SequencerMenuExtensionPoints::AddTrackMenu_PropertiesSection, EExtensionHook::Before, nullptr, FMenuExtensionDelegate::CreateSP(this, &FPrimitiveMaterialTrackEditor::ConstructObjectBindingTrackMenu, ObjectBindings));
 	}
 }
 
-void FPrimitiveMaterialTrackEditor::ConstructObjectBindingTrackMenu(FMenuBuilder& MenuBuilder, FGuid ObjectBinding)
+void FPrimitiveMaterialTrackEditor::ConstructObjectBindingTrackMenu(FMenuBuilder& MenuBuilder, TArray<FGuid> ObjectBindings)
 {
 	auto GetMaterialIndexForTrack = [](UMovieSceneTrack* InTrack)
 	{
@@ -39,7 +39,7 @@ void FPrimitiveMaterialTrackEditor::ConstructObjectBindingTrackMenu(FMenuBuilder
 
 	int32 MinNumMaterials = TNumericLimits<int32>::Max();
 
-	for (TWeakObjectPtr<> WeakObject : GetSequencer()->FindObjectsInCurrentSequence(ObjectBinding))
+	for (TWeakObjectPtr<> WeakObject : GetSequencer()->FindObjectsInCurrentSequence(ObjectBindings[0]))
 	{
 		UPrimitiveComponent* PrimitiveComponent = Cast<UPrimitiveComponent>(WeakObject.Get());
 		if (PrimitiveComponent)
@@ -53,7 +53,7 @@ void FPrimitiveMaterialTrackEditor::ConstructObjectBindingTrackMenu(FMenuBuilder
 		MenuBuilder.BeginSection(NAME_None, LOCTEXT("MaterialSwitcherTitle", "Material Switchers"));
 
 		const UMovieScene*        MovieScene = GetFocusedMovieScene();
-		const FMovieSceneBinding* Binding    = Algo::FindBy(MovieScene->GetBindings(), ObjectBinding, &FMovieSceneBinding::GetObjectGuid);
+		const FMovieSceneBinding* Binding    = Algo::FindBy(MovieScene->GetBindings(), ObjectBindings[0], &FMovieSceneBinding::GetObjectGuid);
 
 		check(Binding);
 
@@ -67,7 +67,7 @@ void FPrimitiveMaterialTrackEditor::ConstructObjectBindingTrackMenu(FMenuBuilder
 					FText(),
 					FSlateIcon(),
 					FUIAction(
-						FExecuteAction::CreateSP(this, &FPrimitiveMaterialTrackEditor::CreateTrackForElement, ObjectBinding, Index)
+						FExecuteAction::CreateSP(this, &FPrimitiveMaterialTrackEditor::CreateTrackForElement, ObjectBindings, Index)
 					)
 				);
 			}
@@ -77,18 +77,21 @@ void FPrimitiveMaterialTrackEditor::ConstructObjectBindingTrackMenu(FMenuBuilder
 	}
 }
 
-void FPrimitiveMaterialTrackEditor::CreateTrackForElement(FGuid ObjectBindingID, int32 MaterialIndex)
+void FPrimitiveMaterialTrackEditor::CreateTrackForElement(TArray<FGuid> ObjectBindingIDs, int32 MaterialIndex)
 {
 	UMovieScene* MovieScene = GetFocusedMovieScene();
 
 	FScopedTransaction Transaction(LOCTEXT("CreateTrack", "Create Material Track"));
 	MovieScene->Modify();
 
-	UMovieScenePrimitiveMaterialTrack* NewTrack = MovieScene->AddTrack<UMovieScenePrimitiveMaterialTrack>(ObjectBindingID);
-	NewTrack->MaterialIndex = MaterialIndex;
-	NewTrack->SetDisplayName(FText::Format(LOCTEXT("MaterialTrackName_Format", "Material Element {0}"), FText::AsNumber(MaterialIndex)));
+	for (FGuid ObjectBindingID : ObjectBindingIDs)
+	{
+		UMovieScenePrimitiveMaterialTrack* NewTrack = MovieScene->AddTrack<UMovieScenePrimitiveMaterialTrack>(ObjectBindingID);
+		NewTrack->MaterialIndex = MaterialIndex;
+		NewTrack->SetDisplayName(FText::Format(LOCTEXT("MaterialTrackName_Format", "Material Element {0}"), FText::AsNumber(MaterialIndex)));
 
-	NewTrack->AddSection(*NewTrack->CreateNewSection());
+		NewTrack->AddSection(*NewTrack->CreateNewSection());
+	}
 
 	GetSequencer()->NotifyMovieSceneDataChanged(EMovieSceneDataChangeType::MovieSceneStructureItemAdded);
 }
diff --git a/Engine/Source/Editor/MovieSceneTools/Private/TrackEditors/PrimitiveMaterialTrackEditor.h b/Engine/Source/Editor/MovieSceneTools/Private/TrackEditors/PrimitiveMaterialTrackEditor.h
--- a/Engine/Source/Editor/MovieSceneTools/Private/TrackEditors/PrimitiveMaterialTrackEditor.h
+++ b/Engine/Source/Editor/MovieSceneTools/Private/TrackEditors/PrimitiveMaterialTrackEditor.h
@@ -19,8 +19,8 @@ public:
 public:
 
 	//~ ISequencerTrackEditor interface
-	virtual void ExtendObjectBindingTrackMenu(TSharedRef<FExtender> Extender, const FGuid& ObjectBinding, const UClass* ObjectClass) override;
+	virtual void ExtendObjectBindingTrackMenu(TSharedRef<FExtender> Extender, const TArray<FGuid>& ObjectBindings, const UClass* ObjectClass) override;
 
-	void ConstructObjectBindingTrackMenu(FMenuBuilder& MenuBuilder, FGuid ObjectBinding);
-	void CreateTrackForElement(FGuid ObjectBindingID, int32 MaterialIndex);
+	void ConstructObjectBindingTrackMenu(FMenuBuilder& MenuBuilder, TArray<FGuid> ObjectBindings);
+	void CreateTrackForElement(TArray<FGuid> ObjectBindingIDs, int32 MaterialIndex);
 };
diff --git a/Engine/Source/Editor/MovieSceneTools/Private/TrackEditors/SkeletalAnimationTrackEditor.cpp b/Engine/Source/Editor/MovieSceneTools/Private/TrackEditors/SkeletalAnimationTrackEditor.cpp
--- a/Engine/Source/Editor/MovieSceneTools/Private/TrackEditors/SkeletalAnimationTrackEditor.cpp
+++ b/Engine/Source/Editor/MovieSceneTools/Private/TrackEditors/SkeletalAnimationTrackEditor.cpp
@@ -503,13 +503,13 @@ bool FSkeletalAnimationTrackEditor::HandleAssetAdded(UObject* Asset, const FGuid
 }
 
 
-void FSkeletalAnimationTrackEditor::BuildObjectBindingTrackMenu(FMenuBuilder& MenuBuilder, const FGuid& ObjectBinding, const UClass* ObjectClass)
+void FSkeletalAnimationTrackEditor::BuildObjectBindingTrackMenu(FMenuBuilder& MenuBuilder, const TArray<FGuid>& ObjectBindings, const UClass* ObjectClass)
 {
 	if (ObjectClass->IsChildOf(USkeletalMeshComponent::StaticClass()) || ObjectClass->IsChildOf(AActor::StaticClass()))
 	{
 		const TSharedPtr<ISequencer> ParentSequencer = GetSequencer();
 
-		USkeleton* Skeleton = AcquireSkeletonFromObjectGuid(ObjectBinding, GetSequencer());
+		USkeleton* Skeleton = AcquireSkeletonFromObjectGuid(ObjectBindings[0], GetSequencer());
 
 		if (Skeleton)
 		{
@@ -526,7 +526,7 @@ void FSkeletalAnimationTrackEditor::BuildObjectBindingTrackMenu(FMenuBuilder& Me
 
 				MenuBuilder.AddSubMenu(
 					LOCTEXT("AddAnimation", "Animation"), NSLOCTEXT("Sequencer", "AddAnimationTooltip", "Adds an animation track."),
-					FNewMenuDelegate::CreateRaw(this, &FSkeletalAnimationTrackEditor::AddAnimationSubMenu, ObjectBinding, Skeleton, Track)
+					FNewMenuDelegate::CreateRaw(this, &FSkeletalAnimationTrackEditor::AddAnimationSubMenu, ObjectBindings, Skeleton, Track)
 				);
 			}
 		}
@@ -537,7 +537,10 @@ TSharedRef<SWidget> FSkeletalAnimationTrackEditor::BuildAnimationSubMenu(FGuid O
 {
 	FMenuBuilder MenuBuilder(true, nullptr);
 	
-	AddAnimationSubMenu(MenuBuilder, ObjectBinding, Skeleton, Track);
+	TArray<FGuid> ObjectBindings;
+	ObjectBindings.Add(ObjectBinding);
+
+	AddAnimationSubMenu(MenuBuilder, ObjectBindings, Skeleton, Track);
 
 	return MenuBuilder.MakeWidget();
 }
@@ -560,12 +563,12 @@ bool FSkeletalAnimationTrackEditor::ShouldFilterAsset(const FAssetData& AssetDat
 	return ((EAdditiveAnimationType)AdditiveTypeEnum->GetValueByName(*EnumString) == AAT_RotationOffsetMeshSpace);
 }
 
-void FSkeletalAnimationTrackEditor::AddAnimationSubMenu(FMenuBuilder& MenuBuilder, FGuid ObjectBinding, USkeleton* Skeleton, UMovieSceneTrack* Track)
+void FSkeletalAnimationTrackEditor::AddAnimationSubMenu(FMenuBuilder& MenuBuilder, TArray<FGuid> ObjectBindings, USkeleton* Skeleton, UMovieSceneTrack* Track)
 {
 	FAssetPickerConfig AssetPickerConfig;
 	{
-		AssetPickerConfig.OnAssetSelected = FOnAssetSelected::CreateRaw( this, &FSkeletalAnimationTrackEditor::OnAnimationAssetSelected, ObjectBinding, Track);
-		AssetPickerConfig.OnAssetEnterPressed = FOnAssetEnterPressed::CreateRaw( this, &FSkeletalAnimationTrackEditor::OnAnimationAssetEnterPressed, ObjectBinding, Track);
+		AssetPickerConfig.OnAssetSelected = FOnAssetSelected::CreateRaw( this, &FSkeletalAnimationTrackEditor::OnAnimationAssetSelected, ObjectBindings, Track);
+		AssetPickerConfig.OnAssetEnterPressed = FOnAssetEnterPressed::CreateRaw( this, &FSkeletalAnimationTrackEditor::OnAnimationAssetEnterPressed, ObjectBindings, Track);
 		AssetPickerConfig.bAllowNullSelection = false;
 		AssetPickerConfig.InitialAssetViewType = EAssetViewType::List;
 		AssetPickerConfig.OnShouldFilterAsset = FOnShouldFilterAsset::CreateRaw(this, &FSkeletalAnimationTrackEditor::ShouldFilterAsset);
@@ -587,7 +590,7 @@ void FSkeletalAnimationTrackEditor::AddAnimationSubMenu(FMenuBuilder& MenuBuilde
 }
 
 
-void FSkeletalAnimationTrackEditor::OnAnimationAssetSelected(const FAssetData& AssetData, FGuid ObjectBinding, UMovieSceneTrack* Track)
+void FSkeletalAnimationTrackEditor::OnAnimationAssetSelected(const FAssetData& AssetData, TArray<FGuid> ObjectBindings, UMovieSceneTrack* Track)
 {
 	FSlateApplication::Get().DismissAllMenus();
 
@@ -598,17 +601,22 @@ void FSkeletalAnimationTrackEditor::OnAnimationAssetSelected(const FAssetData& A
 	{
 		UAnimSequenceBase* AnimSequence = CastChecked<UAnimSequenceBase>(AssetData.GetAsset());
 
-		UObject* Object = SequencerPtr->FindSpawnedObjectOrTemplate(ObjectBinding);
-		int32 RowIndex = INDEX_NONE;
-		AnimatablePropertyChanged( FOnKeyProperty::CreateRaw( this, &FSkeletalAnimationTrackEditor::AddKeyInternal, Object, AnimSequence, Track, RowIndex) );
+		const FScopedTransaction Transaction(LOCTEXT("AddAnimation_Transaction", "Add Animation"));
+
+		for (FGuid ObjectBinding : ObjectBindings)
+		{
+			UObject* Object = SequencerPtr->FindSpawnedObjectOrTemplate(ObjectBinding);
+			int32 RowIndex = INDEX_NONE;
+			AnimatablePropertyChanged(FOnKeyProperty::CreateRaw(this, &FSkeletalAnimationTrackEditor::AddKeyInternal, Object, AnimSequence, Track, RowIndex));
+		}
 	}
 }
 
-void FSkeletalAnimationTrackEditor::OnAnimationAssetEnterPressed(const TArray<FAssetData>& AssetData, FGuid ObjectBinding, UMovieSceneTrack* Track)
+void FSkeletalAnimationTrackEditor::OnAnimationAssetEnterPressed(const TArray<FAssetData>& AssetData, TArray<FGuid> ObjectBindings, UMovieSceneTrack* Track)
 {
 	if (AssetData.Num() > 0)
 	{
-		OnAnimationAssetSelected(AssetData[0].GetAsset(), ObjectBinding, Track);
+		OnAnimationAssetSelected(AssetData[0].GetAsset(), ObjectBindings, Track);
 	}
 }
 
diff --git a/Engine/Source/Editor/MovieSceneTools/Private/TrackEditors/SkeletalAnimationTrackEditor.h b/Engine/Source/Editor/MovieSceneTools/Private/TrackEditors/SkeletalAnimationTrackEditor.h
--- a/Engine/Source/Editor/MovieSceneTools/Private/TrackEditors/SkeletalAnimationTrackEditor.h
+++ b/Engine/Source/Editor/MovieSceneTools/Private/TrackEditors/SkeletalAnimationTrackEditor.h
@@ -47,7 +47,7 @@ public:
 	// ISequencerTrackEditor interface
 
 	virtual void AddKey(const FGuid& ObjectGuid) override;
-	virtual void BuildObjectBindingTrackMenu(FMenuBuilder& MenuBuilder, const FGuid& ObjectBinding, const UClass* ObjectClass) override;
+	virtual void BuildObjectBindingTrackMenu(FMenuBuilder& MenuBuilder, const TArray<FGuid>& ObjectBindings, const UClass* ObjectClass) override;
 	virtual bool HandleAssetAdded(UObject* Asset, const FGuid& TargetObjectGuid) override;
 	virtual TSharedRef<ISequencerSection> MakeSectionInterface( UMovieSceneSection& SectionObject, UMovieSceneTrack& Track, FGuid ObjectBinding ) override;
 	virtual bool SupportsType( TSubclassOf<UMovieSceneTrack> Type ) const override;
@@ -60,16 +60,16 @@ private:
 
 	/** Animation sub menu */
 	TSharedRef<SWidget> BuildAnimationSubMenu(FGuid ObjectBinding, USkeleton* Skeleton, UMovieSceneTrack* Track);
-	void AddAnimationSubMenu(FMenuBuilder& MenuBuilder, FGuid ObjectBinding, USkeleton* Skeleton, UMovieSceneTrack* Track);
+	void AddAnimationSubMenu(FMenuBuilder& MenuBuilder, TArray<FGuid> ObjectBindings, USkeleton* Skeleton, UMovieSceneTrack* Track);
 
 	/** Animation sub menu filter function */
 	bool ShouldFilterAsset(const FAssetData& AssetData);
 
 	/** Animation asset selected */
-	void OnAnimationAssetSelected(const FAssetData& AssetData, FGuid ObjectBinding, UMovieSceneTrack* Track);
+	void OnAnimationAssetSelected(const FAssetData& AssetData, TArray<FGuid> ObjectBindings, UMovieSceneTrack* Track);
 
 	/** Animation asset enter pressed */
-	void OnAnimationAssetEnterPressed(const TArray<FAssetData>& AssetData, FGuid ObjectBinding, UMovieSceneTrack* Track);
+	void OnAnimationAssetEnterPressed(const TArray<FAssetData>& AssetData, TArray<FGuid> ObjectBindings, UMovieSceneTrack* Track);
 
 	/** Delegate for AnimatablePropertyChanged in AddKey */
 	FKeyPropertyResult AddKeyInternal(FFrameNumber KeyTime, UObject* Object, class UAnimSequenceBase* AnimSequence, UMovieSceneTrack* Track, int32 RowIndex);
diff --git a/Engine/Source/Editor/MovieSceneTools/Private/TrackEditors/SpawnTrackEditor.cpp b/Engine/Source/Editor/MovieSceneTools/Private/TrackEditors/SpawnTrackEditor.cpp
--- a/Engine/Source/Editor/MovieSceneTools/Private/TrackEditors/SpawnTrackEditor.cpp
+++ b/Engine/Source/Editor/MovieSceneTools/Private/TrackEditors/SpawnTrackEditor.cpp
@@ -33,11 +33,11 @@ UMovieSceneTrack* FSpawnTrackEditor::AddTrack(UMovieScene* FocusedMovieScene, co
 }
 
 
-void FSpawnTrackEditor::BuildObjectBindingTrackMenu(FMenuBuilder& MenuBuilder, const FGuid& ObjectBinding, const UClass* ObjectClass)
+void FSpawnTrackEditor::BuildObjectBindingTrackMenu(FMenuBuilder& MenuBuilder, const TArray<FGuid>& ObjectBindings, const UClass* ObjectClass)
 {
 	UMovieSceneSequence* MovieSequence = GetSequencer()->GetFocusedMovieSceneSequence();
 
-	if (!MovieSequence || MovieSequence->GetClass()->GetName() != TEXT("LevelSequence") || !MovieSequence->GetMovieScene()->FindSpawnable(ObjectBinding))
+	if (!MovieSequence || MovieSequence->GetClass()->GetName() != TEXT("LevelSequence") || !MovieSequence->GetMovieScene()->FindSpawnable(ObjectBindings[0]))
 	{
 		return;
 	}
@@ -47,8 +47,8 @@ void FSpawnTrackEditor::BuildObjectBindingTrackMenu(FMenuBuilder& MenuBuilder, c
 		LOCTEXT("AddSpawnTrackTooltip", "Adds a new track that controls the lifetime of the track's spawnable object."),
 		FSlateIcon(),
 		FUIAction(
-			FExecuteAction::CreateRaw(this, &FSpawnTrackEditor::HandleAddSpawnTrackMenuEntryExecute, ObjectBinding),
-			FCanExecuteAction::CreateSP(this, &FSpawnTrackEditor::CanAddSpawnTrack, ObjectBinding)
+			FExecuteAction::CreateRaw(this, &FSpawnTrackEditor::HandleAddSpawnTrackMenuEntryExecute, ObjectBindings),
+			FCanExecuteAction::CreateSP(this, &FSpawnTrackEditor::CanAddSpawnTrack, ObjectBindings[0])
 		)
 	);
 }
@@ -66,11 +66,15 @@ bool FSpawnTrackEditor::SupportsSequence(UMovieSceneSequence* InSequence) const
 }
 
 
-void FSpawnTrackEditor::HandleAddSpawnTrackMenuEntryExecute(FGuid ObjectBinding)
+void FSpawnTrackEditor::HandleAddSpawnTrackMenuEntryExecute(TArray<FGuid> ObjectBindings)
 {
 	FScopedTransaction AddSpawnTrackTransaction(LOCTEXT("AddSpawnTrack_Transaction", "Add Spawn Track"));
-	AddTrack(GetSequencer()->GetFocusedMovieSceneSequence()->GetMovieScene(), ObjectBinding, UMovieSceneSpawnTrack::StaticClass(), NAME_None);
-	GetSequencer()->NotifyMovieSceneDataChanged( EMovieSceneDataChangeType::MovieSceneStructureItemAdded );
+
+	for (FGuid ObjectBinding : ObjectBindings)
+	{
+		AddTrack(GetSequencer()->GetFocusedMovieSceneSequence()->GetMovieScene(), ObjectBinding, UMovieSceneSpawnTrack::StaticClass(), NAME_None);
+	}
+	GetSequencer()->NotifyMovieSceneDataChanged(EMovieSceneDataChangeType::MovieSceneStructureItemAdded);
 }
 
 
diff --git a/Engine/Source/Editor/MovieSceneTools/Private/TrackEditors/TransformTrackEditor.cpp b/Engine/Source/Editor/MovieSceneTools/Private/TrackEditors/TransformTrackEditor.cpp
--- a/Engine/Source/Editor/MovieSceneTools/Private/TrackEditors/TransformTrackEditor.cpp
+++ b/Engine/Source/Editor/MovieSceneTools/Private/TrackEditors/TransformTrackEditor.cpp
@@ -493,7 +493,7 @@ void F3DTransformTrackEditor::BuildObjectBindingEditButtons(TSharedPtr<SHorizont
 }
 
 
-void F3DTransformTrackEditor::BuildObjectBindingTrackMenu(FMenuBuilder& MenuBuilder, const FGuid& ObjectBinding, const UClass* ObjectClass)
+void F3DTransformTrackEditor::BuildObjectBindingTrackMenu(FMenuBuilder& MenuBuilder, const TArray<FGuid>& ObjectBindings, const UClass* ObjectClass)
 {
 	if (ObjectClass != nullptr && (ObjectClass->IsChildOf(AActor::StaticClass()) || ObjectClass->IsChildOf(USceneComponent::StaticClass())))
 	{
@@ -502,7 +502,7 @@ void F3DTransformTrackEditor::BuildObjectBindingTrackMenu(FMenuBuilder& MenuBuil
 			NSLOCTEXT("Sequencer", "AddTransformTooltip", "Adds a transform track."),
 			FSlateIcon(),
 			FUIAction( 
-				FExecuteAction::CreateSP( this, &F3DTransformTrackEditor::AddTransformKeysForHandle, ObjectBinding, EMovieSceneTransformChannel::All, ESequencerKeyMode::ManualKey )
+				FExecuteAction::CreateSP( this, &F3DTransformTrackEditor::AddTransformKeysForHandle, ObjectBindings, EMovieSceneTransformChannel::All, ESequencerKeyMode::ManualKey )
 			)
 		);
 	}
@@ -784,11 +784,16 @@ void F3DTransformTrackEditor::GetTransformKeys( const TOptional<FTransformData>&
 	}
 }
 
-void F3DTransformTrackEditor::AddTransformKeysForHandle( FGuid ObjectHandle, EMovieSceneTransformChannel ChannelToKey, ESequencerKeyMode KeyMode )
+void F3DTransformTrackEditor::AddTransformKeysForHandle(TArray<FGuid> ObjectHandles, EMovieSceneTransformChannel ChannelToKey, ESequencerKeyMode KeyMode)
 {
-	for ( TWeakObjectPtr<UObject> Object : GetSequencer()->FindObjectsInCurrentSequence(ObjectHandle) )
+	const FScopedTransaction Transaction(NSLOCTEXT("Sequencer", "AddTransformTrack", "Add Transform Track"));
+	
+	for (FGuid ObjectHandle : ObjectHandles)
 	{
-		AddTransformKeysForObject(Object.Get(), ChannelToKey, KeyMode);
+		for (TWeakObjectPtr<UObject> Object : GetSequencer()->FindObjectsInCurrentSequence(ObjectHandle))
+		{
+			AddTransformKeysForObject(Object.Get(), ChannelToKey, KeyMode);
+		}
 	}
 }
 
diff --git a/Engine/Source/Editor/MovieSceneTools/Private/TrackEditors/TransformTrackEditor.h b/Engine/Source/Editor/MovieSceneTools/Private/TrackEditors/TransformTrackEditor.h
--- a/Engine/Source/Editor/MovieSceneTools/Private/TrackEditors/TransformTrackEditor.h
+++ b/Engine/Source/Editor/MovieSceneTools/Private/TrackEditors/TransformTrackEditor.h
@@ -51,7 +51,7 @@ public:
 
 	virtual void BindCommands(TSharedRef<FUICommandList> SequencerCommandBindings) override;
 	virtual void BuildObjectBindingEditButtons(TSharedPtr<SHorizontalBox> EditBox, const FGuid& ObjectBinding, const UClass* ObjectClass) override;
-	virtual void BuildObjectBindingTrackMenu(FMenuBuilder& MenuBuilder, const FGuid& ObjectBinding, const UClass* ObjectClass) override;
+	virtual void BuildObjectBindingTrackMenu(FMenuBuilder& MenuBuilder, const TArray<FGuid>& ObjectBindinsg, const UClass* ObjectClass) override;
 	virtual TSharedRef<ISequencerSection> MakeSectionInterface( UMovieSceneSection& SectionObject, UMovieSceneTrack& Track, FGuid ObjectBinding ) override;
 	virtual void OnRelease() override;
 	virtual bool SupportsType( TSubclassOf<UMovieSceneTrack> Type ) const override;
@@ -131,7 +131,7 @@ private:
 	 * @param ChannelToKey The channels to add keys to.
 	 * @param KeyParams Parameters which control how the keys are added. 
 	 */
-	void AddTransformKeysForHandle( FGuid ObjectHandle, EMovieSceneTransformChannel ChannelToKey, ESequencerKeyMode KeyMode );
+	void AddTransformKeysForHandle( TArray<FGuid> ObjectHandles, EMovieSceneTransformChannel ChannelToKey, ESequencerKeyMode KeyMode );
 
 	/**
 	* Adds transform keys to a specific object.
diff --git a/Engine/Source/Editor/MovieSceneTools/Public/TrackEditors/SpawnTrackEditor.h b/Engine/Source/Editor/MovieSceneTools/Public/TrackEditors/SpawnTrackEditor.h
--- a/Engine/Source/Editor/MovieSceneTools/Public/TrackEditors/SpawnTrackEditor.h
+++ b/Engine/Source/Editor/MovieSceneTools/Public/TrackEditors/SpawnTrackEditor.h
@@ -43,7 +43,7 @@ public:
 	// ISequencerTrackEditor interface
 
 	virtual UMovieSceneTrack* AddTrack(UMovieScene* FocusedMovieScene, const FGuid& ObjectHandle, TSubclassOf<class UMovieSceneTrack> TrackClass, FName UniqueTypeName) override;
-	virtual void BuildObjectBindingTrackMenu(FMenuBuilder& MenuBuilder, const FGuid& ObjectBinding, const UClass* ObjectClass) override;
+	virtual void BuildObjectBindingTrackMenu(FMenuBuilder& MenuBuilder, const TArray<FGuid>& ObjectBindings, const UClass* ObjectClass) override;
 	virtual TSharedPtr<SWidget> BuildOutlinerEditWidget(const FGuid& ObjectBinding, UMovieSceneTrack* Track, const FBuildEditWidgetParams& Params) override  { return TSharedPtr<SWidget>(); }
 	virtual bool HandleAssetAdded(UObject* Asset, const FGuid& TargetObjectGuid) override { return false; }
 	virtual bool SupportsType(TSubclassOf<UMovieSceneTrack> Type) const override;
@@ -52,6 +52,6 @@ public:
 private:
 
 	/** Callback for executing the "Add Spawn Track" menu entry. */
-	void HandleAddSpawnTrackMenuEntryExecute(FGuid ObjectBinding);
+	void HandleAddSpawnTrackMenuEntryExecute(TArray<FGuid> ObjectBindings);
 	bool CanAddSpawnTrack(FGuid ObjectBinding) const;
 };
diff --git a/Engine/Source/Editor/Sequencer/Private/DisplayNodes/SequencerObjectBindingNode.cpp b/Engine/Source/Editor/Sequencer/Private/DisplayNodes/SequencerObjectBindingNode.cpp
--- a/Engine/Source/Editor/Sequencer/Private/DisplayNodes/SequencerObjectBindingNode.cpp
+++ b/Engine/Source/Editor/Sequencer/Private/DisplayNodes/SequencerObjectBindingNode.cpp
@@ -885,6 +885,30 @@ TSharedRef<SWidget> FSequencerObjectBindingNode::HandleAddTrackComboButtonGetMen
 	UObject* BoundObject = GetSequencer().FindSpawnedObjectOrTemplate(ObjectBinding);
 	const UClass* ObjectClass = GetClassForObjectBinding();
 
+	TArray<FGuid> ObjectBindings;
+	ObjectBindings.Add(ObjectBinding);
+
+	// Add additionally selected object bindings
+	for (const TSharedRef<FSequencerDisplayNode>& Node : Sequencer.GetSelection().GetSelectedOutlinerNodes())
+	{
+		if (Node->GetType() != ESequencerNode::Object)
+		{
+			continue;
+		}
+
+		auto ObjectBindingNode = StaticCastSharedRef<FSequencerObjectBindingNode>(Node);
+
+		FGuid Guid = ObjectBindingNode->GetObjectBinding();
+		for (auto RuntimeObject : Sequencer.FindBoundObjects(Guid, Sequencer.GetFocusedTemplateID()))
+		{
+			if (RuntimeObject->GetClass() == ObjectClass)
+			{
+				ObjectBindings.AddUnique(Guid);
+				continue;
+			}
+		}
+	}
+
 	ISequencerModule& SequencerModule = FModuleManager::GetModuleChecked<ISequencerModule>( "Sequencer" );
 	TSharedRef<FUICommandList> CommandList(new FUICommandList);
 
@@ -892,7 +916,7 @@ TSharedRef<SWidget> FSequencerObjectBindingNode::HandleAddTrackComboButtonGetMen
 
 	for (const TSharedPtr<ISequencerTrackEditor>& TrackEditor : GetSequencer().GetTrackEditors())
 	{
-		TrackEditor->ExtendObjectBindingTrackMenu(Extender, ObjectBinding, ObjectClass);
+		TrackEditor->ExtendObjectBindingTrackMenu(Extender, ObjectBindings, ObjectClass);
 	}
 
 	FMenuBuilder AddTrackMenuBuilder(true, nullptr, Extender);
@@ -900,7 +924,7 @@ TSharedRef<SWidget> FSequencerObjectBindingNode::HandleAddTrackComboButtonGetMen
 	const int32 NumStartingBlocks = AddTrackMenuBuilder.GetMultiBox()->GetBlocks().Num();
 
 	AddTrackMenuBuilder.BeginSection("Tracks", LOCTEXT("TracksMenuHeader" , "Tracks"));
-	GetSequencer().BuildObjectBindingTrackMenu(AddTrackMenuBuilder, ObjectBinding, ObjectClass);
+	GetSequencer().BuildObjectBindingTrackMenu(AddTrackMenuBuilder, ObjectBindings, ObjectClass);
 	AddTrackMenuBuilder.EndSection();
 
 	TArray<FPropertyPath> KeyablePropertyPaths;
diff --git a/Engine/Source/Editor/Sequencer/Private/MovieSceneTrackEditor.cpp b/Engine/Source/Editor/Sequencer/Private/MovieSceneTrackEditor.cpp
--- a/Engine/Source/Editor/Sequencer/Private/MovieSceneTrackEditor.cpp
+++ b/Engine/Source/Editor/Sequencer/Private/MovieSceneTrackEditor.cpp
@@ -162,11 +162,11 @@ void FMovieSceneTrackEditor::BuildObjectBindingEditButtons(TSharedPtr<SHorizonta
 { 
 }
 
-void FMovieSceneTrackEditor::BuildObjectBindingTrackMenu(FMenuBuilder& MenuBuilder, const FGuid& ObjectBinding, const UClass* ObjectClass) 
-{ 
+void FMovieSceneTrackEditor::BuildObjectBindingTrackMenu(FMenuBuilder& MenuBuilder, const TArray<FGuid>& ObjectBindings, const UClass* ObjectClass)
+{
 }
 
-TSharedPtr<SWidget> FMovieSceneTrackEditor::BuildOutlinerEditWidget(const FGuid& ObjectBinding, UMovieSceneTrack* Track, const FBuildEditWidgetParams& Params) 
+TSharedPtr<SWidget> FMovieSceneTrackEditor::BuildOutlinerEditWidget(const FGuid& ObjectBinding, UMovieSceneTrack* Track, const FBuildEditWidgetParams& Params)
 {
 	if (Track->GetSupportedBlendTypes().Num() > 0)
 	{
diff --git a/Engine/Source/Editor/Sequencer/Private/Sequencer.cpp b/Engine/Source/Editor/Sequencer/Private/Sequencer.cpp
--- a/Engine/Source/Editor/Sequencer/Private/Sequencer.cpp
+++ b/Engine/Source/Editor/Sequencer/Private/Sequencer.cpp
@@ -9881,11 +9881,11 @@ void FSequencer::BuildAddObjectBindingsMenu(class FMenuBuilder& MenuBuilder)
 	}
 }
 
-void FSequencer::BuildObjectBindingTrackMenu(FMenuBuilder& MenuBuilder, const FGuid& ObjectBinding, const UClass* ObjectClass)
+void FSequencer::BuildObjectBindingTrackMenu(FMenuBuilder& MenuBuilder, const TArray<FGuid>& InObjectBindings, const UClass* ObjectClass)
 {
 	for (int32 i = 0; i < TrackEditors.Num(); ++i)
 	{
-		TrackEditors[i]->BuildObjectBindingTrackMenu(MenuBuilder, ObjectBinding, ObjectClass);
+		TrackEditors[i]->BuildObjectBindingTrackMenu(MenuBuilder, InObjectBindings, ObjectClass);
 	}
 }
 
diff --git a/Engine/Source/Editor/Sequencer/Private/Sequencer.h b/Engine/Source/Editor/Sequencer/Private/Sequencer.h
--- a/Engine/Source/Editor/Sequencer/Private/Sequencer.h
+++ b/Engine/Source/Editor/Sequencer/Private/Sequencer.h
@@ -394,10 +394,10 @@ public:
 	 * Builds up the track menu for object binding nodes in the outliner
 	 * 
 	 * @param MenuBuilder	The track menu builder to add things to
-	 * @param ObjectBinding	The object binding of the selected node
+	 * @param ObjectBindings The array of object bindings to add tracks to (if there are more than 1 selected)
 	 * @param ObjectClass	The class of the selected object
 	 */
-	void BuildObjectBindingTrackMenu(FMenuBuilder& MenuBuilder, const FGuid& ObjectBinding, const UClass* ObjectClass);
+	void BuildObjectBindingTrackMenu(FMenuBuilder& MenuBuilder, const TArray<FGuid>& ObjectBindings, const UClass* ObjectClass);
 
 	/**
 	 * Builds up the edit buttons for object binding nodes in the outliner
diff --git a/Engine/Source/Editor/Sequencer/Public/ISequencerTrackEditor.h b/Engine/Source/Editor/Sequencer/Public/ISequencerTrackEditor.h
--- a/Engine/Source/Editor/Sequencer/Public/ISequencerTrackEditor.h
+++ b/Engine/Source/Editor/Sequencer/Public/ISequencerTrackEditor.h
@@ -81,7 +81,7 @@ public:
 	 * @param ObjectBinding The object binding this is for.
 	 * @param ObjectClass The class of the object this is for.
 	 */
-	virtual void BuildObjectBindingTrackMenu(FMenuBuilder& MenuBuilder, const FGuid& ObjectBinding, const UClass* ObjectClass) = 0;
+	virtual void BuildObjectBindingTrackMenu(FMenuBuilder& MenuBuilder, const TArray<FGuid>& ObjectBindings, const UClass* ObjectClass) = 0;
 
 	/**
 	 * Extend the object binding track menu for the specified binding and class
@@ -90,7 +90,7 @@ public:
 	 * @param ObjectBinding The object binding this is for.
 	 * @param ObjectClass The class of the object this is for.
 	 */
-	virtual void ExtendObjectBindingTrackMenu(TSharedRef<FExtender> Extender, const FGuid& ObjectBinding, const UClass* ObjectClass) { }
+	virtual void ExtendObjectBindingTrackMenu(TSharedRef<FExtender> Extender, const TArray<FGuid>& ObjectBindings, const UClass* ObjectClass) { }
 
 	/**
 	 * Builds an edit widget for the outliner nodes which represent tracks which are edited by this editor.
diff --git a/Engine/Source/Editor/Sequencer/Public/MovieSceneTrackEditor.h b/Engine/Source/Editor/Sequencer/Public/MovieSceneTrackEditor.h
--- a/Engine/Source/Editor/Sequencer/Public/MovieSceneTrackEditor.h
+++ b/Engine/Source/Editor/Sequencer/Public/MovieSceneTrackEditor.h
@@ -161,7 +161,7 @@ public:
 	virtual void BindCommands(TSharedRef<FUICommandList> SequencerCommandBindings) override;
 	virtual void BuildAddTrackMenu(FMenuBuilder& MenuBuilder) override;
 	virtual void BuildObjectBindingEditButtons(TSharedPtr<SHorizontalBox> EditBox, const FGuid& ObjectBinding, const UClass* ObjectClass) override;
-	virtual void BuildObjectBindingTrackMenu(FMenuBuilder& MenuBuilder, const FGuid& ObjectBinding, const UClass* ObjectClass) override;
+	virtual void BuildObjectBindingTrackMenu(FMenuBuilder& MenuBuilder, const TArray<FGuid>& ObjectBindings, const UClass* ObjectClass) override;
 	virtual TSharedPtr<SWidget> BuildOutlinerEditWidget(const FGuid& ObjectBinding, UMovieSceneTrack* Track, const FBuildEditWidgetParams& Params) override;
 	virtual void BuildTrackContextMenu( FMenuBuilder& MenuBuilder, UMovieSceneTrack* Track ) override;
 	virtual bool HandleAssetAdded(UObject* Asset, const FGuid& TargetObjectGuid) override;
