Sequencer: More CIS fixes

diff --git a/Engine/Plugins/FX/Niagara/Source/NiagaraEditor/Private/Sequencer/NiagaraSequence/NiagaraSequencerFilters.cpp b/Engine/Plugins/FX/Niagara/Source/NiagaraEditor/Private/Sequencer/NiagaraSequence/NiagaraSequencerFilters.cpp
--- a/Engine/Plugins/FX/Niagara/Source/NiagaraEditor/Private/Sequencer/NiagaraSequence/NiagaraSequencerFilters.cpp
+++ b/Engine/Plugins/FX/Niagara/Source/NiagaraEditor/Private/Sequencer/NiagaraSequence/NiagaraSequencerFilters.cpp
@@ -4,6 +4,7 @@
 #include "MovieScene/MovieSceneNiagaraTrack.h"
 #include "NiagaraActor.h"
 #include "NiagaraComponent.h"
+#include "EditorStyleSet.h"
 
 #define LOCTEXT_NAMESPACE "NiagaraSequencerTrackFilters"
 
diff --git a/Engine/Source/Editor/Sequencer/Private/SequencerTrackFilters.h b/Engine/Source/Editor/Sequencer/Private/SequencerTrackFilters.h
--- a/Engine/Source/Editor/Sequencer/Private/SequencerTrackFilters.h
+++ b/Engine/Source/Editor/Sequencer/Private/SequencerTrackFilters.h
@@ -7,6 +7,7 @@
 #include "Misc/FilterCollection.h"
 #include "CollectionManagerTypes.h"
 #include "Styling/SlateIconFinder.h"
+#include "EditorStyleSet.h"
 #include "SequencerDisplayNode.h"
 #include "MovieSceneSequence.h"
 
