Fix for GPU emitters flickering when simulated > 1 time per frame.

diff --git a/Engine/Plugins/FX/Niagara/Source/Niagara/Private/NiagaraEmitterInstanceBatcher.cpp b/Engine/Plugins/FX/Niagara/Source/Niagara/Private/NiagaraEmitterInstanceBatcher.cpp
--- a/Engine/Plugins/FX/Niagara/Source/Niagara/Private/NiagaraEmitterInstanceBatcher.cpp
+++ b/Engine/Plugins/FX/Niagara/Source/Niagara/Private/NiagaraEmitterInstanceBatcher.cpp
@@ -212,7 +212,10 @@ void NiagaraEmitterInstanceBatcher::ResizeBuffersAndGatherResources(FOverlappabl
 			}
 
 			Context->MainDataSet->EndSimulate();
-			Context->SetDataToRender(Instance.DestinationData);
+			if (Instance->bIsFinalTick)
+			{
+				Context->SetDataToRender(Instance.DestinationData);
+			}
 		}
 	}
 }
@@ -368,8 +371,8 @@ void NiagaraEmitterInstanceBatcher::ExecuteAll(FRHICommandList &RHICmdList, FUni
 			// Those ticks should still be executed in order.
 			for (FNiagaraGPUSystemTick& Tick : Ticks_RT)
 			{
-				FNiagaraComputeInstanceData* data = Tick.GetInstanceData();
-				ContextToTicks.FindOrAdd(data->Context).Add(&Tick);
+				FNiagaraComputeInstanceData* Data = Tick.GetInstanceData();
+				ContextToTicks.FindOrAdd(Data->Context).Add(&Tick);
 			}
 
 			// Count the maximum number of tick per context to know the number of passes we need to do
@@ -377,6 +380,8 @@ void NiagaraEmitterInstanceBatcher::ExecuteAll(FRHICommandList &RHICmdList, FUni
 			for (auto& Ticks : ContextToTicks)
 			{
 				NumSimPass = FMath::Max(NumSimPass, (uint32)Ticks.Value.Num());
+				FNiagaraComputeInstanceData* Data = Ticks.Value.Last()->GetInstanceData();
+				Data->bIsFinalTick = true;
 			}
 
 			// Transpose now only once the data to get all independent tick per pass
