Missing file

diff --git a/Engine/Plugins/FX/Niagara/Source/Niagara/Classes/NiagaraScriptExecutionContext.h b/Engine/Plugins/FX/Niagara/Source/Niagara/Classes/NiagaraScriptExecutionContext.h
--- a/Engine/Plugins/FX/Niagara/Source/Niagara/Classes/NiagaraScriptExecutionContext.h
+++ b/Engine/Plugins/FX/Niagara/Source/Niagara/Classes/NiagaraScriptExecutionContext.h
@@ -233,6 +233,8 @@ struct FNiagaraComputeInstanceData
 	//Buffer into which we'll write the new simulation state. Initialized at the start of processing this tick on the RT.
 	FNiagaraDataBuffer* DestinationData;
 
+	bool bIsFinalTick;
+
 	FNiagaraComputeInstanceData()
 		: EventSpawnTotal(0)
 		, SpawnRateInstances(0)
@@ -240,6 +242,7 @@ struct FNiagaraComputeInstanceData
 		, Context(nullptr)
 		, CurrentData(nullptr)
 		, DestinationData(nullptr)
+		, bIsFinalTick(false)
 	{}
 };
 
