- Niagara fix missing final tick on emitters other than the first one

diff --git a/Engine/Plugins/FX/Niagara/Source/Niagara/Classes/NiagaraScriptExecutionContext.h b/Engine/Plugins/FX/Niagara/Source/Niagara/Classes/NiagaraScriptExecutionContext.h
--- a/Engine/Plugins/FX/Niagara/Source/Niagara/Classes/NiagaraScriptExecutionContext.h
+++ b/Engine/Plugins/FX/Niagara/Source/Niagara/Classes/NiagaraScriptExecutionContext.h
@@ -233,8 +233,6 @@ struct FNiagaraComputeInstanceData
 	//Buffer into which we'll write the new simulation state. Initialized at the start of processing this tick on the RT.
 	FNiagaraDataBuffer* DestinationData;
 
-	bool bIsFinalTick;
-
 	FNiagaraComputeInstanceData()
 		: EventSpawnTotal(0)
 		, SpawnRateInstances(0)
@@ -242,7 +240,6 @@ struct FNiagaraComputeInstanceData
 		, Context(nullptr)
 		, CurrentData(nullptr)
 		, DestinationData(nullptr)
-		, bIsFinalTick(false)
 	{}
 };
 
@@ -282,4 +279,5 @@ public:
 	uint8* InstanceData_ParamData_Packed;
 	bool bRequiredDistanceFieldData = false;
 	bool bNeedsReset = false;
+	bool bIsFinalTick = false;
 };
diff --git a/Engine/Plugins/FX/Niagara/Source/Niagara/Private/NiagaraEmitterInstanceBatcher.cpp b/Engine/Plugins/FX/Niagara/Source/Niagara/Private/NiagaraEmitterInstanceBatcher.cpp
--- a/Engine/Plugins/FX/Niagara/Source/Niagara/Private/NiagaraEmitterInstanceBatcher.cpp
+++ b/Engine/Plugins/FX/Niagara/Source/Niagara/Private/NiagaraEmitterInstanceBatcher.cpp
@@ -157,7 +157,9 @@ void NiagaraEmitterInstanceBatcher::ResizeBuffersAndGatherResources(FOverlappabl
 
 	for (FNiagaraGPUSystemTick* Tick : OverlappableTick)
 	{
-		uint32 DispatchCount = Tick->Count;
+		const uint32 DispatchCount = Tick->Count;
+		const bool bIsFinalTick = Tick->bIsFinalTick;
+
 		FNiagaraComputeInstanceData* Instances = Tick->GetInstanceData();
 		for (uint32 Index = 0; Index < DispatchCount; Index++)
 		{
@@ -212,7 +214,7 @@ void NiagaraEmitterInstanceBatcher::ResizeBuffersAndGatherResources(FOverlappabl
 			}
 
 			Context->MainDataSet->EndSimulate();
-			if (Instance.bIsFinalTick)
+			if (bIsFinalTick)
 			{
 				Context->SetDataToRender(Instance.DestinationData);
 			}
@@ -380,8 +382,7 @@ void NiagaraEmitterInstanceBatcher::ExecuteAll(FRHICommandList &RHICmdList, FUni
 			for (auto& Ticks : ContextToTicks)
 			{
 				NumSimPass = FMath::Max(NumSimPass, (uint32)Ticks.Value.Num());
-				FNiagaraComputeInstanceData* Data = Ticks.Value.Last()->GetInstanceData();
-				Data->bIsFinalTick = true;
+				Ticks.Value.Last()->bIsFinalTick = true;
 			}
 
 			// Transpose now only once the data to get all independent tick per pass
