Niagara GPU Readback Refactor

diff --git a/Engine/Plugins/FX/Niagara/Shaders/Private/NiagaraDrawIndirectArgsGen.usf b/Engine/Plugins/FX/Niagara/Shaders/Private/NiagaraDrawIndirectArgsGen.usf
new file mode 100644
--- /dev/null
+++ b/Engine/Plugins/FX/Niagara/Shaders/Private/NiagaraDrawIndirectArgsGen.usf
@@ -0,0 +1,35 @@
+// Copyright 1998-2019 Epic Games, Inc. All Rights Reserved.
+
+/*==============================================================================
+	NiagaraDrawIndirectArgsGen.usf: Shader to copy sorted index buffers.
+==============================================================================*/
+
+// Array of task info. Each being a draw indirect arg gen or instance count reset.
+// ArgGen = (InstanceCountBufferOffset, NumIndicesPerInstance)
+// InstanceCountClear = (InstanceCountBufferOffset)
+Buffer<uint> TaskInfos;
+
+RWBuffer<uint> RWInstanceCounts;
+RWBuffer<uint> RWDrawIndirectArgs;
+
+uint4 TaskCount; // (NumArgGenTasks, NumInstanceCountClearTasks, NumArgGenTasks + NumInstanceCountClearTasks, 0)
+
+[numthreads(THREAD_COUNT,1,1)]
+void MainCS(uint TaskIndex : SV_DispatchThreadID)
+{
+	if (TaskIndex < TaskCount.x)
+	{
+		const uint ArgOffset = NIAGARA_DRAW_INDIRECT_ARGS_SIZE * TaskIndex;
+		const uint InfoOffset = TaskIndex * 2;
+		RWDrawIndirectArgs[ArgOffset + 0] = TaskInfos[InfoOffset + 1]; // NumIndicesPerInstance
+		RWDrawIndirectArgs[ArgOffset + 1] = RWInstanceCounts[TaskInfos[InfoOffset + 0]]; // InstanceCount
+		RWDrawIndirectArgs[ArgOffset + 2] = 0; // StartIndexLocation
+		RWDrawIndirectArgs[ArgOffset + 3] = 0; // BaseVertexLocation
+		RWDrawIndirectArgs[ArgOffset + 4] = 0; // StartInstanceLocation
+	}
+	else if (TaskIndex < TaskCount.z)
+	{
+		const uint InfoOffset = TaskCount.x * 2 + (TaskIndex - TaskCount.x);
+		RWInstanceCounts[TaskInfos[InfoOffset + 0]] = 0;
+	}
+}
\ No newline at end of file
diff --git a/Engine/Plugins/FX/Niagara/Shaders/Private/NiagaraEmitterInstanceShader.usf b/Engine/Plugins/FX/Niagara/Shaders/Private/NiagaraEmitterInstanceShader.usf
--- a/Engine/Plugins/FX/Niagara/Shaders/Private/NiagaraEmitterInstanceShader.usf
+++ b/Engine/Plugins/FX/Niagara/Shaders/Private/NiagaraEmitterInstanceShader.usf
@@ -298,7 +298,6 @@ int rand(int x, int Seed1, int Seed2, int Seed3)
 	 */
 	uint UpdateStartInstance;
 	uint SpawnedInstances;
-	uint NumIndicesPerInstance;
 
 	uint ComponentBufferSizeRead;
 	uint ComponentBufferSizeWrite;
@@ -313,16 +312,15 @@ int rand(int x, int Seed1, int Seed2, int Seed3)
 	Buffer<int> InputInt;
 	RWBuffer<int> RWOutputInt;
 	RWBuffer<float> RWOutputFloat;
-	RWBuffer<uint> RWDataSetIndices;
-	Buffer<uint> ReadDataSetIndices;
 
+	RWBuffer<uint> RWInstanceCounts;
+	uint ReadInstanceCountOffset;
+	uint WriteInstanceCountOffset;
 
 	//TODO: IDs for GPU
 //	void AcquireID(int DataSetID, out int IDIndex, out int IDAcquireTag);
 // 	void UpdateID(int DataSetID, int IDIndex, int InstanceIndex);
 
-
-
 	/* Returns the current instance index relative to the operation (spawn/update)
 	 */
 	int ExecIndex()
@@ -341,14 +339,20 @@ int rand(int x, int Seed1, int Seed2, int Seed3)
 	/* Acquire an output index - the default index is the scratch instance; one additional instance is allocated 
 	 *	at the end of the buffer, so no branching on -1 is necessary during OutputData operations
 	 */
-	int AcquireIndex(int DataSetID, bool bDoAcquire)
+	int AcquireIndex(uniform int DataSetID, bool bDoAcquire)
 	{
+		// Begin static assert : GPU particles only support DataSetID 0
+		int MustBe0[1];
+		MustBe0[DataSetID] = 0;
+		// End static assert
+
 		int PrevIdx = GSpawnStartInstance + SpawnedInstances;	// scratch instance as default; write to that for dead particles
-		if(bDoAcquire == true)
+		if (bDoAcquire)
 		{
 			// Have to use uint's here to avoid PS4 compiler warnings about InterlockedAdd, cannot propagate uint due to CPU VM limitations...
 			uint RetPrevIdx;
-			InterlockedAdd(RWDataSetIndices[DataSetID*5+1], (uint)1U, RetPrevIdx);
+			// @TODO : add some TLS logic to avoid thread group for doing atomic for each thread. (gathering the actual required count)
+			InterlockedAdd(RWInstanceCounts[WriteInstanceCountOffset], (uint)1U, RetPrevIdx);
 			PrevIdx = (int)RetPrevIdx;
 		}
 		return PrevIdx;
@@ -425,7 +429,7 @@ void SimulateMainComputeCS(
 	*/
 
 	uint InstanceID = UpdateStartInstance + DispatchThreadId.x;
-	GSpawnStartInstance = ReadDataSetIndices[1];				// needed by ExecIndex()
+	GSpawnStartInstance = RWInstanceCounts[ReadInstanceCountOffset];				// needed by ExecIndex()
 	int MaxInstances = GSpawnStartInstance + SpawnedInstances;
 
 	if(InstanceID < UpdateStartInstance + MaxInstances )
@@ -461,15 +465,6 @@ void SimulateMainComputeCS(
 			SimulateMain(InstanceID, EventIndex, 0);		// SimulateMain spawn phase
 		}
 	}
-
-	// fill in the rest of RWDataSetIndices for the DrawIndirect call for the main data set; number of instances is incremented in AcquireIndex
-	if(InstanceID == 0)
-	{
-		RWDataSetIndices[0] = NumIndicesPerInstance;		// IndexCountPerInstance; passed in from the renderer - TODO: support multiple renderers with multiple indirect draw calls
-		RWDataSetIndices[2] = 0;		// StartIndexLocation
-		RWDataSetIndices[3] = 0;		// BaseVertexLocation
-		RWDataSetIndices[4] = 0;		// StartInstanceLocation
-	}
 }
 
 #endif//GPU_SIMULATION
diff --git a/Engine/Plugins/FX/Niagara/Shaders/Private/NiagaraSpriteVertexFactory.ush b/Engine/Plugins/FX/Niagara/Shaders/Private/NiagaraSpriteVertexFactory.ush
--- a/Engine/Plugins/FX/Niagara/Shaders/Private/NiagaraSpriteVertexFactory.ush
+++ b/Engine/Plugins/FX/Niagara/Shaders/Private/NiagaraSpriteVertexFactory.ush
@@ -60,7 +60,7 @@ FVertexFactoryInput LoadVertexFactoryInput(uint TriangleIndex, int VertexIndex)
 }
 uint GetNumRayTracingDynamicMeshVerticesIndirect()
 {
-	return 6 * NiagaraSpriteVFLooseParameters.IndirectArgsBuffer[1];
+	return 6 * NiagaraSpriteVFLooseParameters.IndirectArgsBuffer[NiagaraSpriteVFLooseParameters.IndirectArgsOffset + 1];
 }
 #endif
 
diff --git a/Engine/Plugins/FX/Niagara/Source/Niagara/Classes/NiagaraDataSet.h b/Engine/Plugins/FX/Niagara/Source/Niagara/Classes/NiagaraDataSet.h
--- a/Engine/Plugins/FX/Niagara/Source/Niagara/Classes/NiagaraDataSet.h
+++ b/Engine/Plugins/FX/Niagara/Source/Niagara/Classes/NiagaraDataSet.h
@@ -26,6 +26,8 @@ struct FNiagaraVariableLayoutInfo
 
 class FNiagaraDataSet;
 class FNiagaraShader;
+class FNiagaraGPUInstanceCountManager;
+struct FNiagaraComputeExecutionContext;
 
 //Base class for objects in Niagara that are owned by one object but are then passed for reading to other objects, potentially on other threads.
 //This class allows us to know if the object is being used so we do not overwrite it and to ensure it's lifetime so we do not access freed data.
@@ -94,7 +96,7 @@ protected:
 public:
 	FNiagaraDataBuffer(FNiagaraDataSet* InOwner);
 	void Allocate(uint32 NumInstances, bool bMaintainExisting = false);
-	void AllocateGPU(uint32 InNumInstances, FRHICommandList &RHICmdList);
+	void AllocateGPU(uint32 InNumInstances, FNiagaraGPUInstanceCountManager& GPUInstanceCountManager, FRHICommandList &RHICmdList);
 	void SwapInstances(uint32 OldIndex, uint32 NewIndex);
 	void KillInstance(uint32 InstanceIdx);
 	void CopyTo(FNiagaraDataBuffer& DestBuffer, int32 StartIdx, int32 NumInstances)const;
@@ -129,7 +131,8 @@ public:
 	FORCEINLINE uint32 GetSizeBytes()const { return FloatData.Num() + Int32Data.Num(); }
 	FORCEINLINE FRWBuffer& GetGPUBufferFloat() { return GPUBufferFloat; }
 	FORCEINLINE FRWBuffer& GetGPUBufferInt() { return GPUBufferInt;	}
-	FORCEINLINE FRWBuffer& GetGPUIndices() { return GPUIndices; }
+	FORCEINLINE uint32 GetGPUInstanceCountBufferOffset() const { return GPUInstanceCountBufferOffset; }
+	FORCEINLINE void ClearGPUInstanceCountBufferOffset() { GPUInstanceCountBufferOffset = INDEX_NONE; }
 
 	FORCEINLINE int32 GetSafeComponentBufferSize() const { return GetSafeComponentBufferSize(GetNumInstancesAllocated()); }
 	FORCEINLINE uint32 GetFloatStride() const { return FloatStride; }
@@ -146,6 +149,9 @@ public:
 	template<bool bDoResourceTransitions>
 	void SetShaderParams(class FNiagaraShader *Shader, FRHICommandList &CommandList, bool bInput);
 	void UnsetShaderParams(class FNiagaraShader *Shader, FRHICommandList &CommandList);
+
+	void ReleaseGPUInstanceCount(FNiagaraGPUInstanceCountManager& GPUInstanceCountManager);
+
 private:
 
 	void CheckUsage(bool bReadOnly)const;
@@ -173,13 +179,14 @@ private:
 
 	//////////////////////////////////////////////////////////////////////////
 	// GPU Data
+	/** The buffer offset where the instance count is accumulated. */
+	uint32 GPUInstanceCountBufferOffset;
+	/** The num of allocated chunks, each being of size ALLOC_CHUNKSIZE */
 	uint32 NumChunksAllocatedForGPU;
 	/** GPU Buffer containing floating point values for GPU simulations. */
 	FRWBuffer GPUBufferFloat;
 	/** GPU Buffer containing floating point values for GPU simulations. */
 	FRWBuffer GPUBufferInt;
-	/** GPU Buffer containing args to allow indirect rendering of this buffer. */
-	FRWBuffer GPUIndices;
 	//////////////////////////////////////////////////////////////////////////
 
 	/** Number of instances in data. */
@@ -277,6 +284,9 @@ public:
 		return *DestinationData;
 	}
 
+	/** Release the GPU instance counts so that they can be reused */
+	void ReleaseGPUInstanceCounts(FNiagaraGPUInstanceCountManager& GPUInstanceCountManager);
+
 private:
 
 	void Reset();
@@ -1451,14 +1461,14 @@ class NIAGARA_API FScopedNiagaraDataSetGPUReadback
 {
 public:
 	FScopedNiagaraDataSetGPUReadback() {}
-	FScopedNiagaraDataSetGPUReadback(FNiagaraDataSet* DataSet) { ReadbackData(DataSet); }
 	~FScopedNiagaraDataSetGPUReadback();
 
-	void ReadbackData(FNiagaraDataSet* InDataSet);
+	void ReadbackData(class NiagaraEmitterInstanceBatcher* Batcher, FNiagaraDataSet* InDataSet);
 	uint32 GetNumInstances() const { check(DataSet != nullptr); return NumInstances; }
 
 private:
 	FNiagaraDataSet*	DataSet = nullptr;
 	FNiagaraDataBuffer* DataBuffer = nullptr;
+	NiagaraEmitterInstanceBatcher* Batcher = nullptr;
 	uint32				NumInstances = 0;
 };
diff --git a/Engine/Plugins/FX/Niagara/Source/Niagara/Classes/NiagaraEmitterInstanceBatcher.h b/Engine/Plugins/FX/Niagara/Source/Niagara/Classes/NiagaraEmitterInstanceBatcher.h
--- a/Engine/Plugins/FX/Niagara/Source/Niagara/Classes/NiagaraEmitterInstanceBatcher.h
+++ b/Engine/Plugins/FX/Niagara/Source/Niagara/Classes/NiagaraEmitterInstanceBatcher.h
@@ -21,6 +21,7 @@ the same VectorVM byte code / compute shader code
 #include "Runtime/Engine/Private/Particles/ParticleSortingGPU.h"
 #include "NiagaraGPUSortInfo.h"
 #include "NiagaraScriptExecutionContext.h"
+#include "NiagaraGPUInstanceCountManager.h"
 
 class FNiagaraIndicesVertexBuffer : public FParticleIndicesVertexBuffer
 {
@@ -77,7 +78,7 @@ public:
 	virtual void AddVectorField(UVectorFieldComponent* VectorFieldComponent) override {}
 	virtual void RemoveVectorField(UVectorFieldComponent* VectorFieldComponent) override {}
 	virtual void UpdateVectorField(UVectorFieldComponent* VectorFieldComponent) override {}
-	virtual void PreInitViews() override;
+	virtual void PreInitViews(FRHICommandListImmediate& RHICmdList) override;
 	virtual bool UsesGlobalDistanceField() const override;
 	virtual void PreRender(FRHICommandListImmediate& RHICmdList, const class FGlobalDistanceFieldParameterData* GlobalDistanceFieldParameterData) override;
 
@@ -102,12 +103,6 @@ public:
 		const class FShaderParametersMetadata* SceneTexturesUniformBufferStruct,
 		FUniformBufferRHIParamRef SceneTexturesUniformBuffer) override;
 
-	void ExecuteAll(FRHICommandList &RHICmdList, FUniformBufferRHIParamRef ViewUniformBuffer);
-	void TickSingle(const FNiagaraGPUSystemTick& Tick, FNiagaraComputeInstanceData *Instance, FRHICommandList &RHICmdList, FUniformBufferRHIParamRef ViewUniformBuffer) const;
-	void ResizeBuffersAndGatherResources(FOverlappableTicks& OverlappableTick, FRHICommandList& RHICmdList, FNiagaraBufferArray& DestDataBuffers, FNiagaraBufferArray& CurrDataBuffers, FNiagaraBufferArray& DestBufferIntFloat, FNiagaraBufferArray& CurrBufferIntFloat);
-
-	void DispatchAllOnCompute(FOverlappableTicks& OverlappableTick, FRHICommandList& RHICmdList, FUniformBufferRHIParamRef ViewUniformBuffer, FNiagaraBufferArray& DestDataBuffers, FNiagaraBufferArray& CurrDataBuffers, FNiagaraBufferArray& DestBufferIntFloat, FNiagaraBufferArray& CurrBufferIntFloat);
-
 	int32 AddSortedGPUSimulation(const FNiagaraGPUSortInfo& SortInfo);
 	void SortGPUParticles(FRHICommandListImmediate& RHICmdList);
 	void ResolveParticleSortBuffers(FRHICommandListImmediate& RHICmdList, int32 ResultBufferIndex);
@@ -131,11 +126,17 @@ public:
 				bool bCopyBeforeStart = false
 			) const;
 
-	void ResolveDatasetWrites(FRHICommandList &RHICmdList, FNiagaraComputeInstanceData *Context) const;
 	void ResizeCurrentBuffer(FRHICommandList &RHICmdList, FNiagaraComputeExecutionContext *Context, uint32 NewNumInstances, uint32 PrevNumInstances) const;
 
+	FORCEINLINE FNiagaraGPUInstanceCountManager& GetGPUInstanceCounterManager() { check(IsInRenderingThread()); return GPUInstanceCounterManager; }
+
 private:
-	void SimStepClearAndSetup(const FNiagaraComputeInstanceData& Instance, FRHICommandList& RHICmdList) const;
+
+	void ExecuteAll(FRHICommandList &RHICmdList, FUniformBufferRHIParamRef ViewUniformBuffer, bool bSetReadback);
+	void TickSingle(const FNiagaraGPUSystemTick& Tick, FNiagaraComputeInstanceData *Instance, FRHICommandList &RHICmdList, FUniformBufferRHIParamRef ViewUniformBuffer, bool bSetReadback) const;
+	void ResizeBuffersAndGatherResources(FOverlappableTicks& OverlappableTick, FRHICommandList& RHICmdList, FNiagaraBufferArray& DestDataBuffers, FNiagaraBufferArray& CurrDataBuffers, FNiagaraBufferArray& DestBufferIntFloat, FNiagaraBufferArray& CurrBufferIntFloat);
+	void DispatchAllOnCompute(FOverlappableTicks& OverlappableTick, FRHICommandList& RHICmdList, FUniformBufferRHIParamRef ViewUniformBuffer, FNiagaraBufferArray& DestDataBuffers, FNiagaraBufferArray& CurrDataBuffers, FNiagaraBufferArray& DestBufferIntFloat, FNiagaraBufferArray& CurrBufferIntFloat, bool bSetReadback);
+
 //	void SimStepReadback(const FNiagaraComputeInstanceData& Instance, FRHICommandList& RHICmdList) const;
 
 	inline uint32 UnpackEmitterDispatchCount(uint8* PackedData)
@@ -162,6 +163,9 @@ private:
 	TArray<FNiagaraGPUSortInfo> SimulationsToSort;
 	FParticleSortBuffers ParticleSortBuffers;
 
+	// GPU emitter instance count buffer. Contains the actual particle / instance count generate in the GPU tick.
+	mutable FNiagaraGPUInstanceCountManager GPUInstanceCounterManager;
+
 	// @todo REMOVE THIS HACK
 	uint32 LastFrameThatDrainedData;
 
diff --git a/Engine/Plugins/FX/Niagara/Source/Niagara/Classes/NiagaraGPUInstanceCountManager.h b/Engine/Plugins/FX/Niagara/Source/Niagara/Classes/NiagaraGPUInstanceCountManager.h
new file mode 100644
--- /dev/null
+++ b/Engine/Plugins/FX/Niagara/Source/Niagara/Classes/NiagaraGPUInstanceCountManager.h
@@ -0,0 +1,106 @@
+// Copyright 1998-2019 Epic Games, Inc. All Rights Reserved.
+
+/*==============================================================================
+NiagaraGPUInstanceCountManager.h: GPU particle count handling
+==============================================================================*/
+#pragma once
+
+#include "CoreMinimal.h"
+#include "NiagaraCommon.h"
+#include "RHIUtilities.h"
+#include "NiagaraDrawIndirect.h"
+#include "RHIGPUReadback.h"
+
+class FRHIGPUMemoryReadback;
+
+/** 
+ * A manager that handles the buffer containing the GPU particle count. 
+ * Also provides related functionalities like the generation of the draw indirect buffer.
+ */
+class FNiagaraGPUInstanceCountManager
+{
+public:
+
+	FNiagaraGPUInstanceCountManager();
+	~FNiagaraGPUInstanceCountManager();
+
+	// Init resource for the first time.
+	void InitRHI();
+	// Free resources.
+	void ReleaseRHI();
+
+	FRWBuffer& GetInstanceCountBuffer() 
+	{ 
+		check(UsedInstanceCounts <= AllocatedInstanceCounts); // Can't resize after after the buffer gets bound.
+		return CountBuffer; 
+	}
+
+	/** Free the entry and reset it to INDEX_NONE if valid. */
+	void FreeEntry(uint32& BufferOffset);
+	uint32 AcquireEntry();
+
+	const uint32* GetGPUReadback();
+	void ReleaseGPUReadback();
+	void EnqueueGPUReadback(FRHICommandListImmediate& RHICmdList);
+	bool HasPendingGPUReadback() const;
+
+	/** Add a draw indirect task to generate the draw indirect args. Returns the draw indirect arg buffer offset. */
+	uint32 AddDrawIndirect(uint32 InstanceCountBufferOffset, uint32 NumIndicesPerInstance);
+	FRWBuffer& GetDrawIndirectBuffer() { return DrawIndirectBuffer; }
+
+	/** 
+	 * Update the max possible required draw indirect args (one per renderer).
+	 * Called on the renderthread from FNiagaraRenderer::CreateRenderThreadResources()
+	 * or FNiagaraRenderer::ReleaseRenderThreadResources()
+	 */
+	void IncrementMaxDrawIndirectCount() 
+	{ 
+		checkSlow(IsInRenderingThread()); 
+		MaxDrawIndirectArgs += DrawIndirectArgsIncrement; 
+	}
+
+	void DecrementMaxDrawIndirectCount() 
+	{ 
+		checkSlow(IsInRenderingThread()); 
+		MaxDrawIndirectArgs -= DrawIndirectArgsIncrement; 
+
+	}
+
+	// Resize instance count and draw indirect buffers to ensure it is big enough to hold all draw indirect args.
+	void ResizeBuffers(int32 ReservedInstanceCounts);
+
+	// Generate the draw indirect buffers, and reset all release counts.
+	void UpdateDrawIndirectBuffer(FRHICommandList& RHICmdList, ERHIFeatureLevel::Type FeatureLevel);
+
+protected:
+
+	/** Buffer slack size hysteresis used to managed the size of CountBuffer and DrawIndirectBuffer. */
+	float BufferSlack = 1.5f;
+
+	/** The current used instance counts allocated from FNiagaraDataBuffer::AllocateGPU() */
+	int32 UsedInstanceCounts = 0;
+	/** The allocated instance counts in CountBuffer */
+	int32 AllocatedInstanceCounts = 0;
+
+	/** A buffer holding the each emitter particle count after a simulation tick. */
+	FRWBuffer CountBuffer;
+	TArray<uint32> FreeEntries;
+	FRHIGPUMemoryReadback* CountReadback = nullptr;
+	int32 CountReadbackSize = 0;
+
+	/** The max possible required draw indirect args count based on the number of GPU renderer current registered */
+	int32 MaxDrawIndirectArgs = 0;
+	/** The allocated indirect args in DrawIndirectBuffer (each being 5 uint32) */
+	int32 AllocatedDrawIndirectArgs = 0;
+
+	/** The list of all draw indirected tasks that are to be run in UpdateDrawIndirectBuffer() */
+	typedef FNiagaraDrawIndirectArgsGenCS::FArgGenTaskInfo FArgGenTaskInfo;
+	TArray<FArgGenTaskInfo> DrawIndirectArgGenTasks;
+	/** The list of all instance count clear tasks that are to be run in UpdateDrawIndirectBuffer() */
+	TArray<uint32> InstanceCountClearTasks;
+	/** A buffer holding drawindirect data to render GPU emitter renderers. */
+	FRWBuffer DrawIndirectBuffer;
+
+	// A multiplier applied to each increment/decrement in IncrementMaxDrawIndirectCount()/DecrementMaxDrawIndirectCount()
+	int32 DrawIndirectArgsIncrement = 1;
+};
\ No newline at end of file
diff --git a/Engine/Plugins/FX/Niagara/Source/Niagara/Classes/NiagaraScriptExecutionContext.h b/Engine/Plugins/FX/Niagara/Source/Niagara/Classes/NiagaraScriptExecutionContext.h
--- a/Engine/Plugins/FX/Niagara/Source/Niagara/Classes/NiagaraScriptExecutionContext.h
+++ b/Engine/Plugins/FX/Niagara/Source/Niagara/Classes/NiagaraScriptExecutionContext.h
@@ -18,6 +18,8 @@ NiagaraEmitterInstance.h: Niagara emitter simulation class
 #include "RHIGPUReadback.h"
 
 struct FNiagaraDataInterfaceProxy;
+class FNiagaraGPUInstanceCountManager;
+class NiagaraEmitterInstanceBatcher;
 
 /** Container for data needed to process event data. */
 struct FNiagaraEventHandlingInfo
@@ -150,7 +152,7 @@ struct FNiagaraComputeExecutionContext
 	FNiagaraComputeExecutionContext();
 	~FNiagaraComputeExecutionContext();
 
-	void Reset();
+	void Reset(NiagaraEmitterInstanceBatcher* Batcher);
 
 	void InitParams(UNiagaraScript* InGPUComputeScript, ENiagaraSimTarget InSimTarget, const FString& InDebugSimName);
 	void DirtyDataInterfaces();
@@ -161,8 +163,16 @@ struct FNiagaraComputeExecutionContext
 	void SetDataToRender(FNiagaraDataBuffer* InDataToRender);
 	FNiagaraDataBuffer* GetDataToRender()const { return DataToRender; }
 
+	struct 
+	{
+		// The offset at which the GPU instance count (see FNiagaraGPUInstanceCountManager()).
+		uint32 GPUCountOffset = INDEX_NONE;
+		// The CPU instance count at the time the GPU count readback was issued. Always bigger or equal to the GPU count.
+		uint32 CPUCount = 0;
+	}  EmitterInstanceReadback;
+
 private:
-	void ResetInternal();
+	void ResetInternal(NiagaraEmitterInstanceBatcher* Batcher);
 
 public:
 	static uint32 TickCounter;
@@ -186,11 +196,6 @@ public:
 	//Most current buffer that can be used for rendering.
 	FNiagaraDataBuffer* DataToRender;
 
-	bool ResetSinceLastReadbackIssued = false;
-	FRHIGPUMemoryReadback *GPUDataReadback;
-	uint32 AccumulatedSpawnRate;
-	uint32 NumIndicesPerInstance;	// how many vtx indices per instance the renderer is going to have for its draw call
-
 	uint32 EventSpawnTotal_GT;
 	uint32 SpawnRateInstances_GT;
 
@@ -202,6 +207,7 @@ public:
 	mutable uint32 GPUDebugDataIntSize;
 	mutable uint32 GPUDebugDataFloatStride;
 	mutable uint32 GPUDebugDataIntStride;
+	mutable uint32 GPUDebugDataCountOffset;
 	mutable TSharedPtr<struct FNiagaraScriptDebuggerInfo, ESPMode::ThreadSafe> DebugInfo;
 #endif
 
diff --git a/Engine/Plugins/FX/Niagara/Source/Niagara/Private/NiagaraComponent.cpp b/Engine/Plugins/FX/Niagara/Source/Niagara/Private/NiagaraComponent.cpp
--- a/Engine/Plugins/FX/Niagara/Source/Niagara/Private/NiagaraComponent.cpp
+++ b/Engine/Plugins/FX/Niagara/Source/Niagara/Private/NiagaraComponent.cpp
@@ -120,15 +120,17 @@ void FNiagaraSceneProxy::ReleaseRenderers()
 {
 	if (EmitterRenderers.Num() > 0)
 	{
+		NiagaraEmitterInstanceBatcher* TheBatcher = Batcher && !Batcher->IsPendingKill() ? Batcher : nullptr;
+
 		//Renderers must be freed on the render thread.
 		ENQUEUE_RENDER_COMMAND(ReleaseRenderersCommand)(
-			[ToDeleteEmitterRenderers = MoveTemp(EmitterRenderers)](FRHICommandListImmediate& RHICmdList)
+			[ToDeleteEmitterRenderers = MoveTemp(EmitterRenderers), TheBatcher](FRHICommandListImmediate& RHICmdList)
 		{
 			for (FNiagaraRenderer* EmitterRenderer : ToDeleteEmitterRenderers)
 			{
 				if (EmitterRenderer)
 				{
-					EmitterRenderer->ReleaseRenderThreadResources();
+					EmitterRenderer->ReleaseRenderThreadResources(TheBatcher);
 					delete EmitterRenderer;
 				}
 			}
@@ -195,7 +197,7 @@ FNiagaraSceneProxy::~FNiagaraSceneProxy()
 	{
 		if (EmitterRenderer)
 		{
-			EmitterRenderer->ReleaseRenderThreadResources();
+			EmitterRenderer->ReleaseRenderThreadResources(Batcher);
 			delete EmitterRenderer;
 		}
 	}
@@ -208,7 +210,7 @@ void FNiagaraSceneProxy::ReleaseRenderThreadResources()
 	{
 		if (Renderer)
 		{
-			Renderer->ReleaseRenderThreadResources();
+			Renderer->ReleaseRenderThreadResources(Batcher);
 		}
 	}
 	return;
@@ -221,7 +223,7 @@ void FNiagaraSceneProxy::CreateRenderThreadResources()
 	{
 		if (Renderer)
 		{
-			Renderer->CreateRenderThreadResources();
+			Renderer->CreateRenderThreadResources(Batcher);
 		}
 	}
 	return;
diff --git a/Engine/Plugins/FX/Niagara/Source/Niagara/Private/NiagaraDataSet.cpp b/Engine/Plugins/FX/Niagara/Source/Niagara/Private/NiagaraDataSet.cpp
--- a/Engine/Plugins/FX/Niagara/Source/Niagara/Private/NiagaraDataSet.cpp
+++ b/Engine/Plugins/FX/Niagara/Source/Niagara/Private/NiagaraDataSet.cpp
@@ -8,6 +8,8 @@
 #include "ShaderParameterUtils.h"
 #include "NiagaraStats.h"
 #include "NiagaraRenderer.h"
+#include "NiagaraGPUInstanceCountManager.h"
+#include "NiagaraEmitterInstanceBatcher.h"
 
 DECLARE_CYCLE_STAT(TEXT("InitRenderData"), STAT_InitRenderData, STATGROUP_Niagara);
 
@@ -346,6 +348,14 @@ void FNiagaraDataSet::Dump(int32 StartIndex, int32 NumInstances, const FString&
 	}
 }
 
+void FNiagaraDataSet::ReleaseGPUInstanceCounts(FNiagaraGPUInstanceCountManager& GPUInstanceCountManager)
+{
+	for (FNiagaraDataBuffer* Buffer : Data)
+	{
+		Buffer->ReleaseGPUInstanceCount(GPUInstanceCountManager);
+	}
+}
+
 void FNiagaraDataSet::BuildLayout()
 {
 	VariableLayouts.Empty();
@@ -446,6 +456,7 @@ void FNiagaraDataSet::CopyFromGPUReadback(float* GPUReadBackFloat, int* GPUReadB
 
 FNiagaraDataBuffer::FNiagaraDataBuffer(FNiagaraDataSet* InOwner)
 	: Owner(InOwner)
+	, GPUInstanceCountBufferOffset(INDEX_NONE)
 	, NumChunksAllocatedForGPU(0)
 	, NumInstances(0)
 	, NumInstancesAllocated(0)
@@ -457,6 +468,10 @@ FNiagaraDataBuffer::FNiagaraDataBuffer(FNiagaraDataSet* InOwner)
 FNiagaraDataBuffer::~FNiagaraDataBuffer()
 {
 	check(!IsInUse());
+	// If this is data for a GPU emitter, we have to release the GPU instance counts for reuse.
+	// The only exception is if the batcher was pending kill and we couldn't enqueue a rendering command, 
+	// in which case this would have been released on the game thread and not from the batcher DataSetsToDestroy_RT.
+	check(!IsInRenderingThread() || GPUInstanceCountBufferOffset == INDEX_NONE);
 	DEC_MEMORY_STAT_BY(STAT_NiagaraParticleMemory, FloatData.Max() + Int32Data.Max());
 }
 
@@ -587,18 +602,16 @@ void FNiagaraDataBuffer::Allocate(uint32 InNumInstances, bool bMaintainExisting)
 	}
 }
 
-void FNiagaraDataBuffer::AllocateGPU(uint32 InNumInstances, FRHICommandList &RHICmdList)
+void FNiagaraDataBuffer::AllocateGPU(uint32 InNumInstances, FNiagaraGPUInstanceCountManager& GPUInstanceCountManager, FRHICommandList &RHICmdList)
 {
 	CheckUsage(false);
 
 	check(Owner->SimTarget == ENiagaraSimTarget::GPUComputeSim);
 
-	//Allocate the indices buffer if needed
-	if (GPUIndices.Buffer == nullptr)
-	{
-		// Use BUF_KeepCPUAccessible here since some platforms will lock it for readonly (depending on the implementation of RHIEnqueueStagedRead) after GPU simulation.
-		GPUIndices.Initialize(sizeof(int32), 64 /*Context->NumDataSets*/, EPixelFormat::PF_R32_UINT, BUF_DrawIndirect | BUF_Static | BUF_KeepCPUAccessible);	// always allocate for up to 64 data sets
-	}
+	// Release previous entry if any.
+	GPUInstanceCountManager.FreeEntry(GPUInstanceCountBufferOffset);
+	// Get a new entry currently set to 0, since simulation will increment it to the actual instance count.
+	GPUInstanceCountBufferOffset = GPUInstanceCountManager.AcquireEntry();
 
 	// ALLOC_CHUNKSIZE must be greater than zero and divisible by the thread group size
 	const uint32 ALLOC_CHUNKSIZE = 4096;
@@ -998,6 +1011,12 @@ void FNiagaraDataBuffer::UnsetShaderParams(FNiagaraShader *Shader, FRHICommandLi
 	}
 }
 
+void FNiagaraDataBuffer::ReleaseGPUInstanceCount(FNiagaraGPUInstanceCountManager& GPUInstanceCountManager)
+{
+	GPUInstanceCountManager.FreeEntry(GPUInstanceCountBufferOffset);
+}
+
+
 FScopedNiagaraDataSetGPUReadback::~FScopedNiagaraDataSetGPUReadback()
 {
 	if (DataBuffer != nullptr)
@@ -1007,11 +1026,12 @@ FScopedNiagaraDataSetGPUReadback::~FScopedNiagaraDataSetGPUReadback()
 	}
 }
 
-void FScopedNiagaraDataSetGPUReadback::ReadbackData(FNiagaraDataSet* InDataSet)
+void FScopedNiagaraDataSetGPUReadback::ReadbackData(NiagaraEmitterInstanceBatcher* InBatcher, FNiagaraDataSet* InDataSet)
 {
 	check(DataSet == nullptr);
 	check(InDataSet != nullptr);
 
+	Batcher = InBatcher && !InBatcher->IsPendingKill() ? Batcher : nullptr;
 	DataSet = InDataSet;
 	DataBuffer = DataSet->GetCurrentData();
 
@@ -1019,20 +1039,23 @@ void FScopedNiagaraDataSetGPUReadback::ReadbackData(FNiagaraDataSet* InDataSet)
 	check((DataBuffer->FloatData.Num() == 0) && (DataBuffer->Int32Data.Num() == 0));
 
 	// Readback data
-	TArray<uint32> DrawIndirectData;
 	ENQUEUE_RENDER_COMMAND(ReadbackGPUBuffers)
 	(
 		[&](FRHICommandListImmediate& RHICmdList)
 		{
 			// Read DrawIndirect Params
+			const uint32 BufferOffset = DataBuffer->GetGPUInstanceCountBufferOffset();
+			if (Batcher && BufferOffset != INDEX_NONE)
 			{
-				const FRWBuffer& DataSetIndices = DataBuffer->GetGPUIndices();
-				const int32 BytesToCopy = DataSetIndices.NumBytes / sizeof(uint32);
-				DrawIndirectData.AddUninitialized(BytesToCopy);
+				FVertexBufferRHIParamRef InstanceCountBuffer = Batcher->GetGPUInstanceCounterManager().GetInstanceCountBuffer().Buffer;
 
-				void* Data = RHICmdList.LockVertexBuffer(DataSetIndices.Buffer, 0, BytesToCopy, RLM_ReadOnly);
-				FMemory::Memcpy(DrawIndirectData.GetData(), Data, BytesToCopy);
-				RHICmdList.UnlockVertexBuffer(DataSetIndices.Buffer);
+				void* Data = RHICmdList.LockVertexBuffer(InstanceCountBuffer, 0, (BufferOffset + 1) * sizeof(int32), RLM_ReadOnly);
+				NumInstances = reinterpret_cast<int32*>(Data)[BufferOffset];
+				RHICmdList.UnlockVertexBuffer(InstanceCountBuffer);
+			}
+			else
+			{
+				NumInstances = DataBuffer->GetNumInstances();
 			}
 
 			// Read float data
@@ -1059,6 +1082,4 @@ void FScopedNiagaraDataSetGPUReadback::ReadbackData(FNiagaraDataSet* InDataSet)
 		}
 	);
 	FlushRenderingCommands();
-
-	NumInstances = DrawIndirectData[1];
 }
diff --git a/Engine/Plugins/FX/Niagara/Source/Niagara/Private/NiagaraEmitterInstance.cpp b/Engine/Plugins/FX/Niagara/Source/Niagara/Private/NiagaraEmitterInstance.cpp
--- a/Engine/Plugins/FX/Niagara/Source/Niagara/Private/NiagaraEmitterInstance.cpp
+++ b/Engine/Plugins/FX/Niagara/Source/Niagara/Private/NiagaraEmitterInstance.cpp
@@ -387,7 +387,7 @@ void FNiagaraEmitterInstance::ResetSimulation(bool bKillExisting)
 
 		if (CachedEmitter->SimTarget == ENiagaraSimTarget::GPUComputeSim && GPUExecContext != nullptr)
 		{
-			GPUExecContext->Reset();
+			GPUExecContext->Reset(Batcher);
 		}
 	}
 }
@@ -678,7 +678,7 @@ FBox FNiagaraEmitterInstance::CalculateDynamicBounds(const bool bReadGPUSimulati
 		if (!bReadGPUSimulation || (GPUExecContext == nullptr))
 			return FBox(ForceInit);
 
-		ScopedGPUReadback.ReadbackData(GPUExecContext->MainDataSet);
+		ScopedGPUReadback.ReadbackData(Batcher, GPUExecContext->MainDataSet);
 		NumInstances = ScopedGPUReadback.GetNumInstances();
 	}
 	else
@@ -1018,15 +1018,6 @@ void FNiagaraEmitterInstance::Tick(float DeltaSeconds)
 		check(GPUExecContext->GPUScript_RT == CachedEmitter->GetGPUComputeScript()->GetRenderThreadScript());
 		GPUExecContext->GPUScript_RT = CachedEmitter->GetGPUComputeScript()->GetRenderThreadScript();
 
-		if (CachedEmitter->GetRenderers().Num() > 0 && CachedEmitter->GetRenderers()[0] != nullptr)
-		{
-			GPUExecContext->NumIndicesPerInstance = CachedEmitter->GetRenderers()[0]->GetNumIndicesPerInstance();
-		}
-		else
-		{
-			GPUExecContext->NumIndicesPerInstance = 0;
-		}
-
 		GPUExecContext->EventSpawnTotal_GT = EventSpawnTotal;
 		GPUExecContext->SpawnRateInstances_GT = SpawnTotal;
 		
diff --git a/Engine/Plugins/FX/Niagara/Source/Niagara/Private/NiagaraEmitterInstanceBatcher.cpp b/Engine/Plugins/FX/Niagara/Source/Niagara/Private/NiagaraEmitterInstanceBatcher.cpp
--- a/Engine/Plugins/FX/Niagara/Source/Niagara/Private/NiagaraEmitterInstanceBatcher.cpp
+++ b/Engine/Plugins/FX/Niagara/Source/Niagara/Private/NiagaraEmitterInstanceBatcher.cpp
@@ -17,7 +17,6 @@
 DECLARE_CYCLE_STAT(TEXT("Niagara Dispatch Setup"), STAT_NiagaraGPUDispatchSetup_RT, STATGROUP_Niagara);
 DECLARE_CYCLE_STAT(TEXT("GPU Emitter Dispatch [RT]"), STAT_NiagaraGPUSimTick_RT, STATGROUP_Niagara);
 DECLARE_CYCLE_STAT(TEXT("GPU Data Readback [RT]"), STAT_NiagaraGPUReadback_RT, STATGROUP_Niagara);
-DECLARE_CYCLE_STAT(TEXT("Allocate GPU Readback Data [RT]"), STAT_NiagaraAllocateGPUReadback_RT, STATGROUP_Niagara);
 DECLARE_FLOAT_COUNTER_STAT(TEXT("Niagara GPU Sim"), STAT_GPU_NiagaraSim, STATGROUP_GPU);
 DECLARE_DWORD_COUNTER_STAT(TEXT("# GPU Particles"), STAT_NiagaraGPUParticles, STATGROUP_Niagara);
 DECLARE_DWORD_COUNTER_STAT(TEXT("# GPU Sorted Particles"), STAT_NiagaraGPUSortedParticles, STATGROUP_Niagara);
@@ -27,8 +26,6 @@ DECLARE_DWORD_COUNTER_STAT(TEXT("Readback latency (frames)"), STAT_NiagaraReadba
 DECLARE_GPU_STAT_NAMED(NiagaraGPU, TEXT("Niagara"));
 DECLARE_GPU_STAT_NAMED(NiagaraGPUSimulation, TEXT("Niagara GPU Simulation"));
 DECLARE_GPU_STAT_NAMED(NiagaraGPUSorting, TEXT("Niagara GPU sorting"));
-DECLARE_GPU_STAT_NAMED(NiagaraSimReadback, TEXT("Niagara GPU Simulation read back"));
-DECLARE_GPU_STAT_NAMED(NiagaraIndexBufferClear, TEXT("Niagara index buffer clear"));
 
 uint32 FNiagaraComputeExecutionContext::TickCounter = 0;
 
@@ -44,7 +41,7 @@ static FAutoConsoleVariableRef CVarNiagaraUseAsyncCompute(
 int32 GNiagaraGpuMaxQueuedRenderFrames = 10;
 static FAutoConsoleVariableRef CVarNiagaraGpuMaxQueuedRenderFrames(
 	TEXT("fx.NiagaraGpuMaxQueuedRenderFrames"),
-	GNiagaraOverlapCompute,
+	GNiagaraGpuMaxQueuedRenderFrames,
 	TEXT("Number of frames we all to pass before we start to discard GPU ticks.\n"),
 	ECVF_Default
 );
@@ -67,9 +64,8 @@ FFXSystemInterface* NiagaraEmitterInstanceBatcher::GetInterface(const FName& InN
 
 NiagaraEmitterInstanceBatcher::~NiagaraEmitterInstanceBatcher()
 {
-	ParticleSortBuffers.ReleaseRHI();
-
 	FinishDispatches();
+	ParticleSortBuffers.ReleaseRHI();
 }
 
 void NiagaraEmitterInstanceBatcher::GiveSystemTick_RenderThread(FNiagaraGPUSystemTick& Tick)
@@ -121,12 +117,18 @@ void NiagaraEmitterInstanceBatcher::FinishDispatches()
 
 	for (FNiagaraComputeExecutionContext* Context : ContextsToDestroy_RT)
 	{
+		check(Context);
+		// Put back the GPU instance counter the global pool.
+		GPUInstanceCounterManager.FreeEntry(Context->EmitterInstanceReadback.GPUCountOffset);
 		delete Context;
 	}
 	ContextsToDestroy_RT.Reset();
 
 	for (FNiagaraDataSet* DataSet : DataSetsToDestroy_RT)
 	{
+		check(DataSet);
+		// Put back the GPU instance counter the global pool.
+		DataSet->ReleaseGPUInstanceCounts(GPUInstanceCounterManager);
 		delete DataSet;
 	}
 	DataSetsToDestroy_RT.Reset();
@@ -192,7 +194,7 @@ void NiagaraEmitterInstanceBatcher::ResizeBuffersAndGatherResources(FOverlappabl
 			//If this is not true, the read back in ResolveDatasetWrites will shrink the buffers.
 			const uint32 RequiredInstances = FMath::Max(PrevNumInstances, NewNumInstances);
 
-			DestinationData.AllocateGPU(RequiredInstances + 1, RHICmdList);
+			DestinationData.AllocateGPU(RequiredInstances + 1, GPUInstanceCounterManager, RHICmdList);
 			DestinationData.SetNumInstances(RequiredInstances);
 
 			if ( Shader->FloatInputBufferParam.IsBound() )
@@ -222,7 +224,7 @@ void NiagaraEmitterInstanceBatcher::ResizeBuffersAndGatherResources(FOverlappabl
 	}
 }
 
-void NiagaraEmitterInstanceBatcher::DispatchAllOnCompute(FOverlappableTicks& OverlappableTick, FRHICommandList& RHICmdList, FUniformBufferRHIParamRef ViewUniformBuffer, FNiagaraBufferArray& DestDataBuffers, FNiagaraBufferArray& CurrDataBuffers, FNiagaraBufferArray& DestBufferIntFloat, FNiagaraBufferArray& CurrBufferIntFloat)
+void NiagaraEmitterInstanceBatcher::DispatchAllOnCompute(FOverlappableTicks& OverlappableTick, FRHICommandList& RHICmdList, FUniformBufferRHIParamRef ViewUniformBuffer, FNiagaraBufferArray& DestDataBuffers, FNiagaraBufferArray& CurrDataBuffers, FNiagaraBufferArray& DestBufferIntFloat, FNiagaraBufferArray& CurrBufferIntFloat, bool bSetReadback)
 {
 	FRHICommandListImmediate& RHICmdListImmediate = FRHICommandListExecutor::GetImmediateCommandList();
 
@@ -235,10 +237,8 @@ void NiagaraEmitterInstanceBatcher::DispatchAllOnCompute(FOverlappableTicks& Ove
 
 	RHICmdList.TransitionResources(EResourceTransitionAccess::ERWBarrier, EResourceTransitionPipeline::EGfxToCompute, CurrDataBuffers.GetData(), CurrDataBuffers.Num());
 
+#if WITH_EDITORONLY_DATA
 	{
-		SCOPED_DRAW_EVENT(RHICmdList, NiagaraIndexBufferClear);
-		SCOPED_GPU_STAT(RHICmdList, NiagaraIndexBufferClear);
-
 		for (FNiagaraGPUSystemTick* Tick : OverlappableTick)
 		{
 			uint32 DispatchCount = Tick->Count;
@@ -249,24 +249,20 @@ void NiagaraEmitterInstanceBatcher::DispatchAllOnCompute(FOverlappableTicks& Ove
 				FNiagaraComputeExecutionContext* Context = Instance.Context;
 				if (Context && Context->GPUScript_RT->GetShader())
 				{
-#if WITH_EDITORONLY_DATA
 					if (Context->DebugInfo.IsValid())
 					{
 						ProcessDebugInfo(RHICmdList, Context);
 					}
-#endif // WITH_EDITORONLY_DATA
-
-					// clear data set index buffer for the simulation shader to write number of written instances
-					ClearUAV(RHICmdList, Instance.DestinationData->GetGPUIndices(), 0);
 				}
-
 			}
 		}
 	}
+#endif // WITH_EDITORONLY_DATA
 
 	//
 	//	Add a rw barrier for the destination data buffers we just cleared and mark others as read/write as needed for particles simulation.
 	//
+	RHICmdList.TransitionResource(EResourceTransitionAccess::ERWBarrier, EResourceTransitionPipeline::EComputeToCompute, GPUInstanceCounterManager.GetInstanceCountBuffer().UAV);		// transition to readable; we'll be using this next frame
 	RHICmdList.TransitionResources(EResourceTransitionAccess::ERWBarrier, EResourceTransitionPipeline::EComputeToCompute, DestDataBuffers.GetData(), DestDataBuffers.Num());
 	RHICmdList.TransitionResources(EResourceTransitionAccess::EWritable, EResourceTransitionPipeline::EComputeToCompute, DestBufferIntFloat.GetData(), DestBufferIntFloat.Num());
 	RHICmdList.TransitionResources(EResourceTransitionAccess::EReadable, EResourceTransitionPipeline::EComputeToCompute, CurrDataBuffers.GetData(), CurrDataBuffers.Num());
@@ -291,15 +287,28 @@ void NiagaraEmitterInstanceBatcher::DispatchAllOnCompute(FOverlappableTicks& Ove
 					// run shader, sim and spawn in a single dispatch
 					uint32 UpdateStartInstance = 0;
 					Run<false>(*Tick, &Instance, UpdateStartInstance, Instance.DestinationData->GetNumInstances(), Context->GPUScript_RT->GetShader(), RHICmdList, ViewUniformBuffer);
+
+					FNiagaraDataBuffer* CurrentData = Instance.CurrentData;
+					if (bSetReadback && Tick->bIsFinalTick && Index == DispatchCount - 1)
+					{
+						// Now that the current data is not required anymore, stage it for readback.
+						if (CurrentData->GetNumInstances() && Context->EmitterInstanceReadback.GPUCountOffset == INDEX_NONE && CurrentData->GetGPUInstanceCountBufferOffset() != INDEX_NONE)
+						{
+							// Transfer the GPU instance counter ownership to the context. Note that when bSetReadback is true, a readback request will be performed later in the tick update.
+							Context->EmitterInstanceReadback.GPUCountOffset = CurrentData->GetGPUInstanceCountBufferOffset();
+							Context->EmitterInstanceReadback.CPUCount = CurrentData->GetNumInstances();
+							CurrentData->ClearGPUInstanceCountBufferOffset();
+						}
+					}
 				}
 			}
 		}
 	}
 
-
 	//
 	//	Now Copy to staging buffer the data we want to read back (alive particle count). And make buffer ready for that and draw commands on the graphics pipe too.
 	//
+	RHICmdList.TransitionResource(EResourceTransitionAccess::ERWBarrier, EResourceTransitionPipeline::EComputeToCompute, GPUInstanceCounterManager.GetInstanceCountBuffer().UAV);		// transition to readable; we'll be using this next frame
 	RHICmdList.TransitionResources(EResourceTransitionAccess::ERWBarrier, EResourceTransitionPipeline::EComputeToCompute, DestDataBuffers.GetData(), DestDataBuffers.Num());
 	RHICmdList.TransitionResources(EResourceTransitionAccess::ERWBarrier, EResourceTransitionPipeline::EComputeToCompute, DestBufferIntFloat.GetData(), DestBufferIntFloat.Num());
 	RHICmdList.TransitionResources(EResourceTransitionAccess::EReadable, EResourceTransitionPipeline::EComputeToGfx, DestDataBuffers.GetData(), DestDataBuffers.Num());
@@ -309,32 +318,6 @@ void NiagaraEmitterInstanceBatcher::DispatchAllOnCompute(FOverlappableTicks& Ove
 	// We have done all our overlapping compute work on this list so go back to default behavior and flush.
 	RHICmdList.AutomaticCacheFlushAfterComputeShader(true);
 
-	{
-		SCOPED_DRAW_EVENT(RHICmdList, NiagaraSimReadback);
-		SCOPED_GPU_STAT(RHICmdList, NiagaraSimReadback);
-		for (FNiagaraGPUSystemTick* Tick : OverlappableTick)
-		{
-			uint32 DispatchCount = Tick->Count;
-			FNiagaraComputeInstanceData* Instances = Tick->GetInstanceData();
-			for (uint32 Index = 0; Index < DispatchCount; Index++)
-			{
-				FNiagaraComputeInstanceData& Instance = Instances[Index];
-				FNiagaraComputeExecutionContext* Context = Instance.Context;
-				if (Context && Context->GPUScript_RT->GetShader())
-				{
-					// Don't resolve if the data if there are no instances (prevents a transition issue warning).
-					if (Instance.DestinationData->GetNumInstances() > 0)
-					{
-						// resolve data set writes - grabs the number of instances written from the index set during the simulation run
-						ResolveDatasetWrites(RHICmdList, &Instance);
-					}
-					check(Instance.DestinationData->GetGPUIndices().Buffer);
-				}
-			}
-		}
-	}
-	// the VF grabs current state for drawing, so need to transition
-	
 	//TODO: Need to set the data to render in the context?
 
 	RHICmdList.TransitionResources(EResourceTransitionAccess::EReadable, EResourceTransitionPipeline::EComputeToGfx, DestBufferIntFloat.GetData(), DestBufferIntFloat.Num());
@@ -346,10 +329,16 @@ void NiagaraEmitterInstanceBatcher::DispatchAllOnCompute(FOverlappableTicks& Ove
 
 void NiagaraEmitterInstanceBatcher::PostRenderOpaque(FRHICommandListImmediate& RHICmdList, const FUniformBufferRHIParamRef ViewUniformBuffer, const class FShaderParametersMetadata* SceneTexturesUniformBufferStruct, FUniformBufferRHIParamRef SceneTexturesUniformBuffer)
 {
-	ExecuteAll(RHICmdList, ViewUniformBuffer);
+	// Setup new readback since if there is no pending request, there is no risk of having invalid data read (offset being allocated after the readback was sent).
+	ExecuteAll(RHICmdList, ViewUniformBuffer, !GPUInstanceCounterManager.HasPendingGPUReadback());
+
+	if (!GPUInstanceCounterManager.HasPendingGPUReadback())
+	{
+		GPUInstanceCounterManager.EnqueueGPUReadback(RHICmdList);
+	}
 }
 
-void NiagaraEmitterInstanceBatcher::ExecuteAll(FRHICommandList &RHICmdList, FUniformBufferRHIParamRef ViewUniformBuffer)
+void NiagaraEmitterInstanceBatcher::ExecuteAll(FRHICommandList &RHICmdList, FUniformBufferRHIParamRef ViewUniformBuffer, bool bSetReadback)
 {
 	// @todo REMOVE THIS HACK
 	LastFrameThatDrainedData = GFrameNumberRenderThread;
@@ -362,7 +351,6 @@ void NiagaraEmitterInstanceBatcher::ExecuteAll(FRHICommandList &RHICmdList, FUni
 
 	SCOPED_DRAW_EVENT(RHICmdList, NiagaraEmitterInstanceBatcher_ExecuteAll);
 
-
 	if (GNiagaraOverlapCompute>0)
 	{
 		TArray<FOverlappableTicks, TInlineAllocator<2>> SimPasses;
@@ -409,7 +397,7 @@ void NiagaraEmitterInstanceBatcher::ExecuteAll(FRHICommandList &RHICmdList, FUni
 			// It also ensures the GPU buffers are large enough to hold everything.
 			ResizeBuffersAndGatherResources(SimPass, RHICmdList, DestDataBuffers, CurrDataBuffers, DestBufferIntFloat, CurrBufferIntFloat);
 
-			DispatchAllOnCompute(SimPass, RHICmdList, ViewUniformBuffer, DestDataBuffers, CurrDataBuffers, DestBufferIntFloat, CurrBufferIntFloat);
+			DispatchAllOnCompute(SimPass, RHICmdList, ViewUniformBuffer, DestDataBuffers, CurrDataBuffers, DestBufferIntFloat, CurrBufferIntFloat, bSetReadback);
 		}
 	}
 	else
@@ -422,7 +410,7 @@ void NiagaraEmitterInstanceBatcher::ExecuteAll(FRHICommandList &RHICmdList, FUni
 			for (uint32 i = 0; i < DispatchCount; i++)
 			{
 				FNiagaraComputeInstanceData& Instance = Instances[i];
-				TickSingle(Tick, &Instance, RHICmdList, ViewUniformBuffer);
+				TickSingle(Tick, &Instance, RHICmdList, ViewUniformBuffer, bSetReadback && i == DispatchCount - 1); // Only readback last count.
 			}
 		}
 	}
@@ -430,62 +418,7 @@ void NiagaraEmitterInstanceBatcher::ExecuteAll(FRHICommandList &RHICmdList, FUni
 	FinishDispatches();
 }
 
-void NiagaraEmitterInstanceBatcher::SimStepClearAndSetup(const FNiagaraComputeInstanceData& Instance, FRHICommandList& RHICmdList) const
-{
-	check(IsInRenderingThread());
-#if 1
-	FNiagaraComputeExecutionContext* Context = Instance.Context;
-	check(Context);
-
-
-	FNiagaraShader* ComputeShader = Context->GPUScript_RT->GetShader();
-	check(ComputeShader);
-	if (!ComputeShader)
-	{
-		return;
-	}
-
-#if WITH_EDITORONLY_DATA
-	if (Context->DebugInfo.IsValid())
-	{
-		ProcessDebugInfo(RHICmdList, Context);
-	}
-#endif // WITH_EDITORONLY_DATA
-
-	// clear data set index buffer for the simulation shader to write number of written instances to
-	//
-	FRWBuffer &DatasetIndexBufferWrite = Instance.DestinationData->GetGPUIndices();
-//	SCOPED_GPU_EVENTF(RHICmdList, NiagaraIndexBufferClear, TEXT("Niagara index buffer clear"));
-	SCOPED_DRAW_EVENT(RHICmdList, NiagaraIndexBufferClear);
-	SCOPED_GPU_STAT(RHICmdList, NiagaraIndexBufferClear);
-	//UE_LOG(LogNiagara, Log, TEXT("Clearing UAV %p"), (FRHIVertexBuffer *)DatasetIndexBufferWrite.Buffer);
-	ClearUAV(RHICmdList, DatasetIndexBufferWrite, 0);
-#endif
-}
-
-//void NiagaraEmitterInstanceBatcher::SimStepReadback(const FNiagaraComputeInstanceData& Instance, FRHICommandList& RHICmdList) const
-//{
-//#if 1
-//	FNiagaraComputeExecutionContext* Context = Instance.Context;
-//	check(Context);
-//	FNiagaraShader* ComputeShader = Context->GPUScript_RT->GetShader();
-//	if (!ComputeShader)
-//	{
-//		return;
-//	}
-//
-//	// Don't resolve if the data if there are no instances (prevents a transition issue warning).
-//	if (Context->MainDataSet->CurrData().GetNumInstances() > 0)
-//	{
-//		// resolve data set writes - grabs the number of instances written from the index set during the simulation run
-//		ResolveDatasetWrites(RHICmdList, &Instance);
-//	}
-//
-//	check(Context->MainDataSet->HasDatasetIndices());
-//#endif
-//}
-
-void NiagaraEmitterInstanceBatcher::TickSingle(const FNiagaraGPUSystemTick& Tick, FNiagaraComputeInstanceData *Instance, FRHICommandList &RHICmdList, FUniformBufferRHIParamRef ViewUniformBuffer) const
+void NiagaraEmitterInstanceBatcher::TickSingle(const FNiagaraGPUSystemTick& Tick, FNiagaraComputeInstanceData *Instance, FRHICommandList &RHICmdList, FUniformBufferRHIParamRef ViewUniformBuffer, bool bSetReadback) const
 {
 	SCOPE_CYCLE_COUNTER(STAT_NiagaraGPUSimTick_RT);
 
@@ -517,38 +450,16 @@ void NiagaraEmitterInstanceBatcher::TickSingle(const FNiagaraGPUSystemTick& Tick
 	FNiagaraDataBuffer& CurrentData = *Instance->CurrentData;
 	FNiagaraDataBuffer& DestinationData = *Instance->DestinationData;
 
-	uint32 PrevNumInstances = Tick.bNeedsReset ? 0 : CurrentData.GetNumInstances();
-	uint32 NewNumInstances = Instance->SpawnRateInstances + Instance->EventSpawnTotal + PrevNumInstances;
-
-	//We must assume all particles survive when allocating here. 
-	//If this is not true, the read back in ResolveDatasetWrites will shrink the buffers.
-	uint32 RequiredInstances = FMath::Max(PrevNumInstances, NewNumInstances);
+	const uint32 PrevNumInstances = Tick.bNeedsReset ? 0 : CurrentData.GetNumInstances();
+	const uint32 NewNumInstances = Instance->SpawnRateInstances + Instance->EventSpawnTotal + PrevNumInstances;
 
-	DestinationData.AllocateGPU(RequiredInstances + 1, RHICmdList);
-	DestinationData.SetNumInstances(RequiredInstances);
-
-	// clear data set index buffer for the simulation shader to write number of written instances to
-	{
-		FRWBuffer& DatasetIndexBufferWrite = DestinationData.GetGPUIndices();
-		SCOPED_DRAW_EVENTF(RHICmdList, NiagaraIndexBufferClear, TEXT("Niagara index buffer clear"));
-		ClearUAV(RHICmdList, DatasetIndexBufferWrite, 0);
-		RHICmdList.TransitionResource(EResourceTransitionAccess::EWritable, EResourceTransitionPipeline::EComputeToCompute, DatasetIndexBufferWrite.UAV);
-	}
+	DestinationData.AllocateGPU(NewNumInstances + 1, GPUInstanceCounterManager, RHICmdList);
+	DestinationData.SetNumInstances(NewNumInstances);
 
 	// run shader, sim and spawn in a single dispatch
 	uint32 UpdateStartInstance = 0;
 	Run<true>(Tick, Instance, UpdateStartInstance, NewNumInstances, ComputeShader, RHICmdList, ViewUniformBuffer);
 
-	// ResolveDatasetWrites may read this, so we must transition it here.
-	RHICmdList.TransitionResource(EResourceTransitionAccess::EReadable, EResourceTransitionPipeline::EComputeToGfx, DestinationData.GetGPUIndices().UAV);		// transition to readable; we'll be using this next frame
-
-	// Don't resolve if the data if there are no instances (prevents a transition issue warning).
-	if (NewNumInstances > 0)
-	{
-		// resolve data set writes - grabs the number of instances written from the index set during the simulation run
-		ResolveDatasetWrites(RHICmdList, Instance);
-	}
-
 	Context->MainDataSet->EndSimulate();
 	Context->SetDataToRender(Instance->DestinationData);
 
@@ -565,9 +476,21 @@ void NiagaraEmitterInstanceBatcher::TickSingle(const FNiagaraGPUSystemTick& Tick
 	// Better to do this in the VF set?
 	RHICmdList.TransitionResource(EResourceTransitionAccess::EReadable, EResourceTransitionPipeline::EComputeToGfx, Context->DataToRender->GetGPUBufferFloat().UAV);
 	RHICmdList.TransitionResource(EResourceTransitionAccess::EReadable, EResourceTransitionPipeline::EComputeToGfx, Context->DataToRender->GetGPUBufferInt().UAV);
+
+	// Put a barrier so that next tick or arg gen gets final count.
+	RHICmdList.TransitionResource(EResourceTransitionAccess::ERWBarrier, EResourceTransitionPipeline::EComputeToCompute, GPUInstanceCounterManager.GetInstanceCountBuffer().UAV);		// transition to readable; we'll be using this next frame
+
+	// Now that the current data is not required anymore, stage it for readback.
+	if (bSetReadback && PrevNumInstances && Context->EmitterInstanceReadback.GPUCountOffset == INDEX_NONE && CurrentData.GetGPUInstanceCountBufferOffset() != INDEX_NONE)
+	{
+		// Transfer the GPU instance counter ownership to the context. Note that when bSetReadback is true, a readback request will be performed later in the tick update.
+		Context->EmitterInstanceReadback.GPUCountOffset = CurrentData.GetGPUInstanceCountBufferOffset();
+		Context->EmitterInstanceReadback.CPUCount = PrevNumInstances;
+		CurrentData.ClearGPUInstanceCountBufferOffset();
+	}
 }
 
-void NiagaraEmitterInstanceBatcher::PreInitViews()
+void NiagaraEmitterInstanceBatcher::PreInitViews(FRHICommandListImmediate& RHICmdList)
 {
 	SortedParticleCount = 0;
 	SimulationsToSort.Reset();
@@ -576,6 +499,57 @@ void NiagaraEmitterInstanceBatcher::PreInitViews()
 	{
 		SortedVertexBuffer.UsedIndexCount = 0;
 	}
+
+	// Update draw indirect buffer to max possible size.
+	int32 TotalDispatchCount = 0;
+	for (FNiagaraGPUSystemTick& Tick : Ticks_RT)
+	{
+		TotalDispatchCount += (int32)Tick.Count;
+
+		// Cancel any pending readback if the emitter is resetting.
+		if (Tick.bNeedsReset)
+		{
+			FNiagaraComputeExecutionContext* Context = Tick.GetInstanceData()->Context;
+			if (Context)
+			{
+				GPUInstanceCounterManager.FreeEntry(Context->EmitterInstanceReadback.GPUCountOffset);
+			}
+		}
+	}
+	GPUInstanceCounterManager.ResizeBuffers(TotalDispatchCount);
+
+	// Update the instance counts from the GPU readback.
+	{
+		SCOPE_CYCLE_COUNTER(STAT_NiagaraGPUReadback_RT);
+		const uint32* Counts = GPUInstanceCounterManager.GetGPUReadback();
+		if (Counts)
+		{
+			for (FNiagaraGPUSystemTick& Tick : Ticks_RT)
+			{
+				FNiagaraComputeExecutionContext* Context = Tick.GetInstanceData()->Context;
+				if (Context && Context->EmitterInstanceReadback.GPUCountOffset != INDEX_NONE)
+				{
+					check(Context->MainDataSet);
+					FNiagaraDataBuffer* CurrentData = Context->MainDataSet->GetCurrentData();
+					if (CurrentData)
+					{
+						const uint32 DeadInstanceCount = Context->EmitterInstanceReadback.CPUCount - Counts[Context->EmitterInstanceReadback.GPUCountOffset];
+						if (DeadInstanceCount && DeadInstanceCount < CurrentData->GetNumInstances())
+						{
+							CurrentData->SetNumInstances(CurrentData->GetNumInstances() - DeadInstanceCount);
+						}
+					}
+
+					// Now release the readback since another one will be enqueued in the tick.
+					// Also prevents processing the same data again.
+					GPUInstanceCounterManager.FreeEntry(Context->EmitterInstanceReadback.GPUCountOffset);
+				}
+			}
+			// Readback is only valid for one frame, so that any newly allocated instance count
+			// are guarantied to be in the next valid readback data.
+			GPUInstanceCounterManager.ReleaseGPUReadback();
+		}
+	}
 }
 
 bool NiagaraEmitterInstanceBatcher::UsesGlobalDistanceField() const
@@ -597,6 +571,10 @@ void NiagaraEmitterInstanceBatcher::PreRender(FRHICommandListImmediate& RHICmdLi
 
 	// Sort buffer after mesh batches are issued, before tick (which will change the GPU instance count).
 	SortGPUParticles(RHICmdList);
+
+	// Update draw indirect args from the simulation results.
+	GPUInstanceCounterManager.UpdateDrawIndirectBuffer(RHICmdList, FeatureLevel);
+
 }
 
 int32 NiagaraEmitterInstanceBatcher::AddSortedGPUSimulation(const FNiagaraGPUSortInfo& SortInfo)
@@ -789,69 +767,6 @@ void NiagaraEmitterInstanceBatcher::ResolveParticleSortBuffers(FRHICommandListIm
 	CopyBufferCS->UnbindBuffers(RHICmdList);
 }
 
-/* Attempt to read back simulation results (number of live instances) from the GPU via an async readback request; 
- *	If the readback isn't ready to be performed, we accumulate spawn rates and assume all instances have survived, until
- *	the GPU can tell us how many are actually alive; since that data may be several frames old, we'll always end up
- *	overallocating a bit, and the CPU might think we have more particles alive than we actually do;
- *	since we use DrawIndirect with the GPU determining draw call parameters, that's not an issue
- */
-void NiagaraEmitterInstanceBatcher::ResolveDatasetWrites(FRHICommandList &RHICmdList, FNiagaraComputeInstanceData *Instance) const
-{
-	FNiagaraComputeExecutionContext* Context = Instance->Context;
-	FNiagaraDataBuffer& DestinationData = *Instance->DestinationData;
-	FRWBuffer &DatasetIndexBufferWrite = DestinationData.GetGPUIndices();
-	uint32 SpawnedThisFrame = Instance->SpawnRateInstances + Instance->EventSpawnTotal;
-	Context->AccumulatedSpawnRate += SpawnedThisFrame;
-	int32 ExistingDataCount = DestinationData.GetNumInstances();
-	if (!Context->GPUDataReadback)
-	{
-		SCOPE_CYCLE_COUNTER(STAT_NiagaraAllocateGPUReadback_RT);
-
-		Context->GPUDataReadback = new FRHIGPUMemoryReadback(TEXT("Niagara GPU Emitter Readback"));
-		INC_DWORD_STAT(STAT_NiagaraReadbackLatency);
-		Context->GPUDataReadback->EnqueueCopy(RHICmdList, DatasetIndexBufferWrite.Buffer);
-		INC_DWORD_STAT_BY(STAT_NiagaraGPUParticles, ExistingDataCount);
-	}
-	else if (Context->GPUDataReadback->IsReady())
-	{
-		bool bSuccessfullyRead = false;
-		{
-		    SCOPE_CYCLE_COUNTER(STAT_NiagaraGPUReadback_RT);
-		    int32 *NumInstancesAfterSim = static_cast<int32*>(Context->GPUDataReadback->Lock(64 * sizeof(int32)));
-			if (NumInstancesAfterSim)
-			{
-				int32 NewExistingDataCount = (Context->ResetSinceLastReadbackIssued ? 0 : NumInstancesAfterSim[1]) + Context->AccumulatedSpawnRate; // index 1 is always the count
-				DestinationData.SetNumInstances(NewExistingDataCount);
-				//UE_LOG(LogNiagara, Log, TEXT("GPU Syncup %s : Was(%d) Now(%d) Accumulated(%d) NumInstancesAfterSim[1](%d)"), *Context->DebugSimName, ExistingDataCount, NewExistingDataCount , Context->AccumulatedSpawnRate, NumInstancesAfterSim[1]);
-				//UE_LOG(LogNiagara, Log, TEXT("Reading UAV %p"), (FRHIVertexBuffer *)DatasetIndexBufferWrite.Buffer);
-
-				INC_DWORD_STAT_BY(STAT_NiagaraGPUParticles, NewExistingDataCount);
-				SET_DWORD_STAT(STAT_NiagaraReadbackLatency, 0);
-
-				Context->AccumulatedSpawnRate = 0;
-				Context->ResetSinceLastReadbackIssued = false;
-				bSuccessfullyRead = true;
-			}
-			else
-			{
-				UE_LOG(LogNiagara, Warning, TEXT("GPUDataReadback said it was ready, but returned an invalid buffer. Skipping this time.."));
-				INC_DWORD_STAT_BY(STAT_NiagaraGPUParticles, ExistingDataCount);
-			}
-			Context->GPUDataReadback->Unlock();
-		}
-		if (bSuccessfullyRead)
-		{
-			SCOPE_CYCLE_COUNTER(STAT_NiagaraAllocateGPUReadback_RT);
-			// The following code seems to take significant time on d3d12
-			Context->GPUDataReadback->EnqueueCopy(RHICmdList, DatasetIndexBufferWrite.Buffer);
-		}
-	}
-	else
-	{
-		INC_DWORD_STAT_BY(STAT_NiagaraGPUParticles, ExistingDataCount);
-	}
-}
-
 void NiagaraEmitterInstanceBatcher::ProcessDebugInfo(FRHICommandList &RHICmdList, const FNiagaraComputeExecutionContext* Context) const
 {
 #if WITH_EDITORONLY_DATA
@@ -872,8 +787,7 @@ void NiagaraEmitterInstanceBatcher::ProcessDebugInfo(FRHICommandList &RHICmdList
 		{
 			//UE_LOG(LogNiagara, Warning, TEXT("Read back!"));
 
-			int32 *NumInstancesAfterSim = static_cast<int32*>(Context->GPUDebugDataReadbackCounts->Lock(64 * sizeof(int32)));
-			int32 NewExistingDataCount = NumInstancesAfterSim[1];
+			int32 NewExistingDataCount =  static_cast<int32*>(Context->GPUDebugDataReadbackCounts->Lock((Context->GPUDebugDataCountOffset + 1) * sizeof(int32)))[Context->GPUDebugDataCountOffset];
 			{
 				float* FloatDataBuffer = nullptr;
 				if (Context->GPUDebugDataReadbackFloat)
@@ -919,6 +833,7 @@ void NiagaraEmitterInstanceBatcher::ProcessDebugInfo(FRHICommandList &RHICmdList
 				Context->GPUDebugDataIntSize = 0;
 				Context->GPUDebugDataFloatStride = 0;
 				Context->GPUDebugDataIntStride = 0;
+				Context->GPUDebugDataCountOffset = INDEX_NONE;
 			}
 
 			// We've updated the debug info directly, now we need to no longer keep asking and querying because this frame is done!
@@ -1016,32 +931,8 @@ void NiagaraEmitterInstanceBatcher::Run(const FNiagaraGPUSystemTick& Tick, const
 	FNiagaraDataBuffer& DestinationData = *Instance->DestinationData;
 	FNiagaraDataBuffer& CurrentData = *Instance->CurrentData;
 
-	const FRWBuffer& WriteIndexBuffer = DestinationData.GetGPUIndices();
-	FRWBuffer& ReadIndexBuffer = CurrentData.GetGPUIndices();
-
-	// if we don't have a previous index buffer, we need to prep one using the maximum number of instances; this should only happen on the first frame
-	//		the data set index buffer is really the param buffer for the indirect draw call; it contains the number of live instances at index 1, and the simulation
-	//		CS uses this to determine the current number of active instances in the buffer
-	//
-	if (ReadIndexBuffer.Buffer == nullptr)
-	{
-		TResourceArray<int32> InitIndexBuffer;
-		InitIndexBuffer.AddUninitialized(64);
-		InitIndexBuffer[1] = 0;		// number of instances 
-		ReadIndexBuffer.Initialize(sizeof(int32), 64, EPixelFormat::PF_R32_UINT, BUF_DrawIndirect | BUF_Static, nullptr, &InitIndexBuffer);
-	}
-	else if (Tick.bNeedsReset)
-	{
-		//UE_LOG(LogNiagara, Log, TEXT("Clearing UAV due to reset!"));
-		ClearUAV(RHICmdList, ReadIndexBuffer, 0);
-		Context->AccumulatedSpawnRate = 0;
-		Context->ResetSinceLastReadbackIssued = true;
-	}
-
 	RHICmdList.SetComputeShader(Shader->GetComputeShader());
 
-	RHICmdList.SetShaderResourceViewParameter(Shader->GetComputeShader(), Shader->InputIndexBufferParam.GetBaseIndex(), ReadIndexBuffer.SRV);
-
 	// set the view uniform buffer param
 	if (Shader->ViewUniformBufferParam.IsBound() && ViewUniformBuffer)
 	{
@@ -1057,9 +948,13 @@ void NiagaraEmitterInstanceBatcher::Run(const FNiagaraGPUSystemTick& Tick, const
 
 	// set the index buffer uav
 	//
-	if (Shader->OutputIndexBufferParam.IsBound())
+	if (Shader->InstanceCountsParam.IsBound())
 	{
-		RHICmdList.SetUAVParameter(Shader->GetComputeShader(), Shader->OutputIndexBufferParam.GetUAVIndex(), WriteIndexBuffer.UAV);
+		RHICmdList.SetUAVParameter(Shader->GetComputeShader(), Shader->InstanceCountsParam.GetUAVIndex(), GPUInstanceCounterManager.GetInstanceCountBuffer().UAV);
+		const uint32 ReadOffset = CurrentData.GetGPUInstanceCountBufferOffset();
+		const uint32 WriteOffset = DestinationData.GetGPUInstanceCountBufferOffset();
+		RHICmdList.SetShaderParameter(Shader->GetComputeShader(), Shader->ReadInstanceCountOffsetParam.GetBufferIndex(), Shader->ReadInstanceCountOffsetParam.GetBaseIndex(), Shader->ReadInstanceCountOffsetParam.GetNumBytes(), &ReadOffset);
+		RHICmdList.SetShaderParameter(Shader->GetComputeShader(), Shader->WriteInstanceCountOffsetParam.GetBufferIndex(), Shader->WriteInstanceCountOffsetParam.GetBaseIndex(), Shader->WriteInstanceCountOffsetParam.GetNumBytes(), &WriteOffset);
 	}
 
 	// set the execution parameters
@@ -1069,10 +964,7 @@ void NiagaraEmitterInstanceBatcher::Run(const FNiagaraGPUSystemTick& Tick, const
 		RHICmdList.SetShaderParameter(Shader->GetComputeShader(), Shader->EmitterTickCounterParam.GetBufferIndex(), Shader->EmitterTickCounterParam.GetBaseIndex(), Shader->EmitterTickCounterParam.GetNumBytes(), &FNiagaraComputeExecutionContext::TickCounter);
 	}
 
-	uint32 Copy = bCopyBeforeStart ? 1 : 0;
-
 	RHICmdList.SetShaderParameter(Shader->GetComputeShader(), Shader->UpdateStartInstanceParam.GetBufferIndex(), Shader->UpdateStartInstanceParam.GetBaseIndex(), Shader->UpdateStartInstanceParam.GetNumBytes(), &UpdateStartInstance);					// 0, except for event handler runs
-	RHICmdList.SetShaderParameter(Shader->GetComputeShader(), Shader->NumIndicesPerInstanceParam.GetBufferIndex(), Shader->NumIndicesPerInstanceParam.GetBaseIndex(), Shader->NumIndicesPerInstanceParam.GetNumBytes(), &Context->NumIndicesPerInstance);		// set from the renderer in FNiagaraEmitterInstance::Tick
 	int32 InstancesToSpawnThisFrame = Instance->SpawnRateInstances + Instance->EventSpawnTotal;
 	RHICmdList.SetShaderParameter(Shader->GetComputeShader(), Shader->NumSpawnedInstancesParam.GetBufferIndex(), Shader->NumSpawnedInstancesParam.GetBaseIndex(), Shader->NumSpawnedInstancesParam.GetNumBytes(), &InstancesToSpawnThisFrame);				// number of instances in the spawn run
 
@@ -1111,10 +1003,8 @@ void NiagaraEmitterInstanceBatcher::Run(const FNiagaraGPUSystemTick& Tick, const
 	{
 		//UE_LOG(LogNiagara, Warning, TEXT("Queued up!"));
 
-		if (!Context->GPUDebugDataReadbackFloat && !Context->GPUDebugDataReadbackInt && !Context->GPUDebugDataReadbackCounts)
+		if (!Context->GPUDebugDataReadbackFloat && !Context->GPUDebugDataReadbackInt && !Context->GPUDebugDataReadbackCounts && DestinationData.GetGPUInstanceCountBufferOffset() != INDEX_NONE)
 		{
-			FRWBuffer &DatasetIndexBufferWrite = DestinationData.GetGPUIndices();
-
 			Context->GPUDebugDataFloatSize = 0;
 			Context->GPUDebugDataIntSize = 0;
 			Context->GPUDebugDataFloatStride = 0;
@@ -1137,7 +1027,8 @@ void NiagaraEmitterInstanceBatcher::Run(const FNiagaraGPUSystemTick& Tick, const
 			}
 
 			Context->GPUDebugDataReadbackCounts = new FRHIGPUMemoryReadback(TEXT("Niagara GPU Emitter Readback"));
-			Context->GPUDebugDataReadbackCounts->EnqueueCopy(RHICmdList, DatasetIndexBufferWrite.Buffer);
+			Context->GPUDebugDataReadbackCounts->EnqueueCopy(RHICmdList, GPUInstanceCounterManager.GetInstanceCountBuffer().Buffer);
+			Context->GPUDebugDataCountOffset = DestinationData.GetGPUInstanceCountBufferOffset();
 		}
 	}
 #endif // WITH_EDITORONLY_DATA
@@ -1147,5 +1038,5 @@ void NiagaraEmitterInstanceBatcher::Run(const FNiagaraGPUSystemTick& Tick, const
 	UnsetDataInterfaceParameters(DataInterfaceProxies, Shader, RHICmdList, Instance, Tick);
 	CurrentData.UnsetShaderParams(Shader, RHICmdList);
 	DestinationData.UnsetShaderParams(Shader, RHICmdList);
-	Shader->OutputIndexBufferParam.UnsetUAV(RHICmdList, Shader->GetComputeShader());
+	Shader->InstanceCountsParam.UnsetUAV(RHICmdList, Shader->GetComputeShader());
 }
diff --git a/Engine/Plugins/FX/Niagara/Source/Niagara/Private/NiagaraGPUInstanceCountManager.cpp b/Engine/Plugins/FX/Niagara/Source/Niagara/Private/NiagaraGPUInstanceCountManager.cpp
new file mode 100644
--- /dev/null
+++ b/Engine/Plugins/FX/Niagara/Source/Niagara/Private/NiagaraGPUInstanceCountManager.cpp
@@ -0,0 +1,245 @@
+// Copyright 1998-2019 Epic Games, Inc. All Rights Reserved.
+
+#include "NiagaraGPUInstanceCountManager.h"
+#include "NiagaraStats.h"
+
+int32 GNiagaraMinGPUInstanceCount = 2048;
+static FAutoConsoleVariableRef CVarNiagaraMinGPUInstanceCount(
+	TEXT("Niagara.MinGPUInstanceCount"),
+	GNiagaraMinGPUInstanceCount,
+	TEXT("Minimum number of instance count entries allocated in the global buffer. (default=2048)"),
+	ECVF_Default
+);
+
+int32 GNiagaraMinGPUDrawIndirectArgs = 256;
+static FAutoConsoleVariableRef CVarNiagaraMinGPUDrawIndirectArgs(
+	TEXT("Niagara.MinGPUDrawIndirectArgs"),
+	GNiagaraMinGPUDrawIndirectArgs,
+	TEXT("Minimum number of draw indirect argsallocated in the global buffer. (default=256)"),
+	ECVF_Default
+);
+
+DECLARE_DWORD_COUNTER_STAT(TEXT("Used GPU Instance Counters"), STAT_NiagaraUsedGPUInstanceCounters, STATGROUP_Niagara);
+DECLARE_DWORD_COUNTER_STAT(TEXT("Max Num GPU Renderers"), STAT_NiagaraMaxNumGPURenderers, STATGROUP_Niagara);
+
+
+//*****************************************************************************
+
+FNiagaraGPUInstanceCountManager::FNiagaraGPUInstanceCountManager() 
+{
+#if RHI_RAYTRACING
+	// GetDynamicRayTracingInstances()
+	++DrawIndirectArgsIncrement;
+#endif
+	if (GIsEditor)
+	{
+		// Hit proxies
+		++DrawIndirectArgsIncrement;
+	}
+}
+
+FNiagaraGPUInstanceCountManager::~FNiagaraGPUInstanceCountManager()
+{
+	ReleaseRHI();
+}
+
+void FNiagaraGPUInstanceCountManager::InitRHI()
+{
+}
+
+void FNiagaraGPUInstanceCountManager::ReleaseRHI()
+{
+	CountBuffer.Release();
+	DrawIndirectBuffer.Release();
+
+	AllocatedInstanceCounts = 0;
+	AllocatedDrawIndirectArgs = 0;
+
+	if (CountReadback)
+	{
+		delete CountReadback;
+		CountReadback = nullptr;
+		CountReadbackSize = 0;
+	}
+}
+
+uint32 FNiagaraGPUInstanceCountManager::AcquireEntry()
+{
+	checkSlow(IsInRenderingThread());
+
+	if (FreeEntries.Num())
+	{
+		return FreeEntries.Pop();
+	}
+	else if (UsedInstanceCounts < AllocatedInstanceCounts)
+	{
+		// We can't reallocate on the fly, the buffer must be correctly resized before any tick gets scheduled.
+		return UsedInstanceCounts++;
+	}
+	else
+	{
+		// @TODO : add realloc the buffer and copy the current content to it. Might require reallocating the readback in FNiagaraGPUInstanceCountManager::EnqueueGPUReadback()
+		ensure(UsedInstanceCounts < AllocatedInstanceCounts);
+		UE_LOG(LogNiagara, Error, TEXT("Niagara.MinGPUInstanceCount too small."));
+		return INDEX_NONE;
+	}
+}
+
+void FNiagaraGPUInstanceCountManager::FreeEntry(uint32& BufferOffset)
+{
+	checkSlow(IsInRenderingThread());
+
+	if (BufferOffset != INDEX_NONE)
+	{
+		// Add a reset to 0 task.
+		// The entry will only become available/reusable after being reset to 0 in UpdateDrawIndirectBuffer()
+		InstanceCountClearTasks.Add(BufferOffset);
+		BufferOffset = INDEX_NONE;
+	}
+}
+
+void FNiagaraGPUInstanceCountManager::ResizeBuffers(int32 ReservedInstanceCounts)
+{
+	const int32 RequiredInstanceCounts = UsedInstanceCounts + FMath::Max<int32>(ReservedInstanceCounts - FreeEntries.Num(), 0);
+	if (RequiredInstanceCounts > 0 || MaxDrawIndirectArgs > 0)
+	{
+		const int32 RecommendedInstanceCounts = FMath::Max(GNiagaraMinGPUInstanceCount, (int32)(RequiredInstanceCounts * BufferSlack));
+		// If the buffer is not allocated, allocate it to the recommended size.
+		if (!AllocatedInstanceCounts)
+		{
+			AllocatedInstanceCounts = RecommendedInstanceCounts;
+			TResourceArray<uint32> InitData;
+			InitData.AddZeroed(AllocatedInstanceCounts);
+			CountBuffer.Initialize(sizeof(uint32), AllocatedInstanceCounts, EPixelFormat::PF_R32_UINT, BUF_Static | BUF_KeepCPUAccessible, TEXT("NiagaraGPUInstanceCounts"), &InitData);
+		}
+		// If we need to increase the buffer size to RecommendedInstanceCounts because the buffer is too small.
+		else if (RequiredInstanceCounts > AllocatedInstanceCounts)
+		{
+			check(false); // @TODO : reallocate and copy current data, filling new space with 0s.
+		}
+		// If we need to shrink the buffer size because use way to much buffer size.
+		else if ((int32)(RecommendedInstanceCounts * BufferSlack) < AllocatedInstanceCounts)
+		{
+			// possibly shrink but hard to do because of sparse array allocation.
+		}
+
+		// We expect 2 niagara buffer per instance for simplicity
+		int32 RecommendedDrawIndirectArgsCount = FMath::Max(GNiagaraMinGPUDrawIndirectArgs, (int32)(MaxDrawIndirectArgs * BufferSlack));
+		if (MaxDrawIndirectArgs > AllocatedDrawIndirectArgs || (int32)(RecommendedDrawIndirectArgsCount * BufferSlack) < AllocatedDrawIndirectArgs)
+		{
+			DrawIndirectBuffer.Release();
+			AllocatedDrawIndirectArgs = RecommendedDrawIndirectArgsCount;
+			DrawIndirectBuffer.Initialize(sizeof(uint32), RecommendedDrawIndirectArgsCount * NIAGARA_DRAW_INDIRECT_ARGS_SIZE, EPixelFormat::PF_R32_UINT, BUF_Static | BUF_DrawIndirect, TEXT("NiagaraGPUDrawIndirectArgs"));
+		}
+	}
+	else
+	{
+		ReleaseRHI();
+	}
+
+	INC_DWORD_STAT_BY(STAT_NiagaraUsedGPUInstanceCounters, RequiredInstanceCounts);
+	INC_DWORD_STAT_BY(STAT_NiagaraMaxNumGPURenderers, MaxDrawIndirectArgs);
+}
+
+uint32 FNiagaraGPUInstanceCountManager::AddDrawIndirect(uint32 InstanceCountBufferOffset, uint32 NumIndicesPerInstance)
+{
+	checkSlow(IsInRenderingThread());
+	if (DrawIndirectArgGenTasks.Num() < AllocatedDrawIndirectArgs)
+	{
+#if !UE_BUILD_SHIPPING
+		if (DrawIndirectArgGenTasks.Num() >= MaxDrawIndirectArgs)
+		{
+			UE_LOG(LogNiagara, Warning, TEXT("More draw indirect args then expected (%d / %d)"), DrawIndirectArgGenTasks.Num() + 1, MaxDrawIndirectArgs);
+		}
+#endif
+
+		DrawIndirectArgGenTasks.Add(FArgGenTaskInfo(InstanceCountBufferOffset, NumIndicesPerInstance));
+		return (DrawIndirectArgGenTasks.Num() - 1) * NIAGARA_DRAW_INDIRECT_ARGS_SIZE * sizeof(uint32);
+	}
+	else
+	{
+		UE_LOG(LogNiagara, Error, TEXT("Too many draw indirect args. Some draw call will be corrupted"));
+		return INDEX_NONE;
+	}
+}
+
+void FNiagaraGPUInstanceCountManager::UpdateDrawIndirectBuffer(FRHICommandList& RHICmdList, ERHIFeatureLevel::Type FeatureLevel)
+{
+	if (DrawIndirectArgGenTasks.Num() || InstanceCountClearTasks.Num())
+	{
+		RHICmdList.TransitionResource(EResourceTransitionAccess::ERWBarrier, EResourceTransitionPipeline::EComputeToCompute, CountBuffer.UAV);
+
+		FReadBuffer TaskInfosBuffer;
+		{
+			// All draw indirect args task are run first because of the binding between the task ID and arg write offset.
+			const uint32 ArgGenSize = DrawIndirectArgGenTasks.Num() * sizeof(FArgGenTaskInfo);
+			const uint32 InstanceCountClearSize = InstanceCountClearTasks.Num() * sizeof(uint32);
+			const uint32 TaskBufferSize = ArgGenSize + InstanceCountClearSize;
+			TaskInfosBuffer.Initialize(sizeof(uint32), TaskBufferSize / sizeof(uint32), EPixelFormat::PF_R32_UINT, BUF_Volatile);
+			uint8* TaskBufferData = (uint8*)RHILockVertexBuffer(TaskInfosBuffer.Buffer, 0, TaskBufferSize, RLM_WriteOnly);
+			FMemory::Memcpy(TaskBufferData, DrawIndirectArgGenTasks.GetData(), ArgGenSize);
+			FMemory::Memcpy(TaskBufferData + ArgGenSize, InstanceCountClearTasks.GetData(), InstanceCountClearSize);
+			RHIUnlockVertexBuffer(TaskInfosBuffer.Buffer);
+		}
+
+		RHICmdList.TransitionResource(EResourceTransitionAccess::EWritable, EResourceTransitionPipeline::EGfxToCompute, DrawIndirectBuffer.UAV);
+
+		FNiagaraDrawIndirectArgsGenCS::FPermutationDomain PermutationVector;
+		TShaderMapRef<FNiagaraDrawIndirectArgsGenCS> DrawIndirectArgsGenCS(GetGlobalShaderMap(FeatureLevel), PermutationVector);
+		RHICmdList.SetComputeShader(DrawIndirectArgsGenCS->GetComputeShader());
+		DrawIndirectArgsGenCS->SetOutput(RHICmdList, DrawIndirectBuffer.UAV, CountBuffer.UAV);
+		DrawIndirectArgsGenCS->SetParameters(RHICmdList, TaskInfosBuffer.SRV, DrawIndirectArgGenTasks.Num(), InstanceCountClearTasks.Num());
+		DispatchComputeShader(RHICmdList, *DrawIndirectArgsGenCS, FMath::DivideAndRoundUp(DrawIndirectArgGenTasks.Num() + InstanceCountClearTasks.Num(), NIAGARA_DRAW_INDIRECT_ARGS_GEN_THREAD_COUNT), 1, 1);
+		DrawIndirectArgsGenCS->UnbindBuffers(RHICmdList);
+
+		// Sync after clear.
+		RHICmdList.TransitionResource(EResourceTransitionAccess::ERWBarrier, EResourceTransitionPipeline::EComputeToCompute, CountBuffer.UAV);
+		// Transition draw indirect to readable for gfx draw indirect.
+		RHICmdList.TransitionResource(EResourceTransitionAccess::EReadable, EResourceTransitionPipeline::EComputeToGfx, DrawIndirectBuffer.UAV);
+
+		// Once cleared to 0, the count are reusable.
+		FreeEntries.Append(InstanceCountClearTasks);
+
+		DrawIndirectArgGenTasks.Empty();
+		InstanceCountClearTasks.Empty();
+	}
+}
+
+const uint32* FNiagaraGPUInstanceCountManager::GetGPUReadback() 
+{
+	if (CountReadback && CountReadbackSize && CountReadback->IsReady())
+	{
+		return (uint32*)(CountReadback->Lock(CountReadbackSize * sizeof(uint32)));
+	}
+	else
+	{
+		return nullptr; 
+	}
+}
+
+void FNiagaraGPUInstanceCountManager::ReleaseGPUReadback() 
+{
+	check(CountReadback && CountReadbackSize);
+	CountReadback->Unlock();
+	// Readback can only ever be done once, to prevent misusage with index lifetime
+	CountReadbackSize = 0; 
+}
+
+void FNiagaraGPUInstanceCountManager::EnqueueGPUReadback(FRHICommandListImmediate& RHICmdList)
+{
+	if (UsedInstanceCounts)
+	{
+		if (!CountReadback)
+		{
+			CountReadback = new FRHIGPUMemoryReadback(TEXT("Niagara GPU Instance Count Readback"));
+		}
+		CountReadbackSize = UsedInstanceCounts;
+		CountReadback->EnqueueCopy(RHICmdList, CountBuffer.Buffer);
+	}
+}
+
+bool FNiagaraGPUInstanceCountManager::HasPendingGPUReadback() const
+{ 
+	return CountReadback && CountReadbackSize;
+}
+
diff --git a/Engine/Plugins/FX/Niagara/Source/Niagara/Private/NiagaraMeshRendererProperties.cpp b/Engine/Plugins/FX/Niagara/Source/Niagara/Private/NiagaraMeshRendererProperties.cpp
--- a/Engine/Plugins/FX/Niagara/Source/Niagara/Private/NiagaraMeshRendererProperties.cpp
+++ b/Engine/Plugins/FX/Niagara/Source/Niagara/Private/NiagaraMeshRendererProperties.cpp
@@ -109,7 +109,7 @@ void UNiagaraMeshRendererProperties::GetUsedMaterials(TArray<UMaterialInterface*
 	}
 }
 
-uint32 UNiagaraMeshRendererProperties::GetNumIndicesPerInstance() 
+uint32 UNiagaraMeshRendererProperties::GetNumIndicesPerInstance() const
 {
 	// TODO: Add proper support for multiple mesh sections for GPU mesh particles.
 	return ParticleMesh ? ParticleMesh->RenderData->LODResources[0].Sections[0].NumTriangles * 3 : 0;
diff --git a/Engine/Plugins/FX/Niagara/Source/Niagara/Private/NiagaraRenderer.cpp b/Engine/Plugins/FX/Niagara/Source/Niagara/Private/NiagaraRenderer.cpp
--- a/Engine/Plugins/FX/Niagara/Source/Niagara/Private/NiagaraRenderer.cpp
+++ b/Engine/Plugins/FX/Niagara/Source/Niagara/Private/NiagaraRenderer.cpp
@@ -8,6 +8,7 @@
 #include "NiagaraVertexFactory.h"
 #include "Engine/Engine.h"
 #include "DynamicBufferAllocator.h"
+#include "NiagaraEmitterInstanceBatcher.h"
 
 DECLARE_CYCLE_STAT(TEXT("Sort Particles"), STAT_NiagaraSortParticles, STATGROUP_Niagara);
 
@@ -185,6 +186,8 @@ FNiagaraRenderer::FNiagaraRenderer(ERHIFeatureLevel::Type FeatureLevel, const UN
 	, CPUTimeMS(0.0f)
 	, bLocalSpace(Emitter->GetCachedEmitter()->bLocalSpace)
 	, bHasLights(false)
+	, SimTarget(Emitter->GetCachedEmitter()->SimTarget)
+	, NumIndicesPerInstance(InProps ? InProps->GetNumIndicesPerInstance() : 0)
 {
 #if STATS
 	EmitterStatID = Emitter->GetCachedEmitter()->GetStatID(false, false);
@@ -207,10 +210,26 @@ void FNiagaraRenderer::Initialize(ERHIFeatureLevel::Type FeatureLevel, const UNi
 
 FNiagaraRenderer::~FNiagaraRenderer()
 {
-	ReleaseRenderThreadResources();
+	ReleaseRenderThreadResources(nullptr);
 	SetDynamicData_RenderThread(nullptr);
 }
 
+void FNiagaraRenderer::CreateRenderThreadResources(NiagaraEmitterInstanceBatcher* Batcher) 
+{
+	if (Batcher && SimTarget == ENiagaraSimTarget::GPUComputeSim)
+	{
+		Batcher->GetGPUInstanceCounterManager().IncrementMaxDrawIndirectCount();
+	}
+}
+
+void FNiagaraRenderer::ReleaseRenderThreadResources(NiagaraEmitterInstanceBatcher* Batcher)
+{
+	if (Batcher && SimTarget == ENiagaraSimTarget::GPUComputeSim)
+	{
+		Batcher->GetGPUInstanceCounterManager().DecrementMaxDrawIndirectCount();
+	}
+}
+
 FPrimitiveViewRelevance FNiagaraRenderer::GetViewRelevance(const FSceneView* View, const FNiagaraSceneProxy *SceneProxy)const
 {
 	FPrimitiveViewRelevance Result;
diff --git a/Engine/Plugins/FX/Niagara/Source/Niagara/Private/NiagaraRendererLights.cpp b/Engine/Plugins/FX/Niagara/Source/Niagara/Private/NiagaraRendererLights.cpp
--- a/Engine/Plugins/FX/Niagara/Source/Niagara/Private/NiagaraRendererLights.cpp
+++ b/Engine/Plugins/FX/Niagara/Source/Niagara/Private/NiagaraRendererLights.cpp
@@ -64,7 +64,6 @@ FNiagaraDynamicDataBase* FNiagaraRendererLights::GenerateDynamicData(const FNiag
 	FNiagaraDataSetAccessor<float> ScatteringAccessor(Data, Properties->VolumetricScatteringBinding.DataSetVariable);
 	FNiagaraDataSetAccessor<int32> EnabledAccessor(Data, Properties->LightRenderingEnabledBinding.DataSetVariable);
 
-	ENiagaraSimTarget SimTarget = Emitter->GetData().GetSimTarget();
 	FNiagaraDynamicDataLights* DynamicData = new FNiagaraDynamicDataLights(Emitter);
 
 	const FMatrix& LocalToWorldMatrix = Proxy->GetLocalToWorld();
diff --git a/Engine/Plugins/FX/Niagara/Source/Niagara/Private/NiagaraRendererMeshes.cpp b/Engine/Plugins/FX/Niagara/Source/Niagara/Private/NiagaraRendererMeshes.cpp
--- a/Engine/Plugins/FX/Niagara/Source/Niagara/Private/NiagaraRendererMeshes.cpp
+++ b/Engine/Plugins/FX/Niagara/Source/Niagara/Private/NiagaraRendererMeshes.cpp
@@ -107,16 +107,16 @@ FNiagaraRendererMeshes::FNiagaraRendererMeshes(ERHIFeatureLevel::Type FeatureLev
 	MeshMinimumLOD = Properties->ParticleMesh->MinLOD.GetValueForFeatureLevel(FeatureLevel);
 }
 
-void FNiagaraRendererMeshes::ReleaseRenderThreadResources()
+void FNiagaraRendererMeshes::ReleaseRenderThreadResources(NiagaraEmitterInstanceBatcher* Batcher)
 {
-	FNiagaraRenderer::ReleaseRenderThreadResources();
+	FNiagaraRenderer::ReleaseRenderThreadResources(Batcher);
 	VertexFactory->ReleaseResource();
 	WorldSpacePrimitiveUniformBuffer.ReleaseResource();
 }
 
-void FNiagaraRendererMeshes::CreateRenderThreadResources()
+void FNiagaraRendererMeshes::CreateRenderThreadResources(NiagaraEmitterInstanceBatcher* Batcher)
 {
-	FNiagaraRenderer::CreateRenderThreadResources();
+	FNiagaraRenderer::CreateRenderThreadResources(Batcher);
 	VertexFactory->InitResource();
 }
 
@@ -135,11 +135,13 @@ void FNiagaraRendererMeshes::GetDynamicMeshElements(const TArray<const FSceneVie
 {
 	SCOPE_CYCLE_COUNTER(STAT_NiagaraRender);
 	SCOPE_CYCLE_COUNTER(STAT_NiagaraRenderMeshes);
+	check(SceneProxy);
 
 	SimpleTimer MeshElementsTimer;
 
+	NiagaraEmitterInstanceBatcher* Batcher = SceneProxy->GetBatcher();
 	FNiagaraDynamicDataMesh *DynamicDataMesh = (static_cast<FNiagaraDynamicDataMesh*>(DynamicDataRender));
-	if (!DynamicDataMesh)
+	if (!DynamicDataMesh || !Batcher)
 	{
 		return;
 	}
@@ -166,12 +168,18 @@ void FNiagaraRendererMeshes::GetDynamicMeshElements(const TArray<const FSceneVie
 	FGlobalDynamicReadBuffer& DynamicReadBuffer = Collector.GetDynamicReadBuffer();
 	FGlobalDynamicReadBuffer::FAllocation ParticleData;
 
+	// @TODO : support multiple LOD and section, using an inlined array and/or the SceneRenderingAllocator
+	uint32 IndirectArgsOffset = INDEX_NONE;
 	//For cpu sims we allocate render buffers from the global pool. GPU sims own their own.
-	if (DynamicDataMesh->GetSimTarget() == ENiagaraSimTarget::CPUSim)
+	if (SimTarget == ENiagaraSimTarget::CPUSim)
 	{
 		ParticleData = DynamicReadBuffer.AllocateFloat(TotalFloatSize);
 		FMemory::Memcpy(ParticleData.Buffer, SourceParticleData->GetFloatBuffer().GetData(), SourceParticleData->GetFloatBuffer().Num());
 	}
+	else // ENiagaraSimTarget::GPUComputeSim
+	{
+		IndirectArgsOffset = Batcher->GetGPUInstanceCounterManager().AddDrawIndirect(SourceParticleData->GetGPUInstanceCountBufferOffset(), NumIndicesPerInstance);
+	}
 
 	{
 		// Update the primitive uniform buffer if needed.
@@ -260,24 +268,23 @@ void FNiagaraRendererMeshes::GetDynamicMeshElements(const TArray<const FSceneVie
 					}
 				};
 
-				if (DynamicDataMesh->GetSimTarget() == ENiagaraSimTarget::CPUSim)//TODO: Compute shader for sorting gpu sims and larger cpu sims.
+				if (SimTarget == ENiagaraSimTarget::CPUSim)//TODO: Compute shader for sorting gpu sims and larger cpu sims.
 				{
 					check(ParticleData.IsValid());
 					if (SortInfo.SortMode != ENiagaraSortMode::None && SortInfo.SortAttributeOffset != INDEX_NONE)
 					{
 						if (GNiagaraGPUSorting &&
 							GNiagaraGPUSortingCPUToGPUThreshold != INDEX_NONE &&
-							SortInfo.ParticleCount >= GNiagaraGPUSortingCPUToGPUThreshold &&
-							SceneProxy->GetBatcher())
+							SortInfo.ParticleCount >= GNiagaraGPUSortingCPUToGPUThreshold)
 						{
 							SortInfo.ParticleCount = NumInstances;
 							SortInfo.ParticleDataFloatSRV = ParticleData.ReadBuffer->SRV;
 							SortInfo.FloatDataOffset = ParticleData.FirstIndex / sizeof(float);
 							SortInfo.FloatDataStride = SourceParticleData->GetFloatStride() / sizeof(float);
-							const int32 IndexBufferOffset = SceneProxy->GetBatcher()->AddSortedGPUSimulation(SortInfo);
+							const int32 IndexBufferOffset = Batcher->AddSortedGPUSimulation(SortInfo);
 							if (IndexBufferOffset != INDEX_NONE)
 							{
-								CollectorResources.VertexFactory.SetSortedIndices(SceneProxy->GetBatcher()->GetGPUSortedBuffer().VertexBufferSRV, IndexBufferOffset);
+								CollectorResources.VertexFactory.SetSortedIndices(Batcher->GetGPUSortedBuffer().VertexBufferSRV, IndexBufferOffset);
 							}
 						}
 						else
@@ -292,21 +299,7 @@ void FNiagaraRendererMeshes::GetDynamicMeshElements(const TArray<const FSceneVie
 				}
 				else
 				{
-					if (SortInfo.SortMode != ENiagaraSortMode::None && SortInfo.SortAttributeOffset != INDEX_NONE && GNiagaraGPUSorting && SceneProxy->GetBatcher())
-					{
-						SortInfo.ParticleDataFloatSRV = SourceParticleData->GetGPUBufferFloat().SRV;
-						SortInfo.FloatDataOffset = 0;
-						SortInfo.FloatDataStride = SourceParticleData->GetFloatStride() / sizeof(float);
-						SortInfo.GPUParticleCountSRV = SourceParticleData->GetGPUIndices().SRV;
-						SortInfo.GPUParticleCountOffset = 1;
-						const int32 IndexBufferOffset = SceneProxy->GetBatcher()->AddSortedGPUSimulation(SortInfo);
-						if (IndexBufferOffset != INDEX_NONE)
-						{
-							CollectorResources.VertexFactory.SetSortedIndices(SceneProxy->GetBatcher()->GetGPUSortedBuffer().VertexBufferSRV, IndexBufferOffset);
-						}
-					}
-
-					if (SortInfo.SortMode != ENiagaraSortMode::None && SortInfo.SortAttributeOffset != INDEX_NONE && GNiagaraGPUSorting && SceneProxy->GetBatcher())
+					if (SortInfo.SortMode != ENiagaraSortMode::None && SortInfo.SortAttributeOffset != INDEX_NONE && GNiagaraGPUSorting)
 					{
 						// Here we need to be conservative about the InstanceCount, since the final value is only known on the GPU after the simulation.
 						SortInfo.ParticleCount = SourceParticleData->GetNumInstances();
@@ -314,12 +307,12 @@ void FNiagaraRendererMeshes::GetDynamicMeshElements(const TArray<const FSceneVie
 						SortInfo.ParticleDataFloatSRV = SourceParticleData->GetGPUBufferFloat().SRV;
 						SortInfo.FloatDataOffset = 0;
 						SortInfo.FloatDataStride = SourceParticleData->GetFloatStride() / sizeof(float);
-						SortInfo.GPUParticleCountSRV = SourceParticleData->GetGPUIndices().SRV;
-						SortInfo.GPUParticleCountOffset = 1;
-						const int32 IndexBufferOffset = SceneProxy->GetBatcher()->AddSortedGPUSimulation(SortInfo);
-						if (IndexBufferOffset != INDEX_NONE)
+						SortInfo.GPUParticleCountSRV = Batcher->GetGPUInstanceCounterManager().GetInstanceCountBuffer().SRV;
+						SortInfo.GPUParticleCountOffset = SourceParticleData->GetGPUInstanceCountBufferOffset();
+						const int32 IndexBufferOffset = Batcher->AddSortedGPUSimulation(SortInfo);
+						if (IndexBufferOffset != INDEX_NONE && SortInfo.GPUParticleCountOffset != INDEX_NONE)
 						{
-							CollectorResources.VertexFactory.SetSortedIndices(SceneProxy->GetBatcher()->GetGPUSortedBuffer().VertexBufferSRV, IndexBufferOffset);
+							CollectorResources.VertexFactory.SetSortedIndices(Batcher->GetGPUSortedBuffer().VertexBufferSRV, IndexBufferOffset);
 						}
 					}
 					CollectorResources.VertexFactory.SetParticleData(SourceParticleData->GetGPUBufferFloat().SRV, 0, SourceParticleData->GetFloatStride() / sizeof(float));
@@ -335,7 +328,7 @@ void FNiagaraRendererMeshes::GetDynamicMeshElements(const TArray<const FSceneVie
 			
 				// GPU mesh rendering currently only supports one mesh section.
 				// TODO: Add proper support for multiple mesh sections for GPU mesh particles.
-				int32 MaxSection = DynamicDataMesh->GetSimTarget() == ENiagaraSimTarget::GPUComputeSim ? 1 : LODModel.Sections.Num();
+				int32 MaxSection = SimTarget == ENiagaraSimTarget::GPUComputeSim ? 1 : LODModel.Sections.Num();
 				const bool bIsWireframe = AllowDebugViewmodes() && View && View->Family->EngineShowFlags.Wireframe;
 				for (int32 SectionIndex = 0; SectionIndex < MaxSection; SectionIndex++)
 				{
@@ -391,10 +384,11 @@ void FNiagaraRendererMeshes::GetDynamicMeshElements(const TArray<const FSceneVie
 						BatchElement.NumPrimitives = Section.NumTriangles;
 					}
 
-					if (DynamicDataMesh->GetSimTarget() == ENiagaraSimTarget::GPUComputeSim)
+					if (IndirectArgsOffset != INDEX_NONE)
 					{
 						BatchElement.NumPrimitives = 0;
-						BatchElement.IndirectArgsBuffer = SourceParticleData->GetGPUIndices().Buffer;
+						BatchElement.IndirectArgsOffset = IndirectArgsOffset;
+						BatchElement.IndirectArgsBuffer = Batcher->GetGPUInstanceCounterManager().GetDrawIndirectBuffer().Buffer;
 					}
 					else
 					{
@@ -432,7 +426,6 @@ FNiagaraDynamicDataBase *FNiagaraRendererMeshes::GenerateDynamicData(const FNiag
 
 	FNiagaraDynamicDataMesh *DynamicData = nullptr;
 
-	ENiagaraSimTarget SimTarget = Emitter->GetData().GetSimTarget();
 	if (Properties->ParticleMesh)
 	{
 		DynamicData = new FNiagaraDynamicDataMesh(Emitter);
diff --git a/Engine/Plugins/FX/Niagara/Source/Niagara/Private/NiagaraRendererRibbons.cpp b/Engine/Plugins/FX/Niagara/Source/Niagara/Private/NiagaraRendererRibbons.cpp
--- a/Engine/Plugins/FX/Niagara/Source/Niagara/Private/NiagaraRendererRibbons.cpp
+++ b/Engine/Plugins/FX/Niagara/Source/Niagara/Private/NiagaraRendererRibbons.cpp
@@ -198,17 +198,17 @@ FNiagaraRendererRibbons::FNiagaraRendererRibbons(ERHIFeatureLevel::Type FeatureL
 }
 
 
-void FNiagaraRendererRibbons::ReleaseRenderThreadResources()
+void FNiagaraRendererRibbons::ReleaseRenderThreadResources(NiagaraEmitterInstanceBatcher* Batcher)
 {
-	FNiagaraRenderer::ReleaseRenderThreadResources();
+	FNiagaraRenderer::ReleaseRenderThreadResources(Batcher);
 	VertexFactory->ReleaseResource();
 	WorldSpacePrimitiveUniformBuffer.ReleaseResource();
 }
 
 // FPrimitiveSceneProxy interface.
-void FNiagaraRendererRibbons::CreateRenderThreadResources()
+void FNiagaraRendererRibbons::CreateRenderThreadResources(NiagaraEmitterInstanceBatcher* Batcher)
 {
-	FNiagaraRenderer::CreateRenderThreadResources();
+	FNiagaraRenderer::CreateRenderThreadResources(Batcher);
 	VertexFactory->InitResource();
 }
 
@@ -301,7 +301,7 @@ void FNiagaraRendererRibbons::GetDynamicMeshElements(const TArray<const FSceneVi
 	int32 TotalFloatSize = SourceParticleData->GetFloatBuffer().Num() / sizeof(float);
 	FGlobalDynamicReadBuffer::FAllocation ParticleData;
 
-	if (DynamicDataRibbon->GetSimTarget() == ENiagaraSimTarget::CPUSim)
+	if (SimTarget == ENiagaraSimTarget::CPUSim)
 	{
 		FGlobalDynamicReadBuffer& DynamicReadBuffer = Collector.GetDynamicReadBuffer();
 		ParticleData = DynamicReadBuffer.AllocateFloat(TotalFloatSize);
@@ -569,7 +569,7 @@ FNiagaraDynamicDataBase* FNiagaraRendererRibbons::GenerateDynamicData(const FNia
 {
 	SCOPE_CYCLE_COUNTER(STAT_NiagaraGenRibbonVertexData);
 
-	if (Emitter->GetData().GetSimTarget() == ENiagaraSimTarget::GPUComputeSim)
+	if (SimTarget == ENiagaraSimTarget::GPUComputeSim)
 	{
 		return nullptr;
 	}
@@ -598,7 +598,6 @@ FNiagaraDynamicDataBase* FNiagaraRendererRibbons::GenerateDynamicData(const FNia
 	FNiagaraDataSetAccessor<int32> RibbonIdData;
 	FNiagaraDataSetAccessor<FNiagaraID> RibbonFullIDData;
 
-	ENiagaraSimTarget SimTarget = Data.GetSimTarget();
 	FNiagaraDataBuffer* DataToRender = Emitter->GetData().GetCurrentData();
 	if (DataToRender == nullptr || DataToRender->GetNumInstances() < 2 || !PosData.IsValid() || !SortKeyData.IsValid())
 	{
diff --git a/Engine/Plugins/FX/Niagara/Source/Niagara/Private/NiagaraRendererSprites.cpp b/Engine/Plugins/FX/Niagara/Source/Niagara/Private/NiagaraRendererSprites.cpp
--- a/Engine/Plugins/FX/Niagara/Source/Niagara/Private/NiagaraRendererSprites.cpp
+++ b/Engine/Plugins/FX/Niagara/Source/Niagara/Private/NiagaraRendererSprites.cpp
@@ -138,9 +138,9 @@ FNiagaraRendererSprites::FNiagaraRendererSprites(ERHIFeatureLevel::Type FeatureL
 	MaterialParamValidMask |= MaterialParamOffset3 != -1 ? 8 : 0;
 }
 
-void FNiagaraRendererSprites::ReleaseRenderThreadResources()
+void FNiagaraRendererSprites::ReleaseRenderThreadResources(NiagaraEmitterInstanceBatcher* Batcher)
 {
-	FNiagaraRenderer::ReleaseRenderThreadResources();
+	FNiagaraRenderer::ReleaseRenderThreadResources(Batcher);
 	VertexFactory->ReleaseResource();
 	WorldSpacePrimitiveUniformBuffer.ReleaseResource();
 
@@ -154,9 +154,9 @@ void FNiagaraRendererSprites::ReleaseRenderThreadResources()
 #endif
 }
 
-void FNiagaraRendererSprites::CreateRenderThreadResources()
+void FNiagaraRendererSprites::CreateRenderThreadResources(NiagaraEmitterInstanceBatcher* Batcher)
 {
-	FNiagaraRenderer::CreateRenderThreadResources();
+	FNiagaraRenderer::CreateRenderThreadResources(Batcher);
 	CutoutVertexBuffer.InitResource();
 	VertexFactory->InitResource();
 
@@ -218,7 +218,7 @@ FNiagaraRendererSprites::FCPUSimParticleDataAllocation FNiagaraRendererSprites::
 
 	FCPUSimParticleDataAllocation CPUSimParticleDataAllocation { DynamicReadBuffer };
 
-	if (DynamicDataSprites->GetSimTarget() == ENiagaraSimTarget::CPUSim)
+	if (SimTarget == ENiagaraSimTarget::CPUSim)
 	{
 		SCOPE_CYCLE_COUNTER(STAT_NiagaraRenderSpritesCPUSimCopy);
 		CPUSimParticleDataAllocation.ParticleData = DynamicReadBuffer.AllocateFloat(TotalFloatSize);
@@ -316,6 +316,9 @@ void FNiagaraRendererSprites::SetVertexFactoryParticleData(
 	const FSceneView* View,
 	const FNiagaraSceneProxy *SceneProxy) const
 {
+	NiagaraEmitterInstanceBatcher* Batcher = SceneProxy->GetBatcher();
+	check(Batcher);
+
 	// Cutout geometry.
 	const bool bUseSubImage = SubImageSize.X != 1 || SubImageSize.Y != 1;
 	const bool bUseCutout = CutoutVertexBuffer.VertexBufferRHI.IsValid();
@@ -362,24 +365,23 @@ void FNiagaraRendererSprites::SetVertexFactoryParticleData(
 			}
 		};
 
-		if (DynamicDataSprites->GetSimTarget() == ENiagaraSimTarget::CPUSim)//TODO: Compute shader for sorting gpu sims and larger cpu sims.
+		if (SimTarget == ENiagaraSimTarget::CPUSim)//TODO: Compute shader for sorting gpu sims and larger cpu sims.
 		{
 			check(CPUSimParticleDataAllocation.ParticleData.IsValid());
 			if (SortInfo.SortMode != ENiagaraSortMode::None && SortInfo.SortAttributeOffset != INDEX_NONE)
 			{
 				if (GNiagaraGPUSorting &&
 					GNiagaraGPUSortingCPUToGPUThreshold != INDEX_NONE &&
-					SortInfo.ParticleCount >= GNiagaraGPUSortingCPUToGPUThreshold &&
-					SceneProxy->GetBatcher())
+					SortInfo.ParticleCount >= GNiagaraGPUSortingCPUToGPUThreshold)
 				{
 					SortInfo.ParticleCount = NumInstances;
 					SortInfo.ParticleDataFloatSRV = CPUSimParticleDataAllocation.ParticleData.ReadBuffer->SRV;
 					SortInfo.FloatDataOffset = CPUSimParticleDataAllocation.ParticleData.FirstIndex / sizeof(float);
 					SortInfo.FloatDataStride = SourceParticleData->GetFloatStride() / sizeof(float);
-					const int32 IndexBufferOffset = SceneProxy->GetBatcher()->AddSortedGPUSimulation(SortInfo);
+					const int32 IndexBufferOffset = Batcher->AddSortedGPUSimulation(SortInfo);
 					if (IndexBufferOffset != INDEX_NONE)
 					{
-						OutVertexFactory.SetSortedIndices(SceneProxy->GetBatcher()->GetGPUSortedBuffer().VertexBufferSRV, IndexBufferOffset);
+						OutVertexFactory.SetSortedIndices(Batcher->GetGPUSortedBuffer().VertexBufferSRV, IndexBufferOffset);
 					}
 				}
 				else
@@ -394,9 +396,9 @@ void FNiagaraRendererSprites::SetVertexFactoryParticleData(
 			}
 			OutVertexFactory.SetParticleData(CPUSimParticleDataAllocation.ParticleData.ReadBuffer->SRV, CPUSimParticleDataAllocation.ParticleData.FirstIndex / sizeof(float), SourceParticleData->GetFloatStride() / sizeof(float));
 		}
-		else
+		else // ENiagaraSimTarget::GPUSim
 		{
-			if (SortInfo.SortMode != ENiagaraSortMode::None && SortInfo.SortAttributeOffset != INDEX_NONE && GNiagaraGPUSorting && SceneProxy->GetBatcher())
+			if (SortInfo.SortMode != ENiagaraSortMode::None && SortInfo.SortAttributeOffset != INDEX_NONE && GNiagaraGPUSorting)
 			{
 				// Here we need to be conservative about the InstanceCount, since the final value is only known on the GPU after the simulation.
 				SortInfo.ParticleCount = SourceParticleData->GetNumInstances();
@@ -404,12 +406,12 @@ void FNiagaraRendererSprites::SetVertexFactoryParticleData(
 				SortInfo.ParticleDataFloatSRV = SourceParticleData->GetGPUBufferFloat().SRV;
 				SortInfo.FloatDataOffset = 0;
 				SortInfo.FloatDataStride = SourceParticleData->GetFloatStride() / sizeof(float);
-				SortInfo.GPUParticleCountSRV = SourceParticleData->GetGPUIndices().SRV;
-				SortInfo.GPUParticleCountOffset = 1;
-				const int32 IndexBufferOffset = SceneProxy->GetBatcher()->AddSortedGPUSimulation(SortInfo);
-				if (IndexBufferOffset != INDEX_NONE)
+				SortInfo.GPUParticleCountSRV = Batcher->GetGPUInstanceCounterManager().GetInstanceCountBuffer().SRV;
+				SortInfo.GPUParticleCountOffset = SourceParticleData->GetGPUInstanceCountBufferOffset();
+				const int32 IndexBufferOffset = Batcher->AddSortedGPUSimulation(SortInfo);
+				if (IndexBufferOffset != INDEX_NONE && SortInfo.GPUParticleCountOffset != INDEX_NONE)
 				{
-					OutVertexFactory.SetSortedIndices(SceneProxy->GetBatcher()->GetGPUSortedBuffer().VertexBufferSRV, IndexBufferOffset);
+					OutVertexFactory.SetSortedIndices(Batcher->GetGPUSortedBuffer().VertexBufferSRV, IndexBufferOffset);
 				}
 			}
 
@@ -423,6 +425,7 @@ void FNiagaraRendererSprites::CreateMeshBatchForView(
 	const FSceneViewFamily& ViewFamily, 
 	const FNiagaraSceneProxy *SceneProxy,
 	FNiagaraDynamicDataSprites *DynamicDataSprites,
+	uint32 IndirectArgsOffset,
 	FMeshBatch& MeshBatch,
 	FNiagaraMeshCollectorResourcesSprite& CollectorResources) const
 {
@@ -463,13 +466,17 @@ void FNiagaraRendererSprites::CreateMeshBatchForView(
 	VFLooseParams.ParticleFacingMode = CollectorResources.VertexFactory.GetFacingMode();
 	VFLooseParams.SortedIndices = CollectorResources.VertexFactory.GetSortedIndicesSRV() ? CollectorResources.VertexFactory.GetSortedIndicesSRV().GetReference() : GFNiagaraNullSortedIndicesVertexBuffer.VertexBufferSRV.GetReference();
 	VFLooseParams.SortedIndicesOffset = CollectorResources.VertexFactory.GetSortedIndicesOffset();
-	if (DynamicDataSprites->GetSimTarget() == ENiagaraSimTarget::GPUComputeSim)
+	if (IndirectArgsOffset != INDEX_NONE)
 	{
-		VFLooseParams.IndirectArgsBuffer = SourceParticleData->GetGPUIndices().SRV;
+		NiagaraEmitterInstanceBatcher* Batcher = SceneProxy->GetBatcher();
+		check(Batcher); // Already verified at this point.
+		VFLooseParams.IndirectArgsOffset = IndirectArgsOffset / sizeof(uint32);
+		VFLooseParams.IndirectArgsBuffer = Batcher->GetGPUInstanceCounterManager().GetDrawIndirectBuffer().SRV;
 	}
 	else
 	{
 		VFLooseParams.IndirectArgsBuffer = GFNiagaraNullSortedIndicesVertexBuffer.VertexBufferSRV;
+		VFLooseParams.IndirectArgsOffset = 0;
 	}
 
 	CollectorResources.VertexFactory.LooseParameterUniformBuffer = FNiagaraSpriteVFLooseParametersRef::CreateUniformBufferImmediate(VFLooseParams, UniformBuffer_SingleFrame);
@@ -496,20 +503,22 @@ void FNiagaraRendererSprites::CreateMeshBatchForView(
 	FMeshBatchElement& MeshElement = MeshBatch.Elements[0];
 	MeshElement.IndexBuffer = &GParticleIndexBuffer;
 	MeshElement.FirstIndex = 0;
-	MeshElement.NumPrimitives = 2;
+	MeshElement.NumPrimitives = NumIndicesPerInstance / 3;
 	MeshElement.NumInstances = FMath::Max(0, NumInstances);	//->VertexData.Num();
 	MeshElement.MinVertexIndex = 0;
 	MeshElement.MaxVertexIndex = 0;// MeshElement.NumInstances * 4 - 1;
 	MeshElement.PrimitiveUniformBuffer = WorldSpacePrimitiveUniformBuffer.GetUniformBufferRHI();
-	if (DynamicDataSprites->GetSimTarget() == ENiagaraSimTarget::GPUComputeSim)
+	if (IndirectArgsOffset != INDEX_NONE)
 	{
-		MeshElement.IndirectArgsBuffer = SourceParticleData->GetGPUIndices().Buffer;
+		NiagaraEmitterInstanceBatcher* Batcher = SceneProxy->GetBatcher();
+		check(Batcher); // Already verified at this point.
+		MeshElement.IndirectArgsOffset = IndirectArgsOffset;
+		MeshElement.IndirectArgsBuffer = Batcher->GetGPUInstanceCounterManager().GetDrawIndirectBuffer().Buffer;
 		MeshElement.NumPrimitives = 0;
 	}
 
 	if (NumCutoutVertexPerSubImage == 8)
 	{
-		MeshElement.NumPrimitives = 6;
 		MeshElement.IndexBuffer = &GSixTriangleParticleIndexBuffer;
 	}
 	
@@ -520,12 +529,14 @@ void FNiagaraRendererSprites::GetDynamicMeshElements(const TArray<const FSceneVi
 {
 	SCOPE_CYCLE_COUNTER(STAT_NiagaraRender);
 	SCOPE_CYCLE_COUNTER(STAT_NiagaraRenderSprites);
+	check(SceneProxy);
 
 	SimpleTimer MeshElementsTimer;
 
 	//check(DynamicDataRender)
 	FNiagaraDynamicDataSprites *DynamicDataSprites = static_cast<FNiagaraDynamicDataSprites*>(DynamicDataRender);
-	if (DynamicDataSprites == nullptr)
+	NiagaraEmitterInstanceBatcher* Batcher = SceneProxy->GetBatcher();
+	if (!DynamicDataSprites || !Batcher)
 	{
 		return;
 	}
@@ -548,6 +559,12 @@ void FNiagaraRendererSprites::GetDynamicMeshElements(const TArray<const FSceneVi
 
 	ConditionalInitPrimitiveUniformBuffer(SceneProxy);
 
+	uint32 IndirectArgsOffset = INDEX_NONE;
+	if (SimTarget == ENiagaraSimTarget::GPUComputeSim)
+	{
+		IndirectArgsOffset = Batcher->GetGPUInstanceCounterManager().AddDrawIndirect(SourceParticleData->GetGPUInstanceCountBufferOffset(), NumIndicesPerInstance);
+	}
+
 	for (int32 ViewIndex = 0; ViewIndex < Views.Num(); ViewIndex++)
 	{
 		if (VisibilityMap & (1 << ViewIndex))
@@ -558,8 +575,7 @@ void FNiagaraRendererSprites::GetDynamicMeshElements(const TArray<const FSceneVi
 			SetVertexFactoryParticleData(CollectorResources.VertexFactory, DynamicDataSprites, CPUSimParticleDataAllocation, View, SceneProxy);
 			CollectorResources.UniformBuffer = CreatePerViewUniformBuffer(View, ViewFamily, SceneProxy);
 			FMeshBatch& MeshBatch = Collector.AllocateMesh();
-			CreateMeshBatchForView(View, ViewFamily, SceneProxy, DynamicDataSprites, MeshBatch, CollectorResources);
-
+			CreateMeshBatchForView(View, ViewFamily, SceneProxy, DynamicDataSprites, IndirectArgsOffset, MeshBatch, CollectorResources);
 			Collector.AddMesh(ViewIndex, MeshBatch);
 		}
 	}
@@ -572,11 +588,13 @@ void FNiagaraRendererSprites::GetDynamicRayTracingInstances(FRayTracingMaterialG
 {
 	SCOPE_CYCLE_COUNTER(STAT_NiagaraRender);
 	SCOPE_CYCLE_COUNTER(STAT_NiagaraRenderSprites);
+	check(SceneProxy);
 
 	SimpleTimer MeshElementsTimer;
 
 	FNiagaraDynamicDataSprites *DynamicDataSprites = static_cast<FNiagaraDynamicDataSprites*>(DynamicDataRender);
-	if (DynamicDataSprites == nullptr)
+	NiagaraEmitterInstanceBatcher* Batcher = SceneProxy->GetBatcher();
+	if (!DynamicDataSprites || !Batcher)
 	{
 		return;
 	}
@@ -592,6 +610,12 @@ void FNiagaraRendererSprites::GetDynamicRayTracingInstances(FRayTracingMaterialG
 		return;
 	}
 
+	uint32 IndirectArgsOffset = INDEX_NONE;
+	if (SimTarget == ENiagaraSimTarget::GPUComputeSim)
+	{
+		IndirectArgsOffset = Batcher->GetGPUInstanceCounterManager().AddDrawIndirect(SourceParticleData->GetGPUInstanceCountBufferOffset(), NumIndicesPerInstance);
+	}
+
 	FRayTracingInstance RayTracingInstance;
 	RayTracingInstance.Geometry = &RayTracingGeometry;
 	RayTracingInstance.InstanceTransforms.Add(FMatrix::Identity);
@@ -604,7 +628,7 @@ void FNiagaraRendererSprites::GetDynamicRayTracingInstances(FRayTracingMaterialG
 		SetVertexFactoryParticleData(CollectorResources.VertexFactory, DynamicDataSprites, CPUSimParticleDataAllocation, Context.ReferenceView, SceneProxy);
 		CollectorResources.UniformBuffer = CreatePerViewUniformBuffer(Context.ReferenceView, Context.ReferenceViewFamily, SceneProxy);
 		FMeshBatch MeshBatch;
-		CreateMeshBatchForView(Context.ReferenceView, Context.ReferenceViewFamily, SceneProxy, DynamicDataSprites, MeshBatch, CollectorResources);
+		CreateMeshBatchForView(Context.ReferenceView, Context.ReferenceViewFamily, SceneProxy, DynamicDataSprites, IndirectArgsOffset, MeshBatch, CollectorResources);
 
 		ensureMsgf(MeshBatch.Elements[0].IndexBuffer != &GSixTriangleParticleIndexBuffer, TEXT("Cutout geometry is not supported by ray tracing"));
 
@@ -642,7 +666,6 @@ FNiagaraDynamicDataBase *FNiagaraRendererSprites::GenerateDynamicData(const FNia
 		SCOPE_CYCLE_COUNTER(STAT_NiagaraGenSpriteDynamicData);
 
 		FNiagaraDataSet& Data = Emitter->GetData();
-		ENiagaraSimTarget SimTarget = Data.GetSimTarget();
 		if(SimTarget == ENiagaraSimTarget::GPUComputeSim || Data.GetCurrentDataChecked().GetNumInstances() > 0)
 		{
 			DynamicData = new FNiagaraDynamicDataSprites(Emitter);
diff --git a/Engine/Plugins/FX/Niagara/Source/Niagara/Private/NiagaraScriptExecutionContext.cpp b/Engine/Plugins/FX/Niagara/Source/Niagara/Private/NiagaraScriptExecutionContext.cpp
--- a/Engine/Plugins/FX/Niagara/Source/Niagara/Private/NiagaraScriptExecutionContext.cpp
+++ b/Engine/Plugins/FX/Niagara/Source/Niagara/Private/NiagaraScriptExecutionContext.cpp
@@ -8,6 +8,7 @@
 #include "NiagaraWorldManager.h"
 #include "NiagaraSystemInstance.h"
 #include "NiagaraEmitterInstance.h"
+#include "NiagaraGPUInstanceCountManager.h"
 
 DECLARE_CYCLE_STAT(TEXT("Register Setup"), STAT_NiagaraSimRegisterSetup, STATGROUP_Niagara);
 DECLARE_CYCLE_STAT(TEXT("Context Ticking"), STAT_NiagaraScriptExecContextTick, STATGROUP_Niagara);
@@ -395,9 +396,6 @@ FNiagaraComputeExecutionContext::FNiagaraComputeExecutionContext()
 	, GPUScript_RT(nullptr)
 	, CBufferLayout(TEXT("Niagara Compute Sim CBuffer"))
 	, DataToRender(nullptr)
-	, GPUDataReadback(nullptr)
-	, AccumulatedSpawnRate(0)
-	, NumIndicesPerInstance(0)
 #if WITH_EDITORONLY_DATA
 	, GPUDebugDataReadbackFloat(nullptr)
 	, GPUDebugDataReadbackInt(nullptr)
@@ -406,6 +404,7 @@ FNiagaraComputeExecutionContext::FNiagaraComputeExecutionContext()
 	, GPUDebugDataIntSize(0)
 	, GPUDebugDataFloatStride(0)
 	, GPUDebugDataIntStride(0)
+	, GPUDebugDataCountOffset(INDEX_NONE)
 #endif	  
 {
 }
@@ -413,11 +412,7 @@ FNiagaraComputeExecutionContext::FNiagaraComputeExecutionContext()
 FNiagaraComputeExecutionContext::~FNiagaraComputeExecutionContext()
 {
 	checkf(IsInRenderingThread(), TEXT("Can only delete the gpu readback from the render thread"));
-	if (GPUDataReadback)
-	{
-		delete GPUDataReadback;
-		GPUDataReadback = nullptr;
-	}
+	check(EmitterInstanceReadback.GPUCountOffset == INDEX_NONE);
 
 #if WITH_EDITORONLY_DATA
 	if (GPUDebugDataReadbackFloat)
@@ -440,13 +435,14 @@ FNiagaraComputeExecutionContext::~FNiagaraComputeExecutionContext()
 	SetDataToRender(nullptr);
 }
 
-void FNiagaraComputeExecutionContext::Reset()
+void FNiagaraComputeExecutionContext::Reset(NiagaraEmitterInstanceBatcher* Batcher)
 {
 	FNiagaraComputeExecutionContext* Context = this;
+	NiagaraEmitterInstanceBatcher* B = Batcher && !Batcher->IsPendingKill() ? Batcher : nullptr;
 	ENQUEUE_RENDER_COMMAND(ResetRT)(
-		[Context](FRHICommandListImmediate& RHICmdList)
+		[B, Context](FRHICommandListImmediate& RHICmdList)
 	{
-		Context->ResetInternal();
+		Context->ResetInternal(B);
 	}
 	);
 }
@@ -518,14 +514,18 @@ void FNiagaraComputeExecutionContext::PostTick()
 	}
 }
 
-void FNiagaraComputeExecutionContext::ResetInternal()
+void FNiagaraComputeExecutionContext::ResetInternal(NiagaraEmitterInstanceBatcher* Batcher)
 {
 	checkf(IsInRenderingThread(), TEXT("Can only reset the gpu context from the render thread"));
-	AccumulatedSpawnRate = 0;
-	if (GPUDataReadback)
+
+	// Release and reset readback data.
+	if (Batcher)
+	{
+		Batcher->GetGPUInstanceCounterManager().FreeEntry(EmitterInstanceReadback.GPUCountOffset);
+	}
+	else // In this case the batcher is pending kill so no need to putback entry in the pool.
 	{
-		delete GPUDataReadback;
-		GPUDataReadback = nullptr;
+		EmitterInstanceReadback.GPUCountOffset = INDEX_NONE;
 	}
 
 #if WITH_EDITORONLY_DATA
diff --git a/Engine/Plugins/FX/Niagara/Source/Niagara/Private/NiagaraSpriteRendererProperties.cpp b/Engine/Plugins/FX/Niagara/Source/Niagara/Private/NiagaraSpriteRendererProperties.cpp
--- a/Engine/Plugins/FX/Niagara/Source/Niagara/Private/NiagaraSpriteRendererProperties.cpp
+++ b/Engine/Plugins/FX/Niagara/Source/Niagara/Private/NiagaraSpriteRendererProperties.cpp
@@ -295,7 +295,7 @@ int32 UNiagaraSpriteRendererProperties::GetNumCutoutVertexPerSubimage() const
 	}
 }
 
-uint32 UNiagaraSpriteRendererProperties::GetNumIndicesPerInstance()
+uint32 UNiagaraSpriteRendererProperties::GetNumIndicesPerInstance() const
 {
 	// This is a based on cutout vertices making a triangle strip.
 	if (GetNumCutoutVertexPerSubimage() == 8)
diff --git a/Engine/Plugins/FX/Niagara/Source/Niagara/Public/NiagaraMeshRendererProperties.h b/Engine/Plugins/FX/Niagara/Source/Niagara/Public/NiagaraMeshRendererProperties.h
--- a/Engine/Plugins/FX/Niagara/Source/Niagara/Public/NiagaraMeshRendererProperties.h
+++ b/Engine/Plugins/FX/Niagara/Source/Niagara/Public/NiagaraMeshRendererProperties.h
@@ -55,7 +55,7 @@ public:
 	virtual const TArray<FNiagaraVariable>& GetOptionalAttributes() override;
 #endif // WITH_EDITORONLY_DATA
 
-	virtual uint32 GetNumIndicesPerInstance() final override;
+	virtual uint32 GetNumIndicesPerInstance() const final override;
 
 	/** The static mesh to be instanced when rendering mesh particles. If OverrideMaterial is not specified, the mesh's material is used. Note that the material must have the Niagara Mesh Particles flag checked.*/
 	UPROPERTY(EditAnywhere, Category = "Mesh Rendering")
diff --git a/Engine/Plugins/FX/Niagara/Source/Niagara/Public/NiagaraRenderer.h b/Engine/Plugins/FX/Niagara/Source/Niagara/Public/NiagaraRenderer.h
--- a/Engine/Plugins/FX/Niagara/Source/Niagara/Public/NiagaraRenderer.h
+++ b/Engine/Plugins/FX/Niagara/Source/Niagara/Public/NiagaraRenderer.h
@@ -90,8 +90,8 @@ public:
 	FNiagaraRenderer& operator=(const FNiagaraRenderer& Other) = delete;
 
 	virtual void Initialize(ERHIFeatureLevel::Type FeatureLevel, const UNiagaraRendererProperties *InProps, const FNiagaraEmitterInstance* Emitter);
-	virtual void CreateRenderThreadResources() {}
-	virtual void ReleaseRenderThreadResources() {}
+	virtual void CreateRenderThreadResources(NiagaraEmitterInstanceBatcher* Batcher);
+	virtual void ReleaseRenderThreadResources(NiagaraEmitterInstanceBatcher* Batcher);
 
 	virtual FPrimitiveViewRelevance GetViewRelevance(const FSceneView* View, const FNiagaraSceneProxy *SceneProxy)const;
 	virtual void GetDynamicMeshElements(const TArray<const FSceneView*>& Views, const FSceneViewFamily& ViewFamily, uint32 VisibilityMap, FMeshElementCollector& Collector, const FNiagaraSceneProxy *SceneProxy) const {}
@@ -132,6 +132,8 @@ protected:
 
 	uint32 bLocalSpace : 1;
 	uint32 bHasLights : 1;
+	ENiagaraSimTarget SimTarget;
+	uint32 NumIndicesPerInstance;
 
 #if STATS
 	TStatId EmitterStatID;
@@ -140,5 +142,6 @@ protected:
 	/** Cached array of materials used from the properties data. Validated with usage flags etc. */
 	TArray<UMaterialInterface*> BaseMaterials_GT;
 	FMaterialRelevance BaseMaterialRelevance_GT;
+
 };
 
diff --git a/Engine/Plugins/FX/Niagara/Source/Niagara/Public/NiagaraRendererMeshes.h b/Engine/Plugins/FX/Niagara/Source/Niagara/Public/NiagaraRendererMeshes.h
--- a/Engine/Plugins/FX/Niagara/Source/Niagara/Public/NiagaraRendererMeshes.h
+++ b/Engine/Plugins/FX/Niagara/Source/Niagara/Public/NiagaraRendererMeshes.h
@@ -20,8 +20,8 @@ public:
 	FNiagaraRendererMeshes(ERHIFeatureLevel::Type FeatureLevel, const UNiagaraRendererProperties *InProps, const FNiagaraEmitterInstance* Emitter);
 	
 	//FNiagaraRenderer Interface
-	virtual void CreateRenderThreadResources() override;
-	virtual void ReleaseRenderThreadResources() override;
+	virtual void CreateRenderThreadResources(NiagaraEmitterInstanceBatcher* Batcher) override;
+	virtual void ReleaseRenderThreadResources(NiagaraEmitterInstanceBatcher* Batcher) override;
 
 	virtual void GetDynamicMeshElements(const TArray<const FSceneView*>& Views, const FSceneViewFamily& ViewFamily, uint32 VisibilityMap, FMeshElementCollector& Collector, const FNiagaraSceneProxy *SceneProxy) const override;
 	virtual FNiagaraDynamicDataBase* GenerateDynamicData(const FNiagaraSceneProxy* Proxy, const UNiagaraRendererProperties* InProperties, const FNiagaraEmitterInstance* Emitter) const override;
diff --git a/Engine/Plugins/FX/Niagara/Source/Niagara/Public/NiagaraRendererProperties.h b/Engine/Plugins/FX/Niagara/Source/Niagara/Public/NiagaraRendererProperties.h
--- a/Engine/Plugins/FX/Niagara/Source/Niagara/Public/NiagaraRendererProperties.h
+++ b/Engine/Plugins/FX/Niagara/Source/Niagara/Public/NiagaraRendererProperties.h
@@ -53,7 +53,7 @@ public:
 
 
 	// GPU simulation uses DrawIndirect, so the sim step needs to know indices per instance in order to prepare the draw call parameters
-	virtual uint32 GetNumIndicesPerInstance() { return 0; }
+	virtual uint32 GetNumIndicesPerInstance() const { return 0; }
 
 	virtual bool GetIsEnabled() const { return bIsEnabled; }
 	virtual void SetIsEnabled(bool bInIsEnabled) { bIsEnabled = bInIsEnabled; }
diff --git a/Engine/Plugins/FX/Niagara/Source/Niagara/Public/NiagaraRendererRibbons.h b/Engine/Plugins/FX/Niagara/Source/Niagara/Public/NiagaraRendererRibbons.h
--- a/Engine/Plugins/FX/Niagara/Source/Niagara/Public/NiagaraRendererRibbons.h
+++ b/Engine/Plugins/FX/Niagara/Source/Niagara/Public/NiagaraRendererRibbons.h
@@ -20,8 +20,8 @@ public:
 	FNiagaraRendererRibbons(ERHIFeatureLevel::Type FeatureLevel, const UNiagaraRendererProperties *InProps, const FNiagaraEmitterInstance* Emitter);	// FNiagaraRenderer Interface 
 
 	// FNiagaraRenderer Interface 
-	virtual void CreateRenderThreadResources() override;
-	virtual void ReleaseRenderThreadResources() override;
+	virtual void CreateRenderThreadResources(NiagaraEmitterInstanceBatcher* Batcher) override;
+	virtual void ReleaseRenderThreadResources(NiagaraEmitterInstanceBatcher* Batcher) override;
 
 	virtual void GetDynamicMeshElements(const TArray<const FSceneView*>& Views, const FSceneViewFamily& ViewFamily, uint32 VisibilityMap, FMeshElementCollector& Collector, const FNiagaraSceneProxy *SceneProxy) const override;
 	virtual FNiagaraDynamicDataBase *GenerateDynamicData(const FNiagaraSceneProxy* Proxy, const UNiagaraRendererProperties* InProperties, const FNiagaraEmitterInstance* Emitteride) const override;
diff --git a/Engine/Plugins/FX/Niagara/Source/Niagara/Public/NiagaraRendererSprites.h b/Engine/Plugins/FX/Niagara/Source/Niagara/Public/NiagaraRendererSprites.h
--- a/Engine/Plugins/FX/Niagara/Source/Niagara/Public/NiagaraRendererSprites.h
+++ b/Engine/Plugins/FX/Niagara/Source/Niagara/Public/NiagaraRendererSprites.h
@@ -20,8 +20,8 @@ public:
 	FNiagaraRendererSprites(ERHIFeatureLevel::Type FeatureLevel, const UNiagaraRendererProperties *InProps, const FNiagaraEmitterInstance* Emitter);
 
 	//FNiagaraRenderer interface
-	virtual void CreateRenderThreadResources()override;
-	virtual void ReleaseRenderThreadResources()override;
+	virtual void CreateRenderThreadResources(NiagaraEmitterInstanceBatcher* Batcher)override;
+	virtual void ReleaseRenderThreadResources(NiagaraEmitterInstanceBatcher* Batcher)override;
 
 	virtual void GetDynamicMeshElements(const TArray<const FSceneView*>& Views, const FSceneViewFamily& ViewFamily, uint32 VisibilityMap, FMeshElementCollector& Collector, const FNiagaraSceneProxy *SceneProxy) const override;
 	virtual FNiagaraDynamicDataBase* GenerateDynamicData(const FNiagaraSceneProxy* Proxy, const UNiagaraRendererProperties* InProperties, const FNiagaraEmitterInstance* Emitter) const override;
@@ -55,6 +55,7 @@ private:
 		const FSceneViewFamily& ViewFamily,
 		const FNiagaraSceneProxy *SceneProxy,
 		FNiagaraDynamicDataSprites *DynamicDataSprites,
+		uint32 IndirectArgsOffset,
 		FMeshBatch& OutMeshBatch,
 		class FNiagaraMeshCollectorResourcesSprite& OutCollectorResources) const;
 
diff --git a/Engine/Plugins/FX/Niagara/Source/Niagara/Public/NiagaraSpriteRendererProperties.h b/Engine/Plugins/FX/Niagara/Source/Niagara/Public/NiagaraSpriteRendererProperties.h
--- a/Engine/Plugins/FX/Niagara/Source/Niagara/Public/NiagaraSpriteRendererProperties.h
+++ b/Engine/Plugins/FX/Niagara/Source/Niagara/Public/NiagaraSpriteRendererProperties.h
@@ -72,7 +72,7 @@ public:
 	//UNiagaraMaterialRendererProperties interface END
 
 	int32 GetNumCutoutVertexPerSubimage() const;
-	virtual uint32 GetNumIndicesPerInstance();
+	virtual uint32 GetNumIndicesPerInstance() const;
 
 	/** The material used to render the particle. Note that it must have the Use with Niagara Sprites flag checked.*/
 	UPROPERTY(EditAnywhere, Category = "Sprite Rendering")
diff --git a/Engine/Plugins/FX/Niagara/Source/NiagaraCore/Private/NiagaraCustomVersion.cpp b/Engine/Plugins/FX/Niagara/Source/NiagaraCore/Private/NiagaraCustomVersion.cpp
--- a/Engine/Plugins/FX/Niagara/Source/NiagaraCore/Private/NiagaraCustomVersion.cpp
+++ b/Engine/Plugins/FX/Niagara/Source/NiagaraCore/Private/NiagaraCustomVersion.cpp
@@ -8,4 +8,4 @@ const FGuid FNiagaraCustomVersion::GUID(0xFCF57AFA, 0x50764283, 0xB9A9E658, 0xFF
 // Register the custom version with core
 FCustomVersionRegistration GRegisterNiagaraCustomVersion(FNiagaraCustomVersion::GUID, FNiagaraCustomVersion::LatestVersion, TEXT("NiagaraVer"));
 
-const FGuid FNiagaraCustomVersion::LatestScriptCompileVersion(0x8D38659E, 0xFF5CC945, 0xA718F14F, 0x52E92EBA);
+const FGuid FNiagaraCustomVersion::LatestScriptCompileVersion(0x397D4E10, 0x87CC4DF8, 0x9EDAD0CA, 0xE837CA88);
\ No newline at end of file
diff --git a/Engine/Plugins/FX/Niagara/Source/NiagaraShader/Private/NiagaraShader.cpp b/Engine/Plugins/FX/Niagara/Source/NiagaraShader/Private/NiagaraShader.cpp
--- a/Engine/Plugins/FX/Niagara/Source/NiagaraShader/Private/NiagaraShader.cpp
+++ b/Engine/Plugins/FX/Niagara/Source/NiagaraShader/Private/NiagaraShader.cpp
@@ -1275,8 +1275,11 @@ void FNiagaraShader::BindParams(const FShaderParameterMap &ParameterMap)
 	IntInputBufferParam.Bind(ParameterMap, TEXT("InputInt"));
 	FloatOutputBufferParam.Bind(ParameterMap, TEXT("OutputFloat"));
 	IntOutputBufferParam.Bind(ParameterMap, TEXT("OutputInt"));
-	OutputIndexBufferParam.Bind(ParameterMap, TEXT("DataSetIndices"));
-	InputIndexBufferParam.Bind(ParameterMap, TEXT("ReadDataSetIndices"));
+
+	InstanceCountsParam.Bind(ParameterMap, TEXT("InstanceCounts"));
+	ReadInstanceCountOffsetParam.Bind(ParameterMap, TEXT("ReadInstanceCountOffset"));
+	WriteInstanceCountOffsetParam.Bind(ParameterMap, TEXT("WriteInstanceCountOffset"));
+
 	EmitterTickCounterParam.Bind(ParameterMap, TEXT("EmitterTickCounter"));
 	NumEventsPerParticleParam.Bind(ParameterMap, TEXT("NumEventsPerParticle"));
 	NumParticlesPerEventParam.Bind(ParameterMap, TEXT("NumParticlesPerEvent"));
@@ -1284,7 +1287,6 @@ void FNiagaraShader::BindParams(const FShaderParameterMap &ParameterMap)
 
 	NumSpawnedInstancesParam.Bind(ParameterMap, TEXT("SpawnedInstances"));
 	UpdateStartInstanceParam.Bind(ParameterMap, TEXT("UpdateStartInstance"));
-	NumIndicesPerInstanceParam.Bind(ParameterMap, TEXT("NumIndicesPerInstance"));
 
 	ComponentBufferSizeReadParam.Bind(ParameterMap, TEXT("ComponentBufferSizeRead"));
 	ComponentBufferSizeWriteParam.Bind(ParameterMap, TEXT("ComponentBufferSizeWrite"));
@@ -1323,11 +1325,9 @@ void FNiagaraShader::BindParams(const FShaderParameterMap &ParameterMap)
 
 	ensure(FloatOutputBufferParam.IsBound() || IntOutputBufferParam.IsBound());	// we should have at least one output buffer we're writing to
 	ensure(ComponentBufferSizeWriteParam.IsBound());
-	ensure(OutputIndexBufferParam.IsBound());
-	ensure(InputIndexBufferParam.IsBound());
+	ensure(InstanceCountsParam.IsBound());
 	ensure(UpdateStartInstanceParam.IsBound());
 	ensure(NumSpawnedInstancesParam.IsBound());
-	ensure(NumIndicesPerInstanceParam.IsBound());
 }
 
 bool FNiagaraShader::Serialize(FArchive& Ar)
@@ -1341,7 +1341,11 @@ bool FNiagaraShader::Serialize(FArchive& Ar)
 	Ar << IntInputBufferParam;
 	Ar << FloatOutputBufferParam;
 	Ar << IntOutputBufferParam;
-	Ar << InputIndexBufferParam;
+
+	Ar << InstanceCountsParam;
+	Ar << ReadInstanceCountOffsetParam;
+	Ar << WriteInstanceCountOffsetParam;
+
 	Ar << EmitterTickCounterParam;
 
 	Ar << NumSpawnedInstancesParam;
@@ -1366,8 +1370,6 @@ bool FNiagaraShader::Serialize(FArchive& Ar)
 	Ar << CopyInstancesBeforeStartParam;
 	Ar << ViewUniformBufferParam;
 
-	Ar << OutputIndexBufferParam;
-	Ar << NumIndicesPerInstanceParam;
 	return bShaderHasOutdatedParameters;
 }
 
diff --git a/Engine/Plugins/FX/Niagara/Source/NiagaraShader/Public/NiagaraShader.h b/Engine/Plugins/FX/Niagara/Source/NiagaraShader/Public/NiagaraShader.h
--- a/Engine/Plugins/FX/Niagara/Source/NiagaraShader/Public/NiagaraShader.h
+++ b/Engine/Plugins/FX/Niagara/Source/NiagaraShader/Public/NiagaraShader.h
@@ -102,8 +102,9 @@ public:
 	FShaderResourceParameter IntInputBufferParam;
 	FRWShaderParameter FloatOutputBufferParam;
 	FRWShaderParameter IntOutputBufferParam;
-	FRWShaderParameter OutputIndexBufferParam;
-	FShaderResourceParameter InputIndexBufferParam;
+	FRWShaderParameter InstanceCountsParam;
+	FShaderParameter ReadInstanceCountOffsetParam;
+	FShaderParameter WriteInstanceCountOffsetParam;
 	FShaderUniformBufferParameter EmitterConstantBufferParam;
 	FShaderUniformBufferParameter DataInterfaceUniformBufferParam;
 	FShaderUniformBufferParameter ViewUniformBufferParam;
@@ -113,7 +114,6 @@ public:
 	FShaderParameter CopyInstancesBeforeStartParam;
 	FShaderParameter NumSpawnedInstancesParam;
 	FShaderParameter UpdateStartInstanceParam;
-	FShaderParameter NumIndicesPerInstanceParam;
 	FShaderParameter ComponentBufferSizeReadParam;
 	FShaderParameter ComponentBufferSizeWriteParam;
 	FRWShaderParameter EventIntUAVParams[MAX_CONCURRENT_EVENT_DATASETS];
diff --git a/Engine/Plugins/FX/Niagara/Source/NiagaraVertexFactories/Private/NiagaraDrawIndirect.cpp b/Engine/Plugins/FX/Niagara/Source/NiagaraVertexFactories/Private/NiagaraDrawIndirect.cpp
new file mode 100644
--- /dev/null
+++ b/Engine/Plugins/FX/Niagara/Source/NiagaraVertexFactories/Private/NiagaraDrawIndirect.cpp
@@ -0,0 +1,78 @@
+// Copyright 1998-2019 Epic Games, Inc. All Rights Reserved.
+
+/*==============================================================================
+NiagaraDrawIndirect.cpp : Niagara shader to generate the draw indirect args for Niagara renderers.
+==============================================================================*/
+
+#include "NiagaraDrawIndirect.h"
+#include "NiagaraGPUSortInfo.h"
+
+
+IMPLEMENT_GLOBAL_SHADER(FNiagaraDrawIndirectArgsGenCS, "/Plugin/FX/Niagara/Private/NiagaraDrawIndirectArgsGen.usf", "MainCS", SF_Compute);
+
+void FNiagaraDrawIndirectArgsGenCS::ModifyCompilationEnvironment(const FGlobalShaderPermutationParameters& Parameters, FShaderCompilerEnvironment& OutEnvironment)
+{
+	FGlobalShader::ModifyCompilationEnvironment(Parameters, OutEnvironment);
+	OutEnvironment.SetDefine(TEXT("THREAD_COUNT"), NIAGARA_DRAW_INDIRECT_ARGS_GEN_THREAD_COUNT);
+	OutEnvironment.SetDefine(TEXT("NIAGARA_DRAW_INDIRECT_ARGS_SIZE"), NIAGARA_DRAW_INDIRECT_ARGS_SIZE);
+}
+
+FNiagaraDrawIndirectArgsGenCS::FNiagaraDrawIndirectArgsGenCS(const ShaderMetaType::CompiledShaderInitializerType& Initializer)
+	: FGlobalShader(Initializer)
+{
+	TaskInfosParam.Bind(Initializer.ParameterMap, TEXT("TaskInfos"));
+	InstanceCountsParam.Bind(Initializer.ParameterMap, TEXT("InstanceCounts"));
+	DrawIndirectArgsParam.Bind(Initializer.ParameterMap, TEXT("DrawIndirectArgs"));
+	TaskCountParam.Bind(Initializer.ParameterMap, TEXT("TaskCount"));
+}
+
+bool FNiagaraDrawIndirectArgsGenCS::Serialize(FArchive& Ar)
+{
+	bool bShaderHasOutdatedParameters = FGlobalShader::Serialize(Ar);
+	Ar << TaskInfosParam;
+	Ar << InstanceCountsParam;
+	Ar << DrawIndirectArgsParam;
+	Ar << TaskCountParam;
+	return bShaderHasOutdatedParameters;
+}
+
+void FNiagaraDrawIndirectArgsGenCS::SetOutput(FRHICommandList& RHICmdList, FUnorderedAccessViewRHIParamRef DrawIndirectArgsUAV, FUnorderedAccessViewRHIParamRef InstanceCountsUAV)
+{
+	FComputeShaderRHIParamRef ComputeShaderRHI = GetComputeShader();
+	if (DrawIndirectArgsParam.IsBound())
+	{
+		RHICmdList.SetUAVParameter(ComputeShaderRHI, DrawIndirectArgsParam.GetUAVIndex(), DrawIndirectArgsUAV);
+	}
+	if (InstanceCountsParam.IsBound())
+	{
+		RHICmdList.SetUAVParameter(ComputeShaderRHI, InstanceCountsParam.GetUAVIndex(), InstanceCountsUAV);
+	}
+}
+
+void FNiagaraDrawIndirectArgsGenCS::SetParameters(FRHICommandList& RHICmdList, FShaderResourceViewRHIParamRef TaskInfosBuffer, int32 NumArgGenTasks, int32 NumInstanceCountClearTasks)
+{
+	FComputeShaderRHIParamRef ComputeShaderRHI = GetComputeShader();
+
+	RHICmdList.SetShaderResourceViewParameter(ComputeShaderRHI, TaskInfosParam.GetBaseIndex(), TaskInfosBuffer);
+
+	const FUintVector4 TaskCountValue((int32)NumArgGenTasks, (int32)NumInstanceCountClearTasks, (int32)(NumArgGenTasks + NumInstanceCountClearTasks), 0);
+	RHICmdList.SetShaderParameter(ComputeShaderRHI, TaskCountParam.GetBufferIndex(), TaskCountParam.GetBaseIndex(), TaskCountParam.GetNumBytes(), &TaskCountValue);
+}
+
+void FNiagaraDrawIndirectArgsGenCS::UnbindBuffers(FRHICommandList& RHICmdList)
+{
+	FComputeShaderRHIParamRef ComputeShaderRHI = GetComputeShader();
+	if (TaskInfosParam.IsBound())
+	{
+		RHICmdList.SetShaderResourceViewParameter(ComputeShaderRHI, TaskInfosParam.GetBaseIndex(), FShaderResourceViewRHIParamRef());
+	}
+	if (DrawIndirectArgsParam.IsBound())
+	{
+		RHICmdList.SetUAVParameter(ComputeShaderRHI, DrawIndirectArgsParam.GetUAVIndex(), FUnorderedAccessViewRHIParamRef());
+	}
+	if (InstanceCountsParam.IsBound())
+	{
+		RHICmdList.SetUAVParameter(ComputeShaderRHI, InstanceCountsParam.GetUAVIndex(), FUnorderedAccessViewRHIParamRef());
+	}
+}
+
diff --git a/Engine/Plugins/FX/Niagara/Source/NiagaraVertexFactories/Public/NiagaraDrawIndirect.h b/Engine/Plugins/FX/Niagara/Source/NiagaraVertexFactories/Public/NiagaraDrawIndirect.h
new file mode 100644
--- /dev/null
+++ b/Engine/Plugins/FX/Niagara/Source/NiagaraVertexFactories/Public/NiagaraDrawIndirect.h
@@ -0,0 +1,68 @@
+// Copyright 1998-2019 Epic Games, Inc. All Rights Reserved.
+
+/*==============================================================================
+NiagaraDrawIndirect.h: Niagara shader to generate the draw indirect args for Niagara renderers.
+==============================================================================*/
+
+#pragma once
+
+#include "CoreMinimal.h"
+#include "GlobalShader.h"
+#include "ShaderPermutation.h"
+
+#define NIAGARA_DRAW_INDIRECT_ARGS_GEN_THREAD_COUNT 64
+#define NIAGARA_DRAW_INDIRECT_ARGS_SIZE 5
+
+// #define NIAGARA_COPY_BUFFER_THREAD_COUNT 64
+// #define NIAGARA_COPY_BUFFER_BUFFER_COUNT 3
+
+/**
+ * Compute shader used to generate GPU emitter draw indirect args.
+ * It also resets unused instance count entries.
+ */
+class NIAGARAVERTEXFACTORIES_API FNiagaraDrawIndirectArgsGenCS : public FGlobalShader
+{
+	DECLARE_GLOBAL_SHADER(FNiagaraDrawIndirectArgsGenCS);
+
+public:
+
+	/** 
+	 * Task info when generating draw indirect frame buffer.
+	 * Task is either about generate Niagara renderers drawindirect buffer,
+	 * or about resetting released instance counters.
+	 */
+	struct FArgGenTaskInfo
+	{
+		FArgGenTaskInfo(uint32 InInstanceCountBufferOffset, uint32 InNumIndicesPerInstance) 
+			: InstanceCountBufferOffset(InInstanceCountBufferOffset)
+			, NumIndicesPerInstance(InNumIndicesPerInstance)
+		{}
+
+		uint32 InstanceCountBufferOffset;
+		uint32 NumIndicesPerInstance; // When -1 the counter needs to be reset to 0.
+	};
+
+	using FPermutationDomain = TShaderPermutationDomain<>;
+
+	static bool ShouldCompilePermutation(const FGlobalShaderPermutationParameters& Parameters)
+	{
+		return RHISupportsComputeShaders(Parameters.Platform);
+	}
+
+	static void ModifyCompilationEnvironment(const FGlobalShaderPermutationParameters& Parameters, FShaderCompilerEnvironment& OutEnvironment);
+
+	FNiagaraDrawIndirectArgsGenCS() {}
+	FNiagaraDrawIndirectArgsGenCS(const ShaderMetaType::CompiledShaderInitializerType& Initializer);
+
+	bool Serialize(FArchive& Ar);
+	void SetOutput(FRHICommandList& RHICmdList, FUnorderedAccessViewRHIParamRef DrawIndirectArgsUAV, FUnorderedAccessViewRHIParamRef InstanceCountsUAV);
+	void SetParameters(FRHICommandList& RHICmdList, FShaderResourceViewRHIParamRef TaskInfosBuffer, int32 NumArgGenTasks, int32 NumInstanceCountClearTasks);
+	void UnbindBuffers(FRHICommandList& RHICmdList);
+
+protected:
+	
+	FShaderResourceParameter TaskInfosParam;
+	FRWShaderParameter InstanceCountsParam;
+	FRWShaderParameter DrawIndirectArgsParam;
+	FShaderParameter TaskCountParam;
+};
\ No newline at end of file
diff --git a/Engine/Plugins/FX/Niagara/Source/NiagaraVertexFactories/Public/NiagaraSpriteVertexFactory.h b/Engine/Plugins/FX/Niagara/Source/NiagaraVertexFactories/Public/NiagaraSpriteVertexFactory.h
--- a/Engine/Plugins/FX/Niagara/Source/NiagaraVertexFactories/Public/NiagaraSpriteVertexFactory.h
+++ b/Engine/Plugins/FX/Niagara/Source/NiagaraVertexFactories/Public/NiagaraSpriteVertexFactory.h
@@ -64,6 +64,7 @@ BEGIN_GLOBAL_SHADER_PARAMETER_STRUCT(FNiagaraSpriteVFLooseParameters, NIAGARAVER
 	SHADER_PARAMETER(uint32, ParticleAlignmentMode)
 	SHADER_PARAMETER(uint32, ParticleFacingMode)
 	SHADER_PARAMETER(uint32, SortedIndicesOffset)
+	SHADER_PARAMETER(uint32, IndirectArgsOffset)
 	SHADER_PARAMETER_SRV(Buffer<float2>, CutoutGeometry)
 	SHADER_PARAMETER_SRV(Buffer<float>, NiagaraParticleDataFloat)
 	SHADER_PARAMETER_SRV(Buffer<int>, SortedIndices)
diff --git a/Engine/Source/Runtime/Engine/Private/Particles/FXSystem.cpp b/Engine/Source/Runtime/Engine/Private/Particles/FXSystem.cpp
--- a/Engine/Source/Runtime/Engine/Private/Particles/FXSystem.cpp
+++ b/Engine/Source/Runtime/Engine/Private/Particles/FXSystem.cpp
@@ -361,7 +361,7 @@ void FFXSystem::DrawDebug( FCanvas* Canvas )
 	}
 }
 
-void FFXSystem::PreInitViews()
+void FFXSystem::PreInitViews(FRHICommandListImmediate& RHICmdList)
 {
 	if (RHISupportsGPUParticles())
 	{
diff --git a/Engine/Source/Runtime/Engine/Private/Particles/FXSystemPrivate.h b/Engine/Source/Runtime/Engine/Private/Particles/FXSystemPrivate.h
--- a/Engine/Source/Runtime/Engine/Private/Particles/FXSystemPrivate.h
+++ b/Engine/Source/Runtime/Engine/Private/Particles/FXSystemPrivate.h
@@ -149,7 +149,7 @@ public:
 	virtual void RemoveVectorField(UVectorFieldComponent* VectorFieldComponent) override;
 	virtual void UpdateVectorField(UVectorFieldComponent* VectorFieldComponent) override;
 	FParticleEmitterInstance* CreateGPUSpriteEmitterInstance(FGPUSpriteEmitterInfo& EmitterInfo);
-	virtual void PreInitViews() override;
+	virtual void PreInitViews(FRHICommandListImmediate& RHICmdList) override;
 	virtual bool UsesGlobalDistanceField() const override;
 	virtual void PreRender(FRHICommandListImmediate& RHICmdList, const FGlobalDistanceFieldParameterData* GlobalDistanceFieldParameterData) override;
 	virtual void PostRenderOpaque(
diff --git a/Engine/Source/Runtime/Engine/Private/Particles/FXSystemSet.cpp b/Engine/Source/Runtime/Engine/Private/Particles/FXSystemSet.cpp
--- a/Engine/Source/Runtime/Engine/Private/Particles/FXSystemSet.cpp
+++ b/Engine/Source/Runtime/Engine/Private/Particles/FXSystemSet.cpp
@@ -87,12 +87,12 @@ void FFXSystemSet::UpdateVectorField(UVectorFieldComponent* VectorFieldComponent
 	}
 }
 
-void FFXSystemSet::PreInitViews()
+void FFXSystemSet::PreInitViews(FRHICommandListImmediate& RHICmdList)
 {
 	for (FFXSystemInterface* FXSystem : FXSystems)
 	{
 		check(FXSystem);
-		FXSystem->PreInitViews();
+		FXSystem->PreInitViews(RHICmdList);
 	}
 }
 
diff --git a/Engine/Source/Runtime/Engine/Private/Particles/FXSystemSet.h b/Engine/Source/Runtime/Engine/Private/Particles/FXSystemSet.h
--- a/Engine/Source/Runtime/Engine/Private/Particles/FXSystemSet.h
+++ b/Engine/Source/Runtime/Engine/Private/Particles/FXSystemSet.h
@@ -30,7 +30,7 @@ public:
 	virtual void AddVectorField(UVectorFieldComponent* VectorFieldComponent) override;
 	virtual void RemoveVectorField(UVectorFieldComponent* VectorFieldComponent) override;
 	virtual void UpdateVectorField(UVectorFieldComponent* VectorFieldComponent) override;
-	virtual void PreInitViews() override;
+	virtual void PreInitViews(FRHICommandListImmediate& RHICmdList) override;
 	virtual bool UsesGlobalDistanceField() const override;
 	virtual void PreRender(FRHICommandListImmediate& RHICmdList, const class FGlobalDistanceFieldParameterData* GlobalDistanceFieldParameterData) override;
 	virtual void PostRenderOpaque(
diff --git a/Engine/Source/Runtime/Engine/Public/FXSystem.h b/Engine/Source/Runtime/Engine/Public/FXSystem.h
--- a/Engine/Source/Runtime/Engine/Public/FXSystem.h
+++ b/Engine/Source/Runtime/Engine/Public/FXSystem.h
@@ -197,7 +197,7 @@ public:
 	 * Notification from the renderer that it is about to perform visibility
 	 * checks on FX belonging to this system.
 	 */
-	virtual void PreInitViews() = 0;
+	virtual void PreInitViews(FRHICommandListImmediate& RHICmdList) = 0;
 
 	virtual bool UsesGlobalDistanceField() const = 0;
 
diff --git a/Engine/Source/Runtime/Engine/Public/MeshBatch.h b/Engine/Source/Runtime/Engine/Public/MeshBatch.h
--- a/Engine/Source/Runtime/Engine/Public/MeshBatch.h
+++ b/Engine/Source/Runtime/Engine/Public/MeshBatch.h
@@ -96,6 +96,7 @@ struct FMeshBatchElement
 	int8 VisualizeElementIndex;
 #endif
 	FVertexBufferRHIParamRef IndirectArgsBuffer;
+	uint32 IndirectArgsOffset;
 
 	FMeshBatchElement()
 	:	PrimitiveUniformBuffer(nullptr)
@@ -121,6 +122,7 @@ struct FMeshBatchElement
 	,	VisualizeElementIndex(INDEX_NONE)
 #endif
 	,	IndirectArgsBuffer(nullptr)
+	,	IndirectArgsOffset(0)
 	{
 	}
 };
diff --git a/Engine/Source/Runtime/Renderer/Private/MeshPassProcessor.cpp b/Engine/Source/Runtime/Renderer/Private/MeshPassProcessor.cpp
--- a/Engine/Source/Runtime/Renderer/Private/MeshPassProcessor.cpp
+++ b/Engine/Source/Runtime/Renderer/Private/MeshPassProcessor.cpp
@@ -722,7 +722,8 @@ void FMeshDrawCommand::SetDrawParametersAndFinalize(
 	else
 	{
 		checkf(BatchElement.IndirectArgsBuffer, TEXT("It is only valid to set BatchElement.NumPrimitives == 0 when a IndirectArgsBuffer is used"));
-		IndirectArgsBuffer = BatchElement.IndirectArgsBuffer;
+		IndirectArgs.Buffer = BatchElement.IndirectArgsBuffer;
+		IndirectArgs.Offset = BatchElement.IndirectArgsOffset;
 	}
 
 	Finalize(PipelineId, ShadersForDebugging);
@@ -874,9 +875,9 @@ void FMeshDrawCommand::SubmitDraw(
 		FName ResourceName = MeshDrawCommand.DebugData.ResourceName;
 
 		FString DrawEventName = FString::Printf(
-			TEXT("%s %s"),
-			// Note: this is the parent's material name, not the material instance
-			*Material->GetFriendlyName(),
+				TEXT("%s %s"),
+				// Note: this is the parent's material name, not the material instance
+				*Material->GetFriendlyName(),
 			ResourceName.IsValid() ? *ResourceName.ToString() : TEXT(""));
 
 		const uint32 Instances = MeshDrawCommand.NumInstances * InstanceFactor;
@@ -900,37 +901,37 @@ void FMeshDrawCommand::SubmitDraw(
 	{
 		const FGraphicsMinimalPipelineStateInitializer& MeshPipelineState = MeshDrawCommand.CachedPipelineId.GetPipelineState(GraphicsMinimalPipelineStateSet);
 
-		if (MeshDrawCommand.CachedPipelineId.GetId() != StateCache.PipelineId)
-		{
-			FGraphicsPipelineStateInitializer GraphicsPSOInit = MeshPipelineState;
-			RHICmdList.ApplyCachedRenderTargets(GraphicsPSOInit);
-			SetGraphicsPipelineState(RHICmdList, GraphicsPSOInit);
-			StateCache.SetPipelineState(MeshDrawCommand.CachedPipelineId.GetId());
-		}
+	if (MeshDrawCommand.CachedPipelineId.GetId() != StateCache.PipelineId)
+	{
+		FGraphicsPipelineStateInitializer GraphicsPSOInit = MeshPipelineState;
+		RHICmdList.ApplyCachedRenderTargets(GraphicsPSOInit);
+		SetGraphicsPipelineState(RHICmdList, GraphicsPSOInit);
+		StateCache.SetPipelineState(MeshDrawCommand.CachedPipelineId.GetId());
+	}
+
+	if (MeshDrawCommand.StencilRef != StateCache.StencilRef)
+	{
+		RHICmdList.SetStencilRef(MeshDrawCommand.StencilRef);
+		StateCache.StencilRef = MeshDrawCommand.StencilRef;
+	}
+
+	for (int32 VertexBindingIndex = 0; VertexBindingIndex < MeshDrawCommand.VertexStreams.Num(); VertexBindingIndex++)
+	{
+		const FVertexInputStream& Stream = MeshDrawCommand.VertexStreams[VertexBindingIndex];
 
-		if (MeshDrawCommand.StencilRef != StateCache.StencilRef)
+		if (MeshDrawCommand.PrimitiveIdStreamIndex != -1 && Stream.StreamIndex == MeshDrawCommand.PrimitiveIdStreamIndex)
 		{
-			RHICmdList.SetStencilRef(MeshDrawCommand.StencilRef);
-			StateCache.StencilRef = MeshDrawCommand.StencilRef;
+			RHICmdList.SetStreamSource(Stream.StreamIndex, ScenePrimitiveIdsBuffer, PrimitiveIdOffset);
+			StateCache.VertexStreams[Stream.StreamIndex] = Stream;
 		}
-
-		for (int32 VertexBindingIndex = 0; VertexBindingIndex < MeshDrawCommand.VertexStreams.Num(); VertexBindingIndex++)
+		else if (StateCache.VertexStreams[Stream.StreamIndex] != Stream)
 		{
-			const FVertexInputStream& Stream = MeshDrawCommand.VertexStreams[VertexBindingIndex];
-
-			if (MeshDrawCommand.PrimitiveIdStreamIndex != -1 && Stream.StreamIndex == MeshDrawCommand.PrimitiveIdStreamIndex)
-			{
-				RHICmdList.SetStreamSource(Stream.StreamIndex, ScenePrimitiveIdsBuffer, PrimitiveIdOffset);
-				StateCache.VertexStreams[Stream.StreamIndex] = Stream;
-			}
-			else if (StateCache.VertexStreams[Stream.StreamIndex] != Stream)
-			{
-				RHICmdList.SetStreamSource(Stream.StreamIndex, Stream.VertexBuffer, Stream.Offset);
-				StateCache.VertexStreams[Stream.StreamIndex] = Stream;
-			}
+			RHICmdList.SetStreamSource(Stream.StreamIndex, Stream.VertexBuffer, Stream.Offset);
+			StateCache.VertexStreams[Stream.StreamIndex] = Stream;
 		}
+	}
 
-		MeshDrawCommand.ShaderBindings.SetOnCommandList(RHICmdList, MeshPipelineState.BoundShaderState, StateCache.ShaderBindings);
+	MeshDrawCommand.ShaderBindings.SetOnCommandList(RHICmdList, MeshPipelineState.BoundShaderState, StateCache.ShaderBindings);
 	}
 
 	if (MeshDrawCommand.IndexBuffer)
@@ -950,10 +951,10 @@ void FMeshDrawCommand::SubmitDraw(
 		else
 		{
 			RHICmdList.DrawIndexedPrimitiveIndirect(
-				MeshDrawCommand.IndexBuffer,
-				MeshDrawCommand.IndirectArgsBuffer,
-				0
-			);
+				MeshDrawCommand.IndexBuffer, 
+				MeshDrawCommand.IndirectArgs.Buffer, 
+				MeshDrawCommand.IndirectArgs.Offset
+				);
 		}
 	}
 	else
diff --git a/Engine/Source/Runtime/Renderer/Private/SceneVisibility.cpp b/Engine/Source/Runtime/Renderer/Private/SceneVisibility.cpp
--- a/Engine/Source/Runtime/Renderer/Private/SceneVisibility.cpp
+++ b/Engine/Source/Runtime/Renderer/Private/SceneVisibility.cpp
@@ -2923,7 +2923,7 @@ void FSceneRenderer::PreVisibilityFrameSetup(FRHICommandListImmediate& RHICmdLis
 	// Notify the FX system that the scene is about to perform visibility checks.
 	if (Scene->FXSystem && !Views[0].bIsPlanarReflection)
 	{
-		Scene->FXSystem->PreInitViews();
+		Scene->FXSystem->PreInitViews(RHICmdList);
 	}
 
 	// Draw lines to lights affecting this mesh if its selected.
diff --git a/Engine/Source/Runtime/Renderer/Public/MeshPassProcessor.h b/Engine/Source/Runtime/Renderer/Public/MeshPassProcessor.h
--- a/Engine/Source/Runtime/Renderer/Public/MeshPassProcessor.h
+++ b/Engine/Source/Runtime/Renderer/Public/MeshPassProcessor.h
@@ -466,7 +466,11 @@ public:
 			uint32 NumVertices;
 		} VertexParams;
 		
-		FVertexBufferRHIParamRef IndirectArgsBuffer;
+		struct  
+		{
+			FVertexBufferRHIParamRef Buffer;
+			uint32 Offset;
+		} IndirectArgs;
 	};
 
 	int8 PrimitiveIdStreamIndex;
@@ -489,7 +493,7 @@ public:
 			&& NumPrimitives == Rhs.NumPrimitives
 			&& NumInstances == Rhs.NumInstances
 			&& ((NumPrimitives > 0 && VertexParams.BaseVertexIndex == Rhs.VertexParams.BaseVertexIndex && VertexParams.NumVertices == Rhs.VertexParams.NumVertices)
-				|| (NumPrimitives == 0 && IndirectArgsBuffer == Rhs.IndirectArgsBuffer));
+				|| (NumPrimitives == 0 && IndirectArgs.Buffer == Rhs.IndirectArgs.Buffer && IndirectArgs.Offset == Rhs.IndirectArgs.Offset));
 	}
 
 	uint32 GetDynamicInstancingHash() const
