Fix system params not appearing in particle event handlers

diff --git a/Engine/Plugins/FX/Niagara/Source/NiagaraEditor/Private/ViewModels/Stack/NiagaraStackFunctionInput.cpp b/Engine/Plugins/FX/Niagara/Source/NiagaraEditor/Private/ViewModels/Stack/NiagaraStackFunctionInput.cpp
--- a/Engine/Plugins/FX/Niagara/Source/NiagaraEditor/Private/ViewModels/Stack/NiagaraStackFunctionInput.cpp
+++ b/Engine/Plugins/FX/Niagara/Source/NiagaraEditor/Private/ViewModels/Stack/NiagaraStackFunctionInput.cpp
@@ -690,7 +690,7 @@ void UNiagaraStackFunctionInput::SetLinkedValueHandle(const FNiagaraParameterHan
 	RefreshValues();
 }
 
-bool UsageRunsBefore(ENiagaraScriptUsage UsageA, ENiagaraScriptUsage UsageB)
+bool UsageRunsBefore(ENiagaraScriptUsage UsageA, ENiagaraScriptUsage UsageB, bool bCheckInterpSpawn = false)
 {
 	static TArray<ENiagaraScriptUsage> UsagesOrderedByExecution
 	{
@@ -699,13 +699,23 @@ bool UsageRunsBefore(ENiagaraScriptUsage UsageA, ENiagaraScriptUsage UsageB)
 		ENiagaraScriptUsage::EmitterSpawnScript,
 		ENiagaraScriptUsage::EmitterUpdateScript,
 		ENiagaraScriptUsage::ParticleSpawnScript,
-		ENiagaraScriptUsage::ParticleUpdateScript
+		ENiagaraScriptUsage::ParticleEventScript,	// When not using interpolated spawn
+		ENiagaraScriptUsage::ParticleUpdateScript,
+		ENiagaraScriptUsage::ParticleEventScript	// When using interpolated spawn and is spawn
 	};
 
 	int32 IndexA;
 	int32 IndexB;
-	UsagesOrderedByExecution.Find(UsageA, IndexA);
-	UsagesOrderedByExecution.Find(UsageB, IndexB);
+	if (bCheckInterpSpawn)
+	{
+		UsagesOrderedByExecution.FindLast(UsageA, IndexA);
+		UsagesOrderedByExecution.FindLast(UsageB, IndexB);
+	}
+	else
+	{
+		UsagesOrderedByExecution.Find(UsageA, IndexA);
+		UsagesOrderedByExecution.Find(UsageB, IndexB);
+	}
 	return IndexA < IndexB;
 }
 
@@ -723,6 +733,7 @@ FName GetNamespaceForUsage(ENiagaraScriptUsage Usage)
 	{
 	case ENiagaraScriptUsage::ParticleSpawnScript:
 	case ENiagaraScriptUsage::ParticleUpdateScript:
+	case ENiagaraScriptUsage::ParticleEventScript:
 		return FNiagaraParameterHandle::ParticleAttributeNamespace;
 	case ENiagaraScriptUsage::EmitterSpawnScript:
 	case ENiagaraScriptUsage::EmitterUpdateScript:
@@ -767,7 +778,23 @@ void UNiagaraStackFunctionInput::GetAvailableParameterHandles(TArray<FNiagaraPar
 
 	for (UNiagaraNodeOutput* OutputNode : AllOutputNodes)
 	{
-		if (OutputNode == CurrentOutputNode || (CurrentOutputNode != nullptr && UsageRunsBefore(OutputNode->GetUsage(), CurrentOutputNode->GetUsage())) || (CurrentOutputNode != nullptr && IsSpawnUsage(CurrentOutputNode->GetUsage())))
+		// Check if this is in a spawn event handler and the emitter is not using interpolated spawn so we
+		// we can hide particle update parameters
+		bool bSpawnScript = false;
+		if (CurrentOutputNode->GetUsage() == ENiagaraScriptUsage::ParticleEventScript)
+		{
+			for (const FNiagaraEventScriptProperties &EventHandlerProps : GetEmitterViewModel()->GetEmitter()->GetEventHandlers())
+			{
+				if (EventHandlerProps.Script->GetUsageId() == CurrentOutputNode->ScriptTypeId)
+				{
+					bSpawnScript = EventHandlerProps.ExecutionMode == EScriptExecutionMode::SpawnedParticles;
+					break;
+				}
+			}
+		}
+		bool bInterpolatedSpawn = GetEmitterViewModel()->GetEmitter()->bInterpolatedSpawning;
+		bool bCheckInterpSpawn = bInterpolatedSpawn || !bSpawnScript;
+		if (OutputNode == CurrentOutputNode || (CurrentOutputNode != nullptr && UsageRunsBefore(OutputNode->GetUsage(), CurrentOutputNode->GetUsage(), bCheckInterpSpawn)) || (CurrentOutputNode != nullptr && IsSpawnUsage(CurrentOutputNode->GetUsage())))
 		{
 			TArray<FNiagaraParameterHandle> AvailableParameterHandlesForThisOutput;
 			TArray<FNiagaraStackGraphUtilities::FStackNodeGroup> StackGroups;
