Added missing TResourceArray include

diff --git a/Engine/Plugins/FX/Niagara/Source/Niagara/Private/NiagaraGPUInstanceCountManager.cpp b/Engine/Plugins/FX/Niagara/Source/Niagara/Private/NiagaraGPUInstanceCountManager.cpp
--- a/Engine/Plugins/FX/Niagara/Source/Niagara/Private/NiagaraGPUInstanceCountManager.cpp
+++ b/Engine/Plugins/FX/Niagara/Source/Niagara/Private/NiagaraGPUInstanceCountManager.cpp
@@ -2,6 +2,7 @@
 
 #include "NiagaraGPUInstanceCountManager.h"
 #include "NiagaraStats.h"
+#include "Containers/DynamicRHIResourceArray.h"
 
 int32 GNiagaraMinGPUInstanceCount = 2048;
 static FAutoConsoleVariableRef CVarNiagaraMinGPUInstanceCount(
