Add search to niagara graph editor.

diff --git a/Engine/Plugins/FX/Niagara/Source/NiagaraEditor/Private/NiagaraEditorCommands.cpp b/Engine/Plugins/FX/Niagara/Source/NiagaraEditor/Private/NiagaraEditorCommands.cpp
--- a/Engine/Plugins/FX/Niagara/Source/NiagaraEditor/Private/NiagaraEditorCommands.cpp
+++ b/Engine/Plugins/FX/Niagara/Source/NiagaraEditor/Private/NiagaraEditorCommands.cpp
@@ -26,6 +26,7 @@ void FNiagaraEditorCommands::RegisterCommands()
 	UI_COMMAND(ToggleResetDependentSystems, "Reset Dependent Systems", "Toggles whether or not to reset all systems that include this emitter when it is reset by the user.", EUserInterfaceActionType::ToggleButton, FInputChord());
 	UI_COMMAND(CollapseStackToHeaders, "Collapse to Headers", "Expands all headsers and collapse all items.", EUserInterfaceActionType::Button, FInputChord(EModifierKey::Control, EKeys::O));
 
+	UI_COMMAND(FindInCurrentView, "Find", "Contextually finds items in current view.", EUserInterfaceActionType::Button, FInputChord(EModifierKey::Control, EKeys::F));
 }
 
 #undef LOCTEXT_NAMESPACE // NiagaraEditorCommands
diff --git a/Engine/Plugins/FX/Niagara/Source/NiagaraEditor/Private/NiagaraMessageManager.cpp b/Engine/Plugins/FX/Niagara/Source/NiagaraEditor/Private/NiagaraMessageManager.cpp
--- a/Engine/Plugins/FX/Niagara/Source/NiagaraEditor/Private/NiagaraMessageManager.cpp
+++ b/Engine/Plugins/FX/Niagara/Source/NiagaraEditor/Private/NiagaraMessageManager.cpp
@@ -499,13 +499,16 @@ void FNiagaraCompileEventToken::OpenScriptAssetByPathAndFocusNodeOrPinIfSet(
 			FNiagaraEditorModule& NiagaraEditorModule = FModuleManager::LoadModuleChecked<FNiagaraEditorModule>("NiagaraEditor");
 			if (InPinGUID.IsSet())
 			{
-				FNiagaraScriptGraphPinToFocusInfo PinToFocusInfo = FNiagaraScriptGraphPinToFocusInfo(ScriptAsset->GetUniqueID(), InPinGUID.GetValue());
-				NiagaraEditorModule.GetOnScriptToolkitsShouldFocusGraphElement().Broadcast(&PinToFocusInfo);
+				
+				TSharedRef<FNiagaraScriptGraphPinToFocusInfo> PinToFocusInfo = MakeShared<FNiagaraScriptGraphPinToFocusInfo>(InPinGUID.GetValue());
+				FNiagaraScriptIDAndGraphFocusInfo PinToFocusAndScriptID = FNiagaraScriptIDAndGraphFocusInfo(ScriptAsset->GetUniqueID(), PinToFocusInfo);
+				NiagaraEditorModule.GetOnScriptToolkitsShouldFocusGraphElement().Broadcast(&PinToFocusAndScriptID);
 			}
 			else if (InNodeGUID.IsSet())
 			{
-				FNiagaraScriptGraphNodeToFocusInfo NodeToFocusInfo = FNiagaraScriptGraphNodeToFocusInfo(ScriptAsset->GetUniqueID(), InNodeGUID.GetValue());
-				NiagaraEditorModule.GetOnScriptToolkitsShouldFocusGraphElement().Broadcast(&NodeToFocusInfo);
+				TSharedRef<FNiagaraScriptGraphNodeToFocusInfo> NodeToFocusInfo = MakeShared<FNiagaraScriptGraphNodeToFocusInfo>(InNodeGUID.GetValue());
+				FNiagaraScriptIDAndGraphFocusInfo NodeToFocusAndScriptID = FNiagaraScriptIDAndGraphFocusInfo(ScriptAsset->GetUniqueID(), NodeToFocusInfo);
+				NiagaraEditorModule.GetOnScriptToolkitsShouldFocusGraphElement().Broadcast(&NodeToFocusAndScriptID);
 			}
 		}
 		else
diff --git a/Engine/Plugins/FX/Niagara/Source/NiagaraEditor/Private/Toolkits/NiagaraScriptToolkit.cpp b/Engine/Plugins/FX/Niagara/Source/NiagaraEditor/Private/Toolkits/NiagaraScriptToolkit.cpp
--- a/Engine/Plugins/FX/Niagara/Source/NiagaraEditor/Private/Toolkits/NiagaraScriptToolkit.cpp
+++ b/Engine/Plugins/FX/Niagara/Source/NiagaraEditor/Private/Toolkits/NiagaraScriptToolkit.cpp
@@ -120,7 +120,7 @@ void FNiagaraScriptToolkit::UnregisterTabSpawners(const TSharedRef<class FTabMan
 }
 
 void FNiagaraScriptToolkit::Initialize( const EToolkitMode::Type Mode, const TSharedPtr< class IToolkitHost >& InitToolkitHost, UNiagaraScript* InputScript )
-{	
+{
 	check(InputScript != NULL);
 	OriginalNiagaraScript = InputScript;
 	ResetLoaders(GetTransientPackage()); // Make sure that we're not going to get invalid version number linkers into the package we are going into. 
@@ -235,7 +235,7 @@ void FNiagaraScriptToolkit::Initialize( const EToolkitMode::Type Mode, const TSh
 
 	FNiagaraEditorModule& NiagaraEditorModule = FModuleManager::LoadModuleChecked<FNiagaraEditorModule>( "NiagaraEditor" );
 	AddMenuExtender(NiagaraEditorModule.GetMenuExtensibilityManager()->GetAllExtenders(GetToolkitCommands(), GetEditingObjects()));
-	NiagaraEditorModule.GetOnScriptToolkitsShouldFocusGraphElement().AddSP(this, &FNiagaraScriptToolkit::FocusGraphElement);
+	NiagaraEditorModule.GetOnScriptToolkitsShouldFocusGraphElement().AddSP(this, &FNiagaraScriptToolkit::FocusGraphElementIfSameScriptID);
 
 	SetupCommands();
 	ExtendToolbar();
@@ -744,45 +744,12 @@ void FNiagaraScriptToolkit::OnEditedScriptGraphChanged(const FEdGraphEditAction&
 	bEditedScriptHasPendingChanges = true;
 }
 
-void FNiagaraScriptToolkit::FocusGraphElement(const INiagaraScriptGraphFocusInfo* FocusInfo)
+void FNiagaraScriptToolkit::FocusGraphElementIfSameScriptID(const FNiagaraScriptIDAndGraphFocusInfo* FocusInfo)
 {
-	checkf(FocusInfo->GetFocusType() != INiagaraScriptGraphFocusInfo::ENiagaraScriptGraphFocusInfoType::None, TEXT("Failed to assign focus type to FocusInfo parameter!"));
-	
-	if (FocusInfo->GetScriptUniqueAssetID() != OriginalNiagaraScript->GetUniqueID())
+	if (FocusInfo->GetScriptUniqueAssetID() == OriginalNiagaraScript->GetUniqueID())
 	{
-		return;
-	}
-
-	if (FocusInfo->GetFocusType() == INiagaraScriptGraphFocusInfo::ENiagaraScriptGraphFocusInfoType::Node)
-	{
-		const FNiagaraScriptGraphNodeToFocusInfo* NodeFocusInfo = static_cast<const FNiagaraScriptGraphNodeToFocusInfo*>(FocusInfo);
-		const FGuid& NodeGuidToMatch = NodeFocusInfo->GetNodeGuidToFocus();
-		UEdGraphNode* const* NodeToFocus = ScriptViewModel->GetGraphViewModel()->GetGraph()->Nodes.FindByPredicate([&NodeGuidToMatch](const UEdGraphNode* Node) {return Node->NodeGuid == NodeGuidToMatch; });
-		if (NodeToFocus != nullptr && *NodeToFocus != nullptr)
-		{
-			NiagaraScriptGraphWidget->GetGraphEditor()->JumpToNode(*NodeToFocus);
-			return;
-		}
-		ensureMsgf(false, TEXT("Failed to find Node with matching GUID when focusing graph element. Was the graph edited out from underneath us?"));
-		return;
-	}
-	else if (FocusInfo->GetFocusType() == INiagaraScriptGraphFocusInfo::ENiagaraScriptGraphFocusInfoType::Pin)
-	{
-		const FNiagaraScriptGraphPinToFocusInfo* PinFocusInfo = static_cast<const FNiagaraScriptGraphPinToFocusInfo*>(FocusInfo);
-		const FGuid& PinGuidToMatch = PinFocusInfo->GetPinGuidToFocus();
-		for (const UEdGraphNode* Node : ScriptViewModel->GetGraphViewModel()->GetGraph()->Nodes)
-		{
-			const UEdGraphPin* const* PinToFocus = Node->Pins.FindByPredicate([&PinGuidToMatch](const UEdGraphPin* Pin) {return Pin->PersistentGuid == PinGuidToMatch; });
-			if (PinToFocus != nullptr && *PinToFocus != nullptr)
-			{
-				NiagaraScriptGraphWidget->GetGraphEditor()->JumpToPin(*PinToFocus);
-				return;
-			}
-		}
-		ensureMsgf(false, TEXT("Failed to find Pin with matching GUID when focusing graph element. Was the graph edited out from underneath us?"));
-		return;
+		NiagaraScriptGraphWidget->FocusGraphElement(FocusInfo->GetScriptGraphFocusInfo().Get());
 	}
-	checkf(false, TEXT("Requested focus for a graph element without specifying a Node or Pin to focus!"));
 }
 
 #undef LOCTEXT_NAMESPACE
diff --git a/Engine/Plugins/FX/Niagara/Source/NiagaraEditor/Private/Toolkits/NiagaraScriptToolkit.h b/Engine/Plugins/FX/Niagara/Source/NiagaraEditor/Private/Toolkits/NiagaraScriptToolkit.h
--- a/Engine/Plugins/FX/Niagara/Source/NiagaraEditor/Private/Toolkits/NiagaraScriptToolkit.h
+++ b/Engine/Plugins/FX/Niagara/Source/NiagaraEditor/Private/Toolkits/NiagaraScriptToolkit.h
@@ -118,7 +118,7 @@ private:
 	/** Navigates to element in graph (node, pin, etc.) 
 	* @Param ElementToFocus Defines the graph element to navigate to and select.
 	*/
-	void FocusGraphElement(const INiagaraScriptGraphFocusInfo* ElementToFocus);
+	void FocusGraphElementIfSameScriptID(const FNiagaraScriptIDAndGraphFocusInfo* ElementToFocus);
 
 private:
 
diff --git a/Engine/Plugins/FX/Niagara/Source/NiagaraEditor/Private/Widgets/SNiagaraScriptGraph.cpp b/Engine/Plugins/FX/Niagara/Source/NiagaraEditor/Private/Widgets/SNiagaraScriptGraph.cpp
--- a/Engine/Plugins/FX/Niagara/Source/NiagaraEditor/Private/Widgets/SNiagaraScriptGraph.cpp
+++ b/Engine/Plugins/FX/Niagara/Source/NiagaraEditor/Private/Widgets/SNiagaraScriptGraph.cpp
@@ -19,6 +19,8 @@
 #include "ScopedTransaction.h"
 #include "NiagaraEditorUtilities.h"
 #include "NiagaraNodeFactory.h"
+#include "NiagaraEditorCommands.h"
+#include "NiagaraNodeOutput.h"
 
 #define LOCTEXT_NAMESPACE "NiagaraScriptGraph"
 
@@ -38,6 +40,8 @@ void SNiagaraScriptGraph::Construct(const FArguments& InArgs, TSharedRef<FNiagar
 	[
 		GraphEditor.ToSharedRef()
 	];
+
+	CurrentFocusedSearchMatchIndex = 0;
 }
 
 TSharedRef<SGraphEditor> SNiagaraScriptGraph::ConstructGraphEditor()
@@ -45,29 +49,87 @@ TSharedRef<SGraphEditor> SNiagaraScriptGraph::ConstructGraphEditor()
 	FGraphAppearanceInfo AppearanceInfo;
 	AppearanceInfo.CornerText = LOCTEXT("AppearanceCornerText", "NIAGARA");
 
+	const FEditableTextBoxStyle& Style = FCoreStyle::Get().GetWidgetStyle<FEditableTextBoxStyle>("SearchBox");
+
 	TSharedRef<SWidget> TitleBarWidget =
 		SNew(SBorder)
 		.BorderImage(FEditorStyle::GetBrush(TEXT("Graph.TitleBackground")))
 		.HAlign(HAlign_Fill)
 		[
-			SNew(SHorizontalBox)
-			+ SHorizontalBox::Slot()
-			.AutoWidth()
-			.Padding(0.0f, 0.0f, 3.0f, 0.0f)
+			SNew(SOverlay)
+			+SOverlay::Slot()
 			[
-				SNew(SErrorText)
-				.Visibility(ViewModel.ToSharedRef(), &FNiagaraScriptGraphViewModel::GetGraphErrorTextVisible)
-				.BackgroundColor(ViewModel.ToSharedRef(), &FNiagaraScriptGraphViewModel::GetGraphErrorColor)
-				.ToolTipText(ViewModel.ToSharedRef(), &FNiagaraScriptGraphViewModel::GetGraphErrorMsgToolTip)
-				.ErrorText(ViewModel->GetGraphErrorText())
+				SNew(SHorizontalBox)
+				+ SHorizontalBox::Slot()
+				.AutoWidth()
+				.Padding(0.0f, 0.0f, 3.0f, 0.0f)
+				[
+					SNew(SErrorText)
+					.Visibility(ViewModel.ToSharedRef(), &FNiagaraScriptGraphViewModel::GetGraphErrorTextVisible)
+					.BackgroundColor(ViewModel.ToSharedRef(), &FNiagaraScriptGraphViewModel::GetGraphErrorColor)
+					.ToolTipText(ViewModel.ToSharedRef(), &FNiagaraScriptGraphViewModel::GetGraphErrorMsgToolTip)
+					.ErrorText(ViewModel->GetGraphErrorText())
+				]
+				+SHorizontalBox::Slot()
+				.VAlign(VAlign_Center)
+				[
+					SNew(STextBlock)
+					.Text(ViewModel.ToSharedRef(), &FNiagaraScriptGraphViewModel::GetDisplayName)
+					.TextStyle(FEditorStyle::Get(), TEXT("GraphBreadcrumbButtonText"))
+					.Justification(ETextJustify::Center)
+				]
 			]
-			+ SHorizontalBox::Slot()
-			.VAlign(VAlign_Center)
+			+SOverlay::Slot()
+			.HAlign(HAlign_Right)
+			.VAlign(VAlign_Fill)
 			[
-				SNew(STextBlock)
-				.Text(ViewModel.ToSharedRef(), &FNiagaraScriptGraphViewModel::GetDisplayName)
-				.TextStyle(FEditorStyle::Get(), TEXT("GraphBreadcrumbButtonText"))
-				.Justification(ETextJustify::Center)
+				SNew(SHorizontalBox)
+				+ SHorizontalBox::Slot()
+				.HAlign(HAlign_Fill)
+				.VAlign(VAlign_Fill)
+				.AutoWidth()
+				.MaxWidth(400.0f) // Limit max search box width to avoid extending over titlebar
+				[
+					SAssignNew(SearchBox, SSearchBox)
+					.HintText(LOCTEXT("GraphSearchBoxHint", "Search Nodes and Pins in Graph"))
+					.SearchResultData(this, &SNiagaraScriptGraph::GetSearchResultData)
+					.OnTextChanged(this, &SNiagaraScriptGraph::OnSearchTextChanged)
+					.OnTextCommitted(this, &SNiagaraScriptGraph::OnSearchBoxTextCommitted)
+					.DelayChangeNotificationsWhileTyping(true)
+					.OnSearch(this, &SNiagaraScriptGraph::OnSearchBoxSearch)
+					.Visibility(this, &SNiagaraScriptGraph::GetGraphSearchBoxVisibility)
+					.OnKeyDownHandler(this, &SNiagaraScriptGraph::HandleGraphSearchBoxKeyDown)
+				]
+				+SHorizontalBox::Slot()
+				.HAlign(HAlign_Fill)
+				.VAlign(VAlign_Fill)
+				.AutoWidth()
+				[
+					SNew(SBorder)
+					.HAlign(HAlign_Center)
+					.VAlign(VAlign_Fill)
+					.BorderImage(&Style.BackgroundImageHovered)
+					.BorderBackgroundColor(Style.BackgroundColor)
+					.ForegroundColor(Style.ForegroundColor)
+					.Padding(0)
+					[
+						SNew(SButton)
+						.ButtonStyle(FCoreStyle::Get(), "NoBorder")
+						.ContentPadding(0)
+						.HAlign(HAlign_Center)
+						.VAlign(VAlign_Center)
+						.IsFocusable(false)
+						.ToolTipText(LOCTEXT("CloseGraphSearchBox", "Close Graph search box"))
+						.Visibility(this, &SNiagaraScriptGraph::GetGraphSearchBoxVisibility)
+						.OnClicked(this, &SNiagaraScriptGraph::CloseGraphSearchBoxPressed)
+						.Content()
+						[
+ 							SNew(SImage)
+ 							.Image(FEditorStyle::GetBrush("Symbols.X"))
+							.ColorAndOpacity(FSlateColor::UseForeground())
+						]
+					]
+				]
 			]
 		];
 
@@ -78,6 +140,11 @@ TSharedRef<SGraphEditor> SNiagaraScriptGraph::ConstructGraphEditor()
 	Events.OnVerifyTextCommit = FOnNodeVerifyTextCommit::CreateSP(this, &SNiagaraScriptGraph::OnVerifyNodeTextCommit);
 	Events.OnSpawnNodeByShortcut = SGraphEditor::FOnSpawnNodeByShortcut::CreateSP(this, &SNiagaraScriptGraph::OnSpawnGraphNodeByShortcut);
 
+	TSharedRef<FUICommandList> Commands = ViewModel->GetCommands();
+	Commands->MapAction(
+		FNiagaraEditorCommands::Get().FindInCurrentView,
+		FExecuteAction::CreateRaw(this, &SNiagaraScriptGraph::FocusGraphSearchBox));
+	
 	TSharedRef<SGraphEditor> CreatedGraphEditor = SNew(SGraphEditor)
 		.AdditionalCommands(ViewModel->GetCommands())
 		.Appearance(AppearanceInfo)
@@ -318,4 +385,137 @@ void SNiagaraScriptGraph::GraphChanged()
 	];
 }
 
+void SNiagaraScriptGraph::FocusGraphElement(const INiagaraScriptGraphFocusInfo* FocusInfo)
+{
+	checkf(FocusInfo->GetFocusType() != INiagaraScriptGraphFocusInfo::ENiagaraScriptGraphFocusInfoType::None, TEXT("Failed to assign focus type to FocusInfo parameter!"));
+
+	if (FocusInfo->GetFocusType() == INiagaraScriptGraphFocusInfo::ENiagaraScriptGraphFocusInfoType::Node)
+	{
+		const FNiagaraScriptGraphNodeToFocusInfo* NodeFocusInfo = static_cast<const FNiagaraScriptGraphNodeToFocusInfo*>(FocusInfo);
+		const FGuid& NodeGuidToMatch = NodeFocusInfo->GetNodeGuidToFocus();
+		UEdGraphNode* const* NodeToFocus = ViewModel->GetGraph()->Nodes.FindByPredicate([&NodeGuidToMatch](const UEdGraphNode* Node) {return Node->NodeGuid == NodeGuidToMatch; });
+		if (NodeToFocus != nullptr && *NodeToFocus != nullptr)
+		{
+			GetGraphEditor()->JumpToNode(*NodeToFocus);
+			return;
+		}
+		ensureMsgf(false, TEXT("Failed to find Node with matching GUID when focusing graph element. Was the graph edited out from underneath us?"));
+		return;
+	}
+	else if (FocusInfo->GetFocusType() == INiagaraScriptGraphFocusInfo::ENiagaraScriptGraphFocusInfoType::Pin)
+	{
+		const FNiagaraScriptGraphPinToFocusInfo* PinFocusInfo = static_cast<const FNiagaraScriptGraphPinToFocusInfo*>(FocusInfo);
+		const FGuid& PinGuidToMatch = PinFocusInfo->GetPinGuidToFocus();
+		for (const UEdGraphNode* Node : ViewModel->GetGraph()->Nodes)
+		{
+			const UEdGraphPin* const* PinToFocus = Node->Pins.FindByPredicate([&PinGuidToMatch](const UEdGraphPin* Pin) {return Pin->PersistentGuid == PinGuidToMatch; });
+			if (PinToFocus != nullptr && *PinToFocus != nullptr)
+			{
+				GetGraphEditor()->JumpToPin(*PinToFocus);
+				return;
+			}
+		}
+		ensureMsgf(false, TEXT("Failed to find Pin with matching GUID when focusing graph element. Was the graph edited out from underneath us?"));
+		return;
+	}
+	checkf(false, TEXT("Requested focus for a graph element without specifying a Node or Pin to focus!"));
+}
+
+void SNiagaraScriptGraph::OnSearchTextChanged(const FText& SearchText)
+{
+	if (!CurrentSearchText.EqualTo(SearchText))
+	{
+		CurrentSearchResults.Empty();
+		CurrentSearchText = SearchText;
+		TArray<UNiagaraNode*> Nodes;
+		ViewModel->GetGraph()->GetNodesOfClass<UNiagaraNode>(Nodes);
+		for (UNiagaraNode* Node : Nodes)
+		{
+			if (Node->GetNodeTitle(ENodeTitleType::FullTitle).ToString().Contains(SearchText.ToString()))
+			{	
+				CurrentSearchResults.Add(MakeShared<FNiagaraScriptGraphNodeToFocusInfo>(Node->NodeGuid));
+			}
+			
+			if (Node->IsA<UNiagaraNodeOutput>() == false) 
+			{
+				for (UEdGraphPin* Pin : Node->GetAllPins())
+				{
+					if (Pin->GetDisplayName().ToString().Contains(SearchText.ToString()))
+					{
+						CurrentSearchResults.Add(MakeShared<FNiagaraScriptGraphPinToFocusInfo>(Pin->PersistentGuid));
+					}
+				}
+			}
+		}
+
+		CurrentFocusedSearchMatchIndex = 0;
+		if (CurrentSearchResults.Num() > 0)
+		{
+			FocusGraphElement(CurrentSearchResults[0].Get());
+		}
+	}
+}
+
+void SNiagaraScriptGraph::OnSearchBoxTextCommitted(const FText& NewText, ETextCommit::Type CommitInfo)
+{
+	if (SearchBox->HasKeyboardFocus())
+	{
+		OnSearchBoxSearch(SSearchBox::Next);
+	}
+}
+
+TOptional<SSearchBox::FSearchResultData> SNiagaraScriptGraph::GetSearchResultData() const
+{
+	if (CurrentSearchText.IsEmpty() || CurrentSearchResults.Num() == 0)
+	{
+		return TOptional<SSearchBox::FSearchResultData>();
+	}
+	return TOptional<SSearchBox::FSearchResultData>({ CurrentSearchResults.Num(), CurrentFocusedSearchMatchIndex + 1 });
+}
+
+FReply SNiagaraScriptGraph::CloseGraphSearchBoxPressed()
+{
+	bGraphSearchBoxActive = false;
+	return FReply::Handled();
+}
+
+FReply SNiagaraScriptGraph::HandleGraphSearchBoxKeyDown(const FGeometry& MyGeometry, const FKeyEvent& InKeyEvent)
+{
+	if (InKeyEvent.GetKey() == EKeys::Escape)
+	{
+		CloseGraphSearchBoxPressed();
+		return FReply::Handled();
+	}
+	return FReply::Unhandled();
+}
+
+void SNiagaraScriptGraph::FocusGraphSearchBox()
+{
+	bGraphSearchBoxActive = true;
+
+	if (SearchBox.IsValid())
+	{
+		FWidgetPath WidgetToFocusPath;
+		FSlateApplication::Get().GeneratePathToWidgetUnchecked(SearchBox.ToSharedRef(), WidgetToFocusPath, EVisibility::All);
+		FSlateApplication::Get().SetKeyboardFocus(WidgetToFocusPath, EFocusCause::SetDirectly);
+	}
+}
+
+void SNiagaraScriptGraph::OnSearchBoxSearch(SSearchBox::SearchDirection Direction)
+{
+	if (CurrentSearchResults.Num() > 0)
+	{
+		if (Direction == SSearchBox::Next)
+		{
+			CurrentFocusedSearchMatchIndex = CurrentFocusedSearchMatchIndex < CurrentSearchResults.Num() - 1 ? CurrentFocusedSearchMatchIndex + 1 : 0;
+			FocusGraphElement(CurrentSearchResults[CurrentFocusedSearchMatchIndex].Get());
+		}
+		else if (Direction == SSearchBox::Previous)
+		{
+			CurrentFocusedSearchMatchIndex = CurrentFocusedSearchMatchIndex > 0 ? CurrentFocusedSearchMatchIndex - 1 : CurrentSearchResults.Num() - 1;
+			FocusGraphElement(CurrentSearchResults[CurrentFocusedSearchMatchIndex].Get());
+		}
+	}
+}
+
 #undef LOCTEXT_NAMESPACE // "NiagaraScriptGraph"
diff --git a/Engine/Plugins/FX/Niagara/Source/NiagaraEditor/Private/Widgets/SNiagaraScriptGraph.h b/Engine/Plugins/FX/Niagara/Source/NiagaraEditor/Private/Widgets/SNiagaraScriptGraph.h
--- a/Engine/Plugins/FX/Niagara/Source/NiagaraEditor/Private/Widgets/SNiagaraScriptGraph.h
+++ b/Engine/Plugins/FX/Niagara/Source/NiagaraEditor/Private/Widgets/SNiagaraScriptGraph.h
@@ -6,6 +6,8 @@
 #include "Widgets/DeclarativeSyntaxSupport.h"
 #include "GraphEditor.h"
 #include "EdGraph/EdGraphNode.h"
+#include "Widgets/Input/SSearchBox.h"
+#include "NiagaraEditorCommon.h"
 
 class FNiagaraScriptGraphViewModel;
 
@@ -24,6 +26,10 @@ public:
 
 	const TSharedPtr<SGraphEditor> GetGraphEditor() { return GraphEditor; };
 
+	void FocusGraphElement(const INiagaraScriptGraphFocusInfo* FocusInfo);
+
+	void FocusGraphSearchBox();
+
 private:
 	/** Constructs the graph editor widget for the current graph. */
 	TSharedRef<SGraphEditor> ConstructGraphEditor();
@@ -55,6 +61,24 @@ private:
 	/** Called whenever a user is trying to edit the name inline of a node and we want to make sure that it is valid.*/
 	bool OnVerifyNodeTextCommit(const FText& NewText, UEdGraphNode* NodeBeingChanged, FText& OutErrorMessage);
 
+	void OnSearchTextChanged(const FText& SearchText);
+
+	void OnSearchBoxTextCommitted(const FText& NewText, ETextCommit::Type CommitInfo);
+
+	void OnSearchBoxSearch(SSearchBox::SearchDirection Direction);
+
+	TOptional<SSearchBox::FSearchResultData> GetSearchResultData() const;
+
+	int GetCurrentFocusedMatchIndex() const { return CurrentFocusedSearchMatchIndex; }
+
+	FText GetCurrentSearchText() const { return CurrentSearchText; };
+
+	EVisibility GetGraphSearchBoxVisibility() const { return bGraphSearchBoxActive ? EVisibility::Visible : EVisibility::Collapsed; };
+
+	FReply CloseGraphSearchBoxPressed();
+
+	FReply HandleGraphSearchBoxKeyDown(const FGeometry& MyGeometry, const FKeyEvent& InKeyEvent);
+
 private:
 	/** An attribute for the title text of the graph. */
 	TAttribute<FText> GraphTitle;
@@ -65,7 +89,16 @@ private:
 	/** The graph editor which is editing the script graph. */
 	TSharedPtr<SGraphEditor> GraphEditor;
 
+	/** The search box for searching the script graph. */
+	TSharedPtr<SSearchBox> SearchBox;
+
 	/** Flag to prevent modifying the view model selection when updating the graph
 		editor selection due to a view model selection change. */
 	bool bUpdatingGraphSelectionFromViewModel;
+
+	// ~Search stuff
+	FText CurrentSearchText;
+	int CurrentFocusedSearchMatchIndex;
+	TArray<TSharedPtr<INiagaraScriptGraphFocusInfo>> CurrentSearchResults;
+	bool bGraphSearchBoxActive;
 };
diff --git a/Engine/Plugins/FX/Niagara/Source/NiagaraEditor/Public/NiagaraEditorCommands.h b/Engine/Plugins/FX/Niagara/Source/NiagaraEditor/Public/NiagaraEditorCommands.h
--- a/Engine/Plugins/FX/Niagara/Source/NiagaraEditor/Public/NiagaraEditorCommands.h
+++ b/Engine/Plugins/FX/Niagara/Source/NiagaraEditor/Public/NiagaraEditorCommands.h
@@ -49,4 +49,6 @@ public:
 	TSharedPtr<FUICommandInfo> ToggleResetDependentSystems;
 
 	TSharedPtr<FUICommandInfo> CollapseStackToHeaders;
+
+	TSharedPtr<FUICommandInfo> FindInCurrentView;
 };
\ No newline at end of file
diff --git a/Engine/Plugins/FX/Niagara/Source/NiagaraEditor/Public/NiagaraEditorCommon.h b/Engine/Plugins/FX/Niagara/Source/NiagaraEditor/Public/NiagaraEditorCommon.h
--- a/Engine/Plugins/FX/Niagara/Source/NiagaraEditor/Public/NiagaraEditorCommon.h
+++ b/Engine/Plugins/FX/Niagara/Source/NiagaraEditor/Public/NiagaraEditorCommon.h
@@ -52,7 +52,7 @@ public:
 };
 
 /** Interface for struct representing information about where to focus in a Niagara Script Graph after opening the editor for it. */
-struct INiagaraScriptGraphFocusInfo
+struct INiagaraScriptGraphFocusInfo : public TSharedFromThis<INiagaraScriptGraphFocusInfo>
 {
 public:
 	enum class ENiagaraScriptGraphFocusInfoType : uint8
@@ -62,28 +62,24 @@ public:
 		Pin
 	};
 
-	INiagaraScriptGraphFocusInfo(const uint32& InScriptUniqueAssetID, const ENiagaraScriptGraphFocusInfoType InFocusType)
-		: ScriptUniqueAssetID(InScriptUniqueAssetID)
-		, FocusType(InFocusType)
+	INiagaraScriptGraphFocusInfo(const ENiagaraScriptGraphFocusInfoType InFocusType)
+		: FocusType(InFocusType)
 	{
 	};
 
-	const uint32& GetScriptUniqueAssetID() const { return ScriptUniqueAssetID; };
-
 	const ENiagaraScriptGraphFocusInfoType& GetFocusType() const { return FocusType; };
 
 	virtual ~INiagaraScriptGraphFocusInfo() = 0;
 	
 private:
-	const uint32 ScriptUniqueAssetID;
 	const ENiagaraScriptGraphFocusInfoType FocusType;
 };
 
 struct FNiagaraScriptGraphNodeToFocusInfo : public INiagaraScriptGraphFocusInfo
 {
 public:
-	FNiagaraScriptGraphNodeToFocusInfo(const uint32& InScriptUniqueAssetID, const FGuid& InNodeGuidToFocus)
-		: INiagaraScriptGraphFocusInfo(InScriptUniqueAssetID, ENiagaraScriptGraphFocusInfoType::Node)
+	FNiagaraScriptGraphNodeToFocusInfo(const FGuid& InNodeGuidToFocus)
+		: INiagaraScriptGraphFocusInfo(ENiagaraScriptGraphFocusInfoType::Node)
 		, NodeGuidToFocus(InNodeGuidToFocus)
 	{
 	};
@@ -97,8 +93,8 @@ private:
 struct FNiagaraScriptGraphPinToFocusInfo : public INiagaraScriptGraphFocusInfo
 {
 public:
-	FNiagaraScriptGraphPinToFocusInfo(const uint32& InScriptUniqueAssetID, const FGuid& InPinGuidToFocus)
-		: INiagaraScriptGraphFocusInfo(InScriptUniqueAssetID, ENiagaraScriptGraphFocusInfoType::Pin)
+	FNiagaraScriptGraphPinToFocusInfo(const FGuid& InPinGuidToFocus)
+		: INiagaraScriptGraphFocusInfo(ENiagaraScriptGraphFocusInfoType::Pin)
 		, PinGuidToFocus(InPinGuidToFocus)
 	{
 	};
@@ -108,3 +104,21 @@ public:
 private:
 	const FGuid PinGuidToFocus;
 };
+
+struct FNiagaraScriptIDAndGraphFocusInfo
+{
+public:
+	FNiagaraScriptIDAndGraphFocusInfo(const uint32& InScriptUniqueAssetID, const TSharedPtr<INiagaraScriptGraphFocusInfo>& InScriptGraphFocusInfo)
+		: ScriptUniqueAssetID(InScriptUniqueAssetID)
+		, ScriptGraphFocusInfo(InScriptGraphFocusInfo)
+	{
+	};
+
+	const uint32& GetScriptUniqueAssetID() const { return ScriptUniqueAssetID; };
+
+	const TSharedPtr<INiagaraScriptGraphFocusInfo>& GetScriptGraphFocusInfo() const { return ScriptGraphFocusInfo; };
+
+private:
+	const uint32 ScriptUniqueAssetID;
+	const TSharedPtr<INiagaraScriptGraphFocusInfo> ScriptGraphFocusInfo;
+};
diff --git a/Engine/Plugins/FX/Niagara/Source/NiagaraEditor/Public/NiagaraEditorModule.h b/Engine/Plugins/FX/Niagara/Source/NiagaraEditor/Public/NiagaraEditorModule.h
--- a/Engine/Plugins/FX/Niagara/Source/NiagaraEditor/Public/NiagaraEditorModule.h
+++ b/Engine/Plugins/FX/Niagara/Source/NiagaraEditor/Public/NiagaraEditorModule.h
@@ -30,7 +30,7 @@ class FNiagaraEditorModule : public IModuleInterface,
 public:
 	DECLARE_DELEGATE_RetVal_OneParam(TSharedRef<SWidget>, FOnCreateStackWidget, UNiagaraStackViewModel*);
 	DECLARE_DELEGATE_RetVal_OneParam(UMovieSceneNiagaraParameterTrack*, FOnCreateMovieSceneTrackForParameter, FNiagaraVariable);
-	DECLARE_MULTICAST_DELEGATE_OneParam(FOnCheckScriptToolkitsShouldFocusGraphElement, const INiagaraScriptGraphFocusInfo*);
+	DECLARE_MULTICAST_DELEGATE_OneParam(FOnCheckScriptToolkitsShouldFocusGraphElement, const FNiagaraScriptIDAndGraphFocusInfo*);
 
 public:
 	FNiagaraEditorModule();
