Back out changelist 6894880. Cook errors

diff --git a/Engine/Plugins/FX/Niagara/Shaders/Private/NiagaraDrawIndirectArgsGen.usf b/Engine/Plugins/FX/Niagara/Shaders/Private/NiagaraDrawIndirectArgsGen.usf
deleted file mode 100644
diff --git a/Engine/Plugins/FX/Niagara/Shaders/Private/NiagaraEmitterInstanceShader.usf b/Engine/Plugins/FX/Niagara/Shaders/Private/NiagaraEmitterInstanceShader.usf
--- a/Engine/Plugins/FX/Niagara/Shaders/Private/NiagaraEmitterInstanceShader.usf
+++ b/Engine/Plugins/FX/Niagara/Shaders/Private/NiagaraEmitterInstanceShader.usf
@@ -298,6 +298,7 @@ int rand(int x, int Seed1, int Seed2, int Seed3)
 	 */
 	uint UpdateStartInstance;
 	uint SpawnedInstances;
+	uint NumIndicesPerInstance;
 
 	uint ComponentBufferSizeRead;
 	uint ComponentBufferSizeWrite;
@@ -312,15 +313,16 @@ int rand(int x, int Seed1, int Seed2, int Seed3)
 	Buffer<int> InputInt;
 	RWBuffer<int> RWOutputInt;
 	RWBuffer<float> RWOutputFloat;
+	RWBuffer<uint> RWDataSetIndices;
+	Buffer<uint> ReadDataSetIndices;
 
-	RWBuffer<uint> RWInstanceCounts;
-	uint ReadInstanceCountOffset;
-	uint WriteInstanceCountOffset;
 
 	//TODO: IDs for GPU
 //	void AcquireID(int DataSetID, out int IDIndex, out int IDAcquireTag);
 // 	void UpdateID(int DataSetID, int IDIndex, int InstanceIndex);
 
+
+
 	/* Returns the current instance index relative to the operation (spawn/update)
 	 */
 	int ExecIndex()
@@ -339,20 +341,14 @@ int rand(int x, int Seed1, int Seed2, int Seed3)
 	/* Acquire an output index - the default index is the scratch instance; one additional instance is allocated 
 	 *	at the end of the buffer, so no branching on -1 is necessary during OutputData operations
 	 */
-	int AcquireIndex(uniform int DataSetID, bool bDoAcquire)
+	int AcquireIndex(int DataSetID, bool bDoAcquire)
 	{
-		// Begin static assert : GPU particles only support DataSetID 0
-		int MustBe0[1];
-		MustBe0[DataSetID] = 0;
-		// End static assert
-
 		int PrevIdx = GSpawnStartInstance + SpawnedInstances;	// scratch instance as default; write to that for dead particles
-		if (bDoAcquire)
+		if(bDoAcquire == true)
 		{
 			// Have to use uint's here to avoid PS4 compiler warnings about InterlockedAdd, cannot propagate uint due to CPU VM limitations...
 			uint RetPrevIdx;
-			// @TODO : add some TLS logic to avoid thread group for doing atomic for each thread. (gathering the actual required count)
-			InterlockedAdd(RWInstanceCounts[WriteInstanceCountOffset], (uint)1U, RetPrevIdx);
+			InterlockedAdd(RWDataSetIndices[DataSetID*5+1], (uint)1U, RetPrevIdx);
 			PrevIdx = (int)RetPrevIdx;
 		}
 		return PrevIdx;
@@ -429,7 +425,7 @@ void SimulateMainComputeCS(
 	*/
 
 	uint InstanceID = UpdateStartInstance + DispatchThreadId.x;
-	GSpawnStartInstance = RWInstanceCounts[ReadInstanceCountOffset];				// needed by ExecIndex()
+	GSpawnStartInstance = ReadDataSetIndices[1];				// needed by ExecIndex()
 	int MaxInstances = GSpawnStartInstance + SpawnedInstances;
 
 	if(InstanceID < UpdateStartInstance + MaxInstances )
@@ -465,6 +461,15 @@ void SimulateMainComputeCS(
 			SimulateMain(InstanceID, EventIndex, 0);		// SimulateMain spawn phase
 		}
 	}
+
+	// fill in the rest of RWDataSetIndices for the DrawIndirect call for the main data set; number of instances is incremented in AcquireIndex
+	if(InstanceID == 0)
+	{
+		RWDataSetIndices[0] = NumIndicesPerInstance;		// IndexCountPerInstance; passed in from the renderer - TODO: support multiple renderers with multiple indirect draw calls
+		RWDataSetIndices[2] = 0;		// StartIndexLocation
+		RWDataSetIndices[3] = 0;		// BaseVertexLocation
+		RWDataSetIndices[4] = 0;		// StartInstanceLocation
+	}
 }
 
 #endif//GPU_SIMULATION
diff --git a/Engine/Plugins/FX/Niagara/Shaders/Private/NiagaraSpriteVertexFactory.ush b/Engine/Plugins/FX/Niagara/Shaders/Private/NiagaraSpriteVertexFactory.ush
--- a/Engine/Plugins/FX/Niagara/Shaders/Private/NiagaraSpriteVertexFactory.ush
+++ b/Engine/Plugins/FX/Niagara/Shaders/Private/NiagaraSpriteVertexFactory.ush
@@ -60,7 +60,7 @@ FVertexFactoryInput LoadVertexFactoryInput(uint TriangleIndex, int VertexIndex)
 }
 uint GetNumRayTracingDynamicMeshVerticesIndirect()
 {
-	return 6 * NiagaraSpriteVFLooseParameters.IndirectArgsBuffer[NiagaraSpriteVFLooseParameters.IndirectArgsOffset + 1];
+	return 6 * NiagaraSpriteVFLooseParameters.IndirectArgsBuffer[1];
 }
 #endif
 
diff --git a/Engine/Plugins/FX/Niagara/Source/Niagara/Classes/NiagaraDataSet.h b/Engine/Plugins/FX/Niagara/Source/Niagara/Classes/NiagaraDataSet.h
--- a/Engine/Plugins/FX/Niagara/Source/Niagara/Classes/NiagaraDataSet.h
+++ b/Engine/Plugins/FX/Niagara/Source/Niagara/Classes/NiagaraDataSet.h
@@ -26,8 +26,6 @@ struct FNiagaraVariableLayoutInfo
 
 class FNiagaraDataSet;
 class FNiagaraShader;
-class FNiagaraGPUInstanceCountManager;
-struct FNiagaraComputeExecutionContext;
 
 //Base class for objects in Niagara that are owned by one object but are then passed for reading to other objects, potentially on other threads.
 //This class allows us to know if the object is being used so we do not overwrite it and to ensure it's lifetime so we do not access freed data.
@@ -96,7 +94,7 @@ protected:
 public:
 	FNiagaraDataBuffer(FNiagaraDataSet* InOwner);
 	void Allocate(uint32 NumInstances, bool bMaintainExisting = false);
-	void AllocateGPU(uint32 InNumInstances, FNiagaraGPUInstanceCountManager& GPUInstanceCountManager, FRHICommandList &RHICmdList);
+	void AllocateGPU(uint32 InNumInstances, FRHICommandList &RHICmdList);
 	void SwapInstances(uint32 OldIndex, uint32 NewIndex);
 	void KillInstance(uint32 InstanceIdx);
 	void CopyTo(FNiagaraDataBuffer& DestBuffer, int32 StartIdx, int32 NumInstances)const;
@@ -131,8 +129,7 @@ public:
 	FORCEINLINE uint32 GetSizeBytes()const { return FloatData.Num() + Int32Data.Num(); }
 	FORCEINLINE FRWBuffer& GetGPUBufferFloat() { return GPUBufferFloat; }
 	FORCEINLINE FRWBuffer& GetGPUBufferInt() { return GPUBufferInt;	}
-	FORCEINLINE uint32 GetGPUInstanceCountBufferOffset() const { return GPUInstanceCountBufferOffset; }
-	FORCEINLINE void ClearGPUInstanceCountBufferOffset() { GPUInstanceCountBufferOffset = INDEX_NONE; }
+	FORCEINLINE FRWBuffer& GetGPUIndices() { return GPUIndices; }
 
 	FORCEINLINE int32 GetSafeComponentBufferSize() const { return GetSafeComponentBufferSize(GetNumInstancesAllocated()); }
 	FORCEINLINE uint32 GetFloatStride() const { return FloatStride; }
@@ -149,9 +146,6 @@ public:
 	template<bool bDoResourceTransitions>
 	void SetShaderParams(class FNiagaraShader *Shader, FRHICommandList &CommandList, bool bInput);
 	void UnsetShaderParams(class FNiagaraShader *Shader, FRHICommandList &CommandList);
-
-	void ReleaseGPUInstanceCount(FNiagaraGPUInstanceCountManager& GPUInstanceCountManager);
-
 private:
 
 	void CheckUsage(bool bReadOnly)const;
@@ -179,14 +173,13 @@ private:
 
 	//////////////////////////////////////////////////////////////////////////
 	// GPU Data
-	/** The buffer offset where the instance count is accumulated. */
-	uint32 GPUInstanceCountBufferOffset;
-	/** The num of allocated chunks, each being of size ALLOC_CHUNKSIZE */
 	uint32 NumChunksAllocatedForGPU;
 	/** GPU Buffer containing floating point values for GPU simulations. */
 	FRWBuffer GPUBufferFloat;
 	/** GPU Buffer containing floating point values for GPU simulations. */
 	FRWBuffer GPUBufferInt;
+	/** GPU Buffer containing args to allow indirect rendering of this buffer. */
+	FRWBuffer GPUIndices;
 	//////////////////////////////////////////////////////////////////////////
 
 	/** Number of instances in data. */
@@ -284,9 +277,6 @@ public:
 		return *DestinationData;
 	}
 
-	/** Release the GPU instance counts so that they can be reused */
-	void ReleaseGPUInstanceCounts(FNiagaraGPUInstanceCountManager& GPUInstanceCountManager);
-
 private:
 
 	void Reset();
@@ -1461,14 +1451,14 @@ class NIAGARA_API FScopedNiagaraDataSetGPUReadback
 {
 public:
 	FScopedNiagaraDataSetGPUReadback() {}
+	FScopedNiagaraDataSetGPUReadback(FNiagaraDataSet* DataSet) { ReadbackData(DataSet); }
 	~FScopedNiagaraDataSetGPUReadback();
 
-	void ReadbackData(class NiagaraEmitterInstanceBatcher* Batcher, FNiagaraDataSet* InDataSet);
+	void ReadbackData(FNiagaraDataSet* InDataSet);
 	uint32 GetNumInstances() const { check(DataSet != nullptr); return NumInstances; }
 
 private:
 	FNiagaraDataSet*	DataSet = nullptr;
 	FNiagaraDataBuffer* DataBuffer = nullptr;
-	NiagaraEmitterInstanceBatcher* Batcher = nullptr;
 	uint32				NumInstances = 0;
 };
diff --git a/Engine/Plugins/FX/Niagara/Source/Niagara/Classes/NiagaraEmitterInstanceBatcher.h b/Engine/Plugins/FX/Niagara/Source/Niagara/Classes/NiagaraEmitterInstanceBatcher.h
--- a/Engine/Plugins/FX/Niagara/Source/Niagara/Classes/NiagaraEmitterInstanceBatcher.h
+++ b/Engine/Plugins/FX/Niagara/Source/Niagara/Classes/NiagaraEmitterInstanceBatcher.h
@@ -21,7 +21,6 @@ the same VectorVM byte code / compute shader code
 #include "Runtime/Engine/Private/Particles/ParticleSortingGPU.h"
 #include "NiagaraGPUSortInfo.h"
 #include "NiagaraScriptExecutionContext.h"
-#include "NiagaraGPUInstanceCountManager.h"
 
 class FNiagaraIndicesVertexBuffer : public FParticleIndicesVertexBuffer
 {
@@ -78,7 +77,7 @@ public:
 	virtual void AddVectorField(UVectorFieldComponent* VectorFieldComponent) override {}
 	virtual void RemoveVectorField(UVectorFieldComponent* VectorFieldComponent) override {}
 	virtual void UpdateVectorField(UVectorFieldComponent* VectorFieldComponent) override {}
-	virtual void PreInitViews(FRHICommandListImmediate& RHICmdList) override;
+	virtual void PreInitViews() override;
 	virtual bool UsesGlobalDistanceField() const override;
 	virtual void PreRender(FRHICommandListImmediate& RHICmdList, const class FGlobalDistanceFieldParameterData* GlobalDistanceFieldParameterData) override;
 
@@ -103,6 +102,12 @@ public:
 		const class FShaderParametersMetadata* SceneTexturesUniformBufferStruct,
 		FUniformBufferRHIParamRef SceneTexturesUniformBuffer) override;
 
+	void ExecuteAll(FRHICommandList &RHICmdList, FUniformBufferRHIParamRef ViewUniformBuffer);
+	void TickSingle(const FNiagaraGPUSystemTick& Tick, FNiagaraComputeInstanceData *Instance, FRHICommandList &RHICmdList, FUniformBufferRHIParamRef ViewUniformBuffer) const;
+	void ResizeBuffersAndGatherResources(FOverlappableTicks& OverlappableTick, FRHICommandList& RHICmdList, FNiagaraBufferArray& DestDataBuffers, FNiagaraBufferArray& CurrDataBuffers, FNiagaraBufferArray& DestBufferIntFloat, FNiagaraBufferArray& CurrBufferIntFloat);
+
+	void DispatchAllOnCompute(FOverlappableTicks& OverlappableTick, FRHICommandList& RHICmdList, FUniformBufferRHIParamRef ViewUniformBuffer, FNiagaraBufferArray& DestDataBuffers, FNiagaraBufferArray& CurrDataBuffers, FNiagaraBufferArray& DestBufferIntFloat, FNiagaraBufferArray& CurrBufferIntFloat);
+
 	int32 AddSortedGPUSimulation(const FNiagaraGPUSortInfo& SortInfo);
 	void SortGPUParticles(FRHICommandListImmediate& RHICmdList);
 	void ResolveParticleSortBuffers(FRHICommandListImmediate& RHICmdList, int32 ResultBufferIndex);
@@ -126,17 +131,11 @@ public:
 				bool bCopyBeforeStart = false
 			) const;
 
+	void ResolveDatasetWrites(FRHICommandList &RHICmdList, FNiagaraComputeInstanceData *Context) const;
 	void ResizeCurrentBuffer(FRHICommandList &RHICmdList, FNiagaraComputeExecutionContext *Context, uint32 NewNumInstances, uint32 PrevNumInstances) const;
 
-	FORCEINLINE FNiagaraGPUInstanceCountManager& GetGPUInstanceCounterManager() { check(IsInRenderingThread()); return GPUInstanceCounterManager; }
-
 private:
-
-	void ExecuteAll(FRHICommandList &RHICmdList, FUniformBufferRHIParamRef ViewUniformBuffer, bool bSetReadback);
-	void TickSingle(const FNiagaraGPUSystemTick& Tick, FNiagaraComputeInstanceData *Instance, FRHICommandList &RHICmdList, FUniformBufferRHIParamRef ViewUniformBuffer, bool bSetReadback) const;
-	void ResizeBuffersAndGatherResources(FOverlappableTicks& OverlappableTick, FRHICommandList& RHICmdList, FNiagaraBufferArray& DestDataBuffers, FNiagaraBufferArray& CurrDataBuffers, FNiagaraBufferArray& DestBufferIntFloat, FNiagaraBufferArray& CurrBufferIntFloat);
-	void DispatchAllOnCompute(FOverlappableTicks& OverlappableTick, FRHICommandList& RHICmdList, FUniformBufferRHIParamRef ViewUniformBuffer, FNiagaraBufferArray& DestDataBuffers, FNiagaraBufferArray& CurrDataBuffers, FNiagaraBufferArray& DestBufferIntFloat, FNiagaraBufferArray& CurrBufferIntFloat, bool bSetReadback);
-
+	void SimStepClearAndSetup(const FNiagaraComputeInstanceData& Instance, FRHICommandList& RHICmdList) const;
 //	void SimStepReadback(const FNiagaraComputeInstanceData& Instance, FRHICommandList& RHICmdList) const;
 
 	inline uint32 UnpackEmitterDispatchCount(uint8* PackedData)
@@ -163,9 +162,6 @@ private:
 	TArray<FNiagaraGPUSortInfo> SimulationsToSort;
 	FParticleSortBuffers ParticleSortBuffers;
 
-	// GPU emitter instance count buffer. Contains the actual particle / instance count generate in the GPU tick.
-	mutable FNiagaraGPUInstanceCountManager GPUInstanceCounterManager;
-
 	// @todo REMOVE THIS HACK
 	uint32 LastFrameThatDrainedData;
 
diff --git a/Engine/Plugins/FX/Niagara/Source/Niagara/Classes/NiagaraGPUInstanceCountManager.h b/Engine/Plugins/FX/Niagara/Source/Niagara/Classes/NiagaraGPUInstanceCountManager.h
deleted file mode 100644
diff --git a/Engine/Plugins/FX/Niagara/Source/Niagara/Classes/NiagaraScriptExecutionContext.h b/Engine/Plugins/FX/Niagara/Source/Niagara/Classes/NiagaraScriptExecutionContext.h
--- a/Engine/Plugins/FX/Niagara/Source/Niagara/Classes/NiagaraScriptExecutionContext.h
+++ b/Engine/Plugins/FX/Niagara/Source/Niagara/Classes/NiagaraScriptExecutionContext.h
@@ -18,8 +18,6 @@ NiagaraEmitterInstance.h: Niagara emitter simulation class
 #include "RHIGPUReadback.h"
 
 struct FNiagaraDataInterfaceProxy;
-class FNiagaraGPUInstanceCountManager;
-class NiagaraEmitterInstanceBatcher;
 
 /** Container for data needed to process event data. */
 struct FNiagaraEventHandlingInfo
@@ -152,7 +150,7 @@ struct FNiagaraComputeExecutionContext
 	FNiagaraComputeExecutionContext();
 	~FNiagaraComputeExecutionContext();
 
-	void Reset(NiagaraEmitterInstanceBatcher* Batcher);
+	void Reset();
 
 	void InitParams(UNiagaraScript* InGPUComputeScript, ENiagaraSimTarget InSimTarget, const FString& InDebugSimName);
 	void DirtyDataInterfaces();
@@ -163,16 +161,8 @@ struct FNiagaraComputeExecutionContext
 	void SetDataToRender(FNiagaraDataBuffer* InDataToRender);
 	FNiagaraDataBuffer* GetDataToRender()const { return DataToRender; }
 
-	struct 
-	{
-		// The offset at which the GPU instance count (see FNiagaraGPUInstanceCountManager()).
-		uint32 GPUCountOffset = INDEX_NONE;
-		// The CPU instance count at the time the GPU count readback was issued. Always bigger or equal to the GPU count.
-		uint32 CPUCount = 0;
-	}  EmitterInstanceReadback;
-
 private:
-	void ResetInternal(NiagaraEmitterInstanceBatcher* Batcher);
+	void ResetInternal();
 
 public:
 	static uint32 TickCounter;
@@ -196,6 +186,11 @@ public:
 	//Most current buffer that can be used for rendering.
 	FNiagaraDataBuffer* DataToRender;
 
+	bool ResetSinceLastReadbackIssued = false;
+	FRHIGPUMemoryReadback *GPUDataReadback;
+	uint32 AccumulatedSpawnRate;
+	uint32 NumIndicesPerInstance;	// how many vtx indices per instance the renderer is going to have for its draw call
+
 	uint32 EventSpawnTotal_GT;
 	uint32 SpawnRateInstances_GT;
 
@@ -207,7 +202,6 @@ public:
 	mutable uint32 GPUDebugDataIntSize;
 	mutable uint32 GPUDebugDataFloatStride;
 	mutable uint32 GPUDebugDataIntStride;
-	mutable uint32 GPUDebugDataCountOffset;
 	mutable TSharedPtr<struct FNiagaraScriptDebuggerInfo, ESPMode::ThreadSafe> DebugInfo;
 #endif
 
diff --git a/Engine/Plugins/FX/Niagara/Source/Niagara/Private/NiagaraComponent.cpp b/Engine/Plugins/FX/Niagara/Source/Niagara/Private/NiagaraComponent.cpp
--- a/Engine/Plugins/FX/Niagara/Source/Niagara/Private/NiagaraComponent.cpp
+++ b/Engine/Plugins/FX/Niagara/Source/Niagara/Private/NiagaraComponent.cpp
@@ -120,17 +120,15 @@ void FNiagaraSceneProxy::ReleaseRenderers()
 {
 	if (EmitterRenderers.Num() > 0)
 	{
-		NiagaraEmitterInstanceBatcher* TheBatcher = Batcher && !Batcher->IsPendingKill() ? Batcher : nullptr;
-
 		//Renderers must be freed on the render thread.
 		ENQUEUE_RENDER_COMMAND(ReleaseRenderersCommand)(
-			[ToDeleteEmitterRenderers = MoveTemp(EmitterRenderers), TheBatcher](FRHICommandListImmediate& RHICmdList)
+			[ToDeleteEmitterRenderers = MoveTemp(EmitterRenderers)](FRHICommandListImmediate& RHICmdList)
 		{
 			for (FNiagaraRenderer* EmitterRenderer : ToDeleteEmitterRenderers)
 			{
 				if (EmitterRenderer)
 				{
-					EmitterRenderer->ReleaseRenderThreadResources(TheBatcher);
+					EmitterRenderer->ReleaseRenderThreadResources();
 					delete EmitterRenderer;
 				}
 			}
@@ -197,7 +195,7 @@ FNiagaraSceneProxy::~FNiagaraSceneProxy()
 	{
 		if (EmitterRenderer)
 		{
-			EmitterRenderer->ReleaseRenderThreadResources(Batcher);
+			EmitterRenderer->ReleaseRenderThreadResources();
 			delete EmitterRenderer;
 		}
 	}
@@ -210,7 +208,7 @@ void FNiagaraSceneProxy::ReleaseRenderThreadResources()
 	{
 		if (Renderer)
 		{
-			Renderer->ReleaseRenderThreadResources(Batcher);
+			Renderer->ReleaseRenderThreadResources();
 		}
 	}
 	return;
@@ -223,7 +221,7 @@ void FNiagaraSceneProxy::CreateRenderThreadResources()
 	{
 		if (Renderer)
 		{
-			Renderer->CreateRenderThreadResources(Batcher);
+			Renderer->CreateRenderThreadResources();
 		}
 	}
 	return;
diff --git a/Engine/Plugins/FX/Niagara/Source/Niagara/Private/NiagaraDataSet.cpp b/Engine/Plugins/FX/Niagara/Source/Niagara/Private/NiagaraDataSet.cpp
--- a/Engine/Plugins/FX/Niagara/Source/Niagara/Private/NiagaraDataSet.cpp
+++ b/Engine/Plugins/FX/Niagara/Source/Niagara/Private/NiagaraDataSet.cpp
@@ -8,8 +8,6 @@
 #include "ShaderParameterUtils.h"
 #include "NiagaraStats.h"
 #include "NiagaraRenderer.h"
-#include "NiagaraGPUInstanceCountManager.h"
-#include "NiagaraEmitterInstanceBatcher.h"
 
 DECLARE_CYCLE_STAT(TEXT("InitRenderData"), STAT_InitRenderData, STATGROUP_Niagara);
 
@@ -348,14 +346,6 @@ void FNiagaraDataSet::Dump(int32 StartIndex, int32 NumInstances, const FString&
 	}
 }
 
-void FNiagaraDataSet::ReleaseGPUInstanceCounts(FNiagaraGPUInstanceCountManager& GPUInstanceCountManager)
-{
-	for (FNiagaraDataBuffer* Buffer : Data)
-	{
-		Buffer->ReleaseGPUInstanceCount(GPUInstanceCountManager);
-	}
-}
-
 void FNiagaraDataSet::BuildLayout()
 {
 	VariableLayouts.Empty();
@@ -456,7 +446,6 @@ void FNiagaraDataSet::CopyFromGPUReadback(float* GPUReadBackFloat, int* GPUReadB
 
 FNiagaraDataBuffer::FNiagaraDataBuffer(FNiagaraDataSet* InOwner)
 	: Owner(InOwner)
-	, GPUInstanceCountBufferOffset(INDEX_NONE)
 	, NumChunksAllocatedForGPU(0)
 	, NumInstances(0)
 	, NumInstancesAllocated(0)
@@ -468,10 +457,6 @@ FNiagaraDataBuffer::FNiagaraDataBuffer(FNiagaraDataSet* InOwner)
 FNiagaraDataBuffer::~FNiagaraDataBuffer()
 {
 	check(!IsInUse());
-	// If this is data for a GPU emitter, we have to release the GPU instance counts for reuse.
-	// The only exception is if the batcher was pending kill and we couldn't enqueue a rendering command, 
-	// in which case this would have been released on the game thread and not from the batcher DataSetsToDestroy_RT.
-	check(!IsInRenderingThread() || GPUInstanceCountBufferOffset == INDEX_NONE);
 	DEC_MEMORY_STAT_BY(STAT_NiagaraParticleMemory, FloatData.Max() + Int32Data.Max());
 }
 
@@ -602,16 +587,18 @@ void FNiagaraDataBuffer::Allocate(uint32 InNumInstances, bool bMaintainExisting)
 	}
 }
 
-void FNiagaraDataBuffer::AllocateGPU(uint32 InNumInstances, FNiagaraGPUInstanceCountManager& GPUInstanceCountManager, FRHICommandList &RHICmdList)
+void FNiagaraDataBuffer::AllocateGPU(uint32 InNumInstances, FRHICommandList &RHICmdList)
 {
 	CheckUsage(false);
 
 	check(Owner->SimTarget == ENiagaraSimTarget::GPUComputeSim);
 
-	// Release previous entry if any.
-	GPUInstanceCountManager.FreeEntry(GPUInstanceCountBufferOffset);
-	// Get a new entry currently set to 0, since simulation will increment it to the actual instance count.
-	GPUInstanceCountBufferOffset = GPUInstanceCountManager.AcquireEntry();
+	//Allocate the indices buffer if needed
+	if (GPUIndices.Buffer == nullptr)
+	{
+		// Use BUF_KeepCPUAccessible here since some platforms will lock it for readonly (depending on the implementation of RHIEnqueueStagedRead) after GPU simulation.
+		GPUIndices.Initialize(sizeof(int32), 64 /*Context->NumDataSets*/, EPixelFormat::PF_R32_UINT, BUF_DrawIndirect | BUF_Static | BUF_KeepCPUAccessible);	// always allocate for up to 64 data sets
+	}
 
 	// ALLOC_CHUNKSIZE must be greater than zero and divisible by the thread group size
 	const uint32 ALLOC_CHUNKSIZE = 4096;
@@ -1011,12 +998,6 @@ void FNiagaraDataBuffer::UnsetShaderParams(FNiagaraShader *Shader, FRHICommandLi
 	}
 }
 
-void FNiagaraDataBuffer::ReleaseGPUInstanceCount(FNiagaraGPUInstanceCountManager& GPUInstanceCountManager)
-{
-	GPUInstanceCountManager.FreeEntry(GPUInstanceCountBufferOffset);
-}
-
-
 FScopedNiagaraDataSetGPUReadback::~FScopedNiagaraDataSetGPUReadback()
 {
 	if (DataBuffer != nullptr)
@@ -1026,12 +1007,11 @@ FScopedNiagaraDataSetGPUReadback::~FScopedNiagaraDataSetGPUReadback()
 	}
 }
 
-void FScopedNiagaraDataSetGPUReadback::ReadbackData(NiagaraEmitterInstanceBatcher* InBatcher, FNiagaraDataSet* InDataSet)
+void FScopedNiagaraDataSetGPUReadback::ReadbackData(FNiagaraDataSet* InDataSet)
 {
 	check(DataSet == nullptr);
 	check(InDataSet != nullptr);
 
-	Batcher = InBatcher && !InBatcher->IsPendingKill() ? Batcher : nullptr;
 	DataSet = InDataSet;
 	DataBuffer = DataSet->GetCurrentData();
 
@@ -1039,23 +1019,20 @@ void FScopedNiagaraDataSetGPUReadback::ReadbackData(NiagaraEmitterInstanceBatche
 	check((DataBuffer->FloatData.Num() == 0) && (DataBuffer->Int32Data.Num() == 0));
 
 	// Readback data
+	TArray<uint32> DrawIndirectData;
 	ENQUEUE_RENDER_COMMAND(ReadbackGPUBuffers)
 	(
 		[&](FRHICommandListImmediate& RHICmdList)
 		{
 			// Read DrawIndirect Params
-			const uint32 BufferOffset = DataBuffer->GetGPUInstanceCountBufferOffset();
-			if (Batcher && BufferOffset != INDEX_NONE)
 			{
-				FVertexBufferRHIParamRef InstanceCountBuffer = Batcher->GetGPUInstanceCounterManager().GetInstanceCountBuffer().Buffer;
+				const FRWBuffer& DataSetIndices = DataBuffer->GetGPUIndices();
+				const int32 BytesToCopy = DataSetIndices.NumBytes / sizeof(uint32);
+				DrawIndirectData.AddUninitialized(BytesToCopy);
 
-				void* Data = RHICmdList.LockVertexBuffer(InstanceCountBuffer, 0, (BufferOffset + 1) * sizeof(int32), RLM_ReadOnly);
-				NumInstances = reinterpret_cast<int32*>(Data)[BufferOffset];
-				RHICmdList.UnlockVertexBuffer(InstanceCountBuffer);
-			}
-			else
-			{
-				NumInstances = DataBuffer->GetNumInstances();
+				void* Data = RHICmdList.LockVertexBuffer(DataSetIndices.Buffer, 0, BytesToCopy, RLM_ReadOnly);
+				FMemory::Memcpy(DrawIndirectData.GetData(), Data, BytesToCopy);
+				RHICmdList.UnlockVertexBuffer(DataSetIndices.Buffer);
 			}
 
 			// Read float data
@@ -1082,4 +1059,6 @@ void FScopedNiagaraDataSetGPUReadback::ReadbackData(NiagaraEmitterInstanceBatche
 		}
 	);
 	FlushRenderingCommands();
+
+	NumInstances = DrawIndirectData[1];
 }
diff --git a/Engine/Plugins/FX/Niagara/Source/Niagara/Private/NiagaraEmitterInstance.cpp b/Engine/Plugins/FX/Niagara/Source/Niagara/Private/NiagaraEmitterInstance.cpp
--- a/Engine/Plugins/FX/Niagara/Source/Niagara/Private/NiagaraEmitterInstance.cpp
+++ b/Engine/Plugins/FX/Niagara/Source/Niagara/Private/NiagaraEmitterInstance.cpp
@@ -387,7 +387,7 @@ void FNiagaraEmitterInstance::ResetSimulation(bool bKillExisting)
 
 		if (CachedEmitter->SimTarget == ENiagaraSimTarget::GPUComputeSim && GPUExecContext != nullptr)
 		{
-			GPUExecContext->Reset(Batcher);
+			GPUExecContext->Reset();
 		}
 	}
 }
@@ -678,7 +678,7 @@ FBox FNiagaraEmitterInstance::CalculateDynamicBounds(const bool bReadGPUSimulati
 		if (!bReadGPUSimulation || (GPUExecContext == nullptr))
 			return FBox(ForceInit);
 
-		ScopedGPUReadback.ReadbackData(Batcher, GPUExecContext->MainDataSet);
+		ScopedGPUReadback.ReadbackData(GPUExecContext->MainDataSet);
 		NumInstances = ScopedGPUReadback.GetNumInstances();
 	}
 	else
@@ -1018,6 +1018,15 @@ void FNiagaraEmitterInstance::Tick(float DeltaSeconds)
 		check(GPUExecContext->GPUScript_RT == CachedEmitter->GetGPUComputeScript()->GetRenderThreadScript());
 		GPUExecContext->GPUScript_RT = CachedEmitter->GetGPUComputeScript()->GetRenderThreadScript();
 
+		if (CachedEmitter->GetRenderers().Num() > 0 && CachedEmitter->GetRenderers()[0] != nullptr)
+		{
+			GPUExecContext->NumIndicesPerInstance = CachedEmitter->GetRenderers()[0]->GetNumIndicesPerInstance();
+		}
+		else
+		{
+			GPUExecContext->NumIndicesPerInstance = 0;
+		}
+
 		GPUExecContext->EventSpawnTotal_GT = EventSpawnTotal;
 		GPUExecContext->SpawnRateInstances_GT = SpawnTotal;
 		
diff --git a/Engine/Plugins/FX/Niagara/Source/Niagara/Private/NiagaraEmitterInstanceBatcher.cpp b/Engine/Plugins/FX/Niagara/Source/Niagara/Private/NiagaraEmitterInstanceBatcher.cpp
--- a/Engine/Plugins/FX/Niagara/Source/Niagara/Private/NiagaraEmitterInstanceBatcher.cpp
+++ b/Engine/Plugins/FX/Niagara/Source/Niagara/Private/NiagaraEmitterInstanceBatcher.cpp
@@ -17,6 +17,7 @@
 DECLARE_CYCLE_STAT(TEXT("Niagara Dispatch Setup"), STAT_NiagaraGPUDispatchSetup_RT, STATGROUP_Niagara);
 DECLARE_CYCLE_STAT(TEXT("GPU Emitter Dispatch [RT]"), STAT_NiagaraGPUSimTick_RT, STATGROUP_Niagara);
 DECLARE_CYCLE_STAT(TEXT("GPU Data Readback [RT]"), STAT_NiagaraGPUReadback_RT, STATGROUP_Niagara);
+DECLARE_CYCLE_STAT(TEXT("Allocate GPU Readback Data [RT]"), STAT_NiagaraAllocateGPUReadback_RT, STATGROUP_Niagara);
 DECLARE_FLOAT_COUNTER_STAT(TEXT("Niagara GPU Sim"), STAT_GPU_NiagaraSim, STATGROUP_GPU);
 DECLARE_DWORD_COUNTER_STAT(TEXT("# GPU Particles"), STAT_NiagaraGPUParticles, STATGROUP_Niagara);
 DECLARE_DWORD_COUNTER_STAT(TEXT("# GPU Sorted Particles"), STAT_NiagaraGPUSortedParticles, STATGROUP_Niagara);
@@ -26,6 +27,8 @@ DECLARE_DWORD_COUNTER_STAT(TEXT("Readback latency (frames)"), STAT_NiagaraReadba
 DECLARE_GPU_STAT_NAMED(NiagaraGPU, TEXT("Niagara"));
 DECLARE_GPU_STAT_NAMED(NiagaraGPUSimulation, TEXT("Niagara GPU Simulation"));
 DECLARE_GPU_STAT_NAMED(NiagaraGPUSorting, TEXT("Niagara GPU sorting"));
+DECLARE_GPU_STAT_NAMED(NiagaraSimReadback, TEXT("Niagara GPU Simulation read back"));
+DECLARE_GPU_STAT_NAMED(NiagaraIndexBufferClear, TEXT("Niagara index buffer clear"));
 
 uint32 FNiagaraComputeExecutionContext::TickCounter = 0;
 
@@ -41,7 +44,7 @@ static FAutoConsoleVariableRef CVarNiagaraUseAsyncCompute(
 int32 GNiagaraGpuMaxQueuedRenderFrames = 10;
 static FAutoConsoleVariableRef CVarNiagaraGpuMaxQueuedRenderFrames(
 	TEXT("fx.NiagaraGpuMaxQueuedRenderFrames"),
-	GNiagaraGpuMaxQueuedRenderFrames,
+	GNiagaraOverlapCompute,
 	TEXT("Number of frames we all to pass before we start to discard GPU ticks.\n"),
 	ECVF_Default
 );
@@ -64,8 +67,9 @@ FFXSystemInterface* NiagaraEmitterInstanceBatcher::GetInterface(const FName& InN
 
 NiagaraEmitterInstanceBatcher::~NiagaraEmitterInstanceBatcher()
 {
-	FinishDispatches();
 	ParticleSortBuffers.ReleaseRHI();
+
+	FinishDispatches();
 }
 
 void NiagaraEmitterInstanceBatcher::GiveSystemTick_RenderThread(FNiagaraGPUSystemTick& Tick)
@@ -117,18 +121,12 @@ void NiagaraEmitterInstanceBatcher::FinishDispatches()
 
 	for (FNiagaraComputeExecutionContext* Context : ContextsToDestroy_RT)
 	{
-		check(Context);
-		// Put back the GPU instance counter the global pool.
-		GPUInstanceCounterManager.FreeEntry(Context->EmitterInstanceReadback.GPUCountOffset);
 		delete Context;
 	}
 	ContextsToDestroy_RT.Reset();
 
 	for (FNiagaraDataSet* DataSet : DataSetsToDestroy_RT)
 	{
-		check(DataSet);
-		// Put back the GPU instance counter the global pool.
-		DataSet->ReleaseGPUInstanceCounts(GPUInstanceCounterManager);
 		delete DataSet;
 	}
 	DataSetsToDestroy_RT.Reset();
@@ -194,7 +192,7 @@ void NiagaraEmitterInstanceBatcher::ResizeBuffersAndGatherResources(FOverlappabl
 			//If this is not true, the read back in ResolveDatasetWrites will shrink the buffers.
 			const uint32 RequiredInstances = FMath::Max(PrevNumInstances, NewNumInstances);
 
-			DestinationData.AllocateGPU(RequiredInstances + 1, GPUInstanceCounterManager, RHICmdList);
+			DestinationData.AllocateGPU(RequiredInstances + 1, RHICmdList);
 			DestinationData.SetNumInstances(RequiredInstances);
 
 			if ( Shader->FloatInputBufferParam.IsBound() )
@@ -224,7 +222,7 @@ void NiagaraEmitterInstanceBatcher::ResizeBuffersAndGatherResources(FOverlappabl
 	}
 }
 
-void NiagaraEmitterInstanceBatcher::DispatchAllOnCompute(FOverlappableTicks& OverlappableTick, FRHICommandList& RHICmdList, FUniformBufferRHIParamRef ViewUniformBuffer, FNiagaraBufferArray& DestDataBuffers, FNiagaraBufferArray& CurrDataBuffers, FNiagaraBufferArray& DestBufferIntFloat, FNiagaraBufferArray& CurrBufferIntFloat, bool bSetReadback)
+void NiagaraEmitterInstanceBatcher::DispatchAllOnCompute(FOverlappableTicks& OverlappableTick, FRHICommandList& RHICmdList, FUniformBufferRHIParamRef ViewUniformBuffer, FNiagaraBufferArray& DestDataBuffers, FNiagaraBufferArray& CurrDataBuffers, FNiagaraBufferArray& DestBufferIntFloat, FNiagaraBufferArray& CurrBufferIntFloat)
 {
 	FRHICommandListImmediate& RHICmdListImmediate = FRHICommandListExecutor::GetImmediateCommandList();
 
@@ -237,8 +235,10 @@ void NiagaraEmitterInstanceBatcher::DispatchAllOnCompute(FOverlappableTicks& Ove
 
 	RHICmdList.TransitionResources(EResourceTransitionAccess::ERWBarrier, EResourceTransitionPipeline::EGfxToCompute, CurrDataBuffers.GetData(), CurrDataBuffers.Num());
 
-#if WITH_EDITORONLY_DATA
 	{
+		SCOPED_DRAW_EVENT(RHICmdList, NiagaraIndexBufferClear);
+		SCOPED_GPU_STAT(RHICmdList, NiagaraIndexBufferClear);
+
 		for (FNiagaraGPUSystemTick* Tick : OverlappableTick)
 		{
 			uint32 DispatchCount = Tick->Count;
@@ -249,20 +249,24 @@ void NiagaraEmitterInstanceBatcher::DispatchAllOnCompute(FOverlappableTicks& Ove
 				FNiagaraComputeExecutionContext* Context = Instance.Context;
 				if (Context && Context->GPUScript_RT->GetShader())
 				{
+#if WITH_EDITORONLY_DATA
 					if (Context->DebugInfo.IsValid())
 					{
 						ProcessDebugInfo(RHICmdList, Context);
 					}
+#endif // WITH_EDITORONLY_DATA
+
+					// clear data set index buffer for the simulation shader to write number of written instances
+					ClearUAV(RHICmdList, Instance.DestinationData->GetGPUIndices(), 0);
 				}
+
 			}
 		}
 	}
-#endif // WITH_EDITORONLY_DATA
 
 	//
 	//	Add a rw barrier for the destination data buffers we just cleared and mark others as read/write as needed for particles simulation.
 	//
-	RHICmdList.TransitionResource(EResourceTransitionAccess::ERWBarrier, EResourceTransitionPipeline::EComputeToCompute, GPUInstanceCounterManager.GetInstanceCountBuffer().UAV);		// transition to readable; we'll be using this next frame
 	RHICmdList.TransitionResources(EResourceTransitionAccess::ERWBarrier, EResourceTransitionPipeline::EComputeToCompute, DestDataBuffers.GetData(), DestDataBuffers.Num());
 	RHICmdList.TransitionResources(EResourceTransitionAccess::EWritable, EResourceTransitionPipeline::EComputeToCompute, DestBufferIntFloat.GetData(), DestBufferIntFloat.Num());
 	RHICmdList.TransitionResources(EResourceTransitionAccess::EReadable, EResourceTransitionPipeline::EComputeToCompute, CurrDataBuffers.GetData(), CurrDataBuffers.Num());
@@ -287,28 +291,15 @@ void NiagaraEmitterInstanceBatcher::DispatchAllOnCompute(FOverlappableTicks& Ove
 					// run shader, sim and spawn in a single dispatch
 					uint32 UpdateStartInstance = 0;
 					Run<false>(*Tick, &Instance, UpdateStartInstance, Instance.DestinationData->GetNumInstances(), Context->GPUScript_RT->GetShader(), RHICmdList, ViewUniformBuffer);
-
-					FNiagaraDataBuffer* CurrentData = Instance.CurrentData;
-					if (bSetReadback && Tick->bIsFinalTick && Index == DispatchCount - 1)
-					{
-						// Now that the current data is not required anymore, stage it for readback.
-						if (CurrentData->GetNumInstances() && Context->EmitterInstanceReadback.GPUCountOffset == INDEX_NONE && CurrentData->GetGPUInstanceCountBufferOffset() != INDEX_NONE)
-						{
-							// Transfer the GPU instance counter ownership to the context. Note that when bSetReadback is true, a readback request will be performed later in the tick update.
-							Context->EmitterInstanceReadback.GPUCountOffset = CurrentData->GetGPUInstanceCountBufferOffset();
-							Context->EmitterInstanceReadback.CPUCount = CurrentData->GetNumInstances();
-							CurrentData->ClearGPUInstanceCountBufferOffset();
-						}
-					}
 				}
 			}
 		}
 	}
 
+
 	//
 	//	Now Copy to staging buffer the data we want to read back (alive particle count). And make buffer ready for that and draw commands on the graphics pipe too.
 	//
-	RHICmdList.TransitionResource(EResourceTransitionAccess::ERWBarrier, EResourceTransitionPipeline::EComputeToCompute, GPUInstanceCounterManager.GetInstanceCountBuffer().UAV);		// transition to readable; we'll be using this next frame
 	RHICmdList.TransitionResources(EResourceTransitionAccess::ERWBarrier, EResourceTransitionPipeline::EComputeToCompute, DestDataBuffers.GetData(), DestDataBuffers.Num());
 	RHICmdList.TransitionResources(EResourceTransitionAccess::ERWBarrier, EResourceTransitionPipeline::EComputeToCompute, DestBufferIntFloat.GetData(), DestBufferIntFloat.Num());
 	RHICmdList.TransitionResources(EResourceTransitionAccess::EReadable, EResourceTransitionPipeline::EComputeToGfx, DestDataBuffers.GetData(), DestDataBuffers.Num());
@@ -318,6 +309,32 @@ void NiagaraEmitterInstanceBatcher::DispatchAllOnCompute(FOverlappableTicks& Ove
 	// We have done all our overlapping compute work on this list so go back to default behavior and flush.
 	RHICmdList.AutomaticCacheFlushAfterComputeShader(true);
 
+	{
+		SCOPED_DRAW_EVENT(RHICmdList, NiagaraSimReadback);
+		SCOPED_GPU_STAT(RHICmdList, NiagaraSimReadback);
+		for (FNiagaraGPUSystemTick* Tick : OverlappableTick)
+		{
+			uint32 DispatchCount = Tick->Count;
+			FNiagaraComputeInstanceData* Instances = Tick->GetInstanceData();
+			for (uint32 Index = 0; Index < DispatchCount; Index++)
+			{
+				FNiagaraComputeInstanceData& Instance = Instances[Index];
+				FNiagaraComputeExecutionContext* Context = Instance.Context;
+				if (Context && Context->GPUScript_RT->GetShader())
+				{
+					// Don't resolve if the data if there are no instances (prevents a transition issue warning).
+					if (Instance.DestinationData->GetNumInstances() > 0)
+					{
+						// resolve data set writes - grabs the number of instances written from the index set during the simulation run
+						ResolveDatasetWrites(RHICmdList, &Instance);
+					}
+					check(Instance.DestinationData->GetGPUIndices().Buffer);
+				}
+			}
+		}
+	}
+	// the VF grabs current state for drawing, so need to transition
+	
 	//TODO: Need to set the data to render in the context?
 
 	RHICmdList.TransitionResources(EResourceTransitionAccess::EReadable, EResourceTransitionPipeline::EComputeToGfx, DestBufferIntFloat.GetData(), DestBufferIntFloat.Num());
@@ -329,16 +346,10 @@ void NiagaraEmitterInstanceBatcher::DispatchAllOnCompute(FOverlappableTicks& Ove
 
 void NiagaraEmitterInstanceBatcher::PostRenderOpaque(FRHICommandListImmediate& RHICmdList, const FUniformBufferRHIParamRef ViewUniformBuffer, const class FShaderParametersMetadata* SceneTexturesUniformBufferStruct, FUniformBufferRHIParamRef SceneTexturesUniformBuffer)
 {
-	// Setup new readback since if there is no pending request, there is no risk of having invalid data read (offset being allocated after the readback was sent).
-	ExecuteAll(RHICmdList, ViewUniformBuffer, !GPUInstanceCounterManager.HasPendingGPUReadback());
-
-	if (!GPUInstanceCounterManager.HasPendingGPUReadback())
-	{
-		GPUInstanceCounterManager.EnqueueGPUReadback(RHICmdList);
-	}
+	ExecuteAll(RHICmdList, ViewUniformBuffer);
 }
 
-void NiagaraEmitterInstanceBatcher::ExecuteAll(FRHICommandList &RHICmdList, FUniformBufferRHIParamRef ViewUniformBuffer, bool bSetReadback)
+void NiagaraEmitterInstanceBatcher::ExecuteAll(FRHICommandList &RHICmdList, FUniformBufferRHIParamRef ViewUniformBuffer)
 {
 	// @todo REMOVE THIS HACK
 	LastFrameThatDrainedData = GFrameNumberRenderThread;
@@ -351,6 +362,7 @@ void NiagaraEmitterInstanceBatcher::ExecuteAll(FRHICommandList &RHICmdList, FUni
 
 	SCOPED_DRAW_EVENT(RHICmdList, NiagaraEmitterInstanceBatcher_ExecuteAll);
 
+
 	if (GNiagaraOverlapCompute>0)
 	{
 		TArray<FOverlappableTicks, TInlineAllocator<2>> SimPasses;
@@ -397,7 +409,7 @@ void NiagaraEmitterInstanceBatcher::ExecuteAll(FRHICommandList &RHICmdList, FUni
 			// It also ensures the GPU buffers are large enough to hold everything.
 			ResizeBuffersAndGatherResources(SimPass, RHICmdList, DestDataBuffers, CurrDataBuffers, DestBufferIntFloat, CurrBufferIntFloat);
 
-			DispatchAllOnCompute(SimPass, RHICmdList, ViewUniformBuffer, DestDataBuffers, CurrDataBuffers, DestBufferIntFloat, CurrBufferIntFloat, bSetReadback);
+			DispatchAllOnCompute(SimPass, RHICmdList, ViewUniformBuffer, DestDataBuffers, CurrDataBuffers, DestBufferIntFloat, CurrBufferIntFloat);
 		}
 	}
 	else
@@ -410,7 +422,7 @@ void NiagaraEmitterInstanceBatcher::ExecuteAll(FRHICommandList &RHICmdList, FUni
 			for (uint32 i = 0; i < DispatchCount; i++)
 			{
 				FNiagaraComputeInstanceData& Instance = Instances[i];
-				TickSingle(Tick, &Instance, RHICmdList, ViewUniformBuffer, bSetReadback && i == DispatchCount - 1); // Only readback last count.
+				TickSingle(Tick, &Instance, RHICmdList, ViewUniformBuffer);
 			}
 		}
 	}
@@ -418,7 +430,62 @@ void NiagaraEmitterInstanceBatcher::ExecuteAll(FRHICommandList &RHICmdList, FUni
 	FinishDispatches();
 }
 
-void NiagaraEmitterInstanceBatcher::TickSingle(const FNiagaraGPUSystemTick& Tick, FNiagaraComputeInstanceData *Instance, FRHICommandList &RHICmdList, FUniformBufferRHIParamRef ViewUniformBuffer, bool bSetReadback) const
+void NiagaraEmitterInstanceBatcher::SimStepClearAndSetup(const FNiagaraComputeInstanceData& Instance, FRHICommandList& RHICmdList) const
+{
+	check(IsInRenderingThread());
+#if 1
+	FNiagaraComputeExecutionContext* Context = Instance.Context;
+	check(Context);
+
+
+	FNiagaraShader* ComputeShader = Context->GPUScript_RT->GetShader();
+	check(ComputeShader);
+	if (!ComputeShader)
+	{
+		return;
+	}
+
+#if WITH_EDITORONLY_DATA
+	if (Context->DebugInfo.IsValid())
+	{
+		ProcessDebugInfo(RHICmdList, Context);
+	}
+#endif // WITH_EDITORONLY_DATA
+
+	// clear data set index buffer for the simulation shader to write number of written instances to
+	//
+	FRWBuffer &DatasetIndexBufferWrite = Instance.DestinationData->GetGPUIndices();
+//	SCOPED_GPU_EVENTF(RHICmdList, NiagaraIndexBufferClear, TEXT("Niagara index buffer clear"));
+	SCOPED_DRAW_EVENT(RHICmdList, NiagaraIndexBufferClear);
+	SCOPED_GPU_STAT(RHICmdList, NiagaraIndexBufferClear);
+	//UE_LOG(LogNiagara, Log, TEXT("Clearing UAV %p"), (FRHIVertexBuffer *)DatasetIndexBufferWrite.Buffer);
+	ClearUAV(RHICmdList, DatasetIndexBufferWrite, 0);
+#endif
+}
+
+//void NiagaraEmitterInstanceBatcher::SimStepReadback(const FNiagaraComputeInstanceData& Instance, FRHICommandList& RHICmdList) const
+//{
+//#if 1
+//	FNiagaraComputeExecutionContext* Context = Instance.Context;
+//	check(Context);
+//	FNiagaraShader* ComputeShader = Context->GPUScript_RT->GetShader();
+//	if (!ComputeShader)
+//	{
+//		return;
+//	}
+//
+//	// Don't resolve if the data if there are no instances (prevents a transition issue warning).
+//	if (Context->MainDataSet->CurrData().GetNumInstances() > 0)
+//	{
+//		// resolve data set writes - grabs the number of instances written from the index set during the simulation run
+//		ResolveDatasetWrites(RHICmdList, &Instance);
+//	}
+//
+//	check(Context->MainDataSet->HasDatasetIndices());
+//#endif
+//}
+
+void NiagaraEmitterInstanceBatcher::TickSingle(const FNiagaraGPUSystemTick& Tick, FNiagaraComputeInstanceData *Instance, FRHICommandList &RHICmdList, FUniformBufferRHIParamRef ViewUniformBuffer) const
 {
 	SCOPE_CYCLE_COUNTER(STAT_NiagaraGPUSimTick_RT);
 
@@ -450,16 +517,38 @@ void NiagaraEmitterInstanceBatcher::TickSingle(const FNiagaraGPUSystemTick& Tick
 	FNiagaraDataBuffer& CurrentData = *Instance->CurrentData;
 	FNiagaraDataBuffer& DestinationData = *Instance->DestinationData;
 
-	const uint32 PrevNumInstances = Tick.bNeedsReset ? 0 : CurrentData.GetNumInstances();
-	const uint32 NewNumInstances = Instance->SpawnRateInstances + Instance->EventSpawnTotal + PrevNumInstances;
+	uint32 PrevNumInstances = Tick.bNeedsReset ? 0 : CurrentData.GetNumInstances();
+	uint32 NewNumInstances = Instance->SpawnRateInstances + Instance->EventSpawnTotal + PrevNumInstances;
+
+	//We must assume all particles survive when allocating here. 
+	//If this is not true, the read back in ResolveDatasetWrites will shrink the buffers.
+	uint32 RequiredInstances = FMath::Max(PrevNumInstances, NewNumInstances);
 
-	DestinationData.AllocateGPU(NewNumInstances + 1, GPUInstanceCounterManager, RHICmdList);
-	DestinationData.SetNumInstances(NewNumInstances);
+	DestinationData.AllocateGPU(RequiredInstances + 1, RHICmdList);
+	DestinationData.SetNumInstances(RequiredInstances);
+
+	// clear data set index buffer for the simulation shader to write number of written instances to
+	{
+		FRWBuffer& DatasetIndexBufferWrite = DestinationData.GetGPUIndices();
+		SCOPED_DRAW_EVENTF(RHICmdList, NiagaraIndexBufferClear, TEXT("Niagara index buffer clear"));
+		ClearUAV(RHICmdList, DatasetIndexBufferWrite, 0);
+		RHICmdList.TransitionResource(EResourceTransitionAccess::EWritable, EResourceTransitionPipeline::EComputeToCompute, DatasetIndexBufferWrite.UAV);
+	}
 
 	// run shader, sim and spawn in a single dispatch
 	uint32 UpdateStartInstance = 0;
 	Run<true>(Tick, Instance, UpdateStartInstance, NewNumInstances, ComputeShader, RHICmdList, ViewUniformBuffer);
 
+	// ResolveDatasetWrites may read this, so we must transition it here.
+	RHICmdList.TransitionResource(EResourceTransitionAccess::EReadable, EResourceTransitionPipeline::EComputeToGfx, DestinationData.GetGPUIndices().UAV);		// transition to readable; we'll be using this next frame
+
+	// Don't resolve if the data if there are no instances (prevents a transition issue warning).
+	if (NewNumInstances > 0)
+	{
+		// resolve data set writes - grabs the number of instances written from the index set during the simulation run
+		ResolveDatasetWrites(RHICmdList, Instance);
+	}
+
 	Context->MainDataSet->EndSimulate();
 	Context->SetDataToRender(Instance->DestinationData);
 
@@ -476,21 +565,9 @@ void NiagaraEmitterInstanceBatcher::TickSingle(const FNiagaraGPUSystemTick& Tick
 	// Better to do this in the VF set?
 	RHICmdList.TransitionResource(EResourceTransitionAccess::EReadable, EResourceTransitionPipeline::EComputeToGfx, Context->DataToRender->GetGPUBufferFloat().UAV);
 	RHICmdList.TransitionResource(EResourceTransitionAccess::EReadable, EResourceTransitionPipeline::EComputeToGfx, Context->DataToRender->GetGPUBufferInt().UAV);
-
-	// Put a barrier so that next tick or arg gen gets final count.
-	RHICmdList.TransitionResource(EResourceTransitionAccess::ERWBarrier, EResourceTransitionPipeline::EComputeToCompute, GPUInstanceCounterManager.GetInstanceCountBuffer().UAV);		// transition to readable; we'll be using this next frame
-
-	// Now that the current data is not required anymore, stage it for readback.
-	if (bSetReadback && PrevNumInstances && Context->EmitterInstanceReadback.GPUCountOffset == INDEX_NONE && CurrentData.GetGPUInstanceCountBufferOffset() != INDEX_NONE)
-	{
-		// Transfer the GPU instance counter ownership to the context. Note that when bSetReadback is true, a readback request will be performed later in the tick update.
-		Context->EmitterInstanceReadback.GPUCountOffset = CurrentData.GetGPUInstanceCountBufferOffset();
-		Context->EmitterInstanceReadback.CPUCount = PrevNumInstances;
-		CurrentData.ClearGPUInstanceCountBufferOffset();
-	}
 }
 
-void NiagaraEmitterInstanceBatcher::PreInitViews(FRHICommandListImmediate& RHICmdList)
+void NiagaraEmitterInstanceBatcher::PreInitViews()
 {
 	SortedParticleCount = 0;
 	SimulationsToSort.Reset();
@@ -499,57 +576,6 @@ void NiagaraEmitterInstanceBatcher::PreInitViews(FRHICommandListImmediate& RHICm
 	{
 		SortedVertexBuffer.UsedIndexCount = 0;
 	}
-
-	// Update draw indirect buffer to max possible size.
-	int32 TotalDispatchCount = 0;
-	for (FNiagaraGPUSystemTick& Tick : Ticks_RT)
-	{
-		TotalDispatchCount += (int32)Tick.Count;
-
-		// Cancel any pending readback if the emitter is resetting.
-		if (Tick.bNeedsReset)
-		{
-			FNiagaraComputeExecutionContext* Context = Tick.GetInstanceData()->Context;
-			if (Context)
-			{
-				GPUInstanceCounterManager.FreeEntry(Context->EmitterInstanceReadback.GPUCountOffset);
-			}
-		}
-	}
-	GPUInstanceCounterManager.ResizeBuffers(TotalDispatchCount);
-
-	// Update the instance counts from the GPU readback.
-	{
-		SCOPE_CYCLE_COUNTER(STAT_NiagaraGPUReadback_RT);
-		const uint32* Counts = GPUInstanceCounterManager.GetGPUReadback();
-		if (Counts)
-		{
-			for (FNiagaraGPUSystemTick& Tick : Ticks_RT)
-			{
-				FNiagaraComputeExecutionContext* Context = Tick.GetInstanceData()->Context;
-				if (Context && Context->EmitterInstanceReadback.GPUCountOffset != INDEX_NONE)
-				{
-					check(Context->MainDataSet);
-					FNiagaraDataBuffer* CurrentData = Context->MainDataSet->GetCurrentData();
-					if (CurrentData)
-					{
-						const uint32 DeadInstanceCount = Context->EmitterInstanceReadback.CPUCount - Counts[Context->EmitterInstanceReadback.GPUCountOffset];
-						if (DeadInstanceCount && DeadInstanceCount < CurrentData->GetNumInstances())
-						{
-							CurrentData->SetNumInstances(CurrentData->GetNumInstances() - DeadInstanceCount);
-						}
-					}
-
-					// Now release the readback since another one will be enqueued in the tick.
-					// Also prevents processing the same data again.
-					GPUInstanceCounterManager.FreeEntry(Context->EmitterInstanceReadback.GPUCountOffset);
-				}
-			}
-			// Readback is only valid for one frame, so that any newly allocated instance count
-			// are guarantied to be in the next valid readback data.
-			GPUInstanceCounterManager.ReleaseGPUReadback();
-		}
-	}
 }
 
 bool NiagaraEmitterInstanceBatcher::UsesGlobalDistanceField() const
@@ -571,10 +597,6 @@ void NiagaraEmitterInstanceBatcher::PreRender(FRHICommandListImmediate& RHICmdLi
 
 	// Sort buffer after mesh batches are issued, before tick (which will change the GPU instance count).
 	SortGPUParticles(RHICmdList);
-
-	// Update draw indirect args from the simulation results.
-	GPUInstanceCounterManager.UpdateDrawIndirectBuffer(RHICmdList, FeatureLevel);
-
 }
 
 int32 NiagaraEmitterInstanceBatcher::AddSortedGPUSimulation(const FNiagaraGPUSortInfo& SortInfo)
@@ -767,6 +789,69 @@ void NiagaraEmitterInstanceBatcher::ResolveParticleSortBuffers(FRHICommandListIm
 	CopyBufferCS->UnbindBuffers(RHICmdList);
 }
 
+/* Attempt to read back simulation results (number of live instances) from the GPU via an async readback request; 
+ *	If the readback isn't ready to be performed, we accumulate spawn rates and assume all instances have survived, until
+ *	the GPU can tell us how many are actually alive; since that data may be several frames old, we'll always end up
+ *	overallocating a bit, and the CPU might think we have more particles alive than we actually do;
+ *	since we use DrawIndirect with the GPU determining draw call parameters, that's not an issue
+ */
+void NiagaraEmitterInstanceBatcher::ResolveDatasetWrites(FRHICommandList &RHICmdList, FNiagaraComputeInstanceData *Instance) const
+{
+	FNiagaraComputeExecutionContext* Context = Instance->Context;
+	FNiagaraDataBuffer& DestinationData = *Instance->DestinationData;
+	FRWBuffer &DatasetIndexBufferWrite = DestinationData.GetGPUIndices();
+	uint32 SpawnedThisFrame = Instance->SpawnRateInstances + Instance->EventSpawnTotal;
+	Context->AccumulatedSpawnRate += SpawnedThisFrame;
+	int32 ExistingDataCount = DestinationData.GetNumInstances();
+	if (!Context->GPUDataReadback)
+	{
+		SCOPE_CYCLE_COUNTER(STAT_NiagaraAllocateGPUReadback_RT);
+
+		Context->GPUDataReadback = new FRHIGPUMemoryReadback(TEXT("Niagara GPU Emitter Readback"));
+		INC_DWORD_STAT(STAT_NiagaraReadbackLatency);
+		Context->GPUDataReadback->EnqueueCopy(RHICmdList, DatasetIndexBufferWrite.Buffer);
+		INC_DWORD_STAT_BY(STAT_NiagaraGPUParticles, ExistingDataCount);
+	}
+	else if (Context->GPUDataReadback->IsReady())
+	{
+		bool bSuccessfullyRead = false;
+		{
+		    SCOPE_CYCLE_COUNTER(STAT_NiagaraGPUReadback_RT);
+		    int32 *NumInstancesAfterSim = static_cast<int32*>(Context->GPUDataReadback->Lock(64 * sizeof(int32)));
+			if (NumInstancesAfterSim)
+			{
+				int32 NewExistingDataCount = (Context->ResetSinceLastReadbackIssued ? 0 : NumInstancesAfterSim[1]) + Context->AccumulatedSpawnRate; // index 1 is always the count
+				DestinationData.SetNumInstances(NewExistingDataCount);
+				//UE_LOG(LogNiagara, Log, TEXT("GPU Syncup %s : Was(%d) Now(%d) Accumulated(%d) NumInstancesAfterSim[1](%d)"), *Context->DebugSimName, ExistingDataCount, NewExistingDataCount , Context->AccumulatedSpawnRate, NumInstancesAfterSim[1]);
+				//UE_LOG(LogNiagara, Log, TEXT("Reading UAV %p"), (FRHIVertexBuffer *)DatasetIndexBufferWrite.Buffer);
+
+				INC_DWORD_STAT_BY(STAT_NiagaraGPUParticles, NewExistingDataCount);
+				SET_DWORD_STAT(STAT_NiagaraReadbackLatency, 0);
+
+				Context->AccumulatedSpawnRate = 0;
+				Context->ResetSinceLastReadbackIssued = false;
+				bSuccessfullyRead = true;
+			}
+			else
+			{
+				UE_LOG(LogNiagara, Warning, TEXT("GPUDataReadback said it was ready, but returned an invalid buffer. Skipping this time.."));
+				INC_DWORD_STAT_BY(STAT_NiagaraGPUParticles, ExistingDataCount);
+			}
+			Context->GPUDataReadback->Unlock();
+		}
+		if (bSuccessfullyRead)
+		{
+			SCOPE_CYCLE_COUNTER(STAT_NiagaraAllocateGPUReadback_RT);
+			// The following code seems to take significant time on d3d12
+			Context->GPUDataReadback->EnqueueCopy(RHICmdList, DatasetIndexBufferWrite.Buffer);
+		}
+	}
+	else
+	{
+		INC_DWORD_STAT_BY(STAT_NiagaraGPUParticles, ExistingDataCount);
+	}
+}
+
 void NiagaraEmitterInstanceBatcher::ProcessDebugInfo(FRHICommandList &RHICmdList, const FNiagaraComputeExecutionContext* Context) const
 {
 #if WITH_EDITORONLY_DATA
@@ -787,7 +872,8 @@ void NiagaraEmitterInstanceBatcher::ProcessDebugInfo(FRHICommandList &RHICmdList
 		{
 			//UE_LOG(LogNiagara, Warning, TEXT("Read back!"));
 
-			int32 NewExistingDataCount =  static_cast<int32*>(Context->GPUDebugDataReadbackCounts->Lock((Context->GPUDebugDataCountOffset + 1) * sizeof(int32)))[Context->GPUDebugDataCountOffset];
+			int32 *NumInstancesAfterSim = static_cast<int32*>(Context->GPUDebugDataReadbackCounts->Lock(64 * sizeof(int32)));
+			int32 NewExistingDataCount = NumInstancesAfterSim[1];
 			{
 				float* FloatDataBuffer = nullptr;
 				if (Context->GPUDebugDataReadbackFloat)
@@ -833,7 +919,6 @@ void NiagaraEmitterInstanceBatcher::ProcessDebugInfo(FRHICommandList &RHICmdList
 				Context->GPUDebugDataIntSize = 0;
 				Context->GPUDebugDataFloatStride = 0;
 				Context->GPUDebugDataIntStride = 0;
-				Context->GPUDebugDataCountOffset = INDEX_NONE;
 			}
 
 			// We've updated the debug info directly, now we need to no longer keep asking and querying because this frame is done!
@@ -931,8 +1016,32 @@ void NiagaraEmitterInstanceBatcher::Run(const FNiagaraGPUSystemTick& Tick, const
 	FNiagaraDataBuffer& DestinationData = *Instance->DestinationData;
 	FNiagaraDataBuffer& CurrentData = *Instance->CurrentData;
 
+	const FRWBuffer& WriteIndexBuffer = DestinationData.GetGPUIndices();
+	FRWBuffer& ReadIndexBuffer = CurrentData.GetGPUIndices();
+
+	// if we don't have a previous index buffer, we need to prep one using the maximum number of instances; this should only happen on the first frame
+	//		the data set index buffer is really the param buffer for the indirect draw call; it contains the number of live instances at index 1, and the simulation
+	//		CS uses this to determine the current number of active instances in the buffer
+	//
+	if (ReadIndexBuffer.Buffer == nullptr)
+	{
+		TResourceArray<int32> InitIndexBuffer;
+		InitIndexBuffer.AddUninitialized(64);
+		InitIndexBuffer[1] = 0;		// number of instances 
+		ReadIndexBuffer.Initialize(sizeof(int32), 64, EPixelFormat::PF_R32_UINT, BUF_DrawIndirect | BUF_Static, nullptr, &InitIndexBuffer);
+	}
+	else if (Tick.bNeedsReset)
+	{
+		//UE_LOG(LogNiagara, Log, TEXT("Clearing UAV due to reset!"));
+		ClearUAV(RHICmdList, ReadIndexBuffer, 0);
+		Context->AccumulatedSpawnRate = 0;
+		Context->ResetSinceLastReadbackIssued = true;
+	}
+
 	RHICmdList.SetComputeShader(Shader->GetComputeShader());
 
+	RHICmdList.SetShaderResourceViewParameter(Shader->GetComputeShader(), Shader->InputIndexBufferParam.GetBaseIndex(), ReadIndexBuffer.SRV);
+
 	// set the view uniform buffer param
 	if (Shader->ViewUniformBufferParam.IsBound() && ViewUniformBuffer)
 	{
@@ -948,13 +1057,9 @@ void NiagaraEmitterInstanceBatcher::Run(const FNiagaraGPUSystemTick& Tick, const
 
 	// set the index buffer uav
 	//
-	if (Shader->InstanceCountsParam.IsBound())
+	if (Shader->OutputIndexBufferParam.IsBound())
 	{
-		RHICmdList.SetUAVParameter(Shader->GetComputeShader(), Shader->InstanceCountsParam.GetUAVIndex(), GPUInstanceCounterManager.GetInstanceCountBuffer().UAV);
-		const uint32 ReadOffset = CurrentData.GetGPUInstanceCountBufferOffset();
-		const uint32 WriteOffset = DestinationData.GetGPUInstanceCountBufferOffset();
-		RHICmdList.SetShaderParameter(Shader->GetComputeShader(), Shader->ReadInstanceCountOffsetParam.GetBufferIndex(), Shader->ReadInstanceCountOffsetParam.GetBaseIndex(), Shader->ReadInstanceCountOffsetParam.GetNumBytes(), &ReadOffset);
-		RHICmdList.SetShaderParameter(Shader->GetComputeShader(), Shader->WriteInstanceCountOffsetParam.GetBufferIndex(), Shader->WriteInstanceCountOffsetParam.GetBaseIndex(), Shader->WriteInstanceCountOffsetParam.GetNumBytes(), &WriteOffset);
+		RHICmdList.SetUAVParameter(Shader->GetComputeShader(), Shader->OutputIndexBufferParam.GetUAVIndex(), WriteIndexBuffer.UAV);
 	}
 
 	// set the execution parameters
@@ -964,7 +1069,10 @@ void NiagaraEmitterInstanceBatcher::Run(const FNiagaraGPUSystemTick& Tick, const
 		RHICmdList.SetShaderParameter(Shader->GetComputeShader(), Shader->EmitterTickCounterParam.GetBufferIndex(), Shader->EmitterTickCounterParam.GetBaseIndex(), Shader->EmitterTickCounterParam.GetNumBytes(), &FNiagaraComputeExecutionContext::TickCounter);
 	}
 
+	uint32 Copy = bCopyBeforeStart ? 1 : 0;
+
 	RHICmdList.SetShaderParameter(Shader->GetComputeShader(), Shader->UpdateStartInstanceParam.GetBufferIndex(), Shader->UpdateStartInstanceParam.GetBaseIndex(), Shader->UpdateStartInstanceParam.GetNumBytes(), &UpdateStartInstance);					// 0, except for event handler runs
+	RHICmdList.SetShaderParameter(Shader->GetComputeShader(), Shader->NumIndicesPerInstanceParam.GetBufferIndex(), Shader->NumIndicesPerInstanceParam.GetBaseIndex(), Shader->NumIndicesPerInstanceParam.GetNumBytes(), &Context->NumIndicesPerInstance);		// set from the renderer in FNiagaraEmitterInstance::Tick
 	int32 InstancesToSpawnThisFrame = Instance->SpawnRateInstances + Instance->EventSpawnTotal;
 	RHICmdList.SetShaderParameter(Shader->GetComputeShader(), Shader->NumSpawnedInstancesParam.GetBufferIndex(), Shader->NumSpawnedInstancesParam.GetBaseIndex(), Shader->NumSpawnedInstancesParam.GetNumBytes(), &InstancesToSpawnThisFrame);				// number of instances in the spawn run
 
@@ -1003,8 +1111,10 @@ void NiagaraEmitterInstanceBatcher::Run(const FNiagaraGPUSystemTick& Tick, const
 	{
 		//UE_LOG(LogNiagara, Warning, TEXT("Queued up!"));
 
-		if (!Context->GPUDebugDataReadbackFloat && !Context->GPUDebugDataReadbackInt && !Context->GPUDebugDataReadbackCounts && DestinationData.GetGPUInstanceCountBufferOffset() != INDEX_NONE)
+		if (!Context->GPUDebugDataReadbackFloat && !Context->GPUDebugDataReadbackInt && !Context->GPUDebugDataReadbackCounts)
 		{
+			FRWBuffer &DatasetIndexBufferWrite = DestinationData.GetGPUIndices();
+
 			Context->GPUDebugDataFloatSize = 0;
 			Context->GPUDebugDataIntSize = 0;
 			Context->GPUDebugDataFloatStride = 0;
@@ -1027,8 +1137,7 @@ void NiagaraEmitterInstanceBatcher::Run(const FNiagaraGPUSystemTick& Tick, const
 			}
 
 			Context->GPUDebugDataReadbackCounts = new FRHIGPUMemoryReadback(TEXT("Niagara GPU Emitter Readback"));
-			Context->GPUDebugDataReadbackCounts->EnqueueCopy(RHICmdList, GPUInstanceCounterManager.GetInstanceCountBuffer().Buffer);
-			Context->GPUDebugDataCountOffset = DestinationData.GetGPUInstanceCountBufferOffset();
+			Context->GPUDebugDataReadbackCounts->EnqueueCopy(RHICmdList, DatasetIndexBufferWrite.Buffer);
 		}
 	}
 #endif // WITH_EDITORONLY_DATA
@@ -1038,5 +1147,5 @@ void NiagaraEmitterInstanceBatcher::Run(const FNiagaraGPUSystemTick& Tick, const
 	UnsetDataInterfaceParameters(DataInterfaceProxies, Shader, RHICmdList, Instance, Tick);
 	CurrentData.UnsetShaderParams(Shader, RHICmdList);
 	DestinationData.UnsetShaderParams(Shader, RHICmdList);
-	Shader->InstanceCountsParam.UnsetUAV(RHICmdList, Shader->GetComputeShader());
+	Shader->OutputIndexBufferParam.UnsetUAV(RHICmdList, Shader->GetComputeShader());
 }
diff --git a/Engine/Plugins/FX/Niagara/Source/Niagara/Private/NiagaraGPUInstanceCountManager.cpp b/Engine/Plugins/FX/Niagara/Source/Niagara/Private/NiagaraGPUInstanceCountManager.cpp
deleted file mode 100644
diff --git a/Engine/Plugins/FX/Niagara/Source/Niagara/Private/NiagaraMeshRendererProperties.cpp b/Engine/Plugins/FX/Niagara/Source/Niagara/Private/NiagaraMeshRendererProperties.cpp
--- a/Engine/Plugins/FX/Niagara/Source/Niagara/Private/NiagaraMeshRendererProperties.cpp
+++ b/Engine/Plugins/FX/Niagara/Source/Niagara/Private/NiagaraMeshRendererProperties.cpp
@@ -109,7 +109,7 @@ void UNiagaraMeshRendererProperties::GetUsedMaterials(TArray<UMaterialInterface*
 	}
 }
 
-uint32 UNiagaraMeshRendererProperties::GetNumIndicesPerInstance() const
+uint32 UNiagaraMeshRendererProperties::GetNumIndicesPerInstance() 
 {
 	// TODO: Add proper support for multiple mesh sections for GPU mesh particles.
 	return ParticleMesh ? ParticleMesh->RenderData->LODResources[0].Sections[0].NumTriangles * 3 : 0;
diff --git a/Engine/Plugins/FX/Niagara/Source/Niagara/Private/NiagaraRenderer.cpp b/Engine/Plugins/FX/Niagara/Source/Niagara/Private/NiagaraRenderer.cpp
--- a/Engine/Plugins/FX/Niagara/Source/Niagara/Private/NiagaraRenderer.cpp
+++ b/Engine/Plugins/FX/Niagara/Source/Niagara/Private/NiagaraRenderer.cpp
@@ -8,7 +8,6 @@
 #include "NiagaraVertexFactory.h"
 #include "Engine/Engine.h"
 #include "DynamicBufferAllocator.h"
-#include "NiagaraEmitterInstanceBatcher.h"
 
 DECLARE_CYCLE_STAT(TEXT("Sort Particles"), STAT_NiagaraSortParticles, STATGROUP_Niagara);
 
@@ -186,8 +185,6 @@ FNiagaraRenderer::FNiagaraRenderer(ERHIFeatureLevel::Type FeatureLevel, const UN
 	, CPUTimeMS(0.0f)
 	, bLocalSpace(Emitter->GetCachedEmitter()->bLocalSpace)
 	, bHasLights(false)
-	, SimTarget(Emitter->GetCachedEmitter()->SimTarget)
-	, NumIndicesPerInstance(InProps ? InProps->GetNumIndicesPerInstance() : 0)
 {
 #if STATS
 	EmitterStatID = Emitter->GetCachedEmitter()->GetStatID(false, false);
@@ -210,26 +207,10 @@ void FNiagaraRenderer::Initialize(ERHIFeatureLevel::Type FeatureLevel, const UNi
 
 FNiagaraRenderer::~FNiagaraRenderer()
 {
-	ReleaseRenderThreadResources(nullptr);
+	ReleaseRenderThreadResources();
 	SetDynamicData_RenderThread(nullptr);
 }
 
-void FNiagaraRenderer::CreateRenderThreadResources(NiagaraEmitterInstanceBatcher* Batcher) 
-{
-	if (Batcher && SimTarget == ENiagaraSimTarget::GPUComputeSim)
-	{
-		Batcher->GetGPUInstanceCounterManager().IncrementMaxDrawIndirectCount();
-	}
-}
-
-void FNiagaraRenderer::ReleaseRenderThreadResources(NiagaraEmitterInstanceBatcher* Batcher)
-{
-	if (Batcher && SimTarget == ENiagaraSimTarget::GPUComputeSim)
-	{
-		Batcher->GetGPUInstanceCounterManager().DecrementMaxDrawIndirectCount();
-	}
-}
-
 FPrimitiveViewRelevance FNiagaraRenderer::GetViewRelevance(const FSceneView* View, const FNiagaraSceneProxy *SceneProxy)const
 {
 	FPrimitiveViewRelevance Result;
diff --git a/Engine/Plugins/FX/Niagara/Source/Niagara/Private/NiagaraRendererLights.cpp b/Engine/Plugins/FX/Niagara/Source/Niagara/Private/NiagaraRendererLights.cpp
--- a/Engine/Plugins/FX/Niagara/Source/Niagara/Private/NiagaraRendererLights.cpp
+++ b/Engine/Plugins/FX/Niagara/Source/Niagara/Private/NiagaraRendererLights.cpp
@@ -64,6 +64,7 @@ FNiagaraDynamicDataBase* FNiagaraRendererLights::GenerateDynamicData(const FNiag
 	FNiagaraDataSetAccessor<float> ScatteringAccessor(Data, Properties->VolumetricScatteringBinding.DataSetVariable);
 	FNiagaraDataSetAccessor<int32> EnabledAccessor(Data, Properties->LightRenderingEnabledBinding.DataSetVariable);
 
+	ENiagaraSimTarget SimTarget = Emitter->GetData().GetSimTarget();
 	FNiagaraDynamicDataLights* DynamicData = new FNiagaraDynamicDataLights(Emitter);
 
 	const FMatrix& LocalToWorldMatrix = Proxy->GetLocalToWorld();
diff --git a/Engine/Plugins/FX/Niagara/Source/Niagara/Private/NiagaraRendererMeshes.cpp b/Engine/Plugins/FX/Niagara/Source/Niagara/Private/NiagaraRendererMeshes.cpp
--- a/Engine/Plugins/FX/Niagara/Source/Niagara/Private/NiagaraRendererMeshes.cpp
+++ b/Engine/Plugins/FX/Niagara/Source/Niagara/Private/NiagaraRendererMeshes.cpp
@@ -107,16 +107,16 @@ FNiagaraRendererMeshes::FNiagaraRendererMeshes(ERHIFeatureLevel::Type FeatureLev
 	MeshMinimumLOD = Properties->ParticleMesh->MinLOD.GetValueForFeatureLevel(FeatureLevel);
 }
 
-void FNiagaraRendererMeshes::ReleaseRenderThreadResources(NiagaraEmitterInstanceBatcher* Batcher)
+void FNiagaraRendererMeshes::ReleaseRenderThreadResources()
 {
-	FNiagaraRenderer::ReleaseRenderThreadResources(Batcher);
+	FNiagaraRenderer::ReleaseRenderThreadResources();
 	VertexFactory->ReleaseResource();
 	WorldSpacePrimitiveUniformBuffer.ReleaseResource();
 }
 
-void FNiagaraRendererMeshes::CreateRenderThreadResources(NiagaraEmitterInstanceBatcher* Batcher)
+void FNiagaraRendererMeshes::CreateRenderThreadResources()
 {
-	FNiagaraRenderer::CreateRenderThreadResources(Batcher);
+	FNiagaraRenderer::CreateRenderThreadResources();
 	VertexFactory->InitResource();
 }
 
@@ -135,13 +135,11 @@ void FNiagaraRendererMeshes::GetDynamicMeshElements(const TArray<const FSceneVie
 {
 	SCOPE_CYCLE_COUNTER(STAT_NiagaraRender);
 	SCOPE_CYCLE_COUNTER(STAT_NiagaraRenderMeshes);
-	check(SceneProxy);
 
 	SimpleTimer MeshElementsTimer;
 
-	NiagaraEmitterInstanceBatcher* Batcher = SceneProxy->GetBatcher();
 	FNiagaraDynamicDataMesh *DynamicDataMesh = (static_cast<FNiagaraDynamicDataMesh*>(DynamicDataRender));
-	if (!DynamicDataMesh || !Batcher)
+	if (!DynamicDataMesh)
 	{
 		return;
 	}
@@ -168,18 +166,12 @@ void FNiagaraRendererMeshes::GetDynamicMeshElements(const TArray<const FSceneVie
 	FGlobalDynamicReadBuffer& DynamicReadBuffer = Collector.GetDynamicReadBuffer();
 	FGlobalDynamicReadBuffer::FAllocation ParticleData;
 
-	// @TODO : support multiple LOD and section, using an inlined array and/or the SceneRenderingAllocator
-	uint32 IndirectArgsOffset = INDEX_NONE;
 	//For cpu sims we allocate render buffers from the global pool. GPU sims own their own.
-	if (SimTarget == ENiagaraSimTarget::CPUSim)
+	if (DynamicDataMesh->GetSimTarget() == ENiagaraSimTarget::CPUSim)
 	{
 		ParticleData = DynamicReadBuffer.AllocateFloat(TotalFloatSize);
 		FMemory::Memcpy(ParticleData.Buffer, SourceParticleData->GetFloatBuffer().GetData(), SourceParticleData->GetFloatBuffer().Num());
 	}
-	else // ENiagaraSimTarget::GPUComputeSim
-	{
-		IndirectArgsOffset = Batcher->GetGPUInstanceCounterManager().AddDrawIndirect(SourceParticleData->GetGPUInstanceCountBufferOffset(), NumIndicesPerInstance);
-	}
 
 	{
 		// Update the primitive uniform buffer if needed.
@@ -268,23 +260,24 @@ void FNiagaraRendererMeshes::GetDynamicMeshElements(const TArray<const FSceneVie
 					}
 				};
 
-				if (SimTarget == ENiagaraSimTarget::CPUSim)//TODO: Compute shader for sorting gpu sims and larger cpu sims.
+				if (DynamicDataMesh->GetSimTarget() == ENiagaraSimTarget::CPUSim)//TODO: Compute shader for sorting gpu sims and larger cpu sims.
 				{
 					check(ParticleData.IsValid());
 					if (SortInfo.SortMode != ENiagaraSortMode::None && SortInfo.SortAttributeOffset != INDEX_NONE)
 					{
 						if (GNiagaraGPUSorting &&
 							GNiagaraGPUSortingCPUToGPUThreshold != INDEX_NONE &&
-							SortInfo.ParticleCount >= GNiagaraGPUSortingCPUToGPUThreshold)
+							SortInfo.ParticleCount >= GNiagaraGPUSortingCPUToGPUThreshold &&
+							SceneProxy->GetBatcher())
 						{
 							SortInfo.ParticleCount = NumInstances;
 							SortInfo.ParticleDataFloatSRV = ParticleData.ReadBuffer->SRV;
 							SortInfo.FloatDataOffset = ParticleData.FirstIndex / sizeof(float);
 							SortInfo.FloatDataStride = SourceParticleData->GetFloatStride() / sizeof(float);
-							const int32 IndexBufferOffset = Batcher->AddSortedGPUSimulation(SortInfo);
+							const int32 IndexBufferOffset = SceneProxy->GetBatcher()->AddSortedGPUSimulation(SortInfo);
 							if (IndexBufferOffset != INDEX_NONE)
 							{
-								CollectorResources.VertexFactory.SetSortedIndices(Batcher->GetGPUSortedBuffer().VertexBufferSRV, IndexBufferOffset);
+								CollectorResources.VertexFactory.SetSortedIndices(SceneProxy->GetBatcher()->GetGPUSortedBuffer().VertexBufferSRV, IndexBufferOffset);
 							}
 						}
 						else
@@ -299,7 +292,21 @@ void FNiagaraRendererMeshes::GetDynamicMeshElements(const TArray<const FSceneVie
 				}
 				else
 				{
-					if (SortInfo.SortMode != ENiagaraSortMode::None && SortInfo.SortAttributeOffset != INDEX_NONE && GNiagaraGPUSorting)
+					if (SortInfo.SortMode != ENiagaraSortMode::None && SortInfo.SortAttributeOffset != INDEX_NONE && GNiagaraGPUSorting && SceneProxy->GetBatcher())
+					{
+						SortInfo.ParticleDataFloatSRV = SourceParticleData->GetGPUBufferFloat().SRV;
+						SortInfo.FloatDataOffset = 0;
+						SortInfo.FloatDataStride = SourceParticleData->GetFloatStride() / sizeof(float);
+						SortInfo.GPUParticleCountSRV = SourceParticleData->GetGPUIndices().SRV;
+						SortInfo.GPUParticleCountOffset = 1;
+						const int32 IndexBufferOffset = SceneProxy->GetBatcher()->AddSortedGPUSimulation(SortInfo);
+						if (IndexBufferOffset != INDEX_NONE)
+						{
+							CollectorResources.VertexFactory.SetSortedIndices(SceneProxy->GetBatcher()->GetGPUSortedBuffer().VertexBufferSRV, IndexBufferOffset);
+						}
+					}
+
+					if (SortInfo.SortMode != ENiagaraSortMode::None && SortInfo.SortAttributeOffset != INDEX_NONE && GNiagaraGPUSorting && SceneProxy->GetBatcher())
 					{
 						// Here we need to be conservative about the InstanceCount, since the final value is only known on the GPU after the simulation.
 						SortInfo.ParticleCount = SourceParticleData->GetNumInstances();
@@ -307,12 +314,12 @@ void FNiagaraRendererMeshes::GetDynamicMeshElements(const TArray<const FSceneVie
 						SortInfo.ParticleDataFloatSRV = SourceParticleData->GetGPUBufferFloat().SRV;
 						SortInfo.FloatDataOffset = 0;
 						SortInfo.FloatDataStride = SourceParticleData->GetFloatStride() / sizeof(float);
-						SortInfo.GPUParticleCountSRV = Batcher->GetGPUInstanceCounterManager().GetInstanceCountBuffer().SRV;
-						SortInfo.GPUParticleCountOffset = SourceParticleData->GetGPUInstanceCountBufferOffset();
-						const int32 IndexBufferOffset = Batcher->AddSortedGPUSimulation(SortInfo);
-						if (IndexBufferOffset != INDEX_NONE && SortInfo.GPUParticleCountOffset != INDEX_NONE)
+						SortInfo.GPUParticleCountSRV = SourceParticleData->GetGPUIndices().SRV;
+						SortInfo.GPUParticleCountOffset = 1;
+						const int32 IndexBufferOffset = SceneProxy->GetBatcher()->AddSortedGPUSimulation(SortInfo);
+						if (IndexBufferOffset != INDEX_NONE)
 						{
-							CollectorResources.VertexFactory.SetSortedIndices(Batcher->GetGPUSortedBuffer().VertexBufferSRV, IndexBufferOffset);
+							CollectorResources.VertexFactory.SetSortedIndices(SceneProxy->GetBatcher()->GetGPUSortedBuffer().VertexBufferSRV, IndexBufferOffset);
 						}
 					}
 					CollectorResources.VertexFactory.SetParticleData(SourceParticleData->GetGPUBufferFloat().SRV, 0, SourceParticleData->GetFloatStride() / sizeof(float));
@@ -328,7 +335,7 @@ void FNiagaraRendererMeshes::GetDynamicMeshElements(const TArray<const FSceneVie
 			
 				// GPU mesh rendering currently only supports one mesh section.
 				// TODO: Add proper support for multiple mesh sections for GPU mesh particles.
-				int32 MaxSection = SimTarget == ENiagaraSimTarget::GPUComputeSim ? 1 : LODModel.Sections.Num();
+				int32 MaxSection = DynamicDataMesh->GetSimTarget() == ENiagaraSimTarget::GPUComputeSim ? 1 : LODModel.Sections.Num();
 				const bool bIsWireframe = AllowDebugViewmodes() && View && View->Family->EngineShowFlags.Wireframe;
 				for (int32 SectionIndex = 0; SectionIndex < MaxSection; SectionIndex++)
 				{
@@ -384,11 +391,10 @@ void FNiagaraRendererMeshes::GetDynamicMeshElements(const TArray<const FSceneVie
 						BatchElement.NumPrimitives = Section.NumTriangles;
 					}
 
-					if (IndirectArgsOffset != INDEX_NONE)
+					if (DynamicDataMesh->GetSimTarget() == ENiagaraSimTarget::GPUComputeSim)
 					{
 						BatchElement.NumPrimitives = 0;
-						BatchElement.IndirectArgsOffset = IndirectArgsOffset;
-						BatchElement.IndirectArgsBuffer = Batcher->GetGPUInstanceCounterManager().GetDrawIndirectBuffer().Buffer;
+						BatchElement.IndirectArgsBuffer = SourceParticleData->GetGPUIndices().Buffer;
 					}
 					else
 					{
@@ -426,6 +432,7 @@ FNiagaraDynamicDataBase *FNiagaraRendererMeshes::GenerateDynamicData(const FNiag
 
 	FNiagaraDynamicDataMesh *DynamicData = nullptr;
 
+	ENiagaraSimTarget SimTarget = Emitter->GetData().GetSimTarget();
 	if (Properties->ParticleMesh)
 	{
 		DynamicData = new FNiagaraDynamicDataMesh(Emitter);
diff --git a/Engine/Plugins/FX/Niagara/Source/Niagara/Private/NiagaraRendererRibbons.cpp b/Engine/Plugins/FX/Niagara/Source/Niagara/Private/NiagaraRendererRibbons.cpp
--- a/Engine/Plugins/FX/Niagara/Source/Niagara/Private/NiagaraRendererRibbons.cpp
+++ b/Engine/Plugins/FX/Niagara/Source/Niagara/Private/NiagaraRendererRibbons.cpp
@@ -198,17 +198,17 @@ FNiagaraRendererRibbons::FNiagaraRendererRibbons(ERHIFeatureLevel::Type FeatureL
 }
 
 
-void FNiagaraRendererRibbons::ReleaseRenderThreadResources(NiagaraEmitterInstanceBatcher* Batcher)
+void FNiagaraRendererRibbons::ReleaseRenderThreadResources()
 {
-	FNiagaraRenderer::ReleaseRenderThreadResources(Batcher);
+	FNiagaraRenderer::ReleaseRenderThreadResources();
 	VertexFactory->ReleaseResource();
 	WorldSpacePrimitiveUniformBuffer.ReleaseResource();
 }
 
 // FPrimitiveSceneProxy interface.
-void FNiagaraRendererRibbons::CreateRenderThreadResources(NiagaraEmitterInstanceBatcher* Batcher)
+void FNiagaraRendererRibbons::CreateRenderThreadResources()
 {
-	FNiagaraRenderer::CreateRenderThreadResources(Batcher);
+	FNiagaraRenderer::CreateRenderThreadResources();
 	VertexFactory->InitResource();
 }
 
@@ -301,7 +301,7 @@ void FNiagaraRendererRibbons::GetDynamicMeshElements(const TArray<const FSceneVi
 	int32 TotalFloatSize = SourceParticleData->GetFloatBuffer().Num() / sizeof(float);
 	FGlobalDynamicReadBuffer::FAllocation ParticleData;
 
-	if (SimTarget == ENiagaraSimTarget::CPUSim)
+	if (DynamicDataRibbon->GetSimTarget() == ENiagaraSimTarget::CPUSim)
 	{
 		FGlobalDynamicReadBuffer& DynamicReadBuffer = Collector.GetDynamicReadBuffer();
 		ParticleData = DynamicReadBuffer.AllocateFloat(TotalFloatSize);
@@ -569,7 +569,7 @@ FNiagaraDynamicDataBase* FNiagaraRendererRibbons::GenerateDynamicData(const FNia
 {
 	SCOPE_CYCLE_COUNTER(STAT_NiagaraGenRibbonVertexData);
 
-	if (SimTarget == ENiagaraSimTarget::GPUComputeSim)
+	if (Emitter->GetData().GetSimTarget() == ENiagaraSimTarget::GPUComputeSim)
 	{
 		return nullptr;
 	}
@@ -598,6 +598,7 @@ FNiagaraDynamicDataBase* FNiagaraRendererRibbons::GenerateDynamicData(const FNia
 	FNiagaraDataSetAccessor<int32> RibbonIdData;
 	FNiagaraDataSetAccessor<FNiagaraID> RibbonFullIDData;
 
+	ENiagaraSimTarget SimTarget = Data.GetSimTarget();
 	FNiagaraDataBuffer* DataToRender = Emitter->GetData().GetCurrentData();
 	if (DataToRender == nullptr || DataToRender->GetNumInstances() < 2 || !PosData.IsValid() || !SortKeyData.IsValid())
 	{
diff --git a/Engine/Plugins/FX/Niagara/Source/Niagara/Private/NiagaraRendererSprites.cpp b/Engine/Plugins/FX/Niagara/Source/Niagara/Private/NiagaraRendererSprites.cpp
--- a/Engine/Plugins/FX/Niagara/Source/Niagara/Private/NiagaraRendererSprites.cpp
+++ b/Engine/Plugins/FX/Niagara/Source/Niagara/Private/NiagaraRendererSprites.cpp
@@ -138,9 +138,9 @@ FNiagaraRendererSprites::FNiagaraRendererSprites(ERHIFeatureLevel::Type FeatureL
 	MaterialParamValidMask |= MaterialParamOffset3 != -1 ? 8 : 0;
 }
 
-void FNiagaraRendererSprites::ReleaseRenderThreadResources(NiagaraEmitterInstanceBatcher* Batcher)
+void FNiagaraRendererSprites::ReleaseRenderThreadResources()
 {
-	FNiagaraRenderer::ReleaseRenderThreadResources(Batcher);
+	FNiagaraRenderer::ReleaseRenderThreadResources();
 	VertexFactory->ReleaseResource();
 	WorldSpacePrimitiveUniformBuffer.ReleaseResource();
 
@@ -154,9 +154,9 @@ void FNiagaraRendererSprites::ReleaseRenderThreadResources(NiagaraEmitterInstanc
 #endif
 }
 
-void FNiagaraRendererSprites::CreateRenderThreadResources(NiagaraEmitterInstanceBatcher* Batcher)
+void FNiagaraRendererSprites::CreateRenderThreadResources()
 {
-	FNiagaraRenderer::CreateRenderThreadResources(Batcher);
+	FNiagaraRenderer::CreateRenderThreadResources();
 	CutoutVertexBuffer.InitResource();
 	VertexFactory->InitResource();
 
@@ -218,7 +218,7 @@ FNiagaraRendererSprites::FCPUSimParticleDataAllocation FNiagaraRendererSprites::
 
 	FCPUSimParticleDataAllocation CPUSimParticleDataAllocation { DynamicReadBuffer };
 
-	if (SimTarget == ENiagaraSimTarget::CPUSim)
+	if (DynamicDataSprites->GetSimTarget() == ENiagaraSimTarget::CPUSim)
 	{
 		SCOPE_CYCLE_COUNTER(STAT_NiagaraRenderSpritesCPUSimCopy);
 		CPUSimParticleDataAllocation.ParticleData = DynamicReadBuffer.AllocateFloat(TotalFloatSize);
@@ -316,9 +316,6 @@ void FNiagaraRendererSprites::SetVertexFactoryParticleData(
 	const FSceneView* View,
 	const FNiagaraSceneProxy *SceneProxy) const
 {
-	NiagaraEmitterInstanceBatcher* Batcher = SceneProxy->GetBatcher();
-	check(Batcher);
-
 	// Cutout geometry.
 	const bool bUseSubImage = SubImageSize.X != 1 || SubImageSize.Y != 1;
 	const bool bUseCutout = CutoutVertexBuffer.VertexBufferRHI.IsValid();
@@ -365,23 +362,24 @@ void FNiagaraRendererSprites::SetVertexFactoryParticleData(
 			}
 		};
 
-		if (SimTarget == ENiagaraSimTarget::CPUSim)//TODO: Compute shader for sorting gpu sims and larger cpu sims.
+		if (DynamicDataSprites->GetSimTarget() == ENiagaraSimTarget::CPUSim)//TODO: Compute shader for sorting gpu sims and larger cpu sims.
 		{
 			check(CPUSimParticleDataAllocation.ParticleData.IsValid());
 			if (SortInfo.SortMode != ENiagaraSortMode::None && SortInfo.SortAttributeOffset != INDEX_NONE)
 			{
 				if (GNiagaraGPUSorting &&
 					GNiagaraGPUSortingCPUToGPUThreshold != INDEX_NONE &&
-					SortInfo.ParticleCount >= GNiagaraGPUSortingCPUToGPUThreshold)
+					SortInfo.ParticleCount >= GNiagaraGPUSortingCPUToGPUThreshold &&
+					SceneProxy->GetBatcher())
 				{
 					SortInfo.ParticleCount = NumInstances;
 					SortInfo.ParticleDataFloatSRV = CPUSimParticleDataAllocation.ParticleData.ReadBuffer->SRV;
 					SortInfo.FloatDataOffset = CPUSimParticleDataAllocation.ParticleData.FirstIndex / sizeof(float);
 					SortInfo.FloatDataStride = SourceParticleData->GetFloatStride() / sizeof(float);
-					const int32 IndexBufferOffset = Batcher->AddSortedGPUSimulation(SortInfo);
+					const int32 IndexBufferOffset = SceneProxy->GetBatcher()->AddSortedGPUSimulation(SortInfo);
 					if (IndexBufferOffset != INDEX_NONE)
 					{
-						OutVertexFactory.SetSortedIndices(Batcher->GetGPUSortedBuffer().VertexBufferSRV, IndexBufferOffset);
+						OutVertexFactory.SetSortedIndices(SceneProxy->GetBatcher()->GetGPUSortedBuffer().VertexBufferSRV, IndexBufferOffset);
 					}
 				}
 				else
@@ -396,9 +394,9 @@ void FNiagaraRendererSprites::SetVertexFactoryParticleData(
 			}
 			OutVertexFactory.SetParticleData(CPUSimParticleDataAllocation.ParticleData.ReadBuffer->SRV, CPUSimParticleDataAllocation.ParticleData.FirstIndex / sizeof(float), SourceParticleData->GetFloatStride() / sizeof(float));
 		}
-		else // ENiagaraSimTarget::GPUSim
+		else
 		{
-			if (SortInfo.SortMode != ENiagaraSortMode::None && SortInfo.SortAttributeOffset != INDEX_NONE && GNiagaraGPUSorting)
+			if (SortInfo.SortMode != ENiagaraSortMode::None && SortInfo.SortAttributeOffset != INDEX_NONE && GNiagaraGPUSorting && SceneProxy->GetBatcher())
 			{
 				// Here we need to be conservative about the InstanceCount, since the final value is only known on the GPU after the simulation.
 				SortInfo.ParticleCount = SourceParticleData->GetNumInstances();
@@ -406,12 +404,12 @@ void FNiagaraRendererSprites::SetVertexFactoryParticleData(
 				SortInfo.ParticleDataFloatSRV = SourceParticleData->GetGPUBufferFloat().SRV;
 				SortInfo.FloatDataOffset = 0;
 				SortInfo.FloatDataStride = SourceParticleData->GetFloatStride() / sizeof(float);
-				SortInfo.GPUParticleCountSRV = Batcher->GetGPUInstanceCounterManager().GetInstanceCountBuffer().SRV;
-				SortInfo.GPUParticleCountOffset = SourceParticleData->GetGPUInstanceCountBufferOffset();
-				const int32 IndexBufferOffset = Batcher->AddSortedGPUSimulation(SortInfo);
-				if (IndexBufferOffset != INDEX_NONE && SortInfo.GPUParticleCountOffset != INDEX_NONE)
+				SortInfo.GPUParticleCountSRV = SourceParticleData->GetGPUIndices().SRV;
+				SortInfo.GPUParticleCountOffset = 1;
+				const int32 IndexBufferOffset = SceneProxy->GetBatcher()->AddSortedGPUSimulation(SortInfo);
+				if (IndexBufferOffset != INDEX_NONE)
 				{
-					OutVertexFactory.SetSortedIndices(Batcher->GetGPUSortedBuffer().VertexBufferSRV, IndexBufferOffset);
+					OutVertexFactory.SetSortedIndices(SceneProxy->GetBatcher()->GetGPUSortedBuffer().VertexBufferSRV, IndexBufferOffset);
 				}
 			}
 
@@ -425,7 +423,6 @@ void FNiagaraRendererSprites::CreateMeshBatchForView(
 	const FSceneViewFamily& ViewFamily, 
 	const FNiagaraSceneProxy *SceneProxy,
 	FNiagaraDynamicDataSprites *DynamicDataSprites,
-	uint32 IndirectArgsOffset,
 	FMeshBatch& MeshBatch,
 	FNiagaraMeshCollectorResourcesSprite& CollectorResources) const
 {
@@ -466,17 +463,13 @@ void FNiagaraRendererSprites::CreateMeshBatchForView(
 	VFLooseParams.ParticleFacingMode = CollectorResources.VertexFactory.GetFacingMode();
 	VFLooseParams.SortedIndices = CollectorResources.VertexFactory.GetSortedIndicesSRV() ? CollectorResources.VertexFactory.GetSortedIndicesSRV().GetReference() : GFNiagaraNullSortedIndicesVertexBuffer.VertexBufferSRV.GetReference();
 	VFLooseParams.SortedIndicesOffset = CollectorResources.VertexFactory.GetSortedIndicesOffset();
-	if (IndirectArgsOffset != INDEX_NONE)
+	if (DynamicDataSprites->GetSimTarget() == ENiagaraSimTarget::GPUComputeSim)
 	{
-		NiagaraEmitterInstanceBatcher* Batcher = SceneProxy->GetBatcher();
-		check(Batcher); // Already verified at this point.
-		VFLooseParams.IndirectArgsOffset = IndirectArgsOffset / sizeof(uint32);
-		VFLooseParams.IndirectArgsBuffer = Batcher->GetGPUInstanceCounterManager().GetDrawIndirectBuffer().SRV;
+		VFLooseParams.IndirectArgsBuffer = SourceParticleData->GetGPUIndices().SRV;
 	}
 	else
 	{
 		VFLooseParams.IndirectArgsBuffer = GFNiagaraNullSortedIndicesVertexBuffer.VertexBufferSRV;
-		VFLooseParams.IndirectArgsOffset = 0;
 	}
 
 	CollectorResources.VertexFactory.LooseParameterUniformBuffer = FNiagaraSpriteVFLooseParametersRef::CreateUniformBufferImmediate(VFLooseParams, UniformBuffer_SingleFrame);
@@ -503,22 +496,20 @@ void FNiagaraRendererSprites::CreateMeshBatchForView(
 	FMeshBatchElement& MeshElement = MeshBatch.Elements[0];
 	MeshElement.IndexBuffer = &GParticleIndexBuffer;
 	MeshElement.FirstIndex = 0;
-	MeshElement.NumPrimitives = NumIndicesPerInstance / 3;
+	MeshElement.NumPrimitives = 2;
 	MeshElement.NumInstances = FMath::Max(0, NumInstances);	//->VertexData.Num();
 	MeshElement.MinVertexIndex = 0;
 	MeshElement.MaxVertexIndex = 0;// MeshElement.NumInstances * 4 - 1;
 	MeshElement.PrimitiveUniformBuffer = WorldSpacePrimitiveUniformBuffer.GetUniformBufferRHI();
-	if (IndirectArgsOffset != INDEX_NONE)
+	if (DynamicDataSprites->GetSimTarget() == ENiagaraSimTarget::GPUComputeSim)
 	{
-		NiagaraEmitterInstanceBatcher* Batcher = SceneProxy->GetBatcher();
-		check(Batcher); // Already verified at this point.
-		MeshElement.IndirectArgsOffset = IndirectArgsOffset;
-		MeshElement.IndirectArgsBuffer = Batcher->GetGPUInstanceCounterManager().GetDrawIndirectBuffer().Buffer;
+		MeshElement.IndirectArgsBuffer = SourceParticleData->GetGPUIndices().Buffer;
 		MeshElement.NumPrimitives = 0;
 	}
 
 	if (NumCutoutVertexPerSubImage == 8)
 	{
+		MeshElement.NumPrimitives = 6;
 		MeshElement.IndexBuffer = &GSixTriangleParticleIndexBuffer;
 	}
 	
@@ -529,14 +520,12 @@ void FNiagaraRendererSprites::GetDynamicMeshElements(const TArray<const FSceneVi
 {
 	SCOPE_CYCLE_COUNTER(STAT_NiagaraRender);
 	SCOPE_CYCLE_COUNTER(STAT_NiagaraRenderSprites);
-	check(SceneProxy);
 
 	SimpleTimer MeshElementsTimer;
 
 	//check(DynamicDataRender)
 	FNiagaraDynamicDataSprites *DynamicDataSprites = static_cast<FNiagaraDynamicDataSprites*>(DynamicDataRender);
-	NiagaraEmitterInstanceBatcher* Batcher = SceneProxy->GetBatcher();
-	if (!DynamicDataSprites || !Batcher)
+	if (DynamicDataSprites == nullptr)
 	{
 		return;
 	}
@@ -559,12 +548,6 @@ void FNiagaraRendererSprites::GetDynamicMeshElements(const TArray<const FSceneVi
 
 	ConditionalInitPrimitiveUniformBuffer(SceneProxy);
 
-	uint32 IndirectArgsOffset = INDEX_NONE;
-	if (SimTarget == ENiagaraSimTarget::GPUComputeSim)
-	{
-		IndirectArgsOffset = Batcher->GetGPUInstanceCounterManager().AddDrawIndirect(SourceParticleData->GetGPUInstanceCountBufferOffset(), NumIndicesPerInstance);
-	}
-
 	for (int32 ViewIndex = 0; ViewIndex < Views.Num(); ViewIndex++)
 	{
 		if (VisibilityMap & (1 << ViewIndex))
@@ -575,7 +558,8 @@ void FNiagaraRendererSprites::GetDynamicMeshElements(const TArray<const FSceneVi
 			SetVertexFactoryParticleData(CollectorResources.VertexFactory, DynamicDataSprites, CPUSimParticleDataAllocation, View, SceneProxy);
 			CollectorResources.UniformBuffer = CreatePerViewUniformBuffer(View, ViewFamily, SceneProxy);
 			FMeshBatch& MeshBatch = Collector.AllocateMesh();
-			CreateMeshBatchForView(View, ViewFamily, SceneProxy, DynamicDataSprites, IndirectArgsOffset, MeshBatch, CollectorResources);
+			CreateMeshBatchForView(View, ViewFamily, SceneProxy, DynamicDataSprites, MeshBatch, CollectorResources);
+
 			Collector.AddMesh(ViewIndex, MeshBatch);
 		}
 	}
@@ -588,13 +572,11 @@ void FNiagaraRendererSprites::GetDynamicRayTracingInstances(FRayTracingMaterialG
 {
 	SCOPE_CYCLE_COUNTER(STAT_NiagaraRender);
 	SCOPE_CYCLE_COUNTER(STAT_NiagaraRenderSprites);
-	check(SceneProxy);
 
 	SimpleTimer MeshElementsTimer;
 
 	FNiagaraDynamicDataSprites *DynamicDataSprites = static_cast<FNiagaraDynamicDataSprites*>(DynamicDataRender);
-	NiagaraEmitterInstanceBatcher* Batcher = SceneProxy->GetBatcher();
-	if (!DynamicDataSprites || !Batcher)
+	if (DynamicDataSprites == nullptr)
 	{
 		return;
 	}
@@ -610,12 +592,6 @@ void FNiagaraRendererSprites::GetDynamicRayTracingInstances(FRayTracingMaterialG
 		return;
 	}
 
-	uint32 IndirectArgsOffset = INDEX_NONE;
-	if (SimTarget == ENiagaraSimTarget::GPUComputeSim)
-	{
-		IndirectArgsOffset = Batcher->GetGPUInstanceCounterManager().AddDrawIndirect(SourceParticleData->GetGPUInstanceCountBufferOffset(), NumIndicesPerInstance);
-	}
-
 	FRayTracingInstance RayTracingInstance;
 	RayTracingInstance.Geometry = &RayTracingGeometry;
 	RayTracingInstance.InstanceTransforms.Add(FMatrix::Identity);
@@ -628,7 +604,7 @@ void FNiagaraRendererSprites::GetDynamicRayTracingInstances(FRayTracingMaterialG
 		SetVertexFactoryParticleData(CollectorResources.VertexFactory, DynamicDataSprites, CPUSimParticleDataAllocation, Context.ReferenceView, SceneProxy);
 		CollectorResources.UniformBuffer = CreatePerViewUniformBuffer(Context.ReferenceView, Context.ReferenceViewFamily, SceneProxy);
 		FMeshBatch MeshBatch;
-		CreateMeshBatchForView(Context.ReferenceView, Context.ReferenceViewFamily, SceneProxy, DynamicDataSprites, IndirectArgsOffset, MeshBatch, CollectorResources);
+		CreateMeshBatchForView(Context.ReferenceView, Context.ReferenceViewFamily, SceneProxy, DynamicDataSprites, MeshBatch, CollectorResources);
 
 		ensureMsgf(MeshBatch.Elements[0].IndexBuffer != &GSixTriangleParticleIndexBuffer, TEXT("Cutout geometry is not supported by ray tracing"));
 
@@ -666,6 +642,7 @@ FNiagaraDynamicDataBase *FNiagaraRendererSprites::GenerateDynamicData(const FNia
 		SCOPE_CYCLE_COUNTER(STAT_NiagaraGenSpriteDynamicData);
 
 		FNiagaraDataSet& Data = Emitter->GetData();
+		ENiagaraSimTarget SimTarget = Data.GetSimTarget();
 		if(SimTarget == ENiagaraSimTarget::GPUComputeSim || Data.GetCurrentDataChecked().GetNumInstances() > 0)
 		{
 			DynamicData = new FNiagaraDynamicDataSprites(Emitter);
diff --git a/Engine/Plugins/FX/Niagara/Source/Niagara/Private/NiagaraScriptExecutionContext.cpp b/Engine/Plugins/FX/Niagara/Source/Niagara/Private/NiagaraScriptExecutionContext.cpp
--- a/Engine/Plugins/FX/Niagara/Source/Niagara/Private/NiagaraScriptExecutionContext.cpp
+++ b/Engine/Plugins/FX/Niagara/Source/Niagara/Private/NiagaraScriptExecutionContext.cpp
@@ -8,7 +8,6 @@
 #include "NiagaraWorldManager.h"
 #include "NiagaraSystemInstance.h"
 #include "NiagaraEmitterInstance.h"
-#include "NiagaraGPUInstanceCountManager.h"
 
 DECLARE_CYCLE_STAT(TEXT("Register Setup"), STAT_NiagaraSimRegisterSetup, STATGROUP_Niagara);
 DECLARE_CYCLE_STAT(TEXT("Context Ticking"), STAT_NiagaraScriptExecContextTick, STATGROUP_Niagara);
@@ -396,6 +395,9 @@ FNiagaraComputeExecutionContext::FNiagaraComputeExecutionContext()
 	, GPUScript_RT(nullptr)
 	, CBufferLayout(TEXT("Niagara Compute Sim CBuffer"))
 	, DataToRender(nullptr)
+	, GPUDataReadback(nullptr)
+	, AccumulatedSpawnRate(0)
+	, NumIndicesPerInstance(0)
 #if WITH_EDITORONLY_DATA
 	, GPUDebugDataReadbackFloat(nullptr)
 	, GPUDebugDataReadbackInt(nullptr)
@@ -404,7 +406,6 @@ FNiagaraComputeExecutionContext::FNiagaraComputeExecutionContext()
 	, GPUDebugDataIntSize(0)
 	, GPUDebugDataFloatStride(0)
 	, GPUDebugDataIntStride(0)
-	, GPUDebugDataCountOffset(INDEX_NONE)
 #endif	  
 {
 }
@@ -412,7 +413,11 @@ FNiagaraComputeExecutionContext::FNiagaraComputeExecutionContext()
 FNiagaraComputeExecutionContext::~FNiagaraComputeExecutionContext()
 {
 	checkf(IsInRenderingThread(), TEXT("Can only delete the gpu readback from the render thread"));
-	check(EmitterInstanceReadback.GPUCountOffset == INDEX_NONE);
+	if (GPUDataReadback)
+	{
+		delete GPUDataReadback;
+		GPUDataReadback = nullptr;
+	}
 
 #if WITH_EDITORONLY_DATA
 	if (GPUDebugDataReadbackFloat)
@@ -435,14 +440,13 @@ FNiagaraComputeExecutionContext::~FNiagaraComputeExecutionContext()
 	SetDataToRender(nullptr);
 }
 
-void FNiagaraComputeExecutionContext::Reset(NiagaraEmitterInstanceBatcher* Batcher)
+void FNiagaraComputeExecutionContext::Reset()
 {
 	FNiagaraComputeExecutionContext* Context = this;
-	NiagaraEmitterInstanceBatcher* B = Batcher && !Batcher->IsPendingKill() ? Batcher : nullptr;
 	ENQUEUE_RENDER_COMMAND(ResetRT)(
-		[B, Context](FRHICommandListImmediate& RHICmdList)
+		[Context](FRHICommandListImmediate& RHICmdList)
 	{
-		Context->ResetInternal(B);
+		Context->ResetInternal();
 	}
 	);
 }
@@ -514,18 +518,14 @@ void FNiagaraComputeExecutionContext::PostTick()
 	}
 }
 
-void FNiagaraComputeExecutionContext::ResetInternal(NiagaraEmitterInstanceBatcher* Batcher)
+void FNiagaraComputeExecutionContext::ResetInternal()
 {
 	checkf(IsInRenderingThread(), TEXT("Can only reset the gpu context from the render thread"));
-
-	// Release and reset readback data.
-	if (Batcher)
-	{
-		Batcher->GetGPUInstanceCounterManager().FreeEntry(EmitterInstanceReadback.GPUCountOffset);
-	}
-	else // In this case the batcher is pending kill so no need to putback entry in the pool.
+	AccumulatedSpawnRate = 0;
+	if (GPUDataReadback)
 	{
-		EmitterInstanceReadback.GPUCountOffset = INDEX_NONE;
+		delete GPUDataReadback;
+		GPUDataReadback = nullptr;
 	}
 
 #if WITH_EDITORONLY_DATA
diff --git a/Engine/Plugins/FX/Niagara/Source/Niagara/Private/NiagaraSpriteRendererProperties.cpp b/Engine/Plugins/FX/Niagara/Source/Niagara/Private/NiagaraSpriteRendererProperties.cpp
--- a/Engine/Plugins/FX/Niagara/Source/Niagara/Private/NiagaraSpriteRendererProperties.cpp
+++ b/Engine/Plugins/FX/Niagara/Source/Niagara/Private/NiagaraSpriteRendererProperties.cpp
@@ -295,7 +295,7 @@ int32 UNiagaraSpriteRendererProperties::GetNumCutoutVertexPerSubimage() const
 	}
 }
 
-uint32 UNiagaraSpriteRendererProperties::GetNumIndicesPerInstance() const
+uint32 UNiagaraSpriteRendererProperties::GetNumIndicesPerInstance()
 {
 	// This is a based on cutout vertices making a triangle strip.
 	if (GetNumCutoutVertexPerSubimage() == 8)
diff --git a/Engine/Plugins/FX/Niagara/Source/Niagara/Public/NiagaraMeshRendererProperties.h b/Engine/Plugins/FX/Niagara/Source/Niagara/Public/NiagaraMeshRendererProperties.h
--- a/Engine/Plugins/FX/Niagara/Source/Niagara/Public/NiagaraMeshRendererProperties.h
+++ b/Engine/Plugins/FX/Niagara/Source/Niagara/Public/NiagaraMeshRendererProperties.h
@@ -55,7 +55,7 @@ public:
 	virtual const TArray<FNiagaraVariable>& GetOptionalAttributes() override;
 #endif // WITH_EDITORONLY_DATA
 
-	virtual uint32 GetNumIndicesPerInstance() const final override;
+	virtual uint32 GetNumIndicesPerInstance() final override;
 
 	/** The static mesh to be instanced when rendering mesh particles. If OverrideMaterial is not specified, the mesh's material is used. Note that the material must have the Niagara Mesh Particles flag checked.*/
 	UPROPERTY(EditAnywhere, Category = "Mesh Rendering")
diff --git a/Engine/Plugins/FX/Niagara/Source/Niagara/Public/NiagaraRenderer.h b/Engine/Plugins/FX/Niagara/Source/Niagara/Public/NiagaraRenderer.h
--- a/Engine/Plugins/FX/Niagara/Source/Niagara/Public/NiagaraRenderer.h
+++ b/Engine/Plugins/FX/Niagara/Source/Niagara/Public/NiagaraRenderer.h
@@ -90,8 +90,8 @@ public:
 	FNiagaraRenderer& operator=(const FNiagaraRenderer& Other) = delete;
 
 	virtual void Initialize(ERHIFeatureLevel::Type FeatureLevel, const UNiagaraRendererProperties *InProps, const FNiagaraEmitterInstance* Emitter);
-	virtual void CreateRenderThreadResources(NiagaraEmitterInstanceBatcher* Batcher);
-	virtual void ReleaseRenderThreadResources(NiagaraEmitterInstanceBatcher* Batcher);
+	virtual void CreateRenderThreadResources() {}
+	virtual void ReleaseRenderThreadResources() {}
 
 	virtual FPrimitiveViewRelevance GetViewRelevance(const FSceneView* View, const FNiagaraSceneProxy *SceneProxy)const;
 	virtual void GetDynamicMeshElements(const TArray<const FSceneView*>& Views, const FSceneViewFamily& ViewFamily, uint32 VisibilityMap, FMeshElementCollector& Collector, const FNiagaraSceneProxy *SceneProxy) const {}
@@ -132,8 +132,6 @@ protected:
 
 	uint32 bLocalSpace : 1;
 	uint32 bHasLights : 1;
-	ENiagaraSimTarget SimTarget;
-	uint32 NumIndicesPerInstance;
 
 #if STATS
 	TStatId EmitterStatID;
@@ -142,6 +140,5 @@ protected:
 	/** Cached array of materials used from the properties data. Validated with usage flags etc. */
 	TArray<UMaterialInterface*> BaseMaterials_GT;
 	FMaterialRelevance BaseMaterialRelevance_GT;
-
 };
 
diff --git a/Engine/Plugins/FX/Niagara/Source/Niagara/Public/NiagaraRendererMeshes.h b/Engine/Plugins/FX/Niagara/Source/Niagara/Public/NiagaraRendererMeshes.h
--- a/Engine/Plugins/FX/Niagara/Source/Niagara/Public/NiagaraRendererMeshes.h
+++ b/Engine/Plugins/FX/Niagara/Source/Niagara/Public/NiagaraRendererMeshes.h
@@ -20,8 +20,8 @@ public:
 	FNiagaraRendererMeshes(ERHIFeatureLevel::Type FeatureLevel, const UNiagaraRendererProperties *InProps, const FNiagaraEmitterInstance* Emitter);
 	
 	//FNiagaraRenderer Interface
-	virtual void CreateRenderThreadResources(NiagaraEmitterInstanceBatcher* Batcher) override;
-	virtual void ReleaseRenderThreadResources(NiagaraEmitterInstanceBatcher* Batcher) override;
+	virtual void CreateRenderThreadResources() override;
+	virtual void ReleaseRenderThreadResources() override;
 
 	virtual void GetDynamicMeshElements(const TArray<const FSceneView*>& Views, const FSceneViewFamily& ViewFamily, uint32 VisibilityMap, FMeshElementCollector& Collector, const FNiagaraSceneProxy *SceneProxy) const override;
 	virtual FNiagaraDynamicDataBase* GenerateDynamicData(const FNiagaraSceneProxy* Proxy, const UNiagaraRendererProperties* InProperties, const FNiagaraEmitterInstance* Emitter) const override;
diff --git a/Engine/Plugins/FX/Niagara/Source/Niagara/Public/NiagaraRendererProperties.h b/Engine/Plugins/FX/Niagara/Source/Niagara/Public/NiagaraRendererProperties.h
--- a/Engine/Plugins/FX/Niagara/Source/Niagara/Public/NiagaraRendererProperties.h
+++ b/Engine/Plugins/FX/Niagara/Source/Niagara/Public/NiagaraRendererProperties.h
@@ -53,7 +53,7 @@ public:
 
 
 	// GPU simulation uses DrawIndirect, so the sim step needs to know indices per instance in order to prepare the draw call parameters
-	virtual uint32 GetNumIndicesPerInstance() const { return 0; }
+	virtual uint32 GetNumIndicesPerInstance() { return 0; }
 
 	virtual bool GetIsEnabled() const { return bIsEnabled; }
 	virtual void SetIsEnabled(bool bInIsEnabled) { bIsEnabled = bInIsEnabled; }
diff --git a/Engine/Plugins/FX/Niagara/Source/Niagara/Public/NiagaraRendererRibbons.h b/Engine/Plugins/FX/Niagara/Source/Niagara/Public/NiagaraRendererRibbons.h
--- a/Engine/Plugins/FX/Niagara/Source/Niagara/Public/NiagaraRendererRibbons.h
+++ b/Engine/Plugins/FX/Niagara/Source/Niagara/Public/NiagaraRendererRibbons.h
@@ -20,8 +20,8 @@ public:
 	FNiagaraRendererRibbons(ERHIFeatureLevel::Type FeatureLevel, const UNiagaraRendererProperties *InProps, const FNiagaraEmitterInstance* Emitter);	// FNiagaraRenderer Interface 
 
 	// FNiagaraRenderer Interface 
-	virtual void CreateRenderThreadResources(NiagaraEmitterInstanceBatcher* Batcher) override;
-	virtual void ReleaseRenderThreadResources(NiagaraEmitterInstanceBatcher* Batcher) override;
+	virtual void CreateRenderThreadResources() override;
+	virtual void ReleaseRenderThreadResources() override;
 
 	virtual void GetDynamicMeshElements(const TArray<const FSceneView*>& Views, const FSceneViewFamily& ViewFamily, uint32 VisibilityMap, FMeshElementCollector& Collector, const FNiagaraSceneProxy *SceneProxy) const override;
 	virtual FNiagaraDynamicDataBase *GenerateDynamicData(const FNiagaraSceneProxy* Proxy, const UNiagaraRendererProperties* InProperties, const FNiagaraEmitterInstance* Emitteride) const override;
diff --git a/Engine/Plugins/FX/Niagara/Source/Niagara/Public/NiagaraRendererSprites.h b/Engine/Plugins/FX/Niagara/Source/Niagara/Public/NiagaraRendererSprites.h
--- a/Engine/Plugins/FX/Niagara/Source/Niagara/Public/NiagaraRendererSprites.h
+++ b/Engine/Plugins/FX/Niagara/Source/Niagara/Public/NiagaraRendererSprites.h
@@ -20,8 +20,8 @@ public:
 	FNiagaraRendererSprites(ERHIFeatureLevel::Type FeatureLevel, const UNiagaraRendererProperties *InProps, const FNiagaraEmitterInstance* Emitter);
 
 	//FNiagaraRenderer interface
-	virtual void CreateRenderThreadResources(NiagaraEmitterInstanceBatcher* Batcher)override;
-	virtual void ReleaseRenderThreadResources(NiagaraEmitterInstanceBatcher* Batcher)override;
+	virtual void CreateRenderThreadResources()override;
+	virtual void ReleaseRenderThreadResources()override;
 
 	virtual void GetDynamicMeshElements(const TArray<const FSceneView*>& Views, const FSceneViewFamily& ViewFamily, uint32 VisibilityMap, FMeshElementCollector& Collector, const FNiagaraSceneProxy *SceneProxy) const override;
 	virtual FNiagaraDynamicDataBase* GenerateDynamicData(const FNiagaraSceneProxy* Proxy, const UNiagaraRendererProperties* InProperties, const FNiagaraEmitterInstance* Emitter) const override;
@@ -55,7 +55,6 @@ private:
 		const FSceneViewFamily& ViewFamily,
 		const FNiagaraSceneProxy *SceneProxy,
 		FNiagaraDynamicDataSprites *DynamicDataSprites,
-		uint32 IndirectArgsOffset,
 		FMeshBatch& OutMeshBatch,
 		class FNiagaraMeshCollectorResourcesSprite& OutCollectorResources) const;
 
diff --git a/Engine/Plugins/FX/Niagara/Source/Niagara/Public/NiagaraSpriteRendererProperties.h b/Engine/Plugins/FX/Niagara/Source/Niagara/Public/NiagaraSpriteRendererProperties.h
--- a/Engine/Plugins/FX/Niagara/Source/Niagara/Public/NiagaraSpriteRendererProperties.h
+++ b/Engine/Plugins/FX/Niagara/Source/Niagara/Public/NiagaraSpriteRendererProperties.h
@@ -72,7 +72,7 @@ public:
 	//UNiagaraMaterialRendererProperties interface END
 
 	int32 GetNumCutoutVertexPerSubimage() const;
-	virtual uint32 GetNumIndicesPerInstance() const;
+	virtual uint32 GetNumIndicesPerInstance();
 
 	/** The material used to render the particle. Note that it must have the Use with Niagara Sprites flag checked.*/
 	UPROPERTY(EditAnywhere, Category = "Sprite Rendering")
diff --git a/Engine/Plugins/FX/Niagara/Source/NiagaraCore/Private/NiagaraCustomVersion.cpp b/Engine/Plugins/FX/Niagara/Source/NiagaraCore/Private/NiagaraCustomVersion.cpp
--- a/Engine/Plugins/FX/Niagara/Source/NiagaraCore/Private/NiagaraCustomVersion.cpp
+++ b/Engine/Plugins/FX/Niagara/Source/NiagaraCore/Private/NiagaraCustomVersion.cpp
@@ -8,4 +8,4 @@ const FGuid FNiagaraCustomVersion::GUID(0xFCF57AFA, 0x50764283, 0xB9A9E658, 0xFF
 // Register the custom version with core
 FCustomVersionRegistration GRegisterNiagaraCustomVersion(FNiagaraCustomVersion::GUID, FNiagaraCustomVersion::LatestVersion, TEXT("NiagaraVer"));
 
-const FGuid FNiagaraCustomVersion::LatestScriptCompileVersion(0x397D4E10, 0x87CC4DF8, 0x9EDAD0CA, 0xE837CA88);
\ No newline at end of file
+const FGuid FNiagaraCustomVersion::LatestScriptCompileVersion(0x8D38659E, 0xFF5CC945, 0xA718F14F, 0x52E92EBA);
diff --git a/Engine/Plugins/FX/Niagara/Source/NiagaraShader/Private/NiagaraShader.cpp b/Engine/Plugins/FX/Niagara/Source/NiagaraShader/Private/NiagaraShader.cpp
--- a/Engine/Plugins/FX/Niagara/Source/NiagaraShader/Private/NiagaraShader.cpp
+++ b/Engine/Plugins/FX/Niagara/Source/NiagaraShader/Private/NiagaraShader.cpp
@@ -1275,11 +1275,8 @@ void FNiagaraShader::BindParams(const FShaderParameterMap &ParameterMap)
 	IntInputBufferParam.Bind(ParameterMap, TEXT("InputInt"));
 	FloatOutputBufferParam.Bind(ParameterMap, TEXT("OutputFloat"));
 	IntOutputBufferParam.Bind(ParameterMap, TEXT("OutputInt"));
-
-	InstanceCountsParam.Bind(ParameterMap, TEXT("InstanceCounts"));
-	ReadInstanceCountOffsetParam.Bind(ParameterMap, TEXT("ReadInstanceCountOffset"));
-	WriteInstanceCountOffsetParam.Bind(ParameterMap, TEXT("WriteInstanceCountOffset"));
-
+	OutputIndexBufferParam.Bind(ParameterMap, TEXT("DataSetIndices"));
+	InputIndexBufferParam.Bind(ParameterMap, TEXT("ReadDataSetIndices"));
 	EmitterTickCounterParam.Bind(ParameterMap, TEXT("EmitterTickCounter"));
 	NumEventsPerParticleParam.Bind(ParameterMap, TEXT("NumEventsPerParticle"));
 	NumParticlesPerEventParam.Bind(ParameterMap, TEXT("NumParticlesPerEvent"));
@@ -1287,6 +1284,7 @@ void FNiagaraShader::BindParams(const FShaderParameterMap &ParameterMap)
 
 	NumSpawnedInstancesParam.Bind(ParameterMap, TEXT("SpawnedInstances"));
 	UpdateStartInstanceParam.Bind(ParameterMap, TEXT("UpdateStartInstance"));
+	NumIndicesPerInstanceParam.Bind(ParameterMap, TEXT("NumIndicesPerInstance"));
 
 	ComponentBufferSizeReadParam.Bind(ParameterMap, TEXT("ComponentBufferSizeRead"));
 	ComponentBufferSizeWriteParam.Bind(ParameterMap, TEXT("ComponentBufferSizeWrite"));
@@ -1325,9 +1323,11 @@ void FNiagaraShader::BindParams(const FShaderParameterMap &ParameterMap)
 
 	ensure(FloatOutputBufferParam.IsBound() || IntOutputBufferParam.IsBound());	// we should have at least one output buffer we're writing to
 	ensure(ComponentBufferSizeWriteParam.IsBound());
-	ensure(InstanceCountsParam.IsBound());
+	ensure(OutputIndexBufferParam.IsBound());
+	ensure(InputIndexBufferParam.IsBound());
 	ensure(UpdateStartInstanceParam.IsBound());
 	ensure(NumSpawnedInstancesParam.IsBound());
+	ensure(NumIndicesPerInstanceParam.IsBound());
 }
 
 bool FNiagaraShader::Serialize(FArchive& Ar)
@@ -1341,11 +1341,7 @@ bool FNiagaraShader::Serialize(FArchive& Ar)
 	Ar << IntInputBufferParam;
 	Ar << FloatOutputBufferParam;
 	Ar << IntOutputBufferParam;
-
-	Ar << InstanceCountsParam;
-	Ar << ReadInstanceCountOffsetParam;
-	Ar << WriteInstanceCountOffsetParam;
-
+	Ar << InputIndexBufferParam;
 	Ar << EmitterTickCounterParam;
 
 	Ar << NumSpawnedInstancesParam;
@@ -1370,6 +1366,8 @@ bool FNiagaraShader::Serialize(FArchive& Ar)
 	Ar << CopyInstancesBeforeStartParam;
 	Ar << ViewUniformBufferParam;
 
+	Ar << OutputIndexBufferParam;
+	Ar << NumIndicesPerInstanceParam;
 	return bShaderHasOutdatedParameters;
 }
 
diff --git a/Engine/Plugins/FX/Niagara/Source/NiagaraShader/Public/NiagaraShader.h b/Engine/Plugins/FX/Niagara/Source/NiagaraShader/Public/NiagaraShader.h
--- a/Engine/Plugins/FX/Niagara/Source/NiagaraShader/Public/NiagaraShader.h
+++ b/Engine/Plugins/FX/Niagara/Source/NiagaraShader/Public/NiagaraShader.h
@@ -102,9 +102,8 @@ public:
 	FShaderResourceParameter IntInputBufferParam;
 	FRWShaderParameter FloatOutputBufferParam;
 	FRWShaderParameter IntOutputBufferParam;
-	FRWShaderParameter InstanceCountsParam;
-	FShaderParameter ReadInstanceCountOffsetParam;
-	FShaderParameter WriteInstanceCountOffsetParam;
+	FRWShaderParameter OutputIndexBufferParam;
+	FShaderResourceParameter InputIndexBufferParam;
 	FShaderUniformBufferParameter EmitterConstantBufferParam;
 	FShaderUniformBufferParameter DataInterfaceUniformBufferParam;
 	FShaderUniformBufferParameter ViewUniformBufferParam;
@@ -114,6 +113,7 @@ public:
 	FShaderParameter CopyInstancesBeforeStartParam;
 	FShaderParameter NumSpawnedInstancesParam;
 	FShaderParameter UpdateStartInstanceParam;
+	FShaderParameter NumIndicesPerInstanceParam;
 	FShaderParameter ComponentBufferSizeReadParam;
 	FShaderParameter ComponentBufferSizeWriteParam;
 	FRWShaderParameter EventIntUAVParams[MAX_CONCURRENT_EVENT_DATASETS];
diff --git a/Engine/Plugins/FX/Niagara/Source/NiagaraVertexFactories/Private/NiagaraDrawIndirect.cpp b/Engine/Plugins/FX/Niagara/Source/NiagaraVertexFactories/Private/NiagaraDrawIndirect.cpp
deleted file mode 100644
diff --git a/Engine/Plugins/FX/Niagara/Source/NiagaraVertexFactories/Public/NiagaraDrawIndirect.h b/Engine/Plugins/FX/Niagara/Source/NiagaraVertexFactories/Public/NiagaraDrawIndirect.h
deleted file mode 100644
diff --git a/Engine/Plugins/FX/Niagara/Source/NiagaraVertexFactories/Public/NiagaraSpriteVertexFactory.h b/Engine/Plugins/FX/Niagara/Source/NiagaraVertexFactories/Public/NiagaraSpriteVertexFactory.h
--- a/Engine/Plugins/FX/Niagara/Source/NiagaraVertexFactories/Public/NiagaraSpriteVertexFactory.h
+++ b/Engine/Plugins/FX/Niagara/Source/NiagaraVertexFactories/Public/NiagaraSpriteVertexFactory.h
@@ -64,7 +64,6 @@ BEGIN_GLOBAL_SHADER_PARAMETER_STRUCT(FNiagaraSpriteVFLooseParameters, NIAGARAVER
 	SHADER_PARAMETER(uint32, ParticleAlignmentMode)
 	SHADER_PARAMETER(uint32, ParticleFacingMode)
 	SHADER_PARAMETER(uint32, SortedIndicesOffset)
-	SHADER_PARAMETER(uint32, IndirectArgsOffset)
 	SHADER_PARAMETER_SRV(Buffer<float2>, CutoutGeometry)
 	SHADER_PARAMETER_SRV(Buffer<float>, NiagaraParticleDataFloat)
 	SHADER_PARAMETER_SRV(Buffer<int>, SortedIndices)
diff --git a/Engine/Source/Runtime/Engine/Private/Particles/FXSystem.cpp b/Engine/Source/Runtime/Engine/Private/Particles/FXSystem.cpp
--- a/Engine/Source/Runtime/Engine/Private/Particles/FXSystem.cpp
+++ b/Engine/Source/Runtime/Engine/Private/Particles/FXSystem.cpp
@@ -361,7 +361,7 @@ void FFXSystem::DrawDebug( FCanvas* Canvas )
 	}
 }
 
-void FFXSystem::PreInitViews(FRHICommandListImmediate& RHICmdList)
+void FFXSystem::PreInitViews()
 {
 	if (RHISupportsGPUParticles())
 	{
diff --git a/Engine/Source/Runtime/Engine/Private/Particles/FXSystemPrivate.h b/Engine/Source/Runtime/Engine/Private/Particles/FXSystemPrivate.h
--- a/Engine/Source/Runtime/Engine/Private/Particles/FXSystemPrivate.h
+++ b/Engine/Source/Runtime/Engine/Private/Particles/FXSystemPrivate.h
@@ -149,7 +149,7 @@ public:
 	virtual void RemoveVectorField(UVectorFieldComponent* VectorFieldComponent) override;
 	virtual void UpdateVectorField(UVectorFieldComponent* VectorFieldComponent) override;
 	FParticleEmitterInstance* CreateGPUSpriteEmitterInstance(FGPUSpriteEmitterInfo& EmitterInfo);
-	virtual void PreInitViews(FRHICommandListImmediate& RHICmdList) override;
+	virtual void PreInitViews() override;
 	virtual bool UsesGlobalDistanceField() const override;
 	virtual void PreRender(FRHICommandListImmediate& RHICmdList, const FGlobalDistanceFieldParameterData* GlobalDistanceFieldParameterData) override;
 	virtual void PostRenderOpaque(
diff --git a/Engine/Source/Runtime/Engine/Private/Particles/FXSystemSet.cpp b/Engine/Source/Runtime/Engine/Private/Particles/FXSystemSet.cpp
--- a/Engine/Source/Runtime/Engine/Private/Particles/FXSystemSet.cpp
+++ b/Engine/Source/Runtime/Engine/Private/Particles/FXSystemSet.cpp
@@ -87,12 +87,12 @@ void FFXSystemSet::UpdateVectorField(UVectorFieldComponent* VectorFieldComponent
 	}
 }
 
-void FFXSystemSet::PreInitViews(FRHICommandListImmediate& RHICmdList)
+void FFXSystemSet::PreInitViews()
 {
 	for (FFXSystemInterface* FXSystem : FXSystems)
 	{
 		check(FXSystem);
-		FXSystem->PreInitViews(RHICmdList);
+		FXSystem->PreInitViews();
 	}
 }
 
diff --git a/Engine/Source/Runtime/Engine/Private/Particles/FXSystemSet.h b/Engine/Source/Runtime/Engine/Private/Particles/FXSystemSet.h
--- a/Engine/Source/Runtime/Engine/Private/Particles/FXSystemSet.h
+++ b/Engine/Source/Runtime/Engine/Private/Particles/FXSystemSet.h
@@ -30,7 +30,7 @@ public:
 	virtual void AddVectorField(UVectorFieldComponent* VectorFieldComponent) override;
 	virtual void RemoveVectorField(UVectorFieldComponent* VectorFieldComponent) override;
 	virtual void UpdateVectorField(UVectorFieldComponent* VectorFieldComponent) override;
-	virtual void PreInitViews(FRHICommandListImmediate& RHICmdList) override;
+	virtual void PreInitViews() override;
 	virtual bool UsesGlobalDistanceField() const override;
 	virtual void PreRender(FRHICommandListImmediate& RHICmdList, const class FGlobalDistanceFieldParameterData* GlobalDistanceFieldParameterData) override;
 	virtual void PostRenderOpaque(
diff --git a/Engine/Source/Runtime/Engine/Public/FXSystem.h b/Engine/Source/Runtime/Engine/Public/FXSystem.h
--- a/Engine/Source/Runtime/Engine/Public/FXSystem.h
+++ b/Engine/Source/Runtime/Engine/Public/FXSystem.h
@@ -197,7 +197,7 @@ public:
 	 * Notification from the renderer that it is about to perform visibility
 	 * checks on FX belonging to this system.
 	 */
-	virtual void PreInitViews(FRHICommandListImmediate& RHICmdList) = 0;
+	virtual void PreInitViews() = 0;
 
 	virtual bool UsesGlobalDistanceField() const = 0;
 
diff --git a/Engine/Source/Runtime/Engine/Public/MeshBatch.h b/Engine/Source/Runtime/Engine/Public/MeshBatch.h
--- a/Engine/Source/Runtime/Engine/Public/MeshBatch.h
+++ b/Engine/Source/Runtime/Engine/Public/MeshBatch.h
@@ -96,7 +96,6 @@ struct FMeshBatchElement
 	int8 VisualizeElementIndex;
 #endif
 	FVertexBufferRHIParamRef IndirectArgsBuffer;
-	uint32 IndirectArgsOffset;
 
 	FMeshBatchElement()
 	:	PrimitiveUniformBuffer(nullptr)
@@ -122,7 +121,6 @@ struct FMeshBatchElement
 	,	VisualizeElementIndex(INDEX_NONE)
 #endif
 	,	IndirectArgsBuffer(nullptr)
-	,	IndirectArgsOffset(0)
 	{
 	}
 };
diff --git a/Engine/Source/Runtime/Renderer/Private/MeshPassProcessor.cpp b/Engine/Source/Runtime/Renderer/Private/MeshPassProcessor.cpp
--- a/Engine/Source/Runtime/Renderer/Private/MeshPassProcessor.cpp
+++ b/Engine/Source/Runtime/Renderer/Private/MeshPassProcessor.cpp
@@ -722,8 +722,7 @@ void FMeshDrawCommand::SetDrawParametersAndFinalize(
 	else
 	{
 		checkf(BatchElement.IndirectArgsBuffer, TEXT("It is only valid to set BatchElement.NumPrimitives == 0 when a IndirectArgsBuffer is used"));
-		IndirectArgs.Buffer = BatchElement.IndirectArgsBuffer;
-		IndirectArgs.Offset = BatchElement.IndirectArgsOffset;
+		IndirectArgsBuffer = BatchElement.IndirectArgsBuffer;
 	}
 
 	Finalize(PipelineId, ShadersForDebugging);
@@ -875,9 +874,9 @@ void FMeshDrawCommand::SubmitDraw(
 		FName ResourceName = MeshDrawCommand.DebugData.ResourceName;
 
 		FString DrawEventName = FString::Printf(
-				TEXT("%s %s"),
-				// Note: this is the parent's material name, not the material instance
-				*Material->GetFriendlyName(),
+			TEXT("%s %s"),
+			// Note: this is the parent's material name, not the material instance
+			*Material->GetFriendlyName(),
 			ResourceName.IsValid() ? *ResourceName.ToString() : TEXT(""));
 
 		const uint32 Instances = MeshDrawCommand.NumInstances * InstanceFactor;
@@ -901,37 +900,37 @@ void FMeshDrawCommand::SubmitDraw(
 	{
 		const FGraphicsMinimalPipelineStateInitializer& MeshPipelineState = MeshDrawCommand.CachedPipelineId.GetPipelineState(GraphicsMinimalPipelineStateSet);
 
-	if (MeshDrawCommand.CachedPipelineId.GetId() != StateCache.PipelineId)
-	{
-		FGraphicsPipelineStateInitializer GraphicsPSOInit = MeshPipelineState;
-		RHICmdList.ApplyCachedRenderTargets(GraphicsPSOInit);
-		SetGraphicsPipelineState(RHICmdList, GraphicsPSOInit);
-		StateCache.SetPipelineState(MeshDrawCommand.CachedPipelineId.GetId());
-	}
-
-	if (MeshDrawCommand.StencilRef != StateCache.StencilRef)
-	{
-		RHICmdList.SetStencilRef(MeshDrawCommand.StencilRef);
-		StateCache.StencilRef = MeshDrawCommand.StencilRef;
-	}
-
-	for (int32 VertexBindingIndex = 0; VertexBindingIndex < MeshDrawCommand.VertexStreams.Num(); VertexBindingIndex++)
-	{
-		const FVertexInputStream& Stream = MeshDrawCommand.VertexStreams[VertexBindingIndex];
+		if (MeshDrawCommand.CachedPipelineId.GetId() != StateCache.PipelineId)
+		{
+			FGraphicsPipelineStateInitializer GraphicsPSOInit = MeshPipelineState;
+			RHICmdList.ApplyCachedRenderTargets(GraphicsPSOInit);
+			SetGraphicsPipelineState(RHICmdList, GraphicsPSOInit);
+			StateCache.SetPipelineState(MeshDrawCommand.CachedPipelineId.GetId());
+		}
 
-		if (MeshDrawCommand.PrimitiveIdStreamIndex != -1 && Stream.StreamIndex == MeshDrawCommand.PrimitiveIdStreamIndex)
+		if (MeshDrawCommand.StencilRef != StateCache.StencilRef)
 		{
-			RHICmdList.SetStreamSource(Stream.StreamIndex, ScenePrimitiveIdsBuffer, PrimitiveIdOffset);
-			StateCache.VertexStreams[Stream.StreamIndex] = Stream;
+			RHICmdList.SetStencilRef(MeshDrawCommand.StencilRef);
+			StateCache.StencilRef = MeshDrawCommand.StencilRef;
 		}
-		else if (StateCache.VertexStreams[Stream.StreamIndex] != Stream)
+
+		for (int32 VertexBindingIndex = 0; VertexBindingIndex < MeshDrawCommand.VertexStreams.Num(); VertexBindingIndex++)
 		{
-			RHICmdList.SetStreamSource(Stream.StreamIndex, Stream.VertexBuffer, Stream.Offset);
-			StateCache.VertexStreams[Stream.StreamIndex] = Stream;
+			const FVertexInputStream& Stream = MeshDrawCommand.VertexStreams[VertexBindingIndex];
+
+			if (MeshDrawCommand.PrimitiveIdStreamIndex != -1 && Stream.StreamIndex == MeshDrawCommand.PrimitiveIdStreamIndex)
+			{
+				RHICmdList.SetStreamSource(Stream.StreamIndex, ScenePrimitiveIdsBuffer, PrimitiveIdOffset);
+				StateCache.VertexStreams[Stream.StreamIndex] = Stream;
+			}
+			else if (StateCache.VertexStreams[Stream.StreamIndex] != Stream)
+			{
+				RHICmdList.SetStreamSource(Stream.StreamIndex, Stream.VertexBuffer, Stream.Offset);
+				StateCache.VertexStreams[Stream.StreamIndex] = Stream;
+			}
 		}
-	}
 
-	MeshDrawCommand.ShaderBindings.SetOnCommandList(RHICmdList, MeshPipelineState.BoundShaderState, StateCache.ShaderBindings);
+		MeshDrawCommand.ShaderBindings.SetOnCommandList(RHICmdList, MeshPipelineState.BoundShaderState, StateCache.ShaderBindings);
 	}
 
 	if (MeshDrawCommand.IndexBuffer)
@@ -951,10 +950,10 @@ void FMeshDrawCommand::SubmitDraw(
 		else
 		{
 			RHICmdList.DrawIndexedPrimitiveIndirect(
-				MeshDrawCommand.IndexBuffer, 
-				MeshDrawCommand.IndirectArgs.Buffer, 
-				MeshDrawCommand.IndirectArgs.Offset
-				);
+				MeshDrawCommand.IndexBuffer,
+				MeshDrawCommand.IndirectArgsBuffer,
+				0
+			);
 		}
 	}
 	else
diff --git a/Engine/Source/Runtime/Renderer/Private/SceneVisibility.cpp b/Engine/Source/Runtime/Renderer/Private/SceneVisibility.cpp
--- a/Engine/Source/Runtime/Renderer/Private/SceneVisibility.cpp
+++ b/Engine/Source/Runtime/Renderer/Private/SceneVisibility.cpp
@@ -2923,7 +2923,7 @@ void FSceneRenderer::PreVisibilityFrameSetup(FRHICommandListImmediate& RHICmdLis
 	// Notify the FX system that the scene is about to perform visibility checks.
 	if (Scene->FXSystem && !Views[0].bIsPlanarReflection)
 	{
-		Scene->FXSystem->PreInitViews(RHICmdList);
+		Scene->FXSystem->PreInitViews();
 	}
 
 	// Draw lines to lights affecting this mesh if its selected.
diff --git a/Engine/Source/Runtime/Renderer/Public/MeshPassProcessor.h b/Engine/Source/Runtime/Renderer/Public/MeshPassProcessor.h
--- a/Engine/Source/Runtime/Renderer/Public/MeshPassProcessor.h
+++ b/Engine/Source/Runtime/Renderer/Public/MeshPassProcessor.h
@@ -466,11 +466,7 @@ public:
 			uint32 NumVertices;
 		} VertexParams;
 		
-		struct  
-		{
-			FVertexBufferRHIParamRef Buffer;
-			uint32 Offset;
-		} IndirectArgs;
+		FVertexBufferRHIParamRef IndirectArgsBuffer;
 	};
 
 	int8 PrimitiveIdStreamIndex;
@@ -493,7 +489,7 @@ public:
 			&& NumPrimitives == Rhs.NumPrimitives
 			&& NumInstances == Rhs.NumInstances
 			&& ((NumPrimitives > 0 && VertexParams.BaseVertexIndex == Rhs.VertexParams.BaseVertexIndex && VertexParams.NumVertices == Rhs.VertexParams.NumVertices)
-				|| (NumPrimitives == 0 && IndirectArgs.Buffer == Rhs.IndirectArgs.Buffer && IndirectArgs.Offset == Rhs.IndirectArgs.Offset));
+				|| (NumPrimitives == 0 && IndirectArgsBuffer == Rhs.IndirectArgsBuffer));
 	}
 
 	uint32 GetDynamicInstancingHash() const
