Add nullptr check before dereferencing

diff --git a/Engine/Plugins/FX/Niagara/Source/NiagaraEditor/Private/ViewModels/Stack/NiagaraStackFunctionInput.cpp b/Engine/Plugins/FX/Niagara/Source/NiagaraEditor/Private/ViewModels/Stack/NiagaraStackFunctionInput.cpp
--- a/Engine/Plugins/FX/Niagara/Source/NiagaraEditor/Private/ViewModels/Stack/NiagaraStackFunctionInput.cpp
+++ b/Engine/Plugins/FX/Niagara/Source/NiagaraEditor/Private/ViewModels/Stack/NiagaraStackFunctionInput.cpp
@@ -781,7 +781,7 @@ void UNiagaraStackFunctionInput::GetAvailableParameterHandles(TArray<FNiagaraPar
 		// Check if this is in a spawn event handler and the emitter is not using interpolated spawn so we
 		// we can hide particle update parameters
 		bool bSpawnScript = false;
-		if (CurrentOutputNode->GetUsage() == ENiagaraScriptUsage::ParticleEventScript)
+		if (CurrentOutputNode != nullptr && CurrentOutputNode->GetUsage() == ENiagaraScriptUsage::ParticleEventScript)
 		{
 			for (const FNiagaraEventScriptProperties &EventHandlerProps : GetEmitterViewModel()->GetEmitter()->GetEventHandlers())
 			{
