Only hide stack parameters compiled out by static switches, don't remove them completely

diff --git a/Engine/Plugins/FX/Niagara/Source/NiagaraEditor/Private/Customizations/NiagaraStaticSwitchNodeDetails.cpp b/Engine/Plugins/FX/Niagara/Source/NiagaraEditor/Private/Customizations/NiagaraStaticSwitchNodeDetails.cpp
--- a/Engine/Plugins/FX/Niagara/Source/NiagaraEditor/Private/Customizations/NiagaraStaticSwitchNodeDetails.cpp
+++ b/Engine/Plugins/FX/Niagara/Source/NiagaraEditor/Private/Customizations/NiagaraStaticSwitchNodeDetails.cpp
@@ -320,8 +320,12 @@ void FNiagaraStaticSwitchNodeDetails::RefreshDefaultDropdownValues()
 	if (Enum)
 	{
 		SelectedDefaultValue.Reset();
-		for (int i = 0; i < Enum->GetMaxEnumValue() - 1; i++)
+		for (int i = 0; i < Enum->GetMaxEnumValue(); i++)
 		{
+			if (!Enum->IsValidEnumValue(i))
+			{
+				continue;
+			}
 			FText DisplayName = Enum->GetDisplayNameTextByIndex(i);
 			DefaultEnumDropdownOptions.Add(MakeShared<DefaultEnumOption>(DisplayName, i));
 
diff --git a/Engine/Plugins/FX/Niagara/Source/NiagaraEditor/Private/NiagaraEditorUtilities.cpp b/Engine/Plugins/FX/Niagara/Source/NiagaraEditor/Private/NiagaraEditorUtilities.cpp
--- a/Engine/Plugins/FX/Niagara/Source/NiagaraEditor/Private/NiagaraEditorUtilities.cpp
+++ b/Engine/Plugins/FX/Niagara/Source/NiagaraEditor/Private/NiagaraEditorUtilities.cpp
@@ -630,14 +630,16 @@ void FNiagaraEditorUtilities::SetStaticSwitchConstants(UNiagaraGraph* Graph, con
 		UNiagaraNodeStaticSwitch* SwitchNode = Cast<UNiagaraNodeStaticSwitch>(Node);
 		if (SwitchNode)
 		{
-			SwitchNode->IsValueSet = false;
+			SwitchNode->ClearSwitchValue();
 			for (UEdGraphPin* InputPin : CallInputs)
 			{
 				if (InputPin->GetFName().IsEqual(SwitchNode->InputParameterName))
 				{
 					int32 SwitchValue = 0;
-					SwitchNode->IsValueSet = ResolveConstantValue(InputPin, SwitchValue);
-					SwitchNode->SwitchValue = SwitchValue;
+					if (ResolveConstantValue(InputPin, SwitchValue))
+					{
+						SwitchNode->SetSwitchValue(SwitchValue);
+					}
 					break;
 				}
 			}
diff --git a/Engine/Plugins/FX/Niagara/Source/NiagaraEditor/Private/NiagaraGraph.cpp b/Engine/Plugins/FX/Niagara/Source/NiagaraEditor/Private/NiagaraGraph.cpp
--- a/Engine/Plugins/FX/Niagara/Source/NiagaraEditor/Private/NiagaraGraph.cpp
+++ b/Engine/Plugins/FX/Niagara/Source/NiagaraEditor/Private/NiagaraGraph.cpp
@@ -614,20 +614,20 @@ void UNiagaraGraph::FindInputNodes(TArray<UNiagaraNodeInput*>& OutInputNodes, UN
 	}
 }
 
-TArray<FNiagaraVariable> UNiagaraGraph::FindStaticSwitchInputs() const
+TArray<FNiagaraVariable> UNiagaraGraph::FindStaticSwitchInputs(bool bReachableOnly) const
 {
+	TArray<UEdGraphNode*> NodesToProcess = bReachableOnly ? FindReachbleNodes() : Nodes;
+
 	TArray<FNiagaraVariable> Result;
-	for (UEdGraphNode* Node : Nodes)
+	for (UEdGraphNode* Node : NodesToProcess)
 	{
-		UNiagaraNodeStaticSwitch* SwitchNode = Cast<UNiagaraNodeStaticSwitch>(Node);
-		if (SwitchNode)
+		if (UNiagaraNodeStaticSwitch* SwitchNode = Cast<UNiagaraNodeStaticSwitch>(Node))
 		{
 			FNiagaraVariable Variable(SwitchNode->GetInputType(), SwitchNode->InputParameterName);
 			Result.AddUnique(Variable);
 		}
 
-		UNiagaraNodeFunctionCall* FunctionNode = Cast<UNiagaraNodeFunctionCall>(Node);
-		if (FunctionNode)
+		if (UNiagaraNodeFunctionCall* FunctionNode = Cast<UNiagaraNodeFunctionCall>(Node))
 		{
 			for (const FNiagaraPropagatedVariable& Propagated : FunctionNode->PropagatedStaticSwitchParameters)
 			{
@@ -647,6 +647,57 @@ TArray<FNiagaraVariable> UNiagaraGraph::FindStaticSwitchInputs() const
 	return Result;
 }
 
+TArray<UEdGraphNode*> UNiagaraGraph::FindReachbleNodes() const
+{
+	TArray<UEdGraphNode*> ResultNodes;
+	TArray<UNiagaraNodeOutput*> OutNodes;
+	FindOutputNodes(OutNodes);
+	ResultNodes.Append(OutNodes);
+
+	for (int i = 0; i < ResultNodes.Num(); i++)
+	{
+		UEdGraphNode* Node = ResultNodes[i];
+		if (Node == nullptr)
+		{
+			continue;
+		}
+		
+		if (UNiagaraNodeStaticSwitch* SwitchNode = Cast<UNiagaraNodeStaticSwitch>(Node))
+		{
+			TArray<UEdGraphPin*> OutPins;
+			SwitchNode->GetOutputPins(OutPins);
+			for (UEdGraphPin* Pin : OutPins)
+			{
+				UEdGraphPin* TracedPin = SwitchNode->GetTracedOutputPin(Pin, false);
+				if (TracedPin && TracedPin != Pin)
+				{
+					ResultNodes.AddUnique(TracedPin->GetOwningNode());
+				}
+			}
+		}
+		else
+		{
+			TArray<UEdGraphPin*> InputPins;
+			for (UEdGraphPin* Pin : Node->GetAllPins())
+			{
+				if (!Pin || Pin->Direction != EEdGraphPinDirection::EGPD_Input)
+				{
+					continue;
+				}
+				for (UEdGraphPin* LinkedPin : Pin->LinkedTo)
+				{
+					if (!LinkedPin)
+					{
+						continue;
+					}
+					ResultNodes.AddUnique(LinkedPin->GetOwningNode());
+				}
+			}
+		}
+	}
+	return ResultNodes;
+}
+
 void UNiagaraGraph::GetParameters(TArray<FNiagaraVariable>& Inputs, TArray<FNiagaraVariable>& Outputs)const
 {
 	Inputs.Empty();
diff --git a/Engine/Plugins/FX/Niagara/Source/NiagaraEditor/Private/NiagaraNode.cpp b/Engine/Plugins/FX/Niagara/Source/NiagaraEditor/Private/NiagaraNode.cpp
--- a/Engine/Plugins/FX/Niagara/Source/NiagaraEditor/Private/NiagaraNode.cpp
+++ b/Engine/Plugins/FX/Niagara/Source/NiagaraEditor/Private/NiagaraNode.cpp
@@ -354,11 +354,11 @@ bool UNiagaraNode::CanAddToGraph(UNiagaraGraph* TargetGraph, FString& OutErrorMs
 	return true;
 }
 
-void UNiagaraNode::BuildParameterMapHistory(FNiagaraParameterMapHistoryBuilder& OutHistory, bool bRecursive) const
+void UNiagaraNode::BuildParameterMapHistory(FNiagaraParameterMapHistoryBuilder& OutHistory, bool bRecursive /*= true*/, bool bFilterForCompilation /*= true*/) const
 {
 	if (bRecursive)
 	{
-		OutHistory.VisitInputPins(this);
+		OutHistory.VisitInputPins(this, bFilterForCompilation);
 	}
 }
 
@@ -545,7 +545,7 @@ ENiagaraNumericOutputTypeSelectionMode UNiagaraNode::GetNumericOutputTypeSelecti
 }
 
 
-UEdGraphPin* UNiagaraNode::TraceOutputPin(UEdGraphPin* LocallyOwnedOutputPin)
+UEdGraphPin* UNiagaraNode::TraceOutputPin(UEdGraphPin* LocallyOwnedOutputPin, bool bFilterForCompilation)
 {
 	if (LocallyOwnedOutputPin == nullptr)
 	{
diff --git a/Engine/Plugins/FX/Niagara/Source/NiagaraEditor/Private/NiagaraNodeAssignment.cpp b/Engine/Plugins/FX/Niagara/Source/NiagaraEditor/Private/NiagaraNodeAssignment.cpp
--- a/Engine/Plugins/FX/Niagara/Source/NiagaraEditor/Private/NiagaraNodeAssignment.cpp
+++ b/Engine/Plugins/FX/Niagara/Source/NiagaraEditor/Private/NiagaraNodeAssignment.cpp
@@ -172,9 +172,9 @@ void UNiagaraNodeAssignment::PostLoad()
 	}
 }
 
-void UNiagaraNodeAssignment::BuildParameterMapHistory(FNiagaraParameterMapHistoryBuilder& OutHistory, bool bRecursive) const
+void UNiagaraNodeAssignment::BuildParameterMapHistory(FNiagaraParameterMapHistoryBuilder& OutHistory, bool bRecursive /*= true*/, bool bFilterForCompilation /*= true*/) const
 {
-	Super::BuildParameterMapHistory(OutHistory, bRecursive);
+	Super::BuildParameterMapHistory(OutHistory, bRecursive, bFilterForCompilation);
 }
 
 void UNiagaraNodeAssignment::GatherExternalDependencyIDs(ENiagaraScriptUsage InMasterUsage, const FGuid& InMasterUsageId, TArray<FNiagaraCompileHash>& InReferencedCompileHashes, TArray<FGuid>& InReferencedIDs, TArray<UObject*>& InReferencedObjs) const
diff --git a/Engine/Plugins/FX/Niagara/Source/NiagaraEditor/Private/NiagaraNodeCustomHlsl.cpp b/Engine/Plugins/FX/Niagara/Source/NiagaraEditor/Private/NiagaraNodeCustomHlsl.cpp
--- a/Engine/Plugins/FX/Niagara/Source/NiagaraEditor/Private/NiagaraNodeCustomHlsl.cpp
+++ b/Engine/Plugins/FX/Niagara/Source/NiagaraEditor/Private/NiagaraNodeCustomHlsl.cpp
@@ -191,9 +191,9 @@ void UNiagaraNodeCustomHlsl::MoveDynamicPin(UEdGraphPin* Pin, int32 DirectionToM
 	RebuildSignatureFromPins();
 }
 
-void UNiagaraNodeCustomHlsl::BuildParameterMapHistory(FNiagaraParameterMapHistoryBuilder& OutHistory, bool bRecursive) const
+void UNiagaraNodeCustomHlsl::BuildParameterMapHistory(FNiagaraParameterMapHistoryBuilder& OutHistory, bool bRecursive /*= true*/, bool bFilterForCompilation /*= true*/) const
 {
-	Super::BuildParameterMapHistory(OutHistory, bRecursive);
+	Super::BuildParameterMapHistory(OutHistory, bRecursive, bFilterForCompilation);
 	if (!IsNodeEnabled() && OutHistory.GetIgnoreDisabled())
 	{
 		RouteParameterMapAroundMe(OutHistory, bRecursive);
diff --git a/Engine/Plugins/FX/Niagara/Source/NiagaraEditor/Private/NiagaraNodeEmitter.cpp b/Engine/Plugins/FX/Niagara/Source/NiagaraEditor/Private/NiagaraNodeEmitter.cpp
--- a/Engine/Plugins/FX/Niagara/Source/NiagaraEditor/Private/NiagaraNodeEmitter.cpp
+++ b/Engine/Plugins/FX/Niagara/Source/NiagaraEditor/Private/NiagaraNodeEmitter.cpp
@@ -256,9 +256,9 @@ FText UNiagaraNodeEmitter::GetNameFromEmitter()
 	return FText();
 }
 
-void UNiagaraNodeEmitter::BuildParameterMapHistory(FNiagaraParameterMapHistoryBuilder& OutHistory, bool bRecursive) const
+void UNiagaraNodeEmitter::BuildParameterMapHistory(FNiagaraParameterMapHistoryBuilder& OutHistory, bool bRecursive /*= true*/, bool bFilterForCompilation /*= true*/) const
 {
-	Super::BuildParameterMapHistory(OutHistory, bRecursive);
+	Super::BuildParameterMapHistory(OutHistory, bRecursive, bFilterForCompilation);
 
 	if (!IsNodeEnabled() && OutHistory.GetIgnoreDisabled())
 	{
diff --git a/Engine/Plugins/FX/Niagara/Source/NiagaraEditor/Private/NiagaraNodeEmitter.h b/Engine/Plugins/FX/Niagara/Source/NiagaraEditor/Private/NiagaraNodeEmitter.h
--- a/Engine/Plugins/FX/Niagara/Source/NiagaraEditor/Private/NiagaraNodeEmitter.h
+++ b/Engine/Plugins/FX/Niagara/Source/NiagaraEditor/Private/NiagaraNodeEmitter.h
@@ -45,7 +45,7 @@ public:
 	virtual bool VerifyEditablePinName(const FText& InName, FText& OutErrorMessage, const UEdGraphPin* InGraphPinObj) const override;
 	virtual bool CommitEditablePinName(const FText& InName, UEdGraphPin* InGraphPinObj)  override;
 
-	virtual void BuildParameterMapHistory(FNiagaraParameterMapHistoryBuilder& OutHistory, bool bRecursive = true) const override;
+	virtual void BuildParameterMapHistory(FNiagaraParameterMapHistoryBuilder& OutHistory, bool bRecursive = true, bool bFilterForCompilation = true) const override;
 
 	virtual FText GetNodeTitle(ENodeTitleType::Type TitleType) const;
 
diff --git a/Engine/Plugins/FX/Niagara/Source/NiagaraEditor/Private/NiagaraNodeFunctionCall.cpp b/Engine/Plugins/FX/Niagara/Source/NiagaraEditor/Private/NiagaraNodeFunctionCall.cpp
--- a/Engine/Plugins/FX/Niagara/Source/NiagaraEditor/Private/NiagaraNodeFunctionCall.cpp
+++ b/Engine/Plugins/FX/Niagara/Source/NiagaraEditor/Private/NiagaraNodeFunctionCall.cpp
@@ -69,6 +69,16 @@ void UNiagaraNodeFunctionCall::PostLoad()
 		}
 	}
 
+	// Clean up invalid old references to propagated parameters
+	for (int i = PropagatedStaticSwitchParameters.Num() - 1; i >= 0; i--)
+	{
+		FNiagaraPropagatedVariable& Propagated = PropagatedStaticSwitchParameters[i];
+		if (Propagated.SwitchParameter.GetName().IsNone())
+		{
+			PropagatedStaticSwitchParameters.RemoveAt(i);
+		}
+	}
+
 	if (FunctionDisplayName.IsEmpty())
 	{
 		ComputeNodeName();
@@ -644,9 +654,9 @@ bool UNiagaraNodeFunctionCall::ScriptIsValid() const
 	return false;
 }
 
-void UNiagaraNodeFunctionCall::BuildParameterMapHistory(FNiagaraParameterMapHistoryBuilder& OutHistory, bool bRecursive) const
+void UNiagaraNodeFunctionCall::BuildParameterMapHistory(FNiagaraParameterMapHistoryBuilder& OutHistory, bool bRecursive /*= true*/, bool bFilterForCompilation /*= true*/) const
 {
-	Super::BuildParameterMapHistory(OutHistory, bRecursive);
+	Super::BuildParameterMapHistory(OutHistory, bRecursive, bFilterForCompilation);
 	if (!IsNodeEnabled() && OutHistory.GetIgnoreDisabled())
 	{
 		RouteParameterMapAroundMe(OutHistory, bRecursive);
@@ -658,7 +668,7 @@ void UNiagaraNodeFunctionCall::BuildParameterMapHistory(FNiagaraParameterMapHist
 	{
 		UNiagaraScriptSource* Source = CastChecked<UNiagaraScriptSource>(FunctionScript->GetSource());
 		UNiagaraGraph* FunctionGraph = CastChecked<UNiagaraGraph>(Source->NodeGraph);
-		
+
 		UNiagaraNodeOutput* OutputNode = FunctionGraph->FindOutputNode(ENiagaraScriptUsage::Function);
 		if (OutputNode == nullptr)
 		{
@@ -679,13 +689,13 @@ void UNiagaraNodeFunctionCall::BuildParameterMapHistory(FNiagaraParameterMapHist
 				ParamMapIdx = OutHistory.TraceParameterMapOutputPin(UNiagaraNode::TraceOutputPin(GetInputPin(0)->LinkedTo[0]));
 			}
 		}
-		
+
 		OutHistory.EnterFunction(GetFunctionName(), FunctionScript, this);
 		if (ParamMapIdx != INDEX_NONE)
 		{
 			NodeIdx = OutHistory.BeginNodeVisitation(ParamMapIdx, this);
 		}
-		OutputNode->BuildParameterMapHistory(OutHistory, true);
+		OutputNode->BuildParameterMapHistory(OutHistory, true, bFilterForCompilation);
 
 		// Since we're about to lose the pin calling context, we finish up the function call parameter map pin wiring
 		// here when we have the calling context and the child context still available to us...
@@ -693,12 +703,12 @@ void UNiagaraNodeFunctionCall::BuildParameterMapHistory(FNiagaraParameterMapHist
 		GetOutputPins(OutputPins);
 
 		TArray<TPair<UEdGraphPin*, int32> > MatchedPairs;
-		
+
 		// Find the matches of names and types of the sub-graph output pins and this function call nodes' outputs.
 		for (UEdGraphPin* ChildOutputNodePin : OutputNode->GetAllPins())
 		{
 			FNiagaraVariable VarChild = Schema->PinToNiagaraVariable(ChildOutputNodePin);
-			
+
 			if (ChildOutputNodePin->LinkedTo.Num() > 0 && VarChild.GetType() == FNiagaraTypeDefinition::GetParameterMapDef())
 			{
 				for (int32 i = 0; i < OutputPins.Num(); i++)
diff --git a/Engine/Plugins/FX/Niagara/Source/NiagaraEditor/Private/NiagaraNodeInput.cpp b/Engine/Plugins/FX/Niagara/Source/NiagaraEditor/Private/NiagaraNodeInput.cpp
--- a/Engine/Plugins/FX/Niagara/Source/NiagaraEditor/Private/NiagaraNodeInput.cpp
+++ b/Engine/Plugins/FX/Niagara/Source/NiagaraEditor/Private/NiagaraNodeInput.cpp
@@ -72,7 +72,7 @@ void UNiagaraNodeInput::PostLoad()
 	}
 }
 
-void UNiagaraNodeInput::BuildParameterMapHistory(FNiagaraParameterMapHistoryBuilder& OutHistory, bool bRecursive) const
+void UNiagaraNodeInput::BuildParameterMapHistory(FNiagaraParameterMapHistoryBuilder& OutHistory, bool bRecursive /*= true*/, bool bFilterForCompilation /*= true*/) const
 {
 	if (!IsNodeEnabled() && OutHistory.GetIgnoreDisabled())
 	{
diff --git a/Engine/Plugins/FX/Niagara/Source/NiagaraEditor/Private/NiagaraNodeParameterMapGet.cpp b/Engine/Plugins/FX/Niagara/Source/NiagaraEditor/Private/NiagaraNodeParameterMapGet.cpp
--- a/Engine/Plugins/FX/Niagara/Source/NiagaraEditor/Private/NiagaraNodeParameterMapGet.cpp
+++ b/Engine/Plugins/FX/Niagara/Source/NiagaraEditor/Private/NiagaraNodeParameterMapGet.cpp
@@ -291,11 +291,11 @@ FText UNiagaraNodeParameterMapGet::GetNodeTitle(ENodeTitleType::Type TitleType)
 	return LOCTEXT("UNiagaraNodeParameterMapGetName", "Map Get");
 }
 
-void UNiagaraNodeParameterMapGet::BuildParameterMapHistory(FNiagaraParameterMapHistoryBuilder& OutHistory, bool bRecursive) const
+void UNiagaraNodeParameterMapGet::BuildParameterMapHistory(FNiagaraParameterMapHistoryBuilder& OutHistory, bool bRecursive /*= true*/, bool bFilterForCompilation /*= true*/) const
 {
 	if (bRecursive)
 	{
-		OutHistory.VisitInputPin(GetInputPin(0), this);
+		OutHistory.VisitInputPin(GetInputPin(0), this, bFilterForCompilation);
 	}
 
 	if (!IsNodeEnabled() && OutHistory.GetIgnoreDisabled())
@@ -325,11 +325,11 @@ void UNiagaraNodeParameterMapGet::BuildParameterMapHistory(FNiagaraParameterMapH
 			bool bUsedDefaults = false;
 			if (bRecursive)
 			{
-				OutHistory.HandleVariableRead(ParamMapIdx, OutputPins[i], true, GetDefaultPin(OutputPins[i]), bUsedDefaults);
+				OutHistory.HandleVariableRead(ParamMapIdx, OutputPins[i], true, GetDefaultPin(OutputPins[i]), bFilterForCompilation, bUsedDefaults);
 			}
 			else
 			{
-				OutHistory.HandleVariableRead(ParamMapIdx, OutputPins[i], true, nullptr, bUsedDefaults);
+				OutHistory.HandleVariableRead(ParamMapIdx, OutputPins[i], true, nullptr, bFilterForCompilation, bUsedDefaults);
 			}
 		}
 		OutHistory.EndNodeVisitation(ParamMapIdx, NodeIdx);
diff --git a/Engine/Plugins/FX/Niagara/Source/NiagaraEditor/Private/NiagaraNodeParameterMapGet.h b/Engine/Plugins/FX/Niagara/Source/NiagaraEditor/Private/NiagaraNodeParameterMapGet.h
--- a/Engine/Plugins/FX/Niagara/Source/NiagaraEditor/Private/NiagaraNodeParameterMapGet.h
+++ b/Engine/Plugins/FX/Niagara/Source/NiagaraEditor/Private/NiagaraNodeParameterMapGet.h
@@ -30,7 +30,7 @@ public:
 
 	virtual void Compile(class FHlslNiagaraTranslator* Translator, TArray<int32>& Outputs);
 
-	virtual void BuildParameterMapHistory(FNiagaraParameterMapHistoryBuilder& OutHistory, bool bRecursive = true) const override;
+	virtual void BuildParameterMapHistory(FNiagaraParameterMapHistoryBuilder& OutHistory, bool bRecursive = true, bool bFilterForCompilation = true) const override;
 
 	virtual FText GetNodeTitle(ENodeTitleType::Type TitleType) const;
 
diff --git a/Engine/Plugins/FX/Niagara/Source/NiagaraEditor/Private/NiagaraNodeParameterMapSet.cpp b/Engine/Plugins/FX/Niagara/Source/NiagaraEditor/Private/NiagaraNodeParameterMapSet.cpp
--- a/Engine/Plugins/FX/Niagara/Source/NiagaraEditor/Private/NiagaraNodeParameterMapSet.cpp
+++ b/Engine/Plugins/FX/Niagara/Source/NiagaraEditor/Private/NiagaraNodeParameterMapSet.cpp
@@ -191,7 +191,7 @@ FText UNiagaraNodeParameterMapSet::GetNodeTitle(ENodeTitleType::Type TitleType)
 }
 
 
-void UNiagaraNodeParameterMapSet::BuildParameterMapHistory(FNiagaraParameterMapHistoryBuilder& OutHistory, bool bRecursive) const
+void UNiagaraNodeParameterMapSet::BuildParameterMapHistory(FNiagaraParameterMapHistoryBuilder& OutHistory, bool bRecursive /*= true*/, bool bFilterForCompilation /*= true*/) const
 {
 	const UEdGraphSchema_Niagara* Schema = GetDefault<UEdGraphSchema_Niagara>();
 	TArray<UEdGraphPin*> InputPins;
@@ -207,7 +207,7 @@ void UNiagaraNodeParameterMapSet::BuildParameterMapHistory(FNiagaraParameterMapH
 			continue;
 		}
 
-		OutHistory.VisitInputPin(InputPins[i], this);
+		OutHistory.VisitInputPin(InputPins[i], this, bFilterForCompilation);
 
 
 		if (!IsNodeEnabled() && OutHistory.GetIgnoreDisabled())
diff --git a/Engine/Plugins/FX/Niagara/Source/NiagaraEditor/Private/NiagaraNodeParameterMapSet.h b/Engine/Plugins/FX/Niagara/Source/NiagaraEditor/Private/NiagaraNodeParameterMapSet.h
--- a/Engine/Plugins/FX/Niagara/Source/NiagaraEditor/Private/NiagaraNodeParameterMapSet.h
+++ b/Engine/Plugins/FX/Niagara/Source/NiagaraEditor/Private/NiagaraNodeParameterMapSet.h
@@ -26,7 +26,7 @@ public:
 	virtual bool CommitEditablePinName(const FText& InName, UEdGraphPin* InGraphPinObj)  override;
 	virtual bool CancelEditablePinName(const FText& InName, UEdGraphPin* InGraphPinObj) override;
 
-	virtual void BuildParameterMapHistory(FNiagaraParameterMapHistoryBuilder& OutHistory, bool bRecursive = true) const override;
+	virtual void BuildParameterMapHistory(FNiagaraParameterMapHistoryBuilder& OutHistory, bool bRecursive = true, bool bFilterForCompilation = true) const override;
 
 	virtual void Compile(class FHlslNiagaraTranslator* Translator, TArray<int32>& Outputs);
 
diff --git a/Engine/Plugins/FX/Niagara/Source/NiagaraEditor/Private/NiagaraNodeReadDataSet.cpp b/Engine/Plugins/FX/Niagara/Source/NiagaraEditor/Private/NiagaraNodeReadDataSet.cpp
--- a/Engine/Plugins/FX/Niagara/Source/NiagaraEditor/Private/NiagaraNodeReadDataSet.cpp
+++ b/Engine/Plugins/FX/Niagara/Source/NiagaraEditor/Private/NiagaraNodeReadDataSet.cpp
@@ -109,11 +109,11 @@ void UNiagaraNodeReadDataSet::Compile(class FHlslNiagaraTranslator* Translator,
 	Translator->ReadDataSet(DataSet, Variables, ENiagaraDataSetAccessMode::AppendConsume, Inputs[0], Outputs);
 }
 
-void UNiagaraNodeReadDataSet::BuildParameterMapHistory(FNiagaraParameterMapHistoryBuilder& OutHistory, bool bRecursive) const
+void UNiagaraNodeReadDataSet::BuildParameterMapHistory(FNiagaraParameterMapHistoryBuilder& OutHistory, bool bRecursive /*= true*/, bool bFilterForCompilation /*= true*/) const
 {
 	if (bRecursive)
 	{
-		OutHistory.VisitInputPins(this);
+		OutHistory.VisitInputPins(this, bFilterForCompilation);
 	}
 
 	if (!IsNodeEnabled() && OutHistory.GetIgnoreDisabled())
diff --git a/Engine/Plugins/FX/Niagara/Source/NiagaraEditor/Private/NiagaraNodeReroute.cpp b/Engine/Plugins/FX/Niagara/Source/NiagaraEditor/Private/NiagaraNodeReroute.cpp
--- a/Engine/Plugins/FX/Niagara/Source/NiagaraEditor/Private/NiagaraNodeReroute.cpp
+++ b/Engine/Plugins/FX/Niagara/Source/NiagaraEditor/Private/NiagaraNodeReroute.cpp
@@ -117,7 +117,7 @@ void UNiagaraNodeReroute::PinConnectionListChanged(UEdGraphPin* Pin)
 	Super::PinConnectionListChanged(Pin);
 }
 
-void UNiagaraNodeReroute::BuildParameterMapHistory(FNiagaraParameterMapHistoryBuilder& OutHistory, bool bRecursive) const
+void UNiagaraNodeReroute::BuildParameterMapHistory(FNiagaraParameterMapHistoryBuilder& OutHistory, bool bRecursive /*= true*/, bool bFilterForCompilation /*= true*/) const
 {
 	// Should never hit here...
 	check(false);
diff --git a/Engine/Plugins/FX/Niagara/Source/NiagaraEditor/Private/NiagaraNodeReroute.h b/Engine/Plugins/FX/Niagara/Source/NiagaraEditor/Private/NiagaraNodeReroute.h
--- a/Engine/Plugins/FX/Niagara/Source/NiagaraEditor/Private/NiagaraNodeReroute.h
+++ b/Engine/Plugins/FX/Niagara/Source/NiagaraEditor/Private/NiagaraNodeReroute.h
@@ -33,7 +33,7 @@ public:
 	virtual void Compile(class FHlslNiagaraTranslator* Translator, TArray<int32>& Outputs) override;
 	virtual bool RefreshFromExternalChanges() override;
 	virtual void PinConnectionListChanged(UEdGraphPin* Pin) override;
-	void BuildParameterMapHistory(FNiagaraParameterMapHistoryBuilder& OutHistory, bool bRecursive) const override;
+	void BuildParameterMapHistory(FNiagaraParameterMapHistoryBuilder& OutHistory, bool bRecursive = true, bool bFilterForCompilation = true) const override;
 	virtual ENiagaraNumericOutputTypeSelectionMode GetNumericOutputTypeSelectionMode() const;
 
 	/** Trace to an output pin that is not a reroute node output pin. If the reroute nodes resulted in a dead end and no output pin was found return nullptr.*/
diff --git a/Engine/Plugins/FX/Niagara/Source/NiagaraEditor/Private/NiagaraNodeSimTargetSelector.cpp b/Engine/Plugins/FX/Niagara/Source/NiagaraEditor/Private/NiagaraNodeSimTargetSelector.cpp
--- a/Engine/Plugins/FX/Niagara/Source/NiagaraEditor/Private/NiagaraNodeSimTargetSelector.cpp
+++ b/Engine/Plugins/FX/Niagara/Source/NiagaraEditor/Private/NiagaraNodeSimTargetSelector.cpp
@@ -104,9 +104,9 @@ UEdGraphPin* UNiagaraNodeSimTargetSelector::GetPassThroughPin(const UEdGraphPin*
 	return nullptr;
 }
 
-void UNiagaraNodeSimTargetSelector::BuildParameterMapHistory(FNiagaraParameterMapHistoryBuilder& OutHistory, bool bRecursive) const
+void UNiagaraNodeSimTargetSelector::BuildParameterMapHistory(FNiagaraParameterMapHistoryBuilder& OutHistory, bool bRecursive /*= true*/, bool bFilterForCompilation /*= true*/) const
 {
-	UNiagaraNode::BuildParameterMapHistory(OutHistory, bRecursive);
+	UNiagaraNode::BuildParameterMapHistory(OutHistory, bRecursive, bFilterForCompilation);
 }
 
 FText UNiagaraNodeSimTargetSelector::GetTooltipText() const
diff --git a/Engine/Plugins/FX/Niagara/Source/NiagaraEditor/Private/NiagaraNodeSimTargetSelector.h b/Engine/Plugins/FX/Niagara/Source/NiagaraEditor/Private/NiagaraNodeSimTargetSelector.h
--- a/Engine/Plugins/FX/Niagara/Source/NiagaraEditor/Private/NiagaraNodeSimTargetSelector.h
+++ b/Engine/Plugins/FX/Niagara/Source/NiagaraEditor/Private/NiagaraNodeSimTargetSelector.h
@@ -19,7 +19,7 @@ public:
 
 	//~ Begin UNiagaraNode Interface
 	virtual void Compile(class FHlslNiagaraTranslator* Translator, TArray<int32>& Outputs) override;
-	void BuildParameterMapHistory(FNiagaraParameterMapHistoryBuilder& OutHistory, bool bRecursive) const override;
+	void BuildParameterMapHistory(FNiagaraParameterMapHistoryBuilder& OutHistory, bool bRecursive = true, bool bFilterForCompilation = true) const override;
 	virtual UEdGraphPin* GetPassThroughPin(const UEdGraphPin* LocallyOwnedOutputPin, ENiagaraScriptUsage MasterUsage) const override;
 	//~ End UNiagaraNode Interface
 
diff --git a/Engine/Plugins/FX/Niagara/Source/NiagaraEditor/Private/NiagaraNodeStaticSwitch.cpp b/Engine/Plugins/FX/Niagara/Source/NiagaraEditor/Private/NiagaraNodeStaticSwitch.cpp
--- a/Engine/Plugins/FX/Niagara/Source/NiagaraEditor/Private/NiagaraNodeStaticSwitch.cpp
+++ b/Engine/Plugins/FX/Niagara/Source/NiagaraEditor/Private/NiagaraNodeStaticSwitch.cpp
@@ -43,6 +43,18 @@ void UNiagaraNodeStaticSwitch::OnSwitchParameterTypeChanged(const FNiagaraTypeDe
 	RemoveUnusedGraphParameter(FNiagaraVariable(OldType, InputParameterName));
 }
 
+void UNiagaraNodeStaticSwitch::SetSwitchValue(int Value)
+{
+	IsValueSet = true;
+	SwitchValue = Value;
+}
+
+void UNiagaraNodeStaticSwitch::ClearSwitchValue()
+{
+	IsValueSet = false;
+	SwitchValue = 0;
+}
+
 void UNiagaraNodeStaticSwitch::RemoveUnusedGraphParameter(const FNiagaraVariable& OldParameter)
 {
 	TArray<FNiagaraVariable> GraphVariables = GetNiagaraGraph()->FindStaticSwitchInputs();
@@ -279,6 +291,11 @@ bool UNiagaraNodeStaticSwitch::SubstituteCompiledPin(FHlslNiagaraTranslator* Tra
 }
 
 UEdGraphPin* UNiagaraNodeStaticSwitch::GetTracedOutputPin(UEdGraphPin* LocallyOwnedOutputPin) const
+{
+	return GetTracedOutputPin(LocallyOwnedOutputPin, true);
+}
+
+UEdGraphPin* UNiagaraNodeStaticSwitch::GetTracedOutputPin(UEdGraphPin* LocallyOwnedOutputPin, bool bRecursive) const
 {
 	TArray<UEdGraphPin*> InputPins;
 	GetInputPins(InputPins);
@@ -293,18 +310,23 @@ UEdGraphPin* UNiagaraNodeStaticSwitch::GetTracedOutputPin(UEdGraphPin* LocallyOw
 			continue;
 		}
 		int32 VarIdx;
-		if (OutPin == LocallyOwnedOutputPin && GetVarIndex(nullptr, InputPins.Num(), IsValueSet ? SwitchValue : 0, VarIdx))
+		if (OutPin == LocallyOwnedOutputPin && GetVarIndex(nullptr, InputPins.Num(), SwitchValue, VarIdx))
 		{
 			UEdGraphPin* InputPin = InputPins[VarIdx + i];
 			if (InputPin->LinkedTo.Num() == 1)
 			{
-				return UNiagaraNode::TraceOutputPin(InputPin->LinkedTo[0]);
+				return bRecursive ? UNiagaraNode::TraceOutputPin(InputPin->LinkedTo[0]) : InputPin->LinkedTo[0];
 			}
 		}
 	}
 	return LocallyOwnedOutputPin;
 }
 
+void UNiagaraNodeStaticSwitch::BuildParameterMapHistory(FNiagaraParameterMapHistoryBuilder& OutHistory, bool bRecursive /*= true*/, bool bFilterForCompilation /*= true*/) const
+{
+	UNiagaraNode::BuildParameterMapHistory(OutHistory, bRecursive, bFilterForCompilation);
+}
+
 FText UNiagaraNodeStaticSwitch::GetTooltipText() const
 {
 	return LOCTEXT("NiagaraStaticSwitchNodeTooltip", "This is a compile-time switch that selects one branch to compile based on an input parameter.");
diff --git a/Engine/Plugins/FX/Niagara/Source/NiagaraEditor/Private/NiagaraNodeStaticSwitch.h b/Engine/Plugins/FX/Niagara/Source/NiagaraEditor/Private/NiagaraNodeStaticSwitch.h
--- a/Engine/Plugins/FX/Niagara/Source/NiagaraEditor/Private/NiagaraNodeStaticSwitch.h
+++ b/Engine/Plugins/FX/Niagara/Source/NiagaraEditor/Private/NiagaraNodeStaticSwitch.h
@@ -47,19 +47,14 @@ public:
 	UPROPERTY()
 	FStaticSwitchTypeData SwitchTypeData;
 
-	/** If true then the current SwitchValue should be used for compilation, otherwise the node is not yet connected to a constant value */
-	UPROPERTY(Transient)
-	bool IsValueSet;
-
-	/** The current value that the node is evaluated with. For bool 0 is false, for enums the value is the index of the enum value. */
-	UPROPERTY(Transient)
-	int32 SwitchValue;
-
 	FNiagaraTypeDefinition GetInputType() const;
 
 	void ChangeSwitchParameterName(const FName& NewName);
 	void OnSwitchParameterTypeChanged(const FNiagaraTypeDefinition& OldType);
 
+	void SetSwitchValue(int Value);
+	void ClearSwitchValue();
+
 	//~ Begin EdGraphNode Interface
 	virtual void AllocateDefaultPins() override;
 	virtual FText GetTooltipText() const override;
@@ -72,6 +67,8 @@ public:
 	virtual void Compile(class FHlslNiagaraTranslator* Translator, TArray<int32>& Outputs) override;
 	virtual bool SubstituteCompiledPin(FHlslNiagaraTranslator* Translator, UEdGraphPin** LocallyOwnedPin) override;
 	virtual UEdGraphPin* GetTracedOutputPin(UEdGraphPin* LocallyOwnedOutputPin) const override;
+	virtual UEdGraphPin* GetTracedOutputPin(UEdGraphPin* LocallyOwnedOutputPin, bool bRecursive) const;
+	virtual void BuildParameterMapHistory(FNiagaraParameterMapHistoryBuilder& OutHistory, bool bRecursive = true, bool bFilterForCompilation = true) const override;
 	//~ End UNiagaraNode Interface
 
 protected:
@@ -87,4 +84,10 @@ private:
 	bool GetVarIndex(class FHlslNiagaraTranslator* Translator, int32 InputPinCount, int32 Value, int32& VarIndexOut) const;
 
 	void RemoveUnusedGraphParameter(const FNiagaraVariable& OldParameter);
+
+	/** If true then the current SwitchValue should be used for compilation, otherwise the node is not yet connected to a constant value */
+	bool IsValueSet = false;
+
+	/** The current value that the node is evaluated with. For bool 0 is false, for enums the value is the index of the enum value. */
+	int32 SwitchValue = 0;
 };
diff --git a/Engine/Plugins/FX/Niagara/Source/NiagaraEditor/Private/NiagaraNodeUsageSelector.cpp b/Engine/Plugins/FX/Niagara/Source/NiagaraEditor/Private/NiagaraNodeUsageSelector.cpp
--- a/Engine/Plugins/FX/Niagara/Source/NiagaraEditor/Private/NiagaraNodeUsageSelector.cpp
+++ b/Engine/Plugins/FX/Niagara/Source/NiagaraEditor/Private/NiagaraNodeUsageSelector.cpp
@@ -191,7 +191,7 @@ void UNiagaraNodeUsageSelector::AppendFunctionAliasForContext(const FNiagaraGrap
 	}
 }
 
-void UNiagaraNodeUsageSelector::BuildParameterMapHistory(FNiagaraParameterMapHistoryBuilder& OutHistory, bool bRecursive) const
+void UNiagaraNodeUsageSelector::BuildParameterMapHistory(FNiagaraParameterMapHistoryBuilder& OutHistory, bool bRecursive /*= true*/, bool bFilterForCompilation /*= true*/) const
 {
 	const UEdGraphSchema_Niagara* Schema = CastChecked<UEdGraphSchema_Niagara>(GetSchema());
 
@@ -223,7 +223,7 @@ void UNiagaraNodeUsageSelector::BuildParameterMapHistory(FNiagaraParameterMapHis
 
 			for (int32 i = 0; i < OutputVars.Num(); i++)
 			{
-				OutHistory.VisitInputPin(InputPins[VarIdx + i], this);
+				OutHistory.VisitInputPin(InputPins[VarIdx + i], this, bFilterForCompilation);
 			}
 		}
 	}
diff --git a/Engine/Plugins/FX/Niagara/Source/NiagaraEditor/Private/NiagaraNodeUsageSelector.h b/Engine/Plugins/FX/Niagara/Source/NiagaraEditor/Private/NiagaraNodeUsageSelector.h
--- a/Engine/Plugins/FX/Niagara/Source/NiagaraEditor/Private/NiagaraNodeUsageSelector.h
+++ b/Engine/Plugins/FX/Niagara/Source/NiagaraEditor/Private/NiagaraNodeUsageSelector.h
@@ -32,7 +32,7 @@ public:
 	virtual void Compile(class FHlslNiagaraTranslator* Translator, TArray<int32>& Outputs) override;
 	virtual bool RefreshFromExternalChanges() override;
 
-	void BuildParameterMapHistory(FNiagaraParameterMapHistoryBuilder& OutHistory, bool bRecursive) const override;
+	virtual void BuildParameterMapHistory(FNiagaraParameterMapHistoryBuilder& OutHistory, bool bRecursive = true, bool bFilterForCompilation = true) const override;
 
 	//~ End UNiagaraNode Interface
 
diff --git a/Engine/Plugins/FX/Niagara/Source/NiagaraEditor/Private/NiagaraNodeWriteDataSet.cpp b/Engine/Plugins/FX/Niagara/Source/NiagaraEditor/Private/NiagaraNodeWriteDataSet.cpp
--- a/Engine/Plugins/FX/Niagara/Source/NiagaraEditor/Private/NiagaraNodeWriteDataSet.cpp
+++ b/Engine/Plugins/FX/Niagara/Source/NiagaraEditor/Private/NiagaraNodeWriteDataSet.cpp
@@ -113,11 +113,11 @@ void UNiagaraNodeWriteDataSet::PostLoad()
 	}
 }
 
-void UNiagaraNodeWriteDataSet::BuildParameterMapHistory(FNiagaraParameterMapHistoryBuilder& OutHistory, bool bRecursive) const
+void UNiagaraNodeWriteDataSet::BuildParameterMapHistory(FNiagaraParameterMapHistoryBuilder& OutHistory, bool bRecursive /*= true*/, bool bFilterForCompilation /*= true*/) const
 {
 	if (bRecursive)
 	{
-		OutHistory.VisitInputPins(this);
+		OutHistory.VisitInputPins(this, bFilterForCompilation);
 	}
 
 	if (!IsNodeEnabled() && OutHistory.GetIgnoreDisabled())
diff --git a/Engine/Plugins/FX/Niagara/Source/NiagaraEditor/Private/NiagaraParameterMapHistory.cpp b/Engine/Plugins/FX/Niagara/Source/NiagaraEditor/Private/NiagaraParameterMapHistory.cpp
--- a/Engine/Plugins/FX/Niagara/Source/NiagaraEditor/Private/NiagaraParameterMapHistory.cpp
+++ b/Engine/Plugins/FX/Niagara/Source/NiagaraEditor/Private/NiagaraParameterMapHistory.cpp
@@ -19,6 +19,7 @@
 #include "ViewModels/Stack/NiagaraStackGraphUtilities.h"
 #include "NiagaraParameterCollection.h"
 #include "NiagaraConstants.h"
+#include "NiagaraNodeStaticSwitch.h"
 
 #define LOCTEXT_NAMESPACE "NiagaraEditor"
 
@@ -1033,7 +1034,7 @@ const FString* FNiagaraParameterMapHistoryBuilder::GetEmitterAlias() const
 	return AliasMap.Find(TEXT("Emitter"));
 }
 
-void FNiagaraParameterMapHistoryBuilder::VisitInputPin(const UEdGraphPin* Pin, const class UNiagaraNode* InNode)
+void FNiagaraParameterMapHistoryBuilder::VisitInputPin(const UEdGraphPin* Pin, const class UNiagaraNode* InNode, bool bFilterForCompilation)
 {
 	const UEdGraphSchema_Niagara* Schema = GetDefault<UEdGraphSchema_Niagara>();
 
@@ -1041,14 +1042,30 @@ void FNiagaraParameterMapHistoryBuilder::VisitInputPin(const UEdGraphPin* Pin, c
 	{
 		for (int32 j = 0; j < Pin->LinkedTo.Num(); j++)
 		{
-			const UEdGraphPin* OutputPin = UNiagaraNode::TraceOutputPin(Pin->LinkedTo[j]);
+			UEdGraphPin* LinkedPin = Pin->LinkedTo[j];
+			const UEdGraphPin* OutputPin;
+			if (!bFilterForCompilation && LinkedPin && Cast<UNiagaraNodeStaticSwitch>(LinkedPin->GetOwningNode()))
+			{
+				// if we do not filter the parameter maps we want all nodes connected to the static switch, regardless of the switch value
+				OutputPin = LinkedPin;
+			}
+			else
+			{
+				OutputPin = UNiagaraNode::TraceOutputPin(LinkedPin, bFilterForCompilation);
+				if (bFilterForCompilation && OutputPin && Cast<UNiagaraNodeStaticSwitch>(OutputPin->GetOwningNode()))
+				{
+					// this means there is no further node connected to this pin of the static switch, so it is not relevant
+					continue;
+				}
+			}
+
 			if (OutputPin)
 			{
 				UNiagaraNode* Node = CastChecked<UNiagaraNode>(OutputPin->GetOwningNode());
 
 				if (!GetNodePreviouslyVisited(Node))
 				{
-					Node->BuildParameterMapHistory(*this, true);
+					Node->BuildParameterMapHistory(*this, true, bFilterForCompilation);
 					RegisterNodeVisitation(Node);
 				}
 
@@ -1062,14 +1079,14 @@ void FNiagaraParameterMapHistoryBuilder::VisitInputPin(const UEdGraphPin* Pin, c
 	}
 }
 
-void FNiagaraParameterMapHistoryBuilder::VisitInputPins(const class UNiagaraNode* InNode)
+void FNiagaraParameterMapHistoryBuilder::VisitInputPins(const class UNiagaraNode* InNode, bool bFilterForCompilation)
 {
 	TArray<UEdGraphPin*> InputPins;
 	InNode->GetInputPins(InputPins);
 
 	for (int32 i = 0; i < InputPins.Num(); i++)
 	{
-		VisitInputPin(InputPins[i], InNode);
+		VisitInputPin(InputPins[i], InNode, bFilterForCompilation);
 	}
 }
 
@@ -1119,7 +1136,7 @@ int32 FNiagaraParameterMapHistoryBuilder::HandleVariableWrite(int32 ParameterMap
 	return Histories[ParameterMapIndex].AddVariable(ResolvedVar, Var, nullptr);
 }
 
-int32 FNiagaraParameterMapHistoryBuilder::HandleVariableRead(int32 ParamMapIdx, const UEdGraphPin* InPin, bool RegisterReadsAsVariables, const UEdGraphPin* InDefaultPin, bool& OutUsedDefault)
+int32 FNiagaraParameterMapHistoryBuilder::HandleVariableRead(int32 ParamMapIdx, const UEdGraphPin* InPin, bool RegisterReadsAsVariables, const UEdGraphPin* InDefaultPin, bool bFilterForCompilation, bool& OutUsedDefault)
 {
 	OutUsedDefault = false;
 	const UEdGraphSchema_Niagara* Schema = GetDefault<UEdGraphSchema_Niagara>();
@@ -1153,7 +1170,7 @@ int32 FNiagaraParameterMapHistoryBuilder::HandleVariableRead(int32 ParamMapIdx,
 			if (InDefaultPin)
 			{
 				OutUsedDefault = true;
-				VisitInputPin(InDefaultPin, Cast<UNiagaraNode>(InDefaultPin->GetOwningNode()));
+				VisitInputPin(InDefaultPin, Cast<UNiagaraNode>(InDefaultPin->GetOwningNode()), bFilterForCompilation);
 				FoundIdx = Histories[ParamMapIdx].FindVariable(Var.GetName(), Var.GetType());
 			}
 
diff --git a/Engine/Plugins/FX/Niagara/Source/NiagaraEditor/Private/NiagaraScriptSource.cpp b/Engine/Plugins/FX/Niagara/Source/NiagaraEditor/Private/NiagaraScriptSource.cpp
--- a/Engine/Plugins/FX/Niagara/Source/NiagaraEditor/Private/NiagaraScriptSource.cpp
+++ b/Engine/Plugins/FX/Niagara/Source/NiagaraEditor/Private/NiagaraScriptSource.cpp
@@ -193,7 +193,8 @@ void InitializeNewRapidIterationParametersForNode(const UEdGraphSchema_Niagara*
 	if (FunctionCallNode != nullptr)
 	{
 		TArray<const UEdGraphPin*> FunctionInputPins;
-		FNiagaraStackGraphUtilities::GetStackFunctionInputPins(*FunctionCallNode, FunctionInputPins, FNiagaraStackGraphUtilities::ENiagaraGetStackFunctionInputPinsOptions::ModuleInputsOnly, false);
+		TSet<const UEdGraphPin*> HiddenInputPins;
+		FNiagaraStackGraphUtilities::GetStackFunctionInputPins(*FunctionCallNode, FunctionInputPins, HiddenInputPins, FNiagaraStackGraphUtilities::ENiagaraGetStackFunctionInputPinsOptions::ModuleInputsOnly, false);
 		for (const UEdGraphPin* FunctionInputPin : FunctionInputPins)
 		{
 			FNiagaraTypeDefinition InputType = Schema->PinToTypeDefinition(FunctionInputPin);
diff --git a/Engine/Plugins/FX/Niagara/Source/NiagaraEditor/Private/NiagaraStackFunctionInputBinder.cpp b/Engine/Plugins/FX/Niagara/Source/NiagaraEditor/Private/NiagaraStackFunctionInputBinder.cpp
--- a/Engine/Plugins/FX/Niagara/Source/NiagaraEditor/Private/NiagaraStackFunctionInputBinder.cpp
+++ b/Engine/Plugins/FX/Niagara/Source/NiagaraEditor/Private/NiagaraStackFunctionInputBinder.cpp
@@ -91,7 +91,8 @@ bool FNiagaraStackFunctionInputBinder::TryBindInternal(
 	FunctionCallNode = InFunctionCallNode;
 
 	TArray<const UEdGraphPin*> InputPins;
-	FNiagaraStackGraphUtilities::GetStackFunctionInputPins(*FunctionCallNode, InputPins, FNiagaraStackGraphUtilities::ENiagaraGetStackFunctionInputPinsOptions::ModuleInputsOnly);
+	TSet<const UEdGraphPin*> HiddenPins;
+	FNiagaraStackGraphUtilities::GetStackFunctionInputPins(*FunctionCallNode, InputPins, HiddenPins, FNiagaraStackGraphUtilities::ENiagaraGetStackFunctionInputPinsOptions::ModuleInputsOnly);
 	
 	bool bInputFound = false;
 	const UEdGraphSchema_Niagara* Schema = GetDefault<UEdGraphSchema_Niagara>();
diff --git a/Engine/Plugins/FX/Niagara/Source/NiagaraEditor/Private/ViewModels/Stack/NiagaraStackFunctionInput.cpp b/Engine/Plugins/FX/Niagara/Source/NiagaraEditor/Private/ViewModels/Stack/NiagaraStackFunctionInput.cpp
--- a/Engine/Plugins/FX/Niagara/Source/NiagaraEditor/Private/ViewModels/Stack/NiagaraStackFunctionInput.cpp
+++ b/Engine/Plugins/FX/Niagara/Source/NiagaraEditor/Private/ViewModels/Stack/NiagaraStackFunctionInput.cpp
@@ -1595,7 +1595,7 @@ void UNiagaraStackFunctionInput::ReassignDynamicInputScript(UNiagaraScript* Dyna
 
 bool UNiagaraStackFunctionInput::GetShouldPassFilterForVisibleCondition() const
 {
-	return GetHasVisibleCondition() == false || GetVisibleConditionEnabled();
+	return bIsVisible && (GetHasVisibleCondition() == false || GetVisibleConditionEnabled());
 }
 
 void UNiagaraStackFunctionInput::GetSearchItems(TArray<FStackSearchItem>& SearchItems) const
diff --git a/Engine/Plugins/FX/Niagara/Source/NiagaraEditor/Private/ViewModels/Stack/NiagaraStackFunctionInputCollection.cpp b/Engine/Plugins/FX/Niagara/Source/NiagaraEditor/Private/ViewModels/Stack/NiagaraStackFunctionInputCollection.cpp
--- a/Engine/Plugins/FX/Niagara/Source/NiagaraEditor/Private/ViewModels/Stack/NiagaraStackFunctionInputCollection.cpp
+++ b/Engine/Plugins/FX/Niagara/Source/NiagaraEditor/Private/ViewModels/Stack/NiagaraStackFunctionInputCollection.cpp
@@ -15,6 +15,7 @@
 
 #include "EdGraph/EdGraphPin.h"
 #include "ScopedTransaction.h"
+#include "NiagaraEditorUtilities.h"
 
 #define LOCTEXT_NAMESPACE "UNiagaraStackFunctionInputCollection"
 
@@ -77,8 +78,9 @@ void UNiagaraStackFunctionInputCollection::SetShouldShowInStack(bool bInShouldSh
 
 void UNiagaraStackFunctionInputCollection::RefreshChildrenInternal(const TArray<UNiagaraStackEntry*>& CurrentChildren, TArray<UNiagaraStackEntry*>& NewChildren, TArray<FStackIssue>& NewIssues)
 {
+	TSet<const UEdGraphPin*> HiddenPins;
 	TArray<const UEdGraphPin*> InputPins;
-	FNiagaraStackGraphUtilities::GetStackFunctionInputPins(*InputFunctionCallNode, InputPins, FNiagaraStackGraphUtilities::ENiagaraGetStackFunctionInputPinsOptions::ModuleInputsOnly);
+	FNiagaraStackGraphUtilities::GetStackFunctionInputPins(*InputFunctionCallNode, InputPins, HiddenPins, FNiagaraStackGraphUtilities::ENiagaraGetStackFunctionInputPinsOptions::ModuleInputsOnly);
 
 	const UEdGraphSchema_Niagara* NiagaraSchema = GetDefault<UEdGraphSchema_Niagara>();
 
@@ -125,13 +127,15 @@ void UNiagaraStackFunctionInputCollection::RefreshChildrenInternal(const TArray<
 			? InputMetaData->CategoryName
 			: UncategorizedName;
 
-		FInputData InputData = { InputPin, InputVariable.GetType(), InputMetaData ? InputMetaData->EditorSortPriority : 0, InputCategory, false };
+		bool IsVisible = !HiddenPins.Contains(InputPin);
+		FInputData InputData = { InputPin, InputVariable.GetType(), InputMetaData ? InputMetaData->EditorSortPriority : 0, InputCategory, false, IsVisible };
 		InputDataCollection.Add(InputData);
 	}
 
 	// Gather static switch parameters
+	TSet<UEdGraphPin*> HiddenSwitchPins;
 	TArray<UEdGraphPin*> SwitchPins;
-	FNiagaraStackGraphUtilities::GetStackFunctionStaticSwitchPins(*InputFunctionCallNode, SwitchPins);
+	FNiagaraStackGraphUtilities::GetStackFunctionStaticSwitchPins(*InputFunctionCallNode, SwitchPins, HiddenSwitchPins);
 	for (UEdGraphPin* InputPin : SwitchPins)
 	{
 		// The static switch pin names to not contain the module namespace, as they are not part of the parameter maps.
@@ -167,7 +171,8 @@ void UNiagaraStackFunctionInputCollection::RefreshChildrenInternal(const TArray<
 			? InputMetaData->CategoryName
 			: UncategorizedName;
 
-		FInputData InputData = { InputPin, InputVariable.GetType(), InputMetaData ? InputMetaData->EditorSortPriority : 0, InputCategory, true };
+		bool IsVisible = !HiddenSwitchPins.Contains(InputPin);
+		FInputData InputData = { InputPin, InputVariable.GetType(), InputMetaData ? InputMetaData->EditorSortPriority : 0, InputCategory, true, IsVisible };
 		InputDataCollection.Add(InputData);
 	}
 
@@ -222,7 +227,7 @@ void UNiagaraStackFunctionInputCollection::RefreshChildrenInternal(const TArray<
 			}
 			NewChildren.Add(InputCategory);
 		}
-		InputCategory->AddInput(InputData.Pin->PinName, InputData.Type, InputData.bIsStatic ? EStackParameterBehavior::Static : EStackParameterBehavior::Dynamic);
+		InputCategory->AddInput(InputData.Pin->PinName, InputData.Type, InputData.bIsStatic ? EStackParameterBehavior::Static : EStackParameterBehavior::Dynamic, InputData.bIsVisible);
 	}
 	RefreshIssues(DuplicateInputNames, ValidAliasedInputNames, PinsWithInvalidTypes, StaticSwitchInputs, NewIssues);
 }
diff --git a/Engine/Plugins/FX/Niagara/Source/NiagaraEditor/Private/ViewModels/Stack/NiagaraStackGraphUtilities.cpp b/Engine/Plugins/FX/Niagara/Source/NiagaraEditor/Private/ViewModels/Stack/NiagaraStackGraphUtilities.cpp
--- a/Engine/Plugins/FX/Niagara/Source/NiagaraEditor/Private/ViewModels/Stack/NiagaraStackGraphUtilities.cpp
+++ b/Engine/Plugins/FX/Niagara/Source/NiagaraEditor/Private/ViewModels/Stack/NiagaraStackGraphUtilities.cpp
@@ -497,43 +497,62 @@ FString FNiagaraStackGraphUtilities::GenerateStackModuleEditorDataKey(UNiagaraNo
 	return ModuleNode.GetFunctionName();
 }
 
-void FNiagaraStackGraphUtilities::GetStackFunctionInputPins(UNiagaraNodeFunctionCall& FunctionCallNode, TArray<const UEdGraphPin*>& OutInputPins, ENiagaraGetStackFunctionInputPinsOptions Options, bool bIgnoreDisabled)
+void ExtractInputPinsFromHistory(FNiagaraParameterMapHistory& History, UEdGraph* FunctionGraph, FNiagaraStackGraphUtilities::ENiagaraGetStackFunctionInputPinsOptions Options, TArray<const UEdGraphPin*>& OutPins)
+{
+	for (int32 i = 0; i < History.Variables.Num(); i++)
+	{
+		FNiagaraVariable& Variable = History.Variables[i];
+		TArray<TTuple<const UEdGraphPin*, const UEdGraphPin*>>& ReadHistory = History.PerVariableReadHistory[i];
+
+		// A read is only really exposed if it's the first read and it has no corresponding write.
+		if (ReadHistory.Num() > 0 && ReadHistory[0].Get<1>() == nullptr)
+		{
+			const UEdGraphPin* InputPin = ReadHistory[0].Get<0>();
+
+			// Make sure that the module input is from the called graph, and not a nested graph.
+			if (InputPin->GetOwningNode()->GetGraph() == FunctionGraph &&
+				(Options == FNiagaraStackGraphUtilities::ENiagaraGetStackFunctionInputPinsOptions::AllInputs || FNiagaraParameterHandle(InputPin->PinName).IsModuleHandle()))
+			{
+				OutPins.Add(InputPin);
+			}
+		}
+	}
+}
+
+void FNiagaraStackGraphUtilities::GetStackFunctionInputPins(UNiagaraNodeFunctionCall& FunctionCallNode, TArray<const UEdGraphPin*>& OutInputPins, TSet<const UEdGraphPin*>& OutHiddenPins, ENiagaraGetStackFunctionInputPinsOptions Options, bool bIgnoreDisabled)
 {
 	FNiagaraParameterMapHistoryBuilder Builder;
 	Builder.SetIgnoreDisabled(bIgnoreDisabled);
-	FunctionCallNode.BuildParameterMapHistory(Builder, false);
+	FunctionCallNode.BuildParameterMapHistory(Builder, false, false);
 	
 	if (Builder.Histories.Num() == 1)
 	{
-		for (int32 i = 0; i < Builder.Histories[0].Variables.Num(); i++)
-		{
-			FNiagaraVariable& Variable = Builder.Histories[0].Variables[i];
-			TArray<TTuple<const UEdGraphPin*, const UEdGraphPin*>>& ReadHistory = Builder.Histories[0].PerVariableReadHistory[i];
+		ExtractInputPinsFromHistory(Builder.Histories[0], FunctionCallNode.GetCalledGraph(), Options, OutInputPins);
 
-			// A read is only really exposed if it's the first read and it has no corresponding write.
-			if (ReadHistory.Num() > 0 && ReadHistory[0].Get<1>() == nullptr)
-			{
-				const UEdGraphPin* InputPin = ReadHistory[0].Get<0>();
+		FNiagaraParameterMapHistoryBuilder BuilderCompiled;
+		FunctionCallNode.BuildParameterMapHistory(BuilderCompiled, false, true);
+		TArray<const UEdGraphPin*> CompilationPins;
+		ExtractInputPinsFromHistory(BuilderCompiled.Histories[0], FunctionCallNode.GetCalledGraph(), Options, CompilationPins);
 
-				// Make sure that the module input is from the called graph, and not a nested graph.
-				if (InputPin->GetOwningNode()->GetGraph() == FunctionCallNode.GetCalledGraph() &&
-					(Options == ENiagaraGetStackFunctionInputPinsOptions::AllInputs || FNiagaraParameterHandle(InputPin->PinName).IsModuleHandle()))
-				{
-					OutInputPins.Add(InputPin);
-				}
+		for (const UEdGraphPin* Pin : OutInputPins)
+		{
+			if (!CompilationPins.Contains(Pin))
+			{
+				OutHiddenPins.Add(Pin);
 			}
 		}
 	}
 }
 
-void FNiagaraStackGraphUtilities::GetStackFunctionStaticSwitchPins(UNiagaraNodeFunctionCall& FunctionCallNode, TArray<UEdGraphPin*>& OutInputPins)
+void FNiagaraStackGraphUtilities::GetStackFunctionStaticSwitchPins(UNiagaraNodeFunctionCall& FunctionCallNode, TArray<UEdGraphPin*>& OutInputPins, TSet<UEdGraphPin*>& OutHiddenPins)
 {
 	const UEdGraphSchema_Niagara* Schema = CastChecked<UEdGraphSchema_Niagara>(FunctionCallNode.GetSchema());
 	UNiagaraGraph* FunctionCallGraph = FunctionCallNode.GetCalledGraph();
 	if (FunctionCallGraph == nullptr)
 		return;
 
-	for (FNiagaraVariable SwitchInput : FunctionCallGraph->FindStaticSwitchInputs())
+	TArray<FNiagaraVariable> ReachableInputs = FunctionCallGraph->FindStaticSwitchInputs(true);
+	for (FNiagaraVariable SwitchInput : FunctionCallGraph->FindStaticSwitchInputs(false))
 	{
 		FEdGraphPinType PinType = Schema->TypeDefinitionToPinType(SwitchInput.GetType());
 		for (UEdGraphPin* Pin : FunctionCallNode.Pins)
@@ -545,6 +564,10 @@ void FNiagaraStackGraphUtilities::GetStackFunctionStaticSwitchPins(UNiagaraNodeF
 			if (Pin->PinName.IsEqual(SwitchInput.GetName()) && Pin->PinType == PinType)
 			{
 				OutInputPins.Add(Pin);
+				if (!ReachableInputs.Contains(SwitchInput))
+				{
+					OutHiddenPins.Add(Pin);
+				}
 				break;
 			}
 		}
@@ -1514,7 +1537,8 @@ bool TryGetStackFunctionInputValue(UNiagaraScript& OwningScript, const UEdGraphP
 			UNiagaraNodeFunctionCall* DynamicInputFunctionCall = CastChecked<UNiagaraNodeFunctionCall>(InputPin.LinkedTo[0]->GetOwningNode());
 			OutStackFunctionInputValue.DynamicValue = DynamicInputFunctionCall;
 			TArray<const UEdGraphPin*> DynamicValueInputPins;
-			FNiagaraStackGraphUtilities::GetStackFunctionInputPins(*DynamicInputFunctionCall,DynamicValueInputPins,
+			TSet<const UEdGraphPin*> HiddenInputPins;
+			FNiagaraStackGraphUtilities::GetStackFunctionInputPins(*DynamicInputFunctionCall, DynamicValueInputPins, HiddenInputPins,
 				FNiagaraStackGraphUtilities::ENiagaraGetStackFunctionInputPinsOptions::ModuleInputsOnly);
 
 			FRapidIterationParameterContext InputRapidIterationParameterContext = RapidIterationParameterContext.IsValid() 
diff --git a/Engine/Plugins/FX/Niagara/Source/NiagaraEditor/Private/ViewModels/Stack/NiagaraStackInputCategory.cpp b/Engine/Plugins/FX/Niagara/Source/NiagaraEditor/Private/ViewModels/Stack/NiagaraStackInputCategory.cpp
--- a/Engine/Plugins/FX/Niagara/Source/NiagaraEditor/Private/ViewModels/Stack/NiagaraStackInputCategory.cpp
+++ b/Engine/Plugins/FX/Niagara/Source/NiagaraEditor/Private/ViewModels/Stack/NiagaraStackInputCategory.cpp
@@ -33,9 +33,9 @@ void UNiagaraStackInputCategory::ResetInputs()
 	Inputs.Empty();
 }
 
-void UNiagaraStackInputCategory::AddInput(FName InInputParameterHandle, FNiagaraTypeDefinition InInputType, EStackParameterBehavior InParameterBehavior)
+void UNiagaraStackInputCategory::AddInput(FName InInputParameterHandle, FNiagaraTypeDefinition InInputType, EStackParameterBehavior InParameterBehavior, bool bIsVisible)
 {
-	Inputs.Add({ InInputParameterHandle, InInputType, InParameterBehavior });
+	Inputs.Add({ InInputParameterHandle, InInputType, InParameterBehavior, bIsVisible });
 }
 
 void UNiagaraStackInputCategory::RefreshChildrenInternal(const TArray<UNiagaraStackEntry*>& CurrentChildren, TArray<UNiagaraStackEntry*>& NewChildren, TArray<FStackIssue>& NewIssues)
@@ -51,7 +51,7 @@ void UNiagaraStackInputCategory::RefreshChildrenInternal(const TArray<UNiagaraSt
 			InputChild->Initialize(CreateDefaultChildRequiredData(), *ModuleNode, *InputFunctionCallNode,
 				Input.ParameterHandle, Input.Type, Input.ParameterBehavior, GetOwnerStackItemEditorDataKey());
 		}
-
+		InputChild->bIsVisible = Input.bIsVisible;
 		NewChildren.Add(InputChild);
 	}
 }
@@ -63,7 +63,16 @@ FText UNiagaraStackInputCategory::GetDisplayName() const
 
 bool UNiagaraStackInputCategory::GetShouldShowInStack() const
 {
-	return bShouldShowInStack;
+	bool bIsVisible = false;
+	for (FInputParameterHandleAndType Input : Inputs)
+	{
+		if (Input.bIsVisible)
+		{
+			bIsVisible = true;
+			break;
+		}
+	}
+	return bIsVisible && bShouldShowInStack;
 }
 
 UNiagaraStackEntry::EStackRowStyle UNiagaraStackInputCategory::GetStackRowStyle() const
diff --git a/Engine/Plugins/FX/Niagara/Source/NiagaraEditor/Private/ViewModels/Stack/NiagaraStackScriptItemGroup.cpp b/Engine/Plugins/FX/Niagara/Source/NiagaraEditor/Private/ViewModels/Stack/NiagaraStackScriptItemGroup.cpp
--- a/Engine/Plugins/FX/Niagara/Source/NiagaraEditor/Private/ViewModels/Stack/NiagaraStackScriptItemGroup.cpp
+++ b/Engine/Plugins/FX/Niagara/Source/NiagaraEditor/Private/ViewModels/Stack/NiagaraStackScriptItemGroup.cpp
@@ -233,7 +233,8 @@ private:
 		UNiagaraNodeAssignment* NewAssignmentModule = FNiagaraStackGraphUtilities::AddParameterModuleToStack(Vars, *OutputNode, TargetIndex,DefaultVals );
 		
 		TArray<const UEdGraphPin*> InputPins;
-		FNiagaraStackGraphUtilities::GetStackFunctionInputPins(*NewAssignmentModule, InputPins);
+		TSet<const UEdGraphPin*> HiddenPins;
+		FNiagaraStackGraphUtilities::GetStackFunctionInputPins(*NewAssignmentModule, InputPins, HiddenPins);
 		if (InputPins.Num() == 1)
 		{
 			FString FunctionInputEditorDataKey = FNiagaraStackGraphUtilities::GenerateStackFunctionInputEditorDataKey(*NewAssignmentModule, InputPins[0]->PinName);
diff --git a/Engine/Plugins/FX/Niagara/Source/NiagaraEditor/Private/Widgets/SNiagaraParameterMapView.cpp b/Engine/Plugins/FX/Niagara/Source/NiagaraEditor/Private/Widgets/SNiagaraParameterMapView.cpp
--- a/Engine/Plugins/FX/Niagara/Source/NiagaraEditor/Private/Widgets/SNiagaraParameterMapView.cpp
+++ b/Engine/Plugins/FX/Niagara/Source/NiagaraEditor/Private/Widgets/SNiagaraParameterMapView.cpp
@@ -331,6 +331,10 @@ void SNiagaraParameterMapView::CollectAllActions(FGraphActionListBuilderBase& Ou
 	TMap<FNiagaraVariable, TArray<FNiagaraGraphParameterReferenceCollection>> ParameterEntries;
 	for (auto& GraphWeakPtr : Graphs)
 	{
+		if (!GraphWeakPtr.IsValid())
+		{
+			continue;
+		}
 		UNiagaraGraph* Graph = GraphWeakPtr.Get();
 		for (const auto& ParameterElement : FNiagaraEditorUtilities::GetCompiledGraphParameterMapReferences(Graph))
 		{
diff --git a/Engine/Plugins/FX/Niagara/Source/NiagaraEditor/Public/NiagaraGraph.h b/Engine/Plugins/FX/Niagara/Source/NiagaraEditor/Public/NiagaraGraph.h
--- a/Engine/Plugins/FX/Niagara/Source/NiagaraEditor/Public/NiagaraGraph.h
+++ b/Engine/Plugins/FX/Niagara/Source/NiagaraEditor/Public/NiagaraGraph.h
@@ -154,7 +154,7 @@ class UNiagaraGraph : public UEdGraph
 	void FindInputNodes(TArray<class UNiagaraNodeInput*>& OutInputNodes, FFindInputNodeOptions Options = FFindInputNodeOptions()) const;
 
 	/** Returns a list of variable inputs for all static switch nodes in the graph. */
-	TArray<FNiagaraVariable> FindStaticSwitchInputs() const;
+	TArray<FNiagaraVariable> FindStaticSwitchInputs(bool bReachableOnly = false) const;
 
 	/** Get an in-order traversal of a graph by the specified target output script usage.*/
 	void BuildTraversal(TArray<class UNiagaraNode*>& OutNodesTraversed, ENiagaraScriptUsage TargetUsage, FGuid TargetUsageId) const;
@@ -287,6 +287,9 @@ private:
 
 	/** A delegate that broadcasts a notification whenever the graph needs recompile due to structural change. */
 	FOnGraphChanged OnGraphNeedsRecompile;
+
+	/** Find all nodes in the graph that can be reached during compilation. */
+	TArray<UEdGraphNode*> FindReachbleNodes() const;
 	
 	/** The current change identifier for this graph overall. Used to sync status with UNiagaraScripts.*/
 	UPROPERTY()
diff --git a/Engine/Plugins/FX/Niagara/Source/NiagaraEditor/Public/NiagaraNode.h b/Engine/Plugins/FX/Niagara/Source/NiagaraEditor/Public/NiagaraNode.h
--- a/Engine/Plugins/FX/Niagara/Source/NiagaraEditor/Public/NiagaraNode.h
+++ b/Engine/Plugins/FX/Niagara/Source/NiagaraEditor/Public/NiagaraNode.h
@@ -96,14 +96,14 @@ public:
 	/** Notify the rename was cancelled.*/
 	virtual bool CancelEditablePinName(const FText& InName, UEdGraphPin* InGraphPinObj) { return false; }
 
-	virtual void BuildParameterMapHistory(FNiagaraParameterMapHistoryBuilder& OutHistory, bool bRecursive = true) const;
+	virtual void BuildParameterMapHistory(FNiagaraParameterMapHistoryBuilder& OutHistory, bool bRecursive = true, bool bFilterForCompilation = true) const;
 	
 	/** Go through all the external dependencies of this node in isolation and add them to the reference id list.*/
 	virtual void GatherExternalDependencyIDs(ENiagaraScriptUsage InMasterUsage, const FGuid& InMasterUsageId, TArray<FNiagaraCompileHash>& InReferencedCompileHashes, TArray<FGuid>& InReferencedIDs, TArray<UObject*>& InReferencedObjs) const {};
 
 	/** Traces one of this node's output pins to its source output pin if it is a reroute node output pin.*/
 	virtual UEdGraphPin* GetTracedOutputPin(UEdGraphPin* LocallyOwnedOutputPin) const {return LocallyOwnedOutputPin;}
-	static UEdGraphPin* TraceOutputPin(UEdGraphPin* LocallyOwnedOutputPin);
+	static UEdGraphPin* TraceOutputPin(UEdGraphPin* LocallyOwnedOutputPin, bool bFilterForCompilation = true);
 
 	/** Allows a node to replace a pin that is about to be compiled with another pin. This can be used for either optimizations or features such as the static switch. Returns true if the pin was successfully replaced, false otherwise. */
 	virtual bool SubstituteCompiledPin(FHlslNiagaraTranslator* Translator, UEdGraphPin** LocallyOwnedPin);
diff --git a/Engine/Plugins/FX/Niagara/Source/NiagaraEditor/Public/NiagaraNodeAssignment.h b/Engine/Plugins/FX/Niagara/Source/NiagaraEditor/Public/NiagaraNodeAssignment.h
--- a/Engine/Plugins/FX/Niagara/Source/NiagaraEditor/Public/NiagaraNodeAssignment.h
+++ b/Engine/Plugins/FX/Niagara/Source/NiagaraEditor/Public/NiagaraNodeAssignment.h
@@ -44,7 +44,7 @@ public:
 
 	//~ UNiagaraNodeFunctionCall interface
 	virtual bool RefreshFromExternalChanges() override;
-	virtual void BuildParameterMapHistory(FNiagaraParameterMapHistoryBuilder& OutHistory, bool bRecursive = true) const override;
+	virtual void BuildParameterMapHistory(FNiagaraParameterMapHistoryBuilder& OutHistory, bool bRecursive = true, bool bFilterForCompilation = true) const override;
 	virtual void GatherExternalDependencyIDs(ENiagaraScriptUsage InMasterUsage, const FGuid& InMasterUsageId, TArray<FNiagaraCompileHash>& InReferencedCompileHashes, TArray<FGuid>& InReferencedIDs, TArray<UObject*>& InReferencedObjs) const override;
 
 	void AddParameter(FNiagaraVariable InVar, FString InDefaultValue);
diff --git a/Engine/Plugins/FX/Niagara/Source/NiagaraEditor/Public/NiagaraNodeCustomHlsl.h b/Engine/Plugins/FX/Niagara/Source/NiagaraEditor/Public/NiagaraNodeCustomHlsl.h
--- a/Engine/Plugins/FX/Niagara/Source/NiagaraEditor/Public/NiagaraNodeCustomHlsl.h
+++ b/Engine/Plugins/FX/Niagara/Source/NiagaraEditor/Public/NiagaraNodeCustomHlsl.h
@@ -32,7 +32,7 @@ public:
 
 	bool GetTokens(TArray<FString>& OutTokens) const;
 
-	virtual void BuildParameterMapHistory(FNiagaraParameterMapHistoryBuilder& OutHistory, bool bRecursive = true) const override;
+	virtual void BuildParameterMapHistory(FNiagaraParameterMapHistoryBuilder& OutHistory, bool bRecursive = true, bool bFilterForCompilation = true) const override;
 
 	// Replace items in the tokens array if they start with the src string or optionally src string and a namespace delimiter
 	static void ReplaceExactMatchTokens(TArray<FString>& Tokens, const FString& SrcString, const FString& ReplaceString, bool bAllowNamespaceSeparation);
diff --git a/Engine/Plugins/FX/Niagara/Source/NiagaraEditor/Public/NiagaraNodeFunctionCall.h b/Engine/Plugins/FX/Niagara/Source/NiagaraEditor/Public/NiagaraNodeFunctionCall.h
--- a/Engine/Plugins/FX/Niagara/Source/NiagaraEditor/Public/NiagaraNodeFunctionCall.h
+++ b/Engine/Plugins/FX/Niagara/Source/NiagaraEditor/Public/NiagaraNodeFunctionCall.h
@@ -94,7 +94,7 @@ public:
 
 	bool FindAutoBoundInput(UNiagaraNodeInput* InputNode, UEdGraphPin* PinToAutoBind, FNiagaraVariable& OutFoundVar, ENiagaraInputNodeUsage& OutNodeUsage);
 
-	virtual void BuildParameterMapHistory(FNiagaraParameterMapHistoryBuilder& OutHistory, bool bRecursive = true) const override;
+	void BuildParameterMapHistory(FNiagaraParameterMapHistoryBuilder& OutHistory, bool bRecursive = true, bool bFilterForCompilation = true) const;
 
 	FString GetFunctionName() const { return FunctionDisplayName; }
 	UNiagaraGraph* GetCalledGraph() const;
diff --git a/Engine/Plugins/FX/Niagara/Source/NiagaraEditor/Public/NiagaraNodeInput.h b/Engine/Plugins/FX/Niagara/Source/NiagaraEditor/Public/NiagaraNodeInput.h
--- a/Engine/Plugins/FX/Niagara/Source/NiagaraEditor/Public/NiagaraNodeInput.h
+++ b/Engine/Plugins/FX/Niagara/Source/NiagaraEditor/Public/NiagaraNodeInput.h
@@ -92,7 +92,7 @@ public:
 
 	bool ReferencesSameInput(UNiagaraNodeInput* Other) const;
 
-	virtual void BuildParameterMapHistory(FNiagaraParameterMapHistoryBuilder& OutHistory, bool bRecursive = true) const override;
+	virtual void BuildParameterMapHistory(FNiagaraParameterMapHistoryBuilder& OutHistory, bool bRecursive = true, bool bFilterForCompilation = true) const override;
 
 	static const FLinearColor TitleColor_Attribute;
 	static const FLinearColor TitleColor_Constant;
diff --git a/Engine/Plugins/FX/Niagara/Source/NiagaraEditor/Public/NiagaraNodeReadDataSet.h b/Engine/Plugins/FX/Niagara/Source/NiagaraEditor/Public/NiagaraNodeReadDataSet.h
--- a/Engine/Plugins/FX/Niagara/Source/NiagaraEditor/Public/NiagaraNodeReadDataSet.h
+++ b/Engine/Plugins/FX/Niagara/Source/NiagaraEditor/Public/NiagaraNodeReadDataSet.h
@@ -27,6 +27,6 @@ public:
 	//~End UNiagaraGraph Interface
 	virtual void Compile(class FHlslNiagaraTranslator* Translator, TArray<int32>& Outputs)override;
 
-	virtual void BuildParameterMapHistory(FNiagaraParameterMapHistoryBuilder& OutHistory, bool bRecursive = true) const override;
+	virtual void BuildParameterMapHistory(FNiagaraParameterMapHistoryBuilder& OutHistory, bool bRecursive = true, bool bFilterForCompilation = true) const override;
 };
 
diff --git a/Engine/Plugins/FX/Niagara/Source/NiagaraEditor/Public/NiagaraNodeWriteDataSet.h b/Engine/Plugins/FX/Niagara/Source/NiagaraEditor/Public/NiagaraNodeWriteDataSet.h
--- a/Engine/Plugins/FX/Niagara/Source/NiagaraEditor/Public/NiagaraNodeWriteDataSet.h
+++ b/Engine/Plugins/FX/Niagara/Source/NiagaraEditor/Public/NiagaraNodeWriteDataSet.h
@@ -34,7 +34,7 @@ public:
 
 	virtual void Compile(class FHlslNiagaraTranslator* Translator, TArray<int32>& Outputs)override;
 
-	virtual void BuildParameterMapHistory(FNiagaraParameterMapHistoryBuilder& OutHistory, bool bRecursive = true) const override;
+	virtual void BuildParameterMapHistory(FNiagaraParameterMapHistoryBuilder& OutHistory, bool bRecursive = true, bool bFilterForCompilation = true) const override;
 	virtual TSharedPtr<SGraphNode> CreateVisualWidget() override;
 protected:
 	void AddConditionPin(int32 PinIndex = -1);
diff --git a/Engine/Plugins/FX/Niagara/Source/NiagaraEditor/Public/NiagaraParameterMapHistory.h b/Engine/Plugins/FX/Niagara/Source/NiagaraEditor/Public/NiagaraParameterMapHistory.h
--- a/Engine/Plugins/FX/Niagara/Source/NiagaraEditor/Public/NiagaraParameterMapHistory.h
+++ b/Engine/Plugins/FX/Niagara/Source/NiagaraEditor/Public/NiagaraParameterMapHistory.h
@@ -280,10 +280,10 @@ public:
 	bool GetNodePreviouslyVisited(const class UNiagaraNode* Node) const;
 
 	/** If we haven't already visited the owning nodes, do so.*/
-	void VisitInputPins(const class UNiagaraNode*);
+	void VisitInputPins(const class UNiagaraNode*, bool bFilterForCompilation);
 	
 	/** If we haven't already visited the owning node, do so.*/
-	void VisitInputPin(const UEdGraphPin* Pin, const class UNiagaraNode*);
+	void VisitInputPin(const UEdGraphPin* Pin, const class UNiagaraNode*, bool bFilterForCompilation);
 
 	/**
 	* Record that a pin writes to the parameter map. The pin name is expected to be the namespaced parameter map version of the name. If any aliases are in place, they are removed.
@@ -298,7 +298,7 @@ public:
 	/**
 	* Record that a pin reads from the parameter map. The pin name is expected to be the namespaced parameter map version of the name. If any aliases are in place, they are removed.
 	*/
-	int32 HandleVariableRead(int32 ParameterMapIndex, const UEdGraphPin* InPin, bool RegisterReadsAsVariables, const UEdGraphPin* InDefaultPin, bool& OutUsedDefault);
+	int32 HandleVariableRead(int32 ParameterMapIndex, const UEdGraphPin* InPin, bool RegisterReadsAsVariables, const UEdGraphPin* InDefaultPin, bool bFilterForCompilation, bool& OutUsedDefault);
 
 	int32 HandleExternalVariableRead(int32 ParamMapIdx, const FName& InVarName);
 	
diff --git a/Engine/Plugins/FX/Niagara/Source/NiagaraEditor/Public/ViewModels/Stack/NiagaraStackFunctionInput.h b/Engine/Plugins/FX/Niagara/Source/NiagaraEditor/Public/ViewModels/Stack/NiagaraStackFunctionInput.h
--- a/Engine/Plugins/FX/Niagara/Source/NiagaraEditor/Public/ViewModels/Stack/NiagaraStackFunctionInput.h
+++ b/Engine/Plugins/FX/Niagara/Source/NiagaraEditor/Public/ViewModels/Stack/NiagaraStackFunctionInput.h
@@ -220,6 +220,9 @@ public:
 	//~ UNiagaraStackEntry interface
 	virtual void GetSearchItems(TArray<FStackSearchItem>& SearchItems) const override;
 
+	/** If false then the stack parameter is not visible */
+	bool bIsVisible = true;
+
 protected:
 	//~ UNiagaraStackEntry interface
 	virtual void FinalizeInternal() override;
diff --git a/Engine/Plugins/FX/Niagara/Source/NiagaraEditor/Public/ViewModels/Stack/NiagaraStackFunctionInputCollection.h b/Engine/Plugins/FX/Niagara/Source/NiagaraEditor/Public/ViewModels/Stack/NiagaraStackFunctionInputCollection.h
--- a/Engine/Plugins/FX/Niagara/Source/NiagaraEditor/Public/ViewModels/Stack/NiagaraStackFunctionInputCollection.h
+++ b/Engine/Plugins/FX/Niagara/Source/NiagaraEditor/Public/ViewModels/Stack/NiagaraStackFunctionInputCollection.h
@@ -54,6 +54,7 @@ private:
 		int32 SortKey;
 		FText Category;
 		bool bIsStatic;
+		bool bIsVisible;
 	};
 
 	UNiagaraNodeFunctionCall* ModuleNode;
diff --git a/Engine/Plugins/FX/Niagara/Source/NiagaraEditor/Public/ViewModels/Stack/NiagaraStackGraphUtilities.h b/Engine/Plugins/FX/Niagara/Source/NiagaraEditor/Public/ViewModels/Stack/NiagaraStackGraphUtilities.h
--- a/Engine/Plugins/FX/Niagara/Source/NiagaraEditor/Public/ViewModels/Stack/NiagaraStackGraphUtilities.h
+++ b/Engine/Plugins/FX/Niagara/Source/NiagaraEditor/Public/ViewModels/Stack/NiagaraStackGraphUtilities.h
@@ -74,9 +74,9 @@ namespace FNiagaraStackGraphUtilities
 		ModuleInputsOnly
 	};
 
-	void GetStackFunctionInputPins(UNiagaraNodeFunctionCall& FunctionCallNode, TArray<const UEdGraphPin*>& OutInputPins, ENiagaraGetStackFunctionInputPinsOptions Options = ENiagaraGetStackFunctionInputPinsOptions::AllInputs, bool bIgnoreDisabled = false);
+	void GetStackFunctionInputPins(UNiagaraNodeFunctionCall& FunctionCallNode, TArray<const UEdGraphPin*>& OutInputPins, TSet<const UEdGraphPin*>& OutHiddenPins, ENiagaraGetStackFunctionInputPinsOptions Options = ENiagaraGetStackFunctionInputPinsOptions::AllInputs, bool bIgnoreDisabled = false);
 
-	void GetStackFunctionStaticSwitchPins(UNiagaraNodeFunctionCall& FunctionCallNode, TArray<UEdGraphPin*>& OutInputPins);
+	void GetStackFunctionStaticSwitchPins(UNiagaraNodeFunctionCall& FunctionCallNode, TArray<UEdGraphPin*>& OutInputPins, TSet<UEdGraphPin*>& OutHiddenPins);
 
 	UNiagaraNodeParameterMapSet* GetStackFunctionOverrideNode(UNiagaraNodeFunctionCall& FunctionCallNode);
 
diff --git a/Engine/Plugins/FX/Niagara/Source/NiagaraEditor/Public/ViewModels/Stack/NiagaraStackInputCategory.h b/Engine/Plugins/FX/Niagara/Source/NiagaraEditor/Public/ViewModels/Stack/NiagaraStackInputCategory.h
--- a/Engine/Plugins/FX/Niagara/Source/NiagaraEditor/Public/ViewModels/Stack/NiagaraStackInputCategory.h
+++ b/Engine/Plugins/FX/Niagara/Source/NiagaraEditor/Public/ViewModels/Stack/NiagaraStackInputCategory.h
@@ -32,7 +32,7 @@ public:
 
 	void ResetInputs();
 
-	void AddInput(FName InInputParameterHandle, FNiagaraTypeDefinition InInputType, EStackParameterBehavior InParameterBehavior);
+	void AddInput(FName InInputParameterHandle, FNiagaraTypeDefinition InInputType, EStackParameterBehavior InParameterBehavior, bool bIsVisible);
 
 	//~ UNiagaraStackEntry interface
 	virtual FText GetDisplayName() const override;
@@ -56,6 +56,7 @@ private:
 		FName ParameterHandle;
 		FNiagaraTypeDefinition Type;
 		EStackParameterBehavior ParameterBehavior;
+		bool bIsVisible;
 	};
 
 	UNiagaraNodeFunctionCall* ModuleNode;
