Support for compile-time constants in static switches

diff --git a/Engine/Plugins/FX/Niagara/Source/Niagara/Classes/NiagaraConstants.h b/Engine/Plugins/FX/Niagara/Source/Niagara/Classes/NiagaraConstants.h
--- a/Engine/Plugins/FX/Niagara/Source/Niagara/Classes/NiagaraConstants.h
+++ b/Engine/Plugins/FX/Niagara/Source/Niagara/Classes/NiagaraConstants.h
@@ -109,6 +109,7 @@ struct NIAGARA_API FNiagaraConstants
 	static void Init();
 	static const TArray<FNiagaraVariable>& GetEngineConstants();
 	static const TArray<FNiagaraVariable>& GetTranslatorConstants();
+	static const TArray<FNiagaraVariable>& GetStaticSwitchConstants();
 	static FNiagaraVariable UpdateEngineConstant(const FNiagaraVariable& InVar);
 	static const FNiagaraVariable *FindEngineConstant(const FNiagaraVariable& InVar);
 	static FText GetEngineConstantDescription(const FNiagaraVariable& InVar);
@@ -124,12 +125,14 @@ struct NIAGARA_API FNiagaraConstants
 	static const FNiagaraVariableMetaData* GetConstantMetaData(const FNiagaraVariable& InVar);
 
 	static const FNiagaraVariable* GetKnownConstant(const FName& InName, bool bAllowPartialNameMatch);
+	static const FNiagaraVariable *FindStaticSwitchConstant(const FName& InName);
 
 	static bool IsEngineManagedAttribute(const FNiagaraVariable& Var);
 
 private:
 	static TArray<FNiagaraVariable> SystemParameters;
 	static TArray<FNiagaraVariable> TranslatorParameters;
+	static TArray<FNiagaraVariable> SwitchParameters;
 	static TMap<FName, FNiagaraVariable> UpdatedSystemParameters;
 	static TMap<FNiagaraVariable, FText> SystemStrMap;
 	static TArray<FNiagaraVariable> Attributes;
diff --git a/Engine/Plugins/FX/Niagara/Source/Niagara/Private/NiagaraConstants.cpp b/Engine/Plugins/FX/Niagara/Source/Niagara/Private/NiagaraConstants.cpp
--- a/Engine/Plugins/FX/Niagara/Source/Niagara/Private/NiagaraConstants.cpp
+++ b/Engine/Plugins/FX/Niagara/Source/Niagara/Private/NiagaraConstants.cpp
@@ -9,6 +9,7 @@
 
 TArray<FNiagaraVariable> FNiagaraConstants::SystemParameters;
 TArray<FNiagaraVariable> FNiagaraConstants::TranslatorParameters;
+TArray<FNiagaraVariable> FNiagaraConstants::SwitchParameters;
 TMap<FName, FNiagaraVariable> FNiagaraConstants::UpdatedSystemParameters;
 TMap<FNiagaraVariable, FText> FNiagaraConstants::SystemStrMap;
 TArray<FNiagaraVariable> FNiagaraConstants::Attributes;
@@ -71,6 +72,12 @@ void FNiagaraConstants::Init()
 		TranslatorParameters.Add(TRANSLATOR_PARAM_BEGIN_DEFAULTS);
 	}
 
+	if (SwitchParameters.Num() == 0)
+	{
+		SwitchParameters.Add(SYS_PARAM_EMITTER_LOCALSPACE);
+		SwitchParameters.Add(SYS_PARAM_EMITTER_DETERMINISM);
+	}
+
 	if (UpdatedSystemParameters.Num() == 0)
 	{
 		UpdatedSystemParameters.Add(FName(TEXT("System Delta Time")), SYS_PARAM_ENGINE_DELTA_TIME);
@@ -442,6 +449,12 @@ const TArray<FNiagaraVariable>& FNiagaraConstants::GetTranslatorConstants()
 	return TranslatorParameters;
 }
 
+const TArray<FNiagaraVariable>& FNiagaraConstants::GetStaticSwitchConstants()
+{
+	check(SwitchParameters.Num() != 0);
+	return SwitchParameters;
+}
+
 bool FNiagaraConstants::IsEngineManagedAttribute(const FNiagaraVariable& Var)
 {
 	return EngineManagedAttributes.Contains(Var);
@@ -467,15 +480,20 @@ FNiagaraVariable FNiagaraConstants::UpdateEngineConstant(const FNiagaraVariable&
 
 }
 
-
 const FNiagaraVariable* FNiagaraConstants::FindEngineConstant(const FNiagaraVariable& InVar)
 {
-	const TArray<FNiagaraVariable>& LocalSystemParameters = GetEngineConstants();
-	const FNiagaraVariable* FoundSystemVar = LocalSystemParameters.FindByPredicate([&](const FNiagaraVariable& Var)
+	return GetEngineConstants().FindByPredicate([&](const FNiagaraVariable& Var)
 	{
 		return Var.GetName() == InVar.GetName();
 	});
-	return FoundSystemVar;
+}
+
+const FNiagaraVariable * FNiagaraConstants::FindStaticSwitchConstant(const FName& InName)
+{
+	return GetStaticSwitchConstants().FindByPredicate([&](const FNiagaraVariable& Var)
+	{
+		return Var.GetName() == InName;
+	});
 }
 
 FText FNiagaraConstants::GetEngineConstantDescription(const FNiagaraVariable& InAttribute)
diff --git a/Engine/Plugins/FX/Niagara/Source/Niagara/Public/NiagaraTypes.h b/Engine/Plugins/FX/Niagara/Source/Niagara/Public/NiagaraTypes.h
--- a/Engine/Plugins/FX/Niagara/Source/Niagara/Public/NiagaraTypes.h
+++ b/Engine/Plugins/FX/Niagara/Source/Niagara/Public/NiagaraTypes.h
@@ -355,7 +355,7 @@ public:
 	UPROPERTY(EditAnywhere, Category = "Variable", DisplayName = "Property Metadata", meta = (ToolTip = "Property Metadata"))
 	TMap<FName, FString> PropertyMetaData;
 
-	UPROPERTY(AdvancedDisplay, VisibleAnywhere, Category = "Variable")
+	UPROPERTY(AdvancedDisplay, VisibleAnywhere, Category = "Variable", meta = (ToolTip = "This is a read-only variable that designates if the metadata is tied to a static switch or not."))
 	bool bIsStaticSwitch;
 
 	/** The default value to use when creating new pins or stack entries for a static switch parameter */
diff --git a/Engine/Plugins/FX/Niagara/Source/NiagaraEditor/Private/Customizations/NiagaraFunctionCallNodeDetails.cpp b/Engine/Plugins/FX/Niagara/Source/NiagaraEditor/Private/Customizations/NiagaraFunctionCallNodeDetails.cpp
--- a/Engine/Plugins/FX/Niagara/Source/NiagaraEditor/Private/Customizations/NiagaraFunctionCallNodeDetails.cpp
+++ b/Engine/Plugins/FX/Niagara/Source/NiagaraEditor/Private/Customizations/NiagaraFunctionCallNodeDetails.cpp
@@ -74,9 +74,13 @@ void FNiagaraFunctionCallNodeDetails::CustomizeDetails(IDetailLayoutBuilder& Det
 					{
 						Node->RemovePropagatedVariable(SwitchInput);
 					}
-					else if (NewState == ECheckBoxState::Checked && !Node->FindPropagatedVariable(SwitchInput))
+					else if (NewState == ECheckBoxState::Checked)
 					{
-						Node->PropagatedStaticSwitchParameters.Emplace(SwitchInput);
+						if (!Node->FindPropagatedVariable(SwitchInput))
+						{
+							Node->PropagatedStaticSwitchParameters.Emplace(SwitchInput);
+						}
+						CopyMetadataFromCalledGraph(SwitchInput);
 					}
 					Node->RefreshFromExternalChanges();
 				})
@@ -93,7 +97,9 @@ void FNiagaraFunctionCallNodeDetails::CustomizeDetails(IDetailLayoutBuilder& Det
 					FNiagaraPropagatedVariable* Propagated = Node->FindPropagatedVariable(SwitchInput);
 					if (Propagated)
 					{
+						FNiagaraVariable OldVar = Propagated->ToVariable();
 						Propagated->PropagatedName = NewText.ToString();
+						CopyMetadataForNameOverride(OldVar, Propagated->ToVariable());
 					}
 				})
 				.Text_Lambda([SwitchInput, this]()
@@ -110,4 +116,61 @@ void FNiagaraFunctionCallNodeDetails::CustomizeDetails(IDetailLayoutBuilder& Det
 	}
 }
 
+void FNiagaraFunctionCallNodeDetails::CopyMetadataFromCalledGraph(FNiagaraVariable FromVariable)
+{
+	UNiagaraGraph* NodeGraph = GetNodeGraph();
+	UNiagaraGraph* CalledGraph = GetCalledGraph();
+	if (!NodeGraph || !CalledGraph)
+	{
+		return;
+	}
+	if (NodeGraph->GetMetaData(FromVariable).IsSet())
+	{
+		return;
+	}
+	TOptional<FNiagaraVariableMetaData> OriginalData = CalledGraph->GetMetaData(FromVariable);
+	if (OriginalData.IsSet())
+	{
+		NodeGraph->SetMetaData(FromVariable, OriginalData.GetValue());
+		NodeGraph->NotifyGraphChanged();
+	}
+}
+
+void FNiagaraFunctionCallNodeDetails::CopyMetadataForNameOverride(FNiagaraVariable FromVariable, FNiagaraVariable ToVariable)
+{
+	UNiagaraGraph* NodeGraph = GetNodeGraph();
+	if (!NodeGraph)
+	{
+		return;
+	}
+	if (NodeGraph->GetMetaData(ToVariable).IsSet())
+	{
+		return;
+	}
+	TOptional<FNiagaraVariableMetaData> OriginalData = NodeGraph->GetMetaData(FromVariable);
+	if (OriginalData.IsSet())
+	{
+		NodeGraph->SetMetaData(ToVariable, OriginalData.GetValue());
+		NodeGraph->NotifyGraphChanged();
+	}
+}
+
+UNiagaraGraph* FNiagaraFunctionCallNodeDetails::GetNodeGraph()
+{
+	if (!Node.IsValid())
+	{
+		return nullptr;
+	}
+	return Node->GetNiagaraGraph();
+}
+
+UNiagaraGraph* FNiagaraFunctionCallNodeDetails::GetCalledGraph()
+{
+	if (!Node.IsValid())
+	{
+		return nullptr;
+	}
+	return Node->GetCalledGraph();
+}
+
 #undef LOCTEXT_NAMESPACE
diff --git a/Engine/Plugins/FX/Niagara/Source/NiagaraEditor/Private/Customizations/NiagaraFunctionCallNodeDetails.h b/Engine/Plugins/FX/Niagara/Source/NiagaraEditor/Private/Customizations/NiagaraFunctionCallNodeDetails.h
--- a/Engine/Plugins/FX/Niagara/Source/NiagaraEditor/Private/Customizations/NiagaraFunctionCallNodeDetails.h
+++ b/Engine/Plugins/FX/Niagara/Source/NiagaraEditor/Private/Customizations/NiagaraFunctionCallNodeDetails.h
@@ -6,6 +6,8 @@
 #include "IDetailCustomization.h"
 #include "DetailLayoutBuilder.h"
 #include "Misc/Optional.h"
+#include "NiagaraTypes.h"
+#include "NiagaraGraph.h"
 
 /** This customization sets up a custom details panel for the function call node in the niagara module graph. */
 class FNiagaraFunctionCallNodeDetails : public IDetailCustomization
@@ -20,4 +22,10 @@ public:
 	
 private:
 	TWeakObjectPtr<class UNiagaraNodeFunctionCall> Node;
+
+	void CopyMetadataFromCalledGraph(FNiagaraVariable FromVariable);
+	void CopyMetadataForNameOverride(FNiagaraVariable FromVariable, FNiagaraVariable ToVariable);
+
+	UNiagaraGraph* GetNodeGraph();
+	UNiagaraGraph* GetCalledGraph();
 };
diff --git a/Engine/Plugins/FX/Niagara/Source/NiagaraEditor/Private/Customizations/NiagaraStaticSwitchNodeDetails.cpp b/Engine/Plugins/FX/Niagara/Source/NiagaraEditor/Private/Customizations/NiagaraStaticSwitchNodeDetails.cpp
--- a/Engine/Plugins/FX/Niagara/Source/NiagaraEditor/Private/Customizations/NiagaraStaticSwitchNodeDetails.cpp
+++ b/Engine/Plugins/FX/Niagara/Source/NiagaraEditor/Private/Customizations/NiagaraStaticSwitchNodeDetails.cpp
@@ -14,6 +14,7 @@
 #include "Widgets/Input/SEditableTextBox.h"
 #include "Widgets/Input/SCheckBox.h"
 #include "Widgets/Layout/SWidgetSwitcher.h"
+#include "NiagaraConstants.h"
 
 #define LOCTEXT_NAMESPACE "NiagaraStaticSwitchNodeDetails"
 
@@ -25,7 +26,10 @@ TSharedRef<IDetailCustomization> FNiagaraStaticSwitchNodeDetails::MakeInstance()
 FNiagaraStaticSwitchNodeDetails::FNiagaraStaticSwitchNodeDetails()
 {
 	DropdownOptions.Add(MakeShareable(new SwitchDropdownOption("Bool")));
+	DropdownOptions.Add(MakeShareable(new SwitchDropdownOption("Bool Constant")));
 	DropdownOptions.Add(MakeShareable(new SwitchDropdownOption("Integer")));
+	DropdownOptions.Add(MakeShareable(new SwitchDropdownOption("Integer Constant")));
+	DropdownOptions.Add(MakeShareable(new SwitchDropdownOption("Enum Constant")));
 
 	for (FNiagaraTypeDefinition Type : FNiagaraTypeRegistry::GetRegisteredParameterTypes())
 	{
@@ -53,6 +57,7 @@ void FNiagaraStaticSwitchNodeDetails::CustomizeDetails(IDetailLayoutBuilder& Det
 		FDetailWidgetRow& DropdownWidget = CategoryBuilder.AddCustomRow(LOCTEXT("NiagaraSwitchNodeTypeFilterText", "Input parameter type"));
 		FDetailWidgetRow& IntValueOption = CategoryBuilder.AddCustomRow(LOCTEXT("NiagaraSwitchNodeIntFilterText", "Max integer value"));
 		FDetailWidgetRow& DefaultValueOption = CategoryBuilder.AddCustomRow(LOCTEXT("NiagaraSwitchNodeDefaultFilterText", "Default value"));
+		FDetailWidgetRow& ConstantSelection = CategoryBuilder.AddCustomRow(LOCTEXT("NiagaraSwitchNodeConstantFilterText", "Compiler constant"));
 
 		NameWidget
 		.NameContent()
@@ -145,6 +150,7 @@ void FNiagaraStaticSwitchNodeDetails::CustomizeDetails(IDetailLayoutBuilder& Det
 			.Padding(FMargin(0.0f, 2.0f))
 			[
 				SNew(STextBlock)
+				.IsEnabled(this, &FNiagaraStaticSwitchNodeDetails::GetDefaultOptionEnabled)
 				.TextStyle(FNiagaraEditorStyle::Get(), "NiagaraEditor.ParameterText")
 				.Text(LOCTEXT("NiagaraSwitchNodeDefaultOptionText", "Default value"))
 			]
@@ -155,6 +161,7 @@ void FNiagaraStaticSwitchNodeDetails::CustomizeDetails(IDetailLayoutBuilder& Det
 			.Padding(FMargin(0.0f, 2.0f))
 			[
 				SNew(SWidgetSwitcher)
+				.IsEnabled(this, &FNiagaraStaticSwitchNodeDetails::GetDefaultOptionEnabled)
 		  		.WidgetIndex(this, &FNiagaraStaticSwitchNodeDetails::GetDefaultWidgetIndex)
 
 		  		+ SWidgetSwitcher::Slot()
@@ -190,7 +197,41 @@ void FNiagaraStaticSwitchNodeDetails::CustomizeDetails(IDetailLayoutBuilder& Det
 			]
 		];
 
-		RefreshDefaultDropdownValues();
+		ConstantSelection
+		.NameContent()
+		[
+			SNew(SBox)
+			.Padding(FMargin(0.0f, 2.0f))
+			[
+				SNew(STextBlock)
+				.IsEnabled(this, &FNiagaraStaticSwitchNodeDetails::IsConstantSelection)
+				.TextStyle(FNiagaraEditorStyle::Get(), "NiagaraEditor.ParameterText")
+				.Text(LOCTEXT("NiagaraSwitchNodeConstantText", "Compiler constant"))
+			]
+		]
+		.ValueContent()
+		[
+			SNew(SHorizontalBox)
+			+ SHorizontalBox::Slot()
+			.HAlign(HAlign_Fill)
+			.AutoWidth()
+			.Padding(FMargin(0.0f, 2.0f))
+			[
+				SNew(SComboBox<TSharedPtr<ConstantDropdownOption>>)
+				.IsEnabled(this, &FNiagaraStaticSwitchNodeDetails::IsConstantSelection)
+				.OptionsSource(&ConstantOptions)
+				.OnSelectionChanged(this, &FNiagaraStaticSwitchNodeDetails::OnSelectionChanged)
+				.OnGenerateWidget(this, &FNiagaraStaticSwitchNodeDetails::CreateWidgetForDropdownOption)
+				.InitiallySelectedItem(SelectedConstantItem)
+				[
+					SNew(STextBlock)
+					.Margin(FMargin(0.0f, 2.0f))
+					.Text(this, &FNiagaraStaticSwitchNodeDetails::GetConstantSelectionItemLabel)
+				]
+			]
+		];
+
+		RefreshDropdownValues();
 	}
 }
 
@@ -240,6 +281,23 @@ TSharedRef<SWidget> FNiagaraStaticSwitchNodeDetails::CreateWidgetForDropdownOpti
 	return SNew(STextBlock).Text(InOption->DisplayName);
 }
 
+TSharedRef<SWidget> FNiagaraStaticSwitchNodeDetails::CreateWidgetForDropdownOption(TSharedPtr<ConstantDropdownOption> InOption)
+{
+	return SNew(STextBlock)
+		.Text(InOption->DisplayName)
+		.ToolTipText(InOption->Tooltip);
+}
+
+FText FNiagaraStaticSwitchNodeDetails::GetConstantSelectionItemLabel() const
+{
+	if (SelectedConstantItem.IsValid())
+	{
+		return SelectedConstantItem->DisplayName;
+	}
+
+	return LOCTEXT("InvalidNiagaraStaticSwitchNodeComboEntryText", "<Invalid selection>");
+}
+
 void FNiagaraStaticSwitchNodeDetails::OnSelectionChanged(TSharedPtr<DefaultEnumOption> NewValue, ESelectInfo::Type)
 {
 	SelectedDefaultValue = NewValue;
@@ -268,23 +326,46 @@ void FNiagaraStaticSwitchNodeDetails::OnSelectionChanged(TSharedPtr<SwitchDropdo
 	}
 	
 	FNiagaraTypeDefinition OldType = Node->GetInputType();
-	if (SelectedDropdownItem == DropdownOptions[0])
+	Node->SwitchTypeData.Enum = nullptr;
+	Node->SwitchTypeData.SwitchConstant = NAME_None;
+	if (SelectedDropdownItem == DropdownOptions[0] || SelectedDropdownItem == DropdownOptions[1])
 	{
 		Node->SwitchTypeData.SwitchType = ENiagaraStaticSwitchType::Bool;
-		Node->SwitchTypeData.Enum = nullptr;
 	}
-	else if (SelectedDropdownItem == DropdownOptions[1])
+	else if (SelectedDropdownItem == DropdownOptions[2] || SelectedDropdownItem == DropdownOptions[3])
 	{
 		Node->SwitchTypeData.SwitchType = ENiagaraStaticSwitchType::Integer;
-		Node->SwitchTypeData.Enum = nullptr;
 	}
 	else
 	{
 		Node->SwitchTypeData.SwitchType = ENiagaraStaticSwitchType::Enum;
-		Node->SwitchTypeData.Enum = SelectedDropdownItem->Enum;
+		if (!IsConstantSelection())
+		{
+			Node->SwitchTypeData.Enum = SelectedDropdownItem->Enum;
+		}
 	}
+	RefreshDropdownValues();
 	Node->OnSwitchParameterTypeChanged(OldType);
-	RefreshDefaultDropdownValues();
+}
+
+void FNiagaraStaticSwitchNodeDetails::OnSelectionChanged(TSharedPtr<ConstantDropdownOption> NewValue, ESelectInfo::Type)
+{
+	SelectedConstantItem = NewValue;
+	if (!NewValue.IsValid() || !Node.IsValid())
+	{
+		return;
+	}
+	Node->SwitchTypeData.SwitchConstant = NewValue->Constant.GetName();
+
+	// in case of an enum constant we also need the enum type
+	if (SelectedDropdownItem == DropdownOptions[4])
+	{
+		Node->SwitchTypeData.Enum = NewValue->Constant.GetType().Enum;
+	}
+	else
+	{
+		Node->SwitchTypeData.Enum = nullptr;
+	}
 }
 
 FText FNiagaraStaticSwitchNodeDetails::GetDropdownItemLabel() const
@@ -307,9 +388,50 @@ FText FNiagaraStaticSwitchNodeDetails::GetDefaultSelectionItemLabel() const
 	return LOCTEXT("InvalidNiagaraStaticSwitchNodeComboEntryText", "<Invalid selection>");
 }
 
-void FNiagaraStaticSwitchNodeDetails::RefreshDefaultDropdownValues()
+void FNiagaraStaticSwitchNodeDetails::RefreshDropdownValues()
 {
-	if (!Node.IsValid() || Node->SwitchTypeData.SwitchType != ENiagaraStaticSwitchType::Enum)
+	if (!Node.IsValid())
+	{
+		return;
+	}
+
+	// refresh the constant value dropdown values
+	ENiagaraStaticSwitchType SwitchType = Node->SwitchTypeData.SwitchType;
+	if (IsConstantSelection())
+	{
+		ConstantOptions.Empty();
+		SelectedConstantItem.Reset();
+
+		TArray<FNiagaraVariable> AllConstants;
+		AllConstants.Append(FNiagaraConstants::GetStaticSwitchConstants());
+
+		for (const FNiagaraVariable& Var : AllConstants)
+		{
+			FText DisplayName = FText::FromName(Var.GetName());
+			FText Tooltip = FNiagaraConstants::GetEngineConstantDescription(Var);
+			if ((SwitchType == ENiagaraStaticSwitchType::Bool && Var.GetType() == FNiagaraTypeDefinition::GetBoolDef()) ||
+				(SwitchType == ENiagaraStaticSwitchType::Integer && Var.GetType() == FNiagaraTypeDefinition::GetIntDef()) ||
+				(SwitchType == ENiagaraStaticSwitchType::Enum && Var.GetType().IsEnum()))
+			{				
+				ConstantOptions.Add(MakeShared<ConstantDropdownOption>(DisplayName, Tooltip, Var));
+				if (Node->SwitchTypeData.SwitchConstant == Var.GetName())
+				{
+					SelectedConstantItem = ConstantOptions.Last();
+				}
+			}
+		}
+		if (!SelectedConstantItem.IsValid() && ConstantOptions.Num() > 0)
+		{
+			OnSelectionChanged(ConstantOptions[0], ESelectInfo::Type::Direct);
+		}
+		if (ConstantOptions.Num() == 0)
+		{
+			OnSelectionChanged(MakeShared<ConstantDropdownOption>(LOCTEXT("NiagaraNoConstantsForTypeText", "<No entries>"), FText(), FNiagaraVariable()), ESelectInfo::Type::Direct);
+		}
+	}
+
+	// refresh the default options dropdown values
+	if (SwitchType != ENiagaraStaticSwitchType::Enum)
 	{
 		return;
 	}
@@ -336,7 +458,7 @@ void FNiagaraStaticSwitchNodeDetails::RefreshDefaultDropdownValues()
 		}
 		if (!SelectedDefaultValue.IsValid() && DefaultEnumDropdownOptions.Num() > 0)
 		{
-			SelectedDefaultValue = DefaultEnumDropdownOptions[0];
+			OnSelectionChanged(DefaultEnumDropdownOptions[0], ESelectInfo::Type::Direct);
 		}
 	}
 }
@@ -364,32 +486,45 @@ void FNiagaraStaticSwitchNodeDetails::UpdateSelectionFromNode()
 {
 	SelectedDropdownItem.Reset();
 	ENiagaraStaticSwitchType SwitchType = Node->SwitchTypeData.SwitchType;
+	bool IsConstantSelection = Node->IsSetByCompiler();
 
 	if (SwitchType == ENiagaraStaticSwitchType::Bool)
 	{
-		SelectedDropdownItem = DropdownOptions[0];
+		SelectedDropdownItem = IsConstantSelection ? DropdownOptions[1] : DropdownOptions[0];
 	}
 	else if (SwitchType == ENiagaraStaticSwitchType::Integer)
 	{
-		SelectedDropdownItem = DropdownOptions[1];
+		SelectedDropdownItem = IsConstantSelection ? DropdownOptions[3] : DropdownOptions[2];
 	}
 	else if (SwitchType == ENiagaraStaticSwitchType::Enum && Node->SwitchTypeData.Enum)
 	{
-		FString SelectedName = Node->SwitchTypeData.Enum->GetName();
-		for (TSharedPtr<SwitchDropdownOption>& Option : DropdownOptions)
+		if (IsConstantSelection)
+		{
+			SelectedDropdownItem = DropdownOptions[4];
+		}
+		else
 		{
-			if (SelectedName.Equals(*Option->Name))
+			FString SelectedName = Node->SwitchTypeData.Enum->GetName();
+			for (TSharedPtr<SwitchDropdownOption>& Option : DropdownOptions)
 			{
-				SelectedDropdownItem = Option;
-				return;
+				if (SelectedName.Equals(*Option->Name))
+				{
+					SelectedDropdownItem = Option;
+					return;
+				}
 			}
 		}
 	}
 }
 
+bool FNiagaraStaticSwitchNodeDetails::IsConstantSelection() const
+{
+	return SelectedDropdownItem == DropdownOptions[1] || SelectedDropdownItem == DropdownOptions[3] || SelectedDropdownItem == DropdownOptions[4];
+}
+
 bool FNiagaraStaticSwitchNodeDetails::GetIntOptionEnabled() const
 {
-	return Node.IsValid() && Node->SwitchTypeData.SwitchType == ENiagaraStaticSwitchType::Integer;
+	return Node.IsValid() && Node->SwitchTypeData.SwitchType == ENiagaraStaticSwitchType::Integer && (!IsConstantSelection() || SelectedDropdownItem == DropdownOptions[3]);
 }
 
 TOptional<int32> FNiagaraStaticSwitchNodeDetails::GetIntOptionValue() const
@@ -419,4 +554,9 @@ void FNiagaraStaticSwitchNodeDetails::OnParameterNameCommited(const FText& InTex
 	}
 }
 
+bool FNiagaraStaticSwitchNodeDetails::GetDefaultOptionEnabled() const
+{
+	return !IsConstantSelection();
+}
+
 #undef LOCTEXT_NAMESPACE
diff --git a/Engine/Plugins/FX/Niagara/Source/NiagaraEditor/Private/Customizations/NiagaraStaticSwitchNodeDetails.h b/Engine/Plugins/FX/Niagara/Source/NiagaraEditor/Private/Customizations/NiagaraStaticSwitchNodeDetails.h
--- a/Engine/Plugins/FX/Niagara/Source/NiagaraEditor/Private/Customizations/NiagaraStaticSwitchNodeDetails.h
+++ b/Engine/Plugins/FX/Niagara/Source/NiagaraEditor/Private/Customizations/NiagaraStaticSwitchNodeDetails.h
@@ -36,6 +36,22 @@ struct DefaultEnumOption
 	}
 };
 
+// This data structure is used internally by the constant selection to keep track of the user's choice
+struct ConstantDropdownOption
+{
+	FText DisplayName;
+	FText Tooltip;
+	FNiagaraVariable Constant;
+
+	ConstantDropdownOption(FText DisplayName) : DisplayName(DisplayName), Tooltip(FText()), Constant(FNiagaraVariable())
+	{
+	}
+
+	ConstantDropdownOption(FText DisplayName, FText Tooltip, FNiagaraVariable Constant) : DisplayName(DisplayName), Tooltip(Tooltip), Constant(Constant)
+	{
+	}
+};
+
 /** This customization sets up a custom details panel for the static switch node in the niagara module graph. */
 class FNiagaraStaticSwitchNodeDetails : public IDetailCustomization
 {
@@ -50,12 +66,12 @@ public:
 	virtual void CustomizeDetails( IDetailLayoutBuilder& DetailBuilder ) override;
 	
 private:
-	// enum dropdown functions
+	// type and enum dropdown functions
 	TSharedRef<SWidget> CreateWidgetForDropdownOption(TSharedPtr<SwitchDropdownOption> InOption);
 	void OnSelectionChanged(TSharedPtr<SwitchDropdownOption> NewValue, ESelectInfo::Type);
 	FText GetDropdownItemLabel() const;
 	void UpdateSelectionFromNode();
-
+	
 	// float type option functions
 	bool GetIntOptionEnabled() const;
 	TOptional<int32> GetIntOptionValue() const;
@@ -66,6 +82,7 @@ private:
 	void OnParameterNameCommited(const FText& InText, ETextCommit::Type InCommitType);
 
 	// default value option functions
+	bool GetDefaultOptionEnabled() const;
 	int32 GetDefaultWidgetIndex() const;
 	TOptional<int32> GetSwitchDefaultValue() const;
 	void DefaultIntValueCommitted(int32 Value, ETextCommit::Type CommitInfo);
@@ -73,14 +90,22 @@ private:
 	TSharedRef<SWidget> CreateWidgetForDropdownOption(TSharedPtr<DefaultEnumOption> InOption);
 	void OnSelectionChanged(TSharedPtr<DefaultEnumOption> NewValue, ESelectInfo::Type);
 	FText GetDefaultSelectionItemLabel() const;
-	void RefreshDefaultDropdownValues();
+	void RefreshDropdownValues();
 
 	TOptional<FNiagaraVariableMetaData> GetSwitchParameterMetadata() const;
 	void SetSwitchParameterMetadata(const FNiagaraVariableMetaData& MetaData);
 
+	// constant selection functions
+	bool IsConstantSelection() const;
+	void OnSelectionChanged(TSharedPtr<ConstantDropdownOption> NewValue, ESelectInfo::Type);
+	TSharedRef<SWidget> CreateWidgetForDropdownOption(TSharedPtr<ConstantDropdownOption> InOption);
+	FText GetConstantSelectionItemLabel() const;
+
 	TWeakObjectPtr<class UNiagaraNodeStaticSwitch> Node;
 	TArray<TSharedPtr<SwitchDropdownOption>> DropdownOptions;
 	TSharedPtr<SwitchDropdownOption> SelectedDropdownItem;
 	TArray<TSharedPtr<DefaultEnumOption>> DefaultEnumDropdownOptions;
 	TSharedPtr<DefaultEnumOption> SelectedDefaultValue;
+	TArray<TSharedPtr<ConstantDropdownOption>> ConstantOptions;
+	TSharedPtr<ConstantDropdownOption> SelectedConstantItem;
 };
diff --git a/Engine/Plugins/FX/Niagara/Source/NiagaraEditor/Private/NiagaraEditorUtilities.cpp b/Engine/Plugins/FX/Niagara/Source/NiagaraEditor/Private/NiagaraEditorUtilities.cpp
--- a/Engine/Plugins/FX/Niagara/Source/NiagaraEditor/Private/NiagaraEditorUtilities.cpp
+++ b/Engine/Plugins/FX/Niagara/Source/NiagaraEditor/Private/NiagaraEditorUtilities.cpp
@@ -659,7 +659,7 @@ void FNiagaraEditorUtilities::SetStaticSwitchConstants(UNiagaraGraph* Graph, con
 					continue;
 				}
 				ValuePin->DefaultValue = FString();
-				FName PinName = SwitchValue.PropagatedName.IsEmpty() ? ValuePin->GetFName() : FName(*SwitchValue.PropagatedName);
+				FName PinName = SwitchValue.ToVariable().GetName();
 				for (UEdGraphPin* InputPin : CallInputs)
 				{
 					if (InputPin->GetFName().IsEqual(PinName) && InputPin->PinType == ValuePin->PinType)
diff --git a/Engine/Plugins/FX/Niagara/Source/NiagaraEditor/Private/NiagaraGraph.cpp b/Engine/Plugins/FX/Niagara/Source/NiagaraEditor/Private/NiagaraGraph.cpp
--- a/Engine/Plugins/FX/Niagara/Source/NiagaraEditor/Private/NiagaraGraph.cpp
+++ b/Engine/Plugins/FX/Niagara/Source/NiagaraEditor/Private/NiagaraGraph.cpp
@@ -621,7 +621,8 @@ TArray<FNiagaraVariable> UNiagaraGraph::FindStaticSwitchInputs(bool bReachableOn
 	TArray<FNiagaraVariable> Result;
 	for (UEdGraphNode* Node : NodesToProcess)
 	{
-		if (UNiagaraNodeStaticSwitch* SwitchNode = Cast<UNiagaraNodeStaticSwitch>(Node))
+		UNiagaraNodeStaticSwitch* SwitchNode = Cast<UNiagaraNodeStaticSwitch>(Node);
+		if (SwitchNode && !SwitchNode->IsSetByCompiler())
 		{
 			FNiagaraVariable Variable(SwitchNode->GetInputType(), SwitchNode->InputParameterName);
 			Result.AddUnique(Variable);
@@ -631,12 +632,7 @@ TArray<FNiagaraVariable> UNiagaraGraph::FindStaticSwitchInputs(bool bReachableOn
 		{
 			for (const FNiagaraPropagatedVariable& Propagated : FunctionNode->PropagatedStaticSwitchParameters)
 			{
-				FNiagaraVariable Variable = Propagated.SwitchParameter;
-				if (!Propagated.PropagatedName.IsEmpty())
-				{
-					Variable.SetName(FName(*Propagated.PropagatedName));
-				}
-				Result.AddUnique(Variable);
+				Result.AddUnique(Propagated.ToVariable());
 			}			
 		}
 	}
@@ -662,7 +658,8 @@ TArray<UEdGraphNode*> UNiagaraGraph::FindReachbleNodes() const
 			continue;
 		}
 		
-		if (UNiagaraNodeStaticSwitch* SwitchNode = Cast<UNiagaraNodeStaticSwitch>(Node))
+		UNiagaraNodeStaticSwitch* SwitchNode = Cast<UNiagaraNodeStaticSwitch>(Node);
+		if (SwitchNode && !SwitchNode->IsSetByCompiler())
 		{
 			TArray<UEdGraphPin*> OutPins;
 			SwitchNode->GetOutputPins(OutPins);
@@ -1596,7 +1593,8 @@ void UNiagaraGraph::RefreshParameterReferences() const
 	const UEdGraphSchema_Niagara* NiagaraSchema = Cast<UEdGraphSchema_Niagara>(Schema);
 	for (UEdGraphNode* Node : Nodes)
 	{
-		if (UNiagaraNodeStaticSwitch* SwitchNode = Cast<UNiagaraNodeStaticSwitch>(Node))
+		UNiagaraNodeStaticSwitch* SwitchNode = Cast<UNiagaraNodeStaticSwitch>(Node);
+		if (SwitchNode && !SwitchNode->IsSetByCompiler())
 		{
 			FNiagaraVariable Variable(SwitchNode->GetInputType(), SwitchNode->InputParameterName);
 			AddStaticParameterReference(Variable, SwitchNode);
@@ -1605,12 +1603,7 @@ void UNiagaraGraph::RefreshParameterReferences() const
 		{
 			for (const FNiagaraPropagatedVariable& Propagated : FunctionNode->PropagatedStaticSwitchParameters)
 			{
-				FNiagaraVariable Variable = Propagated.SwitchParameter;
-				if (!Propagated.PropagatedName.IsEmpty())
-				{
-					Variable.SetName(FName(*Propagated.PropagatedName));
-				}
-				AddStaticParameterReference(Variable, FunctionNode);
+				AddStaticParameterReference(Propagated.ToVariable(), FunctionNode);
 			}
 		}
 
diff --git a/Engine/Plugins/FX/Niagara/Source/NiagaraEditor/Private/NiagaraNodeStaticSwitch.cpp b/Engine/Plugins/FX/Niagara/Source/NiagaraEditor/Private/NiagaraNodeStaticSwitch.cpp
--- a/Engine/Plugins/FX/Niagara/Source/NiagaraEditor/Private/NiagaraNodeStaticSwitch.cpp
+++ b/Engine/Plugins/FX/Niagara/Source/NiagaraEditor/Private/NiagaraNodeStaticSwitch.cpp
@@ -3,6 +3,7 @@
 #include "NiagaraNodeStaticSwitch.h"
 #include "NiagaraEditorUtilities.h"
 #include "NiagaraHlslTranslator.h"
+#include "NiagaraConstants.h"
 
 #define LOCTEXT_NAMESPACE "NiagaraNodeStaticSwitch"
 
@@ -21,7 +22,7 @@ FNiagaraTypeDefinition UNiagaraNodeStaticSwitch::GetInputType() const
 	{
 		return FNiagaraTypeDefinition::GetIntDef();
 	}
-	else if (SwitchTypeData.SwitchType == ENiagaraStaticSwitchType::Enum)
+	else if (SwitchTypeData.SwitchType == ENiagaraStaticSwitchType::Enum && SwitchTypeData.Enum)
 	{
 		return FNiagaraTypeDefinition(SwitchTypeData.Enum);
 	}
@@ -55,6 +56,11 @@ void UNiagaraNodeStaticSwitch::ClearSwitchValue()
 	SwitchValue = 0;
 }
 
+bool UNiagaraNodeStaticSwitch::IsSetByCompiler() const
+{
+	return !SwitchTypeData.SwitchConstant.IsNone();
+}
+
 void UNiagaraNodeStaticSwitch::RemoveUnusedGraphParameter(const FNiagaraVariable& OldParameter)
 {
 	TArray<FNiagaraVariable> GraphVariables = GetNiagaraGraph()->FindStaticSwitchInputs();
@@ -102,7 +108,7 @@ void UNiagaraNodeStaticSwitch::AllocateDefaultPins()
 			}
 		}
 	}
-	else if (SwitchTypeData.SwitchType == ENiagaraStaticSwitchType::Enum)
+	else if (SwitchTypeData.SwitchType == ENiagaraStaticSwitchType::Enum && SwitchTypeData.Enum)
 	{
 		// The last enum value is a special "max" value that we do not want to display, so we skip it
 		for (int32 i = 0; i < SwitchTypeData.Enum->NumEnums() - 1; i++)
@@ -145,7 +151,7 @@ void UNiagaraNodeStaticSwitch::InsertInputPinsFor(const FNiagaraVariable& Var)
 		// +1 because the range is inclusive
 		OptionsCount = SwitchTypeData.MaxIntCount + 1;
 	}
-	else if (SwitchTypeData.SwitchType == ENiagaraStaticSwitchType::Enum)
+	else if (SwitchTypeData.SwitchType == ENiagaraStaticSwitchType::Enum && SwitchTypeData.Enum)
 	{
 		// -1 because the last enum value is a special "max" value
 		OptionsCount = SwitchTypeData.Enum->NumEnums() - 1;
@@ -174,7 +180,7 @@ void UNiagaraNodeStaticSwitch::InsertInputPinsFor(const FNiagaraVariable& Var)
 		{
 			PathSuffix += FString::FromInt(i);
 		}
-		else if (SwitchTypeData.SwitchType == ENiagaraStaticSwitchType::Enum)
+		else if (SwitchTypeData.SwitchType == ENiagaraStaticSwitchType::Enum && SwitchTypeData.Enum)
 		{
 			FText EnumName = SwitchTypeData.Enum->GetDisplayNameTextByIndex(i);
 			PathSuffix += EnumName.ToString();
@@ -197,7 +203,40 @@ bool UNiagaraNodeStaticSwitch::GetVarIndex(FHlslNiagaraTranslator* Translator, i
 	return GetVarIndex(Translator, InputPinCount, SwitchValue, VarIndexOut);
 }
 
-bool UNiagaraNodeStaticSwitch::GetVarIndex(class FHlslNiagaraTranslator* Translator, int32 InputPinCount, int32 Value, int32& VarIndexOut) const
+void UNiagaraNodeStaticSwitch::UpdateCompilerConstantValue(FHlslNiagaraTranslator* Translator)
+{
+	if (!IsSetByCompiler() || !Translator)
+	{
+		return;
+	}
+	IsValueSet = false;
+
+	const FNiagaraVariable* Found = FNiagaraConstants::FindStaticSwitchConstant(SwitchTypeData.SwitchConstant);
+	FNiagaraVariable Constant = Found ? *Found : FNiagaraVariable();
+	if (Found && Translator->GetLiteralConstantVariable(Constant))
+	{
+		if (SwitchTypeData.SwitchType == ENiagaraStaticSwitchType::Bool)
+		{
+			SwitchValue = Constant.GetValue<bool>();
+			IsValueSet = true;
+		}
+		else if (SwitchTypeData.SwitchType == ENiagaraStaticSwitchType::Integer || SwitchTypeData.SwitchType == ENiagaraStaticSwitchType::Enum)
+		{
+			SwitchValue = Constant.GetValue<int32>();
+			IsValueSet = true;
+		}
+		else
+		{
+			Translator->Error(LOCTEXT("InvalidSwitchType", "Invalid static switch type."), this, nullptr);
+		}
+	}
+	else
+	{
+		Translator->Error(FText::Format(LOCTEXT("InvalidConstantValue", "Unable to determine constant value '{0}' for static switch."), FText::FromName(SwitchTypeData.SwitchConstant)), this, nullptr);
+	}
+}
+
+bool UNiagaraNodeStaticSwitch::GetVarIndex(FHlslNiagaraTranslator* Translator, int32 InputPinCount, int32 Value, int32& VarIndexOut) const
 {
 	bool Success = false;
 	if (SwitchTypeData.SwitchType == ENiagaraStaticSwitchType::Bool)
@@ -255,7 +294,14 @@ bool UNiagaraNodeStaticSwitch::SubstituteCompiledPin(FHlslNiagaraTranslator* Tra
 	// if we compile the standalone module or function we don't have any valid input yet, so we just take the first option to satisfy the compiler
 	ENiagaraScriptUsage TargetUsage = Translator->GetTargetUsage();
 	bool IsDryRun = TargetUsage == ENiagaraScriptUsage::Module || TargetUsage == ENiagaraScriptUsage::Function;
-	SwitchValue = IsDryRun ? 0 : SwitchValue;
+	if (IsDryRun)
+	{
+		SwitchValue = 0;
+	}
+	else
+	{
+		UpdateCompilerConstantValue(Translator);
+	}
 	if (!IsValueSet && !IsDryRun)
 	{
 		FText ErrorMessage = FText::Format(LOCTEXT("MissingSwitchValue", "The input parameter \"{0}\" is not set to a constant value for the static switch node."), FText::FromString(InputParameterName.ToString()));
diff --git a/Engine/Plugins/FX/Niagara/Source/NiagaraEditor/Private/NiagaraNodeStaticSwitch.h b/Engine/Plugins/FX/Niagara/Source/NiagaraEditor/Private/NiagaraNodeStaticSwitch.h
--- a/Engine/Plugins/FX/Niagara/Source/NiagaraEditor/Private/NiagaraNodeStaticSwitch.h
+++ b/Engine/Plugins/FX/Niagara/Source/NiagaraEditor/Private/NiagaraNodeStaticSwitch.h
@@ -30,6 +30,10 @@ struct FStaticSwitchTypeData
 	UPROPERTY()
 	UEnum* Enum;
 
+	/** If set, then this switch is not exposed but will rather be evaluated by the given compile-time constant */
+	UPROPERTY()
+	FName SwitchConstant;
+
 	FStaticSwitchTypeData() : SwitchType(ENiagaraStaticSwitchType::Bool), MaxIntCount(1), Enum(nullptr)
 	{ }
 };
@@ -54,6 +58,10 @@ public:
 
 	void SetSwitchValue(int Value);
 	void ClearSwitchValue();
+	/** If true then the value of this static switch is not set by the user but directly by the compiler via one of the engine constants (e.g. Emitter.Determinism). */
+	bool IsSetByCompiler() const;
+
+	void UpdateCompilerConstantValue(class FHlslNiagaraTranslator* Translator);
 
 	//~ Begin EdGraphNode Interface
 	virtual void AllocateDefaultPins() override;
diff --git a/Engine/Plugins/FX/Niagara/Source/NiagaraEditor/Private/ViewModels/Stack/NiagaraStackGraphUtilities.cpp b/Engine/Plugins/FX/Niagara/Source/NiagaraEditor/Private/ViewModels/Stack/NiagaraStackGraphUtilities.cpp
--- a/Engine/Plugins/FX/Niagara/Source/NiagaraEditor/Private/ViewModels/Stack/NiagaraStackGraphUtilities.cpp
+++ b/Engine/Plugins/FX/Niagara/Source/NiagaraEditor/Private/ViewModels/Stack/NiagaraStackGraphUtilities.cpp
@@ -532,7 +532,10 @@ void FNiagaraStackGraphUtilities::GetStackFunctionInputPins(UNiagaraNodeFunction
 		FNiagaraParameterMapHistoryBuilder BuilderCompiled;
 		FunctionCallNode.BuildParameterMapHistory(BuilderCompiled, false, true);
 		TArray<const UEdGraphPin*> CompilationPins;
-		ExtractInputPinsFromHistory(BuilderCompiled.Histories[0], FunctionCallNode.GetCalledGraph(), Options, CompilationPins);
+		if (BuilderCompiled.Histories.Num() == 1)
+		{
+			ExtractInputPinsFromHistory(BuilderCompiled.Histories[0], FunctionCallNode.GetCalledGraph(), Options, CompilationPins);
+		}
 
 		for (const UEdGraphPin* Pin : OutInputPins)
 		{
@@ -549,7 +552,9 @@ void FNiagaraStackGraphUtilities::GetStackFunctionStaticSwitchPins(UNiagaraNodeF
 	const UEdGraphSchema_Niagara* Schema = CastChecked<UEdGraphSchema_Niagara>(FunctionCallNode.GetSchema());
 	UNiagaraGraph* FunctionCallGraph = FunctionCallNode.GetCalledGraph();
 	if (FunctionCallGraph == nullptr)
+	{
 		return;
+	}
 
 	TArray<FNiagaraVariable> ReachableInputs = FunctionCallGraph->FindStaticSwitchInputs(true);
 	for (FNiagaraVariable SwitchInput : FunctionCallGraph->FindStaticSwitchInputs(false))
diff --git a/Engine/Plugins/FX/Niagara/Source/NiagaraEditor/Public/NiagaraNodeFunctionCall.h b/Engine/Plugins/FX/Niagara/Source/NiagaraEditor/Public/NiagaraNodeFunctionCall.h
--- a/Engine/Plugins/FX/Niagara/Source/NiagaraEditor/Public/NiagaraNodeFunctionCall.h
+++ b/Engine/Plugins/FX/Niagara/Source/NiagaraEditor/Public/NiagaraNodeFunctionCall.h
@@ -29,6 +29,16 @@ public:
 	UPROPERTY()
 	FString PropagatedName;
 
+	FNiagaraVariable ToVariable() const
+	{
+		FNiagaraVariable Copy = SwitchParameter;
+		if (!PropagatedName.IsEmpty())
+		{
+			Copy.SetName(FName(*PropagatedName));
+		}
+		return Copy;
+	}
+
 	bool operator==(const FNiagaraPropagatedVariable& Other)const
 	{
 		return SwitchParameter == Other.SwitchParameter;
@@ -94,7 +104,7 @@ public:
 
 	bool FindAutoBoundInput(UNiagaraNodeInput* InputNode, UEdGraphPin* PinToAutoBind, FNiagaraVariable& OutFoundVar, ENiagaraInputNodeUsage& OutNodeUsage);
 
-	void BuildParameterMapHistory(FNiagaraParameterMapHistoryBuilder& OutHistory, bool bRecursive = true, bool bFilterForCompilation = true) const;
+	void BuildParameterMapHistory(FNiagaraParameterMapHistoryBuilder& OutHistory, bool bRecursive = true, bool bFilterForCompilation = true) const override;
 
 	FString GetFunctionName() const { return FunctionDisplayName; }
 	UNiagaraGraph* GetCalledGraph() const;
