* Add compile target (cpu, gpu) as constant for static switch selection
* Fixed a bug where not all module inputs were culled correctly

diff --git a/Engine/Plugins/FX/Niagara/Source/Niagara/Classes/NiagaraConstants.h b/Engine/Plugins/FX/Niagara/Source/Niagara/Classes/NiagaraConstants.h
--- a/Engine/Plugins/FX/Niagara/Source/Niagara/Classes/NiagaraConstants.h
+++ b/Engine/Plugins/FX/Niagara/Source/Niagara/Classes/NiagaraConstants.h
@@ -65,6 +65,7 @@
 #define SYS_PARAM_EMITTER_RANDOM_SEED                    INiagaraModule::GetVar_Emitter_RandomSeed()
 #define SYS_PARAM_EMITTER_SPAWNRATE                      INiagaraModule::GetVar_Emitter_SpawnRate()
 #define SYS_PARAM_EMITTER_SPAWN_INTERVAL                 INiagaraModule::GetVar_Emitter_SpawnInterval()
+#define SYS_PARAM_EMITTER_SIMULATION_TARGET              INiagaraModule::GetVar_Emitter_SimulationTarget()
 #define SYS_PARAM_EMITTER_INTERP_SPAWN_START_DT          INiagaraModule::GetVar_Emitter_InterpSpawnStartDt()
 #define SYS_PARAM_EMITTER_SPAWN_GROUP                    INiagaraModule::GetVar_Emitter_SpawnGroup()
 
diff --git a/Engine/Plugins/FX/Niagara/Source/Niagara/Private/NiagaraConstants.cpp b/Engine/Plugins/FX/Niagara/Source/Niagara/Private/NiagaraConstants.cpp
--- a/Engine/Plugins/FX/Niagara/Source/Niagara/Private/NiagaraConstants.cpp
+++ b/Engine/Plugins/FX/Niagara/Source/Niagara/Private/NiagaraConstants.cpp
@@ -76,6 +76,7 @@ void FNiagaraConstants::Init()
 	{
 		SwitchParameters.Add(SYS_PARAM_EMITTER_LOCALSPACE);
 		SwitchParameters.Add(SYS_PARAM_EMITTER_DETERMINISM);
+		SwitchParameters.Add(SYS_PARAM_EMITTER_SIMULATION_TARGET);
 	}
 
 	if (UpdatedSystemParameters.Num() == 0)
diff --git a/Engine/Plugins/FX/Niagara/Source/Niagara/Private/NiagaraModule.cpp b/Engine/Plugins/FX/Niagara/Source/Niagara/Private/NiagaraModule.cpp
--- a/Engine/Plugins/FX/Niagara/Source/Niagara/Private/NiagaraModule.cpp
+++ b/Engine/Plugins/FX/Niagara/Source/Niagara/Private/NiagaraModule.cpp
@@ -108,6 +108,7 @@ FNiagaraVariable INiagaraModule::Engine_System_Age;
 FNiagaraVariable INiagaraModule::Emitter_Age;
 FNiagaraVariable INiagaraModule::Emitter_LocalSpace;
 FNiagaraVariable INiagaraModule::Emitter_Determinism;
+FNiagaraVariable INiagaraModule::Emitter_SimulationTarget;
 FNiagaraVariable INiagaraModule::Emitter_RandomSeed;
 FNiagaraVariable INiagaraModule::Emitter_SpawnRate;
 FNiagaraVariable INiagaraModule::Emitter_SpawnInterval;
@@ -215,6 +216,7 @@ void INiagaraModule::StartupModule()
 	Emitter_LocalSpace = FNiagaraVariable(FNiagaraTypeDefinition::GetBoolDef(), TEXT("Emitter.LocalSpace"));
 	Emitter_RandomSeed = FNiagaraVariable(FNiagaraTypeDefinition::GetIntDef(), TEXT("Emitter.RandomSeed"));
 	Emitter_Determinism = FNiagaraVariable(FNiagaraTypeDefinition::GetBoolDef(), TEXT("Emitter.Determinism"));
+	Emitter_SimulationTarget = FNiagaraVariable(FNiagaraTypeDefinition::GetSimulationTargetEnum(), TEXT("Emitter.SimulationTarget"));
 	Emitter_SpawnRate = FNiagaraVariable(FNiagaraTypeDefinition::GetFloatDef(), TEXT("Emitter.SpawnRate"));
 	Emitter_SpawnInterval = FNiagaraVariable(FNiagaraTypeDefinition::GetFloatDef(), TEXT("Emitter.SpawnInterval"));
 	Emitter_InterpSpawnStartDt = FNiagaraVariable(FNiagaraTypeDefinition::GetFloatDef(), TEXT("Emitter.InterpSpawnStartDt"));
@@ -465,6 +467,7 @@ UScriptStruct* FNiagaraTypeDefinition::ColorStruct;
 UScriptStruct* FNiagaraTypeDefinition::QuatStruct;
 
 UEnum* FNiagaraTypeDefinition::ExecutionStateEnum;
+UEnum* FNiagaraTypeDefinition::SimulationTargetEnum;
 UEnum* FNiagaraTypeDefinition::ExecutionStateSourceEnum;
 
 FNiagaraTypeDefinition FNiagaraTypeDefinition::ParameterMapDef;
@@ -572,6 +575,7 @@ void FNiagaraTypeDefinition::Init()
 
 	ExecutionStateEnum = StaticEnum<ENiagaraExecutionState>();
 	ExecutionStateSourceEnum = StaticEnum<ENiagaraExecutionStateSource>();
+	SimulationTargetEnum = StaticEnum<ENiagaraSimTarget>();
 	
 	RecreateUserDefinedTypeRegistry();
 }
diff --git a/Engine/Plugins/FX/Niagara/Source/Niagara/Public/NiagaraModule.h b/Engine/Plugins/FX/Niagara/Source/Niagara/Public/NiagaraModule.h
--- a/Engine/Plugins/FX/Niagara/Source/Niagara/Public/NiagaraModule.h
+++ b/Engine/Plugins/FX/Niagara/Source/Niagara/Public/NiagaraModule.h
@@ -131,6 +131,7 @@ public:
 	FORCEINLINE static const FNiagaraVariable&  GetVar_Emitter_RandomSeed() { return Emitter_RandomSeed; }
 	FORCEINLINE static const FNiagaraVariable&  GetVar_Emitter_SpawnRate() { return Emitter_SpawnRate; }
 	FORCEINLINE static const FNiagaraVariable&  GetVar_Emitter_SpawnInterval() { return Emitter_SpawnInterval; }
+	FORCEINLINE static const FNiagaraVariable&  GetVar_Emitter_SimulationTarget() { return Emitter_SimulationTarget; }
 	FORCEINLINE static const FNiagaraVariable&  GetVar_Emitter_InterpSpawnStartDt() { return Emitter_InterpSpawnStartDt; }
 	FORCEINLINE static const FNiagaraVariable&  GetVar_Emitter_SpawnGroup() { return Emitter_SpawnGroup; }
 
@@ -224,6 +225,7 @@ private:
 	static FNiagaraVariable Emitter_Age;
 	static FNiagaraVariable Emitter_LocalSpace;
 	static FNiagaraVariable Emitter_Determinism;
+	static FNiagaraVariable Emitter_SimulationTarget;
 	static FNiagaraVariable Emitter_RandomSeed;
 	static FNiagaraVariable Emitter_SpawnRate;
 	static FNiagaraVariable Emitter_SpawnInterval;
diff --git a/Engine/Plugins/FX/Niagara/Source/Niagara/Public/NiagaraTypes.h b/Engine/Plugins/FX/Niagara/Source/Niagara/Public/NiagaraTypes.h
--- a/Engine/Plugins/FX/Niagara/Source/Niagara/Public/NiagaraTypes.h
+++ b/Engine/Plugins/FX/Niagara/Source/Niagara/Public/NiagaraTypes.h
@@ -570,6 +570,7 @@ public:
 	static UScriptStruct* GetIDStruct() { return IDStruct; }
 
 	static UEnum* GetExecutionStateEnum() { return ExecutionStateEnum; }
+	static UEnum* GetSimulationTargetEnum() { return SimulationTargetEnum; }
 
 	static const FNiagaraTypeDefinition& GetCollisionEventDef() { return CollisionEventDef; }
 
@@ -617,6 +618,7 @@ private:
 	static UScriptStruct* Matrix4Struct;
 	static UScriptStruct* NumericStruct;
 
+	static UEnum* SimulationTargetEnum;
 	static UEnum* ExecutionStateEnum;
 	static UEnum* ExecutionStateSourceEnum;
 
diff --git a/Engine/Plugins/FX/Niagara/Source/NiagaraEditor/Private/NiagaraHlslTranslator.cpp b/Engine/Plugins/FX/Niagara/Source/NiagaraEditor/Private/NiagaraHlslTranslator.cpp
--- a/Engine/Plugins/FX/Niagara/Source/NiagaraEditor/Private/NiagaraHlslTranslator.cpp
+++ b/Engine/Plugins/FX/Niagara/Source/NiagaraEditor/Private/NiagaraHlslTranslator.cpp
@@ -3362,6 +3362,13 @@ bool FHlslNiagaraTranslator::GetLiteralConstantVariable(FNiagaraVariable& OutVar
 			OutVar.SetValue(bEmitterDeterminism ? FNiagaraBool(true) : FNiagaraBool(false));
 			return true;
 		}
+		if (OutVar == FNiagaraVariable(FNiagaraTypeDefinition::GetSimulationTargetEnum(), TEXT("Emitter.SimulationTarget")))
+		{
+			FNiagaraInt32 EnumValue;
+			EnumValue.Value = (uint8) CompilationTarget;
+			OutVar.SetValue(EnumValue);
+			return true;
+		}
 	}
 	return false;
 }
diff --git a/Engine/Plugins/FX/Niagara/Source/NiagaraEditor/Private/ViewModels/Stack/NiagaraStackGraphUtilities.cpp b/Engine/Plugins/FX/Niagara/Source/NiagaraEditor/Private/ViewModels/Stack/NiagaraStackGraphUtilities.cpp
--- a/Engine/Plugins/FX/Niagara/Source/NiagaraEditor/Private/ViewModels/Stack/NiagaraStackGraphUtilities.cpp
+++ b/Engine/Plugins/FX/Niagara/Source/NiagaraEditor/Private/ViewModels/Stack/NiagaraStackGraphUtilities.cpp
@@ -539,7 +539,17 @@ void FNiagaraStackGraphUtilities::GetStackFunctionInputPins(UNiagaraNodeFunction
 
 		for (const UEdGraphPin* Pin : OutInputPins)
 		{
-			if (!CompilationPins.Contains(Pin))
+			bool bFoundPin = false;
+			for (const UEdGraphPin* CompiledPin : CompilationPins)
+			{
+				if (Pin->GetName() == CompiledPin->GetName() && 
+					Pin->PinType.PinCategory == CompiledPin->PinType.PinCategory && 
+					Pin->PinType.PinSubCategoryObject == CompiledPin->PinType.PinSubCategoryObject)
+				{
+					bFoundPin = true;
+				}
+			}
+			if (!bFoundPin)
 			{
 				OutHiddenPins.Add(Pin);
 			}
