Set the node title of constant switches to the constant's name and disable the input parameter textbox.

diff --git a/Engine/Plugins/FX/Niagara/Source/NiagaraEditor/Private/Customizations/NiagaraStaticSwitchNodeDetails.cpp b/Engine/Plugins/FX/Niagara/Source/NiagaraEditor/Private/Customizations/NiagaraStaticSwitchNodeDetails.cpp
--- a/Engine/Plugins/FX/Niagara/Source/NiagaraEditor/Private/Customizations/NiagaraStaticSwitchNodeDetails.cpp
+++ b/Engine/Plugins/FX/Niagara/Source/NiagaraEditor/Private/Customizations/NiagaraStaticSwitchNodeDetails.cpp
@@ -77,6 +77,7 @@ void FNiagaraStaticSwitchNodeDetails::CustomizeDetails(IDetailLayoutBuilder& Det
 			[
 				SNew(SEditableTextBox)
 				.Text(this, &FNiagaraStaticSwitchNodeDetails::GetParameterNameText)
+				.IsEnabled(this, &FNiagaraStaticSwitchNodeDetails::GetDefaultOptionEnabled)
 				.ToolTipText(LOCTEXT("NiagaraSwitchNodeNameTooltip", "This is the name of the parameter that is exposed to the user calling this function graph."))
 				.OnTextCommitted(this, &FNiagaraStaticSwitchNodeDetails::OnParameterNameCommited)
 				.SelectAllTextWhenFocused(true)
@@ -356,7 +357,7 @@ void FNiagaraStaticSwitchNodeDetails::OnSelectionChanged(TSharedPtr<ConstantDrop
 		return;
 	}
 	Node->SwitchTypeData.SwitchConstant = NewValue->Constant.GetName();
-
+	
 	// in case of an enum constant we also need the enum type
 	if (SelectedDropdownItem == DropdownOptions[4])
 	{
@@ -366,6 +367,7 @@ void FNiagaraStaticSwitchNodeDetails::OnSelectionChanged(TSharedPtr<ConstantDrop
 	{
 		Node->SwitchTypeData.Enum = nullptr;
 	}
+	Node->RefreshFromExternalChanges();
 }
 
 FText FNiagaraStaticSwitchNodeDetails::GetDropdownItemLabel() const
diff --git a/Engine/Plugins/FX/Niagara/Source/NiagaraEditor/Private/NiagaraNodeStaticSwitch.cpp b/Engine/Plugins/FX/Niagara/Source/NiagaraEditor/Private/NiagaraNodeStaticSwitch.cpp
--- a/Engine/Plugins/FX/Niagara/Source/NiagaraEditor/Private/NiagaraNodeStaticSwitch.cpp
+++ b/Engine/Plugins/FX/Niagara/Source/NiagaraEditor/Private/NiagaraNodeStaticSwitch.cpp
@@ -380,7 +380,7 @@ FText UNiagaraNodeStaticSwitch::GetTooltipText() const
 
 FText UNiagaraNodeStaticSwitch::GetNodeTitle(ENodeTitleType::Type TitleType) const
 {
-	return FText::FormatOrdered(LOCTEXT("StaticSwitchTitle", "Static Switch ({0})"), FText::FromName(InputParameterName));
+	return FText::FormatOrdered(LOCTEXT("StaticSwitchTitle", "Static Switch ({0})"), FText::FromName(IsSetByCompiler() ? SwitchTypeData.SwitchConstant : InputParameterName));
 }
 
 FLinearColor UNiagaraNodeStaticSwitch::GetNodeTitleColor() const
