Fixed a crash that could happen when saving a changed niagara system.
(previous CL did not correctly submit the changes)

diff --git a/Engine/Plugins/FX/Niagara/Source/NiagaraEditor/Private/Toolkits/NiagaraScriptToolkit.cpp b/Engine/Plugins/FX/Niagara/Source/NiagaraEditor/Private/Toolkits/NiagaraScriptToolkit.cpp
--- a/Engine/Plugins/FX/Niagara/Source/NiagaraEditor/Private/Toolkits/NiagaraScriptToolkit.cpp
+++ b/Engine/Plugins/FX/Niagara/Source/NiagaraEditor/Private/Toolkits/NiagaraScriptToolkit.cpp
@@ -628,6 +628,10 @@ void FNiagaraScriptToolkit::UpdateOriginalNiagaraScript()
 
 		// First see if it is directly called, as this will force a need to refresh from external changes...
 		UNiagaraScriptSource* Source = Cast<UNiagaraScriptSource>(It->GetSource());
+		if (!Source)
+		{
+			continue;
+		}
 		TArray<UNiagaraNode*> NiagaraNodes;
 		Source->NodeGraph->GetNodesOfClass<UNiagaraNode>(NiagaraNodes);
 		bool bRefreshed = false;
