Add script library and script asset drag and drop

diff --git a/Engine/Plugins/FX/Niagara/Source/Niagara/Classes/NiagaraScript.h b/Engine/Plugins/FX/Niagara/Source/Niagara/Classes/NiagaraScript.h
--- a/Engine/Plugins/FX/Niagara/Source/Niagara/Classes/NiagaraScript.h
+++ b/Engine/Plugins/FX/Niagara/Source/Niagara/Classes/NiagaraScript.h
@@ -347,6 +347,10 @@ public:
 	UPROPERTY(EditAnywhere, Category = Script, meta = (EditCondition = "bDeprecated"))
 	UNiagaraScript* DeprecationRecommendation;
 
+	/* If this script is exposed to the library. */
+	UPROPERTY(AssetRegistrySearchable, EditAnywhere, Category = Script)
+	uint32 bExposeToLibrary : 1;
+
 #endif
 
 
diff --git a/Engine/Plugins/FX/Niagara/Source/Niagara/Private/NiagaraScript.cpp b/Engine/Plugins/FX/Niagara/Source/Niagara/Private/NiagaraScript.cpp
--- a/Engine/Plugins/FX/Niagara/Source/Niagara/Private/NiagaraScript.cpp
+++ b/Engine/Plugins/FX/Niagara/Source/Niagara/Private/NiagaraScript.cpp
@@ -491,6 +491,10 @@ void UNiagaraScript::PostLoad()
 				}
 			}
 		}
+		if (NiagaraVer < FNiagaraCustomVersion::AddLibraryAssetProperty)
+		{
+			bExposeToLibrary = true;
+		}
 	}
 
 	// Invalidate the CachedScriptVM if it's out of date to fix some cook errors, a further investigation is required in how to handle this correctly
diff --git a/Engine/Plugins/FX/Niagara/Source/NiagaraCore/Public/NiagaraCustomVersion.h b/Engine/Plugins/FX/Niagara/Source/NiagaraCore/Public/NiagaraCustomVersion.h
--- a/Engine/Plugins/FX/Niagara/Source/NiagaraCore/Public/NiagaraCustomVersion.h
+++ b/Engine/Plugins/FX/Niagara/Source/NiagaraCore/Public/NiagaraCustomVersion.h
@@ -122,6 +122,8 @@ struct FNiagaraCustomVersion
 
 		MoveInheritanceDataFromTheEmitterHandleToTheEmitter, // Moved the emitter inheritance data from the emitter handle to the emitter to allow for chained emitter inheritance.
 		
+		AddLibraryAssetProperty, // Add property to all Niagara scripts indicating whether or not they belong to the library
+
 		// -----<new versions can be added above this line>-------------------------------------------------
 		VersionPlusOne,
 		LatestVersion = VersionPlusOne - 1,
diff --git a/Engine/Plugins/FX/Niagara/Source/NiagaraEditor/Private/NiagaraEditorUtilities.cpp b/Engine/Plugins/FX/Niagara/Source/NiagaraEditor/Private/NiagaraEditorUtilities.cpp
--- a/Engine/Plugins/FX/Niagara/Source/NiagaraEditor/Private/NiagaraEditorUtilities.cpp
+++ b/Engine/Plugins/FX/Niagara/Source/NiagaraEditor/Private/NiagaraEditorUtilities.cpp
@@ -908,22 +908,21 @@ void FNiagaraEditorUtilities::GetFilteredScriptAssets(FGetFilteredScriptAssetsOp
 	ScriptFilter.TagsAndValues.Add(GET_MEMBER_NAME_CHECKED(UNiagaraScript, Usage), UnqualifiedScriptUsageString);
 
 	const FAssetRegistryModule& AssetRegistryModule = FModuleManager::LoadModuleChecked<FAssetRegistryModule>(TEXT("AssetRegistry"));
-	AssetRegistryModule.Get().GetAssets(ScriptFilter, OutFilteredScriptAssets);
+	TArray<FAssetData> FilteredScriptAssets;
+	AssetRegistryModule.Get().GetAssets(ScriptFilter, FilteredScriptAssets);
 
-	// We remove deprecated scripts separately as FARFilter does not support filtering by non-string tags.
-	if (InFilter.bIncludeDeprecatedScripts == false)
+	for (int i = 0; i < FilteredScriptAssets.Num(); ++i)
 	{
-		bool bScriptIsDeprecated = false;
-		bool bFoundDeprecatedTag = false;
-		for (int i = OutFilteredScriptAssets.Num() - 1; i >= 0; --i)
+		// Check if the script is deprecated
+		if (InFilter.bIncludeDeprecatedScripts == false)
 		{
-			bFoundDeprecatedTag = OutFilteredScriptAssets[i].GetTagValue(GET_MEMBER_NAME_CHECKED(UNiagaraScript, bDeprecated), bScriptIsDeprecated);
-			// If the asset does not have the metadata tag, check if it is loaded and if so check the bDeprecated value directly
+			bool bScriptIsDeprecated = false;
+			bool bFoundDeprecatedTag = FilteredScriptAssets[i].GetTagValue(GET_MEMBER_NAME_CHECKED(UNiagaraScript, bDeprecated), bScriptIsDeprecated);
 			if (bFoundDeprecatedTag == false)
 			{
-				if (OutFilteredScriptAssets[i].IsAssetLoaded())
+				if (FilteredScriptAssets[i].IsAssetLoaded())
 				{
-					UNiagaraScript* Script = static_cast<UNiagaraScript*>(OutFilteredScriptAssets[i].GetAsset());
+					UNiagaraScript* Script = static_cast<UNiagaraScript*>(FilteredScriptAssets[i].GetAsset());
 					if (Script != nullptr)
 					{
 						bScriptIsDeprecated = Script->bDeprecated;
@@ -932,26 +931,48 @@ void FNiagaraEditorUtilities::GetFilteredScriptAssets(FGetFilteredScriptAssetsOp
 			}
 			if (bScriptIsDeprecated)
 			{
-				OutFilteredScriptAssets.RemoveAt(i);
+				continue;
 			}
 		}
-	}
-	// We remove scripts with non matching usage bitmasks separately as FARFilter does not support filtering by non-string tags.
-	if (InFilter.TargetUsageToMatch.IsSet())
-	{
-		FString BitfieldTagValue;
-		int32 BitfieldValue, TargetBit;
-		for (int i = OutFilteredScriptAssets.Num() - 1; i >= 0; --i)
+
+		// Check if usage bitmask matches
+		if (InFilter.TargetUsageToMatch.IsSet())
 		{
-			BitfieldTagValue = OutFilteredScriptAssets[i].GetTagValueRef<FString>(GET_MEMBER_NAME_CHECKED(UNiagaraScript, ModuleUsageBitmask));
+			FString BitfieldTagValue;
+			int32 BitfieldValue, TargetBit;
+			BitfieldTagValue = FilteredScriptAssets[i].GetTagValueRef<FString>(GET_MEMBER_NAME_CHECKED(UNiagaraScript, ModuleUsageBitmask));
 			BitfieldValue = FCString::Atoi(*BitfieldTagValue);
-
 			TargetBit = (BitfieldValue >> (int32)InFilter.TargetUsageToMatch.GetValue()) & 1;
 			if (TargetBit != 1)
 			{
-				OutFilteredScriptAssets.RemoveAt(i);
+				continue;
+			}
+		}
+
+		// Check if library script
+		if (InFilter.bIncludeNonLibraryScripts == false)
+		{
+			bool bScriptIsLibrary = true;
+			bool bFoundLibScriptTag = FilteredScriptAssets[i].GetTagValue(GET_MEMBER_NAME_CHECKED(UNiagaraScript, bExposeToLibrary), bScriptIsLibrary);
+
+			if (bFoundLibScriptTag == false)
+			{
+				if (FilteredScriptAssets[i].IsAssetLoaded())
+				{
+					UNiagaraScript* Script = static_cast<UNiagaraScript*>(FilteredScriptAssets[i].GetAsset());
+					if (Script != nullptr)
+					{
+						bScriptIsLibrary = Script->bExposeToLibrary;
+					}
+				}
+			}
+			if (bScriptIsLibrary == false)
+			{
+				continue;
 			}
 		}
+
+		OutFilteredScriptAssets.Add(FilteredScriptAssets[i]);
 	}
 }
 
diff --git a/Engine/Plugins/FX/Niagara/Source/NiagaraEditor/Private/ViewModels/Stack/NiagaraStackEventHandlerGroup.cpp b/Engine/Plugins/FX/Niagara/Source/NiagaraEditor/Private/ViewModels/Stack/NiagaraStackEventHandlerGroup.cpp
--- a/Engine/Plugins/FX/Niagara/Source/NiagaraEditor/Private/ViewModels/Stack/NiagaraStackEventHandlerGroup.cpp
+++ b/Engine/Plugins/FX/Niagara/Source/NiagaraEditor/Private/ViewModels/Stack/NiagaraStackEventHandlerGroup.cpp
@@ -59,7 +59,7 @@ public:
 		OnItemAdded.ExecuteIfBound(EventScriptProperties);
 	}
 
-	virtual void GenerateAddActions(TArray<TSharedRef<INiagaraStackItemGroupAddAction>>& OutAddActions) const override { unimplemented(); }
+	virtual void GenerateAddActions(TArray<TSharedRef<INiagaraStackItemGroupAddAction>>& OutAddActions, const FNiagaraStackItemGroupAddOptions& AddProperties) const override { unimplemented(); }
 	virtual void ExecuteAddAction(TSharedRef<INiagaraStackItemGroupAddAction> AddAction, int32 TargetIndex) override { unimplemented(); }
 
 private:
diff --git a/Engine/Plugins/FX/Niagara/Source/NiagaraEditor/Private/ViewModels/Stack/NiagaraStackFunctionInput.cpp b/Engine/Plugins/FX/Niagara/Source/NiagaraEditor/Private/ViewModels/Stack/NiagaraStackFunctionInput.cpp
--- a/Engine/Plugins/FX/Niagara/Source/NiagaraEditor/Private/ViewModels/Stack/NiagaraStackFunctionInput.cpp
+++ b/Engine/Plugins/FX/Niagara/Source/NiagaraEditor/Private/ViewModels/Stack/NiagaraStackFunctionInput.cpp
@@ -882,11 +882,12 @@ UNiagaraNodeCustomHlsl* UNiagaraStackFunctionInput::GetExpressionNode() const
 	return InputValues.ExpressionNode.Get();
 }
 
-void UNiagaraStackFunctionInput::GetAvailableDynamicInputs(TArray<UNiagaraScript*>& AvailableDynamicInputs)
+void UNiagaraStackFunctionInput::GetAvailableDynamicInputs(TArray<UNiagaraScript*>& AvailableDynamicInputs, bool bIncludeNonLibraryInputs)
 {
 	TArray<FAssetData> DynamicInputAssets;
 	FNiagaraEditorUtilities::FGetFilteredScriptAssetsOptions DynamicInputScriptFilterOptions;
 	DynamicInputScriptFilterOptions.ScriptUsageToInclude = ENiagaraScriptUsage::DynamicInput;
+	DynamicInputScriptFilterOptions.bIncludeNonLibraryScripts = bIncludeNonLibraryInputs;
 	FNiagaraEditorUtilities::GetFilteredScriptAssets(DynamicInputScriptFilterOptions, DynamicInputAssets);
 
 	for (const FAssetData& DynamicInputAsset : DynamicInputAssets)
diff --git a/Engine/Plugins/FX/Niagara/Source/NiagaraEditor/Private/ViewModels/Stack/NiagaraStackModuleSpacer.cpp b/Engine/Plugins/FX/Niagara/Source/NiagaraEditor/Private/ViewModels/Stack/NiagaraStackModuleSpacer.cpp
--- a/Engine/Plugins/FX/Niagara/Source/NiagaraEditor/Private/ViewModels/Stack/NiagaraStackModuleSpacer.cpp
+++ b/Engine/Plugins/FX/Niagara/Source/NiagaraEditor/Private/ViewModels/Stack/NiagaraStackModuleSpacer.cpp
@@ -5,6 +5,7 @@
 #include "NiagaraTypes.h"
 #include "ViewModels/Stack/NiagaraParameterHandle.h"
 #include "ViewModels/Stack/NiagaraStackGraphUtilities.h"
+#include "DragAndDrop/AssetDragDropOp.h"
 
 void UNiagaraStackModuleSpacer::Initialize(FRequiredEntryData InRequiredEntryData, ENiagaraScriptUsage InScriptUsage, FName InSpacerKey, float InSpacerScale, EStackRowStyle InRowStyle)
 {
@@ -22,6 +23,15 @@ FReply UNiagaraStackModuleSpacer::OnStackSpacerDrop(TSharedPtr<FDragDropOperatio
 		OnStackSpacerAcceptDrop.ExecuteIfBound(this, Action->GetParameter());
 		return FReply::Handled();
 	}
+	else if (DragDropOperation->IsOfType<FAssetDragDropOp>())
+	{
+		TSharedPtr<FAssetDragDropOp> InputDragDropOp = StaticCastSharedPtr<FAssetDragDropOp>(DragDropOperation);
+		for (const FAssetData& AssetData : InputDragDropOp->GetAssets())
+		{
+			OnStackSpacerAssetDrop.ExecuteIfBound(this, AssetData);
+		}
+		return FReply::Handled();
+	}
 	return FReply::Unhandled();
 }
 
@@ -36,5 +46,17 @@ bool UNiagaraStackModuleSpacer::OnStackSpacerAllowDrop(TSharedPtr<FDragDropOpera
 			return FNiagaraStackGraphUtilities::ParameterIsCompatibleWithScriptUsage(Action->GetParameter(), ItemGroupScriptUsage);
 		}	
 	}
+	else if (DragDropOperation->IsOfType<FAssetDragDropOp>())
+	{
+		TSharedPtr<FAssetDragDropOp> InputDragDropOp = StaticCastSharedPtr<FAssetDragDropOp>(DragDropOperation);
+		for (const FAssetData& AssetData : InputDragDropOp->GetAssets())
+		{
+			if (OnStackSpacerRequestAssetDrop.IsBound() && OnStackSpacerRequestAssetDrop.Execute(this, AssetData) == false)
+			{
+				return false;
+			}
+		}
+		return true;
+	}
 	return false;
 }
diff --git a/Engine/Plugins/FX/Niagara/Source/NiagaraEditor/Private/ViewModels/Stack/NiagaraStackParameterStoreGroup.cpp b/Engine/Plugins/FX/Niagara/Source/NiagaraEditor/Private/ViewModels/Stack/NiagaraStackParameterStoreGroup.cpp
--- a/Engine/Plugins/FX/Niagara/Source/NiagaraEditor/Private/ViewModels/Stack/NiagaraStackParameterStoreGroup.cpp
+++ b/Engine/Plugins/FX/Niagara/Source/NiagaraEditor/Private/ViewModels/Stack/NiagaraStackParameterStoreGroup.cpp
@@ -69,7 +69,7 @@ public:
 
 	virtual void AddItemDirectly() override { unimplemented(); }
 
-	virtual void GenerateAddActions(TArray<TSharedRef<INiagaraStackItemGroupAddAction>>& OutAddActions) const override
+	virtual void GenerateAddActions(TArray<TSharedRef<INiagaraStackItemGroupAddAction>>& OutAddActions, const FNiagaraStackItemGroupAddOptions& AddProperties) const override
 	{
 		TArray<FNiagaraTypeDefinition> AvailableTypes;
 		FNiagaraStackGraphUtilities::GetNewParameterAvailableTypes(AvailableTypes);
diff --git a/Engine/Plugins/FX/Niagara/Source/NiagaraEditor/Private/ViewModels/Stack/NiagaraStackRenderItemGroup.cpp b/Engine/Plugins/FX/Niagara/Source/NiagaraEditor/Private/ViewModels/Stack/NiagaraStackRenderItemGroup.cpp
--- a/Engine/Plugins/FX/Niagara/Source/NiagaraEditor/Private/ViewModels/Stack/NiagaraStackRenderItemGroup.cpp
+++ b/Engine/Plugins/FX/Niagara/Source/NiagaraEditor/Private/ViewModels/Stack/NiagaraStackRenderItemGroup.cpp
@@ -64,7 +64,7 @@ public:
 
 	virtual void AddItemDirectly() override { unimplemented(); }
 
-	virtual void GenerateAddActions(TArray<TSharedRef<INiagaraStackItemGroupAddAction>>& OutAddActions) const override
+	virtual void GenerateAddActions(TArray<TSharedRef<INiagaraStackItemGroupAddAction>>& OutAddActions, const FNiagaraStackItemGroupAddOptions& AddProperties) const override
 	{
 		TArray<UClass*> RendererClasses;
 		GetDerivedClasses(UNiagaraRendererProperties::StaticClass(), RendererClasses);
diff --git a/Engine/Plugins/FX/Niagara/Source/NiagaraEditor/Private/ViewModels/Stack/NiagaraStackScriptItemGroup.cpp b/Engine/Plugins/FX/Niagara/Source/NiagaraEditor/Private/ViewModels/Stack/NiagaraStackScriptItemGroup.cpp
--- a/Engine/Plugins/FX/Niagara/Source/NiagaraEditor/Private/ViewModels/Stack/NiagaraStackScriptItemGroup.cpp
+++ b/Engine/Plugins/FX/Niagara/Source/NiagaraEditor/Private/ViewModels/Stack/NiagaraStackScriptItemGroup.cpp
@@ -40,12 +40,35 @@ public:
 			Category = LOCTEXT("ModuleNotCategorized", "Uncategorized Modules");
 		}
 
+		bool bIsLibraryScript = true;
+		bool bIsLibraryTagFound = false;
+		bIsLibraryTagFound = AssetData.GetTagValue(GET_MEMBER_NAME_CHECKED(UNiagaraScript, bExposeToLibrary), bIsLibraryScript);
+		if (bIsLibraryTagFound == false)
+		{
+			if (AssetData.IsAssetLoaded())
+			{
+				UNiagaraScript* Script = static_cast<UNiagaraScript*>(AssetData.GetAsset());
+				if (Script != nullptr)
+				{
+					bIsLibraryScript = Script->bExposeToLibrary;
+				}
+			}
+		}
+
 		FString DisplayNameString = FName::NameToDisplayString(AssetData.AssetName.ToString(), false);
+		if (bIsLibraryScript == false)
+		{
+			DisplayNameString += "*";
+		}
 		FText DisplayName = FText::FromString(DisplayNameString);
 
 		FText AssetDescription;
 		AssetData.GetTagValue(GET_MEMBER_NAME_CHECKED(UNiagaraScript, Description), AssetDescription);
 		FText Description = FNiagaraEditorUtilities::FormatScriptAssetDescription(AssetDescription, AssetData.ObjectPath);
+		if (bIsLibraryScript == false)
+		{
+			Description = FText::Format(LOCTEXT("ScriptDescriptionWithLibraryCommentFormat", "{0}\n{1}"), Description, LOCTEXT("NotExposedToLibrary", "*Not exposed to library"));
+		}
 
 		FText Keywords;
 		AssetData.GetTagValue(GET_MEMBER_NAME_CHECKED(UNiagaraScript, Keywords), Keywords);
@@ -160,7 +183,7 @@ public:
 
 	virtual void AddItemDirectly() override { unimplemented(); }
 
-	virtual void GenerateAddActions(TArray<TSharedRef<INiagaraStackItemGroupAddAction>>& OutAddActions) const override
+	virtual void GenerateAddActions(TArray<TSharedRef<INiagaraStackItemGroupAddAction>>& OutAddActions, const FNiagaraStackItemGroupAddOptions& AddProperties) const override
 	{
 		if (SystemViewModel.IsValid() == false || EmitterViewModel.IsValid() == false || OutputNode == nullptr)
 		{
@@ -172,6 +195,8 @@ public:
 		FNiagaraEditorUtilities::FGetFilteredScriptAssetsOptions ModuleScriptFilterOptions;
 		ModuleScriptFilterOptions.ScriptUsageToInclude = ENiagaraScriptUsage::Module;
 		ModuleScriptFilterOptions.TargetUsageToMatch = OutputNode->GetUsage();
+		ModuleScriptFilterOptions.bIncludeDeprecatedScripts = AddProperties.bIncludeDeprecated;
+		ModuleScriptFilterOptions.bIncludeNonLibraryScripts = AddProperties.bIncludeNonLibrary;
 		FNiagaraEditorUtilities::GetFilteredScriptAssets(ModuleScriptFilterOptions, ModuleAssets);
 		for (const FAssetData& ModuleAsset : ModuleAssets)
 		{
@@ -323,6 +348,8 @@ void UNiagaraStackScriptItemGroup::RefreshChildrenInternal(const TArray<UNiagara
 				ModuleSpacer = NewObject<UNiagaraStackModuleSpacer>(this);
 				ModuleSpacer->Initialize(CreateDefaultChildRequiredData(), GetScriptUsage(), ModuleSpacerKey, 1.4f, UNiagaraStackEntry::EStackRowStyle::None);
 				ModuleSpacer->OnStackSpacerAcceptDrop.BindUObject(this, &UNiagaraStackScriptItemGroup::AddParameterModuleToStack);
+				ModuleSpacer->OnStackSpacerRequestAssetDrop.BindUObject(this, &UNiagaraStackScriptItemGroup::CanAddAssetModuleToStack);
+				ModuleSpacer->OnStackSpacerAssetDrop.BindUObject(this, &UNiagaraStackScriptItemGroup::AddAssetModuleToStack);
 			}
 
 			NewChildren.Add(ModuleSpacer);
@@ -353,6 +380,8 @@ void UNiagaraStackScriptItemGroup::RefreshChildrenInternal(const TArray<UNiagara
 			ModuleSpacer = NewObject<UNiagaraStackModuleSpacer>(this);
 			ModuleSpacer->Initialize(CreateDefaultChildRequiredData(), GetScriptUsage(), ModuleSpacerKey, 1.4f, UNiagaraStackEntry::EStackRowStyle::None);
 			ModuleSpacer->OnStackSpacerAcceptDrop.BindUObject(this, &UNiagaraStackScriptItemGroup::AddParameterModuleToStack);
+			ModuleSpacer->OnStackSpacerRequestAssetDrop.BindUObject(this, &UNiagaraStackScriptItemGroup::CanAddAssetModuleToStack);
+			ModuleSpacer->OnStackSpacerAssetDrop.BindUObject(this, &UNiagaraStackScriptItemGroup::AddAssetModuleToStack);
 		}
 
 		NewChildren.Add(ModuleSpacer);
@@ -659,4 +688,38 @@ void UNiagaraStackScriptItemGroup::AddParameterModuleToStack(const UNiagaraStack
 	AddUtilities->ExecuteAddAction(AddAction, TargetIndex);
 }
 
+void UNiagaraStackScriptItemGroup::AddAssetModuleToStack(const UNiagaraStackModuleSpacer* InModuleSpcaer, const FAssetData& InAsset)
+{
+	int32 TargetIndex = INDEX_NONE;
+	UNiagaraStackModuleItem** TargetModuleItemPtr = StackSpacerToModuleItemMap.Find(FObjectKey(InModuleSpcaer));
+	if (*TargetModuleItemPtr != nullptr)
+	{
+		UNiagaraStackModuleItem* TargetModuleItem = *TargetModuleItemPtr;
+		TargetIndex = TargetModuleItem->GetModuleIndex();
+	}
+
+	TSharedRef<FScriptGroupAddAction> AddAction = FScriptGroupAddAction::CreateAssetModuleAction(InAsset);
+	AddUtilities->ExecuteAddAction(AddAction, TargetIndex);
+}
+
+bool UNiagaraStackScriptItemGroup::CanAddAssetModuleToStack(const UNiagaraStackModuleSpacer* InModuleSpacer, const FAssetData& InAsset)
+{
+	const UEnum* NiagaraScriptUsageEnum = FindObjectChecked<UEnum>(ANY_PACKAGE, TEXT("ENiagaraScriptUsage"), true);
+	FString AssetScriptUsageString = InAsset.GetTagValueRef<FString>(GET_MEMBER_NAME_CHECKED(UNiagaraScript, Usage));
+	int64 AssetScriptUsageValue = NiagaraScriptUsageEnum->GetValueByNameString(AssetScriptUsageString);
+	if (AssetScriptUsageValue == INDEX_NONE)
+	{
+		return false;
+	}
+	ENiagaraScriptUsage AssetScriptUsage = static_cast<ENiagaraScriptUsage>(AssetScriptUsageValue);
+	if (AssetScriptUsage != ENiagaraScriptUsage::Module)
+	{
+		return false;
+	}
+	FString BitfieldTagValue = InAsset.GetTagValueRef<FString>(GET_MEMBER_NAME_CHECKED(UNiagaraScript, ModuleUsageBitmask));
+	int32 BitfieldValue = FCString::Atoi(*BitfieldTagValue);
+	TArray<ENiagaraScriptUsage> SupportedUsages = UNiagaraScript::GetSupportedUsageContextsForBitmask(BitfieldValue);
+	return SupportedUsages.Contains(GetScriptUsage());
+}
+
 #undef LOCTEXT_NAMESPACE
diff --git a/Engine/Plugins/FX/Niagara/Source/NiagaraEditor/Public/NiagaraEditorUtilities.h b/Engine/Plugins/FX/Niagara/Source/NiagaraEditor/Public/NiagaraEditorUtilities.h
--- a/Engine/Plugins/FX/Niagara/Source/NiagaraEditor/Public/NiagaraEditorUtilities.h
+++ b/Engine/Plugins/FX/Niagara/Source/NiagaraEditor/Public/NiagaraEditorUtilities.h
@@ -151,6 +151,7 @@ namespace FNiagaraEditorUtilities
 	** @Param ScriptUsageToInclude Only return Scripts that have this usage
 	** @Param (Optional) TargetUsageToMatch Only return Scripts that have this target usage (output node) 
 	** @Param bIncludeDeprecatedScripts Whether or not to return Scripts that are deprecated (defaults to false) 
+	** @Param bIncludeNonLibraryScripts Whether or not to return non-library scripts (defaults to false)
 	*/
 	struct FGetFilteredScriptAssetsOptions
 	{
@@ -158,12 +159,14 @@ namespace FNiagaraEditorUtilities
 			: ScriptUsageToInclude(ENiagaraScriptUsage::Module)
 			, TargetUsageToMatch()
 			, bIncludeDeprecatedScripts(false)
+			, bIncludeNonLibraryScripts(false)
 		{
 		}
 
 		ENiagaraScriptUsage ScriptUsageToInclude;
 		TOptional<ENiagaraScriptUsage> TargetUsageToMatch;
 		bool bIncludeDeprecatedScripts;
+		bool bIncludeNonLibraryScripts;
 	};
 
 	NIAGARAEDITOR_API void GetFilteredScriptAssets(FGetFilteredScriptAssetsOptions InFilter, TArray<FAssetData>& OutFilteredScriptAssets); 
diff --git a/Engine/Plugins/FX/Niagara/Source/NiagaraEditor/Public/ViewModels/Stack/INiagaraStackItemGroupAddUtilities.h b/Engine/Plugins/FX/Niagara/Source/NiagaraEditor/Public/ViewModels/Stack/INiagaraStackItemGroupAddUtilities.h
--- a/Engine/Plugins/FX/Niagara/Source/NiagaraEditor/Public/ViewModels/Stack/INiagaraStackItemGroupAddUtilities.h
+++ b/Engine/Plugins/FX/Niagara/Source/NiagaraEditor/Public/ViewModels/Stack/INiagaraStackItemGroupAddUtilities.h
@@ -21,6 +21,16 @@ public:
 	virtual ~INiagaraStackItemGroupAddAction() { }
 };
 
+/** Defines options for adding items to groups in the stack. */
+struct FNiagaraStackItemGroupAddOptions
+{
+	/** Whether or not to include deprecated items. */
+	bool bIncludeDeprecated		= false;
+
+	/** Whether or not to include non-library items. */
+	bool bIncludeNonLibrary		= false;
+};
+
 /** Defines utilities for generically handling adding items to groups in the stack. */
 class INiagaraStackItemGroupAddUtilities
 {
@@ -48,7 +58,7 @@ public:
 	virtual void AddItemDirectly() = 0;
 
 	/** Populates an array with the valid add actions. */
-	virtual void GenerateAddActions(TArray<TSharedRef<INiagaraStackItemGroupAddAction>>& OutAddActions) const = 0;
+	virtual void GenerateAddActions(TArray<TSharedRef<INiagaraStackItemGroupAddAction>>& OutAddActions, const FNiagaraStackItemGroupAddOptions& AddProperties = FNiagaraStackItemGroupAddOptions()) const = 0;
 
 	/** Executes the specified add action. */
 	virtual void ExecuteAddAction(TSharedRef<INiagaraStackItemGroupAddAction> AddAction, int32 TargetIndex) = 0;
diff --git a/Engine/Plugins/FX/Niagara/Source/NiagaraEditor/Public/ViewModels/Stack/NiagaraStackFunctionInput.h b/Engine/Plugins/FX/Niagara/Source/NiagaraEditor/Public/ViewModels/Stack/NiagaraStackFunctionInput.h
--- a/Engine/Plugins/FX/Niagara/Source/NiagaraEditor/Public/ViewModels/Stack/NiagaraStackFunctionInput.h
+++ b/Engine/Plugins/FX/Niagara/Source/NiagaraEditor/Public/ViewModels/Stack/NiagaraStackFunctionInput.h
@@ -112,7 +112,7 @@ public:
 	UNiagaraNodeCustomHlsl* GetExpressionNode() const;
 
 	/** Gets the dynamic inputs available for this input. */
-	void GetAvailableDynamicInputs(TArray<UNiagaraScript*>& AvailableDynamicInputs);
+	void GetAvailableDynamicInputs(TArray<UNiagaraScript*>& AvailableDynamicInputs, bool bIncludeNonLibraryInputs = false);
 
 	/** Sets the dynamic input script for this input. */
 	void SetDynamicInput(UNiagaraScript* DynamicInput);
diff --git a/Engine/Plugins/FX/Niagara/Source/NiagaraEditor/Public/ViewModels/Stack/NiagaraStackModuleSpacer.h b/Engine/Plugins/FX/Niagara/Source/NiagaraEditor/Public/ViewModels/Stack/NiagaraStackModuleSpacer.h
--- a/Engine/Plugins/FX/Niagara/Source/NiagaraEditor/Public/ViewModels/Stack/NiagaraStackModuleSpacer.h
+++ b/Engine/Plugins/FX/Niagara/Source/NiagaraEditor/Public/ViewModels/Stack/NiagaraStackModuleSpacer.h
@@ -4,9 +4,9 @@
 
 #include "ViewModels/Stack/NiagaraStackSpacer.h"
 #include "NiagaraStackScriptItemGroup.h"
+#include "AssetData.h"
 #include "NiagaraStackModuleSpacer.generated.h"
 
-
 UCLASS()
 class NIAGARAEDITOR_API UNiagaraStackModuleSpacer : public UNiagaraStackSpacer
 {
@@ -15,6 +15,8 @@ class NIAGARAEDITOR_API UNiagaraStackModuleSpacer : public UNiagaraStackSpacer
 public:
 
 	DECLARE_DELEGATE_TwoParams(FOnStackSpacerAcceptDrop, const UNiagaraStackModuleSpacer*, const FNiagaraVariable&)
+	DECLARE_DELEGATE_TwoParams(FOnStackSpacerAssetDrop, const UNiagaraStackModuleSpacer*, const FAssetData&)
+	DECLARE_DELEGATE_RetVal_TwoParams(bool, FOnStackSpacerRequestAssetDrop, const UNiagaraStackModuleSpacer*, const FAssetData&)
 
 	void Initialize(FRequiredEntryData InRequiredEntryData, ENiagaraScriptUsage InScriptUsage, FName SpacerKey = NAME_None, float InSpacerScale = 1.0f, EStackRowStyle InRowStyle = UNiagaraStackEntry::EStackRowStyle::None);
 
@@ -24,6 +26,9 @@ public:
 
 	FOnStackSpacerAcceptDrop OnStackSpacerAcceptDrop;
 
+	FOnStackSpacerRequestAssetDrop OnStackSpacerRequestAssetDrop;
+	FOnStackSpacerAssetDrop OnStackSpacerAssetDrop;
+
 private:
 	ENiagaraScriptUsage ItemGroupScriptUsage;
 }; 
diff --git a/Engine/Plugins/FX/Niagara/Source/NiagaraEditor/Public/ViewModels/Stack/NiagaraStackScriptItemGroup.h b/Engine/Plugins/FX/Niagara/Source/NiagaraEditor/Public/ViewModels/Stack/NiagaraStackScriptItemGroup.h
--- a/Engine/Plugins/FX/Niagara/Source/NiagaraEditor/Public/ViewModels/Stack/NiagaraStackScriptItemGroup.h
+++ b/Engine/Plugins/FX/Niagara/Source/NiagaraEditor/Public/ViewModels/Stack/NiagaraStackScriptItemGroup.h
@@ -5,6 +5,7 @@
 #include "ViewModels/Stack/NiagaraStackItemGroup.h"
 #include "NiagaraCommon.h"
 #include "UObject/ObjectKey.h"
+#include "AssetRegistry/Public/AssetData.h"
 #include "NiagaraStackScriptItemGroup.generated.h"
 
 class FNiagaraScriptViewModel;
@@ -41,6 +42,10 @@ public:
 	 */
 	void AddParameterModuleToStack(const UNiagaraStackModuleSpacer* InModuleSpacer, const FNiagaraVariable& InVariable);
 
+	void AddAssetModuleToStack(const UNiagaraStackModuleSpacer* InModuleSpcaer, const FAssetData& InAsset);
+
+	bool CanAddAssetModuleToStack(const UNiagaraStackModuleSpacer* InModuleSpacer, const FAssetData& InAsset);
+
 protected:
 	virtual void RefreshChildrenInternal(const TArray<UNiagaraStackEntry*>& CurrentChildren, TArray<UNiagaraStackEntry*>& NewChildren, TArray<FStackIssue>& NewIssues) override;
 
diff --git a/Engine/Plugins/FX/Niagara/Source/NiagaraEditorWidgets/Private/Stack/SNiagaraStackFunctionInputValue.cpp b/Engine/Plugins/FX/Niagara/Source/NiagaraEditorWidgets/Private/Stack/SNiagaraStackFunctionInputValue.cpp
--- a/Engine/Plugins/FX/Niagara/Source/NiagaraEditorWidgets/Private/Stack/SNiagaraStackFunctionInputValue.cpp
+++ b/Engine/Plugins/FX/Niagara/Source/NiagaraEditorWidgets/Private/Stack/SNiagaraStackFunctionInputValue.cpp
@@ -18,6 +18,7 @@
 #include "Widgets/Input/SEditableTextBox.h"
 #include "Widgets/Input/SComboButton.h"
 #include "Widgets/Images/SImage.h"
+#include "Widgets/Input/SCheckBox.h"
 #include "Editor/PropertyEditor/Public/PropertyEditorModule.h"
 #include "Framework/Application/SlateApplication.h"
 #include "IStructureDetailsView.h"
@@ -32,6 +33,8 @@
 
 const float TextIconSize = 16;
 
+bool SNiagaraStackFunctionInputValue::bIncludeNonLibraryInputs = false;
+
 void SNiagaraStackFunctionInputValue::Construct(const FArguments& InArgs, UNiagaraStackFunctionInput* InFunctionInput)
 {
 	FunctionInput = InFunctionInput;
@@ -467,7 +470,6 @@ TSharedRef<SExpanderArrow> SNiagaraStackFunctionInputValue::CreateCustomNiagaraF
 
 TSharedRef<SWidget> SNiagaraStackFunctionInputValue::OnGetAvailableHandleMenu()
 {
-	TSharedPtr<SGraphActionMenu> SelectInputFunctionMenu;
 	TSharedRef<SBorder> MenuWidget = SNew(SBorder)
 	.BorderImage(FEditorStyle::GetBrush("Menu.Background"))
 	.Padding(5)
@@ -476,12 +478,45 @@ TSharedRef<SWidget> SNiagaraStackFunctionInputValue::OnGetAvailableHandleMenu()
 		.WidthOverride(300)
 		.HeightOverride(400)
 		[
-			SAssignNew(SelectInputFunctionMenu, SGraphActionMenu)
-			.OnActionSelected(this, &SNiagaraStackFunctionInputValue::OnActionSelected)
-			.OnCollectAllActions(this, &SNiagaraStackFunctionInputValue::CollectAllActions)
-			.AutoExpandActionMenu(false)
-			.ShowFilterTextBox(true)
-			.OnCreateCustomRowExpander_Static(&CreateCustomNiagaraFunctionInputActionExpander)
+			SNew(SVerticalBox)
+			+SVerticalBox::Slot()
+			.Padding(1.0f)
+			[
+				SNew(SHorizontalBox)
+
+				// Search context description
+				+SHorizontalBox::Slot()
+				.VAlign(VAlign_Center)
+				.AutoWidth()
+				[
+					SNew(STextBlock)
+					.Text(LOCTEXT("FunctionInputValueTitle", "Edit value"))
+				]
+
+				// Library Only Toggle
+				+SHorizontalBox::Slot()
+				.HAlign(HAlign_Right)
+				.VAlign(VAlign_Center)
+				[
+					SNew(SCheckBox)
+					.OnCheckStateChanged(this, &SNiagaraStackFunctionInputValue::OnLibraryToggleChanged)
+					.IsChecked(this, &SNiagaraStackFunctionInputValue::LibraryToggleIsChecked)
+					[
+						SNew(STextBlock)
+						.Text(LOCTEXT("LibraryOnly", "Library Only"))
+					]
+				]
+			]
+			+SVerticalBox::Slot()
+			.FillHeight(15)
+			[
+				SAssignNew(SelectInputFunctionMenu, SGraphActionMenu)
+				.OnActionSelected(this, &SNiagaraStackFunctionInputValue::OnActionSelected)
+				.OnCollectAllActions(this, &SNiagaraStackFunctionInputValue::CollectAllActions)
+				.AutoExpandActionMenu(false)
+				.ShowFilterTextBox(true)
+				.OnCreateCustomRowExpander_Static(&CreateCustomNiagaraFunctionInputActionExpander)
+			]
 		]
 	];
 
@@ -529,11 +564,11 @@ void SNiagaraStackFunctionInputValue::CollectAllActions(FGraphActionListBuilderB
 	{
 		const FText CategoryName = LOCTEXT("DynamicInputValueCategory", "Dynamic Inputs");
 		TArray<UNiagaraScript*> DynamicInputScripts;
-		FunctionInput->GetAvailableDynamicInputs(DynamicInputScripts);
+		FunctionInput->GetAvailableDynamicInputs(DynamicInputScripts, bIncludeNonLibraryInputs);
 		for (UNiagaraScript* DynamicInputScript : DynamicInputScripts)
 		{
-			const FText DynamicInputText = FText::FromString(FName::NameToDisplayString(DynamicInputScript->GetName(), false));
-			const FText Tooltip = FNiagaraEditorUtilities::FormatScriptAssetDescription(DynamicInputScript->Description, *DynamicInputScript->GetPathName());
+			const FText DynamicInputText = FText::FromString(FName::NameToDisplayString(DynamicInputScript->GetName() + (DynamicInputScript->bExposeToLibrary ? "" : "*"), false));
+			const FText Tooltip = FNiagaraEditorUtilities::FormatScriptAssetDescription(DynamicInputScript->Description, *(DynamicInputScript->GetPathName() + (DynamicInputScript->bExposeToLibrary ? "" : "\n*Not exposed to library")));
 			TSharedPtr<FNiagaraMenuAction> DynamicInputAction(new FNiagaraMenuAction(CategoryName, DynamicInputText, Tooltip, 0, DynamicInputScript->Keywords,
 				FNiagaraMenuAction::FOnExecuteStackAction::CreateSP(this, &SNiagaraStackFunctionInputValue::DynamicInputScriptSelected, DynamicInputScript)));
 			OutAllActions.AddAction(DynamicInputAction);
@@ -869,4 +904,15 @@ void SNiagaraStackFunctionInputValue::ShowReassignDynamicInputScriptMenu()
 	FSlateApplication::Get().PushMenu(AsShared(), FWidgetPath(), MenuWidget, MenuPosition, FPopupTransitionEffect::ContextMenu);
 }
 
+void SNiagaraStackFunctionInputValue::OnLibraryToggleChanged(ECheckBoxState CheckState)
+{
+	SNiagaraStackFunctionInputValue::bIncludeNonLibraryInputs = CheckState == ECheckBoxState::Unchecked;
+	SelectInputFunctionMenu->RefreshAllActions(true, false);
+}
+
+ECheckBoxState SNiagaraStackFunctionInputValue::LibraryToggleIsChecked() const
+{
+	return SNiagaraStackFunctionInputValue::bIncludeNonLibraryInputs ? ECheckBoxState::Unchecked : ECheckBoxState::Checked;
+}
+
 #undef LOCTEXT_NAMESPACE
diff --git a/Engine/Plugins/FX/Niagara/Source/NiagaraEditorWidgets/Private/Stack/SNiagaraStackFunctionInputValue.h b/Engine/Plugins/FX/Niagara/Source/NiagaraEditorWidgets/Private/Stack/SNiagaraStackFunctionInputValue.h
--- a/Engine/Plugins/FX/Niagara/Source/NiagaraEditorWidgets/Private/Stack/SNiagaraStackFunctionInputValue.h
+++ b/Engine/Plugins/FX/Niagara/Source/NiagaraEditorWidgets/Private/Stack/SNiagaraStackFunctionInputValue.h
@@ -137,6 +137,10 @@ private:
 
 	void ShowReassignDynamicInputScriptMenu();
 
+	void OnLibraryToggleChanged(ECheckBoxState CheckState);
+
+	ECheckBoxState LibraryToggleIsChecked() const;
+
 private:
 	UNiagaraStackFunctionInput* FunctionInput;
 
@@ -146,4 +150,7 @@ private:
 	TSharedPtr<SNiagaraParameterEditor> LocalValueStructParameterEditor;
 	TSharedPtr<IStructureDetailsView> LocalValueStructDetailsView;
 	TSharedPtr<SComboButton> SetFunctionInputButton;
+	TSharedPtr<SGraphActionMenu> SelectInputFunctionMenu;
+
+	static bool bIncludeNonLibraryInputs;
 };
diff --git a/Engine/Plugins/FX/Niagara/Source/NiagaraEditorWidgets/Private/Stack/SNiagaraStackItemGroupAddMenu.cpp b/Engine/Plugins/FX/Niagara/Source/NiagaraEditorWidgets/Private/Stack/SNiagaraStackItemGroupAddMenu.cpp
--- a/Engine/Plugins/FX/Niagara/Source/NiagaraEditorWidgets/Private/Stack/SNiagaraStackItemGroupAddMenu.cpp
+++ b/Engine/Plugins/FX/Niagara/Source/NiagaraEditorWidgets/Private/Stack/SNiagaraStackItemGroupAddMenu.cpp
@@ -7,6 +7,11 @@
 #include "EdGraph/EdGraph.h"
 #include "SGraphActionMenu.h"
 #include "Framework/Application/SlateApplication.h"
+#include "Widgets/Input/SCheckBox.h"
+
+#define LOCTEXT_NAMESPACE "NiagaraStackItemGroupAddMenu"
+
+bool SNiagaraStackItemGroupAddMenu::bIncludeNonLibraryScripts = false;
 
 void SNiagaraStackItemGroupAddMenu::Construct(const FArguments& InArgs, INiagaraStackItemGroupAddUtilities* InAddUtilities, int32 InInsertIndex)
 {
@@ -23,11 +28,44 @@ void SNiagaraStackItemGroupAddMenu::Construct(const FArguments& InArgs, INiagara
 			.WidthOverride(300)
 			.HeightOverride(400)
 			[
-				SAssignNew(AddMenu, SGraphActionMenu)
-				.OnActionSelected(this, &SNiagaraStackItemGroupAddMenu::OnActionSelected)
-				.OnCollectAllActions(this, &SNiagaraStackItemGroupAddMenu::CollectAllAddActions)
-				.AutoExpandActionMenu(AddUtilities->GetAutoExpandAddActions())
-				.ShowFilterTextBox(true)
+				SNew(SVerticalBox)
+				+SVerticalBox::Slot()
+				.Padding(1.0f)
+				[
+					SNew(SHorizontalBox)
+
+					// Search context description
+					+SHorizontalBox::Slot()
+					.VAlign(VAlign_Center)
+					.AutoWidth()
+					[
+						SNew(STextBlock)
+						.Text(FText::Format(LOCTEXT("AddToGroupFormatTitle", "Add new {0}"), AddUtilities->GetAddItemName()))
+					]
+
+					// Library Only Toggle
+					+SHorizontalBox::Slot()
+					.HAlign(HAlign_Right)
+					.VAlign(VAlign_Center)
+					[
+						SNew(SCheckBox)
+						.OnCheckStateChanged(this, &SNiagaraStackItemGroupAddMenu::OnLibraryToggleChanged)
+						.IsChecked(this, &SNiagaraStackItemGroupAddMenu::LibraryToggleIsChecked)
+						[
+							SNew(STextBlock)
+							.Text(LOCTEXT("LibraryOnly", "Library Only"))
+						]
+					]
+				]
+				+SVerticalBox::Slot()
+				.FillHeight(15)
+				[
+					SAssignNew(AddMenu, SGraphActionMenu)
+					.OnActionSelected(this, &SNiagaraStackItemGroupAddMenu::OnActionSelected)
+					.OnCollectAllActions(this, &SNiagaraStackItemGroupAddMenu::CollectAllAddActions)
+					.AutoExpandActionMenu(AddUtilities->GetAutoExpandAddActions())
+					.ShowFilterTextBox(true)
+				]
 			]
 		]
 	];
@@ -54,8 +92,12 @@ void SNiagaraStackItemGroupAddMenu::CollectAllAddActions(FGraphActionListBuilder
 		OutAllActions.OwnerOfTemporaries = NewObject<UEdGraph>((UObject*)GetTransientPackage());
 	}
 
+	FNiagaraStackItemGroupAddOptions AddOptions;
+	AddOptions.bIncludeDeprecated = false;
+	AddOptions.bIncludeNonLibrary = bIncludeNonLibraryScripts;
+
 	TArray<TSharedRef<INiagaraStackItemGroupAddAction>> AddActions;
-	AddUtilities->GenerateAddActions(AddActions);
+	AddUtilities->GenerateAddActions(AddActions, AddOptions);
 
 	for (TSharedRef<INiagaraStackItemGroupAddAction> AddAction : AddActions)
 	{
@@ -82,3 +124,14 @@ void SNiagaraStackItemGroupAddMenu::OnActionSelected(const TArray< TSharedPtr<FE
 		}
 	}
 }
+
+void SNiagaraStackItemGroupAddMenu::OnLibraryToggleChanged(ECheckBoxState CheckState)
+{
+	SNiagaraStackItemGroupAddMenu::bIncludeNonLibraryScripts = CheckState == ECheckBoxState::Unchecked;
+	AddMenu->RefreshAllActions(true, false);
+}
+
+ECheckBoxState SNiagaraStackItemGroupAddMenu::LibraryToggleIsChecked() const
+{
+	return SNiagaraStackItemGroupAddMenu::bIncludeNonLibraryScripts ? ECheckBoxState::Unchecked : ECheckBoxState::Checked;
+}
\ No newline at end of file
diff --git a/Engine/Plugins/FX/Niagara/Source/NiagaraEditorWidgets/Private/Stack/SNiagaraStackItemGroupAddMenu.h b/Engine/Plugins/FX/Niagara/Source/NiagaraEditorWidgets/Private/Stack/SNiagaraStackItemGroupAddMenu.h
--- a/Engine/Plugins/FX/Niagara/Source/NiagaraEditorWidgets/Private/Stack/SNiagaraStackItemGroupAddMenu.h
+++ b/Engine/Plugins/FX/Niagara/Source/NiagaraEditorWidgets/Private/Stack/SNiagaraStackItemGroupAddMenu.h
@@ -25,6 +25,10 @@ private:
 
 	void OnActionSelected(const TArray< TSharedPtr<FEdGraphSchemaAction> >& SelectedActions, ESelectInfo::Type InSelectionType);
 
+	void OnLibraryToggleChanged(ECheckBoxState CheckState);
+
+	ECheckBoxState LibraryToggleIsChecked() const;
+
 private:
 	INiagaraStackItemGroupAddUtilities* AddUtilities;
 
@@ -33,4 +37,6 @@ private:
 	TSharedPtr<class SGraphActionMenu> AddMenu;
 
 	bool bSetFocusOnNextTick;
+	
+	static bool bIncludeNonLibraryScripts;
 };
\ No newline at end of file
