After Uriel's change to batch up count readbacks from GPU for GPU systems, we weren't ever handling the case of going to zero. This caused us to never "complete" the system.

#tests engine tests pass

diff --git a/Engine/Plugins/FX/Niagara/Source/Niagara/Private/NiagaraEmitterInstanceBatcher.cpp b/Engine/Plugins/FX/Niagara/Source/Niagara/Private/NiagaraEmitterInstanceBatcher.cpp
--- a/Engine/Plugins/FX/Niagara/Source/Niagara/Private/NiagaraEmitterInstanceBatcher.cpp
+++ b/Engine/Plugins/FX/Niagara/Source/Niagara/Private/NiagaraEmitterInstanceBatcher.cpp
@@ -534,7 +534,9 @@ void NiagaraEmitterInstanceBatcher::PreInitViews(FRHICommandListImmediate& RHICm
 					if (CurrentData)
 					{
 						const uint32 DeadInstanceCount = Context->EmitterInstanceReadback.CPUCount - Counts[Context->EmitterInstanceReadback.GPUCountOffset];
-						if (DeadInstanceCount && DeadInstanceCount < CurrentData->GetNumInstances())
+
+						// This will communicate the particle counts to the game thread. If DeadInstanceCount equals CurrentData->GetNumInstances() the game thread will know that the emitter has completed.
+						if (DeadInstanceCount <= CurrentData->GetNumInstances())
 						{
 							CurrentData->SetNumInstances(CurrentData->GetNumInstances() - DeadInstanceCount);
 						}
