Fixes rendering under SM4.
Also unifies how we check feature levels / shader platforms.

diff --git a/Engine/Plugins/FX/Niagara/Source/Niagara/Private/NiagaraComponent.cpp b/Engine/Plugins/FX/Niagara/Source/Niagara/Private/NiagaraComponent.cpp
--- a/Engine/Plugins/FX/Niagara/Source/Niagara/Private/NiagaraComponent.cpp
+++ b/Engine/Plugins/FX/Niagara/Source/Niagara/Private/NiagaraComponent.cpp
@@ -204,7 +204,7 @@ FPrimitiveViewRelevance FNiagaraSceneProxy::GetViewRelevance(const FSceneView* V
 {
 	FPrimitiveViewRelevance Relevance;
 
-	if (bRenderingEnabled == false || (View->GetFeatureLevel() != ERHIFeatureLevel::SM5 && View->GetFeatureLevel() != ERHIFeatureLevel::ES3_1))
+	if (bRenderingEnabled == false || !FNiagaraUtilities::SupportsNiagaraRendering(View->GetFeatureLevel()))
 	{
 		return Relevance;
 	}
diff --git a/Engine/Plugins/FX/Niagara/Source/Niagara/Private/NiagaraEmitterInstance.cpp b/Engine/Plugins/FX/Niagara/Source/Niagara/Private/NiagaraEmitterInstance.cpp
--- a/Engine/Plugins/FX/Niagara/Source/Niagara/Private/NiagaraEmitterInstance.cpp
+++ b/Engine/Plugins/FX/Niagara/Source/Niagara/Private/NiagaraEmitterInstance.cpp
@@ -158,7 +158,7 @@ void FNiagaraEmitterInstance::Init(int32 InEmitterIdx, FName InSystemInstanceNam
 
 	if (!EmitterHandle.GetIsEnabled()
 		|| !CachedEmitter->IsAllowedByDetailLevel()
-		|| (GMaxRHIFeatureLevel != ERHIFeatureLevel::SM5 && GMaxRHIFeatureLevel != ERHIFeatureLevel::ES3_1 && CachedEmitter->SimTarget == ENiagaraSimTarget::GPUComputeSim)  // skip if GPU sim and <SM5. TODO: fall back to CPU sim instead once we have scalability functionality to do so
+		|| (!FNiagaraUtilities::SupportsGPUParticles(GMaxRHIFeatureLevel) && CachedEmitter->SimTarget == ENiagaraSimTarget::GPUComputeSim)  // skip if GPU sim and <SM5. TODO: fall back to CPU sim instead once we have scalability functionality to do so
 		)
 	{
 		ExecutionState = ENiagaraExecutionState::Disabled;
diff --git a/Engine/Plugins/FX/Niagara/Source/Niagara/Private/NiagaraScript.cpp b/Engine/Plugins/FX/Niagara/Source/Niagara/Private/NiagaraScript.cpp
--- a/Engine/Plugins/FX/Niagara/Source/Niagara/Private/NiagaraScript.cpp
+++ b/Engine/Plugins/FX/Niagara/Source/Niagara/Private/NiagaraScript.cpp
@@ -980,7 +980,7 @@ void UNiagaraScript::BeginCacheForCookedPlatformData(const ITargetPlatform *Targ
 		for (int32 FormatIndex = 0; FormatIndex < DesiredShaderFormats.Num(); FormatIndex++)
 		{
 			const EShaderPlatform LegacyShaderPlatform = ShaderFormatToLegacyShaderPlatform(DesiredShaderFormats[FormatIndex]);
-			if (RHISupportsComputeShaders(LegacyShaderPlatform))
+			if (FNiagaraUtilities::SupportsGPUParticles_NoES31(LegacyShaderPlatform))
 			{
 				CacheResourceShadersForCooking(LegacyShaderPlatform, CachedScriptResourcesForPlatform);
 			}
@@ -1074,7 +1074,7 @@ void UNiagaraScript::CacheResourceShadersForRendering(bool bRegenerateId, bool b
 				// SM4 is "in between" these feature levels but we do NOT support GPU particles in SM4. 
 				// So we can't use IsFeatureLevelSupported since SM4 will be seen as past ES3.1 and we do NOT support SM4 GPU particles.
 				// @todo-mattc This check should be rolled into RHISupportsComputeShaders.
-				if (GetMaxSupportedFeatureLevel(ShaderPlatform) == ERHIFeatureLevel::SM5 || GetMaxSupportedFeatureLevel(ShaderPlatform) == ERHIFeatureLevel::ES3_1)
+				if (FNiagaraUtilities::SupportsGPUParticles(CacheFeatureLevel))
 				{
 					ResourceToCache = ScriptResourcesByFeatureLevel[CacheFeatureLevel];
 					CacheShadersForResources(ShaderPlatform, &ScriptResource, true);
diff --git a/Engine/Plugins/FX/Niagara/Source/Niagara/Public/NiagaraCommon.h b/Engine/Plugins/FX/Niagara/Source/Niagara/Public/NiagaraCommon.h
--- a/Engine/Plugins/FX/Niagara/Source/Niagara/Public/NiagaraCommon.h
+++ b/Engine/Plugins/FX/Niagara/Source/Niagara/Public/NiagaraCommon.h
@@ -6,6 +6,7 @@
 #include "UObject/ObjectMacros.h"
 #include "NiagaraTypes.h"
 #include "UObject/SoftObjectPath.h"
+#include "RHIDefinitions.h"
 #include "NiagaraCommon.generated.h"
 
 class UNiagaraSystem;
@@ -534,6 +535,41 @@ namespace FNiagaraUtilities
 
 	void CollectScriptDataInterfaceParameters(const UObject& Owner, const TArray<UNiagaraScript*>& Scripts, FNiagaraParameterStore& OutDataInterfaceParameters);
 
+	inline bool SupportsNiagaraRendering(ERHIFeatureLevel::Type FeatureLevel)
+	{
+		return FeatureLevel == ERHIFeatureLevel::SM4 || FeatureLevel == ERHIFeatureLevel::SM5 || FeatureLevel == ERHIFeatureLevel::ES3_1;
+	}
+
+	inline bool SupportsNiagaraRendering(EShaderPlatform ShaderPlatform)
+	{
+		// Note:
+		// IsFeatureLevelSupported does a FeatureLevel < MaxFeatureLevel(ShaderPlatform) so checking ES3.1 support will return true for SM4. I added it explicitly to be clear what we are doing.
+		return IsFeatureLevelSupported(ShaderPlatform, ERHIFeatureLevel::SM5) || IsFeatureLevelSupported(ShaderPlatform, ERHIFeatureLevel::ES3_1) || IsFeatureLevelSupported(ShaderPlatform, ERHIFeatureLevel::SM4);
+	}
+
+	// There are two different versions of this temporarily to keep the previous behavior.
+	// Cooking previously checked RHISupportsComputeShaders which does NOT include ES3.1.
+	// Runtime DOES include ES3.1.
+	inline bool SupportsGPUParticles_NoES31(ERHIFeatureLevel::Type FeatureLevel)
+	{
+		return FeatureLevel == ERHIFeatureLevel::SM5;
+	}
+
+	inline bool SupportsGPUParticles_NoES31(EShaderPlatform ShaderPlatform)
+	{
+		return RHISupportsComputeShaders(ShaderPlatform);
+	}
+
+	inline bool SupportsGPUParticles(ERHIFeatureLevel::Type FeatureLevel)
+	{
+		return FeatureLevel == ERHIFeatureLevel::SM5 || FeatureLevel == ERHIFeatureLevel::ES3_1;
+	}
+	
+	inline bool SupportsGPUParticles(EShaderPlatform ShaderPlatform)
+	{
+		return GetMaxSupportedFeatureLevel(ShaderPlatform) == ERHIFeatureLevel::SM5 || GetMaxSupportedFeatureLevel(ShaderPlatform) == ERHIFeatureLevel::ES3_1;
+	}
+
 #if WITH_EDITORONLY_DATA
 	/**
 	 * Prepares rapid iteration parameter stores for simulation by removing old parameters no longer used by functions, by initializing new parameters
diff --git a/Engine/Plugins/FX/Niagara/Source/NiagaraVertexFactories/Private/NiagaraMeshVertexFactory.cpp b/Engine/Plugins/FX/Niagara/Source/NiagaraVertexFactories/Private/NiagaraMeshVertexFactory.cpp
--- a/Engine/Plugins/FX/Niagara/Source/NiagaraVertexFactories/Private/NiagaraMeshVertexFactory.cpp
+++ b/Engine/Plugins/FX/Niagara/Source/NiagaraVertexFactories/Private/NiagaraMeshVertexFactory.cpp
@@ -198,7 +198,7 @@ FShaderResourceViewRHIParamRef FNiagaraMeshVertexFactory::GetPreviousTransformBu
 
 bool FNiagaraMeshVertexFactory::ShouldCompilePermutation(EShaderPlatform Platform, const class FMaterial* Material, const class FShaderType* ShaderType)
 {
-	return (IsFeatureLevelSupported(Platform, ERHIFeatureLevel::SM5) || IsFeatureLevelSupported(Platform, ERHIFeatureLevel::ES3_1)) && (Material->IsUsedWithNiagaraMeshParticles() || Material->IsSpecialEngineMaterial());
+	return (FNiagaraUtilities::SupportsNiagaraRendering(Platform)) && (Material->IsUsedWithNiagaraMeshParticles() || Material->IsSpecialEngineMaterial());
 }
 
 void FNiagaraMeshVertexFactory::SetData(const FStaticMeshDataType& InData)
diff --git a/Engine/Plugins/FX/Niagara/Source/NiagaraVertexFactories/Private/NiagaraRibbonVertexFactory.cpp b/Engine/Plugins/FX/Niagara/Source/NiagaraVertexFactories/Private/NiagaraRibbonVertexFactory.cpp
--- a/Engine/Plugins/FX/Niagara/Source/NiagaraVertexFactories/Private/NiagaraRibbonVertexFactory.cpp
+++ b/Engine/Plugins/FX/Niagara/Source/NiagaraVertexFactories/Private/NiagaraRibbonVertexFactory.cpp
@@ -171,7 +171,7 @@ static TGlobalResource<FNiagaraRibbonVertexDeclaration> GNiagaraRibbonVertexDecl
 
 bool FNiagaraRibbonVertexFactory::ShouldCompilePermutation(EShaderPlatform Platform, const class FMaterial* Material, const class FShaderType* ShaderType)
 {
-	return (IsFeatureLevelSupported(Platform, ERHIFeatureLevel::SM5) || IsFeatureLevelSupported(Platform, ERHIFeatureLevel::ES3_1)) && (Material->IsUsedWithNiagaraRibbons() || Material->IsSpecialEngineMaterial());
+	return (FNiagaraUtilities::SupportsNiagaraRendering(Platform)) && (Material->IsUsedWithNiagaraRibbons() || Material->IsSpecialEngineMaterial());
 }
 
 /**
diff --git a/Engine/Plugins/FX/Niagara/Source/NiagaraVertexFactories/Private/NiagaraSpriteVertexFactory.cpp b/Engine/Plugins/FX/Niagara/Source/NiagaraVertexFactories/Private/NiagaraSpriteVertexFactory.cpp
--- a/Engine/Plugins/FX/Niagara/Source/NiagaraVertexFactories/Private/NiagaraSpriteVertexFactory.cpp
+++ b/Engine/Plugins/FX/Niagara/Source/NiagaraVertexFactories/Private/NiagaraSpriteVertexFactory.cpp
@@ -218,7 +218,7 @@ inline TGlobalResource<FNiagaraSpriteVertexDeclaration>& GetNiagaraSpriteVertexD
 
 bool FNiagaraSpriteVertexFactory::ShouldCompilePermutation(EShaderPlatform Platform, const class FMaterial* Material, const class FShaderType* ShaderType)
 {
-	return (IsFeatureLevelSupported(Platform, ERHIFeatureLevel::SM5) || IsFeatureLevelSupported(Platform, ERHIFeatureLevel::ES3_1)) && (Material->IsUsedWithNiagaraSprites() || Material->IsSpecialEngineMaterial());
+	return (FNiagaraUtilities::SupportsNiagaraRendering(Platform)) && (Material->IsUsedWithNiagaraSprites() || Material->IsSpecialEngineMaterial());
 }
 
 /**
