Fix Niagara GPUSim crash in ray tracing with BufferSize == 0

diff --git a/Engine/Plugins/FX/Niagara/Source/Niagara/Private/NiagaraRendererSprites.cpp b/Engine/Plugins/FX/Niagara/Source/Niagara/Private/NiagaraRendererSprites.cpp
--- a/Engine/Plugins/FX/Niagara/Source/Niagara/Private/NiagaraRendererSprites.cpp
+++ b/Engine/Plugins/FX/Niagara/Source/Niagara/Private/NiagaraRendererSprites.cpp
@@ -346,8 +346,21 @@ void NiagaraRendererSprites::GetDynamicMeshElements(const TArray<const FSceneVie
 #if RHI_RAYTRACING
 void NiagaraRendererSprites::GetRayTracingGeometryInstances(TArray<FRayTracingGeometryInstanceCollection>& OutInstanceCollections)
 {
+	FNiagaraDynamicDataSprites *DynamicDataSprites = static_cast<FNiagaraDynamicDataSprites*>(DynamicDataRender);
+
+	if (!DynamicDataSprites
+		|| DynamicDataSprites->RTParticleData.GetNumInstancesAllocated() == 0
+		|| DynamicDataSprites->RTParticleData.GetNumInstances() == 0
+		|| nullptr == Properties
+		|| !GSupportsResourceView // Current shader requires SRV to draw properly in all cases.
+		)
+	{
+		return;
+	}
+
 	FRayTracingGeometryInstanceCollection Collection;
 	Collection.Geometry = &RayTracingGeometry;
+	Collection.NumDynamicVertices = 6 * DynamicDataSprites->RTParticleData.GetNumInstances();
 	Collection.DynamicVertexPositionBuffer = &RayTracingDynamicVertexBuffer;
 	Collection.InstanceTransformMode = FRayTracingGeometryInstanceCollection::TransformMode::Identity;
 
diff --git a/Engine/Source/Runtime/RenderCore/Public/RenderResource.h b/Engine/Source/Runtime/RenderCore/Public/RenderResource.h
--- a/Engine/Source/Runtime/RenderCore/Public/RenderResource.h
+++ b/Engine/Source/Runtime/RenderCore/Public/RenderResource.h
@@ -593,6 +593,7 @@ struct FRayTracingGeometryInstanceCollection
 	};
 
 	FRayTracingGeometry* Geometry = nullptr;
+	uint32 NumDynamicVertices = 0;
 	FRWBuffer* DynamicVertexPositionBuffer = nullptr;
 	TransformMode InstanceTransformMode;
 	TArray<FMatrix> CustomTransforms;
diff --git a/Engine/Source/Runtime/Renderer/Private/DeferredShadingRenderer.cpp b/Engine/Source/Runtime/Renderer/Private/DeferredShadingRenderer.cpp
--- a/Engine/Source/Runtime/Renderer/Private/DeferredShadingRenderer.cpp
+++ b/Engine/Source/Runtime/Renderer/Private/DeferredShadingRenderer.cpp
@@ -854,7 +854,7 @@ bool FDeferredShadingSceneRenderer::GatherRayTracingWorldInstances(FRHICommandLi
 
 										if (ViewIndex == 0)
 										{
-											Scene->GetRayTracingDynamicGeometryCollection()->AddDynamicMeshBatchForGeometryUpdate(Scene, &View, SceneProxy, *BatchAndRelevance.Mesh, *InstanceCollections[GeometryIndex].Geometry, *InstanceCollections[GeometryIndex].DynamicVertexPositionBuffer);
+											Scene->GetRayTracingDynamicGeometryCollection()->AddDynamicMeshBatchForGeometryUpdate(Scene, &View, SceneProxy, *BatchAndRelevance.Mesh, *InstanceCollections[GeometryIndex].Geometry, InstanceCollections[GeometryIndex].NumDynamicVertices, *InstanceCollections[GeometryIndex].DynamicVertexPositionBuffer);
 										}
 									}
 
diff --git a/Engine/Source/Runtime/Renderer/Private/RayTracing/RayTracingDynamicGeometry.cpp b/Engine/Source/Runtime/Renderer/Private/RayTracing/RayTracingDynamicGeometry.cpp
--- a/Engine/Source/Runtime/Renderer/Private/RayTracing/RayTracingDynamicGeometry.cpp
+++ b/Engine/Source/Runtime/Renderer/Private/RayTracing/RayTracingDynamicGeometry.cpp
@@ -82,6 +82,7 @@ void FRayTracingDynamicGeometryCollection::AddDynamicMeshBatchForGeometryUpdate(
 	const FPrimitiveSceneProxy* PrimitiveSceneProxy, 
 	const FMeshBatch& MeshBatch, 
 	FRayTracingGeometry& Geometry,
+	uint32 NumVertices,
 	FRWBuffer& Buffer)
 {
 	const FMaterialRenderProxy* FallbackMaterialRenderProxyPtr = nullptr;
@@ -117,7 +118,7 @@ void FRayTracingDynamicGeometryCollection::AddDynamicMeshBatchForGeometryUpdate(
 	FVertexInputStreamArray DummyArray;
 	Shader->GetElementShaderBindings(Scene, View, MeshBatch.VertexFactory, false, Scene->GetFeatureLevel(), PrimitiveSceneProxy, MeshBatch, MeshBatch.Elements[0], ShaderElementData, SingleShaderBindings, DummyArray);
 
-	DispatchCmd.NumVertices = MeshBatch.Elements[0].NumPrimitives * 3 * MeshBatch.Elements[0].NumInstances;
+	DispatchCmd.NumVertices = NumVertices;
 	DispatchCmd.TargetBuffer = &Buffer;
 	DispatchCmd.TargetGeometry = &Geometry;
 
diff --git a/Engine/Source/Runtime/Renderer/Private/RayTracing/RayTracingDynamicGeometryCollection.h b/Engine/Source/Runtime/Renderer/Private/RayTracing/RayTracingDynamicGeometryCollection.h
--- a/Engine/Source/Runtime/Renderer/Private/RayTracing/RayTracingDynamicGeometryCollection.h
+++ b/Engine/Source/Runtime/Renderer/Private/RayTracing/RayTracingDynamicGeometryCollection.h
@@ -7,7 +7,7 @@
 class RENDERER_API FRayTracingDynamicGeometryCollection
 {
 public:
-	void AddDynamicMeshBatchForGeometryUpdate(const FScene* Scene, const FSceneView* View, const FPrimitiveSceneProxy* PrimitiveSceneProxy, const FMeshBatch& MeshBatch, FRayTracingGeometry& Geometry, FRWBuffer& Buffer);
+	void AddDynamicMeshBatchForGeometryUpdate(const FScene* Scene, const FSceneView* View, const FPrimitiveSceneProxy* PrimitiveSceneProxy, const FMeshBatch& MeshBatch, FRayTracingGeometry& Geometry, uint32 NumVertices, FRWBuffer& Buffer);
 	void DispatchUpdates(FRHICommandListImmediate& RHICmdList);
 	void Clear();
 
