Foundation work for Niagara System Overview Graph.
EdGraphSchema for Overview, and OverviewGraphViewModel.
Moved add and remove emitter to system functions from SystemViewModel to NiagaraEditorUtilities.

diff --git a/Engine/Plugins/FX/Niagara/Source/NiagaraEditor/Private/EdGraphSchema_NiagaraSystemOverview.cpp b/Engine/Plugins/FX/Niagara/Source/NiagaraEditor/Private/EdGraphSchema_NiagaraSystemOverview.cpp
new file mode 100644
--- /dev/null
+++ b/Engine/Plugins/FX/Niagara/Source/NiagaraEditor/Private/EdGraphSchema_NiagaraSystemOverview.cpp
@@ -0,0 +1,19 @@
+// Copyright 1998-2019 Epic Games, Inc. All Rights Reserved.
+
+#include "EdGraphSchema_NiagaraSystemOverview.h"
+#include "EdGraphNode_Comment.h"
+
+#define LOCTEXT_NAMESPACE "NiagaraSchema"
+
+UEdGraphSchema_NiagaraSystemOverview::UEdGraphSchema_NiagaraSystemOverview(const FObjectInitializer& ObjectInitializer)
+	: Super(ObjectInitializer)
+{
+}
+
+void UEdGraphSchema_NiagaraSystemOverview::GetContextMenuActions(const UEdGraph* CurrentGraph, const UEdGraphNode* InGraphNode, const UEdGraphPin* InGraphPin, class FMenuBuilder* MenuBuilder, bool bIsDebugging) const
+{
+	//@TODO System Overview: enable/disable emitter, solo, etc. 
+	return;
+}
+
+#undef LOCTEXT_NAMESPACE
diff --git a/Engine/Plugins/FX/Niagara/Source/NiagaraEditor/Private/NiagaraEditorUtilities.cpp b/Engine/Plugins/FX/Niagara/Source/NiagaraEditor/Private/NiagaraEditorUtilities.cpp
--- a/Engine/Plugins/FX/Niagara/Source/NiagaraEditor/Private/NiagaraEditorUtilities.cpp
+++ b/Engine/Plugins/FX/Niagara/Source/NiagaraEditor/Private/NiagaraEditorUtilities.cpp
@@ -1082,4 +1082,58 @@ TArray<UNiagaraComponent*> FNiagaraEditorUtilities::GetComponentsThatReferenceSy
 	return ReferencingComponents;
 }
 
+const FGuid FNiagaraEditorUtilities::AddEmitterToSystem(UNiagaraSystem& InSystem, UNiagaraEmitter& InEmitterToAdd, bool bSystemIsPlaceholder)
+{
+	// Kill all system instances before modifying the emitter handle list to prevent accessing deleted data.
+	KillSystemInstances(InSystem);
+
+	TSet<FName> EmitterHandleNames;
+	for (const FNiagaraEmitterHandle& EmitterHandle : InSystem.GetEmitterHandles())
+	{
+		EmitterHandleNames.Add(EmitterHandle.GetName());
+	}
+
+	FNiagaraEmitterHandle EmitterHandle;
+	if (false == bSystemIsPlaceholder)
+	{
+		InSystem.Modify();
+		EmitterHandle = InSystem.AddEmitterHandle(InEmitterToAdd, FNiagaraUtilities::GetUniqueName(InEmitterToAdd.GetFName(), EmitterHandleNames));
+	}
+	else
+	{
+		// When editing an emitter asset we add the emitter as a duplicate so that the parent emitter is duplicated, but it's parent emitter
+		// information is maintained.
+		checkf(InSystem.GetNumEmitters() == 0, TEXT("Can not add multiple emitters to a system being edited in emitter asset mode."));
+		FNiagaraEmitterHandle TemporaryEmitterHandle(InEmitterToAdd);
+		EmitterHandle = InSystem.DuplicateEmitterHandle(TemporaryEmitterHandle, *InEmitterToAdd.GetUniqueEmitterName());
+	}
+	
+	FNiagaraStackGraphUtilities::RebuildEmitterNodes(InSystem);
+
+	return EmitterHandle.GetId();
+}
+
+void FNiagaraEditorUtilities::RemoveEmittersFromSystemByEmitterHandleId(UNiagaraSystem& InSystem, TSet<FGuid> EmitterHandleIdsToDelete)
+{
+	// Kill all system instances before modifying the emitter handle list to prevent accessing deleted data.
+	KillSystemInstances(InSystem);
+
+	const FScopedTransaction DeleteTransaction(EmitterHandleIdsToDelete.Num() == 1
+		? LOCTEXT("DeleteEmitter", "Delete emitter")
+		: LOCTEXT("DeleteEmitters", "Delete emitters"));
+
+	InSystem.Modify();
+	InSystem.RemoveEmitterHandlesById(EmitterHandleIdsToDelete);
+	FNiagaraStackGraphUtilities::RebuildEmitterNodes(InSystem);
+}
+
+void FNiagaraEditorUtilities::KillSystemInstances(const UNiagaraSystem& System)
+{
+	TArray<UNiagaraComponent*> ReferencingComponents = FNiagaraEditorUtilities::GetComponentsThatReferenceSystem(System);
+	for (auto Component : ReferencingComponents)
+	{
+		Component->DestroyInstance();
+	}
+}
+
 #undef LOCTEXT_NAMESPACE
diff --git a/Engine/Plugins/FX/Niagara/Source/NiagaraEditor/Private/NiagaraOverviewNodeStackItem.cpp b/Engine/Plugins/FX/Niagara/Source/NiagaraEditor/Private/NiagaraOverviewNodeStackItem.cpp
new file mode 100644
--- /dev/null
+++ b/Engine/Plugins/FX/Niagara/Source/NiagaraEditor/Private/NiagaraOverviewNodeStackItem.cpp
@@ -0,0 +1,60 @@
+// Copyright 1998-2019 Epic Games, Inc. All Rights Reserved.
+
+#include "NiagaraOverviewNodeStackItem.h"
+#include "NiagaraSystem.h"
+
+#define LOCTEXT_NAMESPACE "NiagaraOverviewNodeStackItem"
+
+UNiagaraOverviewNodeStackItem::UNiagaraOverviewNodeStackItem()
+	: OwningSystem(nullptr)
+	, EmitterHandleGuid(FGuid())
+{
+};
+
+void UNiagaraOverviewNodeStackItem::Initialize(UNiagaraSystem* InOwningSystem)
+{
+	OwningSystem = InOwningSystem;
+}
+
+void UNiagaraOverviewNodeStackItem::Initialize(UNiagaraSystem* InOwningSystem, FGuid InEmitterHandleGuid)
+{
+	OwningSystem = InOwningSystem;
+	EmitterHandleGuid = InEmitterHandleGuid;
+}
+
+const FGuid UNiagaraOverviewNodeStackItem::GetEmitterHandleGuid() const
+{
+	return EmitterHandleGuid;
+}
+
+FText UNiagaraOverviewNodeStackItem::GetNodeTitle(ENodeTitleType::Type TitleType) const
+{
+	if (EmitterHandleGuid.IsValid())
+	{
+		for (const FNiagaraEmitterHandle& Handle : OwningSystem->GetEmitterHandles())
+		{
+			if (Handle.GetId() == EmitterHandleGuid)
+			{
+				return FText::FromName(Handle.GetName());
+			}
+		}
+		ensureMsgf(false, TEXT("Failed to find matching emitter handle for existing overview node!"));
+		return LOCTEXT("UnknownEmitterName", "Unknown Emitter");
+	}
+	else
+	{
+		return FText::FromString(OwningSystem->GetName());
+	}
+}
+
+bool UNiagaraOverviewNodeStackItem::CanUserDeleteNode() const
+{
+	return EmitterHandleGuid.IsValid();
+}
+
+bool UNiagaraOverviewNodeStackItem::CanDuplicateNode() const
+{
+	return EmitterHandleGuid.IsValid();
+}
+
+#undef LOCTEXT_NAMESPACE
diff --git a/Engine/Plugins/FX/Niagara/Source/NiagaraEditor/Private/NiagaraSystemEditorData.cpp b/Engine/Plugins/FX/Niagara/Source/NiagaraEditor/Private/NiagaraSystemEditorData.cpp
--- a/Engine/Plugins/FX/Niagara/Source/NiagaraEditor/Private/NiagaraSystemEditorData.cpp
+++ b/Engine/Plugins/FX/Niagara/Source/NiagaraEditor/Private/NiagaraSystemEditorData.cpp
@@ -7,6 +7,10 @@
 #include "NiagaraEmitterHandle.h"
 #include "NiagaraEmitter.h"
 #include "NiagaraEmitterEditorData.h"
+#include "NiagaraOverviewNodeStackItem.h"
+#include "EdGraphSchema_NiagaraSystemOverview.h"
+
+static const float SystemOverviewNodePadding = 300.0f;
 
 const FName UNiagaraSystemEditorFolder::GetFolderName() const
 {
@@ -61,10 +65,18 @@ UNiagaraSystemEditorData::UNiagaraSystemEditorData(const FObjectInitializer& Obj
 	PlaybackRangeMax = 10;
 }
 
-void UNiagaraSystemEditorData::PostLoadFromOwner(UObject* InOwner)
+void UNiagaraSystemEditorData::PostInitProperties()
 {
-	UNiagaraSystem* OwnerSystem = CastChecked<UNiagaraSystem>(InOwner);
+	Super::PostInitProperties();
+	if (HasAnyFlags(RF_ClassDefaultObject | RF_NeedLoad) == false)
+	{
+		SystemOverviewGraph = NewObject<UEdGraph>(this, "NiagaraOverview", RF_Transactional);
+		SystemOverviewGraph->Schema = UEdGraphSchema_NiagaraSystemOverview::StaticClass();
+	}
+}
 
+void UNiagaraSystemEditorData::PostLoadFromOwner(UObject* InOwner)
+{
 	if (RootFolder == nullptr)
 	{
 		RootFolder = NewObject<UNiagaraSystemEditorFolder>(this, TEXT("RootFolder"), RF_Transactional);
@@ -78,7 +90,13 @@ void UNiagaraSystemEditorData::PostLoadFromOwner(UObject* InOwner)
 
 	if (NiagaraVer < FNiagaraCustomVersion::PlaybackRangeStoredOnSystem)
 	{
-		UpdatePlaybackRangeFromEmitters(OwnerSystem);
+		UpdatePlaybackRangeFromEmitters(OwningSystem);
+	}
+
+	if (SystemOverviewGraph == nullptr)
+	{
+		SystemOverviewGraph = NewObject<UEdGraph>(this, "NiagaraOverview", RF_Transactional);
+		SystemOverviewGraph->Schema = UEdGraphSchema_NiagaraSystemOverview::StaticClass();
 	}
 }
 
@@ -103,6 +121,11 @@ void UNiagaraSystemEditorData::SetPlaybackRange(TRange<float> InPlaybackRange)
 	PlaybackRangeMax = InPlaybackRange.GetUpperBoundValue();
 }
 
+UEdGraph* UNiagaraSystemEditorData::GetSystemOverviewGraph() const
+{
+	return SystemOverviewGraph;	
+}
+
 void UNiagaraSystemEditorData::UpdatePlaybackRangeFromEmitters(UNiagaraSystem* OwnerSystem)
 {
 	if (OwnerSystem->GetEmitterHandles().Num() > 0)
@@ -123,4 +146,107 @@ void UNiagaraSystemEditorData::UpdatePlaybackRangeFromEmitters(UNiagaraSystem* O
 		PlaybackRangeMin = EmitterPlaybackRangeMin;
 		PlaybackRangeMax = EmitterPlaybackRangeMax;
 	}
-}
\ No newline at end of file
+}
+
+void UNiagaraSystemEditorData::InitSystemOverviewGraph()
+{
+	float NextNodePosX = 0.0;
+
+	FGraphNodeCreator<UNiagaraOverviewNodeStackItem> SystemOverviewNodeCreator(*SystemOverviewGraph);
+	UNiagaraOverviewNodeStackItem* SystemOverviewNode = SystemOverviewNodeCreator.CreateNode(false);
+	SystemOverviewNode->Initialize(OwningSystem);
+
+	SystemOverviewNode->NodePosX = NextNodePosX;
+	SystemOverviewNode->NodePosY = 0.0;
+	NextNodePosX += SystemOverviewNodePadding;
+
+	SystemOverviewNodeCreator.Finalize();
+
+	for (const FNiagaraEmitterHandle& Handle : OwningSystem->GetEmitterHandles())
+	{
+		FGraphNodeCreator<UNiagaraOverviewNodeStackItem> EmitterOverviewNodeCreator(*SystemOverviewGraph);
+		UNiagaraOverviewNodeStackItem* EmitterOverviewNode = EmitterOverviewNodeCreator.CreateNode(false);
+		EmitterOverviewNode->Initialize(OwningSystem, Handle.GetId());
+
+		EmitterOverviewNode->NodePosX = NextNodePosX;
+		EmitterOverviewNode->NodePosY = 0.0;
+		NextNodePosX += SystemOverviewNodePadding;
+
+		EmitterOverviewNodeCreator.Finalize();
+	}
+}
+
+const FVector2D UNiagaraSystemEditorData::GetGoodPlaceForNewOverviewNode() const
+{
+	TArray<UNiagaraOverviewNodeStackItem*> OverviewNodes;
+	SystemOverviewGraph->GetNodesOfClass<UNiagaraOverviewNodeStackItem>(OverviewNodes);
+	const TArray<FNiagaraEmitterHandle>& EmitterHandles = OwningSystem->GetEmitterHandles();
+
+	// Find the last created Emitter overview node location and place next to that. 
+	if (EmitterHandles.Num() > 1)
+	{
+		const FGuid PreviousEmitterHandleGuid = EmitterHandles.Last(1).GetId();
+		for (const UNiagaraOverviewNodeStackItem* PerOverviewNode : OverviewNodes)
+		{
+			if (PreviousEmitterHandleGuid == PerOverviewNode->GetEmitterHandleGuid())
+			{
+				return FVector2D(PerOverviewNode->NodePosX + SystemOverviewNodePadding, PerOverviewNode->NodePosY);
+			}
+		}
+	}
+	return SystemOverviewGraph->GetGoodPlaceForNewNode();
+}
+
+bool UNiagaraSystemEditorData::GetSystemOverviewGraphIsValid() const
+{
+	TArray<UEdGraphNode*> Nodes;
+	SystemOverviewGraph->GetNodesOfClass(Nodes);
+	if (Nodes.Num() > 0)
+	{
+		return true;
+	}
+	return false;
+}
+
+void UNiagaraSystemEditorData::Initialize(UNiagaraSystem* OwnerSystem, bool bEditingSystem)
+{
+	OwningSystem = OwnerSystem;
+
+	if (bEditingSystem && GetSystemOverviewGraphIsValid() == false)
+	{
+		InitSystemOverviewGraph();
+	}
+}
+
+void UNiagaraSystemEditorData::SystemOverviewHandleAdded(const FGuid AddedHandleGuid) const
+{
+	SystemOverviewGraph->Modify();
+
+	const FVector2D NewNodeLocation = GetGoodPlaceForNewOverviewNode();
+
+	FGraphNodeCreator<UNiagaraOverviewNodeStackItem> EmitterOverviewNodeCreator(*SystemOverviewGraph);
+	UNiagaraOverviewNodeStackItem* EmitterOverviewNode = EmitterOverviewNodeCreator.CreateNode(false);
+	EmitterOverviewNode->Initialize(OwningSystem, AddedHandleGuid);
+
+	EmitterOverviewNode->NodePosX = NewNodeLocation.X;
+	EmitterOverviewNode->NodePosY = NewNodeLocation.Y;
+
+	EmitterOverviewNodeCreator.Finalize();
+
+	//@TODO System Overview: if we have a ref to the SGraphEditor here, then focus the new node
+}
+
+void UNiagaraSystemEditorData::SystemOverviewHandlesRemoved() const
+{
+	TArray<UNiagaraOverviewNodeStackItem*> CurrentOverviewNodes;
+	SystemOverviewGraph->GetNodesOfClass<UNiagaraOverviewNodeStackItem>(CurrentOverviewNodes);
+	for (UNiagaraOverviewNodeStackItem* OverviewNode : CurrentOverviewNodes)
+	{
+		const FGuid OverviewNodeGuid = OverviewNode->GetEmitterHandleGuid();
+		// If the OverviewNode's EmitterHandleGuid is valid (not representing a system) and is not in the current array of the owning System's EmitterHandles, that node can be removed.
+		if (OverviewNodeGuid.IsValid() && false == OwningSystem->GetEmitterHandles().ContainsByPredicate([&OverviewNodeGuid](const FNiagaraEmitterHandle& Handle) {return Handle.GetId() == OverviewNodeGuid; }))
+		{
+			SystemOverviewGraph->RemoveNode(OverviewNode);
+		}
+	}
+}
diff --git a/Engine/Plugins/FX/Niagara/Source/NiagaraEditor/Private/NiagaraSystemEditorData.h b/Engine/Plugins/FX/Niagara/Source/NiagaraEditor/Private/NiagaraSystemEditorData.h
deleted file mode 100644
diff --git a/Engine/Plugins/FX/Niagara/Source/NiagaraEditor/Private/NiagaraSystemFactoryNew.cpp b/Engine/Plugins/FX/Niagara/Source/NiagaraEditor/Private/NiagaraSystemFactoryNew.cpp
--- a/Engine/Plugins/FX/Niagara/Source/NiagaraEditor/Private/NiagaraSystemFactoryNew.cpp
+++ b/Engine/Plugins/FX/Niagara/Source/NiagaraEditor/Private/NiagaraSystemFactoryNew.cpp
@@ -15,6 +15,7 @@
 #include "Modules/ModuleManager.h"
 #include "Interfaces/IMainFrameModule.h"
 #include "Framework/Application/SlateApplication.h"
+#include "NiagaraEditorUtilities.h"
 
 #define LOCTEXT_NAMESPACE "NiagaraSystemFactory"
 
@@ -129,15 +130,9 @@ UObject* UNiagaraSystemFactoryNew::FactoryCreateNew(UClass* Class, UObject* InPa
 		NewSystem = NewObject<UNiagaraSystem>(InParent, Class, Name, Flags | RF_Transactional);
 		InitializeSystem(NewSystem, true);
 
-		FNiagaraSystemViewModelOptions SystemViewModelOptions;
-		SystemViewModelOptions.bCanAutoCompile = false;
-		SystemViewModelOptions.bCanSimulate = false;
-		SystemViewModelOptions.EditMode = ENiagaraSystemViewModelEditMode::SystemAsset;
-
-		TSharedRef<FNiagaraSystemViewModel> NewSystemViewModel = MakeShared<FNiagaraSystemViewModel>(*NewSystem, SystemViewModelOptions);
 		for (UNiagaraEmitter* EmitterToAddToNewSystem : EmittersToAddToNewSystem)
 		{
-			NewSystemViewModel->AddEmitter(*EmitterToAddToNewSystem);
+			FNiagaraEditorUtilities::AddEmitterToSystem(*NewSystem, *EmitterToAddToNewSystem, false);
 		}
 	}
 	else
diff --git a/Engine/Plugins/FX/Niagara/Source/NiagaraEditor/Private/Toolkits/NiagaraSystemToolkit.cpp b/Engine/Plugins/FX/Niagara/Source/NiagaraEditor/Private/Toolkits/NiagaraSystemToolkit.cpp
--- a/Engine/Plugins/FX/Niagara/Source/NiagaraEditor/Private/Toolkits/NiagaraSystemToolkit.cpp
+++ b/Engine/Plugins/FX/Niagara/Source/NiagaraEditor/Private/Toolkits/NiagaraSystemToolkit.cpp
@@ -49,6 +49,7 @@
 #include "HAL/PlatformFilemanager.h"
 #include "Misc/FileHelper.h"
 #include "NiagaraMessageLogViewModel.h"
+#include "Viewmodels/NiagaraOverviewGraphViewModel.h"
 
 #define LOCTEXT_NAMESPACE "NiagaraSystemEditor"
 
@@ -286,6 +287,8 @@ void FNiagaraSystemToolkit::InitializeWithEmitter(const EToolkitMode::Type Mode,
 
 void FNiagaraSystemToolkit::InitializeInternal(const EToolkitMode::Type Mode, const TSharedPtr<IToolkitHost>& InitToolkitHost, const FGuid& MessageLogGuid)
 {
+	SystemViewModel->SetOverviewGraphViewModel(MakeShareable(new FNiagaraOverviewGraphViewModel(SystemViewModel->AsShared()))); //@TODO System Overview: move constructing system overview viewmodel to SystemViewModel Init
+
 	if (SystemViewModel->GetEmitterHandleViewModels().Num() > 0)
 	{
 		SystemViewModel->SetSelectedEmitterHandleById(SystemViewModel->GetEmitterHandleViewModels()[0]->GetId());
diff --git a/Engine/Plugins/FX/Niagara/Source/NiagaraEditor/Private/ViewModels/NiagaraOverviewGraphViewModel.cpp b/Engine/Plugins/FX/Niagara/Source/NiagaraEditor/Private/ViewModels/NiagaraOverviewGraphViewModel.cpp
new file mode 100644
--- /dev/null
+++ b/Engine/Plugins/FX/Niagara/Source/NiagaraEditor/Private/ViewModels/NiagaraOverviewGraphViewModel.cpp
@@ -0,0 +1,321 @@
+// Copyright 1998-2019 Epic Games, Inc. All Rights Reserved.
+
+#include "ViewModels\NiagaraOverviewGraphViewModel.h"
+#include "Framework/Commands/GenericCommands.h"
+#include "ScopedTransaction.h"
+#include "EdGraphUtilities.h"
+#include "HAL/PlatformApplicationMisc.h"
+#include "NiagaraSystemEditorData.h"
+#include "NiagaraOverviewNodeStackItem.h"
+#include "NiagaraObjectSelection.h"
+#include "ViewModels/NiagaraSystemViewModel.h"
+#include "ViewModels/NiagaraEmitterHandleViewModel.h"
+#include "NiagaraSystem.h"
+
+#define LOCTEXT_NAMESPACE "NiagaraOverviewGraphViewModel"
+
+FNiagaraOverviewGraphViewModel::FNiagaraOverviewGraphViewModel(TSharedRef<FNiagaraSystemViewModel> InSystemViewModel)
+	: SystemViewModel(InSystemViewModel)
+	, Commands(MakeShareable(new FUICommandList()))
+	, NodeSelection(MakeShareable(new FNiagaraObjectSelection()))
+{
+	SetupCommands();
+	InitDisplayName();
+	GEditor->RegisterForUndo(this);
+}
+
+FNiagaraOverviewGraphViewModel::~FNiagaraOverviewGraphViewModel()
+{
+	GEditor->UnregisterForUndo(this);
+}
+
+FText FNiagaraOverviewGraphViewModel::GetDisplayName() const
+{
+	return DisplayName;
+}
+
+UEdGraph* FNiagaraOverviewGraphViewModel::GetGraph() const
+{
+	if (SystemViewModel.IsValid())
+	{
+		return SystemViewModel.Pin()->GetEditorData().GetSystemOverviewGraph();
+	}
+	return nullptr;
+}
+
+TSharedRef<FUICommandList> FNiagaraOverviewGraphViewModel::GetCommands()
+{
+	return Commands;
+}
+
+TSharedRef<FNiagaraObjectSelection> FNiagaraOverviewGraphViewModel::GetNodeSelection()
+{
+	return NodeSelection;
+}
+
+FNiagaraOverviewGraphViewModel::FOnNodesPasted& FNiagaraOverviewGraphViewModel::OnNodesPasted()
+{
+	return OnNodesPastedDelegate;
+}
+
+void FNiagaraOverviewGraphViewModel::PostUndo(bool bSuccess)
+{
+	// The graph may have been deleted as a result of the undo operation so make sure it's valid
+	// before using it.
+	UEdGraph* Graph = GetGraph();
+	if (Graph != nullptr)
+	{
+		Graph->NotifyGraphChanged();
+	}
+}
+
+void FNiagaraOverviewGraphViewModel::SetupCommands()
+{
+	Commands->MapAction(
+		FGenericCommands::Get().SelectAll,
+		FExecuteAction::CreateRaw(this, &FNiagaraOverviewGraphViewModel::SelectAllNodes));
+
+	Commands->MapAction(
+		FGenericCommands::Get().Delete,
+		FExecuteAction::CreateRaw(this, &FNiagaraOverviewGraphViewModel::DeleteSelectedNodes),
+		FCanExecuteAction::CreateRaw(this, &FNiagaraOverviewGraphViewModel::CanDeleteNodes));
+
+	Commands->MapAction(
+		FGenericCommands::Get().Copy,
+		FExecuteAction::CreateRaw(this, &FNiagaraOverviewGraphViewModel::CopySelectedNodes),
+		FCanExecuteAction::CreateRaw(this, &FNiagaraOverviewGraphViewModel::CanCopyNodes));
+
+	Commands->MapAction(
+		FGenericCommands::Get().Cut,
+		FExecuteAction::CreateRaw(this, &FNiagaraOverviewGraphViewModel::CutSelectedNodes),
+		FCanExecuteAction::CreateRaw(this, &FNiagaraOverviewGraphViewModel::CanCutNodes));
+
+	Commands->MapAction(
+		FGenericCommands::Get().Paste,
+		FExecuteAction::CreateRaw(this, &FNiagaraOverviewGraphViewModel::PasteNodes),
+		FCanExecuteAction::CreateRaw(this, &FNiagaraOverviewGraphViewModel::CanPasteNodes));
+
+	Commands->MapAction(
+		FGenericCommands::Get().Duplicate,
+		FExecuteAction::CreateRaw(this, &FNiagaraOverviewGraphViewModel::DuplicateNodes),
+		FCanExecuteAction::CreateRaw(this, &FNiagaraOverviewGraphViewModel::CanDuplicateNodes));
+}
+
+void FNiagaraOverviewGraphViewModel::SelectAllNodes()
+{
+	UEdGraph* Graph = GetGraph();
+	if (Graph != nullptr)
+	{
+		TArray<UObject*> AllNodes;
+		Graph->GetNodesOfClass<UObject>(AllNodes);
+		TSet<UObject*> AllNodeSet;
+		AllNodeSet.Append(AllNodes);
+		NodeSelection->SetSelectedObjects(AllNodeSet);
+	}
+}
+
+void FNiagaraOverviewGraphViewModel::DeleteSelectedNodes()
+{
+	UEdGraph* Graph = GetGraph();
+	if (Graph != nullptr)
+	{
+		const FScopedTransaction Transaction(FGenericCommands::Get().Delete->GetDescription());
+		Graph->Modify();
+
+		TArray<UObject*> NodesToDelete = NodeSelection->GetSelectedObjects().Array();
+		NodeSelection->ClearSelectedObjects();
+
+		TSet<FGuid> EmitterGuidsToDelete;
+		for (UObject* NodeToDelete : NodesToDelete)
+		{
+			UNiagaraOverviewNodeStackItem* GraphNodeToDelete = Cast<UNiagaraOverviewNodeStackItem>(NodeToDelete);
+			if (GraphNodeToDelete != nullptr && GraphNodeToDelete->CanUserDeleteNode())
+			{
+				EmitterGuidsToDelete.Add(GraphNodeToDelete->GetEmitterHandleGuid());
+			}
+		}
+		//we have checked SystemViewModel is valid as this is a requisite for Graph to not be null.
+		SystemViewModel.Pin()->DeleteEmitters(EmitterGuidsToDelete);
+	}
+}
+
+bool FNiagaraOverviewGraphViewModel::CanDeleteNodes() const
+{
+	UEdGraph* Graph = GetGraph();
+	if (Graph != nullptr)
+	{
+		for (UObject* SelectedNode : NodeSelection->GetSelectedObjects())
+		{
+			UEdGraphNode* SelectedGraphNode = Cast<UEdGraphNode>(SelectedNode);
+			if (SelectedGraphNode != nullptr && SelectedGraphNode->CanUserDeleteNode())
+			{
+				return true;
+			}
+		}
+	}
+	return false;
+}
+
+void FNiagaraOverviewGraphViewModel::CutSelectedNodes()
+{
+	// Collect nodes which can not be delete or duplicated so they can be reselected.
+	TSet<UObject*> CanBeDuplicatedAndDeleted;
+	TSet<UObject*> CanNotBeDuplicatedAndDeleted;
+	for (UObject* SelectedNode : NodeSelection->GetSelectedObjects())
+	{
+		UEdGraphNode* SelectedGraphNode = Cast<UEdGraphNode>(SelectedNode);
+		if (SelectedGraphNode != nullptr)
+		{
+			if (SelectedGraphNode->CanDuplicateNode() && SelectedGraphNode->CanUserDeleteNode())
+			{
+				CanBeDuplicatedAndDeleted.Add(SelectedNode);
+			}
+			else
+			{
+				CanNotBeDuplicatedAndDeleted.Add(SelectedNode);
+			}
+		}
+	}
+
+	// Select the nodes which can be copied and deleted, copy and delete them, and then restore the ones which couldn't be copied or deleted.
+	NodeSelection->SetSelectedObjects(CanBeDuplicatedAndDeleted);
+	CopySelectedNodes();
+	DeleteSelectedNodes();
+	NodeSelection->SetSelectedObjects(CanNotBeDuplicatedAndDeleted);
+}
+
+bool FNiagaraOverviewGraphViewModel::CanCutNodes() const
+{
+	return CanCopyNodes() && CanDeleteNodes();
+}
+
+void FNiagaraOverviewGraphViewModel::CopySelectedNodes()
+{
+	TSet<UObject*> NodesToCopy;
+	for (UObject* SelectedNode : NodeSelection->GetSelectedObjects())
+	{
+		UEdGraphNode* SelectedGraphNode = Cast<UEdGraphNode>(SelectedNode);
+		if (SelectedGraphNode != nullptr)
+		{
+			if (SelectedGraphNode->CanDuplicateNode())
+			{
+				SelectedGraphNode->PrepareForCopying();
+				NodesToCopy.Add(SelectedNode);
+			}
+		}
+	}
+
+	FString ExportedText;
+	FEdGraphUtilities::ExportNodesToText(NodesToCopy, ExportedText);
+	FPlatformApplicationMisc::ClipboardCopy(*ExportedText);
+}
+
+bool FNiagaraOverviewGraphViewModel::CanCopyNodes() const
+{
+	UEdGraph* Graph = GetGraph();
+	if (Graph != nullptr)
+	{
+		for (UObject* SelectedNode : NodeSelection->GetSelectedObjects())
+		{
+			UEdGraphNode* SelectedGraphNode = Cast<UEdGraphNode>(SelectedNode);
+			if (SelectedGraphNode != nullptr && SelectedGraphNode->CanDuplicateNode())
+			{
+				return true;
+			}
+		}
+	}
+	return false;
+}
+
+
+void FNiagaraOverviewGraphViewModel::PasteNodes()
+{
+	UEdGraph* Graph = GetGraph();
+	if (Graph != nullptr)
+	{
+		const FScopedTransaction Transaction(FGenericCommands::Get().Paste->GetDescription());;
+		Graph->Modify();
+
+		NodeSelection->ClearSelectedObjects();
+
+		// Grab the text to paste from the clipboard.
+		FString TextToImport;
+		FPlatformApplicationMisc::ClipboardPaste(TextToImport);
+
+		// Import the nodes
+		TSet<UEdGraphNode*> PastedNodes;
+		FEdGraphUtilities::ImportNodesFromText(Graph, TextToImport, PastedNodes);
+
+// 		for (UEdGraphNode* PastedNode : PastedNodes) //@TODO System Overview: fix this, copy the impl from sequencer impl in system toolkit
+// 		{
+// 			PastedNode->CreateNewGuid();
+// 			UNiagaraNode* Node = Cast<UNiagaraNode>(PastedNode);
+// 			if (Node)
+// 				Node->MarkNodeRequiresSynchronization(__FUNCTION__, false);
+// 		}
+
+		OnNodesPastedDelegate.Broadcast(PastedNodes);
+
+		TSet<UObject*> PastedObjects;
+		for (UEdGraphNode* PastedNode : PastedNodes)
+		{
+			PastedObjects.Add(PastedNode);
+		}
+
+		NodeSelection->SetSelectedObjects(PastedObjects);
+		Graph->NotifyGraphChanged(); //@TODO System Overview: might not be necessary 
+	}
+}
+
+bool FNiagaraOverviewGraphViewModel::CanPasteNodes() const
+{
+	UEdGraph* Graph = GetGraph();
+	if (Graph == nullptr)
+	{
+		return false;
+	}
+
+	FString ClipboardContent;
+	FPlatformApplicationMisc::ClipboardPaste(ClipboardContent);
+
+	return FEdGraphUtilities::CanImportNodesFromText(Graph, ClipboardContent);
+}
+
+void FNiagaraOverviewGraphViewModel::DuplicateNodes()
+{
+	CopySelectedNodes();
+	PasteNodes();
+}
+
+bool FNiagaraOverviewGraphViewModel::CanDuplicateNodes() const
+{
+	return CanCopyNodes();
+}
+
+void FNiagaraOverviewGraphViewModel::InitDisplayName()
+{
+	ensureMsgf(SystemViewModel.IsValid(), TEXT("SystemViewModel was not initialized before InitDisplayName!"));
+	TSharedPtr<FNiagaraSystemViewModel> WeakSystemViewModel = SystemViewModel.Pin();
+	ENiagaraSystemViewModelEditMode SystemEditMode = WeakSystemViewModel->GetEditMode();
+	if (SystemEditMode == ENiagaraSystemViewModelEditMode::EmitterAsset)
+	{
+		TArray<TSharedRef<FNiagaraEmitterHandleViewModel>> EditedEmitterHandleViewModels = WeakSystemViewModel->GetEmitterHandleViewModels();
+		if (ensureMsgf(EditedEmitterHandleViewModels.Num() > 0, TEXT("SystemViewModel did not have any EmitterHandleViewModels! Cannot get currently edited Emitter's Name.")))
+		{
+			DisplayName = FText::FromName(WeakSystemViewModel->GetEmitterHandleViewModels()[0]->GetEmitterHandle()->GetName());
+			return;
+		}
+	}
+	else if (SystemEditMode == ENiagaraSystemViewModelEditMode::SystemAsset)
+	{
+		DisplayName = FText::FromString(WeakSystemViewModel->GetSystem().GetName());
+		return;
+	}
+	else
+	{
+		ensureMsgf(false, TEXT("Encountered unexpected SystemViewModel edit mode!"));
+	}
+	DisplayName = LOCTEXT("NiagaraOverview_FallbackTitlebar", "Niagara");
+}
+
+#undef LOCTEXT_NAMESPACE // NiagaraScriptGraphViewModel
diff --git a/Engine/Plugins/FX/Niagara/Source/NiagaraEditor/Private/ViewModels/NiagaraSystemViewModel.cpp b/Engine/Plugins/FX/Niagara/Source/NiagaraEditor/Private/ViewModels/NiagaraSystemViewModel.cpp
--- a/Engine/Plugins/FX/Niagara/Source/NiagaraEditor/Private/ViewModels/NiagaraSystemViewModel.cpp
+++ b/Engine/Plugins/FX/Niagara/Source/NiagaraEditor/Private/ViewModels/NiagaraSystemViewModel.cpp
@@ -36,6 +36,7 @@
 #include "Channels/MovieSceneChannelProxy.h"
 #include "Framework/Commands/UICommandList.h"
 #include "NiagaraMessageManager.h"
+#include "Viewmodels/NiagaraOverviewGraphViewModel.h"
 
 DECLARE_CYCLE_STAT(TEXT("Niagara - SystemViewModel - CompileSystem"), STAT_NiagaraEditor_SystemViewModel_CompileSystem, STATGROUP_NiagaraEditor);
 
@@ -69,7 +70,7 @@ FNiagaraSystemViewModel::FNiagaraSystemViewModel(UNiagaraSystem& InSystem, FNiag
 	, EditorSettings(GetMutableDefault<UNiagaraEditorSettings>())
 	, bResetRequestPending(false)
 	, bCompilePendingCompletion(false)
-	,SystemMessageLogGuidKey(InMessageLogGuid)
+	, SystemMessageLogGuidKey(InMessageLogGuid)
 {
 	SetupPreviewComponentAndInstance();
 	SetupSequencer();
@@ -165,7 +166,7 @@ TSharedPtr<FNiagaraSystemScriptViewModel> FNiagaraSystemViewModel::GetSystemScri
 void FNiagaraSystemViewModel::CompileSystem(bool bForce)
 {
 	SCOPE_CYCLE_COUNTER(STAT_NiagaraEditor_SystemViewModel_CompileSystem);
-	KillSystemInstances();
+	FNiagaraEditorUtilities::KillSystemInstances(System);
 	check(SystemScriptViewModel.IsValid());
 	SystemScriptViewModel->CompileSystem(bForce);
 	bCompilePendingCompletion = true;
@@ -250,6 +251,8 @@ UNiagaraSystemEditorData& FNiagaraSystemViewModel::GetOrCreateEditorData()
 		System.Modify();
 		System.SetEditorData(EditorData);
 	}
+	bool bEditingSystem = GetEditMode() == ENiagaraSystemViewModelEditMode::SystemAsset;
+	EditorData->Initialize(&System, bEditingSystem);
 	return *EditorData;
 }
 
@@ -290,41 +293,18 @@ void FNiagaraSystemViewModel::AddEmitterFromAssetData(const FAssetData& AssetDat
 
 void FNiagaraSystemViewModel::AddEmitter(UNiagaraEmitter& Emitter)
 {
-	// Kill all system instances before modifying the emitter handle list to prevent accessing deleted data.
-	KillSystemInstances();
-
-	// When editing an emitter asset the system is a placeholder and we don't want to make adding an emitter to it
-	// undoable.
-	if (EditMode != ENiagaraSystemViewModelEditMode::EmitterAsset)
+	// When editing an emitter asset the system is a placeholder and we don't want to make adding an emitter to it undoable.
+	bool bSystemIsPlaceholder = EditMode == ENiagaraSystemViewModelEditMode::EmitterAsset;
+	if (false == bSystemIsPlaceholder)
 	{
 		GEditor->BeginTransaction(LOCTEXT("AddEmitter", "Add emitter"));
 	}
 
-	TSet<FName> EmitterHandleNames;
-	for (const FNiagaraEmitterHandle& EmitterHandle : System.GetEmitterHandles())
-	{
-		EmitterHandleNames.Add(EmitterHandle.GetName());
-	}
-
-	FNiagaraEmitterHandle EmitterHandle;
-	if (EditMode == ENiagaraSystemViewModelEditMode::SystemAsset)
-	{
-		System.Modify();
-		EmitterHandle = System.AddEmitterHandle(Emitter, FNiagaraUtilities::GetUniqueName(Emitter.GetFName(), EmitterHandleNames));
-	}
-	else if (EditMode == ENiagaraSystemViewModelEditMode::EmitterAsset)
-	{
-		// When editing an emitter asset we add the emitter as a duplicate so that the parent emitter is duplicated, but it's parent emitter
-		// information is maintained.
-		checkf(System.GetNumEmitters() == 0, TEXT("Can not add multiple emitters to a system being edited in emitter asset mode."));
-		FNiagaraEmitterHandle TemporaryEmitterHandle(Emitter);
-		EmitterHandle = System.DuplicateEmitterHandle(TemporaryEmitterHandle, *Emitter.GetUniqueEmitterName());
-	}
-	
-	check(SystemScriptViewModel.IsValid());
-	FNiagaraStackGraphUtilities::RebuildEmitterNodes(System);
+	const FGuid NewEmitterHandleId = FNiagaraEditorUtilities::AddEmitterToSystem(System, Emitter, bSystemIsPlaceholder);
+	UNiagaraSystemEditorData& EditorData = GetOrCreateEditorData();
+	EditorData.SystemOverviewHandleAdded(NewEmitterHandleId);
 
-	if (EditMode == ENiagaraSystemViewModelEditMode::SystemAsset)
+	if (false == bSystemIsPlaceholder)
 	{
 		GEditor->EndTransaction();
 	}
@@ -334,17 +314,16 @@ void FNiagaraSystemViewModel::AddEmitter(UNiagaraEmitter& Emitter)
 		// When adding a new emitter to an empty system start playing.
 		Sequencer->SetPlaybackStatus(EMovieScenePlayerStatus::Playing);
 	}
-		
+
 	RefreshAll();
 
-	TRange<float> SystemPlaybackRange = GetEditorData().GetPlaybackRange();
-	TRange<float> EmitterPlaybackRange = GetEmitterHandleViewModelById(EmitterHandle.GetId())->GetEmitterViewModel()->GetEditorData().GetPlaybackRange();
+	TRange<float> SystemPlaybackRange = EditorData.GetPlaybackRange();
+	TRange<float> EmitterPlaybackRange = GetEmitterHandleViewModelById(NewEmitterHandleId)->GetEmitterViewModel()->GetEditorData().GetPlaybackRange();
 	TRange<float> NewSystemPlaybackRange = TRange<float>(
 		FMath::Min(SystemPlaybackRange.GetLowerBoundValue(), EmitterPlaybackRange.GetLowerBoundValue()),
 		FMath::Max(SystemPlaybackRange.GetUpperBoundValue(), EmitterPlaybackRange.GetUpperBoundValue()));
 
-	GetOrCreateEditorData().Modify();
-	GetOrCreateEditorData().SetPlaybackRange(NewSystemPlaybackRange);
+	EditorData.SetPlaybackRange(NewSystemPlaybackRange);
 
 	TGuardValue<bool> UpdateGuard(bUpdatingSequencerFromEmitterDataChange, true);
 
@@ -352,8 +331,8 @@ void FNiagaraSystemViewModel::AddEmitter(UNiagaraEmitter& Emitter)
 	int32 NewDuration = (NewSystemPlaybackRange.Size<float>() * NiagaraSequence->GetMovieScene()->GetTickResolution()).FrameNumber.Value;
 
 	NiagaraSequence->GetMovieScene()->SetPlaybackRange(NewStartFrame.RoundToFrame(), NewDuration);
-		
-	SetSelectedEmitterHandleById(EmitterHandle.GetId());
+
+	SetSelectedEmitterHandleById(NewEmitterHandleId);
 
 	bForceAutoCompileOnce = true;
 }
@@ -366,7 +345,7 @@ void FNiagaraSystemViewModel::DuplicateEmitters(TArray<FEmitterHandleToDuplicate
 	}
 
 	// Kill all system instances before modifying the emitter handle list to prevent accessing deleted data.
-	KillSystemInstances();
+	FNiagaraEditorUtilities::KillSystemInstances(System);
 	const FScopedTransaction DeleteTransaction(EmitterHandlesToDuplicate.Num() == 1
 		? LOCTEXT("DuplicateEmitter", "Duplicate emitter")
 		: LOCTEXT("DuplicateEmitters", "Duplicate emitters"));
@@ -403,8 +382,9 @@ void FNiagaraSystemViewModel::DuplicateEmitters(TArray<FEmitterHandleToDuplicate
 
 		if (HandleToDuplicate.IsValid())
 		{
-			FNiagaraEmitterHandle EmitterHandle = System.DuplicateEmitterHandle(HandleToDuplicate, FNiagaraUtilities::GetUniqueName(HandleToDuplicate.GetName(), EmitterHandleNames));
+			const FNiagaraEmitterHandle& EmitterHandle = System.DuplicateEmitterHandle(HandleToDuplicate, FNiagaraUtilities::GetUniqueName(HandleToDuplicate.GetName(), EmitterHandleNames));
 			EmitterHandleNames.Add(EmitterHandle.GetName());
+			GetEditorData().SystemOverviewHandleAdded(EmitterHandle.GetId());
 		}
 	}
 
@@ -413,30 +393,12 @@ void FNiagaraSystemViewModel::DuplicateEmitters(TArray<FEmitterHandleToDuplicate
 	bForceAutoCompileOnce = true;
 }
 
-void FNiagaraSystemViewModel::DeleteEmitter(TSharedRef<FNiagaraEmitterHandleViewModel> EmitterHandleToDelete)
-{
-	TSet<FGuid> IdsToDelete;
-	IdsToDelete.Add(EmitterHandleToDelete->GetId());
-	DeleteEmitters(IdsToDelete);
-	bForceAutoCompileOnce = true;
-}
-
 void FNiagaraSystemViewModel::DeleteEmitters(TSet<FGuid> EmitterHandleIdsToDelete)
 {
 	if (EmitterHandleIdsToDelete.Num() > 0)
 	{
-		// Kill all system instances before modifying the emitter handle list to prevent accessing deleted data.
-		KillSystemInstances();
-
-		const FScopedTransaction DeleteTransaction(EmitterHandleIdsToDelete.Num() == 1 
-			? LOCTEXT("DeleteEmitter", "Delete emitter")
-			: LOCTEXT("DeleteEmitters", "Delete emitters"));
-
-		System.Modify();
-		System.RemoveEmitterHandlesById(EmitterHandleIdsToDelete);
-
-		check(SystemScriptViewModel.IsValid());
-		FNiagaraStackGraphUtilities::RebuildEmitterNodes(System);
+		FNiagaraEditorUtilities::RemoveEmittersFromSystemByEmitterHandleId(System, EmitterHandleIdsToDelete);
+		GetEditorData().SystemOverviewHandlesRemoved();
 
 		RefreshAll();
 		bForceAutoCompileOnce = true;
@@ -1048,15 +1010,6 @@ void FNiagaraSystemViewModel::RequestResetSystem()
 	bResetRequestPending = true;
 }
 
-void FNiagaraSystemViewModel::KillSystemInstances()
-{
-	TArray<UNiagaraComponent*> ReferencingComponents = FNiagaraEditorUtilities::GetComponentsThatReferenceSystem(System);
-	for (auto Component : ReferencingComponents)
-	{
-		Component->DestroyInstance();
-	}
-}
-
 struct FNiagaraSystemCurveData
 {
 	FRichCurve* Curve;
@@ -1748,4 +1701,24 @@ void FNiagaraSystemViewModel::BuildStackModuleData(UNiagaraScript* Script, FGuid
 
 }
 
-#undef LOCTEXT_NAMESPACE // NiagaraSystemViewModel
\ No newline at end of file
+TSharedPtr<FNiagaraOverviewGraphViewModel> FNiagaraSystemViewModel::GetOverviewGraphViewModel() const
+{
+	return OverviewGraphViewModel;
+}
+
+void FNiagaraSystemViewModel::SetOverviewGraphViewModel(TSharedRef<FNiagaraOverviewGraphViewModel> InOverviewGraphViewModel)
+{
+	OverviewGraphViewModel = InOverviewGraphViewModel;
+}
+
+void FNiagaraSystemViewModel::EmitterNameChanged(const FText& InText, FGuid ChangedEmitterGuid)
+{
+	TSharedPtr<FNiagaraEmitterHandleViewModel> EmitterHandleViewModel = GetEmitterHandleViewModelById(ChangedEmitterGuid);
+	if (ensureMsgf(EmitterHandleViewModel.IsValid(), TEXT("Failed to find EmitterHandleViewModel with Emitter that has same Guid as Emitter that was renamed!")))
+	{
+		EmitterHandleViewModel->OnNameTextComitted(InText, ETextCommit::Default);
+		OverviewGraphViewModel->GetGraph()->NotifyGraphChanged();
+	}
+}
+
+#undef LOCTEXT_NAMESPACE // NiagaraSystemViewModel
diff --git a/Engine/Plugins/FX/Niagara/Source/NiagaraEditor/Public/EdGraphSchema_NiagaraSystemOverview.h b/Engine/Plugins/FX/Niagara/Source/NiagaraEditor/Public/EdGraphSchema_NiagaraSystemOverview.h
new file mode 100644
--- /dev/null
+++ b/Engine/Plugins/FX/Niagara/Source/NiagaraEditor/Public/EdGraphSchema_NiagaraSystemOverview.h
@@ -0,0 +1,25 @@
+// Copyright 1998-2019 Epic Games, Inc. All Rights Reserved.
+
+#pragma once
+
+#include "CoreMinimal.h"
+#include "EdGraph/EdGraphSchema.h"
+#include "ConnectionDrawingPolicy.h"
+#include "EdGraphSchema_NiagaraSystemOverview.generated.h"
+
+
+
+UCLASS()
+class NIAGARAEDITOR_API UEdGraphSchema_NiagaraSystemOverview : public UEdGraphSchema
+{
+	GENERATED_UCLASS_BODY()
+
+	//~ Begin EdGraphSchema Interface
+	virtual void GetContextMenuActions(const UEdGraph* CurrentGraph, const UEdGraphNode* InGraphNode, const UEdGraphPin* InGraphPin, class FMenuBuilder* MenuBuilder, bool bIsDebugging) const override;
+	virtual const FPinConnectionResponse CanCreateConnection(const UEdGraphPin* A, const UEdGraphPin* B) const override { return FPinConnectionResponse(CONNECT_RESPONSE_DISALLOW, FText()); }; //@TODO System Overview: write text response
+	virtual bool TryCreateConnection(UEdGraphPin* A, UEdGraphPin* B) const override { return false; };
+	virtual bool ShouldAlwaysPurgeOnModification() const override { return false; }
+	//~ End EdGraphSchema Interface
+
+};
+
diff --git a/Engine/Plugins/FX/Niagara/Source/NiagaraEditor/Public/NiagaraEditorUtilities.h b/Engine/Plugins/FX/Niagara/Source/NiagaraEditor/Public/NiagaraEditorUtilities.h
--- a/Engine/Plugins/FX/Niagara/Source/NiagaraEditor/Public/NiagaraEditorUtilities.h
+++ b/Engine/Plugins/FX/Niagara/Source/NiagaraEditor/Public/NiagaraEditorUtilities.h
@@ -189,4 +189,12 @@ namespace FNiagaraEditorUtilities
 	TArray<UNiagaraComponent*> GetComponentsThatReferenceSystem(const UNiagaraSystem& ReferencedSystem);
 
 	TArray<UNiagaraComponent*> GetComponentsThatReferenceSystemViewModel(const FNiagaraSystemViewModel& ReferencedSystemViewModel);
+
+	const FGuid AddEmitterToSystem(UNiagaraSystem& InSystem, UNiagaraEmitter& InEmitterToAdd, bool bSystemIsPlaceholder);
+
+	void RemoveEmittersFromSystemByEmitterHandleId(UNiagaraSystem& InSystem, TSet<FGuid> EmitterHandleIdsToDelete);
+
+	/** Kills all system instances using the referenced system. */
+	void KillSystemInstances(const UNiagaraSystem& System);
+
 };
diff --git a/Engine/Plugins/FX/Niagara/Source/NiagaraEditor/Public/NiagaraOverviewNodeStackItem.h b/Engine/Plugins/FX/Niagara/Source/NiagaraEditor/Public/NiagaraOverviewNodeStackItem.h
new file mode 100644
--- /dev/null
+++ b/Engine/Plugins/FX/Niagara/Source/NiagaraEditor/Public/NiagaraOverviewNodeStackItem.h
@@ -0,0 +1,35 @@
+// Copyright 1998-2019 Epic Games, Inc. All Rights Reserved.
+#pragma once
+
+#include "NiagaraOverviewNodeStackItem.generated.h"
+
+class UNiagaraSystem;
+
+UCLASS()
+class NIAGARAEDITOR_API UNiagaraOverviewNodeStackItem : public UEdGraphNode
+{
+	GENERATED_BODY()
+public:
+	UNiagaraOverviewNodeStackItem();
+	void Initialize(UNiagaraSystem* InOwningSystem);
+	void Initialize(UNiagaraSystem* InOwningSystem, FGuid InEmitterHandleGuid);
+	const FGuid GetEmitterHandleGuid() const;
+
+	//~ Begin UEdGraphNode Interface
+	/** Gets the name of this node, shown in title bar */
+	virtual FText GetNodeTitle(ENodeTitleType::Type TitleType) const override;
+
+	/** Whether or not this node can be deleted by user action */
+	virtual bool CanUserDeleteNode() const override;
+
+	/** Whether or not this node can be safely duplicated (via copy/paste, etc...) in the graph */
+	virtual bool CanDuplicateNode() const override;
+	//~ End UEdGraphNode Interface
+
+private:
+	UPROPERTY()
+	UNiagaraSystem* OwningSystem;
+
+	UPROPERTY()
+	FGuid EmitterHandleGuid;
+};
diff --git a/Engine/Plugins/FX/Niagara/Source/NiagaraEditor/Public/NiagaraSystemEditorData.h b/Engine/Plugins/FX/Niagara/Source/NiagaraEditor/Public/NiagaraSystemEditorData.h
new file mode 100644
--- /dev/null
+++ b/Engine/Plugins/FX/Niagara/Source/NiagaraEditor/Public/NiagaraSystemEditorData.h
@@ -0,0 +1,117 @@
+// Copyright 1998-2019 Epic Games, Inc. All Rights Reserved.
+
+#pragma once
+
+#include "NiagaraEditorDataBase.h"
+#include "NiagaraSystemEditorData.generated.h"
+
+class UNiagaraStackEditorData;
+class UNiagaraSystem;
+
+/** Editor only folder data for emitters in a system. */
+UCLASS()
+class UNiagaraSystemEditorFolder : public UObject
+{
+	GENERATED_BODY()
+
+public:
+	const FName GetFolderName() const;
+
+	void SetFolderName(FName InFolderName);
+
+	const TArray<UNiagaraSystemEditorFolder*>& GetChildFolders() const;
+
+	void AddChildFolder(UNiagaraSystemEditorFolder* ChildFolder);
+
+	void RemoveChildFolder(UNiagaraSystemEditorFolder* ChildFolder);
+
+	const TArray<FGuid>& GetChildEmitterHandleIds() const;
+
+	void AddChildEmitterHandleId(FGuid ChildEmitterHandleId);
+
+	void RemoveChildEmitterHandleId(FGuid ChildEmitterHandleId);
+
+private:
+	UPROPERTY()
+	FName FolderName;
+
+	UPROPERTY()
+	TArray<UNiagaraSystemEditorFolder*> ChildFolders;
+
+	UPROPERTY()
+	TArray<FGuid> ChildEmitterHandleIds;
+};
+
+/** Editor only UI data for systems. */
+UCLASS()
+class NIAGARAEDITOR_API UNiagaraSystemEditorData : public UNiagaraEditorDataBase
+{
+	GENERATED_BODY()
+
+public:
+	UNiagaraSystemEditorData(const FObjectInitializer& ObjectInitializer);
+
+	//~ Begin UObject Interface
+	void PostInitProperties();
+	//~ End UObject Interface
+
+	virtual void PostLoadFromOwner(UObject* InOwner) override;
+
+	/** Gets the root folder for UI folders for emitters. */
+	UNiagaraSystemEditorFolder& GetRootFolder() const;
+
+	/** Gets the stack editor data for the system. */
+	UNiagaraStackEditorData& GetStackEditorData() const;
+
+	const FTransform& GetOwnerTransform() const {
+		return OwnerTransform;
+	}
+
+	void SetOwnerTransform(const FTransform& InTransform) {
+		OwnerTransform = InTransform;
+	}
+
+	TRange<float> GetPlaybackRange() const;
+
+	void SetPlaybackRange(TRange<float> InPlaybackRange);
+
+	UEdGraph* GetSystemOverviewGraph() const;
+
+	void SystemOverviewHandleAdded(const FGuid AddedHandleGuid) const;
+
+	void SystemOverviewHandlesRemoved() const;
+
+	void Initialize(UNiagaraSystem* OwnerSystem, bool bEditingSystem);
+
+private:
+	void UpdatePlaybackRangeFromEmitters(UNiagaraSystem* OwnerSystem);
+
+	bool GetSystemOverviewGraphIsValid() const;
+
+	void InitSystemOverviewGraph();
+
+	const FVector2D GetGoodPlaceForNewOverviewNode() const;
+
+private:
+	UPROPERTY(Instanced)
+	UNiagaraSystemEditorFolder* RootFolder;
+
+	UPROPERTY(Instanced)
+	UNiagaraStackEditorData* StackEditorData;
+
+	UPROPERTY()
+	FTransform OwnerTransform;
+
+	UPROPERTY()
+	float PlaybackRangeMin;
+
+	UPROPERTY()
+	float PlaybackRangeMax;
+
+	/** Graph presenting overview of the current system and its emitters. */
+	UPROPERTY()
+	UEdGraph* SystemOverviewGraph;
+
+	UPROPERTY()
+	UNiagaraSystem* OwningSystem;
+};
diff --git a/Engine/Plugins/FX/Niagara/Source/NiagaraEditor/Public/ViewModels/NiagaraOverviewGraphViewModel.h b/Engine/Plugins/FX/Niagara/Source/NiagaraEditor/Public/ViewModels/NiagaraOverviewGraphViewModel.h
new file mode 100644
--- /dev/null
+++ b/Engine/Plugins/FX/Niagara/Source/NiagaraEditor/Public/ViewModels/NiagaraOverviewGraphViewModel.h
@@ -0,0 +1,91 @@
+// Copyright 1998-2019 Epic Games, Inc. All Rights Reserved.
+
+#pragma once
+
+#include "CoreMinimal.h"
+#include "EdGraph/EdGraphNode.h"
+#include "EditorUndoClient.h"
+#include "Styling/SlateColor.h"
+#include "Framework/Commands/UICommandList.h"
+
+class UNiagaraGraph;
+class FNiagaraSystemViewModel;
+class FNiagaraObjectSelection;
+
+/** A view model for editing a niagara system in a graph editor. */
+class FNiagaraOverviewGraphViewModel : public TSharedFromThis<FNiagaraOverviewGraphViewModel>, public FEditorUndoClient
+{
+public:
+	/** A multicast delegate which is called when nodes are pasted in the graph which supplies the pasted nodes. */
+	DECLARE_MULTICAST_DELEGATE_OneParam(FOnNodesPasted, const TSet<UEdGraphNode*>&);
+
+	DECLARE_MULTICAST_DELEGATE(FOnGraphChanged);
+
+public:
+	/** Create a new view model with the supplied system editor data and graph widget. */
+	FNiagaraOverviewGraphViewModel(TSharedRef<FNiagaraSystemViewModel> InSystemViewModel);
+
+	~FNiagaraOverviewGraphViewModel();
+
+	/** Gets the display text for this graph. */
+	NIAGARAEDITOR_API FText GetDisplayName() const;
+
+	/** Gets the graph which is used to edit and view the system */
+	UEdGraph* GetGraph() const;
+
+	/** Gets commands used for editing the graph. */
+	NIAGARAEDITOR_API TSharedRef<FUICommandList> GetCommands();
+
+	/** Gets the currently selected graph nodes. */
+	TSharedRef<FNiagaraObjectSelection> GetNodeSelection();
+
+	/** Sets the currently selected graph nodes. */
+	void SetSelectedNodes(const TSet<UObject*>& InSelectedNodes);
+
+	/** Clears the currently selected graph nodes. */
+	void ClearSelectedNodes();
+
+	/** Gets a multicast delegate which is called any time nodes are pasted in the graph. */
+	FOnNodesPasted& OnNodesPasted();
+
+	//~ FEditorUndoClient interface.
+	virtual void PostUndo(bool bSuccess) override;
+	virtual void PostRedo(bool bSuccess) override { PostUndo(bSuccess); }
+
+private:
+	void SetupCommands();
+
+	void SelectAllNodes();
+	void DeleteSelectedNodes();
+	bool CanDeleteNodes() const;
+	void CutSelectedNodes();
+	bool CanCutNodes() const;
+	void CopySelectedNodes();
+	bool CanCopyNodes() const;
+	void PasteNodes();
+	bool CanPasteNodes() const;
+	void DuplicateNodes();
+	bool CanDuplicateNodes() const;
+
+	void InitDisplayName();
+
+private:
+
+	/** The view model to interface with the system being viewed and edited by this view model. */
+	TWeakPtr<FNiagaraSystemViewModel> SystemViewModel;
+
+	/** The display name for the overview graph. */
+	FText DisplayName;
+
+	/** Commands for editing the graph. */
+	TSharedRef<FUICommandList> Commands;
+
+	/** The set of nodes objects currently selected in the graph. */
+	TSharedRef<FNiagaraObjectSelection> NodeSelection;
+
+	/** A multicast delegate which is called whenever nodes are pasted into the graph. */
+	FOnNodesPasted OnNodesPastedDelegate;
+
+	/** A multicast delegate which is called whenever the graph object is changed to a different graph. */
+	FOnGraphChanged OnGraphChangedDelegate;
+};
diff --git a/Engine/Plugins/FX/Niagara/Source/NiagaraEditor/Public/ViewModels/NiagaraSystemViewModel.h b/Engine/Plugins/FX/Niagara/Source/NiagaraEditor/Public/ViewModels/NiagaraSystemViewModel.h
--- a/Engine/Plugins/FX/Niagara/Source/NiagaraEditor/Public/ViewModels/NiagaraSystemViewModel.h
+++ b/Engine/Plugins/FX/Niagara/Source/NiagaraEditor/Public/ViewModels/NiagaraSystemViewModel.h
@@ -32,6 +32,7 @@ struct FNiagaraParameterStore;
 struct FEdGraphEditAction;
 class UNiagaraNodeFunctionCall;
 class FNiagaraEmitterViewModel;
+class FNiagaraOverviewGraphViewModel;
 
 /** Defines different editing modes for this system view model. */
 enum class NIAGARAEDITOR_API ENiagaraSystemViewModelEditMode
@@ -141,7 +142,7 @@ public:
 	const TArray<TSharedRef<FNiagaraEmitterHandleViewModel>>& GetEmitterHandleViewModels();
 
 	/** Gets an emitter handle view model by id.  Returns an invalid shared ptr if it can't be found. */
-	TSharedPtr<FNiagaraEmitterHandleViewModel> GetEmitterHandleViewModelById(FGuid InEmitterHandleId);
+	NIAGARAEDITOR_API TSharedPtr<FNiagaraEmitterHandleViewModel> GetEmitterHandleViewModelById(FGuid InEmitterHandleId);
 
 	/** Gets the view model for the System script. */
 	TSharedPtr<FNiagaraSystemScriptViewModel> GetSystemScriptViewModel();
@@ -165,17 +166,11 @@ public:
 	NIAGARAEDITOR_API ENiagaraSystemViewModelEditMode GetEditMode() const;
 
 	/** Adds a new emitter to the System from an emitter asset data. */
-	void AddEmitterFromAssetData(const FAssetData& AssetData);
+	NIAGARAEDITOR_API void AddEmitterFromAssetData(const FAssetData& AssetData);
 
 	/** Adds a new emitter to the System. */
 	void AddEmitter(UNiagaraEmitter& Emitter);
 
-	/** Duplicates the selected emitter in this System. */
-	void DuplicateEmitter(TSharedRef<FNiagaraEmitterHandleViewModel> EmitterHandleToDuplicate);
-
-	/** Deletes the selected emitter from the System. */
-	void DeleteEmitter(TSharedRef<FNiagaraEmitterHandleViewModel> EmitterHandleToDelete);
-
 	/** Deletes the emitters with the supplied ids from the system */
 	void DeleteEmitters(TSet<FGuid> EmitterHandleIdsToDelete);
 
@@ -241,7 +236,7 @@ public:
 	void GetSelectedEmitterHandles(TArray<TSharedRef<FNiagaraEmitterHandleViewModel>>& OutSelectedEmitterHanldles);
 
 	/** Gets editor specific data which can be stored per system.  If this data hasn't been created the default version will be returned. */
-	const UNiagaraSystemEditorData& GetEditorData() const;
+	const NIAGARAEDITOR_API UNiagaraSystemEditorData& GetEditorData() const;
 
 	/** Gets editor specific data which is stored per system.  If this data hasn't been created then it will be created. */
 	UNiagaraSystemEditorData& GetOrCreateEditorData();
@@ -297,6 +292,15 @@ public:
 	/** Gets the stack module data for the provided emitter, for use in module dependencies. */
 	const TArray<FNiagaraStackModuleData>& GetStackModuleDataForEmitter(TSharedRef<FNiagaraEmitterViewModel> EmitterViewModel);
 
+	/** Gets the ViewModel for the system overview graph. */
+	NIAGARAEDITOR_API TSharedPtr<FNiagaraOverviewGraphViewModel> GetOverviewGraphViewModel() const;
+
+	/** Sets the ViewModel for the system overview graph. */
+	void SetOverviewGraphViewModel(TSharedRef<FNiagaraOverviewGraphViewModel> InOverviewGraphViewModel);
+
+	/** Notifies views to update displayed Emitter name that has changed. */
+	NIAGARAEDITOR_API void EmitterNameChanged(const FText& InText, FGuid ChangedEmitterGuid);
+
 private:
 
 	/** Sends message jobs to FNiagaraMessageManager for all compile events from the last compile. */
@@ -323,9 +327,6 @@ private:
 	/** Gets the content for the add menu in sequencer. */
 	void GetSequencerAddMenuContent(FMenuBuilder& MenuBuilder, TSharedRef<ISequencer> Sequencer);
 
-	/** Kills all system instances using the system being displayed by this view model. */
-	void KillSystemInstances();
-
 	/** Resets and rebuilds the data in the curve owner. */
 	void ResetCurveData();
 
@@ -543,4 +544,8 @@ private:
 
 	/** GUID used when sending message jobs to FNiagaraMessageManager for notifying the FNiagaraMessageLogViewModel with the same GUID key */
 	const TOptional<const FGuid> SystemMessageLogGuidKey;
-};
\ No newline at end of file
+
+	/** ViewModel for the system overview graph */
+	TSharedPtr<FNiagaraOverviewGraphViewModel> OverviewGraphViewModel;
+
+};
diff --git a/Engine/Plugins/FX/Niagara/Source/NiagaraEditorWidgets/Private/NiagaraEditorWidgetsModule.cpp b/Engine/Plugins/FX/Niagara/Source/NiagaraEditorWidgets/Private/NiagaraEditorWidgetsModule.cpp
--- a/Engine/Plugins/FX/Niagara/Source/NiagaraEditorWidgets/Private/NiagaraEditorWidgetsModule.cpp
+++ b/Engine/Plugins/FX/Niagara/Source/NiagaraEditorWidgets/Private/NiagaraEditorWidgetsModule.cpp
@@ -10,6 +10,8 @@
 #include "Modules/ModuleManager.h"
 #include "PropertyEditorModule.h"
 
+#include "SNiagaraOverviewGraph.h"
+
 IMPLEMENT_MODULE(FNiagaraEditorWidgetsModule, NiagaraEditorWidgets);
 
 FNiagaraStackCurveEditorOptions::FNiagaraStackCurveEditorOptions()
@@ -95,6 +97,11 @@ void FNiagaraEditorWidgetsModule::StartupModule()
 	PropertyModule.RegisterCustomClassLayout("NiagaraDataInterfaceVector4Curve", FOnGetDetailCustomizationInstance::CreateStatic(&FNiagaraDataInterfaceVector4CurveDetails::MakeInstance));
 	PropertyModule.RegisterCustomClassLayout("NiagaraDataInterfaceColorCurve", FOnGetDetailCustomizationInstance::CreateStatic(&FNiagaraDataInterfaceColorCurveDetails::MakeInstance));
 	PropertyModule.RegisterCustomClassLayout("NiagaraDataInterfaceSkeletalMesh", FOnGetDetailCustomizationInstance::CreateStatic(&FNiagaraDataInterfaceSkeletalMeshDetails::MakeInstance));
+
+	ReinitializeStyleCommand = IConsoleManager::Get().RegisterConsoleCommand(
+		TEXT("fx.NiagaraEditorWidgets.ReinitializeStyle"),
+		TEXT("Reinitializes the style for the niagara editor widgets module.  Used in conjuction with live coding for UI tweaks.  May crash the editor if style objects are in use."),
+		FConsoleCommandDelegate::CreateRaw(this, &FNiagaraEditorWidgetsModule::ReinitializeStyle));
 }
 
 void FNiagaraEditorWidgetsModule::ShutdownModule()
@@ -117,9 +124,20 @@ void FNiagaraEditorWidgetsModule::ShutdownModule()
 		PropertyModule->UnregisterCustomClassLayout("NiagaraDataInterfaceSkeletalMesh");
 	}
 
+	if (ReinitializeStyleCommand != nullptr)
+	{
+		IConsoleManager::Get().UnregisterConsoleObject(ReinitializeStyleCommand);
+	}
+
 	FNiagaraEditorWidgetsStyle::Shutdown();
 }
 
+void FNiagaraEditorWidgetsModule::ReinitializeStyle()
+{
+	FNiagaraEditorWidgetsStyle::Shutdown();
+	FNiagaraEditorWidgetsStyle::Initialize();
+}
+
 TSharedRef<FNiagaraStackCurveEditorOptions> FNiagaraEditorWidgetsModule::GetOrCreateStackCurveEditorOptionsForObject(UObject* Object, bool bDefaultAreCurvesVisible, float DefaultHeight)
 {
 	TSharedRef<FNiagaraStackCurveEditorOptions>* StackCurveEditorOptions = ObjectToStackCurveEditorOptionsMap.Find(FObjectKey(Object));
@@ -139,5 +157,5 @@ TSharedRef<SWidget> FNiagaraEditorWidgetsModule::FNiagaraEditorWidgetProvider::C
 
 TSharedRef<SWidget> FNiagaraEditorWidgetsModule::FNiagaraEditorWidgetProvider::CreateSystemOverview(TSharedRef<FNiagaraSystemViewModel> SystemViewModel)
 {
-	return SNullWidget::NullWidget;
-}
\ No newline at end of file
+	return SNew(SNiagaraOverviewGraph, SystemViewModel);
+}
diff --git a/Engine/Plugins/FX/Niagara/Source/NiagaraEditorWidgets/Private/SNiagaraOverviewGraph.cpp b/Engine/Plugins/FX/Niagara/Source/NiagaraEditorWidgets/Private/SNiagaraOverviewGraph.cpp
new file mode 100644
--- /dev/null
+++ b/Engine/Plugins/FX/Niagara/Source/NiagaraEditorWidgets/Private/SNiagaraOverviewGraph.cpp
@@ -0,0 +1,266 @@
+// Copyright 1998-2019 Epic Games, Inc. All Rights Reserved.
+
+#include "SNiagaraOverviewGraph.h"
+#include "SNiagaraStack.h"
+#include "ViewModels/NiagaraSystemViewModel.h"
+#include "ViewModels/NiagaraEmitterHandleViewModel.h"
+#include "ViewModels/Stack/NiagaraStackViewModel.h"
+#include "NiagaraSystem.h"
+#include "NiagaraSystemEditorData.h"
+#include "IContentBrowserSingleton.h"
+#include "ContentBrowserModule.h"
+#include "ViewModels/NiagaraOverviewGraphViewModel.h"
+#include "EdGraphSchema_Niagara.h"
+#include "NiagaraOverviewNodeStackItem.h"
+#include "GraphEditorActions.h"
+
+#define LOCTEXT_NAMESPACE "NiagaraOverviewGraph"
+
+void SNiagaraOverviewGraph::Construct(const FArguments& InArgs, TSharedRef<FNiagaraSystemViewModel> InSystemViewModel)
+{
+	SystemViewModel = InSystemViewModel;
+	bUpdatingOverviewSelectionFromGraph = false;
+	bUpdatingGraphSelectionFromOverview = false;
+
+	SGraphEditor::FGraphEditorEvents Events;
+	Events.OnSelectionChanged = SGraphEditor::FOnSelectionChanged::CreateSP(this, &SNiagaraOverviewGraph::GraphSelectionChanged);
+	Events.OnCreateActionMenu = SGraphEditor::FOnCreateActionMenu::CreateSP(this, &SNiagaraOverviewGraph::OnCreateGraphActionMenu);
+	Events.OnTextCommitted = FOnNodeTextCommitted::CreateSP(this, &SNiagaraOverviewGraph::OnNodeTitleCommitted);
+
+	TSharedPtr<FNiagaraOverviewGraphViewModel> OverviewGraphViewModel = SystemViewModel->GetOverviewGraphViewModel();
+
+	FGraphAppearanceInfo AppearanceInfo;
+	if (SystemViewModel->GetEditMode() == ENiagaraSystemViewModelEditMode::EmitterAsset)
+	{
+		AppearanceInfo.CornerText = LOCTEXT("NiagaraOverview_AppearanceCornerTextEmitter", "EMITTER");
+
+	}
+	else if (SystemViewModel->GetEditMode() == ENiagaraSystemViewModelEditMode::SystemAsset)
+	{
+		AppearanceInfo.CornerText = LOCTEXT("NiagaraOverview_AppearanceCornerTextSystem", "SYSTEM");
+	}
+	else
+	{
+		ensureMsgf(false, TEXT("Encountered unhandled SystemViewModel Edit Mode!"));
+		AppearanceInfo.CornerText = LOCTEXT("NiagaraOverview_AppearanceCornerTextGeneric", "NIAGARA");
+	}
+	
+	TSharedRef<SWidget> TitleBarWidget = SNew(SBorder)
+	.BorderImage(FEditorStyle::GetBrush(TEXT("Graph.TitleBackground")))
+	.HAlign(HAlign_Fill)
+	[
+		SNew(STextBlock)
+		.Text(OverviewGraphViewModel.ToSharedRef(), &FNiagaraOverviewGraphViewModel::GetDisplayName)
+		.TextStyle(FEditorStyle::Get(), TEXT("GraphBreadcrumbButtonText"))
+		.Justification(ETextJustify::Center)
+	];
+
+	TSharedRef<FUICommandList> Commands = OverviewGraphViewModel->GetCommands();
+	Commands->MapAction(
+		FGraphEditorCommands::Get().CreateComment,
+		FExecuteAction::CreateRaw(this, &SNiagaraOverviewGraph::OnCreateComment));
+
+	GraphEditor = SNew(SGraphEditor)
+		.AdditionalCommands(Commands)
+		.Appearance(AppearanceInfo)
+		.TitleBar(TitleBarWidget)
+		.GraphToEdit(SystemViewModel->GetEditorData().GetSystemOverviewGraph())
+		.GraphEvents(Events);
+
+	//GraphEditor->SetNodeFactory(MakeShared<FNiagaraOverviewGraphNodeFactory>());
+
+	ChildSlot
+	[
+		SNew(SSplitter)
+		+ SSplitter::Slot()
+		.Value(.7f)
+		[
+			GraphEditor.ToSharedRef()
+		]
+		+ SSplitter::Slot()
+		.Value(.3)
+		[
+			SNew(SBox)
+			//SNew(SNiagaraStack, SystemViewModel->GetOverviewSelectionViewModel()->GetOverviewStackViewModel())
+		]
+	];
+
+	//SystemViewModel->GetOverviewSelectionViewModel()->OnSelectionChanged().AddSP(this, &SNiagaraOverviewGraph::OverviewSelectionChanged);
+}
+
+UNiagaraStackEntry* GetStackEntryForSelectedNode(UObject* SelectedNode, TSharedPtr<FNiagaraSystemViewModel> SystemViewModel)
+{
+// 	UNiagaraOverviewNodeStackItem* OverviewStackNode = Cast<UNiagaraOverviewNodeStackItem>(SelectedNode);
+// 	if (OverviewStackNode != nullptr)
+// 	{
+// 		if (OverviewStackNode->GetOwningSystem() != nullptr)
+// 		{
+// 			UNiagaraStackViewModel* StackViewModel = nullptr;
+// 			if (OverviewStackNode->GetEmitterHandleGuid().IsValid())
+// 			{
+// 				TSharedPtr<FNiagaraEmitterHandleViewModel> EmitterHandleViewModel = SystemViewModel->GetEmitterHandleViewModelById(OverviewStackNode->GetEmitterHandleGuid());
+// 				StackViewModel = EmitterHandleViewModel->GetEmitterStackViewModel();
+// 			}
+// 			else
+// 			{
+// 				StackViewModel = SystemViewModel->GetSystemStackViewModel();
+// 			}
+// 			TArray<UNiagaraStackEntry*> RootEntries = StackViewModel->GetRootEntries();
+// 			if (ensureMsgf(RootEntries.Num() == 1, TEXT("Only one root entry expected")))
+// 			{
+// 				return RootEntries[0];
+// 			}
+// 		}
+// 	}
+	return nullptr;
+}
+
+void SNiagaraOverviewGraph::GraphSelectionChanged(const TSet<UObject*>& SelectedNodes)
+{
+// 	if (bUpdatingGraphSelectionFromOverview == false)
+// 	{
+// 		TGuardValue<bool> UpdateGuard(bUpdatingOverviewSelectionFromGraph, true);
+// 		TArray<UNiagaraStackEntry*> SelectedEntries;
+// 		for (UObject* SelectedNode : SelectedNodes)
+// 		{
+// 			UNiagaraStackEntry* SelectedEntry = GetStackEntryForSelectedNode(SelectedNode, SystemViewModel);
+// 			if(SelectedEntry != nullptr)
+// 			{
+// 				SelectedEntries.Add(SelectedEntry);
+// 			}
+// 		}
+// 		bool bClearCurrentSelection = FSlateApplication::Get().GetModifierKeys().IsControlDown() == false;
+// 		SystemViewModel->GetOverviewSelectionViewModel()->UpdateSelectedEntries(SelectedEntries, bClearCurrentSelection);
+// 	}
+}
+
+void SNiagaraOverviewGraph::OverviewSelectionChanged()
+{
+// 	if (bUpdatingOverviewSelectionFromGraph == false)
+// 	{
+// 		TGuardValue<bool> UpdateGuard(bUpdatingGraphSelectionFromOverview, true);
+// 
+// 		TSet<UObject*> SelectedNodeObjects = GraphEditor->GetSelectedNodes();
+// 		TArray<UNiagaraStackEntry*> SelectedOverviewEntries = SystemViewModel->GetOverviewSelectionViewModel()->GetSelectedEntries();
+// 
+// 		TArray<UEdGraphNode*> NodesToDeselect;
+// 		for (UObject* SelectedNodeObject : SelectedNodeObjects)
+// 		{
+// 			UEdGraphNode* SelectedNode = Cast<UEdGraphNode>(SelectedNodeObject);
+// 			if (SelectedNode != nullptr)
+// 			{
+// 				UNiagaraStackEntry* SelectedEntry = GetStackEntryForSelectedNode(SelectedNode, SystemViewModel);
+// 				if (SelectedEntry != nullptr && SelectedOverviewEntries.Contains(SelectedEntry) == false)
+// 				{
+// 					NodesToDeselect.Add(SelectedNode);
+// 				}
+// 			}
+// 		}
+// 
+// 		for (UEdGraphNode* NodeToDeselect : NodesToDeselect)
+// 		{
+// 			GraphEditor->SetNodeSelection(NodeToDeselect, false);
+// 		}
+// 	}
+}
+
+FActionMenuContent SNiagaraOverviewGraph::OnCreateGraphActionMenu(UEdGraph* InGraph, const FVector2D& InNodePosition, const TArray<UEdGraphPin*>& InDraggedPins, bool bAutoExpand, SGraphEditor::FActionMenuClosed InOnMenuClosed)
+{
+	if (SystemViewModel->GetEditMode() == ENiagaraSystemViewModelEditMode::SystemAsset)
+	{
+		FMenuBuilder MenuBuilder(true, NULL);
+		MenuBuilder.BeginSection(TEXT("NiagaraOverview_EditGraph"), LOCTEXT("EditGraph", "Edit Graph"));
+
+		MenuBuilder.AddSubMenu(
+			LOCTEXT("EmitterAddLabel", "Add Emitter"),
+			LOCTEXT("EmitterAddToolTip", "Add an existing emitter"),
+			FNewMenuDelegate::CreateSP(this, &SNiagaraOverviewGraph::CreateAddEmitterMenuContent, InGraph));
+
+		MenuBuilder.AddMenuEntry(
+			LOCTEXT("CommentsLabel", "Add Comment"),
+			LOCTEXT("CommentsToolTip", "Add a comment box"),
+			FSlateIcon(),
+			FExecuteAction::CreateSP(this, &SNiagaraOverviewGraph::OnCreateComment));
+
+		MenuBuilder.EndSection();
+		TSharedRef<SWidget> ActionMenu = MenuBuilder.MakeWidget();
+
+		return FActionMenuContent(ActionMenu, ActionMenu);
+	}
+	return FActionMenuContent(SNullWidget::NullWidget, SNullWidget::NullWidget);
+}
+
+void SNiagaraOverviewGraph::OnCreateComment()
+{
+	FNiagaraSchemaAction_NewComment CommentAction = FNiagaraSchemaAction_NewComment(GraphEditor);
+	CommentAction.PerformAction(SystemViewModel->GetEditorData().GetSystemOverviewGraph(), nullptr, GraphEditor->GetPasteLocation(), false);
+}
+
+void SNiagaraOverviewGraph::OnNodeTitleCommitted(const FText& NewText, ETextCommit::Type CommitInfo, UEdGraphNode* NodeBeingChanged)
+{
+	if (NodeBeingChanged)
+	{
+		// When you request rename on spawn but accept the value, we want to not add a transaction if they just hit "Enter".
+		bool bRename = true;
+		FText CurrentNodeTitleText = NodeBeingChanged->GetNodeTitle(ENodeTitleType::FullTitle);
+		if (CurrentNodeTitleText.EqualTo(NewText))
+		{
+			return;
+		}
+
+		if (NodeBeingChanged->IsA(UNiagaraOverviewNodeStackItem::StaticClass())) //@TODO System Overview: renaming system or emitters locally through this view
+		{
+			UNiagaraOverviewNodeStackItem* OverviewNodeBeingChanged = Cast<UNiagaraOverviewNodeStackItem>(NodeBeingChanged);
+			TSharedPtr<FNiagaraEmitterHandleViewModel> NodeEmitterHandleViewModel = SystemViewModel->GetEmitterHandleViewModelById(OverviewNodeBeingChanged->GetEmitterHandleGuid());
+			if (ensureMsgf(NodeEmitterHandleViewModel.IsValid(), TEXT("Failed to find EmitterHandleViewModel with matching Emitter GUID to Overview Node!")))
+			{
+				NodeEmitterHandleViewModel->OnNameTextComitted(NewText, CommitInfo);
+			}
+			else
+			{
+				bRename = false;
+			}
+		}
+
+		if (bRename)
+		{
+			const FScopedTransaction Transaction(LOCTEXT("RenameNode", "Rename Node"));
+			NodeBeingChanged->Modify();
+			NodeBeingChanged->OnRenameNode(NewText.ToString());
+		}
+	}
+}
+
+void SNiagaraOverviewGraph::CreateAddEmitterMenuContent(FMenuBuilder& MenuBuilder, UEdGraph* InGraph)
+{
+	FAssetPickerConfig AssetPickerConfig;
+	{
+		AssetPickerConfig.OnAssetSelected = FOnAssetSelected::CreateLambda([this, InGraph](const FAssetData& AssetData)
+		{
+			FSlateApplication::Get().DismissAllMenus();
+			SystemViewModel->AddEmitterFromAssetData(AssetData);
+		});
+		AssetPickerConfig.bAllowNullSelection = false;
+		AssetPickerConfig.InitialAssetViewType = EAssetViewType::List;
+		AssetPickerConfig.Filter.ClassNames.Add(UNiagaraEmitter::StaticClass()->GetFName());
+	}
+
+	FContentBrowserModule& ContentBrowserModule = FModuleManager::Get().LoadModuleChecked<FContentBrowserModule>(TEXT("ContentBrowser"));
+
+	TSharedRef<SWidget> EmitterAddSubMenu =
+		SNew(SBorder)
+		.BorderImage(FEditorStyle::GetBrush("Menu.Background"))
+		.Padding(5.0f)
+		[
+			SNew(SBox)
+			.WidthOverride(300.0f)
+			.HeightOverride(300.0f)
+			[
+				ContentBrowserModule.Get().CreateAssetPicker(AssetPickerConfig)
+			]
+		];
+	
+	MenuBuilder.AddWidget(EmitterAddSubMenu, FText());
+}
+
+#undef LOCTEXT_NAMESPACE // "NiagaraOverviewGraph"
diff --git a/Engine/Plugins/FX/Niagara/Source/NiagaraEditorWidgets/Private/SNiagaraOverviewGraph.h b/Engine/Plugins/FX/Niagara/Source/NiagaraEditorWidgets/Private/SNiagaraOverviewGraph.h
new file mode 100644
--- /dev/null
+++ b/Engine/Plugins/FX/Niagara/Source/NiagaraEditorWidgets/Private/SNiagaraOverviewGraph.h
@@ -0,0 +1,35 @@
+// Copyright 1998-2019 Epic Games, Inc. All Rights Reserved.
+
+#pragma once
+
+class FNiagaraSystemViewModel;
+
+class SNiagaraOverviewGraph : public SCompoundWidget
+{
+public:
+	SLATE_BEGIN_ARGS(SNiagaraOverviewGraph) {}
+	SLATE_END_ARGS()
+
+	void Construct(const FArguments& InArgs, TSharedRef<FNiagaraSystemViewModel> InSystemViewModel);
+
+private:
+	void GraphSelectionChanged(const TSet<UObject*>& SelectedNodes);
+
+	void OverviewSelectionChanged();
+
+	/** Called to create context menu when right-clicking on graph */
+	FActionMenuContent OnCreateGraphActionMenu(UEdGraph* InGraph, const FVector2D& InNodePosition, const TArray<UEdGraphPin*>& InDraggedPins, bool bAutoExpand, SGraphEditor::FActionMenuClosed InOnMenuClosed);
+
+	void OnCreateComment();
+
+	void OnNodeTitleCommitted(const FText& NewText, ETextCommit::Type CommitInfo, UEdGraphNode* NodeBeingChanged);
+
+	void CreateAddEmitterMenuContent(FMenuBuilder& MenuBuilder, UEdGraph* InGraph);
+	
+private:
+	TSharedPtr<FNiagaraSystemViewModel> SystemViewModel;
+	TSharedPtr<SGraphEditor> GraphEditor;
+
+	bool bUpdatingOverviewSelectionFromGraph;
+	bool bUpdatingGraphSelectionFromOverview;
+};
diff --git a/Engine/Plugins/FX/Niagara/Source/NiagaraEditorWidgets/Private/SNiagaraStack.cpp b/Engine/Plugins/FX/Niagara/Source/NiagaraEditorWidgets/Private/SNiagaraStack.cpp
--- a/Engine/Plugins/FX/Niagara/Source/NiagaraEditorWidgets/Private/SNiagaraStack.cpp
+++ b/Engine/Plugins/FX/Niagara/Source/NiagaraEditorWidgets/Private/SNiagaraStack.cpp
@@ -1079,7 +1079,8 @@ FText SNiagaraStack::GetEmitterNameToolTip() const
 
 void SNiagaraStack::OnStackViewNameTextCommitted(const FText& InText, ETextCommit::Type CommitInfo) const
 {
-	StackViewModel->GetEmitterHandleViewModel()->OnNameTextComitted(InText, CommitInfo);
+	const FGuid EditedEmitterGuid = StackViewModel->GetEmitterHandleViewModel()->GetEmitterHandle()->GetId();
+	StackViewModel->GetSystemViewModel()->EmitterNameChanged(InText, EditedEmitterGuid);
 }
 
 EVisibility SNiagaraStack::GetSourceEmitterNameVisibility() const
diff --git a/Engine/Plugins/FX/Niagara/Source/NiagaraEditorWidgets/Public/NiagaraEditorWidgetsModule.h b/Engine/Plugins/FX/Niagara/Source/NiagaraEditorWidgets/Public/NiagaraEditorWidgetsModule.h
--- a/Engine/Plugins/FX/Niagara/Source/NiagaraEditorWidgets/Public/NiagaraEditorWidgetsModule.h
+++ b/Engine/Plugins/FX/Niagara/Source/NiagaraEditorWidgets/Public/NiagaraEditorWidgetsModule.h
@@ -59,9 +59,12 @@ public:
 	TSharedRef<FNiagaraStackCurveEditorOptions> GetOrCreateStackCurveEditorOptionsForObject(UObject* Object, bool bDefaultAreCurvesVisible, float DefaultHeight);
 
 private:
+	void ReinitializeStyle();
 
 private:
 	TMap<FObjectKey, TSharedRef<FNiagaraStackCurveEditorOptions>> ObjectToStackCurveEditorOptionsMap;
 
 	TSharedPtr<FNiagaraEditorWidgetProvider> WidgetProvider;
+
+	IConsoleCommand* ReinitializeStyleCommand;
 };
\ No newline at end of file
