Fix includes and forward decls for non unity

diff --git a/Engine/Plugins/FX/Niagara/Source/NiagaraEditor/Private/NiagaraEditorUtilities.cpp b/Engine/Plugins/FX/Niagara/Source/NiagaraEditor/Private/NiagaraEditorUtilities.cpp
--- a/Engine/Plugins/FX/Niagara/Source/NiagaraEditor/Private/NiagaraEditorUtilities.cpp
+++ b/Engine/Plugins/FX/Niagara/Source/NiagaraEditor/Private/NiagaraEditorUtilities.cpp
@@ -35,6 +35,7 @@
 #include "NiagaraNodeStaticSwitch.h"
 #include "NiagaraNodeFunctionCall.h"
 #include "NiagaraParameterMapHistory.h"
+#include "ScopedTransaction.h"
 
 #define LOCTEXT_NAMESPACE "FNiagaraEditorUtilities"
 
diff --git a/Engine/Plugins/FX/Niagara/Source/NiagaraEditor/Private/NiagaraSystemEditorData.cpp b/Engine/Plugins/FX/Niagara/Source/NiagaraEditor/Private/NiagaraSystemEditorData.cpp
--- a/Engine/Plugins/FX/Niagara/Source/NiagaraEditor/Private/NiagaraSystemEditorData.cpp
+++ b/Engine/Plugins/FX/Niagara/Source/NiagaraEditor/Private/NiagaraSystemEditorData.cpp
@@ -9,6 +9,7 @@
 #include "NiagaraEmitterEditorData.h"
 #include "NiagaraOverviewNodeStackItem.h"
 #include "EdGraphSchema_NiagaraSystemOverview.h"
+#include "EdGraph/EdGraph.h"
 
 static const float SystemOverviewNodePadding = 300.0f;
 
diff --git a/Engine/Plugins/FX/Niagara/Source/NiagaraEditor/Private/ViewModels/NiagaraOverviewGraphViewModel.cpp b/Engine/Plugins/FX/Niagara/Source/NiagaraEditor/Private/ViewModels/NiagaraOverviewGraphViewModel.cpp
--- a/Engine/Plugins/FX/Niagara/Source/NiagaraEditor/Private/ViewModels/NiagaraOverviewGraphViewModel.cpp
+++ b/Engine/Plugins/FX/Niagara/Source/NiagaraEditor/Private/ViewModels/NiagaraOverviewGraphViewModel.cpp
@@ -11,6 +11,7 @@
 #include "ViewModels/NiagaraSystemViewModel.h"
 #include "ViewModels/NiagaraEmitterHandleViewModel.h"
 #include "NiagaraSystem.h"
+#include "Editor.h"
 
 #define LOCTEXT_NAMESPACE "NiagaraOverviewGraphViewModel"
 
diff --git a/Engine/Plugins/FX/Niagara/Source/NiagaraEditor/Public/NiagaraOverviewNodeStackItem.h b/Engine/Plugins/FX/Niagara/Source/NiagaraEditor/Public/NiagaraOverviewNodeStackItem.h
--- a/Engine/Plugins/FX/Niagara/Source/NiagaraEditor/Public/NiagaraOverviewNodeStackItem.h
+++ b/Engine/Plugins/FX/Niagara/Source/NiagaraEditor/Public/NiagaraOverviewNodeStackItem.h
@@ -1,6 +1,7 @@
 // Copyright 1998-2019 Epic Games, Inc. All Rights Reserved.
 #pragma once
 
+#include "EdGraph/EdGraphNode.h"
 #include "NiagaraOverviewNodeStackItem.generated.h"
 
 class UNiagaraSystem;
diff --git a/Engine/Plugins/FX/Niagara/Source/NiagaraEditor/Public/NiagaraSystemEditorData.h b/Engine/Plugins/FX/Niagara/Source/NiagaraEditor/Public/NiagaraSystemEditorData.h
--- a/Engine/Plugins/FX/Niagara/Source/NiagaraEditor/Public/NiagaraSystemEditorData.h
+++ b/Engine/Plugins/FX/Niagara/Source/NiagaraEditor/Public/NiagaraSystemEditorData.h
@@ -7,6 +7,7 @@
 
 class UNiagaraStackEditorData;
 class UNiagaraSystem;
+class UEdGraph;
 
 /** Editor only folder data for emitters in a system. */
 UCLASS()
diff --git a/Engine/Plugins/FX/Niagara/Source/NiagaraEditorWidgets/Private/SNiagaraOverviewGraph.cpp b/Engine/Plugins/FX/Niagara/Source/NiagaraEditorWidgets/Private/SNiagaraOverviewGraph.cpp
--- a/Engine/Plugins/FX/Niagara/Source/NiagaraEditorWidgets/Private/SNiagaraOverviewGraph.cpp
+++ b/Engine/Plugins/FX/Niagara/Source/NiagaraEditorWidgets/Private/SNiagaraOverviewGraph.cpp
@@ -13,6 +13,10 @@
 #include "EdGraphSchema_Niagara.h"
 #include "NiagaraOverviewNodeStackItem.h"
 #include "GraphEditorActions.h"
+#include "GraphEditor.h"
+#include "Framework/MultiBox/MultiBoxBuilder.h"
+#include "ScopedTransaction.h"
+#include "Framework/Application/SlateApplication.h"
 
 #define LOCTEXT_NAMESPACE "NiagaraOverviewGraph"
 
diff --git a/Engine/Plugins/FX/Niagara/Source/NiagaraEditorWidgets/Private/SNiagaraOverviewGraph.h b/Engine/Plugins/FX/Niagara/Source/NiagaraEditorWidgets/Private/SNiagaraOverviewGraph.h
--- a/Engine/Plugins/FX/Niagara/Source/NiagaraEditorWidgets/Private/SNiagaraOverviewGraph.h
+++ b/Engine/Plugins/FX/Niagara/Source/NiagaraEditorWidgets/Private/SNiagaraOverviewGraph.h
@@ -2,7 +2,12 @@
 
 #pragma once
 
+#include "Widgets/SCompoundWidget.h"
+
 class FNiagaraSystemViewModel;
+class SGraphEditor;
+struct FActionMenuContent;
+class FMenuBuilder;
 
 class SNiagaraOverviewGraph : public SCompoundWidget
 {
