Resolve non unity include errors SNiagaraOverviewGraph

diff --git a/Engine/Plugins/FX/Niagara/Source/NiagaraEditorWidgets/Private/SNiagaraOverviewGraph.cpp b/Engine/Plugins/FX/Niagara/Source/NiagaraEditorWidgets/Private/SNiagaraOverviewGraph.cpp
--- a/Engine/Plugins/FX/Niagara/Source/NiagaraEditorWidgets/Private/SNiagaraOverviewGraph.cpp
+++ b/Engine/Plugins/FX/Niagara/Source/NiagaraEditorWidgets/Private/SNiagaraOverviewGraph.cpp
@@ -13,11 +13,12 @@
 #include "EdGraphSchema_Niagara.h"
 #include "NiagaraOverviewNodeStackItem.h"
 #include "GraphEditorActions.h"
-#include "GraphEditor.h"
 #include "Framework/MultiBox/MultiBoxBuilder.h"
 #include "ScopedTransaction.h"
 #include "Framework/Application/SlateApplication.h"
 #include "Templates/SharedPointer.h"
+#include "EdGraph/EdGraph.h"
+#include "EdGraph/EdGraphNode.h"
 
 #define LOCTEXT_NAMESPACE "NiagaraOverviewGraph"
 
diff --git a/Engine/Plugins/FX/Niagara/Source/NiagaraEditorWidgets/Private/SNiagaraOverviewGraph.h b/Engine/Plugins/FX/Niagara/Source/NiagaraEditorWidgets/Private/SNiagaraOverviewGraph.h
--- a/Engine/Plugins/FX/Niagara/Source/NiagaraEditorWidgets/Private/SNiagaraOverviewGraph.h
+++ b/Engine/Plugins/FX/Niagara/Source/NiagaraEditorWidgets/Private/SNiagaraOverviewGraph.h
@@ -3,11 +3,13 @@
 #pragma once
 
 #include "Widgets/SCompoundWidget.h"
+#include "GraphEditor.h"
 
 class FNiagaraSystemViewModel;
-class SGraphEditor;
 struct FActionMenuContent;
 class FMenuBuilder;
+class UEdGraph;
+class UEdGraphNode;
 
 class SNiagaraOverviewGraph : public SCompoundWidget
 {
