Rearranged niagara mesh rendering shader code to avoid device reboots (shader compiler bug in VK Note9, S9)

diff --git a/Engine/Plugins/FX/Niagara/Shaders/Private/NiagaraMeshVertexFactory.ush b/Engine/Plugins/FX/Niagara/Shaders/Private/NiagaraMeshVertexFactory.ush
--- a/Engine/Plugins/FX/Niagara/Shaders/Private/NiagaraMeshVertexFactory.ush
+++ b/Engine/Plugins/FX/Niagara/Shaders/Private/NiagaraMeshVertexFactory.ush
@@ -530,16 +530,44 @@ FVertexFactoryIntermediates GetVertexFactoryIntermediates(FVertexFactoryInput In
 	}
 
 #if (DYNAMIC_PARAMETERS_MASK & 1)
-	Intermediates.DynamicParameter = NiagaraMeshVF.MaterialParamDataOffset != -1 ? GetVec4(NiagaraMeshVF.MaterialParamDataOffset, Input.InstanceId) : float4(1, 1, 1, 1);
+	if (NiagaraMeshVF.MaterialParamDataOffset == -1)
+	{
+		Intermediates.DynamicParameter = float4(1, 1, 1, 1);
+	}
+	else
+	{
+		Intermediates.DynamicParameter = GetVec4(NiagaraMeshVF.MaterialParamDataOffset, Input.InstanceId);
+	}
 #endif
 #if (DYNAMIC_PARAMETERS_MASK & 2)
-	Intermediates.DynamicParameter1 = NiagaraMeshVF.MaterialParam1DataOffset != -1 ? GetVec4(NiagaraMeshVF.MaterialParam1DataOffset, Input.InstanceId) : float4(1, 1, 1, 1);
+	if (NiagaraMeshVF.MaterialParam1DataOffset == -1)
+	{
+		Intermediates.DynamicParameter1 = float4(1, 1, 1, 1);
+	}	
+	else
+	{
+		Intermediates.DynamicParameter1 = GetVec4(NiagaraMeshVF.MaterialParam1DataOffset, Input.InstanceId);
+	}
 #endif
 #if (DYNAMIC_PARAMETERS_MASK & 4)
-	Intermediates.DynamicParameter2 = NiagaraMeshVF.MaterialParam2DataOffset != -1 ? GetVec4(NiagaraMeshVF.MaterialParam2DataOffset, Input.InstanceId) : float4(1, 1, 1, 1);
+	if (NiagaraMeshVF.MaterialParam2DataOffset == -1)
+	{
+		Intermediates.DynamicParameter2 = float4(1, 1, 1, 1);
+	}
+	else
+	{
+		Intermediates.DynamicParameter2 = GetVec4(NiagaraMeshVF.MaterialParam2DataOffset, Input.InstanceId);
+	}
 #endif
 #if (DYNAMIC_PARAMETERS_MASK & 8)
-	Intermediates.DynamicParameter3 = NiagaraMeshVF.MaterialParam3DataOffset != -1 ? GetVec4(NiagaraMeshVF.MaterialParam3DataOffset, Input.InstanceId) : float4(1, 1, 1, 1);
+	if (NiagaraMeshVF.MaterialParam3DataOffset == -1)
+	{
+		Intermediates.DynamicParameter3 = float4(1, 1, 1, 1);
+	}
+	else
+	{
+		Intermediates.DynamicParameter3 = GetVec4(NiagaraMeshVF.MaterialParam3DataOffset, Input.InstanceId);
+	}
 #endif
 
 #if NEEDS_PARTICLE_TIME
