Move Emitter Properties out of Emitter Spawn group into separate Stack category.
Move System Properties of out System Properties tab into Stack view.

diff --git a/Engine/Plugins/FX/Niagara/Source/Niagara/Classes/NiagaraSystem.h b/Engine/Plugins/FX/Niagara/Source/Niagara/Classes/NiagaraSystem.h
--- a/Engine/Plugins/FX/Niagara/Source/Niagara/Classes/NiagaraSystem.h
+++ b/Engine/Plugins/FX/Niagara/Source/Niagara/Classes/NiagaraSystem.h
@@ -239,7 +239,7 @@ private:
 protected:
 
 	/** Handles to the emitter this System will simulate. */
-	UPROPERTY(VisibleAnywhere, Category = "Emitters")
+	UPROPERTY()
 	TArray<FNiagaraEmitterHandle> EmitterHandles;
 
 	UPROPERTY(EditAnywhere, Category="System")
diff --git a/Engine/Plugins/FX/Niagara/Source/NiagaraEditor/Private/Toolkits/NiagaraSystemToolkit.cpp b/Engine/Plugins/FX/Niagara/Source/NiagaraEditor/Private/Toolkits/NiagaraSystemToolkit.cpp
--- a/Engine/Plugins/FX/Niagara/Source/NiagaraEditor/Private/Toolkits/NiagaraSystemToolkit.cpp
+++ b/Engine/Plugins/FX/Niagara/Source/NiagaraEditor/Private/Toolkits/NiagaraSystemToolkit.cpp
@@ -109,10 +109,6 @@ void FNiagaraSystemToolkit::RegisterTabSpawners(const TSharedRef<class FTabManag
 		.SetGroup(WorkspaceMenuCategory.ToSharedRef())
 		.SetAutoGenerateMenuEntry(GbShowNiagaraDeveloperWindows != 0);
 
-	InTabManager->RegisterTabSpawner(SystemDetailsTabID, FOnSpawnTab::CreateSP(this, &FNiagaraSystemToolkit::SpawnTab_SystemDetails))
-		.SetDisplayName(LOCTEXT("SystemDetails", "System Details"))
-		.SetGroup(WorkspaceMenuCategory.ToSharedRef());
-
 	InTabManager->RegisterTabSpawner(SystemParametersTabID, FOnSpawnTab::CreateSP(this, &FNiagaraSystemToolkit::SpawnTab_SystemParameters))
 		.SetDisplayName(LOCTEXT("SystemParameters", "Parameters"))
 		.SetGroup(WorkspaceMenuCategory.ToSharedRef());
@@ -484,23 +480,6 @@ TSharedRef<SDockTab> FNiagaraSystemToolkit::SpawnTab_SystemScript(const FSpawnTa
 	return SpawnedTab;
 }
 
-
-TSharedRef<SDockTab> FNiagaraSystemToolkit::SpawnTab_SystemDetails(const FSpawnTabArgs& Args)
-{
-	check(Args.GetTabId().TabType == SystemDetailsTabID);
-
-	TSharedRef<FNiagaraObjectSelection> SystemSelection = MakeShareable(new FNiagaraObjectSelection());
-	SystemSelection->SetSelectedObject(System);
-
-	TSharedRef<SDockTab> SpawnedTab =
-		SNew(SDockTab)
-		[
-			SNew(SNiagaraSelectedObjectsDetails, SystemSelection)
-		];
-
-	return SpawnedTab;
-}
-
 TSharedRef<SDockTab> FNiagaraSystemToolkit::SpawnTab_SystemParameters(const FSpawnTabArgs& Args)
 {
 	check(Args.GetTabId().TabType == SystemParametersTabID);
diff --git a/Engine/Plugins/FX/Niagara/Source/NiagaraEditor/Private/Toolkits/NiagaraSystemToolkit.h b/Engine/Plugins/FX/Niagara/Source/NiagaraEditor/Private/Toolkits/NiagaraSystemToolkit.h
--- a/Engine/Plugins/FX/Niagara/Source/NiagaraEditor/Private/Toolkits/NiagaraSystemToolkit.h
+++ b/Engine/Plugins/FX/Niagara/Source/NiagaraEditor/Private/Toolkits/NiagaraSystemToolkit.h
@@ -94,7 +94,6 @@ private:
 	TSharedRef<SDockTab> SpawnTab_CurveEd(const FSpawnTabArgs& Args);
 	TSharedRef<SDockTab> SpawnTab_Sequencer(const FSpawnTabArgs& Args);
 	TSharedRef<SDockTab> SpawnTab_SystemScript(const FSpawnTabArgs& Args);
-	TSharedRef<SDockTab> SpawnTab_SystemDetails(const FSpawnTabArgs& Args);
 	TSharedRef<SDockTab> SpawnTab_SystemParameters(const FSpawnTabArgs& Args);
 	TSharedRef<SDockTab> SpawnTab_SelectedEmitterStack(const FSpawnTabArgs& Args);
 	TSharedRef<SDockTab> SpawnTab_SelectedEmitterGraph(const FSpawnTabArgs& Args);
diff --git a/Engine/Plugins/FX/Niagara/Source/NiagaraEditor/Private/ViewModels/Stack/NiagaraStackEmitterSettingsGroup.cpp b/Engine/Plugins/FX/Niagara/Source/NiagaraEditor/Private/ViewModels/Stack/NiagaraStackEmitterSettingsGroup.cpp
new file mode 100644
--- /dev/null
+++ b/Engine/Plugins/FX/Niagara/Source/NiagaraEditor/Private/ViewModels/Stack/NiagaraStackEmitterSettingsGroup.cpp
@@ -0,0 +1,125 @@
+// Copyright 1998-2019 Epic Games, Inc. All Rights Reserved.
+
+#include "ViewModels/Stack/NiagaraStackEmitterSettingsGroup.h"
+#include "ViewModels/Stack/NiagaraStackObject.h"
+#include "ViewModels/Stack/NiagaraStackSpacer.h"
+#include "ViewModels/NiagaraSystemViewModel.h"
+#include "ViewModels/NiagaraEmitterViewModel.h"
+#include "NiagaraEmitter.h"
+#include "NiagaraStackEditorData.h"
+#include "ViewModels/Stack/NiagaraStackGraphUtilities.h"
+#include "NiagaraScriptMergeManager.h"
+#include "NiagaraEmitterDetailsCustomization.h"
+
+#define LOCTEXT_NAMESPACE "UNiagaraStackEmitterItemGroup"
+
+void UNiagaraStackEmitterPropertiesItem::Initialize(FRequiredEntryData InRequiredEntryData)
+{
+	Super::Initialize(InRequiredEntryData, TEXT("EmitterProperties"));
+	Emitter = GetEmitterViewModel()->GetEmitter();
+	Emitter->OnPropertiesChanged().AddUObject(this, &UNiagaraStackEmitterPropertiesItem::EmitterPropertiesChanged);
+}
+
+void UNiagaraStackEmitterPropertiesItem::FinalizeInternal()
+{
+	if (Emitter.IsValid())
+	{
+		Emitter->OnPropertiesChanged().RemoveAll(this);
+	}
+	Super::FinalizeInternal();
+}
+
+FText UNiagaraStackEmitterPropertiesItem::GetDisplayName() const
+{
+	return LOCTEXT("EmitterPropertiesName", "Emitter Properties");
+}
+
+FText UNiagaraStackEmitterPropertiesItem::GetTooltipText() const
+{
+	return LOCTEXT("EmitterPropertiesTooltip", "Properties that are handled per Emitter. These cannot change at runtime.");
+}
+
+bool UNiagaraStackEmitterPropertiesItem::CanResetToBase() const
+{
+	if (bCanResetToBaseCache.IsSet() == false)
+	{
+		const UNiagaraEmitter* BaseEmitter = GetEmitterViewModel()->GetEmitter()->GetParent();
+		if (BaseEmitter != nullptr && Emitter != BaseEmitter)
+		{
+			TSharedRef<FNiagaraScriptMergeManager> MergeManager = FNiagaraScriptMergeManager::Get();
+			bCanResetToBaseCache = MergeManager->IsEmitterEditablePropertySetDifferentFromBase(*Emitter.Get(), *BaseEmitter);
+		}
+		else
+		{
+			bCanResetToBaseCache = false;
+		}
+	}
+	return bCanResetToBaseCache.GetValue();
+}
+
+void UNiagaraStackEmitterPropertiesItem::ResetToBase()
+{
+	if (CanResetToBase())
+	{
+		const UNiagaraEmitter* BaseEmitter = GetEmitterViewModel()->GetEmitter()->GetParent();
+		TSharedRef<FNiagaraScriptMergeManager> MergeManager = FNiagaraScriptMergeManager::Get();
+		MergeManager->ResetEmitterEditablePropertySetToBase(*Emitter, *BaseEmitter);
+	}
+}
+
+bool UNiagaraStackEmitterPropertiesItem::IsExpandedByDefault() const
+{
+	return false;
+}
+
+void UNiagaraStackEmitterPropertiesItem::RefreshChildrenInternal(const TArray<UNiagaraStackEntry*>& CurrentChildren, TArray<UNiagaraStackEntry*>& NewChildren, TArray<FStackIssue>& NewIssues)
+{
+	if (EmitterObject == nullptr)
+	{
+		EmitterObject = NewObject<UNiagaraStackObject>(this);
+		FRequiredEntryData RequiredEntryData(GetSystemViewModel(), GetEmitterViewModel(), FExecutionCategoryNames::Emitter, NAME_None, GetStackEditorData());
+		EmitterObject->Initialize(RequiredEntryData, Emitter.Get(), GetStackEditorDataKey());
+		EmitterObject->RegisterInstancedCustomPropertyLayout(UNiagaraEmitter::StaticClass(), FOnGetDetailCustomizationInstance::CreateStatic(&FNiagaraEmitterDetails::MakeInstance));
+	}
+
+	NewChildren.Add(EmitterObject);
+	bCanResetToBaseCache.Reset();
+	Super::RefreshChildrenInternal(CurrentChildren, NewChildren, NewIssues);
+}
+
+void UNiagaraStackEmitterPropertiesItem::EmitterPropertiesChanged()
+{
+	bCanResetToBaseCache.Reset();
+}
+
+UNiagaraStackEmitterSettingsGroup::UNiagaraStackEmitterSettingsGroup()
+	: PropertiesItem(nullptr)
+{
+}
+
+void UNiagaraStackEmitterSettingsGroup::RefreshChildrenInternal(const TArray<UNiagaraStackEntry*>& CurrentChildren, TArray<UNiagaraStackEntry*>& NewChildren, TArray<FStackIssue>& NewIssues)
+{
+	FName PropertiesSpacerKey = "PropertiesSpacer";
+	UNiagaraStackSpacer* PropertiesSpacer = FindCurrentChildOfTypeByPredicate<UNiagaraStackSpacer>(CurrentChildren,
+		[=](UNiagaraStackSpacer* CurrentPropertiesSpacer) { return CurrentPropertiesSpacer->GetSpacerKey() == PropertiesSpacerKey; });
+
+	if (PropertiesSpacer == nullptr)
+	{
+		PropertiesSpacer = NewObject<UNiagaraStackSpacer>(this);
+		PropertiesSpacer->Initialize(CreateDefaultChildRequiredData(), PropertiesSpacerKey);
+	}
+
+	NewChildren.Add(PropertiesSpacer);
+
+	if (PropertiesItem == nullptr)
+	{
+		PropertiesItem = NewObject<UNiagaraStackEmitterPropertiesItem>(this);
+		PropertiesItem->Initialize(CreateDefaultChildRequiredData());
+	}
+
+	NewChildren.Add(PropertiesItem);
+
+	Super::RefreshChildrenInternal(CurrentChildren, NewChildren, NewIssues);
+}
+
+#undef LOCTEXT_NAMESPACE
diff --git a/Engine/Plugins/FX/Niagara/Source/NiagaraEditor/Private/ViewModels/Stack/NiagaraStackEmitterSpawnScriptItemGroup.cpp b/Engine/Plugins/FX/Niagara/Source/NiagaraEditor/Private/ViewModels/Stack/NiagaraStackEmitterSpawnScriptItemGroup.cpp
deleted file mode 100644
diff --git a/Engine/Plugins/FX/Niagara/Source/NiagaraEditor/Private/ViewModels/Stack/NiagaraStackEntry.cpp b/Engine/Plugins/FX/Niagara/Source/NiagaraEditor/Private/ViewModels/Stack/NiagaraStackEntry.cpp
--- a/Engine/Plugins/FX/Niagara/Source/NiagaraEditor/Private/ViewModels/Stack/NiagaraStackEntry.cpp
+++ b/Engine/Plugins/FX/Niagara/Source/NiagaraEditor/Private/ViewModels/Stack/NiagaraStackEntry.cpp
@@ -14,7 +14,7 @@ const FName UNiagaraStackEntry::FExecutionCategoryNames::Emitter = TEXT("Emitter
 const FName UNiagaraStackEntry::FExecutionCategoryNames::Particle = TEXT("Particle");
 const FName UNiagaraStackEntry::FExecutionCategoryNames::Render = TEXT("Render");
 
-const FName UNiagaraStackEntry::FExecutionSubcategoryNames::Parameters = TEXT("Parameters");
+const FName UNiagaraStackEntry::FExecutionSubcategoryNames::Settings = TEXT("Settings");
 const FName UNiagaraStackEntry::FExecutionSubcategoryNames::Spawn = TEXT("Spawn");
 const FName UNiagaraStackEntry::FExecutionSubcategoryNames::Update = TEXT("Update");
 const FName UNiagaraStackEntry::FExecutionSubcategoryNames::Event = TEXT("Event");
diff --git a/Engine/Plugins/FX/Niagara/Source/NiagaraEditor/Private/ViewModels/Stack/NiagaraStackParameterStoreGroup.cpp b/Engine/Plugins/FX/Niagara/Source/NiagaraEditor/Private/ViewModels/Stack/NiagaraStackParameterStoreGroup.cpp
deleted file mode 100644
diff --git a/Engine/Plugins/FX/Niagara/Source/NiagaraEditor/Private/ViewModels/Stack/NiagaraStackRoot.cpp b/Engine/Plugins/FX/Niagara/Source/NiagaraEditor/Private/ViewModels/Stack/NiagaraStackRoot.cpp
--- a/Engine/Plugins/FX/Niagara/Source/NiagaraEditor/Private/ViewModels/Stack/NiagaraStackRoot.cpp
+++ b/Engine/Plugins/FX/Niagara/Source/NiagaraEditor/Private/ViewModels/Stack/NiagaraStackRoot.cpp
@@ -2,7 +2,7 @@
 
 #include "ViewModels/Stack/NiagaraStackRoot.h"
 #include "ViewModels/Stack/NiagaraStackScriptItemGroup.h"
-#include "ViewModels/Stack/NiagaraStackEmitterSpawnScriptItemGroup.h"
+#include "Viewmodels/Stack/NiagaraStackEmitterSettingsGroup.h"
 #include "ViewModels/Stack/NiagaraStackRenderItemGroup.h"
 #include "ViewModels/Stack/NiagaraStackEventHandlerGroup.h"
 #include "ViewModels/Stack/NiagaraStackEventScriptItemGroup.h"
@@ -13,13 +13,14 @@
 #include "NiagaraSystem.h"
 #include "NiagaraSystemEditorData.h"
 #include "NiagaraEmitterEditorData.h"
-#include "ViewModels/Stack/NiagaraStackParameterStoreGroup.h"
+#include "ViewModels/Stack/NiagaraStackSystemSettingsGroup.h"
 #include "ViewModels/Stack/NiagaraStackSpacer.h"
 
 #define LOCTEXT_NAMESPACE "NiagaraStackViewModel"
 
 UNiagaraStackRoot::UNiagaraStackRoot()
-	: SystemExposedVariablesGroup(nullptr)
+	: SystemSettingsGroup(nullptr)
+	, EmitterSettingsGroup(nullptr)
 	, EmitterSpawnGroup(nullptr)
 	, EmitterUpdateGroup(nullptr)
 	, ParticleSpawnGroup(nullptr)
@@ -32,13 +33,14 @@ UNiagaraStackRoot::UNiagaraStackRoot()
 void UNiagaraStackRoot::Initialize(FRequiredEntryData InRequiredEntryData)
 {
 	Super::Initialize(InRequiredEntryData, FString());
+	EmitterSettingsGroup = nullptr;
 	EmitterSpawnGroup = nullptr;
 	EmitterUpdateGroup = nullptr;
 	ParticleSpawnGroup = nullptr;
 	ParticleUpdateGroup = nullptr;
 	AddEventHandlerGroup = nullptr;
 	RenderGroup = nullptr;
-	SystemExposedVariablesGroup = nullptr;
+	SystemSettingsGroup = nullptr;
 }
 
 bool UNiagaraStackRoot::GetCanExpand() const
@@ -57,13 +59,13 @@ void UNiagaraStackRoot::RefreshChildrenInternal(const TArray<UNiagaraStackEntry*
 	bool bShowSystemGroups = GetSystemViewModel()->GetEditMode() == ENiagaraSystemViewModelEditMode::SystemAsset;
 
 	// Create static entries as needed.
-	if (bShowSystemGroups && SystemExposedVariablesGroup == nullptr)
+	if (bShowSystemGroups && SystemSettingsGroup == nullptr)
 	{
-		SystemExposedVariablesGroup = NewObject<UNiagaraStackParameterStoreGroup>(this);
+		SystemSettingsGroup = NewObject<UNiagaraStackSystemSettingsGroup>(this);
 		FRequiredEntryData RequiredEntryData(GetSystemViewModel(), GetEmitterViewModel(),
-			FExecutionCategoryNames::System, FExecutionSubcategoryNames::Parameters,
+			FExecutionCategoryNames::System, FExecutionSubcategoryNames::Settings,
 			GetSystemViewModel()->GetOrCreateEditorData().GetStackEditorData());
-		SystemExposedVariablesGroup->Initialize(RequiredEntryData, &GetSystemViewModel()->GetSystem(), &GetSystemViewModel()->GetSystem().GetExposedParameters());
+		SystemSettingsGroup->Initialize(RequiredEntryData, &GetSystemViewModel()->GetSystem(), &GetSystemViewModel()->GetSystem().GetExposedParameters());
 	}
 
 	if (bShowSystemGroups && SystemSpawnGroup == nullptr)
@@ -88,9 +90,20 @@ void UNiagaraStackRoot::RefreshChildrenInternal(const TArray<UNiagaraStackEntry*
 		SystemUpdateGroup->Initialize(RequiredEntryData, DisplayName, ToolTip, GetSystemViewModel()->GetSystemScriptViewModel().ToSharedRef(), ENiagaraScriptUsage::SystemUpdateScript);
 	}
 
+	if (EmitterSettingsGroup == nullptr)
+	{
+		EmitterSettingsGroup = NewObject<UNiagaraStackEmitterSettingsGroup>(this);
+		FRequiredEntryData RequiredEntryData(GetSystemViewModel(), GetEmitterViewModel(),
+			FExecutionCategoryNames::Emitter, FExecutionSubcategoryNames::Settings,
+			GetEmitterViewModel()->GetOrCreateEditorData().GetStackEditorData());
+		FText DisplayName = LOCTEXT("EmitterSettingsGroupName", "Emitter Settings");
+		FText Tooltip = LOCTEXT("EmitterSettingsTooltip", "Settings that are handled per Emitter.");
+		EmitterSettingsGroup->Initialize(RequiredEntryData, DisplayName, Tooltip, nullptr);
+	}
+
 	if (EmitterSpawnGroup == nullptr)
 	{
-		EmitterSpawnGroup = NewObject<UNiagaraStackEmitterSpawnScriptItemGroup>(this);
+		EmitterSpawnGroup = NewObject<UNiagaraStackScriptItemGroup>(this);
 		FRequiredEntryData RequiredEntryData(GetSystemViewModel(), GetEmitterViewModel(),
 			FExecutionCategoryNames::Emitter, FExecutionSubcategoryNames::Spawn,
 			GetEmitterViewModel()->GetOrCreateEditorData().GetStackEditorData());
@@ -171,13 +184,14 @@ void UNiagaraStackRoot::RefreshChildrenInternal(const TArray<UNiagaraStackEntry*
 	// Populate new children
 	if (bShowSystemGroups)
 	{
-		NewChildren.Add(SystemExposedVariablesGroup);
+		NewChildren.Add(SystemSettingsGroup);
 		NewChildren.Add(SystemSpawnGroup);
 		NewChildren.Add(SystemUpdateGroup);
 		NewChildren.Add(GetOrCreateSpacer(FExecutionCategoryNames::System, "SystemFooter"));
 		NewChildren.Add(GetOrCreateSpacer(NAME_None, "SystemSpacer"));
 	}
 
+	NewChildren.Add(EmitterSettingsGroup);
 	NewChildren.Add(EmitterSpawnGroup);
 	NewChildren.Add(EmitterUpdateGroup);
 	NewChildren.Add(GetOrCreateSpacer(FExecutionCategoryNames::Emitter, "EmitterFooter"));
diff --git a/Engine/Plugins/FX/Niagara/Source/NiagaraEditor/Private/ViewModels/Stack/NiagaraStackSystemSettingsGroup.cpp b/Engine/Plugins/FX/Niagara/Source/NiagaraEditor/Private/ViewModels/Stack/NiagaraStackSystemSettingsGroup.cpp
new file mode 100644
--- /dev/null
+++ b/Engine/Plugins/FX/Niagara/Source/NiagaraEditor/Private/ViewModels/Stack/NiagaraStackSystemSettingsGroup.cpp
@@ -0,0 +1,230 @@
+// Copyright 1998-2019 Epic Games, Inc. All Rights Reserved.
+
+#include "ViewModels/Stack/NiagaraStackSystemSettingsGroup.h"
+#include "ViewModels/Stack/NiagaraStackModuleItem.h"
+#include "NiagaraScriptViewModel.h"
+#include "NiagaraScriptGraphViewModel.h"
+#include "NiagaraGraph.h"
+#include "NiagaraNodeOutput.h"
+#include "NiagaraNodeFunctionCall.h"
+#include "EdGraphSchema_Niagara.h"
+#include "ViewModels/Stack/NiagaraStackErrorItem.h"
+#include "ViewModels/Stack/NiagaraStackParameterStoreEntry.h"
+#include "Internationalization/Internationalization.h"
+#include "ViewModels/Stack/NiagaraStackSpacer.h"
+#include "NiagaraStackEditorData.h"
+#include "ScopedTransaction.h"
+#include "ViewModels/Stack/NiagaraStackGraphUtilities.h"
+#include "ViewModels/Stack/NiagaraStackSystemPropertiesItem.h"
+
+
+#define LOCTEXT_NAMESPACE "UNiagaraStackParameterStoreGroup"
+
+class FParameterStoreGroupAddAction : public INiagaraStackItemGroupAddAction
+{
+public:
+	FParameterStoreGroupAddAction(FNiagaraVariable InNewParameterVariable)
+		: NewParameterVariable(InNewParameterVariable)
+	{
+	}
+
+	FNiagaraVariable GetNewParameterVariable() const
+	{
+		return NewParameterVariable;
+	}
+
+	virtual FText GetCategory() const override
+	{
+		return LOCTEXT("CreateNewParameterCategory", "Parameter Types");
+	}
+
+	virtual FText GetDisplayName() const override
+	{
+		return NewParameterVariable.GetType().GetNameText();
+	}
+
+	virtual FText GetDescription() const override
+	{
+		return FText::Format(LOCTEXT("AddParameterActionDescriptionFormat", "Create a new {0} parameter."), GetDisplayName());
+	}
+
+	virtual FText GetKeywords() const override
+	{
+		return FText();
+	}
+
+private:
+	FNiagaraVariable NewParameterVariable;
+};
+
+class FParameterStoreGroupAddUtiliites : public TNiagaraStackItemGroupAddUtilities<FNiagaraVariable>
+{
+public:
+	FParameterStoreGroupAddUtiliites(UObject& InParameterStoreOwner, FNiagaraParameterStore& InParameterStore, UNiagaraStackEditorData& InStackEditorData, FOnItemAdded InOnItemAdded)
+		: TNiagaraStackItemGroupAddUtilities(LOCTEXT("ScriptGroupAddItemName", "Parameter"), EAddMode::AddFromAction, true, InOnItemAdded)
+		, ParameterStoreOwner(InParameterStoreOwner)
+		, ParameterStore(InParameterStore)
+		, StackEditorData(InStackEditorData)
+	{
+	}
+
+	virtual void AddItemDirectly() override { unimplemented(); }
+
+	virtual void GenerateAddActions(TArray<TSharedRef<INiagaraStackItemGroupAddAction>>& OutAddActions, const FNiagaraStackItemGroupAddOptions& AddProperties) const override
+	{
+		TArray<FNiagaraTypeDefinition> AvailableTypes;
+		FNiagaraStackGraphUtilities::GetNewParameterAvailableTypes(AvailableTypes, FNiagaraParameterHandle::UserNamespace);
+		for (const FNiagaraTypeDefinition& AvailableType : AvailableTypes)
+		{
+			FNiagaraParameterHandle NewParameterHandle(FNiagaraParameterHandle::UserNamespace, *(TEXT("New") + AvailableType.GetName()));
+			FNiagaraVariable NewParameterVariable(AvailableType, NewParameterHandle.GetParameterHandleString());
+			OutAddActions.Add(MakeShared<FParameterStoreGroupAddAction>(NewParameterVariable));
+		}
+	}
+
+	virtual void ExecuteAddAction(TSharedRef<INiagaraStackItemGroupAddAction> AddAction, int32 TargetIndex) override
+	{
+		TSharedRef<FParameterStoreGroupAddAction> ParameterAddAction = StaticCastSharedRef<FParameterStoreGroupAddAction>(AddAction);
+
+		FScopedTransaction AddTransaction(LOCTEXT("AddParameter", "Add Parameter"));
+		ParameterStoreOwner.Modify();
+
+		FNiagaraVariable ParameterVariable = ParameterAddAction->GetNewParameterVariable();
+		FNiagaraEditorUtilities::ResetVariableToDefaultValue(ParameterVariable);
+
+		ParameterStore.AddParameter(ParameterVariable);
+		StackEditorData.SetModuleInputIsRenamePending(ParameterVariable.GetName().ToString(), true);
+
+		OnItemAdded.ExecuteIfBound(ParameterVariable);
+	}
+
+private:
+	UObject& ParameterStoreOwner;
+	FNiagaraParameterStore& ParameterStore;
+	UNiagaraStackEditorData& StackEditorData;
+};
+
+void UNiagaraStackSystemSettingsGroup::Initialize(
+	FRequiredEntryData InRequiredEntryData,
+	UObject* InOwner,
+	FNiagaraParameterStore* InParameterStore)
+{
+	AddUtilities = MakeShared<FParameterStoreGroupAddUtiliites>(*InOwner, *InParameterStore, *InRequiredEntryData.StackEditorData,
+		FParameterStoreGroupAddUtiliites::FOnItemAdded::CreateUObject(this, &UNiagaraStackSystemSettingsGroup::ParameterAdded));
+	FText DisplayName = LOCTEXT("SystemSettingsGroupName", "System Settings");
+	FText Tooltip = LOCTEXT("SystemSettingsTooltip", "Settings of the System.");
+	Super::Initialize(InRequiredEntryData, DisplayName, Tooltip, AddUtilities.Get());
+
+	Owner = InOwner;
+	UserParameterStore = InParameterStore;
+}
+
+void UNiagaraStackSystemSettingsGroup::RefreshChildrenInternal(const TArray<UNiagaraStackEntry*>& CurrentChildren, TArray<UNiagaraStackEntry*>& NewChildren, TArray<FStackIssue>& NewIssues)
+{
+	if (UserParameterStore != nullptr)
+	{
+		FName UserParameterStoreItemSpacerKey = "UserParameterStoreSpacer";
+		UNiagaraStackSpacer* UserParameterStoreItemSpacer = FindCurrentChildOfTypeByPredicate<UNiagaraStackSpacer>(CurrentChildren,
+			[=](UNiagaraStackSpacer* CurrentParameterStoreItemSpacer) { return CurrentParameterStoreItemSpacer->GetSpacerKey() == UserParameterStoreItemSpacerKey; });
+
+		if (UserParameterStoreItemSpacer == nullptr)
+		{
+			UserParameterStoreItemSpacer = NewObject<UNiagaraStackSpacer>(this);
+			UserParameterStoreItemSpacer->Initialize(CreateDefaultChildRequiredData(), UserParameterStoreItemSpacerKey, 1.4f);
+		}
+
+		NewChildren.Add(UserParameterStoreItemSpacer);
+
+		UNiagaraStackParameterStoreItem* UserParameterStoreItem = FindCurrentChildOfTypeByPredicate<UNiagaraStackParameterStoreItem>(CurrentChildren,
+			[=](UNiagaraStackParameterStoreItem* CurrentItem) { return true; });
+
+		if (UserParameterStoreItem == nullptr)
+		{
+			UserParameterStoreItem = NewObject<UNiagaraStackParameterStoreItem>(this);
+			UserParameterStoreItem->Initialize(CreateDefaultChildRequiredData(), Owner.Get(), UserParameterStore);
+		}
+
+		NewChildren.Add(UserParameterStoreItem);
+	}
+
+	FName SystemPropertiesItemSpacerKey = "SystemPropertiesSpacer";
+	UNiagaraStackSpacer* SystemPropertiesItemSpacer = FindCurrentChildOfTypeByPredicate<UNiagaraStackSpacer>(CurrentChildren,
+		[=](UNiagaraStackSpacer* CurrentSystemPropertiesItemSpacer) { return CurrentSystemPropertiesItemSpacer->GetSpacerKey() == SystemPropertiesItemSpacerKey; });
+
+	if (SystemPropertiesItemSpacer == nullptr)
+	{
+		SystemPropertiesItemSpacer = NewObject<UNiagaraStackSpacer>(this);
+		SystemPropertiesItemSpacer->Initialize(CreateDefaultChildRequiredData(), SystemPropertiesItemSpacerKey, 1.4f);
+	}
+	NewChildren.Add(SystemPropertiesItemSpacer);
+
+	UNiagaraStackSystemPropertiesItem* SystemPropertiesItem = FindCurrentChildOfTypeByPredicate<UNiagaraStackSystemPropertiesItem>(CurrentChildren,
+		[=](UNiagaraStackSystemPropertiesItem* CurrentItem) { return true; });
+
+	if (SystemPropertiesItem == nullptr)
+	{
+		SystemPropertiesItem = NewObject<UNiagaraStackSystemPropertiesItem>(this);
+		SystemPropertiesItem->Initialize(CreateDefaultChildRequiredData());
+	}
+	NewChildren.Add(SystemPropertiesItem);
+
+	Super::RefreshChildrenInternal(CurrentChildren, NewChildren, NewIssues);
+}
+
+void UNiagaraStackSystemSettingsGroup::ParameterAdded(FNiagaraVariable AddedParameter)
+{
+	RefreshChildren();
+}
+
+void UNiagaraStackParameterStoreItem::Initialize(
+	FRequiredEntryData InRequiredEntryData,
+	UObject* InOwner,
+	FNiagaraParameterStore* InParameterStore)
+{
+	Super::Initialize(InRequiredEntryData, TEXT("ParameterStoreItem"));
+
+	Owner = InOwner;
+	ParameterStore = InParameterStore;
+}
+
+FText UNiagaraStackParameterStoreItem::GetDisplayName() const
+{
+	return LOCTEXT("ParameterItemDisplayName", "User Parameters");
+}
+
+FText UNiagaraStackParameterStoreItem::GetTooltipText() const
+{
+	return LOCTEXT("ParameterItemTooltip", "Displays the variables created in the User namespace. These variables are exposed to owning UComponents, blueprints, etc.");
+}
+
+void UNiagaraStackParameterStoreItem::RefreshChildrenInternal(const TArray<UNiagaraStackEntry*>& CurrentChildren, TArray<UNiagaraStackEntry*>& NewChildren, TArray<FStackIssue>& NewIssues)
+{
+	if (ParameterStore != nullptr && ParameterStore->GetNumParameters() > 0)
+	{
+		TArray<FNiagaraVariable> Variables;
+		ParameterStore->GetParameters(Variables);
+
+		for (FNiagaraVariable& Var : Variables)
+		{
+			UNiagaraStackParameterStoreEntry* ValueObjectEntry = FindCurrentChildOfTypeByPredicate<UNiagaraStackParameterStoreEntry>(CurrentChildren,
+				[=](UNiagaraStackParameterStoreEntry* CurrentEntry) { return CurrentEntry->GetDisplayName().ToString() == Var.GetName().ToString(); });
+
+			if (ValueObjectEntry == nullptr)
+			{
+				ValueObjectEntry = NewObject<UNiagaraStackParameterStoreEntry>(this);
+				ValueObjectEntry->Initialize(CreateDefaultChildRequiredData(), Owner.Get(), ParameterStore, Var.GetName().ToString(), Var.GetType(), GetStackEditorDataKey());
+				ValueObjectEntry->OnParameterDeleted().AddUObject(this, &UNiagaraStackParameterStoreItem::ParameterDeleted);
+			}
+
+			NewChildren.Add(ValueObjectEntry);
+		}
+	}
+	Super::RefreshChildrenInternal(CurrentChildren, NewChildren, NewIssues);
+}
+
+void UNiagaraStackParameterStoreItem::ParameterDeleted()
+{
+	RefreshChildren();
+}
+
+#undef LOCTEXT_NAMESPACE
diff --git a/Engine/Plugins/FX/Niagara/Source/NiagaraEditor/Public/ViewModels/Stack/NiagaraStackEmitterSettingsGroup.h b/Engine/Plugins/FX/Niagara/Source/NiagaraEditor/Public/ViewModels/Stack/NiagaraStackEmitterSettingsGroup.h
new file mode 100644
--- /dev/null
+++ b/Engine/Plugins/FX/Niagara/Source/NiagaraEditor/Public/ViewModels/Stack/NiagaraStackEmitterSettingsGroup.h
@@ -0,0 +1,63 @@
+// Copyright 1998-2019 Epic Games, Inc. All Rights Reserved.
+
+#pragma once
+
+#include "ViewModels/Stack/NiagaraStackScriptItemGroup.h"
+#include "ViewModels/Stack/NiagaraStackItem.h"
+#include "NiagaraStackEmitterSettingsGroup.generated.h"
+
+class FNiagaraEmitterViewModel;
+class FNiagaraScriptViewModel;
+class UNiagaraStackObject;
+class UNiagaraStackSpacer;
+
+UCLASS()
+class NIAGARAEDITOR_API UNiagaraStackEmitterPropertiesItem : public UNiagaraStackItem
+{
+	GENERATED_BODY()
+
+public:
+	void Initialize(FRequiredEntryData InRequiredEntryData);
+
+	virtual FText GetDisplayName() const override;
+
+	virtual FText GetTooltipText() const override;
+
+	bool CanResetToBase() const;
+
+	void ResetToBase();
+
+	virtual bool IsExpandedByDefault() const override;
+
+protected:
+	virtual void FinalizeInternal() override;
+
+	virtual void RefreshChildrenInternal(const TArray<UNiagaraStackEntry*>& CurrentChildren, TArray<UNiagaraStackEntry*>& NewChildren, TArray<FStackIssue>& NewIssues) override;
+
+private:
+	void EmitterPropertiesChanged();
+
+private:
+	mutable TOptional<bool> bCanResetToBaseCache;
+
+	TWeakObjectPtr<UNiagaraEmitter> Emitter;
+
+	UPROPERTY()
+	UNiagaraStackObject* EmitterObject;
+};
+
+UCLASS()
+class NIAGARAEDITOR_API UNiagaraStackEmitterSettingsGroup : public UNiagaraStackItemGroup
+{
+	GENERATED_BODY()
+
+public:
+	UNiagaraStackEmitterSettingsGroup();
+
+protected:
+	void RefreshChildrenInternal(const TArray<UNiagaraStackEntry*>& CurrentChildren, TArray<UNiagaraStackEntry*>& NewChildren, TArray<FStackIssue>& NewIssues) override;
+
+private:
+	UPROPERTY()
+	UNiagaraStackEmitterPropertiesItem* PropertiesItem;
+};
diff --git a/Engine/Plugins/FX/Niagara/Source/NiagaraEditor/Public/ViewModels/Stack/NiagaraStackEmitterSpawnScriptItemGroup.h b/Engine/Plugins/FX/Niagara/Source/NiagaraEditor/Public/ViewModels/Stack/NiagaraStackEmitterSpawnScriptItemGroup.h
deleted file mode 100644
diff --git a/Engine/Plugins/FX/Niagara/Source/NiagaraEditor/Public/ViewModels/Stack/NiagaraStackEntry.h b/Engine/Plugins/FX/Niagara/Source/NiagaraEditor/Public/ViewModels/Stack/NiagaraStackEntry.h
--- a/Engine/Plugins/FX/Niagara/Source/NiagaraEditor/Public/ViewModels/Stack/NiagaraStackEntry.h
+++ b/Engine/Plugins/FX/Niagara/Source/NiagaraEditor/Public/ViewModels/Stack/NiagaraStackEntry.h
@@ -55,7 +55,7 @@ public:
 
 	struct NIAGARAEDITOR_API FExecutionSubcategoryNames
 	{
-		static const FName Parameters;
+		static const FName Settings;
 		static const FName Spawn;
 		static const FName Update;
 		static const FName Event;
diff --git a/Engine/Plugins/FX/Niagara/Source/NiagaraEditor/Public/ViewModels/Stack/NiagaraStackParameterStoreGroup.h b/Engine/Plugins/FX/Niagara/Source/NiagaraEditor/Public/ViewModels/Stack/NiagaraStackParameterStoreGroup.h
deleted file mode 100644
diff --git a/Engine/Plugins/FX/Niagara/Source/NiagaraEditor/Public/ViewModels/Stack/NiagaraStackRoot.h b/Engine/Plugins/FX/Niagara/Source/NiagaraEditor/Public/ViewModels/Stack/NiagaraStackRoot.h
--- a/Engine/Plugins/FX/Niagara/Source/NiagaraEditor/Public/ViewModels/Stack/NiagaraStackRoot.h
+++ b/Engine/Plugins/FX/Niagara/Source/NiagaraEditor/Public/ViewModels/Stack/NiagaraStackRoot.h
@@ -6,13 +6,11 @@
 #include "NiagaraStackRoot.generated.h"
 
 class FNiagaraEmitterViewModel;
-class UNiagaraStackEmitterSpawnScriptItemGroup;
-class UNiagaraStackScriptItemGroup;
-class UNiagaraStackScriptItemGroup;
+class UNiagaraStackSystemSettingsGroup;
+class UNiagaraStackEmitterSettingsGroup;
 class UNiagaraStackScriptItemGroup;
 class UNiagaraStackEventHandlerGroup;
 class UNiagaraStackRenderItemGroup;
-class UNiagaraStackParameterStoreGroup;
 
 UCLASS()
 class NIAGARAEDITOR_API UNiagaraStackRoot : public UNiagaraStackEntry
@@ -34,6 +32,9 @@ private:
 	void EmitterEventArraysChanged();
 
 private:
+	UPROPERTY()
+	UNiagaraStackSystemSettingsGroup* SystemSettingsGroup;
+
 	UPROPERTY()
 	UNiagaraStackScriptItemGroup* SystemSpawnGroup;
 
@@ -41,10 +42,10 @@ private:
 	UNiagaraStackScriptItemGroup* SystemUpdateGroup;
 
 	UPROPERTY()
-	UNiagaraStackParameterStoreGroup* SystemExposedVariablesGroup;
+	UNiagaraStackEmitterSettingsGroup* EmitterSettingsGroup;
 
 	UPROPERTY()
-	UNiagaraStackEmitterSpawnScriptItemGroup* EmitterSpawnGroup;
+	UNiagaraStackScriptItemGroup* EmitterSpawnGroup;
 
 	UPROPERTY()
 	UNiagaraStackScriptItemGroup* EmitterUpdateGroup;
diff --git a/Engine/Plugins/FX/Niagara/Source/NiagaraEditor/Public/ViewModels/Stack/NiagaraStackSystemSettingsGroup.h b/Engine/Plugins/FX/Niagara/Source/NiagaraEditor/Public/ViewModels/Stack/NiagaraStackSystemSettingsGroup.h
new file mode 100644
--- /dev/null
+++ b/Engine/Plugins/FX/Niagara/Source/NiagaraEditor/Public/ViewModels/Stack/NiagaraStackSystemSettingsGroup.h
@@ -0,0 +1,55 @@
+// Copyright 1998-2019 Epic Games, Inc. All Rights Reserved.
+
+#pragma once
+
+#include "ViewModels/Stack/NiagaraStackItemGroup.h"
+#include "NiagaraCommon.h"
+#include "NiagaraParameterStore.h"
+#include "ViewModels/Stack/NiagaraStackParameterStoreEntry.h"
+#include "NiagaraStackSystemSettingsGroup.generated.h"
+
+class FNiagaraScriptViewModel;
+
+UCLASS()
+class NIAGARAEDITOR_API UNiagaraStackSystemSettingsGroup : public UNiagaraStackItemGroup
+{
+	GENERATED_BODY()
+		
+public:
+	void Initialize(FRequiredEntryData InRequiredEntryData,	UObject* InOwner, FNiagaraParameterStore* InParameterStore);
+
+protected:
+	virtual void RefreshChildrenInternal(const TArray<UNiagaraStackEntry*>& CurrentChildren, TArray<UNiagaraStackEntry*>& NewChildren, TArray<FStackIssue>& NewIssues) override;
+
+private:
+	void ParameterAdded(FNiagaraVariable AddedParameter);
+
+private:
+	TWeakObjectPtr<UObject> Owner;
+	FNiagaraParameterStore* UserParameterStore;
+	FDelegateHandle ParameterStoreChangedHandle;
+	TSharedPtr<INiagaraStackItemGroupAddUtilities> AddUtilities;
+};
+
+UCLASS()
+class UNiagaraStackParameterStoreItem : public UNiagaraStackItem
+{
+	GENERATED_BODY()
+
+public:
+	void Initialize(FRequiredEntryData InRequiredEntryData, UObject* InOwner, FNiagaraParameterStore* InParameterStore);
+
+	virtual FText GetDisplayName() const override;
+
+	virtual FText GetTooltipText() const override;
+
+protected:
+	virtual void RefreshChildrenInternal(const TArray<UNiagaraStackEntry*>& CurrentChildren, TArray<UNiagaraStackEntry*>& NewChildren, TArray<FStackIssue>& NewIssues) override;
+
+private:
+	void ParameterDeleted();
+
+private:
+	TWeakObjectPtr<UObject> Owner;
+	FNiagaraParameterStore* ParameterStore;
+};
\ No newline at end of file
diff --git a/Engine/Plugins/FX/Niagara/Source/NiagaraEditorWidgets/Private/NiagaraEditorWidgetsUtilities.cpp b/Engine/Plugins/FX/Niagara/Source/NiagaraEditorWidgets/Private/NiagaraEditorWidgetsUtilities.cpp
--- a/Engine/Plugins/FX/Niagara/Source/NiagaraEditorWidgets/Private/NiagaraEditorWidgetsUtilities.cpp
+++ b/Engine/Plugins/FX/Niagara/Source/NiagaraEditorWidgets/Private/NiagaraEditorWidgetsUtilities.cpp
@@ -31,7 +31,7 @@ FName FNiagaraStackEditorWidgetsUtilities::GetIconNameForExecutionSubcategory(FN
 {
 	if (bIsHighlighted)
 	{
-		if (ExecutionSubcategoryName == UNiagaraStackEntry::FExecutionSubcategoryNames::Parameters)
+		if (ExecutionSubcategoryName == UNiagaraStackEntry::FExecutionSubcategoryNames::Settings)
 		{
 			return "NiagaraEditor.Stack.ParametersIconHighlighted";
 		}
@@ -50,7 +50,7 @@ FName FNiagaraStackEditorWidgetsUtilities::GetIconNameForExecutionSubcategory(FN
 	}
 	else
 	{
-		if (ExecutionSubcategoryName == UNiagaraStackEntry::FExecutionSubcategoryNames::Parameters)
+		if (ExecutionSubcategoryName == UNiagaraStackEntry::FExecutionSubcategoryNames::Settings)
 		{
 			return "NiagaraEditor.Stack.ParametersIcon";
 		}
diff --git a/Engine/Plugins/FX/Niagara/Source/NiagaraEditorWidgets/Private/SNiagaraStack.cpp b/Engine/Plugins/FX/Niagara/Source/NiagaraEditorWidgets/Private/SNiagaraStack.cpp
--- a/Engine/Plugins/FX/Niagara/Source/NiagaraEditorWidgets/Private/SNiagaraStack.cpp
+++ b/Engine/Plugins/FX/Niagara/Source/NiagaraEditorWidgets/Private/SNiagaraStack.cpp
@@ -24,7 +24,7 @@
 #include "ViewModels/Stack/NiagaraStackModuleItem.h"
 #include "ViewModels/Stack/NiagaraStackRendererItem.h"
 #include "ViewModels/Stack/NiagaraStackParameterStoreEntry.h"
-#include "ViewModels/Stack/NiagaraStackEmitterSpawnScriptItemGroup.h"
+#include "ViewModels/Stack/NiagaraStackEmitterSettingsGroup.h"
 #include "ViewModels/Stack/NiagaraStackEventScriptItemGroup.h"
 #include "Framework/Application/SlateApplication.h"
 #include "IContentBrowserSingleton.h"
diff --git a/Engine/Plugins/FX/Niagara/Source/NiagaraEditorWidgets/Private/Stack/SNiagaraStackEmitterPropertiesItem.cpp b/Engine/Plugins/FX/Niagara/Source/NiagaraEditorWidgets/Private/Stack/SNiagaraStackEmitterPropertiesItem.cpp
--- a/Engine/Plugins/FX/Niagara/Source/NiagaraEditorWidgets/Private/Stack/SNiagaraStackEmitterPropertiesItem.cpp
+++ b/Engine/Plugins/FX/Niagara/Source/NiagaraEditorWidgets/Private/Stack/SNiagaraStackEmitterPropertiesItem.cpp
@@ -4,7 +4,7 @@
 #include "NiagaraEditorWidgetsStyle.h"
 #include "NiagaraEditorStyle.h"
 #include "EditorStyleSet.h"
-#include "ViewModels/Stack/NiagaraStackEmitterSpawnScriptItemGroup.h"
+#include "ViewModels/Stack/NiagaraStackEmitterSettingsGroup.h"
 #include "ViewModels/Stack/NiagaraStackViewModel.h"
 #include "Widgets/SBoxPanel.h"
 #include "Widgets/Text/STextBlock.h"
