Submit missing files

diff --git a/Engine/Plugins/FX/Niagara/Source/NiagaraEditor/Private/ViewModels/Stack/NiagaraStackSystemPropertiesItem.cpp b/Engine/Plugins/FX/Niagara/Source/NiagaraEditor/Private/ViewModels/Stack/NiagaraStackSystemPropertiesItem.cpp
new file mode 100644
--- /dev/null
+++ b/Engine/Plugins/FX/Niagara/Source/NiagaraEditor/Private/ViewModels/Stack/NiagaraStackSystemPropertiesItem.cpp
@@ -0,0 +1,50 @@
+// Copyright 1998-2018 Epic Games, Inc. All Rights Reserved.
+
+#include "ViewModels/Stack/NiagaraStackSystemPropertiesItem.h"
+#include "ViewModels/Stack/NiagaraStackObject.h"
+#include "NiagaraSystem.h"
+#include "ViewModels/NiagaraSystemViewModel.h"
+
+#define LOCTEXT_NAMESPACE "UNiagaraStackSystemItemGroup"
+
+void UNiagaraStackSystemPropertiesItem::Initialize(FRequiredEntryData InRequiredEntryData)
+{
+	Super::Initialize(InRequiredEntryData, TEXT("SystemProperties"));
+	System = &GetSystemViewModel()->GetSystem();
+}
+
+FText UNiagaraStackSystemPropertiesItem::GetDisplayName() const
+{
+	return LOCTEXT("SystemPropertiesName", "System Properties");
+}
+
+FText UNiagaraStackSystemPropertiesItem::GetTooltipText() const
+{
+	return LOCTEXT("SystemPropertiesTooltip", "Properties of the System. These cannot change at runtime.");
+}
+
+bool UNiagaraStackSystemPropertiesItem::IsExpandedByDefault() const
+{
+	return false;
+}
+
+void UNiagaraStackSystemPropertiesItem::RefreshChildrenInternal(const TArray<UNiagaraStackEntry*>& CurrentChildren, TArray<UNiagaraStackEntry*>& NewChildren, TArray<FStackIssue>& NewIssues)
+{
+	if (SystemObject == nullptr)
+	{
+		SystemObject = NewObject<UNiagaraStackObject>(this);
+		FRequiredEntryData RequiredEntryData(GetSystemViewModel(), GetEmitterViewModel(), FExecutionCategoryNames::System, NAME_None, GetStackEditorData());
+		SystemObject->Initialize(RequiredEntryData, System.Get(), GetStackEditorDataKey());
+	}
+
+	NewChildren.Add(SystemObject);
+	bCanResetToBase.Reset();
+	Super::RefreshChildrenInternal(CurrentChildren, NewChildren, NewIssues);
+}
+
+void UNiagaraStackSystemPropertiesItem::SystemPropertiesChanged()
+{
+	bCanResetToBase.Reset();
+}
+
+#undef LOCTEXT_NAMESPACE
diff --git a/Engine/Plugins/FX/Niagara/Source/NiagaraEditor/Public/ViewModels/Stack/NiagaraStackSystemPropertiesItem.h b/Engine/Plugins/FX/Niagara/Source/NiagaraEditor/Public/ViewModels/Stack/NiagaraStackSystemPropertiesItem.h
new file mode 100644
--- /dev/null
+++ b/Engine/Plugins/FX/Niagara/Source/NiagaraEditor/Public/ViewModels/Stack/NiagaraStackSystemPropertiesItem.h
@@ -0,0 +1,39 @@
+// Copyright 1998-2018 Epic Games, Inc. All Rights Reserved.
+
+#pragma once
+
+#include "ViewModels/Stack/NiagaraStackScriptItemGroup.h"
+#include "ViewModels/Stack/NiagaraStackItem.h"
+#include "NiagaraStackSystemPropertiesItem.generated.h"
+
+class UNiagaraStackObject;
+
+UCLASS()
+class NIAGARAEDITOR_API UNiagaraStackSystemPropertiesItem : public UNiagaraStackItem
+{
+	GENERATED_BODY()
+
+public:
+	void Initialize(FRequiredEntryData InRequiredEntryData);
+
+	virtual FText GetDisplayName() const override;
+
+	virtual FText GetTooltipText() const override;
+
+	virtual bool IsExpandedByDefault() const override;
+
+protected:
+	virtual void RefreshChildrenInternal(const TArray<UNiagaraStackEntry*>& CurrentChildren, TArray<UNiagaraStackEntry*>& NewChildren, TArray<FStackIssue>& NewIssues) override;
+
+private:
+	void SystemPropertiesChanged();
+
+private:
+	mutable TOptional<bool> bCanResetToBase;
+
+	TWeakObjectPtr<UNiagaraSystem> System;
+
+	UPROPERTY()
+	UNiagaraStackObject* SystemObject;
+
+};
