Fix non portable capitalization in include filepaths.

diff --git a/Engine/Plugins/FX/Niagara/Source/NiagaraEditor/Private/Toolkits/NiagaraSystemToolkit.cpp b/Engine/Plugins/FX/Niagara/Source/NiagaraEditor/Private/Toolkits/NiagaraSystemToolkit.cpp
--- a/Engine/Plugins/FX/Niagara/Source/NiagaraEditor/Private/Toolkits/NiagaraSystemToolkit.cpp
+++ b/Engine/Plugins/FX/Niagara/Source/NiagaraEditor/Private/Toolkits/NiagaraSystemToolkit.cpp
@@ -49,7 +49,7 @@
 #include "HAL/PlatformFilemanager.h"
 #include "Misc/FileHelper.h"
 #include "NiagaraMessageLogViewModel.h"
-#include "Viewmodels/NiagaraOverviewGraphViewModel.h"
+#include "ViewModels/NiagaraOverviewGraphViewModel.h"
 
 #define LOCTEXT_NAMESPACE "NiagaraSystemEditor"
 
diff --git a/Engine/Plugins/FX/Niagara/Source/NiagaraEditor/Private/ViewModels/NiagaraSystemViewModel.cpp b/Engine/Plugins/FX/Niagara/Source/NiagaraEditor/Private/ViewModels/NiagaraSystemViewModel.cpp
--- a/Engine/Plugins/FX/Niagara/Source/NiagaraEditor/Private/ViewModels/NiagaraSystemViewModel.cpp
+++ b/Engine/Plugins/FX/Niagara/Source/NiagaraEditor/Private/ViewModels/NiagaraSystemViewModel.cpp
@@ -36,7 +36,7 @@
 #include "Channels/MovieSceneChannelProxy.h"
 #include "Framework/Commands/UICommandList.h"
 #include "NiagaraMessageManager.h"
-#include "Viewmodels/NiagaraOverviewGraphViewModel.h"
+#include "ViewModels/NiagaraOverviewGraphViewModel.h"
 
 DECLARE_CYCLE_STAT(TEXT("Niagara - SystemViewModel - CompileSystem"), STAT_NiagaraEditor_SystemViewModel_CompileSystem, STATGROUP_NiagaraEditor);
 
