- Fix circular TSharedPtr's in Niagara causing editor resources to not be released

diff --git a/Engine/Plugins/FX/Niagara/Source/NiagaraEditor/Private/Widgets/SNiagaraSystemViewport.cpp b/Engine/Plugins/FX/Niagara/Source/NiagaraEditor/Private/Widgets/SNiagaraSystemViewport.cpp
--- a/Engine/Plugins/FX/Niagara/Source/NiagaraEditor/Private/Widgets/SNiagaraSystemViewport.cpp
+++ b/Engine/Plugins/FX/Niagara/Source/NiagaraEditor/Private/Widgets/SNiagaraSystemViewport.cpp
@@ -48,7 +48,7 @@ public:
 
 	void DrawInstructionCounts(UNiagaraSystem* ParticleSystem, FCanvas* Canvas, float& CurrentX, float& CurrentY, UFont* Font, const float FontHeight);
 
-	TSharedPtr<SNiagaraSystemViewport> NiagaraViewport;
+	TWeakPtr<SNiagaraSystemViewport> NiagaraViewportPtr;
 	bool bCaptureScreenShot;
 	TWeakObjectPtr<UObject> ScreenShotOwner;
 
@@ -59,7 +59,7 @@ FNiagaraSystemViewportClient::FNiagaraSystemViewportClient(FAdvancedPreviewScene
 	: FEditorViewportClient(nullptr, &InPreviewScene, StaticCastSharedRef<SEditorViewport>(InNiagaraEditorViewport))
 	, OnScreenShotCaptured(InOnScreenShotCaptured)
 {
-	NiagaraViewport = InNiagaraEditorViewport;
+	NiagaraViewportPtr = InNiagaraEditorViewport;
 
 	// Setup defaults for the common draw helper.
 	DrawHelper.bDrawPivot = false;
@@ -99,6 +99,7 @@ void FNiagaraSystemViewportClient::Tick(float DeltaSeconds)
 
 void FNiagaraSystemViewportClient::Draw(FViewport* InViewport, FCanvas* Canvas)
 {
+	TSharedPtr<SNiagaraSystemViewport> NiagaraViewport = NiagaraViewportPtr.Pin();
 	UNiagaraSystem* ParticleSystem = NiagaraViewport.IsValid() ? NiagaraViewport->GetPreviewComponent()->GetAsset() : nullptr;
 
 	if (NiagaraViewport.IsValid() && NiagaraViewport->GetDrawElement(SNiagaraSystemViewport::EDrawElements::Bounds))
