Fixed crash when deleting Object User Variables.

diff --git a/Engine/Plugins/FX/Niagara/Source/NiagaraEditor/Private/ViewModels/Stack/NiagaraStackParameterStoreEntry.cpp b/Engine/Plugins/FX/Niagara/Source/NiagaraEditor/Private/ViewModels/Stack/NiagaraStackParameterStoreEntry.cpp
--- a/Engine/Plugins/FX/Niagara/Source/NiagaraEditor/Private/ViewModels/Stack/NiagaraStackParameterStoreEntry.cpp
+++ b/Engine/Plugins/FX/Niagara/Source/NiagaraEditor/Private/ViewModels/Stack/NiagaraStackParameterStoreEntry.cpp
@@ -285,7 +285,7 @@ void UNiagaraStackParameterStoreEntry::Delete()
 	//remove from store
 	Owner->Modify();
 	ParameterStore->RemoveParameter(FNiagaraVariable(InputType, ParameterName));
-	if (InputType.GetClass() != nullptr)
+	if (InputType.IsDataInterface())
 	{
 		UNiagaraDataInterface* DataInterface = NewObject<UNiagaraDataInterface>(this, const_cast<UClass*>(InputType.GetClass()));
 		if (DataInterface != nullptr)
