Fix non-portable capitalization in include.

diff --git a/Engine/Plugins/FX/Niagara/Source/NiagaraEditor/Private/ViewModels/Stack/NiagaraStackRoot.cpp b/Engine/Plugins/FX/Niagara/Source/NiagaraEditor/Private/ViewModels/Stack/NiagaraStackRoot.cpp
--- a/Engine/Plugins/FX/Niagara/Source/NiagaraEditor/Private/ViewModels/Stack/NiagaraStackRoot.cpp
+++ b/Engine/Plugins/FX/Niagara/Source/NiagaraEditor/Private/ViewModels/Stack/NiagaraStackRoot.cpp
@@ -2,7 +2,7 @@
 
 #include "ViewModels/Stack/NiagaraStackRoot.h"
 #include "ViewModels/Stack/NiagaraStackScriptItemGroup.h"
-#include "Viewmodels/Stack/NiagaraStackEmitterSettingsGroup.h"
+#include "ViewModels/Stack/NiagaraStackEmitterSettingsGroup.h"
 #include "ViewModels/Stack/NiagaraStackRenderItemGroup.h"
 #include "ViewModels/Stack/NiagaraStackEventHandlerGroup.h"
 #include "ViewModels/Stack/NiagaraStackEventScriptItemGroup.h"
