Some Niagara support for migration of legacy anim trails.

diff --git a/Engine/Plugins/FX/Niagara/Source/Niagara/Private/NiagaraFunctionLibrary.cpp b/Engine/Plugins/FX/Niagara/Source/Niagara/Private/NiagaraFunctionLibrary.cpp
--- a/Engine/Plugins/FX/Niagara/Source/Niagara/Private/NiagaraFunctionLibrary.cpp
+++ b/Engine/Plugins/FX/Niagara/Source/Niagara/Private/NiagaraFunctionLibrary.cpp
@@ -35,6 +35,7 @@ UNiagaraComponent* CreateNiagaraSystem(UNiagaraSystem* SystemTemplate, UWorld* W
 	NiagaraComponent->SetAutoDestroy(bAutoDestroy);
 	NiagaraComponent->bAllowAnyoneToDestroyMe = true;
 	NiagaraComponent->SetAsset(SystemTemplate);
+	NiagaraComponent->bAutoActivate = false;
 	return NiagaraComponent;
 }
 
@@ -43,7 +44,7 @@ UNiagaraComponent* CreateNiagaraSystem(UNiagaraSystem* SystemTemplate, UWorld* W
 * Spawns a Niagara System at the specified world location/rotation
 * @return			The spawned UNiagaraComponent
 */
-UNiagaraComponent* UNiagaraFunctionLibrary::SpawnSystemAtLocation(const UObject* WorldContextObject, UNiagaraSystem* SystemTemplate, FVector SpawnLocation, FRotator SpawnRotation, FVector Scale, bool bAutoDestroy)
+UNiagaraComponent* UNiagaraFunctionLibrary::SpawnSystemAtLocation(const UObject* WorldContextObject, UNiagaraSystem* SystemTemplate, FVector SpawnLocation, FRotator SpawnRotation, FVector Scale, bool bAutoDestroy, bool bAutoActivate)
 {
 	UNiagaraComponent* PSC = NULL;
 	if (SystemTemplate)
@@ -55,13 +56,15 @@ UNiagaraComponent* UNiagaraFunctionLibrary::SpawnSystemAtLocation(const UObject*
 #if WITH_EDITORONLY_DATA
 			PSC->bWaitForCompilationOnActivate = true;
 #endif
-			PSC->bAutoActivate = false;
 			PSC->RegisterComponentWithWorld(World);
 
 			PSC->SetAbsolute(true, true, true);
 			PSC->SetWorldLocationAndRotation(SpawnLocation, SpawnRotation);
 			PSC->SetRelativeScale3D(Scale);
-			PSC->Activate(true);
+			if (bAutoActivate)
+			{
+				PSC->Activate(true);
+			}
 		}
 	}
 	return PSC;
@@ -75,7 +78,7 @@ UNiagaraComponent* UNiagaraFunctionLibrary::SpawnSystemAtLocation(const UObject*
 * Spawns a Niagara System attached to a component
 * @return			The spawned UNiagaraComponent
 */
-UNiagaraComponent* UNiagaraFunctionLibrary::SpawnSystemAttached(UNiagaraSystem* SystemTemplate, USceneComponent* AttachToComponent, FName AttachPointName, FVector Location, FRotator Rotation, EAttachLocation::Type LocationType, bool bAutoDestroy)
+UNiagaraComponent* UNiagaraFunctionLibrary::SpawnSystemAttached(UNiagaraSystem* SystemTemplate, USceneComponent* AttachToComponent, FName AttachPointName, FVector Location, FRotator Rotation, EAttachLocation::Type LocationType, bool bAutoDestroy, bool bAutoActivate)
 {
 	UNiagaraComponent* PSC = nullptr;
 	if (SystemTemplate)
@@ -105,6 +108,11 @@ UNiagaraComponent* UNiagaraFunctionLibrary::SpawnSystemAttached(UNiagaraSystem*
 				PSC->SetRelativeLocationAndRotation(Location, Rotation);
 			}
 			PSC->SetRelativeScale3D(FVector(1.f));
+
+			if (bAutoActivate)
+			{
+				PSC->Activate();
+			}
 		}
 	}
 	return PSC;
@@ -124,7 +132,8 @@ UNiagaraComponent* UNiagaraFunctionLibrary::SpawnSystemAttached(
 	FVector Scale,
 	EAttachLocation::Type LocationType,
 	bool bAutoDestroy,
-	EPSCPoolMethod PoolingMethod
+	EPSCPoolMethod PoolingMethod,
+	bool bAutoActivate
 )
 {
 	UNiagaraComponent* PSC = nullptr;
@@ -178,7 +187,10 @@ UNiagaraComponent* UNiagaraFunctionLibrary::SpawnSystemAttached(
 					}
 
 					PSC->RegisterComponentWithWorld(World);
-					PSC->Activate(true);
+					if (bAutoActivate)
+					{
+						PSC->Activate(true);
+					}
 
 					// Notify the texture streamer so that PSC gets managed as a dynamic component.
 					IStreamingManager::Get().NotifyPrimitiveUpdated(PSC);
diff --git a/Engine/Plugins/FX/Niagara/Source/Niagara/Private/NiagaraModule.cpp b/Engine/Plugins/FX/Niagara/Source/Niagara/Private/NiagaraModule.cpp
--- a/Engine/Plugins/FX/Niagara/Source/Niagara/Private/NiagaraModule.cpp
+++ b/Engine/Plugins/FX/Niagara/Source/Niagara/Private/NiagaraModule.cpp
@@ -719,6 +719,8 @@ void FNiagaraTypeDefinition::RecreateUserDefinedTypeRegistry()
 	}
 
 	FNiagaraTypeRegistry::Register(FNiagaraRandInfo::StaticStruct(), true, true, true);
+
+	FNiagaraTypeRegistry::Register(StaticEnum<ENiagaraLegacyTrailWidthMode>(), true, true, false);
 }
 
 bool FNiagaraTypeDefinition::IsScalarDefinition(const FNiagaraTypeDefinition& Type)
diff --git a/Engine/Plugins/FX/Niagara/Source/Niagara/Public/NiagaraCommon.h b/Engine/Plugins/FX/Niagara/Source/Niagara/Public/NiagaraCommon.h
--- a/Engine/Plugins/FX/Niagara/Source/Niagara/Public/NiagaraCommon.h
+++ b/Engine/Plugins/FX/Niagara/Source/Niagara/Public/NiagaraCommon.h
@@ -605,3 +605,20 @@ struct FNiagaraRandInfo
 	UPROPERTY(EditAnywhere, Category = "Random")
 	int32 Seed3;
 };
+
+
+//////////////////////////////////////////////////////////////////////////
+// Legacy Anim Trail Support
+
+
+/** 
+Controls the way that the width scale property affects animation trails. 
+Only used for Legacy Anim Trail support when converting from Cascade to Niagara.
+*/
+UENUM()
+enum class ENiagaraLegacyTrailWidthMode : uint32
+{
+	FromCentre,
+	FromFirst,
+	FromSecond,
+};
diff --git a/Engine/Plugins/FX/Niagara/Source/Niagara/Public/NiagaraFunctionLibrary.h b/Engine/Plugins/FX/Niagara/Source/Niagara/Public/NiagaraFunctionLibrary.h
--- a/Engine/Plugins/FX/Niagara/Source/Niagara/Public/NiagaraFunctionLibrary.h
+++ b/Engine/Plugins/FX/Niagara/Source/Niagara/Public/NiagaraFunctionLibrary.h
@@ -27,12 +27,12 @@ class NIAGARA_API UNiagaraFunctionLibrary : public UBlueprintFunctionLibrary
 	* @return			The spawned UNiagaraComponent
 	*/
 	UFUNCTION(BlueprintCallable, Category = Niagara, meta = (Keywords = "niagara System", WorldContext = "WorldContextObject", UnsafeDuringActorConstruction = "true"))
-	static UNiagaraComponent* SpawnSystemAtLocation(const UObject* WorldContextObject, class UNiagaraSystem* SystemTemplate, FVector Location, FRotator Rotation = FRotator::ZeroRotator, FVector Scale = FVector(1.f), bool bAutoDestroy = true);
+	static UNiagaraComponent* SpawnSystemAtLocation(const UObject* WorldContextObject, class UNiagaraSystem* SystemTemplate, FVector Location, FRotator Rotation = FRotator::ZeroRotator, FVector Scale = FVector(1.f), bool bAutoDestroy = true, bool bAutoActivate = true);
 
 	UFUNCTION(BlueprintCallable, Category = Niagara, meta = (Keywords = "niagara System", UnsafeDuringActorConstruction = "true"))
-	static UNiagaraComponent* SpawnSystemAttached(UNiagaraSystem* SystemTemplate, USceneComponent* AttachToComponent, FName AttachPointName, FVector Location, FRotator Rotation, EAttachLocation::Type LocationType, bool bAutoDestroy);
+	static UNiagaraComponent* SpawnSystemAttached(UNiagaraSystem* SystemTemplate, USceneComponent* AttachToComponent, FName AttachPointName, FVector Location, FRotator Rotation, EAttachLocation::Type LocationType, bool bAutoDestroy, bool bAutoActivate = true);
 
-	static UNiagaraComponent* SpawnSystemAttached(UNiagaraSystem* SystemTemplate, USceneComponent* AttachToComponent, FName AttachPointName, FVector Location, FRotator Rotation, FVector Scale, EAttachLocation::Type LocationType, bool bAutoDestroy, EPSCPoolMethod PoolingMethod);
+	static UNiagaraComponent* SpawnSystemAttached(UNiagaraSystem* SystemTemplate, USceneComponent* AttachToComponent, FName AttachPointName, FVector Location, FRotator Rotation, FVector Scale, EAttachLocation::Type LocationType, bool bAutoDestroy, EPSCPoolMethod PoolingMethod, bool bAutoActivate = true);
 
 	/** Sets a Niagara StaticMesh parameter by name, overriding locally if necessary.*/
 	UFUNCTION(BlueprintCallable, Category = Niagara, meta = (DisplayName = "Set Niagara Static Mesh Component"))
diff --git a/Engine/Source/Runtime/Engine/Classes/Particles/ParticleSystemComponent.h b/Engine/Source/Runtime/Engine/Classes/Particles/ParticleSystemComponent.h
--- a/Engine/Source/Runtime/Engine/Classes/Particles/ParticleSystemComponent.h
+++ b/Engine/Source/Runtime/Engine/Classes/Particles/ParticleSystemComponent.h
@@ -359,6 +359,7 @@ public:
 	/** 
 	 * Get the referenced FXSystem asset.
 	*/
+	UFUNCTION(BlueprintCallable, Category = "Effects|Components|ParticleSystem")
 	virtual UFXSystemAsset* GetFXSystemAsset() const { return nullptr; };
 
 	/**
