- Increase Niagara DeferredDeletionQueue to 3 to remove stalls on the GT
- Only fence if we have things to delete

diff --git a/Engine/Plugins/FX/Niagara/Source/Niagara/Private/NiagaraWorldManager.cpp b/Engine/Plugins/FX/Niagara/Source/Niagara/Private/NiagaraWorldManager.cpp
--- a/Engine/Plugins/FX/Niagara/Source/Niagara/Private/NiagaraWorldManager.cpp
+++ b/Engine/Plugins/FX/Niagara/Source/Niagara/Private/NiagaraWorldManager.cpp
@@ -266,14 +266,21 @@ void FNiagaraWorldManager::Tick(float DeltaSeconds)
 		SystemSimulations.Remove(DeadSystem);
 	}
 
-	// Enqueue fence for deferred deletion
-	DeferredDeletionQueue[DeferredDeletionQueueIndex].Fence.BeginFence();
+	// Enqueue fence for deferred deletion if we need to wait on anything
+	if ( DeferredDeletionQueue[DeferredDeletionQueueIndex].Queue.Num() > 0 )
+	{
+		DeferredDeletionQueue[DeferredDeletionQueueIndex].Fence.BeginFence();
+	}
 
-	// Remove instances from previous frame
+	// Remove instances from oldest frame making sure they aren't in use on the RT
 	DeferredDeletionQueueIndex = (DeferredDeletionQueueIndex + 1) % NumDeferredQueues;
+	if ( DeferredDeletionQueue[DeferredDeletionQueueIndex].Queue.Num() > 0 )
 	{
-		SCOPE_CYCLE_COUNTER(STAT_NiagaraWorldManWaitOnRender);
-		DeferredDeletionQueue[DeferredDeletionQueueIndex].Fence.Wait();
+		if ( !DeferredDeletionQueue[DeferredDeletionQueueIndex].Fence.IsFenceComplete() )
+		{
+			SCOPE_CYCLE_COUNTER(STAT_NiagaraWorldManWaitOnRender);
+			DeferredDeletionQueue[DeferredDeletionQueueIndex].Fence.Wait();
+		}
+		DeferredDeletionQueue[DeferredDeletionQueueIndex].Queue.Empty();
 	}
-	DeferredDeletionQueue[DeferredDeletionQueueIndex].Queue.Empty();
 }
diff --git a/Engine/Plugins/FX/Niagara/Source/Niagara/Public/NiagaraWorldManager.h b/Engine/Plugins/FX/Niagara/Source/Niagara/Public/NiagaraWorldManager.h
--- a/Engine/Plugins/FX/Niagara/Source/Niagara/Public/NiagaraWorldManager.h
+++ b/Engine/Plugins/FX/Niagara/Source/Niagara/Public/NiagaraWorldManager.h
@@ -111,7 +111,7 @@ private:
 		TArray<TUniquePtr<FNiagaraSystemInstance>>	Queue;
 	};
 
-	static constexpr int NumDeferredQueues = 2;
+	static constexpr int NumDeferredQueues = 3;
 	int DeferredDeletionQueueIndex = 0;
 	FDeferredDeletionQueue DeferredDeletionQueue[NumDeferredQueues];
 };
