Bringing Niagara compiler logging in line with expectations from our live games and other Unreal systems.
Compiler errors for Niagara will have a strong UI impact but the game should still function properly, there will just be oddities in the visuals. It should be acted upon, but in no way should the game be blocked from a successful cook because of it. Therefore, we do a warning.
Compiler warnings for Niagara are meant for notification and should have a UI representation, but should be expected to still function properly and can be acted upon at the user's leisure. This makes them best logged as Display messages, as Log will not be shown in the cook.

diff --git a/Engine/Plugins/FX/Niagara/Source/Niagara/Private/NiagaraScript.cpp b/Engine/Plugins/FX/Niagara/Source/Niagara/Private/NiagaraScript.cpp
--- a/Engine/Plugins/FX/Niagara/Source/Niagara/Private/NiagaraScript.cpp
+++ b/Engine/Plugins/FX/Niagara/Source/Niagara/Private/NiagaraScript.cpp
@@ -862,11 +862,17 @@ void UNiagaraScript::SetVMCompilationResults(const FNiagaraVMExecutableDataId& I
 	
 	if (CachedScriptVM.LastCompileStatus == ENiagaraScriptCompileStatus::NCS_Error)
 	{
-		UE_LOG(LogNiagara, Error, TEXT("%s"), *CachedScriptVM.ErrorMsg);
+		// Compiler errors for Niagara will have a strong UI impact but the game should still function properly, there 
+		// will just be oddities in the visuals. It should be acted upon, but in no way should the game be blocked from
+		// a successful cook because of it. Therefore, we do a warning.
+		UE_LOG(LogNiagara, Warning, TEXT("%s System Asset: %s"), *CachedScriptVM.ErrorMsg, *GetPathName());
 	}
 	else if (CachedScriptVM.LastCompileStatus == ENiagaraScriptCompileStatus::NCS_UpToDateWithWarnings)
 	{
-		UE_LOG(LogNiagara, Warning, TEXT("%s"), *CachedScriptVM.ErrorMsg);
+		// Compiler warnings for Niagara are meant for notification and should have a UI representation, but 
+		// should be expected to still function properly and can be acted upon at the user's leisure. This makes
+		// them best logged as Display messages, as Log will not be shown in the cook.
+		UE_LOG(LogNiagara, Display, TEXT("%s System Asset: %s"), *CachedScriptVM.ErrorMsg, *GetPathName());
 	}
 
 	// The compilation process only references via soft references any parameter collections. This resolves those 
diff --git a/Engine/Plugins/FX/Niagara/Source/NiagaraEditor/Private/NiagaraCompiler.cpp b/Engine/Plugins/FX/Niagara/Source/NiagaraEditor/Private/NiagaraCompiler.cpp
--- a/Engine/Plugins/FX/Niagara/Source/NiagaraEditor/Private/NiagaraCompiler.cpp
+++ b/Engine/Plugins/FX/Niagara/Source/NiagaraEditor/Private/NiagaraCompiler.cpp
@@ -571,13 +571,19 @@ TSharedPtr<FNiagaraVMExecutableData> FNiagaraEditorModule::CompileScript(const F
 		else if (Message.Severity == FNiagaraCompileEventSeverity::Warning )
 		{
 		#if defined(NIAGARA_SCRIPT_COMPILE_LOGGING_MEDIUM)
-			UE_LOG(LogNiagaraCompiler, Warning, TEXT("%s"), *Message.Message);
+			// Compiler warnings for Niagara are meant for notification and should have a UI representation, but 
+			// should be expected to still function properly and can be acted upon at the user's leisure. This makes
+			// them best logged as Display messages, as Log will not be shown in the cook.
+			UE_LOG(LogNiagaraCompiler, Display, TEXT("%s System Asset: %s"), *Message.Message, *InCompileOptions.FullName);
 		#endif
 		}
 		else if (Message.Severity == FNiagaraCompileEventSeverity::Error)
 		{
 		#if defined(NIAGARA_SCRIPT_COMPILE_LOGGING_MEDIUM)
-			UE_LOG(LogNiagaraCompiler, Error, TEXT("%s"), *Message.Message);
+			// Compiler errors for Niagara will have a strong UI impact but the game should still function properly, there 
+			// will just be oddities in the visuals. It should be acted upon, but in no way should the game be blocked from
+			// a successful cook because of it. Therefore, we do a warning.
+			UE_LOG(LogNiagaraCompiler, Warning, TEXT("%s System Asset: %s"), *Message.Message, *InCompileOptions.FullName);
 		#endif
 			// Write the error messages to the string as well so that they can be echoed up the chain.
 			if (OutGraphLevelErrorMessages.Len() > 0)
