Rearranged niagara sprtie rendering shader code to avoid device reboots (shader compiler bug in VK Note9, S9)

diff --git a/Engine/Plugins/FX/Niagara/Shaders/Private/NiagaraSpriteVertexFactory.ush b/Engine/Plugins/FX/Niagara/Shaders/Private/NiagaraSpriteVertexFactory.ush
--- a/Engine/Plugins/FX/Niagara/Shaders/Private/NiagaraSpriteVertexFactory.ush
+++ b/Engine/Plugins/FX/Niagara/Shaders/Private/NiagaraSpriteVertexFactory.ush
@@ -697,25 +697,53 @@ float GetNiagaraNormalizedAge(uint InstanceID)
 #if (DYNAMIC_PARAMETERS_MASK & 1)
 float4 GetNiagaraParticleDynamicParameters(uint InstanceID)
 {
-	return NiagaraSpriteVF.MaterialParamDataOffset != -1 ? GetVec4(NiagaraSpriteVF.MaterialParamDataOffset, InstanceID) : float4(1.0f, 1.0f, 1.0f, 1.0f);;
+	if (NiagaraSpriteVF.MaterialParamDataOffset != -1)
+	{
+		return GetVec4(NiagaraSpriteVF.MaterialParamDataOffset, InstanceID);
+	}
+	else
+	{
+		return float4(1.0f, 1.0f, 1.0f, 1.0f);;
+	}
 }
 #endif
 #if (DYNAMIC_PARAMETERS_MASK & 2)
 float4 GetNiagaraParticleDynamicParameters1(uint InstanceID)
 {
-	return NiagaraSpriteVF.MaterialParam1DataOffset != -1 ? GetVec4(NiagaraSpriteVF.MaterialParam1DataOffset, InstanceID) : float4(1.0f, 1.0f, 1.0f, 1.0f);;
+	if (NiagaraSpriteVF.MaterialParam1DataOffset != -1)
+	{
+		return GetVec4(NiagaraSpriteVF.MaterialParam1DataOffset, InstanceID);
+	}
+	else
+	{
+		return float4(1.0f, 1.0f, 1.0f, 1.0f);
+	}
 }
 #endif
 #if (DYNAMIC_PARAMETERS_MASK & 4)
 float4 GetNiagaraParticleDynamicParameters2(uint InstanceID)
 {
-	return NiagaraSpriteVF.MaterialParam2DataOffset != -1 ? GetVec4(NiagaraSpriteVF.MaterialParam2DataOffset, InstanceID) : float4(1.0f, 1.0f, 1.0f, 1.0f);;
+	if (NiagaraSpriteVF.MaterialParam2DataOffset != -1)
+	{
+		return GetVec4(NiagaraSpriteVF.MaterialParam2DataOffset, InstanceID);
+	}
+	else
+	{
+		return float4(1.0f, 1.0f, 1.0f, 1.0f);
+	}
 }
 #endif
 #if (DYNAMIC_PARAMETERS_MASK & 8)
 float4 GetNiagaraParticleDynamicParameters3(uint InstanceID)
 {
-	return NiagaraSpriteVF.MaterialParam3DataOffset != -1 ? GetVec4(NiagaraSpriteVF.MaterialParam3DataOffset, InstanceID) : float4(1.0f, 1.0f, 1.0f, 1.0f);;
+	if (NiagaraSpriteVF.MaterialParam3DataOffset != -1)
+	{
+		return GetVec4(NiagaraSpriteVF.MaterialParam3DataOffset, InstanceID);
+	}
+	else
+	{
+		return float4(1.0f, 1.0f, 1.0f, 1.0f);
+	}
 }
 #endif
 
