Fix for carsh on save after editing module parameter description.

diff --git a/Engine/Plugins/FX/Niagara/Source/NiagaraEditor/Private/NiagaraEditorUtilities.cpp b/Engine/Plugins/FX/Niagara/Source/NiagaraEditor/Private/NiagaraEditorUtilities.cpp
--- a/Engine/Plugins/FX/Niagara/Source/NiagaraEditor/Private/NiagaraEditorUtilities.cpp
+++ b/Engine/Plugins/FX/Niagara/Source/NiagaraEditor/Private/NiagaraEditorUtilities.cpp
@@ -500,7 +500,7 @@ void FNiagaraEditorUtilities::CompileExistingEmitters(const TArray<UNiagaraEmitt
 	for (UNiagaraEmitter* Emitter : AffectedEmitters)
 	{
 		// If we've already compiled this emitter, or it's invalid skip it.
-		if (CompiledEmitters.Contains(Emitter) || Emitter->IsPendingKillOrUnreachable())
+		if (Emitter == nullptr || CompiledEmitters.Contains(Emitter) || Emitter->IsPendingKillOrUnreachable())
 		{
 			continue;
 		}
