Updating Niagara to be in its own CSV scope for easier tracking.

diff --git a/Engine/Plugins/FX/Niagara/Source/Niagara/Private/NiagaraComponent.cpp b/Engine/Plugins/FX/Niagara/Source/Niagara/Private/NiagaraComponent.cpp
--- a/Engine/Plugins/FX/Niagara/Source/Niagara/Private/NiagaraComponent.cpp
+++ b/Engine/Plugins/FX/Niagara/Source/Niagara/Private/NiagaraComponent.cpp
@@ -423,8 +423,8 @@ void UNiagaraComponent::SetEmitterEnable(FName EmitterName, bool bNewEnableState
 void UNiagaraComponent::TickComponent(float DeltaSeconds, enum ELevelTick TickType, FActorComponentTickFunction* ThisTickFunction)
 {
 	LLM_SCOPE(ELLMTag::Niagara);
+	CSV_SCOPED_TIMING_STAT_EXCLUSIVE(Niagara);
 	SCOPE_CYCLE_COUNTER(STAT_NiagaraComponentTick);
-	CSV_SCOPED_TIMING_STAT_EXCLUSIVE(Particles);
 
 	FScopeCycleCounter SystemStatCounter(Asset ? Asset->GetStatID(true, false) : TStatId());
 
@@ -599,6 +599,7 @@ bool UNiagaraComponent::IsWorldReadyToRun() const
 bool UNiagaraComponent::InitializeSystem()
 {
 	LLM_SCOPE(ELLMTag::Niagara);
+	CSV_SCOPED_TIMING_STAT_EXCLUSIVE(Niagara);
 	if (SystemInstance.IsValid() == false)
 	{
 		FNiagaraSystemInstance::AllocateSystemInstance(this, SystemInstance);
@@ -742,6 +743,7 @@ void UNiagaraComponent::Activate(bool bReset /* = false */)
 void UNiagaraComponent::Deactivate()
 {
 	SCOPE_CYCLE_COUNTER(STAT_NiagaraComponentDeactivate);
+	CSV_SCOPED_TIMING_STAT_EXCLUSIVE(Niagara);
 
 	//UE_LOG(LogNiagara, Log, TEXT("Deactivate: %u - %s"), this, *Asset->GetName());
 
@@ -937,6 +939,7 @@ void UNiagaraComponent::CreateRenderState_Concurrent()
 void UNiagaraComponent::SendRenderDynamicData_Concurrent()
 {
 	LLM_SCOPE(ELLMTag::Niagara);
+	CSV_SCOPED_TIMING_STAT_EXCLUSIVE(Niagara);
 	SCOPE_CYCLE_COUNTER(STAT_NiagaraComponentSendRenderData);
 	if (SystemInstance.IsValid() && SceneProxy)
 	{
diff --git a/Engine/Plugins/FX/Niagara/Source/Niagara/Private/NiagaraWorldManager.cpp b/Engine/Plugins/FX/Niagara/Source/Niagara/Private/NiagaraWorldManager.cpp
--- a/Engine/Plugins/FX/Niagara/Source/Niagara/Private/NiagaraWorldManager.cpp
+++ b/Engine/Plugins/FX/Niagara/Source/Niagara/Private/NiagaraWorldManager.cpp
@@ -216,6 +216,7 @@ void FNiagaraWorldManager::OnWorldCleanup(bool bSessionEnded, bool bCleanupResou
 
 void FNiagaraWorldManager::Tick(float DeltaSeconds)
 {
+	CSV_SCOPED_TIMING_STAT_EXCLUSIVE(Niagara);
 	LLM_SCOPE(ELLMTag::Niagara);
 	SCOPE_CYCLE_COUNTER(STAT_NiagaraWorldManTick);
 	SCOPE_CYCLE_COUNTER(STAT_NiagaraOverview_GT);
