Nullptr check before logging CDO name.

diff --git a/Engine/Plugins/FX/Niagara/Source/NiagaraEditor/Private/NiagaraNodeFunctionCall.cpp b/Engine/Plugins/FX/Niagara/Source/NiagaraEditor/Private/NiagaraNodeFunctionCall.cpp
--- a/Engine/Plugins/FX/Niagara/Source/NiagaraEditor/Private/NiagaraNodeFunctionCall.cpp
+++ b/Engine/Plugins/FX/Niagara/Source/NiagaraEditor/Private/NiagaraNodeFunctionCall.cpp
@@ -119,7 +119,10 @@ void UNiagaraNodeFunctionCall::UpgradeDIFunctionCalls()
 	{
 		UE_LOG(LogNiagaraEditor, Log, TEXT("Upgradeing Niagara Data Interface fuction call node. This may cause unnessessary recompiles. Please resave these assets if this occurs. Or use fx.UpgradeAllNiagaraAssets."));
 		UE_LOG(LogNiagaraEditor, Log, TEXT("Node: %s"), *GetFullName());
-		UE_LOG(LogNiagaraEditor, Log, TEXT("Interface: %s"), *InterfaceCDO->GetFullName());
+		if (InterfaceCDO)
+		{
+			UE_LOG(LogNiagaraEditor, Log, TEXT("Interface: %s"), *InterfaceCDO->GetFullName());
+		}
 		UE_LOG(LogNiagaraEditor, Log, TEXT("Function: %s"), *Signature.GetName());
 		UE_LOG(LogNiagaraEditor, Log, TEXT("Upgrade Note: %s"),* UpgradeNote);
 	}
