- Remove bUse*Sampling variables as they are not reliable.  They are bound after the DI is initialized in the editor and the discovery of them during VM binding may not write to the correct UObject (i.e. the one that will be serialized), and are therefore not reliable in cooked builds.  They are also incorrect for GPU emitters.
- Replaced with bAllowCPUMeshDataAccess that the DI determines on instance initialize.  We then don't allow the VM to bind to functions that would fail.

diff --git a/Engine/Plugins/FX/Niagara/Source/Niagara/Classes/NiagaraDataInterfaceSkeletalMesh.h b/Engine/Plugins/FX/Niagara/Source/Niagara/Classes/NiagaraDataInterfaceSkeletalMesh.h
--- a/Engine/Plugins/FX/Niagara/Source/Niagara/Classes/NiagaraDataInterfaceSkeletalMesh.h
+++ b/Engine/Plugins/FX/Niagara/Source/Niagara/Classes/NiagaraDataInterfaceSkeletalMesh.h
@@ -410,6 +410,9 @@ struct FNDISkeletalMesh_InstanceData
 	FSkeletalMeshGpuSpawnStaticBuffers* MeshGpuSpawnStaticBuffers;
 	FSkeletalMeshGpuDynamicBufferProxy* MeshGpuSpawnDynamicBuffers;
 
+	/** Temporary flag to deny binding VM functions that rely on mesh data being accessible on the CPU */
+	bool bAllowCPUMeshDataAccess;
+
 	FORCEINLINE_DEBUGGABLE bool ResetRequired(UNiagaraDataInterfaceSkeletalMesh* Interface)const;
 
 	bool Init(UNiagaraDataInterfaceSkeletalMesh* Interface, FNiagaraSystemInstance* SystemInstance);
@@ -483,17 +486,11 @@ public:
 	UPROPERTY(EditAnywhere, Category = "Skeleton")
 	TArray<FName> SpecificSockets;
 
-	/** Whether any triangle sampling function is bound. Only used in game. */
-	UPROPERTY()
-	bool bUseTriangleSampling;
-
-	/** Whether any vertex sampling function is bound. Only used in game. */
-	UPROPERTY()
-	bool bUseVertexSampling;
-
-	/** Whether any skeleton sampling function is bound. Only used in game. */
-	UPROPERTY()
-	bool bUseSkeletonSampling;
+#if WITH_EDITORONLY_DATA
+	/** Do we require CPU access to the data, this is set during GetVMExternalFunction. */
+	UPROPERTY(transient)
+	bool bRequiresCPUAccess;
+#endif
 
 	/** Cached change id off of the data interface.*/
 	uint32 ChangeId;
diff --git a/Engine/Plugins/FX/Niagara/Source/Niagara/Private/NiagaraDataInterfaceSkeletalMesh.cpp b/Engine/Plugins/FX/Niagara/Source/Niagara/Private/NiagaraDataInterfaceSkeletalMesh.cpp
--- a/Engine/Plugins/FX/Niagara/Source/Niagara/Private/NiagaraDataInterfaceSkeletalMesh.cpp
+++ b/Engine/Plugins/FX/Niagara/Source/Niagara/Private/NiagaraDataInterfaceSkeletalMesh.cpp
@@ -1028,6 +1028,8 @@ bool FNDISkeletalMesh_InstanceData::Init(UNiagaraDataInterfaceSkeletalMesh* Inte
 	MeshGpuSpawnStaticBuffers = nullptr;
 	MeshGpuSpawnDynamicBuffers = nullptr;
 
+	bAllowCPUMeshDataAccess = true;
+
 	// Get skel mesh and confirm have valid data
 	USkeletalMeshComponent* NewSkelComp = nullptr;
 	Mesh = Interface->GetSkeletalMesh(SystemInstance->GetComponent(), Component, NewSkelComp, this);
@@ -1152,16 +1154,15 @@ bool FNDISkeletalMesh_InstanceData::Init(UNiagaraDataInterfaceSkeletalMesh* Inte
 	bool bNeedDataImmediately = true;
 		
 	//Grab a handle to the skinning data if we have a component to skin.
-	ENDISkeletalMesh_SkinningMode SkinningMode = (Interface->bUseTriangleSampling || Interface->bUseVertexSampling) ? Interface->SkinningMode : ENDISkeletalMesh_SkinningMode::None;
+	const ENDISkeletalMesh_SkinningMode SkinningMode = Interface->SkinningMode;
 	FSkeletalMeshSkinningDataUsage Usage(
 		LODIndex,
-		SkinningMode == ENDISkeletalMesh_SkinningMode::SkinOnTheFly || SkinningMode == ENDISkeletalMesh_SkinningMode::PreSkin || Interface->bUseSkeletonSampling,
+		SkinningMode == ENDISkeletalMesh_SkinningMode::SkinOnTheFly || SkinningMode == ENDISkeletalMesh_SkinningMode::PreSkin,
 		SkinningMode == ENDISkeletalMesh_SkinningMode::PreSkin,
 		bNeedDataImmediately);
 
 	if (NewSkelComp)
 	{
-		SkinningMode = Interface->SkinningMode;
 		TWeakObjectPtr<USkeletalMeshComponent> SkelWeakCompPtr = NewSkelComp;
 		FNDI_SkeletalMesh_GeneratedData& GeneratedData = SystemInstance->GetWorldManager()->GetSkeletalMeshGeneratedData();
 		SkinningData = GeneratedData.GetCachedSkinningData(SkelWeakCompPtr, Usage);
@@ -1181,52 +1182,29 @@ bool FNDISkeletalMesh_InstanceData::Init(UNiagaraDataInterfaceSkeletalMesh* Inte
 	FSkinWeightVertexBuffer* SkinWeightBuffer = nullptr;
 	FSkeletalMeshLODRenderData& LODData = GetLODRenderDataAndSkinWeights(SkinWeightBuffer);
 
-	//Check for the validity of the Mesh's cpu data.
-	if ( Interface->bUseTriangleSampling || Interface->bUseVertexSampling )
+	// Check for the validity of the Mesh's cpu data.
 	{
-		if ( !Mesh->GetLODInfo(LODIndex)->bAllowCPUAccess )
+		if ( Mesh->GetLODInfo(LODIndex)->bAllowCPUAccess )
 		{
-			UE_LOG(LogNiagara, Warning, TEXT("Skeletal Mesh Data Interface is trying to spawn from a whole mesh that does not allow CPU Access.\nInterface: %s\nMesh: %s\nLOD: %d"),
-				*Interface->GetFullName(),
-				*Mesh->GetFullName(),
-				LODIndex);
-
-			return false;
+			const bool LODDataNumVerticesCorrect = LODData.GetNumVertices() > 0;
+			const bool LODDataPositonNumVerticesCorrect = LODData.StaticVertexBuffers.PositionVertexBuffer.GetNumVertices() > 0;
+			const bool bSkinWeightBuffer = SkinWeightBuffer != nullptr;
+			const bool SkinWeightBufferNumVerticesCorrect = bSkinWeightBuffer && (SkinWeightBuffer->GetNumVertices() > 0);
+			const bool bIndexBufferValid = LODData.MultiSizeIndexContainer.IsIndexBufferValid();
+			const bool bIndexBufferFound = bIndexBufferValid && (LODData.MultiSizeIndexContainer.GetIndexBuffer() != nullptr);
+			const bool bIndexBufferNumCorrect = bIndexBufferFound && (LODData.MultiSizeIndexContainer.GetIndexBuffer()->Num() > 0);
+
+			bAllowCPUMeshDataAccess = LODDataNumVerticesCorrect &&
+				LODDataPositonNumVerticesCorrect &&
+				bSkinWeightBuffer &&
+				SkinWeightBufferNumVerticesCorrect &&
+				bIndexBufferValid &&
+				bIndexBufferFound &&
+				bIndexBufferNumCorrect;
 		}
-
-		bool LODDataNumVerticesCorrect = LODData.GetNumVertices() > 0;
-		bool LODDataPositonNumVerticesCorrect = LODData.StaticVertexBuffers.PositionVertexBuffer.GetNumVertices() > 0;
-		bool bSkinWeightBuffer = SkinWeightBuffer != nullptr;
-		bool SkinWeightBufferNumVerticesCorrect = bSkinWeightBuffer && (SkinWeightBuffer->GetNumVertices() > 0);
-		bool bIndexBufferValid = LODData.MultiSizeIndexContainer.IsIndexBufferValid();
-		bool bIndexBufferFound = bIndexBufferValid && (LODData.MultiSizeIndexContainer.GetIndexBuffer() != nullptr);
-		bool bIndexBufferNumCorrect = bIndexBufferFound && (LODData.MultiSizeIndexContainer.GetIndexBuffer()->Num() > 0);
-
-		bool bMeshCPUDataValid = LODDataNumVerticesCorrect &&
-			LODDataPositonNumVerticesCorrect &&
-			bSkinWeightBuffer &&
-			SkinWeightBufferNumVerticesCorrect &&
-			bIndexBufferValid &&
-			bIndexBufferFound &&
-			bIndexBufferNumCorrect;
-
-		if (!bMeshCPUDataValid)
+		else
 		{
-			UE_LOG(LogNiagara, Warning, TEXT("Skeletal Mesh Data Interface is trying to sample from a mesh with missing CPU vertex or index data.\nInterface: %s\nMesh: %s\nLOD: %d\n"
-				"LODDataNumVerticesCorrect: %d  LODDataPositonNumVerticesCorrect : %d  bSkinWeightBuffer : %d  SkinWeightBufferNumVerticesCorrect : %d bIndexBufferValid : %d  bIndexBufferFound : %d  bIndexBufferNumCorrect : %d"),
-				*Interface->GetFullName(),
-				*Mesh->GetFullName(),
-				LODIndex,
-				LODDataNumVerticesCorrect ? 1 : 0,
-				LODDataPositonNumVerticesCorrect ? 1 : 0,
-				bSkinWeightBuffer ? 1 : 0,
-				SkinWeightBufferNumVerticesCorrect ? 1 : 0,
-				bIndexBufferValid ? 1 : 0,
-				bIndexBufferFound ? 1 : 0,
-				bIndexBufferNumCorrect ? 1 : 0
-			);
-
-			return false;
+			bAllowCPUMeshDataAccess = false;
 		}
 	}
 
@@ -1499,9 +1477,9 @@ UNiagaraDataInterfaceSkeletalMesh::UNiagaraDataInterfaceSkeletalMesh(FObjectInit
 	, Source(nullptr)
 	, SkinningMode(ENDISkeletalMesh_SkinningMode::SkinOnTheFly)
 	, WholeMeshLOD(INDEX_NONE)
-	, bUseTriangleSampling(true)
-	, bUseVertexSampling(true)
-	, bUseSkeletonSampling(true)
+#if WITH_EDITORONLY_DATA
+	, bRequiresCPUAccess(false)
+#endif
 	, ChangeId(0)
 {
 	Proxy = MakeShared<FNiagaraDataInterfaceProxySkeletalMesh, ESPMode::ThreadSafe>();
@@ -1527,14 +1505,14 @@ void UNiagaraDataInterfaceSkeletalMesh::PostEditChangeProperty(FPropertyChangedE
 {
 	Super::PostEditChangeProperty(PropertyChangedEvent);
 
+#if WITH_EDITORONLY_DATA
 	// If the change comes from an interaction (and not just a generic change) reset the usage flags.
 	// todo : this and the usage binding need to be done in the a precompilation parsing (or whever the script is compiled).
 	if (PropertyChangedEvent.Property)
 	{
-		bUseTriangleSampling = false;
-		bUseVertexSampling = false;
-		bUseSkeletonSampling = false;
+		bRequiresCPUAccess = false;
 	}
+#endif
 	ChangeId++;
 }
 
@@ -1559,46 +1537,53 @@ void UNiagaraDataInterfaceSkeletalMesh::GetVMExternalFunction(const FVMExternalF
 		return;
 	}
 
-	BindTriangleSamplingFunction(BindingInfo, InstData, OutFunc);
-
+	// Bind skeleton sampling function
+	BindSkeletonSamplingFunction(BindingInfo, InstData, OutFunc);
 	if (OutFunc.IsBound())
 	{
 #if WITH_EDITOR
-		if (!bUseTriangleSampling)
+		if ( SkinningMode == ENDISkeletalMesh_SkinningMode::None )
 		{
-			bUseTriangleSampling = true;
-			MarkPackageDirty();
+			UE_LOG(LogNiagara, Warning, TEXT("Skeletal Mesh Data Interface is trying to use skeleton sampling but skinning mode is none. Interface: %s"), *GetFullName());
 		}
 #endif // WITH_EDITOR
 		return;
 	}
 
-	BindVertexSamplingFunction(BindingInfo, InstData, OutFunc);
-
+	// Bind triangle sampling function
+	BindTriangleSamplingFunction(BindingInfo, InstData, OutFunc);
 	if (OutFunc.IsBound())
 	{
-#if WITH_EDITOR
-		if (!bUseVertexSampling)
+		if (!InstData->bAllowCPUMeshDataAccess)
 		{
-			bUseVertexSampling = true;
-			MarkPackageDirty();
+			UE_LOG(LogNiagara, Warning, TEXT("Skeletal Mesh Data Interface is trying to use triangle sampling but CPU access or the data is invalid. Interface: %s"), *GetFullName());
+			OutFunc.Unbind();
+			return;
 		}
+
+#if WITH_EDITOR
+		bRequiresCPUAccess = true;
 #endif // WITH_EDITOR
 		return;
 	}
 
-	BindSkeletonSamplingFunction(BindingInfo, InstData, OutFunc);
+	// Bind vertex sampling function
+	BindVertexSamplingFunction(BindingInfo, InstData, OutFunc);
 
-#if WITH_EDITOR
 	if (OutFunc.IsBound())
 	{
-		if (!bUseSkeletonSampling)
+		if (!InstData->bAllowCPUMeshDataAccess)
 		{
-			bUseSkeletonSampling = true;
-			MarkPackageDirty();
+			UE_LOG(LogNiagara, Warning, TEXT("Skeletal Mesh Data Interface is trying to use vertex sampling but CPU access or the data is invalid. Interface: %s"), *GetFullName());
+			OutFunc.Unbind();
+			return;
 		}
-	}
+
+#if WITH_EDITOR
+		bRequiresCPUAccess = true;
 #endif // WITH_EDITOR
+		return;
+	}
 }
 
 
@@ -1618,9 +1603,9 @@ bool UNiagaraDataInterfaceSkeletalMesh::CopyToInternal(UNiagaraDataInterface* De
 	OtherTyped->WholeMeshLOD = WholeMeshLOD;
 	OtherTyped->SpecificBones = SpecificBones;
 	OtherTyped->SpecificSockets = SpecificSockets;
-	OtherTyped->bUseTriangleSampling = bUseTriangleSampling;
-	OtherTyped->bUseVertexSampling = bUseVertexSampling;
-	OtherTyped->bUseSkeletonSampling = bUseSkeletonSampling;
+#if WITH_EDITORONLY_DATA
+	OtherTyped->bRequiresCPUAccess = bRequiresCPUAccess;
+#endif
 	return true;
 }
 
@@ -1688,7 +1673,7 @@ TArray<FNiagaraDataInterfaceError> UNiagaraDataInterfaceSkeletalMesh::GetErrors(
 	bool bHasNoMeshAssignedError = false;
 	
 	// Collect Errors
-	if (DefaultMesh != nullptr && (bUseTriangleSampling || bUseVertexSampling))
+	if (DefaultMesh != nullptr && bRequiresCPUAccess)
 	{
 		for (auto info : DefaultMesh->GetLODInfoArray())
 		{
