- Transition Niagara Dummy Sort Buffers to RW

diff --git a/Engine/Plugins/FX/Niagara/Source/NiagaraVertexFactories/Private/NiagaraSortingGPU.cpp b/Engine/Plugins/FX/Niagara/Source/NiagaraVertexFactories/Private/NiagaraSortingGPU.cpp
--- a/Engine/Plugins/FX/Niagara/Source/NiagaraVertexFactories/Private/NiagaraSortingGPU.cpp
+++ b/Engine/Plugins/FX/Niagara/Source/NiagaraVertexFactories/Private/NiagaraSortingGPU.cpp
@@ -226,6 +226,7 @@ void FNiagaraCopyIntBufferRegionCS::SetParameters(
 
 	for (int32 Index = DestCount; Index < NIAGARA_COPY_BUFFER_BUFFER_COUNT; ++Index)
 	{
+		RHICmdList.TransitionResource(EResourceTransitionAccess::ERWNoBarrier, EResourceTransitionPipeline::EComputeToCompute, NiagaraSortingDummyUAV[Index].Buffer.UAV);
 		RHICmdList.SetUAVParameter(ComputeShaderRHI, DestData[Index].GetBaseIndex(), NiagaraSortingDummyUAV[Index].Buffer.UAV);
 	}
 
