Updating fix to mesh camera facing world position to account for local space emitters.

diff --git a/Engine/Plugins/FX/Niagara/Shaders/Private/NiagaraMeshVertexFactory.ush b/Engine/Plugins/FX/Niagara/Shaders/Private/NiagaraMeshVertexFactory.ush
--- a/Engine/Plugins/FX/Niagara/Shaders/Private/NiagaraMeshVertexFactory.ush
+++ b/Engine/Plugins/FX/Niagara/Shaders/Private/NiagaraMeshVertexFactory.ush
@@ -640,7 +640,7 @@ FVertexFactoryIntermediates GetVertexFactoryIntermediates(FVertexFactoryInput In
 	}
 	else if (MeshFacingMode == MESH_FACING_CAMERA_POSITION)
 	{
-		float3 WorldPosition = Intermediates.ParticlePosition.xyz;
+		float3 WorldPosition = mul(Intermediates.ParticlePosition, NiagaraMeshVF.LocalToWorld).xyz;
 		float3 CameraVector = ResolvedView.WorldCameraOrigin.xyz - WorldPosition.xyz;
 		float3 CameraDirection = normalize(CameraVector);
 		float4 ToAlignmentQuat = FindQuatBetween(float3(1.0f, 0.0f, 0.0f), CameraDirection);
