- Solo emitters should not read from the world managers cached view locations

diff --git a/Engine/Plugins/FX/Niagara/Source/Niagara/Private/NiagaraSystemInstance.cpp b/Engine/Plugins/FX/Niagara/Source/Niagara/Private/NiagaraSystemInstance.cpp
--- a/Engine/Plugins/FX/Niagara/Source/Niagara/Private/NiagaraSystemInstance.cpp
+++ b/Engine/Plugins/FX/Niagara/Source/Niagara/Private/NiagaraSystemInstance.cpp
@@ -12,6 +12,7 @@
 #include "NiagaraRenderer.h"
 #include "Templates/AlignmentTemplates.h"
 #include "NiagaraEmitterInstanceBatcher.h"
+#include "GameFramework/PlayerController.h"
 
 
 DECLARE_CYCLE_STAT(TEXT("System Activate [GT]"), STAT_NiagaraSystemActivate, STATGROUP_Niagara);
@@ -1171,40 +1172,78 @@ void FNiagaraSystemInstance::TickDataInterfaces(float DeltaSeconds, bool bPostSi
 float FNiagaraSystemInstance::GetLODDistance()
 {
 	check(Component);
-	FVector CurrPos = Component->GetComponentLocation();
 #if WITH_EDITOR
 	if (Component->bEnablePreviewLODDistance)
 	{
 		return Component->PreviewLODDistance;
 	}
-	else
 #endif
+
+	constexpr float DefaultLODDistance = 0.0f;
+
+	FNiagaraWorldManager* WorldManager = GetWorldManager();
+	if ( WorldManager == nullptr )
 	{
-		constexpr float DefaultLODDistance = 0.0f;
+		return DefaultLODDistance;
+	}
 
-		FNiagaraWorldManager* WorldManager = GetWorldManager();
-		if ( WorldManager == nullptr )
-		{
-			return DefaultLODDistance;
-		}
+	const FVector EffectLocation = Component->GetComponentLocation();
 
+	// If we are inside the WorldManager tick we will use the cache player view locations as we can be ticked on different threads
+	if (WorldManager->CachedPlayerViewLocationsValid())
+	{
 		TArrayView<const FVector> PlayerViewLocations = WorldManager->GetCachedPlayerViewLocations();
-		if ( PlayerViewLocations.Num() == 0 )
+		if (PlayerViewLocations.Num() == 0)
 		{
 			return DefaultLODDistance;
 		}
 
+		// We are being ticked inside the WorldManager and can safely use the list of cached player view locations
 		float LODDistanceSqr = FMath::Square(WORLD_MAX);
 		for (const FVector& ViewLocation : PlayerViewLocations)
 		{
-			const float DistanceToEffectSqr = FVector(ViewLocation - CurrPos).SizeSquared();
-			if (DistanceToEffectSqr < LODDistanceSqr)
+			const float DistanceToEffectSqr = FVector(ViewLocation - EffectLocation).SizeSquared();
+			LODDistanceSqr = FMath::Min(LODDistanceSqr, DistanceToEffectSqr);
+		}
+		return FMath::Sqrt(LODDistanceSqr);
+	}
+
+	// If we are not inside the WorldManager tick (solo tick) we must look over the player view locations manually
+	ensureMsgf(IsInGameThread(), TEXT("FNiagaraSystemInstance::GetLODDistance called in potentially thread unsafe way"));
+
+	if ( UWorld* World = Component->GetWorld() )
+	{
+		TArray<FVector, TInlineAllocator<8> > PlayerViewLocations;
+		if (World->GetPlayerControllerIterator())
+		{
+			for (FConstPlayerControllerIterator Iterator = World->GetPlayerControllerIterator(); Iterator; ++Iterator)
 			{
-				LODDistanceSqr = DistanceToEffectSqr;
+				APlayerController* PlayerController = Iterator->Get();
+				if (PlayerController && PlayerController->IsLocalPlayerController())
+				{
+					FVector* ViewLocation = new(PlayerViewLocations) FVector;
+					FRotator ViewRotation;
+					PlayerController->GetPlayerViewPoint(*ViewLocation, ViewRotation);
+				}
 			}
 		}
-		return FMath::Sqrt(LODDistanceSqr);
+		else
+		{
+			PlayerViewLocations = World->ViewLocationsRenderedLastFrame;
+		}
+
+		if (PlayerViewLocations.Num() > 0)
+		{
+			float LODDistanceSqr = FMath::Square(WORLD_MAX);
+			for (const FVector& ViewLocation : PlayerViewLocations)
+			{
+				const float DistanceToEffectSqr = FVector(ViewLocation - EffectLocation).SizeSquared();
+				LODDistanceSqr = FMath::Min(LODDistanceSqr, DistanceToEffectSqr);
+			}
+			return FMath::Sqrt(LODDistanceSqr);
+		}
 	}
+	return DefaultLODDistance;
 }
 
 void FNiagaraSystemInstance::TickInstanceParameters(float DeltaSeconds)
diff --git a/Engine/Plugins/FX/Niagara/Source/Niagara/Private/NiagaraWorldManager.cpp b/Engine/Plugins/FX/Niagara/Source/Niagara/Private/NiagaraWorldManager.cpp
--- a/Engine/Plugins/FX/Niagara/Source/Niagara/Private/NiagaraWorldManager.cpp
+++ b/Engine/Plugins/FX/Niagara/Source/Niagara/Private/NiagaraWorldManager.cpp
@@ -226,7 +226,7 @@ void FNiagaraWorldManager::Tick(float DeltaSeconds)
 	SkeletalMeshGeneratedData.TickGeneratedData(DeltaSeconds);
 
 	// Cache player view locations for all system instances to access
-	CachedPlayerViewLocations.Reset();
+	bCachedPlayerViewLocationsValid = true;
 	if (World->GetPlayerControllerIterator())
 	{
 		for ( FConstPlayerControllerIterator Iterator=World->GetPlayerControllerIterator(); Iterator; ++Iterator)
@@ -267,6 +267,10 @@ void FNiagaraWorldManager::Tick(float DeltaSeconds)
 		SystemSimulations.Remove(DeadSystem);
 	}
 
+	// Clear cached player view location list, it should never be used outside of the world tick
+	bCachedPlayerViewLocationsValid = false;
+	CachedPlayerViewLocations.Reset();
+
 	// Enqueue fence for deferred deletion if we need to wait on anything
 	if ( DeferredDeletionQueue[DeferredDeletionQueueIndex].Queue.Num() > 0 )
 	{
diff --git a/Engine/Plugins/FX/Niagara/Source/Niagara/Public/NiagaraWorldManager.h b/Engine/Plugins/FX/Niagara/Source/Niagara/Public/NiagaraWorldManager.h
--- a/Engine/Plugins/FX/Niagara/Source/Niagara/Public/NiagaraWorldManager.h
+++ b/Engine/Plugins/FX/Niagara/Source/Niagara/Public/NiagaraWorldManager.h
@@ -87,7 +87,8 @@ public:
 	
 	FORCEINLINE FNDI_SkeletalMesh_GeneratedData& GetSkeletalMeshGeneratedData() { return SkeletalMeshGeneratedData; }
 
-	TArrayView<const FVector> GetCachedPlayerViewLocations() const { return MakeArrayView(CachedPlayerViewLocations); }
+	bool CachedPlayerViewLocationsValid() const { return bCachedPlayerViewLocationsValid; }
+	TArrayView<const FVector> GetCachedPlayerViewLocations() const { check(bCachedPlayerViewLocationsValid); return MakeArrayView(CachedPlayerViewLocations); }
 
 private:
 	UWorld* World;
@@ -98,6 +99,7 @@ private:
 
 	int32 CachedEffectsQuality;
 
+	bool bCachedPlayerViewLocationsValid = false;
 	TArray<FVector, TInlineAllocator<8> > CachedPlayerViewLocations;
 
 	/** Generated data used by data interfaces*/
