Cleaning up curve debug logging. The tolerances were too low and normal floating point differences were triggering it. There is still a lingering issue (likely due to undo/redo), so this also makes sure that the LUT is up to date on load.

diff --git a/Engine/Plugins/FX/Niagara/Source/Niagara/Private/NiagaraDataInterface.cpp b/Engine/Plugins/FX/Niagara/Source/Niagara/Private/NiagaraDataInterface.cpp
--- a/Engine/Plugins/FX/Niagara/Source/Niagara/Private/NiagaraDataInterface.cpp
+++ b/Engine/Plugins/FX/Niagara/Source/Niagara/Private/NiagaraDataInterface.cpp
@@ -116,17 +116,21 @@ bool UNiagaraDataInterfaceCurveBase::CompareLUTS(const TArray<float>& OtherLUT)
 {
 	if (ShaderLUT.Num() == OtherLUT.Num())
 	{
+		bool bMatched = true;
 		for (int32 i = 0; i < ShaderLUT.Num(); i++)
 		{
-			if (false == FMath::IsNearlyEqual(ShaderLUT[i], OtherLUT[i]))
+			if (false == FMath::IsNearlyEqual(ShaderLUT[i], OtherLUT[i], 0.0001f))
 			{
-				return false;
+				bMatched = false;
+				UE_LOG(LogNiagara, Log, TEXT("First LUT mismatch found on comparison - LUT[%d] = %.9f  Other = %.9f \t%.9f"), i, ShaderLUT[i], OtherLUT[i], fabsf(ShaderLUT[i] - OtherLUT[i]));
+				break;
 			}
 		}
-		return true;
+		return bMatched;
 	}
 	else
 	{
+		UE_LOG(LogNiagara, Log, TEXT("Table sizes don't match"));
 		return false;
 	}
 }
diff --git a/Engine/Plugins/FX/Niagara/Source/Niagara/Private/NiagaraDataInterfaceColorCurve.cpp b/Engine/Plugins/FX/Niagara/Source/Niagara/Private/NiagaraDataInterfaceColorCurve.cpp
--- a/Engine/Plugins/FX/Niagara/Source/Niagara/Private/NiagaraDataInterfaceColorCurve.cpp
+++ b/Engine/Plugins/FX/Niagara/Source/Niagara/Private/NiagaraDataInterfaceColorCurve.cpp
@@ -41,18 +41,21 @@ void UNiagaraDataInterfaceColorCurve::PostLoad()
 	{
 		UpdateLUT();
 	}
-#if !UE_BUILD_SHIPPING
 	else
 	{
+
+#if !UE_BUILD_SHIPPING
 		TArray<float> OldLUT = ShaderLUT;
+#endif
 		UpdateLUT();
 
+#if !UE_BUILD_SHIPPING
 		if (!CompareLUTS(OldLUT))
 		{
 			UE_LOG(LogNiagara, Log, TEXT("PostLoad LUT generation is out of sync. Please investigate. %s"), *GetPathName());
 		}
-	}
 #endif
+	}
 }
 
 void UNiagaraDataInterfaceColorCurve::UpdateLUT()
diff --git a/Engine/Plugins/FX/Niagara/Source/Niagara/Private/NiagaraDataInterfaceCurve.cpp b/Engine/Plugins/FX/Niagara/Source/Niagara/Private/NiagaraDataInterfaceCurve.cpp
--- a/Engine/Plugins/FX/Niagara/Source/Niagara/Private/NiagaraDataInterfaceCurve.cpp
+++ b/Engine/Plugins/FX/Niagara/Source/Niagara/Private/NiagaraDataInterfaceCurve.cpp
@@ -80,18 +80,22 @@ void UNiagaraDataInterfaceCurve::PostLoad()
 	{
 		UpdateLUT();
 	}
-#if !UE_BUILD_SHIPPING
 	else
 	{
+
+#if !UE_BUILD_SHIPPING
 		TArray<float> OldLUT = ShaderLUT;
+#endif
 		UpdateLUT();
 
+
+#if !UE_BUILD_SHIPPING
 		if (!CompareLUTS(OldLUT))
 		{
 			UE_LOG(LogNiagara, Log, TEXT("PostLoad LUT generation is out of sync. Please investigate. %s"), *GetPathName());
 		}
-	}
 #endif
+	}
 }
 
 
diff --git a/Engine/Plugins/FX/Niagara/Source/Niagara/Private/NiagaraDataInterfaceVector2DCurve.cpp b/Engine/Plugins/FX/Niagara/Source/Niagara/Private/NiagaraDataInterfaceVector2DCurve.cpp
--- a/Engine/Plugins/FX/Niagara/Source/Niagara/Private/NiagaraDataInterfaceVector2DCurve.cpp
+++ b/Engine/Plugins/FX/Niagara/Source/Niagara/Private/NiagaraDataInterfaceVector2DCurve.cpp
@@ -39,18 +39,20 @@ void UNiagaraDataInterfaceVector2DCurve::PostLoad()
 	{
 		UpdateLUT();
 	}
-#if !UE_BUILD_SHIPPING
 	else
 	{
+#if !UE_BUILD_SHIPPING
 		TArray<float> OldLUT = ShaderLUT;
+#endif
 		UpdateLUT();
 
+#if !UE_BUILD_SHIPPING
 		if (!CompareLUTS(OldLUT))
 		{
 			UE_LOG(LogNiagara, Log, TEXT("PostLoad LUT generation is out of sync. Please investigate. %s"), *GetPathName());
 		}
-	}
 #endif
+	}
 }
 
 void UNiagaraDataInterfaceVector2DCurve::UpdateLUT()
diff --git a/Engine/Plugins/FX/Niagara/Source/Niagara/Private/NiagaraDataInterfaceVector4Curve.cpp b/Engine/Plugins/FX/Niagara/Source/Niagara/Private/NiagaraDataInterfaceVector4Curve.cpp
--- a/Engine/Plugins/FX/Niagara/Source/Niagara/Private/NiagaraDataInterfaceVector4Curve.cpp
+++ b/Engine/Plugins/FX/Niagara/Source/Niagara/Private/NiagaraDataInterfaceVector4Curve.cpp
@@ -41,18 +41,20 @@ void UNiagaraDataInterfaceVector4Curve::PostLoad()
 	{
 		UpdateLUT();
 	}
-#if !UE_BUILD_SHIPPING
 	else
 	{
+#if !UE_BUILD_SHIPPING
 		TArray<float> OldLUT = ShaderLUT;
+#endif
 		UpdateLUT();
 
+#if !UE_BUILD_SHIPPING
 		if (!CompareLUTS(OldLUT))
 		{
 			UE_LOG(LogNiagara, Log, TEXT("PostLoad LUT generation is out of sync. Please investigate. %s"), *GetPathName());
 		}
-	}
 #endif
+	}
 }
 
 void UNiagaraDataInterfaceVector4Curve::UpdateLUT()
diff --git a/Engine/Plugins/FX/Niagara/Source/Niagara/Private/NiagaraDataInterfaceVectorCurve.cpp b/Engine/Plugins/FX/Niagara/Source/Niagara/Private/NiagaraDataInterfaceVectorCurve.cpp
--- a/Engine/Plugins/FX/Niagara/Source/Niagara/Private/NiagaraDataInterfaceVectorCurve.cpp
+++ b/Engine/Plugins/FX/Niagara/Source/Niagara/Private/NiagaraDataInterfaceVectorCurve.cpp
@@ -41,17 +41,19 @@ void UNiagaraDataInterfaceVectorCurve::PostLoad()
 	{
 		UpdateLUT();
 	}
-#if !UE_BUILD_SHIPPING
 	else
 	{
+#if !UE_BUILD_SHIPPING
 		TArray<float> OldLUT = ShaderLUT;
+#endif
 		UpdateLUT();
+#if !UE_BUILD_SHIPPING
 		if (!CompareLUTS(OldLUT))
 		{
 			UE_LOG(LogNiagara, Log, TEXT("PostLoad LUT generation is out of sync. Please investigate. %s"), *GetPathName());
 		}
-	}
 #endif
+	}
 }
 
 void UNiagaraDataInterfaceVectorCurve::UpdateLUT()
