Niagara - System overview MVP

diff --git a/Engine/Plugins/FX/Niagara/Source/Niagara/Classes/NiagaraSystem.h b/Engine/Plugins/FX/Niagara/Source/Niagara/Classes/NiagaraSystem.h
--- a/Engine/Plugins/FX/Niagara/Source/Niagara/Classes/NiagaraSystem.h
+++ b/Engine/Plugins/FX/Niagara/Source/Niagara/Classes/NiagaraSystem.h
@@ -166,9 +166,6 @@ public:
 	/** Gets editor specific data stored with this system. */
 	const UNiagaraEditorDataBase* GetEditorData() const;
 
-	/** Sets editor specific data stored with this system. */
-	void SetEditorData(UNiagaraEditorDataBase* InEditorData);
-
 	/** Internal: The thumbnail image.*/
 	UPROPERTY()
 	class UTexture2D* ThumbnailImage;
diff --git a/Engine/Plugins/FX/Niagara/Source/Niagara/Private/NiagaraModule.cpp b/Engine/Plugins/FX/Niagara/Source/Niagara/Private/NiagaraModule.cpp
--- a/Engine/Plugins/FX/Niagara/Source/Niagara/Private/NiagaraModule.cpp
+++ b/Engine/Plugins/FX/Niagara/Source/Niagara/Private/NiagaraModule.cpp
@@ -364,7 +364,7 @@ void INiagaraModule::TickWorld(UWorld* World, ELevelTick TickType, float DeltaSe
 }
 
 #if WITH_EDITOR
-const INiagaraMergeManager& INiagaraModule::GetMergeManager()
+const INiagaraMergeManager& INiagaraModule::GetMergeManager() const
 {
 	checkf(MergeManager.IsValid(), TEXT("Merge manager was never registered, or was unregistered."));
 	return *MergeManager.Get();
@@ -383,24 +383,23 @@ void INiagaraModule::UnregisterMergeManager(TSharedRef<INiagaraMergeManager> InM
 	MergeManager.Reset();
 }
 
-UNiagaraScriptSourceBase* INiagaraModule::CreateDefaultScriptSource(UObject* Outer)
+const INiagaraEditorOnlyDataUtilities& INiagaraModule::GetEditorOnlyDataUtilities() const
 {
-	checkf(OnCreateDefaultScriptSourceDelegate.IsBound(), TEXT("Create default script source delegate not bound."));
-	return OnCreateDefaultScriptSourceDelegate.Execute(Outer);
+	checkf(EditorOnlyDataUtilities.IsValid(), TEXT("Editor only data utilities object was never registered, or was unregistered."));
+	return *EditorOnlyDataUtilities.Get();
 }
 
-FDelegateHandle INiagaraModule::RegisterOnCreateDefaultScriptSource(FOnCreateDefaultScriptSource OnCreateDefaultScriptSource)
+void INiagaraModule::RegisterEditorOnlyDataUtilities(TSharedRef<INiagaraEditorOnlyDataUtilities> InEditorOnlyDataUtilities)
 {
-	checkf(OnCreateDefaultScriptSourceDelegate.IsBound() == false, TEXT("Only one handler is allowed for the OnCreateDefaultScriptSource delegate"));
-	OnCreateDefaultScriptSourceDelegate = OnCreateDefaultScriptSource;
-	return OnCreateDefaultScriptSourceDelegate.GetHandle();
+	checkf(EditorOnlyDataUtilities.IsValid() == false, TEXT("Only one editor only data utilities object can be registered at a time."));
+	EditorOnlyDataUtilities = InEditorOnlyDataUtilities;
 }
 
-void INiagaraModule::UnregisterOnCreateDefaultScriptSource(FDelegateHandle DelegateHandle)
+void INiagaraModule::UnregisterEditorOnlyDataUtilities(TSharedRef<INiagaraEditorOnlyDataUtilities> InEditorOnlyDataUtilities)
 {
-	checkf(OnCreateDefaultScriptSourceDelegate.IsBound(), TEXT("OnCreateDefaultScriptSource is not registered"));
-	checkf(OnCreateDefaultScriptSourceDelegate.GetHandle() == DelegateHandle, TEXT("Can only unregister the OnCreateDefaultScriptSource delegate with the handle it was registered with."));
-	OnCreateDefaultScriptSourceDelegate.Unbind();
+	checkf(EditorOnlyDataUtilities.IsValid(), TEXT("Editor only data utilities object is not registered"));
+	checkf(EditorOnlyDataUtilities == InEditorOnlyDataUtilities, TEXT("Can only unregister the editor only data utilities object which was previously registered."));
+	EditorOnlyDataUtilities.Reset();
 }
 
 TSharedPtr<FNiagaraVMExecutableData> INiagaraModule::CompileScript(const FNiagaraCompileRequestDataBase* InCompileData, const FNiagaraCompileOptions& InCompileOptions)
diff --git a/Engine/Plugins/FX/Niagara/Source/Niagara/Private/NiagaraSystem.cpp b/Engine/Plugins/FX/Niagara/Source/Niagara/Private/NiagaraSystem.cpp
--- a/Engine/Plugins/FX/Niagara/Source/Niagara/Private/NiagaraSystem.cpp
+++ b/Engine/Plugins/FX/Niagara/Source/Niagara/Private/NiagaraSystem.cpp
@@ -15,6 +15,7 @@
 #include "AssetData.h"
 #include "NiagaraStats.h"
 #include "NiagaraEditorDataBase.h"
+#include "INiagaraEditorOnlyDataUtlities.h"
 
 #if WITH_EDITOR
 #include "NiagaraScriptDerivedData.h"
@@ -90,6 +91,9 @@ void UNiagaraSystem::PostInitProperties()
 
 		SystemUpdateScript = NewObject<UNiagaraScript>(this, "SystemUpdateScript", RF_Transactional);
 		SystemUpdateScript->SetUsage(ENiagaraScriptUsage::SystemUpdateScript);
+
+		INiagaraModule& NiagaraModule = FModuleManager::GetModuleChecked<INiagaraModule>("Niagara");
+		EditorData = NiagaraModule.GetEditorOnlyDataUtilities().CreateDefaultEditorData(this);
 	}
 
 	GenerateStatID();
@@ -179,7 +183,6 @@ void UNiagaraSystem::Serialize(FArchive& Ar)
 void UNiagaraSystem::PostEditChangeProperty(struct FPropertyChangedEvent& PropertyChangedEvent)
 {
 	Super::PostEditChangeProperty(PropertyChangedEvent);
-	FNiagaraSystemUpdateContext(this, true);
 
 	ThumbnailImageOutOfDate = true;
 	
@@ -248,7 +251,7 @@ void UNiagaraSystem::PostLoad()
 		SystemSpawnScript = NewObject<UNiagaraScript>(this, "SystemSpawnScript", RF_Transactional);
 		SystemSpawnScript->SetUsage(ENiagaraScriptUsage::SystemSpawnScript);
 		INiagaraModule& NiagaraModule = FModuleManager::GetModuleChecked<INiagaraModule>("Niagara");
-		SystemScriptSource = NiagaraModule.CreateDefaultScriptSource(this);
+		SystemScriptSource = NiagaraModule.GetEditorOnlyDataUtilities().CreateDefaultScriptSource(this);
 		SystemSpawnScript->SetSource(SystemScriptSource);
 	}
 	else
@@ -323,7 +326,12 @@ void UNiagaraSystem::PostLoad()
 		}
 	}
 
-	if (EditorData != nullptr)
+	if (EditorData == nullptr)
+	{
+		INiagaraModule& NiagaraModule = FModuleManager::GetModuleChecked<INiagaraModule>("Niagara");
+		EditorData = NiagaraModule.GetEditorOnlyDataUtilities().CreateDefaultEditorData(this);
+	}
+	else
 	{
 		EditorData->PostLoadFromOwner(this);
 	}
@@ -399,11 +407,6 @@ const UNiagaraEditorDataBase* UNiagaraSystem::GetEditorData() const
 	return EditorData;
 }
 
-void UNiagaraSystem::SetEditorData(UNiagaraEditorDataBase* InEditorData)
-{
-	EditorData = InEditorData;
-}
-
 bool UNiagaraSystem::ReferencesInstanceEmitter(UNiagaraEmitter& Emitter)
 {
 	for (FNiagaraEmitterHandle& Handle : EmitterHandles)
diff --git a/Engine/Plugins/FX/Niagara/Source/Niagara/Public/NiagaraEditorDataBase.h b/Engine/Plugins/FX/Niagara/Source/Niagara/Public/NiagaraEditorDataBase.h
--- a/Engine/Plugins/FX/Niagara/Source/Niagara/Public/NiagaraEditorDataBase.h
+++ b/Engine/Plugins/FX/Niagara/Source/Niagara/Public/NiagaraEditorDataBase.h
@@ -7,6 +7,40 @@
 
 #include "NiagaraEditorDataBase.generated.h"
 
+USTRUCT()
+struct FNiagaraGraphViewSettings
+{
+	GENERATED_BODY()
+public:
+	FNiagaraGraphViewSettings()
+		: Location(FVector2D::ZeroVector)
+		, Zoom(0.0f)
+		, bIsValid(false)
+	{
+	}
+
+	FNiagaraGraphViewSettings(const FVector2D& InLocation, float InZoom)
+		: Location(InLocation)
+		, Zoom(InZoom)
+		, bIsValid(true)
+	{
+	}
+
+	const FVector2D& GetLocation() const { return Location; }
+	float GetZoom() const { return Zoom; }
+	bool IsValid() const { return bIsValid; }
+
+private:
+	UPROPERTY()
+	FVector2D Location;
+
+	UPROPERTY()
+	float Zoom;
+
+	UPROPERTY()
+	bool bIsValid;
+};
+
 /** A base class for editor only data which supports post loading from the runtime owner object. */
 UCLASS(MinimalAPI)
 class UNiagaraEditorDataBase : public UObject
diff --git a/Engine/Plugins/FX/Niagara/Source/Niagara/Public/NiagaraLightRendererProperties.h b/Engine/Plugins/FX/Niagara/Source/Niagara/Public/NiagaraLightRendererProperties.h
--- a/Engine/Plugins/FX/Niagara/Source/Niagara/Public/NiagaraLightRendererProperties.h
+++ b/Engine/Plugins/FX/Niagara/Source/Niagara/Public/NiagaraLightRendererProperties.h
@@ -10,7 +10,7 @@
 
 class FNiagaraEmitterInstance;
 
-UCLASS(editinlinenew)
+UCLASS(editinlinenew, meta = (DisplayName = "Sprite Renderer"))
 class UNiagaraLightRendererProperties : public UNiagaraRendererProperties
 {
 public:
diff --git a/Engine/Plugins/FX/Niagara/Source/Niagara/Public/NiagaraMeshRendererProperties.h b/Engine/Plugins/FX/Niagara/Source/Niagara/Public/NiagaraMeshRendererProperties.h
--- a/Engine/Plugins/FX/Niagara/Source/Niagara/Public/NiagaraMeshRendererProperties.h
+++ b/Engine/Plugins/FX/Niagara/Source/Niagara/Public/NiagaraMeshRendererProperties.h
@@ -25,7 +25,7 @@ enum class ENiagaraMeshFacingMode : uint8
 	CameraPlane
 };
 
-UCLASS(editinlinenew)
+UCLASS(editinlinenew, meta = (DisplayName = "Mesh Renderer"))
 class UNiagaraMeshRendererProperties : public UNiagaraRendererProperties
 {
 public:
diff --git a/Engine/Plugins/FX/Niagara/Source/Niagara/Public/NiagaraModule.h b/Engine/Plugins/FX/Niagara/Source/Niagara/Public/NiagaraModule.h
--- a/Engine/Plugins/FX/Niagara/Source/Niagara/Public/NiagaraModule.h
+++ b/Engine/Plugins/FX/Niagara/Source/Niagara/Public/NiagaraModule.h
@@ -16,6 +16,7 @@ class UNiagaraScript;
 class FNiagaraCompileOptions;
 class FNiagaraCompileRequestDataBase;
 class INiagaraMergeManager;
+class INiagaraEditorOnlyDataUtilities;
 
 extern NIAGARA_API int32 GEnableVerboseNiagaraChangeIdLogging;
 
@@ -27,7 +28,6 @@ class NIAGARA_API INiagaraModule : public IModuleInterface
 public:
 #if WITH_EDITOR
 	typedef TSharedPtr<FNiagaraCompileRequestDataBase, ESPMode::ThreadSafe> CompileRequestPtr;
-	DECLARE_DELEGATE_RetVal_OneParam(class UNiagaraScriptSourceBase*, FOnCreateDefaultScriptSource, UObject*);
 	DECLARE_DELEGATE_RetVal_TwoParams(TSharedPtr<FNiagaraVMExecutableData>, FScriptCompiler,const FNiagaraCompileRequestDataBase*, const FNiagaraCompileOptions&);
 	DECLARE_DELEGATE_RetVal_OneParam(CompileRequestPtr, FOnPrecompile, UObject*);
 #endif
@@ -62,17 +62,17 @@ public:
 	void TickWorld(UWorld* World, ELevelTick TickType, float DeltaSeconds);
 
 #if WITH_EDITOR
-	const INiagaraMergeManager& GetMergeManager();
+	const INiagaraMergeManager& GetMergeManager() const;
 
 	void RegisterMergeManager(TSharedRef<INiagaraMergeManager> InMergeManager);
 
 	void UnregisterMergeManager(TSharedRef<INiagaraMergeManager> InMergeManager);
 
-	UNiagaraScriptSourceBase* CreateDefaultScriptSource(UObject* Outer);
+	const INiagaraEditorOnlyDataUtilities& GetEditorOnlyDataUtilities() const;
 
-	FDelegateHandle RegisterOnCreateDefaultScriptSource(FOnCreateDefaultScriptSource OnCreateDefaultScriptSource);
+	void RegisterEditorOnlyDataUtilities(TSharedRef<INiagaraEditorOnlyDataUtilities> InEditorOnlyDataUtilities);
 
-	void UnregisterOnCreateDefaultScriptSource(FDelegateHandle DelegateHandle);
+	void UnregisterEditorOnlyDataUtilities(TSharedRef<INiagaraEditorOnlyDataUtilities> InEditorOnlyDataUtilities);
 
 	TSharedPtr<FNiagaraVMExecutableData> CompileScript(const FNiagaraCompileRequestDataBase* InCompileData, const FNiagaraCompileOptions& InCompileOptions);
 	FDelegateHandle RegisterScriptCompiler(FScriptCompiler ScriptCompiler);
@@ -178,7 +178,8 @@ private:
 
 #if WITH_EDITORONLY_DATA
 	TSharedPtr<INiagaraMergeManager> MergeManager;
-	FOnCreateDefaultScriptSource OnCreateDefaultScriptSourceDelegate;
+	TSharedPtr<INiagaraEditorOnlyDataUtilities> EditorOnlyDataUtilities;
+
 	FScriptCompiler ScriptCompilerDelegate;
 	FOnPrecompile ObjectPrecompilerDelegate;
 #endif
diff --git a/Engine/Plugins/FX/Niagara/Source/Niagara/Public/NiagaraRibbonRendererProperties.h b/Engine/Plugins/FX/Niagara/Source/Niagara/Public/NiagaraRibbonRendererProperties.h
--- a/Engine/Plugins/FX/Niagara/Source/Niagara/Public/NiagaraRibbonRendererProperties.h
+++ b/Engine/Plugins/FX/Niagara/Source/Niagara/Public/NiagaraRibbonRendererProperties.h
@@ -54,7 +54,7 @@ enum class ENiagaraRibbonTessellationMode : uint8
 	Disabled
 };
 
-UCLASS(editinlinenew)
+UCLASS(editinlinenew, meta = (DisplayName = "Ribbon Renderer"))
 class UNiagaraRibbonRendererProperties : public UNiagaraRendererProperties
 {
 public:
diff --git a/Engine/Plugins/FX/Niagara/Source/Niagara/Public/NiagaraSpriteRendererProperties.h b/Engine/Plugins/FX/Niagara/Source/Niagara/Public/NiagaraSpriteRendererProperties.h
--- a/Engine/Plugins/FX/Niagara/Source/Niagara/Public/NiagaraSpriteRendererProperties.h
+++ b/Engine/Plugins/FX/Niagara/Source/Niagara/Public/NiagaraSpriteRendererProperties.h
@@ -39,7 +39,7 @@ enum class ENiagaraSpriteFacingMode : uint8
 };
 
 
-UCLASS(editinlinenew)
+UCLASS(editinlinenew, meta = (DisplayName = "Sprite Renderer"))
 class NIAGARA_API UNiagaraSpriteRendererProperties : public UNiagaraRendererProperties
 {
 public:
diff --git a/Engine/Plugins/FX/Niagara/Source/NiagaraEditor/Private/NiagaraEditorModule.cpp b/Engine/Plugins/FX/Niagara/Source/NiagaraEditor/Private/NiagaraEditorModule.cpp
--- a/Engine/Plugins/FX/Niagara/Source/NiagaraEditor/Private/NiagaraEditorModule.cpp
+++ b/Engine/Plugins/FX/Niagara/Source/NiagaraEditor/Private/NiagaraEditorModule.cpp
@@ -68,6 +68,7 @@
 #include "NiagaraScriptSource.h"
 #include "NiagaraTypes.h"
 #include "NiagaraSystemFactoryNew.h"
+#include "NiagaraSystemEditorData.h"
 
 #include "MovieScene/Parameters/MovieSceneNiagaraBoolParameterTrack.h"
 #include "MovieScene/Parameters/MovieSceneNiagaraFloatParameterTrack.h"
@@ -96,6 +97,7 @@
 #include "Customizations/NiagaraFunctionCallNodeDetails.h"
 
 #include "NiagaraNodeFunctionCall.h"
+#include "INiagaraEditorOnlyDataUtlities.h"
 
 IMPLEMENT_MODULE( FNiagaraEditorModule, NiagaraEditor );
 
@@ -110,6 +112,26 @@ EAssetTypeCategories::Type FNiagaraEditorModule::NiagaraAssetCategory;
 
 //////////////////////////////////////////////////////////////////////////
 
+class FNiagaraEditorOnlyDataUtilities : public INiagaraEditorOnlyDataUtilities
+{
+	UNiagaraScriptSourceBase* CreateDefaultScriptSource(UObject* InOuter) const
+	{
+		return NewObject<UNiagaraScriptSource>(InOuter);
+	}
+
+	virtual UNiagaraEditorDataBase* CreateDefaultEditorData(UObject* InOuter) const
+	{
+		UNiagaraSystem* System = Cast<UNiagaraSystem>(InOuter);
+		if (System != nullptr)
+		{
+			UNiagaraSystemEditorData* SystemEditorData = NewObject<UNiagaraSystemEditorData>(InOuter);
+			SystemEditorData->SynchronizeOverviewGraphWithSystem(*System);
+			return SystemEditorData;
+		}
+		return nullptr;
+	}
+};
+
 class FNiagaraScriptGraphPanelPinFactory : public FGraphPanelPinFactory
 {
 public:
@@ -623,10 +645,13 @@ void FNiagaraEditorModule::StartupModule()
 		FNiagaraShaderQueueTickable::ProcessQueue();
 	}));
 
-	// Register the emitter merge handler.
+	// Register the emitter merge handler and editor data utilities.
 	ScriptMergeManager = MakeShared<FNiagaraScriptMergeManager>();
 	NiagaraModule.RegisterMergeManager(ScriptMergeManager.ToSharedRef());
 
+	EditorOnlyDataUtilities = MakeShared<FNiagaraEditorOnlyDataUtilities>();
+	NiagaraModule.RegisterEditorOnlyDataUtilities(EditorOnlyDataUtilities.ToSharedRef());
+
 	// Register the script compiler
 	ScriptCompilerHandle = NiagaraModule.RegisterScriptCompiler(INiagaraModule::FScriptCompiler::CreateLambda([this](const FNiagaraCompileRequestDataBase* CompileRequest, const FNiagaraCompileOptions& Options)
 	{
@@ -638,11 +663,6 @@ void FNiagaraEditorModule::StartupModule()
 		return Precompile(InObj);
 	}));
 
-	// Register the create default script source handler.
-	CreateDefaultScriptSourceHandle = NiagaraModule.RegisterOnCreateDefaultScriptSource(
-		INiagaraModule::FOnCreateDefaultScriptSource::CreateLambda([](UObject* Outer) { return NewObject<UNiagaraScriptSource>(Outer); }));
-
-
 	TestCompileScriptCommand = IConsoleManager::Get().RegisterConsoleCommand(
 		TEXT("fx.TestCompileNiagaraScript"),
 		TEXT("Compiles the specified script on disk for the niagara vector vm"),
@@ -658,7 +678,7 @@ void FNiagaraEditorModule::StartupModule()
 		TEXT("Forces the system to refresh all it's dependencies so it won't recompile on load.  This may mark multiple assets dirty for re-saving."),
 		FConsoleCommandWithArgsDelegate::CreateStatic(&PreventSystemRecompile));
 
-	PreventSystemRecompileCommand = IConsoleManager::Get().RegisterConsoleCommand(
+	PreventAllSystemRecompilesCommand = IConsoleManager::Get().RegisterConsoleCommand(
 		TEXT("fx.PreventAllSystemRecompiles"),
 		TEXT("Loads all of the systems in the project and forces each system to refresh all it's dependencies so it won't recompile on load.  This may mark multiple assets dirty for re-saving."),
 		FConsoleCommandDelegate::CreateStatic(&PreventAllSystemRecompiles));
@@ -732,7 +752,7 @@ void FNiagaraEditorModule::ShutdownModule()
 	if (NiagaraModule != nullptr)
 	{
 		NiagaraModule->UnregisterMergeManager(ScriptMergeManager.ToSharedRef());
-		NiagaraModule->UnregisterOnCreateDefaultScriptSource(CreateDefaultScriptSourceHandle);
+		NiagaraModule->UnregisterEditorOnlyDataUtilities(EditorOnlyDataUtilities.ToSharedRef());
 		NiagaraModule->UnregisterScriptCompiler(ScriptCompilerHandle);
 		NiagaraModule->UnregisterPrecompiler(PrecompilerHandle);
 	}
@@ -752,6 +772,16 @@ void FNiagaraEditorModule::ShutdownModule()
 		IConsoleManager::Get().UnregisterConsoleObject(DumpRapidIterationParametersForAsset);
 	}
 
+	if (PreventSystemRecompileCommand != nullptr)
+	{
+		IConsoleManager::Get().UnregisterConsoleObject(PreventSystemRecompileCommand);
+	}
+
+	if (PreventAllSystemRecompilesCommand != nullptr)
+	{
+		IConsoleManager::Get().UnregisterConsoleObject(PreventAllSystemRecompilesCommand);
+	}
+
 	if (UObjectInitialized() && GIsEditor)
 	{
 		UThumbnailManager::Get().UnregisterCustomRenderer(UNiagaraEmitter::StaticClass());
@@ -852,6 +882,11 @@ const FNiagaraEditorCommands& FNiagaraEditorModule::Commands()
 	return FNiagaraEditorCommands::Get();
 }
 
+TSharedPtr<FNiagaraSystemViewModel> FNiagaraEditorModule::GetExistingViewModelForSystem(UNiagaraSystem* InSystem)
+{
+	return FNiagaraSystemViewModel::GetExistingViewModelForObject(InSystem);
+}
+
 void FNiagaraEditorModule::RegisterAssetTypeAction(IAssetTools& AssetTools, TSharedRef<IAssetTypeActions> Action)
 {
 	AssetTools.RegisterAssetTypeActions(Action);
diff --git a/Engine/Plugins/FX/Niagara/Source/NiagaraEditor/Private/NiagaraEditorStyle.cpp b/Engine/Plugins/FX/Niagara/Source/NiagaraEditor/Private/NiagaraEditorStyle.cpp
--- a/Engine/Plugins/FX/Niagara/Source/NiagaraEditor/Private/NiagaraEditorStyle.cpp
+++ b/Engine/Plugins/FX/Niagara/Source/NiagaraEditor/Private/NiagaraEditorStyle.cpp
@@ -125,6 +125,8 @@ TSharedRef< FSlateStyleSet > FNiagaraEditorStyle::Create()
 	FEditableTextBoxStyle HeadingEditableTextBox = FEditableTextBoxStyle(NormalEditableTextBox)
 		.SetFont(DEFAULT_FONT("Regular", 14));
 
+	Style->Set("NiagaraEditor.HeadingTextBlock", HeadingText);
+
 	Style->Set("NiagaraEditor.HeadingEditableTextBox", HeadingEditableTextBox);
 
 	Style->Set("NiagaraEditor.HeadingInlineEditableText", FInlineEditableTextBlockStyle()
diff --git a/Engine/Plugins/FX/Niagara/Source/NiagaraEditor/Private/NiagaraEditorUtilities.cpp b/Engine/Plugins/FX/Niagara/Source/NiagaraEditor/Private/NiagaraEditorUtilities.cpp
--- a/Engine/Plugins/FX/Niagara/Source/NiagaraEditor/Private/NiagaraEditorUtilities.cpp
+++ b/Engine/Plugins/FX/Niagara/Source/NiagaraEditor/Private/NiagaraEditorUtilities.cpp
@@ -10,6 +10,7 @@
 #include "UObject/StructOnScope.h"
 #include "NiagaraGraph.h"
 #include "NiagaraSystem.h"
+#include "NiagaraSystemEditorData.h"
 #include "NiagaraScriptSource.h"
 #include "NiagaraScript.h"
 #include "NiagaraNodeOutput.h"
@@ -1083,7 +1084,7 @@ TArray<UNiagaraComponent*> FNiagaraEditorUtilities::GetComponentsThatReferenceSy
 	return ReferencingComponents;
 }
 
-const FGuid FNiagaraEditorUtilities::AddEmitterToSystem(UNiagaraSystem& InSystem, UNiagaraEmitter& InEmitterToAdd, bool bSystemIsPlaceholder)
+const FGuid FNiagaraEditorUtilities::AddEmitterToSystem(UNiagaraSystem& InSystem, UNiagaraEmitter& InEmitterToAdd)
 {
 	// Kill all system instances before modifying the emitter handle list to prevent accessing deleted data.
 	KillSystemInstances(InSystem);
@@ -1094,8 +1095,9 @@ const FGuid FNiagaraEditorUtilities::AddEmitterToSystem(UNiagaraSystem& InSystem
 		EmitterHandleNames.Add(EmitterHandle.GetName());
 	}
 
+	UNiagaraSystemEditorData* SystemEditorData = CastChecked<UNiagaraSystemEditorData>(InSystem.GetEditorData(), ECastCheckedType::NullChecked);
 	FNiagaraEmitterHandle EmitterHandle;
-	if (false == bSystemIsPlaceholder)
+	if (SystemEditorData->GetOwningSystemIsPlaceholder() == false)
 	{
 		InSystem.Modify();
 		EmitterHandle = InSystem.AddEmitterHandle(InEmitterToAdd, FNiagaraUtilities::GetUniqueName(InEmitterToAdd.GetFName(), EmitterHandleNames));
@@ -1110,6 +1112,7 @@ const FGuid FNiagaraEditorUtilities::AddEmitterToSystem(UNiagaraSystem& InSystem
 	}
 	
 	FNiagaraStackGraphUtilities::RebuildEmitterNodes(InSystem);
+	SystemEditorData->SynchronizeOverviewGraphWithSystem(InSystem);
 
 	return EmitterHandle.GetId();
 }
@@ -1125,7 +1128,10 @@ void FNiagaraEditorUtilities::RemoveEmittersFromSystemByEmitterHandleId(UNiagara
 
 	InSystem.Modify();
 	InSystem.RemoveEmitterHandlesById(EmitterHandleIdsToDelete);
+
 	FNiagaraStackGraphUtilities::RebuildEmitterNodes(InSystem);
+	UNiagaraSystemEditorData* SystemEditorData = CastChecked<UNiagaraSystemEditorData>(InSystem.GetEditorData(), ECastCheckedType::NullChecked);
+	SystemEditorData->SynchronizeOverviewGraphWithSystem(InSystem);
 }
 
 void FNiagaraEditorUtilities::KillSystemInstances(const UNiagaraSystem& System)
diff --git a/Engine/Plugins/FX/Niagara/Source/NiagaraEditor/Private/NiagaraOverviewNode.cpp b/Engine/Plugins/FX/Niagara/Source/NiagaraEditor/Private/NiagaraOverviewNode.cpp
new file mode 100644
--- /dev/null
+++ b/Engine/Plugins/FX/Niagara/Source/NiagaraEditor/Private/NiagaraOverviewNode.cpp
@@ -0,0 +1,65 @@
+// Copyright 1998-2019 Epic Games, Inc. All Rights Reserved.
+
+#include "NiagaraOverviewNode.h"
+#include "NiagaraSystem.h"
+
+#define LOCTEXT_NAMESPACE "NiagaraOverviewNodeStackItem"
+
+UNiagaraOverviewNode::UNiagaraOverviewNode()
+	: OwningSystem(nullptr)
+	, EmitterHandleGuid(FGuid())
+{
+};
+
+void UNiagaraOverviewNode::Initialize(UNiagaraSystem* InOwningSystem)
+{
+	OwningSystem = InOwningSystem;
+}
+
+void UNiagaraOverviewNode::Initialize(UNiagaraSystem* InOwningSystem, FGuid InEmitterHandleGuid)
+{
+	OwningSystem = InOwningSystem;
+	EmitterHandleGuid = InEmitterHandleGuid;
+}
+
+const FGuid UNiagaraOverviewNode::GetEmitterHandleGuid() const
+{
+	return EmitterHandleGuid;
+}
+
+FText UNiagaraOverviewNode::GetNodeTitle(ENodeTitleType::Type TitleType) const
+{
+	if (EmitterHandleGuid.IsValid())
+	{
+		for (const FNiagaraEmitterHandle& Handle : OwningSystem->GetEmitterHandles())
+		{
+			if (Handle.GetId() == EmitterHandleGuid)
+			{
+				return FText::FromName(Handle.GetName());
+			}
+		}
+		ensureMsgf(false, TEXT("Failed to find matching emitter handle for existing overview node!"));
+		return LOCTEXT("UnknownEmitterName", "Unknown Emitter");
+	}
+	else
+	{
+		return FText::FromString(OwningSystem->GetName());
+	}
+}
+
+bool UNiagaraOverviewNode::CanUserDeleteNode() const
+{
+	return EmitterHandleGuid.IsValid();
+}
+
+bool UNiagaraOverviewNode::CanDuplicateNode() const
+{
+	return EmitterHandleGuid.IsValid();
+}
+
+UNiagaraSystem* UNiagaraOverviewNode::GetOwningSystem()
+{
+	return OwningSystem;
+}
+
+#undef LOCTEXT_NAMESPACE
diff --git a/Engine/Plugins/FX/Niagara/Source/NiagaraEditor/Private/NiagaraSystemEditorData.cpp b/Engine/Plugins/FX/Niagara/Source/NiagaraEditor/Private/NiagaraSystemEditorData.cpp
--- a/Engine/Plugins/FX/Niagara/Source/NiagaraEditor/Private/NiagaraSystemEditorData.cpp
+++ b/Engine/Plugins/FX/Niagara/Source/NiagaraEditor/Private/NiagaraSystemEditorData.cpp
@@ -7,11 +7,11 @@
 #include "NiagaraEmitterHandle.h"
 #include "NiagaraEmitter.h"
 #include "NiagaraEmitterEditorData.h"
-#include "NiagaraOverviewNodeStackItem.h"
+#include "NiagaraOverviewNode.h"
 #include "EdGraphSchema_NiagaraSystemOverview.h"
 #include "EdGraph/EdGraph.h"
 
-static const float SystemOverviewNodePadding = 300.0f;
+static const float SystemOverviewNodePadding = 250.0f;
 
 const FName UNiagaraSystemEditorFolder::GetFolderName() const
 {
@@ -64,6 +64,7 @@ UNiagaraSystemEditorData::UNiagaraSystemEditorData(const FObjectInitializer& Obj
 	OwnerTransform.SetLocation(FVector(0.0f, 0.0f, 0.0f));
 	PlaybackRangeMin = 0;
 	PlaybackRangeMax = 10;
+	bSystemIsPlaceholder = false;
 }
 
 void UNiagaraSystemEditorData::PostInitProperties()
@@ -78,6 +79,8 @@ void UNiagaraSystemEditorData::PostInitProperties()
 
 void UNiagaraSystemEditorData::PostLoadFromOwner(UObject* InOwner)
 {
+	UNiagaraSystem* OwningSystem = CastChecked<UNiagaraSystem>(InOwner, ECastCheckedType::NullChecked);
+
 	if (RootFolder == nullptr)
 	{
 		RootFolder = NewObject<UNiagaraSystemEditorFolder>(this, TEXT("RootFolder"), RF_Transactional);
@@ -91,13 +94,14 @@ void UNiagaraSystemEditorData::PostLoadFromOwner(UObject* InOwner)
 
 	if (NiagaraVer < FNiagaraCustomVersion::PlaybackRangeStoredOnSystem)
 	{
-		UpdatePlaybackRangeFromEmitters(OwningSystem);
+		UpdatePlaybackRangeFromEmitters(*OwningSystem);
 	}
 
 	if (SystemOverviewGraph == nullptr)
 	{
 		SystemOverviewGraph = NewObject<UEdGraph>(this, "NiagaraOverview", RF_Transactional);
 		SystemOverviewGraph->Schema = UEdGraphSchema_NiagaraSystemOverview::StaticClass();
+		SynchronizeOverviewGraphWithSystem(*OwningSystem);
 	}
 }
 
@@ -127,14 +131,35 @@ UEdGraph* UNiagaraSystemEditorData::GetSystemOverviewGraph() const
 	return SystemOverviewGraph;	
 }
 
-void UNiagaraSystemEditorData::UpdatePlaybackRangeFromEmitters(UNiagaraSystem* OwnerSystem)
+const FNiagaraGraphViewSettings& UNiagaraSystemEditorData::GetSystemOverviewGraphViewSettings() const
+{
+	return OverviewGraphViewSettings;
+}
+
+void UNiagaraSystemEditorData::SetSystemOverviewGraphViewSettings(const FNiagaraGraphViewSettings& InOverviewGraphViewSettings)
+{
+	OverviewGraphViewSettings = InOverviewGraphViewSettings;
+}
+
+bool UNiagaraSystemEditorData::GetOwningSystemIsPlaceholder() const
 {
-	if (OwnerSystem->GetEmitterHandles().Num() > 0)
+	return bSystemIsPlaceholder;
+}
+
+void UNiagaraSystemEditorData::SetOwningSystemIsPlaceholder(bool bInSystemIsPlaceholder, UNiagaraSystem& OwnerSystem)
+{
+	bSystemIsPlaceholder = bInSystemIsPlaceholder;
+	SynchronizeOverviewGraphWithSystem(OwnerSystem);
+}
+
+void UNiagaraSystemEditorData::UpdatePlaybackRangeFromEmitters(UNiagaraSystem& OwnerSystem)
+{
+	if (OwnerSystem.GetEmitterHandles().Num() > 0)
 	{
 		float EmitterPlaybackRangeMin = TNumericLimits<float>::Max();
 		float EmitterPlaybackRangeMax = TNumericLimits<float>::Lowest();
 
-		for (const FNiagaraEmitterHandle& EmitterHandle : OwnerSystem->GetEmitterHandles())
+		for (const FNiagaraEmitterHandle& EmitterHandle : OwnerSystem.GetEmitterHandles())
 		{
 			UNiagaraEmitterEditorData* EmitterEditorData = Cast<UNiagaraEmitterEditorData>(EmitterHandle.GetInstance()->EditorData);
 			if (EmitterEditorData != nullptr)
@@ -149,105 +174,110 @@ void UNiagaraSystemEditorData::UpdatePlaybackRangeFromEmitters(UNiagaraSystem* O
 	}
 }
 
-void UNiagaraSystemEditorData::InitSystemOverviewGraph()
-{
-	float NextNodePosX = 0.0;
-
-	FGraphNodeCreator<UNiagaraOverviewNodeStackItem> SystemOverviewNodeCreator(*SystemOverviewGraph);
-	UNiagaraOverviewNodeStackItem* SystemOverviewNode = SystemOverviewNodeCreator.CreateNode(false);
-	SystemOverviewNode->Initialize(OwningSystem);
-
-	SystemOverviewNode->NodePosX = NextNodePosX;
-	SystemOverviewNode->NodePosY = 0.0;
-	NextNodePosX += SystemOverviewNodePadding;
-
-	SystemOverviewNodeCreator.Finalize();
-
-	for (const FNiagaraEmitterHandle& Handle : OwningSystem->GetEmitterHandles())
-	{
-		FGraphNodeCreator<UNiagaraOverviewNodeStackItem> EmitterOverviewNodeCreator(*SystemOverviewGraph);
-		UNiagaraOverviewNodeStackItem* EmitterOverviewNode = EmitterOverviewNodeCreator.CreateNode(false);
-		EmitterOverviewNode->Initialize(OwningSystem, Handle.GetId());
-
-		EmitterOverviewNode->NodePosX = NextNodePosX;
-		EmitterOverviewNode->NodePosY = 0.0;
-		NextNodePosX += SystemOverviewNodePadding;
-
-		EmitterOverviewNodeCreator.Finalize();
-	}
-}
-
-const FVector2D UNiagaraSystemEditorData::GetGoodPlaceForNewOverviewNode() const
+void FindOuterNodes(const TArray<UNiagaraOverviewNode*> OverviewNodes, UNiagaraOverviewNode*& OutLeft, UNiagaraOverviewNode*& OutRight)
 {
-	TArray<UNiagaraOverviewNodeStackItem*> OverviewNodes;
-	SystemOverviewGraph->GetNodesOfClass<UNiagaraOverviewNodeStackItem>(OverviewNodes);
-	const TArray<FNiagaraEmitterHandle>& EmitterHandles = OwningSystem->GetEmitterHandles();
-
-	// Find the last created Emitter overview node location and place next to that. 
-	if (EmitterHandles.Num() > 1)
+	if (OverviewNodes.Num() > 0)
 	{
-		const FGuid PreviousEmitterHandleGuid = EmitterHandles.Last(1).GetId();
-		for (const UNiagaraOverviewNodeStackItem* PerOverviewNode : OverviewNodes)
+		OutLeft = OverviewNodes[0];
+		OutRight = OutLeft;
+		for (int32 i = 1; i < OverviewNodes.Num(); i++)
 		{
-			if (PreviousEmitterHandleGuid == PerOverviewNode->GetEmitterHandleGuid())
+			UNiagaraOverviewNode* Current = OverviewNodes[i];
+			if (Current->NodePosX < OutLeft->NodePosX)
+			{
+				OutLeft = Current;
+			}
+			if (Current->NodePosX > OutRight->NodePosX)
 			{
-				return FVector2D(PerOverviewNode->NodePosX + SystemOverviewNodePadding, PerOverviewNode->NodePosY);
+				OutRight = Current;
 			}
 		}
 	}
-	return SystemOverviewGraph->GetGoodPlaceForNewNode();
-}
-
-bool UNiagaraSystemEditorData::GetSystemOverviewGraphIsValid() const
-{
-	TArray<UEdGraphNode*> Nodes;
-	SystemOverviewGraph->GetNodesOfClass(Nodes);
-	if (Nodes.Num() > 0)
+	else
 	{
-		return true;
+		OutLeft = nullptr;
+		OutRight = nullptr;
 	}
-	return false;
 }
 
-void UNiagaraSystemEditorData::Initialize(UNiagaraSystem* OwnerSystem, bool bEditingSystem)
+void UNiagaraSystemEditorData::SynchronizeOverviewGraphWithSystem(UNiagaraSystem& OwnerSystem)
 {
-	OwningSystem = OwnerSystem;
+	float NextNodePosX = 0.0;
 
-	if (bEditingSystem && GetSystemOverviewGraphIsValid() == false)
-	{
-		InitSystemOverviewGraph();
-	}
-}
+	TArray<UNiagaraOverviewNode*> OverviewNodes;
+	SystemOverviewGraph->GetNodesOfClass<UNiagaraOverviewNode>(OverviewNodes);
 
-void UNiagaraSystemEditorData::SystemOverviewHandleAdded(const FGuid AddedHandleGuid) const
-{
-	SystemOverviewGraph->Modify();
+	UNiagaraOverviewNode* LeftNode;
+	UNiagaraOverviewNode* RightNode;
+	FindOuterNodes(OverviewNodes, LeftNode, RightNode);
 
-	const FVector2D NewNodeLocation = GetGoodPlaceForNewOverviewNode();
+	// Validate system node.
+	int32 SystemNodeIndex = OverviewNodes.IndexOfByPredicate([] (UNiagaraOverviewNode* OverviewNode) { return OverviewNode->GetEmitterHandleGuid().IsValid() == false; });
+	if (SystemNodeIndex != INDEX_NONE)
+	{
+		if (bSystemIsPlaceholder)
+		{
+			// Systems which are placeholder are for editing emitter assets and shouldn't have system nodes in their overview graph.
+			UNiagaraOverviewNode* SystemNode = OverviewNodes[SystemNodeIndex];
+			SystemNode->DestroyNode();
+		}
 
-	FGraphNodeCreator<UNiagaraOverviewNodeStackItem> EmitterOverviewNodeCreator(*SystemOverviewGraph);
-	UNiagaraOverviewNodeStackItem* EmitterOverviewNode = EmitterOverviewNodeCreator.CreateNode(false);
-	EmitterOverviewNode->Initialize(OwningSystem, AddedHandleGuid);
+		// Remove the node from the list to track which nodes are no longer used.
+		OverviewNodes.RemoveAt(SystemNodeIndex);
+	}
+	else
+	{
+		if(bSystemIsPlaceholder == false)
+		{
+			FGraphNodeCreator<UNiagaraOverviewNode> SystemOverviewNodeCreator(*SystemOverviewGraph);
+			UNiagaraOverviewNode* SystemOverviewNode = SystemOverviewNodeCreator.CreateNode(false);
+			SystemOverviewNode->Initialize(&OwnerSystem);
 
-	EmitterOverviewNode->NodePosX = NewNodeLocation.X;
-	EmitterOverviewNode->NodePosY = NewNodeLocation.Y;
+			FVector2D SystemOverviewNodeLocation = LeftNode != nullptr
+				? FVector2D(LeftNode->NodePosX - SystemOverviewNodePadding, LeftNode->NodePosY)
+				: SystemOverviewGraph->GetGoodPlaceForNewNode();
+			SystemOverviewNode->NodePosX = SystemOverviewNodeLocation.X;
+			SystemOverviewNode->NodePosY = SystemOverviewNodeLocation.Y;
 
-	EmitterOverviewNodeCreator.Finalize();
+			SystemOverviewNodeCreator.Finalize();
 
-	//@TODO System Overview: if we have a ref to the SGraphEditor here, then focus the new node
-}
+			if (RightNode == nullptr)
+			{
+				RightNode = SystemOverviewNode;
+			}
+		}
+	}
 
-void UNiagaraSystemEditorData::SystemOverviewHandlesRemoved() const
-{
-	TArray<UNiagaraOverviewNodeStackItem*> CurrentOverviewNodes;
-	SystemOverviewGraph->GetNodesOfClass<UNiagaraOverviewNodeStackItem>(CurrentOverviewNodes);
-	for (UNiagaraOverviewNodeStackItem* OverviewNode : CurrentOverviewNodes)
+	// Validate emitter nodes.
+	for (const FNiagaraEmitterHandle& EmitterHandle : OwnerSystem.GetEmitterHandles())
 	{
-		const FGuid OverviewNodeGuid = OverviewNode->GetEmitterHandleGuid();
-		// If the OverviewNode's EmitterHandleGuid is valid (not representing a system) and is not in the current array of the owning System's EmitterHandles, that node can be removed.
-		if (OverviewNodeGuid.IsValid() && false == OwningSystem->GetEmitterHandles().ContainsByPredicate([&OverviewNodeGuid](const FNiagaraEmitterHandle& Handle) {return Handle.GetId() == OverviewNodeGuid; }))
+		int32 EmitterNodeIndex = OverviewNodes.IndexOfByPredicate([&EmitterHandle] (UNiagaraOverviewNode* OverviewNode) { return OverviewNode->GetEmitterHandleGuid() == EmitterHandle.GetId(); });
+		if (EmitterNodeIndex != INDEX_NONE)
+		{
+			// If a node for this emitter exists already, remove it from the collection so that we can track nodes for emitters which no longer exist.
+			OverviewNodes.RemoveAt(EmitterNodeIndex);
+		}
+		else
 		{
-			SystemOverviewGraph->RemoveNode(OverviewNode);
+			FGraphNodeCreator<UNiagaraOverviewNode> EmitterOverviewNodeCreator(*SystemOverviewGraph);
+			UNiagaraOverviewNode* EmitterOverviewNode = EmitterOverviewNodeCreator.CreateNode(false);
+			EmitterOverviewNode->Initialize(&OwnerSystem, EmitterHandle.GetId());
+
+			FVector2D EmitterOverviewNodeLocation = RightNode != nullptr
+				? FVector2D(RightNode->NodePosX + SystemOverviewNodePadding, RightNode->NodePosY)
+				: SystemOverviewGraph->GetGoodPlaceForNewNode();
+			EmitterOverviewNode->NodePosX = EmitterOverviewNodeLocation.X;
+			EmitterOverviewNode->NodePosY = EmitterOverviewNodeLocation.Y;
+
+			EmitterOverviewNodeCreator.Finalize();
+
+			RightNode = EmitterOverviewNode;
 		}
 	}
+
+	// If there are any nodes remaining in the list they're no longer being used so destroy them.
+	for (UNiagaraOverviewNode* OverviewNode : OverviewNodes)
+	{
+		OverviewNode->DestroyNode();
+	}
 }
diff --git a/Engine/Plugins/FX/Niagara/Source/NiagaraEditor/Private/NiagaraSystemFactoryNew.cpp b/Engine/Plugins/FX/Niagara/Source/NiagaraEditor/Private/NiagaraSystemFactoryNew.cpp
--- a/Engine/Plugins/FX/Niagara/Source/NiagaraEditor/Private/NiagaraSystemFactoryNew.cpp
+++ b/Engine/Plugins/FX/Niagara/Source/NiagaraEditor/Private/NiagaraSystemFactoryNew.cpp
@@ -132,7 +132,7 @@ UObject* UNiagaraSystemFactoryNew::FactoryCreateNew(UClass* Class, UObject* InPa
 
 		for (UNiagaraEmitter* EmitterToAddToNewSystem : EmittersToAddToNewSystem)
 		{
-			FNiagaraEditorUtilities::AddEmitterToSystem(*NewSystem, *EmitterToAddToNewSystem, false);
+			FNiagaraEditorUtilities::AddEmitterToSystem(*NewSystem, *EmitterToAddToNewSystem);
 		}
 	}
 	else
diff --git a/Engine/Plugins/FX/Niagara/Source/NiagaraEditor/Private/Sequencer/NiagaraSequence/NiagaraEmitterTrackEditor.cpp b/Engine/Plugins/FX/Niagara/Source/NiagaraEditor/Private/Sequencer/NiagaraSequence/NiagaraEmitterTrackEditor.cpp
--- a/Engine/Plugins/FX/Niagara/Source/NiagaraEditor/Private/Sequencer/NiagaraSequence/NiagaraEmitterTrackEditor.cpp
+++ b/Engine/Plugins/FX/Niagara/Source/NiagaraEditor/Private/Sequencer/NiagaraSequence/NiagaraEmitterTrackEditor.cpp
@@ -6,6 +6,7 @@
 #include "ViewModels/NiagaraSystemViewModel.h"
 #include "ViewModels/NiagaraEmitterHandleViewModel.h"
 #include "ViewModels/NiagaraEmitterViewModel.h"
+#include "ViewModels/NiagaraSystemSelectionViewModel.h"
 #include "Sequencer/NiagaraSequence/NiagaraSequence.h"
 #include "Framework/MultiBox/MultiBoxBuilder.h"
 #include "NiagaraEditorStyle.h"
@@ -141,28 +142,28 @@ private:
 
 	FReply OnToggleIsolateButtonClicked()
 	{
-		TArray<TSharedRef<FNiagaraEmitterHandleViewModel>> EmittersToIsolate;
+		TArray<FGuid> EmitterIdsToIsolate;
 		if (EmitterTrack.IsValid())
 		{
-			if (EmitterTrack->GetSystemViewModel().IsEmitterIsolated(EmitterTrack->GetEmitterHandleViewModel().ToSharedRef()) == false)
+			if (EmitterTrack->GetEmitterHandleViewModel()->GetIsIsolated() == false)
 			{
-				EmittersToIsolate.Add(EmitterTrack->GetEmitterHandleViewModel().ToSharedRef());
+				EmitterIdsToIsolate.Add(EmitterTrack->GetEmitterHandleViewModel()->GetId());
 			}
-			EmitterTrack->GetSystemViewModel().IsolateEmitters(EmittersToIsolate);
+			EmitterTrack->GetSystemViewModel().IsolateEmitters(EmitterIdsToIsolate);
 		}
 		return FReply::Handled();
 	}
 
 	FText GetToggleIsolateToolTip() const
 	{
-		return EmitterTrack.IsValid() && EmitterTrack->GetSystemViewModel().IsEmitterIsolated(EmitterTrack->GetEmitterHandleViewModel().ToSharedRef())
+		return EmitterTrack.IsValid() && EmitterTrack->GetEmitterHandleViewModel()->GetIsIsolated()
 			? LOCTEXT("TurnOffEmitterIsolation", "Disable emitter isolation.")
 			: LOCTEXT("IsolateThisEmitter", "Enable isolation for this emitter.");
 	}
 
 	FSlateColor GetToggleIsolateImageColor() const
 	{
-		return EmitterTrack.IsValid() && EmitterTrack->GetSystemViewModel().IsEmitterIsolated(EmitterTrack->GetEmitterHandleViewModel().ToSharedRef())
+		return EmitterTrack.IsValid() && EmitterTrack->GetEmitterHandleViewModel()->GetIsIsolated()
 			? FEditorStyle::GetSlateColor("SelectionColor")
 			: FLinearColor::Gray;
 	}
@@ -233,22 +234,21 @@ void FNiagaraEmitterTrackEditor::BuildTrackContextMenu( FMenuBuilder& MenuBuilde
 		MenuBuilder.BeginSection("Niagara", LOCTEXT("NiagaraContextMenuSectionName", "Niagara"));
 		{
 			MenuBuilder.AddMenuEntry(
-				SystemViewModel.IsEmitterIsolated(EmitterTrack->GetEmitterHandleViewModel().ToSharedRef())
+				EmitterTrack->GetEmitterHandleViewModel()->GetIsIsolated()
 					? LOCTEXT("RemoveFromIsolation", "Remove this from isolation.")
 					: LOCTEXT("AddToIsolation", "Add this to isolation"),
-				SystemViewModel.IsEmitterIsolated(EmitterTrack->GetEmitterHandleViewModel().ToSharedRef())
+				EmitterTrack->GetEmitterHandleViewModel()->GetIsIsolated()
 					? LOCTEXT("RemoveFromIsolation_NoChangeOthers", "Remove this emitter from isolation, without changing other emitters.")
 					: LOCTEXT("AddToIsolation_NoChangeOthers", "Add this emitter to isolation, without changing other emitters."),
 				FSlateIcon(),
 				FUIAction(FExecuteAction::CreateRaw(&SystemViewModel, &FNiagaraSystemViewModel::ToggleEmitterIsolation, EmitterTrack->GetEmitterHandleViewModel().ToSharedRef())));
 			
-			TArray<TSharedRef<FNiagaraEmitterHandleViewModel>> SelectedEmitters;
-			SystemViewModel.GetSelectedEmitterHandles(SelectedEmitters);
+			TArray<FGuid> SelectedEmitterHandleIds = SystemViewModel.GetSelectionViewModel()->GetSelectedEmitterHandleIds();
 			MenuBuilder.AddMenuEntry(
 				LOCTEXT("IsolateSelected", "Isolate all selected"),
 				LOCTEXT("IsolateSelectedToolTip", "Add all of the selected emitters to isloation"),
 				FSlateIcon(),
-				FUIAction(FExecuteAction::CreateRaw(&SystemViewModel, &FNiagaraSystemViewModel::IsolateEmitters, SelectedEmitters)));
+				FUIAction(FExecuteAction::CreateRaw(&SystemViewModel, &FNiagaraSystemViewModel::IsolateEmitters, SelectedEmitterHandleIds)));
 		}
 		MenuBuilder.EndSection();
 	}
diff --git a/Engine/Plugins/FX/Niagara/Source/NiagaraEditor/Private/Tests/NiagaraEditorTestUtilities.cpp b/Engine/Plugins/FX/Niagara/Source/NiagaraEditor/Private/Tests/NiagaraEditorTestUtilities.cpp
--- a/Engine/Plugins/FX/Niagara/Source/NiagaraEditor/Private/Tests/NiagaraEditorTestUtilities.cpp
+++ b/Engine/Plugins/FX/Niagara/Source/NiagaraEditor/Private/Tests/NiagaraEditorTestUtilities.cpp
@@ -48,7 +48,8 @@ TSharedRef<FNiagaraSystemViewModel> FNiagaraEditorTestUtilities::CreateTestSyste
 	SystemViewModelOptions.bCanAutoCompile = false;
 	SystemViewModelOptions.bCanSimulate = false;
 
-	TSharedRef<FNiagaraSystemViewModel> SystemViewModel = MakeShared<FNiagaraSystemViewModel>(*System, SystemViewModelOptions);
+	TSharedRef<FNiagaraSystemViewModel> SystemViewModel = MakeShared<FNiagaraSystemViewModel>();
+	SystemViewModel->Initialize(*System, SystemViewModelOptions);
 
 	return SystemViewModel;
 }
diff --git a/Engine/Plugins/FX/Niagara/Source/NiagaraEditor/Private/Tests/NiagaraStackViewModelTests.cpp b/Engine/Plugins/FX/Niagara/Source/NiagaraEditor/Private/Tests/NiagaraStackViewModelTests.cpp
--- a/Engine/Plugins/FX/Niagara/Source/NiagaraEditor/Private/Tests/NiagaraStackViewModelTests.cpp
+++ b/Engine/Plugins/FX/Niagara/Source/NiagaraEditor/Private/Tests/NiagaraStackViewModelTests.cpp
@@ -17,6 +17,7 @@
 #include "Misc/AutomationTest.h"
 
 
+/*
 UNiagaraStackViewModel* CreateStackViewModelForEmitter(UNiagaraEmitter* Emitter)
 {
 	TSharedRef<FNiagaraSystemViewModel> SystemViewModel = FNiagaraEditorTestUtilities::CreateTestSystemViewModelForEmitter(Emitter);
@@ -25,6 +26,7 @@ UNiagaraStackViewModel* CreateStackViewModelForEmitter(UNiagaraEmitter* Emitter)
 	StackViewModel->Initialize(SystemViewModel, EmitterHandleViewModel);
 	return StackViewModel;
 }
+*/
 
 /*
 IMPLEMENT_SIMPLE_AUTOMATION_TEST(FNiagaraEditorStackViewModelAddParameterModulesTest, "System.Plugins.Niagara.Editor.StackViewModel.Add Parameter Modules", EAutomationTestFlags::EditorContext | EAutomationTestFlags::CriticalPriority | EAutomationTestFlags::EngineFilter)
diff --git a/Engine/Plugins/FX/Niagara/Source/NiagaraEditor/Private/Toolkits/NiagaraSystemToolkit.cpp b/Engine/Plugins/FX/Niagara/Source/NiagaraEditor/Private/Toolkits/NiagaraSystemToolkit.cpp
--- a/Engine/Plugins/FX/Niagara/Source/NiagaraEditor/Private/Toolkits/NiagaraSystemToolkit.cpp
+++ b/Engine/Plugins/FX/Niagara/Source/NiagaraEditor/Private/Toolkits/NiagaraSystemToolkit.cpp
@@ -8,13 +8,13 @@
 #include "ViewModels/NiagaraSystemViewModel.h"
 #include "ViewModels/NiagaraEmitterHandleViewModel.h"
 #include "ViewModels/NiagaraEmitterViewModel.h"
+#include "ViewModels/NiagaraSystemSelectionViewModel.h"
 #include "NiagaraSystemScriptViewModel.h"
 #include "Widgets/SNiagaraCurveEditor.h"
 #include "Widgets/SNiagaraSystemScript.h"
 #include "Widgets/SNiagaraSystemViewport.h"
 #include "Widgets/SNiagaraSelectedObjectsDetails.h"
 #include "Widgets/SNiagaraParameterMapView.h"
-#include "Widgets/SNiagaraSelectedEmitterHandles.h"
 #include "Widgets/SNiagaraSpreadsheetView.h"
 #include "Widgets/SNiagaraGeneratedCodeView.h"
 #include "Widgets/SNiagaraScriptGraph.h"
@@ -24,6 +24,7 @@
 #include "NiagaraSystemFactoryNew.h"
 #include "NiagaraEmitter.h"
 #include "NiagaraComponent.h"
+#include "NiagaraSystemEditorData.h"
 
 #include "IContentBrowserSingleton.h"
 #include "ContentBrowserModule.h"
@@ -187,9 +188,10 @@ void FNiagaraSystemToolkit::InitializeWithSystem(const EToolkitMode::Type Mode,
 	SystemOptions.bCanModifyEmittersFromTimeline = true;
 	SystemOptions.EditMode = ENiagaraSystemViewModelEditMode::SystemAsset;
 	SystemOptions.OnGetSequencerAddMenuContent.BindSP(this, &FNiagaraSystemToolkit::GetSequencerAddMenuContent);
-	const FGuid MessageLogGuid = FGuid::NewGuid();
+	SystemOptions.MessageLogGuid = FGuid::NewGuid();
 
-	SystemViewModel = MakeShareable(new FNiagaraSystemViewModel(*System, SystemOptions, MessageLogGuid));
+	SystemViewModel = MakeShared<FNiagaraSystemViewModel>();
+	SystemViewModel->Initialize(*System, SystemOptions);
 	SystemViewModel->SetToolkitCommands(GetToolkitCommands());
 	SystemToolkitMode = ESystemToolkitMode::System;
 
@@ -204,7 +206,7 @@ void FNiagaraSystemToolkit::InitializeWithSystem(const EToolkitMode::Type Mode,
 		FNiagaraEditorUtilities::WriteTextFileToDisk(FPaths::ProjectLogDir(), FilenamePart + TEXT(".onLoad.txt"), ExportText, true);
 	}
 
-	InitializeInternal(Mode, InitToolkitHost, MessageLogGuid);
+	InitializeInternal(Mode, InitToolkitHost, SystemOptions.MessageLogGuid.GetValue());
 }
 
 void FNiagaraSystemToolkit::InitializeWithEmitter(const EToolkitMode::Type Mode, const TSharedPtr< class IToolkitHost >& InitToolkitHost, UNiagaraEmitter& InEmitter)
@@ -255,11 +257,13 @@ void FNiagaraSystemToolkit::InitializeWithEmitter(const EToolkitMode::Type Mode,
 	FNiagaraSystemViewModelOptions SystemOptions;
 	SystemOptions.bCanModifyEmittersFromTimeline = false;
 	SystemOptions.EditMode = ENiagaraSystemViewModelEditMode::EmitterAsset;
-	const FGuid MessageLogGuid = FGuid::NewGuid();
+	SystemOptions.MessageLogGuid = FGuid::NewGuid();
 
-	SystemViewModel = MakeShareable(new FNiagaraSystemViewModel(*System, SystemOptions, MessageLogGuid));
+	SystemViewModel = MakeShared<FNiagaraSystemViewModel>();
+	SystemViewModel->Initialize(*System, SystemOptions);
 	SystemViewModel->SetToolkitCommands(GetToolkitCommands());
 	SystemViewModel->AddEmitter(*Emitter);
+	SystemViewModel->GetEditorData().SetOwningSystemIsPlaceholder(true, *System);
 
 	// Adding the emitter to the system has made a copy of it and we set this to the copy's change id here instead of the original emitter's change 
 	// id because the copy's change id may have been updated from the original as part of post load and we use this id to detect if the editable 
@@ -278,29 +282,22 @@ void FNiagaraSystemToolkit::InitializeWithEmitter(const EToolkitMode::Type Mode,
 		FNiagaraEditorUtilities::WriteTextFileToDisk(FPaths::ProjectLogDir(), FilenamePart + TEXT(".onLoad.txt"), ExportText, true);
 	}
 
-	InitializeInternal(Mode, InitToolkitHost, MessageLogGuid);
+	InitializeInternal(Mode, InitToolkitHost, SystemOptions.MessageLogGuid.GetValue());
 }
 
 void FNiagaraSystemToolkit::InitializeInternal(const EToolkitMode::Type Mode, const TSharedPtr<IToolkitHost>& InitToolkitHost, const FGuid& MessageLogGuid)
 {
-	SystemViewModel->SetOverviewGraphViewModel(MakeShareable(new FNiagaraOverviewGraphViewModel(SystemViewModel->AsShared()))); //@TODO System Overview: move constructing system overview viewmodel to SystemViewModel Init
-
-	if (SystemViewModel->GetEmitterHandleViewModels().Num() > 0)
-	{
-		SystemViewModel->SetSelectedEmitterHandleById(SystemViewModel->GetEmitterHandleViewModels()[0]->GetId());
-	}
-
 	NiagaraMessageLogViewModel = MakeShared<FNiagaraMessageLogViewModel>(GetNiagaraSystemMessageLogName(System), MessageLogGuid, NiagaraMessageLog);
 
-	SystemViewModel->OnEmitterHandleViewModelsChanged().AddSP(this, &FNiagaraSystemToolkit::OnRefresh);
-	SystemViewModel->OnSelectedEmitterHandlesChanged().AddSP(this, &FNiagaraSystemToolkit::OnRefresh);
-	SystemViewModel->GetOnPinnedEmittersChanged().AddSP(this, &FNiagaraSystemToolkit::OnRefresh);
+	SystemViewModel->OnEmitterHandleViewModelsChanged().AddSP(this, &FNiagaraSystemToolkit::RefreshParameters);
+	SystemViewModel->GetSelectionViewModel()->OnSelectionChanged().AddSP(this, &FNiagaraSystemToolkit::OnSystemSelectionChanged);
+	SystemViewModel->GetOnPinnedEmittersChanged().AddSP(this, &FNiagaraSystemToolkit::RefreshParameters);
 	SystemViewModel->GetOnPinnedCurvesChanged().AddSP(this, &FNiagaraSystemToolkit::OnPinnedCurvesChanged);
 
 	const float InTime = -0.02f;
 	const float OutTime = 3.2f;
 
-	TSharedRef<FTabManager::FLayout> StandaloneDefaultLayout = FTabManager::NewLayout("Standalone_Niagara_System_Layout_v19")
+	TSharedRef<FTabManager::FLayout> StandaloneDefaultLayout = FTabManager::NewLayout("Standalone_Niagara_System_Layout_v20")
 		->AddArea
 		(
 			FTabManager::NewPrimaryArea()->SetOrientation(Orient_Vertical)
@@ -316,28 +313,36 @@ void FNiagaraSystemToolkit::InitializeInternal(const EToolkitMode::Type Mode, co
 				FTabManager::NewSplitter()->SetOrientation(Orient_Horizontal)
 				->Split
 				(
+					// Top Level Left
 					FTabManager::NewSplitter()->SetOrientation(Orient_Vertical)
-					->SetSizeCoefficient(.60f)
+					->SetSizeCoefficient(.75f)
 					->Split
 					(
+						// Inner Left Top
 						FTabManager::NewSplitter()->SetOrientation(Orient_Horizontal)
 						->SetSizeCoefficient(0.75f)
 						->Split
 						(
 							FTabManager::NewStack()
-							->SetSizeCoefficient(.80f)
+							->SetSizeCoefficient(.25f)
 							->AddTab(ViewportTabID, ETabState::OpenedTab)
-							->AddTab(SystemOverviewTabID, ETabState::ClosedTab)
 						)
 						->Split
 						(
 							FTabManager::NewStack()
-							->SetSizeCoefficient(0.20f)
+							->SetSizeCoefficient(0.15f)
 							->AddTab(SystemParametersTabID, ETabState::OpenedTab)
 						)
+						->Split
+						(
+							FTabManager::NewStack()
+							->SetSizeCoefficient(0.6f)
+							->AddTab(SystemOverviewTabID, ETabState::OpenedTab)
+						)
 					)
 					->Split
 					(
+						// Inner Left Bottom
 						FTabManager::NewStack()
 						->SetSizeCoefficient(0.25f)
 						->AddTab(CurveEditorTabID, ETabState::OpenedTab)
@@ -347,8 +352,9 @@ void FNiagaraSystemToolkit::InitializeInternal(const EToolkitMode::Type Mode, co
 				)
 				->Split
 				(
+					// Top Level Right
 					FTabManager::NewStack()
-					->SetSizeCoefficient(0.40f)
+					->SetSizeCoefficient(0.25f)
 					->AddTab(SelectedEmitterStackTabID, ETabState::OpenedTab)
 					->AddTab(SelectedEmitterGraphTabID, ETabState::ClosedTab)
 					->AddTab(SystemScriptTabID, ETabState::ClosedTab)
@@ -512,10 +518,11 @@ TSharedRef<SDockTab> FNiagaraSystemToolkit::SpawnTab_SelectedEmitterStack(const
 {
 	check(Args.GetTabId().TabType == SelectedEmitterStackTabID);
 
-	TSharedRef<SDockTab> SpawnedTab =
-		SNew(SDockTab)
+	FNiagaraEditorModule& NiagaraEditorModule = FModuleManager::LoadModuleChecked<FNiagaraEditorModule>("NiagaraEditor");
+	TSharedRef<SDockTab> SpawnedTab = SNew(SDockTab)
+		.Label(LOCTEXT("SystemOverviewSelection", "Selection"))
 		[
-			SNew(SNiagaraSelectedEmitterHandles, SystemViewModel.ToSharedRef())
+			NiagaraEditorModule.GetWidgetProvider()->CreateStackView(*SystemViewModel->GetSelectionViewModel()->GetSelectionStackViewModel())
 		];
 
 	return SpawnedTab;
@@ -531,7 +538,7 @@ public:
 	void Construct(const FArguments& InArgs, TSharedRef<FNiagaraSystemViewModel> InSystemViewModel)
 	{
 		SystemViewModel = InSystemViewModel;
-		SystemViewModel->OnSelectedEmitterHandlesChanged().AddRaw(this, &SNiagaraSelectedEmitterGraph::SelectedEmitterHandlesChanged);
+		SystemViewModel->GetSelectionViewModel()->OnSelectionChanged().AddSP(this, &SNiagaraSelectedEmitterGraph::SystemSelectionChanged);
 		ChildSlot
 		[
 			SAssignNew(GraphWidgetContainer, SBox)
@@ -543,26 +550,23 @@ public:
 	{
 		if (SystemViewModel.IsValid())
 		{
-			SystemViewModel->OnEmitterHandleViewModelsChanged().RemoveAll(this);
-			SystemViewModel->OnSelectedEmitterHandlesChanged().RemoveAll(this);
-			SystemViewModel->GetOnPinnedEmittersChanged().RemoveAll(this);
-			SystemViewModel->OnSelectedEmitterHandlesChanged().RemoveAll(this);
+			SystemViewModel->GetSelectionViewModel()->OnSelectionChanged().RemoveAll(this);
 		}
 	}
 
 private:
-	void SelectedEmitterHandlesChanged()
+	void SystemSelectionChanged(UNiagaraSystemSelectionViewModel::ESelectionChangeSource SelectionChangeSource)
 	{
 		UpdateGraphWidget();
 	}
 
 	void UpdateGraphWidget()
 	{
-		TArray<TSharedRef<FNiagaraEmitterHandleViewModel>> SelectedEmitterHandles;
-		SystemViewModel->GetSelectedEmitterHandles(SelectedEmitterHandles);
-		if (SelectedEmitterHandles.Num() == 1)
+		TArray<FGuid> SelectedEmitterHandleIds = SystemViewModel->GetSelectionViewModel()->GetSelectedEmitterHandleIds();
+		if (SelectedEmitterHandleIds.Num() == 1)
 		{
-			GraphWidgetContainer->SetContent(SNew(SNiagaraScriptGraph, SelectedEmitterHandles[0]->GetEmitterViewModel()->GetSharedScriptViewModel()->GetGraphViewModel()));
+			TSharedPtr<FNiagaraEmitterHandleViewModel> SelectedEmitterHandle = SystemViewModel->GetEmitterHandleViewModelById(SelectedEmitterHandleIds[0]);
+			GraphWidgetContainer->SetContent(SNew(SNiagaraScriptGraph, SelectedEmitterHandle->GetEmitterViewModel()->GetSharedScriptViewModel()->GetGraphViewModel()));
 		}
 		else
 		{
@@ -1276,21 +1280,29 @@ void FNiagaraSystemToolkit::OnPinnedCurvesChanged()
 	TabManager->InvokeTab(CurveEditorTabID);
 }
 
-void FNiagaraSystemToolkit::OnRefresh()
+void FNiagaraSystemToolkit::RefreshParameters()
 {
 	if (ParameterMapView.IsValid())
 	{
 		TArray<TSharedPtr<FNiagaraEmitterHandleViewModel>> EmitterHandlesToDisplay;
 		EmitterHandlesToDisplay.Append(SystemViewModel->GetPinnedEmitterHandles());
-		TArray<TSharedRef<FNiagaraEmitterHandleViewModel>> SelectedEmitterHandles;
-		SystemViewModel->GetSelectedEmitterHandles(SelectedEmitterHandles);
-		for (auto Handle : SelectedEmitterHandles)
+		
+		TArray<FGuid> SelectedEmitterHandleIds = SystemViewModel->GetSelectionViewModel()->GetSelectedEmitterHandleIds();
+		for (TSharedRef<FNiagaraEmitterHandleViewModel> EmitterHandleViewModel : SystemViewModel->GetEmitterHandleViewModels())
 		{
-			EmitterHandlesToDisplay.AddUnique(Handle);
+			if (SelectedEmitterHandleIds.Contains(EmitterHandleViewModel->GetId()))
+			{
+				EmitterHandlesToDisplay.AddUnique(EmitterHandleViewModel);
+			}
 		}
 
 		ParameterMapView->RefreshEmitterHandles(EmitterHandlesToDisplay);
 	}
 }
 
+void FNiagaraSystemToolkit::OnSystemSelectionChanged(UNiagaraSystemSelectionViewModel::ESelectionChangeSource SelectionChangeSource)
+{
+	RefreshParameters();
+}
+
 #undef LOCTEXT_NAMESPACE
diff --git a/Engine/Plugins/FX/Niagara/Source/NiagaraEditor/Private/Toolkits/NiagaraSystemToolkit.h b/Engine/Plugins/FX/Niagara/Source/NiagaraEditor/Private/Toolkits/NiagaraSystemToolkit.h
--- a/Engine/Plugins/FX/Niagara/Source/NiagaraEditor/Private/Toolkits/NiagaraSystemToolkit.h
+++ b/Engine/Plugins/FX/Niagara/Source/NiagaraEditor/Private/Toolkits/NiagaraSystemToolkit.h
@@ -14,6 +14,7 @@
 #include "ISequencerTrackEditor.h"
 
 #include "NiagaraScript.h"
+#include "ViewModels/NiagaraSystemSelectionViewModel.h"
 
 class FNiagaraSystemInstance;
 class FNiagaraSystemViewModel;
@@ -122,9 +123,9 @@ private:
 	bool OnApplyEnabled() const;
 
 	void OnPinnedCurvesChanged();
-	void OnRefresh();
+	void RefreshParameters();
+	void OnSystemSelectionChanged(UNiagaraSystemSelectionViewModel::ESelectionChangeSource SelectionChangeSource);
 
-private:
 	TSharedRef<SWidget> GenerateBoundsMenuContent(TSharedRef<FUICommandList> InCommandList);
 	const FName GetNiagaraSystemMessageLogName(UNiagaraSystem* InSystem) const;
 	void OnSaveThumbnailImage();
diff --git a/Engine/Plugins/FX/Niagara/Source/NiagaraEditor/Private/ViewModels/NiagaraEmitterHandleViewModel.cpp b/Engine/Plugins/FX/Niagara/Source/NiagaraEditor/Private/ViewModels/NiagaraEmitterHandleViewModel.cpp
--- a/Engine/Plugins/FX/Niagara/Source/NiagaraEditor/Private/ViewModels/NiagaraEmitterHandleViewModel.cpp
+++ b/Engine/Plugins/FX/Niagara/Source/NiagaraEditor/Private/ViewModels/NiagaraEmitterHandleViewModel.cpp
@@ -3,7 +3,9 @@
 #include "ViewModels/NiagaraEmitterHandleViewModel.h"
 #include "NiagaraSystem.h"
 #include "NiagaraEmitterHandle.h"
+#include "ViewModels/NiagaraSystemViewModel.h"
 #include "ViewModels/NiagaraEmitterViewModel.h"
+#include "ViewModels/Stack/NiagaraStackViewModel.h"
 #include "NiagaraScriptViewModel.h"
 #include "NiagaraScriptGraphViewModel.h"
 #include "NiagaraObjectSelection.h"
@@ -20,56 +22,67 @@
 
 #define LOCTEXT_NAMESPACE "EmitterHandleViewModel"
 
-FNiagaraEmitterHandleViewModel::FNiagaraEmitterHandleViewModel(FNiagaraEmitterHandle* InEmitterHandle, TWeakPtr<FNiagaraEmitterInstance, ESPMode::ThreadSafe> InSimulation, UNiagaraSystem& InOwningSystem)
-	: EmitterHandle(InEmitterHandle)
-	, OwningSystem(InOwningSystem)
-	, EmitterViewModel(MakeShareable(new FNiagaraEmitterViewModel((InEmitterHandle ? InEmitterHandle->GetInstance() : nullptr), InSimulation)))
+FNiagaraEmitterHandleViewModel::FNiagaraEmitterHandleViewModel()
+	: EmitterHandle(nullptr)
+	, EmitterViewModel(MakeShared<FNiagaraEmitterViewModel>())
+	, EmitterStackViewModel(nullptr)
+	, bIsRenamePending(false)
 {
 }
 
 void FNiagaraEmitterHandleViewModel::Cleanup()
 {
-	if (EmitterViewModel.IsValid())
+	EmitterViewModel->Cleanup();
+	if (EmitterStackViewModel != nullptr)
 	{
-		EmitterViewModel->Cleanup();
+		EmitterStackViewModel->Finalize();
+		EmitterStackViewModel = nullptr;
 	}
-	EmitterViewModel.Reset();
 }
 
+TSharedRef<FNiagaraSystemViewModel> FNiagaraEmitterHandleViewModel::GetOwningSystemViewModel() const
+{
+	TSharedPtr<FNiagaraSystemViewModel> OwningSystemViewModelPinned = OwningSystemViewModelWeak.Pin();
+	checkf(OwningSystemViewModelPinned.IsValid(), TEXT("Owning system view model was destroyed before child handle view model."));
+	return OwningSystemViewModelPinned.ToSharedRef();
+}
 
 FNiagaraEmitterHandleViewModel::~FNiagaraEmitterHandleViewModel()
 {
 	Cleanup();
 }
 
-bool FNiagaraEmitterHandleViewModel::Set(FNiagaraEmitterHandle* InEmitterHandle, TWeakPtr<FNiagaraEmitterInstance, ESPMode::ThreadSafe> InSimulation, UNiagaraSystem& InOwningSystem)
+
+void FNiagaraEmitterHandleViewModel::Initialize(TSharedRef<FNiagaraSystemViewModel> InOwningSystemViewModel, FNiagaraEmitterHandle* InEmitterHandle, TWeakPtr<FNiagaraEmitterInstance, ESPMode::ThreadSafe> InSimulation)
 {
-	if (&OwningSystem != &InOwningSystem)
-	{
-		return false;
-	}
+	OwningSystemViewModelWeak = InOwningSystemViewModel;
+	EmitterHandle = InEmitterHandle;
+	UNiagaraEmitter* Emitter = EmitterHandle != nullptr ? EmitterHandle->GetInstance() : nullptr;
+	EmitterViewModel->Initialize(Emitter, InSimulation);
+}
 
-	SetEmitterHandle(InEmitterHandle);
-	SetSimulation(InSimulation);
-	
-	UNiagaraEmitter* EmitterProperties = nullptr;
-	if (InEmitterHandle != nullptr)
+void FNiagaraEmitterHandleViewModel::Reset()
+{
+	OwningSystemViewModelWeak.Reset();
+	EmitterHandle = nullptr;
+	EmitterViewModel->Reset();
+	if (EmitterStackViewModel != nullptr)
 	{
-		EmitterProperties = InEmitterHandle->GetInstance();
+		EmitterStackViewModel->Finalize();
+		EmitterStackViewModel = nullptr;
 	}
-	check(EmitterViewModel.IsValid());
-	return EmitterViewModel->Set(EmitterProperties, InSimulation);
 }
 
-
-void FNiagaraEmitterHandleViewModel::SetEmitterHandle(FNiagaraEmitterHandle* InEmitterHandle)
+void FNiagaraEmitterHandleViewModel::AddReferencedObjects(FReferenceCollector& Collector)
 {
-	EmitterHandle = InEmitterHandle;
+	if (EmitterStackViewModel != nullptr)
+	{
+		Collector.AddReferencedObject(EmitterStackViewModel);
+	}
 }
 
 void FNiagaraEmitterHandleViewModel::SetSimulation(TWeakPtr<FNiagaraEmitterInstance, ESPMode::ThreadSafe> InSimulation)
 {
-	check(EmitterViewModel.IsValid());
 	EmitterViewModel->SetSimulation(InSimulation);
 }
 
@@ -90,8 +103,6 @@ FText FNiagaraEmitterHandleViewModel::GetIdText() const
 
 FText FNiagaraEmitterHandleViewModel::GetErrorText() const
 {
-	check(EmitterViewModel.IsValid());
-
 	switch (EmitterViewModel->GetLatestCompileStatus())
 	{
 	case ENiagaraScriptCompileStatus::NCS_Unknown:
@@ -106,7 +117,6 @@ FText FNiagaraEmitterHandleViewModel::GetErrorText() const
 
 FSlateColor FNiagaraEmitterHandleViewModel::GetErrorTextColor() const
 {
-	check(EmitterViewModel.IsValid());
 	switch (EmitterViewModel->GetLatestCompileStatus())
 	{
 	case ENiagaraScriptCompileStatus::NCS_Unknown:
@@ -121,7 +131,6 @@ FSlateColor FNiagaraEmitterHandleViewModel::GetErrorTextColor() const
 
 EVisibility FNiagaraEmitterHandleViewModel::GetErrorTextVisibility() const
 {
-	check(EmitterViewModel.IsValid());
 	return EmitterViewModel->GetLatestCompileStatus() != ENiagaraScriptCompileStatus::NCS_UpToDate ? EVisibility::Visible : EVisibility::Collapsed;
 }
 
@@ -144,10 +153,11 @@ void FNiagaraEmitterHandleViewModel::SetName(FName InName)
 	if (EmitterHandle)
 	{
 		FScopedTransaction ScopedTransaction(NSLOCTEXT("NiagaraEmitterEditor", "EditEmitterNameTransaction", "Edit emitter name"));
-		OwningSystem.Modify();
-		OwningSystem.RemoveSystemParametersForEmitter(*EmitterHandle);
-		EmitterHandle->SetName(InName, OwningSystem);
-		OwningSystem.RefreshSystemParametersFromEmitter(*EmitterHandle);
+		UNiagaraSystem& System = GetOwningSystemViewModel()->GetSystem();
+		System.Modify();
+		System.RemoveSystemParametersForEmitter(*EmitterHandle);
+		EmitterHandle->SetName(InName, System);
+		System.RefreshSystemParametersFromEmitter(*EmitterHandle);
 		OnPropertyChangedDelegate.Broadcast();
 		OnNameChangedDelegate.Broadcast();
 	}
@@ -162,6 +172,11 @@ FText FNiagaraEmitterHandleViewModel::GetNameText() const
 	return FText();
 }
 
+bool FNiagaraEmitterHandleViewModel::CanRenameEmitter() const
+{
+	return GetOwningSystemViewModel()->GetEditMode() == ENiagaraSystemViewModelEditMode::SystemAsset;
+}
+
 void FNiagaraEmitterHandleViewModel::OnNameTextComitted(const FText& InText, ETextCommit::Type CommitInfo)
 {
 	SetName(*InText.ToString());
@@ -192,12 +207,17 @@ void FNiagaraEmitterHandleViewModel::SetIsEnabled(bool bInIsEnabled)
 	if (EmitterHandle && EmitterHandle->GetIsEnabled() != bInIsEnabled)
 	{
 		FScopedTransaction ScopedTransaction(NSLOCTEXT("NiagaraEmitterEditor", "EditEmitterEnabled", "Change emitter enabled state"));
-		OwningSystem.Modify();
+		GetOwningSystemViewModel()->GetSystem().Modify();
 		EmitterHandle->SetIsEnabled(bInIsEnabled);
 		OnPropertyChangedDelegate.Broadcast();
 	}
 }
 
+bool FNiagaraEmitterHandleViewModel::GetIsIsolated() const
+{
+	return EmitterHandle != nullptr && EmitterHandle->IsIsolated();
+}
+
 ECheckBoxState FNiagaraEmitterHandleViewModel::GetIsEnabledCheckState() const
 {
 	if (EmitterHandle)
@@ -217,11 +237,31 @@ FNiagaraEmitterHandle* FNiagaraEmitterHandleViewModel::GetEmitterHandle()
 	return EmitterHandle;
 }
 
-TSharedPtr<FNiagaraEmitterViewModel> FNiagaraEmitterHandleViewModel::GetEmitterViewModel()
+TSharedRef<FNiagaraEmitterViewModel> FNiagaraEmitterHandleViewModel::GetEmitterViewModel()
 {
 	return EmitterViewModel;
 }
 
+UNiagaraStackViewModel* FNiagaraEmitterHandleViewModel::GetEmitterStackViewModel()
+{
+	if (EmitterStackViewModel == nullptr)
+	{
+		EmitterStackViewModel = NewObject<UNiagaraStackViewModel>(GetTransientPackage());
+		EmitterStackViewModel->InitializeWithViewModels(GetOwningSystemViewModel(), this->AsShared(), FNiagaraStackViewModelOptions(false, true));
+	}
+	return EmitterStackViewModel;
+}
+
+bool FNiagaraEmitterHandleViewModel::GetIsRenamePending() const
+{
+	return bIsRenamePending;
+}
+
+void FNiagaraEmitterHandleViewModel::SetIsRenamePending(bool bInIsRenamePending)
+{
+	bIsRenamePending = bInIsRenamePending;
+}
+
 FNiagaraEmitterHandleViewModel::FOnPropertyChanged& FNiagaraEmitterHandleViewModel::OnPropertyChanged()
 {
 	return OnPropertyChangedDelegate;
diff --git a/Engine/Plugins/FX/Niagara/Source/NiagaraEditor/Private/ViewModels/NiagaraEmitterViewModel.cpp b/Engine/Plugins/FX/Niagara/Source/NiagaraEditor/Private/ViewModels/NiagaraEmitterViewModel.cpp
--- a/Engine/Plugins/FX/Niagara/Source/NiagaraEditor/Private/ViewModels/NiagaraEmitterViewModel.cpp
+++ b/Engine/Plugins/FX/Niagara/Source/NiagaraEditor/Private/ViewModels/NiagaraEmitterViewModel.cpp
@@ -4,6 +4,7 @@
 #include "NiagaraEmitter.h"
 #include "NiagaraEmitterInstance.h"
 #include "NiagaraScriptSourceBase.h"
+#include "ViewModels/NiagaraSystemViewModel.h"
 #include "NiagaraScriptViewModel.h"
 #include "NiagaraScriptGraphViewModel.h"
 #include "NiagaraObjectSelection.h"
@@ -31,14 +32,13 @@ namespace NiagaraCommands
 
 const float Megabyte = 1024.0f * 1024.0f;
 
-FNiagaraEmitterViewModel::FNiagaraEmitterViewModel(UNiagaraEmitter* InEmitter, TWeakPtr<FNiagaraEmitterInstance, ESPMode::ThreadSafe> InSimulation)
-	: Emitter(InEmitter)
-	, Simulation(InSimulation)
-	, SharedScriptViewModel(MakeShareable(new FNiagaraScriptViewModel(InEmitter, LOCTEXT("SharedDisplayName", "Graph"), ENiagaraParameterEditMode::EditAll)))
+FNiagaraEmitterViewModel::FNiagaraEmitterViewModel()
+	: Emitter(nullptr)
+	, Simulation(nullptr)
+	, SharedScriptViewModel(MakeShared<FNiagaraScriptViewModel>(LOCTEXT("SharedDisplayName", "Graph"), ENiagaraParameterEditMode::EditAll))
 	, bUpdatingSelectionInternally(false)
 	, ExecutionStateEnum(StaticEnum<ENiagaraExecutionState>())
 {	
-	SetEmitter(InEmitter);
 }
 
 void FNiagaraEmitterViewModel::Cleanup()
@@ -58,13 +58,19 @@ void FNiagaraEmitterViewModel::Cleanup()
 	RemoveScriptEventHandlers();
 }
 
-bool FNiagaraEmitterViewModel::Set(UNiagaraEmitter* InEmitter, TWeakPtr<FNiagaraEmitterInstance, ESPMode::ThreadSafe> InSimulation)
+bool FNiagaraEmitterViewModel::Initialize(UNiagaraEmitter* InEmitter, TWeakPtr<FNiagaraEmitterInstance, ESPMode::ThreadSafe> InSimulation)
 {
 	SetEmitter(InEmitter);
 	SetSimulation(InSimulation);
 	return true;
 }
 
+void FNiagaraEmitterViewModel::Reset()
+{
+	SetEmitter(nullptr);
+	SetSimulation(nullptr);
+}
+
 FNiagaraEmitterViewModel::~FNiagaraEmitterViewModel()
 {
 	Cleanup();
@@ -80,6 +86,7 @@ void FNiagaraEmitterViewModel::SetEmitter(UNiagaraEmitter* InEmitter)
 	{
 		Emitter->OnEmitterVMCompiled().RemoveAll(this);
 		Emitter->OnPropertiesChanged().RemoveAll(this);
+		UnregisterViewModelWithMap(RegisteredHandle);
 	}
 
 	UnregisterViewModelWithMap(RegisteredHandle);
@@ -104,18 +111,26 @@ void FNiagaraEmitterViewModel::SetEmitter(UNiagaraEmitter* InEmitter)
 	OnEmitterChanged().Broadcast();
 }
 
+TWeakPtr<FNiagaraEmitterInstance, ESPMode::ThreadSafe> FNiagaraEmitterViewModel::GetSimulation() const
+{
+	return Simulation;
+}
 
 void FNiagaraEmitterViewModel::SetSimulation(TWeakPtr<FNiagaraEmitterInstance, ESPMode::ThreadSafe> InSimulation)
 {
 	Simulation = InSimulation;
 }
 
-
 UNiagaraEmitter* FNiagaraEmitterViewModel::GetEmitter()
 {
 	return Emitter.Get();
 }
 
+bool FNiagaraEmitterViewModel::HasParentEmitter() const
+{
+	return Emitter.IsValid() && Emitter->GetParent() != nullptr;
+}
+
 const UNiagaraEmitter* FNiagaraEmitterViewModel::GetParentEmitter() const
 {
 	return Emitter.IsValid() ? Emitter->GetParent() : nullptr;
@@ -139,6 +154,14 @@ FText FNiagaraEmitterViewModel::GetParentPathNameText() const
 	return FText();
 }
 
+void FNiagaraEmitterViewModel::RemoveParentEmitter()
+{
+	FScopedTransaction ScopedTransaction(LOCTEXT("RemoveParentEmitterTransaction", "Remove Parent Emitter"));
+	Emitter->Modify();
+	Emitter->RemoveParent();
+	OnParentRemovedDelegate.Broadcast();
+}
+
 FText FNiagaraEmitterViewModel::GetStatsText() const
 {
 	if (Simulation.IsValid())
@@ -323,6 +346,11 @@ FNiagaraEmitterViewModel::FOnScriptCompiled& FNiagaraEmitterViewModel::OnScriptC
 	return OnScriptCompiledDelegate;
 }
 
+FNiagaraEmitterViewModel::FOnParentRemoved& FNiagaraEmitterViewModel::OnParentRemoved()
+{
+	return OnParentRemovedDelegate;
+}
+
 FNiagaraEmitterViewModel::FOnScriptGraphChanged& FNiagaraEmitterViewModel::OnScriptGraphChanged()
 {
 	return OnScriptGraphChangedDelegate;
diff --git a/Engine/Plugins/FX/Niagara/Source/NiagaraEditor/Private/ViewModels/NiagaraOverviewGraphViewModel.cpp b/Engine/Plugins/FX/Niagara/Source/NiagaraEditor/Private/ViewModels/NiagaraOverviewGraphViewModel.cpp
--- a/Engine/Plugins/FX/Niagara/Source/NiagaraEditor/Private/ViewModels/NiagaraOverviewGraphViewModel.cpp
+++ b/Engine/Plugins/FX/Niagara/Source/NiagaraEditor/Private/ViewModels/NiagaraOverviewGraphViewModel.cpp
@@ -6,33 +6,46 @@
 #include "EdGraphUtilities.h"
 #include "HAL/PlatformApplicationMisc.h"
 #include "NiagaraSystemEditorData.h"
-#include "NiagaraOverviewNodeStackItem.h"
+#include "NiagaraOverviewNode.h"
 #include "NiagaraObjectSelection.h"
 #include "ViewModels/NiagaraSystemViewModel.h"
 #include "ViewModels/NiagaraEmitterHandleViewModel.h"
+#include "ViewModels/Stack/NiagaraStackViewModel.h"
+#include "ViewModels/Stack/NiagaraStackEntry.h"
 #include "NiagaraSystem.h"
 #include "Editor.h"
 
 #define LOCTEXT_NAMESPACE "NiagaraOverviewGraphViewModel"
 
-FNiagaraOverviewGraphViewModel::FNiagaraOverviewGraphViewModel(TSharedRef<FNiagaraSystemViewModel> InSystemViewModel)
-	: SystemViewModel(InSystemViewModel)
-	, Commands(MakeShareable(new FUICommandList()))
+FNiagaraOverviewGraphViewModel::FNiagaraOverviewGraphViewModel()
+	: Commands(MakeShareable(new FUICommandList()))
 	, NodeSelection(MakeShareable(new FNiagaraObjectSelection()))
+	, bUpdatingSystemSelectionFromGraph(false)
+	, bUpdatingGraphSelectionFromSystem(false)
 {
-	SetupCommands();
-	InitDisplayName();
-	GEditor->RegisterForUndo(this);
 }
 
 FNiagaraOverviewGraphViewModel::~FNiagaraOverviewGraphViewModel()
 {
-	GEditor->UnregisterForUndo(this);
+}
+
+void FNiagaraOverviewGraphViewModel::Initialize(TSharedRef<FNiagaraSystemViewModel> InSystemViewModel)
+{
+	SystemViewModel = InSystemViewModel;
+	OverviewGraph = InSystemViewModel->GetEditorData().GetSystemOverviewGraph();
+
+	SetupCommands();
+	NodeSelection->OnSelectedObjectsChanged().AddSP(this, &FNiagaraOverviewGraphViewModel::GraphSelectionChanged);
+	InSystemViewModel->GetSelectionViewModel()->OnSelectionChanged().AddSP(this, &FNiagaraOverviewGraphViewModel::SystemSelectionChanged);
 }
 
 FText FNiagaraOverviewGraphViewModel::GetDisplayName() const
 {
-	return DisplayName;
+	if (DisplayNameCache.IsSet() == false)
+	{
+		DisplayNameCache = GetDisplayNameInternal();
+	}
+	return DisplayNameCache.GetValue();
 }
 
 UEdGraph* FNiagaraOverviewGraphViewModel::GetGraph() const
@@ -54,20 +67,33 @@ TSharedRef<FNiagaraObjectSelection> FNiagaraOverviewGraphViewModel::GetNodeSelec
 	return NodeSelection;
 }
 
+const FNiagaraGraphViewSettings& FNiagaraOverviewGraphViewModel::GetViewSettings() const
+{
+	return GetSystemViewModel()->GetEditorData().GetSystemOverviewGraphViewSettings();
+}
+
+void FNiagaraOverviewGraphViewModel::SetViewSettings(const FNiagaraGraphViewSettings& InOverviewGraphViewSettings)
+{
+	GetSystemViewModel()->GetEditorData().SetSystemOverviewGraphViewSettings(InOverviewGraphViewSettings);
+}
+
 FNiagaraOverviewGraphViewModel::FOnNodesPasted& FNiagaraOverviewGraphViewModel::OnNodesPasted()
 {
 	return OnNodesPastedDelegate;
 }
 
-void FNiagaraOverviewGraphViewModel::PostUndo(bool bSuccess)
+TSharedRef<FNiagaraSystemViewModel> FNiagaraOverviewGraphViewModel::GetSystemViewModel()
 {
-	// The graph may have been deleted as a result of the undo operation so make sure it's valid
-	// before using it.
-	UEdGraph* Graph = GetGraph();
-	if (Graph != nullptr)
-	{
-		Graph->NotifyGraphChanged();
-	}
+	TSharedPtr<FNiagaraSystemViewModel> SystemViewModelPinned = SystemViewModel.Pin();
+	checkf(SystemViewModelPinned.IsValid(), TEXT("System view model destroyed before overview graph view model."));
+	return SystemViewModelPinned.ToSharedRef();
+}
+
+const TSharedRef<FNiagaraSystemViewModel> FNiagaraOverviewGraphViewModel::GetSystemViewModel() const
+{
+	TSharedPtr<FNiagaraSystemViewModel> SystemViewModelPinned = SystemViewModel.Pin();
+	checkf(SystemViewModelPinned.IsValid(), TEXT("System view model destroyed before overview graph view model."));
+	return SystemViewModelPinned.ToSharedRef();
 }
 
 void FNiagaraOverviewGraphViewModel::SetupCommands()
@@ -129,7 +155,7 @@ void FNiagaraOverviewGraphViewModel::DeleteSelectedNodes()
 		TSet<FGuid> EmitterGuidsToDelete;
 		for (UObject* NodeToDelete : NodesToDelete)
 		{
-			UNiagaraOverviewNodeStackItem* GraphNodeToDelete = Cast<UNiagaraOverviewNodeStackItem>(NodeToDelete);
+			UNiagaraOverviewNode* GraphNodeToDelete = Cast<UNiagaraOverviewNode>(NodeToDelete);
 			if (GraphNodeToDelete != nullptr && GraphNodeToDelete->CanUserDeleteNode())
 			{
 				EmitterGuidsToDelete.Add(GraphNodeToDelete->GetEmitterHandleGuid());
@@ -293,7 +319,7 @@ bool FNiagaraOverviewGraphViewModel::CanDuplicateNodes() const
 	return CanCopyNodes();
 }
 
-void FNiagaraOverviewGraphViewModel::InitDisplayName()
+FText FNiagaraOverviewGraphViewModel::GetDisplayNameInternal() const
 {
 	ensureMsgf(SystemViewModel.IsValid(), TEXT("SystemViewModel was not initialized before InitDisplayName!"));
 	TSharedPtr<FNiagaraSystemViewModel> WeakSystemViewModel = SystemViewModel.Pin();
@@ -303,20 +329,79 @@ void FNiagaraOverviewGraphViewModel::InitDisplayName()
 		TArray<TSharedRef<FNiagaraEmitterHandleViewModel>> EditedEmitterHandleViewModels = WeakSystemViewModel->GetEmitterHandleViewModels();
 		if (ensureMsgf(EditedEmitterHandleViewModels.Num() > 0, TEXT("SystemViewModel did not have any EmitterHandleViewModels! Cannot get currently edited Emitter's Name.")))
 		{
-			DisplayName = FText::FromName(WeakSystemViewModel->GetEmitterHandleViewModels()[0]->GetEmitterHandle()->GetName());
-			return;
+			return FText::FromName(WeakSystemViewModel->GetEmitterHandleViewModels()[0]->GetEmitterHandle()->GetName());
 		}
 	}
 	else if (SystemEditMode == ENiagaraSystemViewModelEditMode::SystemAsset)
 	{
-		DisplayName = FText::FromString(WeakSystemViewModel->GetSystem().GetName());
-		return;
+		return FText::FromString(WeakSystemViewModel->GetSystem().GetName());
 	}
 	else
 	{
 		ensureMsgf(false, TEXT("Encountered unexpected SystemViewModel edit mode!"));
 	}
-	DisplayName = LOCTEXT("NiagaraOverview_FallbackTitlebar", "Niagara");
+	return LOCTEXT("NiagaraOverview_FallbackTitlebar", "Niagara");
+}
+
+void FNiagaraOverviewGraphViewModel::GraphSelectionChanged()
+{
+	if (bUpdatingGraphSelectionFromSystem == false)
+	{
+		TGuardValue<bool> UpdateGuard(bUpdatingSystemSelectionFromGraph, true);
+
+		bool bSystemIsSelected = false;
+		TArray<FGuid> SelectedEmitterHandleGuids;
+		for (UObject* SelectedNode : NodeSelection->GetSelectedObjects())
+		{
+			UNiagaraOverviewNode* OverviewNode = Cast<UNiagaraOverviewNode>(SelectedNode);
+			if (OverviewNode != nullptr)
+			{
+				if (OverviewNode->GetEmitterHandleGuid().IsValid())
+				{
+					SelectedEmitterHandleGuids.Add(OverviewNode->GetEmitterHandleGuid());
+				}
+				else
+				{
+					bSystemIsSelected = true;
+				}
+			}
+		}
+
+		bool bClearCurrentSelection = FSlateApplication::Get().GetModifierKeys().IsControlDown() == false;
+		GetSystemViewModel()->GetSelectionViewModel()->UpdateSelectionFromTopLevelObjects(bSystemIsSelected, SelectedEmitterHandleGuids, bClearCurrentSelection);
+	}
+}
+
+void FNiagaraOverviewGraphViewModel::SystemSelectionChanged(UNiagaraSystemSelectionViewModel::ESelectionChangeSource SelectionChangeSource)
+{
+	if (bUpdatingSystemSelectionFromGraph == false && 
+		SelectionChangeSource != UNiagaraSystemSelectionViewModel::ESelectionChangeSource::EntrySelection)
+	{
+		TGuardValue<bool> UpdateGuard(bUpdatingGraphSelectionFromSystem, true);
+
+		TArray<UObject*> SelectedNodes;
+		TArray<UNiagaraOverviewNode*> OverviewNodes;
+		OverviewGraph->GetNodesOfClass<UNiagaraOverviewNode>(OverviewNodes);
+		for (UNiagaraOverviewNode* OverviewNode : OverviewNodes)
+		{
+			if (OverviewNode->GetEmitterHandleGuid().IsValid())
+			{
+				if (GetSystemViewModel()->GetSelectionViewModel()->GetSelectedEmitterHandleIds().Contains(OverviewNode->GetEmitterHandleGuid()))
+				{
+					SelectedNodes.Add(OverviewNode);
+				}
+			}
+			else
+			{
+				if (GetSystemViewModel()->GetSelectionViewModel()->GetSystemIsSelected())
+				{
+					SelectedNodes.Add(OverviewNode);
+				}
+			}
+		}
+
+		NodeSelection->SetSelectedObjects(SelectedNodes);
+	}
 }
 
 #undef LOCTEXT_NAMESPACE // NiagaraScriptGraphViewModel
diff --git a/Engine/Plugins/FX/Niagara/Source/NiagaraEditor/Private/ViewModels/NiagaraScriptGraphViewModel.cpp b/Engine/Plugins/FX/Niagara/Source/NiagaraEditor/Private/ViewModels/NiagaraScriptGraphViewModel.cpp
--- a/Engine/Plugins/FX/Niagara/Source/NiagaraEditor/Private/ViewModels/NiagaraScriptGraphViewModel.cpp
+++ b/Engine/Plugins/FX/Niagara/Source/NiagaraEditor/Private/ViewModels/NiagaraScriptGraphViewModel.cpp
@@ -28,6 +28,16 @@ FNiagaraScriptGraphViewModel::FNiagaraScriptGraphViewModel(UNiagaraScriptSource*
 	ErrorColor = FEditorStyle::GetColor("ErrorReporting.BackgroundColor");
 }
 
+FNiagaraScriptGraphViewModel::FNiagaraScriptGraphViewModel(FText InDisplayName)
+	: DisplayName(InDisplayName)
+	, Commands(MakeShareable(new FUICommandList()))
+	, NodeSelection(MakeShareable(new FNiagaraObjectSelection()))
+{
+	SetupCommands();
+	GEditor->RegisterForUndo(this);
+	ErrorColor = FEditorStyle::GetColor("ErrorReporting.BackgroundColor");
+}
+
 FNiagaraScriptGraphViewModel::~FNiagaraScriptGraphViewModel()
 {
 	GEditor->UnregisterForUndo(this);
diff --git a/Engine/Plugins/FX/Niagara/Source/NiagaraEditor/Private/ViewModels/NiagaraScriptGraphViewModel.h b/Engine/Plugins/FX/Niagara/Source/NiagaraEditor/Private/ViewModels/NiagaraScriptGraphViewModel.h
--- a/Engine/Plugins/FX/Niagara/Source/NiagaraEditor/Private/ViewModels/NiagaraScriptGraphViewModel.h
+++ b/Engine/Plugins/FX/Niagara/Source/NiagaraEditor/Private/ViewModels/NiagaraScriptGraphViewModel.h
@@ -25,6 +25,8 @@ public:
 	/** Create a new view model with the supplied script. */
 	FNiagaraScriptGraphViewModel(UNiagaraScriptSource* InScriptSource, FText InDisplayName);
 
+	FNiagaraScriptGraphViewModel(FText InDisplayName);
+
 	~FNiagaraScriptGraphViewModel();
 
 	/** Sets this view model to a new script. */
diff --git a/Engine/Plugins/FX/Niagara/Source/NiagaraEditor/Private/ViewModels/NiagaraScriptInputCollectionViewModel.cpp b/Engine/Plugins/FX/Niagara/Source/NiagaraEditor/Private/ViewModels/NiagaraScriptInputCollectionViewModel.cpp
--- a/Engine/Plugins/FX/Niagara/Source/NiagaraEditor/Private/ViewModels/NiagaraScriptInputCollectionViewModel.cpp
+++ b/Engine/Plugins/FX/Niagara/Source/NiagaraEditor/Private/ViewModels/NiagaraScriptInputCollectionViewModel.cpp
@@ -50,41 +50,10 @@ FNiagaraScriptInputCollectionViewModel::FNiagaraScriptInputCollectionViewModel(U
 }
 
 
-FNiagaraScriptInputCollectionViewModel::FNiagaraScriptInputCollectionViewModel(UNiagaraEmitter* InEmitter, FText InDisplayName, ENiagaraParameterEditMode InParameterEditMode)
+FNiagaraScriptInputCollectionViewModel::FNiagaraScriptInputCollectionViewModel(FText InDisplayName, ENiagaraParameterEditMode InParameterEditMode)
 	: FNiagaraParameterCollectionViewModel(InParameterEditMode)
 	, DisplayName(FText::Format(DisplayNameFormat, InDisplayName))
 {
-	TArray<UNiagaraScript*> InScripts;
-	InEmitter->GetScripts(InScripts);
-	// Because of weak pointers, we need to copy ourselves..
-	for (UNiagaraScript* Script : InScripts)
-	{
-		Scripts.Add(Script);
-		ensure(Script->GetSource() == InEmitter->GraphSource);
-	}
-	UNiagaraScriptSource* Source = Cast<UNiagaraScriptSource>(InEmitter->GraphSource);
-
-	if (Source != nullptr)
-	{
-		Graph = Source->NodeGraph;
-		bCanHaveNumericParameters = false;
-	}
-	else
-	{
-		Graph = nullptr;
-		bCanHaveNumericParameters = false;
-	}
-
-	RefreshParameterViewModels();
-
-	if (Graph.IsValid())
-	{
-		OnGraphChangedHandle = Graph->AddOnGraphChangedHandler(
-			FOnGraphChanged::FDelegate::CreateRaw(this, &FNiagaraScriptInputCollectionViewModel::OnGraphChanged));
-		OnRecompileHandle = Graph->AddOnGraphNeedsRecompileHandler(
-			FOnGraphChanged::FDelegate::CreateRaw(this, &FNiagaraScriptInputCollectionViewModel::OnGraphChanged));
-
-	}
 }
 
 FNiagaraScriptInputCollectionViewModel::~FNiagaraScriptInputCollectionViewModel()
diff --git a/Engine/Plugins/FX/Niagara/Source/NiagaraEditor/Private/ViewModels/NiagaraScriptInputCollectionViewModel.h b/Engine/Plugins/FX/Niagara/Source/NiagaraEditor/Private/ViewModels/NiagaraScriptInputCollectionViewModel.h
--- a/Engine/Plugins/FX/Niagara/Source/NiagaraEditor/Private/ViewModels/NiagaraScriptInputCollectionViewModel.h
+++ b/Engine/Plugins/FX/Niagara/Source/NiagaraEditor/Private/ViewModels/NiagaraScriptInputCollectionViewModel.h
@@ -17,7 +17,7 @@ class FNiagaraScriptInputCollectionViewModel : public FNiagaraParameterCollectio
 {
 public:
 	FNiagaraScriptInputCollectionViewModel(UNiagaraScript* InScript, FText InDisplayName, ENiagaraParameterEditMode InParameterEditMode);
-	FNiagaraScriptInputCollectionViewModel(UNiagaraEmitter* InEmitter, FText InDisplayName, ENiagaraParameterEditMode InParameterEditMode);
+	FNiagaraScriptInputCollectionViewModel(FText InDisplayName, ENiagaraParameterEditMode InParameterEditMode);
 
 	~FNiagaraScriptInputCollectionViewModel();
 
diff --git a/Engine/Plugins/FX/Niagara/Source/NiagaraEditor/Private/ViewModels/NiagaraScriptOutputCollectionViewModel.cpp b/Engine/Plugins/FX/Niagara/Source/NiagaraEditor/Private/ViewModels/NiagaraScriptOutputCollectionViewModel.cpp
--- a/Engine/Plugins/FX/Niagara/Source/NiagaraEditor/Private/ViewModels/NiagaraScriptOutputCollectionViewModel.cpp
+++ b/Engine/Plugins/FX/Niagara/Source/NiagaraEditor/Private/ViewModels/NiagaraScriptOutputCollectionViewModel.cpp
@@ -53,42 +53,13 @@ FNiagaraScriptOutputCollectionViewModel::FNiagaraScriptOutputCollectionViewModel
 	}
 }
 
-FNiagaraScriptOutputCollectionViewModel::FNiagaraScriptOutputCollectionViewModel(UNiagaraEmitter* InEmitter, ENiagaraParameterEditMode InParameterEditMode)
+FNiagaraScriptOutputCollectionViewModel::FNiagaraScriptOutputCollectionViewModel(ENiagaraParameterEditMode InParameterEditMode)
 	: FNiagaraParameterCollectionViewModel(InParameterEditMode)
 	, DisplayName(LOCTEXT("DisplayName", "Outputs"))
 {
-	TArray<UNiagaraScript*> InScripts;
-	InEmitter->GetScripts(InScripts);
-	// Because of weak pointers, we need to copy ourselves..
-	for (UNiagaraScript* Script : InScripts)
-	{
-		Scripts.Add(Script);
-		ensure(Script->GetSource() == InEmitter->GraphSource);
-	}
-	UNiagaraScriptSource* Source = Cast<UNiagaraScriptSource>(InEmitter->GraphSource);
-
-	if (Source)
-	{
-		Graph = Source->NodeGraph;
-		bCanHaveNumericParameters = false;
-		
-		OutputNode = Graph->FindOutputNode(ENiagaraScriptUsage::ParticleSpawnScript);
-	}
-	else
-	{
-		Graph = nullptr;
-		OutputNode = nullptr;
-		bCanHaveNumericParameters = false;
-	}
-	RefreshParameterViewModels();
-
-	if (Graph.IsValid())
-	{
-		OnGraphChangedHandle = Graph->AddOnGraphChangedHandler(
-			FOnGraphChanged::FDelegate::CreateRaw(this, &FNiagaraScriptOutputCollectionViewModel::OnGraphChanged));
-		OnRecompileHandle = Graph->AddOnGraphNeedsRecompileHandler(
-			FOnGraphChanged::FDelegate::CreateRaw(this, &FNiagaraScriptOutputCollectionViewModel::OnGraphChanged));
-	}
+	Graph = nullptr;
+	OutputNode = nullptr;
+	bCanHaveNumericParameters = false;
 }
 
 FNiagaraScriptOutputCollectionViewModel::~FNiagaraScriptOutputCollectionViewModel()
diff --git a/Engine/Plugins/FX/Niagara/Source/NiagaraEditor/Private/ViewModels/NiagaraScriptOutputCollectionViewModel.h b/Engine/Plugins/FX/Niagara/Source/NiagaraEditor/Private/ViewModels/NiagaraScriptOutputCollectionViewModel.h
--- a/Engine/Plugins/FX/Niagara/Source/NiagaraEditor/Private/ViewModels/NiagaraScriptOutputCollectionViewModel.h
+++ b/Engine/Plugins/FX/Niagara/Source/NiagaraEditor/Private/ViewModels/NiagaraScriptOutputCollectionViewModel.h
@@ -18,7 +18,7 @@ public:
 
 public:
 	FNiagaraScriptOutputCollectionViewModel(UNiagaraScript* InScript, ENiagaraParameterEditMode InParameterEditMode);
-	FNiagaraScriptOutputCollectionViewModel(UNiagaraEmitter* InEmitter, ENiagaraParameterEditMode InParameterEditMode);
+	FNiagaraScriptOutputCollectionViewModel(ENiagaraParameterEditMode InParameterEditMode);
 
 	~FNiagaraScriptOutputCollectionViewModel();
 
diff --git a/Engine/Plugins/FX/Niagara/Source/NiagaraEditor/Private/ViewModels/NiagaraScriptViewModel.cpp b/Engine/Plugins/FX/Niagara/Source/NiagaraEditor/Private/ViewModels/NiagaraScriptViewModel.cpp
--- a/Engine/Plugins/FX/Niagara/Source/NiagaraEditor/Private/ViewModels/NiagaraScriptViewModel.cpp
+++ b/Engine/Plugins/FX/Niagara/Source/NiagaraEditor/Private/ViewModels/NiagaraScriptViewModel.cpp
@@ -64,25 +64,14 @@ FNiagaraScriptViewModel::FNiagaraScriptViewModel(UNiagaraScript* InScript, FText
 
 }
 
-FNiagaraScriptViewModel::FNiagaraScriptViewModel(UNiagaraEmitter* InEmitter, FText DisplayName, ENiagaraParameterEditMode InParameterEditMode)
-	: InputCollectionViewModel(MakeShareable(new FNiagaraScriptInputCollectionViewModel(InEmitter, DisplayName, InParameterEditMode)))
-	, OutputCollectionViewModel(MakeShareable(new FNiagaraScriptOutputCollectionViewModel(InEmitter, InParameterEditMode)))
-	, GraphViewModel(MakeShareable(new FNiagaraScriptGraphViewModel(Cast<UNiagaraScriptSource>(InEmitter->GraphSource), DisplayName)))
+FNiagaraScriptViewModel::FNiagaraScriptViewModel(FText DisplayName, ENiagaraParameterEditMode InParameterEditMode)
+	: InputCollectionViewModel(MakeShareable(new FNiagaraScriptInputCollectionViewModel(DisplayName, InParameterEditMode)))
+	, OutputCollectionViewModel(MakeShareable(new FNiagaraScriptOutputCollectionViewModel(InParameterEditMode)))
+	, GraphViewModel(MakeShareable(new FNiagaraScriptGraphViewModel(DisplayName)))
 	, VariableSelection(MakeShareable(new FNiagaraObjectSelection()))
 	, bUpdatingSelectionInternally(false)
 	, LastCompileStatus(ENiagaraScriptCompileStatus::NCS_Unknown)
 {
-	TArray<UNiagaraScript*> InScripts;
-	InEmitter->GetScripts(InScripts);
-	// Because of weak pointers, we need to copy ourselves..
-	for (UNiagaraScript* Script : InScripts)
-	{
-		int32 i = Scripts.Add(Script);
-		check(Script->GetSource() == InEmitter->GraphSource);
-		Scripts[i]->OnVMScriptCompiled().AddRaw(this, &FNiagaraScriptViewModel::OnVMScriptCompiled);
-	}
-	Source = Cast<UNiagaraScriptSource>(InEmitter->GraphSource);
- 
 	InputCollectionViewModel->GetSelection().OnSelectedObjectsChanged().AddRaw(this, &FNiagaraScriptViewModel::InputViewModelSelectionChanged);
 	InputCollectionViewModel->OnParameterValueChanged().AddRaw(this, &FNiagaraScriptViewModel::InputParameterValueChanged);
 	OutputCollectionViewModel->OnParameterValueChanged().AddRaw(this, &FNiagaraScriptViewModel::OutputParameterValueChanged);
@@ -95,50 +84,6 @@ FNiagaraScriptViewModel::FNiagaraScriptViewModel(UNiagaraEmitter* InEmitter, FTe
 	CompileStatuses.Empty();
 	CompileErrors.Empty();
 	CompilePaths.Empty();
-
-	for (int32 i = 0; i < Scripts.Num(); i++)
-	{
-		FString Message;
-		ENiagaraScriptCompileStatus ScriptStatus = Scripts[i]->GetLastCompileStatus();
-		if (Scripts[i].IsValid() && Scripts[i]->IsCompilable() && Scripts[i]->GetVMExecutableData().IsValid() && Scripts[i]->GetVMExecutableData().ByteCode.Num() == 0) // This is either a brand new script or failed in the past. Since we create a default working script, assume invalid.
-		{
-			LastCompileStatus = ScriptStatus;
-			Message = TEXT("Please recompile for full error stack.");
-			GraphViewModel->SetErrorTextToolTip(Message);
-		}
-		else // Possibly warnings previously, but still compiled. It *could* have been dirtied somehow, but we assume that it is up-to-date.
-		{
-			if (Scripts[i].IsValid() && Scripts[i]->IsCompilable()&& Scripts[i]->AreScriptAndSourceSynchronized())
-			{
-				LastCompileStatus = FNiagaraEditorUtilities::UnionCompileStatus(LastCompileStatus, Scripts[i]->GetLastCompileStatus());
-			}
-			else if (Scripts[i].IsValid() && !Scripts[i]->IsCompilable())
-			{
-				// Do nothing..
-				ScriptStatus = ENiagaraScriptCompileStatus::NCS_UpToDate;
-			}
-			else
-			{
-				LastCompileStatus = FNiagaraEditorUtilities::UnionCompileStatus(LastCompileStatus, ENiagaraScriptCompileStatus::NCS_Error);
-				ScriptStatus = ENiagaraScriptCompileStatus::NCS_Error;
-				Message = TEXT("Please recompile for full error stack.");
-				GraphViewModel->SetErrorTextToolTip(Message);
-			}
-		}
-
-		CompilePaths.Add(Scripts[i]->GetPathName());
-		CompileErrors.Add(Message);
-		CompileStatuses.Add(ScriptStatus);
-
-		RegisteredHandles.Add(RegisterViewModelWithMap(Scripts[i].Get(), this));
-	}
-
-	CompileTypes.SetNum(CompileStatuses.Num());
-	for (int32 i = 0; i < CompileStatuses.Num(); i++)
-	{
-		CompileTypes[i].Key = Scripts[i]->GetUsage();
-		CompileTypes[i].Value = Scripts[i]->GetUsageId();
-	}
 }
 
 void FNiagaraScriptViewModel::OnVMScriptCompiled(UNiagaraScript* InScript)
diff --git a/Engine/Plugins/FX/Niagara/Source/NiagaraEditor/Private/ViewModels/NiagaraScriptViewModel.h b/Engine/Plugins/FX/Niagara/Source/NiagaraEditor/Private/ViewModels/NiagaraScriptViewModel.h
--- a/Engine/Plugins/FX/Niagara/Source/NiagaraEditor/Private/ViewModels/NiagaraScriptViewModel.h
+++ b/Engine/Plugins/FX/Niagara/Source/NiagaraEditor/Private/ViewModels/NiagaraScriptViewModel.h
@@ -24,7 +24,7 @@ class FNiagaraScriptViewModel : public TSharedFromThis<FNiagaraScriptViewModel>,
 {
 public:
 	FNiagaraScriptViewModel(UNiagaraScript* InScript, FText DisplayName, ENiagaraParameterEditMode InParameterEditMode);
-	FNiagaraScriptViewModel(UNiagaraEmitter* InEmitter, FText DisplayName, ENiagaraParameterEditMode InParameterEditMode);
+	FNiagaraScriptViewModel(FText DisplayName, ENiagaraParameterEditMode InParameterEditMode);
 
 	~FNiagaraScriptViewModel();
 
diff --git a/Engine/Plugins/FX/Niagara/Source/NiagaraEditor/Private/ViewModels/NiagaraSystemSelectionViewModel.cpp b/Engine/Plugins/FX/Niagara/Source/NiagaraEditor/Private/ViewModels/NiagaraSystemSelectionViewModel.cpp
new file mode 100644
--- /dev/null
+++ b/Engine/Plugins/FX/Niagara/Source/NiagaraEditor/Private/ViewModels/NiagaraSystemSelectionViewModel.cpp
@@ -0,0 +1,272 @@
+// Copyright 1998-2019 Epic Games, Inc. All Rights Reserved.
+
+#include "ViewModels/NiagaraSystemSelectionViewModel.h"
+#include "ViewModels/Stack/NiagaraStackEntry.h"
+#include "ViewModels/Stack/NiagaraStackSelection.h"
+#include "ViewModels/Stack/NiagaraStackViewModel.h"
+#include "ViewModels/NiagaraSystemViewModel.h"
+#include "ViewModels/NiagaraEmitterHandleViewModel.h"
+#include "ViewModels/NiagaraEmitterViewModel.h"
+#include "NiagaraEditorUtilities.h"
+#include "NiagaraSystemEditorData.h"
+#include "NiagaraEmitterHandle.h"
+
+void UNiagaraSystemSelectionViewModel::Initialize(TSharedRef<FNiagaraSystemViewModel> InSystemViewModel)
+{
+	SystemViewModelWeak = InSystemViewModel;
+
+	StackSelection = NewObject<UNiagaraStackSelection>(this);
+	StackSelection->Initialize(UNiagaraStackEntry::FRequiredEntryData(
+		InSystemViewModel,
+		TSharedPtr<FNiagaraEmitterViewModel>(),
+		UNiagaraStackEntry::FExecutionCategoryNames::System,
+		UNiagaraStackEntry::FExecutionSubcategoryNames::Settings,
+		InSystemViewModel->GetEditorData().GetStackEditorData()));
+
+	SelectionStackViewModel = NewObject<UNiagaraStackViewModel>(this);
+	SelectionStackViewModel->InitializeWithRootEntry(StackSelection);
+
+	bSystemIsSelected = false;
+}
+
+void UNiagaraSystemSelectionViewModel::Finalize()
+{
+	SystemViewModelWeak.Reset();
+
+	if (StackSelection != nullptr)
+	{
+		StackSelection->Finalize();
+		StackSelection = nullptr;
+	}
+
+	if (SelectionStackViewModel != nullptr)
+	{
+		SelectionStackViewModel->Finalize();
+		SelectionStackViewModel = nullptr;
+	}
+}
+
+const TArray<UNiagaraStackEntry*> UNiagaraSystemSelectionViewModel::GetSelectedEntries() const
+{
+	return SelectedEntries;
+}
+
+bool UNiagaraSystemSelectionViewModel::GetSystemIsSelected() const
+{
+	return bSystemIsSelected;
+}
+
+const TArray<FGuid>& UNiagaraSystemSelectionViewModel::GetSelectedEmitterHandleIds() const
+{
+	return SelectedEmitterHandleIds;
+}
+
+void UNiagaraSystemSelectionViewModel::UpdateSelectionFromEntries(const TArray<UNiagaraStackEntry*> InSelectedEntries, const TArray<UNiagaraStackEntry*> InDeselectedEntries, bool bClearCurrentSelection)
+{
+	if (bClearCurrentSelection)
+	{
+		SelectedEntries.Empty();
+		bSystemIsSelected = false;
+		SelectedEmitterHandleIds.Empty();
+		OnSelectionChangedDelegate.Broadcast(ESelectionChangeSource::Clear);
+		UpdateStackSelectionEntry();
+	}
+
+	bool bSelectionChanged = false;
+	TSet<UNiagaraStackEntry*> CurrentSelection(SelectedEntries);
+
+	for (UNiagaraStackEntry* DeselectedEntry : InDeselectedEntries)
+	{
+		if (CurrentSelection.Contains(DeselectedEntry))
+		{
+			SelectedEntries.Remove(DeselectedEntry);
+			TSharedPtr<FNiagaraEmitterViewModel> DeselectedEmitterViewModel = DeselectedEntry->GetEmitterViewModel();
+			if (DeselectedEmitterViewModel.IsValid())
+			{
+				bool bEmitterIsStillSelected = SelectedEntries.ContainsByPredicate([DeselectedEmitterViewModel](const UNiagaraStackEntry* SelectedEntry)
+				{ 
+					return SelectedEntry->GetEmitterViewModel() == DeselectedEmitterViewModel; 
+				});
+				if (bEmitterIsStillSelected == false)
+				{
+					const FNiagaraEmitterHandle* DeselectedEmitterHandle = FNiagaraEditorUtilities::GetEmitterHandleForEmitter(DeselectedEntry->GetSystemViewModel()->GetSystem(), *DeselectedEmitterViewModel->GetEmitter());
+					if (DeselectedEmitterHandle != nullptr)
+					{
+						SelectedEmitterHandleIds.Remove(DeselectedEmitterHandle->GetId());
+					}
+				}
+			}
+			else
+			{
+				bool bSystemIsStillSelected = SelectedEntries.ContainsByPredicate([](const UNiagaraStackEntry* SelectedEntry)
+				{
+					return SelectedEntry->GetEmitterViewModel().IsValid() == false;
+				});
+				if (bSystemIsStillSelected == false)
+				{
+					bSystemIsSelected = false;
+				}
+			}
+			bSelectionChanged = true;
+		}
+	}
+
+	for (UNiagaraStackEntry* SelectedEntry : InSelectedEntries)
+	{
+		if (CurrentSelection.Contains(SelectedEntry) == false)
+		{
+			SelectedEntries.Add(SelectedEntry);
+			TSharedPtr<FNiagaraEmitterViewModel> SelectedEmitterViewModel = SelectedEntry->GetEmitterViewModel();
+			if (SelectedEmitterViewModel.IsValid())
+			{
+				const FNiagaraEmitterHandle* SelectedEmitterHandle = FNiagaraEditorUtilities::GetEmitterHandleForEmitter(SelectedEntry->GetSystemViewModel()->GetSystem(), *SelectedEmitterViewModel->GetEmitter());
+				if (SelectedEmitterHandle != nullptr)
+				{
+					SelectedEmitterHandleIds.AddUnique(SelectedEmitterHandle->GetId());
+				}
+			}
+			else
+			{
+				bSystemIsSelected = true;
+			}
+			bSelectionChanged = true;
+		}
+	}
+
+	if (bSelectionChanged)
+	{
+		OnSelectionChangedDelegate.Broadcast(ESelectionChangeSource::EntrySelection);
+		UpdateStackSelectionEntry();
+	}
+}
+
+void UNiagaraSystemSelectionViewModel::UpdateSelectionFromTopLevelObjects(bool bInSystemIsSelected, const TArray<FGuid> InSelectedEmitterIds, bool bClearCurrentSelection)
+{
+	if (bClearCurrentSelection)
+	{
+		SelectedEntries.Empty();
+		bSystemIsSelected = false;
+		SelectedEmitterHandleIds.Empty();
+		OnSelectionChangedDelegate.Broadcast(ESelectionChangeSource::Clear);
+		UpdateStackSelectionEntry();
+	}
+
+	bool bSelectionChanged = false;
+	if (bSystemIsSelected != bInSystemIsSelected)
+	{
+		bSystemIsSelected = bInSystemIsSelected;
+		if (bSystemIsSelected)
+		{
+			SelectedEntries.Append(GetSystemViewModel()->GetSystemStackViewModel()->GetRootEntries());
+		}
+		else
+		{
+			TArray<UNiagaraStackEntry*> SystemRootEntries = GetSystemViewModel()->GetSystemStackViewModel()->GetRootEntries();
+			SelectedEntries.RemoveAll([&SystemRootEntries](UNiagaraStackEntry* SelectedEntry) { return SystemRootEntries.Contains(SelectedEntry); });
+		}
+		bSelectionChanged = true;
+	}
+
+	TSet<FGuid> CurrentSelectedEmitterHandleIds(SelectedEmitterHandleIds);
+
+	for (const FGuid& CurrentSelectedEmitterHandleId : CurrentSelectedEmitterHandleIds)
+	{
+		if (InSelectedEmitterIds.Contains(CurrentSelectedEmitterHandleId) == false)
+		{
+			SelectedEmitterHandleIds.Remove(CurrentSelectedEmitterHandleId);
+			TSharedPtr<FNiagaraEmitterHandleViewModel> DeselectedEmitterHandleViewModel = GetSystemViewModel()->GetEmitterHandleViewModelById(CurrentSelectedEmitterHandleId);
+			TArray<UNiagaraStackEntry*> DeselectedEmitterRootEntries = DeselectedEmitterHandleViewModel->GetEmitterStackViewModel()->GetRootEntries();
+			SelectedEntries.RemoveAll([&DeselectedEmitterRootEntries](UNiagaraStackEntry* SelectedEntry) { return DeselectedEmitterRootEntries.Contains(SelectedEntry); });
+			bSelectionChanged = true;
+		}
+	}
+
+	for (const FGuid& InSelectedEmitterId : InSelectedEmitterIds)
+	{
+		if (CurrentSelectedEmitterHandleIds.Contains(InSelectedEmitterId) == false)
+		{
+			SelectedEmitterHandleIds.Add(InSelectedEmitterId);
+			TSharedPtr<FNiagaraEmitterHandleViewModel> SelectedEmitterHandleViewModel = GetSystemViewModel()->GetEmitterHandleViewModelById(InSelectedEmitterId);
+			if (SelectedEmitterHandleViewModel.IsValid())
+			{
+				SelectedEntries.Append(SelectedEmitterHandleViewModel->GetEmitterStackViewModel()->GetRootEntries());
+			}
+			bSelectionChanged = true;
+		}
+	}
+
+	if (bSelectionChanged)
+	{
+		OnSelectionChangedDelegate.Broadcast(ESelectionChangeSource::TopObjectLevelSelection);
+		UpdateStackSelectionEntry();
+	}
+}
+
+UNiagaraStackViewModel* UNiagaraSystemSelectionViewModel::GetSelectionStackViewModel()
+{
+	return SelectionStackViewModel;
+}
+
+void UNiagaraSystemSelectionViewModel::Refresh()
+{
+	TSet<FGuid> ValidEmitterHandleIds;
+	for (const TSharedRef<FNiagaraEmitterHandleViewModel>& EmitterHandleViewModel : GetSystemViewModel()->GetEmitterHandleViewModels())
+	{
+		ValidEmitterHandleIds.Add(EmitterHandleViewModel->GetId());
+	}
+
+	int32 NumEmitterHandleIdsRemoved = SelectedEmitterHandleIds.RemoveAll(
+		[&ValidEmitterHandleIds](const FGuid& SelectedEmitterHandleId) { return ValidEmitterHandleIds.Contains(SelectedEmitterHandleId) == false; });
+
+	int32 NumEntriesRemoved = SelectedEntries.RemoveAll(
+		[](const UNiagaraStackEntry* SelectedEntry) { return SelectedEntry->IsFinalized(); });
+
+	if (NumEmitterHandleIdsRemoved > 0 || NumEntriesRemoved > 0)
+	{
+		OnSelectionChangedDelegate.Broadcast(ESelectionChangeSource::Refresh);
+		UpdateStackSelectionEntry();
+	}
+}
+
+UNiagaraSystemSelectionViewModel::FOnSelectionChanged& UNiagaraSystemSelectionViewModel::OnSelectionChanged()
+{
+	return OnSelectionChangedDelegate;
+}
+
+TSharedRef<FNiagaraSystemViewModel> UNiagaraSystemSelectionViewModel::GetSystemViewModel()
+{
+	TSharedPtr<FNiagaraSystemViewModel> SystemViewModel = SystemViewModelWeak.Pin();
+	checkf(SystemViewModel.IsValid(), TEXT("Owning system view model destroyed before system overview view model."));
+	return SystemViewModel.ToSharedRef();
+}
+
+void UNiagaraSystemSelectionViewModel::UpdateStackSelectionEntry()
+{
+	// Check the new selection to make sure that there are no entries which are owned by other entries and if there are
+	// remove them.
+	TArray<UNiagaraStackEntry*> StackSelectionEntries;
+	for (UNiagaraStackEntry* SelectedEntry : SelectedEntries)
+	{
+		bool bIsChildEntry = false;
+		UNiagaraStackEntry* EntryInOuterChain = SelectedEntry->GetTypedOuter<UNiagaraStackEntry>();
+		while (EntryInOuterChain != nullptr)
+		{
+			if (SelectedEntries.Contains(EntryInOuterChain))
+			{
+				bIsChildEntry = true;
+				EntryInOuterChain = nullptr;
+			}
+			else
+			{
+				EntryInOuterChain = EntryInOuterChain->GetTypedOuter<UNiagaraStackEntry>();
+			}
+		}
+
+		if (bIsChildEntry == false)
+		{
+			StackSelectionEntries.Add(SelectedEntry);
+		}
+	}
+
+	StackSelection->SetSelectedEntries(StackSelectionEntries);
+}
\ No newline at end of file
