Niagara - System overview MVP

diff --git a/Engine/Plugins/FX/Niagara/Source/NiagaraEditor/Private/ViewModels/NiagaraSystemViewModel.cpp b/Engine/Plugins/FX/Niagara/Source/NiagaraEditor/Private/ViewModels/NiagaraSystemViewModel.cpp
--- a/Engine/Plugins/FX/Niagara/Source/NiagaraEditor/Private/ViewModels/NiagaraSystemViewModel.cpp
+++ b/Engine/Plugins/FX/Niagara/Source/NiagaraEditor/Private/ViewModels/NiagaraSystemViewModel.cpp
@@ -9,6 +9,8 @@
 #include "NiagaraEditorUtilities.h"
 #include "ViewModels/NiagaraEmitterHandleViewModel.h"
 #include "ViewModels/NiagaraEmitterViewModel.h"
+#include "ViewModels/NiagaraSystemSelectionViewModel.h"
+#include "ViewModels/Stack/NiagaraStackViewModel.h"
 #include "ViewModels/Stack/NiagaraStackGraphUtilities.h"
 #include "NiagaraSystemScriptViewModel.h"
 #include "NiagaraScriptGraphViewModel.h"
@@ -50,19 +52,13 @@ FNiagaraSystemViewModelOptions::FNiagaraSystemViewModelOptions()
 {
 }
 
-FNiagaraSystemViewModel::FNiagaraSystemViewModel(UNiagaraSystem& InSystem, FNiagaraSystemViewModelOptions InOptions, TOptional<const FGuid> InMessageLogGuid)
-	: System(InSystem)
+FNiagaraSystemViewModel::FNiagaraSystemViewModel()
+	: System(nullptr)
 	, PreviewComponent(nullptr)
 	, SystemInstance(nullptr)
-	, SystemScriptViewModel(MakeShareable(new FNiagaraSystemScriptViewModel(System, this)))
 	, NiagaraSequence(nullptr)
 	, bSettingSequencerTimeDirectly(false)
-	, bCanModifyEmittersFromTimeline(InOptions.bCanModifyEmittersFromTimeline)
-	, bCanAutoCompile(InOptions.bCanAutoCompile)
 	, bForceAutoCompileOnce(false)
-	, bCanSimulate(InOptions.bCanSimulate)
-	, EditMode(InOptions.EditMode)
-	, OnGetSequencerAddMenuContent(InOptions.OnGetSequencerAddMenuContent)
 	, bUpdatingEmittersFromSequencerDataChange(false)
 	, bUpdatingSequencerFromEmitterDataChange(false)
 	, bUpdatingSystemSelectionFromSequencer(false)
@@ -70,13 +66,36 @@ FNiagaraSystemViewModel::FNiagaraSystemViewModel(UNiagaraSystem& InSystem, FNiag
 	, EditorSettings(GetMutableDefault<UNiagaraEditorSettings>())
 	, bResetRequestPending(false)
 	, bCompilePendingCompletion(false)
-	, SystemMessageLogGuidKey(InMessageLogGuid)
+	, SystemStackViewModel(nullptr)
+	, SelectionViewModel(nullptr)
 {
+	GEditor->RegisterForUndo(this);
+}
+
+void FNiagaraSystemViewModel::Initialize(UNiagaraSystem& InSystem, FNiagaraSystemViewModelOptions InOptions)
+{
+	System = &InSystem;
+	RegisteredHandle = RegisterViewModelWithMap(System, this);
+
+	bCanModifyEmittersFromTimeline = InOptions.bCanModifyEmittersFromTimeline;
+	bCanAutoCompile = InOptions.bCanAutoCompile;
+	bCanSimulate = InOptions.bCanSimulate;
+	EditMode = InOptions.EditMode;
+	OnGetSequencerAddMenuContent = InOptions.OnGetSequencerAddMenuContent;
+	SystemMessageLogGuidKey = InOptions.MessageLogGuid;	
+	
+	SelectionViewModel = NewObject<UNiagaraSystemSelectionViewModel>(GetTransientPackage());
+	SelectionViewModel->Initialize(this->AsShared());
+	SelectionViewModel->OnSelectionChanged().AddSP(this, &FNiagaraSystemViewModel::SystemSelectionChanged);
+
+	SystemScriptViewModel = MakeShared<FNiagaraSystemScriptViewModel>(GetSystem(), this);
+
+	OverviewGraphViewModel = MakeShared<FNiagaraOverviewGraphViewModel>();
+	OverviewGraphViewModel->Initialize(this->AsShared());
+
 	SetupPreviewComponentAndInstance();
 	SetupSequencer();
 	RefreshAll();
-	GEditor->RegisterForUndo(this);
-	RegisteredHandle = RegisterViewModelWithMap(&InSystem, this);
 	AddSystemEventHandlers();
 	SendLastCompileMessageJobs();
 }
@@ -84,7 +103,7 @@ FNiagaraSystemViewModel::FNiagaraSystemViewModel(UNiagaraSystem& InSystem, FNiag
 void FNiagaraSystemViewModel::DumpToText(FString& ExportText)
 {
 	TSet<UObject*> ExportObjs;
-	ExportObjs.Add(&System);
+	ExportObjs.Add(System);
 	FEdGraphUtilities::ExportNodesToText(ExportObjs, ExportText);
 }
 
@@ -132,6 +151,18 @@ void FNiagaraSystemViewModel::Cleanup()
 	PreviewComponent = nullptr;
 	RemoveSystemEventHandlers();
 	SystemScriptViewModel.Reset();
+
+	if (SystemScriptViewModel.IsValid())
+	{
+		SystemStackViewModel->Finalize();
+		SystemStackViewModel = nullptr;
+	}
+
+	if (SelectionViewModel != nullptr)
+	{
+		SelectionViewModel->Finalize();
+		SelectionViewModel = nullptr;
+	}
 }
 
 FNiagaraSystemViewModel::~FNiagaraSystemViewModel()
@@ -141,6 +172,11 @@ FNiagaraSystemViewModel::~FNiagaraSystemViewModel()
 	UE_LOG(LogNiagaraEditor, Warning, TEXT("Deleting System view model %p"), this);
 }
 
+FText FNiagaraSystemViewModel::GetDisplayName() const
+{
+	return FText::FromString(System->GetName());
+}
+
 const TArray<TSharedRef<FNiagaraEmitterHandleViewModel>>& FNiagaraSystemViewModel::GetEmitterHandleViewModels()
 {
 	return EmitterHandleViewModels;
@@ -158,6 +194,18 @@ TSharedPtr<FNiagaraEmitterHandleViewModel> FNiagaraSystemViewModel::GetEmitterHa
 	return TSharedPtr<FNiagaraEmitterHandleViewModel>();
 }
 
+TSharedPtr<FNiagaraEmitterHandleViewModel> FNiagaraSystemViewModel::GetEmitterHandleViewModelForEmitter(UNiagaraEmitter* InEmitter) const
+{
+	for (TSharedRef<FNiagaraEmitterHandleViewModel> EmitterHandleViewModel : EmitterHandleViewModels)
+	{
+		if (EmitterHandleViewModel->GetEmitterViewModel()->GetEmitter() == InEmitter)
+		{
+			return EmitterHandleViewModel;
+		}
+	}
+	return TSharedPtr<FNiagaraEmitterHandleViewModel>();
+}
+
 TSharedPtr<FNiagaraSystemScriptViewModel> FNiagaraSystemViewModel::GetSystemScriptViewModel()
 {
 	return SystemScriptViewModel;
@@ -166,7 +214,7 @@ TSharedPtr<FNiagaraSystemScriptViewModel> FNiagaraSystemViewModel::GetSystemScri
 void FNiagaraSystemViewModel::CompileSystem(bool bForce)
 {
 	SCOPE_CYCLE_COUNTER(STAT_NiagaraEditor_SystemViewModel_CompileSystem);
-	FNiagaraEditorUtilities::KillSystemInstances(System);
+	FNiagaraEditorUtilities::KillSystemInstances(GetSystem());
 	check(SystemScriptViewModel.IsValid());
 	SystemScriptViewModel->CompileSystem(bForce);
 	bCompilePendingCompletion = true;
@@ -178,82 +226,9 @@ ENiagaraScriptCompileStatus FNiagaraSystemViewModel::GetLatestCompileStatus() co
 	return SystemScriptViewModel->GetLatestCompileStatus();
 }
 
-const TArray<FGuid>& FNiagaraSystemViewModel::GetSelectedEmitterHandleIds()
-{
-	return SelectedEmitterHandleIds;
-}
-
-void FNiagaraSystemViewModel::SetSelectedEmitterHandlesById(TArray<FGuid> InSelectedEmitterHandleIds)
-{
-	bool bSelectionChanged = false;
-	if (SelectedEmitterHandleIds.Num() == InSelectedEmitterHandleIds.Num())
-	{
-		for (FGuid InSelectedEmitterHandleId : InSelectedEmitterHandleIds)
-		{
-			if (SelectedEmitterHandleIds.Contains(InSelectedEmitterHandleId) == false)
-			{
-				bSelectionChanged = true;
-				break;
-			}
-		}
-	}
-	else
-	{
-		bSelectionChanged = true;
-	}
-
-	SelectedEmitterHandleIds.Empty();
-	SelectedEmitterHandleIds.Append(InSelectedEmitterHandleIds);
-	if (bSelectionChanged)
-	{
-		if (bUpdatingSystemSelectionFromSequencer == false)
-		{
-			UpdateSequencerFromEmitterHandleSelection();
-		}
-		OnSelectedEmitterHandlesChangedDelegate.Broadcast();
-	}
-}
-
-void FNiagaraSystemViewModel::SetSelectedEmitterHandleById(FGuid InSelectedEmitterHandleId)
-{
-	TArray<FGuid> SingleSelectedEmitterHandleId;
-	SingleSelectedEmitterHandleId.Add(InSelectedEmitterHandleId);
-	SetSelectedEmitterHandlesById(SingleSelectedEmitterHandleId);
-}
-
-void FNiagaraSystemViewModel::GetSelectedEmitterHandles(TArray<TSharedRef<FNiagaraEmitterHandleViewModel>>& OutSelectedEmitterHanldles)
-{
-	for (TSharedRef<FNiagaraEmitterHandleViewModel> EmitterHandleViewModel : EmitterHandleViewModels)
-	{
-		if (SelectedEmitterHandleIds.Contains(EmitterHandleViewModel->GetId()))
-		{
-			OutSelectedEmitterHanldles.Add(EmitterHandleViewModel);
-		}
-	}
-}
-
-const UNiagaraSystemEditorData& FNiagaraSystemViewModel::GetEditorData() const
+UNiagaraSystemEditorData& FNiagaraSystemViewModel::GetEditorData() const
 {
-	const UNiagaraSystemEditorData* EditorData = Cast<UNiagaraSystemEditorData>(System.GetEditorData());
-	if (EditorData == nullptr)
-	{
-		EditorData = GetDefault<UNiagaraSystemEditorData>();
-	}
-	return *EditorData;
-}
-
-UNiagaraSystemEditorData& FNiagaraSystemViewModel::GetOrCreateEditorData()
-{
-	UNiagaraSystemEditorData* EditorData = Cast<UNiagaraSystemEditorData>(System.GetEditorData());
-	if (EditorData == nullptr)
-	{
-		EditorData = NewObject<UNiagaraSystemEditorData>(&System, NAME_None, RF_Transactional);
-		System.Modify();
-		System.SetEditorData(EditorData);
-	}
-	bool bEditingSystem = GetEditMode() == ENiagaraSystemViewModelEditMode::SystemAsset;
-	EditorData->Initialize(&System, bEditingSystem);
-	return *EditorData;
+	return *CastChecked<UNiagaraSystemEditorData>(GetSystem().GetEditorData(), ECastCheckedType::NullChecked);
 }
 
 UNiagaraComponent* FNiagaraSystemViewModel::GetPreviewComponent()
@@ -271,6 +246,12 @@ FNiagaraCurveOwner& FNiagaraSystemViewModel::GetCurveOwner()
 	return CurveOwner;
 }
 
+UNiagaraSystem& FNiagaraSystemViewModel::GetSystem() const
+{
+	checkf(System != nullptr, TEXT("System view model not initialized before use."));
+	return *System;
+}
+
 bool FNiagaraSystemViewModel::GetCanModifyEmittersFromTimeline() const
 {
 	return bCanModifyEmittersFromTimeline;
@@ -300,16 +281,14 @@ void FNiagaraSystemViewModel::AddEmitter(UNiagaraEmitter& Emitter)
 		GEditor->BeginTransaction(LOCTEXT("AddEmitter", "Add emitter"));
 	}
 
-	const FGuid NewEmitterHandleId = FNiagaraEditorUtilities::AddEmitterToSystem(System, Emitter, bSystemIsPlaceholder);
-	UNiagaraSystemEditorData& EditorData = GetOrCreateEditorData();
-	EditorData.SystemOverviewHandleAdded(NewEmitterHandleId);
+	const FGuid NewEmitterHandleId = FNiagaraEditorUtilities::AddEmitterToSystem(GetSystem(), Emitter);
 
 	if (false == bSystemIsPlaceholder)
 	{
 		GEditor->EndTransaction();
 	}
 
-	if (System.GetNumEmitters() == 1 && EditorSettings->GetAutoPlay())
+	if (GetSystem().GetNumEmitters() == 1 && EditorSettings->GetAutoPlay())
 	{
 		// When adding a new emitter to an empty system start playing.
 		Sequencer->SetPlaybackStatus(EMovieScenePlayerStatus::Playing);
@@ -317,13 +296,13 @@ void FNiagaraSystemViewModel::AddEmitter(UNiagaraEmitter& Emitter)
 
 	RefreshAll();
 
-	TRange<float> SystemPlaybackRange = EditorData.GetPlaybackRange();
+	TRange<float> SystemPlaybackRange = GetEditorData().GetPlaybackRange();
 	TRange<float> EmitterPlaybackRange = GetEmitterHandleViewModelById(NewEmitterHandleId)->GetEmitterViewModel()->GetEditorData().GetPlaybackRange();
 	TRange<float> NewSystemPlaybackRange = TRange<float>(
 		FMath::Min(SystemPlaybackRange.GetLowerBoundValue(), EmitterPlaybackRange.GetLowerBoundValue()),
 		FMath::Max(SystemPlaybackRange.GetUpperBoundValue(), EmitterPlaybackRange.GetUpperBoundValue()));
 
-	EditorData.SetPlaybackRange(NewSystemPlaybackRange);
+	GetEditorData().SetPlaybackRange(NewSystemPlaybackRange);
 
 	TGuardValue<bool> UpdateGuard(bUpdatingSequencerFromEmitterDataChange, true);
 
@@ -332,7 +311,9 @@ void FNiagaraSystemViewModel::AddEmitter(UNiagaraEmitter& Emitter)
 
 	NiagaraSequence->GetMovieScene()->SetPlaybackRange(NewStartFrame.RoundToFrame(), NewDuration);
 
-	SetSelectedEmitterHandleById(NewEmitterHandleId);
+	TArray<FGuid> SelectedEmitterHandles;
+	SelectedEmitterHandles.Add(NewEmitterHandleId);
+	SelectionViewModel->UpdateSelectionFromTopLevelObjects(false, SelectedEmitterHandles, true);
 
 	bForceAutoCompileOnce = true;
 }
@@ -345,13 +326,13 @@ void FNiagaraSystemViewModel::DuplicateEmitters(TArray<FEmitterHandleToDuplicate
 	}
 
 	// Kill all system instances before modifying the emitter handle list to prevent accessing deleted data.
-	FNiagaraEditorUtilities::KillSystemInstances(System);
+	FNiagaraEditorUtilities::KillSystemInstances(GetSystem());
 	const FScopedTransaction DeleteTransaction(EmitterHandlesToDuplicate.Num() == 1
 		? LOCTEXT("DuplicateEmitter", "Duplicate emitter")
 		: LOCTEXT("DuplicateEmitters", "Duplicate emitters"));
 
 	TSet<FName> EmitterHandleNames;
-	for (const FNiagaraEmitterHandle& EmitterHandle : System.GetEmitterHandles())
+	for (const FNiagaraEmitterHandle& EmitterHandle : GetSystem().GetEmitterHandles())
 	{
 		EmitterHandleNames.Add(EmitterHandle.GetName());
 	}
@@ -382,13 +363,13 @@ void FNiagaraSystemViewModel::DuplicateEmitters(TArray<FEmitterHandleToDuplicate
 
 		if (HandleToDuplicate.IsValid())
 		{
-			const FNiagaraEmitterHandle& EmitterHandle = System.DuplicateEmitterHandle(HandleToDuplicate, FNiagaraUtilities::GetUniqueName(HandleToDuplicate.GetName(), EmitterHandleNames));
+			const FNiagaraEmitterHandle& EmitterHandle = GetSystem().DuplicateEmitterHandle(HandleToDuplicate, FNiagaraUtilities::GetUniqueName(HandleToDuplicate.GetName(), EmitterHandleNames));
 			EmitterHandleNames.Add(EmitterHandle.GetName());
-			GetEditorData().SystemOverviewHandleAdded(EmitterHandle.GetId());
 		}
 	}
 
-	FNiagaraStackGraphUtilities::RebuildEmitterNodes(System);
+	FNiagaraStackGraphUtilities::RebuildEmitterNodes(GetSystem());
+	GetEditorData().SynchronizeOverviewGraphWithSystem(GetSystem());
 	RefreshAll();
 	bForceAutoCompileOnce = true;
 }
@@ -397,9 +378,7 @@ void FNiagaraSystemViewModel::DeleteEmitters(TSet<FGuid> EmitterHandleIdsToDelet
 {
 	if (EmitterHandleIdsToDelete.Num() > 0)
 	{
-		FNiagaraEditorUtilities::RemoveEmittersFromSystemByEmitterHandleId(System, EmitterHandleIdsToDelete);
-		GetEditorData().SystemOverviewHandlesRemoved();
-
+		FNiagaraEditorUtilities::RemoveEmittersFromSystemByEmitterHandleId(GetSystem(), EmitterHandleIdsToDelete);
 		RefreshAll();
 		bForceAutoCompileOnce = true;
 	}
@@ -444,11 +423,6 @@ FNiagaraSystemViewModel::FOnCurveOwnerChanged& FNiagaraSystemViewModel::OnCurveO
 	return OnCurveOwnerChangedDelegate;
 }
 
-FNiagaraSystemViewModel::FOnSelectedEmitterHandlesChanged& FNiagaraSystemViewModel::OnSelectedEmitterHandlesChanged()
-{
-	return OnSelectedEmitterHandlesChangedDelegate;
-}
-
 FNiagaraSystemViewModel::FOnPostSequencerTimeChange& FNiagaraSystemViewModel::OnPostSequencerTimeChanged()
 {
 	return OnPostSequencerTimeChangeDelegate;
@@ -469,6 +443,14 @@ void FNiagaraSystemViewModel::AddReferencedObjects(FReferenceCollector& Collecto
 	{
 		Collector.AddReferencedObject(NiagaraSequence);
 	}
+	if (SystemStackViewModel != nullptr)
+	{
+		Collector.AddReferencedObject(SystemStackViewModel);
+	}
+	if (SelectionViewModel != nullptr)
+	{
+		Collector.AddReferencedObject(SelectionViewModel);
+	}
 }
 
 void FNiagaraSystemViewModel::PostUndo(bool bSuccess)
@@ -478,7 +460,7 @@ void FNiagaraSystemViewModel::PostUndo(bool bSuccess)
 
 void FNiagaraSystemViewModel::Tick(float DeltaTime)
 {
-	if (bCompilePendingCompletion && System.HasOutstandingCompilationRequests() == false)
+	if (bCompilePendingCompletion && GetSystem().HasOutstandingCompilationRequests() == false)
 	{
 		bCompilePendingCompletion = false;
 		OnSystemCompiled().Broadcast();
@@ -493,7 +475,7 @@ void FNiagaraSystemViewModel::Tick(float DeltaTime)
 		if (SystemScriptViewModel->GetLatestCompileStatus() == ENiagaraScriptCompileStatus::NCS_Dirty)
 		{
 			//SystemScriptViewModel->CompileSystem();
-			//UE_LOG(LogNiagaraEditor, Log, TEXT("Compiling %s due to dirty scripts."), *System.GetName());
+			//UE_LOG(LogNiagaraEditor, Log, TEXT("Compiling %s due to dirty scripts."), *GetSystem().GetName());
 			bRecompile |= true;
 		}
 
@@ -503,11 +485,11 @@ void FNiagaraSystemViewModel::Tick(float DeltaTime)
 			{
 				bRecompile |= true;
 				//EmitterHandleViewModel->GetEmitterViewModel()->CompileScripts();
-				//UE_LOG(LogNiagaraEditor, Log, TEXT("Compiling %s - %s due to dirty scripts."), *System.GetName(), *EmitterHandleViewModel->GetName().ToString());
+				//UE_LOG(LogNiagaraEditor, Log, TEXT("Compiling %s - %s due to dirty scripts."), *GetSystem().GetName(), *EmitterHandleViewModel->GetName().ToString());
 			}
 		}
 
-		if (System.HasOutstandingCompilationRequests() == false)
+		if (GetSystem().HasOutstandingCompilationRequests() == false)
 		{
 			if (bRecompile || bForceAutoCompileOnce)
 			{
@@ -531,19 +513,19 @@ void FNiagaraSystemViewModel::Tick(float DeltaTime)
 
 void FNiagaraSystemViewModel::OnPreSave()
 {
-	if (System.HasOutstandingCompilationRequests())
+	if (GetSystem().HasOutstandingCompilationRequests())
 	{
-		UE_LOG(LogNiagaraEditor, Log, TEXT("System %s has pending compile jobs. Waiting for that code to complete before Saving.."), *System.GetName());
-		System.WaitForCompilationComplete();
+		UE_LOG(LogNiagaraEditor, Log, TEXT("System %s has pending compile jobs. Waiting for that code to complete before Saving.."), *GetSystem().GetName());
+		GetSystem().WaitForCompilationComplete();
 	}
 }
 
 void FNiagaraSystemViewModel::OnPreClose()
 {
-	if (System.HasOutstandingCompilationRequests())
+	if (GetSystem().HasOutstandingCompilationRequests())
 	{
-		UE_LOG(LogNiagaraEditor, Log, TEXT("System %s has pending compile jobs. Waiting for that code to complete before Closing.."), *System.GetName());
-		System.WaitForCompilationComplete();
+		UE_LOG(LogNiagaraEditor, Log, TEXT("System %s has pending compile jobs. Waiting for that code to complete before Closing.."), *GetSystem().GetName());
+		GetSystem().WaitForCompilationComplete();
 	}
 }
 
@@ -588,6 +570,21 @@ const TArray<FNiagaraStackModuleData>& FNiagaraSystemViewModel::GetStackModuleDa
 	return EmitterToCachedStackModuleData[EmitterHandleId];
 }
 
+UNiagaraStackViewModel* FNiagaraSystemViewModel::GetSystemStackViewModel()
+{
+	if (SystemStackViewModel == nullptr)
+	{
+		SystemStackViewModel = NewObject<UNiagaraStackViewModel>(GetTransientPackage());
+		SystemStackViewModel->InitializeWithViewModels(this->AsShared(), TSharedPtr<FNiagaraEmitterHandleViewModel>(), FNiagaraStackViewModelOptions(true, false));
+	}
+	return SystemStackViewModel;
+}
+
+UNiagaraSystemSelectionViewModel* FNiagaraSystemViewModel::GetSelectionViewModel()
+{
+	return SelectionViewModel;
+}
+
 TStatId FNiagaraSystemViewModel::GetStatId() const
 {
 	RETURN_QUICK_DECLARE_CYCLE_STAT(FNiagaraSystemViewModel, STATGROUP_Tickables);
@@ -616,9 +613,9 @@ void FNiagaraSystemViewModel::SendLastCompileMessageJobs() const
 	int32 WarningCount = 0;
 
 	TArray<FNiagaraScriptAndOwningScriptNameString> ScriptsToGetCompileEventsFrom;
-	ScriptsToGetCompileEventsFrom.Add(FNiagaraScriptAndOwningScriptNameString(System.GetSystemSpawnScript(), System.GetName()));
-	ScriptsToGetCompileEventsFrom.Add(FNiagaraScriptAndOwningScriptNameString(System.GetSystemUpdateScript(), System.GetName()));
-	const TArray<FNiagaraEmitterHandle> EmitterHandles = System.GetEmitterHandles();
+	ScriptsToGetCompileEventsFrom.Add(FNiagaraScriptAndOwningScriptNameString(GetSystem().GetSystemSpawnScript(), GetSystem().GetName()));
+	ScriptsToGetCompileEventsFrom.Add(FNiagaraScriptAndOwningScriptNameString(GetSystem().GetSystemUpdateScript(), GetSystem().GetName()));
+	const TArray<FNiagaraEmitterHandle> EmitterHandles = GetSystem().GetEmitterHandles();
 	for (const FNiagaraEmitterHandle& Handle : EmitterHandles)
 	{
 		const UNiagaraEmitter* EmitterInSystem = Handle.GetInstance();
@@ -663,14 +660,13 @@ void FNiagaraSystemViewModel::SetupPreviewComponentAndInstance()
 		PreviewComponent = NewObject<UNiagaraComponent>(GetTransientPackage(), NAME_None, RF_Transient);
 		PreviewComponent->CastShadow = 1;
 		PreviewComponent->bCastDynamicShadow = 1;
-		PreviewComponent->SetAsset(&System);
+		PreviewComponent->SetAsset(System);
 		PreviewComponent->SetForceSolo(true);
 		PreviewComponent->SetAgeUpdateMode(ENiagaraAgeUpdateMode::DesiredAge);
 		PreviewComponent->SetCanRenderWhileSeeking(false);
 		PreviewComponent->Activate(true);
 
-		UNiagaraSystemEditorData& EditorData = GetOrCreateEditorData();
-		FTransform OwnerTransform = EditorData.GetOwnerTransform();
+		FTransform OwnerTransform = GetEditorData().GetOwnerTransform();
 		PreviewComponent->SetRelativeTransform(OwnerTransform);
 
 		PreviewComponent->OnSystemInstanceChanged().AddRaw(this, &FNiagaraSystemViewModel::PreviewComponentSystemInstanceChanged);
@@ -708,7 +704,7 @@ void FNiagaraSystemViewModel::NotifyDataObjectChanged(UObject* ChangedObject)
 	ResetSystem(ETimeResetMode::AllowResetTime, EMultiResetMode::ResetThisInstance, EReinitMode::ReinitializeSystem);
 }
 
-void FNiagaraSystemViewModel::IsolateEmitters(TArray<TSharedRef<FNiagaraEmitterHandleViewModel>> EmitterHandlesToIsolate)
+void FNiagaraSystemViewModel::IsolateEmitters(TArray<FGuid> EmitterHandlesIdsToIsolate)
 {
 	for (TSharedRef<FNiagaraEmitterHandleViewModel> EmitterHandle : EmitterHandleViewModels)
 	{
@@ -716,13 +712,16 @@ void FNiagaraSystemViewModel::IsolateEmitters(TArray<TSharedRef<FNiagaraEmitterH
 	}
 
 	bool bAnyEmitterIsolated = false;
-	for (TSharedRef<FNiagaraEmitterHandleViewModel> EmitterHandleToIsolate : EmitterHandlesToIsolate)
+	for (TSharedRef<FNiagaraEmitterHandleViewModel> EmitterHandle : EmitterHandleViewModels)
 	{
-		bAnyEmitterIsolated = true;
-		EmitterHandleToIsolate->GetEmitterHandle()->SetIsolated(true);
+		if (EmitterHandlesIdsToIsolate.Contains(EmitterHandle->GetId()))
+		{
+			bAnyEmitterIsolated = true;
+			EmitterHandle->GetEmitterHandle()->SetIsolated(true);
+		}
 	}
 
-	System.SetIsolateEnabled(bAnyEmitterIsolated);
+	GetSystem().SetIsolateEnabled(bAnyEmitterIsolated);
 }
 
 void FNiagaraSystemViewModel::ToggleEmitterIsolation(TSharedRef<FNiagaraEmitterHandleViewModel> InEmitterHandle)
@@ -739,12 +738,7 @@ void FNiagaraSystemViewModel::ToggleEmitterIsolation(TSharedRef<FNiagaraEmitterH
 		}
 	}
 
-	System.SetIsolateEnabled(bAnyEmitterIsolated);
-}
-
-bool FNiagaraSystemViewModel::IsEmitterIsolated(TSharedRef<FNiagaraEmitterHandleViewModel> InEmitterHandle)
-{
-	return InEmitterHandle->GetEmitterHandle()->IsIsolated();
+	GetSystem().SetIsolateEnabled(bAnyEmitterIsolated);
 }
 
 void FNiagaraSystemViewModel::RefreshEmitterHandleViewModels()
@@ -756,35 +750,34 @@ void FNiagaraSystemViewModel::RefreshEmitterHandleViewModels()
 	// Map existing view models to the real instances that now exist. Reuse if we can. Create a new one if we cannot.
 	TArray<FGuid> ValidEmitterHandleIds;
 	int32 i;
-	for (i = 0; i < System.GetNumEmitters(); ++i)
+	for (i = 0; i < GetSystem().GetNumEmitters(); ++i)
 	{
-		FNiagaraEmitterHandle* EmitterHandle = &System.GetEmitterHandle(i);
+		FNiagaraEmitterHandle* EmitterHandle = &GetSystem().GetEmitterHandle(i);
 		TSharedPtr<FNiagaraEmitterInstance, ESPMode::ThreadSafe> Simulation = SystemInstance ? SystemInstance->GetSimulationForHandle(*EmitterHandle) : nullptr;
 		ValidEmitterHandleIds.Add(EmitterHandle->GetId());
 
+		TSharedPtr<FNiagaraEmitterHandleViewModel> ViewModel;
 		bool bAdd = OldViewModels.Num() <= i;
 		if (bAdd)
 		{
-			TSharedRef<FNiagaraEmitterHandleViewModel> ViewModel = MakeShareable(new FNiagaraEmitterHandleViewModel(EmitterHandle, Simulation, System));
-			// Since we're adding fresh, we need to register all the event handlers.
-			ViewModel->OnPropertyChanged().AddRaw(this, &FNiagaraSystemViewModel::EmitterHandlePropertyChanged, ViewModel);
-			ViewModel->OnNameChanged().AddRaw(this, &FNiagaraSystemViewModel::EmitterHandleNameChanged, ViewModel);
-			ViewModel->GetEmitterViewModel()->OnPropertyChanged().AddRaw(this, &FNiagaraSystemViewModel::EmitterPropertyChanged, ViewModel);
+			ViewModel = MakeShared<FNiagaraEmitterHandleViewModel>();
+			ViewModel->OnPropertyChanged().AddRaw(this, &FNiagaraSystemViewModel::EmitterHandlePropertyChanged, ViewModel.ToSharedRef());
+			ViewModel->OnNameChanged().AddRaw(this, &FNiagaraSystemViewModel::EmitterHandleNameChanged, ViewModel.ToSharedRef());
+			ViewModel->GetEmitterViewModel()->OnPropertyChanged().AddRaw(this, &FNiagaraSystemViewModel::EmitterPropertyChanged, ViewModel.ToSharedRef());
 			ViewModel->GetEmitterViewModel()->OnScriptCompiled().AddRaw(this, &FNiagaraSystemViewModel::ScriptCompiled);
-			ViewModel->GetEmitterViewModel()->OnScriptGraphChanged().AddRaw(this, &FNiagaraSystemViewModel::EmitterScriptGraphChanged, ViewModel);
-			ViewModel->GetEmitterViewModel()->OnScriptParameterStoreChanged().AddRaw(this, &FNiagaraSystemViewModel::EmitterParameterStoreChanged, ViewModel);
-			EmitterHandleViewModels.Add(ViewModel);
+			ViewModel->GetEmitterViewModel()->OnScriptGraphChanged().AddRaw(this, &FNiagaraSystemViewModel::EmitterScriptGraphChanged, ViewModel.ToSharedRef());
+			ViewModel->GetEmitterViewModel()->OnScriptParameterStoreChanged().AddRaw(this, &FNiagaraSystemViewModel::EmitterParameterStoreChanged, ViewModel.ToSharedRef());
 		}
 		else
 		{
-			TSharedRef<FNiagaraEmitterHandleViewModel> ViewModel = OldViewModels[i];
-			ViewModel->Set(EmitterHandle, Simulation, System);
-			EmitterHandleViewModels.Add(ViewModel);
+			ViewModel = OldViewModels[i];
 		}
 
+		ViewModel->Initialize(this->AsShared(), EmitterHandle, Simulation);
+		EmitterHandleViewModels.Add(ViewModel.ToSharedRef());
 	}
 
-	check(EmitterHandleViewModels.Num() == System.GetNumEmitters());
+	check(EmitterHandleViewModels.Num() == GetSystem().GetNumEmitters());
 
 	// Clear out any old view models that may still be left around.
 	for (; i < OldViewModels.Num(); i++)
@@ -795,18 +788,15 @@ void FNiagaraSystemViewModel::RefreshEmitterHandleViewModels()
 		ViewModel->GetEmitterViewModel()->OnScriptCompiled().RemoveAll(this);
 		ViewModel->GetEmitterViewModel()->OnScriptGraphChanged().RemoveAll(this);
 		ViewModel->GetEmitterViewModel()->OnScriptParameterStoreChanged().RemoveAll(this);
-		ViewModel->Set(nullptr, nullptr, System);
+		ViewModel->Reset();
 	}
 
-	// Remove any invalid ids from the handle selection.
-	auto EmitterHandleIdIsInvalid = [&](FGuid& EmitterHandleId) { return ValidEmitterHandleIds.Contains(EmitterHandleId) == false; };
-	int32 NumRemoved = SelectedEmitterHandleIds.RemoveAll(EmitterHandleIdIsInvalid);
-
-	OnEmitterHandleViewModelsChangedDelegate.Broadcast();
-	if (NumRemoved > 0)
+	if (SelectionViewModel != nullptr)
 	{
-		OnSelectedEmitterHandlesChangedDelegate.Broadcast();
+		SelectionViewModel->Refresh();
 	}
+
+	OnEmitterHandleViewModelsChangedDelegate.Broadcast();
 }
 
 void PopulateChildMovieSceneFoldersFromNiagaraFolders(const UNiagaraSystemEditorFolder* NiagaraFolder, UMovieSceneFolder* MovieSceneFolder, const TMap<FGuid, UMovieSceneNiagaraEmitterTrack*>& EmitterHandleIdToTrackMap)
@@ -950,7 +940,7 @@ void FNiagaraSystemViewModel::SetupSequencer()
 	Sequencer->OnGlobalTimeChanged().AddRaw(this, &FNiagaraSystemViewModel::SequencerTimeChanged);
 	Sequencer->GetSelectionChangedTracks().AddRaw(this, &FNiagaraSystemViewModel::SequencerTrackSelectionChanged);
 	Sequencer->GetSelectionChangedSections().AddRaw(this, &FNiagaraSystemViewModel::SequencerSectionSelectionChanged);
-	Sequencer->SetPlaybackStatus(System.GetNumEmitters() > 0 && EditorSettings->GetAutoPlay()
+	Sequencer->SetPlaybackStatus(GetSystem().GetNumEmitters() > 0 && EditorSettings->GetAutoPlay()
 		? EMovieScenePlayerStatus::Playing
 		: EMovieScenePlayerStatus::Stopped);
 }
@@ -978,7 +968,7 @@ void FNiagaraSystemViewModel::ResetSystem(ETimeResetMode TimeResetMode, EMultiRe
 		}
 	}
 
-	TArray<UNiagaraComponent*> ReferencingComponents = FNiagaraEditorUtilities::GetComponentsThatReferenceSystem(System);
+	TArray<UNiagaraComponent*> ReferencingComponents = FNiagaraEditorUtilities::GetComponentsThatReferenceSystem(GetSystem());
 	for (auto Component : ReferencingComponents)
 	{
 		if (ReinitMode == EReinitMode::ResetSystem)
@@ -1083,7 +1073,7 @@ void FNiagaraSystemViewModel::ResetCurveData()
 		CurveData,
 		ShownCurveDataInterfaces);
 	// Get curves from user variables
-	for (UNiagaraDataInterface* DataInterface : System.GetExposedParameters().GetDataInterfaces())
+	for (UNiagaraDataInterface* DataInterface : GetSystem().GetExposedParameters().GetDataInterfaces())
 	{
 		UNiagaraDataInterfaceCurveBase* CurveDataInterface = Cast<UNiagaraDataInterfaceCurveBase>(DataInterface);
 		if (CurveDataInterface != nullptr)
@@ -1141,7 +1131,7 @@ void FNiagaraSystemViewModel::UpdateCompiledDataInterfaces(UNiagaraDataInterface
 		// If the data interface was owned by an input node, then we need to try to update the compiled version.
 		UNiagaraEmitter* OwningEmitter;
 		UNiagaraScript* OwningScript;
-		FNiagaraStackGraphUtilities::GetOwningEmitterAndScriptForStackNode(*OuterInputNode, System, OwningEmitter, OwningScript);
+		FNiagaraStackGraphUtilities::GetOwningEmitterAndScriptForStackNode(*OuterInputNode, GetSystem(), OwningEmitter, OwningScript);
 		if (ensureMsgf(OwningScript != nullptr, TEXT("Could not find owning script for data interface input node.")))
 		{
 			switch (OwningScript->GetUsage())
@@ -1158,7 +1148,7 @@ void FNiagaraSystemViewModel::UpdateCompiledDataInterfaces(UNiagaraDataInterface
 			case ENiagaraScriptUsage::EmitterUpdateScript:
 				if (ensureMsgf(OwningEmitter != nullptr, TEXT("Could not find owning emitter for data interface input node.")))
 				{
-					UNiagaraScript& TargetScript = OwningScript->GetUsage() == ENiagaraScriptUsage::EmitterSpawnScript ? *System.GetSystemSpawnScript() : *System.GetSystemUpdateScript();
+					UNiagaraScript& TargetScript = OwningScript->GetUsage() == ENiagaraScriptUsage::EmitterSpawnScript ? *GetSystem().GetSystemSpawnScript() : *GetSystem().GetSystemUpdateScript();
 					FName AliasedInputNodeName = FNiagaraParameterMapHistory::ResolveEmitterAlias(OuterInputNode->Input.GetName(), OwningEmitter->GetUniqueEmitterName());
 					UpdateCompiledDataInterfacesForScript(TargetScript, AliasedInputNodeName, *ChangedDataInterface);
 				}
@@ -1169,10 +1159,10 @@ void FNiagaraSystemViewModel::UpdateCompiledDataInterfaces(UNiagaraDataInterface
 	else
 	{
 		// If the data interface wasn't owned by a script, try to find it in the exposed parameter data interfaces.
-		const FNiagaraVariable* FoundExposedDataInterface = System.GetExposedParameters().FindVariable(ChangedDataInterface);
+		const FNiagaraVariable* FoundExposedDataInterface = GetSystem().GetExposedParameters().FindVariable(ChangedDataInterface);
 		if (FoundExposedDataInterface != nullptr)
 		{
-			System.GetExposedParameters().OnInterfaceChange();
+			GetSystem().GetExposedParameters().OnInterfaceChange();
 		}
 	}
 }
@@ -1343,11 +1333,11 @@ void FNiagaraSystemViewModel::SequencerDataChanged(EMovieSceneDataChangeType Dat
 	{
 		TGuardValue<bool> UpdateGuard(bUpdatingEmittersFromSequencerDataChange, true);
 
-		GetOrCreateEditorData().Modify();
+		GetEditorData().Modify();
 		TRange<FFrameNumber> FramePlaybackRange = NiagaraSequence->GetMovieScene()->GetPlaybackRange();
 		float StartTimeSeconds = NiagaraSequence->GetMovieScene()->GetTickResolution().AsSeconds(FramePlaybackRange.GetLowerBoundValue());
 		float EndTimeSeconds = NiagaraSequence->GetMovieScene()->GetTickResolution().AsSeconds(FramePlaybackRange.GetUpperBoundValue());
-		GetOrCreateEditorData().SetPlaybackRange(TRange<float>(StartTimeSeconds, EndTimeSeconds));
+		GetEditorData().SetPlaybackRange(TRange<float>(StartTimeSeconds, EndTimeSeconds));
 
 		TSet<FGuid> VaildTrackEmitterHandleIds;
 		TArray<FEmitterHandleToDuplicate> EmittersToDuplicate;
@@ -1399,11 +1389,6 @@ void FNiagaraSystemViewModel::SequencerDataChanged(EMovieSceneDataChangeType Dat
 			if (bCanModifyEmittersFromTimeline)
 			{
 				DeleteEmitters(RemovedEmitterHandleIds);
-				// When deleting emitters from sequencer, select a new one if one is available.
-				if (SelectedEmitterHandleIds.Num() == 0 && EmitterHandleViewModels.Num() > 0)
-				{
-					SetSelectedEmitterHandleById(EmitterHandleViewModels[0]->GetId());
-				}
 			}
 			else
 			{
@@ -1427,7 +1412,7 @@ void FNiagaraSystemViewModel::SequencerDataChanged(EMovieSceneDataChangeType Dat
 		TArray<UMovieSceneFolder*> RootFolders = NiagaraSequence->GetMovieScene()->GetRootFolders();
 		if (RootFolders.Num() != 0 || GetEditorData().GetRootFolder().GetChildFolders().Num() != 0)
 		{
-			PopulateNiagaraFoldersFromMovieSceneFolders(RootFolders, RootTracks, &GetOrCreateEditorData().GetRootFolder());
+			PopulateNiagaraFoldersFromMovieSceneFolders(RootFolders, RootTracks, &GetEditorData().GetRootFolder());
 		}
 
 		if (bRefreshAllTracks)
@@ -1491,6 +1476,14 @@ void FNiagaraSystemViewModel::SequencerTimeChanged()
 	OnPostSequencerTimeChangeDelegate.Broadcast();
 }
 
+void FNiagaraSystemViewModel::SystemSelectionChanged(UNiagaraSystemSelectionViewModel::ESelectionChangeSource SelectionChangeSource)
+{
+	if (bUpdatingSystemSelectionFromSequencer == false)
+	{
+		UpdateSequencerFromEmitterHandleSelection();
+	}
+}
+
 void FNiagaraSystemViewModel::SequencerTrackSelectionChanged(TArray<UMovieSceneTrack*> SelectedTracks)
 {
 	if (bUpdatingSequencerSelectionFromSystem == false)
@@ -1534,14 +1527,19 @@ void FNiagaraSystemViewModel::UpdateEmitterHandleSelectionFromSequencer()
 	}
 
 	TGuardValue<bool> UpdateGuard(bUpdatingSystemSelectionFromSequencer, true);
-	SetSelectedEmitterHandlesById(NewSelectedEmitterHandleIds);
+	if (SelectionViewModel != nullptr)
+	{
+		bool bClearCurrentSelection = FSlateApplication::Get().GetModifierKeys().IsControlDown() == false;
+		bool bSelectSystem = NewSelectedEmitterHandleIds.Num() > 0 && GetEditorData().GetOwningSystemIsPlaceholder() == false;
+		SelectionViewModel->UpdateSelectionFromTopLevelObjects(bSelectSystem, NewSelectedEmitterHandleIds, bClearCurrentSelection);
+	}
 }
 
 void FNiagaraSystemViewModel::UpdateSequencerFromEmitterHandleSelection()
 {
 	TGuardValue<bool> UpdateGuard(bUpdatingSequencerSelectionFromSystem, true);
 	Sequencer->EmptySelection();
-	for (FGuid SelectedEmitterHandleId : SelectedEmitterHandleIds)
+	for (FGuid SelectedEmitterHandleId : SelectionViewModel->GetSelectedEmitterHandleIds())
 	{
 		for (UMovieSceneTrack* MasterTrack : NiagaraSequence->GetMovieScene()->GetMasterTracks())
 		{
@@ -1596,12 +1594,13 @@ void FNiagaraSystemViewModel::SystemInstanceInitialized()
 
 void FNiagaraSystemViewModel::UpdateEmitterFixedBounds()
 {
-	TArray<TSharedRef<FNiagaraEmitterHandleViewModel>> SelectedEmitterHandles;
-	GetSelectedEmitterHandles(SelectedEmitterHandles);
-
-	for (TSharedRef<FNiagaraEmitterHandleViewModel>& SelectedEmitterHandleViewModel : SelectedEmitterHandles)
+	for (TSharedRef<FNiagaraEmitterHandleViewModel>& EmitterHandleViewModel : EmitterHandleViewModels)
 	{
-		FNiagaraEmitterHandle* SelectedEmitterHandle = SelectedEmitterHandleViewModel->GetEmitterHandle();
+		if (SelectionViewModel->GetSelectedEmitterHandleIds().Contains(EmitterHandleViewModel->GetId()) == false)
+		{
+			continue;
+		}
+		FNiagaraEmitterHandle* SelectedEmitterHandle = EmitterHandleViewModel->GetEmitterHandle();
 		check(SelectedEmitterHandle);
 		UNiagaraEmitter* Emitter = SelectedEmitterHandle->GetInstance();
 		for (TSharedRef<FNiagaraEmitterInstance, ESPMode::ThreadSafe>& EmitterInst : PreviewComponent->GetSystemInstance()->GetEmitters())
@@ -1617,11 +1616,11 @@ void FNiagaraSystemViewModel::UpdateEmitterFixedBounds()
 
 void FNiagaraSystemViewModel::AddSystemEventHandlers()
 {
-	if (System.IsValid())
+	if (GetSystem().IsValid())
 	{
 		TArray<UNiagaraScript*> Scripts;
-		Scripts.Add(System.GetSystemSpawnScript());
-		Scripts.Add(System.GetSystemUpdateScript());
+		Scripts.Add(GetSystem().GetSystemSpawnScript());
+		Scripts.Add(GetSystem().GetSystemUpdateScript());
 		
 		for (UNiagaraScript* Script : Scripts)
 		{
@@ -1631,9 +1630,9 @@ void FNiagaraSystemViewModel::AddSystemEventHandlers()
 			ScriptToOnParameterStoreChangedHandleMap.Add(FObjectKey(Script), OnParameterStoreChangedHandle);
 		}
 
-		UserParameterStoreChangedHandle = System.GetExposedParameters().AddOnChangedHandler(
+		UserParameterStoreChangedHandle = GetSystem().GetExposedParameters().AddOnChangedHandler(
 			FNiagaraParameterStore::FOnChanged::FDelegate::CreateRaw<FNiagaraSystemViewModel, const FNiagaraParameterStore&, const UNiagaraScript*>(
-				this, &FNiagaraSystemViewModel::SystemParameterStoreChanged, System.GetExposedParameters(), nullptr));
+				this, &FNiagaraSystemViewModel::SystemParameterStoreChanged, GetSystem().GetExposedParameters(), nullptr));
 
 		SystemScriptGraphChangedHandler = SystemScriptViewModel->GetGraphViewModel()->GetGraph()->AddOnGraphChangedHandler(FOnGraphChanged::FDelegate::CreateRaw(this, &FNiagaraSystemViewModel::SystemScriptGraphChanged));
 	}
@@ -1641,11 +1640,11 @@ void FNiagaraSystemViewModel::AddSystemEventHandlers()
 
 void FNiagaraSystemViewModel::RemoveSystemEventHandlers()
 {
-	if (System.IsValid())
+	if (GetSystem().IsValid())
 	{
 		TArray<UNiagaraScript*> Scripts;
-		Scripts.Add(System.GetSystemSpawnScript());
-		Scripts.Add(System.GetSystemUpdateScript());
+		Scripts.Add(GetSystem().GetSystemSpawnScript());
+		Scripts.Add(GetSystem().GetSystemUpdateScript());
 
 		for (UNiagaraScript* Script : Scripts)
 		{
@@ -1656,7 +1655,7 @@ void FNiagaraSystemViewModel::RemoveSystemEventHandlers()
 			}
 		}
 
-		System.GetExposedParameters().RemoveOnChangedHandler(UserParameterStoreChangedHandle);
+		GetSystem().GetExposedParameters().RemoveOnChangedHandler(UserParameterStoreChangedHandle);
 		if (SystemScriptViewModel.IsValid())
 		{
 			SystemScriptViewModel->GetGraphViewModel()->GetGraph()->RemoveOnGraphChangedHandler(SystemScriptGraphChangedHandler);
@@ -1706,19 +1705,4 @@ TSharedPtr<FNiagaraOverviewGraphViewModel> FNiagaraSystemViewModel::GetOverviewG
 	return OverviewGraphViewModel;
 }
 
-void FNiagaraSystemViewModel::SetOverviewGraphViewModel(TSharedRef<FNiagaraOverviewGraphViewModel> InOverviewGraphViewModel)
-{
-	OverviewGraphViewModel = InOverviewGraphViewModel;
-}
-
-void FNiagaraSystemViewModel::EmitterNameChanged(const FText& InText, FGuid ChangedEmitterGuid)
-{
-	TSharedPtr<FNiagaraEmitterHandleViewModel> EmitterHandleViewModel = GetEmitterHandleViewModelById(ChangedEmitterGuid);
-	if (ensureMsgf(EmitterHandleViewModel.IsValid(), TEXT("Failed to find EmitterHandleViewModel with Emitter that has same Guid as Emitter that was renamed!")))
-	{
-		EmitterHandleViewModel->OnNameTextComitted(InText, ETextCommit::Default);
-		OverviewGraphViewModel->GetGraph()->NotifyGraphChanged();
-	}
-}
-
 #undef LOCTEXT_NAMESPACE // NiagaraSystemViewModel
diff --git a/Engine/Plugins/FX/Niagara/Source/NiagaraEditor/Private/ViewModels/Stack/NiagaraStackEntry.cpp b/Engine/Plugins/FX/Niagara/Source/NiagaraEditor/Private/ViewModels/Stack/NiagaraStackEntry.cpp
--- a/Engine/Plugins/FX/Niagara/Source/NiagaraEditor/Private/ViewModels/Stack/NiagaraStackEntry.cpp
+++ b/Engine/Plugins/FX/Niagara/Source/NiagaraEditor/Private/ViewModels/Stack/NiagaraStackEntry.cpp
@@ -18,6 +18,7 @@ const FName UNiagaraStackEntry::FExecutionSubcategoryNames::Settings = TEXT("Set
 const FName UNiagaraStackEntry::FExecutionSubcategoryNames::Spawn = TEXT("Spawn");
 const FName UNiagaraStackEntry::FExecutionSubcategoryNames::Update = TEXT("Update");
 const FName UNiagaraStackEntry::FExecutionSubcategoryNames::Event = TEXT("Event");
+const FName UNiagaraStackEntry::FExecutionSubcategoryNames::Render = TEXT("Render");
 
 UNiagaraStackEntry::FStackIssueFix::FStackIssueFix()
 {
@@ -165,6 +166,11 @@ void UNiagaraStackEntry::Finalize()
 	ErrorChildren.Empty();
 }
 
+bool UNiagaraStackEntry::IsFinalized() const
+{
+	return bIsFinalized;
+}
+
 FText UNiagaraStackEntry::GetDisplayName() const
 {
 	return FText();
@@ -289,11 +295,9 @@ TSharedRef<FNiagaraSystemViewModel> UNiagaraStackEntry::GetSystemViewModel() con
 	return PinnedSystemViewModel.ToSharedRef();
 }
 
-TSharedRef<FNiagaraEmitterViewModel> UNiagaraStackEntry::GetEmitterViewModel() const
+TSharedPtr<FNiagaraEmitterViewModel> UNiagaraStackEntry::GetEmitterViewModel() const
 {
-	TSharedPtr<FNiagaraEmitterViewModel> PinnedEmitterViewModel = EmitterViewModel.Pin();
-	checkf(PinnedEmitterViewModel.IsValid(), TEXT("Base stack entry not initialized or emitter view model was already deleted."));
-	return PinnedEmitterViewModel.ToSharedRef();
+	return EmitterViewModel.Pin();
 }
 
 UNiagaraStackEntry::FOnStructureChanged& UNiagaraStackEntry::OnStructureChanged()
@@ -390,8 +394,7 @@ bool UNiagaraStackEntry::HasBaseEmitter() const
 {
 	if (bHasBaseEmitterCache.IsSet() == false)
 	{
-		TSharedRef<FNiagaraScriptMergeManager> MergeManager = FNiagaraScriptMergeManager::Get();
-		const UNiagaraEmitter* BaseEmitter = GetEmitterViewModel()->GetParentEmitter();
+		const UNiagaraEmitter* BaseEmitter = GetEmitterViewModel().IsValid() ? GetEmitterViewModel()->GetParentEmitter() : nullptr;
 		bHasBaseEmitterCache = BaseEmitter != nullptr;
 	}
 	return bHasBaseEmitterCache.GetValue();
@@ -429,7 +432,7 @@ void UNiagaraStackEntry::FinalizeInternal()
 void UNiagaraStackEntry::RefreshChildren()
 {
 	checkf(bIsFinalized == false, TEXT("Can not refresh children on an entry after it has been finalized."));
-	checkf(SystemViewModel.IsValid() && EmitterViewModel.IsValid(), TEXT("Base stack entry not initialized."));
+	checkf(SystemViewModel.IsValid(), TEXT("Base stack entry not initialized."));
 
 	for (UNiagaraStackEntry* Child : Children)
 	{
@@ -455,11 +458,11 @@ void UNiagaraStackEntry::RefreshChildren()
 	TArray<FStackIssue> NewStackIssues;
 	RefreshChildrenInternal(Children, NewChildren, NewStackIssues);
 
-	// If any of the current children were not moved to the new children collection than finalize them since
+	// If any of the current children were not moved to the new children collection, and they're owned by this entry than finalize them since
 	// they weren't reused.
 	for (UNiagaraStackEntry* Child : Children)
 	{
-		if (NewChildren.Contains(Child) == false)
+		if (NewChildren.Contains(Child) == false && Child->GetOuter() == this)
 		{
 			Child->Finalize();
 		}
diff --git a/Engine/Plugins/FX/Niagara/Source/NiagaraEditor/Private/ViewModels/Stack/NiagaraStackEventHandlerGroup.cpp b/Engine/Plugins/FX/Niagara/Source/NiagaraEditor/Private/ViewModels/Stack/NiagaraStackEventHandlerGroup.cpp
--- a/Engine/Plugins/FX/Niagara/Source/NiagaraEditor/Private/ViewModels/Stack/NiagaraStackEventHandlerGroup.cpp
+++ b/Engine/Plugins/FX/Niagara/Source/NiagaraEditor/Private/ViewModels/Stack/NiagaraStackEventHandlerGroup.cpp
@@ -54,7 +54,9 @@ public:
 
 		// Set the emitter here so that the internal state of the view model is updated.
 		// TODO: Move the logic for managing event handlers into the emitter view model or script view model.
-		EmitterViewModelPinned->SetEmitter(Emitter);
+		TWeakPtr<FNiagaraEmitterInstance, ESPMode::ThreadSafe> Simulation = EmitterViewModelPinned->GetSimulation();
+		EmitterViewModelPinned->Reset();
+		EmitterViewModelPinned->Initialize(Emitter, Simulation);
 
 		OnItemAdded.ExecuteIfBound(EventScriptProperties);
 	}
@@ -70,7 +72,7 @@ void UNiagaraStackEventHandlerGroup::Initialize(FRequiredEntryData InRequiredEnt
 {
 	FText DisplayName = LOCTEXT("EventGroupName", "Add Event Handler");
 	FText ToolTip = LOCTEXT("EventGroupTooltip", "Determines how this Emitter responds to incoming events. There can be more than one event handler script stack per Emitter.");
-	AddUtilities = MakeShared<FEventHandlerGroupAddUtilities>(InRequiredEntryData.EmitterViewModel, 
+	AddUtilities = MakeShared<FEventHandlerGroupAddUtilities>(InRequiredEntryData.EmitterViewModel.ToSharedRef(), 
 		TNiagaraStackItemGroupAddUtilities<FNiagaraEventScriptProperties>::FOnItemAdded::CreateUObject(this, &UNiagaraStackEventHandlerGroup::ItemAddedFromUtilties));
 	Super::Initialize(InRequiredEntryData, DisplayName, ToolTip, AddUtilities.Get());
 }
diff --git a/Engine/Plugins/FX/Niagara/Source/NiagaraEditor/Private/ViewModels/Stack/NiagaraStackFunctionInput.cpp b/Engine/Plugins/FX/Niagara/Source/NiagaraEditor/Private/ViewModels/Stack/NiagaraStackFunctionInput.cpp
--- a/Engine/Plugins/FX/Niagara/Source/NiagaraEditor/Private/ViewModels/Stack/NiagaraStackFunctionInput.cpp
+++ b/Engine/Plugins/FX/Niagara/Source/NiagaraEditor/Private/ViewModels/Stack/NiagaraStackFunctionInput.cpp
@@ -134,7 +134,7 @@ void UNiagaraStackFunctionInput::Initialize(
 	OwningAssignmentNode = Cast<UNiagaraNodeAssignment>(OwningFunctionCallNode.Get());
 
 	UNiagaraNodeOutput* OutputNode = FNiagaraStackGraphUtilities::GetEmitterOutputNodeForStackNode(*OwningModuleNode.Get());
-	UNiagaraEmitter* ParentEmitter = GetEmitterViewModel()->GetEmitter();
+	UNiagaraEmitter* ParentEmitter = GetEmitterViewModel().IsValid() ? GetEmitterViewModel()->GetEmitter() : nullptr;
 	UNiagaraSystem* ParentSystem = &GetSystemViewModel()->GetSystem();
 	if (OutputNode)
 	{
@@ -201,8 +201,9 @@ void UNiagaraStackFunctionInput::Initialize(
 		AffectedScriptsNotWeak.Add(AffectedScript.Get());
 	}
 
-	EditCondition.Initialize(SourceScript.Get(), AffectedScriptsNotWeak, GetEmitterViewModel()->GetEmitter()->GetUniqueEmitterName(), OwningFunctionCallNode.Get());
-	VisibleCondition.Initialize(SourceScript.Get(), AffectedScriptsNotWeak, GetEmitterViewModel()->GetEmitter()->GetUniqueEmitterName(), OwningFunctionCallNode.Get());
+	FString UniqueEmitterName = GetEmitterViewModel().IsValid() ? GetEmitterViewModel()->GetEmitter()->GetUniqueEmitterName() : FString();
+	EditCondition.Initialize(SourceScript.Get(), AffectedScriptsNotWeak, UniqueEmitterName, OwningFunctionCallNode.Get());
+	VisibleCondition.Initialize(SourceScript.Get(), AffectedScriptsNotWeak, UniqueEmitterName, OwningFunctionCallNode.Get());
 }
 
 void UNiagaraStackFunctionInput::FinalizeInternal()
@@ -492,9 +493,9 @@ FString UNiagaraStackFunctionInput::ResolveDisplayNameArgument(const FString& In
 {
 	if (InArg.StartsWith(TEXT("MaterialDynamicParam")))
 	{
-		TSharedRef<FNiagaraEmitterViewModel> ThisEmitterViewModel = GetEmitterViewModel();
+		TSharedPtr<FNiagaraEmitterViewModel> ThisEmitterViewModel = GetEmitterViewModel();
 		TArray<UMaterialExpressionDynamicParameter*> ExpressionParams;
-		if (!UNiagaraStackFunctionInputUtilities::GetMaterialExpressionDynamicParameter(ThisEmitterViewModel->GetEmitter(), ExpressionParams))
+		if (ThisEmitterViewModel.IsValid() == false || !UNiagaraStackFunctionInputUtilities::GetMaterialExpressionDynamicParameter(ThisEmitterViewModel->GetEmitter(), ExpressionParams))
 		{
 			return InArg.Replace(TEXT("MaterialDynamic"), TEXT("")) + TEXT(" (No material found using dynamic params)");
 		}
@@ -1135,7 +1136,8 @@ bool UNiagaraStackFunctionInput::CanReset() const
 					if (FNiagaraStackGraphUtilities::IsValidDefaultDynamicInput(*SourceScript, *DefaultPin))
 					{
 						UEdGraphPin* OverridePin = GetOverridePin();
-						bNewCanReset = OverridePin == nullptr || FNiagaraStackGraphUtilities::DoesDynamicInputMatchDefault(GetEmitterViewModel()->GetEmitter()->GetUniqueEmitterName(), *SourceScript,
+						FString UniqueEmitterName = GetEmitterViewModel().IsValid() ? GetEmitterViewModel()->GetEmitter()->GetUniqueEmitterName() : FString();
+						bNewCanReset = OverridePin == nullptr || FNiagaraStackGraphUtilities::DoesDynamicInputMatchDefault(UniqueEmitterName, *SourceScript,
 							*OwningFunctionCallNode, *OverridePin, InputParameterHandle.GetName(), *DefaultPin) == false;
 					}
 					else
@@ -1380,8 +1382,8 @@ void UNiagaraStackFunctionInput::ResetToBase()
 FNiagaraVariable UNiagaraStackFunctionInput::CreateRapidIterationVariable(const FName& InName)
 {
 	UNiagaraNodeOutput* OutputNode = FNiagaraStackGraphUtilities::GetEmitterOutputNodeForStackNode(*OwningModuleNode.Get());
-	UNiagaraEmitter* ParentEmitter = GetEmitterViewModel()->GetEmitter();
-	return FNiagaraStackGraphUtilities::CreateRapidIterationParameter(ParentEmitter->GetUniqueEmitterName(), OutputNode->GetUsage(), InName, InputType);
+	FString UniqueEmitterName = GetEmitterViewModel().IsValid() ? GetEmitterViewModel()->GetEmitter()->GetUniqueEmitterName() : FString();
+	return FNiagaraStackGraphUtilities::CreateRapidIterationParameter(UniqueEmitterName, OutputNode->GetUsage(), InName, InputType);
 }
 
 bool UNiagaraStackFunctionInput::CanRenameInput() const
diff --git a/Engine/Plugins/FX/Niagara/Source/NiagaraEditor/Private/ViewModels/Stack/NiagaraStackFunctionInputCollection.cpp b/Engine/Plugins/FX/Niagara/Source/NiagaraEditor/Private/ViewModels/Stack/NiagaraStackFunctionInputCollection.cpp
--- a/Engine/Plugins/FX/Niagara/Source/NiagaraEditor/Private/ViewModels/Stack/NiagaraStackFunctionInputCollection.cpp
+++ b/Engine/Plugins/FX/Niagara/Source/NiagaraEditor/Private/ViewModels/Stack/NiagaraStackFunctionInputCollection.cpp
@@ -81,7 +81,9 @@ void UNiagaraStackFunctionInputCollection::RefreshChildrenInternal(const TArray<
 {
 	TSet<const UEdGraphPin*> HiddenPins;
 	TArray<const UEdGraphPin*> InputPins;
-	FCompileConstantResolver ConstantResolver(GetEmitterViewModel()->GetEmitter());
+	FCompileConstantResolver ConstantResolver = GetEmitterViewModel().IsValid() 
+		? FCompileConstantResolver(GetEmitterViewModel()->GetEmitter()) 
+		: FCompileConstantResolver();
 	FNiagaraStackGraphUtilities::GetStackFunctionInputPins(*InputFunctionCallNode, InputPins, HiddenPins, ConstantResolver, FNiagaraStackGraphUtilities::ENiagaraGetStackFunctionInputPinsOptions::ModuleInputsOnly);
 
 	const UEdGraphSchema_Niagara* NiagaraSchema = GetDefault<UEdGraphSchema_Niagara>();
diff --git a/Engine/Plugins/FX/Niagara/Source/NiagaraEditor/Private/ViewModels/Stack/NiagaraStackGraphUtilities.cpp b/Engine/Plugins/FX/Niagara/Source/NiagaraEditor/Private/ViewModels/Stack/NiagaraStackGraphUtilities.cpp
--- a/Engine/Plugins/FX/Niagara/Source/NiagaraEditor/Private/ViewModels/Stack/NiagaraStackGraphUtilities.cpp
+++ b/Engine/Plugins/FX/Niagara/Source/NiagaraEditor/Private/ViewModels/Stack/NiagaraStackGraphUtilities.cpp
@@ -440,7 +440,7 @@ void FNiagaraStackGraphUtilities::ConnectStackNodeGroup(const FStackNodeGroup& C
 
 DECLARE_DELEGATE_RetVal_OneParam(bool, FInputSelector, UNiagaraStackFunctionInput*);
 
-void InitializeStackFunctionInputsInternal(TSharedRef<FNiagaraSystemViewModel> SystemViewModel, TSharedRef<FNiagaraEmitterViewModel> EmitterViewModel, UNiagaraStackEditorData& StackEditorData, UNiagaraNodeFunctionCall& ModuleNode, UNiagaraNodeFunctionCall& InputFunctionCallNode, FInputSelector InputSelector)
+void InitializeStackFunctionInputsInternal(TSharedRef<FNiagaraSystemViewModel> SystemViewModel, TSharedPtr<FNiagaraEmitterViewModel> EmitterViewModel, UNiagaraStackEditorData& StackEditorData, UNiagaraNodeFunctionCall& ModuleNode, UNiagaraNodeFunctionCall& InputFunctionCallNode, FInputSelector InputSelector)
 {
 	UNiagaraStackFunctionInputCollection* FunctionInputCollection = NewObject<UNiagaraStackFunctionInputCollection>(GetTransientPackage()); 
 	UNiagaraStackEntry::FRequiredEntryData RequiredEntryData(SystemViewModel, EmitterViewModel, NAME_None, NAME_None, StackEditorData);
@@ -472,12 +472,12 @@ void InitializeStackFunctionInputsInternal(TSharedRef<FNiagaraSystemViewModel> S
 	SystemViewModel->NotifyDataObjectChanged(nullptr);
 }
 
-void FNiagaraStackGraphUtilities::InitializeStackFunctionInputs(TSharedRef<FNiagaraSystemViewModel> SystemViewModel, TSharedRef<FNiagaraEmitterViewModel> EmitterViewModel, UNiagaraStackEditorData& StackEditorData, UNiagaraNodeFunctionCall& ModuleNode, UNiagaraNodeFunctionCall& InputFunctionCallNode)
+void FNiagaraStackGraphUtilities::InitializeStackFunctionInputs(TSharedRef<FNiagaraSystemViewModel> SystemViewModel, TSharedPtr<FNiagaraEmitterViewModel> EmitterViewModel, UNiagaraStackEditorData& StackEditorData, UNiagaraNodeFunctionCall& ModuleNode, UNiagaraNodeFunctionCall& InputFunctionCallNode)
 {
 	InitializeStackFunctionInputsInternal(SystemViewModel, EmitterViewModel, StackEditorData, ModuleNode, InputFunctionCallNode, FInputSelector());
 }
 
-void FNiagaraStackGraphUtilities::InitializeStackFunctionInput(TSharedRef<FNiagaraSystemViewModel> SystemViewModel, TSharedRef<FNiagaraEmitterViewModel> EmitterViewModel, UNiagaraStackEditorData& StackEditorData, UNiagaraNodeFunctionCall& ModuleNode, UNiagaraNodeFunctionCall& InputFunctionCallNode, FName InputName)
+void FNiagaraStackGraphUtilities::InitializeStackFunctionInput(TSharedRef<FNiagaraSystemViewModel> SystemViewModel, TSharedPtr<FNiagaraEmitterViewModel> EmitterViewModel, UNiagaraStackEditorData& StackEditorData, UNiagaraNodeFunctionCall& ModuleNode, UNiagaraNodeFunctionCall& InputFunctionCallNode, FName InputName)
 {
 	FInputSelector InputSelector;
 	InputSelector.BindLambda([&InputName](UNiagaraStackFunctionInput* Input)
@@ -1650,7 +1650,7 @@ bool FNiagaraStackGraphUtilities::DoesDynamicInputMatchDefault(
 
 void SetInputValue(
 	TSharedRef<FNiagaraSystemViewModel> SystemViewModel,
-	TSharedRef<FNiagaraEmitterViewModel> EmitterViewModel,
+	TSharedPtr<FNiagaraEmitterViewModel> EmitterViewModel,
 	UNiagaraStackEditorData& StackEditorData,
 	UNiagaraScript& SourceScript,
 	const TArray<TWeakObjectPtr<UNiagaraScript>> AffectedScripts,
@@ -1731,7 +1731,7 @@ void SetInputValue(
 
 void FNiagaraStackGraphUtilities::ResetToDefaultDynamicInput(
 	TSharedRef<FNiagaraSystemViewModel> SystemViewModel,
-	TSharedRef<FNiagaraEmitterViewModel> EmitterViewModel,
+	TSharedPtr<FNiagaraEmitterViewModel> EmitterViewModel,
 	UNiagaraStackEditorData& StackEditorData,
 	UNiagaraScript& SourceScript,
 	const TArray<TWeakObjectPtr<UNiagaraScript>> AffectedScripts,
diff --git a/Engine/Plugins/FX/Niagara/Source/NiagaraEditor/Private/ViewModels/Stack/NiagaraStackModuleItem.cpp b/Engine/Plugins/FX/Niagara/Source/NiagaraEditor/Private/ViewModels/Stack/NiagaraStackModuleItem.cpp
--- a/Engine/Plugins/FX/Niagara/Source/NiagaraEditor/Private/ViewModels/Stack/NiagaraStackModuleItem.cpp
+++ b/Engine/Plugins/FX/Niagara/Source/NiagaraEditor/Private/ViewModels/Stack/NiagaraStackModuleItem.cpp
@@ -393,7 +393,12 @@ void UNiagaraStackModuleItem::RefreshIssues(TArray<FStackIssue>& NewIssues)
 	TArray<UNiagaraNodeFunctionCall*> FoundCalls;
 	TArray<FNiagaraModuleDependency> DependenciesNeeded;
 
-	TArray<FNiagaraStackModuleData> SystemModuleData = GetSystemViewModel()->GetStackModuleDataForEmitter(GetEmitterViewModel());
+
+	TArray<FNiagaraStackModuleData> SystemModuleData;
+	if (GetEmitterViewModel().IsValid())
+	{
+		SystemModuleData = GetSystemViewModel()->GetStackModuleDataForEmitter(GetEmitterViewModel().ToSharedRef());
+	}
 	int32 ModuleIndex = INDEX_NONE;
 	for (int i = 0; i < SystemModuleData.Num(); i++)
 	{
diff --git a/Engine/Plugins/FX/Niagara/Source/NiagaraEditor/Private/ViewModels/Stack/NiagaraStackModuleItemLinkedInputCollection.cpp b/Engine/Plugins/FX/Niagara/Source/NiagaraEditor/Private/ViewModels/Stack/NiagaraStackModuleItemLinkedInputCollection.cpp
--- a/Engine/Plugins/FX/Niagara/Source/NiagaraEditor/Private/ViewModels/Stack/NiagaraStackModuleItemLinkedInputCollection.cpp
+++ b/Engine/Plugins/FX/Niagara/Source/NiagaraEditor/Private/ViewModels/Stack/NiagaraStackModuleItemLinkedInputCollection.cpp
@@ -61,7 +61,9 @@ void UNiagaraStackModuleItemLinkedInputCollection::RefreshChildrenInternal(const
 	{
 		FNiagaraParameterMapHistoryBuilder Builder;
 		Builder.SetIgnoreDisabled(false);
-		Builder.ConstantResolver = FCompileConstantResolver(GetEmitterViewModel()->GetEmitter());
+		Builder.ConstantResolver = GetEmitterViewModel().IsValid() 
+			? FCompileConstantResolver(GetEmitterViewModel()->GetEmitter())
+			: FCompileConstantResolver();
 		FunctionCallNode->BuildParameterMapHistory(Builder, false);
 
 		if (ensureMsgf(Builder.Histories.Num() == 1, TEXT("Invalid Stack Graph - Function call node has invalid history count!")))
diff --git a/Engine/Plugins/FX/Niagara/Source/NiagaraEditor/Private/ViewModels/Stack/NiagaraStackModuleItemOutputCollection.cpp b/Engine/Plugins/FX/Niagara/Source/NiagaraEditor/Private/ViewModels/Stack/NiagaraStackModuleItemOutputCollection.cpp
--- a/Engine/Plugins/FX/Niagara/Source/NiagaraEditor/Private/ViewModels/Stack/NiagaraStackModuleItemOutputCollection.cpp
+++ b/Engine/Plugins/FX/Niagara/Source/NiagaraEditor/Private/ViewModels/Stack/NiagaraStackModuleItemOutputCollection.cpp
@@ -50,7 +50,9 @@ void UNiagaraStackModuleItemOutputCollection::RefreshChildrenInternal(const TArr
 	{
 		FNiagaraParameterMapHistoryBuilder Builder;
 		Builder.SetIgnoreDisabled(false);
-		Builder.ConstantResolver = FCompileConstantResolver(GetEmitterViewModel()->GetEmitter());
+		Builder.ConstantResolver = GetEmitterViewModel().IsValid()
+			? FCompileConstantResolver(GetEmitterViewModel()->GetEmitter())
+			: FCompileConstantResolver();
 		FunctionCallNode->BuildParameterMapHistory(Builder, false);
 
 		if (ensureMsgf(Builder.Histories.Num() == 1, TEXT("Invalid Stack Graph - Function call node has invalid history count!")))
diff --git a/Engine/Plugins/FX/Niagara/Source/NiagaraEditor/Private/ViewModels/Stack/NiagaraStackRenderItemGroup.cpp b/Engine/Plugins/FX/Niagara/Source/NiagaraEditor/Private/ViewModels/Stack/NiagaraStackRenderItemGroup.cpp
--- a/Engine/Plugins/FX/Niagara/Source/NiagaraEditor/Private/ViewModels/Stack/NiagaraStackRenderItemGroup.cpp
+++ b/Engine/Plugins/FX/Niagara/Source/NiagaraEditor/Private/ViewModels/Stack/NiagaraStackRenderItemGroup.cpp
@@ -121,7 +121,7 @@ void UNiagaraStackRenderItemGroup::Initialize(FRequiredEntryData InRequiredEntry
 {
 	FText DisplayName = LOCTEXT("RenderGroupName", "Render");
 	FText ToolTip = LOCTEXT("RendererGroupTooltip", "Describes how we should display/present each particle. Note that this doesn't have to be visual. Multiple renderers are supported. Order in this stack is not necessarily relevant to draw order.");
-	AddUtilities = MakeShared<FRenderItemGroupAddUtilities>(InRequiredEntryData.EmitterViewModel,
+	AddUtilities = MakeShared<FRenderItemGroupAddUtilities>(InRequiredEntryData.EmitterViewModel.ToSharedRef(),
 		TNiagaraStackItemGroupAddUtilities<UNiagaraRendererProperties*>::FOnItemAdded::CreateUObject(this, &UNiagaraStackRenderItemGroup::ItemAdded));
 	Super::Initialize(InRequiredEntryData, DisplayName, ToolTip, AddUtilities.Get());
 }
diff --git a/Engine/Plugins/FX/Niagara/Source/NiagaraEditor/Private/ViewModels/Stack/NiagaraStackRendererItem.cpp b/Engine/Plugins/FX/Niagara/Source/NiagaraEditor/Private/ViewModels/Stack/NiagaraStackRendererItem.cpp
--- a/Engine/Plugins/FX/Niagara/Source/NiagaraEditor/Private/ViewModels/Stack/NiagaraStackRendererItem.cpp
+++ b/Engine/Plugins/FX/Niagara/Source/NiagaraEditor/Private/ViewModels/Stack/NiagaraStackRendererItem.cpp
@@ -136,7 +136,7 @@ FText UNiagaraStackRendererItem::GetDisplayName() const
 {
 	if (RendererProperties != nullptr)
 	{
-		return FText::FromString(RendererProperties->GetClass()->GetName());
+		return RendererProperties->GetClass()->GetDisplayNameText();
 	}
 	else
 	{
diff --git a/Engine/Plugins/FX/Niagara/Source/NiagaraEditor/Private/ViewModels/Stack/NiagaraStackRoot.cpp b/Engine/Plugins/FX/Niagara/Source/NiagaraEditor/Private/ViewModels/Stack/NiagaraStackRoot.cpp
--- a/Engine/Plugins/FX/Niagara/Source/NiagaraEditor/Private/ViewModels/Stack/NiagaraStackRoot.cpp
+++ b/Engine/Plugins/FX/Niagara/Source/NiagaraEditor/Private/ViewModels/Stack/NiagaraStackRoot.cpp
@@ -30,9 +30,12 @@ UNiagaraStackRoot::UNiagaraStackRoot()
 {
 }
 
-void UNiagaraStackRoot::Initialize(FRequiredEntryData InRequiredEntryData)
+void UNiagaraStackRoot::Initialize(FRequiredEntryData InRequiredEntryData, bool bInIncludeSystemInformation, bool bInIncludeEmitterInformation)
 {
-	Super::Initialize(InRequiredEntryData, FString());
+	Super::Initialize(InRequiredEntryData, TEXT("Root"));
+	bIncludeSystemInformation = bInIncludeSystemInformation;
+	bIncludeEmitterInformation = bInIncludeEmitterInformation;
+	SystemSettingsGroup = nullptr;
 	EmitterSettingsGroup = nullptr;
 	EmitterSpawnGroup = nullptr;
 	EmitterUpdateGroup = nullptr;
@@ -40,7 +43,6 @@ void UNiagaraStackRoot::Initialize(FRequiredEntryData InRequiredEntryData)
 	ParticleUpdateGroup = nullptr;
 	AddEventHandlerGroup = nullptr;
 	RenderGroup = nullptr;
-	SystemSettingsGroup = nullptr;
 }
 
 bool UNiagaraStackRoot::GetCanExpand() const
@@ -55,42 +57,39 @@ bool UNiagaraStackRoot::GetShouldShowInStack() const
 
 void UNiagaraStackRoot::RefreshChildrenInternal(const TArray<UNiagaraStackEntry*>& CurrentChildren, TArray<UNiagaraStackEntry*>& NewChildren, TArray<FStackIssue>& NewIssues)
 {
-	// We only allow displaying and editing system stacks if the system isn't transient which is the case in the emitter editor.
-	bool bShowSystemGroups = GetSystemViewModel()->GetEditMode() == ENiagaraSystemViewModelEditMode::SystemAsset;
-
 	// Create static entries as needed.
-	if (bShowSystemGroups && SystemSettingsGroup == nullptr)
+	if (bIncludeSystemInformation && SystemSettingsGroup == nullptr)
 	{
 		SystemSettingsGroup = NewObject<UNiagaraStackSystemSettingsGroup>(this);
 		FRequiredEntryData RequiredEntryData(GetSystemViewModel(), GetEmitterViewModel(),
 			FExecutionCategoryNames::System, FExecutionSubcategoryNames::Settings,
-			GetSystemViewModel()->GetOrCreateEditorData().GetStackEditorData());
+			GetSystemViewModel()->GetEditorData().GetStackEditorData());
 		SystemSettingsGroup->Initialize(RequiredEntryData, &GetSystemViewModel()->GetSystem(), &GetSystemViewModel()->GetSystem().GetExposedParameters());
 	}
 
-	if (bShowSystemGroups && SystemSpawnGroup == nullptr)
+	if (bIncludeSystemInformation && SystemSpawnGroup == nullptr)
 	{
 		SystemSpawnGroup = NewObject<UNiagaraStackScriptItemGroup>(this);
 		FRequiredEntryData RequiredEntryData(GetSystemViewModel(), GetEmitterViewModel(),
 			FExecutionCategoryNames::System, FExecutionSubcategoryNames::Spawn,
-			GetSystemViewModel()->GetOrCreateEditorData().GetStackEditorData());
+			GetSystemViewModel()->GetEditorData().GetStackEditorData());
 		FText DisplayName = LOCTEXT("SystemSpawnGroupName", "System Spawn");
 		FText ToolTip = LOCTEXT("SystemSpawnGroupToolTip", "Occurs once at System creation on the CPU. Modules in this section should initialize defaults and/or do initial setup.\r\nModules are executed in order from top to bottom of the stack.");
 		SystemSpawnGroup->Initialize(RequiredEntryData, DisplayName, ToolTip, GetSystemViewModel()->GetSystemScriptViewModel().ToSharedRef(), ENiagaraScriptUsage::SystemSpawnScript);
 	}
 
-	if (bShowSystemGroups && SystemUpdateGroup == nullptr)
+	if (bIncludeSystemInformation && SystemUpdateGroup == nullptr)
 	{
 		SystemUpdateGroup = NewObject<UNiagaraStackScriptItemGroup>(this);
 		FRequiredEntryData RequiredEntryData(GetSystemViewModel(), GetEmitterViewModel(),
 			FExecutionCategoryNames::System, FExecutionSubcategoryNames::Update,
-			GetSystemViewModel()->GetOrCreateEditorData().GetStackEditorData());
+			GetSystemViewModel()->GetEditorData().GetStackEditorData());
 		FText DisplayName = LOCTEXT("SystemUpdateGroupName", "System Update");
 		FText ToolTip = LOCTEXT("SystemUpdateGroupToolTip", "Occurs every Emitter tick on the CPU.Modules in this section should compute values for parameters for emitter or particle update or spawning this frame.\r\nModules are executed in order from top to bottom of the stack.");
 		SystemUpdateGroup->Initialize(RequiredEntryData, DisplayName, ToolTip, GetSystemViewModel()->GetSystemScriptViewModel().ToSharedRef(), ENiagaraScriptUsage::SystemUpdateScript);
 	}
 
-	if (EmitterSettingsGroup == nullptr)
+	if (bIncludeEmitterInformation && EmitterSettingsGroup == nullptr)
 	{
 		EmitterSettingsGroup = NewObject<UNiagaraStackEmitterSettingsGroup>(this);
 		FRequiredEntryData RequiredEntryData(GetSystemViewModel(), GetEmitterViewModel(),
@@ -101,7 +100,7 @@ void UNiagaraStackRoot::RefreshChildrenInternal(const TArray<UNiagaraStackEntry*
 		EmitterSettingsGroup->Initialize(RequiredEntryData, DisplayName, Tooltip, nullptr);
 	}
 
-	if (EmitterSpawnGroup == nullptr)
+	if (bIncludeEmitterInformation && EmitterSpawnGroup == nullptr)
 	{
 		EmitterSpawnGroup = NewObject<UNiagaraStackScriptItemGroup>(this);
 		FRequiredEntryData RequiredEntryData(GetSystemViewModel(), GetEmitterViewModel(),
@@ -112,7 +111,7 @@ void UNiagaraStackRoot::RefreshChildrenInternal(const TArray<UNiagaraStackEntry*
 		EmitterSpawnGroup->Initialize(RequiredEntryData, DisplayName, ToolTip, GetEmitterViewModel()->GetSharedScriptViewModel(), ENiagaraScriptUsage::EmitterSpawnScript);
 	}
 
-	if (EmitterUpdateGroup == nullptr)
+	if (bIncludeEmitterInformation && EmitterUpdateGroup == nullptr)
 	{
 		EmitterUpdateGroup = NewObject<UNiagaraStackScriptItemGroup>(this);
 		FRequiredEntryData RequiredEntryData(GetSystemViewModel(), GetEmitterViewModel(),
@@ -123,7 +122,7 @@ void UNiagaraStackRoot::RefreshChildrenInternal(const TArray<UNiagaraStackEntry*
 		EmitterUpdateGroup->Initialize(RequiredEntryData, DisplayName, ToolTip, GetEmitterViewModel()->GetSharedScriptViewModel(), ENiagaraScriptUsage::EmitterUpdateScript);
 	}
 
-	if (ParticleSpawnGroup == nullptr)
+	if (bIncludeEmitterInformation && ParticleSpawnGroup == nullptr)
 	{
 		ParticleSpawnGroup = NewObject<UNiagaraStackScriptItemGroup>(this);
 		FRequiredEntryData RequiredEntryData(GetSystemViewModel(), GetEmitterViewModel(),
@@ -134,7 +133,7 @@ void UNiagaraStackRoot::RefreshChildrenInternal(const TArray<UNiagaraStackEntry*
 		ParticleSpawnGroup->Initialize(RequiredEntryData, DisplayName, ToolTip, GetEmitterViewModel()->GetSharedScriptViewModel(), ENiagaraScriptUsage::ParticleSpawnScript);
 	}
 
-	if (ParticleUpdateGroup == nullptr)
+	if (bIncludeEmitterInformation && ParticleUpdateGroup == nullptr)
 	{
 		ParticleUpdateGroup = NewObject<UNiagaraStackScriptItemGroup>(this);
 		FRequiredEntryData RequiredEntryData(GetSystemViewModel(), GetEmitterViewModel(),
@@ -145,7 +144,7 @@ void UNiagaraStackRoot::RefreshChildrenInternal(const TArray<UNiagaraStackEntry*
 		ParticleUpdateGroup->Initialize(RequiredEntryData, DisplayName, ToolTip, GetEmitterViewModel()->GetSharedScriptViewModel(), ENiagaraScriptUsage::ParticleUpdateScript);
 	}
 
-	if (AddEventHandlerGroup == nullptr)
+	if (bIncludeEmitterInformation && AddEventHandlerGroup == nullptr)
 	{
 		AddEventHandlerGroup = NewObject<UNiagaraStackEventHandlerGroup>(this);
 		FRequiredEntryData RequiredEntryData(GetSystemViewModel(), GetEmitterViewModel(),
@@ -155,11 +154,11 @@ void UNiagaraStackRoot::RefreshChildrenInternal(const TArray<UNiagaraStackEntry*
 		AddEventHandlerGroup->SetOnItemAdded(UNiagaraStackEventHandlerGroup::FOnItemAdded::CreateUObject(this, &UNiagaraStackRoot::EmitterEventArraysChanged));
 	}
 
-	if (RenderGroup == nullptr)
+	if (bIncludeEmitterInformation && RenderGroup == nullptr)
 	{
 		RenderGroup = NewObject<UNiagaraStackRenderItemGroup>(this);
 		FRequiredEntryData RequiredEntryData(GetSystemViewModel(), GetEmitterViewModel(),
-			FExecutionCategoryNames::Render, NAME_None,
+			FExecutionCategoryNames::Render, FExecutionSubcategoryNames::Render,
 			GetEmitterViewModel()->GetOrCreateEditorData().GetStackEditorData());
 		RenderGroup->Initialize(RequiredEntryData);
 	}
@@ -182,7 +181,7 @@ void UNiagaraStackRoot::RefreshChildrenInternal(const TArray<UNiagaraStackEntry*
 	};
 
 	// Populate new children
-	if (bShowSystemGroups)
+	if (bIncludeSystemInformation)
 	{
 		NewChildren.Add(SystemSettingsGroup);
 		NewChildren.Add(SystemSpawnGroup);
@@ -191,39 +190,42 @@ void UNiagaraStackRoot::RefreshChildrenInternal(const TArray<UNiagaraStackEntry*
 		NewChildren.Add(GetOrCreateSpacer(NAME_None, "SystemSpacer"));
 	}
 
-	NewChildren.Add(EmitterSettingsGroup);
-	NewChildren.Add(EmitterSpawnGroup);
-	NewChildren.Add(EmitterUpdateGroup);
-	NewChildren.Add(GetOrCreateSpacer(FExecutionCategoryNames::Emitter, "EmitterFooter"));
-	NewChildren.Add(GetOrCreateSpacer(NAME_None, "EmitterSpacer"));
-
-	NewChildren.Add(ParticleSpawnGroup);
-	NewChildren.Add(ParticleUpdateGroup);
-
-	for (const FNiagaraEventScriptProperties& EventScriptProperties : GetEmitterViewModel()->GetEmitter()->GetEventHandlers())
+	if (bIncludeEmitterInformation)
 	{
-		UNiagaraStackEventScriptItemGroup* EventHandlerGroup = FindCurrentChildOfTypeByPredicate<UNiagaraStackEventScriptItemGroup>(CurrentChildren,
-			[&](UNiagaraStackEventScriptItemGroup* CurrentEventHandlerGroup) { return CurrentEventHandlerGroup->GetScriptUsageId() == EventScriptProperties.Script->GetUsageId(); });
+		NewChildren.Add(EmitterSettingsGroup);
+		NewChildren.Add(EmitterSpawnGroup);
+		NewChildren.Add(EmitterUpdateGroup);
+		NewChildren.Add(GetOrCreateSpacer(FExecutionCategoryNames::Emitter, "EmitterFooter"));
+		NewChildren.Add(GetOrCreateSpacer(NAME_None, "EmitterSpacer"));
 
-		if (EventHandlerGroup == nullptr)
+		NewChildren.Add(ParticleSpawnGroup);
+		NewChildren.Add(ParticleUpdateGroup);
+
+		for (const FNiagaraEventScriptProperties& EventScriptProperties : GetEmitterViewModel()->GetEmitter()->GetEventHandlers())
 		{
-			EventHandlerGroup = NewObject<UNiagaraStackEventScriptItemGroup>(this, NAME_None, RF_Transactional);
-			FRequiredEntryData RequiredEntryData(GetSystemViewModel(), GetEmitterViewModel(),
-				FExecutionCategoryNames::Particle, FExecutionSubcategoryNames::Event,
-				GetEmitterViewModel()->GetEditorData().GetStackEditorData());
-			EventHandlerGroup->Initialize(RequiredEntryData, GetEmitterViewModel()->GetSharedScriptViewModel(), ENiagaraScriptUsage::ParticleEventScript, EventScriptProperties.Script->GetUsageId());
-			EventHandlerGroup->SetOnModifiedEventHandlers(UNiagaraStackEventScriptItemGroup::FOnModifiedEventHandlers::CreateUObject(this, &UNiagaraStackRoot::EmitterEventArraysChanged));
+			UNiagaraStackEventScriptItemGroup* EventHandlerGroup = FindCurrentChildOfTypeByPredicate<UNiagaraStackEventScriptItemGroup>(CurrentChildren,
+				[&](UNiagaraStackEventScriptItemGroup* CurrentEventHandlerGroup) { return CurrentEventHandlerGroup->GetScriptUsageId() == EventScriptProperties.Script->GetUsageId(); });
+
+			if (EventHandlerGroup == nullptr)
+			{
+				EventHandlerGroup = NewObject<UNiagaraStackEventScriptItemGroup>(this, NAME_None, RF_Transactional);
+				FRequiredEntryData RequiredEntryData(GetSystemViewModel(), GetEmitterViewModel(),
+					FExecutionCategoryNames::Particle, FExecutionSubcategoryNames::Event,
+					GetEmitterViewModel()->GetEditorData().GetStackEditorData());
+				EventHandlerGroup->Initialize(RequiredEntryData, GetEmitterViewModel()->GetSharedScriptViewModel(), ENiagaraScriptUsage::ParticleEventScript, EventScriptProperties.Script->GetUsageId());
+				EventHandlerGroup->SetOnModifiedEventHandlers(UNiagaraStackEventScriptItemGroup::FOnModifiedEventHandlers::CreateUObject(this, &UNiagaraStackRoot::EmitterEventArraysChanged));
+			}
+
+			NewChildren.Add(EventHandlerGroup);
 		}
 
-		NewChildren.Add(EventHandlerGroup);
-	}
+		NewChildren.Add(AddEventHandlerGroup);
+		NewChildren.Add(GetOrCreateSpacer(FExecutionCategoryNames::Particle, "ParticleFooter"));
+		NewChildren.Add(GetOrCreateSpacer(NAME_None, "ParticleSpacer"));
 
-	NewChildren.Add(AddEventHandlerGroup);
-	NewChildren.Add(GetOrCreateSpacer(FExecutionCategoryNames::Particle, "ParticleFooter"));
-	NewChildren.Add(GetOrCreateSpacer(NAME_None, "ParticleSpacer"));
-
-	NewChildren.Add(RenderGroup);
-	NewChildren.Add(GetOrCreateSpacer(FExecutionCategoryNames::Render, "RenderFooter"));
+		NewChildren.Add(RenderGroup);
+		NewChildren.Add(GetOrCreateSpacer(FExecutionCategoryNames::Render, "RenderFooter"));
+	}
 }
 
 void UNiagaraStackRoot::EmitterEventArraysChanged()
diff --git a/Engine/Plugins/FX/Niagara/Source/NiagaraEditor/Private/ViewModels/Stack/NiagaraStackScriptItemGroup.cpp b/Engine/Plugins/FX/Niagara/Source/NiagaraEditor/Private/ViewModels/Stack/NiagaraStackScriptItemGroup.cpp
--- a/Engine/Plugins/FX/Niagara/Source/NiagaraEditor/Private/ViewModels/Stack/NiagaraStackScriptItemGroup.cpp
+++ b/Engine/Plugins/FX/Niagara/Source/NiagaraEditor/Private/ViewModels/Stack/NiagaraStackScriptItemGroup.cpp
@@ -168,7 +168,7 @@ private:
 class FScriptItemGroupAddUtilities : public TNiagaraStackItemGroupAddUtilities<UNiagaraNodeFunctionCall*>
 {
 public:
-	FScriptItemGroupAddUtilities(TSharedRef<FNiagaraSystemViewModel> InSystemViewModel, TSharedRef<FNiagaraEmitterViewModel> InEmitterViewModel, UNiagaraStackEditorData& InStackEditorData, FOnItemAdded InOnItemAdded)
+	FScriptItemGroupAddUtilities(TSharedRef<FNiagaraSystemViewModel> InSystemViewModel, TSharedPtr<FNiagaraEmitterViewModel> InEmitterViewModel, UNiagaraStackEditorData& InStackEditorData, FOnItemAdded InOnItemAdded)
 		: TNiagaraStackItemGroupAddUtilities(LOCTEXT("ScriptGroupAddItemName", "Module"), EAddMode::AddFromAction, false, InOnItemAdded)
 		, SystemViewModel(InSystemViewModel)
 		, EmitterViewModel(InEmitterViewModel)
@@ -185,7 +185,7 @@ public:
 
 	virtual void GenerateAddActions(TArray<TSharedRef<INiagaraStackItemGroupAddAction>>& OutAddActions, const FNiagaraStackItemGroupAddOptions& AddProperties) const override
 	{
-		if (SystemViewModel.IsValid() == false || EmitterViewModel.IsValid() == false || OutputNode == nullptr)
+		if (SystemViewModel.IsValid() == false || OutputNode == nullptr)
 		{
 			return;
 		}
@@ -239,7 +239,7 @@ public:
 		}
 
 		checkf(NewModuleNode != nullptr, TEXT("Add module action failed"));
-		FNiagaraStackGraphUtilities::InitializeStackFunctionInputs(SystemViewModel.Pin().ToSharedRef(), EmitterViewModel.Pin().ToSharedRef(), StackEditorData, *NewModuleNode, *NewModuleNode);
+		FNiagaraStackGraphUtilities::InitializeStackFunctionInputs(SystemViewModel.Pin().ToSharedRef(), EmitterViewModel.Pin(), StackEditorData, *NewModuleNode, *NewModuleNode);
 		FNiagaraStackGraphUtilities::RelayoutGraph(*OutputNode->GetGraph());
 		OnItemAdded.ExecuteIfBound(NewModuleNode);
 	}
diff --git a/Engine/Plugins/FX/Niagara/Source/NiagaraEditor/Private/ViewModels/Stack/NiagaraStackSelection.cpp b/Engine/Plugins/FX/Niagara/Source/NiagaraEditor/Private/ViewModels/Stack/NiagaraStackSelection.cpp
new file mode 100644
--- /dev/null
+++ b/Engine/Plugins/FX/Niagara/Source/NiagaraEditor/Private/ViewModels/Stack/NiagaraStackSelection.cpp
@@ -0,0 +1,68 @@
+// Copyright 1998-2019 Epic Games, Inc. All Rights Reserved.
+
+#include "ViewModels/Stack/NiagaraStackSelection.h"
+#include "ViewModels/Stack/NiagaraStackEntry.h"
+#include "ViewModels/Stack/NiagaraStackSpacer.h"
+#include "ViewModels/NiagaraSystemViewModel.h"
+#include "NiagaraSystemEditorData.h"
+
+void UNiagaraStackSelection::Initialize(FRequiredEntryData InRequiredEntryData)
+{
+	Super::Initialize(InRequiredEntryData, FString());
+}
+
+bool UNiagaraStackSelection::GetCanExpand() const
+{
+	return false;
+}
+
+bool UNiagaraStackSelection::GetShouldShowInStack() const
+{
+	return false;
+}
+
+void UNiagaraStackSelection::SetSelectedEntries(const TArray<UNiagaraStackEntry*>& InSelectedEntries)
+{
+	SelectedEntries.Empty();
+	for (UNiagaraStackEntry* SelectedEntry : InSelectedEntries)
+	{
+		SelectedEntries.Add(SelectedEntry);
+	}
+	RefreshChildren();
+}
+
+void UNiagaraStackSelection::RefreshChildrenInternal(const TArray<UNiagaraStackEntry*>& CurrentChildren, TArray<UNiagaraStackEntry*>& NewChildren, TArray<FStackIssue>& NewIssues)
+{
+	auto GetOrCreateSpacer = [&](FName SpacerExecutionCategory, FName SpacerKey)
+	{
+		UNiagaraStackSpacer* Spacer = FindCurrentChildOfTypeByPredicate<UNiagaraStackSpacer>(CurrentChildren,
+			[SpacerKey](UNiagaraStackSpacer* CurrentSpacer) { return CurrentSpacer->GetSpacerKey() == SpacerKey; });
+
+		if (Spacer == nullptr)
+		{
+			Spacer = NewObject<UNiagaraStackSpacer>(this);
+			FRequiredEntryData RequiredEntryData(GetSystemViewModel(), GetEmitterViewModel(),
+				SpacerExecutionCategory, NAME_None,
+				GetSystemViewModel()->GetEditorData().GetStackEditorData());
+			Spacer->Initialize(RequiredEntryData, SpacerKey, 1.0f);
+		}
+
+		return Spacer;
+	};
+
+	int32 EntryIndex = 0;
+	for (TWeakObjectPtr<UNiagaraStackEntry> SelectedEntry : SelectedEntries)
+	{
+		if (EntryIndex > 0)
+		{
+			NewChildren.Add(GetOrCreateSpacer(NAME_None, *FString::Printf(TEXT("SelectionSpacer_%i"), EntryIndex)));
+		}
+
+		if (SelectedEntry.IsValid())
+		{
+			SelectedEntry->SetIsExpanded(true);
+			NewChildren.Add(SelectedEntry.Get());
+		}
+		EntryIndex++;
+	}
+}
\ No newline at end of file
diff --git a/Engine/Plugins/FX/Niagara/Source/NiagaraEditor/Private/ViewModels/Stack/NiagaraStackViewModel.cpp b/Engine/Plugins/FX/Niagara/Source/NiagaraEditor/Private/ViewModels/Stack/NiagaraStackViewModel.cpp
--- a/Engine/Plugins/FX/Niagara/Source/NiagaraEditor/Private/ViewModels/Stack/NiagaraStackViewModel.cpp
+++ b/Engine/Plugins/FX/Niagara/Source/NiagaraEditor/Private/ViewModels/Stack/NiagaraStackViewModel.cpp
@@ -2,6 +2,7 @@
 
 #include "ViewModels/Stack/NiagaraStackViewModel.h"
 #include "ViewModels/Stack/NiagaraStackRoot.h"
+#include "ViewModels/Stack/NiagaraStackSpacer.h"
 #include "ViewModels/NiagaraSystemViewModel.h"
 #include "ViewModels/NiagaraEmitterViewModel.h"
 #include "ViewModels/NiagaraEmitterHandleViewModel.h"
@@ -21,73 +22,140 @@
 
 #define LOCTEXT_NAMESPACE "NiagaraStackViewModel"
 const double UNiagaraStackViewModel::MaxSearchTime = .02f; // search at 50 fps
-TSharedPtr<FNiagaraSystemViewModel> UNiagaraStackViewModel::GetSystemViewModel()
-{
-	return SystemViewModel;
-}
 
-TSharedPtr<FNiagaraEmitterHandleViewModel> UNiagaraStackViewModel::GetEmitterHandleViewModel()
+bool UNiagaraStackViewModel::FTopLevelViewModel::IsValid() const
 {
-	return EmitterHandleViewModel;
+	return (SystemViewModel.IsValid() && EmitterHandleViewModel.IsValid() == false) ||
+		(SystemViewModel.IsValid() == false && EmitterHandleViewModel.IsValid());
 }
 
-void UNiagaraStackViewModel::Initialize(TSharedPtr<FNiagaraSystemViewModel> InSystemViewModel, TSharedPtr<FNiagaraEmitterHandleViewModel> InEmitterHandleViewModel)
+UNiagaraStackEditorData* UNiagaraStackViewModel::FTopLevelViewModel::GetStackEditorData() const
 {
-	if (RootEntry != nullptr)
+	if (SystemViewModel.IsValid())
 	{
-		RootEntry->OnStructureChanged().RemoveAll(this);
-		RootEntry->OnDataObjectModified().RemoveAll(this);
-		RootEntry->OnRequestFullRefresh().RemoveAll(this);
-		RootEntry->OnRequestFullRefreshDeferred().RemoveAll(this);
-		RootEntries.Empty();
-		RootEntry->Finalize();
-		RootEntry = nullptr;
-		GEditor->UnregisterForUndo(this);
+		return &SystemViewModel->GetEditorData().GetStackEditorData();
 	}
-
-	if (EmitterHandleViewModel.IsValid() && EmitterHandleViewModel->GetEmitterViewModel().IsValid())
+	else if (EmitterHandleViewModel.IsValid())
+	{
+		return &EmitterHandleViewModel->GetEmitterViewModel()->GetEditorData().GetStackEditorData();
+	}
+	else
 	{
-		EmitterHandleViewModel->GetEmitterViewModel()->OnScriptCompiled().RemoveAll(this);
+		return nullptr;
 	}
+}
 
+FText UNiagaraStackViewModel::FTopLevelViewModel::GetDisplayName() const
+{
 	if (SystemViewModel.IsValid())
 	{
-		SystemViewModel->OnSystemCompiled().RemoveAll(this);
+		return SystemViewModel->GetDisplayName();
+	}
+	else if (EmitterHandleViewModel.IsValid())
+	{
+		return EmitterHandleViewModel->GetNameText();
+	}
+	else
+	{
+		return FText();
 	}
+}
 
+void UNiagaraStackViewModel::InitializeWithViewModels(TSharedPtr<FNiagaraSystemViewModel> InSystemViewModel, TSharedPtr<FNiagaraEmitterHandleViewModel> InEmitterHandleViewModel, FNiagaraStackViewModelOptions InOptions)
+{
+	Reset();
+
+	Options = InOptions;
 	SystemViewModel = InSystemViewModel;
 	EmitterHandleViewModel = InEmitterHandleViewModel;
+
+	TSharedPtr<FNiagaraSystemViewModel> SystemViewModelPinned = SystemViewModel.Pin();
 	TSharedPtr<FNiagaraEmitterViewModel> EmitterViewModel = InEmitterHandleViewModel.IsValid() ? InEmitterHandleViewModel->GetEmitterViewModel() : TSharedPtr<FNiagaraEmitterViewModel>();
 
-	if (SystemViewModel.IsValid() && EmitterViewModel.IsValid() && EmitterViewModel->GetSharedScriptViewModel()->GetGraphViewModel()->GetGraph() != nullptr)
+	if (SystemViewModelPinned.IsValid())
 	{
 		GEditor->RegisterForUndo(this);
 
-		EmitterViewModel->OnScriptCompiled().AddUObject(this, &UNiagaraStackViewModel::OnEmitterCompiled);
-		SystemViewModel->OnSystemCompiled().AddUObject(this, &UNiagaraStackViewModel::OnSystemCompiled);
+		if (EmitterViewModel.IsValid())
+		{
+			EmitterViewModel->OnScriptCompiled().AddUObject(this, &UNiagaraStackViewModel::OnEmitterCompiled);
+			EmitterViewModel->OnParentRemoved().AddUObject(this, &UNiagaraStackViewModel::EmitterParentRemoved);
+		}
+		SystemViewModelPinned->OnSystemCompiled().AddUObject(this, &UNiagaraStackViewModel::OnSystemCompiled);
 		
-		RootEntry = NewObject<UNiagaraStackRoot>(this);
-		UNiagaraStackEntry::FRequiredEntryData RequiredEntryData(SystemViewModel.ToSharedRef(), EmitterViewModel.ToSharedRef(),
-			NAME_None, NAME_None,
-			SystemViewModel->GetOrCreateEditorData().GetStackEditorData());
-		RootEntry->Initialize(RequiredEntryData);
-		RootEntry->RefreshChildren();
-		RootEntry->OnStructureChanged().AddUObject(this, &UNiagaraStackViewModel::EntryStructureChanged);
-		RootEntry->OnDataObjectModified().AddUObject(this, &UNiagaraStackViewModel::EntryDataObjectModified);
-		RootEntry->OnRequestFullRefresh().AddUObject(this, &UNiagaraStackViewModel::EntryRequestFullRefresh);
-		RootEntry->OnRequestFullRefreshDeferred().AddUObject(this, &UNiagaraStackViewModel::EntryRequestFullRefreshDeferred);
+		UNiagaraStackRoot* StackRoot = NewObject<UNiagaraStackRoot>(this);
+		UNiagaraStackEntry::FRequiredEntryData RequiredEntryData(SystemViewModelPinned.ToSharedRef(), EmitterViewModel,
+			UNiagaraStackEntry::FExecutionCategoryNames::System, UNiagaraStackEntry::FExecutionSubcategoryNames::Settings,
+			SystemViewModelPinned->GetEditorData().GetStackEditorData());
+		StackRoot->Initialize(RequiredEntryData, Options.GetIncludeSystemInformation(), Options.GetIncludeEmitterInformation());
+		StackRoot->RefreshChildren();
+		StackRoot->OnStructureChanged().AddUObject(this, &UNiagaraStackViewModel::EntryStructureChanged);
+		StackRoot->OnDataObjectModified().AddUObject(this, &UNiagaraStackViewModel::EntryDataObjectModified);
+		StackRoot->OnRequestFullRefresh().AddUObject(this, &UNiagaraStackViewModel::EntryRequestFullRefresh);
+		StackRoot->OnRequestFullRefreshDeferred().AddUObject(this, &UNiagaraStackViewModel::EntryRequestFullRefreshDeferred);
+		RootEntry = StackRoot;
 		RootEntries.Add(RootEntry);
+
+		bExternalRootEntry = false;
 	}
 
-	CurrentFocusedSearchMatchIndex = -1;
 	StructureChangedDelegate.Broadcast();
+}
+
+void UNiagaraStackViewModel::InitializeWithRootEntry(UNiagaraStackEntry* InRootEntry)
+{
+	Reset();
+	bUsesTopLevelViewModels = true;
+
+	RootEntry = InRootEntry;
+	RootEntry->OnStructureChanged().AddUObject(this, &UNiagaraStackViewModel::EntryStructureChanged);
+	RootEntries.Add(RootEntry);
+
+	GEditor->RegisterForUndo(this);
+	StructureChangedDelegate.Broadcast();
+}
+
+void UNiagaraStackViewModel::Reset()
+{
+	if (RootEntry != nullptr)
+	{
+		RootEntry->OnStructureChanged().RemoveAll(this);
+		RootEntry->OnDataObjectModified().RemoveAll(this);
+		RootEntry->OnRequestFullRefresh().RemoveAll(this);
+		RootEntry->OnRequestFullRefreshDeferred().RemoveAll(this);
+		if (bExternalRootEntry == false)
+		{
+			RootEntry->Finalize();
+		}
+		RootEntry = nullptr;
+	}
+	RootEntries.Empty();
+
+	if (EmitterHandleViewModel.IsValid())
+	{
+		EmitterHandleViewModel.Pin()->GetEmitterViewModel()->OnScriptCompiled().RemoveAll(this);
+		EmitterHandleViewModel.Reset();
+	}
+
+	if (SystemViewModel.IsValid())
+	{
+		SystemViewModel.Pin()->OnSystemCompiled().RemoveAll(this);
+		SystemViewModel.Reset();
+	}
+
+	TopLevelViewModels.Empty();
+	
+	GEditor->UnregisterForUndo(this);
+
+	CurrentFocusedSearchMatchIndex = -1;
 	bRestartSearch = false;
 	bRefreshPending = false;
+	bUsesTopLevelViewModels = false;
 }
 
 void UNiagaraStackViewModel::Finalize()
 {
-	Initialize(nullptr, nullptr);
+	Reset();
 }
 
 void UNiagaraStackViewModel::BeginDestroy()
@@ -165,37 +233,69 @@ void UNiagaraStackViewModel::CollapseToHeaders()
 void UNiagaraStackViewModel::UndismissAllIssues()
 {
 	FScopedTransaction ScopedTransaction(LOCTEXT("UnDismissIssues", "Undismiss issues"));
-	UNiagaraStackEditorData& EmitterData = GetEmitterHandleViewModel()->GetEmitterViewModel()->GetEditorData().GetStackEditorData();
-	EmitterData.Modify();
-	EmitterData.UndismissAllIssues();
 
-	UNiagaraStackEditorData& SystemData = GetSystemViewModel()->GetEditorData().GetStackEditorData();
-	SystemData.Modify();
-	SystemData.UndismissAllIssues();
+	TArray<UNiagaraStackEditorData*> StackEditorDatas;
+	for (TSharedRef<UNiagaraStackViewModel::FTopLevelViewModel> TopLevelViewModel : TopLevelViewModels)
+	{
+		if (TopLevelViewModel->IsValid())
+		{
+			StackEditorDatas.AddUnique(TopLevelViewModel->GetStackEditorData());
+		}
+	}
+
+	for (UNiagaraStackEditorData* StackEditorData : StackEditorDatas)
+	{
+		StackEditorData->Modify();
+		StackEditorData->UndismissAllIssues();
+	}
 	
 	RootEntry->RefreshChildren();
 }
 
 bool UNiagaraStackViewModel::HasDismissedStackIssues()
 {
-	return GetSystemViewModel()->GetEditorData().GetStackEditorData().GetDismissedStackIssueIds().Num() > 0
-		|| GetEmitterHandleViewModel()->GetEmitterViewModel()->GetEditorData().GetStackEditorData().GetDismissedStackIssueIds().Num() > 0;
+	for (TSharedRef<FTopLevelViewModel> TopLevelViewModel : TopLevelViewModels)
+	{
+		if (TopLevelViewModel->IsValid() && TopLevelViewModel->GetStackEditorData()->GetDismissedStackIssueIds().Num() > 0)
+		{
+			return true;
+		}
+	}
+	return false;
 }
 
-bool UNiagaraStackViewModel::HasParentEmitter() const
+const TArray<TSharedRef<UNiagaraStackViewModel::FTopLevelViewModel>>& UNiagaraStackViewModel::GetTopLevelViewModels() const
 {
-	return EmitterHandleViewModel->GetEmitterViewModel()->GetParentEmitter() != nullptr;
+	return TopLevelViewModels;
 }
 
-void UNiagaraStackViewModel::RemoveEmitterSource()
+TSharedPtr<UNiagaraStackViewModel::FTopLevelViewModel> UNiagaraStackViewModel::GetTopLevelViewModelForEntry(UNiagaraStackEntry& InEntry) const
 {
+	if (InEntry.GetEmitterViewModel().IsValid())
 	{
-		FScopedTransaction ScopedTransaction(LOCTEXT("RemoveParentEmitterTransaction", "Remove Parent Emitter"));
-		UNiagaraEmitter* Emitter = EmitterHandleViewModel->GetEmitterViewModel()->GetEmitter();
-		Emitter->Modify();
-		Emitter->RemoveParent();
+		const TSharedRef<FTopLevelViewModel>* MatchingTopLevelViewModel = TopLevelViewModels.FindByPredicate([&InEntry](const TSharedRef<FTopLevelViewModel>& TopLevelViewModel) 
+		{ 
+			return TopLevelViewModel->EmitterHandleViewModel.IsValid() && TopLevelViewModel->EmitterHandleViewModel->GetEmitterViewModel() == InEntry.GetEmitterViewModel(); 
+		});
+
+		if (MatchingTopLevelViewModel != nullptr)
+		{
+			return *MatchingTopLevelViewModel;
+		}
 	}
-	RootEntry->RefreshChildren();
+	else
+	{
+		const TSharedRef<FTopLevelViewModel>* MatchingTopLevelViewModel = TopLevelViewModels.FindByPredicate([&InEntry](const TSharedRef<FTopLevelViewModel>& TopLevelViewModel) 
+		{ 
+			return TopLevelViewModel->SystemViewModel == InEntry.GetSystemViewModel(); 
+		});
+
+		if (MatchingTopLevelViewModel != nullptr)
+		{
+			return *MatchingTopLevelViewModel;
+		}
+	}
+	return TSharedPtr<FTopLevelViewModel>();
 }
 
 void UNiagaraStackViewModel::CollapseToHeadersRecursive(TArray<UNiagaraStackEntry*> Entries)
@@ -238,6 +338,11 @@ void UNiagaraStackViewModel::OnEmitterCompiled()
 	OnSearchTextChanged(CurrentSearchText);
 }
 
+void UNiagaraStackViewModel::EmitterParentRemoved()
+{
+	RootEntry->RefreshChildren();
+}
+
 void UNiagaraStackViewModel::SearchTick()
 {
 	// perform partial searches here, by processing a fixed number of entries (maybe more than one?)
@@ -265,7 +370,6 @@ void UNiagaraStackViewModel::SearchTick()
 
 	if (IsSearching())
 	{
-		UNiagaraStackEditorData& EditorData = GetSystemViewModel()->GetEditorData().GetStackEditorData();
 		double SearchStartTime = FPlatformTime::Seconds();
 		double CurrentSearchLoopTime = SearchStartTime;
 		// process at least one item, but don't go over MaxSearchTime for the rest
@@ -280,9 +384,9 @@ void UNiagaraStackViewModel::SearchTick()
 				TSet<FName> MatchedKeys;
 				for (UNiagaraStackEntry::FStackSearchItem SearchItem : SearchItems)
 				{
-					if (!ItemsToSearch[0].GetEntry()->GetStackEditorDataKey().IsEmpty())
+					if (!EntryToProcess->GetStackEditorDataKey().IsEmpty())
 					{
-						EditorData.SetStackEntryWasExpandedPreSearch(ItemsToSearch[0].GetEntry()->GetStackEditorDataKey(), ItemsToSearch[0].GetEntry()->GetIsExpanded());
+						EntryToProcess->GetStackEditorData().SetStackEntryWasExpandedPreSearch(EntryToProcess->GetStackEditorDataKey(), EntryToProcess->GetIsExpanded());
 					}
 
 					if (ItemMatchesSearchCriteria(SearchItem)) 
@@ -347,7 +451,6 @@ void UNiagaraStackViewModel::GeneratePathForEntry(UNiagaraStackEntry* Root, UNia
 
 void UNiagaraStackViewModel::RestoreStackEntryExpansionPreSearch()
 {
-	UNiagaraStackEditorData& StackEditorData = SystemViewModel->GetEditorData().GetStackEditorData();
 	GenerateTraversalEntries(RootEntry, TArray<UNiagaraStackEntry*>(), ItemsToRestoreExpansionState);
 
 	while (ItemsToRestoreExpansionState.Num() != 0)
@@ -355,10 +458,11 @@ void UNiagaraStackViewModel::RestoreStackEntryExpansionPreSearch()
 		UNiagaraStackEntry* EntryToProcess = ItemsToRestoreExpansionState[0].GetEntry();
 		if (EntryToProcess->GetIsSearchResult() && !EntryToProcess->IsA<UNiagaraStackRoot>())
 		{
-			for (auto EntryToExpand : ItemsToRestoreExpansionState[0].EntryPath)
+			for (UNiagaraStackEntry* EntryToExpand : ItemsToRestoreExpansionState[0].EntryPath)
 			{
 				if (!EntryToExpand->IsA<UNiagaraStackRoot>())
 				{
+					UNiagaraStackEditorData& StackEditorData = EntryToExpand->GetStackEditorData();
 					EntryToExpand->SetIsExpanded(StackEditorData.GetStackEntryWasExpandedPreSearch(EntryToExpand->GetStackEditorDataKey(), false));
 				}
 			}
@@ -384,85 +488,100 @@ UNiagaraStackViewModel::FOnSearchCompleted& UNiagaraStackViewModel::OnSearchComp
 
 bool UNiagaraStackViewModel::GetShowAllAdvanced() const
 {
-	if (SystemViewModel.IsValid() && EmitterHandleViewModel.IsValid())
+	for (TSharedRef<FTopLevelViewModel> TopLevelViewModel : TopLevelViewModels)
 	{
-		return SystemViewModel->GetEditorData().GetStackEditorData().GetShowAllAdvanced() ||
-			EmitterHandleViewModel->GetEmitterViewModel()->GetEditorData().GetStackEditorData().GetShowAllAdvanced();
+		if (TopLevelViewModel->IsValid() && TopLevelViewModel->GetStackEditorData()->GetShowAllAdvanced())
+		{
+			return true;
+		}
 	}
 	return false;
 }
 
 void UNiagaraStackViewModel::SetShowAllAdvanced(bool bInShowAllAdvanced)
 {
-	if (SystemViewModel.IsValid() && EmitterHandleViewModel.IsValid())
+	for (TSharedRef<FTopLevelViewModel> TopLevelViewModel : TopLevelViewModels)
 	{
-		SystemViewModel->GetOrCreateEditorData().GetStackEditorData().SetShowAllAdvanced(bInShowAllAdvanced);
-		EmitterHandleViewModel->GetEmitterViewModel()->GetOrCreateEditorData().GetStackEditorData().SetShowAllAdvanced(bInShowAllAdvanced);
-		OnSearchTextChanged(CurrentSearchText);
-		StructureChangedDelegate.Broadcast();
+		if (TopLevelViewModel->IsValid()) 
+		{
+			TopLevelViewModel->GetStackEditorData()->SetShowAllAdvanced(bInShowAllAdvanced);
+		}
 	}
+
+	OnSearchTextChanged(CurrentSearchText);
+	StructureChangedDelegate.Broadcast();
 }
 
 bool UNiagaraStackViewModel::GetShowOutputs() const
 {
-	if (SystemViewModel.IsValid() && EmitterHandleViewModel.IsValid())
+	for (TSharedRef<FTopLevelViewModel> TopLevelViewModel : TopLevelViewModels)
 	{
-		return SystemViewModel->GetEditorData().GetStackEditorData().GetShowOutputs() ||
-			EmitterHandleViewModel->GetEmitterViewModel()->GetEditorData().GetStackEditorData().GetShowOutputs();
+		if (TopLevelViewModel->IsValid() && TopLevelViewModel->GetStackEditorData()->GetShowOutputs())
+		{
+			return true;
+		}
 	}
 	return false;
 }
 
 void UNiagaraStackViewModel::SetShowOutputs(bool bInShowOutputs)
 {
-	if (SystemViewModel.IsValid() && EmitterHandleViewModel.IsValid())
+	for (TSharedRef<FTopLevelViewModel> TopLevelViewModel : TopLevelViewModels)
 	{
-		SystemViewModel->GetOrCreateEditorData().GetStackEditorData().SetShowOutputs(bInShowOutputs);
-		EmitterHandleViewModel->GetEmitterViewModel()->GetOrCreateEditorData().GetStackEditorData().SetShowOutputs(bInShowOutputs);
-		OnSearchTextChanged(CurrentSearchText);
-
-		// Showing outputs changes indenting so a full refresh is needed.
-		RootEntry->RefreshChildren();
+		if (TopLevelViewModel->IsValid())
+		{
+			TopLevelViewModel->GetStackEditorData()->SetShowOutputs(bInShowOutputs);
+		}
 	}
+	
+	// Showing outputs changes indenting so a full refresh is needed.
+	OnSearchTextChanged(CurrentSearchText);
+	RootEntry->RefreshChildren();
 }
 
 bool UNiagaraStackViewModel::GetShowLinkedInputs() const
 {
-	if (SystemViewModel.IsValid() && EmitterHandleViewModel.IsValid())
+	for (TSharedRef<FTopLevelViewModel> TopLevelViewModel : TopLevelViewModels)
 	{
-		return SystemViewModel->GetEditorData().GetStackEditorData().GetShowLinkedInputs() ||
-			EmitterHandleViewModel->GetEmitterViewModel()->GetEditorData().GetStackEditorData().GetShowLinkedInputs();
+		if (TopLevelViewModel->IsValid() && TopLevelViewModel->GetStackEditorData()->GetShowLinkedInputs())
+		{
+			return true;
+		}
 	}
 	return false;
 }
 
 void UNiagaraStackViewModel::SetShowLinkedInputs(bool bInShowLinkedInputs)
 {
-	if (SystemViewModel.IsValid() && EmitterHandleViewModel.IsValid())
+	for (TSharedRef<FTopLevelViewModel> TopLevelViewModel : TopLevelViewModels)
 	{
-		SystemViewModel->GetOrCreateEditorData().GetStackEditorData().SetShowLinkedInputs(bInShowLinkedInputs);
-		EmitterHandleViewModel->GetEmitterViewModel()->GetOrCreateEditorData().GetStackEditorData().SetShowLinkedInputs(bInShowLinkedInputs);
-		OnSearchTextChanged(CurrentSearchText);
-
-		// Showing linked inputs changes indenting so a full refresh is needed.
-		RootEntry->RefreshChildren();
+		if (TopLevelViewModel->IsValid())
+		{
+			TopLevelViewModel->GetStackEditorData()->SetShowLinkedInputs(bInShowLinkedInputs);
+		}
 	}
+
+	// Showing linked inputs changes indenting so a full refresh is needed.
+	OnSearchTextChanged(CurrentSearchText);
+	RootEntry->RefreshChildren();
 }
 
 double UNiagaraStackViewModel::GetLastScrollPosition() const
 {
+	// TODO: Fix this with the new overview paradigm.
 	if (EmitterHandleViewModel.IsValid())
 	{
-		return EmitterHandleViewModel->GetEmitterViewModel()->GetEditorData().GetStackEditorData().GetLastScrollPosition();
+		return EmitterHandleViewModel.Pin()->GetEmitterViewModel()->GetEditorData().GetStackEditorData().GetLastScrollPosition();
 	}
 	return 0;
 }
 
 void UNiagaraStackViewModel::SetLastScrollPosition(double InLastScrollPosition)
 {
+	// TODO: Fix this with the new overview paradigm.
 	if (EmitterHandleViewModel.IsValid())
 	{
-		EmitterHandleViewModel->GetEmitterViewModel()->GetOrCreateEditorData().GetStackEditorData().SetLastScrollPosition(InLastScrollPosition);
+		EmitterHandleViewModel.Pin()->GetEmitterViewModel()->GetOrCreateEditorData().GetStackEditorData().SetLastScrollPosition(InLastScrollPosition);
 	}
 }
 
@@ -479,13 +598,20 @@ void UNiagaraStackViewModel::PostUndo(bool bSuccess)
 
 void UNiagaraStackViewModel::EntryStructureChanged()
 {
+	if (bUsesTopLevelViewModels)
+	{
+		RefreshTopLevelViewModels();
+	}
 	StructureChangedDelegate.Broadcast();
 	OnSearchTextChanged(CurrentSearchText);
 }
 
 void UNiagaraStackViewModel::EntryDataObjectModified(UObject* ChangedObject)
 {
-	SystemViewModel->NotifyDataObjectChanged(ChangedObject);
+	if (SystemViewModel.IsValid())
+	{
+		SystemViewModel.Pin()->NotifyDataObjectChanged(ChangedObject);
+	}
 	OnSearchTextChanged(CurrentSearchText);
 }
 
@@ -500,4 +626,53 @@ void UNiagaraStackViewModel::EntryRequestFullRefreshDeferred()
 	bRefreshPending = true;
 }
 
+void UNiagaraStackViewModel::RefreshTopLevelViewModels()
+{
+	TArray<TSharedRef<FTopLevelViewModel>> CurrentTopLevelViewModels = TopLevelViewModels;
+	TopLevelViewModels.Empty();
+
+	TArray<UNiagaraStackEntry*> RootChildren;
+	RootEntry->GetUnfilteredChildren(RootChildren);
+	for (UNiagaraStackEntry* RootChild : RootChildren)
+	{
+		if (RootChild->IsA<UNiagaraStackSpacer>())
+		{
+			// Skip spacers they're not a good indication of top level object.  This will be unnecessary when spacers are removed.
+			continue;
+		}
+
+		TSharedPtr<FTopLevelViewModel> TopLevelViewModel;
+		if (RootChild->GetEmitterViewModel().IsValid())
+		{
+			TSharedPtr<FNiagaraEmitterHandleViewModel> RootChildEmitterHandleViewModel = RootChild->GetSystemViewModel()->GetEmitterHandleViewModelForEmitter(RootChild->GetEmitterViewModel()->GetEmitter());
+			TSharedRef<FTopLevelViewModel>* CurrentTopLevelViewModelPtr = CurrentTopLevelViewModels.FindByPredicate(
+				[&RootChildEmitterHandleViewModel](const TSharedRef<FTopLevelViewModel>& TopLevelViewModel) { return TopLevelViewModel->EmitterHandleViewModel == RootChildEmitterHandleViewModel; });
+			if (CurrentTopLevelViewModelPtr != nullptr)
+			{
+				if (TopLevelViewModels.Contains(*CurrentTopLevelViewModelPtr) == false)
+				{
+					TopLevelViewModels.Add(*CurrentTopLevelViewModelPtr);
+				}
+			}
+			else
+			{
+				TopLevelViewModels.Add(MakeShared<FTopLevelViewModel>(RootChildEmitterHandleViewModel));
+			}
+		}
+		else
+		{
+			TSharedRef<FTopLevelViewModel>* CurrentTopLevelViewModelPtr = CurrentTopLevelViewModels.FindByPredicate(
+				[RootChild](const TSharedRef<FTopLevelViewModel>& TopLevelViewModel) { return TopLevelViewModel->SystemViewModel == RootChild->GetSystemViewModel(); });
+			if (CurrentTopLevelViewModelPtr != nullptr && TopLevelViewModels.Contains(*CurrentTopLevelViewModelPtr) == false)
+			{
+				TopLevelViewModels.Add(*CurrentTopLevelViewModelPtr);
+			}
+			else
+			{
+				TopLevelViewModels.Add(MakeShared<FTopLevelViewModel>(RootChild->GetSystemViewModel()));
+			}
+		}
+	}
+}
+
 #undef LOCTEXT_NAMESPACE
diff --git a/Engine/Plugins/FX/Niagara/Source/NiagaraEditor/Private/Widgets/SNiagaraGeneratedCodeView.cpp b/Engine/Plugins/FX/Niagara/Source/NiagaraEditor/Private/Widgets/SNiagaraGeneratedCodeView.cpp
--- a/Engine/Plugins/FX/Niagara/Source/NiagaraEditor/Private/Widgets/SNiagaraGeneratedCodeView.cpp
+++ b/Engine/Plugins/FX/Niagara/Source/NiagaraEditor/Private/Widgets/SNiagaraGeneratedCodeView.cpp
@@ -33,7 +33,7 @@ void SNiagaraGeneratedCodeView::Construct(const FArguments& InArgs, TSharedRef<F
 	ensure(ScriptEnum);
 
 	SystemViewModel = InSystemViewModel;
-	SystemViewModel->OnSelectedEmitterHandlesChanged().AddRaw(this, &SNiagaraGeneratedCodeView::SelectedEmitterHandlesChanged);
+	SystemViewModel->GetSelectionViewModel()->OnSelectionChanged().AddSP(this, &SNiagaraGeneratedCodeView::SystemSelectionChanged);
 	SystemViewModel->GetSystemScriptViewModel()->OnSystemCompiled().AddRaw(this, &SNiagaraGeneratedCodeView::OnCodeCompiled);
 
 	TSharedRef<SWidget> HeaderContentsFirstLine = SNew(SHorizontalBox)
@@ -329,7 +329,7 @@ void SNiagaraGeneratedCodeView::OnCodeCompiled()
 	DoSearch(SearchBox->GetText());
 }
 
-void SNiagaraGeneratedCodeView::SelectedEmitterHandlesChanged()
+void SNiagaraGeneratedCodeView::SystemSelectionChanged(UNiagaraSystemSelectionViewModel::ESelectionChangeSource SelectionChangeSource)
 {
 	UpdateUI();
 	DoSearch(SearchBox->GetText());
@@ -343,11 +343,11 @@ void SNiagaraGeneratedCodeView::UpdateUI()
 	Scripts.Add(System.GetSystemSpawnScript());
 	Scripts.Add(System.GetSystemUpdateScript());
 
-	TArray<TSharedRef<FNiagaraEmitterHandleViewModel>> SelectedEmitterHandles;
-	SystemViewModel->GetSelectedEmitterHandles(SelectedEmitterHandles);
-	if (SelectedEmitterHandles.Num() == 1)
+	TArray<FGuid> SelectedEmitterHandleIds = SystemViewModel->GetSelectionViewModel()->GetSelectedEmitterHandleIds();
+	if(SelectedEmitterHandleIds.Num() == 1)
 	{
-		FNiagaraEmitterHandle* Handle = SelectedEmitterHandles[0]->GetEmitterHandle();
+		TSharedPtr<FNiagaraEmitterHandleViewModel> SelectedEmitterHandleViewModel = SystemViewModel->GetEmitterHandleViewModelById(SelectedEmitterHandleIds[0]);
+		FNiagaraEmitterHandle* Handle = SelectedEmitterHandleViewModel.IsValid() ? SelectedEmitterHandleViewModel->GetEmitterHandle() : nullptr;
 		if (Handle)
 		{
 			TArray<UNiagaraScript*> EmitterScripts;
@@ -531,7 +531,7 @@ SNiagaraGeneratedCodeView::~SNiagaraGeneratedCodeView()
 {
 	if (SystemViewModel.IsValid())
 	{
-		SystemViewModel->OnSelectedEmitterHandlesChanged().RemoveAll(this);
+		SystemViewModel->GetSelectionViewModel()->OnSelectionChanged().RemoveAll(this);
 		if (SystemViewModel->GetSystemScriptViewModel().IsValid())
 		{
 			SystemViewModel->GetSystemScriptViewModel()->OnSystemCompiled().RemoveAll(this);
