Niagara - System overview MVP

diff --git a/Engine/Plugins/FX/Niagara/Source/NiagaraEditor/Private/Widgets/SNiagaraGeneratedCodeView.h b/Engine/Plugins/FX/Niagara/Source/NiagaraEditor/Private/Widgets/SNiagaraGeneratedCodeView.h
--- a/Engine/Plugins/FX/Niagara/Source/NiagaraEditor/Private/Widgets/SNiagaraGeneratedCodeView.h
+++ b/Engine/Plugins/FX/Niagara/Source/NiagaraEditor/Private/Widgets/SNiagaraGeneratedCodeView.h
@@ -9,6 +9,7 @@
 #include "Widgets/Views/STableViewBase.h"
 #include "Widgets/Views/STreeView.h"
 #include "ViewModels/NiagaraSystemViewModel.h"
+#include "ViewModels/NiagaraSystemSelectionViewModel.h"
 #include "TickableEditorObject.h"
 #include "UObject/WeakObjectPtrTemplates.h"
 #include "NiagaraDataSet.h"
@@ -61,7 +62,7 @@ protected:
 
 	void SetSearchMofN();
 
-	void SelectedEmitterHandlesChanged();
+	void SystemSelectionChanged(UNiagaraSystemSelectionViewModel::ESelectionChangeSource SelectionChangeSource);
 	void OnTabChanged(uint32 Tab);
 	bool GetTabCheckedState(uint32 Tab) const;
 	EVisibility GetViewVisibility(uint32 Tab) const;
diff --git a/Engine/Plugins/FX/Niagara/Source/NiagaraEditor/Private/Widgets/SNiagaraSpreadsheetView.cpp b/Engine/Plugins/FX/Niagara/Source/NiagaraEditor/Private/Widgets/SNiagaraSpreadsheetView.cpp
--- a/Engine/Plugins/FX/Niagara/Source/NiagaraEditor/Private/Widgets/SNiagaraSpreadsheetView.cpp
+++ b/Engine/Plugins/FX/Niagara/Source/NiagaraEditor/Private/Widgets/SNiagaraSpreadsheetView.cpp
@@ -224,7 +224,7 @@ void SNiagaraSpreadsheetView::Construct(const FArguments& InArgs, TSharedRef<FNi
 	CaptureData.SetNum(UIMax);
 
 	SystemViewModel = InSystemViewModel;
-	SystemViewModel->OnSelectedEmitterHandlesChanged().AddRaw(this, &SNiagaraSpreadsheetView::SelectedEmitterHandlesChanged);
+	SystemViewModel->GetSelectionViewModel()->OnSelectionChanged().AddSP(this, &SNiagaraSpreadsheetView::SystemSelectionChanged);
 	SystemViewModel->OnPostSequencerTimeChanged().AddRaw(this, &SNiagaraSpreadsheetView::OnSequencerTimeChanged);
 
 	bInitialColumns = true;
@@ -618,7 +618,7 @@ SNiagaraSpreadsheetView::~SNiagaraSpreadsheetView()
 {
 	if (SystemViewModel.IsValid())
 	{
-		SystemViewModel->OnSelectedEmitterHandlesChanged().RemoveAll(this);
+		SystemViewModel->GetSelectionViewModel()->OnSelectionChanged().RemoveAll(this);
 		SystemViewModel->OnPostSequencerTimeChanged().RemoveAll(this);
 	}
 }
@@ -786,12 +786,12 @@ TSharedRef< ITableRow > SNiagaraSpreadsheetView::OnGenerateWidgetForList(TShared
 
 FText SNiagaraSpreadsheetView::LastCapturedInfoText() const
 {
-	TArray<TSharedRef<FNiagaraEmitterHandleViewModel>> SelectedEmitterHandles;
-	SystemViewModel->GetSelectedEmitterHandles(SelectedEmitterHandles);
-	if (SelectedEmitterHandles.Num() == 1 && SNiagaraSpreadsheetView::IsPausedAtRightTimeOnRightHandle() && CaptureData[(int32)TabState].DataSet.GetCurrentData())
+	TArray<FGuid> SelectedEmitterHandleIds = SystemViewModel->GetSelectionViewModel()->GetSelectedEmitterHandleIds();
+	if (SelectedEmitterHandleIds.Num() == 1 && SNiagaraSpreadsheetView::IsPausedAtRightTimeOnRightHandle() && CaptureData[(int32)TabState].DataSet.GetCurrentData())
 	{
+		TSharedPtr<FNiagaraEmitterHandleViewModel> SelectedEmitterHandle = SystemViewModel->GetEmitterHandleViewModelById(SelectedEmitterHandleIds[0]);
 		return FText::Format(LOCTEXT("LastCapturedInfoName", "Captured Emitter: \"{0}\"     # Particles: {1}    Script Type: {2}"),
-			SelectedEmitterHandles[0]->GetNameText(),
+			SelectedEmitterHandle->GetNameText(),
 			FText::AsNumber(CaptureData[(int32)TabState].DataSet.GetCurrentDataChecked().GetNumInstances()),
 			ScriptEnum->GetDisplayNameTextByValue((int64)CaptureData[(int32)TabState].TargetUsage));
 	}
@@ -894,7 +894,7 @@ bool SNiagaraSpreadsheetView::IsOutputAttributeEnabled(EUITab Tab, FName Item)
 	return CaptureData[(int32)Tab].FilteredOutputFields.Find(Item) != INDEX_NONE;
 }
 
-void SNiagaraSpreadsheetView::SelectedEmitterHandlesChanged()
+void SNiagaraSpreadsheetView::SystemSelectionChanged(UNiagaraSystemSelectionViewModel::ESelectionChangeSource SelectionChangeSource)
 {
 	// Need to reset the attributes list...
 	for (int32 i = 0; i < (int32)UIMax; i++)
@@ -1000,14 +1000,14 @@ void SNiagaraSpreadsheetView::HandleTimeChange()
 		{
 			for (int32 i = 0; i < (int32)UIMax; i++)
 			{
-				TArray<TSharedRef<FNiagaraEmitterHandleViewModel>> SelectedEmitterHandles;
-				SystemViewModel->GetSelectedEmitterHandles(SelectedEmitterHandles);
-				if (SelectedEmitterHandles.Num() == 1)
+				TArray<FGuid> SelectedEmitterHandleIds = SystemViewModel->GetSelectionViewModel()->GetSelectedEmitterHandleIds();
+				if (SelectedEmitterHandleIds.Num() == 1)
 				{
+					TSharedPtr<FNiagaraEmitterHandleViewModel> SelectedEmitterHandle = SystemViewModel->GetEmitterHandleViewModelById(SelectedEmitterHandleIds[0]);
 					FName EntryName = NAME_None;
 					if (i != UISystemUpdate)
 					{
-						EntryName = SelectedEmitterHandles[0]->GetEmitterHandle()->GetIdName();
+						EntryName = SelectedEmitterHandle->GetEmitterHandle()->GetIdName();
 					}
 					else //if (i == UISystemUpdate)
 					{
@@ -1024,7 +1024,7 @@ void SNiagaraSpreadsheetView::HandleTimeChange()
 						CaptureData[i].CaptureData = *FoundEntry;
 						CaptureData[i].CaptureData->Frame.CopyTo(CaptureData[i].DataSet);
 						CaptureData[i].InputParams = CaptureData[i].CaptureData->Parameters;
-						CaptureData[i].LastCaptureHandleId = SelectedEmitterHandles[0]->GetId();
+						CaptureData[i].LastCaptureHandleId = SelectedEmitterHandle->GetId();
 
 						ResetColumns((EUITab)i);
 						ResetEntries((EUITab)i);
@@ -1057,11 +1057,11 @@ TStatId SNiagaraSpreadsheetView::GetStatId() const
 
 bool SNiagaraSpreadsheetView::CanCapture() const
 {
-	TArray<TSharedRef<FNiagaraEmitterHandleViewModel>> SelectedEmitterHandles;
-	SystemViewModel->GetSelectedEmitterHandles(SelectedEmitterHandles);
-	if (SelectedEmitterHandles.Num() == 1)
+	TArray<FGuid> SelectedEmitterHandleIds = SystemViewModel->GetSelectionViewModel()->GetSelectedEmitterHandleIds();
+	if (SelectedEmitterHandleIds.Num() == 1)
 	{
-		FNiagaraEmitterHandle* Handle = SelectedEmitterHandles[0]->GetEmitterHandle();
+		TSharedPtr<FNiagaraEmitterHandleViewModel> SelectedEmitterHandle = SystemViewModel->GetEmitterHandleViewModelById(SelectedEmitterHandleIds[0]);
+		FNiagaraEmitterHandle* Handle = SelectedEmitterHandle->GetEmitterHandle();
 		if (Handle )
 		{
 			return true;
@@ -1073,13 +1073,13 @@ bool SNiagaraSpreadsheetView::CanCapture() const
 
 bool SNiagaraSpreadsheetView::IsPausedAtRightTimeOnRightHandle() const
 {
-	TArray<TSharedRef<FNiagaraEmitterHandleViewModel>> SelectedEmitterHandles;
-	SystemViewModel->GetSelectedEmitterHandles(SelectedEmitterHandles);
-	if (SelectedEmitterHandles.Num() == 1)
+	TArray<FGuid> SelectedEmitterHandleIds = SystemViewModel->GetSelectionViewModel()->GetSelectedEmitterHandleIds();
+	if (SelectedEmitterHandleIds.Num() == 1)
 	{
+		TSharedPtr<FNiagaraEmitterHandleViewModel> SelectedEmitterHandle = SystemViewModel->GetEmitterHandleViewModelById(SelectedEmitterHandleIds[0]);
 		return SystemViewModel->GetSequencer()->GetPlaybackStatus() == EMovieScenePlayerStatus::Stopped &&
 			CaptureData[(int32)TabState].CaptureData.IsValid() &&
-			CaptureData[(int32)TabState].LastCaptureHandleId == SelectedEmitterHandles[0]->GetId();
+			CaptureData[(int32)TabState].LastCaptureHandleId == SelectedEmitterHandle->GetId();
 	}
 	return false;
 }
@@ -1402,16 +1402,17 @@ FReply SNiagaraSpreadsheetView::OnCaptureRequestPressed()
 	float SimulationStep = SystemViewModel->GetPreviewComponent()->GetSeekDelta();
 	float TargetCaptureTime = LocalTime + SimulationStep;
 
-	TArray<TSharedRef<FNiagaraEmitterHandleViewModel>> SelectedEmitterHandles;
-	SystemViewModel->GetSelectedEmitterHandles(SelectedEmitterHandles);
-	ensure(SelectedEmitterHandles.Num() == 1);
+	TArray<FGuid> SelectedEmitterHandleIds = SystemViewModel->GetSelectionViewModel()->GetSelectedEmitterHandleIds();
+	ensure(SelectedEmitterHandleIds.Num() == 1);
 
-	if (TargetComponent.IsValid() && TargetComponent->GetSystemInstance())
+
+	if (TargetComponent.IsValid() && TargetComponent->GetSystemInstance() && SelectedEmitterHandleIds.Num() == 1)
 	{
 		TargetRequestId = FGuid::NewGuid();
 		TargetComponent->GetSystemInstance()->RequestCapture(TargetRequestId);
 
-		UNiagaraEmitter* Emitter = SelectedEmitterHandles[0]->GetEmitterHandle()->GetInstance();
+		TSharedPtr<FNiagaraEmitterHandleViewModel> SelectedEmitterHandle = SystemViewModel->GetEmitterHandleViewModelById(SelectedEmitterHandleIds[0]);
+		UNiagaraEmitter* Emitter = SelectedEmitterHandle->GetEmitterHandle()->GetInstance();
 		
 		for (int32 i = 0; i < CaptureData.Num(); i++)
 		{
diff --git a/Engine/Plugins/FX/Niagara/Source/NiagaraEditor/Private/Widgets/SNiagaraSpreadsheetView.h b/Engine/Plugins/FX/Niagara/Source/NiagaraEditor/Private/Widgets/SNiagaraSpreadsheetView.h
--- a/Engine/Plugins/FX/Niagara/Source/NiagaraEditor/Private/Widgets/SNiagaraSpreadsheetView.h
+++ b/Engine/Plugins/FX/Niagara/Source/NiagaraEditor/Private/Widgets/SNiagaraSpreadsheetView.h
@@ -10,6 +10,7 @@
 #include "Widgets/Views/STreeView.h"
 #include "Widgets/Input/SComboBox.h"
 #include "ViewModels/NiagaraSystemViewModel.h"
+#include "ViewModels/NiagaraSystemSelectionViewModel.h"
 #include "TickableEditorObject.h"
 #include "UObject/WeakObjectPtrTemplates.h"
 #include "NiagaraDataSet.h"
@@ -132,7 +133,7 @@ protected:
 	bool CanCapture() const;
 	bool IsPausedAtRightTimeOnRightHandle() const;
 
-	void SelectedEmitterHandlesChanged();
+	void SystemSelectionChanged(UNiagaraSystemSelectionViewModel::ESelectionChangeSource SelectionChangeSource);
 
 	void ResetColumns(EUITab Tab);
 	void ResetEntries(EUITab Tab);
diff --git a/Engine/Plugins/FX/Niagara/Source/NiagaraEditor/Private/Widgets/SNiagaraSystemScript.cpp b/Engine/Plugins/FX/Niagara/Source/NiagaraEditor/Private/Widgets/SNiagaraSystemScript.cpp
--- a/Engine/Plugins/FX/Niagara/Source/NiagaraEditor/Private/Widgets/SNiagaraSystemScript.cpp
+++ b/Engine/Plugins/FX/Niagara/Source/NiagaraEditor/Private/Widgets/SNiagaraSystemScript.cpp
@@ -15,16 +15,6 @@ void SNiagaraSystemScript::Construct(const FArguments& InArgs, TSharedRef<FNiaga
 	SystemViewModel = InSystemViewModel;
 	ChildSlot
 	[
-		SNew(SSplitter)
-		+ SSplitter::Slot()
-		.Value(0.3f)
-		[
-			SNew(SNiagaraParameterCollection, SystemViewModel->GetSystemScriptViewModel()->GetInputCollectionViewModel())
-		]
-		+ SSplitter::Slot()
-		.Value(0.7f)
-		[
-			SNew(SNiagaraScriptGraph, SystemViewModel->GetSystemScriptViewModel()->GetGraphViewModel())
-		]
+		SNew(SNiagaraScriptGraph, SystemViewModel->GetSystemScriptViewModel()->GetGraphViewModel())
 	];
 }
diff --git a/Engine/Plugins/FX/Niagara/Source/NiagaraEditor/Public/NiagaraEditorModule.h b/Engine/Plugins/FX/Niagara/Source/NiagaraEditor/Public/NiagaraEditorModule.h
--- a/Engine/Plugins/FX/Niagara/Source/NiagaraEditor/Public/NiagaraEditorModule.h
+++ b/Engine/Plugins/FX/Niagara/Source/NiagaraEditor/Public/NiagaraEditorModule.h
@@ -21,6 +21,7 @@ class FNiagaraCompileOptions;
 class FNiagaraCompileRequestDataBase;
 class UMovieSceneNiagaraParameterTrack;
 struct IConsoleCommand;
+class INiagaraEditorOnlyDataUtilities;
 
 DECLARE_STATS_GROUP(TEXT("Niagara Editor"), STATGROUP_NiagaraEditor, STATCAT_Advanced);
 
@@ -94,6 +95,8 @@ public:
 
 	FOnCheckScriptToolkitsShouldFocusGraphElement& GetOnScriptToolkitsShouldFocusGraphElement() { return OnCheckScriptToolkitsShouldFocusGraphElement; };
 
+	NIAGARAEDITOR_API TSharedPtr<FNiagaraSystemViewModel> GetExistingViewModelForSystem(UNiagaraSystem* InSystem);
+
 private:
 	void RegisterAssetTypeAction(IAssetTools& AssetTools, TSharedRef<IAssetTypeActions> Action);
 	void OnNiagaraSettingsChangedEvent(const FString& PropertyName, const UNiagaraSettings* Settings);
@@ -127,7 +130,6 @@ private:
 	FDelegateHandle CreateVectorParameterTrackEditorHandle;
 	FDelegateHandle CreateColorParameterTrackEditorHandle;
 
-	FDelegateHandle CreateDefaultScriptSourceHandle;
 	FDelegateHandle ScriptCompilerHandle;
 	FDelegateHandle PrecompilerHandle;
 
@@ -137,6 +139,8 @@ private:
 
 	TSharedPtr<FNiagaraScriptMergeManager> ScriptMergeManager;
 
+	TSharedPtr<INiagaraEditorOnlyDataUtilities> EditorOnlyDataUtilities;
+
 	TMap<const UScriptStruct*, FOnCreateMovieSceneTrackForParameter> TypeToParameterTrackCreatorMap;
 
 	IConsoleCommand* TestCompileScriptCommand;
diff --git a/Engine/Plugins/FX/Niagara/Source/NiagaraEditor/Public/NiagaraEditorUtilities.h b/Engine/Plugins/FX/Niagara/Source/NiagaraEditor/Public/NiagaraEditorUtilities.h
--- a/Engine/Plugins/FX/Niagara/Source/NiagaraEditor/Public/NiagaraEditorUtilities.h
+++ b/Engine/Plugins/FX/Niagara/Source/NiagaraEditor/Public/NiagaraEditorUtilities.h
@@ -190,7 +190,7 @@ namespace FNiagaraEditorUtilities
 
 	TArray<UNiagaraComponent*> GetComponentsThatReferenceSystemViewModel(const FNiagaraSystemViewModel& ReferencedSystemViewModel);
 
-	const FGuid AddEmitterToSystem(UNiagaraSystem& InSystem, UNiagaraEmitter& InEmitterToAdd, bool bSystemIsPlaceholder);
+	const FGuid AddEmitterToSystem(UNiagaraSystem& InSystem, UNiagaraEmitter& InEmitterToAdd);
 
 	void RemoveEmittersFromSystemByEmitterHandleId(UNiagaraSystem& InSystem, TSet<FGuid> EmitterHandleIdsToDelete);
 
diff --git a/Engine/Plugins/FX/Niagara/Source/NiagaraEditor/Public/NiagaraOverviewNode.h b/Engine/Plugins/FX/Niagara/Source/NiagaraEditor/Public/NiagaraOverviewNode.h
new file mode 100644
--- /dev/null
+++ b/Engine/Plugins/FX/Niagara/Source/NiagaraEditor/Public/NiagaraOverviewNode.h
@@ -0,0 +1,38 @@
+// Copyright 1998-2019 Epic Games, Inc. All Rights Reserved.
+#pragma once
+
+#include "EdGraph/EdGraphNode.h"
+#include "NiagaraOverviewNode.generated.h"
+
+class UNiagaraSystem;
+
+UCLASS()
+class NIAGARAEDITOR_API UNiagaraOverviewNode : public UEdGraphNode
+{
+	GENERATED_BODY()
+public:
+	UNiagaraOverviewNode();
+	void Initialize(UNiagaraSystem* InOwningSystem);
+	void Initialize(UNiagaraSystem* InOwningSystem, FGuid InEmitterHandleGuid);
+	const FGuid GetEmitterHandleGuid() const;
+
+	//~ Begin UEdGraphNode Interface
+	/** Gets the name of this node, shown in title bar */
+	virtual FText GetNodeTitle(ENodeTitleType::Type TitleType) const override;
+
+	/** Whether or not this node can be deleted by user action */
+	virtual bool CanUserDeleteNode() const override;
+
+	/** Whether or not this node can be safely duplicated (via copy/paste, etc...) in the graph */
+	virtual bool CanDuplicateNode() const override;
+	//~ End UEdGraphNode Interface
+
+	UNiagaraSystem* GetOwningSystem();
+
+private:
+	UPROPERTY()
+	UNiagaraSystem* OwningSystem;
+
+	UPROPERTY()
+	FGuid EmitterHandleGuid;
+};
diff --git a/Engine/Plugins/FX/Niagara/Source/NiagaraEditor/Public/NiagaraSystemEditorData.h b/Engine/Plugins/FX/Niagara/Source/NiagaraEditor/Public/NiagaraSystemEditorData.h
--- a/Engine/Plugins/FX/Niagara/Source/NiagaraEditor/Public/NiagaraSystemEditorData.h
+++ b/Engine/Plugins/FX/Niagara/Source/NiagaraEditor/Public/NiagaraSystemEditorData.h
@@ -78,20 +78,18 @@ public:
 
 	UEdGraph* GetSystemOverviewGraph() const;
 
-	void SystemOverviewHandleAdded(const FGuid AddedHandleGuid) const;
+	const FNiagaraGraphViewSettings& GetSystemOverviewGraphViewSettings() const;
 
-	void SystemOverviewHandlesRemoved() const;
+	void SetSystemOverviewGraphViewSettings(const FNiagaraGraphViewSettings& InOverviewGraphViewSettings);
 
-	void Initialize(UNiagaraSystem* OwnerSystem, bool bEditingSystem);
+	bool GetOwningSystemIsPlaceholder() const;
 
-private:
-	void UpdatePlaybackRangeFromEmitters(UNiagaraSystem* OwnerSystem);
-
-	bool GetSystemOverviewGraphIsValid() const;
+	void SetOwningSystemIsPlaceholder(bool bInSystemIsPlaceholder, UNiagaraSystem& OwnerSystem);
 
-	void InitSystemOverviewGraph();
+	void SynchronizeOverviewGraphWithSystem(UNiagaraSystem& OwnerSystem);
 
-	const FVector2D GetGoodPlaceForNewOverviewNode() const;
+private:
+	void UpdatePlaybackRangeFromEmitters(UNiagaraSystem& OwnerSystem);
 
 private:
 	UPROPERTY(Instanced)
@@ -114,5 +112,8 @@ private:
 	UEdGraph* SystemOverviewGraph;
 
 	UPROPERTY()
-	UNiagaraSystem* OwningSystem;
+	FNiagaraGraphViewSettings OverviewGraphViewSettings;
+
+	UPROPERTY()
+	bool bSystemIsPlaceholder;
 };
diff --git a/Engine/Plugins/FX/Niagara/Source/NiagaraEditor/Public/ViewModels/NiagaraEmitterHandleViewModel.h b/Engine/Plugins/FX/Niagara/Source/NiagaraEditor/Public/ViewModels/NiagaraEmitterHandleViewModel.h
--- a/Engine/Plugins/FX/Niagara/Source/NiagaraEditor/Public/ViewModels/NiagaraEmitterHandleViewModel.h
+++ b/Engine/Plugins/FX/Niagara/Source/NiagaraEditor/Public/ViewModels/NiagaraEmitterHandleViewModel.h
@@ -9,27 +9,32 @@
 #include "Layout/Visibility.h"
 
 class UNiagaraSystem;
+class FNiagaraSystemViewModel;
 struct FNiagaraEmitterHandle;
 class FNiagaraEmitterInstance;
 class FNiagaraEmitterViewModel;
+class UNiagaraStackViewModel;
 
 /** The view model for the FNiagaraEmitterEditorWidget. */
-class FNiagaraEmitterHandleViewModel : public TSharedFromThis<FNiagaraEmitterHandleViewModel>
+class FNiagaraEmitterHandleViewModel : public TSharedFromThis<FNiagaraEmitterHandleViewModel>, public FGCObject
 {
 public:
 	DECLARE_MULTICAST_DELEGATE(FOnPropertyChanged);
 	DECLARE_MULTICAST_DELEGATE(FOnNameChanged);
 public:
-	/** Creates a new emitter editor view model with the supplied emitter handle and simulation. */
-	FNiagaraEmitterHandleViewModel(FNiagaraEmitterHandle* InEmitterHandle, TWeakPtr<FNiagaraEmitterInstance, ESPMode::ThreadSafe> InSimulation, UNiagaraSystem& InOwningSystem);
+	/** Creates a new emitter editor view model.  This must be initialized before it can be used. */
+	FNiagaraEmitterHandleViewModel();
 	
 	~FNiagaraEmitterHandleViewModel();
 
-	/** Reuses a the emitter editor view model with the supplied emitter handle and simulation.*/
-	bool Set(FNiagaraEmitterHandle* InEmitterHandle, TWeakPtr<FNiagaraEmitterInstance, ESPMode::ThreadSafe> InSimulation, UNiagaraSystem& InOwningSystem);
+	/** Initializes the emitter editor view model with the supplied emitter handle and simulation.*/
+	void Initialize(TSharedRef<FNiagaraSystemViewModel> InOwningSystemViewModel, FNiagaraEmitterHandle* InEmitterHandle, TWeakPtr<FNiagaraEmitterInstance, ESPMode::ThreadSafe> InSimulation);
 
-	/** Sets the emitter handle.*/
-	void SetEmitterHandle(FNiagaraEmitterHandle* InEmitterHandle);
+	/** Resets the data in the view model. */
+	void Reset();
+
+	//~ FGCObject interface
+	virtual void AddReferencedObjects(FReferenceCollector& Collector) override;
 
 	/** Sets the simulation for the emitter this handle references. */
 	void SetSimulation(TWeakPtr<FNiagaraEmitterInstance, ESPMode::ThreadSafe> InSimulation);
@@ -47,6 +52,9 @@ public:
 	/** Gets the text representation of the emitter handle name. */
 	NIAGARAEDITOR_API FText GetNameText() const;
 
+	/** Gets whether or not this emitter can be renamed. */
+	NIAGARAEDITOR_API bool CanRenameEmitter() const;
+
 	/** Called when the contents of the name text control is committed. */
 	NIAGARAEDITOR_API void OnNameTextComitted(const FText& InText, ETextCommit::Type CommitInfo);
 
@@ -67,6 +75,9 @@ public:
 	/** Sets whether or not this emitter handle is enabled. */
 	NIAGARAEDITOR_API void SetIsEnabled(bool bInIsEnabled);
 
+	/** Gets whether or not the emitter for this handle has been isolated in the UI. */
+	NIAGARAEDITOR_API bool GetIsIsolated() const;
+
 	/** Gets the check state for the is enabled check box. */
 	NIAGARAEDITOR_API ECheckBoxState GetIsEnabledCheckState() const;
 
@@ -77,7 +88,16 @@ public:
 	NIAGARAEDITOR_API FNiagaraEmitterHandle* GetEmitterHandle();
 
 	/** Gets the view model for the emitter this handle references. */
-	NIAGARAEDITOR_API TSharedPtr<FNiagaraEmitterViewModel> GetEmitterViewModel();
+	NIAGARAEDITOR_API TSharedRef<FNiagaraEmitterViewModel> GetEmitterViewModel();
+
+	/** Gets the stack view model which represents the emitter pointed to by this handle. */
+	NIAGARAEDITOR_API UNiagaraStackViewModel* GetEmitterStackViewModel();
+
+	/** Gets whether or not this emitter handle has a rename pending. */
+	NIAGARAEDITOR_API bool GetIsRenamePending() const;
+
+	/** Sets whether or not this emitter handle has a rename pending. */
+	NIAGARAEDITOR_API void SetIsRenamePending(bool bInIsRenamePending);
 
 	/** Gets a multicast delegate which is called any time a property on the handle changes. */
 	FOnPropertyChanged& OnPropertyChanged();
@@ -88,18 +108,27 @@ public:
 	void Cleanup();
 
 private:
+	TSharedRef<FNiagaraSystemViewModel> GetOwningSystemViewModel() const;
+
+private:
+	/** The system view model which owns this emitter handle view model. */
+	TWeakPtr<FNiagaraSystemViewModel> OwningSystemViewModelWeak;
+
 	/** The emitter handle being displayed and edited by this view model. */
 	FNiagaraEmitterHandle* EmitterHandle;
 
-	/** The System which owns the handled being displayed and edited by this view model. */
-	UNiagaraSystem& OwningSystem;
-
 	/** The view model for emitter this handle references. */
-	TSharedPtr<FNiagaraEmitterViewModel> EmitterViewModel;
+	TSharedRef<FNiagaraEmitterViewModel> EmitterViewModel;
+
+	/** The stack view model which represents the emitter pointed to by this handle. */
+	UNiagaraStackViewModel* EmitterStackViewModel;
 
 	/** A multicast delegate which is called any time a property on the handle changes. */
 	FOnPropertyChanged OnPropertyChangedDelegate;
 
 	/** A multicast delegate which is called any time this emitter handle is renamed. */
 	FOnNameChanged OnNameChangedDelegate;
+
+	/** Gets whether or not this emitter has a pending rename. */
+	bool bIsRenamePending;
 };
diff --git a/Engine/Plugins/FX/Niagara/Source/NiagaraEditor/Public/ViewModels/NiagaraEmitterViewModel.h b/Engine/Plugins/FX/Niagara/Source/NiagaraEditor/Public/ViewModels/NiagaraEmitterViewModel.h
--- a/Engine/Plugins/FX/Niagara/Source/NiagaraEditor/Public/ViewModels/NiagaraEmitterViewModel.h
+++ b/Engine/Plugins/FX/Niagara/Source/NiagaraEditor/Public/ViewModels/NiagaraEmitterViewModel.h
@@ -8,6 +8,7 @@
 
 class UNiagaraEmitter;
 class UNiagaraScript;
+class FNiagaraSystemViewModel;
 class FNiagaraScriptViewModel;
 class FNiagaraScriptGraphViewModel;
 class UNiagaraEmitterEditorData;
@@ -23,19 +24,23 @@ public:
 	DECLARE_MULTICAST_DELEGATE(FOnEmitterChanged);
 	DECLARE_MULTICAST_DELEGATE(FOnPropertyChanged);
 	DECLARE_MULTICAST_DELEGATE(FOnScriptCompiled);
+	DECLARE_MULTICAST_DELEGATE(FOnParentRemoved);
 	DECLARE_MULTICAST_DELEGATE_TwoParams(FOnScriptGraphChanged, const FEdGraphEditAction&, const UNiagaraScript&);
 	DECLARE_MULTICAST_DELEGATE_TwoParams(FOnScriptParameterStoreChanged, const FNiagaraParameterStore&, const UNiagaraScript&);
 
 public:
-	/** Creates a new emitter editor view model with the supplied emitter handle and simulation. */
-	FNiagaraEmitterViewModel(UNiagaraEmitter* InEmitter, TWeakPtr<FNiagaraEmitterInstance, ESPMode::ThreadSafe> InSimulation);
+	/** Creates a new emitter editor view model.  It must be initialized before use. */
+	FNiagaraEmitterViewModel();
 	virtual ~FNiagaraEmitterViewModel();
 
-	/** Reuse this view model with new parameters.*/
-	bool Set(UNiagaraEmitter* InEmitter, TWeakPtr<FNiagaraEmitterInstance, ESPMode::ThreadSafe> InSimulation);
+	/** Initialize this view model with an emitter and simulation. */
+	bool Initialize(UNiagaraEmitter* InEmitter, TWeakPtr<FNiagaraEmitterInstance, ESPMode::ThreadSafe> InSimulation);
 
-	/** Sets this view model to a different emitter. */
-	void SetEmitter(UNiagaraEmitter* InEmitter);
+	/** Resets this view model to initial conditions. */
+	void Reset();
+
+	/** Gets the currently assigned simulation if there is one. */
+	TWeakPtr<FNiagaraEmitterInstance, ESPMode::ThreadSafe> GetSimulation() const;
 
 	/** Sets the current simulation for the emitter. */
 	void SetSimulation(TWeakPtr<FNiagaraEmitterInstance, ESPMode::ThreadSafe> InSimulation);
@@ -43,6 +48,9 @@ public:
 	/** Gets the emitter represented by this view model. */
 	NIAGARAEDITOR_API UNiagaraEmitter* GetEmitter();
 
+	/** Gets whether or not this emitter has a parent emitter. */
+	NIAGARAEDITOR_API bool HasParentEmitter() const;
+
 	/** Gets the parent emitter for the emitter represented by this view model, if it has one. */
 	NIAGARAEDITOR_API const UNiagaraEmitter* GetParentEmitter() const;
 
@@ -52,6 +60,9 @@ public:
 	/** Gets the text representation of the parent emitter path. */
 	NIAGARAEDITOR_API FText GetParentPathNameText() const;
 
+	/** Removes the parent emitter from this emitter. */
+	NIAGARAEDITOR_API void RemoveParentEmitter();
+
 	/** Gets text representing stats for the emitter. */
 	//~ TODO: Instead of a single string here, we should probably have separate controls with tooltips etc.
 	NIAGARAEDITOR_API FText GetStatsText() const;
@@ -71,6 +82,9 @@ public:
 	/** Gets a delegate which is called when the shared script is compiled. */
 	FOnScriptCompiled& OnScriptCompiled();
 
+	/** Gets a delegate which is called when this emitters parent is removed. */
+	FOnParentRemoved& OnParentRemoved();
+
 	/** Gets a multicast delegate which is called any time a graph on a script owned by this emitter changes. */
 	FOnScriptGraphChanged& OnScriptGraphChanged();
 
@@ -86,6 +100,9 @@ public:
 	void Cleanup();
 
 private:
+	/** Sets this view model to a different emitter. */
+	void SetEmitter(UNiagaraEmitter* InEmitter);
+
 	void OnVMCompiled(UNiagaraEmitter* InEmitter);
 
 	void AddScriptEventHandlers();
@@ -99,7 +116,6 @@ private:
 	void OnEmitterPropertiesChanged();
 
 private:
-
 	/** The text format stats display .*/
 	static const FText StatsFormat;
 
@@ -126,6 +142,9 @@ private:
 
 	FOnScriptCompiled OnScriptCompiledDelegate;
 
+	/** A multicast delegate which is called when this emitters parent is removed. */
+	FOnParentRemoved OnParentRemovedDelegate;
+
 	FOnScriptGraphChanged OnScriptGraphChangedDelegate;
 
 	FOnScriptParameterStoreChanged OnScriptParameterStoreChangedDelegate;
diff --git a/Engine/Plugins/FX/Niagara/Source/NiagaraEditor/Public/ViewModels/NiagaraOverviewGraphViewModel.h b/Engine/Plugins/FX/Niagara/Source/NiagaraEditor/Public/ViewModels/NiagaraOverviewGraphViewModel.h
--- a/Engine/Plugins/FX/Niagara/Source/NiagaraEditor/Public/ViewModels/NiagaraOverviewGraphViewModel.h
+++ b/Engine/Plugins/FX/Niagara/Source/NiagaraEditor/Public/ViewModels/NiagaraOverviewGraphViewModel.h
@@ -7,13 +7,15 @@
 #include "EditorUndoClient.h"
 #include "Styling/SlateColor.h"
 #include "Framework/Commands/UICommandList.h"
+#include "ViewModels/NiagaraSystemSelectionViewModel.h"
 
 class UNiagaraGraph;
 class FNiagaraSystemViewModel;
 class FNiagaraObjectSelection;
+struct FNiagaraGraphViewSettings;
 
 /** A view model for editing a niagara system in a graph editor. */
-class FNiagaraOverviewGraphViewModel : public TSharedFromThis<FNiagaraOverviewGraphViewModel>, public FEditorUndoClient
+class FNiagaraOverviewGraphViewModel : public TSharedFromThis<FNiagaraOverviewGraphViewModel>
 {
 public:
 	/** A multicast delegate which is called when nodes are pasted in the graph which supplies the pasted nodes. */
@@ -23,35 +25,35 @@ public:
 
 public:
 	/** Create a new view model with the supplied system editor data and graph widget. */
-	FNiagaraOverviewGraphViewModel(TSharedRef<FNiagaraSystemViewModel> InSystemViewModel);
+	FNiagaraOverviewGraphViewModel();
 
-	~FNiagaraOverviewGraphViewModel();
+	virtual ~FNiagaraOverviewGraphViewModel();
+
+	void Initialize(TSharedRef<FNiagaraSystemViewModel> InSystemViewModel);
+	
+	NIAGARAEDITOR_API TSharedRef<FNiagaraSystemViewModel> GetSystemViewModel();
+
+	NIAGARAEDITOR_API const TSharedRef<FNiagaraSystemViewModel> GetSystemViewModel() const;
 
 	/** Gets the display text for this graph. */
 	NIAGARAEDITOR_API FText GetDisplayName() const;
 
 	/** Gets the graph which is used to edit and view the system */
-	UEdGraph* GetGraph() const;
+	NIAGARAEDITOR_API UEdGraph* GetGraph() const;
 
 	/** Gets commands used for editing the graph. */
 	NIAGARAEDITOR_API TSharedRef<FUICommandList> GetCommands();
 
 	/** Gets the currently selected graph nodes. */
-	TSharedRef<FNiagaraObjectSelection> GetNodeSelection();
+	NIAGARAEDITOR_API TSharedRef<FNiagaraObjectSelection> GetNodeSelection();
 
-	/** Sets the currently selected graph nodes. */
-	void SetSelectedNodes(const TSet<UObject*>& InSelectedNodes);
+	NIAGARAEDITOR_API const FNiagaraGraphViewSettings& GetViewSettings() const;
 
-	/** Clears the currently selected graph nodes. */
-	void ClearSelectedNodes();
+	NIAGARAEDITOR_API void SetViewSettings(const FNiagaraGraphViewSettings& InOverviewGraphViewSettings);
 
 	/** Gets a multicast delegate which is called any time nodes are pasted in the graph. */
 	FOnNodesPasted& OnNodesPasted();
-
-	//~ FEditorUndoClient interface.
-	virtual void PostUndo(bool bSuccess) override;
-	virtual void PostRedo(bool bSuccess) override { PostUndo(bSuccess); }
-
+	
 private:
 	void SetupCommands();
 
@@ -67,15 +69,21 @@ private:
 	void DuplicateNodes();
 	bool CanDuplicateNodes() const;
 
-	void InitDisplayName();
+	FText GetDisplayNameInternal() const;
+
+	void GraphSelectionChanged();
+
+	void SystemSelectionChanged(UNiagaraSystemSelectionViewModel::ESelectionChangeSource SelectionChangeSource);
 
 private:
 
 	/** The view model to interface with the system being viewed and edited by this view model. */
 	TWeakPtr<FNiagaraSystemViewModel> SystemViewModel;
 
+	UEdGraph* OverviewGraph;
+
 	/** The display name for the overview graph. */
-	FText DisplayName;
+	mutable TOptional<FText> DisplayNameCache;
 
 	/** Commands for editing the graph. */
 	TSharedRef<FUICommandList> Commands;
@@ -88,4 +96,8 @@ private:
 
 	/** A multicast delegate which is called whenever the graph object is changed to a different graph. */
 	FOnGraphChanged OnGraphChangedDelegate;
+
+	bool bUpdatingSystemSelectionFromGraph;
+
+	bool bUpdatingGraphSelectionFromSystem;
 };
diff --git a/Engine/Plugins/FX/Niagara/Source/NiagaraEditor/Public/ViewModels/NiagaraSystemSelectionViewModel.h b/Engine/Plugins/FX/Niagara/Source/NiagaraEditor/Public/ViewModels/NiagaraSystemSelectionViewModel.h
new file mode 100644
--- /dev/null
+++ b/Engine/Plugins/FX/Niagara/Source/NiagaraEditor/Public/ViewModels/NiagaraSystemSelectionViewModel.h
@@ -0,0 +1,71 @@
+// Copyright 1998-2019 Epic Games, Inc. All Rights Reserved.
+
+#pragma once
+
+#include "NiagaraSystemSelectionViewModel.generated.h"
+
+
+class FNiagaraSystemViewModel;
+class UNiagaraStackEntry;
+class UNiagaraStackSelection;
+class UNiagaraStackViewModel;
+
+UCLASS()
+class NIAGARAEDITOR_API UNiagaraSystemSelectionViewModel : public UObject
+{
+public:
+	enum class ESelectionChangeSource
+	{
+		EntrySelection,
+		TopObjectLevelSelection,
+		Refresh,
+		Clear
+	};
+
+	DECLARE_MULTICAST_DELEGATE_OneParam(FOnSelectionChanged, ESelectionChangeSource);
+
+public:
+	GENERATED_BODY()
+
+	void Initialize(TSharedRef<FNiagaraSystemViewModel> InSystemViewModel);
+
+	void Finalize();
+
+	const TArray<UNiagaraStackEntry*> GetSelectedEntries() const;
+
+	bool GetSystemIsSelected() const;
+
+	const TArray<FGuid>& GetSelectedEmitterHandleIds() const;
+
+	void UpdateSelectionFromEntries(const TArray<UNiagaraStackEntry*> InSelectedEntries, const TArray<UNiagaraStackEntry*> InDeselectedEntries, bool bClearCurrentSelection);
+
+	void UpdateSelectionFromTopLevelObjects(bool bInSystemIsSelected, const TArray<FGuid> InSelectedEmitterHandleIds, bool bClearCurrentSelection);
+
+	UNiagaraStackViewModel* GetSelectionStackViewModel();
+
+	void Refresh();
+
+	FOnSelectionChanged& OnSelectionChanged();
+
+private:
+	TSharedRef<FNiagaraSystemViewModel> GetSystemViewModel();
+
+	void UpdateStackSelectionEntry();
+
+private:
+	TWeakPtr<FNiagaraSystemViewModel> SystemViewModelWeak;
+
+	TArray<UNiagaraStackEntry*> SelectedEntries;
+
+	bool bSystemIsSelected;
+
+	TArray<FGuid> SelectedEmitterHandleIds;
+
+	UPROPERTY()
+	UNiagaraStackSelection* StackSelection;
+
+	UPROPERTY()
+	UNiagaraStackViewModel* SelectionStackViewModel;
+
+	FOnSelectionChanged OnSelectionChangedDelegate;
+};
\ No newline at end of file
diff --git a/Engine/Plugins/FX/Niagara/Source/NiagaraEditor/Public/ViewModels/NiagaraSystemViewModel.h b/Engine/Plugins/FX/Niagara/Source/NiagaraEditor/Public/ViewModels/NiagaraSystemViewModel.h
--- a/Engine/Plugins/FX/Niagara/Source/NiagaraEditor/Public/ViewModels/NiagaraSystemViewModel.h
+++ b/Engine/Plugins/FX/Niagara/Source/NiagaraEditor/Public/ViewModels/NiagaraSystemViewModel.h
@@ -8,6 +8,7 @@
 #include "UObject/GCObject.h"
 #include "NiagaraCurveOwner.h"
 #include "ViewModels/TNiagaraViewModelManager.h"
+#include "ViewModels/NiagaraSystemSelectionViewModel.h"
 #include "ISequencer.h"
 
 #include "TickableEditorObject.h"
@@ -22,6 +23,8 @@ struct FNiagaraVariable;
 struct FNiagaraEmitterHandle;
 class FNiagaraEmitterHandleViewModel;
 class FNiagaraSystemScriptViewModel;
+class UNiagaraStackViewModel;
+class UNiagaraSystemSelectionViewModel;
 class FNiagaraSystemInstance;
 class ISequencer;
 struct FAssetData;
@@ -60,6 +63,9 @@ struct FNiagaraSystemViewModelOptions
 	/** Whether or not the system represented by this view model can be simulated. True by default. */
 	bool bCanSimulate;
 
+	/** An optional unique id which associates this view model with messages in the message manager. */
+	TOptional<const FGuid> MessageLogGuid;
+
 	/** Gets the current editing mode for this system. */
 	ENiagaraSystemViewModelEditMode EditMode;
 };
@@ -86,8 +92,6 @@ public:
 
 	DECLARE_MULTICAST_DELEGATE(FOnCurveOwnerChanged);
 
-	DECLARE_MULTICAST_DELEGATE(FOnSelectedEmitterHandlesChanged);
-
 	DECLARE_MULTICAST_DELEGATE(FOnPostSequencerTimeChange)
 
 	DECLARE_MULTICAST_DELEGATE(FOnSystemCompiled);
@@ -134,16 +138,23 @@ public:
 		ResetSystem,
 	};
 	/** Creates a new view model with the supplied System and System instance. */
-	FNiagaraSystemViewModel(UNiagaraSystem& InSystem, FNiagaraSystemViewModelOptions InOptions, TOptional<const FGuid> InMessageLogGuid = TOptional<const FGuid>());
+	FNiagaraSystemViewModel();
+	
+	/** Initializes this system view model with the supplied system and options. */
+	void Initialize(UNiagaraSystem& InSystem, FNiagaraSystemViewModelOptions InOptions);
 
 	~FNiagaraSystemViewModel();
 
+	NIAGARAEDITOR_API FText GetDisplayName() const;
+
 	/** Gets an array of the view models for the emitter handles owned by this System. */
-	const TArray<TSharedRef<FNiagaraEmitterHandleViewModel>>& GetEmitterHandleViewModels();
+	NIAGARAEDITOR_API const TArray<TSharedRef<FNiagaraEmitterHandleViewModel>>& GetEmitterHandleViewModels();
 
 	/** Gets an emitter handle view model by id.  Returns an invalid shared ptr if it can't be found. */
 	NIAGARAEDITOR_API TSharedPtr<FNiagaraEmitterHandleViewModel> GetEmitterHandleViewModelById(FGuid InEmitterHandleId);
 
+	TSharedPtr<FNiagaraEmitterHandleViewModel> GetEmitterHandleViewModelForEmitter(UNiagaraEmitter* InEmitter) const;
+
 	/** Gets the view model for the System script. */
 	TSharedPtr<FNiagaraSystemScriptViewModel> GetSystemScriptViewModel();
 
@@ -157,7 +168,7 @@ public:
 	FNiagaraCurveOwner& GetCurveOwner();
 
 	/** Get access to the underlying system*/
-	UNiagaraSystem& GetSystem() const { return System; }
+	NIAGARAEDITOR_API UNiagaraSystem& GetSystem() const;
 
 	/** Gets whether or not emitters can be added from the timeline. */
 	bool GetCanModifyEmittersFromTimeline() const;
@@ -175,13 +186,10 @@ public:
 	void DeleteEmitters(TSet<FGuid> EmitterHandleIdsToDelete);
 
 	/** Gets a multicast delegate which is called any time the array of emitter handle view models changes. */
-	FOnEmitterHandleViewModelsChanged& OnEmitterHandleViewModelsChanged();
+	NIAGARAEDITOR_API FOnEmitterHandleViewModelsChanged& OnEmitterHandleViewModelsChanged();
 
 	/** Gets a delegate which is called any time the data in the curve owner is changed internally by this view model. */
 	FOnCurveOwnerChanged& OnCurveOwnerChanged();
-
-	/** Gets a multicast delegate which is called whenever the selected emitter handles changes. */
-	FOnSelectedEmitterHandlesChanged& OnSelectedEmitterHandlesChanged();
 	
 	/** Gets a multicast delegate which is called whenever we've received and handled a sequencer time update.*/
 	FOnPostSequencerTimeChange& OnPostSequencerTimeChanged();
@@ -223,23 +231,8 @@ public:
 	/* Get the latest status of this view-model's script compilation.*/
 	ENiagaraScriptCompileStatus GetLatestCompileStatus() const;
 
-	/** Gets the ids for the currently selected emitter handles. */
-	const TArray<FGuid>& GetSelectedEmitterHandleIds();
-
-	/** Sets the currently selected emitter handles by id. */
-	void SetSelectedEmitterHandlesById(TArray<FGuid> InSelectedEmitterHandleIds);
-
-	/** Sets the currently selected emitter handle by id. */
-	void SetSelectedEmitterHandleById(FGuid InSelectedEmitterHandleId);
-
-	/** Gets the currently selected emitter handles. */
-	void GetSelectedEmitterHandles(TArray<TSharedRef<FNiagaraEmitterHandleViewModel>>& OutSelectedEmitterHanldles);
-
 	/** Gets editor specific data which can be stored per system.  If this data hasn't been created the default version will be returned. */
-	const NIAGARAEDITOR_API UNiagaraSystemEditorData& GetEditorData() const;
-
-	/** Gets editor specific data which is stored per system.  If this data hasn't been created then it will be created. */
-	UNiagaraSystemEditorData& GetOrCreateEditorData();
+	NIAGARAEDITOR_API UNiagaraSystemEditorData& GetEditorData() const;
 	
 	/** Act as if the system has been fully destroyed although references might persist.*/
 	void Cleanup();
@@ -254,14 +247,11 @@ public:
 	void UpdateEmitterFixedBounds();
 
 	/** Isolates the supplied emitters.  This will remove all other emitters from isolation. */
-	void IsolateEmitters(TArray<TSharedRef<FNiagaraEmitterHandleViewModel>> EmitterHandlesToIsolate);
+	void IsolateEmitters(TArray<FGuid> EmitterHandlesIdsToIsolate);
 
 	/** Toggles the isolation state of a single emitter. */
 	void ToggleEmitterIsolation(TSharedRef<FNiagaraEmitterHandleViewModel> InEmitterHandle);
 
-	/** Returns whether the suplied emitter is isolated. */
-	bool IsEmitterIsolated(TSharedRef<FNiagaraEmitterHandleViewModel> InEmitterHandle);
-
 	/** Export the current state of the system to text.*/
 	void DumpToText(FString& ExportText);
 
@@ -295,11 +285,11 @@ public:
 	/** Gets the ViewModel for the system overview graph. */
 	NIAGARAEDITOR_API TSharedPtr<FNiagaraOverviewGraphViewModel> GetOverviewGraphViewModel() const;
 
-	/** Sets the ViewModel for the system overview graph. */
-	void SetOverviewGraphViewModel(TSharedRef<FNiagaraOverviewGraphViewModel> InOverviewGraphViewModel);
+	/** Gets the stack view model representing this system. */
+	NIAGARAEDITOR_API UNiagaraStackViewModel* GetSystemStackViewModel();
 
-	/** Notifies views to update displayed Emitter name that has changed. */
-	NIAGARAEDITOR_API void EmitterNameChanged(const FText& InText, FGuid ChangedEmitterGuid);
+	/** Gets the a view model representing the selected entries in the overview. */
+	NIAGARAEDITOR_API UNiagaraSystemSelectionViewModel* GetSelectionViewModel();
 
 private:
 
@@ -377,6 +367,9 @@ private:
 	/** Called whenever the global time in the sequencer is changed. */
 	void SequencerTimeChanged();
 
+	/** Called whenever the current selection in the system changes. */
+	void SystemSelectionChanged(UNiagaraSystemSelectionViewModel::ESelectionChangeSource SelectionChangeSource);
+
 	/** Called whenever the track selection in sequencer changes. */
 	void SequencerTrackSelectionChanged(TArray<UMovieSceneTrack*> SelectedTracks);
 
@@ -421,7 +414,7 @@ private:
 
 private:
 	/** The System being viewed and edited by this view model. */
-	UNiagaraSystem& System;
+	UNiagaraSystem* System;
 
 	/** The component used for previewing the System in a viewport. */
 	UNiagaraComponent* PreviewComponent;
@@ -474,9 +467,6 @@ private:
 	/** A multicast delegate which is called when the contents of the curve owner is changed internally by this view model. */
 	FOnCurveOwnerChanged OnCurveOwnerChangedDelegate;
 
-	/** A multicast delegate which is called whenever the selected emitter changes. */
-	FOnSelectedEmitterHandlesChanged OnSelectedEmitterHandlesChangedDelegate;
-
 	/** A multicast delegate which is called whenever we've received and handled a sequencer time update.*/
 	FOnPostSequencerTimeChange OnPostSequencerTimeChangeDelegate;
 
@@ -498,9 +488,6 @@ private:
 	/** A curve owner implementation for curves in a niagara System. */
 	FNiagaraCurveOwner CurveOwner;
 
-	/** The ids for the currently selected emitter handles. */
-	TArray<FGuid> SelectedEmitterHandleIds;
-
 	TNiagaraViewModelManager<UNiagaraSystem, FNiagaraSystemViewModel>::Handle RegisteredHandle;
 
 	/** Array of FNiagaraEmitterHandleViewModel representing emitters that are pinned in the stack UI*/
@@ -543,9 +530,12 @@ private:
 	TArray<FGuid> EmitterIdsRequiringSequencerTrackUpdate;
 
 	/** GUID used when sending message jobs to FNiagaraMessageManager for notifying the FNiagaraMessageLogViewModel with the same GUID key */
-	const TOptional<const FGuid> SystemMessageLogGuidKey;
+	TOptional<const FGuid> SystemMessageLogGuidKey;
 
 	/** ViewModel for the system overview graph */
 	TSharedPtr<FNiagaraOverviewGraphViewModel> OverviewGraphViewModel;
 
-};
+	UNiagaraStackViewModel* SystemStackViewModel;
+
+	UNiagaraSystemSelectionViewModel* SelectionViewModel;
+};
\ No newline at end of file
diff --git a/Engine/Plugins/FX/Niagara/Source/NiagaraEditor/Public/ViewModels/Stack/NiagaraStackEntry.h b/Engine/Plugins/FX/Niagara/Source/NiagaraEditor/Public/ViewModels/Stack/NiagaraStackEntry.h
--- a/Engine/Plugins/FX/Niagara/Source/NiagaraEditor/Public/ViewModels/Stack/NiagaraStackEntry.h
+++ b/Engine/Plugins/FX/Niagara/Source/NiagaraEditor/Public/ViewModels/Stack/NiagaraStackEntry.h
@@ -43,6 +43,7 @@ public:
 	DECLARE_DELEGATE_RetVal_TwoParams(TOptional<FDropResult>, FOnRequestDrop, const UNiagaraStackEntry&, const TArray<UNiagaraStackEntry*>&);
 	DECLARE_DELEGATE_RetVal_OneParam(bool, FOnFilterChild, const UNiagaraStackEntry&);
 	DECLARE_DELEGATE(FStackIssueFixDelegate);
+	DECLARE_DELEGATE(FOnRequestDeferredOperation);
 
 public:
 	struct NIAGARAEDITOR_API FExecutionCategoryNames
@@ -59,6 +60,7 @@ public:
 		static const FName Spawn;
 		static const FName Update;
 		static const FName Event;
+		static const FName Render;
 	};
 
 	enum class EStackRowStyle
@@ -75,7 +77,7 @@ public:
 
 	struct FRequiredEntryData
 	{
-		FRequiredEntryData(TSharedRef<FNiagaraSystemViewModel> InSystemViewModel, TSharedRef<FNiagaraEmitterViewModel> InEmitterViewModel, FName InExecutionCategoryName, FName InExecutionSubcategoryName, UNiagaraStackEditorData& InStackEditorData)
+		FRequiredEntryData(TSharedRef<FNiagaraSystemViewModel> InSystemViewModel, TSharedPtr<FNiagaraEmitterViewModel> InEmitterViewModel, FName InExecutionCategoryName, FName InExecutionSubcategoryName, UNiagaraStackEditorData& InStackEditorData)
 			: SystemViewModel(InSystemViewModel)
 			, EmitterViewModel(InEmitterViewModel)
 			, ExecutionCategoryName(InExecutionCategoryName)
@@ -85,7 +87,7 @@ public:
 		}
 
 		const TSharedRef<FNiagaraSystemViewModel> SystemViewModel;
-		const TSharedRef<FNiagaraEmitterViewModel> EmitterViewModel;
+		const TSharedPtr<FNiagaraEmitterViewModel> EmitterViewModel;
 		const FName ExecutionCategoryName;
 		const FName ExecutionSubcategoryName;
 		UNiagaraStackEditorData* const StackEditorData;
@@ -167,6 +169,8 @@ public:
 
 	void Finalize();
 
+	bool IsFinalized() const;
+
 	virtual FText GetDisplayName() const;
 
 	UNiagaraStackEditorData& GetStackEditorData() const;
@@ -211,12 +215,13 @@ public:
 
 	void RefreshChildren();
 
+	void RefreshChildrenDeferred();
 
 	FDelegateHandle AddChildFilter(FOnFilterChild ChildFilter);
 	void RemoveChildFilter(FDelegateHandle FilterHandle);
 
 	TSharedRef<FNiagaraSystemViewModel> GetSystemViewModel() const;
-	TSharedRef<FNiagaraEmitterViewModel> GetEmitterViewModel() const;
+	TSharedPtr<FNiagaraEmitterViewModel> GetEmitterViewModel() const;
 
 	template<typename ChildType, typename PredicateType>
 	static ChildType* FindCurrentChildOfTypeByPredicate(const TArray<UNiagaraStackEntry*>& CurrentChildren, PredicateType Predicate)
diff --git a/Engine/Plugins/FX/Niagara/Source/NiagaraEditor/Public/ViewModels/Stack/NiagaraStackGraphUtilities.h b/Engine/Plugins/FX/Niagara/Source/NiagaraEditor/Public/ViewModels/Stack/NiagaraStackGraphUtilities.h
--- a/Engine/Plugins/FX/Niagara/Source/NiagaraEditor/Public/ViewModels/Stack/NiagaraStackGraphUtilities.h
+++ b/Engine/Plugins/FX/Niagara/Source/NiagaraEditor/Public/ViewModels/Stack/NiagaraStackGraphUtilities.h
@@ -62,9 +62,9 @@ namespace FNiagaraStackGraphUtilities
 
 	void ConnectStackNodeGroup(const FStackNodeGroup& ConnectGroup, const FStackNodeGroup& NewPreviousGroup, const FStackNodeGroup& NewNextGroup);
 
-	void InitializeStackFunctionInputs(TSharedRef<FNiagaraSystemViewModel> SystemViewModel, TSharedRef<FNiagaraEmitterViewModel> EmitterViewModel, UNiagaraStackEditorData& StackEditorData, UNiagaraNodeFunctionCall& ModuleNode, UNiagaraNodeFunctionCall& InputFunctionCallNode);
+	void InitializeStackFunctionInputs(TSharedRef<FNiagaraSystemViewModel> SystemViewModel, TSharedPtr<FNiagaraEmitterViewModel> EmitterViewModel, UNiagaraStackEditorData& StackEditorData, UNiagaraNodeFunctionCall& ModuleNode, UNiagaraNodeFunctionCall& InputFunctionCallNode);
 
-	void InitializeStackFunctionInput(TSharedRef<FNiagaraSystemViewModel> SystemViewModel, TSharedRef<FNiagaraEmitterViewModel> EmitterViewModel, UNiagaraStackEditorData& StackEditorData, UNiagaraNodeFunctionCall& ModuleNode, UNiagaraNodeFunctionCall& InputFunctionCallNode, FName InputName);
+	void InitializeStackFunctionInput(TSharedRef<FNiagaraSystemViewModel> SystemViewModel, TSharedPtr<FNiagaraEmitterViewModel> EmitterViewModel, UNiagaraStackEditorData& StackEditorData, UNiagaraNodeFunctionCall& ModuleNode, UNiagaraNodeFunctionCall& InputFunctionCallNode, FName InputName);
 
 	FString GenerateStackFunctionInputEditorDataKey(UNiagaraNodeFunctionCall& FunctionCallNode, FNiagaraParameterHandle InputParameterHandle);
 
@@ -156,7 +156,7 @@ namespace FNiagaraStackGraphUtilities
 
 	void ResetToDefaultDynamicInput(
 		TSharedRef<FNiagaraSystemViewModel> SystemViewModel,
-		TSharedRef<FNiagaraEmitterViewModel> EmitterViewModel,
+		TSharedPtr<FNiagaraEmitterViewModel> EmitterViewModel,
 		UNiagaraStackEditorData& StackEditorData,
 		UNiagaraScript& SourceScript,
 		const TArray<TWeakObjectPtr<UNiagaraScript>> AffectedScripts,
diff --git a/Engine/Plugins/FX/Niagara/Source/NiagaraEditor/Public/ViewModels/Stack/NiagaraStackRoot.h b/Engine/Plugins/FX/Niagara/Source/NiagaraEditor/Public/ViewModels/Stack/NiagaraStackRoot.h
--- a/Engine/Plugins/FX/Niagara/Source/NiagaraEditor/Public/ViewModels/Stack/NiagaraStackRoot.h
+++ b/Engine/Plugins/FX/Niagara/Source/NiagaraEditor/Public/ViewModels/Stack/NiagaraStackRoot.h
@@ -20,7 +20,7 @@ class NIAGARAEDITOR_API UNiagaraStackRoot : public UNiagaraStackEntry
 public:
 	UNiagaraStackRoot();
 	
-	void Initialize(FRequiredEntryData InRequiredEntryData);
+	void Initialize(FRequiredEntryData InRequiredEntryData, bool bInIncludeSystemInformation, bool bInIncludeEmitterInformation);
 
 	virtual bool GetCanExpand() const override;
 	virtual bool GetShouldShowInStack() const override;
@@ -61,4 +61,7 @@ private:
 
 	UPROPERTY()
 	UNiagaraStackRenderItemGroup* RenderGroup;
+
+	bool bIncludeSystemInformation;
+	bool bIncludeEmitterInformation;
 };
diff --git a/Engine/Plugins/FX/Niagara/Source/NiagaraEditor/Public/ViewModels/Stack/NiagaraStackSelection.h b/Engine/Plugins/FX/Niagara/Source/NiagaraEditor/Public/ViewModels/Stack/NiagaraStackSelection.h
new file mode 100644
--- /dev/null
+++ b/Engine/Plugins/FX/Niagara/Source/NiagaraEditor/Public/ViewModels/Stack/NiagaraStackSelection.h
@@ -0,0 +1,27 @@
+// Copyright 1998-2019 Epic Games, Inc. All Rights Reserved.
+
+#pragma once
+
+#include "NiagaraStackEntry.h"
+#include "NiagaraStackSelection.generated.h"
+
+UCLASS()
+class NIAGARAEDITOR_API UNiagaraStackSelection : public UNiagaraStackEntry
+{
+	GENERATED_BODY()
+
+public:
+	void Initialize(FRequiredEntryData InRequiredEntryData);
+
+	virtual bool GetCanExpand() const override;
+
+	virtual bool GetShouldShowInStack() const override;
+
+	void SetSelectedEntries(const TArray<UNiagaraStackEntry*>& InSelectedEntries);
+
+protected:
+	virtual void RefreshChildrenInternal(const TArray<UNiagaraStackEntry*>& CurrentChildren, TArray<UNiagaraStackEntry*>& NewChildren, TArray<FStackIssue>& NewIssues) override;
+
+private:
+	TArray<TWeakObjectPtr<UNiagaraStackEntry>> SelectedEntries;
+};
\ No newline at end of file
diff --git a/Engine/Plugins/FX/Niagara/Source/NiagaraEditor/Public/ViewModels/Stack/NiagaraStackViewModel.h b/Engine/Plugins/FX/Niagara/Source/NiagaraEditor/Public/ViewModels/Stack/NiagaraStackViewModel.h
--- a/Engine/Plugins/FX/Niagara/Source/NiagaraEditor/Public/ViewModels/Stack/NiagaraStackViewModel.h
+++ b/Engine/Plugins/FX/Niagara/Source/NiagaraEditor/Public/ViewModels/Stack/NiagaraStackViewModel.h
@@ -13,6 +13,28 @@ class FNiagaraEmitterHandleViewModel;
 class UNiagaraStackEntry;
 class UNiagaraStackRoot;
 
+struct FNiagaraStackViewModelOptions
+{
+	FNiagaraStackViewModelOptions()
+		: bIncludeSystemInformation(false)
+		, bIncludeEmitterInformation(false)
+	{
+	}
+
+	FNiagaraStackViewModelOptions(bool bInIncludeSystemInformation, bool bInIncludeEmitterInformation)
+		: bIncludeSystemInformation(bInIncludeSystemInformation)
+		, bIncludeEmitterInformation(bInIncludeEmitterInformation)
+	{
+	}
+
+	bool GetIncludeSystemInformation() { return bIncludeSystemInformation; }
+	bool GetIncludeEmitterInformation() { return bIncludeEmitterInformation; }
+
+private:
+	bool bIncludeSystemInformation;
+	bool bIncludeEmitterInformation;
+};
+
 UCLASS()
 class NIAGARAEDITOR_API UNiagaraStackViewModel : public UObject, public FEditorUndoClient
 {
@@ -33,10 +55,32 @@ public:
 				nullptr;
 		}
 	};
+
+	struct NIAGARAEDITOR_API FTopLevelViewModel
+	{
+		FTopLevelViewModel(TSharedPtr<FNiagaraSystemViewModel> InSystemViewModel)
+			: SystemViewModel(InSystemViewModel)
+		{
+		}
+
+		FTopLevelViewModel(TSharedPtr<FNiagaraEmitterHandleViewModel> InEmitterHandleViewModel)
+			: EmitterHandleViewModel(InEmitterHandleViewModel)
+		{
+		}
+
+		bool IsValid() const;
+
+		UNiagaraStackEditorData* GetStackEditorData() const;
+
+		FText GetDisplayName() const;
+
+		const TSharedPtr<FNiagaraSystemViewModel> SystemViewModel;
+		const TSharedPtr<FNiagaraEmitterHandleViewModel> EmitterHandleViewModel;
+	};
+
 public:
-	TSharedPtr<FNiagaraEmitterHandleViewModel> GetEmitterHandleViewModel();
-	TSharedPtr<FNiagaraSystemViewModel> GetSystemViewModel();
-	void Initialize(TSharedPtr<FNiagaraSystemViewModel> InSystemViewModel, TSharedPtr<FNiagaraEmitterHandleViewModel> InEmitterHandleViewModel);
+	void InitializeWithViewModels(TSharedPtr<FNiagaraSystemViewModel> InSystemViewModel, TSharedPtr<FNiagaraEmitterHandleViewModel> InEmitterHandleViewModel, FNiagaraStackViewModelOptions InOptions);
+	void InitializeWithRootEntry(UNiagaraStackEntry* Root);
 
 	void Finalize();
 
@@ -83,10 +127,12 @@ public:
 
 	bool HasDismissedStackIssues();
 
-	bool HasParentEmitter() const;
-	void RemoveEmitterSource();
+	const TArray<TSharedRef<FTopLevelViewModel>>& GetTopLevelViewModels() const;
+
+	TSharedPtr<FTopLevelViewModel> GetTopLevelViewModelForEntry(UNiagaraStackEntry& InEntry) const;
 
 private:
+	void Reset();
 
 	/** Recursively Expands all groups and collapses all items in the stack. */
 	void CollapseToHeadersRecursive(TArray<UNiagaraStackEntry*> Entries);
@@ -102,13 +148,14 @@ private:
 		}
 	};
 
-private:
 	void EntryStructureChanged();
 	void EntryDataObjectModified(UObject* ChangedObject);
 	void EntryRequestFullRefresh();
 	void EntryRequestFullRefreshDeferred();
+	void RefreshTopLevelViewModels();
 	void OnSystemCompiled();
 	void OnEmitterCompiled();
+	void EmitterParentRemoved();
 	/** Called by the tick function to perform partial search */
 	void SearchTick();
 	void GenerateTraversalEntries(UNiagaraStackEntry* Root, TArray<UNiagaraStackEntry*> ParentChain, 
@@ -118,13 +165,15 @@ private:
 	void RestoreStackEntryExpansionPreSearch();
 
 private:
-	TSharedPtr<FNiagaraEmitterHandleViewModel> EmitterHandleViewModel;
-	TSharedPtr<FNiagaraSystemViewModel> SystemViewModel;
+	TWeakPtr<FNiagaraEmitterHandleViewModel> EmitterHandleViewModel;
+	TWeakPtr<FNiagaraSystemViewModel> SystemViewModel;
 	
 	TArray<UNiagaraStackEntry*> RootEntries;
 
 	UPROPERTY()
-	UNiagaraStackRoot* RootEntry;
+	UNiagaraStackEntry* RootEntry;
+
+	bool bExternalRootEntry;
 
 	FOnStructureChanged StructureChangedDelegate;
 
@@ -138,4 +187,9 @@ private:
 	static const double MaxSearchTime;
 	bool bRestartSearch;
 	bool bRefreshPending;
+
+	bool bUsesTopLevelViewModels;
+	TArray<TSharedRef<FTopLevelViewModel>> TopLevelViewModels;
+
+	FNiagaraStackViewModelOptions Options;
 };
\ No newline at end of file
diff --git a/Engine/Plugins/FX/Niagara/Source/NiagaraEditorWidgets/Private/NiagaraEditorWidgetsModule.cpp b/Engine/Plugins/FX/Niagara/Source/NiagaraEditorWidgets/Private/NiagaraEditorWidgetsModule.cpp
--- a/Engine/Plugins/FX/Niagara/Source/NiagaraEditorWidgets/Private/NiagaraEditorWidgetsModule.cpp
+++ b/Engine/Plugins/FX/Niagara/Source/NiagaraEditorWidgets/Private/NiagaraEditorWidgetsModule.cpp
@@ -6,6 +6,7 @@
 #include "DetailCustomizations/NiagaraDataInterfaceCurveDetails.h"
 #include "DetailCustomizations/NiagaraDataInterfaceDetails.h"
 #include "DetailCustomizations/NiagaraDataInterfaceSkeletalMeshDetails.h"
+#include "ViewModels/NiagaraSystemViewModel.h"
 
 #include "Modules/ModuleManager.h"
 #include "PropertyEditorModule.h"
@@ -157,5 +158,5 @@ TSharedRef<SWidget> FNiagaraEditorWidgetsModule::FNiagaraEditorWidgetProvider::C
 
 TSharedRef<SWidget> FNiagaraEditorWidgetsModule::FNiagaraEditorWidgetProvider::CreateSystemOverview(TSharedRef<FNiagaraSystemViewModel> SystemViewModel)
 {
-	return SNew(SNiagaraOverviewGraph, SystemViewModel);
+	return SNew(SNiagaraOverviewGraph, SystemViewModel->GetOverviewGraphViewModel().ToSharedRef());
 }
diff --git a/Engine/Plugins/FX/Niagara/Source/NiagaraEditorWidgets/Private/NiagaraEditorWidgetsStyle.cpp b/Engine/Plugins/FX/Niagara/Source/NiagaraEditorWidgets/Private/NiagaraEditorWidgetsStyle.cpp
--- a/Engine/Plugins/FX/Niagara/Source/NiagaraEditorWidgets/Private/NiagaraEditorWidgetsStyle.cpp
+++ b/Engine/Plugins/FX/Niagara/Source/NiagaraEditorWidgets/Private/NiagaraEditorWidgetsStyle.cpp
@@ -97,6 +97,21 @@ TSharedRef< FSlateStyleSet > FNiagaraEditorWidgetsStyle::Create()
 		.SetFont(StackItemFont);
 	Style->Set("NiagaraEditor.Stack.ItemText", StackItemText);
 
+	FSlateFontInfo SystemOverviewListHeaderFont = DEFAULT_FONT("Bold", 12);
+	FTextBlockStyle SystemOverviewListHeaderText = FTextBlockStyle(NormalText)
+		.SetFont(SystemOverviewListHeaderFont);
+	Style->Set("NiagaraEditor.SystemOverview.ListHeaderText", SystemOverviewListHeaderText);
+
+	FSlateFontInfo SystemOverviewGroupHeaderFont = DEFAULT_FONT("Bold", 9);
+	FTextBlockStyle SystemOverviewGroupHeaderText = FTextBlockStyle(NormalText)
+		.SetFont(SystemOverviewGroupHeaderFont);
+	Style->Set("NiagaraEditor.SystemOverview.GroupHeaderText", SystemOverviewGroupHeaderText);
+
+	FSlateFontInfo SystemOverviewItemFont = DEFAULT_FONT("Regular", 10);
+	FTextBlockStyle SystemOverviewItemText = FTextBlockStyle(NormalText)
+		.SetFont(SystemOverviewItemFont);
+	Style->Set("NiagaraEditor.SystemOverview.ItemText", SystemOverviewItemText);
+
 	Style->Set("NiagaraEditor.Stack.Group.BackgroundColor", FLinearColor(FColor(96, 96, 96)));
 	Style->Set("NiagaraEditor.Stack.Item.HeaderBackgroundColor", FLinearColor(FColor(48, 48, 48)));
 	Style->Set("NiagaraEditor.Stack.Item.ContentBackgroundColor", FLinearColor(FColor(62, 62, 62)));
@@ -132,11 +147,13 @@ TSharedRef< FSlateStyleSet > FNiagaraEditorWidgetsStyle::Create()
 	Style->Set("NiagaraEditor.Stack.SpawnIcon", new IMAGE_PLUGIN_BRUSH("Icons/Spawn", Icon12x12, FLinearColor::White));
 	Style->Set("NiagaraEditor.Stack.UpdateIcon", new IMAGE_PLUGIN_BRUSH("Icons/Update", Icon12x12, FLinearColor::White));
 	Style->Set("NiagaraEditor.Stack.EventIcon", new IMAGE_PLUGIN_BRUSH("Icons/Event", Icon12x12, FLinearColor::White));
+	Style->Set("NiagaraEditor.Stack.RenderIcon", new IMAGE_PLUGIN_BRUSH("Icons/Render", Icon12x12, FLinearColor::White));
 
 	Style->Set("NiagaraEditor.Stack.ParametersIconHighlighted", new IMAGE_PLUGIN_BRUSH("Icons/SystemParams", Icon16x16, FLinearColor::White));
 	Style->Set("NiagaraEditor.Stack.SpawnIconHighlighted", new IMAGE_PLUGIN_BRUSH("Icons/Spawn", Icon16x16, FLinearColor::White));
 	Style->Set("NiagaraEditor.Stack.UpdateIconHighlighted", new IMAGE_PLUGIN_BRUSH("Icons/Update", Icon16x16, FLinearColor::White));
 	Style->Set("NiagaraEditor.Stack.EventIconHighlighted", new IMAGE_PLUGIN_BRUSH("Icons/Event", Icon16x16, FLinearColor::White));
+	Style->Set("NiagaraEditor.Stack.RenderIconHighlighted", new IMAGE_PLUGIN_BRUSH("Icons/Render", Icon16x16, FLinearColor::White));
 
 	Style->Set("NiagaraEditor.Stack.IconHighlightedSize", 16.0f);
 
diff --git a/Engine/Plugins/FX/Niagara/Source/NiagaraEditorWidgets/Private/NiagaraEditorWidgetsUtilities.cpp b/Engine/Plugins/FX/Niagara/Source/NiagaraEditorWidgets/Private/NiagaraEditorWidgetsUtilities.cpp
--- a/Engine/Plugins/FX/Niagara/Source/NiagaraEditorWidgets/Private/NiagaraEditorWidgetsUtilities.cpp
+++ b/Engine/Plugins/FX/Niagara/Source/NiagaraEditorWidgets/Private/NiagaraEditorWidgetsUtilities.cpp
@@ -47,6 +47,10 @@ FName FNiagaraStackEditorWidgetsUtilities::GetIconNameForExecutionSubcategory(FN
 		{
 			return "NiagaraEditor.Stack.EventIconHighlighted";
 		}
+		else if (ExecutionSubcategoryName == UNiagaraStackEntry::FExecutionSubcategoryNames::Render)
+		{
+			return "NiagaraEditor.Stack.RenderIconHighlighted";
+		}
 	}
 	else
 	{
@@ -66,6 +70,10 @@ FName FNiagaraStackEditorWidgetsUtilities::GetIconNameForExecutionSubcategory(FN
 		{
 			return "NiagaraEditor.Stack.EventIcon";
 		}
+		else if (ExecutionSubcategoryName == UNiagaraStackEntry::FExecutionSubcategoryNames::Render)
+		{
+			return "NiagaraEditor.Stack.RenderIcon";
+		}
 	}
 
 	return NAME_None;
diff --git a/Engine/Plugins/FX/Niagara/Source/NiagaraEditorWidgets/Private/SNiagaraOverviewGraph.cpp b/Engine/Plugins/FX/Niagara/Source/NiagaraEditorWidgets/Private/SNiagaraOverviewGraph.cpp
--- a/Engine/Plugins/FX/Niagara/Source/NiagaraEditorWidgets/Private/SNiagaraOverviewGraph.cpp
+++ b/Engine/Plugins/FX/Niagara/Source/NiagaraEditorWidgets/Private/SNiagaraOverviewGraph.cpp
@@ -1,17 +1,21 @@
 // Copyright 1998-2019 Epic Games, Inc. All Rights Reserved.
 
 #include "SNiagaraOverviewGraph.h"
+#include "NiagaraOverviewGraphNodeFactory.h"
 #include "SNiagaraStack.h"
 #include "ViewModels/NiagaraSystemViewModel.h"
 #include "ViewModels/NiagaraEmitterHandleViewModel.h"
+#include "ViewModels/NiagaraOverviewGraphViewModel.h"
 #include "ViewModels/Stack/NiagaraStackViewModel.h"
 #include "NiagaraSystem.h"
+#include "NiagaraOverviewNode.h"
 #include "NiagaraSystemEditorData.h"
+#include "NiagaraObjectSelection.h"
 #include "IContentBrowserSingleton.h"
 #include "ContentBrowserModule.h"
 #include "ViewModels/NiagaraOverviewGraphViewModel.h"
 #include "EdGraphSchema_Niagara.h"
-#include "NiagaraOverviewNodeStackItem.h"
+#include "NiagaraOverviewNode.h"
 #include "GraphEditorActions.h"
 #include "Framework/MultiBox/MultiBoxBuilder.h"
 #include "ScopedTransaction.h"
@@ -22,26 +26,26 @@
 
 #define LOCTEXT_NAMESPACE "NiagaraOverviewGraph"
 
-void SNiagaraOverviewGraph::Construct(const FArguments& InArgs, TSharedRef<FNiagaraSystemViewModel> InSystemViewModel)
+void SNiagaraOverviewGraph::Construct(const FArguments& InArgs, TSharedRef<FNiagaraOverviewGraphViewModel> InViewModel)
 {
-	SystemViewModel = InSystemViewModel;
-	bUpdatingOverviewSelectionFromGraph = false;
-	bUpdatingGraphSelectionFromOverview = false;
+	ViewModel = InViewModel;
+	ViewModel->GetNodeSelection()->OnSelectedObjectsChanged().AddSP(this, &SNiagaraOverviewGraph::ViewModelSelectionChanged);
+
+	bUpdatingViewModelSelectionFromGraph = false;
+	bUpdatingGraphSelectionFromViewModel = false;
 
 	SGraphEditor::FGraphEditorEvents Events;
 	Events.OnSelectionChanged = SGraphEditor::FOnSelectionChanged::CreateSP(this, &SNiagaraOverviewGraph::GraphSelectionChanged);
 	Events.OnCreateActionMenu = SGraphEditor::FOnCreateActionMenu::CreateSP(this, &SNiagaraOverviewGraph::OnCreateGraphActionMenu);
 	Events.OnTextCommitted = FOnNodeTextCommitted::CreateSP(this, &SNiagaraOverviewGraph::OnNodeTitleCommitted);
 
-	TSharedPtr<FNiagaraOverviewGraphViewModel> OverviewGraphViewModel = SystemViewModel->GetOverviewGraphViewModel();
-
 	FGraphAppearanceInfo AppearanceInfo;
-	if (SystemViewModel->GetEditMode() == ENiagaraSystemViewModelEditMode::EmitterAsset)
+	if (ViewModel->GetSystemViewModel()->GetEditMode() == ENiagaraSystemViewModelEditMode::EmitterAsset)
 	{
 		AppearanceInfo.CornerText = LOCTEXT("NiagaraOverview_AppearanceCornerTextEmitter", "EMITTER");
 
 	}
-	else if (SystemViewModel->GetEditMode() == ENiagaraSystemViewModelEditMode::SystemAsset)
+	else if (ViewModel->GetSystemViewModel()->GetEditMode() == ENiagaraSystemViewModelEditMode::SystemAsset)
 	{
 		AppearanceInfo.CornerText = LOCTEXT("NiagaraOverview_AppearanceCornerTextSystem", "SYSTEM");
 	}
@@ -56,12 +60,12 @@ void SNiagaraOverviewGraph::Construct(const FArguments& InArgs, TSharedRef<FNiag
 	.HAlign(HAlign_Fill)
 	[
 		SNew(STextBlock)
-		.Text(OverviewGraphViewModel.ToSharedRef(), &FNiagaraOverviewGraphViewModel::GetDisplayName)
+		.Text(ViewModel.ToSharedRef(), &FNiagaraOverviewGraphViewModel::GetDisplayName)
 		.TextStyle(FEditorStyle::Get(), TEXT("GraphBreadcrumbButtonText"))
 		.Justification(ETextJustify::Center)
 	];
 
-	TSharedRef<FUICommandList> Commands = OverviewGraphViewModel->GetCommands();
+	TSharedRef<FUICommandList> Commands = ViewModel->GetCommands();
 	Commands->MapAction(
 		FGraphEditorCommands::Get().CreateComment,
 		FExecuteAction::CreateRaw(this, &SNiagaraOverviewGraph::OnCreateComment));
@@ -70,109 +74,93 @@ void SNiagaraOverviewGraph::Construct(const FArguments& InArgs, TSharedRef<FNiag
 		.AdditionalCommands(Commands)
 		.Appearance(AppearanceInfo)
 		.TitleBar(TitleBarWidget)
-		.GraphToEdit(SystemViewModel->GetEditorData().GetSystemOverviewGraph())
+		.GraphToEdit(ViewModel->GetGraph())
 		.GraphEvents(Events);
 
-	//GraphEditor->SetNodeFactory(MakeShared<FNiagaraOverviewGraphNodeFactory>());
+	GraphEditor->SetNodeFactory(MakeShared<FNiagaraOverviewGraphNodeFactory>());
+
+	FNiagaraGraphViewSettings ViewSettings = ViewModel->GetViewSettings();
+	if (ViewSettings.IsValid())
+	{
+		GraphEditor->SetViewLocation(ViewSettings.GetLocation(), ViewSettings.GetZoom());
+		ZoomToFitFrameDelay = 0;
+	}
+	else
+	{
+		// When initialzing the graph control the stacks inside the nodes aren't actually available until two frames later due to
+		// how the underlying list view works.  In order to zoom to fix correctly we have to delay for an extra fram so we use a
+		// counter here instead of a simple bool.
+		ZoomToFitFrameDelay = 2;
+	}
 
 	ChildSlot
 	[
-		SNew(SSplitter)
-		+ SSplitter::Slot()
-		.Value(.7f)
-		[
-			GraphEditor.ToSharedRef()
-		]
-		+ SSplitter::Slot()
-		.Value(.3)
-		[
-			SNew(SBox)
-			//SNew(SNiagaraStack, SystemViewModel->GetOverviewSelectionViewModel()->GetOverviewStackViewModel())
-		]
+		GraphEditor.ToSharedRef()
 	];
+}
 
-	//SystemViewModel->GetOverviewSelectionViewModel()->OnSelectionChanged().AddSP(this, &SNiagaraOverviewGraph::OverviewSelectionChanged);
+SNiagaraOverviewGraph::~SNiagaraOverviewGraph()
+{
+	if (ViewModel.IsValid() && GraphEditor.IsValid())
+	{
+		FVector2D Location;
+		float Zoom;
+		GraphEditor->GetViewLocation(Location, Zoom);
+		ViewModel->SetViewSettings(FNiagaraGraphViewSettings(Location, Zoom));
+	}
 }
 
-UNiagaraStackEntry* GetStackEntryForSelectedNode(UObject* SelectedNode, TSharedPtr<FNiagaraSystemViewModel> SystemViewModel)
+void SNiagaraOverviewGraph::Tick(const FGeometry& AllottedGeometry, const double InCurrentTime, const float InDeltaTime)
 {
-// 	UNiagaraOverviewNodeStackItem* OverviewStackNode = Cast<UNiagaraOverviewNodeStackItem>(SelectedNode);
-// 	if (OverviewStackNode != nullptr)
-// 	{
-// 		if (OverviewStackNode->GetOwningSystem() != nullptr)
-// 		{
-// 			UNiagaraStackViewModel* StackViewModel = nullptr;
-// 			if (OverviewStackNode->GetEmitterHandleGuid().IsValid())
-// 			{
-// 				TSharedPtr<FNiagaraEmitterHandleViewModel> EmitterHandleViewModel = SystemViewModel->GetEmitterHandleViewModelById(OverviewStackNode->GetEmitterHandleGuid());
-// 				StackViewModel = EmitterHandleViewModel->GetEmitterStackViewModel();
-// 			}
-// 			else
-// 			{
-// 				StackViewModel = SystemViewModel->GetSystemStackViewModel();
-// 			}
-// 			TArray<UNiagaraStackEntry*> RootEntries = StackViewModel->GetRootEntries();
-// 			if (ensureMsgf(RootEntries.Num() == 1, TEXT("Only one root entry expected")))
-// 			{
-// 				return RootEntries[0];
-// 			}
-// 		}
-// 	}
-	return nullptr;
+	if (ZoomToFitFrameDelay > 0)
+	{
+		ZoomToFitFrameDelay--;
+		if(ZoomToFitFrameDelay == 0)
+		{
+			GraphEditor->ZoomToFit(false);
+		}
+	}
 }
 
-void SNiagaraOverviewGraph::GraphSelectionChanged(const TSet<UObject*>& SelectedNodes)
+void SNiagaraOverviewGraph::ViewModelSelectionChanged()
 {
-// 	if (bUpdatingGraphSelectionFromOverview == false)
-// 	{
-// 		TGuardValue<bool> UpdateGuard(bUpdatingOverviewSelectionFromGraph, true);
-// 		TArray<UNiagaraStackEntry*> SelectedEntries;
-// 		for (UObject* SelectedNode : SelectedNodes)
-// 		{
-// 			UNiagaraStackEntry* SelectedEntry = GetStackEntryForSelectedNode(SelectedNode, SystemViewModel);
-// 			if(SelectedEntry != nullptr)
-// 			{
-// 				SelectedEntries.Add(SelectedEntry);
-// 			}
-// 		}
-// 		bool bClearCurrentSelection = FSlateApplication::Get().GetModifierKeys().IsControlDown() == false;
-// 		SystemViewModel->GetOverviewSelectionViewModel()->UpdateSelectedEntries(SelectedEntries, bClearCurrentSelection);
-// 	}
+	if (bUpdatingViewModelSelectionFromGraph == false)
+	{
+		if (FNiagaraEditorUtilities::SetsMatch(GraphEditor->GetSelectedNodes(), ViewModel->GetNodeSelection()->GetSelectedObjects()) == false)
+		{
+			TGuardValue<bool> UpdateGuard(bUpdatingGraphSelectionFromViewModel, true);
+			GraphEditor->ClearSelectionSet();
+			for (UObject* SelectedNode : ViewModel->GetNodeSelection()->GetSelectedObjects())
+			{
+				UEdGraphNode* GraphNode = Cast<UEdGraphNode>(SelectedNode);
+				if (GraphNode != nullptr)
+				{
+					GraphEditor->SetNodeSelection(GraphNode, true);
+				}
+			}
+		}
+	}
 }
 
-void SNiagaraOverviewGraph::OverviewSelectionChanged()
+void SNiagaraOverviewGraph::GraphSelectionChanged(const TSet<UObject*>& SelectedNodes)
 {
-// 	if (bUpdatingOverviewSelectionFromGraph == false)
-// 	{
-// 		TGuardValue<bool> UpdateGuard(bUpdatingGraphSelectionFromOverview, true);
-// 
-// 		TSet<UObject*> SelectedNodeObjects = GraphEditor->GetSelectedNodes();
-// 		TArray<UNiagaraStackEntry*> SelectedOverviewEntries = SystemViewModel->GetOverviewSelectionViewModel()->GetSelectedEntries();
-// 
-// 		TArray<UEdGraphNode*> NodesToDeselect;
-// 		for (UObject* SelectedNodeObject : SelectedNodeObjects)
-// 		{
-// 			UEdGraphNode* SelectedNode = Cast<UEdGraphNode>(SelectedNodeObject);
-// 			if (SelectedNode != nullptr)
-// 			{
-// 				UNiagaraStackEntry* SelectedEntry = GetStackEntryForSelectedNode(SelectedNode, SystemViewModel);
-// 				if (SelectedEntry != nullptr && SelectedOverviewEntries.Contains(SelectedEntry) == false)
-// 				{
-// 					NodesToDeselect.Add(SelectedNode);
-// 				}
-// 			}
-// 		}
-// 
-// 		for (UEdGraphNode* NodeToDeselect : NodesToDeselect)
-// 		{
-// 			GraphEditor->SetNodeSelection(NodeToDeselect, false);
-// 		}
-// 	}
+	if (bUpdatingGraphSelectionFromViewModel == false)
+	{
+		TGuardValue<bool> UpdateGuard(bUpdatingViewModelSelectionFromGraph, true);
+		if (SelectedNodes.Num() == 0)
+		{
+			ViewModel->GetNodeSelection()->ClearSelectedObjects();
+		}
+		else
+		{
+			ViewModel->GetNodeSelection()->SetSelectedObjects(SelectedNodes);
+		}
+	}
 }
 
 FActionMenuContent SNiagaraOverviewGraph::OnCreateGraphActionMenu(UEdGraph* InGraph, const FVector2D& InNodePosition, const TArray<UEdGraphPin*>& InDraggedPins, bool bAutoExpand, SGraphEditor::FActionMenuClosed InOnMenuClosed)
 {
-	if (SystemViewModel->GetEditMode() == ENiagaraSystemViewModelEditMode::SystemAsset)
+	if (ViewModel->GetSystemViewModel()->GetEditMode() == ENiagaraSystemViewModelEditMode::SystemAsset)
 	{
 		FMenuBuilder MenuBuilder(true, NULL);
 		MenuBuilder.BeginSection(TEXT("NiagaraOverview_EditGraph"), LOCTEXT("EditGraph", "Edit Graph"));
@@ -199,7 +187,7 @@ FActionMenuContent SNiagaraOverviewGraph::OnCreateGraphActionMenu(UEdGraph* InGr
 void SNiagaraOverviewGraph::OnCreateComment()
 {
 	FNiagaraSchemaAction_NewComment CommentAction = FNiagaraSchemaAction_NewComment(GraphEditor);
-	CommentAction.PerformAction(SystemViewModel->GetEditorData().GetSystemOverviewGraph(), nullptr, GraphEditor->GetPasteLocation(), false);
+	CommentAction.PerformAction(ViewModel->GetGraph(), nullptr, GraphEditor->GetPasteLocation(), false);
 }
 
 void SNiagaraOverviewGraph::OnNodeTitleCommitted(const FText& NewText, ETextCommit::Type CommitInfo, UEdGraphNode* NodeBeingChanged)
@@ -214,10 +202,10 @@ void SNiagaraOverviewGraph::OnNodeTitleCommitted(const FText& NewText, ETextComm
 			return;
 		}
 
-		if (NodeBeingChanged->IsA(UNiagaraOverviewNodeStackItem::StaticClass())) //@TODO System Overview: renaming system or emitters locally through this view
+		if (NodeBeingChanged->IsA(UNiagaraOverviewNode::StaticClass())) //@TODO System Overview: renaming system or emitters locally through this view
 		{
-			UNiagaraOverviewNodeStackItem* OverviewNodeBeingChanged = Cast<UNiagaraOverviewNodeStackItem>(NodeBeingChanged);
-			TSharedPtr<FNiagaraEmitterHandleViewModel> NodeEmitterHandleViewModel = SystemViewModel->GetEmitterHandleViewModelById(OverviewNodeBeingChanged->GetEmitterHandleGuid());
+			UNiagaraOverviewNode* OverviewNodeBeingChanged = Cast<UNiagaraOverviewNode>(NodeBeingChanged);
+			TSharedPtr<FNiagaraEmitterHandleViewModel> NodeEmitterHandleViewModel = ViewModel->GetSystemViewModel()->GetEmitterHandleViewModelById(OverviewNodeBeingChanged->GetEmitterHandleGuid());
 			if (ensureMsgf(NodeEmitterHandleViewModel.IsValid(), TEXT("Failed to find EmitterHandleViewModel with matching Emitter GUID to Overview Node!")))
 			{
 				NodeEmitterHandleViewModel->OnNameTextComitted(NewText, CommitInfo);
@@ -244,7 +232,7 @@ void SNiagaraOverviewGraph::CreateAddEmitterMenuContent(FMenuBuilder& MenuBuilde
 		AssetPickerConfig.OnAssetSelected = FOnAssetSelected::CreateLambda([this, InGraph](const FAssetData& AssetData)
 		{
 			FSlateApplication::Get().DismissAllMenus();
-			SystemViewModel->AddEmitterFromAssetData(AssetData);
+			ViewModel->GetSystemViewModel()->AddEmitterFromAssetData(AssetData);
 		});
 		AssetPickerConfig.bAllowNullSelection = false;
 		AssetPickerConfig.InitialAssetViewType = EAssetViewType::List;
@@ -265,8 +253,8 @@ void SNiagaraOverviewGraph::CreateAddEmitterMenuContent(FMenuBuilder& MenuBuilde
 				ContentBrowserModule.Get().CreateAssetPicker(AssetPickerConfig)
 			]
 		];
-	
+
 	MenuBuilder.AddWidget(EmitterAddSubMenu, FText());
 }
 
-#undef LOCTEXT_NAMESPACE // "NiagaraOverviewGraph"
+#undef LOCTEXT_NAMESPACE // "NiagaraOverviewGraph"
\ No newline at end of file
diff --git a/Engine/Plugins/FX/Niagara/Source/NiagaraEditorWidgets/Private/SNiagaraOverviewGraph.h b/Engine/Plugins/FX/Niagara/Source/NiagaraEditorWidgets/Private/SNiagaraOverviewGraph.h
--- a/Engine/Plugins/FX/Niagara/Source/NiagaraEditorWidgets/Private/SNiagaraOverviewGraph.h
+++ b/Engine/Plugins/FX/Niagara/Source/NiagaraEditorWidgets/Private/SNiagaraOverviewGraph.h
@@ -5,7 +5,7 @@
 #include "Widgets/SCompoundWidget.h"
 #include "GraphEditor.h"
 
-class FNiagaraSystemViewModel;
+class FNiagaraOverviewGraphViewModel;
 struct FActionMenuContent;
 class FMenuBuilder;
 class UEdGraph;
@@ -17,12 +17,16 @@ public:
 	SLATE_BEGIN_ARGS(SNiagaraOverviewGraph) {}
 	SLATE_END_ARGS()
 
-	void Construct(const FArguments& InArgs, TSharedRef<FNiagaraSystemViewModel> InSystemViewModel);
+	void Construct(const FArguments& InArgs, TSharedRef<FNiagaraOverviewGraphViewModel> InViewModel);
+
+	~SNiagaraOverviewGraph();
+
+	virtual void Tick(const FGeometry& AllottedGeometry, const double InCurrentTime, const float InDeltaTime) override;
 
 private:
-	void GraphSelectionChanged(const TSet<UObject*>& SelectedNodes);
+	void ViewModelSelectionChanged();
 
-	void OverviewSelectionChanged();
+	void GraphSelectionChanged(const TSet<UObject*>& SelectedNodes);
 
 	/** Called to create context menu when right-clicking on graph */
 	FActionMenuContent OnCreateGraphActionMenu(UEdGraph* InGraph, const FVector2D& InNodePosition, const TArray<UEdGraphPin*>& InDraggedPins, bool bAutoExpand, SGraphEditor::FActionMenuClosed InOnMenuClosed);
@@ -32,11 +36,13 @@ private:
 	void OnNodeTitleCommitted(const FText& NewText, ETextCommit::Type CommitInfo, UEdGraphNode* NodeBeingChanged);
 
 	void CreateAddEmitterMenuContent(FMenuBuilder& MenuBuilder, UEdGraph* InGraph);
-	
+
 private:
-	TSharedPtr<FNiagaraSystemViewModel> SystemViewModel;
+	TSharedPtr<FNiagaraOverviewGraphViewModel> ViewModel;
 	TSharedPtr<SGraphEditor> GraphEditor;
 
-	bool bUpdatingOverviewSelectionFromGraph;
-	bool bUpdatingGraphSelectionFromOverview;
-};
+	bool bUpdatingViewModelSelectionFromGraph;
+	bool bUpdatingGraphSelectionFromViewModel;
+
+	int32 ZoomToFitFrameDelay;
+};
\ No newline at end of file
diff --git a/Engine/Plugins/FX/Niagara/Source/NiagaraEditorWidgets/Private/SNiagaraOverviewStack.cpp b/Engine/Plugins/FX/Niagara/Source/NiagaraEditorWidgets/Private/SNiagaraOverviewStack.cpp
new file mode 100644
--- /dev/null
+++ b/Engine/Plugins/FX/Niagara/Source/NiagaraEditorWidgets/Private/SNiagaraOverviewStack.cpp
@@ -0,0 +1,274 @@
+// Copyright 1998-2019 Epic Games, Inc. All Rights Reserved.
+
+#include "SNiagaraOverviewStack.h"
+#include "NiagaraSystemEditorData.h"
+#include "ViewModels/NiagaraSystemViewModel.h"
+#include "ViewModels/NiagaraEmitterHandleViewModel.h"
+#include "ViewModels/Stack/NiagaraStackEntry.h"
+#include "ViewModels/Stack/NiagaraStackItemGroup.h"
+#include "ViewModels/Stack/NiagaraStackItem.h"
+#include "ViewModels/Stack/NiagaraStackViewModel.h"
+#include "ViewModels/Stack/NiagaraStackSelection.h"
+#include "NiagaraEditorWidgetsStyle.h"
+#include "NiagaraEditorWidgetsUtilities.h"
+#include "SNiagaraStack.h"
+
+#include "Widgets/Layout/SScrollBox.h"
+
+class SNiagaraSystemOverviewEntryListRow : public STableRow<UNiagaraStackEntry*>
+{
+	SLATE_BEGIN_ARGS(SNiagaraSystemOverviewEntryListRow)
+	{}
+		SLATE_DEFAULT_SLOT(FArguments, Content)
+	SLATE_END_ARGS();
+
+	void Construct(const FArguments& InArgs, UNiagaraStackEntry* StackEntry, const TSharedRef<STableViewBase>& InOwnerTableView)
+	{
+		FSlateColor IconColor = FNiagaraEditorWidgetsStyle::Get().GetColor(FNiagaraStackEditorWidgetsUtilities::GetColorNameForExecutionCategory(StackEntry->GetExecutionCategoryName()));
+
+		TSharedPtr<SWidget> WrappedContent;
+		if (StackEntry->IsA<UNiagaraStackItem>())
+		{
+			WrappedContent =
+				SNew(SBorder)
+				.BorderImage(FEditorStyle::GetBrush("WhiteBrush"))
+				.BorderBackgroundColor(FNiagaraEditorWidgetsStyle::Get().GetColor("NiagaraEditor.Stack.Item.HeaderBackgroundColor"))
+				.Padding(0)
+				[
+					SNew(SBorder)
+					.BorderImage(this, &SNiagaraSystemOverviewEntryListRow::GetBorder)
+					.Padding(FMargin(8, 4, 8, 4))
+					[
+						InArgs._Content.Widget
+					]
+				];
+		}
+		else
+		{
+			WrappedContent =
+				SNew(SBorder)
+				.BorderImage(FEditorStyle::GetBrush("NoBrush"))
+				.Padding(0)
+				[
+					SNew(SBorder)
+					.BorderImage(this, &SNiagaraSystemOverviewEntryListRow::GetBorder)
+					.Padding(FMargin(2, 6, 8, 2))
+					[
+						InArgs._Content.Widget
+					]
+				];
+		}
+
+		STableRow<UNiagaraStackEntry*>::Construct(STableRow<UNiagaraStackEntry*>::FArguments()
+		[
+			SNew(SBorder)
+			.BorderImage(FEditorStyle::GetBrush("WhiteBrush"))
+			.BorderBackgroundColor(FNiagaraEditorWidgetsStyle::Get().GetColor("NiagaraEditor.Stack.Item.ContentBackgroundColor"))
+			.Padding(0)
+			[
+				SNew(SHorizontalBox)
+				+ SHorizontalBox::Slot()
+				.AutoWidth()
+				.VAlign(VAlign_Fill)
+				[
+					SNew(SBorder)
+					.BorderImage(FEditorStyle::GetBrush("WhiteBrush"))
+					.BorderBackgroundColor(FNiagaraEditorWidgetsStyle::Get().GetColor(FNiagaraStackEditorWidgetsUtilities::GetIconColorNameForExecutionCategory(StackEntry->GetExecutionCategoryName())))
+					.Padding(0)
+					[
+						SNew(SBox)
+						.WidthOverride(5)
+					]
+				]
+				+ SHorizontalBox::Slot()
+				[
+					SNew(SBorder)
+					.BorderImage(this, &SNiagaraSystemOverviewEntryListRow::GetBorder)
+					.Padding(FMargin(7, 0, 7, 4))
+					[
+						WrappedContent.ToSharedRef()
+					]
+				]
+			]
+		],
+		InOwnerTableView);
+	}
+};
+
+void SNiagaraOverviewStack::Construct(const FArguments& InArgs, UNiagaraStackViewModel& InStackViewModel, UNiagaraSystemSelectionViewModel& InOverviewSelectionViewModel)
+{
+	bRefreshEntryListPending = false;
+	bUpdatingOverviewSelectionFromStackSelection = false;
+	bUpdatingStackSelectionFromOverviewSelection = false;
+
+	StackViewModel = &InStackViewModel;
+	OverviewSelectionViewModel = &InOverviewSelectionViewModel;
+	OverviewSelectionViewModel->OnSelectionChanged().AddSP(this, &SNiagaraOverviewStack::SystemSelectionChanged);
+
+	ChildSlot
+	[
+		SAssignNew(EntryListView, SListView<UNiagaraStackEntry*>)
+		.ListItemsSource(&FlattenedEntryList)
+		.OnGenerateRow(this, &SNiagaraOverviewStack::OnGenerateRowForEntry)
+		.OnSelectionChanged(this, &SNiagaraOverviewStack::OnSelectionChanged)
+		.SelectionMode(ESelectionMode::Multi)
+	];
+
+	InStackViewModel.OnStructureChanged().AddSP(this, &SNiagaraOverviewStack::EntryStructureChanged);
+		
+	bRefreshEntryListPending = true;
+	RefreshEntryList();
+}
+
+SNiagaraOverviewStack::~SNiagaraOverviewStack()
+{
+	if (StackViewModel != nullptr)
+	{
+		StackViewModel->OnStructureChanged().RemoveAll(this);
+	}
+}
+
+void SNiagaraOverviewStack::Tick(const FGeometry& AllottedGeometry, const double InCurrentTime, const float InDeltaTime)
+{
+	if (bRefreshEntryListPending)
+	{
+		RefreshEntryList();
+		bRefreshEntryListPending = false;
+	}
+}
+
+void SNiagaraOverviewStack::AddEntriesRecursive(UNiagaraStackEntry& EntryToAdd, TArray<UNiagaraStackEntry*>& EntryList, const TArray<UClass*>& AcceptableClasses, TArray<UNiagaraStackEntry*> ParentChain)
+{
+	if (AcceptableClasses.ContainsByPredicate([&] (UClass* Class) { return EntryToAdd.IsA(Class); }))
+	{
+		EntryList.Add(&EntryToAdd);
+		EntryObjectKeyToParentChain.Add(FObjectKey(&EntryToAdd), ParentChain);
+		TArray<UNiagaraStackEntry*> Children;
+		EntryToAdd.GetFilteredChildren(Children);
+		ParentChain.Add(&EntryToAdd);
+		for (UNiagaraStackEntry* Child : Children)
+		{
+			checkf(Child != nullptr, TEXT("Stack entry had null child."));
+			AddEntriesRecursive(*Child, EntryList, AcceptableClasses, ParentChain);
+		}
+	}
+}
+
+void SNiagaraOverviewStack::RefreshEntryList()
+{
+	FlattenedEntryList.Empty();
+	EntryObjectKeyToParentChain.Empty();
+	TArray<UClass*> AcceptableClasses;
+	AcceptableClasses.Add(UNiagaraStackItemGroup::StaticClass());
+	AcceptableClasses.Add(UNiagaraStackItem::StaticClass());
+	for (UNiagaraStackEntry* RootEntry : StackViewModel->GetRootEntries())
+	{
+		checkf(RootEntry != nullptr, TEXT("Root entry was null."));
+		TArray<UNiagaraStackEntry*> RootChildren;
+		RootEntry->GetFilteredChildren(RootChildren);
+		for (UNiagaraStackEntry* RootChild : RootChildren)
+		{
+			checkf(RootEntry != nullptr, TEXT("Root entry child was null."));
+			TArray<UNiagaraStackEntry*> ParentChain;
+			AddEntriesRecursive(*RootChild, FlattenedEntryList, AcceptableClasses, ParentChain);
+		}
+	}
+	EntryListView->RequestListRefresh();
+}
+
+void SNiagaraOverviewStack::EntryStructureChanged()
+{
+	bRefreshEntryListPending = true;
+}
+
+TSharedRef<ITableRow> SNiagaraOverviewStack::OnGenerateRowForEntry(UNiagaraStackEntry* Item, const TSharedRef<STableViewBase>& OwnerTable)
+{
+	TSharedPtr<SWidget> Content;
+	if (Item->IsA<UNiagaraStackItem>())
+	{
+		Content = SNew(STextBlock)
+			.TextStyle(FNiagaraEditorWidgetsStyle::Get(), "NiagaraEditor.SystemOverview.ItemText")
+			.Text_UObject(Item, &UNiagaraStackEntry::GetDisplayName);
+	}
+	else
+	{
+		Content = SNew(SHorizontalBox)
+			+ SHorizontalBox::Slot()
+			.Padding(2, 0, 6, 0)
+			.VAlign(VAlign_Center)
+			.AutoWidth()
+			[
+				SNew(SImage)
+				//.Visibility(this, &SNiagaraStackTableRow::GetExecutionCategoryIconVisibility)
+				.Image(FNiagaraEditorWidgetsStyle::Get().GetBrush(FNiagaraStackEditorWidgetsUtilities::GetIconNameForExecutionSubcategory(Item->GetExecutionSubcategoryName(), true)))
+				.ColorAndOpacity(FNiagaraEditorWidgetsStyle::Get().GetColor(FNiagaraStackEditorWidgetsUtilities::GetIconColorNameForExecutionCategory(Item->GetExecutionCategoryName())))
+			]
+			+ SHorizontalBox::Slot()
+			[
+				SNew(STextBlock)
+				.TextStyle(FNiagaraEditorWidgetsStyle::Get(), "NiagaraEditor.SystemOverview.GroupHeaderText")
+				.Text_UObject(Item, &UNiagaraStackEntry::GetDisplayName)
+			];
+	}
+
+
+	return SNew(SNiagaraSystemOverviewEntryListRow, Item, OwnerTable)
+	[
+		Content.ToSharedRef()
+	];
+}
+
+void SNiagaraOverviewStack::OnSelectionChanged(UNiagaraStackEntry* InNewSelection, ESelectInfo::Type SelectInfo)
+{
+	if (bUpdatingStackSelectionFromOverviewSelection == false)
+	{
+		TGuardValue<bool> UpdateGuard(bUpdatingOverviewSelectionFromStackSelection, true);
+		TArray<UNiagaraStackEntry*> SelectedEntries;
+		EntryListView->GetSelectedItems(SelectedEntries);
+
+		TArray<UNiagaraStackEntry*> DeselectedEntries;
+
+		for (TWeakObjectPtr<UNiagaraStackEntry> PreviousSelectedEntry : PreviousSelection)
+		{
+			if (PreviousSelectedEntry.IsValid() && SelectedEntries.Contains(PreviousSelectedEntry.Get()) == false)
+			{
+				DeselectedEntries.Add(PreviousSelectedEntry.Get());
+			}
+		}
+
+		bool bClearCurrentSelection = FSlateApplication::Get().GetModifierKeys().IsControlDown() == false;
+		OverviewSelectionViewModel->UpdateSelectionFromEntries(SelectedEntries, DeselectedEntries, bClearCurrentSelection);
+
+		PreviousSelection.Empty();
+		for (UNiagaraStackEntry* SelectedEntry : SelectedEntries)
+		{
+			PreviousSelection.Add(SelectedEntry);
+		}
+	}
+}
+
+void SNiagaraOverviewStack::SystemSelectionChanged(UNiagaraSystemSelectionViewModel::ESelectionChangeSource SelectionChangeSource)
+{
+	if (bUpdatingOverviewSelectionFromStackSelection == false)
+	{
+		TGuardValue<bool> UpdateGuard(bUpdatingStackSelectionFromOverviewSelection, true);
+
+		TArray<UNiagaraStackEntry*> SelectedStackEntries;
+		EntryListView->GetSelectedItems(SelectedStackEntries);
+		TArray<UNiagaraStackEntry*> SelectedOverviewEntries = OverviewSelectionViewModel->GetSelectedEntries();
+
+		TArray<UNiagaraStackEntry*> EntriesToDeselect;
+		for (UNiagaraStackEntry* SelectedStackEntry : SelectedStackEntries)
+		{
+			if (SelectedOverviewEntries.Contains(SelectedStackEntry) == false)
+			{
+				EntriesToDeselect.Add(SelectedStackEntry);
+			}
+		}
+
+		for (UNiagaraStackEntry* EntryToDeselect : EntriesToDeselect)
+		{
+			EntryListView->SetItemSelection(EntryToDeselect, false);
+		}
+	}
+}
\ No newline at end of file
diff --git a/Engine/Plugins/FX/Niagara/Source/NiagaraEditorWidgets/Private/SNiagaraOverviewStackNode.cpp b/Engine/Plugins/FX/Niagara/Source/NiagaraEditorWidgets/Private/SNiagaraOverviewStackNode.cpp
new file mode 100644
--- /dev/null
+++ b/Engine/Plugins/FX/Niagara/Source/NiagaraEditorWidgets/Private/SNiagaraOverviewStackNode.cpp
@@ -0,0 +1,83 @@
+// Copyright 1998-2019 Epic Games, Inc. All Rights Reserved.
+
+#include "SNiagaraOverviewStackNode.h"
+#include "NiagaraOverviewNode.h"
+#include "ViewModels/NiagaraSystemViewModel.h"
+#include "ViewModels/NiagaraEmitterHandleViewModel.h"
+#include "SNiagaraOverviewStack.h"
+#include "NiagaraEditorModule.h"
+
+#include "Modules/ModuleManager.h"
+
+void SNiagaraOverviewStackNode::Construct(const FArguments& InArgs, UNiagaraOverviewNode* InNode)
+{
+	GraphNode = InNode;
+	OverviewStackNode = InNode;
+	if (OverviewStackNode->GetOwningSystem() != nullptr)
+	{
+		FNiagaraEditorModule& NiagaraEditorModule = FModuleManager::Get().LoadModuleChecked<FNiagaraEditorModule>("NiagaraEditor");
+		TSharedPtr<FNiagaraSystemViewModel> OwningSystemViewModel = NiagaraEditorModule.GetExistingViewModelForSystem(OverviewStackNode->GetOwningSystem());
+		if (OwningSystemViewModel.IsValid())
+		{
+			if (OverviewStackNode->GetEmitterHandleGuid().IsValid() == false)
+			{
+				StackViewModel = OwningSystemViewModel->GetSystemStackViewModel();
+			}
+			else
+			{
+				TSharedPtr<FNiagaraEmitterHandleViewModel> EmitterHandleviewModel = OwningSystemViewModel->GetEmitterHandleViewModelById(OverviewStackNode->GetEmitterHandleGuid());
+				if (EmitterHandleviewModel.IsValid())
+				{
+					StackViewModel = EmitterHandleviewModel->GetEmitterStackViewModel();
+				}
+			}
+			OverviewSelectionViewModel = OwningSystemViewModel->GetSelectionViewModel();
+		}
+	}
+	UpdateGraphNode();
+}
+
+TSharedRef<SWidget> SNiagaraOverviewStackNode::CreateNodeContentArea()
+{
+	TSharedPtr<SWidget> ContentWidget;
+	if (StackViewModel != nullptr && OverviewSelectionViewModel != nullptr)
+	{
+		ContentWidget = SNew(SBox)
+			.WidthOverride(210)
+			[
+				SNew(SNiagaraOverviewStack, *StackViewModel, *OverviewSelectionViewModel)
+			];
+	}
+	else
+	{
+		ContentWidget = SNullWidget::NullWidget;
+	}
+
+	// NODE CONTENT AREA
+	return SNew(SBorder)
+		.BorderImage( FEditorStyle::GetBrush("NoBorder") )
+		.HAlign(HAlign_Fill)
+		.VAlign(VAlign_Fill)
+		.Padding( FMargin(0,3) )
+		[
+			SNew(SHorizontalBox)
+			+SHorizontalBox::Slot()
+			.AutoWidth()
+			.HAlign(HAlign_Left)
+			[
+				// LEFT
+				SAssignNew(LeftNodeBox, SVerticalBox)
+			]
+			+SHorizontalBox::Slot()
+			[
+				ContentWidget.ToSharedRef()
+			]
+			+SHorizontalBox::Slot()
+			.AutoWidth()
+			.HAlign(HAlign_Right)
+			[
+				// RIGHT
+				SAssignNew(RightNodeBox, SVerticalBox)
+			]
+		];
+}
\ No newline at end of file
diff --git a/Engine/Plugins/FX/Niagara/Source/NiagaraEditorWidgets/Private/SNiagaraOverviewStackNode.h b/Engine/Plugins/FX/Niagara/Source/NiagaraEditorWidgets/Private/SNiagaraOverviewStackNode.h
new file mode 100644
--- /dev/null
+++ b/Engine/Plugins/FX/Niagara/Source/NiagaraEditorWidgets/Private/SNiagaraOverviewStackNode.h
@@ -0,0 +1,29 @@
+// Copyright 1998-2019 Epic Games, Inc. All Rights Reserved.
+
+#pragma once
+
+#include "SGraphNode.h"
+
+class UNiagaraOverviewNode;
+class UNiagaraStackViewModel;
+class UNiagaraSystemSelectionViewModel;
+
+class SNiagaraOverviewStackNode : public SGraphNode
+{
+public:
+	SLATE_BEGIN_ARGS(SNiagaraOverviewStackNode) {}
+	SLATE_END_ARGS()
+
+	/** Constructs this widget with InArgs */
+	void Construct(const FArguments& InArgs, UNiagaraOverviewNode* InNode);
+
+protected:
+	virtual TSharedRef<SWidget> CreateNodeContentArea() override;
+
+private:
+
+private:
+	UNiagaraOverviewNode* OverviewStackNode;
+	UNiagaraStackViewModel* StackViewModel;
+	UNiagaraSystemSelectionViewModel* OverviewSelectionViewModel;
+};
\ No newline at end of file
