Niagara - System overview MVP

diff --git a/Engine/Plugins/FX/Niagara/Source/NiagaraEditorWidgets/Private/SNiagaraStack.cpp b/Engine/Plugins/FX/Niagara/Source/NiagaraEditorWidgets/Private/SNiagaraStack.cpp
--- a/Engine/Plugins/FX/Niagara/Source/NiagaraEditorWidgets/Private/SNiagaraStack.cpp
+++ b/Engine/Plugins/FX/Niagara/Source/NiagaraEditorWidgets/Private/SNiagaraStack.cpp
@@ -56,6 +56,164 @@
 #include "Stack/SNiagaraStackSpacer.h"
 #include "ViewModels/Stack/NiagaraStackRoot.h"
 
+#define LOCTEXT_NAMESPACE "NiagaraStack"
+
+class SNiagaraStackEmitterHeader : public SCompoundWidget
+{
+public:
+	SLATE_BEGIN_ARGS(SNiagaraStackEmitterHeader) {}
+	SLATE_END_ARGS();
+
+	void Construct(const FArguments& InArgs, TSharedRef<FNiagaraEmitterHandleViewModel> InEmitterHandleViewModel)
+	{
+		EmitterHandleViewModel = InEmitterHandleViewModel;
+
+		ChildSlot
+		[
+			SNew(SVerticalBox)
+
+			//~ Enable check box, view source emitter button, and external header controls.
+			+ SVerticalBox::Slot()
+			.AutoHeight()
+			.HAlign(HAlign_Fill)
+			[
+				//~ Enabled
+				SNew(SHorizontalBox)
+				+ SHorizontalBox::Slot()
+				.AutoWidth()
+				.Padding(2)
+				[
+					SNew(SCheckBox)
+					.ToolTipText(LOCTEXT("EnabledToolTip", "Toggles whether this emitter is enabled. Disabled emitters don't simulate or render."))
+					.IsChecked(EmitterHandleViewModel.ToSharedRef(), &FNiagaraEmitterHandleViewModel::GetIsEnabledCheckState)
+					.OnCheckStateChanged(EmitterHandleViewModel.ToSharedRef(), &FNiagaraEmitterHandleViewModel::OnIsEnabledCheckStateChanged)
+				]
+				// Name and Source Emitter Name
+				+ SHorizontalBox::Slot()
+				.Padding(2)
+				[
+					SNew(SWrapBox)
+					.Clipping(EWidgetClipping::ClipToBoundsAlways) 
+					.UseAllottedWidth(true)
+					+ SWrapBox::Slot()
+					.Padding(3, 0)
+					[						
+						SAssignNew(EmitterNameTextBlock, SInlineEditableTextBlock)
+						.ToolTipText(this, &SNiagaraStackEmitterHeader::GetEmitterNameToolTip)
+						.Style(FNiagaraEditorStyle::Get(), "NiagaraEditor.HeadingInlineEditableText") 
+						.Clipping(EWidgetClipping::ClipToBoundsAlways)
+						.Text(EmitterHandleViewModel.ToSharedRef(), &FNiagaraEmitterHandleViewModel::GetNameText)
+						.OnTextCommitted(EmitterHandleViewModel.ToSharedRef(), &FNiagaraEmitterHandleViewModel::OnNameTextComitted)
+						.OnVerifyTextChanged(EmitterHandleViewModel.ToSharedRef(), &FNiagaraEmitterHandleViewModel::VerifyNameTextChanged)
+						.IsReadOnly(EmitterHandleViewModel->CanRenameEmitter() == false)
+					]
+					+ SWrapBox::Slot()
+					.Padding(3, 0)
+					[
+						SNew(STextBlock)
+						.ToolTipText(this, &SNiagaraStackEmitterHeader::GetEmitterNameToolTip)
+						.TextStyle(FNiagaraEditorStyle::Get(), "NiagaraEditor.SubduedHeadingTextBox")
+						.Clipping(EWidgetClipping::ClipToBoundsAlways)
+						.Text(EmitterHandleViewModel->GetEmitterViewModel(), &FNiagaraEmitterViewModel::GetParentNameText)
+						.Visibility(this, &SNiagaraStackEmitterHeader::GetSourceEmitterNameVisibility) 
+					]
+ 				]
+				// Open and Focus Source Emitter
+				+ SHorizontalBox::Slot()
+				.HAlign(HAlign_Fill)
+				.VAlign(VAlign_Center)
+				.AutoWidth()
+				.Padding(2)
+				[
+					SNew(SButton)
+					.IsFocusable(false)
+					.ToolTipText(LOCTEXT("OpenAndFocusParentEmitterToolTip", "Open and Focus Parent Emitter"))
+					.ButtonStyle(FEditorStyle::Get(), "HoverHintOnly")
+					.ForegroundColor(FNiagaraEditorWidgetsStyle::Get().GetColor("NiagaraEditor.Stack.ForegroundColor"))
+					.ContentPadding(2)
+					.OnClicked(this, &SNiagaraStackEmitterHeader::OpenParentEmitter)
+					.Visibility(this, &SNiagaraStackEmitterHeader::GetOpenSourceEmitterVisibility)
+					// GoToSource icon is 30x30px so we scale it down to stay in line with other 12x12px UI
+					.DesiredSizeScale(FVector2D(0.55f, 0.55f))
+					.Content()
+					[
+						SNew(SImage)
+						.Image(FNiagaraEditorWidgetsStyle::Get().GetBrush("NiagaraEditor.Stack.GoToSourceIcon"))
+						.ColorAndOpacity(FNiagaraEditorWidgetsStyle::Get().GetColor("NiagaraEditor.Stack.FlatButtonColor"))
+					]
+				]
+			]
+
+			//~ Stats
+			+ SVerticalBox::Slot()
+			.AutoHeight()
+			.HAlign(HAlign_Fill)
+			.Padding(2)
+			[
+				SNew(STextBlock)
+				.Text(EmitterHandleViewModel->GetEmitterViewModel(), &FNiagaraEmitterViewModel::GetStatsText)
+			]
+		];
+	}
+
+	void Tick(const FGeometry& AllottedGeometry, const double InCurrentTime, const float InDeltaTime)
+	{
+		if (EmitterHandleViewModel->GetIsRenamePending())
+		{
+			EmitterHandleViewModel->SetIsRenamePending(false);
+			EmitterNameTextBlock->EnterEditingMode();
+		}
+	}
+
+private:
+	FText GetEmitterNameToolTip() const
+	{
+		if (EmitterHandleViewModel->GetEmitterViewModel()->HasParentEmitter())
+		{
+			// We are looking at this Emitter in a System Asset and it has a valid parent Emitter
+			return FText::Format(LOCTEXT("EmitterNameAndPath", "{0}\nParent: {1}"), EmitterHandleViewModel->GetNameText(), EmitterHandleViewModel->GetEmitterViewModel()->GetParentPathNameText());
+		}
+		else
+		{
+			// We are looking at this Emitter in an Emitter Asset or we are looking at this Emitter in a System Asset and it does not have a valid parent Emitter
+			return EmitterHandleViewModel->GetNameText();
+		}
+	}
+
+	EVisibility GetSourceEmitterNameVisibility() const
+	{
+		bool bIsRenamed = false;
+		if (EmitterHandleViewModel->GetEmitterViewModel()->HasParentEmitter())
+		{
+			const FText CurrentNameText = EmitterHandleViewModel->GetNameText();
+			const FText ParentNameText = EmitterHandleViewModel->GetEmitterViewModel()->GetParentNameText();
+			bIsRenamed = CurrentNameText.EqualTo(ParentNameText) == false;
+		}
+		return bIsRenamed ? EVisibility::Visible : EVisibility::Collapsed;
+	}
+
+
+	FReply OpenParentEmitter()
+	{
+		UNiagaraEmitter* ParentEmitter = const_cast<UNiagaraEmitter*>(EmitterHandleViewModel->GetEmitterViewModel()->GetParentEmitter());
+		if (ParentEmitter != nullptr)
+		{
+			FAssetEditorManager::Get().OpenEditorForAsset(ParentEmitter);
+		}
+		return FReply::Handled();
+	}
+
+	EVisibility GetOpenSourceEmitterVisibility() const
+	{
+		return EmitterHandleViewModel->GetEmitterViewModel()->GetEmitter()->GetParent() != nullptr ? EVisibility::Visible : EVisibility::Collapsed;
+	}
+
+private:
+	TSharedPtr<FNiagaraEmitterHandleViewModel> EmitterHandleViewModel;
+
+	TSharedPtr<SInlineEditableTextBlock> EmitterNameTextBlock;
+};
+
 /** Contains data for a socket drag and drop operation in the StackEntry node. */
 class FNiagaraStackEntryDragDropOp : public FDecoratedDragDropOp
 {
@@ -77,8 +235,6 @@ private:
 	TArray<UNiagaraStackEntry*> DraggedEntries;
 };
 
-#define LOCTEXT_NAMESPACE "NiagaraStack"
-
 const float SpacerHeight = 6;
 
 const FText SNiagaraStack::OccurencesFormat = NSLOCTEXT("NiagaraStack", "OccurencesFound", "{0} / {1}");
@@ -90,16 +246,11 @@ void SNiagaraStack::Construct(const FArguments& InArgs, UNiagaraStackViewModel*
 	StackViewModel->OnSearchCompleted().AddSP(this, &SNiagaraStack::OnStackSearchComplete); 
 	NameColumnWidth = .3f;
 	ContentColumnWidth = .7f;
-	PinIsPinnedColor = FNiagaraEditorWidgetsStyle::Get().GetColor("NiagaraEditor.Stack.ForegroundColor");
-	PinIsUnpinnedColor = PinIsPinnedColor.Desaturate(.4f);
-	CurrentPinColor = StackViewModel->GetSystemViewModel()->GetIsEmitterPinned(StackViewModel->GetEmitterHandleViewModel()->AsShared()) ? PinIsPinnedColor: PinIsUnpinnedColor;
 	bNeedsJumpToNextOccurence = false;
 
-	ConstructHeaderWidget();
-	TSharedPtr<SWidget> HeaderBox;
 	ChildSlot
 	[
-		SAssignNew(HeaderBox, SVerticalBox)
+		SNew(SVerticalBox)
 		+ SVerticalBox::Slot()
 		.AutoHeight()
 		.Padding(0, 0, 0, 3)
@@ -108,7 +259,7 @@ void SNiagaraStack::Construct(const FArguments& InArgs, UNiagaraStackViewModel*
 			.BorderImage(FEditorStyle::GetBrush("ToolPanel.GroupBorder"))
 			.Padding(3)
 			[
-				HeaderWidget.ToSharedRef()
+				ConstructHeaderWidget()
 			]
 		]
 		+ SVerticalBox::Slot()
@@ -129,107 +280,6 @@ void SNiagaraStack::Construct(const FArguments& InArgs, UNiagaraStackViewModel*
 
 	StackTree->SetScrollOffset(StackViewModel->GetLastScrollPosition());
 
-	auto OnHeaderMouseButtonUp = [this](const FGeometry&, const FPointerEvent& MouseEvent) -> FReply
-	{
-		if (MouseEvent.GetEffectingButton() == EKeys::RightMouseButton)
-		{
-			FMenuBuilder MenuBuilder(true, nullptr);
-			MenuBuilder.BeginSection("EmitterInlineMenuActions", LOCTEXT("EmitterActions", "Emitter Actions"));
-			{
-				FUIAction Action(FExecuteAction::CreateSP(this, &SNiagaraStack::SetEmitterEnabled, !StackViewModel->GetEmitterHandleViewModel()->GetIsEnabled()),
-					FCanExecuteAction(),
-					FIsActionChecked::CreateSP(this, &SNiagaraStack::CheckEmitterEnabledStatus, true));
-				MenuBuilder.AddMenuEntry(
-					LOCTEXT("IsEnabled", "Is Enabled"),
-					LOCTEXT("ToggleEmitterEnabledToolTip", "Toggle emitter enabled/disabled state"),
-					FSlateIcon(),
-					Action,
-					NAME_None,
-					EUserInterfaceActionType::Check);
-
-				if (!GetEmitterNameIsReadOnly()) // Only allow renaming local copies of Emitters in Systems
-				{
-					MenuBuilder.AddMenuEntry(
-						LOCTEXT("RenameEmitter", "Rename Emitter"),
-						LOCTEXT("RenameEmitterToolTip", "Rename this local emitter copy"),
-						FSlateIcon(),
-						FUIAction(FExecuteAction::CreateSP(InlineEditableTextBlock.Get(), &SInlineEditableTextBlock::EnterEditingMode)));
-				}
-
-				MenuBuilder.AddMenuEntry(
-					LOCTEXT("RemoveSourceEmitter", "Remove Source Emitter"),
-					LOCTEXT("RemoveSourceEmitterToolTip", "Removes source information from this emitter, preventing inheritance of any further changes."),
-					FSlateIcon(),
-					FUIAction(
-						FExecuteAction::CreateUObject(StackViewModel, &UNiagaraStackViewModel::RemoveEmitterSource),
-						FCanExecuteAction::CreateUObject(StackViewModel, &UNiagaraStackViewModel::HasParentEmitter)));
-
-				MenuBuilder.AddMenuEntry(
-					LOCTEXT("ShowEmitterInContentBrowser", "Show in Content Browser"),
-					LOCTEXT("ShowEmitterInContentBrowserToolTip", "Show the emitter in this stack in the Content Browser"),
-					FSlateIcon(),
-					FUIAction(FExecuteAction::CreateSP(this, &SNiagaraStack::ShowEmitterInContentBrowser))); 
-
-				MenuBuilder.AddMenuEntry(
-					LOCTEXT("CollapseStack", "Collapse All"),
-					LOCTEXT("CollapseStackToolTip", "Collapses every row in the stack."),
-					FSlateIcon(),
-					FUIAction(FExecuteAction::CreateSP(this, &SNiagaraStack::CollapseAll)));
-				
-				TSharedPtr<FUICommandInfo> ExpandStackGroups = FNiagaraEditorModule::Get().Commands().CollapseStackToHeaders;
-				MenuBuilder.AddMenuEntry(
-					ExpandStackGroups->GetLabel(),
-					ExpandStackGroups->GetDescription(),
-					FSlateIcon(),
-					FUIAction(FExecuteAction::CreateUObject(StackViewModel, &UNiagaraStackViewModel::CollapseToHeaders)));
-			}
-			MenuBuilder.EndSection();
-			MenuBuilder.BeginSection("EmitterNavigateTo", 
-				FText::Format(LOCTEXT("EmitterNavigateTo", "Navigate to {0} Section:"), StackViewModel->GetEmitterHandleViewModel()->GetNameText()));
-			{
-				// Parse all children of root entries and if they are UNiagaraStackItemGroup then add a navigate menu entry
-				TArray<UNiagaraStackEntry*> EntriesToProcess(StackViewModel->GetRootEntries());
-				TArray<UNiagaraStackEntry*> RootChildren;
-				while (EntriesToProcess.Num() > 0)
-				{
-					UNiagaraStackEntry* EntryToProcess = EntriesToProcess[0];
-					EntriesToProcess.RemoveAtSwap(0);
-					EntryToProcess->GetUnfilteredChildren(RootChildren);
-				}
-				for (auto RootChild: RootChildren)
-				{
-					UNiagaraStackItemGroup* GroupChild = Cast<UNiagaraStackItemGroup>(RootChild);
-					if (GroupChild != nullptr)
-					{
-						MenuBuilder.AddMenuEntry(
-							RootChild->GetDisplayName(),
-							FText::Format(LOCTEXT("EmitterTooltip", "Navigate to {0}"), RootChild->GetDisplayName()),
-							FSlateIcon(),
-							FUIAction(FExecuteAction::CreateSP(this, &SNiagaraStack::NavigateTo, RootChild)));
-					}
-				}
-			}
-			MenuBuilder.EndSection();
-			
-			MenuBuilder.BeginSection("StackActions", LOCTEXT("StackActions", "Stack Actions"));
-			if (StackViewModel->HasDismissedStackIssues())
-			{
-				MenuBuilder.AddMenuEntry(
-					LOCTEXT("UndismissIssues", "Undismiss All Stack Issues"),
-					LOCTEXT("ShowAssetInContentBrowserToolTip", "Undismiss all issues that were previously dismissed for this stack, if any"),
-					FSlateIcon(),
-					FUIAction(FExecuteAction::CreateUObject(StackViewModel, &UNiagaraStackViewModel::UndismissAllIssues)));
-			}
-			MenuBuilder.EndSection();
-
-			FWidgetPath WidgetPath = MouseEvent.GetEventPath() != nullptr ? *MouseEvent.GetEventPath() : FWidgetPath();
-			FSlateApplication::Get().PushMenu(AsShared(), WidgetPath, MenuBuilder.MakeWidget(), MouseEvent.GetScreenSpacePosition(), FPopupTransitionEffect(FPopupTransitionEffect::ContextMenu));
-			return FReply::Handled();
-		}
-		return FReply::Unhandled();
-	};
-	HeaderBox->SetOnMouseButtonUp(FPointerEventHandler::CreateLambda(OnHeaderMouseButtonUp));
-	
 	SynchronizeTreeExpansion();
 }
 
@@ -253,172 +303,64 @@ void SNiagaraStack::SynchronizeTreeExpansion()
 	}
 }
 
-void SNiagaraStack::ConstructHeaderWidget()
+TSharedRef<SWidget> SNiagaraStack::ConstructHeaderWidget()
 {
-	if (!StackViewModel->GetEmitterHandleViewModel().IsValid() || !StackViewModel->GetSystemViewModel().IsValid())
-	{
-		HeaderWidget = SNullWidget::NullWidget;
-	}
-	else
-	{
-		SAssignNew(HeaderWidget, SVerticalBox)
-			//~ Enable check box, pin text box, view source emitter button, and external header controls.
-			+ SVerticalBox::Slot()
-			.AutoHeight()
-			.HAlign(HAlign_Fill)
+	return SNew(SVerticalBox)
+		//~ Top level object list view
+		+ SVerticalBox::Slot()
+		.AutoHeight()
+		.HAlign(HAlign_Fill)
+		[
+			SAssignNew(HeaderList, SListView<TSharedRef<UNiagaraStackViewModel::FTopLevelViewModel>>)
+			.ListItemsSource(&StackViewModel->GetTopLevelViewModels())
+			.SelectionMode(ESelectionMode::None)
+			.OnGenerateRow(this, &SNiagaraStack::OnGenerateRowForTopLevelObject)
+		]
+		
+		//~ Search, view options
+		+ SVerticalBox::Slot()
+		.AutoHeight()
+		.HAlign(HAlign_Fill)
+		.Padding(2, 4, 2, 4)
+		[
+			// Search box
+			SNew(SHorizontalBox)
+			+ SHorizontalBox::Slot()
 			[
-				//~ Enabled
-				SNew(SHorizontalBox)
-				+ SHorizontalBox::Slot()
-				.AutoWidth()
-				.Padding(2)
-				[
-					SNew(SCheckBox)
-					.ToolTipText(LOCTEXT("EnabledToolTip", "Toggles whether this emitter is enabled. Disabled emitters don't simulate or render."))
-					.IsChecked(StackViewModel->GetEmitterHandleViewModel()->AsShared(), &FNiagaraEmitterHandleViewModel::GetIsEnabledCheckState)
-					.OnCheckStateChanged(StackViewModel->GetEmitterHandleViewModel()->AsShared(), &FNiagaraEmitterHandleViewModel::OnIsEnabledCheckStateChanged)
-					.Visibility(this, &SNiagaraStack::GetEnableCheckboxVisibility) 
-				]
-				// Pin
-				+ SHorizontalBox::Slot()
-				.VAlign(VAlign_Center)
-				.AutoWidth()
-				[
-					SNew(SButton)
-					.IsFocusable(false)
-					.ToolTipText(LOCTEXT("PinToolTip", "Pin this emitter"))
-					.ButtonStyle(FEditorStyle::Get(), "HoverHintOnly")
-					.ForegroundColor(this, &SNiagaraStack::GetPinColor)
-					.ContentPadding(2)
-					.OnClicked(this, &SNiagaraStack::PinButtonPressed)
-					.Visibility(this, &SNiagaraStack::GetPinEmitterVisibility)
-					.Content()
-					[
-						SNew(STextBlock)
-						.Font(FEditorStyle::Get().GetFontStyle("FontAwesome.9"))
-						.Text(FText::FromString(FString(TEXT("\xf08d"))))
-						.RenderTransformPivot(FVector2D(.5f, .5f))
-					]
-				]
-				// Name and Source Emitter Name
-				+ SHorizontalBox::Slot()
-				.Padding(2)
-				[
-					SNew(SWrapBox)
-					.Clipping(EWidgetClipping::ClipToBoundsAlways) 
-					.UseAllottedWidth(true)
-					+ SWrapBox::Slot()
-					.Padding(3, 0)
-					[						
-						SAssignNew(InlineEditableTextBlock, SInlineEditableTextBlock)
-						.ToolTipText(this, &SNiagaraStack::GetEmitterNameToolTip)
-						.Style(FNiagaraEditorStyle::Get(), "NiagaraEditor.HeadingInlineEditableText") 
-						.Clipping(EWidgetClipping::ClipToBoundsAlways)
-						.Text(StackViewModel->GetEmitterHandleViewModel()->AsShared(), &FNiagaraEmitterHandleViewModel::GetNameText)
-						.OnTextCommitted(this, &SNiagaraStack::OnStackViewNameTextCommitted)
-						.OnVerifyTextChanged(StackViewModel->GetEmitterHandleViewModel()->AsShared(), &FNiagaraEmitterHandleViewModel::VerifyNameTextChanged)
-						.IsReadOnly(this, &SNiagaraStack::GetEmitterNameIsReadOnly)
-					]
-					+ SWrapBox::Slot()
-					.Padding(3, 0)
-					[
-						SNew(STextBlock)
-						.ToolTipText(this, &SNiagaraStack::GetEmitterNameToolTip)
-						.TextStyle(FNiagaraEditorStyle::Get(), "NiagaraEditor.SubduedHeadingTextBox")
-						.Clipping(EWidgetClipping::ClipToBoundsAlways)
-						.Text(this, &SNiagaraStack::GetSourceEmitterNameText)
-						.Visibility(this, &SNiagaraStack::GetSourceEmitterNameVisibility) 
-					]
- 				]
-				// Open and Focus Source Emitter
-				+ SHorizontalBox::Slot()
-				.HAlign(HAlign_Fill)
-				.VAlign(VAlign_Center)
-				.AutoWidth()
-				.Padding(2)
+				SAssignNew(SearchBox, SSearchBox)
+				.HintText(LOCTEXT("StackSearchBoxHint", "Search the stack"))
+				.SearchResultData(this, &SNiagaraStack::GetSearchResultData)
+				.IsSearching(this, &SNiagaraStack::GetIsSearching)
+				.OnTextChanged(this, &SNiagaraStack::OnSearchTextChanged)
+				.DelayChangeNotificationsWhileTyping(true)
+				.OnTextCommitted(this, &SNiagaraStack::OnSearchBoxTextCommitted)
+				.OnSearch(this, &SNiagaraStack::OnSearchBoxSearch)
+			]
+			// View options
+			+ SHorizontalBox::Slot()
+			.AutoWidth()
+			.Padding(4, 0, 0, 0)
+			[
+				SNew(SComboButton)
+				.ButtonStyle(FEditorStyle::Get(), "HoverHintOnly")
+				.ForegroundColor(FSlateColor::UseForeground())
+				.ToolTipText(LOCTEXT("ViewOptionsToolTip", "View Options"))
+				.OnGetMenuContent(this, &SNiagaraStack::GetViewOptionsMenu)
+				.ContentPadding(0)
+				.MenuPlacement(MenuPlacement_BelowRightAnchor)
+				.ButtonContent()
 				[
-					SNew(SButton)
-					.IsFocusable(false)
-					.ToolTipText(LOCTEXT("OpenAndFocusParentEmitterToolTip", "Open and Focus Parent Emitter"))
-					.ButtonStyle(FEditorStyle::Get(), "HoverHintOnly")
-					.ForegroundColor(this, &SNiagaraStack::GetPinColor)
-					.ContentPadding(2)
-					.OnClicked(this, &SNiagaraStack::OpenParentEmitter)
-					.Visibility(this, &SNiagaraStack::GetOpenSourceEmitterVisibility)
-					// GoToSource icon is 30x30px so we scale it down to stay in line with other 12x12px UI
-					.DesiredSizeScale(FVector2D(0.55f, 0.55f))
-					.Content()
+					SNew(SBox)
+					.HAlign(HAlign_Center)
+					.VAlign(VAlign_Center)
 					[
 						SNew(SImage)
-						.Image(FNiagaraEditorWidgetsStyle::Get().GetBrush("NiagaraEditor.Stack.GoToSourceIcon"))
+						.Image(FEditorStyle::GetBrush("GenericViewButton"))
 						.ColorAndOpacity(FNiagaraEditorWidgetsStyle::Get().GetColor("NiagaraEditor.Stack.FlatButtonColor"))
 					]
 				]
 			]
-
-			//~ Stats
-			+ SVerticalBox::Slot()
-			.AutoHeight()
-			.HAlign(HAlign_Fill)
-			.Padding(2)
-			[
-				SNew(STextBlock)
-				.Text(StackViewModel->GetEmitterHandleViewModel()->GetEmitterViewModel()->AsShared(), &FNiagaraEmitterViewModel::GetStatsText)
-			]
-			
-			//~ Search, view options
-			+ SVerticalBox::Slot()
-			.AutoHeight()
-			.HAlign(HAlign_Fill)
-			.Padding(2, 4, 2, 4)
-			[
-				// Search box
-				SNew(SHorizontalBox)
-				+ SHorizontalBox::Slot()
-				[
-					SAssignNew(SearchBox, SSearchBox)
-					.HintText(LOCTEXT("StackSearchBoxHint", "Search the stack"))
-					.SearchResultData(this, &SNiagaraStack::GetSearchResultData)
-					.IsSearching(this, &SNiagaraStack::GetIsSearching)
-					.OnTextChanged(this, &SNiagaraStack::OnSearchTextChanged)
-					.DelayChangeNotificationsWhileTyping(true)
-					.OnTextCommitted(this, &SNiagaraStack::OnSearchBoxTextCommitted)
-					.OnSearch(this, &SNiagaraStack::OnSearchBoxSearch)
-				]
-				// View options
-				+ SHorizontalBox::Slot()
-				.AutoWidth()
-				.Padding(4, 0, 0, 0)
-				[
-					SNew(SComboButton)
-					.ButtonStyle(FEditorStyle::Get(), "HoverHintOnly")
-					.ForegroundColor(FSlateColor::UseForeground())
-					.ToolTipText(LOCTEXT("ViewOptionsToolTip", "View Options"))
-					.OnGetMenuContent(this, &SNiagaraStack::GetViewOptionsMenu)
-					.ContentPadding(0)
-					.MenuPlacement(MenuPlacement_BelowRightAnchor)
-					.ButtonContent()
-					[
-						SNew(SBox)
-						.HAlign(HAlign_Center)
-						.VAlign(VAlign_Center)
-						[
-							SNew(SImage)
-							.Image(FEditorStyle::GetBrush("GenericViewButton"))
-							.ColorAndOpacity(FNiagaraEditorWidgetsStyle::Get().GetColor("NiagaraEditor.Stack.FlatButtonColor"))
-						]
-					]
-				]
-			];
-	}
-}
-
-FReply SNiagaraStack::PinButtonPressed()
-{
-	bool bNewPinnedState = !StackViewModel->GetSystemViewModel()->GetIsEmitterPinned(StackViewModel->GetEmitterHandleViewModel()->AsShared());
-	StackViewModel->GetSystemViewModel()->SetEmitterPinnedState(StackViewModel->GetEmitterHandleViewModel()->AsShared(), bNewPinnedState);
-	CurrentPinColor = bNewPinnedState ? PinIsPinnedColor : PinIsUnpinnedColor;
-	return FReply::Handled();
+		];
 }
 
 void SNiagaraStack::OnSearchTextChanged(const FText& SearchText)
@@ -530,67 +472,19 @@ bool SNiagaraStack::IsEntryFocusedInSearch(UNiagaraStackEntry* Entry) const
 	return false;
 }
 
-FReply SNiagaraStack::OpenParentEmitter() 
+void SNiagaraStack::ShowEmitterInContentBrowser(TWeakPtr<FNiagaraEmitterHandleViewModel> EmitterHandleViewModelWeak)
 {
-	if (StackViewModel && StackViewModel->GetEmitterHandleViewModel().IsValid() && StackViewModel->GetEmitterHandleViewModel()->GetEmitterViewModel().IsValid())
+	TSharedPtr<FNiagaraEmitterHandleViewModel> EmitterHandleViewModel = EmitterHandleViewModelWeak.Pin();
+	if (EmitterHandleViewModel.IsValid())
 	{
-		UNiagaraEmitter* ParentEmitter = const_cast<UNiagaraEmitter*>(StackViewModel->GetEmitterHandleViewModel()->GetEmitterViewModel()->GetParentEmitter());
-		if (ParentEmitter != nullptr)
+		FContentBrowserModule& ContentBrowserModule = FModuleManager::Get().LoadModuleChecked<FContentBrowserModule>(TEXT("ContentBrowser"));
+		TArray<FAssetData> Assets;
+		if (EmitterHandleViewModel->GetEmitterViewModel()->HasParentEmitter())
 		{
-			FAssetEditorManager::Get().OpenEditorForAsset(ParentEmitter);
+			Assets.Add(FAssetData(EmitterHandleViewModel->GetEmitterViewModel()->GetParentEmitter()));
+			ContentBrowserModule.Get().SyncBrowserToAssets(Assets);
 		}
 	}
-	return FReply::Handled();
-}
-
-EVisibility SNiagaraStack::GetEnableCheckboxVisibility() const
-{
-	return StackViewModel->GetSystemViewModel()->GetEditMode() == ENiagaraSystemViewModelEditMode::SystemAsset ? EVisibility::Visible : EVisibility::Collapsed;
-}
-
-EVisibility SNiagaraStack::GetPinEmitterVisibility() const
-{
-	return StackViewModel->GetSystemViewModel()->GetEditMode() == ENiagaraSystemViewModelEditMode::SystemAsset ? EVisibility::Visible : EVisibility::Collapsed;
-}
-
-EVisibility SNiagaraStack::GetOpenSourceEmitterVisibility() const
-{
-	return StackViewModel->HasParentEmitter() ? EVisibility::Visible : EVisibility::Collapsed;
-}
-
-bool SNiagaraStack::GetEmitterNameIsReadOnly() const
-{
-	if (StackViewModel && StackViewModel->GetSystemViewModel().IsValid())
-	{
-		return StackViewModel->GetSystemViewModel()->GetEditMode() == ENiagaraSystemViewModelEditMode::EmitterAsset;
-	}
-	return true;
-}
-
-void SNiagaraStack::SetEmitterEnabled(bool bIsEnabled)
-{
-	StackViewModel->GetEmitterHandleViewModel()->SetIsEnabled(bIsEnabled);
-}
-
-bool SNiagaraStack::CheckEmitterEnabledStatus(bool bIsEnabled)
-{
-	return StackViewModel->GetEmitterHandleViewModel()->GetIsEnabled() == bIsEnabled;
-}
-
-void SNiagaraStack::ShowEmitterInContentBrowser()
-{
-	FContentBrowserModule& ContentBrowserModule = FModuleManager::Get().LoadModuleChecked<FContentBrowserModule>(TEXT("ContentBrowser"));
-	TArray<FAssetData> Assets;
-	if (StackViewModel->HasParentEmitter())
-	{
-		Assets.Add(FAssetData(StackViewModel->GetEmitterHandleViewModel()->GetEmitterViewModel()->GetParentEmitter()));
-		ContentBrowserModule.Get().SyncBrowserToAssets(Assets);
-	}
-}
-
-void SNiagaraStack::NavigateTo(UNiagaraStackEntry* Item)
-{
-	StackTree->RequestScrollIntoView(Item);
 }
 
 void CollapseEntriesRecursive(TArray<UNiagaraStackEntry*> Entries)
@@ -767,11 +661,6 @@ FSlateColor SNiagaraStack::GetTextColorForItem(UNiagaraStackEntry* Item) const
 	return FSlateColor::UseForeground();
 }
 
-FSlateColor SNiagaraStack::GetPinColor() const
-{
-	return CurrentPinColor;
-}
-
 TSharedRef<ITableRow> SNiagaraStack::OnGenerateRowForStackItem(UNiagaraStackEntry* Item, const TSharedRef<STableViewBase>& OwnerTable)
 {
 	TSharedRef<SNiagaraStackTableRow> Container = ConstructContainerForItem(Item);
@@ -780,6 +669,110 @@ TSharedRef<ITableRow> SNiagaraStack::OnGenerateRowForStackItem(UNiagaraStackEntr
 	return Container;
 }
 
+TSharedRef<ITableRow> SNiagaraStack::OnGenerateRowForTopLevelObject(TSharedRef<UNiagaraStackViewModel::FTopLevelViewModel> Item, const TSharedRef<STableViewBase>& OwnerTable)
+{
+	TSharedPtr<SWidget> Content;
+	if (Item->SystemViewModel.IsValid())
+	{
+		Content = SNew(STextBlock)
+			.TextStyle(FNiagaraEditorStyle::Get(), "NiagaraEditor.HeadingTextBlock")
+			.Text(Item->SystemViewModel.ToSharedRef(), &FNiagaraSystemViewModel::GetDisplayName);
+	}
+	else if(Item->EmitterHandleViewModel.IsValid())
+	{
+		Content = SNew(SNiagaraStackEmitterHeader, Item->EmitterHandleViewModel.ToSharedRef());
+	}
+
+	Content->SetOnMouseButtonUp(FPointerEventHandler::CreateSP(this, &SNiagaraStack::OnTopLevelRowMouseButtonDown, TWeakPtr<UNiagaraStackViewModel::FTopLevelViewModel>(Item)));
+
+	return SNew(STableRow<TSharedRef<UNiagaraStackViewModel::FTopLevelViewModel>>, OwnerTable)
+		[
+			Content.ToSharedRef()
+		];
+}
+
+FReply SNiagaraStack::OnTopLevelRowMouseButtonDown(const FGeometry&, const FPointerEvent& MouseEvent, TWeakPtr<UNiagaraStackViewModel::FTopLevelViewModel> TopLevelViewModelWeak)
+{
+	TSharedPtr<UNiagaraStackViewModel::FTopLevelViewModel> TopLevelViewModel = TopLevelViewModelWeak.Pin();
+	if (TopLevelViewModel.IsValid(), MouseEvent.GetEffectingButton() == EKeys::RightMouseButton)
+	{
+		FMenuBuilder MenuBuilder(true, nullptr);
+
+		if (TopLevelViewModel->EmitterHandleViewModel.IsValid())
+		{
+			TSharedRef<FNiagaraEmitterHandleViewModel> EmitterHandleViewModel = TopLevelViewModel->EmitterHandleViewModel.ToSharedRef();
+			MenuBuilder.BeginSection("EmitterInlineMenuActions", LOCTEXT("EmitterActions", "Emitter Actions"));
+			{
+				FUIAction Action(FExecuteAction::CreateSP(EmitterHandleViewModel, &FNiagaraEmitterHandleViewModel::SetIsEnabled, !EmitterHandleViewModel->GetIsEnabled()),
+					FCanExecuteAction(),
+					FIsActionChecked::CreateSP(EmitterHandleViewModel, &FNiagaraEmitterHandleViewModel::GetIsEnabled));
+				MenuBuilder.AddMenuEntry(
+					LOCTEXT("IsEnabled", "Is Enabled"),
+					LOCTEXT("ToggleEmitterEnabledToolTip", "Toggle emitter enabled/disabled state"),
+					FSlateIcon(),
+					Action,
+					NAME_None,
+					EUserInterfaceActionType::Check);
+
+				if (EmitterHandleViewModel->CanRenameEmitter()) // Only allow renaming local copies of Emitters in Systems
+				{
+					MenuBuilder.AddMenuEntry(
+						LOCTEXT("RenameEmitter", "Rename Emitter"),
+						LOCTEXT("RenameEmitterToolTip", "Rename this local emitter copy"),
+						FSlateIcon(),
+						FUIAction(FExecuteAction::CreateSP(EmitterHandleViewModel, &FNiagaraEmitterHandleViewModel::SetIsRenamePending, true)));
+				}
+
+				MenuBuilder.AddMenuEntry(
+					LOCTEXT("RemoveSourceEmitter", "Remove Source Emitter"),
+					LOCTEXT("RemoveSourceEmitterToolTip", "Removes source information from this emitter, preventing inheritance of any further changes."),
+					FSlateIcon(),
+					FUIAction(
+						FExecuteAction::CreateSP(EmitterHandleViewModel->GetEmitterViewModel(), &FNiagaraEmitterViewModel::RemoveParentEmitter),
+						FCanExecuteAction::CreateSP(EmitterHandleViewModel->GetEmitterViewModel(), &FNiagaraEmitterViewModel::HasParentEmitter)));
+
+				MenuBuilder.AddMenuEntry(
+					LOCTEXT("ShowEmitterInContentBrowser", "Show in Content Browser"),
+					LOCTEXT("ShowEmitterInContentBrowserToolTip", "Show the emitter in this stack in the Content Browser"),
+					FSlateIcon(),
+					FUIAction(FExecuteAction::CreateSP(this, &SNiagaraStack::ShowEmitterInContentBrowser, TWeakPtr<FNiagaraEmitterHandleViewModel>(EmitterHandleViewModel))));
+			}
+			MenuBuilder.EndSection();
+		}
+
+		MenuBuilder.BeginSection("StackActions", LOCTEXT("StackActions", "Stack Actions"));
+		{
+			if (StackViewModel->HasDismissedStackIssues())
+			{
+				MenuBuilder.AddMenuEntry(
+					LOCTEXT("UndismissIssues", "Undismiss All Stack Issues"),
+					LOCTEXT("ShowAssetInContentBrowserToolTip", "Undismiss all issues that were previously dismissed for this stack, if any"),
+					FSlateIcon(),
+					FUIAction(FExecuteAction::CreateUObject(StackViewModel, &UNiagaraStackViewModel::UndismissAllIssues)));
+			}
+
+			MenuBuilder.AddMenuEntry(
+				LOCTEXT("CollapseStack", "Collapse All"),
+				LOCTEXT("CollapseStackToolTip", "Collapses every row in the stack."),
+				FSlateIcon(),
+				FUIAction(FExecuteAction::CreateSP(this, &SNiagaraStack::CollapseAll)));
+
+			TSharedPtr<FUICommandInfo> CollapseToHeadersCommand = FNiagaraEditorModule::Get().Commands().CollapseStackToHeaders;
+			MenuBuilder.AddMenuEntry(
+				CollapseToHeadersCommand->GetLabel(),
+				CollapseToHeadersCommand->GetDescription(),
+				FSlateIcon(),
+				FUIAction(FExecuteAction::CreateUObject(StackViewModel, &UNiagaraStackViewModel::CollapseToHeaders)));
+		}
+		MenuBuilder.EndSection();
+
+		FWidgetPath WidgetPath = MouseEvent.GetEventPath() != nullptr ? *MouseEvent.GetEventPath() : FWidgetPath();
+		FSlateApplication::Get().PushMenu(AsShared(), WidgetPath, MenuBuilder.MakeWidget(), MouseEvent.GetScreenSpacePosition(), FPopupTransitionEffect(FPopupTransitionEffect::ContextMenu));
+		return FReply::Handled();
+	}
+	return FReply::Unhandled();
+}
+
 TSharedRef<SNiagaraStackTableRow> SNiagaraStack::ConstructContainerForItem(UNiagaraStackEntry* Item)
 {
 	float LeftContentPadding = 4;
@@ -1055,44 +1048,7 @@ void SNiagaraStack::StackStructureChanged()
 {
 	SynchronizeTreeExpansion();
 	StackTree->RequestTreeRefresh();
-}
-
-FText SNiagaraStack::GetSourceEmitterNameText() const 
-{
-	return StackViewModel->GetEmitterHandleViewModel()->GetEmitterViewModel()->GetParentNameText();
-}
-
-FText SNiagaraStack::GetEmitterNameToolTip() const
-{
-	if (StackViewModel->HasParentEmitter())
-	{
-		// We are looking at this Emitter in a System Asset and it has a valid parent Emitter
-		TSharedPtr<FNiagaraEmitterHandleViewModel> ThisViewModel = StackViewModel->GetEmitterHandleViewModel();
-		return FText::Format(LOCTEXT("EmitterNameAndPath", "{0}\nParent: {1}"), ThisViewModel->GetNameText(), ThisViewModel->GetEmitterViewModel()->GetParentPathNameText());
-	}
-	else
-	{
-		// We are looking at this Emitter in an Emitter Asset or we are looking at this Emitter in a System Asset and it does not have a valid parent Emitter
-		return StackViewModel->GetEmitterHandleViewModel()->GetNameText();
-	}
-}
-
-void SNiagaraStack::OnStackViewNameTextCommitted(const FText& InText, ETextCommit::Type CommitInfo) const
-{
-	const FGuid EditedEmitterGuid = StackViewModel->GetEmitterHandleViewModel()->GetEmitterHandle()->GetId();
-	StackViewModel->GetSystemViewModel()->EmitterNameChanged(InText, EditedEmitterGuid);
-}
-
-EVisibility SNiagaraStack::GetSourceEmitterNameVisibility() const
-{
-	return StackViewModel->HasParentEmitter() && GetIsEmitterRenamed() ? EVisibility::Visible : EVisibility::Collapsed;
-}
-
-bool SNiagaraStack::GetIsEmitterRenamed() const
-{
-	const FText CurrentNameText = StackViewModel->GetEmitterHandleViewModel()->GetNameText();
-	const FText SourceNameText = StackViewModel->GetEmitterHandleViewModel()->GetEmitterViewModel()->GetParentNameText();
-	return !CurrentNameText.EqualTo(SourceNameText);
+	HeaderList->RequestListRefresh();
 }
 
 #undef LOCTEXT_NAMESPACE
diff --git a/Engine/Plugins/FX/Niagara/Source/NiagaraEditorWidgets/Private/SNiagaraStack.h b/Engine/Plugins/FX/Niagara/Source/NiagaraEditorWidgets/Private/SNiagaraStack.h
--- a/Engine/Plugins/FX/Niagara/Source/NiagaraEditorWidgets/Private/SNiagaraStack.h
+++ b/Engine/Plugins/FX/Niagara/Source/NiagaraEditorWidgets/Private/SNiagaraStack.h
@@ -2,7 +2,6 @@
 
 #pragma once
 
-#include "EditorUndoClient.h"
 #include "Widgets/SCompoundWidget.h"
 #include "Widgets/DeclarativeSyntaxSupport.h"
 #include "Widgets/Views/STableRow.h"
@@ -11,14 +10,14 @@
 #include "ViewModels/Stack/NiagaraStackEntry.h"
 #include "Widgets/Input/SSearchBox.h"
 #include "Widgets/Text/SInlineEditableTextBlock.h"
-
+#include "ViewModels/Stack/NiagaraStackViewModel.h"
 
 class UNiagaraStackViewModel;
 class SNiagaraStackTableRow;
 class SSearchBox;
 class FReply;
 
-class SNiagaraStack : public SCompoundWidget, public FEditorUndoClient
+class SNiagaraStack : public SCompoundWidget
 {
 public:
 	SLATE_BEGIN_ARGS(SNiagaraStack)
@@ -50,6 +49,10 @@ private:
 
 	TSharedRef<ITableRow> OnGenerateRowForStackItem(UNiagaraStackEntry* Item, const TSharedRef<STableViewBase>& OwnerTable);
 
+	TSharedRef<ITableRow> OnGenerateRowForTopLevelObject(TSharedRef<UNiagaraStackViewModel::FTopLevelViewModel> Item, const TSharedRef<STableViewBase>& OwnerTable);
+
+	FReply OnTopLevelRowMouseButtonDown(const FGeometry&, const FPointerEvent& MouseEvent, TWeakPtr<UNiagaraStackViewModel::FTopLevelViewModel> TopLevelViewModelWeak);
+
 	TSharedRef<SNiagaraStackTableRow> ConstructContainerForItem(UNiagaraStackEntry* Item);
 
 	FRowWidgets ConstructNameAndValueWidgetsForItem(UNiagaraStackEntry* Item, TSharedRef<SNiagaraStackTableRow> Container);
@@ -69,21 +72,7 @@ private:
 
 	virtual void Tick(const FGeometry& AllottedGeometry, const double InCurrentTime, const float InDeltaTime) override;
 
-	void ConstructHeaderWidget();
-	FSlateColor GetPinColor() const;
-	FReply PinButtonPressed();
-	FReply OpenParentEmitter();
-	EVisibility GetEnableCheckboxVisibility() const;
-	EVisibility GetPinEmitterVisibility() const;
-	EVisibility GetOpenSourceEmitterVisibility() const;
-
-	// source name handling
-	bool GetEmitterNameIsReadOnly() const;
-	FText GetSourceEmitterNameText() const;
-	FText GetEmitterNameToolTip() const;
-	void OnStackViewNameTextCommitted(const FText& InText, ETextCommit::Type CommitInfo) const;
-	EVisibility GetSourceEmitterNameVisibility() const; 
-	bool GetIsEmitterRenamed() const;
+	TSharedRef<SWidget> ConstructHeaderWidget();
 
 	// ~stack search stuff
 	void OnSearchTextChanged(const FText& SearchText);
@@ -100,9 +89,7 @@ private:
 	bool IsEntryFocusedInSearch(UNiagaraStackEntry* Entry) const;
 	
 	// Inline menu commands
-	void SetEmitterEnabled(bool bIsEnabled);
-	bool CheckEmitterEnabledStatus(bool bIsEnabled);
-	void ShowEmitterInContentBrowser();
+	void ShowEmitterInContentBrowser(TWeakPtr<FNiagaraEmitterHandleViewModel> EmitterHandleViewModelWeak);
 	void NavigateTo(UNiagaraStackEntry* Item);
 	void CollapseAll();
 
@@ -120,20 +107,11 @@ private:
 
 	TSharedPtr<STreeView<UNiagaraStackEntry*>> StackTree;
 
+	TSharedPtr<SListView<TSharedRef<UNiagaraStackViewModel::FTopLevelViewModel>>> HeaderList;
+
 	float NameColumnWidth;
 
 	float ContentColumnWidth;
-	
-	TSharedPtr<SWidget> HeaderWidget;
-
-	FLinearColor PinIsPinnedColor;
-	
-	FLinearColor PinIsUnpinnedColor;
-
-	FLinearColor CurrentPinColor;
-	
-	// emitter name textblock
-	TSharedPtr<SInlineEditableTextBlock> InlineEditableTextBlock;
 
 	// ~ search stuff
 	TSharedPtr<SSearchBox> SearchBox;
diff --git a/Engine/Plugins/FX/Niagara/Source/NiagaraEditorWidgets/Private/SNiagaraStackEntryWidget.cpp b/Engine/Plugins/FX/Niagara/Source/NiagaraEditorWidgets/Private/SNiagaraStackEntryWidget.cpp
--- a/Engine/Plugins/FX/Niagara/Source/NiagaraEditorWidgets/Private/SNiagaraStackEntryWidget.cpp
+++ b/Engine/Plugins/FX/Niagara/Source/NiagaraEditorWidgets/Private/SNiagaraStackEntryWidget.cpp
@@ -1,11 +1,100 @@
 // Copyright 1998-2019 Epic Games, Inc. All Rights Reserved.
 
 #include "SNiagaraStackEntryWidget.h"
-#include "ViewModels/Stack/NiagaraStackViewModel.h"
 #include "NiagaraEditorWidgetsStyle.h"
+#include "Widgets/Layout/SWrapBox.h"
 
 #define LOCTEXT_NAMESPACE "SNiagaraStackEntryWidget"
 
+void SNiagaraStackDisplayName::Construct(const FArguments& InArgs, UNiagaraStackEntry& InStackEntry, UNiagaraStackViewModel& InStackViewModel, FName InTextStyleName)
+{
+	StackEntry = &InStackEntry;
+	StackViewModel = &InStackViewModel;
+	TextStyleName = InTextStyleName;
+
+	ColorAndOpacity = InArgs._ColorAndOpacity;
+
+	StackViewModel->OnStructureChanged().AddSP(this, &SNiagaraStackDisplayName::StackViewModelStructureChanged);
+
+	ChildSlot
+	[
+		SAssignNew(Container, SBox)
+		[
+			ConstructChildren()
+		]
+	];
+}
+
+TSharedRef<SWidget> SNiagaraStackDisplayName::ConstructChildren()
+{
+	TSharedRef<STextBlock> BaseNameWidget = SNew(STextBlock)
+		.TextStyle(FNiagaraEditorWidgetsStyle::Get(), TextStyleName)
+		.ToolTipText_UObject(StackEntry, &UNiagaraStackEntry::GetTooltipText)
+		.Text_UObject(StackEntry, &UNiagaraStackEntry::GetDisplayName)
+		.HighlightText_UObject(StackViewModel, &UNiagaraStackViewModel::GetCurrentSearchText)
+		.ColorAndOpacity(ColorAndOpacity)
+		.IsEnabled_UObject(StackEntry, &UNiagaraStackEntry::GetIsEnabled);
+
+	if (StackViewModel->GetTopLevelViewModels().Num() == 1)
+	{
+		TopLevelViewModelCountAtLastConstruction = 1;
+		return BaseNameWidget;
+	}
+	else if (StackViewModel->GetTopLevelViewModels().Num() > 1)
+	{
+		TopLevelViewModelCountAtLastConstruction = StackViewModel->GetTopLevelViewModels().Num();
+		TSharedPtr<UNiagaraStackViewModel::FTopLevelViewModel> TopLevelViewModel = StackViewModel->GetTopLevelViewModelForEntry(*StackEntry);
+		if(TopLevelViewModel.IsValid())
+		{
+			return SNew(SWrapBox)
+				.Clipping(EWidgetClipping::ClipToBoundsAlways)
+				.UseAllottedWidth(true)
+				+ SWrapBox::Slot()
+				[
+					SNew(STextBlock)
+					.TextStyle(FNiagaraEditorWidgetsStyle::Get(), TextStyleName)
+					.ToolTipText_UObject(StackEntry, &UNiagaraStackEntry::GetTooltipText)
+					.Text(this, &SNiagaraStackDisplayName::GetTopLevelDisplayName, TWeakPtr<UNiagaraStackViewModel::FTopLevelViewModel>(TopLevelViewModel))
+					.HighlightText_UObject(StackViewModel, &UNiagaraStackViewModel::GetCurrentSearchText)
+					.ColorAndOpacity(ColorAndOpacity)
+					.IsEnabled_UObject(StackEntry, &UNiagaraStackEntry::GetIsEnabled)
+				]
+				+ SWrapBox::Slot()
+				[
+					BaseNameWidget
+				];
+		}
+	}
+	TopLevelViewModelCountAtLastConstruction = -1;
+	return SNullWidget::NullWidget;
+}
+
+FText SNiagaraStackDisplayName::GetTopLevelDisplayName(TWeakPtr<UNiagaraStackViewModel::FTopLevelViewModel> TopLevelViewModelWeak) const
+{
+	TSharedPtr<UNiagaraStackViewModel::FTopLevelViewModel> TopLevelViewModel = TopLevelViewModelWeak.Pin();
+	if (TopLevelViewModel.IsValid())
+	{
+		if (TopLevelViewModel->GetDisplayName().IdenticalTo(TopLevelDisplayNameCache) == false)
+		{
+			TopLevelDisplayNameCache = TopLevelViewModel->GetDisplayName();
+			TopLevelDisplayNameFormattedCache = FText::Format(LOCTEXT("TopLevelDisplayNameFormat", "{0} - "), TopLevelDisplayNameCache);
+		}
+	}
+	else
+	{
+		TopLevelDisplayNameFormattedCache = FText();
+	}
+	return TopLevelDisplayNameFormattedCache;
+}
+
+void SNiagaraStackDisplayName::StackViewModelStructureChanged()
+{
+	if (StackViewModel->GetTopLevelViewModels().Num() != TopLevelViewModelCountAtLastConstruction)
+	{
+		Container->SetContent(ConstructChildren());
+	}
+}
+
 FSlateColor SNiagaraStackEntryWidget::GetTextColorForSearch() const
 {
 	if (IsCurrentSearchMatch())
diff --git a/Engine/Plugins/FX/Niagara/Source/NiagaraEditorWidgets/Private/SNiagaraStackEntryWidget.h b/Engine/Plugins/FX/Niagara/Source/NiagaraEditorWidgets/Private/SNiagaraStackEntryWidget.h
--- a/Engine/Plugins/FX/Niagara/Source/NiagaraEditorWidgets/Private/SNiagaraStackEntryWidget.h
+++ b/Engine/Plugins/FX/Niagara/Source/NiagaraEditorWidgets/Private/SNiagaraStackEntryWidget.h
@@ -4,10 +4,41 @@
 
 #include "Widgets/SCompoundWidget.h"
 #include "Widgets/DeclarativeSyntaxSupport.h"
+#include "ViewModels/Stack/NiagaraStackViewModel.h"
 
 class UNiagaraStackViewModel;
 class UNiagaraStackEntry;
 
+class SNiagaraStackDisplayName : public SCompoundWidget
+{
+public:
+	SLATE_BEGIN_ARGS(SNiagaraStackDisplayName) { }
+		SLATE_ATTRIBUTE(FSlateColor, ColorAndOpacity);
+	SLATE_END_ARGS();
+
+	void Construct(const FArguments& InArgs, UNiagaraStackEntry& InStackEntry, UNiagaraStackViewModel& InStackViewModel, FName InTextStyleName);
+
+private:
+	TSharedRef<SWidget> ConstructChildren();
+
+	FText GetTopLevelDisplayName(TWeakPtr<UNiagaraStackViewModel::FTopLevelViewModel> TopLevelViewModelWeak) const;
+
+	void StackViewModelStructureChanged();
+
+private:
+	UNiagaraStackEntry* StackEntry;
+	UNiagaraStackViewModel* StackViewModel;
+	FName TextStyleName;
+
+	TAttribute<FSlateColor> ColorAndOpacity;
+
+	TSharedPtr<SBox> Container;
+
+	mutable FText TopLevelDisplayNameCache;
+	mutable FText TopLevelDisplayNameFormattedCache;
+	int32 TopLevelViewModelCountAtLastConstruction;
+};
+
 class SNiagaraStackEntryWidget : public SCompoundWidget
 {
 public:
diff --git a/Engine/Plugins/FX/Niagara/Source/NiagaraEditorWidgets/Private/Stack/SNiagaraStackEmitterPropertiesItem.cpp b/Engine/Plugins/FX/Niagara/Source/NiagaraEditorWidgets/Private/Stack/SNiagaraStackEmitterPropertiesItem.cpp
--- a/Engine/Plugins/FX/Niagara/Source/NiagaraEditorWidgets/Private/Stack/SNiagaraStackEmitterPropertiesItem.cpp
+++ b/Engine/Plugins/FX/Niagara/Source/NiagaraEditorWidgets/Private/Stack/SNiagaraStackEmitterPropertiesItem.cpp
@@ -25,12 +25,8 @@ void SNiagaraStackEmitterPropertiesItem::Construct(const FArguments& InArgs, UNi
 		+ SHorizontalBox::Slot()
 		.Padding(0)
 		[
-			SNew(STextBlock)
-			.TextStyle(FNiagaraEditorWidgetsStyle::Get(), "NiagaraEditor.Stack.ItemText")
-			.ToolTipText_UObject(EmitterPropertiesItem, &UNiagaraStackEntry::GetTooltipText)
-			.Text_UObject(EmitterPropertiesItem, &UNiagaraStackEntry::GetDisplayName)
-			.HighlightText_UObject(InStackViewModel, &UNiagaraStackViewModel::GetCurrentSearchText)
-			.ColorAndOpacity(this, &SNiagaraStackEmitterPropertiesItem::GetTextColorForSearch)
+			SNew(SNiagaraStackDisplayName, InEmitterPropertiesItem, *InStackViewModel, "NiagaraEditor.Stack.ItemText")
+			.ColorAndOpacity(this, &SNiagaraStackEntryWidget::GetTextColorForSearch)
 		]
 		// Reset to base Button
 		+ SHorizontalBox::Slot()
diff --git a/Engine/Plugins/FX/Niagara/Source/NiagaraEditorWidgets/Private/Stack/SNiagaraStackEventHandlerPropertiesItem.cpp b/Engine/Plugins/FX/Niagara/Source/NiagaraEditorWidgets/Private/Stack/SNiagaraStackEventHandlerPropertiesItem.cpp
--- a/Engine/Plugins/FX/Niagara/Source/NiagaraEditorWidgets/Private/Stack/SNiagaraStackEventHandlerPropertiesItem.cpp
+++ b/Engine/Plugins/FX/Niagara/Source/NiagaraEditorWidgets/Private/Stack/SNiagaraStackEventHandlerPropertiesItem.cpp
@@ -25,12 +25,8 @@ void SNiagaraStackEventHandlerPropertiesItem::Construct(const FArguments& InArgs
 		+ SHorizontalBox::Slot()
 		.Padding(0)
 		[
-			SNew(STextBlock)
-			.TextStyle(FNiagaraEditorWidgetsStyle::Get(), "NiagaraEditor.Stack.ItemText")
-			.ToolTipText_UObject(EventHandlerPropertiesItem, &UNiagaraStackEntry::GetTooltipText)
-			.Text_UObject(EventHandlerPropertiesItem, &UNiagaraStackEntry::GetDisplayName)
-			.HighlightText_UObject(InStackViewModel, &UNiagaraStackViewModel::GetCurrentSearchText)
-			.ColorAndOpacity(this, &SNiagaraStackEventHandlerPropertiesItem::GetTextColorForSearch)
+			SNew(SNiagaraStackDisplayName, InEventHandlerPropertiesItem, *InStackViewModel, "NiagaraEditor.Stack.ItemText")
+			.ColorAndOpacity(this, &SNiagaraStackEntryWidget::GetTextColorForSearch)
 		]
 		// Reset to base Button
 		+ SHorizontalBox::Slot()
diff --git a/Engine/Plugins/FX/Niagara/Source/NiagaraEditorWidgets/Private/Stack/SNiagaraStackItemGroup.cpp b/Engine/Plugins/FX/Niagara/Source/NiagaraEditorWidgets/Private/Stack/SNiagaraStackItemGroup.cpp
--- a/Engine/Plugins/FX/Niagara/Source/NiagaraEditorWidgets/Private/Stack/SNiagaraStackItemGroup.cpp
+++ b/Engine/Plugins/FX/Niagara/Source/NiagaraEditorWidgets/Private/Stack/SNiagaraStackItemGroup.cpp
@@ -32,12 +32,8 @@ void SNiagaraStackItemGroup::Construct(const FArguments& InArgs, UNiagaraStackIt
 		+ SHorizontalBox::Slot()
 		.VAlign(VAlign_Center)
 		[
-			SNew(STextBlock)
-			.TextStyle(FNiagaraEditorWidgetsStyle::Get(), "NiagaraEditor.Stack.GroupText")
-			.ToolTipText_UObject(Group, &UNiagaraStackEntry::GetTooltipText)
-			.Text_UObject(Group, &UNiagaraStackEntry::GetDisplayName)
-			.HighlightText_UObject(InStackViewModel, &UNiagaraStackViewModel::GetCurrentSearchText)
-			.ColorAndOpacity(this, &SNiagaraStackItemGroup::GetTextColorForSearch)
+			SNew(SNiagaraStackDisplayName, InGroup, *InStackViewModel, "NiagaraEditor.Stack.GroupText")
+			.ColorAndOpacity(this, &SNiagaraStackEntryWidget::GetTextColorForSearch)
 		]
 		// Stack issues icon
 		+ SHorizontalBox::Slot()
diff --git a/Engine/Plugins/FX/Niagara/Source/NiagaraEditorWidgets/Private/Stack/SNiagaraStackModuleItem.cpp b/Engine/Plugins/FX/Niagara/Source/NiagaraEditorWidgets/Private/Stack/SNiagaraStackModuleItem.cpp
--- a/Engine/Plugins/FX/Niagara/Source/NiagaraEditorWidgets/Private/Stack/SNiagaraStackModuleItem.cpp
+++ b/Engine/Plugins/FX/Niagara/Source/NiagaraEditorWidgets/Private/Stack/SNiagaraStackModuleItem.cpp
@@ -51,13 +51,8 @@ void SNiagaraStackModuleItem::Construct(const FArguments& InArgs, UNiagaraStackM
 			.VAlign(VAlign_Center)
 			.Padding(0)
 			[
-				SNew(STextBlock)
-				.TextStyle(FNiagaraEditorWidgetsStyle::Get(), "NiagaraEditor.Stack.ItemText")
-				.ToolTipText_UObject(ModuleItem, &UNiagaraStackEntry::GetTooltipText)
-				.Text_UObject(ModuleItem, &UNiagaraStackEntry::GetDisplayName)
-				.HighlightText_UObject(InStackViewModel, &UNiagaraStackViewModel::GetCurrentSearchText)
+				SNew(SNiagaraStackDisplayName, InModuleItem, *InStackViewModel, "NiagaraEditor.Stack.ItemText")
 				.ColorAndOpacity(this, &SNiagaraStackModuleItem::GetTextColorForSearch)
-				.IsEnabled_UObject(ModuleItem, &UNiagaraStackModuleItem::GetIsEnabled)
 			]
 			// Stack issues icon
 			+ SHorizontalBox::Slot()
diff --git a/Engine/Plugins/FX/Niagara/Source/NiagaraEditorWidgets/Private/Stack/SNiagaraStackRendererItem.cpp b/Engine/Plugins/FX/Niagara/Source/NiagaraEditorWidgets/Private/Stack/SNiagaraStackRendererItem.cpp
--- a/Engine/Plugins/FX/Niagara/Source/NiagaraEditorWidgets/Private/Stack/SNiagaraStackRendererItem.cpp
+++ b/Engine/Plugins/FX/Niagara/Source/NiagaraEditorWidgets/Private/Stack/SNiagaraStackRendererItem.cpp
@@ -30,7 +30,7 @@ void SNiagaraStackRendererItem::Construct(const FArguments& InArgs, UNiagaraStac
 		// Renderer icon
 		SNew(SHorizontalBox)
 		+ SHorizontalBox::Slot()
-		.Padding(5, 0, 0, 0)
+		.Padding(2, 0, 0, 0)
 		.AutoWidth()
 		.VAlign(VAlign_Center)
 		[
@@ -42,12 +42,8 @@ void SNiagaraStackRendererItem::Construct(const FArguments& InArgs, UNiagaraStac
 		.Padding(5, 0, 0, 0)
 		.VAlign(VAlign_Center)
 		[
-			SNew(STextBlock)
-			.TextStyle(FNiagaraEditorWidgetsStyle::Get(), "NiagaraEditor.Stack.ItemText")
-			.ToolTipText_UObject(RendererItem, &UNiagaraStackEntry::GetTooltipText)
-			.Text_UObject(RendererItem, &UNiagaraStackEntry::GetDisplayName)
-			.HighlightText_UObject(InStackViewModel, &UNiagaraStackViewModel::GetCurrentSearchText)
-			.ColorAndOpacity(this, &SNiagaraStackRendererItem::GetTextColorForSearch)
+			SNew(SNiagaraStackDisplayName, InRendererItem, *InStackViewModel, "NiagaraEditor.Stack.ItemText")
+			.ColorAndOpacity(this, &SNiagaraStackEntryWidget::GetTextColorForSearch)
 		]
 		// Stack issues icon
 		+ SHorizontalBox::Slot()
