Niagara - Fix cis errors from overview MVP.

diff --git a/Engine/Plugins/FX/Niagara/Source/Niagara/Private/NiagaraSystem.cpp b/Engine/Plugins/FX/Niagara/Source/Niagara/Private/NiagaraSystem.cpp
--- a/Engine/Plugins/FX/Niagara/Source/Niagara/Private/NiagaraSystem.cpp
+++ b/Engine/Plugins/FX/Niagara/Source/Niagara/Private/NiagaraSystem.cpp
@@ -92,8 +92,10 @@ void UNiagaraSystem::PostInitProperties()
 		SystemUpdateScript = NewObject<UNiagaraScript>(this, "SystemUpdateScript", RF_Transactional);
 		SystemUpdateScript->SetUsage(ENiagaraScriptUsage::SystemUpdateScript);
 
+#if WITH_EDITORONLY_DATA && WITH_EDITOR
 		INiagaraModule& NiagaraModule = FModuleManager::GetModuleChecked<INiagaraModule>("Niagara");
 		EditorData = NiagaraModule.GetEditorOnlyDataUtilities().CreateDefaultEditorData(this);
+#endif
 	}
 
 	GenerateStatID();
diff --git a/Engine/Plugins/FX/Niagara/Source/NiagaraEditor/Private/ViewModels/NiagaraOverviewGraphViewModel.cpp b/Engine/Plugins/FX/Niagara/Source/NiagaraEditor/Private/ViewModels/NiagaraOverviewGraphViewModel.cpp
--- a/Engine/Plugins/FX/Niagara/Source/NiagaraEditor/Private/ViewModels/NiagaraOverviewGraphViewModel.cpp
+++ b/Engine/Plugins/FX/Niagara/Source/NiagaraEditor/Private/ViewModels/NiagaraOverviewGraphViewModel.cpp
@@ -14,6 +14,7 @@
 #include "ViewModels/Stack/NiagaraStackEntry.h"
 #include "NiagaraSystem.h"
 #include "Editor.h"
+#include "Framework/Application/SlateApplication.h"
 
 #define LOCTEXT_NAMESPACE "NiagaraOverviewGraphViewModel"
 
diff --git a/Engine/Plugins/FX/Niagara/Source/NiagaraEditor/Private/ViewModels/NiagaraSystemViewModel.cpp b/Engine/Plugins/FX/Niagara/Source/NiagaraEditor/Private/ViewModels/NiagaraSystemViewModel.cpp
--- a/Engine/Plugins/FX/Niagara/Source/NiagaraEditor/Private/ViewModels/NiagaraSystemViewModel.cpp
+++ b/Engine/Plugins/FX/Niagara/Source/NiagaraEditor/Private/ViewModels/NiagaraSystemViewModel.cpp
@@ -39,6 +39,7 @@
 #include "Framework/Commands/UICommandList.h"
 #include "NiagaraMessageManager.h"
 #include "ViewModels/NiagaraOverviewGraphViewModel.h"
+#include "Framework/Application/SlateApplication.h"
 
 DECLARE_CYCLE_STAT(TEXT("Niagara - SystemViewModel - CompileSystem"), STAT_NiagaraEditor_SystemViewModel_CompileSystem, STATGROUP_NiagaraEditor);
 
diff --git a/Engine/Plugins/FX/Niagara/Source/NiagaraEditorWidgets/Private/SNiagaraOverviewStack.h b/Engine/Plugins/FX/Niagara/Source/NiagaraEditorWidgets/Private/SNiagaraOverviewStack.h
--- a/Engine/Plugins/FX/Niagara/Source/NiagaraEditorWidgets/Private/SNiagaraOverviewStack.h
+++ b/Engine/Plugins/FX/Niagara/Source/NiagaraEditorWidgets/Private/SNiagaraOverviewStack.h
@@ -5,6 +5,7 @@
 #include "ViewModels/Stack/NiagaraStackEntry.h"
 #include "ViewModels/NiagaraSystemSelectionViewModel.h"
 #include "Widgets/SCompoundWidget.h"
+#include "Widgets/Views/SListView.h"
 #include "UObject/ObjectKey.h"
 
 class UNiagaraStackEntry;
diff --git a/Engine/Plugins/FX/Niagara/Source/NiagaraEditorWidgets/Private/SNiagaraOverviewStackNode.cpp b/Engine/Plugins/FX/Niagara/Source/NiagaraEditorWidgets/Private/SNiagaraOverviewStackNode.cpp
--- a/Engine/Plugins/FX/Niagara/Source/NiagaraEditorWidgets/Private/SNiagaraOverviewStackNode.cpp
+++ b/Engine/Plugins/FX/Niagara/Source/NiagaraEditorWidgets/Private/SNiagaraOverviewStackNode.cpp
@@ -6,6 +6,7 @@
 #include "ViewModels/NiagaraEmitterHandleViewModel.h"
 #include "SNiagaraOverviewStack.h"
 #include "NiagaraEditorModule.h"
+#include "Widgets/Layout/SBox.h"
 
 #include "Modules/ModuleManager.h"
 
diff --git a/Engine/Plugins/FX/Niagara/Source/NiagaraEditorWidgets/Private/SNiagaraStackEntryWidget.cpp b/Engine/Plugins/FX/Niagara/Source/NiagaraEditorWidgets/Private/SNiagaraStackEntryWidget.cpp
--- a/Engine/Plugins/FX/Niagara/Source/NiagaraEditorWidgets/Private/SNiagaraStackEntryWidget.cpp
+++ b/Engine/Plugins/FX/Niagara/Source/NiagaraEditorWidgets/Private/SNiagaraStackEntryWidget.cpp
@@ -3,6 +3,8 @@
 #include "SNiagaraStackEntryWidget.h"
 #include "NiagaraEditorWidgetsStyle.h"
 #include "Widgets/Layout/SWrapBox.h"
+#include "Widgets/Layout/SBox.h"
+#include "Widgets/Text/STextBlock.h"
 
 #define LOCTEXT_NAMESPACE "SNiagaraStackEntryWidget"
 
diff --git a/Engine/Plugins/FX/Niagara/Source/NiagaraEditorWidgets/Private/SNiagaraStackEntryWidget.h b/Engine/Plugins/FX/Niagara/Source/NiagaraEditorWidgets/Private/SNiagaraStackEntryWidget.h
--- a/Engine/Plugins/FX/Niagara/Source/NiagaraEditorWidgets/Private/SNiagaraStackEntryWidget.h
+++ b/Engine/Plugins/FX/Niagara/Source/NiagaraEditorWidgets/Private/SNiagaraStackEntryWidget.h
@@ -8,6 +8,7 @@
 
 class UNiagaraStackViewModel;
 class UNiagaraStackEntry;
+class SBox;
 
 class SNiagaraStackDisplayName : public SCompoundWidget
 {
