Niagara - Fix more cis issues.

diff --git a/Engine/Plugins/FX/Niagara/Source/NiagaraEditorWidgets/Private/SNiagaraOverviewStack.cpp b/Engine/Plugins/FX/Niagara/Source/NiagaraEditorWidgets/Private/SNiagaraOverviewStack.cpp
--- a/Engine/Plugins/FX/Niagara/Source/NiagaraEditorWidgets/Private/SNiagaraOverviewStack.cpp
+++ b/Engine/Plugins/FX/Niagara/Source/NiagaraEditorWidgets/Private/SNiagaraOverviewStack.cpp
@@ -14,6 +14,9 @@
 #include "SNiagaraStack.h"
 
 #include "Widgets/Layout/SScrollBox.h"
+#include "Widgets/Images/SImage.h"
+#include "EditorStyleSet.h"
+#include "Framework/Application/SlateApplication.h"
 
 class SNiagaraSystemOverviewEntryListRow : public STableRow<UNiagaraStackEntry*>
 {
