Niagara - Overview fixes.
+ Fix light renderer display name.
+ Update node name when emitters are renamed.

diff --git a/Engine/Plugins/FX/Niagara/Source/Niagara/Public/NiagaraLightRendererProperties.h b/Engine/Plugins/FX/Niagara/Source/Niagara/Public/NiagaraLightRendererProperties.h
--- a/Engine/Plugins/FX/Niagara/Source/Niagara/Public/NiagaraLightRendererProperties.h
+++ b/Engine/Plugins/FX/Niagara/Source/Niagara/Public/NiagaraLightRendererProperties.h
@@ -10,7 +10,7 @@
 
 class FNiagaraEmitterInstance;
 
-UCLASS(editinlinenew, meta = (DisplayName = "Sprite Renderer"))
+UCLASS(editinlinenew, meta = (DisplayName = "Light Renderer"))
 class UNiagaraLightRendererProperties : public UNiagaraRendererProperties
 {
 public:
diff --git a/Engine/Plugins/FX/Niagara/Source/NiagaraEditor/Private/EdGraphSchema_NiagaraSystemOverview.cpp b/Engine/Plugins/FX/Niagara/Source/NiagaraEditor/Private/EdGraphSchema_NiagaraSystemOverview.cpp
--- a/Engine/Plugins/FX/Niagara/Source/NiagaraEditor/Private/EdGraphSchema_NiagaraSystemOverview.cpp
+++ b/Engine/Plugins/FX/Niagara/Source/NiagaraEditor/Private/EdGraphSchema_NiagaraSystemOverview.cpp
@@ -5,6 +5,8 @@
 
 #define LOCTEXT_NAMESPACE "NiagaraSchema"
 
+int32 UEdGraphSchema_NiagaraSystemOverview::CurrentCacheRefreshID = 0;
+
 UEdGraphSchema_NiagaraSystemOverview::UEdGraphSchema_NiagaraSystemOverview(const FObjectInitializer& ObjectInitializer)
 	: Super(ObjectInitializer)
 {
@@ -16,4 +18,19 @@ void UEdGraphSchema_NiagaraSystemOverview::GetContextMenuActions(const UEdGraph*
 	return;
 }
 
+bool UEdGraphSchema_NiagaraSystemOverview::IsCacheVisualizationOutOfDate(int32 InVisualizationCacheID) const
+{
+	return InVisualizationCacheID != CurrentCacheRefreshID;
+}
+
+int32 UEdGraphSchema_NiagaraSystemOverview::GetCurrentVisualizationCacheID() const
+{
+	return CurrentCacheRefreshID;
+}
+
+void UEdGraphSchema_NiagaraSystemOverview::ForceVisualizationCacheClear() const
+{
+	CurrentCacheRefreshID++;
+}
+
 #undef LOCTEXT_NAMESPACE
diff --git a/Engine/Plugins/FX/Niagara/Source/NiagaraEditor/Private/ViewModels/NiagaraSystemViewModel.cpp b/Engine/Plugins/FX/Niagara/Source/NiagaraEditor/Private/ViewModels/NiagaraSystemViewModel.cpp
--- a/Engine/Plugins/FX/Niagara/Source/NiagaraEditor/Private/ViewModels/NiagaraSystemViewModel.cpp
+++ b/Engine/Plugins/FX/Niagara/Source/NiagaraEditor/Private/ViewModels/NiagaraSystemViewModel.cpp
@@ -23,6 +23,7 @@
 #include "NiagaraNodeInput.h"
 #include "NiagaraEditorModule.h"
 #include "NiagaraNodeFunctionCall.h"
+#include "EdGraphSchema_NiagaraSystemOverview.h"
 
 #include "Editor.h"
 
@@ -457,6 +458,7 @@ void FNiagaraSystemViewModel::AddReferencedObjects(FReferenceCollector& Collecto
 void FNiagaraSystemViewModel::PostUndo(bool bSuccess)
 {
 	RefreshAll();
+	GetDefault<UEdGraphSchema_NiagaraSystemOverview>()->ForceVisualizationCacheClear();
 }
 
 void FNiagaraSystemViewModel::Tick(float DeltaTime)
@@ -1190,6 +1192,7 @@ void FNiagaraSystemViewModel::EmitterHandlePropertyChanged(TSharedRef<FNiagaraEm
 
 void FNiagaraSystemViewModel::EmitterHandleNameChanged(TSharedRef<FNiagaraEmitterHandleViewModel> EmitterHandleViewModel)
 {
+	GetDefault<UEdGraphSchema_NiagaraSystemOverview>()->ForceVisualizationCacheClear();
 	CompileSystem(false);
 }
 
diff --git a/Engine/Plugins/FX/Niagara/Source/NiagaraEditor/Public/EdGraphSchema_NiagaraSystemOverview.h b/Engine/Plugins/FX/Niagara/Source/NiagaraEditor/Public/EdGraphSchema_NiagaraSystemOverview.h
--- a/Engine/Plugins/FX/Niagara/Source/NiagaraEditor/Public/EdGraphSchema_NiagaraSystemOverview.h
+++ b/Engine/Plugins/FX/Niagara/Source/NiagaraEditor/Public/EdGraphSchema_NiagaraSystemOverview.h
@@ -19,7 +19,11 @@ class NIAGARAEDITOR_API UEdGraphSchema_NiagaraSystemOverview : public UEdGraphSc
 	virtual const FPinConnectionResponse CanCreateConnection(const UEdGraphPin* A, const UEdGraphPin* B) const override { return FPinConnectionResponse(CONNECT_RESPONSE_DISALLOW, FText()); }; //@TODO System Overview: write text response
 	virtual bool TryCreateConnection(UEdGraphPin* A, UEdGraphPin* B) const override { return false; };
 	virtual bool ShouldAlwaysPurgeOnModification() const override { return false; }
-	//~ End EdGraphSchema Interface
+	virtual bool IsCacheVisualizationOutOfDate(int32 InVisualizationCacheID) const override;
+	virtual int32 GetCurrentVisualizationCacheID() const override;
+	virtual void ForceVisualizationCacheClear() const override;
 
+private:
+	static int32 CurrentCacheRefreshID;
 };
 
