Niagara - When cleaning up the system view model, null out pointers to sub-objects so that we don't run clean-up code on them if cleanup gets called twice.

diff --git a/Engine/Plugins/FX/Niagara/Source/NiagaraEditor/Private/ViewModels/NiagaraSystemViewModel.cpp b/Engine/Plugins/FX/Niagara/Source/NiagaraEditor/Private/ViewModels/NiagaraSystemViewModel.cpp
--- a/Engine/Plugins/FX/Niagara/Source/NiagaraEditor/Private/ViewModels/NiagaraSystemViewModel.cpp
+++ b/Engine/Plugins/FX/Niagara/Source/NiagaraEditor/Private/ViewModels/NiagaraSystemViewModel.cpp
@@ -117,12 +117,14 @@ void FNiagaraSystemViewModel::Cleanup()
 	{
 		SystemInstance->OnInitialized().RemoveAll(this);
 		SystemInstance->OnReset().RemoveAll(this);
+		SystemInstance = nullptr;
 	}
 
 	if (PreviewComponent)
 	{
 		PreviewComponent->OnSystemInstanceChanged().RemoveAll(this);
 		PreviewComponent->DeactivateImmediate();
+		PreviewComponent = nullptr;
 	}
 
 	CurveOwner.EmptyCurves();
@@ -150,7 +152,6 @@ void FNiagaraSystemViewModel::Cleanup()
 		Sequencer.Reset();
 	}
 
-	PreviewComponent = nullptr;
 	RemoveSystemEventHandlers();
 	SystemScriptViewModel.Reset();
 
