- Fix for D3DDebug error when we fail to find static mesh DI Data

#rnx

diff --git a/Engine/Plugins/FX/Niagara/Source/Niagara/Private/NiagaraDataInterfaceStaticMesh.cpp b/Engine/Plugins/FX/Niagara/Source/Niagara/Private/NiagaraDataInterfaceStaticMesh.cpp
--- a/Engine/Plugins/FX/Niagara/Source/Niagara/Private/NiagaraDataInterfaceStaticMesh.cpp
+++ b/Engine/Plugins/FX/Niagara/Source/Niagara/Private/NiagaraDataInterfaceStaticMesh.cpp
@@ -495,7 +495,7 @@ struct FNiagaraDataInterfaceParametersCS_StaticMesh : public FNiagaraDataInterfa
 		{
 			FNiagaraDataInterfaceProxyStaticMesh* InterfaceProxy = static_cast<FNiagaraDataInterfaceProxyStaticMesh*>(Context.DataInterface);
 			FNiagaraStaticMeshData* Data = InterfaceProxy->SystemInstancesToMeshData.Find(Context.SystemInstance);
-			ensure(Data);
+			ensureMsgf(Data, TEXT("Failed to find data for instance %s"), *Context.SystemInstance.ToString());
 			if (Data != nullptr)
 			{
 				FStaticMeshGpuSpawnBuffer* SpawnBuffer = Data->MeshGpuSpawnBuffer;
@@ -540,8 +540,8 @@ struct FNiagaraDataInterfaceParametersCS_StaticMesh : public FNiagaraDataInterfa
 				SetSRVParameter(RHICmdList, ComputeShaderRHI, MeshColorBuffer, FNiagaraRenderer::GetDummyFloatBuffer().SRV);
 
 				SetShaderValue(RHICmdList, ComputeShaderRHI, SectionCount, 0);
-				SetSRVParameter(RHICmdList, ComputeShaderRHI, MeshSectionBuffer, FNiagaraRenderer::GetDummyFloatBuffer().SRV);
-				SetSRVParameter(RHICmdList, ComputeShaderRHI, MeshTriangleBuffer, FNiagaraRenderer::GetDummyFloatBuffer().SRV);
+				SetSRVParameter(RHICmdList, ComputeShaderRHI, MeshSectionBuffer, FNiagaraRenderer::GetDummyUIntBuffer().SRV);
+				SetSRVParameter(RHICmdList, ComputeShaderRHI, MeshTriangleBuffer, FNiagaraRenderer::GetDummyUIntBuffer().SRV);
 
 				SetShaderValue(RHICmdList, ComputeShaderRHI, InstanceTransform, FMatrix::Identity);
 				SetShaderValue(RHICmdList, ComputeShaderRHI, InstanceTransformInverseTransposed, FMatrix::Identity);
@@ -581,7 +581,7 @@ void FNiagaraDataInterfaceProxyStaticMesh::DeferredDestroy()
 	for (const FGuid& Sys : DeferredDestroyList)
 	{
 		SystemInstancesToMeshData.Remove(Sys);
-		//UE_LOG(LogNiagara, Log, TEXT("DeferredDestroy() ... Removing %s"), *Sys.ToString());
+		//UE_LOG(LogNiagara, Log, TEXT("RT: StaticMesh DI - DeferredDestroy %s"), *Sys.ToString());
 	}
 
 	DeferredDestroyList.Empty();
@@ -600,8 +600,7 @@ void FNiagaraDataInterfaceProxyStaticMesh::InitializePerInstanceData(const FGuid
 	{
 		Data = &SystemInstancesToMeshData.Add(SystemInstance);
 	}
-
-	//UE_LOG(LogNiagara, Log, TEXT("InitializePerInstanceData() ... %s"), *SystemInstance.ToString());
+	//UE_LOG(LogNiagara, Log, TEXT("RT: StaticMesh DI - InitializePerInstanceData %s"), *SystemInstance.ToString());
 
 	// @todo-threadsafety We should not ever see this case! Though it's not really an error...
 	if (Data->MeshGpuSpawnBuffer)
@@ -617,7 +616,7 @@ void FNiagaraDataInterfaceProxyStaticMesh::DestroyPerInstanceData(NiagaraEmitter
 {
 	check(IsInRenderingThread());
 
-	//UE_LOG(LogNiagara, Log, TEXT("DestroyPerInstanceData() ... %s"), *SystemInstance.ToString());
+	//UE_LOG(LogNiagara, Log, TEXT("RT: StaticMesh DI - DestroyPerInstanceData %s"), *SystemInstance.ToString());
 
 	// @todo-threadsafety verify this destroys the MeshGPUSpawnBuffer data. This thread owns it now.
 	//SystemInstancesToMeshData.Remove(SystemInstance);
@@ -1170,6 +1169,8 @@ bool UNiagaraDataInterfaceStaticMesh::InitPerInstanceData(void* PerInstanceData,
 	FNDIStaticMesh_InstanceData* Inst = new (PerInstanceData) FNDIStaticMesh_InstanceData();
 	bool bSuccess = Inst->Init(this, SystemInstance);
 
+	//UE_LOG(LogNiagara, Log, TEXT("GT: StaticMesh DI - InitPerInstanceData %s"), *SystemInstance->GetId().ToString());
+
 	if (bSuccess)
 	{
 		// Always allocate when bAllowCPUAccess (index buffer can only have SRV created in this case as of today)
@@ -1210,6 +1211,8 @@ void UNiagaraDataInterfaceStaticMesh::DestroyPerInstanceData(void* PerInstanceDa
 {
 	FNDIStaticMesh_InstanceData* Inst = (FNDIStaticMesh_InstanceData*)PerInstanceData;
 
+	//UE_LOG(LogNiagara, Log, TEXT("GT: StaticMesh DI - DestroyPerInstanceData %s"), *SystemInstance->GetId().ToString());
+
 #if WITH_EDITOR
 	if (Inst->Mesh)
 	{
