Niagara - Refresh the reroute node UI when the niagara node calls the visuals changed delegate.  This prevents a crash when refreshing script graphs with reroute nodes in them.

diff --git a/Engine/Plugins/FX/Niagara/Source/NiagaraEditor/Private/NiagaraNodeFactory.cpp b/Engine/Plugins/FX/Niagara/Source/NiagaraEditor/Private/NiagaraNodeFactory.cpp
--- a/Engine/Plugins/FX/Niagara/Source/NiagaraEditor/Private/NiagaraNodeFactory.cpp
+++ b/Engine/Plugins/FX/Niagara/Source/NiagaraEditor/Private/NiagaraNodeFactory.cpp
@@ -3,6 +3,25 @@
 #include "NiagaraNodeReroute.h"
 #include "SGraphNodeKnot.h"
 
+class SNiagaraGraphNodeKnot : public SGraphNodeKnot
+{
+public:
+	SLATE_BEGIN_ARGS(SNiagaraGraphNodeKnot) {}
+	SLATE_END_ARGS();
+
+	void Construct(const FArguments& InArgs, UNiagaraNodeReroute* InKnot)
+	{
+		SGraphNodeKnot::Construct(SGraphNodeKnot::FArguments(), InKnot);
+		InKnot->OnVisualsChanged().AddSP(this, &SNiagaraGraphNodeKnot::HandleNiagaraNodeChanged);
+	}
+
+private:
+	void HandleNiagaraNodeChanged(UNiagaraNode* InNode)
+	{
+
+		UpdateGraphNode();
+	}
+};
 
 FNiagaraNodeFactory::~FNiagaraNodeFactory()
 {
@@ -12,7 +31,7 @@ TSharedPtr<SGraphNode> FNiagaraNodeFactory::CreateNodeWidget(UEdGraphNode* InNod
 {
 	if (UNiagaraNodeReroute* RerouteNode = Cast<UNiagaraNodeReroute>(InNode))
 	{
-		return SNew(SGraphNodeKnot, RerouteNode);
+		return SNew(SNiagaraGraphNodeKnot, RerouteNode);
 	}
 
 	return FGraphNodeFactory::CreateNodeWidget(InNode);
