Niagara - Fix static analysis warning.

diff --git a/Engine/Plugins/FX/Niagara/Source/NiagaraEditorWidgets/Private/SNiagaraStack.cpp b/Engine/Plugins/FX/Niagara/Source/NiagaraEditorWidgets/Private/SNiagaraStack.cpp
--- a/Engine/Plugins/FX/Niagara/Source/NiagaraEditorWidgets/Private/SNiagaraStack.cpp
+++ b/Engine/Plugins/FX/Niagara/Source/NiagaraEditorWidgets/Private/SNiagaraStack.cpp
@@ -694,7 +694,7 @@ TSharedRef<ITableRow> SNiagaraStack::OnGenerateRowForTopLevelObject(TSharedRef<U
 FReply SNiagaraStack::OnTopLevelRowMouseButtonDown(const FGeometry&, const FPointerEvent& MouseEvent, TWeakPtr<UNiagaraStackViewModel::FTopLevelViewModel> TopLevelViewModelWeak)
 {
 	TSharedPtr<UNiagaraStackViewModel::FTopLevelViewModel> TopLevelViewModel = TopLevelViewModelWeak.Pin();
-	if (TopLevelViewModel.IsValid(), MouseEvent.GetEffectingButton() == EKeys::RightMouseButton)
+	if (TopLevelViewModel.IsValid() && MouseEvent.GetEffectingButton() == EKeys::RightMouseButton)
 	{
 		FMenuBuilder MenuBuilder(true, nullptr);
 
