Minor nullptr check

diff --git a/Engine/Plugins/FX/Niagara/Source/NiagaraEditor/Private/Widgets/SNiagaraGeneratedCodeView.cpp b/Engine/Plugins/FX/Niagara/Source/NiagaraEditor/Private/Widgets/SNiagaraGeneratedCodeView.cpp
--- a/Engine/Plugins/FX/Niagara/Source/NiagaraEditor/Private/Widgets/SNiagaraGeneratedCodeView.cpp
+++ b/Engine/Plugins/FX/Niagara/Source/NiagaraEditor/Private/Widgets/SNiagaraGeneratedCodeView.cpp
@@ -531,7 +531,11 @@ SNiagaraGeneratedCodeView::~SNiagaraGeneratedCodeView()
 {
 	if (SystemViewModel.IsValid())
 	{
-		SystemViewModel->GetSelectionViewModel()->OnSelectionChanged().RemoveAll(this);
+		if (SystemViewModel->GetSelectionViewModel())
+		{
+			SystemViewModel->GetSelectionViewModel()->OnSelectionChanged().RemoveAll(this);
+		}
+
 		if (SystemViewModel->GetSystemScriptViewModel().IsValid())
 		{
 			SystemViewModel->GetSystemScriptViewModel()->OnSystemCompiled().RemoveAll(this);
