Merging CL #7660117: Making Niagara respond to exec events for particles

diff --git a/Engine/Plugins/FX/Niagara/Source/Niagara/Private/NiagaraActor.cpp b/Engine/Plugins/FX/Niagara/Source/Niagara/Private/NiagaraActor.cpp
--- a/Engine/Plugins/FX/Niagara/Source/Niagara/Private/NiagaraActor.cpp
+++ b/Engine/Plugins/FX/Niagara/Source/Niagara/Private/NiagaraActor.cpp
@@ -77,5 +77,14 @@ bool ANiagaraActor::GetReferencedContentObjects(TArray<UObject*>& Objects) const
 
 	return true;
 }
+
+void ANiagaraActor::ResetInLevel()
+{
+	if (NiagaraComponent)
+	{
+		NiagaraComponent->Activate(true);
+		NiagaraComponent->ReregisterComponent();
+	}
+}
 #endif // WITH_EDITOR
 
diff --git a/Engine/Plugins/FX/Niagara/Source/Niagara/Public/NiagaraActor.h b/Engine/Plugins/FX/Niagara/Source/Niagara/Public/NiagaraActor.h
--- a/Engine/Plugins/FX/Niagara/Source/Niagara/Public/NiagaraActor.h
+++ b/Engine/Plugins/FX/Niagara/Source/Niagara/Public/NiagaraActor.h
@@ -42,6 +42,9 @@ public:
 	// AActor interface
 	virtual bool GetReferencedContentObjects(TArray<UObject*>& Objects) const override;
 	// End of AActor interface
+
+	/** Reset this actor in the level.*/
+	NIAGARA_API void ResetInLevel();
 #endif // WITH_EDITOR
 
 };
diff --git a/Engine/Plugins/FX/Niagara/Source/NiagaraEditor/Private/NiagaraEditorModule.cpp b/Engine/Plugins/FX/Niagara/Source/NiagaraEditor/Private/NiagaraEditorModule.cpp
--- a/Engine/Plugins/FX/Niagara/Source/NiagaraEditor/Private/NiagaraEditorModule.cpp
+++ b/Engine/Plugins/FX/Niagara/Source/NiagaraEditor/Private/NiagaraEditorModule.cpp
@@ -95,7 +95,9 @@
 #include "Misc/FeedbackContext.h"
 #include "Customizations/NiagaraStaticSwitchNodeDetails.h"
 #include "Customizations/NiagaraFunctionCallNodeDetails.h"
-
+#include "NiagaraNodeFunctionCall.h"
+#include "Engine/Selection.h"
+#include "NiagaraActor.h"
 #include "NiagaraNodeFunctionCall.h"
 #include "INiagaraEditorOnlyDataUtlities.h"
 
@@ -489,7 +491,10 @@ void FNiagaraEditorModule::StartupModule()
 
 	UNiagaraSettings::OnSettingsChanged().AddRaw(this, &FNiagaraEditorModule::OnNiagaraSettingsChangedEvent);
 	FCoreUObjectDelegates::GetPreGarbageCollectDelegate().AddRaw(this, &FNiagaraEditorModule::OnPreGarbageCollection);
-
+	
+	// Any attempt to use GEditor right now will fail as it hasn't been initialized yet. Waiting for post engine init resolves that.
+	FCoreDelegates::OnPostEngineInit.AddRaw(this, &FNiagaraEditorModule::OnPostEngineInit);
+	
 	// register details customization
 	FPropertyEditorModule& PropertyModule = FModuleManager::LoadModuleChecked<FPropertyEditorModule>("PropertyEditor");
 	PropertyModule.RegisterCustomClassLayout("NiagaraComponent", FOnGetDetailCustomizationInstance::CreateStatic(&FNiagaraComponentDetails::MakeInstance));
@@ -724,7 +729,12 @@ void FNiagaraEditorModule::ShutdownModule()
 	UNiagaraSettings::OnSettingsChanged().RemoveAll(this);
 
 	FCoreUObjectDelegates::GetPreGarbageCollectDelegate().RemoveAll(this);
-
+	FCoreDelegates::OnPostEngineInit.RemoveAll(this);
+	
+	if (GEditor)
+	{
+		GEditor->OnExecParticleInvoked().RemoveAll(this);
+	}
 	
 	if (FModuleManager::Get().IsModuleLoaded("PropertyEditor"))
 	{
@@ -789,6 +799,20 @@ void FNiagaraEditorModule::ShutdownModule()
 	}
 }
 
+
+void FNiagaraEditorModule::OnPostEngineInit()
+{
+	// The editor should be valid at this point.. log a warning if not!
+	if (GEditor)
+	{
+		GEditor->OnExecParticleInvoked().AddRaw(this, &FNiagaraEditorModule::OnExecParticleInvoked);
+	}
+	else
+	{
+		UE_LOG(LogNiagaraEditor, Warning, TEXT("GEditor isn't valid! Particle reset commands will not work for Niagara components!"));
+	}
+}
+
 FNiagaraEditorModule& FNiagaraEditorModule::Get()
 {
 	return FModuleManager::LoadModuleChecked<FNiagaraEditorModule>("NiagaraEditor");
@@ -944,6 +968,38 @@ void FNiagaraEditorModule::OnPreGarbageCollection()
 	}
 }
 
+void FNiagaraEditorModule::OnExecParticleInvoked(const TCHAR* Str)
+{
+	// Very similar logic to UEditorEngine::Exec_Particle
+	if (FParse::Command(&Str, TEXT("RESET")))
+	{
+		TArray<AEmitter*> EmittersToReset;
+		if (FParse::Command(&Str, TEXT("SELECTED")))
+		{
+			// Reset any selected emitters in the level
+			for (FSelectionIterator It(GEditor->GetSelectedActorIterator()); It; ++It)
+			{
+				AActor* Actor = static_cast<AActor*>(*It);
+				checkSlow(Actor->IsA(AActor::StaticClass()));
+
+				ANiagaraActor* Emitter = Cast<ANiagaraActor>(Actor);
+				if (Emitter)
+				{
+					Emitter->ResetInLevel();
+				}
+			}
+		}
+		else if (FParse::Command(&Str, TEXT("ALL")))
+		{
+			// Reset ALL emitters in the level
+			for (TObjectIterator<ANiagaraActor> It; It; ++It)
+			{
+				ANiagaraActor* Emitter = *It;
+				Emitter->ResetInLevel();
+			}
+		}
+	}
+}
 PRAGMA_ENABLE_OPTIMIZATION
 
 #undef LOCTEXT_NAMESPACE
diff --git a/Engine/Plugins/FX/Niagara/Source/NiagaraEditor/Public/NiagaraEditorModule.h b/Engine/Plugins/FX/Niagara/Source/NiagaraEditor/Public/NiagaraEditorModule.h
--- a/Engine/Plugins/FX/Niagara/Source/NiagaraEditor/Public/NiagaraEditorModule.h
+++ b/Engine/Plugins/FX/Niagara/Source/NiagaraEditor/Public/NiagaraEditorModule.h
@@ -101,7 +101,8 @@ private:
 	void RegisterAssetTypeAction(IAssetTools& AssetTools, TSharedRef<IAssetTypeActions> Action);
 	void OnNiagaraSettingsChangedEvent(const FString& PropertyName, const UNiagaraSettings* Settings);
 	void OnPreGarbageCollection();
-
+	void OnExecParticleInvoked(const TCHAR* InStr);
+	void OnPostEngineInit();
 
 	/** FGCObject interface */
 	virtual void AddReferencedObjects( FReferenceCollector& Collector ) override;
diff --git a/Engine/Source/Editor/UnrealEd/Classes/Editor/EditorEngine.h b/Engine/Source/Editor/UnrealEd/Classes/Editor/EditorEngine.h
--- a/Engine/Source/Editor/UnrealEd/Classes/Editor/EditorEngine.h
+++ b/Engine/Source/Editor/UnrealEd/Classes/Editor/EditorEngine.h
@@ -866,10 +866,13 @@ public:
 
 	DECLARE_EVENT_TwoParams(UEngine, FHLODActorRemovedFromClusterEvent, const AActor*, const AActor*);
 	FHLODActorRemovedFromClusterEvent& OnHLODActorRemovedFromCluster() { return HLODActorRemovedFromClusterEvent; }
-
+	   
 	/** Called by internal engine systems after an Actor is removed from a cluster */
 	void BroadcastHLODActorRemovedFromCluster(const AActor* InActor, const AActor* ParentActor) { HLODActorRemovedFromClusterEvent.Broadcast(InActor, ParentActor); }
 
+	/** Called when the editor has been asked to perform an exec command on particle systems. */
+	DECLARE_EVENT_OneParam(UEditorEngine, FExecParticleInvoked, const TCHAR*);
+	FExecParticleInvoked& OnExecParticleInvoked() { return ExecParticleInvokedEvent; }
 	/**
 	 * Called before an actor or component is about to be translated, rotated, or scaled by the editor
 	 *
@@ -2959,6 +2962,9 @@ private:
 	/** Broadcasts after an Actor is removed from a cluster */
 	FHLODActorRemovedFromClusterEvent HLODActorRemovedFromClusterEvent;
 
+	/** Broadcasts after an Exec event on particles has been invoked.*/
+	FExecParticleInvoked ExecParticleInvokedEvent; 
+
 	/** Delegate to be called when a matinee is requested to be opened */
 	FShouldOpenMatineeCallback ShouldOpenMatineeCallback;
 
diff --git a/Engine/Source/Editor/UnrealEd/Private/EditorServer.cpp b/Engine/Source/Editor/UnrealEd/Private/EditorServer.cpp
--- a/Engine/Source/Editor/UnrealEd/Private/EditorServer.cpp
+++ b/Engine/Source/Editor/UnrealEd/Private/EditorServer.cpp
@@ -5148,6 +5148,9 @@ bool UEditorEngine::Exec_Particle(const TCHAR* Str, FOutputDevice& Ar)
 {
 	bool bHandled = false;
 	UE_LOG(LogEditorServer, Log, TEXT("Exec Particle!"));
+	
+	// Store off the input string here, as it is adjusted by subsequent parsing commands...
+	const TCHAR* InputStr = Str;
 	if (FParse::Command(&Str,TEXT("RESET")))
 	{
 		TArray<AEmitter*> EmittersToReset;
@@ -5176,6 +5179,12 @@ bool UEditorEngine::Exec_Particle(const TCHAR* Str, FOutputDevice& Ar)
 			}
 		}
 	}
+
+	// Invoke any downstream handlers (like the Niagara editor plugin)
+	if (ExecParticleInvokedEvent.IsBound())
+	{
+		ExecParticleInvokedEvent.Broadcast(InputStr);
+	}
 	return bHandled;
 }
 
