Niagara - Fix static analysis.  This situation should never happen, and no other code checks the results of the node creator, but for some reason this time seems unsafe to the analysis tools.

diff --git a/Engine/Plugins/FX/Niagara/Source/NiagaraEditor/Private/NiagaraSystemEditorData.cpp b/Engine/Plugins/FX/Niagara/Source/NiagaraEditor/Private/NiagaraSystemEditorData.cpp
--- a/Engine/Plugins/FX/Niagara/Source/NiagaraEditor/Private/NiagaraSystemEditorData.cpp
+++ b/Engine/Plugins/FX/Niagara/Source/NiagaraEditor/Private/NiagaraSystemEditorData.cpp
@@ -266,6 +266,8 @@ void UNiagaraSystemEditorData::SynchronizeOverviewGraphWithSystem(UNiagaraSystem
 			FVector2D EmitterOverviewNodeLocation = RightNode != nullptr
 				? FVector2D(RightNode->NodePosX + SystemOverviewNodePadding, RightNode->NodePosY)
 				: SystemOverviewGraph->GetGoodPlaceForNewNode();
+
+			checkf(EmitterOverviewNode != nullptr, TEXT("Emitter overview node creation failed!"));
 			EmitterOverviewNode->NodePosX = EmitterOverviewNodeLocation.X;
 			EmitterOverviewNode->NodePosY = EmitterOverviewNodeLocation.Y;
 
