Niagara - System overview fixes and improvements.
+ Add zoom to fit commands and menu options to the overview graph.
+ Prevent duplicate emitter header lines in the stack selection when using shift select in the overview graph.
+ Prevent a crash when deleting modules and renderers when they are the only thing by selected by binding the OnModifiedGroupItems delegate and removing them from the selection once they are deleted.

diff --git a/Engine/Plugins/FX/Niagara/Source/NiagaraEditor/Private/NiagaraEditorCommands.cpp b/Engine/Plugins/FX/Niagara/Source/NiagaraEditor/Private/NiagaraEditorCommands.cpp
--- a/Engine/Plugins/FX/Niagara/Source/NiagaraEditor/Private/NiagaraEditorCommands.cpp
+++ b/Engine/Plugins/FX/Niagara/Source/NiagaraEditor/Private/NiagaraEditorCommands.cpp
@@ -27,6 +27,9 @@ void FNiagaraEditorCommands::RegisterCommands()
 	UI_COMMAND(CollapseStackToHeaders, "Collapse to Headers", "Expands all headsers and collapse all items.", EUserInterfaceActionType::Button, FInputChord(EModifierKey::Control, EKeys::O));
 
 	UI_COMMAND(FindInCurrentView, "Find", "Contextually finds items in current view.", EUserInterfaceActionType::Button, FInputChord(EModifierKey::Control, EKeys::F));
+
+	UI_COMMAND(ZoomToFit, "Zoom to Fit", "Zooms and pans to fit the current selection.", EUserInterfaceActionType::Button, FInputChord(EKeys::F));
+	UI_COMMAND(ZoomToFitAll, "Zoom to Fit All", "Zooms and pans to fix all items.", EUserInterfaceActionType::Button, FInputChord(EKeys::A));
 }
 
 #undef LOCTEXT_NAMESPACE // NiagaraEditorCommands
diff --git a/Engine/Plugins/FX/Niagara/Source/NiagaraEditor/Private/ViewModels/Stack/NiagaraStackItem.cpp b/Engine/Plugins/FX/Niagara/Source/NiagaraEditor/Private/ViewModels/Stack/NiagaraStackItem.cpp
--- a/Engine/Plugins/FX/Niagara/Source/NiagaraEditor/Private/ViewModels/Stack/NiagaraStackItem.cpp
+++ b/Engine/Plugins/FX/Niagara/Source/NiagaraEditor/Private/ViewModels/Stack/NiagaraStackItem.cpp
@@ -20,9 +20,9 @@ UNiagaraStackEntry::EStackRowStyle UNiagaraStackItem::GetStackRowStyle() const
 	return UNiagaraStackEntry::EStackRowStyle::ItemHeader;
 }
 
-void UNiagaraStackItem::SetOnModifiedGroupItems(FOnModifiedGroupItems OnModifiedGroupItems)
+UNiagaraStackItem::FOnModifiedGroupItems& UNiagaraStackItem::OnModifiedGroupItems()
 {
-	ModifiedGroupItemsDelegate = OnModifiedGroupItems;
+	return ModifiedGroupItemsDelegate;
 }
 
 uint32 UNiagaraStackItem::GetRecursiveStackIssuesCount() const
diff --git a/Engine/Plugins/FX/Niagara/Source/NiagaraEditor/Private/ViewModels/Stack/NiagaraStackModuleItem.cpp b/Engine/Plugins/FX/Niagara/Source/NiagaraEditor/Private/ViewModels/Stack/NiagaraStackModuleItem.cpp
--- a/Engine/Plugins/FX/Niagara/Source/NiagaraEditor/Private/ViewModels/Stack/NiagaraStackModuleItem.cpp
+++ b/Engine/Plugins/FX/Niagara/Source/NiagaraEditor/Private/ViewModels/Stack/NiagaraStackModuleItem.cpp
@@ -797,7 +797,8 @@ void UNiagaraStackModuleItem::Delete()
 				GetSystemViewModel()->NotifyDataObjectChanged(InputNode->GetDataInterface());
 			}
 		}
-		ModifiedGroupItemsDelegate.ExecuteIfBound();
+		Finalize();
+		ModifiedGroupItemsDelegate.Broadcast();
 	}
 }
 
@@ -827,7 +828,7 @@ UNiagaraNodeOutput* UNiagaraStackModuleItem::GetOutputNode() const
 
 void UNiagaraStackModuleItem::NotifyModuleMoved()
 {
-	ModifiedGroupItemsDelegate.ExecuteIfBound();
+	ModifiedGroupItemsDelegate.Broadcast();
 }
 
 bool UNiagaraStackModuleItem::CanAddInput(FNiagaraVariable InputParameter) const
diff --git a/Engine/Plugins/FX/Niagara/Source/NiagaraEditor/Private/ViewModels/Stack/NiagaraStackRenderItemGroup.cpp b/Engine/Plugins/FX/Niagara/Source/NiagaraEditor/Private/ViewModels/Stack/NiagaraStackRenderItemGroup.cpp
--- a/Engine/Plugins/FX/Niagara/Source/NiagaraEditor/Private/ViewModels/Stack/NiagaraStackRenderItemGroup.cpp
+++ b/Engine/Plugins/FX/Niagara/Source/NiagaraEditor/Private/ViewModels/Stack/NiagaraStackRenderItemGroup.cpp
@@ -150,7 +150,7 @@ void UNiagaraStackRenderItemGroup::RefreshChildrenInternal(const TArray<UNiagara
 		{
 			RendererItem = NewObject<UNiagaraStackRendererItem>(this);
 			RendererItem->Initialize(CreateDefaultChildRequiredData(), RendererProperties);
-			RendererItem->SetOnModifiedGroupItems(UNiagaraStackItem::FOnModifiedGroupItems::CreateUObject(this, &UNiagaraStackRenderItemGroup::ChildModifiedGroupItems));
+			RendererItem->OnModifiedGroupItems().AddUObject(this, &UNiagaraStackRenderItemGroup::ChildModifiedGroupItems);
 		}
 
 		NewChildren.Add(RendererItem);
diff --git a/Engine/Plugins/FX/Niagara/Source/NiagaraEditor/Private/ViewModels/Stack/NiagaraStackRendererItem.cpp b/Engine/Plugins/FX/Niagara/Source/NiagaraEditor/Private/ViewModels/Stack/NiagaraStackRendererItem.cpp
--- a/Engine/Plugins/FX/Niagara/Source/NiagaraEditor/Private/ViewModels/Stack/NiagaraStackRendererItem.cpp
+++ b/Engine/Plugins/FX/Niagara/Source/NiagaraEditor/Private/ViewModels/Stack/NiagaraStackRendererItem.cpp
@@ -158,7 +158,8 @@ void UNiagaraStackRendererItem::Delete()
 	Emitter->RemoveRenderer(RendererProperties.Get());
 
 	OnDataObjectModified().Broadcast(RendererProperties.Get());
-	ModifiedGroupItemsDelegate.ExecuteIfBound();
+	Finalize();
+	ModifiedGroupItemsDelegate.Broadcast();
 }
 
 bool UNiagaraStackRendererItem::HasBaseRenderer() const
@@ -198,7 +199,7 @@ void UNiagaraStackRendererItem::ResetToBase()
 		TSharedRef<FNiagaraScriptMergeManager> MergeManager = FNiagaraScriptMergeManager::Get();
 		const UNiagaraEmitter* BaseEmitter = GetEmitterViewModel()->GetEmitter()->GetParent();
 		MergeManager->ResetRendererToBase(*GetEmitterViewModel()->GetEmitter(), *BaseEmitter, RendererProperties->GetMergeId());
-		ModifiedGroupItemsDelegate.ExecuteIfBound();
+		ModifiedGroupItemsDelegate.Broadcast();
 	}
 }
 
diff --git a/Engine/Plugins/FX/Niagara/Source/NiagaraEditor/Private/ViewModels/Stack/NiagaraStackScriptItemGroup.cpp b/Engine/Plugins/FX/Niagara/Source/NiagaraEditor/Private/ViewModels/Stack/NiagaraStackScriptItemGroup.cpp
--- a/Engine/Plugins/FX/Niagara/Source/NiagaraEditor/Private/ViewModels/Stack/NiagaraStackScriptItemGroup.cpp
+++ b/Engine/Plugins/FX/Niagara/Source/NiagaraEditor/Private/ViewModels/Stack/NiagaraStackScriptItemGroup.cpp
@@ -360,7 +360,7 @@ void UNiagaraStackScriptItemGroup::RefreshChildrenInternal(const TArray<UNiagara
 			{
 				ModuleItem = NewObject<UNiagaraStackModuleItem>(this);
 				ModuleItem->Initialize(CreateDefaultChildRequiredData(), GetAddUtilities(), *ModuleNode);
-				ModuleItem->SetOnModifiedGroupItems(UNiagaraStackModuleItem::FOnModifiedGroupItems::CreateUObject(this, &UNiagaraStackScriptItemGroup::ChildModifiedGroupItems));
+				ModuleItem->OnModifiedGroupItems().AddUObject(this, &UNiagaraStackScriptItemGroup::ChildModifiedGroupItems);
 			}
 
 			NewChildren.Add(ModuleItem);
diff --git a/Engine/Plugins/FX/Niagara/Source/NiagaraEditor/Private/ViewModels/Stack/NiagaraStackSelection.cpp b/Engine/Plugins/FX/Niagara/Source/NiagaraEditor/Private/ViewModels/Stack/NiagaraStackSelection.cpp
--- a/Engine/Plugins/FX/Niagara/Source/NiagaraEditor/Private/ViewModels/Stack/NiagaraStackSelection.cpp
+++ b/Engine/Plugins/FX/Niagara/Source/NiagaraEditor/Private/ViewModels/Stack/NiagaraStackSelection.cpp
@@ -3,6 +3,7 @@
 #include "ViewModels/Stack/NiagaraStackSelection.h"
 #include "ViewModels/Stack/NiagaraStackEntry.h"
 #include "ViewModels/Stack/NiagaraStackSpacer.h"
+#include "ViewModels/Stack/NiagaraStackItem.h"
 #include "ViewModels/NiagaraSystemViewModel.h"
 #include "NiagaraSystemEditorData.h"
 
@@ -33,6 +34,15 @@ void UNiagaraStackSelection::SetSelectedEntries(const TArray<UNiagaraStackEntry*
 
 void UNiagaraStackSelection::RefreshChildrenInternal(const TArray<UNiagaraStackEntry*>& CurrentChildren, TArray<UNiagaraStackEntry*>& NewChildren, TArray<FStackIssue>& NewIssues)
 {
+	for (UNiagaraStackEntry* CurrentChild : CurrentChildren)
+	{
+		UNiagaraStackItem* CurrentChildItem = Cast<UNiagaraStackItem>(CurrentChild);
+		if (CurrentChildItem != nullptr)
+		{
+			CurrentChildItem->OnModifiedGroupItems().RemoveAll(this);
+		}
+	}
+
 	auto GetOrCreateSpacer = [&](FName SpacerExecutionCategory, FName SpacerKey)
 	{
 		UNiagaraStackSpacer* Spacer = FindCurrentChildOfTypeByPredicate<UNiagaraStackSpacer>(CurrentChildren,
@@ -53,16 +63,23 @@ void UNiagaraStackSelection::RefreshChildrenInternal(const TArray<UNiagaraStackE
 	int32 EntryIndex = 0;
 	for (TWeakObjectPtr<UNiagaraStackEntry> SelectedEntry : SelectedEntries)
 	{
-		if (EntryIndex > 0)
+		if (SelectedEntry.IsValid() && SelectedEntry->IsFinalized() == false)
 		{
-			NewChildren.Add(GetOrCreateSpacer(NAME_None, *FString::Printf(TEXT("SelectionSpacer_%i"), EntryIndex)));
-		}
+			if (EntryIndex > 0)
+			{
+				NewChildren.Add(GetOrCreateSpacer(NAME_None, *FString::Printf(TEXT("SelectionSpacer_%i"), EntryIndex)));
+			}
+
+			UNiagaraStackItem* SelectedItem = Cast<UNiagaraStackItem>(SelectedEntry);
+			if (SelectedItem != nullptr)
+			{
+				SelectedItem->OnModifiedGroupItems().AddUObject(this, &UNiagaraStackEntry::RefreshChildren);
+			}
 
-		if (SelectedEntry.IsValid())
-		{
 			SelectedEntry->SetIsExpanded(true);
 			NewChildren.Add(SelectedEntry.Get());
+
+			EntryIndex++;
 		}
-		EntryIndex++;
 	}
 }
\ No newline at end of file
diff --git a/Engine/Plugins/FX/Niagara/Source/NiagaraEditor/Private/ViewModels/Stack/NiagaraStackViewModel.cpp b/Engine/Plugins/FX/Niagara/Source/NiagaraEditor/Private/ViewModels/Stack/NiagaraStackViewModel.cpp
--- a/Engine/Plugins/FX/Niagara/Source/NiagaraEditor/Private/ViewModels/Stack/NiagaraStackViewModel.cpp
+++ b/Engine/Plugins/FX/Niagara/Source/NiagaraEditor/Private/ViewModels/Stack/NiagaraStackViewModel.cpp
@@ -61,6 +61,11 @@ FText UNiagaraStackViewModel::FTopLevelViewModel::GetDisplayName() const
 	}
 }
 
+bool UNiagaraStackViewModel::FTopLevelViewModel::operator==(const FTopLevelViewModel& Other) const
+{
+	return Other.SystemViewModel == SystemViewModel && Other.EmitterHandleViewModel == EmitterHandleViewModel;
+}
+
 void UNiagaraStackViewModel::InitializeWithViewModels(TSharedPtr<FNiagaraSystemViewModel> InSystemViewModel, TSharedPtr<FNiagaraEmitterHandleViewModel> InEmitterHandleViewModel, FNiagaraStackViewModelOptions InOptions)
 {
 	Reset();
@@ -649,29 +654,30 @@ void UNiagaraStackViewModel::RefreshTopLevelViewModels()
 				[&RootChildEmitterHandleViewModel](const TSharedRef<FTopLevelViewModel>& TopLevelViewModel) { return TopLevelViewModel->EmitterHandleViewModel == RootChildEmitterHandleViewModel; });
 			if (CurrentTopLevelViewModelPtr != nullptr)
 			{
-				if (TopLevelViewModels.Contains(*CurrentTopLevelViewModelPtr) == false)
-				{
-					TopLevelViewModels.Add(*CurrentTopLevelViewModelPtr);
-				}
+				TopLevelViewModel = *CurrentTopLevelViewModelPtr;
 			}
 			else
 			{
-				TopLevelViewModels.Add(MakeShared<FTopLevelViewModel>(RootChildEmitterHandleViewModel));
+				TopLevelViewModel = MakeShared<FTopLevelViewModel>(RootChildEmitterHandleViewModel);
 			}
 		}
 		else
 		{
 			TSharedRef<FTopLevelViewModel>* CurrentTopLevelViewModelPtr = CurrentTopLevelViewModels.FindByPredicate(
 				[RootChild](const TSharedRef<FTopLevelViewModel>& TopLevelViewModel) { return TopLevelViewModel->SystemViewModel == RootChild->GetSystemViewModel(); });
-			if (CurrentTopLevelViewModelPtr != nullptr && TopLevelViewModels.Contains(*CurrentTopLevelViewModelPtr) == false)
+			if (CurrentTopLevelViewModelPtr != nullptr)
 			{
-				TopLevelViewModels.Add(*CurrentTopLevelViewModelPtr);
+				TopLevelViewModel = *CurrentTopLevelViewModelPtr;
 			}
 			else
 			{
-				TopLevelViewModels.Add(MakeShared<FTopLevelViewModel>(RootChild->GetSystemViewModel()));
+				TopLevelViewModel = MakeShared<FTopLevelViewModel>(RootChild->GetSystemViewModel());
 			}
 		}
+		if (TopLevelViewModels.ContainsByPredicate([&TopLevelViewModel](const TSharedRef<FTopLevelViewModel>& ExistingTopLevelViewModel) { return *TopLevelViewModel == *ExistingTopLevelViewModel; }) == false)
+		{
+			TopLevelViewModels.Add(TopLevelViewModel.ToSharedRef());
+		}
 	}
 }
 
diff --git a/Engine/Plugins/FX/Niagara/Source/NiagaraEditor/Public/NiagaraEditorCommands.h b/Engine/Plugins/FX/Niagara/Source/NiagaraEditor/Public/NiagaraEditorCommands.h
--- a/Engine/Plugins/FX/Niagara/Source/NiagaraEditor/Public/NiagaraEditorCommands.h
+++ b/Engine/Plugins/FX/Niagara/Source/NiagaraEditor/Public/NiagaraEditorCommands.h
@@ -51,4 +51,7 @@ public:
 	TSharedPtr<FUICommandInfo> CollapseStackToHeaders;
 
 	TSharedPtr<FUICommandInfo> FindInCurrentView;
+
+	TSharedPtr<FUICommandInfo> ZoomToFit;
+	TSharedPtr<FUICommandInfo> ZoomToFitAll;
 };
\ No newline at end of file
diff --git a/Engine/Plugins/FX/Niagara/Source/NiagaraEditor/Public/ViewModels/Stack/NiagaraStackItem.h b/Engine/Plugins/FX/Niagara/Source/NiagaraEditor/Public/ViewModels/Stack/NiagaraStackItem.h
--- a/Engine/Plugins/FX/Niagara/Source/NiagaraEditor/Public/ViewModels/Stack/NiagaraStackItem.h
+++ b/Engine/Plugins/FX/Niagara/Source/NiagaraEditor/Public/ViewModels/Stack/NiagaraStackItem.h
@@ -16,14 +16,14 @@ class NIAGARAEDITOR_API UNiagaraStackItem : public UNiagaraStackEntry
 	GENERATED_BODY()
 
 public:
-	DECLARE_DELEGATE(FOnModifiedGroupItems);
+	DECLARE_MULTICAST_DELEGATE(FOnModifiedGroupItems);
 
 public:
 	void Initialize(FRequiredEntryData InRequiredEntryData, FString InStackEditorDataKey);
 
 	virtual EStackRowStyle GetStackRowStyle() const override;
 
-	void SetOnModifiedGroupItems(FOnModifiedGroupItems OnModifiedGroupItems);
+	FOnModifiedGroupItems& OnModifiedGroupItems();
 
 	uint32 GetRecursiveStackIssuesCount() const;
 	EStackIssueSeverity GetHighestStackIssueSeverity() const;
diff --git a/Engine/Plugins/FX/Niagara/Source/NiagaraEditor/Public/ViewModels/Stack/NiagaraStackViewModel.h b/Engine/Plugins/FX/Niagara/Source/NiagaraEditor/Public/ViewModels/Stack/NiagaraStackViewModel.h
--- a/Engine/Plugins/FX/Niagara/Source/NiagaraEditor/Public/ViewModels/Stack/NiagaraStackViewModel.h
+++ b/Engine/Plugins/FX/Niagara/Source/NiagaraEditor/Public/ViewModels/Stack/NiagaraStackViewModel.h
@@ -74,6 +74,8 @@ public:
 
 		FText GetDisplayName() const;
 
+		bool operator==(const FTopLevelViewModel& Other)const;
+
 		const TSharedPtr<FNiagaraSystemViewModel> SystemViewModel;
 		const TSharedPtr<FNiagaraEmitterHandleViewModel> EmitterHandleViewModel;
 	};
diff --git a/Engine/Plugins/FX/Niagara/Source/NiagaraEditorWidgets/Private/SNiagaraOverviewGraph.cpp b/Engine/Plugins/FX/Niagara/Source/NiagaraEditorWidgets/Private/SNiagaraOverviewGraph.cpp
--- a/Engine/Plugins/FX/Niagara/Source/NiagaraEditorWidgets/Private/SNiagaraOverviewGraph.cpp
+++ b/Engine/Plugins/FX/Niagara/Source/NiagaraEditorWidgets/Private/SNiagaraOverviewGraph.cpp
@@ -15,7 +15,7 @@
 #include "ContentBrowserModule.h"
 #include "ViewModels/NiagaraOverviewGraphViewModel.h"
 #include "EdGraphSchema_Niagara.h"
-#include "NiagaraOverviewNode.h"
+#include "NiagaraEditorCommands.h"
 #include "GraphEditorActions.h"
 #include "Framework/MultiBox/MultiBoxBuilder.h"
 #include "ScopedTransaction.h"
@@ -68,8 +68,14 @@ void SNiagaraOverviewGraph::Construct(const FArguments& InArgs, TSharedRef<FNiag
 	TSharedRef<FUICommandList> Commands = ViewModel->GetCommands();
 	Commands->MapAction(
 		FGraphEditorCommands::Get().CreateComment,
-		FExecuteAction::CreateRaw(this, &SNiagaraOverviewGraph::OnCreateComment));
-
+		FExecuteAction::CreateSP(this, &SNiagaraOverviewGraph::OnCreateComment));
+	Commands->MapAction(
+		FNiagaraEditorCommands::Get().ZoomToFit,
+		FExecuteAction::CreateSP(this, &SNiagaraOverviewGraph::ZoomToFit));
+	Commands->MapAction(
+		FNiagaraEditorCommands::Get().ZoomToFitAll,
+		FExecuteAction::CreateSP(this, &SNiagaraOverviewGraph::ZoomToFitAll));
+	
 	GraphEditor = SNew(SGraphEditor)
 		.AdditionalCommands(Commands)
 		.Appearance(AppearanceInfo)
@@ -162,21 +168,30 @@ FActionMenuContent SNiagaraOverviewGraph::OnCreateGraphActionMenu(UEdGraph* InGr
 {
 	if (ViewModel->GetSystemViewModel()->GetEditMode() == ENiagaraSystemViewModelEditMode::SystemAsset)
 	{
-		FMenuBuilder MenuBuilder(true, NULL);
-		MenuBuilder.BeginSection(TEXT("NiagaraOverview_EditGraph"), LOCTEXT("EditGraph", "Edit Graph"));
+		FMenuBuilder MenuBuilder(true, ViewModel->GetCommands());
 
-		MenuBuilder.AddSubMenu(
-			LOCTEXT("EmitterAddLabel", "Add Emitter"),
-			LOCTEXT("EmitterAddToolTip", "Add an existing emitter"),
-			FNewMenuDelegate::CreateSP(this, &SNiagaraOverviewGraph::CreateAddEmitterMenuContent, InGraph));
+		MenuBuilder.BeginSection(TEXT("NiagaraOverview_EditGraph"), LOCTEXT("EditGraph", "Edit Graph"));
+		{
+			MenuBuilder.AddSubMenu(
+				LOCTEXT("EmitterAddLabel", "Add Emitter"),
+				LOCTEXT("EmitterAddToolTip", "Add an existing emitter"),
+				FNewMenuDelegate::CreateSP(this, &SNiagaraOverviewGraph::CreateAddEmitterMenuContent, InGraph));
 
-		MenuBuilder.AddMenuEntry(
-			LOCTEXT("CommentsLabel", "Add Comment"),
-			LOCTEXT("CommentsToolTip", "Add a comment box"),
-			FSlateIcon(),
-			FExecuteAction::CreateSP(this, &SNiagaraOverviewGraph::OnCreateComment));
+			MenuBuilder.AddMenuEntry(
+				LOCTEXT("CommentsLabel", "Add Comment"),
+				LOCTEXT("CommentsToolTip", "Add a comment box"),
+				FSlateIcon(),
+				FExecuteAction::CreateSP(this, &SNiagaraOverviewGraph::OnCreateComment));
+		}
+		MenuBuilder.EndSection();
 
+		MenuBuilder.BeginSection(TEXT("NiagaraOverview_View"), LOCTEXT("EditGraph", "View"));
+		{
+			MenuBuilder.AddMenuEntry(FNiagaraEditorCommands::Get().ZoomToFit);
+			MenuBuilder.AddMenuEntry(FNiagaraEditorCommands::Get().ZoomToFitAll);
+		}
 		MenuBuilder.EndSection();
+
 		TSharedRef<SWidget> ActionMenu = MenuBuilder.MakeWidget();
 
 		return FActionMenuContent(ActionMenu, ActionMenu);
@@ -257,4 +272,14 @@ void SNiagaraOverviewGraph::CreateAddEmitterMenuContent(FMenuBuilder& MenuBuilde
 	MenuBuilder.AddWidget(EmitterAddSubMenu, FText());
 }
 
+void SNiagaraOverviewGraph::ZoomToFit()
+{
+	GraphEditor->ZoomToFit(true);
+}
+
+void SNiagaraOverviewGraph::ZoomToFitAll()
+{
+	GraphEditor->ZoomToFit(false);
+}
+
 #undef LOCTEXT_NAMESPACE // "NiagaraOverviewGraph"
\ No newline at end of file
diff --git a/Engine/Plugins/FX/Niagara/Source/NiagaraEditorWidgets/Private/SNiagaraOverviewGraph.h b/Engine/Plugins/FX/Niagara/Source/NiagaraEditorWidgets/Private/SNiagaraOverviewGraph.h
--- a/Engine/Plugins/FX/Niagara/Source/NiagaraEditorWidgets/Private/SNiagaraOverviewGraph.h
+++ b/Engine/Plugins/FX/Niagara/Source/NiagaraEditorWidgets/Private/SNiagaraOverviewGraph.h
@@ -37,6 +37,10 @@ private:
 
 	void CreateAddEmitterMenuContent(FMenuBuilder& MenuBuilder, UEdGraph* InGraph);
 
+	void ZoomToFit();
+
+	void ZoomToFitAll();
+
 private:
 	TSharedPtr<FNiagaraOverviewGraphViewModel> ViewModel;
 	TSharedPtr<SGraphEditor> GraphEditor;
