Add particle direct read DI for CPU simulations

diff --git a/Engine/Plugins/FX/Niagara/Source/Niagara/Classes/NiagaraDataInterface.h b/Engine/Plugins/FX/Niagara/Source/Niagara/Classes/NiagaraDataInterface.h
--- a/Engine/Plugins/FX/Niagara/Source/Niagara/Classes/NiagaraDataInterface.h
+++ b/Engine/Plugins/FX/Niagara/Source/Niagara/Classes/NiagaraDataInterface.h
@@ -95,6 +95,17 @@ struct NDI_FUNC_BINDER(ClassName, FuncName)\
 	}\
 };
 
+#define DEFINE_NDI_DIRECT_FUNC_BINDER_WITH_PAYLOAD(ClassName, FuncName)\
+struct NDI_FUNC_BINDER(ClassName, FuncName)\
+{\
+	template <typename ... VarTypes>\
+	static void Bind(UNiagaraDataInterface* Interface, FVMExternalFunction &OutFunc, VarTypes... Var)\
+	{\
+		auto Lambda = [Interface, Var...](FVectorVMContext& Context) { static_cast<ClassName*>(Interface)->FuncName(Context, Var...); };\
+		OutFunc = FVMExternalFunction::CreateLambda(Lambda);\
+	}\
+};
+
 #if WITH_EDITOR
 // Helper class for GUI error handling
 DECLARE_DELEGATE_RetVal(bool, FNiagaraDataInterfaceFix);
@@ -272,6 +283,8 @@ public:
 //		checkf(false, TEXT("Unefined HLSL in data interface. Interfaces need to define HLSL for uniforms their functions access."));
 	}
 
+	virtual void PostExecute() {}
+
 #if WITH_EDITOR	
 	/** Refreshes and returns the errors detected with the corresponding data, if any.*/
 	virtual TArray<FNiagaraDataInterfaceError> GetErrors() { return TArray<FNiagaraDataInterfaceError>(); }
diff --git a/Engine/Plugins/FX/Niagara/Source/Niagara/Classes/NiagaraDataInterfaceParticleRead.h b/Engine/Plugins/FX/Niagara/Source/Niagara/Classes/NiagaraDataInterfaceParticleRead.h
new file mode 100644
--- /dev/null
+++ b/Engine/Plugins/FX/Niagara/Source/Niagara/Classes/NiagaraDataInterfaceParticleRead.h
@@ -0,0 +1,66 @@
+// Copyright 1998-2019 Epic Games, Inc. All Rights Reserved.
+
+#pragma once
+
+#include "NiagaraDataInterface.h"
+#include "NiagaraCommon.h"
+#include "NiagaraEmitterInstance.h"
+#include "NiagaraDataInterfaceParticleRead.generated.h"
+
+struct FNDIParticleRead_InstanceData
+{
+	FNiagaraSystemInstance* SystemInstance;
+	FNiagaraEmitterInstance* EmitterInstance;
+};
+
+UCLASS(EditInlineNew, Category = "ParticleRead", meta = (DisplayName = "Particle Attribute Reader"))
+class NIAGARA_API UNiagaraDataInterfaceParticleRead : public UNiagaraDataInterface
+{
+	GENERATED_UCLASS_BODY()
+public:
+	UPROPERTY(EditAnywhere, Category = "ParticleRead")
+	FString EmitterName;
+
+	//UObject Interface
+	virtual void PostInitProperties()override;
+	virtual void PostLoad() override;
+#if WITH_EDITOR
+	virtual void PreEditChange(UProperty* PropertyAboutToChange) override;
+	virtual void PostEditChangeProperty(struct FPropertyChangedEvent& PropertyChangedEvent) override;
+#endif
+	//UObject Interface End
+
+	//UNiagaraDataInterface Interface
+	virtual bool InitPerInstanceData(void* PerInstanceData, FNiagaraSystemInstance* SystemInstance) override;
+	virtual void DestroyPerInstanceData(void* PerInstanceData, FNiagaraSystemInstance* SystemInstance) override;
+	virtual int32 PerInstanceDataSize() const { return sizeof(FNDIParticleRead_InstanceData); }
+	
+	virtual bool PerInstanceTick(void* PerInstanceData, FNiagaraSystemInstance* SystemInstance, float DeltaSeconds);
+	virtual bool PerInstanceTickPostSimulate(void* PerInstanceData, FNiagaraSystemInstance* SystemInstance, float DeltaSeconds);
+
+	virtual void GetFunctions(TArray<FNiagaraFunctionSignature>& OutFunctions) override;
+	virtual void GetVMExternalFunction(const FVMExternalFunctionBindingInfo& BindingInfo, void* InstanceData, FVMExternalFunction &OutFunc) override;
+	virtual bool CanExecuteOnTarget(ENiagaraSimTarget Target)const override { return Target == ENiagaraSimTarget::CPUSim; }
+	virtual bool Equals(const UNiagaraDataInterface* Other) const override;
+
+	// GPU sim functionality
+	virtual bool GetFunctionHLSL(const FName& DefinitionFunctionName, FString InstanceFunctionName, FNiagaraDataInterfaceGPUParamInfo& ParamInfo, FString& OutHLSL) override;
+	virtual void GetParameterDefinitionHLSL(FNiagaraDataInterfaceGPUParamInfo& ParamInfo, FString& OutHLSL) override;
+	virtual FNiagaraDataInterfaceParametersCS* ConstructComputeParameters() const override;
+
+	//UNiagaraDataInterface Interface End
+
+	virtual void ReadFloat(FVectorVMContext& Context, FName AttributeToRead);
+	virtual void ReadVector2(FVectorVMContext& Context, FName AttributeToRead);
+	virtual void ReadVector3(FVectorVMContext& Context, FName AttributeToRead);
+	virtual void ReadVector4(FVectorVMContext& Context, FName AttributeToRead);
+	virtual void ReadInt(FVectorVMContext& Context, FName AttributeToRead);
+	virtual void ReadBool(FVectorVMContext& Context, FName AttributeToRead);
+	virtual void ReadColor(FVectorVMContext& Context, FName AttributeToRead);
+	virtual void ReadQuat(FVectorVMContext& Context, FName AttributeToRead);
+
+private:
+	template<typename T>
+	T RetrieveValueWithCheck(FNiagaraEmitterInstance* EmitterInstance, const FNiagaraTypeDefinition& Type, const FName& Attr, const FNiagaraID& ParticleID, bool &bValid);
+
+};
\ No newline at end of file
diff --git a/Engine/Plugins/FX/Niagara/Source/Niagara/Private/NiagaraDataInterface.cpp b/Engine/Plugins/FX/Niagara/Source/Niagara/Private/NiagaraDataInterface.cpp
--- a/Engine/Plugins/FX/Niagara/Source/Niagara/Private/NiagaraDataInterface.cpp
+++ b/Engine/Plugins/FX/Niagara/Source/Niagara/Private/NiagaraDataInterface.cpp
@@ -79,7 +79,7 @@ void UNiagaraDataInterface::ValidateFunction(const FNiagaraFunctionSignature& Fu
 	TArray<FNiagaraFunctionSignature> DIFuncs;
 	GetFunctions(DIFuncs);
 
-	if (!DIFuncs.Contains(Function))
+	if (!DIFuncs.ContainsByPredicate([&](const FNiagaraFunctionSignature& Sig) { return Sig.EqualsIgnoringSpecifiers(Function); }))
 	{
 		//We couldn't find this signature in the list of available functions.
 		//Lets try to find one with the same name whose parameters may have changed.
diff --git a/Engine/Plugins/FX/Niagara/Source/Niagara/Private/NiagaraDataInterfaceParticleRead.cpp b/Engine/Plugins/FX/Niagara/Source/Niagara/Private/NiagaraDataInterfaceParticleRead.cpp
new file mode 100644
--- /dev/null
+++ b/Engine/Plugins/FX/Niagara/Source/Niagara/Private/NiagaraDataInterfaceParticleRead.cpp
@@ -0,0 +1,597 @@
+// Copyright 1998-2019 Epic Games, Inc. All Rights Reserved.
+
+#include "NiagaraDataInterfaceParticleRead.h"
+#include "NiagaraConstants.h"
+#include "NiagaraSystemInstance.h"
+
+#define LOCTEXT_NAMESPACE "NiagaraDataInterfaceParticleRead"
+
+static const FName GetIntAttributeFunctionName("Get int Attribute");
+static const FName GetFloatAttributeFunctionName("Get float Attribute");
+static const FName GetVec2AttributeFunctionName("Get Vector2 Attribute");
+static const FName GetVec3AttributeFunctionName("Get Vector3 Attribute");
+static const FName GetVec4AttributeFunctionName("Get Vector4 Attribute");
+static const FName GetBoolAttributeFunctionName("Get bool Attribute");
+static const FName GetColorAttributeFunctionName("Get Color Attribute");
+static const FName GetQuatAttributeFunctionName("Get Quaternion Attribute");
+
+UNiagaraDataInterfaceParticleRead::UNiagaraDataInterfaceParticleRead(FObjectInitializer const& ObjectInitializer)
+	: Super(ObjectInitializer)
+{
+}
+
+void UNiagaraDataInterfaceParticleRead::PostInitProperties()
+{
+	Super::PostInitProperties();
+}
+
+void UNiagaraDataInterfaceParticleRead::PostLoad()
+{
+	Super::PostLoad();
+}
+
+#if WITH_EDITOR
+
+void UNiagaraDataInterfaceParticleRead::PreEditChange(UProperty* PropertyAboutToChange)
+{
+	Super::PreEditChange(PropertyAboutToChange);
+}
+
+void UNiagaraDataInterfaceParticleRead::PostEditChangeProperty(struct FPropertyChangedEvent& PropertyChangedEvent)
+{
+	Super::PostEditChangeProperty(PropertyChangedEvent);
+}
+
+#endif
+
+bool UNiagaraDataInterfaceParticleRead::InitPerInstanceData(void* PerInstanceData, FNiagaraSystemInstance* SystemInstance)
+{
+	FNDIParticleRead_InstanceData* PIData = new (PerInstanceData) FNDIParticleRead_InstanceData;
+	PIData->SystemInstance = SystemInstance;
+	PIData->EmitterInstance = nullptr;
+	for (TSharedPtr<FNiagaraEmitterInstance, ESPMode::ThreadSafe> EmitterInstance : SystemInstance->GetEmitters())
+	{
+		if (EmitterName == EmitterInstance->GetCachedEmitter()->GetUniqueEmitterName())
+		{
+			PIData->EmitterInstance = EmitterInstance.Get();
+			break;
+		}
+	}
+	return PIData->EmitterInstance;
+}
+
+void UNiagaraDataInterfaceParticleRead::DestroyPerInstanceData(void* PerInstanceData, FNiagaraSystemInstance* SystemInstance)
+{
+
+}
+
+bool UNiagaraDataInterfaceParticleRead::PerInstanceTick(void* PerInstanceData, FNiagaraSystemInstance* SystemInstance, float DeltaSeconds)
+{
+	return false;
+}
+
+bool UNiagaraDataInterfaceParticleRead::PerInstanceTickPostSimulate(void* PerInstanceData, FNiagaraSystemInstance* SystemInstance, float DeltaSeconds)
+{
+	return false;
+}
+
+void UNiagaraDataInterfaceParticleRead::GetFunctions(TArray<FNiagaraFunctionSignature>& OutFunctions)
+{
+	{
+		FNiagaraFunctionSignature Sig;
+		Sig.Name = GetFloatAttributeFunctionName;
+		Sig.Inputs.Add(FNiagaraVariable(FNiagaraTypeDefinition(GetClass()), TEXT("Particle Reader")));
+		Sig.Inputs.Add(FNiagaraVariable(FNiagaraTypeDefinition::GetIDDef(), TEXT("Particle ID")));
+
+		Sig.Outputs.Add(FNiagaraVariable(FNiagaraTypeDefinition::GetBoolDef(), TEXT("Valid")));
+		Sig.Outputs.Add(FNiagaraVariable(FNiagaraTypeDefinition::GetFloatDef(), TEXT("Value")));
+
+		Sig.FunctionSpecifiers.Add(FName("Attribute"));
+
+		Sig.bMemberFunction = true;
+		Sig.bRequiresContext = false;
+		OutFunctions.Add(Sig);
+	}
+	{
+		FNiagaraFunctionSignature Sig;
+		Sig.Name = GetVec2AttributeFunctionName;
+		Sig.Inputs.Add(FNiagaraVariable(FNiagaraTypeDefinition(GetClass()), TEXT("Particle Reader")));
+		Sig.Inputs.Add(FNiagaraVariable(FNiagaraTypeDefinition::GetIDDef(), TEXT("Particle ID")));
+
+		Sig.Outputs.Add(FNiagaraVariable(FNiagaraTypeDefinition::GetBoolDef(), TEXT("Valid")));
+		Sig.Outputs.Add(FNiagaraVariable(FNiagaraTypeDefinition::GetVec2Def(), TEXT("Value")));
+
+		Sig.FunctionSpecifiers.Add(FName("Attribute"));
+
+		Sig.bMemberFunction = true;
+		Sig.bRequiresContext = false;
+		OutFunctions.Add(Sig);
+	}
+	{
+		FNiagaraFunctionSignature Sig;
+		Sig.Name = GetVec3AttributeFunctionName;
+		Sig.Inputs.Add(FNiagaraVariable(FNiagaraTypeDefinition(GetClass()), TEXT("Particle Reader")));
+		Sig.Inputs.Add(FNiagaraVariable(FNiagaraTypeDefinition::GetIDDef(), TEXT("Particle ID")));
+
+		Sig.Outputs.Add(FNiagaraVariable(FNiagaraTypeDefinition::GetBoolDef(), TEXT("Valid")));
+		Sig.Outputs.Add(FNiagaraVariable(FNiagaraTypeDefinition::GetVec3Def(), TEXT("Value")));
+
+		Sig.FunctionSpecifiers.Add(FName("Attribute"));
+
+		Sig.bMemberFunction = true;
+		Sig.bRequiresContext = false;
+		OutFunctions.Add(Sig);
+	}
+	{
+		FNiagaraFunctionSignature Sig;
+		Sig.Name = GetVec4AttributeFunctionName;
+		Sig.Inputs.Add(FNiagaraVariable(FNiagaraTypeDefinition(GetClass()), TEXT("Particle Reader")));
+		Sig.Inputs.Add(FNiagaraVariable(FNiagaraTypeDefinition::GetIDDef(), TEXT("Particle ID")));
+
+		Sig.Outputs.Add(FNiagaraVariable(FNiagaraTypeDefinition::GetBoolDef(), TEXT("Valid")));
+		Sig.Outputs.Add(FNiagaraVariable(FNiagaraTypeDefinition::GetVec4Def(), TEXT("Value")));
+
+		Sig.FunctionSpecifiers.Add(FName("Attribute"));
+
+		Sig.bMemberFunction = true;
+		Sig.bRequiresContext = false;
+		OutFunctions.Add(Sig);
+	}
+	{
+		FNiagaraFunctionSignature Sig;
+		Sig.Name = GetIntAttributeFunctionName;
+		Sig.Inputs.Add(FNiagaraVariable(FNiagaraTypeDefinition(GetClass()), TEXT("Particle Reader")));
+		Sig.Inputs.Add(FNiagaraVariable(FNiagaraTypeDefinition::GetIDDef(), TEXT("Particle ID")));
+
+		Sig.Outputs.Add(FNiagaraVariable(FNiagaraTypeDefinition::GetBoolDef(), TEXT("Valid")));
+		Sig.Outputs.Add(FNiagaraVariable(FNiagaraTypeDefinition::GetIntDef(), TEXT("Value")));
+
+		Sig.FunctionSpecifiers.Add(FName("Attribute"));
+
+		Sig.bMemberFunction = true;
+		Sig.bRequiresContext = false;
+		OutFunctions.Add(Sig);
+	}
+	{
+		FNiagaraFunctionSignature Sig;
+		Sig.Name = GetBoolAttributeFunctionName;
+		Sig.Inputs.Add(FNiagaraVariable(FNiagaraTypeDefinition(GetClass()), TEXT("Particle Reader")));
+		Sig.Inputs.Add(FNiagaraVariable(FNiagaraTypeDefinition::GetIDDef(), TEXT("Particle ID")));
+
+		Sig.Outputs.Add(FNiagaraVariable(FNiagaraTypeDefinition::GetBoolDef(), TEXT("Valid")));
+		Sig.Outputs.Add(FNiagaraVariable(FNiagaraTypeDefinition::GetBoolDef(), TEXT("Value")));
+
+		Sig.FunctionSpecifiers.Add(FName("Attribute"));
+
+		Sig.bMemberFunction = true;
+		Sig.bRequiresContext = false;
+		OutFunctions.Add(Sig);
+	}
+	{
+		FNiagaraFunctionSignature Sig;
+		Sig.Name = GetColorAttributeFunctionName;
+		Sig.Inputs.Add(FNiagaraVariable(FNiagaraTypeDefinition(GetClass()), TEXT("Particle Reader")));
+		Sig.Inputs.Add(FNiagaraVariable(FNiagaraTypeDefinition::GetIDDef(), TEXT("Particle ID")));
+
+		Sig.Outputs.Add(FNiagaraVariable(FNiagaraTypeDefinition::GetBoolDef(), TEXT("Valid")));
+		Sig.Outputs.Add(FNiagaraVariable(FNiagaraTypeDefinition::GetColorDef(), TEXT("Value")));
+
+		Sig.FunctionSpecifiers.Add(FName("Attribute"));
+
+		Sig.bMemberFunction = true;
+		Sig.bRequiresContext = false;
+		OutFunctions.Add(Sig);
+	}
+	{
+		FNiagaraFunctionSignature Sig;
+		Sig.Name = GetQuatAttributeFunctionName;
+		Sig.Inputs.Add(FNiagaraVariable(FNiagaraTypeDefinition(GetClass()), TEXT("Particle Reader")));
+		Sig.Inputs.Add(FNiagaraVariable(FNiagaraTypeDefinition::GetIDDef(), TEXT("Particle ID")));
+
+		Sig.Outputs.Add(FNiagaraVariable(FNiagaraTypeDefinition::GetBoolDef(), TEXT("Valid")));
+		Sig.Outputs.Add(FNiagaraVariable(FNiagaraTypeDefinition::GetQuatDef(), TEXT("Value")));
+
+		Sig.FunctionSpecifiers.Add(FName("Attribute"));
+
+		Sig.bMemberFunction = true;
+		Sig.bRequiresContext = false;
+		OutFunctions.Add(Sig);
+	}
+}
+DEFINE_NDI_DIRECT_FUNC_BINDER_WITH_PAYLOAD(UNiagaraDataInterfaceParticleRead, ReadFloat);
+DEFINE_NDI_DIRECT_FUNC_BINDER_WITH_PAYLOAD(UNiagaraDataInterfaceParticleRead, ReadVector2);
+DEFINE_NDI_DIRECT_FUNC_BINDER_WITH_PAYLOAD(UNiagaraDataInterfaceParticleRead, ReadVector3);
+DEFINE_NDI_DIRECT_FUNC_BINDER_WITH_PAYLOAD(UNiagaraDataInterfaceParticleRead, ReadVector4);
+DEFINE_NDI_DIRECT_FUNC_BINDER_WITH_PAYLOAD(UNiagaraDataInterfaceParticleRead, ReadInt);
+DEFINE_NDI_DIRECT_FUNC_BINDER_WITH_PAYLOAD(UNiagaraDataInterfaceParticleRead, ReadBool);
+DEFINE_NDI_DIRECT_FUNC_BINDER_WITH_PAYLOAD(UNiagaraDataInterfaceParticleRead, ReadColor);
+DEFINE_NDI_DIRECT_FUNC_BINDER_WITH_PAYLOAD(UNiagaraDataInterfaceParticleRead, ReadQuat);
+void UNiagaraDataInterfaceParticleRead::GetVMExternalFunction(const FVMExternalFunctionBindingInfo& BindingInfo, void* InstanceData, FVMExternalFunction &OutFunc)
+{
+	bool bBindSuccessful = false;
+	FNDIParticleRead_InstanceData* PIData = static_cast<FNDIParticleRead_InstanceData*>(InstanceData);
+	const FName* AttributeToRead = BindingInfo.Specifiers.Find(FName("Attribute"));
+
+	if (!AttributeToRead)
+	{
+		UE_LOG(LogNiagara, Error, TEXT("VMExternalFunction '%s' does not have a function specifier 'attribute'!"), *BindingInfo.Name.ToString());
+		return;
+	}
+
+	if (BindingInfo.Name == GetFloatAttributeFunctionName)
+	{
+		if (AttributeToRead)
+		{
+			FNiagaraVariable VariableToRead(FNiagaraTypeDefinition::GetFloatDef(), *AttributeToRead);
+			FName Attr = *AttributeToRead;
+			if (PIData->EmitterInstance->GetData().GetVariables().Find(VariableToRead) != INDEX_NONE)
+			{
+				NDI_FUNC_BINDER(UNiagaraDataInterfaceParticleRead, ReadFloat)::Bind(this, OutFunc, Attr);
+				bBindSuccessful = true;
+			}
+		}
+	}
+	else if (BindingInfo.Name == GetVec2AttributeFunctionName)
+	{
+		if (AttributeToRead)
+		{
+			FNiagaraVariable VariableToRead(FNiagaraTypeDefinition::GetVec2Def(), *AttributeToRead);
+			FName Attr = *AttributeToRead;
+			if (PIData->EmitterInstance->GetData().GetVariables().Find(VariableToRead) != INDEX_NONE)
+			{
+				NDI_FUNC_BINDER(UNiagaraDataInterfaceParticleRead, ReadVector2)::Bind(this, OutFunc, Attr);
+				bBindSuccessful = true;
+			}
+		}
+	}
+	else if (BindingInfo.Name == GetVec3AttributeFunctionName)
+	{
+		if (AttributeToRead)
+		{
+			FNiagaraVariable VariableToRead(FNiagaraTypeDefinition::GetVec3Def(), *AttributeToRead);
+			FName Attr = *AttributeToRead;
+			if (PIData->EmitterInstance->GetData().GetVariables().Find(VariableToRead) != INDEX_NONE)
+			{
+				NDI_FUNC_BINDER(UNiagaraDataInterfaceParticleRead, ReadVector3)::Bind(this, OutFunc, Attr);
+				bBindSuccessful = true;
+			}
+		}
+	}
+	else if (BindingInfo.Name == GetVec4AttributeFunctionName)
+	{
+		if (AttributeToRead)
+		{
+			FNiagaraVariable VariableToRead(FNiagaraTypeDefinition::GetVec4Def(), *AttributeToRead);
+			FName Attr = *AttributeToRead;
+			if (PIData->EmitterInstance->GetData().GetVariables().Find(VariableToRead) != INDEX_NONE)
+			{
+				NDI_FUNC_BINDER(UNiagaraDataInterfaceParticleRead, ReadVector4)::Bind(this, OutFunc, Attr);
+				bBindSuccessful = true;
+			}
+		}
+	}
+	else if (BindingInfo.Name == GetIntAttributeFunctionName)
+	{
+		if (AttributeToRead)
+		{
+			FNiagaraVariable VariableToRead(FNiagaraTypeDefinition::GetIntDef(), *AttributeToRead);
+			FName Attr = *AttributeToRead;
+			if (PIData->EmitterInstance->GetData().GetVariables().Find(VariableToRead) != INDEX_NONE)
+			{
+				NDI_FUNC_BINDER(UNiagaraDataInterfaceParticleRead, ReadInt)::Bind(this, OutFunc, Attr);
+				bBindSuccessful = true;
+			}
+		}
+	}
+	else if (BindingInfo.Name == GetBoolAttributeFunctionName)
+	{
+		if (AttributeToRead)
+		{
+			FNiagaraVariable VariableToRead(FNiagaraTypeDefinition::GetBoolDef(), *AttributeToRead);
+			FName Attr = *AttributeToRead;
+			if (PIData->EmitterInstance->GetData().GetVariables().Find(VariableToRead) != INDEX_NONE)
+			{
+				NDI_FUNC_BINDER(UNiagaraDataInterfaceParticleRead, ReadBool)::Bind(this, OutFunc, Attr);
+				bBindSuccessful = true;
+			}
+		}
+	}
+	else if (BindingInfo.Name == GetColorAttributeFunctionName)
+	{
+		if (AttributeToRead)
+		{
+			FNiagaraVariable VariableToRead(FNiagaraTypeDefinition::GetColorDef(), *AttributeToRead);
+			FName Attr = *AttributeToRead;
+			if (PIData->EmitterInstance->GetData().GetVariables().Find(VariableToRead) != INDEX_NONE)
+			{
+				NDI_FUNC_BINDER(UNiagaraDataInterfaceParticleRead, ReadColor)::Bind(this, OutFunc, Attr);
+				bBindSuccessful = true;
+			}
+		}
+	}
+	else if (BindingInfo.Name == GetQuatAttributeFunctionName)
+	{
+		if (AttributeToRead)
+		{
+			FNiagaraVariable VariableToRead(FNiagaraTypeDefinition::GetQuatDef(), *AttributeToRead);
+			FName Attr = *AttributeToRead;
+			if (PIData->EmitterInstance->GetData().GetVariables().Find(VariableToRead) != INDEX_NONE)
+			{
+				NDI_FUNC_BINDER(UNiagaraDataInterfaceParticleRead, ReadQuat)::Bind(this, OutFunc, Attr);
+				bBindSuccessful = true;
+			}
+		}
+	}
+
+
+	if (!bBindSuccessful)
+	{
+		UE_LOG(LogNiagara, Error, TEXT("Failed to bind VMExternalFunction '%s' with attribute '%s'! Check that the attribute is named correctly."), *BindingInfo.Name.ToString(), *(*AttributeToRead).ToString());
+	}
+}
+
+void UNiagaraDataInterfaceParticleRead::ReadFloat(FVectorVMContext& Context, FName AttributeToRead)
+{
+	VectorVM::FExternalFuncInputHandler<int32> ParticleIDIndexParam(Context);
+	VectorVM::FExternalFuncInputHandler<int32> ParticleIDAcquireTagParam(Context);
+
+	VectorVM::FUserPtrHandler<FNDIParticleRead_InstanceData> InstanceData(Context);
+
+	VectorVM::FExternalFuncRegisterHandler<FNiagaraBool> OutValid(Context);
+	VectorVM::FExternalFuncRegisterHandler<float> OutValue(Context);
+	
+	for (int32 InstanceIdx = 0; InstanceIdx < Context.NumInstances; ++InstanceIdx)
+	{
+		FNiagaraID ParticleID = { ParticleIDIndexParam.GetAndAdvance(), ParticleIDAcquireTagParam.GetAndAdvance() };
+		bool bValid;
+		float Value = RetrieveValueWithCheck<float>(InstanceData->EmitterInstance, FNiagaraTypeDefinition::GetFloatDef(), AttributeToRead, ParticleID, bValid);
+		FNiagaraBool ValidValue;
+		ValidValue.SetValue(bValid);
+		*OutValid.GetDestAndAdvance() = ValidValue;
+		*OutValue.GetDestAndAdvance() = Value;
+	}
+}
+
+void UNiagaraDataInterfaceParticleRead::ReadVector2(FVectorVMContext& Context, FName AttributeToRead)
+{
+	VectorVM::FExternalFuncInputHandler<int32> ParticleIDIndexParam(Context);
+	VectorVM::FExternalFuncInputHandler<int32> ParticleIDAcquireTagParam(Context);
+
+	VectorVM::FUserPtrHandler<FNDIParticleRead_InstanceData> InstanceData(Context);
+
+	VectorVM::FExternalFuncRegisterHandler<FNiagaraBool> OutValid(Context);
+	VectorVM::FExternalFuncRegisterHandler<float> OutX(Context);
+	VectorVM::FExternalFuncRegisterHandler<float> OutY(Context);
+
+	for (int32 InstanceIdx = 0; InstanceIdx < Context.NumInstances; ++InstanceIdx)
+	{
+		FNiagaraID ParticleID = { ParticleIDIndexParam.GetAndAdvance(), ParticleIDAcquireTagParam.GetAndAdvance() };
+		bool bValid;
+
+		FVector2D Value = RetrieveValueWithCheck<FVector2D>(InstanceData->EmitterInstance, FNiagaraTypeDefinition::GetVec2Def(), AttributeToRead, ParticleID, bValid);
+
+		FNiagaraBool ValidValue;
+		ValidValue.SetValue(bValid);
+		*OutValid.GetDestAndAdvance() = ValidValue;
+		*OutX.GetDestAndAdvance() = Value.X;
+		*OutY.GetDestAndAdvance() = Value.Y;
+	}
+}
+
+void UNiagaraDataInterfaceParticleRead::ReadVector3(FVectorVMContext& Context, FName AttributeToRead)
+{
+	VectorVM::FExternalFuncInputHandler<int32> ParticleIDIndexParam(Context);
+	VectorVM::FExternalFuncInputHandler<int32> ParticleIDAcquireTagParam(Context);
+
+	VectorVM::FUserPtrHandler<FNDIParticleRead_InstanceData> InstanceData(Context);
+
+	VectorVM::FExternalFuncRegisterHandler<FNiagaraBool> OutValid(Context);
+	VectorVM::FExternalFuncRegisterHandler<float> OutX(Context);
+	VectorVM::FExternalFuncRegisterHandler<float> OutY(Context);
+	VectorVM::FExternalFuncRegisterHandler<float> OutZ(Context);
+
+	for (int32 InstanceIdx = 0; InstanceIdx < Context.NumInstances; ++InstanceIdx)
+	{
+		FNiagaraID ParticleID = { ParticleIDIndexParam.GetAndAdvance(), ParticleIDAcquireTagParam.GetAndAdvance() };
+		bool bValid;
+
+		FVector Value = RetrieveValueWithCheck<FVector>(InstanceData->EmitterInstance, FNiagaraTypeDefinition::GetVec3Def(), AttributeToRead, ParticleID, bValid);
+
+		FNiagaraBool ValidValue;
+		ValidValue.SetValue(bValid);
+		*OutValid.GetDestAndAdvance() = ValidValue;
+		*OutX.GetDestAndAdvance() = Value.X;
+		*OutY.GetDestAndAdvance() = Value.Y;
+		*OutZ.GetDestAndAdvance() = Value.Z;
+	}
+}
+
+void UNiagaraDataInterfaceParticleRead::ReadVector4(FVectorVMContext& Context, FName AttributeToRead)
+{
+	VectorVM::FExternalFuncInputHandler<int32> ParticleIDIndexParam(Context);
+	VectorVM::FExternalFuncInputHandler<int32> ParticleIDAcquireTagParam(Context);
+
+	VectorVM::FUserPtrHandler<FNDIParticleRead_InstanceData> InstanceData(Context);
+
+	VectorVM::FExternalFuncRegisterHandler<FNiagaraBool> OutValid(Context);
+	VectorVM::FExternalFuncRegisterHandler<float> OutX(Context);
+	VectorVM::FExternalFuncRegisterHandler<float> OutY(Context);
+	VectorVM::FExternalFuncRegisterHandler<float> OutZ(Context);
+	VectorVM::FExternalFuncRegisterHandler<float> OutW(Context);
+
+	for (int32 InstanceIdx = 0; InstanceIdx < Context.NumInstances; ++InstanceIdx)
+	{
+		FNiagaraID ParticleID = { ParticleIDIndexParam.GetAndAdvance(), ParticleIDAcquireTagParam.GetAndAdvance() };
+		bool bValid;
+
+		FVector4 Value = RetrieveValueWithCheck<FVector4>(InstanceData->EmitterInstance, FNiagaraTypeDefinition::GetVec4Def(), AttributeToRead, ParticleID, bValid);
+
+		FNiagaraBool ValidValue;
+		ValidValue.SetValue(bValid);
+		*OutValid.GetDestAndAdvance() = ValidValue;
+		*OutX.GetDestAndAdvance() = Value.X;
+		*OutY.GetDestAndAdvance() = Value.Y;
+		*OutZ.GetDestAndAdvance() = Value.Z;
+		*OutW.GetDestAndAdvance() = Value.W;
+	}
+}
+
+void UNiagaraDataInterfaceParticleRead::ReadInt(FVectorVMContext& Context, FName AttributeToRead)
+{
+	VectorVM::FExternalFuncInputHandler<int32> ParticleIDIndexParam(Context);
+	VectorVM::FExternalFuncInputHandler<int32> ParticleIDAcquireTagParam(Context);
+
+	VectorVM::FUserPtrHandler<FNDIParticleRead_InstanceData> InstanceData(Context);
+
+	VectorVM::FExternalFuncRegisterHandler<FNiagaraBool> OutValid(Context);
+	VectorVM::FExternalFuncRegisterHandler<int32> OutValue(Context);
+
+	for (int32 InstanceIdx = 0; InstanceIdx < Context.NumInstances; ++InstanceIdx)
+	{
+		FNiagaraID ParticleID = { ParticleIDIndexParam.GetAndAdvance(), ParticleIDAcquireTagParam.GetAndAdvance() };
+		bool bValid;
+
+		int32 Value = RetrieveValueWithCheck<int32>(InstanceData->EmitterInstance, FNiagaraTypeDefinition::GetIntDef(), AttributeToRead, ParticleID, bValid);
+
+		FNiagaraBool ValidValue;
+		ValidValue.SetValue(bValid);
+		*OutValid.GetDestAndAdvance() = ValidValue;
+		*OutValue.GetDestAndAdvance() = Value;
+	}
+}
+
+void UNiagaraDataInterfaceParticleRead::ReadBool(FVectorVMContext& Context, FName AttributeToRead)
+{
+	VectorVM::FExternalFuncInputHandler<int32> ParticleIDIndexParam(Context);
+	VectorVM::FExternalFuncInputHandler<int32> ParticleIDAcquireTagParam(Context);
+
+	VectorVM::FUserPtrHandler<FNDIParticleRead_InstanceData> InstanceData(Context);
+
+	VectorVM::FExternalFuncRegisterHandler<FNiagaraBool> OutValid(Context);
+	VectorVM::FExternalFuncRegisterHandler<FNiagaraBool> OutValue(Context);
+
+	for (int32 InstanceIdx = 0; InstanceIdx < Context.NumInstances; ++InstanceIdx)
+	{
+		FNiagaraID ParticleID = { ParticleIDIndexParam.GetAndAdvance(), ParticleIDAcquireTagParam.GetAndAdvance() };
+		bool bValid;
+
+		FNiagaraBool Value = RetrieveValueWithCheck<FNiagaraBool>(InstanceData->EmitterInstance, FNiagaraTypeDefinition::GetBoolDef(), AttributeToRead, ParticleID, bValid);
+
+		FNiagaraBool ValidValue;
+		ValidValue.SetValue(bValid);
+		*OutValid.GetDestAndAdvance() = ValidValue;
+		*OutValue.GetDestAndAdvance() = Value;
+	}
+}
+
+void UNiagaraDataInterfaceParticleRead::ReadColor(FVectorVMContext& Context, FName AttributeToRead)
+{
+	VectorVM::FExternalFuncInputHandler<int32> ParticleIDIndexParam(Context);
+	VectorVM::FExternalFuncInputHandler<int32> ParticleIDAcquireTagParam(Context);
+
+	VectorVM::FUserPtrHandler<FNDIParticleRead_InstanceData> InstanceData(Context);
+
+	VectorVM::FExternalFuncRegisterHandler<FNiagaraBool> OutValid(Context);
+	VectorVM::FExternalFuncRegisterHandler<float> OutR(Context);
+	VectorVM::FExternalFuncRegisterHandler<float> OutG(Context);
+	VectorVM::FExternalFuncRegisterHandler<float> OutB(Context);
+	VectorVM::FExternalFuncRegisterHandler<float> OutA(Context);
+
+	for (int32 InstanceIdx = 0; InstanceIdx < Context.NumInstances; ++InstanceIdx)
+	{
+		FNiagaraID ParticleID = { ParticleIDIndexParam.GetAndAdvance(), ParticleIDAcquireTagParam.GetAndAdvance() };
+		bool bValid;
+
+		FLinearColor Value = RetrieveValueWithCheck<FLinearColor>(InstanceData->EmitterInstance, FNiagaraTypeDefinition::GetColorDef(), AttributeToRead, ParticleID, bValid);
+
+		FNiagaraBool ValidValue;
+		ValidValue.SetValue(bValid);
+		*OutValid.GetDestAndAdvance() = ValidValue;
+		*OutR.GetDestAndAdvance() = Value.R;
+		*OutG.GetDestAndAdvance() = Value.G;
+		*OutB.GetDestAndAdvance() = Value.B;
+		*OutA.GetDestAndAdvance() = Value.A;
+	}
+}
+
+void UNiagaraDataInterfaceParticleRead::ReadQuat(FVectorVMContext& Context, FName AttributeToRead)
+{
+	VectorVM::FExternalFuncInputHandler<int32> ParticleIDIndexParam(Context);
+	VectorVM::FExternalFuncInputHandler<int32> ParticleIDAcquireTagParam(Context);
+
+	VectorVM::FUserPtrHandler<FNDIParticleRead_InstanceData> InstanceData(Context);
+
+	VectorVM::FExternalFuncRegisterHandler<FNiagaraBool> OutValid(Context);
+	VectorVM::FExternalFuncRegisterHandler<float> OutX(Context);
+	VectorVM::FExternalFuncRegisterHandler<float> OutY(Context);
+	VectorVM::FExternalFuncRegisterHandler<float> OutZ(Context);
+	VectorVM::FExternalFuncRegisterHandler<float> OutW(Context);
+
+	for (int32 InstanceIdx = 0; InstanceIdx < Context.NumInstances; ++InstanceIdx)
+	{
+		FNiagaraID ParticleID = { ParticleIDIndexParam.GetAndAdvance(), ParticleIDAcquireTagParam.GetAndAdvance() };
+		bool bValid;
+
+		FQuat Value = RetrieveValueWithCheck<FQuat>(InstanceData->EmitterInstance, FNiagaraTypeDefinition::GetQuatDef(), AttributeToRead, ParticleID, bValid);
+
+		FNiagaraBool ValidValue;
+		ValidValue.SetValue(bValid);
+		*OutValid.GetDestAndAdvance() = ValidValue;
+		*OutX.GetDestAndAdvance() = Value.X;
+		*OutY.GetDestAndAdvance() = Value.Y;
+		*OutZ.GetDestAndAdvance() = Value.Z;
+		*OutW.GetDestAndAdvance() = Value.W;
+	}
+}
+
+template <typename T>
+T UNiagaraDataInterfaceParticleRead::RetrieveValueWithCheck(FNiagaraEmitterInstance* EmitterInstance, const FNiagaraTypeDefinition& Type, const FName& Attr, const FNiagaraID& ParticleID, bool& bValid)
+{
+	TArray<int32>& IDTable = EmitterInstance->GetData().GetCurrentData()->GetIDTable();
+	if (ParticleID.Index < 0 || ParticleID.Index >= IDTable.Num())
+	{
+		bValid = false;
+		return T();
+	}
+	else
+	{
+		FNiagaraVariable ReadVar(Type, Attr);
+		FNiagaraDataSetAccessor<T> ValueData(EmitterInstance->GetData(), ReadVar);
+
+		FNiagaraVariable IDVar(FNiagaraTypeDefinition::GetIDDef(), "ID");
+		FNiagaraDataSetAccessor<FNiagaraID> IDData(EmitterInstance->GetData(), IDVar);
+
+		int32 CorrectIndex = IDTable[ParticleID.Index];
+		T Value = T();
+		FNiagaraID ID = NIAGARA_INVALID_ID;
+		if (CorrectIndex >= 0)
+		{
+			ID = IDData.GetSafe(CorrectIndex, NIAGARA_INVALID_ID);
+			Value = ValueData.GetSafe(CorrectIndex, T());
+		}
+		bValid = (ID != NIAGARA_INVALID_ID);
+		return Value;
+	}
+}
+
+bool UNiagaraDataInterfaceParticleRead::Equals(const UNiagaraDataInterface* Other) const
+{
+	return false;
+}
+
+bool UNiagaraDataInterfaceParticleRead::GetFunctionHLSL(const FName& DefinitionFunctionName, FString InstanceFunctionName, FNiagaraDataInterfaceGPUParamInfo& ParamInfo, FString& OutHLSL)
+{
+	return false;
+}
+
+void UNiagaraDataInterfaceParticleRead::GetParameterDefinitionHLSL(FNiagaraDataInterfaceGPUParamInfo& ParamInfo, FString& OutHLSL)
+{
+
+}
+
+FNiagaraDataInterfaceParametersCS* UNiagaraDataInterfaceParticleRead::ConstructComputeParameters() const
+{
+	return nullptr;
+}
+
+#undef LOCTEXT_NAMESPACE
\ No newline at end of file
diff --git a/Engine/Plugins/FX/Niagara/Source/Niagara/Public/NiagaraCommon.h b/Engine/Plugins/FX/Niagara/Source/Niagara/Public/NiagaraCommon.h
--- a/Engine/Plugins/FX/Niagara/Source/Niagara/Public/NiagaraCommon.h
+++ b/Engine/Plugins/FX/Niagara/Source/Niagara/Public/NiagaraCommon.h
@@ -217,6 +217,9 @@ struct NIAGARA_API FNiagaraFunctionSignature
 	/** True if this is the signature for a "member" function of a data interface. If this is true, the first input is the owner. */
 	UPROPERTY()
 	bool bMemberFunction;
+	/** Function specifiers verified at bind time. */
+	UPROPERTY()
+	TMap<FName, FName> FunctionSpecifiers;
 
 	/** Localized description of this node. Note that this is *not* used during the operator == below since it may vary from culture to culture.*/
 #if WITH_EDITORONLY_DATA
@@ -236,11 +239,43 @@ struct NIAGARA_API FNiagaraFunctionSignature
 		, Outputs(InOutputs)
 		, bRequiresContext(bInRequiresContext)
 		, bMemberFunction(bInMemberFunction)
+		, FunctionSpecifiers()
+	{
+
+	}
+
+	FNiagaraFunctionSignature(FName InName, TArray<FNiagaraVariable>& InInputs, TArray<FNiagaraVariable>& InOutputs, FName InSource, bool bInRequiresContext, bool bInMemberFunction, TMap<FName, FName>& InFunctionSpecifiers)
+		: Name(InName)
+		, Inputs(InInputs)
+		, Outputs(InOutputs)
+		, bRequiresContext(bInRequiresContext)
+		, bMemberFunction(bInMemberFunction)
+		, FunctionSpecifiers(InFunctionSpecifiers)
 	{
 
 	}
 
 	bool operator==(const FNiagaraFunctionSignature& Other) const
+	{
+		bool bFunctionSpecifiersEqual = [&]()
+		{
+			if (Other.FunctionSpecifiers.Num() != FunctionSpecifiers.Num())
+			{
+				return false;
+			}
+			for (const TTuple<FName, FName>& Specifier : FunctionSpecifiers)
+			{
+				if (Other.FunctionSpecifiers.FindRef(Specifier.Key) != Specifier.Value)
+				{
+					return false;
+				}
+			}
+			return true;
+		}();
+		return EqualsIgnoringSpecifiers(Other) && bFunctionSpecifiersEqual;
+	}
+
+	bool EqualsIgnoringSpecifiers(const FNiagaraFunctionSignature& Other) const
 	{
 		bool bNamesEqual = Name.ToString().Equals(Other.Name.ToString());
 		bool bInputsEqual = Inputs == Other.Inputs;
@@ -384,6 +419,9 @@ struct FVMExternalFunctionBindingInfo
 	UPROPERTY()
 	int32 NumOutputs;
 
+	UPROPERTY()
+	TMap<FName, FName> Specifiers;
+
 	FORCEINLINE int32 GetNumInputs()const { return InputParamLocations.Num(); }
 	FORCEINLINE int32 GetNumOutputs()const { return NumOutputs; }
 };
diff --git a/Engine/Plugins/FX/Niagara/Source/Niagara/Public/NiagaraTypes.h b/Engine/Plugins/FX/Niagara/Source/Niagara/Public/NiagaraTypes.h
--- a/Engine/Plugins/FX/Niagara/Source/Niagara/Public/NiagaraTypes.h
+++ b/Engine/Plugins/FX/Niagara/Source/Niagara/Public/NiagaraTypes.h
@@ -176,9 +176,12 @@ struct FNiagaraID
 	int32 AcquireTag;
 
 	bool operator==(const FNiagaraID& Other)const { return Index == Other.Index && AcquireTag == Other.AcquireTag; }
+	bool operator!=(const FNiagaraID& Other)const { return !(*this == Other); }
 	bool operator<(const FNiagaraID& Other)const { return Index < Other.Index && AcquireTag < Other.AcquireTag; }
 };
 
+#define NIAGARA_INVALID_ID (FNiagaraID({(INDEX_NONE), (INDEX_NONE)}))
+
 FORCEINLINE uint32 GetTypeHash(const FNiagaraID& ID)
 {
 	return HashCombine(GetTypeHash(ID.Index), GetTypeHash(ID.AcquireTag));
@@ -722,6 +725,15 @@ public:
 		}
 	}
 
+	static void Deregister(const FNiagaraTypeDefinition& Type)
+	{
+		RegisteredTypes.Remove(Type);
+		RegisteredParamTypes.Remove(Type);
+		RegisteredPayloadTypes.Remove(Type);
+		RegisteredUserDefinedTypes.Remove(Type);
+		RegisteredNumericTypes.Remove(Type);
+	}
+
 	FNiagaraTypeDefinition GetTypeDefFromStruct(UStruct* Struct)
 	{
 		for (FNiagaraTypeDefinition& TypeDef : RegisteredTypes)
diff --git a/Engine/Plugins/FX/Niagara/Source/NiagaraEditor/Private/NiagaraCompiler.cpp b/Engine/Plugins/FX/Niagara/Source/NiagaraEditor/Private/NiagaraCompiler.cpp
--- a/Engine/Plugins/FX/Niagara/Source/NiagaraEditor/Private/NiagaraCompiler.cpp
+++ b/Engine/Plugins/FX/Niagara/Source/NiagaraEditor/Private/NiagaraCompiler.cpp
@@ -855,6 +855,7 @@ FNiagaraCompileResults FHlslNiagaraCompiler::CompileScript(const FNiagaraCompile
 					int32 NewBindingIdx = CompileResults.Data->CalledVMExternalFunctions.AddDefaulted();
 					CompileResults.Data->CalledVMExternalFunctions[NewBindingIdx].Name = *Sig->GetName();
 					CompileResults.Data->CalledVMExternalFunctions[NewBindingIdx].OwnerName = Sig->OwnerName;
+					CompileResults.Data->CalledVMExternalFunctions[NewBindingIdx].Specifiers = Sig->FunctionSpecifiers;
 
 					CompileResults.Data->CalledVMExternalFunctions[NewBindingIdx].InputParamLocations = FuncInfo.InputParamLocations;
 					CompileResults.Data->CalledVMExternalFunctions[NewBindingIdx].NumOutputs = FuncInfo.NumOutputs;
diff --git a/Engine/Plugins/FX/Niagara/Source/NiagaraEditor/Private/NiagaraHlslTranslator.cpp b/Engine/Plugins/FX/Niagara/Source/NiagaraEditor/Private/NiagaraHlslTranslator.cpp
--- a/Engine/Plugins/FX/Niagara/Source/NiagaraEditor/Private/NiagaraHlslTranslator.cpp
+++ b/Engine/Plugins/FX/Niagara/Source/NiagaraEditor/Private/NiagaraHlslTranslator.cpp
@@ -4961,7 +4961,7 @@ void FHlslNiagaraTranslator::RegisterFunctionCall(ENiagaraScriptUsage ScriptUsag
 				TArray<FNiagaraFunctionSignature> DataInterfaceFunctions;
 				CDO->GetFunctions(DataInterfaceFunctions);
 
-				const bool bFoundMatch = DataInterfaceFunctions.ContainsByPredicate([&](const FNiagaraFunctionSignature& Sig) -> bool { return Sig == OutSignature; });
+				const bool bFoundMatch = DataInterfaceFunctions.ContainsByPredicate([&](const FNiagaraFunctionSignature& Sig) -> bool { return Sig.EqualsIgnoringSpecifiers(OutSignature); });
 				if (!bFoundMatch)
 				{
 					Error(LOCTEXT("FunctionCallDataInterfaceMissing", "Function call signature does not match DataInterface possible signatures?"), nullptr, nullptr);
@@ -5133,6 +5133,10 @@ FString FHlslNiagaraTranslator::GetFunctionSignatureSymbol(const FNiagaraFunctio
 	{
 		SigStr += TEXT("_Func_");
 	}
+	for (const TTuple<FName, FName>& Specifier : Sig.FunctionSpecifiers)
+	{
+		SigStr += TEXT("_") + Specifier.Key.ToString() + Specifier.Value.ToString().Replace(TEXT("."), TEXT(""));
+	}
 	return GetSanitizedSymbolName(SigStr);
 }
 
diff --git a/Engine/Plugins/FX/Niagara/Source/NiagaraEditor/Private/NiagaraNodeFunctionCall.cpp b/Engine/Plugins/FX/Niagara/Source/NiagaraEditor/Private/NiagaraNodeFunctionCall.cpp
--- a/Engine/Plugins/FX/Niagara/Source/NiagaraEditor/Private/NiagaraNodeFunctionCall.cpp
+++ b/Engine/Plugins/FX/Niagara/Source/NiagaraEditor/Private/NiagaraNodeFunctionCall.cpp
@@ -17,6 +17,7 @@
 #include "NiagaraConstants.h"
 #include "NiagaraCustomVersion.h"
 #include "NiagaraDataInterfaceSkeletalMesh.h"
+#include "SNiagaraGraphNodeFunctionCallWithSpecifiers.h"
 
 #define LOCTEXT_NAMESPACE "NiagaraNodeFunctionCall"
 
@@ -128,6 +129,22 @@ void UNiagaraNodeFunctionCall::UpgradeDIFunctionCalls()
 	}
 }
 
+TSharedPtr<SGraphNode> UNiagaraNodeFunctionCall::CreateVisualWidget()
+{
+	if (!FunctionScript && FunctionSpecifiers.Num() == 0)
+	{
+		FunctionSpecifiers = Signature.FunctionSpecifiers;
+	}
+	if (FunctionSpecifiers.Num() == 0)
+	{
+		return Super::CreateVisualWidget();
+	}
+	else
+	{
+		return SNew(SNiagaraGraphNodeFunctionCallWithSpecifiers, this);
+	}
+}
+
 void UNiagaraNodeFunctionCall::PostEditChangeProperty(struct FPropertyChangedEvent& PropertyChangedEvent)
 {
 	if (PropertyChangedEvent.Property != nullptr)
@@ -543,7 +560,7 @@ void UNiagaraNodeFunctionCall::Compile(class FHlslNiagaraTranslator* Translator,
 				}
 			}
 		}
-
+		Signature.FunctionSpecifiers = FunctionSpecifiers;
 		bError = CompileInputPins(Translator, Inputs);
 	}		
 	else
diff --git a/Engine/Plugins/FX/Niagara/Source/NiagaraEditor/Private/Widgets/SNiagaraGraphNodeFunctionCallWithSpecifiers.cpp b/Engine/Plugins/FX/Niagara/Source/NiagaraEditor/Private/Widgets/SNiagaraGraphNodeFunctionCallWithSpecifiers.cpp
new file mode 100644
--- /dev/null
+++ b/Engine/Plugins/FX/Niagara/Source/NiagaraEditor/Private/Widgets/SNiagaraGraphNodeFunctionCallWithSpecifiers.cpp
@@ -0,0 +1,89 @@
+// Copyright 1998-2019 Epic Games, Inc. All Rights Reserved.
+
+#include "SNiagaraGraphNodeFunctionCallWithSpecifiers.h"
+#include "ISinglePropertyView.h"
+#include "Modules/ModuleManager.h"
+#include "NiagaraNodeFunctionCall.h"
+#include "GraphEditorSettings.h"
+#include "Widgets/Input/SEditableTextBox.h"
+
+#define LOCTEXT_NAMESPACE "SNiagaraGraphNodeFunctionCallWithSpecifiers"
+
+void SNiagaraGraphNodeFunctionCallWithSpecifiers::Construct(const FArguments& InArgs, UEdGraphNode* InGraphNode)
+{
+	GraphNode = InGraphNode;
+	this->FunctionSpecifiers = &(dynamic_cast<UNiagaraNodeFunctionCall*>(GraphNode)->FunctionSpecifiers);
+	RegisterNiagaraGraphNode(InGraphNode);
+	this->UpdateGraphNode();
+}
+
+TSharedRef<SWidget> SNiagaraGraphNodeFunctionCallWithSpecifiers::CreateNodeContentArea()
+{
+	TSharedPtr<SVerticalBox> FunctionSpecifierWidget = SNew(SVerticalBox);
+	for (TTuple<FName, FName>& Entry : *FunctionSpecifiers)
+	{
+		FunctionSpecifierWidget->AddSlot()
+			.VAlign(VAlign_Center)
+			[
+				SNew(SNiagaraFunctionSpecifier, Entry.Key, Entry.Value, *FunctionSpecifiers)
+			];
+	}
+
+	TSharedRef<SWidget> ContentAreaWidget = SGraphNode::CreateNodeContentArea();
+	TSharedPtr<SVerticalBox> VertContainer = SNew(SVerticalBox)
+		+ SVerticalBox::Slot()
+		.AutoHeight()
+		.VAlign(VAlign_Center)
+		.Padding(Settings->GetInputPinPadding())
+		[
+			FunctionSpecifierWidget.ToSharedRef()
+		]
+		+ SVerticalBox::Slot()
+		.AutoHeight()
+		[
+			ContentAreaWidget
+		];
+	TSharedRef<SWidget> RetWidget = VertContainer.ToSharedRef();
+	return RetWidget;
+}
+
+void SNiagaraFunctionSpecifier::Construct(const FArguments& InArgs, FName InAttributeName, FName InValueName, TMap<FName, FName>& InSpecifiers)
+{
+	AttributeName = InAttributeName;
+	ValueName = InValueName;
+	Specifiers = &InSpecifiers;
+	ChildSlot
+		.VAlign(VAlign_Fill)
+		.HAlign(HAlign_Fill)
+		[
+			SNew(SHorizontalBox)
+			+ SHorizontalBox::Slot()
+			.HAlign(HAlign_Fill)
+			.AutoWidth()
+			.Padding(5)
+			[
+				SNew(STextBlock)
+				.Text(FText::FromName(AttributeName))
+			]
+			+ SHorizontalBox::Slot()
+			.HAlign(HAlign_Fill)
+			.AutoWidth()
+			.Padding(5)
+			[
+				SNew(SEditableTextBox)
+				.Text(FText::FromName(ValueName))
+				.OnTextCommitted(this, &SNiagaraFunctionSpecifier::OnValueNameCommitted)
+			]
+		];
+}
+
+void SNiagaraFunctionSpecifier::OnValueNameCommitted(const FText& InText, ETextCommit::Type InCommitType)
+{
+	ValueName = FName(*InText.ToString());
+	if (Specifiers)
+	{
+		Specifiers->Add(AttributeName, ValueName);
+	}
+}
+
+#undef LOCTEXT_NAMESPACE
\ No newline at end of file
diff --git a/Engine/Plugins/FX/Niagara/Source/NiagaraEditor/Private/Widgets/SNiagaraGraphNodeFunctionCallWithSpecifiers.h b/Engine/Plugins/FX/Niagara/Source/NiagaraEditor/Private/Widgets/SNiagaraGraphNodeFunctionCallWithSpecifiers.h
new file mode 100644
--- /dev/null
+++ b/Engine/Plugins/FX/Niagara/Source/NiagaraEditor/Private/Widgets/SNiagaraGraphNodeFunctionCallWithSpecifiers.h
@@ -0,0 +1,32 @@
+// Copyright 1998-2019 Epic Games, Inc. All Rights Reserved.
+
+#pragma once
+
+#include "SNiagaraGraphNode.h"
+
+/** A graph node widget representing a function call with function specifiers. */
+class SNiagaraGraphNodeFunctionCallWithSpecifiers : public SNiagaraGraphNode
+{
+public:
+	SLATE_BEGIN_ARGS(SNiagaraGraphNodeFunctionCallWithSpecifiers) {}
+	SLATE_END_ARGS();
+
+	void Construct(const FArguments& InArgs, UEdGraphNode* InGraphNode);
+	virtual TSharedRef<SWidget> CreateNodeContentArea() override;
+	
+	TMap<FName, FName>* FunctionSpecifiers;
+};
+
+class SNiagaraFunctionSpecifier : public SCompoundWidget
+{
+public:
+	SLATE_BEGIN_ARGS(SNiagaraFunctionSpecifier) {}
+	SLATE_END_ARGS();
+	void Construct(const FArguments& InArgs, FName InAttributeName, FName InValueName, TMap<FName, FName>& InSpecifiers);
+
+	void OnValueNameCommitted(const FText& InText, ETextCommit::Type InCommitType);
+
+	FName AttributeName;
+	FName ValueName;
+	TMap<FName, FName>* Specifiers;
+};
\ No newline at end of file
diff --git a/Engine/Plugins/FX/Niagara/Source/NiagaraEditor/Public/NiagaraNodeFunctionCall.h b/Engine/Plugins/FX/Niagara/Source/NiagaraEditor/Public/NiagaraNodeFunctionCall.h
--- a/Engine/Plugins/FX/Niagara/Source/NiagaraEditor/Public/NiagaraNodeFunctionCall.h
+++ b/Engine/Plugins/FX/Niagara/Source/NiagaraEditor/Public/NiagaraNodeFunctionCall.h
@@ -70,6 +70,9 @@ public:
 	UPROPERTY()
 	FNiagaraFunctionSignature Signature;
 
+	UPROPERTY(VisibleAnywhere, Category = "Function")
+	TMap<FName, FName> FunctionSpecifiers;
+
 	/** All the input values the function propagates to the next higher caller instead of forcing the user to set them directly. */
 	UPROPERTY()
 	TArray<FNiagaraPropagatedVariable> PropagatedStaticSwitchParameters;
@@ -128,6 +131,8 @@ public:
 
 	/** Does any automated data manipulated required to update DI function call nodes to the current version. */
 	void UpgradeDIFunctionCalls();
+
+	virtual TSharedPtr<SGraphNode> CreateVisualWidget() override;
 protected:
 
 	virtual bool IsValidPinToCompile(UEdGraphPin* Pin) const;
