Fix attribute reader DI not showing up in UI

diff --git a/Engine/Plugins/FX/Niagara/Source/Niagara/Private/NiagaraDataInterfaceParticleRead.cpp b/Engine/Plugins/FX/Niagara/Source/Niagara/Private/NiagaraDataInterfaceParticleRead.cpp
--- a/Engine/Plugins/FX/Niagara/Source/Niagara/Private/NiagaraDataInterfaceParticleRead.cpp
+++ b/Engine/Plugins/FX/Niagara/Source/Niagara/Private/NiagaraDataInterfaceParticleRead.cpp
@@ -23,6 +23,11 @@ UNiagaraDataInterfaceParticleRead::UNiagaraDataInterfaceParticleRead(FObjectInit
 void UNiagaraDataInterfaceParticleRead::PostInitProperties()
 {
 	Super::PostInitProperties();
+
+	if (HasAnyFlags(RF_ClassDefaultObject))
+	{
+		FNiagaraTypeRegistry::Register(FNiagaraTypeDefinition(GetClass()), true, false, false);
+	}
 }
 
 void UNiagaraDataInterfaceParticleRead::PostLoad()
