Niagara - Fix shutdown issues.
+ Fix shared ptr cycles between the emitter handle view model and the emitter view models caused by shared refs being used for additional delegate parameters.
+ Fix a typo in the system view model where it wasn't properly cleaning up the system stack view model.

diff --git a/Engine/Plugins/FX/Niagara/Source/NiagaraEditor/Private/ViewModels/NiagaraSystemViewModel.cpp b/Engine/Plugins/FX/Niagara/Source/NiagaraEditor/Private/ViewModels/NiagaraSystemViewModel.cpp
--- a/Engine/Plugins/FX/Niagara/Source/NiagaraEditor/Private/ViewModels/NiagaraSystemViewModel.cpp
+++ b/Engine/Plugins/FX/Niagara/Source/NiagaraEditor/Private/ViewModels/NiagaraSystemViewModel.cpp
@@ -155,7 +155,7 @@ void FNiagaraSystemViewModel::Cleanup()
 	RemoveSystemEventHandlers();
 	SystemScriptViewModel.Reset();
 
-	if (SystemScriptViewModel.IsValid())
+	if (SystemStackViewModel != nullptr)
 	{
 		SystemStackViewModel->Finalize();
 		SystemStackViewModel = nullptr;
@@ -765,12 +765,12 @@ void FNiagaraSystemViewModel::RefreshEmitterHandleViewModels()
 		if (bAdd)
 		{
 			ViewModel = MakeShared<FNiagaraEmitterHandleViewModel>();
-			ViewModel->OnPropertyChanged().AddRaw(this, &FNiagaraSystemViewModel::EmitterHandlePropertyChanged, ViewModel.ToSharedRef());
-			ViewModel->OnNameChanged().AddRaw(this, &FNiagaraSystemViewModel::EmitterHandleNameChanged, ViewModel.ToSharedRef());
-			ViewModel->GetEmitterViewModel()->OnPropertyChanged().AddRaw(this, &FNiagaraSystemViewModel::EmitterPropertyChanged, ViewModel.ToSharedRef());
+			ViewModel->OnPropertyChanged().AddRaw(this, &FNiagaraSystemViewModel::EmitterHandlePropertyChanged, ViewModel->GetId());
+			ViewModel->OnNameChanged().AddRaw(this, &FNiagaraSystemViewModel::EmitterHandleNameChanged);
+			ViewModel->GetEmitterViewModel()->OnPropertyChanged().AddRaw(this, &FNiagaraSystemViewModel::EmitterPropertyChanged);
 			ViewModel->GetEmitterViewModel()->OnScriptCompiled().AddRaw(this, &FNiagaraSystemViewModel::ScriptCompiled);
-			ViewModel->GetEmitterViewModel()->OnScriptGraphChanged().AddRaw(this, &FNiagaraSystemViewModel::EmitterScriptGraphChanged, ViewModel.ToSharedRef());
-			ViewModel->GetEmitterViewModel()->OnScriptParameterStoreChanged().AddRaw(this, &FNiagaraSystemViewModel::EmitterParameterStoreChanged, ViewModel.ToSharedRef());
+			ViewModel->GetEmitterViewModel()->OnScriptGraphChanged().AddRaw(this, &FNiagaraSystemViewModel::EmitterScriptGraphChanged, ViewModel->GetId());
+			ViewModel->GetEmitterViewModel()->OnScriptParameterStoreChanged().AddRaw(this, &FNiagaraSystemViewModel::EmitterParameterStoreChanged);
 		}
 		else
 		{
@@ -1171,7 +1171,7 @@ void FNiagaraSystemViewModel::UpdateCompiledDataInterfaces(UNiagaraDataInterface
 	}
 }
 
-void FNiagaraSystemViewModel::EmitterHandlePropertyChanged(TSharedRef<FNiagaraEmitterHandleViewModel> EmitterHandleViewModel)
+void FNiagaraSystemViewModel::EmitterHandlePropertyChanged(FGuid OwningEmitterHandleId)
 {
 	// When the emitter handle changes, refresh the System scripts emitter nodes and the sequencer tracks just in case the
 	// property that changed was the handles emitter.
@@ -1181,7 +1181,7 @@ void FNiagaraSystemViewModel::EmitterHandlePropertyChanged(TSharedRef<FNiagaraEm
 		for (UMovieSceneTrack* Track : NiagaraSequence->GetMovieScene()->GetMasterTracks())
 		{
 			UMovieSceneNiagaraEmitterTrack* EmitterTrack = CastChecked<UMovieSceneNiagaraEmitterTrack>(Track);
-			if (EmitterTrack->GetEmitterHandleViewModel() == EmitterHandleViewModel)
+			if (EmitterTrack->GetEmitterHandleViewModel()->GetId() == OwningEmitterHandleId)
 			{
 				EmitterTrack->UpdateTrackFromEmitterGraphChange(NiagaraSequence->GetMovieScene()->GetTickResolution());
 			}
@@ -1191,13 +1191,13 @@ void FNiagaraSystemViewModel::EmitterHandlePropertyChanged(TSharedRef<FNiagaraEm
 	ResetSystem(ETimeResetMode::AllowResetTime, EMultiResetMode::ResetThisInstance, EReinitMode::ReinitializeSystem);
 }
 
-void FNiagaraSystemViewModel::EmitterHandleNameChanged(TSharedRef<FNiagaraEmitterHandleViewModel> EmitterHandleViewModel)
+void FNiagaraSystemViewModel::EmitterHandleNameChanged()
 {
 	GetDefault<UEdGraphSchema_NiagaraSystemOverview>()->ForceVisualizationCacheClear();
 	CompileSystem(false);
 }
 
-void FNiagaraSystemViewModel::EmitterPropertyChanged(TSharedRef<FNiagaraEmitterHandleViewModel> EmitterHandleViewModel)
+void FNiagaraSystemViewModel::EmitterPropertyChanged()
 {
 	ResetSystem(ETimeResetMode::AllowResetTime, EMultiResetMode::ResetThisInstance, EReinitMode::ReinitializeSystem);
 }
@@ -1213,14 +1213,14 @@ void FNiagaraSystemViewModel::SystemParameterStoreChanged(const FNiagaraParamete
 	UpdateSimulationFromParameterChange();
 }
 
-void FNiagaraSystemViewModel::EmitterScriptGraphChanged(const FEdGraphEditAction& InAction, const UNiagaraScript& OwningScript, const TSharedRef<FNiagaraEmitterHandleViewModel> OwningEmitterHandleViewModel)
+void FNiagaraSystemViewModel::EmitterScriptGraphChanged(const FEdGraphEditAction& InAction, const UNiagaraScript& OwningScript, FGuid OwningEmitterHandleId)
 {
 	if (bUpdatingEmittersFromSequencerDataChange == false)
 	{
-		EmitterIdsRequiringSequencerTrackUpdate.AddUnique(OwningEmitterHandleViewModel->GetId());
+		EmitterIdsRequiringSequencerTrackUpdate.AddUnique(OwningEmitterHandleId);
 	}
 	// Remove from cache on graph change 
-	EmitterToCachedStackModuleData.Remove(OwningEmitterHandleViewModel->GetId());
+	EmitterToCachedStackModuleData.Remove(OwningEmitterHandleId);
 }
 
 void FNiagaraSystemViewModel::SystemScriptGraphChanged(const FEdGraphEditAction& InAction)
@@ -1228,7 +1228,7 @@ void FNiagaraSystemViewModel::SystemScriptGraphChanged(const FEdGraphEditAction&
 	EmitterToCachedStackModuleData.Empty();
 }
 
-void FNiagaraSystemViewModel::EmitterParameterStoreChanged(const FNiagaraParameterStore& ChangedParameterStore, const UNiagaraScript& OwningScript, const TSharedRef<FNiagaraEmitterHandleViewModel> OwningEmitterHandleViewModel)
+void FNiagaraSystemViewModel::EmitterParameterStoreChanged(const FNiagaraParameterStore& ChangedParameterStore, const UNiagaraScript& OwningScript)
 {
 	if (bUpdatingEmittersFromSequencerDataChange == false)
 	{
diff --git a/Engine/Plugins/FX/Niagara/Source/NiagaraEditor/Public/ViewModels/NiagaraSystemViewModel.h b/Engine/Plugins/FX/Niagara/Source/NiagaraEditor/Public/ViewModels/NiagaraSystemViewModel.h
--- a/Engine/Plugins/FX/Niagara/Source/NiagaraEditor/Public/ViewModels/NiagaraSystemViewModel.h
+++ b/Engine/Plugins/FX/Niagara/Source/NiagaraEditor/Public/ViewModels/NiagaraSystemViewModel.h
@@ -324,13 +324,13 @@ private:
 	void UpdateCompiledDataInterfaces(UNiagaraDataInterface* ChangedDataInterface);
 
 	/** Called whenever a property on the emitter handle changes. */
-	void EmitterHandlePropertyChanged(TSharedRef<FNiagaraEmitterHandleViewModel> EmitterHandleViewModel);
+	void EmitterHandlePropertyChanged(FGuid OwningEmitterHandleId);
 
 	/** Called whenever the name on an emitter handle changes. */
-	void EmitterHandleNameChanged(TSharedRef<FNiagaraEmitterHandleViewModel> EmitterHandleViewModel);
+	void EmitterHandleNameChanged();
 
 	/** Called whenever a property on the emitter changes. */
-	void EmitterPropertyChanged(TSharedRef<FNiagaraEmitterHandleViewModel> EmitterHandleViewModel);
+	void EmitterPropertyChanged();
 
 	/** 
 	 * Called whenever a parameter store owned by the system changes.
@@ -340,7 +340,7 @@ private:
 	void SystemParameterStoreChanged(const FNiagaraParameterStore& ChangedParameterStore, const UNiagaraScript* OwningScript);
 
 	/** Called whenever an emitter's script graph changes. */
-	void EmitterScriptGraphChanged(const FEdGraphEditAction& InAction, const UNiagaraScript& OwningScript, const TSharedRef<FNiagaraEmitterHandleViewModel> OwningEmitterHandleViewModel);
+	void EmitterScriptGraphChanged(const FEdGraphEditAction& InAction, const UNiagaraScript& OwningScript, FGuid OwningEmitterHandleId);
 
 	/** Called whenever the system script graph changes. */
 	void SystemScriptGraphChanged(const FEdGraphEditAction& InAction);
@@ -350,7 +350,7 @@ private:
 	* @param ChangedParameterStore The parameter store that changed.
 	* @param OwningScript The script that owns the parameter store, if there is one.
 	*/
-	void EmitterParameterStoreChanged(const FNiagaraParameterStore& ChangedParameterStore, const UNiagaraScript& OwningScript, const TSharedRef<FNiagaraEmitterHandleViewModel> OwningEmitterHandleViewModel);
+	void EmitterParameterStoreChanged(const FNiagaraParameterStore& ChangedParameterStore, const UNiagaraScript& OwningScript);
 
 	/** Updates the current simulation for a parameter changing, based on the current simulation options. */
 	void UpdateSimulationFromParameterChange();
