Niagara - Decrease the log spam for Niagara when things are working as expected.

diff --git a/Engine/Plugins/FX/Niagara/Source/Niagara/Private/NiagaraSystem.cpp b/Engine/Plugins/FX/Niagara/Source/Niagara/Private/NiagaraSystem.cpp
--- a/Engine/Plugins/FX/Niagara/Source/Niagara/Private/NiagaraSystem.cpp
+++ b/Engine/Plugins/FX/Niagara/Source/Niagara/Private/NiagaraSystem.cpp
@@ -344,12 +344,12 @@ void UNiagaraSystem::PostLoad()
 		UE_LOG(LogNiagara, Log, TEXT("System %s being rebuilt because UNiagaraEmitter::GetForceCompileOnLoad() == true."), *GetPathName());
 	}
 
-	if (bSystemScriptsAreSynchronized == false)
+	if (bSystemScriptsAreSynchronized == false && GEnableVerboseNiagaraChangeIdLogging)
 	{
 		UE_LOG(LogNiagara, Log, TEXT("System %s being compiled because there were changes to a system script Change ID."), *GetPathName());
 	}
 
-	if (bEmitterScriptsAreSynchronized == false)
+	if (bEmitterScriptsAreSynchronized == false && GEnableVerboseNiagaraChangeIdLogging)
 	{
 		UE_LOG(LogNiagara, Log, TEXT("System %s being compiled because there were changes to an emitter script Change ID."), *GetPathName());
 	}
@@ -703,10 +703,6 @@ bool UNiagaraSystem::QueryCompileComplete(bool bWait, bool bDoPost, bool bDoNotA
 					{
 						UE_LOG(LogNiagara, Log, TEXT("UNiagraScript \'%s\' was built locally.."), *EmitterCompiledScriptPair.CompiledScript->GetFullName());
 					}
-					else
-					{
-						UE_LOG(LogNiagara, Log, TEXT("UNiagraScript \'%s\' was pulled from DDC."), *EmitterCompiledScriptPair.CompiledScript->GetFullName());
-					}
 
 					TSharedPtr<FNiagaraVMExecutableData> ExeData = MakeShared<FNiagaraVMExecutableData>();
 					EmitterCompiledScriptPair.CompileResults = ExeData;
@@ -853,7 +849,7 @@ bool UNiagaraSystem::QueryCompileComplete(bool bWait, bool bDoPost, bool bDoNotA
 
 		UpdatePostCompileDIInfo();
 
-		UE_LOG(LogNiagara, Log, TEXT("Compiling System %s took %f sec (wall time), %f sec (combined time)."), *GetFullName(), (float)(FPlatformTime::Seconds() - ActiveCompilations[ActiveCompileIdx].StartTime),
+		UE_LOG(LogNiagara, Verbose, TEXT("Compiling System %s took %f sec (wall time), %f sec (combined time)."), *GetFullName(), (float)(FPlatformTime::Seconds() - ActiveCompilations[ActiveCompileIdx].StartTime),
 			CombinedCompileTime);
 
 		ActiveCompilations.RemoveAt(ActiveCompileIdx);
diff --git a/Engine/Plugins/FX/Niagara/Source/NiagaraEditor/Private/NiagaraCompiler.cpp b/Engine/Plugins/FX/Niagara/Source/NiagaraEditor/Private/NiagaraCompiler.cpp
--- a/Engine/Plugins/FX/Niagara/Source/NiagaraEditor/Private/NiagaraCompiler.cpp
+++ b/Engine/Plugins/FX/Niagara/Source/NiagaraEditor/Private/NiagaraCompiler.cpp
@@ -496,7 +496,7 @@ TSharedPtr<FNiagaraCompileRequestDataBase, ESPMode::ThreadSafe> FNiagaraEditorMo
 
 	}
 
-	UE_LOG(LogNiagaraEditor, Log, TEXT("'%s' Precompile took %f sec."), *InObj->GetOutermost()->GetName(),
+	UE_LOG(LogNiagaraEditor, Verbose, TEXT("'%s' Precompile took %f sec."), *InObj->GetOutermost()->GetName(),
 		(float)(FPlatformTime::Seconds() - StartTime));
 
 	return BasePtr;
diff --git a/Engine/Plugins/FX/Niagara/Source/NiagaraEditor/Private/NiagaraGraph.cpp b/Engine/Plugins/FX/Niagara/Source/NiagaraEditor/Private/NiagaraGraph.cpp
--- a/Engine/Plugins/FX/Niagara/Source/NiagaraEditor/Private/NiagaraGraph.cpp
+++ b/Engine/Plugins/FX/Niagara/Source/NiagaraEditor/Private/NiagaraGraph.cpp
@@ -226,7 +226,6 @@ void UNiagaraGraph::PostLoad()
 				// so this will extract "/Path/To/ScriptName"
 				PathName = PathName.Left(ColonPos);
 			}
-			UE_LOG(LogNiagaraEditor, Log, TEXT("Migrated old metadata entry for variable \"%s\" in \"%s\""), *It.Key().GetName().ToString(), *PathName);
 		}
 		VariableToMetaData_DEPRECATED.Empty();
 	}
diff --git a/Engine/Plugins/FX/Niagara/Source/NiagaraEditor/Private/NiagaraScriptMergeManager.cpp b/Engine/Plugins/FX/Niagara/Source/NiagaraEditor/Private/NiagaraScriptMergeManager.cpp
--- a/Engine/Plugins/FX/Niagara/Source/NiagaraEditor/Private/NiagaraScriptMergeManager.cpp
+++ b/Engine/Plugins/FX/Niagara/Source/NiagaraEditor/Private/NiagaraScriptMergeManager.cpp
@@ -773,23 +773,16 @@ FNiagaraScriptMergeManager::FApplyDiffResults FNiagaraScriptMergeManager::Resolv
 		// Now go through all the nodes and set the persistent change ids if we encounter a node that needs it's change id kept.
 		bool bAnySet = false;
 		while (It)
-		{			
-			bool bFound = false;
+		{
 			for (UNiagaraNode* Node : Nodes)
 			{
 				if (Node->NodeGuid == It.Key())
 				{
 					Node->ForceChangeId(It.Value(), false);
 					bAnySet = true;
-					bFound = true;
 					break;
 				}
 			}
-
-			if (!bFound)
-			{
-				UE_LOG(LogNiagaraEditor, Log, TEXT("Looks like %s was removed.. is that expected?"), *It.Key().ToString())
-			}
 			++It;
 		}
 
