#UE4 Basic implementation of IsCachedCookedPlatformDataLoaded for UNiagaraScript. Fixes a crash when using ConcurrentSave

diff --git a/Engine/Plugins/FX/Niagara/Source/Niagara/Classes/NiagaraScript.h b/Engine/Plugins/FX/Niagara/Source/Niagara/Classes/NiagaraScript.h
--- a/Engine/Plugins/FX/Niagara/Source/Niagara/Classes/NiagaraScript.h
+++ b/Engine/Plugins/FX/Niagara/Source/Niagara/Classes/NiagaraScript.h
@@ -480,6 +480,7 @@ public:
 
 	NIAGARA_API void CacheResourceShadersForRendering(bool bRegenerateId, bool bForceRecompile=false);
 	void BeginCacheForCookedPlatformData(const ITargetPlatform *TargetPlatform);
+	bool IsCachedCookedPlatformDataLoaded(const ITargetPlatform* TargetPlatform) override;
 	void CacheShadersForResources(EShaderPlatform ShaderPlatform, FNiagaraShaderScript *ResourceToCache, bool bApplyCompletedShaderMapForRendering, bool bForceRecompile = false, bool bCooking=false);
 	FNiagaraShaderScript* AllocateResource();
 	FNiagaraShaderScript *GetRenderThreadScript()
diff --git a/Engine/Plugins/FX/Niagara/Source/Niagara/Private/NiagaraScript.cpp b/Engine/Plugins/FX/Niagara/Source/Niagara/Private/NiagaraScript.cpp
--- a/Engine/Plugins/FX/Niagara/Source/Niagara/Private/NiagaraScript.cpp
+++ b/Engine/Plugins/FX/Niagara/Source/Niagara/Private/NiagaraScript.cpp
@@ -1120,6 +1120,19 @@ void UNiagaraScript::BeginCacheForCookedPlatformData(const ITargetPlatform *Targ
 	}
 }
 
+bool UNiagaraScript::IsCachedCookedPlatformDataLoaded(const ITargetPlatform* TargetPlatform)
+{
+	if (ShouldCacheShadersForCooking())
+	{
+		if (UNiagaraSystem* SystemOwner = FindRootSystem())
+		{
+			return !SystemOwner->HasOutstandingCompilationRequests();
+		}
+	}
+
+	return true;
+}
+
 void UNiagaraScript::CacheResourceShadersForCooking(EShaderPlatform ShaderPlatform, TArray<FNiagaraShaderScript*>& InOutCachedResources)
 {
 	if (CanBeRunOnGpu())
