#UE4 Fix some WITH_EDITOR usage

diff --git a/Engine/Plugins/FX/Niagara/Source/Niagara/Classes/NiagaraScript.h b/Engine/Plugins/FX/Niagara/Source/Niagara/Classes/NiagaraScript.h
--- a/Engine/Plugins/FX/Niagara/Source/Niagara/Classes/NiagaraScript.h
+++ b/Engine/Plugins/FX/Niagara/Source/Niagara/Classes/NiagaraScript.h
@@ -476,12 +476,14 @@ public:
 	//~ End UObject interface
 
 	// Infrastructure for GPU compute Shaders
+#if WITH_EDITOR
 	NIAGARA_API void CacheResourceShadersForCooking(EShaderPlatform ShaderPlatform, TArray<FNiagaraShaderScript*>& InOutCachedResources);
 
 	NIAGARA_API void CacheResourceShadersForRendering(bool bRegenerateId, bool bForceRecompile=false);
 	void BeginCacheForCookedPlatformData(const ITargetPlatform *TargetPlatform);
-	bool IsCachedCookedPlatformDataLoaded(const ITargetPlatform* TargetPlatform) override;
+	virtual bool IsCachedCookedPlatformDataLoaded(const ITargetPlatform* TargetPlatform) override;
 	void CacheShadersForResources(EShaderPlatform ShaderPlatform, FNiagaraShaderScript *ResourceToCache, bool bApplyCompletedShaderMapForRendering, bool bForceRecompile = false, bool bCooking=false);
+#endif // WITH_EDITOR
 	FNiagaraShaderScript* AllocateResource();
 	FNiagaraShaderScript *GetRenderThreadScript()
 	{
