Copying //UE4/Dev-Physics to Dev-Main (//UE4/Dev-Main) [at] 7703071

diff --git a/Engine/Plugins/FX/Niagara/Source/Niagara/Private/NiagaraComponent.cpp b/Engine/Plugins/FX/Niagara/Source/Niagara/Private/NiagaraComponent.cpp
--- a/Engine/Plugins/FX/Niagara/Source/Niagara/Private/NiagaraComponent.cpp
+++ b/Engine/Plugins/FX/Niagara/Source/Niagara/Private/NiagaraComponent.cpp
@@ -565,6 +565,31 @@ bool UNiagaraComponent::IsPaused()const
 	return false;
 }
 
+UNiagaraDataInterface * UNiagaraComponent::GetDataInterface(const FString &Name)
+{
+
+	// @todo-threadsafety Think of a better way to do this!
+	if (!SystemInstance || SystemInstance->GetEmitters().Num() == 0 || !SystemInstance->GetEmitters()[0]->GetGPUContext())
+	{
+		return nullptr;
+	}
+	
+	FNiagaraComputeExecutionContext* GPUContext = SystemInstance->GetEmitters()[0]->GetGPUContext();
+	const TArray<FNiagaraScriptDataInterfaceCompileInfo> &DataInterfaceInfo = GPUContext->GPUScript->GetVMExecutableData().DataInterfaceInfo;
+	const TArray<UNiagaraDataInterface*>& DataInterfaces = GPUContext->CombinedParamStore.GetDataInterfaces();
+
+	int Index = 0;
+	for (UNiagaraDataInterface* Interface : DataInterfaces)
+	{
+		if (DataInterfaceInfo[Index].Name.GetPlainNameString() == Name)
+		{			
+			return Interface;
+		}	
+		++Index;
+	}
+	return nullptr;
+}
+
 bool UNiagaraComponent::IsWorldReadyToRun() const
 {
 	// The niagara system instance assumes that a batcher exists when it is created. We need to wait until this has happened before successfully activating this system.
diff --git a/Engine/Plugins/FX/Niagara/Source/Niagara/Private/NiagaraDataInterfaceGrid2DCollection.cpp b/Engine/Plugins/FX/Niagara/Source/Niagara/Private/NiagaraDataInterfaceGrid2DCollection.cpp
new file mode 100644
--- /dev/null
+++ b/Engine/Plugins/FX/Niagara/Source/Niagara/Private/NiagaraDataInterfaceGrid2DCollection.cpp
@@ -0,0 +1,467 @@
+// Copyright 1998-2019 Epic Games, Inc. All Rights Reserved.
+#include "NiagaraDataInterfaceGrid2DCollection.h"
+#include "NiagaraShader.h"
+#include "ShaderParameterUtils.h"
+#include "ClearQuad.h"
+#include "TextureResource.h"
+#include "Engine/Texture2D.h"
+#include "NiagaraEmitterInstanceBatcher.h"
+#include "NiagaraSystemInstance.h"
+#include "Engine/TextureRenderTarget2D.h"
+
+#define LOCTEXT_NAMESPACE "NiagaraDataInterfaceGrid2DCollection"
+
+static const FString GridName(TEXT("Grid_"));
+static const FString OutputGridName(TEXT("OutputGrid_"));
+
+
+// Global VM function names, also used by the shaders code generation methods.
+static const FName SetValueFunctionName("SetGridValue");
+static const FName GetValueFunctionName("GetGridValue");
+
+static const FName SampleGridFunctionName("SampleGrid");
+
+
+/*--------------------------------------------------------------------------------------------------------------------------*/
+struct FNiagaraDataInterfaceParametersCS_Grid2DCollection : public FNiagaraDataInterfaceParametersCS
+{
+	virtual void Bind(const FNiagaraDataInterfaceParamRef& ParamRef, const class FShaderParameterMap& ParameterMap) override
+	{			
+		NumCellsParam.Bind(ParameterMap, *(NumCellsName + ParamRef.ParameterInfo.DataInterfaceHLSLSymbol));
+
+		CellSizeParam.Bind(ParameterMap, *(CellSizeName + ParamRef.ParameterInfo.DataInterfaceHLSLSymbol));
+
+		WorldBBoxMinParam.Bind(ParameterMap, *(WorldBBoxMinName + ParamRef.ParameterInfo.DataInterfaceHLSLSymbol));
+		WorldBBoxSizeParam.Bind(ParameterMap, *(WorldBBoxSizeName + ParamRef.ParameterInfo.DataInterfaceHLSLSymbol));
+
+		GridParam.Bind(ParameterMap, *(GridName + ParamRef.ParameterInfo.DataInterfaceHLSLSymbol));
+		OutputGridParam.Bind(ParameterMap, *(OutputGridName + ParamRef.ParameterInfo.DataInterfaceHLSLSymbol));
+	}
+
+	virtual void Serialize(FArchive& Ar)override
+	{		
+		Ar << NumCellsParam;		
+		Ar << CellSizeParam;	
+		Ar << WorldBBoxMinParam;
+		Ar << WorldBBoxSizeParam;
+
+		Ar << GridParam;
+		Ar << OutputGridParam;		
+	}
+
+	virtual void Set(FRHICommandList& RHICmdList, const FNiagaraDataInterfaceSetArgs& Context) const override
+	{
+		check(IsInRenderingThread());
+
+		// Get shader and DI
+		FRHIComputeShader* ComputeShaderRHI = Context.Shader->GetComputeShader();
+		FNiagaraDataInterfaceProxyGrid2DCollection* VFDI = static_cast<FNiagaraDataInterfaceProxyGrid2DCollection*>(Context.DataInterface);
+		
+		Grid2DCollectionRWInstanceData* ProxyData = VFDI->SystemInstancesToProxyData.Find(Context.SystemInstance);
+		ensure(ProxyData);
+
+		int NumCellsTmp[2];
+		NumCellsTmp[0] = ProxyData->NumCellsX;
+		NumCellsTmp[1] = ProxyData->NumCellsY;
+		SetShaderValue(RHICmdList, ComputeShaderRHI, NumCellsParam, NumCellsTmp);		
+		
+		// #todo(dmp): precompute these values
+		float CellSizeTmp[2];
+		CellSizeTmp[0] = ProxyData->WorldBBoxSize.X / ProxyData->NumCellsX;
+		CellSizeTmp[1] = ProxyData->WorldBBoxSize.Y / ProxyData->NumCellsY;
+		SetShaderValue(RHICmdList, ComputeShaderRHI, CellSizeParam, CellSizeTmp);		
+		
+		SetShaderValue(RHICmdList, ComputeShaderRHI, WorldBBoxMinParam, ProxyData->WorldBBoxMin);
+		SetShaderValue(RHICmdList, ComputeShaderRHI, WorldBBoxSizeParam, ProxyData->WorldBBoxSize);
+
+
+
+		if (GridParam.IsBound() && ProxyData->GetCurrentData() != NULL)
+		{
+			RHICmdList.SetShaderResourceViewParameter(Context.Shader->GetComputeShader(), GridParam.GetBaseIndex(), ProxyData->GetCurrentData()->GridBuffer.SRV);
+		}
+
+		if (Context.IsOutputStage && OutputGridParam.IsBound() && ProxyData->GetDestinationData() != NULL)
+		{
+			RHICmdList.SetUAVParameter(Context.Shader->GetComputeShader(), OutputGridParam.GetUAVIndex(), ProxyData->GetDestinationData()->GridBuffer.UAV);		
+		}
+		// Note: There is a flush in PreEditChange to make sure everything is synced up at this point 
+	}
+
+	virtual void Unset(FRHICommandList& RHICmdList, const FNiagaraDataInterfaceSetArgs& Context) const 
+	{
+		if (OutputGridParam.IsBound())
+		{
+			OutputGridParam.UnsetUAV(RHICmdList, Context.Shader->GetComputeShader());
+		}
+	}
+
+private:
+
+	FShaderParameter NumCellsParam;
+	FShaderParameter CellSizeParam;
+	FShaderParameter WorldBBoxMinParam;
+	FShaderParameter WorldBBoxSizeParam;
+
+	FShaderResourceParameter GridParam;
+	FRWShaderParameter OutputGridParam;
+	
+};
+
+
+UNiagaraDataInterfaceGrid2DCollection::UNiagaraDataInterfaceGrid2DCollection(FObjectInitializer const& ObjectInitializer)
+	: Super(ObjectInitializer)
+{
+	Proxy = MakeShared<FNiagaraDataInterfaceProxyGrid2DCollection, ESPMode::ThreadSafe>();
+	RWProxy = (FNiagaraDataInterfaceProxyRW*) Proxy.Get();
+	PushToRenderThread();
+}
+
+
+void UNiagaraDataInterfaceGrid2DCollection::PostInitProperties()
+{
+	Super::PostInitProperties();
+
+	//Can we register data interfaces as regular types and fold them into the FNiagaraVariable framework for UI and function calls etc?
+	if (HasAnyFlags(RF_ClassDefaultObject))
+	{
+		FNiagaraTypeRegistry::Register(FNiagaraTypeDefinition(GetClass()), /*bCanBeParameter*/ true, /*bCanBePayload*/ false, /*bIsUserDefined*/ false);
+	}
+}
+
+void UNiagaraDataInterfaceGrid2DCollection::GetFunctions(TArray<FNiagaraFunctionSignature>& OutFunctions)
+{
+	Super::GetFunctions(OutFunctions);
+
+	{
+		FNiagaraFunctionSignature Sig;
+		Sig.Name = GetValueFunctionName;
+		Sig.Inputs.Add(FNiagaraVariable(FNiagaraTypeDefinition(GetClass()), TEXT("Grid")));
+		Sig.Inputs.Add(FNiagaraVariable(FNiagaraTypeDefinition::GetIntDef(), TEXT("IndexX")));
+		Sig.Inputs.Add(FNiagaraVariable(FNiagaraTypeDefinition::GetIntDef(), TEXT("IndexY")));
+		Sig.Outputs.Add(FNiagaraVariable(FNiagaraTypeDefinition::GetVec4Def(), TEXT("Value")));
+
+		Sig.bMemberFunction = true;
+		Sig.bRequiresContext = false;
+		OutFunctions.Add(Sig);
+	}
+
+	{
+		FNiagaraFunctionSignature Sig;
+		Sig.Name = SetValueFunctionName;
+		Sig.Inputs.Add(FNiagaraVariable(FNiagaraTypeDefinition(GetClass()), TEXT("Grid")));
+		Sig.Inputs.Add(FNiagaraVariable(FNiagaraTypeDefinition::GetIntDef(), TEXT("IndexX")));
+		Sig.Inputs.Add(FNiagaraVariable(FNiagaraTypeDefinition::GetIntDef(), TEXT("IndexY")));
+		Sig.Inputs.Add(FNiagaraVariable(FNiagaraTypeDefinition::GetVec4Def(), TEXT("Value")));
+		Sig.Outputs.Add(FNiagaraVariable(FNiagaraTypeDefinition::GetIntDef(), TEXT("IGNORE")));
+
+		Sig.bMemberFunction = true;
+		Sig.bRequiresContext = false;
+		OutFunctions.Add(Sig);
+	}
+
+	{
+		FNiagaraFunctionSignature Sig;
+		Sig.Name = SampleGridFunctionName;
+		Sig.Inputs.Add(FNiagaraVariable(FNiagaraTypeDefinition(GetClass()), TEXT("Grid")));
+		Sig.Inputs.Add(FNiagaraVariable(FNiagaraTypeDefinition::GetFloatDef(), TEXT("UnitX")));
+		Sig.Inputs.Add(FNiagaraVariable(FNiagaraTypeDefinition::GetFloatDef(), TEXT("UnitY")));		
+		Sig.Outputs.Add(FNiagaraVariable(FNiagaraTypeDefinition::GetVec4Def(), TEXT("Value")));
+
+		Sig.bMemberFunction = true;
+		Sig.bRequiresContext = false;
+		OutFunctions.Add(Sig);
+	}
+
+}
+
+
+void UNiagaraDataInterfaceGrid2DCollection::GetVMExternalFunction(const FVMExternalFunctionBindingInfo& BindingInfo, void* InstanceData, FVMExternalFunction &OutFunc)
+{
+	Super::GetVMExternalFunction(BindingInfo, InstanceData, OutFunc);
+
+	if (BindingInfo.Name == GetValueFunctionName) { OutFunc = FVMExternalFunction::CreateUObject(this, &UNiagaraDataInterfaceRWBase::EmptyVMFunction); }
+	else if (BindingInfo.Name == SetValueFunctionName) { OutFunc = FVMExternalFunction::CreateUObject(this, &UNiagaraDataInterfaceRWBase::EmptyVMFunction); }
+	else if (BindingInfo.Name == SampleGridFunctionName) { OutFunc = FVMExternalFunction::CreateUObject(this, &UNiagaraDataInterfaceRWBase::EmptyVMFunction); }
+}
+
+bool UNiagaraDataInterfaceGrid2DCollection::Equals(const UNiagaraDataInterface* Other) const
+{
+	if (!Super::Equals(Other))
+	{
+		return false;
+	}
+	const UNiagaraDataInterfaceGrid2DCollection* OtherTyped = CastChecked<const UNiagaraDataInterfaceGrid2DCollection>(Other);
+
+	return OtherTyped != nullptr;		
+}
+
+void UNiagaraDataInterfaceGrid2DCollection::GetParameterDefinitionHLSL(FNiagaraDataInterfaceGPUParamInfo& ParamInfo, FString& OutHLSL)
+{
+	Super::GetParameterDefinitionHLSL(ParamInfo, OutHLSL);
+
+	static const TCHAR *FormatDeclarations = TEXT(R"(				
+		Texture2D<float4> {GridName};
+		RWTexture2D<float4> RW{OutputGridName};
+
+		SamplerState {GridName}Sampler
+		{
+			Filter = MIN_MAG_MIP_LINEAR;
+			AddressU = Clamp;
+			AddressV = Clamp;
+		};
+	)");
+	TMap<FString, FStringFormatArg> ArgsDeclarations = {				
+		{ TEXT("GridName"),    GridName + ParamInfo.DataInterfaceHLSLSymbol },
+		{ TEXT("OutputGridName"),    OutputGridName + ParamInfo.DataInterfaceHLSLSymbol },
+	};
+	OutHLSL += FString::Format(FormatDeclarations, ArgsDeclarations);
+}
+
+bool UNiagaraDataInterfaceGrid2DCollection::GetFunctionHLSL(const FName& DefinitionFunctionName, FString InstanceFunctionName, FNiagaraDataInterfaceGPUParamInfo& ParamInfo, FString& OutHLSL)
+{
+	bool ParentRet = Super::GetFunctionHLSL(DefinitionFunctionName, InstanceFunctionName, ParamInfo, OutHLSL);
+	if (ParentRet)
+	{
+		return true;
+	} 
+	else if (DefinitionFunctionName == GetValueFunctionName)
+	{
+		static const TCHAR *FormatBounds = TEXT(R"(
+			void {FunctionName}(int In_IndexX, int In_IndexY, out float4 Out_Val)
+			{
+				Out_Val = {Grid}.Load(int3(In_IndexX, In_IndexY, 0));
+			}
+		)");
+		TMap<FString, FStringFormatArg> ArgsBounds = {
+			{TEXT("FunctionName"), InstanceFunctionName},
+			{TEXT("Grid"), GridName + ParamInfo.DataInterfaceHLSLSymbol},
+		};
+		OutHLSL += FString::Format(FormatBounds, ArgsBounds);
+		return true;
+	}
+	else if (DefinitionFunctionName == SetValueFunctionName)
+	{
+		static const TCHAR *FormatBounds = TEXT(R"(
+			void {FunctionName}(int In_IndexX, int In_IndexY, float4 In_Value, out int val)
+			{				
+				val = 0;
+				RW{OutputGrid}[int2(In_IndexX, In_IndexY)] = In_Value;
+			}
+		)");
+		TMap<FString, FStringFormatArg> ArgsBounds = {
+			{TEXT("FunctionName"), InstanceFunctionName},
+			{TEXT("OutputGrid"), OutputGridName + ParamInfo.DataInterfaceHLSLSymbol},
+		};
+		OutHLSL += FString::Format(FormatBounds, ArgsBounds);
+		return true;
+	}
+	else if (DefinitionFunctionName == SampleGridFunctionName)
+	{
+		static const TCHAR *FormatBounds = TEXT(R"(
+			void {FunctionName}(float In_UnitX, float In_UnitY, out float4 Out_Val)
+			{
+				Out_Val = {Grid}.SampleLevel({Grid}Sampler, float2(In_UnitX, In_UnitY), 0);
+			}
+		)");
+		TMap<FString, FStringFormatArg> ArgsBounds = {
+			{TEXT("FunctionName"), InstanceFunctionName},
+			{TEXT("Grid"), GridName + ParamInfo.DataInterfaceHLSLSymbol},
+		};
+		OutHLSL += FString::Format(FormatBounds, ArgsBounds);
+		return true;
+	}
+	return false;
+}
+
+FNiagaraDataInterfaceParametersCS* UNiagaraDataInterfaceGrid2DCollection::ConstructComputeParameters()const
+{
+	return new FNiagaraDataInterfaceParametersCS_Grid2DCollection();
+}
+
+
+
+bool UNiagaraDataInterfaceGrid2DCollection::CopyToInternal(UNiagaraDataInterface* Destination) const
+{
+	if (!Super::CopyToInternal(Destination))
+	{
+		return false;
+	}
+
+	UNiagaraDataInterfaceGrid2DCollection* OtherTyped = CastChecked<UNiagaraDataInterfaceGrid2DCollection>(Destination);
+
+
+	return true;
+}
+
+bool UNiagaraDataInterfaceGrid2DCollection::InitPerInstanceData(void* PerInstanceData, FNiagaraSystemInstance* SystemInstance)
+{
+	check(Proxy);
+
+	Grid2DCollectionRWInstanceData* InstanceData = new (PerInstanceData) Grid2DCollectionRWInstanceData();
+
+	FNiagaraDataInterfaceProxyGrid2DCollection* RT_Proxy = GetProxyAs<FNiagaraDataInterfaceProxyGrid2DCollection>();
+
+	int RT_NumCellsX = NumCellsX;
+	int RT_NumCellsY = NumCellsY;
+
+	FVector RT_WorldBBoxMin = WorldBBoxMin;
+	FVector2D RT_WorldBBoxSize = WorldBBoxSize;
+
+	// If we are setting the grid from the voxel size, then recompute NumVoxels and change bbox	
+	if (SetGridFromCellSize)
+	{
+		RT_NumCellsX = WorldBBoxSize.X / CellSize;
+		RT_NumCellsY = WorldBBoxSize.Y / CellSize;
+
+		// Pad grid by 1 voxel if our computed bounding box is too small
+		if (!FMath::IsNearlyEqual(CellSize * RT_NumCellsX, WorldBBoxSize.X))
+		{
+			RT_NumCellsX = NumCellsX + 1;
+			RT_NumCellsY = NumCellsY + 1;
+			RT_WorldBBoxSize = FVector2D(RT_NumCellsX, RT_NumCellsY) * CellSize;
+		}
+	}	
+
+	TSet<int> RT_OutputShaderStages = OutputShaderStages;
+	TSet<int> RT_IterationShaderStages = IterationShaderStages;
+
+	// @todo-threadsafety. This would be a race but I'm taking a ref here. Not ideal in the long term.
+	// Push Updates to Proxy.
+	ENQUEUE_RENDER_COMMAND(FUpdateData)(
+		[RT_Proxy, RT_NumCellsX, RT_NumCellsY, RT_WorldBBoxMin, RT_WorldBBoxSize, RT_OutputShaderStages, RT_IterationShaderStages, InstanceID = SystemInstance->GetId()](FRHICommandListImmediate& RHICmdList)
+	{
+		Grid2DCollectionRWInstanceData* TargetData = RT_Proxy->SystemInstancesToProxyData.Find(InstanceID);
+		if (TargetData != nullptr)
+		{
+			RT_Proxy->DeferredDestroyList.Remove(InstanceID);
+		}
+		else
+		{
+			TargetData = &RT_Proxy->SystemInstancesToProxyData.Add(InstanceID);
+		}
+
+		TargetData->NumCellsX = RT_NumCellsX;
+		TargetData->NumCellsY = RT_NumCellsY;
+
+		TargetData->WorldBBoxMin = RT_WorldBBoxMin;
+		TargetData->WorldBBoxSize = RT_WorldBBoxSize;
+
+		RT_Proxy->OutputShaderStages = RT_OutputShaderStages;
+		RT_Proxy->IterationShaderStages = RT_IterationShaderStages;
+
+		RT_Proxy->SetElementCount(TargetData->NumCellsX * TargetData->NumCellsY);
+	});
+
+	return true;
+}
+
+
+void UNiagaraDataInterfaceGrid2DCollection::DestroyPerInstanceData(void* PerInstanceData, FNiagaraSystemInstance* SystemInstance)
+{
+	Grid2DCollectionRWInstanceData* InstanceData = static_cast<Grid2DCollectionRWInstanceData*>(PerInstanceData);
+
+	InstanceData->~Grid2DCollectionRWInstanceData();
+
+	FNiagaraDataInterfaceProxyGrid2DCollection* ThisProxy = GetProxyAs<FNiagaraDataInterfaceProxyGrid2DCollection>();
+
+	ENQUEUE_RENDER_COMMAND(FNiagaraDIDestroyInstanceData) (
+		[ThisProxy, InstanceID = SystemInstance->GetId(), Batcher = SystemInstance->GetBatcher()](FRHICommandListImmediate& CmdList)
+	{
+		ThisProxy->DestroyPerInstanceData(Batcher, InstanceID);
+	}
+	);
+}
+
+UFUNCTION(BlueprintCallable, Category = Niagara)
+void UNiagaraDataInterfaceGrid2DCollection::FillTexture2D(const UNiagaraComponent *Component, UTextureRenderTarget2D *Dest)
+{
+	FNiagaraDataInterfaceProxyGrid2DCollection* TProxy = GetProxyAs<FNiagaraDataInterfaceProxyGrid2DCollection>();
+
+	if (!Component)
+	{
+		return;
+	}
+
+	FNiagaraSystemInstance *SystemInstance = Component->GetSystemInstance();
+	if (!SystemInstance)
+	{
+		return;
+	}
+	FGuid InstanceID = SystemInstance->GetId();
+
+	ENQUEUE_RENDER_COMMAND(FUpdateDIColorCurve)(
+		[Dest, TProxy, InstanceID](FRHICommandListImmediate& RHICmdList)
+	{
+		Grid2DCollectionRWInstanceData* Grid2DInstanceData = TProxy->SystemInstancesToProxyData.Find(InstanceID);
+
+		if (Dest && Dest->Resource && Grid2DInstanceData && Grid2DInstanceData->CurrentData)
+		{
+			FTextureRenderTarget2DResource* Texture2DResource = (FTextureRenderTarget2DResource*)Dest->Resource;
+			RHICopySharedMips(Texture2DResource->GetTextureRHI(), Grid2DInstanceData->CurrentData->GridBuffer.Buffer);
+		}
+	});
+}
+
+// #todo(dmp): move these to super class
+void FNiagaraDataInterfaceProxyGrid2DCollection::DestroyPerInstanceData(NiagaraEmitterInstanceBatcher* Batcher, const FGuid& SystemInstance)
+{
+	check(IsInRenderingThread());
+
+	DeferredDestroyList.Add(SystemInstance);
+	Batcher->EnqueueDeferredDeletesForDI_RenderThread(this->AsShared());
+}
+
+// #todo(dmp): move these to super class
+void FNiagaraDataInterfaceProxyGrid2DCollection::DeferredDestroy()
+{
+	for (const FGuid& Sys : DeferredDestroyList)
+	{
+		SystemInstancesToProxyData.Remove(Sys);
+	}
+
+	DeferredDestroyList.Empty();
+}
+
+
+void FNiagaraDataInterfaceProxyGrid2DCollection::PreStage(FRHICommandList& RHICmdList, const FNiagaraDataInterfaceSetArgs& Context)
+{
+	// #todo(dmp): Context doesnt need to specify if a stage is output or not since we moved pre/post steage to the DI itself.  Not sure which design is better for the future
+	if (Context.IsOutputStage)
+	{
+		Grid2DCollectionRWInstanceData* ProxyData = SystemInstancesToProxyData.Find(Context.SystemInstance);
+
+		ProxyData->BeginSimulate();
+
+		// If we don't have an iteration stage, then we should manually clear the buffer to make sure there is no residual data.  If we are doing something like rasterizing particles into a grid, we want it to be clear before
+		// we start.  If a user wants to access data from the previous stage, then they can read from the current data.
+
+		// #todo(dmp): we might want to expose an option where we have buffers that are write only and need a clear (ie: no buffering like the neighbor grid).  They would be considered transient perhaps?  It'd be more
+		// memory efficient since it would theoretically not require any double buffering.
+		if (!Context.IsIterationStage)
+		{
+			ClearUAV(RHICmdList, ProxyData->DestinationData->GridBuffer, FLinearColor(0, 0, 0, 0));
+		}
+	}
+}
+
+void FNiagaraDataInterfaceProxyGrid2DCollection::PostStage(FRHICommandList& RHICmdList, const FNiagaraDataInterfaceSetArgs& Context)
+{
+	
+	if (Context.IsOutputStage)
+	{
+		Grid2DCollectionRWInstanceData* ProxyData = SystemInstancesToProxyData.Find(Context.SystemInstance);
+		ProxyData->EndSimulate();
+	}
+}
+
+void FNiagaraDataInterfaceProxyGrid2DCollection::ResetData(FRHICommandList& RHICmdList, const FNiagaraDataInterfaceSetArgs& Context)
+{	
+	Grid2DCollectionRWInstanceData* ProxyData = SystemInstancesToProxyData.Find(Context.SystemInstance);
+
+	for (Grid2DBuffer* Buffer : ProxyData->Buffers)
+	{
+		ClearUAV(RHICmdList, Buffer->GridBuffer, FLinearColor(0, 0, 0, 0));
+	}	
+}
+#undef LOCTEXT_NAMESPACE
\ No newline at end of file
diff --git a/Engine/Plugins/FX/Niagara/Source/Niagara/Private/NiagaraDataInterfaceNeighborGrid3D.cpp b/Engine/Plugins/FX/Niagara/Source/Niagara/Private/NiagaraDataInterfaceNeighborGrid3D.cpp
new file mode 100644
--- /dev/null
+++ b/Engine/Plugins/FX/Niagara/Source/Niagara/Private/NiagaraDataInterfaceNeighborGrid3D.cpp
@@ -0,0 +1,511 @@
+// Copyright 1998-2019 Epic Games, Inc. All Rights Reserved.
+#include "NiagaraDataInterfaceNeighborGrid3D.h"
+#include "NiagaraShader.h"
+#include "ShaderParameterUtils.h"
+#include "NiagaraEmitterInstanceBatcher.h"
+#include "NiagaraSystemInstance.h"
+
+
+#define LOCTEXT_NAMESPACE "NiagaraDataInterfaceNeighborGrid3D"
+
+static const FString MaxNeighborsPerVoxelName(TEXT("MaxNeighborsPerVoxel_"));
+static const FString ParticleNeighborsName(TEXT("ParticleNeighbors_"));
+static const FString ParticleNeighborCountName(TEXT("ParticleNeighborCount_"));
+static const FString OutputParticleNeighborsName(TEXT("OutputParticleNeighbors_"));
+static const FString OutputParticleNeighborCountName(TEXT("OutputParticleNeighborCount_"));
+
+
+// Global VM function names, also used by the shaders code generation methods.
+
+static const FName MaxNeighborsPerVoxelFunctionName("MaxNeighborsPerVoxel");
+static const FName NeighborGridIndexToLinearFunctionName("NeighborGridIndexToLinear");
+static const FName GetParticleNeighborFunctionName("GetParticleNeighbor");
+static const FName SetParticleNeighborFunctionName("SetParticleNeighbor");
+static const FName GetParticleNeighborCountFunctionName("GetParticleNeighborCount");
+static const FName SetParticleNeighborCountFunctionName("SetParticleNeighborCount");
+
+
+/*--------------------------------------------------------------------------------------------------------------------------*/
+struct FNiagaraDataInterfaceParametersCS_NeighborGrid3D : public FNiagaraDataInterfaceParametersCS
+{
+	virtual void Bind(const FNiagaraDataInterfaceParamRef& ParamRef, const class FShaderParameterMap& ParameterMap) override
+	{
+		NumVoxelsParam.Bind(ParameterMap, *(NumVoxelsName + ParamRef.ParameterInfo.DataInterfaceHLSLSymbol));		
+		VoxelSizeParam.Bind(ParameterMap, *(VoxelSizeName + ParamRef.ParameterInfo.DataInterfaceHLSLSymbol));
+
+		MaxNeighborsPerVoxelParam.Bind(ParameterMap, *(MaxNeighborsPerVoxelName + ParamRef.ParameterInfo.DataInterfaceHLSLSymbol));
+		WorldBBoxMinParam.Bind(ParameterMap, *(WorldBBoxMinName + ParamRef.ParameterInfo.DataInterfaceHLSLSymbol));
+		WorldBBoxSizeParam.Bind(ParameterMap, *(WorldBBoxSizeName + ParamRef.ParameterInfo.DataInterfaceHLSLSymbol));
+		
+		ParticleNeighborsGridParam.Bind(ParameterMap,  *(ParticleNeighborsName + ParamRef.ParameterInfo.DataInterfaceHLSLSymbol));
+		ParticleNeighborCountGridParam.Bind(ParameterMap,  *(ParticleNeighborCountName + ParamRef.ParameterInfo.DataInterfaceHLSLSymbol));
+
+		OutputParticleNeighborsGridParam.Bind(ParameterMap, *(OutputParticleNeighborsName + ParamRef.ParameterInfo.DataInterfaceHLSLSymbol));
+		OutputParticleNeighborCountGridParam.Bind(ParameterMap, *(OutputParticleNeighborCountName + ParamRef.ParameterInfo.DataInterfaceHLSLSymbol));		
+	}
+
+	virtual void Serialize(FArchive& Ar)override
+	{
+		Ar << NumVoxelsParam;
+		Ar << VoxelSizeParam;
+		Ar << MaxNeighborsPerVoxelParam;
+		Ar << WorldBBoxMinParam;
+		Ar << WorldBBoxSizeParam;
+		
+		Ar << ParticleNeighborsGridParam;
+		Ar << ParticleNeighborCountGridParam;
+		Ar << OutputParticleNeighborsGridParam;
+		Ar << OutputParticleNeighborCountGridParam;		
+	}
+
+	virtual void Set(FRHICommandList& RHICmdList, const FNiagaraDataInterfaceSetArgs& Context) const override
+	{
+		check(IsInRenderingThread());
+
+		// Get shader and DI
+		FRHIComputeShader* ComputeShaderRHI = Context.Shader->GetComputeShader();
+		FNiagaraDataInterfaceProxyNeighborGrid3D* VFDI = static_cast<FNiagaraDataInterfaceProxyNeighborGrid3D*>(Context.DataInterface);
+
+		NeighborGrid3DRWInstanceData* ProxyData = VFDI->SystemInstancesToProxyData.Find(Context.SystemInstance);
+		ensure(ProxyData);
+
+		SetShaderValue(RHICmdList, ComputeShaderRHI, NumVoxelsParam, ProxyData->NumVoxels);
+
+		// #todo(dmp): remove this computation here
+		float VoxelSizeTmp[3];
+		VoxelSizeTmp[0] = ProxyData->WorldBBoxSize.X / ProxyData->NumVoxels.X;
+		VoxelSizeTmp[1] = ProxyData->WorldBBoxSize.Y / ProxyData->NumVoxels.Y;
+		VoxelSizeTmp[2] = ProxyData->WorldBBoxSize.Z / ProxyData->NumVoxels.Z;
+		SetShaderValue(RHICmdList, ComputeShaderRHI, VoxelSizeParam, VoxelSizeTmp);
+
+		SetShaderValue(RHICmdList, ComputeShaderRHI, MaxNeighborsPerVoxelParam, ProxyData->MaxNeighborsPerVoxel);
+		SetShaderValue(RHICmdList, ComputeShaderRHI, WorldBBoxMinParam, ProxyData->WorldBBoxMin);
+		SetShaderValue(RHICmdList, ComputeShaderRHI, WorldBBoxSizeParam, ProxyData->WorldBBoxSize);
+		
+		if (!Context.IsOutputStage)
+		{
+			if (ParticleNeighborsGridParam.IsBound())
+			{
+				RHICmdList.SetShaderResourceViewParameter(Context.Shader->GetComputeShader(), ParticleNeighborsGridParam.GetBaseIndex(), ProxyData->NeighborhoodBuffer.SRV);
+			}
+
+			if (ParticleNeighborCountGridParam.IsBound())
+			{
+				RHICmdList.SetShaderResourceViewParameter(Context.Shader->GetComputeShader(), ParticleNeighborCountGridParam.GetBaseIndex(), ProxyData->NeighborhoodCountBuffer.SRV);
+			}
+		}
+		else
+		{
+			if (OutputParticleNeighborsGridParam.IsBound())
+			{
+				RHICmdList.SetUAVParameter(Context.Shader->GetComputeShader(), OutputParticleNeighborsGridParam.GetUAVIndex(), ProxyData->NeighborhoodBuffer.UAV);
+			}
+
+			if (OutputParticleNeighborCountGridParam.IsBound())
+			{
+				RHICmdList.SetUAVParameter(Context.Shader->GetComputeShader(), OutputParticleNeighborCountGridParam.GetUAVIndex(), ProxyData->NeighborhoodCountBuffer.UAV);
+			}
+		}
+		// Note: There is a flush in PreEditChange to make sure everything is synced up at this point 
+	}
+
+	virtual void Unset(FRHICommandList& RHICmdList, const FNiagaraDataInterfaceSetArgs& Context) const 
+	{
+		if (OutputParticleNeighborsGridParam.IsBound())
+		{
+			OutputParticleNeighborsGridParam.UnsetUAV(RHICmdList, Context.Shader->GetComputeShader());
+		}
+
+		if (OutputParticleNeighborCountGridParam.IsBound())
+		{
+			OutputParticleNeighborCountGridParam.UnsetUAV(RHICmdList, Context.Shader->GetComputeShader());
+		}
+	}
+
+private:
+	FShaderParameter NumVoxelsParam;
+	FShaderParameter VoxelSizeParam;
+	FShaderParameter MaxNeighborsPerVoxelParam;
+	FShaderParameter WorldBBoxMinParam;
+	FShaderParameter WorldBBoxSizeParam;
+	FShaderResourceParameter ParticleNeighborsGridParam;
+	FShaderResourceParameter ParticleNeighborCountGridParam;
+	FRWShaderParameter OutputParticleNeighborCountGridParam;
+	FRWShaderParameter OutputParticleNeighborsGridParam;
+};
+
+
+UNiagaraDataInterfaceNeighborGrid3D::UNiagaraDataInterfaceNeighborGrid3D(FObjectInitializer const& ObjectInitializer)
+	: Super(ObjectInitializer)
+	, MaxNeighborsPerVoxel(8)
+{
+	Proxy = MakeShared<FNiagaraDataInterfaceProxyNeighborGrid3D, ESPMode::ThreadSafe>();
+	PushToRenderThread();
+}
+
+
+void UNiagaraDataInterfaceNeighborGrid3D::GetFunctions(TArray<FNiagaraFunctionSignature>& OutFunctions)
+{
+	Super::GetFunctions(OutFunctions);
+
+	{
+		FNiagaraFunctionSignature Sig;
+		Sig.Name = MaxNeighborsPerVoxelFunctionName;
+		Sig.Inputs.Add(FNiagaraVariable(FNiagaraTypeDefinition(GetClass()), TEXT("Grid")));
+		Sig.Outputs.Add(FNiagaraVariable(FNiagaraTypeDefinition::GetIntDef(), TEXT("MaxNeighborsPerVoxel")));
+
+		Sig.bMemberFunction = true;
+		Sig.bRequiresContext = false;
+		OutFunctions.Add(Sig);
+	}
+
+	{
+		FNiagaraFunctionSignature Sig;
+		Sig.Name = NeighborGridIndexToLinearFunctionName;
+		Sig.Inputs.Add(FNiagaraVariable(FNiagaraTypeDefinition(GetClass()), TEXT("Grid")));
+		Sig.Inputs.Add(FNiagaraVariable(FNiagaraTypeDefinition::GetIntDef(), TEXT("IndexX")));
+		Sig.Inputs.Add(FNiagaraVariable(FNiagaraTypeDefinition::GetIntDef(), TEXT("IndexY")));
+		Sig.Inputs.Add(FNiagaraVariable(FNiagaraTypeDefinition::GetIntDef(), TEXT("IndexZ")));
+		Sig.Inputs.Add(FNiagaraVariable(FNiagaraTypeDefinition::GetIntDef(), TEXT("Neighbor")));
+		Sig.Outputs.Add(FNiagaraVariable(FNiagaraTypeDefinition::GetIntDef(), TEXT("Linear Index")));
+
+		Sig.bMemberFunction = true;
+		Sig.bRequiresContext = false;
+		OutFunctions.Add(Sig);
+	}
+
+	{
+		FNiagaraFunctionSignature Sig;
+		Sig.Name = GetParticleNeighborFunctionName;
+		Sig.Inputs.Add(FNiagaraVariable(FNiagaraTypeDefinition(GetClass()), TEXT("Grid")));
+		Sig.Inputs.Add(FNiagaraVariable(FNiagaraTypeDefinition::GetIntDef(), TEXT("Linear")));
+		Sig.Outputs.Add(FNiagaraVariable(FNiagaraTypeDefinition::GetIntDef(), TEXT("NeighborIndex")));
+
+		Sig.bMemberFunction = true;
+		Sig.bRequiresContext = false;
+		OutFunctions.Add(Sig);
+	}
+
+	{
+		FNiagaraFunctionSignature Sig;
+		Sig.Name = SetParticleNeighborFunctionName;
+		Sig.Inputs.Add(FNiagaraVariable(FNiagaraTypeDefinition(GetClass()), TEXT("Grid")));
+		Sig.Inputs.Add(FNiagaraVariable(FNiagaraTypeDefinition::GetIntDef(), TEXT("Linear")));
+		Sig.Inputs.Add(FNiagaraVariable(FNiagaraTypeDefinition::GetIntDef(), TEXT("NeighborIndex")));
+		Sig.Outputs.Add(FNiagaraVariable(FNiagaraTypeDefinition::GetIntDef(), TEXT("IGNORE")));
+
+		Sig.bMemberFunction = true;
+		Sig.bRequiresContext = false;
+		OutFunctions.Add(Sig);
+	}
+
+	{
+		FNiagaraFunctionSignature Sig;
+		Sig.Name = GetParticleNeighborCountFunctionName;
+		Sig.Inputs.Add(FNiagaraVariable(FNiagaraTypeDefinition(GetClass()), TEXT("Grid")));
+		Sig.Inputs.Add(FNiagaraVariable(FNiagaraTypeDefinition::GetIntDef(), TEXT("Linear")));
+		Sig.Outputs.Add(FNiagaraVariable(FNiagaraTypeDefinition::GetIntDef(), TEXT("NeighborCount")));
+
+		Sig.bMemberFunction = true;
+		Sig.bRequiresContext = false;
+		OutFunctions.Add(Sig);
+	}
+
+	{
+		FNiagaraFunctionSignature Sig;
+		Sig.Name = SetParticleNeighborCountFunctionName;
+		Sig.Inputs.Add(FNiagaraVariable(FNiagaraTypeDefinition(GetClass()), TEXT("Grid")));
+		Sig.Inputs.Add(FNiagaraVariable(FNiagaraTypeDefinition::GetIntDef(), TEXT("Linear")));
+		Sig.Inputs.Add(FNiagaraVariable(FNiagaraTypeDefinition::GetIntDef(), TEXT("Increment")));
+		Sig.Outputs.Add(FNiagaraVariable(FNiagaraTypeDefinition::GetIntDef(), TEXT("PrevNeighborCount")));
+
+		Sig.bMemberFunction = true;
+		Sig.bRequiresContext = false;
+		OutFunctions.Add(Sig);
+	}
+}
+
+
+void UNiagaraDataInterfaceNeighborGrid3D::GetVMExternalFunction(const FVMExternalFunctionBindingInfo& BindingInfo, void* InstanceData, FVMExternalFunction &OutFunc)
+{
+	Super::GetVMExternalFunction(BindingInfo, InstanceData, OutFunc);
+
+	if (BindingInfo.Name == MaxNeighborsPerVoxelFunctionName) { OutFunc = FVMExternalFunction::CreateUObject(this, &UNiagaraDataInterfaceRWBase::EmptyVMFunction); }	
+	else if (BindingInfo.Name == NeighborGridIndexToLinearFunctionName) { OutFunc = FVMExternalFunction::CreateUObject(this, &UNiagaraDataInterfaceRWBase::EmptyVMFunction); }
+	else if (BindingInfo.Name == GetParticleNeighborFunctionName) { OutFunc = FVMExternalFunction::CreateUObject(this, &UNiagaraDataInterfaceRWBase::EmptyVMFunction); }
+	else if (BindingInfo.Name == SetParticleNeighborFunctionName) { OutFunc = FVMExternalFunction::CreateUObject(this, &UNiagaraDataInterfaceRWBase::EmptyVMFunction); }
+	else if (BindingInfo.Name == GetParticleNeighborCountFunctionName) { OutFunc = FVMExternalFunction::CreateUObject(this, &UNiagaraDataInterfaceRWBase::EmptyVMFunction); }
+	else if (BindingInfo.Name == SetParticleNeighborCountFunctionName) { OutFunc = FVMExternalFunction::CreateUObject(this, &UNiagaraDataInterfaceRWBase::EmptyVMFunction); }
+}
+
+bool UNiagaraDataInterfaceNeighborGrid3D::Equals(const UNiagaraDataInterface* Other) const
+{
+	if (!Super::Equals(Other))
+	{
+		return false;
+	}
+	const UNiagaraDataInterfaceNeighborGrid3D* OtherTyped = CastChecked<const UNiagaraDataInterfaceNeighborGrid3D>(Other);
+
+	return OtherTyped->MaxNeighborsPerVoxel == MaxNeighborsPerVoxel;
+}
+
+void UNiagaraDataInterfaceNeighborGrid3D::GetParameterDefinitionHLSL(FNiagaraDataInterfaceGPUParamInfo& ParamInfo, FString& OutHLSL)
+{
+	Super::GetParameterDefinitionHLSL(ParamInfo, OutHLSL);
+
+	static const TCHAR *FormatDeclarations = TEXT(R"(
+		int {MaxNeighborsPerVoxelName};
+		Buffer<int> {ParticleNeighborsName};
+		Buffer<int> {ParticleNeighborCountName};
+		RWBuffer<int> RW{OutputParticleNeighborsName};
+		RWBuffer<int> RW{OutputParticleNeighborCountName};
+	)");
+	TMap<FString, FStringFormatArg> ArgsDeclarations = {
+		{ TEXT("MaxNeighborsPerVoxelName"),  MaxNeighborsPerVoxelName + ParamInfo.DataInterfaceHLSLSymbol },
+		{ TEXT("ParticleNeighborsName"),    ParticleNeighborsName + ParamInfo.DataInterfaceHLSLSymbol },
+		{ TEXT("ParticleNeighborCountName"),    ParticleNeighborCountName + ParamInfo.DataInterfaceHLSLSymbol },
+		{ TEXT("OutputParticleNeighborsName"),    OutputParticleNeighborsName + ParamInfo.DataInterfaceHLSLSymbol },
+		{ TEXT("OutputParticleNeighborCountName"),    OutputParticleNeighborCountName + ParamInfo.DataInterfaceHLSLSymbol },
+	};
+	OutHLSL += FString::Format(FormatDeclarations, ArgsDeclarations);
+}
+
+bool UNiagaraDataInterfaceNeighborGrid3D::GetFunctionHLSL(const FName& DefinitionFunctionName, FString InstanceFunctionName, FNiagaraDataInterfaceGPUParamInfo& ParamInfo, FString& OutHLSL)
+{
+	bool ParentRet = Super::GetFunctionHLSL(DefinitionFunctionName, InstanceFunctionName, ParamInfo, OutHLSL);
+	if (ParentRet)
+	{
+		return true;
+	} else if (DefinitionFunctionName == MaxNeighborsPerVoxelFunctionName)
+	{
+		static const TCHAR *FormatSample = TEXT(R"(
+			void {FunctionName}(out int Out_MaxNeighborsPerVoxel)
+			{
+				Out_MaxNeighborsPerVoxel = {MaxNeighborsPerVoxelName};
+			}
+		)");
+		TMap<FString, FStringFormatArg> ArgsSample = {
+			{TEXT("FunctionName"), InstanceFunctionName},
+			{TEXT("MaxNeighborsPerVoxelName"), MaxNeighborsPerVoxelName + ParamInfo.DataInterfaceHLSLSymbol},
+
+		};
+		OutHLSL += FString::Format(FormatSample, ArgsSample);
+		return true;
+	}
+	else if (DefinitionFunctionName == NeighborGridIndexToLinearFunctionName)
+	{
+		static const TCHAR *FormatBounds = TEXT(R"(
+			void {FunctionName}(int In_IndexX, int In_IndexY, int In_IndexZ, int In_Neighbor, out int Out_Linear)
+			{
+				Out_Linear = In_Neighbor + In_IndexX * {MaxNeighborsPerVoxelName} + In_IndexY * {MaxNeighborsPerVoxelName}*{NumVoxelsName}.x + In_IndexZ * {MaxNeighborsPerVoxelName}*{NumVoxelsName}.x*{NumVoxelsName}.y;
+			}
+		)");
+		TMap<FString, FStringFormatArg> ArgsBounds = {
+			{TEXT("FunctionName"), InstanceFunctionName},
+			{TEXT("MaxNeighborsPerVoxelName"), MaxNeighborsPerVoxelName + ParamInfo.DataInterfaceHLSLSymbol},
+			{TEXT("NumVoxelsName"), NumVoxelsName + ParamInfo.DataInterfaceHLSLSymbol},
+		};
+		OutHLSL += FString::Format(FormatBounds, ArgsBounds);
+		return true;
+	}
+	else if (DefinitionFunctionName == GetParticleNeighborFunctionName)
+	{
+		static const TCHAR *FormatBounds = TEXT(R"(
+			void {FunctionName}(int In_Index, out int Out_ParticleNeighborIndex)
+			{
+				Out_ParticleNeighborIndex = {ParticleNeighbors}[In_Index];				
+			}
+		)");
+		TMap<FString, FStringFormatArg> ArgsBounds = {
+			{TEXT("FunctionName"), InstanceFunctionName},
+			{TEXT("ParticleNeighbors"), ParticleNeighborsName + ParamInfo.DataInterfaceHLSLSymbol},
+		};
+		OutHLSL += FString::Format(FormatBounds, ArgsBounds);
+		return true;
+	}
+	else if (DefinitionFunctionName == SetParticleNeighborFunctionName)
+	{
+		static const TCHAR *FormatBounds = TEXT(R"(
+			void {FunctionName}(int In_Index, int In_ParticleNeighborIndex, out int Out_Ignore)
+			{
+				RW{OutputParticleNeighbors}[In_Index] = In_ParticleNeighborIndex;				
+				Out_Ignore = 0;
+			}
+		)");
+		TMap<FString, FStringFormatArg> ArgsBounds = {
+			{TEXT("FunctionName"), InstanceFunctionName},
+			{TEXT("OutputParticleNeighbors"), OutputParticleNeighborsName + ParamInfo.DataInterfaceHLSLSymbol},
+		};
+		OutHLSL += FString::Format(FormatBounds, ArgsBounds);
+		return true;
+	}
+	else if (DefinitionFunctionName == GetParticleNeighborCountFunctionName)
+	{
+		static const TCHAR *FormatBounds = TEXT(R"(
+			void {FunctionName}(int In_Index, out int Out_ParticleNeighborIndex)
+			{
+				Out_ParticleNeighborIndex = {ParticleNeighborCount}[In_Index];				
+			}
+		)");
+		TMap<FString, FStringFormatArg> ArgsBounds = {
+			{TEXT("FunctionName"), InstanceFunctionName},
+			{TEXT("ParticleNeighborCount"), ParticleNeighborCountName + ParamInfo.DataInterfaceHLSLSymbol},
+		};
+		OutHLSL += FString::Format(FormatBounds, ArgsBounds);
+		return true;
+	}
+	else if (DefinitionFunctionName == SetParticleNeighborCountFunctionName)
+	{
+		static const TCHAR *FormatBounds = TEXT(R"(
+			void {FunctionName}(int In_Index, int In_Increment, out int PreviousNeighborCount)
+			{				
+				InterlockedAdd(RW{OutputParticleNeighborCount}[In_Index], In_Increment, PreviousNeighborCount);				
+			}
+		)");
+		TMap<FString, FStringFormatArg> ArgsBounds = {
+			{TEXT("FunctionName"), InstanceFunctionName},
+			{TEXT("OutputParticleNeighborCount"), OutputParticleNeighborCountName + ParamInfo.DataInterfaceHLSLSymbol},
+		};
+		OutHLSL += FString::Format(FormatBounds, ArgsBounds);
+		return true;
+	}
+
+	return false;
+}
+
+FNiagaraDataInterfaceParametersCS* UNiagaraDataInterfaceNeighborGrid3D::ConstructComputeParameters()const
+{
+	return new FNiagaraDataInterfaceParametersCS_NeighborGrid3D();
+}
+
+
+
+bool UNiagaraDataInterfaceNeighborGrid3D::InitPerInstanceData(void* PerInstanceData, FNiagaraSystemInstance* SystemInstance)
+{
+
+	NeighborGrid3DRWInstanceData* InstanceData = new (PerInstanceData) NeighborGrid3DRWInstanceData();
+
+	FNiagaraDataInterfaceProxyNeighborGrid3D* RT_Proxy = GetProxyAs<FNiagaraDataInterfaceProxyNeighborGrid3D>();
+
+	FIntVector RT_NumVoxels = NumVoxels;
+	uint32 RT_MaxNeighborsPerVoxel = MaxNeighborsPerVoxel;
+	FVector RT_WorldBBoxMin = WorldBBoxMin;
+	FVector RT_WorldBBoxSize = WorldBBoxSize;
+	TSet<int> RT_OutputShaderStages = OutputShaderStages;
+	TSet<int> RT_IterationShaderStages = IterationShaderStages;
+
+	// If we are setting the grid from the voxel size, then recompute NumVoxels and change bbox
+	if (SetGridFromVoxelSize)
+	{
+		RT_NumVoxels = FIntVector(WorldBBoxSize / VoxelSize);
+
+		// Pad grid by 1 voxel if our computed bounding box is too small
+		if (!FMath::IsNearlyEqual(VoxelSize * RT_NumVoxels.X, WorldBBoxSize.X))
+		{
+			RT_NumVoxels += FIntVector(1, 1, 1);
+			RT_WorldBBoxSize = FVector(RT_NumVoxels) * VoxelSize;
+		}
+	}
+
+	InstanceData->NumVoxels = RT_NumVoxels;
+	InstanceData->MaxNeighborsPerVoxel = RT_MaxNeighborsPerVoxel;
+	InstanceData->WorldBBoxMin = RT_WorldBBoxMin;
+	InstanceData->WorldBBoxSize = RT_WorldBBoxSize;
+
+	// @todo-threadsafety. This would be a race but I'm taking a ref here. Not ideal in the long term.
+	// Push Updates to Proxy.
+	ENQUEUE_RENDER_COMMAND(FUpdateData)(
+		[RT_Proxy, RT_NumVoxels, RT_MaxNeighborsPerVoxel, RT_WorldBBoxMin, RT_WorldBBoxSize, RT_OutputShaderStages, RT_IterationShaderStages, InstanceID = SystemInstance->GetId()](FRHICommandListImmediate& RHICmdList)
+	{
+		NeighborGrid3DRWInstanceData* TargetData = RT_Proxy->SystemInstancesToProxyData.Find(InstanceID);
+		if (TargetData != nullptr)
+		{
+			RT_Proxy->DeferredDestroyList.Remove(InstanceID);
+		}
+		else
+		{
+			TargetData = &RT_Proxy->SystemInstancesToProxyData.Add(InstanceID);
+		}
+
+		TargetData->NumVoxels = RT_NumVoxels;
+		TargetData->MaxNeighborsPerVoxel = RT_MaxNeighborsPerVoxel;
+		TargetData->WorldBBoxMin = RT_WorldBBoxMin;
+		TargetData->WorldBBoxSize = RT_WorldBBoxSize;
+
+		RT_Proxy->OutputShaderStages = RT_OutputShaderStages;
+		RT_Proxy->IterationShaderStages = RT_IterationShaderStages;
+
+		// #todo(dmp): element count is still defined on the proxy and not the instance data
+		RT_Proxy->SetElementCount(RT_NumVoxels.X *  RT_NumVoxels.Y * RT_NumVoxels.Z);
+
+		TargetData->NeighborhoodCountBuffer.Initialize(sizeof(int32), TargetData->NumVoxels.X*TargetData->NumVoxels.Y*TargetData->NumVoxels.Z, EPixelFormat::PF_R32_SINT, BUF_Static);
+		TargetData->NeighborhoodBuffer.Initialize(sizeof(int32), TargetData->NumVoxels.X*TargetData->NumVoxels.Y*TargetData->NumVoxels.Z*TargetData->MaxNeighborsPerVoxel, EPixelFormat::PF_R32_SINT, BUF_Static);
+
+	});
+
+	return true;
+}
+
+void UNiagaraDataInterfaceNeighborGrid3D::DestroyPerInstanceData(void* PerInstanceData, FNiagaraSystemInstance* SystemInstance)
+{		
+	NeighborGrid3DRWInstanceData* InstanceData = static_cast<NeighborGrid3DRWInstanceData*>(PerInstanceData);
+	InstanceData->~NeighborGrid3DRWInstanceData();
+
+	FNiagaraDataInterfaceProxyNeighborGrid3D* ThisProxy = GetProxyAs<FNiagaraDataInterfaceProxyNeighborGrid3D>();
+
+	ENQUEUE_RENDER_COMMAND(FNiagaraDIDestroyInstanceData) (
+		[ThisProxy, InstanceID = SystemInstance->GetId(), Batcher = SystemInstance->GetBatcher()](FRHICommandListImmediate& CmdList)
+	{
+		ThisProxy->DestroyPerInstanceData(Batcher, InstanceID);
+	}
+	);
+}
+
+void FNiagaraDataInterfaceProxyNeighborGrid3D::PreStage(FRHICommandList& RHICmdList, const FNiagaraDataInterfaceSetArgs& Context)
+{
+	if (Context.IsOutputStage)
+	{
+		NeighborGrid3DRWInstanceData* ProxyData = SystemInstancesToProxyData.Find(Context.SystemInstance);
+
+		ClearUAV(RHICmdList, ProxyData->NeighborhoodBuffer, -1);
+		ClearUAV(RHICmdList, ProxyData->NeighborhoodCountBuffer, 0);
+	}
+}
+
+// #todo(dmp): move these to super class
+void FNiagaraDataInterfaceProxyNeighborGrid3D::DestroyPerInstanceData(NiagaraEmitterInstanceBatcher* Batcher, const FGuid& SystemInstance)
+{
+	check(IsInRenderingThread());
+
+	DeferredDestroyList.Add(SystemInstance);
+	Batcher->EnqueueDeferredDeletesForDI_RenderThread(this->AsShared());
+}
+
+// #todo(dmp): move these to super class
+void FNiagaraDataInterfaceProxyNeighborGrid3D::DeferredDestroy()
+{
+	for (const FGuid& Sys : DeferredDestroyList)
+	{
+		SystemInstancesToProxyData.Remove(Sys);
+	}
+
+	DeferredDestroyList.Empty();
+}
+
+bool UNiagaraDataInterfaceNeighborGrid3D::CopyToInternal(UNiagaraDataInterface* Destination) const
+{
+	if (!Super::CopyToInternal(Destination))
+	{
+		return false;
+	}
+
+	UNiagaraDataInterfaceNeighborGrid3D* OtherTyped = CastChecked<UNiagaraDataInterfaceNeighborGrid3D>(Destination);
+
+
+	OtherTyped->MaxNeighborsPerVoxel = MaxNeighborsPerVoxel;
+
+
+	return true;
+}
+
+#undef LOCTEXT_NAMESPACE
\ No newline at end of file
diff --git a/Engine/Plugins/FX/Niagara/Source/Niagara/Private/NiagaraDataInterfaceRW.cpp b/Engine/Plugins/FX/Niagara/Source/Niagara/Private/NiagaraDataInterfaceRW.cpp
new file mode 100644
--- /dev/null
+++ b/Engine/Plugins/FX/Niagara/Source/Niagara/Private/NiagaraDataInterfaceRW.cpp
@@ -0,0 +1,711 @@
+// Copyright 1998-2019 Epic Games, Inc. All Rights Reserved.
+#include "NiagaraDataInterfaceRW.h"
+#include "NiagaraShader.h"
+#include "ShaderParameterUtils.h"
+#include "ClearQuad.h"
+
+
+#define LOCTEXT_NAMESPACE "NiagaraDataInterfaceRW"
+
+
+UNiagaraDataInterfaceRWBase::UNiagaraDataInterfaceRWBase(FObjectInitializer const& ObjectInitializer)
+	: Super(ObjectInitializer)
+{
+
+}
+
+bool UNiagaraDataInterfaceRWBase::Equals(const UNiagaraDataInterface* Other) const
+{
+	const UNiagaraDataInterfaceRWBase* OtherTyped = CastChecked<const UNiagaraDataInterfaceRWBase>(Other);
+
+	if (OtherTyped)
+	{
+		return OutputShaderStages.Difference(OtherTyped->OutputShaderStages).Num() == 0 && IterationShaderStages.Difference(OtherTyped->IterationShaderStages).Num() == 0;
+	}
+
+	return false;
+}
+
+bool UNiagaraDataInterfaceRWBase::CopyToInternal(UNiagaraDataInterface* Destination) const
+{
+	if (!Super::CopyToInternal(Destination))
+	{
+		return false;
+	}
+
+	UNiagaraDataInterfaceRWBase* OtherTyped = CastChecked<UNiagaraDataInterfaceRWBase>(Destination);
+
+	OtherTyped->OutputShaderStages = OutputShaderStages;
+	OtherTyped->IterationShaderStages = IterationShaderStages;
+
+	return true;
+}
+
+/*--------------------------------------------------------------------------------------------------------------------------*/
+
+UNiagaraDataInterfaceGrid3D::UNiagaraDataInterfaceGrid3D(FObjectInitializer const& ObjectInitializer)
+	: Super(ObjectInitializer)
+	, NumVoxels(3, 3, 3)
+	, VoxelSize(1.)
+	, SetGridFromVoxelSize(false)	
+	, WorldBBoxMin(0., 0., 0.)
+	, WorldBBoxSize(100., 100., 100.)
+{
+	Proxy = MakeShared<FNiagaraDataInterfaceProxyRW, ESPMode::ThreadSafe>();
+	RWProxy = (FNiagaraDataInterfaceProxyRW*) Proxy.Get();
+	PushToRenderThread();
+}
+
+
+void UNiagaraDataInterfaceGrid3D::GetFunctions(TArray<FNiagaraFunctionSignature>& OutFunctions)
+{
+
+	{
+		FNiagaraFunctionSignature Sig;
+		Sig.Name = NumVoxelsFunctionName;
+		Sig.Inputs.Add(FNiagaraVariable(FNiagaraTypeDefinition(GetClass()), TEXT("Grid")));
+		Sig.Outputs.Add(FNiagaraVariable(FNiagaraTypeDefinition::GetIntDef(), TEXT("NumVoxelsX")));
+		Sig.Outputs.Add(FNiagaraVariable(FNiagaraTypeDefinition::GetIntDef(), TEXT("NumVoxelsY")));
+		Sig.Outputs.Add(FNiagaraVariable(FNiagaraTypeDefinition::GetIntDef(), TEXT("NumVoxelsZ")));
+
+		Sig.bMemberFunction = true;
+		Sig.bRequiresContext = false;
+		OutFunctions.Add(Sig);
+	}
+
+	{
+		FNiagaraFunctionSignature Sig;
+		Sig.Name = WorldToUnitFunctionName;
+		Sig.Inputs.Add(FNiagaraVariable(FNiagaraTypeDefinition(GetClass()), TEXT("Grid")));
+		Sig.Inputs.Add(FNiagaraVariable(FNiagaraTypeDefinition::GetVec3Def(), TEXT("World")));
+		Sig.Outputs.Add(FNiagaraVariable(FNiagaraTypeDefinition::GetVec3Def(), TEXT("Unit")));
+
+		Sig.bMemberFunction = true;
+		Sig.bRequiresContext = false;
+		OutFunctions.Add(Sig);
+	}
+
+	{
+		FNiagaraFunctionSignature Sig;
+		Sig.Name = UnitToIndexFunctionName;
+		Sig.Inputs.Add(FNiagaraVariable(FNiagaraTypeDefinition(GetClass()), TEXT("Grid")));
+		Sig.Inputs.Add(FNiagaraVariable(FNiagaraTypeDefinition::GetVec3Def(), TEXT("Unit")));
+		Sig.Outputs.Add(FNiagaraVariable(FNiagaraTypeDefinition::GetIntDef(), TEXT("IndexX")));
+		Sig.Outputs.Add(FNiagaraVariable(FNiagaraTypeDefinition::GetIntDef(), TEXT("IndexY")));
+		Sig.Outputs.Add(FNiagaraVariable(FNiagaraTypeDefinition::GetIntDef(), TEXT("IndexZ")));
+
+		Sig.bMemberFunction = true;
+		Sig.bRequiresContext = false;
+		OutFunctions.Add(Sig);
+	}
+
+	{
+		FNiagaraFunctionSignature Sig;
+		Sig.Name = IndexToLinearFunctionName;
+		Sig.Inputs.Add(FNiagaraVariable(FNiagaraTypeDefinition(GetClass()), TEXT("Grid")));
+		Sig.Inputs.Add(FNiagaraVariable(FNiagaraTypeDefinition::GetIntDef(), TEXT("IndexX")));
+		Sig.Inputs.Add(FNiagaraVariable(FNiagaraTypeDefinition::GetIntDef(), TEXT("IndexY")));
+		Sig.Inputs.Add(FNiagaraVariable(FNiagaraTypeDefinition::GetIntDef(), TEXT("IndexZ")));
+		Sig.Outputs.Add(FNiagaraVariable(FNiagaraTypeDefinition::GetIntDef(), TEXT("Linear")));
+
+		Sig.bMemberFunction = true;
+		Sig.bRequiresContext = false;
+		OutFunctions.Add(Sig);
+	}
+
+	{
+		FNiagaraFunctionSignature Sig;
+		Sig.Name = NumVoxelsFunctionName;
+		Sig.Inputs.Add(FNiagaraVariable(FNiagaraTypeDefinition(GetClass()), TEXT("Grid")));		
+		Sig.Outputs.Add(FNiagaraVariable(FNiagaraTypeDefinition::GetIntDef(), TEXT("NumVoxelsX")));
+		Sig.Outputs.Add(FNiagaraVariable(FNiagaraTypeDefinition::GetIntDef(), TEXT("NumVoxelsY")));
+		Sig.Outputs.Add(FNiagaraVariable(FNiagaraTypeDefinition::GetIntDef(), TEXT("NumVoxelsZ")));
+
+		Sig.bMemberFunction = true;
+		Sig.bRequiresContext = false;
+		OutFunctions.Add(Sig);
+	}
+
+	{
+		FNiagaraFunctionSignature Sig;
+		Sig.Name = VoxelSizeFunctionName;
+		Sig.Inputs.Add(FNiagaraVariable(FNiagaraTypeDefinition(GetClass()), TEXT("Grid")));
+		Sig.Outputs.Add(FNiagaraVariable(FNiagaraTypeDefinition::GetVec3Def(), TEXT("VoxelSize")));		
+
+		Sig.bMemberFunction = true;
+		Sig.bRequiresContext = false;
+		OutFunctions.Add(Sig);
+	}
+}
+
+
+void UNiagaraDataInterfaceGrid3D::GetVMExternalFunction(const FVMExternalFunctionBindingInfo& BindingInfo, void* InstanceData, FVMExternalFunction &OutFunc)
+{
+	if (BindingInfo.Name == NumVoxelsFunctionName) { OutFunc = FVMExternalFunction::CreateUObject(this, &UNiagaraDataInterfaceRWBase::EmptyVMFunction); }	
+	else if (BindingInfo.Name == WorldToUnitFunctionName) { OutFunc = FVMExternalFunction::CreateUObject(this, &UNiagaraDataInterfaceRWBase::EmptyVMFunction); }
+	else if (BindingInfo.Name == UnitToIndexFunctionName) { OutFunc = FVMExternalFunction::CreateUObject(this, &UNiagaraDataInterfaceRWBase::EmptyVMFunction); }
+	else if (BindingInfo.Name == IndexToLinearFunctionName) { OutFunc = FVMExternalFunction::CreateUObject(this, &UNiagaraDataInterfaceRWBase::EmptyVMFunction); }
+	else if (BindingInfo.Name == VoxelSizeFunctionName) { OutFunc = FVMExternalFunction::CreateUObject(this, &UNiagaraDataInterfaceRWBase::EmptyVMFunction); }
+	else if (BindingInfo.Name == NumVoxelsFunctionName) { OutFunc = FVMExternalFunction::CreateUObject(this, &UNiagaraDataInterfaceRWBase::EmptyVMFunction); }
+}
+
+bool UNiagaraDataInterfaceGrid3D::Equals(const UNiagaraDataInterface* Other) const
+{
+	if (!Super::Equals(Other))
+	{
+		return false;
+	}
+	const UNiagaraDataInterfaceGrid3D* OtherTyped = CastChecked<const UNiagaraDataInterfaceGrid3D>(Other);
+
+	return OtherTyped->NumVoxels == NumVoxels &&
+		FMath::IsNearlyEqual(OtherTyped->VoxelSize, VoxelSize) &&		
+		OtherTyped->WorldBBoxMin.Equals(WorldBBoxMin) &&
+		OtherTyped->WorldBBoxSize.Equals(WorldBBoxSize);
+}
+
+void UNiagaraDataInterfaceGrid3D::GetParameterDefinitionHLSL(FNiagaraDataInterfaceGPUParamInfo& ParamInfo, FString& OutHLSL)
+{
+	static const TCHAR *FormatDeclarations = TEXT(R"(
+		int3 {NumVoxels};
+		float3 {VoxelSize};
+		float3 {WorldBBoxMinName};
+		float3 {WorldBBoxSizeName};
+	)");
+	TMap<FString, FStringFormatArg> ArgsDeclarations = {
+		{ TEXT("NumVoxels"), NumVoxelsName + ParamInfo.DataInterfaceHLSLSymbol },
+		{ TEXT("VoxelSize"), VoxelSizeName + ParamInfo.DataInterfaceHLSLSymbol },
+		{ TEXT("WorldBBoxMinName"),  WorldBBoxMinName + ParamInfo.DataInterfaceHLSLSymbol },
+		{ TEXT("WorldBBoxSizeName"),    WorldBBoxSizeName + ParamInfo.DataInterfaceHLSLSymbol },
+	};
+	OutHLSL += FString::Format(FormatDeclarations, ArgsDeclarations);
+}
+
+bool UNiagaraDataInterfaceGrid3D::GetFunctionHLSL(const FName& DefinitionFunctionName, FString InstanceFunctionName, FNiagaraDataInterfaceGPUParamInfo& ParamInfo, FString& OutHLSL)
+{
+	if (DefinitionFunctionName == NumVoxelsFunctionName)
+	{
+		static const TCHAR *FormatSample = TEXT(R"(
+			void {FunctionName}(out int Out_NumVoxelsX, out int Out_NumVoxelsY, out int Out_NumVoxelsZ)
+			{
+				Out_NumVoxelsX = {NumVoxelsName}.x;
+				Out_NumVoxelsY = {NumVoxelsName}.y;
+				Out_NumVoxelsZ = {NumVoxelsName}.z;
+			}
+		)");
+		TMap<FString, FStringFormatArg> ArgsSample = {
+			{TEXT("FunctionName"), InstanceFunctionName},
+			{TEXT("NumVoxelsName"), NumVoxelsName + ParamInfo.DataInterfaceHLSLSymbol},
+
+		};
+		OutHLSL += FString::Format(FormatSample, ArgsSample);
+		return true;
+	}
+	else if (DefinitionFunctionName == WorldToUnitFunctionName)
+	{
+		static const TCHAR *FormatSample = TEXT(R"(
+			void {FunctionName}(float3 In_World, out float3 Out_Unit)
+			{
+				Out_Unit = (In_World - {WorldBBoxMinName}) / {WorldBBoxSizeName};
+			}
+		)");
+		TMap<FString, FStringFormatArg> ArgsSample = {
+			{TEXT("FunctionName"), InstanceFunctionName},
+			{TEXT("WorldBBoxMinName"), WorldBBoxMinName + ParamInfo.DataInterfaceHLSLSymbol},
+			{TEXT("WorldBBoxSizeName"), WorldBBoxSizeName + ParamInfo.DataInterfaceHLSLSymbol},
+		};
+		OutHLSL += FString::Format(FormatSample, ArgsSample);
+		return true;
+	}
+	else if (DefinitionFunctionName == UnitToIndexFunctionName)
+	{
+		static const TCHAR *FormatSample = TEXT(R"(
+			void {FunctionName}(float3 In_Unit, out int Out_IndexX, out int Out_IndexY, out int Out_IndexZ)
+			{
+				int3 Out_IndexTmp = round(In_Unit * {NumVoxelsName} - .5);
+				Out_IndexX = Out_IndexTmp.x;
+				Out_IndexY = Out_IndexTmp.y;
+				Out_IndexZ = Out_IndexTmp.z;
+			}
+		)");
+		TMap<FString, FStringFormatArg> ArgsSample = {
+			{TEXT("FunctionName"), InstanceFunctionName},
+			{TEXT("NumVoxelsName"), NumVoxelsName + ParamInfo.DataInterfaceHLSLSymbol},
+		};
+		OutHLSL += FString::Format(FormatSample, ArgsSample);
+		return true;
+	}
+	else if (DefinitionFunctionName == IndexToLinearFunctionName)
+	{
+		static const TCHAR *FormatSample = TEXT(R"(
+			void {FunctionName}(int In_IndexX, int In_IndexY, int In_IndexZ, out int Out_Linear)
+			{
+				Out_Linear = In_IndexX + In_IndexY * {NumVoxelsName}.x + In_IndexZ * {NumVoxelsName}.x * {NumVoxelsName}.y;
+			}
+		)");
+		TMap<FString, FStringFormatArg> ArgsSample = {
+			{TEXT("FunctionName"), InstanceFunctionName},
+			{TEXT("NumVoxelsName"), NumVoxelsName + ParamInfo.DataInterfaceHLSLSymbol},
+		};
+		OutHLSL += FString::Format(FormatSample, ArgsSample);
+		return true;
+	}	
+	else if (DefinitionFunctionName == NumVoxelsFunctionName)
+	{
+		static const TCHAR *FormatSample = TEXT(R"(
+			void {FunctionName}(out int Out_NumVoxelsX, out int Out_NumVoxelsY, out int Out_NumVoxelsZ)
+			{
+				Out_NumVoxelsX = {NumVoxelsName}.x;
+				Out_NumVoxelsY = {NumVoxelsName}.y;
+				Out_NumVoxelsZ = {NumVoxelsName}.z;
+			}
+		)");
+		TMap<FString, FStringFormatArg> ArgsSample = {
+			{TEXT("FunctionName"), InstanceFunctionName},
+			{TEXT("NumVoxelsName"), NumVoxelsName + ParamInfo.DataInterfaceHLSLSymbol},
+		};
+		OutHLSL += FString::Format(FormatSample, ArgsSample);
+		return true;
+	}
+	else if (DefinitionFunctionName == VoxelSizeFunctionName)
+	{
+		static const TCHAR *FormatSample = TEXT(R"(
+			void {FunctionName}(out float3 Out_VoxelSize)
+			{
+				Out_VoxelSize = {VoxelSizeName};
+			}
+		)");
+		TMap<FString, FStringFormatArg> ArgsSample = {
+			{TEXT("FunctionName"), InstanceFunctionName},
+			{TEXT("VoxelSizeName"), VoxelSizeName + ParamInfo.DataInterfaceHLSLSymbol},
+		};
+		OutHLSL += FString::Format(FormatSample, ArgsSample);
+		return true;
+	}
+
+	return false;
+}
+
+FNiagaraDataInterfaceParametersCS* UNiagaraDataInterfaceGrid3D::ConstructComputeParameters()const
+{
+	return nullptr;
+}
+
+
+
+bool UNiagaraDataInterfaceGrid3D::CopyToInternal(UNiagaraDataInterface* Destination) const
+{
+	if (!Super::CopyToInternal(Destination))
+	{
+		return false;
+	}
+
+	UNiagaraDataInterfaceGrid3D* OtherTyped = CastChecked<UNiagaraDataInterfaceGrid3D>(Destination);
+
+
+	OtherTyped->NumVoxels = NumVoxels;
+	OtherTyped->VoxelSize = VoxelSize;
+	OtherTyped->SetGridFromVoxelSize = SetGridFromVoxelSize;	
+	OtherTyped->WorldBBoxMin = WorldBBoxMin;
+	OtherTyped->WorldBBoxSize = WorldBBoxSize;
+
+	return true;
+}
+
+void UNiagaraDataInterfaceGrid3D::PushToRenderThread()
+{
+
+
+}
+
+/*--------------------------------------------------------------------------------------------------------------------------*/
+
+UNiagaraDataInterfaceGrid2D::UNiagaraDataInterfaceGrid2D(FObjectInitializer const& ObjectInitializer)
+	: Super(ObjectInitializer)
+	, NumCellsX(3)
+	, NumCellsY(3)
+	, CellSize(1.)
+	, SetGridFromCellSize(false)
+	, WorldBBoxMin(0., 0., 0.)
+	, WorldBBoxSize(100., 100.)
+{
+	Proxy = MakeShared<FNiagaraDataInterfaceProxyRW, ESPMode::ThreadSafe>();
+	PushToRenderThread();
+}
+
+
+void UNiagaraDataInterfaceGrid2D::GetFunctions(TArray<FNiagaraFunctionSignature>& OutFunctions)
+{
+
+	{
+		FNiagaraFunctionSignature Sig;
+		Sig.Name = NumCellsFunctionName;
+		Sig.Inputs.Add(FNiagaraVariable(FNiagaraTypeDefinition(GetClass()), TEXT("Grid")));
+		Sig.Outputs.Add(FNiagaraVariable(FNiagaraTypeDefinition::GetIntDef(), TEXT("NumCellsX")));
+		Sig.Outputs.Add(FNiagaraVariable(FNiagaraTypeDefinition::GetIntDef(), TEXT("NumCellsY")));		
+
+		Sig.bMemberFunction = true;
+		Sig.bRequiresContext = false;
+		OutFunctions.Add(Sig);
+	}
+
+	{
+		FNiagaraFunctionSignature Sig;
+		Sig.Name = WorldToUnitFunctionName;
+		Sig.Inputs.Add(FNiagaraVariable(FNiagaraTypeDefinition(GetClass()), TEXT("Grid")));
+		Sig.Inputs.Add(FNiagaraVariable(FNiagaraTypeDefinition::GetVec3Def(), TEXT("World")));
+		Sig.Outputs.Add(FNiagaraVariable(FNiagaraTypeDefinition::GetVec3Def(), TEXT("Unit")));
+
+		Sig.bMemberFunction = true;
+		Sig.bRequiresContext = false;
+		OutFunctions.Add(Sig);
+	}
+
+	{
+		FNiagaraFunctionSignature Sig;
+		Sig.Name = UnitToIndexFunctionName;
+		Sig.Inputs.Add(FNiagaraVariable(FNiagaraTypeDefinition(GetClass()), TEXT("Grid")));
+		Sig.Inputs.Add(FNiagaraVariable(FNiagaraTypeDefinition::GetVec2Def(), TEXT("Unit")));
+		Sig.Outputs.Add(FNiagaraVariable(FNiagaraTypeDefinition::GetIntDef(), TEXT("IndexX")));
+		Sig.Outputs.Add(FNiagaraVariable(FNiagaraTypeDefinition::GetIntDef(), TEXT("IndexY")));		
+
+		Sig.bMemberFunction = true;
+		Sig.bRequiresContext = false;
+		OutFunctions.Add(Sig);
+	}
+
+	{
+		FNiagaraFunctionSignature Sig;
+		Sig.Name = IndexToUnitFunctionName;
+		Sig.Inputs.Add(FNiagaraVariable(FNiagaraTypeDefinition(GetClass()), TEXT("Grid")));
+		Sig.Inputs.Add(FNiagaraVariable(FNiagaraTypeDefinition::GetFloatDef(), TEXT("IndexX")));
+		Sig.Inputs.Add(FNiagaraVariable(FNiagaraTypeDefinition::GetFloatDef(), TEXT("IndexY")));
+		Sig.Outputs.Add(FNiagaraVariable(FNiagaraTypeDefinition::GetVec3Def(), TEXT("Unit")));
+
+		Sig.bMemberFunction = true;
+		Sig.bRequiresContext = false;
+		OutFunctions.Add(Sig);
+	}
+
+	{
+		FNiagaraFunctionSignature Sig;
+		Sig.Name = IndexToUnitStaggeredXFunctionName;
+		Sig.Inputs.Add(FNiagaraVariable(FNiagaraTypeDefinition(GetClass()), TEXT("Grid")));
+		Sig.Inputs.Add(FNiagaraVariable(FNiagaraTypeDefinition::GetFloatDef(), TEXT("IndexX")));
+		Sig.Inputs.Add(FNiagaraVariable(FNiagaraTypeDefinition::GetFloatDef(), TEXT("IndexY")));
+		Sig.Outputs.Add(FNiagaraVariable(FNiagaraTypeDefinition::GetVec3Def(), TEXT("Unit")));
+
+		Sig.bMemberFunction = true;
+		Sig.bRequiresContext = false;
+		OutFunctions.Add(Sig);
+	}
+
+	{
+		FNiagaraFunctionSignature Sig;
+		Sig.Name = IndexToUnitStaggeredYFunctionName;
+		Sig.Inputs.Add(FNiagaraVariable(FNiagaraTypeDefinition(GetClass()), TEXT("Grid")));
+		Sig.Inputs.Add(FNiagaraVariable(FNiagaraTypeDefinition::GetFloatDef(), TEXT("IndexX")));
+		Sig.Inputs.Add(FNiagaraVariable(FNiagaraTypeDefinition::GetFloatDef(), TEXT("IndexY")));
+		Sig.Outputs.Add(FNiagaraVariable(FNiagaraTypeDefinition::GetVec3Def(), TEXT("Unit")));
+
+		Sig.bMemberFunction = true;
+		Sig.bRequiresContext = false;
+		OutFunctions.Add(Sig);
+	}
+
+	{
+		FNiagaraFunctionSignature Sig;
+		Sig.Name = IndexToLinearFunctionName;
+		Sig.Inputs.Add(FNiagaraVariable(FNiagaraTypeDefinition(GetClass()), TEXT("Grid")));
+		Sig.Inputs.Add(FNiagaraVariable(FNiagaraTypeDefinition::GetIntDef(), TEXT("IndexX")));
+		Sig.Inputs.Add(FNiagaraVariable(FNiagaraTypeDefinition::GetIntDef(), TEXT("IndexY")));		
+		Sig.Outputs.Add(FNiagaraVariable(FNiagaraTypeDefinition::GetIntDef(), TEXT("Linear")));
+
+		Sig.bMemberFunction = true;
+		Sig.bRequiresContext = false;
+		OutFunctions.Add(Sig);
+	}
+
+	{
+		FNiagaraFunctionSignature Sig;
+		Sig.Name = LinearToIndexFunctionName;
+		Sig.Inputs.Add(FNiagaraVariable(FNiagaraTypeDefinition(GetClass()), TEXT("Grid")));
+		Sig.Inputs.Add(FNiagaraVariable(FNiagaraTypeDefinition::GetIntDef(), TEXT("Linear")));
+		Sig.Outputs.Add(FNiagaraVariable(FNiagaraTypeDefinition::GetIntDef(), TEXT("IndexX")));
+		Sig.Outputs.Add(FNiagaraVariable(FNiagaraTypeDefinition::GetIntDef(), TEXT("IndexY")));
+		
+
+		Sig.bMemberFunction = true;
+		Sig.bRequiresContext = false;
+		OutFunctions.Add(Sig);
+	}
+
+	{
+		FNiagaraFunctionSignature Sig;
+		Sig.Name = NumCellsFunctionName;
+		Sig.Inputs.Add(FNiagaraVariable(FNiagaraTypeDefinition(GetClass()), TEXT("Grid")));
+		Sig.Outputs.Add(FNiagaraVariable(FNiagaraTypeDefinition::GetIntDef(), TEXT("NumVoxelsX")));
+		Sig.Outputs.Add(FNiagaraVariable(FNiagaraTypeDefinition::GetIntDef(), TEXT("NumCellsY")));		
+
+		Sig.bMemberFunction = true;
+		Sig.bRequiresContext = false;
+		OutFunctions.Add(Sig);
+	}
+
+	{
+		FNiagaraFunctionSignature Sig;
+		Sig.Name = CellSizeFunctionName;
+		Sig.Inputs.Add(FNiagaraVariable(FNiagaraTypeDefinition(GetClass()), TEXT("Grid")));
+		Sig.Outputs.Add(FNiagaraVariable(FNiagaraTypeDefinition::GetVec2Def(), TEXT("CellSize")));
+
+		Sig.bMemberFunction = true;
+		Sig.bRequiresContext = false;
+		OutFunctions.Add(Sig);
+	}
+}
+
+
+void UNiagaraDataInterfaceGrid2D::GetVMExternalFunction(const FVMExternalFunctionBindingInfo& BindingInfo, void* InstanceData, FVMExternalFunction &OutFunc)
+{
+	if (BindingInfo.Name == NumCellsFunctionName) { OutFunc = FVMExternalFunction::CreateUObject(this, &UNiagaraDataInterfaceRWBase::EmptyVMFunction); }
+	else if (BindingInfo.Name == WorldToUnitFunctionName) { OutFunc = FVMExternalFunction::CreateUObject(this, &UNiagaraDataInterfaceRWBase::EmptyVMFunction); }
+	else if (BindingInfo.Name == UnitToIndexFunctionName) { OutFunc = FVMExternalFunction::CreateUObject(this, &UNiagaraDataInterfaceRWBase::EmptyVMFunction); }
+	else if (BindingInfo.Name == IndexToUnitFunctionName) { OutFunc = FVMExternalFunction::CreateUObject(this, &UNiagaraDataInterfaceRWBase::EmptyVMFunction); }
+	else if (BindingInfo.Name == IndexToUnitStaggeredXFunctionName) { OutFunc = FVMExternalFunction::CreateUObject(this, &UNiagaraDataInterfaceRWBase::EmptyVMFunction); }
+	else if (BindingInfo.Name == IndexToUnitStaggeredYFunctionName) { OutFunc = FVMExternalFunction::CreateUObject(this, &UNiagaraDataInterfaceRWBase::EmptyVMFunction); }
+	else if (BindingInfo.Name == IndexToLinearFunctionName) { OutFunc = FVMExternalFunction::CreateUObject(this, &UNiagaraDataInterfaceRWBase::EmptyVMFunction); }
+	else if (BindingInfo.Name == LinearToIndexFunctionName) { OutFunc = FVMExternalFunction::CreateUObject(this, &UNiagaraDataInterfaceRWBase::EmptyVMFunction); }
+	else if (BindingInfo.Name == CellSizeFunctionName) { OutFunc = FVMExternalFunction::CreateUObject(this, &UNiagaraDataInterfaceRWBase::EmptyVMFunction); }
+	else if (BindingInfo.Name == NumCellsFunctionName) { OutFunc = FVMExternalFunction::CreateUObject(this, &UNiagaraDataInterfaceRWBase::EmptyVMFunction); }
+}
+
+bool UNiagaraDataInterfaceGrid2D::Equals(const UNiagaraDataInterface* Other) const
+{
+	if (!Super::Equals(Other))
+	{
+		return false;
+	}
+	const UNiagaraDataInterfaceGrid2D* OtherTyped = CastChecked<const UNiagaraDataInterfaceGrid2D>(Other);
+
+	return 
+		OtherTyped->NumCellsX == NumCellsX &&
+		OtherTyped->NumCellsY == NumCellsY &&
+		FMath::IsNearlyEqual(OtherTyped->CellSize, CellSize) &&
+		OtherTyped->WorldBBoxMin.Equals(WorldBBoxMin) &&
+		OtherTyped->WorldBBoxSize.Equals(WorldBBoxSize);
+}
+
+void UNiagaraDataInterfaceGrid2D::GetParameterDefinitionHLSL(FNiagaraDataInterfaceGPUParamInfo& ParamInfo, FString& OutHLSL)
+{
+	static const TCHAR *FormatDeclarations = TEXT(R"(
+		int2 {NumCells};
+		float2 {CellSize};
+		float3 {WorldBBoxMinName};
+		float2 {WorldBBoxSizeName};
+	)");
+	TMap<FString, FStringFormatArg> ArgsDeclarations = {
+		{ TEXT("NumCells"), NumCellsName + ParamInfo.DataInterfaceHLSLSymbol },
+		{ TEXT("CellSize"), CellSizeName + ParamInfo.DataInterfaceHLSLSymbol },
+		{ TEXT("WorldBBoxMinName"),  WorldBBoxMinName + ParamInfo.DataInterfaceHLSLSymbol },
+		{ TEXT("WorldBBoxSizeName"),    WorldBBoxSizeName + ParamInfo.DataInterfaceHLSLSymbol },
+	};
+	OutHLSL += FString::Format(FormatDeclarations, ArgsDeclarations);
+}
+
+bool UNiagaraDataInterfaceGrid2D::GetFunctionHLSL(const FName& DefinitionFunctionName, FString InstanceFunctionName, FNiagaraDataInterfaceGPUParamInfo& ParamInfo, FString& OutHLSL)
+{
+	if (DefinitionFunctionName == NumCellsFunctionName)
+	{
+		static const TCHAR *FormatSample = TEXT(R"(
+			void {FunctionName}(out int Out_NumCellsX, out int Out_NumCellsY)
+			{
+				Out_NumCellsX = {NumCellsName}.x;
+				Out_NumCellsY = {NumCellsName}.y;
+			}
+		)");
+		TMap<FString, FStringFormatArg> ArgsSample = {
+			{TEXT("FunctionName"), InstanceFunctionName},
+			{TEXT("NumCellsName"), NumCellsName + ParamInfo.DataInterfaceHLSLSymbol},
+
+		};
+		OutHLSL += FString::Format(FormatSample, ArgsSample);
+		return true;
+	}
+	else if (DefinitionFunctionName == WorldToUnitFunctionName)
+	{
+		static const TCHAR *FormatSample = TEXT(R"(
+			void {FunctionName}(float3 In_World, out float2 Out_Unit)
+			{
+				Out_Unit = (In_World - {WorldBBoxMinName}) / float3({WorldBBoxSizeName}, 1);
+			}
+		)");
+		TMap<FString, FStringFormatArg> ArgsSample = {
+			{TEXT("FunctionName"), InstanceFunctionName},
+			{TEXT("WorldBBoxMinName"), WorldBBoxMinName + ParamInfo.DataInterfaceHLSLSymbol},
+			{TEXT("WorldBBoxSizeName"), WorldBBoxSizeName + ParamInfo.DataInterfaceHLSLSymbol},
+		};
+		OutHLSL += FString::Format(FormatSample, ArgsSample);
+		return true;
+	}
+	else if (DefinitionFunctionName == UnitToIndexFunctionName)
+	{
+		static const TCHAR *FormatSample = TEXT(R"(
+			void {FunctionName}(float2 In_Unit, out int Out_IndexX, out int Out_IndexY)
+			{
+				int2 Out_IndexTmp = round(In_Unit * {NumCellsName}  - .5);
+				Out_IndexX = Out_IndexTmp.x;
+				Out_IndexY = Out_IndexTmp.y;				
+			}
+		)");
+		TMap<FString, FStringFormatArg> ArgsSample = {
+			{TEXT("FunctionName"), InstanceFunctionName},
+			{TEXT("NumCellsName"), NumCellsName + ParamInfo.DataInterfaceHLSLSymbol},			
+		};
+		OutHLSL += FString::Format(FormatSample, ArgsSample);
+		return true;
+	}
+	else if (DefinitionFunctionName == IndexToUnitFunctionName)
+	{
+		static const TCHAR *FormatSample = TEXT(R"(
+			void {FunctionName}(float In_IndexX, float In_IndexY, out float3 Out_Unit)
+			{
+				Out_Unit = float3((float2(In_IndexX, In_IndexY) + .5) / {NumCellsName}, 0);
+			}
+		)");
+		TMap<FString, FStringFormatArg> ArgsSample = {
+			{TEXT("FunctionName"), InstanceFunctionName},
+			{TEXT("NumCellsName"), NumCellsName + ParamInfo.DataInterfaceHLSLSymbol},
+		};
+		OutHLSL += FString::Format(FormatSample, ArgsSample);
+		return true;
+	}
+	else if (DefinitionFunctionName == IndexToUnitStaggeredXFunctionName)
+	{
+		static const TCHAR *FormatSample = TEXT(R"(
+			void {FunctionName}(float In_IndexX, float In_IndexY, out float3 Out_Unit)
+			{
+				Out_Unit = float3((float2(In_IndexX, In_IndexY) + float2(0.0, 0.5)) / {NumCellsName}, 0);
+			}
+		)");
+		TMap<FString, FStringFormatArg> ArgsSample = {
+			{TEXT("FunctionName"), InstanceFunctionName},
+			{TEXT("NumCellsName"), NumCellsName + ParamInfo.DataInterfaceHLSLSymbol},
+		};
+		OutHLSL += FString::Format(FormatSample, ArgsSample);
+		return true;
+	}
+	else if (DefinitionFunctionName == IndexToUnitStaggeredYFunctionName)
+	{
+		static const TCHAR *FormatSample = TEXT(R"(
+			void {FunctionName}(float In_IndexX, float In_IndexY, out float3 Out_Unit)
+			{
+				Out_Unit = float3((float2(In_IndexX, In_IndexY) +  + float2(0.5, 0.0)) / {NumCellsName}, 0);
+			}
+		)");
+		TMap<FString, FStringFormatArg> ArgsSample = {
+			{TEXT("FunctionName"), InstanceFunctionName},
+			{TEXT("NumCellsName"), NumCellsName + ParamInfo.DataInterfaceHLSLSymbol},
+		};
+		OutHLSL += FString::Format(FormatSample, ArgsSample);
+		return true;
+	}
+	else if (DefinitionFunctionName == IndexToLinearFunctionName)
+	{
+		static const TCHAR *FormatSample = TEXT(R"(
+			void {FunctionName}(int In_IndexX, int In_IndexY, out int Out_Linear)
+			{
+				Out_Linear = In_IndexX + In_IndexY * {NumCellsName}.x;
+			}
+		)");
+		TMap<FString, FStringFormatArg> ArgsSample = {
+			{TEXT("FunctionName"), InstanceFunctionName},
+			{TEXT("NumCellsName"), NumCellsName + ParamInfo.DataInterfaceHLSLSymbol},
+		};
+		OutHLSL += FString::Format(FormatSample, ArgsSample);
+		return true;
+	}
+	else if (DefinitionFunctionName == LinearToIndexFunctionName)
+	{
+		static const TCHAR *FormatSample = TEXT(R"(
+			void {FunctionName}(int In_Linear, out int Out_IndexX, out int Out_IndexY)
+			{
+				Out_IndexX = In_Linear % {NumCells}.x;
+				Out_IndexY = In_Linear / {NumCells}.x;				
+			}
+		)");
+		TMap<FString, FStringFormatArg> ArgsSample = {
+			{TEXT("FunctionName"), InstanceFunctionName},
+			{TEXT("NumCells"), NumCellsName + ParamInfo.DataInterfaceHLSLSymbol},
+		};
+		OutHLSL += FString::Format(FormatSample, ArgsSample);
+		return true;
+	}
+	else if (DefinitionFunctionName == NumCellsFunctionName)
+	{
+		static const TCHAR *FormatSample = TEXT(R"(
+			void {FunctionName}(out int Out_NumCellsX, out int Out_NumCellsY)
+			{
+				Out_NumCellsX = {NumCells}.x;
+				Out_NumCellsY = {NumCells}.y;				
+			}
+		)");
+		TMap<FString, FStringFormatArg> ArgsSample = {
+			{TEXT("FunctionName"), InstanceFunctionName},
+			{TEXT("NumCells"), NumCellsName + ParamInfo.DataInterfaceHLSLSymbol},
+		};
+		OutHLSL += FString::Format(FormatSample, ArgsSample);
+		return true;
+	}
+	else if (DefinitionFunctionName == CellSizeFunctionName)
+	{
+		static const TCHAR *FormatSample = TEXT(R"(
+			void {FunctionName}(out float2 Out_CellSize)
+			{
+				Out_CellSize = {CellSizeName};
+			}
+		)");
+		TMap<FString, FStringFormatArg> ArgsSample = {
+			{TEXT("FunctionName"), InstanceFunctionName},
+			{TEXT("CellSizeName"), CellSizeName + ParamInfo.DataInterfaceHLSLSymbol},
+		};
+		OutHLSL += FString::Format(FormatSample, ArgsSample);
+		return true;
+	}
+
+	return false;
+}
+
+FNiagaraDataInterfaceParametersCS* UNiagaraDataInterfaceGrid2D::ConstructComputeParameters()const
+{
+	return nullptr;
+}
+
+
+
+bool UNiagaraDataInterfaceGrid2D::CopyToInternal(UNiagaraDataInterface* Destination) const
+{
+	if (!Super::CopyToInternal(Destination))
+	{
+		return false;
+	}
+
+	UNiagaraDataInterfaceGrid2D* OtherTyped = CastChecked<UNiagaraDataInterfaceGrid2D>(Destination);
+
+
+	OtherTyped->NumCellsX = NumCellsX;
+	OtherTyped->NumCellsY = NumCellsY;
+	OtherTyped->CellSize = CellSize;
+	OtherTyped->SetGridFromCellSize = SetGridFromCellSize;
+	OtherTyped->WorldBBoxMin = WorldBBoxMin;
+	OtherTyped->WorldBBoxSize = WorldBBoxSize;
+
+	return true;
+}
+
+void UNiagaraDataInterfaceGrid2D::PushToRenderThread()
+{
+
+
+}
+
+
+#undef LOCTEXT_NAMESPACE
diff --git a/Engine/Plugins/FX/Niagara/Source/Niagara/Private/NiagaraDataSet.cpp b/Engine/Plugins/FX/Niagara/Source/Niagara/Private/NiagaraDataSet.cpp
--- a/Engine/Plugins/FX/Niagara/Source/Niagara/Private/NiagaraDataSet.cpp
+++ b/Engine/Plugins/FX/Niagara/Source/Niagara/Private/NiagaraDataSet.cpp
@@ -182,13 +182,18 @@ FNiagaraDataBuffer& FNiagaraDataSet::BeginSimulate()
 	return GetDestinationDataChecked();
 }
 
-void FNiagaraDataSet::EndSimulate()
+void FNiagaraDataSet::EndSimulate(bool SetCurrentData)
 {
 	//CheckCorrectThread();
 	//Destination is now complete so make it the current simulation state.
 	DestinationData->Unlock();
 	check(!DestinationData->IsInUse());
-	CurrentData = DestinationData;
+
+	if (SetCurrentData)
+	{
+		CurrentData = DestinationData;
+	}
+
 	DestinationData = nullptr;
 }
 
diff --git a/Engine/Plugins/FX/Niagara/Source/Niagara/Private/NiagaraEmitter.cpp b/Engine/Plugins/FX/Niagara/Source/Niagara/Private/NiagaraEmitter.cpp
--- a/Engine/Plugins/FX/Niagara/Source/Niagara/Private/NiagaraEmitter.cpp
+++ b/Engine/Plugins/FX/Niagara/Source/Niagara/Private/NiagaraEmitter.cpp
@@ -107,6 +107,7 @@ UNiagaraEmitter::UNiagaraEmitter(const FObjectInitializer& Initializer)
 , bUseMaxDetailLevel(false)
 , bRequiresPersistentIDs(false)
 , MaxDeltaTimePerTick(0.125)
+, MaxUpdateIterations(1)
 , bLimitDeltaTime(true)
 #if WITH_EDITORONLY_DATA
 , ThumbnailImageOutOfDate(true)
diff --git a/Engine/Plugins/FX/Niagara/Source/Niagara/Private/NiagaraEmitterInstance.cpp b/Engine/Plugins/FX/Niagara/Source/Niagara/Private/NiagaraEmitterInstance.cpp
--- a/Engine/Plugins/FX/Niagara/Source/Niagara/Private/NiagaraEmitterInstance.cpp
+++ b/Engine/Plugins/FX/Niagara/Source/Niagara/Private/NiagaraEmitterInstance.cpp
@@ -283,7 +283,7 @@ void FNiagaraEmitterInstance::Init(int32 InEmitterIdx, FName InSystemInstanceNam
 	if (CachedEmitter->SimTarget == ENiagaraSimTarget::GPUComputeSim)
 	{
 		GPUExecContext = new FNiagaraComputeExecutionContext();
-		GPUExecContext->InitParams(CachedEmitter->GetGPUComputeScript(), CachedEmitter->SimTarget, CachedEmitter->GetUniqueEmitterName());
+		GPUExecContext->InitParams(CachedEmitter->GetGPUComputeScript(), CachedEmitter->SimTarget, CachedEmitter->GetUniqueEmitterName(), CachedEmitter->MaxUpdateIterations, CachedEmitter->SpawnStages);
 		GPUExecContext->MainDataSet = &Data;
 		GPUExecContext->GPUScript_RT = CachedEmitter->GetGPUComputeScript()->GetRenderThreadScript();
 
diff --git a/Engine/Plugins/FX/Niagara/Source/Niagara/Private/NiagaraEmitterInstanceBatcher.cpp b/Engine/Plugins/FX/Niagara/Source/Niagara/Private/NiagaraEmitterInstanceBatcher.cpp
--- a/Engine/Plugins/FX/Niagara/Source/Niagara/Private/NiagaraEmitterInstanceBatcher.cpp
+++ b/Engine/Plugins/FX/Niagara/Source/Niagara/Private/NiagaraEmitterInstanceBatcher.cpp
@@ -155,6 +155,173 @@ void NiagaraEmitterInstanceBatcher::ReleaseTicks()
 	Ticks_RT.Empty(0);
 }
 
+bool NiagaraEmitterInstanceBatcher::ResetDataInterfaces(const FNiagaraGPUSystemTick& Tick, FNiagaraComputeInstanceData *Instance, FRHICommandList &RHICmdList, FNiagaraShader* ComputeShader ) const
+{
+	bool ValidSpawnStage = true;
+	FNiagaraComputeExecutionContext* Context = Instance->Context;
+
+	// The first spawn stage must not be an iteration stage
+		// #todo(dmp): fix up this logic at some point so we don't have this constraint, but likely going to be tough
+	uint32 MinSpawnStage = TNumericLimits<uint32>::Max();
+	for (uint32 SpawnStage : Context->SpawnStages)
+	{
+		if (SpawnStage < MinSpawnStage)
+		{
+			MinSpawnStage = SpawnStage;
+		}
+	}
+
+	// check all data interfaces to see if one of them has the min stage as an iteration stage
+	// if so, error out
+	for (FNiagaraDataInterfaceProxy* Interface : Instance->DataInterfaceProxies)
+	{
+		if (Interface->IsIterationStage(MinSpawnStage))
+		{
+			UE_LOG(LogNiagara, Error, TEXT("First spawn stage in execution shouldn't be an iteration stage"));
+			ValidSpawnStage = false;
+		}
+	}
+
+	// Reset all rw data interface data
+	if (Tick.bNeedsReset)
+	{
+		uint32 InterfaceIndex = 0;
+		for (FNiagaraDataInterfaceProxy* Interface : Instance->DataInterfaceProxies)
+		{
+			FNiagaraDataInterfaceParamRef& DIParam = ComputeShader->GetDIParameters()[InterfaceIndex];
+			if (DIParam.Parameters)
+			{
+				FNiagaraDataInterfaceSetArgs TmpContext;
+				TmpContext.Shader = ComputeShader;
+				TmpContext.DataInterface = Interface;
+				TmpContext.SystemInstance = Tick.SystemInstanceID;
+				TmpContext.Batcher = this;		
+				Interface->ResetData(RHICmdList, TmpContext);
+			}			
+			InterfaceIndex++;
+		}
+	}
+	return ValidSpawnStage;
+}
+
+FNiagaraDataInterfaceProxy* NiagaraEmitterInstanceBatcher::FindIterationInterface( FNiagaraComputeInstanceData *Instance, const uint32 ShaderStageIndex) const
+{
+	// Determine if the iteration is outputting to a custom data size
+	FNiagaraDataInterfaceProxy* IterationInterface = nullptr;
+
+	for (FNiagaraDataInterfaceProxy* Interface : Instance->DataInterfaceProxies)
+	{
+		if (Interface->IsIterationStage(ShaderStageIndex))
+		{
+			if (IterationInterface)
+			{
+				UE_LOG(LogNiagara, Error, TEXT("Multiple output Data Interfaces found for current stage"));
+			}
+			else
+			{
+				IterationInterface = Interface;
+			}
+		}
+	}
+	return IterationInterface;
+}
+
+void NiagaraEmitterInstanceBatcher::PreStageInterface(const FNiagaraGPUSystemTick& Tick, FNiagaraComputeInstanceData *Instance, FRHICommandList &RHICmdList, FNiagaraShader* ComputeShader, const uint32 ShaderStageIndex) const
+{
+	uint32 InterfaceIndex = 0;
+	for (FNiagaraDataInterfaceProxy* Interface : Instance->DataInterfaceProxies)
+	{
+		FNiagaraDataInterfaceParamRef& DIParam = ComputeShader->GetDIParameters()[InterfaceIndex];
+		if (DIParam.Parameters)
+		{
+			FNiagaraDataInterfaceSetArgs TmpContext;
+			TmpContext.Shader = ComputeShader;
+			TmpContext.DataInterface = Interface;
+			TmpContext.SystemInstance = Tick.SystemInstanceID;
+			TmpContext.Batcher = this;
+			TmpContext.ShaderStageIndex = ShaderStageIndex;
+			TmpContext.IsOutputStage = Interface->IsOutputStage(ShaderStageIndex);
+			TmpContext.IsIterationStage = Interface->IsIterationStage(ShaderStageIndex);
+			Interface->PreStage(RHICmdList, TmpContext);
+		}
+		InterfaceIndex++;
+	}
+}
+
+void NiagaraEmitterInstanceBatcher::PostStageInterface(const FNiagaraGPUSystemTick& Tick, FNiagaraComputeInstanceData *Instance, FRHICommandList &RHICmdList, FNiagaraShader* ComputeShader, const uint32 ShaderStageIndex) const
+{
+	uint32 InterfaceIndex = 0;
+	for (FNiagaraDataInterfaceProxy* Interface : Instance->DataInterfaceProxies)
+	{
+		FNiagaraDataInterfaceParamRef& DIParam = ComputeShader->GetDIParameters()[InterfaceIndex];
+		if (DIParam.Parameters)
+		{
+			FNiagaraDataInterfaceSetArgs TmpContext;
+			TmpContext.Shader = ComputeShader;
+			TmpContext.DataInterface = Interface;
+			TmpContext.SystemInstance = Tick.SystemInstanceID;
+			TmpContext.Batcher = this;
+			TmpContext.ShaderStageIndex = ShaderStageIndex;
+			TmpContext.IsOutputStage = Interface->IsOutputStage(ShaderStageIndex);
+			TmpContext.IsIterationStage = Interface->IsIterationStage(ShaderStageIndex);
+			Interface->PostStage(RHICmdList, TmpContext);
+		}
+		InterfaceIndex++;
+	}
+}
+
+template<bool bDoResourceTransitions>
+void NiagaraEmitterInstanceBatcher::DispatchMultipleStages(const FNiagaraGPUSystemTick& Tick, FNiagaraComputeInstanceData *Instance, FRHICommandList &RHICmdList, FRHIUniformBuffer* ViewUniformBuffer, FNiagaraShader* ComputeShader) const
+{
+	FNiagaraComputeExecutionContext* Context = Instance->Context;
+
+	static const auto UseShaderStagesCVar = IConsoleManager::Get().FindConsoleVariable(TEXT("fx.UseShaderStages"));
+	if (UseShaderStagesCVar->GetInt() == 1)
+	{
+		FNiagaraDataBuffer* CurrentData = Instance->CurrentData;
+		FNiagaraDataBuffer* DestinationData = Instance->DestinationData;
+
+		const uint32 NumStages = Instance->Context->MaxUpdateIterations;
+		for (uint32 ShaderStageIndex = 0; ShaderStageIndex < NumStages; ++ShaderStageIndex)
+		{
+			if (Context->SpawnStages.Num() > 0 &&
+				((Tick.bNeedsReset && !Context->SpawnStages.Contains(ShaderStageIndex)) ||
+				(!Tick.bNeedsReset && Context->SpawnStages.Contains(ShaderStageIndex))))
+			{
+				continue;
+			}
+
+			// Determine if the iteration is outputting to a custom data size
+			FNiagaraDataInterfaceProxy *IterationInterface = FindIterationInterface(Instance, ShaderStageIndex);
+			PreStageInterface(Tick, Instance, RHICmdList, ComputeShader, ShaderStageIndex);
+
+			if (!IterationInterface)
+			{
+				if (ShaderStageIndex > 0)
+				{
+					FNiagaraDataBuffer* StoreData = Instance->CurrentData;
+					Instance->CurrentData = Instance->DestinationData;
+					Instance->DestinationData = StoreData;
+				}
+				Run<bDoResourceTransitions>(Tick, Instance, 0, Instance->DestinationData->GetNumInstances(), ComputeShader, RHICmdList, ViewUniformBuffer, Instance->SpawnInfo, false, ShaderStageIndex,  nullptr);
+			}
+			else
+			{
+				// run with correct number of instances.  This will make curr data junk or empty
+				Run<bDoResourceTransitions>(Tick, Instance, 0, IterationInterface->ElementCount, ComputeShader, RHICmdList, ViewUniformBuffer, Instance->SpawnInfo, false, ShaderStageIndex, IterationInterface);
+			}
+			PostStageInterface(Tick, Instance, RHICmdList, ComputeShader, ShaderStageIndex);
+		}
+		Instance->CurrentData = CurrentData;
+		Instance->DestinationData = DestinationData;
+	}
+	else
+	{
+		// run shader, sim and spawn in a single dispatch
+		Run<bDoResourceTransitions>(Tick, Instance, 0, Instance->DestinationData->GetNumInstances(), ComputeShader, RHICmdList, ViewUniformBuffer, Instance->SpawnInfo);
+	}
+}
+
 void NiagaraEmitterInstanceBatcher::ResizeBuffersAndGatherResources(FOverlappableTicks& OverlappableTick, FRHICommandList& RHICmdList, FNiagaraBufferArray& DestDataBuffers, FNiagaraBufferArray& CurrDataBuffers, FNiagaraBufferArray& DestBufferIntFloat, FNiagaraBufferArray& CurrBufferIntFloat)
 {
 	SCOPE_CYCLE_COUNTER(STAT_NiagaraGPUDispatchSetup_RT);
@@ -288,8 +455,7 @@ void NiagaraEmitterInstanceBatcher::DispatchAllOnCompute(FOverlappableTicks& Ove
 					FNiagaraComputeExecutionContext::TickCounter++;
 
 					// run shader, sim and spawn in a single dispatch
-					uint32 UpdateStartInstance = 0;
-					Run<false>(*Tick, &Instance, UpdateStartInstance, Instance.DestinationData->GetNumInstances(), Context->GPUScript_RT->GetShader(), RHICmdList, ViewUniformBuffer, Instance.SpawnInfo);
+					DispatchMultipleStages<false>(*Tick, &Instance, RHICmdList, ViewUniformBuffer, Context->GPUScript_RT->GetShader());
 
 					FNiagaraDataBuffer* CurrentData = Instance.CurrentData;
 					if (bSetReadback && Tick->bIsFinalTick)
@@ -445,6 +611,7 @@ void NiagaraEmitterInstanceBatcher::TickSingle(const FNiagaraGPUSystemTick& Tick
 		ProcessDebugInfo(RHICmdList, Context);
 	}
 #endif // WITH_EDITORONLY_DATA
+	if (!ResetDataInterfaces(Tick, Instance, RHICmdList, ComputeShader)) return;
 
 	//The buffer containing current simulation state.
 	Instance->CurrentData = Context->MainDataSet->GetCurrentData();
@@ -462,8 +629,7 @@ void NiagaraEmitterInstanceBatcher::TickSingle(const FNiagaraGPUSystemTick& Tick
 	DestinationData.SetNumInstances(NewNumInstances);
 
 	// run shader, sim and spawn in a single dispatch
-	uint32 UpdateStartInstance = 0;
-	Run<true>(Tick, Instance, UpdateStartInstance, NewNumInstances, ComputeShader, RHICmdList, ViewUniformBuffer, Instance->SpawnInfo);
+	DispatchMultipleStages<true>(Tick, Instance, RHICmdList, ViewUniformBuffer, ComputeShader);
 
 	Context->MainDataSet->EndSimulate();
 	Context->SetDataToRender(Instance->DestinationData);
@@ -471,7 +637,7 @@ void NiagaraEmitterInstanceBatcher::TickSingle(const FNiagaraGPUSystemTick& Tick
 	/*
 	// TODO: hack - only updating event set 0 on update scripts now; need to match them to their indices and update them all
 	if (Context->UpdateEventWriteDataSets.Num())
-	{ 
+	{
 		Context->UpdateEventWriteDataSets[0]->CurrDataRender().SetNumInstances(NumInstancesAfterSim[1]);
 	}
 	RunEventHandlers(Context, NumInstancesAfterSim[0], NumInstancesAfterSpawn, NumInstancesAfterNonEventSpawn, RHICmdList);
@@ -875,7 +1041,7 @@ void NiagaraEmitterInstanceBatcher::ProcessDebugInfo(FRHICommandList &RHICmdList
 
 /* Set shader parameters for data interfaces
  */
-void NiagaraEmitterInstanceBatcher::SetDataInterfaceParameters(const TArray<FNiagaraDataInterfaceProxy*>& DataInterfaceProxies, FNiagaraShader* Shader, FRHICommandList &RHICmdList, const FNiagaraComputeInstanceData* Instance, const FNiagaraGPUSystemTick& Tick) const
+void NiagaraEmitterInstanceBatcher::SetDataInterfaceParameters(const TArray<FNiagaraDataInterfaceProxy*>& DataInterfaceProxies, FNiagaraShader* Shader, FRHICommandList &RHICmdList, const FNiagaraComputeInstanceData* Instance, const FNiagaraGPUSystemTick& Tick, uint32 ShaderStageIndex) const
 {
 	// set up data interface buffers, as defined by the DIs during compilation
 	//
@@ -894,6 +1060,8 @@ void NiagaraEmitterInstanceBatcher::SetDataInterfaceParameters(const TArray<FNia
 			Context.DataInterface = Interface;
 			Context.SystemInstance = SystemInstance;
 			Context.Batcher = this;
+			Context.ShaderStageIndex = ShaderStageIndex;
+			Context.IsOutputStage = Interface->IsOutputStage(ShaderStageIndex);
 			DIParam.Parameters->Set(RHICmdList, Context);
 		}
 
@@ -942,14 +1110,15 @@ void NiagaraEmitterInstanceBatcher::UnsetDataInterfaceParameters(const TArray<FN
  */
  template<bool bDoResourceTransitions>
 void NiagaraEmitterInstanceBatcher::Run(const FNiagaraGPUSystemTick& Tick, const FNiagaraComputeInstanceData* Instance, uint32 UpdateStartInstance, const uint32 TotalNumInstances, FNiagaraShader* Shader,
-	FRHICommandList &RHICmdList, FRHIUniformBuffer* ViewUniformBuffer, const FNiagaraGpuSpawnInfo& SpawnInfo, bool bCopyBeforeStart) const
+	FRHICommandList &RHICmdList, FRHIUniformBuffer* ViewUniformBuffer, const FNiagaraGpuSpawnInfo& SpawnInfo, bool bCopyBeforeStart, uint32 ShaderStageIndex, FNiagaraDataInterfaceProxy *IterationInterface) const
 {
 	FNiagaraComputeExecutionContext* Context = Instance->Context;
 	if (TotalNumInstances == 0 && !bDoResourceTransitions)
 	{
-		SCOPED_DRAW_EVENTF(RHICmdList, NiagaraGPUSimulationCS, TEXT("Niagara Gpu Sim - %s - NumInstances: %u"),
+		SCOPED_DRAW_EVENTF(RHICmdList, NiagaraGPUSimulationCS, TEXT("Niagara Gpu Sim - %s - NumInstances: %u - StageNumber: %u"),
 			*Context->DebugSimName,
-			TotalNumInstances);
+			TotalNumInstances,
+			ShaderStageIndex);
 		return;
 	}
 
@@ -963,13 +1132,26 @@ void NiagaraEmitterInstanceBatcher::Run(const FNiagaraGPUSystemTick& Tick, const
 
 	RHICmdList.SetComputeShader(Shader->GetComputeShader());
 
+	// #todo(dmp): clean up this logic for shader stages on first frame
+	if (Tick.bNeedsReset)
+	{
+		int v = 1;
+		RHICmdList.SetShaderParameter(Shader->GetComputeShader(), Shader->SimStartParam.GetBufferIndex(), Shader->SimStartParam.GetBaseIndex(), Shader->SimStartParam.GetNumBytes(), &v);
+	}
+	else
+	{
+		int v = 0;
+		RHICmdList.SetShaderParameter(Shader->GetComputeShader(), Shader->SimStartParam.GetBufferIndex(), Shader->SimStartParam.GetBaseIndex(), Shader->SimStartParam.GetNumBytes(), &v);
+	}
+	
+
 	// set the view uniform buffer param
 	if (Shader->ViewUniformBufferParam.IsBound() && ViewUniformBuffer)
 	{
 		RHICmdList.SetShaderUniformBuffer(Shader->GetComputeShader(), Shader->ViewUniformBufferParam.GetBaseIndex(), ViewUniformBuffer);
 	}
 
-	SetDataInterfaceParameters(DataInterfaceProxies, Shader, RHICmdList, Instance, Tick);
+	SetDataInterfaceParameters(DataInterfaceProxies, Shader, RHICmdList, Instance, Tick, ShaderStageIndex);
 
 	// set the shader and data set params 
 	//
@@ -978,7 +1160,7 @@ void NiagaraEmitterInstanceBatcher::Run(const FNiagaraGPUSystemTick& Tick, const
 
 	// set the index buffer uav
 	//
-	if (Shader->InstanceCountsParam.IsBound())
+	if (Shader->InstanceCountsParam.IsBound() && !IterationInterface )
 	{
 		RHICmdList.TransitionResource(EResourceTransitionAccess::ERWNoBarrier, EResourceTransitionPipeline::EComputeToCompute, GPUInstanceCounterManager.GetInstanceCountBuffer().UAV);
 		RHICmdList.SetUAVParameter(Shader->GetComputeShader(), Shader->InstanceCountsParam.GetUAVIndex(), GPUInstanceCounterManager.GetInstanceCountBuffer().UAV);
@@ -1004,7 +1186,30 @@ void NiagaraEmitterInstanceBatcher::Run(const FNiagaraGPUSystemTick& Tick, const
 
 	RHICmdList.SetShaderParameter(Shader->GetComputeShader(), Shader->UpdateStartInstanceParam.GetBufferIndex(), Shader->UpdateStartInstanceParam.GetBaseIndex(), Shader->UpdateStartInstanceParam.GetNumBytes(), &UpdateStartInstance);					// 0, except for event handler runs
 	int32 InstancesToSpawnThisFrame = Instance->SpawnInfo.SpawnRateInstances + Instance->SpawnInfo.EventSpawnTotal;
+
+	// Only spawn particles on the first stage
+	if (ShaderStageIndex > 0)
+	{
+		InstancesToSpawnThisFrame = 0;
+	}
+
 	RHICmdList.SetShaderParameter(Shader->GetComputeShader(), Shader->NumSpawnedInstancesParam.GetBufferIndex(), Shader->NumSpawnedInstancesParam.GetBaseIndex(), Shader->NumSpawnedInstancesParam.GetNumBytes(), &InstancesToSpawnThisFrame);				// number of instances in the spawn run
+	RHICmdList.SetShaderParameter(Shader->GetComputeShader(), Shader->ShaderStageIndexParam.GetBufferIndex(), Shader->ShaderStageIndexParam.GetBaseIndex(), Shader->ShaderStageIndexParam.GetNumBytes(), &ShaderStageIndex);					// 0, except if several stages are defined
+
+	if (IterationInterface != NULL)
+	{
+
+	if (TotalNumInstances > NIAGARA_COMPUTE_THREADGROUP_SIZE)
+	{
+		RHICmdList.SetShaderParameter(Shader->GetComputeShader(), Shader->IterationInterfaceCount.GetBufferIndex(), Shader->IterationInterfaceCount.GetBaseIndex(), Shader->IterationInterfaceCount.GetNumBytes(), &TotalNumInstances);					// 0, except if several stages are defined
+	}
+	else
+	{
+		const int v = -1;
+		RHICmdList.SetShaderParameter(Shader->GetComputeShader(), Shader->IterationInterfaceCount.GetBufferIndex(), Shader->IterationInterfaceCount.GetBaseIndex(), Shader->IterationInterfaceCount.GetNumBytes(), &v);					// 0, except if several stages are defined
+	}
+	}
+
 
 	uint32 NumThreadGroups = 1;
 	if (TotalNumInstances > NIAGARA_COMPUTE_THREADGROUP_SIZE)
@@ -1025,23 +1230,42 @@ void NiagaraEmitterInstanceBatcher::Run(const FNiagaraGPUSystemTick& Tick, const
 		ensure(!Shader->EmitterConstantBufferParam.IsBound());
 	}
 
+	// #todo(dmp): temporary hack -- unbind UAVs if we have a valid iteration DI.  This way, when we are outputting with a different iteration count, we don't
+	// mess up particle state
+	if (IterationInterface)
+	{
+		CurrentData.UnsetShaderParams(Shader, RHICmdList);
+		DestinationData.UnsetShaderParams(Shader, RHICmdList); 
+	}
+
+	//UE_LOG(LogNiagara, Log, TEXT("Num Instance : %d | Num Group : %d | Spawned Istance : %d | Start Instance : %d | Num Indices : %d | Stage Index : %d"), 
+		//TotalNumInstances, NumThreadGroups, InstancesToSpawnThisFrame, UpdateStartInstance, Context->NumIndicesPerInstance, ShaderStageIndex);
+
 	// Dispatch, if anything needs to be done
-	//
 	if (TotalNumInstances)
 	{
-		SCOPED_DRAW_EVENTF(RHICmdList, NiagaraGPUSimulationCS, TEXT("Niagara Gpu Sim - %s - NumInstances: %u"),
+		SCOPED_DRAW_EVENTF(RHICmdList, NiagaraGPUSimulationCS, TEXT("Niagara Gpu Sim - %s - NumInstances: %u - StageNumber: %u"),
 			*Context->DebugSimName,
-			TotalNumInstances);
+			TotalNumInstances,
+			ShaderStageIndex);
 		DispatchComputeShader(RHICmdList, Shader, NumThreadGroups, 1, 1);
 	}
 
+	// reset iteration count
+	if (IterationInterface)
+	{
+		const int v = -1;
+		RHICmdList.SetShaderParameter(Shader->GetComputeShader(), Shader->IterationInterfaceCount.GetBufferIndex(), Shader->IterationInterfaceCount.GetBaseIndex(), Shader->IterationInterfaceCount.GetNumBytes(), &v);					// 0, except if several stages are defined
+
+	}
+
 #if WITH_EDITORONLY_DATA
 	// Check to see if we need to queue up a debug dump..
 	if (Context->DebugInfo.IsValid())
 	{
 		//UE_LOG(LogNiagara, Warning, TEXT("Queued up!"));
 
-		if (!Context->GPUDebugDataReadbackFloat && !Context->GPUDebugDataReadbackInt && !Context->GPUDebugDataReadbackCounts && DestinationData.GetGPUInstanceCountBufferOffset() != INDEX_NONE)
+		if (!Context->GPUDebugDataReadbackFloat && !Context->GPUDebugDataReadbackInt && !Context->GPUDebugDataReadbackCounts && DestinationData.GetGPUInstanceCountBufferOffset() != INDEX_NONE && ShaderStageIndex == Context->MaxUpdateIterations - 1)
 		{
 			Context->GPUDebugDataFloatSize = 0;
 			Context->GPUDebugDataIntSize = 0;
diff --git a/Engine/Plugins/FX/Niagara/Source/Niagara/Private/NiagaraModule.cpp b/Engine/Plugins/FX/Niagara/Source/Niagara/Private/NiagaraModule.cpp
--- a/Engine/Plugins/FX/Niagara/Source/Niagara/Private/NiagaraModule.cpp
+++ b/Engine/Plugins/FX/Niagara/Source/Niagara/Private/NiagaraModule.cpp
@@ -39,6 +39,16 @@ static FAutoConsoleVariableRef CVarEnableVerboseNiagaraChangeIdLogging(
 	ECVF_Default
 );
 
+/**
+Use Shader Stages CVar.
+Enable the custom dispatch for multiple shader stages 
+*/
+static TAutoConsoleVariable<int32> CVarUseShaderStages(
+	TEXT("fx.UseShaderStages"), 
+	0, 
+	TEXT("Enable or not the shader stages within Niagara (WIP feature only there for temporary testing)."),
+	ECVF_Default);
+
 /**
 Detail Level CVar.
 Effectively replaces the DetaiMode feature but allows for a rolling range of new hardware and emitters to target them.
diff --git a/Engine/Plugins/FX/Niagara/Source/Niagara/Private/NiagaraScriptExecutionContext.cpp b/Engine/Plugins/FX/Niagara/Source/Niagara/Private/NiagaraScriptExecutionContext.cpp
--- a/Engine/Plugins/FX/Niagara/Source/Niagara/Private/NiagaraScriptExecutionContext.cpp
+++ b/Engine/Plugins/FX/Niagara/Source/Niagara/Private/NiagaraScriptExecutionContext.cpp
@@ -449,12 +449,17 @@ void FNiagaraComputeExecutionContext::Reset(NiagaraEmitterInstanceBatcher* Batch
 	);
 }
 
-void FNiagaraComputeExecutionContext::InitParams(UNiagaraScript* InGPUComputeScript, ENiagaraSimTarget InSimTarget, const FString& InDebugSimName)
+void FNiagaraComputeExecutionContext::InitParams(UNiagaraScript* InGPUComputeScript, ENiagaraSimTarget InSimTarget, const FString& InDebugSimName, const int32 InMaxUpdateIterations, const TSet<uint32> InSpawnStages)
 {
 	DebugSimName = InDebugSimName;
 	GPUScript = InGPUComputeScript;
 	CombinedParamStore.InitFromOwningContext(InGPUComputeScript, InSimTarget, true);
+	MaxUpdateIterations = InMaxUpdateIterations;
+	SpawnStages.Empty();
 
+	SpawnStages.Append(InSpawnStages);
+	
+	
 #if DO_CHECK
 	FNiagaraShader *Shader = InGPUComputeScript->GetRenderThreadScript()->GetShaderGameThread();
 	DIParamInfo.Empty();
diff --git a/Engine/Plugins/FX/Niagara/Source/Niagara/Public/NiagaraComponent.h b/Engine/Plugins/FX/Niagara/Source/Niagara/Public/NiagaraComponent.h
--- a/Engine/Plugins/FX/Niagara/Source/Niagara/Public/NiagaraComponent.h
+++ b/Engine/Plugins/FX/Niagara/Source/Niagara/Public/NiagaraComponent.h
@@ -328,6 +328,9 @@ public:
 	UFUNCTION(BlueprintCallable, Category = Niagara)
 	bool IsPaused()const;
 
+	UFUNCTION(BlueprintCallable, Category = Niagara)
+	UNiagaraDataInterface * GetDataInterface(const FString &Name);
+
 	//~ Begin UObject Interface.
 	virtual void PostLoad();
 #if WITH_EDITOR
diff --git a/Engine/Plugins/FX/Niagara/Source/Niagara/Public/NiagaraRenderer.h b/Engine/Plugins/FX/Niagara/Source/Niagara/Public/NiagaraRenderer.h
--- a/Engine/Plugins/FX/Niagara/Source/Niagara/Public/NiagaraRenderer.h
+++ b/Engine/Plugins/FX/Niagara/Source/Niagara/Public/NiagaraRenderer.h
@@ -114,10 +114,10 @@ public:
 	virtual void GetDynamicRayTracingInstances(FRayTracingMaterialGatheringContext& Context, TArray<FRayTracingInstance>& OutRayTracingInstances, const FNiagaraSceneProxy* Proxy) {}
 #endif
 
-	static FRWBuffer& GetDummyFloatBuffer(); 
-	static FRWBuffer& GetDummyFloat4Buffer();
-	static FRWBuffer& GetDummyIntBuffer();
-	static FRWBuffer& GetDummyUIntBuffer();
+	NIAGARA_API static FRWBuffer& GetDummyFloatBuffer();
+	NIAGARA_API static FRWBuffer& GetDummyFloat4Buffer();
+	NIAGARA_API static FRWBuffer& GetDummyIntBuffer();
+	NIAGARA_API static FRWBuffer& GetDummyUIntBuffer();
 	
 	FORCEINLINE ENiagaraSimTarget GetSimTarget() const { return SimTarget; }
 
diff --git a/Engine/Plugins/FX/Niagara/Source/NiagaraCore/Private/NiagaraCustomVersion.cpp b/Engine/Plugins/FX/Niagara/Source/NiagaraCore/Private/NiagaraCustomVersion.cpp
--- a/Engine/Plugins/FX/Niagara/Source/NiagaraCore/Private/NiagaraCustomVersion.cpp
+++ b/Engine/Plugins/FX/Niagara/Source/NiagaraCore/Private/NiagaraCustomVersion.cpp
@@ -8,4 +8,4 @@ const FGuid FNiagaraCustomVersion::GUID(0xFCF57AFA, 0x50764283, 0xB9A9E658, 0xFF
 // Register the custom version with core
 FCustomVersionRegistration GRegisterNiagaraCustomVersion(FNiagaraCustomVersion::GUID, FNiagaraCustomVersion::LatestVersion, TEXT("NiagaraVer"));
 
-const FGuid FNiagaraCustomVersion::LatestScriptCompileVersion(0x784CB4D0, 0x1529422B, 0xB6880B12, 0xFB5A1E3F);
\ No newline at end of file
+const FGuid FNiagaraCustomVersion::LatestScriptCompileVersion(0xA1C84E62, 0x19E6437F, 0x9B095E4A, 0x2F0F8ADE);
diff --git a/Engine/Plugins/FX/Niagara/Source/NiagaraCore/Public/NiagaraDataInterfaceBase.h b/Engine/Plugins/FX/Niagara/Source/NiagaraCore/Public/NiagaraDataInterfaceBase.h
--- a/Engine/Plugins/FX/Niagara/Source/NiagaraCore/Public/NiagaraDataInterfaceBase.h
+++ b/Engine/Plugins/FX/Niagara/Source/NiagaraCore/Public/NiagaraDataInterfaceBase.h
@@ -22,6 +22,9 @@ struct FNiagaraDataInterfaceSetArgs
 	FNiagaraDataInterfaceProxy* DataInterface;
 	FGuid SystemInstance;
 	const NiagaraEmitterInstanceBatcher* Batcher;
+	uint32 ShaderStageIndex;
+	bool IsOutputStage;
+	bool IsIterationStage;
 };
 
 /**
diff --git a/Engine/Plugins/FX/Niagara/Source/NiagaraEditor/Private/EdGraphSchema_Niagara.cpp b/Engine/Plugins/FX/Niagara/Source/NiagaraEditor/Private/EdGraphSchema_Niagara.cpp
--- a/Engine/Plugins/FX/Niagara/Source/NiagaraEditor/Private/EdGraphSchema_Niagara.cpp
+++ b/Engine/Plugins/FX/Niagara/Source/NiagaraEditor/Private/EdGraphSchema_Niagara.cpp
@@ -25,6 +25,7 @@
 #include "NiagaraNodeWriteDataSet.h"
 #include "NiagaraNodeParameterMapGet.h"
 #include "NiagaraNodeParameterMapSet.h"
+#include "NiagaraNodeParameterMapFor.h"
 #include "NiagaraNodeCustomHlsl.h"
 #include "NiagaraNodeOp.h"
 #include "NiagaraNodeConvert.h"
@@ -686,7 +687,7 @@ TArray<TSharedPtr<FNiagaraSchemaAction_NewNode> > UEdGraphSchema_Niagara::GetGra
 		}
 	}
 
-	// Handle parameter map get/set
+	// Handle parameter map get/set/for
 	{
 		FText MenuCat = FText::FromString("Parameter Map");
 		{
@@ -701,6 +702,14 @@ TArray<TSharedPtr<FNiagaraSchemaAction_NewNode> > UEdGraphSchema_Niagara::GetGra
 			UNiagaraNodeParameterMapSet* BaseNode = NewObject<UNiagaraNodeParameterMapSet>(OwnerOfTemporaries);
 			Action->NodeTemplate = BaseNode;
 		}
+		static const auto UseShaderStagesCVar = IConsoleManager::Get().FindConsoleVariable(TEXT("fx.UseShaderStages"));
+		if (UseShaderStagesCVar->GetInt() == 1)
+		{
+			FString Name = TEXT("Parameter Map For");
+			TSharedPtr<FNiagaraSchemaAction_NewNode> Action = AddNewNodeAction(NewActions, MenuCat, FText::FromString(Name), *Name, FText::GetEmpty());
+			UNiagaraNodeParameterMapFor* BaseNode = NewObject<UNiagaraNodeParameterMapFor>(OwnerOfTemporaries);
+			Action->NodeTemplate = BaseNode;
+		}
 	}
 
 	// Handle comment nodes
@@ -842,12 +851,11 @@ TArray<TSharedPtr<FNiagaraSchemaAction_NewNode> > UEdGraphSchema_Niagara::GetGra
 		}
 	}
 
-	const FText MenuCat = LOCTEXT("NiagaraLogicMenuCat", "Logic");
+	
 	{
+		const FText LogicMenuCat = LOCTEXT("NiagaraLogicMenuCat", "Logic");
 		const FText MenuDesc = LOCTEXT("If", "If");
-
-		TSharedPtr<FNiagaraSchemaAction_NewNode> Action = AddNewNodeAction(NewActions, MenuCat, MenuDesc, TEXT("If"), FText::GetEmpty());
-
+		TSharedPtr<FNiagaraSchemaAction_NewNode> Action = AddNewNodeAction(NewActions, LogicMenuCat, MenuDesc, TEXT("If"), FText::GetEmpty());
 		UNiagaraNodeIf* IfNode = NewObject<UNiagaraNodeIf>(OwnerOfTemporaries);
 		Action->NodeTemplate = IfNode;
 	}
diff --git a/Engine/Plugins/FX/Niagara/Source/NiagaraEditor/Private/NiagaraGraph.cpp b/Engine/Plugins/FX/Niagara/Source/NiagaraEditor/Private/NiagaraGraph.cpp
--- a/Engine/Plugins/FX/Niagara/Source/NiagaraEditor/Private/NiagaraGraph.cpp
+++ b/Engine/Plugins/FX/Niagara/Source/NiagaraEditor/Private/NiagaraGraph.cpp
@@ -1574,7 +1574,6 @@ void UNiagaraGraph::RefreshParameterReferences() const
 		for (int32 Index = 0; Index < History.VariablesWithOriginalAliasesIntact.Num(); Index++)
 		{
 			const FNiagaraVariable& Parameter = History.VariablesWithOriginalAliasesIntact[Index];
-
 			for (const UEdGraphPin* WritePin : History.PerVariableWriteHistory[Index])
 			{
 				AddParameterReference(Parameter, WritePin);
@@ -1621,6 +1620,34 @@ void UNiagaraGraph::RefreshParameterReferences() const
 	{
 		ParameterToReferencesMap.Remove(UnreferencedParameterToRemove);
 	}
+	
+	static const auto UseShaderStagesCVar = IConsoleManager::Get().FindConsoleVariable(TEXT("fx.UseShaderStages"));
+	if (UseShaderStagesCVar->GetInt() == 1)
+	{
+		// Add the array indices to the parameters
+		// When a particle attribute is created we need access the corresponding RegisterIdx if we want to 
+		// query this attribute at a different location inside the InputData buffer. This index must be 
+		// available as well inside the UI if we want to pass it to nodes. It is why we are adding them automatically 
+		// to the ParameterToReferencesMap.
+		TArray<FString> RegisterNames;
+		for (auto& ParameterEntry : ParameterToReferencesMap)
+		{
+			const FNiagaraVariable& NiagaraVariable = ParameterEntry.Key;
+			if (FNiagaraParameterMapHistory::IsAttribute(NiagaraVariable))
+			{
+				const FString VariableName = FHlslNiagaraTranslator::GetSanitizedSymbolName(NiagaraVariable.GetName().ToString());
+				RegisterNames.Add(VariableName.Replace(PARAM_MAP_ATTRIBUTE_STR, PARAM_MAP_INDICES_STR));
+			}
+		}
+		for (const FString& RegisterName : RegisterNames)
+		{
+			FNiagaraVariable Parameter = FNiagaraVariable(FNiagaraTypeDefinition::GetIntDef(), *RegisterName);
+			if (ParameterToReferencesMap.Find(Parameter) == nullptr)
+			{
+				ParameterToReferencesMap.Add(Parameter, false);
+			}
+		}
+	}
 
 	bParameterReferenceRefreshPending = false;
 }
