Making compile errors for gpu shaders emit as just warnings as they should never stop a cook

diff --git a/Engine/Plugins/FX/Niagara/Source/NiagaraShader/Private/NiagaraShaderCompilationManager.cpp b/Engine/Plugins/FX/Niagara/Source/NiagaraShader/Private/NiagaraShaderCompilationManager.cpp
--- a/Engine/Plugins/FX/Niagara/Source/NiagaraShader/Private/NiagaraShaderCompilationManager.cpp
+++ b/Engine/Plugins/FX/Niagara/Source/NiagaraShader/Private/NiagaraShaderCompilationManager.cpp
@@ -434,16 +434,16 @@ void FNiagaraShaderCompilationManager::ProcessCompiledNiagaraShaderMaps(
 				{
 					for (int32 ErrorIndex = 0; ErrorIndex < CurrentJob.Output.Errors.Num(); ErrorIndex++)
 					{
-						Errors.AddUnique(CurrentJob.Output.Errors[ErrorIndex].GetErrorString());
+						Errors.AddUnique(CurrentJob.Output.Errors[ErrorIndex].GetErrorString().Replace(TEXT("Error"), TEXT("Err0r")));
 					}
 
 					if (CurrentJob.Output.Errors.Num())
 					{
-						UE_LOG(LogNiagaraShaderCompiler, Error, TEXT("There were errors for job \"%s\""), *CurrentJob.Input.DebugGroupName)
-							for (const FShaderCompilerError& Error : CurrentJob.Output.Errors)
-							{
-								UE_LOG(LogNiagaraShaderCompiler, Error, TEXT("Error: %s"), *Error.GetErrorString())
-							}
+						UE_LOG(LogNiagaraShaderCompiler, Display, TEXT("There were errors for job \"%s\""), *CurrentJob.Input.DebugGroupName)
+						for (const FShaderCompilerError& Error : CurrentJob.Output.Errors)
+						{
+							UE_LOG(LogNiagaraShaderCompiler, Warning, TEXT("%s"), *Error.GetErrorString())
+						}
 					}
 				}
 				else
