Further tweaking of logging to prevent errors.

diff --git a/Engine/Plugins/FX/Niagara/Source/NiagaraShader/Private/NiagaraShaderCompilationManager.cpp b/Engine/Plugins/FX/Niagara/Source/NiagaraShader/Private/NiagaraShaderCompilationManager.cpp
--- a/Engine/Plugins/FX/Niagara/Source/NiagaraShader/Private/NiagaraShaderCompilationManager.cpp
+++ b/Engine/Plugins/FX/Niagara/Source/NiagaraShader/Private/NiagaraShaderCompilationManager.cpp
@@ -253,7 +253,7 @@ void FNiagaraShaderCompilationManager::RunCompileJobs()
 				}
 				else
 				{
-					UE_LOG(LogNiagaraShaderCompiler, Log, TEXT("ERROR: GPU shader compile failed! Id %d"), CurrentJob.Id);
+					UE_LOG(LogNiagaraShaderCompiler, Warning, TEXT("GPU shader compile failed! Id %d"), CurrentJob.Id);
 				}
 
 				CurrentWorkerInfo.bComplete = true;
@@ -439,11 +439,14 @@ void FNiagaraShaderCompilationManager::ProcessCompiledNiagaraShaderMaps(
 
 					if (CurrentJob.Output.Errors.Num())
 					{
-						UE_LOG(LogNiagaraShaderCompiler, Display, TEXT("There were errors for job \"%s\""), *CurrentJob.Input.DebugGroupName)
+						
+						UE_LOG(LogNiagaraShaderCompiler, Display, TEXT("There were issues for job \"%s\""), *CurrentJob.Input.DebugGroupName)
+						
 						for (const FShaderCompilerError& Error : CurrentJob.Output.Errors)
 						{
 							UE_LOG(LogNiagaraShaderCompiler, Warning, TEXT("%s"), *Error.GetErrorString())
 						}
+						
 					}
 				}
 				else
