Niagara - Fix constant recompiles due to dynamic input name changes due to merging.

diff --git a/Engine/Plugins/FX/Niagara/Source/NiagaraEditor/Private/NiagaraGraph.cpp b/Engine/Plugins/FX/Niagara/Source/NiagaraEditor/Private/NiagaraGraph.cpp
--- a/Engine/Plugins/FX/Niagara/Source/NiagaraEditor/Private/NiagaraGraph.cpp
+++ b/Engine/Plugins/FX/Niagara/Source/NiagaraEditor/Private/NiagaraGraph.cpp
@@ -1036,8 +1036,7 @@ void UNiagaraGraph::RebuildCachedCompileIds(bool bForce)
 		FSHA1 HashState;
 		for (UNiagaraNode* Node : NewUsageCache[i].Traversal)
 		{
-			FGuid Guid = Node->GetChangeId();
-			HashState.Update((const uint8*)&Guid, sizeof(FGuid));
+			Node->UpdateCompileHashForNode(HashState);
 		}
 		HashState.Final();
 
@@ -1125,8 +1124,7 @@ void UNiagaraGraph::RebuildCachedCompileIds(bool bForce)
 		FSHA1 HashState;
 		for (UNiagaraNode* Node : GpuUsageInfo.Traversal)
 		{
-			FGuid Guid = Node->GetChangeId();
-			HashState.Update((const uint8*)&Guid, sizeof(FGuid));
+			Node->UpdateCompileHashForNode(HashState);
 		}
 		HashState.Final();
 
diff --git a/Engine/Plugins/FX/Niagara/Source/NiagaraEditor/Private/NiagaraNode.cpp b/Engine/Plugins/FX/Niagara/Source/NiagaraEditor/Private/NiagaraNode.cpp
--- a/Engine/Plugins/FX/Niagara/Source/NiagaraEditor/Private/NiagaraNode.cpp
+++ b/Engine/Plugins/FX/Niagara/Source/NiagaraEditor/Private/NiagaraNode.cpp
@@ -6,6 +6,7 @@
 #include "NiagaraHlslTranslator.h"
 #include "GraphEditAction.h"
 #include "SNiagaraGraphNode.h"
+#include "Misc/SecureHash.h"
 
 #define LOCTEXT_NAMESPACE "NiagaraNode"
 
@@ -602,5 +603,10 @@ UNiagaraNode::FOnNodeVisualsChanged& UNiagaraNode::OnVisualsChanged()
 	return VisualsChangedDelegate;
 }
 
+void UNiagaraNode::UpdateCompileHashForNode(FSHA1& HashState) const
+{
+	HashState.Update((const uint8*)&NodeGuid, sizeof(FGuid));
+}
+
 
 #undef LOCTEXT_NAMESPACE
\ No newline at end of file
diff --git a/Engine/Plugins/FX/Niagara/Source/NiagaraEditor/Private/NiagaraNodeFunctionCall.cpp b/Engine/Plugins/FX/Niagara/Source/NiagaraEditor/Private/NiagaraNodeFunctionCall.cpp
--- a/Engine/Plugins/FX/Niagara/Source/NiagaraEditor/Private/NiagaraNodeFunctionCall.cpp
+++ b/Engine/Plugins/FX/Niagara/Source/NiagaraEditor/Private/NiagaraNodeFunctionCall.cpp
@@ -18,6 +18,7 @@
 #include "NiagaraCustomVersion.h"
 #include "NiagaraDataInterfaceSkeletalMesh.h"
 #include "SNiagaraGraphNodeFunctionCallWithSpecifiers.h"
+#include "Misc/SecureHash.h"
 
 #define LOCTEXT_NAMESPACE "NiagaraNodeFunctionCall"
 
@@ -705,6 +706,12 @@ void UNiagaraNodeFunctionCall::GatherExternalDependencyIDs(ENiagaraScriptUsage I
 	}
 }
 
+void UNiagaraNodeFunctionCall::UpdateCompileHashForNode(FSHA1& HashState) const
+{
+	Super::UpdateCompileHashForNode(HashState);
+	HashState.UpdateWithString(*GetFunctionName(), GetFunctionName().Len());
+}
+
 bool UNiagaraNodeFunctionCall::ScriptIsValid() const
 {
 	if (FunctionScript != nullptr)
diff --git a/Engine/Plugins/FX/Niagara/Source/NiagaraEditor/Private/NiagaraScriptMergeManager.cpp b/Engine/Plugins/FX/Niagara/Source/NiagaraEditor/Private/NiagaraScriptMergeManager.cpp
--- a/Engine/Plugins/FX/Niagara/Source/NiagaraEditor/Private/NiagaraScriptMergeManager.cpp
+++ b/Engine/Plugins/FX/Niagara/Source/NiagaraEditor/Private/NiagaraScriptMergeManager.cpp
@@ -1810,10 +1810,6 @@ FNiagaraScriptMergeManager::FApplyDiffResults FNiagaraScriptMergeManager::AddInp
 					UNiagaraNodeFunctionCall* DynamicInputFunctionCall;
 					FNiagaraStackGraphUtilities::SetDynamicInputForFunctionInput(InputOverridePin, OverrideToAdd->GetDynamicValueFunction()->GetFunctionCallNode()->FunctionScript,
 						DynamicInputFunctionCall, OverrideToAdd->GetOverrideNodeId(), OverrideToAdd->GetDynamicValueFunction()->GetFunctionCallNode()->GetFunctionName());
-					if (DynamicInputFunctionCall->GetFunctionName() != OverrideToAdd->GetDynamicValueFunction()->GetFunctionCallNode()->GetFunctionName())
-					{
-						DynamicInputFunctionCall->NodeGuid = FGuid::NewGuid();
-					}
 					for (TSharedRef<FNiagaraStackFunctionInputOverrideMergeAdapter> DynamicInputInputOverride : OverrideToAdd->GetDynamicValueFunction()->GetInputOverrides())
 					{
 						FApplyDiffResults AddResults = AddInputOverride(UniqueEmitterName, OwningScript, *DynamicInputFunctionCall, DynamicInputInputOverride);
diff --git a/Engine/Plugins/FX/Niagara/Source/NiagaraEditor/Public/NiagaraNode.h b/Engine/Plugins/FX/Niagara/Source/NiagaraEditor/Public/NiagaraNode.h
--- a/Engine/Plugins/FX/Niagara/Source/NiagaraEditor/Public/NiagaraNode.h
+++ b/Engine/Plugins/FX/Niagara/Source/NiagaraEditor/Public/NiagaraNode.h
@@ -14,6 +14,7 @@
 class UEdGraphPin;
 class INiagaraCompiler;
 struct FNiagaraGraphFunctionAliasContext;
+class FSHA1;
 
 UCLASS()
 class NIAGARAEDITOR_API UNiagaraNode : public UEdGraphNode
@@ -124,6 +125,8 @@ public:
 
 	virtual void AppendFunctionAliasForContext(const FNiagaraGraphFunctionAliasContext& InFunctionAliasContext, FString& InOutFunctionAlias) { };
 
+	virtual void UpdateCompileHashForNode(FSHA1& HashState) const;
+
 protected:
 	virtual int32 CompileInputPin(class FHlslNiagaraTranslator *Translator, UEdGraphPin* Pin);
 	virtual bool IsValidPinToCompile(UEdGraphPin* Pin) const { return true; }; 
diff --git a/Engine/Plugins/FX/Niagara/Source/NiagaraEditor/Public/NiagaraNodeFunctionCall.h b/Engine/Plugins/FX/Niagara/Source/NiagaraEditor/Public/NiagaraNodeFunctionCall.h
--- a/Engine/Plugins/FX/Niagara/Source/NiagaraEditor/Public/NiagaraNodeFunctionCall.h
+++ b/Engine/Plugins/FX/Niagara/Source/NiagaraEditor/Public/NiagaraNodeFunctionCall.h
@@ -92,6 +92,7 @@ public:
 	virtual bool CanAddToGraph(UNiagaraGraph* TargetGraph, FString& OutErrorMsg) const override;
 	virtual void SubsumeExternalDependencies(TMap<const UObject*, UObject*>& ExistingConversions) override;
 	virtual void GatherExternalDependencyIDs(ENiagaraScriptUsage InMasterUsage, const FGuid& InMasterUsageId, TArray<FNiagaraCompileHash>& InReferencedCompileHashes, TArray<FGuid>& InReferencedIDs, TArray<UObject*>& InReferencedObjs) const override;
+	virtual void UpdateCompileHashForNode(FSHA1& HashState) const override;
 	//End UNiagaraNode interface
 
 	//~ Begin EdGraphNode Interface
