Niagara - Add a utility method to dump compile ids for an asset.

diff --git a/Engine/Plugins/FX/Niagara/Source/NiagaraEditor/Private/NiagaraEditorModule.cpp b/Engine/Plugins/FX/Niagara/Source/NiagaraEditor/Private/NiagaraEditorModule.cpp
--- a/Engine/Plugins/FX/Niagara/Source/NiagaraEditor/Private/NiagaraEditorModule.cpp
+++ b/Engine/Plugins/FX/Niagara/Source/NiagaraEditor/Private/NiagaraEditorModule.cpp
@@ -33,6 +33,7 @@
 #include "SNiagaraGraphPinNumeric.h"
 #include "SNiagaraGraphPinAdd.h"
 #include "NiagaraNodeConvert.h"
+#include "NiagaraNodeAssignment.h"
 #include "EdGraphSchema_Niagara.h"
 #include "TypeEditorUtilities/NiagaraFloatTypeEditorUtilities.h"
 #include "TypeEditorUtilities/NiagaraIntegerTypeEditorUtilities.h"
@@ -216,6 +217,7 @@ private:
 FNiagaraEditorModule::FNiagaraEditorModule() 
 	: SequencerSettings(nullptr)
 	, TestCompileScriptCommand(nullptr)
+	, DumpCompileIdDataForAssetCommand(nullptr)
 {
 }
 
@@ -455,6 +457,114 @@ void UpgradeAllNiagaraAssets()
 	GWarn->EndSlowTask();
 }
 
+void MakeIndent(int32 IndentLevel, FString& OutIndentString)
+{
+	OutIndentString.Reserve(IndentLevel * 2);
+	int32 InsertStart = OutIndentString.Len();
+	for (int32 i = 0; i < IndentLevel * 2; i++)
+	{
+		OutIndentString.AppendChar(TCHAR(' '));
+	}
+}
+
+void DumpCompileIdDataForScript(UNiagaraScript* Script, int32 IndentLevel, FString& Dump)
+{
+	FString Indent;
+	MakeIndent(IndentLevel, Indent);
+	Dump.Append(FString::Printf(TEXT("%sScript: %s\n"), *Indent, *Script->GetPathName()));
+	UNiagaraScriptSource* ScriptSource = Cast<UNiagaraScriptSource>(Script->GetSource());
+	TArray<UNiagaraNode*> Nodes;
+	ScriptSource->NodeGraph->GetNodesOfClass<UNiagaraNode>(Nodes);
+	for (UNiagaraNode* Node : Nodes)
+	{
+		Dump.Append(FString::Printf(TEXT("%s%s - %s-%s\n"), *Indent, *Node->GetFullName(), *Node->NodeGuid.ToString(EGuidFormats::Digits), *Node->GetChangeId().ToString(EGuidFormats::Digits)));
+		UNiagaraNodeFunctionCall* FunctionCallNode = Cast<UNiagaraNodeFunctionCall>(Node);
+		if (FunctionCallNode != nullptr)
+		{
+			UNiagaraNodeAssignment* AssignmentNode = Cast<UNiagaraNodeAssignment>(FunctionCallNode);
+			if (AssignmentNode != nullptr)
+			{
+				Dump.Append(FString::Printf(TEXT("%sAssignment Node: %s\n"), *Indent, *FunctionCallNode->GetFunctionName()));
+				for (const FNiagaraVariable& AssignmentTarget : AssignmentNode->GetAssignmentTargets())
+				{
+					Dump.Append(FString::Printf(TEXT("%s  Assignment Target: %s - %s\n"), *Indent, *AssignmentTarget.GetName().ToString(), *AssignmentTarget.GetType().GetName()));
+				}
+			}
+			else if (FunctionCallNode->FunctionScript != nullptr)
+			{
+				Dump.Append(FString::Printf(TEXT("%sFunction Call: %s\n"), *Indent, *FunctionCallNode->GetFunctionName()));
+				DumpCompileIdDataForScript(FunctionCallNode->FunctionScript, IndentLevel + 1, Dump);
+			}
+		}
+	}
+}
+
+void DumpCompileIdDataForEmitter(UNiagaraEmitter* Emitter, int32 IndentLevel, FString& Dump)
+{
+
+	FString Indent;
+	MakeIndent(IndentLevel, Indent);
+	Dump.Append(FString::Printf(TEXT("%sEmitter: %s\n"), *Indent, *Emitter->GetUniqueEmitterName()));
+
+	TArray<UNiagaraScript*> Scripts;
+	Emitter->GetScripts(Scripts, false);
+	for (UNiagaraScript* Script : Scripts)
+	{
+		DumpCompileIdDataForScript(Script, IndentLevel + 1, Dump);
+	}
+}
+
+void DumpCompileIdDataForSystem(UNiagaraSystem* System, FString& Dump)
+{
+	Dump.Append(FString::Printf(TEXT("\nSystem %s\n"), *System->GetPathName()));
+	DumpCompileIdDataForScript(System->GetSystemSpawnScript(), 1, Dump);
+	DumpCompileIdDataForScript(System->GetSystemUpdateScript(), 1, Dump);
+	for (const FNiagaraEmitterHandle& EmitterHandle : System->GetEmitterHandles())
+	{
+		DumpCompileIdDataForEmitter(EmitterHandle.GetInstance(), 1, Dump);
+	}
+}
+
+void DumpCompileIdDataForAsset(const TArray<FString>& Arguments)
+{
+	if (Arguments.Num() > 0)
+	{
+		const FAssetRegistryModule& AssetRegistryModule = FModuleManager::LoadModuleChecked<FAssetRegistryModule>(TEXT("AssetRegistry"));
+		FAssetData SystemAsset = AssetRegistryModule.Get().GetAssetByObjectPath(*Arguments[0]);
+		if (SystemAsset.IsValid() == false)
+		{
+			TArray<FAssetData> AssetsInPackage;
+			AssetRegistryModule.Get().GetAssetsByPackageName(*Arguments[0], AssetsInPackage);
+			if (AssetsInPackage.Num() == 1)
+			{
+				SystemAsset = AssetsInPackage[0];
+			}
+		}
+		if (SystemAsset.IsValid())
+		{
+			UNiagaraSystem* System = Cast<UNiagaraSystem>(SystemAsset.GetAsset());
+			if (System != nullptr)
+			{
+				FString Dump;
+				DumpCompileIdDataForSystem(System, Dump);
+				UE_LOG(LogNiagaraEditor, Log, TEXT("%s"), *Dump);
+			}
+			else
+			{
+				UE_LOG(LogNiagaraEditor, Warning, TEXT("Could not load system asset for argument: %s"), *Arguments[0]);
+			}
+		}
+		else
+		{
+			UE_LOG(LogNiagaraEditor, Warning, TEXT("Could not find asset for argument: %s"), *Arguments[0]);
+		}
+	}
+	else
+	{
+		UE_LOG(LogNiagaraEditor, Warning, TEXT("Command required an asset reference to be passed in."));
+	}
+}
+
 class FNiagaraSystemBoolParameterTrackEditor : public FNiagaraSystemParameterTrackEditor<UMovieSceneNiagaraBoolParameterTrack, UMovieSceneBoolSection>
 {
 	virtual TSharedRef<ISequencerSection> MakeSectionInterface(UMovieSceneSection& SectionObject, UMovieSceneTrack& Track, FGuid ObjectBinding) override
@@ -693,6 +803,11 @@ void FNiagaraEditorModule::StartupModule()
 		TEXT("Loads all Niagara assets and preforms any data upgrade processes required. This may mark multiple assets dirty for re-saving."),
 		FConsoleCommandDelegate::CreateStatic(&UpgradeAllNiagaraAssets));
 
+	DumpCompileIdDataForAssetCommand = IConsoleManager::Get().RegisterConsoleCommand(
+		TEXT("fx.DumpCompileIdDataForAsset"),
+		TEXT("Dumps data relevant to generating the compile id for an asset."),
+		FConsoleCommandWithArgsDelegate::CreateStatic(&DumpCompileIdDataForAsset));
+
 	if (GIsEditor)
 	{
 		UThumbnailManager::Get().RegisterCustomRenderer(UNiagaraEmitter::StaticClass(), UNiagaraEmitterThumbnailRenderer::StaticClass());
@@ -792,6 +907,12 @@ void FNiagaraEditorModule::ShutdownModule()
 		IConsoleManager::Get().UnregisterConsoleObject(PreventAllSystemRecompilesCommand);
 	}
 
+	if (DumpCompileIdDataForAssetCommand != nullptr)
+	{
+		IConsoleManager::Get().UnregisterConsoleObject(DumpCompileIdDataForAssetCommand);
+		DumpCompileIdDataForAssetCommand = nullptr;
+	}
+
 	if (UObjectInitialized() && GIsEditor)
 	{
 		UThumbnailManager::Get().UnregisterCustomRenderer(UNiagaraEmitter::StaticClass());
diff --git a/Engine/Plugins/FX/Niagara/Source/NiagaraEditor/Public/NiagaraEditorModule.h b/Engine/Plugins/FX/Niagara/Source/NiagaraEditor/Public/NiagaraEditorModule.h
--- a/Engine/Plugins/FX/Niagara/Source/NiagaraEditor/Public/NiagaraEditorModule.h
+++ b/Engine/Plugins/FX/Niagara/Source/NiagaraEditor/Public/NiagaraEditorModule.h
@@ -149,6 +149,7 @@ private:
 	IConsoleCommand* PreventSystemRecompileCommand;
 	IConsoleCommand* PreventAllSystemRecompilesCommand;
 	IConsoleCommand* UpgradeAllNiagaraAssetsCommand;
+	IConsoleCommand* DumpCompileIdDataForAssetCommand;
 
 	FOnCheckScriptToolkitsShouldFocusGraphElement OnCheckScriptToolkitsShouldFocusGraphElement;
 };
