Niagara - Prevent a crash on shutdown when the spreadsheet view is open.

diff --git a/Engine/Plugins/FX/Niagara/Source/NiagaraEditor/Private/Widgets/SNiagaraSpreadsheetView.cpp b/Engine/Plugins/FX/Niagara/Source/NiagaraEditor/Private/Widgets/SNiagaraSpreadsheetView.cpp
--- a/Engine/Plugins/FX/Niagara/Source/NiagaraEditor/Private/Widgets/SNiagaraSpreadsheetView.cpp
+++ b/Engine/Plugins/FX/Niagara/Source/NiagaraEditor/Private/Widgets/SNiagaraSpreadsheetView.cpp
@@ -618,7 +618,10 @@ SNiagaraSpreadsheetView::~SNiagaraSpreadsheetView()
 {
 	if (SystemViewModel.IsValid())
 	{
-		SystemViewModel->GetSelectionViewModel()->OnSelectionChanged().RemoveAll(this);
+		if (SystemViewModel->GetSelectionViewModel() != nullptr)
+		{
+			SystemViewModel->GetSelectionViewModel()->OnSelectionChanged().RemoveAll(this);
+		}
 		SystemViewModel->OnPostSequencerTimeChanged().RemoveAll(this);
 	}
 }
