Copying

//UE4/Dev-Physics/Engine/Plugins/FX/Niagara/Source/Niagara/Private/NiagaraDataInterfaceGrid2DCollection.cpp

to //UE4/Main/Engine/Plugins/FX/Niagara/Source/Niagara/Private/NiagaraDataInterfaceGrid2DCollection.cpp

Fix CIS -- ensure changed to check #rb none

diff --git a/Engine/Plugins/FX/Niagara/Source/Niagara/Private/NiagaraDataInterfaceGrid2DCollection.cpp b/Engine/Plugins/FX/Niagara/Source/Niagara/Private/NiagaraDataInterfaceGrid2DCollection.cpp
--- a/Engine/Plugins/FX/Niagara/Source/Niagara/Private/NiagaraDataInterfaceGrid2DCollection.cpp
+++ b/Engine/Plugins/FX/Niagara/Source/Niagara/Private/NiagaraDataInterfaceGrid2DCollection.cpp
@@ -58,7 +58,7 @@ struct FNiagaraDataInterfaceParametersCS_Grid2DCollection : public FNiagaraDataI
 		FNiagaraDataInterfaceProxyGrid2DCollection* VFDI = static_cast<FNiagaraDataInterfaceProxyGrid2DCollection*>(Context.DataInterface);
 		
 		Grid2DCollectionRWInstanceData* ProxyData = VFDI->SystemInstancesToProxyData.Find(Context.SystemInstance);
-		ensure(ProxyData);
+		check(ProxyData);
 
 		int NumCellsTmp[2];
 		NumCellsTmp[0] = ProxyData->NumCellsX;
