Niagara - Fix a typo from CL 7745154

diff --git a/Engine/Plugins/FX/Niagara/Source/NiagaraEditor/Private/NiagaraNode.cpp b/Engine/Plugins/FX/Niagara/Source/NiagaraEditor/Private/NiagaraNode.cpp
--- a/Engine/Plugins/FX/Niagara/Source/NiagaraEditor/Private/NiagaraNode.cpp
+++ b/Engine/Plugins/FX/Niagara/Source/NiagaraEditor/Private/NiagaraNode.cpp
@@ -605,7 +605,7 @@ UNiagaraNode::FOnNodeVisualsChanged& UNiagaraNode::OnVisualsChanged()
 
 void UNiagaraNode::UpdateCompileHashForNode(FSHA1& HashState) const
 {
-	HashState.Update((const uint8*)&NodeGuid, sizeof(FGuid));
+	HashState.Update((const uint8*)&ChangeId, sizeof(FGuid));
 }
 
 
