CIS Fixes

diff --git a/Engine/Plugins/Experimental/HairStrands/Source/HairStrandsEditor/Private/HairStrandsFactory.cpp b/Engine/Plugins/Experimental/HairStrands/Source/HairStrandsEditor/Private/HairStrandsFactory.cpp
--- a/Engine/Plugins/Experimental/HairStrands/Source/HairStrandsEditor/Private/HairStrandsFactory.cpp
+++ b/Engine/Plugins/Experimental/HairStrands/Source/HairStrandsEditor/Private/HairStrandsFactory.cpp
@@ -57,6 +57,7 @@ UObject* UHairStrandsFactory::FactoryCreateNew(UClass* InClass, UObject* InParen
 		ExistingAsset->ReleaseResource();
 	}
 	UHairStrandsAsset* CurrentAsset = NewObject<UHairStrandsAsset>(InParent, InClass, InName, Flags);
+	check(CurrentAsset);
 
 	const FString Filter(TEXT("Hair Strands Files (*.hair,*.fbx,*.abc)|*.hair;*.fbx;*.abc"));
 
@@ -93,6 +94,7 @@ UObject* UHairStrandsFactory::FactoryCreateFile(UClass * InClass, UObject * InPa
 	}
 
 	UHairStrandsAsset* CurrentAsset = NewObject<UHairStrandsAsset>(InParent, InClass, InName, Flags);
+	check(CurrentAsset);
 
 	CurrentAsset->FilePath = Filename;
 	BuildHairStrands(CurrentAsset->FilePath, CurrentAsset->StrandsDatas);
diff --git a/Engine/Plugins/FX/Niagara/Source/Niagara/Private/NiagaraDataInterfaceRW.cpp b/Engine/Plugins/FX/Niagara/Source/Niagara/Private/NiagaraDataInterfaceRW.cpp
--- a/Engine/Plugins/FX/Niagara/Source/Niagara/Private/NiagaraDataInterfaceRW.cpp
+++ b/Engine/Plugins/FX/Niagara/Source/Niagara/Private/NiagaraDataInterfaceRW.cpp
@@ -146,7 +146,6 @@ void UNiagaraDataInterfaceGrid3D::GetVMExternalFunction(const FVMExternalFunctio
 	else if (BindingInfo.Name == UnitToIndexFunctionName) { OutFunc = FVMExternalFunction::CreateUObject(this, &UNiagaraDataInterfaceRWBase::EmptyVMFunction); }
 	else if (BindingInfo.Name == IndexToLinearFunctionName) { OutFunc = FVMExternalFunction::CreateUObject(this, &UNiagaraDataInterfaceRWBase::EmptyVMFunction); }
 	else if (BindingInfo.Name == VoxelSizeFunctionName) { OutFunc = FVMExternalFunction::CreateUObject(this, &UNiagaraDataInterfaceRWBase::EmptyVMFunction); }
-	else if (BindingInfo.Name == NumVoxelsFunctionName) { OutFunc = FVMExternalFunction::CreateUObject(this, &UNiagaraDataInterfaceRWBase::EmptyVMFunction); }
 }
 
 bool UNiagaraDataInterfaceGrid3D::Equals(const UNiagaraDataInterface* Other) const
@@ -248,23 +247,6 @@ bool UNiagaraDataInterfaceGrid3D::GetFunctionHLSL(const FName& DefinitionFunctio
 		};
 		OutHLSL += FString::Format(FormatSample, ArgsSample);
 		return true;
-	}	
-	else if (DefinitionFunctionName == NumVoxelsFunctionName)
-	{
-		static const TCHAR *FormatSample = TEXT(R"(
-			void {FunctionName}(out int Out_NumVoxelsX, out int Out_NumVoxelsY, out int Out_NumVoxelsZ)
-			{
-				Out_NumVoxelsX = {NumVoxelsName}.x;
-				Out_NumVoxelsY = {NumVoxelsName}.y;
-				Out_NumVoxelsZ = {NumVoxelsName}.z;
-			}
-		)");
-		TMap<FString, FStringFormatArg> ArgsSample = {
-			{TEXT("FunctionName"), InstanceFunctionName},
-			{TEXT("NumVoxelsName"), NumVoxelsName + ParamInfo.DataInterfaceHLSLSymbol},
-		};
-		OutHLSL += FString::Format(FormatSample, ArgsSample);
-		return true;
 	}
 	else if (DefinitionFunctionName == VoxelSizeFunctionName)
 	{
@@ -474,8 +456,7 @@ void UNiagaraDataInterfaceGrid2D::GetVMExternalFunction(const FVMExternalFunctio
 	else if (BindingInfo.Name == IndexToUnitStaggeredYFunctionName) { OutFunc = FVMExternalFunction::CreateUObject(this, &UNiagaraDataInterfaceRWBase::EmptyVMFunction); }
 	else if (BindingInfo.Name == IndexToLinearFunctionName) { OutFunc = FVMExternalFunction::CreateUObject(this, &UNiagaraDataInterfaceRWBase::EmptyVMFunction); }
 	else if (BindingInfo.Name == LinearToIndexFunctionName) { OutFunc = FVMExternalFunction::CreateUObject(this, &UNiagaraDataInterfaceRWBase::EmptyVMFunction); }
-	else if (BindingInfo.Name == CellSizeFunctionName) { OutFunc = FVMExternalFunction::CreateUObject(this, &UNiagaraDataInterfaceRWBase::EmptyVMFunction); }
-	else if (BindingInfo.Name == NumCellsFunctionName) { OutFunc = FVMExternalFunction::CreateUObject(this, &UNiagaraDataInterfaceRWBase::EmptyVMFunction); }
+	else if (BindingInfo.Name == CellSizeFunctionName) { OutFunc = FVMExternalFunction::CreateUObject(this, &UNiagaraDataInterfaceRWBase::EmptyVMFunction); }	
 }
 
 bool UNiagaraDataInterfaceGrid2D::Equals(const UNiagaraDataInterface* Other) const
@@ -639,22 +620,6 @@ bool UNiagaraDataInterfaceGrid2D::GetFunctionHLSL(const FName& DefinitionFunctio
 		OutHLSL += FString::Format(FormatSample, ArgsSample);
 		return true;
 	}
-	else if (DefinitionFunctionName == NumCellsFunctionName)
-	{
-		static const TCHAR *FormatSample = TEXT(R"(
-			void {FunctionName}(out int Out_NumCellsX, out int Out_NumCellsY)
-			{
-				Out_NumCellsX = {NumCells}.x;
-				Out_NumCellsY = {NumCells}.y;				
-			}
-		)");
-		TMap<FString, FStringFormatArg> ArgsSample = {
-			{TEXT("FunctionName"), InstanceFunctionName},
-			{TEXT("NumCells"), NumCellsName + ParamInfo.DataInterfaceHLSLSymbol},
-		};
-		OutHLSL += FString::Format(FormatSample, ArgsSample);
-		return true;
-	}
 	else if (DefinitionFunctionName == CellSizeFunctionName)
 	{
 		static const TCHAR *FormatSample = TEXT(R"(
diff --git a/Engine/Plugins/FX/Niagara/Source/NiagaraEditor/Private/NiagaraHlslTranslator.cpp b/Engine/Plugins/FX/Niagara/Source/NiagaraEditor/Private/NiagaraHlslTranslator.cpp
--- a/Engine/Plugins/FX/Niagara/Source/NiagaraEditor/Private/NiagaraHlslTranslator.cpp
+++ b/Engine/Plugins/FX/Niagara/Source/NiagaraEditor/Private/NiagaraHlslTranslator.cpp
@@ -4948,7 +4948,7 @@ void FHlslNiagaraTranslator::RegisterFunctionCall(ENiagaraScriptUsage ScriptUsag
 				{
 					if (UseShaderStages)
 					{
-						if (bStageMaxFilter && bStageMaxFilter)
+						if (bStageMinFilter && bStageMaxFilter)
 						{
 							FunctionDefStr += TEXT("if ((GCurrentPhase == 1 && ShaderStageIndex >= ") + MinParam + TEXT(" && ShaderStageIndex <= ") + MaxParam + TEXT(") || ") +
 								TEXT("(GCurrentPhase == 0 && ShaderStageIndex >= ") + MinParamSpawn + TEXT(" && ShaderStageIndex <= ") + MaxParamSpawn + TEXT(")") +
