Trying to fix a crash Jon experienced where the previous bone positons array was missing altogether.

#tests autotests pass

diff --git a/Engine/Plugins/FX/Niagara/Source/Niagara/Classes/NDISkeletalMeshCommon.h b/Engine/Plugins/FX/Niagara/Source/Niagara/Classes/NDISkeletalMeshCommon.h
--- a/Engine/Plugins/FX/Niagara/Source/Niagara/Classes/NDISkeletalMeshCommon.h
+++ b/Engine/Plugins/FX/Niagara/Source/Niagara/Classes/NDISkeletalMeshCommon.h
@@ -166,7 +166,11 @@ struct FSkinnedPositionAccessorHelper<TIntegralConstant<ENDISkeletalMesh_Skinnin
 
 	FORCEINLINE_DEBUGGABLE FVector GetSkinnedBonePreviousPosition(FSkeletalMeshAccessorHelper& Accessor, int32 BoneIndex)
 	{
-		return (*Accessor.PrevBoneComponentSpaceTransforms)[BoneIndex].GetLocation();
+		if (Accessor.PrevBoneComponentSpaceTransforms != nullptr && Accessor.PrevBoneComponentSpaceTransforms->Num() > 0)
+		{
+			return (*Accessor.PrevBoneComponentSpaceTransforms)[BoneIndex].GetLocation();
+		}
+		return GetSkinnedBonePosition(Accessor, BoneIndex);
 	}
 
 	FORCEINLINE_DEBUGGABLE FQuat GetSkinnedBoneRotation(FSkeletalMeshAccessorHelper& Accessor, int32 BoneIndex)
@@ -176,7 +180,11 @@ struct FSkinnedPositionAccessorHelper<TIntegralConstant<ENDISkeletalMesh_Skinnin
 
 	FORCEINLINE_DEBUGGABLE FQuat GetSkinnedBonePreviousRotation(FSkeletalMeshAccessorHelper& Accessor, int32 BoneIndex)
 	{
-		return (*Accessor.PrevBoneComponentSpaceTransforms)[BoneIndex].GetRotation();
+		if (Accessor.PrevBoneComponentSpaceTransforms != nullptr && Accessor.PrevBoneComponentSpaceTransforms->Num() > 0)
+		{
+			return (*Accessor.PrevBoneComponentSpaceTransforms)[BoneIndex].GetRotation();
+		}
+		return GetSkinnedBoneRotation(Accessor, BoneIndex);
 	}
 };
 
diff --git a/Engine/Plugins/FX/Niagara/Source/Niagara/Private/NDISkeletalMesh_BoneSampling.cpp b/Engine/Plugins/FX/Niagara/Source/Niagara/Private/NDISkeletalMesh_BoneSampling.cpp
--- a/Engine/Plugins/FX/Niagara/Source/Niagara/Private/NDISkeletalMesh_BoneSampling.cpp
+++ b/Engine/Plugins/FX/Niagara/Source/Niagara/Private/NDISkeletalMesh_BoneSampling.cpp
@@ -456,11 +456,37 @@ void UNiagaraDataInterfaceSkeletalMesh::GetSkinnedBoneData(FVectorVMContext& Con
 		const float Interp = bInterpolated::Value ? InterpParam.GetAndAdvance() : 1.0f;
 
 		// Determine bone or socket
-		int32 Bone = FMath::Clamp(BoneParam.GetAndAdvance(), 0, BoneAndSocketMax);
+		int32 RawBone = BoneParam.GetAndAdvance();
+		int32 Bone = FMath::Clamp(RawBone, 0, BoneAndSocketMax);
 		const bool bIsSocket = Bone > BoneMax;
-		if ( bIsSocket )
+		const int32 Socket = Bone - SpecificSocketBoneOffset;
+
+		// Handle edge cases first...
+		if ((!bIsSocket && Bone >= BoneMax) ||
+			(bIsSocket && (Socket >= SpecificSocketCurrTransforms.Num() || Socket < 0)))
+		{
+			Pos = FVector::ZeroVector;
+			TransformHandler.TransformPosition(Pos, Transform);
+
+			if (Output.bNeedsVelocity || bInterpolated::Value)
+			{
+				Prev = FVector::ZeroVector;
+				TransformHandler.TransformPosition(Prev, PrevTransform);
+			}
+			if (Output.bNeedsRotation)
+			{
+				FQuat Rotation = FQuat::Identity;
+				if (bInterpolated::Value)
+				{
+					FQuat PrevRotation = FQuat::Identity;
+					Rotation = FQuat::Identity;
+				}
+
+				Output.SetRotation(Rotation);
+			}
+		}
+		else if ( bIsSocket )
 		{
-			const int32 Socket = Bone - SpecificSocketBoneOffset;
 			FTransform CurrSocketTransform = SpecificSocketCurrTransforms[Socket];
 			FTransform PrevSocketTransform = SpecificSocketPrevTransforms[Socket];
 
