Fixing crash when changing out a system asset Jon encountered

#tests n/a

diff --git a/Engine/Plugins/FX/Niagara/Source/NiagaraEditorWidgets/Private/DetailCustomizations/NiagaraDataInterfaceSkeletalMeshDetails.cpp b/Engine/Plugins/FX/Niagara/Source/NiagaraEditorWidgets/Private/DetailCustomizations/NiagaraDataInterfaceSkeletalMeshDetails.cpp
--- a/Engine/Plugins/FX/Niagara/Source/NiagaraEditorWidgets/Private/DetailCustomizations/NiagaraDataInterfaceSkeletalMeshDetails.cpp
+++ b/Engine/Plugins/FX/Niagara/Source/NiagaraEditorWidgets/Private/DetailCustomizations/NiagaraDataInterfaceSkeletalMeshDetails.cpp
@@ -121,15 +121,26 @@ void FNiagaraDataInterfaceSkeletalMeshDetails::CustomizeDetails(IDetailLayoutBui
 
  void FNiagaraDataInterfaceSkeletalMeshDetails::OnDataChanged()
  {
-	TArray<TSharedPtr<FName>> PossibleNames;
-	GenerateRegionsArray(PossibleNames);
-	RegionsBuilder->SetSourceArray(PossibleNames);
-
-	GenerateBonesArray(PossibleNames);
-	BonesBuilder->SetSourceArray(PossibleNames);
-
-	GenerateSocketsArray(PossibleNames);
-	SocketsBuilder->SetSourceArray(PossibleNames);
+	if (RegionsBuilder)
+	{
+		TArray<TSharedPtr<FName>> PossibleNames;
+		GenerateRegionsArray(PossibleNames);
+		RegionsBuilder->SetSourceArray(PossibleNames);
+	}
+
+	if (BonesBuilder)
+	{
+		TArray<TSharedPtr<FName>> PossibleNames;
+		GenerateBonesArray(PossibleNames);
+		BonesBuilder->SetSourceArray(PossibleNames);
+	}
+
+	if (SocketsBuilder)
+	{
+		TArray<TSharedPtr<FName>> PossibleNames;
+		GenerateSocketsArray(PossibleNames);
+		SocketsBuilder->SetSourceArray(PossibleNames);
+	}
  }
 
 void FNiagaraDataInterfaceSkeletalMeshDetails::GenerateRegionsArray(TArray<TSharedPtr<FName>>& SourceArray)
