Niagara - Fix a crash when dragging system modules.

diff --git a/Engine/Plugins/FX/Niagara/Source/NiagaraEditor/Private/ViewModels/Stack/NiagaraStackScriptItemGroup.cpp b/Engine/Plugins/FX/Niagara/Source/NiagaraEditor/Private/ViewModels/Stack/NiagaraStackScriptItemGroup.cpp
--- a/Engine/Plugins/FX/Niagara/Source/NiagaraEditor/Private/ViewModels/Stack/NiagaraStackScriptItemGroup.cpp
+++ b/Engine/Plugins/FX/Niagara/Source/NiagaraEditor/Private/ViewModels/Stack/NiagaraStackScriptItemGroup.cpp
@@ -562,9 +562,7 @@ TOptional<UNiagaraStackEntry::FDropResult> UNiagaraStackScriptItemGroup::ChildRe
 				return FDropResult(false, LOCTEXT("CantMoveByUsage", "This module can't be moved to this section of the stack because its it's not supported in this context."));
 			}
 
-			const FNiagaraEmitterHandle* SourceEmitterHandle = FNiagaraEditorUtilities::GetEmitterHandleForEmitter(
-				SourceModuleItem->GetSystemViewModel()->GetSystem(), *SourceModuleItem->GetEmitterViewModel()->GetEmitter());
-			if (SourceEmitterHandle == nullptr)
+			if (SourceModuleItem->GetSystemViewModel() != TargetChild.GetSystemViewModel())
 			{
 				return FDropResult(false, LOCTEXT("CantMoveFromAnotherSystem", "This module can't be moved into this system from a different system."));
 			}
@@ -622,32 +620,39 @@ TOptional<UNiagaraStackEntry::FDropResult> UNiagaraStackScriptItemGroup::ChildRe
 				UNiagaraStackModuleItem** TargetModuleItemPtr = StackSpacerToModuleItemMap.Find(FObjectKey(TargetSpacer));
 				if (TargetModuleItemPtr != nullptr)
 				{
-					const FNiagaraEmitterHandle* SourceEmitterHandle = FNiagaraEditorUtilities::GetEmitterHandleForEmitter(
-						SourceModuleItem->GetSystemViewModel()->GetSystem(), *SourceModuleItem->GetEmitterViewModel()->GetEmitter());
-					if (SourceEmitterHandle != nullptr)
+					const FNiagaraEmitterHandle* SourceEmitterHandle = SourceModuleItem->GetEmitterViewModel().IsValid()
+						? FNiagaraEditorUtilities::GetEmitterHandleForEmitter(SourceModuleItem->GetSystemViewModel()->GetSystem(), *SourceModuleItem->GetEmitterViewModel()->GetEmitter())
+						: nullptr;
+					FGuid SourceEmitterHandleId = SourceEmitterHandle != nullptr
+						? SourceEmitterHandle->GetId()
+						: FGuid();
+
+					UNiagaraNodeOutput* SourceModuleOutputNode = FNiagaraStackGraphUtilities::GetEmitterOutputNodeForStackNode(SourceModuleItem->GetModuleNode());
+					UNiagaraScript* SourceModuleScript = FNiagaraEditorUtilities::GetScriptFromSystem(SourceModuleItem->GetSystemViewModel()->GetSystem(), SourceEmitterHandleId,
+						SourceModuleOutputNode->GetUsage(), SourceModuleOutputNode->GetUsageId());
+
+					const FNiagaraEmitterHandle* TargetEmitterHandle = GetEmitterViewModel().IsValid()
+						? FNiagaraEditorUtilities::GetEmitterHandleForEmitter(GetSystemViewModel()->GetSystem(), *GetEmitterViewModel()->GetEmitter())
+						: nullptr;
+					FGuid TargetEmitterHandleId = TargetEmitterHandle != nullptr
+						? TargetEmitterHandle->GetId()
+						: FGuid();
+
+					int32 TargetIndex = *TargetModuleItemPtr != nullptr ? (*TargetModuleItemPtr)->GetModuleIndex() : INDEX_NONE;
+
+					FScopedTransaction ScopedTransaction(LOCTEXT("DragAndDropModule", "Drag and drop module"));
 					{
-						UNiagaraNodeOutput* SourceModuleOutputNode = FNiagaraStackGraphUtilities::GetEmitterOutputNodeForStackNode(SourceModuleItem->GetModuleNode());
-						UNiagaraScript* SourceModuleScript = FNiagaraEditorUtilities::GetScriptFromSystem(SourceModuleItem->GetSystemViewModel()->GetSystem(), SourceEmitterHandle->GetId(),
-							SourceModuleOutputNode->GetUsage(), SourceModuleOutputNode->GetUsageId());
+						FNiagaraStackGraphUtilities::MoveModule(*SourceModuleScript, SourceModuleItem->GetModuleNode(), GetSystemViewModel()->GetSystem(), TargetEmitterHandleId,
+							ScriptUsage, ScriptUsageId, TargetIndex);
 
-						const FNiagaraEmitterHandle* TargetEmitterHandle = FNiagaraEditorUtilities::GetEmitterHandleForEmitter(GetSystemViewModel()->GetSystem(), *GetEmitterViewModel()->GetEmitter());
+						UNiagaraGraph* TargetGraph = ScriptViewModel.Pin()->GetGraphViewModel()->GetGraph();
+						FNiagaraStackGraphUtilities::RelayoutGraph(*TargetGraph);
+						TargetGraph->NotifyGraphNeedsRecompile();
 
-						int32 TargetIndex = *TargetModuleItemPtr != nullptr ? (*TargetModuleItemPtr)->GetModuleIndex() : INDEX_NONE;
-
-						FScopedTransaction ScopedTransaction(LOCTEXT("DragAndDropModule", "Drag and drop module"));
-						{
-							FNiagaraStackGraphUtilities::MoveModule(*SourceModuleScript, SourceModuleItem->GetModuleNode(), GetSystemViewModel()->GetSystem(), TargetEmitterHandle->GetId(),
-								ScriptUsage, ScriptUsageId, TargetIndex);
-
-							UNiagaraGraph* TargetGraph = ScriptViewModel.Pin()->GetGraphViewModel()->GetGraph();
-							FNiagaraStackGraphUtilities::RelayoutGraph(*TargetGraph);
-							TargetGraph->NotifyGraphNeedsRecompile();
-
-							SourceModuleItem->NotifyModuleMoved();
-							RefreshChildren();
-						}
-						return FDropResult(true);
+						SourceModuleItem->NotifyModuleMoved();
+						RefreshChildren();
 					}
+					return FDropResult(true);
 				}
 			}
 		}
