Fix compile errors related to VS2019.

diff --git a/Engine/Plugins/FX/Niagara/Source/Niagara/Private/NiagaraDataInterfaceParticleRead.cpp b/Engine/Plugins/FX/Niagara/Source/Niagara/Private/NiagaraDataInterfaceParticleRead.cpp
--- a/Engine/Plugins/FX/Niagara/Source/Niagara/Private/NiagaraDataInterfaceParticleRead.cpp
+++ b/Engine/Plugins/FX/Niagara/Source/Niagara/Private/NiagaraDataInterfaceParticleRead.cpp
@@ -62,7 +62,7 @@ bool UNiagaraDataInterfaceParticleRead::InitPerInstanceData(void* PerInstanceDat
 			break;
 		}
 	}
-	return PIData->EmitterInstance;
+	return PIData->EmitterInstance != nullptr;
 }
 
 void UNiagaraDataInterfaceParticleRead::DestroyPerInstanceData(void* PerInstanceData, FNiagaraSystemInstance* SystemInstance)
