Added "fx.NiagaraAllowTickBeforeRender 1" which allows Niagara GPU particles to render the current frame state if their simulation is view independent
(for example not depening on the depth buffer). Otherwise the rendered particle are the result of the last frame.

diff --git a/Engine/Plugins/FX/Niagara/Source/Niagara/Private/NiagaraEmitterInstanceBatcher.cpp b/Engine/Plugins/FX/Niagara/Source/Niagara/Private/NiagaraEmitterInstanceBatcher.cpp
--- a/Engine/Plugins/FX/Niagara/Source/Niagara/Private/NiagaraEmitterInstanceBatcher.cpp
+++ b/Engine/Plugins/FX/Niagara/Source/Niagara/Private/NiagaraEmitterInstanceBatcher.cpp
@@ -29,6 +29,15 @@ DECLARE_GPU_STAT_NAMED(NiagaraGPUSorting, TEXT("Niagara GPU sorting"));
 
 uint32 FNiagaraComputeExecutionContext::TickCounter = 0;
 
+int32 GNiagaraAllowTickBeforeRender = 1;
+static FAutoConsoleVariableRef CVarNiagaraAllowTickBeforeRender(
+	TEXT("fx.NiagaraAllowTickBeforeRender"),
+	GNiagaraAllowTickBeforeRender,
+	TEXT("If 1, Niagara GPU systems that don't rely on view data will be rendered in sync\n")
+	TEXT("with the current frame simulation instead of the last frame one. (default=1)\n"),
+	ECVF_Default
+);
+
 int32 GNiagaraOverlapCompute = 1;
 static FAutoConsoleVariableRef CVarNiagaraUseAsyncCompute(
 	TEXT("fx.NiagaraOverlapCompute"),
@@ -503,6 +512,8 @@ void NiagaraEmitterInstanceBatcher::PostRenderOpaque(FRHICommandListImmediate& R
 	// Setup new readback since if there is no pending request, there is no risk of having invalid data read (offset being allocated after the readback was sent).
 	ExecuteAll(RHICmdList, ViewUniformBuffer, !GPUInstanceCounterManager.HasPendingGPUReadback());
 
+	FinishDispatches();
+
 	if (!GPUInstanceCounterManager.HasPendingGPUReadback())
 	{
 		GPUInstanceCounterManager.EnqueueGPUReadback(RHICmdList);
@@ -511,9 +522,6 @@ void NiagaraEmitterInstanceBatcher::PostRenderOpaque(FRHICommandListImmediate& R
 
 void NiagaraEmitterInstanceBatcher::ExecuteAll(FRHICommandList &RHICmdList, FRHIUniformBuffer* ViewUniformBuffer, bool bSetReadback)
 {
-	// @todo REMOVE THIS HACK
-	LastFrameThatDrainedData = GFrameNumberRenderThread;
-
 	// This is always called by the renderer so early out if we have no work.
 	if (Ticks_RT.Num() == 0)
 	{
@@ -533,6 +541,18 @@ void NiagaraEmitterInstanceBatcher::ExecuteAll(FRHICommandList &RHICmdList, FRHI
 			for (FNiagaraGPUSystemTick& Tick : Ticks_RT)
 			{
 				FNiagaraComputeInstanceData* Data = Tick.GetInstanceData();
+
+				// This assumes all emitters fallback to the same FNiagaraShaderScript*.
+				FNiagaraShader* ComputeShader = Data->Context->GPUScript_RT->GetShader();
+				if (!ComputeShader)
+				{
+					continue;
+				}
+				else if (GNiagaraAllowTickBeforeRender && ComputeShader->ViewUniformBufferParam.IsBound() != (ViewUniformBuffer != nullptr))
+				{   // When allowing tick before render, skip this emitter if it is not in the right pass.
+					continue;
+				}
+
 				ContextToTicks.FindOrAdd(Data->Context).Add(&Tick);
 			}
 
@@ -585,8 +605,6 @@ void NiagaraEmitterInstanceBatcher::ExecuteAll(FRHICommandList &RHICmdList, FRHI
 			}
 		}
 	}
-
-	FinishDispatches();
 }
 
 void NiagaraEmitterInstanceBatcher::TickSingle(const FNiagaraGPUSystemTick& Tick, FNiagaraComputeInstanceData *Instance, FRHICommandList &RHICmdList, FRHIUniformBuffer* ViewUniformBuffer, bool bSetReadback) const
@@ -604,6 +622,11 @@ void NiagaraEmitterInstanceBatcher::TickSingle(const FNiagaraGPUSystemTick& Tick
 	{
 		return;
 	}
+	// When allowing tick before render, skip this emitter if it is not in the right pass.
+	else if (GNiagaraAllowTickBeforeRender && ComputeShader->ViewUniformBufferParam.IsBound() != (ViewUniformBuffer != nullptr))
+	{
+		return;
+	}
 
 #if WITH_EDITORONLY_DATA
 	if (Context->DebugInfo.IsValid())
@@ -733,6 +756,14 @@ void NiagaraEmitterInstanceBatcher::PreInitViews(FRHICommandListImmediate& RHICm
 			GPUInstanceCounterManager.ReleaseGPUReadback();
 		}
 	}
+
+	// @todo REMOVE THIS HACK
+	LastFrameThatDrainedData = GFrameNumberRenderThread;
+
+	if (GNiagaraAllowTickBeforeRender)
+	{
+		ExecuteAll(RHICmdList, nullptr, !GPUInstanceCounterManager.HasPendingGPUReadback());
+	}
 }
 
 bool NiagaraEmitterInstanceBatcher::UsesGlobalDistanceField() const
