Add Equals and CopyToInternal implementation to direct DI, fixes autotest

diff --git a/Engine/Plugins/FX/Niagara/Source/Niagara/Classes/NiagaraDataInterfaceParticleRead.h b/Engine/Plugins/FX/Niagara/Source/Niagara/Classes/NiagaraDataInterfaceParticleRead.h
--- a/Engine/Plugins/FX/Niagara/Source/Niagara/Classes/NiagaraDataInterfaceParticleRead.h
+++ b/Engine/Plugins/FX/Niagara/Source/Niagara/Classes/NiagaraDataInterfaceParticleRead.h
@@ -59,6 +59,9 @@ public:
 	virtual void ReadColor(FVectorVMContext& Context, FName AttributeToRead);
 	virtual void ReadQuat(FVectorVMContext& Context, FName AttributeToRead);
 
+protected:
+	virtual bool CopyToInternal(UNiagaraDataInterface* Destination) const override;
+
 private:
 	template<typename T>
 	T RetrieveValueWithCheck(FNiagaraEmitterInstance* EmitterInstance, const FNiagaraTypeDefinition& Type, const FName& Attr, const FNiagaraID& ParticleID, bool &bValid);
diff --git a/Engine/Plugins/FX/Niagara/Source/Niagara/Private/NiagaraDataInterfaceParticleRead.cpp b/Engine/Plugins/FX/Niagara/Source/Niagara/Private/NiagaraDataInterfaceParticleRead.cpp
--- a/Engine/Plugins/FX/Niagara/Source/Niagara/Private/NiagaraDataInterfaceParticleRead.cpp
+++ b/Engine/Plugins/FX/Niagara/Source/Niagara/Private/NiagaraDataInterfaceParticleRead.cpp
@@ -581,7 +581,21 @@ T UNiagaraDataInterfaceParticleRead::RetrieveValueWithCheck(FNiagaraEmitterInsta
 
 bool UNiagaraDataInterfaceParticleRead::Equals(const UNiagaraDataInterface* Other) const
 {
-	return false;
+	if (!Super::Equals(Other))
+	{
+		return false;
+	}
+	return CastChecked<UNiagaraDataInterfaceParticleRead>(Other)->EmitterName == EmitterName;
+}
+
+bool UNiagaraDataInterfaceParticleRead::CopyToInternal(UNiagaraDataInterface* Destination) const
+{
+	if (!Super::CopyToInternal(Destination))
+	{
+		return false;
+	}
+	CastChecked<UNiagaraDataInterfaceParticleRead>(Destination)->EmitterName = EmitterName;
+	return true;
 }
 
 bool UNiagaraDataInterfaceParticleRead::GetFunctionHLSL(const FName& DefinitionFunctionName, FString InstanceFunctionName, FNiagaraDataInterfaceGPUParamInfo& ParamInfo, FString& OutHLSL)
