Fixing a null pointer dereference issue

diff --git a/Engine/Plugins/FX/Niagara/Source/NiagaraEditor/Private/Toolkits/NiagaraSystemToolkit.cpp b/Engine/Plugins/FX/Niagara/Source/NiagaraEditor/Private/Toolkits/NiagaraSystemToolkit.cpp
--- a/Engine/Plugins/FX/Niagara/Source/NiagaraEditor/Private/Toolkits/NiagaraSystemToolkit.cpp
+++ b/Engine/Plugins/FX/Niagara/Source/NiagaraEditor/Private/Toolkits/NiagaraSystemToolkit.cpp
@@ -548,7 +548,7 @@ public:
 
 	~SNiagaraSelectedEmitterGraph()
 	{
-		if (SystemViewModel.IsValid())
+		if (SystemViewModel.IsValid() && SystemViewModel->GetSelectionViewModel())
 		{
 			SystemViewModel->GetSelectionViewModel()->OnSelectionChanged().RemoveAll(this);
 		}
