Niagara - Fix crashes and other issues related to object paramers in the stack.

diff --git a/Engine/Plugins/FX/Niagara/Source/NiagaraEditor/Private/ViewModels/Stack/NiagaraStackParameterStoreEntry.cpp b/Engine/Plugins/FX/Niagara/Source/NiagaraEditor/Private/ViewModels/Stack/NiagaraStackParameterStoreEntry.cpp
--- a/Engine/Plugins/FX/Niagara/Source/NiagaraEditor/Private/ViewModels/Stack/NiagaraStackParameterStoreEntry.cpp
+++ b/Engine/Plugins/FX/Niagara/Source/NiagaraEditor/Private/ViewModels/Stack/NiagaraStackParameterStoreEntry.cpp
@@ -104,7 +104,7 @@ TSharedPtr<FStructOnScope> UNiagaraStackParameterStoreEntry::GetValueStruct()
 	return LocalValueStruct;
 }
 
-UNiagaraDataInterface* UNiagaraStackParameterStoreEntry::GetValueObject()
+UObject* UNiagaraStackParameterStoreEntry::GetValueObject()
 {
 	return ValueObject;
 }
@@ -154,8 +154,11 @@ void UNiagaraStackParameterStoreEntry::Reset()
 	}
 	else
 	{
-		UNiagaraDataInterface* DefaultObject = NewObject<UNiagaraDataInterface>(this, const_cast<UClass*>(InputType.GetClass()));
-		DefaultObject->CopyTo(ParameterStore->GetDataInterface(Var));
+		if (Var.IsDataInterface())
+		{
+			UNiagaraDataInterface* DefaultObject = NewObject<UNiagaraDataInterface>(this, const_cast<UClass*>(InputType.GetClass()));
+			DefaultObject->CopyTo(ParameterStore->GetDataInterface(Var));
+		}
 	}
 	RefreshValueAndHandle();
 	RefreshChildren();
@@ -323,11 +326,17 @@ void UNiagaraStackParameterStoreEntry::RemovePins(TArray<UEdGraphPin*> OwningPin
 			}
 			else
 			{
-				UNiagaraDataInterface* OverrideObj = NewObject<UNiagaraDataInterface>(this, const_cast<UClass*>(InputType.GetClass()));
-				FNiagaraStackGraphUtilities::SetDataValueObjectForFunctionInput(*OverridePin, const_cast<UClass*>(InputType.GetClass()), GetCurrentValueObject()->GetName(), OverrideObj);
-				GetCurrentValueObject()->CopyTo(OverrideObj);
+				if (InputType.IsDataInterface())
+				{
+					UNiagaraDataInterface* CurrentValueDataInterface = Cast<UNiagaraDataInterface>(GetCurrentValueObject());
+					if (CurrentValueDataInterface != nullptr)
+					{
+						UNiagaraDataInterface* OverrideObj = NewObject<UNiagaraDataInterface>(this, const_cast<UClass*>(InputType.GetClass()));
+						FNiagaraStackGraphUtilities::SetDataValueObjectForFunctionInput(*OverridePin, const_cast<UClass*>(InputType.GetClass()), CurrentValueDataInterface->GetName(), OverrideObj);
+						CurrentValueDataInterface->CopyTo(OverrideObj);
+					}
+				}
 			}
-			
 		}
 		// now also remove node
 		Graph->RemoveNode(GraphPin->GetOwningNode());
@@ -359,12 +368,19 @@ TSharedPtr<FNiagaraVariable> UNiagaraStackParameterStoreEntry::GetCurrentValueVa
 	return TSharedPtr<FNiagaraVariable>();
 }
 
-UNiagaraDataInterface* UNiagaraStackParameterStoreEntry::GetCurrentValueObject()
+UObject* UNiagaraStackParameterStoreEntry::GetCurrentValueObject()
 {
 	if (InputType.GetClass() != nullptr)
 	{
 		FNiagaraVariable DefaultVariable(InputType, ParameterName);
-		return ParameterStore->GetDataInterface(DefaultVariable);
+		if (DefaultVariable.IsDataInterface())
+		{
+			return ParameterStore->GetDataInterface(DefaultVariable);
+		}
+		else if (DefaultVariable.IsUObject())
+		{
+			return ParameterStore->GetUObject(DefaultVariable);
+		}
 	}
 	return nullptr;
 }
diff --git a/Engine/Plugins/FX/Niagara/Source/NiagaraEditor/Public/ViewModels/Stack/NiagaraStackParameterStoreEntry.h b/Engine/Plugins/FX/Niagara/Source/NiagaraEditor/Public/ViewModels/Stack/NiagaraStackParameterStoreEntry.h
--- a/Engine/Plugins/FX/Niagara/Source/NiagaraEditor/Public/ViewModels/Stack/NiagaraStackParameterStoreEntry.h
+++ b/Engine/Plugins/FX/Niagara/Source/NiagaraEditor/Public/ViewModels/Stack/NiagaraStackParameterStoreEntry.h
@@ -55,7 +55,7 @@ public:
 	TSharedPtr<FStructOnScope> GetValueStruct();
 
 	/** Gets the current object value of this input is there is one. */
-	UNiagaraDataInterface* GetValueObject();
+	UObject* GetValueObject();
 
 	/** Called to notify the input that an ongoing change to it's value has begun. */
 	void NotifyBeginValueChange();
@@ -102,7 +102,7 @@ protected:
 
 	void RefreshValueAndHandle();
 	TSharedPtr<FNiagaraVariable> GetCurrentValueVariable();
-	UNiagaraDataInterface* GetCurrentValueObject();
+	UObject* GetCurrentValueObject();
 
 private:
 	void RemovePins(TArray<UEdGraphPin*> PinsToRemove);
@@ -129,7 +129,7 @@ private:
 
 	/** A pointer to the data interface object for this input if one is available. */
 	UPROPERTY()
-	UNiagaraDataInterface* ValueObject;
+	UObject* ValueObject;
 
 	UPROPERTY()
 	UObject* Owner;
diff --git a/Engine/Plugins/FX/Niagara/Source/NiagaraEditorWidgets/Private/Stack/SNiagaraStackParameterStoreEntryValue.cpp b/Engine/Plugins/FX/Niagara/Source/NiagaraEditorWidgets/Private/Stack/SNiagaraStackParameterStoreEntryValue.cpp
--- a/Engine/Plugins/FX/Niagara/Source/NiagaraEditorWidgets/Private/Stack/SNiagaraStackParameterStoreEntryValue.cpp
+++ b/Engine/Plugins/FX/Niagara/Source/NiagaraEditorWidgets/Private/Stack/SNiagaraStackParameterStoreEntryValue.cpp
@@ -64,21 +64,6 @@ void SNiagaraStackParameterStoreEntryValue::Construct(const FArguments& InArgs,
 			]
 		]
 
-		// Handle drop-down button
-		+ SHorizontalBox::Slot()
-		.AutoWidth()
-		.VAlign(VAlign_Center)
-		.Padding(3, 0, 0, 0)
-		[
-			SNew(SComboButton)
-			.ButtonStyle(FEditorStyle::Get(), "HoverHintOnly")
-			.ForegroundColor(FSlateColor::UseForeground())
-			.OnGetMenuContent(this, &SNiagaraStackParameterStoreEntryValue::OnGetAvailableHandleMenu)
-			.ContentPadding(FMargin(2))
-			.HAlign(HAlign_Center)
-			.VAlign(VAlign_Center)
-		]
-
 		// Reset Button
 		+ SHorizontalBox::Slot()
 		.AutoWidth()
@@ -133,13 +118,6 @@ FReply SNiagaraStackParameterStoreEntryValue::DeleteClicked()
 	return FReply::Handled();
 }
 
-TSharedRef<SWidget> SNiagaraStackParameterStoreEntryValue::OnGetAvailableHandleMenu()
-{
-	// TODO: This will need to be adjusted based on the current stack being edited, i.e. system vs emitter vs particle.
-	FMenuBuilder MenuBuilder(true, nullptr);
-	return MenuBuilder.MakeWidget();
-}
-
 TSharedRef<SWidget> SNiagaraStackParameterStoreEntryValue::ConstructValueStructWidget()
 {
 	ValueStructParameterEditor.Reset();
@@ -184,6 +162,11 @@ TSharedRef<SWidget> SNiagaraStackParameterStoreEntryValue::ConstructValueStructW
 			return StructureDetailsView->GetWidget().ToSharedRef();
 		}
 	}
+	else if (StackEntry->GetInputType().GetClass() != nullptr)
+	{
+		return SNew(STextBlock)
+			.Text(StackEntry->GetInputType().GetClass()->GetDisplayNameText());
+	}
 	else
 	{
 		return SNullWidget::NullWidget;
diff --git a/Engine/Plugins/FX/Niagara/Source/NiagaraEditorWidgets/Private/Stack/SNiagaraStackParameterStoreEntryValue.h b/Engine/Plugins/FX/Niagara/Source/NiagaraEditorWidgets/Private/Stack/SNiagaraStackParameterStoreEntryValue.h
--- a/Engine/Plugins/FX/Niagara/Source/NiagaraEditorWidgets/Private/Stack/SNiagaraStackParameterStoreEntryValue.h
+++ b/Engine/Plugins/FX/Niagara/Source/NiagaraEditorWidgets/Private/Stack/SNiagaraStackParameterStoreEntryValue.h
@@ -24,8 +24,6 @@ public:
 private:
 	FReply DeleteClicked();
 
-	TSharedRef<SWidget> OnGetAvailableHandleMenu();
-
 	TSharedRef<SWidget> ConstructValueStructWidget();
 
 	void OnInputValueChanged();
