Niagara - Fix issues with input and output renaming in function scripts.
+ Fix an issue where renaming an input or output from the parameters panel would not update a scripts change id and would not enable the apply button.
+ Fix an issue where you could create duplicate output names when renaming from the node context menu.
+ Fix an issue where you could not have a function input named the same as a function output.

diff --git a/Engine/Plugins/FX/Niagara/Source/NiagaraEditor/Private/NiagaraEditorUtilities.cpp b/Engine/Plugins/FX/Niagara/Source/NiagaraEditor/Private/NiagaraEditorUtilities.cpp
--- a/Engine/Plugins/FX/Niagara/Source/NiagaraEditor/Private/NiagaraEditorUtilities.cpp
+++ b/Engine/Plugins/FX/Niagara/Source/NiagaraEditor/Private/NiagaraEditorUtilities.cpp
@@ -1143,4 +1143,47 @@ void FNiagaraEditorUtilities::KillSystemInstances(const UNiagaraSystem& System)
 	}
 }
 
+bool FNiagaraEditorUtilities::VerifyNameChangeForInputOrOutputNode(UNiagaraNode& NodeBeingChanged, FName OldName, FName NewName, FText& OutErrorMessage)
+{
+	if (NewName == NAME_None)
+	{
+		OutErrorMessage = LOCTEXT("EmptyNameError", "Name can not be empty.");
+		return false;
+	}
+
+	if (GetSystemConstantNames().Contains(NewName))
+	{
+		OutErrorMessage = LOCTEXT("SystemConstantNameError", "Name can not be the same as a system constant");
+	}
+
+	if (NodeBeingChanged.IsA<UNiagaraNodeInput>())
+	{
+		TArray<UNiagaraNodeInput*> InputNodes;
+		NodeBeingChanged.GetGraph()->GetNodesOfClass<UNiagaraNodeInput>(InputNodes);
+		for (UNiagaraNodeInput* InputNode : InputNodes)
+		{
+			if (InputNode->Input.GetName() != OldName && InputNode->Input.GetName() == NewName)
+			{
+				OutErrorMessage = LOCTEXT("DuplicateInputNameError", "Name can not match an existing input name.");
+				return false;
+			}
+		}
+	}
+
+	if (NodeBeingChanged.IsA<UNiagaraNodeOutput>())
+	{
+		UNiagaraNodeOutput* OutputNodeBeingChanged = CastChecked<UNiagaraNodeOutput>(&NodeBeingChanged);
+		for (const FNiagaraVariable& Output : OutputNodeBeingChanged->GetOutputs())
+		{
+			if (Output.GetName() != OldName && Output.GetName() == NewName)
+			{
+				OutErrorMessage = LOCTEXT("DuplicateOutputNameError", "Name can not match an existing output name.");
+				return false;
+			}
+		}
+	}
+
+	return true;
+}
+
 #undef LOCTEXT_NAMESPACE
diff --git a/Engine/Plugins/FX/Niagara/Source/NiagaraEditor/Private/NiagaraNodeInput.cpp b/Engine/Plugins/FX/Niagara/Source/NiagaraEditor/Private/NiagaraNodeInput.cpp
--- a/Engine/Plugins/FX/Niagara/Source/NiagaraEditor/Private/NiagaraNodeInput.cpp
+++ b/Engine/Plugins/FX/Niagara/Source/NiagaraEditor/Private/NiagaraNodeInput.cpp
@@ -187,84 +187,6 @@ FName UNiagaraNodeInput::GenerateUniqueName(const UNiagaraGraph* Graph, FName& P
 	return FNiagaraUtilities::GetUniqueName(ProposedName, FNiagaraEditorUtilities::GetSystemConstantNames().Union(InputNames));
 }
 
-bool UNiagaraNodeInput::VerifyNodeRenameTextCommit(const FText& NewText, UNiagaraNode* NodeBeingChanged, FText& OutErrorMessage)
-{
-	FName NewName = *NewText.ToString();
-	TSet<FName> SystemConstantNames = FNiagaraEditorUtilities::GetSystemConstantNames();
-
-	// Disallow empty names
-	if (NewName == FName())
-	{
-		OutErrorMessage = LOCTEXT("NiagaraInputNameEmptyWarn", "Cannot have empty name!");
-		return false;
-	}
-
-	// Disallow name changes to system constants
-	if (SystemConstantNames.Contains(NewName))
-	{
-		OutErrorMessage = FText::Format(LOCTEXT("NiagaraInputNameSystemWarn", "\"{0}\" is a system constant name."), FText::FromName(NewName));
-		return false;
-	}
-
-	// @TODO: Prevent any hlsl keywords or invalid hlsl characters from being used as names!
-
-	UNiagaraNodeInput* InputNodeBeingChanged = Cast<UNiagaraNodeInput>(NodeBeingChanged);
-	UNiagaraNodeOutput* OutputNodeBeingChanged = Cast<UNiagaraNodeOutput>(NodeBeingChanged);
-
-	// Make sure that we aren't changing names to something already in the graph.
-	if (NodeBeingChanged != nullptr)
-	{
-		UNiagaraGraph* Graph = CastChecked<UNiagaraGraph>(NodeBeingChanged->GetGraph());
-
-		// If dealing with parameter, check to make sure that we don't conflict with any other parameter name.
-		if (InputNodeBeingChanged != nullptr && InputNodeBeingChanged->Usage == ENiagaraInputNodeUsage::Parameter)
-		{
-			TArray<UNiagaraNodeInput*> InputNodes;
-			Graph->GetNodesOfClass<UNiagaraNodeInput>(InputNodes);
-
-			for (UNiagaraNodeInput* Node : InputNodes)
-			{
-				if (Node == nullptr || Node == InputNodeBeingChanged || Node->Usage != InputNodeBeingChanged->Usage)
-				{
-					continue;
-				}
-
-				bool bIsSame = Node->ReferencesSameInput(InputNodeBeingChanged);
-				// This should still allow case changes b/c we test to make sure that they aren't referencing the same node.
-				if (bIsSame == false && Node->Input.GetName().IsEqual(NewName, ENameCase::IgnoreCase))
-				{
-					OutErrorMessage = FText::Format(LOCTEXT("NiagaraInputNameSameParameterWarn", "\"{0}\" is the name of another parameter."), FText::FromName(NewName));
-					return false;
-				}
-			}
-		}
-
-		// If dealing with Attributes, check to make sure that we don't conflict with any other attribute name.
-		if ((InputNodeBeingChanged != nullptr && InputNodeBeingChanged->Usage == ENiagaraInputNodeUsage::Attribute) || OutputNodeBeingChanged != nullptr)
-		{
-			TArray<UNiagaraNodeOutput*> OutputNodes;
-			Graph->GetNodesOfClass<UNiagaraNodeOutput>(OutputNodes);
-			for (UNiagaraNodeOutput* Node : OutputNodes)
-			{
-				for (const FNiagaraVariable& Output : Node->Outputs)
-				{
-					if (InputNodeBeingChanged != nullptr && Output.GetName().IsEqual(InputNodeBeingChanged->Input.GetName(), ENameCase::IgnoreCase))
-					{
-						continue;
-					}
-
-					if (Output.GetName().IsEqual(NewName, ENameCase::IgnoreCase))
-					{
-						OutErrorMessage = FText::Format(LOCTEXT("NiagaraInputNameSameAttributeWarn", "\"{0}\" is the name of another attribute. Hit \"Escape\" to cancel edit."), FText::FromName(NewName));
-						return false;
-					}
-				}
-			}
-		}
-	}
-	return true;
-}
-
 void UNiagaraNodeInput::OnRenameNode(const FString& NewName)
 {
 	UNiagaraGraph* Graph = CastChecked<UNiagaraGraph>(GetGraph());
@@ -295,10 +217,9 @@ void UNiagaraNodeInput::OnRenameNode(const FString& NewName)
 		{
 			Node->GetDataInterface()->Rename(*NewName);
 		}
-		Node->ReallocatePins();
+		Node->ReallocatePins(false);
+		Node->MarkNodeRequiresSynchronization("Input renamed", true);
 	}
-
-	Graph->MarkGraphRequiresSynchronization(TEXT("Input node renamed"));
 }
 
 TSharedPtr<SGraphNode> UNiagaraNodeInput::CreateVisualWidget()
diff --git a/Engine/Plugins/FX/Niagara/Source/NiagaraEditor/Private/NiagaraNodeOutput.cpp b/Engine/Plugins/FX/Niagara/Source/NiagaraEditor/Private/NiagaraNodeOutput.cpp
--- a/Engine/Plugins/FX/Niagara/Source/NiagaraEditor/Private/NiagaraNodeOutput.cpp
+++ b/Engine/Plugins/FX/Niagara/Source/NiagaraEditor/Private/NiagaraNodeOutput.cpp
@@ -13,6 +13,7 @@
 #include "Framework/MultiBox/MultiBoxBuilder.h"
 #include "Widgets/Layout/SBox.h"
 #include "Widgets/Input/SEditableTextBox.h"
+#include "Widgets/SEditableTextBoxWithVerification.h"
 #include "NiagaraEditorUtilities.h"
 
 #define LOCTEXT_NAMESPACE "NiagaraNodeOutput"
@@ -58,6 +59,10 @@ FText UNiagaraNodeOutput::GetPinNameText(UEdGraphPin* Pin) const
 	return FText::FromName(Pin->PinName);
 }
 
+bool UNiagaraNodeOutput::VerifyPinNameTextChanged(const FText& InText, FText& OutErrorMessage, UEdGraphPin* Pin)
+{
+	return FNiagaraEditorUtilities::VerifyNameChangeForInputOrOutputNode(*this, Pin->PinName, *InText.ToString(), OutErrorMessage);
+}
 
 void UNiagaraNodeOutput::PinNameTextCommitted(const FText& Text, ETextCommit::Type CommitType, UEdGraphPin* Pin)
 {
@@ -88,9 +93,10 @@ void UNiagaraNodeOutput::GetContextMenuActions(const FGraphNodeContextMenuBuilde
 				.WidthOverride(100)
 				.Padding(FMargin(5, 0, 0, 0))
 				[
-					SNew(SEditableTextBox)
+					SNew(SEditableTextBoxWithVerification)
 					.Text_UObject(this, &UNiagaraNodeOutput::GetPinNameText, Pin)
-					.OnTextCommitted_UObject(const_cast<UNiagaraNodeOutput*>(this), &UNiagaraNodeOutput::PinNameTextCommitted, Pin)
+					.OnVerifyTextChanged_UObject(this, &UNiagaraNodeOutput::VerifyPinNameTextChanged, Pin)
+					.OnTextCommitted_UObject(this, &UNiagaraNodeOutput::PinNameTextCommitted, Pin)
 				];
 			Context.MenuBuilder->AddWidget(RenameWidget, LOCTEXT("NameMenuItem", "Name"));
 		}
diff --git a/Engine/Plugins/FX/Niagara/Source/NiagaraEditor/Private/ViewModels/NiagaraScriptInputCollectionViewModel.cpp b/Engine/Plugins/FX/Niagara/Source/NiagaraEditor/Private/ViewModels/NiagaraScriptInputCollectionViewModel.cpp
--- a/Engine/Plugins/FX/Niagara/Source/NiagaraEditor/Private/ViewModels/NiagaraScriptInputCollectionViewModel.cpp
+++ b/Engine/Plugins/FX/Niagara/Source/NiagaraEditor/Private/ViewModels/NiagaraScriptInputCollectionViewModel.cpp
@@ -330,7 +330,7 @@ void FNiagaraScriptInputCollectionViewModel::RefreshParameterViewModels()
 				ParameterViewModel = MakeShareable(new FNiagaraScriptParameterViewModel(GraphVariable, *InputNode, EmitterDataInterface, ParameterEditMode));
 			}
 
-			ParameterViewModel->OnNameChanged().AddRaw(this, &FNiagaraScriptInputCollectionViewModel::OnParameterNameChanged, &GraphVariable);
+			ParameterViewModel->OnNameChanged().AddRaw(this, &FNiagaraScriptInputCollectionViewModel::OnParameterNameChanged, TWeakObjectPtr<UNiagaraNodeInput>(InputNode));
 			ParameterViewModel->OnTypeChanged().AddRaw(this, &FNiagaraScriptInputCollectionViewModel::OnParameterTypeChanged, &GraphVariable);
 			ParameterViewModel->OnDefaultValueChanged().AddRaw(this, &FNiagaraScriptInputCollectionViewModel::OnParameterValueChangedInternal, ParameterViewModel.ToSharedRef());
 			ParameterViewModels.Add(ParameterViewModel.ToSharedRef());
@@ -363,68 +363,13 @@ void FNiagaraScriptInputCollectionViewModel::OnGraphChanged(const struct FEdGrap
 	bNeedsRefresh = true;
 }
 
-void FNiagaraScriptInputCollectionViewModel::OnParameterNameChanged(FName OldName, FName NewName, FNiagaraVariable* ParameterVariable)
+void FNiagaraScriptInputCollectionViewModel::OnParameterNameChanged(FName OldName, FName NewName, TWeakObjectPtr<UNiagaraNodeInput> InputNodeWeak)
 {
-	TSet<FName> CurrentNames;
-	TArray<UNiagaraNodeInput*> InputNodes;
-	TArray<UNiagaraNodeInput*> InputNodesToRename;
-
-	// Check the existing input nodes and get a set of the current names, and find nodes with matching names to rename.
-	if (Graph.IsValid())
-	{
-		Graph->GetNodesOfClass(InputNodes);
-	}
-	
-	for (UNiagaraNodeInput* InputNode : InputNodes)
-	{
-		if (InputNode->Usage == ENiagaraInputNodeUsage::Parameter && &InputNode->Input != ParameterVariable)
-		{
-			if (InputNode->Input.GetName() == OldName)
-			{
-				InputNodesToRename.Add(InputNode);
-			}
-			else
-			{
-				CurrentNames.Add(InputNode->Input.GetName());
-			}
-		}
-	}
-
-	TSet<FName> SystemConstantNames = FNiagaraEditorUtilities::GetSystemConstantNames();
-
-	// Rename the nodes and notify the graph that they've changed.
-	FName UniqueNewName = FNiagaraUtilities::GetUniqueName(ParameterVariable->GetName(), CurrentNames.Union(SystemConstantNames));
-	if (ParameterVariable->GetName() != UniqueNewName)
-	{
-		ParameterVariable->SetName(UniqueNewName);
-	}
-	for (UNiagaraNodeInput* InputNodeToRename : InputNodesToRename)
-	{
-		InputNodeToRename->Modify();
-		InputNodeToRename->Input.SetName(UniqueNewName);
-	}
-
-	// Synchronize script variables...
-	for (TWeakObjectPtr<UNiagaraScript> Script : Scripts)
-	{
-		if (!Script.IsValid() || !Script->GetVMExecutableData().IsValid())
-		{
-			continue;
-		}
-
-		for (FNiagaraVariable& EmitterVariableToCheck : Script->GetVMExecutableData().Parameters.Parameters)
-		{
-			if (EmitterVariableToCheck.GetName() == ParameterVariable->GetName() && ParameterVariable != &EmitterVariableToCheck)
-			{
-				EmitterVariableToCheck = *ParameterVariable;
-				break;
-			}
-		}
-	}
-
-	if (Graph.IsValid())
+	UNiagaraNodeInput* InputNode = InputNodeWeak.Get();
+	if (InputNode != nullptr)
 	{
-		Graph->NotifyGraphChanged();
+		FScopedTransaction ScopedTransaction(LOCTEXT("EditInputName", "Edit input name"));
+		InputNode->OnRenameNode(NewName.ToString());
 	}
 }
 
diff --git a/Engine/Plugins/FX/Niagara/Source/NiagaraEditor/Private/ViewModels/NiagaraScriptInputCollectionViewModel.h b/Engine/Plugins/FX/Niagara/Source/NiagaraEditor/Private/ViewModels/NiagaraScriptInputCollectionViewModel.h
--- a/Engine/Plugins/FX/Niagara/Source/NiagaraEditor/Private/ViewModels/NiagaraScriptInputCollectionViewModel.h
+++ b/Engine/Plugins/FX/Niagara/Source/NiagaraEditor/Private/ViewModels/NiagaraScriptInputCollectionViewModel.h
@@ -51,7 +51,7 @@ private:
 	void OnGraphChanged(const struct FEdGraphEditAction& InAction);
 
 	/** Handles when the name on a parameter changes. */
-	void OnParameterNameChanged(FName OldName, FName NewName, FNiagaraVariable* ParameterVariable);
+	void OnParameterNameChanged(FName OldName, FName NewName, TWeakObjectPtr<UNiagaraNodeInput> InputNodeWeak);
 
 	/** Handles when the type on a parameter changes. */
 	void OnParameterTypeChanged(FNiagaraVariable* ParameterVariable);
diff --git a/Engine/Plugins/FX/Niagara/Source/NiagaraEditor/Private/ViewModels/NiagaraScriptOutputCollectionViewModel.cpp b/Engine/Plugins/FX/Niagara/Source/NiagaraEditor/Private/ViewModels/NiagaraScriptOutputCollectionViewModel.cpp
--- a/Engine/Plugins/FX/Niagara/Source/NiagaraEditor/Private/ViewModels/NiagaraScriptOutputCollectionViewModel.cpp
+++ b/Engine/Plugins/FX/Niagara/Source/NiagaraEditor/Private/ViewModels/NiagaraScriptOutputCollectionViewModel.cpp
@@ -269,26 +269,14 @@ void FNiagaraScriptOutputCollectionViewModel::OnGraphChanged(const struct FEdGra
 
 void FNiagaraScriptOutputCollectionViewModel::OnParameterNameChanged(FName OldName, FName NewName, FNiagaraVariable* ParameterVariable)
 {
-	TSet<FName> CurrentNames;
+	FScopedTransaction ScopedTransaction(LOCTEXT("EditOutputName", "Edit output name"));
 
 	if (OutputNode.IsValid() == false)
 	{
 		return;
 	}
 
-	for (FNiagaraVariable& OutputVariable : OutputNode->Outputs)
-	{
-		if (&OutputVariable != ParameterVariable)
-		{
-			CurrentNames.Add(OutputVariable.GetName());
-		}
-	}
-
-	// If the new name matches a current output name, rename it to something unique.
-	if (CurrentNames.Contains(ParameterVariable->GetName()))
-	{
-		ParameterVariable->SetName(FNiagaraUtilities::GetUniqueName(ParameterVariable->GetName(), CurrentNames));
-	}
+	ParameterVariable->SetName(NewName);
 
 	// Make sure to edit the pin name so that when we regenerate the node the 
 	// links are preserved.
diff --git a/Engine/Plugins/FX/Niagara/Source/NiagaraEditor/Private/ViewModels/NiagaraScriptParameterViewModel.cpp b/Engine/Plugins/FX/Niagara/Source/NiagaraEditor/Private/ViewModels/NiagaraScriptParameterViewModel.cpp
--- a/Engine/Plugins/FX/Niagara/Source/NiagaraEditor/Private/ViewModels/NiagaraScriptParameterViewModel.cpp
+++ b/Engine/Plugins/FX/Niagara/Source/NiagaraEditor/Private/ViewModels/NiagaraScriptParameterViewModel.cpp
@@ -83,16 +83,14 @@ void FNiagaraScriptParameterViewModel::NameTextComitted(const FText& Name, EText
 	FName OldName = GraphVariable->GetName();
 	if (!OldName.IsEqual(NewName, ENameCase::CaseSensitive))
 	{
-		FScopedTransaction ScopedTransaction(LOCTEXT("EditInputName", "Edit input name"));
-		GraphVariableOwner->Modify();
-		GraphVariable->SetName(NewName);
 		OnNameChangedDelegate.Broadcast(OldName, NewName);
 	}
 }
 
 bool FNiagaraScriptParameterViewModel::VerifyNodeNameTextChanged(const FText& NewText, FText& OutErrorMessage)
 {
-	return UNiagaraNodeInput::VerifyNodeRenameTextCommit(NewText, Cast<UNiagaraNode>(GraphVariableOwner.Get()), OutErrorMessage);
+	UNiagaraNode* OwningNode = Cast<UNiagaraNode>(GraphVariableOwner);
+	return OwningNode == nullptr || FNiagaraEditorUtilities::VerifyNameChangeForInputOrOutputNode(*OwningNode, GraphVariable->GetName(), *NewText.ToString(), OutErrorMessage);
 }
 
 TSharedPtr<FNiagaraTypeDefinition> FNiagaraScriptParameterViewModel::GetType() const
diff --git a/Engine/Plugins/FX/Niagara/Source/NiagaraEditor/Private/Widgets/SEditableTextBoxWithVerification.cpp b/Engine/Plugins/FX/Niagara/Source/NiagaraEditor/Private/Widgets/SEditableTextBoxWithVerification.cpp
new file mode 100644
--- /dev/null
+++ b/Engine/Plugins/FX/Niagara/Source/NiagaraEditor/Private/Widgets/SEditableTextBoxWithVerification.cpp
@@ -0,0 +1,49 @@
+// Copyright 1998-2019 Epic Games, Inc. All Rights Reserved.
+
+#include "Widgets/SEditableTextBoxWithVerification.h"
+
+void SEditableTextBoxWithVerification::Construct(const FArguments& InArgs)
+{
+	OnVerifyTextChangedDelegate = InArgs._OnVerifyTextChanged;
+	OnTextCommittedDelegate = InArgs._OnTextCommitted;
+	ChildSlot
+	[
+		SAssignNew(TextBox, SEditableTextBox)
+			.Text(InArgs._Text)
+			.OnTextChanged(this, &SEditableTextBoxWithVerification::OnTextChanged)
+			.OnTextCommitted(this, &SEditableTextBoxWithVerification::OnTextBoxCommitted)
+	];
+}
+
+void SEditableTextBoxWithVerification::OnTextChanged(const FText& InText)
+{
+	FText OutErrorMessage;
+	if (OnVerifyTextChangedDelegate.IsBound() && OnVerifyTextChangedDelegate.Execute(InText, OutErrorMessage) == false)
+	{
+		TextBox->SetError(OutErrorMessage);
+	}
+	else
+	{
+		TextBox->SetError(FText::GetEmpty());
+	}
+}
+
+void SEditableTextBoxWithVerification::OnTextBoxCommitted(const FText& InText, ETextCommit::Type InCommitType)
+{
+	FText OutErrorMessage;
+	if (OnVerifyTextChangedDelegate.IsBound() && OnVerifyTextChangedDelegate.Execute(InText, OutErrorMessage) == false)
+	{
+		if (InCommitType == ETextCommit::OnEnter)
+		{
+			TextBox->SetError(OutErrorMessage);
+		}
+		else
+		{
+			TextBox->SetError(FText::GetEmpty());
+		}
+	}
+	else
+	{
+		OnTextCommittedDelegate.ExecuteIfBound(InText, InCommitType);
+	}
+}
\ No newline at end of file
diff --git a/Engine/Plugins/FX/Niagara/Source/NiagaraEditor/Private/Widgets/SEditableTextBoxWithVerification.h b/Engine/Plugins/FX/Niagara/Source/NiagaraEditor/Private/Widgets/SEditableTextBoxWithVerification.h
new file mode 100644
--- /dev/null
+++ b/Engine/Plugins/FX/Niagara/Source/NiagaraEditor/Private/Widgets/SEditableTextBoxWithVerification.h
@@ -0,0 +1,38 @@
+// Copyright 1998-2019 Epic Games, Inc. All Rights Reserved.
+
+#pragma once
+
+#include "Widgets/SCompoundWidget.h"
+
+/** An editable text block which validates input with a delegate. */
+class SEditableTextBoxWithVerification : public SCompoundWidget
+{
+public:
+	DECLARE_DELEGATE_RetVal_TwoParams(bool, FOnVerifyTextChanged, const FText&, FText&)
+
+public:
+	SLATE_BEGIN_ARGS(SEditableTextBoxWithVerification)
+	{}
+		/** The text displayed in this text block */
+		SLATE_ATTRIBUTE(FText, Text)
+
+		/** Called whenever the text is changed interactively by the user */
+		SLATE_EVENT(FOnVerifyTextChanged, OnVerifyTextChanged)
+
+		/** Called when the text has been committed to the text box.  Text will not be committed if it fails verification. */
+		SLATE_EVENT(FOnTextCommitted, OnTextCommitted);
+
+	SLATE_END_ARGS();
+
+	void Construct(const FArguments& InArgs);
+
+private:
+	void OnTextChanged(const FText& InText);
+
+	void OnTextBoxCommitted(const FText& InText, ETextCommit::Type InCommitType);
+
+private:
+	TSharedPtr<SEditableTextBox> TextBox;
+	FOnVerifyTextChanged OnVerifyTextChangedDelegate;
+	FOnTextCommitted OnTextCommittedDelegate;
+};
diff --git a/Engine/Plugins/FX/Niagara/Source/NiagaraEditor/Private/Widgets/SNiagaraScriptGraph.cpp b/Engine/Plugins/FX/Niagara/Source/NiagaraEditor/Private/Widgets/SNiagaraScriptGraph.cpp
--- a/Engine/Plugins/FX/Niagara/Source/NiagaraEditor/Private/Widgets/SNiagaraScriptGraph.cpp
+++ b/Engine/Plugins/FX/Niagara/Source/NiagaraEditor/Private/Widgets/SNiagaraScriptGraph.cpp
@@ -239,9 +239,10 @@ void SNiagaraScriptGraph::OnNodeTitleCommitted(const FText& NewText, ETextCommit
 bool SNiagaraScriptGraph::OnVerifyNodeTextCommit(const FText& NewText, UEdGraphNode* NodeBeingChanged, FText& OutErrorMessage)
 {
 	bool bValid = true;
-	if (NodeBeingChanged->IsA(UNiagaraNodeInput::StaticClass()))
+	UNiagaraNodeInput* InputNodeBeingChanged = Cast<UNiagaraNodeInput>(NodeBeingChanged);
+	if (InputNodeBeingChanged != nullptr)
 	{
-		return UNiagaraNodeInput::VerifyNodeRenameTextCommit(NewText, Cast<UNiagaraNodeInput>(NodeBeingChanged), OutErrorMessage);
+		return FNiagaraEditorUtilities::VerifyNameChangeForInputOrOutputNode(*InputNodeBeingChanged, InputNodeBeingChanged->Input.GetName(), *NewText.ToString(), OutErrorMessage);
 	}
 	return bValid;
 }
diff --git a/Engine/Plugins/FX/Niagara/Source/NiagaraEditor/Public/NiagaraEditorUtilities.h b/Engine/Plugins/FX/Niagara/Source/NiagaraEditor/Public/NiagaraEditorUtilities.h
--- a/Engine/Plugins/FX/Niagara/Source/NiagaraEditor/Public/NiagaraEditorUtilities.h
+++ b/Engine/Plugins/FX/Niagara/Source/NiagaraEditor/Public/NiagaraEditorUtilities.h
@@ -197,4 +197,6 @@ namespace FNiagaraEditorUtilities
 	/** Kills all system instances using the referenced system. */
 	void KillSystemInstances(const UNiagaraSystem& System);
 
+
+	bool VerifyNameChangeForInputOrOutputNode(UNiagaraNode& NodeBeingChanged, FName OldName, FName NewName, FText& OutErrorMessage);
 };
diff --git a/Engine/Plugins/FX/Niagara/Source/NiagaraEditor/Public/NiagaraNodeInput.h b/Engine/Plugins/FX/Niagara/Source/NiagaraEditor/Public/NiagaraNodeInput.h
--- a/Engine/Plugins/FX/Niagara/Source/NiagaraEditor/Public/NiagaraNodeInput.h
+++ b/Engine/Plugins/FX/Niagara/Source/NiagaraEditor/Public/NiagaraNodeInput.h
@@ -97,11 +97,6 @@ public:
 	static const FLinearColor TitleColor_Attribute;
 	static const FLinearColor TitleColor_Constant;
 
-	/** Verify that the text about to be committed will be valid and doesn't duplicate existing variables based on type.
-	Type validation is done based on the input UObject type, preferably a UNiagaraNodeInput or UNiagaraNodeOutput.*/
-	static bool VerifyNodeRenameTextCommit(const FText& NewText, UNiagaraNode* NodeBeingChanged, FText& OutErrorMessage);
-
-
 	/** Generate a unique name based off of the existing names in the system.*/
 	static FName GenerateUniqueName(const UNiagaraGraph* Graph, FName& ProposedName, ENiagaraInputNodeUsage Usage);
 
diff --git a/Engine/Plugins/FX/Niagara/Source/NiagaraEditor/Public/NiagaraNodeOutput.h b/Engine/Plugins/FX/Niagara/Source/NiagaraEditor/Public/NiagaraNodeOutput.h
--- a/Engine/Plugins/FX/Niagara/Source/NiagaraEditor/Public/NiagaraNodeOutput.h
+++ b/Engine/Plugins/FX/Niagara/Source/NiagaraEditor/Public/NiagaraNodeOutput.h
@@ -69,6 +69,8 @@ private:
 	/** Gets the display text for a pin. */
 	FText GetPinNameText(UEdGraphPin* Pin) const;
 
+	bool VerifyPinNameTextChanged(const FText& InText, FText& OutErrorMessage, UEdGraphPin* Pin);
+
 	/** Called when a pin's name text is committed. */
 	void PinNameTextCommitted(const FText& Text, ETextCommit::Type CommitType, UEdGraphPin* Pin);
 };
