Niagara - Fix non-unity issues.

diff --git a/Engine/Plugins/FX/Niagara/Source/NiagaraEditor/Private/Widgets/SEditableTextBoxWithVerification.cpp b/Engine/Plugins/FX/Niagara/Source/NiagaraEditor/Private/Widgets/SEditableTextBoxWithVerification.cpp
--- a/Engine/Plugins/FX/Niagara/Source/NiagaraEditor/Private/Widgets/SEditableTextBoxWithVerification.cpp
+++ b/Engine/Plugins/FX/Niagara/Source/NiagaraEditor/Private/Widgets/SEditableTextBoxWithVerification.cpp
@@ -1,6 +1,7 @@
 // Copyright 1998-2019 Epic Games, Inc. All Rights Reserved.
 
 #include "Widgets/SEditableTextBoxWithVerification.h"
+#include "Widgets/Input/SEditableTextBox.h"
 
 void SEditableTextBoxWithVerification::Construct(const FArguments& InArgs)
 {
diff --git a/Engine/Plugins/FX/Niagara/Source/NiagaraEditor/Private/Widgets/SEditableTextBoxWithVerification.h b/Engine/Plugins/FX/Niagara/Source/NiagaraEditor/Private/Widgets/SEditableTextBoxWithVerification.h
--- a/Engine/Plugins/FX/Niagara/Source/NiagaraEditor/Private/Widgets/SEditableTextBoxWithVerification.h
+++ b/Engine/Plugins/FX/Niagara/Source/NiagaraEditor/Private/Widgets/SEditableTextBoxWithVerification.h
@@ -3,6 +3,10 @@
 #pragma once
 
 #include "Widgets/SCompoundWidget.h"
+#include "Widgets/DeclarativeSyntaxSupport.h"
+#include "Delegates/DelegateCombinations.h"
+
+class SEditableTextBox;
 
 /** An editable text block which validates input with a delegate. */
 class SEditableTextBoxWithVerification : public SCompoundWidget
