#UE4 Non-unity fix

diff --git a/Engine/Plugins/FX/Niagara/Source/NiagaraEditor/Private/Widgets/SEditableTextBoxWithVerification.h b/Engine/Plugins/FX/Niagara/Source/NiagaraEditor/Private/Widgets/SEditableTextBoxWithVerification.h
--- a/Engine/Plugins/FX/Niagara/Source/NiagaraEditor/Private/Widgets/SEditableTextBoxWithVerification.h
+++ b/Engine/Plugins/FX/Niagara/Source/NiagaraEditor/Private/Widgets/SEditableTextBoxWithVerification.h
@@ -5,6 +5,7 @@
 #include "Widgets/SCompoundWidget.h"
 #include "Widgets/DeclarativeSyntaxSupport.h"
 #include "Delegates/DelegateCombinations.h"
+#include "Framework/SlateDelegates.h"
 
 class SEditableTextBox;
 
