#UE4 Fix concurrent save problems in Niagara

diff --git a/Engine/Plugins/FX/Niagara/Source/Niagara/Private/NiagaraScript.cpp b/Engine/Plugins/FX/Niagara/Source/Niagara/Private/NiagaraScript.cpp
--- a/Engine/Plugins/FX/Niagara/Source/Niagara/Private/NiagaraScript.cpp
+++ b/Engine/Plugins/FX/Niagara/Source/Niagara/Private/NiagaraScript.cpp
@@ -1124,10 +1124,33 @@ bool UNiagaraScript::IsCachedCookedPlatformDataLoaded(const ITargetPlatform* Tar
 {
 	if (ShouldCacheShadersForCooking())
 	{
+		bool bHasOutstandingCompilationRequests = false;
 		if (UNiagaraSystem* SystemOwner = FindRootSystem())
 		{
-			return !SystemOwner->HasOutstandingCompilationRequests();
+			bHasOutstandingCompilationRequests = SystemOwner->HasOutstandingCompilationRequests();
 		}
+
+		if (!bHasOutstandingCompilationRequests)
+		{
+			TArray<FName> DesiredShaderFormats;
+			TargetPlatform->GetAllTargetedShaderFormats(DesiredShaderFormats);
+
+			const TArray<FNiagaraShaderScript*>* CachedScriptResourcesForPlatform = CachedScriptResourcesForCooking.Find(TargetPlatform);
+			if (CachedScriptResourcesForPlatform)
+			{
+				for (const auto& MaterialResource : *CachedScriptResourcesForPlatform)
+				{
+					if (MaterialResource->IsCompilationFinished() == false)
+					{
+						return false;
+					}
+				}
+
+				return true;
+			}
+		}
+
+		return false;
 	}
 
 	return true;
diff --git a/Engine/Plugins/FX/Niagara/Source/NiagaraShader/Private/NiagaraShader.cpp b/Engine/Plugins/FX/Niagara/Source/NiagaraShader/Private/NiagaraShader.cpp
--- a/Engine/Plugins/FX/Niagara/Source/NiagaraShader/Private/NiagaraShader.cpp
+++ b/Engine/Plugins/FX/Niagara/Source/NiagaraShader/Private/NiagaraShader.cpp
@@ -1165,13 +1165,7 @@ void FNiagaraShaderMap::Serialize(FArchive& Ar, bool bInlineShaderResources)
 
 	Ar << DebugDescription;
 
-	if (Ar.IsSaving())
-	{
-		TShaderMap<FNiagaraShaderType>::SerializeInline(Ar, bInlineShaderResources, false, false);
-		RegisterSerializedShaders(false);
-	}
-
-	if (Ar.IsLoading())
+	if (Ar.IsSaving() || Ar.IsLoading())
 	{
 		TShaderMap<FNiagaraShaderType>::SerializeInline(Ar, bInlineShaderResources, false, false);
 	}
diff --git a/Engine/Plugins/FX/Niagara/Source/NiagaraShader/Private/NiagaraShared.cpp b/Engine/Plugins/FX/Niagara/Source/NiagaraShader/Private/NiagaraShared.cpp
--- a/Engine/Plugins/FX/Niagara/Source/NiagaraShader/Private/NiagaraShared.cpp
+++ b/Engine/Plugins/FX/Niagara/Source/NiagaraShader/Private/NiagaraShared.cpp
@@ -417,7 +417,6 @@ NIAGARASHADER_API  FNiagaraShader* FNiagaraShaderScript::GetShaderGameThread() c
 
 void FNiagaraShaderScript::GetShaderMapIDsWithUnfinishedCompilation(TArray<int32>& ShaderMapIds)
 {
-	check(IsInGameThread());
 	// Build an array of the shader map Id's are not finished compiling.
 	if (GameThreadShaderMap && GameThreadShaderMap.IsValid() && !GameThreadShaderMap->IsCompilationFinalized())
 	{
