Niagara - Fixing crashes and memory safety issues with undo/redo while duplicating emitters using the stomp allocator.
+ Make view model delegates safer by changing raw delegate bindings to shared pointer delegate bindings.
+ Properly modify the system when duplicating emitters.
+ Remove the update context from the system property change since we can't safely reset systems in the world from the context.
+ In the system view model, clear out data from emitter handle view models and sequencer before refreshing view models and the system instance.
+ In the system view model, refactor some of the emitter handle delegate handlers to removed unused parameters and to prevent passing arround potentially stale view models.
+ Prevent binding the same command to the same command list multiple times in the niagara graph.

diff --git a/Engine/Plugins/FX/Niagara/Source/NiagaraEditor/Private/Toolkits/NiagaraScriptToolkit.cpp b/Engine/Plugins/FX/Niagara/Source/NiagaraEditor/Private/Toolkits/NiagaraScriptToolkit.cpp
--- a/Engine/Plugins/FX/Niagara/Source/NiagaraEditor/Private/Toolkits/NiagaraScriptToolkit.cpp
+++ b/Engine/Plugins/FX/Niagara/Source/NiagaraEditor/Private/Toolkits/NiagaraScriptToolkit.cpp
@@ -147,7 +147,8 @@ void FNiagaraScriptToolkit::Initialize( const EToolkitMode::Type Mode, const TSh
 	const FGuid MessageLogGuidKey = FGuid::NewGuid();
 	NiagaraMessageLogViewModel = MakeShared<FNiagaraMessageLogViewModel>(GetNiagaraScriptMessageLogName(EditedNiagaraScript), MessageLogGuidKey, NiagaraMessageLog);
 
-	ScriptViewModel = MakeShareable(new FNiagaraStandaloneScriptViewModel(EditedNiagaraScript, DisplayName, ENiagaraParameterEditMode::EditAll, NiagaraMessageLogViewModel, OriginalNiagaraScript, MessageLogGuidKey));
+	ScriptViewModel = MakeShareable(new FNiagaraStandaloneScriptViewModel(DisplayName, ENiagaraParameterEditMode::EditAll, NiagaraMessageLogViewModel, MessageLogGuidKey));
+	ScriptViewModel->Initialize(EditedNiagaraScript, OriginalNiagaraScript);
 
 	OnEditedScriptGraphChangedHandle = ScriptViewModel->GetGraphViewModel()->GetGraph()->AddOnGraphNeedsRecompileHandler(
 		FOnGraphChanged::FDelegate::CreateRaw(this, &FNiagaraScriptToolkit::OnEditedScriptGraphChanged));
@@ -698,7 +699,8 @@ void FNiagaraScriptToolkit::UpdateOriginalNiagaraScript()
 			TSharedPtr<FNiagaraScriptViewModel> AffectedScriptViewModel = FNiagaraScriptViewModel::GetExistingViewModelForObject(Script);
 			if (!AffectedScriptViewModel.IsValid())
 			{
-				AffectedScriptViewModel = MakeShareable(new FNiagaraScriptViewModel(Script, FText::FromString(Script->GetName()), ENiagaraParameterEditMode::EditValueOnly));
+				AffectedScriptViewModel = MakeShareable(new FNiagaraScriptViewModel(FText::FromString(Script->GetName()), ENiagaraParameterEditMode::EditValueOnly));
+				AffectedScriptViewModel->SetScript(Script);
 			}
 			AffectedScriptViewModel->CompileStandaloneScript();
 		}
diff --git a/Engine/Plugins/FX/Niagara/Source/NiagaraEditor/Private/ViewModels/NiagaraEmitterHandleViewModel.cpp b/Engine/Plugins/FX/Niagara/Source/NiagaraEditor/Private/ViewModels/NiagaraEmitterHandleViewModel.cpp
--- a/Engine/Plugins/FX/Niagara/Source/NiagaraEditor/Private/ViewModels/NiagaraEmitterHandleViewModel.cpp
+++ b/Engine/Plugins/FX/Niagara/Source/NiagaraEditor/Private/ViewModels/NiagaraEmitterHandleViewModel.cpp
@@ -30,6 +30,11 @@ FNiagaraEmitterHandleViewModel::FNiagaraEmitterHandleViewModel()
 {
 }
 
+bool FNiagaraEmitterHandleViewModel::IsValid() const
+{
+	return EmitterHandle != nullptr;
+}
+
 void FNiagaraEmitterHandleViewModel::Cleanup()
 {
 	EmitterViewModel->Cleanup();
diff --git a/Engine/Plugins/FX/Niagara/Source/NiagaraEditor/Private/ViewModels/NiagaraEmitterViewModel.cpp b/Engine/Plugins/FX/Niagara/Source/NiagaraEditor/Private/ViewModels/NiagaraEmitterViewModel.cpp
--- a/Engine/Plugins/FX/Niagara/Source/NiagaraEditor/Private/ViewModels/NiagaraEmitterViewModel.cpp
+++ b/Engine/Plugins/FX/Niagara/Source/NiagaraEditor/Private/ViewModels/NiagaraEmitterViewModel.cpp
@@ -33,9 +33,7 @@ namespace NiagaraCommands
 const float Megabyte = 1024.0f * 1024.0f;
 
 FNiagaraEmitterViewModel::FNiagaraEmitterViewModel()
-	: Emitter(nullptr)
-	, Simulation(nullptr)
-	, SharedScriptViewModel(MakeShared<FNiagaraScriptViewModel>(LOCTEXT("SharedDisplayName", "Graph"), ENiagaraParameterEditMode::EditAll))
+	: SharedScriptViewModel(MakeShareable(new FNiagaraScriptViewModel(LOCTEXT("SharedDisplayName", "Graph"), ENiagaraParameterEditMode::EditAll)))
 	, bUpdatingSelectionInternally(false)
 	, ExecutionStateEnum(StaticEnum<ENiagaraExecutionState>())
 {	
@@ -97,8 +95,8 @@ void FNiagaraEmitterViewModel::SetEmitter(UNiagaraEmitter* InEmitter)
 
 	if (Emitter.IsValid())
 	{
-		Emitter->OnEmitterVMCompiled().AddRaw(this, &FNiagaraEmitterViewModel::OnVMCompiled);
-		Emitter->OnPropertiesChanged().AddRaw(this, &FNiagaraEmitterViewModel::OnEmitterPropertiesChanged);
+		Emitter->OnEmitterVMCompiled().AddSP(this, &FNiagaraEmitterViewModel::OnVMCompiled);
+		Emitter->OnPropertiesChanged().AddSP(this, &FNiagaraEmitterViewModel::OnEmitterPropertiesChanged);
 	}
 
 	AddScriptEventHandlers();
@@ -371,16 +369,16 @@ void FNiagaraEmitterViewModel::AddScriptEventHandlers()
 		{
 			UNiagaraScriptSource* ScriptSource = CastChecked<UNiagaraScriptSource>(Script->GetSource());
 			FDelegateHandle OnGraphChangedHandle = ScriptSource->NodeGraph->AddOnGraphChangedHandler(
-				FOnGraphChanged::FDelegate::CreateRaw<FNiagaraEmitterViewModel, const UNiagaraScript&>(this, &FNiagaraEmitterViewModel::ScriptGraphChanged, *Script));
+				FOnGraphChanged::FDelegate::CreateSP<FNiagaraEmitterViewModel, const UNiagaraScript&>(this->AsShared(), &FNiagaraEmitterViewModel::ScriptGraphChanged, *Script));
 			FDelegateHandle OnGraphNeedRecompileHandle = ScriptSource->NodeGraph->AddOnGraphNeedsRecompileHandler(
-				FOnGraphChanged::FDelegate::CreateRaw<FNiagaraEmitterViewModel, const UNiagaraScript&>(this, &FNiagaraEmitterViewModel::ScriptGraphChanged, *Script));
+				FOnGraphChanged::FDelegate::CreateSP<FNiagaraEmitterViewModel, const UNiagaraScript&>(this->AsShared(), &FNiagaraEmitterViewModel::ScriptGraphChanged, *Script));
 
 			ScriptToOnGraphChangedHandleMap.Add(FObjectKey(Script), OnGraphChangedHandle);
 			ScriptToRecompileHandleMap.Add(FObjectKey(Script), OnGraphNeedRecompileHandle);
 
 			FDelegateHandle OnParameterStoreChangedHandle = Script->RapidIterationParameters.AddOnChangedHandler(
-				FNiagaraParameterStore::FOnChanged::FDelegate::CreateRaw<FNiagaraEmitterViewModel, const FNiagaraParameterStore&, const UNiagaraScript&>(
-					this, &FNiagaraEmitterViewModel::ScriptParameterStoreChanged, Script->RapidIterationParameters, *Script));
+				FNiagaraParameterStore::FOnChanged::FDelegate::CreateSP<FNiagaraEmitterViewModel, const FNiagaraParameterStore&, const UNiagaraScript&>(
+					this->AsShared(), &FNiagaraEmitterViewModel::ScriptParameterStoreChanged, Script->RapidIterationParameters, *Script));
 			ScriptToOnParameterStoreChangedHandleMap.Add(FObjectKey(Script), OnParameterStoreChangedHandle);
 		}
 	}
diff --git a/Engine/Plugins/FX/Niagara/Source/NiagaraEditor/Private/ViewModels/NiagaraScriptGraphViewModel.cpp b/Engine/Plugins/FX/Niagara/Source/NiagaraEditor/Private/ViewModels/NiagaraScriptGraphViewModel.cpp
--- a/Engine/Plugins/FX/Niagara/Source/NiagaraEditor/Private/ViewModels/NiagaraScriptGraphViewModel.cpp
+++ b/Engine/Plugins/FX/Niagara/Source/NiagaraEditor/Private/ViewModels/NiagaraScriptGraphViewModel.cpp
@@ -17,17 +17,6 @@
 
 #define LOCTEXT_NAMESPACE "NiagaraScriptGraphViewModel"
 
-FNiagaraScriptGraphViewModel::FNiagaraScriptGraphViewModel(UNiagaraScriptSource* InScriptSrc, FText InDisplayName)
-	: ScriptSource(InScriptSrc)
-	, DisplayName(InDisplayName)
-	, Commands(MakeShareable(new FUICommandList()))
-	, NodeSelection(MakeShareable(new FNiagaraObjectSelection()))
-{
-	SetupCommands();
-	GEditor->RegisterForUndo(this);
-	ErrorColor = FEditorStyle::GetColor("ErrorReporting.BackgroundColor");
-}
-
 FNiagaraScriptGraphViewModel::FNiagaraScriptGraphViewModel(FText InDisplayName)
 	: DisplayName(InDisplayName)
 	, Commands(MakeShareable(new FUICommandList()))
diff --git a/Engine/Plugins/FX/Niagara/Source/NiagaraEditor/Private/ViewModels/NiagaraScriptGraphViewModel.h b/Engine/Plugins/FX/Niagara/Source/NiagaraEditor/Private/ViewModels/NiagaraScriptGraphViewModel.h
--- a/Engine/Plugins/FX/Niagara/Source/NiagaraEditor/Private/ViewModels/NiagaraScriptGraphViewModel.h
+++ b/Engine/Plugins/FX/Niagara/Source/NiagaraEditor/Private/ViewModels/NiagaraScriptGraphViewModel.h
@@ -22,9 +22,7 @@ public:
 	DECLARE_MULTICAST_DELEGATE(FOnGraphChanged);
 
 public:
-	/** Create a new view model with the supplied script. */
-	FNiagaraScriptGraphViewModel(UNiagaraScriptSource* InScriptSource, FText InDisplayName);
-
+	/** Create a new view model with the supplied  display name. */
 	FNiagaraScriptGraphViewModel(FText InDisplayName);
 
 	~FNiagaraScriptGraphViewModel();
diff --git a/Engine/Plugins/FX/Niagara/Source/NiagaraEditor/Private/ViewModels/NiagaraScriptInputCollectionViewModel.cpp b/Engine/Plugins/FX/Niagara/Source/NiagaraEditor/Private/ViewModels/NiagaraScriptInputCollectionViewModel.cpp
--- a/Engine/Plugins/FX/Niagara/Source/NiagaraEditor/Private/ViewModels/NiagaraScriptInputCollectionViewModel.cpp
+++ b/Engine/Plugins/FX/Niagara/Source/NiagaraEditor/Private/ViewModels/NiagaraScriptInputCollectionViewModel.cpp
@@ -18,41 +18,10 @@
 
 FText DisplayNameFormat = NSLOCTEXT("ScriptInputCollection", "DisplayNameFormat", "{0} Inputs");
 
-FNiagaraScriptInputCollectionViewModel::FNiagaraScriptInputCollectionViewModel(UNiagaraScript* InScript, FText InDisplayName, ENiagaraParameterEditMode InParameterEditMode)
-	: FNiagaraParameterCollectionViewModel(InParameterEditMode)
-	, DisplayName(FText::Format(DisplayNameFormat, InDisplayName))
-{
-	if (InScript)
-	{
-		Scripts.Add(InScript);
-	}
-
-	if (Scripts.Num() == 1 && Scripts[0].IsValid() && Scripts[0]->GetSource())
-	{
-		Graph = Cast<UNiagaraScriptSource>(Scripts[0]->GetSource())->NodeGraph;
-		bCanHaveNumericParameters = (Scripts[0]->IsStandaloneScript());
-	}
-	else
-	{
-		Graph = nullptr;
-		bCanHaveNumericParameters = true;
-	}
-
-	RefreshParameterViewModels();
-
-	if (Graph.IsValid())
-	{
-		OnGraphChangedHandle = Graph->AddOnGraphChangedHandler(
-			FOnGraphChanged::FDelegate::CreateRaw(this, &FNiagaraScriptInputCollectionViewModel::OnGraphChanged));
-		OnRecompileHandle = Graph->AddOnGraphNeedsRecompileHandler(
-			FOnGraphChanged::FDelegate::CreateRaw(this, &FNiagaraScriptInputCollectionViewModel::OnGraphChanged));
-	}
-}
-
-
 FNiagaraScriptInputCollectionViewModel::FNiagaraScriptInputCollectionViewModel(FText InDisplayName, ENiagaraParameterEditMode InParameterEditMode)
 	: FNiagaraParameterCollectionViewModel(InParameterEditMode)
 	, DisplayName(FText::Format(DisplayNameFormat, InDisplayName))
+	, bCanHaveNumericParameters(true)
 {
 }
 
@@ -97,9 +66,9 @@ void FNiagaraScriptInputCollectionViewModel::SetScripts(TArray<UNiagaraScript*>
 	{
 		Graph = Cast<UNiagaraScriptSource>(Scripts[0]->GetSource())->NodeGraph;
 		OnGraphChangedHandle = Graph->AddOnGraphChangedHandler(
-			FOnGraphChanged::FDelegate::CreateRaw(this, &FNiagaraScriptInputCollectionViewModel::OnGraphChanged));
+			FOnGraphChanged::FDelegate::CreateSP(this, &FNiagaraScriptInputCollectionViewModel::OnGraphChanged));
 		OnRecompileHandle = Graph->AddOnGraphNeedsRecompileHandler(
-			FOnGraphChanged::FDelegate::CreateRaw(this, &FNiagaraScriptInputCollectionViewModel::OnGraphChanged));
+			FOnGraphChanged::FDelegate::CreateSP(this, &FNiagaraScriptInputCollectionViewModel::OnGraphChanged));
 		bCanHaveNumericParameters = Scripts[0]->IsStandaloneScript();
 	}
 	else
@@ -330,9 +299,9 @@ void FNiagaraScriptInputCollectionViewModel::RefreshParameterViewModels()
 				ParameterViewModel = MakeShareable(new FNiagaraScriptParameterViewModel(GraphVariable, *InputNode, EmitterDataInterface, ParameterEditMode));
 			}
 
-			ParameterViewModel->OnNameChanged().AddRaw(this, &FNiagaraScriptInputCollectionViewModel::OnParameterNameChanged, TWeakObjectPtr<UNiagaraNodeInput>(InputNode));
-			ParameterViewModel->OnTypeChanged().AddRaw(this, &FNiagaraScriptInputCollectionViewModel::OnParameterTypeChanged, &GraphVariable);
-			ParameterViewModel->OnDefaultValueChanged().AddRaw(this, &FNiagaraScriptInputCollectionViewModel::OnParameterValueChangedInternal, ParameterViewModel.ToSharedRef());
+			ParameterViewModel->OnNameChanged().AddSP(this, &FNiagaraScriptInputCollectionViewModel::OnParameterNameChanged, TWeakObjectPtr<UNiagaraNodeInput>(InputNode));
+			ParameterViewModel->OnTypeChanged().AddSP(this, &FNiagaraScriptInputCollectionViewModel::OnParameterTypeChanged, &GraphVariable);
+			ParameterViewModel->OnDefaultValueChanged().AddSP(this, &FNiagaraScriptInputCollectionViewModel::OnParameterValueChangedInternal, ParameterViewModel.ToSharedRef());
 			ParameterViewModels.Add(ParameterViewModel.ToSharedRef());
 			AddedInputNames.Add(InputNode->Input.GetName());
 		}
diff --git a/Engine/Plugins/FX/Niagara/Source/NiagaraEditor/Private/ViewModels/NiagaraScriptInputCollectionViewModel.h b/Engine/Plugins/FX/Niagara/Source/NiagaraEditor/Private/ViewModels/NiagaraScriptInputCollectionViewModel.h
--- a/Engine/Plugins/FX/Niagara/Source/NiagaraEditor/Private/ViewModels/NiagaraScriptInputCollectionViewModel.h
+++ b/Engine/Plugins/FX/Niagara/Source/NiagaraEditor/Private/ViewModels/NiagaraScriptInputCollectionViewModel.h
@@ -16,7 +16,6 @@ class UNiagaraEmitter;
 class FNiagaraScriptInputCollectionViewModel : public FNiagaraParameterCollectionViewModel, public TSharedFromThis<FNiagaraScriptInputCollectionViewModel>
 {
 public:
-	FNiagaraScriptInputCollectionViewModel(UNiagaraScript* InScript, FText InDisplayName, ENiagaraParameterEditMode InParameterEditMode);
 	FNiagaraScriptInputCollectionViewModel(FText InDisplayName, ENiagaraParameterEditMode InParameterEditMode);
 
 	~FNiagaraScriptInputCollectionViewModel();
diff --git a/Engine/Plugins/FX/Niagara/Source/NiagaraEditor/Private/ViewModels/NiagaraScriptOutputCollectionViewModel.cpp b/Engine/Plugins/FX/Niagara/Source/NiagaraEditor/Private/ViewModels/NiagaraScriptOutputCollectionViewModel.cpp
--- a/Engine/Plugins/FX/Niagara/Source/NiagaraEditor/Private/ViewModels/NiagaraScriptOutputCollectionViewModel.cpp
+++ b/Engine/Plugins/FX/Niagara/Source/NiagaraEditor/Private/ViewModels/NiagaraScriptOutputCollectionViewModel.cpp
@@ -12,54 +12,11 @@
 
 #define LOCTEXT_NAMESPACE "NiagaraScriptOutputCollection"
 
-FNiagaraScriptOutputCollectionViewModel::FNiagaraScriptOutputCollectionViewModel(UNiagaraScript* InScript, ENiagaraParameterEditMode InParameterEditMode)
-	: FNiagaraParameterCollectionViewModel(InParameterEditMode)
-	, DisplayName(LOCTEXT("DisplayName", "Outputs"))
-{
-	if (InScript != nullptr)
-	{
-		Scripts.Add(InScript);
-	}
-
-	if (Scripts.Num() == 1 && Scripts[0].IsValid() && Scripts[0]->GetSource())
-	{
-		Graph = Cast<UNiagaraScriptSource>(Scripts[0]->GetSource())->NodeGraph;
-		bCanHaveNumericParameters = Scripts[0]->IsStandaloneScript();
-		
-		OutputNode = nullptr;
-		
-		TArray<UNiagaraNodeOutput*> OutputNodes;
-		Graph->GetNodesOfClass(OutputNodes);
-		if (OutputNodes.Num() == 1)
-		{
-			OutputNode = OutputNodes[0];
-		}
-
-	}
-	else
-	{
-		Graph = nullptr;
-		OutputNode = nullptr;
-		bCanHaveNumericParameters = true;
-	}
-	RefreshParameterViewModels();
-
-	if (Graph.IsValid())
-	{
-		OnGraphChangedHandle = Graph->AddOnGraphChangedHandler(
-			FOnGraphChanged::FDelegate::CreateRaw(this, &FNiagaraScriptOutputCollectionViewModel::OnGraphChanged));
-		OnRecompileHandle = Graph->AddOnGraphNeedsRecompileHandler(
-			FOnGraphChanged::FDelegate::CreateRaw(this, &FNiagaraScriptOutputCollectionViewModel::OnGraphChanged));
-	}
-}
-
 FNiagaraScriptOutputCollectionViewModel::FNiagaraScriptOutputCollectionViewModel(ENiagaraParameterEditMode InParameterEditMode)
 	: FNiagaraParameterCollectionViewModel(InParameterEditMode)
 	, DisplayName(LOCTEXT("DisplayName", "Outputs"))
+	, bCanHaveNumericParameters(true)
 {
-	Graph = nullptr;
-	OutputNode = nullptr;
-	bCanHaveNumericParameters = false;
 }
 
 FNiagaraScriptOutputCollectionViewModel::~FNiagaraScriptOutputCollectionViewModel()
@@ -103,9 +60,9 @@ void FNiagaraScriptOutputCollectionViewModel::SetScripts(TArray<UNiagaraScript*>
 	{
 		Graph = Cast<UNiagaraScriptSource>(Scripts[0]->GetSource())->NodeGraph;
 		OnGraphChangedHandle = Graph->AddOnGraphChangedHandler(
-			FOnGraphChanged::FDelegate::CreateRaw(this, &FNiagaraScriptOutputCollectionViewModel::OnGraphChanged));
+			FOnGraphChanged::FDelegate::CreateSP(this, &FNiagaraScriptOutputCollectionViewModel::OnGraphChanged));
 		OnRecompileHandle = Graph->AddOnGraphNeedsRecompileHandler(
-			FOnGraphChanged::FDelegate::CreateRaw(this, &FNiagaraScriptOutputCollectionViewModel::OnGraphChanged));
+			FOnGraphChanged::FDelegate::CreateSP(this, &FNiagaraScriptOutputCollectionViewModel::OnGraphChanged));
 		bCanHaveNumericParameters = Scripts[0]->IsStandaloneScript();
 
 		if (Scripts.Num() == 1)
@@ -232,8 +189,8 @@ void FNiagaraScriptOutputCollectionViewModel::RefreshParameterViewModels()
 		for (FNiagaraVariable& OutputVariable : OutputNode->Outputs)
 		{
 			TSharedRef<FNiagaraScriptParameterViewModel> ParameterViewModel = MakeShareable(new FNiagaraScriptParameterViewModel(OutputVariable, *OutputNode, nullptr, nullptr, ParameterEditMode));
-			ParameterViewModel->OnNameChanged().AddRaw(this, &FNiagaraScriptOutputCollectionViewModel::OnParameterNameChanged, &OutputVariable);
-			ParameterViewModel->OnTypeChanged().AddRaw(this, &FNiagaraScriptOutputCollectionViewModel::OnParameterTypeChanged, &OutputVariable);
+			ParameterViewModel->OnNameChanged().AddSP(this, &FNiagaraScriptOutputCollectionViewModel::OnParameterNameChanged, &OutputVariable);
+			ParameterViewModel->OnTypeChanged().AddSP(this, &FNiagaraScriptOutputCollectionViewModel::OnParameterTypeChanged, &OutputVariable);
 			ParameterViewModels.Add(ParameterViewModel);
 		}
 	}
diff --git a/Engine/Plugins/FX/Niagara/Source/NiagaraEditor/Private/ViewModels/NiagaraScriptOutputCollectionViewModel.h b/Engine/Plugins/FX/Niagara/Source/NiagaraEditor/Private/ViewModels/NiagaraScriptOutputCollectionViewModel.h
--- a/Engine/Plugins/FX/Niagara/Source/NiagaraEditor/Private/ViewModels/NiagaraScriptOutputCollectionViewModel.h
+++ b/Engine/Plugins/FX/Niagara/Source/NiagaraEditor/Private/ViewModels/NiagaraScriptOutputCollectionViewModel.h
@@ -17,7 +17,6 @@ public:
 	DECLARE_MULTICAST_DELEGATE(FOnOutputParametersChanged);
 
 public:
-	FNiagaraScriptOutputCollectionViewModel(UNiagaraScript* InScript, ENiagaraParameterEditMode InParameterEditMode);
 	FNiagaraScriptOutputCollectionViewModel(ENiagaraParameterEditMode InParameterEditMode);
 
 	~FNiagaraScriptOutputCollectionViewModel();
diff --git a/Engine/Plugins/FX/Niagara/Source/NiagaraEditor/Private/ViewModels/NiagaraScriptViewModel.cpp b/Engine/Plugins/FX/Niagara/Source/NiagaraEditor/Private/ViewModels/NiagaraScriptViewModel.cpp
--- a/Engine/Plugins/FX/Niagara/Source/NiagaraEditor/Private/ViewModels/NiagaraScriptViewModel.cpp
+++ b/Engine/Plugins/FX/Niagara/Source/NiagaraEditor/Private/ViewModels/NiagaraScriptViewModel.cpp
@@ -19,51 +19,6 @@
 
 template<> TMap<UNiagaraScript*, TArray<FNiagaraScriptViewModel*>> TNiagaraViewModelManager<UNiagaraScript, FNiagaraScriptViewModel>::ObjectsToViewModels{};
 
-FNiagaraScriptViewModel::FNiagaraScriptViewModel(UNiagaraScript* InScript, FText DisplayName, ENiagaraParameterEditMode InParameterEditMode)
-	: InputCollectionViewModel(MakeShareable(new FNiagaraScriptInputCollectionViewModel(InScript, DisplayName, InParameterEditMode)))
-	, OutputCollectionViewModel(MakeShareable(new FNiagaraScriptOutputCollectionViewModel(InScript, InParameterEditMode)))
-	, GraphViewModel(MakeShareable(new FNiagaraScriptGraphViewModel(Cast<UNiagaraScriptSource>(InScript->GetSource()), DisplayName)))
-	, VariableSelection(MakeShareable(new FNiagaraObjectSelection()))
-	, bUpdatingSelectionInternally(false)
-	, LastCompileStatus(ENiagaraScriptCompileStatus::NCS_Unknown)
-{
-	if (InScript != nullptr)
-	{
-		Scripts.Add(InScript);
-		Source = Cast<UNiagaraScriptSource>(InScript->GetSource());
-		Scripts[0]->OnVMScriptCompiled().AddRaw(this, &FNiagaraScriptViewModel::OnVMScriptCompiled);
-	}
-
-	InputCollectionViewModel->GetSelection().OnSelectedObjectsChanged().AddRaw(this, &FNiagaraScriptViewModel::InputViewModelSelectionChanged);
-	InputCollectionViewModel->OnParameterValueChanged().AddRaw(this, &FNiagaraScriptViewModel::InputParameterValueChanged);
-	OutputCollectionViewModel->OnParameterValueChanged().AddRaw(this, &FNiagaraScriptViewModel::OutputParameterValueChanged);
-	GraphViewModel->GetNodeSelection()->OnSelectedObjectsChanged().AddRaw(this, &FNiagaraScriptViewModel::GraphViewModelSelectedNodesChanged);
-	GEditor->RegisterForUndo(this);
-
-	// Guess at initial compile status
-	if (InScript != nullptr && InScript->GetVMExecutableData().IsValid() && InScript->GetVMExecutableData().ByteCode.Num() == 0) // This is either a brand new script or failed in the past. Since we create a default working script, assume invalid.
-	{
-		LastCompileStatus = InScript->GetLastCompileStatus();
-		GraphViewModel->SetErrorTextToolTip("Please recompile for full error stack.");
-	}
-	else // Possibly warnings previously, but still compiled. It *could* have been dirtied somehow, but we assume that it is up-to-date.
-	{
-		if (InScript != nullptr && InScript->AreScriptAndSourceSynchronized())
-		{
-			LastCompileStatus = InScript->GetLastCompileStatus();
-		}
-		else
-		{
-			LastCompileStatus = ENiagaraScriptCompileStatus::NCS_Unknown;
-			GraphViewModel->SetErrorTextToolTip("Please recompile for full error stack.");
-		}
-	}
-
-	RegisteredHandles.Add(RegisterViewModelWithMap(InScript, this));
-	//UE_LOG(LogNiagaraEditor, Warning, TEXT("FNiagaraScriptViewModel %p"), this);
-
-}
-
 FNiagaraScriptViewModel::FNiagaraScriptViewModel(FText DisplayName, ENiagaraParameterEditMode InParameterEditMode)
 	: InputCollectionViewModel(MakeShareable(new FNiagaraScriptInputCollectionViewModel(DisplayName, InParameterEditMode)))
 	, OutputCollectionViewModel(MakeShareable(new FNiagaraScriptOutputCollectionViewModel(InParameterEditMode)))
@@ -77,13 +32,6 @@ FNiagaraScriptViewModel::FNiagaraScriptViewModel(FText DisplayName, ENiagaraPara
 	OutputCollectionViewModel->OnParameterValueChanged().AddRaw(this, &FNiagaraScriptViewModel::OutputParameterValueChanged);
 	GraphViewModel->GetNodeSelection()->OnSelectedObjectsChanged().AddRaw(this, &FNiagaraScriptViewModel::GraphViewModelSelectedNodesChanged);
 	GEditor->RegisterForUndo(this);
-
-	// Guess at initial compile status
-	LastCompileStatus = ENiagaraScriptCompileStatus::NCS_UpToDate;
-
-	CompileStatuses.Empty();
-	CompileErrors.Empty();
-	CompilePaths.Empty();
 }
 
 void FNiagaraScriptViewModel::OnVMScriptCompiled(UNiagaraScript* InScript)
@@ -202,7 +150,7 @@ void FNiagaraScriptViewModel::SetScripts(UNiagaraScriptSource* InScriptSource, T
 	{
 		int32 i = Scripts.Add(Script);
 		check(Script->GetSource() == InScriptSource);
-		Scripts[i]->OnVMScriptCompiled().AddRaw(this, &FNiagaraScriptViewModel::OnVMScriptCompiled);
+		Scripts[i]->OnVMScriptCompiled().AddSP(this, &FNiagaraScriptViewModel::OnVMScriptCompiled);
 	}
 	Source = InScriptSource;
 
diff --git a/Engine/Plugins/FX/Niagara/Source/NiagaraEditor/Private/ViewModels/NiagaraScriptViewModel.h b/Engine/Plugins/FX/Niagara/Source/NiagaraEditor/Private/ViewModels/NiagaraScriptViewModel.h
--- a/Engine/Plugins/FX/Niagara/Source/NiagaraEditor/Private/ViewModels/NiagaraScriptViewModel.h
+++ b/Engine/Plugins/FX/Niagara/Source/NiagaraEditor/Private/ViewModels/NiagaraScriptViewModel.h
@@ -23,7 +23,6 @@ class FNiagaraObjectSelection;
 class FNiagaraScriptViewModel : public TSharedFromThis<FNiagaraScriptViewModel>, public FEditorUndoClient, public TNiagaraViewModelManager<UNiagaraScript, FNiagaraScriptViewModel>
 {
 public:
-	FNiagaraScriptViewModel(UNiagaraScript* InScript, FText DisplayName, ENiagaraParameterEditMode InParameterEditMode);
 	FNiagaraScriptViewModel(FText DisplayName, ENiagaraParameterEditMode InParameterEditMode);
 
 	~FNiagaraScriptViewModel();
diff --git a/Engine/Plugins/FX/Niagara/Source/NiagaraEditor/Private/ViewModels/NiagaraStandaloneScriptViewModel.cpp b/Engine/Plugins/FX/Niagara/Source/NiagaraEditor/Private/ViewModels/NiagaraStandaloneScriptViewModel.cpp
--- a/Engine/Plugins/FX/Niagara/Source/NiagaraEditor/Private/ViewModels/NiagaraStandaloneScriptViewModel.cpp
+++ b/Engine/Plugins/FX/Niagara/Source/NiagaraEditor/Private/ViewModels/NiagaraStandaloneScriptViewModel.cpp
@@ -4,19 +4,22 @@
 #include "NiagaraMessageManager.h"
 
 FNiagaraStandaloneScriptViewModel::FNiagaraStandaloneScriptViewModel(
-	UNiagaraScript* InScript
-	, FText DisplayName
-	, ENiagaraParameterEditMode InParameterEditMode
-	, TSharedPtr<FNiagaraMessageLogViewModel> InNiagaraMessageLogViewModel
-	, UNiagaraScript* InSourceScript
-	, const FGuid& InSourceScriptObjKey
-	)
-	: FNiagaraScriptViewModel(InScript, DisplayName, InParameterEditMode)
+	FText DisplayName,
+	ENiagaraParameterEditMode InParameterEditMode,
+	TSharedPtr<FNiagaraMessageLogViewModel> InNiagaraMessageLogViewModel,
+	const FGuid& InSourceScriptObjKey
+)
+	: FNiagaraScriptViewModel(DisplayName, InParameterEditMode)
 	, NiagaraMessageLogViewModel(InNiagaraMessageLogViewModel)
-	, SourceScript(InSourceScript)
 	, ScriptMessageLogGuidKey(InSourceScriptObjKey)
 {
-	SendLastCompileMessageJobs(InSourceScript);
+}
+
+void FNiagaraStandaloneScriptViewModel::Initialize(UNiagaraScript* InScript, UNiagaraScript* InSourceScript)
+{
+	SetScript(InScript);
+	SourceScript = InSourceScript;
+	SendLastCompileMessageJobs(SourceScript);
 }
 
 void FNiagaraStandaloneScriptViewModel::OnVMScriptCompiled(UNiagaraScript* InScript)
@@ -25,14 +28,14 @@ void FNiagaraStandaloneScriptViewModel::OnVMScriptCompiled(UNiagaraScript* InScr
 	SendLastCompileMessageJobs(InScript);
 }
 
-void FNiagaraStandaloneScriptViewModel::SendLastCompileMessageJobs(UNiagaraScript* InScript)
+void FNiagaraStandaloneScriptViewModel::SendLastCompileMessageJobs(const UNiagaraScript* InScript)
 {
 	int32 ErrorCount = 0;
 	int32 WarningCount = 0;
 
 	TArray<TSharedPtr<const INiagaraMessageJob>> JobBatchToQueue;
 
-	TArray<FNiagaraCompileEvent>& CurrentCompileEvents = InScript->GetVMExecutableData().LastCompileEvents;
+	const TArray<FNiagaraCompileEvent>& CurrentCompileEvents = InScript->GetVMExecutableData().LastCompileEvents;
 
 	// Iterate from back to front to avoid reordering the events when they are queued
 	for (int i = CurrentCompileEvents.Num() - 1; i >= 0; --i)
diff --git a/Engine/Plugins/FX/Niagara/Source/NiagaraEditor/Private/ViewModels/NiagaraStandaloneScriptViewModel.h b/Engine/Plugins/FX/Niagara/Source/NiagaraEditor/Private/ViewModels/NiagaraStandaloneScriptViewModel.h
--- a/Engine/Plugins/FX/Niagara/Source/NiagaraEditor/Private/ViewModels/NiagaraStandaloneScriptViewModel.h
+++ b/Engine/Plugins/FX/Niagara/Source/NiagaraEditor/Private/ViewModels/NiagaraStandaloneScriptViewModel.h
@@ -13,19 +13,19 @@ class FNiagaraStandaloneScriptViewModel : public FNiagaraScriptViewModel
 {
 public:
 	FNiagaraStandaloneScriptViewModel(
-		UNiagaraScript* InScript
-		, FText DisplayName
-		, ENiagaraParameterEditMode InParameterEditMode
-		, TSharedPtr<FNiagaraMessageLogViewModel> InNiagaraMessageLogViewModel
-		, UNiagaraScript* InSourceScript
-		, const FGuid& InMessageLogGuidKey
+		FText DisplayName,
+		ENiagaraParameterEditMode InParameterEditMode,
+		TSharedPtr<FNiagaraMessageLogViewModel> InNiagaraMessageLogViewModel,
+		const FGuid& InMessageLogGuidKey
 	);
 
+	void Initialize(UNiagaraScript* InScript, UNiagaraScript* InSourceScript);
+
 private:
 	virtual void OnVMScriptCompiled(UNiagaraScript* InScript) override;
 
 	/** Sends message jobs to FNiagaraMessageManager for all compile events from the last compile. */
-	void SendLastCompileMessageJobs(UNiagaraScript* InScript);
+	void SendLastCompileMessageJobs(const UNiagaraScript* InScript);
 
 	TSharedPtr<FNiagaraMessageLogViewModel> NiagaraMessageLogViewModel;
 	const UNiagaraScript* SourceScript;
diff --git a/Engine/Plugins/FX/Niagara/Source/NiagaraEditor/Private/ViewModels/NiagaraSystemScriptViewModel.cpp b/Engine/Plugins/FX/Niagara/Source/NiagaraEditor/Private/ViewModels/NiagaraSystemScriptViewModel.cpp
--- a/Engine/Plugins/FX/Niagara/Source/NiagaraEditor/Private/ViewModels/NiagaraSystemScriptViewModel.cpp
+++ b/Engine/Plugins/FX/Niagara/Source/NiagaraEditor/Private/ViewModels/NiagaraSystemScriptViewModel.cpp
@@ -22,41 +22,29 @@
 #include "ViewModels/NiagaraEmitterHandleViewModel.h"
 #include "ViewModels/NiagaraEmitterViewModel.h"
 
-FNiagaraSystemScriptViewModel::FNiagaraSystemScriptViewModel(UNiagaraSystem& InSystem, FNiagaraSystemViewModel* InParent)
-	: FNiagaraScriptViewModel(InSystem.GetSystemSpawnScript(), NSLOCTEXT("SystemScriptViewModel", "GraphName", "System"), ENiagaraParameterEditMode::EditAll)
-	, Parent(InParent)
-	, System(InSystem)
+FNiagaraSystemScriptViewModel::FNiagaraSystemScriptViewModel()
+	: FNiagaraScriptViewModel(NSLOCTEXT("SystemScriptViewModel", "GraphName", "System"), ENiagaraParameterEditMode::EditAll)
 {
-	Scripts.Add(InSystem.GetSystemUpdateScript());
-
-	if (GetGraphViewModel()->GetGraph())
-	{
-		OnGraphChangedHandle = GetGraphViewModel()->GetGraph()->AddOnGraphChangedHandler(
-			FOnGraphChanged::FDelegate::CreateRaw(this, &FNiagaraSystemScriptViewModel::OnGraphChanged));
-
-		OnRecompileHandle = GetGraphViewModel()->GetGraph()->AddOnGraphNeedsRecompileHandler(
-			FOnGraphChanged::FDelegate::CreateRaw(this, &FNiagaraSystemScriptViewModel::OnGraphChanged));
-		
-		GetGraphViewModel()->SetErrorTextToolTip("");
-	}
+}
 
-	System.OnSystemCompiled().AddRaw(this, &FNiagaraSystemScriptViewModel::OnSystemVMCompiled);
+void FNiagaraSystemScriptViewModel::Initialize(UNiagaraSystem& InSystem)
+{
+	System = &InSystem;
+	SetScript(System->GetSystemSpawnScript());
+	System->OnSystemCompiled().AddSP(this, &FNiagaraSystemScriptViewModel::OnSystemVMCompiled);
 }
 
 FNiagaraSystemScriptViewModel::~FNiagaraSystemScriptViewModel()
 {
-	System.OnSystemCompiled().RemoveAll(this);
-	if (GetGraphViewModel()->GetGraph())
+	if (System.IsValid())
 	{
-		GetGraphViewModel()->GetGraph()->RemoveOnGraphChangedHandler(OnGraphChangedHandle);
-		GetGraphViewModel()->GetGraph()->RemoveOnGraphNeedsRecompileHandler(OnRecompileHandle);
-		
+		System->OnSystemCompiled().RemoveAll(this);
 	}
 }
 
 void FNiagaraSystemScriptViewModel::OnSystemVMCompiled(UNiagaraSystem* InSystem)
 {
-	if (InSystem != &System)
+	if (InSystem != System.Get())
 	{
 		return;
 	}
@@ -130,13 +118,5 @@ FNiagaraSystemScriptViewModel::FOnSystemCompiled& FNiagaraSystemScriptViewModel:
 
 void FNiagaraSystemScriptViewModel::CompileSystem(bool bForce)
 {
-	System.RequestCompile(bForce);
-}
-
-void FNiagaraSystemScriptViewModel::OnGraphChanged(const struct FEdGraphEditAction& InAction)
-{
-	if (InAction.Action == GRAPHACTION_SelectNode)
-	{
-		return;
-	}
+	System->RequestCompile(bForce);
 }
diff --git a/Engine/Plugins/FX/Niagara/Source/NiagaraEditor/Private/ViewModels/NiagaraSystemScriptViewModel.h b/Engine/Plugins/FX/Niagara/Source/NiagaraEditor/Private/ViewModels/NiagaraSystemScriptViewModel.h
--- a/Engine/Plugins/FX/Niagara/Source/NiagaraEditor/Private/ViewModels/NiagaraSystemScriptViewModel.h
+++ b/Engine/Plugins/FX/Niagara/Source/NiagaraEditor/Private/ViewModels/NiagaraSystemScriptViewModel.h
@@ -4,6 +4,7 @@
 
 #include "NiagaraScriptViewModel.h"
 
+class FNiagaraSystemViewModel;
 class FNiagaraScriptViewModel;
 class UNiagaraSystem;
 
@@ -14,7 +15,9 @@ public:
 	DECLARE_MULTICAST_DELEGATE(FOnSystemCompiled)
 
 public:
-	FNiagaraSystemScriptViewModel(UNiagaraSystem& InSystem, class FNiagaraSystemViewModel* InParent);
+	FNiagaraSystemScriptViewModel();
+
+	void Initialize(UNiagaraSystem& InSystem);
 
 	~FNiagaraSystemScriptViewModel();
 
@@ -23,19 +26,11 @@ public:
 	void CompileSystem(bool bForce);
 
 private:
-	/** Called whenever the graph for the System script changes. */
-	void OnGraphChanged(const struct FEdGraphEditAction& InAction);
-
 	void OnSystemVMCompiled(UNiagaraSystem* InSystem);
 
-	FNiagaraSystemViewModel* Parent;
-
+private:
 	/** The System who's script is getting viewed and edited by this view model. */
-	UNiagaraSystem& System;
-
-	/** A handle to the on graph changed delegate. */
-	FDelegateHandle OnGraphChangedHandle;
-	FDelegateHandle OnRecompileHandle;
+	TWeakObjectPtr<UNiagaraSystem> System;
 
 	FOnSystemCompiled OnSystemCompiledDelegate;
-};
\ No newline at end of file
+};
diff --git a/Engine/Plugins/FX/Niagara/Source/NiagaraEditor/Private/ViewModels/NiagaraSystemViewModel.cpp b/Engine/Plugins/FX/Niagara/Source/NiagaraEditor/Private/ViewModels/NiagaraSystemViewModel.cpp
--- a/Engine/Plugins/FX/Niagara/Source/NiagaraEditor/Private/ViewModels/NiagaraSystemViewModel.cpp
+++ b/Engine/Plugins/FX/Niagara/Source/NiagaraEditor/Private/ViewModels/NiagaraSystemViewModel.cpp
@@ -90,7 +90,8 @@ void FNiagaraSystemViewModel::Initialize(UNiagaraSystem& InSystem, FNiagaraSyste
 	SelectionViewModel->Initialize(this->AsShared());
 	SelectionViewModel->OnSelectionChanged().AddSP(this, &FNiagaraSystemViewModel::SystemSelectionChanged);
 
-	SystemScriptViewModel = MakeShared<FNiagaraSystemScriptViewModel>(GetSystem(), this);
+	SystemScriptViewModel = MakeShared<FNiagaraSystemScriptViewModel>();
+	SystemScriptViewModel->Initialize(GetSystem());
 
 	OverviewGraphViewModel = MakeShared<FNiagaraOverviewGraphViewModel>();
 	OverviewGraphViewModel->Initialize(this->AsShared());
@@ -277,6 +278,9 @@ void FNiagaraSystemViewModel::AddEmitterFromAssetData(const FAssetData& AssetDat
 
 void FNiagaraSystemViewModel::AddEmitter(UNiagaraEmitter& Emitter)
 {
+	// Reset view models before modifying the emitter handle list to prevent accessing deleted data.
+	ResetEmitterHandleViewModelsAndTracks();
+
 	// When editing an emitter asset the system is a placeholder and we don't want to make adding an emitter to it undoable.
 	bool bSystemIsPlaceholder = EditMode == ENiagaraSystemViewModelEditMode::EmitterAsset;
 	if (false == bSystemIsPlaceholder)
@@ -328,7 +332,8 @@ void FNiagaraSystemViewModel::DuplicateEmitters(TArray<FEmitterHandleToDuplicate
 		return;
 	}
 
-	// Kill all system instances before modifying the emitter handle list to prevent accessing deleted data.
+	// Kill all system instances and reset view models before modifying the emitter handle list to prevent accessing deleted data.
+	ResetEmitterHandleViewModelsAndTracks();
 	FNiagaraEditorUtilities::KillSystemInstances(GetSystem());
 	const FScopedTransaction DeleteTransaction(EmitterHandlesToDuplicate.Num() == 1
 		? LOCTEXT("DuplicateEmitter", "Duplicate emitter")
@@ -339,7 +344,8 @@ void FNiagaraSystemViewModel::DuplicateEmitters(TArray<FEmitterHandleToDuplicate
 	{
 		EmitterHandleNames.Add(EmitterHandle.GetName());
 	}
-
+	
+	GetSystem().Modify();
 	for (FEmitterHandleToDuplicate& EmitterHandleToDuplicate : EmitterHandlesToDuplicate)
 	{
 		FNiagaraEmitterHandle HandleToDuplicate;
@@ -381,6 +387,8 @@ void FNiagaraSystemViewModel::DeleteEmitters(TSet<FGuid> EmitterHandleIdsToDelet
 {
 	if (EmitterHandleIdsToDelete.Num() > 0)
 	{
+		// Reset view models before modifying the emitter handle list to prevent accessing deleted data.
+		ResetEmitterHandleViewModelsAndTracks();
 		FNiagaraEditorUtilities::RemoveEmittersFromSystemByEmitterHandleId(GetSystem(), EmitterHandleIdsToDelete);
 		RefreshAll();
 		bForceAutoCompileOnce = true;
@@ -458,6 +466,8 @@ void FNiagaraSystemViewModel::AddReferencedObjects(FReferenceCollector& Collecto
 
 void FNiagaraSystemViewModel::PostUndo(bool bSuccess)
 {
+	// Reset emitter handle view models to prevent accessing invalid handles if they were in the undo operation.
+	ResetEmitterHandleViewModelsAndTracks();
 	RefreshAll();
 	GetDefault<UEdGraphSchema_NiagaraSystemOverview>()->ForceVisualizationCacheClear();
 }
@@ -745,6 +755,23 @@ void FNiagaraSystemViewModel::ToggleEmitterIsolation(TSharedRef<FNiagaraEmitterH
 	GetSystem().SetIsolateEnabled(bAnyEmitterIsolated);
 }
 
+void FNiagaraSystemViewModel::ResetEmitterHandleViewModelsAndTracks()
+{
+	for (TSharedRef<FNiagaraEmitterHandleViewModel> EmitterHandleViewModel : EmitterHandleViewModels)
+	{
+		EmitterHandleViewModel->Reset();
+	}
+
+	TArray<UMovieSceneTrack*> MasterTracks = NiagaraSequence->GetMovieScene()->GetMasterTracks();
+	for (UMovieSceneTrack* MasterTrack : MasterTracks)
+	{
+		if (MasterTrack != nullptr)
+		{
+			NiagaraSequence->GetMovieScene()->RemoveMasterTrack(*MasterTrack);
+		}
+	}
+}
+
 void FNiagaraSystemViewModel::RefreshEmitterHandleViewModels()
 {
 	TArray<TSharedRef<FNiagaraEmitterHandleViewModel>> OldViewModels = EmitterHandleViewModels;
@@ -1590,7 +1617,7 @@ void FNiagaraSystemViewModel::SystemInstanceInitialized()
 {
 	for (TSharedRef<FNiagaraEmitterHandleViewModel> EmitterHandleViewModel : EmitterHandleViewModels)
 	{
-		if (EmitterHandleViewModel->GetEmitterHandle())
+		if (EmitterHandleViewModel->IsValid())
 		{
 			EmitterHandleViewModel->SetSimulation(SystemInstance->GetSimulationForHandle(*EmitterHandleViewModel->GetEmitterHandle()));
 		}
diff --git a/Engine/Plugins/FX/Niagara/Source/NiagaraEditor/Private/Widgets/SNiagaraScriptGraph.cpp b/Engine/Plugins/FX/Niagara/Source/NiagaraEditor/Private/Widgets/SNiagaraScriptGraph.cpp
--- a/Engine/Plugins/FX/Niagara/Source/NiagaraEditor/Private/Widgets/SNiagaraScriptGraph.cpp
+++ b/Engine/Plugins/FX/Niagara/Source/NiagaraEditor/Private/Widgets/SNiagaraScriptGraph.cpp
@@ -145,7 +145,8 @@ TSharedRef<SGraphEditor> SNiagaraScriptGraph::ConstructGraphEditor()
 	Events.OnVerifyTextCommit = FOnNodeVerifyTextCommit::CreateSP(this, &SNiagaraScriptGraph::OnVerifyNodeTextCommit);
 	Events.OnSpawnNodeByShortcut = SGraphEditor::FOnSpawnNodeByShortcut::CreateSP(this, &SNiagaraScriptGraph::OnSpawnGraphNodeByShortcut);
 
-	TSharedRef<FUICommandList> Commands = ViewModel->GetCommands();
+	TSharedRef<FUICommandList> Commands = MakeShared<FUICommandList>();
+	Commands->Append(ViewModel->GetCommands());
 	Commands->MapAction(
 		FNiagaraEditorCommands::Get().FindInCurrentView,
 		FExecuteAction::CreateRaw(this, &SNiagaraScriptGraph::FocusGraphSearchBox));
@@ -154,7 +155,7 @@ TSharedRef<SGraphEditor> SNiagaraScriptGraph::ConstructGraphEditor()
 		FExecuteAction::CreateRaw(this, &SNiagaraScriptGraph::OnCreateComment));
 	
 	TSharedRef<SGraphEditor> CreatedGraphEditor = SNew(SGraphEditor)
-		.AdditionalCommands(ViewModel->GetCommands())
+		.AdditionalCommands(Commands)
 		.Appearance(AppearanceInfo)
 		.TitleBar(TitleBarWidget)
 		.GraphToEdit(ViewModel->GetGraph())
diff --git a/Engine/Plugins/FX/Niagara/Source/NiagaraEditor/Public/ViewModels/NiagaraEmitterHandleViewModel.h b/Engine/Plugins/FX/Niagara/Source/NiagaraEditor/Public/ViewModels/NiagaraEmitterHandleViewModel.h
--- a/Engine/Plugins/FX/Niagara/Source/NiagaraEditor/Public/ViewModels/NiagaraEmitterHandleViewModel.h
+++ b/Engine/Plugins/FX/Niagara/Source/NiagaraEditor/Public/ViewModels/NiagaraEmitterHandleViewModel.h
@@ -30,6 +30,9 @@ public:
 	/** Initializes the emitter editor view model with the supplied emitter handle and simulation.*/
 	void Initialize(TSharedRef<FNiagaraSystemViewModel> InOwningSystemViewModel, FNiagaraEmitterHandle* InEmitterHandle, TWeakPtr<FNiagaraEmitterInstance, ESPMode::ThreadSafe> InSimulation);
 
+	/** Returns whether or not this view model represents a valid emitter handle. */
+	bool IsValid() const;
+
 	/** Resets the data in the view model. */
 	void Reset();
 
diff --git a/Engine/Plugins/FX/Niagara/Source/NiagaraEditor/Public/ViewModels/NiagaraSystemViewModel.h b/Engine/Plugins/FX/Niagara/Source/NiagaraEditor/Public/ViewModels/NiagaraSystemViewModel.h
--- a/Engine/Plugins/FX/Niagara/Source/NiagaraEditor/Public/ViewModels/NiagaraSystemViewModel.h
+++ b/Engine/Plugins/FX/Niagara/Source/NiagaraEditor/Public/ViewModels/NiagaraSystemViewModel.h
@@ -299,6 +299,9 @@ private:
 	/** Sets up the preview component and System instance. */
 	void SetupPreviewComponentAndInstance();
 
+	/** Resets the emitter handle view models and tracks to remove data from them.  This must be called before modifying the emitter handle collection to prevent accessing invalid data. */
+	void ResetEmitterHandleViewModelsAndTracks();
+
 	/** Rebuilds the emitter handle view models. */
 	void RefreshEmitterHandleViewModels();
 
