No longer compiling disabled emitters

diff --git a/Engine/Plugins/FX/Niagara/Source/Niagara/Classes/NiagaraEmitter.h b/Engine/Plugins/FX/Niagara/Source/Niagara/Classes/NiagaraEmitter.h
--- a/Engine/Plugins/FX/Niagara/Source/Niagara/Classes/NiagaraEmitter.h
+++ b/Engine/Plugins/FX/Niagara/Source/Niagara/Classes/NiagaraEmitter.h
@@ -295,6 +295,7 @@ public:
 
 	UNiagaraEmitter* MakeRecursiveDeepCopy(UObject* DestOuter) const;
 	UNiagaraEmitter* MakeRecursiveDeepCopy(UObject* DestOuter, TMap<const UObject*, UObject*>& ExistingConversions) const;
+	void NIAGARA_API InvalidateCompileResults();
 
 	/* Gets a Guid which is updated any time data in this emitter is changed. */
 	FGuid NIAGARA_API GetChangeId() const;
@@ -336,7 +337,7 @@ public:
 #endif
 
 	/** Is this emitter allowed to be enabled by the current system detail level. */
-	bool IsAllowedByDetailLevel(int32 DetailLevel)const;
+	NIAGARA_API bool IsAllowedByDetailLevel(int32 DetailLevel)const;
 	NIAGARA_API bool RequiresPersistantIDs()const;
 
 	NIAGARA_API bool IsValid()const;
diff --git a/Engine/Plugins/FX/Niagara/Source/Niagara/Classes/NiagaraEmitterHandle.h b/Engine/Plugins/FX/Niagara/Source/Niagara/Classes/NiagaraEmitterHandle.h
--- a/Engine/Plugins/FX/Niagara/Source/Niagara/Classes/NiagaraEmitterHandle.h
+++ b/Engine/Plugins/FX/Niagara/Source/Niagara/Classes/NiagaraEmitterHandle.h
@@ -47,8 +47,8 @@ public:
 	/** Gets whether or not this emitter is enabled within the System.  Disabled emitters aren't simulated. */
 	bool GetIsEnabled() const;
 
-	/** Gets whether or not this emitter is enabled within the System.  Disabled emitters aren't simulated. */
-	void SetIsEnabled(bool bInIsEnabled);
+	/** Sets whether or not this emitter is enabled within the System.  Disabled emitters aren't simulated. Returns whether or not the enabled state changed. */
+	bool SetIsEnabled(bool bInIsEnabled, UNiagaraSystem& InOwnerSystem, bool bRecompileIfChanged);
 
 #if WITH_EDITORONLY_DATA
 	bool IsIsolated() const {	return bIsolated; }
@@ -71,6 +71,8 @@ public:
 	/** Whether or not this handle uses the supplied emitter. */
 	bool UsesEmitter(const UNiagaraEmitter& InEmitter) const;
 
+	void ClearEmitter();
+
 #endif
 public:
 	/** A static const invalid handle. */
diff --git a/Engine/Plugins/FX/Niagara/Source/Niagara/Classes/NiagaraScriptSourceBase.h b/Engine/Plugins/FX/Niagara/Source/Niagara/Classes/NiagaraScriptSourceBase.h
--- a/Engine/Plugins/FX/Niagara/Source/Niagara/Classes/NiagaraScriptSourceBase.h
+++ b/Engine/Plugins/FX/Niagara/Source/Niagara/Classes/NiagaraScriptSourceBase.h
@@ -113,6 +113,8 @@ class UNiagaraScriptSourceBase : public UObject
 
 	virtual void InvalidateCachedCompileIds() {}
 
+	virtual void RefreshFromExternalChanges() {}
+
 protected:
 	FOnChanged OnChangedDelegate;
 #endif
diff --git a/Engine/Plugins/FX/Niagara/Source/Niagara/Private/NiagaraEmitter.cpp b/Engine/Plugins/FX/Niagara/Source/Niagara/Private/NiagaraEmitter.cpp
--- a/Engine/Plugins/FX/Niagara/Source/Niagara/Private/NiagaraEmitter.cpp
+++ b/Engine/Plugins/FX/Niagara/Source/Niagara/Private/NiagaraEmitter.cpp
@@ -724,6 +724,16 @@ UNiagaraEmitter::FOnEmitterCompiled& UNiagaraEmitter::OnEmitterVMCompiled()
 	return OnVMScriptCompiledDelegate;
 }
 
+void  UNiagaraEmitter::InvalidateCompileResults()
+{
+	TArray<UNiagaraScript*> Scripts;
+	GetScripts(Scripts, false);
+	for (int32 i = 0; i < Scripts.Num(); i++)
+	{
+		Scripts[i]->InvalidateCompileResults();
+	}
+}
+
 void UNiagaraEmitter::OnPostCompile()
 {
 	SyncEmitterAlias(TEXT("Emitter"), UniqueEmitterName);
diff --git a/Engine/Plugins/FX/Niagara/Source/Niagara/Private/NiagaraEmitterHandle.cpp b/Engine/Plugins/FX/Niagara/Source/Niagara/Private/NiagaraEmitterHandle.cpp
--- a/Engine/Plugins/FX/Niagara/Source/Niagara/Private/NiagaraEmitterHandle.cpp
+++ b/Engine/Plugins/FX/Niagara/Source/Niagara/Private/NiagaraEmitterHandle.cpp
@@ -99,9 +99,38 @@ bool FNiagaraEmitterHandle::GetIsEnabled() const
 	return bIsEnabled;
 }
 
-void FNiagaraEmitterHandle::SetIsEnabled(bool bInIsEnabled)
+bool FNiagaraEmitterHandle::SetIsEnabled(bool bInIsEnabled, UNiagaraSystem& InOwnerSystem, bool bRecompileIfChanged)
 {
-	bIsEnabled = bInIsEnabled;
+	if (bIsEnabled != bInIsEnabled)
+	{
+		bIsEnabled = bInIsEnabled;
+
+#if WITH_EDITOR
+		if (InOwnerSystem.GetSystemSpawnScript() && InOwnerSystem.GetSystemSpawnScript()->GetSource())
+		{
+			// We need to get the NiagaraNodeEmitters to update their enabled state based on what happened.
+			InOwnerSystem.GetSystemSpawnScript()->GetSource()->RefreshFromExternalChanges();
+
+			// Need to cause us to recompile in the future if necessary...
+			InOwnerSystem.GetSystemSpawnScript()->InvalidateCompileResults();
+			InOwnerSystem.GetSystemUpdateScript()->InvalidateCompileResults();
+
+			// Clean out the emitter's compile results for cleanliness.
+			if (Instance)
+			{
+				Instance->InvalidateCompileResults();
+			}
+
+			// In some cases we may do the recompile now.
+			if (bRecompileIfChanged)
+			{
+				InOwnerSystem.RequestCompile(false);
+			}
+		}
+#endif
+		return true;
+	}
+	return false;
 }
 
 UNiagaraEmitter* FNiagaraEmitterHandle::GetInstance() const
@@ -119,14 +148,17 @@ FString FNiagaraEmitterHandle::GetUniqueInstanceName()const
 
 bool FNiagaraEmitterHandle::NeedsRecompile() const
 {
-	TArray<UNiagaraScript*> Scripts;
-	Instance->GetScripts(Scripts);
-
-	for (UNiagaraScript* Script : Scripts)
+	if (GetIsEnabled())
 	{
-		if (Script->IsCompilable() && !Script->AreScriptAndSourceSynchronized())
+		TArray<UNiagaraScript*> Scripts;
+		Instance->GetScripts(Scripts);
+
+		for (UNiagaraScript* Script : Scripts)
 		{
-			return true;
+			if (Script->IsCompilable() && !Script->AreScriptAndSourceSynchronized())
+			{
+				return true;
+			}
 		}
 	}
 	return false;
@@ -168,4 +200,12 @@ bool FNiagaraEmitterHandle::UsesEmitter(const UNiagaraEmitter& InEmitter) const
 	return Instance == &InEmitter || (Instance != nullptr && Instance->UsesEmitter(InEmitter));
 }
 
+void FNiagaraEmitterHandle::ClearEmitter()
+{
+	Instance = nullptr;
+	Source_DEPRECATED = nullptr;
+	LastMergedSource_DEPRECATED = nullptr;
+}
+
+
 #endif
\ No newline at end of file
diff --git a/Engine/Plugins/FX/Niagara/Source/Niagara/Private/NiagaraModule.cpp b/Engine/Plugins/FX/Niagara/Source/Niagara/Private/NiagaraModule.cpp
--- a/Engine/Plugins/FX/Niagara/Source/Niagara/Private/NiagaraModule.cpp
+++ b/Engine/Plugins/FX/Niagara/Source/Niagara/Private/NiagaraModule.cpp
@@ -20,6 +20,9 @@
 #include "NiagaraShaderModule.h"
 #include "UObject/CoreRedirects.h"
 #include "NiagaraEmitterInstanceBatcher.h"
+#include "DeviceProfiles/DeviceProfileManager.h"
+#include "Interfaces/ITargetPlatform.h"
+#include "DeviceProfiles/DeviceProfile.h"
 
 IMPLEMENT_MODULE(INiagaraModule, Niagara);
 
@@ -63,6 +66,13 @@ static TAutoConsoleVariable<float> CVarDetailLevel(
 	TEXT("Default = 4"),
 	ECVF_Scalability);
 
+static TAutoConsoleVariable<float> CVarPruneEmittersOnCookByDetailLevel(
+	TEXT("fx.NiagaraPruneEmittersOnCookByDetailLevel"),
+	0,
+	TEXT("Whether to eliminate all emitters that don't match the detail level.\n")
+	TEXT("This will only work if scalability settings affecting detail level can not be changed at runtime (depends on platform).\n"),
+	ECVF_ReadOnly);
+
 static FAutoConsoleVariableRef CVarNiaraGlobalSpawnCountScale(
 	TEXT("fx.NiagaraGlobalSpawnCountScale"),
 	INiagaraModule::EngineGlobalSpawnCountScale,
@@ -455,6 +465,34 @@ void INiagaraModule::UnregisterPrecompiler(FDelegateHandle DelegateHandle)
 
 #endif
 
+
+bool INiagaraModule::IsTargetPlatformIncludedInLevelRangeForCook(const ITargetPlatform* InTargetPlatform, const UNiagaraEmitter* InEmitter)
+{
+#if WITH_EDITORONLY_DATA
+	if (UDeviceProfile* DeviceProfile = UDeviceProfileManager::Get().FindProfile(InTargetPlatform->IniPlatformName()))
+	{
+		// get local scalability CVars that could cull this actor
+		int32 CVarCullBasedOnDetailLevel;
+		if (DeviceProfile->GetConsolidatedCVarValue(TEXT("fx.NiagaraPruneEmittersOnCookByDetailLevel"), CVarCullBasedOnDetailLevel) && CVarCullBasedOnDetailLevel == 1)
+		{
+			int32 CVarDetailLevelFoundValue;
+			if (InEmitter && DeviceProfile->GetConsolidatedCVarValue(TEXT("r.DetailLevel"), CVarDetailLevelFoundValue))
+			{
+				// Check emitter's detail level range contains the platform's level
+				// If e.g. the emitter's detail level range is between 0 and 2  and the platform detail is 3 only,
+				// then we should cull it.
+				if (InEmitter->IsAllowedByDetailLevel(CVarDetailLevelFoundValue))
+				{
+					return true;
+				}
+				return false;
+			}
+		}
+	}
+#endif
+	return true;
+}
+
 void INiagaraModule::OnChangeDetailLevel(class IConsoleVariable* CVar)
 {
 	//Can only change the detail level at runtime on when not cooked.
diff --git a/Engine/Plugins/FX/Niagara/Source/Niagara/Private/NiagaraScript.cpp b/Engine/Plugins/FX/Niagara/Source/Niagara/Private/NiagaraScript.cpp
--- a/Engine/Plugins/FX/Niagara/Source/Niagara/Private/NiagaraScript.cpp
+++ b/Engine/Plugins/FX/Niagara/Source/Niagara/Private/NiagaraScript.cpp
@@ -288,7 +288,7 @@ void UNiagaraScript::ComputeVMCompilationId(FNiagaraVMExecutableDataId& Id) cons
 		for (const FNiagaraEmitterHandle& EmitterHandle: System->GetEmitterHandles())
 		{
 			UNiagaraEmitter* Emitter = Cast<UNiagaraEmitter>(EmitterHandle.GetInstance());
-			if (Emitter)
+			if (Emitter && EmitterHandle.GetIsEnabled())
 			{
 				if (Emitter->bLocalSpace)
 				{
diff --git a/Engine/Plugins/FX/Niagara/Source/Niagara/Private/NiagaraSystem.cpp b/Engine/Plugins/FX/Niagara/Source/Niagara/Private/NiagaraSystem.cpp
--- a/Engine/Plugins/FX/Niagara/Source/Niagara/Private/NiagaraSystem.cpp
+++ b/Engine/Plugins/FX/Niagara/Source/Niagara/Private/NiagaraSystem.cpp
@@ -72,7 +72,35 @@ void UNiagaraSystem::PreSave(const class ITargetPlatform * TargetPlatform)
 void UNiagaraSystem::BeginCacheForCookedPlatformData(const ITargetPlatform *TargetPlatform)
 {
 	Super::BeginCacheForCookedPlatformData(TargetPlatform);
+	
 #if WITH_EDITORONLY_DATA
+	/*
+	bool bAnyChanged = false;
+	for (int32 i = 0; i < GetNumEmitters(); i++)
+	{
+		FNiagaraEmitterHandle& EmitterHandle = GetEmitterHandle(i);
+		if (EmitterHandle.GetIsEnabled())
+		{
+			const UNiagaraEmitter* Emitter = EmitterHandle.GetInstance();
+			if (Emitter != nullptr)
+			{
+				if (!INiagaraModule::IsTargetPlatformIncludedInLevelRangeForCook(TargetPlatform, Emitter))
+				{
+					EmitterHandle.SetIsEnabled(false, *this, false);
+					UE_LOG(LogNiagara, Display, TEXT("Pruning emitter, detail level mismatch (only works if platform can't change detail level at runtime!): %s - set fx.NiagaraPruneEmittersOnCookByDetailLevel to 0 in DeviceProfile.ini for the target profile to avoid"),  *EmitterHandle.GetName().ToString());
+					EmitterHandle.ClearEmitter();
+					bAnyChanged = true;
+				}
+			}
+		}
+	}
+
+	if (bAnyChanged)
+	{
+		RequestCompile(false);
+	}
+	*/
+
 	WaitForCompilationComplete();
 #endif
 }
@@ -321,7 +349,7 @@ void UNiagaraSystem::PostLoad()
 	for (FNiagaraEmitterHandle& EmitterHandle : EmitterHandles)
 	{
 		EmitterHandle.ConditionalPostLoad(NiagaraVer);
-		if (!EmitterHandle.GetInstance()->AreAllScriptAndSourcesSynchronized())
+		if (EmitterHandle.GetIsEnabled() && !EmitterHandle.GetInstance()->AreAllScriptAndSourcesSynchronized())
 		{
 			bEmitterScriptsAreSynchronized = false;
 			break;
@@ -551,7 +579,7 @@ bool UNiagaraSystem::IsValid()const
 
 	for (const FNiagaraEmitterHandle& Handle : EmitterHandles)
 	{
-		if (!Handle.GetInstance()->IsValid())
+		if (Handle.GetIsEnabled() && !Handle.GetInstance()->IsValid())
 		{
 			return false;
 		}
@@ -580,7 +608,7 @@ FNiagaraEmitterHandle UNiagaraSystem::DuplicateEmitterHandle(const FNiagaraEmitt
 {
 	UNiagaraEmitter* DuplicateEmitter = UNiagaraEmitter::CreateAsDuplicate(*EmitterHandleToDuplicate.GetInstance(), EmitterName, *this);
 	FNiagaraEmitterHandle EmitterHandle(*DuplicateEmitter);
-	EmitterHandle.SetIsEnabled(EmitterHandleToDuplicate.GetIsEnabled());
+	EmitterHandle.SetIsEnabled(EmitterHandleToDuplicate.GetIsEnabled(), *this, false);
 	EmitterHandles.Add(EmitterHandle);
 	RefreshSystemParametersFromEmitter(EmitterHandle);
 	return EmitterHandle;
@@ -753,7 +781,14 @@ bool UNiagaraSystem::QueryCompileComplete(bool bWait, bool bDoPost, bool bDoNotA
 		{
 			for (FNiagaraEmitterHandle Handle : EmitterHandles)
 			{
-				Handle.GetInstance()->OnPostCompile();
+				if (Handle.GetIsEnabled())
+				{
+					Handle.GetInstance()->OnPostCompile();
+				}
+				else
+				{
+					Handle.GetInstance()->InvalidateCompileResults();
+				}
 			}
 		}
 
@@ -910,31 +945,33 @@ bool UNiagaraSystem::RequestCompile(bool bForce)
 	for (int32 i = 0; i < EmitterHandles.Num(); i++)
 	{
 		FNiagaraEmitterHandle Handle = EmitterHandles[i];
-
-		UNiagaraScriptSourceBase* GraphSource = Handle.GetInstance()->GraphSource;
-		TSharedPtr<FNiagaraCompileRequestDataBase, ESPMode::ThreadSafe> EmitterPrecompiledData = SystemPrecompiledData->GetDependentRequest(i);
-		EmitterPrecompiledData->GetReferencedObjects(ActiveCompilations[ActiveCompileIdx].RootObjects);
-
-		TArray<UNiagaraScript*> EmitterScripts;
-		Handle.GetInstance()->GetScripts(EmitterScripts, false);
-		check(EmitterScripts.Num() > 0);
-		for (UNiagaraScript* EmitterScript : EmitterScripts)
+		if (Handle.GetIsEnabled())
 		{
-			ActiveCompilations[ActiveCompileIdx].MappedData.Add(EmitterScript, EmitterPrecompiledData);
-
-			FEmitterCompiledScriptPair Pair;
-			Pair.bResultsReady = false;
-			Pair.Emitter = Handle.GetInstance();
-			Pair.CompiledScript = EmitterScript;			
-			if (EmitterScript->RequestExternallyManagedAsyncCompile(EmitterPrecompiledData, Pair.CompileId, Pair.PendingDDCID, bTrulyAsync))
+			UNiagaraScriptSourceBase* GraphSource = Handle.GetInstance()->GraphSource;
+			TSharedPtr<FNiagaraCompileRequestDataBase, ESPMode::ThreadSafe> EmitterPrecompiledData = SystemPrecompiledData->GetDependentRequest(i);
+			EmitterPrecompiledData->GetReferencedObjects(ActiveCompilations[ActiveCompileIdx].RootObjects);
+
+			TArray<UNiagaraScript*> EmitterScripts;
+			Handle.GetInstance()->GetScripts(EmitterScripts, false);
+			check(EmitterScripts.Num() > 0);
+			for (UNiagaraScript* EmitterScript : EmitterScripts)
 			{
-				bAnyUnsynchronized = true;
+				ActiveCompilations[ActiveCompileIdx].MappedData.Add(EmitterScript, EmitterPrecompiledData);
+
+				FEmitterCompiledScriptPair Pair;
+				Pair.bResultsReady = false;
+				Pair.Emitter = Handle.GetInstance();
+				Pair.CompiledScript = EmitterScript;
+				if (EmitterScript->RequestExternallyManagedAsyncCompile(EmitterPrecompiledData, Pair.CompileId, Pair.PendingDDCID, bTrulyAsync))
+				{
+					bAnyUnsynchronized = true;
+				}
+				ActiveCompilations[ActiveCompileIdx].EmitterCompiledScriptPairs.Add(Pair);
 			}
-			ActiveCompilations[ActiveCompileIdx].EmitterCompiledScriptPairs.Add(Pair);
-		}	
 
-		// Add the emitter's User variables to the encountered list to expose for later.
-		EmitterPrecompiledData->GatherPreCompiledVariables(TEXT("User"), EncounteredExposedVars);
+			// Add the emitter's User variables to the encountered list to expose for later.
+			EmitterPrecompiledData->GatherPreCompiledVariables(TEXT("User"), EncounteredExposedVars);
+		}
 	}
 
 	bool bForceSystems = bForce || bAnyUnsynchronized;
diff --git a/Engine/Plugins/FX/Niagara/Source/Niagara/Public/NiagaraModule.h b/Engine/Plugins/FX/Niagara/Source/Niagara/Public/NiagaraModule.h
--- a/Engine/Plugins/FX/Niagara/Source/Niagara/Public/NiagaraModule.h
+++ b/Engine/Plugins/FX/Niagara/Source/Niagara/Public/NiagaraModule.h
@@ -87,6 +87,7 @@ public:
 	FORCEINLINE static int32 GetDetailLevel() { return EngineDetailLevel; }
 	FORCEINLINE static float GetGlobalSpawnCountScale() { return EngineGlobalSpawnCountScale; }
 	FORCEINLINE static float GetGlobalSystemCountScale() { return EngineGlobalSystemCountScale; }
+	static bool IsTargetPlatformIncludedInLevelRangeForCook(const ITargetPlatform* InTargetPlatform, const class UNiagaraEmitter* InEmitter);
 
 	static float EngineGlobalSpawnCountScale;
 	static float EngineGlobalSystemCountScale;
diff --git a/Engine/Plugins/FX/Niagara/Source/NiagaraCore/Private/NiagaraCustomVersion.cpp b/Engine/Plugins/FX/Niagara/Source/NiagaraCore/Private/NiagaraCustomVersion.cpp
--- a/Engine/Plugins/FX/Niagara/Source/NiagaraCore/Private/NiagaraCustomVersion.cpp
+++ b/Engine/Plugins/FX/Niagara/Source/NiagaraCore/Private/NiagaraCustomVersion.cpp
@@ -8,4 +8,4 @@ const FGuid FNiagaraCustomVersion::GUID(0xFCF57AFA, 0x50764283, 0xB9A9E658, 0xFF
 // Register the custom version with core
 FCustomVersionRegistration GRegisterNiagaraCustomVersion(FNiagaraCustomVersion::GUID, FNiagaraCustomVersion::LatestVersion, TEXT("NiagaraVer"));
 
-const FGuid FNiagaraCustomVersion::LatestScriptCompileVersion(0xBEACE717, 0x0D3A4051, 0xA442CF5E, 0x1CF94ACE);
\ No newline at end of file
+const FGuid FNiagaraCustomVersion::LatestScriptCompileVersion(0x0D9851F7, 0x8A2C4C60, 0xB17F4DEB, 0x8435B868);
diff --git a/Engine/Plugins/FX/Niagara/Source/NiagaraEditor/Private/NiagaraCompiler.cpp b/Engine/Plugins/FX/Niagara/Source/NiagaraEditor/Private/NiagaraCompiler.cpp
--- a/Engine/Plugins/FX/Niagara/Source/NiagaraEditor/Private/NiagaraCompiler.cpp
+++ b/Engine/Plugins/FX/Niagara/Source/NiagaraEditor/Private/NiagaraCompiler.cpp
@@ -50,8 +50,30 @@ static FAutoConsoleVariableRef CVarForceNiagaraTranslatorSingleThreaded(
 	ECVF_Default
 );
 
+// Enable this to log out generated HLSL for debugging purposes.
+static int32 GbForceNiagaraTranslatorDump = 0;
+static FAutoConsoleVariableRef CVarForceNiagaraTranslatorDump(
+	TEXT("fx.ForceNiagaraTranslatorDump"),
+	GbForceNiagaraTranslatorSingleThreaded,
+	TEXT("If > 0 all translation generated HLSL will be dumped \n"),
+	ECVF_Default
+);
+
 static FCriticalSection TranslationCritSec;
 
+void DumpHLSLText(const FString& SourceCode, const FString& DebugName)
+{
+	UE_LOG(LogNiagaraEditor, Log, TEXT("Compile output as text: %s"), *DebugName);
+	UE_LOG(LogNiagaraEditor, Log, TEXT("==================================================================================="));
+	TArray<FString> OutputByLines;
+	SourceCode.ParseIntoArrayLines(OutputByLines, false);
+	for (int32 i = 0; i < OutputByLines.Num(); i++)
+	{
+		UE_LOG(LogNiagaraEditor, Log, TEXT("/*%04d*/\t\t%s"), i + 1, *OutputByLines[i]);
+	}
+	UE_LOG(LogNiagaraEditor, Log, TEXT("==================================================================================="));
+}
+
 void FNiagaraCompileRequestData::VisitReferencedGraphs(UNiagaraGraph* InSrcGraph, UNiagaraGraph* InDupeGraph, ENiagaraScriptUsage InUsage, FCompileConstantResolver ConstantResolver)
 {
 	if (!InDupeGraph && !InSrcGraph)
@@ -234,6 +256,7 @@ void FNiagaraCompileRequestData::VisitReferencedGraphsRecursive(UNiagaraGraph* I
 				{
 					if (Ptr->EmitterUniqueName == EmitterNode->GetEmitterUniqueName())
 					{
+						EmitterNode->SyncEnabledState(); // Just to be safe, sync here while we likely still have the handle source.
 						EmitterNode->SetOwnerSystem(nullptr);
 						EmitterNode->SetCachedVariablesForCompilation(*Ptr->EmitterUniqueName, Ptr->NodeGraphDeepCopy, Ptr->Source);
 					}
@@ -470,7 +493,10 @@ TSharedPtr<FNiagaraCompileRequestDataBase, ESPMode::ThreadSafe> FNiagaraEditorMo
 			const FNiagaraEmitterHandle& Handle = System->GetEmitterHandle(i);
 			FCompileConstantResolver ConstantResolver(Handle.GetInstance());
 			TSharedPtr<FNiagaraCompileRequestData, ESPMode::ThreadSafe> EmitterPtr = MakeShared<FNiagaraCompileRequestData, ESPMode::ThreadSafe>();
-			EmitterPtr->DeepCopyGraphs(Cast<UNiagaraScriptSource>(Handle.GetInstance()->GraphSource), ENiagaraScriptUsage::EmitterSpawnScript, ConstantResolver);
+			if (Handle.GetIsEnabled()) // Don't need to copy the graph if we aren't going to use it.
+			{
+				EmitterPtr->DeepCopyGraphs(Cast<UNiagaraScriptSource>(Handle.GetInstance()->GraphSource), ENiagaraScriptUsage::EmitterSpawnScript, ConstantResolver);
+			}
 			EmitterPtr->EmitterUniqueName = Handle.GetInstance()->GetUniqueEmitterName();
 			EmitterPtr->SourceName = BasePtr->SourceName;
 			BasePtr->EmitterData.Add(EmitterPtr);
@@ -490,8 +516,11 @@ TSharedPtr<FNiagaraCompileRequestDataBase, ESPMode::ThreadSafe> FNiagaraEditorMo
 		{
 			const FNiagaraEmitterHandle& Handle = System->GetEmitterHandle(i);
 			FCompileConstantResolver ConstantResolver(Handle.GetInstance());
-			BasePtr->EmitterData[i]->FinishPrecompile(Cast<UNiagaraScriptSource>(Handle.GetInstance()->GraphSource), EncounterableVars, ENiagaraScriptUsage::EmitterSpawnScript, ConstantResolver);
-			BasePtr->MergeInEmitterPrecompiledData(BasePtr->EmitterData[i].Get());
+			if (Handle.GetIsEnabled()) // Don't pull in the emitter if it isn't going to be used.
+			{
+				BasePtr->EmitterData[i]->FinishPrecompile(Cast<UNiagaraScriptSource>(Handle.GetInstance()->GraphSource), EncounterableVars, ENiagaraScriptUsage::EmitterSpawnScript, ConstantResolver);
+				BasePtr->MergeInEmitterPrecompiledData(BasePtr->EmitterData[i].Get());
+			}
 		}
 
 	}
@@ -544,6 +573,11 @@ TSharedPtr<FNiagaraVMExecutableData> FNiagaraEditorModule::CompileScript(const F
 		}
 		TranslationTime = (float)(FPlatformTime::Seconds() - TranslationStartTime);
 
+		if (GbForceNiagaraTranslatorDump != 0)
+		{
+			DumpHLSLText(Translator.GetTranslatedHLSL(), InCompileOptions.FullName);
+		}
+		
 		Results = Compiler.CompileScript(CompileRequest, InCompileOptions, &Translator.GetTranslateOutput(), Translator.GetTranslatedHLSL());
 		VMCompilationTime = Results.CompileTime;
 	}
diff --git a/Engine/Plugins/FX/Niagara/Source/NiagaraEditor/Private/NiagaraHlslTranslator.cpp b/Engine/Plugins/FX/Niagara/Source/NiagaraEditor/Private/NiagaraHlslTranslator.cpp
--- a/Engine/Plugins/FX/Niagara/Source/NiagaraEditor/Private/NiagaraHlslTranslator.cpp
+++ b/Engine/Plugins/FX/Niagara/Source/NiagaraEditor/Private/NiagaraHlslTranslator.cpp
@@ -3729,6 +3729,24 @@ void FHlslNiagaraTranslator::Emitter(class UNiagaraNodeEmitter* EmitterNode, TAr
 {
 	SCOPE_CYCLE_COUNTER(STAT_NiagaraEditor_Module_NiagaraHLSLTranslator_Emitter);
 
+	// Just pass through the input parameter map pin if the node isn't enabled...
+	if (!EmitterNode->IsNodeEnabled())
+	{
+		TArray<UEdGraphPin*> OutputPins;
+		EmitterNode->GetOutputPins(OutputPins);
+
+		Outputs.SetNum(OutputPins.Num());
+		for (int32 i = 0; i < OutputPins.Num(); i++)
+		{
+			Outputs[i] = INDEX_NONE;
+		}
+		if (Inputs.Num() >= 1)
+		{
+			Outputs[0] = Inputs[0];
+		}
+		return;
+	}
+
 	FNiagaraFunctionSignature Signature;
 	UNiagaraScriptSource* Source = EmitterNode->GetScriptSource();
 	if (Source == nullptr)
diff --git a/Engine/Plugins/FX/Niagara/Source/NiagaraEditor/Private/NiagaraNodeEmitter.cpp b/Engine/Plugins/FX/Niagara/Source/NiagaraEditor/Private/NiagaraNodeEmitter.cpp
--- a/Engine/Plugins/FX/Niagara/Source/NiagaraEditor/Private/NiagaraNodeEmitter.cpp
+++ b/Engine/Plugins/FX/Niagara/Source/NiagaraEditor/Private/NiagaraNodeEmitter.cpp
@@ -41,7 +41,7 @@ UNiagaraSystem* UNiagaraNodeEmitter::GetOwnerSystem() const
 void UNiagaraNodeEmitter::SetOwnerSystem(UNiagaraSystem* InOwnerSystem)
 {
 	OwnerSystem = InOwnerSystem;
-	DisplayName = GetNameFromEmitter();
+	RefreshFromExternalChanges();
 }
 
 FGuid UNiagaraNodeEmitter::GetEmitterHandleId() const
@@ -226,9 +226,31 @@ UNiagaraGraph* UNiagaraNodeEmitter::GetCalledGraph() const
 bool UNiagaraNodeEmitter::RefreshFromExternalChanges()
 {
 	DisplayName = GetNameFromEmitter();
+	SyncEnabledState();
 	return true;
 }
 
+void UNiagaraNodeEmitter::SyncEnabledState()
+{
+	if (OwnerSystem != nullptr && EmitterHandleId.IsValid())
+	{
+		for (const FNiagaraEmitterHandle& EmitterHandle : OwnerSystem->GetEmitterHandles())
+		{
+			if (EmitterHandle.GetId() == EmitterHandleId)
+			{
+				if (EmitterHandle.GetIsEnabled() == false)
+				{
+					SetEnabledState(ENodeEnabledState::Disabled, false);
+				}
+				else
+				{
+					SetEnabledState(ENodeEnabledState::Enabled, false);
+				}
+			}
+		}
+	}
+}
+
 void UNiagaraNodeEmitter::SetCachedVariablesForCompilation(const FName& InUniqueName, UNiagaraGraph* InGraph, UNiagaraScriptSourceBase* InSource)
 {
 	CachedUniqueName = InUniqueName;
@@ -387,29 +409,44 @@ void UNiagaraNodeEmitter::Compile(FHlslNiagaraTranslator *Translator, TArray<int
 	Options.bFilterDuplicates = true;
 	Options.bFilterByScriptUsage = true;
 	Options.TargetScriptUsage = Translator->GetTargetUsage() == ENiagaraScriptUsage::SystemSpawnScript ? ENiagaraScriptUsage::EmitterSpawnScript : ENiagaraScriptUsage::EmitterUpdateScript;
-	CalledGraph->FindInputNodes(InputsNodes, Options);
+	
+	if (CalledGraph && IsNodeEnabled()) // Called graph may be null on an disabled emitter
+	{
+		CalledGraph->FindInputNodes(InputsNodes, Options);
+	}
 
 	TArray<int32> CompileInputs;
-	CompileInputs.Reserve(InputsNodes.Num());
+	
+	if (InputPins.Num() > 1)
+	{
+		Translator->Error(LOCTEXT("TooManyOutputPinsError", "Too many input pins on node."), this, nullptr);
+		return;
+	}
 
-	if (InputsNodes.Num() <= 0)
+	int32 InputPinCompiled = Translator->CompilePin(InputPins[0]);
+	if (!IsNodeEnabled())
 	{
-		Translator->Error(LOCTEXT("InputNodesNotFound", "Input nodes on called graph not found"), this, nullptr);
+		// Do the minimal amount of work necessary if we are disabled.
+		CompileInputs.Reserve(1);
+		CompileInputs.Add(InputPinCompiled);
+		Translator->Emitter(this, CompileInputs, Outputs);
 		return;
 	}
 
-	if (InputPins.Num() > 1)
+	if (InputsNodes.Num() <= 0)
 	{
-		Translator->Error(LOCTEXT("TooManyOutputPinsError", "Too many input pins on node."), this, nullptr);
+		Translator->Error(LOCTEXT("InputNodesNotFound", "Input nodes on called graph not found"), this, nullptr);
 		return;
 	}
 
+	CompileInputs.Reserve(InputsNodes.Num());
+
 	bool bError = false;
 	for (UNiagaraNodeInput* EmitterInputNode : InputsNodes)
 	{
 		if (EmitterInputNode->Input.IsEquivalent(FNiagaraVariable(FNiagaraTypeDefinition::GetParameterMapDef(), TEXT("InputMap"))))
 		{
-			CompileInputs.Add(Translator->CompilePin(InputPins[0]));
+			CompileInputs.Add(InputPinCompiled);
 		}
 		else
 		{
@@ -427,7 +464,7 @@ void UNiagaraNodeEmitter::GatherExternalDependencyIDs(ENiagaraScriptUsage InMast
 {
 	UNiagaraGraph* CalledGraph = GetCalledGraph();
 
-	if (CalledGraph)
+	if (CalledGraph && IsNodeEnabled()) // Skip if disabled
 	{
 		ENiagaraScriptUsage TargetUsage = InMasterUsage == ENiagaraScriptUsage::SystemSpawnScript ? ENiagaraScriptUsage::EmitterSpawnScript : ENiagaraScriptUsage::EmitterUpdateScript;
 		InReferencedCompileHashes.Add(CalledGraph->GetCompileDataHash(TargetUsage, FGuid(0,0,0,0)));
diff --git a/Engine/Plugins/FX/Niagara/Source/NiagaraEditor/Private/NiagaraNodeEmitter.h b/Engine/Plugins/FX/Niagara/Source/NiagaraEditor/Private/NiagaraNodeEmitter.h
--- a/Engine/Plugins/FX/Niagara/Source/NiagaraEditor/Private/NiagaraNodeEmitter.h
+++ b/Engine/Plugins/FX/Niagara/Source/NiagaraEditor/Private/NiagaraNodeEmitter.h
@@ -47,6 +47,9 @@ public:
 
 	virtual void BuildParameterMapHistory(FNiagaraParameterMapHistoryBuilder& OutHistory, bool bRecursive = true, bool bFilterForCompilation = true) const override;
 
+	/** Synchronize with the handle associated with this emitter for enabled/disabled state.*/
+	void SyncEnabledState();
+
 	virtual FText GetNodeTitle(ENodeTitleType::Type TitleType) const;
 
 	ENiagaraScriptUsage GetUsage() const { return ScriptType; }
diff --git a/Engine/Plugins/FX/Niagara/Source/NiagaraEditor/Private/NiagaraScriptSource.cpp b/Engine/Plugins/FX/Niagara/Source/NiagaraEditor/Private/NiagaraScriptSource.cpp
--- a/Engine/Plugins/FX/Niagara/Source/NiagaraEditor/Private/NiagaraScriptSource.cpp
+++ b/Engine/Plugins/FX/Niagara/Source/NiagaraEditor/Private/NiagaraScriptSource.cpp
@@ -72,6 +72,21 @@ void UNiagaraScriptSource::InvalidateCachedCompileIds()
 	NodeGraph->InvalidateCachedCompileIds();
 }
 
+void UNiagaraScriptSource::RefreshFromExternalChanges()
+{
+	if (NodeGraph)
+	{
+		for (UEdGraphNode* Node : NodeGraph->Nodes)
+		{
+			UNiagaraNode* NiagaraNode = Cast<UNiagaraNode>(Node);
+			if (NiagaraNode)
+			{
+				NiagaraNode->RefreshFromExternalChanges();
+			}
+		}
+	}
+}
+
 
 void UNiagaraScriptSource::PostLoad()
 {
diff --git a/Engine/Plugins/FX/Niagara/Source/NiagaraEditor/Private/ViewModels/NiagaraEmitterHandleViewModel.cpp b/Engine/Plugins/FX/Niagara/Source/NiagaraEditor/Private/ViewModels/NiagaraEmitterHandleViewModel.cpp
--- a/Engine/Plugins/FX/Niagara/Source/NiagaraEditor/Private/ViewModels/NiagaraEmitterHandleViewModel.cpp
+++ b/Engine/Plugins/FX/Niagara/Source/NiagaraEditor/Private/ViewModels/NiagaraEmitterHandleViewModel.cpp
@@ -213,7 +213,7 @@ void FNiagaraEmitterHandleViewModel::SetIsEnabled(bool bInIsEnabled)
 	{
 		FScopedTransaction ScopedTransaction(NSLOCTEXT("NiagaraEmitterEditor", "EditEmitterEnabled", "Change emitter enabled state"));
 		GetOwningSystemViewModel()->GetSystem().Modify();
-		EmitterHandle->SetIsEnabled(bInIsEnabled);
+		EmitterHandle->SetIsEnabled(bInIsEnabled, GetOwningSystemViewModel()->GetSystem(), true);
 		OnPropertyChangedDelegate.Broadcast();
 	}
 }
diff --git a/Engine/Plugins/FX/Niagara/Source/NiagaraEditor/Private/ViewModels/NiagaraEmitterViewModel.cpp b/Engine/Plugins/FX/Niagara/Source/NiagaraEditor/Private/ViewModels/NiagaraEmitterViewModel.cpp
--- a/Engine/Plugins/FX/Niagara/Source/NiagaraEditor/Private/ViewModels/NiagaraEmitterViewModel.cpp
+++ b/Engine/Plugins/FX/Niagara/Source/NiagaraEditor/Private/ViewModels/NiagaraEmitterViewModel.cpp
@@ -24,6 +24,7 @@ template<> TMap<UNiagaraEmitter*, TArray<FNiagaraEmitterViewModel*>> TNiagaraVie
 
 const FText FNiagaraEmitterViewModel::StatsFormat = NSLOCTEXT("NiagaraEmitterViewModel", "StatsFormat", "{0} Particles | {1} ms | {2} MB | {3}");
 const FText FNiagaraEmitterViewModel::StatsParticleCountFormat = NSLOCTEXT("NiagaraEmitterViewModel", "StatsParticleCountFormat", "{0} Particles");
+const FText FNiagaraEmitterViewModel::ParticleDisabledDueToDetailLevel = NSLOCTEXT("NiagaraEmitterViewModel", "ParticleDisabledDueToDetailLevel", "Disabled due to detail level mismatch (Current: {0}). See Scalability options for valid range. ");
 
 namespace NiagaraCommands
 {
@@ -189,15 +190,21 @@ FText FNiagaraEmitterViewModel::GetStatsText() const
 				{
 					return LOCTEXT("NullInstance", "Invalid instance");
 				}
+				
+				if (Handle.GetIsEnabled() == false)
+				{
+					return LOCTEXT("DisabledSimulation", "Emitter is not enabled. Excluded from code.");
+				}
 
 				if (!HandleEmitter->IsValid())
 				{
 					return LOCTEXT("InvalidInstance", "Invalid Emitter! May have compile errors.");
 				}
 
-				if (Handle.GetIsEnabled() == false)
+				int32 DetailLevel = SimInstance->GetParentSystemInstance()->GetDetailLevel();
+				if (!HandleEmitter->IsAllowedByDetailLevel(DetailLevel))
 				{
-					return LOCTEXT("DisabledSimulation", "Simulation is not enabled.");
+					return FText::Format(ParticleDisabledDueToDetailLevel, FText::AsNumber(DetailLevel));
 				}
 
 				if (NiagaraCommands::EmitterStatsFormat->GetInt() == 1)
@@ -228,6 +235,25 @@ TSharedRef<FNiagaraScriptViewModel> FNiagaraEmitterViewModel::GetSharedScriptVie
 	return SharedScriptViewModel.ToSharedRef();
 }
 
+bool FNiagaraEmitterViewModel::IsEnabledByDetailLevel() const
+{
+	if (Simulation.IsValid())
+	{
+		TSharedPtr<FNiagaraEmitterInstance, ESPMode::ThreadSafe> SimInstance = Simulation.Pin();
+		if (SimInstance.IsValid())
+		{
+			const FNiagaraEmitterHandle& Handle = SimInstance->GetEmitterHandle();
+			if (Handle.GetInstance())
+			{
+				UNiagaraEmitter* HandleEmitter = Handle.GetInstance();
+				int32 DetailLevel = SimInstance->GetParentSystemInstance()->GetDetailLevel();
+				return HandleEmitter->IsAllowedByDetailLevel(DetailLevel);
+			}
+		}
+	}
+	return false;
+}
+
 const UNiagaraEmitterEditorData& FNiagaraEmitterViewModel::GetEditorData() const
 {
 	check(Emitter.IsValid());
diff --git a/Engine/Plugins/FX/Niagara/Source/NiagaraEditor/Private/ViewModels/NiagaraScriptViewModel.h b/Engine/Plugins/FX/Niagara/Source/NiagaraEditor/Private/ViewModels/NiagaraScriptViewModel.h
--- a/Engine/Plugins/FX/Niagara/Source/NiagaraEditor/Private/ViewModels/NiagaraScriptViewModel.h
+++ b/Engine/Plugins/FX/Niagara/Source/NiagaraEditor/Private/ViewModels/NiagaraScriptViewModel.h
@@ -56,7 +56,7 @@ public:
 	void CompileStandaloneScript(bool bForceCompile = false);
 
 	/** Get the latest status of this view-model's script compilation.*/
-	ENiagaraScriptCompileStatus GetLatestCompileStatus();
+	virtual ENiagaraScriptCompileStatus GetLatestCompileStatus();
 
 	/** Refreshes the nodes in the script graph, updating the pins to match external changes. */
 	void RefreshNodes();
diff --git a/Engine/Plugins/FX/Niagara/Source/NiagaraEditor/Private/ViewModels/NiagaraSystemScriptViewModel.cpp b/Engine/Plugins/FX/Niagara/Source/NiagaraEditor/Private/ViewModels/NiagaraSystemScriptViewModel.cpp
--- a/Engine/Plugins/FX/Niagara/Source/NiagaraEditor/Private/ViewModels/NiagaraSystemScriptViewModel.cpp
+++ b/Engine/Plugins/FX/Niagara/Source/NiagaraEditor/Private/ViewModels/NiagaraSystemScriptViewModel.cpp
@@ -58,18 +58,37 @@ void FNiagaraSystemScriptViewModel::OnSystemVMCompiled(UNiagaraSystem* InSystem)
 	FString AggregateErrors;
 
 	TArray<UNiagaraScript*> SystemScripts;
+	TArray<bool> ScriptsEnabled;
 	SystemScripts.Add(InSystem->GetSystemSpawnScript());
 	SystemScripts.Add(InSystem->GetSystemUpdateScript());
+	ScriptsEnabled.Add(true);
+	ScriptsEnabled.Add(true);
+
+	
 	for (const FNiagaraEmitterHandle& Handle : InSystem->GetEmitterHandles())
 	{
+		int32 NumScripts = SystemScripts.Num();
 		Handle.GetInstance()->GetScripts(SystemScripts, true);
+		for (; NumScripts < SystemScripts.Num(); NumScripts++)
+		{
+			if (Handle.GetIsEnabled())
+			{
+				ScriptsEnabled.Add(true);
+			}
+			else
+			{
+				ScriptsEnabled.Add(false);
+			}
+		}
 	}
 
+	check(ScriptsEnabled.Num() == SystemScripts.Num());
+
 	int32 EventsFound = 0;
 	for (int32 i = 0; i < SystemScripts.Num(); i++)
 	{
 		UNiagaraScript* Script = SystemScripts[i];
-		if (Script != nullptr && Script->GetVMExecutableData().IsValid())
+		if (Script != nullptr && Script->GetVMExecutableData().IsValid() && ScriptsEnabled[i])
 		{
 			InCompileStatuses.Add(Script->GetVMExecutableData().LastCompileStatus);
 			InCompileErrors.Add(Script->GetVMExecutableData().ErrorMsg);
@@ -85,6 +104,22 @@ void FNiagaraSystemScriptViewModel::OnSystemVMCompiled(UNiagaraSystem* InSystem)
 				InUsages.Add(TPair<ENiagaraScriptUsage, int32>(Script->GetUsage(), 0));
 			}
 		}
+		else if (Script != nullptr && ScriptsEnabled[i] == false)
+		{
+			InCompileStatuses.Add(ENiagaraScriptCompileStatus::NCS_UpToDate);
+			InCompileErrors.Add(FString());
+			InCompilePaths.Add(Script->GetPathName());
+
+			if (Script->GetUsage() == ENiagaraScriptUsage::ParticleEventScript)
+			{
+				InUsages.Add(TPair<ENiagaraScriptUsage, int32>(Script->GetUsage(), EventsFound));
+				EventsFound++;
+			}
+			else
+			{
+				InUsages.Add(TPair<ENiagaraScriptUsage, int32>(Script->GetUsage(), 0));
+			}
+		}
 		else
 		{
 			InCompileStatuses.Add(ENiagaraScriptCompileStatus::NCS_Unknown);
@@ -120,3 +155,39 @@ void FNiagaraSystemScriptViewModel::CompileSystem(bool bForce)
 {
 	System->RequestCompile(bForce);
 }
+
+ENiagaraScriptCompileStatus FNiagaraSystemScriptViewModel::GetLatestCompileStatus()
+{
+	TArray<UNiagaraScript*> SystemScripts;
+	SystemScripts.Add(System->GetSystemSpawnScript());
+	SystemScripts.Add(System->GetSystemUpdateScript());
+	for (const FNiagaraEmitterHandle& Handle : System->GetEmitterHandles())
+	{
+		if (Handle.GetIsEnabled())
+		{
+			Handle.GetInstance()->GetScripts(SystemScripts, true);
+		}
+	}
+
+	bool bDirty = false;
+	for (int32 i = 0; i < SystemScripts.Num(); i++)
+	{
+		if (!SystemScripts[i])
+		{
+			continue;
+		}
+
+		if (SystemScripts[i]->IsCompilable() && !SystemScripts[i]->AreScriptAndSourceSynchronized())
+		{
+			bDirty = true;
+			break;
+		}
+	}
+
+	if (bDirty)
+	{
+		return ENiagaraScriptCompileStatus::NCS_Dirty;
+	}
+	return LastCompileStatus;
+}
+
diff --git a/Engine/Plugins/FX/Niagara/Source/NiagaraEditor/Private/ViewModels/NiagaraSystemScriptViewModel.h b/Engine/Plugins/FX/Niagara/Source/NiagaraEditor/Private/ViewModels/NiagaraSystemScriptViewModel.h
--- a/Engine/Plugins/FX/Niagara/Source/NiagaraEditor/Private/ViewModels/NiagaraSystemScriptViewModel.h
+++ b/Engine/Plugins/FX/Niagara/Source/NiagaraEditor/Private/ViewModels/NiagaraSystemScriptViewModel.h
@@ -24,6 +24,8 @@ public:
 	FOnSystemCompiled& OnSystemCompiled();
 
 	void CompileSystem(bool bForce);
+	
+	virtual ENiagaraScriptCompileStatus GetLatestCompileStatus() override;
 
 private:
 	void OnSystemVMCompiled(UNiagaraSystem* InSystem);
diff --git a/Engine/Plugins/FX/Niagara/Source/NiagaraEditor/Private/ViewModels/NiagaraSystemViewModel.cpp b/Engine/Plugins/FX/Niagara/Source/NiagaraEditor/Private/ViewModels/NiagaraSystemViewModel.cpp
--- a/Engine/Plugins/FX/Niagara/Source/NiagaraEditor/Private/ViewModels/NiagaraSystemViewModel.cpp
+++ b/Engine/Plugins/FX/Niagara/Source/NiagaraEditor/Private/ViewModels/NiagaraSystemViewModel.cpp
@@ -495,7 +495,7 @@ void FNiagaraSystemViewModel::Tick(float DeltaTime)
 
 		for (TSharedRef<FNiagaraEmitterHandleViewModel> EmitterHandleViewModel : EmitterHandleViewModels)
 		{
-			if (EmitterHandleViewModel->GetEmitterViewModel()->GetLatestCompileStatus() == ENiagaraScriptCompileStatus::NCS_Dirty)
+			if (EmitterHandleViewModel->GetIsEnabled() && EmitterHandleViewModel->GetEmitterViewModel()->GetLatestCompileStatus() == ENiagaraScriptCompileStatus::NCS_Dirty)
 			{
 				bRecompile |= true;
 				//EmitterHandleViewModel->GetEmitterViewModel()->CompileScripts();
diff --git a/Engine/Plugins/FX/Niagara/Source/NiagaraEditor/Public/NiagaraScriptSource.h b/Engine/Plugins/FX/Niagara/Source/NiagaraEditor/Public/NiagaraScriptSource.h
--- a/Engine/Plugins/FX/Niagara/Source/NiagaraEditor/Public/NiagaraScriptSource.h
+++ b/Engine/Plugins/FX/Niagara/Source/NiagaraEditor/Public/NiagaraScriptSource.h
@@ -47,6 +47,8 @@ class UNiagaraScriptSource : public UNiagaraScriptSourceBase
 
 	virtual void CleanUpOldAndInitializeNewRapidIterationParameters(FString UniqueEmitterName, ENiagaraScriptUsage ScriptUsage, FGuid ScriptUsageId, FNiagaraParameterStore& RapidIterationParameters) const override;
 	virtual void InvalidateCachedCompileIds() override;
+	virtual void RefreshFromExternalChanges() override;
+
 private:
 	void OnGraphChanged(const FEdGraphEditAction &Action);
 	void OnGraphDataInterfaceChanged();
diff --git a/Engine/Plugins/FX/Niagara/Source/NiagaraEditor/Public/ViewModels/NiagaraEmitterViewModel.h b/Engine/Plugins/FX/Niagara/Source/NiagaraEditor/Public/ViewModels/NiagaraEmitterViewModel.h
--- a/Engine/Plugins/FX/Niagara/Source/NiagaraEditor/Public/ViewModels/NiagaraEmitterViewModel.h
+++ b/Engine/Plugins/FX/Niagara/Source/NiagaraEditor/Public/ViewModels/NiagaraEmitterViewModel.h
@@ -99,6 +99,8 @@ public:
 
 	void Cleanup();
 
+	bool IsEnabledByDetailLevel() const;
+
 private:
 	/** Sets this view model to a different emitter. */
 	void SetEmitter(UNiagaraEmitter* InEmitter);
@@ -122,6 +124,9 @@ private:
 	/** The text format stats to only display particles count. */
 	static const FText StatsParticleCountFormat;
 
+	/** The text format stats to display when detail level mismatches. */
+	static const FText ParticleDisabledDueToDetailLevel;
+
 	/** The emitter object being displayed by the control .*/
 	TWeakObjectPtr<UNiagaraEmitter> Emitter;
 
