Added a component pool for niagara components to allow them to be recycled for better performance

diff --git a/Engine/Plugins/FX/Niagara/Source/Niagara/Private/NiagaraComponent.cpp b/Engine/Plugins/FX/Niagara/Source/Niagara/Private/NiagaraComponent.cpp
--- a/Engine/Plugins/FX/Niagara/Source/Niagara/Private/NiagaraComponent.cpp
+++ b/Engine/Plugins/FX/Niagara/Source/Niagara/Private/NiagaraComponent.cpp
@@ -417,6 +417,38 @@ void UNiagaraComponent::SetEmitterEnable(FName EmitterName, bool bNewEnableState
 		SystemInstance->SetEmitterEnable(EmitterName, bNewEnableState);
 	}
 }
+
+void UNiagaraComponent::ReleaseToPool()
+{
+	if (PoolingMethod != ENCPoolMethod::ManualRelease)
+	{
+		UE_LOG(LogNiagara, Warning, TEXT("Manually releasing a PSC to the pool that was not spawned with ENCPoolMethod::ManualRelease. Asset=%s Component=%s"),
+			Asset ? *Asset->GetPathName() : TEXT("NULL"), *GetPathName()
+		);
+		return;
+	}
+
+	if (!bIsActive)
+	{
+		//If we're already complete then release to the pool straight away.
+		UWorld* World = GetWorld();
+		check(World);
+		FNiagaraWorldManager::Get(World)->GetComponentPool()->ReclaimWorldParticleSystem(this);
+	}
+	else
+	{
+		//If we haven't completed, deactivate and defer release to pool.
+		PoolingMethod = ENCPoolMethod::ManualRelease_OnComplete;
+		Deactivate();
+	}
+}
+
+uint32 UNiagaraComponent::GetApproxMemoryUsage() const
+{
+	// TODO: implement memory usage for the component pool statistics
+	return 1;
+}
+
 /********* UFXSystemComponent *********/
 
 
@@ -823,7 +855,16 @@ void UNiagaraComponent::OnSystemComplete()
 	OnSystemFinished.Broadcast(this);
 	//UE_LOG(LogNiagara, Log, TEXT("OnSystemFinished.Broadcast(this);: } %p - %s"), SystemInstance.Get(), *Asset->GetName());
 
-	if (bAutoDestroy)
+	if (PoolingMethod == ENCPoolMethod::AutoRelease)
+	{
+		FNiagaraWorldManager::Get(GetWorld())->GetComponentPool()->ReclaimWorldParticleSystem(this);
+	}
+	else if (PoolingMethod == ENCPoolMethod::ManualRelease_OnComplete)
+	{
+		PoolingMethod = ENCPoolMethod::ManualRelease;
+		FNiagaraWorldManager::Get(GetWorld())->GetComponentPool()->ReclaimWorldParticleSystem(this);
+	}
+	else if (bAutoDestroy)
 	{
 		//UE_LOG(LogNiagara, Log, TEXT("OnSystemComplete DestroyComponent();: { %p - %s"), SystemInstance.Get(), *Asset->GetName());
 		DestroyComponent();
@@ -1352,6 +1393,30 @@ void UNiagaraComponent::PostEditChangeProperty(FPropertyChangedEvent& PropertyCh
 	Super::PostEditChangeProperty(PropertyChangedEvent);
 }
 
+void UNiagaraComponent::SetUserParametersToDefaultValues()
+{
+	OverrideParameters.Empty();
+	if (Asset == nullptr)
+	{		
+		return;
+	}
+
+	TArray<FNiagaraVariable> SourceVars;
+	Asset->GetExposedParameters().GetParameters(SourceVars);
+	for (FNiagaraVariable& Param : SourceVars)
+	{
+		OverrideParameters.AddParameter(Param, true);
+	}
+
+	TArray<FNiagaraVariable> ExistingVars;
+	OverrideParameters.GetUserParameters(ExistingVars);
+	for (FNiagaraVariable ExistingVar : ExistingVars)
+	{
+		Asset->GetExposedParameters().CopyParameterData(OverrideParameters, ExistingVar);
+	}
+
+	OverrideParameters.Rebind();
+}
 
 void UNiagaraComponent::SynchronizeWithSourceSystem()
 {
diff --git a/Engine/Plugins/FX/Niagara/Source/Niagara/Private/NiagaraComponentPool.cpp b/Engine/Plugins/FX/Niagara/Source/Niagara/Private/NiagaraComponentPool.cpp
new file mode 100644
--- /dev/null
+++ b/Engine/Plugins/FX/Niagara/Source/Niagara/Private/NiagaraComponentPool.cpp
@@ -0,0 +1,434 @@
+// Copyright 1998-2019 Epic Games, Inc. All Rights Reserved.
+
+#include "NiagaraComponentPool.h"
+#include "HAL/IConsoleManager.h"
+#include "Engine/World.h"
+#include "NiagaraComponent.h"
+#include "NiagaraWorldManager.h"
+
+static float GParticleSystemPoolKillUnusedTime = 180.0f;
+static FAutoConsoleVariableRef ParticleSystemPoolKillUnusedTime(
+	TEXT("FX.NiagaraComponentPool.KillUnusedTime"),
+	GParticleSystemPoolKillUnusedTime,
+	TEXT("How long a pooled particle component needs to be unused for before it is destroyed.")
+);
+
+static int32 GbEnableParticleSystemPooling = 1;
+static FAutoConsoleVariableRef bEnableParticleSystemPooling(
+	TEXT("FX.NiagaraComponentPool.Enable"),
+	GbEnableParticleSystemPooling,
+	TEXT("How many Particle System Components to preallocate when creating new ones for the pool.")
+);
+
+static float GParticleSystemPoolingCleanTime = 30.0f;
+static FAutoConsoleVariableRef ParticleSystemPoolingCleanTime(
+	TEXT("FX.NiagaraComponentPool.CleanTime"),
+	GParticleSystemPoolingCleanTime,
+	TEXT("How often should the pool be cleaned (in seconds).")
+);
+
+void DumpPooledWorldParticleSystemInfo(UWorld* World)
+{
+	check(World);
+	FNiagaraWorldManager::Get(World)->GetComponentPool()->Dump();
+}
+
+FAutoConsoleCommandWithWorld DumpNCPoolInfoCommand(
+	TEXT("FX.DumpNCPoolInfo"),
+	TEXT("Dump Particle System Pooling Info"),
+	FConsoleCommandWithWorldDelegate::CreateStatic(&DumpPooledWorldParticleSystemInfo)
+);
+
+FNCPool::FNCPool()
+	: MaxUsed(0)
+{
+
+}
+
+void FNCPool::Cleanup()
+{
+	for (FNCPoolElement& Elem : FreeElements)
+	{
+		if (Elem.Component)
+		{
+			Elem.Component->PoolingMethod = ENCPoolMethod::None;//Reset so we don't trigger warnings about destroying pooled NCs.
+			Elem.Component->DestroyComponent();
+		}
+		else
+		{
+			UE_LOG(LogNiagara, Error, TEXT("Free element in the NiagaraComponentPool was null. Someone must be keeping a reference to a NC that has been freed to the pool and then are manually destroying it."));
+		}
+	}
+
+	for (UNiagaraComponent* NC : InUseComponents_Auto)
+	{
+		//It's possible for people to manually destroy these so we have to guard against it. Though we warn about it in UNiagaraComponent::BeginDestroy
+		if (NC)
+		{
+			NC->PoolingMethod = ENCPoolMethod::None;//Reset so we don't trigger warnings about destroying pooled NCs.
+			NC->DestroyComponent();
+		}
+	}
+
+	//Warn if there are any manually released NCs still in the world at cleanup time.
+	for (UNiagaraComponent* NC : InUseComponents_Manual)
+	{
+		//It's possible for people to manually destroy these so we have to guard against it. Though we warn about it in UNiagaraComponent::BeginDestroy
+		if (NC)
+		{
+			UE_LOG(LogNiagara, Warning, TEXT("Pooled NC set to manual release is still in use as the pool is being cleaned up. %s"), *NC->GetAsset()->GetFullName());
+			NC->PoolingMethod = ENCPoolMethod::None;//Reset so we don't trigger warnings about destroying pooled NCs.
+			NC->DestroyComponent();
+		}
+	}
+
+	FreeElements.Empty();
+	InUseComponents_Auto.Empty();
+	InUseComponents_Manual.Empty();
+}
+
+UNiagaraComponent* FNCPool::Acquire(UWorld* World, UNiagaraSystem* Template, ENCPoolMethod PoolingMethod)
+{
+	check(GbEnableParticleSystemPooling);
+	check(PoolingMethod != ENCPoolMethod::None);
+
+	FNCPoolElement RetElem;
+	if (FreeElements.Num())
+	{
+		RetElem = FreeElements.Pop(false);
+		check(RetElem.Component->GetAsset() == Template);
+		check(!RetElem.Component->IsPendingKill());
+		RetElem.Component->SetUserParametersToDefaultValues();
+
+		if (RetElem.Component->GetWorld() != World)
+		{
+			// Rename the NC to move it into the current PersistentLevel - it may have been spawned in one
+			// level but is now needed in another level.
+			// Use the REN_ForceNoResetLoaders flag to prevent the rename from potentially calling FlushAsyncLoading.
+			RetElem.Component->Rename(nullptr, World, REN_ForceNoResetLoaders);
+		}
+	}
+	else
+	{
+		//None in the pool so create a new one.
+		RetElem.Component = NewObject<UNiagaraComponent>(World);
+		RetElem.Component->SetAutoDestroy(false);// we don't auto destroy, just periodically clear up the pool.
+		RetElem.Component->bAutoActivate = false;
+		RetElem.Component->SetAsset(Template);
+	}
+
+	RetElem.Component->PoolingMethod = PoolingMethod;
+
+#if ENABLE_NC_POOL_DEBUGGING
+	if (PoolingMethod == ENCPoolMethod::AutoRelease)
+	{
+		InUseComponents_Auto.Add(RetElem.Component);
+	}
+	else if (PoolingMethod == ENCPoolMethod::ManualRelease)
+	{
+		InUseComponents_Manual.Add(RetElem.Component);
+	}
+#endif 
+
+	MaxUsed = FMath::Max(MaxUsed, InUseComponents_Manual.Num() + InUseComponents_Auto.Num());
+	return RetElem.Component;
+}
+
+void FNCPool::Reclaim(UNiagaraComponent* Component, const float CurrentTimeSeconds)
+{
+	check(Component);
+
+#if ENABLE_NC_POOL_DEBUGGING
+	int32 InUseIdx = INDEX_NONE;
+	if (Component->PoolingMethod == ENCPoolMethod::AutoRelease)
+	{
+		InUseIdx = InUseComponents_Auto.IndexOfByKey(Component);
+		if (InUseIdx != INDEX_NONE)
+		{
+			InUseComponents_Auto.RemoveAtSwap(InUseIdx);
+		}
+	}
+	else if (Component->PoolingMethod == ENCPoolMethod::ManualRelease)
+	{
+		InUseIdx = InUseComponents_Manual.IndexOfByKey(Component);
+		if (InUseIdx != INDEX_NONE)
+		{
+			InUseComponents_Manual.RemoveAtSwap(InUseIdx);
+		}
+	}
+	
+	if(InUseIdx == INDEX_NONE)
+	{
+		UE_LOG(LogNiagara, Error, TEXT("World Particle System Pool is reclaiming a component that is not in it's InUse list!"));
+	}
+#endif
+
+	//Don't add back to the pool if we're no longer pooling or we've hit our max resident pool size.
+	if (GbEnableParticleSystemPooling != 0 && FreeElements.Num() < (int32)Component->GetAsset()->MaxPoolSize)
+	{
+		Component->DeactivateImmediate();
+
+		Component->DetachFromComponent(FDetachmentTransformRules::KeepWorldTransform); // When detaching, maintain world position for optimization purposes
+		Component->SetRelativeScale3D(FVector(1.f)); // Reset scale to avoid future uses of this NC having incorrect scale
+		Component->SetAbsolute(); // Clear out Absolute settings to defaults
+		Component->UnregisterComponent();
+		Component->SetCastShadow(false);
+
+		//TODO: reset the delegates here once they are working
+		
+		//Ensure a small cull distance doesn't linger to future users.
+		Component->SetCullDistance(FLT_MAX);
+
+		Component->PoolingMethod = ENCPoolMethod::FreeInPool;
+		FreeElements.Push(FNCPoolElement(Component, CurrentTimeSeconds));
+	}
+	else
+	{
+		//We've stopped pooling while some effects were in flight so ensure they're destroyed now.
+		Component->PoolingMethod = ENCPoolMethod::None;//Reset so we don't trigger warnings about destroying pooled NCs.
+		Component->DestroyComponent();
+	}
+}
+
+void FNCPool::KillUnusedComponents(float KillTime, UNiagaraSystem* Template)
+{
+	int32 i = 0;
+	while (i < FreeElements.Num())
+	{
+		if (FreeElements[i].LastUsedTime < KillTime)
+		{
+			UNiagaraComponent* Component = FreeElements[i].Component;
+			if (Component)
+			{
+				Component->PoolingMethod = ENCPoolMethod::None; // Reset so we don't trigger warnings about destroying pooled NCs.
+				Component->DestroyComponent();
+			}
+
+			FreeElements.RemoveAtSwap(i, 1, false);
+		}
+		else
+		{
+			++i;
+		}
+	}
+	FreeElements.Shrink();
+
+#if ENABLE_NC_POOL_DEBUGGING
+	//Clean up any in use components that have been cleared out from under the pool. This could happen in someone manually destroys a component for example.
+	i = 0;
+	while (i < InUseComponents_Manual.Num())
+	{
+		if (!InUseComponents_Manual[i])
+		{
+			UE_LOG(LogNiagara, Log, TEXT("Manual Pooled NC has been destroyed! Possibly via a DestroyComponent() call. You should not destroy these but rather call ReleaseToPool on the component so it can be re-used. |\t System: %s"), *Template->GetFullName());
+			InUseComponents_Manual.RemoveAtSwap(i, 1, false);
+		}
+		else
+		{
+			++i;
+		}
+	}
+	InUseComponents_Manual.Shrink();
+
+	i = 0;
+	while (i < InUseComponents_Auto.Num())
+	{
+		if (!InUseComponents_Auto[i])
+		{
+			UE_LOG(LogNiagara, Log, TEXT("Auto Pooled NC has been destroyed! Possibly via a DestroyComponent() call. You should not destroy these manually. Just deactivate them and allow then to be reclaimed by the pool automatically. |\t System: %s"), *Template->GetFullName());
+			InUseComponents_Auto.RemoveAtSwap(i, 1, false);
+		}
+		else
+		{
+			++i;
+		}
+	}
+	InUseComponents_Auto.Shrink();
+#endif
+}
+
+//////////////////////////////////////////////////////////////////////////
+
+
+UNiagaraComponentPool::UNiagaraComponentPool(const FObjectInitializer& ObjectInitializer)
+	: Super(ObjectInitializer)
+{
+	LastParticleSytemPoolCleanTime = 0.0f;
+}
+
+UNiagaraComponentPool::~UNiagaraComponentPool()
+{
+	Cleanup();
+}
+
+void UNiagaraComponentPool::Cleanup()
+{
+	for (TPair<UNiagaraSystem*, FNCPool>& Pool : WorldParticleSystemPools)
+	{
+		Pool.Value.Cleanup();
+	}
+
+	WorldParticleSystemPools.Empty();
+}
+
+UNiagaraComponent* UNiagaraComponentPool::CreateWorldParticleSystem(UNiagaraSystem* Template, UWorld* World, ENCPoolMethod PoolingMethod)
+{
+	check(IsInGameThread());
+	check(World);
+	if (!Template)
+	{
+		UE_LOG(LogNiagara, Warning, TEXT("Attempted CreateWorldParticleSystem() with a NULL Template!"));
+		return nullptr;
+	}
+
+	if (World->bIsTearingDown)
+	{
+		UE_LOG(LogNiagara, Warning, TEXT("Failed to create pooled particle system as we are tearing the world down."));
+		return nullptr;
+	}
+
+	UNiagaraComponent* Component = nullptr;
+	if (GbEnableParticleSystemPooling != 0)
+	{
+		if (Template->MaxPoolSize > 0)
+		{
+			FNCPool& Pool = WorldParticleSystemPools.FindOrAdd(Template);
+			Component = Pool.Acquire(World, Template, PoolingMethod);
+		}
+	}
+	else
+	{
+		WorldParticleSystemPools.Empty();//Ensure the pools are cleared out if we've just switched to not pooling.
+	}
+
+	if(Component == nullptr)
+	{
+		//Create a new auto destroy system if we're not pooling.
+		Component = NewObject<UNiagaraComponent>(World);
+		Component->SetAutoDestroy(true);
+		Component->bAutoActivate = false;
+		Component->SetAsset(Template);
+	}
+
+	check(Component);
+	return Component;
+}
+
+/** Called when an in-use particle component is finished and wishes to be returned to the pool. */
+void UNiagaraComponentPool::ReclaimWorldParticleSystem(UNiagaraComponent* Component)
+{
+	check(IsInGameThread());
+	
+	//If this component has been already destroyed we don't add it back to the pool. Just warn so users can fix it.
+	if (Component->IsPendingKill())
+	{
+		UE_LOG(LogNiagara, Log, TEXT("Pooled NC has been destroyed! Possibly via a DestroyComponent() call. You should not destroy components set to auto destroy manually. \nJust deactivate them and allow them to destroy themselves or be reclaimed by the pool if pooling is enabled. | NC: %p |\t System: %s"), Component, *Component->GetAsset()->GetFullName());
+		return;
+	}
+
+	if (GbEnableParticleSystemPooling)
+	{
+		float CurrentTime = Component->GetWorld()->GetTimeSeconds();
+
+		//Periodically clear up the pools.
+		if (CurrentTime - LastParticleSytemPoolCleanTime > GParticleSystemPoolingCleanTime)
+		{
+			LastParticleSytemPoolCleanTime = CurrentTime;
+			for (TPair<UNiagaraSystem*, FNCPool>& Pair : WorldParticleSystemPools)
+			{
+				Pair.Value.KillUnusedComponents(CurrentTime - GParticleSystemPoolKillUnusedTime, Component->GetAsset());
+			}
+		}
+		
+		FNCPool* NCPool = WorldParticleSystemPools.Find(Component->GetAsset());
+		if (!NCPool)
+		{
+			UE_LOG(LogNiagara, Warning, TEXT("WorldNC Pool trying to reclaim a system for which it doesn't have a pool! Likely because SetAsset() has been called on this NC. | World: %p | NC: %p | Sys: %s"), Component->GetWorld(), Component, *Component->GetAsset()->GetFullName());
+			//Just add the new pool and reclaim to that one.
+			NCPool = &WorldParticleSystemPools.Add(Component->GetAsset());
+		}
+
+		NCPool->Reclaim(Component, CurrentTime);
+	}
+	else
+	{
+		Component->DestroyComponent();
+	}
+}
+
+void UNiagaraComponentPool::ReclaimActiveParticleSystems()
+{
+	check(IsInGameThread());
+
+	for (TPair<UNiagaraSystem*, FNCPool>& Pair : WorldParticleSystemPools)
+	{
+		FNCPool& Pool = Pair.Value;
+
+		for(int32 i = Pool.InUseComponents_Auto.Num() - 1; i >= 0; --i)
+		{
+			UNiagaraComponent* Component = Pool.InUseComponents_Auto[i];
+			if (ensureAlways(Component))
+			{
+				Component->DeactivateImmediate();
+			}
+		}
+
+		for(int32 i = Pool.InUseComponents_Manual.Num() - 1; i >= 0; --i)
+		{
+			UNiagaraComponent* Component = Pool.InUseComponents_Manual[i];
+			if (ensureAlways(Component))
+			{
+				Component->DeactivateImmediate();
+			}
+		}
+	}
+}
+
+void UNiagaraComponentPool::Dump()
+{
+#if ENABLE_NC_POOL_DEBUGGING
+	FString DumpStr;
+
+	uint32 TotalMemUsage = 0;
+	for (TPair<UNiagaraSystem*, FNCPool>& Pair : WorldParticleSystemPools)
+	{
+		UNiagaraSystem* System = Pair.Key;
+		FNCPool& Pool = Pair.Value;		
+		uint32 FreeMemUsage = 0;
+		for (FNCPoolElement& Elem : Pool.FreeElements)
+		{
+			if (ensureAlways(Elem.Component))
+			{
+				FreeMemUsage += Elem.Component->GetApproxMemoryUsage();
+			}
+		}
+		uint32 InUseMemUsage = 0;
+		for (UNiagaraComponent* Component : Pool.InUseComponents_Auto)
+		{
+			if (ensureAlways(Component))
+			{
+				InUseMemUsage += Component->GetApproxMemoryUsage();				
+			}
+		}
+		for (UNiagaraComponent* Component : Pool.InUseComponents_Manual)
+		{
+			if (ensureAlways(Component))
+			{
+				InUseMemUsage += Component->GetApproxMemoryUsage();
+			}
+		}
+
+		TotalMemUsage += FreeMemUsage;
+		TotalMemUsage += InUseMemUsage;
+
+		DumpStr += FString::Printf(TEXT("Free: %d (%uB) \t|\t Used(Auto - Manual): %d - %d (%uB) \t|\t MaxUsed: %d \t|\t System: %s\n"), Pool.FreeElements.Num(), FreeMemUsage, Pool.InUseComponents_Auto.Num(), Pool.InUseComponents_Manual.Num(), InUseMemUsage, Pool.MaxUsed, *System->GetFullName());
+	}
+
+	UE_LOG(LogNiagara, Log, TEXT("***************************************"));
+	UE_LOG(LogNiagara, Log, TEXT("*Particle System Pool Info - Total Mem = %.2fMB*"), TotalMemUsage / 1024.0f / 1024.0f);
+	UE_LOG(LogNiagara, Log, TEXT("***************************************"));
+	UE_LOG(LogNiagara, Log, TEXT("%s"), *DumpStr);
+	UE_LOG(LogNiagara, Log, TEXT("***************************************"));
+
+#endif
+}
diff --git a/Engine/Plugins/FX/Niagara/Source/Niagara/Private/NiagaraFunctionLibrary.cpp b/Engine/Plugins/FX/Niagara/Source/Niagara/Private/NiagaraFunctionLibrary.cpp
--- a/Engine/Plugins/FX/Niagara/Source/Niagara/Private/NiagaraFunctionLibrary.cpp
+++ b/Engine/Plugins/FX/Niagara/Source/Niagara/Private/NiagaraFunctionLibrary.cpp
@@ -27,15 +27,23 @@ UNiagaraFunctionLibrary::UNiagaraFunctionLibrary(const FObjectInitializer& Objec
 }
 
 
-UNiagaraComponent* CreateNiagaraSystem(UNiagaraSystem* SystemTemplate, UWorld* World, AActor* Actor, bool bAutoDestroy, EPSCPoolMethod PoolingMethod)
+UNiagaraComponent* CreateNiagaraSystem(UNiagaraSystem* SystemTemplate, UWorld* World, AActor* Actor, bool bAutoDestroy, ENCPoolMethod PoolingMethod)
 {
-	// todo : implement pooling method.
+	UNiagaraComponent* NiagaraComponent;
+	if (PoolingMethod == ENCPoolMethod::None)
+	{
+		NiagaraComponent = NewObject<UNiagaraComponent>((Actor ? Actor : (UObject*)World));
+		NiagaraComponent->SetAsset(SystemTemplate);
+		NiagaraComponent->bAutoActivate = false;
+	}
+	else
+	{
+		UNiagaraComponentPool* ComponentPool = FNiagaraWorldManager::Get(World)->GetComponentPool();
+		NiagaraComponent = ComponentPool->CreateWorldParticleSystem(SystemTemplate, World, PoolingMethod);
+	}	
 
-	UNiagaraComponent* NiagaraComponent = NewObject<UNiagaraComponent>((Actor ? Actor : (UObject*)World));
 	NiagaraComponent->SetAutoDestroy(bAutoDestroy);
 	NiagaraComponent->bAllowAnyoneToDestroyMe = true;
-	NiagaraComponent->SetAsset(SystemTemplate);
-	NiagaraComponent->bAutoActivate = false;
 	return NiagaraComponent;
 }
 
@@ -44,7 +52,7 @@ UNiagaraComponent* CreateNiagaraSystem(UNiagaraSystem* SystemTemplate, UWorld* W
 * Spawns a Niagara System at the specified world location/rotation
 * @return			The spawned UNiagaraComponent
 */
-UNiagaraComponent* UNiagaraFunctionLibrary::SpawnSystemAtLocation(const UObject* WorldContextObject, UNiagaraSystem* SystemTemplate, FVector SpawnLocation, FRotator SpawnRotation, FVector Scale, bool bAutoDestroy, bool bAutoActivate)
+UNiagaraComponent* UNiagaraFunctionLibrary::SpawnSystemAtLocation(const UObject* WorldContextObject, UNiagaraSystem* SystemTemplate, FVector SpawnLocation, FRotator SpawnRotation, FVector Scale, bool bAutoDestroy, bool bAutoActivate, ENCPoolMethod PoolingMethod)
 {
 	UNiagaraComponent* PSC = NULL;
 	if (SystemTemplate)
@@ -52,7 +60,7 @@ UNiagaraComponent* UNiagaraFunctionLibrary::SpawnSystemAtLocation(const UObject*
 		UWorld* World = GEngine->GetWorldFromContextObject(WorldContextObject, EGetWorldErrorMode::LogAndReturnNull);
 		if (World != nullptr)
 		{
-			PSC = CreateNiagaraSystem(SystemTemplate, World, World->GetWorldSettings(), bAutoDestroy, EPSCPoolMethod::None);
+			PSC = CreateNiagaraSystem(SystemTemplate, World, World->GetWorldSettings(), bAutoDestroy, PoolingMethod);
 #if WITH_EDITORONLY_DATA
 			PSC->bWaitForCompilationOnActivate = true;
 #endif
@@ -78,7 +86,7 @@ UNiagaraComponent* UNiagaraFunctionLibrary::SpawnSystemAtLocation(const UObject*
 * Spawns a Niagara System attached to a component
 * @return			The spawned UNiagaraComponent
 */
-UNiagaraComponent* UNiagaraFunctionLibrary::SpawnSystemAttached(UNiagaraSystem* SystemTemplate, USceneComponent* AttachToComponent, FName AttachPointName, FVector Location, FRotator Rotation, EAttachLocation::Type LocationType, bool bAutoDestroy, bool bAutoActivate)
+UNiagaraComponent* UNiagaraFunctionLibrary::SpawnSystemAttached(UNiagaraSystem* SystemTemplate, USceneComponent* AttachToComponent, FName AttachPointName, FVector Location, FRotator Rotation, EAttachLocation::Type LocationType, bool bAutoDestroy, bool bAutoActivate, ENCPoolMethod PoolingMethod)
 {
 	UNiagaraComponent* PSC = nullptr;
 	if (SystemTemplate)
@@ -89,7 +97,7 @@ UNiagaraComponent* UNiagaraFunctionLibrary::SpawnSystemAttached(UNiagaraSystem*
 		}
 		else
 		{
-			PSC = CreateNiagaraSystem(SystemTemplate, AttachToComponent->GetWorld(), AttachToComponent->GetOwner(), bAutoDestroy, EPSCPoolMethod::None);
+			PSC = CreateNiagaraSystem(SystemTemplate, AttachToComponent->GetWorld(), AttachToComponent->GetOwner(), bAutoDestroy, PoolingMethod);
 #if WITH_EDITOR
 			if (GForceNiagaraSpawnAttachedSolo > 0)
 			{
@@ -132,7 +140,7 @@ UNiagaraComponent* UNiagaraFunctionLibrary::SpawnSystemAttached(
 	FVector Scale,
 	EAttachLocation::Type LocationType,
 	bool bAutoDestroy,
-	EPSCPoolMethod PoolingMethod,
+	ENCPoolMethod PoolingMethod,
 	bool bAutoActivate
 )
 {
diff --git a/Engine/Plugins/FX/Niagara/Source/Niagara/Private/NiagaraSystem.cpp b/Engine/Plugins/FX/Niagara/Source/Niagara/Private/NiagaraSystem.cpp
--- a/Engine/Plugins/FX/Niagara/Source/Niagara/Private/NiagaraSystem.cpp
+++ b/Engine/Plugins/FX/Niagara/Source/Niagara/Private/NiagaraSystem.cpp
@@ -47,6 +47,7 @@ UNiagaraSystem::UNiagaraSystem(const FObjectInitializer& ObjectInitializer)
 , WarmupTickDelta(1.0f / 15.0f)
 , bHasSystemScriptDIsWithPerInstanceData(false)
 {
+	MaxPoolSize = 32;
 }
 
 void UNiagaraSystem::BeginDestroy()
diff --git a/Engine/Plugins/FX/Niagara/Source/Niagara/Private/NiagaraWorldManager.cpp b/Engine/Plugins/FX/Niagara/Source/Niagara/Private/NiagaraWorldManager.cpp
--- a/Engine/Plugins/FX/Niagara/Source/Niagara/Private/NiagaraWorldManager.cpp
+++ b/Engine/Plugins/FX/Niagara/Source/Niagara/Private/NiagaraWorldManager.cpp
@@ -15,6 +15,7 @@
 #include "EngineModule.h"
 #include "NiagaraStats.h"
 #include "NiagaraEmitterInstanceBatcher.h"
+#include "NiagaraComponentPool.h"
 
 DECLARE_CYCLE_STAT(TEXT("Niagara Manager Tick [GT]"), STAT_NiagaraWorldManTick, STATGROUP_Niagara);
 DECLARE_CYCLE_STAT(TEXT("Niagara Manager Wait On Render [GT]"), STAT_NiagaraWorldManWaitOnRender, STATGROUP_Niagara);
@@ -60,19 +61,20 @@ FNiagaraWorldManager::FNiagaraWorldManager(UWorld* InWorld)
 	: World(InWorld)
 	, CachedEffectsQuality(INDEX_NONE)
 {
+	ComponentPool = NewObject<UNiagaraComponentPool>();
 }
 
 FNiagaraWorldManager* FNiagaraWorldManager::Get(UWorld* World)
 {
-	//INiagaraModule& NiagaraModule = FModuleManager::LoadModuleChecked<INiagaraModule>("Niagara");
 	return INiagaraModule::GetWorldManager(World);
 }
 
 void FNiagaraWorldManager::AddReferencedObjects(FReferenceCollector& Collector)
 {
-	// World doesn't need to be added to the reference list. It will be handle via OnWorldInit & OnWorldCleanup & OnPreWorldFinishDestroy in INiagaraModule
+	// World doesn't need to be added to the reference list. It will be handled via OnWorldInit & OnWorldCleanup & OnPreWorldFinishDestroy in INiagaraModule
 
 	Collector.AddReferencedObjects(ParameterCollections);
+	Collector.AddReferencedObject(ComponentPool);
 }
 
 FString FNiagaraWorldManager::GetReferencerName() const
@@ -200,6 +202,8 @@ void FNiagaraWorldManager::OnBatcherDestroyed(NiagaraEmitterInstanceBatcher* InB
 
 void FNiagaraWorldManager::OnWorldCleanup(bool bSessionEnded, bool bCleanupResources)
 {
+	ComponentPool->Cleanup();
+
 	for (TPair<UNiagaraSystem*, TSharedRef<FNiagaraSystemSimulation, ESPMode::ThreadSafe>>& SimPair : SystemSimulations)
 	{
 		SimPair.Value->Destroy();
diff --git a/Engine/Plugins/FX/Niagara/Source/Niagara/Public/NiagaraComponent.h b/Engine/Plugins/FX/Niagara/Source/Niagara/Public/NiagaraComponent.h
--- a/Engine/Plugins/FX/Niagara/Source/Niagara/Public/NiagaraComponent.h
+++ b/Engine/Plugins/FX/Niagara/Source/Niagara/Public/NiagaraComponent.h
@@ -10,6 +10,7 @@
 #include "Particles/ParticleSystemComponent.h"
 #include "NiagaraUserRedirectionParameterStore.h"
 #include "NiagaraSystemInstance.h"
+#include "NiagaraComponentPool.h"
 
 #include "NiagaraComponent.generated.h"
 
@@ -49,6 +50,8 @@ public:
 	void SetActorParameter(FName ParameterName, class AActor* Param) override;
 	virtual UFXSystemAsset* GetFXSystemAsset() const override { return Asset; };
 	void SetEmitterEnable(FName EmitterName, bool bNewEnableState) override;
+	void ReleaseToPool() override;
+	uint32 GetApproxMemoryUsage() const override;
 	/********* UFXSystemComponent *********/
 
 private:
@@ -122,6 +125,8 @@ public:
 	UPROPERTY(EditAnywhere, BlueprintReadOnly, Category = Attachment)
 	uint32 bAutoManageAttachment : 1;
 
+	/** How to handle pooling for this component instance. */
+	ENCPoolMethod PoolingMethod;
 
 	virtual void Activate(bool bReset = false)override;
 	virtual void Deactivate()override;
@@ -368,8 +373,11 @@ public:
 	UPROPERTY(BlueprintAssignable)
 	FOnNiagaraSystemFinished OnSystemFinished;
 
+	/** Removes all local overrides and replaces them with the values from the source System - note: this also removes the editor overrides from the component as it is used by the pooling mechanism to prevent values leaking between different instances. */
+	void SetUserParametersToDefaultValues();
+
 private:
-	/** Compare local overrides with the source System. Remove any that have mismatched types or no longer exist on the System. Returns whether or not any changes occurred.*/
+	/** Compare local overrides with the source System. Remove any that have mismatched types or no longer exist on the System.*/
 	void SynchronizeWithSourceSystem();
 
 	void AssetExposedParametersChanged();
diff --git a/Engine/Plugins/FX/Niagara/Source/Niagara/Public/NiagaraComponentPool.h b/Engine/Plugins/FX/Niagara/Source/Niagara/Public/NiagaraComponentPool.h
new file mode 100644
--- /dev/null
+++ b/Engine/Plugins/FX/Niagara/Source/Niagara/Public/NiagaraComponentPool.h
@@ -0,0 +1,134 @@
+// Copyright 1998-2019 Epic Games, Inc. All Rights Reserved.
+
+#pragma once
+
+#include "CoreMinimal.h"
+#include "UObject/ObjectMacros.h"
+#include "NiagaraComponentPool.generated.h"
+
+class UNiagaraComponent;
+class UNiagaraSystem;
+
+#define ENABLE_NC_POOL_DEBUGGING (!UE_BUILD_SHIPPING)
+
+UENUM()
+enum class ENCPoolMethod : uint8
+{
+	/**
+	* The component will be created fresh and not allocated from the pool.
+	*/
+	None,
+
+	/**
+	* The component is allocated from the pool and will be automatically released back to it.
+	* User need not handle this any more that other NCs but interaction with the NC after the tick it's spawned in are unsafe.
+	* This method is useful for one-shot fx that you don't need to keep a reference to and can fire and forget.
+	*/
+	AutoRelease,
+
+	/**
+	* NC is allocated from the pool but will NOT be automatically released back to it. User has ownership of the NC and must call ReleaseToPool when finished with it otherwise the NC will leak.
+	* Interaction with the NC after it has been released are unsafe.
+	* This method is useful for persistent FX that you need to modify parameters upon etc over it's lifetime.
+	*/
+	ManualRelease,
+	
+	/** 
+	Special entry allowing manual release NCs to be manually released but wait until completion to be returned to the pool.
+	*/
+	ManualRelease_OnComplete UMETA(Hidden),
+
+	/**
+	Special entry that marks a NC as having been returned to the pool. All NCs currently in the pool are marked this way.
+	*/
+	FreeInPool UMETA(Hidden),
+};
+
+USTRUCT()
+struct FNCPoolElement
+{
+	GENERATED_BODY()
+
+	UPROPERTY(transient)
+	UNiagaraComponent* Component;
+
+	float LastUsedTime;
+
+	FNCPoolElement()
+		: Component(nullptr), LastUsedTime(0.0f)
+	{
+
+	}
+	FNCPoolElement(UNiagaraComponent* InNC, float InLastUsedTime)
+		: Component(InNC), LastUsedTime(InLastUsedTime)
+	{
+
+	}
+};
+
+
+USTRUCT()
+struct FNCPool
+{
+	GENERATED_BODY()
+
+	//Collection of all currently allocated, free items ready to be grabbed for use.
+	//TODO: Change this to a FIFO queue to get better usage. May need to make this whole class behave similar to TCircularQueue.
+	UPROPERTY(transient)
+	TArray<FNCPoolElement> FreeElements;
+
+	//Array of currently in flight components that will auto release.
+	UPROPERTY(transient)
+	TArray<UNiagaraComponent*> InUseComponents_Auto;
+
+	//Array of currently in flight components that need manual release.
+	UPROPERTY(transient)
+	TArray<UNiagaraComponent*> InUseComponents_Manual;
+	
+	/** Keeping track of max in flight systems to help inform any future pre-population we do. */
+	int32 MaxUsed;
+
+public:
+
+	FNCPool();
+	void Cleanup();
+
+	/** Gets a component from the pool ready for use. */
+	UNiagaraComponent* Acquire(UWorld* World, UNiagaraSystem* Template, ENCPoolMethod PoolingMethod);
+
+	/** Returns a component to the pool. */
+	void Reclaim(UNiagaraComponent* NC, const float CurrentTimeSeconds);
+
+	/** Kills any components that have not been used since the passed KillTime. */
+	void KillUnusedComponents(float KillTime, UNiagaraSystem* Template);
+
+	int32 NumComponents() { return FreeElements.Num(); }
+};
+
+UCLASS(Transient)
+class NIAGARA_API UNiagaraComponentPool : public UObject
+{
+	GENERATED_UCLASS_BODY()
+
+private:
+	UPROPERTY()
+	TMap<UNiagaraSystem*, FNCPool> WorldParticleSystemPools;
+
+	float LastParticleSytemPoolCleanTime;
+public:
+
+	~UNiagaraComponentPool();
+
+	void Cleanup();
+
+	UNiagaraComponent* CreateWorldParticleSystem(UNiagaraSystem* Template, UWorld* World, ENCPoolMethod PoolingMethod);
+
+	/** Called when an in-use particle component is finished and wishes to be returned to the pool. */
+	void ReclaimWorldParticleSystem(UNiagaraComponent* Component);
+
+	/** Call if you want to halt & reclaim all active particle systems and return them to their respective pools. */
+	void ReclaimActiveParticleSystems();
+	
+	/** Dumps the current state of the pool to the log. */
+	void Dump();
+};
diff --git a/Engine/Plugins/FX/Niagara/Source/Niagara/Public/NiagaraFunctionLibrary.h b/Engine/Plugins/FX/Niagara/Source/Niagara/Public/NiagaraFunctionLibrary.h
--- a/Engine/Plugins/FX/Niagara/Source/Niagara/Public/NiagaraFunctionLibrary.h
+++ b/Engine/Plugins/FX/Niagara/Source/Niagara/Public/NiagaraFunctionLibrary.h
@@ -6,7 +6,7 @@
 #include "UObject/ObjectMacros.h"
 #include "Engine/EngineTypes.h"
 #include "Kismet/BlueprintFunctionLibrary.h"
-#include "Particles/WorldPSCPool.h"
+#include "NiagaraComponentPool.h"
 #include "NiagaraFunctionLibrary.generated.h"
 
 class UNiagaraComponent;
@@ -27,12 +27,12 @@ class NIAGARA_API UNiagaraFunctionLibrary : public UBlueprintFunctionLibrary
 	* @return			The spawned UNiagaraComponent
 	*/
 	UFUNCTION(BlueprintCallable, Category = Niagara, meta = (Keywords = "niagara System", WorldContext = "WorldContextObject", UnsafeDuringActorConstruction = "true"))
-	static UNiagaraComponent* SpawnSystemAtLocation(const UObject* WorldContextObject, class UNiagaraSystem* SystemTemplate, FVector Location, FRotator Rotation = FRotator::ZeroRotator, FVector Scale = FVector(1.f), bool bAutoDestroy = true, bool bAutoActivate = true);
+	static UNiagaraComponent* SpawnSystemAtLocation(const UObject* WorldContextObject, class UNiagaraSystem* SystemTemplate, FVector Location, FRotator Rotation = FRotator::ZeroRotator, FVector Scale = FVector(1.f), bool bAutoDestroy = true, bool bAutoActivate = true, ENCPoolMethod PoolingMethod = ENCPoolMethod::None);
 
 	UFUNCTION(BlueprintCallable, Category = Niagara, meta = (Keywords = "niagara System", UnsafeDuringActorConstruction = "true"))
-	static UNiagaraComponent* SpawnSystemAttached(UNiagaraSystem* SystemTemplate, USceneComponent* AttachToComponent, FName AttachPointName, FVector Location, FRotator Rotation, EAttachLocation::Type LocationType, bool bAutoDestroy, bool bAutoActivate = true);
+	static UNiagaraComponent* SpawnSystemAttached(UNiagaraSystem* SystemTemplate, USceneComponent* AttachToComponent, FName AttachPointName, FVector Location, FRotator Rotation, EAttachLocation::Type LocationType, bool bAutoDestroy, bool bAutoActivate = true, ENCPoolMethod PoolingMethod = ENCPoolMethod::None);
 
-	static UNiagaraComponent* SpawnSystemAttached(UNiagaraSystem* SystemTemplate, USceneComponent* AttachToComponent, FName AttachPointName, FVector Location, FRotator Rotation, FVector Scale, EAttachLocation::Type LocationType, bool bAutoDestroy, EPSCPoolMethod PoolingMethod, bool bAutoActivate = true);
+	static UNiagaraComponent* SpawnSystemAttached(UNiagaraSystem* SystemTemplate, USceneComponent* AttachToComponent, FName AttachPointName, FVector Location, FRotator Rotation, FVector Scale, EAttachLocation::Type LocationType, bool bAutoDestroy, ENCPoolMethod PoolingMethod, bool bAutoActivate = true);
 
 	/** Sets a Niagara StaticMesh parameter by name, overriding locally if necessary.*/
 	UFUNCTION(BlueprintCallable, Category = Niagara, meta = (DisplayName = "Set Niagara Static Mesh Component"))
diff --git a/Engine/Plugins/FX/Niagara/Source/Niagara/Public/NiagaraWorldManager.h b/Engine/Plugins/FX/Niagara/Source/Niagara/Public/NiagaraWorldManager.h
--- a/Engine/Plugins/FX/Niagara/Source/Niagara/Public/NiagaraWorldManager.h
+++ b/Engine/Plugins/FX/Niagara/Source/Niagara/Public/NiagaraWorldManager.h
@@ -12,11 +12,12 @@
 #include "NiagaraSystemInstance.h"
 #include "GlobalDistanceFieldParameters.h"
 #include "NiagaraDataInterfaceSkeletalMesh.h"
+#include "NiagaraComponentPool.h"
 
 class UWorld;
 class UNiagaraParameterCollection;
 class UNiagaraParameterCollectionInstance;
-
+class UNiagaraComponentPool;
 
 
 class FNiagaraViewDataMgr : public FRenderResource
@@ -90,6 +91,8 @@ public:
 	bool CachedPlayerViewLocationsValid() const { return bCachedPlayerViewLocationsValid; }
 	TArrayView<const FVector> GetCachedPlayerViewLocations() const { check(bCachedPlayerViewLocationsValid); return MakeArrayView(CachedPlayerViewLocations); }
 
+	UNiagaraComponentPool* GetComponentPool() { return ComponentPool; }
+
 private:
 	UWorld* World;
 
@@ -102,6 +105,8 @@ private:
 	bool bCachedPlayerViewLocationsValid = false;
 	TArray<FVector, TInlineAllocator<8> > CachedPlayerViewLocations;
 
+	UNiagaraComponentPool* ComponentPool;
+
 	/** Generated data used by data interfaces*/
 	FNDI_SkeletalMesh_GeneratedData SkeletalMeshGeneratedData;
 
diff --git a/Engine/Source/Runtime/Engine/Classes/Particles/ParticleSystem.h b/Engine/Source/Runtime/Engine/Classes/Particles/ParticleSystem.h
--- a/Engine/Source/Runtime/Engine/Classes/Particles/ParticleSystem.h
+++ b/Engine/Source/Runtime/Engine/Classes/Particles/ParticleSystem.h
@@ -109,6 +109,11 @@ class UFXSystemAsset : public UObject
 	GENERATED_UCLASS_BODY()
 public:
 	UFXSystemAsset() {}
+
+	/** Max number of components of this system to keep resident in the world component pool. */
+	UPROPERTY(EditAnywhere, Category = Performance)
+	uint32 MaxPoolSize;
+	//TODO: Allow pool size overriding per world and possibly implement some preallocation too.
 };
 
 /**
@@ -338,11 +343,6 @@ public:
 	UPROPERTY(EditAnywhere, Category = Performance)
 	EParticleSignificanceLevel MaxSignificanceLevel;
 
-	/** Max number of components of this system to keep resident in the world PSC pool. */
-	UPROPERTY(EditAnywhere, Category = Performance)
-	uint32 MaxPoolSize;
-	//TODO: Allow pool size overriding per world and possibly implement some preallocation too.
-
 	/** Local space position that UVs generated with the ParticleMacroUV material node will be centered on. */
 	UPROPERTY(EditAnywhere, Category=MacroUV)
 	FVector MacroUVPosition;
diff --git a/Engine/Source/Runtime/Engine/Classes/Particles/ParticleSystemComponent.h b/Engine/Source/Runtime/Engine/Classes/Particles/ParticleSystemComponent.h
--- a/Engine/Source/Runtime/Engine/Classes/Particles/ParticleSystemComponent.h
+++ b/Engine/Source/Runtime/Engine/Classes/Particles/ParticleSystemComponent.h
@@ -383,6 +383,17 @@ public:
 	 */
 	UFUNCTION(BlueprintCallable, Category = "Effects|Components|ParticleSystem")
 	virtual void SetAutoAttachmentParameters(USceneComponent* Parent, FName SocketName, EAttachmentRule LocationRule, EAttachmentRule RotationRule, EAttachmentRule ScaleRule) {}
+
+	/**
+	 * Deactivates this system and releases it to the pool on completion.
+	 * Usage of this PSC reference after this call is unsafe.
+	 * You should clear out your references to it.
+	 */
+	UFUNCTION(BlueprintCallable, Category = "Effects|Components|ParticleSystem")
+	virtual void ReleaseToPool() {}
+
+	/** Returns an approximate memory usage value for this component. */
+	virtual uint32 GetApproxMemoryUsage() const { return 0; }
 };
 
 
@@ -1168,6 +1179,8 @@ public:
 	UFUNCTION(BlueprintCallable, Category = "Effects|Particles|Trails")
 	void EndTrails();
 
+	void ReleaseToPool() override;
+
 	/**
 	* Sets the defining data for all trails in this component.
 	*
@@ -1179,14 +1192,6 @@ public:
 	UFUNCTION(BlueprintCallable, Category = "Effects|Particles|Trails")
 	void SetTrailSourceData(FName InFirstSocketName, FName InSecondSocketName, ETrailWidthMode InWidthMode, float InWidth);
 
-	/** 
-	 * Deactivates this system and releases it to the pool on completion.
-	 * Usage of this PSC reference after this call is unsafe. 
-	 * You should clear out your references to it.
-	 */
-	UFUNCTION(BlueprintCallable, Category = "Effects|Components|ParticleSystem")
-	void ReleaseToPool();
-
 public:
 	/** Command fence used to shut down properly */
 	class FRenderCommandFence* ReleaseResourcesFence;
@@ -1372,7 +1377,7 @@ public:
 	virtual FName GetNameForMaterial(UMaterialInterface* InMaterial) const;
 
 	/** Returns an approximate memory usage value for this component. */
-	uint32 GetApproxMemoryUsage()const;
+	uint32 GetApproxMemoryUsage() const override;
 
 protected:
 
