- Build fixes for other platforms from Niagara changes

diff --git a/Engine/Plugins/FX/Niagara/Source/Niagara/Private/NiagaraComponentPool.cpp b/Engine/Plugins/FX/Niagara/Source/Niagara/Private/NiagaraComponentPool.cpp
--- a/Engine/Plugins/FX/Niagara/Source/Niagara/Private/NiagaraComponentPool.cpp
+++ b/Engine/Plugins/FX/Niagara/Source/Niagara/Private/NiagaraComponentPool.cpp
@@ -27,7 +27,7 @@ static FAutoConsoleVariableRef ParticleSystemPoolingCleanTime(
 	TEXT("How often should the pool be cleaned (in seconds).")
 );
 
-void DumpPooledWorldParticleSystemInfo(UWorld* World)
+void DumpPooledWorldNiagaraParticleSystemInfo(UWorld* World)
 {
 	check(World);
 	FNiagaraWorldManager::Get(World)->GetComponentPool()->Dump();
@@ -36,7 +36,7 @@ void DumpPooledWorldParticleSystemInfo(UWorld* World)
 FAutoConsoleCommandWithWorld DumpNCPoolInfoCommand(
 	TEXT("FX.DumpNCPoolInfo"),
 	TEXT("Dump Particle System Pooling Info"),
-	FConsoleCommandWithWorldDelegate::CreateStatic(&DumpPooledWorldParticleSystemInfo)
+	FConsoleCommandWithWorldDelegate::CreateStatic(&DumpPooledWorldNiagaraParticleSystemInfo)
 );
 
 FNCPool::FNCPool()
