Fixed build warnings for Linux and Mac builds

diff --git a/Engine/Plugins/FX/Niagara/Source/NiagaraEditor/Private/NiagaraEditorUtilities.cpp b/Engine/Plugins/FX/Niagara/Source/NiagaraEditor/Private/NiagaraEditorUtilities.cpp
--- a/Engine/Plugins/FX/Niagara/Source/NiagaraEditor/Private/NiagaraEditorUtilities.cpp
+++ b/Engine/Plugins/FX/Niagara/Source/NiagaraEditor/Private/NiagaraEditorUtilities.cpp
@@ -1143,7 +1143,7 @@ void FNiagaraEditorUtilities::KillSystemInstances(const UNiagaraSystem& System)
 	}
 }
 
-bool FNiagaraEditorUtilities::VerifyNameChangeForInputOrOutputNode(UNiagaraNode& NodeBeingChanged, FName OldName, FName NewName, FText& OutErrorMessage)
+bool FNiagaraEditorUtilities::VerifyNameChangeForInputOrOutputNode(const UNiagaraNode& NodeBeingChanged, FName OldName, FName NewName, FText& OutErrorMessage)
 {
 	if (NewName == NAME_None)
 	{
@@ -1172,7 +1172,7 @@ bool FNiagaraEditorUtilities::VerifyNameChangeForInputOrOutputNode(UNiagaraNode&
 
 	if (NodeBeingChanged.IsA<UNiagaraNodeOutput>())
 	{
-		UNiagaraNodeOutput* OutputNodeBeingChanged = CastChecked<UNiagaraNodeOutput>(&NodeBeingChanged);
+		const UNiagaraNodeOutput* OutputNodeBeingChanged = CastChecked<const UNiagaraNodeOutput>(&NodeBeingChanged);
 		for (const FNiagaraVariable& Output : OutputNodeBeingChanged->GetOutputs())
 		{
 			if (Output.GetName() != OldName && Output.GetName() == NewName)
diff --git a/Engine/Plugins/FX/Niagara/Source/NiagaraEditor/Private/NiagaraNodeOutput.cpp b/Engine/Plugins/FX/Niagara/Source/NiagaraEditor/Private/NiagaraNodeOutput.cpp
--- a/Engine/Plugins/FX/Niagara/Source/NiagaraEditor/Private/NiagaraNodeOutput.cpp
+++ b/Engine/Plugins/FX/Niagara/Source/NiagaraEditor/Private/NiagaraNodeOutput.cpp
@@ -59,7 +59,7 @@ FText UNiagaraNodeOutput::GetPinNameText(UEdGraphPin* Pin) const
 	return FText::FromName(Pin->PinName);
 }
 
-bool UNiagaraNodeOutput::VerifyPinNameTextChanged(const FText& InText, FText& OutErrorMessage, UEdGraphPin* Pin)
+bool UNiagaraNodeOutput::VerifyPinNameTextChanged(const FText& InText, FText& OutErrorMessage, UEdGraphPin* Pin) const
 {
 	return FNiagaraEditorUtilities::VerifyNameChangeForInputOrOutputNode(*this, Pin->PinName, *InText.ToString(), OutErrorMessage);
 }
diff --git a/Engine/Plugins/FX/Niagara/Source/NiagaraEditor/Private/ViewModels/NiagaraStandaloneScriptViewModel.cpp b/Engine/Plugins/FX/Niagara/Source/NiagaraEditor/Private/ViewModels/NiagaraStandaloneScriptViewModel.cpp
--- a/Engine/Plugins/FX/Niagara/Source/NiagaraEditor/Private/ViewModels/NiagaraStandaloneScriptViewModel.cpp
+++ b/Engine/Plugins/FX/Niagara/Source/NiagaraEditor/Private/ViewModels/NiagaraStandaloneScriptViewModel.cpp
@@ -50,7 +50,7 @@ void FNiagaraStandaloneScriptViewModel::SendLastCompileMessageJobs(const UNiagar
 			WarningCount++;
 		}
 
-		JobBatchToQueue.Add(MakeShared<FNiagaraMessageJobCompileEvent>(CompileEvent, TWeakObjectPtr<UNiagaraScript>(InScript), TOptional<const FString>(), SourceScript->GetPathName()));
+		JobBatchToQueue.Add(MakeShared<FNiagaraMessageJobCompileEvent>(CompileEvent, TWeakObjectPtr<const UNiagaraScript>(InScript), TOptional<const FString>(), SourceScript->GetPathName()));
 	}
 
 	JobBatchToQueue.Insert(MakeShared<FNiagaraMessageJobPostCompileSummary>(ErrorCount, WarningCount, GetLatestCompileStatus(), FText::FromString("Script")), 0);
diff --git a/Engine/Plugins/FX/Niagara/Source/NiagaraEditor/Public/NiagaraEditorUtilities.h b/Engine/Plugins/FX/Niagara/Source/NiagaraEditor/Public/NiagaraEditorUtilities.h
--- a/Engine/Plugins/FX/Niagara/Source/NiagaraEditor/Public/NiagaraEditorUtilities.h
+++ b/Engine/Plugins/FX/Niagara/Source/NiagaraEditor/Public/NiagaraEditorUtilities.h
@@ -198,5 +198,5 @@ namespace FNiagaraEditorUtilities
 	void KillSystemInstances(const UNiagaraSystem& System);
 
 
-	bool VerifyNameChangeForInputOrOutputNode(UNiagaraNode& NodeBeingChanged, FName OldName, FName NewName, FText& OutErrorMessage);
+	bool VerifyNameChangeForInputOrOutputNode(const UNiagaraNode& NodeBeingChanged, FName OldName, FName NewName, FText& OutErrorMessage);
 };
diff --git a/Engine/Plugins/FX/Niagara/Source/NiagaraEditor/Public/NiagaraNodeOutput.h b/Engine/Plugins/FX/Niagara/Source/NiagaraEditor/Public/NiagaraNodeOutput.h
--- a/Engine/Plugins/FX/Niagara/Source/NiagaraEditor/Public/NiagaraNodeOutput.h
+++ b/Engine/Plugins/FX/Niagara/Source/NiagaraEditor/Public/NiagaraNodeOutput.h
@@ -69,7 +69,7 @@ private:
 	/** Gets the display text for a pin. */
 	FText GetPinNameText(UEdGraphPin* Pin) const;
 
-	bool VerifyPinNameTextChanged(const FText& InText, FText& OutErrorMessage, UEdGraphPin* Pin);
+	bool VerifyPinNameTextChanged(const FText& InText, FText& OutErrorMessage, UEdGraphPin* Pin) const;
 
 	/** Called when a pin's name text is committed. */
 	void PinNameTextCommitted(const FText& Text, ETextCommit::Type CommitType, UEdGraphPin* Pin);
