Merging various CIS fixes from //UE4/Dev-Main to Dev-Framework (//UE4/Dev-Framework)
#rnx

diff --git a/Engine/Plugins/Experimental/HairStrands/Source/HairStrandsEditor/Private/HairStrandsFactory.cpp b/Engine/Plugins/Experimental/HairStrands/Source/HairStrandsEditor/Private/HairStrandsFactory.cpp
--- a/Engine/Plugins/Experimental/HairStrands/Source/HairStrandsEditor/Private/HairStrandsFactory.cpp
+++ b/Engine/Plugins/Experimental/HairStrands/Source/HairStrandsEditor/Private/HairStrandsFactory.cpp
@@ -57,6 +57,7 @@ UObject* UHairStrandsFactory::FactoryCreateNew(UClass* InClass, UObject* InParen
 		ExistingAsset->ReleaseResource();
 	}
 	UHairStrandsAsset* CurrentAsset = NewObject<UHairStrandsAsset>(InParent, InClass, InName, Flags);
+	check(CurrentAsset);
 
 	const FString Filter(TEXT("Hair Strands Files (*.hair,*.fbx,*.abc)|*.hair;*.fbx;*.abc"));
 
@@ -93,6 +94,7 @@ UObject* UHairStrandsFactory::FactoryCreateFile(UClass * InClass, UObject * InPa
 	}
 
 	UHairStrandsAsset* CurrentAsset = NewObject<UHairStrandsAsset>(InParent, InClass, InName, Flags);
+	check(CurrentAsset);
 
 	CurrentAsset->FilePath = Filename;
 	BuildHairStrands(CurrentAsset->FilePath, CurrentAsset->StrandsDatas);
diff --git a/Engine/Plugins/FX/Niagara/Source/Niagara/Private/NiagaraDataInterfaceGrid2DCollection.cpp b/Engine/Plugins/FX/Niagara/Source/Niagara/Private/NiagaraDataInterfaceGrid2DCollection.cpp
--- a/Engine/Plugins/FX/Niagara/Source/Niagara/Private/NiagaraDataInterfaceGrid2DCollection.cpp
+++ b/Engine/Plugins/FX/Niagara/Source/Niagara/Private/NiagaraDataInterfaceGrid2DCollection.cpp
@@ -58,7 +58,7 @@ struct FNiagaraDataInterfaceParametersCS_Grid2DCollection : public FNiagaraDataI
 		FNiagaraDataInterfaceProxyGrid2DCollection* VFDI = static_cast<FNiagaraDataInterfaceProxyGrid2DCollection*>(Context.DataInterface);
 		
 		Grid2DCollectionRWInstanceData* ProxyData = VFDI->SystemInstancesToProxyData.Find(Context.SystemInstance);
-		ensure(ProxyData);
+		check(ProxyData);
 
 		int NumCellsTmp[2];
 		NumCellsTmp[0] = ProxyData->NumCellsX;
diff --git a/Engine/Plugins/FX/Niagara/Source/Niagara/Private/NiagaraDataInterfaceNeighborGrid3D.cpp b/Engine/Plugins/FX/Niagara/Source/Niagara/Private/NiagaraDataInterfaceNeighborGrid3D.cpp
--- a/Engine/Plugins/FX/Niagara/Source/Niagara/Private/NiagaraDataInterfaceNeighborGrid3D.cpp
+++ b/Engine/Plugins/FX/Niagara/Source/Niagara/Private/NiagaraDataInterfaceNeighborGrid3D.cpp
@@ -67,7 +67,7 @@ struct FNiagaraDataInterfaceParametersCS_NeighborGrid3D : public FNiagaraDataInt
 		FNiagaraDataInterfaceProxyNeighborGrid3D* VFDI = static_cast<FNiagaraDataInterfaceProxyNeighborGrid3D*>(Context.DataInterface);
 
 		NeighborGrid3DRWInstanceData* ProxyData = VFDI->SystemInstancesToProxyData.Find(Context.SystemInstance);
-		ensure(ProxyData);
+		check(ProxyData);
 
 		SetShaderValue(RHICmdList, ComputeShaderRHI, NumVoxelsParam, ProxyData->NumVoxels);
 
diff --git a/Engine/Plugins/FX/Niagara/Source/Niagara/Private/NiagaraDataInterfaceRW.cpp b/Engine/Plugins/FX/Niagara/Source/Niagara/Private/NiagaraDataInterfaceRW.cpp
--- a/Engine/Plugins/FX/Niagara/Source/Niagara/Private/NiagaraDataInterfaceRW.cpp
+++ b/Engine/Plugins/FX/Niagara/Source/Niagara/Private/NiagaraDataInterfaceRW.cpp
@@ -146,7 +146,6 @@ void UNiagaraDataInterfaceGrid3D::GetVMExternalFunction(const FVMExternalFunctio
 	else if (BindingInfo.Name == UnitToIndexFunctionName) { OutFunc = FVMExternalFunction::CreateUObject(this, &UNiagaraDataInterfaceRWBase::EmptyVMFunction); }
 	else if (BindingInfo.Name == IndexToLinearFunctionName) { OutFunc = FVMExternalFunction::CreateUObject(this, &UNiagaraDataInterfaceRWBase::EmptyVMFunction); }
 	else if (BindingInfo.Name == VoxelSizeFunctionName) { OutFunc = FVMExternalFunction::CreateUObject(this, &UNiagaraDataInterfaceRWBase::EmptyVMFunction); }
-	else if (BindingInfo.Name == NumVoxelsFunctionName) { OutFunc = FVMExternalFunction::CreateUObject(this, &UNiagaraDataInterfaceRWBase::EmptyVMFunction); }
 }
 
 bool UNiagaraDataInterfaceGrid3D::Equals(const UNiagaraDataInterface* Other) const
@@ -248,23 +247,6 @@ bool UNiagaraDataInterfaceGrid3D::GetFunctionHLSL(const FName& DefinitionFunctio
 		};
 		OutHLSL += FString::Format(FormatSample, ArgsSample);
 		return true;
-	}	
-	else if (DefinitionFunctionName == NumVoxelsFunctionName)
-	{
-		static const TCHAR *FormatSample = TEXT(R"(
-			void {FunctionName}(out int Out_NumVoxelsX, out int Out_NumVoxelsY, out int Out_NumVoxelsZ)
-			{
-				Out_NumVoxelsX = {NumVoxelsName}.x;
-				Out_NumVoxelsY = {NumVoxelsName}.y;
-				Out_NumVoxelsZ = {NumVoxelsName}.z;
-			}
-		)");
-		TMap<FString, FStringFormatArg> ArgsSample = {
-			{TEXT("FunctionName"), InstanceFunctionName},
-			{TEXT("NumVoxelsName"), NumVoxelsName + ParamInfo.DataInterfaceHLSLSymbol},
-		};
-		OutHLSL += FString::Format(FormatSample, ArgsSample);
-		return true;
 	}
 	else if (DefinitionFunctionName == VoxelSizeFunctionName)
 	{
@@ -474,8 +456,7 @@ void UNiagaraDataInterfaceGrid2D::GetVMExternalFunction(const FVMExternalFunctio
 	else if (BindingInfo.Name == IndexToUnitStaggeredYFunctionName) { OutFunc = FVMExternalFunction::CreateUObject(this, &UNiagaraDataInterfaceRWBase::EmptyVMFunction); }
 	else if (BindingInfo.Name == IndexToLinearFunctionName) { OutFunc = FVMExternalFunction::CreateUObject(this, &UNiagaraDataInterfaceRWBase::EmptyVMFunction); }
 	else if (BindingInfo.Name == LinearToIndexFunctionName) { OutFunc = FVMExternalFunction::CreateUObject(this, &UNiagaraDataInterfaceRWBase::EmptyVMFunction); }
-	else if (BindingInfo.Name == CellSizeFunctionName) { OutFunc = FVMExternalFunction::CreateUObject(this, &UNiagaraDataInterfaceRWBase::EmptyVMFunction); }
-	else if (BindingInfo.Name == NumCellsFunctionName) { OutFunc = FVMExternalFunction::CreateUObject(this, &UNiagaraDataInterfaceRWBase::EmptyVMFunction); }
+	else if (BindingInfo.Name == CellSizeFunctionName) { OutFunc = FVMExternalFunction::CreateUObject(this, &UNiagaraDataInterfaceRWBase::EmptyVMFunction); }	
 }
 
 bool UNiagaraDataInterfaceGrid2D::Equals(const UNiagaraDataInterface* Other) const
@@ -639,22 +620,6 @@ bool UNiagaraDataInterfaceGrid2D::GetFunctionHLSL(const FName& DefinitionFunctio
 		OutHLSL += FString::Format(FormatSample, ArgsSample);
 		return true;
 	}
-	else if (DefinitionFunctionName == NumCellsFunctionName)
-	{
-		static const TCHAR *FormatSample = TEXT(R"(
-			void {FunctionName}(out int Out_NumCellsX, out int Out_NumCellsY)
-			{
-				Out_NumCellsX = {NumCells}.x;
-				Out_NumCellsY = {NumCells}.y;				
-			}
-		)");
-		TMap<FString, FStringFormatArg> ArgsSample = {
-			{TEXT("FunctionName"), InstanceFunctionName},
-			{TEXT("NumCells"), NumCellsName + ParamInfo.DataInterfaceHLSLSymbol},
-		};
-		OutHLSL += FString::Format(FormatSample, ArgsSample);
-		return true;
-	}
 	else if (DefinitionFunctionName == CellSizeFunctionName)
 	{
 		static const TCHAR *FormatSample = TEXT(R"(
diff --git a/Engine/Plugins/FX/Niagara/Source/NiagaraEditor/Private/NiagaraHlslTranslator.cpp b/Engine/Plugins/FX/Niagara/Source/NiagaraEditor/Private/NiagaraHlslTranslator.cpp
--- a/Engine/Plugins/FX/Niagara/Source/NiagaraEditor/Private/NiagaraHlslTranslator.cpp
+++ b/Engine/Plugins/FX/Niagara/Source/NiagaraEditor/Private/NiagaraHlslTranslator.cpp
@@ -4948,7 +4948,7 @@ void FHlslNiagaraTranslator::RegisterFunctionCall(ENiagaraScriptUsage ScriptUsag
 				{
 					if (UseShaderStages)
 					{
-						if (bStageMaxFilter && bStageMaxFilter)
+						if (bStageMinFilter && bStageMaxFilter)
 						{
 							FunctionDefStr += TEXT("if ((GCurrentPhase == 1 && ShaderStageIndex >= ") + MinParam + TEXT(" && ShaderStageIndex <= ") + MaxParam + TEXT(") || ") +
 								TEXT("(GCurrentPhase == 0 && ShaderStageIndex >= ") + MinParamSpawn + TEXT(" && ShaderStageIndex <= ") + MaxParamSpawn + TEXT(")") +
diff --git a/Engine/Plugins/Runtime/NetworkPrediction/Source/NetworkPrediction/Private/Movement/ParametricMovement.cpp b/Engine/Plugins/Runtime/NetworkPrediction/Source/NetworkPrediction/Private/Movement/ParametricMovement.cpp
--- a/Engine/Plugins/Runtime/NetworkPrediction/Source/NetworkPrediction/Private/Movement/ParametricMovement.cpp
+++ b/Engine/Plugins/Runtime/NetworkPrediction/Source/NetworkPrediction/Private/Movement/ParametricMovement.cpp
@@ -112,7 +112,7 @@ FNetworkSimulationModelInitParameters UParametricMovementComponent::GetSimulatio
 {
 	// These are reasonable defaults but may not be right for everyone
 	FNetworkSimulationModelInitParameters InitParams;
-	InitParams.InputBufferSize = Role != ROLE_SimulatedProxy ? 32 : 32;  // Fixme.. not good
+	InitParams.InputBufferSize = 32; //Role != ROLE_SimulatedProxy ? 32 : 32;  // Fixme.. not good
 	InitParams.SyncedBufferSize = Role != ROLE_AutonomousProxy ? 2 : 32;
 	InitParams.AuxBufferSize = Role != ROLE_AutonomousProxy ? 2 : 32;
 	InitParams.DebugBufferSize = 32;
diff --git a/Engine/Plugins/Runtime/NetworkPrediction/Source/NetworkPrediction/Public/Movement/ParametricMovement.h b/Engine/Plugins/Runtime/NetworkPrediction/Source/NetworkPrediction/Public/Movement/ParametricMovement.h
--- a/Engine/Plugins/Runtime/NetworkPrediction/Source/NetworkPrediction/Public/Movement/ParametricMovement.h
+++ b/Engine/Plugins/Runtime/NetworkPrediction/Source/NetworkPrediction/Public/Movement/ParametricMovement.h
@@ -165,7 +165,13 @@ protected:
 
 	TUniquePtr<ParametricMovement::FMovementSystem> NetworkSim;
 
-	// Temp....
+	// ------------------------------------------------------------------------
+	// Temp Parametric movement example
+	//	The essence of this movement simulation is to map some Time value to a transform. That is it.
+	//	(It could be mapped via a spline, a curve, a simple blueprint function, etc).
+	//	What is below is just a simple C++ implementation to stand things up. Most likely we would 
+	//	do additional subclasses to vary the way this is implemented)
+	// ------------------------------------------------------------------------
 	UPROPERTY(EditAnywhere, BlueprintReadOnly, Category=ParametricMovement)
 	FVector ParametricDelta = FVector(0.f, 0.f, 500.f);
 
diff --git a/Engine/Plugins/Runtime/NetworkPrediction/Source/NetworkPrediction/Public/NetworkSimulationModelTemplates.h b/Engine/Plugins/Runtime/NetworkPrediction/Source/NetworkPrediction/Public/NetworkSimulationModelTemplates.h
--- a/Engine/Plugins/Runtime/NetworkPrediction/Source/NetworkPrediction/Public/NetworkSimulationModelTemplates.h
+++ b/Engine/Plugins/Runtime/NetworkPrediction/Source/NetworkPrediction/Public/NetworkSimulationModelTemplates.h
@@ -470,13 +470,8 @@ struct TReplicator_BasicReconciliar
 
 		TSyncState* LatestClient = SyncBuffer.GetElementFromHead(0);
 
-		
+		//Temp		
 		UE_LOG(LogNetworkSim, Warning, TEXT("Client misprediction. Frame: %d"), ReconciliationKeyframe);
-		/*
-		UE_LOG(LogNetworkSim, Warning, TEXT("  Client: %s"), *ClientMotionState->Location.ToString());
-		UE_LOG(LogNetworkSim, Warning, TEXT("  Server: %s"), *ServerState->Location.ToString());
-		UE_LOG(LogNetworkSim, Warning, TEXT("  LatestClient: %s"), *LatestClient->Location.ToString());
-		*/
 
 		// Copy authoritative state over client state (FIXME: we may want to store this off historically somewhere? Or will VLog be enough for debugging?)
 		*ClientSyncState = *ServerState;
@@ -685,7 +680,7 @@ public:
 				{
 					if (LastProcessedInputKeyframe != 0)
 					{
-						// This shouldn't happen, but is not fatal. We are reseting the motion state buffer.
+						// This shouldn't happen, but is not fatal. We are reseting the sync state buffer.
 						UE_LOG(LogNetworkSim, Warning, TEXT("Break in SyncState continuity. LastProcessedInputKeyframe: %d. SyncBuffer.GetHeadKeyframe(): %d. Role=%d"), LastProcessedInputKeyframe, SyncBuffer.GetHeadKeyframe(), (int32)Parameters.Role);
 					}
 
diff --git a/Engine/Source/Runtime/Experimental/Chaos/Public/Chaos/ImplicitObjectIntersection.h b/Engine/Source/Runtime/Experimental/Chaos/Public/Chaos/ImplicitObjectIntersection.h
--- a/Engine/Source/Runtime/Experimental/Chaos/Public/Chaos/ImplicitObjectIntersection.h
+++ b/Engine/Source/Runtime/Experimental/Chaos/Public/Chaos/ImplicitObjectIntersection.h
@@ -64,7 +64,7 @@ class TImplicitObjectIntersection : public TImplicitObject<T, d>
 
 		for(const TUniquePtr<TImplicitObject<T, d>>& Ptr : MObjects)
 		{
-			HashCombine(Ptr->GetTypeHash(), OutHash);
+			OutHash = HashCombine(Ptr->GetTypeHash(), OutHash);
 		}
 
 		return OutHash;
