Niagara - Fix system editor issues.
+ Fix an undo stability issues.
+ Fix copy/paste for nodes in the system overview graph.
+ Fix an issue where systems could load with an empty system overview graph.
+ Fix an issue where undoing a delete emitter operation would break the overview node ui for that emitter.
+ Fix an issue where comment nodes can not be delete from the system overview.

diff --git a/Engine/Plugins/FX/Niagara/Source/NiagaraEditor/Private/NiagaraGraph.cpp b/Engine/Plugins/FX/Niagara/Source/NiagaraEditor/Private/NiagaraGraph.cpp
--- a/Engine/Plugins/FX/Niagara/Source/NiagaraEditor/Private/NiagaraGraph.cpp
+++ b/Engine/Plugins/FX/Niagara/Source/NiagaraEditor/Private/NiagaraGraph.cpp
@@ -768,7 +768,7 @@ void UNiagaraGraph::AddParameter(const FNiagaraVariable& Parameter)
 	UNiagaraScriptVariable** FoundScriptVariable = VariableToScriptVariable.Find(Parameter);
 	if (!FoundScriptVariable)
 	{
-		UNiagaraScriptVariable* NewScriptVariable = NewObject<UNiagaraScriptVariable>(this);
+		UNiagaraScriptVariable* NewScriptVariable = NewObject<UNiagaraScriptVariable>(this, Parameter.GetName(), RF_Transactional);
 		NewScriptVariable->Variable = Parameter;
 		VariableToScriptVariable.Add(Parameter, NewScriptVariable);
 	}
@@ -1444,7 +1444,7 @@ void UNiagaraGraph::SetMetaData(const FNiagaraVariable& InVar, const FNiagaraVar
 	}
 	else 
 	{
-		UNiagaraScriptVariable*& NewScriptVariable = VariableToScriptVariable.Add(InVar, NewObject<UNiagaraScriptVariable>(this));
+		UNiagaraScriptVariable*& NewScriptVariable = VariableToScriptVariable.Add(InVar, NewObject<UNiagaraScriptVariable>(this, InVar.GetName(), RF_Transactional));
 		NewScriptVariable->Variable = InVar;
 		NewScriptVariable->Metadata = InMetaData;
 	}
@@ -1530,7 +1530,7 @@ void UNiagaraGraph::RefreshParameterReferences() const
 				{
 					// Warning! This method (RefreshParameterReferences) isn't const at all!
 					// Need to cast away the const to be able to create a serializable UObject here...
-					UNiagaraScriptVariable* NewScriptVariable = NewObject<UNiagaraScriptVariable>(const_cast<UNiagaraGraph*>(this));
+					UNiagaraScriptVariable* NewScriptVariable = NewObject<UNiagaraScriptVariable>(const_cast<UNiagaraGraph*>(this), Parameter.GetName(), RF_Transactional);
 					NewScriptVariable->Variable = Parameter;
 					VariableToScriptVariable.Add(Parameter, NewScriptVariable);
 				}
@@ -1556,7 +1556,7 @@ void UNiagaraGraph::RefreshParameterReferences() const
 		UNiagaraScriptVariable** FoundScriptVariable = VariableToScriptVariable.Find(Variable);
 		if (!FoundScriptVariable)
 		{
-			UNiagaraScriptVariable* NewScriptVariable = NewObject<UNiagaraScriptVariable>(const_cast<UNiagaraGraph*>(this));
+			UNiagaraScriptVariable* NewScriptVariable = NewObject<UNiagaraScriptVariable>(const_cast<UNiagaraGraph*>(this), Variable.GetName(), RF_Transactional);
 			NewScriptVariable->Variable = Variable;
 			NewScriptVariable->Metadata.bIsStaticSwitch = true;
 			VariableToScriptVariable.Add(Variable, NewScriptVariable);
diff --git a/Engine/Plugins/FX/Niagara/Source/NiagaraEditor/Private/NiagaraOverviewNode.cpp b/Engine/Plugins/FX/Niagara/Source/NiagaraEditor/Private/NiagaraOverviewNode.cpp
--- a/Engine/Plugins/FX/Niagara/Source/NiagaraEditor/Private/NiagaraOverviewNode.cpp
+++ b/Engine/Plugins/FX/Niagara/Source/NiagaraEditor/Private/NiagaraOverviewNode.cpp
@@ -54,7 +54,8 @@ bool UNiagaraOverviewNode::CanUserDeleteNode() const
 
 bool UNiagaraOverviewNode::CanDuplicateNode() const
 {
-	return EmitterHandleGuid.IsValid();
+	// The class object must return true for can duplicate otherwise the CanImportNodesFromText utility function fails.
+	return HasAnyFlags(RF_ClassDefaultObject) || EmitterHandleGuid.IsValid();
 }
 
 UNiagaraSystem* UNiagaraOverviewNode::GetOwningSystem()
diff --git a/Engine/Plugins/FX/Niagara/Source/NiagaraEditor/Private/NiagaraScriptVariable.cpp b/Engine/Plugins/FX/Niagara/Source/NiagaraEditor/Private/NiagaraScriptVariable.cpp
--- a/Engine/Plugins/FX/Niagara/Source/NiagaraEditor/Private/NiagaraScriptVariable.cpp
+++ b/Engine/Plugins/FX/Niagara/Source/NiagaraEditor/Private/NiagaraScriptVariable.cpp
@@ -7,3 +7,13 @@ UNiagaraScriptVariable::UNiagaraScriptVariable(const FObjectInitializer& ObjectI
 {
 	
 }
+
+void UNiagaraScriptVariable::PostLoad()
+{
+	Super::PostLoad();
+
+	if (GIsEditor)
+	{
+		SetFlags(RF_Transactional);
+	}
+}
diff --git a/Engine/Plugins/FX/Niagara/Source/NiagaraEditor/Private/NiagaraScriptVariable.h b/Engine/Plugins/FX/Niagara/Source/NiagaraEditor/Private/NiagaraScriptVariable.h
--- a/Engine/Plugins/FX/Niagara/Source/NiagaraEditor/Private/NiagaraScriptVariable.h
+++ b/Engine/Plugins/FX/Niagara/Source/NiagaraEditor/Private/NiagaraScriptVariable.h
@@ -11,6 +11,8 @@ class UNiagaraScriptVariable : public UObject
 public:
 	GENERATED_UCLASS_BODY()
 
+	virtual void PostLoad() override;
+
 	UPROPERTY()
 	FNiagaraVariable Variable;
 
diff --git a/Engine/Plugins/FX/Niagara/Source/NiagaraEditor/Private/NiagaraSystemEditorData.cpp b/Engine/Plugins/FX/Niagara/Source/NiagaraEditor/Private/NiagaraSystemEditorData.cpp
--- a/Engine/Plugins/FX/Niagara/Source/NiagaraEditor/Private/NiagaraSystemEditorData.cpp
+++ b/Engine/Plugins/FX/Niagara/Source/NiagaraEditor/Private/NiagaraSystemEditorData.cpp
@@ -101,6 +101,10 @@ void UNiagaraSystemEditorData::PostLoadFromOwner(UObject* InOwner)
 	{
 		SystemOverviewGraph = NewObject<UEdGraph>(this, "NiagaraOverview", RF_Transactional);
 		SystemOverviewGraph->Schema = UEdGraphSchema_NiagaraSystemOverview::StaticClass();
+	}
+
+	if(SystemOverviewGraph->Nodes.Num() == 0)
+	{
 		SynchronizeOverviewGraphWithSystem(*OwningSystem);
 	}
 }
@@ -280,6 +284,7 @@ void UNiagaraSystemEditorData::SynchronizeOverviewGraphWithSystem(UNiagaraSystem
 	// If there are any nodes remaining in the list they're no longer being used so destroy them.
 	for (UNiagaraOverviewNode* OverviewNode : OverviewNodes)
 	{
+		OverviewNode->Modify();
 		OverviewNode->DestroyNode();
 	}
 }
diff --git a/Engine/Plugins/FX/Niagara/Source/NiagaraEditor/Private/ViewModels/NiagaraEmitterViewModel.cpp b/Engine/Plugins/FX/Niagara/Source/NiagaraEditor/Private/ViewModels/NiagaraEmitterViewModel.cpp
--- a/Engine/Plugins/FX/Niagara/Source/NiagaraEditor/Private/ViewModels/NiagaraEmitterViewModel.cpp
+++ b/Engine/Plugins/FX/Niagara/Source/NiagaraEditor/Private/ViewModels/NiagaraEmitterViewModel.cpp
@@ -457,9 +457,13 @@ void FNiagaraEmitterViewModel::ScriptParameterStoreChanged(const FNiagaraParamet
 
 void FNiagaraEmitterViewModel::OnEmitterPropertiesChanged()
 {
-	// When the properties change we reset the scripts on the script view model because gpu/cpu or interpolation may have changed.
-	SharedScriptViewModel->SetScripts(Emitter.Get());
-	OnPropertyChangedDelegate.Broadcast();
+	// Check that these are valid since post edit changed is called on objects even when they've been deleted as a result of undo/redo.
+	if (SharedScriptViewModel.IsValid() && Emitter.IsValid())
+	{
+		// When the properties change we reset the scripts on the script view model because gpu/cpu or interpolation may have changed.
+		SharedScriptViewModel->SetScripts(Emitter.Get());
+		OnPropertyChangedDelegate.Broadcast();
+	}
 }
 
 #undef LOCTEXT_NAMESPACE
diff --git a/Engine/Plugins/FX/Niagara/Source/NiagaraEditor/Private/ViewModels/NiagaraOverviewGraphViewModel.cpp b/Engine/Plugins/FX/Niagara/Source/NiagaraEditor/Private/ViewModels/NiagaraOverviewGraphViewModel.cpp
--- a/Engine/Plugins/FX/Niagara/Source/NiagaraEditor/Private/ViewModels/NiagaraOverviewGraphViewModel.cpp
+++ b/Engine/Plugins/FX/Niagara/Source/NiagaraEditor/Private/ViewModels/NiagaraOverviewGraphViewModel.cpp
@@ -28,6 +28,7 @@ FNiagaraOverviewGraphViewModel::FNiagaraOverviewGraphViewModel()
 
 FNiagaraOverviewGraphViewModel::~FNiagaraOverviewGraphViewModel()
 {
+	GEditor->UnregisterForUndo(this);
 }
 
 void FNiagaraOverviewGraphViewModel::Initialize(TSharedRef<FNiagaraSystemViewModel> InSystemViewModel)
@@ -36,6 +37,8 @@ void FNiagaraOverviewGraphViewModel::Initialize(TSharedRef<FNiagaraSystemViewMod
 	OverviewGraph = InSystemViewModel->GetEditorData().GetSystemOverviewGraph();
 
 	SetupCommands();
+	GEditor->RegisterForUndo(this);
+
 	NodeSelection->OnSelectedObjectsChanged().AddSP(this, &FNiagaraOverviewGraphViewModel::GraphSelectionChanged);
 	InSystemViewModel->GetSelectionViewModel()->OnSelectionChanged().AddSP(this, &FNiagaraOverviewGraphViewModel::SystemSelectionChanged);
 }
@@ -78,11 +81,6 @@ void FNiagaraOverviewGraphViewModel::SetViewSettings(const FNiagaraGraphViewSett
 	GetSystemViewModel()->GetEditorData().SetSystemOverviewGraphViewSettings(InOverviewGraphViewSettings);
 }
 
-FNiagaraOverviewGraphViewModel::FOnNodesPasted& FNiagaraOverviewGraphViewModel::OnNodesPasted()
-{
-	return OnNodesPastedDelegate;
-}
-
 TSharedRef<FNiagaraSystemViewModel> FNiagaraOverviewGraphViewModel::GetSystemViewModel()
 {
 	TSharedPtr<FNiagaraSystemViewModel> SystemViewModelPinned = SystemViewModel.Pin();
@@ -150,16 +148,27 @@ void FNiagaraOverviewGraphViewModel::DeleteSelectedNodes()
 		const FScopedTransaction Transaction(FGenericCommands::Get().Delete->GetDescription());
 		Graph->Modify();
 
-		TArray<UObject*> NodesToDelete = NodeSelection->GetSelectedObjects().Array();
+		TArray<UObject*> ObjectsToDelete = NodeSelection->GetSelectedObjects().Array();
 		NodeSelection->ClearSelectedObjects();
 
 		TSet<FGuid> EmitterGuidsToDelete;
-		for (UObject* NodeToDelete : NodesToDelete)
+		for (UObject* ObjectToDelete : ObjectsToDelete)
 		{
-			UNiagaraOverviewNode* GraphNodeToDelete = Cast<UNiagaraOverviewNode>(NodeToDelete);
-			if (GraphNodeToDelete != nullptr && GraphNodeToDelete->CanUserDeleteNode())
+			UEdGraphNode* NodeToDelete = Cast<UEdGraphNode>(ObjectToDelete);
+			if (NodeToDelete != nullptr && NodeToDelete->CanUserDeleteNode())
 			{
-				EmitterGuidsToDelete.Add(GraphNodeToDelete->GetEmitterHandleGuid());
+				NodeToDelete->Modify();
+				UNiagaraOverviewNode* OverviewNodeToDelete = Cast<UNiagaraOverviewNode>(NodeToDelete);
+				if (OverviewNodeToDelete != nullptr)
+				{
+					// For emitter nodes, just collect their ids and let the system view model delete them.  It will update the overview graph.
+					EmitterGuidsToDelete.Add(OverviewNodeToDelete->GetEmitterHandleGuid());
+				}
+				else
+				{
+					// Other types of nodes can be deleted directly.
+					NodeToDelete->DestroyNode();
+				}
 			}
 		}
 		//we have checked SystemViewModel is valid as this is a requisite for Graph to not be null.
@@ -274,24 +283,27 @@ void FNiagaraOverviewGraphViewModel::PasteNodes()
 		TSet<UEdGraphNode*> PastedNodes;
 		FEdGraphUtilities::ImportNodesFromText(Graph, TextToImport, PastedNodes);
 
-// 		for (UEdGraphNode* PastedNode : PastedNodes) //@TODO System Overview: fix this, copy the impl from sequencer impl in system toolkit
-// 		{
-// 			PastedNode->CreateNewGuid();
-// 			UNiagaraNode* Node = Cast<UNiagaraNode>(PastedNode);
-// 			if (Node)
-// 				Node->MarkNodeRequiresSynchronization(__FUNCTION__, false);
-// 		}
-
-		OnNodesPastedDelegate.Broadcast(PastedNodes);
-
-		TSet<UObject*> PastedObjects;
-		for (UEdGraphNode* PastedNode : PastedNodes)
-		{
-			PastedObjects.Add(PastedNode);
-		}
+		TArray<FNiagaraSystemViewModel::FEmitterHandleToDuplicate> EmitterHandlesToDuplicate;
+ 		for (UEdGraphNode* PastedNode : PastedNodes)
+ 		{
+			UNiagaraOverviewNode* OverviewNode = Cast<UNiagaraOverviewNode>(PastedNode);
+			
+			if(OverviewNode != nullptr)
+			{
+				if (OverviewNode->GetOwningSystem() != nullptr && OverviewNode->GetEmitterHandleGuid().IsValid())
+				{
+					FNiagaraSystemViewModel::FEmitterHandleToDuplicate EmitterHandleToDuplicate;
+					EmitterHandleToDuplicate.SystemPath = OverviewNode->GetOwningSystem()->GetPathName();
+					EmitterHandleToDuplicate.EmitterHandleId = OverviewNode->GetEmitterHandleGuid();
+					EmitterHandlesToDuplicate.Add(EmitterHandleToDuplicate);
+				}
+				// Once we've collected the data from the pasted overview node delete it, since a proper node will be created as part
+				// of the duplication process.
+				OverviewNode->DestroyNode();
+			}
+ 		}
 
-		NodeSelection->SetSelectedObjects(PastedObjects);
-		Graph->NotifyGraphChanged(); //@TODO System Overview: might not be necessary 
+		GetSystemViewModel()->DuplicateEmitters(EmitterHandlesToDuplicate);
 	}
 }
 
@@ -405,4 +417,10 @@ void FNiagaraOverviewGraphViewModel::SystemSelectionChanged(UNiagaraSystemSelect
 	}
 }
 
+void FNiagaraOverviewGraphViewModel::PostUndo(bool bSuccess)
+{
+	// This is neccessary to have the graph editor respond correctly to data changes due to undo.
+	OverviewGraph->NotifyGraphChanged();
+}
+
 #undef LOCTEXT_NAMESPACE // NiagaraScriptGraphViewModel
diff --git a/Engine/Plugins/FX/Niagara/Source/NiagaraEditor/Private/ViewModels/Stack/NiagaraStackScriptItemGroup.cpp b/Engine/Plugins/FX/Niagara/Source/NiagaraEditor/Private/ViewModels/Stack/NiagaraStackScriptItemGroup.cpp
--- a/Engine/Plugins/FX/Niagara/Source/NiagaraEditor/Private/ViewModels/Stack/NiagaraStackScriptItemGroup.cpp
+++ b/Engine/Plugins/FX/Niagara/Source/NiagaraEditor/Private/ViewModels/Stack/NiagaraStackScriptItemGroup.cpp
@@ -313,6 +313,14 @@ void UNiagaraStackScriptItemGroup::FinalizeInternal()
 	{
 		ScriptGraph->RemoveOnGraphChangedHandler(OnGraphChangedHandle);
 	}
+	if (AddUtilities.IsValid())
+	{
+		AddUtilities.Reset();
+	}
+	if (ScriptViewModel.IsValid())
+	{
+		ScriptViewModel.Reset();
+	}
 	Super::FinalizeInternal();
 }
 
diff --git a/Engine/Plugins/FX/Niagara/Source/NiagaraEditor/Public/ViewModels/NiagaraOverviewGraphViewModel.h b/Engine/Plugins/FX/Niagara/Source/NiagaraEditor/Public/ViewModels/NiagaraOverviewGraphViewModel.h
--- a/Engine/Plugins/FX/Niagara/Source/NiagaraEditor/Public/ViewModels/NiagaraOverviewGraphViewModel.h
+++ b/Engine/Plugins/FX/Niagara/Source/NiagaraEditor/Public/ViewModels/NiagaraOverviewGraphViewModel.h
@@ -15,14 +15,8 @@ class FNiagaraObjectSelection;
 struct FNiagaraGraphViewSettings;
 
 /** A view model for editing a niagara system in a graph editor. */
-class FNiagaraOverviewGraphViewModel : public TSharedFromThis<FNiagaraOverviewGraphViewModel>
+class FNiagaraOverviewGraphViewModel : public TSharedFromThis<FNiagaraOverviewGraphViewModel>, public FEditorUndoClient
 {
-public:
-	/** A multicast delegate which is called when nodes are pasted in the graph which supplies the pasted nodes. */
-	DECLARE_MULTICAST_DELEGATE_OneParam(FOnNodesPasted, const TSet<UEdGraphNode*>&);
-
-	DECLARE_MULTICAST_DELEGATE(FOnGraphChanged);
-
 public:
 	/** Create a new view model with the supplied system editor data and graph widget. */
 	FNiagaraOverviewGraphViewModel();
@@ -51,8 +45,9 @@ public:
 
 	NIAGARAEDITOR_API void SetViewSettings(const FNiagaraGraphViewSettings& InOverviewGraphViewSettings);
 
-	/** Gets a multicast delegate which is called any time nodes are pasted in the graph. */
-	FOnNodesPasted& OnNodesPasted();
+	//~ FEditorUndoClient interface.
+	virtual void PostUndo(bool bSuccess) override;
+	virtual void PostRedo(bool bSuccess) override { PostUndo(bSuccess); }
 	
 private:
 	void SetupCommands();
@@ -91,12 +86,6 @@ private:
 	/** The set of nodes objects currently selected in the graph. */
 	TSharedRef<FNiagaraObjectSelection> NodeSelection;
 
-	/** A multicast delegate which is called whenever nodes are pasted into the graph. */
-	FOnNodesPasted OnNodesPastedDelegate;
-
-	/** A multicast delegate which is called whenever the graph object is changed to a different graph. */
-	FOnGraphChanged OnGraphChangedDelegate;
-
 	bool bUpdatingSystemSelectionFromGraph;
 
 	bool bUpdatingGraphSelectionFromSystem;
