Copying change 8008903 - Niagara - Fix script graph key bound commands.

diff --git a/Engine/Plugins/FX/Niagara/Source/NiagaraEditor/Private/Widgets/SNiagaraScriptGraph.cpp b/Engine/Plugins/FX/Niagara/Source/NiagaraEditor/Private/Widgets/SNiagaraScriptGraph.cpp
--- a/Engine/Plugins/FX/Niagara/Source/NiagaraEditor/Private/Widgets/SNiagaraScriptGraph.cpp
+++ b/Engine/Plugins/FX/Niagara/Source/NiagaraEditor/Private/Widgets/SNiagaraScriptGraph.cpp
@@ -145,7 +145,7 @@ TSharedRef<SGraphEditor> SNiagaraScriptGraph::ConstructGraphEditor()
 	Events.OnVerifyTextCommit = FOnNodeVerifyTextCommit::CreateSP(this, &SNiagaraScriptGraph::OnVerifyNodeTextCommit);
 	Events.OnSpawnNodeByShortcut = SGraphEditor::FOnSpawnNodeByShortcut::CreateSP(this, &SNiagaraScriptGraph::OnSpawnGraphNodeByShortcut);
 
-	TSharedRef<FUICommandList> Commands = MakeShared<FUICommandList>();
+	Commands = MakeShared<FUICommandList>();
 	Commands->Append(ViewModel->GetCommands());
 	Commands->MapAction(
 		FNiagaraEditorCommands::Get().FindInCurrentView,
@@ -155,7 +155,7 @@ TSharedRef<SGraphEditor> SNiagaraScriptGraph::ConstructGraphEditor()
 		FExecuteAction::CreateRaw(this, &SNiagaraScriptGraph::OnCreateComment));
 	
 	TSharedRef<SGraphEditor> CreatedGraphEditor = SNew(SGraphEditor)
-		.AdditionalCommands(Commands)
+		.AdditionalCommands(Commands.ToSharedRef())
 		.Appearance(AppearanceInfo)
 		.TitleBar(TitleBarWidget)
 		.GraphToEdit(ViewModel->GetGraph())
diff --git a/Engine/Plugins/FX/Niagara/Source/NiagaraEditor/Private/Widgets/SNiagaraScriptGraph.h b/Engine/Plugins/FX/Niagara/Source/NiagaraEditor/Private/Widgets/SNiagaraScriptGraph.h
--- a/Engine/Plugins/FX/Niagara/Source/NiagaraEditor/Private/Widgets/SNiagaraScriptGraph.h
+++ b/Engine/Plugins/FX/Niagara/Source/NiagaraEditor/Private/Widgets/SNiagaraScriptGraph.h
@@ -82,6 +82,9 @@ private:
 	FReply HandleGraphSearchBoxKeyDown(const FGeometry& MyGeometry, const FKeyEvent& InKeyEvent);
 
 private:
+	/** The combined commands used by the graph editor. */
+	TSharedPtr<FUICommandList> Commands;
+
 	/** An attribute for the title text of the graph. */
 	TAttribute<FText> GraphTitle;
 
