Can now optionally enable detailed stats for niagara modules.
Turn on via vm.DetailedVMScriptStats

diff --git a/Engine/Plugins/FX/Niagara/Source/NiagaraEditor/Private/NiagaraHlslTranslator.cpp b/Engine/Plugins/FX/Niagara/Source/NiagaraEditor/Private/NiagaraHlslTranslator.cpp
--- a/Engine/Plugins/FX/Niagara/Source/NiagaraEditor/Private/NiagaraHlslTranslator.cpp
+++ b/Engine/Plugins/FX/Niagara/Source/NiagaraEditor/Private/NiagaraHlslTranslator.cpp
@@ -3762,6 +3762,9 @@ void FHlslNiagaraTranslator::Emitter(class UNiagaraNodeEmitter* EmitterNode, TAr
 	FString Name = EmitterNode->GetName();
 	FString FullName = EmitterNode->GetFullName();
 
+	FName StatName = *EmitterUniqueName;
+	EnterStatsScope(FNiagaraStatScope(StatName, StatName));
+
 	TArray<UEdGraphPin*> CallOutputs;
 	TArray<UEdGraphPin*> CallInputs;
 	EmitterNode->GetOutputPins(CallOutputs);
@@ -3799,7 +3802,7 @@ void FHlslNiagaraTranslator::Emitter(class UNiagaraNodeEmitter* EmitterNode, TAr
 
 	// We act like a function call here as the semantics are identical.
 	RegisterFunctionCall(ScriptUsage, Name, FullName, EmitterNode->NodeGuid, Source, Signature, false, FString(), Inputs, CallInputs, CallOutputs, Signature);
-	GenerateFunctionCall(Signature, Inputs, Outputs);
+	GenerateFunctionCall(ScriptUsage, Signature, Inputs, Outputs);
 
 	// Clear out the parameter map writes to emitter module parameters as they should not be shared across emitters.
 	if (ParamMapHistoryIdx != -1 && ParamMapHistoryIdx < ParamMapHistories.Num())
@@ -3815,6 +3818,8 @@ void FHlslNiagaraTranslator::Emitter(class UNiagaraNodeEmitter* EmitterNode, TAr
 		}
 	}
 	ActiveHistoryForFunctionCalls.ExitEmitter(EmitterUniqueName, EmitterNode);
+
+	ExitStatsScope();
 }
 
 void FHlslNiagaraTranslator::ParameterMapGet(UNiagaraNodeParameterMapGet* GetNode, TArray<int32>& Inputs, TArray<int32>& Outputs)
@@ -4548,7 +4553,7 @@ void FHlslNiagaraTranslator::FunctionCall(UNiagaraNodeFunctionCall* FunctionNode
 		return;
 	}
 
-	GenerateFunctionCall(OutputSignature, Inputs, Outputs);
+	GenerateFunctionCall(ScriptUsage, OutputSignature, Inputs, Outputs);
 
 	if (bCustomHlsl)
 	{
@@ -5118,11 +5123,16 @@ void FHlslNiagaraTranslator::RegisterFunctionCall(ENiagaraScriptUsage ScriptUsag
 	}
 }
 
-void FHlslNiagaraTranslator::GenerateFunctionCall(FNiagaraFunctionSignature& FunctionSignature, TArray<int32>& Inputs, TArray<int32>& Outputs)
+void FHlslNiagaraTranslator::GenerateFunctionCall(ENiagaraScriptUsage ScriptUsage, FNiagaraFunctionSignature& FunctionSignature, TArray<int32>& Inputs, TArray<int32>& Outputs)
 {
 	SCOPE_CYCLE_COUNTER(STAT_NiagaraEditor_Module_NiagaraHLSLTranslator_GenerateFunctionCall);
 
-	EnterStatsScope(FNiagaraStatScope(*GetFunctionSignatureSymbol(FunctionSignature), *(FunctionSignature.GetName())));
+	bool bEnteredStatScope = false;   
+	if (ScriptUsage == ENiagaraScriptUsage::Module)
+	{
+		bEnteredStatScope = true;
+		EnterStatsScope(FNiagaraStatScope(*GetFunctionSignatureSymbol(FunctionSignature), *(FunctionSignature.GetName())));
+	}
 
 	TArray<FString> MissingParameters;
 	int32 ParamIdx = 0;
@@ -5248,7 +5258,10 @@ void FHlslNiagaraTranslator::GenerateFunctionCall(FNiagaraFunctionSignature& Fun
 
 	AddBodyChunk(TEXT(""), DefStr, FNiagaraTypeDefinition::GetFloatDef(), Params);
 
-	ExitStatsScope();
+	if (bEnteredStatScope)
+	{
+		ExitStatsScope();
+	}
 }
 
 FString FHlslNiagaraTranslator::GetFunctionSignatureSymbol(const FNiagaraFunctionSignature& Sig)
diff --git a/Engine/Plugins/FX/Niagara/Source/NiagaraEditor/Private/NiagaraHlslTranslator.h b/Engine/Plugins/FX/Niagara/Source/NiagaraEditor/Private/NiagaraHlslTranslator.h
--- a/Engine/Plugins/FX/Niagara/Source/NiagaraEditor/Private/NiagaraHlslTranslator.h
+++ b/Engine/Plugins/FX/Niagara/Source/NiagaraEditor/Private/NiagaraHlslTranslator.h
@@ -339,7 +339,7 @@ protected:
 
 	void RegisterFunctionCall(ENiagaraScriptUsage ScriptUsage, const FString& InName, const FString& InFullName, const FGuid& CallNodeId, UNiagaraScriptSource* Source, FNiagaraFunctionSignature& InSignature, bool bIsCustomHlsl, const FString& InCustomHlsl, TArray<int32>& Inputs, const TArray<UEdGraphPin*>& CallInputs, const TArray<UEdGraphPin*>& CallOutputs,
 		FNiagaraFunctionSignature& OutSignature);
-	void GenerateFunctionCall(FNiagaraFunctionSignature& FunctionSignature, TArray<int32>& Inputs, TArray<int32>& Outputs);
+	void GenerateFunctionCall(ENiagaraScriptUsage ScriptUsage, FNiagaraFunctionSignature& FunctionSignature, TArray<int32>& Inputs, TArray<int32>& Outputs);
 	FString GetFunctionSignature(const FNiagaraFunctionSignature& Sig);
 
 	/** Compiles an output Pin on a graph node. Caches the result for any future inputs connected to it. */
diff --git a/Engine/Source/Runtime/VectorVM/Private/VectorVM.cpp b/Engine/Source/Runtime/VectorVM/Private/VectorVM.cpp
--- a/Engine/Source/Runtime/VectorVM/Private/VectorVM.cpp
+++ b/Engine/Source/Runtime/VectorVM/Private/VectorVM.cpp
@@ -81,6 +81,15 @@ static FAutoConsoleVariableRef CVarParallelVVMChunksPerBatch(
 	ECVF_Default
 );
 
+//These are possibly too granular to enable for everyone.
+static int32 GbDetailedVMScriptStats = 0;
+static FAutoConsoleVariableRef CVarDetailedVMScriptStats(
+	TEXT("vm.DetailedVMScriptStats"),
+	GbDetailedVMScriptStats,
+	TEXT("If > 0 the vector VM will emit stats for it's internal module calls. \n"),
+	ECVF_Default
+);
+
 //////////////////////////////////////////////////////////////////////////
 //  Constant Handlers
 
@@ -780,10 +789,11 @@ struct FVectorKernelEnterStatScope
 	static VM_FORCEINLINE void Exec(FVectorVMContext& Context)
 	{
 		FConstantHandler<int32> ScopeIdx(Context);
-#if STATS
-		//int32 CounterIdx = Context.StatCounterStack.AddDefaulted(1);
-		//Context.StatCounterStack[CounterIdx].Start((*Context.StatScopes)[ScopeIdx.Get()]);
-#endif
+		if (STATS && GbDetailedVMScriptStats)
+		{
+			int32 CounterIdx = Context.StatCounterStack.AddDefaulted(1);
+			Context.StatCounterStack[CounterIdx].Start((*Context.StatScopes)[ScopeIdx.Get()]);
+		}
 	}
 };
 
@@ -791,10 +801,11 @@ struct FVectorKernelExitStatScope
 {
 	static VM_FORCEINLINE void Exec(FVectorVMContext& Context)
 	{
-#if STATS
-		//Context.StatCounterStack.Last().Stop();
-		//Context.StatCounterStack.Pop(false);
-#endif
+		if (STATS && GbDetailedVMScriptStats)
+		{
+			Context.StatCounterStack.Last().Stop();
+			Context.StatCounterStack.Pop(false);
+		}
 	}
 };
 
