Disabled Niagara GPU system through NiagaraSupportsComputeShaders().

diff --git a/Engine/Plugins/FX/Niagara/Source/Niagara/Private/NiagaraGPUInstanceCountManager.cpp b/Engine/Plugins/FX/Niagara/Source/Niagara/Private/NiagaraGPUInstanceCountManager.cpp
--- a/Engine/Plugins/FX/Niagara/Source/Niagara/Private/NiagaraGPUInstanceCountManager.cpp
+++ b/Engine/Plugins/FX/Niagara/Source/Niagara/Private/NiagaraGPUInstanceCountManager.cpp
@@ -166,7 +166,7 @@ void FNiagaraGPUInstanceCountManager::UpdateDrawIndirectBuffer(FRHICommandList&
 {
 	if (DrawIndirectArgGenTasks.Num() || InstanceCountClearTasks.Num())
 	{
-		if (FeatureLevel == ERHIFeatureLevel::SM5 /*|| FeatureLevel == ERHIFeatureLevel::ES3_1*/) // Must match with RHISupportsComputeShaders()
+		if (NiagaraSupportsComputeShaders(GShaderPlatformForFeatureLevel[FeatureLevel]))
 		{
 			RHICmdList.TransitionResource(EResourceTransitionAccess::ERWBarrier, EResourceTransitionPipeline::EComputeToCompute, CountBuffer.UAV);
 
diff --git a/Engine/Plugins/FX/Niagara/Source/Niagara/Private/NiagaraSystemSimulation.cpp b/Engine/Plugins/FX/Niagara/Source/Niagara/Private/NiagaraSystemSimulation.cpp
--- a/Engine/Plugins/FX/Niagara/Source/Niagara/Private/NiagaraSystemSimulation.cpp
+++ b/Engine/Plugins/FX/Niagara/Source/Niagara/Private/NiagaraSystemSimulation.cpp
@@ -687,7 +687,7 @@ bool FNiagaraSystemSimulation::Tick(float DeltaSeconds)
 			SystemInstance->FinalizeTick(DeltaSeconds);
 			// @todo If we buffer up the Tick object we create here we can push all the data for every system in this simulation at once.
 			// @todo It may also be worth considering hoisting the push out to FNiagaraWorldManager::Tick to push every desired simulation at once.
-			if (SystemInstance->ActiveGPUEmitterCount > 0 && RHISupportsComputeShaders(SystemInstance->GetBatcher()->GetShaderPlatform()))
+			if (SystemInstance->ActiveGPUEmitterCount > 0 && NiagaraSupportsComputeShaders(SystemInstance->GetBatcher()->GetShaderPlatform()))
 			{
 				ensure(!SystemInstance->IsComplete());
 				FNiagaraGPUSystemTick GPUTick;
diff --git a/Engine/Plugins/FX/Niagara/Source/Niagara/Public/NiagaraCommon.h b/Engine/Plugins/FX/Niagara/Source/Niagara/Public/NiagaraCommon.h
--- a/Engine/Plugins/FX/Niagara/Source/Niagara/Public/NiagaraCommon.h
+++ b/Engine/Plugins/FX/Niagara/Source/Niagara/Public/NiagaraCommon.h
@@ -99,6 +99,12 @@ enum class ENiagaraScriptCompileStatus : uint8
 	NCS_MAX,
 };
 
+FORCEINLINE bool NiagaraSupportsComputeShaders(EShaderPlatform ShaderPlatform)
+{
+	// Change to RHISupportsComputeShaders(ShaderPlatform) to support ES3_1
+	return IsFeatureLevelSupported(ShaderPlatform, ERHIFeatureLevel::SM5);
+}
+
 USTRUCT()
 struct FNiagaraDataSetID
 {
@@ -591,12 +597,12 @@ namespace FNiagaraUtilities
 	inline bool SupportsGPUParticles(ERHIFeatureLevel::Type FeatureLevel)
 	{
 		EShaderPlatform ShaderPlatform = GShaderPlatformForFeatureLevel[FeatureLevel];
-		return RHISupportsComputeShaders(ShaderPlatform);
+		return NiagaraSupportsComputeShaders(ShaderPlatform);
 	}
 
 	inline bool SupportsGPUParticles(EShaderPlatform ShaderPlatform)
 	{
-		return RHISupportsComputeShaders(ShaderPlatform);
+		return NiagaraSupportsComputeShaders(ShaderPlatform);
 	}
 
 #if WITH_EDITORONLY_DATA
