Move NiagaraEditor module to load in PreDefault phase to avoid assets in Niagara module that depend on NiagaraEditor loading too early and subsequently failing to launch editor.

diff --git a/Engine/Plugins/FX/Niagara/Niagara.uplugin b/Engine/Plugins/FX/Niagara/Niagara.uplugin
--- a/Engine/Plugins/FX/Niagara/Niagara.uplugin
+++ b/Engine/Plugins/FX/Niagara/Niagara.uplugin
@@ -44,7 +44,7 @@
 		{
 			"Name" : "NiagaraEditor",
 			"Type" : "Editor",
-			"LoadingPhase" : "Default"
+			"LoadingPhase" : "PreDefault"
 		},
 		{
 			"Name" : "NiagaraEditorWidgets",
