Niagara - Fix mac build warning.

diff --git a/Engine/Plugins/FX/Niagara/Source/NiagaraEditor/Private/NiagaraNodeOutput.cpp b/Engine/Plugins/FX/Niagara/Source/NiagaraEditor/Private/NiagaraNodeOutput.cpp
--- a/Engine/Plugins/FX/Niagara/Source/NiagaraEditor/Private/NiagaraNodeOutput.cpp
+++ b/Engine/Plugins/FX/Niagara/Source/NiagaraEditor/Private/NiagaraNodeOutput.cpp
@@ -96,7 +96,7 @@ void UNiagaraNodeOutput::GetContextMenuActions(const FGraphNodeContextMenuBuilde
 					SNew(SEditableTextBoxWithVerification)
 					.Text_UObject(this, &UNiagaraNodeOutput::GetPinNameText, Pin)
 					.OnVerifyTextChanged_UObject(this, &UNiagaraNodeOutput::VerifyPinNameTextChanged, Pin)
-					.OnTextCommitted_UObject(this, &UNiagaraNodeOutput::PinNameTextCommitted, Pin)
+					.OnTextCommitted_UObject(const_cast<UNiagaraNodeOutput*>(this), &UNiagaraNodeOutput::PinNameTextCommitted, Pin)
 				];
 			Context.MenuBuilder->AddWidget(RenameWidget, LOCTEXT("NameMenuItem", "Name"));
 		}
