Added a bAllowGPUParticleUpdate parameter to most callbacks for FXSystemInterface.
This allows each callback to still do some processing even if the GPU particle should not be ticked.
Centralized the logic in FSceneView::AllowGPUParticleUpdate().

Unbound all rendertargets after UpdateSceneCaptureContent_RenderThread() as they could be used
directly after for GPU particle simulation.

diff --git a/Engine/Plugins/FX/Niagara/Source/Niagara/Classes/NiagaraEmitterInstanceBatcher.h b/Engine/Plugins/FX/Niagara/Source/Niagara/Classes/NiagaraEmitterInstanceBatcher.h
--- a/Engine/Plugins/FX/Niagara/Source/Niagara/Classes/NiagaraEmitterInstanceBatcher.h
+++ b/Engine/Plugins/FX/Niagara/Source/Niagara/Classes/NiagaraEmitterInstanceBatcher.h
@@ -78,9 +78,9 @@ public:
 	virtual void AddVectorField(UVectorFieldComponent* VectorFieldComponent) override {}
 	virtual void RemoveVectorField(UVectorFieldComponent* VectorFieldComponent) override {}
 	virtual void UpdateVectorField(UVectorFieldComponent* VectorFieldComponent) override {}
-	virtual void PreInitViews(FRHICommandListImmediate& RHICmdList) override;
+	virtual void PreInitViews(FRHICommandListImmediate& RHICmdList, bool bAllowGPUParticleUpdate) override;
 	virtual bool UsesGlobalDistanceField() const override;
-	virtual void PreRender(FRHICommandListImmediate& RHICmdList, const class FGlobalDistanceFieldParameterData* GlobalDistanceFieldParameterData, bool bAllowGPUParticleSceneUpdate) override;
+	virtual void PreRender(FRHICommandListImmediate& RHICmdList, const class FGlobalDistanceFieldParameterData* GlobalDistanceFieldParameterData, bool bAllowGPUParticleUpdate) override;
 	virtual void OnDestroy() override; // Called on the gamethread to delete the batcher on the renderthread.
 
 	virtual void Tick(float DeltaTime) override
@@ -102,7 +102,8 @@ public:
 		FRHICommandListImmediate& RHICmdList,
 		FRHIUniformBuffer* ViewUniformBuffer,
 		const class FShaderParametersMetadata* SceneTexturesUniformBufferStruct,
-		FRHIUniformBuffer* SceneTexturesUniformBuffer) override;
+		FRHIUniformBuffer* SceneTexturesUniformBuffer,
+		bool bAllowGPUParticleUpdate) override;
 
 	int32 AddSortedGPUSimulation(const FNiagaraGPUSortInfo& SortInfo);
 	void SortGPUParticles(FRHICommandListImmediate& RHICmdList);
diff --git a/Engine/Plugins/FX/Niagara/Source/Niagara/Private/NiagaraEmitterInstanceBatcher.cpp b/Engine/Plugins/FX/Niagara/Source/Niagara/Private/NiagaraEmitterInstanceBatcher.cpp
--- a/Engine/Plugins/FX/Niagara/Source/Niagara/Private/NiagaraEmitterInstanceBatcher.cpp
+++ b/Engine/Plugins/FX/Niagara/Source/Niagara/Private/NiagaraEmitterInstanceBatcher.cpp
@@ -505,14 +505,17 @@ void NiagaraEmitterInstanceBatcher::DispatchAllOnCompute(FOverlappableTicks& Ove
 	RHICmdList.SubmitCommandsHint();
 }
 
-void NiagaraEmitterInstanceBatcher::PostRenderOpaque(FRHICommandListImmediate& RHICmdList, FRHIUniformBuffer* ViewUniformBuffer, const class FShaderParametersMetadata* SceneTexturesUniformBufferStruct, FRHIUniformBuffer* SceneTexturesUniformBuffer)
+void NiagaraEmitterInstanceBatcher::PostRenderOpaque(FRHICommandListImmediate& RHICmdList, FRHIUniformBuffer* ViewUniformBuffer, const class FShaderParametersMetadata* SceneTexturesUniformBufferStruct, FRHIUniformBuffer* SceneTexturesUniformBuffer, bool bAllowGPUParticleUpdate)
 {
 	LLM_SCOPE(ELLMTag::Niagara);
 
-	// Setup new readback since if there is no pending request, there is no risk of having invalid data read (offset being allocated after the readback was sent).
-	ExecuteAll(RHICmdList, ViewUniformBuffer, !GPUInstanceCounterManager.HasPendingGPUReadback());
+	if (bAllowGPUParticleUpdate)
+	{
+		// Setup new readback since if there is no pending request, there is no risk of having invalid data read (offset being allocated after the readback was sent).
+		ExecuteAll(RHICmdList, ViewUniformBuffer, !GPUInstanceCounterManager.HasPendingGPUReadback());
 
-	FinishDispatches();
+		FinishDispatches();
+	}
 
 	if (!GPUInstanceCounterManager.HasPendingGPUReadback())
 	{
@@ -684,7 +687,7 @@ void NiagaraEmitterInstanceBatcher::TickSingle(const FNiagaraGPUSystemTick& Tick
 	}
 }
 
-void NiagaraEmitterInstanceBatcher::PreInitViews(FRHICommandListImmediate& RHICmdList)
+void NiagaraEmitterInstanceBatcher::PreInitViews(FRHICommandListImmediate& RHICmdList, bool bAllowGPUParticleUpdate)
 {
 	LLM_SCOPE(ELLMTag::Niagara);
 
@@ -697,72 +700,75 @@ void NiagaraEmitterInstanceBatcher::PreInitViews(FRHICommandListImmediate& RHICm
 	}
 
 	// Update draw indirect buffer to max possible size.
-	int32 TotalDispatchCount = 0;
-	for (FNiagaraGPUSystemTick& Tick : Ticks_RT)
+	if (bAllowGPUParticleUpdate)
 	{
-		TotalDispatchCount += (int32)Tick.Count;
-
-		// Cancel any pending readback if the emitter is resetting.
-		if (Tick.bNeedsReset)
+		int32 TotalDispatchCount = 0;
+		for (FNiagaraGPUSystemTick& Tick : Ticks_RT)
 		{
-			FNiagaraComputeInstanceData* Instances = Tick.GetInstanceData();
-			for (uint32 InstanceIndex = 0; InstanceIndex < Tick.Count; ++InstanceIndex)
+			TotalDispatchCount += (int32)Tick.Count;
+
+			// Cancel any pending readback if the emitter is resetting.
+			if (Tick.bNeedsReset)
 			{
-				FNiagaraComputeExecutionContext* Context = Instances[InstanceIndex].Context;
-				if (Context)
+				FNiagaraComputeInstanceData* Instances = Tick.GetInstanceData();
+				for (uint32 InstanceIndex = 0; InstanceIndex < Tick.Count; ++InstanceIndex)
 				{
-					GPUInstanceCounterManager.FreeEntry(Context->EmitterInstanceReadback.GPUCountOffset);
+					FNiagaraComputeExecutionContext* Context = Instances[InstanceIndex].Context;
+					if (Context)
+					{
+						GPUInstanceCounterManager.FreeEntry(Context->EmitterInstanceReadback.GPUCountOffset);
+					}
 				}
 			}
 		}
-	}
-	GPUInstanceCounterManager.ResizeBuffers(TotalDispatchCount);
+		GPUInstanceCounterManager.ResizeBuffers(TotalDispatchCount);
 
-	// Update the instance counts from the GPU readback.
-	{
-		SCOPE_CYCLE_COUNTER(STAT_NiagaraGPUReadback_RT);
-		const uint32* Counts = GPUInstanceCounterManager.GetGPUReadback();
-		if (Counts)
+		// Update the instance counts from the GPU readback.
 		{
-			for (FNiagaraGPUSystemTick& Tick : Ticks_RT)
+			SCOPE_CYCLE_COUNTER(STAT_NiagaraGPUReadback_RT);
+			const uint32* Counts = GPUInstanceCounterManager.GetGPUReadback();
+			if (Counts)
 			{
-				FNiagaraComputeInstanceData* Instances = Tick.GetInstanceData();
-				for (uint32 InstanceIndex = 0; InstanceIndex < Tick.Count; ++InstanceIndex)
+				for (FNiagaraGPUSystemTick& Tick : Ticks_RT)
 				{
-					FNiagaraComputeExecutionContext* Context = Instances[InstanceIndex].Context;
-					if (Context && Context->EmitterInstanceReadback.GPUCountOffset != INDEX_NONE)
+					FNiagaraComputeInstanceData* Instances = Tick.GetInstanceData();
+					for (uint32 InstanceIndex = 0; InstanceIndex < Tick.Count; ++InstanceIndex)
 					{
-						check(Context->MainDataSet);
-						FNiagaraDataBuffer* CurrentData = Context->MainDataSet->GetCurrentData();
-						if (CurrentData)
+						FNiagaraComputeExecutionContext* Context = Instances[InstanceIndex].Context;
+						if (Context && Context->EmitterInstanceReadback.GPUCountOffset != INDEX_NONE)
 						{
-							const uint32 DeadInstanceCount = Context->EmitterInstanceReadback.CPUCount - Counts[Context->EmitterInstanceReadback.GPUCountOffset];
-
-							// This will communicate the particle counts to the game thread. If DeadInstanceCount equals CurrentData->GetNumInstances() the game thread will know that the emitter has completed.
-							if (DeadInstanceCount <= CurrentData->GetNumInstances())
+							check(Context->MainDataSet);
+							FNiagaraDataBuffer* CurrentData = Context->MainDataSet->GetCurrentData();
+							if (CurrentData)
 							{
-								CurrentData->SetNumInstances(CurrentData->GetNumInstances() - DeadInstanceCount); 
+								const uint32 DeadInstanceCount = Context->EmitterInstanceReadback.CPUCount - Counts[Context->EmitterInstanceReadback.GPUCountOffset];
+
+								// This will communicate the particle counts to the game thread. If DeadInstanceCount equals CurrentData->GetNumInstances() the game thread will know that the emitter has completed.
+								if (DeadInstanceCount <= CurrentData->GetNumInstances())
+								{
+									CurrentData->SetNumInstances(CurrentData->GetNumInstances() - DeadInstanceCount); 
+								}
 							}
-						}
 
-						// Now release the readback since another one will be enqueued in the tick.
-						// Also prevents processing the same data again.
-						GPUInstanceCounterManager.FreeEntry(Context->EmitterInstanceReadback.GPUCountOffset);
+							// Now release the readback since another one will be enqueued in the tick.
+							// Also prevents processing the same data again.
+							GPUInstanceCounterManager.FreeEntry(Context->EmitterInstanceReadback.GPUCountOffset);
+						}
 					}
 				}
+				// Readback is only valid for one frame, so that any newly allocated instance count
+				// are guarantied to be in the next valid readback data.
+				GPUInstanceCounterManager.ReleaseGPUReadback();
 			}
-			// Readback is only valid for one frame, so that any newly allocated instance count
-			// are guarantied to be in the next valid readback data.
-			GPUInstanceCounterManager.ReleaseGPUReadback();
 		}
-	}
 
-	// @todo REMOVE THIS HACK
-	LastFrameThatDrainedData = GFrameNumberRenderThread;
+		// @todo REMOVE THIS HACK
+		LastFrameThatDrainedData = GFrameNumberRenderThread;
 
-	if (GNiagaraAllowTickBeforeRender)
-	{
-		ExecuteAll(RHICmdList, nullptr, !GPUInstanceCounterManager.HasPendingGPUReadback());
+		if (GNiagaraAllowTickBeforeRender)
+		{
+			ExecuteAll(RHICmdList, nullptr, !GPUInstanceCounterManager.HasPendingGPUReadback());
+		}
 	}
 }
 
@@ -779,7 +785,7 @@ bool NiagaraEmitterInstanceBatcher::UsesGlobalDistanceField() const
 	return false;
 }
 
-void NiagaraEmitterInstanceBatcher::PreRender(FRHICommandListImmediate& RHICmdList, const class FGlobalDistanceFieldParameterData* GlobalDistanceFieldParameterData, bool bAllowGPUParticleSceneUpdate)
+void NiagaraEmitterInstanceBatcher::PreRender(FRHICommandListImmediate& RHICmdList, const class FGlobalDistanceFieldParameterData* GlobalDistanceFieldParameterData, bool bAllowGPUParticleUpdate)
 {
 	LLM_SCOPE(ELLMTag::Niagara);
 
diff --git a/Engine/Source/Runtime/Engine/Private/Particles/FXSystem.cpp b/Engine/Source/Runtime/Engine/Private/Particles/FXSystem.cpp
--- a/Engine/Source/Runtime/Engine/Private/Particles/FXSystem.cpp
+++ b/Engine/Source/Runtime/Engine/Private/Particles/FXSystem.cpp
@@ -356,11 +356,11 @@ void FFXSystem::DrawDebug( FCanvas* Canvas )
 	}
 }
 
-void FFXSystem::PreInitViews(FRHICommandListImmediate& RHICmdList)
+void FFXSystem::PreInitViews(FRHICommandListImmediate& RHICmdList, bool bAllowGPUParticleUpdate)
 {
 	if (RHISupportsGPUParticles())
 	{
-		AdvanceGPUParticleFrame();
+		AdvanceGPUParticleFrame(bAllowGPUParticleUpdate);
 	}
 }
 
@@ -416,16 +416,18 @@ void FFXSystem::PostRenderOpaque(
 	FRHICommandListImmediate& RHICmdList, 
 	FRHIUniformBuffer* ViewUniformBuffer,
 	const FShaderParametersMetadata* SceneTexturesUniformBufferStruct,
-	FRHIUniformBuffer* SceneTexturesUniformBuffer)
+	FRHIUniformBuffer* SceneTexturesUniformBuffer,
+	bool bAllowGPUParticleUpdate)
 {
 	if (RHISupportsGPUParticles() && IsParticleCollisionModeSupported(GetShaderPlatform(), PCM_DepthBuffer))
 	{
 		SCOPED_DRAW_EVENT(RHICmdList, GPUParticles_PostRenderOpaque);
-		PrepareGPUSimulation(RHICmdList);
-		
-		SimulateGPUParticles(RHICmdList, EParticleSimulatePhase::CollisionDepthBuffer, ViewUniformBuffer, NULL, SceneTexturesUniformBufferStruct, SceneTexturesUniformBuffer);
-		
-		FinalizeGPUSimulation(RHICmdList);
+		if (bAllowGPUParticleUpdate)
+		{
+			PrepareGPUSimulation(RHICmdList);
+			SimulateGPUParticles(RHICmdList, EParticleSimulatePhase::CollisionDepthBuffer, ViewUniformBuffer, NULL, SceneTexturesUniformBufferStruct, SceneTexturesUniformBuffer);
+			FinalizeGPUSimulation(RHICmdList);
+		}
 
 		SortGPUParticles(RHICmdList);		
 	}
diff --git a/Engine/Source/Runtime/Engine/Private/Particles/FXSystemPrivate.h b/Engine/Source/Runtime/Engine/Private/Particles/FXSystemPrivate.h
--- a/Engine/Source/Runtime/Engine/Private/Particles/FXSystemPrivate.h
+++ b/Engine/Source/Runtime/Engine/Private/Particles/FXSystemPrivate.h
@@ -149,14 +149,15 @@ public:
 	virtual void RemoveVectorField(UVectorFieldComponent* VectorFieldComponent) override;
 	virtual void UpdateVectorField(UVectorFieldComponent* VectorFieldComponent) override;
 	FParticleEmitterInstance* CreateGPUSpriteEmitterInstance(FGPUSpriteEmitterInfo& EmitterInfo);
-	virtual void PreInitViews(FRHICommandListImmediate& RHICmdList) override;
+	virtual void PreInitViews(FRHICommandListImmediate& RHICmdList, bool bAllowGPUParticleUpdate) override;
 	virtual bool UsesGlobalDistanceField() const override;
-	virtual void PreRender(FRHICommandListImmediate& RHICmdList, const FGlobalDistanceFieldParameterData* GlobalDistanceFieldParameterData, bool bAllowGPUParticleSceneUpdate) override;
+	virtual void PreRender(FRHICommandListImmediate& RHICmdList, const FGlobalDistanceFieldParameterData* GlobalDistanceFieldParameterData, bool bAllowGPUParticleUpdate) override;
 	virtual void PostRenderOpaque(
 		FRHICommandListImmediate& RHICmdList, 
 		FRHIUniformBuffer* ViewUniformBuffer,
 		const FShaderParametersMetadata* SceneTexturesUniformBufferStruct,
-		FRHIUniformBuffer* SceneTexturesUniformBuffer) override;
+		FRHIUniformBuffer* SceneTexturesUniformBuffer,
+		bool bAllowGPUParticleUpdate) override;
 	// End FFXSystemInterface.
 
 	/*--------------------------------------------------------------------------
@@ -232,7 +233,7 @@ private:
 	/**
 	 * Prepares GPU particles for simulation and rendering in the next frame.
 	 */
-	void AdvanceGPUParticleFrame();
+	void AdvanceGPUParticleFrame(bool bAllowGPUParticleUpdate);
 
 	/**
 	 * Sorts all GPU particles that have called AddSortedGPUSimulation since the
diff --git a/Engine/Source/Runtime/Engine/Private/Particles/FXSystemSet.cpp b/Engine/Source/Runtime/Engine/Private/Particles/FXSystemSet.cpp
--- a/Engine/Source/Runtime/Engine/Private/Particles/FXSystemSet.cpp
+++ b/Engine/Source/Runtime/Engine/Private/Particles/FXSystemSet.cpp
@@ -87,12 +87,12 @@ void FFXSystemSet::UpdateVectorField(UVectorFieldComponent* VectorFieldComponent
 	}
 }
 
-void FFXSystemSet::PreInitViews(FRHICommandListImmediate& RHICmdList)
+void FFXSystemSet::PreInitViews(FRHICommandListImmediate& RHICmdList, bool bAllowGPUParticleUpdate)
 {
 	for (FFXSystemInterface* FXSystem : FXSystems)
 	{
 		check(FXSystem);
-		FXSystem->PreInitViews(RHICmdList);
+		FXSystem->PreInitViews(RHICmdList, bAllowGPUParticleUpdate);
 	}
 }
 
@@ -122,7 +122,8 @@ void FFXSystemSet::PostRenderOpaque(
 	FRHICommandListImmediate& RHICmdList,
 	FRHIUniformBuffer* ViewUniformBuffer,
 	const class FShaderParametersMetadata* SceneTexturesUniformBufferStruct,
-	FRHIUniformBuffer* SceneTexturesUniformBuffer)
+	FRHIUniformBuffer* SceneTexturesUniformBuffer,
+	bool bAllowGPUParticleUpdate)
 {
 	for (FFXSystemInterface* FXSystem : FXSystems)
 	{
@@ -131,7 +132,8 @@ void FFXSystemSet::PostRenderOpaque(
 			RHICmdList,
 			ViewUniformBuffer,
 			SceneTexturesUniformBufferStruct,
-			SceneTexturesUniformBuffer
+			SceneTexturesUniformBuffer,
+			bAllowGPUParticleUpdate
 		);
 	}
 }
diff --git a/Engine/Source/Runtime/Engine/Private/Particles/FXSystemSet.h b/Engine/Source/Runtime/Engine/Private/Particles/FXSystemSet.h
--- a/Engine/Source/Runtime/Engine/Private/Particles/FXSystemSet.h
+++ b/Engine/Source/Runtime/Engine/Private/Particles/FXSystemSet.h
@@ -30,14 +30,15 @@ public:
 	virtual void AddVectorField(UVectorFieldComponent* VectorFieldComponent) override;
 	virtual void RemoveVectorField(UVectorFieldComponent* VectorFieldComponent) override;
 	virtual void UpdateVectorField(UVectorFieldComponent* VectorFieldComponent) override;
-	virtual void PreInitViews(FRHICommandListImmediate& RHICmdList) override;
+	virtual void PreInitViews(FRHICommandListImmediate& RHICmdList, bool bAllowGPUParticleUpdate) override;
 	virtual bool UsesGlobalDistanceField() const override;
 	virtual void PreRender(FRHICommandListImmediate& RHICmdList, const class FGlobalDistanceFieldParameterData* GlobalDistanceFieldParameterData, bool bAllowGPUParticleSceneUpdate) override;
 	virtual void PostRenderOpaque(
 		FRHICommandListImmediate& RHICmdList, 
 		FRHIUniformBuffer* ViewUniformBuffer,
 		const class FShaderParametersMetadata* SceneTexturesUniformBufferStruct,
-		FRHIUniformBuffer* SceneTexturesUniformBuffer) override;
+		FRHIUniformBuffer* SceneTexturesUniformBuffer,
+		bool bAllowGPUParticleUpdate) override;
 
 	virtual void OnDestroy() override;
 
diff --git a/Engine/Source/Runtime/Engine/Private/Particles/ParticleGpuSimulation.cpp b/Engine/Source/Runtime/Engine/Private/Particles/ParticleGpuSimulation.cpp
--- a/Engine/Source/Runtime/Engine/Private/Particles/ParticleGpuSimulation.cpp
+++ b/Engine/Source/Runtime/Engine/Private/Particles/ParticleGpuSimulation.cpp
@@ -4605,10 +4605,13 @@ int32 FFXSystem::AddSortedGPUSimulation(FParticleSimulationGPU* Simulation, cons
 	return BufferOffset;
 }
 
-void FFXSystem::AdvanceGPUParticleFrame()
+void FFXSystem::AdvanceGPUParticleFrame(bool bAllowGPUParticleUpdate)
 {
-	// We double buffer, so swap the current and previous textures.
-	ParticleSimulationResources->FrameIndex ^= 0x1;
+	if (bAllowGPUParticleUpdate)
+	{
+		// We double buffer, so swap the current and previous textures.
+		ParticleSimulationResources->FrameIndex ^= 0x1;
+	}
 
 	// Reset the list of sorted simulations. As PreRenderView is called on GPU simulations we'll
 	// allocate space for them in the sorted index buffer.
diff --git a/Engine/Source/Runtime/Engine/Public/FXSystem.h b/Engine/Source/Runtime/Engine/Public/FXSystem.h
--- a/Engine/Source/Runtime/Engine/Public/FXSystem.h
+++ b/Engine/Source/Runtime/Engine/Public/FXSystem.h
@@ -203,7 +203,7 @@ public:
 	 * Notification from the renderer that it is about to perform visibility
 	 * checks on FX belonging to this system.
 	 */
-	virtual void PreInitViews(FRHICommandListImmediate& RHICmdList) = 0;
+	virtual void PreInitViews(FRHICommandListImmediate& RHICmdList, bool bAllowGPUParticleUpdate) = 0;
 
 	virtual bool UsesGlobalDistanceField() const = 0;
 
@@ -220,12 +220,8 @@ public:
 		FRHICommandListImmediate& RHICmdList, 
 		FRHIUniformBuffer* ViewUniformBuffer,
 		const class FShaderParametersMetadata* SceneTexturesUniformBufferStruct,
-		FRHIUniformBuffer* SceneTexturesUniformBuffer) = 0;
-
-	/**
-	 * Helper in case the data necessary for collision is not available.
-	 */
-	void PostRenderOpaque(FRHICommandListImmediate& RHICmdList) { PostRenderOpaque(RHICmdList, nullptr, nullptr, nullptr); }
+		FRHIUniformBuffer* SceneTexturesUniformBuffer,
+		bool bAllowGPUParticleUpdate) = 0;
 
 	bool IsPendingKill() const { return bIsPendingKill; }
 
diff --git a/Engine/Source/Runtime/Engine/Public/SceneView.h b/Engine/Source/Runtime/Engine/Public/SceneView.h
--- a/Engine/Source/Runtime/Engine/Public/SceneView.h
+++ b/Engine/Source/Runtime/Engine/Public/SceneView.h
@@ -1081,6 +1081,8 @@ public:
 	bool VerifyMembersChecks() const;
 #endif
 
+	FORCEINLINE bool AllowGPUParticleUpdate() const { return !bIsPlanarReflection && !bIsSceneCapture && !bIsReflectionCapture; }
+
 	/** Transforms a point from world-space to the view's screen-space. */
 	FVector4 WorldToScreen(const FVector& WorldPoint) const;
 
diff --git a/Engine/Source/Runtime/Renderer/Private/DeferredShadingRenderer.cpp b/Engine/Source/Runtime/Renderer/Private/DeferredShadingRenderer.cpp
--- a/Engine/Source/Runtime/Renderer/Private/DeferredShadingRenderer.cpp
+++ b/Engine/Source/Runtime/Renderer/Private/DeferredShadingRenderer.cpp
@@ -1107,23 +1107,16 @@ void FDeferredShadingSceneRenderer::Render(FRHICommandListImmediate& RHICmdList)
 		}
 	}
 
-	// Only update the GPU particle simulation for the main view
-	//@todo - this is needed because the GPU particle simulation is updated within a frame render.  Simulation should happen outside of a visible frame rendering.
-	// This also causes GPU particles to be one frame behind in scene captures and planar reflections.
-	const bool bAllowGPUParticleSceneUpdate = !Views[0].bIsPlanarReflection && !Views[0].bIsSceneCapture && !Views[0].bIsReflectionCapture;
-
 	// Notify the FX system that the scene is about to be rendered.
-	bool bDoFXPrerender = Scene->FXSystem && Views.IsValidIndex(0);
-
-	if (bDoFXPrerender)
+	if (Scene->FXSystem && Views.IsValidIndex(0))
 	{
 		SCOPE_CYCLE_COUNTER(STAT_FDeferredShadingSceneRenderer_FXSystem_PreRender);
 		RHICmdList.SetCurrentStat(GET_STATID(STAT_CLM_FXPreRender));
-		Scene->FXSystem->PreRender(RHICmdList, &Views[0].GlobalDistanceFieldInfo.ParameterData, bAllowGPUParticleSceneUpdate);
+		Scene->FXSystem->PreRender(RHICmdList, &Views[0].GlobalDistanceFieldInfo.ParameterData, Views[0].AllowGPUParticleUpdate());
 	}
 
 	bool bDidAfterTaskWork = false;
-	auto AfterTasksAreStarted = [&bDidAfterTaskWork, bDoInitViewAftersPrepass, this, &RHICmdList, &ILCTaskData, &UpdateViewCustomDataEvents, bDoFXPrerender]()
+	auto AfterTasksAreStarted = [&bDidAfterTaskWork, bDoInitViewAftersPrepass, this, &RHICmdList, &ILCTaskData, &UpdateViewCustomDataEvents]()
 	{
 		if (!bDidAfterTaskWork)
 		{
@@ -1863,7 +1856,7 @@ void FDeferredShadingSceneRenderer::Render(FRHICommandListImmediate& RHICmdList)
 	UnbindRenderTargets(RHICmdList);
 
 	// Notify the FX system that opaque primitives have been rendered and we now have a valid depth buffer.
-	if (Scene->FXSystem && Views.IsValidIndex(0) && bAllowGPUParticleSceneUpdate)
+	if (Scene->FXSystem && Views.IsValidIndex(0))
 	{
 		CSV_SCOPED_TIMING_STAT_EXCLUSIVE(RenderOpaqueFX);
 		SCOPE_CYCLE_COUNTER(STAT_FDeferredShadingSceneRenderer_FXSystem_PostRenderOpaque);
@@ -1877,7 +1870,8 @@ void FDeferredShadingSceneRenderer::Render(FRHICommandListImmediate& RHICmdList)
 			RHICmdList,
 			Views[0].ViewUniformBuffer,
 			&FSceneTexturesUniformParameters::StaticStructMetadata,
-			SceneTextureUniformBuffer.GetReference()
+			SceneTextureUniformBuffer.GetReference(),
+			Views[0].AllowGPUParticleUpdate()
 		);
 		ServiceLocalQueue();
 	}
diff --git a/Engine/Source/Runtime/Renderer/Private/SceneCaptureRendering.cpp b/Engine/Source/Runtime/Renderer/Private/SceneCaptureRendering.cpp
--- a/Engine/Source/Runtime/Renderer/Private/SceneCaptureRendering.cpp
+++ b/Engine/Source/Runtime/Renderer/Private/SceneCaptureRendering.cpp
@@ -461,7 +461,9 @@ static void UpdateSceneCaptureContent_RenderThread(
 		default:
 			checkNoEntry();
 			break;
-		}
+	}
+	// Unbind everything in case FX has to read.
+	UnbindRenderTargets(RHICmdList);
 }
 
 void BuildProjectionMatrix(FIntPoint RenderTargetSize, ECameraProjectionMode::Type ProjectionType, float FOV, float InOrthoWidth, float InNearClippingPlane, FMatrix& ProjectionMatrix)
diff --git a/Engine/Source/Runtime/Renderer/Private/SceneVisibility.cpp b/Engine/Source/Runtime/Renderer/Private/SceneVisibility.cpp
--- a/Engine/Source/Runtime/Renderer/Private/SceneVisibility.cpp
+++ b/Engine/Source/Runtime/Renderer/Private/SceneVisibility.cpp
@@ -2945,9 +2945,10 @@ void FSceneRenderer::PreVisibilityFrameSetup(FRHICommandListImmediate& RHICmdLis
 	}
 
 	// Notify the FX system that the scene is about to perform visibility checks.
-	if (Scene->FXSystem && !Views[0].bIsPlanarReflection)
+
+	if (Scene->FXSystem && Views.IsValidIndex(0))
 	{
-		Scene->FXSystem->PreInitViews(RHICmdList);
+		Scene->FXSystem->PreInitViews(RHICmdList, Views[0].AllowGPUParticleUpdate());
 	}
 
 	// Draw lines to lights affecting this mesh if its selected.
