Niagara - Fix stability and usability issues in the system editor.

diff --git a/Engine/Plugins/FX/Niagara/Source/NiagaraEditor/Private/NiagaraNode.cpp b/Engine/Plugins/FX/Niagara/Source/NiagaraEditor/Private/NiagaraNode.cpp
--- a/Engine/Plugins/FX/Niagara/Source/NiagaraEditor/Private/NiagaraNode.cpp
+++ b/Engine/Plugins/FX/Niagara/Source/NiagaraEditor/Private/NiagaraNode.cpp
@@ -116,6 +116,9 @@ bool UNiagaraNode::ReallocatePins(bool bMarkNeedsResynchronizeOnChange)
 	}
 	else
 	{
+		// Even if we're not marking the graph as needing sychronization we still need to let other listeners,
+		// such as the UI, know that the graph has changed.
+		GetNiagaraGraph()->NotifyGraphNeedsRecompile();
 		VisualsChangedDelegate.Broadcast(this);
 	}
 	return bAllSame;
diff --git a/Engine/Plugins/FX/Niagara/Source/NiagaraEditor/Private/NiagaraSystemEditorData.cpp b/Engine/Plugins/FX/Niagara/Source/NiagaraEditor/Private/NiagaraSystemEditorData.cpp
--- a/Engine/Plugins/FX/Niagara/Source/NiagaraEditor/Private/NiagaraSystemEditorData.cpp
+++ b/Engine/Plugins/FX/Niagara/Source/NiagaraEditor/Private/NiagaraSystemEditorData.cpp
@@ -287,4 +287,7 @@ void UNiagaraSystemEditorData::SynchronizeOverviewGraphWithSystem(UNiagaraSystem
 		OverviewNode->Modify();
 		OverviewNode->DestroyNode();
 	}
+
+	// Dispatch an empty graph changed message here so that any graph UIs which are visible refresh their node's widgets.
+	SystemOverviewGraph->NotifyGraphChanged();
 }
diff --git a/Engine/Plugins/FX/Niagara/Source/NiagaraEditor/Private/ViewModels/NiagaraEmitterHandleViewModel.cpp b/Engine/Plugins/FX/Niagara/Source/NiagaraEditor/Private/ViewModels/NiagaraEmitterHandleViewModel.cpp
--- a/Engine/Plugins/FX/Niagara/Source/NiagaraEditor/Private/ViewModels/NiagaraEmitterHandleViewModel.cpp
+++ b/Engine/Plugins/FX/Niagara/Source/NiagaraEditor/Private/ViewModels/NiagaraEmitterHandleViewModel.cpp
@@ -64,6 +64,10 @@ void FNiagaraEmitterHandleViewModel::Initialize(TSharedRef<FNiagaraSystemViewMod
 	EmitterHandle = InEmitterHandle;
 	UNiagaraEmitter* Emitter = EmitterHandle != nullptr ? EmitterHandle->GetInstance() : nullptr;
 	EmitterViewModel->Initialize(Emitter, InSimulation);
+	if (EmitterStackViewModel != nullptr)
+	{
+		EmitterStackViewModel->InitializeWithViewModels(InOwningSystemViewModel, this->AsShared(), FNiagaraStackViewModelOptions(false, true));
+	}
 }
 
 void FNiagaraEmitterHandleViewModel::Reset()
