Add spatial hash data interface for CPU

diff --git a/Engine/Plugins/FX/Niagara/Source/Niagara/Classes/NiagaraDataInterfaceSpatialHash.h b/Engine/Plugins/FX/Niagara/Source/Niagara/Classes/NiagaraDataInterfaceSpatialHash.h
new file mode 100644
--- /dev/null
+++ b/Engine/Plugins/FX/Niagara/Source/Niagara/Classes/NiagaraDataInterfaceSpatialHash.h
@@ -0,0 +1,272 @@
+// Copyright 1998-2019 Epic Games, Inc. All Rights Reserved.
+
+#pragma once
+
+#include "NiagaraDataInterface.h"
+#include "NiagaraCommon.h"
+#include "VectorVM.h"
+#include "NiagaraDataInterfaceSpatialHash.generated.h"
+
+struct FNDISpatialHash_InstanceData;
+class FSpatialHashGPUBuffers;
+
+class FNiagaraDINearestNeighborBatch
+{
+public:
+	FNiagaraDINearestNeighborBatch()
+		: CurrentID(0)
+	{
+	}
+
+	void ClearWrite()
+	{
+		uint32 PrevNum = NearestNeighborResults.Num();
+		NearestNeighborResults.Reset();
+		IDToResultIndex.Reset();
+	}
+
+	void Init(FName InBatchID, FNDISpatialHash_InstanceData* InSpatialHashInstanceData);
+
+	int32 SubmitQuery(FNiagaraID ParticleID, FVector Position, float SearchRadius, uint32 MaxNeighbors, bool bIncludeSelf);
+	bool GetQueryResult(uint32 InQueryID, TArray<FNiagaraID>& Result);
+
+private:
+	/** Stores the results of nearest neighbor queries. */
+	TArray<FNiagaraID> NearestNeighborResults;
+
+	/** Maps a query ID to the index in the NearestNeighborResults array that the results begin. */
+	TMap<uint32, int32> IDToResultIndex;
+	uint32 CurrentID;
+	FNDISpatialHash_InstanceData* SpatialHashInstanceData;
+};
+
+struct FNDISpatialHash_InstanceData
+{
+	void AllocatePersistentTables();
+	void ResetTables();
+	void BuildTable();
+	void BuildTableGPU();
+	uint32 NearestNeighbor(FNiagaraID ParticleID, FVector Position, float SearchRadius, uint32 MaxNeighbors, bool bIncludeSelf, TArray<FNiagaraID> &ClosestParticles);
+
+	bool Init(UNiagaraDataInterfaceSpatialHash* Interface, FNiagaraSystemInstance* SystemInstance);
+	FORCEINLINE_DEBUGGABLE bool Tick(UNiagaraDataInterfaceSpatialHash* Interface, FNiagaraSystemInstance* SystemInstance, float InDeltaSeconds);
+	FORCEINLINE_DEBUGGABLE void Release();
+
+	FNiagaraSystemInstance* SystemInstance;
+	FNiagaraDINearestNeighborBatch SpatialHashBatch;
+
+	uint32 MaximumParticleCount;
+	uint32 TableSize;
+	uint32 MaximumNeighborCount;
+	float MaximumSearchRadius;
+
+	/** Number of particles this spatial hash keeps track of. Resets to 0 after a build finishes. */
+	uint32 NumParticles;
+
+	/** The length, width, and height of a cell. Automatically set based on the maximum search radius. */
+	float CellLength;
+
+	// Size = number of particles
+	struct ParticleData
+	{
+		uint32 CellHash;
+		uint32 ParticleID;
+		FVector ParticlePosition;
+		FNiagaraID ExternalID;
+	};
+	TArray<ParticleData> Particles;
+	TArray<ParticleData> Particles_Built;
+
+	// Size = table size
+	TArray<int32> StartIndex;
+	TArray<int32> EndIndex;
+
+	/** GPU Buffers */
+	FSpatialHashGPUBuffers* SpatialHashGpuBuffers;
+};
+
+UCLASS(EditInlineNew, Category = "Spatial Hash", meta = (DisplayName = "Spatial Hash"))
+class NIAGARA_API UNiagaraDataInterfaceSpatialHash : public UNiagaraDataInterface
+{
+	GENERATED_UCLASS_BODY()
+public:
+
+	/** Maximum number of particles that can be stored in the spatial hash. */
+	UPROPERTY(EditAnywhere, Category = "SpatialHash")
+	uint32 MaximumParticleCount;
+
+	/** Size of the hash table. Make this a prime number larger than the number of particles stored for better performance. */
+	UPROPERTY(EditAnywhere, Category = "SpatialHash")
+	uint32 TableSize;
+
+	/** The maximum number of neighbors that will ever be searched for. */
+	UPROPERTY(EditAnywhere, Category = "SpatialHash")
+	uint32 MaximumNeighborCount;
+
+	/** The maximum search radius that neighbors will ever be searched in. This determines the cell size. */
+	UPROPERTY(EditAnywhere, Category = "SpatialHash")
+	float MaximumSearchRadius;
+
+	//UObject Interface
+	virtual void PostInitProperties()override;
+	virtual void PostLoad() override;
+#if WITH_EDITOR
+	virtual void PreEditChange(UProperty* PropertyAboutToChange) override;
+	virtual void PostEditChangeProperty(struct FPropertyChangedEvent& PropertyChangedEvent) override;
+#endif
+	//UObject Interface End
+
+	//UNiagaraDataInterface Interface
+	virtual bool InitPerInstanceData(void* PerInstanceData, FNiagaraSystemInstance* SystemInstance) override;
+	virtual void DestroyPerInstanceData(void* PerInstanceData, FNiagaraSystemInstance* SystemInstance) override;
+	virtual int32 PerInstanceDataSize() const { return sizeof(FNDISpatialHash_InstanceData); }
+
+	virtual bool PerInstanceTick(void* PerInstanceData, FNiagaraSystemInstance* SystemInstance, float DeltaSeconds);
+	virtual bool PerInstanceTickPostSimulate(void* PerInstanceData, FNiagaraSystemInstance* SystemInstance, float DeltaSeconds);
+
+	virtual void GetFunctions(TArray<FNiagaraFunctionSignature>& OutFunctions) override;
+	virtual void GetVMExternalFunction(const FVMExternalFunctionBindingInfo& BindingInfo, void* InstanceData, FVMExternalFunction &OutFunc) override;
+	virtual bool CanExecuteOnTarget(ENiagaraSimTarget Target)const override { return Target == ENiagaraSimTarget::CPUSim; }
+	virtual bool Equals(const UNiagaraDataInterface* Other) const override;
+
+	void AddParticle(FVectorVMContext& Context);
+	void BuildTable(FVectorVMContext& Context);
+	void PerformKNearestNeighborQuery(FVectorVMContext& Context);
+	void GetClosestNeighborFromQueryByIndex(FVectorVMContext& Context);
+	void Get16ClosestNeighborsFromQuery(FVectorVMContext& Context);
+
+	// GPU sim functionality
+	void GetCommonHLSL(FString& OutHLSL) override;
+	virtual bool GetFunctionHLSL(const FName& DefinitionFunctionName, FString InstanceFunctionName, FNiagaraDataInterfaceGPUParamInfo& ParamInfo, FString& OutHLSL) override;
+	virtual void GetParameterDefinitionHLSL(FNiagaraDataInterfaceGPUParamInfo& ParamInfo, FString& OutHLSL) override;
+	virtual FNiagaraDataInterfaceParametersCS* ConstructComputeParameters() const override;
+
+	virtual void ProvidePerInstanceDataForRenderThread(void* DataForRenderThread, void* PerInstanceData, const FGuid& SystemInstance) override;
+
+	virtual void PostExecute() override;
+
+	//UNiagaraDataInterface Interface End
+
+protected:
+	virtual bool CopyToInternal(UNiagaraDataInterface* Destination) const override;
+	void PushToRenderThread();
+
+private:
+	friend class FNDISPatialHash_InstanceData;
+	void BuildTableHelper();
+	void GetXClosestNeighborsFromQueryHelper(FNiagaraDINearestNeighborBatch* Batch, uint32 QueryID, uint32 NumberToRetrieve, TArray<FNiagaraID>& Neighbors);
+	
+	static FCriticalSection CriticalSection;
+
+public:
+	static const FString ParticleIDBufferName;
+	static const FString ParticlePosBufferName;
+	static const FString Built_ParticleIDBufferName;
+	static const FString Built_ParticlePosBufferName;
+	static const FString CellCountBufferName;
+	static const FString CellStartIndicesBufferName;
+	static const FString CellEndIndicesBufferName;
+
+	static const FString TableSizeName;
+	static const FString MaximumNeighborCountName;
+	static const FString MaximumSearchRadiusName;
+	static const FString NumParticlesName;
+	static const FString CellLengthName;
+
+	static const FString NearestNeighborResultsBufferName;
+	static const FString CurrentNNIDName;
+};
+
+class FSpatialHashGPUBuffers : public FRenderResource
+{
+public:
+	virtual ~FSpatialHashGPUBuffers() {}
+
+	void Initialize(FNDISpatialHash_InstanceData* InstanceData);
+
+	virtual void InitRHI() override;
+	virtual void ReleaseRHI() override;
+
+	virtual FString GetFriendlyName() const override { return TEXT("FSpatialHashGPUBuffers"); }
+
+	FRWBuffer& GetParticleIDs() { return ParticleID; }
+	FRWBuffer& GetParticlePos() { return ParticlePos; }
+	FRWBuffer& GetBuiltParticleIDs() { return Built_ParticleID; }
+	FRWBuffer& GetBuiltParticlePos() { return Built_ParticlePos; }
+	FRWBuffer& GetCellCount() { return CellCount; }
+	FRWBuffer& GetCellStartIndices() { return CellStartIndices; }
+	FRWBuffer& GetCellEndIndices() { return CellEndIndices; }
+
+	FRWBuffer& GetNumParticles() { return NumParticles; }
+
+	FRWBuffer& GetNearestNeighborResults() { return NearestNeighborResults; }
+	FRWBuffer& GetCurrentNNID() { return CurrentNNID; }
+
+	int MaximumParticleCount;
+	int TableSize;
+	int NumberOfParticles;
+	int MaximumNeighborCount;
+
+private:
+	FRWBuffer ParticleID;
+	FRWBuffer ParticlePos;
+	FRWBuffer Built_ParticleID;
+	FRWBuffer Built_ParticlePos;
+	FRWBuffer CellCount;
+	FRWBuffer CellStartIndices;
+	FRWBuffer CellEndIndices;
+
+	FRWBuffer NumParticles;
+
+	FRWBuffer NearestNeighborResults;
+	FRWBuffer CurrentNNID;
+};
+
+struct FNiagaraDISpatialHashPassedDataToRT
+{
+	FSpatialHashGPUBuffers* SpatialHashGpuBuffers;
+	uint32 MaximumParticleCount;
+	uint32 TableSize;
+	uint32 MaximumNeighborCount;
+	float MaximumSearchRadius;
+	uint32 NumParticles;
+	float CellLength;
+};
+
+struct FNiagaraDataInterfaceProxySpatialHash : public FNiagaraDataInterfaceProxy
+{
+	virtual int32 PerInstanceDataPassedToRenderThreadSize() const override
+	{
+		return sizeof(FNiagaraDISpatialHashPassedDataToRT);
+	}
+
+	virtual void ConsumePerInstanceDataFromGameThread(void* PerInstanceData, const FGuid& Instance) override;
+
+	TMap<FGuid, FNiagaraDISpatialHashPassedDataToRT> SystemInstancesToData;
+};
+
+struct FNiagaraDataInterfaceParametersCS_SpatialHash : public FNiagaraDataInterfaceParametersCS
+{
+	virtual ~FNiagaraDataInterfaceParametersCS_SpatialHash() {}
+	virtual void Bind(const FNiagaraDataInterfaceParamRef& ParamRef, const class FShaderParameterMap& ParameterMap);
+	virtual void Serialize(FArchive& Ar);
+	virtual void Set(FRHICommandList& RHICmdList, const FNiagaraDataInterfaceSetArgs& Context) const override;
+	virtual void Unset(FRHICommandList& RHICmdList, const FNiagaraDataInterfaceSetArgs& Context) const override;
+
+private:
+	FShaderResourceParameter ParticleIDBuffer;
+	FShaderResourceParameter ParticlePosBuffer;
+	FShaderResourceParameter Built_ParticleIDBuffer;
+	FShaderResourceParameter Built_ParticlePosBuffer;
+	FShaderResourceParameter CellCountBuffer;
+	FShaderResourceParameter CellStartIndicesBuffer;
+	FShaderResourceParameter CellEndIndicesBuffer;
+	FShaderResourceParameter NumParticles;
+	FShaderResourceParameter NearestNeighborResultsBuffer;
+	FShaderResourceParameter CurrentNNID;
+
+	FShaderParameter TableSize;
+	FShaderParameter MaximumNeighborCount;
+	FShaderParameter MaximumSearchRadius;
+	FShaderParameter CellLength;
+};
\ No newline at end of file
diff --git a/Engine/Plugins/FX/Niagara/Source/Niagara/Private/NiagaraDataInterfaceSpatialHash.cpp b/Engine/Plugins/FX/Niagara/Source/Niagara/Private/NiagaraDataInterfaceSpatialHash.cpp
new file mode 100644
--- /dev/null
+++ b/Engine/Plugins/FX/Niagara/Source/Niagara/Private/NiagaraDataInterfaceSpatialHash.cpp
@@ -0,0 +1,1589 @@
+// Copyright 1998-2019 Epic Games, Inc. All Rights Reserved.
+
+#include "NiagaraDataInterfaceSpatialHash.h"
+#include "ClearQuad.h"
+#include "NiagaraEmitterInstanceBatcher.h"
+#include "NiagaraSystemInstance.h"
+#include "ShaderParameterUtils.h"
+#include "NiagaraSpatialHashBuild.h"
+
+#define LOCTEXT_NAMESPACE "NiagaraDataInterfaceSpatialHash"
+
+// Cvar to enable/disable discovery of this data interface in the UI
+static int32 GbEnableSpatialHashDataInterface = 0;
+static FAutoConsoleVariableRef CVarEnableSpatialHashDataInterface(
+	TEXT("fx.EnableSpatialHashDataInterface"),
+	GbEnableSpatialHashDataInterface,
+	TEXT("If > 0 the spatial hash data interface will be accessible in the stack and module scripts.\n"),
+	ECVF_Default
+);
+
+void OnChangeEnableSpatialHashDataInterface(IConsoleVariable* Var)
+{
+	if (GbEnableSpatialHashDataInterface > 0)
+	{
+		FNiagaraTypeRegistry::Register(FNiagaraTypeDefinition(UNiagaraDataInterfaceSpatialHash::StaticClass()), true, false, false);
+	}
+	else
+	{
+		FNiagaraTypeRegistry::Deregister(FNiagaraTypeDefinition(UNiagaraDataInterfaceSpatialHash::StaticClass()));
+	}
+}
+
+static const FName AddParticleFunctionName("AddParticleToSpatialHash");
+static const FName PerformKNNQueryFunctionName("PerformKNearestNeighborQuery");
+
+static const FName GetClosestNeighborByIndexFunctionName("GetClosestNeighborFromQueryByIndex");
+static const FName GetClosestNeighborFunctionName("GetClosestNeighborFromQuery");
+static const FName Get16ClosestNeighborsFunctionName("Get 16 Closest Neighbors From Query");
+
+FCriticalSection UNiagaraDataInterfaceSpatialHash::CriticalSection;
+
+// Hash function
+// From http://matthias-mueller-fischer.ch/publications/tetraederCollision.pdf
+uint32 SpatialHash_HashFunction(FIntVector Position, uint32 TableSize)
+{
+	const uint32 p1 = 73856093;
+	const uint32 p2 = 19349663;
+	const uint32 p3 = 83492791;
+
+	uint32 n = (p1 * Position.X) ^ (p2 * Position.Y) ^ (p3 * Position.Z);
+	return n % TableSize;
+}
+
+FIntVector SpatialHash_GetCellIndex(FVector Position, float CellLength)
+{
+	auto FloorFVector = [](FVector InVector)
+	{
+		return FIntVector(FMath::FloorToInt(InVector.X), FMath::FloorToInt(InVector.Y), FMath::FloorToInt(InVector.Z));
+	};
+	return FloorFVector(Position / CellLength);
+}
+
+UNiagaraDataInterfaceSpatialHash::UNiagaraDataInterfaceSpatialHash(FObjectInitializer const& ObjectInitializer)
+	: Super(ObjectInitializer)
+	, MaximumParticleCount(500)
+	, TableSize(997)
+	, MaximumNeighborCount(32)
+	, MaximumSearchRadius(100.f)
+{
+	Proxy = MakeShared<FNiagaraDataInterfaceProxySpatialHash, ESPMode::ThreadSafe>();
+}
+
+void UNiagaraDataInterfaceSpatialHash::PostInitProperties()
+{
+	Super::PostInitProperties();
+
+	CVarEnableSpatialHashDataInterface->SetOnChangedCallback(FConsoleVariableDelegate::CreateStatic(&OnChangeEnableSpatialHashDataInterface));
+	if (HasAnyFlags(RF_ClassDefaultObject) && GbEnableSpatialHashDataInterface > 0)
+	{
+		FNiagaraTypeRegistry::Register(FNiagaraTypeDefinition(GetClass()), true, false, false);
+	}
+}
+
+void UNiagaraDataInterfaceSpatialHash::PostLoad()
+{
+	Super::PostLoad();
+}
+
+#if WITH_EDITOR
+
+void UNiagaraDataInterfaceSpatialHash::PreEditChange(UProperty* PropertyAboutToChange)
+{
+	Super::PreEditChange(PropertyAboutToChange);
+}
+
+void UNiagaraDataInterfaceSpatialHash::PostEditChangeProperty(struct FPropertyChangedEvent& PropertyChangedEvent)
+{
+	Super::PostEditChangeProperty(PropertyChangedEvent);
+}
+
+#endif
+
+bool UNiagaraDataInterfaceSpatialHash::InitPerInstanceData(void* InPerInstanceData, FNiagaraSystemInstance* InSystemInstance)
+{
+	FNDISpatialHash_InstanceData* PIData = new (InPerInstanceData) FNDISpatialHash_InstanceData;
+	PIData->Init(this, InSystemInstance);
+	return true;
+}
+
+void UNiagaraDataInterfaceSpatialHash::DestroyPerInstanceData(void* InPerInstanceData, FNiagaraSystemInstance* SystemInstance)
+{
+	FNDISpatialHash_InstanceData* PIData = (FNDISpatialHash_InstanceData*) InPerInstanceData;
+	PIData->Release();
+	PIData->~FNDISpatialHash_InstanceData();
+
+	if (PIData->SpatialHashGpuBuffers) {
+		FNiagaraDataInterfaceProxySpatialHash* ThisProxy = GetProxyAs<FNiagaraDataInterfaceProxySpatialHash>();
+		ENQUEUE_RENDER_COMMAND(FNiagaraDestroySpatialHashInstanceData)(
+			[ThisProxy, InstanceID = SystemInstance->GetId()](FRHICommandListImmediate& CmdList)
+			{
+				ThisProxy->SystemInstancesToData.Remove(InstanceID);
+			}
+		);
+	}
+}
+
+bool UNiagaraDataInterfaceSpatialHash::PerInstanceTick(void* InPerInstanceData, FNiagaraSystemInstance* SystemInstance, float DeltaSeconds)
+{
+	FNDISpatialHash_InstanceData* PIData = (FNDISpatialHash_InstanceData*) InPerInstanceData;
+	
+	return false;
+}
+
+bool UNiagaraDataInterfaceSpatialHash::PerInstanceTickPostSimulate(void* InPerInstanceData, FNiagaraSystemInstance* SystemInstance, float DeltaSeconds)
+{
+	FNDISpatialHash_InstanceData* PIData = static_cast<FNDISpatialHash_InstanceData*>(InPerInstanceData);
+	if (PIData->SpatialHashGpuBuffers == nullptr)
+	{
+		PIData->SpatialHashBatch.ClearWrite();
+		PIData->BuildTable();
+	}
+	else
+	{
+		FSpatialHashGPUBuffers* Buffers = PIData->SpatialHashGpuBuffers;
+		PIData->BuildTableGPU();
+	}
+	return false;
+}
+
+void UNiagaraDataInterfaceSpatialHash::GetFunctions(TArray<FNiagaraFunctionSignature>& OutFunctions)
+{
+	Super::GetFunctions(OutFunctions);
+
+	{
+		FNiagaraFunctionSignature Sig;
+		Sig.Name = AddParticleFunctionName;
+		Sig.Inputs.Add(FNiagaraVariable(FNiagaraTypeDefinition(GetClass()), TEXT("Spatial Hash")));
+		Sig.Inputs.Add(FNiagaraVariable(FNiagaraTypeDefinition::GetIDDef(), TEXT("ParticleID")));
+		Sig.Inputs.Add(FNiagaraVariable(FNiagaraTypeDefinition::GetVec3Def(), TEXT("ParticlePosition")));
+
+		Sig.Outputs.Add(FNiagaraVariable(FNiagaraTypeDefinition::GetIntDef(), TEXT("NumParticles")));
+
+		Sig.SetDescription(LOCTEXT("SpatialHashAddParticleDesc", "Adds a particle with ID ParticleID and position ParticlePosition"
+			" to the spatial hash structure. This does not build or update the structure.\n"
+			"Building occurs after all particle update scripts have run for the current frame.\n"
+			"Call this function in the particle update script so that the particle will remain in the structure across different frames."));
+
+		Sig.bMemberFunction = true;
+		Sig.bRequiresContext = false;
+		OutFunctions.Add(Sig);
+	}
+
+	{
+		FNiagaraFunctionSignature Sig;
+		Sig.Name = PerformKNNQueryFunctionName;
+		Sig.Inputs.Add(FNiagaraVariable(FNiagaraTypeDefinition(GetClass()), TEXT("Spatial Hash")));
+
+		Sig.Inputs.Add(FNiagaraVariable(FNiagaraTypeDefinition::GetIDDef(), TEXT("ParticleID")));
+		Sig.Inputs.Add(FNiagaraVariable(FNiagaraTypeDefinition::GetVec3Def(), TEXT("Position")));
+		Sig.Inputs.Add(FNiagaraVariable(FNiagaraTypeDefinition::GetFloatDef(), TEXT("SearchRadius")));
+		Sig.Inputs.Add(FNiagaraVariable(FNiagaraTypeDefinition::GetIntDef(), TEXT("MaximumNeighbors")));
+		Sig.Inputs.Add(FNiagaraVariable(FNiagaraTypeDefinition::GetBoolDef(), TEXT("IncludeSelf")));
+
+		Sig.Outputs.Add(FNiagaraVariable(FNiagaraTypeDefinition::GetIntDef(), TEXT("QueryResultID")));
+		Sig.Outputs.Add(FNiagaraVariable(FNiagaraTypeDefinition::GetIntDef(), TEXT("NumberOfNeighbors")));
+
+		Sig.SetDescription(LOCTEXT("SpatialHashPerformKNNQueryDesc", "Performs a k-nearest neighbor query. Requires a position, a search radius,"
+			" and a maximum number of neighbors to find.\n"
+			"Returns a query ID used to read the results of the query, and a count of how many neighbors were found that satisfy the conditions.\n"
+			"If IncludeSelf is set to true, then the ParticleID parameter could be a possible nearest neighbor result."));
+
+		Sig.bMemberFunction = true;
+		Sig.bRequiresContext = false;
+		OutFunctions.Add(Sig);
+	}
+
+	{
+		FNiagaraFunctionSignature Sig;
+		Sig.Name = GetClosestNeighborByIndexFunctionName;
+		Sig.Inputs.Add(FNiagaraVariable(FNiagaraTypeDefinition(GetClass()), TEXT("Spatial Hash")));
+
+		Sig.Inputs.Add(FNiagaraVariable(FNiagaraTypeDefinition::GetIntDef(), TEXT("Index")));
+		Sig.Inputs.Add(FNiagaraVariable(FNiagaraTypeDefinition::GetIntDef(), TEXT("QueryID")));
+
+		Sig.Outputs.Add(FNiagaraVariable(FNiagaraTypeDefinition::GetBoolDef(), TEXT("Valid")));
+		Sig.Outputs.Add(FNiagaraVariable(FNiagaraTypeDefinition::GetIDDef(), TEXT("Neighbor")));
+
+		Sig.SetDescription(LOCTEXT("SpatialHashGetClosestNeighborByIndexDesc", "Reads a k-nearest neighbor query and returns the i-th closest neighbor.\n"
+			"For example, if the input index is 1, it returns the closest neighbor. If the input is 4, it returns the 4th closest neighbor.\n"
+			"If the requested neighbor does not exist, the Valid flag is set to false."));
+		Sig.bMemberFunction = true;
+		Sig.bRequiresContext = false;
+		OutFunctions.Add(Sig);
+	}
+	{
+		FNiagaraFunctionSignature Sig;
+		Sig.Name = Get16ClosestNeighborsFunctionName;
+		Sig.Inputs.Add(FNiagaraVariable(FNiagaraTypeDefinition(GetClass()), TEXT("Spatial Hash")));
+
+		Sig.Inputs.Add(FNiagaraVariable(FNiagaraTypeDefinition::GetIntDef(), TEXT("QueryID")));
+
+		Sig.Outputs.Add(FNiagaraVariable(FNiagaraTypeDefinition::GetIDDef(), TEXT("Closest Neighbor 1")));
+		Sig.Outputs.Add(FNiagaraVariable(FNiagaraTypeDefinition::GetIDDef(), TEXT("Closest Neighbor 2")));
+		Sig.Outputs.Add(FNiagaraVariable(FNiagaraTypeDefinition::GetIDDef(), TEXT("Closest Neighbor 3")));
+		Sig.Outputs.Add(FNiagaraVariable(FNiagaraTypeDefinition::GetIDDef(), TEXT("Closest Neighbor 4")));
+		Sig.Outputs.Add(FNiagaraVariable(FNiagaraTypeDefinition::GetIDDef(), TEXT("Closest Neighbor 5")));
+		Sig.Outputs.Add(FNiagaraVariable(FNiagaraTypeDefinition::GetIDDef(), TEXT("Closest Neighbor 6")));
+		Sig.Outputs.Add(FNiagaraVariable(FNiagaraTypeDefinition::GetIDDef(), TEXT("Closest Neighbor 7")));
+		Sig.Outputs.Add(FNiagaraVariable(FNiagaraTypeDefinition::GetIDDef(), TEXT("Closest Neighbor 8")));
+		Sig.Outputs.Add(FNiagaraVariable(FNiagaraTypeDefinition::GetIDDef(), TEXT("Closest Neighbor 9")));
+		Sig.Outputs.Add(FNiagaraVariable(FNiagaraTypeDefinition::GetIDDef(), TEXT("Closest Neighbor 10")));
+		Sig.Outputs.Add(FNiagaraVariable(FNiagaraTypeDefinition::GetIDDef(), TEXT("Closest Neighbor 11")));
+		Sig.Outputs.Add(FNiagaraVariable(FNiagaraTypeDefinition::GetIDDef(), TEXT("Closest Neighbor 12")));
+		Sig.Outputs.Add(FNiagaraVariable(FNiagaraTypeDefinition::GetIDDef(), TEXT("Closest Neighbor 13")));
+		Sig.Outputs.Add(FNiagaraVariable(FNiagaraTypeDefinition::GetIDDef(), TEXT("Closest Neighbor 14")));
+		Sig.Outputs.Add(FNiagaraVariable(FNiagaraTypeDefinition::GetIDDef(), TEXT("Closest Neighbor 15")));
+		Sig.Outputs.Add(FNiagaraVariable(FNiagaraTypeDefinition::GetIDDef(), TEXT("Closest Neighbor 16")));
+
+		Sig.SetDescription(LOCTEXT("SpatialHashGet16ClosestNeighborsDesc", "Reads a k-nearest neighbor query and returns the 16 closest neighbors.\n"
+			"If there are fewer than 16 neighbors, the corresponding output ID will be set to NIAGARA_INVALID_ID (Index: -1, AcquireTag: -1)"));
+
+		Sig.bMemberFunction = true;
+		Sig.bRequiresContext = false;
+		OutFunctions.Add(Sig);
+	}
+}
+
+DEFINE_NDI_DIRECT_FUNC_BINDER(UNiagaraDataInterfaceSpatialHash, AddParticle);
+DEFINE_NDI_DIRECT_FUNC_BINDER(UNiagaraDataInterfaceSpatialHash, PerformKNearestNeighborQuery);
+DEFINE_NDI_DIRECT_FUNC_BINDER(UNiagaraDataInterfaceSpatialHash, GetClosestNeighborFromQueryByIndex);
+DEFINE_NDI_DIRECT_FUNC_BINDER(UNiagaraDataInterfaceSpatialHash, Get16ClosestNeighborsFromQuery);
+void UNiagaraDataInterfaceSpatialHash::GetVMExternalFunction(const FVMExternalFunctionBindingInfo& BindingInfo, void* InstanceData, FVMExternalFunction &OutFunc)
+{
+	if (BindingInfo.Name == AddParticleFunctionName)
+	{
+		NDI_FUNC_BINDER(UNiagaraDataInterfaceSpatialHash, AddParticle)::Bind(this, OutFunc);
+	}
+	else if (BindingInfo.Name == PerformKNNQueryFunctionName)
+	{
+		NDI_FUNC_BINDER(UNiagaraDataInterfaceSpatialHash, PerformKNearestNeighborQuery)::Bind(this, OutFunc);
+	}
+	else if (BindingInfo.Name == GetClosestNeighborByIndexFunctionName)
+	{
+		NDI_FUNC_BINDER(UNiagaraDataInterfaceSpatialHash, GetClosestNeighborFromQueryByIndex)::Bind(this, OutFunc);
+	}
+	else if (BindingInfo.Name == Get16ClosestNeighborsFunctionName)
+	{
+		NDI_FUNC_BINDER(UNiagaraDataInterfaceSpatialHash, Get16ClosestNeighborsFromQuery)::Bind(this, OutFunc);
+	}
+}
+
+bool UNiagaraDataInterfaceSpatialHash::Equals(const UNiagaraDataInterface* Other) const
+{
+	if (!Super::Equals(Other))
+	{
+		return false;
+	}
+
+	bool bMaxParticleCountEq = CastChecked<UNiagaraDataInterfaceSpatialHash>(Other)->MaximumParticleCount == MaximumParticleCount;
+	bool bTableSizeEq = CastChecked<UNiagaraDataInterfaceSpatialHash>(Other)->TableSize == TableSize;
+	bool bMaxNeighborCountEq = CastChecked<UNiagaraDataInterfaceSpatialHash>(Other)->MaximumNeighborCount == MaximumNeighborCount;
+	bool bMaxSearchRadEq = CastChecked<UNiagaraDataInterfaceSpatialHash>(Other)->MaximumSearchRadius == MaximumSearchRadius;
+
+	return bMaxParticleCountEq && bTableSizeEq && bMaxNeighborCountEq && bMaxSearchRadEq;
+}
+
+void UNiagaraDataInterfaceSpatialHash::AddParticle(FVectorVMContext& Context)
+{
+	VectorVM::FExternalFuncInputHandler<int32> ParticleIDIndexParam(Context);
+	VectorVM::FExternalFuncInputHandler<int32> ParticleIDAcquireTagParam(Context);
+	VectorVM::FExternalFuncInputHandler<float> XParticlePositionParam(Context);
+	VectorVM::FExternalFuncInputHandler<float> YParticlePositionParam(Context);
+	VectorVM::FExternalFuncInputHandler<float> ZParticlePositionParam(Context);
+
+	VectorVM::FUserPtrHandler<FNDISpatialHash_InstanceData> InstanceData(Context);
+
+	VectorVM::FExternalFuncRegisterHandler<uint32> OutNumParticles(Context);
+
+	FScopeLock ScopeLock(&CriticalSection);
+	for (int32 InstanceIdx = 0; InstanceIdx < Context.NumInstances; ++InstanceIdx)
+	{
+		FNiagaraID ParticleID = { ParticleIDIndexParam.GetAndAdvance(), ParticleIDAcquireTagParam.GetAndAdvance() };
+		float XParticlePosition = XParticlePositionParam.GetAndAdvance();
+		float YParticlePosition = YParticlePositionParam.GetAndAdvance();
+		float ZParticlePosition = ZParticlePositionParam.GetAndAdvance();
+		FVector ParticlePosition(XParticlePosition, YParticlePosition, ZParticlePosition);
+
+		uint32 ParticleCellHash = SpatialHash_HashFunction(SpatialHash_GetCellIndex(ParticlePosition, InstanceData->CellLength), TableSize);
+
+		InstanceData->Particles.Add({ ParticleCellHash, static_cast<uint32>(InstanceData->Particles.Num()), ParticlePosition, ParticleID });
+
+		*OutNumParticles.GetDestAndAdvance() = (InstanceData->NumParticles)++;
+	}
+}
+
+void UNiagaraDataInterfaceSpatialHash::BuildTable(FVectorVMContext& Context)
+{
+	VectorVM::FUserPtrHandler<FNDISpatialHash_InstanceData> InstanceData(Context);
+
+	VectorVM::FExternalFuncRegisterHandler<uint32> OutNumParticles(Context);
+
+	BuildTableHelper();
+
+	*OutNumParticles.GetDestAndAdvance() = InstanceData->Particles_Built.Num();
+}
+
+void UNiagaraDataInterfaceSpatialHash::PerformKNearestNeighborQuery(FVectorVMContext& Context)
+{
+	VectorVM::FExternalFuncInputHandler<int32> ParticleIDIndexParam(Context);
+	VectorVM::FExternalFuncInputHandler<int32> ParticleIDAcquireTagParam(Context);
+	VectorVM::FExternalFuncInputHandler<float> XPositionParam(Context);
+	VectorVM::FExternalFuncInputHandler<float> YPositionParam(Context);
+	VectorVM::FExternalFuncInputHandler<float> ZPositionParam(Context);
+	VectorVM::FExternalFuncInputHandler<float> SearchRadiusParam(Context);
+	VectorVM::FExternalFuncInputHandler<uint32> MaximumNeighborsParam(Context);
+	VectorVM::FExternalFuncInputHandler<FNiagaraBool> IncludeSelfParam(Context);
+
+	VectorVM::FUserPtrHandler<FNDISpatialHash_InstanceData> InstanceData(Context);
+
+	VectorVM::FExternalFuncRegisterHandler<int32> OutQueryResultID(Context);
+	VectorVM::FExternalFuncRegisterHandler<int32> OutNumNeighbors(Context);
+
+	FScopeLock ScopeLock(&CriticalSection);
+	for (int32 InstanceIdx = 0; InstanceIdx < Context.NumInstances; ++InstanceIdx)
+	{
+		FNiagaraID ParticleID = { ParticleIDIndexParam.GetAndAdvance(), ParticleIDAcquireTagParam.GetAndAdvance() };
+		float XPosition = XPositionParam.GetAndAdvance();
+		float YPosition = YPositionParam.GetAndAdvance();
+		float ZPosition = ZPositionParam.GetAndAdvance();
+		float SearchRadius = SearchRadiusParam.GetAndAdvance();
+		uint32 MaximumNeighbors = MaximumNeighborsParam.GetAndAdvance();
+		FNiagaraBool bIncludeSelf = IncludeSelfParam.GetAndAdvance();
+		FVector Position(XPosition, YPosition, ZPosition);
+
+		uint32 QueryResultID = InstanceData->SpatialHashBatch.SubmitQuery(ParticleID, Position, SearchRadius, MaximumNeighbors, bIncludeSelf);
+
+		*OutQueryResultID.GetDestAndAdvance() = QueryResultID;
+
+		// Retrieve query result
+		TArray<FNiagaraID> QueryResult;
+		InstanceData->SpatialHashBatch.GetQueryResult(QueryResultID, QueryResult);
+		*OutNumNeighbors.GetDestAndAdvance() = QueryResult.Num();
+	}
+}
+
+void UNiagaraDataInterfaceSpatialHash::GetClosestNeighborFromQueryByIndex(FVectorVMContext& Context)
+{
+	VectorVM::FExternalFuncInputHandler<uint32> IndexParam(Context);
+	VectorVM::FExternalFuncInputHandler<uint32> QueryID(Context);
+
+	VectorVM::FUserPtrHandler<FNDISpatialHash_InstanceData> InstanceData(Context);
+
+	VectorVM::FExternalFuncRegisterHandler<FNiagaraBool> OutValid(Context);
+	VectorVM::FExternalFuncRegisterHandler<int32> OutNeighborIDIndex(Context);
+	VectorVM::FExternalFuncRegisterHandler<int32> OutNeighborIDAcquireTag(Context);
+
+	FScopeLock ScopeLock(&CriticalSection);
+	for (int32 InstanceIdx = 0; InstanceIdx < Context.NumInstances; ++InstanceIdx)
+	{
+		uint32 Index = IndexParam.GetAndAdvance();
+		TArray<FNiagaraID> QueryResult;
+		GetXClosestNeighborsFromQueryHelper(&(InstanceData->SpatialHashBatch), QueryID.GetAndAdvance(), Index, QueryResult);
+		FNiagaraID NeighborID = QueryResult.Num() >= static_cast<int32>(Index) ? QueryResult[Index - 1] : NIAGARA_INVALID_ID;
+		FNiagaraBool ValidValue;
+		ValidValue.SetValue(NeighborID != NIAGARA_INVALID_ID);
+		*OutValid.GetDestAndAdvance() = ValidValue;
+		*OutNeighborIDIndex.GetDestAndAdvance() = NeighborID.Index;
+		*OutNeighborIDAcquireTag.GetDestAndAdvance() = NeighborID.AcquireTag;
+	}
+}
+
+void UNiagaraDataInterfaceSpatialHash::Get16ClosestNeighborsFromQuery(FVectorVMContext& Context)
+{
+	VectorVM::FExternalFuncInputHandler<uint32> QueryID(Context);
+
+	VectorVM::FUserPtrHandler<FNDISpatialHash_InstanceData> InstanceData(Context);
+
+	VectorVM::FExternalFuncRegisterHandler<int32> OutNeighborIDIndex_01(Context);
+	VectorVM::FExternalFuncRegisterHandler<int32> OutNeighborIDAcquireTag_01(Context);
+	VectorVM::FExternalFuncRegisterHandler<int32> OutNeighborIDIndex_02(Context);
+	VectorVM::FExternalFuncRegisterHandler<int32> OutNeighborIDAcquireTag_02(Context);
+	VectorVM::FExternalFuncRegisterHandler<int32> OutNeighborIDIndex_03(Context);
+	VectorVM::FExternalFuncRegisterHandler<int32> OutNeighborIDAcquireTag_03(Context);
+	VectorVM::FExternalFuncRegisterHandler<int32> OutNeighborIDIndex_04(Context);
+	VectorVM::FExternalFuncRegisterHandler<int32> OutNeighborIDAcquireTag_04(Context);
+	VectorVM::FExternalFuncRegisterHandler<int32> OutNeighborIDIndex_05(Context);
+	VectorVM::FExternalFuncRegisterHandler<int32> OutNeighborIDAcquireTag_05(Context);
+	VectorVM::FExternalFuncRegisterHandler<int32> OutNeighborIDIndex_06(Context);
+	VectorVM::FExternalFuncRegisterHandler<int32> OutNeighborIDAcquireTag_06(Context);
+	VectorVM::FExternalFuncRegisterHandler<int32> OutNeighborIDIndex_07(Context);
+	VectorVM::FExternalFuncRegisterHandler<int32> OutNeighborIDAcquireTag_07(Context);
+	VectorVM::FExternalFuncRegisterHandler<int32> OutNeighborIDIndex_08(Context);
+	VectorVM::FExternalFuncRegisterHandler<int32> OutNeighborIDAcquireTag_08(Context);
+	VectorVM::FExternalFuncRegisterHandler<int32> OutNeighborIDIndex_09(Context);
+	VectorVM::FExternalFuncRegisterHandler<int32> OutNeighborIDAcquireTag_09(Context);
+	VectorVM::FExternalFuncRegisterHandler<int32> OutNeighborIDIndex_10(Context);
+	VectorVM::FExternalFuncRegisterHandler<int32> OutNeighborIDAcquireTag_10(Context);
+	VectorVM::FExternalFuncRegisterHandler<int32> OutNeighborIDIndex_11(Context);
+	VectorVM::FExternalFuncRegisterHandler<int32> OutNeighborIDAcquireTag_11(Context);
+	VectorVM::FExternalFuncRegisterHandler<int32> OutNeighborIDIndex_12(Context);
+	VectorVM::FExternalFuncRegisterHandler<int32> OutNeighborIDAcquireTag_12(Context);
+	VectorVM::FExternalFuncRegisterHandler<int32> OutNeighborIDIndex_13(Context);
+	VectorVM::FExternalFuncRegisterHandler<int32> OutNeighborIDAcquireTag_13(Context);
+	VectorVM::FExternalFuncRegisterHandler<int32> OutNeighborIDIndex_14(Context);
+	VectorVM::FExternalFuncRegisterHandler<int32> OutNeighborIDAcquireTag_14(Context);
+	VectorVM::FExternalFuncRegisterHandler<int32> OutNeighborIDIndex_15(Context);
+	VectorVM::FExternalFuncRegisterHandler<int32> OutNeighborIDAcquireTag_15(Context);
+	VectorVM::FExternalFuncRegisterHandler<int32> OutNeighborIDIndex_16(Context);
+	VectorVM::FExternalFuncRegisterHandler<int32> OutNeighborIDAcquireTag_16(Context);
+
+	FScopeLock ScopeLock(&CriticalSection);
+	for (int32 InstanceIdx = 0; InstanceIdx < Context.NumInstances; ++InstanceIdx)
+	{
+		TArray<FNiagaraID> QueryResult;
+		GetXClosestNeighborsFromQueryHelper(&(InstanceData->SpatialHashBatch), QueryID.GetAndAdvance(), 16, QueryResult);
+
+		*OutNeighborIDIndex_01.GetDestAndAdvance() = QueryResult[0].Index;
+		*OutNeighborIDIndex_02.GetDestAndAdvance() = QueryResult.Num() > 1 ? QueryResult[1].Index : INDEX_NONE;
+		*OutNeighborIDIndex_03.GetDestAndAdvance() = QueryResult.Num() > 2 ? QueryResult[2].Index : INDEX_NONE;
+		*OutNeighborIDIndex_04.GetDestAndAdvance() = QueryResult.Num() > 3 ? QueryResult[3].Index : INDEX_NONE;
+		*OutNeighborIDIndex_05.GetDestAndAdvance() = QueryResult.Num() > 4 ? QueryResult[4].Index : INDEX_NONE;
+		*OutNeighborIDIndex_06.GetDestAndAdvance() = QueryResult.Num() > 5 ? QueryResult[5].Index : INDEX_NONE;
+		*OutNeighborIDIndex_07.GetDestAndAdvance() = QueryResult.Num() > 6 ? QueryResult[6].Index : INDEX_NONE;
+		*OutNeighborIDIndex_08.GetDestAndAdvance() = QueryResult.Num() > 7 ? QueryResult[7].Index : INDEX_NONE;
+		*OutNeighborIDIndex_09.GetDestAndAdvance() = QueryResult.Num() > 8 ? QueryResult[8].Index : INDEX_NONE;
+		*OutNeighborIDIndex_10.GetDestAndAdvance() = QueryResult.Num() > 9 ? QueryResult[9].Index : INDEX_NONE;
+		*OutNeighborIDIndex_11.GetDestAndAdvance() = QueryResult.Num() > 10 ? QueryResult[10].Index : INDEX_NONE;
+		*OutNeighborIDIndex_12.GetDestAndAdvance() = QueryResult.Num() > 11 ? QueryResult[11].Index : INDEX_NONE;
+		*OutNeighborIDIndex_13.GetDestAndAdvance() = QueryResult.Num() > 12 ? QueryResult[12].Index : INDEX_NONE;
+		*OutNeighborIDIndex_14.GetDestAndAdvance() = QueryResult.Num() > 13 ? QueryResult[13].Index : INDEX_NONE;
+		*OutNeighborIDIndex_15.GetDestAndAdvance() = QueryResult.Num() > 14 ? QueryResult[14].Index : INDEX_NONE;
+		*OutNeighborIDIndex_16.GetDestAndAdvance() = QueryResult.Num() > 15 ? QueryResult[15].Index : INDEX_NONE;
+
+		*OutNeighborIDAcquireTag_01.GetDestAndAdvance() = QueryResult[0].AcquireTag;
+		*OutNeighborIDAcquireTag_02.GetDestAndAdvance() = QueryResult.Num() > 1 ? QueryResult[1].AcquireTag : INDEX_NONE;
+		*OutNeighborIDAcquireTag_03.GetDestAndAdvance() = QueryResult.Num() > 2 ? QueryResult[2].AcquireTag : INDEX_NONE;
+		*OutNeighborIDAcquireTag_04.GetDestAndAdvance() = QueryResult.Num() > 3 ? QueryResult[3].AcquireTag : INDEX_NONE;
+		*OutNeighborIDAcquireTag_05.GetDestAndAdvance() = QueryResult.Num() > 4 ? QueryResult[4].AcquireTag : INDEX_NONE;
+		*OutNeighborIDAcquireTag_06.GetDestAndAdvance() = QueryResult.Num() > 5 ? QueryResult[5].AcquireTag : INDEX_NONE;
+		*OutNeighborIDAcquireTag_07.GetDestAndAdvance() = QueryResult.Num() > 6 ? QueryResult[6].AcquireTag : INDEX_NONE;
+		*OutNeighborIDAcquireTag_08.GetDestAndAdvance() = QueryResult.Num() > 7 ? QueryResult[7].AcquireTag : INDEX_NONE;
+		*OutNeighborIDAcquireTag_09.GetDestAndAdvance() = QueryResult.Num() > 8 ? QueryResult[8].AcquireTag : INDEX_NONE;
+		*OutNeighborIDAcquireTag_10.GetDestAndAdvance() = QueryResult.Num() > 9 ? QueryResult[9].AcquireTag : INDEX_NONE;
+		*OutNeighborIDAcquireTag_11.GetDestAndAdvance() = QueryResult.Num() > 10 ? QueryResult[10].AcquireTag : INDEX_NONE;
+		*OutNeighborIDAcquireTag_12.GetDestAndAdvance() = QueryResult.Num() > 11 ? QueryResult[11].AcquireTag : INDEX_NONE;
+		*OutNeighborIDAcquireTag_13.GetDestAndAdvance() = QueryResult.Num() > 12 ? QueryResult[12].AcquireTag : INDEX_NONE;
+		*OutNeighborIDAcquireTag_14.GetDestAndAdvance() = QueryResult.Num() > 13 ? QueryResult[13].AcquireTag : INDEX_NONE;
+		*OutNeighborIDAcquireTag_15.GetDestAndAdvance() = QueryResult.Num() > 14 ? QueryResult[14].AcquireTag : INDEX_NONE;
+		*OutNeighborIDAcquireTag_16.GetDestAndAdvance() = QueryResult.Num() > 15 ? QueryResult[15].AcquireTag : INDEX_NONE;
+	}
+}
+
+const FString UNiagaraDataInterfaceSpatialHash::ParticleIDBufferName(TEXT("ParticleIDs_"));
+const FString UNiagaraDataInterfaceSpatialHash::ParticlePosBufferName(TEXT("ParticlePosition_"));
+const FString UNiagaraDataInterfaceSpatialHash::Built_ParticleIDBufferName(TEXT("Built_ParticleIDs_"));
+const FString UNiagaraDataInterfaceSpatialHash::Built_ParticlePosBufferName(TEXT("Built_ParticlePosition_"));
+const FString UNiagaraDataInterfaceSpatialHash::CellCountBufferName(TEXT("CellCount_"));
+const FString UNiagaraDataInterfaceSpatialHash::CellStartIndicesBufferName(TEXT("CellStartIndices_"));
+const FString UNiagaraDataInterfaceSpatialHash::CellEndIndicesBufferName(TEXT("CellEndIndices_"));
+const FString UNiagaraDataInterfaceSpatialHash::TableSizeName(TEXT("TableSize_"));
+const FString UNiagaraDataInterfaceSpatialHash::MaximumNeighborCountName(TEXT("MaximumNeighborCount_"));
+const FString UNiagaraDataInterfaceSpatialHash::MaximumSearchRadiusName(TEXT("MaximumSearchRadius_"));
+const FString UNiagaraDataInterfaceSpatialHash::NumParticlesName(TEXT("NumParticles_"));
+const FString UNiagaraDataInterfaceSpatialHash::CellLengthName(TEXT("CellLength_"));
+const FString UNiagaraDataInterfaceSpatialHash::NearestNeighborResultsBufferName(TEXT("NearestNeighborResults_"));
+const FString UNiagaraDataInterfaceSpatialHash::CurrentNNIDName(TEXT("CurrentID_"));
+
+struct FNDISpatialHashParametersName
+{
+	FString ParticleIDBufferName;
+	FString ParticlePosBufferName;
+	FString Built_ParticleIDBufferName;
+	FString Built_ParticlePosBufferName;
+	FString CellCountBufferName;
+	FString CellStartIndicesBufferName;
+	FString CellEndIndicesBufferName;
+	FString TableSizeName;
+	FString MaximumNeighborCountName;
+	FString MaximumSearchRadiusName;
+	FString NumParticlesName;
+	FString CellLengthName;
+	FString NearestNeighborResultsName;
+	FString CurrentNNIDName;
+};
+
+static void GetNiagaraDataInterfaceParametersName(FNDISpatialHashParametersName& Names, const FString& Suffix)
+{
+	Names.ParticleIDBufferName = UNiagaraDataInterfaceSpatialHash::ParticleIDBufferName + Suffix;
+	Names.ParticlePosBufferName = UNiagaraDataInterfaceSpatialHash::ParticlePosBufferName + Suffix;
+	Names.Built_ParticleIDBufferName = UNiagaraDataInterfaceSpatialHash::Built_ParticleIDBufferName + Suffix;
+	Names.Built_ParticlePosBufferName = UNiagaraDataInterfaceSpatialHash::Built_ParticlePosBufferName + Suffix;
+	Names.CellCountBufferName = UNiagaraDataInterfaceSpatialHash::CellCountBufferName + Suffix;
+	Names.CellStartIndicesBufferName = UNiagaraDataInterfaceSpatialHash::CellStartIndicesBufferName + Suffix;
+	Names.CellEndIndicesBufferName = UNiagaraDataInterfaceSpatialHash::CellEndIndicesBufferName + Suffix;
+	Names.TableSizeName = UNiagaraDataInterfaceSpatialHash::TableSizeName + Suffix;
+	Names.MaximumNeighborCountName = UNiagaraDataInterfaceSpatialHash::MaximumNeighborCountName + Suffix;
+	Names.MaximumSearchRadiusName = UNiagaraDataInterfaceSpatialHash::MaximumSearchRadiusName + Suffix;
+	Names.NumParticlesName = UNiagaraDataInterfaceSpatialHash::NumParticlesName + Suffix;
+	Names.CellLengthName = UNiagaraDataInterfaceSpatialHash::CellLengthName + Suffix;
+	Names.NearestNeighborResultsName = UNiagaraDataInterfaceSpatialHash::NearestNeighborResultsBufferName + Suffix;
+	Names.CurrentNNIDName = UNiagaraDataInterfaceSpatialHash::CurrentNNIDName + Suffix;
+}
+
+// GPU sim functionality
+void UNiagaraDataInterfaceSpatialHash::GetCommonHLSL(FString& OutHLSL)
+{
+	OutHLSL += TEXT("#include \"/Plugin/FX/Niagara/Private/NiagaraDataInterfaceSpatialHash.ush\"\n");
+}
+
+bool UNiagaraDataInterfaceSpatialHash::GetFunctionHLSL(const FName& DefinitionFunctionName, FString InstanceFunctionName, FNiagaraDataInterfaceGPUParamInfo& ParamInfo, FString& OutHLSL)
+{
+	FNDISpatialHashParametersName ParamNames;
+	GetNiagaraDataInterfaceParametersName(ParamNames, ParamInfo.DataInterfaceHLSLSymbol);
+	TMap<FString, FStringFormatArg> ArgsSample = {
+		{ TEXT("InstanceFunctionName"), InstanceFunctionName },
+		{ TEXT("ParticleIDBufferName"), ParamNames.ParticleIDBufferName },
+		{ TEXT("ParticlePosBufferName"), ParamNames.ParticlePosBufferName },
+		{ TEXT("Built_ParticleIDBufferName"), ParamNames.Built_ParticleIDBufferName },
+		{ TEXT("Built_ParticlePosBufferName"), ParamNames.Built_ParticlePosBufferName },
+		{ TEXT("CellCountBufferName"), ParamNames.CellCountBufferName },
+		{ TEXT("CellStartIndicesBufferName"), ParamNames.CellStartIndicesBufferName },
+		{ TEXT("CellEndIndicesBufferName"), ParamNames.CellEndIndicesBufferName },
+		{ TEXT("TableSizeName"), ParamNames.TableSizeName },
+		{ TEXT("MaximumNeighborCountName"), ParamNames.MaximumNeighborCountName },
+		{ TEXT("MaximumSearchRadiusName"), ParamNames.MaximumSearchRadiusName },
+		{ TEXT("NumParticlesName"), ParamNames.NumParticlesName },
+		{ TEXT("CellLengthName"), ParamNames.CellLengthName },
+		{ TEXT("NearestNeighborResultsName"), ParamNames.NearestNeighborResultsName },
+		{ TEXT("CurrentNNIDName"), ParamNames.CurrentNNIDName }
+	};
+
+	if (DefinitionFunctionName == AddParticleFunctionName)
+	{
+		static const TCHAR* FormatFunction = TEXT(R"(
+			void {InstanceFunctionName}(in NiagaraID InParticleID, in float3 InParticlePosition, out int OutNumParticles)
+			{
+				int WriteIndex;
+				InterlockedAdd({NumParticlesName}[0], 1, WriteIndex);
+				uint ParticleCellHash = SpatialHash_HashFunction(SpatialHash_GetCellIndex(InParticlePosition, {CellLengthName}), {TableSizeName});
+				int4 ParticleInfoPack;
+				ParticleInfoPack[0] = ParticleCellHash;
+				ParticleInfoPack[1] = WriteIndex;
+				ParticleInfoPack[2] = InParticleID.Index;
+				ParticleInfoPack[3] = InParticleID.AcquireTag;
+				{ParticleIDBufferName}[WriteIndex] = ParticleInfoPack;
+				{ParticlePosBufferName}[WriteIndex] = InParticlePosition;
+
+				int PreviousCellCount;
+				InterlockedAdd({CellCountBufferName}[ParticleCellHash], 1, PreviousCellCount);
+				OutNumParticles = WriteIndex;
+			}
+		)");
+		OutHLSL += FString::Format(FormatFunction, ArgsSample);
+	}
+	else if (DefinitionFunctionName == PerformKNNQueryFunctionName)
+	{
+	static const TCHAR* FormatFunction = LR"(
+			void {InstanceFunctionName}(in NiagaraID InParticleID, in float3 InPosition, in float InSearchRadius, in int InMaxNeighbors, in bool bIncludeSelf, out int OutQueryID, out int OutNumFoundNeighbors)
+			{
+				int3 CellIndex = SpatialHash_GetCellIndex(InPosition, {CellLengthName});
+				int CellRange = max(1, floor(InSearchRadius / {CellLengthName}));
+				uint CurrentCandidateCount = 0;
+				CandidateParticle Closest[32];
+				Initialize_ClosestCandidates();
+				uint CellHash = SpatialHash_HashFunction(CellIndex, {TableSizeName});
+				int CellStart = {CellStartIndicesBufferName}[CellHash];
+				int CellEnd = {CellEndIndicesBufferName}[CellHash];
+				if (CellStart != -1 && CellEnd != -1)
+				{
+					[loop][allow_uav_condition]
+					for (int p = CellStart; p <= CellEnd; ++p)
+					{
+						NiagaraID CandidateID;
+						CandidateID.Index = {Built_ParticleIDBufferName}[p][2];
+						CandidateID.AcquireTag = {Built_ParticleIDBufferName}[p][3];
+						float3 CandidatePosition = {Built_ParticlePosBufferName}[p];
+						if (!bIncludeSelf && InParticleID.Index == CandidateID.Index && InParticleID.AcquireTag == CandidateID.AcquireTag)
+						{
+							continue;
+						}
+						float Distance = distance(CandidatePosition, InPosition);
+						CandidateParticle Candidate;
+						Candidate.ExternalID = CandidateID;
+						Candidate.Position = CandidatePosition;
+						[branch]
+						if (Distance < InSearchRadius && CurrentCandidateCount < 32)
+						{
+							Closest[CurrentCandidateCount] = Candidate;
+							CurrentCandidateCount++;
+						}
+					}
+				}
+				[loop]
+				for (int l = 1; l < CellRange + 1; ++l)
+				{
+					// Top and bottom
+					[loop]
+					for (int i = -1; i <= l; ++i)
+					{
+						[loop]
+						for (int k = -l; k <= l; ++k)
+						{
+							int3 TopNeighborIndex = CellIndex + int3(i, l, k);
+							int TopCellHash = SpatialHash_HashFunction(TopNeighborIndex, {TableSizeName});
+							int TopCellStart = {CellStartIndicesBufferName}[TopCellHash];
+							int TopCellEnd = {CellEndIndicesBufferName}[TopCellHash];
+							if (TopCellStart != -1 && TopCellEnd != -1)
+							{
+								[loop][allow_uav_condition]
+								for (int p = TopCellStart; p <= TopCellEnd; ++p)
+								{
+									NiagaraID CandidateID;
+									CandidateID.Index = {Built_ParticleIDBufferName}[p][2];
+									CandidateID.AcquireTag = {Built_ParticleIDBufferName}[p][3];
+									float3 CandidatePosition = {Built_ParticlePosBufferName}[p];
+									if (!bIncludeSelf && InParticleID.Index == CandidateID.Index && InParticleID.AcquireTag == CandidateID.AcquireTag)
+									{
+										continue;
+									}
+									float Distance = distance(CandidatePosition, InPosition);
+									CandidateParticle Candidate;
+									Candidate.ExternalID = CandidateID;
+									Candidate.Position = CandidatePosition;
+									if (Distance < InSearchRadius && CurrentCandidateCount < 32)
+									{
+										Closest[CurrentCandidateCount] = Candidate;
+										CurrentCandidateCount++;
+									}
+								}
+							}
+
+							int3 BottomNeighborIndex = CellIndex + int3(i, -l, k);
+							int BottomCellHash = SpatialHash_HashFunction(BottomNeighborIndex, {TableSizeName});
+							int BottomCellStart = {CellStartIndicesBufferName}[BottomCellHash];
+							int BottomCellEnd = {CellEndIndicesBufferName}[BottomCellHash];
+							if (BottomCellStart != -1 && BottomCellEnd != -1)
+							{
+								[loop][allow_uav_condition]
+								for (int p = BottomCellStart; p <= BottomCellEnd; ++p)
+								{
+									NiagaraID CandidateID;
+									CandidateID.Index = {Built_ParticleIDBufferName}[p][2];
+									CandidateID.AcquireTag = {Built_ParticleIDBufferName}[p][3];
+									float3 CandidatePosition = {Built_ParticlePosBufferName}[p];
+									if (!bIncludeSelf && InParticleID.Index == CandidateID.Index && InParticleID.AcquireTag == CandidateID.AcquireTag)
+									{
+										continue;
+									}
+									float Distance = distance(CandidatePosition, InPosition);
+									CandidateParticle Candidate;
+									Candidate.ExternalID = CandidateID;
+									Candidate.Position = CandidatePosition;
+									if (Distance < InSearchRadius && CurrentCandidateCount < 32)
+									{
+										Closest[CurrentCandidateCount] = Candidate;
+										CurrentCandidateCount++;
+									}
+								}
+							}
+						}
+					}
+				)"
+				LR"(
+					[loop]
+					for (int j = -l + 1; j <= l - 1; ++j)
+					{
+						[loop]
+						for (int k = -l; k <= l; ++k)
+						{
+							int3 LeftNeighborIndex = CellIndex + int3(-l, j, k);
+							int LeftCellHash = SpatialHash_HashFunction(LeftNeighborIndex, {TableSizeName});
+							int LeftCellStart = {CellStartIndicesBufferName}[LeftCellHash];
+							int LeftCellEnd = {CellEndIndicesBufferName}[LeftCellHash];
+							if (LeftCellStart != -1 && LeftCellEnd != -1)
+							{
+								[loop][allow_uav_condition]
+								for (int p = LeftCellStart; p <= LeftCellEnd; ++p)
+								{
+									NiagaraID CandidateID;
+									CandidateID.Index = {Built_ParticleIDBufferName}[p][2];
+									CandidateID.AcquireTag = {Built_ParticleIDBufferName}[p][3];
+									float3 CandidatePosition = {Built_ParticlePosBufferName}[p];
+									if (!bIncludeSelf && InParticleID.Index == CandidateID.Index && InParticleID.AcquireTag == CandidateID.AcquireTag)
+									{
+										continue;
+									}
+									float Distance = distance(CandidatePosition, InPosition);
+									CandidateParticle Candidate;
+									Candidate.ExternalID = CandidateID;
+									Candidate.Position = CandidatePosition;
+									if (Distance < InSearchRadius && CurrentCandidateCount < 32)
+									{
+										Closest[CurrentCandidateCount] = Candidate;
+										CurrentCandidateCount++;
+									}
+								}
+							}
+
+							int3 RightNeighborIndex = CellIndex + int3(l, j, k);
+							int RightCellHash = SpatialHash_HashFunction(RightNeighborIndex, {TableSizeName});
+							int RightCellStart = {CellStartIndicesBufferName}[RightCellHash];
+							int RightCellEnd = {CellEndIndicesBufferName}[RightCellHash];
+							if (RightCellStart != -1 && RightCellEnd != -1)
+							{
+								[loop][allow_uav_condition]
+								for (int p = RightCellStart; p <= RightCellEnd; ++p)
+								{
+									NiagaraID CandidateID;
+									CandidateID.Index = {Built_ParticleIDBufferName}[p][2];
+									CandidateID.AcquireTag = {Built_ParticleIDBufferName}[p][3];
+									float3 CandidatePosition = {Built_ParticlePosBufferName}[p];
+									if (!bIncludeSelf && InParticleID.Index == CandidateID.Index && InParticleID.AcquireTag == CandidateID.AcquireTag)
+									{
+										continue;
+									}
+									float Distance = distance(CandidatePosition, InPosition);
+									CandidateParticle Candidate;
+									Candidate.ExternalID = CandidateID;
+									Candidate.Position = CandidatePosition;
+									if (Distance < InSearchRadius && CurrentCandidateCount < 32)
+									{
+										Closest[CurrentCandidateCount] = Candidate;
+										CurrentCandidateCount++;
+									}
+								}
+							}
+						}
+					}
+
+					[loop]
+					for (int xi = -l + 1; xi <= l - 1; ++xi)
+					{
+						[loop]
+						for (int j = -l + 1; j <= l - 1; ++j)
+						{
+							int3 FrontNeighborIndex = CellIndex + int3(xi, j, -l);
+							int FrontCellHash = SpatialHash_HashFunction(FrontNeighborIndex, {TableSizeName});
+							int FrontCellStart = {CellStartIndicesBufferName}[FrontCellHash];
+							int FrontCellEnd = {CellEndIndicesBufferName}[FrontCellHash];
+							if (FrontCellStart != -1 && FrontCellEnd != -1)
+							{
+								[loop][allow_uav_condition]
+								for (int p = FrontCellStart; p <= FrontCellEnd; ++p)
+								{
+									NiagaraID CandidateID;
+									CandidateID.Index = {Built_ParticleIDBufferName}[p][2];
+									CandidateID.AcquireTag = {Built_ParticleIDBufferName}[p][3];
+									float3 CandidatePosition = {Built_ParticlePosBufferName}[p];
+									if (!bIncludeSelf && InParticleID.Index == CandidateID.Index && InParticleID.AcquireTag == CandidateID.AcquireTag)
+									{
+										continue;
+									}
+									float Distance = distance(CandidatePosition, InPosition);
+									CandidateParticle Candidate;
+									Candidate.ExternalID = CandidateID;
+									Candidate.Position = CandidatePosition;
+									if (Distance < InSearchRadius && CurrentCandidateCount < 32)
+									{
+										Closest[CurrentCandidateCount] = Candidate;
+										CurrentCandidateCount++;
+									}
+								}
+							}
+
+							int3 BackNeighborIndex = CellIndex + int3(xi, j, l);
+							int BackCellHash = SpatialHash_HashFunction(BackNeighborIndex, {TableSizeName});
+							int BackCellStart = {CellStartIndicesBufferName}[BackCellHash];
+							int BackCellEnd = {CellEndIndicesBufferName}[BackCellHash];
+							if (BackCellStart != -1 && BackCellEnd != -1)
+							{
+								[loop][allow_uav_condition]
+								for (int p = BackCellStart; p <= BackCellEnd; ++p)
+								{
+									NiagaraID CandidateID;
+									CandidateID.Index = {Built_ParticleIDBufferName}[p][2];
+									CandidateID.AcquireTag = {Built_ParticleIDBufferName}[p][3];
+									float3 CandidatePosition = {Built_ParticlePosBufferName}[p];
+									if (!bIncludeSelf && InParticleID.Index == CandidateID.Index && InParticleID.AcquireTag == CandidateID.AcquireTag)
+									{
+										continue;
+									}
+									float Distance = distance(CandidatePosition, InPosition);
+									CandidateParticle Candidate;
+									Candidate.ExternalID = CandidateID;
+									Candidate.Position = CandidatePosition;
+									if (Distance < InSearchRadius && CurrentCandidateCount < 32)
+									{
+										Closest[CurrentCandidateCount] = Candidate;
+										CurrentCandidateCount++;
+									}
+								}
+							}
+						}
+					}
+				}
+				[unroll(32)]
+				for (uint i = 0; i < CurrentCandidateCount; ++i)
+				{
+					ClosestCandidates[i] = Closest[i];
+				}
+				
+				uint tmp = CurrentCandidateCount;
+				OutQueryID = 0;
+				OutNumFoundNeighbors = tmp;
+
+				//Heapify_ClosestCandidates(CurrentCandidateCount, InPosition);
+
+				int QueryResultIndex;
+				InterlockedAdd({CurrentNNIDName}[0], 1, QueryResultIndex);
+				OutQueryID = QueryResultIndex;
+
+				NiagaraID InvalidID;
+				InvalidID.Index = -1;
+				InvalidID.AcquireTag = -1;
+
+				if (CurrentCandidateCount <= {MaximumNeighborCountName})
+				{
+					[unroll(32)]
+					for (uint i = 0; i < CurrentCandidateCount; ++i)
+					{
+						{NearestNeighborResultsName}[QueryResultIndex * {MaximumNeighborCountName} + i] = Closest[i].ExternalID;
+					}
+					OutNumFoundNeighbors = CurrentCandidateCount;
+				}
+				else
+				{
+					uint StartIndex = QueryResultIndex * {MaximumNeighborCountName};
+					[allow_uav_condition]
+					for (uint i = 0; i < {MaximumNeighborCountName}; ++i)
+					{
+						{NearestNeighborResultsName}[StartIndex + i] = Closest[i].ExternalID;
+					}
+					OutNumFoundNeighbors = {MaximumNeighborCountName};
+				}
+			}
+		)";
+		OutHLSL += FString::Format(FormatFunction, ArgsSample);
+	}
+	else if (DefinitionFunctionName == GetClosestNeighborByIndexFunctionName)
+	{
+		static const TCHAR* FormatFunction = TEXT(R"(
+			void {InstanceFunctionName}(in int Index, in int QueryID, out bool bValid, out NiagaraID NeighborID)
+			{
+				bValid = Index <= {MaximumNeighborCountName};
+				if (bValid)
+				{
+					NeighborID = {NearestNeighborResultsName}[QueryID * {MaximumNeighborCountName} + Index - 1];
+				}
+				bValid = (NeighborID.Index != -1) && (NeighborID.AcquireTag != -1);
+			}
+		)");
+		OutHLSL += FString::Format(FormatFunction, ArgsSample);
+	}
+	else if (DefinitionFunctionName == Get16ClosestNeighborsFunctionName)
+	{
+		static const TCHAR* FormatFunction = TEXT(R"(
+			void {InstanceFunctionName}(in int QueryID, out NiagaraID NeighborID_01,
+														out NiagaraID NeighborID_02,
+														out NiagaraID NeighborID_03,
+														out NiagaraID NeighborID_04,
+														out NiagaraID NeighborID_05,
+														out NiagaraID NeighborID_06,
+														out NiagaraID NeighborID_07,
+														out NiagaraID NeighborID_08,
+														out NiagaraID NeighborID_09,
+														out NiagaraID NeighborID_10,
+														out NiagaraID NeighborID_11,
+														out NiagaraID NeighborID_12,
+														out NiagaraID NeighborID_13,
+														out NiagaraID NeighborID_14,
+														out NiagaraID NeighborID_15,
+														out NiagaraID NeighborID_16)
+			{
+				NiagaraID InvalidID;
+				InvalidID.Index = -1;
+				InvalidID.AcquireTag = -1;
+				NeighborID_01 = 1 <= {MaximumNeighborCountName} ? {NearestNeighborResultsName}[QueryID * {MaximumNeighborCountName} + 00] : InvalidID;
+				NeighborID_02 = 2 <= {MaximumNeighborCountName} ? {NearestNeighborResultsName}[QueryID * {MaximumNeighborCountName} + 01] : InvalidID;
+				NeighborID_03 = 3 <= {MaximumNeighborCountName} ? {NearestNeighborResultsName}[QueryID * {MaximumNeighborCountName} + 02] : InvalidID;
+				NeighborID_04 = 4 <= {MaximumNeighborCountName} ? {NearestNeighborResultsName}[QueryID * {MaximumNeighborCountName} + 03] : InvalidID;
+				NeighborID_05 = 5 <= {MaximumNeighborCountName} ? {NearestNeighborResultsName}[QueryID * {MaximumNeighborCountName} + 04] : InvalidID;
+				NeighborID_06 = 6 <= {MaximumNeighborCountName} ? {NearestNeighborResultsName}[QueryID * {MaximumNeighborCountName} + 05] : InvalidID;
+				NeighborID_07 = 7 <= {MaximumNeighborCountName} ? {NearestNeighborResultsName}[QueryID * {MaximumNeighborCountName} + 06] : InvalidID;
+				NeighborID_08 = 8 <= {MaximumNeighborCountName} ? {NearestNeighborResultsName}[QueryID * {MaximumNeighborCountName} + 07] : InvalidID;
+				NeighborID_09 = 9 <= {MaximumNeighborCountName} ? {NearestNeighborResultsName}[QueryID * {MaximumNeighborCountName} + 08] : InvalidID;
+				NeighborID_10 = 10 <= {MaximumNeighborCountName} ? {NearestNeighborResultsName}[QueryID * {MaximumNeighborCountName} + 09] : InvalidID;
+				NeighborID_11 = 11 <= {MaximumNeighborCountName} ? {NearestNeighborResultsName}[QueryID * {MaximumNeighborCountName} + 10] : InvalidID;
+				NeighborID_12 = 12 <= {MaximumNeighborCountName} ? {NearestNeighborResultsName}[QueryID * {MaximumNeighborCountName} + 11] : InvalidID;
+				NeighborID_13 = 13 <= {MaximumNeighborCountName} ? {NearestNeighborResultsName}[QueryID * {MaximumNeighborCountName} + 12] : InvalidID;
+				NeighborID_14 = 14 <= {MaximumNeighborCountName} ? {NearestNeighborResultsName}[QueryID * {MaximumNeighborCountName} + 13] : InvalidID;
+				NeighborID_15 = 15 <= {MaximumNeighborCountName} ? {NearestNeighborResultsName}[QueryID * {MaximumNeighborCountName} + 14] : InvalidID;
+				NeighborID_16 = 16 <= {MaximumNeighborCountName} ? {NearestNeighborResultsName}[QueryID * {MaximumNeighborCountName} + 15] : InvalidID;
+			}
+		)");
+	}
+	else
+	{
+		return false;
+	}
+	OutHLSL += TEXT("\n");
+	return true;
+}
+
+void UNiagaraDataInterfaceSpatialHash::GetParameterDefinitionHLSL(FNiagaraDataInterfaceGPUParamInfo& ParamInfo, FString& OutHLSL)
+{
+	OutHLSL += TEXT("DISPATIALHASH_DECLARE_CONSTANTS(") + ParamInfo.DataInterfaceHLSLSymbol + TEXT(")\n");
+}
+
+FNiagaraDataInterfaceParametersCS* UNiagaraDataInterfaceSpatialHash::ConstructComputeParameters() const
+{
+	return new FNiagaraDataInterfaceParametersCS_SpatialHash();
+}
+
+void UNiagaraDataInterfaceSpatialHash::ProvidePerInstanceDataForRenderThread(void* DataForRenderThread, void* PerInstanceData, const FGuid& SystemInstance)
+{
+	FNiagaraDISpatialHashPassedDataToRT* Data = static_cast<FNiagaraDISpatialHashPassedDataToRT*>(DataForRenderThread);
+	FNDISpatialHash_InstanceData* SourceData = static_cast<FNDISpatialHash_InstanceData*>(PerInstanceData);
+	Data->TableSize = SourceData->TableSize;
+	Data->MaximumNeighborCount = SourceData->MaximumNeighborCount;
+	Data->MaximumSearchRadius = SourceData->MaximumSearchRadius;
+	Data->NumParticles = SourceData->NumParticles;
+	Data->CellLength = SourceData->CellLength;
+	Data->SpatialHashGpuBuffers = SourceData->SpatialHashGpuBuffers;
+}
+
+void UNiagaraDataInterfaceSpatialHash::PostExecute()
+{
+	BuildTableHelper();
+}
+
+bool UNiagaraDataInterfaceSpatialHash::CopyToInternal(UNiagaraDataInterface* Destination) const
+{
+	if (!Super::CopyToInternal(Destination))
+	{
+		return false;
+	}
+	CastChecked<UNiagaraDataInterfaceSpatialHash>(Destination)->MaximumParticleCount = MaximumParticleCount;
+	CastChecked<UNiagaraDataInterfaceSpatialHash>(Destination)->TableSize = TableSize;
+	CastChecked<UNiagaraDataInterfaceSpatialHash>(Destination)->MaximumNeighborCount = MaximumNeighborCount;
+	CastChecked<UNiagaraDataInterfaceSpatialHash>(Destination)->MaximumSearchRadius = MaximumSearchRadius;
+	return true;
+}
+
+void UNiagaraDataInterfaceSpatialHash::PushToRenderThread()
+{ 
+
+}
+
+void UNiagaraDataInterfaceSpatialHash::BuildTableHelper()
+{
+}
+
+void UNiagaraDataInterfaceSpatialHash::GetXClosestNeighborsFromQueryHelper(FNiagaraDINearestNeighborBatch* Batch, uint32 QueryID, uint32 NumberToRetrieve, TArray<FNiagaraID>& Neighbors)
+{
+	if (NumberToRetrieve > MaximumNeighborCount)
+	{
+		UE_LOG(LogNiagara, Warning, TEXT("Requested neighbor count greater than allowed by spatial hash grid, clamping requested count."));
+		NumberToRetrieve = MaximumNeighborCount;
+	}
+	Neighbors.SetNum(NumberToRetrieve);
+	TArray<FNiagaraID> AllNeighbors;
+	Batch->GetQueryResult(QueryID, AllNeighbors);
+	for (uint32 i = 0; i < NumberToRetrieve; ++i)
+	{
+		if (static_cast<int32>(i) >= AllNeighbors.Num())
+		{
+			Neighbors[i] = { -1, -1 };
+		}
+		else
+		{
+			Neighbors[i] = AllNeighbors[i];
+		}
+	}
+}
+
+/////////////////////////////////////////////////////////
+// FNDISpatialHash_InstanceData
+/////////////////////////////////////////////////////////
+
+void FNDISpatialHash_InstanceData::AllocatePersistentTables()
+{
+	StartIndex.SetNum(TableSize, true);
+	EndIndex.SetNum(TableSize, true);
+}
+
+void FNDISpatialHash_InstanceData::ResetTables()
+{
+	StartIndex.Init(-1, TableSize);
+	EndIndex.Init(-1, TableSize);
+}
+
+void FNDISpatialHash_InstanceData::BuildTable()
+{
+	ResetTables();
+
+	NumParticles = Particles.Num();
+
+	// Sort by cell hash
+	Particles.Sort([](const ParticleData& ParticleA, const ParticleData& ParticleB)
+	{
+		return ParticleA.CellHash < ParticleB.CellHash;
+	});
+
+	if (NumParticles > 0)
+	{
+		// Get start and end indices
+		StartIndex[Particles[0].CellHash] = 0;
+		EndIndex[Particles[NumParticles - 1].CellHash] = NumParticles - 1;
+		for (uint32 i = 1; i < NumParticles; ++i)
+		{
+			if (Particles[i].CellHash != Particles[i - 1].CellHash)
+			{
+				StartIndex[Particles[i].CellHash] = i;
+				EndIndex[Particles[i - 1].CellHash] = i - 1;
+			}
+		}
+	}
+
+	NumParticles = 0;
+	Swap(Particles, Particles_Built);
+	Particles.Reset();
+}
+
+void FNDISpatialHash_InstanceData::BuildTableGPU()
+{
+	check(SpatialHashGpuBuffers);
+
+	FSpatialHashGPUBuffers* Proxy = SpatialHashGpuBuffers;
+	ENQUEUE_RENDER_COMMAND(PrefixSum)(
+		[Proxy](FRHICommandListImmediate& RHICmdList)
+	{
+		FRHIResourceCreateInfo CreateInfo;
+		// Pad out and in sizes to next largest multiple of NIAGARA_SPATIAL_HASH_THREAD_COUNT so that
+		// the prefix sum gives the correct result
+		int32 NumberOfBlocks = FMath::DivideAndRoundUp(Proxy->TableSize, NIAGARA_SPATIAL_HASH_THREAD_COUNT);
+		check(NumberOfBlocks <= NIAGARA_SPATIAL_HASH_THREAD_COUNT);
+		int32 NumberOfBlocksPadded = FMath::DivideAndRoundUp(NumberOfBlocks, NIAGARA_SPATIAL_HASH_THREAD_COUNT) * NIAGARA_SPATIAL_HASH_THREAD_COUNT;
+		int32 NumberOfElementsPadded = NumberOfBlocks * NIAGARA_SPATIAL_HASH_THREAD_COUNT;
+
+		ClearUAV(RHICmdList, Proxy->GetCurrentNNID().UAV, sizeof(uint32), 0);
+		ClearUAV(RHICmdList, Proxy->GetNumParticles().UAV, sizeof(uint32), 0);
+		int32 NumParticles = Proxy->MaximumParticleCount;
+
+		FRWBuffer ScanFirstOutput;
+		ScanFirstOutput.Initialize(sizeof(int32), NumberOfElementsPadded, EPixelFormat::PF_R32_SINT);
+
+		FRWBuffer ScanFinalOutput;
+		ScanFinalOutput.Initialize(sizeof(int32), NumberOfElementsPadded, EPixelFormat::PF_R32_SINT, 0, TEXT("ScanFinalOutput"));
+
+		FRWBuffer BlockScans;
+		BlockScans.Initialize(sizeof(int32), NumberOfBlocksPadded, EPixelFormat::PF_R32_SINT);
+
+		FRWBuffer BlockScansOut;
+		BlockScansOut.Initialize(sizeof(int32), NumberOfBlocksPadded, EPixelFormat::PF_R32_SINT);
+
+		FRWBuffer DummyScanBuffer;
+		DummyScanBuffer.Initialize(sizeof(int32), 1, EPixelFormat::PF_R32_SINT);
+
+		RHICmdList.TransitionResource(EResourceTransitionAccess::EReadable, EResourceTransitionPipeline::EComputeToCompute, Proxy->GetCellCount().UAV);
+
+		TShaderMapRef<FNiagaraPrefixSumCS> PrefixSumCS(GetGlobalShaderMap(ERHIFeatureLevel::SM5));
+		RHICmdList.SetComputeShader(PrefixSumCS->GetComputeShader());
+		PrefixSumCS->SetOutput(RHICmdList, ScanFirstOutput.UAV, BlockScans.UAV);
+		PrefixSumCS->SetParameters(RHICmdList, Proxy->GetCellCount().SRV);
+		DispatchComputeShader(RHICmdList, *PrefixSumCS, NumberOfBlocks, 1, 1);
+
+		// Every NIAGARA_SPATIAL_HASH_THREAD_COUNT elements in the cell count buffer is now scanned in ScanFirstOutput
+		// The sum of each block of elements is in BlockScans
+
+		RHICmdList.TransitionResource(EResourceTransitionAccess::EReadable, EResourceTransitionPipeline::EComputeToCompute, Proxy->GetCellCount().UAV);
+		ClearUAV(RHICmdList, Proxy->GetCellCount().UAV, sizeof(int32) * NumberOfElementsPadded, 0);
+		RHICmdList.TransitionResource(EResourceTransitionAccess::EReadable, EResourceTransitionPipeline::EComputeToCompute, BlockScans.UAV);
+
+		TShaderMapRef<FNiagaraPrefixSumCS> PrefixSumCS2(GetGlobalShaderMap(ERHIFeatureLevel::SM5));
+		RHICmdList.SetComputeShader(PrefixSumCS2->GetComputeShader());
+		PrefixSumCS2->SetOutput(RHICmdList, BlockScansOut.UAV, DummyScanBuffer.UAV);
+		PrefixSumCS2->SetParameters(RHICmdList, BlockScans.SRV);
+		DispatchComputeShader(RHICmdList, *PrefixSumCS2, 1, 1, 1);
+
+		// BlockScans itself is scanned. Because of this, the maximum table size that we can have is NIAGARA_SPATIAL_HAS_THREAD_COUNT ^ 2.
+
+		RHICmdList.TransitionResource(EResourceTransitionAccess::EReadable, EResourceTransitionPipeline::EComputeToCompute, BlockScansOut.UAV);
+		RHICmdList.TransitionResource(EResourceTransitionAccess::EReadable, EResourceTransitionPipeline::EComputeToCompute, ScanFirstOutput.UAV);
+		RHICmdList.TransitionResource(EResourceTransitionAccess::EWritable, EResourceTransitionPipeline::EComputeToCompute, ScanFinalOutput.UAV);
+
+		TShaderMapRef<FNiagaraScanAddBlockResultsCS> PrefixSumAddBlockResultsCS(GetGlobalShaderMap(ERHIFeatureLevel::SM5));
+		RHICmdList.SetComputeShader(PrefixSumAddBlockResultsCS->GetComputeShader());
+		PrefixSumAddBlockResultsCS->SetOutput(RHICmdList, ScanFinalOutput.UAV);
+		PrefixSumAddBlockResultsCS->SetParameters(RHICmdList, ScanFirstOutput.SRV, BlockScansOut.SRV);
+		DispatchComputeShader(RHICmdList, *PrefixSumAddBlockResultsCS, NumberOfBlocks, 1, 1);
+
+		// The sum of each subsequent block is added to the corresponding elements. This gives the final correct scan result.
+
+		PrefixSumCS->UnbindBuffers(RHICmdList);
+		PrefixSumCS2->UnbindBuffers(RHICmdList);
+		PrefixSumAddBlockResultsCS->UnbindBuffers(RHICmdList);
+		ScanFirstOutput.Release();
+		BlockScans.Release();
+		BlockScansOut.Release();
+		DummyScanBuffer.Release();
+		// END PREFIX SUM
+
+		// BEGIN COUNTING SORT
+		RHICmdList.TransitionResource(EResourceTransitionAccess::EReadable, EResourceTransitionPipeline::EComputeToCompute, Proxy->GetParticleIDs().UAV);
+		RHICmdList.TransitionResource(EResourceTransitionAccess::EReadable, EResourceTransitionPipeline::EComputeToCompute, Proxy->GetParticlePos().UAV);
+		RHICmdList.TransitionResource(EResourceTransitionAccess::EWritable, EResourceTransitionPipeline::EComputeToCompute, Proxy->GetBuiltParticleIDs().UAV);
+		RHICmdList.TransitionResource(EResourceTransitionAccess::EWritable, EResourceTransitionPipeline::EComputeToCompute, Proxy->GetBuiltParticlePos().UAV);
+		RHICmdList.TransitionResource(EResourceTransitionAccess::ERWBarrier, EResourceTransitionPipeline::EComputeToCompute, ScanFinalOutput.UAV);
+
+		TShaderMapRef<FNiagaraCountingSortCS> CountingSortCS(GetGlobalShaderMap(ERHIFeatureLevel::SM5));
+		RHICmdList.SetComputeShader(CountingSortCS->GetComputeShader());
+		CountingSortCS->SetOutput(RHICmdList, Proxy->GetBuiltParticleIDs().UAV, Proxy->GetBuiltParticlePos().UAV, ScanFinalOutput.UAV);
+		CountingSortCS->SetParameters(RHICmdList, Proxy->GetParticleIDs().SRV, Proxy->GetParticlePos().SRV, NumParticles);
+		DispatchComputeShader(RHICmdList, *CountingSortCS, FMath::DivideAndRoundUp(Proxy->MaximumParticleCount, NIAGARA_SPATIAL_HASH_THREAD_COUNT), 1, 1);
+
+		ScanFinalOutput.Release();
+		CountingSortCS->UnbindBuffers(RHICmdList);
+		// END COUNTING SORT
+
+		// BEGIN CELL INDEXING
+
+		RHICmdList.TransitionResource(EResourceTransitionAccess::EReadable, EResourceTransitionPipeline::EComputeToCompute, Proxy->GetBuiltParticleIDs().UAV);
+		RHICmdList.TransitionResource(EResourceTransitionAccess::EReadable, EResourceTransitionPipeline::EComputeToCompute, Proxy->GetBuiltParticlePos().UAV);
+		RHICmdList.TransitionResource(EResourceTransitionAccess::EWritable, EResourceTransitionPipeline::EComputeToCompute, Proxy->GetCellStartIndices().UAV);
+		RHICmdList.TransitionResource(EResourceTransitionAccess::EWritable, EResourceTransitionPipeline::EComputeToCompute, Proxy->GetCellEndIndices().UAV);
+
+		ClearUAV(RHICmdList, Proxy->GetCellStartIndices().UAV, sizeof(int32) * Proxy->TableSize, -1);
+		ClearUAV(RHICmdList, Proxy->GetCellEndIndices().UAV, sizeof(int32) * Proxy->TableSize, -1);
+
+		TShaderMapRef<FNiagaraSpatialHashIndexCellsCS> IndexCellsCS(GetGlobalShaderMap(ERHIFeatureLevel::SM5));
+		RHICmdList.SetComputeShader(IndexCellsCS->GetComputeShader());
+		IndexCellsCS->SetOutput(RHICmdList, Proxy->GetCellStartIndices().UAV, Proxy->GetCellEndIndices().UAV);
+		IndexCellsCS->SetParameters(RHICmdList, Proxy->GetBuiltParticleIDs().SRV, NumParticles);
+		DispatchComputeShader(RHICmdList, *IndexCellsCS, FMath::DivideAndRoundUp(Proxy->MaximumParticleCount, NIAGARA_SPATIAL_HASH_THREAD_COUNT), 1, 1);
+
+		RHICmdList.TransitionResource(EResourceTransitionAccess::EReadable, EResourceTransitionPipeline::EComputeToCompute, Proxy->GetCellStartIndices().UAV);
+		RHICmdList.TransitionResource(EResourceTransitionAccess::EReadable, EResourceTransitionPipeline::EComputeToCompute, Proxy->GetCellEndIndices().UAV);
+
+		IndexCellsCS->UnbindBuffers(RHICmdList);
+
+		// END CELL INDEXING
+	}
+	);
+}
+
+uint32 FNDISpatialHash_InstanceData::NearestNeighbor(FNiagaraID ParticleID, FVector Position, float SearchRadius, uint32 MaxNeighbors, bool bIncludeSelf, TArray<FNiagaraID> &ClosestParticles)
+{
+	if (Particles_Built.Num() == 0) { return 0; }
+
+	if (MaxNeighbors > MaximumNeighborCount)
+	{
+		UE_LOG(LogNiagara, Warning, TEXT("Requested neighbor count greater than allowed by spatial hash grid, clamping requested count."));
+		MaxNeighbors = MaximumNeighborCount;
+	}
+
+	if (SearchRadius > MaximumSearchRadius)
+	{
+		UE_LOG(LogNiagara, Warning, TEXT("Requested search radius greater than allowed by spatial hash grid, clamping requested radius."));
+		SearchRadius = MaximumSearchRadius;
+	}
+
+	check(SearchRadius >= 0);
+
+	FIntVector CellIndex = SpatialHash_GetCellIndex(Position, CellLength);
+
+	// Calculate how many cells away from the center we have to search
+	int CellRange = FMath::Max(1, FMath::FloorToInt(SearchRadius / CellLength));
+
+	struct CandidateParticle
+	{
+		FNiagaraID ExternalID;
+		FVector Position;
+
+		bool operator==(const CandidateParticle& Other) const { return ExternalID == Other.ExternalID && Position == Other.Position; }
+		operator FNiagaraID() const { return ExternalID; }
+	};
+
+	TArray<CandidateParticle, TInlineAllocator<32>> ClosestCandidates;
+
+	auto CollectCandidateParticlesInCell = [&](FIntVector CellIndex)
+	{
+		uint32 CellHash = SpatialHash_HashFunction(CellIndex, TableSize);
+		int32 CellStart = StartIndex[CellHash];
+		int32 CellEnd = EndIndex[CellHash];
+
+		// There are no particles in this cell
+		if (CellStart == -1 || CellEnd == -1)
+		{
+			return;
+		}
+		for (int p = CellStart; p <= CellEnd; ++p)
+		{
+			if (!bIncludeSelf && Particles_Built[p].ExternalID == ParticleID)
+			{
+				continue;
+			}
+			if (FVector::DistSquared(Particles_Built[p].ParticlePosition, Position) < SearchRadius * SearchRadius)
+			{
+				ClosestCandidates.AddUnique({ Particles_Built[p].ExternalID, Particles_Built[p].ParticlePosition });
+			}
+		}
+	};
+
+	// First collect all candidate particles in the center cell
+	CollectCandidateParticlesInCell(CellIndex);
+
+	// Move layer by layer out from the center cell (l = 1: 1 away from center, l = 2, 2 away from center...)
+	for (int l = 1; l < CellRange + 1; ++l)
+	{
+		// Check top and bottom of current layer
+		for (int i = -l; i <= l; ++i)
+		{
+			for (int k = -l; k <= l; ++k)
+			{
+				FIntVector TopNeighborIndex = CellIndex + FIntVector(i, l, k);
+				FIntVector BottomNeighborIndex = CellIndex + FIntVector(i, -l, k);
+
+				CollectCandidateParticlesInCell(TopNeighborIndex);
+				CollectCandidateParticlesInCell(BottomNeighborIndex);
+			}
+		}
+
+		// Check left and right of current layer
+		for (int j = -l + 1; j <= l - 1; ++j)
+		{
+			for (int k = -l; k <= l; ++k)
+			{
+				FIntVector LeftNeighborIndex = CellIndex + FIntVector(-l, j, k);
+				FIntVector RightNeighborIndex = CellIndex + FIntVector(l, j, k);
+
+				CollectCandidateParticlesInCell(LeftNeighborIndex);
+				CollectCandidateParticlesInCell(RightNeighborIndex);
+			}
+		}
+
+		// Check front and back of current layer
+		for (int i = -l + 1; i <= l - 1; ++i)
+		{
+			for (int j = -l + 1; j <= l - 1; ++j)
+			{
+				FIntVector FrontNeighborIndex = CellIndex + FIntVector(i, j, -l);
+				FIntVector BackNeighborIndex = CellIndex + FIntVector(i, j, l);
+
+				CollectCandidateParticlesInCell(FrontNeighborIndex);
+				CollectCandidateParticlesInCell(BackNeighborIndex);
+			}
+		}
+
+		if (ClosestCandidates.Num() >= static_cast<int32>(MaxNeighbors)) { break; }
+	}
+
+	uint32 FoundNeighbors = 0;
+
+	if (ClosestCandidates.Num() <= static_cast<int32>(MaxNeighbors))
+	{
+		// Fewer neighbors found than requested, return them all
+		for (int32 i = 0; i < ClosestCandidates.Num(); ++i)
+		{
+			ClosestParticles[i] = ClosestCandidates[i];
+		}
+		FoundNeighbors = ClosestCandidates.Num();
+	}
+	else
+	{
+		// Use min-heap to determine the k-nearest
+		auto ParticleCandidatePred = [&Position](const CandidateParticle& CandidateA, const CandidateParticle& CandidateB)
+		{
+			float Dist2A = FVector::DistSquared(CandidateA.Position, Position);
+			float Dist2B = FVector::DistSquared(CandidateB.Position, Position);
+			return Dist2A < Dist2B;
+		};
+		ClosestCandidates.Heapify(ParticleCandidatePred);
+		for (uint32 i = 0; i < MaxNeighbors; ++i)
+		{
+			CandidateParticle NextClosestParticle;
+			ClosestCandidates.HeapPop(NextClosestParticle, ParticleCandidatePred, false);
+			ClosestParticles[i] = (NextClosestParticle);
+		}
+		FoundNeighbors = MaxNeighbors;
+	}
+
+	return FoundNeighbors;
+}
+
+bool FNDISpatialHash_InstanceData::Init(UNiagaraDataInterfaceSpatialHash* Interface, FNiagaraSystemInstance* InSystemInstance)
+{
+	check(InSystemInstance);
+
+	SystemInstance = InSystemInstance;
+	MaximumParticleCount = Interface->MaximumParticleCount;
+	TableSize = Interface->TableSize;
+	MaximumNeighborCount = Interface->MaximumNeighborCount;
+	MaximumSearchRadius = Interface->MaximumSearchRadius;
+	NumParticles = 0;
+	CellLength = MaximumSearchRadius / 3.f;
+	AllocatePersistentTables();
+	SpatialHashGpuBuffers = nullptr;
+	if (SystemInstance->HasGPUEmitters())
+	{
+		SpatialHashGpuBuffers = new FSpatialHashGPUBuffers();
+		SpatialHashGpuBuffers->Initialize(this);
+		BeginInitResource(SpatialHashGpuBuffers);
+	}
+	SpatialHashBatch.Init(SystemInstance->GetIDName(), this);
+	return true;
+}
+
+bool FNDISpatialHash_InstanceData::Tick(UNiagaraDataInterfaceSpatialHash* Interface, FNiagaraSystemInstance* InSystemInstance, float InDeltaSeconds)
+{
+}
+
+void FNDISpatialHash_InstanceData::Release()
+{
+	if (SpatialHashGpuBuffers)
+	{
+		BeginReleaseResource(SpatialHashGpuBuffers);
+		ENQUEUE_RENDER_COMMAND(DeleteResource)(
+			[ParamPointerToRelease = SpatialHashGpuBuffers](FRHICommandListImmediate& RHICmdList)
+			{
+				delete ParamPointerToRelease;
+			}
+		);
+		SpatialHashGpuBuffers = nullptr;
+	}
+}
+
+/////////////////////////////////////////////////////////
+// FNiagaraDINearestNeighborBatch
+/////////////////////////////////////////////////////////
+
+void FNiagaraDINearestNeighborBatch::Init(FName InBatchID, FNDISpatialHash_InstanceData* InSpatialHashInstanceData)
+{
+	SpatialHashInstanceData = InSpatialHashInstanceData;
+	NearestNeighborResults.Empty(InSpatialHashInstanceData->MaximumNeighborCount);
+	IDToResultIndex.Empty(InSpatialHashInstanceData->MaximumNeighborCount);
+	CurrentID = 0;
+}
+
+int32 FNiagaraDINearestNeighborBatch::SubmitQuery(FNiagaraID ParticleID, FVector Position, float SearchRadius, uint32 MaxNeighbors, bool bIncludeSelf)
+{
+	if (!SpatialHashInstanceData)
+	{
+		return INDEX_NONE;
+	}
+
+	TArray<FNiagaraID> ClosestNeighbors;
+	ClosestNeighbors.Init(NIAGARA_INVALID_ID, SpatialHashInstanceData->MaximumNeighborCount);
+	uint32 FoundNeighbors = SpatialHashInstanceData->NearestNeighbor(ParticleID, Position, SearchRadius, MaxNeighbors, bIncludeSelf, ClosestNeighbors);
+	int32 ResultIndex = NearestNeighborResults.Num();
+	NearestNeighborResults.Append(ClosestNeighbors);
+	IDToResultIndex.Add(CurrentID) = ResultIndex; // Result spans from ResultIndex to ResultIndex + MaximumNeighborCount
+
+	int32 Ret = CurrentID;
+	CurrentID++;
+	return Ret;
+}
+
+bool FNiagaraDINearestNeighborBatch::GetQueryResult(uint32 InQueryID, TArray<FNiagaraID>& Result)
+{
+	int32* ResultIndexPtr = IDToResultIndex.Find(InQueryID);
+	int32 ResultIndex = INDEX_NONE;
+	if (ResultIndexPtr)
+	{
+		Result.SetNum(SpatialHashInstanceData->MaximumNeighborCount);
+		ResultIndex = *ResultIndexPtr;
+		for (uint32 i = 0; i < SpatialHashInstanceData->MaximumNeighborCount; ++i)
+		{
+			FNiagaraID ParticleID = NearestNeighborResults[ResultIndex + i];
+			if (ParticleID == NIAGARA_INVALID_ID)
+			{
+				Result.SetNum(i);
+				break;
+			}
+			else
+			{
+				Result[i] = ParticleID;
+			}
+		}
+	}
+	return true;
+}
+
+/////////////////////////////////////////////////////////
+// FSpatialHashGPUBuffers
+/////////////////////////////////////////////////////////
+
+void FSpatialHashGPUBuffers::Initialize(FNDISpatialHash_InstanceData* InstanceData)
+{
+	MaximumParticleCount = InstanceData->MaximumParticleCount;
+	TableSize = InstanceData->TableSize;
+	NumberOfParticles = 0;
+	MaximumNeighborCount = InstanceData->MaximumNeighborCount;
+}
+
+void FSpatialHashGPUBuffers::InitRHI()
+{
+	FRHIResourceCreateInfo CreateInfo;
+
+	int32 PaddedTableSize = FMath::DivideAndRoundUp(TableSize, NIAGARA_SPATIAL_HASH_THREAD_COUNT) * NIAGARA_SPATIAL_HASH_THREAD_COUNT;
+	
+	ParticleID.Initialize(sizeof(FIntVector4), MaximumParticleCount, EPixelFormat::PF_R32G32B32A32_UINT, 0, TEXT("ParticleIDsBufferGPU"));
+	ParticlePos.Initialize(sizeof(FVector), MaximumParticleCount, EPixelFormat::PF_FloatRGB, 0, TEXT("ParticlePosBufferGPU"));
+	Built_ParticleID.Initialize(sizeof(FIntVector4), MaximumParticleCount, EPixelFormat::PF_R32G32B32A32_UINT, 0, TEXT("Built_ParticleIDsBufferGPU"));
+	Built_ParticlePos.Initialize(sizeof(FVector), MaximumParticleCount, EPixelFormat::PF_FloatRGB, 0, TEXT("Built_ParticlePosBufferGPU"));
+
+	CellCount.Initialize(sizeof(int32), PaddedTableSize, EPixelFormat::PF_R32_UINT, 0, TEXT("CellCountBufferGPU"));
+
+	CellStartIndices.Initialize(sizeof(int32), TableSize, EPixelFormat::PF_R32_SINT, 0, TEXT("CellStartIndices"));
+	CellEndIndices.Initialize(sizeof(int32), TableSize, EPixelFormat::PF_R32_SINT, 0, TEXT("CellEndIndices"));
+
+	CreateInfo.DebugName = TEXT("NumParticlesGPU");
+	NumParticles.Initialize(sizeof(int32), 1, EPixelFormat::PF_R32_SINT);
+
+	NearestNeighborResults.Initialize(sizeof(FNiagaraID), MaximumNeighborCount * MaximumParticleCount, EPixelFormat::PF_R32G32_UINT, 0, TEXT("NearestNeighborResultsGPU"));
+	CurrentNNID.Initialize(sizeof(int32), 1, EPixelFormat::PF_R32_UINT);
+
+	FSpatialHashGPUBuffers* ThisProxy = this;
+	ENQUEUE_RENDER_COMMAND(InitSpatialHashBuffers)(
+		[ThisProxy, PaddedTableSize](FRHICommandListImmediate& RHICmdList) mutable
+	{
+		ClearUAV(RHICmdList, ThisProxy->GetCellCount().UAV, sizeof(int32) * PaddedTableSize, 0);
+	}
+	);
+}
+
+void FSpatialHashGPUBuffers::ReleaseRHI()
+{
+	ParticleID.Release();
+	ParticlePos.Release();
+	Built_ParticleID.Release();
+	Built_ParticlePos.Release();
+	CellCount.Release();
+	CellStartIndices.Release();
+	CellEndIndices.Release();
+
+	NumParticles.Release();
+
+	NearestNeighborResults.Release();
+	CurrentNNID.Release();
+}
+
+/////////////////////////////////////////////////////////
+// FNiagaraDataInterfaceParametersCS_SpatialHash
+/////////////////////////////////////////////////////////
+
+void FNiagaraDataInterfaceParametersCS_SpatialHash::Bind(const FNiagaraDataInterfaceParamRef& ParamRef, const class FShaderParameterMap& ParameterMap)
+{
+	FNDISpatialHashParametersName ParamNames;
+	GetNiagaraDataInterfaceParametersName(ParamNames, ParamRef.ParameterInfo.DataInterfaceHLSLSymbol);
+
+	ParticleIDBuffer.Bind(ParameterMap, *ParamNames.ParticleIDBufferName);
+	ParticlePosBuffer.Bind(ParameterMap, *ParamNames.ParticlePosBufferName);
+	Built_ParticleIDBuffer.Bind(ParameterMap, *ParamNames.Built_ParticleIDBufferName);
+	Built_ParticlePosBuffer.Bind(ParameterMap, *ParamNames.Built_ParticlePosBufferName);
+	CellCountBuffer.Bind(ParameterMap, *ParamNames.CellCountBufferName);
+	CellStartIndicesBuffer.Bind(ParameterMap, *ParamNames.CellStartIndicesBufferName);
+	CellEndIndicesBuffer.Bind(ParameterMap, *ParamNames.CellEndIndicesBufferName);
+	NumParticles.Bind(ParameterMap, *ParamNames.NumParticlesName); 
+	NearestNeighborResultsBuffer.Bind(ParameterMap, *ParamNames.NearestNeighborResultsName);
+	CurrentNNID.Bind(ParameterMap, *ParamNames.CurrentNNIDName);
+	TableSize.Bind(ParameterMap, *ParamNames.TableSizeName);
+	MaximumNeighborCount.Bind(ParameterMap, *ParamNames.MaximumNeighborCountName);
+	MaximumSearchRadius.Bind(ParameterMap, *ParamNames.MaximumSearchRadiusName);
+	CellLength.Bind(ParameterMap, *ParamNames.CellLengthName);
+}
+
+void FNiagaraDataInterfaceParametersCS_SpatialHash::Serialize(FArchive& Ar)
+{
+	Ar << ParticleIDBuffer;
+	Ar << ParticlePosBuffer;
+	Ar << Built_ParticleIDBuffer;
+	Ar << Built_ParticlePosBuffer;
+	Ar << CellCountBuffer;
+	Ar << CellStartIndicesBuffer;
+	Ar << CellEndIndicesBuffer;
+	Ar << NumParticles;
+	Ar << NearestNeighborResultsBuffer;
+	Ar << CurrentNNID;
+	Ar << TableSize;
+	Ar << MaximumNeighborCount;
+	Ar << MaximumSearchRadius;
+	Ar << CellLength;
+}
+
+void FNiagaraDataInterfaceParametersCS_SpatialHash::Set(FRHICommandList& RHICmdList, const FNiagaraDataInterfaceSetArgs& Context) const
+{
+	check(IsInRenderingThread());
+
+	FRHIComputeShader* ComputeShaderRHI = Context.Shader->GetComputeShader();
+	FNiagaraDataInterfaceProxySpatialHash* InterfaceProxy = static_cast<FNiagaraDataInterfaceProxySpatialHash*>(Context.DataInterface);
+	FNiagaraDISpatialHashPassedDataToRT* InstanceData = InterfaceProxy->SystemInstancesToData.Find(Context.SystemInstance);
+	if (InstanceData && InstanceData->SpatialHashGpuBuffers)
+	{
+		FSpatialHashGPUBuffers* SpatialHashBuffers = InstanceData->SpatialHashGpuBuffers;
+
+		SetUAVParameter(RHICmdList, ComputeShaderRHI, ParticleIDBuffer, SpatialHashBuffers->GetParticleIDs().UAV);
+		SetUAVParameter(RHICmdList, ComputeShaderRHI, ParticlePosBuffer, SpatialHashBuffers->GetParticlePos().UAV);
+		SetUAVParameter(RHICmdList, ComputeShaderRHI, Built_ParticleIDBuffer, SpatialHashBuffers->GetBuiltParticleIDs().UAV);
+		SetUAVParameter(RHICmdList, ComputeShaderRHI, Built_ParticlePosBuffer, SpatialHashBuffers->GetBuiltParticlePos().UAV);
+		SetUAVParameter(RHICmdList, ComputeShaderRHI, CellCountBuffer, SpatialHashBuffers->GetCellCount().UAV);
+		SetUAVParameter(RHICmdList, ComputeShaderRHI, CellStartIndicesBuffer, SpatialHashBuffers->GetCellStartIndices().UAV);
+		SetUAVParameter(RHICmdList, ComputeShaderRHI, CellEndIndicesBuffer, SpatialHashBuffers->GetCellEndIndices().UAV);
+		SetUAVParameter(RHICmdList, ComputeShaderRHI, NumParticles, SpatialHashBuffers->GetNumParticles().UAV);
+		SetUAVParameter(RHICmdList, ComputeShaderRHI, NearestNeighborResultsBuffer, SpatialHashBuffers->GetNearestNeighborResults().UAV);
+		SetUAVParameter(RHICmdList, ComputeShaderRHI, CurrentNNID, SpatialHashBuffers->GetCurrentNNID().UAV);
+
+		SetShaderValue(RHICmdList, ComputeShaderRHI, TableSize, InstanceData->TableSize);
+		SetShaderValue(RHICmdList, ComputeShaderRHI, MaximumNeighborCount, InstanceData->MaximumNeighborCount);
+		SetShaderValue(RHICmdList, ComputeShaderRHI, MaximumSearchRadius, InstanceData->MaximumSearchRadius);
+		SetShaderValue(RHICmdList, ComputeShaderRHI, CellLength, InstanceData->CellLength);
+	}
+	else
+	{
+
+		SetShaderValue(RHICmdList, ComputeShaderRHI, TableSize, 0);
+		SetShaderValue(RHICmdList, ComputeShaderRHI, MaximumNeighborCount, 0);
+		SetShaderValue(RHICmdList, ComputeShaderRHI, MaximumSearchRadius, 0);
+		SetShaderValue(RHICmdList, ComputeShaderRHI, CellLength, 0);
+	}
+}
+
+void FNiagaraDataInterfaceParametersCS_SpatialHash::Unset(FRHICommandList& RHICmdList, const FNiagaraDataInterfaceSetArgs& Context) const
+{
+	check(IsInRenderingThread());
+	FRHIComputeShader* ComputeShaderRHI = Context.Shader->GetComputeShader();
+	FNiagaraDataInterfaceProxySpatialHash* InterfaceProxy = static_cast<FNiagaraDataInterfaceProxySpatialHash*>(Context.DataInterface);
+	FNiagaraDISpatialHashPassedDataToRT* InstanceData = InterfaceProxy->SystemInstancesToData.Find(Context.SystemInstance);
+	if (InstanceData && InstanceData->SpatialHashGpuBuffers)
+	{
+		FSpatialHashGPUBuffers* SpatialHashBuffers = InstanceData->SpatialHashGpuBuffers;
+		SetUAVParameter(RHICmdList, ComputeShaderRHI, ParticleIDBuffer, FUnorderedAccessViewRHIParamRef());
+		SetUAVParameter(RHICmdList, ComputeShaderRHI, ParticlePosBuffer, FUnorderedAccessViewRHIParamRef());
+		SetUAVParameter(RHICmdList, ComputeShaderRHI, Built_ParticleIDBuffer, FUnorderedAccessViewRHIParamRef());
+		SetUAVParameter(RHICmdList, ComputeShaderRHI, Built_ParticlePosBuffer, FUnorderedAccessViewRHIParamRef());
+		SetUAVParameter(RHICmdList, ComputeShaderRHI, CellCountBuffer, FUnorderedAccessViewRHIParamRef());
+		SetUAVParameter(RHICmdList, ComputeShaderRHI, CellStartIndicesBuffer, FUnorderedAccessViewRHIParamRef());
+		SetUAVParameter(RHICmdList, ComputeShaderRHI, CellEndIndicesBuffer, FUnorderedAccessViewRHIParamRef());
+		SetUAVParameter(RHICmdList, ComputeShaderRHI, NumParticles, FUnorderedAccessViewRHIParamRef());
+		SetUAVParameter(RHICmdList, ComputeShaderRHI, NearestNeighborResultsBuffer, FUnorderedAccessViewRHIParamRef());
+		SetUAVParameter(RHICmdList, ComputeShaderRHI, CurrentNNID, FUnorderedAccessViewRHIParamRef());
+
+		SetShaderValue(RHICmdList, ComputeShaderRHI, TableSize, InstanceData->TableSize);
+		SetShaderValue(RHICmdList, ComputeShaderRHI, MaximumNeighborCount, InstanceData->MaximumNeighborCount);
+		SetShaderValue(RHICmdList, ComputeShaderRHI, MaximumSearchRadius, InstanceData->MaximumSearchRadius);
+		SetShaderValue(RHICmdList, ComputeShaderRHI, CellLength, InstanceData->CellLength);
+	}
+	else
+	{
+
+		SetShaderValue(RHICmdList, ComputeShaderRHI, TableSize, 0);
+		SetShaderValue(RHICmdList, ComputeShaderRHI, MaximumNeighborCount, 0);
+		SetShaderValue(RHICmdList, ComputeShaderRHI, MaximumSearchRadius, 0);
+		SetShaderValue(RHICmdList, ComputeShaderRHI, CellLength, 0);
+	}
+}
+
+void FNiagaraDataInterfaceProxySpatialHash::ConsumePerInstanceDataFromGameThread(void* PerInstanceData, const FGuid& Instance)
+{
+	FNiagaraDISpatialHashPassedDataToRT* SourceData = static_cast<FNiagaraDISpatialHashPassedDataToRT*>(PerInstanceData);
+	FNiagaraDISpatialHashPassedDataToRT& Data = SystemInstancesToData.FindOrAdd(Instance);
+	Data.TableSize = SourceData->TableSize;
+	Data.MaximumNeighborCount = SourceData->MaximumNeighborCount;
+	Data.MaximumSearchRadius = SourceData->MaximumSearchRadius;
+	Data.NumParticles = SourceData->NumParticles;
+	Data.CellLength = SourceData->CellLength;
+	Data.SpatialHashGpuBuffers = SourceData->SpatialHashGpuBuffers;
+}
+
+#undef LOCTEXT_NAMESPACE
\ No newline at end of file
