Add missing shader files from previous change

diff --git a/Engine/Plugins/FX/Niagara/Shaders/Private/NiagaraDISpatialHashCountingSort.usf b/Engine/Plugins/FX/Niagara/Shaders/Private/NiagaraDISpatialHashCountingSort.usf
new file mode 100644
--- /dev/null
+++ b/Engine/Plugins/FX/Niagara/Shaders/Private/NiagaraDISpatialHashCountingSort.usf
@@ -0,0 +1,35 @@
+// Copyright 1998-2019 Epic Games, Inc. All Rights Reserved.
+
+/*------------------------------------------------------------------------------
+	Compile time parameters:
+		THREAD_COUNT - The number of threads to launch per workgroup.
+------------------------------------------------------------------------------*/
+
+Buffer<int4> InDataToSortID;
+Buffer<float3> InDataToSortPos;
+RWBuffer<int4> OutSortedDataID;
+RWBuffer<float3> OutSortedDataPos;
+
+// Prefix sum of cell counts
+RWBuffer<int> ScannedCounts;
+
+// Number of particles
+int N;
+
+[numthreads(THREAD_COUNT,1,1)]
+void SpatialHashCountingSortCS(uint ParticleID : SV_DispatchThreadID)
+{
+    if (ParticleID < N)
+    {
+		// Grab the particle cell hash
+        int CellHash = InDataToSortID[ParticleID][0];
+	
+		// Compute what index this thread's particle belongs in
+        int SortedIndex;
+        InterlockedAdd(ScannedCounts[CellHash], 1, SortedIndex);
+	
+		// Move the particle ID and position into that index
+        OutSortedDataID[SortedIndex] = InDataToSortID[ParticleID];
+        OutSortedDataPos[SortedIndex] = InDataToSortPos[ParticleID];
+    }
+}
\ No newline at end of file
diff --git a/Engine/Plugins/FX/Niagara/Shaders/Private/NiagaraDISpatialHashIndexCells.usf b/Engine/Plugins/FX/Niagara/Shaders/Private/NiagaraDISpatialHashIndexCells.usf
new file mode 100644
--- /dev/null
+++ b/Engine/Plugins/FX/Niagara/Shaders/Private/NiagaraDISpatialHashIndexCells.usf
@@ -0,0 +1,42 @@
+// Copyright 1998-2019 Epic Games, Inc. All Rights Reserved.
+
+/*------------------------------------------------------------------------------
+	Compile time parameters:
+		THREAD_COUNT - The number of threads to launch per workgroup.
+------------------------------------------------------------------------------*/
+
+Buffer<int4> InParticleIDs;
+RWBuffer<int> CellStartIndices;
+RWBuffer<int> CellEndIndices;
+
+// Number of particles
+uint N;
+
+[numthreads(THREAD_COUNT,1,1)]
+void IndexCellsCS(uint ParticleID : SV_DispatchThreadID)
+{
+    if (ParticleID < N)
+    {
+        int CellHash = InParticleIDs[ParticleID][0];
+        if (ParticleID == 0)
+        {
+            CellStartIndices[CellHash] = 0;
+        }
+        else if (ParticleID == N - 1)
+        {
+            CellEndIndices[CellHash] = N - 1;
+        }
+        else
+        {
+			// If this particle's cell and the previous particle's cell are different
+			// this is a cell boundary. It is the previous cell's end index and the
+			// current cell's start index.
+            int PrevCell = InParticleIDs[ParticleID - 1][0];
+            if (PrevCell != CellHash)
+            {
+                CellStartIndices[CellHash] = ParticleID;
+                CellEndIndices[PrevCell] = ParticleID - 1;
+            }
+        }
+    }
+}
\ No newline at end of file
diff --git a/Engine/Plugins/FX/Niagara/Shaders/Private/NiagaraDataInterfaceSpatialHash.ush b/Engine/Plugins/FX/Niagara/Shaders/Private/NiagaraDataInterfaceSpatialHash.ush
new file mode 100644
--- /dev/null
+++ b/Engine/Plugins/FX/Niagara/Shaders/Private/NiagaraDataInterfaceSpatialHash.ush
@@ -0,0 +1,201 @@
+// Copyright 1998-2019 Epic Games, Inc. All Rights Reserved.
+
+#pragma once
+
+// Hash function
+// From http://matthias-mueller-fischer.ch/publications/tetraederCollision.pdf
+uint SpatialHash_HashFunction(int3 Position, uint TableSize)
+{
+	const uint p1 = 73856093;
+	const uint p2 = 19349663;
+	const uint p3 = 83492791;
+
+	uint n = (p1 * Position.x) ^ (p2 * Position.y) ^ (p3 * Position.z);
+	return n % TableSize;
+}
+
+int3 SpatialHash_GetCellIndex(float3 Position, float CellLength)
+{
+    return floor(Position / CellLength);
+}
+
+struct CandidateParticle
+{
+    NiagaraID ExternalID;
+    float3 Position;
+};
+
+CandidateParticle ClosestCandidates[32];
+
+void Initialize_ClosestCandidates()
+{
+    NiagaraID ID;
+    ID.Index = -1;
+    ID.AcquireTag = -1;
+    CandidateParticle Candidate;
+    Candidate.ExternalID = ID;
+    Candidate.Position = float3(0, 0, 0);
+    for (int i = 0; i < 32; ++i)
+    {
+        ClosestCandidates[i] = Candidate;
+    }
+}
+
+#define FLT_MAX 3.402823466e+38F
+
+void Heapify_ClosestCandidates(int Num, float3 Position)
+{
+    if (Num == 0 || Num == 1) { return; }
+    CandidateParticle Closest[32];
+    for (int ci = 0; ci < 32; ++ci) {
+        Closest[ci] = ClosestCandidates[ci];
+    }
+    int StartIndex = 14;
+    while (Closest[StartIndex].ExternalID.Index == -1)
+    {
+        StartIndex--;
+    }
+	
+    for (int i = StartIndex; i >= 0; ++i)
+    {
+        int LeftChildIndex = ((i + 1) * 2) - 1;
+        int RightChildIndex = LeftChildIndex + 1;
+        float LeftDistance = Closest[LeftChildIndex].ExternalID.Index != -1 ? distance(Position, Closest[LeftChildIndex].Position) : FLT_MAX;
+        float RightDistance = Closest[RightChildIndex].ExternalID.Index != -1 ? distance(Position, Closest[RightChildIndex].Position) : FLT_MAX;
+        float CurrNodeDistance = distance(Position, Closest[i].Position);
+        
+        bool bSwapLeft = false;
+        bool bSwapRight = false;
+
+		if (CurrNodeDistance > LeftDistance && CurrNodeDistance > RightDistance)
+        {
+			// Swap with the smaller of the two
+            if (LeftDistance < RightDistance)
+            {
+                bSwapLeft = true;
+            }
+            else
+            {
+                bSwapRight = true;
+            }
+        }
+        else if (CurrNodeDistance > LeftDistance)
+        {
+            bSwapLeft = true;
+        }
+        else if (CurrNodeDistance > RightDistance)
+        {
+            bSwapRight = true;
+        }
+
+        if (bSwapLeft)
+        {
+            CandidateParticle Tmp = Closest[i];
+            Closest[i] = Closest[LeftChildIndex];
+            Closest[LeftChildIndex] = Tmp;
+        }
+        else if (bSwapRight)
+        {
+            CandidateParticle Tmp = Closest[i];
+            Closest[i] = Closest[RightChildIndex];
+            Closest[RightChildIndex] = Tmp;
+        }
+    }
+    for (int i = 0; i < 32; ++i)
+    {
+        ClosestCandidates[i] = Closest[i];
+    }
+}
+
+NiagaraID Pop_ClosestCandidates(int Num, float3 Position)
+{
+    NiagaraID Ret;
+    if (Num == 0)
+    {
+        Ret.Index = -1;
+        Ret.AcquireTag = -1;
+        return Ret;
+    }
+	// Swap root with last element
+    CandidateParticle Tmp = ClosestCandidates[0];
+    Ret = Tmp.ExternalID;
+    ClosestCandidates[0] = ClosestCandidates[Num - 1];
+
+    NiagaraID ID;
+    ID.Index = -1;
+    ID.AcquireTag = -1;
+    ClosestCandidates[Num - 1].ExternalID = ID;
+
+	// Preserve heap order property
+    int CurrIndex = 0;
+    int LeftChildIndex = 1;
+    int RightChildIndex = 2;
+    float LeftDistance = ClosestCandidates[LeftChildIndex].ExternalID.Index != -1 ? distance(Position, ClosestCandidates[LeftChildIndex].Position) : FLT_MAX;
+    float RightDistance = ClosestCandidates[RightChildIndex].ExternalID.Index != -1 ? distance(Position, ClosestCandidates[RightChildIndex].Position) : FLT_MAX;
+    float CurrDistance = distance(Position, Tmp.Position);
+
+    bool bSwapLeft = false;
+    bool bSwapRight = false;
+
+    while (CurrDistance > LeftDistance || CurrDistance > RightDistance)
+    {
+        if (CurrDistance > LeftDistance && CurrDistance > RightDistance)
+        {
+            if (LeftDistance < RightDistance)
+            {
+                bSwapLeft = true;
+            }
+            else
+            {
+                bSwapRight = true;
+            }
+        }
+        else if (CurrDistance > LeftDistance)
+        {
+            bSwapLeft = true;
+        }
+        else if (CurrDistance > RightDistance)
+        {
+            bSwapRight = true;
+        }
+        if (bSwapLeft)
+        {
+            CandidateParticle Tmp = ClosestCandidates[CurrIndex];
+            ClosestCandidates[CurrIndex] = ClosestCandidates[LeftChildIndex];
+            ClosestCandidates[LeftChildIndex] = Tmp;
+            CurrIndex = LeftChildIndex;
+            CurrDistance = LeftDistance;
+        }
+        else if (bSwapRight)
+        {
+            CandidateParticle Tmp = ClosestCandidates[CurrIndex];
+            ClosestCandidates[CurrIndex] = ClosestCandidates[RightChildIndex];
+            ClosestCandidates[RightChildIndex] = Tmp;
+            CurrIndex = RightChildIndex;
+            CurrDistance = RightDistance;
+        }
+        LeftChildIndex = ((CurrIndex + 1) * 2) - 1;
+        RightChildIndex = LeftChildIndex + 1;
+        LeftDistance = ClosestCandidates[LeftChildIndex].ExternalID.Index != -1 ? distance(Position, ClosestCandidates[LeftChildIndex].Position) : FLT_MAX;
+        RightDistance = ClosestCandidates[RightChildIndex].ExternalID.Index != -1 ? distance(Position, ClosestCandidates[RightChildIndex].Position) : FLT_MAX;
+    }
+    return Ret;
+}
+
+#undef FLT_MAX
+
+#define DISPATIALHASH_DECLARE_CONSTANTS(NAME) \
+	uint TableSize_ ## NAME; \
+	uint MaximumNeighborCount_ ## NAME; \
+	float MaximumSearchRadius_ ## NAME; \
+	RWBuffer<int> NumParticles_ ## NAME; \
+	float CellLength_ ## NAME; \
+	RWBuffer<uint4> ParticleIDs_ ## NAME; \
+	RWBuffer<float3> ParticlePosition_ ## NAME; \
+	Buffer<uint4> Built_ParticleIDs_ ## NAME; \
+	Buffer<float3> Built_ParticlePosition_ ## NAME; \
+	RWBuffer<int> CellStartIndices_ ## NAME; \
+	RWBuffer<int> CellEndIndices_ ## NAME; \
+	RWBuffer<uint> CellCount_ ## NAME; \
+	RWBuffer<NiagaraID> NearestNeighborResults_ ## NAME; \
+	RWBuffer<uint> CurrentID_ ## NAME;
diff --git a/Engine/Plugins/FX/Niagara/Shaders/Private/NiagaraPrefixSum.usf b/Engine/Plugins/FX/Niagara/Shaders/Private/NiagaraPrefixSum.usf
new file mode 100644
--- /dev/null
+++ b/Engine/Plugins/FX/Niagara/Shaders/Private/NiagaraPrefixSum.usf
@@ -0,0 +1,71 @@
+// Copyright 1998-2019 Epic Games, Inc. All Rights Reserved.
+
+/*------------------------------------------------------------------------------
+	Compile time parameters:
+		THREAD_COUNT - The number of threads to launch per workgroup.
+------------------------------------------------------------------------------*/
+
+Buffer<int> InScan; // Input data to scan
+RWBuffer<int> OutScan; // Intermediate scan result (each block scanned)
+
+RWBuffer<int> BlockScans;
+Buffer<int> InBlockScans;
+
+/** Temporary thread group shared buffer for storing intermediate block scan results */
+groupshared int BlockRes[THREAD_COUNT];
+
+[numthreads(THREAD_COUNT,1,1)]
+void ScanInBlock(uint BlockIndex : SV_GroupID, uint InBlockIndex : SV_GroupThreadID, uint InScanIndex : SV_DispatchThreadID)
+{
+	// Copy the data into shared memory
+    BlockRes[InBlockIndex] = InScan[InScanIndex];
+
+	// UPSWEEP 
+	// n = THREAD_COUNT
+    int d;
+    for (d = 2; d < THREAD_COUNT; d <<= 1)
+    {
+        GroupMemoryBarrierWithGroupSync();
+        if (InBlockIndex % d == 0)
+        {
+			// k = InBlockIndex
+            int a = InBlockIndex + d - 1;
+            int b = InBlockIndex + (d >> 1) - 1;
+            BlockRes[a] += BlockRes[b];
+        }
+    }
+
+	// Zero last element
+    if (InBlockIndex == 0)
+    {
+        BlockRes[THREAD_COUNT - 1] = 0;
+    }
+	// DOWNSWEEP
+    for (d = THREAD_COUNT; d >= 2; d >>= 1)
+    {
+        GroupMemoryBarrierWithGroupSync();
+        if (InBlockIndex % d == 0)
+        {
+            int a = InBlockIndex + (d >> 1) - 1;
+            int b = InBlockIndex + d - 1;
+            int t = BlockRes[a];
+            BlockRes[a] = BlockRes[b];
+            BlockRes[b] += t;
+        }
+    }
+	
+	// Write scan to PrefixSum
+    OutScan[InScanIndex] = BlockRes[InBlockIndex];
+    //OutScan[InScanIndex] = InScan[InScanIndex];
+	// Write block scan result to auxiliary array
+    if (InBlockIndex == THREAD_COUNT - 1)
+    {
+        BlockScans[BlockIndex] = BlockRes[THREAD_COUNT - 1] + InScan[InScanIndex];
+    }
+}
+
+[numthreads(THREAD_COUNT,1,1)]
+void AddBlockScansToFinalScan(uint BlockIndex : SV_GroupID, uint InScanIndex : SV_DispatchThreadID)
+{
+    OutScan[InScanIndex] = InScan[InScanIndex] + InBlockScans[BlockIndex];
+}
\ No newline at end of file
diff --git a/Engine/Plugins/FX/Niagara/Source/NiagaraVertexFactories/Private/NiagaraSpatialHashBuild.cpp b/Engine/Plugins/FX/Niagara/Source/NiagaraVertexFactories/Private/NiagaraSpatialHashBuild.cpp
new file mode 100644
--- /dev/null
+++ b/Engine/Plugins/FX/Niagara/Source/NiagaraVertexFactories/Private/NiagaraSpatialHashBuild.cpp
@@ -0,0 +1,291 @@
+// Copyright 1998-2019 Epic Games, Inc. All Rights Reserved.
+
+#include "NiagaraSpatialHashBuild.h"
+
+IMPLEMENT_GLOBAL_SHADER(FNiagaraPrefixSumCS, "/Plugin/FX/Niagara/Private/NiagaraPrefixSum.usf", "ScanInBlock", SF_Compute);
+
+void FNiagaraPrefixSumCS::ModifyCompilationEnvironment(const FGlobalShaderPermutationParameters& Parameters, FShaderCompilerEnvironment& OutEnvironment)
+{
+	FGlobalShader::ModifyCompilationEnvironment(Parameters, OutEnvironment);
+	OutEnvironment.SetDefine(TEXT("THREAD_COUNT"), NIAGARA_SPATIAL_HASH_THREAD_COUNT);
+}
+
+FNiagaraPrefixSumCS::FNiagaraPrefixSumCS(const ShaderMetaType::CompiledShaderInitializerType& Initializer)
+	: FGlobalShader(Initializer)
+{
+	InScanBuffer.Bind(Initializer.ParameterMap, TEXT("InScan"));
+	OutScanBuffer.Bind(Initializer.ParameterMap, TEXT("OutScan"));
+	BlockScans.Bind(Initializer.ParameterMap, TEXT("BlockScans"));
+}
+
+bool FNiagaraPrefixSumCS::Serialize(FArchive& Ar)
+{
+	bool bShaderHasOutdatedParameters = FGlobalShader::Serialize(Ar);
+	Ar << InScanBuffer;
+	Ar << OutScanBuffer;
+	Ar << BlockScans;
+	return bShaderHasOutdatedParameters;
+}
+
+void FNiagaraPrefixSumCS::SetOutput(FRHICommandList& RHICmdList, FRHIUnorderedAccessView* OutScanUAV, FRHIUnorderedAccessView* OutBlockScansUAV)
+{
+	FRHIComputeShader* ComputeShaderRHI = GetComputeShader();
+	if (OutScanBuffer.IsBound())
+	{
+		RHICmdList.SetUAVParameter(ComputeShaderRHI, OutScanBuffer.GetBaseIndex(), OutScanUAV);
+	}
+	if (BlockScans.IsBound())
+	{
+		RHICmdList.SetUAVParameter(ComputeShaderRHI, BlockScans.GetBaseIndex(), OutBlockScansUAV);
+	}
+}
+
+void FNiagaraPrefixSumCS::SetParameters(FRHICommandList& RHICmdList, FRHIShaderResourceView* InScanUAV)
+{
+	FRHIComputeShader* ComputeShaderRHI = GetComputeShader();
+	if (InScanBuffer.IsBound())
+	{
+		RHICmdList.SetShaderResourceViewParameter(ComputeShaderRHI, InScanBuffer.GetBaseIndex(), InScanUAV);
+	}
+}
+
+void FNiagaraPrefixSumCS::UnbindBuffers(FRHICommandList& RHICmdList)
+{
+	FRHIComputeShader* ComputeShaderRHI = GetComputeShader();
+	if (InScanBuffer.IsBound())
+	{
+		RHICmdList.SetShaderResourceViewParameter(ComputeShaderRHI, InScanBuffer.GetBaseIndex(), FShaderResourceViewRHIParamRef());
+	}
+	if (OutScanBuffer.IsBound())
+	{
+		RHICmdList.SetUAVParameter(ComputeShaderRHI, OutScanBuffer.GetBaseIndex(), FUnorderedAccessViewRHIParamRef());
+	}
+	if (BlockScans.IsBound())
+	{
+		RHICmdList.SetUAVParameter(ComputeShaderRHI, BlockScans.GetBaseIndex(), FUnorderedAccessViewRHIParamRef());
+	}
+}
+
+IMPLEMENT_GLOBAL_SHADER(FNiagaraScanAddBlockResultsCS, "/Plugin/Fx/Niagara/Private/NiagaraPrefixSum.usf", "AddBlockScansToFinalScan", SF_Compute);
+
+void FNiagaraScanAddBlockResultsCS::ModifyCompilationEnvironment(const FGlobalShaderPermutationParameters& Parameters, FShaderCompilerEnvironment& OutEnvironment)
+{
+	FGlobalShader::ModifyCompilationEnvironment(Parameters, OutEnvironment);
+	OutEnvironment.SetDefine(TEXT("THREAD_COUNT"), NIAGARA_SPATIAL_HASH_THREAD_COUNT);
+}
+
+FNiagaraScanAddBlockResultsCS::FNiagaraScanAddBlockResultsCS(const ShaderMetaType::CompiledShaderInitializerType& Initializer)
+	: FGlobalShader(Initializer)
+{
+	IntermediateScan.Bind(Initializer.ParameterMap, TEXT("InScan"));
+	BlockScans.Bind(Initializer.ParameterMap, TEXT("InBlockScans"));
+	FinalScan.Bind(Initializer.ParameterMap, TEXT("OutScan"));
+}
+
+bool FNiagaraScanAddBlockResultsCS::Serialize(FArchive& Ar)
+{
+	bool bShaderHasOutdatedParameters = FGlobalShader::Serialize(Ar);
+	Ar << IntermediateScan;
+	Ar << BlockScans;
+	Ar << FinalScan;
+	return bShaderHasOutdatedParameters;
+}
+
+void FNiagaraScanAddBlockResultsCS::SetOutput(FRHICommandList& RHICmdList, FRHIUnorderedAccessView* OutFinalScan)
+{
+	FRHIComputeShader* ComputeShaderRHI = GetComputeShader();
+	if (FinalScan.IsBound())
+	{
+		RHICmdList.SetUAVParameter(ComputeShaderRHI, FinalScan.GetBaseIndex(), OutFinalScan);
+	}
+}
+
+void FNiagaraScanAddBlockResultsCS::SetParameters(FRHICommandList& RHICmdList, FRHIShaderResourceView* InIntermediateScanSRV, FRHIShaderResourceView* InBlockScanSRV)
+{
+	FRHIComputeShader* ComputeShaderRHI = GetComputeShader();
+	if (IntermediateScan.IsBound())
+	{
+		RHICmdList.SetShaderResourceViewParameter(ComputeShaderRHI, IntermediateScan.GetBaseIndex(), InIntermediateScanSRV);
+	}
+	if (BlockScans.IsBound())
+	{
+		RHICmdList.SetShaderResourceViewParameter(ComputeShaderRHI, BlockScans.GetBaseIndex(), InBlockScanSRV);
+	}
+}
+
+void FNiagaraScanAddBlockResultsCS::UnbindBuffers(FRHICommandList& RHICmdList)
+{
+	FRHIComputeShader* ComputeShaderRHI = GetComputeShader();
+	if (IntermediateScan.IsBound())
+	{
+		RHICmdList.SetShaderResourceViewParameter(ComputeShaderRHI, IntermediateScan.GetBaseIndex(), FShaderResourceViewRHIParamRef());
+	}
+	if (BlockScans.IsBound())
+	{
+		RHICmdList.SetShaderResourceViewParameter(ComputeShaderRHI, BlockScans.GetBaseIndex(), FShaderResourceViewRHIParamRef());
+	}
+	if (FinalScan.IsBound())
+	{
+		RHICmdList.SetUAVParameter(ComputeShaderRHI, FinalScan.GetBaseIndex(), FUnorderedAccessViewRHIParamRef());
+	}
+}
+
+IMPLEMENT_GLOBAL_SHADER(FNiagaraCountingSortCS, "/Plugin/FX/Niagara/Private/NiagaraDISpatialHashCountingSort.usf", "SpatialHashCountingSortCS", SF_Compute);
+
+void FNiagaraCountingSortCS::ModifyCompilationEnvironment(const FGlobalShaderPermutationParameters& Parameters, FShaderCompilerEnvironment& OutEnvironment)
+{
+	FGlobalShader::ModifyCompilationEnvironment(Parameters, OutEnvironment);
+	OutEnvironment.SetDefine(TEXT("THREAD_COUNT"), NIAGARA_SPATIAL_HASH_THREAD_COUNT);
+}
+
+FNiagaraCountingSortCS::FNiagaraCountingSortCS(const ShaderMetaType::CompiledShaderInitializerType& Initializer)
+	: FGlobalShader(Initializer)
+{
+	DataToSortID.Bind(Initializer.ParameterMap, TEXT("InDataToSortID"));
+	DataToSortPos.Bind(Initializer.ParameterMap, TEXT("InDataToSortPos"));
+	SortedDataID.Bind(Initializer.ParameterMap, TEXT("OutSortedDataID"));
+	SortedDataPos.Bind(Initializer.ParameterMap, TEXT("OutSortedDataPos"));
+	ScannedCounts.Bind(Initializer.ParameterMap, TEXT("ScannedCounts"));
+	NumParticles.Bind(Initializer.ParameterMap, TEXT("N"));
+}
+
+bool FNiagaraCountingSortCS::Serialize(FArchive& Ar)
+{
+	bool bShaderHasOutdatedParameters = FGlobalShader::Serialize(Ar);
+	Ar << DataToSortID;
+	Ar << DataToSortPos;
+	Ar << SortedDataID;
+	Ar << SortedDataPos;
+	Ar << ScannedCounts;
+	Ar << NumParticles;
+	return bShaderHasOutdatedParameters;
+}
+
+void FNiagaraCountingSortCS::SetOutput(FRHICommandList& RHICmdList, FRHIUnorderedAccessView* OutSortedDataID, FRHIUnorderedAccessView* OutSortedDataPos, FRHIUnorderedAccessView* OutScannedCounts)
+{
+	FRHIComputeShader* ComputeShaderRHI = GetComputeShader();
+	if (SortedDataID.IsBound())
+	{
+		RHICmdList.SetUAVParameter(ComputeShaderRHI, SortedDataID.GetBaseIndex(), OutSortedDataID);
+	}
+	if (SortedDataPos.IsBound())
+	{
+		RHICmdList.SetUAVParameter(ComputeShaderRHI, SortedDataPos.GetBaseIndex(), OutSortedDataPos);
+	}
+	if (ScannedCounts.IsBound())
+	{
+		RHICmdList.SetUAVParameter(ComputeShaderRHI, ScannedCounts.GetBaseIndex(), OutScannedCounts);
+	}
+}
+
+void FNiagaraCountingSortCS::SetParameters(FRHICommandList& RHICmdList, FRHIShaderResourceView* InDataToSortID, FRHIShaderResourceView* InDataToSortPos, int InNumParticles)
+{
+	FRHIComputeShader* ComputeShaderRHI = GetComputeShader();
+	if (DataToSortID.IsBound())
+	{
+		RHICmdList.SetShaderResourceViewParameter(ComputeShaderRHI, DataToSortID.GetBaseIndex(), InDataToSortID);
+	}
+	if (DataToSortPos.IsBound())
+	{
+		RHICmdList.SetShaderResourceViewParameter(ComputeShaderRHI, DataToSortPos.GetBaseIndex(), InDataToSortPos);
+	}
+	if (NumParticles.IsBound())
+	{
+		RHICmdList.SetShaderParameter(ComputeShaderRHI, NumParticles.GetBufferIndex(), NumParticles.GetBaseIndex(), sizeof(int32), &InNumParticles);
+	}
+}
+
+void FNiagaraCountingSortCS::UnbindBuffers(FRHICommandList& RHICmdList)
+{
+	FRHIComputeShader* ComputeShaderRHI = GetComputeShader();
+	if (SortedDataID.IsBound())
+	{
+		RHICmdList.SetUAVParameter(ComputeShaderRHI, SortedDataID.GetBaseIndex(), FUnorderedAccessViewRHIParamRef());
+	}
+	if (SortedDataPos.IsBound())
+	{
+		RHICmdList.SetUAVParameter(ComputeShaderRHI, SortedDataPos.GetBaseIndex(), FUnorderedAccessViewRHIParamRef());
+	}
+	if (ScannedCounts.IsBound())
+	{
+		RHICmdList.SetUAVParameter(ComputeShaderRHI, ScannedCounts.GetBaseIndex(), FUnorderedAccessViewRHIParamRef());
+	}
+	if (DataToSortID.IsBound())
+	{
+		RHICmdList.SetShaderResourceViewParameter(ComputeShaderRHI, DataToSortID.GetBaseIndex(), FShaderResourceViewRHIParamRef());
+	}
+	if (DataToSortPos.IsBound())
+	{
+		RHICmdList.SetShaderResourceViewParameter(ComputeShaderRHI, DataToSortPos.GetBaseIndex(), FShaderResourceViewRHIParamRef());
+	}
+}
+
+IMPLEMENT_GLOBAL_SHADER(FNiagaraSpatialHashIndexCellsCS, "/Plugin/FX/Niagara/Private/NiagaraDISpatialHashIndexCells.usf", "IndexCellsCS", SF_Compute);
+
+void FNiagaraSpatialHashIndexCellsCS::ModifyCompilationEnvironment(const FGlobalShaderPermutationParameters& Parameters, FShaderCompilerEnvironment& OutEnvironment)
+{
+	FGlobalShader::ModifyCompilationEnvironment(Parameters, OutEnvironment);
+	OutEnvironment.SetDefine(TEXT("THREAD_COUNT"), NIAGARA_SPATIAL_HASH_THREAD_COUNT);
+}
+
+FNiagaraSpatialHashIndexCellsCS::FNiagaraSpatialHashIndexCellsCS(const ShaderMetaType::CompiledShaderInitializerType& Initializer)
+	: FGlobalShader(Initializer)
+{
+	CellStartIndices.Bind(Initializer.ParameterMap, TEXT("CellStartIndices"));
+	CellEndIndices.Bind(Initializer.ParameterMap, TEXT("CellEndIndices"));
+	SortedIDs.Bind(Initializer.ParameterMap, TEXT("InParticleIDs"));
+	NumParticles.Bind(Initializer.ParameterMap, TEXT("N"));
+}
+
+bool FNiagaraSpatialHashIndexCellsCS::Serialize(FArchive& Ar)
+{
+	bool bShaderHasOutdatedParameters = FGlobalShader::Serialize(Ar);
+	Ar << CellStartIndices;
+	Ar << CellEndIndices;
+	Ar << SortedIDs;
+	Ar << NumParticles;
+	return bShaderHasOutdatedParameters;
+}
+
+void FNiagaraSpatialHashIndexCellsCS::SetOutput(FRHICommandList& RHICmdList, FRHIUnorderedAccessView* OutCellStartIndices, FRHIUnorderedAccessView* OutCellEndIndices)
+{
+	FRHIComputeShader* ComputeShaderRHI = GetComputeShader();
+	if (CellStartIndices.IsBound())
+	{
+		RHICmdList.SetUAVParameter(ComputeShaderRHI, CellStartIndices.GetBaseIndex(), OutCellStartIndices);
+	}
+	if (CellEndIndices.IsBound())
+	{
+		RHICmdList.SetUAVParameter(ComputeShaderRHI, CellEndIndices.GetBaseIndex(), OutCellEndIndices);
+	}
+}
+
+void FNiagaraSpatialHashIndexCellsCS::SetParameters(FRHICommandList& RHICmdList, FRHIShaderResourceView* InSortedIDs, int InNumParticles)
+{
+	FRHIComputeShader* ComputeShaderRHI = GetComputeShader();
+	if (SortedIDs.IsBound())
+	{
+		RHICmdList.SetShaderResourceViewParameter(ComputeShaderRHI, SortedIDs.GetBaseIndex(), InSortedIDs);
+	}
+	if (NumParticles.IsBound())
+	{
+		RHICmdList.SetShaderParameter(ComputeShaderRHI, NumParticles.GetBufferIndex(), NumParticles.GetBaseIndex(), sizeof(int32), &InNumParticles);
+	}
+}
+
+void FNiagaraSpatialHashIndexCellsCS::UnbindBuffers(FRHICommandList& RHICmdList)
+{
+	FRHIComputeShader* ComputeShaderRHI = GetComputeShader();
+	if (CellStartIndices.IsBound())
+	{
+		RHICmdList.SetUAVParameter(ComputeShaderRHI, CellStartIndices.GetBaseIndex(), FUnorderedAccessViewRHIParamRef());
+	}
+	if (CellEndIndices.IsBound())
+	{
+		RHICmdList.SetUAVParameter(ComputeShaderRHI, CellEndIndices.GetBaseIndex(), FUnorderedAccessViewRHIParamRef());
+	}
+	if (SortedIDs.IsBound())
+	{
+		RHICmdList.SetShaderResourceViewParameter(ComputeShaderRHI, SortedIDs.GetBaseIndex(), FShaderResourceViewRHIParamRef());
+	}
+}
diff --git a/Engine/Plugins/FX/Niagara/Source/NiagaraVertexFactories/Public/NiagaraSpatialHashBuild.h b/Engine/Plugins/FX/Niagara/Source/NiagaraVertexFactories/Public/NiagaraSpatialHashBuild.h
new file mode 100644
--- /dev/null
+++ b/Engine/Plugins/FX/Niagara/Source/NiagaraVertexFactories/Public/NiagaraSpatialHashBuild.h
@@ -0,0 +1,143 @@
+// Copyright 1998-2019 Epic Games, Inc. All Rights Reserved.
+
+#pragma once
+
+#include "CoreMinimal.h"
+#include "GlobalShader.h"
+#include "ShaderPermutation.h"
+
+#define NIAGARA_SPATIAL_HASH_THREAD_COUNT 128
+
+/**
+ * Compute shader used to perform a prefix sum on a buffer.
+ */
+class NIAGARAVERTEXFACTORIES_API FNiagaraPrefixSumCS : public FGlobalShader
+{
+	DECLARE_GLOBAL_SHADER(FNiagaraPrefixSumCS)
+public:
+	static bool ShouldCompilePermutation(const FGlobalShaderPermutationParameters& Parameters)
+	{
+		return RHISupportsComputeShaders(Parameters.Platform);
+	}
+
+	static void ModifyCompilationEnvironment(const FGlobalShaderPermutationParameters& Parameters, FShaderCompilerEnvironment& OutEnvironment);
+
+	FNiagaraPrefixSumCS() {}
+
+	explicit FNiagaraPrefixSumCS(const ShaderMetaType::CompiledShaderInitializerType& Initializer);
+
+	virtual bool Serialize(FArchive& Ar) override;
+
+	void SetOutput(FRHICommandList& RHICmdList, FRHIUnorderedAccessView* OutScanUAV, FRHIUnorderedAccessView* OutBlockScansUAV);
+
+	void SetParameters(FRHICommandList& RHICmdList, FRHIShaderResourceView* InScanSRV);
+
+	void UnbindBuffers(FRHICommandList& RHICmdList);
+
+private:
+	FShaderResourceParameter InScanBuffer;
+	FShaderResourceParameter OutScanBuffer;
+	FShaderResourceParameter BlockScans;
+};
+
+/**
+ * Compute shader used to perform the final prefix sum pass (adding the block scan results to each element).
+ */
+class NIAGARAVERTEXFACTORIES_API FNiagaraScanAddBlockResultsCS : public FGlobalShader
+{
+	DECLARE_GLOBAL_SHADER(FNiagaraScanAddBlockResultsCS)
+public:
+	static bool ShouldCompilePermutation(const FGlobalShaderPermutationParameters& Parameters)
+	{
+		return RHISupportsComputeShaders(Parameters.Platform);
+	}
+
+	static void ModifyCompilationEnvironment(const FGlobalShaderPermutationParameters& Parameters, FShaderCompilerEnvironment& OutEnvironment);
+
+	FNiagaraScanAddBlockResultsCS() {}
+
+	explicit FNiagaraScanAddBlockResultsCS(const ShaderMetaType::CompiledShaderInitializerType& Initializer);
+
+	virtual bool Serialize(FArchive& Ar) override;
+
+	void SetOutput(FRHICommandList& RHICmdList, FRHIUnorderedAccessView* OutFinalScan);
+
+	void SetParameters(FRHICommandList& RHICmdList, FRHIShaderResourceView* InIntermediateScanSRV, FRHIShaderResourceView* InBlockScanSRV);
+
+	void UnbindBuffers(FRHICommandList& RHICmdList);
+
+private:
+	FShaderResourceParameter IntermediateScan;
+	FShaderResourceParameter BlockScans;
+	FShaderResourceParameter FinalScan;
+};
+
+/**
+ * Compute shader used to sort spatial has particle IDs and positions based on a scanned cell count buffer.
+ */
+class NIAGARAVERTEXFACTORIES_API FNiagaraCountingSortCS : public FGlobalShader
+{
+	DECLARE_GLOBAL_SHADER(FNiagaraCountingSortCS)
+public:
+	static bool ShouldCompilePermutation(const FGlobalShaderPermutationParameters& Parameters)
+	{
+		return RHISupportsComputeShaders(Parameters.Platform);
+	}
+
+	static void ModifyCompilationEnvironment(const FGlobalShaderPermutationParameters& Parameters, FShaderCompilerEnvironment& OutEnvironment);
+
+	FNiagaraCountingSortCS() {}
+
+	explicit FNiagaraCountingSortCS(const ShaderMetaType::CompiledShaderInitializerType& Initializer);
+
+	virtual bool Serialize(FArchive& Ar) override;
+
+	void SetOutput(FRHICommandList& RHICmdList, FRHIUnorderedAccessView* OutSortedDataID, FRHIUnorderedAccessView* OutSortedDataPos, FRHIUnorderedAccessView* ScannedCounts);
+
+	void SetParameters(FRHICommandList& RHICmdList, FRHIShaderResourceView* InDataToSortID, FRHIShaderResourceView* InDataToSortPos, int InNumParticles);
+
+	void UnbindBuffers(FRHICommandList& RHICmdList);
+
+private:
+	FShaderResourceParameter DataToSortID;
+	FShaderResourceParameter DataToSortPos;
+	FShaderResourceParameter SortedDataID;
+	FShaderResourceParameter SortedDataPos;
+	FShaderResourceParameter ScannedCounts;
+
+	FShaderParameter NumParticles;
+};
+
+/**
+ * Compute shader used to calculate the start and end indices of each cell.
+ */
+class NIAGARAVERTEXFACTORIES_API FNiagaraSpatialHashIndexCellsCS : public FGlobalShader
+{
+	DECLARE_GLOBAL_SHADER(FNiagaraSpatialHashIndexCellsCS)
+public:
+	static bool ShouldCompilePermutation(const FGlobalShaderPermutationParameters& Parameters)
+	{
+		return RHISupportsComputeShaders(Parameters.Platform);
+	}
+
+	static void ModifyCompilationEnvironment(const FGlobalShaderPermutationParameters& Parameters, FShaderCompilerEnvironment& OutEnvironment);
+
+	FNiagaraSpatialHashIndexCellsCS() {}
+
+	explicit FNiagaraSpatialHashIndexCellsCS(const ShaderMetaType::CompiledShaderInitializerType& Initializer);
+
+	virtual bool Serialize(FArchive& Ar) override;
+
+	void SetOutput(FRHICommandList& RHICmdList, FRHIUnorderedAccessView* OutCellStartIndices, FRHIUnorderedAccessView* OutCellEndIndices);
+
+	void SetParameters(FRHICommandList& RHICmdList, FRHIShaderResourceView* InSortedIDs, int InNumParticles);
+
+	void UnbindBuffers(FRHICommandList& RHICmdList);
+
+private:
+	FShaderResourceParameter CellStartIndices;
+	FShaderResourceParameter CellEndIndices;
+	FShaderResourceParameter SortedIDs;
+
+	FShaderParameter NumParticles;
+};
\ No newline at end of file
