#UE4 For now, finishing compilation for UNiagaraScript in IsCachedCookedPlatformDataLoaded until we can make sure compilation is finished in the cook commandlet asyncronously before serialize

diff --git a/Engine/Plugins/FX/Niagara/Source/Niagara/Private/NiagaraScript.cpp b/Engine/Plugins/FX/Niagara/Source/Niagara/Private/NiagaraScript.cpp
--- a/Engine/Plugins/FX/Niagara/Source/Niagara/Private/NiagaraScript.cpp
+++ b/Engine/Plugins/FX/Niagara/Source/Niagara/Private/NiagaraScript.cpp
@@ -1140,6 +1140,9 @@ bool UNiagaraScript::IsCachedCookedPlatformDataLoaded(const ITargetPlatform* Tar
 			{
 				for (const auto& MaterialResource : *CachedScriptResourcesForPlatform)
 				{
+					// For now, finish compilation here until we can make sure compilation is finished in the cook commandlet asyncronously before serialize
+					MaterialResource->FinishCompilation();
+
 					if (MaterialResource->IsCompilationFinished() == false)
 					{
 						return false;
