Niagara - Fix a crash when deleting modules from the system script.

diff --git a/Engine/Plugins/FX/Niagara/Source/NiagaraEditor/Private/ViewModels/Stack/NiagaraStackModuleItem.cpp b/Engine/Plugins/FX/Niagara/Source/NiagaraEditor/Private/ViewModels/Stack/NiagaraStackModuleItem.cpp
--- a/Engine/Plugins/FX/Niagara/Source/NiagaraEditor/Private/ViewModels/Stack/NiagaraStackModuleItem.cpp
+++ b/Engine/Plugins/FX/Niagara/Source/NiagaraEditor/Private/ViewModels/Stack/NiagaraStackModuleItem.cpp
@@ -84,7 +84,8 @@ UNiagaraNodeOutput* GetOutputNodeForModuleDependency(ENiagaraScriptUsage Dependa
 				if ((FMath::Abs(Distance) < ClosestDistance) && bCorrectOrder)
 				{
 					ClosestDistance = Distance;
-					OutputScript = FNiagaraEditorUtilities::GetScriptFromSystem(System, EmitterHandle->GetId(), PossibleUsage, FGuid());
+					FGuid EmitterHandleId = EmitterHandle != nullptr ? EmitterHandle->GetId() : FGuid();
+					OutputScript = FNiagaraEditorUtilities::GetScriptFromSystem(System, EmitterHandleId, PossibleUsage, FGuid());
 				}
 			}
 		}
@@ -92,7 +93,8 @@ UNiagaraNodeOutput* GetOutputNodeForModuleDependency(ENiagaraScriptUsage Dependa
 		{
 			if (SupportedUsages.Contains(DependantUsage))
 			{
-				OutputScript = FNiagaraEditorUtilities::GetScriptFromSystem(System, EmitterHandle->GetId(), DependantUsage, FGuid());
+				FGuid EmitterHandleId = EmitterHandle != nullptr ? EmitterHandle->GetId() : FGuid();
+				OutputScript = FNiagaraEditorUtilities::GetScriptFromSystem(System, EmitterHandleId, DependantUsage, FGuid());
 			}
 		}
 
@@ -559,7 +561,9 @@ void UNiagaraStackModuleItem::RefreshIssues(TArray<FStackIssue>& NewIssues)
 								if (bRequiredDependencyFound) // check for multiple dependents along the way, and stop adjacent to the last one
 								{
 									ENiagaraScriptUsage DependencyUsage = SystemModuleData[i].Usage;
-									const FNiagaraEmitterHandle* EmitterHandle = FNiagaraEditorUtilities::GetEmitterHandleForEmitter(GetSystemViewModel()->GetSystem(), *GetEmitterViewModel()->GetEmitter());
+									const FNiagaraEmitterHandle* EmitterHandle = GetEmitterViewModel().IsValid()
+										? FNiagaraEditorUtilities::GetEmitterHandleForEmitter(GetSystemViewModel()->GetSystem(), *GetEmitterViewModel()->GetEmitter())
+										: nullptr;
 									UNiagaraNodeOutput* FoundTargetOutputNode = GetOutputNodeForModuleDependency(DependencyUsage, DependencyScript, GetSystemViewModel()->GetSystem(), EmitterHandle, Dependency);
 									if (FoundTargetOutputNode != nullptr)
 									{
@@ -780,11 +784,13 @@ void UNiagaraStackModuleItem::Delete()
 
 	FScopedTransaction ScopedTransaction(LOCTEXT("RemoveAModuleFromTheStack", "Remove a module from the stack"));
 
-	const FNiagaraEmitterHandle* EmitterHandle = FNiagaraEditorUtilities::GetEmitterHandleForEmitter(GetSystemViewModel()->GetSystem(), *GetEmitterViewModel()->GetEmitter());
-	checkf(EmitterHandle != nullptr, TEXT("Invalid Stack - Emitter handle could not be found for module"));
+	const FNiagaraEmitterHandle* EmitterHandle = GetEmitterViewModel().IsValid()
+		? FNiagaraEditorUtilities::GetEmitterHandleForEmitter(GetSystemViewModel()->GetSystem(), *GetEmitterViewModel()->GetEmitter())
+		: nullptr;
+	FGuid EmitterHandleId = EmitterHandle != nullptr ? EmitterHandle->GetId() : FGuid();
 
 	TArray<TWeakObjectPtr<UNiagaraNodeInput>> RemovedNodes;
-	bool bRemoved = FNiagaraStackGraphUtilities::RemoveModuleFromStack(GetSystemViewModel()->GetSystem(), EmitterHandle->GetId(), *FunctionCallNode, RemovedNodes);
+	bool bRemoved = FNiagaraStackGraphUtilities::RemoveModuleFromStack(GetSystemViewModel()->GetSystem(), EmitterHandleId, *FunctionCallNode, RemovedNodes);
 	if (bRemoved)
 	{
 		UNiagaraGraph* Graph = FunctionCallNode->GetNiagaraGraph();
