Fix missing editgrate from Main. Do not bind InstanceParameters to System Spawn and Update contexts.

diff --git a/Engine/Plugins/FX/Niagara/Source/Niagara/Private/NiagaraSystemInstance.cpp b/Engine/Plugins/FX/Niagara/Source/Niagara/Private/NiagaraSystemInstance.cpp
--- a/Engine/Plugins/FX/Niagara/Source/Niagara/Private/NiagaraSystemInstance.cpp
+++ b/Engine/Plugins/FX/Niagara/Source/Niagara/Private/NiagaraSystemInstance.cpp
@@ -914,8 +914,6 @@ void FNiagaraSystemInstance::BindParameters()
 	}
 
 	Component->GetOverrideParameters().Bind(&InstanceParameters);
-	InstanceParameters.Bind(&SystemSimulation->GetSpawnExecutionContext().Parameters);
-	InstanceParameters.Bind(&SystemSimulation->GetUpdateExecutionContext().Parameters);
 
 	if (SystemSimulation->GetIsSolo())
 	{
@@ -948,8 +946,6 @@ void FNiagaraSystemInstance::UnbindParameters()
 
 	if (SystemSimulation.IsValid())
 	{
-		InstanceParameters.Unbind(&SystemSimulation->GetSpawnExecutionContext().Parameters);
-		InstanceParameters.Unbind(&SystemSimulation->GetUpdateExecutionContext().Parameters);
 		if (SystemSimulation->GetIsSolo())
 		{
 			if (Component)
