Optimize Niagara Skeletal Mesh tick data generation when pre-skinning is not needed.

diff --git a/Engine/Plugins/FX/Niagara/Source/Niagara/Classes/NiagaraDataInterfaceSkeletalMesh.h b/Engine/Plugins/FX/Niagara/Source/Niagara/Classes/NiagaraDataInterfaceSkeletalMesh.h
--- a/Engine/Plugins/FX/Niagara/Source/Niagara/Classes/NiagaraDataInterfaceSkeletalMesh.h
+++ b/Engine/Plugins/FX/Niagara/Source/Niagara/Classes/NiagaraDataInterfaceSkeletalMesh.h
@@ -102,7 +102,7 @@ struct FSkeletalMeshSkinningData
 	bool IsUsed()const;
 	void ForceDataRefresh();
 
-	bool Tick(float InDeltaSeconds);
+	bool Tick(float InDeltaSeconds, bool bRequirePreskin = true);
 
 	FORCEINLINE FVector GetPosition(int32 LODIndex, int32 VertexIndex)const
 	{
@@ -169,7 +169,14 @@ private:
 
 class FNDI_SkeletalMesh_GeneratedData
 {
-	TMap<TWeakObjectPtr<USkeletalMeshComponent>, TSharedPtr<FSkeletalMeshSkinningData>> CachedSkinningData;
+	// Encapsulates skinning data and mesh usage information. 
+	// Set by GetCachedSkinningData and used by TickGeneratedData to determine whether we need to pre-skin or not.
+	struct CachedSkinningDataAndUsage {
+		TSharedPtr<FSkeletalMeshSkinningData> SkinningData;
+		FSkeletalMeshSkinningDataUsage Usage;
+	};
+	TMap<TWeakObjectPtr<USkeletalMeshComponent>, CachedSkinningDataAndUsage> CachedSkinningData;
+
 public:
 
 	FSkeletalMeshSkinningDataHandle GetCachedSkinningData(TWeakObjectPtr<USkeletalMeshComponent>& InComponent, FSkeletalMeshSkinningDataUsage Usage);
diff --git a/Engine/Plugins/FX/Niagara/Source/Niagara/Private/NiagaraDataInterfaceSkeletalMesh.cpp b/Engine/Plugins/FX/Niagara/Source/Niagara/Private/NiagaraDataInterfaceSkeletalMesh.cpp
--- a/Engine/Plugins/FX/Niagara/Source/Niagara/Private/NiagaraDataInterfaceSkeletalMesh.cpp
+++ b/Engine/Plugins/FX/Niagara/Source/Niagara/Private/NiagaraDataInterfaceSkeletalMesh.cpp
@@ -178,7 +178,7 @@ bool FSkeletalMeshSkinningData::IsUsed()const
 	return false;
 }
 
-bool FSkeletalMeshSkinningData::Tick(float InDeltaSeconds)
+bool FSkeletalMeshSkinningData::Tick(float InDeltaSeconds, bool bRequirePreskin)
 {
 	USkeletalMeshComponent* SkelComp = MeshComp.Get();
 	check(SkelComp);
@@ -196,24 +196,27 @@ bool FSkeletalMeshSkinningData::Tick(float InDeltaSeconds)
 		PrevBoneRefToLocals() = CurrBoneRefToLocals();
 	}
 
-	for (int32 LODIndex = 0; LODIndex < LODData.Num(); ++LODIndex)
+	if (bRequirePreskin)
 	{
-		FLODData& LOD = LODData[LODIndex];
-		if (LOD.PreSkinnedVertsUsers > 0)
+		for (int32 LODIndex = 0; LODIndex < LODData.Num(); ++LODIndex)
 		{
-			//TODO: If we pass the sections in the usage too, we can probably skin a minimal set of verts just for the used regions.
-			FSkeletalMeshLODRenderData& SkelMeshLODData = SkelComp->SkeletalMesh->GetResourceForRendering()->LODRenderData[LODIndex];
-			FSkinWeightVertexBuffer* SkinWeightBuffer = SkelComp->GetSkinWeightBuffer(LODIndex);
-			USkeletalMeshComponent::ComputeSkinnedPositions(SkelComp, CurrSkinnedPositions(LODIndex), CurrBoneRefToLocals(), SkelMeshLODData, *SkinWeightBuffer);
-			//check(CurrSkinnedPositions(LODIndex).Num() == SkelMeshLODData.NumVertices);
-			//Prime the previous positions if they're missing
-			if (PrevSkinnedPositions(LODIndex).Num() != CurrSkinnedPositions(LODIndex).Num())
+			FLODData& LOD = LODData[LODIndex];
+			if (LOD.PreSkinnedVertsUsers > 0)
 			{
-				PrevSkinnedPositions(LODIndex) = CurrSkinnedPositions(LODIndex);
+				//TODO: If we pass the sections in the usage too, we can probably skin a minimal set of verts just for the used regions.
+				FSkeletalMeshLODRenderData& SkelMeshLODData = SkelComp->SkeletalMesh->GetResourceForRendering()->LODRenderData[LODIndex];
+				FSkinWeightVertexBuffer* SkinWeightBuffer = SkelComp->GetSkinWeightBuffer(LODIndex);
+				USkeletalMeshComponent::ComputeSkinnedPositions(SkelComp, CurrSkinnedPositions(LODIndex), CurrBoneRefToLocals(), SkelMeshLODData, *SkinWeightBuffer);
+				//check(CurrSkinnedPositions(LODIndex).Num() == SkelMeshLODData.NumVertices);
+				//Prime the previous positions if they're missing
+				if (PrevSkinnedPositions(LODIndex).Num() != CurrSkinnedPositions(LODIndex).Num())
+				{
+					PrevSkinnedPositions(LODIndex) = CurrSkinnedPositions(LODIndex);
+				}
 			}
 		}
 	}
-	
+
 	bForceDataRefresh = false;
 	return true;
 }
@@ -228,15 +231,16 @@ FSkeletalMeshSkinningDataHandle FNDI_SkeletalMesh_GeneratedData::GetCachedSkinni
 	check(Component);
 	TSharedPtr<FSkeletalMeshSkinningData> SkinningData = nullptr;
 
-	if (TSharedPtr<FSkeletalMeshSkinningData>* Existing = CachedSkinningData.Find(Component))
+	if (CachedSkinningDataAndUsage* Existing = CachedSkinningData.Find(Component))
 	{
-		check(Existing->IsValid());//We shouldn't be able to have an invalid ptr here.
-		SkinningData = *Existing;
+		check(Existing && Existing->SkinningData.IsValid());//We shouldn't be able to have an invalid ptr here.
+		SkinningData = Existing->SkinningData; // Reuse the old skinning data
+		Existing->Usage = Usage; // Overwrite the old Usage
 	}
 	else
 	{
 		SkinningData = MakeShared<FSkeletalMeshSkinningData>(InComponent);
-		CachedSkinningData.Add(Component) = SkinningData;
+		CachedSkinningData.Add(Component) = { SkinningData, Usage };
 	}
 
 	return FSkeletalMeshSkinningDataHandle(Usage, SkinningData);
@@ -248,34 +252,53 @@ void FNDI_SkeletalMesh_GeneratedData::TickGeneratedData(float DeltaSeconds)
 	SCOPE_CYCLE_COUNTER(STAT_NiagaraSkel_PreSkin);
 
 	//Tick skinning data.
-	{
-		TArray<TWeakObjectPtr<USkeletalMeshComponent>, TInlineAllocator<64>> ToRemove;
-		TArray<FSkeletalMeshSkinningData*> ToTick;
-		ToTick.Reserve(CachedSkinningData.Num());
-		for (TPair<TWeakObjectPtr<USkeletalMeshComponent>, TSharedPtr<FSkeletalMeshSkinningData>>& Pair : CachedSkinningData)
+	TArray<TWeakObjectPtr<USkeletalMeshComponent>, TInlineAllocator<64>> ToRemove;
+	TArray<FSkeletalMeshSkinningData*> ToTickBonesOnly;
+	TArray<FSkeletalMeshSkinningData*> ToTickPreskin;
+	ToTickBonesOnly.Reserve(CachedSkinningData.Num());
+	ToTickPreskin.Reserve(CachedSkinningData.Num());
+	for (TPair<TWeakObjectPtr<USkeletalMeshComponent>, CachedSkinningDataAndUsage>& Pair : CachedSkinningData)
+	{
+		TSharedPtr<FSkeletalMeshSkinningData>& Ptr = Pair.Value.SkinningData;
+			
+		if (Ptr.IsUnique() || !Pair.Key.Get() || !Ptr->IsUsed())
+		{
+			ToRemove.Add(Pair.Key);//Remove unused skin data or for those with GCd components as we go.
+		}
+		else
 		{
-			TSharedPtr<FSkeletalMeshSkinningData>& Ptr = Pair.Value;
+			FSkeletalMeshSkinningDataUsage Usage = Pair.Value.Usage;
 			FSkeletalMeshSkinningData* SkinData = Ptr.Get();
-			USkeletalMeshComponent* Component = Pair.Key.Get();
 			check(SkinData);
-			if (Ptr.IsUnique() || !Component || !Ptr->IsUsed())
+
+			if (Usage.NeedPreSkinnedVerts())
 			{
-				ToRemove.Add(Pair.Key);//Remove unused skin data or for those with GCd components as we go.
+				ToTickPreskin.Add(SkinData);
 			}
 			else
 			{
-				ToTick.Add(SkinData);
+				ToTickBonesOnly.Add(SkinData);
 			}
 		}
+	}
 
-		for (TWeakObjectPtr<USkeletalMeshComponent> Key : ToRemove)
-		{
-			CachedSkinningData.Remove(Key);
-		}
-
-		ParallelFor(ToTick.Num(), [&](int32 Index)
+	for (TWeakObjectPtr<USkeletalMeshComponent> Key : ToRemove)
+	{
+		CachedSkinningData.Remove(Key);
+	}
+	
+	// First tick the meshes that don't need pre-skinning. 
+	// This prevents additional threading overhead when we don't need to pre-skin.
+	for (int i = 0; i < ToTickBonesOnly.Num(); i++) {
+		ToTickBonesOnly[i]->Tick(DeltaSeconds, false);
+	}
+		
+	// Then tick the remaining meshes requiring pre-skinning in parallel
+	if (ToTickPreskin.Num() != 0)
+	{
+		ParallelFor(ToTickPreskin.Num(), [&](int32 Index)
 		{
-			ToTick[Index]->Tick(DeltaSeconds);
+			ToTickPreskin[Index]->Tick(DeltaSeconds, true);
 		});
 	}
 }
