Fixed Niagara texture data interface possibly not binding the good RHI texture.

diff --git a/Engine/Plugins/FX/Niagara/Source/Niagara/Classes/NiagaraDataInterfaceTexture.h b/Engine/Plugins/FX/Niagara/Source/Niagara/Classes/NiagaraDataInterfaceTexture.h
--- a/Engine/Plugins/FX/Niagara/Source/Niagara/Classes/NiagaraDataInterfaceTexture.h
+++ b/Engine/Plugins/FX/Niagara/Source/Niagara/Classes/NiagaraDataInterfaceTexture.h
@@ -62,7 +62,7 @@ protected:
 struct FNiagaraDataInterfaceProxyTexture : public FNiagaraDataInterfaceProxy
 {
 	FSamplerStateRHIRef SamplerStateRHI;
-	FTextureRHIRef TextureRHI;
+	FTextureReferenceRHIRef	TextureReferenceRHI;
 	float TexDims[2];
 
 	virtual void ConsumePerInstanceDataFromGameThread(void* PerInstanceData, const FGuid& Instance) override { check(false); }
diff --git a/Engine/Plugins/FX/Niagara/Source/Niagara/Private/NiagaraDataInterfaceTexture.cpp b/Engine/Plugins/FX/Niagara/Source/Niagara/Private/NiagaraDataInterfaceTexture.cpp
--- a/Engine/Plugins/FX/Niagara/Source/Niagara/Private/NiagaraDataInterfaceTexture.cpp
+++ b/Engine/Plugins/FX/Niagara/Source/Niagara/Private/NiagaraDataInterfaceTexture.cpp
@@ -49,6 +49,8 @@ void UNiagaraDataInterfaceTexture::PostLoad()
 		}
 	}
 #endif
+	// Not safe since the UTexture might not have yet PostLoad() called and so UpdateResource() called.
+	// This will affect whether the SamplerStateRHI will be available or not.
 	PushToRenderThread();
 }
 
@@ -385,23 +387,42 @@ struct FNiagaraDataInterfaceParametersCS_Texture : public FNiagaraDataInterfaceP
 
 		FRHIComputeShader* ComputeShaderRHI = Context.Shader->GetComputeShader();
 		FNiagaraDataInterfaceProxyTexture* TextureDI = static_cast<FNiagaraDataInterfaceProxyTexture*>(Context.DataInterface);
-		if (!TextureDI->TextureRHI)
+
+		FRHITexture* TextureRHI = (TextureDI && TextureDI->TextureReferenceRHI) ? TextureDI->TextureReferenceRHI->GetReferencedTexture() : nullptr;
+		if (TextureRHI)
+		{
+			FRHITexture* TextureRHI = TextureDI->TextureReferenceRHI->GetReferencedTexture();
+			FRHISamplerState* SamplerStateRHI = TextureDI->SamplerStateRHI;
+			
+			if (!SamplerStateRHI)
+			{
+				// Fallback required because PostLoad() order affects whether RHI resources 
+				// are initalized in UNiagaraDataInterfaceTexture::PushToRenderThread().
+				SamplerStateRHI = TStaticSamplerState<SF_Point, AM_Clamp, AM_Clamp, AM_Clamp>::GetRHI();
+			}
+			SetTextureParameter(
+				RHICmdList,
+				ComputeShaderRHI,
+				TextureParam,
+				SamplerParam,
+				SamplerStateRHI,
+				TextureRHI
+			);
+			SetShaderValue(RHICmdList, ComputeShaderRHI, Dimensions, TextureDI->TexDims);
+		}
+		else
 		{
+			SetTextureParameter(
+				RHICmdList,
+				ComputeShaderRHI,
+				TextureParam,
+				SamplerParam,
+				GBlackTexture->SamplerStateRHI,
+				GBlackTexture->TextureRHI
+			);
 			float TexDims[] = { 0.0f, 0.0f };
 			SetShaderValue(RHICmdList, ComputeShaderRHI, Dimensions, TexDims);
-			return;
 		}
-
-		SetTextureParameter(
-			RHICmdList,
-			ComputeShaderRHI,
-			TextureParam,
-			SamplerParam,
-			TextureDI->SamplerStateRHI,
-			TextureDI->TextureRHI
-		);
-
-		SetShaderValue(RHICmdList, ComputeShaderRHI, Dimensions, TextureDI->TexDims);
 	}
 	
 
@@ -418,14 +439,14 @@ void UNiagaraDataInterfaceTexture::PushToRenderThread()
 
 	float DimX = 0.0f;
 	float DimY = 0.0f;
-	FTextureRHIRef RT_Texture;
+	FTextureReferenceRHIRef  RT_TextureReference;
 	FSamplerStateRHIRef RT_SamplerState;
 
 	if (Texture && Texture->TextureReference.TextureReferenceRHI.IsValid())
 	{
 		if (Texture->TextureReference.TextureReferenceRHI)
 		{
-			RT_Texture = Texture->TextureReference.TextureReferenceRHI->GetReferencedTexture();
+			RT_TextureReference = Texture->TextureReference.TextureReferenceRHI;
 		}
 		if (Texture->Resource)
 		{
@@ -436,9 +457,9 @@ void UNiagaraDataInterfaceTexture::PushToRenderThread()
 	}
 
 	ENQUEUE_RENDER_COMMAND(FPushDITextureToRT) (
-		[RT_Proxy, RT_Texture, RT_SamplerState, DimX, DimY](FRHICommandListImmediate& RHICmdList)
+		[RT_Proxy, RT_TextureReference, RT_SamplerState, DimX, DimY](FRHICommandListImmediate& RHICmdList)
 		{
-			RT_Proxy->TextureRHI = RT_Texture;
+			RT_Proxy->TextureReferenceRHI = RT_TextureReference;
 			RT_Proxy->SamplerStateRHI = RT_SamplerState;
 			RT_Proxy->TexDims[0] = DimX;
 			RT_Proxy->TexDims[1] = DimY;
