Fixed compilation warning with CL 8153422

diff --git a/Engine/Plugins/FX/Niagara/Source/Niagara/Private/NiagaraDataInterfaceTexture.cpp b/Engine/Plugins/FX/Niagara/Source/Niagara/Private/NiagaraDataInterfaceTexture.cpp
--- a/Engine/Plugins/FX/Niagara/Source/Niagara/Private/NiagaraDataInterfaceTexture.cpp
+++ b/Engine/Plugins/FX/Niagara/Source/Niagara/Private/NiagaraDataInterfaceTexture.cpp
@@ -391,9 +391,7 @@ struct FNiagaraDataInterfaceParametersCS_Texture : public FNiagaraDataInterfaceP
 		FRHITexture* TextureRHI = (TextureDI && TextureDI->TextureReferenceRHI) ? TextureDI->TextureReferenceRHI->GetReferencedTexture() : nullptr;
 		if (TextureRHI)
 		{
-			FRHITexture* TextureRHI = TextureDI->TextureReferenceRHI->GetReferencedTexture();
 			FRHISamplerState* SamplerStateRHI = TextureDI->SamplerStateRHI;
-			
 			if (!SamplerStateRHI)
 			{
 				// Fallback required because PostLoad() order affects whether RHI resources 
