Reimplemented fix from CL 6453585 which was accidentally lost in a later merge:

"Deletion of NiagaraSystemInstance: bNotifyOnCompletion modified after the class instance has already been deleted."

diff --git a/Engine/Plugins/FX/Niagara/Source/Niagara/Private/NiagaraSystemInstance.cpp b/Engine/Plugins/FX/Niagara/Source/Niagara/Private/NiagaraSystemInstance.cpp
--- a/Engine/Plugins/FX/Niagara/Source/Niagara/Private/NiagaraSystemInstance.cpp
+++ b/Engine/Plugins/FX/Niagara/Source/Niagara/Private/NiagaraSystemInstance.cpp
@@ -448,22 +448,16 @@ void FNiagaraSystemInstance::Complete()
 
 	if (bNeedToNotifyOthers)
 	{
-		//UE_LOG(LogNiagara, Log, TEXT("FNiagaraSystemInstance::BroadCast OnCompleteDelegate { %p"), this);
+		// We've already notified once, no need to do so again.
+		bNotifyOnCompletion = false;
+
 		OnCompleteDelegate.Broadcast(this);
-		//UE_LOG(LogNiagara, Log, TEXT("FNiagaraSystemInstance::BroadCast OnCompleteDelegate } %p"), this);
 
 		if (Component)
 		{
-			//UE_LOG(LogNiagara, Log, TEXT("FNiagaraSystemInstance::Component->OnSystemComplete { %p"), this);
+			// Note: This call may destroy this instance of FNiagaraSystemInstance, so don't use bNotifyOnCompletion after it!
 			Component->OnSystemComplete();
-			//UE_LOG(LogNiagara, Log, TEXT("FNiagaraSystemInstance::Component->OnSystemComplete } %p"), this);
 		}
-
-		//UE_LOG(LogNiagara, Log, TEXT("FNiagaraSystemInstance::Complete } %p"), this);
-
-		
-		// We've already notified once, no need to do so again.
-		bNotifyOnCompletion = false;
 	}
 }
 
