Fixing the Niagara Randomness autotest: Make non-deterministic RNG calls in Niagara GPU scripts not impact subsequent deterministic calls.
Slight refactor of the randomness functionality in NiagaraEmitterInstanceShader.usf
Bumps the Niagara custom script version.

diff --git a/Engine/Plugins/FX/Niagara/Shaders/Private/NiagaraEmitterInstanceShader.usf b/Engine/Plugins/FX/Niagara/Shaders/Private/NiagaraEmitterInstanceShader.usf
--- a/Engine/Plugins/FX/Niagara/Shaders/Private/NiagaraEmitterInstanceShader.usf
+++ b/Engine/Plugins/FX/Niagara/Shaders/Private/NiagaraEmitterInstanceShader.usf
@@ -120,28 +120,19 @@ float Reciprocal(float x){ return 1.0f/x; }
  * all relevant architectures.
  *
  * Warning: Only contains 24 bits of randomness, since we produce values in 
- *          the unit interval. 
+ *          the unit interval. Uses the upper 24 bits, as they have the best
+ *          quality.
  *
- * Note: Deterministic as long as Seed1, Seed2 and Seed3 are the same *and* as 
- *       long as no Random or Seeded Random ops are inserted into or removed 
- *       from the current script.
- *
- * TODO(mv): Might want to use a special case based on NiagaraRand3DPCG32 if 
- *           only 3 or less components are needed, and we can drop one of the 
- *           seeds, to reduce the operation count
+ *          By removing the bit swaps in NiagaraRand4DPCG32 we save a few 
+ *          operations, but lose a bit of statistical (but not visual) quality, 
+ *          and for our use case this is an acceptable compromise.
  * ----------------------------------------------------------------------------
  */
 
-// static counter that is incremented every time rand3 or rand4 is called. 
-static int RandomCounter = 0;
-
-// Returns 4 random normalized floats based on 3 explicit integer seeds and one 
-// implicit integer seed. 
-float4 rand4(int Seed1, int Seed2, int Seed3)
+// Returns 4 random normalized floats based on 4 explicit integer seeds
+float4 rand4(int Seed1, int Seed2, int Seed3, int Seed4)
 {
-	RandomCounter += 1;
-
-	int4 v = int4(RandomCounter, Seed1, Seed2, Seed3) * 1664525 + 1013904223;
+	int4 v = int4(Seed4, Seed1, Seed2, Seed3) * 1664525 + 1013904223;
 
 	v.x += v.y*v.w;
 	v.y += v.z*v.x;
@@ -152,73 +143,73 @@ float4 rand4(int Seed1, int Seed2, int Seed3)
 	v.z += v.x*v.y;
 	v.w += v.y*v.z;
 
-	// NOTE(mv): We can use 24 bits of randomness, as all integers in [0, 2^24] 
-    //           are exactly representable in single precision floats.
-    //           We use the upper 24 bits as they tend to be higher quality.
+	// We can use 24 bits of randomness, as all integers in [0, 2^24] 
+	// are exactly representable in single precision floats.
+	// We use the upper 24 bits as they tend to be higher quality.
 
-    // NOTE(mv): The divide can often be folded with the range scale in the rand functions
-
-    return float4((v >> 8) & 0x00ffffff) / 16777216.0; // 0x01000000 == 16777216
-    // return float4((v >> 8) & 0x00ffffff) * (1.0/16777216.0); // bugged, see UE-67738
+	// The divide is often folded with the range scale in the rand functions
+	return float4((v >> 8) & 0x00ffffff) / 16777216.0; // 0x01000000 == 16777216
+	// return float4((v >> 8) & 0x00ffffff) * (1.0/16777216.0); // bugged, see UE-67738
 }
 
 // float3 specialization of the above:
-// Returns 3 random normalized floats based on 3 explicit integer seeds and one 
-// implicit integer seed. The implicit seed takes the upper 16 bits of the third 
-// explicit seed, as this is often a global seed. 
-// The expression 
 // 
-//    RandomCounter | (Seed3 << 16)
-//
-// should be optimized away at compile time, as the third seed is typically 
-// a global constant.
-float3 rand3(int Seed1, int Seed2, int Seed3)
+// Returns 3 random normalized floats based on 4 explicit integer seeds.
+// 
+// All bits of the first and second seeds are used, while only 
+// the lower 16 bits of the third and fourth seeds are used.
+float3 rand3(int Seed1, int Seed2, int Seed3, int Seed4)
 {
-    RandomCounter += 1;
-
-    int3 v = int3(Seed1, Seed2, RandomCounter | (Seed3 << 16)) * 1664525 + 1013904223;
+	int3 v = int3(Seed1, Seed2, Seed4 | (Seed3 << 16)) * 1664525 + 1013904223;
 
-    v.x += v.y*v.z;
-    v.y += v.z*v.x;
-    v.z += v.x*v.y;
-    v.x += v.y*v.z;
-    v.y += v.z*v.x;
-    v.z += v.x*v.y;
+	v.x += v.y*v.z;
+	v.y += v.z*v.x;
+	v.z += v.x*v.y;
+	v.x += v.y*v.z;
+	v.y += v.z*v.x;
+	v.z += v.x*v.y;
 
-    return float3((v >> 8) & 0x00ffffff) / 16777216.0; // 0x01000000 == 16777216
+	return float3((v >> 8) & 0x00ffffff) / 16777216.0; // 0x01000000 == 16777216
 }
 
+// Internal counter used to generate a different sequence of random numbers for each call
+static int RandomCounterDeterministic = 0;
+
 // Cost using rand4: 6 imad, 1 itof, 1 ishr, 1 add, 2 mul
 float rand(float x, int Seed1, int Seed2, int Seed3)
 {
-	return rand3(Seed1, Seed2, Seed3).x * x;
+	RandomCounterDeterministic += 1;
+	return rand3(Seed1, Seed2, Seed3, RandomCounterDeterministic).x * x;
 }
 
 // Cost using rand4: 7 imad, 1 itof, 1 ishr, 1 add, 2 mul
 float2 rand(float2 x, int Seed1, int Seed2, int Seed3)
 {
-	return rand3(Seed1, Seed2, Seed3).xy * x;
+	RandomCounterDeterministic += 1;
+	return rand3(Seed1, Seed2, Seed3, RandomCounterDeterministic).xy * x;
 }
 
 // Cost using rand4: 8 imad, 1 itof, 1 ishr, 1 add, 2 mul
 float3 rand(float3 x, int Seed1, int Seed2, int Seed3)
 {
-	return rand3(Seed1, Seed2, Seed3).xyz * x;
+	RandomCounterDeterministic += 1;
+	return rand3(Seed1, Seed2, Seed3, RandomCounterDeterministic).xyz * x;
 }
 
 // Cost using rand4: 9 imad, 1 itof, 1 ishr, 1 and, 2 mul
 float4 rand(float4 x, int Seed1, int Seed2, int Seed3) 
 {
-    return rand4(Seed1, Seed2, Seed3).xyzw * x;
+	RandomCounterDeterministic += 1;
+	return rand4(Seed1, Seed2, Seed3, RandomCounterDeterministic).xyzw * x;
 }
 
 // Cost using rand4: 6 imad, 2 itof, 1 ishr, 1 add, 2 mul, 1 ftoi
 int rand(int x, int Seed1, int Seed2, int Seed3)
 {
-	// NOTE: Scaling a uniform float range provides better distribution of 
-	//       numbers than using %.
-	// NOTE: Inclusive! So [0, x] instead of [0, x)
-	return int(rand3(Seed1, Seed2, Seed3).x * (x+1));
+	// Scaling a uniform float range provides better distribution of numbers than using %.
+	// Inclusive! So [0, x] instead of [0, x)
+	RandomCounterDeterministic += 1;
+	return int(rand3(Seed1, Seed2, Seed3, RandomCounterDeterministic).x * (x+1));
 }
 
 /* -----------------------------------------------------------------
@@ -226,41 +217,45 @@ int rand(int x, int Seed1, int Seed2, int Seed3)
  * -----------------------------------------------------------------
  */
 #if GPU_SIMULATION
-	// NOTE(mv): This simply calls the deterministic random number functions 
-    //           from the Seeded RNG section, but uses non-deterministic seeds 
-    //           as input. In particular, GDispatchThreadId for the first seed 
-	//           and EmitterTickCounter for the second seed. Also note that 
-	//           EmitterTickCounter does *not* reset for each frame, which 
-	//           secures non-determinism across simulations.
+	// This simply calls the deterministic random number functions from the Seeded RNG section, 
+	// but uses non-deterministic seeds as input. 
 
-	// TODO(mv): This could perhaps be optimized a little bit using slightly 
-	//           cheaper functions, but the difference is likely negligible. 
-	//           So for the sake of brevity and cleanliness, keep it like this.
+	// This could perhaps be optimized by using slightly cheaper functions, but the difference is likely negligible. 
+	
+	// Internal counter used to generate a different sequence of random numbers for each call
+	// We need to keep this separate from the Deterministic version so that non-deterministic 
+	// calls do not interfere with the deterministic ones. 
+	static int RandomCounterNonDeterministic = -1;
 
 	float rand(float x)
 	{
-		return rand4(GDispatchThreadId.x, EmitterTickCounter, 0).x * x;
+		RandomCounterNonDeterministic -= 1;
+		return rand4(GDispatchThreadId.x, EmitterTickCounter, GDispatchThreadId.x, RandomCounterNonDeterministic).x * x;
 	}
 
 	float2 rand(float2 x)
 	{
-		return rand4(GDispatchThreadId.x, EmitterTickCounter, 0).xy * x;
+		RandomCounterNonDeterministic -= 1;
+		return rand4(GDispatchThreadId.x, EmitterTickCounter, GDispatchThreadId.x, RandomCounterNonDeterministic).xy * x;
 	}
 
 	float3 rand(float3 x)
 	{
-		return rand4(GDispatchThreadId.x, EmitterTickCounter, 0).xyz * x;
+		RandomCounterNonDeterministic -= 1;
+		return rand4(GDispatchThreadId.x, EmitterTickCounter, GDispatchThreadId.x, RandomCounterNonDeterministic).xyz * x;
 	}
 
 	float4 rand(float4 x) 
 	{
-		return rand4(GDispatchThreadId.x, EmitterTickCounter, 0).xyzw * x;
+		RandomCounterNonDeterministic -= 1;
+		return rand4(GDispatchThreadId.x, EmitterTickCounter, GDispatchThreadId.x, RandomCounterNonDeterministic).xyzw * x;
 	}
 
-	// Note: integer randoms are INCLUSIVE, i.e. includes both the upper and lower limit
+	// Integer randoms are INCLUSIVE, i.e. includes both the upper and lower limits
 	int rand(int x)
 	{
-		return int(rand4(GDispatchThreadId.x, EmitterTickCounter, 0).x * (x+1));
+		RandomCounterNonDeterministic -= 1;
+		return int(rand4(GDispatchThreadId.x, EmitterTickCounter, GDispatchThreadId.x, RandomCounterNonDeterministic).x * (x+1));
 	}
 #else
 	// Old unseeded, passthrough to FRandomStream
diff --git a/Engine/Plugins/FX/Niagara/Source/NiagaraCore/Private/NiagaraCustomVersion.cpp b/Engine/Plugins/FX/Niagara/Source/NiagaraCore/Private/NiagaraCustomVersion.cpp
--- a/Engine/Plugins/FX/Niagara/Source/NiagaraCore/Private/NiagaraCustomVersion.cpp
+++ b/Engine/Plugins/FX/Niagara/Source/NiagaraCore/Private/NiagaraCustomVersion.cpp
@@ -8,4 +8,6 @@ const FGuid FNiagaraCustomVersion::GUID(0xFCF57AFA, 0x50764283, 0xB9A9E658, 0xFF
 // Register the custom version with core
 FCustomVersionRegistration GRegisterNiagaraCustomVersion(FNiagaraCustomVersion::GUID, FNiagaraCustomVersion::LatestVersion, TEXT("NiagaraVer"));
 
-const FGuid FNiagaraCustomVersion::LatestScriptCompileVersion(0x9BA7BE7D, 0x57FD4A80, 0xB2DD5740, 0xB9FCCC78);
\ No newline at end of file
+const FGuid FNiagaraCustomVersion::LatestScriptCompileVersion(0x563D9FF1, 0x1E9042CC, 0xA8BF7AC4, 0xC9366901);
+
+
