Tick Refactor.

- Moves Niagara tick work into tick function in LastDemotable. In future CL we can push other work into earlier tick groups.
- Can now optioally push system simulation work off the GT.
- Can now optionally push system instance tick work off the GT in batches. Currently of 8.
- Refactor of system simulation code to be far more clear and maintainable.
- Modified spawn flow to facilitate ticking inside main world tick. Now we spawn all our new systems at the end of the frames actor ticking so we catch any new systems spwaned during the frame.
- Various infrastructure changes to support the above.

- Misc:
Changed some checks to check slow.
Fixed a UI crash
Removed VM stat scope that was too granular.
Moved some static functions relating to the world manager actually into the world manager rather than in the NiagaraModule

diff --git a/Engine/Plugins/FX/Niagara/Source/Niagara/Classes/NiagaraDataSet.h b/Engine/Plugins/FX/Niagara/Source/Niagara/Classes/NiagaraDataSet.h
--- a/Engine/Plugins/FX/Niagara/Source/Niagara/Classes/NiagaraDataSet.h
+++ b/Engine/Plugins/FX/Niagara/Source/Niagara/Classes/NiagaraDataSet.h
@@ -99,7 +99,7 @@ public:
 	void AllocateGPU(uint32 InNumInstances, FNiagaraGPUInstanceCountManager& GPUInstanceCountManager, FRHICommandList &RHICmdList);
 	void SwapInstances(uint32 OldIndex, uint32 NewIndex);
 	void KillInstance(uint32 InstanceIdx);
-	void CopyTo(FNiagaraDataBuffer& DestBuffer, int32 StartIdx, int32 NumInstances)const;
+	void CopyTo(FNiagaraDataBuffer& DestBuffer, int32 SrcStartIdx, int32 DestStartIdx, int32 NumInstances)const;
 	void CopyTo(FNiagaraDataBuffer& DestBuffer)const;
 	void GPUCopyFrom(float* GPUReadBackFloat, int* GPUReadBackInt, int32 StartIdx, int32 NumInstances, uint32 InSrcFloatStride, uint32 InSrcIntStride);
 	void Dump(int32 StartIndex, int32 NumInstances, const FString& Label)const;
@@ -140,7 +140,7 @@ public:
 
 	FORCEINLINE FNiagaraDataSet* GetOwner()const { return Owner; }
 
-	int32 TransferInstance(FNiagaraDataBuffer& SourceBuffer, int32 InstanceIndex);
+	int32 TransferInstance(FNiagaraDataBuffer& SourceBuffer, int32 InstanceIndex, bool bRemoveFromSource=true);
 
 	bool CheckForNaNs()const;
 
@@ -260,7 +260,7 @@ public:
 	const FNiagaraVariableLayoutInfo* GetVariableLayout(const FNiagaraVariable& Var)const;
 	bool GetVariableComponentOffsets(const FNiagaraVariable& Var, int32 &FloatStart, int32 &IntStart) const;
 
-	void CopyTo(FNiagaraDataSet& Other, int32 StartIdx = 0, int32 NumInstances = INDEX_NONE)const;
+	void CopyTo(FNiagaraDataSet& Other, int32 StartIdx = 0, int32 NumInstances = INDEX_NONE, bool bResetOther=true)const;
 
 	void CopyFromGPUReadback(float* GPUReadBackFloat, int* GPUReadBackInt, int32 StartIdx = 0, int32 NumInstances = INDEX_NONE, uint32 FloatStride = 0, uint32 IntStride = 0);
 
@@ -300,10 +300,10 @@ private:
 	{
 		// In some rare occasions, the render thread might be null, like when offloading work to Lightmass 
 		// The final GRenderingThread check keeps us from inadvertently failing when that happens.
-#if DO_CHECK
+#if DO_GUARD_SLOW
 		bool CPUSimOK = (SimTarget == ENiagaraSimTarget::CPUSim && !IsInRenderingThread());
 		bool GPUSimOK = (SimTarget == ENiagaraSimTarget::GPUComputeSim && IsInRenderingThread());
-		checkf(!GRenderingThread || CPUSimOK || GPUSimOK, TEXT("NiagaraDataSet function being called on incorrect thread."));
+		checkfSlow(!GRenderingThread || CPUSimOK || GPUSimOK, TEXT("NiagaraDataSet function being called on incorrect thread."));
 #endif
 	}
 
@@ -370,20 +370,30 @@ struct FNiagaraDataSetAccessorBase
 	FNiagaraDataSetAccessorBase(FNiagaraDataSet* InDataSet, FNiagaraVariable InVar)
 		: DataSet(InDataSet)
 	{
+		Var = InVar;
 		VarLayout = DataSet->GetVariableLayout(InVar);
 	}
 
 	void Create(FNiagaraDataSet* InDataSet, FNiagaraVariable InVar)
 	{
 		DataSet = InDataSet;
+		Var = InVar;
 		VarLayout = DataSet->GetVariableLayout(InVar);
 	}
 
+	void SetDataSet(FNiagaraDataSet& InDataSet)
+	{
+		ensure(Var.IsValid());
+		DataSet = &InDataSet;
+		VarLayout = InDataSet.GetVariableLayout(Var);
+	}
+
 	FORCEINLINE bool IsValid()const { return DataSet && VarLayout != nullptr; }
 protected:
 
 	FNiagaraDataSet* DataSet;
 	const FNiagaraVariableLayoutInfo* VarLayout;
+	FNiagaraVariable Var;
 };
 
 template<typename T>
diff --git a/Engine/Plugins/FX/Niagara/Source/Niagara/Private/NiagaraCommon.cpp b/Engine/Plugins/FX/Niagara/Source/Niagara/Private/NiagaraCommon.cpp
--- a/Engine/Plugins/FX/Niagara/Source/Niagara/Private/NiagaraCommon.cpp
+++ b/Engine/Plugins/FX/Niagara/Source/Niagara/Private/NiagaraCommon.cpp
@@ -12,6 +12,7 @@
 #include "UObject/Class.h"
 #include "UObject/Package.h"
 #include "Modules/ModuleManager.h"
+#include "NiagaraWorldManager.h"
 
 DECLARE_CYCLE_STAT(TEXT("Niagara - Utilities - PrepareRapidIterationParameters"), STAT_Niagara_Utilities_PrepareRapidIterationParameters, STATGROUP_Niagara);
 
@@ -70,10 +71,9 @@ FString FNiagaraTypeHelper::ToString(const uint8* ValueData, const UScriptStruct
 
 FNiagaraSystemUpdateContext::~FNiagaraSystemUpdateContext()
 {
-	INiagaraModule& NiagaraModule = FModuleManager::LoadModuleChecked<INiagaraModule>("Niagara");
 	for (UNiagaraSystem* Sys : SystemSimsToDestroy)
 	{
-		NiagaraModule.DestroyAllSystemSimulations(Sys);
+		FNiagaraWorldManager::DestroyAllSystemSimulations(Sys);
 	}
 
 	for (UNiagaraComponent* Comp : ComponentsToReInit)
diff --git a/Engine/Plugins/FX/Niagara/Source/Niagara/Private/NiagaraComponent.cpp b/Engine/Plugins/FX/Niagara/Source/Niagara/Private/NiagaraComponent.cpp
--- a/Engine/Plugins/FX/Niagara/Source/Niagara/Private/NiagaraComponent.cpp
+++ b/Engine/Plugins/FX/Niagara/Source/Niagara/Private/NiagaraComponent.cpp
@@ -487,7 +487,7 @@ void UNiagaraComponent::TickComponent(float DeltaSeconds, enum ELevelTick TickTy
 
 		if (AgeUpdateMode == ENiagaraAgeUpdateMode::TickDeltaTime)
 		{
-			SystemInstance->ComponentTick(DeltaSeconds);
+			SystemInstance->ComponentTick(DeltaSeconds, ThisTickFunction->GetCompletionHandle());
 		}
 		else
 		{
@@ -516,7 +516,9 @@ void UNiagaraComponent::TickComponent(float DeltaSeconds, enum ELevelTick TickTy
 						TicksToProcess = FMath::FloorToInt(AgeDiff / SeekDelta);
 						for (; TicksToProcess > 0 && CurrentTime - StartTime < MaxSimTime; --TicksToProcess)
 						{
-							SystemInstance->ComponentTick(SeekDelta);
+							//Cannot do multiple tick off the game thread here without additional work. So we pass in null for the completion event which will force GT execution.
+							//If this becomes a perf problem I can add a new path for the tick code to handle multiple ticks.
+							SystemInstance->ComponentTick(SeekDelta, nullptr);
 							CurrentTime = FPlatformTime::Seconds();
 						}
 					}
diff --git a/Engine/Plugins/FX/Niagara/Source/Niagara/Private/NiagaraDataSet.cpp b/Engine/Plugins/FX/Niagara/Source/Niagara/Private/NiagaraDataSet.cpp
--- a/Engine/Plugins/FX/Niagara/Source/Niagara/Private/NiagaraDataSet.cpp
+++ b/Engine/Plugins/FX/Niagara/Source/Niagara/Private/NiagaraDataSet.cpp
@@ -425,24 +425,58 @@ bool FNiagaraDataSet::GetVariableComponentOffsets(const FNiagaraVariable& Var, i
 	return false;
 }
 
-void FNiagaraDataSet::CopyTo(FNiagaraDataSet& Other, int32 StartIdx, int32 NumInstances)const
+void FNiagaraDataSet::CopyTo(FNiagaraDataSet& Other, int32 StartIdx, int32 NumInstances, bool bResetOther)const
 {
 	CheckCorrectThread();
-	Other.Reset();
-	Other.Variables = Variables;
-	Other.VariableLayouts = VariableLayouts;
-	Other.TotalFloatComponents = TotalFloatComponents;
-	Other.TotalInt32Components = TotalInt32Components;
-	Other.Finalize();
 
-	//Read the most current data. Even if it's possibly partially complete simulation data.
-	FNiagaraDataBuffer& OtherDataBuffer = Other.BeginSimulate();
+	if (bResetOther)
+	{
+		Other.Reset();
+		Other.Variables = Variables;
+		Other.VariableLayouts = VariableLayouts;
+		Other.TotalFloatComponents = TotalFloatComponents;
+		Other.TotalInt32Components = TotalInt32Components;
+		Other.Finalize();
+	}
+	else
+	{
+		checkSlow(Other.GetVariables() == Variables);
+	}
+
+	FNiagaraDataBuffer* OtherCurrentBuffer = Other.GetCurrentData();
+	FNiagaraDataBuffer& OtherDestBuffer = Other.BeginSimulate();
 
+	//Read the most current data. Even if it's possibly partially complete simulation data.
 	FNiagaraDataBuffer* SourceBuffer = GetDestinationData() ? GetDestinationData() : GetCurrentData();
+
 	if (SourceBuffer != nullptr)
 	{
-		OtherDataBuffer.Allocate(SourceBuffer->GetNumInstancesAllocated());
-		SourceBuffer->CopyTo(OtherDataBuffer, StartIdx, NumInstances);
+		int32 SourceInstances = SourceBuffer->GetNumInstances();
+		int32 OrigNumInstances = OtherCurrentBuffer->GetNumInstances();
+
+		if (StartIdx >= SourceInstances)
+		{
+			return; //We can't start beyond the end of the source buffer.
+		}
+
+		if (NumInstances == INDEX_NONE || StartIdx + NumInstances >= SourceInstances)
+		{
+			NumInstances = SourceBuffer->GetNumInstances() - StartIdx;
+		}
+
+		//We need to allocate enough space for the new data and existing data if we're keeping it.
+		int32 RequiredInstances = bResetOther ? NumInstances : NumInstances + OrigNumInstances;
+		OtherDestBuffer.Allocate(RequiredInstances);
+		OtherDestBuffer.SetNumInstances(RequiredInstances);
+
+		//Copy the data in our current buffer over into the new buffer.
+		if (!bResetOther)
+		{
+			OtherCurrentBuffer->CopyTo(OtherDestBuffer, 0, 0, OtherCurrentBuffer->GetNumInstances());
+		}
+
+		//Now copy the data from the source buffer into the newly allocated space.
+		SourceBuffer->CopyTo(OtherDestBuffer, 0, OrigNumInstances, NumInstances);
 	}
 
 	Other.EndSimulate();
@@ -497,7 +531,7 @@ void FNiagaraDataBuffer::CheckUsage(bool bReadOnly)const
 	}
 }
 
-int32 FNiagaraDataBuffer::TransferInstance(FNiagaraDataBuffer& SourceBuffer, int32 InstanceIndex)
+int32 FNiagaraDataBuffer::TransferInstance(FNiagaraDataBuffer& SourceBuffer, int32 InstanceIndex, bool bRemoveFromSource)
 {
 	CheckUsage(false);
 	if (SourceBuffer.GetNumInstances() > (uint32)InstanceIndex)
@@ -525,6 +559,11 @@ int32 FNiagaraDataBuffer::TransferInstance(FNiagaraDataBuffer& SourceBuffer, int
 			*Dst = *Src;
 		}
 
+		if (bRemoveFromSource)
+		{
+			SourceBuffer.KillInstance(InstanceIndex);
+		}
+
 		return OldNumInstances;
 	}
 
@@ -715,31 +754,35 @@ void FNiagaraDataBuffer::KillInstance(uint32 InstanceIdx)
 #endif
 }
 
-void FNiagaraDataBuffer::CopyTo(FNiagaraDataBuffer& DestBuffer, int32 InStartIdx, int32 InNumInstances)const
+void FNiagaraDataBuffer::CopyTo(FNiagaraDataBuffer& DestBuffer, int32 StartIdx, int32 DestStartIdx, int32 InNumInstances)const
 {
 	CheckUsage(false);
-	
-	if (InStartIdx < 0 || (uint32)InStartIdx > NumInstances)
+
+	if (StartIdx < 0 || (uint32)StartIdx >= NumInstances)
 	{
-		InStartIdx = NumInstances;
+		return;
 	}
-	if (InNumInstances == INDEX_NONE || ((uint32)InNumInstances + (uint32)InStartIdx) > NumInstances)
+
+	uint32 InstancesToCopy = InNumInstances;
+	if (InstancesToCopy == INDEX_NONE)
 	{
-		InNumInstances = NumInstances - InStartIdx;
+		InstancesToCopy = NumInstances - StartIdx;
 	}
 
-	if (InNumInstances != 0)
+	if (InstancesToCopy != 0)
 	{
-		if (DestBuffer.NumInstancesAllocated != NumInstancesAllocated)
+		uint32 NewNumInstances = DestStartIdx + InstancesToCopy;
+		if (DestStartIdx < 0 || NewNumInstances >= DestBuffer.GetNumInstances())
 		{
-			DestBuffer.Allocate(NumInstancesAllocated);
+			DestBuffer.Allocate(NewNumInstances, true);
 		}
+		DestBuffer.SetNumInstances(NewNumInstances);
 
 		for (uint32 CompIdx = 0; CompIdx < Owner->TotalFloatComponents; ++CompIdx)
 		{
-			const float* SrcStart = GetInstancePtrFloat(CompIdx, InStartIdx);
-			const float* SrcEnd = GetInstancePtrFloat(CompIdx, InStartIdx + InNumInstances);
-			float* Dst = DestBuffer.GetInstancePtrFloat(CompIdx, 0);
+			const float* SrcStart = GetInstancePtrFloat(CompIdx, StartIdx);
+			const float* SrcEnd = GetInstancePtrFloat(CompIdx, StartIdx + InstancesToCopy);
+			float* Dst = DestBuffer.GetInstancePtrFloat(CompIdx, DestStartIdx);
 			size_t Count = SrcEnd - SrcStart;
 			FMemory::Memcpy(Dst, SrcStart, Count*sizeof(float));
 
@@ -747,15 +790,15 @@ void FNiagaraDataBuffer::CopyTo(FNiagaraDataBuffer& DestBuffer, int32 InStartIdx
 			{
 				for (size_t i = 0; i < Count; i++)
 				{
-					check(SrcStart[i] == Dst[i]);
+					checkSlow(SrcStart[i] == Dst[i]);
 				}
 			}
 		}
 		for (uint32 CompIdx = 0; CompIdx < Owner->TotalInt32Components; ++CompIdx)
 		{
-			const int32* SrcStart = GetInstancePtrInt32(CompIdx, InStartIdx);
-			const int32* SrcEnd = GetInstancePtrInt32(CompIdx, InStartIdx + InNumInstances);
-			int32* Dst = DestBuffer.GetInstancePtrInt32(CompIdx, 0);
+			const int32* SrcStart = GetInstancePtrInt32(CompIdx, StartIdx);
+			const int32* SrcEnd = GetInstancePtrInt32(CompIdx, StartIdx + InstancesToCopy);
+			int32* Dst = DestBuffer.GetInstancePtrInt32(CompIdx, DestStartIdx);
 			size_t Count = SrcEnd - SrcStart;
 			FMemory::Memcpy(Dst, SrcStart, Count * sizeof(int32));
 
@@ -763,11 +806,10 @@ void FNiagaraDataBuffer::CopyTo(FNiagaraDataBuffer& DestBuffer, int32 InStartIdx
 			{
 				for (size_t i = 0; i < Count; i++)
 				{
-					check(SrcStart[i] == Dst[i]);
+					checkSlow(SrcStart[i] == Dst[i]);
 				}
 			}
 		}
-		DestBuffer.SetNumInstances(InNumInstances);
 	}
 }
 
diff --git a/Engine/Plugins/FX/Niagara/Source/Niagara/Private/NiagaraEmitterInstance.cpp b/Engine/Plugins/FX/Niagara/Source/Niagara/Private/NiagaraEmitterInstance.cpp
--- a/Engine/Plugins/FX/Niagara/Source/Niagara/Private/NiagaraEmitterInstance.cpp
+++ b/Engine/Plugins/FX/Niagara/Source/Niagara/Private/NiagaraEmitterInstance.cpp
@@ -778,6 +778,11 @@ void FNiagaraEmitterInstance::PostTick()
 
 bool FNiagaraEmitterInstance::HandleCompletion(bool bForce)
 {
+	if (IsDisabled())
+	{
+		return true;
+	}
+
 	if (bForce)
 	{
 		SetExecutionState(ENiagaraExecutionState::Complete);
diff --git a/Engine/Plugins/FX/Niagara/Source/Niagara/Private/NiagaraEmitterInstanceBatcher.cpp b/Engine/Plugins/FX/Niagara/Source/Niagara/Private/NiagaraEmitterInstanceBatcher.cpp
--- a/Engine/Plugins/FX/Niagara/Source/Niagara/Private/NiagaraEmitterInstanceBatcher.cpp
+++ b/Engine/Plugins/FX/Niagara/Source/Niagara/Private/NiagaraEmitterInstanceBatcher.cpp
@@ -335,6 +335,7 @@ void NiagaraEmitterInstanceBatcher::ResizeBuffersAndGatherResources(FOverlappabl
 {
 	SCOPE_CYCLE_COUNTER(STAT_NiagaraGPUDispatchSetup_RT);
 
+	//UE_LOG(LogNiagara, Warning, TEXT("NiagaraEmitterInstanceBatcher::ResizeBuffersAndGatherResources:  %0xP"), this);
 	for (FNiagaraGPUSystemTick* Tick : OverlappableTick)
 	{
 		const uint32 DispatchCount = Tick->Count;
@@ -407,6 +408,7 @@ void NiagaraEmitterInstanceBatcher::DispatchAllOnCompute(FOverlappableTicks& Ove
 {
 	FRHICommandListImmediate& RHICmdListImmediate = FRHICommandListExecutor::GetImmediateCommandList();
 
+	//UE_LOG(LogNiagara, Warning, TEXT("NiagaraEmitterInstanceBatcher::DispatchAllOnCompute:  %0xP"), this);
 	// Disable automatic cache flush so that we can have our compute work overlapping. Barrier will be used as a sync mechanism.
 	RHICmdList.AutomaticCacheFlushAfterComputeShader(false);
 
@@ -801,7 +803,7 @@ void NiagaraEmitterInstanceBatcher::PreRender(FRHICommandListImmediate& RHICmdLi
 
 void NiagaraEmitterInstanceBatcher::OnDestroy()
 {
-	INiagaraModule::OnBatcherDestroyed(this);
+	FNiagaraWorldManager::OnBatcherDestroyed(this);
 	FFXSystemInterface::OnDestroy();
 }
 
@@ -835,6 +837,8 @@ void NiagaraEmitterInstanceBatcher::SortGPUParticles(FRHICommandListImmediate& R
 	{
 		SCOPED_GPU_STAT(RHICmdList, NiagaraGPUSorting);
 
+		//UE_LOG(LogNiagara, Warning, TEXT("NiagaraEmitterInstanceBatcher::SortGPUParticles:  %0xP"), this);
+
 		ensure(SortedVertexBuffers.Last().IndexCount >= SortedParticleCount);
 
 		// The particle sort buffer must be able to hold all the particles.
diff --git a/Engine/Plugins/FX/Niagara/Source/Niagara/Private/NiagaraModule.cpp b/Engine/Plugins/FX/Niagara/Source/Niagara/Private/NiagaraModule.cpp
--- a/Engine/Plugins/FX/Niagara/Source/Niagara/Private/NiagaraModule.cpp
+++ b/Engine/Plugins/FX/Niagara/Source/Niagara/Private/NiagaraModule.cpp
@@ -28,8 +28,6 @@ IMPLEMENT_MODULE(INiagaraModule, Niagara);
 
 #define LOCTEXT_NAMESPACE "NiagaraModule"
 
-TMap<class UWorld*, class FNiagaraWorldManager*> INiagaraModule::WorldManagers;
-
 float INiagaraModule::EngineGlobalSpawnCountScale = 1.0f;
 float INiagaraModule::EngineGlobalSystemCountScale = 1.0f;
 int32 INiagaraModule::EngineDetailLevel = 4;
@@ -175,9 +173,7 @@ void INiagaraModule::StartupModule()
 	FNiagaraTypeDefinition::Init();
 	FNiagaraViewDataMgr::Init();
 
-	FWorldDelegates::OnPreWorldInitialization.AddRaw(this, &INiagaraModule::OnWorldInit);
-	FWorldDelegates::OnWorldCleanup.AddRaw(this, &INiagaraModule::OnWorldCleanup);
-	FWorldDelegates::OnPreWorldFinishDestroy.AddRaw(this, &INiagaraModule::OnPreWorldFinishDestroy);
+	FNiagaraWorldManager::OnStartup();
 
 #if WITH_EDITOR	
 	if (!GIsEditor)
@@ -189,7 +185,6 @@ void INiagaraModule::StartupModule()
 	}
 #endif
 
-	FWorldDelegates::OnWorldPostActorTick.AddRaw(this, &INiagaraModule::TickWorld);
 	CVarDetailLevel.AsVariable()->SetOnChangedCallback(FConsoleVariableDelegate::CreateRaw(this, &INiagaraModule::OnChangeDetailLevel));
 	OnChangeDetailLevel(CVarDetailLevel.AsVariable());
 
@@ -307,13 +302,7 @@ void INiagaraModule::ShutdownRenderingResources()
 
 void INiagaraModule::ShutdownModule()
 {
-	//Should have cleared up all world managers by now.
-	check(WorldManagers.Num() == 0);
-	for (TPair<UWorld*, FNiagaraWorldManager*> Pair : WorldManagers)
-	{
-		delete Pair.Value;
-		Pair.Value = nullptr;
-	}
+	FNiagaraWorldManager::OnShutdown();
 
 	// Clear out the handler when shutting down..
 	INiagaraShaderModule& NiagaraShaderModule = FModuleManager::LoadModuleChecked<INiagaraShaderModule>("NiagaraShader");
@@ -323,66 +312,6 @@ void INiagaraModule::ShutdownModule()
 	ShutdownRenderingResources();
 }
 
-FNiagaraWorldManager* INiagaraModule::GetWorldManager(UWorld* World)
-{
-	FNiagaraWorldManager** OutWorld = WorldManagers.Find(World);
-	if (OutWorld == nullptr)
-	{
-		UE_LOG(LogNiagara, Warning, TEXT("Calling INiagaraModule::GetWorldManager \"%s\", but Niagara has never encountered this world before. "
-			" This means that WorldInit never happened. This may happen in some edge cases in the editor, like saving invisible child levels, "
-			"in which case the calling context needs to be safe against this returning nullptr."), World ? *World->GetName() : TEXT("nullptr"));
-		return nullptr;
-	}
-	return *OutWorld;
-}
-
-void INiagaraModule::OnBatcherDestroyed(NiagaraEmitterInstanceBatcher* InBatcher)
-{
-	for (TPair<UWorld*, FNiagaraWorldManager*>& Pair : WorldManagers)
-	{
-		Pair.Value->OnBatcherDestroyed(InBatcher);
-	}
-}
-
-void INiagaraModule::DestroyAllSystemSimulations(class UNiagaraSystem* System)
-{
-	for (TPair<UWorld*, FNiagaraWorldManager*>& Pair : WorldManagers)
-	{
-		Pair.Value->DestroySystemSimulation(System);
-	}
-}
-
-void INiagaraModule::OnWorldInit(UWorld* World, const UWorld::InitializationValues IVS)
-{
-	check(WorldManagers.Find(World) == nullptr);
-	WorldManagers.Add(World) = new FNiagaraWorldManager(World);
-}
-
-void INiagaraModule::OnWorldCleanup(UWorld* World, bool bSessionEnded, bool bCleanupResources)
-{
-	//Cleanup world manager contents but not the manager itself.
-	FNiagaraWorldManager** Manager = WorldManagers.Find(World);
-	if (Manager)
-	{
-		(*Manager)->OnWorldCleanup(bSessionEnded, bCleanupResources);
-	}
-}
-
-void INiagaraModule::OnPreWorldFinishDestroy(UWorld* World)
-{
-	FNiagaraWorldManager** Manager = WorldManagers.Find(World);
-	if (Manager)
-	{
-		delete (*Manager);
-		WorldManagers.Remove(World);
-	}
-}
-
-void INiagaraModule::TickWorld(UWorld* World, ELevelTick TickType, float DeltaSeconds)
-{
-	GetWorldManager(World)->Tick(DeltaSeconds);
-}
-
 #if WITH_EDITOR
 const INiagaraMergeManager& INiagaraModule::GetMergeManager() const
 {
diff --git a/Engine/Plugins/FX/Niagara/Source/Niagara/Private/NiagaraParameterStore.cpp b/Engine/Plugins/FX/Niagara/Source/Niagara/Private/NiagaraParameterStore.cpp
--- a/Engine/Plugins/FX/Niagara/Source/Niagara/Private/NiagaraParameterStore.cpp
+++ b/Engine/Plugins/FX/Niagara/Source/Niagara/Private/NiagaraParameterStore.cpp
@@ -187,12 +187,12 @@ void FNiagaraParameterStore::CheckForNaNs()const
 
 void FNiagaraParameterStore::Tick()
 {
-	SCOPE_CYCLE_COUNTER(STAT_NiagaraParameterStoreTick);
 #if NIAGARA_NAN_CHECKING
 	CheckForNaNs();
 #endif
-	if (bParametersDirty || bInterfacesDirty || bUObjectsDirty)
+	if (Bindings.Num() > 0 && (bParametersDirty || bInterfacesDirty || bUObjectsDirty))
 	{
+		SCOPE_CYCLE_COUNTER(STAT_NiagaraParameterStoreTick);
 		for (TPair<FNiagaraParameterStore*, FNiagaraParameterStoreBinding>& Binding : Bindings)
 		{
 			Binding.Value.Tick(Binding.Key, this);
diff --git a/Engine/Plugins/FX/Niagara/Source/Niagara/Private/NiagaraSystemInstance.cpp b/Engine/Plugins/FX/Niagara/Source/Niagara/Private/NiagaraSystemInstance.cpp
--- a/Engine/Plugins/FX/Niagara/Source/Niagara/Private/NiagaraSystemInstance.cpp
+++ b/Engine/Plugins/FX/Niagara/Source/Niagara/Private/NiagaraSystemInstance.cpp
@@ -24,10 +24,21 @@ DECLARE_CYCLE_STAT(TEXT("System Init Emitters [GT]"), STAT_NiagaraSystemInitEmit
 DECLARE_CYCLE_STAT(TEXT("System Advance Simulation [GT] "), STAT_NiagaraSystemAdvanceSim, STATGROUP_Niagara);
 DECLARE_CYCLE_STAT(TEXT("System SetSolo[GT] "), STAT_NiagaraSystemSetSolo, STATGROUP_Niagara);
 
-DECLARE_CYCLE_STAT(TEXT("System PreSimulateTick [CNC]"), STAT_NiagaraSystemPreSimulateTick, STATGROUP_Niagara);
-DECLARE_CYCLE_STAT(TEXT("System Instance Tick [CNC]"), STAT_NiagaraSystemInstanceTick, STATGROUP_Niagara);
+//High level stats for system instance ticks.
+DECLARE_CYCLE_STAT(TEXT("System Instance Tick (Component) [GT]"), STAT_NiagaraSystemInst_ComponentTickGT, STATGROUP_Niagara);
+DECLARE_CYCLE_STAT(TEXT("System Instance Tick [GT]"), STAT_NiagaraSystemInst_TickGT, STATGROUP_Niagara);
+DECLARE_CYCLE_STAT(TEXT("System Instance Tick [CNC]"), STAT_NiagaraSystemInst_TickCNC, STATGROUP_Niagara);
+DECLARE_CYCLE_STAT(TEXT("System Instance Finalize [GT]"), STAT_NiagaraSystemInst_FinalizeGT, STATGROUP_Niagara);
+DECLARE_CYCLE_STAT(TEXT("System Instance WaitForAsyncTick [GT]"), STAT_NiagaraSystemWaitForAsyncTick, STATGROUP_Niagara);
+
+static float GWaitForAsyncStallWarnThresholdMS = 0.01f;
+static FAutoConsoleVariableRef CVarWaitForAsyncStallWarnThresholdMS(
+	TEXT("fx.WaitForAsyncStallWarnThresholdMS"),
+	GWaitForAsyncStallWarnThresholdMS,
+	TEXT("If we stall in WaitForAsync for longer than this threshold then we emit a stall warning message."),
+	ECVF_Default
+);
 
- 
 /** Safety time to allow for the LastRenderTime coming back from the RT. This is overkill but that's ok.*/
 static float GLastRenderTimeSafetyBias = 0.1f;
 static FAutoConsoleVariableRef CVarLastRenderTimeSafetyBias(
@@ -53,6 +64,9 @@ FNiagaraSystemInstance::FNiagaraSystemInstance(UNiagaraComponent* InComponent)
 	, RequestedExecutionState(ENiagaraExecutionState::Complete)
 	, ActualExecutionState(ENiagaraExecutionState::Complete)
 	, bDataInterfacesInitialized(false)
+	, bAsyncWorkInProgress(false)
+	, bNeedsFinalize(false)
+	, CachedDeltaSeconds(0.0f)
 {
 	SystemBounds.Init();
 
@@ -73,6 +87,8 @@ FNiagaraSystemInstance::FNiagaraSystemInstance(UNiagaraComponent* InComponent)
 
 void FNiagaraSystemInstance::SetEmitterEnable(FName EmitterName, bool bNewEnableState)
 {
+	WaitForAsyncTick();
+
 	UNiagaraSystem* System = GetSystem();
 	if (System != nullptr)
 	{
@@ -119,6 +135,8 @@ void FNiagaraSystemInstance::SetEmitterEnable(FName EmitterName, bool bNewEnable
 
 void FNiagaraSystemInstance::Init(bool bInForceSolo)
 {
+	WaitForAsyncTick(true);
+
 	bForceSolo = bInForceSolo;
 	ActualExecutionState = ENiagaraExecutionState::Inactive;
 	RequestedExecutionState = ENiagaraExecutionState::Inactive;
@@ -202,6 +220,8 @@ bool FNiagaraSystemInstance::RequestCapture(const FGuid& RequestId)
 		return false;
 	}
 
+	WaitForAsyncTick();
+
 	UE_LOG(LogNiagara, Warning, TEXT("Capture requested!"));
 
 	bWasSoloPriorToCaptureRequest = bSolo;
@@ -235,6 +255,8 @@ bool FNiagaraSystemInstance::RequestCapture(const FGuid& RequestId)
 
 void FNiagaraSystemInstance::FinishCapture()
 {
+	WaitForAsyncTick();
+
 	if (!CurrentCapture.IsValid())
 	{
 		return;
@@ -247,6 +269,8 @@ void FNiagaraSystemInstance::FinishCapture()
 
 bool FNiagaraSystemInstance::QueryCaptureResults(const FGuid& RequestId, TArray<TSharedPtr<struct FNiagaraScriptDebuggerInfo, ESPMode::ThreadSafe>>& OutCaptureResults)
 {
+	WaitForAsyncTick();
+
 	if (CurrentCaptureGuid.IsValid() && RequestId == *CurrentCaptureGuid.Get())
 	{
 		return false;
@@ -327,6 +351,8 @@ void FNiagaraSystemInstance::SetSolo(bool bInSolo)
 		return;
 	}
 
+	WaitForAsyncTick();
+
 	UNiagaraSystem* System = GetSystem();
 	if (bInSolo)
 	{
@@ -352,7 +378,9 @@ void FNiagaraSystemInstance::SetSolo(bool bInSolo)
 void FNiagaraSystemInstance::Activate(EResetMode InResetMode)
 {
 	SCOPE_CYCLE_COUNTER(STAT_NiagaraSystemActivate);
-	
+
+	WaitForAsyncTick();
+
 	UNiagaraSystem* System = GetSystem();
 	if (System && System->IsValid() && IsReadyToRun())
 	{
@@ -367,6 +395,9 @@ void FNiagaraSystemInstance::Activate(EResetMode InResetMode)
 void FNiagaraSystemInstance::Deactivate(bool bImmediate)
 {
 	SCOPE_CYCLE_COUNTER(STAT_NiagaraSystemDeactivate);
+
+	WaitForAsyncTick();
+
 	if (IsComplete())
 	{
 		return;
@@ -417,7 +448,7 @@ bool FNiagaraSystemInstance::DeallocateSystemInstance(TUniquePtr< FNiagaraSystem
 void FNiagaraSystemInstance::Complete()
 {
 	SCOPE_CYCLE_COUNTER(STAT_NiagaraSystemComplete);
-	
+
 	// Only notify others if have yet to complete
 	bool bNeedToNotifyOthers = bNotifyOnCompletion;
 
@@ -446,6 +477,8 @@ void FNiagaraSystemInstance::Complete()
 
 	UnbindParameters();
 
+	bPendingSpawn = false;
+
 	if (bNeedToNotifyOthers)
 	{
 		// We've already notified once, no need to do so again.
@@ -467,6 +500,8 @@ void FNiagaraSystemInstance::SetPaused(bool bInPaused)
 	{
 		return;
 	}
+
+	WaitForAsyncTick();
 	
 	if (SystemInstanceIndex != INDEX_NONE)
 	{
@@ -500,7 +535,9 @@ void FNiagaraSystemInstance::Reset(FNiagaraSystemInstance::EResetMode Mode)
 		}*/
 		return;
 	}
-	
+
+	WaitForAsyncTick();
+
 	if (Component)
 	{
 		Component->SetLastRenderTime(Component->GetWorld()->GetTimeSeconds());
@@ -634,15 +671,19 @@ UNiagaraParameterCollectionInstance* FNiagaraSystemInstance::GetParameterCollect
 
 void FNiagaraSystemInstance::AdvanceSimulation(int32 TickCountToSimulate, float TickDeltaSeconds)
 {
-	if (TickCountToSimulate > 0)
+	if (TickCountToSimulate > 0 && !IsPaused())
 	{
+		WaitForAsyncTick();
+
 		SCOPE_CYCLE_COUNTER(STAT_NiagaraSystemAdvanceSim);
 		bool bWasSolo = bSolo;
 		SetSolo(true);
 
 		for (int32 TickIdx = 0; TickIdx < TickCountToSimulate; ++TickIdx)
 		{
-			ComponentTick(TickDeltaSeconds);
+			//Cannot do multiple tick off the game thread here without additional work. So we pass in null for the completion event which will force GT execution.
+			//If this becomes a perf problem I can add a new path for the tick code to handle multiple ticks.
+			ComponentTick(TickDeltaSeconds, nullptr);
 		}
 		SetSolo(bWasSolo);
 	}
@@ -696,6 +737,11 @@ bool DoSystemDataInterfacesRequireSolo(const UNiagaraSystem& System, const UNiag
 void FNiagaraSystemInstance::ReInitInternal()
 {
 	SCOPE_CYCLE_COUNTER(STAT_NiagaraSystemReinit);
+
+	SCOPE_CYCLE_COUNTER(STAT_NiagaraOverview_GT);
+	CSV_SCOPED_TIMING_STAT_EXCLUSIVE(Niagara);
+	LLM_SCOPE(ELLMTag::Niagara);
+
 	Age = 0;
 	TickCount = 0;
 	UNiagaraSystem* System = GetSystem();
@@ -870,6 +916,8 @@ FNiagaraSystemInstance::~FNiagaraSystemInstance()
 
 void FNiagaraSystemInstance::Cleanup()
 {
+	WaitForAsyncTick(true);
+
 	if (SystemInstanceIndex != INDEX_NONE)
 	{
 		TSharedPtr<FNiagaraSystemSimulation, ESPMode::ThreadSafe> SystemSim = GetSystemSimulation();
@@ -932,6 +980,8 @@ void FNiagaraSystemInstance::BindParameters()
 
 void FNiagaraSystemInstance::UnbindParameters()
 {
+	//TODO: Split unbind into internal and external.
+	//We need to unbind from external sources when we complete but internal can stay bound to improve reset times.
 
 	if (Component)
 	{
@@ -1010,6 +1060,8 @@ void FNiagaraSystemInstance::InitDataInterfaces()
 		return;
 	}
 
+	WaitForAsyncTick(true);
+
 	Component->GetOverrideParameters().Tick();
 	
 	DestroyDataInterfaceInstanceData();
@@ -1406,35 +1458,56 @@ void FNiagaraSystemInstance::InitEmitters()
 	}
 }
 
-void FNiagaraSystemInstance::ComponentTick(float DeltaSeconds)
+void FNiagaraSystemInstance::ComponentTick(float DeltaSeconds, const FGraphEventRef& MyCompletionGraphEvent)
 {
+	SCOPE_CYCLE_COUNTER(STAT_NiagaraOverview_GT);
+	SCOPE_CYCLE_COUNTER(STAT_NiagaraSystemInst_ComponentTickGT);
+	CSV_SCOPED_TIMING_STAT_EXCLUSIVE(Niagara);
+	LLM_SCOPE(ELLMTag::Niagara);
+
 	if (IsDisabled())
 	{
 		return;
 	}
 
-	TSharedPtr<FNiagaraSystemSimulation, ESPMode::ThreadSafe> Sim = GetSystemSimulation();
-	check(Sim.IsValid());
+	TSharedPtr<FNiagaraSystemSimulation, ESPMode::ThreadSafe> SystemSim = GetSystemSimulation();
+	check(SystemSim.IsValid());
 	check(IsInGameThread());
 	check(bSolo);
 	check(Component);
-	
-	TSharedPtr<FNiagaraSystemSimulation, ESPMode::ThreadSafe> SystemSim = GetSystemSimulation();
-	SystemSim->Tick(DeltaSeconds);
+
+	if (IsPendingSpawn())
+	{
+		SystemSim->SpawnNew_GameThread(DeltaSeconds, MyCompletionGraphEvent);
+	}
+	else
+	{
+		SystemSim->Tick_GameThread(DeltaSeconds, MyCompletionGraphEvent);
+	}
 }
 
-void FNiagaraSystemInstance::FinalizeTick(float DeltaSeconds)
+void FNiagaraSystemInstance::WaitForAsyncTick(bool bEnsureComplete)
 {
-	//Post tick our interfaces.
-	TickDataInterfaces(DeltaSeconds, true);
-
-	if (Component)
+	if (bAsyncWorkInProgress)
 	{
-		if (HasTickingEmitters())
+		ensureAlwaysMsgf(!bEnsureComplete, TEXT("Niagara System Async Task should be complete by now. %s"), *GetSystem()->GetPathName());
+		ensureAlwaysMsgf(IsInGameThread(), TEXT("NiagaraSystemInstance::WaitForAsyncTick() call is assuming execution on GT but is not on GT. %s"), *GetSystem()->GetPathName());
+
+		SCOPE_CYCLE_COUNTER(STAT_NiagaraSystemWaitForAsyncTick);
+
+		double StartTime = FPlatformTime::Seconds();
+		while (bAsyncWorkInProgress)
 		{
-			Component->UpdateComponentToWorld();//Needed for bounds updates. Can probably skip if using fixed bounds
-			Component->MarkRenderDynamicDataDirty();
+			FPlatformProcess::SleepNoStats(0.0f);
 		}
+		double StallTimeMS = (FPlatformTime::Seconds() - StartTime) * 1000.0f;
+
+		if (StallTimeMS > GWaitForAsyncStallWarnThresholdMS)
+		{
+			UE_LOG(LogNiagara, Warning, TEXT("Niagara Effect stalled GT for %g ms.\nComponent: %s \nSystem: %s"), StallTimeMS, *GetFullNameSafe(Component), *GetFullNameSafe(GetSystem()));
+		}
+
+		FinalizeTick_GameThread();
 	}
 }
 
@@ -1456,33 +1529,51 @@ bool FNiagaraSystemInstance::HandleCompletion()
 		Complete();
 		return true;
 	}
+
+	bPendingSpawn = false;
 	return false;
 }
 
-void FNiagaraSystemInstance::PreSimulateTick(float DeltaSeconds)
+void FNiagaraSystemInstance::Tick_GameThread(float DeltaSeconds)
 {
-	SCOPE_CYCLE_COUNTER(STAT_NiagaraSystemPreSimulateTick);
+	SCOPE_CYCLE_COUNTER(STAT_NiagaraSystemInst_TickGT);
+	SCOPE_CYCLE_COUNTER(STAT_NiagaraOverview_GT);
+	CSV_SCOPED_TIMING_STAT_EXCLUSIVE(Niagara);
+	LLM_SCOPE(ELLMTag::Niagara);
+
 	UNiagaraSystem* System = GetSystem();
 	FScopeCycleCounter SystemStat(System->GetStatID(true, true));
 
+	WaitForAsyncTick(true);
+
+	CachedDeltaSeconds = DeltaSeconds;
+	bNeedsFinalize = true;
+
 	TickInstanceParameters(DeltaSeconds);
 
+	TickDataInterfaces(DeltaSeconds, false);
+
 	Age += DeltaSeconds;
 	TickCount += 1;
+	
+	BeginAsyncWork();
 }
 
-void FNiagaraSystemInstance::PostSimulateTick(float DeltaSeconds)
+void FNiagaraSystemInstance::Tick_Concurrent()
 {
-	SCOPE_CYCLE_COUNTER(STAT_NiagaraSystemInstanceTick);
+	SCOPE_CYCLE_COUNTER(STAT_NiagaraSystemInst_TickCNC);
 	SCOPE_CYCLE_COUNTER(STAT_NiagaraOverview_GT_CNC);
+	CSV_SCOPED_TIMING_STAT_EXCLUSIVE(Niagara);
+	LLM_SCOPE(ELLMTag::Niagara);
 
 	// Reset values that will be accumulated during emitter tick.
 	TotalParamSize = 0;
 	ActiveGPUEmitterCount = 0;
 	UNiagaraSystem* System = GetSystem();
 
-	if (IsComplete() || !bHasTickingEmitters || GetSystem() == nullptr || Component == nullptr || DeltaSeconds < SMALL_NUMBER)
+	if (IsComplete() || !bHasTickingEmitters || GetSystem() == nullptr || Component == nullptr || CachedDeltaSeconds < SMALL_NUMBER)
 	{
+		bAsyncWorkInProgress = false;
 		return;
 	}
 
@@ -1498,7 +1589,7 @@ void FNiagaraSystemInstance::PostSimulateTick(float DeltaSeconds)
 	for (int32 EmitterIdx = 0; EmitterIdx < Emitters.Num(); EmitterIdx++)
 	{
 		FNiagaraEmitterInstance& Inst = Emitters[EmitterIdx].Get();
-		Inst.Tick(DeltaSeconds);
+		Inst.Tick(CachedDeltaSeconds);
 
 		if (Inst.GetCachedEmitter()->SimTarget == ENiagaraSimTarget::GPUComputeSim && Inst.GetGPUContext() != nullptr && (Inst.GetExecutionState() != ENiagaraExecutionState::Complete))
 		{
@@ -1506,6 +1597,66 @@ void FNiagaraSystemInstance::PostSimulateTick(float DeltaSeconds)
 			ActiveGPUEmitterCount++;
 		}
 	}
+
+	bAsyncWorkInProgress = false;
+}
+
+void FNiagaraSystemInstance::FinalizeTick_GameThread()
+{
+	if (bNeedsFinalize)//We can come in here twice in one tick if the GT calls WaitForAsync() while there is a GT finalize task in the queue.
+	{
+		SCOPE_CYCLE_COUNTER(STAT_NiagaraOverview_GT);
+		SCOPE_CYCLE_COUNTER(STAT_NiagaraSystemInst_FinalizeGT);
+		CSV_SCOPED_TIMING_STAT_EXCLUSIVE(Niagara);
+		LLM_SCOPE(ELLMTag::Niagara);
+
+		bNeedsFinalize = false;
+		if (!HandleCompletion())
+		{
+			//Post tick our interfaces.
+			TickDataInterfaces(CachedDeltaSeconds, true);
+
+			bool bIsRegistered = true;
+			if (Component)
+			{
+				Component->UpdateComponentToWorld();//Needed for bounds updates. Can probably skip if using fixed bounds
+				Component->MarkRenderDynamicDataDirty();
+				bIsRegistered = Component->IsRegistered();
+			}
+
+			// Push any GPU ticks for this system instance.
+			if (ActiveGPUEmitterCount > 0 && NiagaraSupportsComputeShaders(GetBatcher()->GetShaderPlatform()) && bIsRegistered)
+			{
+				ensure(!IsComplete());
+				FNiagaraGPUSystemTick GPUTick;
+				GPUTick.Init(this);
+
+				//if (GPUTick.DIInstanceData)
+				//{
+				//	uint8* BasePointer = (uint8*)GPUTick.DIInstanceData->PerInstanceDataForRT;
+
+				//	//UE_LOG(LogNiagara, Log, TEXT("GT Testing (dipacket) %p (baseptr) %p"), GPUTick.DIInstanceData, BasePointer);
+				//	for (auto& Pair : GPUTick.DIInstanceData->InterfaceProxiesToOffsets)
+				//	{
+				//		FNiagaraDataInterfaceProxy* Proxy = Pair.Key;
+				//		UE_LOG(LogNiagara, Log, TEXT("\tGT (proxy) %p (size) %u"), Proxy, Proxy->PerInstanceDataPassedToRenderThreadSize());
+				//	}
+				//}
+
+				// We will give the data over to the render thread. It is responsible for freeing it.
+				// We no longer own it and cannot modify it after this point.
+				// @todo We are taking a copy of the object here. This object is small so this overhead should
+				// not be very high. And we avoid making a bunch of small allocations here.
+				NiagaraEmitterInstanceBatcher* TheBatcher = GetBatcher();
+				ENQUEUE_RENDER_COMMAND(FGiveSystemInstanceTickToRT)(
+					[TheBatcher, GPUTick](FRHICommandListImmediate& RHICmdList) mutable
+				{
+					TheBatcher->GiveSystemTick_RenderThread(GPUTick);
+				}
+				);
+			}
+		}
+	}
 }
 
 #if WITH_EDITORONLY_DATA
diff --git a/Engine/Plugins/FX/Niagara/Source/Niagara/Private/NiagaraSystemSimulation.cpp b/Engine/Plugins/FX/Niagara/Source/Niagara/Private/NiagaraSystemSimulation.cpp
--- a/Engine/Plugins/FX/Niagara/Source/Niagara/Private/NiagaraSystemSimulation.cpp
+++ b/Engine/Plugins/FX/Niagara/Source/Niagara/Private/NiagaraSystemSimulation.cpp
@@ -15,14 +15,19 @@
 #include "NiagaraWorldManager.h"
 #include "NiagaraEmitterInstanceBatcher.h"
 
-DECLARE_CYCLE_STAT(TEXT("System Simulation [GT]"), STAT_NiagaraSystemSim, STATGROUP_Niagara);
-DECLARE_CYCLE_STAT(TEXT("System Pre Simulate [GT]"), STAT_NiagaraSystemSim_PreSimulate, STATGROUP_Niagara);
-DECLARE_CYCLE_STAT(TEXT("System Prepare For Simulate [GT]"), STAT_NiagaraSystemSim_PrepareForSimulate, STATGROUP_Niagara);
-DECLARE_CYCLE_STAT(TEXT("System Sim Update [GT]"), STAT_NiagaraSystemSim_Update, STATGROUP_Niagara);
-DECLARE_CYCLE_STAT(TEXT("System Sim Spawn [GT]"), STAT_NiagaraSystemSim_Spawn, STATGROUP_Niagara);
-DECLARE_CYCLE_STAT(TEXT("System Sim Transfer Parameters [GT]"), STAT_NiagaraSystemSim_TransferParameters, STATGROUP_Niagara);
-DECLARE_CYCLE_STAT(TEXT("System Post Simulate [GT]"), STAT_NiagaraSystemSim_PostSimulate, STATGROUP_Niagara);
-DECLARE_CYCLE_STAT(TEXT("System Mark Component Dirty [GT]"), STAT_NiagaraSystemSim_MarkComponentDirty, STATGROUP_Niagara);
+//High level stats for system sim tick.
+DECLARE_CYCLE_STAT(TEXT("System Simulaton Tick [GT]"), STAT_NiagaraSystemSim_TickGT, STATGROUP_Niagara);
+DECLARE_CYCLE_STAT(TEXT("System Simulaton Tick [CNC]"), STAT_NiagaraSystemSim_TickCNC, STATGROUP_Niagara);
+DECLARE_CYCLE_STAT(TEXT("System Simulaton SpawnNew [GT]"), STAT_NiagaraSystemSim_SpawnNewGT, STATGROUP_Niagara);
+//Some more detailed stats for system sim tick
+DECLARE_CYCLE_STAT(TEXT("System Prepare For Simulate [CNC]"), STAT_NiagaraSystemSim_PrepareForSimulateCNC, STATGROUP_Niagara);
+DECLARE_CYCLE_STAT(TEXT("System Sim Update [CNC]"), STAT_NiagaraSystemSim_UpdateCNC, STATGROUP_Niagara);
+DECLARE_CYCLE_STAT(TEXT("System Sim Spawn [CNC]"), STAT_NiagaraSystemSim_SpawnCNC, STATGROUP_Niagara);
+DECLARE_CYCLE_STAT(TEXT("System Sim Transfer Results [CNC]"), STAT_NiagaraSystemSim_TransferResultsCNC, STATGROUP_Niagara);
+
+
+DECLARE_CYCLE_STAT(TEXT("ForcedWaitForAsync"), STAT_NiagaraSystemSim_ForceWaitForAsync, STATGROUP_Niagara);
+DECLARE_CYCLE_STAT(TEXT("ForcedWait Fake Stall"), STAT_NiagaraSystemSim_ForceWaitFakeStall, STATGROUP_Niagara);
 
 
 static int32 GbDumpSystemData = 0;
@@ -41,31 +46,209 @@ static FAutoConsoleVariableRef CVarSystemUpdateOnSpawn(
 	ECVF_Default
 );
 
-//Pretick can no longer be run in parallel. Will likely remain this way.
-// static int32 GbParallelSystemPreTick = 0;
-// static FAutoConsoleVariableRef CVarNiagaraParallelSystemPreTick(
-// 	TEXT("fx.ParallelSystemPreTick"),
-// 	GbParallelSystemPreTick,
-// 	TEXT("If > 0, system pre tick is parallelized. \n"),
-// 	ECVF_Default
-// );
-
-static int32 GbParallelSystemPostTick = 1;
-static FAutoConsoleVariableRef CVarNiagaraParallelSystemPostTick(
-	TEXT("fx.ParallelSystemPostTick"),
-	GbParallelSystemPostTick,
+static int32 GbParallelSystemSimTick = 1;
+static FAutoConsoleVariableRef CVarParallelSystemSimTick(
+	TEXT("fx.ParallelSystemSimTick"),
+	GbParallelSystemSimTick,
 	TEXT("If > 0, system post tick is parallelized. \n"),
 	ECVF_Default
 );
 
-//TODO: Experiment with parallel param transfer.
-//static int32 GbParallelSystemParamTransfer = 1;
-//static FAutoConsoleVariableRef CVarNiagaraParallelSystemParamTransfer(
-//	TEXT("fx.ParallelSystemParamTransfer"),
-//	GbParallelSystemParamTransfer,
-//	TEXT("If > 0, system param transfer is parallelized. \n"),
-//	ECVF_Default
-//);
+static int32 GbParallelSystemInstanceTick = 1;
+static FAutoConsoleVariableRef CVarParallelSystemInstanceTick(
+	TEXT("fx.ParallelSystemInstanceTick"),
+	GbParallelSystemInstanceTick,
+	TEXT("If > 0, system post tick is parallelized. \n"),
+	ECVF_Default
+);
+
+static int32 GbParallelSystemInstanceTickBatchSize = NiagaraSystemTickBatchSize;
+static FAutoConsoleVariableRef CVarParallelSystemInstanceTickBatchSize(
+	TEXT("fx.ParallelSystemInstanceTickBatchSize"),
+	GbParallelSystemInstanceTickBatchSize,
+	TEXT("The number of system instances to process per async task. \n"),
+	ECVF_Default
+);
+
+static int32 GbSystemSimTransferParamsParallelThreshold = 64;
+static FAutoConsoleVariableRef CVarSystemSimTransferParamsParallelThreshold(
+	TEXT("fx.SystemSimTransferParamsParallelThreshold"),
+	GbSystemSimTransferParamsParallelThreshold,
+	TEXT("The number of system instances required for the transfer parameters portion of the system tick to go wide. \n"),
+	ECVF_Default
+);
+
+static int32 GbForceNiagaraAsyncWait = 0;
+static FAutoConsoleVariableRef CVarForceNiagaraAsyncWait(
+	TEXT("fx.ForceNiagaraAsyncWait"),
+	GbForceNiagaraAsyncWait,
+	TEXT(""),
+	ECVF_Default
+);
+
+static float GbWaitForAsyncStallMS = .5f;
+static FAutoConsoleVariableRef CVarWaitForAsyncStallMS(
+	TEXT("fx.WaitForAsyncStallMS"),
+	GbWaitForAsyncStallMS,
+	TEXT(""),
+	ECVF_Default
+);
+
+//////////////////////////////////////////////////////////////////////////
+
+FNiagaraSystemSimulationTickContext::FNiagaraSystemSimulationTickContext(FNiagaraSystemSimulation* InOwner, float InDeltaSeconds, bool bInPendingSpawnPass, const FGraphEventRef& InMyCompletionGraphEvent)
+	: Owner(InOwner)
+	, DeltaSeconds(InDeltaSeconds)
+	, bPendingSpawnPass(bInPendingSpawnPass)
+	, MyCompletionGraphEvent(InMyCompletionGraphEvent)
+{
+	System = Owner->GetSystem();
+	
+	bTickAsync = Owner->ShouldTickAsync(*this);
+	bTickInstancesAsync = Owner->ShouldTickInstancesAsync(*this);
+}
+
+TArray<FNiagaraSystemInstance*>& FNiagaraSystemSimulationTickContext::GetInstances()
+{
+	return bPendingSpawnPass ? Owner->SpawningSystemInstances : Owner->SystemInstances;
+}
+
+FNiagaraDataSet& FNiagaraSystemSimulationTickContext::GetDataSet()
+{
+	return bPendingSpawnPass ? Owner->PendingSpawnDataSet : Owner->DataSet;
+}
+
+//////////////////////////////////////////////////////////////////////////
+//////////////////////////////////////////////////////////////////////////
+
+FAutoConsoleTaskPriority CPrio_NiagaraSystemSimulationTickTask(
+	TEXT("TaskGraph.TaskPriorities.NiagaraSystemSimulationTickcTask"),
+	TEXT("Task and thread priority for FNiagaraSystemSimulationTickTask."),
+	ENamedThreads::HighThreadPriority, // if we have high priority task threads, then use them...
+	ENamedThreads::NormalTaskPriority, // .. at normal task priority
+	ENamedThreads::HighTaskPriority // if we don't have hi pri threads, then use normal priority threads at high task priority instead
+);
+
+//This task performs the concurrent part of the system simulation tick.
+class FNiagaraSystemSimulationTickTask
+{
+	FNiagaraSystemSimulationTickContext Context;
+public:
+	FNiagaraSystemSimulationTickTask(FNiagaraSystemSimulationTickContext InContext)
+		: Context(InContext)
+	{
+	}
+
+	FORCEINLINE TStatId GetStatId() const
+	{
+		RETURN_QUICK_DECLARE_CYCLE_STAT(FNiagaraSystemSimulationTickTask, STATGROUP_TaskGraphTasks);
+	}
+
+	ENamedThreads::Type GetDesiredThread()
+	{
+		return CPrio_NiagaraSystemSimulationTickTask.Get();
+	}
+
+	static ESubsequentsMode::Type GetSubsequentsMode() { return ESubsequentsMode::TrackSubsequents; }
+
+	void DoTask(ENamedThreads::Type CurrentThread, const FGraphEventRef& MyCompletionGraphEvent)
+	{
+		Context.MyCompletionGraphEvent = MyCompletionGraphEvent;
+		Context.Owner->Tick_Concurrent(Context);
+	}
+};
+
+//////////////////////////////////////////////////////////////////////////
+
+/** 
+Task to call FinalizeTick_GameThread() on a batch of FNiagaraSystemInstances 
+Must be done on the game thread.
+*/
+class FNiagaraSystemInstanceFinalizeTask
+{
+	FNiagaraSystemSimulation* SystemSim;
+	FNiagaraSystemTickBatch Batch;
+public:
+	FNiagaraSystemInstanceFinalizeTask(FNiagaraSystemSimulation* InSystemSim, FNiagaraSystemTickBatch& InBatch)
+		: SystemSim(InSystemSim)
+		, Batch(InBatch)
+	{
+
+	}
+
+	FORCEINLINE TStatId GetStatId() const
+	{
+		RETURN_QUICK_DECLARE_CYCLE_STAT(FNiagaraSystemInstanceFinalizeTask, STATGROUP_TaskGraphTasks);
+	}
+
+	ENamedThreads::Type GetDesiredThread()
+	{
+		return ENamedThreads::GameThread;
+	}
+
+	static ESubsequentsMode::Type GetSubsequentsMode() { return ESubsequentsMode::TrackSubsequents; }
+
+	void DoTask(ENamedThreads::Type CurrentThread, const FGraphEventRef& MyCompletionGraphEvent)
+	{
+		check(CurrentThread == ENamedThreads::GameThread);
+
+		for (FNiagaraSystemInstance* Inst : Batch)
+		{
+			Inst->FinalizeTick_GameThread();
+		}
+	}
+};
+
+FAutoConsoleTaskPriority CPrio_NiagaraSystemInstanceAsyncTask(
+	TEXT("TaskGraph.TaskPriorities.NiagaraSystemAsyncTask"),
+	TEXT("Task and thread priority for FNiagaraSystemAsyncTask."),
+	ENamedThreads::HighThreadPriority, // if we have high priority task threads, then use them...
+	ENamedThreads::NormalTaskPriority, // .. at normal task priority
+	ENamedThreads::HighTaskPriority // if we don't have hi pri threads, then use normal priority threads at high task priority instead
+);
+
+/** 
+Async task to call Tick_Concurrent() on batches of FNiagaraSystemInstances.
+Can be performed on task threads. 
+*/
+class FNiagaraSystemInstanceAsyncTask
+{
+	FNiagaraSystemSimulation* SystemSim;
+	FNiagaraSystemTickBatch Batch;
+
+public:
+	FNiagaraSystemInstanceAsyncTask(FNiagaraSystemSimulation* InSystemSim, FNiagaraSystemTickBatch& InBatch)
+		: SystemSim(InSystemSim)
+		, Batch(InBatch)
+	{
+
+	}
+
+	FORCEINLINE TStatId GetStatId() const
+	{
+		RETURN_QUICK_DECLARE_CYCLE_STAT(FNiagaraSystemInstanceAsyncTask, STATGROUP_TaskGraphTasks);
+	}
+
+	ENamedThreads::Type GetDesiredThread()
+	{
+		return CPrio_NiagaraSystemInstanceAsyncTask.Get();
+	}
+
+	static ESubsequentsMode::Type GetSubsequentsMode() { return ESubsequentsMode::TrackSubsequents; }
+
+	void DoTask(ENamedThreads::Type CurrentThread, const FGraphEventRef& MyCompletionGraphEvent)
+	{
+		for (FNiagaraSystemInstance* Inst : Batch)
+		{
+			Inst->Tick_Concurrent();
+		}
+
+		//Now kick off a finalize task for this batch on the game thread.
+		FGraphEventRef FinalizeTask = TGraphTask<FNiagaraSystemInstanceFinalizeTask>::CreateTask(nullptr, CurrentThread).ConstructAndDispatchWhenReady(SystemSim, Batch);
+		MyCompletionGraphEvent->SetGatherThreadForDontCompleteUntil(ENamedThreads::GameThread);
+		MyCompletionGraphEvent->DontCompleteUntil(FinalizeTask);
+	}
+};
 
 //////////////////////////////////////////////////////////////////////////
 
@@ -83,6 +266,9 @@ bool FNiagaraSystemSimulation::Init(UNiagaraSystem* InSystem, UWorld* InWorld, b
 
 	bIsSolo = bInIsSolo;
 
+	bBindingsInitialized = false;
+	bInSpawnPhase = false;
+
 	FNiagaraWorldManager* WorldMan = FNiagaraWorldManager::Get(InWorld);
 	check(WorldMan);
 
@@ -91,11 +277,21 @@ bool FNiagaraSystemSimulation::Init(UNiagaraSystem* InSystem, UWorld* InWorld, b
 
 	if (bCanExecute)
 	{
+		//TODO: Move all layout data into the System!
+
+		//Initilize the main simulation dataset.
 		DataSet.Init(FNiagaraDataSetID(), ENiagaraSimTarget::CPUSim, System->GetFullName());
 		DataSet.AddVariables(System->GetSystemSpawnScript()->GetVMExecutableData().Attributes);
 		DataSet.AddVariables(System->GetSystemUpdateScript()->GetVMExecutableData().Attributes);
 		DataSet.Finalize();
 
+		//Initialize the data set for newly spawned systems.
+		PendingSpawnDataSet.Init(FNiagaraDataSetID(), ENiagaraSimTarget::CPUSim, System->GetFullName());
+		PendingSpawnDataSet.AddVariables(System->GetSystemSpawnScript()->GetVMExecutableData().Attributes);
+		PendingSpawnDataSet.AddVariables(System->GetSystemUpdateScript()->GetVMExecutableData().Attributes);
+		PendingSpawnDataSet.Finalize();
+
+		//Initialize the dataset for paused systems.
 		PausedInstanceData.Init(FNiagaraDataSetID(), ENiagaraSimTarget::CPUSim, System->GetFullName());
 		PausedInstanceData.AddVariables(System->GetSystemSpawnScript()->GetVMExecutableData().Attributes);
 		PausedInstanceData.AddVariables(System->GetSystemUpdateScript()->GetVMExecutableData().Attributes);
@@ -190,6 +386,7 @@ bool FNiagaraSystemSimulation::Init(UNiagaraSystem* InSystem, UWorld* InWorld, b
 
 void FNiagaraSystemSimulation::Destroy()
 {
+	check(IsInGameThread());
 	while (SystemInstances.Num())
 	{
 		SystemInstances.Last()->Deactivate(true);
@@ -198,6 +395,10 @@ void FNiagaraSystemSimulation::Destroy()
 	{
 		PendingSystemInstances.Last()->Deactivate(true);
 	}
+	SystemInstances.Empty();
+	PendingSystemInstances.Empty();
+	SpawningSystemInstances.Empty();
+
 
 	FNiagaraWorldManager* WorldMan = FNiagaraWorldManager::Get(World);
 	check(WorldMan);
@@ -231,16 +432,19 @@ void FNiagaraSystemSimulation::TransferInstance(FNiagaraSystemSimulation* Source
 	check(SourceSimulation->GetSystem() == GetSystem());
 	check(SystemInst);
 
+	check(!SystemInst->IsPaused());
+	check(!bInSpawnPhase);
+	check(!SourceSimulation->bInSpawnPhase);
+
 	int32 SystemInstIdx = SystemInst->SystemInstanceIndex;
-	int32 NewDataIndex = INDEX_NONE;
-	if (!SystemInst->IsPendingSpawn())
+	if (!SystemInst->IsPendingSpawn() && SystemInst->SystemInstanceIndex != INDEX_NONE)
 	{
 // 		UE_LOG(LogNiagara, Log, TEXT("== Dataset Transfer ========================"));
 // 		UE_LOG(LogNiagara, Log, TEXT(" ----- Existing values in src. Idx: %d -----"), SystemInstIdx);
 // 		SourceSimulation->DataSet.Dump(true, SystemInstIdx, 1);
 
 		//If we're not pending then the system actually has data to pull over. This is not fast.
-		NewDataIndex = DataSet.GetCurrentDataChecked().TransferInstance(SourceSimulation->DataSet.GetCurrentDataChecked(), SystemInstIdx);
+		int32 NewDataSetIndex = DataSet.GetCurrentDataChecked().TransferInstance(SourceSimulation->DataSet.GetCurrentDataChecked(), SystemInstIdx, false);
 
 // 		UE_LOG(LogNiagara, Log, TEXT(" ----- Transfered values in dest. Idx: %d -----"), NewDataIndex);
 // 		DataSet.Dump(true, NewDataIndex, 1);
@@ -249,12 +453,12 @@ void FNiagaraSystemSimulation::TransferInstance(FNiagaraSystemSimulation* Source
 	
 		//Move the system direct to the new sim's 
 		SystemInst->SystemInstanceIndex = SystemInstances.Add(SystemInst);
-		if (SystemInst->SystemInstanceIndex == 0)
+		check(NewDataSetIndex == SystemInst->SystemInstanceIndex);
+
+		if (!bBindingsInitialized)
 		{
 			InitParameterDataSetBindings(SystemInst);
 		}
-
-		check(NewDataIndex == SystemInst->SystemInstanceIndex);
 	}
 	else
 	{
@@ -262,6 +466,8 @@ void FNiagaraSystemSimulation::TransferInstance(FNiagaraSystemSimulation* Source
 
 		AddInstance(SystemInst);			
 	}
+
+	SystemInst->SystemSimulation = this->AsShared();
 }
 
 void FNiagaraSystemSimulation::DumpInstance(const FNiagaraSystemInstance* Inst)const
@@ -277,459 +483,609 @@ void FNiagaraSystemSimulation::DumpInstance(const FNiagaraSystemInstance* Inst)c
 	DataSet.Dump(Inst->SystemInstanceIndex, 1, TEXT("System Data"));
 }
 
-bool FNiagaraSystemSimulation::Tick(float DeltaSeconds)
+bool FNiagaraSystemSimulation::ShouldTickAsync(const FNiagaraSystemSimulationTickContext& Context)
 {
-	UNiagaraSystem* System = WeakSystem.Get();
+	return GbParallelSystemSimTick && FApp::ShouldUseThreadingForPerformance() && !Context.bPendingSpawnPass && Context.MyCompletionGraphEvent.IsValid();
+}
 
-	if (System == nullptr || bCanExecute == false)
-	{
-		// TODO: evaluate whether or not we should have removed this from the world manager instead?
-		return false;
-	}
-	else if (World == nullptr)
+bool FNiagaraSystemSimulation::ShouldTickInstancesAsync(const FNiagaraSystemSimulationTickContext& Context)
+{
+	return GbParallelSystemInstanceTick && !bIsSolo && FApp::ShouldUseThreadingForPerformance() && !Context.bPendingSpawnPass && Context.MyCompletionGraphEvent.IsValid();
+}
+
+void FNiagaraSystemSimulation::AddSystemToTickBatch(FNiagaraSystemInstance* Instance, FNiagaraSystemSimulationTickContext& Context)
+{
+	TickBatch.Add(Instance);
+	if (TickBatch.Num() == GbParallelSystemInstanceTickBatchSize)
 	{
-		ensureMsgf(false, TEXT("World was null while ticking System Simulation!"));
-		return false;
+		FlushTickBatch(Context);
 	}
+}
 
-	FScopeCycleCounter SystemStatCounter(System->GetStatID(true, false));
-
-	if (MaxDeltaTime.IsSet())
+void FNiagaraSystemSimulation::FlushTickBatch(FNiagaraSystemSimulationTickContext& Context)
+{
+	if (TickBatch.Num() > 0)
 	{
-		DeltaSeconds = FMath::Clamp(DeltaSeconds, 0.0f, MaxDeltaTime.GetValue());
+		if (Context.bTickInstancesAsync)
+		{
+			checkSlow(Context.MyCompletionGraphEvent.IsValid());
+			//If we're able, kick off a task to process this batch on a task thread.
+			//When this task has finished the concurrent ticks for this batch it will enqueue a finalize task for the batch on the game thread.
+			FGraphEventRef AsyncTask = TGraphTask<FNiagaraSystemInstanceAsyncTask>::CreateTask(nullptr).ConstructAndDispatchWhenReady(this, TickBatch);
+			Context.MyCompletionGraphEvent->SetGatherThreadForDontCompleteUntil(ENamedThreads::GameThread);
+			Context.MyCompletionGraphEvent->DontCompleteUntil(AsyncTask);
+		}
+		else
+		{
+			//Otherwise just do directly here. This may already be on a task thread or could be on game thread.
+			for (FNiagaraSystemInstance* Inst : TickBatch)
+			{
+				Inst->Tick_Concurrent();
+			}
+
+			if (Context.bTickAsync)
+			{
+				checkSlow(Context.MyCompletionGraphEvent.IsValid());
+				//If we're ticking off the main thread then we still need to add a finalize task on the GT for this batch.
+				FGraphEventRef FinalizeTask = TGraphTask<FNiagaraSystemInstanceFinalizeTask>::CreateTask(nullptr).ConstructAndDispatchWhenReady(this, TickBatch);
+				Context.MyCompletionGraphEvent->SetGatherThreadForDontCompleteUntil(ENamedThreads::GameThread);
+				Context.MyCompletionGraphEvent->DontCompleteUntil(FinalizeTask);
+			}
+		}
+		TickBatch.Reset();
 	}
+}
 
-	UNiagaraScript* SystemSpawnScript = System->GetSystemSpawnScript();
-	UNiagaraScript* SystemUpdateScript = System->GetSystemUpdateScript();
-#if WITH_EDITOR
-	SystemSpawnScript->RapidIterationParameters.Tick();
-	SystemUpdateScript->RapidIterationParameters.Tick();
-#endif
+void FNiagaraSystemSimulation::SpawnNew_GameThread(float DeltaSeconds, const FGraphEventRef& MyCompletionGraphEvent)
+{
+	SCOPE_CYCLE_COUNTER(STAT_NiagaraSystemSim_SpawnNewGT);
+	SCOPE_CYCLE_COUNTER(STAT_NiagaraOverview_GT);
+	CSV_SCOPED_TIMING_STAT_EXCLUSIVE(Niagara);
+	LLM_SCOPE(ELLMTag::Niagara);
 
-	SCOPE_CYCLE_COUNTER(STAT_NiagaraSystemSim);
+	check(SystemInstances.Num() == DataSet.GetCurrentDataChecked().GetNumInstances());
+	check(PausedSystemInstances.Num() == PausedInstanceData.GetCurrentDataChecked().GetNumInstances());
 
-	int32 OrigNum = SystemInstances.Num();
-	int32 SpawnNum = 0;
-	int32 NewNum = 0;
+	//We have to do new spwaning at the end of the frame to account for new niagara systmes spawned after our main simulation tick.
+	//So here we build a list of instances to actually spawn from the PendingSpawnInstances.
+	//Then spwan these into a new dataset which is then transfered into the main dataset.
+	//Currently this is all done on the game thread but we could likely also move this off the GT too.
 
-	{
-		SCOPE_CYCLE_COUNTER(STAT_NiagaraSystemSim_PreSimulate);
+	UNiagaraSystem* System = WeakSystem.Get();
 
+	check(IsInGameThread());
+	if (PendingSystemInstances.Num() > 0)
+	{
+		int32 SystemIndex = 0;
+		SpawningSystemInstances.Reserve(PendingSystemInstances.Num());
+		
+		while (SystemIndex < PendingSystemInstances.Num())
 		{
-			int32 SystemIndex = 0;
-			while ( SystemIndex < SystemInstances.Num() )
-			{
-				FNiagaraSystemInstance* Inst = SystemInstances[SystemIndex];
-				Inst->TickDataInterfaces(DeltaSeconds, false);
+			FNiagaraSystemInstance* Inst = PendingSystemInstances[SystemIndex];
 
+			//If an instance is paused, just leave it in the PendingSpawnInstances.
+			if (!Inst->IsPaused())
+			{
+				Inst->Tick_GameThread(DeltaSeconds);
+								
 				// TickDataInterfaces could remove the system so we only increment if the system has changed
 				if (Inst->SystemInstanceIndex != INDEX_NONE)
 				{
+					checkSlow(!Inst->IsComplete());
 					checkSlow(Inst->SystemInstanceIndex == SystemIndex);
-					++SystemIndex;
-				}
-				else
-				{
-					checkSlow((SystemInstances.Num() <= SystemIndex) || (Inst == SystemInstances[SystemIndex]));
-				}
-			}
-			OrigNum = SystemInstances.Num();
+					
+					if (!bBindingsInitialized)
+					{
+						// When the first instance is added we need to initialize the parameter store to data set bindings.
+						InitParameterDataSetBindings(Inst);
+					}
 
-			//Pre tick and gather any still valid pending instances for spawn.
-			SystemInstances.Reserve(NewNum);
-			SpawnNum = 0;
-			int32 NumPending = PendingSystemInstances.Num();
-			for (int32 i = NumPending - 1; i >= 0; --i)
-			{
-				FNiagaraSystemInstance* Inst = PendingSystemInstances[i];
+					int32 SpawnInstIdx = SpawningSystemInstances.Add(Inst);
+					PendingSystemInstances.RemoveAtSwap(SystemIndex);
 
-				//Don't spawn systems that are paused. Keep them in pending list so they are spawned when unpaused.
-				if (Inst->IsPaused())
-				{
-					if (GbDumpSystemData || (System && System->bDumpDebugSystemInfo))
+					Inst->SystemInstanceIndex = SpawnInstIdx;
+					if (PendingSystemInstances.IsValidIndex(SystemIndex))
 					{
-						UE_LOG(LogNiagara, Log, TEXT("=== Skipping Paused Pending Spawn %d ==="), Inst->SystemInstanceIndex);
-						//DataSet.Dump(true, Instance->SystemInstanceIndex, 1);
+						PendingSystemInstances[SystemIndex]->SystemInstanceIndex = SystemIndex;
 					}
+
 					continue;
 				}
+			}
+			
+			++SystemIndex;
+		}		
 
-				check(Inst->SystemInstanceIndex == i);
-				PendingSystemInstances.RemoveAt(i);
+		if (SpawningSystemInstances.Num() > 0)
+		{
+			//Lets any RemoveInstance calls inside Tick_Concurrent know that we're spawning and all our instances are in SpawningSystemInstances;
+			bInSpawnPhase = true;
 
-				if (GbDumpSystemData || (System && System->bDumpDebugSystemInfo))
-				{
-					UE_LOG(LogNiagara, Log, TEXT("=== Spawning %d -> %d ==="), Inst->SystemInstanceIndex, SystemInstances.Num());
-					//DataSet.Dump(true, Instance->SystemInstanceIndex, 1);
-				}
+			//For now just do the concurrent tick for new systems on the GT here but we may also want to push this off to task threads too.
+			FNiagaraSystemSimulationTickContext Context(this, DeltaSeconds, true, MyCompletionGraphEvent);
+			Tick_Concurrent(Context);
 
-				Inst->TickDataInterfaces(DeltaSeconds, false);
-				Inst->SetPendingSpawn(false);
-				if (!Inst->IsComplete())
-				{
-					Inst->SystemInstanceIndex = SystemInstances.Add(Inst);
-					if (Inst->SystemInstanceIndex == 0)
-					{
-						// When the first instance is added we need to initialize the parameter store to data set bindings.
-						InitParameterDataSetBindings(Inst);
-					}
-					++SpawnNum;
-				}
-				else
-				{
-					checkSlow(Inst->SystemInstanceIndex == INDEX_NONE);
-				}
+			check(SystemInstances.Num() == DataSet.GetCurrentDataChecked().GetNumInstances());
+			check(PausedSystemInstances.Num() == PausedInstanceData.GetCurrentDataChecked().GetNumInstances());
+			check(SpawningSystemInstances.Num() == PendingSpawnDataSet.GetCurrentDataChecked().GetNumInstances());
+
+			//Copy the sim data.
+			PendingSpawnDataSet.CopyTo(DataSet, 0, INDEX_NONE, false);
+			PendingSpawnDataSet.ResetBuffers();
+
+			//Move new instances over to the main buffer and fix their indices.
+			SystemInstances.Reserve(SystemInstances.Num() + SpawningSystemInstances.Num());
+			for(FNiagaraSystemInstance* NewInst : SpawningSystemInstances)
+			{
+				checkSlow(!NewInst->IsComplete());
+				NewInst->SystemInstanceIndex = SystemInstances.Add(NewInst);
 			}
+
+			check(SystemInstances.Num() == DataSet.GetCurrentDataChecked().GetNumInstances());
+			check(PausedSystemInstances.Num() == PausedInstanceData.GetCurrentDataChecked().GetNumInstances());
+
+			SpawningSystemInstances.Reset();
+			bInSpawnPhase = false;
 		}
+	}
+}
+
+/** First phase of system sim tick. Must run on GameThread. */
+void FNiagaraSystemSimulation::Tick_GameThread(float DeltaSeconds, const FGraphEventRef& MyCompletionGraphEvent)
+{
+	check(IsInGameThread());
+	SCOPE_CYCLE_COUNTER(STAT_NiagaraOverview_GT);
+	SCOPE_CYCLE_COUNTER(STAT_NiagaraSystemSim_TickGT);
+	CSV_SCOPED_TIMING_STAT_EXCLUSIVE(Niagara);
+	LLM_SCOPE(ELLMTag::Niagara);
+
+	check(SystemInstances.Num() == DataSet.GetCurrentDataChecked().GetNumInstances());
+	check(PausedSystemInstances.Num() == PausedInstanceData.GetCurrentDataChecked().GetNumInstances());
 
-		NewNum = OrigNum + SpawnNum;
-		check(NewNum == SystemInstances.Num());
+	UNiagaraSystem* System = WeakSystem.Get();
+
+	FScopeCycleCounter SystemStatCounter(System->GetStatID(true, false));
+
+	if (MaxDeltaTime.IsSet())
+	{
+		DeltaSeconds = FMath::Clamp(DeltaSeconds, 0.0f, MaxDeltaTime.GetValue());
 	}
 
-	if (bCanExecute && NewNum > 0)
+	UNiagaraScript* SystemSpawnScript = System->GetSystemSpawnScript();
+	UNiagaraScript* SystemUpdateScript = System->GetSystemUpdateScript();
+#if WITH_EDITOR
+	SystemSpawnScript->RapidIterationParameters.Tick();
+	SystemUpdateScript->RapidIterationParameters.Tick();
+#endif
+
+	int32 SystemIndex = 0;
+	while (SystemIndex < SystemInstances.Num())
 	{
-		if (GbDumpSystemData || System->bDumpDebugSystemInfo)
+		FNiagaraSystemInstance* Inst = SystemInstances[SystemIndex];
+		Inst->Tick_GameThread(DeltaSeconds);
+
+		// TickDataInterfaces could remove the system so we only increment if the system has changed
+		if (Inst->SystemInstanceIndex != INDEX_NONE)
 		{
-			UE_LOG(LogNiagara, Log, TEXT("=========================================================="));
-			UE_LOG(LogNiagara, Log, TEXT("Niagara System Sim Tick: %s"), *System->GetName());
-			UE_LOG(LogNiagara, Log, TEXT("=========================================================="));
+			checkSlow(Inst->SystemInstanceIndex == SystemIndex);
+			++SystemIndex;
 		}
-
+		else
 		{
-			SCOPE_CYCLE_COUNTER(STAT_NiagaraSystemSim_PrepareForSimulate);
-			auto PreSimulateAndTransferParams = [&](int32 SystemIndex)
-			{
-				FNiagaraSystemInstance* Inst = SystemInstances[SystemIndex];
-				Inst->PreSimulateTick(DeltaSeconds);
-
-				if (Inst->GetParameters().GetParametersDirty() && bCanExecute)
-				{
-					SpawnInstanceParameterToDataSetBinding.ParameterStoreToDataSet(Inst->GetParameters(), SpawnInstanceParameterDataSet, SystemIndex);
-					UpdateInstanceParameterToDataSetBinding.ParameterStoreToDataSet(Inst->GetParameters(), UpdateInstanceParameterDataSet, SystemIndex);
-				}
+			checkSlow((SystemInstances.Num() <= SystemIndex) || (Inst == SystemInstances[SystemIndex]));
+		}
+	}
 
-				//TODO: Find good way to check that we're not using any instance parameter data interfaces in the system scripts here.
-				//In that case we need to solo and will never get here.
+	//Setup the few real constants like delta time.
+	float InvDt = 1.0f / DeltaSeconds;
 
-				TArray<TSharedRef<FNiagaraEmitterInstance, ESPMode::ThreadSafe>>& Emitters = Inst->GetEmitters();
-				for (int32 EmitterIdx = 0; EmitterIdx < Emitters.Num(); ++EmitterIdx)
-				{
-					FNiagaraEmitterInstance& EmitterInst = Emitters[EmitterIdx].Get();
-					if (EmitterExecutionStateAccessors.Num() > EmitterIdx && EmitterExecutionStateAccessors[EmitterIdx].IsValidForWrite())
-					{
-						EmitterExecutionStateAccessors[EmitterIdx].Set(SystemIndex, (int32)EmitterInst.GetExecutionState());
-					}
-				}
-			};
-			
-			//Begin filling the state of the instance paramter datasets.
-			SpawnInstanceParameterDataSet.BeginSimulate();
-			UpdateInstanceParameterDataSet.BeginSimulate();
+	SpawnTimeParam.SetValue(World->TimeSeconds);
+	UpdateTimeParam.SetValue(World->TimeSeconds);
+	SpawnDeltaTimeParam.SetValue(DeltaSeconds);
+	UpdateDeltaTimeParam.SetValue(DeltaSeconds);
+	SpawnInvDeltaTimeParam.SetValue(InvDt);
+	UpdateInvDeltaTimeParam.SetValue(InvDt);
+	SpawnNumSystemInstancesParam.SetValue(SystemInstances.Num());
+	UpdateNumSystemInstancesParam.SetValue(SystemInstances.Num());
 
-			SpawnInstanceParameterDataSet.Allocate(NewNum);
-			UpdateInstanceParameterDataSet.Allocate(NewNum);
+	FNiagaraSystemSimulationTickContext Context(this, DeltaSeconds, false, MyCompletionGraphEvent);
 
-			for (int32 EmitterIdx = 0; EmitterIdx < System->GetNumEmitters(); ++EmitterIdx)
-			{
-				EmitterExecutionStateAccessors[EmitterIdx].InitForAccess();
-			}
+	//Now kick of the concurrent tick.
+	if (Context.bTickAsync)
+	{
+		FGraphEventRef AsyncTask = TGraphTask<FNiagaraSystemSimulationTickTask>::CreateTask(nullptr, ENamedThreads::GameThread).ConstructAndDispatchWhenReady(Context);
+		MyCompletionGraphEvent->SetGatherThreadForDontCompleteUntil(ENamedThreads::GameThread);
+		MyCompletionGraphEvent->DontCompleteUntil(AsyncTask);
+	}
+	else
+	{
+		Tick_Concurrent(Context);
+	}
 
-			//Transfer any values like execution state from the system instance into the dataset for simulation.
-			ParallelFor(SystemInstances.Num(), [&](int32 SystemIndex)
-			{
-				PreSimulateAndTransferParams(SystemIndex);
-			});
-
-			SpawnInstanceParameterDataSet.GetDestinationDataChecked().SetNumInstances(NewNum);
-			UpdateInstanceParameterDataSet.GetDestinationDataChecked().SetNumInstances(NewNum);
-
-			//We're done filling in the current state for the instance parameter datasets.
-			SpawnInstanceParameterDataSet.EndSimulate();
-			UpdateInstanceParameterDataSet.EndSimulate();
-
-			//Setup the few real constants like delta time.
-			float InvDt = 1.0f / DeltaSeconds;
-			float GlobalSpawnCountScale = INiagaraModule::GetGlobalSpawnCountScale();
-			float GlobalSystemCountScale = INiagaraModule::GetGlobalSystemCountScale();
-
-			SpawnTimeParam.SetValue(World->TimeSeconds);
-			UpdateTimeParam.SetValue(World->TimeSeconds);
-			SpawnDeltaTimeParam.SetValue(DeltaSeconds);
-			UpdateDeltaTimeParam.SetValue(DeltaSeconds);
-			SpawnInvDeltaTimeParam.SetValue(InvDt);
-			UpdateInvDeltaTimeParam.SetValue(InvDt);
-			SpawnNumSystemInstancesParam.SetValue(NewNum);
-			UpdateNumSystemInstancesParam.SetValue(NewNum);
-			SpawnGlobalSpawnCountScaleParam.SetValue(GlobalSpawnCountScale);
-			UpdateGlobalSpawnCountScaleParam.SetValue(GlobalSpawnCountScale);
-			SpawnGlobalSystemCountScaleParam.SetValue(GlobalSystemCountScale);
-			UpdateGlobalSystemCountScaleParam.SetValue(GlobalSystemCountScale);
-		}
-
-		FNiagaraSystemInstance* SoloSystemInstance = bIsSolo && SystemInstances.Num() == 1 ? SystemInstances[0] : nullptr;
-
-		//TODO: JIRA - UE-60096 - Remove.
-		//We're having to allocate and spawn before update here so we have to do needless copies.			
-		//Ideally this should be compiled directly into the script similarly to interpolated particle spawning.
-		if (SpawnNum)
-		{
-			SCOPE_CYCLE_COUNTER(STAT_NiagaraSystemSim_Spawn);
-			DataSet.BeginSimulate();
-			DataSet.Allocate(NewNum, true);
-			DataSet.GetDestinationDataChecked().SetNumInstances(NewNum);
-
-			//Run Spawn
-			SpawnExecContext.Tick(SoloSystemInstance);//We can't require a specific instance here as these are for all instances.
-			SpawnExecContext.BindData(0, DataSet, OrigNum, false);
-			SpawnExecContext.BindData(1, SpawnInstanceParameterDataSet, OrigNum, false);
-			SpawnExecContext.Execute(SpawnNum);
-
-			if (GbDumpSystemData || System->bDumpDebugSystemInfo)
-			{
-				UE_LOG(LogNiagara, Log, TEXT("=== Spwaned %d Systems ==="), SpawnNum);
-				DataSet.GetDestinationDataChecked().Dump(OrigNum, SpawnNum, TEXT("System Dataset - Post Spawn"));
-				SpawnInstanceParameterDataSet.GetCurrentDataChecked().Dump(OrigNum, SpawnNum, TEXT("Spawn Instance Parameter Data"));
-			}
+	//TEMPORARY DEV CODE:
+	//Test that the wait system is working
+	if (GbForceNiagaraAsyncWait)
+	{
+		WaitForTickComplete();
 
-			DataSet.EndSimulate();
+		{
+			//We spin a bit here to help make sure that all the Niagara work did infact complete before we get here. Makes things easier to see in Razor.
+			SCOPE_CYCLE_COUNTER(STAT_NiagaraSystemSim_ForceWaitFakeStall);
 
-#if WITH_EDITORONLY_DATA
-			if (SoloSystemInstance && SoloSystemInstance->ShouldCaptureThisFrame())
+			float CurrTimeMS = FPlatformTime::Seconds() * 1000;
+			float WaitEnd = CurrTimeMS + GbWaitForAsyncStallMS;
+			do
 			{
-				TSharedPtr<struct FNiagaraScriptDebuggerInfo, ESPMode::ThreadSafe> DebugInfo = SoloSystemInstance->GetActiveCaptureWrite(NAME_None, ENiagaraScriptUsage::SystemSpawnScript, FGuid());
-				if (DebugInfo)
-				{
-					DataSet.CopyTo(DebugInfo->Frame, OrigNum, SpawnNum);
-					DebugInfo->Parameters = SpawnExecContext.Parameters;
-					DebugInfo->bWritten = true;
-				}
-			}
-#endif
+				CurrTimeMS = FPlatformTime::Seconds() * 1000;
+			} 
+			while (CurrTimeMS < WaitEnd);
 		}
+	}
+}
 
-
+void FNiagaraSystemSimulation::WaitForTickComplete()
+{
+	SCOPE_CYCLE_COUNTER(STAT_NiagaraSystemSim_ForceWaitForAsync);
+	check(IsInGameThread());
+	int32 SystemInstIndex = 0;
+	while (SystemInstIndex < SystemInstances.Num())
+	{
+		FNiagaraSystemInstance* Inst = SystemInstances[SystemInstIndex];
+		Inst->WaitForAsyncTick(bInSpawnPhase);//if we're in a spawn phase all existing instances should be complete already.
+		if (!Inst->IsComplete())
 		{
-			SCOPE_CYCLE_COUNTER(STAT_NiagaraSystemSim_Update);
-
-			FNiagaraDataBuffer& DestinationData = DataSet.BeginSimulate();
-			DestinationData.Allocate(NewNum);
-			DestinationData.SetNumInstances(NewNum);
+			//If the system completes during finalize it will be removed from instances so we don't update the index.
+			++SystemInstIndex;
+		}
 
-			//Run update.
-			UpdateExecContext.Tick(SystemInstances[0]);
-			UpdateExecContext.BindData(0, DataSet, 0, false);
-			UpdateExecContext.BindData(1, UpdateInstanceParameterDataSet, 0, false);
-			UpdateExecContext.Execute(OrigNum);
+		check(DataSet.GetCurrentDataChecked().GetNumInstances() == SystemInstances.Num());
+	}
 
-			if (GbDumpSystemData || System->bDumpDebugSystemInfo)
+	if (bInSpawnPhase)
+	{
+		SystemInstIndex = 0;
+		while (SystemInstIndex < SpawningSystemInstances.Num())
+		{
+			FNiagaraSystemInstance* Inst = SpawningSystemInstances[SystemInstIndex];
+			Inst->WaitForAsyncTick();
+			if (!Inst->IsComplete())
 			{
-				UE_LOG(LogNiagara, Log, TEXT("=== Updated %d Systems ==="), OrigNum); 
-				DestinationData.Dump(0, OrigNum, TEXT("System Data - Post Update"));
-				UpdateInstanceParameterDataSet.GetCurrentDataChecked().Dump(0, OrigNum, TEXT("Update Instance Paramter Data"));
+				//If the system completes during finalize it will be removed from instances so we don't update the index.
+				++SystemInstIndex;
 			}
 
- 			//Also run the update script on the newly spawned systems too.
-			//TODO: JIRA - UE-60096 - Remove.
-			//Ideally this should be compiled directly into the script similarly to interpolated particle spawning.
-			if (SpawnNum && GbSystemUpdateOnSpawn)
-			{
-				//Run update.
-				UpdateExecContext.Tick(SystemInstances[0]);
+			check(PendingSpawnDataSet.GetCurrentDataChecked().GetNumInstances() == SpawningSystemInstances.Num());
+		}
+	}
+}
 
-				UpdateExecContext.Parameters.SetParameterValue(0.0001f, SYS_PARAM_ENGINE_DELTA_TIME);
-				UpdateExecContext.Parameters.SetParameterValue(10000.0f, SYS_PARAM_ENGINE_INV_DELTA_TIME);
+void FNiagaraSystemSimulation::Tick_Concurrent(FNiagaraSystemSimulationTickContext& Context)
+{
+	SCOPE_CYCLE_COUNTER(STAT_NiagaraSystemSim_TickCNC);
+	SCOPE_CYCLE_COUNTER(STAT_NiagaraOverview_GT_CNC);
+	CSV_SCOPED_TIMING_STAT_EXCLUSIVE(Niagara);
+	LLM_SCOPE(ELLMTag::Niagara);
 
-				UpdateExecContext.BindData(0, DataSet, OrigNum, false);
-				UpdateExecContext.BindData(1, UpdateInstanceParameterDataSet, OrigNum, false);
-				UpdateExecContext.Execute(SpawnNum);
+	TArray<FNiagaraSystemInstance*>& Instances = Context.GetInstances();
+	FNiagaraDataSet& SimulationDataSet = Context.GetDataSet();
+	FNiagaraSystemInstance* SoloSystemInstance = bIsSolo && Instances.Num() == 1 ? Instances[0] : nullptr;
 
-				if (GbDumpSystemData || System->bDumpDebugSystemInfo)
-				{
-					UE_LOG(LogNiagara, Log, TEXT("=== Spawn Updated %d Systems ==="), SpawnNum);
-					DestinationData.Dump(OrigNum, SpawnNum, TEXT("System Data - Post Update (new systems)"));
-					UpdateInstanceParameterDataSet.GetCurrentDataChecked().Dump(OrigNum, SpawnNum, TEXT("Update Instance Paramter Data (new systems)"));
-				}
-			}
-			
-			DataSet.EndSimulate();
+	if (bCanExecute && Instances.Num())
+	{
+		if (GbDumpSystemData || Context.System->bDumpDebugSystemInfo)
+		{
+			UE_LOG(LogNiagara, Log, TEXT("=========================================================="));
+			UE_LOG(LogNiagara, Log, TEXT("Niagara System Sim Tick_Concurrent(): %s"), *Context.System->GetName());
+			UE_LOG(LogNiagara, Log, TEXT("=========================================================="));
+		}
 
+		PrepareForSystemSimulate(Context);
 
-#if WITH_EDITORONLY_DATA
-			if (SoloSystemInstance && SoloSystemInstance->ShouldCaptureThisFrame())
+		if (Context.bPendingSpawnPass)
+		{
+			SpawnSystemInstances(Context);
+		}
+
+		UpdateSystemInstances(Context);
+
+		TransferSystemSimResults(Context);
+
+		for (FNiagaraSystemInstance* Instance : Instances)
+		{
+			AddSystemToTickBatch(Instance, Context);
+		}
+		FlushTickBatch(Context);
+
+		//If both the instances and the main sim are run on the GT then we need to finalize here.
+		if (!Context.bTickAsync && !Context.bTickInstancesAsync)
+		{
+			check(IsInGameThread());
+			int32 SystemInstIndex = 0;
+			while (SystemInstIndex < Instances.Num())
 			{
-				TSharedPtr<struct FNiagaraScriptDebuggerInfo, ESPMode::ThreadSafe> DebugInfo = SoloSystemInstance->GetActiveCaptureWrite(NAME_None, ENiagaraScriptUsage::SystemUpdateScript, FGuid());
-				if (DebugInfo)
+				FNiagaraSystemInstance* Inst = Instances[SystemInstIndex];
+				checkSlow(!Inst->IsComplete());
+				Inst->FinalizeTick_GameThread();
+				if (!Inst->IsComplete())
 				{
-					DataSet.CopyTo(DebugInfo->Frame);
-					DebugInfo->Parameters = UpdateExecContext.Parameters;
-					DebugInfo->bWritten = true;
+					//If the system completes during finalize it will be removed from instances so we don't update the index.
+					++SystemInstIndex;
 				}
+
+				check(SimulationDataSet.GetCurrentDataChecked().GetNumInstances() == Instances.Num());
 			}
-#endif
 		}
 
-		SystemExecutionStateAccessor.InitForAccess();
-		for (int32 EmitterIdx = 0; EmitterIdx < System->GetNumEmitters(); ++EmitterIdx)
+	#if WITH_EDITORONLY_DATA
+		if (SoloSystemInstance)
 		{
-			EmitterExecutionStateAccessors[EmitterIdx].InitForAccess();
-			for (int32 SpawnInfoIdx = 0; SpawnInfoIdx < EmitterSpawnInfoAccessors[EmitterIdx].Num(); ++SpawnInfoIdx)
-			{
-				EmitterSpawnInfoAccessors[EmitterIdx][SpawnInfoIdx].InitForAccess();
-			}
+			SoloSystemInstance->FinishCapture();
 		}
+	#endif
 
+		if (!Context.bPendingSpawnPass)
 		{
-			SCOPE_CYCLE_COUNTER(STAT_NiagaraSystemSim_TransferParameters);
-			for (int32 SystemIndex = 0; SystemIndex < SystemInstances.Num(); ++SystemIndex)
-			{
-				ENiagaraExecutionState ExecutionState = (ENiagaraExecutionState)SystemExecutionStateAccessor.GetSafe(SystemIndex, (int32)ENiagaraExecutionState::Disabled);
-				FNiagaraSystemInstance* SystemInst = SystemInstances[SystemIndex];
+			INC_DWORD_STAT_BY(STAT_NiagaraNumSystems, Instances.Num());
+		}
+	}
+}
 
-				//Apply the systems requested execution state to it's actual execution state.
-				SystemInst->SetActualExecutionState(ExecutionState);
+void FNiagaraSystemSimulation::PrepareForSystemSimulate(FNiagaraSystemSimulationTickContext& Context)
+{
+	SCOPE_CYCLE_COUNTER(STAT_NiagaraSystemSim_PrepareForSimulateCNC);
 
-				if (!SystemInst->IsDisabled() && !SystemInst->HandleCompletion())
-				{
-					//Now pull data out of the simulation and drive the emitters with it.
-					TArray<TSharedRef<FNiagaraEmitterInstance, ESPMode::ThreadSafe>>& Emitters = SystemInst->GetEmitters();
-					for (int32 EmitterIdx = 0; EmitterIdx < Emitters.Num(); ++EmitterIdx)
-					{
-						FNiagaraEmitterInstance& EmitterInst = Emitters[EmitterIdx].Get();
-
-						//Early exit before we set the state as if we're complete or disabled we should never let the emitter turn itself back. It needs to be reset/reinited manually.
-						if (EmitterInst.IsComplete())
-						{
-							continue;
-						}
-
-						check(Emitters.Num() > EmitterIdx);
-						ENiagaraExecutionState State = (ENiagaraExecutionState)EmitterExecutionStateAccessors[EmitterIdx].GetSafe(SystemIndex, (int32)ENiagaraExecutionState::Disabled);
-						EmitterInst.SetExecutionState(State);
-												
-						TArray<FNiagaraSpawnInfo>& EmitterInstSpawnInfos = EmitterInst.GetSpawnInfo();
-						for (int32 SpawnInfoIdx=0; SpawnInfoIdx < EmitterSpawnInfoAccessors[EmitterIdx].Num(); ++SpawnInfoIdx)
-						{
-							if (SpawnInfoIdx < EmitterInstSpawnInfos.Num())
-							{
-								EmitterInstSpawnInfos[SpawnInfoIdx] = EmitterSpawnInfoAccessors[EmitterIdx][SpawnInfoIdx].Get(SystemIndex);
-							}
-							else
-							{
-								ensure(SpawnInfoIdx < EmitterInstSpawnInfos.Num());
-							}
-						}
-
-
-						//TODO: Any other fixed function stuff like this?
-
-						FNiagaraScriptExecutionContext& SpawnContext = EmitterInst.GetSpawnExecutionContext();
-						DataSetToEmitterSpawnParameters[EmitterIdx].DataSetToParameterStore(SpawnContext.Parameters, DataSet, SystemIndex);
-
-						FNiagaraScriptExecutionContext& UpdateContext = EmitterInst.GetUpdateExecutionContext();
-						DataSetToEmitterUpdateParameters[EmitterIdx].DataSetToParameterStore(UpdateContext.Parameters, DataSet, SystemIndex);
-
-						TArray<FNiagaraScriptExecutionContext>& EventContexts = EmitterInst.GetEventExecutionContexts();
-						for (int32 EventIdx = 0; EventIdx < EventContexts.Num(); ++EventIdx)
-						{
-							FNiagaraScriptExecutionContext& EventContext = EventContexts[EventIdx];
-							if (DataSetToEmitterEventParameters[EmitterIdx].Num() > EventIdx)
-							{
-								DataSetToEmitterEventParameters[EmitterIdx][EventIdx].DataSetToParameterStore(EventContext.Parameters, DataSet, SystemIndex);
-							}
-							else
-							{
-								UE_LOG(LogNiagara, Log, TEXT("Skipping DataSetToEmitterEventParameters because EventIdx is out-of-bounds. %d of %d"), EventIdx, DataSetToEmitterEventParameters[EmitterIdx].Num());
-							}
-						}
-					}
-				}
-			}
-		}
+	TArray<FNiagaraSystemInstance*>& Instances = Context.GetInstances();
+	int32 NumInstances = Instances.Num();
+
+	if (NumInstances == 0)
+	{
+		return;
+	}
+
+	//Begin filling the state of the instance parameter datasets.
+	SpawnInstanceParameterDataSet.BeginSimulate();
+	UpdateInstanceParameterDataSet.BeginSimulate();
+
+	SpawnInstanceParameterDataSet.Allocate(NumInstances);
+	UpdateInstanceParameterDataSet.Allocate(NumInstances);
+
+	for (int32 EmitterIdx = 0; EmitterIdx < Context.System->GetNumEmitters(); ++EmitterIdx)
+	{
+		EmitterExecutionStateAccessors[EmitterIdx].InitForAccess();
 	}
 
+	//Tick instance parameters and transfer any needed into the system simulation dataset.
+	auto TransferInstanceParameters = [&](int32 SystemIndex)
 	{
-		SCOPE_CYCLE_COUNTER(STAT_NiagaraSystemSim_PostSimulate);
+		FNiagaraSystemInstance* Inst = Instances[SystemIndex];
+
+		//Inst->TickInstanceParameters(Context.DeltaSeconds);
 
-		if (GbParallelSystemPostTick && FApp::ShouldUseThreadingForPerformance())
+		if (Inst->GetParameters().GetParametersDirty() && bCanExecute)
 		{
-			ParallelFor(SystemInstances.Num(),
-				[&](int32 SystemIndex)
-			{
-				FNiagaraSystemInstance* SystemInst = SystemInstances[SystemIndex];
-				SystemInst->PostSimulateTick(DeltaSeconds);
-			});
+			SpawnInstanceParameterToDataSetBinding.ParameterStoreToDataSet(Inst->GetParameters(), SpawnInstanceParameterDataSet, SystemIndex);
+			UpdateInstanceParameterToDataSetBinding.ParameterStoreToDataSet(Inst->GetParameters(), UpdateInstanceParameterDataSet, SystemIndex);
 		}
-		else
+
+		//TODO: Find good way to check that we're not using any instance parameter data interfaces in the system scripts here.
+		//In that case we need to solo and will never get here.
+
+		TArray<TSharedRef<FNiagaraEmitterInstance, ESPMode::ThreadSafe>>& Emitters = Inst->GetEmitters();
+		for (int32 EmitterIdx = 0; EmitterIdx < Emitters.Num(); ++EmitterIdx)
 		{
-			//Now actually tick emitters.
-			for (int32 SystemIndex = 0; SystemIndex < SystemInstances.Num(); ++SystemIndex)
+			FNiagaraEmitterInstance& EmitterInst = Emitters[EmitterIdx].Get();
+			if (EmitterExecutionStateAccessors.Num() > EmitterIdx && EmitterExecutionStateAccessors[EmitterIdx].IsValidForWrite())
 			{
-				FNiagaraSystemInstance* SystemInst = SystemInstances[SystemIndex];
-				SystemInst->PostSimulateTick(DeltaSeconds);
+				EmitterExecutionStateAccessors[EmitterIdx].Set(SystemIndex, (int32)EmitterInst.GetExecutionState());
 			}
 		}
+	};
+
+	//This can go wide if we have a very large number of instances.
+	//ParallelFor(Instances.Num(), TransferInstanceParameters, Instances.Num() < GbSystemSimTransferParamsParallelThreshold);
+	ParallelFor(Instances.Num(), TransferInstanceParameters, true);
+
+	SpawnInstanceParameterDataSet.GetDestinationDataChecked().SetNumInstances(NumInstances);
+	UpdateInstanceParameterDataSet.GetDestinationDataChecked().SetNumInstances(NumInstances);
+
+	//We're done filling in the current state for the instance parameter datasets.
+	SpawnInstanceParameterDataSet.EndSimulate();
+	UpdateInstanceParameterDataSet.EndSimulate();
+}
+
+void FNiagaraSystemSimulation::SpawnSystemInstances(FNiagaraSystemSimulationTickContext& Context)
+{
+	//All instance spawning is done in a separate pass at the end of the frame so we can be sure we have all new spawns ready for processing.
+	//We run the spawn and update scripts separately here as their own sim passes.
+
+	SCOPE_CYCLE_COUNTER(STAT_NiagaraSystemSim_SpawnCNC);
+
+	TArray<FNiagaraSystemInstance*>& Instances = Context.GetInstances();
+	FNiagaraDataSet& SimulationDataSet = Context.GetDataSet();
+	int32 NumInstances = Instances.Num();
+	
+	check(Context.bPendingSpawnPass);
+
+	if (NumInstances > 0)
+	{
+
+		FNiagaraSystemInstance* SoloSystemInstance = bIsSolo && Instances.Num() == 1 ? Instances[0] : nullptr;
+		SimulationDataSet.BeginSimulate();
+		SimulationDataSet.Allocate(NumInstances);
+		SimulationDataSet.GetDestinationDataChecked().SetNumInstances(NumInstances);
+
+		//Run Spawn
+		SpawnExecContext.Tick(SoloSystemInstance);//We can't require a specific instance here as these are for all instances.
+		SpawnExecContext.BindData(0, SimulationDataSet, 0, false);
+		SpawnExecContext.BindData(1, SpawnInstanceParameterDataSet, 0, false);
+		SpawnExecContext.Execute(NumInstances);
+
+		if (GbDumpSystemData || Context.System->bDumpDebugSystemInfo)
+		{
+			UE_LOG(LogNiagara, Log, TEXT("=== Spwaned %d Systems ==="), NumInstances);
+			SimulationDataSet.GetDestinationDataChecked().Dump(0, NumInstances, TEXT("System Dataset - Post Spawn"));
+			SpawnInstanceParameterDataSet.GetCurrentDataChecked().Dump(0, NumInstances, TEXT("Spawn Instance Parameter Data"));
+		}
+
+		SimulationDataSet.EndSimulate();
 	}
 
+	check(SimulationDataSet.GetCurrentDataChecked().GetNumInstances() == Instances.Num());
+}
+
+void FNiagaraSystemSimulation::UpdateSystemInstances(FNiagaraSystemSimulationTickContext& Context)
+{
+	SCOPE_CYCLE_COUNTER(STAT_NiagaraSystemSim_UpdateCNC);
+
+	TArray<FNiagaraSystemInstance*>& Instances = Context.GetInstances();
+	FNiagaraDataSet& SimulationDataSet = Context.GetDataSet();
+	int32 NumInstances = Instances.Num();
+
+	if (NumInstances > 0)
 	{
-		SCOPE_CYCLE_COUNTER(STAT_NiagaraSystemSim_MarkComponentDirty);
-		//This is not a small amount of the update time. 
-		//Annoyingly these can't be done in parallel it seems.
-		//TODO: Find some way to parallelize this. Especially UpdateComponentToWorld();
-		int32 SystemIndex = 0;
-		while (SystemIndex < SystemInstances.Num())
+		if (Context.bPendingSpawnPass)
 		{
-			FNiagaraSystemInstance* SystemInstance = SystemInstances[SystemIndex];
-			++SystemIndex;
-			if (SystemIndex < SystemInstances.Num())
+			UpdateDeltaTimeParam.SetValue(0.0001f);
+			UpdateInvDeltaTimeParam.SetValue(10000.0f);
+		}
+
+
+		FNiagaraSystemInstance* SoloSystemInstance = bIsSolo && Instances.Num() == 1 ? Instances[0] : nullptr;
+
+		FNiagaraDataBuffer& DestinationData = SimulationDataSet.BeginSimulate();
+		DestinationData.Allocate(NumInstances);
+		DestinationData.SetNumInstances(NumInstances);
+
+		//Run update.
+		UpdateExecContext.Tick(Instances[0]);
+		UpdateExecContext.BindData(0, SimulationDataSet, 0, false);
+		UpdateExecContext.BindData(1, UpdateInstanceParameterDataSet, 0, false);
+		UpdateExecContext.Execute(NumInstances);
+
+		if (GbDumpSystemData || Context.System->bDumpDebugSystemInfo)
+		{
+			UE_LOG(LogNiagara, Log, TEXT("=== Updated %d Systems ==="), NumInstances);
+			DestinationData.Dump(0, NumInstances, TEXT("System Data - Post Update"));
+			UpdateInstanceParameterDataSet.GetCurrentDataChecked().Dump(0, NumInstances, TEXT("Update Instance Paramter Data"));
+		}
+
+		SimulationDataSet.EndSimulate();
+
+#if WITH_EDITORONLY_DATA
+		if (SoloSystemInstance && SoloSystemInstance->ShouldCaptureThisFrame())
+		{
+			TSharedPtr<struct FNiagaraScriptDebuggerInfo, ESPMode::ThreadSafe> DebugInfo;
+
+			if (Context.bPendingSpawnPass)
 			{
-				FPlatformMisc::Prefetch(SystemInstances[SystemIndex]->GetComponent());
+				//We say this update phase is the out put of the SpawnScript but it's really just the initial post spwan state of the particles. We may want to alter the UI to reflect this?
+				SoloSystemInstance->GetActiveCaptureWrite(NAME_None, ENiagaraScriptUsage::SystemSpawnScript, FGuid());
 			}
-			SystemInstance->FinalizeTick(DeltaSeconds);
-			// @todo If we buffer up the Tick object we create here we can push all the data for every system in this simulation at once.
-			// @todo It may also be worth considering hoisting the push out to FNiagaraWorldManager::Tick to push every desired simulation at once.
-			if (SystemInstance->ActiveGPUEmitterCount > 0 && NiagaraSupportsComputeShaders(SystemInstance->GetBatcher()->GetShaderPlatform()))
+			else
 			{
-				ensure(!SystemInstance->IsComplete());
-				FNiagaraGPUSystemTick GPUTick;
-				GPUTick.Init(SystemInstance);
-
-				//if (GPUTick.DIInstanceData)
-				//{
-				//	uint8* BasePointer = (uint8*)GPUTick.DIInstanceData->PerInstanceDataForRT;
-
-				//	//UE_LOG(LogNiagara, Log, TEXT("GT Testing (dipacket) %p (baseptr) %p"), GPUTick.DIInstanceData, BasePointer);
-				//	for (auto& Pair : GPUTick.DIInstanceData->InterfaceProxiesToOffsets)
-				//	{
-				//		FNiagaraDataInterfaceProxy* Proxy = Pair.Key;
-				//		UE_LOG(LogNiagara, Log, TEXT("\tGT (proxy) %p (size) %u"), Proxy, Proxy->PerInstanceDataPassedToRenderThreadSize());
-				//	}
-				//}
-
-				// We will give the data over to the render thread. It is responsible for freeing it.
-				// We no longer own it and cannot modify it after this point.
-				// @todo We are taking a copy of the object here. This object is small so this overhead should
-				// not be very high. And we avoid making a bunch of small allocations here.
-				NiagaraEmitterInstanceBatcher* Batcher = SystemInstance->GetBatcher();
-				ENQUEUE_RENDER_COMMAND(FGiveSystemInstanceTickToRT)(
-					[Batcher, GPUTick](FRHICommandListImmediate& RHICmdList) mutable
-				{
-					Batcher->GiveSystemTick_RenderThread(GPUTick);
-				}
-				);
+				SoloSystemInstance->GetActiveCaptureWrite(NAME_None, ENiagaraScriptUsage::SystemUpdateScript, FGuid());
+			}
+
+			if (DebugInfo)
+			{
+				SimulationDataSet.CopyTo(DebugInfo->Frame);
+				DebugInfo->Parameters = UpdateExecContext.Parameters;
+				DebugInfo->bWritten = true;
 			}
 		}
+#endif
 	}
 
-#if WITH_EDITORONLY_DATA
-	if (bIsSolo && SystemInstances.Num() == 1)
+	check(SimulationDataSet.GetCurrentDataChecked().GetNumInstances() == Instances.Num());
+}
+
+void FNiagaraSystemSimulation::TransferSystemSimResults(FNiagaraSystemSimulationTickContext& Context)
+{
+	SCOPE_CYCLE_COUNTER(STAT_NiagaraSystemSim_TransferResultsCNC);
+
+	TArray<FNiagaraSystemInstance*>& Instances = Context.GetInstances();
+	if (Instances.Num() == 0)
 	{
-		SystemInstances[0]->FinishCapture();
+		return;
 	}
-#endif
 
-	INC_DWORD_STAT_BY(STAT_NiagaraNumSystems, SystemInstances.Num());
+	FNiagaraDataSet& SimulationDataSet = Context.GetDataSet();
+	SystemExecutionStateAccessor.SetDataSet(SimulationDataSet);
+	SystemExecutionStateAccessor.InitForAccess();
+	for (int32 EmitterIdx = 0; EmitterIdx < Context.System->GetNumEmitters(); ++EmitterIdx)
+	{
+		EmitterExecutionStateAccessors[EmitterIdx].SetDataSet(SimulationDataSet);
+		EmitterExecutionStateAccessors[EmitterIdx].InitForAccess();
+		for (int32 SpawnInfoIdx = 0; SpawnInfoIdx < EmitterSpawnInfoAccessors[EmitterIdx].Num(); ++SpawnInfoIdx)
+		{
+			EmitterSpawnInfoAccessors[EmitterIdx][SpawnInfoIdx].SetDataSet(SimulationDataSet);
+			EmitterSpawnInfoAccessors[EmitterIdx][SpawnInfoIdx].InitForAccess();
+		}
+	}
 
-	return true;
+	for (int32 SystemIndex = 0; SystemIndex < Instances.Num(); ++SystemIndex)
+	{
+		ENiagaraExecutionState ExecutionState = (ENiagaraExecutionState)SystemExecutionStateAccessor.GetSafe(SystemIndex, (int32)ENiagaraExecutionState::Disabled);
+		FNiagaraSystemInstance* SystemInst = Instances[SystemIndex];
+
+		//Apply the systems requested execution state to it's actual execution state.
+		SystemInst->SetActualExecutionState(ExecutionState);
+
+		if (!SystemInst->IsDisabled())
+		{
+			//Now pull data out of the simulation and drive the emitters with it.
+			TArray<TSharedRef<FNiagaraEmitterInstance, ESPMode::ThreadSafe>>& Emitters = SystemInst->GetEmitters();
+			for (int32 EmitterIdx = 0; EmitterIdx < Emitters.Num(); ++EmitterIdx)
+			{
+				FNiagaraEmitterInstance& EmitterInst = Emitters[EmitterIdx].Get();
+
+				//Early exit before we set the state as if we're complete or disabled we should never let the emitter turn itself back. It needs to be reset/reinited manually.
+				if (EmitterInst.IsComplete())
+				{
+					continue;
+				}
+
+				check(Emitters.Num() > EmitterIdx);
+				ENiagaraExecutionState State = (ENiagaraExecutionState)EmitterExecutionStateAccessors[EmitterIdx].GetSafe(SystemIndex, (int32)ENiagaraExecutionState::Disabled);
+				EmitterInst.SetExecutionState(State);
+
+				TArray<FNiagaraSpawnInfo>& EmitterInstSpawnInfos = EmitterInst.GetSpawnInfo();
+				for (int32 SpawnInfoIdx = 0; SpawnInfoIdx < EmitterSpawnInfoAccessors[EmitterIdx].Num(); ++SpawnInfoIdx)
+				{
+					if (SpawnInfoIdx < EmitterInstSpawnInfos.Num())
+					{
+						EmitterInstSpawnInfos[SpawnInfoIdx] = EmitterSpawnInfoAccessors[EmitterIdx][SpawnInfoIdx].Get(SystemIndex);
+					}
+					else
+					{
+						ensure(SpawnInfoIdx < EmitterInstSpawnInfos.Num());
+					}
+				}
+
+				//TODO: Any other fixed function stuff like this?
+
+				FNiagaraScriptExecutionContext& SpawnContext = EmitterInst.GetSpawnExecutionContext();
+				DataSetToEmitterSpawnParameters[EmitterIdx].DataSetToParameterStore(SpawnContext.Parameters, SimulationDataSet, SystemIndex);
+
+				FNiagaraScriptExecutionContext& UpdateContext = EmitterInst.GetUpdateExecutionContext();
+				DataSetToEmitterUpdateParameters[EmitterIdx].DataSetToParameterStore(UpdateContext.Parameters, SimulationDataSet, SystemIndex);
+
+				TArray<FNiagaraScriptExecutionContext>& EventContexts = EmitterInst.GetEventExecutionContexts();
+				for (int32 EventIdx = 0; EventIdx < EventContexts.Num(); ++EventIdx)
+				{
+					FNiagaraScriptExecutionContext& EventContext = EventContexts[EventIdx];
+					if (DataSetToEmitterEventParameters[EmitterIdx].Num() > EventIdx)
+					{
+						DataSetToEmitterEventParameters[EmitterIdx][EventIdx].DataSetToParameterStore(EventContext.Parameters, SimulationDataSet, SystemIndex);
+					}
+					else
+					{
+						UE_LOG(LogNiagara, Log, TEXT("Skipping DataSetToEmitterEventParameters because EventIdx is out-of-bounds. %d of %d"), EventIdx, DataSetToEmitterEventParameters[EmitterIdx].Num());
+					}
+				}
+			}
+		}
+	}
 }
 
 void FNiagaraSystemSimulation::RemoveInstance(FNiagaraSystemInstance* Instance)
@@ -739,6 +1095,10 @@ void FNiagaraSystemSimulation::RemoveInstance(FNiagaraSystemInstance* Instance)
 		return;
 	}
 
+	check(SystemInstances.Num() == DataSet.GetCurrentDataChecked().GetNumInstances());
+	check(PausedSystemInstances.Num() == PausedInstanceData.GetCurrentDataChecked().GetNumInstances());
+
+	check(IsInGameThread());
 	check(DataSet.GetDestinationData() == nullptr);
 	UNiagaraSystem* System = WeakSystem.Get();
 	if (Instance->IsPendingSpawn())
@@ -748,14 +1108,26 @@ void FNiagaraSystemSimulation::RemoveInstance(FNiagaraSystemInstance* Instance)
 			UE_LOG(LogNiagara, Log, TEXT("=== Removing Pending Spawn %d ==="), Instance->SystemInstanceIndex);
 			DataSet.GetCurrentDataChecked().Dump(Instance->SystemInstanceIndex, 1, TEXT("System data being removed."));
 		}
+
+		//If we're currently spawning then all our pending spawn instances will be inside
+		TArray<FNiagaraSystemInstance*>& Instances = bInSpawnPhase ? SpawningSystemInstances : PendingSystemInstances;
+
 		int32 SystemIndex = Instance->SystemInstanceIndex;
-		check(Instance == PendingSystemInstances[SystemIndex]);
-		PendingSystemInstances.RemoveAtSwap(SystemIndex);
+		check(Instances.IsValidIndex(SystemIndex));
+		check(Instance == Instances[SystemIndex]);
+		
+		if (bInSpawnPhase)
+		{
+			//We should have already spawned this into our dataset so need to kill it.
+			PendingSpawnDataSet.GetCurrentDataChecked().KillInstance(SystemIndex);
+		}
+
+		Instances.RemoveAtSwap(SystemIndex);
 		Instance->SystemInstanceIndex = INDEX_NONE;
 		Instance->SetPendingSpawn(false);
-		if (PendingSystemInstances.IsValidIndex(SystemIndex))
+		if (Instances.IsValidIndex(SystemIndex))
 		{
-			PendingSystemInstances[SystemIndex]->SystemInstanceIndex = SystemIndex;
+			Instances[SystemIndex]->SystemInstanceIndex = SystemIndex;
 		}
 	}
 	else if (Instance->IsPaused())
@@ -766,14 +1138,23 @@ void FNiagaraSystemSimulation::RemoveInstance(FNiagaraSystemInstance* Instance)
 			DataSet.GetCurrentDataChecked().Dump(Instance->SystemInstanceIndex, 1, TEXT("System data being removed."));
 		}
 
+		int32 NumInstances = PausedInstanceData.GetCurrentDataChecked().GetNumInstances();
+		check(PausedSystemInstances.Num() == NumInstances);
+
 		int32 SystemIndex = Instance->SystemInstanceIndex;
+		check(PausedSystemInstances.IsValidIndex(SystemIndex));
 		check(Instance == PausedSystemInstances[SystemIndex]);
+
+		PausedInstanceData.GetCurrentDataChecked().KillInstance(SystemIndex);
 		PausedSystemInstances.RemoveAtSwap(SystemIndex);
 		Instance->SystemInstanceIndex = INDEX_NONE;
 		if (PausedSystemInstances.IsValidIndex(SystemIndex))
 		{
 			PausedSystemInstances[SystemIndex]->SystemInstanceIndex = SystemIndex;
 		}
+
+		check(SystemInstances.Num() == DataSet.GetCurrentDataChecked().GetNumInstances());
+		check(PausedSystemInstances.Num() == PausedInstanceData.GetCurrentDataChecked().GetNumInstances());
 	}
 	else if (SystemInstances.IsValidIndex(Instance->SystemInstanceIndex))
 	{
@@ -789,14 +1170,17 @@ void FNiagaraSystemSimulation::RemoveInstance(FNiagaraSystemInstance* Instance)
 		int32 SystemIndex = Instance->SystemInstanceIndex;
 		check(Instance == SystemInstances[SystemIndex]);
 		check(SystemInstances.IsValidIndex(SystemIndex));
+
 		DataSet.GetCurrentDataChecked().KillInstance(SystemIndex);
 		SystemInstances.RemoveAtSwap(SystemIndex);
 		Instance->SystemInstanceIndex = INDEX_NONE;
-
 		if (SystemInstances.IsValidIndex(SystemIndex))
 		{
 			SystemInstances[SystemIndex]->SystemInstanceIndex = SystemIndex;
 		}
+
+		check(SystemInstances.Num() == DataSet.GetCurrentDataChecked().GetNumInstances());
+		check(PausedSystemInstances.Num() == PausedInstanceData.GetCurrentDataChecked().GetNumInstances());
 	}
 
 #if NIAGARA_NAN_CHECKING
@@ -806,6 +1190,7 @@ void FNiagaraSystemSimulation::RemoveInstance(FNiagaraSystemInstance* Instance)
 
 void FNiagaraSystemSimulation::AddInstance(FNiagaraSystemInstance* Instance)
 {
+	check(IsInGameThread());
 	check(Instance->SystemInstanceIndex == INDEX_NONE);
 	Instance->SetPendingSpawn(true);
 	Instance->SystemInstanceIndex = PendingSystemInstances.Add(Instance);
@@ -816,14 +1201,21 @@ void FNiagaraSystemSimulation::AddInstance(FNiagaraSystemInstance* Instance)
 		UE_LOG(LogNiagara, Log, TEXT("=== Adding To Pending Spawn %d ==="), Instance->SystemInstanceIndex);
 		//DataSet.Dump(true, Instance->SystemInstanceIndex, 1);
 	}
+
+	check(SystemInstances.Num() == DataSet.GetCurrentDataChecked().GetNumInstances());
+	check(PausedSystemInstances.Num() == PausedInstanceData.GetCurrentDataChecked().GetNumInstances());
 }
 
 void FNiagaraSystemSimulation::PauseInstance(FNiagaraSystemInstance* Instance)
 {
+	check(IsInGameThread());
 	check(!Instance->IsPaused());
 	check(!DataSet.GetDestinationData());
 	check(!PausedInstanceData.GetDestinationData());
 
+	check(SystemInstances.Num() == DataSet.GetCurrentDataChecked().GetNumInstances());
+	check(PausedSystemInstances.Num() == PausedInstanceData.GetCurrentDataChecked().GetNumInstances());
+
 	UNiagaraSystem* System = WeakSystem.Get();
 	if (Instance->IsPendingSpawn())
 	{
@@ -847,26 +1239,32 @@ void FNiagaraSystemSimulation::PauseInstance(FNiagaraSystemInstance* Instance)
 	check(SystemInstances.IsValidIndex(SystemIndex));
 	check(Instance == SystemInstances[SystemIndex]);
 
-	int32 NewSystemIndex = PausedInstanceData.GetCurrentDataChecked().TransferInstance(DataSet.GetCurrentDataChecked(), SystemIndex);
-	DataSet.GetCurrentDataChecked().KillInstance(SystemIndex);
+	int32 NewDataSetIndex = PausedInstanceData.GetCurrentDataChecked().TransferInstance(DataSet.GetCurrentDataChecked(), SystemIndex);
+
+	Instance->SystemInstanceIndex = PausedSystemInstances.Add(Instance);
 
-	check(PausedSystemInstances.Num() == NewSystemIndex);
-	Instance->SystemInstanceIndex = NewSystemIndex;
-	PausedSystemInstances.Add(Instance);
+	check(NewDataSetIndex == Instance->SystemInstanceIndex);
 
 	SystemInstances.RemoveAtSwap(SystemIndex);
 	if (SystemInstances.IsValidIndex(SystemIndex))
 	{
 		SystemInstances[SystemIndex]->SystemInstanceIndex = SystemIndex;
 	}
+
+	check(SystemInstances.Num() == DataSet.GetCurrentDataChecked().GetNumInstances());
+	check(PausedSystemInstances.Num() == PausedInstanceData.GetCurrentDataChecked().GetNumInstances());
 }
 
 void FNiagaraSystemSimulation::UnpauseInstance(FNiagaraSystemInstance* Instance)
 {
+	check(IsInGameThread());
 	check(Instance->IsPaused());
 	check(!DataSet.GetDestinationData());
 	check(!PausedInstanceData.GetDestinationData());
 
+	check(SystemInstances.Num() == DataSet.GetCurrentDataChecked().GetNumInstances());
+	check(PausedSystemInstances.Num() == PausedInstanceData.GetCurrentDataChecked().GetNumInstances());
+
 	UNiagaraSystem* System = WeakSystem.Get();
 	if (Instance->IsPendingSpawn())
 	{
@@ -890,18 +1288,19 @@ void FNiagaraSystemSimulation::UnpauseInstance(FNiagaraSystemInstance* Instance)
 	check(PausedSystemInstances.IsValidIndex(SystemIndex));
 	check(Instance == PausedSystemInstances[SystemIndex]);
 
-	int32 NewSystemIndex = DataSet.GetCurrentDataChecked().TransferInstance(PausedInstanceData.GetCurrentDataChecked(), SystemIndex);
-	PausedInstanceData.GetCurrentDataChecked().KillInstance(SystemIndex);
+	int32 NewDataSetIndex = DataSet.GetCurrentDataChecked().TransferInstance(PausedInstanceData.GetCurrentDataChecked(), SystemIndex);
 
-	check(SystemInstances.Num() == NewSystemIndex);
-	Instance->SystemInstanceIndex = NewSystemIndex;
-	SystemInstances.Add(Instance);
+	Instance->SystemInstanceIndex = SystemInstances.Add(Instance);
+	check(NewDataSetIndex == Instance->SystemInstanceIndex);
 
 	PausedSystemInstances.RemoveAtSwap(SystemIndex);
 	if (PausedSystemInstances.IsValidIndex(SystemIndex))
 	{
 		PausedSystemInstances[SystemIndex]->SystemInstanceIndex = SystemIndex;
 	}
+
+	check(SystemInstances.Num() == DataSet.GetCurrentDataChecked().GetNumInstances());
+	check(PausedSystemInstances.Num() == PausedInstanceData.GetCurrentDataChecked().GetNumInstances());
 }
 
 void FNiagaraSystemSimulation::InitParameterDataSetBindings(FNiagaraSystemInstance* SystemInst)
@@ -910,8 +1309,10 @@ void FNiagaraSystemSimulation::InitParameterDataSetBindings(FNiagaraSystemInstan
 	//TODO: Pull the layout stuff out of each data set and store. So much duplicated data.
 	//This assumes that all layouts for all emitters is the same. Which it should be.
 	//Ideally we can store all this layout info in the systm/emitter assets so we can just generate this in Init()
-	if (SystemInst != nullptr)
+	if (!bBindingsInitialized && SystemInst != nullptr)
 	{
+		bBindingsInitialized = true;
+
 		SpawnInstanceParameterToDataSetBinding.Init(SpawnInstanceParameterDataSet, SystemInst->GetInstanceParameters());
 		UpdateInstanceParameterToDataSetBinding.Init(UpdateInstanceParameterDataSet, SystemInst->GetInstanceParameters());
 
diff --git a/Engine/Plugins/FX/Niagara/Source/Niagara/Private/NiagaraWorldManager.cpp b/Engine/Plugins/FX/Niagara/Source/Niagara/Private/NiagaraWorldManager.cpp
--- a/Engine/Plugins/FX/Niagara/Source/Niagara/Private/NiagaraWorldManager.cpp
+++ b/Engine/Plugins/FX/Niagara/Source/Niagara/Private/NiagaraWorldManager.cpp
@@ -20,6 +20,13 @@
 DECLARE_CYCLE_STAT(TEXT("Niagara Manager Tick [GT]"), STAT_NiagaraWorldManTick, STATGROUP_Niagara);
 DECLARE_CYCLE_STAT(TEXT("Niagara Manager Wait On Render [GT]"), STAT_NiagaraWorldManWaitOnRender, STATGROUP_Niagara);
 
+FDelegateHandle FNiagaraWorldManager::OnWorldInitHandle;
+FDelegateHandle FNiagaraWorldManager::OnWorldCleanupHandle;
+FDelegateHandle FNiagaraWorldManager::OnPreWorldFinishDestroyHandle;
+FDelegateHandle FNiagaraWorldManager::OnWorldBeginTearDownHandle;
+FDelegateHandle FNiagaraWorldManager::TickWorldHandle;
+TMap<class UWorld*, class FNiagaraWorldManager*> FNiagaraWorldManager::WorldManagers;
+
 TGlobalResource<FNiagaraViewDataMgr> GNiagaraViewDataManager;
 
 FNiagaraViewDataMgr::FNiagaraViewDataMgr() 
@@ -57,16 +64,86 @@ void FNiagaraViewDataMgr::ReleaseDynamicRHI()
 	SceneTexturesUniformParams.SafeRelease();
 }
 
+//////////////////////////////////////////////////////////////////////////
+
+void FNiagaraWorldManagerTickFunction::ExecuteTick(float DeltaTime, enum ELevelTick TickType, ENamedThreads::Type CurrentThread, const FGraphEventRef& MyCompletionGraphEvent)
+{
+	check(Owner);
+	Owner->Tick(TickGroup, DeltaTime, TickType, CurrentThread, MyCompletionGraphEvent);
+}
+
+FString FNiagaraWorldManagerTickFunction::DiagnosticMessage()
+{
+	static const UEnum* EnumType = FindObjectChecked<UEnum>(ANY_PACKAGE, TEXT("ETickingGroup"));
+
+	return TEXT("FParticleSystemManager::Tick(") + EnumType->GetNameStringByIndex(static_cast<uint32>(TickGroup)) + TEXT(")");
+}
+
+FName FNiagaraWorldManagerTickFunction::DiagnosticContext(bool bDetailed)
+{
+	return FName(TEXT("ParticleSystemManager"));
+}
+
+//////////////////////////////////////////////////////////////////////////
+
 FNiagaraWorldManager::FNiagaraWorldManager(UWorld* InWorld)
 	: World(InWorld)
 	, CachedEffectsQuality(INDEX_NONE)
 {
+	FNiagaraWorldManagerTickFunction& TickFunc = TickFunction;//TODO: Add more tick functions to kick work earlier in frame.
+	TickFunc.TickGroup = ETickingGroup::TG_LastDemotable;
+	TickFunc.EndTickGroup = TickFunc.TickGroup;
+	TickFunc.bCanEverTick = true;
+	TickFunc.bStartWithTickEnabled = true;
+	TickFunc.bHighPriority = true;
+	TickFunc.Owner = this;
+	TickFunc.RegisterTickFunction(InWorld->PersistentLevel);
+
 	ComponentPool = NewObject<UNiagaraComponentPool>();
 }
 
+FNiagaraWorldManager::~FNiagaraWorldManager()
+{
+	OnWorldCleanup(true, true);
+}
+
 FNiagaraWorldManager* FNiagaraWorldManager::Get(UWorld* World)
 {
-	return INiagaraModule::GetWorldManager(World);
+	FNiagaraWorldManager** OutWorld = WorldManagers.Find(World);
+	if (OutWorld == nullptr)
+	{
+		UE_LOG(LogNiagara, Warning, TEXT("Calling FNiagaraWorldManager::Get \"%s\", but Niagara has never encountered this world before. "
+			" This means that WorldInit never happened. This may happen in some edge cases in the editor, like saving invisible child levels, "
+			"in which case the calling context needs to be safe against this returning nullptr."), World ? *World->GetName() : TEXT("nullptr"));
+		return nullptr;
+	}
+	return *OutWorld;
+}
+
+void FNiagaraWorldManager::OnStartup()
+{
+	OnWorldInitHandle = FWorldDelegates::OnPreWorldInitialization.AddStatic(&FNiagaraWorldManager::OnWorldInit);
+	OnWorldCleanupHandle = FWorldDelegates::OnWorldCleanup.AddStatic(&FNiagaraWorldManager::OnWorldCleanup);
+	OnPreWorldFinishDestroyHandle = FWorldDelegates::OnPreWorldFinishDestroy.AddStatic(&FNiagaraWorldManager::OnPreWorldFinishDestroy);
+	OnWorldBeginTearDownHandle = FWorldDelegates::OnWorldBeginTearDown.AddStatic(&FNiagaraWorldManager::OnWorldBeginTearDown);
+	TickWorldHandle = FWorldDelegates::OnWorldPostActorTick.AddStatic(&FNiagaraWorldManager::TickWorld);
+}
+
+void FNiagaraWorldManager::OnShutdown()
+{
+	FWorldDelegates::OnPreWorldInitialization.Remove(OnWorldInitHandle);
+	FWorldDelegates::OnWorldCleanup.Remove(OnWorldCleanupHandle);
+	FWorldDelegates::OnPreWorldFinishDestroy.Remove(OnPreWorldFinishDestroyHandle);
+	FWorldDelegates::OnWorldBeginTearDown.Remove(OnWorldBeginTearDownHandle);
+	FWorldDelegates::OnWorldPostActorTick.Remove(TickWorldHandle);
+
+	//Should have cleared up all world managers by now.
+	check(WorldManagers.Num() == 0);
+	for (TPair<UWorld*, FNiagaraWorldManager*> Pair : WorldManagers)
+	{
+		delete Pair.Value;
+		Pair.Value = nullptr;
+	}
 }
 
 void FNiagaraWorldManager::AddReferencedObjects(FReferenceCollector& Collector)
@@ -185,7 +262,7 @@ void FNiagaraWorldManager::DestroySystemInstance(TUniquePtr<FNiagaraSystemInstan
 	DeferredDeletionQueue[DeferredDeletionQueueIndex].Queue.Emplace(MoveTemp(InPtr));
 }
 
-void FNiagaraWorldManager::OnBatcherDestroyed(NiagaraEmitterInstanceBatcher* InBatcher)
+void FNiagaraWorldManager::OnBatcherDestroyed_Internal(NiagaraEmitterInstanceBatcher* InBatcher)
 {
 	// Process the deferred deletion queue before deleting the batcher of this world.
 	// This is required because the batcher is accessed in FNiagaraEmitterInstance::~FNiagaraEmitterInstance
@@ -218,7 +295,117 @@ void FNiagaraWorldManager::OnWorldCleanup(bool bSessionEnded, bool bCleanupResou
 	}
 }
 
-void FNiagaraWorldManager::Tick(float DeltaSeconds)
+void FNiagaraWorldManager::OnWorldInit(UWorld* World, const UWorld::InitializationValues IVS)
+{
+	check(WorldManagers.Find(World) == nullptr);
+	WorldManagers.Add(World) = new FNiagaraWorldManager(World);
+}
+
+void FNiagaraWorldManager::OnWorldCleanup(UWorld* World, bool bSessionEnded, bool bCleanupResources)
+{
+	//Cleanup world manager contents but not the manager itself.
+	FNiagaraWorldManager** Manager = WorldManagers.Find(World);
+	if (Manager)
+	{
+		(*Manager)->OnWorldCleanup(bSessionEnded, bCleanupResources);
+	}
+}
+
+void FNiagaraWorldManager::OnPreWorldFinishDestroy(UWorld* World)
+{
+	FNiagaraWorldManager** Manager = WorldManagers.Find(World);
+	if (Manager)
+	{
+		delete (*Manager);
+		WorldManagers.Remove(World);
+	}
+}
+
+void FNiagaraWorldManager::OnWorldBeginTearDown(UWorld* World)
+{
+// 	FNiagaraWorldManager** Manager = WorldManagers.Find(World);
+// 	if (Manager)
+// 	{
+// 		delete (*Manager);
+// 		WorldManagers.Remove(World);
+// 	}
+// 	FNiagaraWorldManager** Manager = WorldManagers.Find(World);
+// 	if (Manager)
+// 	{
+// 		Manager->SystemSimulations
+// 	}
+}
+
+void FNiagaraWorldManager::OnBatcherDestroyed(NiagaraEmitterInstanceBatcher* InBatcher)
+{
+	for (TPair<UWorld*, FNiagaraWorldManager*>& Pair : WorldManagers)
+	{
+		Pair.Value->OnBatcherDestroyed_Internal(InBatcher);
+	}
+}
+
+void FNiagaraWorldManager::DestroyAllSystemSimulations(class UNiagaraSystem* System)
+{
+	for (TPair<UWorld*, FNiagaraWorldManager*>& Pair : WorldManagers)
+	{
+		Pair.Value->DestroySystemSimulation(System);
+	}
+}
+
+void FNiagaraWorldManager::TickWorld(UWorld* World, ELevelTick TickType, float DeltaSeconds)
+{
+	Get(World)->PostActorTick(DeltaSeconds);
+}
+
+void FNiagaraWorldManager::PostActorTick(float DeltaSeconds)
+{
+	CSV_SCOPED_TIMING_STAT_EXCLUSIVE(Niagara);
+	LLM_SCOPE(ELLMTag::Niagara);
+	SCOPE_CYCLE_COUNTER(STAT_NiagaraWorldManTick);
+	SCOPE_CYCLE_COUNTER(STAT_NiagaraOverview_GT);
+
+	TArray<UNiagaraSystem*, TInlineAllocator<4>> DeadSystems;
+	for (TPair<UNiagaraSystem*, TSharedRef<FNiagaraSystemSimulation, ESPMode::ThreadSafe>>& SystemSim : SystemSimulations)
+	{
+		FNiagaraSystemSimulation*  Sim = &SystemSim.Value.Get();
+		if (Sim->IsValid())
+		{
+			Sim->SpawnNew_GameThread(DeltaSeconds, nullptr);
+		}
+		else
+		{
+			DeadSystems.Add(SystemSim.Key);
+		}
+	}
+	for (UNiagaraSystem* DeadSystem : DeadSystems)
+	{
+		SystemSimulations.Remove(DeadSystem);
+	}
+
+	// Clear cached player view location list, it should never be used outside of the world tick
+	bCachedPlayerViewLocationsValid = false;
+	CachedPlayerViewLocations.Reset();
+
+	// Enqueue fence for deferred deletion if we need to wait on anything
+	if (DeferredDeletionQueue[DeferredDeletionQueueIndex].Queue.Num() > 0)
+	{
+		DeferredDeletionQueue[DeferredDeletionQueueIndex].Fence.BeginFence();
+	}
+
+	// Remove instances from oldest frame making sure they aren't in use on the RT
+	DeferredDeletionQueueIndex = (DeferredDeletionQueueIndex + 1) % NumDeferredQueues;
+	if (DeferredDeletionQueue[DeferredDeletionQueueIndex].Queue.Num() > 0)
+	{
+		if (!DeferredDeletionQueue[DeferredDeletionQueueIndex].Fence.IsFenceComplete())
+		{
+			SCOPE_CYCLE_COUNTER(STAT_NiagaraWorldManWaitOnRender);
+			DeferredDeletionQueue[DeferredDeletionQueueIndex].Fence.Wait();
+		}
+		DeferredDeletionQueue[DeferredDeletionQueueIndex].Queue.Empty();
+	}
+}
+
+void FNiagaraWorldManager::Tick(ETickingGroup TickGroup, float DeltaSeconds, ELevelTick TickType, ENamedThreads::Type CurrentThread, const FGraphEventRef& MyCompletionGraphEvent)
 {
 	CSV_SCOPED_TIMING_STAT_EXCLUSIVE(Niagara);
 	LLM_SCOPE(ELLMTag::Niagara);
@@ -233,7 +420,7 @@ void FNiagaraWorldManager::Tick(float DeltaSeconds)
 	bCachedPlayerViewLocationsValid = true;
 	if (World->GetPlayerControllerIterator())
 	{
-		for ( FConstPlayerControllerIterator Iterator=World->GetPlayerControllerIterator(); Iterator; ++Iterator)
+		for (FConstPlayerControllerIterator Iterator = World->GetPlayerControllerIterator(); Iterator; ++Iterator)
 		{
 			APlayerController* PlayerController = Iterator->Get();
 			if (PlayerController && PlayerController->IsLocalPlayerController())
@@ -257,10 +444,16 @@ void FNiagaraWorldManager::Tick(float DeltaSeconds)
 	}
 
 	//Now tick all system instances. 
-	TArray<UNiagaraSystem*> DeadSystems;
+	TArray<UNiagaraSystem*, TInlineAllocator<4>> DeadSystems;
 	for (TPair<UNiagaraSystem*, TSharedRef<FNiagaraSystemSimulation, ESPMode::ThreadSafe>>& SystemSim : SystemSimulations)
 	{
-		if (SystemSim.Value->Tick(DeltaSeconds) == false)
+		FNiagaraSystemSimulation*  Sim = &SystemSim.Value.Get();
+
+		if (Sim->IsValid())
+		{
+			Sim->Tick_GameThread(DeltaSeconds, MyCompletionGraphEvent);
+		}
+		else
 		{
 			DeadSystems.Add(SystemSim.Key);
 		}
@@ -270,26 +463,4 @@ void FNiagaraWorldManager::Tick(float DeltaSeconds)
 	{
 		SystemSimulations.Remove(DeadSystem);
 	}
-
-	// Clear cached player view location list, it should never be used outside of the world tick
-	bCachedPlayerViewLocationsValid = false;
-	CachedPlayerViewLocations.Reset();
-
-	// Enqueue fence for deferred deletion if we need to wait on anything
-	if ( DeferredDeletionQueue[DeferredDeletionQueueIndex].Queue.Num() > 0 )
-	{
-		DeferredDeletionQueue[DeferredDeletionQueueIndex].Fence.BeginFence();
-	}
-
-	// Remove instances from oldest frame making sure they aren't in use on the RT
-	DeferredDeletionQueueIndex = (DeferredDeletionQueueIndex + 1) % NumDeferredQueues;
-	if ( DeferredDeletionQueue[DeferredDeletionQueueIndex].Queue.Num() > 0 )
-	{
-		if ( !DeferredDeletionQueue[DeferredDeletionQueueIndex].Fence.IsFenceComplete() )
-		{
-			SCOPE_CYCLE_COUNTER(STAT_NiagaraWorldManWaitOnRender);
-			DeferredDeletionQueue[DeferredDeletionQueueIndex].Fence.Wait();
-		}
-		DeferredDeletionQueue[DeferredDeletionQueueIndex].Queue.Empty();
-	}
 }
diff --git a/Engine/Plugins/FX/Niagara/Source/Niagara/Public/NiagaraModule.h b/Engine/Plugins/FX/Niagara/Source/Niagara/Public/NiagaraModule.h
--- a/Engine/Plugins/FX/Niagara/Source/Niagara/Public/NiagaraModule.h
+++ b/Engine/Plugins/FX/Niagara/Source/Niagara/Public/NiagaraModule.h
@@ -42,25 +42,6 @@ public:
 	void ResetOnProcessShaderCompilationQueue(FDelegateHandle DelegateHandle);
 	void ProcessShaderCompilationQueue();
 
-
-	static FNiagaraWorldManager* GetWorldManager(UWorld* World);
-
-	// Gamethread callback to cleanup references to the given batcher before it gets deleted on the renderthread.
-	static void OnBatcherDestroyed(class NiagaraEmitterInstanceBatcher* InBatcher);
-
-	void DestroyAllSystemSimulations(class UNiagaraSystem* System);
-
-	// Callback function registered with global world delegates to instantiate world manager when a game world is created
-	void OnWorldInit(UWorld* World, const UWorld::InitializationValues IVS);
-
-	// Callback function registered with global world delegates to cleanup world manager contents
-	void OnWorldCleanup(UWorld* World, bool bSessionEnded, bool bCleanupResources);
-
-	// Callback function registered with global world delegates to cleanup world manager when a game world is destroyed
-	void OnPreWorldFinishDestroy(UWorld* World);
-
-	void TickWorld(UWorld* World, ELevelTick TickType, float DeltaSeconds);
-
 #if WITH_EDITOR
 	const INiagaraMergeManager& GetMergeManager() const;
 
@@ -185,8 +166,6 @@ private:
 	FOnPrecompile ObjectPrecompilerDelegate;
 #endif
 
-	static TMap<class UWorld*, class FNiagaraWorldManager*> WorldManagers;
-
 	void OnChangeDetailLevel(class IConsoleVariable* CVar);
 	static int32 EngineDetailLevel;
 
diff --git a/Engine/Plugins/FX/Niagara/Source/Niagara/Public/NiagaraParameterStore.h b/Engine/Plugins/FX/Niagara/Source/Niagara/Public/NiagaraParameterStore.h
--- a/Engine/Plugins/FX/Niagara/Source/Niagara/Public/NiagaraParameterStore.h
+++ b/Engine/Plugins/FX/Niagara/Source/Niagara/Public/NiagaraParameterStore.h
@@ -400,6 +400,16 @@ public:
 		OnParameterChange();
 	}
 
+	template<typename T>
+	FORCEINLINE_DEBUGGABLE void SetParameterDataTyped(const uint8* Data, int32 Offset)
+	{
+		checkSlow(Data != nullptr);
+		checkSlow((Offset + sizeof(T)) <= ParameterData.Num());
+		uint8* Dest = GetParameterData_Internal(Offset);
+		*((T*)Dest) = *((T*)Data);
+		OnParameterChange();
+	}
+
 	FORCEINLINE_DEBUGGABLE bool SetParameterData(const uint8* Data, const FNiagaraVariable& Param, bool bAdd = false)
 	{
 		checkSlow(Data != nullptr);
@@ -816,8 +826,8 @@ struct FNiagaraParameterDirectBinding
 
 	FORCEINLINE void SetValue(const T& InValue)
 	{
-		check(BoundVariable.GetSizeInBytes() == sizeof(T));
-		checkf(LayoutVersion == BoundStore->GetLayoutVersion(), TEXT("This binding is invalid, its bound parameter store's layout was changed since it was created"));
+		checkSlow(BoundVariable.GetSizeInBytes() == sizeof(T));
+		checkfSlow(LayoutVersion == BoundStore->GetLayoutVersion(), TEXT("This binding is invalid, its bound parameter store's layout was changed since it was created"));
 
 		if (ValuePtr)
 		{
@@ -827,8 +837,8 @@ struct FNiagaraParameterDirectBinding
 
 	FORCEINLINE T GetValue()const 
 	{
-		check(BoundVariable.GetSizeInBytes() == sizeof(T));
-		checkf(LayoutVersion == BoundStore->GetLayoutVersion(), TEXT("This binding is invalid, its bound parameter store's layout was changed since it was created"));
+		checkSlow(BoundVariable.GetSizeInBytes() == sizeof(T));
+		checkfSlow(LayoutVersion == BoundStore->GetLayoutVersion(), TEXT("This binding is invalid, its bound parameter store's layout was changed since it was created"));
 
 		if (ValuePtr)
 		{
@@ -861,8 +871,8 @@ struct FNiagaraParameterDirectBinding<FMatrix>
 
 	FORCEINLINE void SetValue(const FMatrix& InValue)
 	{
-		check(BoundVariable.GetSizeInBytes() == sizeof(FMatrix));
-		checkf(LayoutVersion == BoundStore->GetLayoutVersion(), TEXT("This binding is invalid, its bound parameter store's layout was changed since it was created"));
+		checkSlow(BoundVariable.GetSizeInBytes() == sizeof(FMatrix));
+		checkfSlow(LayoutVersion == BoundStore->GetLayoutVersion(), TEXT("This binding is invalid, its bound parameter store's layout was changed since it was created"));
 
 		if (ValuePtr)
 		{
@@ -872,8 +882,8 @@ struct FNiagaraParameterDirectBinding<FMatrix>
 
 	FORCEINLINE FMatrix GetValue()const
 	{
-		check(BoundVariable.GetSizeInBytes() == sizeof(FMatrix));
-		checkf(LayoutVersion == BoundStore->GetLayoutVersion(), TEXT("This binding is invalid, its bound parameter store's layout was changed since it was created"));
+		checkSlow(BoundVariable.GetSizeInBytes() == sizeof(FMatrix));
+		checkfSlow(LayoutVersion == BoundStore->GetLayoutVersion(), TEXT("This binding is invalid, its bound parameter store's layout was changed since it was created"));
 
 		FMatrix Ret;
 		if (ValuePtr)
@@ -907,8 +917,8 @@ struct FNiagaraParameterDirectBinding<FVector4>
 
 	FORCEINLINE void SetValue(const FVector4& InValue)
 	{
-		check(BoundVariable.GetSizeInBytes() == sizeof(FVector4));
-		checkf(LayoutVersion == BoundStore->GetLayoutVersion(), TEXT("This binding is invalid, its bound parameter store's layout was changed since it was created"));
+		checkSlow(BoundVariable.GetSizeInBytes() == sizeof(FVector4));
+		checkfSlow(LayoutVersion == BoundStore->GetLayoutVersion(), TEXT("This binding is invalid, its bound parameter store's layout was changed since it was created"));
 
 		if (ValuePtr)
 		{
@@ -919,7 +929,7 @@ struct FNiagaraParameterDirectBinding<FVector4>
 	FORCEINLINE FVector4 GetValue()const
 	{
 		check(BoundVariable.GetSizeInBytes() == sizeof(FVector4));
-		checkf(LayoutVersion == BoundStore->GetLayoutVersion(), TEXT("This binding is invalid, its bound parameter store's layout was changed since it was created"));
+		checkfSlow(LayoutVersion == BoundStore->GetLayoutVersion(), TEXT("This binding is invalid, its bound parameter store's layout was changed since it was created"));
 
 		FVector4 Ret;
 		if (ValuePtr)
@@ -953,8 +963,8 @@ struct FNiagaraParameterDirectBinding<FQuat>
 
 	FORCEINLINE void SetValue(const FQuat& InValue)
 	{
-		check(BoundVariable.GetSizeInBytes() == sizeof(FQuat));
-		checkf(LayoutVersion == BoundStore->GetLayoutVersion(), TEXT("This binding is invalid, its bound parameter store's layout was changed since it was created"));
+		checkSlow(BoundVariable.GetSizeInBytes() == sizeof(FQuat));
+		checkfSlow(LayoutVersion == BoundStore->GetLayoutVersion(), TEXT("This binding is invalid, its bound parameter store's layout was changed since it was created"));
 
 		if (ValuePtr)
 		{
@@ -964,8 +974,8 @@ struct FNiagaraParameterDirectBinding<FQuat>
 
 	FORCEINLINE FQuat GetValue()const
 	{
-		check(BoundVariable.GetSizeInBytes() == sizeof(FQuat));
-		checkf(LayoutVersion == BoundStore->GetLayoutVersion(), TEXT("This binding is invalid, its bound parameter store's layout was changed since it was created"));
+		checkSlow(BoundVariable.GetSizeInBytes() == sizeof(FQuat));
+		checkfSlow(LayoutVersion == BoundStore->GetLayoutVersion(), TEXT("This binding is invalid, its bound parameter store's layout was changed since it was created"));
 
 		FQuat Ret;
 		if (ValuePtr)
@@ -1000,9 +1010,9 @@ struct FNiagaraParameterDirectBinding<FNiagaraBool>
 
 	FORCEINLINE void SetValue(const FNiagaraBool& InValue)
 	{
-		check(BoundVariable.GetSizeInBytes() == sizeof(FNiagaraBool));
-		check(sizeof(uint32) == sizeof(FNiagaraBool));
-		checkf(LayoutVersion == BoundStore->GetLayoutVersion(), TEXT("This binding is invalid, its bound parameter store's layout was changed since it was created"));
+		checkSlow(BoundVariable.GetSizeInBytes() == sizeof(FNiagaraBool));
+		checkSlow(sizeof(uint32) == sizeof(FNiagaraBool));
+		checkfSlow(LayoutVersion == BoundStore->GetLayoutVersion(), TEXT("This binding is invalid, its bound parameter store's layout was changed since it was created"));
 
 		if (ValuePtr)
 		{
@@ -1012,8 +1022,8 @@ struct FNiagaraParameterDirectBinding<FNiagaraBool>
 
 	FORCEINLINE void SetValue(const bool& InValue)
 	{
-		check(BoundVariable.GetSizeInBytes() == sizeof(FNiagaraBool));
-		checkf(LayoutVersion == BoundStore->GetLayoutVersion(), TEXT("This binding is invalid, its bound parameter store's layout was changed since it was created"));
+		checkSlow(BoundVariable.GetSizeInBytes() == sizeof(FNiagaraBool));
+		checkfSlow(LayoutVersion == BoundStore->GetLayoutVersion(), TEXT("This binding is invalid, its bound parameter store's layout was changed since it was created"));
 
 		if (ValuePtr)
 		{
@@ -1030,8 +1040,8 @@ struct FNiagaraParameterDirectBinding<FNiagaraBool>
 
 	FORCEINLINE FNiagaraBool GetValue()const
 	{
-		check(BoundVariable.GetSizeInBytes() == sizeof(FNiagaraBool));
-		checkf(LayoutVersion == BoundStore->GetLayoutVersion(), TEXT("This binding is invalid, its bound parameter store's layout was changed since it was created"));
+		checkSlow(BoundVariable.GetSizeInBytes() == sizeof(FNiagaraBool));
+		checkfSlow(LayoutVersion == BoundStore->GetLayoutVersion(), TEXT("This binding is invalid, its bound parameter store's layout was changed since it was created"));
 
 		FNiagaraBool Ret(false);
 		if (ValuePtr)
@@ -1081,8 +1091,8 @@ struct FNiagaraParameterDirectBinding<UObject*>
 	{
 		if (UObjectOffset != INDEX_NONE)
 		{
-			check(BoundVariable.GetType() == FNiagaraTypeDefinition::GetUObjectDef());
-			checkf(LayoutVersion == BoundStore->GetLayoutVersion(), TEXT("This binding is invalid, its bound parameter store's layout was changed since it was created"));
+			checkSlow(BoundVariable.GetType() == FNiagaraTypeDefinition::GetUObjectDef());
+			checkfSlow(LayoutVersion == BoundStore->GetLayoutVersion(), TEXT("This binding is invalid, its bound parameter store's layout was changed since it was created"));
 
 			BoundStore->SetUObject(InValue, UObjectOffset);
 		}
@@ -1092,8 +1102,8 @@ struct FNiagaraParameterDirectBinding<UObject*>
 	{
 		if (UObjectOffset != INDEX_NONE)
 		{
-			check(BoundVariable.GetType() == FNiagaraTypeDefinition::GetUObjectDef());
-			checkf(LayoutVersion == BoundStore->GetLayoutVersion(), TEXT("This binding is invalid, its bound parameter store's layout was changed since it was created"));
+			checkSlow(BoundVariable.GetType() == FNiagaraTypeDefinition::GetUObjectDef());
+			checkfSlow(LayoutVersion == BoundStore->GetLayoutVersion(), TEXT("This binding is invalid, its bound parameter store's layout was changed since it was created"));
 
 			return BoundStore->GetUObject(UObjectOffset);
 		}
diff --git a/Engine/Plugins/FX/Niagara/Source/Niagara/Public/NiagaraSystemInstance.h b/Engine/Plugins/FX/Niagara/Source/Niagara/Public/NiagaraSystemInstance.h
--- a/Engine/Plugins/FX/Niagara/Source/Niagara/Public/NiagaraSystemInstance.h
+++ b/Engine/Plugins/FX/Niagara/Source/Niagara/Public/NiagaraSystemInstance.h
@@ -80,10 +80,18 @@ public:
 	/** Requests the the simulation be reset on the next tick. */
 	void Reset(EResetMode Mode);
 
-	void ComponentTick(float DeltaSeconds);
-	void PreSimulateTick(float DeltaSeconds);
-	void PostSimulateTick(float DeltaSeconds);
-	void FinalizeTick(float DeltaSeconds);
+	void ComponentTick(float DeltaSeconds, const FGraphEventRef& MyCompletionGraphEvent);
+
+	/** Initial phase of system instance tick. Must be executed on the game thread. */
+	void Tick_GameThread(float DeltaSeconds);
+	/** Secondary phase of the system instance tick that can be executed on any thread. */
+	void Tick_Concurrent();
+	/** Final phase of system instance tick. Must be executed on the game thread. */
+	void FinalizeTick_GameThread();
+
+	/** Blocks until any async work for this system instance has completed. Must be called on the game thread. */
+	void WaitForAsyncTick(bool bEnsureComplete=false);
+
 	/** Handles completion of the system and returns true if the system is complete. */
 	bool HandleCompletion();
 
@@ -150,8 +158,6 @@ public:
 		return nullptr;
 	}
 
-	void DestroyDataInterfaceInstanceData();
-
 	bool UsesEmitter(const UNiagaraEmitter* Emitter)const;
 	bool UsesScript(const UNiagaraScript* Script)const;
 	//bool UsesDataInterface(UNiagaraDataInterface* Interface);
@@ -173,6 +179,8 @@ public:
 	/** Index of this instance in the system simulation. */
 	int32 SystemInstanceIndex;
 
+	TSharedPtr<class FNiagaraSystemSimulation, ESPMode::ThreadSafe> SystemSimulation;
+
 	FORCEINLINE bool HasTickingEmitters()const { return bHasTickingEmitters; }
 
 	UNiagaraParameterCollectionInstance* GetParameterCollectionInstance(UNiagaraParameterCollection* Collection);
@@ -217,8 +225,19 @@ public:
 	bool HasGPUEmitters() { return bHasGPUEmitters;  }
 
 	int32 GetDetailLevel()const;
+
+	FORCEINLINE void BeginAsyncWork()
+	{
+		bAsyncWorkInProgress = true;
+		bNeedsFinalize = true;
+	}
+
+	void TickInstanceParameters(float DeltaSeconds);
+
 private:
 
+	void DestroyDataInterfaceInstanceData();
+
 	/** Builds the emitter simulations. */
 	void InitEmitters();
 
@@ -230,7 +249,6 @@ private:
 	/** Call PrepareForSImulation on each data source from the simulations and determine which need per-tick updates.*/
 	void InitDataInterfaces();	
 
-	void TickInstanceParameters(float DeltaSeconds);
 
 	void BindParameterCollections(FNiagaraScriptExecutionContext& ExecContext);
 	
@@ -238,7 +256,6 @@ private:
 	float GetLODDistance();
 
 	UNiagaraComponent* Component;
-	TSharedPtr<class FNiagaraSystemSimulation, ESPMode::ThreadSafe> SystemSimulation;
 	FBox SystemBounds;
 
 	/** The age of the System instance. */
@@ -337,8 +354,15 @@ private:
 
 	NiagaraEmitterInstanceBatcher* Batcher = nullptr;
 public:
+	/** True if we have async work in flight. */
+	volatile bool bAsyncWorkInProgress;
+	/** True if we require a call to FinalizeTick_GameThread(). Typically this is called from a GT task but can be called in WaitForAsync. */
+	volatile bool bNeedsFinalize;
+
 	// Transient data that is accumulated during tick.
 	uint32 TotalParamSize = 0;
 	uint32 ActiveGPUEmitterCount = 0;
 	int32 GPUDataInterfaceInstanceDataSize = 0;
+
+	float CachedDeltaSeconds;
 };
diff --git a/Engine/Plugins/FX/Niagara/Source/Niagara/Public/NiagaraSystemSimulation.h b/Engine/Plugins/FX/Niagara/Source/Niagara/Public/NiagaraSystemSimulation.h
--- a/Engine/Plugins/FX/Niagara/Source/Niagara/Public/NiagaraSystemSimulation.h
+++ b/Engine/Plugins/FX/Niagara/Source/Niagara/Public/NiagaraSystemSimulation.h
@@ -12,6 +12,11 @@ class UWorld;
 class UNiagaraParameterCollection;
 class UNiagaraParameterCollectionInstance;
 
+//TODO: It would be good to have the batch size be variable per system to try to keep a good work/overhead ratio.
+//Can possibly adjust in future based on average batch execution time.
+#define NiagaraSystemTickBatchSize 4
+typedef TArray<FNiagaraSystemInstance*, TInlineAllocator<NiagaraSystemTickBatchSize>> FNiagaraSystemTickBatch;
+
 //TODO: Pull all the layout information here, in the data set and in parameter stores out into a single layout structure that's shared between all instances of it.
 //Currently there's tons of extra data and work done setting these up.
 struct FNiagaraParameterStoreToDataSetBinding
@@ -116,15 +121,50 @@ struct FNiagaraParameterStoreToDataSetBinding
 	}
 };
 
+struct FNiagaraSystemSimulationTickContext
+{
+	class FNiagaraSystemSimulation* Owner;
+	UNiagaraSystem* System;
+	float DeltaSeconds;
+	/** If true, this is a pending spawn pass where we spawn and do an initial tick on newly spawned instances. Then transfer them over to the correct tick group. */
+	bool bPendingSpawnPass;
+	FGraphEventRef MyCompletionGraphEvent;
+
+	bool bTickAsync;
+	bool bTickInstancesAsync;
+
+	FNiagaraSystemSimulationTickContext(FNiagaraSystemSimulation* InOwner, float InDeltaSeconds, bool bInPendingSpawnPass, const FGraphEventRef& InMyCompletionGraphEvent);
+
+	TArray<FNiagaraSystemInstance*>& GetInstances();
+	FNiagaraDataSet& GetDataSet();
+};
+
 /** Simulation performing all system and emitter scripts for a instances of a UNiagaraSystem in a world. */
-class FNiagaraSystemSimulation
+class FNiagaraSystemSimulation : public TSharedFromThis<FNiagaraSystemSimulation, ESPMode::ThreadSafe>
 {
+	friend FNiagaraSystemSimulationTickContext;
 public:
 	~FNiagaraSystemSimulation();
 	bool Init(UNiagaraSystem* InSystem, UWorld* InWorld, bool bInIsSolo);
 	void Destroy();
 	bool Tick(float DeltaSeconds);
 
+	bool IsValid()const { return WeakSystem.Get() != nullptr && bCanExecute && World != nullptr; }
+
+	/** At the end of the frame, we do a separate tick pass for newly spawned systems which are then pushed into the correct simulations for next frame. */
+	void SpawnNew_GameThread(float DeltaSeconds, const FGraphEventRef& MyCompletionGraphEvent);
+	
+	/** First phase of system sim tick. Must run on GameThread. */
+	void Tick_GameThread(float DeltaSeconds, const FGraphEventRef& MyCompletionGraphEvent);
+
+	/** Second phase of system sim tick that can run on any thread. */
+	void Tick_Concurrent(FNiagaraSystemSimulationTickContext& Context);
+
+	void WaitForTickComplete();
+
+	/** Old tick for AB Testing.*/
+	bool Tick_Old(float DeltaSeconds);
+
 	void RemoveInstance(FNiagaraSystemInstance* Instance);
 	void AddInstance(FNiagaraSystemInstance* Instance);
 
@@ -151,6 +191,23 @@ public:
 
 protected:
 
+	/** Does any prep work for system simulation such as pulling instance parameters into a dataset. */
+	void PrepareForSystemSimulate(FNiagaraSystemSimulationTickContext& Context);
+	/** Runs the system spawn script for new system instances. */
+	void SpawnSystemInstances(FNiagaraSystemSimulationTickContext& Context);
+	/** Runs the system update script. */
+	void UpdateSystemInstances(FNiagaraSystemSimulationTickContext& Context);
+	/** Transfers the results of the system simulation into the emitter instances. */
+	void TransferSystemSimResults(FNiagaraSystemSimulationTickContext& Context);
+
+	/** Should we push the system sim tick off the game thread. */
+	FORCEINLINE bool ShouldTickAsync(const FNiagaraSystemSimulationTickContext& Context);
+	/** Should we push the system instance ticks off the game thread. */
+	FORCEINLINE bool ShouldTickInstancesAsync(const FNiagaraSystemSimulationTickContext& Context);
+
+	void AddSystemToTickBatch(FNiagaraSystemInstance* Instance, FNiagaraSystemSimulationTickContext& Context);
+	void FlushTickBatch(FNiagaraSystemSimulationTickContext& Context);
+
 	/** System of instances being simulated.  We use a weak object ptr here because once the last referencing object goes away this system may be come invalid at runtime. */
 	TWeakObjectPtr<UNiagaraSystem> WeakSystem;
 
@@ -159,6 +216,9 @@ protected:
 
 	/** Main dataset containing system instance attribute data. */
 	FNiagaraDataSet DataSet;
+
+	/** Dataset containing system instance data for newly spawned systems. */
+	FNiagaraDataSet PendingSpawnDataSet;
 	
 	/**
 	As there's a 1 to 1 relationship between system instance and their execution in this simulation we must pull all that instances parameters into a dataset for simulation.
@@ -202,11 +262,12 @@ protected:
 	FNiagaraParameterDirectBinding<float> SpawnGlobalSystemCountScaleParam;
 	FNiagaraParameterDirectBinding<float> UpdateGlobalSystemCountScaleParam;
 
-
 	/** System instances that have been spawned and are now simulating. */
 	TArray<FNiagaraSystemInstance*> SystemInstances;
 	/** System instances that are pending to be spawned. */
 	TArray<FNiagaraSystemInstance*> PendingSystemInstances;
+	/** System instances that we are spawning this tick. Instances are moved here from PendingSystemInstances if they are not paused. */
+	TArray<FNiagaraSystemInstance*> SpawningSystemInstances;
 
 	/** System instances that are paused. */
 	TArray<FNiagaraSystemInstance*> PausedSystemInstances;
@@ -220,6 +281,8 @@ protected:
 	TArray<FNiagaraDataSetAccessor<int32>> EmitterExecutionStateAccessors;
 
 	uint32 bCanExecute : 1;
+	uint32 bBindingsInitialized : 1;
+	uint32 bInSpawnPhase : 1;
 
 	/** A parameter store which contains the data interfaces parameters which were defined by the scripts. */
 	FNiagaraParameterStore ScriptDefinedDataInterfaceParameters;
@@ -228,5 +291,8 @@ protected:
 
 	TOptional<float> MaxDeltaTime;
 
+	/** Current tick batch we're filling ready for processing, potentially in an async task. */
+	FNiagaraSystemTickBatch TickBatch;
+
 	TArray<TUniquePtr< FNiagaraSystemInstance> > DeferredDeletionQueue;
 };
diff --git a/Engine/Plugins/FX/Niagara/Source/Niagara/Public/NiagaraWorldManager.h b/Engine/Plugins/FX/Niagara/Source/Niagara/Public/NiagaraWorldManager.h
--- a/Engine/Plugins/FX/Niagara/Source/Niagara/Public/NiagaraWorldManager.h
+++ b/Engine/Plugins/FX/Niagara/Source/Niagara/Public/NiagaraWorldManager.h
@@ -3,6 +3,9 @@
 #pragma once
 
 #include "CoreMinimal.h"
+#include "UObject/ObjectMacros.h"
+#include "Engine/EngineTypes.h"
+#include "Engine/EngineBaseTypes.h"
 #include "Engine/World.h"
 #include "NiagaraParameterCollection.h"
 #include "UObject/GCObject.h"
@@ -14,6 +17,8 @@
 #include "NiagaraDataInterfaceSkeletalMesh.h"
 #include "NiagaraComponentPool.h"
 
+#include "NiagaraWorldManager.generated.h"
+
 class UWorld;
 class UNiagaraParameterCollection;
 class UNiagaraParameterCollectionInstance;
@@ -56,6 +61,28 @@ private:
 
 extern TGlobalResource<FNiagaraViewDataMgr> GNiagaraViewDataManager;
 
+USTRUCT()
+struct FNiagaraWorldManagerTickFunction : public FTickFunction
+{
+	GENERATED_USTRUCT_BODY()
+
+	//~ FTickFunction Interface
+	virtual void ExecuteTick(float DeltaTime, ELevelTick TickType, ENamedThreads::Type CurrentThread, const FGraphEventRef& MyCompletionGraphEvent) override;
+	virtual FString DiagnosticMessage() override;
+	virtual FName DiagnosticContext(bool bDetailed)  override;
+	//~ FTickFunction Interface
+
+	FNiagaraWorldManager* Owner;
+};
+
+template<>
+struct TStructOpsTypeTraits<FNiagaraWorldManagerTickFunction> : public TStructOpsTypeTraitsBase2<FNiagaraWorldManagerTickFunction>
+{
+	enum
+	{
+		WithCopy = false
+	};
+};
 
 /**
 * Manager class for any data relating to a particular world.
@@ -65,7 +92,16 @@ class FNiagaraWorldManager : public FGCObject
 public:
 	
 	FNiagaraWorldManager(UWorld* InWorld);
+	~FNiagaraWorldManager();
+
 	static FNiagaraWorldManager* Get(UWorld* World);
+	static void OnStartup();
+	static void OnShutdown();
+
+	// Gamethread callback to cleanup references to the given batcher before it gets deleted on the renderthread.
+	static void OnBatcherDestroyed(class NiagaraEmitterInstanceBatcher* InBatcher);
+
+	static void DestroyAllSystemSimulations(class UNiagaraSystem* System);
 
 	//~ GCObject Interface
 	virtual void AddReferencedObjects(FReferenceCollector& Collector) override;
@@ -77,12 +113,12 @@ public:
 	void CleanupParameterCollections();
 	TSharedRef<FNiagaraSystemSimulation, ESPMode::ThreadSafe> GetSystemSimulation(UNiagaraSystem* System);
 	void DestroySystemSimulation(UNiagaraSystem* System);
-	void DestroySystemInstance(TUniquePtr<FNiagaraSystemInstance>& InPtr);
+	void DestroySystemInstance(TUniquePtr<FNiagaraSystemInstance>& InPtr);	
 
-	// Gamethread callback to cleanup references to the given batcher before it gets deleted on the renderthread.
-	void OnBatcherDestroyed(NiagaraEmitterInstanceBatcher* InBatcher);
+	void Tick(ETickingGroup TickGroup, float DeltaSeconds, ELevelTick TickType, ENamedThreads::Type CurrentThread, const FGraphEventRef& MyCompletionGraphEvent);
 
-	void Tick(float DeltaSeconds);
+	/** Called after all actor tick groups are complete. */
+	void PostActorTick(float DeltaSeconds);
 
 	void OnWorldCleanup(bool bSessionEnded, bool bCleanupResources);
 	
@@ -94,8 +130,37 @@ public:
 	UNiagaraComponentPool* GetComponentPool() { return ComponentPool; }
 
 private:
+m
+	// Callback function registered with global world delegates to instantiate world manager when a game world is created
+	static void OnWorldInit(UWorld* World, const UWorld::InitializationValues IVS);
+
+	// Callback function registered with global world delegates to cleanup world manager contents
+	static void OnWorldCleanup(UWorld* World, bool bSessionEnded, bool bCleanupResources);
+
+	// Callback function registered with global world delegates to cleanup world manager when a game world is destroyed
+	static void OnPreWorldFinishDestroy(UWorld* World);
+
+	// Called when the world begins to be torn down for example by level streaming.
+	static void OnWorldBeginTearDown(UWorld* World);
+
+	// Callback for when a world is ticked.
+	static void TickWorld(UWorld* World, ELevelTick TickType, float DeltaSeconds);
+	
+	// Gamethread callback to cleanup references to the given batcher before it gets deleted on the renderthread.
+	void OnBatcherDestroyed_Internal(NiagaraEmitterInstanceBatcher* InBatcher);
+
+	static FDelegateHandle OnWorldInitHandle;
+	static FDelegateHandle OnWorldCleanupHandle;
+	static FDelegateHandle OnPreWorldFinishDestroyHandle;
+	static FDelegateHandle OnWorldBeginTearDownHandle;
+	static FDelegateHandle TickWorldHandle;
+
+	static TMap<class UWorld*, class FNiagaraWorldManager*> WorldManagers;
+
 	UWorld* World;
 
+	FNiagaraWorldManagerTickFunction TickFunction;//TODO: Add more tick functions to kick some work earlier in frame.
+
 	TMap<UNiagaraParameterCollection*, UNiagaraParameterCollectionInstance*> ParameterCollections;
 
 	TMap<UNiagaraSystem*, TSharedRef<FNiagaraSystemSimulation, ESPMode::ThreadSafe>> SystemSimulations;
diff --git a/Engine/Plugins/FX/Niagara/Source/NiagaraEditor/Private/NiagaraOverviewNode.cpp b/Engine/Plugins/FX/Niagara/Source/NiagaraEditor/Private/NiagaraOverviewNode.cpp
--- a/Engine/Plugins/FX/Niagara/Source/NiagaraEditor/Private/NiagaraOverviewNode.cpp
+++ b/Engine/Plugins/FX/Niagara/Source/NiagaraEditor/Private/NiagaraOverviewNode.cpp
@@ -29,6 +29,11 @@ const FGuid UNiagaraOverviewNode::GetEmitterHandleGuid() const
 
 FText UNiagaraOverviewNode::GetNodeTitle(ENodeTitleType::Type TitleType) const
 {
+	if (OwningSystem == nullptr)
+	{
+		return FText::GetEmpty();
+	}
+
 	if (EmitterHandleGuid.IsValid())
 	{
 		for (const FNiagaraEmitterHandle& Handle : OwningSystem->GetEmitterHandles())
diff --git a/Engine/Source/Runtime/VectorVM/Private/VectorVM.cpp b/Engine/Source/Runtime/VectorVM/Private/VectorVM.cpp
--- a/Engine/Source/Runtime/VectorVM/Private/VectorVM.cpp
+++ b/Engine/Source/Runtime/VectorVM/Private/VectorVM.cpp
@@ -1841,7 +1841,7 @@ void VectorVM::Exec(
 
 	auto ExecChunkBatch = [&](int32 BatchIdx)
 	{
-		SCOPE_CYCLE_COUNTER(STAT_VVMExecChunk);
+		//SCOPE_CYCLE_COUNTER(STAT_VVMExecChunk);
 
 		FVectorVMContext& Context = FVectorVMContext::Get();
 		Context.PrepareForExec(InputRegisters, OutputRegisters, NumInputRegisters, NumOutputRegisters, ConstantTable, DataSetIndexTable.GetData(), DataSetOffsetTable.GetData(), DataSetOffsetTable.Num(),
