Compile fix. Removed mysterious 'm'

diff --git a/Engine/Plugins/FX/Niagara/Source/Niagara/Public/NiagaraWorldManager.h b/Engine/Plugins/FX/Niagara/Source/Niagara/Public/NiagaraWorldManager.h
--- a/Engine/Plugins/FX/Niagara/Source/Niagara/Public/NiagaraWorldManager.h
+++ b/Engine/Plugins/FX/Niagara/Source/Niagara/Public/NiagaraWorldManager.h
@@ -130,7 +130,7 @@ public:
 	UNiagaraComponentPool* GetComponentPool() { return ComponentPool; }
 
 private:
-m
+
 	// Callback function registered with global world delegates to instantiate world manager when a game world is created
 	static void OnWorldInit(UWorld* World, const UWorld::InitializationValues IVS);
 
