Fix check. Could sometimes begin a sim pass without a corresponding end call and get into a bad state.

diff --git a/Engine/Plugins/FX/Niagara/Source/Niagara/Private/NiagaraDataSet.cpp b/Engine/Plugins/FX/Niagara/Source/Niagara/Private/NiagaraDataSet.cpp
--- a/Engine/Plugins/FX/Niagara/Source/Niagara/Private/NiagaraDataSet.cpp
+++ b/Engine/Plugins/FX/Niagara/Source/Niagara/Private/NiagaraDataSet.cpp
@@ -443,11 +443,9 @@ void FNiagaraDataSet::CopyTo(FNiagaraDataSet& Other, int32 StartIdx, int32 NumIn
 		checkSlow(Other.GetVariables() == Variables);
 	}
 
-	FNiagaraDataBuffer* OtherCurrentBuffer = Other.GetCurrentData();
-	FNiagaraDataBuffer& OtherDestBuffer = Other.BeginSimulate();
-
 	//Read the most current data. Even if it's possibly partially complete simulation data.
 	FNiagaraDataBuffer* SourceBuffer = GetDestinationData() ? GetDestinationData() : GetCurrentData();
+	FNiagaraDataBuffer* OtherCurrentBuffer = Other.GetCurrentData();
 
 	if (SourceBuffer != nullptr)
 	{
@@ -464,6 +462,8 @@ void FNiagaraDataSet::CopyTo(FNiagaraDataSet& Other, int32 StartIdx, int32 NumIn
 			NumInstances = SourceBuffer->GetNumInstances() - StartIdx;
 		}
 
+		FNiagaraDataBuffer& OtherDestBuffer = Other.BeginSimulate();
+
 		//We need to allocate enough space for the new data and existing data if we're keeping it.
 		int32 RequiredInstances = bResetOther ? NumInstances : NumInstances + OrigNumInstances;
 		OtherDestBuffer.Allocate(RequiredInstances);
@@ -477,9 +477,9 @@ void FNiagaraDataSet::CopyTo(FNiagaraDataSet& Other, int32 StartIdx, int32 NumIn
 
 		//Now copy the data from the source buffer into the newly allocated space.
 		SourceBuffer->CopyTo(OtherDestBuffer, 0, OrigNumInstances, NumInstances);
-	}
 
-	Other.EndSimulate();
+		Other.EndSimulate();
+	}
 }
 
 void FNiagaraDataSet::CopyFromGPUReadback(float* GPUReadBackFloat, int* GPUReadBackInt, int32 StartIdx /* = 0 */, int32 NumInstances /* = INDEX_NONE */, uint32 FloatStride, uint32 IntStride)
