Fixed issue with transfering structs from the system sim dataset into parameter stores.
We were using the register layout of components which is not the same as the dataset component buffer layout.
Possibly should unify these for simplicity in the future.

diff --git a/Engine/Plugins/FX/Niagara/Source/Niagara/Public/NiagaraSystemSimulation.h b/Engine/Plugins/FX/Niagara/Source/Niagara/Public/NiagaraSystemSimulation.h
--- a/Engine/Plugins/FX/Niagara/Source/Niagara/Public/NiagaraSystemSimulation.h
+++ b/Engine/Plugins/FX/Niagara/Source/Niagara/Public/NiagaraSystemSimulation.h
@@ -48,19 +48,21 @@ struct FNiagaraParameterStoreToDataSetBinding
 		{
 			const FNiagaraVariableLayoutInfo* Layout = DataSet.GetVariableLayout(Var);
 			const int32* ParameterOffsetPtr = ParameterStore.FindParameterOffset(Var);
+			int32 NumFloats = 0;
+			int32 NumInts = 0;
 			if (ParameterOffsetPtr && Layout)
 			{
 				int32 ParameterOffset = *ParameterOffsetPtr;
 				for (uint32 CompIdx = 0; CompIdx < Layout->GetNumFloatComponents(); ++CompIdx)
 				{
 					int32 ParamOffset = ParameterOffset + Layout->LayoutInfo.FloatComponentByteOffsets[CompIdx];
-					int32 DataSetOffset = Layout->FloatComponentStart + Layout->LayoutInfo.FloatComponentRegisterOffsets[CompIdx];
+					int32 DataSetOffset = Layout->FloatComponentStart + NumFloats++;
 					FloatOffsets.Emplace(ParamOffset, DataSetOffset);
 				}
 				for (uint32 CompIdx = 0; CompIdx < Layout->GetNumInt32Components(); ++CompIdx)
 				{
 					int32 ParamOffset = ParameterOffset + Layout->LayoutInfo.Int32ComponentByteOffsets[CompIdx];
-					int32 DataSetOffset = Layout->Int32ComponentStart + Layout->LayoutInfo.Int32ComponentRegisterOffsets[CompIdx];
+					int32 DataSetOffset = Layout->Int32ComponentStart + NumInts++;
 					Int32Offsets.Emplace(ParamOffset, DataSetOffset);
 				}
 			}
