Check GPU scripts are up to date in UNiagaraEmitter::AreAllScriptAndSourcesSynchronized()

diff --git a/Engine/Plugins/FX/Niagara/Source/Niagara/Private/NiagaraEmitter.cpp b/Engine/Plugins/FX/Niagara/Source/Niagara/Private/NiagaraEmitter.cpp
--- a/Engine/Plugins/FX/Niagara/Source/Niagara/Private/NiagaraEmitter.cpp
+++ b/Engine/Plugins/FX/Niagara/Source/Niagara/Private/NiagaraEmitter.cpp
@@ -715,6 +715,11 @@ bool UNiagaraEmitter::AreAllScriptAndSourcesSynchronized() const
 		}
 	}
 
+	if (GPUComputeScript->IsCompilable() && !GPUComputeScript->AreScriptAndSourceSynchronized())
+	{
+		return false;
+	}
+
 	return true;
 }
 
