Cleaned up ribbon tangent computations
Removed "Niagara.Ribbon.CullTwistedStrips" as it turns out the underlying issue can only be fixed through tessellation
Normalized the tessellation metrics by the length of each segments, removing the need for "Niagara.Ribbon.Tessellation.MaxAveraging"
Tessellation override per emitter is now limited by the platform limit

diff --git a/Engine/Plugins/FX/Niagara/Shaders/Private/NiagaraRibbonVertexFactory.ush b/Engine/Plugins/FX/Niagara/Shaders/Private/NiagaraRibbonVertexFactory.ush
--- a/Engine/Plugins/FX/Niagara/Shaders/Private/NiagaraRibbonVertexFactory.ush
+++ b/Engine/Plugins/FX/Niagara/Shaders/Private/NiagaraRibbonVertexFactory.ush
@@ -91,7 +91,6 @@ struct FVertexFactoryIntermediates
 	float3 Right;
 	float3 Up;
 	float3 Facing;
-	float3 SideVector;
 	float4 TexCoord;
 	
 #if USE_PARTICLE_POSITION
@@ -558,7 +557,11 @@ float3 SafeNormalize(float3 V)
 #endif
 
 void GetNiagaraParticleTangents(in FInterpVtxData InterpVtxData, in out FVertexFactoryIntermediates Intermediates)
-{	
+{
+	// Up : direction of the ribbon expanding.
+	// Facing : normal of the ribbon surface.
+	// Right : direction of the ribbon's width.
+
 	BRANCH
 	if(NiagaraRibbonVF.FacingDataOffset != -1)
 	{
@@ -571,36 +574,32 @@ void GetNiagaraParticleTangents(in FInterpVtxData InterpVtxData, in out FVertexF
 		Intermediates.Facing = normalize(ResolvedView.WorldCameraOrigin - Intermediates.Position);
 	}
 
-	BRANCH
-    if (FacingMode == 2 /* CustomSide */)
-	{
-		Intermediates.SideVector = -Intermediates.Facing;
-		Intermediates.Right = normalize(Intermediates.SideVector - dot(Intermediates.SideVector, Intermediates.Direction * Intermediates.Direction));
-		Intermediates.Facing = normalize(cross(Intermediates.Direction, Intermediates.Right));
-	}
-	else
-	{
-		Intermediates.Right = normalize(cross(Intermediates.Facing, Intermediates.Direction));
-		Intermediates.SideVector = Intermediates.Right;
-	}
+	Intermediates.Right = normalize(cross(Intermediates.Facing, Intermediates.Direction));
 
 	ISOLATE
 	{
 		Intermediates.Up = cross(Intermediates.Right, Intermediates.Facing);
 	}
-	
+
+	if (FacingMode == 2 /* CustomSide */)
+	{
+		const float3 Temp = Intermediates.Right;
+		Intermediates.Right = -Intermediates.Facing;
+		Intermediates.Facing = Temp;
+	}
+
 	BRANCH
 	if (NiagaraRibbonVF.TwistDataOffset != -1)
 	{
 		const float Twist0 = GetFloat(NiagaraRibbonVF.TwistDataOffset, InterpVtxData.ParticleId0);
 		const float Twist1 = GetFloat(NiagaraRibbonVF.TwistDataOffset, InterpVtxData.ParticleId1);
 		float SinTwist, CosTwist;
-		sincos(lerp(Twist0, Twist1, InterpVtxData.Alpha), SinTwist, CosTwist);
+		sincos(lerp(Twist0, Twist1, InterpVtxData.Alpha) /* + ResolvedView.RealTime * .5 */, SinTwist, CosTwist);
 
 		// Apply a rotation matrix to up and right.
 		float3 Right = Intermediates.Right; // Back it up for the transform.
-		Intermediates.Right = CosTwist * Right + SinTwist * Intermediates.Up;
-		Intermediates.Up = -SinTwist * Right + CosTwist * Intermediates.Up;
+		Intermediates.Right = CosTwist * Right + SinTwist * Intermediates.Facing;
+		Intermediates.Facing = -SinTwist * Right + CosTwist * Intermediates.Facing;
 	}
 }
 
@@ -610,7 +609,7 @@ float4 CalcWorldPosition(FVertexFactoryInput Input, FVertexFactoryIntermediates
 	{
 		// expand and align
 		int Sgn = (((int)Input.InterpVtxId & 0x1) * 2) - 1;
-		float3 Pos = Intermediates.Position.xyz + Intermediates.SideVector * Sgn * (Intermediates.RibbonWidth.xxx / 2);
+		float3 Pos = Intermediates.Position.xyz + Intermediates.Right * Sgn * (Intermediates.RibbonWidth.xxx / 2);
 		// Intermediates are already in world space, only need to account for the translation around view origin.
 		return float4(Pos + ResolvedView.PreViewTranslation, 1);
 	}
diff --git a/Engine/Plugins/FX/Niagara/Source/Niagara/Private/NiagaraRendererRibbons.cpp b/Engine/Plugins/FX/Niagara/Source/Niagara/Private/NiagaraRendererRibbons.cpp
--- a/Engine/Plugins/FX/Niagara/Source/Niagara/Private/NiagaraRendererRibbons.cpp
+++ b/Engine/Plugins/FX/Niagara/Source/Niagara/Private/NiagaraRendererRibbons.cpp
@@ -11,14 +11,6 @@ DECLARE_CYCLE_STAT(TEXT("Render Ribbons [RT]"), STAT_NiagaraRenderRibbons, STATG
 
 DECLARE_CYCLE_STAT(TEXT("Genereate GPU Buffers"), STAT_NiagaraGenRibbonGpuBuffers, STATGROUP_Niagara);
 
-int32 GNiagaraRibbonCullTwistedStrips = 1;
-static FAutoConsoleVariableRef CVarNiagaraRibbonCullTwistedStrips(
-	TEXT("Niagara.Ribbon.CullTwistedStrips"),
-	GNiagaraRibbonCullTwistedStrips,
-	TEXT("Apply backface culling to remove twisted strips. Double indices. (default=1)"),
-	ECVF_Default
-);
-
 float GNiagaraRibbonTessellationAngle = 15.f * (2.f * PI) / 360.f; // Every 15 degrees
 static FAutoConsoleVariableRef CVarNiagaraRibbonTessellationAngle(
 	TEXT("Niagara.Ribbon.Tessellation.MinAngle"),
@@ -63,14 +55,6 @@ static FAutoConsoleVariableRef CVarNiagaraRibbonMinSegmentLength(
 	ECVF_Scalability
 );
 
-float GNiagaraRibbonMaxAveraging = .5f;
-static FAutoConsoleVariableRef CVarNiagaraRibbonMaxAveraging(
-	TEXT("Niagara.Ribbon.Tessellation.MaxAveraging"),
-	GNiagaraRibbonMaxAveraging,
-	TEXT("Averaging value between the max metrics and the current metrics. (default=50%)"),
-	ECVF_Scalability
-);
-
 static int32 GbEnableNiagaraRibbonRendering = 1;
 static FAutoConsoleVariableRef CVarEnableNiagaraRibbonRendering(
 	TEXT("fx.EnableNiagaraRibbonRendering"),
@@ -232,9 +216,9 @@ void FNiagaraRendererRibbons::CreateRenderThreadResources(NiagaraEmitterInstance
 	VertexFactory->InitResource();
 }
 
-void FNiagaraRendererRibbons::GenerateIndexBuffer(uint16* OutIndices, const TArray<int32>& SegmentData,  int32 InterpCount, bool bInvertOrder, bool bCullTwistedStrips)
+void FNiagaraRendererRibbons::GenerateIndexBuffer(uint16* OutIndices, const TArray<int32>& SegmentData,  int32 InterpCount, bool bInvertOrder)
 {
-	auto AddTriangleIndices = [&OutIndices, InterpCount, bCullTwistedStrips](int32 SegmentIndex)
+	auto AddTriangleIndices = [&OutIndices, InterpCount](int32 SegmentIndex)
 	{
 		for (int32 SubSegmentIndex = 0; SubSegmentIndex < InterpCount; ++SubSegmentIndex)
 		{
@@ -245,25 +229,7 @@ void FNiagaraRendererRibbons::GenerateIndexBuffer(uint16* OutIndices, const TArr
 			OutIndices[3] = BaseVertexIndex + 1;
 			OutIndices[4] = BaseVertexIndex + 3;
 			OutIndices[5] = BaseVertexIndex + 2;
-
-			if (bCullTwistedStrips)
-			{
-				// When the ribbon right vector abruptly changes direction, for example when a segment starts with a tangent going up, but ends with a tanging going down,
-				// the triangles defined above will cross each other. One will be clockwise and the other counter-clockwise. We add two triangle here that replaces the invalid one
-				// this prevents the crossing/twisting artefacts while putting a bit more pressure on triangle culling.
-
-				OutIndices[6] = BaseVertexIndex + 0;
-				OutIndices[7] = BaseVertexIndex + 2;
-				OutIndices[8] = BaseVertexIndex + 3;
-				OutIndices[9] = BaseVertexIndex + 0;
-				OutIndices[10] = BaseVertexIndex + 3;
-				OutIndices[11] = BaseVertexIndex + 1;
-				OutIndices += 12;
-			}
-			else
-			{
-				OutIndices += 6;
-			}
+			OutIndices += 6;
 		}
 	};
 
@@ -312,7 +278,6 @@ void FNiagaraRendererRibbons::GetDynamicMeshElements(const TArray<const FSceneVi
 #endif
 
 	const bool bIsWireframe = ViewFamily.EngineShowFlags.Wireframe;
-	const bool bCullTwistedStrips = FacingMode == ENiagaraRibbonFacingMode::Screen && GNiagaraRibbonCullTwistedStrips != 0;
 	FMaterialRenderProxy* MaterialRenderProxy = DynamicDataRibbon->Material;
 	check(MaterialRenderProxy);
 
@@ -364,7 +329,7 @@ void FNiagaraRendererRibbons::GetDynamicMeshElements(const TArray<const FSceneVi
 	case ENiagaraRibbonTessellationMode::Automatic:
 		break;
 	case ENiagaraRibbonTessellationMode::Custom:
-		TessellationFactor = CustomTessellationFactor;
+		TessellationFactor = FMath::Min<int32>(TessellationFactor, CustomTessellationFactor); // Don't allow factors bigger than the platform limit.
 		bUseConstantFactor = bCustomUseConstantFactor;
 		TessellationMinAngle = CustomTessellationMinAngle;
 		ScreenPercentage = bCustomUseScreenSpace && !bUseConstantFactor ? GNiagaraRibbonTessellationScreenPercentage : 0.f;
@@ -415,7 +380,7 @@ void FNiagaraRendererRibbons::GetDynamicMeshElements(const TArray<const FSceneVi
 				NumSegments *= SegmentTessellation;
 			}
 
-			const int32 NumberOfPrimitives = NumSegments * 2 * (bCullTwistedStrips ? 2 : 1);
+			const int32 NumberOfPrimitives = NumSegments * 2;
 
 			// Figure out whether start is closer to the view plane than end
 			// TODO : This doesn't work with multi-ribbons.
@@ -426,7 +391,7 @@ void FNiagaraRendererRibbons::GetDynamicMeshElements(const TArray<const FSceneVi
 
 			// Copy the vertex data over.
 			FGlobalDynamicIndexBuffer::FAllocation DynamicIndexAllocation = DynamicIndexBuffer.Allocate(NumberOfPrimitives * 3, sizeof(int16));
-			GenerateIndexBuffer((uint16*)DynamicIndexAllocation.Buffer, DynamicDataRibbon->SegmentData, SegmentTessellation, bInvertOrder, bCullTwistedStrips); 
+			GenerateIndexBuffer((uint16*)DynamicIndexAllocation.Buffer, DynamicDataRibbon->SegmentData, SegmentTessellation, bInvertOrder); 
 
 			FNiagaraMeshCollectorResourcesRibbon& CollectorResources = Collector.AllocateOneFrameResource<FNiagaraMeshCollectorResourcesRibbon>();
 			FNiagaraRibbonUniformParameters PerViewUniformParameters;// = UniformParameters;
@@ -512,7 +477,7 @@ void FNiagaraRendererRibbons::GetDynamicMeshElements(const TArray<const FSceneVi
 			MeshBatch.CastShadow = SceneProxy->CastsDynamicShadow();
 			MeshBatch.bUseAsOccluder = false;
 			MeshBatch.ReverseCulling = SceneProxy->IsLocalToWorldDeterminantNegative();
-			MeshBatch.bDisableBackfaceCulling = !bCullTwistedStrips;
+			MeshBatch.bDisableBackfaceCulling = true;
 			MeshBatch.Type = PT_TriangleList;
 			MeshBatch.DepthPriorityGroup = SceneProxy->GetDepthPriorityGroup(View);
 			MeshBatch.bCanApplyViewModeOverrides = true;
@@ -670,10 +635,11 @@ FNiagaraDynamicDataBase* FNiagaraRendererRibbons::GenerateDynamicData(const FNia
 
 	TArray<int32>& SegmentData = DynamicData->SegmentData;
 	float TotalSegmentLength = 0;
-	float TotalSegmentAngle = 0;
-	float TotalTwistAngle = 0;
-	float TotalWidth = 0;
-	int NumTotalSamples = 0;
+	// weighted sum based on the segment length :
+	float AverageSegmentLength = 0; 
+	float AverageSegmentAngle = 0;
+	float AverageTwistAngle = 0;
+	float AverageWidth = 0;
 
 	if (Properties->RibbonIdBinding.DataSetVariable.GetType() == FNiagaraTypeDefinition::GetIDDef())
 	{
@@ -693,8 +659,7 @@ FNiagaraDynamicDataBase* FNiagaraRendererRibbons::GenerateDynamicData(const FNia
 
 	auto AddRibbonVerts = [&](TArray<int32>& RibbonIndices, uint32 RibbonIndex)
 	{
-		int32 StartIndex = DynamicData->SortedIndices.Num();
-		int32 NumIndices = RibbonIndices.Num();
+		const int32 StartIndex = DynamicData->SortedIndices.Num();
 
 		float TotalDistance = 0.0f;
 
@@ -729,7 +694,6 @@ FNiagaraDynamicDataBase* FNiagaraRendererRibbons::GenerateDynamicData(const FNia
 				DynamicData->SortedIndices.Add(RibbonIndices[0]);
 				DynamicData->TangentAndDistances.Add(FVector4(LastToCurrVec.X, LastToCurrVec.Y, LastToCurrVec.Z, 0));
 				DynamicData->MultiRibbonIndices.Add(RibbonIndex);
-
 				break;
 			}
 			else
@@ -773,10 +737,10 @@ FNiagaraDynamicDataBase* FNiagaraRendererRibbons::GenerateDynamicData(const FNia
 
 				// Assumed equal to dot(Tangent, CurrToNextVec)
 				TotalSegmentLength += CurrToNextSize;
-				TotalSegmentAngle += AcosFast(FVector::DotProduct(LastToCurrVec, CurrToNextVec));
-				TotalTwistAngle += FMath::Abs(NextTwist - LastTwist);
-				TotalWidth += LastWidth;
-				++NumTotalSamples;
+				AverageSegmentLength += CurrToNextSize * CurrToNextSize;
+				AverageSegmentAngle += CurrToNextSize * AcosFast(FVector::DotProduct(LastToCurrVec, CurrToNextVec));
+				AverageTwistAngle += FMath::Abs(NextTwist - LastTwist) * CurrToNextSize;
+				AverageWidth += LastWidth * CurrToNextSize;
 
 				// Move to next segment.
 				CurrentIndex = NextIndex;
@@ -902,25 +866,28 @@ FNiagaraDynamicDataBase* FNiagaraRendererRibbons::GenerateDynamicData(const FNia
 		}
 	}
 
-	if (NumTotalSamples > 0)
+	if (TotalSegmentLength > 0)
 	{
-		const float OneOverSampleCount = 1.f / (float)NumTotalSamples;
-		const float AverageSegmentAngle = TotalSegmentAngle * OneOverSampleCount;
-		const float AverageSegmentCurvature = TotalSegmentLength * OneOverSampleCount * .5 / (FMath::Max(SMALL_NUMBER, FMath::Abs(FMath::Sin(AverageSegmentAngle))));
+		// Blend the result between the last frame tessellation factors and the current frame base on the total length of all segments.
+		// This is only used to increase the tessellation value of the current frame data to prevent glitches where tessellation is significantly changin between frames.
+		const float OneOverTotalSegmentLength = 1.f / FMath::Max(1.f, TotalSegmentLength);
+		const float AveragingFactor = TessellationTotalSegmentLength / (TotalSegmentLength + TessellationTotalSegmentLength);
+		TessellationTotalSegmentLength = TotalSegmentLength;
+
+		AverageSegmentAngle *= OneOverTotalSegmentLength;
+		AverageSegmentLength *= OneOverTotalSegmentLength;
+		const float AverageSegmentCurvature = AverageSegmentLength / (FMath::Max(SMALL_NUMBER, FMath::Abs(FMath::Sin(AverageSegmentAngle))));
 
-		// Accumulate the max so that tessellation is ever increasing and stabilizes over time.
-		// This is to avoid the spawning and unspawning of ribbons to constantly generate different tessellation factors.
-		
-		TessellationAngle = FMath::Lerp<float>(AverageSegmentAngle, FMath::Max(TessellationAngle, AverageSegmentAngle), GNiagaraRibbonMaxAveraging);
-		TessellationCurvature = FMath::Lerp<float>(AverageSegmentCurvature, FMath::Max(TessellationCurvature, AverageSegmentCurvature), GNiagaraRibbonMaxAveraging);
+		TessellationAngle = FMath::Lerp<float>(AverageSegmentAngle, FMath::Max(TessellationAngle, AverageSegmentAngle), AveragingFactor);
+		TessellationCurvature = FMath::Lerp<float>(AverageSegmentCurvature, FMath::Max(TessellationCurvature, AverageSegmentCurvature), AveragingFactor);
 
 		if (bHasTwist)
 		{
-			const float AverageTwistAngle = TotalTwistAngle * OneOverSampleCount;
-			const float AverageTwistCurvature = TotalWidth * OneOverSampleCount;
+			AverageTwistAngle *= OneOverTotalSegmentLength;
+			AverageWidth *= OneOverTotalSegmentLength;
 			
-			TessellationTwistAngle = FMath::Lerp<float>(AverageTwistAngle, FMath::Max(TessellationTwistAngle, AverageTwistAngle), GNiagaraRibbonMaxAveraging);
-			TessellationTwistCurvature = FMath::Lerp<float>(AverageTwistCurvature, FMath::Max(TessellationTwistCurvature, AverageTwistCurvature), GNiagaraRibbonMaxAveraging);
+			TessellationTwistAngle = FMath::Lerp<float>(AverageTwistAngle, FMath::Max(TessellationTwistAngle, AverageTwistAngle), AveragingFactor);
+			TessellationTwistCurvature = FMath::Lerp<float>(AverageWidth, FMath::Max(TessellationTwistCurvature, AverageWidth), AveragingFactor);
 		}
 	}
 	else // Reset the metrics when the ribbons are reset.
@@ -929,6 +896,7 @@ FNiagaraDynamicDataBase* FNiagaraRendererRibbons::GenerateDynamicData(const FNia
 		TessellationCurvature = 0;
 		TessellationTwistAngle = 0;
 		TessellationTwistCurvature = 0;
+		TessellationTotalSegmentLength = 0;
 	}
 
 	CPUTimeMS = VertexDataTimer.GetElapsedMilliseconds();
diff --git a/Engine/Plugins/FX/Niagara/Source/Niagara/Public/NiagaraRendererRibbons.h b/Engine/Plugins/FX/Niagara/Source/Niagara/Public/NiagaraRendererRibbons.h
--- a/Engine/Plugins/FX/Niagara/Source/Niagara/Public/NiagaraRendererRibbons.h
+++ b/Engine/Plugins/FX/Niagara/Source/Niagara/Public/NiagaraRendererRibbons.h
@@ -33,7 +33,7 @@ public:
 
 	FORCEINLINE void AddDynamicParam(TArray<FNiagaraRibbonVertexDynamicParameter>& ParamData, const FVector4& DynamicParam);
 protected:
-	static void GenerateIndexBuffer(uint16* OutIndices, const TArray<int32>& SegmentData, int32 MaxTessellation, bool bInvertOrder, bool bCullTwistedStrips);
+	static void GenerateIndexBuffer(uint16* OutIndices, const TArray<int32>& SegmentData, int32 MaxTessellation, bool bInvertOrder);
 
 private:
 	class FNiagaraRibbonVertexFactory *VertexFactory;
@@ -72,10 +72,12 @@ private:
 
 	// Average curvature of the segments.
 	mutable float TessellationAngle = 0;
-	// Average angle of the curvature of the segments (in radian).
+	// Average curvature of the segments (computed from the segment angle in radian).
 	mutable float TessellationCurvature = 0;
 	// Average twist of the segments.
 	mutable float TessellationTwistAngle = 0;
 	// Average twist curvature of the segments.
 	mutable float TessellationTwistCurvature = 0;
+	// Average twist curvature of the segments.
+	mutable float TessellationTotalSegmentLength = 0;
 };
