Adds GetSpecificSocketTransform to UNiagaraDataInterfaceSkeletalMesh. This Data Interface function returns the FTransform components of the current USkeletalMeshComponent or USkeletalMesh.

An optional ComponentToWorld transform can be applied.

The GPU version is not yet implemented.

diff --git a/Engine/Plugins/FX/Niagara/Source/Niagara/Classes/NiagaraDataInterfaceSkeletalMesh.h b/Engine/Plugins/FX/Niagara/Source/Niagara/Classes/NiagaraDataInterfaceSkeletalMesh.h
--- a/Engine/Plugins/FX/Niagara/Source/Niagara/Classes/NiagaraDataInterfaceSkeletalMesh.h
+++ b/Engine/Plugins/FX/Niagara/Source/Niagara/Classes/NiagaraDataInterfaceSkeletalMesh.h
@@ -675,6 +675,8 @@ public:
 
 	void GetSpecificSocketBoneAt(FVectorVMContext& Context);
 
+	void GetSpecificSocketTransform(FVectorVMContext& Context);
+
 	void RandomSpecificSocketBone(FVectorVMContext& Context);
 		
 	// End of Direct Bone + Socket Sampling
@@ -712,6 +714,7 @@ public:
 	static const FName RandomSpecificSocketBoneName;
 	static const FName GetSpecificSocketCountName;
 	static const FName GetSpecificSocketBoneAtName;
+	static const FName GetSpecificSocketTransformName;
 
 	// Vertex Sampling
 	static const FName IsValidVertexName;
diff --git a/Engine/Plugins/FX/Niagara/Source/Niagara/Private/NDISkeletalMesh_BoneSampling.cpp b/Engine/Plugins/FX/Niagara/Source/Niagara/Private/NDISkeletalMesh_BoneSampling.cpp
--- a/Engine/Plugins/FX/Niagara/Source/Niagara/Private/NDISkeletalMesh_BoneSampling.cpp
+++ b/Engine/Plugins/FX/Niagara/Source/Niagara/Private/NDISkeletalMesh_BoneSampling.cpp
@@ -19,6 +19,7 @@ DEFINE_NDI_DIRECT_FUNC_BINDER(UNiagaraDataInterfaceSkeletalMesh, IsValidBone)
 DEFINE_NDI_DIRECT_FUNC_BINDER(UNiagaraDataInterfaceSkeletalMesh, GetSpecificBoneAt)
 
 DEFINE_NDI_DIRECT_FUNC_BINDER(UNiagaraDataInterfaceSkeletalMesh, GetSpecificSocketBoneAt)
+DEFINE_NDI_DIRECT_FUNC_BINDER(UNiagaraDataInterfaceSkeletalMesh, GetSpecificSocketTransform)
 
 const FName FSkeletalMeshInterfaceHelper::GetSkinnedBoneDataName("GetSkinnedBoneData");
 const FName FSkeletalMeshInterfaceHelper::GetSkinnedBoneDataWSName("GetSkinnedBoneDataWS");
@@ -30,6 +31,7 @@ const FName FSkeletalMeshInterfaceHelper::GetSpecificBoneCountName("GetSpecificB
 const FName FSkeletalMeshInterfaceHelper::GetSpecificBoneAtName("GetSpecificBone");
 const FName FSkeletalMeshInterfaceHelper::RandomSpecificSocketBoneName("RandomSpecificSocketBone");
 const FName FSkeletalMeshInterfaceHelper::GetSpecificSocketCountName("GetSpecificSocketCount");
+const FName FSkeletalMeshInterfaceHelper::GetSpecificSocketTransformName("GetSpecificSocketTransform");
 const FName FSkeletalMeshInterfaceHelper::GetSpecificSocketBoneAtName("GetSpecificSocketBone");
 
 void UNiagaraDataInterfaceSkeletalMesh::GetSkeletonSamplingFunctions(TArray<FNiagaraFunctionSignature>& OutFunctions)
@@ -194,6 +196,23 @@ void UNiagaraDataInterfaceSkeletalMesh::GetSkeletonSamplingFunctions(TArray<FNia
 		Sig.bRequiresContext = false;
 #if WITH_EDITORONLY_DATA
 		Sig.Description = LOCTEXT("GetSpecificSocketBoneAtDesc", "Gets the bone for the socket at the passed index in the DI's specfic socket list.");
+#endif
+		OutFunctions.Add(Sig);
+	}
+
+	{
+		FNiagaraFunctionSignature Sig;
+		Sig.Name = FSkeletalMeshInterfaceHelper::GetSpecificSocketTransformName;
+		Sig.Inputs.Add(FNiagaraVariable(FNiagaraTypeDefinition(GetClass()), TEXT("SkeletalMesh")));
+		Sig.Inputs.Add(FNiagaraVariable(FNiagaraTypeDefinition::GetIntDef(), TEXT("Socket Index")));
+		Sig.Inputs.Add(FNiagaraVariable(FNiagaraTypeDefinition::GetBoolDef(), TEXT("Apply World Transform")));
+		Sig.Outputs.Add(FNiagaraVariable(FNiagaraTypeDefinition::GetVec3Def(), TEXT("Socket Translation")));
+		Sig.Outputs.Add(FNiagaraVariable(FNiagaraTypeDefinition::GetQuatDef(), TEXT("Socket Rotation")));
+		Sig.Outputs.Add(FNiagaraVariable(FNiagaraTypeDefinition::GetVec3Def(), TEXT("Socket Scale")));
+		Sig.bMemberFunction = true;
+		Sig.bRequiresContext = false;
+#if WITH_EDITORONLY_DATA
+		Sig.Description = LOCTEXT("GetSpecificSocketTransformDesc", "Gets the transform for the socket at the passed index in the DI's specfic socket list. If the Source component is set it will respect the Relative Transform Space as well..");
 #endif
 		OutFunctions.Add(Sig);
 	}
@@ -262,6 +281,11 @@ void UNiagaraDataInterfaceSkeletalMesh::BindSkeletonSamplingFunction(const FVMEx
 		check(BindingInfo.GetNumInputs() == 2 && BindingInfo.GetNumOutputs() == 1);
 		NDI_FUNC_BINDER(UNiagaraDataInterfaceSkeletalMesh, GetSpecificSocketBoneAt)::Bind(this, OutFunc);
 	}
+	else if (BindingInfo.Name == FSkeletalMeshInterfaceHelper::GetSpecificSocketTransformName)
+	{
+		check(BindingInfo.GetNumInputs() == 3 && BindingInfo.GetNumOutputs() == 10);
+		NDI_FUNC_BINDER(UNiagaraDataInterfaceSkeletalMesh, GetSpecificSocketTransform)::Bind(this, OutFunc);
+	}
 }
 
 
@@ -598,6 +622,75 @@ void UNiagaraDataInterfaceSkeletalMesh::GetSpecificSocketBoneAt(FVectorVMContext
 	}
 }
 
+void UNiagaraDataInterfaceSkeletalMesh::GetSpecificSocketTransform(FVectorVMContext& Context)
+{
+	SCOPE_CYCLE_COUNTER(STAT_NiagaraSkel_Bone_Sample);
+
+	VectorVM::FExternalFuncInputHandler<int32> SocketParam(Context);
+	VectorVM::FExternalFuncInputHandler<int32> ApplyWorldTransform(Context);
+	
+	VectorVM::FUserPtrHandler<FNDISkeletalMesh_InstanceData> InstData(Context);
+
+	VectorVM::FExternalFuncRegisterHandler<float> OutSocketTranslateX(Context);
+	VectorVM::FExternalFuncRegisterHandler<float> OutSocketTranslateY(Context);
+	VectorVM::FExternalFuncRegisterHandler<float> OutSocketTranslateZ(Context);
+	
+	VectorVM::FExternalFuncRegisterHandler<float> OutSocketRotationX(Context);
+	VectorVM::FExternalFuncRegisterHandler<float> OutSocketRotationY(Context);
+	VectorVM::FExternalFuncRegisterHandler<float> OutSocketRotationZ(Context);
+	VectorVM::FExternalFuncRegisterHandler<float> OutSocketRotationW(Context);
+
+	VectorVM::FExternalFuncRegisterHandler<float> OutSocketScaleX(Context);
+	VectorVM::FExternalFuncRegisterHandler<float> OutSocketScaleY(Context);
+	VectorVM::FExternalFuncRegisterHandler<float> OutSocketScaleZ(Context);
+ 
+ 	//USkeletalMesh* Mesh = InstData->Mesh;
+	USkeletalMeshComponent* SkelComp = Cast<USkeletalMeshComponent>(InstData->Component.Get());
+	// UE_LOG(LogNiagara, Warning, TEXT("SkelComp: %p, Mesh: %p, Num: %d"), SkelComp, InstData->Mesh, Context.NumInstances);
+
+	FTransform ComponentToWorld;
+	if (SkelComp != nullptr) {
+		ComponentToWorld = SkelComp->GetComponentToWorld();
+	}
+
+	const TArray<FTransform>& CurrentSpecificSockets = InstData->GetSpecificSocketsCurrBuffer();
+
+	for (int32 i = 0; i < Context.NumInstances; ++i)
+	{//
+		const int32 SocketIndex = FMath::Clamp(SocketParam.GetAndAdvance(), 0, CurrentSpecificSockets.Num()-1);
+		const int32 ShouldApplyWorldTransform = ApplyWorldTransform.GetAndAdvance();
+		
+		FTransform CurrentTransform = CurrentSpecificSockets[SocketIndex];
+		if (ShouldApplyWorldTransform == -1)
+		{
+			CurrentTransform = CurrentTransform * ComponentToWorld;
+		}
+
+		FQuat Rotation = CurrentTransform.GetRotation();
+		FVector Translation = CurrentTransform.GetTranslation();
+		FVector Scale = CurrentTransform.GetScale3D();
+
+		*OutSocketTranslateX.GetDestAndAdvance() = Translation.X;
+		*OutSocketTranslateY.GetDestAndAdvance() = Translation.Y;
+		*OutSocketTranslateZ.GetDestAndAdvance() = Translation.Z;
+
+		*OutSocketRotationX.GetDestAndAdvance() = Rotation.X;
+		*OutSocketRotationY.GetDestAndAdvance() = Rotation.Y;
+		*OutSocketRotationZ.GetDestAndAdvance() = Rotation.Z;
+		*OutSocketRotationW.GetDestAndAdvance() = Rotation.W;
+
+		*OutSocketScaleX.GetDestAndAdvance() = Scale.X;
+		*OutSocketScaleY.GetDestAndAdvance() = Scale.Y;
+		*OutSocketScaleZ.GetDestAndAdvance() = Scale.Z;
+			
+		// UE_LOG(LogNiagara, Warning, TEXT("Enum: %d, translation: %f %f %f, Rotation: %f %f %f %f, scale: %f %f %f"), 
+		// 		  ShouldApplyWorldTransform,
+		// 		  Translation.X, Translation.Y, Translation.Z,
+		// 		  Rotation.X, Rotation.Y, Rotation.Z, Rotation.W,
+		// 		  Scale.X, Scale.Y, Scale.Z);
+	}
+}
+
 void UNiagaraDataInterfaceSkeletalMesh::RandomSpecificSocketBone(FVectorVMContext& Context)
 {
 	SCOPE_CYCLE_COUNTER(STAT_NiagaraSkel_Bone_Sample);
diff --git a/Engine/Plugins/FX/Niagara/Source/Niagara/Private/NiagaraDataInterfaceSkeletalMesh.cpp b/Engine/Plugins/FX/Niagara/Source/Niagara/Private/NiagaraDataInterfaceSkeletalMesh.cpp
--- a/Engine/Plugins/FX/Niagara/Source/Niagara/Private/NiagaraDataInterfaceSkeletalMesh.cpp
+++ b/Engine/Plugins/FX/Niagara/Source/Niagara/Private/NiagaraDataInterfaceSkeletalMesh.cpp
@@ -2023,7 +2023,14 @@ bool UNiagaraDataInterfaceSkeletalMesh::GetFunctionHLSL(const FName& DefinitionF
 	{
 		static const TCHAR* FormatSample = TEXT("void {InstanceFunctionName} (in int SocketIndex, out int Bone) { {GetDISkelMeshContextName} DISkelMesh_GetSpecificSocketBoneAt(DIContext, SocketIndex, Bone); }");
 		OutHLSL += FString::Format(FormatSample, ArgsSample);
-	}
+	} 
+	else if (DefinitionFunctionName == FSkeletalMeshInterfaceHelper::GetSpecificSocketTransformName)
+	{
+		// TODO: This just returns the Identity transform.
+		// TODO: Make this work on the GPU?
+		static const TCHAR* FormatSample = TEXT("void {InstanceFunctionName} (in int SocketIndex, in int bShouldApplyTransform, out float3 Translation, out float4 Rotation, out float4 Scale) { Translation = float3(0.0, 0.0, 0.0); Rotation = float4(0.0, 0.0, 0.0, 1.0); Scale = float3(1.0, 1.0, 1.0); }");
+		OutHLSL += FString::Format(FormatSample, ArgsSample);
+	} 
 	else if (DefinitionFunctionName == FSkeletalMeshInterfaceHelper::RandomSpecificSocketBoneName)
 	{
 		static const TCHAR* FormatSample = TEXT("void {InstanceFunctionName} (out int SocketBone) { {GetDISkelMeshContextName} DISkelMesh_RandomSpecificSocketBone(DIContext, SocketBone); }");
