- Remove function / member that is no longer used

diff --git a/Engine/Plugins/FX/Niagara/Source/Niagara/Public/NiagaraSystemSimulation.h b/Engine/Plugins/FX/Niagara/Source/Niagara/Public/NiagaraSystemSimulation.h
--- a/Engine/Plugins/FX/Niagara/Source/Niagara/Public/NiagaraSystemSimulation.h
+++ b/Engine/Plugins/FX/Niagara/Source/Niagara/Public/NiagaraSystemSimulation.h
@@ -189,10 +189,7 @@ public:
 	FNiagaraScriptExecutionContext& GetSpawnExecutionContext() { return SpawnExecContext; }
 	FNiagaraScriptExecutionContext& GetUpdateExecutionContext() { return UpdateExecContext; }
 
-	void TransitionToDeferredDeletionQueue(TUniquePtr< FNiagaraSystemInstance>& InPtr);
-
 protected:
-
 	/** Does any prep work for system simulation such as pulling instance parameters into a dataset. */
 	void PrepareForSystemSimulate(FNiagaraSystemSimulationTickContext& Context);
 	/** Runs the system spawn script for new system instances. */
@@ -295,6 +292,4 @@ protected:
 
 	/** Current tick batch we're filling ready for processing, potentially in an async task. */
 	FNiagaraSystemTickBatch TickBatch;
-
-	TArray<TUniquePtr< FNiagaraSystemInstance> > DeferredDeletionQueue;
 };
