- Remove incorrect warning about skinning mode being set to none when using skeleton sampling
- This warning has some merit with triangle / vertex sampling functions, but is only relevant to specific functions so it has been removed to avoid confusion

diff --git a/Engine/Plugins/FX/Niagara/Source/Niagara/Private/NiagaraDataInterfaceSkeletalMesh.cpp b/Engine/Plugins/FX/Niagara/Source/Niagara/Private/NiagaraDataInterfaceSkeletalMesh.cpp
--- a/Engine/Plugins/FX/Niagara/Source/Niagara/Private/NiagaraDataInterfaceSkeletalMesh.cpp
+++ b/Engine/Plugins/FX/Niagara/Source/Niagara/Private/NiagaraDataInterfaceSkeletalMesh.cpp
@@ -1564,12 +1564,6 @@ void UNiagaraDataInterfaceSkeletalMesh::GetVMExternalFunction(const FVMExternalF
 	BindSkeletonSamplingFunction(BindingInfo, InstData, OutFunc);
 	if (OutFunc.IsBound())
 	{
-#if WITH_EDITOR
-		if ( SkinningMode == ENDISkeletalMesh_SkinningMode::None )
-		{
-			UE_LOG(LogNiagara, Warning, TEXT("Skeletal Mesh Data Interface is trying to use skeleton sampling but skinning mode is none. Interface: %s"), *GetFullName());
-		}
-#endif // WITH_EDITOR
 		return;
 	}
 
@@ -1592,7 +1586,6 @@ void UNiagaraDataInterfaceSkeletalMesh::GetVMExternalFunction(const FVMExternalF
 
 	// Bind vertex sampling function
 	BindVertexSamplingFunction(BindingInfo, InstData, OutFunc);
-
 	if (OutFunc.IsBound())
 	{
 		if (!InstData->bAllowCPUMeshDataAccess)
