- Remove single tick from instance batcher, instead when you disable overlapping we use the same method but each tick is it's own simulation pass.  Disabling overlapping is inefficient and should only be used for debugging (i.e. transition related issues).
- Added fx.NiagaraSubmitCommands to control if we submit commands, currently set to 0.  On some platforms this can save 0.2ms on the RHI thread.
- Remove a duplicate transition that looks like a bad merge

diff --git a/Engine/Plugins/FX/Niagara/Source/Niagara/Classes/NiagaraDataSet.h b/Engine/Plugins/FX/Niagara/Source/Niagara/Classes/NiagaraDataSet.h
--- a/Engine/Plugins/FX/Niagara/Source/Niagara/Classes/NiagaraDataSet.h
+++ b/Engine/Plugins/FX/Niagara/Source/Niagara/Classes/NiagaraDataSet.h
@@ -146,7 +146,6 @@ public:
 
 	FORCEINLINE TArray<int32>& GetIDTable() { return IDToIndexTable; }
 
-	template<bool bDoResourceTransitions>
 	void SetShaderParams(class FNiagaraShader *Shader, FRHICommandList &CommandList, bool bInput);
 	void UnsetShaderParams(class FNiagaraShader *Shader, FRHICommandList &CommandList);
 
diff --git a/Engine/Plugins/FX/Niagara/Source/Niagara/Classes/NiagaraEmitterInstanceBatcher.h b/Engine/Plugins/FX/Niagara/Source/Niagara/Classes/NiagaraEmitterInstanceBatcher.h
--- a/Engine/Plugins/FX/Niagara/Source/Niagara/Classes/NiagaraEmitterInstanceBatcher.h
+++ b/Engine/Plugins/FX/Niagara/Source/Niagara/Classes/NiagaraEmitterInstanceBatcher.h
@@ -117,7 +117,6 @@ public:
 	void SetDataInterfaceParameters(const TArray<FNiagaraDataInterfaceProxy*>& DataInterfaceProxies, FNiagaraShader* Shader, FRHICommandList &RHICmdList, const FNiagaraComputeInstanceData* Instance, const FNiagaraGPUSystemTick& Tick, uint32 ShaderStageIndex) const;
 	void UnsetDataInterfaceParameters(const TArray<FNiagaraDataInterfaceProxy*>& DataInterfaceProxies, FNiagaraShader* Shader, FRHICommandList& RHICmdList, const FNiagaraComputeInstanceData* Instance, const FNiagaraGPUSystemTick& Tick) const;
 
-	template<bool bDoResourceTransitions>
 	void Run(	const FNiagaraGPUSystemTick& Tick,
 				const FNiagaraComputeInstanceData* Instance, 
 				uint32 UpdateStartInstance, 
@@ -151,13 +150,10 @@ public:
 	void PostStageInterface(const FNiagaraGPUSystemTick& Tick, FNiagaraComputeInstanceData *Instance, FRHICommandList &RHICmdList, FNiagaraShader* ComputeShader, const uint32 ShaderStageIndex) const;
 
 	/** Run the dispatch over multiple stages */
-	template<bool bDoResourceTransitions>
 	void DispatchMultipleStages(const FNiagaraGPUSystemTick& Tick, FNiagaraComputeInstanceData *Instance, FRHICommandList &RHICmdList, FRHIUniformBuffer* ViewUniformBuffer, FNiagaraShader* ComputeShader) const;
 
 private:
-
 	void ExecuteAll(FRHICommandList &RHICmdList, FRHIUniformBuffer* ViewUniformBuffer, bool bSetReadback);
-	void TickSingle(const FNiagaraGPUSystemTick& Tick, FNiagaraComputeInstanceData *Instance, FRHICommandList &RHICmdList, FRHIUniformBuffer* ViewUniformBuffer, bool bSetReadback) const;
 	void ResizeBuffersAndGatherResources(FOverlappableTicks& OverlappableTick, FRHICommandList& RHICmdList, FNiagaraBufferArray& DestDataBuffers, FNiagaraBufferArray& CurrDataBuffers, FNiagaraBufferArray& DestBufferIntFloat, FNiagaraBufferArray& CurrBufferIntFloat);
 	void DispatchAllOnCompute(FOverlappableTicks& OverlappableTick, FRHICommandList& RHICmdList, FRHIUniformBuffer* ViewUniformBuffer, FNiagaraBufferArray& DestDataBuffers, FNiagaraBufferArray& CurrDataBuffers, FNiagaraBufferArray& DestBufferIntFloat, FNiagaraBufferArray& CurrBufferIntFloat, bool bSetReadback);
 
diff --git a/Engine/Plugins/FX/Niagara/Source/Niagara/Private/NiagaraDataSet.cpp b/Engine/Plugins/FX/Niagara/Source/Niagara/Private/NiagaraDataSet.cpp
--- a/Engine/Plugins/FX/Niagara/Source/Niagara/Private/NiagaraDataSet.cpp
+++ b/Engine/Plugins/FX/Niagara/Source/Niagara/Private/NiagaraDataSet.cpp
@@ -970,7 +970,6 @@ bool FNiagaraDataBuffer::AppendToRegisterTable(uint8** Registers, int32& NumRegi
 
 /////////////////////////////////////////////////////////////////////////
 
-template<bool bDoResourceTransitions>
 void FNiagaraDataBuffer::SetShaderParams(FNiagaraShader *Shader, FRHICommandList &CommandList, bool bInput)
 {
 	check(IsInRenderingThread());
@@ -978,10 +977,6 @@ void FNiagaraDataBuffer::SetShaderParams(FNiagaraShader *Shader, FRHICommandList
 	{
 		if (Shader->FloatInputBufferParam.IsBound())
 		{
-			if (bDoResourceTransitions)
-			{
-				CommandList.TransitionResource(EResourceTransitionAccess::EReadable, EResourceTransitionPipeline::EComputeToCompute, GetGPUBufferFloat().UAV);
-			}
 			if (GetNumInstancesAllocated() > 0)
 			{
 				CommandList.SetShaderResourceViewParameter(Shader->GetComputeShader(), Shader->FloatInputBufferParam.GetBaseIndex(), GetGPUBufferFloat().SRV);
@@ -994,10 +989,6 @@ void FNiagaraDataBuffer::SetShaderParams(FNiagaraShader *Shader, FRHICommandList
 
 		if (Shader->IntInputBufferParam.IsBound())
 		{
-			if (bDoResourceTransitions)
-			{
-				CommandList.TransitionResource(EResourceTransitionAccess::EReadable, EResourceTransitionPipeline::EComputeToCompute, GetGPUBufferInt().UAV);
-			}
 			if (GetNumInstancesAllocated() > 0)
 			{
 				CommandList.SetShaderResourceViewParameter(Shader->GetComputeShader(), Shader->IntInputBufferParam.GetBaseIndex(), GetGPUBufferInt().SRV);
@@ -1018,19 +1009,11 @@ void FNiagaraDataBuffer::SetShaderParams(FNiagaraShader *Shader, FRHICommandList
 	{
 		if (Shader->FloatOutputBufferParam.IsUAVBound())
 		{
-			if (bDoResourceTransitions)
-			{
-				CommandList.TransitionResource(EResourceTransitionAccess::EWritable, EResourceTransitionPipeline::EGfxToCompute, GetGPUBufferFloat().UAV);
-			}
 			CommandList.SetUAVParameter(Shader->GetComputeShader(), Shader->FloatOutputBufferParam.GetUAVIndex(), GetGPUBufferFloat().UAV);
 		}
 
 		if (Shader->IntOutputBufferParam.IsUAVBound())
 		{
-			if (bDoResourceTransitions)
-			{
-				CommandList.TransitionResource(EResourceTransitionAccess::EWritable, EResourceTransitionPipeline::EGfxToCompute, GetGPUBufferInt().UAV);
-			}
 			CommandList.SetUAVParameter(Shader->GetComputeShader(), Shader->IntOutputBufferParam.GetUAVIndex(), GetGPUBufferInt().UAV);
 		}
 
@@ -1042,10 +1025,6 @@ void FNiagaraDataBuffer::SetShaderParams(FNiagaraShader *Shader, FRHICommandList
 	}
 }
 
-template void FNiagaraDataBuffer::SetShaderParams<true>(FNiagaraShader*, FRHICommandList&, bool);
-template void FNiagaraDataBuffer::SetShaderParams<false>(FNiagaraShader*, FRHICommandList&, bool);
-
-
 void FNiagaraDataBuffer::UnsetShaderParams(FNiagaraShader *Shader, FRHICommandList &RHICmdList)
 {
 	check(IsInRenderingThread());
diff --git a/Engine/Plugins/FX/Niagara/Source/Niagara/Private/NiagaraEmitterInstanceBatcher.cpp b/Engine/Plugins/FX/Niagara/Source/Niagara/Private/NiagaraEmitterInstanceBatcher.cpp
--- a/Engine/Plugins/FX/Niagara/Source/Niagara/Private/NiagaraEmitterInstanceBatcher.cpp
+++ b/Engine/Plugins/FX/Niagara/Source/Niagara/Private/NiagaraEmitterInstanceBatcher.cpp
@@ -42,7 +42,16 @@ int32 GNiagaraOverlapCompute = 1;
 static FAutoConsoleVariableRef CVarNiagaraUseAsyncCompute(
 	TEXT("fx.NiagaraOverlapCompute"),
 	GNiagaraOverlapCompute,
-	TEXT("If 1, use compute dispatch overlap for better performance. \n"),
+	TEXT("0 - Disable compute dispatch overlap, this will result in poor performance due to resource barriers between each dispatch call, but can be used to debug resource transition issues.\n")
+	TEXT("1 - (Default) Enable compute dispatch overlap where possible, this increases GPU utilization.\n"),
+	ECVF_Default
+);
+
+int32 GNiagaraSubmitCommands = 0;
+static FAutoConsoleVariableRef CVarNiagaraSubmitCommands(
+	TEXT("fx.NiagaraSubmitCommands"),
+	GNiagaraSubmitCommands,
+	TEXT("1 - (Default) Submit commands to the GPU once we have finished dispatching.\n"),
 	ECVF_Default
 );
 
@@ -279,7 +288,6 @@ void NiagaraEmitterInstanceBatcher::PostStageInterface(const FNiagaraGPUSystemTi
 	}
 }
 
-template<bool bDoResourceTransitions>
 void NiagaraEmitterInstanceBatcher::DispatchMultipleStages(const FNiagaraGPUSystemTick& Tick, FNiagaraComputeInstanceData *Instance, FRHICommandList &RHICmdList, FRHIUniformBuffer* ViewUniformBuffer, FNiagaraShader* ComputeShader) const
 {
 	FNiagaraComputeExecutionContext* Context = Instance->Context;
@@ -312,12 +320,12 @@ void NiagaraEmitterInstanceBatcher::DispatchMultipleStages(const FNiagaraGPUSyst
 					Instance->CurrentData = Instance->DestinationData;
 					Instance->DestinationData = StoreData;
 				}
-				Run<bDoResourceTransitions>(Tick, Instance, 0, Instance->DestinationData->GetNumInstances(), ComputeShader, RHICmdList, ViewUniformBuffer, Instance->SpawnInfo, false, ShaderStageIndex,  nullptr);
+				Run(Tick, Instance, 0, Instance->DestinationData->GetNumInstances(), ComputeShader, RHICmdList, ViewUniformBuffer, Instance->SpawnInfo, false, ShaderStageIndex,  nullptr);
 			}
 			else
 			{
 				// run with correct number of instances.  This will make curr data junk or empty
-				Run<bDoResourceTransitions>(Tick, Instance, 0, IterationInterface->ElementCount, ComputeShader, RHICmdList, ViewUniformBuffer, Instance->SpawnInfo, false, ShaderStageIndex, IterationInterface);
+				Run(Tick, Instance, 0, IterationInterface->ElementCount, ComputeShader, RHICmdList, ViewUniformBuffer, Instance->SpawnInfo, false, ShaderStageIndex, IterationInterface);
 			}
 			PostStageInterface(Tick, Instance, RHICmdList, ComputeShader, ShaderStageIndex);
 		}
@@ -327,7 +335,7 @@ void NiagaraEmitterInstanceBatcher::DispatchMultipleStages(const FNiagaraGPUSyst
 	else
 	{
 		// run shader, sim and spawn in a single dispatch
-		Run<bDoResourceTransitions>(Tick, Instance, 0, Instance->DestinationData->GetNumInstances(), ComputeShader, RHICmdList, ViewUniformBuffer, Instance->SpawnInfo);
+		Run(Tick, Instance, 0, Instance->DestinationData->GetNumInstances(), ComputeShader, RHICmdList, ViewUniformBuffer, Instance->SpawnInfo);
 	}
 }
 
@@ -445,40 +453,36 @@ void NiagaraEmitterInstanceBatcher::DispatchAllOnCompute(FOverlappableTicks& Ove
 	//
 	RHICmdList.TransitionResource(EResourceTransitionAccess::ERWBarrier, EResourceTransitionPipeline::EComputeToCompute, GPUInstanceCounterManager.GetInstanceCountBuffer().UAV);		// transition to readable; we'll be using this next frame
 	RHICmdList.TransitionResources(EResourceTransitionAccess::ERWBarrier, EResourceTransitionPipeline::EComputeToCompute, DestDataBuffers.GetData(), DestDataBuffers.Num());
-	RHICmdList.TransitionResources(EResourceTransitionAccess::EWritable, EResourceTransitionPipeline::EComputeToCompute, DestBufferIntFloat.GetData(), DestBufferIntFloat.Num());
+	RHICmdList.TransitionResources(EResourceTransitionAccess::ERWBarrier, EResourceTransitionPipeline::EComputeToCompute, DestBufferIntFloat.GetData(), DestBufferIntFloat.Num());
 	RHICmdList.TransitionResources(EResourceTransitionAccess::EReadable, EResourceTransitionPipeline::EComputeToCompute, CurrDataBuffers.GetData(), CurrDataBuffers.Num());
 	RHICmdList.TransitionResources(EResourceTransitionAccess::EReadable, EResourceTransitionPipeline::EComputeToCompute, CurrBufferIntFloat.GetData(), CurrBufferIntFloat.Num());
 	RHICmdList.FlushComputeShaderCache();
 
+	for (FNiagaraGPUSystemTick* Tick : OverlappableTick)
 	{
-		SCOPED_DRAW_EVENT(RHICmdList, NiagaraGPUSimulation);
-		SCOPED_GPU_STAT(RHICmdList, NiagaraGPUSimulation);
-		for (FNiagaraGPUSystemTick* Tick : OverlappableTick)
+		uint32 DispatchCount = Tick->Count;
+		FNiagaraComputeInstanceData* Instances = Tick->GetInstanceData();
+		for (uint32 Index = 0; Index < DispatchCount; Index++)
 		{
-			uint32 DispatchCount = Tick->Count;
-			FNiagaraComputeInstanceData* Instances = Tick->GetInstanceData();
-			for (uint32 Index = 0; Index < DispatchCount; Index++)
+			FNiagaraComputeInstanceData& Instance = Instances[Index];
+			FNiagaraComputeExecutionContext* Context = Instance.Context;
+			if (Context && Context->GPUScript_RT->GetShader())
 			{
-				FNiagaraComputeInstanceData& Instance = Instances[Index];
-				FNiagaraComputeExecutionContext* Context = Instance.Context;
-				if (Context && Context->GPUScript_RT->GetShader())
-				{
-					FNiagaraComputeExecutionContext::TickCounter++;
+				FNiagaraComputeExecutionContext::TickCounter++;
 
-					// run shader, sim and spawn in a single dispatch
-					DispatchMultipleStages<false>(*Tick, &Instance, RHICmdList, ViewUniformBuffer, Context->GPUScript_RT->GetShader());
+				// run shader, sim and spawn in a single dispatch
+				DispatchMultipleStages(*Tick, &Instance, RHICmdList, ViewUniformBuffer, Context->GPUScript_RT->GetShader());
 
-					FNiagaraDataBuffer* CurrentData = Instance.CurrentData;
-					if (bSetReadback && Tick->bIsFinalTick)
+				FNiagaraDataBuffer* CurrentData = Instance.CurrentData;
+				if (bSetReadback && Tick->bIsFinalTick)
+				{
+					// Now that the current data is not required anymore, stage it for readback.
+					if (CurrentData->GetNumInstances() && Context->EmitterInstanceReadback.GPUCountOffset == INDEX_NONE && CurrentData->GetGPUInstanceCountBufferOffset() != INDEX_NONE)
 					{
-						// Now that the current data is not required anymore, stage it for readback.
-						if (CurrentData->GetNumInstances() && Context->EmitterInstanceReadback.GPUCountOffset == INDEX_NONE && CurrentData->GetGPUInstanceCountBufferOffset() != INDEX_NONE)
-						{
-							// Transfer the GPU instance counter ownership to the context. Note that when bSetReadback is true, a readback request will be performed later in the tick update.
-							Context->EmitterInstanceReadback.GPUCountOffset = CurrentData->GetGPUInstanceCountBufferOffset();
-							Context->EmitterInstanceReadback.CPUCount = CurrentData->GetNumInstances();
-							CurrentData->ClearGPUInstanceCountBufferOffset();
-						}
+						// Transfer the GPU instance counter ownership to the context. Note that when bSetReadback is true, a readback request will be performed later in the tick update.
+						Context->EmitterInstanceReadback.GPUCountOffset = CurrentData->GetGPUInstanceCountBufferOffset();
+						Context->EmitterInstanceReadback.CPUCount = CurrentData->GetNumInstances();
+						CurrentData->ClearGPUInstanceCountBufferOffset();
 					}
 				}
 			}
@@ -498,13 +502,12 @@ void NiagaraEmitterInstanceBatcher::DispatchAllOnCompute(FOverlappableTicks& Ove
 	// We have done all our overlapping compute work on this list so go back to default behavior and flush.
 	RHICmdList.AutomaticCacheFlushAfterComputeShader(true);
 
-	//TODO: Need to set the data to render in the context?
-
-	RHICmdList.TransitionResources(EResourceTransitionAccess::EReadable, EResourceTransitionPipeline::EComputeToGfx, DestBufferIntFloat.GetData(), DestBufferIntFloat.Num());
-
 	// We have done all our compute work
 	RHICmdList.FlushComputeShaderCache();
-	RHICmdList.SubmitCommandsHint();
+	if (GNiagaraSubmitCommands)
+	{
+		RHICmdList.SubmitCommandsHint();
+	}
 }
 
 void NiagaraEmitterInstanceBatcher::PostRenderOpaque(FRHICommandListImmediate& RHICmdList, FRHIUniformBuffer* ViewUniformBuffer, const class FShaderParametersMetadata* SceneTexturesUniformBufferStruct, FRHIUniformBuffer* SceneTexturesUniformBuffer, bool bAllowGPUParticleUpdate)
@@ -535,40 +538,53 @@ void NiagaraEmitterInstanceBatcher::ExecuteAll(FRHICommandList &RHICmdList, FRHI
 
 	SCOPED_DRAW_EVENT(RHICmdList, NiagaraEmitterInstanceBatcher_ExecuteAll);
 
-	if (GNiagaraOverlapCompute>0)
+	TArray<FOverlappableTicks, TInlineAllocator<2>> SimPasses;
 	{
-		TArray<FOverlappableTicks, TInlineAllocator<2>> SimPasses;
+		TMap<FNiagaraComputeExecutionContext*, FOverlappableTicks> ContextToTicks;
+
+		// Ticks are added in order. Two tick with the same context should not overlap so should be in two different batch.
+		// Those ticks should still be executed in order.
+		for (FNiagaraGPUSystemTick& Tick : Ticks_RT)
 		{
-			TMap<FNiagaraComputeExecutionContext*, FOverlappableTicks> ContextToTicks;
+			FNiagaraComputeInstanceData* Data = Tick.GetInstanceData();
 
-			// Ticks are added in order. Two tick with the same context should not overlap so should be in two different batch.
-			// Those ticks should still be executed in order.
-			for (FNiagaraGPUSystemTick& Tick : Ticks_RT)
+			// This assumes all emitters fallback to the same FNiagaraShaderScript*.
+			FNiagaraShader* ComputeShader = Data->Context->GPUScript_RT->GetShader();
+			if (!ComputeShader)
 			{
-				FNiagaraComputeInstanceData* Data = Tick.GetInstanceData();
+				continue;
+			}
+			else if (GNiagaraAllowTickBeforeRender && ComputeShader->ViewUniformBufferParam.IsBound() != (ViewUniformBuffer != nullptr))
+			{   // When allowing tick before render, skip this emitter if it is not in the right pass.
+				continue;
+			}
 
-				// This assumes all emitters fallback to the same FNiagaraShaderScript*.
-				FNiagaraShader* ComputeShader = Data->Context->GPUScript_RT->GetShader();
-				if (!ComputeShader)
-				{
-					continue;
-				}
-				else if (GNiagaraAllowTickBeforeRender && ComputeShader->ViewUniformBufferParam.IsBound() != (ViewUniformBuffer != nullptr))
-				{   // When allowing tick before render, skip this emitter if it is not in the right pass.
-					continue;
-				}
+			ContextToTicks.FindOrAdd(Data->Context).Add(&Tick);
+		}
 
-				ContextToTicks.FindOrAdd(Data->Context).Add(&Tick);
-			}
+		// Count the maximum number of tick per context to know the number of passes we need to do
+		uint32 NumSimPass = 0;
+		for (auto& Ticks : ContextToTicks)
+		{
+			NumSimPass = FMath::Max(NumSimPass, (uint32)Ticks.Value.Num());
+			Ticks.Value.Last()->bIsFinalTick = true;
+		}
 
-			// Count the maximum number of tick per context to know the number of passes we need to do
-			uint32 NumSimPass = 0;
+		// Force dispatches to run individually, this should only be used for debugging as it is highly inefficient on the GPU
+		if (GNiagaraOverlapCompute == 0)
+		{
 			for (auto& Ticks : ContextToTicks)
 			{
-				NumSimPass = FMath::Max(NumSimPass, (uint32)Ticks.Value.Num());
-				Ticks.Value.Last()->bIsFinalTick = true;
-			}
+				SimPasses.Reserve(Ticks.Value.Num() - SimPasses.GetSlack());
 
+				for (auto& Tick : Ticks.Value)
+				{
+					SimPasses.AddDefaulted_GetRef().Add(Tick);
+				}
+			}
+		}
+		else
+		{
 			// Transpose now only once the data to get all independent tick per pass
 			SimPasses.AddDefaulted(NumSimPass);
 			for (auto& Ticks : ContextToTicks)
@@ -581,111 +597,22 @@ void NiagaraEmitterInstanceBatcher::ExecuteAll(FRHICommandList &RHICmdList, FRHI
 				}
 			}
 		}
-
-		for (auto& SimPass : SimPasses)
-		{
-			FNiagaraBufferArray DestDataBuffers;
-			FNiagaraBufferArray CurrDataBuffers;
-			FNiagaraBufferArray DestBufferIntFloat;
-			FNiagaraBufferArray CurrBufferIntFloat;
-
-			// This initial pass gathers all the buffers that are read from and written to so we can do batch resource transitions.
-			// It also ensures the GPU buffers are large enough to hold everything.
-			ResizeBuffersAndGatherResources(SimPass, RHICmdList, DestDataBuffers, CurrDataBuffers, DestBufferIntFloat, CurrBufferIntFloat);
-
-			DispatchAllOnCompute(SimPass, RHICmdList, ViewUniformBuffer, DestDataBuffers, CurrDataBuffers, DestBufferIntFloat, CurrBufferIntFloat, bSetReadback);
-		}
 	}
-	else
-	{
-		for (FNiagaraGPUSystemTick& Tick : Ticks_RT)
-		{
-			uint32 DispatchCount = Tick.Count;
-			FNiagaraComputeInstanceData* Instances = Tick.GetInstanceData();
 
-			for (uint32 i = 0; i < DispatchCount; i++)
-			{
-				FNiagaraComputeInstanceData& Instance = Instances[i];
-				TickSingle(Tick, &Instance, RHICmdList, ViewUniformBuffer, bSetReadback);
-			}
-		}
-	}
-}
-
-void NiagaraEmitterInstanceBatcher::TickSingle(const FNiagaraGPUSystemTick& Tick, FNiagaraComputeInstanceData *Instance, FRHICommandList &RHICmdList, FRHIUniformBuffer* ViewUniformBuffer, bool bSetReadback) const
-{
-	SCOPE_CYCLE_COUNTER(STAT_NiagaraGPUSimTick_RT);
-
-	FNiagaraComputeExecutionContext* Context = Instance->Context;
-	
-	check(IsInRenderingThread());
-
-	FNiagaraComputeExecutionContext::TickCounter++;
-
-	FNiagaraShader* ComputeShader = Context->GPUScript_RT->GetShader();
-	if (!ComputeShader)
-	{
-		return;
-	}
-	// When allowing tick before render, skip this emitter if it is not in the right pass.
-	else if (GNiagaraAllowTickBeforeRender && ComputeShader->ViewUniformBufferParam.IsBound() != (ViewUniformBuffer != nullptr))
+	for (auto& SimPass : SimPasses)
 	{
-		return;
-	}
-
-#if WITH_EDITORONLY_DATA
-	if (Context->DebugInfo.IsValid())
-	{
-		ProcessDebugInfo(RHICmdList, Context);
-	}
-#endif // WITH_EDITORONLY_DATA
-	if (!ResetDataInterfaces(Tick, Instance, RHICmdList, ComputeShader)) return;
-
-	//The buffer containing current simulation state.
-	Instance->CurrentData = Context->MainDataSet->GetCurrentData();
-	//The buffer we're going to write simulation results to.
-	Instance->DestinationData = &Context->MainDataSet->BeginSimulate();
-
-	check(Instance->CurrentData && Instance->DestinationData);
-	FNiagaraDataBuffer& CurrentData = *Instance->CurrentData;
-	FNiagaraDataBuffer& DestinationData = *Instance->DestinationData;
+		FNiagaraBufferArray DestDataBuffers;
+		FNiagaraBufferArray CurrDataBuffers;
+		FNiagaraBufferArray DestBufferIntFloat;
+		FNiagaraBufferArray CurrBufferIntFloat;
 
-	const uint32 PrevNumInstances = Tick.bNeedsReset ? 0 : CurrentData.GetNumInstances();
-	const uint32 NewNumInstances = Instance->SpawnInfo.SpawnRateInstances + Instance->SpawnInfo.EventSpawnTotal + PrevNumInstances;
+		// This initial pass gathers all the buffers that are read from and written to so we can do batch resource transitions.
+		// It also ensures the GPU buffers are large enough to hold everything.
+		ResizeBuffersAndGatherResources(SimPass, RHICmdList, DestDataBuffers, CurrDataBuffers, DestBufferIntFloat, CurrBufferIntFloat);
 
-	DestinationData.AllocateGPU(NewNumInstances + 1, GPUInstanceCounterManager, RHICmdList);
-	DestinationData.SetNumInstances(NewNumInstances);
-
-	// run shader, sim and spawn in a single dispatch
-	DispatchMultipleStages<true>(Tick, Instance, RHICmdList, ViewUniformBuffer, ComputeShader);
-
-	Context->MainDataSet->EndSimulate();
-	Context->SetDataToRender(Instance->DestinationData);
-
-	/*
-	// TODO: hack - only updating event set 0 on update scripts now; need to match them to their indices and update them all
-	if (Context->UpdateEventWriteDataSets.Num())
-	{
-		Context->UpdateEventWriteDataSets[0]->CurrDataRender().SetNumInstances(NumInstancesAfterSim[1]);
-	}
-	RunEventHandlers(Context, NumInstancesAfterSim[0], NumInstancesAfterSpawn, NumInstancesAfterNonEventSpawn, RHICmdList);
-	*/
-
-	// the VF grabs PrevDataRender for drawing, so need to transition
-	// Better to do this in the VF set?
-	RHICmdList.TransitionResource(EResourceTransitionAccess::EReadable, EResourceTransitionPipeline::EComputeToGfx, Context->DataToRender->GetGPUBufferFloat().UAV);
-	RHICmdList.TransitionResource(EResourceTransitionAccess::EReadable, EResourceTransitionPipeline::EComputeToGfx, Context->DataToRender->GetGPUBufferInt().UAV);
-
-	// Put a barrier so that next tick or arg gen gets final count.
-	RHICmdList.TransitionResource(EResourceTransitionAccess::ERWBarrier, EResourceTransitionPipeline::EComputeToCompute, GPUInstanceCounterManager.GetInstanceCountBuffer().UAV);		// transition to readable; we'll be using this next frame
-
-	// Now that the current data is not required anymore, stage it for readback.
-	if (bSetReadback && PrevNumInstances && Context->EmitterInstanceReadback.GPUCountOffset == INDEX_NONE && CurrentData.GetGPUInstanceCountBufferOffset() != INDEX_NONE)
-	{
-		// Transfer the GPU instance counter ownership to the context. Note that when bSetReadback is true, a readback request will be performed later in the tick update.
-		Context->EmitterInstanceReadback.GPUCountOffset = CurrentData.GetGPUInstanceCountBufferOffset();
-		Context->EmitterInstanceReadback.CPUCount = PrevNumInstances;
-		CurrentData.ClearGPUInstanceCountBufferOffset();
+		SCOPED_DRAW_EVENT(RHICmdList, NiagaraGPUSimulation);
+		SCOPED_GPU_STAT(RHICmdList, NiagaraGPUSimulation);
+		DispatchAllOnCompute(SimPass, RHICmdList, ViewUniformBuffer, DestDataBuffers, CurrDataBuffers, DestBufferIntFloat, CurrBufferIntFloat, bSetReadback);
 	}
 }
 
@@ -1008,7 +935,7 @@ void NiagaraEmitterInstanceBatcher::ProcessDebugInfo(FRHICommandList &RHICmdList
 {
 #if WITH_EDITORONLY_DATA
 	// This method may be called from one of two places: in the tick or as part of a paused frame looking for the debug info that was submitted previously...
-	// Note that PrevData is where we expect the data to be for rendering, as per NiagaraEmitterInstanceBatcher::TickSingle
+	// Note that PrevData is where we expect the data to be for rendering
 	if (Context && Context->DebugInfo.IsValid())
 	{
 		// Fire off the readback if not already doing so
@@ -1149,12 +1076,11 @@ void NiagaraEmitterInstanceBatcher::UnsetDataInterfaceParameters(const TArray<FN
 
 /* Kick off a simulation/spawn run
  */
- template<bool bDoResourceTransitions>
 void NiagaraEmitterInstanceBatcher::Run(const FNiagaraGPUSystemTick& Tick, const FNiagaraComputeInstanceData* Instance, uint32 UpdateStartInstance, const uint32 TotalNumInstances, FNiagaraShader* Shader,
 	FRHICommandList &RHICmdList, FRHIUniformBuffer* ViewUniformBuffer, const FNiagaraGpuSpawnInfo& SpawnInfo, bool bCopyBeforeStart, uint32 ShaderStageIndex, FNiagaraDataInterfaceProxy *IterationInterface) const
 {
 	FNiagaraComputeExecutionContext* Context = Instance->Context;
-	if (TotalNumInstances == 0 && !bDoResourceTransitions)
+	if (TotalNumInstances == 0)
 	{
 		SCOPED_DRAW_EVENTF(RHICmdList, NiagaraGPUSimulationCS, TEXT("Niagara Gpu Sim - %s - NumInstances: %u - StageNumber: %u"),
 			*Context->DebugSimName,
@@ -1196,8 +1122,8 @@ void NiagaraEmitterInstanceBatcher::Run(const FNiagaraGPUSystemTick& Tick, const
 
 	// set the shader and data set params 
 	//
-	CurrentData.SetShaderParams<bDoResourceTransitions>(Shader, RHICmdList, true);
-	DestinationData.SetShaderParams<bDoResourceTransitions>(Shader, RHICmdList, false);
+	CurrentData.SetShaderParams(Shader, RHICmdList, true);
+	DestinationData.SetShaderParams(Shader, RHICmdList, false);
 
 	// set the index buffer uav
 	//
