Niagara - Hide exectution state enum values which are invalid to set directly.

diff --git a/Engine/Plugins/FX/Niagara/Source/Niagara/Public/NiagaraTypes.h b/Engine/Plugins/FX/Niagara/Source/Niagara/Public/NiagaraTypes.h
--- a/Engine/Plugins/FX/Niagara/Source/Niagara/Public/NiagaraTypes.h
+++ b/Engine/Plugins/FX/Niagara/Source/Niagara/Public/NiagaraTypes.h
@@ -294,10 +294,10 @@ enum class ENiagaraExecutionState : uint32
 	/** Complete. When the system or all emitters are complete the effect is considered finished. */
 	Complete,
 	/** Emitter only. Emitter is disabled. Will not tick or render again until a full re initialization of the system. */
-	Disabled,
+	Disabled UMETA(Hidden),
 
 	// insert new states before
-	Num
+	Num UMETA(Hidden)
 };
 
 /** Defines options for conditionally editing and showing script inputs in the UI. */
