Manual Import of the CL 7711015  (integrate->edit)

diff --git a/Engine/Plugins/FX/Niagara/Source/Niagara/Private/NiagaraScript.cpp b/Engine/Plugins/FX/Niagara/Source/Niagara/Private/NiagaraScript.cpp
--- a/Engine/Plugins/FX/Niagara/Source/Niagara/Private/NiagaraScript.cpp
+++ b/Engine/Plugins/FX/Niagara/Source/Niagara/Private/NiagaraScript.cpp
@@ -302,6 +302,12 @@ void UNiagaraScript::ComputeVMCompilationId(FNiagaraVMExecutableDataId& Id) cons
 		}
 	}
 
+	static const auto UseShaderStagesCVar = IConsoleManager::Get().FindConsoleVariable(TEXT("fx.UseShaderStages"));
+	if (UseShaderStagesCVar->GetInt())
+	{
+		Id.AdditionalDefines.Add(TEXT("fx.UseShaderStages"));
+	}
+
 	Source->ComputeVMCompilationId(Id, Usage, UsageId);
 	
 	LastGeneratedVMId = Id;
@@ -1177,7 +1183,7 @@ void UNiagaraScript::CacheResourceShadersForCooking(EShaderPlatform ShaderPlatfo
 			check(CachedScriptVMId.BaseScriptID.IsValid());
 			check(CachedScriptVMId.BaseScriptCompileHash.IsValid());
 
-			NewResource->SetScript(this, (ERHIFeatureLevel::Type)TargetFeatureLevel, CachedScriptVMId.CompilerVersionID, CachedScriptVMId.BaseScriptID,
+			NewResource->SetScript(this, (ERHIFeatureLevel::Type)TargetFeatureLevel, CachedScriptVMId.CompilerVersionID, CachedScriptVMId.BaseScriptID, CachedScriptVMId.AdditionalDefines,
 				CachedScriptVMId.BaseScriptCompileHash,	CachedScriptVMId.ReferencedCompileHashes, CachedScriptVMId.ReferencedDependencyIds, GetFullName());
 			ResourceToCache = NewResource;
 
@@ -1242,7 +1248,7 @@ void UNiagaraScript::CacheResourceShadersForRendering(bool bRegenerateId, bool b
 		{
 			FNiagaraShaderScript* ResourceToCache;
 			ERHIFeatureLevel::Type CacheFeatureLevel = GMaxRHIFeatureLevel;
-			ScriptResource.SetScript(this, FeatureLevel, CachedScriptVMId.CompilerVersionID, CachedScriptVMId.BaseScriptID,
+			ScriptResource.SetScript(this, FeatureLevel, CachedScriptVMId.CompilerVersionID, CachedScriptVMId.BaseScriptID, CachedScriptVMId.AdditionalDefines,
 				CachedScriptVMId.BaseScriptCompileHash, CachedScriptVMId.ReferencedCompileHashes, CachedScriptVMId.ReferencedDependencyIds, GetFullName());
 
 			//if (ScriptResourcesByFeatureLevel[FeatureLevel])
diff --git a/Engine/Plugins/FX/Niagara/Source/NiagaraCore/Private/NiagaraCustomVersion.cpp b/Engine/Plugins/FX/Niagara/Source/NiagaraCore/Private/NiagaraCustomVersion.cpp
--- a/Engine/Plugins/FX/Niagara/Source/NiagaraCore/Private/NiagaraCustomVersion.cpp
+++ b/Engine/Plugins/FX/Niagara/Source/NiagaraCore/Private/NiagaraCustomVersion.cpp
@@ -8,6 +8,4 @@ const FGuid FNiagaraCustomVersion::GUID(0xFCF57AFA, 0x50764283, 0xB9A9E658, 0xFF
 // Register the custom version with core
 FCustomVersionRegistration GRegisterNiagaraCustomVersion(FNiagaraCustomVersion::GUID, FNiagaraCustomVersion::LatestVersion, TEXT("NiagaraVer"));
 
-const FGuid FNiagaraCustomVersion::LatestScriptCompileVersion(0x563D9FF1, 0x1E9042CC, 0xA8BF7AC4, 0xC9366901);
-
-
+const FGuid FNiagaraCustomVersion::LatestScriptCompileVersion(0xAA5883AE, 0xCB3648DB, 0xAC0BF434, 0xFEA4BE4);
diff --git a/Engine/Plugins/FX/Niagara/Source/NiagaraCore/Public/NiagaraCustomVersion.h b/Engine/Plugins/FX/Niagara/Source/NiagaraCore/Public/NiagaraCustomVersion.h
--- a/Engine/Plugins/FX/Niagara/Source/NiagaraCore/Public/NiagaraCustomVersion.h
+++ b/Engine/Plugins/FX/Niagara/Source/NiagaraCore/Public/NiagaraCustomVersion.h
@@ -124,6 +124,8 @@ struct FNiagaraCustomVersion
 		
 		AddLibraryAssetProperty, // Add property to all Niagara scripts indicating whether or not they belong to the library
 
+		AddAdditionalDefinesProperty, // Addding additional defines to the GPU script
+
 		// -----<new versions can be added above this line>-------------------------------------------------
 		VersionPlusOne,
 		LatestVersion = VersionPlusOne - 1,
diff --git a/Engine/Plugins/FX/Niagara/Source/NiagaraShader/Private/NiagaraShader.cpp b/Engine/Plugins/FX/Niagara/Source/NiagaraShader/Private/NiagaraShader.cpp
--- a/Engine/Plugins/FX/Niagara/Source/NiagaraShader/Private/NiagaraShader.cpp
+++ b/Engine/Plugins/FX/Niagara/Source/NiagaraShader/Private/NiagaraShader.cpp
@@ -128,6 +128,10 @@ void FNiagaraShaderMapId::Serialize(FArchive& Ar)
 		Ar << BaseCompileHash;
 		Ar << ReferencedCompileHashes;
 	}
+	if (NiagaraVer >= FNiagaraCustomVersion::AddAdditionalDefinesProperty)
+	{
+		Ar << AdditionalDefines;
+	}
 	
 	//ParameterSet.Serialize(Ar);		// NIAGARATODO: at some point we'll need stuff for static switches here
 }
@@ -142,6 +146,11 @@ void FNiagaraShaderMapId::GetScriptHash(FSHAHash& OutHash) const
 	HashState.Update(BaseCompileHash.GetData(), FNiagaraCompileHash::HashSize);
 	HashState.Update((const uint8*)&FeatureLevel, sizeof(FeatureLevel));
 		
+	for (int32 Index = 0; Index < AdditionalDefines.Num(); Index++)
+	{
+		HashState.UpdateWithString(*AdditionalDefines[Index], AdditionalDefines[Index].Len());
+	}
+
 	for (int32 Index = 0; Index < ReferencedCompileHashes.Num(); Index++)
 	{
 		HashState.Update(ReferencedCompileHashes[Index].GetData(), FNiagaraCompileHash::HashSize);
@@ -173,6 +182,18 @@ bool FNiagaraShaderMapId::operator==(const FNiagaraShaderMapId& ReferenceSet) co
 	{
 		return false;
 	}
+	if (AdditionalDefines.Num() != ReferenceSet.AdditionalDefines.Num())
+	{
+		return false;
+	}
+
+	for (int32 i = 0; i < AdditionalDefines.Num(); i++)
+	{
+		if (AdditionalDefines[i] != ReferenceSet.AdditionalDefines[i])
+		{
+			return false;
+		}
+	}
 
 	if (ReferencedCompileHashes.Num() != ReferenceSet.ReferencedCompileHashes.Num())
 	{
@@ -234,6 +255,16 @@ void FNiagaraShaderMapId::AppendKeyString(FString& KeyString) const
 	KeyString += CompilerVersionID.ToString();
 	KeyString += TEXT("_");
 
+	// Add additional defines
+	for (int32 DefinesIndex = 0; DefinesIndex < AdditionalDefines.Num(); DefinesIndex++)
+	{
+		KeyString += AdditionalDefines[DefinesIndex];
+		if (DefinesIndex < AdditionalDefines.Num() - 1)
+		{
+			KeyString += TEXT("_");
+		}
+	}
+
 	// Add any referenced top level compile hashes to the key so that we will recompile when they are changed
 	for (int32 HashIndex = 0; HashIndex < ReferencedCompileHashes.Num(); HashIndex++)
 	{
diff --git a/Engine/Plugins/FX/Niagara/Source/NiagaraShader/Private/NiagaraShared.cpp b/Engine/Plugins/FX/Niagara/Source/NiagaraShader/Private/NiagaraShared.cpp
--- a/Engine/Plugins/FX/Niagara/Source/NiagaraShader/Private/NiagaraShared.cpp
+++ b/Engine/Plugins/FX/Niagara/Source/NiagaraShader/Private/NiagaraShared.cpp
@@ -89,6 +89,7 @@ bool FNiagaraShaderScript::IsSame(const FNiagaraShaderMapId& InId) const
 	return
 		InId.FeatureLevel == FeatureLevel &&
 		InId.BaseScriptID == BaseScriptId &&
+		InId.AdditionalDefines == AdditionalDefines &&
 		InId.BaseCompileHash == BaseCompileHash &&
 		InId.ReferencedCompileHashes == ReferencedCompileHashes &&
 		InId.ReferencedDependencyIds == ReferencedDependencyIds &&
@@ -123,6 +124,7 @@ NIAGARASHADER_API void FNiagaraShaderScript::GetShaderMapId(EShaderPlatform Plat
 		GetDependentShaderTypes(Platform, ShaderTypes);
 		OutId.FeatureLevel = GetFeatureLevel();
 		OutId.BaseScriptID = BaseScriptId;
+		OutId.AdditionalDefines = AdditionalDefines;
 		OutId.BaseCompileHash = BaseCompileHash;
 		OutId.ReferencedCompileHashes = ReferencedCompileHashes;
 		OutId.ReferencedDependencyIds = ReferencedDependencyIds;
@@ -224,7 +226,7 @@ void FNiagaraShaderScript::SerializeShaderMap(FArchive& Ar)
 	}
 }
 
-void FNiagaraShaderScript::SetScript(UNiagaraScript *InScript, ERHIFeatureLevel::Type InFeatureLevel, const FGuid& InCompilerVersionID, const FGuid& InBaseScriptID,
+void FNiagaraShaderScript::SetScript(UNiagaraScript *InScript, ERHIFeatureLevel::Type InFeatureLevel, const FGuid& InCompilerVersionID, const FGuid& InBaseScriptID, const TArray<FString>& InAdditionalDefines,
 	const FNiagaraCompileHash& InBaseCompileHash, const TArray<FNiagaraCompileHash>& InReferencedCompileHashes, const TArray<FGuid>& InReferencedDependencyIds, FString InFriendlyName)
 {
 	checkf(InBaseScriptID.IsValid(), TEXT("Invalid base script id.  Script caching will fail."));
@@ -232,6 +234,7 @@ void FNiagaraShaderScript::SetScript(UNiagaraScript *InScript, ERHIFeatureLevel:
 	BaseVMScript = InScript;
 	CompilerVersionId = InCompilerVersionID;
 	BaseScriptId = InBaseScriptID;
+	AdditionalDefines = InAdditionalDefines;
 	BaseCompileHash = InBaseCompileHash;
 	ReferencedCompileHashes = InReferencedCompileHashes;
 	ReferencedDependencyIds = InReferencedDependencyIds;
diff --git a/Engine/Plugins/FX/Niagara/Source/NiagaraShader/Public/NiagaraShared.h b/Engine/Plugins/FX/Niagara/Source/NiagaraShader/Public/NiagaraShared.h
--- a/Engine/Plugins/FX/Niagara/Source/NiagaraShader/Public/NiagaraShared.h
+++ b/Engine/Plugins/FX/Niagara/Source/NiagaraShader/Public/NiagaraShared.h
@@ -161,6 +161,9 @@ public:
 	*/
 	FGuid BaseScriptID;
 
+	/** Configuration options */
+	TArray<FString> AdditionalDefines;
+
 	/**
 	* The hash of the subgraph this shader primarily represents.
 	*/
@@ -666,7 +669,7 @@ public:
 	const FString& GetFriendlyName()	const { return FriendlyName; }
 
 
-	NIAGARASHADER_API void SetScript(UNiagaraScript *InScript, ERHIFeatureLevel::Type InFeatureLevel, const FGuid& InCompilerVersion, const FGuid& InBaseScriptID,
+	NIAGARASHADER_API void SetScript(UNiagaraScript *InScript, ERHIFeatureLevel::Type InFeatureLevel, const FGuid& InCompilerVersion, const FGuid& InBaseScriptID, const TArray<FString>& InAdditionalDefines,
 		const FNiagaraCompileHash& InBaseCompileHash, const TArray<FNiagaraCompileHash>& InReferencedCompileHashes, const TArray<FGuid>& InReferencedDependencyIds, FString InFriendlyName);
 
 	UNiagaraScript *GetBaseVMScript()
@@ -742,6 +745,9 @@ private:
 	/** Guid id for base script*/
 	FGuid BaseScriptId;
 
+	/** Configuration options */
+	TArray<FString> AdditionalDefines;
+
 	/** Compile hash for the base script. */
 	FNiagaraCompileHash BaseCompileHash;
 
