Niagara - Remove the generated compile id from the cached compile info for graphs since it's state is already captured by the hash which is deterministic.  This prevents all recompiles in running the engine tests.

diff --git a/Engine/Plugins/FX/Niagara/Source/Niagara/Classes/NiagaraParameterCollection.h b/Engine/Plugins/FX/Niagara/Source/Niagara/Classes/NiagaraParameterCollection.h
--- a/Engine/Plugins/FX/Niagara/Source/Niagara/Classes/NiagaraParameterCollection.h
+++ b/Engine/Plugins/FX/Niagara/Source/Niagara/Classes/NiagaraParameterCollection.h
@@ -4,6 +4,7 @@
 
 #include "NiagaraParameterStore.h"
 #include "NiagaraCommon.h"
+#include "NiagaraCompileHash.h"
 #include "NiagaraParameterCollection.generated.h"
 
 class UNiagaraParameterCollection;
@@ -154,7 +155,7 @@ public:
 	FString GetFullNamespace()const;
 
 	/** The compile Id is an indicator to any compiled scripts that reference this collection that contents may have changed and a recompile is recommended to be safe.*/
-	FGuid GetCompileId() const { return CompileId; };
+	FNiagaraCompileHash GetCompileHash() const;
 
 	/** If for any reason this data has changed externally and needs recompilation and it isn't autodetected, use this method. */
 	void RefreshCompileId();
diff --git a/Engine/Plugins/FX/Niagara/Source/Niagara/Classes/NiagaraScript.h b/Engine/Plugins/FX/Niagara/Source/Niagara/Classes/NiagaraScript.h
--- a/Engine/Plugins/FX/Niagara/Source/Niagara/Classes/NiagaraScript.h
+++ b/Engine/Plugins/FX/Niagara/Source/Niagara/Classes/NiagaraScript.h
@@ -150,10 +150,6 @@ public:
 	UPROPERTY()
 	TArray<FNiagaraCompileHash> ReferencedCompileHashes;
 
-	/** Guids of any functions, module scripts, parameter collections, or other assets the script was dependent on that might trigger a recompile if they change. */
-	UPROPERTY()
-	TArray<FGuid> ReferencedDependencyIds;
-
 	/** Temp storage while generating the Id. This is NOT serialized and shouldn't be used in any comparisons*/
 	TArray<UObject*> ReferencedObjects;
 #endif
diff --git a/Engine/Plugins/FX/Niagara/Source/Niagara/Private/NiagaraParameterCollection.cpp b/Engine/Plugins/FX/Niagara/Source/Niagara/Private/NiagaraParameterCollection.cpp
--- a/Engine/Plugins/FX/Niagara/Source/Niagara/Private/NiagaraParameterCollection.cpp
+++ b/Engine/Plugins/FX/Niagara/Source/Niagara/Private/NiagaraParameterCollection.cpp
@@ -284,6 +284,20 @@ FString UNiagaraParameterCollection::GetFullNamespace()const
 	return TEXT("NPC.") + Namespace.ToString() + TEXT(".");
 }
 
+FNiagaraCompileHash UNiagaraParameterCollection::GetCompileHash() const
+{
+	// TODO - Implement an actual hash for parameter collections instead of just hashing a change id.
+	FSHA1 CompileHash;
+	CompileHash.Update((const uint8*)&CompileId, sizeof(FGuid));
+	CompileHash.Final();
+
+	TArray<uint8> DataHash;
+	DataHash.AddUninitialized(20);
+	CompileHash.GetHash(DataHash.GetData());
+
+	return FNiagaraCompileHash(DataHash);
+}
+
 void UNiagaraParameterCollection::RefreshCompileId()
 {
 	CompileId = FGuid::NewGuid();
diff --git a/Engine/Plugins/FX/Niagara/Source/Niagara/Private/NiagaraScript.cpp b/Engine/Plugins/FX/Niagara/Source/Niagara/Private/NiagaraScript.cpp
--- a/Engine/Plugins/FX/Niagara/Source/Niagara/Private/NiagaraScript.cpp
+++ b/Engine/Plugins/FX/Niagara/Source/Niagara/Private/NiagaraScript.cpp
@@ -136,21 +136,6 @@ bool FNiagaraVMExecutableDataId::operator==(const FNiagaraVMExecutableDataId& Re
 			return false;
 		}
 	}
-	
-	if (ReferencedDependencyIds.Num() != ReferenceSet.ReferencedDependencyIds.Num())
-	{
-		return false;
-	}
-
-	for (int32 RefFunctionIndex = 0; RefFunctionIndex < ReferenceSet.ReferencedDependencyIds.Num(); RefFunctionIndex++)
-	{
-		const FGuid& ReferenceGuid = ReferenceSet.ReferencedDependencyIds[RefFunctionIndex];
-
-		if (ReferencedDependencyIds[RefFunctionIndex] != ReferenceGuid)
-		{
-			return false;
-		}
-	}
 
 	if (AdditionalDefines.Num() != ReferenceSet.AdditionalDefines.Num())
 	{
@@ -203,17 +188,6 @@ void FNiagaraVMExecutableDataId::AppendKeyString(FString& KeyString) const
 			KeyString += TEXT("_");
 		}
 	}
-
-	// Add any referenced functions to the key so that we will recompile when they are changed
-	for (int32 FunctionIndex = 0; FunctionIndex < ReferencedDependencyIds.Num(); FunctionIndex++)
-	{
-		KeyString += ReferencedDependencyIds[FunctionIndex].ToString();
-
-		if (FunctionIndex < ReferencedDependencyIds.Num() - 1)
-		{
-			KeyString += TEXT("_");
-		}
-	}	
 }
 #endif
 
@@ -645,24 +619,6 @@ bool UNiagaraScript::AreScriptAndSourceSynchronized() const
 							}
 						}
 					}
-
-					if (NewId.ReferencedDependencyIds.Num() != CachedScriptVMId.ReferencedDependencyIds.Num())
-					{
-						UE_LOG(LogNiagara, Log, TEXT("AreScriptAndSourceSynchronized num dependencies don't match. %d != %d, script %s"),
-							NewId.ReferencedDependencyIds.Num(), CachedScriptVMId.ReferencedDependencyIds.Num(), *GetPathName());
-					}
-					else
-					{
-						for (int32 i = 0; i < NewId.ReferencedDependencyIds.Num(); i++)
-						{
-							if (NewId.ReferencedDependencyIds[i] != CachedScriptVMId.ReferencedDependencyIds[i])
-							{
-								UE_LOG(LogNiagara, Log, TEXT("AreScriptAndSourceSynchronized reference id %d doesn't match. %s != %s, script %s, source %s"),
-									i, *NewId.ReferencedDependencyIds[i].ToString(), *CachedScriptVMId.ReferencedDependencyIds[i].ToString(), *GetPathName(),
-									NewId.ReferencedObjects[i] != nullptr ? *NewId.ReferencedObjects[i]->GetPathName() : TEXT("nullptr"));
-							}
-						}
-					}
 				}
 				LastReportedVMId = NewId;
 			}
@@ -1184,7 +1140,7 @@ void UNiagaraScript::CacheResourceShadersForCooking(EShaderPlatform ShaderPlatfo
 			check(CachedScriptVMId.BaseScriptCompileHash.IsValid());
 
 			NewResource->SetScript(this, (ERHIFeatureLevel::Type)TargetFeatureLevel, CachedScriptVMId.CompilerVersionID, CachedScriptVMId.BaseScriptID, CachedScriptVMId.AdditionalDefines,
-				CachedScriptVMId.BaseScriptCompileHash,	CachedScriptVMId.ReferencedCompileHashes, CachedScriptVMId.ReferencedDependencyIds, GetFullName());
+				CachedScriptVMId.BaseScriptCompileHash,	CachedScriptVMId.ReferencedCompileHashes, GetFullName());
 			ResourceToCache = NewResource;
 
 			check(ResourceToCache);
@@ -1249,7 +1205,7 @@ void UNiagaraScript::CacheResourceShadersForRendering(bool bRegenerateId, bool b
 			FNiagaraShaderScript* ResourceToCache;
 			ERHIFeatureLevel::Type CacheFeatureLevel = GMaxRHIFeatureLevel;
 			ScriptResource.SetScript(this, FeatureLevel, CachedScriptVMId.CompilerVersionID, CachedScriptVMId.BaseScriptID, CachedScriptVMId.AdditionalDefines,
-				CachedScriptVMId.BaseScriptCompileHash, CachedScriptVMId.ReferencedCompileHashes, CachedScriptVMId.ReferencedDependencyIds, GetFullName());
+				CachedScriptVMId.BaseScriptCompileHash, CachedScriptVMId.ReferencedCompileHashes, GetFullName());
 
 			//if (ScriptResourcesByFeatureLevel[FeatureLevel])
 			{
diff --git a/Engine/Plugins/FX/Niagara/Source/NiagaraCore/Public/NiagaraCustomVersion.h b/Engine/Plugins/FX/Niagara/Source/NiagaraCore/Public/NiagaraCustomVersion.h
--- a/Engine/Plugins/FX/Niagara/Source/NiagaraCore/Public/NiagaraCustomVersion.h
+++ b/Engine/Plugins/FX/Niagara/Source/NiagaraCore/Public/NiagaraCustomVersion.h
@@ -126,6 +126,8 @@ struct FNiagaraCustomVersion
 
 		AddAdditionalDefinesProperty, // Addding additional defines to the GPU script
 
+		RemoveGraphUsageCompileIds, // Remove the random compile id guids from the cached script usage and from the compile and script ids since the hashes serve the same purpose and are deterministic.
+
 		// -----<new versions can be added above this line>-------------------------------------------------
 		VersionPlusOne,
 		LatestVersion = VersionPlusOne - 1,
diff --git a/Engine/Plugins/FX/Niagara/Source/NiagaraEditor/Private/NiagaraGraph.cpp b/Engine/Plugins/FX/Niagara/Source/NiagaraEditor/Private/NiagaraGraph.cpp
--- a/Engine/Plugins/FX/Niagara/Source/NiagaraEditor/Private/NiagaraGraph.cpp
+++ b/Engine/Plugins/FX/Niagara/Source/NiagaraEditor/Private/NiagaraGraph.cpp
@@ -65,8 +65,11 @@ void FNiagaraGraphScriptUsageInfo::PostLoad(UObject* Owner)
 	const int32 NiagaraVer = Owner->GetLinkerCustomVersion(FNiagaraCustomVersion::GUID);
 	if (NiagaraVer < FNiagaraCustomVersion::UseHashesToIdentifyCompileStateOfTopLevelScripts)
 	{
-		// When loading old data use the last generated compile id as the base id to prevent recompiles on load for existing scripts.
-		BaseId = GeneratedCompileId;
+		if (BaseId.IsValid() == false && GeneratedCompileId_DEPRECATED.IsValid())
+		{
+			// When loading old data use the last generated compile id as the base id to prevent recompiles on load for existing scripts.
+			BaseId = GeneratedCompileId_DEPRECATED;
+		}
 
 		if (CompileHash.IsValid() == false && DataHash_DEPRECATED.Num() == FNiagaraCompileHash::HashSize)
 		{
@@ -289,22 +292,6 @@ class UNiagaraScriptSource* UNiagaraGraph::GetSource() const
 	return CastChecked<UNiagaraScriptSource>(GetOuter());
 }
 
-FGuid UNiagaraGraph::ComputeCompileID(ENiagaraScriptUsage InUsage, const FGuid& InUsageId)
-{
-	RebuildCachedCompileIds();
-
-	for (int32 j = 0; j < CachedUsageInfo.Num(); j++)
-	{
-		if (UNiagaraScript::IsEquivalentUsage(CachedUsageInfo[j].UsageType, InUsage) && CachedUsageInfo[j].UsageId == InUsageId)
-		{
-			return CachedUsageInfo[j].GeneratedCompileId;
-		}
-	}
-
-	return FGuid();
-
-}
-
 FNiagaraCompileHash UNiagaraGraph::GetCompileDataHash(ENiagaraScriptUsage InUsage, const FGuid& InUsageId) const
 {
 	for (int32 i = 0; i < CachedUsageInfo.Num(); i++)
@@ -1043,24 +1030,6 @@ void UNiagaraGraph::RebuildCachedCompileIds(bool bForce)
 		HashState.GetHash(DataHash.GetData());
 		NewUsageCache[i].CompileHash = FNiagaraCompileHash(DataHash);
 
-		bool bNeedsNewCompileId = true;
-
-		// Now compare the hashed data. If it is the same as before, then leave the compile ID as-is. If it is different, generate a new guid.
-		if (FoundMatchIdx != INDEX_NONE)
-		{
-			if (NewUsageCache[i].CompileHash == CachedUsageInfo[FoundMatchIdx].CompileHash)
-			{
-				NewUsageCache[i].GeneratedCompileId = CachedUsageInfo[FoundMatchIdx].GeneratedCompileId;
-				bNeedsNewCompileId = false;
-			}
-		}
-
-		if (bNeedsNewCompileId)
-		{
-			NewUsageCache[i].GeneratedCompileId = FGuid::NewGuid();
-			bNeedsAnyNewCompileIds = true;
-		}
-
 		// TODO sckime debug logic... should be disabled or put under a cvar in the future
 		{
 
@@ -1074,15 +1043,6 @@ void UNiagaraGraph::RebuildCachedCompileIds(bool bForce)
 			{
 				ResultsEnum = FoundEnum->GetNameStringByValue((int64)NewUsageCache[i].UsageType);
 			}
-
-			if (bNeedsNewCompileId)
-			{
-				//UE_LOG(LogNiagaraEditor, Log, TEXT("'%s' changes detected in %s .. new guid: %s"), *GetFullName(), *ResultsEnum, *NewUsageCache[i].GeneratedCompileId.ToString());
-			}
-			else
-			{
-				//UE_LOG(LogNiagaraEditor, Log, TEXT("'%s' changes NOT detected in %s .. keeping guid: %s"), *GetFullName(), *ResultsEnum, *NewUsageCache[i].GeneratedCompileId.ToString());
-			}
 		}
 
 		if (UNiagaraScript::IsEquivalentUsage(NewUsageCache[i].UsageType, ENiagaraScriptUsage::ParticleSpawnScript) && NewUsageCache[i].UsageId == FGuid())
@@ -1115,31 +1075,11 @@ void UNiagaraGraph::RebuildCachedCompileIds(bool bForce)
 			GpuUsageInfo.BaseId = OldGpuInfo->BaseId;
 		}
 
-		GpuUsageInfo.Traversal.Append(ParticleSpawnUsageInfo->Traversal);
-		GpuUsageInfo.Traversal.Append(ParticleUpdateUsageInfo->Traversal);
-
-		FSHA1 HashState;
-		for (UNiagaraNode* Node : GpuUsageInfo.Traversal)
-		{
-			Node->UpdateCompileHashForNode(HashState);
-		}
-		HashState.Final();
-
+		// The gpu script has no graph representation so the hash is empty.  Changes to the gpu script are handled by collecting the compile hashes of it's dependencies.
 		TArray<uint8> DataHash;
-		DataHash.AddUninitialized(20);
-		HashState.GetHash(DataHash.GetData());
+		DataHash.AddZeroed(20);
 		GpuUsageInfo.CompileHash = FNiagaraCompileHash(DataHash);
 
-		FNiagaraGraphScriptUsageInfo* OldGpuUsageInfo = CachedUsageInfo.FindByPredicate([](const FNiagaraGraphScriptUsageInfo& UsageInfo) { return UsageInfo.UsageType == ENiagaraScriptUsage::ParticleGPUComputeScript && UsageInfo.UsageId == FGuid(); });
-		if (OldGpuUsageInfo != nullptr && OldGpuUsageInfo->CompileHash == GpuUsageInfo.CompileHash)
-		{
-			GpuUsageInfo.GeneratedCompileId = OldGpuUsageInfo->GeneratedCompileId;
-		}
-		else
-		{
-			GpuUsageInfo.GeneratedCompileId = FGuid::NewGuid();
-		}
-
 		NewUsageCache.Add(GpuUsageInfo);
 	}
 
@@ -1233,63 +1173,6 @@ void UNiagaraGraph::ResolveNumerics(TMap<UNiagaraNode*, bool>& VisitedNodes, UEd
 	}
 }
 
-
-void UNiagaraGraph::SynchronizeInternalCacheWithGraph(UNiagaraGraph* Other)
-{
-	// Force us to rebuild the cache, note that this builds traversals and everything else, keeping it in sync if nothing changed from the current version.
-	RebuildCachedCompileIds(true);
-	
-	UEnum* FoundEnum = nullptr;
-
-	// Now go through all of the other graph's usage info. If we find a match for its usage and our data hashes match, use the generated compile id from
-	// the other graph.
-	for (int32 i = 0; i < CachedUsageInfo.Num(); i++)
-	{
-		int32 FoundMatchIdx = INDEX_NONE;
-		for (int32 j = 0; j < Other->CachedUsageInfo.Num(); j++)
-		{
-			if (UNiagaraScript::IsEquivalentUsage(Other->CachedUsageInfo[j].UsageType, CachedUsageInfo[i].UsageType) && Other->CachedUsageInfo[j].UsageId == CachedUsageInfo[i].UsageId)
-			{
-				FoundMatchIdx = j;
-				break;
-			}
-		}
-
-		if (FoundMatchIdx != INDEX_NONE)
-		{
-			if (CachedUsageInfo[i].CompileHash == Other->CachedUsageInfo[FoundMatchIdx].CompileHash)
-			{
-				CachedUsageInfo[i].GeneratedCompileId = Other->CachedUsageInfo[FoundMatchIdx].GeneratedCompileId;		
-
-				// TODO sckime debug logic... should be disabled or put under a cvar in the future
-				{
-					if (FoundEnum == nullptr)
-					{
-						FoundEnum = StaticEnum<ENiagaraScriptUsage>();
-					}
-
-					FString ResultsEnum = TEXT("??");
-					if (FoundEnum)
-					{
-						ResultsEnum = FoundEnum->GetNameStringByValue((int64)CachedUsageInfo[i].UsageType);
-					}
-					if (GEnableVerboseNiagaraChangeIdLogging)
-					{
-						UE_LOG(LogNiagaraEditor, Log, TEXT("'%s' changes synchronized with master script in %s .. synced guid: %s"), *GetFullName(), *ResultsEnum, *CachedUsageInfo[i].GeneratedCompileId.ToString());
-					}
-				}
-			}
-		}
-	}
-
-	if (bWriteToLog)
-	{
-		TMap<FGuid, FGuid> ComputeChangeIds;
-		FNiagaraEditorUtilities::GatherChangeIds(*this, ComputeChangeIds, GetName() + TEXT(".Synced"));
-	}
-}
-
-
 void UNiagaraGraph::InvalidateCachedCompileIds()
 {
 	Modify();
@@ -1297,52 +1180,29 @@ void UNiagaraGraph::InvalidateCachedCompileIds()
 	MarkGraphRequiresSynchronization(__FUNCTION__);
 }
 
-void UNiagaraGraph::GatherExternalDependencyIDs(ENiagaraScriptUsage InUsage, const FGuid& InUsageId, TArray<FNiagaraCompileHash>& InReferencedCompileHashes, TArray<FGuid>& InReferencedIDs, TArray<UObject*>& InReferencedObjs)
+void UNiagaraGraph::GatherExternalDependencyData(ENiagaraScriptUsage InUsage, const FGuid& InUsageId, TArray<FNiagaraCompileHash>& InReferencedCompileHashes, TArray<UObject*>& InReferencedObjs)
 {
 	RebuildCachedCompileIds();
 	
-	// Particle compute scripts get all particle scripts baked into their dependency chain. 
-	if (InUsage == ENiagaraScriptUsage::ParticleGPUComputeScript)
+	for (int32 i = 0; i < CachedUsageInfo.Num(); i++)
 	{
-		for (int32 i = 0; i < CachedUsageInfo.Num(); i++)
+		// First add our direct dependency chain...
+		if (UNiagaraScript::IsEquivalentUsage(CachedUsageInfo[i].UsageType, InUsage) && CachedUsageInfo[i].UsageId == InUsageId)
 		{
-			// Add all chains that we depend on.
-			if (UNiagaraScript::IsUsageDependentOn(InUsage, CachedUsageInfo[i].UsageType)) 
+			for (UNiagaraNode* Node : CachedUsageInfo[i].Traversal)
 			{
-				InReferencedCompileHashes.Add(CachedUsageInfo[i].CompileHash);
-				InReferencedObjs.Add(CachedUsageInfo[i].Traversal.Last());
-
-				for (UNiagaraNode* Node : CachedUsageInfo[i].Traversal)
-				{
-					Node->GatherExternalDependencyIDs(InUsage, InUsageId, InReferencedCompileHashes, InReferencedIDs, InReferencedObjs);
-				}
+				Node->GatherExternalDependencyData(InUsage, InUsageId, InReferencedCompileHashes, InReferencedObjs);
 			}
 		}
-	}
-	// Otherwise, just add downstream dependencies for the specific usage type we're on.
-	else
-	{
-		for (int32 i = 0; i < CachedUsageInfo.Num(); i++)
+		// Now add any other dependency chains that we might have...
+		else if (UNiagaraScript::IsUsageDependentOn(InUsage, CachedUsageInfo[i].UsageType))
 		{
-			// First add our direct dependency chain...
-			if (UNiagaraScript::IsEquivalentUsage(CachedUsageInfo[i].UsageType, InUsage) && CachedUsageInfo[i].UsageId == InUsageId)
-			{
-				// Skip adding to list because we already did it in GetCompileId above.
-				for (UNiagaraNode* Node : CachedUsageInfo[i].Traversal)
-				{
-					Node->GatherExternalDependencyIDs(InUsage, InUsageId, InReferencedCompileHashes, InReferencedIDs, InReferencedObjs);
-				}
-			}
-			// Now add any other dependency chains that we might have...
-			else if (UNiagaraScript::IsUsageDependentOn(InUsage, CachedUsageInfo[i].UsageType))
-			{
-				InReferencedCompileHashes.Add(CachedUsageInfo[i].CompileHash);
-				InReferencedObjs.Add(CachedUsageInfo[i].Traversal.Last());
+			InReferencedCompileHashes.Add(CachedUsageInfo[i].CompileHash);
+			InReferencedObjs.Add(CachedUsageInfo[i].Traversal.Last());
 
-				for (UNiagaraNode* Node : CachedUsageInfo[i].Traversal)
-				{
-					Node->GatherExternalDependencyIDs(InUsage, InUsageId, InReferencedCompileHashes, InReferencedIDs, InReferencedObjs);
-				}
+			for (UNiagaraNode* Node : CachedUsageInfo[i].Traversal)
+			{
+				Node->GatherExternalDependencyData(InUsage, InUsageId, InReferencedCompileHashes, InReferencedObjs);
 			}
 		}
 	}
diff --git a/Engine/Plugins/FX/Niagara/Source/NiagaraEditor/Private/NiagaraNodeAssignment.cpp b/Engine/Plugins/FX/Niagara/Source/NiagaraEditor/Private/NiagaraNodeAssignment.cpp
--- a/Engine/Plugins/FX/Niagara/Source/NiagaraEditor/Private/NiagaraNodeAssignment.cpp
+++ b/Engine/Plugins/FX/Niagara/Source/NiagaraEditor/Private/NiagaraNodeAssignment.cpp
@@ -177,7 +177,7 @@ void UNiagaraNodeAssignment::BuildParameterMapHistory(FNiagaraParameterMapHistor
 	Super::BuildParameterMapHistory(OutHistory, bRecursive, bFilterForCompilation);
 }
 
-void UNiagaraNodeAssignment::GatherExternalDependencyIDs(ENiagaraScriptUsage InMasterUsage, const FGuid& InMasterUsageId, TArray<FNiagaraCompileHash>& InReferencedCompileHashes, TArray<FGuid>& InReferencedIDs, TArray<UObject*>& InReferencedObjs) const
+void UNiagaraNodeAssignment::GatherExternalDependencyData(ENiagaraScriptUsage InMasterUsage, const FGuid& InMasterUsageId, TArray<FNiagaraCompileHash>& InReferencedCompileHashes, TArray<UObject*>& InReferencedObjs) const
 {
 	// Assignment nodes own their function graphs and therefore have no external dependencies so we override the default function behavior here to avoid 
 	// adding additional non-deterministic guids to the compile id generation which can invalid the DDC for compiled scripts, especially during emitter merging.
diff --git a/Engine/Plugins/FX/Niagara/Source/NiagaraEditor/Private/NiagaraNodeEmitter.cpp b/Engine/Plugins/FX/Niagara/Source/NiagaraEditor/Private/NiagaraNodeEmitter.cpp
--- a/Engine/Plugins/FX/Niagara/Source/NiagaraEditor/Private/NiagaraNodeEmitter.cpp
+++ b/Engine/Plugins/FX/Niagara/Source/NiagaraEditor/Private/NiagaraNodeEmitter.cpp
@@ -460,7 +460,7 @@ void UNiagaraNodeEmitter::Compile(FHlslNiagaraTranslator *Translator, TArray<int
 	}
 }
 
-void UNiagaraNodeEmitter::GatherExternalDependencyIDs(ENiagaraScriptUsage InMasterUsage, const FGuid& InMasterUsageId, TArray<FNiagaraCompileHash>& InReferencedCompileHashes, TArray<FGuid>& InReferencedIDs, TArray<UObject*>& InReferencedObjs) const
+void UNiagaraNodeEmitter::GatherExternalDependencyData(ENiagaraScriptUsage InMasterUsage, const FGuid& InMasterUsageId, TArray<FNiagaraCompileHash>& InReferencedCompileHashes, TArray<UObject*>& InReferencedObjs) const
 {
 	UNiagaraGraph* CalledGraph = GetCalledGraph();
 
@@ -469,7 +469,7 @@ void UNiagaraNodeEmitter::GatherExternalDependencyIDs(ENiagaraScriptUsage InMast
 		ENiagaraScriptUsage TargetUsage = InMasterUsage == ENiagaraScriptUsage::SystemSpawnScript ? ENiagaraScriptUsage::EmitterSpawnScript : ENiagaraScriptUsage::EmitterUpdateScript;
 		InReferencedCompileHashes.Add(CalledGraph->GetCompileDataHash(TargetUsage, FGuid(0,0,0,0)));
 		InReferencedObjs.Add(CalledGraph);
-		CalledGraph->GatherExternalDependencyIDs(TargetUsage, FGuid(0, 0, 0, 0), InReferencedCompileHashes, InReferencedIDs, InReferencedObjs);
+		CalledGraph->GatherExternalDependencyData(TargetUsage, FGuid(0, 0, 0, 0), InReferencedCompileHashes, InReferencedObjs);
 	}
 }
 
diff --git a/Engine/Plugins/FX/Niagara/Source/NiagaraEditor/Private/NiagaraNodeEmitter.h b/Engine/Plugins/FX/Niagara/Source/NiagaraEditor/Private/NiagaraNodeEmitter.h
--- a/Engine/Plugins/FX/Niagara/Source/NiagaraEditor/Private/NiagaraNodeEmitter.h
+++ b/Engine/Plugins/FX/Niagara/Source/NiagaraEditor/Private/NiagaraNodeEmitter.h
@@ -60,7 +60,7 @@ public:
 	UNiagaraGraph* GetCalledGraph() const;
 	
 	virtual void Compile(FHlslNiagaraTranslator *Translator, TArray<int32>& Outputs) override;
-	virtual void GatherExternalDependencyIDs(ENiagaraScriptUsage InMasterUsage, const FGuid& InMasterUsageId, TArray<FNiagaraCompileHash>& InReferencedCompileHashes, TArray<FGuid>& InReferencedIDs, TArray<UObject*>& InReferencedObjs) const override;
+	virtual void GatherExternalDependencyData(ENiagaraScriptUsage InMasterUsage, const FGuid& InMasterUsageId, TArray<FNiagaraCompileHash>& InReferencedCompileHashes, TArray<UObject*>& InReferencedObjs) const override;
 
 	void SetCachedVariablesForCompilation(const FName& InUniqueName, UNiagaraGraph* InGraph, UNiagaraScriptSourceBase* InSource);
 
diff --git a/Engine/Plugins/FX/Niagara/Source/NiagaraEditor/Private/NiagaraNodeFunctionCall.cpp b/Engine/Plugins/FX/Niagara/Source/NiagaraEditor/Private/NiagaraNodeFunctionCall.cpp
--- a/Engine/Plugins/FX/Niagara/Source/NiagaraEditor/Private/NiagaraNodeFunctionCall.cpp
+++ b/Engine/Plugins/FX/Niagara/Source/NiagaraEditor/Private/NiagaraNodeFunctionCall.cpp
@@ -681,7 +681,7 @@ void UNiagaraNodeFunctionCall::SubsumeExternalDependencies(TMap<const UObject*,
 	}
 }
 
-void UNiagaraNodeFunctionCall::GatherExternalDependencyIDs(ENiagaraScriptUsage InMasterUsage, const FGuid& InMasterUsageId, TArray<FNiagaraCompileHash>& InReferencedCompileHashes, TArray<FGuid>& InReferencedIDs, TArray<UObject*>& InReferencedObjs) const
+void UNiagaraNodeFunctionCall::GatherExternalDependencyData(ENiagaraScriptUsage InMasterUsage, const FGuid& InMasterUsageId, TArray<FNiagaraCompileHash>& InReferencedCompileHashes, TArray<UObject*>& InReferencedObjs) const
 {
 	if (FunctionScript && FunctionScript->GetOutermost() != this->GetOutermost())
 	{
@@ -691,15 +691,16 @@ void UNiagaraNodeFunctionCall::GatherExternalDependencyIDs(ENiagaraScriptUsage I
 		// We don't know which graph type we're referencing, so we try them all... may need to replace this with something faster in the future.
 		if (FunctionGraph)
 		{
+			FunctionGraph->RebuildCachedCompileIds();
 			for (int32 i = (int32)ENiagaraScriptUsage::Function; i <= (int32)ENiagaraScriptUsage::DynamicInput; i++)
 			{
-				FGuid FoundGuid = FunctionGraph->ComputeCompileID((ENiagaraScriptUsage)i, FGuid(0, 0, 0, 0));
-				if (FoundGuid.IsValid())
+				FGuid FoundGuid = FunctionGraph->GetBaseId((ENiagaraScriptUsage)i, FGuid(0, 0, 0, 0));
+				FNiagaraCompileHash FoundCompileHash = FunctionGraph->GetCompileDataHash((ENiagaraScriptUsage)i, FGuid(0, 0, 0, 0));
+				if (FoundGuid.IsValid() && FoundCompileHash.IsValid())
 				{
-					InReferencedIDs.Add(FoundGuid);
+					InReferencedCompileHashes.Add(FoundCompileHash);
 					InReferencedObjs.Add(FunctionGraph);
-
-					FunctionGraph->GatherExternalDependencyIDs((ENiagaraScriptUsage)i, FGuid(0, 0, 0, 0), InReferencedCompileHashes, InReferencedIDs, InReferencedObjs);
+					FunctionGraph->GatherExternalDependencyData((ENiagaraScriptUsage)i, FGuid(0, 0, 0, 0), InReferencedCompileHashes, InReferencedObjs);
 				}
 			}
 		}
diff --git a/Engine/Plugins/FX/Niagara/Source/NiagaraEditor/Private/NiagaraNodeParameterMapGet.cpp b/Engine/Plugins/FX/Niagara/Source/NiagaraEditor/Private/NiagaraNodeParameterMapGet.cpp
--- a/Engine/Plugins/FX/Niagara/Source/NiagaraEditor/Private/NiagaraNodeParameterMapGet.cpp
+++ b/Engine/Plugins/FX/Niagara/Source/NiagaraEditor/Private/NiagaraNodeParameterMapGet.cpp
@@ -414,7 +414,7 @@ bool UNiagaraNodeParameterMapGet::CommitEditablePinName(const FText& InName, UEd
 	return false;
 }
 
-void UNiagaraNodeParameterMapGet::GatherExternalDependencyIDs(ENiagaraScriptUsage InMasterUsage, const FGuid& InMasterUsageId, TArray<FNiagaraCompileHash>& InReferencedCompileHashes, TArray<FGuid>& InReferencedIDs, TArray<UObject*>& InReferencedObjs) const
+void UNiagaraNodeParameterMapGet::GatherExternalDependencyData(ENiagaraScriptUsage InMasterUsage, const FGuid& InMasterUsageId, TArray<FNiagaraCompileHash>& InReferencedCompileHashes, TArray<UObject*>& InReferencedObjs) const
 {
 	// If we are referencing any parameter collections, we need to register them here... might want to speeed this up in the future 
 	// by caching any parameter collections locally.
@@ -433,10 +433,9 @@ void UNiagaraNodeParameterMapGet::GatherExternalDependencyIDs(ENiagaraScriptUsag
 		UNiagaraParameterCollection* Collection = Schema->VariableIsFromParameterCollection(Var);
 		if (Collection)
 		{
-			InReferencedIDs.Add(Collection->GetCompileId());
+			InReferencedCompileHashes.Add(Collection->GetCompileHash());
 			InReferencedObjs.Add(Collection);
 		}
-
 	}
 }
 
diff --git a/Engine/Plugins/FX/Niagara/Source/NiagaraEditor/Private/NiagaraNodeParameterMapGet.h b/Engine/Plugins/FX/Niagara/Source/NiagaraEditor/Private/NiagaraNodeParameterMapGet.h
--- a/Engine/Plugins/FX/Niagara/Source/NiagaraEditor/Private/NiagaraNodeParameterMapGet.h
+++ b/Engine/Plugins/FX/Niagara/Source/NiagaraEditor/Private/NiagaraNodeParameterMapGet.h
@@ -39,7 +39,7 @@ public:
 
 	virtual void PostLoad() override;
 
-	void GatherExternalDependencyIDs(ENiagaraScriptUsage InMasterUsage, const FGuid& InMasterUsageId, TArray<FNiagaraCompileHash>& InReferencedCompileHashes, TArray<FGuid>& InReferencedIDs, TArray<UObject*>& InReferencedObjs) const override;
+	void GatherExternalDependencyData(ENiagaraScriptUsage InMasterUsage, const FGuid& InMasterUsageId, TArray<FNiagaraCompileHash>& InReferencedCompileHashes, TArray<UObject*>& InReferencedObjs) const override;
 	virtual void GetPinHoverText(const UEdGraphPin& Pin, FString& HoverTextOut) const override;
 
 protected:
diff --git a/Engine/Plugins/FX/Niagara/Source/NiagaraEditor/Private/NiagaraScriptMergeManager.cpp b/Engine/Plugins/FX/Niagara/Source/NiagaraEditor/Private/NiagaraScriptMergeManager.cpp
--- a/Engine/Plugins/FX/Niagara/Source/NiagaraEditor/Private/NiagaraScriptMergeManager.cpp
+++ b/Engine/Plugins/FX/Niagara/Source/NiagaraEditor/Private/NiagaraScriptMergeManager.cpp
@@ -796,19 +796,6 @@ FNiagaraScriptMergeManager::FApplyDiffResults FNiagaraScriptMergeManager::Resolv
 
 		DiffResults.bModifiedGraph = bAnySet;
 
-		UNiagaraScriptSource* ScriptSource = Cast<UNiagaraScriptSource>(OriginalEmitterInstance.GraphSource);
-		if (ScriptSource && bAnySet)
-		{
-			UNiagaraGraph* OtherGraph = ScriptSource->NodeGraph;
-			if (OtherGraph)
-			{
-				for (UNiagaraGraph* Graph : Graphs)
-				{
-					Graph->SynchronizeInternalCacheWithGraph(OtherGraph);
-				}
-			}
-		}
-
 		if (bAnySet)
 		{
 			bool bAnyUpdated = false;
diff --git a/Engine/Plugins/FX/Niagara/Source/NiagaraEditor/Private/NiagaraScriptSource.cpp b/Engine/Plugins/FX/Niagara/Source/NiagaraEditor/Private/NiagaraScriptSource.cpp
--- a/Engine/Plugins/FX/Niagara/Source/NiagaraEditor/Private/NiagaraScriptSource.cpp
+++ b/Engine/Plugins/FX/Niagara/Source/NiagaraEditor/Private/NiagaraScriptSource.cpp
@@ -53,7 +53,7 @@ void UNiagaraScriptSource::ComputeVMCompilationId(FNiagaraVMExecutableDataId& Id
 		NodeGraph->RebuildCachedCompileIds(bForceRebuild);
 		Id.BaseScriptID = NodeGraph->GetBaseId(InUsage, InUsageId);
 		Id.BaseScriptCompileHash = FNiagaraCompileHash(NodeGraph->GetCompileDataHash(InUsage, InUsageId));
-		NodeGraph->GatherExternalDependencyIDs(InUsage, InUsageId, Id.ReferencedCompileHashes, Id.ReferencedDependencyIds, Id.ReferencedObjects);
+		NodeGraph->GatherExternalDependencyData(InUsage, InUsageId, Id.ReferencedCompileHashes, Id.ReferencedObjects);
 	}
 }
 
diff --git a/Engine/Plugins/FX/Niagara/Source/NiagaraEditor/Public/NiagaraGraph.h b/Engine/Plugins/FX/Niagara/Source/NiagaraEditor/Public/NiagaraGraph.h
--- a/Engine/Plugins/FX/Niagara/Source/NiagaraEditor/Public/NiagaraGraph.h
+++ b/Engine/Plugins/FX/Niagara/Source/NiagaraEditor/Public/NiagaraGraph.h
@@ -63,10 +63,6 @@ public:
 	UPROPERTY()
 	FGuid UsageId;
 
-	/** The compile ID last associated with this traversal.*/
-	UPROPERTY()
-	FGuid GeneratedCompileId;
-
 	/** The hash that we calculated last traversal. */
 	UPROPERTY()
 	FNiagaraCompileHash CompileHash;
@@ -80,6 +76,9 @@ public:
 private:
 	UPROPERTY()
 	TArray<uint8> DataHash_DEPRECATED;
+
+	UPROPERTY()
+	FGuid GeneratedCompileId_DEPRECATED;
 };
 
 struct FNiagaraGraphFunctionAliasContext
@@ -182,7 +181,7 @@ class UNiagaraGraph : public UEdGraph
 	void GetAllReferencedGraphs(TArray<const UNiagaraGraph*>& Graphs) const;
 
 	/** Gather all the change ids of external references for this specific graph traversal.*/
-	void GatherExternalDependencyIDs(ENiagaraScriptUsage InUsage, const FGuid& InUsageId, TArray<FNiagaraCompileHash>& InReferencedCompileHashes, TArray<FGuid>& InReferencedIDs, TArray<UObject*>& InReferencedObjs);
+	void GatherExternalDependencyData(ENiagaraScriptUsage InUsage, const FGuid& InUsageId, TArray<FNiagaraCompileHash>& InReferencedCompileHashes, TArray<UObject*>& InReferencedObjs);
 
 	/** Determine if there are any external dependencies wrt to scripts and ensure that those dependencies are sucked into the existing package.*/
 	void SubsumeExternalDependencies(TMap<const UObject*, UObject*>& ExistingConversions);
@@ -201,9 +200,6 @@ class UNiagaraGraph : public UEdGraph
 	the changes that were produced during the traversal of each output node, which are referred to as the CompileID.*/
 	FGuid GetChangeID() { return ChangeId; }
 
-	/** Recomputes the current compile id associated with the output node traversal specified by InUsage and InUsageId. If the usage is not found, an invalid Guid is returned.*/
-	FGuid ComputeCompileID(ENiagaraScriptUsage InUsage, const FGuid& InUsageId);
-
 	/** Gets the current compile data hash associated with the output node traversal specified by InUsage and InUsageId. If the usage is not found, an invalid hash is returned.*/
 	FNiagaraCompileHash GetCompileDataHash(ENiagaraScriptUsage InUsage, const FGuid& InUsageId) const;
 
@@ -242,7 +238,6 @@ class UNiagaraGraph : public UEdGraph
 	/** Gets a delegate which is called whenever a contained data interfaces changes. */
 	FOnDataInterfaceChanged& OnDataInterfaceChanged();
 
-	void SynchronizeInternalCacheWithGraph(UNiagaraGraph* Other);
 	void InvalidateCachedCompileIds();
 
 	/** Add a listener for OnGraphNeedsRecompile events */
diff --git a/Engine/Plugins/FX/Niagara/Source/NiagaraEditor/Public/NiagaraNode.h b/Engine/Plugins/FX/Niagara/Source/NiagaraEditor/Public/NiagaraNode.h
--- a/Engine/Plugins/FX/Niagara/Source/NiagaraEditor/Public/NiagaraNode.h
+++ b/Engine/Plugins/FX/Niagara/Source/NiagaraEditor/Public/NiagaraNode.h
@@ -100,7 +100,7 @@ public:
 	virtual void BuildParameterMapHistory(FNiagaraParameterMapHistoryBuilder& OutHistory, bool bRecursive = true, bool bFilterForCompilation = true) const;
 	
 	/** Go through all the external dependencies of this node in isolation and add them to the reference id list.*/
-	virtual void GatherExternalDependencyIDs(ENiagaraScriptUsage InMasterUsage, const FGuid& InMasterUsageId, TArray<FNiagaraCompileHash>& InReferencedCompileHashes, TArray<FGuid>& InReferencedIDs, TArray<UObject*>& InReferencedObjs) const {};
+	virtual void GatherExternalDependencyData(ENiagaraScriptUsage InMasterUsage, const FGuid& InMasterUsageId, TArray<FNiagaraCompileHash>& InReferencedCompileHashes, TArray<UObject*>& InReferencedObjs) const {};
 
 	/** Traces one of this node's output pins to its source output pin if it is a reroute node output pin.*/
 	virtual UEdGraphPin* GetTracedOutputPin(UEdGraphPin* LocallyOwnedOutputPin) const {return LocallyOwnedOutputPin;}
diff --git a/Engine/Plugins/FX/Niagara/Source/NiagaraEditor/Public/NiagaraNodeAssignment.h b/Engine/Plugins/FX/Niagara/Source/NiagaraEditor/Public/NiagaraNodeAssignment.h
--- a/Engine/Plugins/FX/Niagara/Source/NiagaraEditor/Public/NiagaraNodeAssignment.h
+++ b/Engine/Plugins/FX/Niagara/Source/NiagaraEditor/Public/NiagaraNodeAssignment.h
@@ -45,7 +45,7 @@ public:
 	//~ UNiagaraNodeFunctionCall interface
 	virtual bool RefreshFromExternalChanges() override;
 	virtual void BuildParameterMapHistory(FNiagaraParameterMapHistoryBuilder& OutHistory, bool bRecursive = true, bool bFilterForCompilation = true) const override;
-	virtual void GatherExternalDependencyIDs(ENiagaraScriptUsage InMasterUsage, const FGuid& InMasterUsageId, TArray<FNiagaraCompileHash>& InReferencedCompileHashes, TArray<FGuid>& InReferencedIDs, TArray<UObject*>& InReferencedObjs) const override;
+	virtual void GatherExternalDependencyData(ENiagaraScriptUsage InMasterUsage, const FGuid& InMasterUsageId, TArray<FNiagaraCompileHash>& InReferencedCompileHashes, TArray<UObject*>& InReferencedObjs) const override;
 
 	void AddParameter(FNiagaraVariable InVar, FString InDefaultValue);
 	void RemoveParameter(const FNiagaraVariable& InVar);
diff --git a/Engine/Plugins/FX/Niagara/Source/NiagaraEditor/Public/NiagaraNodeFunctionCall.h b/Engine/Plugins/FX/Niagara/Source/NiagaraEditor/Public/NiagaraNodeFunctionCall.h
--- a/Engine/Plugins/FX/Niagara/Source/NiagaraEditor/Public/NiagaraNodeFunctionCall.h
+++ b/Engine/Plugins/FX/Niagara/Source/NiagaraEditor/Public/NiagaraNodeFunctionCall.h
@@ -91,7 +91,7 @@ public:
 	virtual ENiagaraNumericOutputTypeSelectionMode GetNumericOutputTypeSelectionMode() const override;
 	virtual bool CanAddToGraph(UNiagaraGraph* TargetGraph, FString& OutErrorMsg) const override;
 	virtual void SubsumeExternalDependencies(TMap<const UObject*, UObject*>& ExistingConversions) override;
-	virtual void GatherExternalDependencyIDs(ENiagaraScriptUsage InMasterUsage, const FGuid& InMasterUsageId, TArray<FNiagaraCompileHash>& InReferencedCompileHashes, TArray<FGuid>& InReferencedIDs, TArray<UObject*>& InReferencedObjs) const override;
+	virtual void GatherExternalDependencyData(ENiagaraScriptUsage InMasterUsage, const FGuid& InMasterUsageId, TArray<FNiagaraCompileHash>& InReferencedCompileHashes, TArray<UObject*>& InReferencedObjs) const override;
 	virtual void UpdateCompileHashForNode(FSHA1& HashState) const override;
 	//End UNiagaraNode interface
 
diff --git a/Engine/Plugins/FX/Niagara/Source/NiagaraShader/Private/NiagaraShader.cpp b/Engine/Plugins/FX/Niagara/Source/NiagaraShader/Private/NiagaraShader.cpp
--- a/Engine/Plugins/FX/Niagara/Source/NiagaraShader/Private/NiagaraShader.cpp
+++ b/Engine/Plugins/FX/Niagara/Source/NiagaraShader/Private/NiagaraShader.cpp
@@ -122,7 +122,14 @@ void FNiagaraShaderMapId::Serialize(FArchive& Ar)
 
 	Ar << BaseScriptID;
 	Ar << (int32&)FeatureLevel;
-	Ar << ReferencedDependencyIds;
+
+	if (Ar.IsLoading() && NiagaraVer < FNiagaraCustomVersion::RemoveGraphUsageCompileIds)
+	{
+		// These values are no longer used, but need to be loaded from old files to keep the archive valid.
+		TArray<FGuid> ReferencedDependencyIds;
+		Ar << ReferencedDependencyIds;
+	}
+
 	if (NiagaraVer >= FNiagaraCustomVersion::UseHashesToIdentifyCompileStateOfTopLevelScripts)
 	{
 		Ar << BaseCompileHash;
@@ -156,11 +163,6 @@ void FNiagaraShaderMapId::GetScriptHash(FSHAHash& OutHash) const
 		HashState.Update(ReferencedCompileHashes[Index].GetData(), FNiagaraCompileHash::HashSize);
 	}
 
-	for (int32 Index = 0; Index < ReferencedDependencyIds.Num(); Index++)
-	{
-		HashState.Update((const uint8*)&ReferencedDependencyIds[Index], sizeof(ReferencedDependencyIds[Index]));
-	}
-
 	//ParameterSet.UpdateHash(HashState);		// will need for static switches
 	
 	HashState.Final();
@@ -208,18 +210,6 @@ bool FNiagaraShaderMapId::operator==(const FNiagaraShaderMapId& ReferenceSet) co
 		}
 	}
 
-	if (ReferencedDependencyIds.Num() != ReferenceSet.ReferencedDependencyIds.Num())
-	{
-		return false;
-	}
-
-	for (int32 i = 0; i < ReferencedDependencyIds.Num(); i++)
-	{
-		if (ReferencedDependencyIds[i] != ReferenceSet.ReferencedDependencyIds[i])
-		{
-			return false;
-		}
-	}
 	/*
 	if (ParameterSet != ReferenceSet.ParameterSet
 		|| ReferencedFunctions.Num() != ReferenceSet.ReferencedFunctions.Num()
@@ -275,17 +265,6 @@ void FNiagaraShaderMapId::AppendKeyString(FString& KeyString) const
 		}
 	}
 
-	// Add any referenced functions to the key so that we will recompile when they are changed
-	for (int32 FunctionIndex = 0; FunctionIndex < ReferencedDependencyIds.Num(); FunctionIndex++)
-	{
-		KeyString += ReferencedDependencyIds[FunctionIndex].ToString();
-
-		if (FunctionIndex < ReferencedDependencyIds.Num() - 1)
-		{
-			KeyString += TEXT("_");
-		}
-	}
-
 	/*
 	ParameterSet.AppendKeyString(KeyString);
 
diff --git a/Engine/Plugins/FX/Niagara/Source/NiagaraShader/Private/NiagaraShared.cpp b/Engine/Plugins/FX/Niagara/Source/NiagaraShader/Private/NiagaraShared.cpp
--- a/Engine/Plugins/FX/Niagara/Source/NiagaraShader/Private/NiagaraShared.cpp
+++ b/Engine/Plugins/FX/Niagara/Source/NiagaraShader/Private/NiagaraShared.cpp
@@ -92,7 +92,6 @@ bool FNiagaraShaderScript::IsSame(const FNiagaraShaderMapId& InId) const
 		InId.AdditionalDefines == AdditionalDefines &&
 		InId.BaseCompileHash == BaseCompileHash &&
 		InId.ReferencedCompileHashes == ReferencedCompileHashes &&
-		InId.ReferencedDependencyIds == ReferencedDependencyIds &&
 		InId.CompilerVersionID == CompilerVersionId;
 }
 
@@ -127,7 +126,6 @@ NIAGARASHADER_API void FNiagaraShaderScript::GetShaderMapId(EShaderPlatform Plat
 		OutId.AdditionalDefines = AdditionalDefines;
 		OutId.BaseCompileHash = BaseCompileHash;
 		OutId.ReferencedCompileHashes = ReferencedCompileHashes;
-		OutId.ReferencedDependencyIds = ReferencedDependencyIds;
 		OutId.CompilerVersionID = FNiagaraCustomVersion::LatestScriptCompileVersion;
 	}
 }
@@ -227,7 +225,7 @@ void FNiagaraShaderScript::SerializeShaderMap(FArchive& Ar)
 }
 
 void FNiagaraShaderScript::SetScript(UNiagaraScript *InScript, ERHIFeatureLevel::Type InFeatureLevel, const FGuid& InCompilerVersionID, const FGuid& InBaseScriptID, const TArray<FString>& InAdditionalDefines,
-	const FNiagaraCompileHash& InBaseCompileHash, const TArray<FNiagaraCompileHash>& InReferencedCompileHashes, const TArray<FGuid>& InReferencedDependencyIds, FString InFriendlyName)
+	const FNiagaraCompileHash& InBaseCompileHash, const TArray<FNiagaraCompileHash>& InReferencedCompileHashes, FString InFriendlyName)
 {
 	checkf(InBaseScriptID.IsValid(), TEXT("Invalid base script id.  Script caching will fail."));
 	checkf(InBaseCompileHash.IsValid(), TEXT("Invalid base compile hash.  Script caching will fail."))
@@ -237,7 +235,6 @@ void FNiagaraShaderScript::SetScript(UNiagaraScript *InScript, ERHIFeatureLevel:
 	AdditionalDefines = InAdditionalDefines;
 	BaseCompileHash = InBaseCompileHash;
 	ReferencedCompileHashes = InReferencedCompileHashes;
-	ReferencedDependencyIds = InReferencedDependencyIds;
 	FriendlyName = InFriendlyName;
 	SetFeatureLevel(InFeatureLevel);
 }
diff --git a/Engine/Plugins/FX/Niagara/Source/NiagaraShader/Public/NiagaraShared.h b/Engine/Plugins/FX/Niagara/Source/NiagaraShader/Public/NiagaraShared.h
--- a/Engine/Plugins/FX/Niagara/Source/NiagaraShader/Public/NiagaraShared.h
+++ b/Engine/Plugins/FX/Niagara/Source/NiagaraShader/Public/NiagaraShared.h
@@ -172,9 +172,6 @@ public:
 	/** The compile hashes of the top level scripts the script is dependent on. */
 	TArray<FNiagaraCompileHash> ReferencedCompileHashes;
 
-	/** Guids of any functions or module scripts the script was dependent on. */
-	TArray<FGuid> ReferencedDependencyIds;
-
 	FNiagaraShaderMapId()
 		: CompilerVersionID()
 		, FeatureLevel(GMaxRHIFeatureLevel)
@@ -670,7 +667,7 @@ public:
 
 
 	NIAGARASHADER_API void SetScript(UNiagaraScript *InScript, ERHIFeatureLevel::Type InFeatureLevel, const FGuid& InCompilerVersion, const FGuid& InBaseScriptID, const TArray<FString>& InAdditionalDefines,
-		const FNiagaraCompileHash& InBaseCompileHash, const TArray<FNiagaraCompileHash>& InReferencedCompileHashes, const TArray<FGuid>& InReferencedDependencyIds, FString InFriendlyName);
+		const FNiagaraCompileHash& InBaseCompileHash, const TArray<FNiagaraCompileHash>& InReferencedCompileHashes, FString InFriendlyName);
 
 	UNiagaraScript *GetBaseVMScript()
 	{
@@ -757,9 +754,6 @@ private:
 	/** The compile hashes for the top level scripts referenced by the script. */
 	TArray<FNiagaraCompileHash> ReferencedCompileHashes;
 
-	/** Dependencies of the script*/
-	TArray<FGuid> ReferencedDependencyIds;
-
 	/** 
 	 * Contains the compiling id of this shader map when it is being compiled asynchronously. 
 	 * This can be used to access the shader map during async compiling, since GameThreadShaderMap will not have been set yet.
