Make sure FNiagaraParameterStore::ParameterData is always properly padded with at least VECTOR_WIDTH_BYTES in addition to being aligned to VECTOR_WIDTH_BYTES when being sent to the VectorVM. This is required to not go out of bounds on unaligned reads.

diff --git a/Engine/Plugins/FX/Niagara/Source/Niagara/Private/NiagaraParameterStore.cpp b/Engine/Plugins/FX/Niagara/Source/Niagara/Private/NiagaraParameterStore.cpp
--- a/Engine/Plugins/FX/Niagara/Source/Niagara/Private/NiagaraParameterStore.cpp
+++ b/Engine/Plugins/FX/Niagara/Source/Niagara/Private/NiagaraParameterStore.cpp
@@ -586,6 +586,15 @@ void FNiagaraParameterStore::Reset(bool bClearBindings)
 
 void FNiagaraParameterStore::OnLayoutChange()
 {
+	// The VM require that the parameter data we send it in FNiagaraScriptExecutionContext::Execute
+	// is aligned to VECTOR_WIDTH_BYTES *and* is padded with an additional VECTOR_WIDTH_BYTES.
+	// This is due to possible unaligned reads, e.g. an integer might be stored in the very last byte
+	// of the aligned parameter data due to the packing, which will spill 3 bytes outside the bounds
+	int32 ExpectedSlack = Align(ParameterData.Num(), VECTOR_WIDTH_BYTES) + VECTOR_WIDTH_BYTES;
+	if (ParameterData.Max() < ExpectedSlack)
+	{
+		ParameterData.Reserve(ExpectedSlack);
+	}
 	Rebind();
 	++LayoutVersion;
 
