Removed unused TUniformParameter to avoid redefinition errors in unity build should the user have both NiagaraShader.h and MaterialShader.h in the same unity blob.

#rnx

diff --git a/Engine/Plugins/FX/Niagara/Source/NiagaraShader/Public/NiagaraShader.h b/Engine/Plugins/FX/Niagara/Source/NiagaraShader/Public/NiagaraShader.h
--- a/Engine/Plugins/FX/Niagara/Source/NiagaraShader/Public/NiagaraShader.h
+++ b/Engine/Plugins/FX/Niagara/Source/NiagaraShader/Public/NiagaraShader.h
@@ -22,17 +22,6 @@ class UClass;
 
 template<typename TBufferStruct> class TUniformBufferRef;
 
-template<typename ParameterType> 
-struct TUniformParameter
-{
-	int32 Index;
-	ParameterType ShaderParameter;
-	friend FArchive& operator<<(FArchive& Ar,TUniformParameter<ParameterType>& P)
-	{
-		return Ar << P.Index << P.ShaderParameter;
-	}
-};
-
 /** Base class of all shaders that need material parameters. */
 class NIAGARASHADER_API FNiagaraShader : public FShader
 {
diff --git a/Engine/Source/Runtime/Renderer/Public/MaterialShader.h b/Engine/Source/Runtime/Renderer/Public/MaterialShader.h
--- a/Engine/Source/Runtime/Renderer/Public/MaterialShader.h
+++ b/Engine/Source/Runtime/Renderer/Public/MaterialShader.h
@@ -20,18 +20,6 @@
 
 template<typename TBufferStruct> class TUniformBufferRef;
 
-template<typename ParameterType> 
-struct TUniformParameter
-{
-	int32 Index;
-	ParameterType ShaderParameter;
-	friend FArchive& operator<<(FArchive& Ar,TUniformParameter<ParameterType>& P)
-	{
-		return Ar << P.Index << P.ShaderParameter;
-	}
-};
-
-
 /**
  * Debug information related to uniform expression sets.
  */
