Merge CIS and VS2019 fixes from UE4/Main
#rnx

diff --git a/Engine/Plugins/FX/Niagara/Source/NiagaraEditor/Private/NiagaraNodeOutput.cpp b/Engine/Plugins/FX/Niagara/Source/NiagaraEditor/Private/NiagaraNodeOutput.cpp
--- a/Engine/Plugins/FX/Niagara/Source/NiagaraEditor/Private/NiagaraNodeOutput.cpp
+++ b/Engine/Plugins/FX/Niagara/Source/NiagaraEditor/Private/NiagaraNodeOutput.cpp
@@ -96,7 +96,7 @@ void UNiagaraNodeOutput::GetContextMenuActions(const FGraphNodeContextMenuBuilde
 					SNew(SEditableTextBoxWithVerification)
 					.Text_UObject(this, &UNiagaraNodeOutput::GetPinNameText, Pin)
 					.OnVerifyTextChanged_UObject(this, &UNiagaraNodeOutput::VerifyPinNameTextChanged, Pin)
-					.OnTextCommitted_UObject(this, &UNiagaraNodeOutput::PinNameTextCommitted, Pin)
+					.OnTextCommitted_UObject(const_cast<UNiagaraNodeOutput*>(this), &UNiagaraNodeOutput::PinNameTextCommitted, Pin)
 				];
 			Context.MenuBuilder->AddWidget(RenameWidget, LOCTEXT("NameMenuItem", "Name"));
 		}
diff --git a/Engine/Source/Runtime/Experimental/Chaos/Public/Chaos/ParticleIterator.h b/Engine/Source/Runtime/Experimental/Chaos/Public/Chaos/ParticleIterator.h
--- a/Engine/Source/Runtime/Experimental/Chaos/Public/Chaos/ParticleIterator.h
+++ b/Engine/Source/Runtime/Experimental/Chaos/Public/Chaos/ParticleIterator.h
@@ -251,7 +251,7 @@ public:
 
 	TConstParticleIterator(const TConstParticleIterator& Rhs) = default;
 
-	operator bool() const { return TransientHandle.GeometryParticles; }
+	operator bool() const { return TransientHandle.GeometryParticles != nullptr; }
 	TConstParticleIterator<TSOA>& operator++()
 	{
 		RangedForValidation();
diff --git a/Engine/Source/Runtime/Experimental/ChaosSolvers/Public/PBDRigidsSolver.h b/Engine/Source/Runtime/Experimental/ChaosSolvers/Public/PBDRigidsSolver.h
--- a/Engine/Source/Runtime/Experimental/ChaosSolvers/Public/PBDRigidsSolver.h
+++ b/Engine/Source/Runtime/Experimental/ChaosSolvers/Public/PBDRigidsSolver.h
@@ -179,7 +179,7 @@ namespace Chaos
 			});
 		}
 
-		bool GetNumPhysicsProxies() const {
+		int32 GetNumPhysicsProxies() const {
 			return GeometryParticlePhysicsProxies.Num() + KinematicGeometryParticlePhysicsProxies.Num() + RigidParticlePhysicsProxies.Num()
 				+ SkeletalMeshPhysicsProxies.Num() + StaticMeshPhysicsProxies.Num()
 				+ GeometryCollectionPhysicsProxies.Num() + FieldSystemPhysicsProxies.Num() ;
