Fix If node crash when compiling after adding and deleting inputs or outputs.

diff --git a/Engine/Plugins/FX/Niagara/Source/NiagaraEditor/Private/NiagaraNodeIf.cpp b/Engine/Plugins/FX/Niagara/Source/NiagaraEditor/Private/NiagaraNodeIf.cpp
--- a/Engine/Plugins/FX/Niagara/Source/NiagaraEditor/Private/NiagaraNodeIf.cpp
+++ b/Engine/Plugins/FX/Niagara/Source/NiagaraEditor/Private/NiagaraNodeIf.cpp
@@ -8,6 +8,7 @@
 
 const FString UNiagaraNodeIf::InputAPinSuffix(" A");
 const FString UNiagaraNodeIf::InputBPinSuffix(" B");
+const FName UNiagaraNodeIf::ConditionPinName("Condition");
 
 UNiagaraNodeIf::UNiagaraNodeIf(const FObjectInitializer& ObjectInitializer)
 	: Super(ObjectInitializer)
@@ -29,19 +30,9 @@ void UNiagaraNodeIf::PostLoad()
 {
 	Super::PostLoad();
 
-	if (OutputVars.Num() != OutputVarGuids.Num())
+	if (PathAssociatedPinGuids.Num() != OutputVars.Num())
 	{
-		OutputVarGuids.SetNum(OutputVars.Num());
-	}
-
-	if (OutputVars.Num() != InputAVarGuids.Num())
-	{
-		InputAVarGuids.SetNum(OutputVars.Num());
-	}
-
-	if (OutputVars.Num() != InputBVarGuids.Num())
-	{
-		InputBVarGuids.SetNum(OutputVars.Num());
+		PathAssociatedPinGuids.SetNum(OutputVars.Num());
 	}
 
 	auto LoadGuid = [&](FGuid& Guid, const FString& Name, const EEdGraphPinDirection Direction)
@@ -60,15 +51,16 @@ void UNiagaraNodeIf::PostLoad()
 			UE_LOG(LogNiagaraEditor, Error, TEXT("Unable to find output pin named %s"), *Name);
 		}
 	};
-
+	
+	LoadGuid(ConditionPinGuid, ConditionPinName.ToString(), EGPD_Input);
 	for (int32 i = 0; i < OutputVars.Num(); i++)
 	{
 		const FString VarName = OutputVars[i].GetName().ToString();
-		LoadGuid(OutputVarGuids[i], VarName, EGPD_Output);
+		LoadGuid(PathAssociatedPinGuids[i].OutputPinGuid, VarName, EGPD_Output);
 		const FString InputAName = VarName + InputAPinSuffix;
-		LoadGuid(InputAVarGuids[i], InputAName, EGPD_Input);
+		LoadGuid(PathAssociatedPinGuids[i].InputAPinGuid, InputAName, EGPD_Input);
 		const FString InputBName = VarName + InputBPinSuffix;
-		LoadGuid(InputBVarGuids[i], InputBName, EGPD_Input);
+		LoadGuid(PathAssociatedPinGuids[i].InputBPinGuid, InputBName, EGPD_Input);
 	}
 }
 
@@ -84,28 +76,29 @@ void UNiagaraNodeIf::AllocateDefaultPins()
 	const UEdGraphSchema_Niagara* Schema = GetDefault<UEdGraphSchema_Niagara>();
 
 	//Add the condition pin.
-	CreatePin(EGPD_Input, Schema->TypeDefinitionToPinType(FNiagaraTypeDefinition::GetBoolDef()), TEXT("Condition"));
+	UEdGraphPin* ConditionPin = CreatePin(EGPD_Input, Schema->TypeDefinitionToPinType(FNiagaraTypeDefinition::GetBoolDef()), ConditionPinName);
+	ConditionPin->PersistentGuid = ConditionPinGuid;
 
 	//Create the inputs for each path.
 	for (int32 Index = 0; Index < OutputVars.Num(); Index++)
 	{
 		const FNiagaraVariable& Var = OutputVars[Index];
 		UEdGraphPin* NewPin = CreatePin(EGPD_Input, Schema->TypeDefinitionToPinType(Var.GetType()), *(Var.GetName().ToString() + InputAPinSuffix));
-		NewPin->PersistentGuid = InputAVarGuids[Index];
+		NewPin->PersistentGuid = PathAssociatedPinGuids[Index].InputAPinGuid;
 	}
 
 	for (int32 Index = 0; Index < OutputVars.Num(); Index++)
 	{
 		const FNiagaraVariable& Var = OutputVars[Index];
 		UEdGraphPin* NewPin = CreatePin(EGPD_Input, Schema->TypeDefinitionToPinType(Var.GetType()), *(Var.GetName().ToString() + InputBPinSuffix));
-		NewPin->PersistentGuid = InputBVarGuids[Index];
+		NewPin->PersistentGuid = PathAssociatedPinGuids[Index].InputBPinGuid;
 	}
 
 	for (int32 Index = 0; Index < OutputVars.Num(); Index++)
 	{
 		const FNiagaraVariable& Var = OutputVars[Index];
 		UEdGraphPin* NewPin = CreatePin(EGPD_Output, Schema->TypeDefinitionToPinType(Var.GetType()), Var.GetName());
-		NewPin->PersistentGuid = OutputVarGuids[Index];
+		NewPin->PersistentGuid = PathAssociatedPinGuids[Index].OutputPinGuid;
 	}
 
 	CreateAddPin(EGPD_Output);
@@ -113,32 +106,34 @@ void UNiagaraNodeIf::AllocateDefaultPins()
 
 void UNiagaraNodeIf::Compile(class FHlslNiagaraTranslator* Translator, TArray<int32>& Outputs)
 {
+	checkSlow(PathAssociatedPinGuids.Num() == OutputVars.Num());
 	const UEdGraphSchema_Niagara* Schema = CastChecked<UEdGraphSchema_Niagara>(GetSchema());
 
-	int32 PinIdx = 0;
-	int32 Condition = Translator->CompilePin(Pins[PinIdx++]);
+	int32 Condition = Translator->CompilePin(GetPinByGuid(ConditionPinGuid));
+
 	TArray<int32> PathA;
-	PathA.Reserve(OutputVars.Num());
-	for (int32 i = 0; i < OutputVars.Num(); ++i)
+	PathA.Reserve(PathAssociatedPinGuids.Num());
+	for (const FPinGuidsForPath& PerPathAssociatedPinGuids : PathAssociatedPinGuids)
 	{
-		if (Schema->PinToTypeDefinition(Pins[PinIdx]) == FNiagaraTypeDefinition::GetParameterMapDef())
+		const UEdGraphPin* InputAPin = GetPinByGuid(PerPathAssociatedPinGuids.InputAPinGuid);
+		if (Schema->PinToTypeDefinition(InputAPin) == FNiagaraTypeDefinition::GetParameterMapDef())
 		{
-			Translator->Error(LOCTEXT("UnsupportedParamMapInIf","Parameter maps are not supported in if nodes."), this, Pins[PinIdx]);
+			Translator->Error(LOCTEXT("UnsupportedParamMapInIf", "Parameter maps are not supported in if nodes."), this, InputAPin);
 		}
-
-		PathA.Add(Translator->CompilePin(Pins[PinIdx++]));
+		PathA.Add(Translator->CompilePin(InputAPin));
 	}
 	TArray<int32> PathB;
-	PathB.Reserve(OutputVars.Num());
-	for (int32 i = 0; i < OutputVars.Num(); ++i)
+	PathB.Reserve(PathAssociatedPinGuids.Num());
+	for (const FPinGuidsForPath& PerPathAssociatedPinGuids : PathAssociatedPinGuids)
 	{
-		if (Schema->PinToTypeDefinition(Pins[PinIdx]) == FNiagaraTypeDefinition::GetParameterMapDef())
+		const UEdGraphPin* InputBPin = GetPinByGuid(PerPathAssociatedPinGuids.InputBPinGuid);
+		if (Schema->PinToTypeDefinition(InputBPin) == FNiagaraTypeDefinition::GetParameterMapDef())
 		{
-			Translator->Error(LOCTEXT("UnsupportedParamMapInIf", "Parameter maps are not supported in if nodes."), this, Pins[PinIdx]);
+			Translator->Error(LOCTEXT("UnsupportedParamMapInIf", "Parameter maps are not supported in if nodes."), this, InputBPin);
 		}
-
-		PathB.Add(Translator->CompilePin(Pins[PinIdx++]));
+		PathB.Add(Translator->CompilePin(InputBPin));
 	}
+
 	Translator->If(this, OutputVars, Condition, PathA, PathB, Outputs);
 }
 
@@ -173,33 +168,42 @@ bool UNiagaraNodeIf::RefreshFromExternalChanges()
 
 FGuid UNiagaraNodeIf::AddOutput(FNiagaraTypeDefinition Type, const FName& Name)
 {
+	FPinGuidsForPath& NewPinGuidsForPath = PathAssociatedPinGuids.Add_GetRef(FPinGuidsForPath());
+
 	FNiagaraVariable NewOutput(Type, Name);
 	OutputVars.Add(NewOutput);
 	FGuid Guid = FGuid::NewGuid();
-	OutputVarGuids.Add(Guid);
+	NewPinGuidsForPath.OutputPinGuid = Guid;
 
 	const UEdGraphSchema_Niagara* Schema = CastChecked<UEdGraphSchema_Niagara>(GetSchema());
 	FGuid PinAGuid = FGuid::NewGuid();
-	InputAVarGuids.Add(PinAGuid);
-	UEdGraphPin* PinA = CreatePin(EGPD_Input, Schema->TypeDefinitionToPinType(Type), *(Name.ToString() + InputAPinSuffix), InputAVarGuids.Num());
+	UEdGraphPin* PinA = CreatePin(EGPD_Input, Schema->TypeDefinitionToPinType(Type), *(Name.ToString() + InputAPinSuffix), PathAssociatedPinGuids.Num());
 	PinA->PersistentGuid = PinAGuid;
+	NewPinGuidsForPath.InputAPinGuid = PinAGuid;
 
 	FGuid PinBGuid = FGuid::NewGuid();
-	InputBVarGuids.Add(PinBGuid);
-	UEdGraphPin* PinB = CreatePin(EGPD_Input, Schema->TypeDefinitionToPinType(Type), *(Name.ToString() + InputBPinSuffix), InputAVarGuids.Num() + InputBVarGuids.Num());
+	UEdGraphPin* PinB = CreatePin(EGPD_Input, Schema->TypeDefinitionToPinType(Type), *(Name.ToString() + InputBPinSuffix), PathAssociatedPinGuids.Num() * 2);
 	PinB->PersistentGuid = PinBGuid;
+	NewPinGuidsForPath.InputBPinGuid = PinBGuid;
 
 	return Guid;
 }
 
+const UEdGraphPin* UNiagaraNodeIf::GetPinByGuid(const FGuid& InGuid)
+{
+	UEdGraphPin* FoundPin = *Pins.FindByPredicate([&InGuid](const UEdGraphPin* Pin) { return Pin->PersistentGuid == InGuid; });
+	checkf(FoundPin != nullptr, TEXT("Failed to get pin by cached Guid!"));
+	return FoundPin;
+}
+
 void UNiagaraNodeIf::OnPinRemoved(UEdGraphPin* PinToRemove)
 {
-	auto FindPredicate = [=](const FGuid& Guid) { return Guid == PinToRemove->PersistentGuid; };
-	int32 FoundIndex = OutputVarGuids.IndexOfByPredicate(FindPredicate);
+	auto FindByOutputPinGuidPredicate = [=](const FPinGuidsForPath& PerPinGuidsForPath) { return PerPinGuidsForPath.OutputPinGuid == PinToRemove->PersistentGuid; };
+	int32 FoundIndex = PathAssociatedPinGuids.IndexOfByPredicate(FindByOutputPinGuidPredicate);
 	if (FoundIndex != INDEX_NONE)
 	{
-		OutputVarGuids.RemoveAt(FoundIndex);
 		OutputVars.RemoveAt(FoundIndex);
+		PathAssociatedPinGuids.RemoveAt(FoundIndex);
 	}
 	ReallocatePins();
 }
@@ -227,8 +231,8 @@ void UNiagaraNodeIf::OnNewTypedPinAdded(UEdGraphPin* NewPin)
 
 void UNiagaraNodeIf::OnPinRenamed(UEdGraphPin* RenamedPin, const FString& OldName)
 {
-	auto FindPredicate = [=](const FGuid& Guid) { return Guid == RenamedPin->PersistentGuid; };
-	int32 FoundIndex = OutputVarGuids.IndexOfByPredicate(FindPredicate);
+	auto FindByOutputPinGuidPredicate = [=](const FPinGuidsForPath& PerPinGuidsForPath) { return PerPinGuidsForPath.OutputPinGuid == RenamedPin->PersistentGuid; };
+	int32 FoundIndex = PathAssociatedPinGuids.IndexOfByPredicate(FindByOutputPinGuidPredicate);
 	if(FoundIndex != INDEX_NONE)
 	{
 		TSet<FName> OutputNames;
diff --git a/Engine/Plugins/FX/Niagara/Source/NiagaraEditor/Private/NiagaraNodeIf.h b/Engine/Plugins/FX/Niagara/Source/NiagaraEditor/Private/NiagaraNodeIf.h
--- a/Engine/Plugins/FX/Niagara/Source/NiagaraEditor/Private/NiagaraNodeIf.h
+++ b/Engine/Plugins/FX/Niagara/Source/NiagaraEditor/Private/NiagaraNodeIf.h
@@ -5,6 +5,33 @@
 #include "EdGraphSchema_Niagara.h"
 #include "NiagaraNodeIf.generated.h"
 
+USTRUCT()
+struct FPinGuidsForPath
+{
+	GENERATED_USTRUCT_BODY()
+
+	FPinGuidsForPath()
+		: OutputPinGuid(FGuid())
+		, InputAPinGuid(FGuid())
+		, InputBPinGuid(FGuid())
+	{
+	}
+
+	bool IsValid() const
+	{
+		return OutputPinGuid.IsValid() && InputAPinGuid.IsValid() && InputBPinGuid.IsValid();
+	}
+
+	UPROPERTY()
+	FGuid OutputPinGuid;
+
+	UPROPERTY()
+	FGuid InputAPinGuid;
+
+	UPROPERTY()
+	FGuid InputBPinGuid;
+};
+
 UCLASS(MinimalAPI)
 class UNiagaraNodeIf : public UNiagaraNodeWithDynamicPins
 {
@@ -17,13 +44,10 @@ public:
 	TArray<FNiagaraVariable> OutputVars;
 
 	UPROPERTY()
-	TArray<FGuid> OutputVarGuids;
+	TArray<FPinGuidsForPath> PathAssociatedPinGuids;
 
 	UPROPERTY()
-	TArray<FGuid> InputAVarGuids;
-
-	UPROPERTY()
-	TArray<FGuid> InputBVarGuids;
+	FGuid ConditionPinGuid;
 
 	//~ Begin UObject Interface
 	virtual void PostEditChangeProperty(FPropertyChangedEvent& PropertyChangedEvent) override;
@@ -45,9 +69,12 @@ public:
 
 protected:
 
-	/** Helper function to create a variable to add to the OutputVars and FGuid to add to OutputVarGuids. */
+	/** Helper function to create a variable to add to the OutputVars and FGuid to add to PathAssociatedPinGuids. */
 	FGuid AddOutput(FNiagaraTypeDefinition Type, const FName& Name);
 
+	/** Helper to get a pin in the pins list by GUID */
+	const UEdGraphPin* GetPinByGuid(const FGuid& InGuid);
+
 	//~ Begin EdGraphNode Interface
 	virtual void OnPinRemoved(UEdGraphPin* PinToRemove) override;
 	//~ End EdGraphNode Interface
@@ -64,4 +91,5 @@ protected:
 private:
 	static const FString InputAPinSuffix;
 	static const FString InputBPinSuffix;
+	static const FName ConditionPinName;
 };
