- Fix race condition with async system simulations
- The finalize task was allow to run before the async simulation task had completed which could allow concurrent access to the SystemInstances array

diff --git a/Engine/Plugins/FX/Niagara/Source/Niagara/Private/NiagaraSystemSimulation.cpp b/Engine/Plugins/FX/Niagara/Source/Niagara/Private/NiagaraSystemSimulation.cpp
--- a/Engine/Plugins/FX/Niagara/Source/Niagara/Private/NiagaraSystemSimulation.cpp
+++ b/Engine/Plugins/FX/Niagara/Source/Niagara/Private/NiagaraSystemSimulation.cpp
@@ -15,6 +15,12 @@
 #include "NiagaraWorldManager.h"
 #include "NiagaraEmitterInstanceBatcher.h"
 
+// Niagara simulations async will block the tick task from completion until all async work is finished
+// If simulations are allowed to tick async we will create a FNiagaraSystemSimulationTickTask task to run on any thread
+// If instances are allowed to tick async we will create a FNiagaraSystemInstanceAsyncTask in batches to run on any thread
+// If any async is enabled we create a FNiagaraSystemInstanceFinalizeTask for each batch that will not run until FNiagaraSystemSimulationTickTask is complete (Due to contention with SystemInstances) and will run on the GameThread
+// If any async is enabled we create a null graph task to wait for all FNiagaraSystemInstanceFinalizeTask's to complete before allowing the tick group to advance
+
 //High level stats for system sim tick.
 DECLARE_CYCLE_STAT(TEXT("System Simulaton Tick [GT]"), STAT_NiagaraSystemSim_TickGT, STATGROUP_Niagara);
 DECLARE_CYCLE_STAT(TEXT("System Simulaton Tick [CNC]"), STAT_NiagaraSystemSim_TickCNC, STATGROUP_Niagara);
@@ -25,11 +31,9 @@ DECLARE_CYCLE_STAT(TEXT("System Sim Update [CNC]"), STAT_NiagaraSystemSim_Update
 DECLARE_CYCLE_STAT(TEXT("System Sim Spawn [CNC]"), STAT_NiagaraSystemSim_SpawnCNC, STATGROUP_Niagara);
 DECLARE_CYCLE_STAT(TEXT("System Sim Transfer Results [CNC]"), STAT_NiagaraSystemSim_TransferResultsCNC, STATGROUP_Niagara);
 
-
 DECLARE_CYCLE_STAT(TEXT("ForcedWaitForAsync"), STAT_NiagaraSystemSim_ForceWaitForAsync, STATGROUP_Niagara);
 DECLARE_CYCLE_STAT(TEXT("ForcedWait Fake Stall"), STAT_NiagaraSystemSim_ForceWaitFakeStall, STATGROUP_Niagara);
 
-
 static int32 GbDumpSystemData = 0;
 static FAutoConsoleVariableRef CVarNiagaraDumpSystemData(
 	TEXT("fx.DumpSystemData"),
@@ -133,9 +137,11 @@ FAutoConsoleTaskPriority CPrio_NiagaraSystemSimulationTickTask(
 class FNiagaraSystemSimulationTickTask
 {
 	FNiagaraSystemSimulationTickContext Context;
+	TGraphTask<FNullGraphTask>* FinalizeCompleteTask;
 public:
-	FNiagaraSystemSimulationTickTask(FNiagaraSystemSimulationTickContext InContext)
+	FNiagaraSystemSimulationTickTask(FNiagaraSystemSimulationTickContext InContext, TGraphTask<FNullGraphTask>* InFinalizeCompleteTask)
 		: Context(InContext)
+		, FinalizeCompleteTask(InFinalizeCompleteTask)
 	{
 	}
 
@@ -155,6 +161,8 @@ public:
 	{
 		Context.MyCompletionGraphEvent = MyCompletionGraphEvent;
 		Context.Owner->Tick_Concurrent(Context);
+
+		FinalizeCompleteTask->Unlock(ENamedThreads::GameThread);
 	}
 };
 
@@ -242,11 +250,6 @@ public:
 		{
 			Inst->Tick_Concurrent();
 		}
-
-		//Now kick off a finalize task for this batch on the game thread.
-		FGraphEventRef FinalizeTask = TGraphTask<FNiagaraSystemInstanceFinalizeTask>::CreateTask(nullptr, CurrentThread).ConstructAndDispatchWhenReady(SystemSim, Batch);
-		MyCompletionGraphEvent->SetGatherThreadForDontCompleteUntil(ENamedThreads::GameThread);
-		MyCompletionGraphEvent->DontCompleteUntil(FinalizeTask);
 	}
 };
 
@@ -506,32 +509,34 @@ void FNiagaraSystemSimulation::FlushTickBatch(FNiagaraSystemSimulationTickContex
 {
 	if (TickBatch.Num() > 0)
 	{
-		if (Context.bTickInstancesAsync)
+		FGraphEventArray FinalizePrereq;
+		FinalizePrereq.Add(Context.MyCompletionGraphEvent);
+
+		// Enqueue or tick the instances
+		if ( Context.bTickInstancesAsync )
 		{
-			checkSlow(Context.MyCompletionGraphEvent.IsValid());
-			//If we're able, kick off a task to process this batch on a task thread.
-			//When this task has finished the concurrent ticks for this batch it will enqueue a finalize task for the batch on the game thread.
+			check(Context.FinalizeCompleteGraphEvent.IsValid());
+
 			FGraphEventRef AsyncTask = TGraphTask<FNiagaraSystemInstanceAsyncTask>::CreateTask(nullptr).ConstructAndDispatchWhenReady(this, TickBatch);
-			Context.MyCompletionGraphEvent->SetGatherThreadForDontCompleteUntil(ENamedThreads::GameThread);
-			Context.MyCompletionGraphEvent->DontCompleteUntil(AsyncTask);
+			FinalizePrereq.Add(AsyncTask);
 		}
 		else
 		{
-			//Otherwise just do directly here. This may already be on a task thread or could be on game thread.
 			for (FNiagaraSystemInstance* Inst : TickBatch)
 			{
 				Inst->Tick_Concurrent();
 			}
+		}
 
-			if (Context.bTickAsync)
-			{
-				checkSlow(Context.MyCompletionGraphEvent.IsValid());
-				//If we're ticking off the main thread then we still need to add a finalize task on the GT for this batch.
-				FGraphEventRef FinalizeTask = TGraphTask<FNiagaraSystemInstanceFinalizeTask>::CreateTask(nullptr).ConstructAndDispatchWhenReady(this, TickBatch);
-				Context.MyCompletionGraphEvent->SetGatherThreadForDontCompleteUntil(ENamedThreads::GameThread);
-				Context.MyCompletionGraphEvent->DontCompleteUntil(FinalizeTask);
-			}
+		// Enqueue a finalize task?
+		if (Context.bTickAsync || Context.bTickInstancesAsync )
+		{
+			check(Context.FinalizeCompleteGraphEvent.IsValid());
+
+			FGraphEventRef FinalizeTask = TGraphTask<FNiagaraSystemInstanceFinalizeTask>::CreateTask(&FinalizePrereq).ConstructAndDispatchWhenReady(this, TickBatch);
+			Context.FinalizeCompleteGraphEvent->DontCompleteUntil(FinalizeTask);
 		}
+
 		TickBatch.Reset();
 	}
 }
@@ -693,13 +698,39 @@ void FNiagaraSystemSimulation::Tick_GameThread(float DeltaSeconds, const FGraphE
 	//Now kick of the concurrent tick.
 	if (Context.bTickAsync)
 	{
-		FGraphEventRef AsyncTask = TGraphTask<FNiagaraSystemSimulationTickTask>::CreateTask(nullptr, ENamedThreads::GameThread).ConstructAndDispatchWhenReady(Context);
+		DECLARE_CYCLE_STAT(TEXT("FNullGraphTask.NiagaraFinalizeComplete"), STAT_FNullGraphTask_NiagaraFinalizeComplete, STATGROUP_TaskGraphTasks);
+
+		TGraphTask<FNullGraphTask>* FinalizeCompleteTask = TGraphTask<FNullGraphTask>::CreateTask(nullptr, ENamedThreads::GameThread).ConstructAndHold(GET_STATID(STAT_FNullGraphTask_NiagaraFinalizeComplete), ENamedThreads::AnyThread);
+		Context.FinalizeCompleteGraphEvent = FinalizeCompleteTask->GetCompletionEvent();
+
+		TGraphTask<FNiagaraSystemSimulationTickTask>* AsyncTask = TGraphTask<FNiagaraSystemSimulationTickTask>::CreateTask(nullptr, ENamedThreads::GameThread).ConstructAndHold(Context, FinalizeCompleteTask);
+		FGraphEventRef AsyncGraphEvent = AsyncTask->GetCompletionEvent();
+
 		MyCompletionGraphEvent->SetGatherThreadForDontCompleteUntil(ENamedThreads::GameThread);
-		MyCompletionGraphEvent->DontCompleteUntil(AsyncTask);
+		MyCompletionGraphEvent->DontCompleteUntil(Context.FinalizeCompleteGraphEvent);
+		Context.FinalizeCompleteGraphEvent->SetGatherThreadForDontCompleteUntil(ENamedThreads::GameThread);
+		Context.FinalizeCompleteGraphEvent->DontCompleteUntil(AsyncGraphEvent);
+
+		AsyncTask->Unlock(ENamedThreads::GameThread);
 	}
 	else
 	{
+		TGraphTask<FNullGraphTask>* WaitConcurrentTick = nullptr;
+		if (Context.bTickInstancesAsync)
+		{
+			DECLARE_CYCLE_STAT(TEXT("FNullGraphTask.NiagaraWaitToFinalize"), STAT_FNullGraphTask_NiagaraWaitToFinalize, STATGROUP_TaskGraphTasks);
+
+			WaitConcurrentTick = TGraphTask<FNullGraphTask>::CreateTask(nullptr, ENamedThreads::GameThread).ConstructAndHold(GET_STATID(STAT_FNullGraphTask_NiagaraWaitToFinalize), ENamedThreads::AnyThread);
+			Context.MyCompletionGraphEvent = WaitConcurrentTick->GetCompletionEvent();
+			Context.FinalizeCompleteGraphEvent = MyCompletionGraphEvent;
+		}
+
 		Tick_Concurrent(Context);
+
+		if (Context.bTickInstancesAsync)
+		{
+			WaitConcurrentTick->Unlock(ENamedThreads::GameThread);
+		}
 	}
 
 	//TEMPORARY DEV CODE:
diff --git a/Engine/Plugins/FX/Niagara/Source/Niagara/Public/NiagaraSystemSimulation.h b/Engine/Plugins/FX/Niagara/Source/Niagara/Public/NiagaraSystemSimulation.h
--- a/Engine/Plugins/FX/Niagara/Source/Niagara/Public/NiagaraSystemSimulation.h
+++ b/Engine/Plugins/FX/Niagara/Source/Niagara/Public/NiagaraSystemSimulation.h
@@ -132,6 +132,8 @@ struct FNiagaraSystemSimulationTickContext
 	bool bPendingSpawnPass;
 	FGraphEventRef MyCompletionGraphEvent;
 
+	FGraphEventRef FinalizeCompleteGraphEvent;
+
 	bool bTickAsync;
 	bool bTickInstancesAsync;
 
@@ -282,12 +284,11 @@ protected:
 	uint32 bCanExecute : 1;
 	uint32 bBindingsInitialized : 1;
 	uint32 bInSpawnPhase : 1;
+	uint32 bIsSolo : 1;
 
 	/** A parameter store which contains the data interfaces parameters which were defined by the scripts. */
 	FNiagaraParameterStore ScriptDefinedDataInterfaceParameters;
 
-	bool bIsSolo;
-
 	TOptional<float> MaxDeltaTime;
 
 	/** Current tick batch we're filling ready for processing, potentially in an async task. */
