- Remove thread name from finalize unlock task

diff --git a/Engine/Plugins/FX/Niagara/Source/Niagara/Private/NiagaraSystemSimulation.cpp b/Engine/Plugins/FX/Niagara/Source/Niagara/Private/NiagaraSystemSimulation.cpp
--- a/Engine/Plugins/FX/Niagara/Source/Niagara/Private/NiagaraSystemSimulation.cpp
+++ b/Engine/Plugins/FX/Niagara/Source/Niagara/Private/NiagaraSystemSimulation.cpp
@@ -162,7 +162,7 @@ public:
 		Context.MyCompletionGraphEvent = MyCompletionGraphEvent;
 		Context.Owner->Tick_Concurrent(Context);
 
-		FinalizeCompleteTask->Unlock(ENamedThreads::GameThread);
+		FinalizeCompleteTask->Unlock();
 	}
 };
 
