Removed niagara VF LocalToWorld and LocalToWorldInverseTransposed.
Also removed Niagara renderers WorldSpacePrimitiveUniformBuffer.
Now using the Primive UB instead.

diff --git a/Engine/Plugins/FX/Niagara/Shaders/Private/NiagaraMeshVertexFactory.ush b/Engine/Plugins/FX/Niagara/Shaders/Private/NiagaraMeshVertexFactory.ush
--- a/Engine/Plugins/FX/Niagara/Shaders/Private/NiagaraMeshVertexFactory.ush
+++ b/Engine/Plugins/FX/Niagara/Shaders/Private/NiagaraMeshVertexFactory.ush
@@ -14,6 +14,30 @@
 #define USE_PARTICLE_RANDOM (FEATURE_LEVEL >= FEATURE_LEVEL_SM5 && NEEDS_PARTICLE_RANDOM)
 #define USE_PARTICLE_TRANSFORM (FEATURE_LEVEL >= FEATURE_LEVEL_SM5 && NEEDS_PARTICLE_TRANSFORM)
 
+float3 TransformPosition(float3 InPosition)
+{
+	if (NiagaraMeshVF.bLocalSpace)
+	{
+		return TransformLocalToWorld(InPosition);
+	}
+	else
+	{
+		return InPosition;
+	}
+}
+
+float3 TransformVector(float3 InVector)
+{
+	if (NiagaraMeshVF.bLocalSpace)
+	{
+		return RotateLocalToWorld(InVector);
+	}
+	else
+	{
+		return InVector;
+	}
+}
+
 // Only support previous transform for deferred feature level
 //#define USE_PREVIOUS_TRANSFORM (FEATURE_LEVEL >= FEATURE_LEVEL_SM4)
 #define USE_PREVIOUS_TRANSFORM 0
@@ -608,7 +632,7 @@ FVertexFactoryIntermediates GetVertexFactoryIntermediates(FVertexFactoryInput In
 		float4 Quat = GetVec4(NiagaraMeshVF.TransformDataOffset, Input.InstanceId);
 
 		//Transform axis to handle local space.
-		//AxisAngle.xyz = mul(float4(AxisAngle.xyz, 0), NiagaraMeshVF.LocalToWorldInverseTransposed);
+		//AxisAngle.xyz = TransformVector(AxisAngle.xyz);
 
 
 		float4x4 EulerToMat = QuatTo4x4(normalize(Quat));
@@ -640,7 +664,7 @@ FVertexFactoryIntermediates GetVertexFactoryIntermediates(FVertexFactoryInput In
 	}
 	else if (MeshFacingMode == MESH_FACING_CAMERA_POSITION)
 	{
-		float3 WorldPosition = mul(Intermediates.ParticlePosition, NiagaraMeshVF.LocalToWorld).xyz;
+		float3 WorldPosition = TransformPosition(Intermediates.ParticlePosition);
 		float3 CameraVector = ResolvedView.WorldCameraOrigin.xyz - WorldPosition.xyz;
 		float3 CameraDirection = normalize(CameraVector);
 		float4 ToAlignmentQuat = FindQuatBetween(float3(1.0f, 0.0f, 0.0f), CameraDirection);
@@ -662,9 +686,12 @@ FVertexFactoryIntermediates GetVertexFactoryIntermediates(FVertexFactoryInput In
 	//Handle local space. TODO: Optimize behind define.
 	//Apply local space to local to world, also pretranspse here to avoid it in the vs&ps
 	Intermediates.LocalToWorld = transpose(Intermediates.LocalToWorld);
-	Intermediates.LocalToWorld = mul(Intermediates.LocalToWorld, NiagaraMeshVF.LocalToWorld);
-	Intermediates.ParticlePosition = mul(Intermediates.ParticlePosition, NiagaraMeshVF.LocalToWorld);
-	//Intermediates.ParticleVelocity.xyz = mul(float4(Intermediates.ParticleVelocity.xyz,0), NiagaraMeshVF.LocalToWorldInverseTransposed);
+	if (NiagaraMeshVF.bLocalSpace)
+	{
+		Intermediates.LocalToWorld = mul(Intermediates.LocalToWorld, Primitive.LocalToWorld);
+		Intermediates.ParticlePosition = TransformLocalToWorld(Intermediates.ParticlePosition);
+	}
+	//Intermediates.ParticleVelocity.xyz = TransformVector(Intermediates.ParticleVelocity.xyz);
 
 #if NUM_MATERIAL_TEXCOORDS_VERTEX
 	for(int CoordinateIndex = 0; CoordinateIndex < NUM_MATERIAL_TEXCOORDS_VERTEX; CoordinateIndex++)
diff --git a/Engine/Plugins/FX/Niagara/Shaders/Private/NiagaraRibbonVertexFactory.ush b/Engine/Plugins/FX/Niagara/Shaders/Private/NiagaraRibbonVertexFactory.ush
--- a/Engine/Plugins/FX/Niagara/Shaders/Private/NiagaraRibbonVertexFactory.ush
+++ b/Engine/Plugins/FX/Niagara/Shaders/Private/NiagaraRibbonVertexFactory.ush
@@ -16,12 +16,26 @@
 
 float3 TransformPosition(float3 InPosition)
 {
-	return mul(float4(InPosition,1.0), NiagaraRibbonVF.LocalToWorld).xyz;
+	if (NiagaraRibbonVF.bLocalSpace)
+	{
+		return TransformLocalToWorld(InPosition);
+	}
+	else
+	{
+		return InPosition;
+	}
 }
 
 float3 TransformVector(float3 InVector)
 {
-	return mul(float4(InVector,0.0), NiagaraRibbonVF.LocalToWorldInverseTransposed).xyz;
+	if (NiagaraRibbonVF.bLocalSpace)
+	{
+		return RotateLocalToWorld(InVector);
+	}
+	else
+	{
+		return InVector;
+	}
 }
 
 float3 CubicInterp(float3 P0, float3 T0, float3 P1, float3 T1, float A, float TS)
diff --git a/Engine/Plugins/FX/Niagara/Shaders/Private/NiagaraSpriteVertexFactory.ush b/Engine/Plugins/FX/Niagara/Shaders/Private/NiagaraSpriteVertexFactory.ush
--- a/Engine/Plugins/FX/Niagara/Shaders/Private/NiagaraSpriteVertexFactory.ush
+++ b/Engine/Plugins/FX/Niagara/Shaders/Private/NiagaraSpriteVertexFactory.ush
@@ -343,7 +343,7 @@ FMaterialVertexParameters GetMaterialVertexParameters(FVertexFactoryInput Input,
 
 	// Previous frame not handled deliberately. Lacks necessary information and
 	// primitives using this VF are usually transparent and hence don't output velocity
-	Result.PrevFrameLocalToWorld = Primitive.LocalToWorld;
+	Result.PrevFrameLocalToWorld = float4x4(float4(1.0f, 0.0f, 0.0f, 0.0f), float4(0.0f, 1.0f, 0.0f, 0.0f), float4(0.0f, 0.0f, 1.0f, 0.0f), float4(0.0f, 0.0f, 0.0f, 1.0f));
 
 #if (DYNAMIC_PARAMETERS_MASK != 0)
 	Result.Particle.DynamicParameterValidMask = NiagaraSpriteVF.MaterialParamValidMask;
@@ -538,14 +538,14 @@ float3x3 CalcTangentBasis(FVertexFactoryIntermediates Intermediates)
 	// ENM_Spherical
 	else if (NiagaraSpriteVF.NormalsType < 1.5f)
 	{
-		float3 TangentZ = normalize(Intermediates.ParticleWorldPosition.xyz - NiagaraSpriteVF.NormalsSphereCenter.xyz);
+		float3 TangentZ = normalize(Intermediates.ParticleWorldPosition - NiagaraSpriteVF.NormalsSphereCenter.xyz);
 		Result[2] = TangentZ;	
 	}
 	// ENM_Cylindrical
 	else
 	{
-		float3 ClosestPointOnCylinder = NiagaraSpriteVF.NormalsSphereCenter.xyz + dot(NiagaraSpriteVF.NormalsCylinderUnitDirection.xyz, Intermediates.ParticleWorldPosition.xyz - NiagaraSpriteVF.NormalsSphereCenter.xyz) * NiagaraSpriteVF.NormalsCylinderUnitDirection.xyz;
-		float3 TangentZ = normalize(Intermediates.ParticleWorldPosition.xyz - ClosestPointOnCylinder);
+		float3 ClosestPointOnCylinder = NiagaraSpriteVF.NormalsSphereCenter.xyz + dot(NiagaraSpriteVF.NormalsCylinderUnitDirection.xyz, Intermediates.ParticleWorldPosition - NiagaraSpriteVF.NormalsSphereCenter.xyz) * NiagaraSpriteVF.NormalsCylinderUnitDirection.xyz;
+		float3 TangentZ = normalize(Intermediates.ParticleWorldPosition - ClosestPointOnCylinder);
 		Result[2] = TangentZ;	
 	}
 
@@ -595,12 +595,26 @@ void ComputeBillboardUVs(FVertexFactoryInput Input, FVertexFactoryIntermediates
 
 float3 TransformPosition(float3 InPosition)
 {
-	return mul(float4(InPosition,1.0), NiagaraSpriteVF.LocalToWorld).xyz;
+	if (NiagaraSpriteVF.bLocalSpace)
+	{
+		return TransformLocalToWorld(InPosition);
+	}
+	else
+	{
+		return InPosition;
+	}
 }
 
 float3 TransformVector(float3 InVector)
 {
-	return mul(float4(InVector,0.0), NiagaraSpriteVF.LocalToWorldInverseTransposed).xyz;
+	if (NiagaraSpriteVF.bLocalSpace)
+	{
+		return RotateLocalToWorld(InVector);
+	}
+	else
+	{
+		return InVector;
+	}
 }
 
 /*********************************************************************************
@@ -818,23 +832,23 @@ FVertexFactoryIntermediates GetVertexFactoryIntermediates(FVertexFactoryInput In
 
 	// World position.
 	
-	float4 ParticleWorldPosition = mul(float4(Intermediates.Position.xyz,1), Primitive.LocalToWorld);
-	float4 ParticleOldWorldPosition = mul(float4(Intermediates.OldPosition.xyz,1), Primitive.LocalToWorld); 
+	float3 ParticleWorldPosition = Intermediates.Position;
+	float3 ParticleOldWorldPosition = Intermediates.OldPosition;
 	float CameraOffset = GetNiagaraCameraOffset(ParticleID);
-	float3 CameraVector = ResolvedView.WorldCameraOrigin.xyz - ParticleWorldPosition.xyz;
+	float3 CameraVector = ResolvedView.WorldCameraOrigin.xyz - ParticleWorldPosition;
 	float3 CameraDirection = SafeNormalize(CameraVector);
 
 	// Do we need the hack that forces particles at the camera origin to be *really* behind the camera origin?
 	float3 OffsetPosition = CameraDirection * CameraOffset;
 
-	ParticleWorldPosition.xyz = ParticleWorldPosition.xyz + OffsetPosition;
-	ParticleOldWorldPosition.xyz = ParticleOldWorldPosition.xyz + OffsetPosition;
+	ParticleWorldPosition = ParticleWorldPosition + OffsetPosition;
+	ParticleOldWorldPosition = ParticleOldWorldPosition + OffsetPosition;
 	
-	Intermediates.ParticleWorldPosition = ParticleWorldPosition.xyz;
+	Intermediates.ParticleWorldPosition = ParticleWorldPosition;
 
 	// Tangents.
 	float3 Right,Up;
-	GetTangents(Input, Intermediates, ParticleWorldPosition.xyz,ParticleOldWorldPosition.xyz,Right,Up);
+	GetTangents(Input, Intermediates, ParticleWorldPosition,ParticleOldWorldPosition,Right,Up);
 	Intermediates.TangentUp = Up;
 	Intermediates.TangentRight = Right;
 
@@ -847,14 +861,14 @@ FVertexFactoryIntermediates GetVertexFactoryIntermediates(FVertexFactoryInput In
 	UVForTexturing.xy = float2(0.5f, 0.5f) + (UVForTexturing.xy - float2(0.5f, 0.5f) ) *  UVScale.xy; 
 
 	// Vertex position.
-	float4 VertexWorldPosition = float4(ParticleWorldPosition.xyz + ResolvedView.PreViewTranslation,1);
+	float4 VertexWorldPosition = float4(ParticleWorldPosition + ResolvedView.PreViewTranslation,1);
 	float2 Size = abs(Intermediates.SizeRotSubImage.xy);
 	VertexWorldPosition += Size.x * (UVForPosition.x + NiagaraSpriteVF.PivotOffset.x) * float4(Right,0);
 	VertexWorldPosition += Size.y * (UVForPosition.y + NiagaraSpriteVF.PivotOffset.y) * float4(Up,0);
 	Intermediates.VertexWorldPosition = VertexWorldPosition.xyz;
 
 	// Previous vertex position.
-	Intermediates.PreviousVertexWorldPosition = VertexWorldPosition.xyz + (ParticleOldWorldPosition.xyz - ParticleWorldPosition.xyz);
+	Intermediates.PreviousVertexWorldPosition = VertexWorldPosition.xyz + (ParticleOldWorldPosition - ParticleWorldPosition);
 	
 	// SubUV.
 	float SubImageIndex = Intermediates.SizeRotSubImage.w;
@@ -876,7 +890,7 @@ FVertexFactoryIntermediates GetVertexFactoryIntermediates(FVertexFactoryInput In
 
 #if USE_PARTICLE_POSITION
 	float ParticleRadius = .5f * min(Intermediates.SizeRotSubImage.x, Intermediates.SizeRotSubImage.y);
-	float3 ParticleTranslatedWorldPosition = ParticleWorldPosition.xyz + ResolvedView.PreViewTranslation.xyz;
+	float3 ParticleTranslatedWorldPosition = ParticleWorldPosition + ResolvedView.PreViewTranslation.xyz;
 	Intermediates.TranslatedWorldPositionAndSize = float4(ParticleTranslatedWorldPosition, ParticleRadius);
 
 #if GENERATE_SPHERICAL_PARTICLE_NORMALS
@@ -886,7 +900,7 @@ FVertexFactoryIntermediates GetVertexFactoryIntermediates(FVertexFactoryInput In
 #endif
 
 #if USE_PARTICLE_VELOCITY
-	float3 ParticleVelocity = (ParticleWorldPosition.xyz - ParticleOldWorldPosition.xyz) * (1.0f/NiagaraSpriteVF.DeltaSeconds);
+	float3 ParticleVelocity = (ParticleWorldPosition - ParticleOldWorldPosition) * (1.0f/NiagaraSpriteVF.DeltaSeconds);
 	float3 DirectionVector = lerp(ParticleVelocity, Intermediates.CustomAlignmentVector, NiagaraSpriteVFLooseParameters.ParticleAlignmentMode == ALIGNMENT_MODE_CUSTOMALIGNMENT);
 	Intermediates.ParticleVelocity.xyz = normalize(DirectionVector);
 	Intermediates.ParticleVelocity.w = length(DirectionVector);
@@ -943,7 +957,7 @@ float3 VertexFactoryGetPositionForVertexLighting(FVertexFactoryInput Input, FVer
 {
 #if (TRANSLUCENCY_LIGHTING_VOLUMETRIC_PERVERTEX_DIRECTIONAL || TRANSLUCENCY_LIGHTING_VOLUMETRIC_PERVERTEX_NONDIRECTIONAL)
 	// Do per-vertex lighting with particle position instead of vertex position when we're using SubUV cutouts, since the vertex positions vary per-frame which would cause popping
-	return NiagaraSpriteVFLooseParameters.NumCutoutVerticesPerFrame > 0 ? (Intermediates.ParticleWorldPosition.xyz + ResolvedView.PreViewTranslation.xyz) : TranslatedWorldPosition;
+	return NiagaraSpriteVFLooseParameters.NumCutoutVerticesPerFrame > 0 ? (Intermediates.ParticleWorldPosition + ResolvedView.PreViewTranslation.xyz) : TranslatedWorldPosition;
 #else
 	return TranslatedWorldPosition;
 #endif
diff --git a/Engine/Plugins/FX/Niagara/Source/Niagara/Private/NiagaraComponent.cpp b/Engine/Plugins/FX/Niagara/Source/Niagara/Private/NiagaraComponent.cpp
--- a/Engine/Plugins/FX/Niagara/Source/Niagara/Private/NiagaraComponent.cpp
+++ b/Engine/Plugins/FX/Niagara/Source/Niagara/Private/NiagaraComponent.cpp
@@ -232,17 +232,6 @@ void FNiagaraSceneProxy::CreateRenderThreadResources()
 	return;
 }
 
-void FNiagaraSceneProxy::OnTransformChanged()
-{
-	for (FNiagaraRenderer* Renderer : EmitterRenderers)
-	{
-		if (Renderer)
-		{
-			Renderer->TransformChanged();
-		}
-	}
-}
-
 FPrimitiveViewRelevance FNiagaraSceneProxy::GetViewRelevance(const FSceneView* View) const
 {
 	FPrimitiveViewRelevance Relevance;
diff --git a/Engine/Plugins/FX/Niagara/Source/Niagara/Private/NiagaraRendererMeshes.cpp b/Engine/Plugins/FX/Niagara/Source/Niagara/Private/NiagaraRendererMeshes.cpp
--- a/Engine/Plugins/FX/Niagara/Source/Niagara/Private/NiagaraRendererMeshes.cpp
+++ b/Engine/Plugins/FX/Niagara/Source/Niagara/Private/NiagaraRendererMeshes.cpp
@@ -12,6 +12,7 @@
 
 DECLARE_CYCLE_STAT(TEXT("Generate Mesh Vertex Data [GT]"), STAT_NiagaraGenMeshVertexData, STATGROUP_Niagara);
 DECLARE_CYCLE_STAT(TEXT("Render Meshes [RT]"), STAT_NiagaraRenderMeshes, STATGROUP_Niagara);
+DECLARE_CYCLE_STAT(TEXT("Render Meshes - Allocate GPU Data [RT]"), STAT_NiagaraRenderMeshes_AllocateGPUData, STATGROUP_Niagara);
 
 
 DECLARE_DWORD_COUNTER_STAT(TEXT("NumMeshesRenderer"), STAT_NiagaraNumMeshes, STATGROUP_Niagara);
@@ -120,7 +121,6 @@ void FNiagaraRendererMeshes::ReleaseRenderThreadResources(NiagaraEmitterInstance
 {
 	FNiagaraRenderer::ReleaseRenderThreadResources(Batcher);
 	VertexFactory->ReleaseResource();
-	WorldSpacePrimitiveUniformBuffer.ReleaseResource();
 }
 
 void FNiagaraRendererMeshes::CreateRenderThreadResources(NiagaraEmitterInstanceBatcher* Batcher)
@@ -182,6 +182,7 @@ void FNiagaraRendererMeshes::GetDynamicMeshElements(const TArray<const FSceneVie
 	//For cpu sims we allocate render buffers from the global pool. GPU sims own their own.
 	if (SimTarget == ENiagaraSimTarget::CPUSim)
 	{
+		SCOPE_CYCLE_COUNTER(STAT_NiagaraRenderMeshes_AllocateGPUData);
 		ParticleData = DynamicReadBuffer.AllocateFloat(TotalFloatSize);
 		FMemory::Memcpy(ParticleData.Buffer, SourceParticleData->GetFloatBuffer().GetData(), SourceParticleData->GetFloatBuffer().Num());
 	}
@@ -191,31 +192,6 @@ void FNiagaraRendererMeshes::GetDynamicMeshElements(const TArray<const FSceneVie
 	}
 
 	{
-		// Update the primitive uniform buffer if needed.
-		if (!WorldSpacePrimitiveUniformBuffer.IsInitialized())
-		{
-			FPrimitiveUniformShaderParameters PrimitiveUniformShaderParameters = GetPrimitiveUniformShaderParameters(
-				FMatrix::Identity,
-				FMatrix::Identity,
-				SceneProxy->GetActorPosition(),
-				SceneProxy->GetBounds(),
-				SceneProxy->GetLocalBounds(),
-				SceneProxy->ReceivesDecals(),
-				false,
-				false,
-				false,
-				false,
-				SceneProxy->DrawsVelocity(),
-				SceneProxy->GetLightingChannelMask(),
-				0,
-				INDEX_NONE,
-				INDEX_NONE,
-				SceneProxy->AlwaysHasVelocity()
-				);
-			WorldSpacePrimitiveUniformBuffer.SetContents(PrimitiveUniformShaderParameters);
-			WorldSpacePrimitiveUniformBuffer.InitResource();
-		}
-
 		// Compute the per-view uniform buffers.
 		for (int32 ViewIndex = 0; ViewIndex < Views.Num(); ViewIndex++)
 		{
@@ -235,8 +211,9 @@ void FNiagaraRendererMeshes::GetDynamicMeshElements(const TArray<const FSceneVie
 				FNiagaraMeshCollectorResourcesMesh& CollectorResources = Collector.AllocateOneFrameResource<FNiagaraMeshCollectorResourcesMesh>();
 				SetupVertexFactory(&CollectorResources.VertexFactory, LODModel);
 				FNiagaraMeshUniformParameters PerViewUniformParameters;// = UniformParameters;
-				PerViewUniformParameters.LocalToWorld = bLocalSpace ? SceneProxy->GetLocalToWorld() : FMatrix::Identity;//For now just handle local space like this but maybe in future have a VF variant to avoid the transform entirely?
-				PerViewUniformParameters.LocalToWorldInverseTransposed = bLocalSpace ? SceneProxy->GetLocalToWorld().Inverse().GetTransposed() : FMatrix::Identity;
+				FMemory::Memzero(&PerViewUniformParameters,sizeof(PerViewUniformParameters)); // Clear unset bytes
+
+				PerViewUniformParameters.bLocalSpace = bLocalSpace;
 				PerViewUniformParameters.PrevTransformAvailable = false;
 				PerViewUniformParameters.DeltaSeconds = ViewFamily.DeltaWorldTime;
 				PerViewUniformParameters.PositionDataOffset = PositionOffset;
@@ -370,7 +347,7 @@ void FNiagaraRendererMeshes::GetDynamicMeshElements(const TArray<const FSceneVie
 					Mesh.DepthPriorityGroup = (ESceneDepthPriorityGroup)SceneProxy->GetDepthPriorityGroup(View);
 
 					FMeshBatchElement& BatchElement = Mesh.Elements[0];
-					BatchElement.PrimitiveUniformBuffer = WorldSpacePrimitiveUniformBuffer.GetUniformBufferRHI();
+					BatchElement.PrimitiveUniformBuffer = SceneProxy->GetUniformBuffer();
 					BatchElement.FirstIndex = 0;
 					BatchElement.MinVertexIndex = 0;
 					BatchElement.MaxVertexIndex = 0;
@@ -488,11 +465,6 @@ int FNiagaraRendererMeshes::GetDynamicDataSize()const
 	return Size;
 }
 
-void FNiagaraRendererMeshes::TransformChanged()
-{
-	WorldSpacePrimitiveUniformBuffer.ReleaseResource();
-}
-
 bool FNiagaraRendererMeshes::IsMaterialValid(UMaterialInterface* Mat)const
 {
 	return Mat && Mat->CheckMaterialUsage(MATUSAGE_NiagaraMeshParticles);
diff --git a/Engine/Plugins/FX/Niagara/Source/Niagara/Private/NiagaraRendererRibbons.cpp b/Engine/Plugins/FX/Niagara/Source/Niagara/Private/NiagaraRendererRibbons.cpp
--- a/Engine/Plugins/FX/Niagara/Source/Niagara/Private/NiagaraRendererRibbons.cpp
+++ b/Engine/Plugins/FX/Niagara/Source/Niagara/Private/NiagaraRendererRibbons.cpp
@@ -206,7 +206,6 @@ void FNiagaraRendererRibbons::ReleaseRenderThreadResources(NiagaraEmitterInstanc
 {
 	FNiagaraRenderer::ReleaseRenderThreadResources(Batcher);
 	VertexFactory->ReleaseResource();
-	WorldSpacePrimitiveUniformBuffer.ReleaseResource();
 }
 
 // FPrimitiveSceneProxy interface.
@@ -293,31 +292,6 @@ void FNiagaraRendererRibbons::GetDynamicMeshElements(const TArray<const FSceneVi
 		FMemory::Memcpy(ParticleData.Buffer, SourceParticleData->GetFloatBuffer().GetData(), SourceParticleData->GetFloatBuffer().Num());
 	}
 
-	// Update the primitive uniform buffer if needed.
-	if (!WorldSpacePrimitiveUniformBuffer.IsInitialized())
-	{
-		FPrimitiveUniformShaderParameters PrimitiveUniformShaderParameters = GetPrimitiveUniformShaderParameters(
-			FMatrix::Identity,
-			FMatrix::Identity,
-			SceneProxy->GetActorPosition(),
-			SceneProxy->GetBounds(),
-			SceneProxy->GetLocalBounds(),
-			SceneProxy->ReceivesDecals(),
-			false,
-			false,
-			SceneProxy->UseSingleSampleShadowFromStationaryLights(),
-			SceneProxy->GetScene().HasPrecomputedVolumetricLightmap_RenderThread(),
-			SceneProxy->DrawsVelocity(),
-			SceneProxy->GetLightingChannelMask(),
-			0,
-			INDEX_NONE,
-			INDEX_NONE,
-			SceneProxy->AlwaysHasVelocity()
-			);
-		WorldSpacePrimitiveUniformBuffer.SetContents(PrimitiveUniformShaderParameters);
-		WorldSpacePrimitiveUniformBuffer.InitResource();
-	}
-
 	// Modify tessellation parameters based on tessellation mode
 	// Set auto defaults here:
 	int32 TessellationFactor = GNiagaraRibbonMaxTessellation;
@@ -395,8 +369,9 @@ void FNiagaraRendererRibbons::GetDynamicMeshElements(const TArray<const FSceneVi
 
 			FNiagaraMeshCollectorResourcesRibbon& CollectorResources = Collector.AllocateOneFrameResource<FNiagaraMeshCollectorResourcesRibbon>();
 			FNiagaraRibbonUniformParameters PerViewUniformParameters;// = UniformParameters;
-			PerViewUniformParameters.LocalToWorld = bLocalSpace ? SceneProxy->GetLocalToWorld() : FMatrix::Identity;//For now just handle local space like this but maybe in future have a VF variant to avoid the transform entirely?
-			PerViewUniformParameters.LocalToWorldInverseTransposed = bLocalSpace ? SceneProxy->GetLocalToWorld().Inverse().GetTransposed() : FMatrix::Identity;
+			FMemory::Memzero(&PerViewUniformParameters,sizeof(PerViewUniformParameters)); // Clear unset bytes
+			
+			PerViewUniformParameters.bLocalSpace = bLocalSpace;
 			PerViewUniformParameters.DeltaSeconds = ViewFamily.DeltaWorldTime;
 			PerViewUniformParameters.CameraUp = View->GetViewUp(); // FVector4(0.0f, 0.0f, 1.0f, 0.0f);
 			PerViewUniformParameters.CameraRight = View->GetViewRight();//	FVector4(1.0f, 0.0f, 0.0f, 0.0f);
@@ -500,7 +475,7 @@ void FNiagaraRendererRibbons::GetDynamicMeshElements(const TArray<const FSceneVi
 			MeshElement.NumInstances = 1;
 			MeshElement.MinVertexIndex = 0;
 			MeshElement.MaxVertexIndex = 0;
-			MeshElement.PrimitiveUniformBuffer = WorldSpacePrimitiveUniformBuffer.GetUniformBufferRHI();
+			MeshElement.PrimitiveUniformBuffer = SceneProxy->GetUniformBuffer();
 
 			Collector.AddMesh(ViewIndex, MeshBatch);
 		}
@@ -525,11 +500,6 @@ int FNiagaraRendererRibbons::GetDynamicDataSize()const
 	return Size;
 }
 
-void FNiagaraRendererRibbons::TransformChanged()
-{
-	WorldSpacePrimitiveUniformBuffer.ReleaseResource();
-}
-
 void CalculateUVScaleAndOffsets(
 	const FNiagaraDataSetAccessor<float>& SortKeyData, const TArray<int32>& RibbonIndices,
 	bool bSortKeyIsAge, int32 StartIndex, int32 EndIndex, int32 NumSegments,
diff --git a/Engine/Plugins/FX/Niagara/Source/Niagara/Private/NiagaraRendererSprites.cpp b/Engine/Plugins/FX/Niagara/Source/Niagara/Private/NiagaraRendererSprites.cpp
--- a/Engine/Plugins/FX/Niagara/Source/Niagara/Private/NiagaraRendererSprites.cpp
+++ b/Engine/Plugins/FX/Niagara/Source/Niagara/Private/NiagaraRendererSprites.cpp
@@ -151,7 +151,6 @@ void FNiagaraRendererSprites::ReleaseRenderThreadResources(NiagaraEmitterInstanc
 {
 	FNiagaraRenderer::ReleaseRenderThreadResources(Batcher);
 	VertexFactory->ReleaseResource();
-	WorldSpacePrimitiveUniformBuffer.ReleaseResource();
 
 	CutoutVertexBuffer.ReleaseResource();
 #if RHI_RAYTRACING
@@ -191,34 +190,6 @@ void FNiagaraRendererSprites::CreateRenderThreadResources(NiagaraEmitterInstance
 #endif
 }
 
-void FNiagaraRendererSprites::ConditionalInitPrimitiveUniformBuffer(const FNiagaraSceneProxy *SceneProxy) const
-{
-	// Update the primitive uniform buffer if needed.
-	if (!WorldSpacePrimitiveUniformBuffer.IsInitialized())
-	{
-		FPrimitiveUniformShaderParameters PrimitiveUniformShaderParameters = GetPrimitiveUniformShaderParameters(
-			FMatrix::Identity,
-			FMatrix::Identity,
-			SceneProxy->GetActorPosition(),
-			SceneProxy->GetBounds(),
-			SceneProxy->GetLocalBounds(),
-			SceneProxy->ReceivesDecals(),
-			false,
-			false,
-			SceneProxy->UseSingleSampleShadowFromStationaryLights(),
-			SceneProxy->GetScene().HasPrecomputedVolumetricLightmap_RenderThread(),
-			SceneProxy->DrawsVelocity(),
-			SceneProxy->GetLightingChannelMask(),
-			0,
-			INDEX_NONE,
-			INDEX_NONE,
-			SceneProxy->AlwaysHasVelocity()
-		);
-		WorldSpacePrimitiveUniformBuffer.SetContents(PrimitiveUniformShaderParameters);
-		WorldSpacePrimitiveUniformBuffer.InitResource();
-	}
-}
-
 FNiagaraRendererSprites::FCPUSimParticleDataAllocation FNiagaraRendererSprites::ConditionalAllocateCPUSimParticleData(FNiagaraDynamicDataSprites* DynamicDataSprites, FGlobalDynamicReadBuffer& DynamicReadBuffer) const
 {
 	FNiagaraDataBuffer* SourceParticleData = DynamicDataSprites->GetParticleDataToRender();
@@ -243,8 +214,9 @@ FNiagaraRendererSprites::FCPUSimParticleDataAllocation FNiagaraRendererSprites::
 FNiagaraSpriteUniformBufferRef FNiagaraRendererSprites::CreatePerViewUniformBuffer(const FSceneView* View, const FSceneViewFamily& ViewFamily, const FNiagaraSceneProxy *SceneProxy) const
 {
 	FNiagaraSpriteUniformParameters PerViewUniformParameters;
-	PerViewUniformParameters.LocalToWorld = bLocalSpace ? SceneProxy->GetLocalToWorld() : FMatrix::Identity;//For now just handle local space like this but maybe in future have a VF variant to avoid the transform entirely?
-	PerViewUniformParameters.LocalToWorldInverseTransposed = bLocalSpace ? SceneProxy->GetLocalToWorld().Inverse().GetTransposed() : FMatrix::Identity;
+	FMemory::Memzero(&PerViewUniformParameters,sizeof(PerViewUniformParameters)); // Clear unset bytes
+
+	PerViewUniformParameters.bLocalSpace = bLocalSpace;
 	PerViewUniformParameters.RotationBias = 0.0f;
 	PerViewUniformParameters.RotationScale = 1.0f;
 	PerViewUniformParameters.TangentSelector = FVector4(0.0f, 0.0f, 0.0f, 1.0f);
@@ -526,7 +498,7 @@ void FNiagaraRendererSprites::CreateMeshBatchForView(
 	MeshElement.NumInstances = FMath::Max(0, NumInstances);	//->VertexData.Num();
 	MeshElement.MinVertexIndex = 0;
 	MeshElement.MaxVertexIndex = 0;// MeshElement.NumInstances * 4 - 1;
-	MeshElement.PrimitiveUniformBuffer = WorldSpacePrimitiveUniformBuffer.GetUniformBufferRHI();
+	MeshElement.PrimitiveUniformBuffer = SceneProxy->GetUniformBuffer();
 	if (IndirectArgsOffset != INDEX_NONE)
 	{
 		NiagaraEmitterInstanceBatcher* Batcher = SceneProxy->GetBatcher();
@@ -576,8 +548,6 @@ void FNiagaraRendererSprites::GetDynamicMeshElements(const TArray<const FSceneVi
 	
 	FCPUSimParticleDataAllocation CPUSimParticleDataAllocation = ConditionalAllocateCPUSimParticleData(DynamicDataSprites, Collector.GetDynamicReadBuffer());
 
-	ConditionalInitPrimitiveUniformBuffer(SceneProxy);
-
 	uint32 IndirectArgsOffset = INDEX_NONE;
 	if (SimTarget == ENiagaraSimTarget::GPUComputeSim)
 	{
@@ -642,7 +612,6 @@ void FNiagaraRendererSprites::GetDynamicRayTracingInstances(FRayTracingMaterialG
 	{
 		// Setup material for our ray tracing instance
 		FCPUSimParticleDataAllocation CPUSimParticleDataAllocation = ConditionalAllocateCPUSimParticleData(DynamicDataSprites, Context.RayTracingMeshResourceCollector.GetDynamicReadBuffer());
-		ConditionalInitPrimitiveUniformBuffer(SceneProxy);
 		FNiagaraMeshCollectorResourcesSprite& CollectorResources = Context.RayTracingMeshResourceCollector.AllocateOneFrameResource<FNiagaraMeshCollectorResourcesSprite>();
 		SetVertexFactoryParticleData(CollectorResources.VertexFactory, DynamicDataSprites, CPUSimParticleDataAllocation, Context.ReferenceView, SceneProxy);
 		CollectorResources.UniformBuffer = CreatePerViewUniformBuffer(Context.ReferenceView, Context.ReferenceViewFamily, SceneProxy);
@@ -715,11 +684,6 @@ int FNiagaraRendererSprites::GetDynamicDataSize()const
 	return Size;
 }
 
-void FNiagaraRendererSprites::TransformChanged()
-{
-	WorldSpacePrimitiveUniformBuffer.ReleaseResource();
-}
-
 bool FNiagaraRendererSprites::IsMaterialValid(UMaterialInterface* Mat)const
 {
 	return Mat && Mat->CheckMaterialUsage(MATUSAGE_NiagaraSprites);
diff --git a/Engine/Plugins/FX/Niagara/Source/Niagara/Public/NiagaraComponent.h b/Engine/Plugins/FX/Niagara/Source/Niagara/Public/NiagaraComponent.h
--- a/Engine/Plugins/FX/Niagara/Source/Niagara/Public/NiagaraComponent.h
+++ b/Engine/Plugins/FX/Niagara/Source/Niagara/Public/NiagaraComponent.h
@@ -516,7 +516,6 @@ private:
 	virtual void CreateRenderThreadResources() override;
 
 	//virtual void OnActorPositionChanged() override;
-	virtual void OnTransformChanged() override;
 
 	virtual void GetDynamicMeshElements(const TArray<const FSceneView*>& Views, const FSceneViewFamily& ViewFamily, uint32 VisibilityMap, FMeshElementCollector& Collector) const override;
 
diff --git a/Engine/Plugins/FX/Niagara/Source/Niagara/Public/NiagaraRenderer.h b/Engine/Plugins/FX/Niagara/Source/Niagara/Public/NiagaraRenderer.h
--- a/Engine/Plugins/FX/Niagara/Source/Niagara/Public/NiagaraRenderer.h
+++ b/Engine/Plugins/FX/Niagara/Source/Niagara/Public/NiagaraRenderer.h
@@ -99,7 +99,6 @@ public:
 
 	virtual void GatherSimpleLights(FSimpleLightArray& OutParticleLights)const {}
 	virtual int32 GetDynamicDataSize()const { return 0; }
-	virtual void TransformChanged() {}
 	virtual bool IsMaterialValid(UMaterialInterface* Mat)const { return Mat != nullptr; }
 
 	void SortIndices(ENiagaraSortMode SortMode, int32 SortAttributeOffset, const FNiagaraDataBuffer& Buffer, const FMatrix& LocalToWorld, const FSceneView* View, FGlobalDynamicReadBuffer::FAllocation& OutIndices)const;
diff --git a/Engine/Plugins/FX/Niagara/Source/Niagara/Public/NiagaraRendererMeshes.h b/Engine/Plugins/FX/Niagara/Source/Niagara/Public/NiagaraRendererMeshes.h
--- a/Engine/Plugins/FX/Niagara/Source/Niagara/Public/NiagaraRendererMeshes.h
+++ b/Engine/Plugins/FX/Niagara/Source/Niagara/Public/NiagaraRendererMeshes.h
@@ -26,14 +26,13 @@ public:
 	virtual void GetDynamicMeshElements(const TArray<const FSceneView*>& Views, const FSceneViewFamily& ViewFamily, uint32 VisibilityMap, FMeshElementCollector& Collector, const FNiagaraSceneProxy *SceneProxy) const override;
 	virtual FNiagaraDynamicDataBase* GenerateDynamicData(const FNiagaraSceneProxy* Proxy, const UNiagaraRendererProperties* InProperties, const FNiagaraEmitterInstance* Emitter) const override;
 	virtual int32 GetDynamicDataSize()const override;
-	virtual void TransformChanged() override;
 	virtual bool IsMaterialValid(UMaterialInterface* Mat)const override;
 	//FNiagaraRenderer Interface END
 
 	void SetupVertexFactory(FNiagaraMeshVertexFactory *InVertexFactory, const FStaticMeshLODResources& LODResources) const;
 
 private:
-	mutable TUniformBuffer<FPrimitiveUniformShaderParameters> WorldSpacePrimitiveUniformBuffer;
+
 	class FNiagaraMeshVertexFactory* VertexFactory;
 
 	/** Render data of the static mesh we use. */
diff --git a/Engine/Plugins/FX/Niagara/Source/Niagara/Public/NiagaraRendererRibbons.h b/Engine/Plugins/FX/Niagara/Source/Niagara/Public/NiagaraRendererRibbons.h
--- a/Engine/Plugins/FX/Niagara/Source/Niagara/Public/NiagaraRendererRibbons.h
+++ b/Engine/Plugins/FX/Niagara/Source/Niagara/Public/NiagaraRendererRibbons.h
@@ -27,7 +27,6 @@ public:
 	virtual void GetDynamicMeshElements(const TArray<const FSceneView*>& Views, const FSceneViewFamily& ViewFamily, uint32 VisibilityMap, FMeshElementCollector& Collector, const FNiagaraSceneProxy *SceneProxy) const override;
 	virtual FNiagaraDynamicDataBase *GenerateDynamicData(const FNiagaraSceneProxy* Proxy, const UNiagaraRendererProperties* InProperties, const FNiagaraEmitterInstance* Emitteride) const override;
 	virtual int32 GetDynamicDataSize()const override;
-	virtual void TransformChanged() override;
 	virtual bool IsMaterialValid(UMaterialInterface* Mat)const override;
 	//FNiagaraInterface END
 
@@ -37,7 +36,6 @@ protected:
 
 private:
 	class FNiagaraRibbonVertexFactory *VertexFactory;
-	mutable TUniformBuffer<FPrimitiveUniformShaderParameters> WorldSpacePrimitiveUniformBuffer;
 
 	ENiagaraRibbonFacingMode FacingMode;
 	float UV0TilingDistance;
diff --git a/Engine/Plugins/FX/Niagara/Source/Niagara/Public/NiagaraRendererSprites.h b/Engine/Plugins/FX/Niagara/Source/Niagara/Public/NiagaraRendererSprites.h
--- a/Engine/Plugins/FX/Niagara/Source/Niagara/Public/NiagaraRendererSprites.h
+++ b/Engine/Plugins/FX/Niagara/Source/Niagara/Public/NiagaraRendererSprites.h
@@ -26,7 +26,6 @@ public:
 	virtual void GetDynamicMeshElements(const TArray<const FSceneView*>& Views, const FSceneViewFamily& ViewFamily, uint32 VisibilityMap, FMeshElementCollector& Collector, const FNiagaraSceneProxy *SceneProxy) const override;
 	virtual FNiagaraDynamicDataBase* GenerateDynamicData(const FNiagaraSceneProxy* Proxy, const UNiagaraRendererProperties* InProperties, const FNiagaraEmitterInstance* Emitter) const override;
 	virtual int GetDynamicDataSize()const override;
-	virtual void TransformChanged()override;
 	virtual bool IsMaterialValid(UMaterialInterface* Mat)const override;
 
 #if RHI_RAYTRACING
@@ -41,7 +40,6 @@ private:
 		FGlobalDynamicReadBuffer::FAllocation ParticleData;
 	};
 
-	void ConditionalInitPrimitiveUniformBuffer(const FNiagaraSceneProxy *SceneProxy) const;
 	FCPUSimParticleDataAllocation ConditionalAllocateCPUSimParticleData(FNiagaraDynamicDataSprites *DynamicDataSprites, FGlobalDynamicReadBuffer& DynamicReadBuffer) const;
 	TUniformBufferRef<class FNiagaraSpriteUniformParameters> CreatePerViewUniformBuffer(const FSceneView* View, const FSceneViewFamily& ViewFamily, const FNiagaraSceneProxy *SceneProxy) const;
 	void SetVertexFactoryParticleData(
@@ -59,7 +57,6 @@ private:
 		FMeshBatch& OutMeshBatch,
 		class FNiagaraMeshCollectorResourcesSprite& OutCollectorResources) const;
 
-	mutable TUniformBuffer<FPrimitiveUniformShaderParameters> WorldSpacePrimitiveUniformBuffer;
 	class FNiagaraSpriteVertexFactory* VertexFactory;
 
 	//Cached data from the properties struct.
diff --git a/Engine/Plugins/FX/Niagara/Source/NiagaraVertexFactories/Public/NiagaraMeshVertexFactory.h b/Engine/Plugins/FX/Niagara/Source/NiagaraVertexFactories/Public/NiagaraMeshVertexFactory.h
--- a/Engine/Plugins/FX/Niagara/Source/NiagaraVertexFactories/Public/NiagaraMeshVertexFactory.h
+++ b/Engine/Plugins/FX/Niagara/Source/NiagaraVertexFactories/Public/NiagaraMeshVertexFactory.h
@@ -28,8 +28,7 @@ struct FShaderCompilerEnvironment;
 * Uniform buffer for mesh particle vertex factories.
 */
 BEGIN_GLOBAL_SHADER_PARAMETER_STRUCT(FNiagaraMeshUniformParameters, NIAGARAVERTEXFACTORIES_API)
-	SHADER_PARAMETER_EX( FMatrix, LocalToWorld, EShaderPrecisionModifier::Half)
-	SHADER_PARAMETER_EX( FMatrix, LocalToWorldInverseTransposed, EShaderPrecisionModifier::Half)
+	SHADER_PARAMETER(bool , bLocalSpace)
 	SHADER_PARAMETER(FVector4, SubImageSize)
 	SHADER_PARAMETER(uint32, TexCoordWeightA)
 	SHADER_PARAMETER(uint32, TexCoordWeightB)
diff --git a/Engine/Plugins/FX/Niagara/Source/NiagaraVertexFactories/Public/NiagaraRibbonVertexFactory.h b/Engine/Plugins/FX/Niagara/Source/NiagaraVertexFactories/Public/NiagaraRibbonVertexFactory.h
--- a/Engine/Plugins/FX/Niagara/Source/NiagaraVertexFactories/Public/NiagaraRibbonVertexFactory.h
+++ b/Engine/Plugins/FX/Niagara/Source/NiagaraVertexFactories/Public/NiagaraRibbonVertexFactory.h
@@ -48,8 +48,7 @@ BEGIN_GLOBAL_SHADER_PARAMETER_STRUCT(FNiagaraRibbonUniformParameters, NIAGARAVER
 	SHADER_PARAMETER(float, OneOverUV0TilingDistance)
 	SHADER_PARAMETER(float, OneOverUV1TilingDistance)
 	SHADER_PARAMETER(FVector4, PackedVData)
-	SHADER_PARAMETER_EX(FMatrix, LocalToWorld, EShaderPrecisionModifier::Half)
-	SHADER_PARAMETER_EX(FMatrix, LocalToWorldInverseTransposed, EShaderPrecisionModifier::Half)
+	SHADER_PARAMETER(bool, bLocalSpace)
 	SHADER_PARAMETER_EX(float, DeltaSeconds, EShaderPrecisionModifier::Half)
 END_GLOBAL_SHADER_PARAMETER_STRUCT()
 typedef TUniformBufferRef<FNiagaraRibbonUniformParameters> FNiagaraRibbonUniformBufferRef;
diff --git a/Engine/Plugins/FX/Niagara/Source/NiagaraVertexFactories/Public/NiagaraSpriteVertexFactory.h b/Engine/Plugins/FX/Niagara/Source/NiagaraVertexFactories/Public/NiagaraSpriteVertexFactory.h
--- a/Engine/Plugins/FX/Niagara/Source/NiagaraVertexFactories/Public/NiagaraSpriteVertexFactory.h
+++ b/Engine/Plugins/FX/Niagara/Source/NiagaraVertexFactories/Public/NiagaraSpriteVertexFactory.h
@@ -19,8 +19,7 @@ class FMaterial;
  * Uniform buffer for particle sprite vertex factories.
  */
 BEGIN_GLOBAL_SHADER_PARAMETER_STRUCT( FNiagaraSpriteUniformParameters, NIAGARAVERTEXFACTORIES_API)
-	SHADER_PARAMETER_EX( FMatrix, LocalToWorld, EShaderPrecisionModifier::Half)
-	SHADER_PARAMETER_EX( FMatrix, LocalToWorldInverseTransposed, EShaderPrecisionModifier::Half)
+	SHADER_PARAMETER( bool, bLocalSpace)
 	SHADER_PARAMETER_EX( FVector, CustomFacingVectorMask, EShaderPrecisionModifier::Half)
 	SHADER_PARAMETER_EX( FVector4, TangentSelector, EShaderPrecisionModifier::Half )
 	SHADER_PARAMETER_EX( FVector4, NormalsSphereCenter, EShaderPrecisionModifier::Half )
diff --git a/Engine/Shaders/Private/VertexFactoryCommon.ush b/Engine/Shaders/Private/VertexFactoryCommon.ush
--- a/Engine/Shaders/Private/VertexFactoryCommon.ush
+++ b/Engine/Shaders/Private/VertexFactoryCommon.ush
@@ -11,6 +11,12 @@ float3 TransformLocalToWorld(float3 LocalPosition, uint PrimitiveId)
 	return INVARIANT((LocalToWorld[0].xyz * LocalPosition.xxx + LocalToWorld[1].xyz * LocalPosition.yyy + LocalToWorld[2].xyz * LocalPosition.zzz) + LocalToWorld[3].xyz);	
 }
 
+float4 TransformLocalToWorld(float3 LocalPosition)
+{
+	float3 RotatedPosition = Primitive.LocalToWorld[0].xyz * LocalPosition.xxx + Primitive.LocalToWorld[1].xyz * LocalPosition.yyy + Primitive.LocalToWorld[2].xyz * LocalPosition.zzz;
+	return float4(RotatedPosition + Primitive.LocalToWorld[3].xyz ,1);
+}
+
 float4 TransformLocalToTranslatedWorld(float3 LocalPosition, uint PrimitiveId)
 {
 	float4x4 LocalToWorld = GetPrimitiveData(PrimitiveId).LocalToWorld;
