Adding material uobject params to Mesh renderer. This now allows materials to be set on Niagara systems.

diff --git a/Engine/Plugins/FX/Niagara/Source/Niagara/Classes/NiagaraEmitterInstance.h b/Engine/Plugins/FX/Niagara/Source/Niagara/Classes/NiagaraEmitterInstance.h
--- a/Engine/Plugins/FX/Niagara/Source/Niagara/Classes/NiagaraEmitterInstance.h
+++ b/Engine/Plugins/FX/Niagara/Source/Niagara/Classes/NiagaraEmitterInstance.h
@@ -119,6 +119,8 @@ public:
 
 	void SetSystemFixedBoundsOverride(FBox SystemFixedBounds);
 
+	bool FindBinding(const FNiagaraUserParameterBinding& InBinding, TArray<UMaterialInterface*>& OutMaterials) const;
+
 private:
 
 	void CheckForErrors();
diff --git a/Engine/Plugins/FX/Niagara/Source/Niagara/Private/NiagaraComponent.cpp b/Engine/Plugins/FX/Niagara/Source/Niagara/Private/NiagaraComponent.cpp
--- a/Engine/Plugins/FX/Niagara/Source/Niagara/Private/NiagaraComponent.cpp
+++ b/Engine/Plugins/FX/Niagara/Source/Niagara/Private/NiagaraComponent.cpp
@@ -1073,7 +1073,7 @@ int32 UNiagaraComponent::GetNumMaterials() const
 			for (int32 EmitterIdx = 0; EmitterIdx < Emitter->GetRenderers().Num(); EmitterIdx++)
 			{
 				UNiagaraRendererProperties* Properties = Emitter->GetRenderers()[EmitterIdx];
-				Properties->GetUsedMaterials(UsedMaterials);
+				Properties->GetUsedMaterials(EmitterInst, UsedMaterials);
 			}
 		}
 	}
@@ -1143,7 +1143,7 @@ void UNiagaraComponent::GetUsedMaterials(TArray<UMaterialInterface*>& OutMateria
 			{
 				if (UNiagaraRendererProperties* Renderer = Props->GetRenderers()[i])
 				{
-					Renderer->GetUsedMaterials(OutMaterials);
+					Renderer->GetUsedMaterials(&Sim.Get(), OutMaterials);
 				}
 			}
 		}
@@ -1274,6 +1274,12 @@ void UNiagaraComponent::SetVariableObject(FName InVariableName, UObject* InValue
 	OverrideParameters.SetUObject(InValue, FNiagaraVariable(FNiagaraTypeDefinition::GetUObjectDef(), InVariableName));
 }
 
+void UNiagaraComponent::SetVariableMaterial(FName InVariableName, UMaterialInterface* InValue)
+{
+	OverrideParameters.SetUObject(InValue, FNiagaraVariable(FNiagaraTypeDefinition::GetUMaterialDef(), InVariableName));
+	MarkRenderStateDirty(); // Materials might be using this on the system, so invalidate the render state to re-gather them.
+}
+
 TArray<FVector> UNiagaraComponent::GetNiagaraParticlePositions_DebugOnly(const FString& InEmitterName)
 {
 	return GetNiagaraParticleValueVec3_DebugOnly(InEmitterName, TEXT("Position"));
@@ -1409,6 +1415,13 @@ void UNiagaraComponent::PostEditChangeProperty(FPropertyChangedEvent& PropertyCh
 	Super::PostEditChangeProperty(PropertyChangedEvent);
 }
 
+void UNiagaraComponent::OverrideUObjectParameter(const FNiagaraVariable& InVar, UObject* InObj)
+{
+	FNiagaraUserRedirectionParameterStore& ParamStore = GetOverrideParameters();
+	ParamStore.SetUObject(InObj, InVar);
+	SetParameterValueOverriddenLocally(InVar, true, false);
+}
+
 void UNiagaraComponent::SynchronizeWithSourceSystem()
 {
 	// Synchronizing parameters will create new data interface objects and if the old data
@@ -1584,6 +1597,11 @@ void UNiagaraComponent::SetParameterValueOverriddenLocally(const FNiagaraVariabl
 		Asset->GetExposedParameters().CopyParameterData(OverrideParameters, InParam);
 	}
 	
+	if (InParam.IsUObject() && InParam.GetType().GetClass()->IsChildOf(UMaterialInterface::StaticClass()))
+	{
+		MarkRenderStateDirty();
+	}
+
 	if (bRequiresSystemInstanceReset && SystemInstance)
 	{
 		SystemInstance->Reset(FNiagaraSystemInstance::EResetMode::ResetAll);
diff --git a/Engine/Plugins/FX/Niagara/Source/Niagara/Private/NiagaraEmitterInstance.cpp b/Engine/Plugins/FX/Niagara/Source/Niagara/Private/NiagaraEmitterInstance.cpp
--- a/Engine/Plugins/FX/Niagara/Source/Niagara/Private/NiagaraEmitterInstance.cpp
+++ b/Engine/Plugins/FX/Niagara/Source/Niagara/Private/NiagaraEmitterInstance.cpp
@@ -1630,3 +1630,26 @@ void FNiagaraEmitterInstance::SetExecutionState(ENiagaraExecutionState InState)
 	}
 
 }
+
+bool FNiagaraEmitterInstance::FindBinding(const FNiagaraUserParameterBinding& InBinding, TArray<UMaterialInterface*>& OutMaterials) const
+{
+	FNiagaraSystemInstance* SystemInstance = GetParentSystemInstance();
+	if (SystemInstance)
+	{
+		UNiagaraComponent* Component = SystemInstance->GetComponent();
+		if (Component)
+		{
+			UObject* Obj = Component->GetOverrideParameters().GetUObject(InBinding.Parameter);
+			if (Obj)
+			{
+				UMaterialInterface* Material = Cast<UMaterialInterface>(Obj);
+				if (Material)
+				{
+					OutMaterials.Add(Material);
+					return true;
+				}
+			}
+		}
+	}
+	return false;
+}
\ No newline at end of file
diff --git a/Engine/Plugins/FX/Niagara/Source/Niagara/Private/NiagaraLightRendererProperties.cpp b/Engine/Plugins/FX/Niagara/Source/Niagara/Private/NiagaraLightRendererProperties.cpp
--- a/Engine/Plugins/FX/Niagara/Source/Niagara/Private/NiagaraLightRendererProperties.cpp
+++ b/Engine/Plugins/FX/Niagara/Source/Niagara/Private/NiagaraLightRendererProperties.cpp
@@ -44,7 +44,7 @@ FNiagaraRenderer* UNiagaraLightRendererProperties::CreateEmitterRenderer(ERHIFea
 	return NewRenderer;
 }
 
-void UNiagaraLightRendererProperties::GetUsedMaterials(TArray<UMaterialInterface*>& OutMaterials) const
+void UNiagaraLightRendererProperties::GetUsedMaterials(const FNiagaraEmitterInstance* InEmitter, TArray<UMaterialInterface*>& OutMaterials) const
 {
 	//OutMaterials.Add(Material);
 	//Material should live here.
diff --git a/Engine/Plugins/FX/Niagara/Source/Niagara/Private/NiagaraMeshRendererProperties.cpp b/Engine/Plugins/FX/Niagara/Source/Niagara/Private/NiagaraMeshRendererProperties.cpp
--- a/Engine/Plugins/FX/Niagara/Source/Niagara/Private/NiagaraMeshRendererProperties.cpp
+++ b/Engine/Plugins/FX/Niagara/Source/Niagara/Private/NiagaraMeshRendererProperties.cpp
@@ -6,6 +6,25 @@
 #include "NiagaraConstants.h"
 #include "NiagaraBoundsCalculatorHelper.h"
 
+FNiagaraMeshMaterialOverride::FNiagaraMeshMaterialOverride()
+{
+	FNiagaraTypeDefinition MaterialDef(UMaterialInterface::StaticClass());
+	UserParamBinding.Parameter.SetType(MaterialDef);
+}
+
+bool FNiagaraMeshMaterialOverride::SerializeFromMismatchedTag(const struct FPropertyTag& Tag, FStructuredArchive::FSlot Slot)
+{
+	// We have to handle the fact that UNiagaraMeshRendererProperties OverrideMaterials just used to be an array of UMaterialInterfaces
+	if (Tag.Type == NAME_ObjectProperty)
+	{
+		Slot << ExplicitMat;
+		return true;
+	}
+
+	return false;
+}
+
+
 UNiagaraMeshRendererProperties::UNiagaraMeshRendererProperties()
 	: ParticleMesh(nullptr)
 	, SortMode(ENiagaraSortMode::ViewDistance)
@@ -75,7 +94,9 @@ void UNiagaraMeshRendererProperties::InitBindings()
 	}
 }
 
-void UNiagaraMeshRendererProperties::GetUsedMaterials(TArray<UMaterialInterface*>& OutMaterials) const
+
+
+void UNiagaraMeshRendererProperties::GetUsedMaterials(const FNiagaraEmitterInstance* InEmitter, TArray<UMaterialInterface*>& OutMaterials) const
 {
 	if (ParticleMesh)
 	{
@@ -87,9 +108,27 @@ void UNiagaraMeshRendererProperties::GetUsedMaterials(TArray<UMaterialInterface*
 				const FStaticMeshSection& Section = LODModel.Sections[SectionIndex];
 				UMaterialInterface* ParticleMeshMaterial = ParticleMesh->GetMaterial(Section.MaterialIndex);
 
-				if (Section.MaterialIndex >= 0 && OverrideMaterials.Num() > Section.MaterialIndex && OverrideMaterials[Section.MaterialIndex] != nullptr)
+				if (Section.MaterialIndex >= 0 && OverrideMaterials.Num() > Section.MaterialIndex)
 				{
-					OutMaterials.Add(OverrideMaterials[Section.MaterialIndex]);
+					bool bSet = false;
+					
+					// UserParamBinding, if mapped to a real value, always wins. Otherwise, use the ExplictMat if it is set. Finally, fall
+					// back to the particle mesh material. This allows the user to effectively optionally bind to a Material binding
+					// and still have good defaults if it isn't set to anything.
+					if (OverrideMaterials[Section.MaterialIndex].UserParamBinding.Parameter.IsValid() && InEmitter->FindBinding(OverrideMaterials[Section.MaterialIndex].UserParamBinding, OutMaterials))
+					{
+						bSet = true;
+					}
+					else if (OverrideMaterials[Section.MaterialIndex].ExplicitMat != nullptr)
+					{
+						bSet = true;
+						OutMaterials.Add(OverrideMaterials[Section.MaterialIndex].ExplicitMat);
+					}
+
+					if (!bSet)
+					{
+						OutMaterials.Add(ParticleMeshMaterial);
+					}
 				}
 				else
 				{
@@ -116,6 +155,12 @@ uint32 UNiagaraMeshRendererProperties::GetNumIndicesPerInstance() const
 }
 
 
+void UNiagaraMeshRendererProperties::PostLoad()
+{
+	Super::PostLoad();
+}
+
+
 #if WITH_EDITORONLY_DATA
 bool UNiagaraMeshRendererProperties::IsMaterialValidForRenderer(UMaterial* Material, FText& InvalidMessage)
 {
diff --git a/Engine/Plugins/FX/Niagara/Source/Niagara/Private/NiagaraModule.cpp b/Engine/Plugins/FX/Niagara/Source/Niagara/Private/NiagaraModule.cpp
--- a/Engine/Plugins/FX/Niagara/Source/Niagara/Private/NiagaraModule.cpp
+++ b/Engine/Plugins/FX/Niagara/Source/Niagara/Private/NiagaraModule.cpp
@@ -453,6 +453,7 @@ UScriptStruct* FNiagaraTypeDefinition::ColorStruct;
 UScriptStruct* FNiagaraTypeDefinition::QuatStruct;
 
 UClass* FNiagaraTypeDefinition::UObjectClass;
+UClass* FNiagaraTypeDefinition::UMaterialClass;
 
 UEnum* FNiagaraTypeDefinition::ExecutionStateEnum;
 UEnum* FNiagaraTypeDefinition::SimulationTargetEnum;
@@ -473,6 +474,7 @@ FNiagaraTypeDefinition FNiagaraTypeDefinition::ColorDef;
 FNiagaraTypeDefinition FNiagaraTypeDefinition::QuatDef;
 
 FNiagaraTypeDefinition FNiagaraTypeDefinition::UObjectDef;
+FNiagaraTypeDefinition FNiagaraTypeDefinition::UMaterialDef;
 
 TSet<UScriptStruct*> FNiagaraTypeDefinition::NumericStructs;
 TArray<FNiagaraTypeDefinition> FNiagaraTypeDefinition::OrderedNumericTypes;
@@ -517,6 +519,7 @@ void FNiagaraTypeDefinition::Init()
 	FNiagaraTypeDefinition::QuatStruct = FindObjectChecked<UScriptStruct>(CoreUObjectPkg, TEXT("Quat"));
 
 	FNiagaraTypeDefinition::UObjectClass = UObject::StaticClass();
+	FNiagaraTypeDefinition::UMaterialClass = UMaterialInterface::StaticClass();
 	
 	ParameterMapDef = FNiagaraTypeDefinition(ParameterMapStruct);
 	IDDef = FNiagaraTypeDefinition(IDStruct);
@@ -532,6 +535,7 @@ void FNiagaraTypeDefinition::Init()
 	Matrix4Def = FNiagaraTypeDefinition(Matrix4Struct);
 
 	UObjectDef = FNiagaraTypeDefinition(UObjectClass);
+	UMaterialDef = FNiagaraTypeDefinition(UMaterialClass);
 
 	CollisionEventDef = FNiagaraTypeDefinition(FNiagaraCollisionEventPayload::StaticStruct());
 	NumericStructs.Add(NumericStruct);
@@ -618,6 +622,7 @@ void FNiagaraTypeDefinition::RecreateUserDefinedTypeRegistry()
 	FNiagaraTypeRegistry::Register(FNiagaraTypeDefinition(SpawnInfoStruct), true, false, false);
 
 	FNiagaraTypeRegistry::Register(UObjectDef, true, false, false);
+	FNiagaraTypeRegistry::Register(UMaterialDef, true, false, false);
 
 	const UNiagaraSettings* Settings = GetDefault<UNiagaraSettings>();
 	check(Settings);
diff --git a/Engine/Plugins/FX/Niagara/Source/Niagara/Private/NiagaraRenderer.cpp b/Engine/Plugins/FX/Niagara/Source/Niagara/Private/NiagaraRenderer.cpp
--- a/Engine/Plugins/FX/Niagara/Source/Niagara/Private/NiagaraRenderer.cpp
+++ b/Engine/Plugins/FX/Niagara/Source/Niagara/Private/NiagaraRenderer.cpp
@@ -197,7 +197,7 @@ FNiagaraRenderer::FNiagaraRenderer(ERHIFeatureLevel::Type FeatureLevel, const UN
 void FNiagaraRenderer::Initialize(ERHIFeatureLevel::Type FeatureLevel, const UNiagaraRendererProperties *InProps, const FNiagaraEmitterInstance* Emitter)
 {
 	//Get our list of valid base materials. Fall back to default material if they're not valid.
-	InProps->GetUsedMaterials(BaseMaterials_GT);
+	InProps->GetUsedMaterials(Emitter, BaseMaterials_GT);
 	for (UMaterialInterface*& Mat : BaseMaterials_GT)
 	{
 		if (!IsMaterialValid(Mat))
diff --git a/Engine/Plugins/FX/Niagara/Source/Niagara/Private/NiagaraRibbonRendererProperties.cpp b/Engine/Plugins/FX/Niagara/Source/Niagara/Private/NiagaraRibbonRendererProperties.cpp
--- a/Engine/Plugins/FX/Niagara/Source/Niagara/Private/NiagaraRibbonRendererProperties.cpp
+++ b/Engine/Plugins/FX/Niagara/Source/Niagara/Private/NiagaraRibbonRendererProperties.cpp
@@ -35,7 +35,7 @@ FNiagaraBoundsCalculator* UNiagaraRibbonRendererProperties::CreateBoundsCalculat
 	return new FNiagaraBoundsCalculatorHelper<false, false, true>();
 }
 
-void UNiagaraRibbonRendererProperties::GetUsedMaterials(TArray<UMaterialInterface*>& OutMaterials) const
+void UNiagaraRibbonRendererProperties::GetUsedMaterials(const FNiagaraEmitterInstance* InEmitter, TArray<UMaterialInterface*>& OutMaterials) const
 {
 	OutMaterials.Add(Material);
 }
diff --git a/Engine/Plugins/FX/Niagara/Source/Niagara/Private/NiagaraSpriteRendererProperties.cpp b/Engine/Plugins/FX/Niagara/Source/Niagara/Private/NiagaraSpriteRendererProperties.cpp
--- a/Engine/Plugins/FX/Niagara/Source/Niagara/Private/NiagaraSpriteRendererProperties.cpp
+++ b/Engine/Plugins/FX/Niagara/Source/Niagara/Private/NiagaraSpriteRendererProperties.cpp
@@ -58,7 +58,7 @@ FNiagaraBoundsCalculator* UNiagaraSpriteRendererProperties::CreateBoundsCalculat
 	return new FNiagaraBoundsCalculatorHelper<true, false, false>();
 }
 
-void UNiagaraSpriteRendererProperties::GetUsedMaterials(TArray<UMaterialInterface*>& OutMaterials) const
+void UNiagaraSpriteRendererProperties::GetUsedMaterials(const FNiagaraEmitterInstance* InEmitter, TArray<UMaterialInterface*>& OutMaterials) const
 {
 	OutMaterials.Add(Material);
 }
diff --git a/Engine/Plugins/FX/Niagara/Source/Niagara/Public/NiagaraComponent.h b/Engine/Plugins/FX/Niagara/Source/Niagara/Public/NiagaraComponent.h
--- a/Engine/Plugins/FX/Niagara/Source/Niagara/Public/NiagaraComponent.h
+++ b/Engine/Plugins/FX/Niagara/Source/Niagara/Public/NiagaraComponent.h
@@ -292,6 +292,9 @@ public:
 	UFUNCTION(BlueprintCallable, Category = Niagara, meta = (DisplayName = "Set Niagara Variable (Object)"))
 	void SetVariableObject(FName InVariableName, UObject* Object);
 
+	UFUNCTION(BlueprintCallable, Category = Niagara, meta = (DisplayName = "Set Niagara Variable (Material)"))
+	void SetVariableMaterial(FName InVariableName, UMaterialInterface* Object);
+
 	/** Debug accessors for getting positions in blueprints. */
 	UFUNCTION(BlueprintCallable, Category = Niagara, meta = (DisplayName = "Get Niagara Emitter Positions"))
 	TArray<FVector> GetNiagaraParticlePositions_DebugOnly(const FString& InEmitterName);
@@ -341,6 +344,8 @@ public:
 #if WITH_EDITOR
 	virtual void PreEditChange(UProperty* PropertyAboutToChange) override;
 	virtual void PostEditChangeProperty(FPropertyChangedEvent& PropertyChangedEvent) override;
+	void OverrideUObjectParameter(const FNiagaraVariable& InVar, UObject* InObj);
+
 #endif
 	//~ End UObject Interface
 
diff --git a/Engine/Plugins/FX/Niagara/Source/Niagara/Public/NiagaraLightRendererProperties.h b/Engine/Plugins/FX/Niagara/Source/Niagara/Public/NiagaraLightRendererProperties.h
--- a/Engine/Plugins/FX/Niagara/Source/Niagara/Public/NiagaraLightRendererProperties.h
+++ b/Engine/Plugins/FX/Niagara/Source/Niagara/Public/NiagaraLightRendererProperties.h
@@ -27,7 +27,7 @@ public:
 	//~ UNiagaraRendererProperties interface
 	virtual FNiagaraRenderer* CreateEmitterRenderer(ERHIFeatureLevel::Type FeatureLevel, const FNiagaraEmitterInstance* Emitter) override;
 	virtual class FNiagaraBoundsCalculator* CreateBoundsCalculator() override { return nullptr; }
-	virtual void GetUsedMaterials(TArray<UMaterialInterface*>& OutMaterials) const override;
+	virtual void GetUsedMaterials(const FNiagaraEmitterInstance* InEmitter, TArray<UMaterialInterface*>& OutMaterials) const override;
 	virtual bool IsSimTargetSupported(ENiagaraSimTarget InSimTarget) const override { return (InSimTarget == ENiagaraSimTarget::CPUSim); };
 #if WITH_EDITORONLY_DATA
 	virtual bool IsMaterialValidForRenderer(UMaterial* Material, FText& InvalidMessage) override;
diff --git a/Engine/Plugins/FX/Niagara/Source/Niagara/Public/NiagaraMeshRendererProperties.h b/Engine/Plugins/FX/Niagara/Source/Niagara/Public/NiagaraMeshRendererProperties.h
--- a/Engine/Plugins/FX/Niagara/Source/Niagara/Public/NiagaraMeshRendererProperties.h
+++ b/Engine/Plugins/FX/Niagara/Source/Niagara/Public/NiagaraMeshRendererProperties.h
@@ -25,6 +25,34 @@ enum class ENiagaraMeshFacingMode : uint8
 	CameraPlane
 };
 
+USTRUCT()
+struct FNiagaraMeshMaterialOverride 
+{
+	GENERATED_USTRUCT_BODY()
+public:	
+	FNiagaraMeshMaterialOverride();
+	
+	/** Used to upgrade a serialized FNiagaraParameterStore property to our own struct */
+	bool SerializeFromMismatchedTag(const struct FPropertyTag& Tag, FStructuredArchive::FSlot Slot);
+
+	/** Use this UMaterialInterface if set to a valid value. This will be subordinate to UserParamBinding if it is set to a valid user variable.*/
+	UPROPERTY(EditAnywhere, Category = "Mesh Rendering")
+	UMaterialInterface* ExplicitMat;
+
+	/** Use the UMaterialInterface bound to this user variable if it is set to a valid value. If this is bound to a valid value and ExplicitMat is also set, UserParamBinding wins.*/
+	UPROPERTY(EditAnywhere, Category = "Mesh Rendering")
+	FNiagaraUserParameterBinding UserParamBinding;
+};
+
+template<>
+struct TStructOpsTypeTraits<FNiagaraMeshMaterialOverride> : public TStructOpsTypeTraitsBase2<FNiagaraMeshMaterialOverride>
+{
+	enum
+	{
+		WithStructuredSerializeFromMismatchedTag = true,
+	};
+};
+
 UCLASS(editinlinenew, meta = (DisplayName = "Mesh Renderer"))
 class UNiagaraMeshRendererProperties : public UNiagaraRendererProperties
 {
@@ -45,8 +73,9 @@ public:
 	//~ UNiagaraRendererProperties interface
 	virtual FNiagaraRenderer* CreateEmitterRenderer(ERHIFeatureLevel::Type FeatureLevel, const FNiagaraEmitterInstance* Emitter) override;
 	virtual class FNiagaraBoundsCalculator* CreateBoundsCalculator() override;
-	virtual void GetUsedMaterials(TArray<UMaterialInterface*>& OutMaterials) const override;
+	virtual void GetUsedMaterials(const FNiagaraEmitterInstance* InEmitter, TArray<UMaterialInterface*>& OutMaterials) const override;
 	virtual bool IsSimTargetSupported(ENiagaraSimTarget InSimTarget) const override { return true; };
+	virtual void PostLoad();
 
 #if WITH_EDITORONLY_DATA
 	virtual bool IsMaterialValidForRenderer(UMaterial* Material, FText& InvalidMessage) override;
@@ -76,7 +105,7 @@ public:
 	/** The materials to be used instead of the StaticMesh's materials. Note that each material must have the Niagara Mesh Particles flag checked. If the ParticleMesh 
 	requires more materials than exist in this array or any entry in this array is set to None, we will use the ParticleMesh's existing Material instead.*/
 	UPROPERTY(EditAnywhere, Category = "Mesh Rendering", meta = (EditCondition = "bOverrideMaterials"))
-	TArray<UMaterialInterface*> OverrideMaterials;
+	TArray<FNiagaraMeshMaterialOverride> OverrideMaterials;
 
 	/** Determines how the mesh orients itself relative to the camera.*/
 	UPROPERTY(EditAnywhere, Category = "Mesh Rendering")
@@ -130,6 +159,8 @@ public:
 	UPROPERTY(EditAnywhere, Category = "Bindings")
 	FNiagaraVariableAttributeBinding NormalizedAgeBinding;
 
+
 protected:
+	bool FindBinding(const FNiagaraUserParameterBinding& InBinding, const FNiagaraEmitterInstance* InEmitter, TArray<UMaterialInterface*>& OutMaterials);
 	void InitBindings();
 };
diff --git a/Engine/Plugins/FX/Niagara/Source/Niagara/Public/NiagaraParameterStore.h b/Engine/Plugins/FX/Niagara/Source/Niagara/Public/NiagaraParameterStore.h
--- a/Engine/Plugins/FX/Niagara/Source/Niagara/Public/NiagaraParameterStore.h
+++ b/Engine/Plugins/FX/Niagara/Source/Niagara/Public/NiagaraParameterStore.h
@@ -311,7 +311,7 @@ public:
 	{
 		int32 Offset = IndexOf(Parameter);
 		UObject* Obj = GetUObject(Offset);
-		checkSlow(!Obj || Parameter.GetType().GetClass() == Obj->GetClass());
+		checkSlow(!Obj || Obj->IsA(Parameter.GetType().GetClass()));
 		return Obj;
 	}
 
@@ -1079,7 +1079,7 @@ struct FNiagaraParameterDirectBinding<UObject*>
 			BoundVariable = DestVariable;
 			LayoutVersion = BoundStore->GetLayoutVersion();
 
-			check(BoundVariable.GetType() == FNiagaraTypeDefinition::GetUObjectDef());
+			check(BoundVariable.GetType().IsUObject());
 			UObjectOffset = BoundStore->IndexOf(DestVariable);
 			UObject* Ret = BoundStore->GetUObject(UObjectOffset);
 			return Ret;
@@ -1091,7 +1091,7 @@ struct FNiagaraParameterDirectBinding<UObject*>
 	{
 		if (UObjectOffset != INDEX_NONE)
 		{
-			checkSlow(BoundVariable.GetType() == FNiagaraTypeDefinition::GetUObjectDef());
+			checkSlow(BoundVariable.GetType().IsUObject());
 			checkfSlow(LayoutVersion == BoundStore->GetLayoutVersion(), TEXT("This binding is invalid, its bound parameter store's layout was changed since it was created"));
 
 			BoundStore->SetUObject(InValue, UObjectOffset);
@@ -1102,7 +1102,7 @@ struct FNiagaraParameterDirectBinding<UObject*>
 	{
 		if (UObjectOffset != INDEX_NONE)
 		{
-			checkSlow(BoundVariable.GetType() == FNiagaraTypeDefinition::GetUObjectDef());
+			checkSlow(BoundVariable.GetType().IsUObject());
 			checkfSlow(LayoutVersion == BoundStore->GetLayoutVersion(), TEXT("This binding is invalid, its bound parameter store's layout was changed since it was created"));
 
 			return BoundStore->GetUObject(UObjectOffset);
diff --git a/Engine/Plugins/FX/Niagara/Source/Niagara/Public/NiagaraRendererProperties.h b/Engine/Plugins/FX/Niagara/Source/Niagara/Public/NiagaraRendererProperties.h
--- a/Engine/Plugins/FX/Niagara/Source/Niagara/Public/NiagaraRendererProperties.h
+++ b/Engine/Plugins/FX/Niagara/Source/Niagara/Public/NiagaraRendererProperties.h
@@ -38,7 +38,7 @@ public:
 	//UObject Interface End
 	virtual FNiagaraRenderer* CreateEmitterRenderer(ERHIFeatureLevel::Type FeatureLevel, const FNiagaraEmitterInstance* Emitter) PURE_VIRTUAL ( UNiagaraRendererProperties::CreateEmitterRenderer, return nullptr;);
 	virtual class FNiagaraBoundsCalculator* CreateBoundsCalculator() PURE_VIRTUAL(UNiagaraRendererProperties::CreateBoundsCalculator, return nullptr;);
-	virtual void GetUsedMaterials(TArray<UMaterialInterface*>& OutMaterials) const PURE_VIRTUAL(UNiagaraRendererProperties::GetUsedMaterials,);
+	virtual void GetUsedMaterials(const FNiagaraEmitterInstance* InEmitter, TArray<UMaterialInterface*>& OutMaterials) const PURE_VIRTUAL(UNiagaraRendererProperties::GetUsedMaterials,);
 
 	virtual bool IsSimTargetSupported(ENiagaraSimTarget InSimTarget) const { return false; };
 
diff --git a/Engine/Plugins/FX/Niagara/Source/Niagara/Public/NiagaraRibbonRendererProperties.h b/Engine/Plugins/FX/Niagara/Source/Niagara/Public/NiagaraRibbonRendererProperties.h
--- a/Engine/Plugins/FX/Niagara/Source/Niagara/Public/NiagaraRibbonRendererProperties.h
+++ b/Engine/Plugins/FX/Niagara/Source/Niagara/Public/NiagaraRibbonRendererProperties.h
@@ -75,7 +75,7 @@ public:
 	//UNiagaraRendererProperties Interface
 	virtual FNiagaraRenderer* CreateEmitterRenderer(ERHIFeatureLevel::Type FeatureLevel, const FNiagaraEmitterInstance* Emitter) override;
 	virtual class FNiagaraBoundsCalculator* CreateBoundsCalculator() override;
-	virtual void GetUsedMaterials(TArray<UMaterialInterface*>& OutMaterials) const override;
+	virtual void GetUsedMaterials(const FNiagaraEmitterInstance* InEmitter, TArray<UMaterialInterface*>& OutMaterials) const override;
 	virtual bool IsSimTargetSupported(ENiagaraSimTarget InSimTarget) const override { return (InSimTarget == ENiagaraSimTarget::CPUSim); };
 #if WITH_EDITOR
 	virtual bool IsMaterialValidForRenderer(UMaterial* Material, FText& InvalidMessage) override;
diff --git a/Engine/Plugins/FX/Niagara/Source/Niagara/Public/NiagaraSpriteRendererProperties.h b/Engine/Plugins/FX/Niagara/Source/Niagara/Public/NiagaraSpriteRendererProperties.h
--- a/Engine/Plugins/FX/Niagara/Source/Niagara/Public/NiagaraSpriteRendererProperties.h
+++ b/Engine/Plugins/FX/Niagara/Source/Niagara/Public/NiagaraSpriteRendererProperties.h
@@ -61,7 +61,7 @@ public:
 	//UNiagaraRendererProperties interface
 	virtual FNiagaraRenderer* CreateEmitterRenderer(ERHIFeatureLevel::Type FeatureLevel, const FNiagaraEmitterInstance* Emitter) override;
 	virtual class FNiagaraBoundsCalculator* CreateBoundsCalculator() override;
-	virtual void GetUsedMaterials(TArray<UMaterialInterface*>& OutMaterials) const override;
+	virtual void GetUsedMaterials(const FNiagaraEmitterInstance* InEmitter, TArray<UMaterialInterface*>& OutMaterials) const override;
 	virtual bool IsSimTargetSupported(ENiagaraSimTarget InSimTarget) const override { return true; };
 #if WITH_EDITOR
 	virtual bool IsMaterialValidForRenderer(UMaterial* Material, FText& InvalidMessage) override;
diff --git a/Engine/Plugins/FX/Niagara/Source/Niagara/Public/NiagaraTypes.h b/Engine/Plugins/FX/Niagara/Source/Niagara/Public/NiagaraTypes.h
--- a/Engine/Plugins/FX/Niagara/Source/Niagara/Public/NiagaraTypes.h
+++ b/Engine/Plugins/FX/Niagara/Source/Niagara/Public/NiagaraTypes.h
@@ -564,6 +564,7 @@ public:
 	static const FNiagaraTypeDefinition& GetParameterMapDef() { return ParameterMapDef; }
 	static const FNiagaraTypeDefinition& GetIDDef() { return IDDef; }
 	static const FNiagaraTypeDefinition& GetUObjectDef() { return UObjectDef; }
+	static const FNiagaraTypeDefinition& GetUMaterialDef() { return UMaterialDef; }
 
 	static UScriptStruct* GetFloatStruct() { return FloatStruct; }
 	static UScriptStruct* GetBoolStruct() { return BoolStruct; }
@@ -617,6 +618,7 @@ private:
 	static FNiagaraTypeDefinition ParameterMapDef;
 	static FNiagaraTypeDefinition IDDef;
 	static FNiagaraTypeDefinition UObjectDef;
+	static FNiagaraTypeDefinition UMaterialDef;
 
 	static UScriptStruct* FloatStruct;
 	static UScriptStruct* BoolStruct;
@@ -630,6 +632,7 @@ private:
 	static UScriptStruct* NumericStruct;
 
 	static UClass* UObjectClass;
+	static UClass* UMaterialClass;
 
 	static UEnum* SimulationTargetEnum;
 	static UEnum* ScriptUsageEnum;
diff --git a/Engine/Plugins/FX/Niagara/Source/NiagaraEditor/Private/Customizations/NiagaraComponentDetails.cpp b/Engine/Plugins/FX/Niagara/Source/NiagaraEditor/Private/Customizations/NiagaraComponentDetails.cpp
--- a/Engine/Plugins/FX/Niagara/Source/NiagaraEditor/Private/Customizations/NiagaraComponentDetails.cpp
+++ b/Engine/Plugins/FX/Niagara/Source/NiagaraEditor/Private/Customizations/NiagaraComponentDetails.cpp
@@ -9,6 +9,7 @@
 #include "DetailLayoutBuilder.h"
 #include "IDetailPropertyRow.h"
 #include "DetailWidgetRow.h"
+#include "Materials/Material.h"
 #include "Widgets/Text/SInlineEditableTextBlock.h"
 #include "Widgets/SToolTip.h"
 #include "PropertyCustomizationHelpers.h"
@@ -87,6 +88,33 @@ public:
 		return NiagaraComponentNodeBuilder;
 	}
 
+
+	void OnAssetSelectedFromPicker(const FAssetData& InAssetData, FNiagaraVariable InVar)
+	{
+		FScopedTransaction ScopedTransaction(LOCTEXT("ChangeAsset", "Change asset"));
+
+		UObject* Asset = InAssetData.GetAsset();
+		if (Asset == nullptr || Asset->GetClass()->IsChildOf(InVar.GetType().GetClass()))
+		{
+			check(Component.IsValid());
+			Component->Modify();
+			Component->OverrideUObjectParameter(InVar, Asset);
+		}
+	}
+
+	FString GetCurrentAssetPath(FNiagaraVariable InVar) const
+	{
+		check(Component.IsValid());
+		TArray<FNiagaraVariable> Parameters;
+		FNiagaraUserRedirectionParameterStore& ParamStore = Component->GetOverrideParameters();
+		UObject* Obj = ParamStore.GetUObject(InVar);
+		if (Obj)
+		{
+			return Obj->GetPathName();
+		}
+		return FString();
+	}
+
 	virtual void GenerateChildContent(IDetailChildrenBuilder& ChildrenBuilder) override
 	{
 		check(Component.IsValid());
@@ -153,11 +181,27 @@ public:
 // 					FOnGetActorFilters::CreateRaw(this, &FNiagaraComponentNodeBuilder::OnGetActorFiltersForSceneOutliner),
 // 					FOnActorSelected::CreateRaw(this, &FNiagaraComponentNodeBuilder::OnActorSelected));
 
-				FString ValueName = DefaultValueObject ? FName::NameToDisplayString(DefaultValueObject->GetClass()->GetName(), false) : TEXT("null");
-				CustomValueWidget =
-					SNew(STextBlock)
-					.TextStyle(FNiagaraEditorStyle::Get(), "NiagaraEditor.ParameterText")
-					.Text(FText::FromString(ValueName));
+				if (Parameter.GetType().GetClass() == UMaterialInterface::StaticClass())
+				{
+									
+					CustomValueWidget = SNew(SObjectPropertyEntryBox)
+						.ObjectPath_Raw(this, &FNiagaraComponentNodeBuilder::GetCurrentAssetPath, Parameter )
+						.AllowedClass(UMaterialInterface::StaticClass())
+						.OnObjectChanged_Raw(this, &FNiagaraComponentNodeBuilder::OnAssetSelectedFromPicker, Parameter)
+						.AllowClear(false)
+						.DisplayUseSelected(true)
+						.DisplayBrowse(true)
+						.DisplayThumbnail(true)
+						.NewAssetFactories(TArray<UFactory*>());
+				}
+				else
+				{
+					FString ValueName = DefaultValueObject ? FName::NameToDisplayString(DefaultValueObject->GetClass()->GetName(), false) : TEXT("null");
+					CustomValueWidget =
+						SNew(STextBlock)
+						.TextStyle(FNiagaraEditorStyle::Get(), "NiagaraEditor.ParameterText")
+						.Text(FText::FromString(ValueName));
+				}
 			}
 			else
 			{
diff --git a/Engine/Plugins/FX/Niagara/Source/NiagaraEditor/Private/Customizations/NiagaraTypeCustomizations.cpp b/Engine/Plugins/FX/Niagara/Source/NiagaraEditor/Private/Customizations/NiagaraTypeCustomizations.cpp
--- a/Engine/Plugins/FX/Niagara/Source/NiagaraEditor/Private/Customizations/NiagaraTypeCustomizations.cpp
+++ b/Engine/Plugins/FX/Niagara/Source/NiagaraEditor/Private/Customizations/NiagaraTypeCustomizations.cpp
@@ -417,7 +417,7 @@ void FNiagaraUserParameterBindingCustomization::ChangeSource(FName InVarName)
 
 	PropertyHandle->NotifyPreChange();
 	TargetUserParameterBinding->Parameter.SetName(InVarName);
-	TargetUserParameterBinding->Parameter.SetType(FNiagaraTypeDefinition::GetUObjectDef());
+	//TargetUserParameterBinding->Parameter.SetType(FNiagaraTypeDefinition::GetUObjectDef()); Do not override the type here!
 	//TargetVariableBinding->DataSetVariable = FNiagaraConstants::GetAttributeAsDataSetKey(TargetVariableBinding->BoundVariable);
 	PropertyHandle->NotifyPostChange();
 	PropertyHandle->NotifyFinishedChangingProperties();
diff --git a/Engine/Plugins/FX/Niagara/Source/NiagaraEditor/Private/ViewModels/Stack/NiagaraStackFunctionInput.cpp b/Engine/Plugins/FX/Niagara/Source/NiagaraEditor/Private/ViewModels/Stack/NiagaraStackFunctionInput.cpp
--- a/Engine/Plugins/FX/Niagara/Source/NiagaraEditor/Private/ViewModels/Stack/NiagaraStackFunctionInput.cpp
+++ b/Engine/Plugins/FX/Niagara/Source/NiagaraEditor/Private/ViewModels/Stack/NiagaraStackFunctionInput.cpp
@@ -435,9 +435,9 @@ void UNiagaraStackFunctionInput::RefreshChildrenInternal(const TArray<UNiagaraSt
 class UNiagaraStackFunctionInputUtilities
 {
 public:
-	static bool GetMaterialExpressionDynamicParameter(UNiagaraEmitter* InEmitter, TArray<UMaterialExpressionDynamicParameter*>& OutDynamicParameterExpressions)
+	static bool GetMaterialExpressionDynamicParameter(FNiagaraEmitterInstance* InEmitterInstance, TArray<UMaterialExpressionDynamicParameter*>& OutDynamicParameterExpressions)
 	{
-		TArray<UMaterial*> Materials = GetMaterialFromEmitter(InEmitter);
+		TArray<UMaterial*> Materials = GetMaterialFromEmitter(InEmitterInstance);
 		
 		OutDynamicParameterExpressions.Reset();
 
@@ -462,15 +462,16 @@ public:
 		return OutDynamicParameterExpressions.Num() > 0;
 	}
 
-	static TArray<UMaterial*> GetMaterialFromEmitter(UNiagaraEmitter* InEmitter)
+	static TArray<UMaterial*> GetMaterialFromEmitter(FNiagaraEmitterInstance* InEmitterInstance)
 	{
+		UNiagaraEmitter* InEmitter = InEmitterInstance->GetCachedEmitter();
 		TArray<UMaterial*> ResultMaterials;
 		if (InEmitter->GetRenderers().Num() > 0)
 		{
 			for (UNiagaraRendererProperties* RenderProperties : InEmitter->GetRenderers())
 			{
 				TArray<UMaterialInterface*> UsedMaterialInteraces;
-				RenderProperties->GetUsedMaterials(UsedMaterialInteraces);
+				RenderProperties->GetUsedMaterials(InEmitterInstance, UsedMaterialInteraces);
 				for (UMaterialInterface* UsedMaterialInterface : UsedMaterialInteraces)
 				{
 					if (UsedMaterialInterface != nullptr)
@@ -495,7 +496,8 @@ FString UNiagaraStackFunctionInput::ResolveDisplayNameArgument(const FString& In
 	{
 		TSharedPtr<FNiagaraEmitterViewModel> ThisEmitterViewModel = GetEmitterViewModel();
 		TArray<UMaterialExpressionDynamicParameter*> ExpressionParams;
-		if (ThisEmitterViewModel.IsValid() == false || !UNiagaraStackFunctionInputUtilities::GetMaterialExpressionDynamicParameter(ThisEmitterViewModel->GetEmitter(), ExpressionParams))
+		FNiagaraEmitterInstance* Instance = ThisEmitterViewModel->GetSimulation().Pin().Get();
+		if (ThisEmitterViewModel.IsValid() == false || !UNiagaraStackFunctionInputUtilities::GetMaterialExpressionDynamicParameter(Instance, ExpressionParams))
 		{
 			return InArg.Replace(TEXT("MaterialDynamic"), TEXT("")) + TEXT(" (No material found using dynamic params)");
 		}
diff --git a/Engine/Plugins/FX/Niagara/Source/NiagaraEditor/Private/ViewModels/Stack/NiagaraStackGraphUtilities.cpp b/Engine/Plugins/FX/Niagara/Source/NiagaraEditor/Private/ViewModels/Stack/NiagaraStackGraphUtilities.cpp
--- a/Engine/Plugins/FX/Niagara/Source/NiagaraEditor/Private/ViewModels/Stack/NiagaraStackGraphUtilities.cpp
+++ b/Engine/Plugins/FX/Niagara/Source/NiagaraEditor/Private/ViewModels/Stack/NiagaraStackGraphUtilities.cpp
@@ -1298,7 +1298,7 @@ void FNiagaraStackGraphUtilities::GetNewParameterAvailableTypes(TArray<FNiagaraT
 	for (const FNiagaraTypeDefinition& RegisteredParameterType : FNiagaraTypeRegistry::GetRegisteredParameterTypes())
 	{
 		//Object types only allowed in user namespace at the moment.
-		if (RegisteredParameterType == FNiagaraTypeDefinition::GetUObjectDef() && Namespace != FNiagaraParameterHandle::UserNamespace)
+		if (RegisteredParameterType.IsUObject() && Namespace != FNiagaraParameterHandle::UserNamespace)
 		{
 			continue;
 		}
diff --git a/Engine/Plugins/FX/Niagara/Source/NiagaraEditor/Private/ViewModels/Stack/NiagaraStackParameterStoreEntry.cpp b/Engine/Plugins/FX/Niagara/Source/NiagaraEditor/Private/ViewModels/Stack/NiagaraStackParameterStoreEntry.cpp
--- a/Engine/Plugins/FX/Niagara/Source/NiagaraEditor/Private/ViewModels/Stack/NiagaraStackParameterStoreEntry.cpp
+++ b/Engine/Plugins/FX/Niagara/Source/NiagaraEditor/Private/ViewModels/Stack/NiagaraStackParameterStoreEntry.cpp
@@ -159,6 +159,10 @@ void UNiagaraStackParameterStoreEntry::Reset()
 			UNiagaraDataInterface* DefaultObject = NewObject<UNiagaraDataInterface>(this, const_cast<UClass*>(InputType.GetClass()));
 			DefaultObject->CopyTo(ParameterStore->GetDataInterface(Var));
 		}
+		else if (Var.IsUObject())
+		{
+			ParameterStore->SetUObject(nullptr, Var);
+		}
 	}
 	RefreshValueAndHandle();
 	RefreshChildren();
@@ -277,6 +281,22 @@ void UNiagaraStackParameterStoreEntry::RenameInput(FString NewName)
 	}
 }
 
+void UNiagaraStackParameterStoreEntry::ReplaceValueObject(UObject* Obj)
+{
+	NotifyBeginValueChange();
+	FNiagaraVariable Var(InputType, ParameterName);
+	if (Obj == nullptr || Obj->GetClass()->IsChildOf(GetInputType().GetClass()))
+	{
+		ParameterStore->SetUObject(Obj, Var);
+
+		RefreshValueAndHandle();
+		RefreshChildren();
+	}
+	NotifyEndValueChange();
+	GetSystemViewModel()->ResetSystem();
+}
+
+
 void UNiagaraStackParameterStoreEntry::Delete()
 {
 	FScopedTransaction ScopedTransaction(LOCTEXT("RemoveUserParameter", "Remove user parameter"));
diff --git a/Engine/Plugins/FX/Niagara/Source/NiagaraEditor/Public/ViewModels/Stack/NiagaraStackParameterStoreEntry.h b/Engine/Plugins/FX/Niagara/Source/NiagaraEditor/Public/ViewModels/Stack/NiagaraStackParameterStoreEntry.h
--- a/Engine/Plugins/FX/Niagara/Source/NiagaraEditor/Public/ViewModels/Stack/NiagaraStackParameterStoreEntry.h
+++ b/Engine/Plugins/FX/Niagara/Source/NiagaraEditor/Public/ViewModels/Stack/NiagaraStackParameterStoreEntry.h
@@ -96,6 +96,9 @@ public:
 	/** Delete the parameter from the ParameterStore and notify that the store changed. */
 	void Delete();
 
+	/** Use an external asset instead of the local value object.*/
+	void ReplaceValueObject(UObject* Obj);
+
 protected:
 	//~ UNiagaraStackEntry interface
 	virtual void RefreshChildrenInternal(const TArray<UNiagaraStackEntry*>& CurrentChildren, TArray<UNiagaraStackEntry*>& NewChildren, TArray<FStackIssue>& NewIssues) override;
diff --git a/Engine/Plugins/FX/Niagara/Source/NiagaraEditorWidgets/Private/Stack/SNiagaraStackParameterStoreEntryValue.cpp b/Engine/Plugins/FX/Niagara/Source/NiagaraEditorWidgets/Private/Stack/SNiagaraStackParameterStoreEntryValue.cpp
--- a/Engine/Plugins/FX/Niagara/Source/NiagaraEditorWidgets/Private/Stack/SNiagaraStackParameterStoreEntryValue.cpp
+++ b/Engine/Plugins/FX/Niagara/Source/NiagaraEditorWidgets/Private/Stack/SNiagaraStackParameterStoreEntryValue.cpp
@@ -17,6 +17,8 @@
 #include "PropertyEditorModule.h"
 #include "IStructureDetailsView.h"
 #include "Modules/ModuleManager.h"
+#include "PropertyEditorModule.h"
+#include "PropertyCustomizationHelpers.h"
 
 #define LOCTEXT_NAMESPACE "NiagaraStackParameterStoreEntryValue"
 
@@ -118,6 +120,19 @@ FReply SNiagaraStackParameterStoreEntryValue::DeleteClicked()
 	return FReply::Handled();
 }
 
+void SNiagaraStackParameterStoreEntryValue::OnAssetSelectedFromPicker(const FAssetData& InAssetData)
+{	
+	UMaterialInterface* Mat = Cast<UMaterialInterface>(InAssetData.GetAsset());
+	StackEntry->ReplaceValueObject(Mat);
+}
+
+FString SNiagaraStackParameterStoreEntryValue::GetCurrentAssetPath() const
+{
+	UObject* Obj = StackEntry->GetValueObject();
+	return  Obj != nullptr ? Obj->GetPathName() : FString();
+}
+
+
 TSharedRef<SWidget> SNiagaraStackParameterStoreEntryValue::ConstructValueStructWidget()
 {
 	ValueStructParameterEditor.Reset();
@@ -164,8 +179,27 @@ TSharedRef<SWidget> SNiagaraStackParameterStoreEntryValue::ConstructValueStructW
 	}
 	else if (StackEntry->GetInputType().GetClass() != nullptr)
 	{
-		return SNew(STextBlock)
-			.Text(StackEntry->GetInputType().GetClass()->GetDisplayNameText());
+		if (StackEntry->GetInputType().GetClass()->IsChildOf(UMaterialInterface::StaticClass()))
+		{
+			TArray<const UClass*> AllowedClasses;
+			AllowedClasses.Add(UMaterialInterface::StaticClass());
+			
+			return SNew(SObjectPropertyEntryBox)
+				.ObjectPath_Raw(this, &SNiagaraStackParameterStoreEntryValue::GetCurrentAssetPath)
+				.AllowedClass(UMaterialInterface::StaticClass())
+				.OnObjectChanged_Raw(this, &SNiagaraStackParameterStoreEntryValue::OnAssetSelectedFromPicker)
+				.AllowClear(false)
+				.DisplayUseSelected(true)
+				.DisplayBrowse(true)
+				.DisplayThumbnail(true)
+				.NewAssetFactories(TArray<UFactory*>());
+
+		}
+		else
+		{
+			return SNew(STextBlock)
+				.Text(StackEntry->GetInputType().GetClass()->GetDisplayNameText());
+		}
 	}
 	else
 	{
diff --git a/Engine/Plugins/FX/Niagara/Source/NiagaraEditorWidgets/Private/Stack/SNiagaraStackParameterStoreEntryValue.h b/Engine/Plugins/FX/Niagara/Source/NiagaraEditorWidgets/Private/Stack/SNiagaraStackParameterStoreEntryValue.h
--- a/Engine/Plugins/FX/Niagara/Source/NiagaraEditorWidgets/Private/Stack/SNiagaraStackParameterStoreEntryValue.h
+++ b/Engine/Plugins/FX/Niagara/Source/NiagaraEditorWidgets/Private/Stack/SNiagaraStackParameterStoreEntryValue.h
@@ -50,6 +50,9 @@ private:
 
 	FSlateColor GetInputIconColor() const;
 
+	void OnAssetSelectedFromPicker(const FAssetData& InAssetData);
+	FString GetCurrentAssetPath() const;
+
 private:
 	UNiagaraStackParameterStoreEntry* StackEntry;
 
