#UE4 Fix uninitialized variable

diff --git a/Engine/Plugins/FX/Niagara/Source/Niagara/Private/NiagaraMeshRendererProperties.cpp b/Engine/Plugins/FX/Niagara/Source/Niagara/Private/NiagaraMeshRendererProperties.cpp
--- a/Engine/Plugins/FX/Niagara/Source/Niagara/Private/NiagaraMeshRendererProperties.cpp
+++ b/Engine/Plugins/FX/Niagara/Source/Niagara/Private/NiagaraMeshRendererProperties.cpp
@@ -7,6 +7,7 @@
 #include "NiagaraBoundsCalculatorHelper.h"
 
 FNiagaraMeshMaterialOverride::FNiagaraMeshMaterialOverride()
+	: ExplicitMat(nullptr)
 {
 	FNiagaraTypeDefinition MaterialDef(UMaterialInterface::StaticClass());
 	UserParamBinding.Parameter.SetType(MaterialDef);
