Fix niagara components having collision on

diff --git a/Engine/Plugins/FX/Niagara/Source/Niagara/Private/NiagaraComponent.cpp b/Engine/Plugins/FX/Niagara/Source/Niagara/Private/NiagaraComponent.cpp
--- a/Engine/Plugins/FX/Niagara/Source/Niagara/Private/NiagaraComponent.cpp
+++ b/Engine/Plugins/FX/Niagara/Source/Niagara/Private/NiagaraComponent.cpp
@@ -20,6 +20,7 @@
 #include "NiagaraWorldManager.h"
 #include "EngineUtils.h"
 #include "ProfilingDebugging/CsvProfiler.h"
+#include "Engine/CollisionProfile.h"
 #include "NiagaraEmitterInstanceBatcher.h"
 
 DECLARE_CYCLE_STAT(TEXT("Sceneproxy create (GT)"), STAT_NiagaraCreateSceneProxy, STATGROUP_Niagara);
@@ -371,6 +372,7 @@ UNiagaraComponent::UNiagaraComponent(const FObjectInitializer& ObjectInitializer
 	SavedAutoAttachRelativeScale3D = FVector(1.f, 1.f, 1.f);
 
 	SetGenerateOverlapEvents(false);
+	SetCollisionProfileName(UCollisionProfile::NoCollision_ProfileName);
 }
 
 /********* UFXSystemComponent *********/
