#ue4 - UNiagaraComponent::CalcBounds() should use the auto attach parent's bounds when not attached, otherwise the actor bounds are huge when detached and sitting near the origin. Matches a change I made to UParticleSystemComponent previously.

diff --git a/Engine/Plugins/FX/Niagara/Source/Niagara/Private/NiagaraComponent.cpp b/Engine/Plugins/FX/Niagara/Source/Niagara/Private/NiagaraComponent.cpp
--- a/Engine/Plugins/FX/Niagara/Source/Niagara/Private/NiagaraComponent.cpp
+++ b/Engine/Plugins/FX/Niagara/Source/Niagara/Private/NiagaraComponent.cpp
@@ -1086,6 +1086,13 @@ int32 UNiagaraComponent::GetNumMaterials() const
 
 FBoxSphereBounds UNiagaraComponent::CalcBounds(const FTransform& LocalToWorld) const
 {
+	const USceneComponent* UseAutoParent = (bAutoManageAttachment && GetAttachParent() == nullptr) ? AutoAttachParent.Get() : nullptr;
+	if (UseAutoParent)
+	{
+		// We use auto attachment but have detached, don't use our own bogus bounds (we're off near 0,0,0), use the usual parent's bounds.
+		return UseAutoParent->Bounds;
+	}
+
 	FBoxSphereBounds SystemBounds;
 	if (SystemInstance.IsValid())
 	{
