Null check when looking in cached emitter for dynamic material.

diff --git a/Engine/Plugins/FX/Niagara/Source/NiagaraEditor/Private/ViewModels/Stack/NiagaraStackFunctionInput.cpp b/Engine/Plugins/FX/Niagara/Source/NiagaraEditor/Private/ViewModels/Stack/NiagaraStackFunctionInput.cpp
--- a/Engine/Plugins/FX/Niagara/Source/NiagaraEditor/Private/ViewModels/Stack/NiagaraStackFunctionInput.cpp
+++ b/Engine/Plugins/FX/Niagara/Source/NiagaraEditor/Private/ViewModels/Stack/NiagaraStackFunctionInput.cpp
@@ -464,23 +464,26 @@ public:
 
 	static TArray<UMaterial*> GetMaterialFromEmitter(FNiagaraEmitterInstance* InEmitterInstance)
 	{
-		UNiagaraEmitter* InEmitter = InEmitterInstance->GetCachedEmitter();
 		TArray<UMaterial*> ResultMaterials;
-		if (InEmitter->GetRenderers().Num() > 0)
+		UNiagaraEmitter* InEmitter = InEmitterInstance->GetCachedEmitter();
+		ensureMsgf(InEmitter != nullptr, TEXT("Tried to get null Emitter from EmitterInstance!"));
 		{
-			for (UNiagaraRendererProperties* RenderProperties : InEmitter->GetRenderers())
+			if (InEmitter->GetRenderers().Num() > 0)
 			{
-				TArray<UMaterialInterface*> UsedMaterialInteraces;
-				RenderProperties->GetUsedMaterials(InEmitterInstance, UsedMaterialInteraces);
-				for (UMaterialInterface* UsedMaterialInterface : UsedMaterialInteraces)
+				for (UNiagaraRendererProperties* RenderProperties : InEmitter->GetRenderers())
 				{
-					if (UsedMaterialInterface != nullptr)
+					TArray<UMaterialInterface*> UsedMaterialInteraces;
+					RenderProperties->GetUsedMaterials(InEmitterInstance, UsedMaterialInteraces);
+					for (UMaterialInterface* UsedMaterialInterface : UsedMaterialInteraces)
 					{
-						UMaterial* UsedMaterial = UsedMaterialInterface->GetBaseMaterial();
-						if (UsedMaterial != nullptr)
+						if (UsedMaterialInterface != nullptr)
 						{
-							ResultMaterials.AddUnique(UsedMaterial);
-							break;
+							UMaterial* UsedMaterial = UsedMaterialInterface->GetBaseMaterial();
+							if (UsedMaterial != nullptr)
+							{
+								ResultMaterials.AddUnique(UsedMaterial);
+								break;
+							}
 						}
 					}
 				}
