Attempt to fix crash in GetMaterialFromEmitter

diff --git a/Engine/Plugins/FX/Niagara/Source/NiagaraEditor/Private/ViewModels/Stack/NiagaraStackFunctionInput.cpp b/Engine/Plugins/FX/Niagara/Source/NiagaraEditor/Private/ViewModels/Stack/NiagaraStackFunctionInput.cpp
--- a/Engine/Plugins/FX/Niagara/Source/NiagaraEditor/Private/ViewModels/Stack/NiagaraStackFunctionInput.cpp
+++ b/Engine/Plugins/FX/Niagara/Source/NiagaraEditor/Private/ViewModels/Stack/NiagaraStackFunctionInput.cpp
@@ -465,10 +465,10 @@ public:
 	static TArray<UMaterial*> GetMaterialFromEmitter(FNiagaraEmitterInstance* InEmitterInstance)
 	{
 		TArray<UMaterial*> ResultMaterials;
-		UNiagaraEmitter* InEmitter = InEmitterInstance->GetCachedEmitter();
-		ensureMsgf(InEmitter != nullptr, TEXT("Tried to get null Emitter from EmitterInstance!"));
+		if (InEmitterInstance)
 		{
-			if (InEmitter->GetRenderers().Num() > 0)
+			UNiagaraEmitter* InEmitter = InEmitterInstance->GetCachedEmitter();
+			if (InEmitter && InEmitter->GetRenderers().Num() > 0)
 			{
 				for (UNiagaraRendererProperties* RenderProperties : InEmitter->GetRenderers())
 				{
