Turning a check into a warning for vertexcount = 0 or trianglecount = 0

diff --git a/Engine/Plugins/FX/Niagara/Source/Niagara/Private/NiagaraDataInterfaceSkeletalMesh.cpp b/Engine/Plugins/FX/Niagara/Source/Niagara/Private/NiagaraDataInterfaceSkeletalMesh.cpp
--- a/Engine/Plugins/FX/Niagara/Source/Niagara/Private/NiagaraDataInterfaceSkeletalMesh.cpp
+++ b/Engine/Plugins/FX/Niagara/Source/Niagara/Private/NiagaraDataInterfaceSkeletalMesh.cpp
@@ -320,8 +320,15 @@ void FSkeletalMeshGpuSpawnStaticBuffers::Initialise(FNDISkeletalMesh_InstanceDat
 	LODRenderData = &SkeletalMeshLODRenderData;
 	TriangleCount = SkeletalMeshLODRenderData.MultiSizeIndexContainer.GetIndexBuffer()->Num() / 3;
 	VertexCount = SkeletalMeshLODRenderData.GetNumVertices();
-	check(TriangleCount > 0);
-	check(VertexCount > 0);
+	
+	if (TriangleCount == 0)
+	{
+		UE_LOG(LogNiagara, Warning, TEXT("FSkeletalMeshGpuSpawnStaticBuffers> Triangle count is invalid %d"), TriangleCount, (InstData && InstData->Mesh) ? *InstData->Mesh->GetFullName() : TEXT("Unknown Mesh"));
+	}
+	if (VertexCount == 0)
+	{
+		UE_LOG(LogNiagara, Warning, TEXT("FSkeletalMeshGpuSpawnStaticBuffers> Vertex count is invalid %d"), VertexCount, (InstData && InstData->Mesh) ? *InstData->Mesh->GetFullName() : TEXT("Unknown Mesh"));
+	}
 
 	// Copy Specific Bones / Socket data into arrays that the renderer will use to create read buffers
 	//-TODO: Exclude setting up these arrays if we don't sample from them
