Fixed a possible crash for invalid render properties

diff --git a/Engine/Plugins/FX/Niagara/Source/Niagara/Private/NiagaraRendererLights.cpp b/Engine/Plugins/FX/Niagara/Source/Niagara/Private/NiagaraRendererLights.cpp
--- a/Engine/Plugins/FX/Niagara/Source/Niagara/Private/NiagaraRendererLights.cpp
+++ b/Engine/Plugins/FX/Niagara/Source/Niagara/Private/NiagaraRendererLights.cpp
@@ -54,6 +54,11 @@ FNiagaraDynamicDataBase* FNiagaraRendererLights::GenerateDynamicData(const FNiag
 	const UNiagaraLightRendererProperties* Properties = CastChecked<const UNiagaraLightRendererProperties>(InProperties);
 	FNiagaraDataSet& Data = Emitter->GetData();
 	FNiagaraDataBuffer& ParticleData = Data.GetCurrentDataChecked();
+	FNiagaraDynamicDataLights* DynamicData = new FNiagaraDynamicDataLights(Emitter);
+	if (!Properties)
+	{
+		return DynamicData;
+	}
 
 	//I'm not a great fan of pulling scalar components out to a structured vert buffer like this.
 	//TODO: Experiment with a new VF that reads the data directly from the scalar layout.
@@ -64,8 +69,6 @@ FNiagaraDynamicDataBase* FNiagaraRendererLights::GenerateDynamicData(const FNiag
 	FNiagaraDataSetAccessor<float> ScatteringAccessor(Data, Properties->VolumetricScatteringBinding.DataSetVariable);
 	FNiagaraDataSetAccessor<int32> EnabledAccessor(Data, Properties->LightRenderingEnabledBinding.DataSetVariable);
 
-	FNiagaraDynamicDataLights* DynamicData = new FNiagaraDynamicDataLights(Emitter);
-
 	const FMatrix& LocalToWorldMatrix = Proxy->GetLocalToWorld();
 	FVector DefaultColor = FVector(Properties->ColorBinding.DefaultValueIfNonExistent.GetValue<FLinearColor>());
 	FVector DefaultPos = FVector4(LocalToWorldMatrix.GetOrigin());
