- Niagara System Simulations can now tick in multiple tick groups and run concurrently across all tick groups until the last demotable tick group

diff --git a/Engine/Plugins/FX/Niagara/Source/Niagara/Classes/NiagaraDataInterface.h b/Engine/Plugins/FX/Niagara/Source/Niagara/Classes/NiagaraDataInterface.h
--- a/Engine/Plugins/FX/Niagara/Source/Niagara/Classes/NiagaraDataInterface.h
+++ b/Engine/Plugins/FX/Niagara/Source/Niagara/Classes/NiagaraDataInterface.h
@@ -285,6 +285,9 @@ public:
 
 	virtual bool RequiresDistanceFieldData() const { return false; }
 
+	virtual bool HasTickGroupPrereqs() const { return false; }
+	virtual ETickingGroup CalculateTickGroup(void* PerInstanceData) const { return NiagaraFirstTickGroup; }
+
 	/** Determines if this type definition matches to a known data interface type.*/
 	static bool IsDataInterfaceType(const FNiagaraTypeDefinition& TypeDef);
 
diff --git a/Engine/Plugins/FX/Niagara/Source/Niagara/Classes/NiagaraDataInterfaceSkeletalMesh.h b/Engine/Plugins/FX/Niagara/Source/Niagara/Classes/NiagaraDataInterfaceSkeletalMesh.h
--- a/Engine/Plugins/FX/Niagara/Source/Niagara/Classes/NiagaraDataInterfaceSkeletalMesh.h
+++ b/Engine/Plugins/FX/Niagara/Source/Niagara/Classes/NiagaraDataInterfaceSkeletalMesh.h
@@ -171,19 +171,19 @@ class FNDI_SkeletalMesh_GeneratedData
 {
 	// Encapsulates skinning data and mesh usage information. 
 	// Set by GetCachedSkinningData and used by TickGeneratedData to determine whether we need to pre-skin or not.
-	struct CachedSkinningDataAndUsage {
+	struct CachedSkinningDataAndUsage
+	{
+		bool bHasTicked = false;
 		TSharedPtr<FSkeletalMeshSkinningData> SkinningData;
 		FSkeletalMeshSkinningDataUsage Usage;
 	};
+
+	FRWLock CachedSkinningDataGuard;
 	TMap<TWeakObjectPtr<USkeletalMeshComponent>, CachedSkinningDataAndUsage> CachedSkinningData;
 
 public:
-
 	FSkeletalMeshSkinningDataHandle GetCachedSkinningData(TWeakObjectPtr<USkeletalMeshComponent>& InComponent, FSkeletalMeshSkinningDataUsage Usage);
-
-	void TickGeneratedData(float DeltaSeconds);
-
-	FCriticalSection CriticalSection;
+	void TickGeneratedData(ETickingGroup TickGroup, float DeltaSeconds);
 };
 
 //////////////////////////////////////////////////////////////////////////
@@ -524,6 +524,8 @@ public:
 	virtual TArray<FNiagaraDataInterfaceError> GetErrors() override;
 	virtual void ValidateFunction(const FNiagaraFunctionSignature& Function, TArray<FText>& OutValidationErrors) override;
 #endif
+	virtual bool HasTickGroupPrereqs() const override { return true; }
+	virtual ETickingGroup CalculateTickGroup(void* PerInstanceData) const override;
 	//~ UNiagaraDataInterface interface END
 
 	USkeletalMesh* GetSkeletalMesh(class UNiagaraComponent* OwningComponent, TWeakObjectPtr<USceneComponent>& SceneComponent, USkeletalMeshComponent*& FoundSkelComp, FNDISkeletalMesh_InstanceData* InstData = nullptr);
diff --git a/Engine/Plugins/FX/Niagara/Source/Niagara/Private/NiagaraComponent.cpp b/Engine/Plugins/FX/Niagara/Source/Niagara/Private/NiagaraComponent.cpp
--- a/Engine/Plugins/FX/Niagara/Source/Niagara/Private/NiagaraComponent.cpp
+++ b/Engine/Plugins/FX/Niagara/Source/Niagara/Private/NiagaraComponent.cpp
@@ -967,16 +967,6 @@ void UNiagaraComponent::BeginDestroy()
 	Super::BeginDestroy();
 }
 
-// Uncertain about this. 
-// void UNiagaraComponent::OnAttachmentChanged()
-// {
-// 	Super::OnAttachmentChanged();
-// 	if (bIsActive && !bIsChangingAutoAttachment && !GetOwner()->IsPendingKillPending())
-// 	{
-// 		ResetSystem();
-// 	}
-// }
-
 TSharedPtr<FNiagaraSystemSimulation, ESPMode::ThreadSafe> UNiagaraComponent::GetSystemSimulation()
 {
 	if (SystemInstance)
@@ -1159,6 +1149,48 @@ void UNiagaraComponent::GetUsedMaterials(TArray<UMaterialInterface*>& OutMateria
 	}
 }
 
+void UNiagaraComponent::SetComponentTickEnabled(bool bEnabled)
+{
+	Super::SetComponentTickEnabled(bEnabled);
+	if (SystemInstance.IsValid())
+	{
+		SystemInstance->UpdatePrereqs();
+	}
+}
+
+void UNiagaraComponent::OnAttachmentChanged()
+{
+	// Uncertain about this. 
+	// 	if (bIsActive && !bIsChangingAutoAttachment && !GetOwner()->IsPendingKillPending())
+	// 	{
+	// 		ResetSystem();
+	// 	}
+
+	Super::OnAttachmentChanged();
+	if ( SystemInstance.IsValid() )
+	{
+		SystemInstance->UpdatePrereqs();
+	}
+}
+
+void UNiagaraComponent::OnChildAttached(USceneComponent* ChildComponent)
+{
+	Super::OnChildAttached(ChildComponent);
+	if (SystemInstance.IsValid())
+	{
+		SystemInstance->UpdatePrereqs();
+	}
+}
+
+void UNiagaraComponent::OnChildDetached(USceneComponent* ChildComponent)
+{
+	Super::OnChildDetached(ChildComponent);
+	if (SystemInstance.IsValid())
+	{
+		SystemInstance->UpdatePrereqs();
+	}
+}
+
 FNiagaraSystemInstance* UNiagaraComponent::GetSystemInstance() const
 {
 	return SystemInstance.Get();
diff --git a/Engine/Plugins/FX/Niagara/Source/Niagara/Private/NiagaraDataInterfaceSkeletalMesh.cpp b/Engine/Plugins/FX/Niagara/Source/Niagara/Private/NiagaraDataInterfaceSkeletalMesh.cpp
--- a/Engine/Plugins/FX/Niagara/Source/Niagara/Private/NiagaraDataInterfaceSkeletalMesh.cpp
+++ b/Engine/Plugins/FX/Niagara/Source/Niagara/Private/NiagaraDataInterfaceSkeletalMesh.cpp
@@ -223,63 +223,91 @@ bool FSkeletalMeshSkinningData::Tick(float InDeltaSeconds, bool bRequirePreskin)
 
 //////////////////////////////////////////////////////////////////////////
 
-FSkeletalMeshSkinningDataHandle FNDI_SkeletalMesh_GeneratedData::GetCachedSkinningData(TWeakObjectPtr<USkeletalMeshComponent>& InComponent, FSkeletalMeshSkinningDataUsage Usage)
+FSkeletalMeshSkinningDataHandle FNDI_SkeletalMesh_GeneratedData::GetCachedSkinningData(TWeakObjectPtr<USkeletalMeshComponent>& Component, FSkeletalMeshSkinningDataUsage Usage)
 {
-	FScopeLock Lock(&CriticalSection);
-	
-	USkeletalMeshComponent* Component = InComponent.Get();
-	check(Component);
-	TSharedPtr<FSkeletalMeshSkinningData> SkinningData = nullptr;
+	check(Component.Get() != nullptr);
 
-	if (CachedSkinningDataAndUsage* Existing = CachedSkinningData.Find(Component))
-	{
-		check(Existing && Existing->SkinningData.IsValid());//We shouldn't be able to have an invalid ptr here.
-		SkinningData = Existing->SkinningData; // Reuse the old skinning data
-		Existing->Usage = Usage; // Overwrite the old Usage
-	}
-	else
+	// Attempt to Find data
 	{
-		SkinningData = MakeShared<FSkeletalMeshSkinningData>(InComponent);
-		CachedSkinningData.Add(Component) = { SkinningData, Usage };
+		FRWScopeLock ReadLock(CachedSkinningDataGuard, SLT_ReadOnly);
+		if ( CachedSkinningDataAndUsage* Existing = CachedSkinningData.Find(Component) )
+		{
+			check(Existing->SkinningData.IsValid());
+			return FSkeletalMeshSkinningDataHandle(Existing->Usage, Existing->SkinningData);
+		}
 	}
 
-	return FSkeletalMeshSkinningDataHandle(Usage, SkinningData);
+	// We need to add
+	FRWScopeLock WriteLock(CachedSkinningDataGuard, SLT_Write);
+	CachedSkinningDataAndUsage& NewData = CachedSkinningData.FindOrAdd(Component);
+	NewData.Usage = Usage;
+	NewData.SkinningData = MakeShared<FSkeletalMeshSkinningData>(Component);
+	return FSkeletalMeshSkinningDataHandle(NewData.Usage, NewData.SkinningData);
 }
 
-void FNDI_SkeletalMesh_GeneratedData::TickGeneratedData(float DeltaSeconds)
+void FNDI_SkeletalMesh_GeneratedData::TickGeneratedData(ETickingGroup TickGroup, float DeltaSeconds)
 {
 	check(IsInGameThread());
 	SCOPE_CYCLE_COUNTER(STAT_NiagaraSkel_PreSkin);
 
-	//Tick skinning data.
-	TArray<TWeakObjectPtr<USkeletalMeshComponent>, TInlineAllocator<64>> ToRemove;
-	TArray<FSkeletalMeshSkinningData*> ToTickBonesOnly;
-	TArray<FSkeletalMeshSkinningData*> ToTickPreskin;
+	FRWScopeLock WriteLock(CachedSkinningDataGuard, SLT_Write);
+
+	// We may want to look at separating out how we manage the ticks here
+	//-OPT: Move into different arrays per tick group, manage promotions, demotions, etc, or add ourselves as a subsequent of the component's tick
+	TArray<TWeakObjectPtr<USkeletalMeshComponent>, TInlineAllocator<32>> ToRemove;
+	TArray<FSkeletalMeshSkinningData*, TInlineAllocator<32>> ToTickBonesOnly;
+	TArray<FSkeletalMeshSkinningData*, TInlineAllocator<32>> ToTickPreskin;
+	const bool bForceTick = TickGroup == NiagaraLastTickGroup;
+
 	ToTickBonesOnly.Reserve(CachedSkinningData.Num());
 	ToTickPreskin.Reserve(CachedSkinningData.Num());
+
 	for (TPair<TWeakObjectPtr<USkeletalMeshComponent>, CachedSkinningDataAndUsage>& Pair : CachedSkinningData)
 	{
-		TSharedPtr<FSkeletalMeshSkinningData>& Ptr = Pair.Value.SkinningData;
-			
-		if (Ptr.IsUnique() || !Pair.Key.Get() || !Ptr->IsUsed())
+		if ( TickGroup == NiagaraFirstTickGroup )
 		{
-			ToRemove.Add(Pair.Key);//Remove unused skin data or for those with GCd components as we go.
+			Pair.Value.bHasTicked = false;
 		}
-		else
+
+		// Should remove?
+		TSharedPtr<FSkeletalMeshSkinningData>& SkinningData = Pair.Value.SkinningData;
+		USkeletalMeshComponent* Component = Pair.Key.Get();
+		if ( (Component == nullptr) || SkinningData.IsUnique() || !SkinningData->IsUsed() )
 		{
-			FSkeletalMeshSkinningDataUsage Usage = Pair.Value.Usage;
-			FSkeletalMeshSkinningData* SkinData = Ptr.Get();
-			check(SkinData);
+			ToRemove.Add(Pair.Key);
+			continue;
+		}
 
-			if (Usage.NeedPreSkinnedVerts())
-			{
-				ToTickPreskin.Add(SkinData);
-			}
-			else
+		if ( Pair.Value.bHasTicked == true )
+		{
+			continue;
+		}
+
+		// Has ticked or can be ticked
+		if (bForceTick == false)
+		{
+			const ETickingGroup PrereqTickGroup = FMath::Max(Component->PrimaryComponentTick.TickGroup, Component->PrimaryComponentTick.EndTickGroup);
+			if ((PrereqTickGroup > TickGroup) || (Component->PrimaryComponentTick.IsCompletionHandleValid() && !Component->PrimaryComponentTick.GetCompletionHandle()->IsComplete()))
 			{
-				ToTickBonesOnly.Add(SkinData);
+				continue;
 			}
 		}
+
+		// We are going to tick this one
+		Pair.Value.bHasTicked = true;
+
+		FSkeletalMeshSkinningDataUsage Usage = Pair.Value.Usage;
+		FSkeletalMeshSkinningData* SkinningDataPtr = SkinningData.Get();
+		check(SkinningDataPtr);
+
+		if (Usage.NeedPreSkinnedVerts())
+		{
+			ToTickPreskin.Add(SkinningDataPtr);
+		}
+		else
+		{
+			ToTickBonesOnly.Add(SkinningDataPtr);
+		}
 	}
 
 	for (TWeakObjectPtr<USkeletalMeshComponent> Key : ToRemove)
@@ -289,7 +317,8 @@ void FNDI_SkeletalMesh_GeneratedData::TickGeneratedData(float DeltaSeconds)
 	
 	// First tick the meshes that don't need pre-skinning. 
 	// This prevents additional threading overhead when we don't need to pre-skin.
-	for (int i = 0; i < ToTickBonesOnly.Num(); i++) {
+	for (int i = 0; i < ToTickBonesOnly.Num(); i++)
+	{
 		ToTickBonesOnly[i]->Tick(DeltaSeconds, false);
 	}
 		
@@ -2069,6 +2098,17 @@ void UNiagaraDataInterfaceSkeletalMesh::SetSourceComponentFromBlueprints(USkelet
 	Source = ComponentToUse->GetOwner();
 }
 
+ETickingGroup UNiagaraDataInterfaceSkeletalMesh::CalculateTickGroup(void* PerInstanceData) const
+{
+	FNDISkeletalMesh_InstanceData* InstData = static_cast<FNDISkeletalMesh_InstanceData*>(PerInstanceData);
+	USkeletalMeshComponent* Component = Cast<USkeletalMeshComponent>(InstData->Component.Get());
+	if ( Component )
+	{
+		return ETickingGroup(FMath::Max(Component->PrimaryComponentTick.TickGroup, Component->PrimaryComponentTick.EndTickGroup) + 1);
+	}
+	return NiagaraFirstTickGroup;
+}
+
 //UNiagaraDataInterfaceSkeletalMesh END
 //////////////////////////////////////////////////////////////////////////
 
diff --git a/Engine/Plugins/FX/Niagara/Source/Niagara/Private/NiagaraSystemInstance.cpp b/Engine/Plugins/FX/Niagara/Source/Niagara/Private/NiagaraSystemInstance.cpp
--- a/Engine/Plugins/FX/Niagara/Source/Niagara/Private/NiagaraSystemInstance.cpp
+++ b/Engine/Plugins/FX/Niagara/Source/Niagara/Private/NiagaraSystemInstance.cpp
@@ -48,9 +48,19 @@ static FAutoConsoleVariableRef CVarLastRenderTimeSafetyBias(
 	ECVF_Default
 );
 
+static int GNiagaraForceLastTickGroup = 0;
+static FAutoConsoleVariableRef CVarNiagaraForceLastTickGroup(
+	TEXT("fx.Niagara.ForceLastTickGroup"),
+	GNiagaraForceLastTickGroup,
+	TEXT("Force Niagara ticks to be in the last tick group, this mirrors old behavour and can be useful to test for async overlapping issues."),
+	ECVF_Default
+);
+
 FNiagaraSystemInstance::FNiagaraSystemInstance(UNiagaraComponent* InComponent)
 	: SystemInstanceIndex(INDEX_NONE)
 	, Component(InComponent)
+	, TickGroup(TG_MAX)
+	, PrereqComponent(nullptr)
 	, Age(0.0f)
 	, TickCount(0)
 	, ID(FGuid::NewGuid())
@@ -61,6 +71,7 @@ FNiagaraSystemInstance::FNiagaraSystemInstance(UNiagaraComponent* InComponent)
 	, bPendingSpawn(false)
 	, bHasTickingEmitters(true)
 	, bPaused(false)
+	, bDataInterfacesHaveTickPrereqs(false)
 	, RequestedExecutionState(ENiagaraExecutionState::Complete)
 	, ActualExecutionState(ENiagaraExecutionState::Complete)
 	, bDataInterfacesInitialized(false)
@@ -353,6 +364,8 @@ void FNiagaraSystemInstance::SetSolo(bool bInSolo)
 
 	WaitForAsyncTick();
 
+	TickGroup = TG_MAX;
+
 	UNiagaraSystem* System = GetSystem();
 	if (bInSolo)
 	{
@@ -366,7 +379,9 @@ void FNiagaraSystemInstance::SetSolo(bool bInSolo)
 	}
 	else
 	{
-		TSharedPtr<FNiagaraSystemSimulation, ESPMode::ThreadSafe> NewSim = GetWorldManager()->GetSystemSimulation(System);
+		UpdatePrereqs();
+		TickGroup = CalculateTickGroup();
+		TSharedPtr<FNiagaraSystemSimulation, ESPMode::ThreadSafe> NewSim = GetWorldManager()->GetSystemSimulation(TickGroup, System);
 
 		NewSim->TransferInstance(SystemSimulation.Get(), this);
 		
@@ -375,6 +390,11 @@ void FNiagaraSystemInstance::SetSolo(bool bInSolo)
 	}
 }
 
+void FNiagaraSystemInstance::UpdatePrereqs()
+{
+	PrereqComponent = Component != nullptr ? Component->GetAttachParent() : nullptr;
+}
+
 void FNiagaraSystemInstance::Activate(EResetMode InResetMode)
 {
 	SCOPE_CYCLE_COUNTER(STAT_NiagaraSystemActivate);
@@ -769,6 +789,7 @@ void FNiagaraSystemInstance::ReInitInternal()
 	}
 
 	/** Do we need to run in solo mode? */
+	TickGroup = TG_MAX;
 	bSolo = bForceSolo || DoSystemDataInterfacesRequireSolo(*System, *Component);
 	if (bSolo)
 	{
@@ -780,7 +801,9 @@ void FNiagaraSystemInstance::ReInitInternal()
 	}
 	else
 	{
-		SystemSimulation = GetWorldManager()->GetSystemSimulation(System);
+		UpdatePrereqs();
+		TickGroup = CalculateTickGroup();
+		SystemSimulation = GetWorldManager()->GetSystemSimulation(TickGroup, System);
 	}
 
 	//When re initializing, throw away old emitters and init new ones.
@@ -1047,6 +1070,8 @@ bool FNiagaraSystemInstance::RequiresDistanceFieldData() const
 
 void FNiagaraSystemInstance::InitDataInterfaces()
 {
+	bDataInterfacesHaveTickPrereqs = false;
+
 	// If either the System or the component is invalid, it is possible that our cached data interfaces
 	// are now bogus and could point to invalid memory. Only the UNiagaraComponent or UNiagaraSystem
 	// can hold onto GC references to the DataInterfaces.
@@ -1091,6 +1116,11 @@ void FNiagaraSystemInstance::InitDataInterfaces()
 					InstanceDataSize += Align(Size, 16);
 				}
 			}
+
+			if (bDataInterfacesHaveTickPrereqs == false)
+			{
+				bDataInterfacesHaveTickPrereqs = Interface->HasTickGroupPrereqs();
+			}
 		}
 	};
 
@@ -1290,6 +1320,42 @@ float FNiagaraSystemInstance::GetLODDistance()
 	return DefaultLODDistance;
 }
 
+ETickingGroup FNiagaraSystemInstance::CalculateTickGroup()
+{
+	ETickingGroup NewTickGroup = (ETickingGroup)0;
+
+	// Debugging feature to force last tick group
+	if ( GNiagaraForceLastTickGroup )
+	{
+		return NiagaraLastTickGroup;
+	}
+
+	// Handle attached component tick group
+	if (PrereqComponent != nullptr)
+	{
+		//-TODO: This doesn't deal with 'DontCompleteUntil' on the prereq's tick, if we have to handle that it could mean continual TG demotion
+		ETickingGroup PrereqTG = ETickingGroup(FMath::Max(PrereqComponent->PrimaryComponentTick.TickGroup, PrereqComponent->PrimaryComponentTick.EndTickGroup) + 1);
+		NewTickGroup = FMath::Max(NewTickGroup, PrereqTG);
+	}
+
+	// Handle data interfaces that have tick dependencies
+	if ( bDataInterfacesHaveTickPrereqs )
+	{
+		for (TPair<TWeakObjectPtr<UNiagaraDataInterface>, int32>& Pair : DataInterfaceInstanceDataOffsets)
+		{
+			if (UNiagaraDataInterface* Interface = Pair.Key.Get())
+			{
+				ETickingGroup PrereqTG = Interface->CalculateTickGroup(&DataInterfaceInstanceData[Pair.Value]);
+				NewTickGroup = FMath::Max(NewTickGroup, PrereqTG);
+			}
+		}
+	}
+
+	// Clamp tick group to our range
+	NewTickGroup = FMath::Clamp(NewTickGroup, NiagaraFirstTickGroup, NiagaraLastTickGroup);
+	return NewTickGroup;
+}
+
 void FNiagaraSystemInstance::TickInstanceParameters(float DeltaSeconds)
 {
 	if (!Component)
@@ -1495,16 +1561,17 @@ void FNiagaraSystemInstance::WaitForAsyncTick(bool bEnsureComplete)
 
 		SCOPE_CYCLE_COUNTER(STAT_NiagaraSystemWaitForAsyncTick);
 
-		double StartTime = FPlatformTime::Seconds();
+		const uint64 StartCycles = FPlatformTime::Cycles64();
 		while (bAsyncWorkInProgress)
 		{
 			FPlatformProcess::SleepNoStats(0.0f);
 		}
-		double StallTimeMS = (FPlatformTime::Seconds() - StartTime) * 1000.0f;
 
+		const double StallTimeMS = FPlatformTime::ToMilliseconds64(FPlatformTime::Cycles64() - StartCycles);
 		if (StallTimeMS > GWaitForAsyncStallWarnThresholdMS)
 		{
-			UE_LOG(LogNiagara, Warning, TEXT("Niagara Effect stalled GT for %g ms.\nComponent: %s \nSystem: %s"), StallTimeMS, *GetFullNameSafe(Component), *GetFullNameSafe(GetSystem()));
+			//-TODO: This should be put back to a warning once EngineTests no longer cause it show up.  The reason it's triggered is because we pause in latent actions right after a TG running Niagara sims.
+			UE_LOG(LogNiagara, Log, TEXT("Niagara Effect stalled GT for %g ms.\nComponent: %s \nSystem: %s"), StallTimeMS, *GetFullNameSafe(Component), *GetFullNameSafe(GetSystem()));
 		}
 
 		FinalizeTick_GameThread();
@@ -1544,6 +1611,32 @@ void FNiagaraSystemInstance::Tick_GameThread(float DeltaSeconds)
 	UNiagaraSystem* System = GetSystem();
 	FScopeCycleCounter SystemStat(System->GetStatID(true, true));
 
+	// Update tick group, this can cause our Tick_GameThread to exit early due to a demotion
+	// When not in solo mode we need to deal with tick groups potentially changing due to attachments or data interfaces having prereqs
+	if ( !bSolo && !bPendingSpawn )
+	{
+		ETickingGroup NewTickGroup = CalculateTickGroup();
+		if ( NewTickGroup != TickGroup )
+		{
+			// Tick demotion we need to do this now to ensure we complete in the correct group
+			if ( NewTickGroup > TickGroup )
+			{
+				TickGroup = NewTickGroup;
+
+				TSharedPtr<FNiagaraSystemSimulation, ESPMode::ThreadSafe> NewSim = GetWorldManager()->GetSystemSimulation(TickGroup, System);
+				NewSim->TransferInstance(SystemSimulation.Get(), this);
+				SystemSimulation = NewSim;
+				return;
+			}
+			// Tick promotions must be deferred as the tick group has already been processed
+			// Note: We still tick in this group and will be transfered later on
+			else
+			{
+				SystemSimulation->AddTickGroupPromotion(this);
+			}
+		}
+	}
+
 	WaitForAsyncTick(true);
 
 	CachedDeltaSeconds = DeltaSeconds;
diff --git a/Engine/Plugins/FX/Niagara/Source/Niagara/Private/NiagaraSystemSimulation.cpp b/Engine/Plugins/FX/Niagara/Source/Niagara/Private/NiagaraSystemSimulation.cpp
--- a/Engine/Plugins/FX/Niagara/Source/Niagara/Private/NiagaraSystemSimulation.cpp
+++ b/Engine/Plugins/FX/Niagara/Source/Niagara/Private/NiagaraSystemSimulation.cpp
@@ -486,6 +486,13 @@ void FNiagaraSystemSimulation::DumpInstance(const FNiagaraSystemInstance* Inst)c
 	DataSet.Dump(Inst->SystemInstanceIndex, 1, TEXT("System Data"));
 }
 
+void FNiagaraSystemSimulation::AddTickGroupPromotion(FNiagaraSystemInstance* Instance)
+{
+	check(IsInGameThread());
+	check(!PendingTickGroupPromotions.Contains(Instance));
+	PendingTickGroupPromotions.Add(Instance);
+}
+
 bool FNiagaraSystemSimulation::ShouldTickAsync(const FNiagaraSystemSimulationTickContext& Context)
 {
 	return GbParallelSystemSimTick && FApp::ShouldUseThreadingForPerformance() && !Context.bPendingSpawnPass && Context.MyCompletionGraphEvent.IsValid();
@@ -601,6 +608,7 @@ void FNiagaraSystemSimulation::SpawnNew_GameThread(float DeltaSeconds, const FGr
 			++SystemIndex;
 		}		
 
+		//-OPT: Our tick group may have changed so next frame we could end up moving the instance again
 		if (SpawningSystemInstances.Num() > 0)
 		{
 			//Lets any RemoveInstance calls inside Tick_Concurrent know that we're spawning and all our instances are in SpawningSystemInstances;
@@ -633,6 +641,25 @@ void FNiagaraSystemSimulation::SpawnNew_GameThread(float DeltaSeconds, const FGr
 			bInSpawnPhase = false;
 		}
 	}
+
+	if (PendingTickGroupPromotions.Num() > 0)
+	{
+		FNiagaraWorldManager* WorldManager = FNiagaraWorldManager::Get(World);
+		check(WorldManager != nullptr);
+
+		for ( FNiagaraSystemInstance* Instance : PendingTickGroupPromotions)
+		{
+			const ETickingGroup TickGroup = Instance->CalculateTickGroup();
+			if ( TickGroup != Instance->TickGroup )
+			{
+				TSharedPtr<FNiagaraSystemSimulation, ESPMode::ThreadSafe> NewSim = WorldManager->GetSystemSimulation(TickGroup, System);
+				NewSim->TransferInstance(this, Instance);
+				Instance->TickGroup = TickGroup;
+				Instance->SystemSimulation = NewSim;
+			}
+		}
+		PendingTickGroupPromotions.Reset();
+	}
 }
 
 /** First phase of system sim tick. Must run on GameThread. */
diff --git a/Engine/Plugins/FX/Niagara/Source/Niagara/Private/NiagaraWorldManager.cpp b/Engine/Plugins/FX/Niagara/Source/Niagara/Private/NiagaraWorldManager.cpp
--- a/Engine/Plugins/FX/Niagara/Source/Niagara/Private/NiagaraWorldManager.cpp
+++ b/Engine/Plugins/FX/Niagara/Source/Niagara/Private/NiagaraWorldManager.cpp
@@ -20,6 +20,14 @@
 DECLARE_CYCLE_STAT(TEXT("Niagara Manager Tick [GT]"), STAT_NiagaraWorldManTick, STATGROUP_Niagara);
 DECLARE_CYCLE_STAT(TEXT("Niagara Manager Wait On Render [GT]"), STAT_NiagaraWorldManWaitOnRender, STATGROUP_Niagara);
 
+static int GNiagaraAllowAsyncWorkToEndOfFrame = 1;
+static FAutoConsoleVariableRef CVarNiagaraAllowAsyncWorkToEndOfFrame(
+	TEXT("fx.Niagara.AllowAsyncWorkToEndOfFrame"),
+	GNiagaraAllowAsyncWorkToEndOfFrame,
+	TEXT("Allow async work to continue until the end of the frame, if false it will complete within the tick group it's started in."),
+	ECVF_Default
+);
+
 FDelegateHandle FNiagaraWorldManager::OnWorldInitHandle;
 FDelegateHandle FNiagaraWorldManager::OnWorldCleanupHandle;
 FDelegateHandle FNiagaraWorldManager::OnPreWorldFinishDestroyHandle;
@@ -29,7 +37,16 @@ TMap<class UWorld*, class FNiagaraWorldManager*> FNiagaraWorldManager::WorldMana
 
 TGlobalResource<FNiagaraViewDataMgr> GNiagaraViewDataManager;
 
-FNiagaraViewDataMgr::FNiagaraViewDataMgr() 
+namespace FNiagaraUtilities
+{
+	int GetNiagaraTickGroup(ETickingGroup TickGroup)
+	{
+		const int ActualTickGroup = FMath::Clamp(TickGroup - NiagaraFirstTickGroup, 0, NiagaraNumTickGroups - 1);
+		return ActualTickGroup;
+	}
+}
+
+FNiagaraViewDataMgr::FNiagaraViewDataMgr()
 	: FRenderResource()
 	, SceneDepthTexture(nullptr)
 	, SceneNormalTexture(nullptr)
@@ -90,14 +107,17 @@ FNiagaraWorldManager::FNiagaraWorldManager(UWorld* InWorld)
 	: World(InWorld)
 	, CachedEffectsQuality(INDEX_NONE)
 {
-	FNiagaraWorldManagerTickFunction& TickFunc = TickFunction;//TODO: Add more tick functions to kick work earlier in frame.
-	TickFunc.TickGroup = ETickingGroup::TG_LastDemotable;
-	TickFunc.EndTickGroup = TickFunc.TickGroup;
-	TickFunc.bCanEverTick = true;
-	TickFunc.bStartWithTickEnabled = true;
-	TickFunc.bHighPriority = true;
-	TickFunc.Owner = this;
-	TickFunc.RegisterTickFunction(InWorld->PersistentLevel);
+	for (int32 TickGroup=0; TickGroup < NiagaraNumTickGroups; ++TickGroup)
+	{
+		FNiagaraWorldManagerTickFunction& TickFunc = TickFunctions[TickGroup];
+		TickFunc.TickGroup = ETickingGroup(NiagaraFirstTickGroup + TickGroup);
+		TickFunc.EndTickGroup = GNiagaraAllowAsyncWorkToEndOfFrame ? TG_LastDemotable : (ETickingGroup)TickFunc.TickGroup;
+		TickFunc.bCanEverTick = true;
+		TickFunc.bStartWithTickEnabled = true;
+		TickFunc.bHighPriority = true;
+		TickFunc.Owner = this;
+		TickFunc.RegisterTickFunction(InWorld->PersistentLevel);
+	}
 
 	ComponentPool = NewObject<UNiagaraComponentPool>();
 }
@@ -230,29 +250,35 @@ void FNiagaraWorldManager::CleanupParameterCollections()
 	ParameterCollections.Empty();
 }
 
-TSharedRef<FNiagaraSystemSimulation, ESPMode::ThreadSafe> FNiagaraWorldManager::GetSystemSimulation(UNiagaraSystem* System)
+TSharedRef<FNiagaraSystemSimulation, ESPMode::ThreadSafe> FNiagaraWorldManager::GetSystemSimulation(ETickingGroup TickGroup, UNiagaraSystem* System)
 {
 	LLM_SCOPE(ELLMTag::Niagara);
-	TSharedRef<FNiagaraSystemSimulation, ESPMode::ThreadSafe>* SimPtr = SystemSimulations.Find(System);
+
+	const int32 ActualTickGroup = FNiagaraUtilities::GetNiagaraTickGroup(TickGroup);
+
+	TSharedRef<FNiagaraSystemSimulation, ESPMode::ThreadSafe>* SimPtr = SystemSimulations[ActualTickGroup].Find(System);
 	if (SimPtr != nullptr)
 	{
 		return *SimPtr;
 	}
 	
 	TSharedRef<FNiagaraSystemSimulation, ESPMode::ThreadSafe> Sim = MakeShared<FNiagaraSystemSimulation, ESPMode::ThreadSafe>();
-	SystemSimulations.Add(System, Sim);
+	SystemSimulations[ActualTickGroup].Add(System, Sim);
 	Sim->Init(System, World, false);
 	return Sim;
 }
 
 void FNiagaraWorldManager::DestroySystemSimulation(UNiagaraSystem* System)
 {
-	TSharedRef<FNiagaraSystemSimulation, ESPMode::ThreadSafe>* Sim = SystemSimulations.Find(System);
-	if (Sim)
+	for ( int TG=0; TG < NiagaraNumTickGroups; ++TG )
 	{
-		(*Sim)->Destroy();
-		SystemSimulations.Remove(System);
-	}	
+		TSharedRef<FNiagaraSystemSimulation, ESPMode::ThreadSafe>* Simulation = SystemSimulations[TG].Find(System);
+		if (Simulation != nullptr)
+		{
+			(*Simulation)->Destroy();
+			SystemSimulations[TG].Remove(System);
+		}
+	}
 }
 
 void FNiagaraWorldManager::DestroySystemInstance(TUniquePtr<FNiagaraSystemInstance>& InPtr)
@@ -281,11 +307,14 @@ void FNiagaraWorldManager::OnWorldCleanup(bool bSessionEnded, bool bCleanupResou
 {
 	ComponentPool->Cleanup();
 
-	for (TPair<UNiagaraSystem*, TSharedRef<FNiagaraSystemSimulation, ESPMode::ThreadSafe>>& SimPair : SystemSimulations)
+	for (int TG = 0; TG < NiagaraNumTickGroups; ++TG)
 	{
-		SimPair.Value->Destroy();
+		for (TPair<UNiagaraSystem*, TSharedRef<FNiagaraSystemSimulation, ESPMode::ThreadSafe>>& SimPair : SystemSimulations[TG])
+		{
+			SimPair.Value->Destroy();
+		}
+		SystemSimulations[TG].Empty();
 	}
-	SystemSimulations.Empty();
 	CleanupParameterCollections();
 
 	for ( int32 i=0; i < NumDeferredQueues; ++i)
@@ -364,23 +393,27 @@ void FNiagaraWorldManager::PostActorTick(float DeltaSeconds)
 	SCOPE_CYCLE_COUNTER(STAT_NiagaraWorldManTick);
 	SCOPE_CYCLE_COUNTER(STAT_NiagaraOverview_GT);
 
-	TArray<UNiagaraSystem*, TInlineAllocator<4>> DeadSystems;
-	for (TPair<UNiagaraSystem*, TSharedRef<FNiagaraSystemSimulation, ESPMode::ThreadSafe>>& SystemSim : SystemSimulations)
+	//-TODO: This doesn't all have to happen here...
+	for (int TG = 0; TG < NiagaraNumTickGroups; ++TG)
 	{
-		FNiagaraSystemSimulation*  Sim = &SystemSim.Value.Get();
-		if (Sim->IsValid())
+		TArray<UNiagaraSystem*, TInlineAllocator<4>> DeadSystems;
+		for (TPair<UNiagaraSystem*, TSharedRef<FNiagaraSystemSimulation, ESPMode::ThreadSafe>>& SystemSim : SystemSimulations[TG])
 		{
-			Sim->SpawnNew_GameThread(DeltaSeconds, nullptr);
+			FNiagaraSystemSimulation*  Sim = &SystemSim.Value.Get();
+			if (Sim->IsValid())
+			{
+				Sim->SpawnNew_GameThread(DeltaSeconds, nullptr);
+			}
+			else
+			{
+				DeadSystems.Add(SystemSim.Key);
+			}
 		}
-		else
+		for (UNiagaraSystem* DeadSystem : DeadSystems)
 		{
-			DeadSystems.Add(SystemSim.Key);
+			SystemSimulations[TG].Remove(DeadSystem);
 		}
 	}
-	for (UNiagaraSystem* DeadSystem : DeadSystems)
-	{
-		SystemSimulations.Remove(DeadSystem);
-	}
 
 	// Clear cached player view location list, it should never be used outside of the world tick
 	bCachedPlayerViewLocationsValid = false;
@@ -403,49 +436,66 @@ void FNiagaraWorldManager::PostActorTick(float DeltaSeconds)
 		}
 		DeferredDeletionQueue[DeferredDeletionQueueIndex].Queue.Empty();
 	}
+
+	// Update tick groups
+	for (FNiagaraWorldManagerTickFunction& TickFunc : TickFunctions )
+	{
+		TickFunc.EndTickGroup = GNiagaraAllowAsyncWorkToEndOfFrame ? TG_LastDemotable : (ETickingGroup)TickFunc.TickGroup;
+	}
 }
 
 void FNiagaraWorldManager::Tick(ETickingGroup TickGroup, float DeltaSeconds, ELevelTick TickType, ENamedThreads::Type CurrentThread, const FGraphEventRef& MyCompletionGraphEvent)
 {
+	check(TickGroup >= NiagaraFirstTickGroup && TickGroup <= NiagaraLastTickGroup);
+
 	CSV_SCOPED_TIMING_STAT_EXCLUSIVE(Niagara);
 	LLM_SCOPE(ELLMTag::Niagara);
 	SCOPE_CYCLE_COUNTER(STAT_NiagaraWorldManTick);
 	SCOPE_CYCLE_COUNTER(STAT_NiagaraOverview_GT);
 
-	FNiagaraSharedObject::FlushDeletionList();
-
-	SkeletalMeshGeneratedData.TickGeneratedData(DeltaSeconds);
-
-	// Cache player view locations for all system instances to access
-	bCachedPlayerViewLocationsValid = true;
-	if (World->GetPlayerControllerIterator())
+	// We do book keeping in the first tick group
+	if ( TickGroup == NiagaraFirstTickGroup )
 	{
-		for (FConstPlayerControllerIterator Iterator = World->GetPlayerControllerIterator(); Iterator; ++Iterator)
+		FNiagaraSharedObject::FlushDeletionList();
+
+		// Cache player view locations for all system instances to access
+		//-TODO: Do we need to do this per tick group?
+		bCachedPlayerViewLocationsValid = true;
+		if (World->GetPlayerControllerIterator())
 		{
-			APlayerController* PlayerController = Iterator->Get();
-			if (PlayerController && PlayerController->IsLocalPlayerController())
+			for (FConstPlayerControllerIterator Iterator = World->GetPlayerControllerIterator(); Iterator; ++Iterator)
 			{
-				FVector* POVLoc = new(CachedPlayerViewLocations) FVector;
-				FRotator POVRotation;
-				PlayerController->GetPlayerViewPoint(*POVLoc, POVRotation);
+				APlayerController* PlayerController = Iterator->Get();
+				if (PlayerController && PlayerController->IsLocalPlayerController())
+				{
+					FVector* POVLoc = new(CachedPlayerViewLocations) FVector;
+					FRotator POVRotation;
+					PlayerController->GetPlayerViewPoint(*POVLoc, POVRotation);
+				}
 			}
 		}
-	}
-	else
-	{
-		CachedPlayerViewLocations.Append(World->ViewLocationsRenderedLastFrame);
-	}
+		else
+		{
+			CachedPlayerViewLocations.Append(World->ViewLocationsRenderedLastFrame);
+		}
 
-	//Tick our collections to push any changes to bound stores.
-	for (TPair<UNiagaraParameterCollection*, UNiagaraParameterCollectionInstance*> CollectionInstPair : ParameterCollections)
-	{
-		check(CollectionInstPair.Value);
-		CollectionInstPair.Value->Tick();
+		//Tick our collections to push any changes to bound stores.
+		//-TODO: Do we need to do this per tick group?
+		for (TPair<UNiagaraParameterCollection*, UNiagaraParameterCollectionInstance*> CollectionInstPair : ParameterCollections)
+		{
+			check(CollectionInstPair.Value);
+			CollectionInstPair.Value->Tick();
+		}
 	}
 
-	//Now tick all system instances. 
+	// Tick skeletal mesh data
+	SkeletalMeshGeneratedData.TickGeneratedData(TickGroup, DeltaSeconds);
+
+	// Now tick all system instances. 
+	const int ActualTickGroup = FNiagaraUtilities::GetNiagaraTickGroup(TickGroup);
+
 	TArray<UNiagaraSystem*, TInlineAllocator<4>> DeadSystems;
-	for (TPair<UNiagaraSystem*, TSharedRef<FNiagaraSystemSimulation, ESPMode::ThreadSafe>>& SystemSim : SystemSimulations)
+	for (TPair<UNiagaraSystem*, TSharedRef<FNiagaraSystemSimulation, ESPMode::ThreadSafe>>& SystemSim : SystemSimulations[ActualTickGroup])
 	{
 		FNiagaraSystemSimulation*  Sim = &SystemSim.Value.Get();
 
@@ -461,6 +511,6 @@ void FNiagaraWorldManager::Tick(ETickingGroup TickGroup, float DeltaSeconds, ELe
 
 	for (UNiagaraSystem* DeadSystem : DeadSystems)
 	{
-		SystemSimulations.Remove(DeadSystem);
+		SystemSimulations[ActualTickGroup].Remove(DeadSystem);
 	}
 }
diff --git a/Engine/Plugins/FX/Niagara/Source/Niagara/Public/NiagaraCommon.h b/Engine/Plugins/FX/Niagara/Source/Niagara/Public/NiagaraCommon.h
--- a/Engine/Plugins/FX/Niagara/Source/Niagara/Public/NiagaraCommon.h
+++ b/Engine/Plugins/FX/Niagara/Source/Niagara/Public/NiagaraCommon.h
@@ -29,6 +29,11 @@ const FString INTERPOLATED_PARAMETER_PREFIX = TEXT("PREV_");
 constexpr uint32 NIAGARA_MAX_GPU_SPAWN_INFOS = 8;
 constexpr uint32 NIAGARA_MAX_GPU_SPAWN_INFOS_V4 = (NIAGARA_MAX_GPU_SPAWN_INFOS + 3) / 4;
 
+/** TickGroup information for Niagara.  */
+constexpr ETickingGroup NiagaraFirstTickGroup = TG_PrePhysics;
+constexpr ETickingGroup NiagaraLastTickGroup = TG_LastDemotable;
+constexpr int NiagaraNumTickGroups = NiagaraLastTickGroup - NiagaraFirstTickGroup + 1;
+
 enum ENiagaraBaseTypes
 {
 	NBT_Float,
diff --git a/Engine/Plugins/FX/Niagara/Source/Niagara/Public/NiagaraComponent.h b/Engine/Plugins/FX/Niagara/Source/Niagara/Public/NiagaraComponent.h
--- a/Engine/Plugins/FX/Niagara/Source/Niagara/Public/NiagaraComponent.h
+++ b/Engine/Plugins/FX/Niagara/Source/Niagara/Public/NiagaraComponent.h
@@ -132,6 +132,8 @@ public:
 	virtual void Deactivate()override;
 	void DeactivateImmediate();
 
+	virtual void SetComponentTickEnabled(bool bEnabled) override;
+
 	virtual void TickComponent(float DeltaTime, enum ELevelTick TickType, FActorComponentTickFunction* ThisTickFunction) override;
 	virtual const UObject* AdditionalStatObject() const override;
 	virtual void OnComponentDestroyed(bool bDestroyingHierarchy) override;
@@ -142,8 +144,15 @@ public:
 	virtual FBoxSphereBounds CalcBounds(const FTransform& LocalToWorld) const override;
 	virtual FPrimitiveSceneProxy* CreateSceneProxy() override;
 	virtual void GetUsedMaterials(TArray<UMaterialInterface*>& OutMaterials, bool bGetDebugMaterials = false) const override;
+
+	virtual void OnAttachmentChanged() override;
 	//~ End UPrimitiveComponent Interface
 
+	//~ Begin USceneComponent Interface
+	virtual void OnChildAttached(USceneComponent* ChildComponent) override;
+	virtual void OnChildDetached(USceneComponent* ChildComponent) override;
+	//~ Begin USceneComponent Interface
+
 	TSharedPtr<FNiagaraSystemSimulation, ESPMode::ThreadSafe> GetSystemSimulation();
 
 	bool InitializeSystem();
diff --git a/Engine/Plugins/FX/Niagara/Source/Niagara/Public/NiagaraSystemInstance.h b/Engine/Plugins/FX/Niagara/Source/Niagara/Public/NiagaraSystemInstance.h
--- a/Engine/Plugins/FX/Niagara/Source/Niagara/Public/NiagaraSystemInstance.h
+++ b/Engine/Plugins/FX/Niagara/Source/Niagara/Public/NiagaraSystemInstance.h
@@ -20,6 +20,7 @@ class FNiagaraGPUSystemTick;
 
 class NIAGARA_API FNiagaraSystemInstance 
 {
+	friend class FNiagaraSystemSimulation;
 	friend class FNiagaraGPUSystemTick;
 
 public:
@@ -68,6 +69,8 @@ public:
 
 	void SetSolo(bool bInSolo);
 
+	void UpdatePrereqs();
+
 	//void RebindParameterCollection(UNiagaraParameterCollectionInstance* OldInstance, UNiagaraParameterCollectionInstance* NewInstance);
 	void BindParameters();
 	void UnbindParameters();
@@ -176,11 +179,6 @@ public:
 
 	bool IsReadyToRun() const;
 
-	/** Index of this instance in the system simulation. */
-	int32 SystemInstanceIndex;
-
-	TSharedPtr<class FNiagaraSystemSimulation, ESPMode::ThreadSafe> SystemSimulation;
-
 	FORCEINLINE bool HasTickingEmitters()const { return bHasTickingEmitters; }
 
 	UNiagaraParameterCollectionInstance* GetParameterCollectionInstance(UNiagaraParameterCollection* Collection);
@@ -249,15 +247,25 @@ private:
 	/** Call PrepareForSImulation on each data source from the simulations and determine which need per-tick updates.*/
 	void InitDataInterfaces();	
 
-
 	void BindParameterCollections(FNiagaraScriptExecutionContext& ExecContext);
 	
 	/** Calculates the distance to use for distance based LODing / culling. */
 	float GetLODDistance();
 
+	/** Calculates which tick group the instance should be in. */
+	ETickingGroup CalculateTickGroup();
+
+	/** Index of this instance in the system simulation. */
+	int32 SystemInstanceIndex;
+
+	TSharedPtr<class FNiagaraSystemSimulation, ESPMode::ThreadSafe> SystemSimulation;
+
 	UNiagaraComponent* Component;
 	FBox SystemBounds;
 
+	ETickingGroup TickGroup;
+	UActorComponent* PrereqComponent;
+
 	/** The age of the System instance. */
 	float Age;
 
@@ -342,7 +350,9 @@ private:
 	/** If this system is paused. When paused it will not tick and never complete etc. */
 	uint32 bPaused : 1;
 	/** If this system has emitters that will run GPU Simulations */
-	bool bHasGPUEmitters : 1;
+	uint32 bHasGPUEmitters : 1;
+	/** The system contains data interfaces that can have tick group prerequisites. */
+	uint32 bDataInterfacesHaveTickPrereqs : 1;
 
 	/* Execution state requested by external code/BPs calling Activate/Deactivate. */
 	ENiagaraExecutionState RequestedExecutionState;
diff --git a/Engine/Plugins/FX/Niagara/Source/Niagara/Public/NiagaraSystemSimulation.h b/Engine/Plugins/FX/Niagara/Source/Niagara/Public/NiagaraSystemSimulation.h
--- a/Engine/Plugins/FX/Niagara/Source/Niagara/Public/NiagaraSystemSimulation.h
+++ b/Engine/Plugins/FX/Niagara/Source/Niagara/Public/NiagaraSystemSimulation.h
@@ -191,6 +191,8 @@ public:
 	FNiagaraScriptExecutionContext& GetSpawnExecutionContext() { return SpawnExecContext; }
 	FNiagaraScriptExecutionContext& GetUpdateExecutionContext() { return UpdateExecContext; }
 
+	void AddTickGroupPromotion(FNiagaraSystemInstance* Instance);
+
 protected:
 	/** Does any prep work for system simulation such as pulling instance parameters into a dataset. */
 	void PrepareForSystemSimulate(FNiagaraSystemSimulationTickContext& Context);
@@ -274,6 +276,9 @@ protected:
 	TArray<FNiagaraSystemInstance*> PausedSystemInstances;
 	FNiagaraDataSet PausedInstanceData;
 
+	/** List of instances that are pending a tick group promotion. */
+	TArray<FNiagaraSystemInstance*> PendingTickGroupPromotions;
+
 	TArray<TArray<FNiagaraDataSetAccessor<FNiagaraSpawnInfo>>> EmitterSpawnInfoAccessors;
 
 	void InitParameterDataSetBindings(FNiagaraSystemInstance* SystemInst);
diff --git a/Engine/Plugins/FX/Niagara/Source/Niagara/Public/NiagaraWorldManager.h b/Engine/Plugins/FX/Niagara/Source/Niagara/Public/NiagaraWorldManager.h
--- a/Engine/Plugins/FX/Niagara/Source/Niagara/Public/NiagaraWorldManager.h
+++ b/Engine/Plugins/FX/Niagara/Source/Niagara/Public/NiagaraWorldManager.h
@@ -111,7 +111,7 @@ public:
 	UNiagaraParameterCollectionInstance* GetParameterCollection(UNiagaraParameterCollection* Collection);
 	void SetParameterCollection(UNiagaraParameterCollectionInstance* NewInstance);
 	void CleanupParameterCollections();
-	TSharedRef<FNiagaraSystemSimulation, ESPMode::ThreadSafe> GetSystemSimulation(UNiagaraSystem* System);
+	TSharedRef<FNiagaraSystemSimulation, ESPMode::ThreadSafe> GetSystemSimulation(ETickingGroup TickGroup, UNiagaraSystem* System);
 	void DestroySystemSimulation(UNiagaraSystem* System);
 	void DestroySystemInstance(TUniquePtr<FNiagaraSystemInstance>& InPtr);	
 
@@ -159,11 +159,11 @@ private:
 
 	UWorld* World;
 
-	FNiagaraWorldManagerTickFunction TickFunction;//TODO: Add more tick functions to kick some work earlier in frame.
+	FNiagaraWorldManagerTickFunction TickFunctions[NiagaraNumTickGroups];
 
 	TMap<UNiagaraParameterCollection*, UNiagaraParameterCollectionInstance*> ParameterCollections;
 
-	TMap<UNiagaraSystem*, TSharedRef<FNiagaraSystemSimulation, ESPMode::ThreadSafe>> SystemSimulations;
+	TMap<UNiagaraSystem*, TSharedRef<FNiagaraSystemSimulation, ESPMode::ThreadSafe>> SystemSimulations[NiagaraNumTickGroups];
 
 	int32 CachedEffectsQuality;
 
