- improved dx12 support for RWDI
- support for particles and RWI in a single emitter
- exposed additional grid functions on DIs

diff --git a/Engine/Plugins/FX/Niagara/Source/Niagara/Classes/NiagaraDataInterfaceGrid2DCollection.h b/Engine/Plugins/FX/Niagara/Source/Niagara/Classes/NiagaraDataInterfaceGrid2DCollection.h
--- a/Engine/Plugins/FX/Niagara/Source/Niagara/Classes/NiagaraDataInterfaceGrid2DCollection.h
+++ b/Engine/Plugins/FX/Niagara/Source/Niagara/Classes/NiagaraDataInterfaceGrid2DCollection.h
@@ -28,7 +28,7 @@ public:
 
 	int NumCellsX;
 	int NumCellsY;
-	float CellSize;
+	FVector2D CellSize;
 	bool SetGridFromVoxelSize;
 
 	FVector WorldBBoxMin;
@@ -133,6 +133,8 @@ public:
 	UFUNCTION(BlueprintCallable, Category = Niagara)
 	virtual void FillTexture2D(const UNiagaraComponent *Component, UTextureRenderTarget2D *dest);
 
+	void GetCellSize(FVectorVMContext& Context);
+
 protected:
 	//~ UNiagaraDataInterface interface
 	virtual bool CopyToInternal(UNiagaraDataInterface* Destination) const override;
diff --git a/Engine/Plugins/FX/Niagara/Source/Niagara/Classes/NiagaraDataInterfaceRW.h b/Engine/Plugins/FX/Niagara/Source/Niagara/Classes/NiagaraDataInterfaceRW.h
--- a/Engine/Plugins/FX/Niagara/Source/Niagara/Classes/NiagaraDataInterfaceRW.h
+++ b/Engine/Plugins/FX/Niagara/Source/Niagara/Classes/NiagaraDataInterfaceRW.h
@@ -23,6 +23,7 @@ static const FName NumCellsFunctionName("GetNumCells");
 static const FName CellSizeFunctionName("GetCellSize");
 
 static const FName WorldToUnitFunctionName("WorldToUnit");
+static const FName UnitToWorldFunctionName("UnitToWorld");
 static const FName UnitToIndexFunctionName("UnitToIndex");
 static const FName IndexToUnitFunctionName("IndexToUnit");
 static const FName IndexToUnitStaggeredXFunctionName("IndexToUnitStaggeredX");
diff --git a/Engine/Plugins/FX/Niagara/Source/Niagara/Classes/NiagaraEmitterInstanceBatcher.h b/Engine/Plugins/FX/Niagara/Source/Niagara/Classes/NiagaraEmitterInstanceBatcher.h
--- a/Engine/Plugins/FX/Niagara/Source/Niagara/Classes/NiagaraEmitterInstanceBatcher.h
+++ b/Engine/Plugins/FX/Niagara/Source/Niagara/Classes/NiagaraEmitterInstanceBatcher.h
@@ -127,7 +127,8 @@ public:
 				const FNiagaraGpuSpawnInfo& SpawnInfo,
 				bool bCopyBeforeStart = false,
 				uint32 ShaderStageIndex = 0,
-				FNiagaraDataInterfaceProxy *IterationInterface = nullptr
+				FNiagaraDataInterfaceProxy *IterationInterface = nullptr,
+				bool HasRunParticleStage = false
 			) const;
 
 	void ResizeCurrentBuffer(FRHICommandList &RHICmdList, FNiagaraComputeExecutionContext *Context, uint32 NewNumInstances, uint32 PrevNumInstances) const;
diff --git a/Engine/Plugins/FX/Niagara/Source/Niagara/Private/NiagaraDataInterfaceGrid2DCollection.cpp b/Engine/Plugins/FX/Niagara/Source/Niagara/Private/NiagaraDataInterfaceGrid2DCollection.cpp
--- a/Engine/Plugins/FX/Niagara/Source/Niagara/Private/NiagaraDataInterfaceGrid2DCollection.cpp
+++ b/Engine/Plugins/FX/Niagara/Source/Niagara/Private/NiagaraDataInterfaceGrid2DCollection.cpp
@@ -65,11 +65,7 @@ struct FNiagaraDataInterfaceParametersCS_Grid2DCollection : public FNiagaraDataI
 		NumCellsTmp[1] = ProxyData->NumCellsY;
 		SetShaderValue(RHICmdList, ComputeShaderRHI, NumCellsParam, NumCellsTmp);		
 		
-		// #todo(dmp): precompute these values
-		float CellSizeTmp[2];
-		CellSizeTmp[0] = ProxyData->WorldBBoxSize.X / ProxyData->NumCellsX;
-		CellSizeTmp[1] = ProxyData->WorldBBoxSize.Y / ProxyData->NumCellsY;
-		SetShaderValue(RHICmdList, ComputeShaderRHI, CellSizeParam, CellSizeTmp);		
+		SetShaderValue(RHICmdList, ComputeShaderRHI, CellSizeParam, ProxyData->CellSize);		
 		
 		SetShaderValue(RHICmdList, ComputeShaderRHI, WorldBBoxMinParam, ProxyData->WorldBBoxMin);
 		SetShaderValue(RHICmdList, ComputeShaderRHI, WorldBBoxSizeParam, ProxyData->WorldBBoxSize);
@@ -78,14 +74,15 @@ struct FNiagaraDataInterfaceParametersCS_Grid2DCollection : public FNiagaraDataI
 
 		if (GridParam.IsBound() && ProxyData->GetCurrentData() != NULL)
 		{
+			RHICmdList.TransitionResource(EResourceTransitionAccess::EReadable, EResourceTransitionPipeline::EComputeToCompute, ProxyData->GetCurrentData()->GridBuffer.UAV);
 			RHICmdList.SetShaderResourceViewParameter(Context.Shader->GetComputeShader(), GridParam.GetBaseIndex(), ProxyData->GetCurrentData()->GridBuffer.SRV);
 		}
 
 		if (Context.IsOutputStage && OutputGridParam.IsBound() && ProxyData->GetDestinationData() != NULL)
 		{
+			RHICmdList.TransitionResource(EResourceTransitionAccess::EWritable, EResourceTransitionPipeline::EComputeToCompute, ProxyData->GetDestinationData()->GridBuffer.UAV);
 			RHICmdList.SetUAVParameter(Context.Shader->GetComputeShader(), OutputGridParam.GetUAVIndex(), ProxyData->GetDestinationData()->GridBuffer.UAV);		
 		}
-		// Note: There is a flush in PreEditChange to make sure everything is synced up at this point 
 	}
 
 	virtual void Unset(FRHICommandList& RHICmdList, const FNiagaraDataInterfaceSetArgs& Context) const 
@@ -175,12 +172,25 @@ void UNiagaraDataInterfaceGrid2DCollection::GetFunctions(TArray<FNiagaraFunction
 
 }
 
-
+// #todo(dmp): expose more CPU functinality
+DEFINE_NDI_DIRECT_FUNC_BINDER(UNiagaraDataInterfaceGrid2DCollection, GetCellSize);
 void UNiagaraDataInterfaceGrid2DCollection::GetVMExternalFunction(const FVMExternalFunctionBindingInfo& BindingInfo, void* InstanceData, FVMExternalFunction &OutFunc)
 {
-	Super::GetVMExternalFunction(BindingInfo, InstanceData, OutFunc);
-
-	if (BindingInfo.Name == GetValueFunctionName) { OutFunc = FVMExternalFunction::CreateUObject(this, &UNiagaraDataInterfaceRWBase::EmptyVMFunction); }
+	if (BindingInfo.Name == NumCellsFunctionName) { OutFunc = FVMExternalFunction::CreateUObject(this, &UNiagaraDataInterfaceRWBase::EmptyVMFunction); }
+	else if (BindingInfo.Name == WorldToUnitFunctionName) { OutFunc = FVMExternalFunction::CreateUObject(this, &UNiagaraDataInterfaceRWBase::EmptyVMFunction); }
+	else if (BindingInfo.Name == UnitToWorldFunctionName) { OutFunc = FVMExternalFunction::CreateUObject(this, &UNiagaraDataInterfaceRWBase::EmptyVMFunction); }
+	else if (BindingInfo.Name == UnitToIndexFunctionName) { OutFunc = FVMExternalFunction::CreateUObject(this, &UNiagaraDataInterfaceRWBase::EmptyVMFunction); }
+	else if (BindingInfo.Name == IndexToUnitFunctionName) { OutFunc = FVMExternalFunction::CreateUObject(this, &UNiagaraDataInterfaceRWBase::EmptyVMFunction); }
+	else if (BindingInfo.Name == IndexToUnitStaggeredXFunctionName) { OutFunc = FVMExternalFunction::CreateUObject(this, &UNiagaraDataInterfaceRWBase::EmptyVMFunction); }
+	else if (BindingInfo.Name == IndexToUnitStaggeredYFunctionName) { OutFunc = FVMExternalFunction::CreateUObject(this, &UNiagaraDataInterfaceRWBase::EmptyVMFunction); }
+	else if (BindingInfo.Name == IndexToLinearFunctionName) { OutFunc = FVMExternalFunction::CreateUObject(this, &UNiagaraDataInterfaceRWBase::EmptyVMFunction); }
+	else if (BindingInfo.Name == LinearToIndexFunctionName) { OutFunc = FVMExternalFunction::CreateUObject(this, &UNiagaraDataInterfaceRWBase::EmptyVMFunction); }
+	else if (BindingInfo.Name == CellSizeFunctionName) 
+	{ 
+		check(BindingInfo.GetNumInputs() == 1 && BindingInfo.GetNumOutputs() == 2);
+		NDI_FUNC_BINDER(UNiagaraDataInterfaceGrid2DCollection, GetCellSize)::Bind(this, OutFunc);
+	}
+	else if (BindingInfo.Name == GetValueFunctionName) { OutFunc = FVMExternalFunction::CreateUObject(this, &UNiagaraDataInterfaceRWBase::EmptyVMFunction); }
 	else if (BindingInfo.Name == SetValueFunctionName) { OutFunc = FVMExternalFunction::CreateUObject(this, &UNiagaraDataInterfaceRWBase::EmptyVMFunction); }
 	else if (BindingInfo.Name == SampleGridFunctionName) { OutFunc = FVMExternalFunction::CreateUObject(this, &UNiagaraDataInterfaceRWBase::EmptyVMFunction); }
 }
@@ -325,11 +335,13 @@ bool UNiagaraDataInterfaceGrid2DCollection::InitPerInstanceData(void* PerInstanc
 
 	TSet<int> RT_OutputShaderStages = OutputShaderStages;
 	TSet<int> RT_IterationShaderStages = IterationShaderStages;
-
-	// @todo-threadsafety. This would be a race but I'm taking a ref here. Not ideal in the long term.
+	
+	FVector2D RT_CellSize = RT_WorldBBoxSize / FVector2D(RT_NumCellsX, RT_NumCellsY);
+	InstanceData->CellSize = RT_CellSize;
+	
 	// Push Updates to Proxy.
 	ENQUEUE_RENDER_COMMAND(FUpdateData)(
-		[RT_Proxy, RT_NumCellsX, RT_NumCellsY, RT_WorldBBoxMin, RT_WorldBBoxSize, RT_OutputShaderStages, RT_IterationShaderStages, InstanceID = SystemInstance->GetId()](FRHICommandListImmediate& RHICmdList)
+		[RT_Proxy, RT_NumCellsX, RT_NumCellsY, RT_CellSize, RT_WorldBBoxMin, RT_WorldBBoxSize, RT_OutputShaderStages, RT_IterationShaderStages, InstanceID = SystemInstance->GetId()](FRHICommandListImmediate& RHICmdList)
 	{
 		Grid2DCollectionRWInstanceData* TargetData = RT_Proxy->SystemInstancesToProxyData.Find(InstanceID);
 		if (TargetData != nullptr)
@@ -344,6 +356,8 @@ bool UNiagaraDataInterfaceGrid2DCollection::InitPerInstanceData(void* PerInstanc
 		TargetData->NumCellsX = RT_NumCellsX;
 		TargetData->NumCellsY = RT_NumCellsY;
 
+		TargetData->CellSize = RT_CellSize;
+
 		TargetData->WorldBBoxMin = RT_WorldBBoxMin;
 		TargetData->WorldBBoxSize = RT_WorldBBoxSize;
 
@@ -397,12 +411,25 @@ void UNiagaraDataInterfaceGrid2DCollection::FillTexture2D(const UNiagaraComponen
 
 		if (Dest && Dest->Resource && Grid2DInstanceData && Grid2DInstanceData->CurrentData)
 		{
-			FTextureRenderTarget2DResource* Texture2DResource = (FTextureRenderTarget2DResource*)Dest->Resource;
-			RHICopySharedMips(Texture2DResource->GetTextureRHI(), Grid2DInstanceData->CurrentData->GridBuffer.Buffer);
+			FRHICopyTextureInfo CopyInfo;
+			RHICmdList.CopyTexture(Grid2DInstanceData->CurrentData->GridBuffer.Buffer, Dest->Resource->TextureRHI, CopyInfo);
 		}
 	});
 }
 
+void UNiagaraDataInterfaceGrid2DCollection::GetCellSize(FVectorVMContext& Context)
+{
+	VectorVM::FUserPtrHandler<Grid2DCollectionRWInstanceData> InstData(Context);
+	VectorVM::FExternalFuncRegisterHandler<float> OutCellSizeX(Context);
+	VectorVM::FExternalFuncRegisterHandler<float> OutCellSizeY(Context);
+
+	for (int32 InstanceIdx = 0; InstanceIdx < Context.NumInstances; ++InstanceIdx)
+	{	
+		*OutCellSizeX.GetDestAndAdvance() = InstData->CellSize.X;
+		*OutCellSizeY.GetDestAndAdvance() = InstData->CellSize.Y;
+	}
+}
+
 // #todo(dmp): move these to super class
 void FNiagaraDataInterfaceProxyGrid2DCollection::DestroyPerInstanceData(NiagaraEmitterInstanceBatcher* Batcher, const FGuid& SystemInstance)
 {
@@ -426,7 +453,7 @@ void FNiagaraDataInterfaceProxyGrid2DCollection::DeferredDestroy()
 
 void FNiagaraDataInterfaceProxyGrid2DCollection::PreStage(FRHICommandList& RHICmdList, const FNiagaraDataInterfaceSetArgs& Context)
 {
-	// #todo(dmp): Context doesnt need to specify if a stage is output or not since we moved pre/post steage to the DI itself.  Not sure which design is better for the future
+	// #todo(dmp): Context doesnt need to specify if a stage is output or not since we moved pre/post stage to the DI itself.  Not sure which design is better for the future
 	if (Context.IsOutputStage)
 	{
 		Grid2DCollectionRWInstanceData* ProxyData = SystemInstancesToProxyData.Find(Context.SystemInstance);
diff --git a/Engine/Plugins/FX/Niagara/Source/Niagara/Private/NiagaraDataInterfaceNeighborGrid3D.cpp b/Engine/Plugins/FX/Niagara/Source/Niagara/Private/NiagaraDataInterfaceNeighborGrid3D.cpp
--- a/Engine/Plugins/FX/Niagara/Source/Niagara/Private/NiagaraDataInterfaceNeighborGrid3D.cpp
+++ b/Engine/Plugins/FX/Niagara/Source/Niagara/Private/NiagaraDataInterfaceNeighborGrid3D.cpp
@@ -58,6 +58,7 @@ struct FNiagaraDataInterfaceParametersCS_NeighborGrid3D : public FNiagaraDataInt
 		Ar << OutputParticleNeighborCountGridParam;		
 	}
 
+	// #todo(dmp): make resource transitions batched
 	virtual void Set(FRHICommandList& RHICmdList, const FNiagaraDataInterfaceSetArgs& Context) const override
 	{
 		check(IsInRenderingThread());
@@ -86,11 +87,13 @@ struct FNiagaraDataInterfaceParametersCS_NeighborGrid3D : public FNiagaraDataInt
 		{
 			if (ParticleNeighborsGridParam.IsBound())
 			{
+				RHICmdList.TransitionResource(EResourceTransitionAccess::EReadable, EResourceTransitionPipeline::EComputeToCompute, ProxyData->NeighborhoodBuffer.UAV);
 				RHICmdList.SetShaderResourceViewParameter(Context.Shader->GetComputeShader(), ParticleNeighborsGridParam.GetBaseIndex(), ProxyData->NeighborhoodBuffer.SRV);
 			}
 
 			if (ParticleNeighborCountGridParam.IsBound())
 			{
+				RHICmdList.TransitionResource(EResourceTransitionAccess::EReadable, EResourceTransitionPipeline::EComputeToCompute, ProxyData->NeighborhoodCountBuffer.UAV);
 				RHICmdList.SetShaderResourceViewParameter(Context.Shader->GetComputeShader(), ParticleNeighborCountGridParam.GetBaseIndex(), ProxyData->NeighborhoodCountBuffer.SRV);
 			}
 		}
@@ -98,11 +101,13 @@ struct FNiagaraDataInterfaceParametersCS_NeighborGrid3D : public FNiagaraDataInt
 		{
 			if (OutputParticleNeighborsGridParam.IsBound())
 			{
+				RHICmdList.TransitionResource(EResourceTransitionAccess::EWritable, EResourceTransitionPipeline::EComputeToCompute, ProxyData->NeighborhoodBuffer.UAV);
 				RHICmdList.SetUAVParameter(Context.Shader->GetComputeShader(), OutputParticleNeighborsGridParam.GetUAVIndex(), ProxyData->NeighborhoodBuffer.UAV);
 			}
 
 			if (OutputParticleNeighborCountGridParam.IsBound())
 			{
+				RHICmdList.TransitionResource(EResourceTransitionAccess::EWritable, EResourceTransitionPipeline::EComputeToCompute, ProxyData->NeighborhoodCountBuffer.UAV);
 				RHICmdList.SetUAVParameter(Context.Shader->GetComputeShader(), OutputParticleNeighborCountGridParam.GetUAVIndex(), ProxyData->NeighborhoodCountBuffer.UAV);
 			}
 		}
diff --git a/Engine/Plugins/FX/Niagara/Source/Niagara/Private/NiagaraDataInterfaceRW.cpp b/Engine/Plugins/FX/Niagara/Source/Niagara/Private/NiagaraDataInterfaceRW.cpp
--- a/Engine/Plugins/FX/Niagara/Source/Niagara/Private/NiagaraDataInterfaceRW.cpp
+++ b/Engine/Plugins/FX/Niagara/Source/Niagara/Private/NiagaraDataInterfaceRW.cpp
@@ -342,6 +342,18 @@ void UNiagaraDataInterfaceGrid2D::GetFunctions(TArray<FNiagaraFunctionSignature>
 		OutFunctions.Add(Sig);
 	}
 
+	{
+		FNiagaraFunctionSignature Sig;
+		Sig.Name = UnitToWorldFunctionName;
+		Sig.Inputs.Add(FNiagaraVariable(FNiagaraTypeDefinition(GetClass()), TEXT("Grid")));
+		Sig.Inputs.Add(FNiagaraVariable(FNiagaraTypeDefinition::GetVec3Def(), TEXT("Unit")));
+		Sig.Outputs.Add(FNiagaraVariable(FNiagaraTypeDefinition::GetVec3Def(), TEXT("World")));
+
+		Sig.bMemberFunction = true;
+		Sig.bRequiresContext = false;
+		OutFunctions.Add(Sig);
+	}
+
 	{
 		FNiagaraFunctionSignature Sig;
 		Sig.Name = UnitToIndexFunctionName;
@@ -425,7 +437,7 @@ void UNiagaraDataInterfaceGrid2D::GetFunctions(TArray<FNiagaraFunctionSignature>
 		FNiagaraFunctionSignature Sig;
 		Sig.Name = NumCellsFunctionName;
 		Sig.Inputs.Add(FNiagaraVariable(FNiagaraTypeDefinition(GetClass()), TEXT("Grid")));
-		Sig.Outputs.Add(FNiagaraVariable(FNiagaraTypeDefinition::GetIntDef(), TEXT("NumVoxelsX")));
+		Sig.Outputs.Add(FNiagaraVariable(FNiagaraTypeDefinition::GetIntDef(), TEXT("NumCellsX")));
 		Sig.Outputs.Add(FNiagaraVariable(FNiagaraTypeDefinition::GetIntDef(), TEXT("NumCellsY")));		
 
 		Sig.bMemberFunction = true;
@@ -450,6 +462,7 @@ void UNiagaraDataInterfaceGrid2D::GetVMExternalFunction(const FVMExternalFunctio
 {
 	if (BindingInfo.Name == NumCellsFunctionName) { OutFunc = FVMExternalFunction::CreateUObject(this, &UNiagaraDataInterfaceRWBase::EmptyVMFunction); }
 	else if (BindingInfo.Name == WorldToUnitFunctionName) { OutFunc = FVMExternalFunction::CreateUObject(this, &UNiagaraDataInterfaceRWBase::EmptyVMFunction); }
+	else if (BindingInfo.Name == UnitToWorldFunctionName) { OutFunc = FVMExternalFunction::CreateUObject(this, &UNiagaraDataInterfaceRWBase::EmptyVMFunction); }
 	else if (BindingInfo.Name == UnitToIndexFunctionName) { OutFunc = FVMExternalFunction::CreateUObject(this, &UNiagaraDataInterfaceRWBase::EmptyVMFunction); }
 	else if (BindingInfo.Name == IndexToUnitFunctionName) { OutFunc = FVMExternalFunction::CreateUObject(this, &UNiagaraDataInterfaceRWBase::EmptyVMFunction); }
 	else if (BindingInfo.Name == IndexToUnitStaggeredXFunctionName) { OutFunc = FVMExternalFunction::CreateUObject(this, &UNiagaraDataInterfaceRWBase::EmptyVMFunction); }
@@ -527,6 +540,22 @@ bool UNiagaraDataInterfaceGrid2D::GetFunctionHLSL(const FName& DefinitionFunctio
 		OutHLSL += FString::Format(FormatSample, ArgsSample);
 		return true;
 	}
+	else if (DefinitionFunctionName == UnitToWorldFunctionName)
+	{
+		static const TCHAR *FormatSample = TEXT(R"(
+			void {FunctionName}(float3 In_Unit, out float3 Out_World)
+			{
+				Out_World = In_Unit * float3({WorldBBoxSizeName}, 1) + {WorldBBoxMinName};
+			}
+		)");
+		TMap<FString, FStringFormatArg> ArgsSample = {
+			{TEXT("FunctionName"), InstanceFunctionName},
+			{TEXT("WorldBBoxMinName"), WorldBBoxMinName + ParamInfo.DataInterfaceHLSLSymbol},
+			{TEXT("WorldBBoxSizeName"), WorldBBoxSizeName + ParamInfo.DataInterfaceHLSLSymbol},
+		};
+		OutHLSL += FString::Format(FormatSample, ArgsSample);
+		return true;
+	}
 	else if (DefinitionFunctionName == UnitToIndexFunctionName)
 	{
 		static const TCHAR *FormatSample = TEXT(R"(
diff --git a/Engine/Plugins/FX/Niagara/Source/Niagara/Private/NiagaraEmitterInstanceBatcher.cpp b/Engine/Plugins/FX/Niagara/Source/Niagara/Private/NiagaraEmitterInstanceBatcher.cpp
--- a/Engine/Plugins/FX/Niagara/Source/Niagara/Private/NiagaraEmitterInstanceBatcher.cpp
+++ b/Engine/Plugins/FX/Niagara/Source/Niagara/Private/NiagaraEmitterInstanceBatcher.cpp
@@ -178,28 +178,6 @@ bool NiagaraEmitterInstanceBatcher::ResetDataInterfaces(const FNiagaraGPUSystemT
 	bool ValidSpawnStage = true;
 	FNiagaraComputeExecutionContext* Context = Instance->Context;
 
-	// The first spawn stage must not be an iteration stage
-		// #todo(dmp): fix up this logic at some point so we don't have this constraint, but likely going to be tough
-	uint32 MinSpawnStage = TNumericLimits<uint32>::Max();
-	for (uint32 SpawnStage : Context->SpawnStages)
-	{
-		if (SpawnStage < MinSpawnStage)
-		{
-			MinSpawnStage = SpawnStage;
-		}
-	}
-
-	// check all data interfaces to see if one of them has the min stage as an iteration stage
-	// if so, error out
-	for (FNiagaraDataInterfaceProxy* Interface : Instance->DataInterfaceProxies)
-	{
-		if (Interface->IsIterationStage(MinSpawnStage))
-		{
-			UE_LOG(LogNiagara, Error, TEXT("First spawn stage in execution shouldn't be an iteration stage"));
-			ValidSpawnStage = false;
-		}
-	}
-
 	// Reset all rw data interface data
 	if (Tick.bNeedsReset)
 	{
@@ -292,35 +270,41 @@ void NiagaraEmitterInstanceBatcher::DispatchMultipleStages(const FNiagaraGPUSyst
 {
 	FNiagaraComputeExecutionContext* Context = Instance->Context;
 
+	if (!ResetDataInterfaces(Tick, Instance, RHICmdList, ComputeShader)) return;
+
 	static const auto UseShaderStagesCVar = IConsoleManager::Get().FindConsoleVariable(TEXT("fx.UseShaderStages"));
 	if (UseShaderStagesCVar->GetInt() == 1)
 	{
+
+		bool HasRunParticleStage = false;
+
 		FNiagaraDataBuffer* CurrentData = Instance->CurrentData;
 		FNiagaraDataBuffer* DestinationData = Instance->DestinationData;
 
 		const uint32 NumStages = Instance->Context->MaxUpdateIterations;
 		for (uint32 ShaderStageIndex = 0; ShaderStageIndex < NumStages; ++ShaderStageIndex)
 		{
-			if (Context->SpawnStages.Num() > 0 &&
+			// Determine if the iteration is outputting to a custom data size
+			FNiagaraDataInterfaceProxy *IterationInterface = FindIterationInterface(Instance, ShaderStageIndex);
+
+			if (IterationInterface && Context->SpawnStages.Num() > 0 &&
 				((Tick.bNeedsReset && !Context->SpawnStages.Contains(ShaderStageIndex)) ||
 				(!Tick.bNeedsReset && Context->SpawnStages.Contains(ShaderStageIndex))))
 			{
 				continue;
 			}
-
-			// Determine if the iteration is outputting to a custom data size
-			FNiagaraDataInterfaceProxy *IterationInterface = FindIterationInterface(Instance, ShaderStageIndex);
 			PreStageInterface(Tick, Instance, RHICmdList, ComputeShader, ShaderStageIndex);
 
 			if (!IterationInterface)
 			{
-				if (ShaderStageIndex > 0)
+				if (HasRunParticleStage)
 				{
 					FNiagaraDataBuffer* StoreData = Instance->CurrentData;
 					Instance->CurrentData = Instance->DestinationData;
 					Instance->DestinationData = StoreData;
 				}
-				Run(Tick, Instance, 0, Instance->DestinationData->GetNumInstances(), ComputeShader, RHICmdList, ViewUniformBuffer, Instance->SpawnInfo, false, ShaderStageIndex,  nullptr);
+				Run(Tick, Instance, 0, Instance->DestinationData->GetNumInstances(), ComputeShader, RHICmdList, ViewUniformBuffer, Instance->SpawnInfo, false, ShaderStageIndex,  nullptr, HasRunParticleStage);
+				HasRunParticleStage = true;
 			}
 			else
 			{
@@ -458,36 +442,36 @@ void NiagaraEmitterInstanceBatcher::DispatchAllOnCompute(FOverlappableTicks& Ove
 	RHICmdList.TransitionResources(EResourceTransitionAccess::EReadable, EResourceTransitionPipeline::EComputeToCompute, CurrBufferIntFloat.GetData(), CurrBufferIntFloat.Num());
 	RHICmdList.FlushComputeShaderCache();
 
-	for (FNiagaraGPUSystemTick* Tick : OverlappableTick)
-	{
-		uint32 DispatchCount = Tick->Count;
-		FNiagaraComputeInstanceData* Instances = Tick->GetInstanceData();
-		for (uint32 Index = 0; Index < DispatchCount; Index++)
+		for (FNiagaraGPUSystemTick* Tick : OverlappableTick)
 		{
-			FNiagaraComputeInstanceData& Instance = Instances[Index];
-			FNiagaraComputeExecutionContext* Context = Instance.Context;
-			if (Context && Context->GPUScript_RT->GetShader())
+			uint32 DispatchCount = Tick->Count;
+			FNiagaraComputeInstanceData* Instances = Tick->GetInstanceData();
+			for (uint32 Index = 0; Index < DispatchCount; Index++)
 			{
-				FNiagaraComputeExecutionContext::TickCounter++;
+				FNiagaraComputeInstanceData& Instance = Instances[Index];
+				FNiagaraComputeExecutionContext* Context = Instance.Context;
+				if (Context && Context->GPUScript_RT->GetShader())
+				{
+					FNiagaraComputeExecutionContext::TickCounter++;
 
-				// run shader, sim and spawn in a single dispatch
+					// run shader, sim and spawn in a single dispatch
 				DispatchMultipleStages(*Tick, &Instance, RHICmdList, ViewUniformBuffer, Context->GPUScript_RT->GetShader());
 
-				FNiagaraDataBuffer* CurrentData = Instance.CurrentData;
-				if (bSetReadback && Tick->bIsFinalTick)
-				{
-					// Now that the current data is not required anymore, stage it for readback.
-					if (CurrentData->GetNumInstances() && Context->EmitterInstanceReadback.GPUCountOffset == INDEX_NONE && CurrentData->GetGPUInstanceCountBufferOffset() != INDEX_NONE)
+					FNiagaraDataBuffer* CurrentData = Instance.CurrentData;
+					if (bSetReadback && Tick->bIsFinalTick)
 					{
-						// Transfer the GPU instance counter ownership to the context. Note that when bSetReadback is true, a readback request will be performed later in the tick update.
-						Context->EmitterInstanceReadback.GPUCountOffset = CurrentData->GetGPUInstanceCountBufferOffset();
-						Context->EmitterInstanceReadback.CPUCount = CurrentData->GetNumInstances();
-						CurrentData->ClearGPUInstanceCountBufferOffset();
+						// Now that the current data is not required anymore, stage it for readback.
+						if (CurrentData->GetNumInstances() && Context->EmitterInstanceReadback.GPUCountOffset == INDEX_NONE && CurrentData->GetGPUInstanceCountBufferOffset() != INDEX_NONE)
+						{
+							// Transfer the GPU instance counter ownership to the context. Note that when bSetReadback is true, a readback request will be performed later in the tick update.
+							Context->EmitterInstanceReadback.GPUCountOffset = CurrentData->GetGPUInstanceCountBufferOffset();
+							Context->EmitterInstanceReadback.CPUCount = CurrentData->GetNumInstances();
+							CurrentData->ClearGPUInstanceCountBufferOffset();
+						}
 					}
 				}
 			}
 		}
-	}
 
 	//
 	//	Now Copy to staging buffer the data we want to read back (alive particle count). And make buffer ready for that and draw commands on the graphics pipe too.
@@ -506,7 +490,7 @@ void NiagaraEmitterInstanceBatcher::DispatchAllOnCompute(FOverlappableTicks& Ove
 	RHICmdList.FlushComputeShaderCache();
 	if (GNiagaraSubmitCommands)
 	{
-		RHICmdList.SubmitCommandsHint();
+	RHICmdList.SubmitCommandsHint();
 	}
 }
 
@@ -538,37 +522,37 @@ void NiagaraEmitterInstanceBatcher::ExecuteAll(FRHICommandList &RHICmdList, FRHI
 
 	SCOPED_DRAW_EVENT(RHICmdList, NiagaraEmitterInstanceBatcher_ExecuteAll);
 
-	TArray<FOverlappableTicks, TInlineAllocator<2>> SimPasses;
-	{
-		TMap<FNiagaraComputeExecutionContext*, FOverlappableTicks> ContextToTicks;
-
-		// Ticks are added in order. Two tick with the same context should not overlap so should be in two different batch.
-		// Those ticks should still be executed in order.
-		for (FNiagaraGPUSystemTick& Tick : Ticks_RT)
+		TArray<FOverlappableTicks, TInlineAllocator<2>> SimPasses;
 		{
-			FNiagaraComputeInstanceData* Data = Tick.GetInstanceData();
+			TMap<FNiagaraComputeExecutionContext*, FOverlappableTicks> ContextToTicks;
 
-			// This assumes all emitters fallback to the same FNiagaraShaderScript*.
-			FNiagaraShader* ComputeShader = Data->Context->GPUScript_RT->GetShader();
-			if (!ComputeShader)
+			// Ticks are added in order. Two tick with the same context should not overlap so should be in two different batch.
+			// Those ticks should still be executed in order.
+			for (FNiagaraGPUSystemTick& Tick : Ticks_RT)
 			{
-				continue;
-			}
-			else if (GNiagaraAllowTickBeforeRender && ComputeShader->ViewUniformBufferParam.IsBound() != (ViewUniformBuffer != nullptr))
-			{   // When allowing tick before render, skip this emitter if it is not in the right pass.
-				continue;
-			}
+				FNiagaraComputeInstanceData* Data = Tick.GetInstanceData();
 
-			ContextToTicks.FindOrAdd(Data->Context).Add(&Tick);
-		}
+				// This assumes all emitters fallback to the same FNiagaraShaderScript*.
+				FNiagaraShader* ComputeShader = Data->Context->GPUScript_RT->GetShader();
+				if (!ComputeShader)
+				{
+					continue;
+				}
+				else if (GNiagaraAllowTickBeforeRender && ComputeShader->ViewUniformBufferParam.IsBound() != (ViewUniformBuffer != nullptr))
+				{   // When allowing tick before render, skip this emitter if it is not in the right pass.
+					continue;
+				}
 
-		// Count the maximum number of tick per context to know the number of passes we need to do
-		uint32 NumSimPass = 0;
-		for (auto& Ticks : ContextToTicks)
-		{
-			NumSimPass = FMath::Max(NumSimPass, (uint32)Ticks.Value.Num());
-			Ticks.Value.Last()->bIsFinalTick = true;
-		}
+				ContextToTicks.FindOrAdd(Data->Context).Add(&Tick);
+			}
+
+			// Count the maximum number of tick per context to know the number of passes we need to do
+			uint32 NumSimPass = 0;
+			for (auto& Ticks : ContextToTicks)
+			{
+				NumSimPass = FMath::Max(NumSimPass, (uint32)Ticks.Value.Num());
+				Ticks.Value.Last()->bIsFinalTick = true;
+			}
 
 		// Force dispatches to run individually, this should only be used for debugging as it is highly inefficient on the GPU
 		if (GNiagaraOverlapCompute == 0)
@@ -599,21 +583,21 @@ void NiagaraEmitterInstanceBatcher::ExecuteAll(FRHICommandList &RHICmdList, FRHI
 		}
 	}
 
-	for (auto& SimPass : SimPasses)
-	{
-		FNiagaraBufferArray DestDataBuffers;
-		FNiagaraBufferArray CurrDataBuffers;
-		FNiagaraBufferArray DestBufferIntFloat;
-		FNiagaraBufferArray CurrBufferIntFloat;
+		for (auto& SimPass : SimPasses)
+		{
+			FNiagaraBufferArray DestDataBuffers;
+			FNiagaraBufferArray CurrDataBuffers;
+			FNiagaraBufferArray DestBufferIntFloat;
+			FNiagaraBufferArray CurrBufferIntFloat;
 
-		// This initial pass gathers all the buffers that are read from and written to so we can do batch resource transitions.
-		// It also ensures the GPU buffers are large enough to hold everything.
-		ResizeBuffersAndGatherResources(SimPass, RHICmdList, DestDataBuffers, CurrDataBuffers, DestBufferIntFloat, CurrBufferIntFloat);
+			// This initial pass gathers all the buffers that are read from and written to so we can do batch resource transitions.
+			// It also ensures the GPU buffers are large enough to hold everything.
+			ResizeBuffersAndGatherResources(SimPass, RHICmdList, DestDataBuffers, CurrDataBuffers, DestBufferIntFloat, CurrBufferIntFloat);
 
 		SCOPED_DRAW_EVENT(RHICmdList, NiagaraGPUSimulation);
 		SCOPED_GPU_STAT(RHICmdList, NiagaraGPUSimulation);
-		DispatchAllOnCompute(SimPass, RHICmdList, ViewUniformBuffer, DestDataBuffers, CurrDataBuffers, DestBufferIntFloat, CurrBufferIntFloat, bSetReadback);
-	}
+			DispatchAllOnCompute(SimPass, RHICmdList, ViewUniformBuffer, DestDataBuffers, CurrDataBuffers, DestBufferIntFloat, CurrBufferIntFloat, bSetReadback);
+		}
 }
 
 void NiagaraEmitterInstanceBatcher::PreInitViews(FRHICommandListImmediate& RHICmdList, bool bAllowGPUParticleUpdate)
@@ -1077,7 +1061,7 @@ void NiagaraEmitterInstanceBatcher::UnsetDataInterfaceParameters(const TArray<FN
 /* Kick off a simulation/spawn run
  */
 void NiagaraEmitterInstanceBatcher::Run(const FNiagaraGPUSystemTick& Tick, const FNiagaraComputeInstanceData* Instance, uint32 UpdateStartInstance, const uint32 TotalNumInstances, FNiagaraShader* Shader,
-	FRHICommandList &RHICmdList, FRHIUniformBuffer* ViewUniformBuffer, const FNiagaraGpuSpawnInfo& SpawnInfo, bool bCopyBeforeStart, uint32 ShaderStageIndex, FNiagaraDataInterfaceProxy *IterationInterface) const
+	FRHICommandList &RHICmdList, FRHIUniformBuffer* ViewUniformBuffer, const FNiagaraGpuSpawnInfo& SpawnInfo, bool bCopyBeforeStart, uint32 ShaderStageIndex, FNiagaraDataInterfaceProxy *IterationInterface, bool HasRunParticleStage) const
 {
 	FNiagaraComputeExecutionContext* Context = Instance->Context;
 	if (TotalNumInstances == 0)
@@ -1155,7 +1139,7 @@ void NiagaraEmitterInstanceBatcher::Run(const FNiagaraGPUSystemTick& Tick, const
 	int32 InstancesToSpawnThisFrame = Instance->SpawnInfo.SpawnRateInstances + Instance->SpawnInfo.EventSpawnTotal;
 
 	// Only spawn particles on the first stage
-	if (ShaderStageIndex > 0)
+	if (HasRunParticleStage)
 	{
 		InstancesToSpawnThisFrame = 0;
 	}
@@ -1163,18 +1147,15 @@ void NiagaraEmitterInstanceBatcher::Run(const FNiagaraGPUSystemTick& Tick, const
 	RHICmdList.SetShaderParameter(Shader->GetComputeShader(), Shader->NumSpawnedInstancesParam.GetBufferIndex(), Shader->NumSpawnedInstancesParam.GetBaseIndex(), Shader->NumSpawnedInstancesParam.GetNumBytes(), &InstancesToSpawnThisFrame);				// number of instances in the spawn run
 	RHICmdList.SetShaderParameter(Shader->GetComputeShader(), Shader->ShaderStageIndexParam.GetBufferIndex(), Shader->ShaderStageIndexParam.GetBaseIndex(), Shader->ShaderStageIndexParam.GetNumBytes(), &ShaderStageIndex);					// 0, except if several stages are defined
 
-	if (IterationInterface != NULL)
-	{
+	const int32 DefaultIterationCount = -1;
+	RHICmdList.SetShaderParameter(Shader->GetComputeShader(), Shader->IterationInterfaceCount.GetBufferIndex(), Shader->IterationInterfaceCount.GetBaseIndex(), Shader->IterationInterfaceCount.GetNumBytes(), &DefaultIterationCount);					// 0, except if several stages are defined
 
+	if (IterationInterface)
+	{
 	if (TotalNumInstances > NIAGARA_COMPUTE_THREADGROUP_SIZE)
 	{
 		RHICmdList.SetShaderParameter(Shader->GetComputeShader(), Shader->IterationInterfaceCount.GetBufferIndex(), Shader->IterationInterfaceCount.GetBaseIndex(), Shader->IterationInterfaceCount.GetNumBytes(), &TotalNumInstances);					// 0, except if several stages are defined
 	}
-	else
-	{
-		const int v = -1;
-		RHICmdList.SetShaderParameter(Shader->GetComputeShader(), Shader->IterationInterfaceCount.GetBufferIndex(), Shader->IterationInterfaceCount.GetBaseIndex(), Shader->IterationInterfaceCount.GetNumBytes(), &v);					// 0, except if several stages are defined
-	}
 	}
 
 
@@ -1221,9 +1202,7 @@ void NiagaraEmitterInstanceBatcher::Run(const FNiagaraGPUSystemTick& Tick, const
 	// reset iteration count
 	if (IterationInterface)
 	{
-		const int v = -1;
-		RHICmdList.SetShaderParameter(Shader->GetComputeShader(), Shader->IterationInterfaceCount.GetBufferIndex(), Shader->IterationInterfaceCount.GetBaseIndex(), Shader->IterationInterfaceCount.GetNumBytes(), &v);					// 0, except if several stages are defined
-
+		RHICmdList.SetShaderParameter(Shader->GetComputeShader(), Shader->IterationInterfaceCount.GetBufferIndex(), Shader->IterationInterfaceCount.GetBaseIndex(), Shader->IterationInterfaceCount.GetNumBytes(), &DefaultIterationCount);					// 0, except if several stages are defined
 	}
 
 #if WITH_EDITORONLY_DATA
