Optimized Niagara particle sorting :
- Cached the inverse transform at the proxy level instead of recomputing it in the renderer
- Moved all sort key computation to local space instead of world space
- Added radixsort32 sort at 400+ particles. Controlled by "Niagara.RadixSortThreshold"

diff --git a/Engine/Plugins/FX/Niagara/Source/Niagara/Private/NiagaraComponent.cpp b/Engine/Plugins/FX/Niagara/Source/Niagara/Private/NiagaraComponent.cpp
--- a/Engine/Plugins/FX/Niagara/Source/Niagara/Private/NiagaraComponent.cpp
+++ b/Engine/Plugins/FX/Niagara/Source/Niagara/Private/NiagaraComponent.cpp
@@ -233,6 +233,11 @@ void FNiagaraSceneProxy::CreateRenderThreadResources()
 	return;
 }
 
+void FNiagaraSceneProxy::OnTransformChanged()
+{
+	LocalToWorldInverse = GetLocalToWorld().Inverse();
+}
+
 FPrimitiveViewRelevance FNiagaraSceneProxy::GetViewRelevance(const FSceneView* View) const
 {
 	FPrimitiveViewRelevance Relevance;
diff --git a/Engine/Plugins/FX/Niagara/Source/Niagara/Private/NiagaraRenderer.cpp b/Engine/Plugins/FX/Niagara/Source/Niagara/Private/NiagaraRenderer.cpp
--- a/Engine/Plugins/FX/Niagara/Source/Niagara/Private/NiagaraRenderer.cpp
+++ b/Engine/Plugins/FX/Niagara/Source/Niagara/Private/NiagaraRenderer.cpp
@@ -9,14 +9,23 @@
 #include "Engine/Engine.h"
 #include "DynamicBufferAllocator.h"
 #include "NiagaraEmitterInstanceBatcher.h"
+#include "NiagaraGPUSortInfo.h"
 
 DECLARE_CYCLE_STAT(TEXT("Sort Particles"), STAT_NiagaraSortParticles, STATGROUP_Niagara);
-
 DECLARE_CYCLE_STAT(TEXT("Global Float Alloc - All"), STAT_NiagaraAllocateGlobalFloatAll, STATGROUP_Niagara);
 DECLARE_CYCLE_STAT(TEXT("Global Float Alloc - InsideLock"), STAT_NiagaraAllocateGlobalFloatInsideLock, STATGROUP_Niagara);
 DECLARE_CYCLE_STAT(TEXT("Global Float Alloc - Alloc New Buffer"), STAT_NiagaraAllocateGlobalFloatAllocNew, STATGROUP_Niagara);
 DECLARE_CYCLE_STAT(TEXT("Global Float Alloc - Map Buffer"), STAT_NiagaraAllocateGlobalFloatMapBuffer, STATGROUP_Niagara);
 
+int32 GNiagaraRadixSortThreshold = 400;
+static FAutoConsoleVariableRef CVarNiagaraRadixSortThreshold(
+	TEXT("Niagara.RadixSortThreshold"),
+	GNiagaraRadixSortThreshold,
+	TEXT("Instance count at which radix sort gets used instead of introspective sort.\n")
+	TEXT("Set to  -1 to never use radixsort. (default=400)"),
+	ECVF_Default
+);
+
 class FNiagaraDummyRWBufferFloat : public FRenderResource
 {
 public:
@@ -259,106 +268,103 @@ void FNiagaraRenderer::SetDynamicData_RenderThread(FNiagaraDynamicDataBase* NewD
 	DynamicDataRender = NewDynamicData;
 }
 
-void FNiagaraRenderer::SortIndices(ENiagaraSortMode SortMode, int32 SortAttributeOffset, const FNiagaraDataBuffer& Buffer, const FMatrix& LocalToWorld, const FSceneView* View, FGlobalDynamicReadBuffer::FAllocation& OutIndices)const
+struct FParticleOrderAsUint
+{
+	uint32 OrderAsUint;
+	int32 Index;
+
+	template <bool bStrictlyPositive, bool bAscending>
+	FORCEINLINE void SetAsUint(int32 InIndex, float InOrder) 
+	{
+		const uint32 SortKeySignBit = 0x80000000;
+		uint32 InOrderAsUint = reinterpret_cast<uint32&>(InOrder);
+		InOrderAsUint = (bStrictlyPositive || InOrder >= 0) ? (InOrderAsUint | SortKeySignBit) : ~InOrderAsUint;
+		OrderAsUint = bAscending ? InOrderAsUint : ~InOrderAsUint;
+
+		Index = InIndex;
+	}
+		
+	FORCEINLINE operator uint32() const { return OrderAsUint; }
+};
+
+void FNiagaraRenderer::SortIndices(const FNiagaraGPUSortInfo& SortInfo, const FNiagaraDataBuffer& Buffer, FGlobalDynamicReadBuffer::FAllocation& OutIndices)const
 {
 	SCOPE_CYCLE_COUNTER(STAT_NiagaraSortParticles);
 
 	uint32 NumInstances = Buffer.GetNumInstances();
 	check(OutIndices.ReadBuffer->NumBytes >= OutIndices.FirstIndex + NumInstances * sizeof(int32));
-	check(SortMode != ENiagaraSortMode::None);
-	check(SortAttributeOffset != INDEX_NONE);
+	check(SortInfo.SortMode != ENiagaraSortMode::None);
+	check(SortInfo.SortAttributeOffset != INDEX_NONE);
+
+	const bool bUseRadixSort = GNiagaraRadixSortThreshold != -1 && (int32)NumInstances  > GNiagaraRadixSortThreshold;
 
 	int32* RESTRICT IndexBuffer = (int32*)(OutIndices.Buffer);
 
-	struct FParticleOrder
-	{
-		int32 Index;
-		float Order;
-	};
 	FMemMark Mark(FMemStack::Get());
-	FParticleOrder* RESTRICT ParticleOrder = (FParticleOrder*)FMemStack::Get().Alloc(sizeof(FParticleOrder) * NumInstances, alignof(FParticleOrder));
+	FParticleOrderAsUint* RESTRICT ParticleOrder = (FParticleOrderAsUint*)FMemStack::Get().Alloc(sizeof(FParticleOrderAsUint) * NumInstances, alignof(FParticleOrderAsUint));
 
-	auto SortAscending = [](const FParticleOrder& A, const FParticleOrder& B) { return A.Order < B.Order; };
-	auto SortDescending = [](const FParticleOrder& A, const FParticleOrder& B) { return A.Order > B.Order; };
-
-	if (SortMode == ENiagaraSortMode::ViewDepth || SortMode == ENiagaraSortMode::ViewDistance)
+	if (SortInfo.SortMode == ENiagaraSortMode::ViewDepth || SortInfo.SortMode == ENiagaraSortMode::ViewDistance)
 	{
-		float* RESTRICT PositionX = (float*)Buffer.GetComponentPtrFloat(SortAttributeOffset);
-		float* RESTRICT PositionY = (float*)Buffer.GetComponentPtrFloat(SortAttributeOffset + 1);
-		float* RESTRICT PositionZ = (float*)Buffer.GetComponentPtrFloat(SortAttributeOffset + 2);
+		float* RESTRICT PositionX = (float*)Buffer.GetComponentPtrFloat(SortInfo.SortAttributeOffset);
+		float* RESTRICT PositionY = (float*)Buffer.GetComponentPtrFloat(SortInfo.SortAttributeOffset + 1);
+		float* RESTRICT PositionZ = (float*)Buffer.GetComponentPtrFloat(SortInfo.SortAttributeOffset + 2);
 		auto GetPos = [&PositionX, &PositionY, &PositionZ](int32 Idx)
 		{
 			return FVector(PositionX[Idx], PositionY[Idx], PositionZ[Idx]);
 		};
 
-		//TODO Parallelize in batches? Move to GPU for large emitters?
-		if (SortMode == ENiagaraSortMode::ViewDepth)
+		if (SortInfo.SortMode == ENiagaraSortMode::ViewDepth)
 		{
-			FMatrix ViewProjMatrix = View->ViewMatrices.GetViewProjectionMatrix();
-			if (bLocalSpace)
-			{
-				for (uint32 i = 0; i < NumInstances; ++i)
-				{
-					ParticleOrder[i].Index = i;
-					ParticleOrder[i].Order = ViewProjMatrix.TransformPosition(LocalToWorld.TransformPosition(GetPos(i))).W;
-				}
-			}
-			else
+			for (uint32 i = 0; i < NumInstances; ++i)
 			{
-				for (uint32 i = 0; i < NumInstances; ++i)
-				{
-					ParticleOrder[i].Index = i;
-					ParticleOrder[i].Order = ViewProjMatrix.TransformPosition(GetPos(i)).W;
-				}
+				ParticleOrder[i].SetAsUint<true, false>(i, FVector::DotProduct(GetPos(i) - SortInfo.ViewOrigin, SortInfo.ViewDirection));
 			}
-
-			Sort(ParticleOrder, NumInstances, SortDescending);
 		}
-		else 
+		else
 		{
-			// check(SortMode == ENiagaraSortMode::ViewDistance); Should always be true because we are within this block. If code is moved or copied, please change.
-			FVector ViewOrigin = View->ViewMatrices.GetViewOrigin();
-			if (bLocalSpace)
-			{
-				for (uint32 i = 0; i < NumInstances; ++i)
-				{
-					ParticleOrder[i].Index = i;
-					ParticleOrder[i].Order = (ViewOrigin - LocalToWorld.TransformPosition(GetPos(i))).SizeSquared();
-				}
-			}
-			else
+			for (uint32 i = 0; i < NumInstances; ++i)
 			{
-				for (uint32 i = 0; i < NumInstances; ++i)
-				{
-					ParticleOrder[i].Index = i;
-					ParticleOrder[i].Order = (ViewOrigin - GetPos(i)).SizeSquared();
-				}
+				ParticleOrder[i].SetAsUint<true, false>(i, (GetPos(i) - SortInfo.ViewOrigin).SizeSquared());
 			}
-
-			Sort(ParticleOrder, NumInstances, SortDescending);
 		}
 	}
 	else
 	{
-		float* RESTRICT CustomSorting = (float*)Buffer.GetComponentPtrFloat(SortAttributeOffset);
-		for (uint32 i = 0; i < NumInstances; ++i)
+		float* RESTRICT CustomSorting = (float*)Buffer.GetComponentPtrFloat(SortInfo.SortAttributeOffset);
+
+		if (SortInfo.SortMode == ENiagaraSortMode::CustomAscending)
 		{
-			ParticleOrder[i].Index = i;
-			ParticleOrder[i].Order = CustomSorting[i];
+			for (uint32 i = 0; i < NumInstances; ++i)
+			{
+				ParticleOrder[i].SetAsUint<false, true>(i, CustomSorting[i]);
+			}
 		}
-		if (SortMode == ENiagaraSortMode::CustomAscending)
+		else // ENiagaraSortMode::CustomDecending
 		{
-			Sort(ParticleOrder, NumInstances, SortAscending);
+			for (uint32 i = 0; i < NumInstances; ++i)
+			{
+				ParticleOrder[i].SetAsUint<false, false>(i, CustomSorting[i]);
+			}
 		}
-		else if (SortMode == ENiagaraSortMode::CustomDecending)
+	}
+
+	if (!bUseRadixSort)
+	{
+		Sort(ParticleOrder, NumInstances, [](const FParticleOrderAsUint& A, const FParticleOrderAsUint& B) { return A.OrderAsUint < B.OrderAsUint; });
+		//Now transfer to the real index buffer.
+		for (uint32 i = 0; i < NumInstances; ++i)
 		{
-			Sort(ParticleOrder, NumInstances, SortDescending);
+			IndexBuffer[i] = ParticleOrder[i].Index;
 		}
 	}
-
-	//Now transfer to the real index buffer.
-	for (uint32 i = 0; i < NumInstances; ++i)
+	else
 	{
-		IndexBuffer[i] = ParticleOrder[i].Index;
+		FParticleOrderAsUint* RESTRICT ParticleOrderResult = bUseRadixSort ? (FParticleOrderAsUint*)FMemStack::Get().Alloc(sizeof(FParticleOrderAsUint) * NumInstances, alignof(FParticleOrderAsUint)) : nullptr;
+		RadixSort32(ParticleOrderResult, ParticleOrder, NumInstances);
+		//Now transfer to the real index buffer.
+		for (uint32 i = 0; i < NumInstances; ++i)
+		{
+			IndexBuffer[i] = ParticleOrderResult[i].Index;
+		}
 	}
 }
diff --git a/Engine/Plugins/FX/Niagara/Source/Niagara/Private/NiagaraRendererMeshes.cpp b/Engine/Plugins/FX/Niagara/Source/Niagara/Private/NiagaraRendererMeshes.cpp
--- a/Engine/Plugins/FX/Niagara/Source/Niagara/Private/NiagaraRendererMeshes.cpp
+++ b/Engine/Plugins/FX/Niagara/Source/Niagara/Private/NiagaraRendererMeshes.cpp
@@ -249,14 +249,22 @@ void FNiagaraRendererMeshes::GetDynamicMeshElements(const TArray<const FSceneVie
 				{
 					SortInfo.ParticleCount = NumInstances;
 					SortInfo.SortMode = SortMode;
-					SortInfo.SortAttributeOffset = (SortInfo.SortMode == ENiagaraSortMode::CustomAscending || SortInfo.SortMode == ENiagaraSortMode::CustomDecending) ? CustomSortingOffset : PositionOffset;
-					SortInfo.ViewOrigin = View->ViewMatrices.GetViewOrigin();
-					SortInfo.ViewDirection = View->GetViewDirection();
-					if (bLocalSpace)
+					if (SortInfo.SortMode == ENiagaraSortMode::CustomAscending || SortInfo.SortMode == ENiagaraSortMode::CustomDecending)
 					{
-						FMatrix InvTransform = SceneProxy->GetLocalToWorld().InverseFast();
-						SortInfo.ViewOrigin = InvTransform.TransformPosition(SortInfo.ViewOrigin);
-						SortInfo.ViewDirection = InvTransform.TransformVector(SortInfo.ViewDirection);
+						SortInfo.SortAttributeOffset = CustomSortingOffset;
+						SortInfo.ViewOrigin.Set(0, 0, 0);
+						SortInfo.ViewDirection.Set(0, 0, 1);
+					}
+					else
+					{
+						SortInfo.SortAttributeOffset = PositionOffset;
+						SortInfo.ViewOrigin = View->ViewMatrices.GetViewOrigin();
+						SortInfo.ViewDirection = View->GetViewDirection();
+						if (bLocalSpace)
+						{
+							SortInfo.ViewOrigin = SceneProxy->GetLocalToWorldInverse().TransformPosition(SortInfo.ViewOrigin);
+							SortInfo.ViewDirection = SceneProxy->GetLocalToWorld().GetTransposed().TransformVector(SortInfo.ViewDirection);
+						}
 					}
 				};
 
@@ -283,7 +291,7 @@ void FNiagaraRendererMeshes::GetDynamicMeshElements(const TArray<const FSceneVie
 						{
 							FGlobalDynamicReadBuffer::FAllocation SortedIndices;
 							SortedIndices = DynamicReadBuffer.AllocateInt32(NumInstances);
-							SortIndices(SortInfo.SortMode, SortInfo.SortAttributeOffset, *SourceParticleData, SceneProxy->GetLocalToWorld(), View, SortedIndices);
+							SortIndices(SortInfo, *SourceParticleData, SortedIndices);
 							CollectorResources.VertexFactory.SetSortedIndices(SortedIndices.ReadBuffer->SRV, SortedIndices.FirstIndex / sizeof(float));
 						}
 					}
diff --git a/Engine/Plugins/FX/Niagara/Source/Niagara/Private/NiagaraRendererSprites.cpp b/Engine/Plugins/FX/Niagara/Source/Niagara/Private/NiagaraRendererSprites.cpp
--- a/Engine/Plugins/FX/Niagara/Source/Niagara/Private/NiagaraRendererSprites.cpp
+++ b/Engine/Plugins/FX/Niagara/Source/Niagara/Private/NiagaraRendererSprites.cpp
@@ -335,14 +335,22 @@ void FNiagaraRendererSprites::SetVertexFactoryParticleData(
 		{
 			SortInfo.ParticleCount = NumInstances;
 			SortInfo.SortMode = SortMode;
-			SortInfo.SortAttributeOffset = (SortInfo.SortMode == ENiagaraSortMode::CustomAscending || SortInfo.SortMode == ENiagaraSortMode::CustomDecending) ? CustomSortingOffset : PositionOffset;
-			SortInfo.ViewOrigin = View->ViewMatrices.GetViewOrigin();
-			SortInfo.ViewDirection = View->GetViewDirection();
-			if (bLocalSpace)
+			if (SortInfo.SortMode == ENiagaraSortMode::CustomAscending || SortInfo.SortMode == ENiagaraSortMode::CustomDecending)
 			{
-				FMatrix InvTransform = SceneProxy->GetLocalToWorld().InverseFast();
-				SortInfo.ViewOrigin = InvTransform.TransformPosition(SortInfo.ViewOrigin);
-				SortInfo.ViewDirection = InvTransform.TransformVector(SortInfo.ViewDirection);
+				SortInfo.SortAttributeOffset = CustomSortingOffset;
+				SortInfo.ViewOrigin.Set(0, 0, 0);
+				SortInfo.ViewDirection.Set(0, 0, 1);
+			}
+			else
+			{
+				SortInfo.SortAttributeOffset = PositionOffset;
+				SortInfo.ViewOrigin = View->ViewMatrices.GetViewOrigin();
+				SortInfo.ViewDirection = View->GetViewDirection();
+				if (bLocalSpace)
+				{
+					SortInfo.ViewOrigin = SceneProxy->GetLocalToWorldInverse().TransformPosition(SortInfo.ViewOrigin);
+					SortInfo.ViewDirection = SceneProxy->GetLocalToWorld().GetTransposed().TransformVector(SortInfo.ViewDirection);
+				}
 			}
 		};
 
@@ -371,7 +379,7 @@ void FNiagaraRendererSprites::SetVertexFactoryParticleData(
 
 					FGlobalDynamicReadBuffer::FAllocation SortedIndices;
 					SortedIndices = CPUSimParticleDataAllocation.DynamicReadBuffer.AllocateInt32(NumInstances);
-					SortIndices(SortInfo.SortMode, SortInfo.SortAttributeOffset, *SourceParticleData, SceneProxy->GetLocalToWorld(), View, SortedIndices);
+					SortIndices(SortInfo, *SourceParticleData, SortedIndices);
 					OutVertexFactory.SetSortedIndices(SortedIndices.ReadBuffer->SRV, SortedIndices.FirstIndex / sizeof(float));
 				}
 			}
diff --git a/Engine/Plugins/FX/Niagara/Source/Niagara/Public/NiagaraComponent.h b/Engine/Plugins/FX/Niagara/Source/Niagara/Public/NiagaraComponent.h
--- a/Engine/Plugins/FX/Niagara/Source/Niagara/Public/NiagaraComponent.h
+++ b/Engine/Plugins/FX/Niagara/Source/Niagara/Public/NiagaraComponent.h
@@ -523,6 +523,8 @@ public:
 	virtual bool IsRayTracingRelevant() const override { return true; }
 #endif
 
+	FORCEINLINE const FMatrix& GetLocalToWorldInverse() const { return LocalToWorldInverse; }
+
 private:
 	void ReleaseRenderThreadResources();
 
@@ -530,6 +532,7 @@ private:
 	virtual void CreateRenderThreadResources() override;
 
 	//virtual void OnActorPositionChanged() override;
+	virtual void OnTransformChanged() override;
 
 	virtual void GetDynamicMeshElements(const TArray<const FSceneView*>& Views, const FSceneViewFamily& ViewFamily, uint32 VisibilityMap, FMeshElementCollector& Collector) const override;
 
@@ -560,6 +563,8 @@ private:
 	bool bRenderingEnabled;
 	NiagaraEmitterInstanceBatcher* Batcher = nullptr;
 
+	FMatrix LocalToWorldInverse;
+
 #if STATS
 	TStatId SystemStatID;
 #endif
diff --git a/Engine/Plugins/FX/Niagara/Source/Niagara/Public/NiagaraRenderer.h b/Engine/Plugins/FX/Niagara/Source/Niagara/Public/NiagaraRenderer.h
--- a/Engine/Plugins/FX/Niagara/Source/Niagara/Public/NiagaraRenderer.h
+++ b/Engine/Plugins/FX/Niagara/Source/Niagara/Public/NiagaraRenderer.h
@@ -101,7 +101,7 @@ public:
 	virtual int32 GetDynamicDataSize()const { return 0; }
 	virtual bool IsMaterialValid(UMaterialInterface* Mat)const { return Mat != nullptr; }
 
-	void SortIndices(ENiagaraSortMode SortMode, int32 SortAttributeOffset, const FNiagaraDataBuffer& Buffer, const FMatrix& LocalToWorld, const FSceneView* View, FGlobalDynamicReadBuffer::FAllocation& OutIndices)const;
+	void SortIndices(const struct FNiagaraGPUSortInfo& SortInfo, const FNiagaraDataBuffer& Buffer, FGlobalDynamicReadBuffer::FAllocation& OutIndices)const;
 
 	void SetDynamicData_RenderThread(FNiagaraDynamicDataBase* NewDynamicData);
 	FORCEINLINE FNiagaraDynamicDataBase *GetDynamicData()const { return DynamicDataRender; }
