- Ensure we wait for instance tick if we send EOF updates mid frame

#rnx

diff --git a/Engine/Plugins/FX/Niagara/Source/Niagara/Private/NiagaraComponent.cpp b/Engine/Plugins/FX/Niagara/Source/Niagara/Private/NiagaraComponent.cpp
--- a/Engine/Plugins/FX/Niagara/Source/Niagara/Private/NiagaraComponent.cpp
+++ b/Engine/Plugins/FX/Niagara/Source/Niagara/Private/NiagaraComponent.cpp
@@ -982,6 +982,15 @@ TSharedPtr<FNiagaraSystemSimulation, ESPMode::ThreadSafe> UNiagaraComponent::Get
 	return nullptr;
 }
 
+void UNiagaraComponent::OnEndOfFrameUpdateDuringTick()
+{
+	Super::OnEndOfFrameUpdateDuringTick();
+	if ( SystemInstance )
+	{
+		SystemInstance->WaitForAsyncTick(true);
+	}
+}
+
 void UNiagaraComponent::CreateRenderState_Concurrent()
 {
 	Super::CreateRenderState_Concurrent();
diff --git a/Engine/Plugins/FX/Niagara/Source/Niagara/Public/NiagaraComponent.h b/Engine/Plugins/FX/Niagara/Source/Niagara/Public/NiagaraComponent.h
--- a/Engine/Plugins/FX/Niagara/Source/Niagara/Public/NiagaraComponent.h
+++ b/Engine/Plugins/FX/Niagara/Source/Niagara/Public/NiagaraComponent.h
@@ -110,6 +110,7 @@ private:
 protected:
 	virtual void OnRegister() override;
 	virtual void OnUnregister() override;
+	virtual void OnEndOfFrameUpdateDuringTick() override;
 	virtual void CreateRenderState_Concurrent() override;
 	virtual void SendRenderDynamicData_Concurrent() override;
 	virtual void BeginDestroy() override;
