- Add explicit FNiagaraSystemSimulationWaitAllFinalizeTask task to handle dependencies rather than using a FNullGraphTask
- This fixes a checkSlow in debug as we were adding EventsToWaitFor outside of the task execution, tehnically it worked fine since the task was not unlocked but this is the correct way

#rnx

diff --git a/Engine/Plugins/FX/Niagara/Source/Niagara/Private/NiagaraSystemSimulation.cpp b/Engine/Plugins/FX/Niagara/Source/Niagara/Private/NiagaraSystemSimulation.cpp
--- a/Engine/Plugins/FX/Niagara/Source/Niagara/Private/NiagaraSystemSimulation.cpp
+++ b/Engine/Plugins/FX/Niagara/Source/Niagara/Private/NiagaraSystemSimulation.cpp
@@ -19,7 +19,7 @@
 // If simulations are allowed to tick async we will create a FNiagaraSystemSimulationTickTask task to run on any thread
 // If instances are allowed to tick async we will create a FNiagaraSystemInstanceAsyncTask in batches to run on any thread
 // If any async is enabled we create a FNiagaraSystemInstanceFinalizeTask for each batch that will not run until FNiagaraSystemSimulationTickTask is complete (Due to contention with SystemInstances) and will run on the GameThread
-// If any async is enabled we create a null graph task to wait for all FNiagaraSystemInstanceFinalizeTask's to complete before allowing the tick group to advance
+// If any async is enabled we create a FNiagaraSystemSimulationWaitAllFinalizeTask task to wait for all FNiagaraSystemInstanceFinalizeTask's to complete before allowing the tick group to advance
 
 //High level stats for system sim tick.
 DECLARE_CYCLE_STAT(TEXT("System Simulaton Tick [GT]"), STAT_NiagaraSystemSim_TickGT, STATGROUP_Niagara);
@@ -133,36 +133,52 @@ FAutoConsoleTaskPriority CPrio_NiagaraSystemSimulationTickTask(
 	ENamedThreads::HighTaskPriority // if we don't have hi pri threads, then use normal priority threads at high task priority instead
 );
 
-//This task performs the concurrent part of the system simulation tick.
-class FNiagaraSystemSimulationTickTask
+// This task is used to wait for all finalize tasks to complete
+class FNiagaraSystemSimulationWaitAllFinalizeTask
 {
-	FNiagaraSystemSimulationTickContext Context;
-	TGraphTask<FNullGraphTask>* FinalizeCompleteTask;
+	FGraphEventArray EventsToWaitFor;
 public:
-	FNiagaraSystemSimulationTickTask(FNiagaraSystemSimulationTickContext InContext, TGraphTask<FNullGraphTask>* InFinalizeCompleteTask)
-		: Context(InContext)
-		, FinalizeCompleteTask(InFinalizeCompleteTask)
+	FNiagaraSystemSimulationWaitAllFinalizeTask(FGraphEventArray*& OutEventsToWaitFor)
 	{
+		OutEventsToWaitFor = &EventsToWaitFor;
 	}
 
-	FORCEINLINE TStatId GetStatId() const
+	FORCEINLINE TStatId GetStatId() const { RETURN_QUICK_DECLARE_CYCLE_STAT(FNiagaraSystemSimulationWaitAllFinalizeTask, STATGROUP_TaskGraphTasks); }
+	ENamedThreads::Type GetDesiredThread() { return CPrio_NiagaraSystemSimulationTickTask.Get(); }
+	static ESubsequentsMode::Type GetSubsequentsMode() { return ESubsequentsMode::TrackSubsequents; }
+
+	void DoTask(ENamedThreads::Type CurrentThread, const FGraphEventRef& MyCompletionGraphEvent)
 	{
-		RETURN_QUICK_DECLARE_CYCLE_STAT(FNiagaraSystemSimulationTickTask, STATGROUP_TaskGraphTasks);
+		for ( FGraphEventRef& Event : EventsToWaitFor )
+		{
+			MyCompletionGraphEvent->DontCompleteUntil(Event);
+		}
+		EventsToWaitFor.Empty();
 	}
+};
 
-	ENamedThreads::Type GetDesiredThread()
+//This task performs the concurrent part of the system simulation tick.
+class FNiagaraSystemSimulationTickTask
+{
+	FNiagaraSystemSimulationTickContext Context;
+	TGraphTask<FNiagaraSystemSimulationWaitAllFinalizeTask>* WaitAllFinalizeTask;
+public:
+	FNiagaraSystemSimulationTickTask(FNiagaraSystemSimulationTickContext InContext, TGraphTask<FNiagaraSystemSimulationWaitAllFinalizeTask>* InWaitAllFinalizeTask)
+		: Context(InContext)
+		, WaitAllFinalizeTask(InWaitAllFinalizeTask)
 	{
-		return CPrio_NiagaraSystemSimulationTickTask.Get();
 	}
 
+	FORCEINLINE TStatId GetStatId() const { RETURN_QUICK_DECLARE_CYCLE_STAT(FNiagaraSystemSimulationTickTask, STATGROUP_TaskGraphTasks); }
+	ENamedThreads::Type GetDesiredThread() { return CPrio_NiagaraSystemSimulationTickTask.Get(); }
 	static ESubsequentsMode::Type GetSubsequentsMode() { return ESubsequentsMode::TrackSubsequents; }
 
 	void DoTask(ENamedThreads::Type CurrentThread, const FGraphEventRef& MyCompletionGraphEvent)
 	{
 		Context.MyCompletionGraphEvent = MyCompletionGraphEvent;
 		Context.Owner->Tick_Concurrent(Context);
-
-		FinalizeCompleteTask->Unlock();
+		Context.FinalizeEvents = nullptr;
+		WaitAllFinalizeTask->Unlock();
 	}
 };
 
@@ -184,16 +200,8 @@ public:
 
 	}
 
-	FORCEINLINE TStatId GetStatId() const
-	{
-		RETURN_QUICK_DECLARE_CYCLE_STAT(FNiagaraSystemInstanceFinalizeTask, STATGROUP_TaskGraphTasks);
-	}
-
-	ENamedThreads::Type GetDesiredThread()
-	{
-		return ENamedThreads::GameThread;
-	}
-
+	FORCEINLINE TStatId GetStatId() const { RETURN_QUICK_DECLARE_CYCLE_STAT(FNiagaraSystemInstanceFinalizeTask, STATGROUP_TaskGraphTasks); }
+	ENamedThreads::Type GetDesiredThread() { return ENamedThreads::GameThread; }
 	static ESubsequentsMode::Type GetSubsequentsMode() { return ESubsequentsMode::TrackSubsequents; }
 
 	void DoTask(ENamedThreads::Type CurrentThread, const FGraphEventRef& MyCompletionGraphEvent)
@@ -229,19 +237,10 @@ public:
 		: SystemSim(InSystemSim)
 		, Batch(InBatch)
 	{
-
-	}
-
-	FORCEINLINE TStatId GetStatId() const
-	{
-		RETURN_QUICK_DECLARE_CYCLE_STAT(FNiagaraSystemInstanceAsyncTask, STATGROUP_TaskGraphTasks);
-	}
-
-	ENamedThreads::Type GetDesiredThread()
-	{
-		return CPrio_NiagaraSystemInstanceAsyncTask.Get();
 	}
 
+	FORCEINLINE TStatId GetStatId() const { RETURN_QUICK_DECLARE_CYCLE_STAT(FNiagaraSystemInstanceAsyncTask, STATGROUP_TaskGraphTasks); }
+	ENamedThreads::Type GetDesiredThread() 	{ return CPrio_NiagaraSystemInstanceAsyncTask.Get(); }
 	static ESubsequentsMode::Type GetSubsequentsMode() { return ESubsequentsMode::TrackSubsequents; }
 
 	void DoTask(ENamedThreads::Type CurrentThread, const FGraphEventRef& MyCompletionGraphEvent)
@@ -522,7 +521,7 @@ void FNiagaraSystemSimulation::FlushTickBatch(FNiagaraSystemSimulationTickContex
 		// Enqueue or tick the instances
 		if ( Context.bTickInstancesAsync )
 		{
-			check(Context.FinalizeCompleteGraphEvent.IsValid());
+			check(Context.FinalizeEvents != nullptr);
 
 			FGraphEventRef AsyncTask = TGraphTask<FNiagaraSystemInstanceAsyncTask>::CreateTask(nullptr).ConstructAndDispatchWhenReady(this, TickBatch);
 			FinalizePrereq.Add(AsyncTask);
@@ -538,10 +537,10 @@ void FNiagaraSystemSimulation::FlushTickBatch(FNiagaraSystemSimulationTickContex
 		// Enqueue a finalize task?
 		if (Context.bTickAsync || Context.bTickInstancesAsync )
 		{
-			check(Context.FinalizeCompleteGraphEvent.IsValid());
+			check(Context.FinalizeEvents != nullptr);
 
 			FGraphEventRef FinalizeTask = TGraphTask<FNiagaraSystemInstanceFinalizeTask>::CreateTask(&FinalizePrereq).ConstructAndDispatchWhenReady(this, TickBatch);
-			Context.FinalizeCompleteGraphEvent->DontCompleteUntil(FinalizeTask);
+			Context.FinalizeEvents->Add(FinalizeTask);
 		}
 
 		TickBatch.Reset();
@@ -725,38 +724,32 @@ void FNiagaraSystemSimulation::Tick_GameThread(float DeltaSeconds, const FGraphE
 	//Now kick of the concurrent tick.
 	if (Context.bTickAsync)
 	{
-		DECLARE_CYCLE_STAT(TEXT("FNullGraphTask.NiagaraFinalizeComplete"), STAT_FNullGraphTask_NiagaraFinalizeComplete, STATGROUP_TaskGraphTasks);
+		TGraphTask<FNiagaraSystemSimulationWaitAllFinalizeTask>* WaitAllFinalizeTask = TGraphTask<FNiagaraSystemSimulationWaitAllFinalizeTask>::CreateTask(nullptr, ENamedThreads::GameThread).ConstructAndHold(Context.FinalizeEvents);
+		FGraphEventRef WaitAllFinalizeGraphEvent = WaitAllFinalizeTask->GetCompletionEvent();
 
-		TGraphTask<FNullGraphTask>* FinalizeCompleteTask = TGraphTask<FNullGraphTask>::CreateTask(nullptr, ENamedThreads::GameThread).ConstructAndHold(GET_STATID(STAT_FNullGraphTask_NiagaraFinalizeComplete), ENamedThreads::AnyThread);
-		Context.FinalizeCompleteGraphEvent = FinalizeCompleteTask->GetCompletionEvent();
-
-		TGraphTask<FNiagaraSystemSimulationTickTask>* AsyncTask = TGraphTask<FNiagaraSystemSimulationTickTask>::CreateTask(nullptr, ENamedThreads::GameThread).ConstructAndHold(Context, FinalizeCompleteTask);
-		FGraphEventRef AsyncGraphEvent = AsyncTask->GetCompletionEvent();
+		TGraphTask<FNiagaraSystemSimulationTickTask>* SimulationTickTask = TGraphTask<FNiagaraSystemSimulationTickTask>::CreateTask(nullptr, ENamedThreads::GameThread).ConstructAndHold(Context, WaitAllFinalizeTask);
+		FGraphEventRef SimulationTickGraphEvent = SimulationTickTask->GetCompletionEvent();
+		Context.FinalizeEvents->Add(SimulationTickGraphEvent);
 
 		MyCompletionGraphEvent->SetGatherThreadForDontCompleteUntil(ENamedThreads::GameThread);
-		MyCompletionGraphEvent->DontCompleteUntil(Context.FinalizeCompleteGraphEvent);
-		Context.FinalizeCompleteGraphEvent->SetGatherThreadForDontCompleteUntil(ENamedThreads::GameThread);
-		Context.FinalizeCompleteGraphEvent->DontCompleteUntil(AsyncGraphEvent);
+		MyCompletionGraphEvent->DontCompleteUntil(WaitAllFinalizeGraphEvent);
 
-		AsyncTask->Unlock(ENamedThreads::GameThread);
+		SimulationTickTask->Unlock(ENamedThreads::GameThread);
 	}
 	else
 	{
-		TGraphTask<FNullGraphTask>* WaitConcurrentTick = nullptr;
+		TGraphTask<FNiagaraSystemSimulationWaitAllFinalizeTask>* WaitAllFinalizeTask = nullptr;
 		if (Context.bTickInstancesAsync)
 		{
-			DECLARE_CYCLE_STAT(TEXT("FNullGraphTask.NiagaraWaitToFinalize"), STAT_FNullGraphTask_NiagaraWaitToFinalize, STATGROUP_TaskGraphTasks);
-
-			WaitConcurrentTick = TGraphTask<FNullGraphTask>::CreateTask(nullptr, ENamedThreads::GameThread).ConstructAndHold(GET_STATID(STAT_FNullGraphTask_NiagaraWaitToFinalize), ENamedThreads::AnyThread);
-			Context.MyCompletionGraphEvent = WaitConcurrentTick->GetCompletionEvent();
-			Context.FinalizeCompleteGraphEvent = MyCompletionGraphEvent;
+			WaitAllFinalizeTask = TGraphTask<FNiagaraSystemSimulationWaitAllFinalizeTask>::CreateTask(nullptr, ENamedThreads::GameThread).ConstructAndHold(Context.FinalizeEvents);
 		}
 
 		Tick_Concurrent(Context);
 
 		if (Context.bTickInstancesAsync)
 		{
-			WaitConcurrentTick->Unlock(ENamedThreads::GameThread);
+			WaitAllFinalizeTask->Unlock(ENamedThreads::GameThread);
+			Context.FinalizeEvents = nullptr;
 		}
 	}
 
diff --git a/Engine/Plugins/FX/Niagara/Source/Niagara/Public/NiagaraSystemSimulation.h b/Engine/Plugins/FX/Niagara/Source/Niagara/Public/NiagaraSystemSimulation.h
--- a/Engine/Plugins/FX/Niagara/Source/Niagara/Public/NiagaraSystemSimulation.h
+++ b/Engine/Plugins/FX/Niagara/Source/Niagara/Public/NiagaraSystemSimulation.h
@@ -132,7 +132,7 @@ struct FNiagaraSystemSimulationTickContext
 	bool bPendingSpawnPass;
 	FGraphEventRef MyCompletionGraphEvent;
 
-	FGraphEventRef FinalizeCompleteGraphEvent;
+	FGraphEventArray* FinalizeEvents = nullptr;
 
 	bool bTickAsync;
 	bool bTickInstancesAsync;
