Added a data interface to allow bulk exporting particle data to blueprint or c++

diff --git a/Engine/Plugins/FX/Niagara/Source/Niagara/Classes/NiagaraDataInterfaceExport.h b/Engine/Plugins/FX/Niagara/Source/Niagara/Classes/NiagaraDataInterfaceExport.h
new file mode 100644
--- /dev/null
+++ b/Engine/Plugins/FX/Niagara/Source/Niagara/Classes/NiagaraDataInterfaceExport.h
@@ -0,0 +1,84 @@
+// Copyright 1998-2019 Epic Games, Inc. All Rights Reserved.
+#pragma once
+
+#include "NiagaraCommon.h"
+#include "NiagaraShared.h"
+#include "NiagaraDataInterface.h"
+#include "NiagaraParameterStore.h"
+#include "NiagaraDataInterfaceExport.generated.h"
+
+USTRUCT(BlueprintType, Blueprintable)
+struct FBasicParticleData
+{
+	GENERATED_USTRUCT_BODY()
+
+	UPROPERTY(BlueprintReadOnly, Category = "Particle Data")
+	FVector Position;
+
+	UPROPERTY(BlueprintReadOnly, Category = "Particle Data")
+	float Size;
+
+	UPROPERTY(BlueprintReadOnly, Category = "Particle Data")
+	FVector Velocity;
+};
+
+UINTERFACE(BlueprintType)
+class UNiagaraParticleCallbackHandler : public UInterface
+{
+	GENERATED_BODY()
+};
+
+class INiagaraParticleCallbackHandler
+{
+	GENERATED_BODY()
+
+public:
+	/** This function is called once per tick with the gathered particle data. It will not be called if there is no particle data to call it with. */
+	UFUNCTION(BlueprintCallable, BlueprintNativeEvent, Category = "Niagara")
+	void ReceiveParticleData(const TArray<FBasicParticleData>& Data, UNiagaraSystem* NiagaraSystem) const;
+};
+
+struct ExportInterface_InstanceData
+{
+	UObject* CallbackHandler = nullptr;
+	TArray<FBasicParticleData> GatheredData;
+
+	/** A binding to the user ptr we're reading the handler from. */
+	FNiagaraParameterDirectBinding<UObject*> UserParamBinding;
+};
+
+/** This Data Interface can be used to gather particles at execution time and call either a 
+C++ or blueprint object with the gathered particle data each tick. */
+UCLASS(EditInlineNew, Category = "Export", meta = (DisplayName = "Export particle data"))
+class NIAGARA_API UNiagaraDataInterfaceExport : public UNiagaraDataInterface
+{
+	GENERATED_UCLASS_BODY()
+
+public:
+	/** Reference to a user parameter that should receive the particle data after the simulation tick. The supplied parameter object needs to implement the INiagaraParticleCallbackHandler interface. */
+	UPROPERTY(EditAnywhere, Category = "Export")
+	FNiagaraUserParameterBinding CallbackHandlerParameter;
+
+	//UObject Interface
+	virtual void PostInitProperties() override;
+	//UObject Interface End
+
+	//UNiagaraDataInterface Interface
+	virtual void GetFunctions(TArray<FNiagaraFunctionSignature>& OutFunctions)override;
+	virtual void GetVMExternalFunction(const FVMExternalFunctionBindingInfo& BindingInfo, void* InstanceData, FVMExternalFunction &OutFunc) override;
+	virtual bool InitPerInstanceData(void* PerInstanceData, FNiagaraSystemInstance* SystemInstance) override;
+	virtual int32 PerInstanceDataSize() const { return sizeof(ExportInterface_InstanceData); }
+	virtual bool PerInstanceTick(void* PerInstanceData, FNiagaraSystemInstance* SystemInstance, float DeltaSeconds) override;
+	virtual bool PerInstanceTickPostSimulate(void* PerInstanceData, FNiagaraSystemInstance* SystemInstance, float DeltaSeconds);
+	virtual bool Equals(const UNiagaraDataInterface* Other) const override;
+	virtual bool GetFunctionHLSL(const FName&  DefinitionFunctionName, FString InstanceFunctionName, FNiagaraDataInterfaceGPUParamInfo& ParamInfo, FString& OutHLSL) override;
+	//UNiagaraDataInterface Interface
+
+	virtual void StoreData(FVectorVMContext& Context);
+
+protected:
+	virtual bool CopyToInternal(UNiagaraDataInterface* Destination) const override;
+	
+private:
+	static const FName StoreDataName;
+};
diff --git a/Engine/Plugins/FX/Niagara/Source/Niagara/Private/NiagaraDataInterfaceExport.cpp b/Engine/Plugins/FX/Niagara/Source/Niagara/Private/NiagaraDataInterfaceExport.cpp
new file mode 100644
--- /dev/null
+++ b/Engine/Plugins/FX/Niagara/Source/Niagara/Private/NiagaraDataInterfaceExport.cpp
@@ -0,0 +1,157 @@
+// Copyright 1998-2019 Epic Games, Inc. All Rights Reserved.
+
+#include "NiagaraDataInterfaceExport.h"
+#include "NiagaraTypes.h"
+#include "NiagaraCustomVersion.h"
+#include "..\Classes\NiagaraDataInterfaceExport.h"
+#include "Internationalization/Internationalization.h"
+#include "NiagaraSystemInstance.h"
+
+const FName UNiagaraDataInterfaceExport::StoreDataName(TEXT("StoreParticleData"));
+
+UNiagaraDataInterfaceExport::UNiagaraDataInterfaceExport(FObjectInitializer const& ObjectInitializer)
+	: Super(ObjectInitializer)
+{
+}
+
+void UNiagaraDataInterfaceExport::PostInitProperties()
+{
+	Super::PostInitProperties();
+
+	if (HasAnyFlags(RF_ClassDefaultObject))
+	{
+		FNiagaraTypeRegistry::Register(FNiagaraTypeDefinition(GetClass()), true, false, false);
+	}
+}
+
+bool UNiagaraDataInterfaceExport::InitPerInstanceData(void* PerInstanceData, FNiagaraSystemInstance* SystemInstance)
+{
+	ExportInterface_InstanceData* PIData = new (PerInstanceData) ExportInterface_InstanceData;
+	return true;
+}
+
+bool UNiagaraDataInterfaceExport::PerInstanceTick(void* PerInstanceData, FNiagaraSystemInstance* SystemInstance, float DeltaSeconds)
+{
+	ExportInterface_InstanceData* PIData = (ExportInterface_InstanceData*)PerInstanceData;
+	if (!PIData)
+	{
+		return true;
+	}
+
+	if (CallbackHandlerParameter.Parameter.IsValid() && SystemInstance)
+	{
+		PIData->CallbackHandler = PIData->UserParamBinding.Init(SystemInstance->GetInstanceParameters(), CallbackHandlerParameter.Parameter);
+	}
+	else
+	{
+		PIData->CallbackHandler = nullptr;
+	}
+	return false;
+}
+
+bool UNiagaraDataInterfaceExport::PerInstanceTickPostSimulate(void* PerInstanceData, FNiagaraSystemInstance* SystemInstance, float DeltaSeconds)
+{
+	ExportInterface_InstanceData* PIData = (ExportInterface_InstanceData*) PerInstanceData;
+	if (PIData->GatheredData.Num() > 0 && PIData->CallbackHandler && PIData->CallbackHandler->GetClass()->ImplementsInterface(UNiagaraParticleCallbackHandler::StaticClass()))
+	{
+		INiagaraParticleCallbackHandler::Execute_ReceiveParticleData(PIData->CallbackHandler, PIData->GatheredData, SystemInstance->GetSystem());
+	}
+	PIData->GatheredData.Empty();
+	return false;
+}
+
+bool UNiagaraDataInterfaceExport::Equals(const UNiagaraDataInterface* Other) const
+{
+	if (!Super::Equals(Other))
+	{
+		return false;
+	}
+
+	return CastChecked<UNiagaraDataInterfaceExport>(Other)->CallbackHandlerParameter == CallbackHandlerParameter;
+}
+
+void UNiagaraDataInterfaceExport::GetFunctions(TArray<FNiagaraFunctionSignature>& OutFunctions)
+{
+	FNiagaraFunctionSignature Sig;
+	Sig.Name = StoreDataName;
+#if WITH_EDITORONLY_DATA
+	Sig.Description = NSLOCTEXT("Niagara", "ExportDataFunctionDescription", "This function takes the particle data and stores it to be exported to the registered callback handler after the simulation has ticked.");
+#endif
+	Sig.bMemberFunction = true;
+	Sig.bRequiresContext = false;
+	Sig.Inputs.Add(FNiagaraVariable(FNiagaraTypeDefinition(GetClass()), TEXT("Export interface")));
+	Sig.Inputs.Add(FNiagaraVariable(FNiagaraTypeDefinition::GetBoolDef(), TEXT("Store Data")));
+	Sig.Inputs.Add(FNiagaraVariable(FNiagaraTypeDefinition::GetVec3Def(), TEXT("Position")));
+	Sig.Inputs.Add(FNiagaraVariable(FNiagaraTypeDefinition::GetFloatDef(), TEXT("Size")));
+	Sig.Inputs.Add(FNiagaraVariable(FNiagaraTypeDefinition::GetVec3Def(), TEXT("Velocity")));
+	Sig.Outputs.Add(FNiagaraVariable(FNiagaraTypeDefinition::GetBoolDef(), TEXT("Success")));
+	OutFunctions.Add(Sig);
+}
+
+bool UNiagaraDataInterfaceExport::GetFunctionHLSL(const FName& DefinitionFunctionName, FString InstanceFunctionName, FNiagaraDataInterfaceGPUParamInfo& ParamInfo, FString& OutHLSL)
+{
+	return false;
+}
+
+DEFINE_NDI_DIRECT_FUNC_BINDER(UNiagaraDataInterfaceExport, StoreData);
+void UNiagaraDataInterfaceExport::GetVMExternalFunction(const FVMExternalFunctionBindingInfo& BindingInfo, void* InstanceData, FVMExternalFunction &OutFunc)
+{
+	if (BindingInfo.Name == StoreDataName)
+	{
+		NDI_FUNC_BINDER(UNiagaraDataInterfaceExport, StoreData)::Bind(this, OutFunc);
+	}
+	else
+	{
+		UE_LOG(LogNiagara, Error, TEXT("Could not find data interface external function. Expected Name: %s  Actual Name: %s"), *StoreDataName.ToString(), *BindingInfo.Name.ToString());
+	}
+}
+
+void UNiagaraDataInterfaceExport::StoreData(FVectorVMContext& Context)
+{
+	VectorVM::FExternalFuncInputHandler<FNiagaraBool> StoreDataParam(Context);
+
+	VectorVM::FExternalFuncInputHandler<float> PositionParamX(Context);
+	VectorVM::FExternalFuncInputHandler<float> PositionParamY(Context);
+	VectorVM::FExternalFuncInputHandler<float> PositionParamZ(Context);
+
+	VectorVM::FExternalFuncInputHandler<float> SizeParam(Context);
+	
+	VectorVM::FExternalFuncInputHandler<float> VelocityParamX(Context);
+	VectorVM::FExternalFuncInputHandler<float> VelocityParamY(Context);
+	VectorVM::FExternalFuncInputHandler<float> VelocityParamZ(Context);
+
+	VectorVM::FUserPtrHandler<ExportInterface_InstanceData> InstData(Context);
+
+	VectorVM::FExternalFuncRegisterHandler<FNiagaraBool> OutSample(Context);
+
+	checkfSlow(InstData.Get(), TEXT("Export data interface has invalid instance data. %s"), *GetPathName());
+
+	for (int32 i = 0; i < Context.NumInstances; ++i)
+	{
+		FNiagaraBool ShouldStore = StoreDataParam.GetAndAdvance();
+		FBasicParticleData Data;
+		Data.Position = FVector(PositionParamX.GetAndAdvance(), PositionParamY.GetAndAdvance(), PositionParamZ.GetAndAdvance());
+		Data.Size = SizeParam.GetAndAdvance();
+		Data.Position = FVector(VelocityParamX.GetAndAdvance(), VelocityParamY.GetAndAdvance(), VelocityParamZ.GetAndAdvance());
+
+		FNiagaraBool Valid;
+		if (ShouldStore && InstData->UserParamBinding.BoundVariable.IsValid() && InstData->CallbackHandler)
+		{
+			InstData->GatheredData.Add(Data);
+			Valid.SetValue(true);
+		}
+		*OutSample.GetDestAndAdvance() = Valid;
+	}
+}
+
+bool UNiagaraDataInterfaceExport::CopyToInternal(UNiagaraDataInterface* Destination) const
+{
+	if (!Super::CopyToInternal(Destination))
+	{
+		return false;
+	}
+
+	UNiagaraDataInterfaceExport* OtherTyped = CastChecked<UNiagaraDataInterfaceExport>(Destination);
+	OtherTyped->CallbackHandlerParameter = CallbackHandlerParameter;
+	return true;
+}
