[NO NON-PII COMMIT MESSAGE]

diff --git a/Engine/Plugins/FX/Niagara/Source/Niagara/Classes/NiagaraDataInterfaceExport.h b/Engine/Plugins/FX/Niagara/Source/Niagara/Classes/NiagaraDataInterfaceExport.h
--- a/Engine/Plugins/FX/Niagara/Source/Niagara/Classes/NiagaraDataInterfaceExport.h
+++ b/Engine/Plugins/FX/Niagara/Source/Niagara/Classes/NiagaraDataInterfaceExport.h
@@ -41,7 +41,9 @@ public:
 struct ExportInterface_InstanceData
 {
 	UObject* CallbackHandler = nullptr;
-	TArray<FBasicParticleData> GatheredData;
+
+	/** We use a lock-free queue here because multiple threads might try to push data to it at the same time. */
+	TQueue<FBasicParticleData, EQueueMode::Mpsc> GatheredData;
 
 	/** A binding to the user ptr we're reading the handler from. */
 	FNiagaraParameterDirectBinding<UObject*> UserParamBinding;
diff --git a/Engine/Plugins/FX/Niagara/Source/Niagara/Private/NiagaraDataInterfaceExport.cpp b/Engine/Plugins/FX/Niagara/Source/Niagara/Private/NiagaraDataInterfaceExport.cpp
--- a/Engine/Plugins/FX/Niagara/Source/Niagara/Private/NiagaraDataInterfaceExport.cpp
+++ b/Engine/Plugins/FX/Niagara/Source/Niagara/Private/NiagaraDataInterfaceExport.cpp
@@ -52,11 +52,17 @@ bool UNiagaraDataInterfaceExport::PerInstanceTick(void* PerInstanceData, FNiagar
 bool UNiagaraDataInterfaceExport::PerInstanceTickPostSimulate(void* PerInstanceData, FNiagaraSystemInstance* SystemInstance, float DeltaSeconds)
 {
 	ExportInterface_InstanceData* PIData = (ExportInterface_InstanceData*) PerInstanceData;
-	if (PIData->GatheredData.Num() > 0 && PIData->CallbackHandler && PIData->CallbackHandler->GetClass()->ImplementsInterface(UNiagaraParticleCallbackHandler::StaticClass()))
+	if (!PIData->GatheredData.IsEmpty() && PIData->CallbackHandler && PIData->CallbackHandler->GetClass()->ImplementsInterface(UNiagaraParticleCallbackHandler::StaticClass()))
 	{
-		INiagaraParticleCallbackHandler::Execute_ReceiveParticleData(PIData->CallbackHandler, PIData->GatheredData, SystemInstance->GetSystem());
+		//Drain the queue into an array here
+		TArray<FBasicParticleData> Data;
+		FBasicParticleData Value;
+		while (PIData->GatheredData.Dequeue(Value))
+		{
+			Data.Add(Value);
+		}
+		INiagaraParticleCallbackHandler::Execute_ReceiveParticleData(PIData->CallbackHandler, Data, SystemInstance->GetSystem());
 	}
-	PIData->GatheredData.Empty();
 	return false;
 }
 
@@ -125,6 +131,7 @@ void UNiagaraDataInterfaceExport::StoreData(FVectorVMContext& Context)
 	VectorVM::FExternalFuncRegisterHandler<FNiagaraBool> OutSample(Context);
 
 	checkfSlow(InstData.Get(), TEXT("Export data interface has invalid instance data. %s"), *GetPathName());
+	bool ValidHandlerData = InstData->UserParamBinding.BoundVariable.IsValid() && InstData->CallbackHandler;
 
 	for (int32 i = 0; i < Context.NumInstances; ++i)
 	{
@@ -135,10 +142,9 @@ void UNiagaraDataInterfaceExport::StoreData(FVectorVMContext& Context)
 		Data.Position = FVector(VelocityParamX.GetAndAdvance(), VelocityParamY.GetAndAdvance(), VelocityParamZ.GetAndAdvance());
 
 		FNiagaraBool Valid;
-		if (ShouldStore && InstData->UserParamBinding.BoundVariable.IsValid() && InstData->CallbackHandler)
+		if (ValidHandlerData && ShouldStore)
 		{
-			InstData->GatheredData.Add(Data);
-			Valid.SetValue(true);
+			Valid.SetValue(InstData->GatheredData.Enqueue(Data));
 		}
 		*OutSample.GetDestAndAdvance() = Valid;
 	}
