Fixed iOS compile error

diff --git a/Engine/Plugins/FX/Niagara/Source/Niagara/Private/NiagaraDataInterfaceExport.cpp b/Engine/Plugins/FX/Niagara/Source/Niagara/Private/NiagaraDataInterfaceExport.cpp
--- a/Engine/Plugins/FX/Niagara/Source/Niagara/Private/NiagaraDataInterfaceExport.cpp
+++ b/Engine/Plugins/FX/Niagara/Source/Niagara/Private/NiagaraDataInterfaceExport.cpp
@@ -3,7 +3,6 @@
 #include "NiagaraDataInterfaceExport.h"
 #include "NiagaraTypes.h"
 #include "NiagaraCustomVersion.h"
-#include "..\Classes\NiagaraDataInterfaceExport.h"
 #include "Internationalization/Internationalization.h"
 #include "NiagaraSystemInstance.h"
 
