Fixed warning when compiling niagara vertex factories.

diff --git a/Engine/Plugins/FX/Niagara/Shaders/Private/NiagaraMeshVertexFactory.ush b/Engine/Plugins/FX/Niagara/Shaders/Private/NiagaraMeshVertexFactory.ush
--- a/Engine/Plugins/FX/Niagara/Shaders/Private/NiagaraMeshVertexFactory.ush
+++ b/Engine/Plugins/FX/Niagara/Shaders/Private/NiagaraMeshVertexFactory.ush
@@ -18,7 +18,7 @@ float3 TransformPosition(float3 InPosition)
 {
 	if (NiagaraMeshVF.bLocalSpace)
 	{
-		return TransformLocalToWorld(InPosition);
+		return TransformLocalToWorld(InPosition).xyz;
 	}
 	else
 	{
diff --git a/Engine/Plugins/FX/Niagara/Shaders/Private/NiagaraRibbonVertexFactory.ush b/Engine/Plugins/FX/Niagara/Shaders/Private/NiagaraRibbonVertexFactory.ush
--- a/Engine/Plugins/FX/Niagara/Shaders/Private/NiagaraRibbonVertexFactory.ush
+++ b/Engine/Plugins/FX/Niagara/Shaders/Private/NiagaraRibbonVertexFactory.ush
@@ -18,7 +18,7 @@ float3 TransformPosition(float3 InPosition)
 {
 	if (NiagaraRibbonVF.bLocalSpace)
 	{
-		return TransformLocalToWorld(InPosition);
+		return TransformLocalToWorld(InPosition).xyz;
 	}
 	else
 	{
diff --git a/Engine/Plugins/FX/Niagara/Shaders/Private/NiagaraSpriteVertexFactory.ush b/Engine/Plugins/FX/Niagara/Shaders/Private/NiagaraSpriteVertexFactory.ush
--- a/Engine/Plugins/FX/Niagara/Shaders/Private/NiagaraSpriteVertexFactory.ush
+++ b/Engine/Plugins/FX/Niagara/Shaders/Private/NiagaraSpriteVertexFactory.ush
@@ -597,7 +597,7 @@ float3 TransformPosition(float3 InPosition)
 {
 	if (NiagaraSpriteVF.bLocalSpace)
 	{
-		return TransformLocalToWorld(InPosition);
+		return TransformLocalToWorld(InPosition).xyz;
 	}
 	else
 	{
