Removed ensures from FNiagaraRibbonVertexFactoryShaderParametersVS as some shader might not be using the GetVertexFactoryIntermediates()

diff --git a/Engine/Plugins/FX/Niagara/Source/NiagaraVertexFactories/Private/NiagaraRibbonVertexFactory.cpp b/Engine/Plugins/FX/Niagara/Source/NiagaraVertexFactories/Private/NiagaraRibbonVertexFactory.cpp
--- a/Engine/Plugins/FX/Niagara/Source/NiagaraVertexFactories/Private/NiagaraRibbonVertexFactory.cpp
+++ b/Engine/Plugins/FX/Niagara/Source/NiagaraVertexFactories/Private/NiagaraRibbonVertexFactory.cpp
@@ -42,13 +42,6 @@ public:
 		MultiRibbonIndices.Bind(ParameterMap, TEXT("MultiRibbonIndices"));
 		PackedPerRibbonDataByIndex.Bind(ParameterMap, TEXT("PackedPerRibbonDataByIndex"));
 		FacingMode.Bind(ParameterMap, TEXT("FacingMode"));
-
-		ensure(NiagaraParticleDataFloat.IsBound());
-		ensure(FloatDataOffset.IsBound());
-		ensure(FloatDataStride.IsBound());
-		ensure(SortedIndices.IsBound());
-		ensure(SortedIndicesOffset.IsBound());
-		ensure(FacingMode.IsBound());
 	}
 
 	virtual void Serialize(FArchive& Ar) override
