Working to resolve CI static analysis warning.

diff --git a/Engine/Plugins/FX/Niagara/Source/Niagara/Private/NiagaraDataInterfaceSkeletalMesh.cpp b/Engine/Plugins/FX/Niagara/Source/Niagara/Private/NiagaraDataInterfaceSkeletalMesh.cpp
--- a/Engine/Plugins/FX/Niagara/Source/Niagara/Private/NiagaraDataInterfaceSkeletalMesh.cpp
+++ b/Engine/Plugins/FX/Niagara/Source/Niagara/Private/NiagaraDataInterfaceSkeletalMesh.cpp
@@ -343,6 +343,22 @@ FSkeletalMeshGpuSpawnStaticBuffers::~FSkeletalMeshGpuSpawnStaticBuffers()
 
 void FSkeletalMeshGpuSpawnStaticBuffers::Initialise(FNDISkeletalMesh_InstanceData* InstData, const FSkeletalMeshLODRenderData& SkeletalMeshLODRenderData, const FSkeletalMeshSamplingLODBuiltData& MeshSamplingLODBuiltData)
 {
+	if (!InstData)
+	{
+		SkeletalMeshSamplingLODBuiltData = nullptr;
+		bUseGpuUniformlyDistributedSampling = false;
+
+		LODRenderData = nullptr;
+		TriangleCount = 0;
+		VertexCount = 0;
+
+		NumSpecificBones = 0;
+		SpecificBonesArray.Empty();
+		NumSpecificSockets = 0;
+		SpecificSocketBoneOffset = 0;
+		return;
+	}
+
 	SkeletalMeshSamplingLODBuiltData = &MeshSamplingLODBuiltData;
 	bUseGpuUniformlyDistributedSampling = InstData->bIsGpuUniformlyDistributedSampling;
 
