- Collect garbage after closing Niagara editor
- This allows all the buffers to be deleted reducing RHI memory pressure in the editor

diff --git a/Engine/Plugins/FX/Niagara/Source/NiagaraEditor/Private/Toolkits/NiagaraSystemToolkit.cpp b/Engine/Plugins/FX/Niagara/Source/NiagaraEditor/Private/Toolkits/NiagaraSystemToolkit.cpp
--- a/Engine/Plugins/FX/Niagara/Source/NiagaraEditor/Private/Toolkits/NiagaraSystemToolkit.cpp
+++ b/Engine/Plugins/FX/Niagara/Source/NiagaraEditor/Private/Toolkits/NiagaraSystemToolkit.cpp
@@ -1236,9 +1236,11 @@ bool FNiagaraSystemToolkit::OnRequestClose()
 				return false;
 			}
 		}
+		GEngine->ForceGarbageCollection(true);
 		return true;
 	}
 	
+	GEngine->ForceGarbageCollection(true);
 	return FAssetEditorToolkit::OnRequestClose();
 }
 
