- Remove true from WaitForAsyncTick in OnEndOfFrameUpdateDuringTick
- We can expect to wait here if we fire a RTT and ticking is not complete

#rnx

diff --git a/Engine/Plugins/FX/Niagara/Source/Niagara/Private/NiagaraComponent.cpp b/Engine/Plugins/FX/Niagara/Source/Niagara/Private/NiagaraComponent.cpp
--- a/Engine/Plugins/FX/Niagara/Source/Niagara/Private/NiagaraComponent.cpp
+++ b/Engine/Plugins/FX/Niagara/Source/Niagara/Private/NiagaraComponent.cpp
@@ -987,7 +987,7 @@ void UNiagaraComponent::OnEndOfFrameUpdateDuringTick()
 	Super::OnEndOfFrameUpdateDuringTick();
 	if ( SystemInstance )
 	{
-		SystemInstance->WaitForAsyncTick(true);
+		SystemInstance->WaitForAsyncTick();
 	}
 }
 
