Niagara - Drag copy modules

diff --git a/Engine/Plugins/FX/Niagara/Source/NiagaraEditor/Private/NiagaraEditorUtilities.cpp b/Engine/Plugins/FX/Niagara/Source/NiagaraEditor/Private/NiagaraEditorUtilities.cpp
--- a/Engine/Plugins/FX/Niagara/Source/NiagaraEditor/Private/NiagaraEditorUtilities.cpp
+++ b/Engine/Plugins/FX/Niagara/Source/NiagaraEditor/Private/NiagaraEditorUtilities.cpp
@@ -4,6 +4,8 @@
 #include "NiagaraEditorModule.h"
 #include "INiagaraEditorTypeUtilities.h"
 #include "NiagaraNodeInput.h"
+#include "NiagaraNodeFunctionCall.h"
+#include "NiagaraNodeParameterMapSet.h"
 #include "NiagaraDataInterface.h"
 #include "NiagaraComponent.h"
 #include "Modules/ModuleManager.h"
@@ -33,6 +35,7 @@
 #include "Misc/FileHelper.h"
 #include "EdGraph/EdGraphPin.h"
 #include "NiagaraNodeWriteDataSet.h"
+#include "ViewModels/Stack/NiagaraParameterHandle.h"
 #include "NiagaraNodeStaticSwitch.h"
 #include "NiagaraNodeFunctionCall.h"
 #include "NiagaraParameterMapHistory.h"
@@ -141,7 +144,7 @@ void FNiagaraEditorUtilities::GetParameterVariablesFromSystem(UNiagaraSystem& Sy
 }
 
 // TODO: This is overly complicated.
-void FNiagaraEditorUtilities::FixUpPastedInputNodes(UEdGraph* Graph, TSet<UEdGraphNode*> PastedNodes)
+void FNiagaraEditorUtilities::FixUpPastedNodes(UEdGraph* Graph, TSet<UEdGraphNode*> PastedNodes)
 {
 	// Collect existing inputs.
 	TArray<UNiagaraNodeInput*> CurrentInputs;
@@ -242,6 +245,54 @@ void FNiagaraEditorUtilities::FixUpPastedInputNodes(UEdGraph* Graph, TSet<UEdGra
 			}
 		}
 	}
+
+	// Fix up pasted function call nodes
+	TArray<UNiagaraNodeFunctionCall*> FunctionCallNodes;
+	Graph->GetNodesOfClass<UNiagaraNodeFunctionCall>(FunctionCallNodes);
+	TSet<FName> ExistingNames;
+	for (UNiagaraNodeFunctionCall* FunctionCallNode : FunctionCallNodes)
+	{
+		if (PastedNodes.Contains(FunctionCallNode) == false)
+		{
+			ExistingNames.Add(*FunctionCallNode->GetFunctionName());
+		}
+	}
+
+	TMap<FName, FName> OldFunctionToNewFunctionNameMap;
+	for (UEdGraphNode* PastedNode : PastedNodes)
+	{
+		UNiagaraNodeFunctionCall* PastedFunctionCallNode = Cast<UNiagaraNodeFunctionCall>(PastedNode);
+		if (PastedFunctionCallNode != nullptr)
+		{
+			if (ExistingNames.Contains(*PastedFunctionCallNode->GetFunctionName()))
+			{
+				FName FunctionCallName = *PastedFunctionCallNode->GetFunctionName();
+				FName UniqueFunctionCallName = FNiagaraUtilities::GetUniqueName(FunctionCallName, ExistingNames);
+				PastedFunctionCallNode->SuggestName(UniqueFunctionCallName.ToString());
+				ExistingNames.Add(UniqueFunctionCallName);
+				OldFunctionToNewFunctionNameMap.Add(FunctionCallName, UniqueFunctionCallName);
+			}
+		}
+	}
+
+	for (UEdGraphNode* PastedNode : PastedNodes)
+	{
+		UNiagaraNodeParameterMapSet* ParameterMapSetNode = Cast<UNiagaraNodeParameterMapSet>(PastedNode);
+		if (ParameterMapSetNode != nullptr)
+		{
+			TArray<UEdGraphPin*> InputPins;
+			ParameterMapSetNode->GetInputPins(InputPins);
+			for (UEdGraphPin* InputPin : InputPins)
+			{
+				FNiagaraParameterHandle InputHandle(InputPin->PinName);
+				if (OldFunctionToNewFunctionNameMap.Contains(InputHandle.GetNamespace()))
+				{
+					// Rename any inputs pins on parameter map sets who's function calls were renamed.
+					InputPin->PinName = FNiagaraParameterHandle(OldFunctionToNewFunctionNameMap[InputHandle.GetNamespace()], InputHandle.GetName()).GetParameterHandleString();
+				}
+			}
+		}
+	}
 }
 
 void FNiagaraEditorUtilities::WriteTextFileToDisk(FString SaveDirectory, FString FileName, FString TextToSave, bool bAllowOverwriting)
diff --git a/Engine/Plugins/FX/Niagara/Source/NiagaraEditor/Private/ViewModels/NiagaraScriptGraphViewModel.cpp b/Engine/Plugins/FX/Niagara/Source/NiagaraEditor/Private/ViewModels/NiagaraScriptGraphViewModel.cpp
--- a/Engine/Plugins/FX/Niagara/Source/NiagaraEditor/Private/ViewModels/NiagaraScriptGraphViewModel.cpp
+++ b/Engine/Plugins/FX/Niagara/Source/NiagaraEditor/Private/ViewModels/NiagaraScriptGraphViewModel.cpp
@@ -297,7 +297,7 @@ void FNiagaraScriptGraphViewModel::PasteNodes()
 				Node->MarkNodeRequiresSynchronization(__FUNCTION__, false);
 		}
 
-		FNiagaraEditorUtilities::FixUpPastedInputNodes(Graph, PastedNodes);
+		FNiagaraEditorUtilities::FixUpPastedNodes(Graph, PastedNodes);
 
 		OnNodesPastedDelegate.Broadcast(PastedNodes);
 
diff --git a/Engine/Plugins/FX/Niagara/Source/NiagaraEditor/Private/ViewModels/Stack/NiagaraStackEntry.cpp b/Engine/Plugins/FX/Niagara/Source/NiagaraEditor/Private/ViewModels/Stack/NiagaraStackEntry.cpp
--- a/Engine/Plugins/FX/Niagara/Source/NiagaraEditor/Private/ViewModels/Stack/NiagaraStackEntry.cpp
+++ b/Engine/Plugins/FX/Niagara/Source/NiagaraEditor/Private/ViewModels/Stack/NiagaraStackEntry.cpp
@@ -340,9 +340,9 @@ bool UNiagaraStackEntry::CanDrag() const
 	return false;
 }
 
-TOptional<UNiagaraStackEntry::FDropResult> UNiagaraStackEntry::CanDrop(const TArray<UNiagaraStackEntry*>& DraggedEntries)
+TOptional<UNiagaraStackEntry::FDropResult> UNiagaraStackEntry::CanDrop(const TArray<UNiagaraStackEntry*>& DraggedEntries, EDragOptions DragOptions)
 {
-	TOptional<FDropResult> Result = CanDropInternal(DraggedEntries);
+	TOptional<FDropResult> Result = CanDropInternal(DraggedEntries, DragOptions);
 	if (Result.IsSet())
 	{
 		return Result;
@@ -350,14 +350,14 @@ TOptional<UNiagaraStackEntry::FDropResult> UNiagaraStackEntry::CanDrop(const TAr
 	else
 	{
 		return OnRequestCanDropDelegate.IsBound() 
-			? OnRequestCanDropDelegate.Execute(*this, DraggedEntries)
+			? OnRequestCanDropDelegate.Execute(*this, DraggedEntries, DragOptions)
 			: TOptional<FDropResult>();
 	}
 }
 
-TOptional<UNiagaraStackEntry::FDropResult> UNiagaraStackEntry::Drop(const TArray<UNiagaraStackEntry*>& DraggedEntries)
+TOptional<UNiagaraStackEntry::FDropResult> UNiagaraStackEntry::Drop(const TArray<UNiagaraStackEntry*>& DraggedEntries, EDragOptions DragOptions)
 {
-	TOptional<FDropResult> Result = DropInternal(DraggedEntries);
+	TOptional<FDropResult> Result = DropInternal(DraggedEntries, DragOptions);
 	if (Result.IsSet())
 	{
 		return Result;
@@ -365,7 +365,7 @@ TOptional<UNiagaraStackEntry::FDropResult> UNiagaraStackEntry::Drop(const TArray
 	else
 	{
 		return OnRequestDropDelegate.IsBound()
-			? OnRequestDropDelegate.Execute(*this, DraggedEntries)
+			? OnRequestDropDelegate.Execute(*this, DraggedEntries, DragOptions)
 			: TOptional<FDropResult>();
 	}
 }
@@ -400,22 +400,22 @@ bool UNiagaraStackEntry::HasBaseEmitter() const
 	return bHasBaseEmitterCache.GetValue();
 }
 
-TOptional<UNiagaraStackEntry::FDropResult> UNiagaraStackEntry::CanDropInternal(const TArray<UNiagaraStackEntry*>& DraggedEntries)
+TOptional<UNiagaraStackEntry::FDropResult> UNiagaraStackEntry::CanDropInternal(const TArray<UNiagaraStackEntry*>& DraggedEntries, EDragOptions DragOptions)
 {
 	return TOptional<FDropResult>();
 }
 
-TOptional<UNiagaraStackEntry::FDropResult> UNiagaraStackEntry::DropInternal(const TArray<UNiagaraStackEntry*>& DraggedEntries)
+TOptional<UNiagaraStackEntry::FDropResult> UNiagaraStackEntry::DropInternal(const TArray<UNiagaraStackEntry*>& DraggedEntries, EDragOptions DragOptions)
 {
 	return TOptional<FDropResult>();
 }
 
-TOptional<UNiagaraStackEntry::FDropResult> UNiagaraStackEntry::ChildRequestCanDropInternal(const UNiagaraStackEntry& TargetChild, const TArray<UNiagaraStackEntry*>& DraggedEntries)
+TOptional<UNiagaraStackEntry::FDropResult> UNiagaraStackEntry::ChildRequestCanDropInternal(const UNiagaraStackEntry& TargetChild, const TArray<UNiagaraStackEntry*>& DraggedEntries, EDragOptions DragOptions)
 {
 	return TOptional<FDropResult>();
 }
 
-TOptional<UNiagaraStackEntry::FDropResult> UNiagaraStackEntry::ChildRequestDropInternal(const UNiagaraStackEntry& TargetChild, const TArray<UNiagaraStackEntry*>& DraggedEntries)
+TOptional<UNiagaraStackEntry::FDropResult> UNiagaraStackEntry::ChildRequestDropInternal(const UNiagaraStackEntry& TargetChild, const TArray<UNiagaraStackEntry*>& DraggedEntries, EDragOptions DragOptions)
 {
 	return TOptional<FDropResult>();
 }
@@ -595,9 +595,9 @@ void UNiagaraStackEntry::ChildRequestFullRefreshDeferred()
 	RequestFullRefreshDeferredDelegate.Broadcast();
 }
 
-TOptional<UNiagaraStackEntry::FDropResult> UNiagaraStackEntry::ChildRequestCanDrop(const UNiagaraStackEntry& TargetChild, const TArray<UNiagaraStackEntry*>& DraggedEntries)
+TOptional<UNiagaraStackEntry::FDropResult> UNiagaraStackEntry::ChildRequestCanDrop(const UNiagaraStackEntry& TargetChild, const TArray<UNiagaraStackEntry*>& DraggedEntries, EDragOptions DragOptions)
 {
-	TOptional<FDropResult> Result = ChildRequestCanDropInternal(TargetChild, DraggedEntries);
+	TOptional<FDropResult> Result = ChildRequestCanDropInternal(TargetChild, DraggedEntries, DragOptions);
 	if (Result.IsSet())
 	{
 		return Result;
@@ -605,14 +605,14 @@ TOptional<UNiagaraStackEntry::FDropResult> UNiagaraStackEntry::ChildRequestCanDr
 	else
 	{
 		return OnRequestCanDropDelegate.IsBound()
-			? OnRequestCanDropDelegate.Execute(TargetChild, DraggedEntries)
+			? OnRequestCanDropDelegate.Execute(TargetChild, DraggedEntries, DragOptions)
 			: TOptional<FDropResult>();
 	}
 }
 
-TOptional<UNiagaraStackEntry::FDropResult> UNiagaraStackEntry::ChildRequestDrop(const UNiagaraStackEntry& TargetChild, const TArray<UNiagaraStackEntry*>& DraggedEntries)
+TOptional<UNiagaraStackEntry::FDropResult> UNiagaraStackEntry::ChildRequestDrop(const UNiagaraStackEntry& TargetChild, const TArray<UNiagaraStackEntry*>& DraggedEntries, EDragOptions DragOptions)
 {
-	TOptional<FDropResult> Result = ChildRequestDropInternal(TargetChild, DraggedEntries);
+	TOptional<FDropResult> Result = ChildRequestDropInternal(TargetChild, DraggedEntries, DragOptions);
 	if (Result.IsSet())
 	{
 		return Result;
@@ -620,7 +620,7 @@ TOptional<UNiagaraStackEntry::FDropResult> UNiagaraStackEntry::ChildRequestDrop(
 	else
 	{
 		return OnRequestDropDelegate.IsBound()
-			? OnRequestDropDelegate.Execute(TargetChild, DraggedEntries)
+			? OnRequestDropDelegate.Execute(TargetChild, DraggedEntries, DragOptions)
 			: TOptional<FDropResult>();
 	}
 }
diff --git a/Engine/Plugins/FX/Niagara/Source/NiagaraEditor/Private/ViewModels/Stack/NiagaraStackGraphUtilities.cpp b/Engine/Plugins/FX/Niagara/Source/NiagaraEditor/Private/ViewModels/Stack/NiagaraStackGraphUtilities.cpp
--- a/Engine/Plugins/FX/Niagara/Source/NiagaraEditor/Private/ViewModels/Stack/NiagaraStackGraphUtilities.cpp
+++ b/Engine/Plugins/FX/Niagara/Source/NiagaraEditor/Private/ViewModels/Stack/NiagaraStackGraphUtilities.cpp
@@ -1768,7 +1768,7 @@ bool FNiagaraStackGraphUtilities::GetStackIssuesRecursively(const UNiagaraStackE
 	return OutIssues.Num() > 0;
 }
 
-void FNiagaraStackGraphUtilities::MoveModule(UNiagaraScript& SourceScript, UNiagaraNodeFunctionCall& ModuleToMove, UNiagaraSystem& TargetSystem, FGuid TargetEmitterHandleId, ENiagaraScriptUsage TargetUsage, FGuid TargetUsageId, int32 TargetModuleIndex)
+void FNiagaraStackGraphUtilities::MoveModule(UNiagaraScript& SourceScript, UNiagaraNodeFunctionCall& ModuleToMove, UNiagaraSystem& TargetSystem, FGuid TargetEmitterHandleId, ENiagaraScriptUsage TargetUsage, FGuid TargetUsageId, int32 TargetModuleIndex, bool bForceCopy)
 {
 	UNiagaraScript* TargetScript = FNiagaraEditorUtilities::GetScriptFromSystem(TargetSystem, TargetEmitterHandleId, TargetUsage, TargetUsageId);
 	checkf(TargetScript != nullptr, TEXT("Target script not found"));
@@ -1785,10 +1785,10 @@ void FNiagaraStackGraphUtilities::MoveModule(UNiagaraScript& SourceScript, UNiag
 	UNiagaraGraph* SourceGraph = ModuleToMove.GetNiagaraGraph();
 	UNiagaraGraph* TargetGraph = TargetOutputNode->GetNiagaraGraph();
 
-	// If the source and target scripts don't match we need to collect the rapid iteration parameter values for each function in the source group
+	// If the source and target scripts don't match, or we're forcing a copy we need to collect the rapid iteration parameter values for each function in the source group
 	// so we can restore them after moving.
 	TMap<FGuid, TArray<FNiagaraVariable>> SourceFunctionIdToRapidIterationParametersMap;
-	if (&SourceScript != TargetScript)
+	if (&SourceScript != TargetScript || bForceCopy)
 	{
 		TMap<FString, FGuid> FunctionCallNameToNodeIdMap;
 		for (UNiagaraNode* SourceGroupNode : SourceGroupNodes)
@@ -1824,14 +1824,14 @@ void FNiagaraStackGraphUtilities::MoveModule(UNiagaraScript& SourceScript, UNiag
 	FStackNodeGroup TargetGroup;
 	TArray<UNiagaraNode*> TargetGroupNodes;
 	TMap<FGuid, FGuid> OldNodeIdToNewIdMap;
-	if (SourceGraph == TargetGraph)
+	if (SourceGraph == TargetGraph && bForceCopy == false)
 	{
 		TargetGroup = SourceGroups[SourceGroupIndex];
 		TargetGroupNodes = SourceGroupNodes;
 	}
 	else 
 	{
-		// If the module is being inserted into a different graph all of the nodes need to be copied into the target graph.
+		// If the module is being inserted into a different graph, or it's being copied, all of the nodes need to be copied into the target graph.
 		FStackNodeGroup SourceGroup = SourceGroups[SourceGroupIndex];
 		
 		TSet<UObject*> NodesToCopy;
@@ -1876,6 +1876,8 @@ void FNiagaraStackGraphUtilities::MoveModule(UNiagaraScript& SourceScript, UNiag
 				CopiedNiagaraNode->MarkNodeRequiresSynchronization(__FUNCTION__, false);
 			}
 		}
+
+		FNiagaraEditorUtilities::FixUpPastedNodes(TargetGraph, CopiedNodesSet);
 	}
 
 	TArray<FStackNodeGroup> TargetGroups;
@@ -1885,17 +1887,21 @@ void FNiagaraStackGraphUtilities::MoveModule(UNiagaraScript& SourceScript, UNiag
 	// if a valid index wasn't supplied, than we insert at the end.
 	int32 TargetGroupIndex = TargetModuleIndex != INDEX_NONE ? TargetModuleIndex + 1 : TargetGroups.Num() - 1;
 
-	// Remove the source module group from it's stack, and insert the source or copied nodes into the target stack.
-	DisconnectStackNodeGroup(SourceGroups[SourceGroupIndex], SourceGroups[SourceGroupIndex - 1], SourceGroups[SourceGroupIndex + 1]);
-	if (SourceGraph != TargetGraph)
+	// If we're not forcing a copy of the moved module, remove the source module group from it's stack.
+	if (bForceCopy == false)
 	{
-		// If the graphs were different also remove the nodes from the source graph.
-		for (UNiagaraNode* SourceGroupNode : SourceGroupNodes)
+		DisconnectStackNodeGroup(SourceGroups[SourceGroupIndex], SourceGroups[SourceGroupIndex - 1], SourceGroups[SourceGroupIndex + 1]);
+		if (SourceGraph != TargetGraph)
 		{
-			SourceGraph->RemoveNode(SourceGroupNode);
+			// If the graphs were different also remove the nodes from the source graph.
+			for (UNiagaraNode* SourceGroupNode : SourceGroupNodes)
+			{
+				SourceGraph->RemoveNode(SourceGroupNode);
+			}
 		}
 	}
 
+	// Insert the source or copied nodes into the target stack.
 	ConnectStackNodeGroup(TargetGroup, TargetGroups[TargetGroupIndex - 1], TargetGroups[TargetGroupIndex]);
 
 	// Copy any rapid iteration parameters cached earlier into the target script.
@@ -1903,9 +1909,9 @@ void FNiagaraStackGraphUtilities::MoveModule(UNiagaraScript& SourceScript, UNiag
 	{
 		SourceScript.Modify();
 		TargetScript->Modify();
-		if (SourceGraph == TargetGraph)
+		if (SourceGraph == TargetGraph && bForceCopy == false)
 		{
-			// If the module was dropped in the same graph than neither the emitter or function call name could have changed
+			// If we're not copying and if the module was dropped in the same graph than neither the emitter or function call name could have changed
 			// so we can just add them directly to the target script.
 			for (auto It = SourceFunctionIdToRapidIterationParametersMap.CreateIterator(); It; ++It)
 			{
@@ -1919,7 +1925,7 @@ void FNiagaraStackGraphUtilities::MoveModule(UNiagaraScript& SourceScript, UNiag
 		}
 		else
 		{
-			// If the module was moved to a different graph it's possible that the emitter name or function call name could have
+			// If we're copying or the module was moved to a different graph it's possible that the emitter name or function call name could have
 			// changed so we need to construct new rapid iteration parameters.
 			FString EmitterName;
 			if (TargetEmitterHandleId.IsValid())
@@ -1929,34 +1935,26 @@ void FNiagaraStackGraphUtilities::MoveModule(UNiagaraScript& SourceScript, UNiag
 				EmitterName = TargetEmitterHandle->GetUniqueInstanceName();
 			}
 
-			TArray<FNiagaraVariable> TargetRapidIterationParameters;
 			for (auto It = SourceFunctionIdToRapidIterationParametersMap.CreateIterator(); It; ++It)
 			{
 				FGuid& FunctionId = It.Key();
 				TArray<FNiagaraVariable>& RapidIterationParameters = It.Value();
 
-				if (SourceGraph == TargetGraph)
-				{
-					TargetRapidIterationParameters.Append(RapidIterationParameters);
-				}
-				else
+				FGuid TargetNodeId = OldNodeIdToNewIdMap[FunctionId];
+				UNiagaraNode** TargetFunctionNodePtr = TargetGroupNodes.FindByPredicate(
+					[TargetNodeId](UNiagaraNode* TargetGroupNode) { return TargetGroupNode->NodeGuid == TargetNodeId; });
+				checkf(TargetFunctionNodePtr != nullptr, TEXT("Target nodes not copied correctly"));
+				UNiagaraNodeFunctionCall* TargetFunctionNode = CastChecked<UNiagaraNodeFunctionCall>(*TargetFunctionNodePtr);
+				for (FNiagaraVariable& RapidIterationParameter : RapidIterationParameters)
 				{
-					FGuid TargetNodeId = OldNodeIdToNewIdMap[FunctionId];
-					UNiagaraNode** TargetFunctionNodePtr = TargetGroupNodes.FindByPredicate(
-						[TargetNodeId](UNiagaraNode* TargetGroupNode) { return TargetGroupNode->NodeGuid == TargetNodeId; });
-					checkf(TargetFunctionNodePtr != nullptr, TEXT("Target nodes not copied correctly"));
-					UNiagaraNodeFunctionCall* TargetFunctionNode = CastChecked<UNiagaraNodeFunctionCall>(*TargetFunctionNodePtr);
-					for (FNiagaraVariable& RapidIterationParameter : RapidIterationParameters)
-					{
-						FString OldEmitterName;
-						FString OldFunctionCallName;
-						FString InputName;
-						FNiagaraParameterMapHistory::SplitRapidIterationParameterName(RapidIterationParameter, OldEmitterName, OldFunctionCallName, InputName);
-						FNiagaraParameterHandle ModuleHandle = FNiagaraParameterHandle::CreateModuleParameterHandle(*InputName);
-						FNiagaraParameterHandle AliasedModuleHandle = FNiagaraParameterHandle::CreateAliasedModuleParameterHandle(ModuleHandle, TargetFunctionNode);
-						FNiagaraVariable TargetRapidIterationParameter = CreateRapidIterationParameter(EmitterName, TargetUsage, AliasedModuleHandle.GetParameterHandleString(), RapidIterationParameter.GetType());
-						TargetScript->RapidIterationParameters.SetParameterData(RapidIterationParameter.GetData(), TargetRapidIterationParameter, true);
-					}
+					FString OldEmitterName;
+					FString OldFunctionCallName;
+					FString InputName;
+					FNiagaraParameterMapHistory::SplitRapidIterationParameterName(RapidIterationParameter, OldEmitterName, OldFunctionCallName, InputName);
+					FNiagaraParameterHandle ModuleHandle = FNiagaraParameterHandle::CreateModuleParameterHandle(*InputName);
+					FNiagaraParameterHandle AliasedModuleHandle = FNiagaraParameterHandle::CreateAliasedModuleParameterHandle(ModuleHandle, TargetFunctionNode);
+					FNiagaraVariable TargetRapidIterationParameter = CreateRapidIterationParameter(EmitterName, TargetUsage, AliasedModuleHandle.GetParameterHandleString(), RapidIterationParameter.GetType());
+					TargetScript->RapidIterationParameters.SetParameterData(RapidIterationParameter.GetData(), TargetRapidIterationParameter, true);
 				}
 			}
 		}
diff --git a/Engine/Plugins/FX/Niagara/Source/NiagaraEditor/Private/ViewModels/Stack/NiagaraStackModuleItem.cpp b/Engine/Plugins/FX/Niagara/Source/NiagaraEditor/Private/ViewModels/Stack/NiagaraStackModuleItem.cpp
--- a/Engine/Plugins/FX/Niagara/Source/NiagaraEditor/Private/ViewModels/Stack/NiagaraStackModuleItem.cpp
+++ b/Engine/Plugins/FX/Niagara/Source/NiagaraEditor/Private/ViewModels/Stack/NiagaraStackModuleItem.cpp
@@ -511,7 +511,7 @@ void UNiagaraStackModuleItem::RefreshIssues(TArray<FStackIssue>& NewIssues)
 						int32 CorrectIndex = Dependency.Type == ENiagaraModuleDependencyType::PostDependency ? DisorderedNode.Index : DisorderedNode.Index + 1;
 						checkf(ModuleIndex != INDEX_NONE, TEXT("Module data wasn't found in system for current module!"));
 						UNiagaraScript& OwningScript = *FNiagaraEditorUtilities::GetScriptFromSystem(GetSystemViewModel()->GetSystem(), SystemModuleData[ModuleIndex].EmitterHandleId, SystemModuleData[ModuleIndex].Usage, SystemModuleData[ModuleIndex].UsageId);
-						FNiagaraStackGraphUtilities::MoveModule(OwningScript, *FunctionCallNode, GetSystemViewModel()->GetSystem(), DisorderedNode.EmitterHandleId, DisorderedNode.Usage, DisorderedNode.UsageId, CorrectIndex);
+						FNiagaraStackGraphUtilities::MoveModule(OwningScript, *FunctionCallNode, GetSystemViewModel()->GetSystem(), DisorderedNode.EmitterHandleId, DisorderedNode.Usage, DisorderedNode.UsageId, CorrectIndex, false);
 						// enable if needed
 						if (bNeedsEnable)
 						{
diff --git a/Engine/Plugins/FX/Niagara/Source/NiagaraEditor/Private/ViewModels/Stack/NiagaraStackScriptItemGroup.cpp b/Engine/Plugins/FX/Niagara/Source/NiagaraEditor/Private/ViewModels/Stack/NiagaraStackScriptItemGroup.cpp
--- a/Engine/Plugins/FX/Niagara/Source/NiagaraEditor/Private/ViewModels/Stack/NiagaraStackScriptItemGroup.cpp
+++ b/Engine/Plugins/FX/Niagara/Source/NiagaraEditor/Private/ViewModels/Stack/NiagaraStackScriptItemGroup.cpp
@@ -552,14 +552,14 @@ void GenerateDragDropData(
 	}
 }
 
-TOptional<UNiagaraStackEntry::FDropResult> UNiagaraStackScriptItemGroup::ChildRequestCanDropInternal(const UNiagaraStackEntry& TargetChild, const TArray<UNiagaraStackEntry*>& DraggedEntries)
+TOptional<UNiagaraStackEntry::FDropResult> UNiagaraStackScriptItemGroup::ChildRequestCanDropInternal(const UNiagaraStackEntry& TargetChild, const TArray<UNiagaraStackEntry*>& DraggedEntries, EDragOptions DragOptions)
 {
-	if (bIsValidForOutput && DraggedEntries.Num() == 1)
+	if (bIsValidForOutput && DraggedEntries.Num() == 1 && (DragOptions == UNiagaraStackEntry::EDragOptions::Move || DragOptions == UNiagaraStackEntry::EDragOptions::Copy))
 	{
 		UNiagaraStackModuleItem* SourceModuleItem = Cast<UNiagaraStackModuleItem>(DraggedEntries[0]);
 		if (SourceModuleItem != nullptr)
 		{
-			if (SourceModuleItem->CanMoveAndDelete() == false)
+			if (DragOptions == EDragOptions::Move && SourceModuleItem->CanMoveAndDelete() == false)
 			{
 				return FDropResult(false, LOCTEXT("CantMoveModuleError", "This inherited module can't be moved."));
 			}
@@ -570,11 +570,6 @@ TOptional<UNiagaraStackEntry::FDropResult> UNiagaraStackScriptItemGroup::ChildRe
 				return FDropResult(false, LOCTEXT("CantMoveByUsage", "This module can't be moved to this section of the stack because its it's not supported in this context."));
 			}
 
-			if (SourceModuleItem->GetSystemViewModel() != TargetChild.GetSystemViewModel())
-			{
-				return FDropResult(false, LOCTEXT("CantMoveFromAnotherSystem", "This module can't be moved into this system from a different system."));
-			}
-
 			const UNiagaraStackSpacer* TargetSpacer = Cast<UNiagaraStackSpacer>(&TargetChild);
 			if (TargetSpacer != nullptr)
 			{
@@ -594,9 +589,12 @@ TOptional<UNiagaraStackEntry::FDropResult> UNiagaraStackScriptItemGroup::ChildRe
 
 					// Make sure the source and target indices are within safe ranges, and make sure that the insert target isn't the source target or the spot directly
 					// after the source target since that won't actually move the module.
-					if (SourceGroupIndex > 0 && SourceGroupIndex < SourceStackGroups.Num() - 1 && TargetGroupIndex > 0 && TargetGroupIndex < TargetStackGroups.Num() &&
-						SourceStackGroups[SourceGroupIndex].EndNode != TargetStackGroups[TargetGroupIndex].EndNode &&
-						SourceStackGroups[SourceGroupIndex].EndNode != TargetStackGroups[TargetGroupIndex - 1].EndNode)
+					bool bSourceWithinRange = SourceGroupIndex > 0 && SourceGroupIndex < SourceStackGroups.Num() - 1;
+					bool bTargetWithinRange = TargetGroupIndex > 0 && TargetGroupIndex < TargetStackGroups.Num();
+					bool bWillMove = DragOptions == EDragOptions::Copy ||
+						(SourceStackGroups[SourceGroupIndex].EndNode != TargetStackGroups[TargetGroupIndex].EndNode &&
+						SourceStackGroups[SourceGroupIndex].EndNode != TargetStackGroups[TargetGroupIndex - 1].EndNode);
+					if(bSourceWithinRange && bTargetWithinRange && bWillMove)
 					{
 						return FDropResult(true, LOCTEXT("MoveModuleResult", "Move this module here."));
 					}
@@ -607,8 +605,13 @@ TOptional<UNiagaraStackEntry::FDropResult> UNiagaraStackScriptItemGroup::ChildRe
 	return TOptional<FDropResult>();
 }
 
-TOptional<UNiagaraStackEntry::FDropResult> UNiagaraStackScriptItemGroup::ChildRequestDropInternal(const UNiagaraStackEntry& TargetChild, const TArray<UNiagaraStackEntry*>& DraggedEntries)
+TOptional<UNiagaraStackEntry::FDropResult> UNiagaraStackScriptItemGroup::ChildRequestDropInternal(const UNiagaraStackEntry& TargetChild, const TArray<UNiagaraStackEntry*>& DraggedEntries, EDragOptions DragOptions)
 {
+	if (DragOptions != EDragOptions::Move && DragOptions != EDragOptions::Copy)
+	{
+		return TOptional<UNiagaraStackEntry::FDropResult>();
+	}
+
 	if (bIsValidForOutput && DraggedEntries.Num() == 1 && DraggedEntries[0]->IsA<UNiagaraStackModuleItem>())
 	{
 		UNiagaraStackModuleItem* SourceModuleItem = CastChecked<UNiagaraStackModuleItem>(DraggedEntries[0]);
@@ -651,7 +654,7 @@ TOptional<UNiagaraStackEntry::FDropResult> UNiagaraStackScriptItemGroup::ChildRe
 					FScopedTransaction ScopedTransaction(LOCTEXT("DragAndDropModule", "Drag and drop module"));
 					{
 						FNiagaraStackGraphUtilities::MoveModule(*SourceModuleScript, SourceModuleItem->GetModuleNode(), GetSystemViewModel()->GetSystem(), TargetEmitterHandleId,
-							ScriptUsage, ScriptUsageId, TargetIndex);
+							ScriptUsage, ScriptUsageId, TargetIndex, DragOptions == EDragOptions::Copy);
 
 						UNiagaraGraph* TargetGraph = ScriptViewModel.Pin()->GetGraphViewModel()->GetGraph();
 						FNiagaraStackGraphUtilities::RelayoutGraph(*TargetGraph);
diff --git a/Engine/Plugins/FX/Niagara/Source/NiagaraEditor/Public/NiagaraEditorUtilities.h b/Engine/Plugins/FX/Niagara/Source/NiagaraEditor/Public/NiagaraEditorUtilities.h
--- a/Engine/Plugins/FX/Niagara/Source/NiagaraEditor/Public/NiagaraEditorUtilities.h
+++ b/Engine/Plugins/FX/Niagara/Source/NiagaraEditor/Public/NiagaraEditorUtilities.h
@@ -108,7 +108,7 @@ namespace FNiagaraEditorUtilities
 	void GetParameterVariablesFromSystem(UNiagaraSystem& System, TArray<FNiagaraVariable>& ParameterVariables, FGetParameterVariablesFromSystemOptions Options = FGetParameterVariablesFromSystemOptions());
 
 	/** Helper to clean up copy & pasted graphs.*/
-	void FixUpPastedInputNodes(UEdGraph* Graph, TSet<UEdGraphNode*> PastedNodes);
+	void FixUpPastedNodes(UEdGraph* Graph, TSet<UEdGraphNode*> PastedNodes);
 
 	/** Helper to convert compile status to text.*/
 	FText StatusToText(ENiagaraScriptCompileStatus Status);
diff --git a/Engine/Plugins/FX/Niagara/Source/NiagaraEditor/Public/ViewModels/Stack/NiagaraStackEntry.h b/Engine/Plugins/FX/Niagara/Source/NiagaraEditor/Public/ViewModels/Stack/NiagaraStackEntry.h
--- a/Engine/Plugins/FX/Niagara/Source/NiagaraEditor/Public/ViewModels/Stack/NiagaraStackEntry.h
+++ b/Engine/Plugins/FX/Niagara/Source/NiagaraEditor/Public/ViewModels/Stack/NiagaraStackEntry.h
@@ -36,11 +36,18 @@ public:
 		const FText DropMessage;
 	};
 
+	enum class EDragOptions
+	{
+		Copy,
+		Move,
+		None
+	};
+
 	DECLARE_MULTICAST_DELEGATE(FOnStructureChanged);
 	DECLARE_MULTICAST_DELEGATE_OneParam(FOnDataObjectModified, UObject*);
 	DECLARE_MULTICAST_DELEGATE(FOnRequestFullRefresh);
 	DECLARE_MULTICAST_DELEGATE(FOnRequestFullRefreshDeferred);
-	DECLARE_DELEGATE_RetVal_TwoParams(TOptional<FDropResult>, FOnRequestDrop, const UNiagaraStackEntry&, const TArray<UNiagaraStackEntry*>&);
+	DECLARE_DELEGATE_RetVal_ThreeParams(TOptional<FDropResult>, FOnRequestDrop, const UNiagaraStackEntry&, const TArray<UNiagaraStackEntry*>&, EDragOptions);
 	DECLARE_DELEGATE_RetVal_OneParam(bool, FOnFilterChild, const UNiagaraStackEntry&);
 	DECLARE_DELEGATE(FStackIssueFixDelegate);
 	DECLARE_DELEGATE(FOnRequestDeferredOperation);
@@ -243,9 +250,9 @@ public:
 
 	virtual bool CanDrag() const;
 
-	TOptional<FDropResult> CanDrop(const TArray<UNiagaraStackEntry*>& DraggedEntries);
+	TOptional<FDropResult> CanDrop(const TArray<UNiagaraStackEntry*>& DraggedEntries, EDragOptions DragOptions);
 
-	TOptional<FDropResult> Drop(const TArray<UNiagaraStackEntry*>& DraggedEntries);
+	TOptional<FDropResult> Drop(const TArray<UNiagaraStackEntry*>& DraggedEntries, EDragOptions DragOptions);
 
 	void SetOnRequestCanDrop(FOnRequestDrop InOnRequestCanDrop);
 
@@ -268,13 +275,13 @@ protected:
 
 	virtual int32 GetChildIndentLevel() const;
 
-	virtual TOptional<FDropResult> CanDropInternal(const TArray<UNiagaraStackEntry*>& DraggedEntries);
+	virtual TOptional<FDropResult> CanDropInternal(const TArray<UNiagaraStackEntry*>& DraggedEntries, EDragOptions DragOptions);
 
-	virtual TOptional<FDropResult> DropInternal(const TArray<UNiagaraStackEntry*>& DraggedEntries);
+	virtual TOptional<FDropResult> DropInternal(const TArray<UNiagaraStackEntry*>& DraggedEntries, EDragOptions DragOptions);
 
-	virtual TOptional<FDropResult> ChildRequestCanDropInternal(const UNiagaraStackEntry& TargetChild, const TArray<UNiagaraStackEntry*>& DraggedEntries);
+	virtual TOptional<FDropResult> ChildRequestCanDropInternal(const UNiagaraStackEntry& TargetChild, const TArray<UNiagaraStackEntry*>& DraggedEntries, EDragOptions DragOptions);
 
-	virtual TOptional<FDropResult> ChildRequestDropInternal(const UNiagaraStackEntry& TargetChild, const TArray<UNiagaraStackEntry*>& DraggedEntries);
+	virtual TOptional<FDropResult> ChildRequestDropInternal(const UNiagaraStackEntry& TargetChild, const TArray<UNiagaraStackEntry*>& DraggedEntries, EDragOptions DragOptions);
 
 	virtual void ChlildStructureChangedInternal();
 
@@ -289,9 +296,9 @@ private:
 
 	void ChildRequestFullRefreshDeferred();
 
-	TOptional<FDropResult> ChildRequestCanDrop(const UNiagaraStackEntry& TargetChild, const TArray<UNiagaraStackEntry*>& DraggedEntries);
+	TOptional<FDropResult> ChildRequestCanDrop(const UNiagaraStackEntry& TargetChild, const TArray<UNiagaraStackEntry*>& DraggedEntries, EDragOptions DragOptions);
 
-	TOptional<FDropResult> ChildRequestDrop(const UNiagaraStackEntry& TargetChild, const TArray<UNiagaraStackEntry*>& DraggedEntries);
+	TOptional<FDropResult> ChildRequestDrop(const UNiagaraStackEntry& TargetChild, const TArray<UNiagaraStackEntry*>& DraggedEntries, EDragOptions DragOptions);
 
 	void RefreshStackErrorChildren();
 
diff --git a/Engine/Plugins/FX/Niagara/Source/NiagaraEditor/Public/ViewModels/Stack/NiagaraStackGraphUtilities.h b/Engine/Plugins/FX/Niagara/Source/NiagaraEditor/Public/ViewModels/Stack/NiagaraStackGraphUtilities.h
--- a/Engine/Plugins/FX/Niagara/Source/NiagaraEditor/Public/ViewModels/Stack/NiagaraStackGraphUtilities.h
+++ b/Engine/Plugins/FX/Niagara/Source/NiagaraEditor/Public/ViewModels/Stack/NiagaraStackGraphUtilities.h
@@ -167,7 +167,7 @@ namespace FNiagaraStackGraphUtilities
 	
 	bool GetStackIssuesRecursively(const UNiagaraStackEntry* const Entry, TArray<UNiagaraStackErrorItem*>& OutIssues);
 
-	void MoveModule(UNiagaraScript& SourceScript, UNiagaraNodeFunctionCall& ModuleToMove, UNiagaraSystem& TargetSystem, FGuid TargetEmitterHandleId, ENiagaraScriptUsage TargetUsage, FGuid TargetUsageId, int32 TargetModuleIndex);
+	void MoveModule(UNiagaraScript& SourceScript, UNiagaraNodeFunctionCall& ModuleToMove, UNiagaraSystem& TargetSystem, FGuid TargetEmitterHandleId, ENiagaraScriptUsage TargetUsage, FGuid TargetUsageId, int32 TargetModuleIndex, bool bForceCopy);
 
 	/** Whether a parameter is allowed to be used in a certain execution category. 
 		Used to check if parameter can be dropped on a module or funciton stack entry. */
diff --git a/Engine/Plugins/FX/Niagara/Source/NiagaraEditor/Public/ViewModels/Stack/NiagaraStackScriptItemGroup.h b/Engine/Plugins/FX/Niagara/Source/NiagaraEditor/Public/ViewModels/Stack/NiagaraStackScriptItemGroup.h
--- a/Engine/Plugins/FX/Niagara/Source/NiagaraEditor/Public/ViewModels/Stack/NiagaraStackScriptItemGroup.h
+++ b/Engine/Plugins/FX/Niagara/Source/NiagaraEditor/Public/ViewModels/Stack/NiagaraStackScriptItemGroup.h
@@ -51,9 +51,9 @@ protected:
 
 	virtual void FinalizeInternal() override;
 
-	virtual TOptional<FDropResult> ChildRequestCanDropInternal(const UNiagaraStackEntry& TargetChild, const TArray<UNiagaraStackEntry*>& DraggedEntries) override;
+	virtual TOptional<FDropResult> ChildRequestCanDropInternal(const UNiagaraStackEntry& TargetChild, const TArray<UNiagaraStackEntry*>& DraggedEntries, EDragOptions DragOptions) override;
 
-	virtual TOptional<FDropResult> ChildRequestDropInternal(const UNiagaraStackEntry& TargetChild, const TArray<UNiagaraStackEntry*>& DraggedEntries) override;
+	virtual TOptional<FDropResult> ChildRequestDropInternal(const UNiagaraStackEntry& TargetChild, const TArray<UNiagaraStackEntry*>& DraggedEntries, EDragOptions DragOptions) override;
 
 private:
 	void ItemAdded(UNiagaraNodeFunctionCall* AddedModule);
diff --git a/Engine/Plugins/FX/Niagara/Source/NiagaraEditorWidgets/Private/SNiagaraStack.cpp b/Engine/Plugins/FX/Niagara/Source/NiagaraEditorWidgets/Private/SNiagaraStack.cpp
--- a/Engine/Plugins/FX/Niagara/Source/NiagaraEditorWidgets/Private/SNiagaraStack.cpp
+++ b/Engine/Plugins/FX/Niagara/Source/NiagaraEditorWidgets/Private/SNiagaraStack.cpp
@@ -572,7 +572,24 @@ TOptional<EItemDropZone> SNiagaraStack::OnRowCanAcceptDrop(const FDragDropEvent&
 	if (DragDropOp.IsValid())
 	{
 		DragDropOp->ResetToDefaultToolTip();
-		TOptional<UNiagaraStackEntry::FDropResult> Result = InTargetEntry->CanDrop(DragDropOp->GetDraggedEntries());
+
+		UNiagaraStackEntry::EDragOptions DragOptions = UNiagaraStackEntry::EDragOptions::None;
+		if (InDragDropEvent.IsAltDown() &&
+			InDragDropEvent.IsShiftDown() == false &&
+			InDragDropEvent.IsControlDown() == false &&
+			InDragDropEvent.IsCommandDown() == false)
+		{
+			DragOptions = UNiagaraStackEntry::EDragOptions::Copy;
+		}
+		else if (InDragDropEvent.IsAltDown() == false &&
+			InDragDropEvent.IsShiftDown() == false &&
+			InDragDropEvent.IsControlDown() == false &&
+			InDragDropEvent.IsCommandDown() == false)
+		{
+			DragOptions = UNiagaraStackEntry::EDragOptions::Move;
+		}
+		
+		TOptional<UNiagaraStackEntry::FDropResult> Result = InTargetEntry->CanDrop(DragDropOp->GetDraggedEntries(), DragOptions);
 		if (Result.IsSet())
 		{
 			if (Result.GetValue().DropMessage.IsEmptyOrWhitespace() == false)
@@ -599,7 +616,23 @@ FReply SNiagaraStack::OnRowAcceptDrop(const FDragDropEvent& InDragDropEvent, EIt
 	TSharedPtr<FNiagaraStackEntryDragDropOp> DragDropOp = InDragDropEvent.GetOperationAs<FNiagaraStackEntryDragDropOp>();
 	if (DragDropOp.IsValid())
 	{
-		if (ensureMsgf(InTargetEntry->Drop(DragDropOp->GetDraggedEntries()).IsSet(), TEXT("Failed to drop stack entry when can drop returned true")))
+		UNiagaraStackEntry::EDragOptions DragOptions = UNiagaraStackEntry::EDragOptions::None;
+		if (InDragDropEvent.IsAltDown() &&
+			InDragDropEvent.IsShiftDown() == false &&
+			InDragDropEvent.IsControlDown() == false &&
+			InDragDropEvent.IsCommandDown() == false)
+		{
+			DragOptions = UNiagaraStackEntry::EDragOptions::Copy;
+		}
+		else if (InDragDropEvent.IsAltDown() == false &&
+			InDragDropEvent.IsShiftDown() == false &&
+			InDragDropEvent.IsControlDown() == false &&
+			InDragDropEvent.IsCommandDown() == false)
+		{
+			DragOptions = UNiagaraStackEntry::EDragOptions::Move;
+		}
+
+		if (ensureMsgf(InTargetEntry->Drop(DragDropOp->GetDraggedEntries(), DragOptions).IsSet(), TEXT("Failed to drop stack entry when can drop returned true")))
 		{
 			return FReply::Handled();
 		}
