- Add more thread safety around pausing / unpausing instances and async ticking in general
- We now ensure that interactions with the system simulation will wait for the async tick and will wait for instances where applicable (i.e. touching the data set)

diff --git a/Engine/Plugins/FX/Niagara/Source/Niagara/Private/NiagaraSystemSimulation.cpp b/Engine/Plugins/FX/Niagara/Source/Niagara/Private/NiagaraSystemSimulation.cpp
--- a/Engine/Plugins/FX/Niagara/Source/Niagara/Private/NiagaraSystemSimulation.cpp
+++ b/Engine/Plugins/FX/Niagara/Source/Niagara/Private/NiagaraSystemSimulation.cpp
@@ -82,22 +82,6 @@ static FAutoConsoleVariableRef CVarSystemSimTransferParamsParallelThreshold(
 	ECVF_Default
 );
 
-static int32 GbForceNiagaraAsyncWait = 0;
-static FAutoConsoleVariableRef CVarForceNiagaraAsyncWait(
-	TEXT("fx.ForceNiagaraAsyncWait"),
-	GbForceNiagaraAsyncWait,
-	TEXT(""),
-	ECVF_Default
-);
-
-static float GbWaitForAsyncStallMS = .5f;
-static FAutoConsoleVariableRef CVarWaitForAsyncStallMS(
-	TEXT("fx.WaitForAsyncStallMS"),
-	GbWaitForAsyncStallMS,
-	TEXT(""),
-	ECVF_Default
-);
-
 //////////////////////////////////////////////////////////////////////////
 
 FNiagaraSystemSimulationTickContext::FNiagaraSystemSimulationTickContext(FNiagaraSystemSimulation* InOwner, float InDeltaSeconds, bool bInPendingSpawnPass, const FGraphEventRef& InMyCompletionGraphEvent)
@@ -389,6 +373,8 @@ bool FNiagaraSystemSimulation::Init(UNiagaraSystem* InSystem, UWorld* InWorld, b
 void FNiagaraSystemSimulation::Destroy()
 {
 	check(IsInGameThread());
+	WaitForSystemTickComplete();
+
 	while (SystemInstances.Num())
 	{
 		SystemInstances.Last()->Deactivate(true);
@@ -438,6 +424,9 @@ void FNiagaraSystemSimulation::TransferInstance(FNiagaraSystemSimulation* Source
 	check(!bInSpawnPhase);
 	check(!SourceSimulation->bInSpawnPhase);
 
+	WaitForInstancesTickComplete();
+	SourceSimulation->WaitForInstancesTickComplete();
+
 	int32 SystemInstIdx = SystemInst->SystemInstanceIndex;
 	if (!SystemInst->IsPendingSpawn() && SystemInst->SystemInstanceIndex != INDEX_NONE)
 	{
@@ -474,6 +463,8 @@ void FNiagaraSystemSimulation::TransferInstance(FNiagaraSystemSimulation* Source
 
 void FNiagaraSystemSimulation::DumpInstance(const FNiagaraSystemInstance* Inst)const
 {
+	ensure(Inst->bAsyncWorkInProgress == false);
+
 	UE_LOG(LogNiagara, Log, TEXT("==  %s (%d) ========"), *Inst->GetSystem()->GetFullName(), Inst->SystemInstanceIndex);
 	UE_LOG(LogNiagara, Log, TEXT(".................Spawn................."));
 	SpawnExecContext.Parameters.DumpParameters(false);
@@ -554,6 +545,8 @@ void FNiagaraSystemSimulation::SpawnNew_GameThread(float DeltaSeconds, const FGr
 	CSV_SCOPED_TIMING_STAT_EXCLUSIVE(Niagara);
 	LLM_SCOPE(ELLMTag::Niagara);
 
+	WaitForSystemTickComplete(true);
+
 	check(SystemInstances.Num() == DataSet.GetCurrentDataChecked().GetNumInstances());
 	check(PausedSystemInstances.Num() == PausedInstanceData.GetCurrentDataChecked().GetNumInstances());
 
@@ -665,11 +658,15 @@ void FNiagaraSystemSimulation::SpawnNew_GameThread(float DeltaSeconds, const FGr
 void FNiagaraSystemSimulation::Tick_GameThread(float DeltaSeconds, const FGraphEventRef& MyCompletionGraphEvent)
 {
 	check(IsInGameThread());
+	WaitForSystemTickComplete(true);
+
 	SCOPE_CYCLE_COUNTER(STAT_NiagaraOverview_GT);
 	SCOPE_CYCLE_COUNTER(STAT_NiagaraSystemSim_TickGT);
 	CSV_SCOPED_TIMING_STAT_EXCLUSIVE(Niagara);
 	LLM_SCOPE(ELLMTag::Niagara);
 
+	SystemTickGraphEvent = nullptr;
+
 	check(SystemInstances.Num() == DataSet.GetCurrentDataChecked().GetNumInstances());
 	check(PausedSystemInstances.Num() == PausedInstanceData.GetCurrentDataChecked().GetNumInstances());
 
@@ -725,14 +722,14 @@ void FNiagaraSystemSimulation::Tick_GameThread(float DeltaSeconds, const FGraphE
 	if (Context.bTickAsync)
 	{
 		TGraphTask<FNiagaraSystemSimulationWaitAllFinalizeTask>* WaitAllFinalizeTask = TGraphTask<FNiagaraSystemSimulationWaitAllFinalizeTask>::CreateTask(nullptr, ENamedThreads::GameThread).ConstructAndHold(Context.FinalizeEvents);
-		FGraphEventRef WaitAllFinalizeGraphEvent = WaitAllFinalizeTask->GetCompletionEvent();
+		FGraphEventRef FinalizeGraphEvent = WaitAllFinalizeTask->GetCompletionEvent();
 
 		TGraphTask<FNiagaraSystemSimulationTickTask>* SimulationTickTask = TGraphTask<FNiagaraSystemSimulationTickTask>::CreateTask(nullptr, ENamedThreads::GameThread).ConstructAndHold(Context, WaitAllFinalizeTask);
-		FGraphEventRef SimulationTickGraphEvent = SimulationTickTask->GetCompletionEvent();
-		Context.FinalizeEvents->Add(SimulationTickGraphEvent);
+		SystemTickGraphEvent = SimulationTickTask->GetCompletionEvent();
+		Context.FinalizeEvents->Add(SystemTickGraphEvent);
 
 		MyCompletionGraphEvent->SetGatherThreadForDontCompleteUntil(ENamedThreads::GameThread);
-		MyCompletionGraphEvent->DontCompleteUntil(WaitAllFinalizeGraphEvent);
+		MyCompletionGraphEvent->DontCompleteUntil(FinalizeGraphEvent);
 
 		SimulationTickTask->Unlock(ENamedThreads::GameThread);
 	}
@@ -752,40 +749,38 @@ void FNiagaraSystemSimulation::Tick_GameThread(float DeltaSeconds, const FGraphE
 			Context.FinalizeEvents = nullptr;
 		}
 	}
+}
 
-	//TEMPORARY DEV CODE:
-	//Test that the wait system is working
-	if (GbForceNiagaraAsyncWait)
-	{
-		WaitForTickComplete();
-
-		{
-			//We spin a bit here to help make sure that all the Niagara work did infact complete before we get here. Makes things easier to see in Razor.
-			SCOPE_CYCLE_COUNTER(STAT_NiagaraSystemSim_ForceWaitFakeStall);
+void FNiagaraSystemSimulation::WaitForSystemTickComplete(bool bEnsureComplete)
+{
+	check(IsInGameThread());
 
-			float CurrTimeMS = FPlatformTime::Seconds() * 1000;
-			float WaitEnd = CurrTimeMS + GbWaitForAsyncStallMS;
-			do
-			{
-				CurrTimeMS = FPlatformTime::Seconds() * 1000;
-			} 
-			while (CurrTimeMS < WaitEnd);
-		}
+	if (SystemTickGraphEvent.IsValid() && !SystemTickGraphEvent->IsComplete())
+	{
+		SCOPE_CYCLE_COUNTER(STAT_NiagaraSystemSim_ForceWaitForAsync);
+		ensureAlwaysMsgf(!bEnsureComplete, TEXT("Niagara System Simulation Tasks should be complete by now. %s"), *GetSystem()->GetPathName());
+		FTaskGraphInterface::Get().WaitUntilTaskCompletes(SystemTickGraphEvent, ENamedThreads::GameThread);
 	}
+	SystemTickGraphEvent = nullptr;
 }
 
-void FNiagaraSystemSimulation::WaitForTickComplete()
+void FNiagaraSystemSimulation::WaitForInstancesTickComplete(bool bEnsureComplete)
 {
-	SCOPE_CYCLE_COUNTER(STAT_NiagaraSystemSim_ForceWaitForAsync);
 	check(IsInGameThread());
+	WaitForSystemTickComplete(bEnsureComplete);
+
+	SCOPE_CYCLE_COUNTER(STAT_NiagaraSystemSim_ForceWaitForAsync);
+
 	int32 SystemInstIndex = 0;
 	while (SystemInstIndex < SystemInstances.Num())
 	{
+		// If we're in a spawn phase all existing instances should be complete already.
 		FNiagaraSystemInstance* Inst = SystemInstances[SystemInstIndex];
-		Inst->WaitForAsyncTick(bInSpawnPhase);//if we're in a spawn phase all existing instances should be complete already.
-		if (!Inst->IsComplete())
+		Inst->WaitForAsyncTick(bInSpawnPhase);
+
+		// If the system completes during finalize it can be removed from instances so we don't update the index.
+		if ( SystemInstances[SystemInstIndex] == Inst)
 		{
-			//If the system completes during finalize it will be removed from instances so we don't update the index.
 			++SystemInstIndex;
 		}
 
@@ -799,9 +794,10 @@ void FNiagaraSystemSimulation::WaitForTickComplete()
 		{
 			FNiagaraSystemInstance* Inst = SpawningSystemInstances[SystemInstIndex];
 			Inst->WaitForAsyncTick();
-			if (!Inst->IsComplete())
+
+			// If the system completes during finalize it can be removed from instances so we don't update the index.
+			if (SystemInstances[SystemInstIndex] == Inst)
 			{
-				//If the system completes during finalize it will be removed from instances so we don't update the index.
 				++SystemInstIndex;
 			}
 
@@ -1215,6 +1211,8 @@ void FNiagaraSystemSimulation::RemoveInstance(FNiagaraSystemInstance* Instance)
 			DataSet.GetCurrentDataChecked().Dump(Instance->SystemInstanceIndex, 1, TEXT("System data being removed."));
 		}
 
+		WaitForInstancesTickComplete();
+
 		int32 NumInstances = DataSet.GetCurrentDataChecked().GetNumInstances();
 		check(SystemInstances.Num() == NumInstances);
 
@@ -1243,6 +1241,7 @@ void FNiagaraSystemSimulation::AddInstance(FNiagaraSystemInstance* Instance)
 {
 	check(IsInGameThread());
 	check(Instance->SystemInstanceIndex == INDEX_NONE);
+
 	Instance->SetPendingSpawn(true);
 	Instance->SystemInstanceIndex = PendingSystemInstances.Add(Instance);
 
@@ -1260,6 +1259,8 @@ void FNiagaraSystemSimulation::AddInstance(FNiagaraSystemInstance* Instance)
 void FNiagaraSystemSimulation::PauseInstance(FNiagaraSystemInstance* Instance)
 {
 	check(IsInGameThread());
+	WaitForInstancesTickComplete();
+
 	check(!Instance->IsPaused());
 	check(!DataSet.GetDestinationData());
 	check(!PausedInstanceData.GetDestinationData());
@@ -1309,6 +1310,8 @@ void FNiagaraSystemSimulation::PauseInstance(FNiagaraSystemInstance* Instance)
 void FNiagaraSystemSimulation::UnpauseInstance(FNiagaraSystemInstance* Instance)
 {
 	check(IsInGameThread());
+	WaitForInstancesTickComplete();
+
 	check(Instance->IsPaused());
 	check(!DataSet.GetDestinationData());
 	check(!PausedInstanceData.GetDestinationData());
diff --git a/Engine/Plugins/FX/Niagara/Source/Niagara/Public/NiagaraSystemSimulation.h b/Engine/Plugins/FX/Niagara/Source/Niagara/Public/NiagaraSystemSimulation.h
--- a/Engine/Plugins/FX/Niagara/Source/Niagara/Public/NiagaraSystemSimulation.h
+++ b/Engine/Plugins/FX/Niagara/Source/Niagara/Public/NiagaraSystemSimulation.h
@@ -164,7 +164,10 @@ public:
 	/** Second phase of system sim tick that can run on any thread. */
 	void Tick_Concurrent(FNiagaraSystemSimulationTickContext& Context);
 
-	void WaitForTickComplete();
+	/** Wait for system simulation tick to complete.  If bEnsureComplete is true we will trigger an ensure if it is not complete. */
+	void WaitForSystemTickComplete(bool bEnsureComplete = false);
+	/** Wait for instances tick to complete.  If bEnsureComplete is true we will trigger an ensure if it is not complete. */
+	void WaitForInstancesTickComplete(bool bEnsureComplete = false);
 
 	/** Old tick for AB Testing.*/
 	bool Tick_Old(float DeltaSeconds);
@@ -298,4 +301,7 @@ protected:
 
 	/** Current tick batch we're filling ready for processing, potentially in an async task. */
 	FNiagaraSystemTickBatch TickBatch;
+
+	/** Current task that is executing */
+	FGraphEventRef SystemTickGraphEvent;
 };
