- Added fx.Niagara.Solo.TickEarly - Allow Niagara components to tick in the first tick group
- Added fx.Niagara.Solo.AllowAsyncWorkToEndOfFrame - Allow async work to continue until the last tick group

diff --git a/Engine/Plugins/FX/Niagara/Source/Niagara/Private/NiagaraComponent.cpp b/Engine/Plugins/FX/Niagara/Source/Niagara/Private/NiagaraComponent.cpp
--- a/Engine/Plugins/FX/Niagara/Source/Niagara/Private/NiagaraComponent.cpp
+++ b/Engine/Plugins/FX/Niagara/Source/Niagara/Private/NiagaraComponent.cpp
@@ -32,6 +32,22 @@ DECLARE_CYCLE_STAT(TEXT("Get Dynamic Mesh Elements (RT)"), STAT_NiagaraComponent
 
 DEFINE_LOG_CATEGORY(LogNiagara);
 
+static int GNiagaraSoloTickEarly = 1;
+static FAutoConsoleVariableRef CVarNiagaraSoloTickEarly(
+	TEXT("fx.Niagara.Solo.TickEarly"),
+	GNiagaraSoloTickEarly,
+	TEXT("When enabled will tick kin the first available tick group."),
+	ECVF_Default
+);
+
+static int GNiagaraSoloAllowAsyncWorkToEndOfFrame = 1;
+static FAutoConsoleVariableRef CVarNiagaraSoloAllowAsyncWorkToEndOfFrame(
+	TEXT("fx.Niagara.Solo.AllowAsyncWorkToEndOfFrame"),
+	GNiagaraSoloAllowAsyncWorkToEndOfFrame,
+	TEXT("Allow async work to continue until the end of the frame for solo Niagara instances, if false it will complete within the tick group it started in."),
+	ECVF_Default
+);
+
 static int32 GbSuppressNiagaraSystems = 0;
 static FAutoConsoleVariableRef CVarSuppressNiagaraSystems(
 	TEXT("fx.SuppressNiagaraSystems"),
@@ -368,7 +384,8 @@ UNiagaraComponent::UNiagaraComponent(const FObjectInitializer& ObjectInitializer
 	//, bIsChangingAutoAttachment(false)
 {
 	PrimaryComponentTick.bCanEverTick = true;
-	PrimaryComponentTick.TickGroup = TG_DuringPhysics;
+	PrimaryComponentTick.TickGroup = GNiagaraSoloTickEarly ? TG_PrePhysics : TG_DuringPhysics;
+	PrimaryComponentTick.EndTickGroup = GNiagaraSoloAllowAsyncWorkToEndOfFrame ? TG_LastDemotable : ETickingGroup(PrimaryComponentTick.TickGroup);
 	PrimaryComponentTick.bStartWithTickEnabled = false;
 	PrimaryComponentTick.SetTickFunctionEnable(false);
 	bTickInEditor = true;
