#UE4 Make FNiagaraParameterStore Owner transient. This should be set in the constructor and not serialized.

diff --git a/Engine/Plugins/FX/Niagara/Source/Niagara/Public/NiagaraParameterStore.h b/Engine/Plugins/FX/Niagara/Source/Niagara/Public/NiagaraParameterStore.h
--- a/Engine/Plugins/FX/Niagara/Source/Niagara/Public/NiagaraParameterStore.h
+++ b/Engine/Plugins/FX/Niagara/Source/Niagara/Public/NiagaraParameterStore.h
@@ -99,7 +99,7 @@ struct NIAGARA_API FNiagaraParameterStore
 
 private:
 	/** Owner of this store. Used to provide an outer to data interfaces in this store. */
-	UPROPERTY()
+	UPROPERTY(Transient)
 	UObject* Owner;
 	
 	/** Map from parameter defs to their offset in the data table or the data interface. TODO: Separate out into a layout and instance class to reduce duplicated data for this?  */
