Move loading NiagaraEditor module from PreDefault to Default phase, now enforce loading the NiagaraEditor module when loading the Niagara module.

diff --git a/Engine/Plugins/FX/Niagara/Niagara.uplugin b/Engine/Plugins/FX/Niagara/Niagara.uplugin
--- a/Engine/Plugins/FX/Niagara/Niagara.uplugin
+++ b/Engine/Plugins/FX/Niagara/Niagara.uplugin
@@ -44,7 +44,7 @@
 		{
 			"Name" : "NiagaraEditor",
 			"Type" : "Editor",
-			"LoadingPhase" : "PreDefault"
+			"LoadingPhase" : "Default"
 		},
 		{
 			"Name" : "NiagaraEditorWidgets",
diff --git a/Engine/Plugins/FX/Niagara/Source/Niagara/Private/NiagaraModule.cpp b/Engine/Plugins/FX/Niagara/Source/Niagara/Private/NiagaraModule.cpp
--- a/Engine/Plugins/FX/Niagara/Source/Niagara/Private/NiagaraModule.cpp
+++ b/Engine/Plugins/FX/Niagara/Source/Niagara/Private/NiagaraModule.cpp
@@ -176,13 +176,10 @@ void INiagaraModule::StartupModule()
 	FNiagaraWorldManager::OnStartup();
 
 #if WITH_EDITOR	
-	if (!GIsEditor)
-	{
-		// Loading uncooked data in a game environment, we still need to get some functionality from the NiagaraEditor module.
-		// This includes the ability to compile scripts and load WITH_EDITOR_ONLY data.
-		// Note that when loading with the Editor, the NiagaraEditor module is loaded based on the plugin description.
-		FModuleManager::Get().LoadModule(TEXT("NiagaraEditor"));
-	}
+	// Loading uncooked data in a game environment, we still need to get some functionality from the NiagaraEditor module.
+	// This includes the ability to compile scripts and load WITH_EDITOR_ONLY data.
+	// Note that when loading with the Editor, the NiagaraEditor module is loaded based on the plugin description.
+	FModuleManager::Get().LoadModule(TEXT("NiagaraEditor"));
 #endif
 
 	CVarDetailLevel.AsVariable()->SetOnChangedCallback(FConsoleVariableDelegate::CreateRaw(this, &INiagaraModule::OnChangeDetailLevel));
