- Add on screen message when a NaN is encountered
- Disable NaN dumps by default and replace with console variables to control

diff --git a/Engine/Plugins/FX/Niagara/Source/Niagara/Private/NiagaraEmitterInstance.cpp b/Engine/Plugins/FX/Niagara/Source/Niagara/Private/NiagaraEmitterInstance.cpp
--- a/Engine/Plugins/FX/Niagara/Source/Niagara/Private/NiagaraEmitterInstance.cpp
+++ b/Engine/Plugins/FX/Niagara/Source/Niagara/Private/NiagaraEmitterInstance.cpp
@@ -31,6 +31,22 @@ static FAutoConsoleVariableRef CVarNiagaraDumpParticleData(
 	ECVF_Default
 	);
 
+static int32 GbNiagaraDumpNans = 0;
+static FAutoConsoleVariableRef CVarNiagaraDumpNans(
+	TEXT("fx.Niagara.DumpNans"),
+	GbDumpParticleData,
+	TEXT("If not 0 any NaNs will be dumped always.\n"),
+	ECVF_Default
+);
+
+static int32 GbNiagaraDumpNansOnce = 0;
+static FAutoConsoleVariableRef CVarNiagaraDumpNansOnce(
+	TEXT("fx.Niagara.DumpNansOnce"),
+	GbNiagaraDumpNansOnce,
+	TEXT("If not 0 any NaNs will be dumped for the first emitter that encounters NaNs.\n"),
+	ECVF_Default
+);
+
 /**
 TODO: This is mainly to avoid hard limits in our storage/alloc code etc rather than for perf reasons.
 We should improve our hard limit/safety code and possibly add a max for perf reasons.
@@ -705,8 +721,15 @@ FBox FNiagaraEmitterInstance::CalculateDynamicBounds(const bool bReadGPUSimulati
 #if !UE_BUILD_SHIPPING
 	if (bContainsNaN && ParentSystemInstance != nullptr && CachedEmitter != nullptr && ParentSystemInstance->GetSystem() != nullptr)
 	{
-		UE_LOG(LogNiagara, Warning, TEXT("Particle position data contains NaNs. Likely a divide by zero somewhere in your modules. Emitter \"%s\" in System \"%s\""), *CachedEmitter->GetName(), *ParentSystemInstance->GetSystem()->GetName());
-		ParentSystemInstance->Dump();
+		const FString OnScreenMessage = FString::Printf(TEXT("Niagara Particle position data contains NaNs. Likely a divide by zero somewhere in your modules. Emitter \"%s\" in System \"%s\".  Use fx.Niagara.DumpNansOnce to get full log."), *CachedEmitter->GetName(), *ParentSystemInstance->GetSystem()->GetName());
+		GEngine->AddOnScreenDebugMessage((uint64)((PTRINT)this), 3.f, FColor::Red, OnScreenMessage);
+
+		if (GbNiagaraDumpNans || GbNiagaraDumpNansOnce)
+		{
+			GbNiagaraDumpNansOnce = 0;
+			UE_LOG(LogNiagara, Warning, TEXT("Particle position data contains NaNs. Likely a divide by zero somewhere in your modules. Emitter \"%s\" in System \"%s\""), *CachedEmitter->GetName(), *ParentSystemInstance->GetSystem()->GetName());
+			ParentSystemInstance->Dump();
+		}
 	}
 #endif
 
