Expose and connect fx.NiagaraGlobalSpawnCountScale and fx.NiagaraGlobalSystemCountScale to Engine.GlobalSpawnCountScale and Engine.GlobalSystemCountScale, respectively.

diff --git a/Engine/Plugins/FX/Niagara/Source/Niagara/Private/NiagaraSystemSimulation.cpp b/Engine/Plugins/FX/Niagara/Source/Niagara/Private/NiagaraSystemSimulation.cpp
--- a/Engine/Plugins/FX/Niagara/Source/Niagara/Private/NiagaraSystemSimulation.cpp
+++ b/Engine/Plugins/FX/Niagara/Source/Niagara/Private/NiagaraSystemSimulation.cpp
@@ -715,6 +715,10 @@ void FNiagaraSystemSimulation::Tick_GameThread(float DeltaSeconds, const FGraphE
 	UpdateInvDeltaTimeParam.SetValue(InvDt);
 	SpawnNumSystemInstancesParam.SetValue(SystemInstances.Num());
 	UpdateNumSystemInstancesParam.SetValue(SystemInstances.Num());
+	SpawnGlobalSpawnCountScaleParam.SetValue(INiagaraModule::GetGlobalSpawnCountScale());
+	UpdateGlobalSpawnCountScaleParam.SetValue(INiagaraModule::GetGlobalSpawnCountScale());
+	SpawnGlobalSystemCountScaleParam.SetValue(INiagaraModule::GetGlobalSystemCountScale());
+	UpdateGlobalSystemCountScaleParam.SetValue(INiagaraModule::GetGlobalSystemCountScale());
 
 	FNiagaraSystemSimulationTickContext Context(this, DeltaSeconds, false, MyCompletionGraphEvent);
 
