Change the SpawnRate and SpawnBurst_Instantaneous Niagara modules to account for the new fx.NiagaraGlobalSpawnCountScale CVar.

There's three parts to this change:
1) Hooking in the global spawn count scale.
2) Clamping the burst spawn count to 1.
3) Moving the scalability application to an earlier point in the SpawnRate module. The module ended up using two values (pre and post scaled) of the spawn rate in different locations. This might change some autotests.

Updates the Niagara custom script version.

diff --git a/Engine/Plugins/FX/Niagara/Source/NiagaraCore/Private/NiagaraCustomVersion.cpp b/Engine/Plugins/FX/Niagara/Source/NiagaraCore/Private/NiagaraCustomVersion.cpp
--- a/Engine/Plugins/FX/Niagara/Source/NiagaraCore/Private/NiagaraCustomVersion.cpp
+++ b/Engine/Plugins/FX/Niagara/Source/NiagaraCore/Private/NiagaraCustomVersion.cpp
@@ -8,6 +8,6 @@ const FGuid FNiagaraCustomVersion::GUID(0xFCF57AFA, 0x50764283, 0xB9A9E658, 0xFF
 // Register the custom version with core
 FCustomVersionRegistration GRegisterNiagaraCustomVersion(FNiagaraCustomVersion::GUID, FNiagaraCustomVersion::LatestVersion, TEXT("NiagaraVer"));
 
-const FGuid FNiagaraCustomVersion::LatestScriptCompileVersion(0x3BD04C41, 0x37635E34, 0xDBC85417, 0x62BFD88C);
+const FGuid FNiagaraCustomVersion::LatestScriptCompileVersion(0xE2AC5117, 0x5B44424D, 0x91265F00, 0xA1AEB559);
 
 
