- Fix an issue where a spawned system would not have it's GameThread tick called before we started to spawn new instances
- This will be refactored with the system spawning changes

#rnx

diff --git a/Engine/Plugins/FX/Niagara/Source/Niagara/Private/NiagaraSystemSimulation.cpp b/Engine/Plugins/FX/Niagara/Source/Niagara/Private/NiagaraSystemSimulation.cpp
--- a/Engine/Plugins/FX/Niagara/Source/Niagara/Private/NiagaraSystemSimulation.cpp
+++ b/Engine/Plugins/FX/Niagara/Source/Niagara/Private/NiagaraSystemSimulation.cpp
@@ -254,6 +254,7 @@ bool FNiagaraSystemSimulation::Init(UNiagaraSystem* InSystem, UWorld* InWorld, b
 
 	bBindingsInitialized = false;
 	bInSpawnPhase = false;
+	bHasEverTicked = false;
 
 	FNiagaraWorldManager* WorldMan = FNiagaraWorldManager::Get(InWorld);
 	check(WorldMan);
@@ -606,6 +607,12 @@ void FNiagaraSystemSimulation::SpawnNew_GameThread(float DeltaSeconds, const FGr
 			//Lets any RemoveInstance calls inside Tick_Concurrent know that we're spawning and all our instances are in SpawningSystemInstances;
 			bInSpawnPhase = true;
 
+			// When a new simulation is added it may miss it's tick in which case we need to ensure it's ticked before we execute the Tick_Concurrent for spawning
+			if (!bHasEverTicked)
+			{
+				Tick_GameThread(DeltaSeconds, MyCompletionGraphEvent);
+			}
+
 			//For now just do the concurrent tick for new systems on the GT here but we may also want to push this off to task threads too.
 			FNiagaraSystemSimulationTickContext Context(this, DeltaSeconds, true, MyCompletionGraphEvent);
 			Tick_Concurrent(Context);
@@ -665,6 +672,8 @@ void FNiagaraSystemSimulation::Tick_GameThread(float DeltaSeconds, const FGraphE
 	CSV_SCOPED_TIMING_STAT_EXCLUSIVE(Niagara);
 	LLM_SCOPE(ELLMTag::Niagara);
 
+	bHasEverTicked = true;
+
 	SystemTickGraphEvent = nullptr;
 
 	check(SystemInstances.Num() == DataSet.GetCurrentDataChecked().GetNumInstances());
diff --git a/Engine/Plugins/FX/Niagara/Source/Niagara/Public/NiagaraSystemSimulation.h b/Engine/Plugins/FX/Niagara/Source/Niagara/Public/NiagaraSystemSimulation.h
--- a/Engine/Plugins/FX/Niagara/Source/Niagara/Public/NiagaraSystemSimulation.h
+++ b/Engine/Plugins/FX/Niagara/Source/Niagara/Public/NiagaraSystemSimulation.h
@@ -293,6 +293,7 @@ protected:
 	uint32 bBindingsInitialized : 1;
 	uint32 bInSpawnPhase : 1;
 	uint32 bIsSolo : 1;
+	uint32 bHasEverTicked : 1;
 
 	/** A parameter store which contains the data interfaces parameters which were defined by the scripts. */
 	FNiagaraParameterStore ScriptDefinedDataInterfaceParameters;
