Cache EmitterInstance NiagaraVariable names at compile time to avoid FName generation at runtime.
Add stat marker for NiagaraSystem initialization.
Move some NiagaraSystem initialization logic to #WITH_EDITORONLY_DATA.

diff --git a/Engine/Plugins/FX/Niagara/Source/Niagara/Classes/NiagaraEmitter.h b/Engine/Plugins/FX/Niagara/Source/Niagara/Classes/NiagaraEmitter.h
--- a/Engine/Plugins/FX/Niagara/Source/Niagara/Classes/NiagaraEmitter.h
+++ b/Engine/Plugins/FX/Niagara/Source/Niagara/Classes/NiagaraEmitter.h
@@ -349,9 +349,6 @@ public:
 	FString NIAGARA_API GetUniqueEmitterName()const;
 	bool NIAGARA_API SetUniqueEmitterName(const FString& InName);
 
-	/** Converts an emitter paramter "Emitter.XXXX" into it's real parameter name. */
-	FNiagaraVariable ToEmitterParameter(const FNiagaraVariable& EmitterVar)const;
-
 	const TArray<UNiagaraRendererProperties*>& GetRenderers() const { return RendererProperties; }
 
 	void NIAGARA_API AddRenderer(UNiagaraRendererProperties* Renderer);
diff --git a/Engine/Plugins/FX/Niagara/Source/Niagara/Classes/NiagaraSystem.h b/Engine/Plugins/FX/Niagara/Source/Niagara/Classes/NiagaraSystem.h
--- a/Engine/Plugins/FX/Niagara/Source/Niagara/Classes/NiagaraSystem.h
+++ b/Engine/Plugins/FX/Niagara/Source/Niagara/Classes/NiagaraSystem.h
@@ -19,12 +19,32 @@ class UNiagaraEditorDataBase;
 #endif
 
 USTRUCT()
-struct FNiagaraEmitterSpawnAttributes
+struct FNiagaraEmitterCompiledData
 {
 	GENERATED_USTRUCT_BODY()
 
+	FNiagaraEmitterCompiledData();
+	
+	/** Attribute names in the data set that are driving each emitter's spawning. */
 	UPROPERTY()
 	TArray<FName> SpawnAttributes;
+
+
+	/** Explicit list of Niagara Variables to bind to Emitter instances. */
+	UPROPERTY()
+	FNiagaraVariable EmitterSpawnIntervalVar;
+
+	UPROPERTY()
+	FNiagaraVariable EmitterInterpSpawnStartDTVar;
+
+	UPROPERTY()
+	FNiagaraVariable EmitterSpawnGroupVar;
+
+	UPROPERTY()
+	FNiagaraVariable EmitterAgeVar;
+
+	UPROPERTY()
+	FNiagaraVariable EmitterRandomSeedVar;
 };
 
 
@@ -187,7 +207,7 @@ public:
 	bool ShouldAutoDeactivate() const { return bAutoDeactivate; }
 	bool IsLooping() const;
 
-	const TArray<FNiagaraEmitterSpawnAttributes>& GetEmitterSpawnAttributes()const {	return EmitterSpawnAttributes;	};
+	const TArray<FNiagaraEmitterCompiledData>& GetEmitterCompiledData()const {	return EmitterCompiledData;	};
 
 	bool UsesCollection(const UNiagaraParameterCollection* Collection)const;
 #if WITH_EDITORONLY_DATA
@@ -230,6 +250,14 @@ public:
 private:
 #if WITH_EDITORONLY_DATA
 	bool QueryCompileComplete(bool bWait, bool bDoPost, bool bDoNotApply = false);
+
+	void InitEmitterCompiledData();
+
+	/** Helper for filling in precomputed variable names per emitter. Converts an emitter paramter "Emitter.XXXX" into it's real parameter name. */
+	void InitEmitterVariableAliasNames(FNiagaraEmitterCompiledData& EmitterCompiledDataToInit, const UNiagaraEmitter& InAssociatedEmitter);
+
+	/** Helper for generating aliased FNiagaraVariable names for the Emitter they are associated with. */
+	const FName GetEmitterVariableAliasName(const FNiagaraVariable& InEmitterVar, const UNiagaraEmitter& InEmitter) const;
 #endif
 
 	void UpdatePostCompileDIInfo();
@@ -261,9 +289,9 @@ protected:
 	UPROPERTY()
 	UNiagaraScript* SystemUpdateScript;
 
-	/** Attribute names in the data set that are driving each emitter's spawning. */
+	//** Post compile generated data used for initializing Emitter Instances during runtime. */
 	UPROPERTY()
-	TArray<FNiagaraEmitterSpawnAttributes> EmitterSpawnAttributes;
+	TArray<FNiagaraEmitterCompiledData> EmitterCompiledData;
 
 	/** Variables exposed to the outside work for tweaking*/
 	UPROPERTY()
@@ -299,8 +327,6 @@ protected:
 	UPROPERTY(EditAnywhere, Category = Warmup)
 	float WarmupTickDelta;
 
-	void InitEmitterSpawnAttributes();
-
 	UPROPERTY()
 	bool bHasSystemScriptDIsWithPerInstanceData;
 
diff --git a/Engine/Plugins/FX/Niagara/Source/Niagara/Private/NiagaraEmitter.cpp b/Engine/Plugins/FX/Niagara/Source/Niagara/Private/NiagaraEmitter.cpp
--- a/Engine/Plugins/FX/Niagara/Source/Niagara/Private/NiagaraEmitter.cpp
+++ b/Engine/Plugins/FX/Niagara/Source/Niagara/Private/NiagaraEmitter.cpp
@@ -1015,14 +1015,6 @@ bool UNiagaraEmitter::SetUniqueEmitterName(const FString& InName)
 	return false;
 }
 
-
-FNiagaraVariable UNiagaraEmitter::ToEmitterParameter(const FNiagaraVariable& EmitterVar)const
-{
-	FNiagaraVariable Var = EmitterVar;
-	Var.SetName(*Var.GetName().ToString().Replace(TEXT("Emitter."), *(GetUniqueEmitterName() + TEXT("."))));
-	return Var;
-}
-
 void UNiagaraEmitter::AddRenderer(UNiagaraRendererProperties* Renderer)
 {
 	Modify();
diff --git a/Engine/Plugins/FX/Niagara/Source/Niagara/Private/NiagaraEmitterInstance.cpp b/Engine/Plugins/FX/Niagara/Source/Niagara/Private/NiagaraEmitterInstance.cpp
--- a/Engine/Plugins/FX/Niagara/Source/Niagara/Private/NiagaraEmitterInstance.cpp
+++ b/Engine/Plugins/FX/Niagara/Source/Niagara/Private/NiagaraEmitterInstance.cpp
@@ -23,6 +23,7 @@ DECLARE_CYCLE_STAT(TEXT("Emitter Spawn [CNC]"), STAT_NiagaraSpawn, STATGROUP_Nia
 DECLARE_CYCLE_STAT(TEXT("Emitter Post Tick [CNC]"), STAT_NiagaraEmitterPostTick, STATGROUP_Niagara);
 DECLARE_CYCLE_STAT(TEXT("Emitter Event Handling [CNC]"), STAT_NiagaraEventHandle, STATGROUP_Niagara);
 DECLARE_CYCLE_STAT(TEXT("Emitter Error Check [CNC]"), STAT_NiagaraEmitterErrorCheck, STATGROUP_Niagara);
+DECLARE_CYCLE_STAT(TEXT("Init Emitters [GT]"), STAT_NiagaraEmitterInit, STATGROUP_Niagara);
 
 static int32 GbDumpParticleData = 0;
 static FAutoConsoleVariableRef CVarNiagaraDumpParticleData(
@@ -186,6 +187,7 @@ bool FNiagaraEmitterInstance::IsAllowedToExecute() const
 
 void FNiagaraEmitterInstance::Init(int32 InEmitterIdx, FName InSystemInstanceName)
 {
+	SCOPE_CYCLE_COUNTER(STAT_NiagaraEmitterInit);
 	check(ParticleDataSet);
 	FNiagaraDataSet& Data = *ParticleDataSet;
 	EmitterIdx = InEmitterIdx;
@@ -196,8 +198,8 @@ void FNiagaraEmitterInstance::Init(int32 InEmitterIdx, FName InSystemInstanceNam
 	CachedIDName = EmitterHandle.GetIdName();
 
 	if (!IsAllowedToExecute())
-
 	{
+		//@TODO FNiagaraMessageManager Error bubbling here
 		ExecutionState = ENiagaraExecutionState::Disabled;
 		return;
 	}
@@ -209,13 +211,21 @@ void FNiagaraEmitterInstance::Init(int32 InEmitterIdx, FName InSystemInstanceNam
 	Data.Init(FNiagaraDataSetID(CachedIDName, ENiagaraDataSetType::ParticleData), CachedEmitter->SimTarget, ParentSystemInstance->GetSystem()->GetName() + TEXT("/") + CachedEmitter->GetName());
 
 	//Init the spawn infos to the correct number for this system.
-	const TArray<FNiagaraEmitterSpawnAttributes>& EmitterSpawnInfoAttrs = ParentSystemInstance->GetSystem()->GetEmitterSpawnAttributes();
-	if (EmitterSpawnInfoAttrs.IsValidIndex(EmitterIdx))
+	int32 NumEvents = CachedEmitter->GetEventHandlers().Num();
+	const TArray<FNiagaraEmitterCompiledData>& EmitterCompiledData = ParentSystemInstance->GetSystem()->GetEmitterCompiledData();
+	if (EmitterCompiledData.IsValidIndex(EmitterIdx))
 	{
-		SpawnInfos.SetNum(EmitterSpawnInfoAttrs[EmitterIdx].SpawnAttributes.Num());
+		SpawnInfos.SetNum(EmitterCompiledData[EmitterIdx].SpawnAttributes.Num());
+	}
+	else
+	{
+		ExecutionState = ENiagaraExecutionState::Disabled;
+		return;
 	}
 
+#if !UE_BUILD_SHIPPING && !UE_BUILD_TEST 
 	CheckForErrors();
+#endif
 
 	if (IsDisabled())
 	{
@@ -309,7 +319,6 @@ void FNiagaraEmitterInstance::Init(int32 InEmitterIdx, FName InSystemInstanceNam
 	}
 
 	EventExecContexts.SetNum(CachedEmitter->GetEventHandlers().Num());
-	int32 NumEvents = CachedEmitter->GetEventHandlers().Num();
 	for (int32 i = 0; i < NumEvents; i++)
 	{
 		ensure(CachedEmitter->GetEventHandlers()[i].DataSetAccessSynchronized());
@@ -321,31 +330,28 @@ void FNiagaraEmitterInstance::Init(int32 InEmitterIdx, FName InSystemInstanceNam
 	}
 
 	//Setup direct bindings for setting parameter values.
-	SpawnIntervalBinding.Init(SpawnExecContext.Parameters, CachedEmitter->ToEmitterParameter(SYS_PARAM_EMITTER_SPAWN_INTERVAL));
-	InterpSpawnStartBinding.Init(SpawnExecContext.Parameters, CachedEmitter->ToEmitterParameter(SYS_PARAM_EMITTER_INTERP_SPAWN_START_DT));
-	SpawnGroupBinding.Init(SpawnExecContext.Parameters, CachedEmitter->ToEmitterParameter(SYS_PARAM_EMITTER_SPAWN_GROUP));
+	SpawnIntervalBinding.Init(SpawnExecContext.Parameters, EmitterCompiledData[EmitterIdx].EmitterSpawnIntervalVar);
+	InterpSpawnStartBinding.Init(SpawnExecContext.Parameters, EmitterCompiledData[EmitterIdx].EmitterInterpSpawnStartDTVar);
+	SpawnGroupBinding.Init(SpawnExecContext.Parameters, EmitterCompiledData[EmitterIdx].EmitterSpawnGroupVar);
 
-	FNiagaraVariable EmitterAgeParam = CachedEmitter->ToEmitterParameter(SYS_PARAM_EMITTER_AGE);
-	SpawnEmitterAgeBinding.Init(SpawnExecContext.Parameters, EmitterAgeParam);
-	UpdateEmitterAgeBinding.Init(UpdateExecContext.Parameters, EmitterAgeParam);
+	SpawnEmitterAgeBinding.Init(SpawnExecContext.Parameters, EmitterCompiledData[EmitterIdx].EmitterAgeVar);
+	UpdateEmitterAgeBinding.Init(UpdateExecContext.Parameters, EmitterCompiledData[EmitterIdx].EmitterAgeVar);
 	EventEmitterAgeBindings.SetNum(NumEvents);
 	for (int32 i = 0; i < NumEvents; i++)
 	{
-		EventEmitterAgeBindings[i].Init(EventExecContexts[i].Parameters, EmitterAgeParam);
+		EventEmitterAgeBindings[i].Init(EventExecContexts[i].Parameters, EmitterCompiledData[EmitterIdx].EmitterAgeVar);
 	}
 
 	if (CachedEmitter->SimTarget == ENiagaraSimTarget::GPUComputeSim && GPUExecContext != nullptr)
 	{
-		EmitterAgeBindingGPU.Init(GPUExecContext->CombinedParamStore, EmitterAgeParam);
+		EmitterAgeBindingGPU.Init(GPUExecContext->CombinedParamStore, EmitterCompiledData[EmitterIdx].EmitterAgeVar);
 	}
 
-	// Initialize the random seed
-	FNiagaraVariable EmitterRandomSeedParam = CachedEmitter->ToEmitterParameter(SYS_PARAM_EMITTER_RANDOM_SEED);
-	SpawnRandomSeedBinding.Init(SpawnExecContext.Parameters, EmitterRandomSeedParam);
-	UpdateRandomSeedBinding.Init(UpdateExecContext.Parameters, EmitterRandomSeedParam);
+	SpawnRandomSeedBinding.Init(SpawnExecContext.Parameters, EmitterCompiledData[EmitterIdx].EmitterRandomSeedVar);
+	UpdateRandomSeedBinding.Init(UpdateExecContext.Parameters, EmitterCompiledData[EmitterIdx].EmitterRandomSeedVar);
 	if (CachedEmitter->SimTarget == ENiagaraSimTarget::GPUComputeSim && GPUExecContext != nullptr)
 	{
-		GPURandomSeedBinding.Init(GPUExecContext->CombinedParamStore, CachedEmitter->ToEmitterParameter(EmitterRandomSeedParam));
+		GPURandomSeedBinding.Init(GPUExecContext->CombinedParamStore, EmitterCompiledData[EmitterIdx].EmitterRandomSeedVar);
 	}
 
 	// Initialize the exec count
diff --git a/Engine/Plugins/FX/Niagara/Source/Niagara/Private/NiagaraSystem.cpp b/Engine/Plugins/FX/Niagara/Source/Niagara/Private/NiagaraSystem.cpp
--- a/Engine/Plugins/FX/Niagara/Source/Niagara/Private/NiagaraSystem.cpp
+++ b/Engine/Plugins/FX/Niagara/Source/Niagara/Private/NiagaraSystem.cpp
@@ -383,6 +383,11 @@ void UNiagaraSystem::PostLoad()
 		UE_LOG(LogNiagara, Log, TEXT("System %s being compiled because there were changes to an emitter script Change ID."), *GetPathName());
 	}
 
+	if (EmitterCompiledData.Num() == 0)
+	{
+		InitEmitterCompiledData();
+	}
+
 #if 0
 	UE_LOG(LogNiagara, Log, TEXT("Before"));
 	for (FNiagaraEmitterHandle& EmitterHandle : EmitterHandles)
@@ -452,7 +457,7 @@ bool UNiagaraSystem::ReferencesInstanceEmitter(UNiagaraEmitter& Emitter)
 
 void UNiagaraSystem::RefreshSystemParametersFromEmitter(const FNiagaraEmitterHandle& EmitterHandle)
 {
-	InitEmitterSpawnAttributes();
+	InitEmitterCompiledData();
 	if (ensureMsgf(EmitterHandles.ContainsByPredicate([=](const FNiagaraEmitterHandle& OwnedEmitterHandle) { return OwnedEmitterHandle.GetId() == EmitterHandle.GetId(); }),
 		TEXT("Can't refresh parameters from an emitter handle this system doesn't own.")))
 	{
@@ -465,7 +470,7 @@ void UNiagaraSystem::RefreshSystemParametersFromEmitter(const FNiagaraEmitterHan
 
 void UNiagaraSystem::RemoveSystemParametersForEmitter(const FNiagaraEmitterHandle& EmitterHandle)
 {
-	InitEmitterSpawnAttributes();
+	InitEmitterCompiledData();
 	if (ensureMsgf(EmitterHandles.ContainsByPredicate([=](const FNiagaraEmitterHandle& OwnedEmitterHandle) { return OwnedEmitterHandle.GetId() == EmitterHandle.GetId(); }),
 		TEXT("Can't remove parameters for an emitter handle this system doesn't own.")))
 	{
@@ -631,7 +636,7 @@ void UNiagaraSystem::RemoveEmitterHandlesById(const TSet<FGuid>& HandlesToRemove
 	};
 	EmitterHandles.RemoveAll(RemovePredicate);
 
-	InitEmitterSpawnAttributes();
+	InitEmitterCompiledData();
 }
 #endif
 
@@ -793,7 +798,7 @@ bool UNiagaraSystem::QueryCompileComplete(bool bWait, bool bDoPost, bool bDoNotA
 			}
 		}
 
-		InitEmitterSpawnAttributes();
+		InitEmitterCompiledData();
 
 		// Prepare rapid iteration parameters for execution.
 		TArray<UNiagaraScript*> Scripts;
@@ -903,6 +908,22 @@ bool UNiagaraSystem::QueryCompileComplete(bool bWait, bool bDoPost, bool bDoNotA
 	return false;
 }
 
+#if WITH_EDITORONLY_DATA
+void UNiagaraSystem::InitEmitterVariableAliasNames(FNiagaraEmitterCompiledData& EmitterCompiledDataToInit, const UNiagaraEmitter& InAssociatedEmitter)
+{
+	EmitterCompiledDataToInit.EmitterSpawnIntervalVar.SetName(GetEmitterVariableAliasName(SYS_PARAM_EMITTER_SPAWN_INTERVAL, InAssociatedEmitter));
+	EmitterCompiledDataToInit.EmitterInterpSpawnStartDTVar.SetName(GetEmitterVariableAliasName(SYS_PARAM_EMITTER_INTERP_SPAWN_START_DT, InAssociatedEmitter));
+	EmitterCompiledDataToInit.EmitterAgeVar.SetName(GetEmitterVariableAliasName(SYS_PARAM_EMITTER_AGE, InAssociatedEmitter));
+	EmitterCompiledDataToInit.EmitterSpawnGroupVar.SetName(GetEmitterVariableAliasName(SYS_PARAM_EMITTER_SPAWN_GROUP, InAssociatedEmitter));
+	EmitterCompiledDataToInit.EmitterRandomSeedVar.SetName(GetEmitterVariableAliasName(SYS_PARAM_EMITTER_RANDOM_SEED, InAssociatedEmitter));
+}
+
+const FName UNiagaraSystem::GetEmitterVariableAliasName(const FNiagaraVariable& InEmitterVar, const UNiagaraEmitter& InEmitter) const
+{
+	return FName(*InEmitterVar.GetName().ToString().Replace(TEXT("Emitter."), *(InEmitter.GetUniqueEmitterName() + TEXT("."))));
+}
+#endif
+
 bool UNiagaraSystem::RequestCompile(bool bForce)
 {
 	if (bForce)
@@ -1021,13 +1042,15 @@ bool UNiagaraSystem::RequestCompile(bool bForce)
 
 #endif
 
-void UNiagaraSystem::InitEmitterSpawnAttributes()
+#if WITH_EDITORONLY_DATA
+void UNiagaraSystem::InitEmitterCompiledData()
 {
-	EmitterSpawnAttributes.Empty();
-	EmitterSpawnAttributes.SetNum(EmitterHandles.Num());
-	FNiagaraTypeDefinition SpawnInfoDef = FNiagaraTypeDefinition(FNiagaraSpawnInfo::StaticStruct());
-	if (SystemSpawnScript->GetVMExecutableData().IsValid())
+	if (SystemSpawnScript->GetVMExecutableData().IsValid() && SystemUpdateScript->GetVMExecutableData().IsValid())
 	{
+		EmitterCompiledData.Empty();
+		EmitterCompiledData.SetNum(EmitterHandles.Num());
+		FNiagaraTypeDefinition SpawnInfoDef = FNiagaraTypeDefinition(FNiagaraSpawnInfo::StaticStruct());
+
 		for (FNiagaraVariable& Var : SystemSpawnScript->GetVMExecutableData().Attributes)
 		{
 			for (int32 EmitterIdx = 0; EmitterIdx < EmitterHandles.Num(); ++EmitterIdx)
@@ -1036,13 +1059,11 @@ void UNiagaraSystem::InitEmitterSpawnAttributes()
 				FString EmitterName = Emitter->GetUniqueEmitterName() + TEXT(".");
 				if (Var.GetType() == SpawnInfoDef && Var.GetName().ToString().StartsWith(EmitterName))
 				{
-					EmitterSpawnAttributes[EmitterIdx].SpawnAttributes.AddUnique(Var.GetName());
+					EmitterCompiledData[EmitterIdx].SpawnAttributes.AddUnique(Var.GetName());
 				}
 			}
 		}
-	}
-	if (SystemUpdateScript->GetVMExecutableData().IsValid())
-	{
+
 		for (FNiagaraVariable& Var : SystemUpdateScript->GetVMExecutableData().Attributes)
 		{
 			for (int32 EmitterIdx = 0; EmitterIdx < EmitterHandles.Num(); ++EmitterIdx)
@@ -1051,12 +1072,18 @@ void UNiagaraSystem::InitEmitterSpawnAttributes()
 				FString EmitterName = Emitter->GetUniqueEmitterName() + TEXT(".");
 				if (Var.GetType() == SpawnInfoDef && Var.GetName().ToString().StartsWith(EmitterName))
 				{
-					EmitterSpawnAttributes[EmitterIdx].SpawnAttributes.AddUnique(Var.GetName());
+					EmitterCompiledData[EmitterIdx].SpawnAttributes.AddUnique(Var.GetName());
 				}
 			}
 		}
+
+		for (int32 EmitterIdx = 0; EmitterIdx < EmitterHandles.Num(); ++EmitterIdx)
+		{
+			InitEmitterVariableAliasNames(EmitterCompiledData[EmitterIdx], *EmitterHandles[EmitterIdx].GetInstance());
+		}
 	}
 }
+#endif
 
 TStatId UNiagaraSystem::GetStatID(bool bGameThread, bool bConcurrent)const
 {
@@ -1095,4 +1122,13 @@ void UNiagaraSystem::GenerateStatID()
 	StatID_RT = FDynamicStats::CreateStatId<FStatGroup_STATGROUP_NiagaraSystems>(GetPathName() + TEXT("[RT]"));
 	StatID_RT_CNC = FDynamicStats::CreateStatId<FStatGroup_STATGROUP_NiagaraSystems>(GetPathName() + TEXT("[RT_CNC]"));
 #endif
-}
\ No newline at end of file
+}
+
+FNiagaraEmitterCompiledData::FNiagaraEmitterCompiledData()
+{
+	EmitterSpawnIntervalVar = SYS_PARAM_EMITTER_SPAWN_INTERVAL;
+	EmitterInterpSpawnStartDTVar = SYS_PARAM_EMITTER_INTERP_SPAWN_START_DT;
+	EmitterAgeVar = SYS_PARAM_EMITTER_AGE;
+	EmitterSpawnGroupVar = SYS_PARAM_EMITTER_SPAWN_GROUP;
+	EmitterRandomSeedVar = SYS_PARAM_EMITTER_RANDOM_SEED;
+}
diff --git a/Engine/Plugins/FX/Niagara/Source/Niagara/Private/NiagaraSystemInstance.cpp b/Engine/Plugins/FX/Niagara/Source/Niagara/Private/NiagaraSystemInstance.cpp
--- a/Engine/Plugins/FX/Niagara/Source/Niagara/Private/NiagaraSystemInstance.cpp
+++ b/Engine/Plugins/FX/Niagara/Source/Niagara/Private/NiagaraSystemInstance.cpp
@@ -1504,8 +1504,6 @@ void FNiagaraSystemInstance::InitEmitters()
 		const TArray<FNiagaraEmitterHandle>& EmitterHandles = GetSystem()->GetEmitterHandles();
 		for (int32 EmitterIdx=0; EmitterIdx < GetSystem()->GetEmitterHandles().Num(); ++EmitterIdx)
 		{
-			const FNiagaraEmitterHandle& EmitterHandle = EmitterHandles[EmitterIdx];
-
 			TSharedRef<FNiagaraEmitterInstance, ESPMode::ThreadSafe> Sim = MakeShared<FNiagaraEmitterInstance, ESPMode::ThreadSafe>(this);
 			Sim->Init(EmitterIdx, IDName);
 			if (System->bFixedBounds)
diff --git a/Engine/Plugins/FX/Niagara/Source/Niagara/Private/NiagaraSystemSimulation.cpp b/Engine/Plugins/FX/Niagara/Source/Niagara/Private/NiagaraSystemSimulation.cpp
--- a/Engine/Plugins/FX/Niagara/Source/Niagara/Private/NiagaraSystemSimulation.cpp
+++ b/Engine/Plugins/FX/Niagara/Source/Niagara/Private/NiagaraSystemSimulation.cpp
@@ -30,6 +30,7 @@ DECLARE_CYCLE_STAT(TEXT("System Prepare For Simulate [CNC]"), STAT_NiagaraSystem
 DECLARE_CYCLE_STAT(TEXT("System Sim Update [CNC]"), STAT_NiagaraSystemSim_UpdateCNC, STATGROUP_Niagara);
 DECLARE_CYCLE_STAT(TEXT("System Sim Spawn [CNC]"), STAT_NiagaraSystemSim_SpawnCNC, STATGROUP_Niagara);
 DECLARE_CYCLE_STAT(TEXT("System Sim Transfer Results [CNC]"), STAT_NiagaraSystemSim_TransferResultsCNC, STATGROUP_Niagara);
+DECLARE_CYCLE_STAT(TEXT("System Sim Init [GT]"), STAT_NiagaraSystemSim_Init, STATGROUP_Niagara);
 
 DECLARE_CYCLE_STAT(TEXT("ForcedWaitForAsync"), STAT_NiagaraSystemSim_ForceWaitForAsync, STATGROUP_Niagara);
 DECLARE_CYCLE_STAT(TEXT("ForcedWait Fake Stall"), STAT_NiagaraSystemSim_ForceWaitFakeStall, STATGROUP_Niagara);
@@ -245,6 +246,7 @@ FNiagaraSystemSimulation::~FNiagaraSystemSimulation()
 
 bool FNiagaraSystemSimulation::Init(UNiagaraSystem* InSystem, UWorld* InWorld, bool bInIsSolo)
 {
+	SCOPE_CYCLE_COUNTER(STAT_NiagaraSystemSim_Init);
 	UNiagaraSystem* System = InSystem;
 	WeakSystem = System;
 
@@ -340,10 +342,10 @@ bool FNiagaraSystemSimulation::Init(UNiagaraSystem* InSystem, UWorld* InWorld, b
 			FString EmitterName = Emitter->GetUniqueEmitterName();
 			check(Emitter);
 			EmitterExecutionStateAccessors.Emplace(DataSet, FNiagaraVariable(EnumPtr, *(EmitterName + TEXT(".ExecutionState"))));
-			const TArray<FNiagaraEmitterSpawnAttributes>& EmitterSpawnAttrNames = System->GetEmitterSpawnAttributes();
+			const TArray<FNiagaraEmitterCompiledData>& EmitterCompiledData = System->GetEmitterCompiledData();
 			
-			check(EmitterSpawnAttrNames.Num() == System->GetNumEmitters());
-			for (FName AttrName : EmitterSpawnAttrNames[EmitterIdx].SpawnAttributes)
+			check(EmitterCompiledData.Num() == System->GetNumEmitters());
+			for (FName AttrName : EmitterCompiledData[EmitterIdx].SpawnAttributes)
 			{
 				EmitterSpawnInfoAccessors[EmitterIdx].Emplace(DataSet, FNiagaraVariable(FNiagaraTypeDefinition(FNiagaraSpawnInfo::StaticStruct()), AttrName));
 			}
