Fixed issue where Niagara components would never clear their render state dirty flag.

diff --git a/Engine/Plugins/FX/Niagara/Source/Niagara/Private/NiagaraComponent.cpp b/Engine/Plugins/FX/Niagara/Source/Niagara/Private/NiagaraComponent.cpp
--- a/Engine/Plugins/FX/Niagara/Source/Niagara/Private/NiagaraComponent.cpp
+++ b/Engine/Plugins/FX/Niagara/Source/Niagara/Private/NiagaraComponent.cpp
@@ -1021,6 +1021,9 @@ void UNiagaraComponent::SendRenderDynamicData_Concurrent()
 	LLM_SCOPE(ELLMTag::Niagara);
 	CSV_SCOPED_TIMING_STAT_EXCLUSIVE(Niagara);
 	SCOPE_CYCLE_COUNTER(STAT_NiagaraComponentSendRenderData);
+
+	Super::SendRenderDynamicData_Concurrent();
+
 	if (SystemInstance.IsValid() && SceneProxy)
 	{
 #if STATS
