Fixed some missing scopes for stat niagarasystems

diff --git a/Engine/Plugins/FX/Niagara/Source/Niagara/Private/NiagaraSystemSimulation.cpp b/Engine/Plugins/FX/Niagara/Source/Niagara/Private/NiagaraSystemSimulation.cpp
--- a/Engine/Plugins/FX/Niagara/Source/Niagara/Private/NiagaraSystemSimulation.cpp
+++ b/Engine/Plugins/FX/Niagara/Source/Niagara/Private/NiagaraSystemSimulation.cpp
@@ -550,6 +550,9 @@ void FNiagaraSystemSimulation::SpawnNew_GameThread(float DeltaSeconds, const FGr
 
 	WaitForSystemTickComplete(true);
 
+	UNiagaraSystem* System = WeakSystem.Get();
+	FScopeCycleCounter SystemStatCounter(System->GetStatID(true, false));
+
 	check(SystemInstances.Num() == DataSet.GetCurrentDataChecked().GetNumInstances());
 	check(PausedSystemInstances.Num() == PausedInstanceData.GetCurrentDataChecked().GetNumInstances());
 
@@ -558,7 +561,6 @@ void FNiagaraSystemSimulation::SpawnNew_GameThread(float DeltaSeconds, const FGr
 	//Then spwan these into a new dataset which is then transfered into the main dataset.
 	//Currently this is all done on the game thread but we could likely also move this off the GT too.
 
-	UNiagaraSystem* System = WeakSystem.Get();
 
 	check(IsInGameThread());
 	if (PendingSystemInstances.Num() > 0)
@@ -841,6 +843,8 @@ void FNiagaraSystemSimulation::Tick_Concurrent(FNiagaraSystemSimulationTickConte
 			UE_LOG(LogNiagara, Log, TEXT("=========================================================="));
 		}
 
+		FScopeCycleCounter SystemStatCounter(Context.System->GetStatID(true, true));
+
 		PrepareForSystemSimulate(Context);
 
 		if (Context.bPendingSpawnPass)
