Added a DumpToString method to NiagaraUtilities

diff --git a/Engine/Plugins/FX/Niagara/Source/Niagara/Private/NiagaraCommon.cpp b/Engine/Plugins/FX/Niagara/Source/Niagara/Private/NiagaraCommon.cpp
--- a/Engine/Plugins/FX/Niagara/Source/Niagara/Private/NiagaraCommon.cpp
+++ b/Engine/Plugins/FX/Niagara/Source/Niagara/Private/NiagaraCommon.cpp
@@ -290,6 +290,19 @@ UNiagaraDataInterface* FNiagaraScriptDataInterfaceCompileInfo::GetDefaultDataInt
 	return Obj;
 }
 
+void FNiagaraUtilities::DumpHLSLText(const FString& SourceCode, const FString& DebugName)
+{
+	UE_LOG(LogNiagara, Display, TEXT("Compile output as text: %s"), *DebugName);
+	UE_LOG(LogNiagara, Display, TEXT("==================================================================================="));
+	TArray<FString> OutputByLines;
+	SourceCode.ParseIntoArrayLines(OutputByLines, false);
+	for (int32 i = 0; i < OutputByLines.Num(); i++)
+	{
+		UE_LOG(LogNiagara, Display, TEXT("/*%04d*/\t\t%s"), i + 1, *OutputByLines[i]);
+	}
+	UE_LOG(LogNiagara, Display, TEXT("==================================================================================="));
+}
+
 #if WITH_EDITORONLY_DATA
 void FNiagaraUtilities::PrepareRapidIterationParameters(const TArray<UNiagaraScript*>& Scripts, const TMap<UNiagaraScript*, UNiagaraScript*>& ScriptDependencyMap, const TMap<UNiagaraScript*, FString>& ScriptToEmitterNameMap)
 {
diff --git a/Engine/Plugins/FX/Niagara/Source/Niagara/Public/NiagaraCommon.h b/Engine/Plugins/FX/Niagara/Source/Niagara/Public/NiagaraCommon.h
--- a/Engine/Plugins/FX/Niagara/Source/Niagara/Public/NiagaraCommon.h
+++ b/Engine/Plugins/FX/Niagara/Source/Niagara/Public/NiagaraCommon.h
@@ -623,6 +623,9 @@ namespace FNiagaraUtilities
 	 */
 	void NIAGARA_API PrepareRapidIterationParameters(const TArray<UNiagaraScript*>& Scripts, const TMap<UNiagaraScript*, UNiagaraScript*>& ScriptDependencyMap, const TMap<UNiagaraScript*, FString>& ScriptToEmitterNameMap);
 #endif
+
+	void NIAGARA_API DumpHLSLText(const FString& SourceCode, const FString& DebugName);
+
 };
 
 USTRUCT()
