Making  FNiagaraDataSetAccessor logic a little more robust to a missing data set

diff --git a/Engine/Plugins/FX/Niagara/Source/Niagara/Classes/NiagaraDataSet.h b/Engine/Plugins/FX/Niagara/Source/Niagara/Classes/NiagaraDataSet.h
--- a/Engine/Plugins/FX/Niagara/Source/Niagara/Classes/NiagaraDataSet.h
+++ b/Engine/Plugins/FX/Niagara/Source/Niagara/Classes/NiagaraDataSet.h
@@ -382,9 +382,16 @@ struct FNiagaraDataSetAccessorBase
 
 	void SetDataSet(FNiagaraDataSet& InDataSet)
 	{
-		ensure(Var.IsValid());
-		DataSet = &InDataSet;
-		VarLayout = InDataSet.GetVariableLayout(Var);
+		if (Var.IsValid())
+		{
+			DataSet = &InDataSet;
+			VarLayout = InDataSet.GetVariableLayout(Var);
+		}
+		else
+		{
+			DataSet = nullptr;
+			VarLayout = nullptr;
+		}
 	}
 
 	FORCEINLINE bool IsValid()const { return DataSet && VarLayout != nullptr; }
@@ -556,16 +563,19 @@ struct FNiagaraDataSetAccessor<int32> : public FNiagaraDataSetAccessorBase
 	{
 		SrcBase = nullptr;
 		DestBase = nullptr;
-		FNiagaraDataBuffer* SrcBuffer = DataSet->GetCurrentData();
-		if (IsValid() && SrcBuffer)
+		if (IsValid())
 		{
-			SrcBase = (int32*)SrcBuffer->GetComponentPtrInt32(VarLayout->Int32ComponentStart);
-
-			//Writes are only valid if we're during a simulation pass.
-			FNiagaraDataBuffer* DestBuffer = DataSet->GetDestinationData();
-			if (DestBuffer)
+			FNiagaraDataBuffer* SrcBuffer = DataSet->GetCurrentData();
+			if (SrcBuffer)
 			{
-				DestBase = (int32*)DestBuffer->GetComponentPtrInt32(VarLayout->Int32ComponentStart);
+				SrcBase = (int32*)SrcBuffer->GetComponentPtrInt32(VarLayout->Int32ComponentStart);
+
+				//Writes are only valid if we're during a simulation pass.
+				FNiagaraDataBuffer* DestBuffer = DataSet->GetDestinationData();
+				if (DestBuffer)
+				{
+					DestBase = (int32*)DestBuffer->GetComponentPtrInt32(VarLayout->Int32ComponentStart);
+				}
 			}
 		}
 	}
