Optimized FNiagaraParameterStore::Bind() by skipping empty bindings and also processing niagara variables by the smallest set first.
Fixed a bug in FNiagaraSystemInstance::BindParameters() where we would not bind the FNiagaraEmitterInstance::UpdateExecContext parameters.

diff --git a/Engine/Plugins/FX/Niagara/Source/Niagara/Private/NiagaraParameterStore.cpp b/Engine/Plugins/FX/Niagara/Source/Niagara/Private/NiagaraParameterStore.cpp
--- a/Engine/Plugins/FX/Niagara/Source/Niagara/Private/NiagaraParameterStore.cpp
+++ b/Engine/Plugins/FX/Niagara/Source/Niagara/Private/NiagaraParameterStore.cpp
@@ -92,8 +92,84 @@ void FNiagaraParameterStore::Bind(FNiagaraParameterStore* DestStore)
 {
 	check(DestStore);
 	SCOPE_CYCLE_COUNTER(STAT_NiagaraParameterStoreBind);
-	FNiagaraParameterStoreBinding& Binding = Bindings.FindOrAdd(DestStore);
-	Binding.Initialize(DestStore, this);
+	if (!Bindings.Contains(DestStore))
+	{
+		// Bind the parameter stores only if they have variables in common.
+		FNiagaraParameterStoreBinding HeapBinding;
+		if  (HeapBinding.Initialize(DestStore, this))
+		{
+			FNiagaraParameterStoreBinding& Binding = Bindings.FindOrAdd(DestStore);
+			FMemory::Memswap(&Binding, &HeapBinding, sizeof(FNiagaraParameterStoreBinding));
+			return;
+		}
+	}
+}
+
+bool FNiagaraParameterStoreBinding::BindParameters(FNiagaraParameterStore* DestStore, FNiagaraParameterStore* SrcStore)
+{
+	InterfaceBindings.Reset();
+	ParameterBindings.Reset();
+	UObjectBindings.Reset();
+
+	bool bAnyBinding = false;
+
+	int32 DestOffset = INDEX_NONE;
+	int32 SrcOffset = INDEX_NONE;
+
+	auto BindVariable = [&](const FNiagaraVariable& InParameter)
+	{
+		bAnyBinding = true;
+
+		if (InParameter.IsDataInterface())
+		{
+			InterfaceBindings.Add(FInterfaceBinding(SrcOffset, DestOffset));
+		}
+		else if (InParameter.IsUObject())
+		{
+			UObjectBindings.Add(FUObjectBinding(SrcOffset, DestOffset));
+		}
+		else
+		{
+			ParameterBindings.Add(FParameterBinding(SrcOffset, DestOffset, InParameter.GetSizeInBytes()));
+		}
+	};
+
+	// Process the smaller parameter store the get the least amount of iterations when it is small (often empty).
+	if (DestStore->GetParameterOffsets().Num() <= SrcStore->GetParameterOffsets().Num())
+	{
+		for (const TPair<FNiagaraVariable, int32>& ParamOffsetPair : DestStore->GetParameterOffsets())
+		{
+			const FNiagaraVariable& Parameter = ParamOffsetPair.Key;
+			DestOffset = ParamOffsetPair.Value;
+			SrcOffset = SrcStore->IndexOf(Parameter);
+
+			if (SrcOffset != INDEX_NONE && DestOffset != INDEX_NONE)
+			{
+				BindVariable(Parameter);
+			}
+		}
+	}
+	else
+	{
+		for (const TPair<FNiagaraVariable, int32>& ParamOffsetPair : SrcStore->GetParameterOffsets())
+		{
+			const FNiagaraVariable& Parameter = ParamOffsetPair.Key;
+			SrcOffset = ParamOffsetPair.Value;
+			DestOffset = DestStore->IndexOf(Parameter);
+
+			if (SrcOffset != INDEX_NONE && DestOffset != INDEX_NONE)
+			{
+				BindVariable(Parameter);
+			}
+		}
+	}
+
+	if (bAnyBinding)
+	{
+		//Force an initial tick to prime our values in the destination store.
+		Tick(DestStore, SrcStore, true);
+	}
+	return bAnyBinding;
 }
 
 void FNiagaraParameterStore::Unbind(FNiagaraParameterStore* DestStore)
diff --git a/Engine/Plugins/FX/Niagara/Source/Niagara/Private/NiagaraSystemInstance.cpp b/Engine/Plugins/FX/Niagara/Source/Niagara/Private/NiagaraSystemInstance.cpp
--- a/Engine/Plugins/FX/Niagara/Source/Niagara/Private/NiagaraSystemInstance.cpp
+++ b/Engine/Plugins/FX/Niagara/Source/Niagara/Private/NiagaraSystemInstance.cpp
@@ -992,7 +992,7 @@ void FNiagaraSystemInstance::BindParameters()
 		// If this simulation is not solo than we have bind the source system parameters to the system simulation contexts so that
 		// the system and emitter scripts use the default shared data interfaces.
 		GetSystem()->GetExposedParameters().Bind(&SystemSimulation->GetSpawnExecutionContext().Parameters);
-		GetSystem()->GetExposedParameters().Bind(&SystemSimulation->GetSpawnExecutionContext().Parameters);
+		GetSystem()->GetExposedParameters().Bind(&SystemSimulation->GetUpdateExecutionContext().Parameters);
 	}
 
 	for (TSharedRef<FNiagaraEmitterInstance, ESPMode::ThreadSafe> Simulation : Emitters)
diff --git a/Engine/Plugins/FX/Niagara/Source/Niagara/Public/NiagaraParameterStore.h b/Engine/Plugins/FX/Niagara/Source/Niagara/Public/NiagaraParameterStore.h
--- a/Engine/Plugins/FX/Niagara/Source/Niagara/Public/NiagaraParameterStore.h
+++ b/Engine/Plugins/FX/Niagara/Source/Niagara/Public/NiagaraParameterStore.h
@@ -76,7 +76,7 @@ struct FNiagaraParameterStoreBinding
 	}
 	
 	FORCEINLINE_DEBUGGABLE void Empty(FNiagaraParameterStore* DestStore, FNiagaraParameterStore* SrcStore);
-	FORCEINLINE_DEBUGGABLE void Initialize(FNiagaraParameterStore* DestStore, FNiagaraParameterStore* SrcStore);
+	FORCEINLINE_DEBUGGABLE bool Initialize(FNiagaraParameterStore* DestStore, FNiagaraParameterStore* SrcStore);
 	FORCEINLINE_DEBUGGABLE bool VerifyBinding(const FNiagaraParameterStore* DestStore, const FNiagaraParameterStore* SrcStore)const;
 	FORCEINLINE_DEBUGGABLE void Tick(FNiagaraParameterStore* DestStore, FNiagaraParameterStore* SrcStore, bool bForce = false);
 	FORCEINLINE_DEBUGGABLE void Dump(const FNiagaraParameterStore* DestStore, const FNiagaraParameterStore* SrcStore)const;
@@ -84,7 +84,7 @@ struct FNiagaraParameterStoreBinding
 	//FORCEINLINE_DEBUGGABLE void Optimize();
 
 private:
-	void BindParameters(FNiagaraParameterStore* DestStore, FNiagaraParameterStore* SrcStore);
+	bool BindParameters(FNiagaraParameterStore* DestStore, FNiagaraParameterStore* SrcStore);
 };
 
 /** Base storage class for Niagara parameter values. */
@@ -571,45 +571,20 @@ FORCEINLINE_DEBUGGABLE void FNiagaraParameterStoreBinding::Empty(FNiagaraParamet
 	UObjectBindings.Reset();
 }
 
-FORCEINLINE_DEBUGGABLE void FNiagaraParameterStoreBinding::BindParameters(FNiagaraParameterStore* DestStore, FNiagaraParameterStore* SrcStore)
-{
-	InterfaceBindings.Reset();
-	ParameterBindings.Reset();
-	UObjectBindings.Reset();
-	for (const TPair<FNiagaraVariable, int32>& ParamOffsetPair : DestStore->GetParameterOffsets())
-	{
-		const FNiagaraVariable& Parameter = ParamOffsetPair.Key;
-		int32 DestOffset = ParamOffsetPair.Value;
-		int32 SrcOffset = SrcStore->IndexOf(Parameter);
-
-		if (SrcOffset != INDEX_NONE && DestOffset != INDEX_NONE)
-		{
-			if (Parameter.IsDataInterface())
-			{
-				InterfaceBindings.Add(FInterfaceBinding(SrcOffset, DestOffset));
-			}
-			else if (Parameter.IsUObject())
-			{
-				UObjectBindings.Add(FUObjectBinding(SrcOffset, DestOffset));
-			}
-			else
-			{
-				ParameterBindings.Add(FParameterBinding(SrcOffset, DestOffset, Parameter.GetSizeInBytes()));
-			}
-		}
-	}
-
-	//Force an initial tick to prime our values in the destination store.
-	Tick(DestStore, SrcStore, true);
-}
-
-FORCEINLINE_DEBUGGABLE void FNiagaraParameterStoreBinding::Initialize(FNiagaraParameterStore* DestStore, FNiagaraParameterStore* SrcStore)
+FORCEINLINE_DEBUGGABLE bool FNiagaraParameterStoreBinding::Initialize(FNiagaraParameterStore* DestStore, FNiagaraParameterStore* SrcStore)
 {
 	checkSlow(DestStore);
 	checkSlow(SrcStore);
-	DestStore->GetSourceParameterStores().AddUnique(SrcStore);
 
-	BindParameters(DestStore, SrcStore);
+	if (BindParameters(DestStore, SrcStore))
+	{
+		DestStore->GetSourceParameterStores().AddUnique(SrcStore);
+		return true;
+	}
+	else
+	{
+		return false;
+	}
 }
 
 FORCEINLINE_DEBUGGABLE bool FNiagaraParameterStoreBinding::VerifyBinding(const FNiagaraParameterStore* DestStore, const FNiagaraParameterStore* SrcStore)const
diff --git a/Engine/Plugins/FX/Niagara/Source/Niagara/Public/NiagaraTypes.h b/Engine/Plugins/FX/Niagara/Source/Niagara/Public/NiagaraTypes.h
--- a/Engine/Plugins/FX/Niagara/Source/Niagara/Public/NiagaraTypes.h
+++ b/Engine/Plugins/FX/Niagara/Source/Niagara/Public/NiagaraTypes.h
@@ -810,8 +810,8 @@ struct FNiagaraVariable
 		return Name == Other.Name && (TypeDef == Other.TypeDef || (bAllowAssignableTypes && FNiagaraTypeDefinition::TypesAreAssignable(TypeDef, Other.TypeDef)));
 	}
 	
-	void SetName(FName InName) { Name = InName; }
-	FName GetName()const { return Name; }
+	FORCEINLINE void SetName(FName InName) { Name = InName; }
+	FORCEINLINE const FName& GetName() const { return Name; }
 
 	void SetType(const FNiagaraTypeDefinition& InTypeDef) { TypeDef = InTypeDef; }
 	const FNiagaraTypeDefinition& GetType()const { return TypeDef; }
