Remove debug FName from Niagara Datasets when building for Shipping or Test.

diff --git a/Engine/Plugins/FX/Niagara/Source/Niagara/Classes/NiagaraCollision.h b/Engine/Plugins/FX/Niagara/Source/Niagara/Classes/NiagaraCollision.h
--- a/Engine/Plugins/FX/Niagara/Source/Niagara/Classes/NiagaraCollision.h
+++ b/Engine/Plugins/FX/Niagara/Source/Niagara/Classes/NiagaraCollision.h
@@ -50,7 +50,11 @@ public:
 	{
 		if (CollisionEventDataSet)
 		{
+#if !UE_BUILD_SHIPPING && !UE_BUILD_TEST
 			CollisionEventDataSet->Init(FNiagaraDataSetID(), ENiagaraSimTarget::CPUSim, TEXT("NiagaraCollisionEventDataset"));
+#else
+			CollisionEventDataSet->Init(FNiagaraDataSetID(), ENiagaraSimTarget::CPUSim);
+#endif
 		}
 		EmitterName = InEmitterName;
 		OwnerSystemInstanceName = InOwnerSystemInstanceName;
diff --git a/Engine/Plugins/FX/Niagara/Source/Niagara/Classes/NiagaraDataSet.h b/Engine/Plugins/FX/Niagara/Source/Niagara/Classes/NiagaraDataSet.h
--- a/Engine/Plugins/FX/Niagara/Source/Niagara/Classes/NiagaraDataSet.h
+++ b/Engine/Plugins/FX/Niagara/Source/Niagara/Classes/NiagaraDataSet.h
@@ -212,7 +212,11 @@ public:
 	~FNiagaraDataSet();
 	FNiagaraDataSet& operator=(const FNiagaraDataSet&) = delete;
 
+#if !UE_BUILD_SHIPPING && !UE_BUILD_TEST
 	void Init(FNiagaraDataSetID InID, ENiagaraSimTarget InSimTarget, const FString& InDebugName);
+#else
+	void Init(FNiagaraDataSetID InID, ENiagaraSimTarget InSimTarget);
+#endif
 
 	/** Resets current data but leaves variable/layout information etc intact. */
 	void ResetBuffers();
@@ -245,7 +249,6 @@ public:
 	FORCEINLINE int32& GetIDAcquireTag() { return IDAcquireTag; }
 	FORCEINLINE void SetIDAcquireTag(int32 InTag) { IDAcquireTag = InTag; }
 
-
 	FORCEINLINE TArray<FNiagaraVariable>& GetVariables() { return Variables; }
 	FORCEINLINE uint32 GetNumVariables()const { return Variables.Num(); }
 	FORCEINLINE bool HasVariable(const FNiagaraVariable& Var)const { return Variables.Contains(Var); }
@@ -254,6 +257,7 @@ public:
 
 	void AddVariable(FNiagaraVariable& Variable);
 	void AddVariables(const TArray<FNiagaraVariable>& Vars);
+
 	/** Finalize the addition of variables and other setup before this data set can be used. */
 	void Finalize();
 	const FNiagaraVariableLayoutInfo* GetVariableLayout(const FNiagaraVariable& Var)const;
@@ -353,7 +357,9 @@ private:
 	*/
 	TArray<FNiagaraDataBuffer*, TInlineAllocator<2>> Data;
 
+#if !UE_BUILD_SHIPPING && !UE_BUILD_TEST
 	FString DebugName;
+#endif
 };
 
 /**
diff --git a/Engine/Plugins/FX/Niagara/Source/Niagara/Private/NiagaraDataSet.cpp b/Engine/Plugins/FX/Niagara/Source/Niagara/Private/NiagaraDataSet.cpp
--- a/Engine/Plugins/FX/Niagara/Source/Niagara/Private/NiagaraDataSet.cpp
+++ b/Engine/Plugins/FX/Niagara/Source/Niagara/Private/NiagaraDataSet.cpp
@@ -87,6 +87,7 @@ FNiagaraDataSet::~FNiagaraDataSet()
 	ReleaseBuffers();
 }
 
+#if !UE_BUILD_SHIPPING && !UE_BUILD_TEST
 void FNiagaraDataSet::Init(FNiagaraDataSetID InID, ENiagaraSimTarget InSimTarget, const FString& InDebugName)
 {
 	Reset();
@@ -94,6 +95,14 @@ void FNiagaraDataSet::Init(FNiagaraDataSetID InID, ENiagaraSimTarget InSimTarget
 	SimTarget = InSimTarget;
 	DebugName = InDebugName;
 }
+#else
+void FNiagaraDataSet::Init(FNiagaraDataSetID InID, ENiagaraSimTarget InSimTarget)
+{
+	Reset();
+	ID = InID;
+	SimTarget = InSimTarget;
+}
+#endif
 
 void FNiagaraDataSet::Reset()
 {
@@ -695,7 +704,11 @@ void FNiagaraDataBuffer::AllocateGPU(uint32 InNumInstances, FNiagaraGPUInstanceC
 			{
 				GPUBufferFloat.Release();
 			}
+#if !UE_BUILD_SHIPPING && !UE_BUILD_TEST 
 			GPUBufferFloat.Initialize(sizeof(float), NumElementsToAlloc * Owner->GetNumFloatComponents(), EPixelFormat::PF_R32_FLOAT, BUF_Static, *Owner->DebugName );
+#else
+			GPUBufferFloat.Initialize(sizeof(float), NumElementsToAlloc * Owner->GetNumFloatComponents(), EPixelFormat::PF_R32_FLOAT, BUF_Static);
+#endif
 		}
 		if (Owner->GetNumInt32Components())
 		{
@@ -703,7 +716,11 @@ void FNiagaraDataBuffer::AllocateGPU(uint32 InNumInstances, FNiagaraGPUInstanceC
 			{
 				GPUBufferInt.Release();
 			}
+#if !UE_BUILD_SHIPPING && !UE_BUILD_TEST 
 			GPUBufferInt.Initialize(sizeof(int32), NumElementsToAlloc * Owner->GetNumInt32Components(), EPixelFormat::PF_R32_SINT, BUF_Static, *Owner->DebugName);
+#else
+			GPUBufferInt.Initialize(sizeof(int32), NumElementsToAlloc * Owner->GetNumInt32Components(), EPixelFormat::PF_R32_SINT, BUF_Static);
+#endif
 		}
 	}
 }
diff --git a/Engine/Plugins/FX/Niagara/Source/Niagara/Private/NiagaraSystemSimulation.cpp b/Engine/Plugins/FX/Niagara/Source/Niagara/Private/NiagaraSystemSimulation.cpp
--- a/Engine/Plugins/FX/Niagara/Source/Niagara/Private/NiagaraSystemSimulation.cpp
+++ b/Engine/Plugins/FX/Niagara/Source/Niagara/Private/NiagaraSystemSimulation.cpp
@@ -269,32 +269,52 @@ bool FNiagaraSystemSimulation::Init(UNiagaraSystem* InSystem, UWorld* InWorld, b
 		//TODO: Move all layout data into the System!
 
 		//Initilize the main simulation dataset.
+#if !UE_BUILD_SHIPPING && !UE_BUILD_TEST 
 		DataSet.Init(FNiagaraDataSetID(), ENiagaraSimTarget::CPUSim, System->GetFullName());
+#else
+		DataSet.Init(FNiagaraDataSetID(), ENiagaraSimTarget::CPUSim);
+#endif
 		DataSet.AddVariables(System->GetSystemSpawnScript()->GetVMExecutableData().Attributes);
 		DataSet.AddVariables(System->GetSystemUpdateScript()->GetVMExecutableData().Attributes);
 		DataSet.Finalize();
 
 		//Initialize the data set for newly spawned systems.
+#if !UE_BUILD_SHIPPING && !UE_BUILD_TEST 
 		PendingSpawnDataSet.Init(FNiagaraDataSetID(), ENiagaraSimTarget::CPUSim, System->GetFullName());
+#else
+		PendingSpawnDataSet.Init(FNiagaraDataSetID(), ENiagaraSimTarget::CPUSim);
+#endif
 		PendingSpawnDataSet.AddVariables(System->GetSystemSpawnScript()->GetVMExecutableData().Attributes);
 		PendingSpawnDataSet.AddVariables(System->GetSystemUpdateScript()->GetVMExecutableData().Attributes);
 		PendingSpawnDataSet.Finalize();
 
 		//Initialize the dataset for paused systems.
+#if !UE_BUILD_SHIPPING && !UE_BUILD_TEST 
 		PausedInstanceData.Init(FNiagaraDataSetID(), ENiagaraSimTarget::CPUSim, System->GetFullName());
+#else
+		PausedInstanceData.Init(FNiagaraDataSetID(), ENiagaraSimTarget::CPUSim);
+#endif
 		PausedInstanceData.AddVariables(System->GetSystemSpawnScript()->GetVMExecutableData().Attributes);
 		PausedInstanceData.AddVariables(System->GetSystemUpdateScript()->GetVMExecutableData().Attributes);
 		PausedInstanceData.Finalize();
 
 		{
+#if !UE_BUILD_SHIPPING && !UE_BUILD_TEST 
 			SpawnInstanceParameterDataSet.Init(FNiagaraDataSetID(), ENiagaraSimTarget::CPUSim, System->GetFullName());
+#else
+			SpawnInstanceParameterDataSet.Init(FNiagaraDataSetID(), ENiagaraSimTarget::CPUSim);
+#endif
 			FNiagaraParameters* EngineParamsSpawn = System->GetSystemSpawnScript()->GetVMExecutableData().DataSetToParameters.Find(TEXT("Engine"));
 			if (EngineParamsSpawn != nullptr)
 			{
 				SpawnInstanceParameterDataSet.AddVariables(EngineParamsSpawn->Parameters);
 			}
 			SpawnInstanceParameterDataSet.Finalize();
+#if !UE_BUILD_SHIPPING && !UE_BUILD_TEST 
 			UpdateInstanceParameterDataSet.Init(FNiagaraDataSetID(), ENiagaraSimTarget::CPUSim, System->GetFullName());
+#else
+			UpdateInstanceParameterDataSet.Init(FNiagaraDataSetID(), ENiagaraSimTarget::CPUSim);
+#endif
 			FNiagaraParameters* EngineParamsUpdate = System->GetSystemUpdateScript()->GetVMExecutableData().DataSetToParameters.Find(TEXT("Engine"));
 			if (EngineParamsUpdate != nullptr)
 			{
