Init NiagaraSystem Instance Parameters post compile instead of at runtime.

diff --git a/Engine/Plugins/FX/Niagara/Source/Niagara/Classes/NiagaraSystem.h b/Engine/Plugins/FX/Niagara/Source/Niagara/Classes/NiagaraSystem.h
--- a/Engine/Plugins/FX/Niagara/Source/Niagara/Classes/NiagaraSystem.h
+++ b/Engine/Plugins/FX/Niagara/Source/Niagara/Classes/NiagaraSystem.h
@@ -47,7 +47,20 @@ struct FNiagaraEmitterCompiledData
 	FNiagaraVariable EmitterRandomSeedVar;
 };
 
+USTRUCT()
+struct FNiagaraSystemCompiledData
+{
+	GENERATED_USTRUCT_BODY()
 
+	UPROPERTY()
+	TArray<FNiagaraVariable> NumParticleVars;
+
+	UPROPERTY()
+	TArray<FNiagaraVariable> TotalSpawnedParticlesVars;
+
+	UPROPERTY()
+	FNiagaraParameterStore InstanceParamStore;
+};
 
 USTRUCT()
 struct FEmitterCompiledScriptPair
@@ -209,6 +222,8 @@ public:
 
 	const TArray<FNiagaraEmitterCompiledData>& GetEmitterCompiledData()const {	return EmitterCompiledData;	};
 
+	const FNiagaraSystemCompiledData& GetSystemCompiledData() const { return SystemCompiledData; };
+
 	bool UsesCollection(const UNiagaraParameterCollection* Collection)const;
 #if WITH_EDITORONLY_DATA
 	bool UsesEmitter(const UNiagaraEmitter* Emitter) const;
@@ -253,6 +268,8 @@ private:
 
 	void InitEmitterCompiledData();
 
+	void InitSystemCompiledData();
+
 	/** Helper for filling in precomputed variable names per emitter. Converts an emitter paramter "Emitter.XXXX" into it's real parameter name. */
 	void InitEmitterVariableAliasNames(FNiagaraEmitterCompiledData& EmitterCompiledDataToInit, const UNiagaraEmitter& InAssociatedEmitter);
 
@@ -293,6 +310,10 @@ protected:
 	UPROPERTY()
 	TArray<FNiagaraEmitterCompiledData> EmitterCompiledData;
 
+	//** Post compile generated data used for initializing System Instances during runtime. */
+	UPROPERTY()
+	FNiagaraSystemCompiledData SystemCompiledData;
+
 	/** Variables exposed to the outside work for tweaking*/
 	UPROPERTY()
 	FNiagaraUserRedirectionParameterStore ExposedParameters;
diff --git a/Engine/Plugins/FX/Niagara/Source/Niagara/Private/NiagaraModule.cpp b/Engine/Plugins/FX/Niagara/Source/Niagara/Private/NiagaraModule.cpp
--- a/Engine/Plugins/FX/Niagara/Source/Niagara/Private/NiagaraModule.cpp
+++ b/Engine/Plugins/FX/Niagara/Source/Niagara/Private/NiagaraModule.cpp
@@ -40,6 +40,10 @@ static FAutoConsoleVariableRef CVarEnableVerboseNiagaraChangeIdLogging(
 	ECVF_Default
 );
 
+#if WITH_EDITORONLY_DATA
+FNiagaraParameterStore INiagaraModule::FixedSystemInstanceParameters = FNiagaraParameterStore();
+#endif
+
 /**
 Use Shader Stages CVar.
 Enable the custom dispatch for multiple shader stages 
@@ -288,6 +292,10 @@ void INiagaraModule::StartupModule()
 	{
 		return new NiagaraEmitterInstanceBatcher(InFeatureLevel, InShaderPlatform);
 	}));
+
+#if WITH_EDITORONLY_DATA
+	InitFixedSystemInstanceParameterStore();
+#endif
 }
 
 void INiagaraModule::ShutdownRenderingResources()
@@ -419,6 +427,36 @@ bool INiagaraModule::IsTargetPlatformIncludedInLevelRangeForCook(const ITargetPl
 	return true;
 }
 
+#if WITH_EDITORONLY_DATA
+void INiagaraModule::InitFixedSystemInstanceParameterStore()
+{
+	FixedSystemInstanceParameters.AddParameter(SYS_PARAM_ENGINE_POSITION, true, false);
+	FixedSystemInstanceParameters.AddParameter(SYS_PARAM_ENGINE_ROTATION, true, false);
+	FixedSystemInstanceParameters.AddParameter(SYS_PARAM_ENGINE_SCALE, true, false);
+	FixedSystemInstanceParameters.AddParameter(SYS_PARAM_ENGINE_VELOCITY, true, false);
+	FixedSystemInstanceParameters.AddParameter(SYS_PARAM_ENGINE_X_AXIS, true, false);
+	FixedSystemInstanceParameters.AddParameter(SYS_PARAM_ENGINE_Y_AXIS, true, false);
+	FixedSystemInstanceParameters.AddParameter(SYS_PARAM_ENGINE_Z_AXIS, true, false);
+	FixedSystemInstanceParameters.AddParameter(SYS_PARAM_ENGINE_LOCAL_TO_WORLD, true, false);
+	FixedSystemInstanceParameters.AddParameter(SYS_PARAM_ENGINE_WORLD_TO_LOCAL, true, false);
+	FixedSystemInstanceParameters.AddParameter(SYS_PARAM_ENGINE_LOCAL_TO_WORLD_TRANSPOSED, true, false);
+	FixedSystemInstanceParameters.AddParameter(SYS_PARAM_ENGINE_WORLD_TO_LOCAL_TRANSPOSED, true, false);
+	FixedSystemInstanceParameters.AddParameter(SYS_PARAM_ENGINE_LOCAL_TO_WORLD_NO_SCALE, true, false);
+	FixedSystemInstanceParameters.AddParameter(SYS_PARAM_ENGINE_WORLD_TO_LOCAL_NO_SCALE, true, false);
+	FixedSystemInstanceParameters.AddParameter(SYS_PARAM_ENGINE_DELTA_TIME, true, false);
+	FixedSystemInstanceParameters.AddParameter(SYS_PARAM_ENGINE_TIME, true, false);
+	FixedSystemInstanceParameters.AddParameter(SYS_PARAM_ENGINE_REAL_TIME, true, false);
+	FixedSystemInstanceParameters.AddParameter(SYS_PARAM_ENGINE_INV_DELTA_TIME, true, false);
+	FixedSystemInstanceParameters.AddParameter(SYS_PARAM_ENGINE_TIME_SINCE_RENDERED, true, false);
+	FixedSystemInstanceParameters.AddParameter(SYS_PARAM_ENGINE_EXECUTION_STATE, true, false);
+	FixedSystemInstanceParameters.AddParameter(SYS_PARAM_ENGINE_LOD_DISTANCE, true, false);
+	FixedSystemInstanceParameters.AddParameter(SYS_PARAM_ENGINE_SYSTEM_NUM_EMITTERS, true, false);
+	FixedSystemInstanceParameters.AddParameter(SYS_PARAM_ENGINE_SYSTEM_NUM_EMITTERS_ALIVE, true, false);
+	FixedSystemInstanceParameters.AddParameter(SYS_PARAM_ENGINE_SYSTEM_AGE);
+	FixedSystemInstanceParameters.AddParameter(SYS_PARAM_ENGINE_SYSTEM_TICK_COUNT);
+}
+#endif
+
 void INiagaraModule::OnChangeDetailLevel(class IConsoleVariable* CVar)
 {
 	//Can only change the detail level at runtime on when not cooked.
diff --git a/Engine/Plugins/FX/Niagara/Source/Niagara/Private/NiagaraSystem.cpp b/Engine/Plugins/FX/Niagara/Source/Niagara/Private/NiagaraSystem.cpp
--- a/Engine/Plugins/FX/Niagara/Source/Niagara/Private/NiagaraSystem.cpp
+++ b/Engine/Plugins/FX/Niagara/Source/Niagara/Private/NiagaraSystem.cpp
@@ -388,6 +388,11 @@ void UNiagaraSystem::PostLoad()
 		InitEmitterCompiledData();
 	}
 
+	if (SystemCompiledData.InstanceParamStore.GetNumParameters() == 0)
+	{
+		InitSystemCompiledData();
+	}
+
 #if 0
 	UE_LOG(LogNiagara, Log, TEXT("Before"));
 	for (FNiagaraEmitterHandle& EmitterHandle : EmitterHandles)
@@ -799,6 +804,7 @@ bool UNiagaraSystem::QueryCompileComplete(bool bWait, bool bDoPost, bool bDoNotA
 		}
 
 		InitEmitterCompiledData();
+		InitSystemCompiledData();
 
 		// Prepare rapid iteration parameters for execution.
 		TArray<UNiagaraScript*> Scripts;
@@ -1083,6 +1089,39 @@ void UNiagaraSystem::InitEmitterCompiledData()
 		}
 	}
 }
+
+void UNiagaraSystem::InitSystemCompiledData()
+{
+	SystemCompiledData.NumParticleVars.Empty();
+	SystemCompiledData.TotalSpawnedParticlesVars.Empty();
+	SystemCompiledData.InstanceParamStore.Empty();
+
+	SystemCompiledData.InstanceParamStore = FNiagaraParameterStore(INiagaraModule::GetFixedSystemInstanceParameterStore());
+	ExposedParameters.CopyParametersTo(SystemCompiledData.InstanceParamStore, false, FNiagaraParameterStore::EDataInterfaceCopyMethod::Reference);
+
+	for (const FNiagaraEmitterHandle& PerEmitterHandle : EmitterHandles)
+	{
+		const UNiagaraEmitter* Emitter = PerEmitterHandle.GetInstance();
+		if (ensureMsgf(Emitter != nullptr, TEXT("Failed to get Emitter Instance from Emitter Handle when post compiling Niagara System!")))
+		{
+			const FString EmitterName = Emitter->GetUniqueEmitterName();
+			{
+				FNiagaraVariable Var = SYS_PARAM_ENGINE_EMITTER_NUM_PARTICLES;
+				const FString ParamName = Var.GetName().ToString().Replace(TEXT("Emitter"), *EmitterName);
+				Var.SetName(*ParamName);
+				SystemCompiledData.InstanceParamStore.AddParameter(Var, true, false);
+				SystemCompiledData.NumParticleVars.Add(Var);
+			}
+			{
+				FNiagaraVariable Var = SYS_PARAM_ENGINE_EMITTER_TOTAL_SPAWNED_PARTICLES;
+				const FString ParamName = Var.GetName().ToString().Replace(TEXT("Emitter"), *EmitterName);
+				Var.SetName(*ParamName);
+				SystemCompiledData.InstanceParamStore.AddParameter(Var, true, false);
+				SystemCompiledData.TotalSpawnedParticlesVars.Add(Var);
+			}
+		}
+	}
+}
 #endif
 
 TStatId UNiagaraSystem::GetStatID(bool bGameThread, bool bConcurrent)const
diff --git a/Engine/Plugins/FX/Niagara/Source/Niagara/Private/NiagaraSystemInstance.cpp b/Engine/Plugins/FX/Niagara/Source/Niagara/Private/NiagaraSystemInstance.cpp
--- a/Engine/Plugins/FX/Niagara/Source/Niagara/Private/NiagaraSystemInstance.cpp
+++ b/Engine/Plugins/FX/Niagara/Source/Niagara/Private/NiagaraSystemInstance.cpp
@@ -810,59 +810,10 @@ void FNiagaraSystemInstance::ReInitInternal()
 	Emitters.Reset();
 	InitEmitters();
 	
-	InstanceParameters.Reset();
-	InstanceParameters.AddParameter(SYS_PARAM_ENGINE_POSITION, true, false);
-	InstanceParameters.AddParameter(SYS_PARAM_ENGINE_ROTATION, true, false);
-	InstanceParameters.AddParameter(SYS_PARAM_ENGINE_SCALE, true, false);
-	InstanceParameters.AddParameter(SYS_PARAM_ENGINE_VELOCITY, true, false);
-	InstanceParameters.AddParameter(SYS_PARAM_ENGINE_X_AXIS, true, false);
-	InstanceParameters.AddParameter(SYS_PARAM_ENGINE_Y_AXIS, true, false);
-	InstanceParameters.AddParameter(SYS_PARAM_ENGINE_Z_AXIS, true, false);
-	InstanceParameters.AddParameter(SYS_PARAM_ENGINE_LOCAL_TO_WORLD, true, false);
-	InstanceParameters.AddParameter(SYS_PARAM_ENGINE_WORLD_TO_LOCAL, true, false);
-	InstanceParameters.AddParameter(SYS_PARAM_ENGINE_LOCAL_TO_WORLD_TRANSPOSED, true, false);
-	InstanceParameters.AddParameter(SYS_PARAM_ENGINE_WORLD_TO_LOCAL_TRANSPOSED, true, false);
-	InstanceParameters.AddParameter(SYS_PARAM_ENGINE_LOCAL_TO_WORLD_NO_SCALE, true, false);
-	InstanceParameters.AddParameter(SYS_PARAM_ENGINE_WORLD_TO_LOCAL_NO_SCALE, true, false);
-	InstanceParameters.AddParameter(SYS_PARAM_ENGINE_DELTA_TIME, true, false);
-	InstanceParameters.AddParameter(SYS_PARAM_ENGINE_TIME, true, false);
-	InstanceParameters.AddParameter(SYS_PARAM_ENGINE_REAL_TIME, true, false);
-	InstanceParameters.AddParameter(SYS_PARAM_ENGINE_INV_DELTA_TIME, true, false);
-	InstanceParameters.AddParameter(SYS_PARAM_ENGINE_TIME_SINCE_RENDERED, true, false);
-	InstanceParameters.AddParameter(SYS_PARAM_ENGINE_EXECUTION_STATE, true, false);
-	InstanceParameters.AddParameter(SYS_PARAM_ENGINE_LOD_DISTANCE, true, false);
-	InstanceParameters.AddParameter(SYS_PARAM_ENGINE_SYSTEM_NUM_EMITTERS, true, false);
-	InstanceParameters.AddParameter(SYS_PARAM_ENGINE_SYSTEM_NUM_EMITTERS_ALIVE, true, false);
-	InstanceParameters.AddParameter(SYS_PARAM_ENGINE_SYSTEM_AGE);
-	InstanceParameters.AddParameter(SYS_PARAM_ENGINE_SYSTEM_TICK_COUNT);
-
-	// This is required for user default data interface's (like say static meshes) to be set up properly.
-	// Additionally, it must happen here for data to be properly found below.
-	const bool bOnlyAdd = false;
-	System->GetExposedParameters().CopyParametersTo(InstanceParameters, bOnlyAdd, FNiagaraParameterStore::EDataInterfaceCopyMethod::Reference);
-
-	TArray<FNiagaraVariable> NumParticleVars;
-	TArray<FNiagaraVariable> TotalSpawnedParticlesVars;
-	for (int32 i = 0; i < Emitters.Num(); i++)
-	{
-		TSharedRef<FNiagaraEmitterInstance, ESPMode::ThreadSafe> Simulation = Emitters[i];
-		FString EmitterName = Simulation->GetEmitterHandle().GetInstance()->GetUniqueEmitterName();
-		
-		{
-			FNiagaraVariable Var = SYS_PARAM_ENGINE_EMITTER_NUM_PARTICLES;
-			FString ParamName = Var.GetName().ToString().Replace(TEXT("Emitter"), *EmitterName);
-			Var.SetName(*ParamName);
-			InstanceParameters.AddParameter(Var, true, false);
-			NumParticleVars.Add(Var);
-		}
-		{
-			FNiagaraVariable Var = SYS_PARAM_ENGINE_EMITTER_TOTAL_SPAWNED_PARTICLES;
-			FString ParamName = Var.GetName().ToString().Replace(TEXT("Emitter"), *EmitterName);
-			Var.SetName(*ParamName);
-			InstanceParameters.AddParameter(Var, true, false);
-			TotalSpawnedParticlesVars.Add(Var);
-		}
-	}
+	const FNiagaraSystemCompiledData& SystemCompiledData = System->GetSystemCompiledData();
+	InstanceParameters = FNiagaraParameterStore(SystemCompiledData.InstanceParamStore);
+	const TArray<FNiagaraVariable>& NumParticleVars = SystemCompiledData.NumParticleVars;
+	const TArray<FNiagaraVariable>& TotalSpawnedParticlesVars = SystemCompiledData.TotalSpawnedParticlesVars;
 
 	// Make sure all parameters are added before initializing the bindings, otherwise parameter store layout changes might invalidate the bindings.
 	OwnerPositionParam.Init(InstanceParameters, SYS_PARAM_ENGINE_POSITION);
diff --git a/Engine/Plugins/FX/Niagara/Source/Niagara/Public/NiagaraModule.h b/Engine/Plugins/FX/Niagara/Source/Niagara/Public/NiagaraModule.h
--- a/Engine/Plugins/FX/Niagara/Source/Niagara/Public/NiagaraModule.h
+++ b/Engine/Plugins/FX/Niagara/Source/Niagara/Public/NiagaraModule.h
@@ -17,6 +17,7 @@ class FNiagaraCompileOptions;
 class FNiagaraCompileRequestDataBase;
 class INiagaraMergeManager;
 class INiagaraEditorOnlyDataUtilities;
+struct FNiagaraParameterStore;
 
 extern NIAGARA_API int32 GEnableVerboseNiagaraChangeIdLogging;
 
@@ -155,7 +156,15 @@ public:
 	FORCEINLINE static const FNiagaraVariable&  GetVar_DataInstance_Alive() { return DataInstance_Alive; }
 	FORCEINLINE static const FNiagaraVariable&  GetVar_BeginDefaults() { return Translator_BeginDefaults; }
 
+#if WITH_EDITORONLY_DATA
+	FORCEINLINE static const FNiagaraParameterStore& GetFixedSystemInstanceParameterStore() { return FixedSystemInstanceParameters; }
+#endif
 private:
+
+#if WITH_EDITORONLY_DATA
+	void InitFixedSystemInstanceParameterStore();
+#endif
+
 	FOnProcessQueue OnProcessQueue;
 
 #if WITH_EDITORONLY_DATA
@@ -251,5 +260,9 @@ private:
 	static FNiagaraVariable ScriptUsage;
 	static FNiagaraVariable DataInstance_Alive;
 	static FNiagaraVariable Translator_BeginDefaults;
+
+#if WITH_EDITORONLY_DATA
+	static FNiagaraParameterStore FixedSystemInstanceParameters;
+#endif
 };
 
