#Niagara Added more info to warnings about IsSynchronizedWithStruct failures

diff --git a/Engine/Plugins/FX/Niagara/Source/NiagaraEditor/Private/NiagaraNodeDataSetBase.cpp b/Engine/Plugins/FX/Niagara/Source/NiagaraEditor/Private/NiagaraNodeDataSetBase.cpp
--- a/Engine/Plugins/FX/Niagara/Source/NiagaraEditor/Private/NiagaraNodeDataSetBase.cpp
+++ b/Engine/Plugins/FX/Niagara/Source/NiagaraEditor/Private/NiagaraNodeDataSetBase.cpp
@@ -117,7 +117,7 @@ bool UNiagaraNodeDataSetBase::IsSynchronizedWithStruct(bool bIgnoreConditionVari
 					bFoundIssues = true;
 					if (bLogIssues)
 					{
-						UE_LOG(LogNiagaraEditor, Warning, TEXT("Unable to find matching variable for struct property: '%s' ... possible add?"), *Property->GetName());
+						UE_LOG(LogNiagaraEditor, Warning, TEXT("%s Unable to find matching variable for struct property: '%s' ... possible add?"), *GetPathNameSafe(this), *GetPathNameSafe(Property));
 					}
 					if (Issues != nullptr)
 					{
@@ -134,7 +134,7 @@ bool UNiagaraNodeDataSetBase::IsSynchronizedWithStruct(bool bIgnoreConditionVari
 							bFoundIssues = true;
 							if (bLogIssues)
 							{
-								UE_LOG(LogNiagaraEditor, Warning, TEXT("Matching variable for struct property '%s', but different type:  Existing Type:'%s' vs New Type:'%s' ... possible type change in user-defined struct?"), *Property->GetName(), *VarFound->GetType().GetName(), *TypeDefFound.GetName());
+								UE_LOG(LogNiagaraEditor, Warning, TEXT("%s Matching variable for struct property '%s', but different type:  Existing Type:'%s' vs New Type:'%s' ... possible type change in user-defined struct?"), *GetPathNameSafe(this), *GetPathNameSafe(Property), *VarFound->GetType().GetName(), *TypeDefFound.GetName());
 							}
 							if (Issues != nullptr)
 							{
@@ -147,7 +147,7 @@ bool UNiagaraNodeDataSetBase::IsSynchronizedWithStruct(bool bIgnoreConditionVari
 						bFoundIssues = true;
 						if (bLogIssues)
 						{
-							UE_LOG(LogNiagaraEditor, Warning, TEXT("Matching variable for struct property '%s', but different type:  Existing Type:'%s' vs New Type:'Unsupported' ... possible type change in user-defined struct?"), *Property->GetName(), *VarFound->GetType().GetName());
+							UE_LOG(LogNiagaraEditor, Warning, TEXT("%s Matching variable for struct property '%s', but different type:  Existing Type:'%s' vs New Type:'Unsupported' ... possible type change in user-defined struct?"), *GetPathNameSafe(this), *GetPathNameSafe(Property), *VarFound->GetType().GetName());
 						}
 						if (Issues != nullptr)
 						{
@@ -178,7 +178,7 @@ bool UNiagaraNodeDataSetBase::IsSynchronizedWithStruct(bool bIgnoreConditionVari
 					bFoundIssues = true;
 					if (bLogIssues)
 					{
-						UE_LOG(LogNiagaraEditor, Warning, TEXT("Unable to find matching struct property for variable: '%s' ... possible removal?"), *Var.GetName().ToString());
+						UE_LOG(LogNiagaraEditor, Warning, TEXT("%s Unable to find matching struct property for variable: '%s' ... possible removal?"), *GetPathNameSafe(this), *Var.GetName().ToString());
 					}
 
 					if (Issues != nullptr)
@@ -194,7 +194,7 @@ bool UNiagaraNodeDataSetBase::IsSynchronizedWithStruct(bool bIgnoreConditionVari
 	{
 		if (bLogIssues)
 		{
-			UE_LOG(LogNiagaraEditor, Warning, TEXT("Unable to find matching Niagara Type: %s"), *DataSet.Name.ToString());
+			UE_LOG(LogNiagaraEditor, Warning, TEXT("%s Unable to find matching Niagara Type: %s"), *GetPathNameSafe(this), *DataSet.Name.ToString());
 		}
 
 		if (Issues != nullptr)
