- Handle reimport of StaticMeshes with Niagara Mesh Renderer

diff --git a/Engine/Plugins/FX/Niagara/Source/Niagara/Private/NiagaraMeshRendererProperties.cpp b/Engine/Plugins/FX/Niagara/Source/Niagara/Private/NiagaraMeshRendererProperties.cpp
--- a/Engine/Plugins/FX/Niagara/Source/Niagara/Private/NiagaraMeshRendererProperties.cpp
+++ b/Engine/Plugins/FX/Niagara/Source/Niagara/Private/NiagaraMeshRendererProperties.cpp
@@ -159,9 +159,14 @@ uint32 UNiagaraMeshRendererProperties::GetNumIndicesPerInstance() const
 void UNiagaraMeshRendererProperties::PostLoad()
 {
 	Super::PostLoad();
+#if WITH_EDITOR
+	if (GIsEditor && (ParticleMesh != nullptr))
+	{
+		ParticleMesh->GetOnMeshChanged().AddUObject(this, &UNiagaraMeshRendererProperties::OnMeshChanged);
+	}
+#endif
 }
 
-
 #if WITH_EDITORONLY_DATA
 bool UNiagaraMeshRendererProperties::IsMaterialValidForRenderer(UMaterial* Material, FText& InvalidMessage)
 {
@@ -213,9 +218,60 @@ const TArray<FNiagaraVariable>& UNiagaraMeshRendererProperties::GetOptionalAttri
 	return Attrs;
 }
 
+void UNiagaraMeshRendererProperties::BeginDestroy()
+{
+	Super::BeginDestroy();
+#if WITH_EDITOR
+	if (GIsEditor && (ParticleMesh != nullptr))
+	{
+		ParticleMesh->GetOnMeshChanged().RemoveAll(this);
+	}
+#endif
+}
+
+void UNiagaraMeshRendererProperties::PreEditChange(class UProperty* PropertyThatWillChange)
+{
+	Super::PreEditChange(PropertyThatWillChange);
+
+	static FName ParticleMeshName(TEXT("ParticleMesh"));
+	if ( PropertyThatWillChange->GetFName() == FName(ParticleMeshName))
+	{
+		if (ParticleMesh != nullptr)
+		{
+			ParticleMesh->GetOnMeshChanged().RemoveAll(this);
+		}
+	}
+}
+
 void UNiagaraMeshRendererProperties::PostEditChangeProperty(FPropertyChangedEvent& PropertyChangedEvent)
 {
-	if (ParticleMesh && PropertyChangedEvent.Property && PropertyChangedEvent.Property->GetName() == "ParticleMesh")
+	static FName ParticleMeshName(TEXT("ParticleMesh"));
+	if (ParticleMesh && PropertyChangedEvent.Property && PropertyChangedEvent.Property->GetFName() == ParticleMeshName)
+	{
+		// We only need to check material usage as we will invalidate any renderers later on
+		CheckMaterialUsage();
+		ParticleMesh->GetOnMeshChanged().AddUObject(this, &UNiagaraMeshRendererProperties::OnMeshChanged);
+	}
+
+	Super::PostEditChangeProperty(PropertyChangedEvent);
+}
+
+void UNiagaraMeshRendererProperties::OnMeshChanged()
+{
+	FNiagaraSystemUpdateContext ReregisterContext;
+
+	UNiagaraEmitter* Emitter = Cast<UNiagaraEmitter>(GetOuter());
+	if (Emitter != nullptr)
+	{
+		ReregisterContext.Add(Emitter, true);
+	}
+
+	CheckMaterialUsage();
+}
+
+void UNiagaraMeshRendererProperties::CheckMaterialUsage()
+{
+	if ( ParticleMesh != nullptr )
 	{
 		const FStaticMeshLODResources& LODModel = ParticleMesh->RenderData->LODResources[0];
 		for (int32 SectionIndex = 0; SectionIndex < LODModel.Sections.Num(); SectionIndex++)
@@ -231,4 +287,4 @@ void UNiagaraMeshRendererProperties::PostEditChangeProperty(FPropertyChangedEven
 	}
 }
 
-#endif // WITH_EDITORONLY_DATA
\ No newline at end of file
+#endif // WITH_EDITORONLY_DATA
diff --git a/Engine/Plugins/FX/Niagara/Source/Niagara/Public/NiagaraMeshRendererProperties.h b/Engine/Plugins/FX/Niagara/Source/Niagara/Public/NiagaraMeshRendererProperties.h
--- a/Engine/Plugins/FX/Niagara/Source/Niagara/Public/NiagaraMeshRendererProperties.h
+++ b/Engine/Plugins/FX/Niagara/Source/Niagara/Public/NiagaraMeshRendererProperties.h
@@ -64,6 +64,8 @@ public:
 	//UObject Interface
 	virtual void PostInitProperties() override;
 #if WITH_EDITORONLY_DATA
+	virtual void BeginDestroy() override;
+	virtual void PreEditChange(class UProperty* PropertyThatWillChange) override;
 	virtual void PostEditChangeProperty(struct FPropertyChangedEvent& PropertyChangedEvent) override;
 #endif// WITH_EDITORONLY_DATA
 	//UObject Interface END
@@ -82,6 +84,9 @@ public:
 	virtual void FixMaterial(UMaterial* Material) override;
 	virtual const TArray<FNiagaraVariable>& GetRequiredAttributes() override;
 	virtual const TArray<FNiagaraVariable>& GetOptionalAttributes() override;
+
+	void OnMeshChanged();
+	void CheckMaterialUsage();
 #endif // WITH_EDITORONLY_DATA
 
 	virtual uint32 GetNumIndicesPerInstance() const final override;
