Stop pointlessly allocating temporary parameter stores.

diff --git a/Engine/Plugins/FX/Niagara/Source/Niagara/Private/NiagaraSystem.cpp b/Engine/Plugins/FX/Niagara/Source/Niagara/Private/NiagaraSystem.cpp
--- a/Engine/Plugins/FX/Niagara/Source/Niagara/Private/NiagaraSystem.cpp
+++ b/Engine/Plugins/FX/Niagara/Source/Niagara/Private/NiagaraSystem.cpp
@@ -1096,7 +1096,7 @@ void UNiagaraSystem::InitSystemCompiledData()
 	SystemCompiledData.TotalSpawnedParticlesVars.Empty();
 	SystemCompiledData.InstanceParamStore.Empty();
 
-	SystemCompiledData.InstanceParamStore = FNiagaraParameterStore(INiagaraModule::GetFixedSystemInstanceParameterStore());
+	SystemCompiledData.InstanceParamStore = INiagaraModule::GetFixedSystemInstanceParameterStore();
 	ExposedParameters.CopyParametersTo(SystemCompiledData.InstanceParamStore, false, FNiagaraParameterStore::EDataInterfaceCopyMethod::Reference);
 
 	for (const FNiagaraEmitterHandle& PerEmitterHandle : EmitterHandles)
diff --git a/Engine/Plugins/FX/Niagara/Source/Niagara/Private/NiagaraSystemInstance.cpp b/Engine/Plugins/FX/Niagara/Source/Niagara/Private/NiagaraSystemInstance.cpp
--- a/Engine/Plugins/FX/Niagara/Source/Niagara/Private/NiagaraSystemInstance.cpp
+++ b/Engine/Plugins/FX/Niagara/Source/Niagara/Private/NiagaraSystemInstance.cpp
@@ -811,7 +811,7 @@ void FNiagaraSystemInstance::ReInitInternal()
 	InitEmitters();
 	
 	const FNiagaraSystemCompiledData& SystemCompiledData = System->GetSystemCompiledData();
-	InstanceParameters = FNiagaraParameterStore(SystemCompiledData.InstanceParamStore);
+	InstanceParameters = SystemCompiledData.InstanceParamStore;
 	const TArray<FNiagaraVariable>& NumParticleVars = SystemCompiledData.NumParticleVars;
 	const TArray<FNiagaraVariable>& TotalSpawnedParticlesVars = SystemCompiledData.TotalSpawnedParticlesVars;
 
