Creating a way to insert non-data interface "fast path" functions into the VM.

diff --git a/Engine/Plugins/FX/Niagara/Source/Niagara/Classes/NiagaraScript.h b/Engine/Plugins/FX/Niagara/Source/Niagara/Classes/NiagaraScript.h
--- a/Engine/Plugins/FX/Niagara/Source/Niagara/Classes/NiagaraScript.h
+++ b/Engine/Plugins/FX/Niagara/Source/Niagara/Classes/NiagaraScript.h
@@ -237,6 +237,9 @@ public:
 	UPROPERTY()
 	FNiagaraScriptDataUsageInfo DataUsage;
 
+	UPROPERTY()
+	TArray<FNiagaraFunctionSignature> AdditionalExternalFunctions;
+
 	/** Information about all data interfaces used by this script. */
 	UPROPERTY()
 	TArray<FNiagaraScriptDataInterfaceCompileInfo> DataInterfaceInfo;
diff --git a/Engine/Plugins/FX/Niagara/Source/Niagara/Private/NiagaraFunctionLibrary.cpp b/Engine/Plugins/FX/Niagara/Source/Niagara/Private/NiagaraFunctionLibrary.cpp
--- a/Engine/Plugins/FX/Niagara/Source/Niagara/Private/NiagaraFunctionLibrary.cpp
+++ b/Engine/Plugins/FX/Niagara/Source/Niagara/Private/NiagaraFunctionLibrary.cpp
@@ -7,9 +7,15 @@
 #include "NiagaraComponent.h"
 #include "NiagaraSystem.h"
 #include "ContentStreaming.h"
+#include "Internationalization\Internationalization.h"
 
 #include "NiagaraWorldManager.h"
 #include "NiagaraDataInterfaceStaticMesh.h"
+#include "NiagaraStats.h"
+
+#define LOCTEXT_NAMESPACE "NiagaraFunctionLibrary"
+
+//DECLARE_CYCLE_STAT(TEXT("FastDot4"), STAT_NiagaraFastDot4, STATGROUP_Niagara);
 
 #if WITH_EDITOR
 int32 GForceNiagaraSpawnAttachedSolo = 0;
@@ -21,6 +27,17 @@ static FAutoConsoleVariableRef CVarForceNiagaraSpawnAttachedSolo(
 );
 #endif
 
+
+int32 GAllowFastPathFunctionLibrary = 0;
+static FAutoConsoleVariableRef CVarAllowFastPathFunctionLibrary(
+	TEXT("fx.AllowFastPathFunctionLibrary"),
+	GAllowFastPathFunctionLibrary,
+	TEXT("If > 0 Allow the graph to insert custom fastpath operations into the graph.\n"),
+	ECVF_Default
+);
+
+TArray<FNiagaraFunctionSignature> UNiagaraFunctionLibrary::VectorVMOps;
+
 UNiagaraFunctionLibrary::UNiagaraFunctionLibrary(const FObjectInitializer& ObjectInitializer)
 : Super(ObjectInitializer)
 {
@@ -345,3 +362,251 @@ UNiagaraParameterCollectionInstance* UNiagaraFunctionLibrary::GetNiagaraParamete
 	}
 	return nullptr;
 }
+
+
+const TArray<FNiagaraFunctionSignature>& UNiagaraFunctionLibrary::GetVectorVMFastPathOps()
+{
+	if (VectorVMOps.Num() == 0)
+	{
+		InitVectorVMFastPathOps();
+	}
+	return VectorVMOps;
+}
+
+struct FVectorKernelFastDot4
+{
+	static void Exec(FVectorVMContext& Context)
+	{
+		//SCOPE_CYCLE_COUNTER(STAT_NiagaraFastDot4);
+
+		VectorVM::FExternalFuncInputHandler<float>InVecA[4] =
+		{
+			VectorVM::FExternalFuncInputHandler<float>(Context), 
+			VectorVM::FExternalFuncInputHandler<float>(Context), 
+			VectorVM::FExternalFuncInputHandler<float>(Context),
+			VectorVM::FExternalFuncInputHandler<float>(Context)
+		};
+
+		VectorVM::FExternalFuncInputHandler<float>InVecB[4] =
+		{
+			VectorVM::FExternalFuncInputHandler<float>(Context),
+			VectorVM::FExternalFuncInputHandler<float>(Context),
+			VectorVM::FExternalFuncInputHandler<float>(Context),
+			VectorVM::FExternalFuncInputHandler<float>(Context)
+		};
+
+
+		VectorVM::FExternalFuncRegisterHandler<float> OutValue(Context);
+
+		VectorRegister* RESTRICT AX = (VectorRegister *)InVecA[0].GetDest();
+		VectorRegister* RESTRICT AY = (VectorRegister *)InVecA[1].GetDest();
+		VectorRegister* RESTRICT AZ = (VectorRegister *)InVecA[2].GetDest();
+		VectorRegister* RESTRICT AW = (VectorRegister *)InVecA[3].GetDest();
+
+		VectorRegister* RESTRICT BX = (VectorRegister *)InVecB[0].GetDest();
+		VectorRegister* RESTRICT BY = (VectorRegister *)InVecB[1].GetDest();
+		VectorRegister* RESTRICT BZ = (VectorRegister *)InVecB[2].GetDest();
+		VectorRegister* RESTRICT BW = (VectorRegister *)InVecB[3].GetDest();
+		VectorRegister* RESTRICT Out = (VectorRegister *)OutValue.GetDest();
+
+		int32 LoopInstances = Align(Context.NumInstances, 4) / 4;
+		for (int32 i = 0; i < LoopInstances; ++i)
+		{
+
+			VectorRegister AVX0= VectorLoadAligned(&AX[i]);
+			VectorRegister AVY0= VectorLoadAligned(&AY[i]);
+			VectorRegister AVZ0= VectorLoadAligned(&AZ[i]);
+			VectorRegister AVW0= VectorLoadAligned(&AW[i]);
+			VectorRegister BVX0= VectorLoadAligned(&BX[i]);
+			VectorRegister BVY0= VectorLoadAligned(&BY[i]);
+			VectorRegister BVZ0= VectorLoadAligned(&BZ[i]);
+			VectorRegister BVW0= VectorLoadAligned(&BW[i]);
+
+			/*
+				 R[19] = :mul(R[21], R[25]);
+				 R[21] = :mad(R[20], R[24], R[19]);
+				 R[19] = :mad(R[22], R[26], R[21]);
+				 R[20] = :mad(R[23], R[27], R[19]);
+			*/
+			VectorRegister AMBX0	= VectorMultiply(AVX0, BVX0);
+			VectorRegister AMBXY0= VectorMultiplyAdd(AVY0, BVY0, AMBX0);
+			VectorRegister AMBXYZ0= VectorMultiplyAdd(AVZ0, BVZ0, AMBXY0);
+			VectorRegister AMBXYZW0= VectorMultiplyAdd(AVW0, BVW0, AMBXYZ0);
+			VectorStoreAligned(AMBXYZW0, &Out[i]) ;
+
+			/*
+			(float Sum = 0.0f;
+			for (int32 VecIdx = 0; VecIdx < 4; VecIdx++)
+			{
+				Sum += InVecA[VecIdx].GetAndAdvance() * InVecB[VecIdx].GetAndAdvance();
+			}
+
+			*OutValue.GetDestAndAdvance() = Sum;
+			*/
+		}
+	}
+};
+
+struct FVectorKernelFastTransformPosition
+{
+	static void Exec(FVectorVMContext& Context)
+	{
+#if 0
+		TArray<VectorVM::FExternalFuncInputHandler<float>, TInlineAllocator<16>> InMatrix;
+		for (int i = 0; i < 16; i++)
+		{
+			InMatrix.Emplace(Context);
+		}
+
+		TArray<VectorVM::FExternalFuncInputHandler<float>, TInlineAllocator<16>> InVec;
+		for (int i = 0; i < 3; i++)
+		{
+			InVec.Emplace(Context);
+		}
+
+		TArray<VectorVM::FExternalFuncInputHandler<float>, TInlineAllocator<16>> OutVec;
+		for (int i = 0; i < 3; i++)
+		{
+			OutVec.Emplace(Context);
+		}
+
+		/*
+	29	| R[19] = :mul(R[20], R[26]);
+	30	| R[27] = :mul(R[21], R[26]);
+	31	| R[28] = :mul(R[22], R[26]);
+	32	| R[29] = :mul(R[23], R[26]);
+	33	| R[26] = :mul(R[20], R[25]);
+	34	| R[30] = :mul(R[21], R[25]);
+	35	| R[31] = :mul(R[22], R[25]);
+	36	| R[32] = :mul(R[23], R[25]);
+	37	| R[25] = :mul(R[20], R[24]);
+	38	| R[33] = :mul(R[21], R[24]);
+	39	| R[34] = :mul(R[22], R[24]);
+	40	| R[35] = :mul(R[23], R[24]);
+	41	| R[24] = :add(R[26], R[25]);
+	42	| R[25] = :add(R[30], R[33]);
+	43	| R[26] = :add(R[31], R[34]);
+	44	| R[30] = :add(R[32], R[35]);
+	45	| R[31] = :add(R[19], R[24]);
+	46	| R[19] = :add(R[27], R[25]);
+	47	| R[24] = :add(R[28], R[26]);
+	48	| R[25] = :add(R[29], R[30]);
+	49	| R[26] = :add(R[20], R[31]);
+	50	| R[20] = :add(R[21], R[19]);
+	51	| R[19] = :add(R[22], R[24]);
+		*/
+		/*float* RESTRICT M00 = (float *)InMatrix[0].GetDest();
+		float* RESTRICT M01 = (float *)InMatrix[1].GetDest();
+		float* RESTRICT M02 = (float *)InMatrix[2].GetDest();
+		float* RESTRICT M03 = (float *)InMatrix[3].GetDest();
+		float* RESTRICT M10 = (float *)InMatrix[4].GetDest();
+		float* RESTRICT M11 = (float *)InMatrix[5].GetDest();
+		float* RESTRICT M12 = (float *)InMatrix[6].GetDest();
+		float* RESTRICT M13 = (float *)InMatrix[7].GetDest(); 
+		float* RESTRICT M20 = (float *)InMatrix[8].GetDest();
+		float* RESTRICT M21 = (float *)InMatrix[9].GetDest();
+		float* RESTRICT M22 = (float *)InMatrix[10].GetDest();
+		float* RESTRICT M23 = (float *)InMatrix[11].GetDest();
+		float* RESTRICT M30 = (float *)InMatrix[12].GetDest();
+		float* RESTRICT M31 = (float *)InMatrix[13].GetDest();
+		float* RESTRICT M32 = (float *)InMatrix[14].GetDest();
+		float* RESTRICT M33 = (float *)InMatrix[15].GetDest();*/
+		TArray<float* RESTRICT, TInlineAllocator<16>> InMatrixDest;
+		for (int i = 0; i < 16; i++)
+		{
+			InMatrixDest.Emplace((float *)InMatrix[i].GetDest());
+		}
+
+		float* RESTRICT BX = (float *)InVec[0].GetDest();
+		float* RESTRICT BY = (float *)InVec[1].GetDest();
+		float* RESTRICT BZ = (float *)InVec[2].GetDest();
+		
+		float* RESTRICT OutValueX = (float *)OutVec[0].GetDest();
+		float* RESTRICT OutValueY = (float *)OutVec[1].GetDest();
+		float* RESTRICT OutValueZ = (float *)OutVec[2].GetDest();
+		
+		int32 LoopInstances = Align(Context.NumInstances, 1) /1;
+		for (int32 i = 0; i < Context.NumInstances; ++i)
+		{
+			/*FMatrix LocalMat;
+			for (int32 Row = 0; Row < 4; Row++)
+			{
+				for (int32 Column = 0; Column < 4; Column++)
+				{
+					LocalMat.M[Row][Column] = *InMatrixDest[Row * 4 + Column]++;
+				}
+			}*/
+
+			for (int32 Row = 0; Row < 4; Row++)
+			{
+				for (int32 Column = 0; Column < 4; Column++)
+				{
+					*InMatrixDest[Row * 4 + Column]++;
+				}
+			}
+			*BX++; *BY++; *BZ++;
+
+			//FVector LocalVec(*BX++, *BY++, *BZ++);
+
+			//FVector LocalOutVec = LocalMat.TransformPosition(LocalVec);
+
+			//*(OutValueX)++ = LocalOutVec.X;
+			//*(OutValueY)++ = LocalOutVec.Y;
+			//*(OutValueZ)++ = LocalOutVec.Z;
+		}
+#endif
+	}
+};
+
+const FName FastPathLibraryName(TEXT("FastPathLibrary"));
+const FName FastPathDot4Name(TEXT("FastPathDot4"));
+const FName FastPathTransformPositionName(TEXT("FastPathTransformPosition"));
+
+bool UNiagaraFunctionLibrary::GetVectorVMFastPathExternalFunction(const FVMExternalFunctionBindingInfo& BindingInfo, FVMExternalFunction &OutFunc)
+{
+	if (BindingInfo.Name == FastPathDot4Name)
+	{
+		OutFunc = FVMExternalFunction::CreateStatic(FVectorKernelFastDot4::Exec);
+		return true;
+	}
+	else if (BindingInfo.Name == FastPathTransformPositionName)
+	{
+		OutFunc = FVMExternalFunction::CreateStatic(FVectorKernelFastTransformPosition::Exec);
+		return true;
+	}
+	return false;
+}
+
+
+void UNiagaraFunctionLibrary::InitVectorVMFastPathOps()
+{
+	if (GAllowFastPathFunctionLibrary == 0)
+		return;
+
+	{
+		FNiagaraFunctionSignature Sig;
+		Sig.Name = FastPathDot4Name;
+		Sig.OwnerName = FastPathLibraryName;
+		Sig.bMemberFunction = false;
+		Sig.bRequiresContext = false;
+		Sig.SetDescription(LOCTEXT("FastPathDot4Desc", "Fast path for Vector4 dot product."));
+		Sig.Inputs.Add(FNiagaraVariable(FNiagaraTypeDefinition::GetVec4Def(), TEXT("A")));
+		Sig.Inputs.Add(FNiagaraVariable(FNiagaraTypeDefinition::GetVec4Def(), TEXT("B")));
+		Sig.Outputs.Add(FNiagaraVariable(FNiagaraTypeDefinition::GetFloatDef(), TEXT("Value")));
+		VectorVMOps.Add(Sig);
+	}
+	{
+		FNiagaraFunctionSignature Sig;
+		Sig.Name = FastPathTransformPositionName;
+		Sig.OwnerName = FastPathLibraryName;
+		Sig.bMemberFunction = false;
+		Sig.bRequiresContext = false;
+		Sig.SetDescription(LOCTEXT("FastPathTransformPositionDesc", "Fast path for Matrix4 transforming a Vector3 position"));
+		Sig.Inputs.Add(FNiagaraVariable(FNiagaraTypeDefinition::GetMatrix4Def(), TEXT("Mat")));
+		Sig.Inputs.Add(FNiagaraVariable(FNiagaraTypeDefinition::GetVec3Def(), TEXT("Position")));
+		Sig.Outputs.Add(FNiagaraVariable(FNiagaraTypeDefinition::GetVec3Def(), TEXT("PositionTransformed")));
+		VectorVMOps.Add(Sig);
+	}
+}
+
+#undef LOCTEXT_NAMESPACE
\ No newline at end of file
diff --git a/Engine/Plugins/FX/Niagara/Source/Niagara/Private/NiagaraScriptExecutionContext.cpp b/Engine/Plugins/FX/Niagara/Source/Niagara/Private/NiagaraScriptExecutionContext.cpp
--- a/Engine/Plugins/FX/Niagara/Source/Niagara/Private/NiagaraScriptExecutionContext.cpp
+++ b/Engine/Plugins/FX/Niagara/Source/Niagara/Private/NiagaraScriptExecutionContext.cpp
@@ -8,6 +8,7 @@
 #include "NiagaraWorldManager.h"
 #include "NiagaraSystemInstance.h"
 #include "NiagaraEmitterInstance.h"
+#include "NiagaraFunctionLibrary.h"
 #include "NiagaraGPUInstanceCountManager.h"
 
 DECLARE_CYCLE_STAT(TEXT("Register Setup"), STAT_NiagaraSimRegisterSetup, STATGROUP_Niagara);
@@ -93,6 +94,14 @@ bool FNiagaraScriptExecutionContext::Tick(FNiagaraSystemInstance* ParentSystemIn
 			bool bSuccessfullyMapped = true;
 			for (FVMExternalFunctionBindingInfo& BindingInfo : Script->GetVMExecutableData().CalledVMExternalFunctions)
 			{
+				// First check to see if we can pull from the fast path library..
+				FVMExternalFunction FuncBind;
+				if (UNiagaraFunctionLibrary::GetVectorVMFastPathExternalFunction(BindingInfo, FuncBind) && FuncBind.IsBound())
+				{
+					FunctionTable.Add(FuncBind);
+					continue;
+				}
+
 				for (int32 i = 0; i < Script->GetVMExecutableData().DataInterfaceInfo.Num(); i++)
 				{
 					FNiagaraScriptDataInterfaceCompileInfo& ScriptInfo = Script->GetVMExecutableData().DataInterfaceInfo[i];
diff --git a/Engine/Plugins/FX/Niagara/Source/Niagara/Public/NiagaraFunctionLibrary.h b/Engine/Plugins/FX/Niagara/Source/Niagara/Public/NiagaraFunctionLibrary.h
--- a/Engine/Plugins/FX/Niagara/Source/Niagara/Public/NiagaraFunctionLibrary.h
+++ b/Engine/Plugins/FX/Niagara/Source/Niagara/Public/NiagaraFunctionLibrary.h
@@ -7,6 +7,8 @@
 #include "Engine/EngineTypes.h"
 #include "Kismet/BlueprintFunctionLibrary.h"
 #include "NiagaraComponentPool.h"
+#include "NiagaraCommon.h"
+#include "VectorVM.h"
 #include "NiagaraFunctionLibrary.generated.h"
 
 class UNiagaraComponent;
@@ -22,6 +24,7 @@ class NIAGARA_API UNiagaraFunctionLibrary : public UBlueprintFunctionLibrary
 {
 	GENERATED_UCLASS_BODY()
 
+public:
 	/**
 	* Spawns a Niagara System at the specified world location/rotation
 	* @return			The spawned UNiagaraComponent
@@ -52,4 +55,11 @@ class NIAGARA_API UNiagaraFunctionLibrary : public UBlueprintFunctionLibrary
 	UFUNCTION(BlueprintCallable, Category = Niagara, meta = (Keywords = "niagara parameter collection", WorldContext = "WorldContextObject"))
 	static UNiagaraParameterCollectionInstance* GetNiagaraParameterCollection(UObject* WorldContextObject, UNiagaraParameterCollection* Collection);
 
+	static const TArray<FNiagaraFunctionSignature>& GetVectorVMFastPathOps();
+
+	static bool GetVectorVMFastPathExternalFunction(const FVMExternalFunctionBindingInfo& BindingInfo, FVMExternalFunction &OutFunc);
+
+private:
+	static void InitVectorVMFastPathOps();
+	static TArray<FNiagaraFunctionSignature> VectorVMOps;
 };
diff --git a/Engine/Plugins/FX/Niagara/Source/NiagaraEditor/Private/EdGraphSchema_Niagara.cpp b/Engine/Plugins/FX/Niagara/Source/NiagaraEditor/Private/EdGraphSchema_Niagara.cpp
--- a/Engine/Plugins/FX/Niagara/Source/NiagaraEditor/Private/EdGraphSchema_Niagara.cpp
+++ b/Engine/Plugins/FX/Niagara/Source/NiagaraEditor/Private/EdGraphSchema_Niagara.cpp
@@ -17,7 +17,7 @@
 #include "GraphEditorSettings.h"
 #include "GraphEditorActions.h"
 #include "NiagaraConstants.h"
-
+#include "NiagaraFunctionLibrary.h"
 #include "NiagaraScript.h"
 #include "NiagaraNodeOutput.h"
 #include "NiagaraNodeInput.h"
@@ -426,6 +426,25 @@ TArray<TSharedPtr<FNiagaraSchemaAction_NewNode> > UEdGraphSchema_Niagara::GetGra
 		{
 			AddScriptFunctionAction(LOCTEXT("Function Menu Title", "Functions"), FunctionScriptAsset);
 		}
+
+		// Insert custom opcodes here...
+		const TArray<FNiagaraFunctionSignature>& VMFastPathOps = UNiagaraFunctionLibrary::GetVectorVMFastPathOps();
+		for (const FNiagaraFunctionSignature& Sig : VMFastPathOps)
+		{
+			UNiagaraNodeFunctionCall* FunctionCallNode = NewObject<UNiagaraNodeFunctionCall>(OwnerOfTemporaries);
+			
+			FunctionCallNode->Signature = Sig;
+			FText AssetDesc;
+			FText Keywords;
+			
+			FString DisplayNameString = FName::NameToDisplayString(Sig.Name.ToString(), false);
+
+			const FText MenuDesc = FText::FromString(DisplayNameString);
+			const FText TooltipDesc = MenuDesc;
+
+			TSharedPtr<FNiagaraSchemaAction_NewNode> FunctionCallAction = AddNewNodeAction(NewActions, LOCTEXT("Function Menu Title", "Functions"), MenuDesc, *DisplayNameString, TooltipDesc, Keywords);
+			FunctionCallAction->NodeTemplate = FunctionCallNode;
+		}
 	}
 
 	//Add modules
diff --git a/Engine/Plugins/FX/Niagara/Source/NiagaraEditor/Private/NiagaraCompiler.cpp b/Engine/Plugins/FX/Niagara/Source/NiagaraEditor/Private/NiagaraCompiler.cpp
--- a/Engine/Plugins/FX/Niagara/Source/NiagaraEditor/Private/NiagaraCompiler.cpp
+++ b/Engine/Plugins/FX/Niagara/Source/NiagaraEditor/Private/NiagaraCompiler.cpp
@@ -884,6 +884,15 @@ FNiagaraCompileResults FHlslNiagaraCompiler::CompileScript(const FNiagaraCompile
 					}
 				}
 
+				if (!Sig)
+				{
+					Sig = TranslatorOutput->ScriptData.AdditionalExternalFunctions.FindByPredicate([&](const FNiagaraFunctionSignature& CheckSig)
+					{
+						FString SigSymbol = FHlslNiagaraTranslator::GetFunctionSignatureSymbol(CheckSig);
+						return SigSymbol == FuncInfo.Name;
+					});
+				}
+
 				if (Sig)
 				{
 					int32 NewBindingIdx = CompileResults.Data->CalledVMExternalFunctions.AddDefaulted();
diff --git a/Engine/Plugins/FX/Niagara/Source/NiagaraEditor/Private/NiagaraHlslTranslator.cpp b/Engine/Plugins/FX/Niagara/Source/NiagaraEditor/Private/NiagaraHlslTranslator.cpp
--- a/Engine/Plugins/FX/Niagara/Source/NiagaraEditor/Private/NiagaraHlslTranslator.cpp
+++ b/Engine/Plugins/FX/Niagara/Source/NiagaraEditor/Private/NiagaraHlslTranslator.cpp
@@ -5040,7 +5040,6 @@ void FHlslNiagaraTranslator::RegisterFunctionCall(ENiagaraScriptUsage ScriptUsag
 		SCOPE_CYCLE_COUNTER(STAT_NiagaraEditor_Module_NiagaraHLSLTranslator_RegisterFunctionCall_Signature);
 
 		check(InSignature.IsValid());
-		check(InSignature.bMemberFunction || bIsCustomHlsl);
 		check(Inputs.Num() > 0);
 
 		OutSignature = InSignature;
@@ -5066,6 +5065,14 @@ void FHlslNiagaraTranslator::RegisterFunctionCall(ENiagaraScriptUsage ScriptUsag
 				ExitFunction();
 			}
 		}
+		else if (!InSignature.bMemberFunction) // Fastpath or other provided function
+		{
+			if (INDEX_NONE == CompilationOutput.ScriptData.AdditionalExternalFunctions.Find(OutSignature))
+			{
+				CompilationOutput.ScriptData.AdditionalExternalFunctions.Add(OutSignature);
+			}
+			Functions.FindOrAdd(OutSignature);
+		}
 		else
 		{
 			int32 OwnerIdx = Inputs[0];
diff --git a/Engine/Plugins/FX/Niagara/Source/NiagaraEditor/Private/NiagaraNodeOp.cpp b/Engine/Plugins/FX/Niagara/Source/NiagaraEditor/Private/NiagaraNodeOp.cpp
--- a/Engine/Plugins/FX/Niagara/Source/NiagaraEditor/Private/NiagaraNodeOp.cpp
+++ b/Engine/Plugins/FX/Niagara/Source/NiagaraEditor/Private/NiagaraNodeOp.cpp
@@ -18,7 +18,10 @@ void UNiagaraNodeOp::AllocateDefaultPins()
 	const UEdGraphSchema_Niagara* Schema = GetDefault<UEdGraphSchema_Niagara>();
 
 	const FNiagaraOpInfo* OpInfo = FNiagaraOpInfo::GetOpInfo(OpName);
-	check(OpInfo);
+	if (!OpInfo)
+	{
+		return;
+	}
 
 	// Create input pins from the op
 	for (int32 SrcIndex = 0; SrcIndex < OpInfo->Inputs.Num(); ++SrcIndex)
@@ -92,6 +95,10 @@ void UNiagaraNodeOp::Compile(class FHlslNiagaraTranslator* Translator, TArray<in
 	const FNiagaraOpInfo* OpInfo = FNiagaraOpInfo::GetOpInfo(OpName);
 	if (!OpInfo)
 	{
+		FFormatNamedArguments Args;
+		Args.Add(TEXT("OpName"), FText::FromName(OpName));
+		FText Format = LOCTEXT("Unknown opcode", "Unknown opcode on {OpName} node.");
+		Translator->Error(FText::Format(Format, Args), this, nullptr);
 		return;
 	}
 
@@ -183,8 +190,14 @@ bool UNiagaraNodeOp::RefreshFromExternalChanges()
 ENiagaraNumericOutputTypeSelectionMode UNiagaraNodeOp::GetNumericOutputTypeSelectionMode() const
 {
 	const FNiagaraOpInfo* OpInfo = FNiagaraOpInfo::GetOpInfo(OpName);
-	check(OpInfo);
-	return OpInfo->NumericOuputTypeSelectionMode;
+	if (OpInfo)
+	{
+		return OpInfo->NumericOuputTypeSelectionMode;
+	}
+	else
+	{
+		return ENiagaraNumericOutputTypeSelectionMode::Largest;
+	}
 }
 
 bool UNiagaraNodeOp::AllowNiagaraTypeForAddPin(const FNiagaraTypeDefinition& InType)
