Fix crash adding new emitter to running system by fixing check in UNiagaraSystem::IsReadyToRun()

diff --git a/Engine/Plugins/FX/Niagara/Source/Niagara/Private/NiagaraSystem.cpp b/Engine/Plugins/FX/Niagara/Source/Niagara/Private/NiagaraSystem.cpp
--- a/Engine/Plugins/FX/Niagara/Source/Niagara/Private/NiagaraSystem.cpp
+++ b/Engine/Plugins/FX/Niagara/Source/Niagara/Private/NiagaraSystem.cpp
@@ -508,6 +508,14 @@ bool UNiagaraSystem::IsReadyToRun() const
 	{
 		return false;
 	}
+
+	/* Check that our post compile data is in sync with the current emitter handles count. If we have just added a new emitter handle, we will not have any outstanding compilation requests as the new compile
+	 * will not be added to the outstanding compilation requests until the next tick.
+	 */
+	if (EmitterHandles.Num() != EmitterCompiledData.Num() || EmitterHandles.Num() != SystemCompiledData.NumParticleVars.Num() || EmitterHandles.Num() != SystemCompiledData.TotalSpawnedParticlesVars.Num())
+	{
+		return false;
+	}
 #endif
 
 	if (SystemSpawnScript->IsScriptCompilationPending(false) || 
