Fixing improper direction of slash in include path.

diff --git a/Engine/Plugins/FX/Niagara/Source/Niagara/Private/NiagaraFunctionLibrary.cpp b/Engine/Plugins/FX/Niagara/Source/Niagara/Private/NiagaraFunctionLibrary.cpp
--- a/Engine/Plugins/FX/Niagara/Source/Niagara/Private/NiagaraFunctionLibrary.cpp
+++ b/Engine/Plugins/FX/Niagara/Source/Niagara/Private/NiagaraFunctionLibrary.cpp
@@ -7,7 +7,7 @@
 #include "NiagaraComponent.h"
 #include "NiagaraSystem.h"
 #include "ContentStreaming.h"
-#include "Internationalization\Internationalization.h"
+#include "Internationalization/Internationalization.h"
 
 #include "NiagaraWorldManager.h"
 #include "NiagaraDataInterfaceStaticMesh.h"
