#UE4 Change verbosity on a couple logs

diff --git a/Engine/Plugins/FX/Niagara/Source/Niagara/Private/NiagaraScript.cpp b/Engine/Plugins/FX/Niagara/Source/Niagara/Private/NiagaraScript.cpp
--- a/Engine/Plugins/FX/Niagara/Source/Niagara/Private/NiagaraScript.cpp
+++ b/Engine/Plugins/FX/Niagara/Source/Niagara/Private/NiagaraScript.cpp
@@ -980,7 +980,7 @@ void UNiagaraScript::RequestCompile(bool bForceCompile)
 	}
 	else
 	{
-		UE_LOG(LogNiagara, Log, TEXT("Script '%s' is in-sync skipping compile.."), *GetFullName());
+		UE_LOG(LogNiagara, Verbose, TEXT("Script '%s' is in-sync skipping compile.."), *GetFullName());
 	}
 }
 
@@ -1012,7 +1012,7 @@ bool UNiagaraScript::RequestExternallyManagedAsyncCompile(const TSharedPtr<FNiag
 	{
 		OutCompileId = LastGeneratedVMId;
 		OutAsyncHandle = (uint32)INDEX_NONE;
-		UE_LOG(LogNiagara, Log, TEXT("Script '%s' is in-sync skipping compile.."), *GetFullName());
+		UE_LOG(LogNiagara, Verbose, TEXT("Script '%s' is in-sync skipping compile.."), *GetFullName());
 		return false;
 	}
 }
