Fixed compilation warning

diff --git a/Engine/Plugins/FX/Niagara/Source/Niagara/Private/NiagaraParameterStore.cpp b/Engine/Plugins/FX/Niagara/Source/Niagara/Private/NiagaraParameterStore.cpp
--- a/Engine/Plugins/FX/Niagara/Source/Niagara/Private/NiagaraParameterStore.cpp
+++ b/Engine/Plugins/FX/Niagara/Source/Niagara/Private/NiagaraParameterStore.cpp
@@ -26,7 +26,7 @@ static FAutoConsoleVariableRef CVarNiagaraDumpParticleParameterStores(
 
 //////////////////////////////////////////////////////////////////////////
 
-int32 FNiagaraVariableSearch::Find(const FNiagaraVariableWithOffset* Variables, const FNiagaraVariable& Ref, int32 Start, int32 Num, int32& CheckIndex)
+bool FNiagaraVariableSearch::Find(const FNiagaraVariableWithOffset* Variables, const FNiagaraVariable& Ref, int32 Start, int32 Num, int32& CheckIndex)
 {
 	while (Num)
 	{
diff --git a/Engine/Plugins/FX/Niagara/Source/Niagara/Public/NiagaraParameterStore.h b/Engine/Plugins/FX/Niagara/Source/Niagara/Public/NiagaraParameterStore.h
--- a/Engine/Plugins/FX/Niagara/Source/Niagara/Public/NiagaraParameterStore.h
+++ b/Engine/Plugins/FX/Niagara/Source/Niagara/Public/NiagaraParameterStore.h
@@ -31,7 +31,7 @@ struct FNiagaraVariableSearch
 		}
 	}
 
-	static int32 Find(const struct FNiagaraVariableWithOffset* Variables, const FNiagaraVariable& Ref, int32 Start, int32 Num, int32& CheckIndex);
+	static bool Find(const struct FNiagaraVariableWithOffset* Variables, const FNiagaraVariable& Ref, int32 Start, int32 Num, int32& CheckIndex);
 };
 
 struct FNiagaraParameterStore;
