- Add inline array allocator when gathering Scripts for CollectScriptDataInterfaceParameters
- Remove missing function / unused variable

#rnx

diff --git a/Engine/Plugins/FX/Niagara/Source/Niagara/Classes/NiagaraEmitterInstance.h b/Engine/Plugins/FX/Niagara/Source/Niagara/Classes/NiagaraEmitterInstance.h
--- a/Engine/Plugins/FX/Niagara/Source/Niagara/Classes/NiagaraEmitterInstance.h
+++ b/Engine/Plugins/FX/Niagara/Source/Niagara/Classes/NiagaraEmitterInstance.h
@@ -207,14 +207,6 @@ private:
 	 ** EmitterInstance and the cached bounds will be unset. */
 	TOptional<FBox> CachedSystemFixedBounds;
 
-#if WITH_EDITORONLY_DATA
-	bool CheckAttributesForRenderer(int32 Index);
-#endif
-
-#if !UE_BUILD_SHIPPING
-	bool bEncounteredNaNs;
-#endif
-
 	/** A parameter store which contains the data interfaces parameters which were defined by the scripts. */
 	FNiagaraParameterStore ScriptDefinedDataInterfaceParameters;
 
diff --git a/Engine/Plugins/FX/Niagara/Source/Niagara/Private/NiagaraCommon.cpp b/Engine/Plugins/FX/Niagara/Source/Niagara/Private/NiagaraCommon.cpp
--- a/Engine/Plugins/FX/Niagara/Source/Niagara/Private/NiagaraCommon.cpp
+++ b/Engine/Plugins/FX/Niagara/Source/Niagara/Private/NiagaraCommon.cpp
@@ -248,7 +248,7 @@ FNiagaraVariable FNiagaraUtilities::ConvertVariableToRapidIterationConstantName(
 	return Var;
 }
 
-void FNiagaraUtilities::CollectScriptDataInterfaceParameters(const UObject& Owner, const TArray<UNiagaraScript*>& Scripts, FNiagaraParameterStore& OutDataInterfaceParameters)
+void FNiagaraUtilities::CollectScriptDataInterfaceParameters(const UObject& Owner, const TArrayView<UNiagaraScript*>& Scripts, FNiagaraParameterStore& OutDataInterfaceParameters)
 {
 	for (UNiagaraScript* Script : Scripts)
 	{
diff --git a/Engine/Plugins/FX/Niagara/Source/Niagara/Private/NiagaraEmitterInstance.cpp b/Engine/Plugins/FX/Niagara/Source/Niagara/Private/NiagaraEmitterInstance.cpp
--- a/Engine/Plugins/FX/Niagara/Source/Niagara/Private/NiagaraEmitterInstance.cpp
+++ b/Engine/Plugins/FX/Niagara/Source/Niagara/Private/NiagaraEmitterInstance.cpp
@@ -70,9 +70,6 @@ FNiagaraEmitterInstance::FNiagaraEmitterInstance(FNiagaraSystemInstance* InParen
 , ParentSystemInstance(InParentSystemInstance)
 , CachedEmitter(nullptr)
 , CachedSystemFixedBounds()
-#if !UE_BUILD_SHIPPING
-, bEncounteredNaNs(false)
-#endif
 , EventSpawnTotal(0)
 {
 	bDumpAfterEvent = false;
@@ -204,9 +201,6 @@ void FNiagaraEmitterInstance::Init(int32 InEmitterIdx, FName InSystemInstanceNam
 		return;
 	}
 
-#if !UE_BUILD_SHIPPING
-	bEncounteredNaNs = false;
-#endif
 
 #if !UE_BUILD_SHIPPING && !UE_BUILD_TEST 
 	Data.Init(FNiagaraDataSetID(CachedIDName, ENiagaraDataSetType::ParticleData), CachedEmitter->SimTarget, ParentSystemInstance->GetSystem()->GetName() + TEXT("/") + CachedEmitter->GetName());
@@ -376,16 +370,17 @@ void FNiagaraEmitterInstance::Init(int32 InEmitterIdx, FName InSystemInstanceNam
 	}
 
 	// Collect script defined data interface parameters.
-	TArray<UNiagaraScript*> Scripts;
+	TArray<UNiagaraScript*, TInlineAllocator<4>> Scripts;
 	Scripts.Add(CachedEmitter->SpawnScriptProps.Script);
 	Scripts.Add(CachedEmitter->UpdateScriptProps.Script);
 	for (const FNiagaraEventScriptProperties& EventHandler : CachedEmitter->GetEventHandlers())
 	{
 		Scripts.Add(EventHandler.Script);
 	}
-	FNiagaraUtilities::CollectScriptDataInterfaceParameters(*CachedEmitter, Scripts, ScriptDefinedDataInterfaceParameters);
+	FNiagaraUtilities::CollectScriptDataInterfaceParameters(*CachedEmitter, MakeArrayView(Scripts), ScriptDefinedDataInterfaceParameters);
 
 	// Initialize bounds calculators
+	//-OPT: Could skip creating this if we won't ever use it
 	BoundsCalculators.Reserve(CachedEmitter->GetRenderers().Num());
 	for (UNiagaraRendererProperties* RendererProperties : CachedEmitter->GetRenderers())
 	{
diff --git a/Engine/Plugins/FX/Niagara/Source/Niagara/Public/NiagaraCommon.h b/Engine/Plugins/FX/Niagara/Source/Niagara/Public/NiagaraCommon.h
--- a/Engine/Plugins/FX/Niagara/Source/Niagara/Public/NiagaraCommon.h
+++ b/Engine/Plugins/FX/Niagara/Source/Niagara/Public/NiagaraCommon.h
@@ -586,7 +586,7 @@ namespace FNiagaraUtilities
 
 	FNiagaraVariable NIAGARA_API ConvertVariableToRapidIterationConstantName(FNiagaraVariable InVar, const TCHAR* InEmitterName, ENiagaraScriptUsage InUsage);
 
-	void CollectScriptDataInterfaceParameters(const UObject& Owner, const TArray<UNiagaraScript*>& Scripts, FNiagaraParameterStore& OutDataInterfaceParameters);
+	void CollectScriptDataInterfaceParameters(const UObject& Owner, const TArrayView<UNiagaraScript*>& Scripts, FNiagaraParameterStore& OutDataInterfaceParameters);
 
 	inline bool SupportsNiagaraRendering(ERHIFeatureLevel::Type FeatureLevel)
 	{
