- Execute Niagara system spawns in tick groups as we progress ticking each group
- This pushes the spawning to async where possible, we can still be left with system that have to be spawned in post actor tick if the last demotable group spawned additional systems

diff --git a/Engine/Plugins/FX/Niagara/Source/Niagara/Private/NiagaraSystemInstance.cpp b/Engine/Plugins/FX/Niagara/Source/Niagara/Private/NiagaraSystemInstance.cpp
--- a/Engine/Plugins/FX/Niagara/Source/Niagara/Private/NiagaraSystemInstance.cpp
+++ b/Engine/Plugins/FX/Niagara/Source/Niagara/Private/NiagaraSystemInstance.cpp
@@ -31,7 +31,7 @@ DECLARE_CYCLE_STAT(TEXT("System Instance Tick [CNC]"), STAT_NiagaraSystemInst_Ti
 DECLARE_CYCLE_STAT(TEXT("System Instance Finalize [GT]"), STAT_NiagaraSystemInst_FinalizeGT, STATGROUP_Niagara);
 DECLARE_CYCLE_STAT(TEXT("System Instance WaitForAsyncTick [GT]"), STAT_NiagaraSystemWaitForAsyncTick, STATGROUP_Niagara);
 
-static float GWaitForAsyncStallWarnThresholdMS = 0.01f;
+static float GWaitForAsyncStallWarnThresholdMS = 0.2f;
 static FAutoConsoleVariableRef CVarWaitForAsyncStallWarnThresholdMS(
 	TEXT("fx.WaitForAsyncStallWarnThresholdMS"),
 	GWaitForAsyncStallWarnThresholdMS,
@@ -59,7 +59,6 @@ static FAutoConsoleVariableRef CVarNiagaraForceLastTickGroup(
 FNiagaraSystemInstance::FNiagaraSystemInstance(UNiagaraComponent* InComponent)
 	: SystemInstanceIndex(INDEX_NONE)
 	, Component(InComponent)
-	, TickGroup(TG_MAX)
 	, PrereqComponent(nullptr)
 	, Age(0.0f)
 	, TickCount(0)
@@ -223,6 +222,34 @@ void FNiagaraSystemInstance::Dump()const
 	}
 }
 
+void FNiagaraSystemInstance::DumpTickInfo(FOutputDevice& Ar)
+{
+	WaitForAsyncTick();
+
+	static const UEnum* TickingGroupEnum = FindObjectChecked<UEnum>(ANY_PACKAGE, TEXT("ETickingGroup"));
+
+	FString PrereqInfo;
+	if (PrereqComponent != nullptr)
+	{
+		ETickingGroup PrereqTG = FMath::Max(PrereqComponent->PrimaryComponentTick.TickGroup, PrereqComponent->PrimaryComponentTick.EndTickGroup);
+		PrereqInfo.Appendf(TEXT(" PreReq(%s = %s)"), *PrereqComponent->GetFullName(), *TickingGroupEnum->GetNameStringByIndex(PrereqTG));
+	}
+
+	if (bDataInterfacesHaveTickPrereqs)
+	{
+		for (TPair<TWeakObjectPtr<UNiagaraDataInterface>, int32>& Pair : DataInterfaceInstanceDataOffsets)
+		{
+			if (UNiagaraDataInterface* Interface = Pair.Key.Get())
+			{
+				ETickingGroup PrereqTG = Interface->CalculateTickGroup(&DataInterfaceInstanceData[Pair.Value]);
+				PrereqInfo.Appendf(TEXT(" DataInterface(%s = %s)"), *Interface->GetFullName(), *TickingGroupEnum->GetNameStringByIndex(PrereqTG));
+			}
+		}
+	}
+
+	Ar.Logf(TEXT("\t\t\tInstance%s"), *PrereqInfo);
+}
+
 #if WITH_EDITORONLY_DATA
 bool FNiagaraSystemInstance::RequestCapture(const FGuid& RequestId)
 {
@@ -364,13 +391,11 @@ void FNiagaraSystemInstance::SetSolo(bool bInSolo)
 
 	WaitForAsyncTick();
 
-	TickGroup = TG_MAX;
-
 	UNiagaraSystem* System = GetSystem();
 	if (bInSolo)
 	{
 		TSharedPtr<FNiagaraSystemSimulation, ESPMode::ThreadSafe> NewSoloSim = MakeShared<FNiagaraSystemSimulation, ESPMode::ThreadSafe>();
-		NewSoloSim->Init(System, Component->GetWorld(), true);
+		NewSoloSim->Init(System, Component->GetWorld(), true, TG_MAX);
 
 		NewSoloSim->TransferInstance(SystemSimulation.Get(), this);	
 
@@ -380,7 +405,7 @@ void FNiagaraSystemInstance::SetSolo(bool bInSolo)
 	else
 	{
 		UpdatePrereqs();
-		TickGroup = CalculateTickGroup();
+		const ETickingGroup TickGroup = CalculateTickGroup();
 		TSharedPtr<FNiagaraSystemSimulation, ESPMode::ThreadSafe> NewSim = GetWorldManager()->GetSystemSimulation(TickGroup, System);
 
 		NewSim->TransferInstance(SystemSimulation.Get(), this);
@@ -568,6 +593,7 @@ void FNiagaraSystemInstance::Reset(FNiagaraSystemInstance::EResetMode Mode)
 	if (SystemSimulation.IsValid())
 	{
 		SystemSimulation->RemoveInstance(this);
+		bNeedsFinalize = false;
 	}
 	else
 	{
@@ -789,20 +815,19 @@ void FNiagaraSystemInstance::ReInitInternal()
 	}
 
 	/** Do we need to run in solo mode? */
-	TickGroup = TG_MAX;
 	bSolo = bForceSolo || DoSystemDataInterfacesRequireSolo(*System, *Component);
 	if (bSolo)
 	{
 		if (!SystemSimulation.IsValid())
 		{
 			SystemSimulation = MakeShared<FNiagaraSystemSimulation, ESPMode::ThreadSafe>();
-			SystemSimulation->Init(System, Component->GetWorld(), true);
+			SystemSimulation->Init(System, Component->GetWorld(), true, TG_MAX);
 		}
 	}
 	else
 	{
 		UpdatePrereqs();
-		TickGroup = CalculateTickGroup();
+		const ETickingGroup TickGroup = CalculateTickGroup();
 		SystemSimulation = GetWorldManager()->GetSystemSimulation(TickGroup, System);
 	}
 
@@ -1491,14 +1516,7 @@ void FNiagaraSystemInstance::ComponentTick(float DeltaSeconds, const FGraphEvent
 	check(bSolo);
 	check(Component);
 
-	if (IsPendingSpawn())
-	{
-		SystemSim->SpawnNew_GameThread(DeltaSeconds, MyCompletionGraphEvent);
-	}
-	else
-	{
-		SystemSim->Tick_GameThread(DeltaSeconds, MyCompletionGraphEvent);
-	}
+	SystemSim->Tick_GameThread(DeltaSeconds, MyCompletionGraphEvent);
 }
 
 void FNiagaraSystemInstance::WaitForAsyncTick(bool bEnsureComplete)
@@ -1560,32 +1578,6 @@ void FNiagaraSystemInstance::Tick_GameThread(float DeltaSeconds)
 	UNiagaraSystem* System = GetSystem();
 	FScopeCycleCounter SystemStat(System->GetStatID(true, true));
 
-	// Update tick group, this can cause our Tick_GameThread to exit early due to a demotion
-	// When not in solo mode we need to deal with tick groups potentially changing due to attachments or data interfaces having prereqs
-	if ( !bSolo && !bPendingSpawn )
-	{
-		ETickingGroup NewTickGroup = CalculateTickGroup();
-		if ( NewTickGroup != TickGroup )
-		{
-			// Tick demotion we need to do this now to ensure we complete in the correct group
-			if ( NewTickGroup > TickGroup )
-			{
-				TickGroup = NewTickGroup;
-
-				TSharedPtr<FNiagaraSystemSimulation, ESPMode::ThreadSafe> NewSim = GetWorldManager()->GetSystemSimulation(TickGroup, System);
-				NewSim->TransferInstance(SystemSimulation.Get(), this);
-				SystemSimulation = NewSim;
-				return;
-			}
-			// Tick promotions must be deferred as the tick group has already been processed
-			// Note: We still tick in this group and will be transfered later on
-			else
-			{
-				SystemSimulation->AddTickGroupPromotion(this);
-			}
-		}
-	}
-
 	WaitForAsyncTick(true);
 
 	CachedDeltaSeconds = DeltaSeconds;
@@ -1607,6 +1599,7 @@ void FNiagaraSystemInstance::Tick_Concurrent()
 	SCOPE_CYCLE_COUNTER(STAT_NiagaraOverview_GT_CNC);
 	CSV_SCOPED_TIMING_STAT_EXCLUSIVE(Niagara);
 	LLM_SCOPE(ELLMTag::Niagara);
+	FScopeCycleCounterUObject AdditionalScope(GetSystem(), GET_STATID(STAT_NiagaraOverview_GT_CNC));
 
 	// Reset values that will be accumulated during emitter tick.
 	TotalParamSize = 0;
diff --git a/Engine/Plugins/FX/Niagara/Source/Niagara/Private/NiagaraSystemSimulation.cpp b/Engine/Plugins/FX/Niagara/Source/Niagara/Private/NiagaraSystemSimulation.cpp
--- a/Engine/Plugins/FX/Niagara/Source/Niagara/Private/NiagaraSystemSimulation.cpp
+++ b/Engine/Plugins/FX/Niagara/Source/Niagara/Private/NiagaraSystemSimulation.cpp
@@ -85,29 +85,20 @@ static FAutoConsoleVariableRef CVarSystemSimTransferParamsParallelThreshold(
 
 //////////////////////////////////////////////////////////////////////////
 
-FNiagaraSystemSimulationTickContext::FNiagaraSystemSimulationTickContext(FNiagaraSystemSimulation* InOwner, float InDeltaSeconds, bool bInPendingSpawnPass, const FGraphEventRef& InMyCompletionGraphEvent)
+FNiagaraSystemSimulationTickContext::FNiagaraSystemSimulationTickContext(FNiagaraSystemSimulation* InOwner, TArray<FNiagaraSystemInstance*>& InInstances, FNiagaraDataSet& InDataSet, float InDeltaSeconds, int32 InSpawnNum, const FGraphEventRef& InMyCompletionGraphEvent)
 	: Owner(InOwner)
+	, System(InOwner->GetSystem())
+	, Instances(InInstances)
+	, DataSet(InDataSet)
 	, DeltaSeconds(InDeltaSeconds)
-	, bPendingSpawnPass(bInPendingSpawnPass)
+	, SpawnNum(InSpawnNum)
 	, MyCompletionGraphEvent(InMyCompletionGraphEvent)
+	, FinalizeEvents(nullptr)
+	, bTickAsync(GbParallelSystemSimTick && FApp::ShouldUseThreadingForPerformance() && InMyCompletionGraphEvent.IsValid())
+	, bTickInstancesAsync(GbParallelSystemInstanceTick && FApp::ShouldUseThreadingForPerformance() && MyCompletionGraphEvent.IsValid() && !InOwner->GetIsSolo())
 {
-	System = Owner->GetSystem();
-	
-	bTickAsync = Owner->ShouldTickAsync(*this);
-	bTickInstancesAsync = Owner->ShouldTickInstancesAsync(*this);
-}
-
-TArray<FNiagaraSystemInstance*>& FNiagaraSystemSimulationTickContext::GetInstances()
-{
-	return bPendingSpawnPass ? Owner->SpawningSystemInstances : Owner->SystemInstances;
-}
-
-FNiagaraDataSet& FNiagaraSystemSimulationTickContext::GetDataSet()
-{
-	return bPendingSpawnPass ? Owner->PendingSpawnDataSet : Owner->DataSet;
 }
 
-//////////////////////////////////////////////////////////////////////////
 //////////////////////////////////////////////////////////////////////////
 
 FAutoConsoleTaskPriority CPrio_NiagaraSystemSimulationTickTask(
@@ -244,12 +235,14 @@ FNiagaraSystemSimulation::~FNiagaraSystemSimulation()
 	Destroy();
 }
 
-bool FNiagaraSystemSimulation::Init(UNiagaraSystem* InSystem, UWorld* InWorld, bool bInIsSolo)
+bool FNiagaraSystemSimulation::Init(UNiagaraSystem* InSystem, UWorld* InWorld, bool bInIsSolo, ETickingGroup InTickGroup)
 {
 	SCOPE_CYCLE_COUNTER(STAT_NiagaraSystemSim_Init);
 	UNiagaraSystem* System = InSystem;
 	WeakSystem = System;
 
+	SystemTickGroup = InTickGroup;
+
 	World = InWorld;
 
 	bIsSolo = bInIsSolo;
@@ -268,25 +261,25 @@ bool FNiagaraSystemSimulation::Init(UNiagaraSystem* InSystem, UWorld* InWorld, b
 	{
 		//TODO: Move all layout data into the System!
 
-		//Initilize the main simulation dataset.
+		//Initialize the main simulation dataset.
 #if !UE_BUILD_SHIPPING && !UE_BUILD_TEST 
-		DataSet.Init(FNiagaraDataSetID(), ENiagaraSimTarget::CPUSim, System->GetFullName());
+		MainDataSet.Init(FNiagaraDataSetID(), ENiagaraSimTarget::CPUSim, System->GetFullName());
 #else
-		DataSet.Init(FNiagaraDataSetID(), ENiagaraSimTarget::CPUSim);
+		MainDataSet.Init(FNiagaraDataSetID(), ENiagaraSimTarget::CPUSim);
 #endif
-		DataSet.AddVariables(System->GetSystemSpawnScript()->GetVMExecutableData().Attributes);
-		DataSet.AddVariables(System->GetSystemUpdateScript()->GetVMExecutableData().Attributes);
-		DataSet.Finalize();
+		MainDataSet.AddVariables(System->GetSystemSpawnScript()->GetVMExecutableData().Attributes);
+		MainDataSet.AddVariables(System->GetSystemUpdateScript()->GetVMExecutableData().Attributes);
+		MainDataSet.Finalize();
 
-		//Initialize the data set for newly spawned systems.
+		//Initialize the main simulation dataset.
 #if !UE_BUILD_SHIPPING && !UE_BUILD_TEST 
-		PendingSpawnDataSet.Init(FNiagaraDataSetID(), ENiagaraSimTarget::CPUSim, System->GetFullName());
+		SpawningDataSet.Init(FNiagaraDataSetID(), ENiagaraSimTarget::CPUSim, System->GetFullName());
 #else
-		PendingSpawnDataSet.Init(FNiagaraDataSetID(), ENiagaraSimTarget::CPUSim);
+		SpawningDataSet.Init(FNiagaraDataSetID(), ENiagaraSimTarget::CPUSim);
 #endif
-		PendingSpawnDataSet.AddVariables(System->GetSystemSpawnScript()->GetVMExecutableData().Attributes);
-		PendingSpawnDataSet.AddVariables(System->GetSystemUpdateScript()->GetVMExecutableData().Attributes);
-		PendingSpawnDataSet.Finalize();
+		SpawningDataSet.AddVariables(System->GetSystemSpawnScript()->GetVMExecutableData().Attributes);
+		SpawningDataSet.AddVariables(System->GetSystemUpdateScript()->GetVMExecutableData().Attributes);
+		SpawningDataSet.Finalize();
 
 		//Initialize the dataset for paused systems.
 #if !UE_BUILD_SHIPPING && !UE_BUILD_TEST 
@@ -350,7 +343,7 @@ bool FNiagaraSystemSimulation::Init(UNiagaraSystem* InSystem, UWorld* InWorld, b
 		SpawnScript->RapidIterationParameters.Bind(&SpawnExecContext.Parameters);
 		UpdateScript->RapidIterationParameters.Bind(&UpdateExecContext.Parameters);
 
-		SystemExecutionStateAccessor.Create(&DataSet, FNiagaraVariable(EnumPtr, TEXT("System.ExecutionState")));
+		SystemExecutionStateAccessor.Create(&MainDataSet, FNiagaraVariable(EnumPtr, TEXT("System.ExecutionState")));
 		EmitterSpawnInfoAccessors.Reset();
 		EmitterExecutionStateAccessors.Reset();
 		EmitterSpawnInfoAccessors.SetNum(System->GetNumEmitters());
@@ -361,13 +354,13 @@ bool FNiagaraSystemSimulation::Init(UNiagaraSystem* InSystem, UWorld* InWorld, b
 			UNiagaraEmitter* Emitter = EmitterHandle.GetInstance();
 			FString EmitterName = Emitter->GetUniqueEmitterName();
 			check(Emitter);
-			EmitterExecutionStateAccessors.Emplace(DataSet, FNiagaraVariable(EnumPtr, *(EmitterName + TEXT(".ExecutionState"))));
+			EmitterExecutionStateAccessors.Emplace(MainDataSet, FNiagaraVariable(EnumPtr, *(EmitterName + TEXT(".ExecutionState"))));
 			const TArray<FNiagaraEmitterCompiledData>& EmitterCompiledData = System->GetEmitterCompiledData();
 			
 			check(EmitterCompiledData.Num() == System->GetNumEmitters());
 			for (FName AttrName : EmitterCompiledData[EmitterIdx].SpawnAttributes)
 			{
-				EmitterSpawnInfoAccessors[EmitterIdx].Emplace(DataSet, FNiagaraVariable(FNiagaraTypeDefinition(FNiagaraSpawnInfo::StaticStruct()), AttrName));
+				EmitterSpawnInfoAccessors[EmitterIdx].Emplace(MainDataSet, FNiagaraVariable(FNiagaraTypeDefinition(FNiagaraSpawnInfo::StaticStruct()), AttrName));
 			}
 
 			if (Emitter->bLimitDeltaTime)
@@ -408,8 +401,11 @@ void FNiagaraSystemSimulation::Destroy()
 	}
 	SystemInstances.Empty();
 	PendingSystemInstances.Empty();
-	SpawningSystemInstances.Empty();
 
+	//TArray<FNiagaraSystemInstance*> SystemInstances;
+	//TArray<FNiagaraSystemInstance*> SpawningInstances;
+	//TArray<FNiagaraSystemInstance*> PausedSystemInstances;
+	//PausedSystemInstances
 
 	FNiagaraWorldManager* WorldMan = FNiagaraWorldManager::Get(World);
 	check(WorldMan);
@@ -458,7 +454,7 @@ void FNiagaraSystemSimulation::TransferInstance(FNiagaraSystemSimulation* Source
 // 		SourceSimulation->DataSet.Dump(true, SystemInstIdx, 1);
 
 		//If we're not pending then the system actually has data to pull over. This is not fast.
-		int32 NewDataSetIndex = DataSet.GetCurrentDataChecked().TransferInstance(SourceSimulation->DataSet.GetCurrentDataChecked(), SystemInstIdx, false);
+		int32 NewDataSetIndex = MainDataSet.GetCurrentDataChecked().TransferInstance(SourceSimulation->MainDataSet.GetCurrentDataChecked(), SystemInstIdx, false);
 
 // 		UE_LOG(LogNiagara, Log, TEXT(" ----- Transfered values in dest. Idx: %d -----"), NewDataIndex);
 // 		DataSet.Dump(true, NewDataIndex, 1);
@@ -496,24 +492,45 @@ void FNiagaraSystemSimulation::DumpInstance(const FNiagaraSystemInstance* Inst)c
 	UpdateExecContext.Parameters.DumpParameters(false);
 	UpdateInstanceParameterDataSet.Dump(Inst->SystemInstanceIndex, 1, TEXT("Update Instance Parameters"));
 	UE_LOG(LogNiagara, Log, TEXT("................. System Instance ................."));
-	DataSet.Dump(Inst->SystemInstanceIndex, 1, TEXT("System Data"));
+	MainDataSet.Dump(Inst->SystemInstanceIndex, 1, TEXT("System Data"));
 }
 
-void FNiagaraSystemSimulation::AddTickGroupPromotion(FNiagaraSystemInstance* Instance)
+void FNiagaraSystemSimulation::DumpTickInfo(FOutputDevice& Ar)
 {
 	check(IsInGameThread());
-	check(!PendingTickGroupPromotions.Contains(Instance));
-	PendingTickGroupPromotions.Add(Instance);
-}
+	if (SystemInstances.Num() > 0)
+	{
+		Ar.Logf(TEXT("\t\tSystemInstances %d"), SystemInstances.Num());
+		for (FNiagaraSystemInstance* Instance : SystemInstances)
+		{
+			Instance->DumpTickInfo(Ar);
+		}
+	}
 
-bool FNiagaraSystemSimulation::ShouldTickAsync(const FNiagaraSystemSimulationTickContext& Context)
-{
-	return GbParallelSystemSimTick && FApp::ShouldUseThreadingForPerformance() && !Context.bPendingSpawnPass && Context.MyCompletionGraphEvent.IsValid();
+	if (PendingSystemInstances.Num() > 0)
+	{
+		Ar.Logf(TEXT("\t\tPendingSystemInstances %d"), PendingSystemInstances.Num());
+		for (FNiagaraSystemInstance* Instance : PendingSystemInstances)
+		{
+			Instance->DumpTickInfo(Ar);
+		}
+	}
+
+	if (PausedSystemInstances.Num() > 0)
+	{
+		Ar.Logf(TEXT("\t\tPausedSystemInstances %d"), PausedSystemInstances.Num());
+		for (FNiagaraSystemInstance* Instance : PausedSystemInstances)
+		{
+			Instance->DumpTickInfo(Ar);
+		}
+	}
 }
 
-bool FNiagaraSystemSimulation::ShouldTickInstancesAsync(const FNiagaraSystemSimulationTickContext& Context)
+void FNiagaraSystemSimulation::AddTickGroupPromotion(FNiagaraSystemInstance* Instance)
 {
-	return GbParallelSystemInstanceTick && !bIsSolo && FApp::ShouldUseThreadingForPerformance() && !Context.bPendingSpawnPass && Context.MyCompletionGraphEvent.IsValid();
+	check(IsInGameThread());
+	check(!PendingTickGroupPromotions.Contains(Instance));
+	PendingTickGroupPromotions.Add(Instance);
 }
 
 void FNiagaraSystemSimulation::AddSystemToTickBatch(FNiagaraSystemInstance* Instance, FNiagaraSystemSimulationTickContext& Context)
@@ -561,231 +578,378 @@ void FNiagaraSystemSimulation::FlushTickBatch(FNiagaraSystemSimulationTickContex
 	}
 }
 
-void FNiagaraSystemSimulation::SpawnNew_GameThread(float DeltaSeconds, const FGraphEventRef& MyCompletionGraphEvent)
+/** First phase of system sim tick. Must run on GameThread. */
+void FNiagaraSystemSimulation::Tick_GameThread(float DeltaSeconds, const FGraphEventRef& MyCompletionGraphEvent)
 {
-	SCOPE_CYCLE_COUNTER(STAT_NiagaraSystemSim_SpawnNewGT);
+	check(IsInGameThread());
+	WaitForSystemTickComplete(true);
+
 	SCOPE_CYCLE_COUNTER(STAT_NiagaraOverview_GT);
+	SCOPE_CYCLE_COUNTER(STAT_NiagaraSystemSim_TickGT);
 	CSV_SCOPED_TIMING_STAT_EXCLUSIVE(Niagara);
 	LLM_SCOPE(ELLMTag::Niagara);
+	FScopeCycleCounterUObject AdditionalScope(GetSystem(), GET_STATID(STAT_NiagaraOverview_GT_CNC));
 
-	WaitForSystemTickComplete(true);
+	bHasEverTicked = true;
+
+	SystemTickGraphEvent = nullptr;
+
+	check(SystemInstances.Num() == MainDataSet.GetCurrentDataChecked().GetNumInstances());
+	check(PausedSystemInstances.Num() == PausedInstanceData.GetCurrentDataChecked().GetNumInstances());
 
 	UNiagaraSystem* System = WeakSystem.Get();
+
 	FScopeCycleCounter SystemStatCounter(System->GetStatID(true, false));
 
-	check(SystemInstances.Num() == DataSet.GetCurrentDataChecked().GetNumInstances());
-	check(PausedSystemInstances.Num() == PausedInstanceData.GetCurrentDataChecked().GetNumInstances());
+	if (MaxDeltaTime.IsSet())
+	{
+		DeltaSeconds = FMath::Clamp(DeltaSeconds, 0.0f, MaxDeltaTime.GetValue());
+	}
 
-	//We have to do new spwaning at the end of the frame to account for new niagara systmes spawned after our main simulation tick.
-	//So here we build a list of instances to actually spawn from the PendingSpawnInstances.
-	//Then spwan these into a new dataset which is then transfered into the main dataset.
-	//Currently this is all done on the game thread but we could likely also move this off the GT too.
+	UNiagaraScript* SystemSpawnScript = System->GetSystemSpawnScript();
+	UNiagaraScript* SystemUpdateScript = System->GetSystemUpdateScript();
+#if WITH_EDITOR
+	SystemSpawnScript->RapidIterationParameters.Tick();
+	SystemUpdateScript->RapidIterationParameters.Tick();
+#endif
 
+	const bool bUpdateTickGroups = !bIsSolo && !bInSpawnPhase;
 
-	check(IsInGameThread());
-	if (PendingSystemInstances.Num() > 0)
+	// Update instances
+	int32 SystemIndex = 0;
+	while (SystemIndex < SystemInstances.Num())
+	{
+		FNiagaraSystemInstance* Inst = SystemInstances[SystemIndex];
+
+		// Update instance tick group, this can involve demoting the instance (i.e. removing from our list)
+		if ( bUpdateTickGroups )
+		{
+			ETickingGroup DesiredTickGroup = Inst->CalculateTickGroup();
+			if (DesiredTickGroup != SystemTickGroup )
+			{
+				// Tick demotion we need to do this now to ensure we complete in the correct group
+				if (DesiredTickGroup > SystemTickGroup)
+				{
+					FNiagaraWorldManager* WorldManager = FNiagaraWorldManager::Get(World);
+					check(WorldManager != nullptr);
+
+					TSharedPtr<FNiagaraSystemSimulation, ESPMode::ThreadSafe> NewSim = WorldManager->GetSystemSimulation(DesiredTickGroup, System);
+					NewSim->WaitForInstancesTickComplete();
+					NewSim->MainDataSet.GetCurrentDataChecked().TransferInstance(MainDataSet.GetCurrentDataChecked(), SystemIndex, true);
+
+					SystemInstances.RemoveAtSwap(SystemIndex);
+					if (SystemInstances.IsValidIndex(SystemIndex))
+					{
+						SystemInstances[SystemIndex]->SystemInstanceIndex = SystemIndex;
+					}
+
+					Inst->SystemInstanceIndex = NewSim->SystemInstances.Add(Inst);
+					Inst->SystemSimulation = NewSim;
+
+					if (!NewSim->bBindingsInitialized)
+					{
+						NewSim->InitParameterDataSetBindings(Inst);
+					}
+					continue;
+				}
+				// Tick promotions must be deferred as the tick group has already been processed
+				//-OPT: We could tick in this group and add a task dependent on both groups to do the transform async
+				else
+				{
+					AddTickGroupPromotion(Inst);
+				}
+			}
+		}
+
+		// Perform instance tick
+		Inst->Tick_GameThread(DeltaSeconds);
+
+		// TickDataInterfaces could remove the system so we only increment if the system has changed
+		// Also possible for this system to have been transfered to another system simulation.
+		if ( Inst->GetSystemSimulation().Get() == this )
+		{
+			if (Inst->SystemInstanceIndex != INDEX_NONE)
+			{
+				checkSlow(Inst->SystemInstanceIndex == SystemIndex);
+				++SystemIndex;
+			}
+			else
+			{
+				checkSlow((SystemInstances.Num() <= SystemIndex) || (Inst == SystemInstances[SystemIndex]));
+			}
+		}
+	}
+
+	//Setup the few real constants like delta time.
+	const float InvDt = 1.0f / DeltaSeconds;
+
+	SpawnTimeParam.SetValue(World->TimeSeconds);
+	UpdateTimeParam.SetValue(World->TimeSeconds);
+	SpawnDeltaTimeParam.SetValue(DeltaSeconds);
+	UpdateDeltaTimeParam.SetValue(DeltaSeconds);
+	SpawnInvDeltaTimeParam.SetValue(InvDt);
+	UpdateInvDeltaTimeParam.SetValue(InvDt);
+	SpawnNumSystemInstancesParam.SetValue(SystemInstances.Num());
+	UpdateNumSystemInstancesParam.SetValue(SystemInstances.Num());
+	SpawnGlobalSpawnCountScaleParam.SetValue(INiagaraModule::GetGlobalSpawnCountScale());
+	UpdateGlobalSpawnCountScaleParam.SetValue(INiagaraModule::GetGlobalSpawnCountScale());
+	SpawnGlobalSystemCountScaleParam.SetValue(INiagaraModule::GetGlobalSystemCountScale());
+	UpdateGlobalSystemCountScaleParam.SetValue(INiagaraModule::GetGlobalSystemCountScale());
+
+	// Somethings we don't want to happen during the spawn phase
+	if ( !bInSpawnPhase )
 	{
-		int32 SystemIndex = 0;
-		SpawningSystemInstances.Reserve(PendingSystemInstances.Num());
-		
-		while (SystemIndex < PendingSystemInstances.Num())
+		int32 SpawnNum = 0;
+
+		// Gather any pending spawn systems and add to the end of the system instances
+		if (PendingSystemInstances.Num() > 0)
 		{
-			FNiagaraSystemInstance* Inst = PendingSystemInstances[SystemIndex];
+			SystemInstances.Reserve(SystemInstances.Num() + PendingSystemInstances.Num());
 
-			//If an instance is paused, just leave it in the PendingSpawnInstances.
-			if (!Inst->IsPaused())
+			SystemIndex = 0;
+			while (SystemIndex < PendingSystemInstances.Num())
 			{
+				FNiagaraSystemInstance* Inst = PendingSystemInstances[SystemIndex];
+
+				// If we are paused continue
+				if (Inst->IsPaused())
+				{
+					++SystemIndex;
+					continue;
+				}
+
+				// If our tick group has changed we need to move this pending instance
+				// Note we don't do this with solo instances
+				if (!bIsSolo)
+				{
+					const ETickingGroup DesiredTickGroup = Inst->CalculateTickGroup();
+					if (DesiredTickGroup != SystemTickGroup)
+					{
+						PendingSystemInstances.RemoveAtSwap(SystemIndex);
+						if (PendingSystemInstances.IsValidIndex(SystemIndex))
+						{
+							PendingSystemInstances[SystemIndex]->SystemInstanceIndex = SystemIndex;
+						}
+						Inst->SystemInstanceIndex = INDEX_NONE;
+
+						FNiagaraWorldManager* WorldManager = FNiagaraWorldManager::Get(World);
+						check(WorldManager != nullptr);
+
+						TSharedPtr<FNiagaraSystemSimulation, ESPMode::ThreadSafe> DestSim = WorldManager->GetSystemSimulation(DesiredTickGroup, System);
+
+						Inst->SystemSimulation = DestSim;
+						Inst->SystemInstanceIndex = DestSim->PendingSystemInstances.Add(Inst);
+						continue;
+					}
+				}
+
+				// Execute instance tick
 				Inst->Tick_GameThread(DeltaSeconds);
-								
-				// TickDataInterfaces could remove the system so we only increment if the system has changed
+
 				if (Inst->SystemInstanceIndex != INDEX_NONE)
 				{
-					checkSlow(!Inst->IsComplete());
-					checkSlow(Inst->SystemInstanceIndex == SystemIndex);
-					
+					// We should not move tick group during Tick_GameThread but let's be safe
+					check(Inst->SystemSimulation.Get() == this);
+
+					// When the first instance is added we need to initialize the parameter store to data set bindings.
 					if (!bBindingsInitialized)
 					{
-						// When the first instance is added we need to initialize the parameter store to data set bindings.
 						InitParameterDataSetBindings(Inst);
 					}
 
-					int32 SpawnInstIdx = SpawningSystemInstances.Add(Inst);
+					check(PendingSystemInstances[SystemIndex] == Inst);
 					PendingSystemInstances.RemoveAtSwap(SystemIndex);
-
-					Inst->SystemInstanceIndex = SpawnInstIdx;
 					if (PendingSystemInstances.IsValidIndex(SystemIndex))
 					{
 						PendingSystemInstances[SystemIndex]->SystemInstanceIndex = SystemIndex;
 					}
 
-					continue;
+					Inst->SystemInstanceIndex = SystemInstances.Add(Inst);
+					Inst->SetPendingSpawn(false);
+					++SpawnNum;
 				}
 			}
-			
-			++SystemIndex;
-		}		
+		}
 
-		//-OPT: Our tick group may have changed so next frame we could end up moving the instance again
-		if (SpawningSystemInstances.Num() > 0)
+		FNiagaraSystemSimulationTickContext Context(this, SystemInstances, MainDataSet, DeltaSeconds, SpawnNum, MyCompletionGraphEvent);
+
+		//Now kick of the concurrent tick.
+		if (Context.bTickAsync)
 		{
-			//Lets any RemoveInstance calls inside Tick_Concurrent know that we're spawning and all our instances are in SpawningSystemInstances;
-			bInSpawnPhase = true;
+			TGraphTask<FNiagaraSystemSimulationWaitAllFinalizeTask>* WaitAllFinalizeTask = TGraphTask<FNiagaraSystemSimulationWaitAllFinalizeTask>::CreateTask(nullptr, ENamedThreads::GameThread).ConstructAndHold(Context.FinalizeEvents);
+			FGraphEventRef FinalizeGraphEvent = WaitAllFinalizeTask->GetCompletionEvent();
+
+			TGraphTask<FNiagaraSystemSimulationTickTask>* SimulationTickTask = TGraphTask<FNiagaraSystemSimulationTickTask>::CreateTask(nullptr, ENamedThreads::GameThread).ConstructAndHold(Context, WaitAllFinalizeTask);
+			SystemTickGraphEvent = SimulationTickTask->GetCompletionEvent();
+			Context.FinalizeEvents->Add(SystemTickGraphEvent);
 
-			// When a new simulation is added it may miss it's tick in which case we need to ensure it's ticked before we execute the Tick_Concurrent for spawning
-			if (!bHasEverTicked)
+			MyCompletionGraphEvent->SetGatherThreadForDontCompleteUntil(ENamedThreads::GameThread);
+			MyCompletionGraphEvent->DontCompleteUntil(FinalizeGraphEvent);
+
+			SimulationTickTask->Unlock(ENamedThreads::GameThread);
+		}
+		else
+		{
+			TGraphTask<FNiagaraSystemSimulationWaitAllFinalizeTask>* WaitAllFinalizeTask = nullptr;
+			if (Context.bTickInstancesAsync)
 			{
-				Tick_GameThread(DeltaSeconds, MyCompletionGraphEvent);
+				WaitAllFinalizeTask = TGraphTask<FNiagaraSystemSimulationWaitAllFinalizeTask>::CreateTask(nullptr, ENamedThreads::GameThread).ConstructAndHold(Context.FinalizeEvents);
 			}
 
-			//For now just do the concurrent tick for new systems on the GT here but we may also want to push this off to task threads too.
-			FNiagaraSystemSimulationTickContext Context(this, DeltaSeconds, true, MyCompletionGraphEvent);
 			Tick_Concurrent(Context);
 
-			check(SystemInstances.Num() == DataSet.GetCurrentDataChecked().GetNumInstances());
-			check(PausedSystemInstances.Num() == PausedInstanceData.GetCurrentDataChecked().GetNumInstances());
-			check(SpawningSystemInstances.Num() == PendingSpawnDataSet.GetCurrentDataChecked().GetNumInstances());
-
-			//Copy the sim data.
-			PendingSpawnDataSet.CopyTo(DataSet, 0, INDEX_NONE, false);
-			PendingSpawnDataSet.ResetBuffers();
-
-			//Move new instances over to the main buffer and fix their indices.
-			SystemInstances.Reserve(SystemInstances.Num() + SpawningSystemInstances.Num());
-			for(FNiagaraSystemInstance* NewInst : SpawningSystemInstances)
+			if (Context.bTickInstancesAsync)
 			{
-				checkSlow(!NewInst->IsComplete());
-				NewInst->SystemInstanceIndex = SystemInstances.Add(NewInst);
+				WaitAllFinalizeTask->Unlock(ENamedThreads::GameThread);
+				Context.FinalizeEvents = nullptr;
 			}
+		}
+	}
+}
 
-			check(SystemInstances.Num() == DataSet.GetCurrentDataChecked().GetNumInstances());
-			check(PausedSystemInstances.Num() == PausedInstanceData.GetCurrentDataChecked().GetNumInstances());
+void FNiagaraSystemSimulation::UpdateTickGroups_GameThread()
+{
+	check(IsInGameThread());
+	check(!bIsSolo);
+
+	SCOPE_CYCLE_COUNTER(STAT_NiagaraSystemSim_SpawnNewGT);
+	SCOPE_CYCLE_COUNTER(STAT_NiagaraOverview_GT);
+	CSV_SCOPED_TIMING_STAT_EXCLUSIVE(Niagara);
+	LLM_SCOPE(ELLMTag::Niagara);
+	FScopeCycleCounterUObject AdditionalScope(GetSystem(), GET_STATID(STAT_NiagaraOverview_GT_CNC));
+
+	FNiagaraWorldManager* WorldManager = FNiagaraWorldManager::Get(World);
+	check(WorldManager != nullptr);
+
+	UNiagaraSystem* System = WeakSystem.Get();
+	check(System != nullptr);
 
-			SpawningSystemInstances.Reset();
-			bInSpawnPhase = false;
+	// Transfer promoted instances to the new tick group
+	//-OPT: This can be done async
+	for (FNiagaraSystemInstance* Instance : PendingTickGroupPromotions)
+	{
+		const ETickingGroup TickGroup = Instance->CalculateTickGroup();
+		if (TickGroup != SystemTickGroup)
+		{
+			TSharedPtr<FNiagaraSystemSimulation, ESPMode::ThreadSafe> NewSim = WorldManager->GetSystemSimulation(TickGroup, System);
+			NewSim->TransferInstance(this, Instance);
 		}
 	}
+	PendingTickGroupPromotions.Reset();
 
-	if (PendingTickGroupPromotions.Num() > 0)
+	// Move pending system instances into new tick groups
+	int32 SystemIndex = 0;
+	while ( SystemIndex < PendingSystemInstances.Num() )
 	{
-		FNiagaraWorldManager* WorldManager = FNiagaraWorldManager::Get(World);
-		check(WorldManager != nullptr);
-
-		for ( FNiagaraSystemInstance* Instance : PendingTickGroupPromotions)
+		FNiagaraSystemInstance* Instance = PendingSystemInstances[SystemIndex];
+		if ( !Instance->IsPaused() )
 		{
-			const ETickingGroup TickGroup = Instance->CalculateTickGroup();
-			if ( TickGroup != Instance->TickGroup )
+			const ETickingGroup DesiredTickGroup = Instance->CalculateTickGroup();
+			if (DesiredTickGroup != SystemTickGroup)
 			{
-				TSharedPtr<FNiagaraSystemSimulation, ESPMode::ThreadSafe> NewSim = WorldManager->GetSystemSimulation(TickGroup, System);
-				NewSim->TransferInstance(this, Instance);
-				Instance->TickGroup = TickGroup;
-				Instance->SystemSimulation = NewSim;
+				PendingSystemInstances.RemoveAtSwap(SystemIndex);
+				if (PendingSystemInstances.IsValidIndex(SystemIndex))
+				{
+					PendingSystemInstances[SystemIndex]->SystemInstanceIndex = SystemIndex;
+				}
+				Instance->SystemInstanceIndex = INDEX_NONE;
+
+				TSharedPtr<FNiagaraSystemSimulation, ESPMode::ThreadSafe> DestSim = WorldManager->GetSystemSimulation(DesiredTickGroup, System);
+
+				Instance->SystemSimulation = DestSim;
+				Instance->SystemInstanceIndex = DestSim->PendingSystemInstances.Add(Instance);
+				continue;
 			}
 		}
-		PendingTickGroupPromotions.Reset();
+		++SystemIndex;
 	}
 }
 
-/** First phase of system sim tick. Must run on GameThread. */
-void FNiagaraSystemSimulation::Tick_GameThread(float DeltaSeconds, const FGraphEventRef& MyCompletionGraphEvent)
+void FNiagaraSystemSimulation::Spawn_GameThread(float DeltaSeconds)
 {
-	check(IsInGameThread());
-	WaitForSystemTickComplete(true);
+	// Early out, nothing to do
+	if (PendingSystemInstances.Num() == 0)
+	{
+		return;
+	}
 
+	SCOPE_CYCLE_COUNTER(STAT_NiagaraSystemSim_SpawnNewGT);
 	SCOPE_CYCLE_COUNTER(STAT_NiagaraOverview_GT);
-	SCOPE_CYCLE_COUNTER(STAT_NiagaraSystemSim_TickGT);
 	CSV_SCOPED_TIMING_STAT_EXCLUSIVE(Niagara);
 	LLM_SCOPE(ELLMTag::Niagara);
+	FScopeCycleCounterUObject AdditionalScope(GetSystem(), GET_STATID(STAT_NiagaraOverview_GT_CNC));
 
-	bHasEverTicked = true;
-
-	SystemTickGraphEvent = nullptr;
-
-	check(SystemInstances.Num() == DataSet.GetCurrentDataChecked().GetNumInstances());
-	check(PausedSystemInstances.Num() == PausedInstanceData.GetCurrentDataChecked().GetNumInstances());
+	WaitForSystemTickComplete(true);
 
-	UNiagaraSystem* System = WeakSystem.Get();
+	bInSpawnPhase = true;
 
-	FScopeCycleCounter SystemStatCounter(System->GetStatID(true, false));
-
-	if (MaxDeltaTime.IsSet())
+	if (bHasEverTicked == false)
 	{
-		DeltaSeconds = FMath::Clamp(DeltaSeconds, 0.0f, MaxDeltaTime.GetValue());
+		Tick_GameThread(DeltaSeconds, nullptr);
 	}
 
-	UNiagaraScript* SystemSpawnScript = System->GetSystemSpawnScript();
-	UNiagaraScript* SystemUpdateScript = System->GetSystemUpdateScript();
-#if WITH_EDITOR
-	SystemSpawnScript->RapidIterationParameters.Tick();
-	SystemUpdateScript->RapidIterationParameters.Tick();
-#endif
+	// Spawn instances
+	SpawningInstances.Reserve(PendingSystemInstances.Num());
 
 	int32 SystemIndex = 0;
-	while (SystemIndex < SystemInstances.Num())
+	while (SystemIndex < PendingSystemInstances.Num())
 	{
-		FNiagaraSystemInstance* Inst = SystemInstances[SystemIndex];
-		Inst->Tick_GameThread(DeltaSeconds);
-
-		// TickDataInterfaces could remove the system so we only increment if the system has changed
-		if (Inst->SystemInstanceIndex != INDEX_NONE)
+		FNiagaraSystemInstance* Instance = PendingSystemInstances[SystemIndex];
+		if (Instance->IsPaused())
 		{
-			checkSlow(Inst->SystemInstanceIndex == SystemIndex);
 			++SystemIndex;
+			continue;
 		}
-		else
-		{
-			checkSlow((SystemInstances.Num() <= SystemIndex) || (Inst == SystemInstances[SystemIndex]));
-		}
-	}
+		Instance->Tick_GameThread(DeltaSeconds);
 
-	//Setup the few real constants like delta time.
-	float InvDt = 1.0f / DeltaSeconds;
+		if (Instance->SystemInstanceIndex != INDEX_NONE)
+		{
+			// When the first instance is added we need to initialize the parameter store to data set bindings.
+			if (!bBindingsInitialized)
+			{
+				InitParameterDataSetBindings(Instance);
+			}
 
-	SpawnTimeParam.SetValue(World->TimeSeconds);
-	UpdateTimeParam.SetValue(World->TimeSeconds);
-	SpawnDeltaTimeParam.SetValue(DeltaSeconds);
-	UpdateDeltaTimeParam.SetValue(DeltaSeconds);
-	SpawnInvDeltaTimeParam.SetValue(InvDt);
-	UpdateInvDeltaTimeParam.SetValue(InvDt);
-	SpawnNumSystemInstancesParam.SetValue(SystemInstances.Num());
-	UpdateNumSystemInstancesParam.SetValue(SystemInstances.Num());
-	SpawnGlobalSpawnCountScaleParam.SetValue(INiagaraModule::GetGlobalSpawnCountScale());
-	UpdateGlobalSpawnCountScaleParam.SetValue(INiagaraModule::GetGlobalSpawnCountScale());
-	SpawnGlobalSystemCountScaleParam.SetValue(INiagaraModule::GetGlobalSystemCountScale());
-	UpdateGlobalSystemCountScaleParam.SetValue(INiagaraModule::GetGlobalSystemCountScale());
+			check(PendingSystemInstances[SystemIndex] == Instance);
+			PendingSystemInstances.RemoveAtSwap(SystemIndex);
+			if (PendingSystemInstances.IsValidIndex(SystemIndex))
+			{
+				PendingSystemInstances[SystemIndex]->SystemInstanceIndex = SystemIndex;
+			}
 
-	FNiagaraSystemSimulationTickContext Context(this, DeltaSeconds, false, MyCompletionGraphEvent);
+			Instance->SystemInstanceIndex = SpawningInstances.Add(Instance);
+		}
+	}
 
-	//Now kick of the concurrent tick.
-	if (Context.bTickAsync)
+	if ( SpawningInstances.Num() > 0 )
 	{
-		TGraphTask<FNiagaraSystemSimulationWaitAllFinalizeTask>* WaitAllFinalizeTask = TGraphTask<FNiagaraSystemSimulationWaitAllFinalizeTask>::CreateTask(nullptr, ENamedThreads::GameThread).ConstructAndHold(Context.FinalizeEvents);
-		FGraphEventRef FinalizeGraphEvent = WaitAllFinalizeTask->GetCompletionEvent();
-
-		TGraphTask<FNiagaraSystemSimulationTickTask>* SimulationTickTask = TGraphTask<FNiagaraSystemSimulationTickTask>::CreateTask(nullptr, ENamedThreads::GameThread).ConstructAndHold(Context, WaitAllFinalizeTask);
-		SystemTickGraphEvent = SimulationTickTask->GetCompletionEvent();
-		Context.FinalizeEvents->Add(SystemTickGraphEvent);
+		//-OPT: This can be async :)
 
-		MyCompletionGraphEvent->SetGatherThreadForDontCompleteUntil(ENamedThreads::GameThread);
-		MyCompletionGraphEvent->DontCompleteUntil(FinalizeGraphEvent);
+		FNiagaraSystemSimulationTickContext Context(this, SpawningInstances, SpawningDataSet, DeltaSeconds, SpawningInstances.Num(), nullptr);
+		Tick_Concurrent(Context);
 
-		SimulationTickTask->Unlock(ENamedThreads::GameThread);
-	}
-	else
-	{
-		TGraphTask<FNiagaraSystemSimulationWaitAllFinalizeTask>* WaitAllFinalizeTask = nullptr;
-		if (Context.bTickInstancesAsync)
-		{
-			WaitAllFinalizeTask = TGraphTask<FNiagaraSystemSimulationWaitAllFinalizeTask>::CreateTask(nullptr, ENamedThreads::GameThread).ConstructAndHold(Context.FinalizeEvents);
-		}
+		check(MainDataSet.GetCurrentDataChecked().GetNumInstances() == SystemInstances.Num());
+		check(PausedInstanceData.GetCurrentDataChecked().GetNumInstances() == PausedSystemInstances.Num());
+		check(SpawningDataSet.GetCurrentDataChecked().GetNumInstances() == SpawningInstances.Num());
 
-		Tick_Concurrent(Context);
+		// Append spawned data to our active DataSet
+		SpawningDataSet.CopyTo(MainDataSet, 0, INDEX_NONE, false);
+		SpawningDataSet.ResetBuffers();
 
-		if (Context.bTickInstancesAsync)
+		// Move instances
+		SystemInstances.Reserve(SystemInstances.Num() + SpawningInstances.Num());
+		for ( FNiagaraSystemInstance* Instance : SpawningInstances )
 		{
-			WaitAllFinalizeTask->Unlock(ENamedThreads::GameThread);
-			Context.FinalizeEvents = nullptr;
+			checkSlow(!Instance->IsComplete());
+			Instance->SystemInstanceIndex = SystemInstances.Add(Instance);
 		}
+		SpawningInstances.Reset();
+
+		check(MainDataSet.GetCurrentDataChecked().GetNumInstances() == SystemInstances.Num());
+		check(PausedInstanceData.GetCurrentDataChecked().GetNumInstances() == PausedSystemInstances.Num());
 	}
+
+	bInSpawnPhase = false;
 }
 
 void FNiagaraSystemSimulation::WaitForSystemTickComplete(bool bEnsureComplete)
@@ -821,25 +985,7 @@ void FNiagaraSystemSimulation::WaitForInstancesTickComplete(bool bEnsureComplete
 			++SystemInstIndex;
 		}
 
-		check(DataSet.GetCurrentDataChecked().GetNumInstances() == SystemInstances.Num());
-	}
-
-	if (bInSpawnPhase)
-	{
-		SystemInstIndex = 0;
-		while (SystemInstIndex < SpawningSystemInstances.Num())
-		{
-			FNiagaraSystemInstance* Inst = SpawningSystemInstances[SystemInstIndex];
-			Inst->WaitForAsyncTick();
-
-			// If the system completes during finalize it can be removed from instances so we don't update the index.
-			if (SystemInstances[SystemInstIndex] == Inst)
-			{
-				++SystemInstIndex;
-			}
-
-			check(PendingSpawnDataSet.GetCurrentDataChecked().GetNumInstances() == SpawningSystemInstances.Num());
-		}
+		check(MainDataSet.GetCurrentDataChecked().GetNumInstances() == SystemInstances.Num());
 	}
 }
 
@@ -849,12 +995,11 @@ void FNiagaraSystemSimulation::Tick_Concurrent(FNiagaraSystemSimulationTickConte
 	SCOPE_CYCLE_COUNTER(STAT_NiagaraOverview_GT_CNC);
 	CSV_SCOPED_TIMING_STAT_EXCLUSIVE(Niagara);
 	LLM_SCOPE(ELLMTag::Niagara);
+	FScopeCycleCounterUObject AdditionalScope(GetSystem(), GET_STATID(STAT_NiagaraOverview_GT_CNC));
 
-	TArray<FNiagaraSystemInstance*>& Instances = Context.GetInstances();
-	FNiagaraDataSet& SimulationDataSet = Context.GetDataSet();
-	FNiagaraSystemInstance* SoloSystemInstance = bIsSolo && Instances.Num() == 1 ? Instances[0] : nullptr;
+	FNiagaraSystemInstance* SoloSystemInstance = bIsSolo && Context.Instances.Num() == 1 ? Context.Instances[0] : nullptr;
 
-	if (bCanExecute && Instances.Num())
+	if (bCanExecute && Context.Instances.Num())
 	{
 		if (GbDumpSystemData || Context.System->bDumpDebugSystemInfo)
 		{
@@ -867,7 +1012,7 @@ void FNiagaraSystemSimulation::Tick_Concurrent(FNiagaraSystemSimulationTickConte
 
 		PrepareForSystemSimulate(Context);
 
-		if (Context.bPendingSpawnPass)
+		if (Context.SpawnNum > 0)
 		{
 			SpawnSystemInstances(Context);
 		}
@@ -876,7 +1021,7 @@ void FNiagaraSystemSimulation::Tick_Concurrent(FNiagaraSystemSimulationTickConte
 
 		TransferSystemSimResults(Context);
 
-		for (FNiagaraSystemInstance* Instance : Instances)
+		for (FNiagaraSystemInstance* Instance : Context.Instances)
 		{
 			AddSystemToTickBatch(Instance, Context);
 		}
@@ -887,18 +1032,19 @@ void FNiagaraSystemSimulation::Tick_Concurrent(FNiagaraSystemSimulationTickConte
 		{
 			check(IsInGameThread());
 			int32 SystemInstIndex = 0;
-			while (SystemInstIndex < Instances.Num())
+			while (SystemInstIndex < Context.Instances.Num())
 			{
-				FNiagaraSystemInstance* Inst = Instances[SystemInstIndex];
+				FNiagaraSystemInstance* Inst = Context.Instances[SystemInstIndex];
 				checkSlow(!Inst->IsComplete());
 				Inst->FinalizeTick_GameThread();
+
+				// If the system completes during finalize it will be removed from the instances, therefore we do not need to increment our system index;
 				if (!Inst->IsComplete())
 				{
-					//If the system completes during finalize it will be removed from instances so we don't update the index.
 					++SystemInstIndex;
 				}
 
-				check(SimulationDataSet.GetCurrentDataChecked().GetNumInstances() == Instances.Num());
+				check(Context.DataSet.GetCurrentDataChecked().GetNumInstances() == Context.Instances.Num());
 			}
 		}
 
@@ -909,10 +1055,7 @@ void FNiagaraSystemSimulation::Tick_Concurrent(FNiagaraSystemSimulationTickConte
 		}
 	#endif
 
-		if (!Context.bPendingSpawnPass)
-		{
-			INC_DWORD_STAT_BY(STAT_NiagaraNumSystems, Instances.Num());
-		}
+		INC_DWORD_STAT_BY(STAT_NiagaraNumSystems, Context.Instances.Num());
 	}
 }
 
@@ -920,9 +1063,7 @@ void FNiagaraSystemSimulation::PrepareForSystemSimulate(FNiagaraSystemSimulation
 {
 	SCOPE_CYCLE_COUNTER(STAT_NiagaraSystemSim_PrepareForSimulateCNC);
 
-	TArray<FNiagaraSystemInstance*>& Instances = Context.GetInstances();
-	int32 NumInstances = Instances.Num();
-
+	const int32 NumInstances = Context.Instances.Num();
 	if (NumInstances == 0)
 	{
 		return;
@@ -943,7 +1084,7 @@ void FNiagaraSystemSimulation::PrepareForSystemSimulate(FNiagaraSystemSimulation
 	//Tick instance parameters and transfer any needed into the system simulation dataset.
 	auto TransferInstanceParameters = [&](int32 SystemIndex)
 	{
-		FNiagaraSystemInstance* Inst = Instances[SystemIndex];
+		FNiagaraSystemInstance* Inst = Context.Instances[SystemIndex];
 
 		//Inst->TickInstanceParameters(Context.DeltaSeconds);
 
@@ -969,7 +1110,7 @@ void FNiagaraSystemSimulation::PrepareForSystemSimulate(FNiagaraSystemSimulation
 
 	//This can go wide if we have a very large number of instances.
 	//ParallelFor(Instances.Num(), TransferInstanceParameters, Instances.Num() < GbSystemSimTransferParamsParallelThreshold);
-	ParallelFor(Instances.Num(), TransferInstanceParameters, true);
+	ParallelFor(Context.Instances.Num(), TransferInstanceParameters, true);
 
 	SpawnInstanceParameterDataSet.GetDestinationDataChecked().SetNumInstances(NumInstances);
 	UpdateInstanceParameterDataSet.GetDestinationDataChecked().SetNumInstances(NumInstances);
@@ -986,67 +1127,69 @@ void FNiagaraSystemSimulation::SpawnSystemInstances(FNiagaraSystemSimulationTick
 
 	SCOPE_CYCLE_COUNTER(STAT_NiagaraSystemSim_SpawnCNC);
 
-	TArray<FNiagaraSystemInstance*>& Instances = Context.GetInstances();
-	FNiagaraDataSet& SimulationDataSet = Context.GetDataSet();
-	int32 NumInstances = Instances.Num();
-	
-	check(Context.bPendingSpawnPass);
+	const int32 NumInstances = Context.Instances.Num();
+	const int32 OrigNum = Context.Instances.Num() - Context.SpawnNum;
+	const int32 SpawnNum = Context.SpawnNum;
 
-	if (NumInstances > 0)
-	{
+	check(NumInstances >= Context.SpawnNum);
 
-		FNiagaraSystemInstance* SoloSystemInstance = bIsSolo && Instances.Num() == 1 ? Instances[0] : nullptr;
-		SimulationDataSet.BeginSimulate();
-		SimulationDataSet.Allocate(NumInstances);
-		SimulationDataSet.GetDestinationDataChecked().SetNumInstances(NumInstances);
+	FNiagaraSystemInstance* SoloSystemInstance = bIsSolo && Context.Instances.Num() == 1 ? Context.Instances[0] : nullptr;
+	Context.DataSet.BeginSimulate();
+	Context.DataSet.Allocate(NumInstances, true);
+	Context.DataSet.GetDestinationDataChecked().SetNumInstances(NumInstances);
 
-		//Run Spawn
-		SpawnExecContext.Tick(SoloSystemInstance);//We can't require a specific instance here as these are for all instances.
-		SpawnExecContext.BindData(0, SimulationDataSet, 0, false);
-		SpawnExecContext.BindData(1, SpawnInstanceParameterDataSet, 0, false);
-		SpawnExecContext.Execute(NumInstances);
+	//Run Spawn
+	SpawnExecContext.Tick(SoloSystemInstance);//We can't require a specific instance here as these are for all instances.
+	SpawnExecContext.BindData(0, Context.DataSet, OrigNum, false);
+	SpawnExecContext.BindData(1, SpawnInstanceParameterDataSet, OrigNum, false);
+	SpawnExecContext.Execute(Context.SpawnNum);
 
-		if (GbDumpSystemData || Context.System->bDumpDebugSystemInfo)
+	if (GbDumpSystemData || Context.System->bDumpDebugSystemInfo)
+	{
+		UE_LOG(LogNiagara, Log, TEXT("=== Spwaned %d Systems ==="), NumInstances);
+		Context.DataSet.GetDestinationDataChecked().Dump(0, NumInstances, TEXT("System Dataset - Post Spawn"));
+		SpawnInstanceParameterDataSet.GetCurrentDataChecked().Dump(0, NumInstances, TEXT("Spawn Instance Parameter Data"));
+	}
+
+#if WITH_EDITORONLY_DATA
+	if (SoloSystemInstance && SoloSystemInstance->ShouldCaptureThisFrame())
+	{
+		TSharedPtr<struct FNiagaraScriptDebuggerInfo, ESPMode::ThreadSafe> DebugInfo = SoloSystemInstance->GetActiveCaptureWrite(NAME_None, ENiagaraScriptUsage::SystemSpawnScript, FGuid());
+		if (DebugInfo)
 		{
-			UE_LOG(LogNiagara, Log, TEXT("=== Spwaned %d Systems ==="), NumInstances);
-			SimulationDataSet.GetDestinationDataChecked().Dump(0, NumInstances, TEXT("System Dataset - Post Spawn"));
-			SpawnInstanceParameterDataSet.GetCurrentDataChecked().Dump(0, NumInstances, TEXT("Spawn Instance Parameter Data"));
+			Context.DataSet.CopyTo(DebugInfo->Frame, OrigNum, SpawnNum);
+			DebugInfo->Parameters = UpdateExecContext.Parameters;
+			DebugInfo->bWritten = true;
 		}
-
-		SimulationDataSet.EndSimulate();
 	}
+#endif
 
-	check(SimulationDataSet.GetCurrentDataChecked().GetNumInstances() == Instances.Num());
+	Context.DataSet.EndSimulate();
+
+	check(Context.DataSet.GetCurrentDataChecked().GetNumInstances() == Context.Instances.Num());
 }
 
 void FNiagaraSystemSimulation::UpdateSystemInstances(FNiagaraSystemSimulationTickContext& Context)
 {
 	SCOPE_CYCLE_COUNTER(STAT_NiagaraSystemSim_UpdateCNC);
 
-	TArray<FNiagaraSystemInstance*>& Instances = Context.GetInstances();
-	FNiagaraDataSet& SimulationDataSet = Context.GetDataSet();
-	int32 NumInstances = Instances.Num();
+	const int32 NumInstances = Context.Instances.Num();
+	const int32 OrigNum = Context.Instances.Num() - Context.SpawnNum;
+	const int32 SpawnNum = Context.SpawnNum;
 
 	if (NumInstances > 0)
 	{
-		if (Context.bPendingSpawnPass)
-		{
-			UpdateDeltaTimeParam.SetValue(0.0001f);
-			UpdateInvDeltaTimeParam.SetValue(10000.0f);
-		}
-
+		FNiagaraSystemInstance* SoloSystemInstance = bIsSolo && Context.Instances.Num() == 1 ? Context.Instances[0] : nullptr;
 
-		FNiagaraSystemInstance* SoloSystemInstance = bIsSolo && Instances.Num() == 1 ? Instances[0] : nullptr;
-
-		FNiagaraDataBuffer& DestinationData = SimulationDataSet.BeginSimulate();
+		FNiagaraDataBuffer& DestinationData = Context.DataSet.BeginSimulate();
 		DestinationData.Allocate(NumInstances);
 		DestinationData.SetNumInstances(NumInstances);
 
 		//Run update.
-		UpdateExecContext.Tick(Instances[0]);
-		UpdateExecContext.BindData(0, SimulationDataSet, 0, false);
+		UpdateExecContext.Tick(Context.Instances[0]);
+		UpdateExecContext.BindData(0, Context.DataSet, 0, false);
 		UpdateExecContext.BindData(1, UpdateInstanceParameterDataSet, 0, false);
-		UpdateExecContext.Execute(NumInstances);
+		UpdateExecContext.Execute(OrigNum);
 
 		if (GbDumpSystemData || Context.System->bDumpDebugSystemInfo)
 		{
@@ -1055,26 +1198,36 @@ void FNiagaraSystemSimulation::UpdateSystemInstances(FNiagaraSystemSimulationTic
 			UpdateInstanceParameterDataSet.GetCurrentDataChecked().Dump(0, NumInstances, TEXT("Update Instance Paramter Data"));
 		}
 
-		SimulationDataSet.EndSimulate();
-
-#if WITH_EDITORONLY_DATA
-		if (SoloSystemInstance && SoloSystemInstance->ShouldCaptureThisFrame())
+		//Also run the update script on the newly spawned systems too.
+		//TODO: JIRA - UE-60096 - Remove.
+		//Ideally this should be compiled directly into the script similarly to interpolated particle spawning.
+		if ( (SpawnNum > 0) && GbSystemUpdateOnSpawn)
 		{
-			TSharedPtr<struct FNiagaraScriptDebuggerInfo, ESPMode::ThreadSafe> DebugInfo;
+			UpdateExecContext.Tick(Context.Instances[0]);
+			UpdateExecContext.Parameters.SetParameterValue(0.0001f, SYS_PARAM_ENGINE_DELTA_TIME);
+			UpdateExecContext.Parameters.SetParameterValue(10000.0f, SYS_PARAM_ENGINE_INV_DELTA_TIME);
 
-			if (Context.bPendingSpawnPass)
-			{
-				//We say this update phase is the out put of the SpawnScript but it's really just the initial post spwan state of the particles. We may want to alter the UI to reflect this?
-				SoloSystemInstance->GetActiveCaptureWrite(NAME_None, ENiagaraScriptUsage::SystemSpawnScript, FGuid());
-			}
-			else
+			UpdateExecContext.BindData(0, Context.DataSet, OrigNum, false);
+			UpdateExecContext.BindData(1, UpdateInstanceParameterDataSet, OrigNum, false);
+			UpdateExecContext.Execute(SpawnNum);
+
+			if (GbDumpSystemData || Context.System->bDumpDebugSystemInfo)
 			{
-				SoloSystemInstance->GetActiveCaptureWrite(NAME_None, ENiagaraScriptUsage::SystemUpdateScript, FGuid());
+				UE_LOG(LogNiagara, Log, TEXT("=== Spawn Updated %d Systems ==="), SpawnNum);
+				DestinationData.Dump(OrigNum, SpawnNum, TEXT("System Data - Post Update (new systems)"));
+				UpdateInstanceParameterDataSet.GetCurrentDataChecked().Dump(OrigNum, SpawnNum, TEXT("Update Instance Paramter Data (new systems)"));
 			}
+		}
+
+		Context.DataSet.EndSimulate();
 
+#if WITH_EDITORONLY_DATA
+		if (SoloSystemInstance && SoloSystemInstance->ShouldCaptureThisFrame())
+		{
+			TSharedPtr<struct FNiagaraScriptDebuggerInfo, ESPMode::ThreadSafe> DebugInfo = SoloSystemInstance->GetActiveCaptureWrite(NAME_None, ENiagaraScriptUsage::SystemUpdateScript, FGuid());
 			if (DebugInfo)
 			{
-				SimulationDataSet.CopyTo(DebugInfo->Frame);
+				Context.DataSet.CopyTo(DebugInfo->Frame);
 				DebugInfo->Parameters = UpdateExecContext.Parameters;
 				DebugInfo->bWritten = true;
 			}
@@ -1082,37 +1235,35 @@ void FNiagaraSystemSimulation::UpdateSystemInstances(FNiagaraSystemSimulationTic
 #endif
 	}
 
-	check(SimulationDataSet.GetCurrentDataChecked().GetNumInstances() == Instances.Num());
+	check(Context.DataSet.GetCurrentDataChecked().GetNumInstances() == Context.Instances.Num());
 }
 
 void FNiagaraSystemSimulation::TransferSystemSimResults(FNiagaraSystemSimulationTickContext& Context)
 {
 	SCOPE_CYCLE_COUNTER(STAT_NiagaraSystemSim_TransferResultsCNC);
 
-	TArray<FNiagaraSystemInstance*>& Instances = Context.GetInstances();
-	if (Instances.Num() == 0)
+	if (Context.Instances.Num() == 0)
 	{
 		return;
 	}
 
-	FNiagaraDataSet& SimulationDataSet = Context.GetDataSet();
-	SystemExecutionStateAccessor.SetDataSet(SimulationDataSet);
+	SystemExecutionStateAccessor.SetDataSet(Context.DataSet);
 	SystemExecutionStateAccessor.InitForAccess();
 	for (int32 EmitterIdx = 0; EmitterIdx < Context.System->GetNumEmitters(); ++EmitterIdx)
 	{
-		EmitterExecutionStateAccessors[EmitterIdx].SetDataSet(SimulationDataSet);
+		EmitterExecutionStateAccessors[EmitterIdx].SetDataSet(Context.DataSet);
 		EmitterExecutionStateAccessors[EmitterIdx].InitForAccess();
 		for (int32 SpawnInfoIdx = 0; SpawnInfoIdx < EmitterSpawnInfoAccessors[EmitterIdx].Num(); ++SpawnInfoIdx)
 		{
-			EmitterSpawnInfoAccessors[EmitterIdx][SpawnInfoIdx].SetDataSet(SimulationDataSet);
+			EmitterSpawnInfoAccessors[EmitterIdx][SpawnInfoIdx].SetDataSet(Context.DataSet);
 			EmitterSpawnInfoAccessors[EmitterIdx][SpawnInfoIdx].InitForAccess();
 		}
 	}
 
-	for (int32 SystemIndex = 0; SystemIndex < Instances.Num(); ++SystemIndex)
+	for (int32 SystemIndex = 0; SystemIndex < Context.Instances.Num(); ++SystemIndex)
 	{
 		ENiagaraExecutionState ExecutionState = (ENiagaraExecutionState)SystemExecutionStateAccessor.GetSafe(SystemIndex, (int32)ENiagaraExecutionState::Disabled);
-		FNiagaraSystemInstance* SystemInst = Instances[SystemIndex];
+		FNiagaraSystemInstance* SystemInst = Context.Instances[SystemIndex];
 
 		//Apply the systems requested execution state to it's actual execution state.
 		SystemInst->SetActualExecutionState(ExecutionState);
@@ -1151,10 +1302,10 @@ void FNiagaraSystemSimulation::TransferSystemSimResults(FNiagaraSystemSimulation
 				//TODO: Any other fixed function stuff like this?
 
 				FNiagaraScriptExecutionContext& SpawnContext = EmitterInst.GetSpawnExecutionContext();
-				DataSetToEmitterSpawnParameters[EmitterIdx].DataSetToParameterStore(SpawnContext.Parameters, SimulationDataSet, SystemIndex);
+				DataSetToEmitterSpawnParameters[EmitterIdx].DataSetToParameterStore(SpawnContext.Parameters, Context.DataSet, SystemIndex);
 
 				FNiagaraScriptExecutionContext& UpdateContext = EmitterInst.GetUpdateExecutionContext();
-				DataSetToEmitterUpdateParameters[EmitterIdx].DataSetToParameterStore(UpdateContext.Parameters, SimulationDataSet, SystemIndex);
+				DataSetToEmitterUpdateParameters[EmitterIdx].DataSetToParameterStore(UpdateContext.Parameters, Context.DataSet, SystemIndex);
 
 				TArray<FNiagaraScriptExecutionContext>& EventContexts = EmitterInst.GetEventExecutionContexts();
 				for (int32 EventIdx = 0; EventIdx < EventContexts.Num(); ++EventIdx)
@@ -1162,7 +1313,7 @@ void FNiagaraSystemSimulation::TransferSystemSimResults(FNiagaraSystemSimulation
 					FNiagaraScriptExecutionContext& EventContext = EventContexts[EventIdx];
 					if (DataSetToEmitterEventParameters[EmitterIdx].Num() > EventIdx)
 					{
-						DataSetToEmitterEventParameters[EmitterIdx][EventIdx].DataSetToParameterStore(EventContext.Parameters, SimulationDataSet, SystemIndex);
+						DataSetToEmitterEventParameters[EmitterIdx][EventIdx].DataSetToParameterStore(EventContext.Parameters, Context.DataSet, SystemIndex);
 					}
 					else
 					{
@@ -1181,31 +1332,29 @@ void FNiagaraSystemSimulation::RemoveInstance(FNiagaraSystemInstance* Instance)
 		return;
 	}
 
-	check(SystemInstances.Num() == DataSet.GetCurrentDataChecked().GetNumInstances());
+	check(SystemInstances.Num() == MainDataSet.GetCurrentDataChecked().GetNumInstances());
 	check(PausedSystemInstances.Num() == PausedInstanceData.GetCurrentDataChecked().GetNumInstances());
 
 	check(IsInGameThread());
-	check(DataSet.GetDestinationData() == nullptr);
 	UNiagaraSystem* System = WeakSystem.Get();
 	if (Instance->IsPendingSpawn())
 	{
 		if (GbDumpSystemData || (System && System->bDumpDebugSystemInfo))
 		{
 			UE_LOG(LogNiagara, Log, TEXT("=== Removing Pending Spawn %d ==="), Instance->SystemInstanceIndex);
-			DataSet.GetCurrentDataChecked().Dump(Instance->SystemInstanceIndex, 1, TEXT("System data being removed."));
+			MainDataSet.GetCurrentDataChecked().Dump(Instance->SystemInstanceIndex, 1, TEXT("System data being removed."));
 		}
 
-		//If we're currently spawning then all our pending spawn instances will be inside
-		TArray<FNiagaraSystemInstance*>& Instances = bInSpawnPhase ? SpawningSystemInstances : PendingSystemInstances;
+		// Note: If we go async with PostActor spawning we will need to ensure this remove doesn't happen other than inside our task
+		TArray<FNiagaraSystemInstance*>& Instances = bInSpawnPhase ? SpawningInstances : PendingSystemInstances;
 
 		int32 SystemIndex = Instance->SystemInstanceIndex;
 		check(Instances.IsValidIndex(SystemIndex));
 		check(Instance == Instances[SystemIndex]);
-		
-		if (bInSpawnPhase)
+
+		if ( bInSpawnPhase )
 		{
-			//We should have already spawned this into our dataset so need to kill it.
-			PendingSpawnDataSet.GetCurrentDataChecked().KillInstance(SystemIndex);
+			SpawningDataSet.GetCurrentDataChecked().KillInstance(Instance->SystemInstanceIndex);
 		}
 
 		Instances.RemoveAtSwap(SystemIndex);
@@ -1221,7 +1370,7 @@ void FNiagaraSystemSimulation::RemoveInstance(FNiagaraSystemInstance* Instance)
 		if (GbDumpSystemData || (System && System->bDumpDebugSystemInfo))
 		{
 			UE_LOG(LogNiagara, Log, TEXT("=== Removing Paused %d ==="), Instance->SystemInstanceIndex);
-			DataSet.GetCurrentDataChecked().Dump(Instance->SystemInstanceIndex, 1, TEXT("System data being removed."));
+			MainDataSet.GetCurrentDataChecked().Dump(Instance->SystemInstanceIndex, 1, TEXT("System data being removed."));
 		}
 
 		int32 NumInstances = PausedInstanceData.GetCurrentDataChecked().GetNumInstances();
@@ -1239,7 +1388,7 @@ void FNiagaraSystemSimulation::RemoveInstance(FNiagaraSystemInstance* Instance)
 			PausedSystemInstances[SystemIndex]->SystemInstanceIndex = SystemIndex;
 		}
 
-		check(SystemInstances.Num() == DataSet.GetCurrentDataChecked().GetNumInstances());
+		check(SystemInstances.Num() == MainDataSet.GetCurrentDataChecked().GetNumInstances());
 		check(PausedSystemInstances.Num() == PausedInstanceData.GetCurrentDataChecked().GetNumInstances());
 	}
 	else if (SystemInstances.IsValidIndex(Instance->SystemInstanceIndex))
@@ -1247,19 +1396,19 @@ void FNiagaraSystemSimulation::RemoveInstance(FNiagaraSystemInstance* Instance)
 		if (GbDumpSystemData || (System && System->bDumpDebugSystemInfo))
 		{
 			UE_LOG(LogNiagara, Log, TEXT("=== Removing System %d ==="), Instance->SystemInstanceIndex);
-			DataSet.GetCurrentDataChecked().Dump(Instance->SystemInstanceIndex, 1, TEXT("System data being removed."));
+			MainDataSet.GetCurrentDataChecked().Dump(Instance->SystemInstanceIndex, 1, TEXT("System data being removed."));
 		}
 
 		WaitForInstancesTickComplete();
 
-		int32 NumInstances = DataSet.GetCurrentDataChecked().GetNumInstances();
+		int32 NumInstances = MainDataSet.GetCurrentDataChecked().GetNumInstances();
 		check(SystemInstances.Num() == NumInstances);
 
 		int32 SystemIndex = Instance->SystemInstanceIndex;
 		check(Instance == SystemInstances[SystemIndex]);
 		check(SystemInstances.IsValidIndex(SystemIndex));
 
-		DataSet.GetCurrentDataChecked().KillInstance(SystemIndex);
+		MainDataSet.GetCurrentDataChecked().KillInstance(SystemIndex);
 		SystemInstances.RemoveAtSwap(SystemIndex);
 		Instance->SystemInstanceIndex = INDEX_NONE;
 		if (SystemInstances.IsValidIndex(SystemIndex))
@@ -1267,12 +1416,12 @@ void FNiagaraSystemSimulation::RemoveInstance(FNiagaraSystemInstance* Instance)
 			SystemInstances[SystemIndex]->SystemInstanceIndex = SystemIndex;
 		}
 
-		check(SystemInstances.Num() == DataSet.GetCurrentDataChecked().GetNumInstances());
+		check(SystemInstances.Num() == MainDataSet.GetCurrentDataChecked().GetNumInstances());
 		check(PausedSystemInstances.Num() == PausedInstanceData.GetCurrentDataChecked().GetNumInstances());
 	}
 
 #if NIAGARA_NAN_CHECKING
-	DataSet.CheckForNaNs();
+	MainDataSet.CheckForNaNs();
 #endif
 }
 
@@ -1288,10 +1437,10 @@ void FNiagaraSystemSimulation::AddInstance(FNiagaraSystemInstance* Instance)
 	if (GbDumpSystemData || (System && System->bDumpDebugSystemInfo))
 	{
 		UE_LOG(LogNiagara, Log, TEXT("=== Adding To Pending Spawn %d ==="), Instance->SystemInstanceIndex);
-		//DataSet.Dump(true, Instance->SystemInstanceIndex, 1);
+		//MainDataSet.Dump(true, Instance->SystemInstanceIndex, 1);
 	}
 
-	check(SystemInstances.Num() == DataSet.GetCurrentDataChecked().GetNumInstances());
+	check(SystemInstances.Num() == MainDataSet.GetCurrentDataChecked().GetNumInstances());
 	check(PausedSystemInstances.Num() == PausedInstanceData.GetCurrentDataChecked().GetNumInstances());
 }
 
@@ -1301,10 +1450,10 @@ void FNiagaraSystemSimulation::PauseInstance(FNiagaraSystemInstance* Instance)
 	WaitForInstancesTickComplete();
 
 	check(!Instance->IsPaused());
-	check(!DataSet.GetDestinationData());
+	check(!MainDataSet.GetDestinationData());
 	check(!PausedInstanceData.GetDestinationData());
 
-	check(SystemInstances.Num() == DataSet.GetCurrentDataChecked().GetNumInstances());
+	check(SystemInstances.Num() == MainDataSet.GetCurrentDataChecked().GetNumInstances());
 	check(PausedSystemInstances.Num() == PausedInstanceData.GetCurrentDataChecked().GetNumInstances());
 
 	UNiagaraSystem* System = WeakSystem.Get();
@@ -1313,7 +1462,7 @@ void FNiagaraSystemSimulation::PauseInstance(FNiagaraSystemInstance* Instance)
 		if (GbDumpSystemData || (System && System->bDumpDebugSystemInfo))
 		{
 			UE_LOG(LogNiagara, Log, TEXT("=== Pausing Pending Spawn %d ==="), Instance->SystemInstanceIndex);
-			//DataSet.Dump(true, Instance->SystemInstanceIndex, 1);
+			//MainDataSet.Dump(true, Instance->SystemInstanceIndex, 1);
 		}
 		//Nothing to do for pending spawn systems.
 		check(PendingSystemInstances[Instance->SystemInstanceIndex] == Instance);
@@ -1323,14 +1472,14 @@ void FNiagaraSystemSimulation::PauseInstance(FNiagaraSystemInstance* Instance)
 	if (GbDumpSystemData || (System && System->bDumpDebugSystemInfo))
 	{
 		UE_LOG(LogNiagara, Log, TEXT("=== Pausing System %d ==="), Instance->SystemInstanceIndex);
-		DataSet.GetCurrentDataChecked().Dump(Instance->SystemInstanceIndex, 1, TEXT("System data being paused."));
+		MainDataSet.GetCurrentDataChecked().Dump(Instance->SystemInstanceIndex, 1, TEXT("System data being paused."));
 	}
 
 	int32 SystemIndex = Instance->SystemInstanceIndex;
 	check(SystemInstances.IsValidIndex(SystemIndex));
 	check(Instance == SystemInstances[SystemIndex]);
 
-	int32 NewDataSetIndex = PausedInstanceData.GetCurrentDataChecked().TransferInstance(DataSet.GetCurrentDataChecked(), SystemIndex);
+	int32 NewDataSetIndex = PausedInstanceData.GetCurrentDataChecked().TransferInstance(MainDataSet.GetCurrentDataChecked(), SystemIndex);
 
 	Instance->SystemInstanceIndex = PausedSystemInstances.Add(Instance);
 
@@ -1342,7 +1491,7 @@ void FNiagaraSystemSimulation::PauseInstance(FNiagaraSystemInstance* Instance)
 		SystemInstances[SystemIndex]->SystemInstanceIndex = SystemIndex;
 	}
 
-	check(SystemInstances.Num() == DataSet.GetCurrentDataChecked().GetNumInstances());
+	check(SystemInstances.Num() == MainDataSet.GetCurrentDataChecked().GetNumInstances());
 	check(PausedSystemInstances.Num() == PausedInstanceData.GetCurrentDataChecked().GetNumInstances());
 }
 
@@ -1352,10 +1501,10 @@ void FNiagaraSystemSimulation::UnpauseInstance(FNiagaraSystemInstance* Instance)
 	WaitForInstancesTickComplete();
 
 	check(Instance->IsPaused());
-	check(!DataSet.GetDestinationData());
+	check(!MainDataSet.GetDestinationData());
 	check(!PausedInstanceData.GetDestinationData());
 
-	check(SystemInstances.Num() == DataSet.GetCurrentDataChecked().GetNumInstances());
+	check(SystemInstances.Num() == MainDataSet.GetCurrentDataChecked().GetNumInstances());
 	check(PausedSystemInstances.Num() == PausedInstanceData.GetCurrentDataChecked().GetNumInstances());
 
 	UNiagaraSystem* System = WeakSystem.Get();
@@ -1364,7 +1513,7 @@ void FNiagaraSystemSimulation::UnpauseInstance(FNiagaraSystemInstance* Instance)
 		if (GbDumpSystemData || (System && System->bDumpDebugSystemInfo))
 		{
 			UE_LOG(LogNiagara, Log, TEXT("=== Unpausing Pending Spawn %d ==="), Instance->SystemInstanceIndex);
-			//DataSet.Dump(true, Instance->SystemInstanceIndex, 1);
+			//MainDataSet.Dump(true, Instance->SystemInstanceIndex, 1);
 		}
 		//Nothing to do for pending spawn systems.
 		check(PendingSystemInstances[Instance->SystemInstanceIndex] == Instance);
@@ -1374,14 +1523,14 @@ void FNiagaraSystemSimulation::UnpauseInstance(FNiagaraSystemInstance* Instance)
 	if (GbDumpSystemData || (System && System->bDumpDebugSystemInfo))
 	{
 		UE_LOG(LogNiagara, Log, TEXT("=== Unpausing System %d ==="), Instance->SystemInstanceIndex);
-		DataSet.GetCurrentDataChecked().Dump(Instance->SystemInstanceIndex, 1, TEXT("System data being unpaused."));
+		MainDataSet.GetCurrentDataChecked().Dump(Instance->SystemInstanceIndex, 1, TEXT("System data being unpaused."));
 	}
 
 	int32 SystemIndex = Instance->SystemInstanceIndex;
 	check(PausedSystemInstances.IsValidIndex(SystemIndex));
 	check(Instance == PausedSystemInstances[SystemIndex]);
 
-	int32 NewDataSetIndex = DataSet.GetCurrentDataChecked().TransferInstance(PausedInstanceData.GetCurrentDataChecked(), SystemIndex);
+	int32 NewDataSetIndex = MainDataSet.GetCurrentDataChecked().TransferInstance(PausedInstanceData.GetCurrentDataChecked(), SystemIndex);
 
 	Instance->SystemInstanceIndex = SystemInstances.Add(Instance);
 	check(NewDataSetIndex == Instance->SystemInstanceIndex);
@@ -1392,7 +1541,7 @@ void FNiagaraSystemSimulation::UnpauseInstance(FNiagaraSystemInstance* Instance)
 		PausedSystemInstances[SystemIndex]->SystemInstanceIndex = SystemIndex;
 	}
 
-	check(SystemInstances.Num() == DataSet.GetCurrentDataChecked().GetNumInstances());
+	check(SystemInstances.Num() == MainDataSet.GetCurrentDataChecked().GetNumInstances());
 	check(PausedSystemInstances.Num() == PausedInstanceData.GetCurrentDataChecked().GetNumInstances());
 }
 
@@ -1417,17 +1566,17 @@ void FNiagaraSystemSimulation::InitParameterDataSetBindings(FNiagaraSystemInstan
 		{
 			FNiagaraEmitterInstance& EmitterInst = Emitters[EmitterIdx].Get();
 			FNiagaraScriptExecutionContext& SpawnContext = EmitterInst.GetSpawnExecutionContext();
-			DataSetToEmitterSpawnParameters[EmitterIdx].Init(DataSet, SpawnContext.Parameters);
+			DataSetToEmitterSpawnParameters[EmitterIdx].Init(MainDataSet, SpawnContext.Parameters);
 
 			FNiagaraScriptExecutionContext& UpdateContext = EmitterInst.GetUpdateExecutionContext();
-			DataSetToEmitterUpdateParameters[EmitterIdx].Init(DataSet, UpdateContext.Parameters);
+			DataSetToEmitterUpdateParameters[EmitterIdx].Init(MainDataSet, UpdateContext.Parameters);
 
 			TArray<FNiagaraScriptExecutionContext>& EventContexts = EmitterInst.GetEventExecutionContexts();
 			DataSetToEmitterEventParameters[EmitterIdx].SetNum(EventContexts.Num());
 			for (int32 EventIdx = 0; EventIdx < EventContexts.Num(); ++EventIdx)
 			{
 				FNiagaraScriptExecutionContext& EventContext = EventContexts[EventIdx];
-				DataSetToEmitterEventParameters[EmitterIdx][EventIdx].Init(DataSet, EventContext.Parameters);
+				DataSetToEmitterEventParameters[EmitterIdx][EventIdx].Init(MainDataSet, EventContext.Parameters);
 			}
 		}
 	}
diff --git a/Engine/Plugins/FX/Niagara/Source/Niagara/Private/NiagaraWorldManager.cpp b/Engine/Plugins/FX/Niagara/Source/Niagara/Private/NiagaraWorldManager.cpp
--- a/Engine/Plugins/FX/Niagara/Source/Niagara/Private/NiagaraWorldManager.cpp
+++ b/Engine/Plugins/FX/Niagara/Source/Niagara/Private/NiagaraWorldManager.cpp
@@ -28,6 +28,21 @@ static FAutoConsoleVariableRef CVarNiagaraAllowAsyncWorkToEndOfFrame(
 	ECVF_Default
 );
 
+FAutoConsoleCommandWithWorld DumpNiagaraWorldManagerCommand(
+	TEXT("DumpNiagaraWorldManager"),
+	TEXT("Dump Information About the Niagara World Manager Contents"),
+	FConsoleCommandWithWorldDelegate::CreateLambda(
+		[](UWorld* World)
+		{
+			FNiagaraWorldManager* WorldManager = FNiagaraWorldManager::Get(World);
+			if (WorldManager != nullptr && GLog != nullptr)
+			{
+				WorldManager->DumpDetails(*GLog);
+			}
+		}
+	)
+);
+
 FDelegateHandle FNiagaraWorldManager::OnWorldInitHandle;
 FDelegateHandle FNiagaraWorldManager::OnWorldCleanupHandle;
 FDelegateHandle FNiagaraWorldManager::OnPreWorldFinishDestroyHandle;
@@ -264,7 +279,7 @@ TSharedRef<FNiagaraSystemSimulation, ESPMode::ThreadSafe> FNiagaraWorldManager::
 	
 	TSharedRef<FNiagaraSystemSimulation, ESPMode::ThreadSafe> Sim = MakeShared<FNiagaraSystemSimulation, ESPMode::ThreadSafe>();
 	SystemSimulations[ActualTickGroup].Add(System, Sim);
-	Sim->Init(System, World, false);
+	Sim->Init(System, World, false, TickGroup);
 	return Sim;
 }
 
@@ -302,7 +317,6 @@ void FNiagaraWorldManager::OnBatcherDestroyed_Internal(NiagaraEmitterInstanceBat
 	}
 }
 
-
 void FNiagaraWorldManager::OnWorldCleanup(bool bSessionEnded, bool bCleanupResources)
 {
 	ComponentPool->Cleanup();
@@ -393,25 +407,29 @@ void FNiagaraWorldManager::PostActorTick(float DeltaSeconds)
 	SCOPE_CYCLE_COUNTER(STAT_NiagaraWorldManTick);
 	SCOPE_CYCLE_COUNTER(STAT_NiagaraOverview_GT);
 
-	//-TODO: This doesn't all have to happen here...
-	for (int TG = 0; TG < NiagaraNumTickGroups; ++TG)
+	// Resolve tick groups for pending spawn instances
+	for (int TG=0; TG < NiagaraNumTickGroups; ++TG)
 	{
-		TArray<UNiagaraSystem*, TInlineAllocator<4>> DeadSystems;
 		for (TPair<UNiagaraSystem*, TSharedRef<FNiagaraSystemSimulation, ESPMode::ThreadSafe>>& SystemSim : SystemSimulations[TG])
 		{
-			FNiagaraSystemSimulation*  Sim = &SystemSim.Value.Get();
-			if (Sim->IsValid())
+			FNiagaraSystemSimulation* Sim = &SystemSim.Value.Get();
+			if ( Sim->IsValid() )
 			{
-				Sim->SpawnNew_GameThread(DeltaSeconds, nullptr);
-			}
-			else
-			{
-				DeadSystems.Add(SystemSim.Key);
+				Sim->UpdateTickGroups_GameThread();
 			}
 		}
-		for (UNiagaraSystem* DeadSystem : DeadSystems)
+	}
+
+	// Execute spawn game thread
+	for (int TG = 0; TG < NiagaraNumTickGroups; ++TG)
+	{
+		for (TPair<UNiagaraSystem*, TSharedRef<FNiagaraSystemSimulation, ESPMode::ThreadSafe>>& SystemSim : SystemSimulations[TG])
 		{
-			SystemSimulations[TG].Remove(DeadSystem);
+			FNiagaraSystemSimulation* Sim = &SystemSim.Value.Get();
+			if (Sim->IsValid())
+			{
+				Sim->Spawn_GameThread(DeltaSeconds);
+			}
 		}
 	}
 
@@ -514,3 +532,32 @@ void FNiagaraWorldManager::Tick(ETickingGroup TickGroup, float DeltaSeconds, ELe
 		SystemSimulations[ActualTickGroup].Remove(DeadSystem);
 	}
 }
+
+void FNiagaraWorldManager::DumpDetails(FOutputDevice& Ar)
+{
+	Ar.Logf(TEXT("=== FNiagaraWorldManager Dumping Detailed Information"));
+
+	static const UEnum* TickingGroupEnum = FindObjectChecked<UEnum>(ANY_PACKAGE, TEXT("ETickingGroup"));
+
+	for ( int TG=0; TG < NiagaraNumTickGroups; ++TG )
+	{
+		if (SystemSimulations[TG].Num() == 0 )
+		{
+			continue;
+		}
+
+		Ar.Logf(TEXT("TickingGroup %s"), *TickingGroupEnum->GetNameStringByIndex(TG + NiagaraFirstTickGroup));
+
+		for (TPair<UNiagaraSystem*, TSharedRef<FNiagaraSystemSimulation, ESPMode::ThreadSafe>>& SystemSim : SystemSimulations[TG])
+		{
+			FNiagaraSystemSimulation* Sim = &SystemSim.Value.Get();
+			if ( !Sim->IsValid() )
+			{
+				continue;
+			}
+
+			Ar.Logf(TEXT("\tSimulation %s"), *Sim->GetSystem()->GetFullName());
+			Sim->DumpTickInfo(Ar);
+		}
+	}
+}
diff --git a/Engine/Plugins/FX/Niagara/Source/Niagara/Public/NiagaraSystemInstance.h b/Engine/Plugins/FX/Niagara/Source/Niagara/Public/NiagaraSystemInstance.h
--- a/Engine/Plugins/FX/Niagara/Source/Niagara/Public/NiagaraSystemInstance.h
+++ b/Engine/Plugins/FX/Niagara/Source/Niagara/Public/NiagaraSystemInstance.h
@@ -213,6 +213,9 @@ public:
 	/** Dumps all of this systems info to the log. */
 	void Dump()const;
 
+	/** Dumps information about the instances tick to the log */
+	void DumpTickInfo(FOutputDevice& Ar);
+
 	bool GetPerInstanceDataAndOffsets(void*& OutData, uint32& OutDataSize, TMap<TWeakObjectPtr<UNiagaraDataInterface>, int32>*& OutOffsets);
 
 	NiagaraEmitterInstanceBatcher* GetBatcher() const { return Batcher; }
@@ -246,8 +249,6 @@ private:
 
 	/** Call PrepareForSImulation on each data source from the simulations and determine which need per-tick updates.*/
 	void InitDataInterfaces();	
-
-	void BindParameterCollections(FNiagaraScriptExecutionContext& ExecContext);
 	
 	/** Calculates the distance to use for distance based LODing / culling. */
 	float GetLODDistance();
@@ -263,7 +264,6 @@ private:
 	UNiagaraComponent* Component;
 	FBox SystemBounds;
 
-	ETickingGroup TickGroup;
 	UActorComponent* PrereqComponent;
 
 	/** The age of the System instance. */
diff --git a/Engine/Plugins/FX/Niagara/Source/Niagara/Public/NiagaraSystemSimulation.h b/Engine/Plugins/FX/Niagara/Source/Niagara/Public/NiagaraSystemSimulation.h
--- a/Engine/Plugins/FX/Niagara/Source/Niagara/Public/NiagaraSystemSimulation.h
+++ b/Engine/Plugins/FX/Niagara/Source/Niagara/Public/NiagaraSystemSimulation.h
@@ -125,22 +125,22 @@ struct FNiagaraParameterStoreToDataSetBinding
 
 struct FNiagaraSystemSimulationTickContext
 {
-	class FNiagaraSystemSimulation* Owner;
-	UNiagaraSystem* System;
-	float DeltaSeconds;
-	/** If true, this is a pending spawn pass where we spawn and do an initial tick on newly spawned instances. Then transfer them over to the correct tick group. */
-	bool bPendingSpawnPass;
-	FGraphEventRef MyCompletionGraphEvent;
+	FNiagaraSystemSimulationTickContext(class FNiagaraSystemSimulation* Owner, TArray<FNiagaraSystemInstance*>& Instances, FNiagaraDataSet& DataSet, float DeltaSeconds, int32 SpawnNum, const FGraphEventRef& MyCompletionGraphEvent);
 
-	FGraphEventArray* FinalizeEvents = nullptr;
+	class FNiagaraSystemSimulation*		Owner;
+	UNiagaraSystem*						System;
 
-	bool bTickAsync;
-	bool bTickInstancesAsync;
+	TArray<FNiagaraSystemInstance*>&	Instances;
+	FNiagaraDataSet&					DataSet;
 
-	FNiagaraSystemSimulationTickContext(FNiagaraSystemSimulation* InOwner, float InDeltaSeconds, bool bInPendingSpawnPass, const FGraphEventRef& InMyCompletionGraphEvent);
+	float								DeltaSeconds;
+	int32								SpawnNum;
 
-	TArray<FNiagaraSystemInstance*>& GetInstances();
-	FNiagaraDataSet& GetDataSet();
+	FGraphEventRef						MyCompletionGraphEvent;
+	FGraphEventArray*					FinalizeEvents;
+
+	bool								bTickAsync;
+	bool								bTickInstancesAsync;
 };
 
 /** Simulation performing all system and emitter scripts for a instances of a UNiagaraSystem in a world. */
@@ -149,29 +149,29 @@ class FNiagaraSystemSimulation : public TSharedFromThis<FNiagaraSystemSimulation
 	friend FNiagaraSystemSimulationTickContext;
 public:
 	~FNiagaraSystemSimulation();
-	bool Init(UNiagaraSystem* InSystem, UWorld* InWorld, bool bInIsSolo);
+	bool Init(UNiagaraSystem* InSystem, UWorld* InWorld, bool bInIsSolo, ETickingGroup TickGroup);
 	void Destroy();
 	bool Tick(float DeltaSeconds);
 
 	bool IsValid()const { return WeakSystem.Get() != nullptr && bCanExecute && World != nullptr; }
 
-	/** At the end of the frame, we do a separate tick pass for newly spawned systems which are then pushed into the correct simulations for next frame. */
-	void SpawnNew_GameThread(float DeltaSeconds, const FGraphEventRef& MyCompletionGraphEvent);
-	
 	/** First phase of system sim tick. Must run on GameThread. */
 	void Tick_GameThread(float DeltaSeconds, const FGraphEventRef& MyCompletionGraphEvent);
-
 	/** Second phase of system sim tick that can run on any thread. */
 	void Tick_Concurrent(FNiagaraSystemSimulationTickContext& Context);
 
+	/** Update TickGroups for pending instances and execute tick group promotions. */
+	void UpdateTickGroups_GameThread();
+	/** Spawn any pending instances, assumes that you have update tick groups ahead of time. */
+	void Spawn_GameThread(float DeltaSeconds);
+
+	/** Promote instances that have ticked during */
+
 	/** Wait for system simulation tick to complete.  If bEnsureComplete is true we will trigger an ensure if it is not complete. */
 	void WaitForSystemTickComplete(bool bEnsureComplete = false);
 	/** Wait for instances tick to complete.  If bEnsureComplete is true we will trigger an ensure if it is not complete. */
 	void WaitForInstancesTickComplete(bool bEnsureComplete = false);
 
-	/** Old tick for AB Testing.*/
-	bool Tick_Old(float DeltaSeconds);
-
 	void RemoveInstance(FNiagaraSystemInstance* Instance);
 	void AddInstance(FNiagaraSystemInstance* Instance);
 
@@ -189,6 +189,9 @@ public:
 
 	void DumpInstance(const FNiagaraSystemInstance* Inst)const;
 
+	/** Dump information about all instances tick */
+	void DumpTickInfo(FOutputDevice& Ar);
+
 	bool GetIsSolo() const { return bIsSolo; }
 
 	FNiagaraScriptExecutionContext& GetSpawnExecutionContext() { return SpawnExecContext; }
@@ -217,15 +220,19 @@ protected:
 	/** System of instances being simulated.  We use a weak object ptr here because once the last referencing object goes away this system may be come invalid at runtime. */
 	TWeakObjectPtr<UNiagaraSystem> WeakSystem;
 
+	/** Which tick group we are in, only valid when not in Solo mode. */
+	ETickingGroup SystemTickGroup = TG_MAX;
+
 	/** World this system simulation belongs to. */
 	UWorld* World;
 
 	/** Main dataset containing system instance attribute data. */
-	FNiagaraDataSet DataSet;
+	FNiagaraDataSet MainDataSet;
+	/** DataSet used if we have to spawn instances outside of their tick. */
+	FNiagaraDataSet SpawningDataSet;
+	/** DataSet used to store pausing instance data. */
+	FNiagaraDataSet PausedInstanceData;
 
-	/** Dataset containing system instance data for newly spawned systems. */
-	FNiagaraDataSet PendingSpawnDataSet;
-	
 	/**
 	As there's a 1 to 1 relationship between system instance and their execution in this simulation we must pull all that instances parameters into a dataset for simulation.
 	In some cases this might be a big waste of memory as there'll be duplicated data from a parameter store that's shared across all instances.
@@ -270,14 +277,13 @@ protected:
 
 	/** System instances that have been spawned and are now simulating. */
 	TArray<FNiagaraSystemInstance*> SystemInstances;
-	/** System instances that are pending to be spawned. */
-	TArray<FNiagaraSystemInstance*> PendingSystemInstances;
-	/** System instances that we are spawning this tick. Instances are moved here from PendingSystemInstances if they are not paused. */
-	TArray<FNiagaraSystemInstance*> SpawningSystemInstances;
-
+	/** System instances that are about to be spawned outside of regular ticking. */
+	TArray<FNiagaraSystemInstance*> SpawningInstances;
 	/** System instances that are paused. */
 	TArray<FNiagaraSystemInstance*> PausedSystemInstances;
-	FNiagaraDataSet PausedInstanceData;
+
+	/** System instances that are pending to be spawned. */
+	TArray<FNiagaraSystemInstance*> PendingSystemInstances;
 
 	/** List of instances that are pending a tick group promotion. */
 	TArray<FNiagaraSystemInstance*> PendingTickGroupPromotions;
diff --git a/Engine/Plugins/FX/Niagara/Source/Niagara/Public/NiagaraWorldManager.h b/Engine/Plugins/FX/Niagara/Source/Niagara/Public/NiagaraWorldManager.h
--- a/Engine/Plugins/FX/Niagara/Source/Niagara/Public/NiagaraWorldManager.h
+++ b/Engine/Plugins/FX/Niagara/Source/Niagara/Public/NiagaraWorldManager.h
@@ -129,6 +129,9 @@ public:
 
 	UNiagaraComponentPool* GetComponentPool() { return ComponentPool; }
 
+	// Dump details about what's inside the world manager
+	void DumpDetails(FOutputDevice& Ar);
+
 private:
 
 	// Callback function registered with global world delegates to instantiate world manager when a game world is created
