Modified Niagara component's Activate() call for correct reset behavior.

diff --git a/Engine/Plugins/FX/Niagara/Source/Niagara/Private/NiagaraComponent.cpp b/Engine/Plugins/FX/Niagara/Source/Niagara/Private/NiagaraComponent.cpp
--- a/Engine/Plugins/FX/Niagara/Source/Niagara/Private/NiagaraComponent.cpp
+++ b/Engine/Plugins/FX/Niagara/Source/Niagara/Private/NiagaraComponent.cpp
@@ -794,7 +794,7 @@ void UNiagaraComponent::Activate(bool bReset /* = false */)
 		}
 	}
 
-	FNiagaraSystemInstance::EResetMode ResetMode = bReset ? FNiagaraSystemInstance::EResetMode::ResetAll : FNiagaraSystemInstance::EResetMode::ResetSystem;
+	FNiagaraSystemInstance::EResetMode ResetMode = FNiagaraSystemInstance::EResetMode::ResetSystem;
 	if (InitializeSystem())
 	{
 		ResetMode = FNiagaraSystemInstance::EResetMode::None;//Already done a reinit sete
