CIS warning fix.

diff --git a/Engine/Plugins/FX/Niagara/Source/Niagara/Private/NiagaraScript.cpp b/Engine/Plugins/FX/Niagara/Source/Niagara/Private/NiagaraScript.cpp
--- a/Engine/Plugins/FX/Niagara/Source/Niagara/Private/NiagaraScript.cpp
+++ b/Engine/Plugins/FX/Niagara/Source/Niagara/Private/NiagaraScript.cpp
@@ -56,7 +56,8 @@ UNiagaraScriptSourceBase::UNiagaraScriptSourceBase(const FObjectInitializer& Obj
 
 
 FNiagaraVMExecutableData::FNiagaraVMExecutableData() 
-	: NumUserPtrs(0)
+	: NumTempRegisters(0)
+	, NumUserPtrs(0)
 #if WITH_EDITORONLY_DATA
 	, LastOpCount(0)
 #endif
