Workaround for possible unbuilt niagara register table

diff --git a/Engine/Plugins/FX/Niagara/Source/Niagara/Classes/NiagaraDataSet.h b/Engine/Plugins/FX/Niagara/Source/Niagara/Classes/NiagaraDataSet.h
--- a/Engine/Plugins/FX/Niagara/Source/Niagara/Classes/NiagaraDataSet.h
+++ b/Engine/Plugins/FX/Niagara/Source/Niagara/Classes/NiagaraDataSet.h
@@ -105,7 +105,7 @@ public:
 	void GPUCopyFrom(float* GPUReadBackFloat, int* GPUReadBackInt, int32 StartIdx, int32 NumInstances, uint32 InSrcFloatStride, uint32 InSrcIntStride);
 	void Dump(int32 StartIndex, int32 NumInstances, const FString& Label)const;
 
-	FORCEINLINE void AppendToRegisterTable(FNiagaraRegisterTable& Registers, int32 StartInstance);
+	FORCEINLINE_DEBUGGABLE void AppendToRegisterTable(FNiagaraRegisterTable& Registers, int32 StartInstance);
 	FORCEINLINE void ClearRegisterTable(FNiagaraRegisterTable& Registers);
 
 	const TArray<uint8>& GetFloatBuffer()const { return FloatData; }
@@ -1517,12 +1517,17 @@ FORCEINLINE void FNiagaraDataBuffer::CheckUsage(bool bReadOnly)const
 		(Owner->SimTarget == ENiagaraSimTarget::GPUComputeSim) && IsInRenderingThread());
 }
 
-FORCEINLINE void FNiagaraDataBuffer::AppendToRegisterTable(FNiagaraRegisterTable& Registers, int32 StartInstance)
+FORCEINLINE_DEBUGGABLE void FNiagaraDataBuffer::AppendToRegisterTable(FNiagaraRegisterTable& Registers, int32 StartInstance)
 {
 	checkSlow(Owner && Owner->IsInitialized());
 	checkSlow(Owner->GetSimTarget() == ENiagaraSimTarget::CPUSim);
 	CheckUsage(true);
 
+	if (!ensure(RegisterTable.Num() != 0))
+	{
+		BuildRegisterTable();
+	}
+
 	int32 RegisterCount = RegisterTable.Num();
 	int32 NumRegisters = Registers.Num();
 	Registers.AddUninitialized(RegisterCount);
