- Convert system instance ID to an atomically increasing number rather than a GUID
- Remove system instance ID to FName, this fixes ever increasing memory usage on FNames

#rnx

diff --git a/Engine/Plugins/Experimental/ChaosNiagara/Source/ChaosNiagara/Classes/NiagaraDataInterfaceChaosDestruction.h b/Engine/Plugins/Experimental/ChaosNiagara/Source/ChaosNiagara/Classes/NiagaraDataInterfaceChaosDestruction.h
--- a/Engine/Plugins/Experimental/ChaosNiagara/Source/ChaosNiagara/Classes/NiagaraDataInterfaceChaosDestruction.h
+++ b/Engine/Plugins/Experimental/ChaosNiagara/Source/ChaosNiagara/Classes/NiagaraDataInterfaceChaosDestruction.h
@@ -680,7 +680,7 @@ public:
 		return SolverTime;
 	}
 
-	virtual void ProvidePerInstanceDataForRenderThread(void* DataForRenderThread, void* PerInstanceData, const FGuid& SystemInstance) override;
+	virtual void ProvidePerInstanceDataForRenderThread(void* DataForRenderThread, void* PerInstanceData, const FNiagaraSystemInstanceID& SystemInstance) override;
 protected:
 	virtual bool CopyToInternal(UNiagaraDataInterface* Destination) const override;
 
@@ -766,7 +766,7 @@ protected:
 
 struct FNiagaraDataInterfaceProxyChaosDestruction : public FNiagaraDataInterfaceProxy
 {
-	virtual void ConsumePerInstanceDataFromGameThread(void* PerInstanceData, const FGuid& Instance) override;
+	virtual void ConsumePerInstanceDataFromGameThread(void* PerInstanceData, const FNiagaraSystemInstanceID& Instance) override;
 	virtual int32 PerInstanceDataPassedToRenderThreadSize() const override
 	{
 		return sizeof(FNiagaraDIChaosDestruction_InstanceDataToPassToRT);
diff --git a/Engine/Plugins/Experimental/ChaosNiagara/Source/ChaosNiagara/Private/NiagaraDataInterfaceChaosDestruction.cpp b/Engine/Plugins/Experimental/ChaosNiagara/Source/ChaosNiagara/Private/NiagaraDataInterfaceChaosDestruction.cpp
--- a/Engine/Plugins/Experimental/ChaosNiagara/Source/ChaosNiagara/Private/NiagaraDataInterfaceChaosDestruction.cpp
+++ b/Engine/Plugins/Experimental/ChaosNiagara/Source/ChaosNiagara/Private/NiagaraDataInterfaceChaosDestruction.cpp
@@ -3948,7 +3948,7 @@ void UNiagaraDataInterfaceChaosDestruction::PushToRenderThread()
 	);
 }
 
-void UNiagaraDataInterfaceChaosDestruction::ProvidePerInstanceDataForRenderThread(void* DataForRenderThread, void* PerInstanceData, const FGuid& SystemInstance)
+void UNiagaraDataInterfaceChaosDestruction::ProvidePerInstanceDataForRenderThread(void* DataForRenderThread, void* PerInstanceData, const FNiagaraSystemInstanceID& SystemInstance)
 {
 	check(Proxy);
 
@@ -4095,7 +4095,7 @@ void FNiagaraDataInterfaceProxyChaosDestruction::CreatePerInstanceData(const FGu
 	SystemsToGPUInstanceData.Add(SystemInstance, FNiagaraDIChaosDestruction_GPUData());
 }
 
-void FNiagaraDataInterfaceProxyChaosDestruction::DestroyInstanceData(NiagaraEmitterInstanceBatcher* Batcher, const FGuid& SystemInstance)
+void FNiagaraDataInterfaceProxyChaosDestruction::DestroyInstanceData(NiagaraEmitterInstanceBatcher* Batcher, const FNiagaraSystemInstanceID& SystemInstance)
 {
 	check(IsInRenderingThread());
 	// @todo-threadsafety This object contains GPU buffers. This _should_ delete them safely but would we rather do so manually?
@@ -4105,7 +4105,7 @@ void FNiagaraDataInterfaceProxyChaosDestruction::DestroyInstanceData(NiagaraEmit
 	Batcher->EnqueueDeferredDeletesForDI_RenderThread(this->AsShared());
 }
 
-void FNiagaraDataInterfaceProxyChaosDestruction::ConsumePerInstanceDataFromGameThread(void* PerInstanceDataFromGameThread, const FGuid& Instance) 
+void FNiagaraDataInterfaceProxyChaosDestruction::ConsumePerInstanceDataFromGameThread(void* PerInstanceDataFromGameThread, const FNiagaraSystemInstanceID& Instance)
 { 
 	FNiagaraDIChaosDestruction_GPUData* DataPtr = SystemsToGPUInstanceData.Find(Instance);
 	ensure(DataPtr);
diff --git a/Engine/Plugins/Experimental/HairStrands/Source/HairStrandsNiagara/Private/NiagaraDataInterfaceHairStrands.cpp b/Engine/Plugins/Experimental/HairStrands/Source/HairStrandsNiagara/Private/NiagaraDataInterfaceHairStrands.cpp
--- a/Engine/Plugins/Experimental/HairStrands/Source/HairStrandsNiagara/Private/NiagaraDataInterfaceHairStrands.cpp
+++ b/Engine/Plugins/Experimental/HairStrands/Source/HairStrandsNiagara/Private/NiagaraDataInterfaceHairStrands.cpp
@@ -276,7 +276,7 @@ void FNDIHairStrandsProxy::DeferredDestroy()
 	DeferredDestroyList.Empty();
 }
 
-void FNDIHairStrandsProxy::ConsumePerInstanceDataFromGameThread(void* PerInstanceData, const FGuid& Instance)
+void FNDIHairStrandsProxy::ConsumePerInstanceDataFromGameThread(void* PerInstanceData, const FNiagaraSystemInstanceID& Instance)
 {
 	FNDIHairStrandsData* SourceData = static_cast<FNDIHairStrandsData*>(PerInstanceData);
 	FNDIHairStrandsData* TargetData = SystemInstancesToProxyData.Find(Instance);
@@ -292,7 +292,7 @@ void FNDIHairStrandsProxy::ConsumePerInstanceDataFromGameThread(void* PerInstanc
 	}
 }
 
-void FNDIHairStrandsProxy::InitializePerInstanceData(const FGuid& SystemInstance, FNDIHairStrandsBuffer* StrandBuffer, uint32 NumStrands, const uint8 StrandSize,
+void FNDIHairStrandsProxy::InitializePerInstanceData(const FNiagaraSystemInstanceID& SystemInstance, FNDIHairStrandsBuffer* StrandBuffer, uint32 NumStrands, const uint8 StrandSize,
 	const float StrandDensity, const float RootThickness, const float TipThickness)
 {
 	check(IsInRenderingThread());
@@ -885,7 +885,7 @@ void UNiagaraDataInterfaceHairStrands::GetParameterDefinitionHLSL(FNiagaraDataIn
 	OutHLSL += TEXT("DIHAIRSTRANDS_DECLARE_CONSTANTS(") + ParamInfo.DataInterfaceHLSLSymbol + TEXT(")\n");
 }
 
-void UNiagaraDataInterfaceHairStrands::ProvidePerInstanceDataForRenderThread(void* DataForRenderThread, void* PerInstanceData, const FGuid& SystemInstance)
+void UNiagaraDataInterfaceHairStrands::ProvidePerInstanceDataForRenderThread(void* DataForRenderThread, void* PerInstanceData, const FNiagaraSystemInstanceID& SystemInstance)
 {
 	check(Proxy);
 
diff --git a/Engine/Plugins/Experimental/HairStrands/Source/HairStrandsNiagara/Public/NiagaraDataInterfaceHairStrands.h b/Engine/Plugins/Experimental/HairStrands/Source/HairStrandsNiagara/Public/NiagaraDataInterfaceHairStrands.h
--- a/Engine/Plugins/Experimental/HairStrands/Source/HairStrandsNiagara/Public/NiagaraDataInterfaceHairStrands.h
+++ b/Engine/Plugins/Experimental/HairStrands/Source/HairStrandsNiagara/Public/NiagaraDataInterfaceHairStrands.h
@@ -145,7 +145,7 @@ public:
 	virtual bool GetFunctionHLSL(const FName&  DefinitionFunctionName, FString InstanceFunctionName, FNiagaraDataInterfaceGPUParamInfo& ParamInfo, FString& OutHLSL) override;
 	virtual void GetParameterDefinitionHLSL(FNiagaraDataInterfaceGPUParamInfo& ParamInfo, FString& OutHLSL) override;
 	virtual FNiagaraDataInterfaceParametersCS* ConstructComputeParameters() const override;
-	virtual void ProvidePerInstanceDataForRenderThread(void* DataForRenderThread, void* PerInstanceData, const FGuid& SystemInstance) override;
+	virtual void ProvidePerInstanceDataForRenderThread(void* DataForRenderThread, void* PerInstanceData, const FNiagaraSystemInstanceID& SystemInstance) override;
 	virtual void GetCommonHLSL(FString& OutHLSL) override;
 
 	/** Get the number of strands */
@@ -220,10 +220,10 @@ struct FNDIHairStrandsProxy : public FNiagaraDataInterfaceProxy
 	virtual int32 PerInstanceDataPassedToRenderThreadSize() const override { return sizeof(FNDIHairStrandsData); }
 
 	/** Get the data that will be passed to render*/
-	virtual void ConsumePerInstanceDataFromGameThread(void* PerInstanceData, const FGuid& Instance) override;
+	virtual void ConsumePerInstanceDataFromGameThread(void* PerInstanceData, const FNiagaraSystemInstanceID& Instance) override;
 
 	/** Initialize the Proxy data strands buffer */
-	void InitializePerInstanceData(const FGuid& SystemInstance, FNDIHairStrandsBuffer* StrandsBuffer, const uint32 NumStrands, const uint8 StrandSize,
+	void InitializePerInstanceData(const FNiagaraSystemInstanceID& SystemInstance, FNDIHairStrandsBuffer* StrandsBuffer, const uint32 NumStrands, const uint8 StrandSize,
 		const float StrandDensity, const float RootThickness, const float TipThickness);
 
 	/** Destroy the proxy data if necessary */
diff --git a/Engine/Plugins/FX/Niagara/Source/Niagara/Classes/NiagaraCollision.h b/Engine/Plugins/FX/Niagara/Source/Niagara/Classes/NiagaraCollision.h
--- a/Engine/Plugins/FX/Niagara/Source/Niagara/Classes/NiagaraCollision.h
+++ b/Engine/Plugins/FX/Niagara/Source/Niagara/Classes/NiagaraCollision.h
@@ -35,7 +35,7 @@ public:
 
 	~FNiagaraCollisionBatch()
 	{
-		FNiagaraEventDataSetMgr::Reset(OwnerSystemInstanceName, EmitterName);
+		FNiagaraEventDataSetMgr::Reset(OwnerSystemInstanceID, EmitterName);
 	}
 
 	void Reset()
@@ -46,7 +46,7 @@ public:
 		}
 	}
 
-	void Init(FName InOwnerSystemInstanceName, FName InEmitterName)
+	void Init(FNiagaraSystemInstanceID InOwnerSystemInstanceID, FName InEmitterName)
 	{
 		if (CollisionEventDataSet)
 		{
@@ -57,8 +57,8 @@ public:
 #endif
 		}
 		EmitterName = InEmitterName;
-		OwnerSystemInstanceName = InOwnerSystemInstanceName;
-		CollisionEventDataSet = FNiagaraEventDataSetMgr::CreateEventDataSet(OwnerSystemInstanceName, EmitterName, NIAGARA_BUILTIN_EVENTNAME_COLLISION);
+		OwnerSystemInstanceID = InOwnerSystemInstanceID;
+		CollisionEventDataSet = FNiagaraEventDataSetMgr::CreateEventDataSet(OwnerSystemInstanceID, EmitterName, NIAGARA_BUILTIN_EVENTNAME_COLLISION);
 		// TODO: this should go away once we can use the FNiagaraCollisionEventPayload
 		// UStruct to create the data set
 		//FNiagaraVariable ValidVar(FNiagaraTypeDefinition::GetIntDef(), "Valid");
@@ -87,7 +87,7 @@ private:
 	TArray<FNiagaraCollisionEventPayload> CollisionEvents;
 	FNiagaraDataSet *CollisionEventDataSet;
 	FName EmitterName;
-	FName OwnerSystemInstanceName;
+	FNiagaraSystemInstanceID OwnerSystemInstanceID;
 };
 
 
@@ -140,7 +140,7 @@ public:
 	}
 
 
-	void Init(FName InBatchID, UWorld *InCollisionWorld)
+	void Init(FNiagaraSystemInstanceID InBatchID, UWorld *InCollisionWorld)
 	{
 		BatchID = InBatchID;
 		CollisionWorld = InCollisionWorld;
@@ -163,7 +163,7 @@ private:
 	TArray<FNiagaraCollisionEventPayload> CollisionEvents;
 	FNiagaraDataSet *CollisionEventDataSet;
 
-	FName BatchID;
+	FNiagaraSystemInstanceID BatchID;
 	TArray<FNiagaraCollisionTrace> CollisionTraces[2];
 	TMap<uint32, int32> IdToTraceIdx[2];
 	uint32 CurrBuffer;
diff --git a/Engine/Plugins/FX/Niagara/Source/Niagara/Classes/NiagaraDataInterface.h b/Engine/Plugins/FX/Niagara/Source/Niagara/Classes/NiagaraDataInterface.h
--- a/Engine/Plugins/FX/Niagara/Source/Niagara/Classes/NiagaraDataInterface.h
+++ b/Engine/Plugins/FX/Niagara/Source/Niagara/Classes/NiagaraDataInterface.h
@@ -160,7 +160,7 @@ struct FNiagaraDataInterfaceProxy : TSharedFromThis<FNiagaraDataInterfaceProxy,
 	virtual ~FNiagaraDataInterfaceProxy() {/*check(IsInRenderingThread());*/}
 
 	virtual int32 PerInstanceDataPassedToRenderThreadSize() const = 0;
-	virtual void ConsumePerInstanceDataFromGameThread(void* PerInstanceData, const FGuid& Instance) { check(false); }
+	virtual void ConsumePerInstanceDataFromGameThread(void* PerInstanceData, const FNiagaraSystemInstanceID& Instance) { check(false); }
 
 	virtual void DeferredDestroy() {}
 
@@ -204,7 +204,7 @@ struct FNiagaraDataInterfaceProxyCurveBase : public FNiagaraDataInterfaceProxy
 	}
 
 	// @todo REMOVEME
-	virtual void ConsumePerInstanceDataFromGameThread(void* PerInstanceData, const FGuid& Instance) override { check(false); }
+	virtual void ConsumePerInstanceDataFromGameThread(void* PerInstanceData, const FNiagaraSystemInstanceID& Instance) override { check(false); }
 };
 
 //////////////////////////////////////////////////////////////////////////
@@ -244,7 +244,7 @@ public:
 			
 		This will not be called if PerInstanceDataPassedToRenderThreadSize is 0.
 	*/
-	virtual void ProvidePerInstanceDataForRenderThread(void* DataForRenderThread, void* PerInstanceData, const FGuid& SystemInstance)
+	virtual void ProvidePerInstanceDataForRenderThread(void* DataForRenderThread, void* PerInstanceData, const FNiagaraSystemInstanceID& SystemInstance)
 	{
 		check(false);
 	}
diff --git a/Engine/Plugins/FX/Niagara/Source/Niagara/Classes/NiagaraDataInterfaceCurlNoise.h b/Engine/Plugins/FX/Niagara/Source/Niagara/Classes/NiagaraDataInterfaceCurlNoise.h
--- a/Engine/Plugins/FX/Niagara/Source/Niagara/Classes/NiagaraDataInterfaceCurlNoise.h
+++ b/Engine/Plugins/FX/Niagara/Source/Niagara/Classes/NiagaraDataInterfaceCurlNoise.h
@@ -56,7 +56,7 @@ struct FNiagaraDataInterfaceProxyCurlNoise : public FNiagaraDataInterfaceProxy
 		OffsetFromSeed = InOffset;
 	}
 
-	virtual void ConsumePerInstanceDataFromGameThread(void* PerInstanceData, const FGuid& Instance) override { check(false); }
+	virtual void ConsumePerInstanceDataFromGameThread(void* PerInstanceData, const FNiagaraSystemInstanceID& Instance) override { check(false); }
 	virtual int32 PerInstanceDataPassedToRenderThreadSize() const override
 	{
 		return 0;
diff --git a/Engine/Plugins/FX/Niagara/Source/Niagara/Classes/NiagaraDataInterfaceGrid2DCollection.h b/Engine/Plugins/FX/Niagara/Source/Niagara/Classes/NiagaraDataInterfaceGrid2DCollection.h
--- a/Engine/Plugins/FX/Niagara/Source/Niagara/Classes/NiagaraDataInterfaceGrid2DCollection.h
+++ b/Engine/Plugins/FX/Niagara/Source/Niagara/Classes/NiagaraDataInterfaceGrid2DCollection.h
@@ -89,15 +89,15 @@ struct FNiagaraDataInterfaceProxyGrid2DCollection : public FNiagaraDataInterface
 	virtual void ResetData(FRHICommandList& RHICmdList, const FNiagaraDataInterfaceSetArgs& Context) override;
 
 	virtual void DeferredDestroy() override;		
-	void DestroyPerInstanceData(NiagaraEmitterInstanceBatcher* Batcher, const FGuid& SystemInstance);
+	void DestroyPerInstanceData(NiagaraEmitterInstanceBatcher* Batcher, const FNiagaraSystemInstanceID& SystemInstance);
 
 	/* List of proxy data for each system instances*/
 	// #todo(dmp): this should all be refactored to avoid duplicate code
-	TMap<FGuid, Grid2DCollectionRWInstanceData> SystemInstancesToProxyData;
+	TMap<FNiagaraSystemInstanceID, Grid2DCollectionRWInstanceData> SystemInstancesToProxyData;
 
 	/* List of proxy data to destroy later */
 	// #todo(dmp): this should all be refactored to avoid duplicate code
-	TSet<FGuid> DeferredDestroyList;
+	TSet<FNiagaraSystemInstanceID> DeferredDestroyList;
 };
 
 UCLASS(EditInlineNew, Category = "Grid", meta = (DisplayName = "Grid2D Collection"), Blueprintable, BlueprintType)
@@ -121,7 +121,7 @@ public:
 	virtual bool GetFunctionHLSL(const FName&  DefinitionFunctionName, FString InstanceFunctionName, FNiagaraDataInterfaceGPUParamInfo& ParamInfo, FString& OutHLSL) override;
 	virtual FNiagaraDataInterfaceParametersCS* ConstructComputeParameters() const override;
 	
-	virtual void ProvidePerInstanceDataForRenderThread(void* DataForRenderThread, void* PerInstanceData, const FGuid& SystemInstance) override {}
+	virtual void ProvidePerInstanceDataForRenderThread(void* DataForRenderThread, void* PerInstanceData, const FNiagaraSystemInstanceID& SystemInstance) override {}
 	virtual bool InitPerInstanceData(void* PerInstanceData, FNiagaraSystemInstance* SystemInstance) override;
 	virtual void DestroyPerInstanceData(void* PerInstanceData, FNiagaraSystemInstance* SystemInstance) override;
 	virtual bool PerInstanceTick(void* PerInstanceData, FNiagaraSystemInstance* SystemInstance, float DeltaSeconds) override { return false; }
diff --git a/Engine/Plugins/FX/Niagara/Source/Niagara/Classes/NiagaraDataInterfaceNeighborGrid3D.h b/Engine/Plugins/FX/Niagara/Source/Niagara/Classes/NiagaraDataInterfaceNeighborGrid3D.h
--- a/Engine/Plugins/FX/Niagara/Source/Niagara/Classes/NiagaraDataInterfaceNeighborGrid3D.h
+++ b/Engine/Plugins/FX/Niagara/Source/Niagara/Classes/NiagaraDataInterfaceNeighborGrid3D.h
@@ -33,18 +33,18 @@ public:
 struct FNiagaraDataInterfaceProxyNeighborGrid3D : public FNiagaraDataInterfaceProxyRW
 {	
 	virtual void PreStage(FRHICommandList& RHICmdList, const FNiagaraDataInterfaceSetArgs& Context) override;	
-	virtual void ConsumePerInstanceDataFromGameThread(void* PerInstanceData, const FGuid& Instance) override {}
+	virtual void ConsumePerInstanceDataFromGameThread(void* PerInstanceData, const FNiagaraSystemInstanceID& Instance) override {}
 	virtual int32 PerInstanceDataPassedToRenderThreadSize() const override { return sizeof(NeighborGrid3DRWInstanceData); }	
-	void DestroyPerInstanceData(NiagaraEmitterInstanceBatcher* Batcher, const FGuid& SystemInstance);
+	void DestroyPerInstanceData(NiagaraEmitterInstanceBatcher* Batcher, const FNiagaraSystemInstanceID& SystemInstance);
 	void DeferredDestroy();
 
 	/* List of proxy data for each system instances*/
 	// #todo(dmp): this should all be refactored to avoid duplicate code
-	TMap<FGuid, NeighborGrid3DRWInstanceData> SystemInstancesToProxyData;
+	TMap<FNiagaraSystemInstanceID, NeighborGrid3DRWInstanceData> SystemInstancesToProxyData;
 
 	/* List of proxy data to destroy later */
 	// #todo(dmp): this should all be refactored to avoid duplicate code
-	TSet<FGuid> DeferredDestroyList;
+	TSet<FNiagaraSystemInstanceID> DeferredDestroyList;
 };
 
 UCLASS(EditInlineNew, Category = "Grid", meta = (DisplayName = "Neighbor Grid3D"))
@@ -82,7 +82,7 @@ public:
 	virtual bool GetFunctionHLSL(const FName&  DefinitionFunctionName, FString InstanceFunctionName, FNiagaraDataInterfaceGPUParamInfo& ParamInfo, FString& OutHLSL) override;
 	virtual FNiagaraDataInterfaceParametersCS* ConstructComputeParameters() const override;
 		
-	virtual void ProvidePerInstanceDataForRenderThread(void* DataForRenderThread, void* PerInstanceData, const FGuid& SystemInstance) override {}
+	virtual void ProvidePerInstanceDataForRenderThread(void* DataForRenderThread, void* PerInstanceData, const FNiagaraSystemInstanceID& SystemInstance) override {}
 	virtual bool InitPerInstanceData(void* PerInstanceData, FNiagaraSystemInstance* SystemInstance) override;
 	virtual void DestroyPerInstanceData(void* PerInstanceData, FNiagaraSystemInstance* SystemInstance) override;
 	virtual bool PerInstanceTick(void* PerInstanceData, FNiagaraSystemInstance* SystemInstance, float DeltaSeconds) override { return false;  }
diff --git a/Engine/Plugins/FX/Niagara/Source/Niagara/Classes/NiagaraDataInterfaceRW.h b/Engine/Plugins/FX/Niagara/Source/Niagara/Classes/NiagaraDataInterfaceRW.h
--- a/Engine/Plugins/FX/Niagara/Source/Niagara/Classes/NiagaraDataInterfaceRW.h
+++ b/Engine/Plugins/FX/Niagara/Source/Niagara/Classes/NiagaraDataInterfaceRW.h
@@ -39,7 +39,7 @@ struct FNiagaraDataInterfaceProxyRW : public FNiagaraDataInterfaceProxy
 {
 public:
 
-	virtual void ConsumePerInstanceDataFromGameThread(void* PerInstanceData, const FGuid& Instance) override { check(false); }
+	virtual void ConsumePerInstanceDataFromGameThread(void* PerInstanceData, const FNiagaraSystemInstanceID& Instance) override { check(false); }
 	virtual int32 PerInstanceDataPassedToRenderThreadSize() const override
 	{
 		return 0;
diff --git a/Engine/Plugins/FX/Niagara/Source/Niagara/Classes/NiagaraDataInterfaceSimpleCounter.h b/Engine/Plugins/FX/Niagara/Source/Niagara/Classes/NiagaraDataInterfaceSimpleCounter.h
--- a/Engine/Plugins/FX/Niagara/Source/Niagara/Classes/NiagaraDataInterfaceSimpleCounter.h
+++ b/Engine/Plugins/FX/Niagara/Source/Niagara/Classes/NiagaraDataInterfaceSimpleCounter.h
@@ -60,7 +60,7 @@ public:
 	virtual bool GetFunctionHLSL(const FName&  DefinitionFunctionName, FString InstanceFunctionName, FNiagaraDataInterfaceGPUParamInfo& ParamInfo, FString& OutHLSL) override;
 	virtual void GetParameterDefinitionHLSL(FNiagaraDataInterfaceGPUParamInfo& ParamInfo, FString& OutHLSL) override;
 	
-	virtual void ProvidePerInstanceDataForRenderThread(void* DataForRenderThread, void* PerInstanceData, const FGuid& SystemInstance) override
+	virtual void ProvidePerInstanceDataForRenderThread(void* DataForRenderThread, void* PerInstanceData, const FNiagaraSystemInstanceID& SystemInstance) override
 	{
 		check(false);
 	}
diff --git a/Engine/Plugins/FX/Niagara/Source/Niagara/Classes/NiagaraDataInterfaceSkeletalMesh.h b/Engine/Plugins/FX/Niagara/Source/Niagara/Classes/NiagaraDataInterfaceSkeletalMesh.h
--- a/Engine/Plugins/FX/Niagara/Source/Niagara/Classes/NiagaraDataInterfaceSkeletalMesh.h
+++ b/Engine/Plugins/FX/Niagara/Source/Niagara/Classes/NiagaraDataInterfaceSkeletalMesh.h
@@ -535,7 +535,7 @@ public:
 	virtual void GetParameterDefinitionHLSL(FNiagaraDataInterfaceGPUParamInfo& ParamInfo, FString& OutHLSL) override;
 	virtual FNiagaraDataInterfaceParametersCS* ConstructComputeParameters() const override;
 
-	virtual void ProvidePerInstanceDataForRenderThread(void* DataForRenderThread, void* PerInstanceData, const FGuid& SystemInstance) override;
+	virtual void ProvidePerInstanceDataForRenderThread(void* DataForRenderThread, void* PerInstanceData, const FNiagaraSystemInstanceID& SystemInstance) override;
 
 	static const FString MeshIndexBufferName;
 	static const FString MeshVertexBufferName;
@@ -752,7 +752,7 @@ struct FNiagaraDataInterfaceProxySkeletalMesh : public FNiagaraDataInterfaceProx
 		return sizeof(FNiagaraDISkeletalMeshPassedDataToRT);
 	}
 
-	virtual void ConsumePerInstanceDataFromGameThread(void* PerInstanceData, const FGuid& Instance) override;
+	virtual void ConsumePerInstanceDataFromGameThread(void* PerInstanceData, const FNiagaraSystemInstanceID& Instance) override;
 
-	TMap<FGuid, FNiagaraDataInterfaceProxySkeletalMeshData> SystemInstancesToData;
+	TMap<FNiagaraSystemInstanceID, FNiagaraDataInterfaceProxySkeletalMeshData> SystemInstancesToData;
 };
diff --git a/Engine/Plugins/FX/Niagara/Source/Niagara/Classes/NiagaraDataInterfaceStaticMesh.h b/Engine/Plugins/FX/Niagara/Source/Niagara/Classes/NiagaraDataInterfaceStaticMesh.h
--- a/Engine/Plugins/FX/Niagara/Source/Niagara/Classes/NiagaraDataInterfaceStaticMesh.h
+++ b/Engine/Plugins/FX/Niagara/Source/Niagara/Classes/NiagaraDataInterfaceStaticMesh.h
@@ -211,7 +211,7 @@ public:
 	virtual void GetParameterDefinitionHLSL(FNiagaraDataInterfaceGPUParamInfo& ParamInfo, FString& OutHLSL) override;
 	virtual FNiagaraDataInterfaceParametersCS* ConstructComputeParameters()const override;
 
-	virtual void ProvidePerInstanceDataForRenderThread(void* DataForRenderThread, void* PerInstanceData, const FGuid& SystemInstance) override;
+	virtual void ProvidePerInstanceDataForRenderThread(void* DataForRenderThread, void* PerInstanceData, const FNiagaraSystemInstanceID& SystemInstance) override;
 
 	static const FString MeshIndexBufferName;
 	static const FString MeshVertexBufferName;
@@ -361,12 +361,12 @@ struct FNiagaraDataInterfaceProxyStaticMesh : public FNiagaraDataInterfaceProxy
 		return sizeof(FNiagaraPassedInstanceDataForRT);
 	}
 
-	virtual void ConsumePerInstanceDataFromGameThread(void* PerInstanceData, const FGuid& Instance) override;
+	virtual void ConsumePerInstanceDataFromGameThread(void* PerInstanceData, const FNiagaraSystemInstanceID& Instance) override;
 
-	void InitializePerInstanceData(const FGuid& SystemInstance, FStaticMeshGpuSpawnBuffer* MeshGPUSpawnBuffer);
-	void DestroyPerInstanceData(NiagaraEmitterInstanceBatcher* Batcher, const FGuid& SystemInstance);
+	void InitializePerInstanceData(const FNiagaraSystemInstanceID& SystemInstance, FStaticMeshGpuSpawnBuffer* MeshGPUSpawnBuffer);
+	void DestroyPerInstanceData(NiagaraEmitterInstanceBatcher* Batcher, const FNiagaraSystemInstanceID& SystemInstance);
 
-	TMap<FGuid, FNiagaraStaticMeshData> SystemInstancesToMeshData;
+	TMap<FNiagaraSystemInstanceID, FNiagaraStaticMeshData> SystemInstancesToMeshData;
 
-	TSet<FGuid> DeferredDestroyList;
+	TSet<FNiagaraSystemInstanceID> DeferredDestroyList;
 };
diff --git a/Engine/Plugins/FX/Niagara/Source/Niagara/Classes/NiagaraDataInterfaceTexture.h b/Engine/Plugins/FX/Niagara/Source/Niagara/Classes/NiagaraDataInterfaceTexture.h
--- a/Engine/Plugins/FX/Niagara/Source/Niagara/Classes/NiagaraDataInterfaceTexture.h
+++ b/Engine/Plugins/FX/Niagara/Source/Niagara/Classes/NiagaraDataInterfaceTexture.h
@@ -65,7 +65,7 @@ struct FNiagaraDataInterfaceProxyTexture : public FNiagaraDataInterfaceProxy
 	FTextureReferenceRHIRef	TextureReferenceRHI;
 	float TexDims[2];
 
-	virtual void ConsumePerInstanceDataFromGameThread(void* PerInstanceData, const FGuid& Instance) override { check(false); }
+	virtual void ConsumePerInstanceDataFromGameThread(void* PerInstanceData, const FNiagaraSystemInstanceID& Instance) override { check(false); }
 	virtual int32 PerInstanceDataPassedToRenderThreadSize() const override
 	{
 		return 0;
diff --git a/Engine/Plugins/FX/Niagara/Source/Niagara/Classes/NiagaraDataInterfaceVectorField.h b/Engine/Plugins/FX/Niagara/Source/Niagara/Classes/NiagaraDataInterfaceVectorField.h
--- a/Engine/Plugins/FX/Niagara/Source/Niagara/Classes/NiagaraDataInterfaceVectorField.h
+++ b/Engine/Plugins/FX/Niagara/Source/Niagara/Classes/NiagaraDataInterfaceVectorField.h
@@ -87,7 +87,7 @@ struct FNiagaraDataInterfaceProxyVectorField : public FNiagaraDataInterfaceProxy
 		return FVector(float(bTileX), float(bTileY), float(bTileZ));
 	}
 
-	virtual void ConsumePerInstanceDataFromGameThread(void* PerInstanceData, const FGuid& Instance) override { check(false); }
+	virtual void ConsumePerInstanceDataFromGameThread(void* PerInstanceData, const FNiagaraSystemInstanceID& Instance) override { check(false); }
 	virtual int32 PerInstanceDataPassedToRenderThreadSize() const override
 	{
 		return 0;
diff --git a/Engine/Plugins/FX/Niagara/Source/Niagara/Classes/NiagaraEmitterInstance.h b/Engine/Plugins/FX/Niagara/Source/Niagara/Classes/NiagaraEmitterInstance.h
--- a/Engine/Plugins/FX/Niagara/Source/Niagara/Classes/NiagaraEmitterInstance.h
+++ b/Engine/Plugins/FX/Niagara/Source/Niagara/Classes/NiagaraEmitterInstance.h
@@ -33,7 +33,7 @@ public:
 	bool bDumpAfterEvent;
 	virtual ~FNiagaraEmitterInstance();
 
-	void Init(int32 InEmitterIdx, FName SystemInstanceName);
+	void Init(int32 InEmitterIdx, FNiagaraSystemInstanceID SystemInstanceID);
 
 	void ResetSimulation(bool bKillExisting=true);
 
@@ -201,7 +201,7 @@ private:
 	TArray<bool> UpdateEventGeneratorIsSharedByIndex;
 	TArray<bool> SpawnEventGeneratorIsSharedByIndex;
 
-	FName OwnerSystemInstanceName;
+	FNiagaraSystemInstanceID OwnerSystemInstanceID;
 
 	/** Cached fixed bounds of the parent system which override this Emitter Instances bounds if set. Whenever we initialize the owning SystemInstance we will reconstruct this
 	 ** EmitterInstance and the cached bounds will be unset. */
diff --git a/Engine/Plugins/FX/Niagara/Source/Niagara/Classes/NiagaraEvents.h b/Engine/Plugins/FX/Niagara/Source/Niagara/Classes/NiagaraEvents.h
--- a/Engine/Plugins/FX/Niagara/Source/Niagara/Classes/NiagaraEvents.h
+++ b/Engine/Plugins/FX/Niagara/Source/Niagara/Classes/NiagaraEvents.h
@@ -82,12 +82,12 @@ public:
 
 	// call this to allocate a data set for a specific emitter and event
 	// if the data set already exists, returns the existing one
-	static FNiagaraDataSet *CreateEventDataSet(FName OwnerSystemInstanceName, FName EmitterName, FName EventName)
+	static FNiagaraDataSet *CreateEventDataSet(FNiagaraSystemInstanceID OwnerSystemInstanceID, FName EmitterName, FName EventName)
 	{
-		PerSystemInstanceDataSetMap *SystemInstanceMap = EmitterEventDataSets.Find(OwnerSystemInstanceName);
+		PerSystemInstanceDataSetMap *SystemInstanceMap = EmitterEventDataSets.Find(OwnerSystemInstanceID);
 		if (!SystemInstanceMap)
 		{
-			SystemInstanceMap = &EmitterEventDataSets.Add(OwnerSystemInstanceName);
+			SystemInstanceMap = &EmitterEventDataSets.Add(OwnerSystemInstanceID);
 		}
 
 		// find the emitter's map
@@ -116,9 +116,9 @@ public:
 
 	// deletes and cleans up all event data sets for an emitter
 	//  should be called when the emitter is destroyed
-	static void Reset(FName OwnerSystemInstanceName, FName EmitterName)
+	static void Reset(FNiagaraSystemInstanceID OwnerSystemInstanceID, FName EmitterName)
 	{
-		PerSystemInstanceDataSetMap *SystemInstanceMap = EmitterEventDataSets.Find(OwnerSystemInstanceName);
+		PerSystemInstanceDataSetMap *SystemInstanceMap = EmitterEventDataSets.Find(OwnerSystemInstanceID);
 		if (SystemInstanceMap)
 		{
 			PerEmitterEventDataSetMap *Map = SystemInstanceMap->Find(EmitterName);
@@ -133,7 +133,7 @@ public:
 			}
 			if (SystemInstanceMap->Num() == 0)
 			{
-				EmitterEventDataSets.Remove(OwnerSystemInstanceName);
+				EmitterEventDataSets.Remove(OwnerSystemInstanceID);
 			}
 		}
 	}
@@ -159,9 +159,9 @@ public:
 	}
 	*/
 
-	static PerEmitterEventDataSetMap * GetEmitterMap(FName OwnerSystemInstanceName, FName EmitterName)
+	static PerEmitterEventDataSetMap * GetEmitterMap(FNiagaraSystemInstanceID OwnerSystemInstanceID, FName EmitterName)
 	{
-		TMap<FName, PerEmitterEventDataSetMap> *SystemMap = EmitterEventDataSets.Find(OwnerSystemInstanceName);
+		TMap<FName, PerEmitterEventDataSetMap> *SystemMap = EmitterEventDataSets.Find(OwnerSystemInstanceID);
 		if (SystemMap)
 		{
 			return SystemMap->Find(EmitterName);
@@ -170,9 +170,9 @@ public:
 		return nullptr;
 	}
 
-	static FNiagaraDataSet * GetEventDataSet(FName OwnerSystemInstanceName, FName EmitterName, FName EventName)
+	static FNiagaraDataSet * GetEventDataSet(FNiagaraSystemInstanceID OwnerSystemInstanceID, FName EmitterName, FName EventName)
 	{
-		PerEmitterEventDataSetMap *SetMap = GetEmitterMap(OwnerSystemInstanceName, EmitterName);
+		PerEmitterEventDataSetMap *SetMap = GetEmitterMap(OwnerSystemInstanceID, EmitterName);
 		if (!SetMap)
 		{
 			return nullptr;
@@ -188,5 +188,5 @@ public:
 	}
 
 private:
-	static TMap<FName, TMap<FName, PerEmitterEventDataSetMap>> EmitterEventDataSets;
+	static TMap<FNiagaraSystemInstanceID, TMap<FName, PerEmitterEventDataSetMap>> EmitterEventDataSets;
 };
\ No newline at end of file
diff --git a/Engine/Plugins/FX/Niagara/Source/Niagara/Classes/NiagaraScriptExecutionContext.h b/Engine/Plugins/FX/Niagara/Source/Niagara/Classes/NiagaraScriptExecutionContext.h
--- a/Engine/Plugins/FX/Niagara/Source/Niagara/Classes/NiagaraScriptExecutionContext.h
+++ b/Engine/Plugins/FX/Niagara/Source/Niagara/Classes/NiagaraScriptExecutionContext.h
@@ -288,7 +288,7 @@ public:
 	FORCEINLINE FNiagaraComputeInstanceData* GetInstanceData()const{ return reinterpret_cast<FNiagaraComputeInstanceData*>(InstanceData_ParamData_Packed); }
 
 	uint32 Count;
-	FGuid SystemInstanceID;
+	FNiagaraSystemInstanceID SystemInstanceID;
 	FNiagaraDataInterfaceInstanceData* DIInstanceData;
 	uint8* InstanceData_ParamData_Packed;
 	bool bRequiredDistanceFieldData = false;
diff --git a/Engine/Plugins/FX/Niagara/Source/Niagara/Private/NiagaraCommon.cpp b/Engine/Plugins/FX/Niagara/Source/Niagara/Private/NiagaraCommon.cpp
--- a/Engine/Plugins/FX/Niagara/Source/Niagara/Private/NiagaraCommon.cpp
+++ b/Engine/Plugins/FX/Niagara/Source/Niagara/Private/NiagaraCommon.cpp
@@ -303,6 +303,21 @@ void FNiagaraUtilities::DumpHLSLText(const FString& SourceCode, const FString& D
 	UE_LOG(LogNiagara, Display, TEXT("==================================================================================="));
 }
 
+FString FNiagaraUtilities::SystemInstanceIDToString(FNiagaraSystemInstanceID ID)
+{
+	TCHAR Buffer[17];
+	uint64 Value = ID;
+	for (int i = 15; i >= 0; --i)
+	{
+		TCHAR ch = Value & 0xf;
+		Value >>= 4;
+		Buffer[i] = (ch >= 10 ? TCHAR('A' - 10) : TCHAR('0')) + ch;
+	}
+	Buffer[16] = 0;
+
+	return FString(Buffer);
+}
+
 #if WITH_EDITORONLY_DATA
 void FNiagaraUtilities::PrepareRapidIterationParameters(const TArray<UNiagaraScript*>& Scripts, const TMap<UNiagaraScript*, UNiagaraScript*>& ScriptDependencyMap, const TMap<UNiagaraScript*, FString>& ScriptToEmitterNameMap)
 {
diff --git a/Engine/Plugins/FX/Niagara/Source/Niagara/Private/NiagaraComponent.cpp b/Engine/Plugins/FX/Niagara/Source/Niagara/Private/NiagaraComponent.cpp
--- a/Engine/Plugins/FX/Niagara/Source/Niagara/Private/NiagaraComponent.cpp
+++ b/Engine/Plugins/FX/Niagara/Source/Niagara/Private/NiagaraComponent.cpp
@@ -1003,7 +1003,7 @@ void UNiagaraComponent::OnEndOfFrameUpdateDuringTick()
 	Super::OnEndOfFrameUpdateDuringTick();
 	if ( SystemInstance )
 	{
-		SystemInstance->WaitForAsyncTick();
+		//SystemInstance->WaitForAsyncTick();
 	}
 }
 
diff --git a/Engine/Plugins/FX/Niagara/Source/Niagara/Private/NiagaraDataInterfaceCollisionQuery.cpp b/Engine/Plugins/FX/Niagara/Source/Niagara/Private/NiagaraDataInterfaceCollisionQuery.cpp
--- a/Engine/Plugins/FX/Niagara/Source/Niagara/Private/NiagaraDataInterfaceCollisionQuery.cpp
+++ b/Engine/Plugins/FX/Niagara/Source/Niagara/Private/NiagaraDataInterfaceCollisionQuery.cpp
@@ -28,7 +28,7 @@ bool UNiagaraDataInterfaceCollisionQuery::InitPerInstanceData(void* PerInstanceD
 	PIData->SystemInstance = InSystemInstance;
 	if (InSystemInstance)
 	{
-		PIData->CollisionBatch.Init(InSystemInstance->GetIDName(), InSystemInstance->GetComponent()->GetWorld());
+		PIData->CollisionBatch.Init(InSystemInstance->GetId(), InSystemInstance->GetComponent()->GetWorld());
 	}
 	return true;
 }
diff --git a/Engine/Plugins/FX/Niagara/Source/Niagara/Private/NiagaraDataInterfaceGrid2DCollection.cpp b/Engine/Plugins/FX/Niagara/Source/Niagara/Private/NiagaraDataInterfaceGrid2DCollection.cpp
--- a/Engine/Plugins/FX/Niagara/Source/Niagara/Private/NiagaraDataInterfaceGrid2DCollection.cpp
+++ b/Engine/Plugins/FX/Niagara/Source/Niagara/Private/NiagaraDataInterfaceGrid2DCollection.cpp
@@ -402,7 +402,7 @@ void UNiagaraDataInterfaceGrid2DCollection::FillTexture2D(const UNiagaraComponen
 	{
 		return;
 	}
-	FGuid InstanceID = SystemInstance->GetId();
+	FNiagaraSystemInstanceID InstanceID = SystemInstance->GetId();
 
 	ENQUEUE_RENDER_COMMAND(FUpdateDIColorCurve)(
 		[Dest, TProxy, InstanceID](FRHICommandListImmediate& RHICmdList)
@@ -431,7 +431,7 @@ void UNiagaraDataInterfaceGrid2DCollection::GetCellSize(FVectorVMContext& Contex
 }
 
 // #todo(dmp): move these to super class
-void FNiagaraDataInterfaceProxyGrid2DCollection::DestroyPerInstanceData(NiagaraEmitterInstanceBatcher* Batcher, const FGuid& SystemInstance)
+void FNiagaraDataInterfaceProxyGrid2DCollection::DestroyPerInstanceData(NiagaraEmitterInstanceBatcher* Batcher, const FNiagaraSystemInstanceID& SystemInstance)
 {
 	check(IsInRenderingThread());
 
@@ -442,7 +442,7 @@ void FNiagaraDataInterfaceProxyGrid2DCollection::DestroyPerInstanceData(NiagaraE
 // #todo(dmp): move these to super class
 void FNiagaraDataInterfaceProxyGrid2DCollection::DeferredDestroy()
 {
-	for (const FGuid& Sys : DeferredDestroyList)
+	for (const FNiagaraSystemInstanceID& Sys : DeferredDestroyList)
 	{
 		SystemInstancesToProxyData.Remove(Sys);
 	}
diff --git a/Engine/Plugins/FX/Niagara/Source/Niagara/Private/NiagaraDataInterfaceNeighborGrid3D.cpp b/Engine/Plugins/FX/Niagara/Source/Niagara/Private/NiagaraDataInterfaceNeighborGrid3D.cpp
--- a/Engine/Plugins/FX/Niagara/Source/Niagara/Private/NiagaraDataInterfaceNeighborGrid3D.cpp
+++ b/Engine/Plugins/FX/Niagara/Source/Niagara/Private/NiagaraDataInterfaceNeighborGrid3D.cpp
@@ -478,7 +478,7 @@ void FNiagaraDataInterfaceProxyNeighborGrid3D::PreStage(FRHICommandList& RHICmdL
 }
 
 // #todo(dmp): move these to super class
-void FNiagaraDataInterfaceProxyNeighborGrid3D::DestroyPerInstanceData(NiagaraEmitterInstanceBatcher* Batcher, const FGuid& SystemInstance)
+void FNiagaraDataInterfaceProxyNeighborGrid3D::DestroyPerInstanceData(NiagaraEmitterInstanceBatcher* Batcher, const FNiagaraSystemInstanceID& SystemInstance)
 {
 	check(IsInRenderingThread());
 
@@ -489,7 +489,7 @@ void FNiagaraDataInterfaceProxyNeighborGrid3D::DestroyPerInstanceData(NiagaraEmi
 // #todo(dmp): move these to super class
 void FNiagaraDataInterfaceProxyNeighborGrid3D::DeferredDestroy()
 {
-	for (const FGuid& Sys : DeferredDestroyList)
+	for (const FNiagaraSystemInstanceID& Sys : DeferredDestroyList)
 	{
 		SystemInstancesToProxyData.Remove(Sys);
 	}
diff --git a/Engine/Plugins/FX/Niagara/Source/Niagara/Private/NiagaraDataInterfaceSkeletalMesh.cpp b/Engine/Plugins/FX/Niagara/Source/Niagara/Private/NiagaraDataInterfaceSkeletalMesh.cpp
--- a/Engine/Plugins/FX/Niagara/Source/Niagara/Private/NiagaraDataInterfaceSkeletalMesh.cpp
+++ b/Engine/Plugins/FX/Niagara/Source/Niagara/Private/NiagaraDataInterfaceSkeletalMesh.cpp
@@ -913,7 +913,7 @@ private:
 
 //////////////////////////////////////////////////////////////////////////
 
-void FNiagaraDataInterfaceProxySkeletalMesh::ConsumePerInstanceDataFromGameThread(void* PerInstanceData, const FGuid& Instance)
+void FNiagaraDataInterfaceProxySkeletalMesh::ConsumePerInstanceDataFromGameThread(void* PerInstanceData, const FNiagaraSystemInstanceID& Instance)
 {
 	FNiagaraDISkeletalMeshPassedDataToRT* SourceData = static_cast<FNiagaraDISkeletalMeshPassedDataToRT*>(PerInstanceData);
 
@@ -934,7 +934,7 @@ void FNiagaraDataInterfaceProxySkeletalMesh::ConsumePerInstanceDataFromGameThrea
 //////////////////////////////////////////////////////////////////////////
 //FNDISkeletalMesh_InstanceData
 
-void UNiagaraDataInterfaceSkeletalMesh::ProvidePerInstanceDataForRenderThread(void* DataForRenderThread, void* PerInstanceData, const FGuid& SystemInstance)
+void UNiagaraDataInterfaceSkeletalMesh::ProvidePerInstanceDataForRenderThread(void* DataForRenderThread, void* PerInstanceData, const FNiagaraSystemInstanceID& SystemInstance)
 {
 	FNiagaraDISkeletalMeshPassedDataToRT* Data = static_cast<FNiagaraDISkeletalMeshPassedDataToRT*>(DataForRenderThread);
 	FNDISkeletalMesh_InstanceData* SourceData = static_cast<FNDISkeletalMesh_InstanceData*>(PerInstanceData);
diff --git a/Engine/Plugins/FX/Niagara/Source/Niagara/Private/NiagaraDataInterfaceStaticMesh.cpp b/Engine/Plugins/FX/Niagara/Source/Niagara/Private/NiagaraDataInterfaceStaticMesh.cpp
--- a/Engine/Plugins/FX/Niagara/Source/Niagara/Private/NiagaraDataInterfaceStaticMesh.cpp
+++ b/Engine/Plugins/FX/Niagara/Source/Niagara/Private/NiagaraDataInterfaceStaticMesh.cpp
@@ -494,7 +494,7 @@ struct FNiagaraDataInterfaceParametersCS_StaticMesh : public FNiagaraDataInterfa
 		{
 			FNiagaraDataInterfaceProxyStaticMesh* InterfaceProxy = static_cast<FNiagaraDataInterfaceProxyStaticMesh*>(Context.DataInterface);
 			FNiagaraStaticMeshData* Data = InterfaceProxy->SystemInstancesToMeshData.Find(Context.SystemInstance);
-			ensureMsgf(Data, TEXT("Failed to find data for instance %s"), *Context.SystemInstance.ToString());
+			ensureMsgf(Data, TEXT("Failed to find data for instance %s"), *FNiagaraUtilities::SystemInstanceIDToString(Context.SystemInstance));
 			if (Data != nullptr)
 			{
 				FStaticMeshGpuSpawnBuffer* SpawnBuffer = Data->MeshGpuSpawnBuffer;
@@ -577,7 +577,7 @@ private:
 void FNiagaraDataInterfaceProxyStaticMesh::DeferredDestroy()
 {
 	//-TODO: This is incorrect, we could be destroying instance data for a batcher is yet to tick, we should only be destroying data for this batcher!
-	for (const FGuid& Sys : DeferredDestroyList)
+	for (const FNiagaraSystemInstanceID& Sys : DeferredDestroyList)
 	{
 		SystemInstancesToMeshData.Remove(Sys);
 		//UE_LOG(LogNiagara, Log, TEXT("RT: StaticMesh DI - DeferredDestroy %s"), *Sys.ToString());
@@ -586,7 +586,7 @@ void FNiagaraDataInterfaceProxyStaticMesh::DeferredDestroy()
 	DeferredDestroyList.Empty();
 }
 
-void FNiagaraDataInterfaceProxyStaticMesh::InitializePerInstanceData(const FGuid& SystemInstance, FStaticMeshGpuSpawnBuffer* MeshGPUSpawnBuffer)
+void FNiagaraDataInterfaceProxyStaticMesh::InitializePerInstanceData(const FNiagaraSystemInstanceID& SystemInstance, FStaticMeshGpuSpawnBuffer* MeshGPUSpawnBuffer)
 {
 	check(IsInRenderingThread());
 
@@ -611,7 +611,7 @@ void FNiagaraDataInterfaceProxyStaticMesh::InitializePerInstanceData(const FGuid
 	Data->MeshGpuSpawnBuffer = MeshGPUSpawnBuffer;
 }
 
-void FNiagaraDataInterfaceProxyStaticMesh::DestroyPerInstanceData(NiagaraEmitterInstanceBatcher* Batcher, const FGuid& SystemInstance)
+void FNiagaraDataInterfaceProxyStaticMesh::DestroyPerInstanceData(NiagaraEmitterInstanceBatcher* Batcher, const FNiagaraSystemInstanceID& SystemInstance)
 {
 	check(IsInRenderingThread());
 
@@ -623,7 +623,7 @@ void FNiagaraDataInterfaceProxyStaticMesh::DestroyPerInstanceData(NiagaraEmitter
 	Batcher->EnqueueDeferredDeletesForDI_RenderThread(this->AsShared());
 }
 
-void FNiagaraDataInterfaceProxyStaticMesh::ConsumePerInstanceDataFromGameThread(void* PerInstanceData, const FGuid& Instance)
+void FNiagaraDataInterfaceProxyStaticMesh::ConsumePerInstanceDataFromGameThread(void* PerInstanceData, const FNiagaraSystemInstanceID& Instance)
 {
 	FNiagaraPassedInstanceDataForRT* SourceData = static_cast<FNiagaraPassedInstanceDataForRT*>(PerInstanceData);
 	FNiagaraStaticMeshData* Data = SystemInstancesToMeshData.Find(Instance);
@@ -641,7 +641,7 @@ void FNiagaraDataInterfaceProxyStaticMesh::ConsumePerInstanceDataFromGameThread(
 	}
 	else
 	{
-		UE_LOG(LogNiagara, Log, TEXT("ConsumePerInstanceDataFromGameThread() ... could not find %s"), *Instance.ToString());
+		UE_LOG(LogNiagara, Log, TEXT("ConsumePerInstanceDataFromGameThread() ... could not find %s"), *FNiagaraUtilities::SystemInstanceIDToString(Instance));
 	}
 }
 
@@ -2364,7 +2364,7 @@ void UNiagaraDataInterfaceStaticMesh::GetParameterDefinitionHLSL(FNiagaraDataInt
 	OutHLSL += TEXT("uint ") + ParamNames.NumTexCoordName + TEXT(";\n");
 }
 
-void UNiagaraDataInterfaceStaticMesh::ProvidePerInstanceDataForRenderThread(void* DataForRenderThread, void* PerInstanceData, const FGuid& SystemInstance)
+void UNiagaraDataInterfaceStaticMesh::ProvidePerInstanceDataForRenderThread(void* DataForRenderThread, void* PerInstanceData, const FNiagaraSystemInstanceID& SystemInstance)
 {
 	check(Proxy);
 
diff --git a/Engine/Plugins/FX/Niagara/Source/Niagara/Private/NiagaraEmitterInstance.cpp b/Engine/Plugins/FX/Niagara/Source/Niagara/Private/NiagaraEmitterInstance.cpp
--- a/Engine/Plugins/FX/Niagara/Source/Niagara/Private/NiagaraEmitterInstance.cpp
+++ b/Engine/Plugins/FX/Niagara/Source/Niagara/Private/NiagaraEmitterInstance.cpp
@@ -182,13 +182,13 @@ bool FNiagaraEmitterInstance::IsAllowedToExecute() const
 	return true;
 }
 
-void FNiagaraEmitterInstance::Init(int32 InEmitterIdx, FName InSystemInstanceName)
+void FNiagaraEmitterInstance::Init(int32 InEmitterIdx, FNiagaraSystemInstanceID InSystemInstanceID)
 {
 	SCOPE_CYCLE_COUNTER(STAT_NiagaraEmitterInit);
 	check(ParticleDataSet);
 	FNiagaraDataSet& Data = *ParticleDataSet;
 	EmitterIdx = InEmitterIdx;
-	OwnerSystemInstanceName = InSystemInstanceName;
+	OwnerSystemInstanceID = InSystemInstanceID;
 	const FNiagaraEmitterHandle& EmitterHandle = GetEmitterHandle();
 	CachedEmitter = EmitterHandle.GetInstance();
 	checkSlow(CachedEmitter);
@@ -277,7 +277,7 @@ void FNiagaraEmitterInstance::Init(int32 InEmitterIdx, FName InSystemInstanceNam
 	int32 UpdateEventGeneratorIndex = 0;
 	for (const FNiagaraEventGeneratorProperties &GeneratorProps : CachedEmitter->UpdateScriptProps.EventGenerators)
 	{
-		FNiagaraDataSet *Set = FNiagaraEventDataSetMgr::CreateEventDataSet(ParentSystemInstance->GetIDName(), EmitterHandle.GetIdName(), GeneratorProps.SetProps.ID.Name);
+		FNiagaraDataSet *Set = FNiagaraEventDataSetMgr::CreateEventDataSet(ParentSystemInstance->GetId(), EmitterHandle.GetIdName(), GeneratorProps.SetProps.ID.Name);
 #if !UE_BUILD_SHIPPING && !UE_BUILD_TEST 
 		Set->Init(FNiagaraDataSetID(), ENiagaraSimTarget::CPUSim, CachedEmitter->GetFullName() + TEXT("/") + GeneratorProps.SetProps.ID.Name.ToString());
 #else
@@ -296,7 +296,7 @@ void FNiagaraEmitterInstance::Init(int32 InEmitterIdx, FName InSystemInstanceNam
 	int32 SpawnEventGeneratorIndex = 0;
 	for (const FNiagaraEventGeneratorProperties &GeneratorProps : CachedEmitter->SpawnScriptProps.EventGenerators)
 	{
-		FNiagaraDataSet *Set = FNiagaraEventDataSetMgr::CreateEventDataSet(ParentSystemInstance->GetIDName(), EmitterHandle.GetIdName(), GeneratorProps.SetProps.ID.Name);
+		FNiagaraDataSet *Set = FNiagaraEventDataSetMgr::CreateEventDataSet(ParentSystemInstance->GetId(), EmitterHandle.GetIdName(), GeneratorProps.SetProps.ID.Name);
 #if !UE_BUILD_SHIPPING && !UE_BUILD_TEST 
 		Set->Init(FNiagaraDataSetID(), ENiagaraSimTarget::CPUSim, CachedEmitter->GetFullName() + TEXT("/") + GeneratorProps.SetProps.ID.Name.ToString());
 #else
@@ -634,14 +634,14 @@ void FNiagaraEmitterInstance::PostInitSimulation()
 		for (const FNiagaraEventReceiverProperties& Receiver : CachedEmitter->SpawnScriptProps.EventReceivers)
 		{
 			//FNiagaraDataSet* ReceiverSet = ParentSystemInstance->GetDataSet(FNiagaraDataSetID(Receiver.SourceEventGenerator, ENiagaraDataSetType::Event), Receiver.SourceEmitter);
-			const FNiagaraDataSet* ReceiverSet = FNiagaraEventDataSetMgr::GetEventDataSet(ParentSystemInstance->GetIDName(), Receiver.SourceEmitter, Receiver.SourceEventGenerator);
+			const FNiagaraDataSet* ReceiverSet = FNiagaraEventDataSetMgr::GetEventDataSet(ParentSystemInstance->GetId(), Receiver.SourceEmitter, Receiver.SourceEventGenerator);
 
 		}
 
 		for (const FNiagaraEventReceiverProperties& Receiver : CachedEmitter->UpdateScriptProps.EventReceivers)
 		{
 			//FNiagaraDataSet* ReceiverSet = ParentSystemInstance->GetDataSet(FNiagaraDataSetID(Receiver.SourceEventGenerator, ENiagaraDataSetType::Event), Receiver.SourceEmitter);
-			const FNiagaraDataSet* ReceiverSet = FNiagaraEventDataSetMgr::GetEventDataSet(ParentSystemInstance->GetIDName(), Receiver.SourceEmitter, Receiver.SourceEventGenerator);
+			const FNiagaraDataSet* ReceiverSet = FNiagaraEventDataSetMgr::GetEventDataSet(ParentSystemInstance->GetId(), Receiver.SourceEmitter, Receiver.SourceEventGenerator);
 		}
 	}
 }
@@ -937,7 +937,7 @@ void FNiagaraEmitterInstance::PreTick()
 		Info.SpawnCounts.Reset();
 		Info.TotalSpawnCount = 0;
 		Info.EventData = nullptr;
-		if (FNiagaraDataSet* EventSet = FNiagaraEventDataSetMgr::GetEventDataSet(ParentSystemInstance->GetIDName(), Info.SourceEmitterName, EventHandlerProps.SourceEventName))
+		if (FNiagaraDataSet* EventSet = FNiagaraEventDataSetMgr::GetEventDataSet(ParentSystemInstance->GetId(), Info.SourceEmitterName, EventHandlerProps.SourceEventName))
 		{
 			Info.SetEventData(&EventSet->GetCurrentDataChecked());
 			uint32 EventSpawnNum = CalculateEventSpawnCount(EventHandlerProps, Info.SpawnCounts, EventSet);
diff --git a/Engine/Plugins/FX/Niagara/Source/Niagara/Private/NiagaraEmitterInstanceBatcher.cpp b/Engine/Plugins/FX/Niagara/Source/Niagara/Private/NiagaraEmitterInstanceBatcher.cpp
--- a/Engine/Plugins/FX/Niagara/Source/Niagara/Private/NiagaraEmitterInstanceBatcher.cpp
+++ b/Engine/Plugins/FX/Niagara/Source/Niagara/Private/NiagaraEmitterInstanceBatcher.cpp
@@ -999,7 +999,7 @@ void NiagaraEmitterInstanceBatcher::SetDataInterfaceParameters(const TArray<FNia
 	//
 
 	// @todo-threadsafety This is a bit gross. Need to rethink this api.
-	const FGuid& SystemInstance = Tick.SystemInstanceID;
+	const FNiagaraSystemInstanceID& SystemInstance = Tick.SystemInstanceID;
 
 	uint32 InterfaceIndex = 0;
 	for (FNiagaraDataInterfaceProxy* Interface : DataInterfaceProxies)
@@ -1027,7 +1027,7 @@ void NiagaraEmitterInstanceBatcher::UnsetDataInterfaceParameters(const TArray<FN
 	//
 
 	// @todo-threadsafety This is a bit gross. Need to rethink this api.
-	const FGuid& SystemInstance = Tick.SystemInstanceID;
+	const FNiagaraSystemInstanceID& SystemInstance = Tick.SystemInstanceID;
 
 	uint32 InterfaceIndex = 0;
 	for (FNiagaraDataInterfaceProxy* Interface : DataInterfaceProxies)
diff --git a/Engine/Plugins/FX/Niagara/Source/Niagara/Private/NiagaraEvents.cpp b/Engine/Plugins/FX/Niagara/Source/Niagara/Private/NiagaraEvents.cpp
--- a/Engine/Plugins/FX/Niagara/Source/Niagara/Private/NiagaraEvents.cpp
+++ b/Engine/Plugins/FX/Niagara/Source/Niagara/Private/NiagaraEvents.cpp
@@ -9,7 +9,7 @@ DECLARE_DWORD_COUNTER_STAT(TEXT("Num Death Events"), STAT_NiagaraNumDeathEvents,
 DECLARE_DWORD_COUNTER_STAT(TEXT("Num Spawn Events"), STAT_NiagaraNumSpawnEvents, STATGROUP_Niagara);
 
 
-TMap<FName, TMap<FName, PerEmitterEventDataSetMap>> FNiagaraEventDataSetMgr::EmitterEventDataSets;
+TMap<FNiagaraSystemInstanceID, TMap<FName, PerEmitterEventDataSetMap>> FNiagaraEventDataSetMgr::EmitterEventDataSets;
 
 
 void UNiagaraEventReceiverEmitterAction_SpawnParticles::PerformAction(FNiagaraEmitterInstance& OwningSim, const FNiagaraEventReceiverProperties& OwningEventReceiver)
diff --git a/Engine/Plugins/FX/Niagara/Source/Niagara/Private/NiagaraSystemInstance.cpp b/Engine/Plugins/FX/Niagara/Source/Niagara/Private/NiagaraSystemInstance.cpp
--- a/Engine/Plugins/FX/Niagara/Source/Niagara/Private/NiagaraSystemInstance.cpp
+++ b/Engine/Plugins/FX/Niagara/Source/Niagara/Private/NiagaraSystemInstance.cpp
@@ -62,8 +62,6 @@ FNiagaraSystemInstance::FNiagaraSystemInstance(UNiagaraComponent* InComponent)
 	, PrereqComponent(nullptr)
 	, Age(0.0f)
 	, TickCount(0)
-	, ID(FGuid::NewGuid())
-	, IDName(*ID.ToString())
 	, InstanceParameters(Component)
 	, bSolo(false)
 	, bForceSolo(false)
@@ -78,6 +76,9 @@ FNiagaraSystemInstance::FNiagaraSystemInstance(UNiagaraComponent* InComponent)
 	, bNeedsFinalize(false)
 	, CachedDeltaSeconds(0.0f)
 {
+	static TAtomic<uint64> IDCounter(1);
+	ID = IDCounter.IncrementExchange();
+
 	SystemBounds.Init();
 
 	if (Component)
@@ -570,6 +571,7 @@ void FNiagaraSystemInstance::SetPaused(bool bInPaused)
 void FNiagaraSystemInstance::Reset(FNiagaraSystemInstance::EResetMode Mode)
 {
 	SCOPE_CYCLE_COUNTER(STAT_NiagaraSystemReset);
+	FScopeCycleCounterUObject AdditionalScope(GetSystem(), GET_STATID(STAT_NiagaraSystemReset));
 
 	if (Mode == EResetMode::None)
 	{
@@ -1481,7 +1483,7 @@ void FNiagaraSystemInstance::InitEmitters()
 		for (int32 EmitterIdx=0; EmitterIdx < GetSystem()->GetEmitterHandles().Num(); ++EmitterIdx)
 		{
 			TSharedRef<FNiagaraEmitterInstance, ESPMode::ThreadSafe> Sim = MakeShared<FNiagaraEmitterInstance, ESPMode::ThreadSafe>(this);
-			Sim->Init(EmitterIdx, IDName);
+			Sim->Init(EmitterIdx, ID);
 			if (System->bFixedBounds)
 			{
 				Sim->SetSystemFixedBoundsOverride(System->GetFixedBounds());
diff --git a/Engine/Plugins/FX/Niagara/Source/Niagara/Public/NiagaraCommon.h b/Engine/Plugins/FX/Niagara/Source/Niagara/Public/NiagaraCommon.h
--- a/Engine/Plugins/FX/Niagara/Source/Niagara/Public/NiagaraCommon.h
+++ b/Engine/Plugins/FX/Niagara/Source/Niagara/Public/NiagaraCommon.h
@@ -8,6 +8,7 @@
 #include "NiagaraTypes.h"
 #include "UObject/SoftObjectPath.h"
 #include "RHI.h"
+#include "NiagaraCore.h"
 #include "NiagaraCommon.generated.h"
 
 class UNiagaraSystem;
@@ -626,6 +627,7 @@ namespace FNiagaraUtilities
 
 	void NIAGARA_API DumpHLSLText(const FString& SourceCode, const FString& DebugName);
 
+	NIAGARA_API FString SystemInstanceIDToString(FNiagaraSystemInstanceID ID);
 };
 
 USTRUCT()
diff --git a/Engine/Plugins/FX/Niagara/Source/Niagara/Public/NiagaraSystemInstance.h b/Engine/Plugins/FX/Niagara/Source/Niagara/Public/NiagaraSystemInstance.h
--- a/Engine/Plugins/FX/Niagara/Source/Niagara/Public/NiagaraSystemInstance.h
+++ b/Engine/Plugins/FX/Niagara/Source/Niagara/Public/NiagaraSystemInstance.h
@@ -148,8 +148,7 @@ public:
 	bool GetIsolateEnabled() const;
 #endif
 
-	FName GetIDName() { return IDName; }
-	FGuid GetId() { return ID; }
+	FNiagaraSystemInstanceID GetId() { return ID; }
 
 	/** Returns the instance data for a particular interface for this System. */
 	FORCEINLINE void* FindDataInterfaceInstanceData(UNiagaraDataInterface* Interface) 
@@ -291,7 +290,7 @@ private:
 	TMap<FGuid, TSharedPtr<TArray<TSharedPtr<struct FNiagaraScriptDebuggerInfo, ESPMode::ThreadSafe>>, ESPMode::ThreadSafe> > CapturedFrames;
 #endif
 
-	FGuid ID;
+	FNiagaraSystemInstanceID ID;
 	FName IDName;
 	
 	/** Per instance data for any data interfaces requiring it. */
diff --git a/Engine/Plugins/FX/Niagara/Source/NiagaraCore/Public/NiagaraCore.h b/Engine/Plugins/FX/Niagara/Source/NiagaraCore/Public/NiagaraCore.h
new file mode 100644
--- /dev/null
+++ b/Engine/Plugins/FX/Niagara/Source/NiagaraCore/Public/NiagaraCore.h
@@ -0,0 +1,6 @@
+// Copyright 1998-2019 Epic Games, Inc. All Rights Reserved.
+
+#pragma once
+
+
+typedef uint64 FNiagaraSystemInstanceID;
diff --git a/Engine/Plugins/FX/Niagara/Source/NiagaraCore/Public/NiagaraDataInterfaceBase.h b/Engine/Plugins/FX/Niagara/Source/NiagaraCore/Public/NiagaraDataInterfaceBase.h
--- a/Engine/Plugins/FX/Niagara/Source/NiagaraCore/Public/NiagaraDataInterfaceBase.h
+++ b/Engine/Plugins/FX/Niagara/Source/NiagaraCore/Public/NiagaraDataInterfaceBase.h
@@ -1,6 +1,7 @@
 // Copyright 1998-2019 Epic Games, Inc. All Rights Reserved.
 #pragma once
 
+#include "NiagaraCore.h"
 #include "NiagaraMergeable.h"
 #include "NiagaraDataInterfaceBase.generated.h"
 
@@ -20,7 +21,7 @@ struct FNiagaraDataInterfaceSetArgs
 {
 	FNiagaraShader* Shader;
 	FNiagaraDataInterfaceProxy* DataInterface;
-	FGuid SystemInstance;
+	FNiagaraSystemInstanceID SystemInstance;
 	const NiagaraEmitterInstanceBatcher* Batcher;
 	uint32 ShaderStageIndex;
 	bool IsOutputStage;
diff --git a/Engine/Plugins/FX/Niagara/Source/NiagaraVertexFactories/NiagaraVertexFactories.Build.cs b/Engine/Plugins/FX/Niagara/Source/NiagaraVertexFactories/NiagaraVertexFactories.Build.cs
--- a/Engine/Plugins/FX/Niagara/Source/NiagaraVertexFactories/NiagaraVertexFactories.Build.cs
+++ b/Engine/Plugins/FX/Niagara/Source/NiagaraVertexFactories/NiagaraVertexFactories.Build.cs
@@ -12,14 +12,15 @@ public class NiagaraVertexFactories : ModuleRules
 				"Engine",		
 				"RenderCore",
                 "RHI",
-			}
+            }
         );
 
 		PublicIncludePathModuleNames.AddRange(
 			new string[] {
-				"Niagara",
-			}
-		);
+                "NiagaraCore",
+                "Niagara",
+            }
+        );
 
         PublicDependencyModuleNames.AddRange(
             new string[] {
