Fixing potential crash in GPU sims. We don't need to build the register table for them

diff --git a/Engine/Plugins/FX/Niagara/Source/Niagara/Private/NiagaraDataSet.cpp b/Engine/Plugins/FX/Niagara/Source/Niagara/Private/NiagaraDataSet.cpp
--- a/Engine/Plugins/FX/Niagara/Source/Niagara/Private/NiagaraDataSet.cpp
+++ b/Engine/Plugins/FX/Niagara/Source/Niagara/Private/NiagaraDataSet.cpp
@@ -397,8 +397,10 @@ void FNiagaraDataSet::Finalize()
 	BuildLayout();
 
 	ResetBuffers();
-
-	GetCurrentDataChecked().BuildRegisterTable();
+	if (SimTarget == ENiagaraSimTarget::CPUSim)
+	{
+		GetCurrentDataChecked().BuildRegisterTable();
+	}
 }
 
 const FNiagaraVariableLayoutInfo* FNiagaraDataSet::GetVariableLayout(const FNiagaraVariable& Var)const
