Crash fix when trying to bind registers for a niagara buffer that hansn't been allocated yet.

diff --git a/Engine/Plugins/FX/Niagara/Source/Niagara/Private/NiagaraScriptExecutionContext.cpp b/Engine/Plugins/FX/Niagara/Source/Niagara/Private/NiagaraScriptExecutionContext.cpp
--- a/Engine/Plugins/FX/Niagara/Source/Niagara/Private/NiagaraScriptExecutionContext.cpp
+++ b/Engine/Plugins/FX/Niagara/Source/Niagara/Private/NiagaraScriptExecutionContext.cpp
@@ -191,7 +191,7 @@ bool FNiagaraScriptExecutionContext::Execute(uint32 NumInstances)
 				DestinationData ? &DestinationData->GetIDTable() : nullptr, &Info.DataSet->GetFreeIDTable(), &Info.DataSet->GetNumFreeIDs(), &Info.DataSet->GetMaxUsedID(), Info.DataSet->GetIDAcquireTag());
 
 			int32 TotalComponents = Info.DataSet->GetNumFloatComponents() + Info.DataSet->GetNumInt32Components();
-			if (Info.Input)
+			if (Info.Input && Info.Input->GetNumInstances() > 0)
 			{
 				Info.Input->AppendToRegisterTable(InputRegisters, Info.StartInstance);
 			}
