Niagara: Fix crashes trying to load cooked content.

diff --git a/Engine/Plugins/FX/Niagara/Source/Niagara/Private/NiagaraEmitter.cpp b/Engine/Plugins/FX/Niagara/Source/Niagara/Private/NiagaraEmitter.cpp
--- a/Engine/Plugins/FX/Niagara/Source/Niagara/Private/NiagaraEmitter.cpp
+++ b/Engine/Plugins/FX/Niagara/Source/Niagara/Private/NiagaraEmitter.cpp
@@ -317,8 +317,11 @@ void UNiagaraEmitter::PostLoad()
 		ParentAtLastMerge = nullptr;
 	}
 
-	GraphSource->ConditionalPostLoad();
-	GraphSource->PostLoadFromEmitter(*this);
+	if (!GetOutermost()->bIsCookedForEditor)
+	{
+		GraphSource->ConditionalPostLoad();
+		GraphSource->PostLoadFromEmitter(*this);
+	}
 #endif
 
 	TArray<UNiagaraScript*> AllScripts;
@@ -331,99 +334,101 @@ void UNiagaraEmitter::PostLoad()
 	}
 
 #if WITH_EDITORONLY_DATA
-
-	// Handle emitter inheritance.
-	if (Parent != nullptr)
-	{
-		Parent->ConditionalPostLoad();
-	}
-	if (ParentAtLastMerge != nullptr)
+	if (!GetOutermost()->bIsCookedForEditor)
 	{
-		ParentAtLastMerge->ConditionalPostLoad();
-	}
-	if (IsSynchronizedWithParent() == false)
-	{
-		MergeChangesFromParent();
-	}
+		// Handle emitter inheritance.
+		if (Parent != nullptr)
+		{
+			Parent->ConditionalPostLoad();
+		}
+		if (ParentAtLastMerge != nullptr)
+		{
+			ParentAtLastMerge->ConditionalPostLoad();
+		}
+		if (IsSynchronizedWithParent() == false)
+		{
+			MergeChangesFromParent();
+		}
 
-	// Reset scripts if recompile is forced.
-	bool bGenerateNewChangeId = false;
-	FString GenerateNewChangeIdReason;
-	if (GetForceCompileOnLoad())
-	{
-		// If we are a standalone emitter, then we invalidate id's, which should cause systems dependent on us to regenerate.
-		UObject* OuterObj = GetOuter();
-		if (OuterObj == GetOutermost())
+		// Reset scripts if recompile is forced.
+		bool bGenerateNewChangeId = false;
+		FString GenerateNewChangeIdReason;
+		if (GetForceCompileOnLoad())
 		{
-			GraphSource->InvalidateCachedCompileIds();
-			bGenerateNewChangeId = true;
-			GenerateNewChangeIdReason = TEXT("PostLoad - Force compile on load");
-			if (GEnableVerboseNiagaraChangeIdLogging)
+			// If we are a standalone emitter, then we invalidate id's, which should cause systems dependent on us to regenerate.
+			UObject* OuterObj = GetOuter();
+			if (OuterObj == GetOutermost())
 			{
-				UE_LOG(LogNiagara, Log, TEXT("InvalidateCachedCompileIds for %s because GbForceNiagaraCompileOnLoad = %d"), *GetPathName(), GbForceNiagaraCompileOnLoad);
+				GraphSource->InvalidateCachedCompileIds();
+				bGenerateNewChangeId = true;
+				GenerateNewChangeIdReason = TEXT("PostLoad - Force compile on load");
+				if (GEnableVerboseNiagaraChangeIdLogging)
+				{
+					UE_LOG(LogNiagara, Log, TEXT("InvalidateCachedCompileIds for %s because GbForceNiagaraCompileOnLoad = %d"), *GetPathName(), GbForceNiagaraCompileOnLoad);
+				}
 			}
 		}
-	}
 	
-	if (ChangeId.IsValid() == false)
-	{
-		// If the change id is already invalid we need to generate a new one, and can skip checking the owned scripts.
-		bGenerateNewChangeId = true;
-		GenerateNewChangeIdReason = TEXT("PostLoad - Change id was invalid.");
-		if (GEnableVerboseNiagaraChangeIdLogging)
+		if (ChangeId.IsValid() == false)
 		{
-			UE_LOG(LogNiagara, Log, TEXT("Change ID updated for emitter %s because the ID was invalid."), *GetPathName());
+			// If the change id is already invalid we need to generate a new one, and can skip checking the owned scripts.
+			bGenerateNewChangeId = true;
+			GenerateNewChangeIdReason = TEXT("PostLoad - Change id was invalid.");
+			if (GEnableVerboseNiagaraChangeIdLogging)
+			{
+				UE_LOG(LogNiagara, Log, TEXT("Change ID updated for emitter %s because the ID was invalid."), *GetPathName());
+			}
 		}
-	}
-	else
-	{
-		for (UNiagaraScript* Script : AllScripts)
+		else
 		{
-			if (Script->AreScriptAndSourceSynchronized() == false)
+			for (UNiagaraScript* Script : AllScripts)
 			{
-				bGenerateNewChangeId = true;
-				GenerateNewChangeIdReason = TEXT("PostLoad - Script out of sync");
-				if (GEnableVerboseNiagaraChangeIdLogging)
+				if (Script->AreScriptAndSourceSynchronized() == false)
 				{
-					UE_LOG(LogNiagara, Log, TEXT("Change ID updated for emitter %s because of a change to its script %s"), *GetPathName(), *Script->GetPathName());
+					bGenerateNewChangeId = true;
+					GenerateNewChangeIdReason = TEXT("PostLoad - Script out of sync");
+					if (GEnableVerboseNiagaraChangeIdLogging)
+					{
+						UE_LOG(LogNiagara, Log, TEXT("Change ID updated for emitter %s because of a change to its script %s"), *GetPathName(), *Script->GetPathName());
+					}
 				}
 			}
 		}
-	}
 
-	if (bGenerateNewChangeId)
-	{
-		UpdateChangeId(GenerateNewChangeIdReason);
-	}
+		if (bGenerateNewChangeId)
+		{
+			UpdateChangeId(GenerateNewChangeIdReason);
+		}
 
-	GraphSource->OnChanged().AddUObject(this, &UNiagaraEmitter::GraphSourceChanged);
+		GraphSource->OnChanged().AddUObject(this, &UNiagaraEmitter::GraphSourceChanged);
 
-	EmitterSpawnScriptProps.Script->RapidIterationParameters.AddOnChangedHandler(
-		FNiagaraParameterStore::FOnChanged::FDelegate::CreateUObject(this, &UNiagaraEmitter::ScriptRapidIterationParameterChanged));
-	EmitterUpdateScriptProps.Script->RapidIterationParameters.AddOnChangedHandler(
-		FNiagaraParameterStore::FOnChanged::FDelegate::CreateUObject(this, &UNiagaraEmitter::ScriptRapidIterationParameterChanged));
-
-	if (SpawnScriptProps.Script)
-	{
-		SpawnScriptProps.Script->RapidIterationParameters.AddOnChangedHandler(
+		EmitterSpawnScriptProps.Script->RapidIterationParameters.AddOnChangedHandler(
 			FNiagaraParameterStore::FOnChanged::FDelegate::CreateUObject(this, &UNiagaraEmitter::ScriptRapidIterationParameterChanged));
-	}
-	
-	if (UpdateScriptProps.Script)
-	{
-		UpdateScriptProps.Script->RapidIterationParameters.AddOnChangedHandler(
+		EmitterUpdateScriptProps.Script->RapidIterationParameters.AddOnChangedHandler(
 			FNiagaraParameterStore::FOnChanged::FDelegate::CreateUObject(this, &UNiagaraEmitter::ScriptRapidIterationParameterChanged));
-	}
 
-	for (FNiagaraEventScriptProperties& EventScriptProperties : EventHandlerScriptProps)
-	{
-		EventScriptProperties.Script->RapidIterationParameters.AddOnChangedHandler(
-			FNiagaraParameterStore::FOnChanged::FDelegate::CreateUObject(this, &UNiagaraEmitter::ScriptRapidIterationParameterChanged));
-	}
+		if (SpawnScriptProps.Script)
+		{
+			SpawnScriptProps.Script->RapidIterationParameters.AddOnChangedHandler(
+				FNiagaraParameterStore::FOnChanged::FDelegate::CreateUObject(this, &UNiagaraEmitter::ScriptRapidIterationParameterChanged));
+		}
+	
+		if (UpdateScriptProps.Script)
+		{
+			UpdateScriptProps.Script->RapidIterationParameters.AddOnChangedHandler(
+				FNiagaraParameterStore::FOnChanged::FDelegate::CreateUObject(this, &UNiagaraEmitter::ScriptRapidIterationParameterChanged));
+		}
 
-	for (UNiagaraRendererProperties* Renderer : RendererProperties)
-	{
-		Renderer->OnChanged().AddUObject(this, &UNiagaraEmitter::RendererChanged);
+		for (FNiagaraEventScriptProperties& EventScriptProperties : EventHandlerScriptProps)
+		{
+			EventScriptProperties.Script->RapidIterationParameters.AddOnChangedHandler(
+				FNiagaraParameterStore::FOnChanged::FDelegate::CreateUObject(this, &UNiagaraEmitter::ScriptRapidIterationParameterChanged));
+		}
+
+		for (UNiagaraRendererProperties* Renderer : RendererProperties)
+		{
+			Renderer->OnChanged().AddUObject(this, &UNiagaraEmitter::RendererChanged);
+		}
 	}
 #endif
 }
diff --git a/Engine/Plugins/FX/Niagara/Source/Niagara/Private/NiagaraSpriteRendererProperties.cpp b/Engine/Plugins/FX/Niagara/Source/Niagara/Private/NiagaraSpriteRendererProperties.cpp
--- a/Engine/Plugins/FX/Niagara/Source/Niagara/Private/NiagaraSpriteRendererProperties.cpp
+++ b/Engine/Plugins/FX/Niagara/Source/Niagara/Private/NiagaraSpriteRendererProperties.cpp
@@ -93,7 +93,7 @@ void UNiagaraSpriteRendererProperties::Serialize(FStructuredArchive::FRecord Rec
 	Super::Serialize(Record);
 
 	FArchive& UnderlyingArchive = Record.GetUnderlyingArchive();
-	if (UnderlyingArchive.IsCooking() || (FPlatformProperties::RequiresCookedData() && UnderlyingArchive.IsLoading()))
+	if (UnderlyingArchive.IsCooking() || (FPlatformProperties::RequiresCookedData() && UnderlyingArchive.IsLoading()) || GetOutermost()->bIsCookedForEditor)
 	{
 		DerivedData.Serialize(Record.EnterField(FIELD_NAME_TEXT("DerivedData")));
 	}
diff --git a/Engine/Plugins/FX/Niagara/Source/Niagara/Private/NiagaraSystem.cpp b/Engine/Plugins/FX/Niagara/Source/Niagara/Private/NiagaraSystem.cpp
--- a/Engine/Plugins/FX/Niagara/Source/Niagara/Private/NiagaraSystem.cpp
+++ b/Engine/Plugins/FX/Niagara/Source/Niagara/Private/NiagaraSystem.cpp
@@ -277,165 +277,168 @@ void UNiagaraSystem::PostLoad()
 #endif
 
 #if WITH_EDITORONLY_DATA
-	TArray<UNiagaraScript*> AllSystemScripts;
-	
-	UNiagaraScriptSourceBase* SystemScriptSource = nullptr;
-	if (SystemSpawnScript == nullptr)
-	{
-		SystemSpawnScript = NewObject<UNiagaraScript>(this, "SystemSpawnScript", RF_Transactional);
-		SystemSpawnScript->SetUsage(ENiagaraScriptUsage::SystemSpawnScript);
-		INiagaraModule& NiagaraModule = FModuleManager::GetModuleChecked<INiagaraModule>("Niagara");
-		SystemScriptSource = NiagaraModule.GetEditorOnlyDataUtilities().CreateDefaultScriptSource(this);
-		SystemSpawnScript->SetSource(SystemScriptSource);
-	}
-	else
+	if (!GetOutermost()->bIsCookedForEditor)
 	{
-		SystemSpawnScript->ConditionalPostLoad();
-		SystemScriptSource = SystemSpawnScript->GetSource();
-	}
-	AllSystemScripts.Add(SystemSpawnScript);
+		TArray<UNiagaraScript*> AllSystemScripts;
+	
+		UNiagaraScriptSourceBase* SystemScriptSource = nullptr;
+		if (SystemSpawnScript == nullptr)
+		{
+			SystemSpawnScript = NewObject<UNiagaraScript>(this, "SystemSpawnScript", RF_Transactional);
+			SystemSpawnScript->SetUsage(ENiagaraScriptUsage::SystemSpawnScript);
+			INiagaraModule& NiagaraModule = FModuleManager::GetModuleChecked<INiagaraModule>("Niagara");
+			SystemScriptSource = NiagaraModule.GetEditorOnlyDataUtilities().CreateDefaultScriptSource(this);
+			SystemSpawnScript->SetSource(SystemScriptSource);
+		}
+		else
+		{
+			SystemSpawnScript->ConditionalPostLoad();
+			SystemScriptSource = SystemSpawnScript->GetSource();
+		}
+		AllSystemScripts.Add(SystemSpawnScript);
 
-	if (SystemUpdateScript == nullptr)
-	{
-		SystemUpdateScript = NewObject<UNiagaraScript>(this, "SystemUpdateScript", RF_Transactional);
-		SystemUpdateScript->SetUsage(ENiagaraScriptUsage::SystemUpdateScript);
-		SystemUpdateScript->SetSource(SystemScriptSource);
-	}
-	else
-	{
-		SystemUpdateScript->ConditionalPostLoad();
-	}
-	AllSystemScripts.Add(SystemUpdateScript);
+		if (SystemUpdateScript == nullptr)
+		{
+			SystemUpdateScript = NewObject<UNiagaraScript>(this, "SystemUpdateScript", RF_Transactional);
+			SystemUpdateScript->SetUsage(ENiagaraScriptUsage::SystemUpdateScript);
+			SystemUpdateScript->SetSource(SystemScriptSource);
+		}
+		else
+		{
+			SystemUpdateScript->ConditionalPostLoad();
+		}
+		AllSystemScripts.Add(SystemUpdateScript);
 
-	const int32 NiagaraVer = GetLinkerCustomVersion(FNiagaraCustomVersion::GUID);
+		const int32 NiagaraVer = GetLinkerCustomVersion(FNiagaraCustomVersion::GUID);
 
-	//TODO: This causes a crash becuase the script source ptr is null? Fix
-	//For existing emitters before the lifecylce rework, ensure they have the system lifecycle module.
-	if (NiagaraVer < FNiagaraCustomVersion::LifeCycleRework)
-	{
-		/*UNiagaraScriptSourceBase* SystemScriptSource = SystemUpdateScript->GetSource();
-		if (SystemScriptSource)
+		//TODO: This causes a crash becuase the script source ptr is null? Fix
+		//For existing emitters before the lifecylce rework, ensure they have the system lifecycle module.
+		if (NiagaraVer < FNiagaraCustomVersion::LifeCycleRework)
 		{
-			bool bFoundModule;
-			if (SystemScriptSource->AddModuleIfMissing(TEXT("/Niagara/Modules/System/SystemLifeCycle.SystemLifeCycle"), ENiagaraScriptUsage::SystemUpdateScript, bFoundModule))
+			/*UNiagaraScriptSourceBase* SystemScriptSource = SystemUpdateScript->GetSource();
+			if (SystemScriptSource)
 			{
-				bNeedsRecompile = true;
-			}
-		}*/
-	}
+				bool bFoundModule;
+				if (SystemScriptSource->AddModuleIfMissing(TEXT("/Niagara/Modules/System/SystemLifeCycle.SystemLifeCycle"), ENiagaraScriptUsage::SystemUpdateScript, bFoundModule))
+				{
+					bNeedsRecompile = true;
+				}
+			}*/
+		}
 
-	bool bSystemScriptsAreSynchronized = true;
-	for (UNiagaraScript* SystemScript : AllSystemScripts)
-	{
-		bSystemScriptsAreSynchronized &= SystemScript->AreScriptAndSourceSynchronized();
-	}
+		bool bSystemScriptsAreSynchronized = true;
+		for (UNiagaraScript* SystemScript : AllSystemScripts)
+		{
+			bSystemScriptsAreSynchronized &= SystemScript->AreScriptAndSourceSynchronized();
+		}
 
-	bool bEmitterScriptsAreSynchronized = true;
+		bool bEmitterScriptsAreSynchronized = true;
 
 #if 0
-	UE_LOG(LogNiagara, Log, TEXT("PreMerger"));
-	for (FNiagaraEmitterHandle& EmitterHandle : EmitterHandles)
-	{
-		UE_LOG(LogNiagara, Log, TEXT("Emitter Handle: %s"), *EmitterHandle.GetUniqueInstanceName());
-		UNiagaraScript* UpdateScript = EmitterHandle.GetInstance()->GetScript(ENiagaraScriptUsage::ParticleUpdateScript, FGuid());
-		UNiagaraScript* SpawnScript = EmitterHandle.GetInstance()->GetScript(ENiagaraScriptUsage::ParticleSpawnScript, FGuid());
-		UE_LOG(LogNiagara, Log, TEXT("Spawn Parameters"));
-		SpawnScript->GetVMExecutableData().Parameters.DumpParameters();
-		UE_LOG(LogNiagara, Log, TEXT("Spawn RI Parameters"));
-		SpawnScript->RapidIterationParameters.DumpParameters();
-		UE_LOG(LogNiagara, Log, TEXT("Update Parameters"));
-		UpdateScript->GetVMExecutableData().Parameters.DumpParameters();
-		UE_LOG(LogNiagara, Log, TEXT("Update RI Parameters"));
-		UpdateScript->RapidIterationParameters.DumpParameters();
-	}
+		UE_LOG(LogNiagara, Log, TEXT("PreMerger"));
+		for (FNiagaraEmitterHandle& EmitterHandle : EmitterHandles)
+		{
+			UE_LOG(LogNiagara, Log, TEXT("Emitter Handle: %s"), *EmitterHandle.GetUniqueInstanceName());
+			UNiagaraScript* UpdateScript = EmitterHandle.GetInstance()->GetScript(ENiagaraScriptUsage::ParticleUpdateScript, FGuid());
+			UNiagaraScript* SpawnScript = EmitterHandle.GetInstance()->GetScript(ENiagaraScriptUsage::ParticleSpawnScript, FGuid());
+			UE_LOG(LogNiagara, Log, TEXT("Spawn Parameters"));
+			SpawnScript->GetVMExecutableData().Parameters.DumpParameters();
+			UE_LOG(LogNiagara, Log, TEXT("Spawn RI Parameters"));
+			SpawnScript->RapidIterationParameters.DumpParameters();
+			UE_LOG(LogNiagara, Log, TEXT("Update Parameters"));
+			UpdateScript->GetVMExecutableData().Parameters.DumpParameters();
+			UE_LOG(LogNiagara, Log, TEXT("Update RI Parameters"));
+			UpdateScript->RapidIterationParameters.DumpParameters();
+		}
 #endif
 
-	for (FNiagaraEmitterHandle& EmitterHandle : EmitterHandles)
-	{
-		EmitterHandle.ConditionalPostLoad(NiagaraVer);
-		if (EmitterHandle.GetIsEnabled() && !EmitterHandle.GetInstance()->AreAllScriptAndSourcesSynchronized())
+		for (FNiagaraEmitterHandle& EmitterHandle : EmitterHandles)
 		{
-			bEmitterScriptsAreSynchronized = false;
-			break;
+			EmitterHandle.ConditionalPostLoad(NiagaraVer);
+			if (EmitterHandle.GetIsEnabled() && !EmitterHandle.GetInstance()->AreAllScriptAndSourcesSynchronized())
+			{
+				bEmitterScriptsAreSynchronized = false;
+				break;
+			}
 		}
-	}
 
-	if (EditorData == nullptr)
-	{
-		INiagaraModule& NiagaraModule = FModuleManager::GetModuleChecked<INiagaraModule>("Niagara");
-		EditorData = NiagaraModule.GetEditorOnlyDataUtilities().CreateDefaultEditorData(this);
-	}
-	else
-	{
-		EditorData->PostLoadFromOwner(this);
-	}
+		if (EditorData == nullptr)
+		{
+			INiagaraModule& NiagaraModule = FModuleManager::GetModuleChecked<INiagaraModule>("Niagara");
+			EditorData = NiagaraModule.GetEditorOnlyDataUtilities().CreateDefaultEditorData(this);
+		}
+		else
+		{
+			EditorData->PostLoadFromOwner(this);
+		}
 
-	if (UNiagaraEmitter::GetForceCompileOnLoad())
-	{
-		InvalidateCachedCompileIds();
-		UE_LOG(LogNiagara, Log, TEXT("System %s being rebuilt because UNiagaraEmitter::GetForceCompileOnLoad() == true."), *GetPathName());
-	}
+		if (UNiagaraEmitter::GetForceCompileOnLoad())
+		{
+			InvalidateCachedCompileIds();
+			UE_LOG(LogNiagara, Log, TEXT("System %s being rebuilt because UNiagaraEmitter::GetForceCompileOnLoad() == true."), *GetPathName());
+		}
 
-	if (bSystemScriptsAreSynchronized == false && GEnableVerboseNiagaraChangeIdLogging)
-	{
-		UE_LOG(LogNiagara, Log, TEXT("System %s being compiled because there were changes to a system script Change ID."), *GetPathName());
-	}
+		if (bSystemScriptsAreSynchronized == false && GEnableVerboseNiagaraChangeIdLogging)
+		{
+			UE_LOG(LogNiagara, Log, TEXT("System %s being compiled because there were changes to a system script Change ID."), *GetPathName());
+		}
 
-	if (bEmitterScriptsAreSynchronized == false && GEnableVerboseNiagaraChangeIdLogging)
-	{
-		UE_LOG(LogNiagara, Log, TEXT("System %s being compiled because there were changes to an emitter script Change ID."), *GetPathName());
-	}
+		if (bEmitterScriptsAreSynchronized == false && GEnableVerboseNiagaraChangeIdLogging)
+		{
+			UE_LOG(LogNiagara, Log, TEXT("System %s being compiled because there were changes to an emitter script Change ID."), *GetPathName());
+		}
 
-	if (EmitterCompiledData.Num() == 0)
-	{
-		InitEmitterCompiledData();
-	}
+		if (EmitterCompiledData.Num() == 0)
+		{
+			InitEmitterCompiledData();
+		}
 
-	if (SystemCompiledData.InstanceParamStore.GetNumParameters() == 0)
-	{
-		InitSystemCompiledData();
-	}
+		if (SystemCompiledData.InstanceParamStore.GetNumParameters() == 0)
+		{
+			InitSystemCompiledData();
+		}
 
 #if 0
-	UE_LOG(LogNiagara, Log, TEXT("Before"));
-	for (FNiagaraEmitterHandle& EmitterHandle : EmitterHandles)
-	{
-		UE_LOG(LogNiagara, Log, TEXT("Emitter Handle: %s"), *EmitterHandle.GetUniqueInstanceName());
-		UNiagaraScript* UpdateScript = EmitterHandle.GetInstance()->GetScript(ENiagaraScriptUsage::ParticleUpdateScript, FGuid());
-		UNiagaraScript* SpawnScript = EmitterHandle.GetInstance()->GetScript(ENiagaraScriptUsage::ParticleSpawnScript, FGuid());
-		UE_LOG(LogNiagara, Log, TEXT("Spawn Parameters"));
-		SpawnScript->GetVMExecutableData().Parameters.DumpParameters();
-		UE_LOG(LogNiagara, Log, TEXT("Spawn RI Parameters"));
-		SpawnScript->RapidIterationParameters.DumpParameters();
-		UE_LOG(LogNiagara, Log, TEXT("Update Parameters"));
-		UpdateScript->GetVMExecutableData().Parameters.DumpParameters();
-		UE_LOG(LogNiagara, Log, TEXT("Update RI Parameters"));
-		UpdateScript->RapidIterationParameters.DumpParameters();
-	}
+		UE_LOG(LogNiagara, Log, TEXT("Before"));
+		for (FNiagaraEmitterHandle& EmitterHandle : EmitterHandles)
+		{
+			UE_LOG(LogNiagara, Log, TEXT("Emitter Handle: %s"), *EmitterHandle.GetUniqueInstanceName());
+			UNiagaraScript* UpdateScript = EmitterHandle.GetInstance()->GetScript(ENiagaraScriptUsage::ParticleUpdateScript, FGuid());
+			UNiagaraScript* SpawnScript = EmitterHandle.GetInstance()->GetScript(ENiagaraScriptUsage::ParticleSpawnScript, FGuid());
+			UE_LOG(LogNiagara, Log, TEXT("Spawn Parameters"));
+			SpawnScript->GetVMExecutableData().Parameters.DumpParameters();
+			UE_LOG(LogNiagara, Log, TEXT("Spawn RI Parameters"));
+			SpawnScript->RapidIterationParameters.DumpParameters();
+			UE_LOG(LogNiagara, Log, TEXT("Update Parameters"));
+			UpdateScript->GetVMExecutableData().Parameters.DumpParameters();
+			UE_LOG(LogNiagara, Log, TEXT("Update RI Parameters"));
+			UpdateScript->RapidIterationParameters.DumpParameters();
+		}
 #endif
 
-	if (bSystemScriptsAreSynchronized == false || bEmitterScriptsAreSynchronized == false)
-	{
-		RequestCompile(false);
-	}
+		if (bSystemScriptsAreSynchronized == false || bEmitterScriptsAreSynchronized == false)
+		{
+			RequestCompile(false);
+		}
 
 #if 0
-	UE_LOG(LogNiagara, Log, TEXT("After"));
-	for (FNiagaraEmitterHandle& EmitterHandle : EmitterHandles)
-	{
-		UE_LOG(LogNiagara, Log, TEXT("Emitter Handle: %s"), *EmitterHandle.GetUniqueInstanceName());
-		UNiagaraScript* UpdateScript = EmitterHandle.GetInstance()->GetScript(ENiagaraScriptUsage::ParticleUpdateScript, FGuid());
-		UNiagaraScript* SpawnScript = EmitterHandle.GetInstance()->GetScript(ENiagaraScriptUsage::ParticleSpawnScript, FGuid());
-		UE_LOG(LogNiagara, Log, TEXT("Spawn Parameters"));
-		SpawnScript->GetVMExecutableData().Parameters.DumpParameters();
-		UE_LOG(LogNiagara, Log, TEXT("Spawn RI Parameters"));
-		SpawnScript->RapidIterationParameters.DumpParameters();
-		UE_LOG(LogNiagara, Log, TEXT("Update Parameters"));
-		UpdateScript->GetVMExecutableData().Parameters.DumpParameters();
-		UE_LOG(LogNiagara, Log, TEXT("Update RI Parameters"));
-		UpdateScript->RapidIterationParameters.DumpParameters();
-	}
+		UE_LOG(LogNiagara, Log, TEXT("After"));
+		for (FNiagaraEmitterHandle& EmitterHandle : EmitterHandles)
+		{
+			UE_LOG(LogNiagara, Log, TEXT("Emitter Handle: %s"), *EmitterHandle.GetUniqueInstanceName());
+			UNiagaraScript* UpdateScript = EmitterHandle.GetInstance()->GetScript(ENiagaraScriptUsage::ParticleUpdateScript, FGuid());
+			UNiagaraScript* SpawnScript = EmitterHandle.GetInstance()->GetScript(ENiagaraScriptUsage::ParticleSpawnScript, FGuid());
+			UE_LOG(LogNiagara, Log, TEXT("Spawn Parameters"));
+			SpawnScript->GetVMExecutableData().Parameters.DumpParameters();
+			UE_LOG(LogNiagara, Log, TEXT("Spawn RI Parameters"));
+			SpawnScript->RapidIterationParameters.DumpParameters();
+			UE_LOG(LogNiagara, Log, TEXT("Update Parameters"));
+			UpdateScript->GetVMExecutableData().Parameters.DumpParameters();
+			UE_LOG(LogNiagara, Log, TEXT("Update RI Parameters"));
+			UpdateScript->RapidIterationParameters.DumpParameters();
+		}
 #endif
+	}
 #endif // WITH_EDITORONLY_DATA
 }
 
@@ -1066,7 +1069,7 @@ void UNiagaraSystem::InitEmitterCompiledData()
 	if (SystemSpawnScript->GetVMExecutableData().IsValid() && SystemUpdateScript->GetVMExecutableData().IsValid())
 	{
 		EmitterCompiledData.SetNum(EmitterHandles.Num());
-		FNiagaraTypeDefinition SpawnInfoDef = FNiagaraTypeDefinition(FNiagaraSpawnInfo::StaticStruct());
+	FNiagaraTypeDefinition SpawnInfoDef = FNiagaraTypeDefinition(FNiagaraSpawnInfo::StaticStruct());
 
 		for (FNiagaraVariable& Var : SystemSpawnScript->GetVMExecutableData().Attributes)
 		{
@@ -1091,8 +1094,8 @@ void UNiagaraSystem::InitEmitterCompiledData()
 				{
 					EmitterCompiledData[EmitterIdx].SpawnAttributes.AddUnique(Var.GetName());
 				}
+				}
 			}
-		}
 
 		for (int32 EmitterIdx = 0; EmitterIdx < EmitterHandles.Num(); ++EmitterIdx)
 		{
diff --git a/Engine/Source/Runtime/Engine/Private/Particles/ParticleComponents.cpp b/Engine/Source/Runtime/Engine/Private/Particles/ParticleComponents.cpp
--- a/Engine/Source/Runtime/Engine/Private/Particles/ParticleComponents.cpp
+++ b/Engine/Source/Runtime/Engine/Private/Particles/ParticleComponents.cpp
@@ -1670,10 +1670,13 @@ void UParticleEmitter::Build()
 		if (HighLODLevel->TypeDataModule != nullptr)
 		{
 			if(HighLODLevel->TypeDataModule->RequiresBuild())
-		{
-			FParticleEmitterBuildInfo EmitterBuildInfo;
+			{
+				FParticleEmitterBuildInfo EmitterBuildInfo;
 #if WITH_EDITOR
-				HighLODLevel->CompileModules( EmitterBuildInfo );
+				if(!GetOutermost()->bIsCookedForEditor)
+				{
+					HighLODLevel->CompileModules( EmitterBuildInfo );
+				}
 #endif
 				HighLODLevel->TypeDataModule->Build( EmitterBuildInfo );
 			}
diff --git a/Engine/Source/Runtime/Engine/Private/Particles/ParticleModules.cpp b/Engine/Source/Runtime/Engine/Private/Particles/ParticleModules.cpp
--- a/Engine/Source/Runtime/Engine/Private/Particles/ParticleModules.cpp
+++ b/Engine/Source/Runtime/Engine/Private/Particles/ParticleModules.cpp
@@ -4562,6 +4562,11 @@ void UParticleModuleTypeDataGpu::BeginDestroy()
 void UParticleModuleTypeDataGpu::Build( FParticleEmitterBuildInfo& EmitterBuildInfo )
 {
 #if WITH_EDITOR
+	if (GetOutermost()->bIsCookedForEditor)
+	{
+		return;
+	}
+
 	FVector4Distribution Curve;
 	FComposableFloatDistribution ZeroDistribution;
 	FComposableFloatDistribution OneDistribution;
