#Niagara CIS quick fix

diff --git a/Engine/Plugins/FX/Niagara/Source/Niagara/Private/NiagaraSpriteRendererProperties.cpp b/Engine/Plugins/FX/Niagara/Source/Niagara/Private/NiagaraSpriteRendererProperties.cpp
--- a/Engine/Plugins/FX/Niagara/Source/Niagara/Private/NiagaraSpriteRendererProperties.cpp
+++ b/Engine/Plugins/FX/Niagara/Source/Niagara/Private/NiagaraSpriteRendererProperties.cpp
@@ -92,8 +92,13 @@ void UNiagaraSpriteRendererProperties::Serialize(FStructuredArchive::FRecord Rec
 {
 	Super::Serialize(Record);
 
+	bool bIsCookedForEditor = false;
+#if WITH_EDITORONLY_DATA
+	bIsCookedForEditor = GetOutermost()->bIsCookedForEditor;
+#endif // WITH_EDITORONLY_DATA
+
 	FArchive& UnderlyingArchive = Record.GetUnderlyingArchive();
-	if (UnderlyingArchive.IsCooking() || (FPlatformProperties::RequiresCookedData() && UnderlyingArchive.IsLoading()) || GetOutermost()->bIsCookedForEditor)
+	if (UnderlyingArchive.IsCooking() || (FPlatformProperties::RequiresCookedData() && UnderlyingArchive.IsLoading()) || bIsCookedForEditor)
 	{
 		DerivedData.Serialize(Record.EnterField(FIELD_NAME_TEXT("DerivedData")));
 	}
