Cache DataSet info in post compile of Niagara System instead of generating during runtime.

diff --git a/Engine/Plugins/FX/Niagara/Source/Niagara/Classes/NiagaraDataSet.h b/Engine/Plugins/FX/Niagara/Source/Niagara/Classes/NiagaraDataSet.h
--- a/Engine/Plugins/FX/Niagara/Source/Niagara/Classes/NiagaraDataSet.h
+++ b/Engine/Plugins/FX/Niagara/Source/Niagara/Classes/NiagaraDataSet.h
@@ -7,21 +7,29 @@
 #include "RHI.h"
 #include "VectorVM.h"
 #include "RenderingThread.h"
+#include "NiagaraDataSet.generated.h"
 
 typedef TArray<uint8*, TInlineAllocator<128>> FNiagaraRegisterTable;
 
 /** Helper class defining the layout and location of an FNiagaraVariable in an FNiagaraDataBuffer-> */
+USTRUCT()
 struct FNiagaraVariableLayoutInfo
 {
+	GENERATED_BODY()
+
 	/** Start index for the float components in the main buffer. */
+	UPROPERTY()
 	uint32 FloatComponentStart;
+
 	/** Start index for the int32 components in the main buffer. */
+	UPROPERTY()
 	uint32 Int32ComponentStart;
 
 	uint32 GetNumFloatComponents()const { return LayoutInfo.FloatComponentByteOffsets.Num(); }
 	uint32 GetNumInt32Components()const { return LayoutInfo.Int32ComponentByteOffsets.Num(); }
 
 	/** This variable's type layout info. */
+	UPROPERTY()
 	FNiagaraTypeLayoutInfo LayoutInfo;
 };
 
@@ -208,6 +216,44 @@ private:
 
 //////////////////////////////////////////////////////////////////////////
 
+USTRUCT()
+struct FNiagaraDataSetCompiledData
+{
+	GENERATED_BODY()
+
+	/** Variables in the data set. */
+	UPROPERTY()
+	TArray<FNiagaraVariable> Variables;
+
+	/** Data describing the layout of variable data. */
+	UPROPERTY()
+	TArray<FNiagaraVariableLayoutInfo> VariableLayouts;
+
+	/** Total number of components of each type in the data set. */
+	UPROPERTY()
+	uint32 TotalFloatComponents;
+
+	UPROPERTY()
+	uint32 TotalInt32Components;
+
+	/** Whether or not this dataset require persistent IDs. */
+	UPROPERTY()
+	uint32 bNeedsPersistentIDs : 1;
+
+	/** Unique ID for this DataSet. Used to allow referencing from other emitters and Systems. */
+	UPROPERTY()
+	FNiagaraDataSetID ID;
+
+	/** Sim target this DataSet is targeting (CPU/GPU). */
+	UPROPERTY()
+	ENiagaraSimTarget SimTarget;
+
+	void BuildLayout();
+	void Empty();
+};
+
+//////////////////////////////////////////////////////////////////////////
+
 /**
 General storage class for all per instance simulation data in Niagara.
 */
@@ -222,8 +268,10 @@ public:
 
 #if !UE_BUILD_SHIPPING && !UE_BUILD_TEST
 	void Init(FNiagaraDataSetID InID, ENiagaraSimTarget InSimTarget, const FString& InDebugName);
+	void InitFromCompiledData(const FNiagaraDataSetCompiledData& InDataSetCompiledData, const FString& InDebugName);
 #else
 	void Init(FNiagaraDataSetID InID, ENiagaraSimTarget InSimTarget);
+	void InitFromCompiledData(const FNiagaraDataSetCompiledData& InDataSetCompiledData);
 #endif
 
 	/** Resets current data but leaves variable/layout information etc intact. */
@@ -266,6 +314,7 @@ public:
 
 	/** Finalize the addition of variables and other setup before this data set can be used. */
 	void Finalize();
+
 	const FNiagaraVariableLayoutInfo* GetVariableLayout(const FNiagaraVariable& Var)const;
 	bool GetVariableComponentOffsets(const FNiagaraVariable& Var, int32 &FloatStart, int32 &IntStart) const;
 
diff --git a/Engine/Plugins/FX/Niagara/Source/Niagara/Classes/NiagaraSystem.h b/Engine/Plugins/FX/Niagara/Source/Niagara/Classes/NiagaraSystem.h
--- a/Engine/Plugins/FX/Niagara/Source/Niagara/Classes/NiagaraSystem.h
+++ b/Engine/Plugins/FX/Niagara/Source/Niagara/Classes/NiagaraSystem.h
@@ -29,7 +29,6 @@ struct FNiagaraEmitterCompiledData
 	UPROPERTY()
 	TArray<FName> SpawnAttributes;
 
-
 	/** Explicit list of Niagara Variables to bind to Emitter instances. */
 	UPROPERTY()
 	FNiagaraVariable EmitterSpawnIntervalVar;
@@ -45,6 +44,10 @@ struct FNiagaraEmitterCompiledData
 
 	UPROPERTY()
 	FNiagaraVariable EmitterRandomSeedVar;
+
+	/** Per-Emitter DataSet Data. */
+	UPROPERTY()
+	FNiagaraDataSetCompiledData DataSetCompiledData;
 };
 
 USTRUCT()
@@ -271,10 +274,13 @@ private:
 	void InitSystemCompiledData();
 
 	/** Helper for filling in precomputed variable names per emitter. Converts an emitter paramter "Emitter.XXXX" into it's real parameter name. */
-	void InitEmitterVariableAliasNames(FNiagaraEmitterCompiledData& EmitterCompiledDataToInit, const UNiagaraEmitter& InAssociatedEmitter);
+	void InitEmitterVariableAliasNames(FNiagaraEmitterCompiledData& EmitterCompiledDataToInit, const UNiagaraEmitter* InAssociatedEmitter);
 
 	/** Helper for generating aliased FNiagaraVariable names for the Emitter they are associated with. */
-	const FName GetEmitterVariableAliasName(const FNiagaraVariable& InEmitterVar, const UNiagaraEmitter& InEmitter) const;
+	const FName GetEmitterVariableAliasName(const FNiagaraVariable& InEmitterVar, const UNiagaraEmitter* InEmitter) const;
+
+	/** Helper for filling in attribute datasets per emitter. */
+	void InitEmitterDataSetCompiledData(FNiagaraDataSetCompiledData& DataSetToInit, const UNiagaraEmitter* InAssociatedEmitter, const FNiagaraEmitterHandle& InAssociatedEmitterHandle);
 #endif
 
 	void UpdatePostCompileDIInfo();
diff --git a/Engine/Plugins/FX/Niagara/Source/Niagara/Private/NiagaraDataSet.cpp b/Engine/Plugins/FX/Niagara/Source/Niagara/Private/NiagaraDataSet.cpp
--- a/Engine/Plugins/FX/Niagara/Source/Niagara/Private/NiagaraDataSet.cpp
+++ b/Engine/Plugins/FX/Niagara/Source/Niagara/Private/NiagaraDataSet.cpp
@@ -95,6 +95,19 @@ void FNiagaraDataSet::Init(FNiagaraDataSetID InID, ENiagaraSimTarget InSimTarget
 	SimTarget = InSimTarget;
 	DebugName = InDebugName;
 }
+
+void FNiagaraDataSet::InitFromCompiledData(const FNiagaraDataSetCompiledData& InDataSetCompiledData, const FString& InDebugName)
+{
+	Variables = InDataSetCompiledData.Variables;
+	VariableLayouts = InDataSetCompiledData.VariableLayouts;
+	TotalFloatComponents = InDataSetCompiledData.TotalFloatComponents;
+	TotalInt32Components = InDataSetCompiledData.TotalInt32Components;
+	bNeedsPersistentIDs = InDataSetCompiledData.bNeedsPersistentIDs;
+	ID = InDataSetCompiledData.ID;
+	SimTarget = InDataSetCompiledData.SimTarget;
+	DebugName = InDebugName;
+	bFinalized = true;
+}
 #else
 void FNiagaraDataSet::Init(FNiagaraDataSetID InID, ENiagaraSimTarget InSimTarget)
 {
@@ -102,6 +115,18 @@ void FNiagaraDataSet::Init(FNiagaraDataSetID InID, ENiagaraSimTarget InSimTarget
 	ID = InID;
 	SimTarget = InSimTarget;
 }
+
+void FNiagaraDataSet::InitFromCompiledData(const FNiagaraDataSetCompiledData& InDataSetCompiledData)
+{
+	Variables = InDataSetCompiledData.Variables;
+	VariableLayouts = InDataSetCompiledData.VariableLayouts;
+	TotalFloatComponents = InDataSetCompiledData.TotalFloatComponents;
+	TotalInt32Components = InDataSetCompiledData.TotalInt32Components;
+	bNeedsPersistentIDs = InDataSetCompiledData.bNeedsPersistentIDs;
+	ID = InDataSetCompiledData.ID;
+	SimTarget = InDataSetCompiledData.SimTarget;
+	bFinalized = true;
+}
 #endif
 
 void FNiagaraDataSet::Reset()
@@ -357,6 +382,7 @@ void FNiagaraDataSet::ReleaseGPUInstanceCounts(FNiagaraGPUInstanceCountManager&
 	}
 }
 
+// note, this method is also implemented in FNiagaraDataSetCompiledData::BuildLayout()
 void FNiagaraDataSet::BuildLayout()
 {
 	VariableLayouts.Empty();
@@ -1113,3 +1139,33 @@ void FNiagaraDataBuffer::BuildRegisterTable()
 		NumRegisters += NumFloats + NumInts;
 	}
 }
+
+// note, this method is also implemented in FNiagaraDataSet::BuildLayout()
+void FNiagaraDataSetCompiledData::BuildLayout()
+{
+	VariableLayouts.Empty();
+	TotalFloatComponents = 0;
+	TotalInt32Components = 0;
+
+	VariableLayouts.Reserve(Variables.Num());
+	for (FNiagaraVariable& Var : Variables)
+	{
+		FNiagaraVariableLayoutInfo& VarInfo = VariableLayouts[VariableLayouts.AddDefaulted()];
+		FNiagaraTypeLayoutInfo::GenerateLayoutInfo(VarInfo.LayoutInfo, Var.GetType().GetScriptStruct());
+		VarInfo.FloatComponentStart = TotalFloatComponents;
+		VarInfo.Int32ComponentStart = TotalInt32Components;
+		TotalFloatComponents += VarInfo.GetNumFloatComponents();
+		TotalInt32Components += VarInfo.GetNumInt32Components();
+	}
+}
+
+void FNiagaraDataSetCompiledData::Empty()
+{
+	bNeedsPersistentIDs = 0;
+	TotalFloatComponents = 0;
+	TotalInt32Components = 0;
+	Variables.Empty();
+	VariableLayouts.Empty();
+	ID = FNiagaraDataSetID();
+	SimTarget = ENiagaraSimTarget::CPUSim;
+}
diff --git a/Engine/Plugins/FX/Niagara/Source/Niagara/Private/NiagaraEmitterInstance.cpp b/Engine/Plugins/FX/Niagara/Source/Niagara/Private/NiagaraEmitterInstance.cpp
--- a/Engine/Plugins/FX/Niagara/Source/Niagara/Private/NiagaraEmitterInstance.cpp
+++ b/Engine/Plugins/FX/Niagara/Source/Niagara/Private/NiagaraEmitterInstance.cpp
@@ -196,27 +196,15 @@ void FNiagaraEmitterInstance::Init(int32 InEmitterIdx, FNiagaraSystemInstanceID
 
 	if (!IsAllowedToExecute())
 	{
-		//@TODO FNiagaraMessageManager Error bubbling here
+		//@todo(message manager) Error bubbling here
 		ExecutionState = ENiagaraExecutionState::Disabled;
 		return;
 	}
 
-
-#if !UE_BUILD_SHIPPING && !UE_BUILD_TEST 
-	Data.Init(FNiagaraDataSetID(CachedIDName, ENiagaraDataSetType::ParticleData), CachedEmitter->SimTarget, ParentSystemInstance->GetSystem()->GetName() + TEXT("/") + CachedEmitter->GetName());
-#else
-	Data.Init(FNiagaraDataSetID(CachedIDName, ENiagaraDataSetType::ParticleData), CachedEmitter->SimTarget);
-#endif
-
-	//Init the spawn infos to the correct number for this system.
-	int32 NumEvents = CachedEmitter->GetEventHandlers().Num();
 	const TArray<FNiagaraEmitterCompiledData>& EmitterCompiledData = ParentSystemInstance->GetSystem()->GetEmitterCompiledData();
-	if (EmitterCompiledData.IsValidIndex(EmitterIdx))
-	{
-		SpawnInfos.SetNum(EmitterCompiledData[EmitterIdx].SpawnAttributes.Num());
-	}
-	else
+	if (EmitterCompiledData.IsValidIndex(EmitterIdx) == false)
 	{
+		//@todo(message manager) Error bubbling here
 		ExecutionState = ENiagaraExecutionState::Disabled;
 		return;
 	}
@@ -230,6 +218,16 @@ void FNiagaraEmitterInstance::Init(int32 InEmitterIdx, FNiagaraSystemInstanceID
 		return;
 	}
 
+	//Init the spawn infos to the correct number for this system.
+	int32 NumEvents = CachedEmitter->GetEventHandlers().Num();
+	SpawnInfos.SetNum(EmitterCompiledData[EmitterIdx].SpawnAttributes.Num());
+
+#if !UE_BUILD_SHIPPING && !UE_BUILD_TEST 
+	Data.InitFromCompiledData(EmitterCompiledData[EmitterIdx].DataSetCompiledData, ParentSystemInstance->GetSystem()->GetName() + TEXT("/") + CachedEmitter->GetName());
+#else
+	Data.InitFromCompiledData(EmitterCompiledData[EmitterIdx].DataSetCompiledData);
+#endif
+
 	ResetSimulation();
 
 	DataSetMap.Empty();
@@ -242,6 +240,7 @@ void FNiagaraEmitterInstance::Init(int32 InEmitterIdx, FNiagaraSystemInstanceID
 
 	const bool bVerboseAttributeLogging = false;
 
+#if !UE_BUILD_SHIPPING && !UE_BUILD_TEST 
 	if (bVerboseAttributeLogging)
 	{
 		for (FNiagaraVariable& Attr : CachedEmitter->UpdateScriptProps.Script->GetVMExecutableData().Attributes)
@@ -260,16 +259,7 @@ void FNiagaraEmitterInstance::Init(int32 InEmitterIdx, FNiagaraSystemInstanceID
 			}
 		}
 	}
-	Data.AddVariables(CachedEmitter->UpdateScriptProps.Script->GetVMExecutableData().Attributes);
-	Data.AddVariables(CachedEmitter->SpawnScriptProps.Script->GetVMExecutableData().Attributes);
-
-	//if we use persistent IDs then add that here too.
-	if (RequiredPersistentID())
-	{
-		Data.SetNeedsPersistentIDs(true);
-	}
-
-	Data.Finalize();
+#endif
 
 	ensure(CachedEmitter->UpdateScriptProps.DataSetAccessSynchronized());
 	UpdateScriptEventDataSets.Empty();
@@ -396,7 +386,7 @@ void FNiagaraEmitterInstance::Init(int32 InEmitterIdx, FNiagaraSystemInstanceID
 	}
 }
 
-void FNiagaraEmitterInstance::ResetSimulation(bool bKillExisting)
+void FNiagaraEmitterInstance::ResetSimulation(bool bKillExisting /*= true*/)
 {
 	Age = 0;
 	Loops = 0;
diff --git a/Engine/Plugins/FX/Niagara/Source/Niagara/Private/NiagaraSystem.cpp b/Engine/Plugins/FX/Niagara/Source/Niagara/Private/NiagaraSystem.cpp
--- a/Engine/Plugins/FX/Niagara/Source/Niagara/Private/NiagaraSystem.cpp
+++ b/Engine/Plugins/FX/Niagara/Source/Niagara/Private/NiagaraSystem.cpp
@@ -280,7 +280,6 @@ void UNiagaraSystem::PostLoad()
 	if (!GetOutermost()->bIsCookedForEditor)
 	{
 		TArray<UNiagaraScript*> AllSystemScripts;
-	
 		UNiagaraScriptSourceBase* SystemScriptSource = nullptr;
 		if (SystemSpawnScript == nullptr)
 		{
@@ -388,10 +387,10 @@ void UNiagaraSystem::PostLoad()
 			UE_LOG(LogNiagara, Log, TEXT("System %s being compiled because there were changes to an emitter script Change ID."), *GetPathName());
 		}
 
-		if (EmitterCompiledData.Num() == 0)
-		{
-			InitEmitterCompiledData();
-		}
+	if (EmitterCompiledData.Num() == 0 || EmitterCompiledData[0].DataSetCompiledData.Variables.Num() == 0)
+	{
+		InitEmitterCompiledData();
+	}
 
 		if (SystemCompiledData.InstanceParamStore.GetNumParameters() == 0)
 		{
@@ -929,7 +928,7 @@ bool UNiagaraSystem::QueryCompileComplete(bool bWait, bool bDoPost, bool bDoNotA
 }
 
 #if WITH_EDITORONLY_DATA
-void UNiagaraSystem::InitEmitterVariableAliasNames(FNiagaraEmitterCompiledData& EmitterCompiledDataToInit, const UNiagaraEmitter& InAssociatedEmitter)
+void UNiagaraSystem::InitEmitterVariableAliasNames(FNiagaraEmitterCompiledData& EmitterCompiledDataToInit, const UNiagaraEmitter* InAssociatedEmitter)
 {
 	EmitterCompiledDataToInit.EmitterSpawnIntervalVar.SetName(GetEmitterVariableAliasName(SYS_PARAM_EMITTER_SPAWN_INTERVAL, InAssociatedEmitter));
 	EmitterCompiledDataToInit.EmitterInterpSpawnStartDTVar.SetName(GetEmitterVariableAliasName(SYS_PARAM_EMITTER_INTERP_SPAWN_START_DT, InAssociatedEmitter));
@@ -938,9 +937,26 @@ void UNiagaraSystem::InitEmitterVariableAliasNames(FNiagaraEmitterCompiledData&
 	EmitterCompiledDataToInit.EmitterRandomSeedVar.SetName(GetEmitterVariableAliasName(SYS_PARAM_EMITTER_RANDOM_SEED, InAssociatedEmitter));
 }
 
-const FName UNiagaraSystem::GetEmitterVariableAliasName(const FNiagaraVariable& InEmitterVar, const UNiagaraEmitter& InEmitter) const
+const FName UNiagaraSystem::GetEmitterVariableAliasName(const FNiagaraVariable& InEmitterVar, const UNiagaraEmitter* InEmitter) const
 {
-	return FName(*InEmitterVar.GetName().ToString().Replace(TEXT("Emitter."), *(InEmitter.GetUniqueEmitterName() + TEXT("."))));
+	return FName(*InEmitterVar.GetName().ToString().Replace(TEXT("Emitter."), *(InEmitter->GetUniqueEmitterName() + TEXT("."))));
+}
+
+void UNiagaraSystem::InitEmitterDataSetCompiledData(FNiagaraDataSetCompiledData& DataSetToInit, const UNiagaraEmitter* InAssociatedEmitter, const FNiagaraEmitterHandle& InAssociatedEmitterHandle)
+{
+	DataSetToInit.Empty();
+
+	DataSetToInit.Variables = InAssociatedEmitter->UpdateScriptProps.Script->GetVMExecutableData().Attributes;
+	for (const FNiagaraVariable& Var : InAssociatedEmitter->SpawnScriptProps.Script->GetVMExecutableData().Attributes)
+	{
+		DataSetToInit.Variables.AddUnique(Var);
+	}
+
+	DataSetToInit.bNeedsPersistentIDs = InAssociatedEmitter->RequiresPersistantIDs() || DataSetToInit.Variables.Contains(SYS_PARAM_PARTICLES_ID);
+	DataSetToInit.ID = FNiagaraDataSetID(InAssociatedEmitterHandle.GetIdName(), ENiagaraDataSetType::ParticleData);
+	DataSetToInit.SimTarget = InAssociatedEmitter->SimTarget;
+
+	DataSetToInit.BuildLayout();
 }
 #endif
 
@@ -1076,10 +1092,13 @@ void UNiagaraSystem::InitEmitterCompiledData()
 			for (int32 EmitterIdx = 0; EmitterIdx < EmitterHandles.Num(); ++EmitterIdx)
 			{
 				UNiagaraEmitter* Emitter = EmitterHandles[EmitterIdx].GetInstance();
-				FString EmitterName = Emitter->GetUniqueEmitterName() + TEXT(".");
-				if (Var.GetType() == SpawnInfoDef && Var.GetName().ToString().StartsWith(EmitterName))
+				if ensureMsgf(Emitter != nullptr, TEXT("Failed to get Emitter Instance from Emitter Handle in post compile, please investigate."))
 				{
-					EmitterCompiledData[EmitterIdx].SpawnAttributes.AddUnique(Var.GetName());
+					FString EmitterName = Emitter->GetUniqueEmitterName() + TEXT(".");
+					if (Var.GetType() == SpawnInfoDef && Var.GetName().ToString().StartsWith(EmitterName))
+					{
+						EmitterCompiledData[EmitterIdx].SpawnAttributes.AddUnique(Var.GetName());
+					}
 				}
 			}
 		}
@@ -1089,17 +1108,27 @@ void UNiagaraSystem::InitEmitterCompiledData()
 			for (int32 EmitterIdx = 0; EmitterIdx < EmitterHandles.Num(); ++EmitterIdx)
 			{
 				UNiagaraEmitter* Emitter = EmitterHandles[EmitterIdx].GetInstance();
-				FString EmitterName = Emitter->GetUniqueEmitterName() + TEXT(".");
-				if (Var.GetType() == SpawnInfoDef && Var.GetName().ToString().StartsWith(EmitterName))
+				if ensureMsgf(Emitter != nullptr, TEXT("Failed to get Emitter Instance from Emitter Handle in post compile, please investigate."))
 				{
-					EmitterCompiledData[EmitterIdx].SpawnAttributes.AddUnique(Var.GetName());
+					FString EmitterName = Emitter->GetUniqueEmitterName() + TEXT(".");
+					if (Var.GetType() == SpawnInfoDef && Var.GetName().ToString().StartsWith(EmitterName))
+					{
+						EmitterCompiledData[EmitterIdx].SpawnAttributes.AddUnique(Var.GetName());
+					}
 				}
 				}
 			}
 
 		for (int32 EmitterIdx = 0; EmitterIdx < EmitterHandles.Num(); ++EmitterIdx)
 		{
-			InitEmitterVariableAliasNames(EmitterCompiledData[EmitterIdx], *EmitterHandles[EmitterIdx].GetInstance());
+			const FNiagaraEmitterHandle& EmitterHandle = EmitterHandles[EmitterIdx];
+			const UNiagaraEmitter* Emitter = EmitterHandle.GetInstance();
+			FNiagaraDataSetCompiledData& EmitterDataSetCompiledData = EmitterCompiledData[EmitterIdx].DataSetCompiledData;
+			if ensureMsgf(Emitter != nullptr, TEXT("Failed to get Emitter Instance from Emitter Handle in post compile, please investigate."))
+			{
+				InitEmitterVariableAliasNames(EmitterCompiledData[EmitterIdx], Emitter);
+				InitEmitterDataSetCompiledData(EmitterDataSetCompiledData, Emitter, EmitterHandle);
+			}
 		}
 	}
 }
diff --git a/Engine/Plugins/FX/Niagara/Source/Niagara/Public/NiagaraTypes.h b/Engine/Plugins/FX/Niagara/Source/Niagara/Public/NiagaraTypes.h
--- a/Engine/Plugins/FX/Niagara/Source/Niagara/Public/NiagaraTypes.h
+++ b/Engine/Plugins/FX/Niagara/Source/Niagara/Public/NiagaraTypes.h
@@ -188,19 +188,28 @@ FORCEINLINE uint32 GetTypeHash(const FNiagaraID& ID)
 }
 
 /** Information about how this type should be laid out in an FNiagaraDataSet */
+USTRUCT()
 struct FNiagaraTypeLayoutInfo
 {
+	GENERATED_BODY()
+
 	FNiagaraTypeLayoutInfo()
 	{}
 
 	/** Byte offset of each float component in a structured layout. */
+	UPROPERTY()
 	TArray<uint32> FloatComponentByteOffsets;
+
 	/** Offset into register table for each float component. */
+	UPROPERTY()
 	TArray<uint32> FloatComponentRegisterOffsets;
 
 	/** Byte offset of each int32 component in a structured layout. */
+	UPROPERTY()
 	TArray<uint32> Int32ComponentByteOffsets;
+
 	/** Offset into register table for each int32 component. */
+	UPROPERTY()
 	TArray<uint32> Int32ComponentRegisterOffsets;
 
 	FORCEINLINE uint32 GetNumComponents()const { return FloatComponentByteOffsets.Num() + Int32ComponentByteOffsets.Num(); }
