- Ensure state of variables / reset ones that were missing when we Reset / ReInit a system
- This fixes a bug with system that have completed not reactivating when bursting

#rnx

diff --git a/Engine/Plugins/FX/Niagara/Source/Niagara/Private/NiagaraSystemInstance.cpp b/Engine/Plugins/FX/Niagara/Source/Niagara/Private/NiagaraSystemInstance.cpp
--- a/Engine/Plugins/FX/Niagara/Source/Niagara/Private/NiagaraSystemInstance.cpp
+++ b/Engine/Plugins/FX/Niagara/Source/Niagara/Private/NiagaraSystemInstance.cpp
@@ -672,8 +672,17 @@ void FNiagaraSystemInstance::Reset(FNiagaraSystemInstance::EResetMode Mode)
 
 void FNiagaraSystemInstance::ResetInternal(bool bResetSimulations)
 {
+	check(SystemInstanceIndex == INDEX_NONE);
+	ensure(bPendingSpawn == false);
+	ensure(bPaused == false);
+	ensure(bAsyncWorkInProgress == false);
+	ensure(bNeedsFinalize == false);
+
 	Age = 0;
 	TickCount = 0;
+	bHasTickingEmitters = true;
+	CachedDeltaSeconds = 0.0f;
+
 	UNiagaraSystem* System = GetSystem();
 	if (System == nullptr || Component == nullptr || IsDisabled())
 	{
@@ -784,6 +793,12 @@ bool DoSystemDataInterfacesRequireSolo(const UNiagaraSystem& System, const UNiag
 
 void FNiagaraSystemInstance::ReInitInternal()
 {
+	check(SystemInstanceIndex == INDEX_NONE);
+	ensure(bPendingSpawn == false);
+	ensure(bPaused == false);
+	ensure(bAsyncWorkInProgress == false);
+	ensure(bNeedsFinalize == false);
+
 	SCOPE_CYCLE_COUNTER(STAT_NiagaraSystemReinit);
 
 	SCOPE_CYCLE_COUNTER(STAT_NiagaraOverview_GT);
@@ -792,6 +807,9 @@ void FNiagaraSystemInstance::ReInitInternal()
 
 	Age = 0;
 	TickCount = 0;
+	bHasTickingEmitters = true;
+	CachedDeltaSeconds = 0.0f;
+
 	UNiagaraSystem* System = GetSystem();
 	if (System == nullptr || Component == nullptr)
 	{
diff --git a/Engine/Plugins/FX/Niagara/Source/Niagara/Private/NiagaraSystemSimulation.cpp b/Engine/Plugins/FX/Niagara/Source/Niagara/Private/NiagaraSystemSimulation.cpp
--- a/Engine/Plugins/FX/Niagara/Source/Niagara/Private/NiagaraSystemSimulation.cpp
+++ b/Engine/Plugins/FX/Niagara/Source/Niagara/Private/NiagaraSystemSimulation.cpp
@@ -1142,7 +1142,7 @@ void FNiagaraSystemSimulation::SpawnSystemInstances(FNiagaraSystemSimulationTick
 	SpawnExecContext.Tick(SoloSystemInstance);//We can't require a specific instance here as these are for all instances.
 	SpawnExecContext.BindData(0, Context.DataSet, OrigNum, false);
 	SpawnExecContext.BindData(1, SpawnInstanceParameterDataSet, OrigNum, false);
-	SpawnExecContext.Execute(Context.SpawnNum);
+	SpawnExecContext.Execute(SpawnNum);
 
 	if (GbDumpSystemData || Context.System->bDumpDebugSystemInfo)
 	{
@@ -1151,6 +1151,8 @@ void FNiagaraSystemSimulation::SpawnSystemInstances(FNiagaraSystemSimulationTick
 		SpawnInstanceParameterDataSet.GetCurrentDataChecked().Dump(0, NumInstances, TEXT("Spawn Instance Parameter Data"));
 	}
 
+	Context.DataSet.EndSimulate();
+
 #if WITH_EDITORONLY_DATA
 	if (SoloSystemInstance && SoloSystemInstance->ShouldCaptureThisFrame())
 	{
@@ -1164,8 +1166,6 @@ void FNiagaraSystemSimulation::SpawnSystemInstances(FNiagaraSystemSimulationTick
 	}
 #endif
 
-	Context.DataSet.EndSimulate();
-
 	check(Context.DataSet.GetCurrentDataChecked().GetNumInstances() == Context.Instances.Num());
 }
 
@@ -1186,10 +1186,13 @@ void FNiagaraSystemSimulation::UpdateSystemInstances(FNiagaraSystemSimulationTic
 		DestinationData.SetNumInstances(NumInstances);
 
 		//Run update.
-		UpdateExecContext.Tick(Context.Instances[0]);
-		UpdateExecContext.BindData(0, Context.DataSet, 0, false);
-		UpdateExecContext.BindData(1, UpdateInstanceParameterDataSet, 0, false);
-		UpdateExecContext.Execute(OrigNum);
+		if (OrigNum > 0)
+		{
+			UpdateExecContext.Tick(Context.Instances[0]);
+			UpdateExecContext.BindData(0, Context.DataSet, 0, false);
+			UpdateExecContext.BindData(1, UpdateInstanceParameterDataSet, 0, false);
+			UpdateExecContext.Execute(OrigNum);
+		}
 
 		if (GbDumpSystemData || Context.System->bDumpDebugSystemInfo)
 		{
