- Enable mid frame waiting on Niagara Component again (was accidentally disabled)

#rnx

diff --git a/Engine/Plugins/FX/Niagara/Source/Niagara/Private/NiagaraComponent.cpp b/Engine/Plugins/FX/Niagara/Source/Niagara/Private/NiagaraComponent.cpp
--- a/Engine/Plugins/FX/Niagara/Source/Niagara/Private/NiagaraComponent.cpp
+++ b/Engine/Plugins/FX/Niagara/Source/Niagara/Private/NiagaraComponent.cpp
@@ -1003,7 +1003,7 @@ void UNiagaraComponent::OnEndOfFrameUpdateDuringTick()
 	Super::OnEndOfFrameUpdateDuringTick();
 	if ( SystemInstance )
 	{
-		//SystemInstance->WaitForAsyncTick();
+		SystemInstance->WaitForAsyncTick();
 	}
 }
 
