Remove unsused module inputs during compilation

#tests all niagara automation tests are green

diff --git a/Engine/Plugins/FX/Niagara/Source/NiagaraEditor/Private/NiagaraHlslTranslator.cpp b/Engine/Plugins/FX/Niagara/Source/NiagaraEditor/Private/NiagaraHlslTranslator.cpp
--- a/Engine/Plugins/FX/Niagara/Source/NiagaraEditor/Private/NiagaraHlslTranslator.cpp
+++ b/Engine/Plugins/FX/Niagara/Source/NiagaraEditor/Private/NiagaraHlslTranslator.cpp
@@ -4564,7 +4564,25 @@ void FHlslNiagaraTranslator::FunctionCall(UNiagaraNodeFunctionCall* FunctionNode
 	ActiveHistoryForFunctionCalls.ExitFunction(FunctionNode->GetFunctionName(), FunctionNode->FunctionScript, FunctionNode);
 }
 
+void FHlslNiagaraTranslator::EnterFunctionCallNode(const TSet<FName>& UnusedInputs)
+{
+	FunctionNodeStack.Add(UnusedInputs);
+}
+
+void FHlslNiagaraTranslator::ExitFunctionCallNode()
+{
+	ensure(FunctionNodeStack.Num() > 0);
+	FunctionNodeStack.Pop(false);
+}
 
+bool FHlslNiagaraTranslator::IsFunctionVariableCulledFromCompilation(const FName& InputName) const
+{
+	if (FunctionNodeStack.Num() == 0)
+	{
+		return false;
+	}
+	return FunctionNodeStack.Last().Contains(InputName);
+}
 
 // From a valid list of namespaces, resolve any aliased tokens and promote namespaced variables without a master namespace to the input parameter map instance namespace
 void FHlslNiagaraTranslator::FinalResolveNamespacedTokens(const FString& ParameterMapInstanceNamespace, TArray<FString>& Tokens, TArray<FString>& ValidChildNamespaces, FNiagaraParameterMapHistoryBuilder& Builder, TArray<FNiagaraVariable>& UniqueParameterMapEntriesAliasesIntact, TArray<FNiagaraVariable>& UniqueParameterMapEntries, int32 ParamMapHistoryIdx)
diff --git a/Engine/Plugins/FX/Niagara/Source/NiagaraEditor/Private/NiagaraHlslTranslator.h b/Engine/Plugins/FX/Niagara/Source/NiagaraEditor/Private/NiagaraHlslTranslator.h
--- a/Engine/Plugins/FX/Niagara/Source/NiagaraEditor/Private/NiagaraHlslTranslator.h
+++ b/Engine/Plugins/FX/Niagara/Source/NiagaraEditor/Private/NiagaraHlslTranslator.h
@@ -447,6 +447,9 @@ public:
 	FString CompileDataInterfaceFunction(UNiagaraDataInterface* DataInterface, FNiagaraFunctionSignature& Signature);
 
 	virtual void FunctionCall(UNiagaraNodeFunctionCall* FunctionNode, TArray<int32>& Inputs, TArray<int32>& Outputs);
+	void EnterFunctionCallNode(const TSet<FName>& UnusedInputs);
+	void ExitFunctionCallNode();
+	bool IsFunctionVariableCulledFromCompilation(const FName& InputName) const;
 
 	virtual void Convert(class UNiagaraNodeConvert* Convert, TArray <int32>& Inputs, TArray<int32>& Outputs);
 	virtual void If(class UNiagaraNodeIf* IfNode, TArray<FNiagaraVariable>& Vars, int32 Condition, TArray<int32>& PathA, TArray<int32>& PathB, TArray<int32>& Outputs);
@@ -572,7 +575,10 @@ private:
 	FNiagaraParameterMapHistoryBuilder ActiveHistoryForFunctionCalls;
 
 	// Synced to the ParamMapHistories.
-	TArray<TArray<int32> > ParamMapSetVariablesToChunks;
+	TArray<TArray<int32>> ParamMapSetVariablesToChunks;
+
+	// Used to keep track of contextual information about the currently compiled function node (currently only the unsused parameter names)
+	TArray<TSet<FName>> FunctionNodeStack;
 
 	// Synced to the System uniforms encountered for parameter maps thus far.
 	TMap<FName, int32> ParamMapDefinedSystemVarsToUniformChunks; // Map from the defined constants to the uniform chunk expressing them (i.e. have we encountered before in this graph?)
diff --git a/Engine/Plugins/FX/Niagara/Source/NiagaraEditor/Private/NiagaraNodeFunctionCall.cpp b/Engine/Plugins/FX/Niagara/Source/NiagaraEditor/Private/NiagaraNodeFunctionCall.cpp
--- a/Engine/Plugins/FX/Niagara/Source/NiagaraEditor/Private/NiagaraNodeFunctionCall.cpp
+++ b/Engine/Plugins/FX/Niagara/Source/NiagaraEditor/Private/NiagaraNodeFunctionCall.cpp
@@ -19,6 +19,7 @@
 #include "NiagaraDataInterfaceSkeletalMesh.h"
 #include "SNiagaraGraphNodeFunctionCallWithSpecifiers.h"
 #include "Misc/SecureHash.h"
+#include "ViewModels/Stack/NiagaraStackGraphUtilities.h"
 
 #define LOCTEXT_NAMESPACE "NiagaraNodeFunctionCall"
 
@@ -452,6 +453,24 @@ void UNiagaraNodeFunctionCall::Compile(class FHlslNiagaraTranslator* Translator,
 		Options.bFilterDuplicates = true;
 		FunctionGraph->FindInputNodes(FunctionInputNodes, Options);
 
+		// We check which module inputs are not used so we can later remove them from the compilation of the
+		// parameter map that sets the input values for our function. This is mainly done to prevent data interfaces being
+		// initialized as parameter when they are not used in the function or module.
+		TArray<const UEdGraphPin*> OutInputPins;
+		TSet<const UEdGraphPin*> HiddenPins;
+		TSet<FName> HiddenPinNames;
+		FNiagaraStackGraphUtilities::GetStackFunctionInputPins(*this, OutInputPins, HiddenPins, FCompileConstantResolver(Translator));
+		for (const UEdGraphPin* Pin : HiddenPins)
+		{
+			FString PinNameString = Pin->PinName.ToString();
+			if (PinNameString.RemoveFromStart(TEXT("Module.")))
+			{
+				PinNameString = FunctionScript->GetFName().ToString() + TEXT(".") + PinNameString;
+				HiddenPinNames.Add(FName(*PinNameString));
+			}
+		}
+		Translator->EnterFunctionCallNode(HiddenPinNames);
+
 		for (UNiagaraNodeInput* FunctionInputNode : FunctionInputNodes)
 		{
 			//Finds the matching Pin in the caller.
@@ -534,6 +553,7 @@ void UNiagaraNodeFunctionCall::Compile(class FHlslNiagaraTranslator* Translator,
 
 		FCompileConstantResolver ConstantResolver(Translator);
 		FNiagaraEditorUtilities::SetStaticSwitchConstants(GetCalledGraph(), CallerInputPins, ConstantResolver);
+		Translator->ExitFunctionCallNode();
 	}
 	else if (Signature.IsValid())
 	{
@@ -561,8 +581,10 @@ void UNiagaraNodeFunctionCall::Compile(class FHlslNiagaraTranslator* Translator,
 				}
 			}
 		}
+		Translator->EnterFunctionCallNode(TSet<FName>());
 		Signature.FunctionSpecifiers = FunctionSpecifiers;
 		bError = CompileInputPins(Translator, Inputs);
+		Translator->ExitFunctionCallNode();
 	}		
 	else
 	{
diff --git a/Engine/Plugins/FX/Niagara/Source/NiagaraEditor/Private/NiagaraNodeParameterMapSet.cpp b/Engine/Plugins/FX/Niagara/Source/NiagaraEditor/Private/NiagaraNodeParameterMapSet.cpp
--- a/Engine/Plugins/FX/Niagara/Source/NiagaraEditor/Private/NiagaraNodeParameterMapSet.cpp
+++ b/Engine/Plugins/FX/Niagara/Source/NiagaraEditor/Private/NiagaraNodeParameterMapSet.cpp
@@ -166,6 +166,12 @@ void UNiagaraNodeParameterMapSet::Compile(class FHlslNiagaraTranslator* Translat
 			continue;
 		}
 
+		if (Translator->IsFunctionVariableCulledFromCompilation(InputPin->PinName))
+		{
+			CompileInputs.Add(INDEX_NONE);
+			continue;
+		}
+
 		if (IsNodeEnabled() == false && Schema->PinToTypeDefinition(InputPin) != FNiagaraTypeDefinition::GetParameterMapDef())
 		{
 			continue;
