- Remove unused VertexFactories

#rnx

diff --git a/Engine/Plugins/FX/Niagara/Source/Niagara/Private/NiagaraRendererMeshes.cpp b/Engine/Plugins/FX/Niagara/Source/Niagara/Private/NiagaraRendererMeshes.cpp
--- a/Engine/Plugins/FX/Niagara/Source/Niagara/Private/NiagaraRendererMeshes.cpp
+++ b/Engine/Plugins/FX/Niagara/Source/Niagara/Private/NiagaraRendererMeshes.cpp
@@ -72,8 +72,6 @@ FNiagaraRendererMeshes::FNiagaraRendererMeshes(ERHIFeatureLevel::Type FeatureLev
 	check(Emitter);
 	check(Props);
 
-	VertexFactory = ConstructNiagaraMeshVertexFactory(NVFT_Mesh, FeatureLevel);
-	
 	const UNiagaraMeshRendererProperties* Properties = CastChecked<const UNiagaraMeshRendererProperties>(Props);
 	check(Properties->ParticleMesh);
 
@@ -110,23 +108,16 @@ FNiagaraRendererMeshes::FNiagaraRendererMeshes(ERHIFeatureLevel::Type FeatureLev
 
 FNiagaraRendererMeshes::~FNiagaraRendererMeshes()
 {
-	if (VertexFactory != nullptr)
-	{
-		delete VertexFactory;
-		VertexFactory = nullptr;
-	}
 }
 
 void FNiagaraRendererMeshes::ReleaseRenderThreadResources(NiagaraEmitterInstanceBatcher* Batcher)
 {
 	FNiagaraRenderer::ReleaseRenderThreadResources(Batcher);
-	VertexFactory->ReleaseResource();
 }
 
 void FNiagaraRendererMeshes::CreateRenderThreadResources(NiagaraEmitterInstanceBatcher* Batcher)
 {
 	FNiagaraRenderer::CreateRenderThreadResources(Batcher);
-	VertexFactory->InitResource();
 }
 
 void FNiagaraRendererMeshes::SetupVertexFactory(FNiagaraMeshVertexFactory *InVertexFactory, const FStaticMeshLODResources& LODResources) const
@@ -237,7 +228,7 @@ void FNiagaraRendererMeshes::GetDynamicMeshElements(const TArray<const FSceneVie
 				for (FMaterialRenderProxy* MaterialProxy : DynamicDataMesh->Materials)
 				{
 					check(MaterialProxy);
-					EBlendMode BlendMode = MaterialProxy->GetMaterial(SceneProxy->GetScene().GetFeatureLevel())->GetBlendMode();
+					EBlendMode BlendMode = MaterialProxy->GetMaterial(View->FeatureLevel)->GetBlendMode();
 					bHasTranslucentMaterials |= BlendMode == BLEND_AlphaComposite || BlendMode == BLEND_AlphaHoldout || BlendMode == BLEND_Translucent;
 				}
 
diff --git a/Engine/Plugins/FX/Niagara/Source/Niagara/Private/NiagaraRendererRibbons.cpp b/Engine/Plugins/FX/Niagara/Source/Niagara/Private/NiagaraRendererRibbons.cpp
--- a/Engine/Plugins/FX/Niagara/Source/Niagara/Private/NiagaraRendererRibbons.cpp
+++ b/Engine/Plugins/FX/Niagara/Source/Niagara/Private/NiagaraRendererRibbons.cpp
@@ -150,8 +150,6 @@ FNiagaraRendererRibbons::FNiagaraRendererRibbons(ERHIFeatureLevel::Type FeatureL
 	, NormalizedAgeDataOffset(INDEX_NONE)
 	, MaterialRandomDataOffset(INDEX_NONE)
 {
-	VertexFactory = new FNiagaraRibbonVertexFactory(NVFT_Ribbon, FeatureLevel);
-
 	const UNiagaraRibbonRendererProperties* Properties = CastChecked<const UNiagaraRibbonRendererProperties>(InProps);
 	const FNiagaraDataSet& Data = Emitter->GetData();
 
@@ -199,17 +197,11 @@ FNiagaraRendererRibbons::FNiagaraRendererRibbons(ERHIFeatureLevel::Type FeatureL
 
 FNiagaraRendererRibbons::~FNiagaraRendererRibbons()
 {
-	if (VertexFactory != nullptr)
-	{
-		delete VertexFactory;
-		VertexFactory = nullptr;
-	}
 }
 
 void FNiagaraRendererRibbons::ReleaseRenderThreadResources(NiagaraEmitterInstanceBatcher* Batcher)
 {
 	FNiagaraRenderer::ReleaseRenderThreadResources(Batcher);
-	VertexFactory->ReleaseResource();
 #if RHI_RAYTRACING
 	if (IsRayTracingEnabled())
 	{
@@ -223,7 +215,6 @@ void FNiagaraRendererRibbons::ReleaseRenderThreadResources(NiagaraEmitterInstanc
 void FNiagaraRendererRibbons::CreateRenderThreadResources(NiagaraEmitterInstanceBatcher* Batcher)
 {
 	FNiagaraRenderer::CreateRenderThreadResources(Batcher);
-	VertexFactory->InitResource();
 #if RHI_RAYTRACING
 	if (IsRayTracingEnabled())
 	{
diff --git a/Engine/Plugins/FX/Niagara/Source/Niagara/Private/NiagaraRendererSprites.cpp b/Engine/Plugins/FX/Niagara/Source/Niagara/Private/NiagaraRendererSprites.cpp
--- a/Engine/Plugins/FX/Niagara/Source/Niagara/Private/NiagaraRendererSprites.cpp
+++ b/Engine/Plugins/FX/Niagara/Source/Niagara/Private/NiagaraRendererSprites.cpp
@@ -93,8 +93,6 @@ FNiagaraRendererSprites::FNiagaraRendererSprites(ERHIFeatureLevel::Type FeatureL
 {
 	check(InProps && Emitter);
 
-	VertexFactory = new FNiagaraSpriteVertexFactory(NVFT_Sprite, FeatureLevel);
-
 	const UNiagaraSpriteRendererProperties* Properties = CastChecked<const UNiagaraSpriteRendererProperties>(InProps);
 	
 	Alignment = Properties->Alignment;
@@ -140,17 +138,11 @@ FNiagaraRendererSprites::FNiagaraRendererSprites(ERHIFeatureLevel::Type FeatureL
 
 FNiagaraRendererSprites::~FNiagaraRendererSprites()
 {
-	if ( VertexFactory != nullptr )
-	{
-		delete VertexFactory;
-		VertexFactory = nullptr;
-	}
 }
 
 void FNiagaraRendererSprites::ReleaseRenderThreadResources(NiagaraEmitterInstanceBatcher* Batcher)
 {
 	FNiagaraRenderer::ReleaseRenderThreadResources(Batcher);
-	VertexFactory->ReleaseResource();
 
 	CutoutVertexBuffer.ReleaseResource();
 #if RHI_RAYTRACING
@@ -166,7 +158,6 @@ void FNiagaraRendererSprites::CreateRenderThreadResources(NiagaraEmitterInstance
 {
 	FNiagaraRenderer::CreateRenderThreadResources(Batcher);
 	CutoutVertexBuffer.InitResource();
-	VertexFactory->InitResource();
 
 #if RHI_RAYTRACING
 	if (IsRayTracingEnabled())
@@ -325,7 +316,7 @@ void FNiagaraRendererSprites::SetVertexFactoryParticleData(
 
 		FMaterialRenderProxy* MaterialRenderProxy = DynamicDataSprites->Material;
 		check(MaterialRenderProxy);
-		EBlendMode BlendMode = MaterialRenderProxy->GetMaterial(VertexFactory->GetFeatureLevel())->GetBlendMode();
+		EBlendMode BlendMode = MaterialRenderProxy->GetMaterial(View->FeatureLevel)->GetBlendMode();
 		OutVertexFactory.SetSortedIndices(nullptr, 0xFFFFFFFF);
 
 		int32 NumInstances = SourceParticleData->GetNumInstances();
diff --git a/Engine/Plugins/FX/Niagara/Source/Niagara/Public/NiagaraRendererMeshes.h b/Engine/Plugins/FX/Niagara/Source/Niagara/Public/NiagaraRendererMeshes.h
--- a/Engine/Plugins/FX/Niagara/Source/Niagara/Public/NiagaraRendererMeshes.h
+++ b/Engine/Plugins/FX/Niagara/Source/Niagara/Public/NiagaraRendererMeshes.h
@@ -33,8 +33,6 @@ public:
 
 private:
 
-	class FNiagaraMeshVertexFactory* VertexFactory;
-
 	/** Render data of the static mesh we use. */
 	FStaticMeshRenderData* MeshRenderData;
 
diff --git a/Engine/Plugins/FX/Niagara/Source/Niagara/Public/NiagaraRendererRibbons.h b/Engine/Plugins/FX/Niagara/Source/Niagara/Public/NiagaraRendererRibbons.h
--- a/Engine/Plugins/FX/Niagara/Source/Niagara/Public/NiagaraRendererRibbons.h
+++ b/Engine/Plugins/FX/Niagara/Source/Niagara/Public/NiagaraRendererRibbons.h
@@ -61,8 +61,6 @@ private:
 
 	FCPUSimParticleDataAllocation AllocateParticleDataIfCPUSim(struct FNiagaraDynamicDataRibbon* DynamicDataRibbon, FGlobalDynamicReadBuffer& DynamicReadBuffer) const;
 
-	class FNiagaraRibbonVertexFactory *VertexFactory;
-
 	ENiagaraRibbonFacingMode FacingMode;
 	float UV0TilingDistance;
 	FVector2D UV0Scale;
diff --git a/Engine/Plugins/FX/Niagara/Source/Niagara/Public/NiagaraRendererSprites.h b/Engine/Plugins/FX/Niagara/Source/Niagara/Public/NiagaraRendererSprites.h
--- a/Engine/Plugins/FX/Niagara/Source/Niagara/Public/NiagaraRendererSprites.h
+++ b/Engine/Plugins/FX/Niagara/Source/Niagara/Public/NiagaraRendererSprites.h
@@ -57,8 +57,6 @@ private:
 		FMeshBatch& OutMeshBatch,
 		class FNiagaraMeshCollectorResourcesSprite& OutCollectorResources) const;
 
-	class FNiagaraSpriteVertexFactory* VertexFactory;
-
 	//Cached data from the properties struct.
 	ENiagaraSpriteAlignment Alignment;
 	ENiagaraSpriteFacingMode FacingMode;
