- Fix for override parameters being ignored in cooked builds
- We rejected the parameter store binding as we had no matching parameters at the time of binding, they were added later

#rnx

diff --git a/Engine/Plugins/FX/Niagara/Source/Niagara/Private/NiagaraParameterStore.cpp b/Engine/Plugins/FX/Niagara/Source/Niagara/Private/NiagaraParameterStore.cpp
--- a/Engine/Plugins/FX/Niagara/Source/Niagara/Private/NiagaraParameterStore.cpp
+++ b/Engine/Plugins/FX/Niagara/Source/Niagara/Private/NiagaraParameterStore.cpp
@@ -116,19 +116,26 @@ FNiagaraParameterStore::~FNiagaraParameterStore()
 	Bindings.Empty();
 }
 
-void FNiagaraParameterStore::Bind(FNiagaraParameterStore* DestStore)
+void FNiagaraParameterStore::Bind(FNiagaraParameterStore* DestStore, bool bAlwaysAddToBindings)
 {
 	check(DestStore);
 	SCOPE_CYCLE_COUNTER(STAT_NiagaraParameterStoreBind);
 	if (!Bindings.Contains(DestStore))
 	{
-		// Bind the parameter stores only if they have variables in common.
-		FNiagaraParameterStoreBinding HeapBinding;
-		if  (HeapBinding.Initialize(DestStore, this))
+		if (bAlwaysAddToBindings)
 		{
 			FNiagaraParameterStoreBinding& Binding = Bindings.FindOrAdd(DestStore);
-			FMemory::Memswap(&Binding, &HeapBinding, sizeof(FNiagaraParameterStoreBinding));
-			return;
+			Binding.Initialize(DestStore, this);
+		}
+		else
+		{
+			// Bind the parameter stores only if they have variables in common.
+			FNiagaraParameterStoreBinding HeapBinding;
+			if  (HeapBinding.Initialize(DestStore, this))
+			{
+				FNiagaraParameterStoreBinding& Binding = Bindings.FindOrAdd(DestStore);
+				FMemory::Memswap(&Binding, &HeapBinding, sizeof(FNiagaraParameterStoreBinding));
+			}
 		}
 	}
 }
diff --git a/Engine/Plugins/FX/Niagara/Source/Niagara/Private/NiagaraSystemInstance.cpp b/Engine/Plugins/FX/Niagara/Source/Niagara/Private/NiagaraSystemInstance.cpp
--- a/Engine/Plugins/FX/Niagara/Source/Niagara/Private/NiagaraSystemInstance.cpp
+++ b/Engine/Plugins/FX/Niagara/Source/Niagara/Private/NiagaraSystemInstance.cpp
@@ -974,14 +974,14 @@ void FNiagaraSystemInstance::BindParameters()
 		return;
 	}
 
-	Component->GetOverrideParameters().Bind(&InstanceParameters);
+	Component->GetOverrideParameters().Bind(&InstanceParameters, true);
 
 	if (SystemSimulation->GetIsSolo())
 	{
 		// If this simulation is solo than we can bind the instance parameters to the system simulation contexts so that
 		// the system and emitter scripts use the per-instance data interfaces.
-		Component->GetOverrideParameters().Bind(&SystemSimulation->GetSpawnExecutionContext().Parameters);
-		Component->GetOverrideParameters().Bind(&SystemSimulation->GetUpdateExecutionContext().Parameters);
+		Component->GetOverrideParameters().Bind(&SystemSimulation->GetSpawnExecutionContext().Parameters, true);
+		Component->GetOverrideParameters().Bind(&SystemSimulation->GetUpdateExecutionContext().Parameters, true);
 	}
 	else 
 	{
diff --git a/Engine/Plugins/FX/Niagara/Source/Niagara/Public/NiagaraParameterStore.h b/Engine/Plugins/FX/Niagara/Source/Niagara/Public/NiagaraParameterStore.h
--- a/Engine/Plugins/FX/Niagara/Source/Niagara/Public/NiagaraParameterStore.h
+++ b/Engine/Plugins/FX/Niagara/Source/Niagara/Public/NiagaraParameterStore.h
@@ -213,8 +213,12 @@ public:
 
 	FORCEINLINE uint32 GetLayoutVersion() const { return LayoutVersion; }
 
-	/** Binds this parameter store to another. During Tick the values of this store will be pushed into all bound stores */
-	void Bind(FNiagaraParameterStore* DestStore);
+	/**
+	  * Binds this parameter store to another, by default if we find no matching parameters we will not maintain a pointer to the store.
+	  * Set bAlwaysAddToBindings to true to always add to the list of ParameterStore bindings, regardless of no matching bindings.
+	  * During Tick the values of this store will be pushed into all bound stores
+	  */
+	void Bind(FNiagaraParameterStore* DestStore, bool bAlwaysAddToBindings = false);
 	/** Unbinds this store form one it's bound to. */
 	void Unbind(FNiagaraParameterStore* DestStore);
 	/** Recreates any bindings to reflect a layout change etc. */
