- Remove SimpleTimer, other tools should be used to time sections like UnrealInsights, etc

#rnx

diff --git a/Engine/Plugins/FX/Niagara/Source/Niagara/Classes/NiagaraEmitterInstance.h b/Engine/Plugins/FX/Niagara/Source/Niagara/Classes/NiagaraEmitterInstance.h
--- a/Engine/Plugins/FX/Niagara/Source/Niagara/Classes/NiagaraEmitterInstance.h
+++ b/Engine/Plugins/FX/Niagara/Source/Niagara/Classes/NiagaraEmitterInstance.h
@@ -87,7 +87,7 @@ public:
 
 	FNiagaraSystemInstance* GetParentSystemInstance()const	{ return ParentSystemInstance; }
 
-	float NIAGARA_API GetTotalCPUTime();
+	float NIAGARA_API GetTotalCPUTimeMS();
 	int	NIAGARA_API GetTotalBytesUsed();
 	
 	ENiagaraExecutionState NIAGARA_API GetExecutionState()	{ return ExecutionState; }
@@ -140,8 +140,8 @@ private:
 	/** Typical resets must be deferred until the tick as the RT could still be using the current buffer. */
 	uint32 bResetPending : 1;
 
-	/* Seconds taken to process everything (including rendering) */
-	float CPUTimeMS;
+	/* Cycles taken to process the tick. */
+	uint32 CPUTimeCycles;
 	/* Emitter tick state */
 	ENiagaraExecutionState ExecutionState;
 	/* Emitter bounds */
diff --git a/Engine/Plugins/FX/Niagara/Source/Niagara/Private/NiagaraComponent.cpp b/Engine/Plugins/FX/Niagara/Source/Niagara/Private/NiagaraComponent.cpp
--- a/Engine/Plugins/FX/Niagara/Source/Niagara/Private/NiagaraComponent.cpp
+++ b/Engine/Plugins/FX/Niagara/Source/Niagara/Private/NiagaraComponent.cpp
@@ -93,7 +93,7 @@ void DumpNiagaraComponents(UWorld* World)
 						for (TSharedRef<FNiagaraEmitterInstance, ESPMode::ThreadSafe> Emitter : SysInst->GetEmitters())
 						{
 							UE_LOG(LogNiagara, Log, TEXT("    Emitter: \"%s\" | ExecState: %d | NumParticles: %d | CPUTime: %f"), *Emitter->GetEmitterHandle().GetUniqueInstanceName(),
-								(int32)Emitter->GetExecutionState(), Emitter->GetNumParticles(), Emitter->GetTotalCPUTime());
+								(int32)Emitter->GetExecutionState(), Emitter->GetNumParticles(), Emitter->GetTotalCPUTimeMS());
 						}
 					}
 				}
diff --git a/Engine/Plugins/FX/Niagara/Source/Niagara/Private/NiagaraEmitterInstance.cpp b/Engine/Plugins/FX/Niagara/Source/Niagara/Private/NiagaraEmitterInstance.cpp
--- a/Engine/Plugins/FX/Niagara/Source/Niagara/Private/NiagaraEmitterInstance.cpp
+++ b/Engine/Plugins/FX/Niagara/Source/Niagara/Private/NiagaraEmitterInstance.cpp
@@ -62,8 +62,40 @@ static FAutoConsoleVariableRef CVarMaxNiagaraCPUParticlesPerEmitter(
 );
 //////////////////////////////////////////////////////////////////////////
 
+template<bool bAccumulate>
+struct FNiagaraEditorOnlyCycleTimer
+{
+	FORCEINLINE FNiagaraEditorOnlyCycleTimer(uint32& InCyclesOut)
+#if WITH_EDITOR
+		: CyclesOut(InCyclesOut)
+		, StartCycles(FPlatformTime::Cycles())
+#endif
+	{
+	}
+
+#if WITH_EDITOR
+	FORCEINLINE ~FNiagaraEditorOnlyCycleTimer()
+	{
+		uint32 DeltaCycles = FPlatformTime::Cycles() - StartCycles;
+		if (bAccumulate)
+		{
+			CyclesOut += DeltaCycles;
+		}
+		else
+		{
+			CyclesOut = DeltaCycles;
+		}
+	}
+
+	uint32& CyclesOut;
+	uint32 StartCycles;
+#endif
+};
+
+//////////////////////////////////////////////////////////////////////////
+
 FNiagaraEmitterInstance::FNiagaraEmitterInstance(FNiagaraSystemInstance* InParentSystemInstance)
-: CPUTimeMS(0.0f)
+: CPUTimeCycles(0)
 , ExecutionState(ENiagaraExecutionState::Inactive)
 , CachedBounds(ForceInit)
 , GPUExecContext(nullptr)
@@ -659,9 +691,9 @@ const FNiagaraEmitterHandle& FNiagaraEmitterInstance::GetEmitterHandle() const
 	return Sys->GetEmitterHandles()[EmitterIdx];
 }
 
-float FNiagaraEmitterInstance::GetTotalCPUTime()
+float FNiagaraEmitterInstance::GetTotalCPUTimeMS()
 {
-	float Total = CPUTimeMS;
+	uint32 TotalCycles = CPUTimeCycles;
 
 	//TODO: Find some way to include the RT cost here?
 	//Possibly have the proxy write back it's most recent frame time during EOF updates?
@@ -673,7 +705,7 @@ float FNiagaraEmitterInstance::GetTotalCPUTime()
 // 		}
 // 	}
 
-	return Total;
+	return FPlatformTime::ToMilliseconds(TotalCycles);
 }
 
 int FNiagaraEmitterInstance::GetTotalBytesUsed()
@@ -963,7 +995,7 @@ void FNiagaraEmitterInstance::SetSystemFixedBoundsOverride(FBox SystemFixedBound
 void FNiagaraEmitterInstance::Tick(float DeltaSeconds)
 {
 	SCOPE_CYCLE_COUNTER(STAT_NiagaraTick);
-	SimpleTimer TickTime;
+	FNiagaraEditorOnlyCycleTimer<false> TickTime(CPUTimeCycles);
 
 #if STATS
 	FScopeCycleCounter SystemStatCounter(CachedEmitter->GetStatID(true, true));
@@ -971,7 +1003,6 @@ void FNiagaraEmitterInstance::Tick(float DeltaSeconds)
 
 	if (HandleCompletion())
 	{
-		CPUTimeMS = TickTime.GetElapsedMilliseconds();
 		return;
 	}
 
@@ -985,14 +1016,12 @@ void FNiagaraEmitterInstance::Tick(float DeltaSeconds)
 	{
 		Data.ResetBuffers();
 		ExecutionState = ENiagaraExecutionState::Inactive;
-		CPUTimeMS = TickTime.GetElapsedMilliseconds();
 		return;
 	}
 
 	if (CachedEmitter->SimTarget == ENiagaraSimTarget::CPUSim && Data.GetCurrentDataChecked().GetNumInstances() == 0 && ExecutionState != ENiagaraExecutionState::Active)
 	{
 		Data.ResetBuffers();
-		CPUTimeMS = TickTime.GetElapsedMilliseconds();
 		return;
 	}
 
@@ -1155,8 +1184,6 @@ void FNiagaraEmitterInstance::Tick(float DeltaSeconds)
 
 		CachedBounds = CachedEmitter->FixedBounds;
 
-		CPUTimeMS = TickTime.GetElapsedMilliseconds();
-
 		/*if (CachedEmitter->SpawnScriptProps.Script->GetComputedVMCompilationId().HasInterpolatedParameters())
 		{
 			GPUExecContext.CombinedParamStore.CopyCurrToPrev();
@@ -1579,8 +1606,6 @@ void FNiagaraEmitterInstance::Tick(float DeltaSeconds)
 		EventContext.PostTick();
 	}
 
-	CPUTimeMS = TickTime.GetElapsedMilliseconds();
-
 	if (GbDumpParticleData || System->bDumpDebugEmitterInfo)
 	{
 		UE_LOG(LogNiagara, Log, TEXT("|=== END OF FNiagaraEmitterInstance::Tick [ %s ] ===============|"), *CachedEmitter->GetPathName());
diff --git a/Engine/Plugins/FX/Niagara/Source/Niagara/Private/NiagaraRenderer.cpp b/Engine/Plugins/FX/Niagara/Source/Niagara/Private/NiagaraRenderer.cpp
--- a/Engine/Plugins/FX/Niagara/Source/Niagara/Private/NiagaraRenderer.cpp
+++ b/Engine/Plugins/FX/Niagara/Source/Niagara/Private/NiagaraRenderer.cpp
@@ -192,7 +192,6 @@ FNiagaraDataBuffer* FNiagaraDynamicDataBase::GetParticleDataToRender()const
 
 FNiagaraRenderer::FNiagaraRenderer(ERHIFeatureLevel::Type FeatureLevel, const UNiagaraRendererProperties *InProps, const FNiagaraEmitterInstance* Emitter)
 	: DynamicDataRender(nullptr)
-	, CPUTimeMS(0.0f)
 	, bLocalSpace(Emitter->GetCachedEmitter()->bLocalSpace)
 	, bHasLights(false)
 	, SimTarget(Emitter->GetCachedEmitter()->SimTarget)
diff --git a/Engine/Plugins/FX/Niagara/Source/Niagara/Private/NiagaraRendererLights.cpp b/Engine/Plugins/FX/Niagara/Source/Niagara/Private/NiagaraRendererLights.cpp
--- a/Engine/Plugins/FX/Niagara/Source/Niagara/Private/NiagaraRendererLights.cpp
+++ b/Engine/Plugins/FX/Niagara/Source/Niagara/Private/NiagaraRendererLights.cpp
@@ -48,8 +48,6 @@ FNiagaraDynamicDataBase* FNiagaraRendererLights::GenerateDynamicData(const FNiag
 {
 	SCOPE_CYCLE_COUNTER(STAT_NiagaraGenLights);
 
-	SimpleTimer VertexDataTimer;
-
 	//Bail if we don't have the required attributes to render this emitter.
 	const UNiagaraLightRendererProperties* Properties = CastChecked<const UNiagaraLightRendererProperties>(InProperties);
 	FNiagaraDataSet& Data = Emitter->GetData();
@@ -97,7 +95,6 @@ FNiagaraDynamicDataBase* FNiagaraRendererLights::GenerateDynamicData(const FNiag
 		}
 	}
 
-	CPUTimeMS = VertexDataTimer.GetElapsedMilliseconds();
 	return DynamicData;
 }
 
diff --git a/Engine/Plugins/FX/Niagara/Source/Niagara/Private/NiagaraRendererMeshes.cpp b/Engine/Plugins/FX/Niagara/Source/Niagara/Private/NiagaraRendererMeshes.cpp
--- a/Engine/Plugins/FX/Niagara/Source/Niagara/Private/NiagaraRendererMeshes.cpp
+++ b/Engine/Plugins/FX/Niagara/Source/Niagara/Private/NiagaraRendererMeshes.cpp
@@ -137,8 +137,6 @@ void FNiagaraRendererMeshes::GetDynamicMeshElements(const TArray<const FSceneVie
 	SCOPE_CYCLE_COUNTER(STAT_NiagaraRenderMeshes);
 	check(SceneProxy);
 
-	SimpleTimer MeshElementsTimer;
-
 	NiagaraEmitterInstanceBatcher* Batcher = SceneProxy->GetBatcher();
 	FNiagaraDynamicDataMesh *DynamicDataMesh = (static_cast<FNiagaraDynamicDataMesh*>(DynamicDataRender));
 	if (!DynamicDataMesh || !Batcher)
@@ -404,8 +402,6 @@ void FNiagaraRendererMeshes::GetDynamicMeshElements(const TArray<const FSceneVie
 			}
 		}
 	}
-
-	CPUTimeMS += MeshElementsTimer.GetElapsedMilliseconds();
 }
 
 FNiagaraDynamicDataBase *FNiagaraRendererMeshes::GenerateDynamicData(const FNiagaraSceneProxy* Proxy, const UNiagaraRendererProperties* InProperties, const FNiagaraEmitterInstance* Emitter) const
@@ -420,8 +416,6 @@ FNiagaraDynamicDataBase *FNiagaraRendererMeshes::GenerateDynamicData(const FNiag
 		return nullptr;
 	}
 
-	SimpleTimer VertexDataTimer;
-
 	FNiagaraDynamicDataMesh *DynamicData = nullptr;
 
 	if (Properties->ParticleMesh)
@@ -454,7 +448,6 @@ FNiagaraDynamicDataBase *FNiagaraRendererMeshes::GenerateDynamicData(const FNiag
 		}
 	}
 
-	CPUTimeMS = VertexDataTimer.GetElapsedMilliseconds();
 	return DynamicData;  
 }
 
diff --git a/Engine/Plugins/FX/Niagara/Source/Niagara/Private/NiagaraRendererRibbons.cpp b/Engine/Plugins/FX/Niagara/Source/Niagara/Private/NiagaraRendererRibbons.cpp
--- a/Engine/Plugins/FX/Niagara/Source/Niagara/Private/NiagaraRendererRibbons.cpp
+++ b/Engine/Plugins/FX/Niagara/Source/Niagara/Private/NiagaraRendererRibbons.cpp
@@ -278,7 +278,6 @@ void FNiagaraRendererRibbons::GetDynamicMeshElements(const TArray<const FSceneVi
 	SCOPE_CYCLE_COUNTER(STAT_NiagaraRender);
 	SCOPE_CYCLE_COUNTER(STAT_NiagaraRenderRibbons);
 
-	SimpleTimer MeshElementsTimer;
 	FNiagaraDynamicDataRibbon *DynamicDataRibbon = static_cast<FNiagaraDynamicDataRibbon*>(DynamicDataRender);
 	if (!DynamicDataRibbon)
 	{
@@ -325,8 +324,6 @@ void FNiagaraRendererRibbons::GetDynamicMeshElements(const TArray<const FSceneVi
 			Collector.AddMesh(ViewIndex, MeshBatch);
 		}
 	}
-
-	CPUTimeMS += MeshElementsTimer.GetElapsedMilliseconds();
 }
 
 int FNiagaraRendererRibbons::GetDynamicDataSize()const
@@ -408,7 +405,6 @@ FNiagaraDynamicDataBase* FNiagaraRendererRibbons::GenerateDynamicData(const FNia
 		return nullptr;
 	}
 
-	SimpleTimer VertexDataTimer;
 	FNiagaraDataSet& Data = Emitter->GetData();
 	const UNiagaraRibbonRendererProperties* Properties = CastChecked<const UNiagaraRibbonRendererProperties>(InProperties);
 
@@ -714,8 +710,6 @@ FNiagaraDynamicDataBase* FNiagaraRendererRibbons::GenerateDynamicData(const FNia
 		TessellationTotalSegmentLength = 0;
 	}
 
-	CPUTimeMS = VertexDataTimer.GetElapsedMilliseconds();
-
 	return DynamicData;
 }
 
diff --git a/Engine/Plugins/FX/Niagara/Source/Niagara/Private/NiagaraRendererSprites.cpp b/Engine/Plugins/FX/Niagara/Source/Niagara/Private/NiagaraRendererSprites.cpp
--- a/Engine/Plugins/FX/Niagara/Source/Niagara/Private/NiagaraRendererSprites.cpp
+++ b/Engine/Plugins/FX/Niagara/Source/Niagara/Private/NiagaraRendererSprites.cpp
@@ -520,8 +520,6 @@ void FNiagaraRendererSprites::GetDynamicMeshElements(const TArray<const FSceneVi
 	SCOPE_CYCLE_COUNTER(STAT_NiagaraRenderSprites);
 	check(SceneProxy);
 
-	SimpleTimer MeshElementsTimer;
-
 	//check(DynamicDataRender)
 	FNiagaraDynamicDataSprites *DynamicDataSprites = static_cast<FNiagaraDynamicDataSprites*>(DynamicDataRender);
 	NiagaraEmitterInstanceBatcher* Batcher = SceneProxy->GetBatcher();
@@ -566,8 +564,6 @@ void FNiagaraRendererSprites::GetDynamicMeshElements(const TArray<const FSceneVi
 			Collector.AddMesh(ViewIndex, MeshBatch);
 		}
 	}
-
-	CPUTimeMS += MeshElementsTimer.GetElapsedMilliseconds();
 }
 
 #if RHI_RAYTRACING
@@ -577,8 +573,6 @@ void FNiagaraRendererSprites::GetDynamicRayTracingInstances(FRayTracingMaterialG
 	SCOPE_CYCLE_COUNTER(STAT_NiagaraRenderSprites);
 	check(SceneProxy);
 
-	SimpleTimer MeshElementsTimer;
-
 	FNiagaraDynamicDataSprites *DynamicDataSprites = static_cast<FNiagaraDynamicDataSprites*>(DynamicDataRender);
 	NiagaraEmitterInstanceBatcher* Batcher = SceneProxy->GetBatcher();
 	if (!DynamicDataSprites || !Batcher)
@@ -638,8 +632,6 @@ void FNiagaraRendererSprites::GetDynamicRayTracingInstances(FRayTracingMaterialG
 	RayTracingInstance.BuildInstanceMaskAndFlags();
 
 	OutRayTracingInstances.Add(RayTracingInstance);
-
-	CPUTimeMS += MeshElementsTimer.GetElapsedMilliseconds();
 }
 #endif
 
@@ -651,8 +643,6 @@ FNiagaraDynamicDataBase *FNiagaraRendererSprites::GenerateDynamicData(const FNia
 
 	if (Properties)
 	{
-		SimpleTimer VertexDataTimer;
-
 		SCOPE_CYCLE_COUNTER(STAT_NiagaraGenSpriteDynamicData);
 
 		FNiagaraDataSet& Data = Emitter->GetData();
@@ -669,8 +659,6 @@ FNiagaraDynamicDataBase *FNiagaraRendererSprites::GenerateDynamicData(const FNia
 			DynamicData->Material = BaseMaterials_GT[0]->GetRenderProxy();
 			DynamicData->SetMaterialRelevance(BaseMaterialRelevance_GT);
 		}
-
-		CPUTimeMS = VertexDataTimer.GetElapsedMilliseconds();
 	}
 
 	return DynamicData;  // for VF that can fetch from particle data directly
diff --git a/Engine/Plugins/FX/Niagara/Source/Niagara/Public/NiagaraRenderer.h b/Engine/Plugins/FX/Niagara/Source/Niagara/Public/NiagaraRenderer.h
--- a/Engine/Plugins/FX/Niagara/Source/Niagara/Public/NiagaraRenderer.h
+++ b/Engine/Plugins/FX/Niagara/Source/Niagara/Public/NiagaraRenderer.h
@@ -51,29 +51,6 @@ protected:
 	}Data;
 };
 
-
-
-class SimpleTimer
-{
-public:
-	SimpleTimer()
-	{
-		StartTime = FPlatformTime::Seconds() * 1000.0;
-	}
-
-	double GetElapsedMilliseconds()
-	{
-		return (FPlatformTime::Seconds()*1000.0) - StartTime;
-	}
-
-	~SimpleTimer()
-	{
-	}
-
-private:
-	double StartTime;
-};
-
 //////////////////////////////////////////////////////////////////////////
 
 /**
@@ -106,7 +83,6 @@ public:
 	void SetDynamicData_RenderThread(FNiagaraDynamicDataBase* NewDynamicData);
 	FORCEINLINE FNiagaraDynamicDataBase *GetDynamicData()const { return DynamicDataRender; }
 	FORCEINLINE bool HasDynamicData()const { return DynamicDataRender != nullptr; }
-	FORCEINLINE float GetCPUTimeMS() const { return CPUTimeMS; }
 	FORCEINLINE bool HasLights()const { return bHasLights; }
 
 #if RHI_RAYTRACING
@@ -129,8 +105,6 @@ protected:
 	FRayTracingGeometry RayTracingGeometry;
 #endif
 
-	mutable float CPUTimeMS;
-
 	uint32 bLocalSpace : 1;
 	uint32 bHasLights : 1;
 	ENiagaraSimTarget SimTarget;
diff --git a/Engine/Plugins/FX/Niagara/Source/NiagaraEditor/Private/ViewModels/NiagaraEmitterViewModel.cpp b/Engine/Plugins/FX/Niagara/Source/NiagaraEditor/Private/ViewModels/NiagaraEmitterViewModel.cpp
--- a/Engine/Plugins/FX/Niagara/Source/NiagaraEditor/Private/ViewModels/NiagaraEmitterViewModel.cpp
+++ b/Engine/Plugins/FX/Niagara/Source/NiagaraEditor/Private/ViewModels/NiagaraEmitterViewModel.cpp
@@ -215,7 +215,7 @@ FText FNiagaraEmitterViewModel::GetStatsText() const
 				{
 					return FText::Format(StatsFormat,
 						FText::AsNumber(SimInstance->GetNumParticles()),
-						FText::AsNumber(SimInstance->GetTotalCPUTime(), &FractionalFormatOptions),
+						FText::AsNumber(SimInstance->GetTotalCPUTimeMS(), &FractionalFormatOptions),
 						FText::AsNumber(SimInstance->GetTotalBytesUsed() / Megabyte, &FractionalFormatOptions),
 						ExecutionStateEnum->GetDisplayNameTextByValue((int32)SimInstance->GetExecutionState()));
 				}
