- Build fix: Remove usage of SimpleTimer from new code inside UE4Main

diff --git a/Engine/Plugins/FX/Niagara/Source/Niagara/Private/NiagaraRendererRibbons.cpp b/Engine/Plugins/FX/Niagara/Source/Niagara/Private/NiagaraRendererRibbons.cpp
--- a/Engine/Plugins/FX/Niagara/Source/Niagara/Private/NiagaraRendererRibbons.cpp
+++ b/Engine/Plugins/FX/Niagara/Source/Niagara/Private/NiagaraRendererRibbons.cpp
@@ -981,8 +981,6 @@ void FNiagaraRendererRibbons::GetDynamicRayTracingInstances(FRayTracingMaterialG
 	SCOPE_CYCLE_COUNTER(STAT_NiagaraRenderRibbons);
 	check(SceneProxy);
 
-	SimpleTimer MeshElementsTimer;
-
 	FNiagaraDynamicDataRibbon *DynamicDataRibbon = static_cast<FNiagaraDynamicDataRibbon*>(DynamicDataRender);
 	NiagaraEmitterInstanceBatcher* Batcher = SceneProxy->GetBatcher();
 	if (!DynamicDataRibbon || !Batcher)
@@ -1052,7 +1050,5 @@ void FNiagaraRendererRibbons::GetDynamicRayTracingInstances(FRayTracingMaterialG
 	RayTracingInstance.BuildInstanceMaskAndFlags();
 
 	OutRayTracingInstances.Add(RayTracingInstance);
-
-	CPUTimeMS += MeshElementsTimer.GetElapsedMilliseconds();
 }
 #endif
\ No newline at end of file
