- Make sure we unbind parameters when removing the system instance to fix a crash in GC due to override parameters still being bound on a solo emitter

#rnx

diff --git a/Engine/Plugins/FX/Niagara/Source/Niagara/Private/NiagaraSystemInstance.cpp b/Engine/Plugins/FX/Niagara/Source/Niagara/Private/NiagaraSystemInstance.cpp
--- a/Engine/Plugins/FX/Niagara/Source/Niagara/Private/NiagaraSystemInstance.cpp
+++ b/Engine/Plugins/FX/Niagara/Source/Niagara/Private/NiagaraSystemInstance.cpp
@@ -474,6 +474,7 @@ bool FNiagaraSystemInstance::DeallocateSystemInstance(TUniquePtr< FNiagaraSystem
 		// Make sure we remove the instance
 		if (SystemInstanceAllocation->SystemInstanceIndex != INDEX_NONE)
 		{
+			SystemInstanceAllocation->UnbindParameters();
 			SystemSim->RemoveInstance(SystemInstanceAllocation.Get());
 		}
 
