Disabling the function call pruning for now as it isn't threadsafe

diff --git a/Engine/Plugins/FX/Niagara/Source/NiagaraEditor/Private/NiagaraNodeFunctionCall.cpp b/Engine/Plugins/FX/Niagara/Source/NiagaraEditor/Private/NiagaraNodeFunctionCall.cpp
--- a/Engine/Plugins/FX/Niagara/Source/NiagaraEditor/Private/NiagaraNodeFunctionCall.cpp
+++ b/Engine/Plugins/FX/Niagara/Source/NiagaraEditor/Private/NiagaraNodeFunctionCall.cpp
@@ -459,7 +459,9 @@ void UNiagaraNodeFunctionCall::Compile(class FHlslNiagaraTranslator* Translator,
 		TArray<const UEdGraphPin*> OutInputPins;
 		TSet<const UEdGraphPin*> HiddenPins;
 		TSet<FName> HiddenPinNames;
-		FNiagaraStackGraphUtilities::GetStackFunctionInputPins(*this, OutInputPins, HiddenPins, FCompileConstantResolver(Translator));
+		//FNiagaraStackGraphUtilities::GetStackFunctionInputPins(*this, OutInputPins, HiddenPins, FCompileConstantResolver(Translator));
+		// The above line of code isn't threadsafe as it builds a parameter map history live. This does things like load the asset registry module
+		// and  there are some static finds as well. We need to re-implement that using straight graph traversal logic.
 		for (const UEdGraphPin* Pin : HiddenPins)
 		{
 			FString PinNameString = Pin->PinName.ToString();
