Fixing potential issue encountered by [REDACTED_PERSON] where we were trying to staticduplicate data interfaces during postload. Had to add copy to render thread to Vector fields on post CopyTo.

#tests engine tests ran correctly.

diff --git a/Engine/Plugins/FX/Niagara/Source/Niagara/Classes/NiagaraScript.h b/Engine/Plugins/FX/Niagara/Source/Niagara/Classes/NiagaraScript.h
--- a/Engine/Plugins/FX/Niagara/Source/Niagara/Classes/NiagaraScript.h
+++ b/Engine/Plugins/FX/Niagara/Source/Niagara/Classes/NiagaraScript.h
@@ -640,4 +640,6 @@ private:
 
 	UPROPERTY()
 	TArray<FNiagaraScriptDataInterfaceInfo> CachedDefaultDataInterfaces;
+
+	static UNiagaraDataInterface* CopyDataInterface(UNiagaraDataInterface* Src, UObject* Owner);
 };
diff --git a/Engine/Plugins/FX/Niagara/Source/Niagara/Private/NiagaraScript.cpp b/Engine/Plugins/FX/Niagara/Source/Niagara/Private/NiagaraScript.cpp
--- a/Engine/Plugins/FX/Niagara/Source/Niagara/Private/NiagaraScript.cpp
+++ b/Engine/Plugins/FX/Niagara/Source/Niagara/Private/NiagaraScript.cpp
@@ -814,6 +814,17 @@ void WriteTextFileToDisk(FString SaveDirectory, FString FileName, FString TextTo
 	}
 }
 
+UNiagaraDataInterface* UNiagaraScript::CopyDataInterface(UNiagaraDataInterface* Src, UObject* Owner)
+{
+	if (Src)
+	{
+		UNiagaraDataInterface* DI = NewObject<UNiagaraDataInterface>(Owner, const_cast<UClass*>(Src->GetClass()), NAME_None, RF_Transactional | RF_Public);
+		Src->CopyTo(DI);
+		return DI;
+	}
+	return nullptr;
+}
+
 void UNiagaraScript::SetVMCompilationResults(const FNiagaraVMExecutableDataId& InCompileId, FNiagaraVMExecutableData& InScriptVM, FNiagaraCompileRequestDataBase* InRequestData)
 {
 	check(InRequestData != nullptr);
@@ -864,7 +875,7 @@ void UNiagaraScript::SetVMCompilationResults(const FNiagaraVMExecutableDataId& I
 		UNiagaraDataInterface*const* FindDIById = InRequestData->GetObjectNameMap().Find(Info.Name);
 		if (FindDIById != nullptr && *(FindDIById) != nullptr)
 		{
-			CachedDefaultDataInterfaces[Idx].DataInterface = DuplicateObject<UNiagaraDataInterface>(*(FindDIById), this);
+			CachedDefaultDataInterfaces[Idx].DataInterface = CopyDataInterface(*(FindDIById), this);
 			check(CachedDefaultDataInterfaces[Idx].DataInterface != nullptr);
 		}			
 		
@@ -872,7 +883,7 @@ void UNiagaraScript::SetVMCompilationResults(const FNiagaraVMExecutableDataId& I
 		{
 			// Use the CDO since we didn't have a default..
 			UObject* Obj = const_cast<UClass*>(Info.Type.GetClass())->GetDefaultObject(true);
-			CachedDefaultDataInterfaces[Idx].DataInterface = Cast<UNiagaraDataInterface>(DuplicateObject(Obj, this));
+			CachedDefaultDataInterfaces[Idx].DataInterface = Cast<UNiagaraDataInterface>(CopyDataInterface(CastChecked<UNiagaraDataInterface>(Obj), this));
 
 			if (Info.bIsPlaceholder == false)
 			{
@@ -1294,7 +1305,7 @@ bool UNiagaraScript::SynchronizeExecutablesWithMaster(const UNiagaraScript* Scri
 		{
 			FNiagaraScriptDataInterfaceInfo AddInfo;
 			AddInfo = Info;
-			AddInfo.DataInterface = DuplicateObject<UNiagaraDataInterface>(Info.DataInterface, this);
+			AddInfo.DataInterface = CopyDataInterface(Info.DataInterface, this);
 			CachedDefaultDataInterfaces.Add(AddInfo);
 		}
 
