- Fix using PropertyThatWillChange during undo where it can be nullptr

#rnx

diff --git a/Engine/Plugins/FX/Niagara/Source/Niagara/Private/NiagaraMeshRendererProperties.cpp b/Engine/Plugins/FX/Niagara/Source/Niagara/Private/NiagaraMeshRendererProperties.cpp
--- a/Engine/Plugins/FX/Niagara/Source/Niagara/Private/NiagaraMeshRendererProperties.cpp
+++ b/Engine/Plugins/FX/Niagara/Source/Niagara/Private/NiagaraMeshRendererProperties.cpp
@@ -234,7 +234,7 @@ void UNiagaraMeshRendererProperties::PreEditChange(class UProperty* PropertyThat
 	Super::PreEditChange(PropertyThatWillChange);
 
 	static FName ParticleMeshName(TEXT("ParticleMesh"));
-	if ( PropertyThatWillChange->GetFName() == FName(ParticleMeshName))
+	if ((PropertyThatWillChange != nullptr) && (PropertyThatWillChange->GetFName() == FName(ParticleMeshName)))
 	{
 		if (ParticleMesh != nullptr)
 		{
diff --git a/Engine/Source/Runtime/Engine/Private/Particles/ParticleModules.cpp b/Engine/Source/Runtime/Engine/Private/Particles/ParticleModules.cpp
--- a/Engine/Source/Runtime/Engine/Private/Particles/ParticleModules.cpp
+++ b/Engine/Source/Runtime/Engine/Private/Particles/ParticleModules.cpp
@@ -3481,7 +3481,7 @@ void UParticleModuleTypeDataMesh::OnMeshChanged()
 
 void UParticleModuleTypeDataMesh::PreEditChange(UProperty* PropertyThatWillChange)
 {
-	if ( PropertyThatWillChange->GetFName() == FName(TEXT("Mesh")) )
+	if ( (PropertyThatWillChange != nullptr) && (PropertyThatWillChange->GetFName() == FName(TEXT("Mesh"))) )
 	{
 		if ( GIsEditor && (Mesh != nullptr) )
 		{
