Fixing nightly build issue with vector fields after previous change.

diff --git a/Engine/Plugins/FX/Niagara/Source/Niagara/Private/NiagaraDataInterfaceVectorField.cpp b/Engine/Plugins/FX/Niagara/Source/Niagara/Private/NiagaraDataInterfaceVectorField.cpp
--- a/Engine/Plugins/FX/Niagara/Source/Niagara/Private/NiagaraDataInterfaceVectorField.cpp
+++ b/Engine/Plugins/FX/Niagara/Source/Niagara/Private/NiagaraDataInterfaceVectorField.cpp
@@ -665,6 +665,8 @@ bool UNiagaraDataInterfaceVectorField::CopyToInternal(UNiagaraDataInterface* Des
 	OtherTyped->bTileX = bTileX;
 	OtherTyped->bTileY = bTileY;
 	OtherTyped->bTileZ = bTileZ;
+
+	OtherTyped->PushToRenderThread();
 	return true;
 }
 
