- Fix for PVS warnings about View pointer, we don't need to check against null as it is guaranteed to be valid

#rnx

diff --git a/Engine/Plugins/FX/Niagara/Source/Niagara/Private/NiagaraRendererMeshes.cpp b/Engine/Plugins/FX/Niagara/Source/Niagara/Private/NiagaraRendererMeshes.cpp
--- a/Engine/Plugins/FX/Niagara/Source/Niagara/Private/NiagaraRendererMeshes.cpp
+++ b/Engine/Plugins/FX/Niagara/Source/Niagara/Private/NiagaraRendererMeshes.cpp
@@ -234,7 +234,7 @@ void FNiagaraRendererMeshes::GetDynamicMeshElements(const TArray<const FSceneVie
 				CollectorResources.VertexFactory.SetSortedIndices(nullptr, 0xFFFFFFFF);
 
 				FNiagaraGPUSortInfo SortInfo;
-				if (View && SortMode != ENiagaraSortMode::None && (bHasTranslucentMaterials || !bSortOnlyWhenTranslucent))
+				if (SortMode != ENiagaraSortMode::None && (bHasTranslucentMaterials || !bSortOnlyWhenTranslucent))
 				{
 					SortInfo.ParticleCount = NumInstances;
 					SortInfo.SortMode = SortMode;
diff --git a/Engine/Plugins/FX/Niagara/Source/Niagara/Private/NiagaraRendererSprites.cpp b/Engine/Plugins/FX/Niagara/Source/Niagara/Private/NiagaraRendererSprites.cpp
--- a/Engine/Plugins/FX/Niagara/Source/Niagara/Private/NiagaraRendererSprites.cpp
+++ b/Engine/Plugins/FX/Niagara/Source/Niagara/Private/NiagaraRendererSprites.cpp
@@ -321,7 +321,7 @@ void FNiagaraRendererSprites::SetVertexFactoryParticleData(
 
 		int32 NumInstances = SourceParticleData->GetNumInstances();
 		FNiagaraGPUSortInfo SortInfo;
-		if (View && SortMode != ENiagaraSortMode::None && (BlendMode == BLEND_AlphaComposite || BlendMode == BLEND_AlphaHoldout|| BlendMode == BLEND_Translucent || !bSortOnlyWhenTranslucent))
+		if (SortMode != ENiagaraSortMode::None && (BlendMode == BLEND_AlphaComposite || BlendMode == BLEND_AlphaHoldout|| BlendMode == BLEND_Translucent || !bSortOnlyWhenTranslucent))
 		{
 			SortInfo.ParticleCount = NumInstances;
 			SortInfo.SortMode = SortMode;
