- Update bounds during concurrent tick and only call UpdateComponentToWorld on the GT in finalize if they have changed
- This follows was Cascade does and saves ~30us on the GT finalize when attached, it also saves updating the octree on the RT

#rnx

diff --git a/Engine/Plugins/FX/Niagara/Source/Niagara/Private/NiagaraComponent.cpp b/Engine/Plugins/FX/Niagara/Source/Niagara/Private/NiagaraComponent.cpp
--- a/Engine/Plugins/FX/Niagara/Source/Niagara/Private/NiagaraComponent.cpp
+++ b/Engine/Plugins/FX/Niagara/Source/Niagara/Private/NiagaraComponent.cpp
@@ -372,6 +372,7 @@ UNiagaraComponent::UNiagaraComponent(const FObjectInitializer& ObjectInitializer
 	, MaxSimTime(33.0f / 1000.0f)
 	, bIsSeeking(false)
 	, bAutoDestroy(false)
+	, MaxTimeBeforeForceUpdateTransform(5.0f)
 #if WITH_EDITOR
 	, PreviewDetailLevel(INDEX_NONE)
 	, PreviewLODDistance(0.0f)
@@ -1119,23 +1120,7 @@ FBoxSphereBounds UNiagaraComponent::CalcBounds(const FTransform& LocalToWorld) c
 	FBoxSphereBounds SystemBounds;
 	if (SystemInstance.IsValid())
 	{
-		UNiagaraSystem* System = SystemInstance->GetSystem();
-		if (System->bFixedBounds)
-		{
-			SystemBounds = System->GetFixedBounds();
-		}
-		else
-		{
-			SystemInstance->GetSystemBounds().Init();
-			for (int32 i = 0; i < SystemInstance->GetEmitters().Num(); i++)
-			{
-				FNiagaraEmitterInstance &Sim = *(SystemInstance->GetEmitters()[i]);
-				SystemInstance->GetSystemBounds() += Sim.GetBounds();
-			}
-			FBox BoundingBox = SystemInstance->GetSystemBounds();
-
-			SystemBounds = FBoxSphereBounds(BoundingBox);
-		}
+		SystemBounds = SystemInstance->GetLocalBounds();
 	}
 	else
 	{
diff --git a/Engine/Plugins/FX/Niagara/Source/Niagara/Private/NiagaraSystemInstance.cpp b/Engine/Plugins/FX/Niagara/Source/Niagara/Private/NiagaraSystemInstance.cpp
--- a/Engine/Plugins/FX/Niagara/Source/Niagara/Private/NiagaraSystemInstance.cpp
+++ b/Engine/Plugins/FX/Niagara/Source/Niagara/Private/NiagaraSystemInstance.cpp
@@ -56,6 +56,14 @@ static FAutoConsoleVariableRef CVarNiagaraForceLastTickGroup(
 	ECVF_Default
 );
 
+static float GNiagaraBoundsExpandByPercent = 0.1f;
+static FAutoConsoleVariableRef CVarNiagaraBoundsExpandByPercent(
+	TEXT("fx.Niagara.BoundsExpandByPercent"),
+	GNiagaraBoundsExpandByPercent,
+	TEXT("The percentage we expand the bounds to avoid updating every frame."),
+	ECVF_Default
+);
+
 FNiagaraSystemInstance::FNiagaraSystemInstance(UNiagaraComponent* InComponent)
 	: SystemInstanceIndex(INDEX_NONE)
 	, Component(InComponent)
@@ -69,17 +77,18 @@ FNiagaraSystemInstance::FNiagaraSystemInstance(UNiagaraComponent* InComponent)
 	, bHasTickingEmitters(true)
 	, bPaused(false)
 	, bDataInterfacesHaveTickPrereqs(false)
-	, RequestedExecutionState(ENiagaraExecutionState::Complete)
-	, ActualExecutionState(ENiagaraExecutionState::Complete)
+	, bIsTransformDirty(true)
+	, bNeedsFinalize(false)
 	, bDataInterfacesInitialized(false)
 	, bAsyncWorkInProgress(false)
-	, bNeedsFinalize(false)
 	, CachedDeltaSeconds(0.0f)
+	, RequestedExecutionState(ENiagaraExecutionState::Complete)
+	, ActualExecutionState(ENiagaraExecutionState::Complete)
 {
 	static TAtomic<uint64> IDCounter(1);
 	ID = IDCounter.IncrementExchange();
 
-	SystemBounds.Init();
+	LocalBounds = FBox(FVector::ZeroVector, FVector::ZeroVector);
 
 	if (Component)
 	{
@@ -683,6 +692,7 @@ void FNiagaraSystemInstance::ResetInternal(bool bResetSimulations)
 	TickCount = 0;
 	bHasTickingEmitters = true;
 	CachedDeltaSeconds = 0.0f;
+	// Note: We do not need to update our bounds here as they are still valid
 
 	UNiagaraSystem* System = GetSystem();
 	if (System == nullptr || Component == nullptr || IsDisabled())
@@ -809,6 +819,9 @@ void FNiagaraSystemInstance::ReInitInternal()
 	Age = 0;
 	TickCount = 0;
 	bHasTickingEmitters = true;
+	bIsTransformDirty = true;
+	TimeSinceLastForceUpdateTransform = 0.0f;
+	LocalBounds = FBox(FVector::ZeroVector, FVector::ZeroVector);
 	CachedDeltaSeconds = 0.0f;
 
 	UNiagaraSystem* System = GetSystem();
@@ -1517,6 +1530,15 @@ void FNiagaraSystemInstance::InitEmitters()
 			Simulation->PostInitSimulation();
 		}
 	}
+
+	if ( System->bFixedBounds )
+	{
+		LocalBounds = System->GetFixedBounds();
+	}
+	else
+	{
+		LocalBounds = FBox(FVector::ZeroVector, FVector::ZeroVector);
+	}
 }
 
 void FNiagaraSystemInstance::ComponentTick(float DeltaSeconds, const FGraphEventRef& MyCompletionGraphEvent)
@@ -1654,6 +1676,35 @@ void FNiagaraSystemInstance::Tick_Concurrent()
 		}
 	}
 
+	// Update local bounds
+	if ( System->bFixedBounds )
+	{
+		LocalBounds = System->GetFixedBounds();
+	}
+	else
+	{
+		FBox NewLocalBounds(EForceInit::ForceInit);
+		for ( const auto& Emitter : Emitters )
+		{
+			NewLocalBounds += Emitter->GetBounds();
+		}
+
+		if ( NewLocalBounds.IsValid )
+		{
+			TimeSinceLastForceUpdateTransform += CachedDeltaSeconds;
+			if ((TimeSinceLastForceUpdateTransform > Component->MaxTimeBeforeForceUpdateTransform) || !LocalBounds.IsInsideOrOn(NewLocalBounds.Min) || !LocalBounds.IsInsideOrOn(NewLocalBounds.Max))
+			{
+				bIsTransformDirty = true;
+				LocalBounds = NewLocalBounds.ExpandBy(NewLocalBounds.GetExtent() * GNiagaraBoundsExpandByPercent);
+				TimeSinceLastForceUpdateTransform = 0.0f;
+			}
+		}
+		else
+		{
+			LocalBounds = FBox(FVector::ZeroVector, FVector::ZeroVector);
+		}
+	}
+
 	bAsyncWorkInProgress = false;
 }
 
@@ -1675,7 +1726,11 @@ void FNiagaraSystemInstance::FinalizeTick_GameThread()
 			bool bIsRegistered = true;
 			if (Component)
 			{
-				Component->UpdateComponentToWorld();//Needed for bounds updates. Can probably skip if using fixed bounds
+				if ( bIsTransformDirty )
+				{
+					bIsTransformDirty = false;
+					Component->UpdateComponentToWorld();
+				}
 				Component->MarkRenderDynamicDataDirty();
 				bIsRegistered = Component->IsRegistered();
 			}
diff --git a/Engine/Plugins/FX/Niagara/Source/Niagara/Public/NiagaraComponent.h b/Engine/Plugins/FX/Niagara/Source/Niagara/Public/NiagaraComponent.h
--- a/Engine/Plugins/FX/Niagara/Source/Niagara/Public/NiagaraComponent.h
+++ b/Engine/Plugins/FX/Niagara/Source/Niagara/Public/NiagaraComponent.h
@@ -126,6 +126,13 @@ public:
 	UPROPERTY(EditAnywhere, BlueprintReadOnly, Category = Attachment)
 	uint32 bAutoManageAttachment : 1;
 
+	/**
+	 * Time between forced UpdateTransforms for systems that use dynamically calculated bounds,
+	 * Which is effectively how often the bounds are shrunk.
+	 */
+	UPROPERTY()
+	float MaxTimeBeforeForceUpdateTransform;
+
 	/** How to handle pooling for this component instance. */
 	ENCPoolMethod PoolingMethod;
 
diff --git a/Engine/Plugins/FX/Niagara/Source/Niagara/Public/NiagaraSystemInstance.h b/Engine/Plugins/FX/Niagara/Source/Niagara/Public/NiagaraSystemInstance.h
--- a/Engine/Plugins/FX/Niagara/Source/Niagara/Public/NiagaraSystemInstance.h
+++ b/Engine/Plugins/FX/Niagara/Source/Niagara/Public/NiagaraSystemInstance.h
@@ -118,7 +118,7 @@ public:
 	UNiagaraSystem* GetSystem()const;
 	FORCEINLINE UNiagaraComponent *GetComponent() { return Component; }
 	FORCEINLINE TArray<TSharedRef<FNiagaraEmitterInstance, ESPMode::ThreadSafe> > &GetEmitters()	{ return Emitters; }
-	FORCEINLINE FBox &GetSystemBounds()	{ return SystemBounds;  }
+	FORCEINLINE const FBox& GetLocalBounds() { return LocalBounds;  }
 
 	FNiagaraEmitterInstance* GetEmitterByID(FGuid InID);
 
@@ -261,7 +261,6 @@ private:
 	TSharedPtr<class FNiagaraSystemSimulation, ESPMode::ThreadSafe> SystemSimulation;
 
 	UNiagaraComponent* Component;
-	FBox SystemBounds;
 
 	UActorComponent* PrereqComponent;
 
@@ -353,25 +352,37 @@ private:
 	/** The system contains data interfaces that can have tick group prerequisites. */
 	uint32 bDataInterfacesHaveTickPrereqs : 1;
 
+	/** True if our bounds have changed and we require pushing that to the rendering thread. */
+	uint32 bIsTransformDirty : 1;
+
+	/** True if we require a call to FinalizeTick_GameThread(). Typically this is called from a GT task but can be called in WaitForAsync. */
+	uint32 bNeedsFinalize : 1;
+
+	uint32 bDataInterfacesInitialized : 1;
+
+	/** True if we have async work in flight. */
+	volatile bool bAsyncWorkInProgress;
+
+	/** Cached delta time, written during Tick_GameThread and used during other phases. */
+	float CachedDeltaSeconds;
+
+	/** Time since we last forced a bounds update. */
+	float TimeSinceLastForceUpdateTransform;
+
+	/** Current calculated local bounds. */
+	FBox LocalBounds;
+
 	/* Execution state requested by external code/BPs calling Activate/Deactivate. */
 	ENiagaraExecutionState RequestedExecutionState;
 
 	/** Copy of simulations internal state so that it can be passed to emitters etc. */
 	ENiagaraExecutionState ActualExecutionState;
 
-	bool bDataInterfacesInitialized;
-
 	NiagaraEmitterInstanceBatcher* Batcher = nullptr;
-public:
-	/** True if we have async work in flight. */
-	volatile bool bAsyncWorkInProgress;
-	/** True if we require a call to FinalizeTick_GameThread(). Typically this is called from a GT task but can be called in WaitForAsync. */
-	volatile bool bNeedsFinalize;
 
+public:
 	// Transient data that is accumulated during tick.
 	uint32 TotalParamSize = 0;
 	uint32 ActiveGPUEmitterCount = 0;
 	int32 GPUDataInterfaceInstanceDataSize = 0;
-
-	float CachedDeltaSeconds;
 };
