Store all 2d grid data in 1 big texture rather than rgba channels

diff --git a/Engine/Plugins/FX/Niagara/Source/Niagara/Classes/NiagaraDataInterfaceGrid2DCollection.h b/Engine/Plugins/FX/Niagara/Source/Niagara/Classes/NiagaraDataInterfaceGrid2DCollection.h
--- a/Engine/Plugins/FX/Niagara/Source/Niagara/Classes/NiagaraDataInterfaceGrid2DCollection.h
+++ b/Engine/Plugins/FX/Niagara/Source/Niagara/Classes/NiagaraDataInterfaceGrid2DCollection.h
@@ -15,7 +15,7 @@ class Grid2DBuffer
 {
 public:
 	Grid2DBuffer(int NumX, int NumY)
-	{
+	{		
 		GridBuffer.Initialize(NumX, NumY);
 	}
 
@@ -28,6 +28,10 @@ public:
 
 	int NumCellsX;
 	int NumCellsY;
+	int NumAttributes;
+	int NumTilesX;
+	int NumTilesY;
+
 	FVector2D CellSize;
 	bool SetGridFromVoxelSize;
 
@@ -66,7 +70,7 @@ public:
 
 		if (DestinationData == nullptr)
 		{
-			DestinationData = new Grid2DBuffer(NumCellsX, NumCellsY);
+			DestinationData = new Grid2DBuffer(NumCellsX * NumTilesX, NumCellsY * NumTilesY);
 			Buffers.Add(DestinationData);
 		}
 
@@ -131,7 +135,7 @@ public:
 	// Fills a texture render target 2d with the current data from the simulation
 	// #todo(dmp): this will eventually go away when we formalize how data makes it out of Niagara
 	UFUNCTION(BlueprintCallable, Category = Niagara)
-	virtual void FillTexture2D(const UNiagaraComponent *Component, UTextureRenderTarget2D *dest);
+	virtual void FillTexture2D(const UNiagaraComponent *Component, UTextureRenderTarget2D *dest, int AttributeIndex);
 
 	void GetCellSize(FVectorVMContext& Context);
 
diff --git a/Engine/Plugins/FX/Niagara/Source/Niagara/Classes/NiagaraDataInterfaceRW.h b/Engine/Plugins/FX/Niagara/Source/Niagara/Classes/NiagaraDataInterfaceRW.h
--- a/Engine/Plugins/FX/Niagara/Source/Niagara/Classes/NiagaraDataInterfaceRW.h
+++ b/Engine/Plugins/FX/Niagara/Source/Niagara/Classes/NiagaraDataInterfaceRW.h
@@ -178,10 +178,13 @@ public:
 
 	UPROPERTY(EditAnywhere, Category = "Grid", meta = (EditCondition = "!SetGridFromCellSize"))
 	int32 NumCellsY;
-
+	
 	UPROPERTY(EditAnywhere, Category = "Grid", meta = (EditCondition = "SetGridFromCellSize"))
 	float CellSize;
 
+	UPROPERTY(EditAnywhere, Category = "Grid")
+	int32 NumAttributes;
+
 	UPROPERTY(EditAnywhere, Category = "Grid")
 	bool SetGridFromCellSize;
 
diff --git a/Engine/Plugins/FX/Niagara/Source/Niagara/Private/NiagaraDataInterfaceGrid2DCollection.cpp b/Engine/Plugins/FX/Niagara/Source/Niagara/Private/NiagaraDataInterfaceGrid2DCollection.cpp
--- a/Engine/Plugins/FX/Niagara/Source/Niagara/Private/NiagaraDataInterfaceGrid2DCollection.cpp
+++ b/Engine/Plugins/FX/Niagara/Source/Niagara/Private/NiagaraDataInterfaceGrid2DCollection.cpp
@@ -11,6 +11,8 @@
 
 #define LOCTEXT_NAMESPACE "NiagaraDataInterfaceGrid2DCollection"
 
+static const FString NumTilesName(TEXT("NumTiles_"));
+
 static const FString GridName(TEXT("Grid_"));
 static const FString OutputGridName(TEXT("OutputGrid_"));
 
@@ -28,6 +30,7 @@ struct FNiagaraDataInterfaceParametersCS_Grid2DCollection : public FNiagaraDataI
 	virtual void Bind(const FNiagaraDataInterfaceParamRef& ParamRef, const class FShaderParameterMap& ParameterMap) override
 	{			
 		NumCellsParam.Bind(ParameterMap, *(NumCellsName + ParamRef.ParameterInfo.DataInterfaceHLSLSymbol));
+		NumTilesParam.Bind(ParameterMap, *(NumTilesName + ParamRef.ParameterInfo.DataInterfaceHLSLSymbol));
 
 		CellSizeParam.Bind(ParameterMap, *(CellSizeName + ParamRef.ParameterInfo.DataInterfaceHLSLSymbol));
 
@@ -41,6 +44,7 @@ struct FNiagaraDataInterfaceParametersCS_Grid2DCollection : public FNiagaraDataI
 	virtual void Serialize(FArchive& Ar)override
 	{		
 		Ar << NumCellsParam;		
+		Ar << NumTilesParam;	
 		Ar << CellSizeParam;	
 		Ar << WorldBBoxMinParam;
 		Ar << WorldBBoxSizeParam;
@@ -63,8 +67,13 @@ struct FNiagaraDataInterfaceParametersCS_Grid2DCollection : public FNiagaraDataI
 		int NumCellsTmp[2];
 		NumCellsTmp[0] = ProxyData->NumCellsX;
 		NumCellsTmp[1] = ProxyData->NumCellsY;
-		SetShaderValue(RHICmdList, ComputeShaderRHI, NumCellsParam, NumCellsTmp);		
-		
+		SetShaderValue(RHICmdList, ComputeShaderRHI, NumCellsParam, NumCellsTmp);	
+
+		int NumTilesTmp[2];
+		NumTilesTmp[0] = ProxyData->NumTilesX;
+		NumTilesTmp[1] = ProxyData->NumTilesY;
+		SetShaderValue(RHICmdList, ComputeShaderRHI, NumTilesParam, NumTilesTmp);		
+
 		SetShaderValue(RHICmdList, ComputeShaderRHI, CellSizeParam, ProxyData->CellSize);		
 		
 		SetShaderValue(RHICmdList, ComputeShaderRHI, WorldBBoxMinParam, ProxyData->WorldBBoxMin);
@@ -96,6 +105,7 @@ struct FNiagaraDataInterfaceParametersCS_Grid2DCollection : public FNiagaraDataI
 private:
 
 	FShaderParameter NumCellsParam;
+	FShaderParameter NumTilesParam;
 	FShaderParameter CellSizeParam;
 	FShaderParameter WorldBBoxMinParam;
 	FShaderParameter WorldBBoxSizeParam;
@@ -136,7 +146,8 @@ void UNiagaraDataInterfaceGrid2DCollection::GetFunctions(TArray<FNiagaraFunction
 		Sig.Inputs.Add(FNiagaraVariable(FNiagaraTypeDefinition(GetClass()), TEXT("Grid")));
 		Sig.Inputs.Add(FNiagaraVariable(FNiagaraTypeDefinition::GetIntDef(), TEXT("IndexX")));
 		Sig.Inputs.Add(FNiagaraVariable(FNiagaraTypeDefinition::GetIntDef(), TEXT("IndexY")));
-		Sig.Outputs.Add(FNiagaraVariable(FNiagaraTypeDefinition::GetVec4Def(), TEXT("Value")));
+		Sig.Inputs.Add(FNiagaraVariable(FNiagaraTypeDefinition::GetIntDef(), TEXT("AttributeIndex")));
+		Sig.Outputs.Add(FNiagaraVariable(FNiagaraTypeDefinition::GetFloatDef(), TEXT("Value")));
 
 		Sig.bMemberFunction = true;
 		Sig.bRequiresContext = false;
@@ -149,7 +160,8 @@ void UNiagaraDataInterfaceGrid2DCollection::GetFunctions(TArray<FNiagaraFunction
 		Sig.Inputs.Add(FNiagaraVariable(FNiagaraTypeDefinition(GetClass()), TEXT("Grid")));
 		Sig.Inputs.Add(FNiagaraVariable(FNiagaraTypeDefinition::GetIntDef(), TEXT("IndexX")));
 		Sig.Inputs.Add(FNiagaraVariable(FNiagaraTypeDefinition::GetIntDef(), TEXT("IndexY")));
-		Sig.Inputs.Add(FNiagaraVariable(FNiagaraTypeDefinition::GetVec4Def(), TEXT("Value")));
+		Sig.Inputs.Add(FNiagaraVariable(FNiagaraTypeDefinition::GetIntDef(), TEXT("AttributeIndex")));
+		Sig.Inputs.Add(FNiagaraVariable(FNiagaraTypeDefinition::GetFloatDef(), TEXT("Value")));
 		Sig.Outputs.Add(FNiagaraVariable(FNiagaraTypeDefinition::GetIntDef(), TEXT("IGNORE")));
 
 		Sig.bMemberFunction = true;
@@ -162,8 +174,9 @@ void UNiagaraDataInterfaceGrid2DCollection::GetFunctions(TArray<FNiagaraFunction
 		Sig.Name = SampleGridFunctionName;
 		Sig.Inputs.Add(FNiagaraVariable(FNiagaraTypeDefinition(GetClass()), TEXT("Grid")));
 		Sig.Inputs.Add(FNiagaraVariable(FNiagaraTypeDefinition::GetFloatDef(), TEXT("UnitX")));
-		Sig.Inputs.Add(FNiagaraVariable(FNiagaraTypeDefinition::GetFloatDef(), TEXT("UnitY")));		
-		Sig.Outputs.Add(FNiagaraVariable(FNiagaraTypeDefinition::GetVec4Def(), TEXT("Value")));
+		Sig.Inputs.Add(FNiagaraVariable(FNiagaraTypeDefinition::GetFloatDef(), TEXT("UnitY")));
+		Sig.Inputs.Add(FNiagaraVariable(FNiagaraTypeDefinition::GetIntDef(), TEXT("AttributeIndex")));
+		Sig.Outputs.Add(FNiagaraVariable(FNiagaraTypeDefinition::GetFloatDef(), TEXT("Value")));
 
 		Sig.bMemberFunction = true;
 		Sig.bRequiresContext = false;
@@ -211,8 +224,9 @@ void UNiagaraDataInterfaceGrid2DCollection::GetParameterDefinitionHLSL(FNiagaraD
 	Super::GetParameterDefinitionHLSL(ParamInfo, OutHLSL);
 
 	static const TCHAR *FormatDeclarations = TEXT(R"(				
-		Texture2D<float4> {GridName};
-		RWTexture2D<float4> RW{OutputGridName};
+		Texture2D<float> {GridName};
+		RWTexture2D<float> RW{OutputGridName};
+		int2 {NumTiles};
 
 		SamplerState {GridName}Sampler
 		{
@@ -224,6 +238,7 @@ void UNiagaraDataInterfaceGrid2DCollection::GetParameterDefinitionHLSL(FNiagaraD
 	TMap<FString, FStringFormatArg> ArgsDeclarations = {				
 		{ TEXT("GridName"),    GridName + ParamInfo.DataInterfaceHLSLSymbol },
 		{ TEXT("OutputGridName"),    OutputGridName + ParamInfo.DataInterfaceHLSLSymbol },
+		{ TEXT("NumTiles"),    NumTilesName + ParamInfo.DataInterfaceHLSLSymbol },
 	};
 	OutHLSL += FString::Format(FormatDeclarations, ArgsDeclarations);
 }
@@ -238,14 +253,19 @@ bool UNiagaraDataInterfaceGrid2DCollection::GetFunctionHLSL(const FName& Definit
 	else if (DefinitionFunctionName == GetValueFunctionName)
 	{
 		static const TCHAR *FormatBounds = TEXT(R"(
-			void {FunctionName}(int In_IndexX, int In_IndexY, out float4 Out_Val)
+			void {FunctionName}(int In_IndexX, int In_IndexY, int In_AttributeIndex, out float Out_Val)
 			{
-				Out_Val = {Grid}.Load(int3(In_IndexX, In_IndexY, 0));
+				int TileIndexX = In_AttributeIndex % {NumTiles}.x;
+				int TileIndexY = In_AttributeIndex / {NumTiles}.x;
+
+				Out_Val = {Grid}.Load(int3(In_IndexX + TileIndexX * {NumCellsName}.x, In_IndexY + TileIndexY * {NumCellsName}.y, 0));
 			}
 		)");
 		TMap<FString, FStringFormatArg> ArgsBounds = {
 			{TEXT("FunctionName"), InstanceFunctionName},
 			{TEXT("Grid"), GridName + ParamInfo.DataInterfaceHLSLSymbol},
+			{TEXT("NumCellsName"), NumCellsName + ParamInfo.DataInterfaceHLSLSymbol},
+			{TEXT("NumTiles"),    NumTilesName + ParamInfo.DataInterfaceHLSLSymbol},
 		};
 		OutHLSL += FString::Format(FormatBounds, ArgsBounds);
 		return true;
@@ -253,15 +273,21 @@ bool UNiagaraDataInterfaceGrid2DCollection::GetFunctionHLSL(const FName& Definit
 	else if (DefinitionFunctionName == SetValueFunctionName)
 	{
 		static const TCHAR *FormatBounds = TEXT(R"(
-			void {FunctionName}(int In_IndexX, int In_IndexY, float4 In_Value, out int val)
-			{				
+			void {FunctionName}(int In_IndexX, int In_IndexY, int In_AttributeIndex, float In_Value, out int val)
+			{			
+				int TileIndexX = In_AttributeIndex % {NumTiles}.x;
+				int TileIndexY = In_AttributeIndex / {NumTiles}.x;
+	
 				val = 0;
-				RW{OutputGrid}[int2(In_IndexX, In_IndexY)] = In_Value;
+				RW{OutputGrid}[int2(In_IndexX + TileIndexX * {NumCellsName}.x, In_IndexY + TileIndexY * {NumCellsName}.y)] = In_Value;
 			}
 		)");
 		TMap<FString, FStringFormatArg> ArgsBounds = {
 			{TEXT("FunctionName"), InstanceFunctionName},
 			{TEXT("OutputGrid"), OutputGridName + ParamInfo.DataInterfaceHLSLSymbol},
+			{TEXT("NumCellsName"), NumCellsName + ParamInfo.DataInterfaceHLSLSymbol},
+			{TEXT("NumTiles"),    NumTilesName + ParamInfo.DataInterfaceHLSLSymbol},
+
 		};
 		OutHLSL += FString::Format(FormatBounds, ArgsBounds);
 		return true;
@@ -269,14 +295,18 @@ bool UNiagaraDataInterfaceGrid2DCollection::GetFunctionHLSL(const FName& Definit
 	else if (DefinitionFunctionName == SampleGridFunctionName)
 	{
 		static const TCHAR *FormatBounds = TEXT(R"(
-			void {FunctionName}(float In_UnitX, float In_UnitY, out float4 Out_Val)
+			void {FunctionName}(float In_UnitX, float In_UnitY, int In_AttributeIndex, out float Out_Val)
 			{
-				Out_Val = {Grid}.SampleLevel({Grid}Sampler, float2(In_UnitX, In_UnitY), 0);
+				int TileIndexX = In_AttributeIndex % {NumTiles}.x;
+				int TileIndexY = In_AttributeIndex / {NumTiles}.x;
+				
+				Out_Val = {Grid}.SampleLevel({Grid}Sampler, float2(In_UnitX / {NumTiles}.x + 1.0*TileIndexX/{NumTiles}.x, In_UnitY / {NumTiles}.y + 1.0*TileIndexY/{NumTiles}.y), 0);
 			}
 		)");
 		TMap<FString, FStringFormatArg> ArgsBounds = {
 			{TEXT("FunctionName"), InstanceFunctionName},
 			{TEXT("Grid"), GridName + ParamInfo.DataInterfaceHLSLSymbol},
+			{TEXT("NumTiles"),    NumTilesName + ParamInfo.DataInterfaceHLSLSymbol},
 		};
 		OutHLSL += FString::Format(FormatBounds, ArgsBounds);
 		return true;
@@ -315,6 +345,18 @@ bool UNiagaraDataInterfaceGrid2DCollection::InitPerInstanceData(void* PerInstanc
 	int RT_NumCellsX = NumCellsX;
 	int RT_NumCellsY = NumCellsY;
 
+	// #todo(dmp): refactor
+	int MaxDim = 16384;
+	int MaxTilesX = floor(MaxDim / NumCellsX);
+	int MaxTilesY = floor(MaxDim / NumCellsY);
+
+	// need to determine number of tiles in x and y based on number of attributes and max dimension size
+	int NumTilesX = NumAttributes <= MaxTilesX ? NumAttributes : MaxTilesX;
+	int NumTilesY = ceil(1.0*NumAttributes / NumTilesX);
+
+	int RT_NumTilesX = NumTilesX;
+	int RT_NumTilesY = NumTilesY;
+
 	FVector RT_WorldBBoxMin = WorldBBoxMin;
 	FVector2D RT_WorldBBoxSize = WorldBBoxSize;
 
@@ -341,7 +383,7 @@ bool UNiagaraDataInterfaceGrid2DCollection::InitPerInstanceData(void* PerInstanc
 	
 	// Push Updates to Proxy.
 	ENQUEUE_RENDER_COMMAND(FUpdateData)(
-		[RT_Proxy, RT_NumCellsX, RT_NumCellsY, RT_CellSize, RT_WorldBBoxMin, RT_WorldBBoxSize, RT_OutputShaderStages, RT_IterationShaderStages, InstanceID = SystemInstance->GetId()](FRHICommandListImmediate& RHICmdList)
+		[RT_Proxy, RT_NumCellsX, RT_NumCellsY, RT_NumTilesX, RT_NumTilesY, RT_CellSize, RT_WorldBBoxMin, RT_WorldBBoxSize, RT_OutputShaderStages, RT_IterationShaderStages, InstanceID = SystemInstance->GetId()](FRHICommandListImmediate& RHICmdList)
 	{
 		Grid2DCollectionRWInstanceData* TargetData = RT_Proxy->SystemInstancesToProxyData.Find(InstanceID);
 		if (TargetData != nullptr)
@@ -355,6 +397,9 @@ bool UNiagaraDataInterfaceGrid2DCollection::InitPerInstanceData(void* PerInstanc
 
 		TargetData->NumCellsX = RT_NumCellsX;
 		TargetData->NumCellsY = RT_NumCellsY;
+		
+		TargetData->NumTilesX = RT_NumTilesX;
+		TargetData->NumTilesY = RT_NumTilesY;
 
 		TargetData->CellSize = RT_CellSize;
 
@@ -388,7 +433,7 @@ void UNiagaraDataInterfaceGrid2DCollection::DestroyPerInstanceData(void* PerInst
 }
 
 UFUNCTION(BlueprintCallable, Category = Niagara)
-void UNiagaraDataInterfaceGrid2DCollection::FillTexture2D(const UNiagaraComponent *Component, UTextureRenderTarget2D *Dest)
+void UNiagaraDataInterfaceGrid2DCollection::FillTexture2D(const UNiagaraComponent *Component, UTextureRenderTarget2D *Dest, int AttributeIndex)
 {
 	FNiagaraDataInterfaceProxyGrid2DCollection* TProxy = GetProxyAs<FNiagaraDataInterfaceProxyGrid2DCollection>();
 
@@ -405,13 +450,21 @@ void UNiagaraDataInterfaceGrid2DCollection::FillTexture2D(const UNiagaraComponen
 	FNiagaraSystemInstanceID InstanceID = SystemInstance->GetId();
 
 	ENQUEUE_RENDER_COMMAND(FUpdateDIColorCurve)(
-		[Dest, TProxy, InstanceID](FRHICommandListImmediate& RHICmdList)
+		[Dest, TProxy, InstanceID, AttributeIndex](FRHICommandListImmediate& RHICmdList)
 	{
 		Grid2DCollectionRWInstanceData* Grid2DInstanceData = TProxy->SystemInstancesToProxyData.Find(InstanceID);
 
 		if (Dest && Dest->Resource && Grid2DInstanceData && Grid2DInstanceData->CurrentData)
 		{
 			FRHICopyTextureInfo CopyInfo;
+			CopyInfo.Size = FIntVector(Grid2DInstanceData->NumCellsX, Grid2DInstanceData->NumCellsY, 1);
+
+			int TileIndexX = AttributeIndex % Grid2DInstanceData->NumTilesX;
+			int TileIndexY = AttributeIndex / Grid2DInstanceData->NumTilesX;
+			int StartX = TileIndexX * Grid2DInstanceData->NumCellsX;
+			int StartY = TileIndexY * Grid2DInstanceData->NumCellsY;
+			CopyInfo.SourcePosition = FIntVector(StartX, StartY, 0);
+
 			RHICmdList.CopyTexture(Grid2DInstanceData->CurrentData->GridBuffer.Buffer, Dest->Resource->TextureRHI, CopyInfo);
 		}
 	});
diff --git a/Engine/Plugins/FX/Niagara/Source/Niagara/Private/NiagaraDataInterfaceRW.cpp b/Engine/Plugins/FX/Niagara/Source/Niagara/Private/NiagaraDataInterfaceRW.cpp
--- a/Engine/Plugins/FX/Niagara/Source/Niagara/Private/NiagaraDataInterfaceRW.cpp
+++ b/Engine/Plugins/FX/Niagara/Source/Niagara/Private/NiagaraDataInterfaceRW.cpp
@@ -306,6 +306,7 @@ UNiagaraDataInterfaceGrid2D::UNiagaraDataInterfaceGrid2D(FObjectInitializer cons
 	, NumCellsX(3)
 	, NumCellsY(3)
 	, CellSize(1.)
+	, NumAttributes(1)
 	, SetGridFromCellSize(false)
 	, WorldBBoxMin(0., 0., 0.)
 	, WorldBBoxSize(100., 100.)
@@ -483,6 +484,7 @@ bool UNiagaraDataInterfaceGrid2D::Equals(const UNiagaraDataInterface* Other) con
 	return 
 		OtherTyped->NumCellsX == NumCellsX &&
 		OtherTyped->NumCellsY == NumCellsY &&
+		OtherTyped->NumAttributes == NumAttributes &&
 		FMath::IsNearlyEqual(OtherTyped->CellSize, CellSize) &&
 		OtherTyped->WorldBBoxMin.Equals(WorldBBoxMin) &&
 		OtherTyped->WorldBBoxSize.Equals(WorldBBoxSize);
@@ -687,6 +689,7 @@ bool UNiagaraDataInterfaceGrid2D::CopyToInternal(UNiagaraDataInterface* Destinat
 
 	OtherTyped->NumCellsX = NumCellsX;
 	OtherTyped->NumCellsY = NumCellsY;
+	OtherTyped->NumAttributes = NumAttributes;
 	OtherTyped->CellSize = CellSize;
 	OtherTyped->SetGridFromCellSize = SetGridFromCellSize;
 	OtherTyped->WorldBBoxMin = WorldBBoxMin;
diff --git a/Engine/Source/Runtime/RHI/Public/RHIUtilities.h b/Engine/Source/Runtime/RHI/Public/RHIUtilities.h
--- a/Engine/Source/Runtime/RHI/Public/RHIUtilities.h
+++ b/Engine/Source/Runtime/RHI/Public/RHIUtilities.h
@@ -82,13 +82,13 @@ struct FTextureRWBuffer
 		);		
 				
 		// #todo(dmp): hardcoded for now...
-		uint32 BlockBytes = 16;
+		uint32 BlockBytes = 4;
 
 		NumBytes = SizeX * SizeY * BlockBytes;
 
 		FRHIResourceCreateInfo CreateInfo;
 		Buffer = RHICreateTexture2D(
-			SizeX, SizeY, PF_A32B32G32R32F,
+			SizeX, SizeY, PF_R32_FLOAT,
 			/*NumMips=*/ 1,
 			1,
 			/*Flags=*/ TexCreate_ShaderResource | TexCreate_UAV,
