- Remove always add binding and replace with gathering the override parameters inside SetAsset
- This ensures that all bindings are hooked up, so if we bind before adding parameters we don't miss the binding

#rnx

diff --git a/Engine/Plugins/FX/Niagara/Source/Niagara/Private/NiagaraComponent.cpp b/Engine/Plugins/FX/Niagara/Source/Niagara/Private/NiagaraComponent.cpp
--- a/Engine/Plugins/FX/Niagara/Source/Niagara/Private/NiagaraComponent.cpp
+++ b/Engine/Plugins/FX/Niagara/Source/Niagara/Private/NiagaraComponent.cpp
@@ -1700,6 +1700,30 @@ void UNiagaraComponent::SetAsset(UNiagaraSystem* InAsset)
 		{
 			AssetExposedParametersChangedHandle.Reset();
 		}
+#else
+		// We need to populate the override parameters here
+		{
+			TArray<FNiagaraVariable> SourceVars;
+			Asset->GetExposedParameters().GetParameters(SourceVars);
+			for (FNiagaraVariable& Param : SourceVars)
+			{
+				OverrideParameters.AddParameter(Param, true);
+			}
+
+			TArray<FNiagaraVariable> ExistingVars;
+			OverrideParameters.GetUserParameters(ExistingVars);
+			Asset->GetExposedParameters().GetUserParameters(SourceVars);
+
+			for (FNiagaraVariable ExistingVar : ExistingVars)
+			{
+				if (!SourceVars.Contains(ExistingVar))
+				{
+					OverrideParameters.RemoveParameter(ExistingVar);
+				}
+			}
+
+			OverrideParameters.Rebind();
+		}
 #endif
 
 		//Force a reinit.
diff --git a/Engine/Plugins/FX/Niagara/Source/Niagara/Private/NiagaraParameterStore.cpp b/Engine/Plugins/FX/Niagara/Source/Niagara/Private/NiagaraParameterStore.cpp
--- a/Engine/Plugins/FX/Niagara/Source/Niagara/Private/NiagaraParameterStore.cpp
+++ b/Engine/Plugins/FX/Niagara/Source/Niagara/Private/NiagaraParameterStore.cpp
@@ -116,26 +116,18 @@ FNiagaraParameterStore::~FNiagaraParameterStore()
 	Bindings.Empty();
 }
 
-void FNiagaraParameterStore::Bind(FNiagaraParameterStore* DestStore, bool bAlwaysAddToBindings)
+void FNiagaraParameterStore::Bind(FNiagaraParameterStore* DestStore)
 {
 	check(DestStore);
 	SCOPE_CYCLE_COUNTER(STAT_NiagaraParameterStoreBind);
 	if (!Bindings.Contains(DestStore))
 	{
-		if (bAlwaysAddToBindings)
+		// Bind the parameter stores only if they have variables in common.
+		FNiagaraParameterStoreBinding HeapBinding;
+		if  (HeapBinding.Initialize(DestStore, this))
 		{
 			FNiagaraParameterStoreBinding& Binding = Bindings.FindOrAdd(DestStore);
-			Binding.Initialize(DestStore, this);
-		}
-		else
-		{
-			// Bind the parameter stores only if they have variables in common.
-			FNiagaraParameterStoreBinding HeapBinding;
-			if  (HeapBinding.Initialize(DestStore, this))
-			{
-				FNiagaraParameterStoreBinding& Binding = Bindings.FindOrAdd(DestStore);
-				FMemory::Memswap(&Binding, &HeapBinding, sizeof(FNiagaraParameterStoreBinding));
-			}
+			FMemory::Memswap(&Binding, &HeapBinding, sizeof(FNiagaraParameterStoreBinding));
 		}
 	}
 }
diff --git a/Engine/Plugins/FX/Niagara/Source/Niagara/Private/NiagaraSystemInstance.cpp b/Engine/Plugins/FX/Niagara/Source/Niagara/Private/NiagaraSystemInstance.cpp
--- a/Engine/Plugins/FX/Niagara/Source/Niagara/Private/NiagaraSystemInstance.cpp
+++ b/Engine/Plugins/FX/Niagara/Source/Niagara/Private/NiagaraSystemInstance.cpp
@@ -988,14 +988,14 @@ void FNiagaraSystemInstance::BindParameters()
 		return;
 	}
 
-	Component->GetOverrideParameters().Bind(&InstanceParameters, true);
+	Component->GetOverrideParameters().Bind(&InstanceParameters);
 
 	if (SystemSimulation->GetIsSolo())
 	{
 		// If this simulation is solo than we can bind the instance parameters to the system simulation contexts so that
 		// the system and emitter scripts use the per-instance data interfaces.
-		Component->GetOverrideParameters().Bind(&SystemSimulation->GetSpawnExecutionContext().Parameters, true);
-		Component->GetOverrideParameters().Bind(&SystemSimulation->GetUpdateExecutionContext().Parameters, true);
+		Component->GetOverrideParameters().Bind(&SystemSimulation->GetSpawnExecutionContext().Parameters);
+		Component->GetOverrideParameters().Bind(&SystemSimulation->GetUpdateExecutionContext().Parameters);
 	}
 	else 
 	{
diff --git a/Engine/Plugins/FX/Niagara/Source/Niagara/Public/NiagaraParameterStore.h b/Engine/Plugins/FX/Niagara/Source/Niagara/Public/NiagaraParameterStore.h
--- a/Engine/Plugins/FX/Niagara/Source/Niagara/Public/NiagaraParameterStore.h
+++ b/Engine/Plugins/FX/Niagara/Source/Niagara/Public/NiagaraParameterStore.h
@@ -213,12 +213,8 @@ public:
 
 	FORCEINLINE uint32 GetLayoutVersion() const { return LayoutVersion; }
 
-	/**
-	  * Binds this parameter store to another, by default if we find no matching parameters we will not maintain a pointer to the store.
-	  * Set bAlwaysAddToBindings to true to always add to the list of ParameterStore bindings, regardless of no matching bindings.
-	  * During Tick the values of this store will be pushed into all bound stores
-	  */
-	void Bind(FNiagaraParameterStore* DestStore, bool bAlwaysAddToBindings = false);
+	/** Binds this parameter store to another, by default if we find no matching parameters we will not maintain a pointer to the store. */
+	void Bind(FNiagaraParameterStore* DestStore);
 	/** Unbinds this store form one it's bound to. */
 	void Unbind(FNiagaraParameterStore* DestStore);
 	/** Recreates any bindings to reflect a layout change etc. */
