Fixed a problem  with function parameter pruning when compiling outside of the game-thread

#tests all niagara automation tests are green

diff --git a/Engine/Plugins/FX/Niagara/Source/NiagaraEditor/Private/NiagaraNodeFunctionCall.cpp b/Engine/Plugins/FX/Niagara/Source/NiagaraEditor/Private/NiagaraNodeFunctionCall.cpp
--- a/Engine/Plugins/FX/Niagara/Source/NiagaraEditor/Private/NiagaraNodeFunctionCall.cpp
+++ b/Engine/Plugins/FX/Niagara/Source/NiagaraEditor/Private/NiagaraNodeFunctionCall.cpp
@@ -456,20 +456,10 @@ void UNiagaraNodeFunctionCall::Compile(class FHlslNiagaraTranslator* Translator,
 		// We check which module inputs are not used so we can later remove them from the compilation of the
 		// parameter map that sets the input values for our function. This is mainly done to prevent data interfaces being
 		// initialized as parameter when they are not used in the function or module.
-		TArray<const UEdGraphPin*> OutInputPins;
-		TSet<const UEdGraphPin*> HiddenPins;
 		TSet<FName> HiddenPinNames;
-		//FNiagaraStackGraphUtilities::GetStackFunctionInputPins(*this, OutInputPins, HiddenPins, FCompileConstantResolver(Translator));
-		// The above line of code isn't threadsafe as it builds a parameter map history live. This does things like load the asset registry module
-		// and  there are some static finds as well. We need to re-implement that using straight graph traversal logic.
-		for (const UEdGraphPin* Pin : HiddenPins)
+		for (UEdGraphPin* Pin : FNiagaraStackGraphUtilities::GetUnusedFunctionInputPins(*this, FCompileConstantResolver(Translator)))
 		{
-			FString PinNameString = Pin->PinName.ToString();
-			if (PinNameString.RemoveFromStart(TEXT("Module.")))
-			{
-				PinNameString = FunctionScript->GetFName().ToString() + TEXT(".") + PinNameString;
-				HiddenPinNames.Add(FName(*PinNameString));
-			}
+			HiddenPinNames.Add(Pin->PinName);
 		}
 		Translator->EnterFunctionCallNode(HiddenPinNames);
 
diff --git a/Engine/Plugins/FX/Niagara/Source/NiagaraEditor/Private/ViewModels/Stack/NiagaraStackGraphUtilities.cpp b/Engine/Plugins/FX/Niagara/Source/NiagaraEditor/Private/ViewModels/Stack/NiagaraStackGraphUtilities.cpp
--- a/Engine/Plugins/FX/Niagara/Source/NiagaraEditor/Private/ViewModels/Stack/NiagaraStackGraphUtilities.cpp
+++ b/Engine/Plugins/FX/Niagara/Source/NiagaraEditor/Private/ViewModels/Stack/NiagaraStackGraphUtilities.cpp
@@ -30,6 +30,7 @@
 #include "EdGraph/EdGraphPin.h"
 #include "ViewModels/NiagaraEmitterViewModel.h"
 #include "AssetRegistryModule.h"
+#include "NiagaraNodeStaticSwitch.h"
 
 DECLARE_CYCLE_STAT(TEXT("Niagara - StackGraphUtilities - RelayoutGraph"), STAT_NiagaraEditor_StackGraphUtilities_RelayoutGraph, STATGROUP_NiagaraEditor);
 
@@ -560,6 +561,121 @@ void FNiagaraStackGraphUtilities::GetStackFunctionInputPins(UNiagaraNodeFunction
 	}
 }
 
+#pragma optimize( "", off )
+
+TArray<UEdGraphPin*> FNiagaraStackGraphUtilities::GetUnusedFunctionInputPins(UNiagaraNodeFunctionCall& FunctionCallNode, FCompileConstantResolver ConstantResolver)
+{
+	UNiagaraGraph* FunctionGraph = FunctionCallNode.GetCalledGraph();
+	if (!FunctionGraph || FunctionCallNode.FunctionScript->Usage != ENiagaraScriptUsage::Module)
+	{
+		return TArray<UEdGraphPin*>();
+	}
+	
+	// Set the static switch values so we traverse the correct node paths
+	TArray<UEdGraphPin*> InputPins;
+	FunctionCallNode.GetInputPins(InputPins);
+	FNiagaraEditorUtilities::SetStaticSwitchConstants(FunctionGraph, InputPins, ConstantResolver);
+
+	// Find the start node for the traversal
+	UNiagaraNodeOutput* OutputNode = FunctionGraph->FindOutputNode(ENiagaraScriptUsage::Function);
+	if (OutputNode == nullptr)
+	{
+		OutputNode = FunctionGraph->FindOutputNode(ENiagaraScriptUsage::Module);
+	}
+	if (OutputNode == nullptr)
+	{
+		OutputNode = FunctionGraph->FindOutputNode(ENiagaraScriptUsage::DynamicInput);
+	}
+
+	// Find reachable nodes
+	TArray<UEdGraphPin*> PinsToCheck;
+	OutputNode->GetInputPins(PinsToCheck);
+	TSet<UNiagaraNode*> ReachedNodes;
+	while (PinsToCheck.Num() > 0)
+	{
+		UEdGraphPin* Pin = PinsToCheck.Pop(false);
+		UNiagaraNode* Owner = Cast<UNiagaraNode>(Pin->GetOwningNode());
+		if (!Owner)
+		{
+			continue;
+		}
+		for (UEdGraphPin* LinkedPin : Pin->LinkedTo)
+		{
+			UEdGraphPin* TracedPin = UNiagaraNode::TraceOutputPin(LinkedPin);
+			if (TracedPin->Direction == EGPD_Output && TracedPin->GetOwningNode()->IsA<UNiagaraNodeStaticSwitch>())
+			{
+				// if the static switch has no connected input (the value is set on the input pin directly) then the trace returns the output pin and we can ignore it
+				continue;
+			}
+			UNiagaraNode* TracedPinOwner = Cast<UNiagaraNode>(TracedPin->GetOwningNode());
+			if (!TracedPinOwner || ReachedNodes.Contains(TracedPinOwner))
+			{
+				continue;
+			}
+			ReachedNodes.Add(TracedPinOwner);
+			TArray<UEdGraphPin*> TracedInputs;
+			TracedPinOwner->GetInputPins(TracedInputs);
+			PinsToCheck.Append(TracedInputs);
+		}
+	}
+
+	// Get the used function parameters from the parameter map set node linked to the function's input pin.
+	// Note that this is only valid for module scripts, not function scripts.
+	TArray<UEdGraphPin*> ResultPins;
+	FString FunctionScriptName = FunctionCallNode.FunctionScript->GetFName().ToString();
+	if (InputPins.Num() > 0 && InputPins[0]->LinkedTo.Num() > 0)
+	{
+		UNiagaraNodeParameterMapSet* ParamMapNode = Cast<UNiagaraNodeParameterMapSet>(InputPins[0]->LinkedTo[0]->GetOwningNode());
+		if (ParamMapNode)
+		{
+			ParamMapNode->GetInputPins(InputPins);
+			for (UEdGraphPin* Pin : InputPins)
+			{
+				FString PinName = Pin->PinName.ToString();
+				if (PinName.StartsWith(FunctionScriptName + "."))
+				{
+					ResultPins.Add(Pin);
+				}
+			}
+		}
+	}
+	if (ResultPins.Num() == 0)
+	{
+		return ResultPins;
+	}
+
+	// We only care about reachable parameter map get nodes with module inputs
+	const UEdGraphSchema_Niagara* Schema = GetDefault<UEdGraphSchema_Niagara>();
+	for (UNiagaraNode* Node : ReachedNodes)
+	{
+		UNiagaraNodeParameterMapGet* ParamMapNode = Cast<UNiagaraNodeParameterMapGet>(Node);
+		if (ParamMapNode)
+		{
+			TArray<UEdGraphPin*> OutPins;
+			ParamMapNode->GetOutputPins(OutPins);
+			for (UEdGraphPin* OutPin : OutPins)
+			{
+				FString OutPinName = OutPin->PinName.ToString();
+				if (!OutPinName.RemoveFromStart(TEXT("Module.")) || OutPin->LinkedTo.Num() == 0)
+				{
+					continue;
+				}
+				for (UEdGraphPin* Pin : ResultPins)
+				{
+					if (Pin->GetName() == FunctionScriptName + "." + OutPinName && Pin->PinType == OutPin->PinType)
+					{
+						ResultPins.RemoveSwap(Pin);
+						break;
+					}
+				}
+			}
+		}
+	}
+	return ResultPins;
+}
+
+#pragma optimize( "", on )
+
 void FNiagaraStackGraphUtilities::GetStackFunctionInputPins(UNiagaraNodeFunctionCall& FunctionCallNode, TArray<const UEdGraphPin*>& OutInputPins, ENiagaraGetStackFunctionInputPinsOptions Options /*= ENiagaraGetStackFunctionInputPinsOptions::AllInputs*/, bool bIgnoreDisabled /*= false*/)
 {
 	TSet<const UEdGraphPin*> HiddenPins;
diff --git a/Engine/Plugins/FX/Niagara/Source/NiagaraEditor/Public/ViewModels/Stack/NiagaraStackGraphUtilities.h b/Engine/Plugins/FX/Niagara/Source/NiagaraEditor/Public/ViewModels/Stack/NiagaraStackGraphUtilities.h
--- a/Engine/Plugins/FX/Niagara/Source/NiagaraEditor/Public/ViewModels/Stack/NiagaraStackGraphUtilities.h
+++ b/Engine/Plugins/FX/Niagara/Source/NiagaraEditor/Public/ViewModels/Stack/NiagaraStackGraphUtilities.h
@@ -80,6 +80,9 @@ namespace FNiagaraStackGraphUtilities
 
 	void GetStackFunctionInputPins(UNiagaraNodeFunctionCall& FunctionCallNode, TArray<const UEdGraphPin*>& OutInputPins, TSet<const UEdGraphPin*>& OutHiddenPins, FCompileConstantResolver ConstantResolver, ENiagaraGetStackFunctionInputPinsOptions Options = ENiagaraGetStackFunctionInputPinsOptions::AllInputs, bool bIgnoreDisabled = false);
 
+	/* Module script calls do not have direct inputs, but rely on the parameter map being initialized correctly. This utility function resolves which of the module's parameters are reachable during compilation and returns a list of pins on the parameter map node that do not have to be compiled. */
+	TArray<UEdGraphPin*> GetUnusedFunctionInputPins(UNiagaraNodeFunctionCall& FunctionCallNode, FCompileConstantResolver ConstantResolver);
+
 	void GetStackFunctionStaticSwitchPins(UNiagaraNodeFunctionCall& FunctionCallNode, TArray<UEdGraphPin*>& OutInputPins, TSet<UEdGraphPin*>& OutHiddenPins);
 
 	UNiagaraNodeParameterMapSet* GetStackFunctionOverrideNode(UNiagaraNodeFunctionCall& FunctionCallNode);
