Removed debug macro

diff --git a/Engine/Plugins/FX/Niagara/Source/NiagaraEditor/Private/ViewModels/Stack/NiagaraStackGraphUtilities.cpp b/Engine/Plugins/FX/Niagara/Source/NiagaraEditor/Private/ViewModels/Stack/NiagaraStackGraphUtilities.cpp
--- a/Engine/Plugins/FX/Niagara/Source/NiagaraEditor/Private/ViewModels/Stack/NiagaraStackGraphUtilities.cpp
+++ b/Engine/Plugins/FX/Niagara/Source/NiagaraEditor/Private/ViewModels/Stack/NiagaraStackGraphUtilities.cpp
@@ -561,8 +561,6 @@ void FNiagaraStackGraphUtilities::GetStackFunctionInputPins(UNiagaraNodeFunction
 	}
 }
 
-#pragma optimize( "", off )
-
 TArray<UEdGraphPin*> FNiagaraStackGraphUtilities::GetUnusedFunctionInputPins(UNiagaraNodeFunctionCall& FunctionCallNode, FCompileConstantResolver ConstantResolver)
 {
 	UNiagaraGraph* FunctionGraph = FunctionCallNode.GetCalledGraph();
@@ -674,8 +672,6 @@ TArray<UEdGraphPin*> FNiagaraStackGraphUtilities::GetUnusedFunctionInputPins(UNi
 	return ResultPins;
 }
 
-#pragma optimize( "", on )
-
 void FNiagaraStackGraphUtilities::GetStackFunctionInputPins(UNiagaraNodeFunctionCall& FunctionCallNode, TArray<const UEdGraphPin*>& OutInputPins, ENiagaraGetStackFunctionInputPinsOptions Options /*= ENiagaraGetStackFunctionInputPinsOptions::AllInputs*/, bool bIgnoreDisabled /*= false*/)
 {
 	TSet<const UEdGraphPin*> HiddenPins;
