Added a null check

diff --git a/Engine/Plugins/FX/Niagara/Source/NiagaraEditor/Private/ViewModels/Stack/NiagaraStackGraphUtilities.cpp b/Engine/Plugins/FX/Niagara/Source/NiagaraEditor/Private/ViewModels/Stack/NiagaraStackGraphUtilities.cpp
--- a/Engine/Plugins/FX/Niagara/Source/NiagaraEditor/Private/ViewModels/Stack/NiagaraStackGraphUtilities.cpp
+++ b/Engine/Plugins/FX/Niagara/Source/NiagaraEditor/Private/ViewModels/Stack/NiagaraStackGraphUtilities.cpp
@@ -575,14 +575,10 @@ TArray<UEdGraphPin*> FNiagaraStackGraphUtilities::GetUnusedFunctionInputPins(UNi
 	FNiagaraEditorUtilities::SetStaticSwitchConstants(FunctionGraph, InputPins, ConstantResolver);
 
 	// Find the start node for the traversal
-	UNiagaraNodeOutput* OutputNode = FunctionGraph->FindOutputNode(ENiagaraScriptUsage::Function);
+	UNiagaraNodeOutput* OutputNode = FunctionGraph->FindOutputNode(ENiagaraScriptUsage::Module);
 	if (OutputNode == nullptr)
 	{
-		OutputNode = FunctionGraph->FindOutputNode(ENiagaraScriptUsage::Module);
-	}
-	if (OutputNode == nullptr)
-	{
-		OutputNode = FunctionGraph->FindOutputNode(ENiagaraScriptUsage::DynamicInput);
+		return TArray<UEdGraphPin*>();
 	}
 
 	// Find reachable nodes
