- Fix static analysis warning about using System after nullptr check

#rnx

diff --git a/Engine/Plugins/FX/Niagara/Source/Niagara/Private/NiagaraSystemInstance.cpp b/Engine/Plugins/FX/Niagara/Source/Niagara/Private/NiagaraSystemInstance.cpp
--- a/Engine/Plugins/FX/Niagara/Source/Niagara/Private/NiagaraSystemInstance.cpp
+++ b/Engine/Plugins/FX/Niagara/Source/Niagara/Private/NiagaraSystemInstance.cpp
@@ -1507,12 +1507,14 @@ void FNiagaraSystemInstance::InitEmitters()
 
 	bHasGPUEmitters = false;
 
+	LocalBounds = FBox(FVector::ZeroVector, FVector::ZeroVector);
+
 	Emitters.Empty();
 	UNiagaraSystem* System = GetSystem();
 	if (System != nullptr)
 	{
-		const TArray<FNiagaraEmitterHandle>& EmitterHandles = GetSystem()->GetEmitterHandles();
-		for (int32 EmitterIdx=0; EmitterIdx < GetSystem()->GetEmitterHandles().Num(); ++EmitterIdx)
+		const TArray<FNiagaraEmitterHandle>& EmitterHandles = System->GetEmitterHandles();
+		for (int32 EmitterIdx=0; EmitterIdx < System->GetEmitterHandles().Num(); ++EmitterIdx)
 		{
 			TSharedRef<FNiagaraEmitterInstance, ESPMode::ThreadSafe> Sim = MakeShared<FNiagaraEmitterInstance, ESPMode::ThreadSafe>(this);
 			Sim->Init(EmitterIdx, ID);
@@ -1529,15 +1531,11 @@ void FNiagaraSystemInstance::InitEmitters()
 
 			Simulation->PostInitSimulation();
 		}
-	}
 
-	if ( System->bFixedBounds )
-	{
-		LocalBounds = System->GetFixedBounds();
-	}
-	else
-	{
-		LocalBounds = FBox(FVector::ZeroVector, FVector::ZeroVector);
+		if (System->bFixedBounds)
+		{
+			LocalBounds = System->GetFixedBounds();
+		}
 	}
 }
 
