Niagara - System script fast path

diff --git a/Engine/Plugins/FX/Niagara/Source/Niagara/Classes/NiagaraEmitter.h b/Engine/Plugins/FX/Niagara/Source/Niagara/Classes/NiagaraEmitter.h
--- a/Engine/Plugins/FX/Niagara/Source/Niagara/Classes/NiagaraEmitter.h
+++ b/Engine/Plugins/FX/Niagara/Source/Niagara/Classes/NiagaraEmitter.h
@@ -9,6 +9,7 @@
 #include "NiagaraScript.h"
 #include "NiagaraCollision.h"
 #include "INiagaraMergeManager.h"
+#include "NiagaraSystemFastPath.h"
 #include "NiagaraEmitter.generated.h"
 
 class UMaterial;
@@ -374,6 +375,29 @@ public:
 
 	NIAGARA_API void RemoveParent();
 
+	void InitFastPathAttributeNames();
+
+	UPROPERTY(EditAnywhere, Category = "Script Fast Path")
+	FNiagaraFastPath_Module_EmitterScalability EmitterScalability;
+
+	UPROPERTY(EditAnywhere, Category = "Script Fast Path")
+	FNiagaraFastPath_Module_EmitterLifeCycle EmitterLifeCycle;
+
+	UPROPERTY(EditAnywhere, Category = "Script Fast Path")
+	TArray<FNiagaraFastPath_Module_SpawnRate> SpawnRate;
+
+	UPROPERTY(EditAnywhere, Category = "Script Fast Path")
+	TArray<FNiagaraFastPath_Module_SpawnPerUnit> SpawnPerUnit;
+
+	UPROPERTY(EditAnywhere, Category = "Script Fast Path")
+	TArray<FNiagaraFastPath_Module_SpawnBurstInstantaneous> SpawnBurstInstantaneous;
+
+	UPROPERTY()
+	FNiagaraFastPathAttributeNames SpawnFastPathAttributeNames;
+
+	UPROPERTY()
+	FNiagaraFastPathAttributeNames UpdateFastPathAttributeNames;
+
 protected:
 	virtual void BeginDestroy() override;
 
diff --git a/Engine/Plugins/FX/Niagara/Source/Niagara/Classes/NiagaraEmitterInstance.h b/Engine/Plugins/FX/Niagara/Source/Niagara/Classes/NiagaraEmitterInstance.h
--- a/Engine/Plugins/FX/Niagara/Source/Niagara/Classes/NiagaraEmitterInstance.h
+++ b/Engine/Plugins/FX/Niagara/Source/Niagara/Classes/NiagaraEmitterInstance.h
@@ -15,6 +15,7 @@ NiagaraEmitterInstance.h: Niagara emitter simulation class
 #include "NiagaraEmitter.h"
 #include "NiagaraScriptExecutionContext.h"
 #include "NiagaraBoundsCalculator.h"
+#include "NiagaraSystemFastPath.h"
 
 class FNiagaraSystemInstance;
 struct FNiagaraEmitterHandle;
@@ -35,13 +36,13 @@ public:
 
 	void Init(int32 InEmitterIdx, FNiagaraSystemInstanceID SystemInstanceID);
 
-	void ResetSimulation(bool bKillExisting=true);
+	void ResetSimulation(bool bKillExisting = true);
 
 	/** Called after all emitters in an System have been initialized, allows emitters to access information from one another. */
 	void PostInitSimulation();
 
 	void DirtyDataInterfaces();
-	
+
 	/** Replaces the binding for a single parameter colleciton instance. If for example the component begins to override the global instance. */
 	//void RebindParameterCollection(UNiagaraParameterCollectionInstance* OldInstance, UNiagaraParameterCollectionInstance* NewInstance);
 	void BindParameters();
@@ -52,7 +53,7 @@ public:
 	void PreTick();
 	void Tick(float DeltaSeconds);
 	void PostTick();
-	bool HandleCompletion(bool bForce=false);
+	bool HandleCompletion(bool bForce = false);
 
 	bool RequiredPersistentID()const;
 
@@ -65,10 +66,10 @@ public:
 	NIAGARA_API void CalculateFixedBounds(const FTransform& ToWorldSpace);
 
 	FNiagaraDataSet& GetData()const { return *ParticleDataSet; }
-	
+
 	FORCEINLINE bool IsDisabled()const { return ExecutionState == ENiagaraExecutionState::Disabled; }
 	FORCEINLINE bool IsComplete()const { return ExecutionState == ENiagaraExecutionState::Complete || ExecutionState == ENiagaraExecutionState::Disabled; }
-		
+
 	/** Create a new NiagaraRenderer. The old renderer is not immediately deleted, but instead put in the ToBeRemoved list.*/
 	//void NIAGARA_API UpdateEmitterRenderer(ERHIFeatureLevel::Type FeatureLevel, TArray<NiagaraRenderer*>& ToBeAddedList, TArray<NiagaraRenderer*>& ToBeRemovedList);
 
@@ -81,16 +82,16 @@ public:
 		}
 		return 0;
 	}
-	FORCEINLINE int32 GetTotalSpawnedParticles()const	{ return TotalSpawnedParticles; }
+	FORCEINLINE int32 GetTotalSpawnedParticles()const { return TotalSpawnedParticles; }
 
 	NIAGARA_API const FNiagaraEmitterHandle& GetEmitterHandle() const;
 
-	FNiagaraSystemInstance* GetParentSystemInstance()const	{ return ParentSystemInstance; }
+	FNiagaraSystemInstance* GetParentSystemInstance()const { return ParentSystemInstance; }
 
 	float NIAGARA_API GetTotalCPUTimeMS();
 	int	NIAGARA_API GetTotalBytesUsed();
-	
-	ENiagaraExecutionState NIAGARA_API GetExecutionState()	{ return ExecutionState; }
+
+	ENiagaraExecutionState NIAGARA_API GetExecutionState() { return ExecutionState; }
 	void NIAGARA_API SetExecutionState(ENiagaraExecutionState InState);
 
 	FNiagaraDataSet* GetDataSet(FNiagaraDataSetID SetID);
@@ -121,10 +122,30 @@ public:
 
 	bool FindBinding(const FNiagaraUserParameterBinding& InBinding, TArray<UMaterialInterface*>& OutMaterials) const;
 
-private:
+	void InitFastPathAttributeBindings();
+
+	void TickFastPathAttributeBindings();
+
+	FNiagaraEmitterFastPath::FParamMap0& GetFastPathMap() { return FastPathMap; }
 
+private:
 	void CheckForErrors();
 
+	void InitFastPathParameterBindingsInternal(const FNiagaraFastPathAttributeNames& SourceParameterNames, FNiagaraParameterStore& TargetParameterStore);
+
+	template<typename TBindingType, typename TVariableType>
+	static void AddBinding(FName ParameterName, TVariableType ParameterType, TBindingType* SourceValuePtr, FNiagaraParameterStore& TargetParameterStore, TArray<TNiagaraFastPathAttributeBinding<TBindingType>>& TargetBindings)
+	{
+		TNiagaraFastPathAttributeBinding<TBindingType> Binding;
+		FNiagaraVariable ParameterVariable = FNiagaraVariable(ParameterType, ParameterName);
+		Binding.ParameterBinding.Init(TargetParameterStore, ParameterVariable);
+		if (Binding.ParameterBinding.ValuePtr != nullptr)
+		{
+			Binding.ParameterValue = SourceValuePtr;
+			TargetBindings.Add(Binding);
+		}
+	}
+
 	/** The index of our emitter in our parent system instance. */
 	int32 EmitterIdx;
 
@@ -218,4 +239,9 @@ private:
 
 	/** Optional list of bounds calculators. */
 	TArray<TUniquePtr<FNiagaraBoundsCalculator>, TInlineAllocator<1>> BoundsCalculators;
+
+	FNiagaraEmitterFastPath::FParamMap0 FastPathMap;
+
+	TArray<TNiagaraFastPathAttributeBinding<int32>> FastPathIntAttributeBindings;
+	TArray<TNiagaraFastPathAttributeBinding<float>> FastPathFloatAttributeBindings;
 };
diff --git a/Engine/Plugins/FX/Niagara/Source/Niagara/Classes/NiagaraSystem.h b/Engine/Plugins/FX/Niagara/Source/Niagara/Classes/NiagaraSystem.h
--- a/Engine/Plugins/FX/Niagara/Source/Niagara/Classes/NiagaraSystem.h
+++ b/Engine/Plugins/FX/Niagara/Source/Niagara/Classes/NiagaraSystem.h
@@ -12,6 +12,7 @@
 #include "NiagaraEmitterHandle.h"
 #include "NiagaraParameterCollection.h"
 #include "NiagaraUserRedirectionParameterStore.h"
+#include "NiagaraSystemFastPath.h"
 #include "NiagaraSystem.generated.h"
 
 #if WITH_EDITORONLY_DATA
@@ -265,6 +266,15 @@ public:
 
 	TStatId GetStatID(bool bGameThread, bool bConcurrent)const;
 
+	UPROPERTY(EditAnywhere, Category = "Script Fast Path")
+	ENiagaraFastPathMode FastPathMode;
+
+	UPROPERTY(EditAnywhere, Category = "Script Fast Path")
+	FNiagaraFastPath_Module_SystemScalability SystemScalability;
+
+	UPROPERTY(EditAnywhere, Category = "Script Fast Path")
+	FNiagaraFastPath_Module_SystemLifeCycle SystemLifeCycle;
+
 private:
 #if WITH_EDITORONLY_DATA
 	bool QueryCompileComplete(bool bWait, bool bDoPost, bool bDoNotApply = false);
diff --git a/Engine/Plugins/FX/Niagara/Source/Niagara/Private/NiagaraEmitter.cpp b/Engine/Plugins/FX/Niagara/Source/Niagara/Private/NiagaraEmitter.cpp
--- a/Engine/Plugins/FX/Niagara/Source/Niagara/Private/NiagaraEmitter.cpp
+++ b/Engine/Plugins/FX/Niagara/Source/Niagara/Private/NiagaraEmitter.cpp
@@ -784,6 +784,8 @@ void UNiagaraEmitter::OnPostCompile()
 	}
 
 	OnEmitterVMCompiled().Broadcast(this);
+
+	InitFastPathAttributeNames();
 }
 
 UNiagaraEmitter* UNiagaraEmitter::MakeRecursiveDeepCopy(UObject* DestOuter) const
@@ -1176,3 +1178,34 @@ void UNiagaraEmitter::RemoveParent()
 	Parent = nullptr;
 	ParentAtLastMerge = nullptr;
 }
+
+void UNiagaraEmitter::InitFastPathAttributeNames()
+{
+	auto InitParameters = [](const FNiagaraParameters& Parameters, const FString& EmitterName, FNiagaraFastPathAttributeNames& FastPathParameterNames)
+	{
+		FastPathParameterNames.System.Empty();
+		FastPathParameterNames.SystemFullNames.Empty();
+		FastPathParameterNames.Emitter.Empty();
+		FastPathParameterNames.EmitterFullNames.Empty();
+
+		FString SystemPrefix = TEXT("System.");
+		FString EmitterPrefix = EmitterName + TEXT(".");
+		for (const FNiagaraVariable& Parameter : Parameters.Parameters)
+		{
+			FString ParameterNameString = Parameter.GetName().ToString();
+			if (ParameterNameString.StartsWith(SystemPrefix))
+			{
+				FastPathParameterNames.System.Add(*(Parameter.GetName().ToString().RightChop(SystemPrefix.Len())));
+				FastPathParameterNames.SystemFullNames.Add(Parameter.GetName());
+			}
+			else if (ParameterNameString.StartsWith(EmitterPrefix))
+			{
+				FastPathParameterNames.Emitter.Add(*(Parameter.GetName().ToString().RightChop(EmitterPrefix.Len())));
+				FastPathParameterNames.EmitterFullNames.Add(Parameter.GetName());
+			}
+		}
+	};
+
+	InitParameters(SpawnScriptProps.Script->GetVMExecutableData().Parameters, UniqueEmitterName, SpawnFastPathAttributeNames);
+	InitParameters(UpdateScriptProps.Script->GetVMExecutableData().Parameters, UniqueEmitterName, UpdateFastPathAttributeNames);
+}
diff --git a/Engine/Plugins/FX/Niagara/Source/Niagara/Private/NiagaraEmitterInstance.cpp b/Engine/Plugins/FX/Niagara/Source/Niagara/Private/NiagaraEmitterInstance.cpp
--- a/Engine/Plugins/FX/Niagara/Source/Niagara/Private/NiagaraEmitterInstance.cpp
+++ b/Engine/Plugins/FX/Niagara/Source/Niagara/Private/NiagaraEmitterInstance.cpp
@@ -252,7 +252,14 @@ void FNiagaraEmitterInstance::Init(int32 InEmitterIdx, FNiagaraSystemInstanceID
 
 	//Init the spawn infos to the correct number for this system.
 	int32 NumEvents = CachedEmitter->GetEventHandlers().Num();
-	SpawnInfos.SetNum(EmitterCompiledData[EmitterIdx].SpawnAttributes.Num());
+	if (ParentSystemInstance->GetSystem()->FastPathMode != ENiagaraFastPathMode::ScriptVMOnly)
+	{
+		SpawnInfos.SetNum(CachedEmitter->SpawnRate.Num() + CachedEmitter->SpawnPerUnit.Num() + CachedEmitter->SpawnBurstInstantaneous.Num());
+	}
+	else
+	{
+		SpawnInfos.SetNum(EmitterCompiledData[EmitterIdx].SpawnAttributes.Num());
+	}
 
 #if !UE_BUILD_SHIPPING && !UE_BUILD_TEST 
 	Data.InitFromCompiledData(EmitterCompiledData[EmitterIdx].DataSetCompiledData, ParentSystemInstance->GetSystem()->GetName() + TEXT("/") + CachedEmitter->GetName());
@@ -445,6 +452,10 @@ void FNiagaraEmitterInstance::ResetSimulation(bool bKillExisting /*= true*/)
 		{
 			GPUExecContext->Reset(Batcher);
 		}
+
+		FastPathMap = FNiagaraEmitterFastPath::FParamMap0();
+		FastPathIntAttributeBindings.Empty();
+		FastPathFloatAttributeBindings.Empty();
 	}
 }
 
@@ -1704,4 +1715,28 @@ bool FNiagaraEmitterInstance::FindBinding(const FNiagaraUserParameterBinding& In
 		}
 	}
 	return false;
-}
\ No newline at end of file
+}
+
+void FNiagaraEmitterInstance::InitFastPathAttributeBindings()
+{
+	FastPathIntAttributeBindings.Empty();
+	FastPathFloatAttributeBindings.Empty();
+
+	FNiagaraEmitterFastPath::InitFastPathAttributeBindings(CachedEmitter->SpawnFastPathAttributeNames, SpawnExecContext.Parameters, 
+		GetParentSystemInstance()->FastPathMap, FastPathMap, FastPathIntAttributeBindings, FastPathFloatAttributeBindings);
+	FNiagaraEmitterFastPath::InitFastPathAttributeBindings(CachedEmitter->UpdateFastPathAttributeNames, UpdateExecContext.Parameters, 
+		GetParentSystemInstance()->FastPathMap, FastPathMap, FastPathIntAttributeBindings, FastPathFloatAttributeBindings);
+}
+
+void FNiagaraEmitterInstance::TickFastPathAttributeBindings()
+{
+	for (TNiagaraFastPathAttributeBinding<int32>& IntBinding : FastPathIntAttributeBindings)
+	{
+		IntBinding.Tick();
+	}
+
+	for (TNiagaraFastPathAttributeBinding<float>& FloatBinding : FastPathFloatAttributeBindings)
+	{
+		FloatBinding.Tick();
+	}
+}
diff --git a/Engine/Plugins/FX/Niagara/Source/Niagara/Private/NiagaraSystem.cpp b/Engine/Plugins/FX/Niagara/Source/Niagara/Private/NiagaraSystem.cpp
--- a/Engine/Plugins/FX/Niagara/Source/Niagara/Private/NiagaraSystem.cpp
+++ b/Engine/Plugins/FX/Niagara/Source/Niagara/Private/NiagaraSystem.cpp
@@ -36,6 +36,7 @@ DECLARE_CYCLE_STAT(TEXT("Niagara - System - CompileScript_ResetAfter"), STAT_Nia
 UNiagaraSystem::UNiagaraSystem(const FObjectInitializer& ObjectInitializer)
 : Super(ObjectInitializer)
 , bFixedBounds(false)
+, FastPathMode(ENiagaraFastPathMode::ScriptVMOnly)
 , ExposedParameters(this)
 #if WITH_EDITORONLY_DATA
 , bIsolateEnabled(false)
diff --git a/Engine/Plugins/FX/Niagara/Source/Niagara/Private/NiagaraSystemFastPath.cpp b/Engine/Plugins/FX/Niagara/Source/Niagara/Private/NiagaraSystemFastPath.cpp
new file mode 100644
--- /dev/null
+++ b/Engine/Plugins/FX/Niagara/Source/Niagara/Private/NiagaraSystemFastPath.cpp
@@ -0,0 +1,719 @@
+// Copyright 1998-2019 Epic Games, Inc. All Rights Reserved.
+
+#include "NiagaraSystemFastPath.h"
+#include "NiagaraEmitter.h"
+
+FName FNiagaraSystemFastPath::FParameterNames::ExecutionState = "ExecutionState";
+FName FNiagaraSystemFastPath::FParameterNames::ExecutionStateSource = "ExecutionStateSource";
+
+void FNiagaraSystemFastPath::SetSpawnMapDefaults(FParamMap0& Map)
+{
+	Map.System.ExecutionState = (ENiagaraExecutionState)0;
+	Map.System.ExecutionStateSource = (ENiagaraExecutionStateSource)0;
+}
+
+void FNiagaraSystemFastPath::SetUpdateMapDefaults(FParamMap0& Map)
+{
+	Map.Scalability.ExecutionState = (ENiagaraExecutionState)0;
+}
+
+void FNiagaraSystemFastPath::Module_SystemScalability(const FNiagaraFastPath_Module_SystemScalability& Context_Map_SystemScalability, FParamMap0& Context_Map)
+{
+	bool Result9 = Context_Map.Engine.Owner.TimeSinceRendered > Context_Map_SystemScalability.VisibilityCullDelay;
+	bool Result10 = Context_Map_SystemScalability.bEnableVisibilityCulling && Result9;
+	ENiagaraExecutionState ENiagaraExecutionState_IfResult3;
+	if (Result10)
+	{
+		ENiagaraExecutionState_IfResult3 = Context_Map_SystemScalability.CulledState;
+	}
+	else
+	{
+		ENiagaraExecutionState_IfResult3 = Context_Map.Scalability.ExecutionState;
+	}
+	Context_Map.Scalability.ExecutionState = ENiagaraExecutionState_IfResult3;
+}
+
+void FNiagaraSystemFastPath::Function_SystemChangeState(ENiagaraExecutionState In_NewState, bool In_Condition, ENiagaraExecutionStateSource In_NewStateSource, FParamMap0& Context_Map)
+{
+	bool Result4 = Context_Map.System.ExecutionStateSource <= In_NewStateSource;
+	bool Result5 = In_Condition && Result4;
+	ENiagaraExecutionState ENiagaraExecutionState_IfResult1;
+	ENiagaraExecutionStateSource ENiagaraExecutionStateSource_IfResult1;
+	if (Result5)
+	{
+		ENiagaraExecutionState_IfResult1 = In_NewState;
+		ENiagaraExecutionStateSource_IfResult1 = In_NewStateSource;
+	}
+	else
+	{
+		ENiagaraExecutionState_IfResult1 = Context_Map.System.ExecutionState;
+		ENiagaraExecutionStateSource_IfResult1 = Context_Map.System.ExecutionStateSource;
+	}
+	Context_Map.System.ExecutionState = ENiagaraExecutionState_IfResult1;
+	Context_Map.System.ExecutionStateSource = ENiagaraExecutionStateSource_IfResult1;
+}
+
+void FNiagaraSystemFastPath::Module_SystemLifeCycle(const FNiagaraFastPath_Module_SystemLifeCycle& Context_Map_SystemLifeCycle, FParamMap0& Context_Map)
+{
+	bool Constant2 = true;
+	ENiagaraExecutionStateSource Constant3 = ENiagaraExecutionStateSource::Scalability;
+	Function_SystemChangeState(Context_Map.Scalability.ExecutionState, Constant2, Constant3, Context_Map);
+	ENiagaraExecutionState Constant4 = ENiagaraExecutionState::Complete;
+	ENiagaraExecutionState Constant5 = ENiagaraExecutionState::Active;
+	bool Result2 = Context_Map.System.ExecutionState != Constant5;
+	bool Result3 = Context_Map_SystemLifeCycle.bCompleteOnInactive && Result2;
+	ENiagaraExecutionStateSource Constant6 = ENiagaraExecutionStateSource::InternalCompletion;
+	Function_SystemChangeState(Constant4, Result3, Constant6, Context_Map);
+	ENiagaraExecutionState Constant7 = ENiagaraExecutionState::Active;
+	bool Result6 = Context_Map.Engine.Owner.ExecutionState != Constant7;
+	ENiagaraExecutionStateSource Constant8 = ENiagaraExecutionStateSource::Owner;
+	Function_SystemChangeState(Context_Map.Engine.Owner.ExecutionState, Result6, Constant8, Context_Map);
+}
+
+FName FNiagaraEmitterFastPath::FAttributeNames::Age = "Age";
+FName FNiagaraEmitterFastPath::FAttributeNames::CurrentLoopDelay = "CurrentLoopDelay";
+FName FNiagaraEmitterFastPath::FAttributeNames::CurrentLoopDuration = "CurrentLoopDuration";
+FName FNiagaraEmitterFastPath::FAttributeNames::ExecutionState = "ExecutionState";
+FName FNiagaraEmitterFastPath::FAttributeNames::ExecutionStateSource = "ExecutionStateSource";
+FName FNiagaraEmitterFastPath::FAttributeNames::LoopCount = "LoopCount";
+FName FNiagaraEmitterFastPath::FAttributeNames::LoopedAge = "LoopedAge";
+FName FNiagaraEmitterFastPath::FAttributeNames::NormalizedLoopAge = "NormalizedLoopAge";
+FName FNiagaraEmitterFastPath::FAttributeNames::DistanceTravelled = "DistanceTraveled";
+FName FNiagaraEmitterFastPath::FAttributeNames::Scalability_DistanceFraction = "Scalability.DistanceFraction";
+
+template<typename TSourceInput, typename TTargetInput>
+void InitInput(
+	const TSourceInput& SourceInput,
+	TTargetInput& TargetInput,
+	FNiagaraTypeDefinition InputType,
+	FNiagaraParameterStore& InstanceParameters,
+	TArray<TNiagaraFastPathUserParameterInputBinding<TTargetInput>>& UserParameterInputBindings,
+	TArray<TNiagaraFastPathRangedInputBinding<TTargetInput>>& UpdateRangedInputBindings)
+{
+	switch (SourceInput.Type)
+	{
+	case ENiagaraFastPathInputType::Local:
+		TargetInput = SourceInput.Local;
+		break;
+	case ENiagaraFastPathInputType::RangedSpawn:
+		TargetInput = FMath::RandRange(SourceInput.RangeMin, SourceInput.RangeMax);
+		break;
+	case ENiagaraFastPathInputType::RangedUpdate:
+		UpdateRangedInputBindings.Add(TNiagaraFastPathRangedInputBinding<TTargetInput>(SourceInput.RangeMin, SourceInput.RangeMax, &TargetInput));
+		break;
+	case ENiagaraFastPathInputType::User:
+		FNiagaraVariable UserParameterVariable = FNiagaraVariable(InputType, SourceInput.UserParameterName);
+		TNiagaraFastPathUserParameterInputBinding<TTargetInput> UserParameterBinding;
+		UserParameterBinding.ParameterBinding.Init(InstanceParameters, UserParameterVariable);
+		if (UserParameterBinding.ParameterBinding.ValuePtr != nullptr)
+		{
+			UserParameterBinding.InputValue = &TargetInput;
+			UserParameterInputBindings.Add(UserParameterBinding);
+		}
+		else
+		{
+			TargetInput = SourceInput.Local;
+		}
+		break;
+	}
+}
+
+void FNiagaraEmitterFastPath::FParamMap0_SpawnRate::Init(
+	const UNiagaraEmitter* Emitter,
+	int32 ModuleIndex,
+	FNiagaraParameterStore& InstanceParameters,
+	TArray<TNiagaraFastPathUserParameterInputBinding<int32>>& IntUserParameterInputBindings,
+	TArray<TNiagaraFastPathUserParameterInputBinding<float>>& FloatUserParameterInputBindings,
+	TArray<TNiagaraFastPathRangedInputBinding<int32>>& IntUpdateRangedInputBindings,
+	TArray<TNiagaraFastPathRangedInputBinding<float>>& FloatUpdateRangedInputBindings)
+{
+	const FNiagaraFastPath_Module_SpawnRate& SpawnRateInputs = Emitter->SpawnRate[ModuleIndex];
+	InitInput(SpawnRateInputs.SpawnRate, SpawnRate, FNiagaraTypeDefinition::GetFloatDef(), InstanceParameters, FloatUserParameterInputBindings, FloatUpdateRangedInputBindings);
+	SpawnGroup = SpawnRateInputs.SpawnGroup;
+}
+
+void FNiagaraEmitterFastPath::FParamMap0_SpawnPerUnit::Init(
+	const UNiagaraEmitter* Emitter,
+	int32 ModuleIndex,
+	FNiagaraParameterStore& InstanceParameters,
+	TArray<TNiagaraFastPathUserParameterInputBinding<int32>>& IntUserParameterInputBindings,
+	TArray<TNiagaraFastPathUserParameterInputBinding<float>>& FloatUserParameterInputBindings,
+	TArray<TNiagaraFastPathRangedInputBinding<int32>>& IntUpdateRangedInputBindings,
+	TArray<TNiagaraFastPathRangedInputBinding<float>>& FloatUpdateRangedInputBindings)
+{
+	const FNiagaraFastPath_Module_SpawnPerUnit& SpawnPerUnitInputs = Emitter->SpawnPerUnit[ModuleIndex];
+	InitInput(SpawnPerUnitInputs.SpawnPerUnit, SpawnPerUnit, FNiagaraTypeDefinition::GetFloatDef(), InstanceParameters, FloatUserParameterInputBindings, FloatUpdateRangedInputBindings);
+	bUseMovementTolerance = SpawnPerUnitInputs.bUseMovementTolerance;
+	MovementTolerance = SpawnPerUnitInputs.MovementTolerance;
+	bUseMaxMovementThreshold = SpawnPerUnitInputs.bUseMaxMovementThreshold;
+	MaxMovementThreshold = SpawnPerUnitInputs.MaxMovementThreshold;
+	SpawnGroup = SpawnPerUnitInputs.SpawnGroup;
+}
+
+void FNiagaraEmitterFastPath::FParamMap0_SpawnBurst_Instantaneous::Init(
+	const UNiagaraEmitter* Emitter,
+	int32 ModuleIndex,
+	FNiagaraParameterStore& InstanceParameters,
+	TArray<TNiagaraFastPathUserParameterInputBinding<int32>>& IntUserParameterInputBindings,
+	TArray<TNiagaraFastPathUserParameterInputBinding<float>>& FloatUserParameterInputBindings,
+	TArray<TNiagaraFastPathRangedInputBinding<int32>>& IntUpdateRangedInputBindings,
+	TArray<TNiagaraFastPathRangedInputBinding<float>>& FloatUpdateRangedInputBindings)
+{
+	const FNiagaraFastPath_Module_SpawnBurstInstantaneous& SpawnBurstInstantaneousInputs = Emitter->SpawnBurstInstantaneous[ModuleIndex];
+	InitInput(SpawnBurstInstantaneousInputs.SpawnCount, SpawnCount, FNiagaraTypeDefinition::GetIntDef(), InstanceParameters, IntUserParameterInputBindings, IntUpdateRangedInputBindings);
+	InitInput(SpawnBurstInstantaneousInputs.SpawnTime, SpawnTime, FNiagaraTypeDefinition::GetFloatDef(), InstanceParameters, FloatUserParameterInputBindings, FloatUpdateRangedInputBindings);
+	SpawnGroup = SpawnBurstInstantaneousInputs.SpawnGroup;
+}
+
+void FNiagaraEmitterFastPath::InitFastPathAttributeBindings(
+	const FNiagaraFastPathAttributeNames& SourceAttributeNames,
+	FNiagaraParameterStore& TargetParameterStore,
+	FNiagaraSystemFastPath::FParamMap0& SystemMap,
+	FNiagaraEmitterFastPath::FParamMap0& EmitterMap,
+	TArray<TNiagaraFastPathAttributeBinding<int32>>& FastPathIntAttributeBindings,
+	TArray<TNiagaraFastPathAttributeBinding<float>>& FastPathFloatAttributeBindings)
+{
+	for (int32 SystemParameterNameIndex = 0; SystemParameterNameIndex < SourceAttributeNames.System.Num(); ++SystemParameterNameIndex)
+	{
+		FName SystemParameterName = SourceAttributeNames.System[SystemParameterNameIndex];
+		FName SystemParameterFullName = SourceAttributeNames.SystemFullNames[SystemParameterNameIndex];
+		if (SystemParameterName == FNiagaraSystemFastPath::FParameterNames::ExecutionState)
+		{
+			AddBinding(SystemParameterFullName, FNiagaraTypeDefinition::GetExecutionStateEnum(), (int32*)&(SystemMap.System.ExecutionState), TargetParameterStore, FastPathIntAttributeBindings);
+		}
+		else if (SystemParameterName == FNiagaraSystemFastPath::FParameterNames::ExecutionStateSource)
+		{
+			AddBinding(SystemParameterFullName, FNiagaraTypeDefinition::GetExecutionStateSouceEnum(), (int32*)&(SystemMap.System.ExecutionStateSource), TargetParameterStore, FastPathIntAttributeBindings);
+		}
+	}
+
+	for (int32 EmitterParameterNameIndex = 0; EmitterParameterNameIndex < SourceAttributeNames.Emitter.Num(); ++EmitterParameterNameIndex)
+	{
+		FName EmitterParameterName = SourceAttributeNames.Emitter[EmitterParameterNameIndex];
+		FName EmitterParameterFullName = SourceAttributeNames.EmitterFullNames[EmitterParameterNameIndex];
+		if (EmitterParameterName == FNiagaraEmitterFastPath::FAttributeNames::Age)
+		{
+			AddBinding(EmitterParameterFullName, FNiagaraTypeDefinition::GetFloatDef(), &(EmitterMap.Emitter.Age), TargetParameterStore, FastPathFloatAttributeBindings);
+		}
+		else if (EmitterParameterName == FNiagaraEmitterFastPath::FAttributeNames::CurrentLoopDelay)
+		{
+			AddBinding(EmitterParameterFullName, FNiagaraTypeDefinition::GetFloatDef(), &(EmitterMap.Emitter.CurrentLoopDelay), TargetParameterStore, FastPathFloatAttributeBindings);
+		}
+		else if (EmitterParameterName == FNiagaraEmitterFastPath::FAttributeNames::CurrentLoopDuration)
+		{
+			AddBinding(EmitterParameterFullName, FNiagaraTypeDefinition::GetFloatDef(), &(EmitterMap.Emitter.CurrentLoopDuration), TargetParameterStore, FastPathFloatAttributeBindings);
+		}
+		else if (EmitterParameterName == FNiagaraEmitterFastPath::FAttributeNames::ExecutionState)
+		{
+			AddBinding(EmitterParameterFullName, FNiagaraTypeDefinition::GetExecutionStateEnum(), (int32*)&(EmitterMap.Emitter.ExecutionState), TargetParameterStore, FastPathIntAttributeBindings);
+		}
+		else if (EmitterParameterName == FNiagaraEmitterFastPath::FAttributeNames::ExecutionStateSource)
+		{
+			AddBinding(EmitterParameterFullName, FNiagaraTypeDefinition::GetExecutionStateSouceEnum(), (int32*)&(EmitterMap.Emitter.ExecutionStateSource), TargetParameterStore, FastPathIntAttributeBindings);
+		}
+		else if (EmitterParameterName == FNiagaraEmitterFastPath::FAttributeNames::LoopCount)
+		{
+			AddBinding(EmitterParameterFullName, FNiagaraTypeDefinition::GetIntDef(), &(EmitterMap.Emitter.LoopCount), TargetParameterStore, FastPathIntAttributeBindings);
+		}
+		else if (EmitterParameterName == FNiagaraEmitterFastPath::FAttributeNames::LoopedAge)
+		{
+			AddBinding(EmitterParameterFullName, FNiagaraTypeDefinition::GetFloatDef(), &(EmitterMap.Emitter.LoopedAge), TargetParameterStore, FastPathFloatAttributeBindings);
+		}
+		else if (EmitterParameterName == FNiagaraEmitterFastPath::FAttributeNames::NormalizedLoopAge)
+		{
+			AddBinding(EmitterParameterFullName, FNiagaraTypeDefinition::GetFloatDef(), &(EmitterMap.Emitter.NormalizedLoopAge), TargetParameterStore, FastPathFloatAttributeBindings);
+		}
+		else if (EmitterParameterName == FNiagaraEmitterFastPath::FAttributeNames::DistanceTravelled)
+		{
+			AddBinding(EmitterParameterFullName, FNiagaraTypeDefinition::GetFloatDef(), &(EmitterMap.Emitter.DistanceTraveled), TargetParameterStore, FastPathFloatAttributeBindings);
+		}
+		else if (EmitterParameterName == FNiagaraEmitterFastPath::FAttributeNames::Scalability_DistanceFraction)
+		{
+			AddBinding(EmitterParameterFullName, FNiagaraTypeDefinition::GetFloatDef(), &(EmitterMap.Emitter.Scalability.DistanceFraction), TargetParameterStore, FastPathFloatAttributeBindings);
+		}
+	}
+}
+
+void FNiagaraEmitterFastPath::SetSpawnMapDefaults(const UNiagaraEmitter& Emitter, FParamMap0& Map)
+{
+	Map.Emitter.Age = 0;
+	Map.Emitter.CurrentLoopDelay = 0;
+	Map.Emitter.CurrentLoopDuration = 0;
+	Map.Emitter.ExecutionState = (ENiagaraExecutionState)0;
+	Map.Emitter.ExecutionStateSource = (ENiagaraExecutionStateSource)0;
+	Map.Emitter.LoopCount = 0;
+	Map.Emitter.LoopedAge = 0.0f;
+	Map.Emitter.NormalizedLoopAge = 0.0f;
+	Map.Emitter.DistanceTraveled = 0;
+	Map.Emitter.Scalability.DistanceFraction = 1;
+	Map.Scalability.Emitter.ExecutionState = (ENiagaraExecutionState)0;
+	Map.Scalability.Emitter.SpawnCountScale = 1;
+	Map.Emitter.SpawnRate.Empty();
+	Map.Emitter.SpawnRate.AddDefaulted(Emitter.SpawnRate.Num());
+	Map.Emitter.SpawnPerUnit.Empty();
+	Map.Emitter.SpawnPerUnit.AddDefaulted(Emitter.SpawnPerUnit.Num());
+	Map.Emitter.SpawnBurst_Instantaneous.Empty();
+	Map.Emitter.SpawnBurst_Instantaneous.AddDefaulted(Emitter.SpawnBurstInstantaneous.Num());
+}
+
+void FNiagaraEmitterFastPath::SetUpdateMapDefaults(FParamMap0& Map)
+{
+	/* TODO?
+	Map.Emitter.SpawnRate.SpawnOutputInfo.Count = 0;
+	Map.Emitter.SpawnRate.SpawnOutputInfo.InterpStartDt = 0.0f;
+	Map.Emitter.SpawnRate.SpawnOutputInfo.IntervalDt = 0.0f;
+	Map.Emitter.SpawnRate.SpawnOutputInfo.SpawnGroup = 0;
+	*/
+}
+
+void FNiagaraEmitterFastPath::Function_SampleCurve_SpawnCountScaleByDistanceFraction(const FNiagaraFastPath_Module_EmitterScalability& Context_Map_EmitterScalability, float InFraction, float& OutScale)
+{
+	OutScale = FMath::Lerp(Context_Map_EmitterScalability.ScaleForMinFraction, Context_Map_EmitterScalability.ScaleForMaxFraction, InFraction);
+}
+
+void FNiagaraEmitterFastPath::Module_EmitterScalability(const FNiagaraFastPath_Module_EmitterScalability& Context_Map_EmitterScalability, FParamMap0& Context_Map)
+{
+	float Constant13 = 0;
+	float float_IfResult;
+	if (Context_Map_EmitterScalability.bUseMinDistance)
+	{
+		float_IfResult = Context_Map_EmitterScalability.MinDistance;
+	}
+	else
+	{
+		float_IfResult = Constant13;
+	}
+	float Result9 = Context_Map.Engine.Owner.LODDistance - float_IfResult;
+	float Constant14 = 1e+12;
+	float float_IfResult1;
+	if (Context_Map_EmitterScalability.bUseMaxDistance)
+	{
+		float_IfResult1 = Context_Map_EmitterScalability.MaxDistance;
+	}
+	else
+	{
+		float_IfResult1 = Constant14;
+	}
+	float Result10 = float_IfResult1 - float_IfResult;
+	float Result11 = Result9 / Result10;
+	Context_Map.Emitter.Scalability.DistanceFraction = Result11;
+	// Value is initialized in input struct constructor now so ignore this to avoid overwritting a value set on the emitter.
+	// ENiagaraExecutionState Constant15 = ENiagaraExecutionState::Inactive;
+	// Context_Map_EmitterScalability.MinCulledState = Constant15;
+	ENiagaraExecutionState Constant16 = ENiagaraExecutionState::Active;
+	Context_Map.Scalability.Emitter.ExecutionState = Constant16;
+	float Constant17 = 0;
+	bool Result12 = Context_Map.Emitter.Scalability.DistanceFraction < Constant17;
+	bool Result13 = Result12 && Context_Map_EmitterScalability.bUseMinDistance;
+	ENiagaraExecutionState ENiagaraExecutionState_IfResult3;
+	if (Result13)
+	{
+		ENiagaraExecutionState_IfResult3 = Context_Map_EmitterScalability.MinCulledState;
+	}
+	else
+	{
+		ENiagaraExecutionState_IfResult3 = Context_Map.Scalability.Emitter.ExecutionState;
+	}
+	Context_Map.Scalability.Emitter.ExecutionState = ENiagaraExecutionState_IfResult3;
+	// Value is initialized in input struct constructor now so ignore this to avoid overwriting a value set on the emitter.
+	// ENiagaraExecutionState Constant18 = ENiagaraExecutionState::Inactive;
+	// Context_Map_EmitterScalability.MaxCulledState = Constant18;
+	float Constant19 = 1;
+	bool Result14 = Context_Map.Emitter.Scalability.DistanceFraction > Constant19;
+	bool Result15 = Result14 && Context_Map_EmitterScalability.bUseMaxDistance;
+	ENiagaraExecutionState ENiagaraExecutionState_IfResult4;
+	if (Result15)
+	{
+		ENiagaraExecutionState_IfResult4 = Context_Map_EmitterScalability.MaxCulledState;
+	}
+	else
+	{
+		ENiagaraExecutionState_IfResult4 = Context_Map.Scalability.Emitter.ExecutionState;
+	}
+	Context_Map.Scalability.Emitter.ExecutionState = ENiagaraExecutionState_IfResult4;
+	float Constant20 = 1;
+	float float_IfResult2;
+	if (Context_Map_EmitterScalability.bApplySpawnCountScale)
+	{
+		float_IfResult2 = Context_Map_EmitterScalability.SpawnCountScale;
+	}
+	else
+	{
+		float_IfResult2 = Constant20;
+	}
+	float SampleCurve_SpawnCountScaleByDistanceFractionOutput_Value;
+	Function_SampleCurve_SpawnCountScaleByDistanceFraction(Context_Map_EmitterScalability, Context_Map.Emitter.Scalability.DistanceFraction, SampleCurve_SpawnCountScaleByDistanceFractionOutput_Value);
+	float Constant21 = 1;
+	float float_IfResult3;
+	if (Context_Map_EmitterScalability.bApplySpawnCountScaleByDistanceFraction)
+	{
+		float_IfResult3 = SampleCurve_SpawnCountScaleByDistanceFractionOutput_Value;
+	}
+	else
+	{
+		float_IfResult3 = Constant21;
+	}
+	float Result16 = float_IfResult2 * float_IfResult3;
+	Context_Map.Scalability.Emitter.SpawnCountScale = Result16;
+}
+
+void FNiagaraEmitterFastPath::Function_EmitterLifeCycle_EmitterChangeState(ENiagaraExecutionState In_NewState, bool In_Condition, ENiagaraExecutionStateSource In_NewStateSource, FParamMap0& Context_Map)
+{
+	ENiagaraExecutionState Constant33 = ENiagaraExecutionState::Complete;
+	bool Result32 = Context_Map.Emitter.ExecutionState != Constant33;
+	ENiagaraExecutionState Constant34 = ENiagaraExecutionState::Disabled;
+	bool Result33 = Context_Map.Emitter.ExecutionState != Constant34;
+	bool Result34 = Result32 && Result33;
+	bool Context_Map_Local_Emitter_CanChangeState = Result34;
+	bool Result35 = In_Condition && Context_Map_Local_Emitter_CanChangeState;
+	bool Result36 = Context_Map.Emitter.ExecutionStateSource <= In_NewStateSource;
+	bool Result37 = Result35 && Result36;
+	ENiagaraExecutionState ENiagaraExecutionState_IfResult6;
+	ENiagaraExecutionStateSource ENiagaraExecutionStateSource_IfResult3;
+	if (Result37)
+	{
+		ENiagaraExecutionState_IfResult6 = In_NewState;
+		ENiagaraExecutionStateSource_IfResult3 = In_NewStateSource;
+	}
+	else
+	{
+		ENiagaraExecutionState_IfResult6 = Context_Map.Emitter.ExecutionState;
+		ENiagaraExecutionStateSource_IfResult3 = Context_Map.Emitter.ExecutionStateSource;
+	}
+	Context_Map.Emitter.ExecutionState = ENiagaraExecutionState_IfResult6;
+	Context_Map.Emitter.ExecutionStateSource = ENiagaraExecutionStateSource_IfResult3;
+}
+
+void FNiagaraEmitterFastPath::Module_EmitterLifeCycle(const FNiagaraFastPath_Module_EmitterLifeCycle& Context_Map_EmitterLifeCycle, FParamMap0& Context_Map)
+{
+	float Constant24 = 0;
+	bool Result19 = Context_Map.Emitter.Age == Constant24;
+	float Result20 = -(Context_Map_EmitterLifeCycle.NextLoopDelay);
+	float Constant25 = 0;
+	float Constant26 = 0;
+	float float_IfResult4;
+	float float001_IfResult;
+	float float002_IfResult;
+	float float003_IfResult;
+	if (Result19)
+	{
+		float_IfResult4 = Result20;
+		float001_IfResult = Context_Map_EmitterLifeCycle.NextLoopDuration;
+		float002_IfResult = Context_Map_EmitterLifeCycle.NextLoopDelay;
+		float003_IfResult = Constant25;
+	}
+	else
+	{
+		float_IfResult4 = Context_Map.Emitter.LoopedAge;
+		float001_IfResult = Context_Map.Emitter.CurrentLoopDuration;
+		float002_IfResult = Context_Map.Emitter.CurrentLoopDelay;
+		float003_IfResult = Constant26;
+	}
+	Context_Map.Emitter.LoopedAge = float_IfResult4;
+	Context_Map.Emitter.CurrentLoopDuration = float001_IfResult;
+	Context_Map.Emitter.CurrentLoopDelay = float002_IfResult;
+	float Result21 = Context_Map.Emitter.Age + Context_Map.Engine.DeltaTime;
+	float Result22 = Context_Map.Engine.DeltaTime + Context_Map.Emitter.LoopedAge;
+	float Result23 = Result22 / Context_Map.Emitter.CurrentLoopDuration;
+	int32 Count;
+	Count = Result23;
+	int32 Constant27 = 0;
+	int32 Result24 = FMath::Max(Count, Constant27);
+	float Result25 = Result24 * Context_Map.Emitter.CurrentLoopDuration;
+	float Result26 = Result22 - Result25;
+	int32 Result27 = Result24 + Context_Map.Emitter.LoopCount;
+	Context_Map.Emitter.Age = Result21;
+	Context_Map.Emitter.LoopedAge = Result26;
+	Context_Map.Emitter.LoopCount = Result27;
+	int32 Constant28 = 0;
+	bool Result28 = Result24 > Constant28;
+	float float_IfResult5;
+	if (Context_Map_EmitterLifeCycle.bDurationRecalcEachLoop)
+	{
+		float_IfResult5 = Context_Map_EmitterLifeCycle.NextLoopDuration;
+	}
+	else
+	{
+		float_IfResult5 = Context_Map.Emitter.CurrentLoopDuration;
+	}
+	float Constant31 = 0;
+	float float_IfResult6;
+	if (Context_Map_EmitterLifeCycle.bDelayFirstLoopOnly)
+	{
+		float_IfResult6 = Constant31;
+	}
+	else
+	{
+		float_IfResult6 = Context_Map_EmitterLifeCycle.NextLoopDelay;
+	}
+	float Result29 = Context_Map.Emitter.LoopedAge - float_IfResult6;
+	float float_IfResult7;
+	float float001_IfResult1;
+	float float002_IfResult1;
+	if (Result28)
+	{
+		float_IfResult7 = float_IfResult5;
+		float001_IfResult1 = float_IfResult6;
+		float002_IfResult1 = Result29;
+	}
+	else
+	{
+		float_IfResult7 = Context_Map.Emitter.CurrentLoopDuration;
+		float001_IfResult1 = Context_Map.Emitter.CurrentLoopDelay;
+		float002_IfResult1 = Context_Map.Emitter.LoopedAge;
+	}
+	float Result30 = float002_IfResult1 / float_IfResult7;
+	Context_Map.Emitter.CurrentLoopDuration = float_IfResult7;
+	Context_Map.Emitter.CurrentLoopDelay = float001_IfResult1;
+	Context_Map.Emitter.LoopedAge = float002_IfResult1;
+	Context_Map.Emitter.NormalizedLoopAge = Result30;
+	ENiagaraExecutionStateSource Constant32 = ENiagaraExecutionStateSource::Scalability;
+	bool Result31 = Context_Map.System->ExecutionStateSource != Constant32;
+	Function_EmitterLifeCycle_EmitterChangeState(Context_Map.System->ExecutionState, Result31, Context_Map.System->ExecutionStateSource, Context_Map);
+	ENiagaraExecutionStateSource Constant35 = ENiagaraExecutionStateSource::Scalability;
+	bool Result38 = Context_Map.System->ExecutionStateSource == Constant35;
+	ENiagaraExecutionState Constant36 = (ENiagaraExecutionState::Active);
+	bool Result39 = Context_Map.System->ExecutionState != Constant36;
+	bool Result40 = Result38 && Result39;
+	ENiagaraExecutionState ENiagaraExecutionState_IfResult7;
+	if (Result40)
+	{
+		ENiagaraExecutionState_IfResult7 = Context_Map.System->ExecutionState;
+	}
+	else
+	{
+		ENiagaraExecutionState_IfResult7 = Context_Map.Scalability.Emitter.ExecutionState;
+	}
+	Context_Map.Scalability.Emitter.ExecutionState = ENiagaraExecutionState_IfResult7;
+	bool Constant37 = true;
+	ENiagaraExecutionStateSource Constant38 = ENiagaraExecutionStateSource::Scalability;
+	Function_EmitterLifeCycle_EmitterChangeState(Context_Map.Scalability.Emitter.ExecutionState, Constant37, Constant38, Context_Map);
+	ENiagaraExecutionState Constant39 = ENiagaraExecutionState::Inactive;
+	bool Constant40 = true;
+	Context_Map.Spawning.bCanEverSpawn = Constant40;
+	int32 Constant41 = 0;
+	bool Result44 = Context_Map_EmitterLifeCycle.MaxLoopCount > Constant41;
+	bool Result45 = Context_Map.Emitter.LoopCount >= Context_Map_EmitterLifeCycle.MaxLoopCount;
+	bool Result46 = Result44 && Result45;
+	ENiagaraExecutionStateSource Constant42 = ENiagaraExecutionStateSource::Internal;
+	Function_EmitterLifeCycle_EmitterChangeState(Constant39, Result46, Constant42, Context_Map);
+	ENiagaraExecutionState Constant43 = ENiagaraExecutionState::Complete;
+	// Values are initialized in input struct constructor now so ignore this to avoid overwriting a value set on the emitter.
+	// bool Constant42 = true;
+	// Context_Map_EmitterLifeCycle.AutoComplete = Constant42;
+	// bool Constant43 = false;
+	// Context_Map_EmitterLifeCycle.CompleteOnInactive = Constant43;
+	int32 Constant46 = 0;
+	bool Result50 = Context_Map.Engine.Emitter.NumParticles == Constant46;
+	ENiagaraExecutionState Constant47 = ENiagaraExecutionState::Active;
+	bool Result51 = Context_Map.Emitter.ExecutionState != Constant47;
+	bool Result52 = Result50 && Result51;
+	bool Result53 = Result52 && Context_Map_EmitterLifeCycle.bAutoComplete;
+	bool Result54 = Context_Map_EmitterLifeCycle.bCompleteOnInactive && Result51;
+	bool Result55 = Result53 || Result54;
+	ENiagaraExecutionStateSource Constant48 = ENiagaraExecutionStateSource::InternalCompletion;
+	Function_EmitterLifeCycle_EmitterChangeState(Constant43, Result55, Constant48, Context_Map);
+}
+
+void FNiagaraEmitterFastPath::Module_SpawnRate(FParamMap0& Context_Map)
+{
+	for (int32 SpawnRateIndex = 0; SpawnRateIndex < Context_Map.SpawnRate.Num(); ++SpawnRateIndex)
+	{
+		const FParamMap0_SpawnRate& Context_Map_SpawnRate = Context_Map.SpawnRate[SpawnRateIndex];
+		FParamMap0_Emitter_SpawnRate& Context_Map_Emitter_SpawnRate = Context_Map.Emitter.SpawnRate[SpawnRateIndex];
+
+		float Result57 = Context_Map_SpawnRate.SpawnRate * Context_Map.Scalability.Emitter.SpawnCountScale * Context_Map.Engine.GlobalSpawnCountScale;
+		float Result58 = 1.0f / Result57;// Reciprocal(Result57);
+		float Result59 = 1 - Context_Map_Emitter_SpawnRate.SpawnRemainder;
+		float Result60 = Result58 * Result59;
+		float Context_Map_Local_SpawnRate_SpawnRate = Result57;
+		float Context_Map_Local_SpawnRate_IntervalDT = Result58;
+		float Context_Map_Local_SpawnRate_InterpStartDT = Result60;
+		float Constant47 = 0;
+		bool Result61 = Context_Map.Emitter.LoopedAge >= Constant47;
+		float Constant48 = 1;
+		float Constant49 = 0;
+		float float_IfResult8;
+		if (Result61)
+		{
+			float_IfResult8 = Constant48;
+		}
+		else
+		{
+			float_IfResult8 = Constant49;
+		}
+		float Result62 = Context_Map_Local_SpawnRate_SpawnRate * float_IfResult8;
+		float Result63 = Result62 * Context_Map.Engine.DeltaTime + Context_Map_Emitter_SpawnRate.SpawnRemainder;
+		int32 Result64 = floor(Result63);
+		float Result65 = Result63 - Result64;
+		int32 Context_Map_Local_SpawnRate_SpawnCount = Result64;
+		Context_Map_Emitter_SpawnRate.SpawnRemainder = Result65;
+		FNiagaraSpawnInfo Output1;
+		Output1.Count = Context_Map_Local_SpawnRate_SpawnCount;
+		Output1.InterpStartDt = Context_Map_Local_SpawnRate_InterpStartDT;
+		Output1.IntervalDt = Context_Map_Local_SpawnRate_IntervalDT;
+		Output1.SpawnGroup = Context_Map_SpawnRate.SpawnGroup;
+		bool Constant50 = true;
+		Context_Map_Emitter_SpawnRate.SpawnOutputInfo = Output1;
+		Context_Map.Spawning.bCanEverSpawn = Constant50;
+	}
+}
+
+void FNiagaraEmitterFastPath::Module_SpawnPerUnit(FParamMap0& Context_Map)
+{
+	for (int32 SpawnPerUnitIndex = 0; SpawnPerUnitIndex < Context_Map.SpawnPerUnit.Num(); ++SpawnPerUnitIndex)
+	{
+		const FParamMap0_SpawnPerUnit& Context_Map_SpawnPerUnit = Context_Map.SpawnPerUnit[SpawnPerUnitIndex];
+		FParamMap0_Emitter_SpawnPerUnit& Context_Map_Emitter_SpawnPerUnit = Context_Map.Emitter.SpawnPerUnit[SpawnPerUnitIndex];
+
+		// Not exposed on the fast path module.
+		FVector Context_Map_SpawnPerUnit_VelocityVector = Context_Map.Engine.Owner.Velocity;
+		float Context_Map_SpawnPerUnit_DeltaTime = Context_Map.Engine.DeltaTime;
+
+		// Value is initialized in input struct constructor now so ignore this to avoid overwriting a value set on the emitter.
+		// bool Constant51 = false;
+		// Context_Map_SpawnPerUnit.UseMovementTolerance = Constant51;
+		float Result66 = Context_Map_SpawnPerUnit_VelocityVector.Size();
+		float Context_Map_Local_SpawnPerUnit_VelocityVectorLength = Result66;
+		float Result67 = Result66 * Context_Map_SpawnPerUnit_DeltaTime;
+		bool Result68 = Result67 > Context_Map_SpawnPerUnit.MovementTolerance;
+		float Constant52 = 0;
+		float float_IfResult9;
+		if (Result68)
+		{
+			float_IfResult9 = Result66;
+		}
+		else
+		{
+			float_IfResult9 = Constant52;
+		}
+		float float_IfResult10;
+		if (Context_Map_SpawnPerUnit.bUseMovementTolerance)
+		{
+			float_IfResult10 = float_IfResult9;
+		}
+		else
+		{
+			float_IfResult10 = Result66;
+		}
+		float Context_Map_Local_SpawnPerUnit_MovementThresholdVectorLength = float_IfResult10;
+		float Constant53 = 100;
+		float Result69 = Context_Map_Local_SpawnPerUnit_VelocityVectorLength / Constant53;
+		float Constant54 = 500000;
+		float Result70 = FMath::Fmod(Context_Map.Emitter.DistanceTraveled, Constant54);
+		float Result71 = Result69 + Result70;
+		float Result72 = 1.0f / Context_Map_SpawnPerUnit.SpawnPerUnit; // Reciprocal(Context_Map_SpawnPerUnit.SpawnPerUnit);
+		float Result73 = Context_Map_Local_SpawnPerUnit_MovementThresholdVectorLength * Result72;
+		Context_Map.Emitter.DistanceTraveled = Result71;
+		float Context_Map_Local_SpawnPerUnit_SpawnSpacing = Result73;
+		float Result74 = Context_Map.Scalability.Emitter.SpawnCountScale * Context_Map_Local_SpawnPerUnit_SpawnSpacing;
+		float Result75 = Result74 * Context_Map.Engine.DeltaTime + Context_Map_Emitter_SpawnPerUnit.SpawnRemainder;
+		int32 Result76 = FMath::FloorToInt(Result75);
+		float Result77 = Result75 - Result76;
+		float Result78 = 1.0f / Context_Map_Local_SpawnPerUnit_SpawnSpacing; // Reciprocal(Context.Map.Local.SpawnPerUnit.SpawnSpacing);
+		float Result79 = 1 - Context_Map_Emitter_SpawnPerUnit.SpawnRemainder;
+		float Result80 = Result79 * Result78;
+		int32 Context_Map_Local_SpawnPerUnit_SpawnCountInt = Result76;
+		Context_Map_Emitter_SpawnPerUnit.SpawnRemainder = Result77;
+		float Context_Map_Local_SpawnPerUnit_IntervalDt = Result78;
+		float Context_Map_Local_SpawnPerUnit_InterpStartDt = Result80;
+		// Value is initialized in input struct constructor now so ignore this to avoid overwriting a value set on the emitter.
+		//bool Constant55 = true;
+		//Context_Map_SpawnPerUnit.bUseMaxMovementThreshold = Constant55;
+		bool Result81 = Context_Map_Local_SpawnPerUnit_VelocityVectorLength > Context_Map_SpawnPerUnit.MaxMovementThreshold;
+		int32 Constant56 = 0;
+		int32 int32_IfResult;
+		if (Result81)
+		{
+			int32_IfResult = Constant56;
+		}
+		else
+		{
+			int32_IfResult = Context_Map_Local_SpawnPerUnit_SpawnCountInt;
+		}
+		int32 int32_IfResult1;
+		if (Context_Map_SpawnPerUnit.bUseMaxMovementThreshold)
+		{
+			int32_IfResult1 = int32_IfResult;
+		}
+		else
+		{
+			int32_IfResult1 = Context_Map_Local_SpawnPerUnit_SpawnCountInt;
+		}
+		FNiagaraSpawnInfo Output11;
+		Output11.Count = int32_IfResult1;
+		Output11.InterpStartDt = Context_Map_Local_SpawnPerUnit_InterpStartDt;
+		Output11.IntervalDt = Context_Map_Local_SpawnPerUnit_IntervalDt;
+		Output11.SpawnGroup = Context_Map_SpawnPerUnit.SpawnGroup;
+		bool Constant57 = true;
+		int32 Constant58 = 0;
+		bool Result82 = int32_IfResult1 > Constant58;
+		Context_Map_Emitter_SpawnPerUnit.SpawnOutputInfo = Output11;
+		Context_Map.Spawning.bCanEverSpawn = Constant57;
+		//Context.Map.OUTPUT_VAR.SpawnPerUnit.HasSpawnedThisFrame = Result82;
+	}
+}
+
+void FNiagaraEmitterFastPath::Module_SpawnBurstInstantaneous(FParamMap0& Context_Map)
+{
+	for (int32 SpawnBurstInstantaneousIndex = 0; SpawnBurstInstantaneousIndex < Context_Map.SpawnBurst_Instantaneous.Num(); ++SpawnBurstInstantaneousIndex)
+	{
+		const FParamMap0_SpawnBurst_Instantaneous& Context_Map_SpawnBurst_Instantaneous = Context_Map.SpawnBurst_Instantaneous[SpawnBurstInstantaneousIndex];
+		FParamMap0_Emitter_SpawnBurst_Instantaneous& Context_Map_Emitter_SpawnBurst_Instantaneous = Context_Map.Emitter.SpawnBurst_Instantaneous[SpawnBurstInstantaneousIndex];
+
+		// Not exposed on the fast path module
+		float Context_Map_SpawnBurst_Instantaneous_Age = Context_Map.Emitter.LoopedAge;
+
+		float Result83 = Context_Map_SpawnBurst_Instantaneous_Age - Context_Map.Engine.DeltaTime;
+		float Result84 = Context_Map_SpawnBurst_Instantaneous.SpawnTime - Result83;
+		float Constant59 = 0;
+		float Output12;
+		Output12 = Constant59;
+		bool Result85 = Result84 >= Output12;
+		float Result86 = Context_Map_SpawnBurst_Instantaneous.SpawnTime - Context_Map_SpawnBurst_Instantaneous_Age;
+		bool Result87 = Result86 < Output12;
+		bool Result88 = Result85 && Result87;
+		int32 Constant60 = 0;
+		bool Result89 = Context_Map_SpawnBurst_Instantaneous.SpawnCount == Constant60;
+		float Constant61 = 0;
+		float Result90 = Context_Map_SpawnBurst_Instantaneous.SpawnCount * Context_Map.Scalability.Emitter.SpawnCountScale * Context_Map.Engine.GlobalSpawnCountScale;
+		float Constant62 = 1;
+		float Result91 = FMath::Max(Result90, Constant62);
+		float float_IfResult11;
+		if (Result89)
+		{
+			float_IfResult11 = Constant61;
+		}
+		else
+		{
+			float_IfResult11 = Result91;
+		}
+		float Constant63 = 0;
+		FNiagaraSpawnInfo Output13;
+		Output13.Count = float_IfResult11;
+		Output13.InterpStartDt = Result84;
+		Output13.IntervalDt = Constant63;
+		Output13.SpawnGroup = Context_Map_SpawnBurst_Instantaneous.SpawnGroup;
+		int32 Constant64 = 0;
+		float Constant65 = 0;
+		float Constant66 = 0;
+		FNiagaraSpawnInfo Output14;
+		Output14.Count = Constant64;
+		Output14.InterpStartDt = Constant65;
+		Output14.SpawnGroup = Context_Map_SpawnBurst_Instantaneous.SpawnGroup;
+		Output14.IntervalDt = Constant66;
+		FNiagaraSpawnInfo SpawnInfo_IfResult;
+		if (Result88)
+		{
+			SpawnInfo_IfResult = Output13;
+		}
+		else
+		{
+			SpawnInfo_IfResult = Output14;
+		}
+		Context_Map_Emitter_SpawnBurst_Instantaneous.SpawnBurst = SpawnInfo_IfResult;
+		bool Result92 = Context_Map.Emitter.LoopedAge <= Context_Map_SpawnBurst_Instantaneous.SpawnTime;
+		bool Result93 = Context_Map.Spawning.bCanEverSpawn || Result92;
+		Context_Map.Spawning.bCanEverSpawn = Result93;
+	}
+}
\ No newline at end of file
diff --git a/Engine/Plugins/FX/Niagara/Source/Niagara/Private/NiagaraSystemInstance.cpp b/Engine/Plugins/FX/Niagara/Source/Niagara/Private/NiagaraSystemInstance.cpp
--- a/Engine/Plugins/FX/Niagara/Source/Niagara/Private/NiagaraSystemInstance.cpp
+++ b/Engine/Plugins/FX/Niagara/Source/Niagara/Private/NiagaraSystemInstance.cpp
@@ -824,6 +824,11 @@ void FNiagaraSystemInstance::ReInitInternal()
 	LocalBounds = FBox(FVector::ZeroVector, FVector::ZeroVector);
 	CachedDeltaSeconds = 0.0f;
 
+	FastPathIntUserParameterInputBindings.Empty();
+	FastPathFloatUserParameterInputBindings.Empty();
+	FastPathIntUpdateRangedInputBindings.Empty();
+	FastPathFloatUpdateRangedInputBindings.Empty();
+
 	UNiagaraSystem* System = GetSystem();
 	if (System == nullptr || Component == nullptr)
 	{
@@ -1452,6 +1457,37 @@ void FNiagaraSystemInstance::TickInstanceParameters(float DeltaSeconds)
 	InstanceParameters.MarkParametersDirty();
 }
 
+void FNiagaraSystemInstance::TickFastPathBindings()
+{
+	for (TNiagaraFastPathUserParameterInputBinding<int32>& IntUserParameterInputBinding : FastPathIntUserParameterInputBindings)
+	{
+		IntUserParameterInputBinding.Tick();
+	}
+
+	for (TNiagaraFastPathUserParameterInputBinding<float>& FloatUserParameterInputBinding : FastPathFloatUserParameterInputBindings)
+	{
+		FloatUserParameterInputBinding.Tick();
+	}
+
+	for (TNiagaraFastPathRangedInputBinding<int32>& IntUpdateRangedInputBinding : FastPathIntUpdateRangedInputBindings)
+	{
+		IntUpdateRangedInputBinding.Tick();
+	}
+
+	for (TNiagaraFastPathRangedInputBinding<float>& FloatUpdateRangedInputBinding : FastPathFloatUpdateRangedInputBindings)
+	{
+		FloatUpdateRangedInputBinding.Tick();
+	}
+}
+
+void FNiagaraSystemInstance::ResetFastPathBindings()
+{
+	FastPathIntUserParameterInputBindings.Empty();
+	FastPathFloatUserParameterInputBindings.Empty();
+	FastPathIntUpdateRangedInputBindings.Empty();
+	FastPathFloatUpdateRangedInputBindings.Empty();
+}
+
 #if WITH_EDITORONLY_DATA
 
 bool FNiagaraSystemInstance::UsesEmitter(const UNiagaraEmitter* Emitter)const
diff --git a/Engine/Plugins/FX/Niagara/Source/Niagara/Private/NiagaraSystemSimulation.cpp b/Engine/Plugins/FX/Niagara/Source/Niagara/Private/NiagaraSystemSimulation.cpp
--- a/Engine/Plugins/FX/Niagara/Source/Niagara/Private/NiagaraSystemSimulation.cpp
+++ b/Engine/Plugins/FX/Niagara/Source/Niagara/Private/NiagaraSystemSimulation.cpp
@@ -31,6 +31,7 @@ DECLARE_CYCLE_STAT(TEXT("System Sim Update [CNC]"), STAT_NiagaraSystemSim_Update
 DECLARE_CYCLE_STAT(TEXT("System Sim Spawn [CNC]"), STAT_NiagaraSystemSim_SpawnCNC, STATGROUP_Niagara);
 DECLARE_CYCLE_STAT(TEXT("System Sim Transfer Results [CNC]"), STAT_NiagaraSystemSim_TransferResultsCNC, STATGROUP_Niagara);
 DECLARE_CYCLE_STAT(TEXT("System Sim Init [GT]"), STAT_NiagaraSystemSim_Init, STATGROUP_Niagara);
+DECLARE_CYCLE_STAT(TEXT("System Sim FastPath [CNC]"), STAT_NiagaraSystemSim_FastPathCNC, STATGROUP_Niagara);
 
 DECLARE_CYCLE_STAT(TEXT("ForcedWaitForAsync"), STAT_NiagaraSystemSim_ForceWaitForAsync, STATGROUP_Niagara);
 DECLARE_CYCLE_STAT(TEXT("ForcedWait Fake Stall"), STAT_NiagaraSystemSim_ForceWaitFakeStall, STATGROUP_Niagara);
@@ -1010,16 +1011,25 @@ void FNiagaraSystemSimulation::Tick_Concurrent(FNiagaraSystemSimulationTickConte
 
 		FScopeCycleCounter SystemStatCounter(Context.System->GetStatID(true, true));
 
-		PrepareForSystemSimulate(Context);
 
-		if (Context.SpawnNum > 0)
+		if (Context.Owner->GetSystem()->FastPathMode != ENiagaraFastPathMode::ScriptVMOnly)
 		{
-			SpawnSystemInstances(Context);
+			TickFastPath(Context);
 		}
 
-		UpdateSystemInstances(Context);
+		if(Context.Owner->GetSystem()->FastPathMode != ENiagaraFastPathMode::FastPathOnly)
+		{
+			PrepareForSystemSimulate(Context);
+
+			if (Context.SpawnNum > 0)
+			{
+				SpawnSystemInstances(Context);
+			}
 
-		TransferSystemSimResults(Context);
+			UpdateSystemInstances(Context);
+
+			TransferSystemSimResults(Context);
+		}
 
 		for (FNiagaraSystemInstance* Instance : Context.Instances)
 		{
@@ -1059,6 +1069,218 @@ void FNiagaraSystemSimulation::Tick_Concurrent(FNiagaraSystemSimulationTickConte
 	}
 }
 
+void FNiagaraSystemSimulation::TickFastPath(FNiagaraSystemSimulationTickContext& Context)
+{
+	SCOPE_CYCLE_COUNTER(STAT_NiagaraSystemSim_FastPathCNC);
+	
+	// PrepareForSystemSimulate
+	for (FNiagaraSystemInstance* SystemInstance : Context.Instances)
+	{
+		FNiagaraSystemFastPath::FParamMap0& SystemMap = SystemInstance->GetFastPathMap();
+		SystemMap.Engine.Owner.ExecutionState = SystemInstance->GetRequestedExecutionState();
+		SystemMap.Engine.Owner.TimeSinceRendered = SystemInstance->GetSystemTimeSinceRendered();
+
+		TArray<TSharedRef<FNiagaraEmitterInstance, ESPMode::ThreadSafe>>& EmitterInstances = SystemInstance->GetEmitters();
+		for (int32 EmitterIndex = 0; EmitterIndex < EmitterInstances.Num(); ++EmitterIndex)
+		{
+			TSharedRef<FNiagaraEmitterInstance, ESPMode::ThreadSafe>& EmitterInstance = EmitterInstances[EmitterIndex];
+			if (EmitterInstance->GetExecutionState() == ENiagaraExecutionState::Disabled)
+			{
+				continue;
+			}
+			FNiagaraEmitterFastPath::FParamMap0& EmitterMap = EmitterInstance->GetFastPathMap();
+			EmitterMap.Engine.Owner.LODDistance = SystemInstance->GetOwnerLODDistance();
+			EmitterMap.Engine.DeltaTime = Context.DeltaSeconds;
+			EmitterMap.Engine.Emitter.NumParticles = SystemInstance->GetNumParticles(EmitterIndex);
+			EmitterMap.Engine.Owner.Velocity = SystemInstance->GetOwnerVelocity();
+			EmitterMap.Engine.GlobalSpawnCountScale = INiagaraModule::GetGlobalSpawnCountScale();
+			EmitterMap.Emitter.ExecutionState = EmitterInstance->GetExecutionState();
+		}
+	}
+
+	// SpawnSystemInstances
+	int32 CurrentNumInstances = Context.Instances.Num() - Context.SpawnNum;
+	if (Context.SpawnNum > 0)
+	{
+		// There is no simulation work to do here, but the number of instances in the data set needs to be correct for other codepaths.
+		FNiagaraDataSet& SimulationDataSet = Context.DataSet;
+		SimulationDataSet.BeginSimulate();
+		SimulationDataSet.Allocate(Context.Instances.Num());
+		SimulationDataSet.GetDestinationDataChecked().SetNumInstances(Context.Instances.Num());
+		SimulationDataSet.EndSimulate();
+
+		for(int32 SpawnIndex = CurrentNumInstances; SpawnIndex < Context.Instances.Num(); ++SpawnIndex)
+		{
+			FNiagaraSystemInstance* SystemInstance = Context.Instances[SpawnIndex];
+			SystemInstance->ResetFastPathBindings();
+
+			FNiagaraSystemFastPath::SetSpawnMapDefaults(SystemInstance->FastPathMap);
+			TArray<TSharedRef<FNiagaraEmitterInstance, ESPMode::ThreadSafe>>& EmitterInstances = SystemInstance->GetEmitters();
+			for (TSharedRef<FNiagaraEmitterInstance, ESPMode::ThreadSafe>& EmitterInstance : EmitterInstances)
+			{
+				const UNiagaraEmitter* Emitter = EmitterInstance->GetCachedEmitter();
+				if (EmitterInstance->GetExecutionState() == ENiagaraExecutionState::Disabled)
+				{
+					continue;
+				}
+
+				FNiagaraEmitterFastPath::FParamMap0& EmitterMap = EmitterInstance->GetFastPathMap();
+				FNiagaraEmitterFastPath::SetSpawnMapDefaults(*Emitter, EmitterMap);
+				EmitterInstance->InitFastPathAttributeBindings();
+
+				EmitterMap.SpawnRate.SetNumUninitialized(Emitter->SpawnRate.Num());
+				for (int32 SpawnRateIndex = 0; SpawnRateIndex < Emitter->SpawnRate.Num(); ++SpawnRateIndex)
+				{
+					EmitterMap.SpawnRate[SpawnRateIndex].Init(
+						Emitter, SpawnRateIndex, SystemInstance->GetInstanceParameters(),
+						SystemInstance->FastPathIntUserParameterInputBindings, SystemInstance->FastPathFloatUserParameterInputBindings,
+						SystemInstance->FastPathIntUpdateRangedInputBindings, SystemInstance->FastPathFloatUpdateRangedInputBindings);
+				}
+
+				EmitterMap.SpawnPerUnit.SetNumUninitialized(Emitter->SpawnPerUnit.Num());
+				for (int32 SpawnPerUnitIndex = 0; SpawnPerUnitIndex < Emitter->SpawnPerUnit.Num(); ++SpawnPerUnitIndex)
+				{
+					EmitterMap.SpawnPerUnit[SpawnPerUnitIndex].Init(
+						Emitter, SpawnPerUnitIndex, SystemInstance->GetInstanceParameters(),
+						SystemInstance->FastPathIntUserParameterInputBindings, SystemInstance->FastPathFloatUserParameterInputBindings,
+						SystemInstance->FastPathIntUpdateRangedInputBindings, SystemInstance->FastPathFloatUpdateRangedInputBindings);
+				}
+
+				EmitterMap.SpawnBurst_Instantaneous.SetNumUninitialized(Emitter->SpawnBurstInstantaneous.Num());
+				for (int32 SpawnBurstInstantaneousIndex = 0; SpawnBurstInstantaneousIndex < Emitter->SpawnBurstInstantaneous.Num(); ++SpawnBurstInstantaneousIndex)
+				{
+					EmitterMap.SpawnBurst_Instantaneous[SpawnBurstInstantaneousIndex].Init(
+						Emitter, SpawnBurstInstantaneousIndex, SystemInstance->GetInstanceParameters(),
+						SystemInstance->FastPathIntUserParameterInputBindings, SystemInstance->FastPathFloatUserParameterInputBindings,
+						SystemInstance->FastPathIntUpdateRangedInputBindings, SystemInstance->FastPathFloatUpdateRangedInputBindings);
+				}
+			}
+		}
+	}
+
+	// UpdateSystemInstances
+	const UNiagaraSystem* System = GetSystem();
+	auto UpdateSystemInstance = [System](FNiagaraSystemInstance* SystemInstance)
+	{
+		FNiagaraSystemFastPath::FParamMap0& SystemMap = SystemInstance->FastPathMap;
+		FNiagaraSystemFastPath::SetUpdateMapDefaults(SystemMap);
+		SystemInstance->TickFastPathBindings();
+
+		if (System->SystemScalability.bUseSystemScalability)
+		{
+			FNiagaraSystemFastPath::Module_SystemScalability(System->SystemScalability, SystemMap);
+		}
+
+		FNiagaraSystemFastPath::Module_SystemLifeCycle(System->SystemLifeCycle, SystemMap);
+
+		TArray<TSharedRef<FNiagaraEmitterInstance, ESPMode::ThreadSafe>>& EmitterInstances = SystemInstance->GetEmitters();
+		for (TSharedRef<FNiagaraEmitterInstance, ESPMode::ThreadSafe>& EmitterInstance : EmitterInstances)
+		{
+			if (EmitterInstance->GetExecutionState() == ENiagaraExecutionState::Disabled)
+			{
+				continue;
+			}
+
+			FNiagaraEmitterFastPath::FParamMap0& EmitterMap = EmitterInstance->GetFastPathMap();
+			FNiagaraEmitterFastPath::SetUpdateMapDefaults(EmitterMap);
+			EmitterMap.System = &SystemMap.System;
+
+			const UNiagaraEmitter* Emitter = EmitterInstance->GetCachedEmitter();
+			if (Emitter->EmitterScalability.bUseEmitterScalability)
+			{
+				FNiagaraEmitterFastPath::Module_EmitterScalability(Emitter->EmitterScalability, EmitterMap);
+			}
+
+			FNiagaraEmitterFastPath::Module_EmitterLifeCycle(Emitter->EmitterLifeCycle, EmitterMap);
+
+			if (Emitter->SpawnRate.Num() > 0)
+			{
+				FNiagaraEmitterFastPath::Module_SpawnRate(EmitterMap);
+			}
+
+			if (Emitter->SpawnPerUnit.Num() > 0)
+			{
+				FNiagaraEmitterFastPath::Module_SpawnPerUnit(EmitterMap);
+			}
+
+			if (Emitter->SpawnBurstInstantaneous.Num() > 0)
+			{
+				FNiagaraEmitterFastPath::Module_SpawnBurstInstantaneous(EmitterMap);
+			}
+		}
+	};
+
+	// Run update on current instances.
+	for (int32 CurrentIndex = 0; CurrentIndex < CurrentNumInstances; ++CurrentIndex)
+	{
+		FNiagaraSystemInstance* CurrentInstance = Context.Instances[CurrentIndex];
+		UpdateSystemInstance(CurrentInstance);
+	}
+
+	// Run updated on spawned instances with a small delta time.
+	for (int32 SpawnIndex = CurrentNumInstances; SpawnIndex < Context.Instances.Num(); ++SpawnIndex)
+	{
+		FNiagaraSystemInstance* SpawnedInstance = Context.Instances[SpawnIndex];
+		TArray<TSharedRef<FNiagaraEmitterInstance, ESPMode::ThreadSafe>>& EmitterInstances = SpawnedInstance->GetEmitters();
+		for (int32 EmitterIndex = 0; EmitterIndex < EmitterInstances.Num(); ++EmitterIndex)
+		{
+			TSharedRef<FNiagaraEmitterInstance, ESPMode::ThreadSafe>& EmitterInstance = EmitterInstances[EmitterIndex];
+			if (EmitterInstance->GetExecutionState() == ENiagaraExecutionState::Disabled)
+			{
+				continue;
+			}
+
+			EmitterInstance->GetFastPathMap().Engine.DeltaTime = 0.0001f;
+		}
+		UpdateSystemInstance(SpawnedInstance);
+	}
+
+	// TransferSystemSimResults
+	for (FNiagaraSystemInstance* SystemInstance : Context.Instances)
+	{
+		SystemInstance->SetActualExecutionState(SystemInstance->FastPathMap.System.ExecutionState);
+		if (SystemInstance->IsDisabled() == false)
+		{
+			TArray<TSharedRef<FNiagaraEmitterInstance, ESPMode::ThreadSafe>>& EmitterInstances = SystemInstance->GetEmitters();
+			for (TSharedRef<FNiagaraEmitterInstance, ESPMode::ThreadSafe>& EmitterInstance : EmitterInstances)
+			{
+				if (EmitterInstance->GetExecutionState() == ENiagaraExecutionState::Disabled)
+				{
+					continue;
+				}
+
+				FNiagaraEmitterFastPath::FParamMap0& EmitterMap = EmitterInstance->GetFastPathMap();
+
+				EmitterInstance->SetExecutionState(EmitterMap.Emitter.ExecutionState);
+
+				TArray<FNiagaraSpawnInfo>& EmitterInstSpawnInfos = EmitterInstance->GetSpawnInfo();
+
+				int32 SpawnInfoIndex = 0;
+
+				for (const FNiagaraEmitterFastPath::FParamMap0_Emitter_SpawnRate& SpawnRateOutput : EmitterMap.Emitter.SpawnRate)
+				{
+					EmitterInstSpawnInfos[SpawnInfoIndex] = SpawnRateOutput.SpawnOutputInfo;
+					SpawnInfoIndex++;
+				}
+
+				for (const FNiagaraEmitterFastPath::FParamMap0_Emitter_SpawnPerUnit& SpawnPerUnitOutput : EmitterMap.Emitter.SpawnPerUnit)
+				{
+					EmitterInstSpawnInfos[SpawnInfoIndex] = SpawnPerUnitOutput.SpawnOutputInfo;
+					SpawnInfoIndex++;
+				}
+
+				for (const FNiagaraEmitterFastPath::FParamMap0_Emitter_SpawnBurst_Instantaneous& SpawnBurstInstantaneousOutput : EmitterMap.Emitter.SpawnBurst_Instantaneous)
+				{
+					EmitterInstSpawnInfos[SpawnInfoIndex] = SpawnBurstInstantaneousOutput.SpawnBurst;
+					SpawnInfoIndex++;
+				}
+
+				EmitterInstance->TickFastPathAttributeBindings();
+			}
+		}
+	}
+}
+
 void FNiagaraSystemSimulation::PrepareForSystemSimulate(FNiagaraSystemSimulationTickContext& Context)
 {
 	SCOPE_CYCLE_COUNTER(STAT_NiagaraSystemSim_PrepareForSimulateCNC);
@@ -1076,9 +1298,13 @@ void FNiagaraSystemSimulation::PrepareForSystemSimulate(FNiagaraSystemSimulation
 	SpawnInstanceParameterDataSet.Allocate(NumInstances);
 	UpdateInstanceParameterDataSet.Allocate(NumInstances);
 
-	for (int32 EmitterIdx = 0; EmitterIdx < Context.System->GetNumEmitters(); ++EmitterIdx)
+	bool bUsingFastPath = GetSystem()->FastPathMode != ENiagaraFastPathMode::ScriptVMOnly;
+	if (bUsingFastPath == false)
 	{
-		EmitterExecutionStateAccessors[EmitterIdx].InitForAccess();
+		for (int32 EmitterIdx = 0; EmitterIdx < Context.System->GetNumEmitters(); ++EmitterIdx)
+		{
+			EmitterExecutionStateAccessors[EmitterIdx].InitForAccess();
+		}
 	}
 
 	//Tick instance parameters and transfer any needed into the system simulation dataset.
@@ -1097,13 +1323,16 @@ void FNiagaraSystemSimulation::PrepareForSystemSimulate(FNiagaraSystemSimulation
 		//TODO: Find good way to check that we're not using any instance parameter data interfaces in the system scripts here.
 		//In that case we need to solo and will never get here.
 
-		TArray<TSharedRef<FNiagaraEmitterInstance, ESPMode::ThreadSafe>>& Emitters = Inst->GetEmitters();
-		for (int32 EmitterIdx = 0; EmitterIdx < Emitters.Num(); ++EmitterIdx)
+		if (bUsingFastPath == false)
 		{
-			FNiagaraEmitterInstance& EmitterInst = Emitters[EmitterIdx].Get();
-			if (EmitterExecutionStateAccessors.Num() > EmitterIdx && EmitterExecutionStateAccessors[EmitterIdx].IsValidForWrite())
+			TArray<TSharedRef<FNiagaraEmitterInstance, ESPMode::ThreadSafe>>& Emitters = Inst->GetEmitters();
+			for (int32 EmitterIdx = 0; EmitterIdx < Emitters.Num(); ++EmitterIdx)
 			{
-				EmitterExecutionStateAccessors[EmitterIdx].Set(SystemIndex, (int32)EmitterInst.GetExecutionState());
+				FNiagaraEmitterInstance& EmitterInst = Emitters[EmitterIdx].Get();
+				if (EmitterExecutionStateAccessors.Num() > EmitterIdx && EmitterExecutionStateAccessors[EmitterIdx].IsValidForWrite())
+				{
+					EmitterExecutionStateAccessors[EmitterIdx].Set(SystemIndex, (int32)EmitterInst.GetExecutionState());
+				}
 			}
 		}
 	};
@@ -1250,16 +1479,20 @@ void FNiagaraSystemSimulation::TransferSystemSimResults(FNiagaraSystemSimulation
 		return;
 	}
 
-	SystemExecutionStateAccessor.SetDataSet(Context.DataSet);
-	SystemExecutionStateAccessor.InitForAccess();
-	for (int32 EmitterIdx = 0; EmitterIdx < Context.System->GetNumEmitters(); ++EmitterIdx)
+	bool bIsUsingFastPath = GetSystem()->FastPathMode != ENiagaraFastPathMode::ScriptVMOnly;
+	if (bIsUsingFastPath == false)
 	{
-		EmitterExecutionStateAccessors[EmitterIdx].SetDataSet(Context.DataSet);
-		EmitterExecutionStateAccessors[EmitterIdx].InitForAccess();
-		for (int32 SpawnInfoIdx = 0; SpawnInfoIdx < EmitterSpawnInfoAccessors[EmitterIdx].Num(); ++SpawnInfoIdx)
+		SystemExecutionStateAccessor.SetDataSet(Context.DataSet);
+		SystemExecutionStateAccessor.InitForAccess();
+		for (int32 EmitterIdx = 0; EmitterIdx < Context.System->GetNumEmitters(); ++EmitterIdx)
 		{
-			EmitterSpawnInfoAccessors[EmitterIdx][SpawnInfoIdx].SetDataSet(Context.DataSet);
-			EmitterSpawnInfoAccessors[EmitterIdx][SpawnInfoIdx].InitForAccess();
+			EmitterExecutionStateAccessors[EmitterIdx].SetDataSet(Context.DataSet);
+			EmitterExecutionStateAccessors[EmitterIdx].InitForAccess();
+			for (int32 SpawnInfoIdx = 0; SpawnInfoIdx < EmitterSpawnInfoAccessors[EmitterIdx].Num(); ++SpawnInfoIdx)
+			{
+				EmitterSpawnInfoAccessors[EmitterIdx][SpawnInfoIdx].SetDataSet(Context.DataSet);
+				EmitterSpawnInfoAccessors[EmitterIdx][SpawnInfoIdx].InitForAccess();
+			}
 		}
 	}
 
@@ -1268,8 +1501,11 @@ void FNiagaraSystemSimulation::TransferSystemSimResults(FNiagaraSystemSimulation
 		ENiagaraExecutionState ExecutionState = (ENiagaraExecutionState)SystemExecutionStateAccessor.GetSafe(SystemIndex, (int32)ENiagaraExecutionState::Disabled);
 		FNiagaraSystemInstance* SystemInst = Context.Instances[SystemIndex];
 
-		//Apply the systems requested execution state to it's actual execution state.
-		SystemInst->SetActualExecutionState(ExecutionState);
+		if (bIsUsingFastPath == false)
+		{
+			//Apply the systems requested execution state to it's actual execution state.
+			SystemInst->SetActualExecutionState(ExecutionState);
+		}
 
 		if (!SystemInst->IsDisabled())
 		{
@@ -1286,19 +1522,23 @@ void FNiagaraSystemSimulation::TransferSystemSimResults(FNiagaraSystemSimulation
 				}
 
 				check(Emitters.Num() > EmitterIdx);
-				ENiagaraExecutionState State = (ENiagaraExecutionState)EmitterExecutionStateAccessors[EmitterIdx].GetSafe(SystemIndex, (int32)ENiagaraExecutionState::Disabled);
-				EmitterInst.SetExecutionState(State);
 
-				TArray<FNiagaraSpawnInfo>& EmitterInstSpawnInfos = EmitterInst.GetSpawnInfo();
-				for (int32 SpawnInfoIdx = 0; SpawnInfoIdx < EmitterSpawnInfoAccessors[EmitterIdx].Num(); ++SpawnInfoIdx)
+				if (bIsUsingFastPath == false)
 				{
-					if (SpawnInfoIdx < EmitterInstSpawnInfos.Num())
-					{
-						EmitterInstSpawnInfos[SpawnInfoIdx] = EmitterSpawnInfoAccessors[EmitterIdx][SpawnInfoIdx].Get(SystemIndex);
-					}
-					else
+					ENiagaraExecutionState State = (ENiagaraExecutionState)EmitterExecutionStateAccessors[EmitterIdx].GetSafe(SystemIndex, (int32)ENiagaraExecutionState::Disabled);
+					EmitterInst.SetExecutionState(State);
+
+					TArray<FNiagaraSpawnInfo>& EmitterInstSpawnInfos = EmitterInst.GetSpawnInfo();
+					for (int32 SpawnInfoIdx = 0; SpawnInfoIdx < EmitterSpawnInfoAccessors[EmitterIdx].Num(); ++SpawnInfoIdx)
 					{
-						ensure(SpawnInfoIdx < EmitterInstSpawnInfos.Num());
+						if (SpawnInfoIdx < EmitterInstSpawnInfos.Num())
+						{
+							EmitterInstSpawnInfos[SpawnInfoIdx] = EmitterSpawnInfoAccessors[EmitterIdx][SpawnInfoIdx].Get(SystemIndex);
+						}
+						else
+						{
+							ensure(SpawnInfoIdx < EmitterInstSpawnInfos.Num());
+						}
 					}
 				}
 
diff --git a/Engine/Plugins/FX/Niagara/Source/Niagara/Public/NiagaraSystemFastPath.h b/Engine/Plugins/FX/Niagara/Source/Niagara/Public/NiagaraSystemFastPath.h
new file mode 100644
--- /dev/null
+++ b/Engine/Plugins/FX/Niagara/Source/Niagara/Public/NiagaraSystemFastPath.h
@@ -0,0 +1,586 @@
+// Copyright 1998-2019 Epic Games, Inc. All Rights Reserved.
+
+#pragma once
+
+#include "NiagaraTypes.h"
+#include "NiagaraParameterStore.h"
+#include "NiagaraSystemFastPath.generated.h"
+
+class UNiagaraEmitter;
+
+UENUM()
+enum class ENiagaraFastPathMode
+{
+	ScriptVMOnly,
+	FastPathOnly,
+	ScrptVMAndFastPath
+};
+
+UENUM()
+enum class ENiagaraFastPathInputType
+{
+	Local,
+	RangedSpawn,
+	RangedUpdate,
+	User
+};
+
+USTRUCT()
+struct FNiagaraFastPathFloatInput
+{
+	GENERATED_BODY()
+
+	UPROPERTY(EditAnywhere, Category="Input")
+	ENiagaraFastPathInputType Type;
+
+	UPROPERTY(EditAnywhere, Category = "Input")
+	float Local;
+
+	UPROPERTY(EditAnywhere, Category = "Input")
+	float RangeMin;
+
+	UPROPERTY(EditAnywhere, Category = "Input")
+	float RangeMax;
+
+	UPROPERTY(EditAnywhere, Category = "Input")
+	FName UserParameterName;
+
+	FNiagaraFastPathFloatInput()
+		: Type(ENiagaraFastPathInputType::Local)
+		, Local(0.0f)
+		, RangeMin(0.0f)
+		, RangeMax(0.0f)
+	{
+	}
+};
+
+USTRUCT()
+struct FNiagaraFastPathIntInput
+{
+	GENERATED_BODY()
+
+	UPROPERTY(EditAnywhere, Category="Input")
+	ENiagaraFastPathInputType Type;
+
+	UPROPERTY(EditAnywhere, Category = "Input")
+	int32 Local;
+
+	UPROPERTY(EditAnywhere, Category = "Input")
+	int32 RangeMin;
+
+	UPROPERTY(EditAnywhere, Category = "Input")
+	int32 RangeMax;
+
+	UPROPERTY(EditAnywhere, Category = "Input")
+	FName UserParameterName;
+
+	FNiagaraFastPathIntInput()
+		: Type(ENiagaraFastPathInputType::Local)
+		, Local(0)
+		, RangeMin(0)
+		, RangeMax(0)
+	{
+	}
+};
+
+USTRUCT()
+struct FNiagaraFastPathAttributeNames
+{
+	GENERATED_BODY()
+
+	UPROPERTY()
+	TArray<FName> System;
+
+	UPROPERTY()
+	TArray<FName> SystemFullNames;
+
+	UPROPERTY()
+	TArray<FName> Emitter;
+
+	UPROPERTY()
+	TArray<FName> EmitterFullNames;
+};
+
+template<typename T>
+struct TNiagaraFastPathUserParameterInputBinding
+{
+	T* InputValue;
+	FNiagaraParameterDirectBinding<T> ParameterBinding;
+
+	void Tick()
+	{
+		*InputValue = ParameterBinding.GetValue();
+	}
+};
+
+template<typename T>
+struct TNiagaraFastPathRangedInputBinding
+{
+	const T Min;
+	const T Max;
+	T* const InputValue;
+
+	TNiagaraFastPathRangedInputBinding(T SourceMin, T SourceMax, T* TargetInputValue)
+		: Min(SourceMin)
+		, Max(SourceMax)
+		, InputValue(TargetInputValue)
+	{
+	}
+
+	void Tick()
+	{
+		*InputValue = FMath::RandRange(Min, Max);
+	}
+};
+
+template<typename T>
+struct TNiagaraFastPathAttributeBinding
+{
+	T* ParameterValue;
+	FNiagaraParameterDirectBinding<T> ParameterBinding;
+
+	void Tick()
+	{
+		ParameterBinding.SetValue(*ParameterValue);
+	}
+};
+
+USTRUCT()
+struct FNiagaraFastPath_Module_SystemScalability
+{
+	GENERATED_BODY()
+
+	UPROPERTY(EditAnywhere, Category = "System Scalability")
+	bool bUseSystemScalability;
+
+	UPROPERTY(EditAnywhere, Category = "System Scalability")
+	ENiagaraExecutionState CulledState;
+
+	UPROPERTY(EditAnywhere, Category = "System Scalability")
+	bool bEnableVisibilityCulling;
+
+	UPROPERTY(EditAnywhere, Category = "System Scalability")
+	float VisibilityCullDelay;
+
+	FNiagaraFastPath_Module_SystemScalability()
+		: bUseSystemScalability(true)
+		, CulledState(ENiagaraExecutionState::Inactive)
+		, bEnableVisibilityCulling(true)
+		, VisibilityCullDelay(1.0f)
+	{
+	}
+};
+
+USTRUCT()
+struct FNiagaraFastPath_Module_SystemLifeCycle
+{
+	GENERATED_BODY()
+
+	UPROPERTY(EditAnywhere, Category = "System Life Cycle")
+	bool bCompleteOnInactive;
+
+	FNiagaraFastPath_Module_SystemLifeCycle()
+	: bCompleteOnInactive(false)
+	{
+	}
+};
+
+namespace FNiagaraSystemFastPath
+{
+	struct FParamMap0_Engine_Owner
+	{
+		float TimeSinceRendered;
+		ENiagaraExecutionState ExecutionState;
+	};
+
+	struct FParamMap0_Engine
+	{
+		FParamMap0_Engine_Owner Owner;
+	};
+
+	struct FParamMap0_System
+	{
+		ENiagaraExecutionState ExecutionState;
+		ENiagaraExecutionStateSource ExecutionStateSource;
+	};
+
+	struct FParamMap0_Scalability
+	{
+		ENiagaraExecutionState ExecutionState;
+	};
+
+	struct FParamMap0
+	{
+		FParamMap0_Engine Engine;
+		FParamMap0_System System;
+		FParamMap0_Scalability Scalability;
+	};
+
+	struct FParameterNames
+	{
+		static FName ExecutionState;
+		static FName ExecutionStateSource;
+	};
+
+	void SetSpawnMapDefaults(FParamMap0& Map);
+	void SetUpdateMapDefaults(FParamMap0& Map);
+
+	void Module_SystemScalability(const FNiagaraFastPath_Module_SystemScalability& Context_Map_SystemScalability, FParamMap0& Context_Map);
+
+	void Function_SystemChangeState(ENiagaraExecutionState In_NewState, bool In_Condition, ENiagaraExecutionStateSource In_NewStateSource, FParamMap0& Context_Map);
+
+	void Module_SystemLifeCycle(const FNiagaraFastPath_Module_SystemLifeCycle& Context_Map_SystemLifeCycle, FParamMap0& Context_Map);
+};
+
+USTRUCT()
+struct FNiagaraFastPath_Module_EmitterScalability
+{
+	GENERATED_BODY()
+
+	UPROPERTY(EditAnywhere, Category = "Emitter Scalability")
+	bool bUseEmitterScalability;
+
+	UPROPERTY(EditAnywhere, Category = "Emitter Scalability")
+	bool bUseMaxDistance;
+	UPROPERTY(EditAnywhere, Category = "Emitter Scalability")
+	ENiagaraExecutionState MaxCulledState;
+	UPROPERTY(EditAnywhere, Category = "Emitter Scalability")
+	float MaxDistance;
+
+	UPROPERTY(EditAnywhere, Category = "Emitter Scalability")
+	bool bUseMinDistance;
+	UPROPERTY(EditAnywhere, Category = "Emitter Scalability")
+	ENiagaraExecutionState MinCulledState;
+	UPROPERTY(EditAnywhere, Category = "Emitter Scalability")
+	float MinDistance;
+
+	UPROPERTY(EditAnywhere, Category = "Emitter Scalability")
+	bool bApplySpawnCountScale;
+	UPROPERTY(EditAnywhere, Category = "Emitter Scalability")
+	float SpawnCountScale;
+
+	UPROPERTY(EditAnywhere, Category = "Emitter Scalability")
+	bool bApplySpawnCountScaleByDistanceFraction;
+	UPROPERTY(EditAnywhere, Category = "Emitter Scalability")
+	float ScaleForMinFraction;
+	UPROPERTY(EditAnywhere, Category = "Emitter Scalability")
+	float ScaleForMaxFraction;
+
+	FNiagaraFastPath_Module_EmitterScalability()
+		: bUseEmitterScalability(true)
+		, bUseMaxDistance(false)
+		, MaxCulledState(ENiagaraExecutionState::Inactive)
+		, MaxDistance(5000.f)
+		, bUseMinDistance(false)
+		, MinCulledState(ENiagaraExecutionState::Inactive)
+		, MinDistance(0.0f)
+		, bApplySpawnCountScale(false)
+		, SpawnCountScale(1.0f)
+		, bApplySpawnCountScaleByDistanceFraction(false)
+		, ScaleForMinFraction(1.0f)
+		, ScaleForMaxFraction(1.0f)
+	{
+	}
+};
+
+USTRUCT()
+struct FNiagaraFastPath_Module_EmitterLifeCycle
+{
+	GENERATED_BODY()
+
+	UPROPERTY(EditAnywhere, Category = "Emitter Life Cycle")
+	bool bAutoComplete;
+	UPROPERTY(EditAnywhere, Category = "Emitter Life Cycle")
+	bool bCompleteOnInactive;
+	UPROPERTY(EditAnywhere, Category = "Emitter Life Cycle")
+	bool bDelayFirstLoopOnly;
+	UPROPERTY(EditAnywhere, Category = "Emitter Life Cycle")
+	bool bDurationRecalcEachLoop;
+
+	UPROPERTY(EditAnywhere, Category = "Emitter Life Cycle")
+	int32 MaxLoopCount;
+
+	UPROPERTY(EditAnywhere, Category = "Emitter Life Cycle")
+	float NextLoopDelay;
+
+	UPROPERTY(EditAnywhere, Category = "Emitter Life Cycle")
+	float NextLoopDuration;
+
+	FNiagaraFastPath_Module_EmitterLifeCycle()
+		: bAutoComplete(true)
+		, bCompleteOnInactive(false)
+		, bDelayFirstLoopOnly(false)
+		, bDurationRecalcEachLoop(false)
+		, MaxLoopCount(0)
+		, NextLoopDelay(0.0f)
+		, NextLoopDuration(5.0f)
+	{
+	}
+};
+
+USTRUCT()
+struct FNiagaraFastPath_Module_SpawnRate
+{
+	GENERATED_BODY()
+
+	UPROPERTY(EditAnywhere, Category = "Spawn Rate")
+	FNiagaraFastPathFloatInput SpawnRate;
+
+	UPROPERTY(EditAnywhere, Category = "Spawn Rate")
+	int32 SpawnGroup;
+
+	FNiagaraFastPath_Module_SpawnRate()
+		: SpawnGroup(0)
+	{
+	}
+};
+
+USTRUCT()
+struct FNiagaraFastPath_Module_SpawnPerUnit
+{
+	GENERATED_BODY()
+
+	UPROPERTY(EditAnywhere, Category = "Spawn Per Unit")
+	FNiagaraFastPathFloatInput SpawnPerUnit;
+
+	UPROPERTY(EditAnywhere, Category = "Spawn Per Unit")
+	bool bUseMovementTolerance;
+	UPROPERTY(EditAnywhere, Category = "Spawn Per Unit")
+	float MovementTolerance;
+
+	UPROPERTY(EditAnywhere, Category = "Spawn Per Unit")
+	bool bUseMaxMovementThreshold;
+	UPROPERTY(EditAnywhere, Category = "Spawn Per Unit")
+	float MaxMovementThreshold;
+
+	UPROPERTY(EditAnywhere, Category = "Spawn Per Unit")
+	int32 SpawnGroup;
+
+	FNiagaraFastPath_Module_SpawnPerUnit()
+		: bUseMovementTolerance(false)
+		, MovementTolerance(0.1f)
+		, bUseMaxMovementThreshold(true)
+		, MaxMovementThreshold(5000.0f)
+		, SpawnGroup(0)
+	{
+	}
+};
+
+USTRUCT()
+struct FNiagaraFastPath_Module_SpawnBurstInstantaneous
+{
+	GENERATED_BODY()
+
+	UPROPERTY(EditAnywhere, Category = "Spawn Burst Instantaneous")
+	FNiagaraFastPathIntInput SpawnCount;
+
+	UPROPERTY(EditAnywhere, Category = "Spawn Burst Instantaneous")
+	FNiagaraFastPathFloatInput SpawnTime;
+
+	UPROPERTY(EditAnywhere, Category = "Spawn Burst Instantaneous")
+	int32 SpawnGroup;
+
+	FNiagaraFastPath_Module_SpawnBurstInstantaneous()
+		: SpawnGroup(0)
+	{
+	}
+};
+
+namespace FNiagaraEmitterFastPath
+{
+	struct FParamMap0_Engine_Emitter
+	{
+		int NumParticles;
+	};
+
+	struct FParamMap0_Engine_Owner
+	{
+		float LODDistance;
+		FVector Velocity;
+	};
+
+	struct FParamMap0_Engine
+	{
+		float DeltaTime;
+		float GlobalSpawnCountScale;
+		FParamMap0_Engine_Emitter Emitter;
+		FParamMap0_Engine_Owner Owner;
+	};
+
+	struct FParamMap0_Scalability_Emitter
+	{
+		ENiagaraExecutionState ExecutionState;
+		float SpawnCountScale;
+	};
+
+	struct FParamMap0_Scalability
+	{
+		FParamMap0_Scalability_Emitter Emitter;
+	};
+
+	struct FParamMap0_Spawning
+	{
+		bool bCanEverSpawn;
+	};
+
+	struct FParamMap0_SpawnRate
+	{
+		int32 SpawnGroup;
+		float SpawnRate;
+
+		void Init(
+			const UNiagaraEmitter* Emitter,
+			int32 ModuleIndex,
+			FNiagaraParameterStore& InstanceParameters,
+			TArray<TNiagaraFastPathUserParameterInputBinding<int32>>& IntUserParameterInputBindings,
+			TArray<TNiagaraFastPathUserParameterInputBinding<float>>& FloatUserParameterInputBindings,
+			TArray<TNiagaraFastPathRangedInputBinding<int32>>& IntUpdateRangedInputBindings,
+			TArray<TNiagaraFastPathRangedInputBinding<float>>& FloatUpdateRangedInputBindings);
+	};
+
+	struct FParamMap0_SpawnPerUnit
+	{
+		float SpawnPerUnit;
+		bool bUseMovementTolerance;
+		float MovementTolerance;
+		bool bUseMaxMovementThreshold;
+		float MaxMovementThreshold;
+		int32 SpawnGroup;
+
+		void Init(
+			const UNiagaraEmitter* Emitter,
+			int32 ModuleIndex,
+			FNiagaraParameterStore& InstanceParameters,
+			TArray<TNiagaraFastPathUserParameterInputBinding<int32>>& IntUserParameterInputBindings,
+			TArray<TNiagaraFastPathUserParameterInputBinding<float>>& FloatUserParameterInputBindings,
+			TArray<TNiagaraFastPathRangedInputBinding<int32>>& IntUpdateRangedInputBindings,
+			TArray<TNiagaraFastPathRangedInputBinding<float>>& FloatUpdateRangedInputBindings);
+	};
+
+	struct FParamMap0_SpawnBurst_Instantaneous
+	{
+		int32 SpawnCount;
+		int32 SpawnGroup;
+		float SpawnTime;
+
+		void Init(
+			const UNiagaraEmitter* Emitter,
+			int32 ModuleIndex,
+			FNiagaraParameterStore& InstanceParameters,
+			TArray<TNiagaraFastPathUserParameterInputBinding<int32>>& IntUserParameterInputBindings,
+			TArray<TNiagaraFastPathUserParameterInputBinding<float>>& FloatUserParameterInputBindings,
+			TArray<TNiagaraFastPathRangedInputBinding<int32>>& IntUpdateRangedInputBindings,
+			TArray<TNiagaraFastPathRangedInputBinding<float>>& FloatUpdateRangedInputBindings);
+	};
+
+	struct FParamMap0_Emitter_Scalability
+	{
+		float DistanceFraction;
+	};
+
+	struct FParamMap0_Emitter_SpawnRate
+	{
+		FNiagaraSpawnInfo SpawnOutputInfo;
+		float SpawnRemainder;
+
+		FParamMap0_Emitter_SpawnRate()
+			: SpawnRemainder(0.0f)
+		{
+		}
+	};
+
+	struct FParamMap0_Emitter_SpawnPerUnit
+	{
+		FNiagaraSpawnInfo SpawnOutputInfo;
+		float SpawnRemainder;
+
+		FParamMap0_Emitter_SpawnPerUnit()
+			: SpawnRemainder(0.0f)
+		{
+		}
+	};
+
+	struct FParamMap0_Emitter_SpawnBurst_Instantaneous
+	{
+		FNiagaraSpawnInfo SpawnBurst;
+	};
+
+	struct FParamMap0_Emitter
+	{
+		float Age;
+		float CurrentLoopDelay;
+		float CurrentLoopDuration;
+		ENiagaraExecutionState ExecutionState;
+		ENiagaraExecutionStateSource ExecutionStateSource;
+		int32 LoopCount;
+		float LoopedAge;
+		float NormalizedLoopAge;
+		float DistanceTraveled;
+		FParamMap0_Emitter_Scalability Scalability;
+		TArray<FParamMap0_Emitter_SpawnRate> SpawnRate;
+		TArray<FParamMap0_Emitter_SpawnPerUnit> SpawnPerUnit;
+		TArray<FParamMap0_Emitter_SpawnBurst_Instantaneous> SpawnBurst_Instantaneous;
+	};
+
+	struct FParamMap0
+	{
+		TArray<FParamMap0_SpawnRate> SpawnRate;
+		TArray<FParamMap0_SpawnPerUnit> SpawnPerUnit;
+		TArray<FParamMap0_SpawnBurst_Instantaneous> SpawnBurst_Instantaneous;
+		FParamMap0_Engine Engine;
+		FParamMap0_Emitter Emitter;
+		FParamMap0_Scalability Scalability;
+		FParamMap0_Spawning Spawning;
+		const FNiagaraSystemFastPath::FParamMap0_System* System;
+	};
+
+	struct FAttributeNames
+	{
+		static FName Age;
+		static FName CurrentLoopDelay;
+		static FName CurrentLoopDuration;
+		static FName ExecutionState;
+		static FName ExecutionStateSource;
+		static FName LoopCount;
+		static FName LoopedAge;
+		static FName NormalizedLoopAge;
+		static FName DistanceTravelled;
+		static FName Scalability_DistanceFraction;
+	};
+
+	void InitFastPathAttributeBindings(
+		const FNiagaraFastPathAttributeNames& SourceAttributeNames,
+		FNiagaraParameterStore& TargetParameterStore,
+		FNiagaraSystemFastPath::FParamMap0& SystemMap,
+		FNiagaraEmitterFastPath::FParamMap0& EmitterMap,
+		TArray<TNiagaraFastPathAttributeBinding<int32>>& FastPathIntAttributeBindings,
+		TArray<TNiagaraFastPathAttributeBinding<float>>& FastPathFloatAttributeBindings);
+
+	template<typename TBindingType, typename TVariableType>
+	static void AddBinding(FName ParameterName, TVariableType ParameterType, TBindingType* SourceValuePtr, FNiagaraParameterStore& TargetParameterStore, TArray<TNiagaraFastPathAttributeBinding<TBindingType>>& TargetBindings)
+	{
+		TNiagaraFastPathAttributeBinding<TBindingType> Binding;
+		FNiagaraVariable ParameterVariable = FNiagaraVariable(ParameterType, ParameterName);
+		Binding.ParameterBinding.Init(TargetParameterStore, ParameterVariable);
+		if (Binding.ParameterBinding.ValuePtr != nullptr)
+		{
+			Binding.ParameterValue = SourceValuePtr;
+			TargetBindings.Add(Binding);
+		}
+	}
+
+	void SetSpawnMapDefaults(const UNiagaraEmitter& Emitter, FParamMap0& Map);
+	void SetUpdateMapDefaults(FParamMap0& Map);
+
+	void Function_SampleCurve_SpawnCountScaleByDistanceFraction(const FNiagaraFastPath_Module_EmitterScalability& Context_Map_EmitterScalability, float InFraction, float& OutScale);
+	void Module_EmitterScalability(const FNiagaraFastPath_Module_EmitterScalability& Context_Map_EmitterScalability, FParamMap0& Context_Map);
+
+	void Function_EmitterLifeCycle_EmitterChangeState(ENiagaraExecutionState In_NewState, bool In_Condition, ENiagaraExecutionStateSource In_NewStateSource, FParamMap0& Context_Map);
+	void Module_EmitterLifeCycle(const FNiagaraFastPath_Module_EmitterLifeCycle& Context_Map_EmitterLifeCycle, FParamMap0& Context_Map);
+
+	void Module_SpawnRate(FParamMap0& Context_Map);
+
+	void Module_SpawnPerUnit(FParamMap0& ContextMap);
+
+	void Module_SpawnBurstInstantaneous(FParamMap0& ContextMap);
+}
\ No newline at end of file
diff --git a/Engine/Plugins/FX/Niagara/Source/Niagara/Public/NiagaraSystemInstance.h b/Engine/Plugins/FX/Niagara/Source/Niagara/Public/NiagaraSystemInstance.h
--- a/Engine/Plugins/FX/Niagara/Source/Niagara/Public/NiagaraSystemInstance.h
+++ b/Engine/Plugins/FX/Niagara/Source/Niagara/Public/NiagaraSystemInstance.h
@@ -10,6 +10,7 @@
 #include "Templates/UniquePtr.h"
 #include "NiagaraCommon.h"
 #include "NiagaraDataInterface.h"
+#include "NiagaraSystemFastPath.h"
 
 class FNiagaraWorldManager;
 class UNiagaraComponent;
@@ -109,6 +110,14 @@ public:
 	ENiagaraExecutionState GetActualExecutionState() { return ActualExecutionState; }
 	void SetActualExecutionState(ENiagaraExecutionState InState);
 
+	float GetSystemTimeSinceRendered() const { return SystemTimeSinceRenderedParam.GetValue(); }
+
+	float GetOwnerLODDistance() const { return OwnerLODDistanceParam.GetValue(); }
+
+	int32 GetNumParticles(int32 EmitterIndex) const { return ParameterNumParticleBindings[EmitterIndex].GetValue(); }
+
+	FVector GetOwnerVelocity() const { return OwnerVelocityParam.GetValue(); }
+
 	FORCEINLINE bool IsComplete()const { return ActualExecutionState == ENiagaraExecutionState::Complete || ActualExecutionState == ENiagaraExecutionState::Disabled; }
 	FORCEINLINE bool IsDisabled()const { return ActualExecutionState == ENiagaraExecutionState::Disabled; }
 
@@ -234,6 +243,12 @@ public:
 
 	void TickInstanceParameters(float DeltaSeconds);
 
+	void TickFastPathBindings();
+
+	void ResetFastPathBindings();
+
+	FNiagaraSystemFastPath::FParamMap0& GetFastPathMap() { return FastPathMap; }
+
 private:
 
 	void DestroyDataInterfaceInstanceData();
@@ -385,4 +400,11 @@ public:
 	uint32 TotalParamSize = 0;
 	uint32 ActiveGPUEmitterCount = 0;
 	int32 GPUDataInterfaceInstanceDataSize = 0;
+
+	FNiagaraSystemFastPath::FParamMap0 FastPathMap;
+
+	TArray<TNiagaraFastPathRangedInputBinding<int32>> FastPathIntUpdateRangedInputBindings;
+	TArray<TNiagaraFastPathRangedInputBinding<float>> FastPathFloatUpdateRangedInputBindings;
+	TArray<TNiagaraFastPathUserParameterInputBinding<int32>> FastPathIntUserParameterInputBindings;
+	TArray<TNiagaraFastPathUserParameterInputBinding<float>> FastPathFloatUserParameterInputBindings;
 };
diff --git a/Engine/Plugins/FX/Niagara/Source/Niagara/Public/NiagaraSystemSimulation.h b/Engine/Plugins/FX/Niagara/Source/Niagara/Public/NiagaraSystemSimulation.h
--- a/Engine/Plugins/FX/Niagara/Source/Niagara/Public/NiagaraSystemSimulation.h
+++ b/Engine/Plugins/FX/Niagara/Source/Niagara/Public/NiagaraSystemSimulation.h
@@ -160,6 +160,8 @@ public:
 	/** Second phase of system sim tick that can run on any thread. */
 	void Tick_Concurrent(FNiagaraSystemSimulationTickContext& Context);
 
+	void TickFastPath(FNiagaraSystemSimulationTickContext& Context);
+
 	/** Update TickGroups for pending instances and execute tick group promotions. */
 	void UpdateTickGroups_GameThread();
 	/** Spawn any pending instances, assumes that you have update tick groups ahead of time. */
diff --git a/Engine/Plugins/FX/Niagara/Source/Niagara/Public/NiagaraTypes.h b/Engine/Plugins/FX/Niagara/Source/Niagara/Public/NiagaraTypes.h
--- a/Engine/Plugins/FX/Niagara/Source/Niagara/Public/NiagaraTypes.h
+++ b/Engine/Plugins/FX/Niagara/Source/Niagara/Public/NiagaraTypes.h
@@ -589,6 +589,7 @@ public:
 	static UScriptStruct* GetIDStruct() { return IDStruct; }
 
 	static UEnum* GetExecutionStateEnum() { return ExecutionStateEnum; }
+	static UEnum* GetExecutionStateSouceEnum() { return ExecutionStateSourceEnum; }
 	static UEnum* GetSimulationTargetEnum() { return SimulationTargetEnum; }
 	static UEnum* GetScriptUsageEnum() { return ScriptUsageEnum; }
 
