Extend graph traversal builder to evaluate static switches

diff --git a/Engine/Plugins/FX/Niagara/Source/NiagaraEditor/Private/NiagaraGraph.cpp b/Engine/Plugins/FX/Niagara/Source/NiagaraEditor/Private/NiagaraGraph.cpp
--- a/Engine/Plugins/FX/Niagara/Source/NiagaraEditor/Private/NiagaraGraph.cpp
+++ b/Engine/Plugins/FX/Niagara/Source/NiagaraEditor/Private/NiagaraGraph.cpp
@@ -478,9 +478,9 @@ UNiagaraNodeOutput* UNiagaraGraph::FindEquivalentOutputNode(ENiagaraScriptUsage
 }
 
 
-void BuildTraversalHelper(TArray<class UNiagaraNode*>& OutNodesTraversed, UNiagaraNode* CurrentNode)
+void BuildTraversalHelper(TArray<class UNiagaraNode*>& OutNodesTraversed, UNiagaraNode* CurrentNode, bool bEvaluateStaticSwitches)
 {
-	if (CurrentNode == nullptr)
+	if (CurrentNode == nullptr || (bEvaluateStaticSwitches && CurrentNode->IsA<UNiagaraNodeStaticSwitch>()))
 	{
 		return;
 	}
@@ -490,34 +490,38 @@ void BuildTraversalHelper(TArray<class UNiagaraNode*>& OutNodesTraversed, UNiaga
 	TArray<UEdGraphPin*> Pins = CurrentNode->GetAllPins();
 	for (int32 i = 0; i < Pins.Num(); i++)
 	{
-		if (Pins[i]->Direction == EEdGraphPinDirection::EGPD_Input && Pins[i]->LinkedTo.Num() == 1)
+		if (Pins[i]->Direction == EEdGraphPinDirection::EGPD_Input && Pins[i]->LinkedTo.Num() > 0)
 		{
-			UNiagaraNode* Node = Cast<UNiagaraNode>(Pins[i]->LinkedTo[0]->GetOwningNode());
-			if (OutNodesTraversed.Contains(Node))
+			for (UEdGraphPin* LinkedPin : Pins[i]->LinkedTo)
 			{
-				continue;
+				UEdGraphPin* TracedPin = bEvaluateStaticSwitches ? UNiagaraNode::TraceOutputPin(LinkedPin) : LinkedPin;
+				UNiagaraNode* Node = Cast<UNiagaraNode>(TracedPin->GetOwningNode());
+				if (OutNodesTraversed.Contains(Node))
+				{
+					continue;
+				}
+				BuildTraversalHelper(OutNodesTraversed, Node, bEvaluateStaticSwitches);
 			}
-			BuildTraversalHelper(OutNodesTraversed, Node);
 		}
 	}
 
 	OutNodesTraversed.Add(CurrentNode);
 }
 
-void UNiagaraGraph::BuildTraversal(TArray<class UNiagaraNode*>& OutNodesTraversed, ENiagaraScriptUsage TargetUsage, FGuid TargetUsageId) const
+void UNiagaraGraph::BuildTraversal(TArray<class UNiagaraNode*>& OutNodesTraversed, ENiagaraScriptUsage TargetUsage, FGuid TargetUsageId, bool bEvaluateStaticSwitches) const
 {
 	UNiagaraNodeOutput* Output = FindOutputNode(TargetUsage, TargetUsageId);
 	if (Output)
 	{
-		BuildTraversalHelper(OutNodesTraversed, Output);
+		BuildTraversalHelper(OutNodesTraversed, Output, bEvaluateStaticSwitches);
 	}
 }
 
-void UNiagaraGraph::BuildTraversal(TArray<class UNiagaraNode*>& OutNodesTraversed, UNiagaraNode* FinalNode) 
+void UNiagaraGraph::BuildTraversal(TArray<class UNiagaraNode*>& OutNodesTraversed, UNiagaraNode* FinalNode, bool bEvaluateStaticSwitches)
 {
 	if (FinalNode)
 	{
-		BuildTraversalHelper(OutNodesTraversed, FinalNode);
+		BuildTraversalHelper(OutNodesTraversed, FinalNode, bEvaluateStaticSwitches);
 	}
 }
 
diff --git a/Engine/Plugins/FX/Niagara/Source/NiagaraEditor/Private/ViewModels/Stack/NiagaraStackGraphUtilities.cpp b/Engine/Plugins/FX/Niagara/Source/NiagaraEditor/Private/ViewModels/Stack/NiagaraStackGraphUtilities.cpp
--- a/Engine/Plugins/FX/Niagara/Source/NiagaraEditor/Private/ViewModels/Stack/NiagaraStackGraphUtilities.cpp
+++ b/Engine/Plugins/FX/Niagara/Source/NiagaraEditor/Private/ViewModels/Stack/NiagaraStackGraphUtilities.cpp
@@ -30,7 +30,6 @@
 #include "EdGraph/EdGraphPin.h"
 #include "ViewModels/NiagaraEmitterViewModel.h"
 #include "AssetRegistryModule.h"
-#include "NiagaraNodeStaticSwitch.h"
 
 DECLARE_CYCLE_STAT(TEXT("Niagara - StackGraphUtilities - RelayoutGraph"), STAT_NiagaraEditor_StackGraphUtilities_RelayoutGraph, STATGROUP_NiagaraEditor);
 
@@ -581,38 +580,6 @@ TArray<UEdGraphPin*> FNiagaraStackGraphUtilities::GetUnusedFunctionInputPins(UNi
 		return TArray<UEdGraphPin*>();
 	}
 
-	// Find reachable nodes
-	TArray<UEdGraphPin*> PinsToCheck;
-	OutputNode->GetInputPins(PinsToCheck);
-	TSet<UNiagaraNode*> ReachedNodes;
-	while (PinsToCheck.Num() > 0)
-	{
-		UEdGraphPin* Pin = PinsToCheck.Pop(false);
-		UNiagaraNode* Owner = Cast<UNiagaraNode>(Pin->GetOwningNode());
-		if (!Owner)
-		{
-			continue;
-		}
-		for (UEdGraphPin* LinkedPin : Pin->LinkedTo)
-		{
-			UEdGraphPin* TracedPin = UNiagaraNode::TraceOutputPin(LinkedPin);
-			if (TracedPin->Direction == EGPD_Output && TracedPin->GetOwningNode()->IsA<UNiagaraNodeStaticSwitch>())
-			{
-				// if the static switch has no connected input (the value is set on the input pin directly) then the trace returns the output pin and we can ignore it
-				continue;
-			}
-			UNiagaraNode* TracedPinOwner = Cast<UNiagaraNode>(TracedPin->GetOwningNode());
-			if (!TracedPinOwner || ReachedNodes.Contains(TracedPinOwner))
-			{
-				continue;
-			}
-			ReachedNodes.Add(TracedPinOwner);
-			TArray<UEdGraphPin*> TracedInputs;
-			TracedPinOwner->GetInputPins(TracedInputs);
-			PinsToCheck.Append(TracedInputs);
-		}
-	}
-
 	// Get the used function parameters from the parameter map set node linked to the function's input pin.
 	// Note that this is only valid for module scripts, not function scripts.
 	TArray<UEdGraphPin*> ResultPins;
@@ -638,6 +605,10 @@ TArray<UEdGraphPin*> FNiagaraStackGraphUtilities::GetUnusedFunctionInputPins(UNi
 		return ResultPins;
 	}
 
+	// Find reachable nodes
+	TArray<UNiagaraNode*> ReachedNodes;
+	FunctionGraph->BuildTraversal(ReachedNodes, OutputNode, true);
+
 	// We only care about reachable parameter map get nodes with module inputs
 	const UEdGraphSchema_Niagara* Schema = GetDefault<UEdGraphSchema_Niagara>();
 	for (UNiagaraNode* Node : ReachedNodes)
diff --git a/Engine/Plugins/FX/Niagara/Source/NiagaraEditor/Public/NiagaraGraph.h b/Engine/Plugins/FX/Niagara/Source/NiagaraEditor/Public/NiagaraGraph.h
--- a/Engine/Plugins/FX/Niagara/Source/NiagaraEditor/Public/NiagaraGraph.h
+++ b/Engine/Plugins/FX/Niagara/Source/NiagaraEditor/Public/NiagaraGraph.h
@@ -156,8 +156,8 @@ class UNiagaraGraph : public UEdGraph
 	TArray<FNiagaraVariable> FindStaticSwitchInputs(bool bReachableOnly = false) const;
 
 	/** Get an in-order traversal of a graph by the specified target output script usage.*/
-	void BuildTraversal(TArray<class UNiagaraNode*>& OutNodesTraversed, ENiagaraScriptUsage TargetUsage, FGuid TargetUsageId) const;
-	static void BuildTraversal(TArray<class UNiagaraNode*>& OutNodesTraversed, UNiagaraNode* FinalNode);
+	void BuildTraversal(TArray<class UNiagaraNode*>& OutNodesTraversed, ENiagaraScriptUsage TargetUsage, FGuid TargetUsageId, bool bEvaluateStaticSwitches = false) const;
+	static void BuildTraversal(TArray<class UNiagaraNode*>& OutNodesTraversed, UNiagaraNode* FinalNode, bool bEvaluateStaticSwitches = false);
 
 	/** Generates a list of unique input and output parameters for when this script is used as a function. */
 	void GetParameters(TArray<FNiagaraVariable>& Inputs, TArray<FNiagaraVariable>& Outputs) const;
