Now possible to bake out rapid iteration parameters if the end user sets up a system to do so.

diff --git a/Engine/Plugins/FX/Niagara/Source/Niagara/Classes/NiagaraEmitter.h b/Engine/Plugins/FX/Niagara/Source/Niagara/Classes/NiagaraEmitter.h
--- a/Engine/Plugins/FX/Niagara/Source/Niagara/Classes/NiagaraEmitter.h
+++ b/Engine/Plugins/FX/Niagara/Source/Niagara/Classes/NiagaraEmitter.h
@@ -291,6 +291,10 @@ public:
 	UPROPERTY()
 	class UNiagaraScriptSourceBase*	GraphSource;
 
+	/** Should we enable rapid iteration removal if the system is also set to remove rapid iteration parameters on compile? This value defaults to true.*/
+	UPROPERTY(EditAnywhere, AdvancedDisplay, Category = "Emitter")
+	uint32 bBakeOutRapidIteration : 1;
+
 	bool NIAGARA_API AreAllScriptAndSourcesSynchronized() const;
 	void NIAGARA_API OnPostCompile();
 
diff --git a/Engine/Plugins/FX/Niagara/Source/Niagara/Classes/NiagaraScript.h b/Engine/Plugins/FX/Niagara/Source/Niagara/Classes/NiagaraScript.h
--- a/Engine/Plugins/FX/Niagara/Source/Niagara/Classes/NiagaraScript.h
+++ b/Engine/Plugins/FX/Niagara/Source/Niagara/Classes/NiagaraScript.h
@@ -133,6 +133,14 @@ public:
 	UPROPERTY()
 	TArray<FString> AdditionalDefines;
 
+	/** Bitfield of supported detail levels in this compile.*/
+	UPROPERTY()
+	uint32 DetailLevelMask;
+
+	/** Whether or not we need to bake Rapid Iteration params. True to keep params, false to bake.*/
+	UPROPERTY()
+	bool bUsesRapidIterationParams;
+
 	/**
 	* The GUID of the subgraph this shader primarily represents.
 	*/
@@ -157,6 +165,8 @@ public:
 	FNiagaraVMExecutableDataId()
 		: CompilerVersionID()
 		, ScriptUsageType(ENiagaraScriptUsage::Function)
+		, DetailLevelMask(0xFFFFFFFF)
+		, bUsesRapidIterationParams(true)
 		, BaseScriptID(0, 0, 0, 0)
 	{ }
 
diff --git a/Engine/Plugins/FX/Niagara/Source/Niagara/Classes/NiagaraScriptSourceBase.h b/Engine/Plugins/FX/Niagara/Source/Niagara/Classes/NiagaraScriptSourceBase.h
--- a/Engine/Plugins/FX/Niagara/Source/Niagara/Classes/NiagaraScriptSourceBase.h
+++ b/Engine/Plugins/FX/Niagara/Source/Niagara/Classes/NiagaraScriptSourceBase.h
@@ -32,6 +32,11 @@ public:
 	virtual int32 GetDependentRequestCount() const = 0;
 	virtual TSharedPtr<FNiagaraCompileRequestDataBase, ESPMode::ThreadSafe> GetDependentRequest(int32 Index) = 0;
 	virtual FName ResolveEmitterAlias(FName VariableName) const = 0;
+	virtual uint32 GetDetailLevelMask() const = 0;
+	virtual bool GetUseRapidIterationParams() const = 0;
+
+	static const uint32 CookForAllDetailLevelMask = 0xFFFFFFFF;
+
 };
 
 class FNiagaraCompileOptions
@@ -58,6 +63,7 @@ public:
 	FString Name;
 	int32 TargetUsageBitmask;
 	TArray<FString> AdditionalDefines;
+	uint32 DetailLevelMask = 0xFFFFFFFF;
 };
 
 struct FNiagaraParameterStore;
diff --git a/Engine/Plugins/FX/Niagara/Source/Niagara/Classes/NiagaraSystem.h b/Engine/Plugins/FX/Niagara/Source/Niagara/Classes/NiagaraSystem.h
--- a/Engine/Plugins/FX/Niagara/Source/Niagara/Classes/NiagaraSystem.h
+++ b/Engine/Plugins/FX/Niagara/Source/Niagara/Classes/NiagaraSystem.h
@@ -235,6 +235,10 @@ public:
 	void InvalidateCachedCompileIds();
 
 	static void RequestCompileForEmitter(UNiagaraEmitter* InEmitter);
+
+	/** Experimental feature that allows us to bake out rapid iteration parameters during the normal compile process. */
+	UPROPERTY(EditAnywhere, AdvancedDisplay, Category = "Emitter")
+	uint32 bBakeOutRapidIteration : 1;
 #endif
 
 	FORCEINLINE UNiagaraParameterCollectionInstance* GetParameterCollectionOverride(UNiagaraParameterCollection* Collection)
diff --git a/Engine/Plugins/FX/Niagara/Source/Niagara/Private/NiagaraComponent.cpp b/Engine/Plugins/FX/Niagara/Source/Niagara/Private/NiagaraComponent.cpp
--- a/Engine/Plugins/FX/Niagara/Source/Niagara/Private/NiagaraComponent.cpp
+++ b/Engine/Plugins/FX/Niagara/Source/Niagara/Private/NiagaraComponent.cpp
@@ -1156,11 +1156,14 @@ void UNiagaraComponent::GetUsedMaterials(TArray<UMaterialInterface*>& OutMateria
 
 		if (UNiagaraEmitter* Props = Sim->GetEmitterHandle().GetInstance())
 		{
-			for (int32 i = 0; i < Props->GetRenderers().Num(); i++)
+			if (Props)
 			{
-				if (UNiagaraRendererProperties* Renderer = Props->GetRenderers()[i])
+				for (int32 i = 0; i < Props->GetRenderers().Num(); i++)
 				{
-					Renderer->GetUsedMaterials(&Sim.Get(), OutMaterials);
+					if (UNiagaraRendererProperties* Renderer = Props->GetRenderers()[i])
+					{
+						Renderer->GetUsedMaterials(&Sim.Get(), OutMaterials);
+					}
 				}
 			}
 		}
diff --git a/Engine/Plugins/FX/Niagara/Source/Niagara/Private/NiagaraEmitter.cpp b/Engine/Plugins/FX/Niagara/Source/Niagara/Private/NiagaraEmitter.cpp
--- a/Engine/Plugins/FX/Niagara/Source/Niagara/Private/NiagaraEmitter.cpp
+++ b/Engine/Plugins/FX/Niagara/Source/Niagara/Private/NiagaraEmitter.cpp
@@ -110,6 +110,7 @@ UNiagaraEmitter::UNiagaraEmitter(const FObjectInitializer& Initializer)
 , MaxUpdateIterations(1)
 , bLimitDeltaTime(true)
 #if WITH_EDITORONLY_DATA
+, bBakeOutRapidIteration(true)
 , ThumbnailImageOutOfDate(true)
 #endif
 {
diff --git a/Engine/Plugins/FX/Niagara/Source/Niagara/Private/NiagaraEmitterInstance.cpp b/Engine/Plugins/FX/Niagara/Source/Niagara/Private/NiagaraEmitterInstance.cpp
--- a/Engine/Plugins/FX/Niagara/Source/Niagara/Private/NiagaraEmitterInstance.cpp
+++ b/Engine/Plugins/FX/Niagara/Source/Niagara/Private/NiagaraEmitterInstance.cpp
@@ -223,10 +223,9 @@ void FNiagaraEmitterInstance::Init(int32 InEmitterIdx, FNiagaraSystemInstanceID
 	OwnerSystemInstanceID = InSystemInstanceID;
 	const FNiagaraEmitterHandle& EmitterHandle = GetEmitterHandle();
 	CachedEmitter = EmitterHandle.GetInstance();
-	checkSlow(CachedEmitter);
 	CachedIDName = EmitterHandle.GetIdName();
 
-	if (!IsAllowedToExecute())
+	if (CachedEmitter == nullptr || !IsAllowedToExecute())
 	{
 		//@todo(message manager) Error bubbling here
 		ExecutionState = ENiagaraExecutionState::Disabled;
diff --git a/Engine/Plugins/FX/Niagara/Source/Niagara/Private/NiagaraScript.cpp b/Engine/Plugins/FX/Niagara/Source/Niagara/Private/NiagaraScript.cpp
--- a/Engine/Plugins/FX/Niagara/Source/Niagara/Private/NiagaraScript.cpp
+++ b/Engine/Plugins/FX/Niagara/Source/Niagara/Private/NiagaraScript.cpp
@@ -18,7 +18,8 @@
 #include "Misc/FileHelper.h"
 #include "UObject/EditorObjectVersion.h"
 #include "UObject/ReleaseObjectVersion.h"
-
+#include "NiagaraDataInterfaceSkeletalMesh.h"
+#include "NiagaraDataInterfaceStaticMesh.h"
 #if WITH_EDITOR
 	#include "NiagaraScriptDerivedData.h"
 	#include "DerivedDataCacheInterface.h"
@@ -118,7 +119,9 @@ bool FNiagaraVMExecutableDataId::operator==(const FNiagaraVMExecutableDataId& Re
 		ScriptUsageType != ReferenceSet.ScriptUsageType || 
 		ScriptUsageTypeID != ReferenceSet.ScriptUsageTypeID ||
 		BaseScriptID != ReferenceSet.BaseScriptID ||
-		BaseScriptCompileHash != ReferenceSet.BaseScriptCompileHash)
+		BaseScriptCompileHash != ReferenceSet.BaseScriptCompileHash ||
+		DetailLevelMask != ReferenceSet.DetailLevelMask ||
+		bUsesRapidIterationParams != ReferenceSet.bUsesRapidIterationParams)
 	{
 		return false;
 	}
@@ -167,6 +170,24 @@ void FNiagaraVMExecutableDataId::AppendKeyString(FString& KeyString) const
 	KeyString += BaseScriptCompileHash.ToString();
 	KeyString += TEXT("_");
 
+	if (DetailLevelMask != 0xFFFFFFFF)
+	{
+		KeyString += FString::Printf(TEXT("DL_%d"), DetailLevelMask);
+	}
+	else
+	{
+		KeyString += TEXT("ALLDL_");
+	}
+
+	if (bUsesRapidIterationParams)
+	{
+		KeyString += TEXT("USESRI_");
+	}
+	else
+	{
+		KeyString += TEXT("NORI_");
+	}
+
 	for (int32 Idx = 0; Idx < AdditionalDefines.Num(); Idx++)
 	{
 		KeyString += AdditionalDefines[Idx];
@@ -230,12 +251,21 @@ class UNiagaraSystem* UNiagaraScript::FindRootSystem()
 void UNiagaraScript::ComputeVMCompilationId(FNiagaraVMExecutableDataId& Id) const
 {
 	Id = FNiagaraVMExecutableDataId();
+
+	Id.bUsesRapidIterationParams = true;
+	Id.DetailLevelMask = 0xFFFFFFFF; // Unused for now.
 	
 	// Ideally we wouldn't want to do this but rather than push the data down
 	// from the emitter.
 	UObject* Obj = GetOuter();
 	if (UNiagaraEmitter* Emitter = Cast<UNiagaraEmitter>(Obj))
 	{
+		UNiagaraSystem* EmitterOwner = Cast<UNiagaraSystem>(Emitter->GetOuter());
+		if (EmitterOwner && EmitterOwner->bBakeOutRapidIteration)
+		{
+			Id.bUsesRapidIterationParams = false;
+		}
+
 		if ((Emitter->bInterpolatedSpawning && Usage == ENiagaraScriptUsage::ParticleGPUComputeScript) || 
 			(Emitter->bInterpolatedSpawning && Usage == ENiagaraScriptUsage::ParticleSpawnScript) ||
 			Usage == ENiagaraScriptUsage::ParticleSpawnScriptInterpolated)
@@ -254,10 +284,20 @@ void UNiagaraScript::ComputeVMCompilationId(FNiagaraVMExecutableDataId& Id) cons
 		{
 			Id.AdditionalDefines.Add(TEXT("Emitter.Determinism"));
 		}
+
+		if (!Emitter->bBakeOutRapidIteration)
+		{
+			Id.bUsesRapidIterationParams = true;
+		}
 	}
 
 	if (UNiagaraSystem* System = Cast<UNiagaraSystem>(Obj))
 	{
+		if (System->bBakeOutRapidIteration)
+		{
+			Id.bUsesRapidIterationParams = false;
+		}
+
 		for (const FNiagaraEmitterHandle& EmitterHandle: System->GetEmitterHandles())
 		{
 			UNiagaraEmitter* Emitter = Cast<UNiagaraEmitter>(EmitterHandle.GetInstance());
@@ -275,6 +315,46 @@ void UNiagaraScript::ComputeVMCompilationId(FNiagaraVMExecutableDataId& Id) cons
 		}
 	}
 
+	// If we aren't using rapid iteration parameters, we need to bake them into the hashstate for the compile id. This 
+	// makes their values part of the lookup. 
+	if (false == Id.bUsesRapidIterationParams)
+	{
+		FSHA1 HashState;
+		TArray<FNiagaraVariable> Vars;
+		RapidIterationParameters.GetParameters(Vars);
+		//UE_LOG(LogNiagara, Display, TEXT("AreScriptAndSourceSynchronized %s ======================== "), *GetFullName());
+		for (int32 i = 0; i < Vars.Num(); i++)
+		{
+			if (Vars[i].IsDataInterface() || Vars[i].IsUObject())
+			{
+				// Skip these types as they don't bake out, just normal parameters get baked.
+			}
+			else
+			{
+				// Hash the name, type, and value of each parameter..
+				FString VarName = Vars[i].GetName().ToString();
+				FString VarTypeName = Vars[i].GetType().GetName();
+				HashState.UpdateWithString(*VarName, VarName.Len());
+				HashState.UpdateWithString(*VarTypeName, VarTypeName.Len());
+				const uint8* VarData = RapidIterationParameters.GetParameterData(Vars[i]);
+				if (VarData)
+				{
+					//UE_LOG(LogNiagara, Display, TEXT("Param %s %s %s"), *VarTypeName, *VarName, *ByteStr);
+					HashState.Update(VarData, Vars[i].GetType().GetSize());
+				}
+			}			
+		}
+		HashState.Final();
+
+		TArray<uint8> DataHash;
+		DataHash.AddUninitialized(20);
+		HashState.GetHash(DataHash.GetData());
+
+		FNiagaraCompileHash Hash(DataHash);
+		Id.ReferencedCompileHashes.Add(Hash);
+		Id.ReferencedObjects.Add(nullptr);
+	}
+
 	static const auto UseShaderStagesCVar = IConsoleManager::Get().FindConsoleVariable(TEXT("fx.UseShaderStages"));
 	if (UseShaderStagesCVar->GetInt())
 	{
@@ -344,11 +424,69 @@ FNiagaraScriptExecutionParameterStore* UNiagaraScript::GetExecutionReadyParamete
 
 void UNiagaraScript::Serialize(FArchive& Ar)
 {
-	Super::Serialize(Ar);
-
 	Ar.UsingCustomVersion(FNiagaraCustomVersion::GUID);		// only changes version if not loading
 	const int32 NiagaraVer = Ar.CustomVer(FNiagaraCustomVersion::GUID);
-	
+
+	FNiagaraParameterStore TemporaryStore;
+	int32 NumRemoved = 0;
+	if (Ar.IsCooking())
+	{
+		bool bUsesRapidIterationParams = true;
+
+#if WITH_EDITORONLY_DATA
+
+		if (UNiagaraEmitter* Emitter = Cast<UNiagaraEmitter>(GetOuter()))
+		{
+			UNiagaraSystem* EmitterOwner = Cast<UNiagaraSystem>(Emitter->GetOuter());
+			if (EmitterOwner && EmitterOwner->bBakeOutRapidIteration)
+			{
+				bUsesRapidIterationParams = false;
+			}
+
+			if (!Emitter->bBakeOutRapidIteration)
+			{
+				bUsesRapidIterationParams = true;
+			}
+		}
+		else if (UNiagaraSystem* System = Cast<UNiagaraSystem>(GetOuter()))
+		{
+			if (System && System->bBakeOutRapidIteration)
+			{
+				bUsesRapidIterationParams = false;
+			}
+		}
+#endif
+
+		if (!bUsesRapidIterationParams)
+		{
+			// Copy off the parameter store for now..
+			TemporaryStore = RapidIterationParameters;
+
+			// Get the active parameters
+			TArray<FNiagaraVariable> Vars;
+			RapidIterationParameters.GetParameters(Vars);
+
+			// Remove all parameters that aren't data interfaces or uobjects
+			for (const FNiagaraVariable& Var : Vars)
+			{
+				if (Var.IsDataInterface() || Var.IsUObject())
+					continue;
+				RapidIterationParameters.RemoveParameter(Var);
+				NumRemoved++;
+			}
+
+			UE_LOG(LogNiagara, Display, TEXT("Pruned %d/%d parameters from script %s"), NumRemoved, TemporaryStore.GetNumParameters(), *GetFullName());
+		}
+	}
+
+	Super::Serialize(Ar);
+
+	// Restore after serialize
+	if (Ar.IsCooking() && NumRemoved > 0)
+	{
+		RapidIterationParameters = TemporaryStore;
+	}
+
 	bool IsValidShaderScript = false;
 	if (NiagaraVer < FNiagaraCustomVersion::DontCompileGPUWhenNotNeeded)
 	{
@@ -365,6 +503,23 @@ void UNiagaraScript::Serialize(FArchive& Ar)
 		IsValidShaderScript = CanBeRunOnGpu();
 	}
 
+	if (IsValidShaderScript)
+	{
+		if (NiagaraVer < FNiagaraCustomVersion::UseHashesToIdentifyCompileStateOfTopLevelScripts)
+		{
+			// In some rare cases a GPU script could have been saved in an error state in a version where skeletal mesh or static mesh data interfaces didn't work properly on GPU.
+			// This would fail in the current regime.
+			for (const FNiagaraScriptDataInterfaceCompileInfo& InterfaceInfo : CachedScriptVM.DataInterfaceInfo)
+			{
+				if (InterfaceInfo.Type.GetClass() == UNiagaraDataInterfaceSkeletalMesh::StaticClass() ||
+					InterfaceInfo.Type.GetClass() == UNiagaraDataInterfaceStaticMesh::StaticClass())
+				{
+					IsValidShaderScript = false;
+				}
+			}
+		}
+	}
+
 	if ( (!Ar.IsLoading() && IsValidShaderScript)		// saving shader maps only for particle sim and spawn scripts
 		|| (Ar.IsLoading() && NiagaraVer >= FNiagaraCustomVersion::NiagaraShaderMaps && (NiagaraVer < FNiagaraCustomVersion::NiagaraShaderMapCooking || IsValidShaderScript))  // load only if we know shader map is presen
 		)
@@ -505,6 +660,8 @@ void UNiagaraScript::PostLoad()
 	GenerateStatScopeIDs();
 #endif
 
+	//FNiagaraUtilities::DumpHLSLText(RapidIterationParameters.ToString(), *GetPathName());
+
 }
 
 bool UNiagaraScript::IsReadyToRun(ENiagaraSimTarget SimTarget) const
@@ -1151,7 +1308,8 @@ void UNiagaraScript::CacheResourceShadersForCooking(EShaderPlatform ShaderPlatfo
 			check(CachedScriptVMId.BaseScriptCompileHash.IsValid());
 
 			NewResource->SetScript(this, (ERHIFeatureLevel::Type)TargetFeatureLevel, CachedScriptVMId.CompilerVersionID, CachedScriptVMId.BaseScriptID, CachedScriptVMId.AdditionalDefines,
-				CachedScriptVMId.BaseScriptCompileHash,	CachedScriptVMId.ReferencedCompileHashes, GetFullName());
+				CachedScriptVMId.BaseScriptCompileHash,	CachedScriptVMId.ReferencedCompileHashes, 
+				CachedScriptVMId.bUsesRapidIterationParams, CachedScriptVMId.DetailLevelMask, GetFullName());
 			ResourceToCache = NewResource;
 
 			check(ResourceToCache);
@@ -1209,14 +1367,15 @@ void UNiagaraScript::CacheResourceShadersForRendering(bool bRegenerateId, bool b
 
 	//UpdateResourceAllocations();
 
-	if (FApp::CanEverRender() && CanBeRunOnGpu())
+	if (CanBeRunOnGpu())
 	{
 		if (Source)
 		{
 			FNiagaraShaderScript* ResourceToCache;
 			ERHIFeatureLevel::Type CacheFeatureLevel = GMaxRHIFeatureLevel;
 			ScriptResource.SetScript(this, FeatureLevel, CachedScriptVMId.CompilerVersionID, CachedScriptVMId.BaseScriptID, CachedScriptVMId.AdditionalDefines,
-				CachedScriptVMId.BaseScriptCompileHash, CachedScriptVMId.ReferencedCompileHashes, GetFullName());
+				CachedScriptVMId.BaseScriptCompileHash, CachedScriptVMId.ReferencedCompileHashes, 
+				CachedScriptVMId.bUsesRapidIterationParams, CachedScriptVMId.DetailLevelMask, GetFullName());
 
 			//if (ScriptResourcesByFeatureLevel[FeatureLevel])
 			{
diff --git a/Engine/Plugins/FX/Niagara/Source/Niagara/Private/NiagaraScriptExecutionContext.cpp b/Engine/Plugins/FX/Niagara/Source/Niagara/Private/NiagaraScriptExecutionContext.cpp
--- a/Engine/Plugins/FX/Niagara/Source/Niagara/Private/NiagaraScriptExecutionContext.cpp
+++ b/Engine/Plugins/FX/Niagara/Source/Niagara/Private/NiagaraScriptExecutionContext.cpp
@@ -341,7 +341,7 @@ void FNiagaraGPUSystemTick::Init(FNiagaraSystemInstance* InSystemInstance)
 	{
 		FNiagaraEmitterInstance* Emitter = &InSystemInstance->GetEmitters()[i].Get();
 
-		if (Emitter->GetCachedEmitter()->SimTarget == ENiagaraSimTarget::GPUComputeSim && Emitter->GetGPUContext() != nullptr && Emitter->GetExecutionState() != ENiagaraExecutionState::Complete)
+		if (Emitter && Emitter->GetCachedEmitter()->SimTarget == ENiagaraSimTarget::GPUComputeSim && Emitter->GetGPUContext() != nullptr && Emitter->GetExecutionState() != ENiagaraExecutionState::Complete)
 		{
 			FNiagaraComputeInstanceData* InstanceData = new (&Instances[InstanceIndex]) FNiagaraComputeInstanceData;
 			InstanceIndex++;
diff --git a/Engine/Plugins/FX/Niagara/Source/Niagara/Private/NiagaraSystem.cpp b/Engine/Plugins/FX/Niagara/Source/Niagara/Private/NiagaraSystem.cpp
--- a/Engine/Plugins/FX/Niagara/Source/Niagara/Private/NiagaraSystem.cpp
+++ b/Engine/Plugins/FX/Niagara/Source/Niagara/Private/NiagaraSystem.cpp
@@ -31,6 +31,14 @@ DECLARE_CYCLE_STAT(TEXT("Niagara - System - CompileScript_ResetAfter"), STAT_Nia
 // {
 // }
 
+static int GNiagaraForceSystemsToCookOutRapidIterationOnLoad = 0;
+static FAutoConsoleVariableRef CVarNiagaraForceSystemsToCookOutRapidIterationOnLoad(
+	TEXT("fx.NiagaraForceSystemsToCookOutRapidIterationOnLoad"),
+	GNiagaraForceSystemsToCookOutRapidIterationOnLoad,
+	TEXT("When enabled UNiagaraSystem's bBakeOutRapidIteration will be forced to true on PostLoad of the system."),
+	ECVF_Default
+);
+
 //////////////////////////////////////////////////////////////////////////
 
 UNiagaraSystem::UNiagaraSystem(const FObjectInitializer& ObjectInitializer)
@@ -73,35 +81,10 @@ void UNiagaraSystem::PreSave(const class ITargetPlatform * TargetPlatform)
 #if WITH_EDITOR
 void UNiagaraSystem::BeginCacheForCookedPlatformData(const ITargetPlatform *TargetPlatform)
 {
+	//UE_LOG(LogNiagara, Display, TEXT("UNiagaraSystem::BeginCacheForCookedPlatformData %s %s"), *GetFullName(), GIsSavingPackage ? TEXT("Saving...") : TEXT("Not Saving..."));
 	Super::BeginCacheForCookedPlatformData(TargetPlatform);
 	
 #if WITH_EDITORONLY_DATA
-	/*
-	bool bAnyChanged = false;
-	for (int32 i = 0; i < GetNumEmitters(); i++)
-	{
-		FNiagaraEmitterHandle& EmitterHandle = GetEmitterHandle(i);
-		if (EmitterHandle.GetIsEnabled())
-		{
-			const UNiagaraEmitter* Emitter = EmitterHandle.GetInstance();
-			if (Emitter != nullptr)
-			{
-				if (!INiagaraModule::IsTargetPlatformIncludedInLevelRangeForCook(TargetPlatform, Emitter))
-				{
-					EmitterHandle.SetIsEnabled(false, *this, false);
-					UE_LOG(LogNiagara, Display, TEXT("Pruning emitter, detail level mismatch (only works if platform can't change detail level at runtime!): %s - set fx.NiagaraPruneEmittersOnCookByDetailLevel to 0 in DeviceProfile.ini for the target profile to avoid"),  *EmitterHandle.GetName().ToString());
-					EmitterHandle.ClearEmitter();
-					bAnyChanged = true;
-				}
-			}
-		}
-	}
-
-	if (bAnyChanged)
-	{
-		RequestCompile(false);
-	}
-	*/
 
 	WaitForCompilationComplete();
 #endif
@@ -146,7 +129,7 @@ bool UNiagaraSystem::UsesCollection(const UNiagaraParameterCollection* Collectio
 
 	for (const FNiagaraEmitterHandle& EmitterHandle : GetEmitterHandles())
 	{
-		if (EmitterHandle.GetInstance()->UsesCollection(Collection))
+		if (EmitterHandle.GetInstance() && EmitterHandle.GetInstance()->UsesCollection(Collection))
 		{
 			return true;
 		}
@@ -439,7 +422,14 @@ void UNiagaraSystem::PostLoad()
 		}
 #endif
 	}
+	if (GNiagaraForceSystemsToCookOutRapidIterationOnLoad == 1 && !bBakeOutRapidIteration)
+	{
+		WaitForCompilationComplete();
+		bBakeOutRapidIteration = true;
+		RequestCompile(false);
+	}
 #endif // WITH_EDITORONLY_DATA
+
 }
 
 #if WITH_EDITORONLY_DATA
@@ -456,6 +446,11 @@ const UNiagaraEditorDataBase* UNiagaraSystem::GetEditorData() const
 
 bool UNiagaraSystem::ReferencesInstanceEmitter(UNiagaraEmitter& Emitter)
 {
+	if (&Emitter == nullptr)
+	{
+		return false;
+	}
+
 	for (FNiagaraEmitterHandle& Handle : EmitterHandles)
 	{
 		if (&Emitter == Handle.GetInstance())
@@ -472,10 +467,13 @@ void UNiagaraSystem::RefreshSystemParametersFromEmitter(const FNiagaraEmitterHan
 	if (ensureMsgf(EmitterHandles.ContainsByPredicate([=](const FNiagaraEmitterHandle& OwnedEmitterHandle) { return OwnedEmitterHandle.GetId() == EmitterHandle.GetId(); }),
 		TEXT("Can't refresh parameters from an emitter handle this system doesn't own.")))
 	{
-		EmitterHandle.GetInstance()->EmitterSpawnScriptProps.Script->RapidIterationParameters.CopyParametersTo(
-			SystemSpawnScript->RapidIterationParameters, false, FNiagaraParameterStore::EDataInterfaceCopyMethod::None);
-		EmitterHandle.GetInstance()->EmitterUpdateScriptProps.Script->RapidIterationParameters.CopyParametersTo(
-			SystemUpdateScript->RapidIterationParameters, false, FNiagaraParameterStore::EDataInterfaceCopyMethod::None);
+		if (EmitterHandle.GetInstance())
+		{
+			EmitterHandle.GetInstance()->EmitterSpawnScriptProps.Script->RapidIterationParameters.CopyParametersTo(
+				SystemSpawnScript->RapidIterationParameters, false, FNiagaraParameterStore::EDataInterfaceCopyMethod::None);
+			EmitterHandle.GetInstance()->EmitterUpdateScriptProps.Script->RapidIterationParameters.CopyParametersTo(
+				SystemUpdateScript->RapidIterationParameters, false, FNiagaraParameterStore::EDataInterfaceCopyMethod::None);
+		}
 	}
 }
 
@@ -485,8 +483,11 @@ void UNiagaraSystem::RemoveSystemParametersForEmitter(const FNiagaraEmitterHandl
 	if (ensureMsgf(EmitterHandles.ContainsByPredicate([=](const FNiagaraEmitterHandle& OwnedEmitterHandle) { return OwnedEmitterHandle.GetId() == EmitterHandle.GetId(); }),
 		TEXT("Can't remove parameters for an emitter handle this system doesn't own.")))
 	{
-		EmitterHandle.GetInstance()->EmitterSpawnScriptProps.Script->RapidIterationParameters.RemoveParameters(SystemSpawnScript->RapidIterationParameters);
-		EmitterHandle.GetInstance()->EmitterUpdateScriptProps.Script->RapidIterationParameters.RemoveParameters(SystemUpdateScript->RapidIterationParameters);
+		if (EmitterHandle.GetInstance())
+		{
+			EmitterHandle.GetInstance()->EmitterSpawnScriptProps.Script->RapidIterationParameters.RemoveParameters(SystemSpawnScript->RapidIterationParameters);
+			EmitterHandle.GetInstance()->EmitterUpdateScriptProps.Script->RapidIterationParameters.RemoveParameters(SystemUpdateScript->RapidIterationParameters);
+		}
 	}
 }
 #endif
@@ -532,7 +533,7 @@ bool UNiagaraSystem::IsReadyToRun() const
 
 	for (const FNiagaraEmitterHandle& Handle : EmitterHandles)
 	{
-		if (!Handle.GetInstance()->IsReadyToRun())
+		if (Handle.GetInstance() && !Handle.GetInstance()->IsReadyToRun())
 		{
 			return false;
 		}
@@ -604,7 +605,7 @@ bool UNiagaraSystem::IsValid()const
 
 	for (const FNiagaraEmitterHandle& Handle : EmitterHandles)
 	{
-		if (Handle.GetIsEnabled() && !Handle.GetInstance()->IsValid())
+		if (Handle.GetIsEnabled() && Handle.GetInstance() && !Handle.GetInstance()->IsValid())
 		{
 			return false;
 		}
@@ -694,8 +695,11 @@ void UNiagaraSystem::InvalidateCachedCompileIds()
 
 	for (FNiagaraEmitterHandle Handle : EmitterHandles)
 	{
-		UNiagaraScriptSourceBase* GraphSource = Handle.GetInstance()->GraphSource;
-		GraphSource->InvalidateCachedCompileIds();
+		if (Handle.GetInstance())
+		{
+			UNiagaraScriptSourceBase* GraphSource = Handle.GetInstance()->GraphSource;
+			GraphSource->InvalidateCachedCompileIds();
+		}
 	}
 }
 
@@ -806,13 +810,16 @@ bool UNiagaraSystem::QueryCompileComplete(bool bWait, bool bDoPost, bool bDoNotA
 		{
 			for (FNiagaraEmitterHandle Handle : EmitterHandles)
 			{
-				if (Handle.GetIsEnabled())
-				{
-					Handle.GetInstance()->OnPostCompile();
-				}
-				else
+				if (Handle.GetInstance())
 				{
-					Handle.GetInstance()->InvalidateCompileResults();
+					if (Handle.GetIsEnabled())
+					{
+						Handle.GetInstance()->OnPostCompile();
+					}
+					else
+					{
+						Handle.GetInstance()->InvalidateCompileResults();
+					}
 				}
 			}
 		}
@@ -1004,7 +1011,7 @@ bool UNiagaraSystem::RequestCompile(bool bForce)
 	for (int32 i = 0; i < EmitterHandles.Num(); i++)
 	{
 		FNiagaraEmitterHandle Handle = EmitterHandles[i];
-		if (Handle.GetIsEnabled())
+		if (Handle.GetInstance() && Handle.GetIsEnabled())
 		{
 			UNiagaraScriptSourceBase* GraphSource = Handle.GetInstance()->GraphSource;
 			TSharedPtr<FNiagaraCompileRequestDataBase, ESPMode::ThreadSafe> EmitterPrecompiledData = SystemPrecompiledData->GetDependentRequest(i);
@@ -1093,7 +1100,7 @@ void UNiagaraSystem::InitEmitterCompiledData()
 			for (int32 EmitterIdx = 0; EmitterIdx < EmitterHandles.Num(); ++EmitterIdx)
 			{
 				UNiagaraEmitter* Emitter = EmitterHandles[EmitterIdx].GetInstance();
-				if ensureMsgf(Emitter != nullptr, TEXT("Failed to get Emitter Instance from Emitter Handle in post compile, please investigate."))
+				if (Emitter)
 				{
 					FString EmitterName = Emitter->GetUniqueEmitterName() + TEXT(".");
 					if (Var.GetType() == SpawnInfoDef && Var.GetName().ToString().StartsWith(EmitterName))
@@ -1109,7 +1116,7 @@ void UNiagaraSystem::InitEmitterCompiledData()
 			for (int32 EmitterIdx = 0; EmitterIdx < EmitterHandles.Num(); ++EmitterIdx)
 			{
 				UNiagaraEmitter* Emitter = EmitterHandles[EmitterIdx].GetInstance();
-				if ensureMsgf(Emitter != nullptr, TEXT("Failed to get Emitter Instance from Emitter Handle in post compile, please investigate."))
+				if (Emitter)
 				{
 					FString EmitterName = Emitter->GetUniqueEmitterName() + TEXT(".");
 					if (Var.GetType() == SpawnInfoDef && Var.GetName().ToString().StartsWith(EmitterName))
diff --git a/Engine/Plugins/FX/Niagara/Source/Niagara/Private/NiagaraSystemInstance.cpp b/Engine/Plugins/FX/Niagara/Source/Niagara/Private/NiagaraSystemInstance.cpp
--- a/Engine/Plugins/FX/Niagara/Source/Niagara/Private/NiagaraSystemInstance.cpp
+++ b/Engine/Plugins/FX/Niagara/Source/Niagara/Private/NiagaraSystemInstance.cpp
@@ -285,14 +285,17 @@ bool FNiagaraSystemInstance::RequestCapture(const FGuid& RequestId)
 	for (const FNiagaraEmitterHandle& Handle : GetSystem()->GetEmitterHandles())
 	{
 		TArray<UNiagaraScript*> Scripts;
-		Handle.GetInstance()->GetScripts(Scripts, false);
-
-		for (UNiagaraScript* Script : Scripts)
+		if (Handle.GetInstance())
 		{
-			TSharedPtr<struct FNiagaraScriptDebuggerInfo, ESPMode::ThreadSafe> DebugInfoPtr = MakeShared<FNiagaraScriptDebuggerInfo, ESPMode::ThreadSafe>(Handle.GetIdName(), Script->GetUsage(), Script->GetUsageId());
-			DebugInfoPtr->bWritten = false;
+			Handle.GetInstance()->GetScripts(Scripts, false);
+
+			for (UNiagaraScript* Script : Scripts)
+			{
+				TSharedPtr<struct FNiagaraScriptDebuggerInfo, ESPMode::ThreadSafe> DebugInfoPtr = MakeShared<FNiagaraScriptDebuggerInfo, ESPMode::ThreadSafe>(Handle.GetIdName(), Script->GetUsage(), Script->GetUsageId());
+				DebugInfoPtr->bWritten = false;
 
-			TempCaptureHolder->Add(DebugInfoPtr);
+				TempCaptureHolder->Add(DebugInfoPtr);
+			}
 		}
 	}
 	CapturedFrames.Add(RequestId, TempCaptureHolder);
@@ -1552,20 +1555,25 @@ void FNiagaraSystemInstance::InitEmitters()
 		const TArray<FNiagaraEmitterHandle>& EmitterHandles = System->GetEmitterHandles();
 		for (int32 EmitterIdx=0; EmitterIdx < System->GetEmitterHandles().Num(); ++EmitterIdx)
 		{
-			TSharedRef<FNiagaraEmitterInstance, ESPMode::ThreadSafe> Sim = MakeShared<FNiagaraEmitterInstance, ESPMode::ThreadSafe>(this);
-			Sim->Init(EmitterIdx, ID);
-			if (System->bFixedBounds)
 			{
-				Sim->SetSystemFixedBoundsOverride(System->GetFixedBounds());
+				TSharedRef<FNiagaraEmitterInstance, ESPMode::ThreadSafe> Sim = MakeShared<FNiagaraEmitterInstance, ESPMode::ThreadSafe>(this);
+				Sim->Init(EmitterIdx, ID);
+				if (System->bFixedBounds)
+				{
+					Sim->SetSystemFixedBoundsOverride(System->GetFixedBounds());
+				}
+				Emitters.Add(Sim);
 			}
-			Emitters.Add(Sim);
 		}
 
 		for (TSharedRef<FNiagaraEmitterInstance, ESPMode::ThreadSafe> Simulation : Emitters)
 		{
-			bHasGPUEmitters |= Simulation->GetCachedEmitter()->SimTarget == ENiagaraSimTarget::GPUComputeSim;
+			if (Simulation->GetCachedEmitter() != nullptr)
+			{
+				bHasGPUEmitters |= Simulation->GetCachedEmitter()->SimTarget == ENiagaraSimTarget::GPUComputeSim;
 
-			Simulation->PostInitSimulation();
+				Simulation->PostInitSimulation();
+			}
 		}
 
 		if (System->bFixedBounds)
@@ -1703,7 +1711,7 @@ void FNiagaraSystemInstance::Tick_Concurrent()
 		FNiagaraEmitterInstance& Inst = Emitters[EmitterIdx].Get();
 		Inst.Tick(CachedDeltaSeconds);
 
-		if (Inst.GetCachedEmitter()->SimTarget == ENiagaraSimTarget::GPUComputeSim && Inst.GetGPUContext() != nullptr && (Inst.GetExecutionState() != ENiagaraExecutionState::Complete))
+		if (Inst.GetCachedEmitter() && Inst.GetCachedEmitter()->SimTarget == ENiagaraSimTarget::GPUComputeSim && Inst.GetGPUContext() != nullptr && (Inst.GetExecutionState() != ENiagaraExecutionState::Complete))
 		{
 			TotalParamSize += Inst.GetGPUContext()->CombinedParamStore.GetPaddedParameterSizeInBytes();
 			ActiveGPUEmitterCount++;
diff --git a/Engine/Plugins/FX/Niagara/Source/Niagara/Private/NiagaraSystemSimulation.cpp b/Engine/Plugins/FX/Niagara/Source/Niagara/Private/NiagaraSystemSimulation.cpp
--- a/Engine/Plugins/FX/Niagara/Source/Niagara/Private/NiagaraSystemSimulation.cpp
+++ b/Engine/Plugins/FX/Niagara/Source/Niagara/Private/NiagaraSystemSimulation.cpp
@@ -353,20 +353,26 @@ bool FNiagaraSystemSimulation::Init(UNiagaraSystem* InSystem, UWorld* InWorld, b
 		{
 			FNiagaraEmitterHandle& EmitterHandle = System->GetEmitterHandle(EmitterIdx);
 			UNiagaraEmitter* Emitter = EmitterHandle.GetInstance();
-			FString EmitterName = Emitter->GetUniqueEmitterName();
-			check(Emitter);
-			EmitterExecutionStateAccessors.Emplace(MainDataSet, FNiagaraVariable(EnumPtr, *(EmitterName + TEXT(".ExecutionState"))));
-			const TArray<FNiagaraEmitterCompiledData>& EmitterCompiledData = System->GetEmitterCompiledData();
-			
-			check(EmitterCompiledData.Num() == System->GetNumEmitters());
-			for (FName AttrName : EmitterCompiledData[EmitterIdx].SpawnAttributes)
+			if (Emitter)
 			{
-				EmitterSpawnInfoAccessors[EmitterIdx].Emplace(MainDataSet, FNiagaraVariable(FNiagaraTypeDefinition(FNiagaraSpawnInfo::StaticStruct()), AttrName));
-			}
+				FString EmitterName = Emitter->GetUniqueEmitterName();
+				EmitterExecutionStateAccessors.Emplace(MainDataSet, FNiagaraVariable(EnumPtr, *(EmitterName + TEXT(".ExecutionState"))));
+				const TArray<FNiagaraEmitterCompiledData>& EmitterCompiledData = System->GetEmitterCompiledData();
+
+				check(EmitterCompiledData.Num() == System->GetNumEmitters());
+				for (FName AttrName : EmitterCompiledData[EmitterIdx].SpawnAttributes)
+				{
+					EmitterSpawnInfoAccessors[EmitterIdx].Emplace(MainDataSet, FNiagaraVariable(FNiagaraTypeDefinition(FNiagaraSpawnInfo::StaticStruct()), AttrName));
+				}
 
-			if (Emitter->bLimitDeltaTime)
+				if (Emitter->bLimitDeltaTime)
+				{
+					MaxDeltaTime = MaxDeltaTime.IsSet() ? FMath::Min(MaxDeltaTime.GetValue(), Emitter->MaxDeltaTimePerTick) : Emitter->MaxDeltaTimePerTick;
+				}
+			}
+			else
 			{
-				MaxDeltaTime = MaxDeltaTime.IsSet() ? FMath::Min(MaxDeltaTime.GetValue(), Emitter->MaxDeltaTimePerTick) : Emitter->MaxDeltaTimePerTick;
+				EmitterExecutionStateAccessors.Emplace();
 			}
 		}
 
diff --git a/Engine/Plugins/FX/Niagara/Source/NiagaraCore/Private/NiagaraCustomVersion.cpp b/Engine/Plugins/FX/Niagara/Source/NiagaraCore/Private/NiagaraCustomVersion.cpp
--- a/Engine/Plugins/FX/Niagara/Source/NiagaraCore/Private/NiagaraCustomVersion.cpp
+++ b/Engine/Plugins/FX/Niagara/Source/NiagaraCore/Private/NiagaraCustomVersion.cpp
@@ -8,4 +8,4 @@ const FGuid FNiagaraCustomVersion::GUID(0xFCF57AFA, 0x50764283, 0xB9A9E658, 0xFF
 // Register the custom version with core
 FCustomVersionRegistration GRegisterNiagaraCustomVersion(FNiagaraCustomVersion::GUID, FNiagaraCustomVersion::LatestVersion, TEXT("NiagaraVer"));
 
-const FGuid FNiagaraCustomVersion::LatestScriptCompileVersion(0x31F56712, 0xC9264924, 0x9B11B480, 0x1499071A);
+const FGuid FNiagaraCustomVersion::LatestScriptCompileVersion(0x31B222C7, 0x2AAF448A, 0xA9DCA5E7, 0x1DAE433D);
diff --git a/Engine/Plugins/FX/Niagara/Source/NiagaraCore/Public/NiagaraCustomVersion.h b/Engine/Plugins/FX/Niagara/Source/NiagaraCore/Public/NiagaraCustomVersion.h
--- a/Engine/Plugins/FX/Niagara/Source/NiagaraCore/Public/NiagaraCustomVersion.h
+++ b/Engine/Plugins/FX/Niagara/Source/NiagaraCore/Public/NiagaraCustomVersion.h
@@ -125,10 +125,13 @@ struct FNiagaraCustomVersion
 		AddLibraryAssetProperty, // Add property to all Niagara scripts indicating whether or not they belong to the library
 
 		AddAdditionalDefinesProperty, // Addding additional defines to the GPU script
-
+		
 		RemoveGraphUsageCompileIds, // Remove the random compile id guids from the cached script usage and from the compile and script ids since the hashes serve the same purpose and are deterministic.
+			
+		AddRIAndDetailLevel, //Adding UseRapidIterationParams and DetailLevelMask to the GPU script
 
-		// -----<new versions can be added above this line>-------------------------------------------------
+		// DO NOT ADD A NEW VERSION UNLESS YOU HAVE TALKED TO THE NIAGARA LEAD. Mismanagement of these versions can lead to data loss if it is adjusted in multiple streams simultaneously.
+		// -----<new versions can be added above this line>  -------------------------------------------------
 		VersionPlusOne,
 		LatestVersion = VersionPlusOne - 1,
 	};
diff --git a/Engine/Plugins/FX/Niagara/Source/NiagaraEditor/Private/NiagaraCompiler.cpp b/Engine/Plugins/FX/Niagara/Source/NiagaraEditor/Private/NiagaraCompiler.cpp
--- a/Engine/Plugins/FX/Niagara/Source/NiagaraEditor/Private/NiagaraCompiler.cpp
+++ b/Engine/Plugins/FX/Niagara/Source/NiagaraEditor/Private/NiagaraCompiler.cpp
@@ -54,24 +54,25 @@ static FAutoConsoleVariableRef CVarForceNiagaraTranslatorSingleThreaded(
 static int32 GbForceNiagaraTranslatorDump = 0;
 static FAutoConsoleVariableRef CVarForceNiagaraTranslatorDump(
 	TEXT("fx.ForceNiagaraTranslatorDump"),
-	GbForceNiagaraTranslatorSingleThreaded,
+	GbForceNiagaraTranslatorDump,
 	TEXT("If > 0 all translation generated HLSL will be dumped \n"),
 	ECVF_Default
 );
 
+static int32 GbForceNiagaraVMBinaryDump = 0;
+static FAutoConsoleVariableRef CVarForceNiagaraVMBinaryDump(
+	TEXT("fx.ForceNiagaraVMBinaryDump"),
+	GbForceNiagaraVMBinaryDump,
+	TEXT("If > 0 all translation generated binary text will be dumped \n"),
+	ECVF_Default
+);
+
 static FCriticalSection TranslationCritSec;
 
 void DumpHLSLText(const FString& SourceCode, const FString& DebugName)
 {
-	UE_LOG(LogNiagaraEditor, Log, TEXT("Compile output as text: %s"), *DebugName);
-	UE_LOG(LogNiagaraEditor, Log, TEXT("==================================================================================="));
-	TArray<FString> OutputByLines;
-	SourceCode.ParseIntoArrayLines(OutputByLines, false);
-	for (int32 i = 0; i < OutputByLines.Num(); i++)
-	{
-		UE_LOG(LogNiagaraEditor, Log, TEXT("/*%04d*/\t\t%s"), i + 1, *OutputByLines[i]);
-	}
-	UE_LOG(LogNiagaraEditor, Log, TEXT("==================================================================================="));
+	FScopeLock Lock(&TranslationCritSec);
+	FNiagaraUtilities::DumpHLSLText(SourceCode, DebugName);
 }
 
 void FNiagaraCompileRequestData::VisitReferencedGraphs(UNiagaraGraph* InSrcGraph, UNiagaraGraph* InDupeGraph, ENiagaraScriptUsage InUsage, FCompileConstantResolver ConstantResolver)
@@ -172,7 +173,7 @@ void FNiagaraCompileRequestData::VisitReferencedGraphsRecursive(UNiagaraGraph* I
 							}
 						}
 
-						/*UE_LOG(LogNiagaraEditor, Log, TEXT("FunctionGraph = %p SrcScriptName = %s DiffPackage? %s Numerics? %s NumericInputs? %s"), FunctionGraph, *FunctionGraph->GetPathName(),
+						/*UE_LOG(LogNiagaraEditor, Display, TEXT("FunctionGraph = %p SrcScriptName = %s DiffPackage? %s Numerics? %s NumericInputs? %s"), FunctionGraph, *FunctionGraph->GetPathName(),
 							bFromDifferentPackage ? TEXT("yes") : TEXT("no"), bHasNumericParams ? TEXT("yes") : TEXT("no"), bHasNumericInputs ? TEXT("yes") : TEXT("no"));*/
 
 						UNiagaraGraph* ProcessedGraph = nullptr;
@@ -194,7 +195,7 @@ void FNiagaraCompileRequestData::VisitReferencedGraphsRecursive(UNiagaraGraph* I
 								FNiagaraEditorUtilities::PreprocessFunctionGraph(Schema, ProcessedGraph, CallInputs, CallOutputs, ScriptUsage, ConstantResolver);
 								FunctionCallNode->FunctionScript = DupeScript;
 
-								/*UE_LOG(LogNiagaraEditor, Log, TEXT("Duplicating script %s OriginalGraph %p NewGraph %p"), *FunctionScript->GetFName().ToString(),
+								/*UE_LOG(LogNiagaraEditor, Display, TEXT("Duplicating script %s OriginalGraph %p NewGraph %p"), *FunctionScript->GetFName().ToString(),
 									Cast<UNiagaraScriptSource>(FunctionScript->GetSource())->NodeGraph,
 									Cast<UNiagaraScriptSource>(DupeScript->GetSource())->NodeGraph);*/
 							}
@@ -221,7 +222,7 @@ void FNiagaraCompileRequestData::VisitReferencedGraphsRecursive(UNiagaraGraph* I
 							FNiagaraEditorUtilities::PreprocessFunctionGraph(Schema, ProcessedGraph, CallInputs, CallOutputs, ScriptUsage, ConstantResolver);
 							FunctionCallNode->FunctionScript = DupeScript;
 
-							/*UE_LOG(LogNiagaraEditor, Log, TEXT("Duplicating script %s OriginalGraph %p NewGraph %p"), *FunctionScript->GetFName().ToString(),
+							/*UE_LOG(LogNiagaraEditor, Display, TEXT("Duplicating script %s OriginalGraph %p NewGraph %p"), *FunctionScript->GetFName().ToString(),
 								Cast<UNiagaraScriptSource>(FunctionScript->GetSource())->NodeGraph,
 								Cast<UNiagaraScriptSource>(DupeScript->GetSource())->NodeGraph);*/
 
@@ -365,6 +366,65 @@ void FNiagaraCompileRequestData::DeepCopyGraphs(UNiagaraScriptSource* ScriptSour
 }
 
 
+void FNiagaraCompileRequestData::AddRapidIterationParameters(const FNiagaraParameterStore& InParamStore, FCompileConstantResolver InResolver)
+{
+	TArray<FNiagaraVariable> StoreParams;
+	InParamStore.GetParameters(StoreParams);
+
+	for (int32 i = 0; i < StoreParams.Num(); i++)
+	{
+		// Only support POD data...
+		if (StoreParams[i].IsDataInterface() || StoreParams[i].IsUObject())
+		{
+			continue;
+		}
+
+		if (InResolver.ResolveConstant(StoreParams[i]))
+		{
+			continue;
+		}
+
+		// Check to see if we already have this RI var...
+		int32 OurFoundIdx = INDEX_NONE;
+		for (int32 OurIdx = 0; OurIdx < RapidIterationParams.Num(); OurIdx++)
+		{
+			if (RapidIterationParams[OurIdx].GetType() == StoreParams[i].GetType() && RapidIterationParams[OurIdx].GetName() == StoreParams[i].GetName())
+			{
+				OurFoundIdx = OurIdx;
+				break;
+			}
+		}
+
+		// If we don't already have it, add it with the up-to-date value.
+		if (OurFoundIdx == INDEX_NONE)
+		{
+			// In parameter stores, the data isn't always up-to-date in the variable, so make sure to get the most up-to-date data before passing in.
+			const int32* Index = InParamStore.FindParameterOffset(StoreParams[i]);
+			if (Index != nullptr)
+			{
+				StoreParams[i].SetData(InParamStore.GetParameterData(*Index)); // This will memcopy the data in.
+				//UE_LOG(LogNiagaraEditor, Display, TEXT("Added Rapid iteration param: %s"), *StoreParams[i].ToString());
+				RapidIterationParams.Add(StoreParams[i]);
+			}
+		}
+		else
+		{
+			FNiagaraVariable ExistingVar = RapidIterationParams[OurFoundIdx];
+
+			const int32* Index = InParamStore.FindParameterOffset(StoreParams[i]);
+			if (Index != nullptr)
+			{
+				StoreParams[i].SetData(InParamStore.GetParameterData(*Index)); // This will memcopy the data in.
+
+				if (StoreParams[i] != ExistingVar)
+				{
+					UE_LOG(LogNiagaraEditor, Display, TEXT("Mismatch in values for Rapid iteration param: %s vs %s"), *StoreParams[i].ToString(), *ExistingVar.ToString());
+				}
+			}
+		}
+	}
+}
+
 void FNiagaraCompileRequestData::FinishPrecompile(UNiagaraScriptSource* ScriptSource, const TArray<FNiagaraVariable>& EncounterableVariables, ENiagaraScriptUsage InUsage, FCompileConstantResolver ConstantResolver)
 {
 	{
@@ -404,6 +464,7 @@ void FNiagaraCompileRequestData::FinishPrecompile(UNiagaraScriptSource* ScriptSo
 			Builder.EndTranslation(TranslationName);
 		}
 
+
 		// Generate CDO's for any referenced data interfaces...
 		for (int32 i = 0; i < PrecompiledHistories.Num(); i++)
 		{
@@ -475,6 +536,8 @@ TSharedPtr<FNiagaraCompileRequestDataBase, ESPMode::ThreadSafe> FNiagaraEditorMo
 	}
 	else if (System)
 	{
+		BasePtr->bUseRapidIterationParams = !System->bBakeOutRapidIteration;
+
 		// Store off the current variables in the exposed parameters list.
 		TArray<FNiagaraVariable> OriginalExposedParams;
 		System->GetExposedParameters().GetParameters(OriginalExposedParams);
@@ -485,6 +548,7 @@ TSharedPtr<FNiagaraCompileRequestDataBase, ESPMode::ThreadSafe> FNiagaraEditorMo
 		// Create an array of variables that we *did* encounter when traversing the graphs.
 		TArray<FNiagaraVariable> EncounteredExposedVars;
 		check(System->GetSystemSpawnScript()->GetSource() == System->GetSystemUpdateScript()->GetSource());
+		uint32 DetailLevelMask = 0xFFFFFFFF;
 
 		// First deep copy all the emitter graphs referenced by the system so that we can later hook up emitter handles in the system traversal.
 		BasePtr->EmitterData.Empty();
@@ -499,12 +563,15 @@ TSharedPtr<FNiagaraCompileRequestDataBase, ESPMode::ThreadSafe> FNiagaraEditorMo
 			}
 			EmitterPtr->EmitterUniqueName = Handle.GetInstance()->GetUniqueEmitterName();
 			EmitterPtr->SourceName = BasePtr->SourceName;
+			EmitterPtr->bUseRapidIterationParams = BasePtr->bUseRapidIterationParams || (!Handle.GetInstance()->bBakeOutRapidIteration);
 			BasePtr->EmitterData.Add(EmitterPtr);
 		}
 
 		// Now deep copy the system graphs, skipping traversal into any emitter references.
 		UNiagaraScriptSource* Source = Cast<UNiagaraScriptSource>(System->GetSystemSpawnScript()->GetSource());
 		BasePtr->DeepCopyGraphs(Source, ENiagaraScriptUsage::SystemSpawnScript, FCompileConstantResolver());
+		BasePtr->AddRapidIterationParameters(System->GetSystemSpawnScript()->RapidIterationParameters, FCompileConstantResolver());
+		BasePtr->AddRapidIterationParameters(System->GetSystemUpdateScript()->RapidIterationParameters, FCompileConstantResolver());
 		BasePtr->FinishPrecompile(Source, EncounterableVars, ENiagaraScriptUsage::SystemSpawnScript, EmptyResolver);
 
 		// Add the User and System variables that we did encounter to the list that emitters might also encounter.
@@ -518,6 +585,16 @@ TSharedPtr<FNiagaraCompileRequestDataBase, ESPMode::ThreadSafe> FNiagaraEditorMo
 			FCompileConstantResolver ConstantResolver(Handle.GetInstance());
 			if (Handle.GetIsEnabled()) // Don't pull in the emitter if it isn't going to be used.
 			{
+               TArray<UNiagaraScript*> EmitterScripts;
+				Handle.GetInstance()->GetScripts(EmitterScripts, false);
+
+				for (int32 ScriptIdx = 0; ScriptIdx < EmitterScripts.Num(); ScriptIdx++)
+				{
+					if (EmitterScripts[ScriptIdx])
+					{
+						BasePtr->EmitterData[i]->AddRapidIterationParameters(EmitterScripts[ScriptIdx]->RapidIterationParameters, ConstantResolver);
+					}
+				}
 				BasePtr->EmitterData[i]->FinishPrecompile(Cast<UNiagaraScriptSource>(Handle.GetInstance()->GraphSource), EncounterableVars, ENiagaraScriptUsage::EmitterSpawnScript, ConstantResolver);
 				BasePtr->MergeInEmitterPrecompiledData(BasePtr->EmitterData[i].Get());
 			}
@@ -539,6 +616,7 @@ TSharedPtr<FNiagaraVMExecutableData> FNiagaraEditorModule::CompileScript(const F
 
 	check(InCompileRequest != NULL);
 	const FNiagaraCompileRequestData* CompileRequest = (const FNiagaraCompileRequestData*)InCompileRequest;
+	TArray<FNiagaraVariable> CookedRapidIterationParams = CompileRequest->GetUseRapidIterationParams() ? TArray<FNiagaraVariable>() : CompileRequest->RapidIterationParams;
 
 	UE_LOG(LogNiagaraEditor, Log, TEXT("Compiling System %s ..................................................................."), *InCompileOptions.FullName);
 
@@ -560,6 +638,8 @@ TSharedPtr<FNiagaraVMExecutableData> FNiagaraEditorModule::CompileScript(const F
 		{
 			TranslateOptions.SimTarget = ENiagaraSimTarget::CPUSim;
 		}
+		TranslateOptions.OverrideModuleConstants = CookedRapidIterationParams;
+		TranslateOptions.bParameterRapidIteration = InCompileRequest->GetUseRapidIterationParams();
 
 		double TranslationStartTime = FPlatformTime::Seconds();
 		if (GbForceNiagaraTranslatorSingleThreaded > 0)
@@ -579,6 +659,12 @@ TSharedPtr<FNiagaraVMExecutableData> FNiagaraEditorModule::CompileScript(const F
 		}
 		
 		Results = Compiler.CompileScript(CompileRequest, InCompileOptions, &Translator.GetTranslateOutput(), Translator.GetTranslatedHLSL());
+
+		if (GbForceNiagaraVMBinaryDump != 0 && Results.Data.IsValid())
+		{
+			DumpHLSLText(Results.Data->LastAssemblyTranslation, InCompileOptions.FullName);
+		}
+
 		VMCompilationTime = Results.CompileTime;
 	}
 	
@@ -837,10 +923,13 @@ FNiagaraCompileResults FHlslNiagaraCompiler::CompileScript(const FNiagaraCompile
 			CompileResults.Data->NumTempRegisters = CompilationOutput.MaxTempRegistersUsed + 1;
 			CompileResults.Data->LastAssemblyTranslation = CompilationOutput.AssemblyAsString;
 			CompileResults.Data->LastOpCount = CompilationOutput.NumOps;
+			
 			//Build internal parameters
 			CompileResults.Data->InternalParameters.Empty();
+
 			for (int32 i = 0; i < CompilationOutput.InternalConstantOffsets.Num(); ++i)
 			{
+				FString ConstantName = FString::Printf(TEXT("InternalConstant_%d"),i);
 				EVectorVMBaseTypes Type = CompilationOutput.InternalConstantTypes[i];
 				int32 Offset = CompilationOutput.InternalConstantOffsets[i];
 				switch (Type)
@@ -848,19 +937,19 @@ FNiagaraCompileResults FHlslNiagaraCompiler::CompileScript(const FNiagaraCompile
 				case EVectorVMBaseTypes::Float:
 				{
 					float Val = *(float*)(CompilationOutput.InternalConstantData.GetData() + Offset);
-					CompileResults.Data->InternalParameters.SetOrAdd(FNiagaraVariable(FNiagaraTypeDefinition::GetFloatDef(), *LexToString(Val)))->SetValue(Val);
+					CompileResults.Data->InternalParameters.SetOrAdd(FNiagaraVariable(FNiagaraTypeDefinition::GetFloatDef(), *ConstantName))->SetValue(Val);
 				}
 				break;
 				case EVectorVMBaseTypes::Int:
 				{
 					int32 Val = *(int32*)(CompilationOutput.InternalConstantData.GetData() + Offset);
-					CompileResults.Data->InternalParameters.SetOrAdd(FNiagaraVariable(FNiagaraTypeDefinition::GetIntDef(), *LexToString(Val)))->SetValue(Val);
+					CompileResults.Data->InternalParameters.SetOrAdd(FNiagaraVariable(FNiagaraTypeDefinition::GetIntDef(), *ConstantName))->SetValue(Val);
 				}
 				break;
 				case EVectorVMBaseTypes::Bool:
 				{
 					int32 Val = *(int32*)(CompilationOutput.InternalConstantData.GetData() + Offset);
-					CompileResults.Data->InternalParameters.SetOrAdd(FNiagaraVariable(FNiagaraTypeDefinition::GetIntDef(), Val == 0 ? TEXT("FALSE") : TEXT("TRUE")))->SetValue(Val);
+					CompileResults.Data->InternalParameters.SetOrAdd(FNiagaraVariable(FNiagaraTypeDefinition::GetIntDef(), *ConstantName))->SetValue(Val);
 				}
 				break;
 				}
diff --git a/Engine/Plugins/FX/Niagara/Source/NiagaraEditor/Private/NiagaraHlslTranslator.cpp b/Engine/Plugins/FX/Niagara/Source/NiagaraEditor/Private/NiagaraHlslTranslator.cpp
--- a/Engine/Plugins/FX/Niagara/Source/NiagaraEditor/Private/NiagaraHlslTranslator.cpp
+++ b/Engine/Plugins/FX/Niagara/Source/NiagaraEditor/Private/NiagaraHlslTranslator.cpp
@@ -2813,7 +2813,9 @@ int32 FHlslNiagaraTranslator::GetRapidIterationParameter(const FNiagaraVariable&
 		int32 FoundIdx = TranslationOptions.OverrideModuleConstants.Find(RapidIterationConstantVar);
 		if (FoundIdx != INDEX_NONE)
 		{
-			int32 OutputChunkId = GetConstant(TranslationOptions.OverrideModuleConstants[FoundIdx]);
+			FString DebugConstantStr;
+			int32 OutputChunkId = GetConstant(TranslationOptions.OverrideModuleConstants[FoundIdx], &DebugConstantStr);
+			//UE_LOG(LogNiagaraEditor, Display, TEXT("Converted parameter %s to constant %s for script %s"), *RapidIterationConstantVar.GetName().ToString(), *DebugConstantStr, *CompileOptions.FullName);
 			return OutputChunkId;
 		}
 	}
@@ -2942,7 +2944,7 @@ int32 FHlslNiagaraTranslator::GetParameter(const FNiagaraVariable& Parameter)
 	return OutputChunkIdx;
 }
 
-int32 FHlslNiagaraTranslator::GetConstant(const FNiagaraVariable& Constant)
+int32 FHlslNiagaraTranslator::GetConstant(const FNiagaraVariable& Constant, FString* DebugOutputValue)
 {
 	if (Constant.IsDataInterface())
 	{
@@ -2961,6 +2963,10 @@ int32 FHlslNiagaraTranslator::GetConstant(const FNiagaraVariable& Constant)
 		ConstantStr = GenerateConstantString(Constant);
 	}
 
+	if (DebugOutputValue != nullptr)
+	{
+		*DebugOutputValue = ConstantStr;
+	}
 	if (ConstantStr.IsEmpty())
 	{
 		return INDEX_NONE;
@@ -3962,8 +3968,9 @@ void FHlslNiagaraTranslator::HandleParameterRead(int32 ParamMapHistoryIdx, const
 		}
 
 		// Check to see if this is the first time we've encountered this node and it is a viable candidate for rapid iteration
-		if (LastSetChunkIdx == INDEX_NONE && bIsCandidateForRapidIteration && TranslationOptions.bParameterRapidIteration)
+		if (LastSetChunkIdx == INDEX_NONE && bIsCandidateForRapidIteration)
 		{
+			FNiagaraVariable OriginalVar = Var;
 			bool bVarChanged = false;
 			if (!bWasEmitterAliased && ActiveHistoryForFunctionCalls.GetEmitterAlias() != nullptr)
 			{
@@ -3976,27 +3983,61 @@ void FHlslNiagaraTranslator::HandleParameterRead(int32 ParamMapHistoryIdx, const
 				bVarChanged = true;
 			}
 
-			// Now try to look up with the new name.. we may have already made this an external variable before..
-			if (bVarChanged)
+			if (TranslationOptions.bParameterRapidIteration)
 			{
-				VarIdx = ParamMapHistories[ParamMapHistoryIdx].FindVariableByName(Var.GetName());
-				if (VarIdx != INDEX_NONE && VarIdx < ParamMapSetVariablesToChunks[ParamMapHistoryIdx].Num())
+				// Now try to look up with the new name.. we may have already made this an external variable before..
+				if (bVarChanged)
 				{
-					LastSetChunkIdx = ParamMapSetVariablesToChunks[ParamMapHistoryIdx][VarIdx];
+					VarIdx = ParamMapHistories[ParamMapHistoryIdx].FindVariableByName(Var.GetName());
+					if (VarIdx != INDEX_NONE && VarIdx < ParamMapSetVariablesToChunks[ParamMapHistoryIdx].Num())
+					{
+						LastSetChunkIdx = ParamMapSetVariablesToChunks[ParamMapHistoryIdx][VarIdx];
+					}
 				}
-			}
 
-			// If it isn't found yet.. go ahead and make it into a constant variable..
-			if (LastSetChunkIdx == INDEX_NONE && ParameterMapRegisterExternalConstantNamespaceVariable(Var, ErrorNode, ParamMapHistoryIdx, OutputChunkId, InputPin))
+				// If it isn't found yet.. go ahead and make it into a constant variable..
+				if (LastSetChunkIdx == INDEX_NONE && ParameterMapRegisterExternalConstantNamespaceVariable(Var, ErrorNode, ParamMapHistoryIdx, OutputChunkId, InputPin))
+				{
+					LastSetChunkIdx = OutputChunkId;
+					if (VarIdx != INDEX_NONE && VarIdx < ParamMapSetVariablesToChunks[ParamMapHistoryIdx].Num())
+					{
+						// Record that we wrote to it.
+						ParamMapSetVariablesToChunks[ParamMapHistoryIdx][VarIdx] = LastSetChunkIdx;
+						ParamMapDefinedAttributesToNamespaceVars.FindOrAdd(Var.GetName()) = Var;
+					}
+					return;
+				}				
+			}
+			else
 			{
-				LastSetChunkIdx = OutputChunkId;
-				if (VarIdx != INDEX_NONE && VarIdx < ParamMapSetVariablesToChunks[ParamMapHistoryIdx].Num())
+				int32 FoundIdx = TranslationOptions.OverrideModuleConstants.Find(Var);
+				if (FoundIdx == INDEX_NONE)
 				{
-					// Record that we wrote to it.
-					ParamMapSetVariablesToChunks[ParamMapHistoryIdx][VarIdx] = LastSetChunkIdx;
-					ParamMapDefinedAttributesToNamespaceVars.FindOrAdd(Var.GetName()) = Var;
+					if (!bWasEmitterAliased && ActiveHistoryForFunctionCalls.GetEmitterAlias() != nullptr && CompileData != nullptr)
+					{
+						Var = FNiagaraParameterMapHistory::ConvertVariableToRapidIterationConstantName(OriginalVar, *CompileData->EmitterUniqueName, GetTargetUsage());
+						bVarChanged = true;
+						FoundIdx = TranslationOptions.OverrideModuleConstants.Find(Var);
+					}
 				}
-				return;
+				
+				if (FoundIdx != INDEX_NONE)
+				{
+					FString DebugConstantStr;
+					OutputChunkId = GetConstant(TranslationOptions.OverrideModuleConstants[FoundIdx], &DebugConstantStr);
+					UE_LOG(LogNiagaraEditor, Display, TEXT("Converted parameter %s to constant %s for script %s"), *Var.GetName().ToString(), *DebugConstantStr, *CompileOptions.FullName);
+					return;
+				}
+				else if (InputPin != nullptr && !InputPin->bDefaultValueIsIgnored) // Use the default from the input pin because this variable was previously never encountered.
+				{
+					FNiagaraVariable PinVar = Schema->PinToNiagaraVariable(InputPin, true);
+					FString DebugConstantStr;
+					OutputChunkId = GetConstant(PinVar, &DebugConstantStr);
+					UE_LOG(LogNiagaraEditor, Display, TEXT("Converted default value of parameter %s to constant %s for script %s. Likely added since this system was last compiled."), *Var.GetName().ToString(), *DebugConstantStr, *CompileOptions.FullName);
+					return;
+				}
+				
+				Error(FText::Format(LOCTEXT("InvalidRapidIterationReplacement", "Variable {0} is a rapid iteration param, but it wasn't found in the override list to bake out!"), FText::FromName(Var.GetName())), ErrorNode, nullptr);
 			}
 		}
 
@@ -4079,9 +4120,21 @@ void FHlslNiagaraTranslator::HandleParameterRead(int32 ParamMapHistoryIdx, const
 					int32 FoundIdx = TranslationOptions.OverrideModuleConstants.Find(RapidIterationConstantVar);
 					if (FoundIdx != INDEX_NONE)
 					{
-						OutputChunkId = GetConstant(TranslationOptions.OverrideModuleConstants[FoundIdx]);
+						FString DebugConstantStr;
+						OutputChunkId = GetConstant(TranslationOptions.OverrideModuleConstants[FoundIdx], &DebugConstantStr);
+						UE_LOG(LogNiagaraEditor, Display, TEXT("Converted parameter %s to constant %s for script %s"), *Var.GetName().ToString(), *DebugConstantStr, *CompileOptions.FullName);
+						return;
+					}
+					else if (InputPin != nullptr && !InputPin->bDefaultValueIsIgnored) // Use the default from the input pin because this variable was previously never encountered.
+					{
+						FNiagaraVariable PinVar = Schema->PinToNiagaraVariable(InputPin, true);
+						FString DebugConstantStr;
+						OutputChunkId = GetConstant(PinVar, &DebugConstantStr);
+						UE_LOG(LogNiagaraEditor, Display, TEXT("Converted default value of parameter %s to constant %s for script %s. Likely added since this system was last compiled."), *Var.GetName().ToString(), *DebugConstantStr, *CompileOptions.FullName);
 						return;
 					}
+
+					Error(FText::Format(LOCTEXT("InvalidRapidIterationReplacement", "Variable {0} is a rapid iteration param, but it wasn't found in the override list to bake out!"), FText::FromName(Var.GetName())), ErrorNode, nullptr);
 				}
 
 				CurrentDefaultPinTraversal.Push(InputPin);
@@ -5405,6 +5458,16 @@ FNiagaraTypeDefinition FHlslNiagaraTranslator::GetChildType(const FNiagaraTypeDe
 				{
 					return FNiagaraTypeDefinition::GetBoolDef();
 				}
+				else if (Property->IsA(UEnumProperty::StaticClass()))
+				{
+					const UEnumProperty* EnumProp = Cast<UEnumProperty>(Property);
+					return FNiagaraTypeDefinition(EnumProp->GetEnum());
+				}
+				else if (Property->IsA(UByteProperty::StaticClass()))
+				{
+					const UByteProperty* ByteProp = Cast<UByteProperty>(Property);
+					return FNiagaraTypeDefinition(ByteProp->GetIntPropertyEnum());
+				}
 				else if (const UStructProperty* StructProp = CastChecked<UStructProperty>(Property))
 				{
 					return FNiagaraTypeDefinition(StructProp->Struct);
@@ -6015,7 +6078,10 @@ FString FHlslNiagaraTranslator::GetPropertyHlslTypeName(const UProperty* Propert
 	}
 	else if (Property->IsA(UEnumProperty::StaticClass()))
 	{
-		const UEnumProperty* EnumProp = Cast<const UEnumProperty>(Property);
+		return "int";
+	}
+	else if (Property->IsA(UByteProperty::StaticClass()))
+	{
 		return "int";
 	}
 	else if (Property->IsA(UBoolProperty::StaticClass()))
diff --git a/Engine/Plugins/FX/Niagara/Source/NiagaraEditor/Private/NiagaraHlslTranslator.h b/Engine/Plugins/FX/Niagara/Source/NiagaraEditor/Private/NiagaraHlslTranslator.h
--- a/Engine/Plugins/FX/Niagara/Source/NiagaraEditor/Private/NiagaraHlslTranslator.h
+++ b/Engine/Plugins/FX/Niagara/Source/NiagaraEditor/Private/NiagaraHlslTranslator.h
@@ -82,6 +82,11 @@ struct FNiagaraTranslateResults
 class FNiagaraCompileRequestData : public FNiagaraCompileRequestDataBase
 {
 public:
+	FNiagaraCompileRequestData() : bUseRapidIterationParams(true), DetailLevelMask(FNiagaraCompileRequestDataBase::CookForAllDetailLevelMask)
+	{
+
+	}
+
 	virtual ~FNiagaraCompileRequestData() {}
 	virtual bool GatherPreCompiledVariables(const FString& InNamespaceFilter, TArray<FNiagaraVariable>& OutVars) override;
 	virtual void GetReferencedObjects(TArray<UObject*>& Objects) override;
@@ -103,7 +108,11 @@ public:
 	virtual TSharedPtr<FNiagaraCompileRequestDataBase, ESPMode::ThreadSafe> GetDependentRequest(int32 Index) override {
 		return EmitterData[Index];
 	}
-
+	void AddRapidIterationParameters(const FNiagaraParameterStore& InParamStore, FCompileConstantResolver InResolver);
+	virtual uint32 GetDetailLevelMask(void) const override {
+		return DetailLevelMask;
+	};
+	virtual bool GetUseRapidIterationParams() const override { return bUseRapidIterationParams; }
 
 	// If this is being held onto for any length of time, make sure to hold onto it in a gc-aware object. Right now in this information-passing struct,
 	// we could have a leaked garbage collected pointer if not held onto by someone capable of registering a reference.
@@ -116,10 +125,14 @@ public:
 	TArray<TSharedPtr<FNiagaraCompileRequestData, ESPMode::ThreadSafe>> EmitterData;
 	UNiagaraScriptSource* Source;
 	FString SourceName;
+	bool bUseRapidIterationParams = true;
+	uint32 DetailLevelMask = 0xFFFFFFFF;
 
 	UEnum* ENiagaraScriptCompileStatusEnum;
 	UEnum* ENiagaraScriptUsageEnum;
 
+	TArray<FNiagaraVariable> RapidIterationParams;
+
 	struct FunctionData
 	{
 		UNiagaraScript* ClonedScript;
@@ -250,6 +263,10 @@ public:
 
 	}
 
+	FHlslNiagaraTranslatorOptions(const FHlslNiagaraTranslatorOptions& InOptions) : SimTarget(InOptions.SimTarget), InstanceParameterNamespaces(InOptions.InstanceParameterNamespaces), bParameterRapidIteration(InOptions.bParameterRapidIteration), OverrideModuleConstants(InOptions.OverrideModuleConstants)
+	{
+	}
+
 	ENiagaraSimTarget SimTarget;
 
 	/** Any parameters in these namespaces will be pulled from an "InstanceParameters" dataset rather than from the uniform table. */
@@ -260,6 +277,8 @@ public:
 
 	/** Whether or not to override top-level module variables with values from the constant override table. This is only used for variables that were candidates for rapid iteration.*/
 	TArray<FNiagaraVariable> OverrideModuleConstants;
+
+
 };
 
 class NIAGARAEDITOR_API FHlslNiagaraTranslationStage
@@ -420,7 +439,7 @@ public:
 
 	virtual int32 GetAttribute(const FNiagaraVariable& Attribute);
 
-	virtual int32 GetConstant(const FNiagaraVariable& Constant);
+	virtual int32 GetConstant(const FNiagaraVariable& Constant, FString* DebugOutputValue = nullptr);
 	
 	virtual void ReadDataSet(const FNiagaraDataSetID DataSet, const TArray<FNiagaraVariable>& Variable, ENiagaraDataSetAccessMode AccessMode, int32 InputChunk, TArray<int32>& Outputs);
 	virtual void WriteDataSet(const FNiagaraDataSetID DataSet, const TArray<FNiagaraVariable>& Variable, ENiagaraDataSetAccessMode AccessMode, const TArray<int32>& Inputs, TArray<int32>& Outputs);
diff --git a/Engine/Plugins/FX/Niagara/Source/NiagaraEditor/Private/Toolkits/NiagaraSystemToolkit.cpp b/Engine/Plugins/FX/Niagara/Source/NiagaraEditor/Private/Toolkits/NiagaraSystemToolkit.cpp
--- a/Engine/Plugins/FX/Niagara/Source/NiagaraEditor/Private/Toolkits/NiagaraSystemToolkit.cpp
+++ b/Engine/Plugins/FX/Niagara/Source/NiagaraEditor/Private/Toolkits/NiagaraSystemToolkit.cpp
@@ -51,6 +51,7 @@
 #include "Misc/FileHelper.h"
 #include "NiagaraMessageLogViewModel.h"
 #include "ViewModels/NiagaraOverviewGraphViewModel.h"
+#include "NiagaraScriptSourceBase.h"
 
 #define LOCTEXT_NAMESPACE "NiagaraSystemEditor"
 
diff --git a/Engine/Plugins/FX/Niagara/Source/NiagaraShader/Private/NiagaraShader.cpp b/Engine/Plugins/FX/Niagara/Source/NiagaraShader/Private/NiagaraShader.cpp
--- a/Engine/Plugins/FX/Niagara/Source/NiagaraShader/Private/NiagaraShader.cpp
+++ b/Engine/Plugins/FX/Niagara/Source/NiagaraShader/Private/NiagaraShader.cpp
@@ -135,11 +135,18 @@ void FNiagaraShaderMapId::Serialize(FArchive& Ar)
 		Ar << BaseCompileHash;
 		Ar << ReferencedCompileHashes;
 	}
+
 	if (NiagaraVer >= FNiagaraCustomVersion::AddAdditionalDefinesProperty)
 	{
 		Ar << AdditionalDefines;
 	}
 	
+	if (NiagaraVer >= FNiagaraCustomVersion::AddRIAndDetailLevel)
+	{
+		Ar << DetailLevelMask;
+		Ar << bUsesRapidIterationParams;
+	}
+	
 	//ParameterSet.Serialize(Ar);		// NIAGARATODO: at some point we'll need stuff for static switches here
 }
 
@@ -180,18 +187,24 @@ bool FNiagaraShaderMapId::operator==(const FNiagaraShaderMapId& ReferenceSet) co
 	if (BaseScriptID != ReferenceSet.BaseScriptID 
 		|| BaseCompileHash != ReferenceSet.BaseCompileHash
 		|| FeatureLevel != ReferenceSet.FeatureLevel
-		|| CompilerVersionID != ReferenceSet.CompilerVersionID)
+		|| CompilerVersionID != ReferenceSet.CompilerVersionID 
+		|| DetailLevelMask != ReferenceSet.DetailLevelMask 
+		|| bUsesRapidIterationParams != ReferenceSet.bUsesRapidIterationParams)
 	{
 		return false;
 	}
+	
+	
 	if (AdditionalDefines.Num() != ReferenceSet.AdditionalDefines.Num())
 	{
 		return false;
 	}
 
-	for (int32 i = 0; i < AdditionalDefines.Num(); i++)
+	for (int32 Idx = 0; Idx < ReferenceSet.AdditionalDefines.Num(); Idx++)
 	{
-		if (AdditionalDefines[i] != ReferenceSet.AdditionalDefines[i])
+		const FString& ReferenceStr = ReferenceSet.AdditionalDefines[Idx];
+
+		if (AdditionalDefines[Idx] != ReferenceStr)
 		{
 			return false;
 		}
@@ -245,6 +258,24 @@ void FNiagaraShaderMapId::AppendKeyString(FString& KeyString) const
 	KeyString += CompilerVersionID.ToString();
 	KeyString += TEXT("_");
 
+	if (DetailLevelMask != 0xFFFFFFFF)
+	{
+		KeyString += FString::Printf(TEXT("DL_%d"), DetailLevelMask);
+	}
+	else
+	{
+		KeyString += TEXT("ALLDL_");
+	}
+
+	if (bUsesRapidIterationParams)
+	{
+		KeyString += TEXT("USESRI_");
+	}
+	else
+	{
+		KeyString += TEXT("NORI_");
+	}
+	
 	// Add additional defines
 	for (int32 DefinesIndex = 0; DefinesIndex < AdditionalDefines.Num(); DefinesIndex++)
 	{
diff --git a/Engine/Plugins/FX/Niagara/Source/NiagaraShader/Private/NiagaraShared.cpp b/Engine/Plugins/FX/Niagara/Source/NiagaraShader/Private/NiagaraShared.cpp
--- a/Engine/Plugins/FX/Niagara/Source/NiagaraShader/Private/NiagaraShared.cpp
+++ b/Engine/Plugins/FX/Niagara/Source/NiagaraShader/Private/NiagaraShared.cpp
@@ -90,6 +90,8 @@ bool FNiagaraShaderScript::IsSame(const FNiagaraShaderMapId& InId) const
 		InId.FeatureLevel == FeatureLevel &&
 		InId.BaseScriptID == BaseScriptId &&
 		InId.AdditionalDefines == AdditionalDefines &&
+		InId.DetailLevelMask == DetailLevelMask &&
+		InId.bUsesRapidIterationParams == bUsesRapidIterationParams &&
 		InId.BaseCompileHash == BaseCompileHash &&
 		InId.ReferencedCompileHashes == ReferencedCompileHashes &&
 		InId.CompilerVersionID == CompilerVersionId;
@@ -124,6 +126,8 @@ NIAGARASHADER_API void FNiagaraShaderScript::GetShaderMapId(EShaderPlatform Plat
 		OutId.FeatureLevel = GetFeatureLevel();
 		OutId.BaseScriptID = BaseScriptId;
 		OutId.AdditionalDefines = AdditionalDefines;
+		OutId.DetailLevelMask = DetailLevelMask;
+		OutId.bUsesRapidIterationParams = bUsesRapidIterationParams;		
 		OutId.BaseCompileHash = BaseCompileHash;
 		OutId.ReferencedCompileHashes = ReferencedCompileHashes;
 		OutId.CompilerVersionID = FNiagaraCustomVersion::LatestScriptCompileVersion;
@@ -225,7 +229,8 @@ void FNiagaraShaderScript::SerializeShaderMap(FArchive& Ar)
 }
 
 void FNiagaraShaderScript::SetScript(UNiagaraScript *InScript, ERHIFeatureLevel::Type InFeatureLevel, const FGuid& InCompilerVersionID, const FGuid& InBaseScriptID, const TArray<FString>& InAdditionalDefines,
-	const FNiagaraCompileHash& InBaseCompileHash, const TArray<FNiagaraCompileHash>& InReferencedCompileHashes, FString InFriendlyName)
+		const FNiagaraCompileHash& InBaseCompileHash, const TArray<FNiagaraCompileHash>& InReferencedCompileHashes, 
+		bool bInUsesRapidIterationParams, uint32 InDetailLevelMask, FString InFriendlyName)
 {
 	checkf(InBaseScriptID.IsValid(), TEXT("Invalid base script id.  Script caching will fail."));
 	checkf(InBaseCompileHash.IsValid(), TEXT("Invalid base compile hash.  Script caching will fail."))
@@ -233,6 +238,8 @@ void FNiagaraShaderScript::SetScript(UNiagaraScript *InScript, ERHIFeatureLevel:
 	CompilerVersionId = InCompilerVersionID;
 	BaseScriptId = InBaseScriptID;
 	AdditionalDefines = InAdditionalDefines;
+	bUsesRapidIterationParams = bInUsesRapidIterationParams;
+	DetailLevelMask = InDetailLevelMask;
 	BaseCompileHash = InBaseCompileHash;
 	ReferencedCompileHashes = InReferencedCompileHashes;
 	FriendlyName = InFriendlyName;
diff --git a/Engine/Plugins/FX/Niagara/Source/NiagaraShader/Public/NiagaraShared.h b/Engine/Plugins/FX/Niagara/Source/NiagaraShader/Public/NiagaraShared.h
--- a/Engine/Plugins/FX/Niagara/Source/NiagaraShader/Public/NiagaraShared.h
+++ b/Engine/Plugins/FX/Niagara/Source/NiagaraShader/Public/NiagaraShared.h
@@ -172,6 +172,13 @@ public:
 	/** The compile hashes of the top level scripts the script is dependent on. */
 	TArray<FNiagaraCompileHash> ReferencedCompileHashes;
 
+	
+	/** Bitfield of supported detail levels in this compile.*/
+	uint32 DetailLevelMask = 0xFFFFFFFF;
+
+	/** Whether or not we need to bake Rapid Iteration params. True to keep params, false to bake.*/
+	bool bUsesRapidIterationParams = true;
+	
 	FNiagaraShaderMapId()
 		: CompilerVersionID()
 		, FeatureLevel(GMaxRHIFeatureLevel)
@@ -667,7 +674,8 @@ public:
 
 
 	NIAGARASHADER_API void SetScript(UNiagaraScript *InScript, ERHIFeatureLevel::Type InFeatureLevel, const FGuid& InCompilerVersion, const FGuid& InBaseScriptID, const TArray<FString>& InAdditionalDefines,
-		const FNiagaraCompileHash& InBaseCompileHash, const TArray<FNiagaraCompileHash>& InReferencedCompileHashes, FString InFriendlyName);
+		const FNiagaraCompileHash& InBaseCompileHash, const TArray<FNiagaraCompileHash>& InReferencedCompileHashes, 
+		bool bInUsesRapidIterationParams, uint32 InDetailLevelMask, FString InFriendlyName);
 
 	UNiagaraScript *GetBaseVMScript()
 	{
@@ -744,6 +752,12 @@ private:
 
 	/** Configuration options */
 	TArray<FString> AdditionalDefines;
+	
+	/** Bitfield of supported detail levels in this compile.*/
+	uint32 DetailLevelMask = 0xFFFFFFFF;
+
+	/** Whether or not we need to bake Rapid Iteration params. True to keep params, false to bake.*/
+	bool bUsesRapidIterationParams = true;
 
 	/** Compile hash for the base script. */
 	FNiagaraCompileHash BaseCompileHash;
