- Fix a crash from using aligned memory writes for types like FQuat / FVector4 when adding new parameter data as they may not be aligned in the parameter store

#rnx

diff --git a/Engine/Plugins/FX/Niagara/Source/Niagara/Public/NiagaraParameterStore.h b/Engine/Plugins/FX/Niagara/Source/Niagara/Public/NiagaraParameterStore.h
--- a/Engine/Plugins/FX/Niagara/Source/Niagara/Public/NiagaraParameterStore.h
+++ b/Engine/Plugins/FX/Niagara/Source/Niagara/Public/NiagaraParameterStore.h
@@ -434,7 +434,10 @@ public:
 				bool bTriggerRebind = false;
 				AddParameter(Param, bInitInterfaces, bTriggerRebind, &Offset);
 				check(Offset != INDEX_NONE);
-				*(T*)(GetParameterData_Internal(Offset)) = InValue;
+				//Until we solve our alignment issues, temporarily just doing a memcpy here.
+				//*(T*)(GetParameterData_Internal(Offset)) = InValue;
+				T* ParamData = reinterpret_cast<T*>(GetParameterData_Internal(Offset));
+				FMemory::Memcpy(ParamData, &InValue, sizeof(T));
 				OnLayoutChange();
 				return true;
 			}
