- Copy over default values to fix issues where they are not all set by the user

#rnx

diff --git a/Engine/Plugins/FX/Niagara/Source/Niagara/Private/NiagaraComponent.cpp b/Engine/Plugins/FX/Niagara/Source/Niagara/Private/NiagaraComponent.cpp
--- a/Engine/Plugins/FX/Niagara/Source/Niagara/Private/NiagaraComponent.cpp
+++ b/Engine/Plugins/FX/Niagara/Source/Niagara/Private/NiagaraComponent.cpp
@@ -1719,7 +1719,11 @@ void UNiagaraComponent::SetAsset(UNiagaraSystem* InAsset)
 
 			for (FNiagaraVariable ExistingVar : ExistingVars)
 			{
-				if (!SourceVars.Contains(ExistingVar))
+				if (SourceVars.Contains(ExistingVar))
+				{
+					Asset->GetExposedParameters().CopyParameterData(OverrideParameters, ExistingVar);
+				}
+				else
 				{
 					OverrideParameters.RemoveParameter(ExistingVar);
 				}
