- Make Parent & ParentAtLastMerge editor only data to avoid pulling emitter hierarchy into memory on cooked builds

#rnx

diff --git a/Engine/Plugins/FX/Niagara/Source/Niagara/Classes/NiagaraEmitter.h b/Engine/Plugins/FX/Niagara/Source/Niagara/Classes/NiagaraEmitter.h
--- a/Engine/Plugins/FX/Niagara/Source/Niagara/Classes/NiagaraEmitter.h
+++ b/Engine/Plugins/FX/Niagara/Source/Niagara/Classes/NiagaraEmitter.h
@@ -375,9 +375,11 @@ public:
 
 	TStatId GetStatID(bool bGameThread, bool bConcurrent)const;
 
+#if WITH_EDITORONLY_DATA
 	NIAGARA_API UNiagaraEmitter* GetParent() const;
 
 	NIAGARA_API void RemoveParent();
+#endif
 
 	void InitFastPathAttributeNames();
 
@@ -443,11 +445,13 @@ private:
 	UPROPERTY()
 	TArray<FName> SharedEventGeneratorIds;
 
+#if WITH_EDITORONLY_DATA
 	UPROPERTY()
 	UNiagaraEmitter* Parent;
 
 	UPROPERTY()
 	UNiagaraEmitter* ParentAtLastMerge;
+#endif
 
 #if WITH_EDITOR
 	FOnPropertiesChanged OnPropertiesChangedDelegate;
diff --git a/Engine/Plugins/FX/Niagara/Source/Niagara/Private/NiagaraEmitter.cpp b/Engine/Plugins/FX/Niagara/Source/Niagara/Private/NiagaraEmitter.cpp
--- a/Engine/Plugins/FX/Niagara/Source/Niagara/Private/NiagaraEmitter.cpp
+++ b/Engine/Plugins/FX/Niagara/Source/Niagara/Private/NiagaraEmitter.cpp
@@ -1169,6 +1169,7 @@ void UNiagaraEmitter::GenerateStatID()
 #endif
 }
 
+#if WITH_EDITORONLY_DATA
 UNiagaraEmitter* UNiagaraEmitter::GetParent() const
 {
 	return Parent;
@@ -1179,6 +1180,7 @@ void UNiagaraEmitter::RemoveParent()
 	Parent = nullptr;
 	ParentAtLastMerge = nullptr;
 }
+#endif
 
 void UNiagaraEmitter::InitFastPathAttributeNames()
 {
