- Fix for Niagara Component activating automatically when changing sg.EffectsQuality

#rnx

diff --git a/Engine/Plugins/FX/Niagara/Source/Niagara/Private/NiagaraModule.cpp b/Engine/Plugins/FX/Niagara/Source/Niagara/Private/NiagaraModule.cpp
--- a/Engine/Plugins/FX/Niagara/Source/Niagara/Private/NiagaraModule.cpp
+++ b/Engine/Plugins/FX/Niagara/Source/Niagara/Private/NiagaraModule.cpp
@@ -465,9 +465,21 @@ void INiagaraModule::OnChangeDetailLevel(class IConsoleVariable* CVar)
 	if (EngineDetailLevel != NewDetailLevel)
 	{
 		EngineDetailLevel = NewDetailLevel;
-		//If the detail level has changed we have to reset all systems.
-		FNiagaraSystemUpdateContext UpdateContext;
-		UpdateContext.AddAll(true);
+
+		// If the detail level has changed we have to reset all systems,
+		// and only activate ones that were previously active
+		for (TObjectIterator<UNiagaraComponent> It; It; ++It)
+		{
+			UNiagaraComponent* Comp = *It;
+			check(Comp);
+
+			const bool bWasActive = Comp->bIsActive;
+			Comp->DestroyInstance();
+			if ( bWasActive )
+			{
+				Comp->Activate(true);
+			}
+		}
 	}
 #endif
 }
