Improvements to stat niagaraemitters

diff --git a/Engine/Plugins/FX/Niagara/Source/Niagara/Classes/NiagaraEmitter.h b/Engine/Plugins/FX/Niagara/Source/Niagara/Classes/NiagaraEmitter.h
--- a/Engine/Plugins/FX/Niagara/Source/Niagara/Classes/NiagaraEmitter.h
+++ b/Engine/Plugins/FX/Niagara/Source/Niagara/Classes/NiagaraEmitter.h
@@ -457,12 +457,12 @@ private:
 	FOnPropertiesChanged OnPropertiesChangedDelegate;
 #endif
 
-	void GenerateStatID();
+	void GenerateStatID()const;
 #if STATS
-	TStatId StatID_GT;
-	TStatId StatID_GT_CNC;
-	TStatId StatID_RT;
-	TStatId StatID_RT_CNC;
+	mutable TStatId StatID_GT;
+	mutable TStatId StatID_GT_CNC;
+	mutable TStatId StatID_RT;
+	mutable TStatId StatID_RT_CNC;
 #endif
 };
 
diff --git a/Engine/Plugins/FX/Niagara/Source/Niagara/Classes/NiagaraSystem.h b/Engine/Plugins/FX/Niagara/Source/Niagara/Classes/NiagaraSystem.h
--- a/Engine/Plugins/FX/Niagara/Source/Niagara/Classes/NiagaraSystem.h
+++ b/Engine/Plugins/FX/Niagara/Source/Niagara/Classes/NiagaraSystem.h
@@ -374,11 +374,11 @@ protected:
 	UPROPERTY()
 	TArray<FName> UserDINamesReadInSystemScripts;
 
-	void GenerateStatID();
+	void GenerateStatID()const;
 #if STATS
-	TStatId StatID_GT;
-	TStatId StatID_GT_CNC;
-	TStatId StatID_RT;
-	TStatId StatID_RT_CNC;
+	mutable TStatId StatID_GT;
+	mutable TStatId StatID_GT_CNC;
+	mutable TStatId StatID_RT;
+	mutable TStatId StatID_RT_CNC;
 #endif
 };
diff --git a/Engine/Plugins/FX/Niagara/Source/Niagara/Private/NiagaraEmitter.cpp b/Engine/Plugins/FX/Niagara/Source/Niagara/Private/NiagaraEmitter.cpp
--- a/Engine/Plugins/FX/Niagara/Source/Niagara/Private/NiagaraEmitter.cpp
+++ b/Engine/Plugins/FX/Niagara/Source/Niagara/Private/NiagaraEmitter.cpp
@@ -138,8 +138,6 @@ void UNiagaraEmitter::PostInitProperties()
 
 	}
 	UniqueEmitterName = TEXT("Emitter");
-
-	GenerateStatID();
 }
 
 #if WITH_EDITORONLY_DATA
@@ -1134,38 +1132,45 @@ void UNiagaraEmitter::GraphSourceChanged()
 TStatId UNiagaraEmitter::GetStatID(bool bGameThread, bool bConcurrent)const
 {
 #if STATS
+	if (!StatID_GT.IsValidStat())
+	{
+		GenerateStatID();
+	}
+
 	if (bGameThread)
 	{
 		if (bConcurrent)
 		{
-			return StatID_GT;
+			return StatID_GT_CNC;
 		}
 		else
 		{
-			return StatID_GT_CNC;
+			return StatID_GT;
 		}
 	}
 	else
 	{
 		if (bConcurrent)
 		{
-			return StatID_RT;
+			return StatID_RT_CNC;
 		}
 		else
 		{
-			return StatID_RT_CNC;
+			return StatID_RT;
 		}
 	}
 #endif
 	return TStatId();
 }
-void UNiagaraEmitter::GenerateStatID()
+void UNiagaraEmitter::GenerateStatID()const
 {
 #if STATS
-	StatID_GT = FDynamicStats::CreateStatId<FStatGroup_STATGROUP_NiagaraEmitters>(GetName() + TEXT("[GT]"));
-	StatID_GT_CNC = FDynamicStats::CreateStatId<FStatGroup_STATGROUP_NiagaraEmitters>(GetName() + TEXT("[GT_CNC]"));
-	StatID_RT = FDynamicStats::CreateStatId<FStatGroup_STATGROUP_NiagaraEmitters>(GetName() + TEXT("[RT]"));
-	StatID_RT_CNC = FDynamicStats::CreateStatId<FStatGroup_STATGROUP_NiagaraEmitters>(GetName() + TEXT("[RT_CNC]"));
+	FString Name = GetOuter() ? GetOuter()->GetFName().ToString() : TEXT("");
+	Name += TEXT("/") + UniqueEmitterName;
+	StatID_GT = FDynamicStats::CreateStatId<FStatGroup_STATGROUP_NiagaraEmitters>(Name + TEXT("[GT]"));
+	StatID_GT_CNC = FDynamicStats::CreateStatId<FStatGroup_STATGROUP_NiagaraEmitters>(Name + TEXT("[GT_CNC]"));
+	StatID_RT = FDynamicStats::CreateStatId<FStatGroup_STATGROUP_NiagaraEmitters>(Name + TEXT("[RT]"));
+	StatID_RT_CNC = FDynamicStats::CreateStatId<FStatGroup_STATGROUP_NiagaraEmitters>(Name + TEXT("[RT_CNC]"));
 #endif
 }
 
diff --git a/Engine/Plugins/FX/Niagara/Source/Niagara/Private/NiagaraSystem.cpp b/Engine/Plugins/FX/Niagara/Source/Niagara/Private/NiagaraSystem.cpp
--- a/Engine/Plugins/FX/Niagara/Source/Niagara/Private/NiagaraSystem.cpp
+++ b/Engine/Plugins/FX/Niagara/Source/Niagara/Private/NiagaraSystem.cpp
@@ -110,8 +110,6 @@ void UNiagaraSystem::PostInitProperties()
 		EditorData = NiagaraModule.GetEditorOnlyDataUtilities().CreateDefaultEditorData(this);
 #endif
 	}
-
-	GenerateStatID();
 }
 
 bool UNiagaraSystem::IsLooping() const
@@ -1178,33 +1176,38 @@ void UNiagaraSystem::InitSystemCompiledData()
 TStatId UNiagaraSystem::GetStatID(bool bGameThread, bool bConcurrent)const
 {
 #if STATS
+	if (!StatID_GT.IsValidStat())
+	{
+		GenerateStatID();
+	}
+
 	if (bGameThread)
 	{
 		if (bConcurrent)
 		{
-			return StatID_GT;
+			return StatID_GT_CNC;
 		}
 		else
 		{
-			return StatID_GT_CNC;
+			return StatID_GT;
 		}
 	}
 	else
 	{
 		if(bConcurrent)
 		{
-			return StatID_RT;
+			return StatID_RT_CNC;
 		}
 		else
 		{
-			return StatID_RT_CNC;
+			return StatID_RT;
 		}
 	}
 #endif
 	return TStatId();
 }
 
-void UNiagaraSystem::GenerateStatID()
+void UNiagaraSystem::GenerateStatID()const
 {
 #if STATS
 	StatID_GT = FDynamicStats::CreateStatId<FStatGroup_STATGROUP_NiagaraSystems>(GetPathName() + TEXT("[GT]"));
diff --git a/Engine/Plugins/FX/Niagara/Source/Niagara/Private/NiagaraSystemInstance.cpp b/Engine/Plugins/FX/Niagara/Source/Niagara/Private/NiagaraSystemInstance.cpp
--- a/Engine/Plugins/FX/Niagara/Source/Niagara/Private/NiagaraSystemInstance.cpp
+++ b/Engine/Plugins/FX/Niagara/Source/Niagara/Private/NiagaraSystemInstance.cpp
@@ -1661,7 +1661,7 @@ void FNiagaraSystemInstance::Tick_GameThread(float DeltaSeconds)
 	LLM_SCOPE(ELLMTag::Niagara);
 
 	UNiagaraSystem* System = GetSystem();
-	FScopeCycleCounter SystemStat(System->GetStatID(true, true));
+	FScopeCycleCounter SystemStat(System->GetStatID(true, false));
 
 	WaitForAsyncTick(true);
 
