Niagara - Update the instance parameters and tick it in the concurrent system simulation tick instead of the game thread tick.

diff --git a/Engine/Plugins/FX/Niagara/Source/Niagara/Private/NiagaraSystemInstance.cpp b/Engine/Plugins/FX/Niagara/Source/Niagara/Private/NiagaraSystemInstance.cpp
--- a/Engine/Plugins/FX/Niagara/Source/Niagara/Private/NiagaraSystemInstance.cpp
+++ b/Engine/Plugins/FX/Niagara/Source/Niagara/Private/NiagaraSystemInstance.cpp
@@ -930,7 +930,8 @@ void FNiagaraSystemInstance::ReInitInternal()
 	// rebind now after all parameters have been added
 	InstanceParameters.Rebind();
 
-	TickInstanceParameters(0.01f);
+	TickInstanceParameters_GameThread(0.01f);
+	TickInstanceParameters_Concurrent();
 
 	//Invalidate the component render state so we recreate the scene proxy and the renderers.
 	Component->MarkRenderStateDirty();
@@ -1374,28 +1375,76 @@ ETickingGroup FNiagaraSystemInstance::CalculateTickGroup()
 	return NewTickGroup;
 }
 
-void FNiagaraSystemInstance::TickInstanceParameters(float DeltaSeconds)
+void FNiagaraSystemInstance::TickInstanceParameters_GameThread(float DeltaSeconds)
 {
 	if (!Component)
 	{
 		return;
 	}
 
-	//TODO: Create helper binding objects to avoid the search in set parameter value.
-	//Set System params.
-	FTransform ComponentTrans = Component->GetComponentTransform();
-	FVector CurrPos = ComponentTrans.GetLocation();
-	FVector OldPos = FMath::IsNearlyZero(Age) ? CurrPos : OwnerPositionParam.GetValue(); // The first frame the value in OwnerPositionParam is uninitialized memory, we need to make sure that we don't use it.
+	GatheredInstanceParameters.DeltaSeconds = DeltaSeconds;
+
+	GatheredInstanceParameters.ComponentTrans = Component->GetComponentTransform();
+	GatheredInstanceParameters.OldPos = OwnerPositionParam.GetValue();
+
+	UWorld* World = Component->GetWorld();
+	if (World != NULL)
+	{
+		GatheredInstanceParameters.LODDistance = GetLODDistance();
+		GatheredInstanceParameters.TimeSeconds = World->TimeSeconds;
+		GatheredInstanceParameters.RealTimeSeconds = World->RealTimeSeconds;
+	}
+	else
+	{
+		GatheredInstanceParameters.LODDistance = 0;
+		GatheredInstanceParameters.TimeSeconds = Age;
+		GatheredInstanceParameters.RealTimeSeconds = Age;
+	}
+
+	GatheredInstanceParameters.Age = Age;
+	GatheredInstanceParameters.TickCount = TickCount;
+
+	GatheredInstanceParameters.NumAlive = 0;
+	for (int32 i = 0; i < Emitters.Num(); ++i)
+	{
+		const TSharedRef<FNiagaraEmitterInstance, ESPMode::ThreadSafe>& Emitter = Emitters[i];
+		if (Emitter->GetExecutionState() != ENiagaraExecutionState::Disabled)
+		{
+			GatheredInstanceParameters.EmitterNumParticles[i] = Emitter->GetNumParticles();
+			GatheredInstanceParameters.EmitterTotalSpawnedParticles[i] = Emitter->GetTotalSpawnedParticles();
+			if (!Emitter->IsComplete())
+			{
+				GatheredInstanceParameters.NumAlive++;
+			}
+		}
+	}
+
+	//Bias the LastRenderTime slightly to account for any delay as it's written by the RT.
+	GatheredInstanceParameters.SafeTimeSinceRendererd = FMath::Max(0.0f, GatheredInstanceParameters.TimeSeconds - Component->GetLastRenderTime() - GLastRenderTimeSafetyBias);
+
+	GatheredInstanceParameters.RequestedExecutionState = RequestedExecutionState;
+
+	Component->GetOverrideParameters().Tick();
+}
+
+
+void FNiagaraSystemInstance::TickInstanceParameters_Concurrent()
+{
+	FVector CurrPos = GatheredInstanceParameters.ComponentTrans.GetLocation();
+	FVector OldPos = FMath::IsNearlyZero(Age) ? CurrPos : GatheredInstanceParameters.OldPos; // The first frame the value in OwnerPositionParam is uninitialized memory, we need to make sure that we don't use it
+
 	OwnerPositionParam.SetValue(CurrPos);
-	OwnerScaleParam.SetValue(ComponentTrans.GetScale3D());
-	OwnerVelocityParam.SetValue((CurrPos - OldPos) / DeltaSeconds);
-	OwnerXAxisParam.SetValue(ComponentTrans.GetRotation().GetAxisX());
-	OwnerYAxisParam.SetValue(ComponentTrans.GetRotation().GetAxisY());
-	OwnerZAxisParam.SetValue(ComponentTrans.GetRotation().GetAxisZ());
+	OwnerScaleParam.SetValue(GatheredInstanceParameters.ComponentTrans.GetScale3D());
+	OwnerVelocityParam.SetValue((CurrPos - OldPos) / GatheredInstanceParameters.DeltaSeconds);
+
+	FQuat Rotation = GatheredInstanceParameters.ComponentTrans.GetRotation();
+	OwnerXAxisParam.SetValue(Rotation.GetAxisX());
+	OwnerYAxisParam.SetValue(Rotation.GetAxisY());
+	OwnerZAxisParam.SetValue(Rotation.GetAxisZ());
 
-	OwnerRotationParam.SetValue(ComponentTrans.GetRotation());
+	OwnerRotationParam.SetValue(Rotation);
 
-	FMatrix Transform = ComponentTrans.ToMatrixWithScale();
+	FMatrix Transform = GatheredInstanceParameters.ComponentTrans.ToMatrixWithScale();
 	FMatrix Inverse = Transform.Inverse();
 	FMatrix Transpose = Transform.GetTransposed();
 	FMatrix InverseTranspose = Inverse.GetTransposed();
@@ -1404,58 +1453,38 @@ void FNiagaraSystemInstance::TickInstanceParameters(float DeltaSeconds)
 	OwnerTransposeParam.SetValue(Transpose);
 	OwnerInverseTransposeParam.SetValue(InverseTranspose);
 
-	FMatrix TransformNoScale = ComponentTrans.ToMatrixNoScale();
+	FMatrix TransformNoScale = GatheredInstanceParameters.ComponentTrans.ToMatrixNoScale();
 	FMatrix InverseNoScale = TransformNoScale.Inverse();
 	OwnerTransformNoScaleParam.SetValue(TransformNoScale);
 	OwnerInverseNoScaleParam.SetValue(InverseNoScale);
 
-	OwnerDeltaSecondsParam.SetValue(DeltaSeconds);
-	OwnerInverseDeltaSecondsParam.SetValue(1.0f / DeltaSeconds);
+	OwnerDeltaSecondsParam.SetValue(GatheredInstanceParameters.DeltaSeconds);
+	OwnerInverseDeltaSecondsParam.SetValue(1.0f / GatheredInstanceParameters.DeltaSeconds);
 
-	//Calculate the min distance to a camera.
-	UWorld* World = Component->GetWorld();
-	if (World != NULL)
-	{
-		float LODDistance = GetLODDistance();
-		OwnerLODDistanceParam.SetValue(LODDistance);
-		OwnerEngineTimeParam.SetValue(World->TimeSeconds);
-		OwnerEngineRealtimeParam.SetValue(World->RealTimeSeconds);
-	}
-	else
-	{
-		OwnerEngineTimeParam.SetValue(Age);
-		OwnerEngineRealtimeParam.SetValue(Age);
-	}
-	SystemAgeParam.SetValue(Age);
-	SystemTickCountParam.SetValue(TickCount);
+	OwnerLODDistanceParam.SetValue(GatheredInstanceParameters.LODDistance);
+	OwnerEngineTimeParam.SetValue(GatheredInstanceParameters.TimeSeconds);
+	OwnerEngineRealtimeParam.SetValue(GatheredInstanceParameters.RealTimeSeconds);
+
+	SystemAgeParam.SetValue(GatheredInstanceParameters.Age);
+	SystemTickCountParam.SetValue(GatheredInstanceParameters.TickCount);
 
-	int32 NumAlive = 0;
-	const TArray<FNiagaraEmitterHandle>& EmitterHandles = GetSystem()->GetEmitterHandles();
 	for (int32 i = 0; i < Emitters.Num(); i++)
 	{
-		if (EmitterHandles[i].GetIsEnabled())//TODO: We should just null out the entry to the emitter in the array.
+		const TSharedRef<FNiagaraEmitterInstance, ESPMode::ThreadSafe>& Emitter = Emitters[i];
+		if (Emitter->GetExecutionState() != ENiagaraExecutionState::Disabled)
 		{
-			int32 NumParticles = Emitters[i]->GetNumParticles();
-			if (!Emitters[i]->IsComplete())
-			{
-				NumAlive++;
-			}
-			ParameterNumParticleBindings[i].SetValue(NumParticles);
-			ParameterTotalSpawnedParticlesBindings[i].SetValue(Emitters[i]->GetTotalSpawnedParticles());
+			ParameterNumParticleBindings[i].SetValue(GatheredInstanceParameters.EmitterNumParticles[i]);
+			ParameterTotalSpawnedParticlesBindings[i].SetValue(GatheredInstanceParameters.EmitterTotalSpawnedParticles[i]);
 		}
 	}
+
 	SystemNumEmittersParam.SetValue(Emitters.Num());
-	SystemNumEmittersAliveParam.SetValue(NumAlive);
+	SystemNumEmittersAliveParam.SetValue(GatheredInstanceParameters.NumAlive);
 
-	check(World);
-	float WorldTime = World->GetTimeSeconds();
-	//Bias the LastRenderTime slightly to account for any delay as it's written by the RT.
-	float SafeTimeSinceRendererd = FMath::Max(0.0f, WorldTime - Component->GetLastRenderTime() - GLastRenderTimeSafetyBias);
-	SystemTimeSinceRenderedParam.SetValue(SafeTimeSinceRendererd);
-	
-	OwnerExecutionStateParam.SetValue((int32)RequestedExecutionState);
+	SystemTimeSinceRenderedParam.SetValue(GatheredInstanceParameters.SafeTimeSinceRendererd);
+
+	OwnerExecutionStateParam.SetValue((int32)GatheredInstanceParameters.RequestedExecutionState);
 	
-	Component->GetOverrideParameters().Tick();
 	InstanceParameters.Tick();
 	InstanceParameters.MarkParametersDirty();
 }
@@ -1581,6 +1610,8 @@ void FNiagaraSystemInstance::InitEmitters()
 			LocalBounds = System->GetFixedBounds();
 		}
 	}
+
+	GatheredInstanceParameters.Init(Emitters.Num());
 }
 
 void FNiagaraSystemInstance::ComponentTick(float DeltaSeconds, const FGraphEventRef& MyCompletionGraphEvent)
@@ -1668,7 +1699,7 @@ void FNiagaraSystemInstance::Tick_GameThread(float DeltaSeconds)
 	CachedDeltaSeconds = DeltaSeconds;
 	bNeedsFinalize = true;
 
-	TickInstanceParameters(DeltaSeconds);
+	TickInstanceParameters_GameThread(DeltaSeconds);
 
 	TickDataInterfaces(DeltaSeconds, false);
 
diff --git a/Engine/Plugins/FX/Niagara/Source/Niagara/Private/NiagaraSystemSimulation.cpp b/Engine/Plugins/FX/Niagara/Source/Niagara/Private/NiagaraSystemSimulation.cpp
--- a/Engine/Plugins/FX/Niagara/Source/Niagara/Private/NiagaraSystemSimulation.cpp
+++ b/Engine/Plugins/FX/Niagara/Source/Niagara/Private/NiagaraSystemSimulation.cpp
@@ -1017,6 +1017,10 @@ void FNiagaraSystemSimulation::Tick_Concurrent(FNiagaraSystemSimulationTickConte
 
 		FScopeCycleCounter SystemStatCounter(Context.System->GetStatID(true, true));
 
+		for (FNiagaraSystemInstance* SystemInstance : Context.Instances)
+		{
+			SystemInstance->TickInstanceParameters_Concurrent();
+		}
 
 		if (Context.Owner->GetSystem()->FastPathMode != ENiagaraFastPathMode::ScriptVMOnly)
 		{
@@ -1318,8 +1322,6 @@ void FNiagaraSystemSimulation::PrepareForSystemSimulate(FNiagaraSystemSimulation
 	{
 		FNiagaraSystemInstance* Inst = Context.Instances[SystemIndex];
 
-		//Inst->TickInstanceParameters(Context.DeltaSeconds);
-
 		if (Inst->GetParameters().GetParametersDirty() && bCanExecute)
 		{
 			SpawnInstanceParameterToDataSetBinding.ParameterStoreToDataSet(Inst->GetParameters(), SpawnInstanceParameterDataSet, SystemIndex);
diff --git a/Engine/Plugins/FX/Niagara/Source/Niagara/Public/NiagaraSystemInstance.h b/Engine/Plugins/FX/Niagara/Source/Niagara/Public/NiagaraSystemInstance.h
--- a/Engine/Plugins/FX/Niagara/Source/Niagara/Public/NiagaraSystemInstance.h
+++ b/Engine/Plugins/FX/Niagara/Source/Niagara/Public/NiagaraSystemInstance.h
@@ -241,7 +241,9 @@ public:
 		bNeedsFinalize = true;
 	}
 
-	void TickInstanceParameters(float DeltaSeconds);
+	void TickInstanceParameters_GameThread(float DeltaSeconds);
+
+	void TickInstanceParameters_Concurrent();
 
 	void TickFastPathBindings();
 
@@ -401,6 +403,37 @@ public:
 	uint32 ActiveGPUEmitterCount = 0;
 	int32 GPUDataInterfaceInstanceDataSize = 0;
 
+	struct FInstanceParameters
+	{
+		float DeltaSeconds;
+
+		FTransform ComponentTrans;
+		FVector OldPos;
+
+		float LODDistance;
+		float TimeSeconds;
+		float RealTimeSeconds;
+
+		float Age;
+		int32 TickCount;
+
+		TArray<int32> EmitterNumParticles;
+		TArray<int32> EmitterTotalSpawnedParticles;
+		int32 NumAlive;
+
+		float SafeTimeSinceRendererd;
+
+		ENiagaraExecutionState RequestedExecutionState;
+
+		void Init(int32 NumEmitters)
+		{
+			EmitterNumParticles.AddUninitialized(NumEmitters);
+			EmitterTotalSpawnedParticles.AddUninitialized(NumEmitters);
+		}
+	};
+
+	FInstanceParameters GatheredInstanceParameters;
+
 	FNiagaraSystemFastPath::FParamMap0 FastPathMap;
 
 	TArray<TNiagaraFastPathRangedInputBinding<int32>> FastPathIntUpdateRangedInputBindings;
