Added "Niagara.AllowQuickSortedParameterOffsetsCopy=1" to enable using memcopy when copying FNiagaraVariableWithOffset arrays.
Also added safeties to ensure that no data is held by any FNiagaraVariableWithOffset.
Added a few FORCEINLINE for FNiagaraVariableWithOffset, FNiagaraVariable and FNiagaraTypeDefinition constructors data don't allocate data.

diff --git a/Engine/Plugins/FX/Niagara/Source/Niagara/Private/NiagaraParameterStore.cpp b/Engine/Plugins/FX/Niagara/Source/Niagara/Private/NiagaraParameterStore.cpp
--- a/Engine/Plugins/FX/Niagara/Source/Niagara/Private/NiagaraParameterStore.cpp
+++ b/Engine/Plugins/FX/Niagara/Source/Niagara/Private/NiagaraParameterStore.cpp
@@ -26,6 +26,43 @@ static FAutoConsoleVariableRef CVarNiagaraDumpParticleParameterStores(
 
 //////////////////////////////////////////////////////////////////////////
 
+int32 GNiagaraAllowQuickSortedParameterOffetsCopy = 1;
+static FAutoConsoleVariableRef CVarNiagaraAllowQuickSortedParameterOffetsCopy(
+	TEXT("Niagara.AllowQuickSortedParameterOffsetsCopy"),
+	GNiagaraAllowQuickSortedParameterOffetsCopy,
+	TEXT("Whether to use memcpy to copy sortedparameteroffset arrays. (default=1)\n"),
+	ECVF_Scalability
+);
+
+namespace
+{
+	void CopySortedParameterOffsets(TArray<FNiagaraVariableWithOffset>& Dest, const TArray<FNiagaraVariableWithOffset>& Src)
+	{
+		if (GNiagaraAllowQuickSortedParameterOffetsCopy)
+		{
+		#if DO_GUARD_SLOW
+			// Safeguard while we don't have yet the FNiagaraVariable type without data.
+			for (const FNiagaraVariableWithOffset& ParamWithOffset : Dest)
+			{
+				checkSlow(!ParamWithOffset.IsDataAllocated());
+			}
+			for (const FNiagaraVariableWithOffset& ParamWithOffset : Src)
+			{
+				checkSlow(!ParamWithOffset.IsDataAllocated());
+			}
+		#endif
+			Dest.SetNumUninitialized(Src.Num());
+			FMemory::Memcpy(Dest.GetData(), Src.GetData(), Dest.GetTypeSize() * Dest.Num());
+		}
+		else
+		{
+			Dest = Src;
+		}
+	}
+}
+
+//////////////////////////////////////////////////////////////////////////
+
 bool FNiagaraVariableSearch::Find(const FNiagaraVariableWithOffset* Variables, const FNiagaraVariable& Ref, int32 Start, int32 Num, int32& CheckIndex)
 {
 	while (Num)
@@ -88,7 +125,7 @@ FNiagaraParameterStore& FNiagaraParameterStore::operator=(const FNiagaraParamete
 #if WITH_EDITORONLY_DATA
 	ParameterOffsets = Other.ParameterOffsets;
 #endif // WITH_EDITORONLY_DATA
-	SortedParameterOffsets = Other.SortedParameterOffsets;
+	CopySortedParameterOffsets(SortedParameterOffsets, Other.SortedParameterOffsets);
 	DEC_MEMORY_STAT_BY(STAT_NiagaraParamStoreMemory, ParameterData.GetAllocatedSize());
 	ParameterData = Other.ParameterData;
 	INC_MEMORY_STAT_BY(STAT_NiagaraParamStoreMemory, ParameterData.GetAllocatedSize());
@@ -513,7 +550,7 @@ bool FNiagaraParameterStore::RemoveParameter(const FNiagaraVariable& ToRemove)
 
 		DEC_MEMORY_STAT_BY(STAT_NiagaraParamStoreMemory, ParameterData.GetAllocatedSize());
 
-		SortedParameterOffsets = NewOffsets;
+		CopySortedParameterOffsets(SortedParameterOffsets, NewOffsets);
 		ParameterData = NewData;
 		DataInterfaces = NewInterfaces;
 		UObjects = NewUObjects;
@@ -722,8 +759,7 @@ void FNiagaraParameterStore::InitFromSource(const FNiagaraParameterStore* SrcSto
 #if WITH_EDITORONLY_DATA
 	ParameterOffsets = SrcStore->ParameterOffsets;
 #endif // WITH_EDITORONLY_DATA
-	SortedParameterOffsets = SrcStore->SortedParameterOffsets;
-
+	CopySortedParameterOffsets(SortedParameterOffsets, SrcStore->SortedParameterOffsets);
 	DEC_MEMORY_STAT_BY(STAT_NiagaraParamStoreMemory, ParameterData.GetAllocatedSize());
 	ParameterData = SrcStore->ParameterData;
 	INC_MEMORY_STAT_BY(STAT_NiagaraParamStoreMemory, ParameterData.GetAllocatedSize());
@@ -859,6 +895,11 @@ void FNiagaraParameterStore::PostLoad()
 		}
 		ParameterOffsets.Empty();
 	}
+	// Safeguard while we don't have yet the FNiagaraVariable type without data
+	for (FNiagaraVariableWithOffset& ParamWithOffset : SortedParameterOffsets)
+	{
+		ParamWithOffset.ClearData();
+	}
 #endif
 
 	// Not always required if NIAGARA_VARIABLE_LEXICAL_SORTING
diff --git a/Engine/Plugins/FX/Niagara/Source/Niagara/Public/NiagaraParameterStore.h b/Engine/Plugins/FX/Niagara/Source/Niagara/Public/NiagaraParameterStore.h
--- a/Engine/Plugins/FX/Niagara/Source/Niagara/Public/NiagaraParameterStore.h
+++ b/Engine/Plugins/FX/Niagara/Source/Niagara/Public/NiagaraParameterStore.h
@@ -121,9 +121,10 @@ struct FNiagaraVariableWithOffset : public FNiagaraVariable
 {
 	GENERATED_USTRUCT_BODY()
 
-	FNiagaraVariableWithOffset() : Offset(INDEX_NONE) {}
-	FNiagaraVariableWithOffset(const FNiagaraVariableWithOffset& InRef) : FNiagaraVariable(InRef), Offset(InRef.Offset) {}
-	FNiagaraVariableWithOffset(const FNiagaraVariable& InVariable, int32 InOffset) : FNiagaraVariable(InVariable), Offset(InOffset) {}
+	// Those constructor enforce that there are no data allocated.
+	FORCEINLINE FNiagaraVariableWithOffset() : Offset(INDEX_NONE) {}
+	FORCEINLINE FNiagaraVariableWithOffset(const FNiagaraVariableWithOffset& InRef) : FNiagaraVariable(InRef.GetType(), InRef.GetName()), Offset(InRef.Offset) {}
+	FORCEINLINE FNiagaraVariableWithOffset(const FNiagaraVariable& InVariable, int32 InOffset) : FNiagaraVariable(InVariable.GetType(), InVariable.GetName()), Offset(InOffset) {}
 
 	UPROPERTY()
 	int32 Offset;
diff --git a/Engine/Plugins/FX/Niagara/Source/Niagara/Public/NiagaraTypes.h b/Engine/Plugins/FX/Niagara/Source/Niagara/Public/NiagaraTypes.h
--- a/Engine/Plugins/FX/Niagara/Source/Niagara/Public/NiagaraTypes.h
+++ b/Engine/Plugins/FX/Niagara/Source/Niagara/Public/NiagaraTypes.h
@@ -382,40 +382,40 @@ struct NIAGARA_API FNiagaraTypeDefinition
 public:
 
 	// Construct blank raw type definition 
-	FNiagaraTypeDefinition(UClass *ClassDef)
+	FORCEINLINE FNiagaraTypeDefinition(UClass *ClassDef)
 		: Struct(ClassDef), Enum(nullptr), Size(INDEX_NONE), Alignment(INDEX_NONE)
 	{
 		checkSlow(Struct != nullptr);
 	}
 
-	FNiagaraTypeDefinition(UEnum *EnumDef)
+	FORCEINLINE FNiagaraTypeDefinition(UEnum *EnumDef)
 		: Struct(IntStruct), Enum(EnumDef), Size(INDEX_NONE), Alignment(INDEX_NONE)
 	{
 		checkSlow(Struct != nullptr);
 	}
 
-	FNiagaraTypeDefinition(UScriptStruct *StructDef)
+	FORCEINLINE FNiagaraTypeDefinition(UScriptStruct *StructDef)
 		: Struct(StructDef), Enum(nullptr), Size(INDEX_NONE), Alignment(INDEX_NONE)
 	{
 		checkSlow(Struct != nullptr);
 	}
 
-	FNiagaraTypeDefinition(const FNiagaraTypeDefinition &Other)
+	FORCEINLINE FNiagaraTypeDefinition(const FNiagaraTypeDefinition &Other)
 		: Struct(Other.Struct), Enum(Other.Enum), Size(INDEX_NONE), Alignment(INDEX_NONE)
 	{
 	}
 
 	// Construct a blank raw type definition
-	FNiagaraTypeDefinition()
+	FORCEINLINE FNiagaraTypeDefinition()
 		: Struct(nullptr), Enum(nullptr), Size(INDEX_NONE), Alignment(INDEX_NONE)
 	{}
 
-	bool operator !=(const FNiagaraTypeDefinition &Other) const
+	FORCEINLINE bool operator !=(const FNiagaraTypeDefinition &Other) const
 	{
 		return !(*this == Other);
 	}
 
-	bool operator == (const FNiagaraTypeDefinition &Other) const
+	FORCEINLINE bool operator == (const FNiagaraTypeDefinition &Other) const
 	{
 		return Struct == Other.Struct && Enum == Other.Enum;
 	}
@@ -796,7 +796,7 @@ struct FNiagaraVariable
 		}
 	}
 
-	FNiagaraVariable(FNiagaraTypeDefinition InType, FName InName)
+	FORCEINLINE FNiagaraVariable(const FNiagaraTypeDefinition& InType, const FName& InName)
 		: Name(InName)
 		, TypeDef(InType)
 	{
@@ -881,6 +881,11 @@ struct FNiagaraVariable
 		return VarData.GetData();
 	}
 
+	void ClearData()
+	{
+		VarData.Empty();
+	}
+
 	int32 GetSizeInBytes() const
 	{
 		return TypeDef.GetSize();
