- Avoid string allocation per frame per spawn info
- Avoid setting variables if we aren't going to spawn anything

#rnx

diff --git a/Engine/Plugins/FX/Niagara/Source/Niagara/Private/NiagaraEmitterInstance.cpp b/Engine/Plugins/FX/Niagara/Source/Niagara/Private/NiagaraEmitterInstance.cpp
--- a/Engine/Plugins/FX/Niagara/Source/Niagara/Private/NiagaraEmitterInstance.cpp
+++ b/Engine/Plugins/FX/Niagara/Source/Niagara/Private/NiagaraEmitterInstance.cpp
@@ -1293,60 +1293,60 @@ void FNiagaraEmitterInstance::Tick(float DeltaSeconds)
 		SCOPE_CYCLE_COUNTER(STAT_NiagaraSpawn);
 
 		//Handle main spawn rate spawning
-		auto SpawnParticles = [&](int32 Num, FString DumpLabel)
+		auto SpawnParticles = [&](int32 Num, const TCHAR* DumpLabel)
 		{
-			if (Num > 0)
-			{
-				int32 OrigNum = Data.GetDestinationDataChecked().GetNumInstances();
-				Data.GetDestinationDataChecked().SetNumInstances(OrigNum + Num);
+			int32 OrigNum = Data.GetDestinationDataChecked().GetNumInstances();
+			Data.GetDestinationDataChecked().SetNumInstances(OrigNum + Num);
 
-				// NOTE(mv): Updates the count after setting the variable, such that the TotalSpawnedParticles value read 
-				//           in the script has the count at the start of the frame. 
-				//           This way UniqueID = TotalSpawnedParticles + ExecIndex provide unique and sequential identifiers. 
-				// NOTE(mv): Only for CPU particles, as GPU particles early outs further up and has a separate increment. 
-				TotalSpawnedParticles += Num;
+			// NOTE(mv): Updates the count after setting the variable, such that the TotalSpawnedParticles value read 
+			//           in the script has the count at the start of the frame. 
+			//           This way UniqueID = TotalSpawnedParticles + ExecIndex provide unique and sequential identifiers. 
+			// NOTE(mv): Only for CPU particles, as GPU particles early outs further up and has a separate increment. 
+			TotalSpawnedParticles += Num;
 
-				SpawnExecCountBinding.SetValue(Num);
-				SpawnExecContext.BindData(0, Data, OrigNum, true);
+			SpawnExecCountBinding.SetValue(Num);
+			SpawnExecContext.BindData(0, Data, OrigNum, true);
 
-				//UE_LOG(LogNiagara, Log, TEXT("SpawnScriptEventDataSets: %d"), SpawnScriptEventDataSets.Num());
-				int32 EventDataSetIdx = 1;
-				for (FNiagaraDataSet* EventDataSet : SpawnScriptEventDataSets)
-				{
-					//UE_LOG(LogNiagara, Log, TEXT("SpawnScriptEventDataSets.. %d"), EventDataSet->GetNumVariables());
-					int32 EventOrigNum = EventDataSet->GetDestinationDataChecked().GetNumInstances();
-					EventDataSet->GetDestinationDataChecked().SetNumInstances(EventOrigNum + Num);
-					SpawnExecContext.BindData(EventDataSetIdx++, *EventDataSet, EventOrigNum, true);
-				}
+			//UE_LOG(LogNiagara, Log, TEXT("SpawnScriptEventDataSets: %d"), SpawnScriptEventDataSets.Num());
+			int32 EventDataSetIdx = 1;
+			for (FNiagaraDataSet* EventDataSet : SpawnScriptEventDataSets)
+			{
+				//UE_LOG(LogNiagara, Log, TEXT("SpawnScriptEventDataSets.. %d"), EventDataSet->GetNumVariables());
+				int32 EventOrigNum = EventDataSet->GetDestinationDataChecked().GetNumInstances();
+				EventDataSet->GetDestinationDataChecked().SetNumInstances(EventOrigNum + Num);
+				SpawnExecContext.BindData(EventDataSetIdx++, *EventDataSet, EventOrigNum, true);
+			}
 
-				SpawnExecContext.Execute(Num);
+			SpawnExecContext.Execute(Num);
 
-				if (GbDumpParticleData || System->bDumpDebugEmitterInfo)
+			if (GbDumpParticleData || System->bDumpDebugEmitterInfo)
+			{
+				Data.GetDestinationDataChecked().Dump(OrigNum, Num, FString::Printf(TEXT("===  %s Spawned %d Particles==="), DumpLabel, Num));
+				for (int32 EventIdx = 0; EventIdx < SpawnScriptEventDataSets.Num(); ++EventIdx)
 				{
-					Data.GetDestinationDataChecked().Dump(OrigNum, Num, FString::Printf(TEXT("===  %s Spawned %d Particles==="), *DumpLabel, Num));
-					for (int32 EventIdx = 0; EventIdx < SpawnScriptEventDataSets.Num(); ++EventIdx)
+					FNiagaraDataSet* EventDataSet = SpawnScriptEventDataSets[EventIdx];
+					if (EventDataSet && EventDataSet->GetDestinationDataChecked().GetNumInstances() > 0)
 					{
-						FNiagaraDataSet* EventDataSet = SpawnScriptEventDataSets[EventIdx];
-						if (EventDataSet && EventDataSet->GetDestinationDataChecked().GetNumInstances() > 0)
-						{
-							EventDataSet->GetDestinationDataChecked().Dump(0, INDEX_NONE, FString::Printf(TEXT("Spawn Script Event %d"), EventIdx));
-						}
+						EventDataSet->GetDestinationDataChecked().Dump(0, INDEX_NONE, FString::Printf(TEXT("Spawn Script Event %d"), EventIdx));
 					}
-					//UE_LOG(LogNiagara, Log, TEXT("=== %s Spawn Parameters ==="), *DumpLabel);
-					SpawnExecContext.Parameters.Dump();
 				}
+				//UE_LOG(LogNiagara, Log, TEXT("=== %s Spawn Parameters ==="), *DumpLabel);
+				SpawnExecContext.Parameters.Dump();
 			}
 		};
 
 		//Perform all our regular spawning that's driven by our emitter script.
 		for (FNiagaraSpawnInfo& Info : SpawnInfos)
 		{
-			SpawnIntervalBinding.SetValue(Info.IntervalDt);
-			InterpSpawnStartBinding.SetValue(Info.InterpStartDt);
-			SpawnGroupBinding.SetValue(Info.SpawnGroup);
+			if ( Info.Count > 0 )
+			{
+				SpawnIntervalBinding.SetValue(Info.IntervalDt);
+				InterpSpawnStartBinding.SetValue(Info.InterpStartDt);
+				SpawnGroupBinding.SetValue(Info.SpawnGroup);
 
-			SpawnParticles(Info.Count, TEXT("Regular Spawn"));
-		};
+				SpawnParticles(Info.Count, TEXT("Regular Spawn"));
+			}
+		}
 
 		EventSpawnStart = Data.GetDestinationDataChecked().GetNumInstances();
 
@@ -1356,20 +1356,23 @@ void FNiagaraEmitterInstance::Tick(float DeltaSeconds)
 			//Spawn particles coming from events.
 			for (int32 i = 0; i < Info.SpawnCounts.Num(); i++)
 			{
-				int32 EventNumToSpawn = Info.SpawnCounts[i];
+				const int32 EventNumToSpawn = Info.SpawnCounts[i];
+				if (EventNumToSpawn > 0)
+				{
+					const int32 CurrNumParticles = Data.GetDestinationDataChecked().GetNumInstances();
 
-				int32 CurrNumParticles = Data.GetDestinationDataChecked().GetNumInstances();
-				//Event spawns are instantaneous at the middle of the frame?
-				SpawnIntervalBinding.SetValue(0.0f);
-				InterpSpawnStartBinding.SetValue(DeltaSeconds * 0.5f);
-				SpawnGroupBinding.SetValue(0);
+					//Event spawns are instantaneous at the middle of the frame?
+					SpawnIntervalBinding.SetValue(0.0f);
+					InterpSpawnStartBinding.SetValue(DeltaSeconds * 0.5f);
+					SpawnGroupBinding.SetValue(0);
 
-				SpawnParticles(EventNumToSpawn, TEXT("Event Spawn"));
+					SpawnParticles(EventNumToSpawn, TEXT("Event Spawn"));
 
-				//Update EventSpawnCounts to the number actually spawned.
-				int32 NumActuallySpawned = Data.GetDestinationDataChecked().GetNumInstances() - CurrNumParticles;
-				TotalActualEventSpawns += NumActuallySpawned;
-				Info.SpawnCounts[i] = NumActuallySpawned;
+					//Update EventSpawnCounts to the number actually spawned.
+					int32 NumActuallySpawned = Data.GetDestinationDataChecked().GetNumInstances() - CurrNumParticles;
+					TotalActualEventSpawns += NumActuallySpawned;
+					Info.SpawnCounts[i] = NumActuallySpawned;
+				}
 			}
 		}
 	}
