Fixed a crash where an invalid renderer was being called on a gpu emitter

#tests engine tests are all green

diff --git a/Engine/Plugins/FX/Niagara/Source/Niagara/Classes/NiagaraEmitter.h b/Engine/Plugins/FX/Niagara/Source/Niagara/Classes/NiagaraEmitter.h
--- a/Engine/Plugins/FX/Niagara/Source/Niagara/Classes/NiagaraEmitter.h
+++ b/Engine/Plugins/FX/Niagara/Source/Niagara/Classes/NiagaraEmitter.h
@@ -355,6 +355,7 @@ public:
 	bool NIAGARA_API SetUniqueEmitterName(const FString& InName);
 
 	const TArray<UNiagaraRendererProperties*>& GetRenderers() const { return RendererProperties; }
+	TArray<UNiagaraRendererProperties*> GetEnabledRenderers() const;
 
 	void NIAGARA_API AddRenderer(UNiagaraRendererProperties* Renderer);
 
diff --git a/Engine/Plugins/FX/Niagara/Source/Niagara/Private/NiagaraComponent.cpp b/Engine/Plugins/FX/Niagara/Source/Niagara/Private/NiagaraComponent.cpp
--- a/Engine/Plugins/FX/Niagara/Source/Niagara/Private/NiagaraComponent.cpp
+++ b/Engine/Plugins/FX/Niagara/Source/Niagara/Private/NiagaraComponent.cpp
@@ -180,15 +180,10 @@ void FNiagaraSceneProxy::CreateRenderers(const UNiagaraComponent* Component)
 	{
 		if (UNiagaraEmitter* Emitter = EmitterInst->GetCachedEmitter())
 		{
-			for (UNiagaraRendererProperties* Properties : Emitter->GetRenderers())
+			for (UNiagaraRendererProperties* Properties : Emitter->GetEnabledRenderers())
 			{
 				RendererSortInfo.Emplace(Properties->SortOrderHint, EmitterRenderers.Num());
-				FNiagaraRenderer* NewRenderer = nullptr;
-				if (Properties->GetIsEnabled())
-				{
-					NewRenderer = Properties->CreateEmitterRenderer(FeatureLevel, &EmitterInst.Get());
-				}
-				EmitterRenderers.Add(NewRenderer);
+				EmitterRenderers.Add(Properties->CreateEmitterRenderer(FeatureLevel, &EmitterInst.Get()));
 			}
 		}
 	}
@@ -1048,9 +1043,10 @@ void UNiagaraComponent::SendRenderDynamicData_Concurrent()
 			FScopeCycleCounter EmitterStatCounter(EmitterStatID);
 #endif
 
-			for (int32 EmitterIdx = 0; EmitterIdx < Emitter->GetRenderers().Num(); EmitterIdx++, RendererIndex++)
+			const TArray<UNiagaraRendererProperties*>& Renderers = Emitter->GetEnabledRenderers();
+			for (int32 EmitterIdx = 0; EmitterIdx < Renderers.Num(); EmitterIdx++, RendererIndex++)
 			{
-				UNiagaraRendererProperties* Properties = Emitter->GetRenderers()[EmitterIdx];
+				UNiagaraRendererProperties* Properties = Renderers[EmitterIdx];
 				FNiagaraRenderer* Renderer = EmitterRenderers[RendererIndex];
 				FNiagaraDynamicDataBase* NewData = nullptr;
 				
@@ -1061,7 +1057,6 @@ void UNiagaraComponent::SendRenderDynamicData_Concurrent()
 					const FNiagaraEmitterHandle& Handle = Asset->GetEmitterHandle(i);
 					bRendererEditorEnabled = (!SystemInstance->GetIsolateEnabled() || Handle.IsIsolated());
 #endif
-					bRendererEditorEnabled &= Properties->GetIsEnabled();
 					if (bRendererEditorEnabled && !EmitterInst->IsComplete() && !SystemInstance->IsComplete())
 					{
 						NewData = Renderer->GenerateDynamicData(NiagaraProxy, Properties, EmitterInst);
@@ -1096,9 +1091,8 @@ int32 UNiagaraComponent::GetNumMaterials() const
 		{
 			FNiagaraEmitterInstance* EmitterInst = &SystemInstance->GetEmitters()[i].Get();
 			UNiagaraEmitter* Emitter = EmitterInst->GetCachedEmitter();
-			for (int32 EmitterIdx = 0; EmitterIdx < Emitter->GetRenderers().Num(); EmitterIdx++)
+			for (UNiagaraRendererProperties* Properties : Emitter->GetEnabledRenderers())
 			{
-				UNiagaraRendererProperties* Properties = Emitter->GetRenderers()[EmitterIdx];
 				Properties->GetUsedMaterials(EmitterInst, UsedMaterials);
 			}
 		}
@@ -1158,12 +1152,9 @@ void UNiagaraComponent::GetUsedMaterials(TArray<UMaterialInterface*>& OutMateria
 		{
 			if (Props)
 			{
-				for (int32 i = 0; i < Props->GetRenderers().Num(); i++)
+				for (UNiagaraRendererProperties* Renderer : Props->GetEnabledRenderers())
 				{
-					if (UNiagaraRendererProperties* Renderer = Props->GetRenderers()[i])
-					{
-						Renderer->GetUsedMaterials(&Sim.Get(), OutMaterials);
-					}
+					Renderer->GetUsedMaterials(&Sim.Get(), OutMaterials);
 				}
 			}
 		}
diff --git a/Engine/Plugins/FX/Niagara/Source/Niagara/Private/NiagaraEmitter.cpp b/Engine/Plugins/FX/Niagara/Source/Niagara/Private/NiagaraEmitter.cpp
--- a/Engine/Plugins/FX/Niagara/Source/Niagara/Private/NiagaraEmitter.cpp
+++ b/Engine/Plugins/FX/Niagara/Source/Niagara/Private/NiagaraEmitter.cpp
@@ -1021,6 +1021,19 @@ bool UNiagaraEmitter::SetUniqueEmitterName(const FString& InName)
 	return false;
 }
 
+TArray<UNiagaraRendererProperties*> UNiagaraEmitter::GetEnabledRenderers() const
+{
+	TArray<UNiagaraRendererProperties*> Renderers;
+	for (UNiagaraRendererProperties* Renderer : RendererProperties)
+	{
+		if (Renderer && Renderer->GetIsEnabled() && Renderer->IsSimTargetSupported(this->SimTarget))
+		{
+			Renderers.Add(Renderer);
+		}
+	}
+	return Renderers;
+}
+
 void UNiagaraEmitter::AddRenderer(UNiagaraRendererProperties* Renderer)
 {
 	Modify();
diff --git a/Engine/Plugins/FX/Niagara/Source/Niagara/Private/NiagaraEmitterInstance.cpp b/Engine/Plugins/FX/Niagara/Source/Niagara/Private/NiagaraEmitterInstance.cpp
--- a/Engine/Plugins/FX/Niagara/Source/Niagara/Private/NiagaraEmitterInstance.cpp
+++ b/Engine/Plugins/FX/Niagara/Source/Niagara/Private/NiagaraEmitterInstance.cpp
@@ -409,17 +409,15 @@ void FNiagaraEmitterInstance::Init(int32 InEmitterIdx, FNiagaraSystemInstanceID
 
 	// Initialize bounds calculators
 	//-OPT: Could skip creating this if we won't ever use it
-	BoundsCalculators.Reserve(CachedEmitter->GetRenderers().Num());
-	for (UNiagaraRendererProperties* RendererProperties : CachedEmitter->GetRenderers())
+	const TArray<UNiagaraRendererProperties*>& Renderers = CachedEmitter->GetEnabledRenderers();
+	BoundsCalculators.Reserve(Renderers.Num());
+	for (UNiagaraRendererProperties* RendererProperties : Renderers)
 	{
-		if ((RendererProperties != nullptr) && RendererProperties->GetIsEnabled())
+		FNiagaraBoundsCalculator* BoundsCalculator = RendererProperties->CreateBoundsCalculator();
+		if (BoundsCalculator != nullptr)
 		{
-			FNiagaraBoundsCalculator* BoundsCalculator = RendererProperties->CreateBoundsCalculator();
-			if (BoundsCalculator != nullptr)
-			{
-				BoundsCalculator->InitAccessors(*ParticleDataSet);
-				BoundsCalculators.Emplace(BoundsCalculator);
-			}
+			BoundsCalculator->InitAccessors(*ParticleDataSet);
+			BoundsCalculators.Emplace(BoundsCalculator);
 		}
 	}
 }
