Fixing crash when searching through stack. Need to figure out root cause, but this lets us keep going.

diff --git a/Engine/Plugins/FX/Niagara/Source/NiagaraEditor/Private/ViewModels/Stack/NiagaraStackEntry.cpp b/Engine/Plugins/FX/Niagara/Source/NiagaraEditor/Private/ViewModels/Stack/NiagaraStackEntry.cpp
--- a/Engine/Plugins/FX/Niagara/Source/NiagaraEditor/Private/ViewModels/Stack/NiagaraStackEntry.cpp
+++ b/Engine/Plugins/FX/Niagara/Source/NiagaraEditor/Private/ViewModels/Stack/NiagaraStackEntry.cpp
@@ -218,7 +218,10 @@ bool UNiagaraStackEntry::GetIsExpanded() const
 
 void UNiagaraStackEntry::SetIsExpanded(bool bInExpanded)
 {
-	StackEditorData->SetStackEntryIsExpanded(GetStackEditorDataKey(), bInExpanded);
+	if (StackEditorData)
+	{
+		StackEditorData->SetStackEntryIsExpanded(GetStackEditorDataKey(), bInExpanded);
+	}
 	bIsExpandedCache.Reset();
 }
 
