Passing minimal data to the gpu for Niagara rendering.

diff --git a/Engine/Plugins/FX/Niagara/Source/Niagara/Private/NiagaraLightRendererProperties.cpp b/Engine/Plugins/FX/Niagara/Source/Niagara/Private/NiagaraLightRendererProperties.cpp
--- a/Engine/Plugins/FX/Niagara/Source/Niagara/Private/NiagaraLightRendererProperties.cpp
+++ b/Engine/Plugins/FX/Niagara/Source/Niagara/Private/NiagaraLightRendererProperties.cpp
@@ -40,7 +40,7 @@ void UNiagaraLightRendererProperties::InitCDOPropertiesAfterModuleStartup()
 FNiagaraRenderer* UNiagaraLightRendererProperties::CreateEmitterRenderer(ERHIFeatureLevel::Type FeatureLevel, const FNiagaraEmitterInstance* Emitter)
 {
 	FNiagaraRenderer* NewRenderer = new FNiagaraRendererLights(FeatureLevel, this, Emitter);
-	NewRenderer->Initialize(FeatureLevel, this, Emitter);
+	NewRenderer->Initialize(this, Emitter);
 	return NewRenderer;
 }
 
diff --git a/Engine/Plugins/FX/Niagara/Source/Niagara/Private/NiagaraMeshRendererProperties.cpp b/Engine/Plugins/FX/Niagara/Source/Niagara/Private/NiagaraMeshRendererProperties.cpp
--- a/Engine/Plugins/FX/Niagara/Source/Niagara/Private/NiagaraMeshRendererProperties.cpp
+++ b/Engine/Plugins/FX/Niagara/Source/Niagara/Private/NiagaraMeshRendererProperties.cpp
@@ -38,7 +38,7 @@ FNiagaraRenderer* UNiagaraMeshRendererProperties::CreateEmitterRenderer(ERHIFeat
 	if (ParticleMesh)
 	{
 		FNiagaraRenderer* NewRenderer = new FNiagaraRendererMeshes(FeatureLevel, this, Emitter);
-		NewRenderer->Initialize(FeatureLevel, this, Emitter);
+		NewRenderer->Initialize(this, Emitter);
 		return NewRenderer;
 	}
 
diff --git a/Engine/Plugins/FX/Niagara/Source/Niagara/Private/NiagaraRenderer.cpp b/Engine/Plugins/FX/Niagara/Source/Niagara/Private/NiagaraRenderer.cpp
--- a/Engine/Plugins/FX/Niagara/Source/Niagara/Private/NiagaraRenderer.cpp
+++ b/Engine/Plugins/FX/Niagara/Source/Niagara/Private/NiagaraRenderer.cpp
@@ -134,6 +134,58 @@ FRWBuffer& FNiagaraRenderer::GetDummyUIntBuffer()
 	return DummyUIntBuffer.Buffer;
 }
 
+bool FNiagaraRenderer::SetVertexFactoryVariable(const FNiagaraDataSet& DataSet, const FNiagaraVariable& Var, int32 VFVarOffset)
+{
+	int32 FloatOffset;
+	int32 IntOffset;//TODO: No VF uses ints atm but it should be trivial to copy the float path if some vf should need to.
+	DataSet.GetVariableComponentOffsets(Var, FloatOffset, IntOffset);
+	int32 NumComponents = Var.GetSizeInBytes() / sizeof(float);
+
+	int32 GPULocation = INDEX_NONE;
+	bool bUpload = true;
+	if (FloatOffset != INDEX_NONE)
+	{
+		if (FNiagaraRendererVariableInfo* ExistingVarInfo = VFVariables.FindByPredicate([&](const FNiagaraRendererVariableInfo& VarInfo) { return VarInfo.DatasetOffset == FloatOffset; }))
+		{
+			//Don't need to upload this var again if it's already been uploaded for another var info. Just point to that.
+			//E.g. when custom sorting uses age.
+			GPULocation = ExistingVarInfo->GPUBufferOffset;
+			bUpload = false;
+		}
+		else
+		{
+			//For CPU Sims we pack just the required data tightly in a GPU buffer we upload. For GPU sims the data is there already so we just provide the real data location.
+			GPULocation = SimTarget == ENiagaraSimTarget::CPUSim ? TotalVFComponents : FloatOffset;
+			TotalVFComponents += NumComponents;
+		}
+	}
+
+	VFVariables[VFVarOffset] = FNiagaraRendererVariableInfo(FloatOffset, GPULocation, NumComponents, bUpload);
+
+	return FloatOffset != INDEX_NONE;
+}
+
+FGlobalDynamicReadBuffer::FAllocation FNiagaraRenderer::TransferDataToGPU(FGlobalDynamicReadBuffer& DynamicReadBuffer, FNiagaraDataBuffer* SrcData)const
+{
+	int32 TotalFloatSize = TotalVFComponents * SrcData->GetNumInstances();
+	int32 ComponentStrideDest = SrcData->GetNumInstances() * sizeof(float);
+	FGlobalDynamicReadBuffer::FAllocation Allocation = DynamicReadBuffer.AllocateFloat(TotalFloatSize);
+	for (const FNiagaraRendererVariableInfo& VarInfo : VFVariables)
+	{
+		int32 GpuOffset = VarInfo.GetGPUOffset();
+		if (GpuOffset != INDEX_NONE && VarInfo.bUpload)
+		{
+			for (int32 CompIdx = 0; CompIdx < VarInfo.NumComponents; ++CompIdx)
+			{
+				float* SrcComponent = (float*)SrcData->GetComponentPtrFloat(VarInfo.DatasetOffset + CompIdx);
+				void* Dest = Allocation.Buffer + ComponentStrideDest * (GpuOffset + CompIdx);
+				FMemory::Memcpy(Dest, SrcComponent, ComponentStrideDest);
+			}
+		}
+	}
+	return Allocation;
+}
+
 //////////////////////////////////////////////////////////////////////////
 
 FNiagaraDynamicDataBase::FNiagaraDynamicDataBase(const FNiagaraEmitterInstance* InEmitter)
@@ -190,19 +242,21 @@ FNiagaraDataBuffer* FNiagaraDynamicDataBase::GetParticleDataToRender()const
 //////////////////////////////////////////////////////////////////////////
 
 
-FNiagaraRenderer::FNiagaraRenderer(ERHIFeatureLevel::Type FeatureLevel, const UNiagaraRendererProperties *InProps, const FNiagaraEmitterInstance* Emitter)
+FNiagaraRenderer::FNiagaraRenderer(ERHIFeatureLevel::Type InFeatureLevel, const UNiagaraRendererProperties *InProps, const FNiagaraEmitterInstance* Emitter)
 	: DynamicDataRender(nullptr)
 	, bLocalSpace(Emitter->GetCachedEmitter()->bLocalSpace)
 	, bHasLights(false)
 	, SimTarget(Emitter->GetCachedEmitter()->SimTarget)
 	, NumIndicesPerInstance(InProps ? InProps->GetNumIndicesPerInstance() : 0)
+	, FeatureLevel(InFeatureLevel)
+	, TotalVFComponents(0)
 {
 #if STATS
 	EmitterStatID = Emitter->GetCachedEmitter()->GetStatID(false, false);
 #endif
 }
 
-void FNiagaraRenderer::Initialize(ERHIFeatureLevel::Type FeatureLevel, const UNiagaraRendererProperties *InProps, const FNiagaraEmitterInstance* Emitter)
+void FNiagaraRenderer::Initialize(const UNiagaraRendererProperties *InProps, const FNiagaraEmitterInstance* Emitter)
 {
 	//Get our list of valid base materials. Fall back to default material if they're not valid.
 	InProps->GetUsedMaterials(Emitter, BaseMaterials_GT);
@@ -286,13 +340,14 @@ struct FParticleOrderAsUint
 	FORCEINLINE operator uint32() const { return OrderAsUint; }
 };
 
-void FNiagaraRenderer::SortIndices(const FNiagaraGPUSortInfo& SortInfo, const FNiagaraDataBuffer& Buffer, FGlobalDynamicReadBuffer::FAllocation& OutIndices)const
+void FNiagaraRenderer::SortIndices(const FNiagaraGPUSortInfo& SortInfo, int32 SortVarIdx, const FNiagaraDataBuffer& Buffer, FGlobalDynamicReadBuffer::FAllocation& OutIndices)const
 {
 	SCOPE_CYCLE_COUNTER(STAT_NiagaraSortParticles);
 
 	uint32 NumInstances = Buffer.GetNumInstances();
 	check(OutIndices.ReadBuffer->NumBytes >= OutIndices.FirstIndex + NumInstances * sizeof(int32));
 	check(SortInfo.SortMode != ENiagaraSortMode::None);
+	check(VFVariables.IsValidIndex(SortVarIdx));
 	check(SortInfo.SortAttributeOffset != INDEX_NONE);
 
 	const bool bUseRadixSort = GNiagaraRadixSortThreshold != -1 && (int32)NumInstances  > GNiagaraRadixSortThreshold;
@@ -304,9 +359,10 @@ void FNiagaraRenderer::SortIndices(const FNiagaraGPUSortInfo& SortInfo, const FN
 
 	if (SortInfo.SortMode == ENiagaraSortMode::ViewDepth || SortInfo.SortMode == ENiagaraSortMode::ViewDistance)
 	{
-		float* RESTRICT PositionX = (float*)Buffer.GetComponentPtrFloat(SortInfo.SortAttributeOffset);
-		float* RESTRICT PositionY = (float*)Buffer.GetComponentPtrFloat(SortInfo.SortAttributeOffset + 1);
-		float* RESTRICT PositionZ = (float*)Buffer.GetComponentPtrFloat(SortInfo.SortAttributeOffset + 2);
+		int32 BaseCompOffset = VFVariables[SortVarIdx].DatasetOffset;
+		float* RESTRICT PositionX = (float*)Buffer.GetComponentPtrFloat(BaseCompOffset);
+		float* RESTRICT PositionY = (float*)Buffer.GetComponentPtrFloat(BaseCompOffset + 1);
+		float* RESTRICT PositionZ = (float*)Buffer.GetComponentPtrFloat(BaseCompOffset + 2);
 		auto GetPos = [&PositionX, &PositionY, &PositionZ](int32 Idx)
 		{
 			return FVector(PositionX[Idx], PositionY[Idx], PositionZ[Idx]);
@@ -329,7 +385,7 @@ void FNiagaraRenderer::SortIndices(const FNiagaraGPUSortInfo& SortInfo, const FN
 	}
 	else
 	{
-		float* RESTRICT CustomSorting = (float*)Buffer.GetComponentPtrFloat(SortInfo.SortAttributeOffset);
+		float* RESTRICT CustomSorting = (float*)Buffer.GetComponentPtrFloat(VFVariables[SortVarIdx].DatasetOffset);
 
 		if (SortInfo.SortMode == ENiagaraSortMode::CustomAscending)
 		{
diff --git a/Engine/Plugins/FX/Niagara/Source/Niagara/Private/NiagaraRendererMeshes.cpp b/Engine/Plugins/FX/Niagara/Source/Niagara/Private/NiagaraRendererMeshes.cpp
--- a/Engine/Plugins/FX/Niagara/Source/Niagara/Private/NiagaraRendererMeshes.cpp
+++ b/Engine/Plugins/FX/Niagara/Source/Niagara/Private/NiagaraRendererMeshes.cpp
@@ -26,6 +26,8 @@ static FAutoConsoleVariableRef CVarEnableNiagaraMeshRendering(
 	ECVF_Default
 );
 
+extern int32 GbEnableMinimalGPUBuffers;
+
 struct FNiagaraDynamicDataMesh : public FNiagaraDynamicDataBase
 {
 	FNiagaraDynamicDataMesh(const FNiagaraEmitterInstance* InEmitter)
@@ -52,22 +54,30 @@ public:
 
 //////////////////////////////////////////////////////////////////////////
 
+namespace ENiagaraMeshVFLayout
+{
+	enum Type
+	{
+		Position,
+		Velocity,
+		Color,
+		Scale,
+		Transform,
+		MaterialRandom,
+		NormalizedAge,
+		CustomSorting,
+		DynamicParam0,
+		DynamicParam1,
+		DynamicParam2,
+		DynamicParam3,
+
+		Num,
+	};
+};
+
 FNiagaraRendererMeshes::FNiagaraRendererMeshes(ERHIFeatureLevel::Type FeatureLevel, const UNiagaraRendererProperties *Props, const FNiagaraEmitterInstance* Emitter)
 	: FNiagaraRenderer(FeatureLevel, Props, Emitter)
-	, PositionOffset(INDEX_NONE)
-	, VelocityOffset(INDEX_NONE)
-	, ColorOffset(INDEX_NONE)
-	, ScaleOffset(INDEX_NONE)
-	, SizeOffset(INDEX_NONE)
 	, MaterialParamValidMask(0)
-	, MaterialParamOffset(INDEX_NONE)
-	, MaterialParamOffset1(INDEX_NONE)
-	, MaterialParamOffset2(INDEX_NONE)
-	, MaterialParamOffset3(INDEX_NONE)
-	, TransformOffset(INDEX_NONE)
-	, NormalizedAgeOffset(INDEX_NONE)
-	, MaterialRandomOffset(INDEX_NONE)
-	, CustomSortingOffset(INDEX_NONE)
 {
 	check(Emitter);
 	check(Props);
@@ -83,26 +93,23 @@ FNiagaraRendererMeshes::FNiagaraRendererMeshes(ERHIFeatureLevel::Type FeatureLev
 	bOverrideMaterials = Properties->bOverrideMaterials;
 
 	const FNiagaraDataSet& Data = Emitter->GetData();
-	int32 IntDummy;
-	SizeOffset = -1;
-	Data.GetVariableComponentOffsets(Properties->PositionBinding.DataSetVariable, PositionOffset, IntDummy);
-	Data.GetVariableComponentOffsets(Properties->VelocityBinding.DataSetVariable, VelocityOffset, IntDummy);
-	Data.GetVariableComponentOffsets(Properties->ColorBinding.DataSetVariable, ColorOffset, IntDummy);
-	Data.GetVariableComponentOffsets(Properties->ScaleBinding.DataSetVariable, ScaleOffset, IntDummy);
-	Data.GetVariableComponentOffsets(Properties->DynamicMaterialBinding.DataSetVariable, MaterialParamOffset, IntDummy);
-	Data.GetVariableComponentOffsets(Properties->DynamicMaterial1Binding.DataSetVariable, MaterialParamOffset1, IntDummy);
-	Data.GetVariableComponentOffsets(Properties->DynamicMaterial2Binding.DataSetVariable, MaterialParamOffset2, IntDummy);
-	Data.GetVariableComponentOffsets(Properties->DynamicMaterial3Binding.DataSetVariable, MaterialParamOffset3, IntDummy);
-	Data.GetVariableComponentOffsets(Properties->MeshOrientationBinding.DataSetVariable, TransformOffset, IntDummy);
-	Data.GetVariableComponentOffsets(Properties->NormalizedAgeBinding.DataSetVariable, NormalizedAgeOffset, IntDummy);
-	Data.GetVariableComponentOffsets(Properties->MaterialRandomBinding.DataSetVariable, MaterialRandomOffset, IntDummy);
-	Data.GetVariableComponentOffsets(Properties->CustomSortingBinding.DataSetVariable, CustomSortingOffset, IntDummy);
-
-	MaterialParamValidMask = MaterialParamOffset  == -1 ? 0 : 0x1;
-	MaterialParamValidMask |= MaterialParamOffset1 == -1 ? 0 : 0x2;
-	MaterialParamValidMask |= MaterialParamOffset2 == -1 ? 0 : 0x4;
-	MaterialParamValidMask |= MaterialParamOffset3 == -1 ? 0 : 0x8;
 
+	MaterialParamValidMask = 0;
+	TotalVFComponents = 0;
+	VFVariables.SetNum(ENiagaraMeshVFLayout::Num);
+	SetVertexFactoryVariable(Data, Properties->PositionBinding.DataSetVariable, ENiagaraMeshVFLayout::Position);
+	SetVertexFactoryVariable(Data, Properties->VelocityBinding.DataSetVariable, ENiagaraMeshVFLayout::Velocity);
+	SetVertexFactoryVariable(Data, Properties->ColorBinding.DataSetVariable, ENiagaraMeshVFLayout::Color);
+	SetVertexFactoryVariable(Data, Properties->ScaleBinding.DataSetVariable, ENiagaraMeshVFLayout::Scale);
+	SetVertexFactoryVariable(Data, Properties->MeshOrientationBinding.DataSetVariable, ENiagaraMeshVFLayout::Transform);
+	SetVertexFactoryVariable(Data, Properties->MaterialRandomBinding.DataSetVariable, ENiagaraMeshVFLayout::MaterialRandom);
+	SetVertexFactoryVariable(Data, Properties->NormalizedAgeBinding.DataSetVariable, ENiagaraMeshVFLayout::NormalizedAge);
+	SetVertexFactoryVariable(Data, Properties->CustomSortingBinding.DataSetVariable, ENiagaraMeshVFLayout::CustomSorting);
+	MaterialParamValidMask |= SetVertexFactoryVariable(Data, Properties->DynamicMaterialBinding.DataSetVariable, ENiagaraMeshVFLayout::DynamicParam0) ? 0x1 : 0;
+	MaterialParamValidMask |= SetVertexFactoryVariable(Data, Properties->DynamicMaterial1Binding.DataSetVariable, ENiagaraMeshVFLayout::DynamicParam1) ? 0x2 : 0;
+	MaterialParamValidMask |= SetVertexFactoryVariable(Data, Properties->DynamicMaterial2Binding.DataSetVariable, ENiagaraMeshVFLayout::DynamicParam2) ? 0x4 : 0;
+	MaterialParamValidMask |= SetVertexFactoryVariable(Data, Properties->DynamicMaterial3Binding.DataSetVariable, ENiagaraMeshVFLayout::DynamicParam3) ? 0x8 : 0;
+	
 	MeshMinimumLOD = Properties->ParticleMesh->MinLOD.GetValueForFeatureLevel(FeatureLevel);
 }
 
@@ -160,22 +167,46 @@ void FNiagaraRendererMeshes::GetDynamicMeshElements(const TArray<const FSceneVie
 	FScopeCycleCounter EmitterStatsCounter(EmitterStatID);
 #endif
 
+	// @TODO : support multiple LOD and section, using an inlined array and/or the SceneRenderingAllocator
+	uint32 IndirectArgsOffset = INDEX_NONE;
 	int32 NumInstances = SourceParticleData->GetNumInstances();
 
-	int32 TotalFloatSize = SourceParticleData->GetFloatBuffer().Num() / sizeof(float);
 	FGlobalDynamicReadBuffer& DynamicReadBuffer = Collector.GetDynamicReadBuffer();
 	FGlobalDynamicReadBuffer::FAllocation ParticleData;
 
-	// @TODO : support multiple LOD and section, using an inlined array and/or the SceneRenderingAllocator
-	uint32 IndirectArgsOffset = INDEX_NONE;
+	//Grab the material proxies we'll be using for each section and check them for translucency.
+	bool bHasTranslucentMaterials = false;
+	for (FMaterialRenderProxy* MaterialProxy : DynamicDataMesh->Materials)
+	{
+		check(MaterialProxy);
+		EBlendMode BlendMode = MaterialProxy->GetMaterial(FeatureLevel)->GetBlendMode();
+		bHasTranslucentMaterials |= BlendMode == BLEND_AlphaComposite || BlendMode == BLEND_AlphaHoldout || BlendMode == BLEND_Translucent;
+	}
+
+	bool bShouldSort = SortMode != ENiagaraSortMode::None && (bHasTranslucentMaterials || !bSortOnlyWhenTranslucent);
+	bool bCustomSorting = SortMode == ENiagaraSortMode::CustomAscending || SortMode == ENiagaraSortMode::CustomDecending;
+	//Disable the upload of sorting data if we're using a material that doesn't need it.
+	//TODO: we can probably reinit the GPU layout info entirely to remove custom sorting from the buffer but for now just skip the upload if it's not needed.
+	VFVariables[ENiagaraMeshVFLayout::CustomSorting].bUpload &= bCustomSorting;
+
 	//For cpu sims we allocate render buffers from the global pool. GPU sims own their own.
 	if (SimTarget == ENiagaraSimTarget::CPUSim)
 	{
-		SCOPE_CYCLE_COUNTER(STAT_NiagaraRenderMeshes_AllocateGPUData);
-		ParticleData = DynamicReadBuffer.AllocateFloat(TotalFloatSize);
-		FMemory::Memcpy(ParticleData.Buffer, SourceParticleData->GetFloatBuffer().GetData(), SourceParticleData->GetFloatBuffer().Num());
+		if (GbEnableMinimalGPUBuffers)
+		{
+			ParticleData = TransferDataToGPU(DynamicReadBuffer, SourceParticleData);
+		}
+		else
+		{
+			SCOPE_CYCLE_COUNTER(STAT_NiagaraRenderMeshes_AllocateGPUData);
+			int32 TotalFloatSize = SourceParticleData->GetFloatBuffer().Num() / sizeof(float);
+			ParticleData = DynamicReadBuffer.AllocateFloat(TotalFloatSize);
+			FMemory::Memcpy(ParticleData.Buffer, SourceParticleData->GetFloatBuffer().GetData(), SourceParticleData->GetFloatBuffer().Num());
+		}
 	}
-	else // ENiagaraSimTarget::GPUComputeSim
+	
+	
+	if(SimTarget == ENiagaraSimTarget::GPUComputeSim)
 	{
 		IndirectArgsOffset = Batcher->GetGPUInstanceCounterManager().AddDrawIndirect(SourceParticleData->GetGPUInstanceCountBufferOffset(), NumIndicesPerInstance);
 	}
@@ -200,53 +231,49 @@ void FNiagaraRendererMeshes::GetDynamicMeshElements(const TArray<const FSceneVie
 				FNiagaraMeshCollectorResourcesMesh& CollectorResources = Collector.AllocateOneFrameResource<FNiagaraMeshCollectorResourcesMesh>();
 				SetupVertexFactory(&CollectorResources.VertexFactory, LODModel);
 				FNiagaraMeshUniformParameters PerViewUniformParameters;// = UniformParameters;
-				FMemory::Memzero(&PerViewUniformParameters,sizeof(PerViewUniformParameters)); // Clear unset bytes
+				FMemory::Memzero(&PerViewUniformParameters, sizeof(PerViewUniformParameters)); // Clear unset bytes
 
 				PerViewUniformParameters.bLocalSpace = bLocalSpace;
 				PerViewUniformParameters.PrevTransformAvailable = false;
 				PerViewUniformParameters.DeltaSeconds = ViewFamily.DeltaWorldTime;
-				PerViewUniformParameters.PositionDataOffset = PositionOffset;
-				PerViewUniformParameters.VelocityDataOffset = VelocityOffset;
-				PerViewUniformParameters.ColorDataOffset = ColorOffset;
-				PerViewUniformParameters.TransformDataOffset = TransformOffset;
-				PerViewUniformParameters.ScaleDataOffset = ScaleOffset;
-				PerViewUniformParameters.SizeDataOffset = SizeOffset;
+
+				PerViewUniformParameters.PositionDataOffset = VFVariables[ENiagaraMeshVFLayout::Position].GetGPUOffset();
+				PerViewUniformParameters.VelocityDataOffset = VFVariables[ENiagaraMeshVFLayout::Velocity].GetGPUOffset();
+				PerViewUniformParameters.ColorDataOffset = VFVariables[ENiagaraMeshVFLayout::Color].GetGPUOffset();
+				PerViewUniformParameters.ScaleDataOffset = VFVariables[ENiagaraMeshVFLayout::Scale].GetGPUOffset();
+				PerViewUniformParameters.TransformDataOffset = VFVariables[ENiagaraMeshVFLayout::Transform].GetGPUOffset();
+				PerViewUniformParameters.NormalizedAgeDataOffset = VFVariables[ENiagaraMeshVFLayout::NormalizedAge].GetGPUOffset();
+				PerViewUniformParameters.MaterialRandomDataOffset = VFVariables[ENiagaraMeshVFLayout::MaterialRandom].GetGPUOffset();
+				PerViewUniformParameters.MaterialParamDataOffset = VFVariables[ENiagaraMeshVFLayout::DynamicParam0].GetGPUOffset();
+				PerViewUniformParameters.MaterialParam1DataOffset = VFVariables[ENiagaraMeshVFLayout::DynamicParam1].GetGPUOffset();
+				PerViewUniformParameters.MaterialParam2DataOffset = VFVariables[ENiagaraMeshVFLayout::DynamicParam2].GetGPUOffset();
+				PerViewUniformParameters.MaterialParam3DataOffset = VFVariables[ENiagaraMeshVFLayout::DynamicParam3].GetGPUOffset();
+
 				PerViewUniformParameters.MaterialParamValidMask = MaterialParamValidMask;
-				PerViewUniformParameters.MaterialParamDataOffset = MaterialParamOffset;
-				PerViewUniformParameters.MaterialParam1DataOffset = MaterialParamOffset1;
-				PerViewUniformParameters.MaterialParam2DataOffset = MaterialParamOffset2;
-				PerViewUniformParameters.MaterialParam3DataOffset = MaterialParamOffset3;
-				PerViewUniformParameters.NormalizedAgeDataOffset = NormalizedAgeOffset;
-				PerViewUniformParameters.MaterialRandomDataOffset = MaterialRandomOffset;
+				PerViewUniformParameters.SizeDataOffset = INDEX_NONE;
 				PerViewUniformParameters.DefaultPos = bLocalSpace ? FVector4(0.0f, 0.0f, 0.0f, 1.0f) : FVector4(SceneProxy->GetLocalToWorld().GetOrigin());
 
-				//Grab the material proxies we'll be using for each section and check them for translucency.
-				TArray<FMaterialRenderProxy*, TInlineAllocator<32>> MaterialProxies;
-				bool bHasTranslucentMaterials = false;
-				for (FMaterialRenderProxy* MaterialProxy : DynamicDataMesh->Materials)
-				{
-					check(MaterialProxy);
-					EBlendMode BlendMode = MaterialProxy->GetMaterial(View->FeatureLevel)->GetBlendMode();
-					bHasTranslucentMaterials |= BlendMode == BLEND_AlphaComposite || BlendMode == BLEND_AlphaHoldout || BlendMode == BLEND_Translucent;
-				}
-
 				//Sort particles if needed.
 				CollectorResources.VertexFactory.SetSortedIndices(nullptr, 0xFFFFFFFF);
 
 				FNiagaraGPUSortInfo SortInfo;
-				if (SortMode != ENiagaraSortMode::None && (bHasTranslucentMaterials || !bSortOnlyWhenTranslucent))
+				SortInfo.SortAttributeOffset = INDEX_NONE;
+				int32 SortVarIdx = INDEX_NONE;
+				if (bShouldSort)
 				{
 					SortInfo.ParticleCount = NumInstances;
 					SortInfo.SortMode = SortMode;
-					if (SortInfo.SortMode == ENiagaraSortMode::CustomAscending || SortInfo.SortMode == ENiagaraSortMode::CustomDecending)
+					if (bCustomSorting)
 					{
-						SortInfo.SortAttributeOffset = CustomSortingOffset;
+						SortVarIdx = ENiagaraMeshVFLayout::CustomSorting;
+						SortInfo.SortAttributeOffset = VFVariables[ENiagaraMeshVFLayout::CustomSorting].GetGPUOffset();
 						SortInfo.ViewOrigin.Set(0, 0, 0);
 						SortInfo.ViewDirection.Set(0, 0, 1);
 					}
 					else
 					{
-						SortInfo.SortAttributeOffset = PositionOffset;
+						SortVarIdx = ENiagaraMeshVFLayout::Position;
+						SortInfo.SortAttributeOffset = VFVariables[ENiagaraMeshVFLayout::Position].GetGPUOffset();
 						SortInfo.ViewOrigin = View->ViewMatrices.GetViewOrigin();
 						SortInfo.ViewDirection = View->GetViewDirection();
 						if (bLocalSpace)
@@ -280,11 +307,12 @@ void FNiagaraRendererMeshes::GetDynamicMeshElements(const TArray<const FSceneVie
 						{
 							FGlobalDynamicReadBuffer::FAllocation SortedIndices;
 							SortedIndices = DynamicReadBuffer.AllocateInt32(NumInstances);
-							SortIndices(SortInfo, *SourceParticleData, SortedIndices);
+							SortIndices(SortInfo, SortVarIdx, *SourceParticleData, SortedIndices);
 							CollectorResources.VertexFactory.SetSortedIndices(SortedIndices.ReadBuffer->SRV, SortedIndices.FirstIndex / sizeof(float));
 						}
 					}
-					CollectorResources.VertexFactory.SetParticleData(ParticleData.ReadBuffer->SRV, ParticleData.FirstIndex / sizeof(float), SourceParticleData->GetFloatStride() / sizeof(float));
+					int32 ParticleDataStride = GbEnableMinimalGPUBuffers ? SourceParticleData->GetNumInstances() : SourceParticleData->GetFloatStride() / sizeof(float);
+					CollectorResources.VertexFactory.SetParticleData(ParticleData.ReadBuffer->SRV, ParticleData.FirstIndex / sizeof(float), ParticleDataStride);
 				}
 				else
 				{
@@ -292,7 +320,6 @@ void FNiagaraRendererMeshes::GetDynamicMeshElements(const TArray<const FSceneVie
 					{
 						// Here we need to be conservative about the InstanceCount, since the final value is only known on the GPU after the simulation.
 						SortInfo.ParticleCount = SourceParticleData->GetNumInstances();
-
 						SortInfo.ParticleDataFloatSRV = SourceParticleData->GetGPUBufferFloat().SRV;
 						SortInfo.FloatDataOffset = 0;
 						SortInfo.FloatDataStride = SourceParticleData->GetFloatStride() / sizeof(float);
diff --git a/Engine/Plugins/FX/Niagara/Source/Niagara/Private/NiagaraRendererRibbons.cpp b/Engine/Plugins/FX/Niagara/Source/Niagara/Private/NiagaraRendererRibbons.cpp
--- a/Engine/Plugins/FX/Niagara/Source/Niagara/Private/NiagaraRendererRibbons.cpp
+++ b/Engine/Plugins/FX/Niagara/Source/Niagara/Private/NiagaraRendererRibbons.cpp
@@ -125,6 +125,26 @@ public:
 	}
 };
 
+namespace ENiagaraRibbonVFLayout
+{
+	enum Type
+	{
+		Position,
+		Velocity,
+		Color,
+		Width,
+		Twist,
+		Facing,
+		NormalizedAge,
+		MaterialRandom,
+		MaterialParam0,
+		MaterialParam1,
+		MaterialParam2,
+		MaterialParam3,
+		Num,
+	};
+};
+
 
 FNiagaraRendererRibbons::FNiagaraRendererRibbons(ERHIFeatureLevel::Type FeatureLevel, const UNiagaraRendererProperties *InProps, const FNiagaraEmitterInstance* Emitter)
 	: FNiagaraRenderer(FeatureLevel, InProps, Emitter)
@@ -141,14 +161,6 @@ FNiagaraRendererRibbons::FNiagaraRendererRibbons(ERHIFeatureLevel::Type FeatureL
 	, bCustomUseConstantFactor(false)
 	, CustomTessellationMinAngle(15.f * PI / 180.f)
 	, bCustomUseScreenSpace(true)
-	, PositionDataOffset(INDEX_NONE)
-	, VelocityDataOffset(INDEX_NONE)
-	, WidthDataOffset(INDEX_NONE)
-	, TwistDataOffset(INDEX_NONE)
-	, FacingDataOffset(INDEX_NONE)
-	, ColorDataOffset(INDEX_NONE)
-	, NormalizedAgeDataOffset(INDEX_NONE)
-	, MaterialRandomDataOffset(INDEX_NONE)
 {
 	const UNiagaraRibbonRendererProperties* Properties = CastChecked<const UNiagaraRibbonRendererProperties>(InProps);
 	const FNiagaraDataSet& Data = Emitter->GetData();
@@ -171,28 +183,25 @@ FNiagaraRendererRibbons::FNiagaraRendererRibbons(ERHIFeatureLevel::Type FeatureL
 	CustomTessellationMinAngle *= PI / 180.f;
 	bCustomUseScreenSpace = Properties->bScreenSpaceTessellation;
 
+	TotalVFComponents = 0;
+	VFVariables.SetNum(ENiagaraRibbonVFLayout::Num);
 	// required attributes
-	int32 IntDummy;
-	Data.GetVariableComponentOffsets(Properties->PositionBinding.DataSetVariable, PositionDataOffset, IntDummy);
-	Data.GetVariableComponentOffsets(Properties->VelocityBinding.DataSetVariable, VelocityDataOffset, IntDummy);
-	Data.GetVariableComponentOffsets(Properties->ColorBinding.DataSetVariable, ColorDataOffset, IntDummy);
+	SetVertexFactoryVariable(Data, Properties->PositionBinding.DataSetVariable, ENiagaraRibbonVFLayout::Position);
+	SetVertexFactoryVariable(Data, Properties->VelocityBinding.DataSetVariable, ENiagaraRibbonVFLayout::Velocity);
+	SetVertexFactoryVariable(Data, Properties->ColorBinding.DataSetVariable, ENiagaraRibbonVFLayout::Color);
 
 	// optional attributes
-	Data.GetVariableComponentOffsets(Properties->RibbonWidthBinding.DataSetVariable, WidthDataOffset, IntDummy);
-	Data.GetVariableComponentOffsets(Properties->RibbonTwistBinding.DataSetVariable, TwistDataOffset, IntDummy);
-	Data.GetVariableComponentOffsets(Properties->RibbonFacingBinding.DataSetVariable, FacingDataOffset, IntDummy);
-	Data.GetVariableComponentOffsets(Properties->NormalizedAgeBinding.DataSetVariable, NormalizedAgeDataOffset, IntDummy);
-	Data.GetVariableComponentOffsets(Properties->MaterialRandomBinding.DataSetVariable, MaterialRandomDataOffset, IntDummy);
-
-	Data.GetVariableComponentOffsets(Properties->DynamicMaterialBinding.DataSetVariable, MaterialParamOffset, IntDummy);
-	Data.GetVariableComponentOffsets(Properties->DynamicMaterial1Binding.DataSetVariable, MaterialParamOffset1, IntDummy);
-	Data.GetVariableComponentOffsets(Properties->DynamicMaterial2Binding.DataSetVariable, MaterialParamOffset2, IntDummy);
-	Data.GetVariableComponentOffsets(Properties->DynamicMaterial3Binding.DataSetVariable, MaterialParamOffset3, IntDummy);
-
-	MaterialParamValidMask = MaterialParamOffset != -1 ? 1 : 0;
-	MaterialParamValidMask |= MaterialParamOffset1 != -1 ? 2 : 0;
-	MaterialParamValidMask |= MaterialParamOffset2 != -1 ? 4 : 0;
-	MaterialParamValidMask |= MaterialParamOffset3 != -1 ? 8 : 0;
+	SetVertexFactoryVariable(Data, Properties->RibbonWidthBinding.DataSetVariable, ENiagaraRibbonVFLayout::Width);
+	SetVertexFactoryVariable(Data, Properties->RibbonTwistBinding.DataSetVariable, ENiagaraRibbonVFLayout::Twist);
+	SetVertexFactoryVariable(Data, Properties->RibbonFacingBinding.DataSetVariable, ENiagaraRibbonVFLayout::Facing);
+	SetVertexFactoryVariable(Data, Properties->NormalizedAgeBinding.DataSetVariable, ENiagaraRibbonVFLayout::NormalizedAge);
+	SetVertexFactoryVariable(Data, Properties->MaterialRandomBinding.DataSetVariable, ENiagaraRibbonVFLayout::MaterialRandom);
+
+	MaterialParamValidMask = 0;
+	MaterialParamValidMask |= SetVertexFactoryVariable(Data, Properties->DynamicMaterialBinding.DataSetVariable, ENiagaraRibbonVFLayout::MaterialParam0) ? 1 : 0;
+	MaterialParamValidMask |= SetVertexFactoryVariable(Data, Properties->DynamicMaterial1Binding.DataSetVariable, ENiagaraRibbonVFLayout::MaterialParam1) ? 2 : 0;
+	MaterialParamValidMask |= SetVertexFactoryVariable(Data, Properties->DynamicMaterial2Binding.DataSetVariable, ENiagaraRibbonVFLayout::MaterialParam2) ? 4 : 0;
+	MaterialParamValidMask |= SetVertexFactoryVariable(Data, Properties->DynamicMaterial3Binding.DataSetVariable, ENiagaraRibbonVFLayout::MaterialParam3) ? 8 : 0;
 }
 
 FNiagaraRendererRibbons::~FNiagaraRendererRibbons()
@@ -750,7 +759,9 @@ void FNiagaraRendererRibbons::SetupMeshBatchAndCollectorResourceForView(
 
 	FCPUSimParticleDataAllocation CPUSimParticleDataAllocation = AllocateParticleDataIfCPUSim(DynamicDataRibbon, Collector.GetDynamicReadBuffer());
 	auto& ParticleData = CPUSimParticleDataAllocation.ParticleData;
-	CollectorResources.VertexFactory.SetParticleData(ParticleData.ReadBuffer->SRV, ParticleData.FirstIndex / sizeof(float), SourceParticleData->GetFloatStride() / sizeof(float));
+
+	int32 ParticleDataStride = GbEnableMinimalGPUBuffers ? SourceParticleData->GetNumInstances() : SourceParticleData->GetFloatStride() / sizeof(float);
+	CollectorResources.VertexFactory.SetParticleData(ParticleData.ReadBuffer->SRV, ParticleData.FirstIndex / sizeof(float), ParticleDataStride);
 
 	// TODO: need to make these two a global alloc buffer as well, not recreate
 				// pass in the sorted indices so the VS can fetch the particle data in order
@@ -845,13 +856,20 @@ FNiagaraRendererRibbons::FCPUSimParticleDataAllocation FNiagaraRendererRibbons::
 
 	FCPUSimParticleDataAllocation CPUSimParticleDataAllocation{ DynamicReadBuffer };
 
+	bool bShouldDoFacing = FacingMode == ENiagaraRibbonFacingMode::Custom || FacingMode == ENiagaraRibbonFacingMode::CustomSideVector;
+	VFVariables[ENiagaraRibbonVFLayout::Facing].bUpload = bShouldDoFacing;
+
 	if (SimTarget == ENiagaraSimTarget::CPUSim)
 	{
-		SCOPE_CYCLE_COUNTER(STAT_NiagaraRenderRibbonsCPUSimCopy);
-		CPUSimParticleDataAllocation.ParticleData = DynamicReadBuffer.AllocateFloat(TotalFloatSize);
+		SCOPE_CYCLE_COUNTER(STAT_NiagaraRenderRibbonsCPUSimMemCopy);
+		if (GbEnableMinimalGPUBuffers)
+		{
+			CPUSimParticleDataAllocation.ParticleData = TransferDataToGPU(DynamicReadBuffer, SourceParticleData);
+		}
+		else
 		{
-			SCOPE_CYCLE_COUNTER(STAT_NiagaraRenderRibbonsCPUSimMemCopy);
-			FMemory::Memcpy(CPUSimParticleDataAllocation.ParticleData.Buffer, SourceParticleData->GetFloatBuffer().GetData(), SourceParticleData->GetFloatBuffer().Num());
+			CPUSimParticleDataAllocation.ParticleData = DynamicReadBuffer.AllocateFloat(TotalFloatSize);
+			FMemory::Memcpy(ParticleData.Buffer, SourceParticleData->GetFloatBuffer().GetData(), SourceParticleData->GetFloatBuffer().Num());
 		}
 	}
 
@@ -953,19 +971,26 @@ void FNiagaraRendererRibbons::CreatePerViewResources(
 	PerViewUniformParameters.InterpCount = SegmentTessellation;
 	PerViewUniformParameters.OneOverInterpCount = 1.f / (float)SegmentTessellation;
 
-	PerViewUniformParameters.PositionDataOffset = PositionDataOffset;
-	PerViewUniformParameters.VelocityDataOffset = VelocityDataOffset;
-	PerViewUniformParameters.ColorDataOffset = ColorDataOffset;
-	PerViewUniformParameters.WidthDataOffset = WidthDataOffset;
-	PerViewUniformParameters.TwistDataOffset = TwistDataOffset;
-	PerViewUniformParameters.FacingDataOffset = FacingMode == ENiagaraRibbonFacingMode::Custom || FacingMode == ENiagaraRibbonFacingMode::CustomSideVector ? FacingDataOffset : -1;
-	PerViewUniformParameters.NormalizedAgeDataOffset = NormalizedAgeDataOffset;
-	PerViewUniformParameters.MaterialRandomDataOffset = MaterialRandomDataOffset;
+	PerViewUniformParameters.PositionDataOffset = VFVariables[ENiagaraRibbonVFLayout::Position].GetGPUOffset();
+	PerViewUniformParameters.VelocityDataOffset = VFVariables[ENiagaraRibbonVFLayout::Velocity].GetGPUOffset();
+	PerViewUniformParameters.ColorDataOffset = VFVariables[ENiagaraRibbonVFLayout::Color].GetGPUOffset();
+	PerViewUniformParameters.WidthDataOffset = VFVariables[ENiagaraRibbonVFLayout::Width].GetGPUOffset();
+	PerViewUniformParameters.TwistDataOffset = VFVariables[ENiagaraRibbonVFLayout::Twist].GetGPUOffset();
+	PerViewUniformParameters.NormalizedAgeDataOffset = VFVariables[ENiagaraRibbonVFLayout::NormalizedAge].GetGPUOffset();
+	PerViewUniformParameters.MaterialRandomDataOffset = VFVariables[ENiagaraRibbonVFLayout::MaterialRandom].GetGPUOffset();
+	PerViewUniformParameters.MaterialParamDataOffset = VFVariables[ENiagaraRibbonVFLayout::MaterialParam0].GetGPUOffset();
+	PerViewUniformParameters.MaterialParam1DataOffset = VFVariables[ENiagaraRibbonVFLayout::MaterialParam1].GetGPUOffset();
+	PerViewUniformParameters.MaterialParam2DataOffset = VFVariables[ENiagaraRibbonVFLayout::MaterialParam2].GetGPUOffset();
+	PerViewUniformParameters.MaterialParam3DataOffset = VFVariables[ENiagaraRibbonVFLayout::MaterialParam3].GetGPUOffset();
+
 	PerViewUniformParameters.MaterialParamValidMask = MaterialParamValidMask;
-	PerViewUniformParameters.MaterialParamDataOffset = MaterialParamOffset;
-	PerViewUniformParameters.MaterialParam1DataOffset = MaterialParamOffset1;
-	PerViewUniformParameters.MaterialParam2DataOffset = MaterialParamOffset2;
-	PerViewUniformParameters.MaterialParam3DataOffset = MaterialParamOffset3;
+
+	bool bShouldDoFacing = FacingMode == ENiagaraRibbonFacingMode::Custom || FacingMode == ENiagaraRibbonFacingMode::CustomSideVector;
+	PerViewUniformParameters.FacingDataOffset = bShouldDoFacing ? VFVariables[ENiagaraRibbonVFLayout::Facing].GetGPUOffset() : -1;
+	PerViewUniformParameters.OneOverUV0TilingDistance = UV0TilingDistance ? 1.f / (UV0TilingDistance) : 0.f;
+	PerViewUniformParameters.OneOverUV1TilingDistance = UV1TilingDistance ? 1.f / (UV1TilingDistance) : 0.f;
+	PerViewUniformParameters.PackedVData = FVector4(UV0Scale.Y, UV0Offset.Y, UV1Scale.Y, UV1Offset.Y);
+
 	PerViewUniformParameters.OneOverUV0TilingDistance = UV0TilingDistance ? 1.f / (UV0TilingDistance) : 0.f;
 	PerViewUniformParameters.OneOverUV1TilingDistance = UV1TilingDistance ? 1.f / (UV1TilingDistance) : 0.f;
 	PerViewUniformParameters.PackedVData = FVector4(UV0Scale.Y, UV0Offset.Y, UV1Scale.Y, UV1Offset.Y);
diff --git a/Engine/Plugins/FX/Niagara/Source/Niagara/Private/NiagaraRendererSprites.cpp b/Engine/Plugins/FX/Niagara/Source/Niagara/Private/NiagaraRendererSprites.cpp
--- a/Engine/Plugins/FX/Niagara/Source/Niagara/Private/NiagaraRendererSprites.cpp
+++ b/Engine/Plugins/FX/Niagara/Source/Niagara/Private/NiagaraRendererSprites.cpp
@@ -31,6 +31,15 @@ static FAutoConsoleVariableRef CVarEnableNiagaraSpriteRendering(
 	ECVF_Default
 );
 
+int32 GbEnableMinimalGPUBuffers = 1;
+static FAutoConsoleVariableRef CVarbEnableMinimalGPUBuffers(
+	TEXT("fx.EnableMinimalGPUBuffers"),
+	GbEnableMinimalGPUBuffers,
+	TEXT("If > 0 we use new code to pass the gpu only data the VF actuially uses for redering, rather than the whole partilce buffer. \n"),
+	ECVF_Default
+);
+
+
 /** Dynamic data for sprite renderers. */
 struct FNiagaraDynamicDataSprites : public FNiagaraDynamicDataBase
 {
@@ -59,6 +68,32 @@ public:
 
 //////////////////////////////////////////////////////////////////////////
 
+namespace ENiagaraSpriteVFLayout
+{
+	enum Type
+	{
+		Position, 
+		Color, 
+		Velocity, 
+		Rotation, 
+		Size, 
+		Facing, 
+		Alignment, 
+		SubImage, 
+		MaterialParam0, 
+		MaterialParam1,
+		MaterialParam2,
+		MaterialParam3,
+		CameraOffset, 
+		UVScale, 
+		MaterialRandom, 
+		CustomSorting, 
+		NormalizedAge, 
+
+		Num,
+	};
+};
+
 FNiagaraRendererSprites::FNiagaraRendererSprites(ERHIFeatureLevel::Type FeatureLevel, const UNiagaraRendererProperties *InProps, const FNiagaraEmitterInstance* Emitter)
 	: FNiagaraRenderer(FeatureLevel, InProps, Emitter)
 	, Alignment(ENiagaraSpriteAlignment::Unaligned)
@@ -72,24 +107,7 @@ FNiagaraRendererSprites::FNiagaraRendererSprites(ERHIFeatureLevel::Type FeatureL
 	, bSortOnlyWhenTranslucent(true)
 	, MinFacingCameraBlendDistance(0.0f)
 	, MaxFacingCameraBlendDistance(0.0f)
-	, PositionOffset(INDEX_NONE)
-	, ColorOffset(INDEX_NONE)
-	, VelocityOffset(INDEX_NONE)
-	, RotationOffset(INDEX_NONE)
-	, SizeOffset(INDEX_NONE)
-	, FacingOffset(INDEX_NONE)
-	, AlignmentOffset(INDEX_NONE)
-	, SubImageOffset(INDEX_NONE)
 	, MaterialParamValidMask(0)
-	, MaterialParamOffset(INDEX_NONE)
-	, MaterialParamOffset1(INDEX_NONE)
-	, MaterialParamOffset2(INDEX_NONE)
-	, MaterialParamOffset3(INDEX_NONE)
-	, CameraOffsetOffset(INDEX_NONE)
-	, UVScaleOffset(INDEX_NONE)
-	, MaterialRandomOffset(INDEX_NONE)
-	, CustomSortingOffset(INDEX_NONE)
-	, NormalizedAgeOffset(INDEX_NONE)
 {
 	check(InProps && Emitter);
 
@@ -110,30 +128,27 @@ FNiagaraRendererSprites::FNiagaraRendererSprites(ERHIFeatureLevel::Type FeatureL
 	NumCutoutVertexPerSubImage = Properties->GetNumCutoutVertexPerSubimage();
 	CutoutVertexBuffer.Data = Properties->GetCutoutData();
 
-	int32 IntDummy;
 	const FNiagaraDataSet& Data = Emitter->GetData();
-	Data.GetVariableComponentOffsets(Properties->PositionBinding.DataSetVariable, PositionOffset, IntDummy);
-	Data.GetVariableComponentOffsets(Properties->VelocityBinding.DataSetVariable, VelocityOffset, IntDummy);
-	Data.GetVariableComponentOffsets(Properties->SpriteRotationBinding.DataSetVariable, RotationOffset, IntDummy);
-	Data.GetVariableComponentOffsets(Properties->SpriteSizeBinding.DataSetVariable, SizeOffset, IntDummy);
-	Data.GetVariableComponentOffsets(Properties->ColorBinding.DataSetVariable, ColorOffset, IntDummy);
-	Data.GetVariableComponentOffsets(Properties->SpriteFacingBinding.DataSetVariable, FacingOffset, IntDummy);
-	Data.GetVariableComponentOffsets(Properties->SpriteAlignmentBinding.DataSetVariable, AlignmentOffset, IntDummy);
-	Data.GetVariableComponentOffsets(Properties->SubImageIndexBinding.DataSetVariable, SubImageOffset, IntDummy);
-	Data.GetVariableComponentOffsets(Properties->DynamicMaterialBinding.DataSetVariable, MaterialParamOffset, IntDummy);
-	Data.GetVariableComponentOffsets(Properties->DynamicMaterial1Binding.DataSetVariable, MaterialParamOffset1, IntDummy);
-	Data.GetVariableComponentOffsets(Properties->DynamicMaterial2Binding.DataSetVariable, MaterialParamOffset2, IntDummy);
-	Data.GetVariableComponentOffsets(Properties->DynamicMaterial3Binding.DataSetVariable, MaterialParamOffset3, IntDummy);
-	Data.GetVariableComponentOffsets(Properties->CameraOffsetBinding.DataSetVariable, CameraOffsetOffset, IntDummy);
-	Data.GetVariableComponentOffsets(Properties->UVScaleBinding.DataSetVariable, UVScaleOffset, IntDummy);
-	Data.GetVariableComponentOffsets(Properties->NormalizedAgeBinding.DataSetVariable, NormalizedAgeOffset, IntDummy);
-	Data.GetVariableComponentOffsets(Properties->MaterialRandomBinding.DataSetVariable, MaterialRandomOffset, IntDummy);
-	Data.GetVariableComponentOffsets(Properties->CustomSortingBinding.DataSetVariable, CustomSortingOffset, IntDummy);
-
-	MaterialParamValidMask = MaterialParamOffset != -1 ? 1 : 0;
-	MaterialParamValidMask |= MaterialParamOffset1 != -1 ? 2 : 0;
-	MaterialParamValidMask |= MaterialParamOffset2 != -1 ? 4 : 0;
-	MaterialParamValidMask |= MaterialParamOffset3 != -1 ? 8 : 0;
+
+	TotalVFComponents = 0;
+	VFVariables.SetNum(ENiagaraSpriteVFLayout::Num);
+	SetVertexFactoryVariable(Data, Properties->PositionBinding.DataSetVariable, ENiagaraSpriteVFLayout::Position);
+	SetVertexFactoryVariable(Data, Properties->VelocityBinding.DataSetVariable, ENiagaraSpriteVFLayout::Velocity);
+	SetVertexFactoryVariable(Data, Properties->ColorBinding.DataSetVariable, ENiagaraSpriteVFLayout::Color);
+	SetVertexFactoryVariable(Data, Properties->SpriteRotationBinding.DataSetVariable, ENiagaraSpriteVFLayout::Rotation);
+	SetVertexFactoryVariable(Data, Properties->SpriteSizeBinding.DataSetVariable, ENiagaraSpriteVFLayout::Size);
+	SetVertexFactoryVariable(Data, Properties->SpriteFacingBinding.DataSetVariable, ENiagaraSpriteVFLayout::Facing);
+	SetVertexFactoryVariable(Data, Properties->SpriteAlignmentBinding.DataSetVariable, ENiagaraSpriteVFLayout::Alignment);
+	SetVertexFactoryVariable(Data, Properties->SubImageIndexBinding.DataSetVariable, ENiagaraSpriteVFLayout::SubImage);
+	SetVertexFactoryVariable(Data, Properties->CameraOffsetBinding.DataSetVariable, ENiagaraSpriteVFLayout::CameraOffset);
+	SetVertexFactoryVariable(Data, Properties->UVScaleBinding.DataSetVariable, ENiagaraSpriteVFLayout::UVScale);
+	SetVertexFactoryVariable(Data, Properties->NormalizedAgeBinding.DataSetVariable, ENiagaraSpriteVFLayout::NormalizedAge);
+	SetVertexFactoryVariable(Data, Properties->MaterialRandomBinding.DataSetVariable, ENiagaraSpriteVFLayout::MaterialRandom);
+	SetVertexFactoryVariable(Data, Properties->CustomSortingBinding.DataSetVariable, ENiagaraSpriteVFLayout::CustomSorting);
+	MaterialParamValidMask |= SetVertexFactoryVariable(Data, Properties->DynamicMaterialBinding.DataSetVariable, ENiagaraSpriteVFLayout::MaterialParam0) ? 0x1 : 0;
+	MaterialParamValidMask |= SetVertexFactoryVariable(Data, Properties->DynamicMaterial1Binding.DataSetVariable, ENiagaraSpriteVFLayout::MaterialParam1) ? 0x2 : 0;
+	MaterialParamValidMask |= SetVertexFactoryVariable(Data, Properties->DynamicMaterial2Binding.DataSetVariable, ENiagaraSpriteVFLayout::MaterialParam2) ? 0x4 : 0;
+	MaterialParamValidMask |= SetVertexFactoryVariable(Data, Properties->DynamicMaterial3Binding.DataSetVariable, ENiagaraSpriteVFLayout::MaterialParam3) ? 0x8 : 0;
 }
 
 FNiagaraRendererSprites::~FNiagaraRendererSprites()
@@ -184,16 +199,23 @@ FNiagaraRendererSprites::FCPUSimParticleDataAllocation FNiagaraRendererSprites::
 {
 	FNiagaraDataBuffer* SourceParticleData = DynamicDataSprites->GetParticleDataToRender();
 	check(SourceParticleData);//Can be null but should be checked before here.
-	int32 TotalFloatSize = SourceParticleData->GetFloatBuffer().Num() / sizeof(float);
 
 	FCPUSimParticleDataAllocation CPUSimParticleDataAllocation { DynamicReadBuffer };
 
 	if (SimTarget == ENiagaraSimTarget::CPUSim)
 	{
 		SCOPE_CYCLE_COUNTER(STAT_NiagaraRenderSpritesCPUSimCopy);
-		CPUSimParticleDataAllocation.ParticleData = DynamicReadBuffer.AllocateFloat(TotalFloatSize);
+
+		if (GbEnableMinimalGPUBuffers)
 		{
 			SCOPE_CYCLE_COUNTER(STAT_NiagaraRenderSpritesCPUSimMemCopy);
+			CPUSimParticleDataAllocation.ParticleData = TransferDataToGPU(DynamicReadBuffer, SourceParticleData);
+		}
+		else
+		{
+			int32 TotalFloatSize = SourceParticleData->GetFloatBuffer().Num() / sizeof(float);
+			CPUSimParticleDataAllocation.ParticleData = DynamicReadBuffer.AllocateFloat(TotalFloatSize);
+			SCOPE_CYCLE_COUNTER(STAT_NiagaraRenderSpritesCPUSimMemCopy);			
 			FMemory::Memcpy(CPUSimParticleDataAllocation.ParticleData.Buffer, SourceParticleData->GetFloatBuffer().GetData(), SourceParticleData->GetFloatBuffer().Num());
 		}
 	}
@@ -220,24 +242,26 @@ FNiagaraSpriteUniformBufferRef FNiagaraRendererSprites::CreatePerViewUniformBuff
 	PerViewUniformParameters.RemoveHMDRoll = bRemoveHMDRollInVR;
 	PerViewUniformParameters.CustomFacingVectorMask = FVector4(0.0f, 0.0f, 0.0f, 0.0f);
 	PerViewUniformParameters.SubImageSize = FVector4(SubImageSize.X, SubImageSize.Y, 1.0f / SubImageSize.X, 1.0f / SubImageSize.Y);
-	PerViewUniformParameters.PositionDataOffset = PositionOffset;
-	PerViewUniformParameters.VelocityDataOffset = VelocityOffset;
-	PerViewUniformParameters.RotationDataOffset = RotationOffset;
-	PerViewUniformParameters.SizeDataOffset = SizeOffset;
-	PerViewUniformParameters.ColorDataOffset = ColorOffset;
-	PerViewUniformParameters.MaterialParamValidMask = MaterialParamValidMask;
-	PerViewUniformParameters.MaterialParamDataOffset = MaterialParamOffset;
-	PerViewUniformParameters.MaterialParam1DataOffset = MaterialParamOffset1;
-	PerViewUniformParameters.MaterialParam2DataOffset = MaterialParamOffset2;
-	PerViewUniformParameters.MaterialParam3DataOffset = MaterialParamOffset3;
-	PerViewUniformParameters.SubimageDataOffset = SubImageOffset;
-	PerViewUniformParameters.FacingDataOffset = FacingOffset;
-	PerViewUniformParameters.AlignmentDataOffset = AlignmentOffset;
+
+	PerViewUniformParameters.PositionDataOffset = VFVariables[ENiagaraSpriteVFLayout::Position].GetGPUOffset();
+	PerViewUniformParameters.VelocityDataOffset = VFVariables[ENiagaraSpriteVFLayout::Velocity].GetGPUOffset();
+	PerViewUniformParameters.RotationDataOffset = VFVariables[ENiagaraSpriteVFLayout::Rotation].GetGPUOffset();
+	PerViewUniformParameters.SizeDataOffset = VFVariables[ENiagaraSpriteVFLayout::Size].GetGPUOffset();
+	PerViewUniformParameters.ColorDataOffset = VFVariables[ENiagaraSpriteVFLayout::Color].GetGPUOffset();
+	PerViewUniformParameters.MaterialParamDataOffset = VFVariables[ENiagaraSpriteVFLayout::MaterialParam0].GetGPUOffset();
+	PerViewUniformParameters.MaterialParam1DataOffset = VFVariables[ENiagaraSpriteVFLayout::MaterialParam1].GetGPUOffset();
+	PerViewUniformParameters.MaterialParam2DataOffset = VFVariables[ENiagaraSpriteVFLayout::MaterialParam2].GetGPUOffset();
+	PerViewUniformParameters.MaterialParam3DataOffset = VFVariables[ENiagaraSpriteVFLayout::MaterialParam3].GetGPUOffset();
+	PerViewUniformParameters.SubimageDataOffset = VFVariables[ENiagaraSpriteVFLayout::SubImage].GetGPUOffset();
+	PerViewUniformParameters.FacingDataOffset = VFVariables[ENiagaraSpriteVFLayout::Facing].GetGPUOffset();
+	PerViewUniformParameters.AlignmentDataOffset = VFVariables[ENiagaraSpriteVFLayout::Alignment].GetGPUOffset();
+	PerViewUniformParameters.CameraOffsetDataOffset = VFVariables[ENiagaraSpriteVFLayout::CameraOffset].GetGPUOffset();
+	PerViewUniformParameters.UVScaleDataOffset = VFVariables[ENiagaraSpriteVFLayout::UVScale].GetGPUOffset();
+	PerViewUniformParameters.NormalizedAgeDataOffset = VFVariables[ENiagaraSpriteVFLayout::NormalizedAge].GetGPUOffset();
+	PerViewUniformParameters.MaterialRandomDataOffset = VFVariables[ENiagaraSpriteVFLayout::MaterialRandom].GetGPUOffset();
+
 	PerViewUniformParameters.SubImageBlendMode = bSubImageBlend;
-	PerViewUniformParameters.CameraOffsetDataOffset = CameraOffsetOffset;
-	PerViewUniformParameters.UVScaleDataOffset = UVScaleOffset;
-	PerViewUniformParameters.NormalizedAgeDataOffset = NormalizedAgeOffset;
-	PerViewUniformParameters.MaterialRandomDataOffset = MaterialRandomOffset;
+	PerViewUniformParameters.MaterialParamValidMask = MaterialParamValidMask;
 	PerViewUniformParameters.DefaultPos = bLocalSpace ? FVector4(0.0f, 0.0f, 0.0f, 1.0f) : FVector4(SceneProxy->GetLocalToWorld().GetOrigin());
 
 	ENiagaraSpriteFacingMode ActualFacingMode = FacingMode;
@@ -316,24 +340,27 @@ void FNiagaraRendererSprites::SetVertexFactoryParticleData(
 
 		FMaterialRenderProxy* MaterialRenderProxy = DynamicDataSprites->Material;
 		check(MaterialRenderProxy);
-		EBlendMode BlendMode = MaterialRenderProxy->GetMaterial(View->FeatureLevel)->GetBlendMode();
+		EBlendMode BlendMode = MaterialRenderProxy->GetMaterial(FeatureLevel)->GetBlendMode();
 		OutVertexFactory.SetSortedIndices(nullptr, 0xFFFFFFFF);
 
 		int32 NumInstances = SourceParticleData->GetNumInstances();
 		FNiagaraGPUSortInfo SortInfo;
-		if (SortMode != ENiagaraSortMode::None && (BlendMode == BLEND_AlphaComposite || BlendMode == BLEND_AlphaHoldout|| BlendMode == BLEND_Translucent || !bSortOnlyWhenTranslucent))
+		SortInfo.SortAttributeOffset = INDEX_NONE;
+		int32 SortingVarIdx = INDEX_NONE;
+		bool bShouldSort = SortMode != ENiagaraSortMode::None && (BlendMode == BLEND_AlphaComposite || BlendMode == BLEND_AlphaHoldout || BlendMode == BLEND_Translucent || !bSortOnlyWhenTranslucent);
+		if (bShouldSort)
 		{
 			SortInfo.ParticleCount = NumInstances;
 			SortInfo.SortMode = SortMode;
 			if (SortInfo.SortMode == ENiagaraSortMode::CustomAscending || SortInfo.SortMode == ENiagaraSortMode::CustomDecending)
 			{
-				SortInfo.SortAttributeOffset = CustomSortingOffset;
+				SortingVarIdx = ENiagaraSpriteVFLayout::CustomSorting;	
 				SortInfo.ViewOrigin.Set(0, 0, 0);
 				SortInfo.ViewDirection.Set(0, 0, 1);
 			}
 			else
 			{
-				SortInfo.SortAttributeOffset = PositionOffset;
+				SortingVarIdx = ENiagaraSpriteVFLayout::Position;
 				SortInfo.ViewOrigin = View->ViewMatrices.GetViewOrigin();
 				SortInfo.ViewDirection = View->GetViewDirection();
 				if (bLocalSpace)
@@ -342,10 +369,14 @@ void FNiagaraRendererSprites::SetVertexFactoryParticleData(
 					SortInfo.ViewDirection = SceneProxy->GetLocalToWorld().GetTransposed().TransformVector(SortInfo.ViewDirection);
 				}
 			}
-		};
+
+			SortInfo.SortAttributeOffset = VFVariables[SortingVarIdx].GetGPUOffset();
+		}
+
 
 		if (SimTarget == ENiagaraSimTarget::CPUSim)//TODO: Compute shader for sorting gpu sims and larger cpu sims.
 		{
+			int32 ParticleStrideInFloats = GbEnableMinimalGPUBuffers ? SourceParticleData->GetNumInstances() : SourceParticleData->GetFloatStride() / sizeof(float);
 			check(CPUSimParticleDataAllocation.ParticleData.IsValid());
 			if (SortInfo.SortMode != ENiagaraSortMode::None && SortInfo.SortAttributeOffset != INDEX_NONE)
 			{
@@ -356,7 +387,7 @@ void FNiagaraRendererSprites::SetVertexFactoryParticleData(
 					SortInfo.ParticleCount = NumInstances;
 					SortInfo.ParticleDataFloatSRV = CPUSimParticleDataAllocation.ParticleData.ReadBuffer->SRV;
 					SortInfo.FloatDataOffset = CPUSimParticleDataAllocation.ParticleData.FirstIndex / sizeof(float);
-					SortInfo.FloatDataStride = SourceParticleData->GetFloatStride() / sizeof(float);
+					SortInfo.FloatDataStride = ParticleStrideInFloats;
 					const int32 IndexBufferOffset = Batcher->AddSortedGPUSimulation(SortInfo);
 					if (IndexBufferOffset != INDEX_NONE)
 					{
@@ -369,11 +400,11 @@ void FNiagaraRendererSprites::SetVertexFactoryParticleData(
 
 					FGlobalDynamicReadBuffer::FAllocation SortedIndices;
 					SortedIndices = CPUSimParticleDataAllocation.DynamicReadBuffer.AllocateInt32(NumInstances);
-					SortIndices(SortInfo, *SourceParticleData, SortedIndices);
+					SortIndices(SortInfo, SortingVarIdx, *SourceParticleData, SortedIndices);
 					OutVertexFactory.SetSortedIndices(SortedIndices.ReadBuffer->SRV, SortedIndices.FirstIndex / sizeof(float));
 				}
 			}
-			OutVertexFactory.SetParticleData(CPUSimParticleDataAllocation.ParticleData.ReadBuffer->SRV, CPUSimParticleDataAllocation.ParticleData.FirstIndex / sizeof(float), SourceParticleData->GetFloatStride() / sizeof(float));
+			OutVertexFactory.SetParticleData(CPUSimParticleDataAllocation.ParticleData.ReadBuffer->SRV, CPUSimParticleDataAllocation.ParticleData.FirstIndex / sizeof(float), ParticleStrideInFloats);
 		}
 		else // ENiagaraSimTarget::GPUSim
 		{
@@ -541,7 +572,15 @@ void FNiagaraRendererSprites::GetDynamicMeshElements(const TArray<const FSceneVi
 #if STATS
 	FScopeCycleCounter EmitterStatsCounter(EmitterStatID);
 #endif
-	
+	FMaterialRenderProxy* MaterialRenderProxy = DynamicDataSprites->Material;
+	check(MaterialRenderProxy);
+	EBlendMode BlendMode = MaterialRenderProxy->GetMaterial(FeatureLevel)->GetBlendMode();
+	bool bShouldSort = SortMode != ENiagaraSortMode::None && (BlendMode == BLEND_AlphaComposite || BlendMode == BLEND_AlphaHoldout || BlendMode == BLEND_Translucent || !bSortOnlyWhenTranslucent);
+	bool bNeedCustomSort = bShouldSort && (SortMode == ENiagaraSortMode::CustomAscending || SortMode == ENiagaraSortMode::CustomDecending);
+	//Disable the upload of sorting data if we're using a material that doesn't need it.
+	//TODO: we can probably reinit the GPU layout info entirely to remove custom sorting from the buffer but for now just skip the upload if it's not needed.
+	VFVariables[ENiagaraSpriteVFLayout::CustomSorting].bUpload &= bNeedCustomSort;
+
 	FCPUSimParticleDataAllocation CPUSimParticleDataAllocation = ConditionalAllocateCPUSimParticleData(DynamicDataSprites, Collector.GetDynamicReadBuffer());
 
 	uint32 IndirectArgsOffset = INDEX_NONE;
diff --git a/Engine/Plugins/FX/Niagara/Source/Niagara/Private/NiagaraRibbonRendererProperties.cpp b/Engine/Plugins/FX/Niagara/Source/Niagara/Private/NiagaraRibbonRendererProperties.cpp
--- a/Engine/Plugins/FX/Niagara/Source/Niagara/Private/NiagaraRibbonRendererProperties.cpp
+++ b/Engine/Plugins/FX/Niagara/Source/Niagara/Private/NiagaraRibbonRendererProperties.cpp
@@ -26,7 +26,7 @@ UNiagaraRibbonRendererProperties::UNiagaraRibbonRendererProperties()
 FNiagaraRenderer* UNiagaraRibbonRendererProperties::CreateEmitterRenderer(ERHIFeatureLevel::Type FeatureLevel, const FNiagaraEmitterInstance* Emitter)
 {
 	FNiagaraRenderer* NewRenderer = new FNiagaraRendererRibbons(FeatureLevel, this, Emitter);
-	NewRenderer->Initialize(FeatureLevel, this, Emitter);
+	NewRenderer->Initialize(this, Emitter);
 	return NewRenderer;
 }
 
diff --git a/Engine/Plugins/FX/Niagara/Source/Niagara/Private/NiagaraSpriteRendererProperties.cpp b/Engine/Plugins/FX/Niagara/Source/Niagara/Private/NiagaraSpriteRendererProperties.cpp
--- a/Engine/Plugins/FX/Niagara/Source/Niagara/Private/NiagaraSpriteRendererProperties.cpp
+++ b/Engine/Plugins/FX/Niagara/Source/Niagara/Private/NiagaraSpriteRendererProperties.cpp
@@ -49,7 +49,7 @@ UNiagaraSpriteRendererProperties::UNiagaraSpriteRendererProperties()
 FNiagaraRenderer* UNiagaraSpriteRendererProperties::CreateEmitterRenderer(ERHIFeatureLevel::Type FeatureLevel, const FNiagaraEmitterInstance* Emitter)
 {
 	FNiagaraRenderer* NewRenderer = new FNiagaraRendererSprites(FeatureLevel, this, Emitter);	
-	NewRenderer->Initialize(FeatureLevel, this, Emitter);
+	NewRenderer->Initialize(this, Emitter);
 	return NewRenderer;
 }
 
diff --git a/Engine/Plugins/FX/Niagara/Source/Niagara/Public/NiagaraRenderer.h b/Engine/Plugins/FX/Niagara/Source/Niagara/Public/NiagaraRenderer.h
--- a/Engine/Plugins/FX/Niagara/Source/Niagara/Public/NiagaraRenderer.h
+++ b/Engine/Plugins/FX/Niagara/Source/Niagara/Public/NiagaraRenderer.h
@@ -53,6 +53,29 @@ protected:
 
 //////////////////////////////////////////////////////////////////////////
 
+extern int32 GbEnableMinimalGPUBuffers;
+
+/** Mapping between a variable in the source dataset and the location we place it in the GPU buffer passed to the VF. */
+struct FNiagaraRendererVariableInfo
+{
+	FNiagaraRendererVariableInfo()
+		: DatasetOffset(INDEX_NONE), GPUBufferOffset(INDEX_NONE), NumComponents(0), bUpload(true)
+	{
+	}
+	
+	FNiagaraRendererVariableInfo(int32 InDatasetOffset, int32 InGPUBufferOffset, int32 InNumComponents, bool InbUpload)
+		: DatasetOffset(InDatasetOffset), GPUBufferOffset(InGPUBufferOffset), NumComponents(InNumComponents), bUpload(InbUpload)
+	{
+	}
+
+	FORCEINLINE int32 GetGPUOffset()const { return GbEnableMinimalGPUBuffers ? GPUBufferOffset : DatasetOffset; }
+
+	int32 DatasetOffset;
+	int32 GPUBufferOffset;
+	int32 NumComponents;
+	bool bUpload;
+};
+
 /**
 * Base class for Niagara System renderers.
 */
@@ -66,7 +89,7 @@ public:
 	FNiagaraRenderer(const FNiagaraRenderer& Other) = delete;
 	FNiagaraRenderer& operator=(const FNiagaraRenderer& Other) = delete;
 
-	virtual void Initialize(ERHIFeatureLevel::Type FeatureLevel, const UNiagaraRendererProperties *InProps, const FNiagaraEmitterInstance* Emitter);
+	virtual void Initialize(const UNiagaraRendererProperties *InProps, const FNiagaraEmitterInstance* Emitter);
 	virtual void CreateRenderThreadResources(NiagaraEmitterInstanceBatcher* Batcher);
 	virtual void ReleaseRenderThreadResources(NiagaraEmitterInstanceBatcher* Batcher);
 
@@ -78,7 +101,7 @@ public:
 	virtual int32 GetDynamicDataSize()const { return 0; }
 	virtual bool IsMaterialValid(UMaterialInterface* Mat)const { return Mat != nullptr; }
 
-	void SortIndices(const struct FNiagaraGPUSortInfo& SortInfo, const FNiagaraDataBuffer& Buffer, FGlobalDynamicReadBuffer::FAllocation& OutIndices)const;
+	void SortIndices(const struct FNiagaraGPUSortInfo& SortInfo, int32 SortVarIdx, const FNiagaraDataBuffer& Buffer, FGlobalDynamicReadBuffer::FAllocation& OutIndices)const;
 
 	void SetDynamicData_RenderThread(FNiagaraDynamicDataBase* NewDynamicData);
 	FORCEINLINE FNiagaraDynamicDataBase *GetDynamicData()const { return DynamicDataRender; }
@@ -110,10 +133,18 @@ protected:
 	ENiagaraSimTarget SimTarget;
 	uint32 NumIndicesPerInstance;
 
+	ERHIFeatureLevel::Type FeatureLevel;
+
 #if STATS
 	TStatId EmitterStatID;
 #endif
 
+	bool SetVertexFactoryVariable(const FNiagaraDataSet& DataSet, const FNiagaraVariable& Var, int32 VFVarOffset);
+	FGlobalDynamicReadBuffer::FAllocation TransferDataToGPU(FGlobalDynamicReadBuffer& DynamicReadBuffer, FNiagaraDataBuffer* SrcData)const;
+	
+	mutable TArray<FNiagaraRendererVariableInfo, TInlineAllocator<16>> VFVariables;
+	int32 TotalVFComponents;
+
 	/** Cached array of materials used from the properties data. Validated with usage flags etc. */
 	TArray<UMaterialInterface*> BaseMaterials_GT;
 	FMaterialRelevance BaseMaterialRelevance_GT;
diff --git a/Engine/Plugins/FX/Niagara/Source/Niagara/Public/NiagaraRendererMeshes.h b/Engine/Plugins/FX/Niagara/Source/Niagara/Public/NiagaraRendererMeshes.h
--- a/Engine/Plugins/FX/Niagara/Source/Niagara/Public/NiagaraRendererMeshes.h
+++ b/Engine/Plugins/FX/Niagara/Source/Niagara/Public/NiagaraRendererMeshes.h
@@ -41,20 +41,7 @@ private:
 	uint32 bOverrideMaterials : 1;
 	uint32 bSortOnlyWhenTranslucent : 1;
 
-	int32 PositionOffset;
-	int32 VelocityOffset;
-	int32 ColorOffset;
-	int32 ScaleOffset;
-	int32 SizeOffset;
 	uint32 MaterialParamValidMask;
-	int32 MaterialParamOffset;
-	int32 MaterialParamOffset1;
-	int32 MaterialParamOffset2;
-	int32 MaterialParamOffset3;
-	int32 TransformOffset;
-	int32 NormalizedAgeOffset;
-	int32 MaterialRandomOffset;
-	int32 CustomSortingOffset;
 
 	int32 MeshMinimumLOD = 0;
 };
diff --git a/Engine/Plugins/FX/Niagara/Source/Niagara/Public/NiagaraRendererRibbons.h b/Engine/Plugins/FX/Niagara/Source/Niagara/Public/NiagaraRendererRibbons.h
--- a/Engine/Plugins/FX/Niagara/Source/Niagara/Public/NiagaraRendererRibbons.h
+++ b/Engine/Plugins/FX/Niagara/Source/Niagara/Public/NiagaraRendererRibbons.h
@@ -78,19 +78,7 @@ private:
 	float CustomTessellationMinAngle;
 	bool bCustomUseScreenSpace;
 
-	int32 PositionDataOffset;
-	int32 VelocityDataOffset;
-	int32 WidthDataOffset;
-	int32 TwistDataOffset;
-	int32 FacingDataOffset;
-	int32 ColorDataOffset;
-	int32 NormalizedAgeDataOffset;
-	int32 MaterialRandomDataOffset;
 	uint32 MaterialParamValidMask;
-	int32 MaterialParamOffset;
-	int32 MaterialParamOffset1;
-	int32 MaterialParamOffset2;
-	int32 MaterialParamOffset3;
 
 	// Average curvature of the segments.
 	mutable float TessellationAngle = 0;
