Exposing more of the API outside the DLL

diff --git a/Engine/Plugins/FX/Niagara/Source/Niagara/Public/NiagaraWorldManager.h b/Engine/Plugins/FX/Niagara/Source/Niagara/Public/NiagaraWorldManager.h
--- a/Engine/Plugins/FX/Niagara/Source/Niagara/Public/NiagaraWorldManager.h
+++ b/Engine/Plugins/FX/Niagara/Source/Niagara/Public/NiagaraWorldManager.h
@@ -94,7 +94,7 @@ public:
 	FNiagaraWorldManager(UWorld* InWorld);
 	~FNiagaraWorldManager();
 
-	static FNiagaraWorldManager* Get(UWorld* World);
+	static NIAGARA_API FNiagaraWorldManager* Get(UWorld* World);
 	static void OnStartup();
 	static void OnShutdown();
 
@@ -124,8 +124,8 @@ public:
 	
 	FORCEINLINE FNDI_SkeletalMesh_GeneratedData& GetSkeletalMeshGeneratedData() { return SkeletalMeshGeneratedData; }
 
-	bool CachedPlayerViewLocationsValid() const { return bCachedPlayerViewLocationsValid; }
-	TArrayView<const FVector> GetCachedPlayerViewLocations() const { check(bCachedPlayerViewLocationsValid); return MakeArrayView(CachedPlayerViewLocations); }
+	NIAGARA_API bool CachedPlayerViewLocationsValid() const { return bCachedPlayerViewLocationsValid; }
+	NIAGARA_API TArrayView<const FVector> GetCachedPlayerViewLocations() const { check(bCachedPlayerViewLocationsValid); return MakeArrayView(CachedPlayerViewLocations); }
 
 	UNiagaraComponentPool* GetComponentPool() { return ComponentPool; }
 
