Build fix

#rnx

diff --git a/Engine/Plugins/FX/Niagara/Source/Niagara/Private/NiagaraRendererRibbons.cpp b/Engine/Plugins/FX/Niagara/Source/Niagara/Private/NiagaraRendererRibbons.cpp
--- a/Engine/Plugins/FX/Niagara/Source/Niagara/Private/NiagaraRendererRibbons.cpp
+++ b/Engine/Plugins/FX/Niagara/Source/Niagara/Private/NiagaraRendererRibbons.cpp
@@ -869,7 +869,7 @@ FNiagaraRendererRibbons::FCPUSimParticleDataAllocation FNiagaraRendererRibbons::
 		else
 		{
 			CPUSimParticleDataAllocation.ParticleData = DynamicReadBuffer.AllocateFloat(TotalFloatSize);
-			FMemory::Memcpy(ParticleData.Buffer, SourceParticleData->GetFloatBuffer().GetData(), SourceParticleData->GetFloatBuffer().Num());
+			FMemory::Memcpy(CPUSimParticleDataAllocation.ParticleData.Buffer, SourceParticleData->GetFloatBuffer().GetData(), SourceParticleData->GetFloatBuffer().Num());
 		}
 	}
 
