Detail level was not being set correctly for cooked platforms.

diff --git a/Engine/Plugins/FX/Niagara/Source/Niagara/Private/NiagaraModule.cpp b/Engine/Plugins/FX/Niagara/Source/Niagara/Private/NiagaraModule.cpp
--- a/Engine/Plugins/FX/Niagara/Source/Niagara/Private/NiagaraModule.cpp
+++ b/Engine/Plugins/FX/Niagara/Source/Niagara/Private/NiagaraModule.cpp
@@ -187,7 +187,7 @@ void INiagaraModule::StartupModule()
 #endif
 
 	CVarDetailLevel.AsVariable()->SetOnChangedCallback(FConsoleVariableDelegate::CreateRaw(this, &INiagaraModule::OnChangeDetailLevel));
-	OnChangeDetailLevel(CVarDetailLevel.AsVariable());
+	EngineDetailLevel = CVarDetailLevel.GetValueOnGameThread();
 
 	//Init commonly used FNiagaraVariables
 
