- Only wait on system & instance concurrent tick to complete when removing an instance
- This avoid a potential stall if an instance completes during finalize before other instances have completed there concurrent tick

#rnx

diff --git a/Engine/Plugins/FX/Niagara/Source/Niagara/Private/NiagaraSystemSimulation.cpp b/Engine/Plugins/FX/Niagara/Source/Niagara/Private/NiagaraSystemSimulation.cpp
--- a/Engine/Plugins/FX/Niagara/Source/Niagara/Private/NiagaraSystemSimulation.cpp
+++ b/Engine/Plugins/FX/Niagara/Source/Niagara/Private/NiagaraSystemSimulation.cpp
@@ -1650,7 +1650,10 @@ void FNiagaraSystemSimulation::RemoveInstance(FNiagaraSystemInstance* Instance)
 			MainDataSet.GetCurrentDataChecked().Dump(Instance->SystemInstanceIndex, 1, TEXT("System data being removed."));
 		}
 
-		WaitForInstancesTickComplete();
+		// Wait for the system simulation & the system instances tick to complete as we are touching both the SystemInstances & DataSet
+		// Note: We do not need to wait for all instances to complete as the system simulation concurrent tick will have transfered data from the DataSet out to ParameterStores
+		WaitForSystemTickComplete();
+		Instance->WaitForAsyncTick();
 
 		int32 NumInstances = MainDataSet.GetCurrentDataChecked().GetNumInstances();
 		check(SystemInstances.Num() == NumInstances);
