Removing default mesh hard ref in cooked.

diff --git a/Engine/Plugins/FX/Niagara/Source/Niagara/Classes/NiagaraDataInterfaceSkeletalMesh.h b/Engine/Plugins/FX/Niagara/Source/Niagara/Classes/NiagaraDataInterfaceSkeletalMesh.h
--- a/Engine/Plugins/FX/Niagara/Source/Niagara/Classes/NiagaraDataInterfaceSkeletalMesh.h
+++ b/Engine/Plugins/FX/Niagara/Source/Niagara/Classes/NiagaraDataInterfaceSkeletalMesh.h
@@ -458,9 +458,11 @@ class NIAGARA_API UNiagaraDataInterfaceSkeletalMesh : public UNiagaraDataInterfa
 
 public:
 	
-	/** Mesh used to sample from when not overridden by a source actor from the scene. Also useful for previewing in the editor. */
+#if WITH_EDITORONLY_DATA
+	/** Mesh used to sample from when not overridden by a source actor from the scene. Only available in editor for previewing. This is removed in cooked builds. */
 	UPROPERTY(EditAnywhere, Category = "Mesh")
 	USkeletalMesh* DefaultMesh;
+#endif
 
 	/** The source actor from which to sample. Takes precedence over the direct mesh. */
 	UPROPERTY(EditAnywhere, Category = "Mesh")
diff --git a/Engine/Plugins/FX/Niagara/Source/Niagara/Private/NiagaraDataInterfaceSkeletalMesh.cpp b/Engine/Plugins/FX/Niagara/Source/Niagara/Private/NiagaraDataInterfaceSkeletalMesh.cpp
--- a/Engine/Plugins/FX/Niagara/Source/Niagara/Private/NiagaraDataInterfaceSkeletalMesh.cpp
+++ b/Engine/Plugins/FX/Niagara/Source/Niagara/Private/NiagaraDataInterfaceSkeletalMesh.cpp
@@ -1074,11 +1074,13 @@ USkeletalMesh* UNiagaraDataInterfaceSkeletalMesh::GetSkeletalMesh(UNiagaraCompon
 	{
 		SceneComponent = FoundSkelComp;
 	}
-	
+
+#if WITH_EDITORONLY_DATA
 	if (!Mesh && DefaultMesh)
 	{
 		Mesh = DefaultMesh;
 	}
+#endif
 
 	return Mesh;
 }
@@ -1417,10 +1419,12 @@ bool FNDISkeletalMesh_InstanceData::ResetRequired(UNiagaraDataInterfaceSkeletalM
 	}
 	else
 	{
+#if WITH_EDITORONLY_DATA
 		if (!Interface->DefaultMesh)
 		{
 			return true;
 		}
+#endif
 	}
 
 	if (Interface->ChangeId != ChangeId)
@@ -1548,7 +1552,9 @@ void FNDISkeletalMesh_InstanceData::Release()
 
 UNiagaraDataInterfaceSkeletalMesh::UNiagaraDataInterfaceSkeletalMesh(FObjectInitializer const& ObjectInitializer)
 	: Super(ObjectInitializer)
+#if WITH_EDITORONLY_DATA
 	, DefaultMesh(nullptr)
+#endif
 	, Source(nullptr)
 	, SkinningMode(ENDISkeletalMesh_SkinningMode::SkinOnTheFly)
 	, WholeMeshLOD(INDEX_NONE)
@@ -1665,13 +1671,13 @@ bool UNiagaraDataInterfaceSkeletalMesh::CopyToInternal(UNiagaraDataInterface* De
 	UNiagaraDataInterfaceSkeletalMesh* OtherTyped = CastChecked<UNiagaraDataInterfaceSkeletalMesh>(Destination);
 	OtherTyped->Source = Source;
 	OtherTyped->MeshUserParameter = MeshUserParameter;
-	OtherTyped->DefaultMesh = DefaultMesh;
 	OtherTyped->SkinningMode = SkinningMode;
 	OtherTyped->SamplingRegions = SamplingRegions;
 	OtherTyped->WholeMeshLOD = WholeMeshLOD;
 	OtherTyped->SpecificBones = SpecificBones;
 	OtherTyped->SpecificSockets = SpecificSockets;
 #if WITH_EDITORONLY_DATA
+	OtherTyped->DefaultMesh = DefaultMesh;
 	OtherTyped->bRequiresCPUAccess = bRequiresCPUAccess;
 #endif
 	return true;
@@ -1685,7 +1691,9 @@ bool UNiagaraDataInterfaceSkeletalMesh::Equals(const UNiagaraDataInterface* Othe
 	}
 	const UNiagaraDataInterfaceSkeletalMesh* OtherTyped = CastChecked<const UNiagaraDataInterfaceSkeletalMesh>(Other);
 	return OtherTyped->Source == Source &&
+#if WITH_EDITORONLY_DATA
 		OtherTyped->DefaultMesh == DefaultMesh &&
+#endif
 		OtherTyped->MeshUserParameter == MeshUserParameter &&
 		OtherTyped->SkinningMode == SkinningMode &&
 		OtherTyped->SamplingRegions == SamplingRegions &&
@@ -1741,6 +1749,7 @@ TArray<FNiagaraDataInterfaceError> UNiagaraDataInterfaceSkeletalMesh::GetErrors(
 	bool bHasNoMeshAssignedError = false;
 	
 	// Collect Errors
+#if WITH_EDITORONLY_DATA
 	if (DefaultMesh != nullptr && bRequiresCPUAccess)
 	{
 		for (auto info : DefaultMesh->GetLODInfoArray())
@@ -1773,6 +1782,7 @@ TArray<FNiagaraDataInterfaceError> UNiagaraDataInterfaceSkeletalMesh::GetErrors(
 
 		Errors.Add(CPUAccessNotAllowedError);
 	}
+#endif
 
 	if (Source == nullptr && bHasNoMeshAssignedError)
 	{
