Some improvements to activation and reset times.

Datasets no longer keep a local copy of their layout information that has to be copied each init.
Removed some debug string work.
Pulled out some bindings between some const parameter stores that were happening each activation to happen once on init to prime those stores.
Reduced the number of stores binding/unbinding on each activation and doing only the minimal set needed.

New stat scopes to help identify cost of activation.
A few minor tweaks to reduce allocations.
Data driven pre allocation for emitters to help avoid many allocations as emitters spin up.
Removed some dead code that needed fixup for above changes.
Emitter renderer objects no longer created for disabled / uninitialized emitters.

diff --git a/Engine/Plugins/FX/Niagara/Source/Niagara/Classes/NiagaraCollision.h b/Engine/Plugins/FX/Niagara/Source/Niagara/Classes/NiagaraCollision.h
--- a/Engine/Plugins/FX/Niagara/Source/Niagara/Classes/NiagaraCollision.h
+++ b/Engine/Plugins/FX/Niagara/Source/Niagara/Classes/NiagaraCollision.h
@@ -24,74 +24,6 @@ struct FNiagaraCollisionTrace
 };
 
 
-class FNiagaraCollisionBatch
-{
-public:
-	FNiagaraCollisionBatch()
-		: CollisionEventDataSet(nullptr)
-		, EmitterName(NAME_None)
-	{
-	}
-
-	~FNiagaraCollisionBatch()
-	{
-		FNiagaraEventDataSetMgr::Reset(OwnerSystemInstanceID, EmitterName);
-	}
-
-	void Reset()
-	{
-		if (CollisionEventDataSet)
-		{
-			CollisionEventDataSet->ResetBuffers();
-		}
-	}
-
-	void Init(FNiagaraSystemInstanceID InOwnerSystemInstanceID, FName InEmitterName)
-	{
-		if (CollisionEventDataSet)
-		{
-#if !UE_BUILD_SHIPPING && !UE_BUILD_TEST
-			CollisionEventDataSet->Init(FNiagaraDataSetID(), ENiagaraSimTarget::CPUSim, TEXT("NiagaraCollisionEventDataset"));
-#else
-			CollisionEventDataSet->Init(FNiagaraDataSetID(), ENiagaraSimTarget::CPUSim);
-#endif
-		}
-		EmitterName = InEmitterName;
-		OwnerSystemInstanceID = InOwnerSystemInstanceID;
-		CollisionEventDataSet = FNiagaraEventDataSetMgr::CreateEventDataSet(OwnerSystemInstanceID, EmitterName, NIAGARA_BUILTIN_EVENTNAME_COLLISION);
-		// TODO: this should go away once we can use the FNiagaraCollisionEventPayload
-		// UStruct to create the data set
-		//FNiagaraVariable ValidVar(FNiagaraTypeDefinition::GetIntDef(), "Valid");
-		FNiagaraVariable PosVar(FNiagaraTypeDefinition::GetVec3Def(), "CollisionLocation");
-		FNiagaraVariable NrmVar(FNiagaraTypeDefinition::GetVec3Def(), "CollisionNormal");
-		FNiagaraVariable PhysMatIdxVar(FNiagaraTypeDefinition::GetIntDef(), "PhysicalMaterialIndex");
-		FNiagaraVariable VelVar(FNiagaraTypeDefinition::GetVec3Def(), "CollisionVelocity");
-		FNiagaraVariable ParticleIndexVar(FNiagaraTypeDefinition::GetIntDef(), "ParticleIndex");
-		//CollisionEventDataSet->AddVariable(ValidVar);
-		CollisionEventDataSet->AddVariable(PosVar);
-		CollisionEventDataSet->AddVariable(NrmVar);
-		CollisionEventDataSet->AddVariable(PhysMatIdxVar);
-		CollisionEventDataSet->AddVariable(VelVar);
-		CollisionEventDataSet->AddVariable(ParticleIndexVar);
-		CollisionEventDataSet->Finalize();
-	}
-	
-	void KickoffNewBatch(class FNiagaraEmitterInstance *Sim, float DeltaSeconds);
-
-	void GenerateEventsFromResults(class FNiagaraEmitterInstance *Sim);
-
-	const FNiagaraDataSet *GetDataSet() const { return CollisionEventDataSet; }
-private:
-	TArray<FTraceHandle> CollisionTraceHandles;
-	TArray<FNiagaraCollisionTrace> CollisionTraces;
-	TArray<FNiagaraCollisionEventPayload> CollisionEvents;
-	FNiagaraDataSet *CollisionEventDataSet;
-	FName EmitterName;
-	FNiagaraSystemInstanceID OwnerSystemInstanceID;
-};
-
-
-
 struct FNiagaraDICollsionQueryResult
 {
 	uint32 TraceID;
diff --git a/Engine/Plugins/FX/Niagara/Source/Niagara/Classes/NiagaraDataSet.h b/Engine/Plugins/FX/Niagara/Source/Niagara/Classes/NiagaraDataSet.h
--- a/Engine/Plugins/FX/Niagara/Source/Niagara/Classes/NiagaraDataSet.h
+++ b/Engine/Plugins/FX/Niagara/Source/Niagara/Classes/NiagaraDataSet.h
@@ -217,7 +217,7 @@ private:
 //////////////////////////////////////////////////////////////////////////
 
 USTRUCT()
-struct FNiagaraDataSetCompiledData
+struct NIAGARA_API FNiagaraDataSetCompiledData
 {
 	GENERATED_BODY()
 
@@ -248,10 +248,14 @@ struct FNiagaraDataSetCompiledData
 	UPROPERTY()
 	ENiagaraSimTarget SimTarget;
 
+	FNiagaraDataSetCompiledData();
 	void BuildLayout();
 	void Empty();
+
+	static FNiagaraDataSetCompiledData DummyCompiledData;
 };
 
+
 //////////////////////////////////////////////////////////////////////////
 
 /**
@@ -266,13 +270,12 @@ public:
 	~FNiagaraDataSet();
 	FNiagaraDataSet& operator=(const FNiagaraDataSet&) = delete;
 
-#if !UE_BUILD_SHIPPING && !UE_BUILD_TEST
-	void Init(FNiagaraDataSetID InID, ENiagaraSimTarget InSimTarget, const FString& InDebugName);
-	void InitFromCompiledData(const FNiagaraDataSetCompiledData& InDataSetCompiledData, const FString& InDebugName);
-#else
-	void Init(FNiagaraDataSetID InID, ENiagaraSimTarget InSimTarget);
-	void InitFromCompiledData(const FNiagaraDataSetCompiledData& InDataSetCompiledData);
-#endif
+	FORCEINLINE void Init(const FNiagaraDataSetCompiledData* InDataSetCompiledData)
+	{
+		CompiledData = InDataSetCompiledData != nullptr ? InDataSetCompiledData : &FNiagaraDataSetCompiledData::DummyCompiledData;
+	
+		Reset();
+	}
 
 	/** Resets current data but leaves variable/layout information etc intact. */
 	void ResetBuffers();
@@ -289,13 +292,10 @@ public:
 	/** Returns size in bytes for all data buffers currently allocated by this dataset. */
 	uint32 GetSizeBytes()const;
 
-	FORCEINLINE bool IsInitialized()const { return bFinalized; }
-	FORCEINLINE ENiagaraSimTarget GetSimTarget()const { return SimTarget; }
-	FORCEINLINE FNiagaraDataSetID GetID()const { return ID; }
-	FORCEINLINE void SetID(FNiagaraDataSetID InID) { ID = InID; }
-
-	FORCEINLINE void SetNeedsPersistentIDs(bool bNeedsIDs) { bNeedsPersistentIDs = bNeedsIDs; }
-	FORCEINLINE bool GetNeedsPersistentIDs()const { return bNeedsPersistentIDs; }
+	FORCEINLINE bool IsInitialized()const { return CompiledData != nullptr; }
+	FORCEINLINE ENiagaraSimTarget GetSimTarget()const { return CompiledData->SimTarget; }
+	FORCEINLINE FNiagaraDataSetID GetID()const { return CompiledData->ID; }	
+	FORCEINLINE bool GetNeedsPersistentIDs()const { return CompiledData->bNeedsPersistentIDs; }
 
 	FORCEINLINE TArray<int32>& GetFreeIDTable() { return FreeIDsTable; }
 	FORCEINLINE int32& GetNumFreeIDs() { return NumFreeIDs; }
@@ -303,18 +303,13 @@ public:
 	FORCEINLINE int32& GetIDAcquireTag() { return IDAcquireTag; }
 	FORCEINLINE void SetIDAcquireTag(int32 InTag) { IDAcquireTag = InTag; }
 
-	FORCEINLINE TArray<FNiagaraVariable>& GetVariables() { return Variables; }
-	FORCEINLINE uint32 GetNumVariables()const { return Variables.Num(); }
-	FORCEINLINE bool HasVariable(const FNiagaraVariable& Var)const { return Variables.Contains(Var); }
-	FORCEINLINE uint32 GetNumFloatComponents()const { return TotalFloatComponents; }
-	FORCEINLINE uint32 GetNumInt32Components()const { return TotalInt32Components; }
-
-	void AddVariable(FNiagaraVariable& Variable);
-	void AddVariables(const TArray<FNiagaraVariable>& Vars);
-
-	/** Finalize the addition of variables and other setup before this data set can be used. */
-	void Finalize();
+	FORCEINLINE const TArray<FNiagaraVariable>& GetVariables()const { return CompiledData->Variables; }
+	FORCEINLINE uint32 GetNumVariables()const { return CompiledData->Variables.Num(); }
+	FORCEINLINE bool HasVariable(const FNiagaraVariable& Var)const { return CompiledData->Variables.Contains(Var); }
+	FORCEINLINE uint32 GetNumFloatComponents()const { return CompiledData->TotalFloatComponents; }
+	FORCEINLINE uint32 GetNumInt32Components()const { return CompiledData->TotalInt32Components; }
 
+	const TArray<FNiagaraVariableLayoutInfo>& GetVariableLayouts()const { return CompiledData->VariableLayouts; }
 	const FNiagaraVariableLayoutInfo* GetVariableLayout(const FNiagaraVariable& Var)const;
 	bool GetVariableComponentOffsets(const FNiagaraVariable& Var, int32 &FloatStart, int32 &IntStart) const;
 
@@ -359,31 +354,15 @@ private:
 		// In some rare occasions, the render thread might be null, like when offloading work to Lightmass 
 		// The final GRenderingThread check keeps us from inadvertently failing when that happens.
 #if DO_GUARD_SLOW
+		ENiagaraSimTarget SimTarget = GetSimTarget();
 		bool CPUSimOK = (SimTarget == ENiagaraSimTarget::CPUSim && !IsInRenderingThread());
 		bool GPUSimOK = (SimTarget == ENiagaraSimTarget::GPUComputeSim && IsInRenderingThread());
 		checkfSlow(!GRenderingThread || CPUSimOK || GPUSimOK, TEXT("NiagaraDataSet function being called on incorrect thread."));
 #endif
 	}
 
-	/** Unique ID for this data set. Used to allow referencing from other emitters and Systems. */
-	FNiagaraDataSetID ID;
-
-	//////////////////////////////////////////////////////////////////////////
-	//TODO: All this layout is known per emitter / system so doesn't need to be generated and stored for every dataset!
-	/** Variables in the data set. */
-	TArray<FNiagaraVariable> Variables;
-	/** Data describing the layout of variable data. */
-	TArray<FNiagaraVariableLayoutInfo> VariableLayouts;
-	/** Total number of components of each type in the data set. */
-	uint32 TotalFloatComponents;
-	uint32 TotalInt32Components;
-	//////////////////////////////////////////////////////////////////////////
-
-	ENiagaraSimTarget SimTarget;
+	const FNiagaraDataSetCompiledData* CompiledData;
 
-	/** Once finalized, the data layout etc is built and no more variables can be added. */
-	uint32 bFinalized : 1;
-	uint32 bNeedsPersistentIDs : 1;
 
 	/** Table of free IDs available to allocate next tick. */
 	TArray<int32> FreeIDsTable;
@@ -411,10 +390,6 @@ private:
 	Additional buffers may be in here if they are currently being used by the render thread.
 	*/
 	TArray<FNiagaraDataBuffer*, TInlineAllocator<2>> Data;
-
-#if !UE_BUILD_SHIPPING && !UE_BUILD_TEST
-	FString DebugName;
-#endif
 };
 
 /**
@@ -1476,10 +1451,10 @@ struct FNiagaraDataVariableIterator
 	{
 		for (int32 VarIdx = 0; VarIdx < Variables.Num(); ++VarIdx)
 		{
-			FNiagaraVariable* Var = Variables[VarIdx];
+			const FNiagaraVariable* Var = Variables[VarIdx];
 			const FNiagaraVariableLayoutInfo* Layout = VarLayouts[VarIdx];
 			check(Var && Layout);
-			uint8* ValuePtr = Var->GetData();
+			const uint8* ValuePtr = Var->GetData();
 
 			for (uint32 CompIdx = 0; CompIdx < Layout->GetNumFloatComponents(); ++CompIdx)
 			{
@@ -1526,7 +1501,7 @@ struct FNiagaraDataVariableIterator
 private:
 
 	const FNiagaraDataBuffer* Data;
-	TArray<FNiagaraVariable*> Variables;
+	TArray<const FNiagaraVariable*> Variables;
 	TArray<const FNiagaraVariableLayoutInfo*> VarLayouts;
 
 	uint32 CurrIdx;
@@ -1562,8 +1537,8 @@ FORCEINLINE void FNiagaraDataBuffer::CheckUsage(bool bReadOnly)const
 
 	//We can read on the RT but any modifications must be GT (or GT Task).
 	//For GPU sims we must be on the RT.
-	checkSlow((Owner->SimTarget == ENiagaraSimTarget::CPUSim  && (IsInGameThread() || (bReadOnly || !IsInRenderingThread()))) ||
-		(Owner->SimTarget == ENiagaraSimTarget::GPUComputeSim) && IsInRenderingThread());
+	checkSlow((Owner->GetSimTarget() == ENiagaraSimTarget::CPUSim  && (IsInGameThread() || (bReadOnly || !IsInRenderingThread()))) ||
+		(Owner->GetSimTarget() == ENiagaraSimTarget::GPUComputeSim) && IsInRenderingThread());
 }
 
 FORCEINLINE_DEBUGGABLE void FNiagaraDataBuffer::AppendToRegisterTable(FNiagaraRegisterTable& Registers, int32 StartInstance)
diff --git a/Engine/Plugins/FX/Niagara/Source/Niagara/Classes/NiagaraEmitter.h b/Engine/Plugins/FX/Niagara/Source/Niagara/Classes/NiagaraEmitter.h
--- a/Engine/Plugins/FX/Niagara/Source/Niagara/Classes/NiagaraEmitter.h
+++ b/Engine/Plugins/FX/Niagara/Source/Niagara/Classes/NiagaraEmitter.h
@@ -70,9 +70,11 @@ struct FNiagaraEventGeneratorProperties
 
 	FNiagaraEventGeneratorProperties(FNiagaraDataSetProperties &Props, FName InEventGenerator)
 		: ID(Props.ID.Name)
-		, SetProps(Props)		
 	{
-
+		DataSetCompiledData.Variables = Props.Variables;
+		DataSetCompiledData.ID = Props.ID;
+		DataSetCompiledData.SimTarget = ENiagaraSimTarget::CPUSim;
+		DataSetCompiledData.BuildLayout();
 	}
 
 	/** Max Number of Events that can be generated per frame. */
@@ -83,7 +85,7 @@ struct FNiagaraEventGeneratorProperties
 	FName ID;
 
 	UPROPERTY()
-	FNiagaraDataSetProperties SetProps;
+	FNiagaraDataSetCompiledData DataSetCompiledData;
 };
 
 
@@ -215,6 +217,13 @@ public:
 	/** An emitter-based seed for the deterministic random number generator. */
 	UPROPERTY(EditAnywhere, Category = "Emitter", meta = (EditCondition = "bDeterminism"))
 	int32 RandomSeed;
+	
+	/** 
+	The emitter will allocate at least this many particles on it's first tick.
+	This can aid performance by avoiding many allocations as an emitter ramps up to it's max size.
+	*/
+	UPROPERTY(EditAnywhere, Category = "Emitter")
+	int32 PreAllocationCount;
 
 	UPROPERTY()
 	FNiagaraEmitterScriptProperties UpdateScriptProps;
diff --git a/Engine/Plugins/FX/Niagara/Source/Niagara/Classes/NiagaraEmitterInstance.h b/Engine/Plugins/FX/Niagara/Source/Niagara/Classes/NiagaraEmitterInstance.h
--- a/Engine/Plugins/FX/Niagara/Source/Niagara/Classes/NiagaraEmitterInstance.h
+++ b/Engine/Plugins/FX/Niagara/Source/Niagara/Classes/NiagaraEmitterInstance.h
@@ -45,8 +45,8 @@ public:
 
 	/** Replaces the binding for a single parameter colleciton instance. If for example the component begins to override the global instance. */
 	//void RebindParameterCollection(UNiagaraParameterCollectionInstance* OldInstance, UNiagaraParameterCollectionInstance* NewInstance);
-	void BindParameters();
-	void UnbindParameters();
+	void BindParameters(bool bExternalOnly);
+	void UnbindParameters(bool bExternalOnly);
 
 	bool IsAllowedToExecute() const;
 
@@ -127,7 +127,6 @@ public:
 	void TickFastPathAttributeBindings();
 
 	FNiagaraEmitterFastPath::FParamMap0& GetFastPathMap() { return FastPathMap; }
-
 private:
 	void CheckForErrors();
 
diff --git a/Engine/Plugins/FX/Niagara/Source/Niagara/Classes/NiagaraSystem.h b/Engine/Plugins/FX/Niagara/Source/Niagara/Classes/NiagaraSystem.h
--- a/Engine/Plugins/FX/Niagara/Source/Niagara/Classes/NiagaraSystem.h
+++ b/Engine/Plugins/FX/Niagara/Source/Niagara/Classes/NiagaraSystem.h
@@ -49,6 +49,11 @@ struct FNiagaraEmitterCompiledData
 	/** Per-Emitter DataSet Data. */
 	UPROPERTY()
 	FNiagaraDataSetCompiledData DataSetCompiledData;
+
+#if WITH_EDITORONLY_DATA
+	UPROPERTY()
+	FNiagaraDataSetCompiledData GPUCaptureDataSetCompiledData;
+#endif
 };
 
 USTRUCT()
@@ -64,6 +69,15 @@ struct FNiagaraSystemCompiledData
 
 	UPROPERTY()
 	FNiagaraParameterStore InstanceParamStore;
+
+	UPROPERTY()
+	FNiagaraDataSetCompiledData DataSetCompiledData;
+
+	UPROPERTY()
+	FNiagaraDataSetCompiledData SpawnInstanceParamsDataSetCompiledData;
+
+	UPROPERTY()
+	FNiagaraDataSetCompiledData UpdateInstanceParamsDataSetCompiledData;
 };
 
 USTRUCT()
diff --git a/Engine/Plugins/FX/Niagara/Source/Niagara/Private/NiagaraCollision.cpp b/Engine/Plugins/FX/Niagara/Source/Niagara/Private/NiagaraCollision.cpp
--- a/Engine/Plugins/FX/Niagara/Source/Niagara/Private/NiagaraCollision.cpp
+++ b/Engine/Plugins/FX/Niagara/Source/Niagara/Private/NiagaraCollision.cpp
@@ -10,167 +10,6 @@
 DECLARE_CYCLE_STAT(TEXT("Collision"), STAT_NiagaraCollision, STATGROUP_Niagara);
 DECLARE_CYCLE_STAT(TEXT("Event Emission"), STAT_NiagaraEventWrite, STATGROUP_Niagara);
 
-void FNiagaraCollisionBatch::KickoffNewBatch(FNiagaraEmitterInstance *Sim, float DeltaSeconds)
-{
-	SCOPE_CYCLE_COUNTER(STAT_NiagaraCollision);
-	FNiagaraVariable PosVar(FNiagaraTypeDefinition::GetVec3Def(), "Position");
-	FNiagaraVariable VelVar(FNiagaraTypeDefinition::GetVec3Def(), "Velocity");
-	FNiagaraVariable SizeVar(FNiagaraTypeDefinition::GetVec2Def(), "SpriteSize");
-	FNiagaraVariable TstVar(FNiagaraTypeDefinition::GetBoolDef(), "PerformCollision");
-	FNiagaraDataSetAccessor<FVector> PosData(Sim->GetData(), PosVar);
-	FNiagaraDataSetAccessor<FVector> VelData(Sim->GetData(), VelVar);
-	FNiagaraDataSetAccessor<FVector2D> SizeData(Sim->GetData(), SizeVar);
-	//FNiagaraDataSetIterator<int32> TstIt(Sim->GetData(), TstVar, 0, false);
-
-	bool bUseSize = SizeData.IsValidForRead();
-
-	if (!PosData.IsValidForRead() || !VelData.IsValidForRead() /*|| !TstIt.IsValid()*/)
-	{
-		return;
-	}
-
-	UWorld *SystemWorld = Sim->GetParentSystemInstance()->GetComponent()->GetWorld();
-	if (SystemWorld)
-	{
-		CollisionTraces.Empty();
-
-		for (uint32 i = 0; i < Sim->GetData().GetCurrentDataChecked().GetNumInstances(); i++)
-		{
-			//int32 TestCollision = *TstIt;
-			//if (TestCollision)
-			{
-				FVector Position;
-				FVector EndPosition;
-				FVector Velocity;
-
-				PosData.Get(i, Position);
-				VelData.Get(i, Velocity);
-				EndPosition = Position + Velocity * DeltaSeconds;
-
-				if (bUseSize)
-				{
-					//TODO:  Handle mesh particles too.  Also this can probably be better and or faster.
-					FVector2D SpriteSize;
-					SizeData.Get(i, SpriteSize);
-					float MaxSize = FMath::Max(SpriteSize.X, SpriteSize.Y);
-
-					float Length;
-					FVector Direction;
-					Velocity.ToDirectionAndLength(Direction, Length);
-
-					Position -= Direction * (MaxSize / 2);
-					EndPosition += Direction * (MaxSize / 2);
-				}
-
-				FCollisionQueryParams QueryParams(SCENE_QUERY_STAT(NiagraAsync));
-				QueryParams.OwnerTag = "Niagara";
-				FTraceHandle Handle = SystemWorld->AsyncLineTraceByChannel(EAsyncTraceType::Single, Position, EndPosition, ECollisionChannel::ECC_WorldStatic, QueryParams, FCollisionResponseParams::DefaultResponseParam, nullptr, i);
-				FNiagaraCollisionTrace Trace;
-				Trace.CollisionTraceHandle = Handle;
-				Trace.SourceParticleIndex = i;
-				CollisionTraces.Add(Trace);
-			}
-		}
-	}
-}
-
-void FNiagaraCollisionBatch::GenerateEventsFromResults(FNiagaraEmitterInstance *Sim)
-{
-	SCOPE_CYCLE_COUNTER(STAT_NiagaraCollision);
-	//CollisionEventDataSet.Allocate(CollisionTraceHandles.Num());
-
-	UWorld *SystemWorld = Sim->GetParentSystemInstance()->GetComponent()->GetWorld();
-	if (SystemWorld)
-	{
-		TArray<FNiagaraCollisionEventPayload> Payloads;
-
-		// generate events for last frame's collisions
-		//
-		for (FNiagaraCollisionTrace CurCheck: CollisionTraces)
-		{
-			FTraceHandle Handle = CurCheck.CollisionTraceHandle;
-			FTraceDatum CurTrace;
-			// wait for trace handles; this should block rarely to never
-			bool bReady = false;
-			while (!bReady)
-			{
-				bReady = SystemWorld->QueryTraceData(Handle, CurTrace);
-				if (!bReady)
-				{
-					// if the query came back false, it's possible that the hanle is invalid for some reason; skip in that case
-					// TODO: handle this more gracefully
-					if (!SystemWorld->IsTraceHandleValid(Handle, false))
-					{
-						break;
-					}
-					break;
-				}
-			}
-
-			if (bReady && CurTrace.OutHits.Num())
-			{
-				// grab the first hit that blocks
-				FHitResult *Hit = FHitResult::GetFirstBlockingHit(CurTrace.OutHits);
-				if (Hit && Hit->IsValidBlockingHit())
-				{
-					FNiagaraCollisionEventPayload Event;
-					Event.CollisionNormal = Hit->ImpactNormal;
-					Event.CollisionPos = Hit->ImpactPoint;
-					Event.ParticleIndex = CurTrace.UserData;
-					check(!Event.CollisionNormal.ContainsNaN());
-					check(Event.CollisionNormal.IsNormalized());
-					check(!Event.CollisionPos.ContainsNaN());
-					check(!Event.CollisionVelocity.ContainsNaN());
-
-					// TODO add to unique list of physical materials for Blueprint
-					Event.PhysicalMaterialIndex = 0;// Hit->PhysMaterial->GetUniqueID();
-
-					Payloads.Add(Event);
-				}
-			}
-		}
-
-		int32 NumInstances = Payloads.Num();
-		if (NumInstances)
-		{
-			// now allocate the data set and write all the event structs
-			//
-			CollisionEventDataSet->Allocate(NumInstances);
-			CollisionEventDataSet->GetDestinationDataChecked().SetNumInstances(NumInstances);
-			//FNiagaraVariable ValidVar(FNiagaraTypeDefinition::GetIntDef(), "Valid");
-			FNiagaraVariable PosVar(FNiagaraTypeDefinition::GetVec3Def(), "CollisionLocation");
-			FNiagaraVariable VelVar(FNiagaraTypeDefinition::GetVec3Def(), "CollisionVelocity");
-			FNiagaraVariable NormVar(FNiagaraTypeDefinition::GetVec3Def(), "CollisionNormal");
-			FNiagaraVariable PhysMatIdxVar(FNiagaraTypeDefinition::GetIntDef(), "PhysicalMaterialIndex");
-			FNiagaraVariable ParticleIndexVar(FNiagaraTypeDefinition::GetIntDef(), "ParticleIndex");
-			//FNiagaraDataSetIterator<int32> ValidItr(*CollisionEventDataSet, ValidVar, 0, true);
-			FNiagaraDataSetAccessor<FVector> PosItr(*CollisionEventDataSet, PosVar);
-			FNiagaraDataSetAccessor<FVector> NormItr(*CollisionEventDataSet, NormVar);
-			FNiagaraDataSetAccessor<FVector> VelItr(*CollisionEventDataSet, VelVar);
-			FNiagaraDataSetAccessor<int32> PhysMatItr(*CollisionEventDataSet, PhysMatIdxVar);
-			FNiagaraDataSetAccessor<int32> ParticleIndexItr(*CollisionEventDataSet, ParticleIndexVar);
-
-			for (int32 i = 0; i < NumInstances; ++i)
-			{
-				SCOPE_CYCLE_COUNTER(STAT_NiagaraEventWrite);
-				FNiagaraCollisionEventPayload& Payload = Payloads[i];
-
-				check(/*ValidItr.IsValid() && */PosItr.IsValid() && VelItr.IsValid() && NormItr.IsValid() && PhysMatItr.IsValid());
-				//ValidItr.Set(1);
-				PosItr.Set(i, Payload.CollisionPos);
-				VelItr.Set(i, Payload.CollisionVelocity);
-				NormItr.Set(i, Payload.CollisionNormal);
-				ParticleIndexItr.Set(i, Payload.ParticleIndex);
-				PhysMatItr.Set(i, 0);
-			}
-		}
-		else
-		{
-			CollisionEventDataSet->ResetBuffers();
-		}
-
-	}
-}
 
 int32 FNiagaraDICollisionQueryBatch::SubmitQuery(FVector Position, FVector Direction, float CollisionSize, float DeltaSeconds)
 {
diff --git a/Engine/Plugins/FX/Niagara/Source/Niagara/Private/NiagaraComponent.cpp b/Engine/Plugins/FX/Niagara/Source/Niagara/Private/NiagaraComponent.cpp
--- a/Engine/Plugins/FX/Niagara/Source/Niagara/Private/NiagaraComponent.cpp
+++ b/Engine/Plugins/FX/Niagara/Source/Niagara/Private/NiagaraComponent.cpp
@@ -123,6 +123,7 @@ FNiagaraSceneProxy::FNiagaraSceneProxy(const UNiagaraComponent* InComponent)
 		CreateRenderers(InComponent);
 		bAlwaysHasVelocity = true;
 		Batcher = SystemInst->GetBatcher();
+		SystemStatID = InComponent->GetAsset()->GetStatID(false, false);
 	}
 }
 
@@ -183,7 +184,12 @@ void FNiagaraSceneProxy::CreateRenderers(const UNiagaraComponent* Component)
 			for (UNiagaraRendererProperties* Properties : Emitter->GetEnabledRenderers())
 			{
 				RendererSortInfo.Emplace(Properties->SortOrderHint, EmitterRenderers.Num());
-				EmitterRenderers.Add(Properties->CreateEmitterRenderer(FeatureLevel, &EmitterInst.Get()));
+				FNiagaraRenderer* NewRenderer = nullptr;
+				if (Properties->GetIsEnabled() && EmitterInst->GetData().IsInitialized() && !EmitterInst->IsDisabled())
+				{
+					NewRenderer = Properties->CreateEmitterRenderer(FeatureLevel, &EmitterInst.Get());
+				}
+				EmitterRenderers.Add(NewRenderer);
 			}
 		}
 	}
@@ -304,6 +310,9 @@ void FNiagaraSceneProxy::GetDynamicMeshElements(const TArray<const FSceneView*>&
 {
 	SCOPE_CYCLE_COUNTER(STAT_NiagaraOverview_RT);
 	SCOPE_CYCLE_COUNTER(STAT_NiagaraComponentGetDynamicMeshElements);
+
+	FScopeCycleCounter SystemStatCounter(SystemStatID);
+
 	for (int32 RendererIdx : RendererDrawOrder)
 	{
 		FNiagaraRenderer* Renderer = EmitterRenderers[RendererIdx];
diff --git a/Engine/Plugins/FX/Niagara/Source/Niagara/Private/NiagaraDataSet.cpp b/Engine/Plugins/FX/Niagara/Source/Niagara/Private/NiagaraDataSet.cpp
--- a/Engine/Plugins/FX/Niagara/Source/Niagara/Private/NiagaraDataSet.cpp
+++ b/Engine/Plugins/FX/Niagara/Source/Niagara/Private/NiagaraDataSet.cpp
@@ -67,11 +67,7 @@ static FAutoConsoleVariableRef CVarNiagaraRenderBufferShrinkFactor(
 );
 
 FNiagaraDataSet::FNiagaraDataSet()
-	: TotalFloatComponents(0)
-	, TotalInt32Components(0)
-	, SimTarget(ENiagaraSimTarget::CPUSim)
-	, bFinalized(0)
-	, bNeedsPersistentIDs(0)
+	: CompiledData(&FNiagaraDataSetCompiledData::DummyCompiledData)
 	, NumFreeIDs(0)
 	, MaxUsedID(0)
 	, IDAcquireTag(0)
@@ -87,69 +83,22 @@ FNiagaraDataSet::~FNiagaraDataSet()
 	ReleaseBuffers();
 }
 
-#if !UE_BUILD_SHIPPING && !UE_BUILD_TEST
-void FNiagaraDataSet::Init(FNiagaraDataSetID InID, ENiagaraSimTarget InSimTarget, const FString& InDebugName)
-{
-	Reset();
-	ID = InID;
-	SimTarget = InSimTarget;
-	DebugName = InDebugName;
-}
-
-void FNiagaraDataSet::InitFromCompiledData(const FNiagaraDataSetCompiledData& InDataSetCompiledData, const FString& InDebugName)
-{
-	Variables = InDataSetCompiledData.Variables;
-	VariableLayouts = InDataSetCompiledData.VariableLayouts;
-	TotalFloatComponents = InDataSetCompiledData.TotalFloatComponents;
-	TotalInt32Components = InDataSetCompiledData.TotalInt32Components;
-	bNeedsPersistentIDs = InDataSetCompiledData.bNeedsPersistentIDs;
-	ID = InDataSetCompiledData.ID;
-	SimTarget = InDataSetCompiledData.SimTarget;
-	DebugName = InDebugName;
-	bFinalized = true;
-}
-#else
-void FNiagaraDataSet::Init(FNiagaraDataSetID InID, ENiagaraSimTarget InSimTarget)
-{
-	Reset();
-	ID = InID;
-	SimTarget = InSimTarget;
-}
-
-void FNiagaraDataSet::InitFromCompiledData(const FNiagaraDataSetCompiledData& InDataSetCompiledData)
-{
-	Variables = InDataSetCompiledData.Variables;
-	VariableLayouts = InDataSetCompiledData.VariableLayouts;
-	TotalFloatComponents = InDataSetCompiledData.TotalFloatComponents;
-	TotalInt32Components = InDataSetCompiledData.TotalInt32Components;
-	bNeedsPersistentIDs = InDataSetCompiledData.bNeedsPersistentIDs;
-	ID = InDataSetCompiledData.ID;
-	SimTarget = InDataSetCompiledData.SimTarget;
-	bFinalized = true;
-}
-#endif
-
 void FNiagaraDataSet::Reset()
 {
 	ResetBuffers();
-
-	Variables.Empty();
-	VariableLayouts.Empty();
-	bFinalized = false;
-	TotalFloatComponents = 0;
-	TotalInt32Components = 0;
-	bNeedsPersistentIDs = 0;
 }
 
 void FNiagaraDataSet::ResetBuffers()
 {
-	if (SimTarget == ENiagaraSimTarget::CPUSim)
+	checkSlow(CompiledData);
+	
+	if (GetSimTarget() == ENiagaraSimTarget::CPUSim)
 	{
 		ResetBuffersInternal();
 	}
 	else
 	{
-		check(SimTarget == ENiagaraSimTarget::GPUComputeSim);
+		checkSlow(GetSimTarget() == ENiagaraSimTarget::GPUComputeSim);
 		ENQUEUE_RENDER_COMMAND(ResetBuffersCommand)([=](FRHICommandListImmediate& RHICmdList)
 		{
 			ResetBuffersInternal();
@@ -171,6 +120,11 @@ void FNiagaraDataSet::ResetBuffersInternal()
 	//Ensure we have a valid current buffer
 	BeginSimulate();
 	EndSimulate();
+
+	if (GetSimTarget() == ENiagaraSimTarget::CPUSim)
+	{
+		GetCurrentDataChecked().BuildRegisterTable();
+	}
 }
 
 void FNiagaraDataSet::ReleaseBuffers()
@@ -207,8 +161,8 @@ FNiagaraDataBuffer& FNiagaraDataSet::BeginSimulate()
 	{
 		Data.Add(new FNiagaraDataBuffer(this));
 		DestinationData = Data.Last();
-		verify(DestinationData->TryLock());
-		check(DestinationData->IsBeingWritten());
+		verifySlow(DestinationData->TryLock());
+		checkSlow(DestinationData->IsBeingWritten());
 	}
 
 	DestinationData->SetNumInstances(0);
@@ -222,7 +176,7 @@ void FNiagaraDataSet::EndSimulate(bool SetCurrentData)
 	//CheckCorrectThread();
 	//Destination is now complete so make it the current simulation state.
 	DestinationData->Unlock();
-	check(!DestinationData->IsInUse());
+	checkSlow(!DestinationData->IsInUse());
 
 	if (SetCurrentData)
 	{
@@ -235,9 +189,9 @@ void FNiagaraDataSet::EndSimulate(bool SetCurrentData)
 
 void FNiagaraDataSet::Allocate(int32 NumInstances, bool bMaintainExisting)
 {
-	check(bFinalized);
+	checkSlow(IsInitialized());
 	CheckCorrectThread();
-	check(DestinationData);
+	checkSlow(DestinationData);
 
 	DestinationData->GetIDTable().Reset();
 	if (bMaintainExisting)
@@ -251,7 +205,7 @@ void FNiagaraDataSet::Allocate(int32 NumInstances, bool bMaintainExisting)
 	CheckForNaNs();
 #endif
 
-	if (bNeedsPersistentIDs)
+	if (GetNeedsPersistentIDs())
 	{
 		TArray<int32>& CurrentIDTable = CurrentData->GetIDTable();
 		TArray<int32>& DestinationIDTable = DestinationData->GetIDTable();
@@ -382,57 +336,10 @@ void FNiagaraDataSet::ReleaseGPUInstanceCounts(FNiagaraGPUInstanceCountManager&
 	}
 }
 
-// note, this method is also implemented in FNiagaraDataSetCompiledData::BuildLayout()
-void FNiagaraDataSet::BuildLayout()
-{
-	VariableLayouts.Empty();
-	TotalFloatComponents = 0;
-	TotalInt32Components = 0;
-
-	VariableLayouts.Reserve(Variables.Num());
-	for (FNiagaraVariable& Var : Variables)
-	{
-		FNiagaraVariableLayoutInfo& VarInfo = VariableLayouts[VariableLayouts.AddDefaulted()];
-		FNiagaraTypeLayoutInfo::GenerateLayoutInfo(VarInfo.LayoutInfo, Var.GetType().GetScriptStruct());
-		VarInfo.FloatComponentStart = TotalFloatComponents;
-		VarInfo.Int32ComponentStart = TotalInt32Components;
-		TotalFloatComponents += VarInfo.GetNumFloatComponents();
-		TotalInt32Components += VarInfo.GetNumInt32Components();
-	}
-}
-
-void FNiagaraDataSet::AddVariable(FNiagaraVariable& Variable)
-{
-	check(!bFinalized);
-	Variables.AddUnique(Variable);
-}
-
-void FNiagaraDataSet::AddVariables(const TArray<FNiagaraVariable>& Vars)
-{
-	check(!bFinalized);
-	for (const FNiagaraVariable& Var : Vars)
-	{
-		Variables.AddUnique(Var);
-	}
-}
-
-void FNiagaraDataSet::Finalize()
-{
-	check(!bFinalized);
-	bFinalized = true;
-	BuildLayout();
-
-	ResetBuffers();
-	if (SimTarget == ENiagaraSimTarget::CPUSim)
-	{
-		GetCurrentDataChecked().BuildRegisterTable();
-	}
-}
-
 const FNiagaraVariableLayoutInfo* FNiagaraDataSet::GetVariableLayout(const FNiagaraVariable& Var)const
 {
-	int32 VarLayoutIndex = Variables.IndexOfByKey(Var);
-	return VarLayoutIndex != INDEX_NONE ? &VariableLayouts[VarLayoutIndex] : nullptr;
+	int32 VarLayoutIndex = GetVariables().IndexOfByKey(Var);
+	return VarLayoutIndex != INDEX_NONE ? &CompiledData->VariableLayouts[VarLayoutIndex] : nullptr;
 }
 
 bool FNiagaraDataSet::GetVariableComponentOffsets(const FNiagaraVariable& Var, int32 &FloatStart, int32 &IntStart) const
@@ -452,20 +359,18 @@ bool FNiagaraDataSet::GetVariableComponentOffsets(const FNiagaraVariable& Var, i
 
 void FNiagaraDataSet::CopyTo(FNiagaraDataSet& Other, int32 StartIdx, int32 NumInstances, bool bResetOther)const
 {
+	check(CompiledData);
+
 	CheckCorrectThread();
 
 	if (bResetOther)
 	{
+		Other.CompiledData = CompiledData;
 		Other.Reset();
-		Other.Variables = Variables;
-		Other.VariableLayouts = VariableLayouts;
-		Other.TotalFloatComponents = TotalFloatComponents;
-		Other.TotalInt32Components = TotalInt32Components;
-		Other.Finalize();
 	}
 	else
 	{
-		checkSlow(Other.GetVariables() == Variables);
+		checkSlow(Other.GetVariables() == GetVariables());
 	}
 
 	//Read the most current data. Even if it's possibly partially complete simulation data.
@@ -510,7 +415,7 @@ void FNiagaraDataSet::CopyTo(FNiagaraDataSet& Other, int32 StartIdx, int32 NumIn
 void FNiagaraDataSet::CopyFromGPUReadback(float* GPUReadBackFloat, int* GPUReadBackInt, int32 StartIdx /* = 0 */, int32 NumInstances /* = INDEX_NONE */, uint32 FloatStride, uint32 IntStride)
 {
 	check(IsInRenderingThread());
-	check(bFinalized);//We should be finalized with proper layout information already.
+	check(IsInitialized());//We should be finalized with proper layout information already.
 
 	FNiagaraDataBuffer& DestBuffer = BeginSimulate();
 	DestBuffer.GPUCopyFrom(GPUReadBackFloat, GPUReadBackInt, StartIdx, NumInstances, FloatStride, IntStride);
@@ -555,13 +460,15 @@ int32 FNiagaraDataBuffer::TransferInstance(FNiagaraDataBuffer& SourceBuffer, int
 		SetNumInstances(OldNumInstances + 1);
 
 		/** Copy the instance data. */
-		for (int32 CompIdx = (int32)Owner->TotalFloatComponents - 1; CompIdx >= 0; --CompIdx)
+		int32 FloatComponents = Owner->GetNumFloatComponents();
+		for (int32 CompIdx = FloatComponents - 1; CompIdx >= 0; --CompIdx)
 		{
 			float* Src = SourceBuffer.GetInstancePtrFloat(CompIdx, InstanceIndex);
 			float* Dst = GetInstancePtrFloat(CompIdx, OldNumInstances);
 			*Dst = *Src;
 		}
-		for (int32 CompIdx = (int32)Owner->TotalInt32Components - 1; CompIdx >= 0; --CompIdx)
+		int32 IntComponents = Owner->GetNumInt32Components();
+		for (int32 CompIdx = IntComponents - 1; CompIdx >= 0; --CompIdx)
 		{
 			int32* Src = SourceBuffer.GetInstancePtrInt32(CompIdx, InstanceIndex);
 			int32* Dst = GetInstancePtrInt32(CompIdx, OldNumInstances);
@@ -599,7 +506,7 @@ bool FNiagaraDataBuffer::CheckForNaNs()const
 void FNiagaraDataBuffer::Allocate(uint32 InNumInstances, bool bMaintainExisting)
 {
 	//CheckUsage(false);
-	check(Owner->SimTarget == ENiagaraSimTarget::CPUSim);
+	checkSlow(Owner->GetSimTarget() == ENiagaraSimTarget::CPUSim);
 
 	NumInstancesAllocated = InNumInstances;
 	NumInstances = 0;
@@ -639,18 +546,21 @@ void FNiagaraDataBuffer::Allocate(uint32 InNumInstances, bool bMaintainExisting)
 	{
 		if (FloatStride != OldFloatStride && FloatStride > 0 && OldFloatStride > 0)
 		{
+			int32 FloatComponents = Owner->GetNumFloatComponents();
 			const uint32 BytesToCopy = FMath::Min(OldFloatStride, FloatStride);
-			for (int32 CompIdx = (int32)Owner->TotalFloatComponents-1; CompIdx >= 0; --CompIdx)
+			for (int32 CompIdx = FloatComponents - 1; CompIdx >= 0; --CompIdx)
 			{
 				uint8* Src = OldFloatData.GetData() + OldFloatStride * CompIdx;
 				uint8* Dst = FloatData.GetData() + FloatStride * CompIdx;
 				FMemory::Memcpy(Dst, Src, BytesToCopy);
 			}
 		}
+
 		if (Int32Stride != OldInt32Stride && Int32Stride > 0 && OldInt32Stride > 0)
 		{
+			int32 IntComponents = Owner->GetNumInt32Components();
 			const uint32 BytesToCopy = FMath::Min(OldInt32Stride, Int32Stride);
-			for (int32 CompIdx = (int32)Owner->TotalInt32Components - 1; CompIdx >= 0; --CompIdx)
+			for (int32 CompIdx = IntComponents - 1; CompIdx >= 0; --CompIdx)
 			{
 				uint8* Src = OldIntData.GetData() + OldInt32Stride * CompIdx;
 				uint8* Dst = Int32Data.GetData() + Int32Stride * CompIdx;
@@ -671,7 +581,7 @@ void FNiagaraDataBuffer::AllocateGPU(uint32 InNumInstances, FNiagaraGPUInstanceC
 {
 	CheckUsage(false);
 
-	check(Owner->SimTarget == ENiagaraSimTarget::GPUComputeSim);
+	checkSlow(Owner->GetSimTarget() == ENiagaraSimTarget::GPUComputeSim);
 
 	// Release previous entry if any.
 	GPUInstanceCountManager.FreeEntry(GPUInstanceCountBufferOffset);
@@ -707,11 +617,8 @@ void FNiagaraDataBuffer::AllocateGPU(uint32 InNumInstances, FNiagaraGPUInstanceC
 			{
 				GPUBufferFloat.Release();
 			}
-#if !UE_BUILD_SHIPPING && !UE_BUILD_TEST 
-			GPUBufferFloat.Initialize(sizeof(float), NumElementsToAlloc * Owner->GetNumFloatComponents(), EPixelFormat::PF_R32_FLOAT, BUF_Static, *Owner->DebugName );
-#else
 			GPUBufferFloat.Initialize(sizeof(float), NumElementsToAlloc * Owner->GetNumFloatComponents(), EPixelFormat::PF_R32_FLOAT, BUF_Static);
-#endif
+
 		}
 		if (Owner->GetNumInt32Components())
 		{
@@ -719,11 +626,7 @@ void FNiagaraDataBuffer::AllocateGPU(uint32 InNumInstances, FNiagaraGPUInstanceC
 			{
 				GPUBufferInt.Release();
 			}
-#if !UE_BUILD_SHIPPING && !UE_BUILD_TEST 
-			GPUBufferInt.Initialize(sizeof(int32), NumElementsToAlloc * Owner->GetNumInt32Components(), EPixelFormat::PF_R32_SINT, BUF_Static, *Owner->DebugName);
-#else
 			GPUBufferInt.Initialize(sizeof(int32), NumElementsToAlloc * Owner->GetNumInt32Components(), EPixelFormat::PF_R32_SINT, BUF_Static);
-#endif
 		}
 	}
 }
@@ -732,7 +635,8 @@ void FNiagaraDataBuffer::SwapInstances(uint32 OldIndex, uint32 NewIndex)
 {
 	CheckUsage(false);
 
-	for (uint32 CompIdx = 0; CompIdx < Owner->TotalFloatComponents; ++CompIdx)
+	uint32 FloatComponents = Owner->GetNumFloatComponents();
+	for (uint32 CompIdx = 0; CompIdx < FloatComponents; ++CompIdx)
 	{
 		float* Src = GetInstancePtrFloat(CompIdx, OldIndex);
 		float* Dst = GetInstancePtrFloat(CompIdx, NewIndex);
@@ -740,7 +644,8 @@ void FNiagaraDataBuffer::SwapInstances(uint32 OldIndex, uint32 NewIndex)
 		*Dst = *Src;
 		*Src = Temp;
 	}
-	for (uint32 CompIdx = 0; CompIdx < Owner->TotalInt32Components; ++CompIdx)
+	uint32 IntComponents = Owner->GetNumInt32Components();
+	for (uint32 CompIdx = 0; CompIdx < IntComponents; ++CompIdx)
 	{
 		int32* Src = GetInstancePtrInt32(CompIdx, OldIndex);
 		int32* Dst = GetInstancePtrInt32(CompIdx, NewIndex);
@@ -756,13 +661,15 @@ void FNiagaraDataBuffer::KillInstance(uint32 InstanceIdx)
 	check(InstanceIdx < NumInstances);
 	--NumInstances;
 
-	for (uint32 CompIdx = 0; CompIdx < Owner->TotalFloatComponents; ++CompIdx)
+	uint32 FloatComponents = Owner->GetNumFloatComponents();
+	for (uint32 CompIdx = 0; CompIdx < FloatComponents; ++CompIdx)
 	{
 		float* Src = GetInstancePtrFloat(CompIdx, NumInstances);
 		float* Dst = GetInstancePtrFloat(CompIdx, InstanceIdx);
 		*Dst = *Src;
 	}
-	for (uint32 CompIdx = 0; CompIdx < Owner->TotalInt32Components; ++CompIdx)
+	uint32 IntComponents = Owner->GetNumInt32Components();
+	for (uint32 CompIdx = 0; CompIdx < IntComponents; ++CompIdx)
 	{
 		int32* Src = GetInstancePtrInt32(CompIdx, NumInstances);
 		int32* Dst = GetInstancePtrInt32(CompIdx, InstanceIdx);
@@ -798,7 +705,8 @@ void FNiagaraDataBuffer::CopyTo(FNiagaraDataBuffer& DestBuffer, int32 StartIdx,
 		}
 		DestBuffer.SetNumInstances(NewNumInstances);
 
-		for (uint32 CompIdx = 0; CompIdx < Owner->TotalFloatComponents; ++CompIdx)
+		uint32 FloatComponents = Owner->GetNumFloatComponents();
+		for (uint32 CompIdx = 0; CompIdx < FloatComponents; ++CompIdx)
 		{
 			const float* SrcStart = GetInstancePtrFloat(CompIdx, StartIdx);
 			const float* SrcEnd = GetInstancePtrFloat(CompIdx, StartIdx + InstancesToCopy);
@@ -814,7 +722,8 @@ void FNiagaraDataBuffer::CopyTo(FNiagaraDataBuffer& DestBuffer, int32 StartIdx,
 				}
 			}
 		}
-		for (uint32 CompIdx = 0; CompIdx < Owner->TotalInt32Components; ++CompIdx)
+		uint32 IntComponents = Owner->GetNumInt32Components();
+		for (uint32 CompIdx = 0; CompIdx < IntComponents; ++CompIdx)
 		{
 			const int32* SrcStart = GetInstancePtrInt32(CompIdx, StartIdx);
 			const int32* SrcEnd = GetInstancePtrInt32(CompIdx, StartIdx + InstancesToCopy);
@@ -847,7 +756,8 @@ void FNiagaraDataBuffer::GPUCopyFrom(float* GPUReadBackFloat, int* GPUReadBackIn
 
 	if (GPUReadBackFloat)
 	{
-		for (uint32 CompIdx = 0; CompIdx < Owner->TotalFloatComponents; ++CompIdx)
+		uint32 FloatComponents = Owner->GetNumFloatComponents();
+		for (uint32 CompIdx = 0; CompIdx < FloatComponents; ++CompIdx)
 		{
 			// We have to reimplement the logic from GetInstancePtrFloat here because the incoming stride may be different than this 
 			// data buffer's stride.
@@ -868,7 +778,8 @@ void FNiagaraDataBuffer::GPUCopyFrom(float* GPUReadBackFloat, int* GPUReadBackIn
 	}
 	if (GPUReadBackInt)
 	{
-		for (uint32 CompIdx = 0; CompIdx < Owner->TotalInt32Components; ++CompIdx)
+		uint32 IntComponents = Owner->GetNumInt32Components();
+		for (uint32 CompIdx = 0; CompIdx < IntComponents; ++CompIdx)
 		{
 			// We have to reimplement the logic from GetInstancePtrInt here because the incoming stride may be different than this 
 			// data buffer's stride.
@@ -904,7 +815,7 @@ void FNiagaraDataBuffer::CopyTo(FNiagaraDataBuffer& DestBuffer)const
 
 void FNiagaraDataBuffer::Dump(int32 StartIndex, int32 InNumInstances, const FString& Label)const
 {
-	TArray<FNiagaraVariable>& Variables = Owner->GetVariables();
+	TArray<FNiagaraVariable> Variables = Owner->GetVariables();
 	FNiagaraDataVariableIterator Itr(this, StartIndex);
 	Itr.AddVariables(Variables);
 
@@ -922,7 +833,7 @@ void FNiagaraDataBuffer::Dump(int32 StartIndex, int32 InNumInstances, const FStr
 		Itr.Get();
 
 		FString Line = TEXT("| ");
-		for (FNiagaraVariable& Var : Owner->GetVariables())
+		for (const FNiagaraVariable& Var : Owner->GetVariables())
 		{
 			Line += Var.ToString() + TEXT(" | ");
 		}
@@ -1120,7 +1031,7 @@ void FNiagaraDataBuffer::BuildRegisterTable()
 	RegisterTable.Reset();
 	RegisterTable.SetNumUninitialized(TotalRegisters);
 	int32 NumRegisters = 0;
-	for (FNiagaraVariableLayoutInfo& VarLayout : Owner->VariableLayouts)
+	for (const FNiagaraVariableLayoutInfo& VarLayout : Owner->GetVariableLayouts())
 	{
 		int32 NumFloats = VarLayout.GetNumFloatComponents();
 		int32 NumInts = VarLayout.GetNumInt32Components();
@@ -1140,6 +1051,10 @@ void FNiagaraDataBuffer::BuildRegisterTable()
 	}
 }
 
+//////////////////////////////////////////////////////////////////////////
+
+FNiagaraDataSetCompiledData FNiagaraDataSetCompiledData::DummyCompiledData;
+
 // note, this method is also implemented in FNiagaraDataSet::BuildLayout()
 void FNiagaraDataSetCompiledData::BuildLayout()
 {
@@ -1159,6 +1074,11 @@ void FNiagaraDataSetCompiledData::BuildLayout()
 	}
 }
 
+FNiagaraDataSetCompiledData::FNiagaraDataSetCompiledData()
+{
+	Empty();
+}
+
 void FNiagaraDataSetCompiledData::Empty()
 {
 	bNeedsPersistentIDs = 0;
diff --git a/Engine/Plugins/FX/Niagara/Source/Niagara/Private/NiagaraEmitter.cpp b/Engine/Plugins/FX/Niagara/Source/Niagara/Private/NiagaraEmitter.cpp
--- a/Engine/Plugins/FX/Niagara/Source/Niagara/Private/NiagaraEmitter.cpp
+++ b/Engine/Plugins/FX/Niagara/Source/Niagara/Private/NiagaraEmitter.cpp
@@ -98,6 +98,7 @@ bool FNiagaraEmitterScriptProperties::DataSetAccessSynchronized() const
 
 UNiagaraEmitter::UNiagaraEmitter(const FObjectInitializer& Initializer)
 : Super(Initializer)
+, PreAllocationCount(0)
 , FixedBounds(FBox(FVector(-100), FVector(100)))
 , MinDetailLevel(0)
 , MaxDetailLevel(4)
diff --git a/Engine/Plugins/FX/Niagara/Source/Niagara/Private/NiagaraEmitterInstance.cpp b/Engine/Plugins/FX/Niagara/Source/Niagara/Private/NiagaraEmitterInstance.cpp
--- a/Engine/Plugins/FX/Niagara/Source/Niagara/Private/NiagaraEmitterInstance.cpp
+++ b/Engine/Plugins/FX/Niagara/Source/Niagara/Private/NiagaraEmitterInstance.cpp
@@ -25,6 +25,12 @@ DECLARE_CYCLE_STAT(TEXT("Emitter Event Handling [CNC]"), STAT_NiagaraEventHandle
 DECLARE_CYCLE_STAT(TEXT("Emitter Error Check [CNC]"), STAT_NiagaraEmitterErrorCheck, STATGROUP_Niagara);
 DECLARE_CYCLE_STAT(TEXT("Init Emitters [GT]"), STAT_NiagaraEmitterInit, STATGROUP_Niagara);
 
+DECLARE_CYCLE_STAT(TEXT("Init Emitters (DataSets) [GT]"), STAT_NiagaraEmitterInit_Datasets, STATGROUP_Niagara);
+DECLARE_CYCLE_STAT(TEXT("Init Emitters (ExecContexts) [GT]"), STAT_NiagaraEmitterInit_ExecContexts, STATGROUP_Niagara);
+DECLARE_CYCLE_STAT(TEXT("Init Emitters (DirectBindings) [GT]"), STAT_NiagaraEmitterInit_DirectBindings, STATGROUP_Niagara);
+DECLARE_CYCLE_STAT(TEXT("Init Emitters (DIDefaultsAndBoundCalcs) [GT]"), STAT_NiagaraEmitterInit_DIDefaultsAndBoundsCalcs, STATGROUP_Niagara);
+
+
 static int32 GbDumpParticleData = 0;
 static FAutoConsoleVariableRef CVarNiagaraDumpParticleData(
 	TEXT("fx.DumpParticleData"),
@@ -118,7 +124,7 @@ FNiagaraEmitterInstance::~FNiagaraEmitterInstance()
 
 	//UE_LOG(LogNiagara, Warning, TEXT("~Simulator %p"), this);
 	CachedBounds.Init();
-	UnbindParameters();
+	UnbindParameters(false);
 
 	if (GPUExecContext != nullptr)
 	{
@@ -260,166 +266,191 @@ void FNiagaraEmitterInstance::Init(int32 InEmitterIdx, FNiagaraSystemInstanceID
 		SpawnInfos.SetNum(EmitterCompiledData[EmitterIdx].SpawnAttributes.Num());
 	}
 
-#if !UE_BUILD_SHIPPING && !UE_BUILD_TEST 
-	Data.InitFromCompiledData(EmitterCompiledData[EmitterIdx].DataSetCompiledData, ParentSystemInstance->GetSystem()->GetName() + TEXT("/") + CachedEmitter->GetName());
-#else
-	Data.InitFromCompiledData(EmitterCompiledData[EmitterIdx].DataSetCompiledData);
-#endif
+	{
+		SCOPE_CYCLE_COUNTER(STAT_NiagaraEmitterInit_Datasets);
+		Data.Init(&EmitterCompiledData[EmitterIdx].DataSetCompiledData);
 
-	ResetSimulation();
+		ResetSimulation();
 
-	DataSetMap.Empty();
+		DataSetMap.Empty();
 
-	//Add the particle data to the data set map.
-	//Currently just used for the tick loop but will also allow access directly to the particle data from other emitters.
-	DataSetMap.Add(Data.GetID(), &Data);
-	//Warn the user if there are any attributes used in the update script that are not initialized in the spawn script.
-	//TODO: We need some window in the System editor and possibly the graph editor for warnings and errors.
+		//Add the particle data to the data set map.
+		//Currently just used for the tick loop but will also allow access directly to the particle data from other emitters.
+		DataSetMap.Add(Data.GetID(), &Data);
+		//Warn the user if there are any attributes used in the update script that are not initialized in the spawn script.
+		//TODO: We need some window in the System editor and possibly the graph editor for warnings and errors.
 
-	const bool bVerboseAttributeLogging = false;
+		const bool bVerboseAttributeLogging = false;
 
 #if !UE_BUILD_SHIPPING && !UE_BUILD_TEST 
-	if (bVerboseAttributeLogging)
-	{
-		for (FNiagaraVariable& Attr : CachedEmitter->UpdateScriptProps.Script->GetVMExecutableData().Attributes)
+		if (bVerboseAttributeLogging)
 		{
-			int32 FoundIdx;
-			if (!CachedEmitter->SpawnScriptProps.Script->GetVMExecutableData().Attributes.Find(Attr, FoundIdx))
+			for (FNiagaraVariable& Attr : CachedEmitter->UpdateScriptProps.Script->GetVMExecutableData().Attributes)
 			{
-				UE_LOG(LogNiagara, Warning, TEXT("Attribute %s is used in the Update script for %s but it is not initialised in the Spawn script!"), *Attr.GetName().ToString(), *EmitterHandle.GetName().ToString());
-			}
-			for (int32 i = 0; i < CachedEmitter->GetEventHandlers().Num(); i++)
-			{
-				if (CachedEmitter->GetEventHandlers()[i].Script && !CachedEmitter->GetEventHandlers()[i].Script->GetVMExecutableData().Attributes.Find(Attr, FoundIdx))
+				int32 FoundIdx;
+				if (!CachedEmitter->SpawnScriptProps.Script->GetVMExecutableData().Attributes.Find(Attr, FoundIdx))
+				{
+					UE_LOG(LogNiagara, Warning, TEXT("Attribute %s is used in the Update script for %s but it is not initialised in the Spawn script!"), *Attr.GetName().ToString(), *EmitterHandle.GetName().ToString());
+				}
+				for (int32 i = 0; i < CachedEmitter->GetEventHandlers().Num(); i++)
 				{
-					UE_LOG(LogNiagara, Warning, TEXT("Attribute %s is used in the event handler script for %s but it is not initialised in the Spawn script!"), *Attr.GetName().ToString(), *EmitterHandle.GetName().ToString());
+					if (CachedEmitter->GetEventHandlers()[i].Script && !CachedEmitter->GetEventHandlers()[i].Script->GetVMExecutableData().Attributes.Find(Attr, FoundIdx))
+					{
+						UE_LOG(LogNiagara, Warning, TEXT("Attribute %s is used in the event handler script for %s but it is not initialised in the Spawn script!"), *Attr.GetName().ToString(), *EmitterHandle.GetName().ToString());
+					}
 				}
 			}
 		}
-	}
 #endif
 
-	ensure(CachedEmitter->UpdateScriptProps.DataSetAccessSynchronized());
-	UpdateScriptEventDataSets.Empty();
-	UpdateEventGeneratorIsSharedByIndex.SetNumZeroed(CachedEmitter->UpdateScriptProps.EventGenerators.Num());
-	int32 UpdateEventGeneratorIndex = 0;
-	for (const FNiagaraEventGeneratorProperties &GeneratorProps : CachedEmitter->UpdateScriptProps.EventGenerators)
-	{
-		FNiagaraDataSet *Set = FNiagaraEventDataSetMgr::CreateEventDataSet(ParentSystemInstance->GetId(), EmitterHandle.GetIdName(), GeneratorProps.SetProps.ID.Name);
-#if !UE_BUILD_SHIPPING && !UE_BUILD_TEST 
-		Set->Init(FNiagaraDataSetID(), ENiagaraSimTarget::CPUSim, CachedEmitter->GetFullName() + TEXT("/") + GeneratorProps.SetProps.ID.Name.ToString());
-#else
-		Set->Init(FNiagaraDataSetID(), ENiagaraSimTarget::CPUSim);
-#endif
-		Set->AddVariables(GeneratorProps.SetProps.Variables);
-		Set->Finalize();
-		UpdateScriptEventDataSets.Add(Set);
-		UpdateEventGeneratorIsSharedByIndex[UpdateEventGeneratorIndex] = CachedEmitter->IsEventGeneratorShared(GeneratorProps.ID);
-		UpdateEventGeneratorIndex++;
-	}
+		ensure(CachedEmitter->UpdateScriptProps.DataSetAccessSynchronized());
+		UpdateScriptEventDataSets.Empty();
+		UpdateEventGeneratorIsSharedByIndex.SetNumZeroed(CachedEmitter->UpdateScriptProps.EventGenerators.Num());
+		int32 UpdateEventGeneratorIndex = 0;
+		for (const FNiagaraEventGeneratorProperties &GeneratorProps : CachedEmitter->UpdateScriptProps.EventGenerators)
+		{
+			FNiagaraDataSet *Set = FNiagaraEventDataSetMgr::CreateEventDataSet(ParentSystemInstance->GetId(), EmitterHandle.GetIdName(), GeneratorProps.ID);
+			Set->Init(&GeneratorProps.DataSetCompiledData);
 
-	ensure(CachedEmitter->SpawnScriptProps.DataSetAccessSynchronized());
-	SpawnScriptEventDataSets.Empty();
-	SpawnEventGeneratorIsSharedByIndex.SetNumZeroed(CachedEmitter->SpawnScriptProps.EventGenerators.Num());
-	int32 SpawnEventGeneratorIndex = 0;
-	for (const FNiagaraEventGeneratorProperties &GeneratorProps : CachedEmitter->SpawnScriptProps.EventGenerators)
-	{
-		FNiagaraDataSet *Set = FNiagaraEventDataSetMgr::CreateEventDataSet(ParentSystemInstance->GetId(), EmitterHandle.GetIdName(), GeneratorProps.SetProps.ID.Name);
-#if !UE_BUILD_SHIPPING && !UE_BUILD_TEST 
-		Set->Init(FNiagaraDataSetID(), ENiagaraSimTarget::CPUSim, CachedEmitter->GetFullName() + TEXT("/") + GeneratorProps.SetProps.ID.Name.ToString());
-#else
-		Set->Init(FNiagaraDataSetID(), ENiagaraSimTarget::CPUSim);
-#endif
-		Set->AddVariables(GeneratorProps.SetProps.Variables);
-		Set->Finalize();
-		SpawnScriptEventDataSets.Add(Set);
-		SpawnEventGeneratorIsSharedByIndex[SpawnEventGeneratorIndex] = CachedEmitter->IsEventGeneratorShared(GeneratorProps.ID);
-		SpawnEventGeneratorIndex++;
-	}
+			UpdateScriptEventDataSets.Add(Set);
+			UpdateEventGeneratorIsSharedByIndex[UpdateEventGeneratorIndex] = CachedEmitter->IsEventGeneratorShared(GeneratorProps.ID);
+			UpdateEventGeneratorIndex++;
+		}
 
-	SpawnExecContext.Init(CachedEmitter->SpawnScriptProps.Script, CachedEmitter->SimTarget);
-	UpdateExecContext.Init(CachedEmitter->UpdateScriptProps.Script, CachedEmitter->SimTarget);
+		ensure(CachedEmitter->SpawnScriptProps.DataSetAccessSynchronized());
+		SpawnScriptEventDataSets.Empty();
+		SpawnEventGeneratorIsSharedByIndex.SetNumZeroed(CachedEmitter->SpawnScriptProps.EventGenerators.Num());
+		int32 SpawnEventGeneratorIndex = 0;
+		for (const FNiagaraEventGeneratorProperties &GeneratorProps : CachedEmitter->SpawnScriptProps.EventGenerators)
+		{
+			FNiagaraDataSet *Set = FNiagaraEventDataSetMgr::CreateEventDataSet(ParentSystemInstance->GetId(), EmitterHandle.GetIdName(), GeneratorProps.ID);
 
-	// setup the parameer store for the GPU execution context; since spawn and update are combined here, we build one with params from both script props
-	if (CachedEmitter->SimTarget == ENiagaraSimTarget::GPUComputeSim)
-	{
-		GPUExecContext = new FNiagaraComputeExecutionContext();
-		GPUExecContext->InitParams(CachedEmitter->GetGPUComputeScript(), CachedEmitter->SimTarget, CachedEmitter->GetUniqueEmitterName(), CachedEmitter->MaxUpdateIterations, CachedEmitter->SpawnStages);
-		GPUExecContext->MainDataSet = &Data;
-		GPUExecContext->GPUScript_RT = CachedEmitter->GetGPUComputeScript()->GetRenderThreadScript();
+			Set->Init(&GeneratorProps.DataSetCompiledData);
 
-		SpawnExecContext.Parameters.Bind(&GPUExecContext->CombinedParamStore);
-		UpdateExecContext.Parameters.Bind(&GPUExecContext->CombinedParamStore);
+			SpawnScriptEventDataSets.Add(Set);
+			SpawnEventGeneratorIsSharedByIndex[SpawnEventGeneratorIndex] = CachedEmitter->IsEventGeneratorShared(GeneratorProps.ID);
+			SpawnEventGeneratorIndex++;
+		}
 	}
 
-	EventExecContexts.SetNum(CachedEmitter->GetEventHandlers().Num());
-	for (int32 i = 0; i < NumEvents; i++)
 	{
-		ensure(CachedEmitter->GetEventHandlers()[i].DataSetAccessSynchronized());
+		SCOPE_CYCLE_COUNTER(STAT_NiagaraEmitterInit_ExecContexts);
 
-		UNiagaraScript* EventScript = CachedEmitter->GetEventHandlers()[i].Script;
+		SpawnExecContext.Init(CachedEmitter->SpawnScriptProps.Script, CachedEmitter->SimTarget);
+		UpdateExecContext.Init(CachedEmitter->UpdateScriptProps.Script, CachedEmitter->SimTarget);
 
-		//This is cpu explicitly? Are we doing event handlers on GPU?
-		EventExecContexts[i].Init(EventScript, ENiagaraSimTarget::CPUSim);
-	}
+		// setup the parameer store for the GPU execution context; since spawn and update are combined here, we build one with params from both script props
+		if (CachedEmitter->SimTarget == ENiagaraSimTarget::GPUComputeSim)
+		{
+			GPUExecContext = new FNiagaraComputeExecutionContext();
+			GPUExecContext->InitParams(CachedEmitter->GetGPUComputeScript(), CachedEmitter->SimTarget, CachedEmitter->GetUniqueEmitterName(), CachedEmitter->MaxUpdateIterations, CachedEmitter->SpawnStages);
+			GPUExecContext->MainDataSet = &Data;
+			GPUExecContext->GPUScript_RT = CachedEmitter->GetGPUComputeScript()->GetRenderThreadScript();
 
-	//Setup direct bindings for setting parameter values.
-	SpawnIntervalBinding.Init(SpawnExecContext.Parameters, EmitterCompiledData[EmitterIdx].EmitterSpawnIntervalVar);
-	InterpSpawnStartBinding.Init(SpawnExecContext.Parameters, EmitterCompiledData[EmitterIdx].EmitterInterpSpawnStartDTVar);
-	SpawnGroupBinding.Init(SpawnExecContext.Parameters, EmitterCompiledData[EmitterIdx].EmitterSpawnGroupVar);
+			SpawnExecContext.Parameters.Bind(&GPUExecContext->CombinedParamStore);
+			UpdateExecContext.Parameters.Bind(&GPUExecContext->CombinedParamStore);
+		}
+	
+		EventExecContexts.SetNum(CachedEmitter->GetEventHandlers().Num());
+		for (int32 i = 0; i < NumEvents; i++)
+		{
+			ensure(CachedEmitter->GetEventHandlers()[i].DataSetAccessSynchronized());
 
-	SpawnEmitterAgeBinding.Init(SpawnExecContext.Parameters, EmitterCompiledData[EmitterIdx].EmitterAgeVar);
-	UpdateEmitterAgeBinding.Init(UpdateExecContext.Parameters, EmitterCompiledData[EmitterIdx].EmitterAgeVar);
-	EventEmitterAgeBindings.SetNum(NumEvents);
-	for (int32 i = 0; i < NumEvents; i++)
-	{
-		EventEmitterAgeBindings[i].Init(EventExecContexts[i].Parameters, EmitterCompiledData[EmitterIdx].EmitterAgeVar);
-	}
+			UNiagaraScript* EventScript = CachedEmitter->GetEventHandlers()[i].Script;
 
-	if (CachedEmitter->SimTarget == ENiagaraSimTarget::GPUComputeSim && GPUExecContext != nullptr)
-	{
-		EmitterAgeBindingGPU.Init(GPUExecContext->CombinedParamStore, EmitterCompiledData[EmitterIdx].EmitterAgeVar);
+			//This is cpu explicitly? Are we doing event handlers on GPU?
+			EventExecContexts[i].Init(EventScript, ENiagaraSimTarget::CPUSim);
+		}
 	}
 
-	SpawnRandomSeedBinding.Init(SpawnExecContext.Parameters, EmitterCompiledData[EmitterIdx].EmitterRandomSeedVar);
-	UpdateRandomSeedBinding.Init(UpdateExecContext.Parameters, EmitterCompiledData[EmitterIdx].EmitterRandomSeedVar);
-	if (CachedEmitter->SimTarget == ENiagaraSimTarget::GPUComputeSim && GPUExecContext != nullptr)
-	{
-		GPURandomSeedBinding.Init(GPUExecContext->CombinedParamStore, EmitterCompiledData[EmitterIdx].EmitterRandomSeedVar);
-	}
 
-	// Initialize the exec count
-	SpawnExecCountBinding.Init(SpawnExecContext.Parameters, SYS_PARAM_ENGINE_EXEC_COUNT);
-	UpdateExecCountBinding.Init(UpdateExecContext.Parameters, SYS_PARAM_ENGINE_EXEC_COUNT);
-	EventExecCountBindings.SetNum(NumEvents);
-	for (int32 i = 0; i < NumEvents; i++)
 	{
-		EventExecCountBindings[i].Init(EventExecContexts[i].Parameters, SYS_PARAM_ENGINE_EXEC_COUNT);
-	}
+		SCOPE_CYCLE_COUNTER(STAT_NiagaraEmitterInit_DirectBindings);
+		//Setup direct bindings for setting parameter values.
+		SpawnIntervalBinding.Init(SpawnExecContext.Parameters, EmitterCompiledData[EmitterIdx].EmitterSpawnIntervalVar);
+		InterpSpawnStartBinding.Init(SpawnExecContext.Parameters, EmitterCompiledData[EmitterIdx].EmitterInterpSpawnStartDTVar);
+		SpawnGroupBinding.Init(SpawnExecContext.Parameters, EmitterCompiledData[EmitterIdx].EmitterSpawnGroupVar);
 
-	// Collect script defined data interface parameters.
-	TArray<UNiagaraScript*, TInlineAllocator<4>> Scripts;
-	Scripts.Add(CachedEmitter->SpawnScriptProps.Script);
-	Scripts.Add(CachedEmitter->UpdateScriptProps.Script);
-	for (const FNiagaraEventScriptProperties& EventHandler : CachedEmitter->GetEventHandlers())
-	{
-		Scripts.Add(EventHandler.Script);
+		SpawnEmitterAgeBinding.Init(SpawnExecContext.Parameters, EmitterCompiledData[EmitterIdx].EmitterAgeVar);
+		UpdateEmitterAgeBinding.Init(UpdateExecContext.Parameters, EmitterCompiledData[EmitterIdx].EmitterAgeVar);
+		EventEmitterAgeBindings.SetNum(NumEvents);
+		for (int32 i = 0; i < NumEvents; i++)
+		{
+			EventEmitterAgeBindings[i].Init(EventExecContexts[i].Parameters, EmitterCompiledData[EmitterIdx].EmitterAgeVar);
+		}
+
+		if (CachedEmitter->SimTarget == ENiagaraSimTarget::GPUComputeSim && GPUExecContext != nullptr)
+		{
+			EmitterAgeBindingGPU.Init(GPUExecContext->CombinedParamStore, EmitterCompiledData[EmitterIdx].EmitterAgeVar);
+		}
+
+		SpawnRandomSeedBinding.Init(SpawnExecContext.Parameters, EmitterCompiledData[EmitterIdx].EmitterRandomSeedVar);
+		UpdateRandomSeedBinding.Init(UpdateExecContext.Parameters, EmitterCompiledData[EmitterIdx].EmitterRandomSeedVar);
+		if (CachedEmitter->SimTarget == ENiagaraSimTarget::GPUComputeSim && GPUExecContext != nullptr)
+		{
+			GPURandomSeedBinding.Init(GPUExecContext->CombinedParamStore, EmitterCompiledData[EmitterIdx].EmitterRandomSeedVar);
+		}
+
+		// Initialize the exec count
+		SpawnExecCountBinding.Init(SpawnExecContext.Parameters, SYS_PARAM_ENGINE_EXEC_COUNT);
+		UpdateExecCountBinding.Init(UpdateExecContext.Parameters, SYS_PARAM_ENGINE_EXEC_COUNT);
+		EventExecCountBindings.SetNum(NumEvents);
+		for (int32 i = 0; i < NumEvents; i++)
+		{
+			EventExecCountBindings[i].Init(EventExecContexts[i].Parameters, SYS_PARAM_ENGINE_EXEC_COUNT);
+		}
 	}
-	FNiagaraUtilities::CollectScriptDataInterfaceParameters(*CachedEmitter, MakeArrayView(Scripts), ScriptDefinedDataInterfaceParameters);
 
-	// Initialize bounds calculators
-	//-OPT: Could skip creating this if we won't ever use it
-	const TArray<UNiagaraRendererProperties*>& Renderers = CachedEmitter->GetEnabledRenderers();
-	BoundsCalculators.Reserve(Renderers.Num());
-	for (UNiagaraRendererProperties* RendererProperties : Renderers)
 	{
-		FNiagaraBoundsCalculator* BoundsCalculator = RendererProperties->CreateBoundsCalculator();
-		if (BoundsCalculator != nullptr)
+		SCOPE_CYCLE_COUNTER(STAT_NiagaraEmitterInit_DIDefaultsAndBoundsCalcs);
+		// Collect script defined data interface parameters.
+		TArray<UNiagaraScript*, TInlineAllocator<8>> Scripts;
+		Scripts.Add(CachedEmitter->SpawnScriptProps.Script);
+		Scripts.Add(CachedEmitter->UpdateScriptProps.Script);
+		for (const FNiagaraEventScriptProperties& EventHandler : CachedEmitter->GetEventHandlers())
 		{
-			BoundsCalculator->InitAccessors(*ParticleDataSet);
-			BoundsCalculators.Emplace(BoundsCalculator);
+			Scripts.Add(EventHandler.Script);
 		}
-	}
+		FNiagaraUtilities::CollectScriptDataInterfaceParameters(*CachedEmitter, MakeArrayView(Scripts), ScriptDefinedDataInterfaceParameters);
+
+		// Initialize bounds calculators
+		//-OPT: Could skip creating this if we won't ever use it
+		bool bUseDynamicBounds = !CachedSystemFixedBounds.IsSet() && !CachedEmitter->bFixedBounds && CachedEmitter->SimTarget == ENiagaraSimTarget::CPUSim;
+		if (bUseDynamicBounds)
+		{
+			BoundsCalculators.Reserve(CachedEmitter->GetRenderers().Num());
+			for (UNiagaraRendererProperties* RendererProperties : CachedEmitter->GetRenderers())
+			{
+				if ((RendererProperties != nullptr) && RendererProperties->GetIsEnabled())
+				{
+					FNiagaraBoundsCalculator* BoundsCalculator = RendererProperties->CreateBoundsCalculator();
+					if (BoundsCalculator != nullptr)
+					{
+						BoundsCalculator->InitAccessors(*ParticleDataSet);
+						BoundsCalculators.Emplace(BoundsCalculator);
+					}
+				}
+			}
+		}
+
+		//Bind some stores and unbind immediately just to prime some data from those stores.
+		FNiagaraParameterStore& SystemScriptDefinedDataInterfaceParameters = ParentSystemInstance->GetSystemSimulation()->GetScriptDefinedDataInterfaceParameters();
+		
+		SystemScriptDefinedDataInterfaceParameters.Bind(&SpawnExecContext.Parameters);
+		ScriptDefinedDataInterfaceParameters.Bind(&SpawnExecContext.Parameters);
+		SpawnExecContext.Parameters.UnbindFromSourceStores();
+
+		SystemScriptDefinedDataInterfaceParameters.Bind(&UpdateExecContext.Parameters);
+		ScriptDefinedDataInterfaceParameters.Bind(&UpdateExecContext.Parameters);
+		UpdateExecContext.Parameters.UnbindFromSourceStores();
+
+		for (FNiagaraScriptExecutionContext& EventContext : EventExecContexts)
+		{
+			SystemScriptDefinedDataInterfaceParameters.Bind(&EventContext.Parameters);
+			ScriptDefinedDataInterfaceParameters.Bind(&EventContext.Parameters);
+			EventContext.Parameters.UnbindFromSourceStores();
+		}
+	}	
 }
 
 void FNiagaraEmitterInstance::ResetSimulation(bool bKillExisting /*= true*/)
@@ -428,6 +459,12 @@ void FNiagaraEmitterInstance::ResetSimulation(bool bKillExisting /*= true*/)
 	Loops = 0;
 	TickCount = 0;
 	CachedBounds.Init();
+
+	if (IsDisabled())
+	{
+		return;
+	}
+
 	SetExecutionState(ENiagaraExecutionState::Active);
 
 	if (bKillExisting)
@@ -576,31 +613,59 @@ void FNiagaraEmitterInstance::DirtyDataInterfaces()
 // 	}
 // }
 
-void FNiagaraEmitterInstance::UnbindParameters()
+void FNiagaraEmitterInstance::UnbindParameters(bool bExternalOnly)
 {
-	SpawnExecContext.Parameters.UnbindFromSourceStores();
-	UpdateExecContext.Parameters.UnbindFromSourceStores();
-	if (GPUExecContext != nullptr)
+	if (bExternalOnly && !IsDisabled())
 	{
-		GPUExecContext->CombinedParamStore.UnbindFromSourceStores();
-	}
+		for (UNiagaraParameterCollection* Collection : SpawnExecContext.Script->GetCachedParameterCollectionReferences())
+		{
+			if (UNiagaraParameterCollectionInstance* NPCInst = ParentSystemInstance->GetParameterCollectionInstance(Collection))
+			{
+				NPCInst->GetParameterStore().Unbind(&SpawnExecContext.Parameters);
+			}
+		}
+		for (UNiagaraParameterCollection* Collection : UpdateExecContext.Script->GetCachedParameterCollectionReferences())
+		{
+			if (UNiagaraParameterCollectionInstance* NPCInst = ParentSystemInstance->GetParameterCollectionInstance(Collection))
+			{
+				NPCInst->GetParameterStore().Unbind(&UpdateExecContext.Parameters);
+			}
+		}
 
-	for (int32 EventIdx = 0; EventIdx < EventExecContexts.Num(); ++EventIdx)
+		for (int32 EventIdx = 0; EventIdx < EventExecContexts.Num(); ++EventIdx)
+		{
+			for (UNiagaraParameterCollection* Collection : EventExecContexts[EventIdx].Script->GetCachedParameterCollectionReferences())
+			{
+				if (UNiagaraParameterCollectionInstance* NPCInst = ParentSystemInstance->GetParameterCollectionInstance(Collection))
+				{
+					NPCInst->GetParameterStore().Unbind(&EventExecContexts[EventIdx].Parameters);
+				}
+			}
+		}
+	}
+	else
 	{
-		EventExecContexts[EventIdx].Parameters.UnbindFromSourceStores();
+		SpawnExecContext.Parameters.UnbindFromSourceStores();
+		UpdateExecContext.Parameters.UnbindFromSourceStores();
+		if (GPUExecContext != nullptr)
+		{
+			GPUExecContext->CombinedParamStore.UnbindFromSourceStores();
+		}
+
+		for (int32 EventIdx = 0; EventIdx < EventExecContexts.Num(); ++EventIdx)
+		{
+			EventExecContexts[EventIdx].Parameters.UnbindFromSourceStores();
+		}
 	}
 }
 
-void FNiagaraEmitterInstance::BindParameters()
+void FNiagaraEmitterInstance::BindParameters(bool bExternalOnly)
 {
 	if (IsDisabled())
 	{
 		return;
 	}
 
-	FNiagaraWorldManager* WorldMan = ParentSystemInstance->GetWorldManager();
-	check(WorldMan);
-
 	for (UNiagaraParameterCollection* Collection : SpawnExecContext.Script->GetCachedParameterCollectionReferences())
 	{
 		ParentSystemInstance->GetParameterCollectionInstance(Collection)->GetParameterStore().Bind(&SpawnExecContext.Parameters);
@@ -618,39 +683,34 @@ void FNiagaraEmitterInstance::BindParameters()
 		}
 	}
 
-	//Now bind parameters from the component and system.
-	FNiagaraParameterStore& InstanceParams = ParentSystemInstance->GetParameters();
-	FNiagaraParameterStore& SystemScriptDefinedDataInterfaceParameters = ParentSystemInstance->GetSystemSimulation()->GetScriptDefinedDataInterfaceParameters();
+	if (!bExternalOnly)
+	{
+		//Now bind parameters from the component and system.
+		FNiagaraParameterStore& InstanceParams = ParentSystemInstance->GetParameters();		
 	
-	InstanceParams.Bind(&SpawnExecContext.Parameters);
-	SystemScriptDefinedDataInterfaceParameters.Bind(&SpawnExecContext.Parameters);
-	ScriptDefinedDataInterfaceParameters.Bind(&SpawnExecContext.Parameters);
+		InstanceParams.Bind(&SpawnExecContext.Parameters);
+		InstanceParams.Bind(&UpdateExecContext.Parameters);
 
-	InstanceParams.Bind(&UpdateExecContext.Parameters);
-	SystemScriptDefinedDataInterfaceParameters.Bind(&UpdateExecContext.Parameters);
-	ScriptDefinedDataInterfaceParameters.Bind(&UpdateExecContext.Parameters);
-
-	for (FNiagaraScriptExecutionContext& EventContext : EventExecContexts)
-	{
-		InstanceParams.Bind(&EventContext.Parameters);
-		SystemScriptDefinedDataInterfaceParameters.Bind(&EventContext.Parameters);
-		ScriptDefinedDataInterfaceParameters.Bind(&EventContext.Parameters);
-	}
+		for (FNiagaraScriptExecutionContext& EventContext : EventExecContexts)
+		{
+			InstanceParams.Bind(&EventContext.Parameters);
+		}
 
-#if WITH_EDITORONLY_DATA
-	CachedEmitter->SpawnScriptProps.Script->RapidIterationParameters.Bind(&SpawnExecContext.Parameters);
-	CachedEmitter->UpdateScriptProps.Script->RapidIterationParameters.Bind(&UpdateExecContext.Parameters);
-	ensure(CachedEmitter->GetEventHandlers().Num() == EventExecContexts.Num());
-	for (int32 i = 0; i < CachedEmitter->GetEventHandlers().Num(); i++)
-	{
-		CachedEmitter->GetEventHandlers()[i].Script->RapidIterationParameters.Bind(&EventExecContexts[i].Parameters);
-	}
-#endif
+	#if WITH_EDITORONLY_DATA
+		CachedEmitter->SpawnScriptProps.Script->RapidIterationParameters.Bind(&SpawnExecContext.Parameters);
+		CachedEmitter->UpdateScriptProps.Script->RapidIterationParameters.Bind(&UpdateExecContext.Parameters);
+		ensure(CachedEmitter->GetEventHandlers().Num() == EventExecContexts.Num());
+		for (int32 i = 0; i < CachedEmitter->GetEventHandlers().Num(); i++)
+		{
+			CachedEmitter->GetEventHandlers()[i].Script->RapidIterationParameters.Bind(&EventExecContexts[i].Parameters);
+		}
+	#endif
 
-	if (CachedEmitter->SimTarget == ENiagaraSimTarget::GPUComputeSim)
-	{
-		SpawnExecContext.Parameters.Bind(&GPUExecContext->CombinedParamStore);
-		UpdateExecContext.Parameters.Bind(&GPUExecContext->CombinedParamStore);
+		if (CachedEmitter->SimTarget == ENiagaraSimTarget::GPUComputeSim)
+		{
+			SpawnExecContext.Parameters.Bind(&GPUExecContext->CombinedParamStore);
+			UpdateExecContext.Parameters.Bind(&GPUExecContext->CombinedParamStore);
+		}
 	}
 }
 
@@ -1097,14 +1157,14 @@ void FNiagaraEmitterInstance::Tick(float DeltaSeconds)
 			TSharedPtr<struct FNiagaraScriptDebuggerInfo, ESPMode::ThreadSafe> DebugInfo = ParentSystemInstance->GetActiveCaptureWrite(CachedIDName, ENiagaraScriptUsage::ParticleGPUComputeScript, FGuid());
 			if (DebugInfo)
 			{
+				const TArray<FNiagaraEmitterCompiledData>& EmitterCompiledData = ParentSystemInstance->GetSystem()->GetEmitterCompiledData();
+
 				//Data.Dump(DebugInfo->Frame, true, 0, OrigNumParticles);
 				//DebugInfo->Frame.Dump(true, 0, OrigNumParticles);
 				DebugInfo->Parameters = GPUExecContext->CombinedParamStore;
 				
 				//TODO: This layout info can be pulled into the emitter/systems etc and all sets just refer to them. They are becoming an annoyance here.
-				DebugInfo->Frame.Init(FNiagaraDataSetID(TEXT("GPU Capture Data"), ENiagaraDataSetType::ParticleData), ENiagaraSimTarget::CPUSim, TEXT("GPU Capture Data"));
-				DebugInfo->Frame.AddVariables(Data.GetVariables());
-				DebugInfo->Frame.Finalize();
+				DebugInfo->Frame.Init(&EmitterCompiledData[EmitterIdx].GPUCaptureDataSetCompiledData);
 
 				GPUExecContext->DebugInfo = DebugInfo;
 			}
@@ -1202,7 +1262,7 @@ void FNiagaraEmitterInstance::Tick(float DeltaSeconds)
 
 	int32 OrigNumParticles = Data.GetCurrentDataChecked().GetNumInstances();
 
-	int32 AllocationSize = OrigNumParticles + SpawnTotal + EventSpawnTotal;
+	int32 AllocationSize = FMath::Max<int32>(CachedEmitter->PreAllocationCount, OrigNumParticles + SpawnTotal + EventSpawnTotal);
 	//Ensure we don't blow our current hard limits on cpu particle count.
 	//TODO: These current limits can be improved relatively easily. Though perf in at these counts will obviously be an issue anyway.
 	if (CachedEmitter->SimTarget == ENiagaraSimTarget::CPUSim && AllocationSize > GMaxNiagaraCPUParticlesPerEmitter)
diff --git a/Engine/Plugins/FX/Niagara/Source/Niagara/Private/NiagaraSystem.cpp b/Engine/Plugins/FX/Niagara/Source/Niagara/Private/NiagaraSystem.cpp
--- a/Engine/Plugins/FX/Niagara/Source/Niagara/Private/NiagaraSystem.cpp
+++ b/Engine/Plugins/FX/Niagara/Source/Niagara/Private/NiagaraSystem.cpp
@@ -369,12 +369,12 @@ void UNiagaraSystem::PostLoad()
 			UE_LOG(LogNiagara, Log, TEXT("System %s being compiled because there were changes to an emitter script Change ID."), *GetPathName());
 		}
 
-	if (EmitterCompiledData.Num() == 0 || EmitterCompiledData[0].DataSetCompiledData.Variables.Num() == 0)
-	{
-		InitEmitterCompiledData();
-	}
+		if (EmitterCompiledData.Num() == 0 || EmitterCompiledData[0].DataSetCompiledData.Variables.Num() == 0)
+		{
+			InitEmitterCompiledData();
+		}
 
-		if (SystemCompiledData.InstanceParamStore.GetNumParameters() == 0)
+		if (SystemCompiledData.InstanceParamStore.GetNumParameters() == 0 ||SystemCompiledData.DataSetCompiledData.Variables.Num() == 0)
 		{
 			InitSystemCompiledData();
 		}
@@ -1130,10 +1130,16 @@ void UNiagaraSystem::InitEmitterCompiledData()
 			const FNiagaraEmitterHandle& EmitterHandle = EmitterHandles[EmitterIdx];
 			const UNiagaraEmitter* Emitter = EmitterHandle.GetInstance();
 			FNiagaraDataSetCompiledData& EmitterDataSetCompiledData = EmitterCompiledData[EmitterIdx].DataSetCompiledData;
+			FNiagaraDataSetCompiledData& GPUCaptureCompiledData = EmitterCompiledData[EmitterIdx].GPUCaptureDataSetCompiledData;
 			if ensureMsgf(Emitter != nullptr, TEXT("Failed to get Emitter Instance from Emitter Handle in post compile, please investigate."))
 			{
+				static FName GPUCaptureDataSetName = TEXT("GPU Capture Dataset");
 				InitEmitterVariableAliasNames(EmitterCompiledData[EmitterIdx], Emitter);
 				InitEmitterDataSetCompiledData(EmitterDataSetCompiledData, Emitter, EmitterHandle);
+				GPUCaptureCompiledData.ID = FNiagaraDataSetID(GPUCaptureDataSetName, ENiagaraDataSetType::ParticleData);
+				GPUCaptureCompiledData.Variables = EmitterDataSetCompiledData.Variables;
+				GPUCaptureCompiledData.SimTarget = ENiagaraSimTarget::CPUSim;
+				GPUCaptureCompiledData.BuildLayout();				
 			}
 		}
 	}
@@ -1170,6 +1176,30 @@ void UNiagaraSystem::InitSystemCompiledData()
 			}
 		}
 	}
+
+	auto CreateDataSetCompiledData = [&](FNiagaraDataSetCompiledData& CompiledData, TArrayView<FNiagaraVariable> Vars)
+	{
+		CompiledData.Empty();
+
+		CompiledData.Variables.Reset(Vars.Num());
+		for (const FNiagaraVariable& Var : Vars)
+		{
+			CompiledData.Variables.AddUnique(Var);
+		}
+
+		CompiledData.bNeedsPersistentIDs = false;
+		CompiledData.ID = FNiagaraDataSetID();
+		CompiledData.SimTarget = ENiagaraSimTarget::CPUSim;
+
+		CompiledData.BuildLayout();
+	};
+
+	CreateDataSetCompiledData(SystemCompiledData.DataSetCompiledData, GetSystemUpdateScript()->GetVMExecutableData().Attributes);
+
+	FNiagaraParameters* EngineParamsSpawn = GetSystemSpawnScript()->GetVMExecutableData().DataSetToParameters.Find(TEXT("Engine"));
+	CreateDataSetCompiledData(SystemCompiledData.SpawnInstanceParamsDataSetCompiledData, EngineParamsSpawn ? EngineParamsSpawn->Parameters : TArrayView<FNiagaraVariable>());
+	FNiagaraParameters* EngineParamsUpdate = GetSystemUpdateScript()->GetVMExecutableData().DataSetToParameters.Find(TEXT("Engine"));
+	CreateDataSetCompiledData(SystemCompiledData.UpdateInstanceParamsDataSetCompiledData, EngineParamsUpdate ? EngineParamsUpdate->Parameters : TArrayView<FNiagaraVariable>());
 }
 #endif
 
diff --git a/Engine/Plugins/FX/Niagara/Source/Niagara/Private/NiagaraSystemInstance.cpp b/Engine/Plugins/FX/Niagara/Source/Niagara/Private/NiagaraSystemInstance.cpp
--- a/Engine/Plugins/FX/Niagara/Source/Niagara/Private/NiagaraSystemInstance.cpp
+++ b/Engine/Plugins/FX/Niagara/Source/Niagara/Private/NiagaraSystemInstance.cpp
@@ -80,6 +80,7 @@ FNiagaraSystemInstance::FNiagaraSystemInstance(UNiagaraComponent* InComponent)
 	, bIsTransformDirty(true)
 	, bNeedsFinalize(false)
 	, bDataInterfacesInitialized(false)
+	, bAlreadyBound(false)
 	, bAsyncWorkInProgress(false)
 	, CachedDeltaSeconds(0.0f)
 	, RequestedExecutionState(ENiagaraExecutionState::Complete)
@@ -160,6 +161,7 @@ void FNiagaraSystemInstance::Init(bool bInForceSolo)
 	bForceSolo = bInForceSolo;
 	ActualExecutionState = ENiagaraExecutionState::Inactive;
 	RequestedExecutionState = ENiagaraExecutionState::Inactive;
+	bAlreadyBound = false;
 
 	//InstanceParameters = GetSystem()->GetInstanceParameters();
 	// In order to get user data interface parameters in the component to work properly,
@@ -534,7 +536,7 @@ void FNiagaraSystemInstance::Complete()
 
 	DestroyDataInterfaceInstanceData();
 
-	UnbindParameters();
+	UnbindParameters(true);
 
 	bPendingSpawn = false;
 
@@ -827,6 +829,8 @@ void FNiagaraSystemInstance::ReInitInternal()
 	LocalBounds = FBox(FVector::ZeroVector, FVector::ZeroVector);
 	CachedDeltaSeconds = 0.0f;
 
+	bAlreadyBound = false;
+
 	FastPathIntUserParameterInputBindings.Empty();
 	FastPathFloatUserParameterInputBindings.Empty();
 	FastPathIntUpdateRangedInputBindings.Empty();
@@ -997,7 +1001,11 @@ void FNiagaraSystemInstance::BindParameters()
 		return;
 	}
 
-	Component->GetOverrideParameters().Bind(&InstanceParameters);
+	if (!bAlreadyBound)
+	{
+		// NOTE: We don't rebind if it's already bound to improve reset times.
+		Component->GetOverrideParameters().Bind(&InstanceParameters); 
+	}
 
 	if (SystemSimulation->GetIsSolo())
 	{
@@ -1006,27 +1014,20 @@ void FNiagaraSystemInstance::BindParameters()
 		Component->GetOverrideParameters().Bind(&SystemSimulation->GetSpawnExecutionContext().Parameters);
 		Component->GetOverrideParameters().Bind(&SystemSimulation->GetUpdateExecutionContext().Parameters);
 	}
-	else 
-	{
-		// If this simulation is not solo than we have bind the source system parameters to the system simulation contexts so that
-		// the system and emitter scripts use the default shared data interfaces.
-		GetSystem()->GetExposedParameters().Bind(&SystemSimulation->GetSpawnExecutionContext().Parameters);
-		GetSystem()->GetExposedParameters().Bind(&SystemSimulation->GetUpdateExecutionContext().Parameters);
-	}
 
 	for (TSharedRef<FNiagaraEmitterInstance, ESPMode::ThreadSafe> Simulation : Emitters)
 	{
-		Simulation->BindParameters();
+		Simulation->BindParameters(bAlreadyBound);
 	}
+
+	bAlreadyBound = true;
 }
 
-void FNiagaraSystemInstance::UnbindParameters()
+void FNiagaraSystemInstance::UnbindParameters(bool bFromComplete)
 {
-	//TODO: Split unbind into internal and external.
-	//We need to unbind from external sources when we complete but internal can stay bound to improve reset times.
-
-	if (Component)
+	if (Component && !bFromComplete)
 	{
+		// NOTE: We don't unbind this on complete to improve reset times.
 		Component->GetOverrideParameters().Unbind(&InstanceParameters);
 	}
 
@@ -1040,20 +1041,12 @@ void FNiagaraSystemInstance::UnbindParameters()
 				Component->GetOverrideParameters().Unbind(&SystemSimulation->GetUpdateExecutionContext().Parameters);
 			}
 		}
-		else 
-		{
-			UNiagaraSystem* System = GetSystem();
-			if (System)
-			{
-				System->GetExposedParameters().Unbind(&SystemSimulation->GetSpawnExecutionContext().Parameters);
-				System->GetExposedParameters().Unbind(&SystemSimulation->GetSpawnExecutionContext().Parameters);
-			}
-		}
 	}
 
+	bAlreadyBound = bFromComplete && bAlreadyBound;
 	for (TSharedRef<FNiagaraEmitterInstance, ESPMode::ThreadSafe> Simulation : Emitters)
 	{
-		Simulation->UnbindParameters();
+		Simulation->UnbindParameters(bFromComplete);
 	}
 }
 
@@ -1577,31 +1570,30 @@ void FNiagaraSystemInstance::InitEmitters()
 
 	LocalBounds = FBox(FVector::ZeroVector, FVector::ZeroVector);
 
-	Emitters.Empty();
+	Emitters.Empty(false);
 	UNiagaraSystem* System = GetSystem();
 	if (System != nullptr)
 	{
 		const TArray<FNiagaraEmitterHandle>& EmitterHandles = System->GetEmitterHandles();
-		for (int32 EmitterIdx=0; EmitterIdx < System->GetEmitterHandles().Num(); ++EmitterIdx)
+		Emitters.Reserve(EmitterHandles.Num());
+		for (int32 EmitterIdx = 0; EmitterIdx < System->GetEmitterHandles().Num(); ++EmitterIdx)
 		{
+			TSharedRef<FNiagaraEmitterInstance, ESPMode::ThreadSafe> Sim = MakeShared<FNiagaraEmitterInstance, ESPMode::ThreadSafe>(this);
+			Sim->Init(EmitterIdx, ID);
+			if (System->bFixedBounds)
 			{
-				TSharedRef<FNiagaraEmitterInstance, ESPMode::ThreadSafe> Sim = MakeShared<FNiagaraEmitterInstance, ESPMode::ThreadSafe>(this);
-				Sim->Init(EmitterIdx, ID);
-				if (System->bFixedBounds)
-				{
-					Sim->SetSystemFixedBoundsOverride(System->GetFixedBounds());
-				}
-				Emitters.Add(Sim);
+				Sim->SetSystemFixedBoundsOverride(System->GetFixedBounds());
 			}
+			Emitters.Add(Sim);
 		}
 
-		for (TSharedRef<FNiagaraEmitterInstance, ESPMode::ThreadSafe> Simulation : Emitters)
+		for (TSharedRef<FNiagaraEmitterInstance, ESPMode::ThreadSafe>& Simulation : Emitters)
 		{
-			if (Simulation->GetCachedEmitter() != nullptr)
+			FNiagaraEmitterInstance& Inst = Simulation.Get();
+			if (Inst.GetCachedEmitter() != nullptr)
 			{
-				bHasGPUEmitters |= Simulation->GetCachedEmitter()->SimTarget == ENiagaraSimTarget::GPUComputeSim;
-
-				Simulation->PostInitSimulation();
+				bHasGPUEmitters |= Inst.GetCachedEmitter()->SimTarget == ENiagaraSimTarget::GPUComputeSim;
+				Inst.PostInitSimulation();
 			}
 		}
 
diff --git a/Engine/Plugins/FX/Niagara/Source/Niagara/Private/NiagaraSystemSimulation.cpp b/Engine/Plugins/FX/Niagara/Source/Niagara/Private/NiagaraSystemSimulation.cpp
--- a/Engine/Plugins/FX/Niagara/Source/Niagara/Private/NiagaraSystemSimulation.cpp
+++ b/Engine/Plugins/FX/Niagara/Source/Niagara/Private/NiagaraSystemSimulation.cpp
@@ -33,6 +33,13 @@ DECLARE_CYCLE_STAT(TEXT("System Sim Transfer Results [CNC]"), STAT_NiagaraSystem
 DECLARE_CYCLE_STAT(TEXT("System Sim Init [GT]"), STAT_NiagaraSystemSim_Init, STATGROUP_Niagara);
 DECLARE_CYCLE_STAT(TEXT("System Sim FastPath [CNC]"), STAT_NiagaraSystemSim_FastPathCNC, STATGROUP_Niagara);
 
+DECLARE_CYCLE_STAT(TEXT("System Sim Init (DataSets) [GT]"), STAT_NiagaraSystemSim_Init_DataSets, STATGROUP_Niagara);
+DECLARE_CYCLE_STAT(TEXT("System Sim Init (ExecContexts) [GT]"), STAT_NiagaraSystemSim_Init_ExecContexts, STATGROUP_Niagara);
+DECLARE_CYCLE_STAT(TEXT("System Sim Init (BindParams) [GT]"), STAT_NiagaraSystemSim_Init_BindParams, STATGROUP_Niagara);
+DECLARE_CYCLE_STAT(TEXT("System Sim Init (DatasetAccessors) [GT]"), STAT_NiagaraSystemSim_Init_DatasetAccessors, STATGROUP_Niagara);
+DECLARE_CYCLE_STAT(TEXT("System Sim Init (DirectBindings) [GT]"), STAT_NiagaraSystemSim_Init_DirectBindings, STATGROUP_Niagara);
+
+
 DECLARE_CYCLE_STAT(TEXT("ForcedWaitForAsync"), STAT_NiagaraSystemSim_ForceWaitForAsync, STATGROUP_Niagara);
 DECLARE_CYCLE_STAT(TEXT("ForcedWait Fake Stall"), STAT_NiagaraSystemSim_ForceWaitFakeStall, STATGROUP_Niagara);
 
@@ -260,134 +267,127 @@ bool FNiagaraSystemSimulation::Init(UNiagaraSystem* InSystem, UWorld* InWorld, b
 
 	if (bCanExecute)
 	{
-		//TODO: Move all layout data into the System!
-
-		//Initialize the main simulation dataset.
-#if !UE_BUILD_SHIPPING && !UE_BUILD_TEST 
-		MainDataSet.Init(FNiagaraDataSetID(), ENiagaraSimTarget::CPUSim, System->GetFullName());
-#else
-		MainDataSet.Init(FNiagaraDataSetID(), ENiagaraSimTarget::CPUSim);
-#endif
-		MainDataSet.AddVariables(System->GetSystemSpawnScript()->GetVMExecutableData().Attributes);
-		MainDataSet.AddVariables(System->GetSystemUpdateScript()->GetVMExecutableData().Attributes);
-		MainDataSet.Finalize();
-
-		//Initialize the main simulation dataset.
-#if !UE_BUILD_SHIPPING && !UE_BUILD_TEST 
-		SpawningDataSet.Init(FNiagaraDataSetID(), ENiagaraSimTarget::CPUSim, System->GetFullName());
-#else
-		SpawningDataSet.Init(FNiagaraDataSetID(), ENiagaraSimTarget::CPUSim);
-#endif
-		SpawningDataSet.AddVariables(System->GetSystemSpawnScript()->GetVMExecutableData().Attributes);
-		SpawningDataSet.AddVariables(System->GetSystemUpdateScript()->GetVMExecutableData().Attributes);
-		SpawningDataSet.Finalize();
-
-		//Initialize the dataset for paused systems.
-#if !UE_BUILD_SHIPPING && !UE_BUILD_TEST 
-		PausedInstanceData.Init(FNiagaraDataSetID(), ENiagaraSimTarget::CPUSim, System->GetFullName());
-#else
-		PausedInstanceData.Init(FNiagaraDataSetID(), ENiagaraSimTarget::CPUSim);
-#endif
-		PausedInstanceData.AddVariables(System->GetSystemSpawnScript()->GetVMExecutableData().Attributes);
-		PausedInstanceData.AddVariables(System->GetSystemUpdateScript()->GetVMExecutableData().Attributes);
-		PausedInstanceData.Finalize();
 
 		{
-#if !UE_BUILD_SHIPPING && !UE_BUILD_TEST 
-			SpawnInstanceParameterDataSet.Init(FNiagaraDataSetID(), ENiagaraSimTarget::CPUSim, System->GetFullName());
-#else
-			SpawnInstanceParameterDataSet.Init(FNiagaraDataSetID(), ENiagaraSimTarget::CPUSim);
-#endif
-			FNiagaraParameters* EngineParamsSpawn = System->GetSystemSpawnScript()->GetVMExecutableData().DataSetToParameters.Find(TEXT("Engine"));
-			if (EngineParamsSpawn != nullptr)
-			{
-				SpawnInstanceParameterDataSet.AddVariables(EngineParamsSpawn->Parameters);
-			}
-			SpawnInstanceParameterDataSet.Finalize();
-#if !UE_BUILD_SHIPPING && !UE_BUILD_TEST 
-			UpdateInstanceParameterDataSet.Init(FNiagaraDataSetID(), ENiagaraSimTarget::CPUSim, System->GetFullName());
-#else
-			UpdateInstanceParameterDataSet.Init(FNiagaraDataSetID(), ENiagaraSimTarget::CPUSim);
-#endif
-			FNiagaraParameters* EngineParamsUpdate = System->GetSystemUpdateScript()->GetVMExecutableData().DataSetToParameters.Find(TEXT("Engine"));
-			if (EngineParamsUpdate != nullptr)
-			{
-				UpdateInstanceParameterDataSet.AddVariables(EngineParamsUpdate->Parameters);
-			}
-			UpdateInstanceParameterDataSet.Finalize();
+			SCOPE_CYCLE_COUNTER(STAT_NiagaraSystemSim_Init_DataSets);
+
+			const FNiagaraSystemCompiledData& SystemCompiledData = System->GetSystemCompiledData();
+			//Initialize the main simulation dataset.
+			MainDataSet.Init(&SystemCompiledData.DataSetCompiledData);
+
+			//Initialize the main simulation dataset.
+			SpawningDataSet.Init(&SystemCompiledData.DataSetCompiledData);
+
+			//Initialize the dataset for paused systems.
+			PausedInstanceData.Init(&SystemCompiledData.DataSetCompiledData);
+			
+			SpawnInstanceParameterDataSet.Init(&SystemCompiledData.SpawnInstanceParamsDataSetCompiledData);
+
+			UpdateInstanceParameterDataSet.Init(&SystemCompiledData.UpdateInstanceParamsDataSetCompiledData);
+			
 		}
 
 		UNiagaraScript* SpawnScript = System->GetSystemSpawnScript();
 		UNiagaraScript* UpdateScript = System->GetSystemUpdateScript();
 
-		SpawnExecContext.Init(SpawnScript, ENiagaraSimTarget::CPUSim);
-		UpdateExecContext.Init(UpdateScript, ENiagaraSimTarget::CPUSim);
-
-		//Bind parameter collections.
-		for (UNiagaraParameterCollection* Collection : SpawnScript->GetCachedParameterCollectionReferences())
 		{
-			GetParameterCollectionInstance(Collection)->GetParameterStore().Bind(&SpawnExecContext.Parameters);
+			SCOPE_CYCLE_COUNTER(STAT_NiagaraSystemSim_Init_ExecContexts);
+
+
+
+			SpawnExecContext.Init(SpawnScript, ENiagaraSimTarget::CPUSim);
+			UpdateExecContext.Init(UpdateScript, ENiagaraSimTarget::CPUSim);
 		}
-		for (UNiagaraParameterCollection* Collection : UpdateScript->GetCachedParameterCollectionReferences())
+
+
 		{
-			GetParameterCollectionInstance(Collection)->GetParameterStore().Bind(&UpdateExecContext.Parameters);
-		}
+			SCOPE_CYCLE_COUNTER(STAT_NiagaraSystemSim_Init_BindParams);
+
+
+			//Bind parameter collections.
+			for (UNiagaraParameterCollection* Collection : SpawnScript->GetCachedParameterCollectionReferences())
+			{
+				GetParameterCollectionInstance(Collection)->GetParameterStore().Bind(&SpawnExecContext.Parameters);
+			}
+			for (UNiagaraParameterCollection* Collection : UpdateScript->GetCachedParameterCollectionReferences())
+			{
+				GetParameterCollectionInstance(Collection)->GetParameterStore().Bind(&UpdateExecContext.Parameters);
+			}
+
+			TArray<UNiagaraScript*, TInlineAllocator<2>> Scripts;
+			Scripts.Add(SpawnScript);
+			Scripts.Add(UpdateScript);
+			FNiagaraUtilities::CollectScriptDataInterfaceParameters(*System, Scripts, ScriptDefinedDataInterfaceParameters);
 
-		TArray<UNiagaraScript*> Scripts;
-		Scripts.Add(SpawnScript);
-		Scripts.Add(UpdateScript);
-		FNiagaraUtilities::CollectScriptDataInterfaceParameters(*System, Scripts, ScriptDefinedDataInterfaceParameters);
+			ScriptDefinedDataInterfaceParameters.Bind(&SpawnExecContext.Parameters);
+			ScriptDefinedDataInterfaceParameters.Bind(&UpdateExecContext.Parameters);
 
-		ScriptDefinedDataInterfaceParameters.Bind(&SpawnExecContext.Parameters);
-		ScriptDefinedDataInterfaceParameters.Bind(&UpdateExecContext.Parameters);
+			SpawnScript->RapidIterationParameters.Bind(&SpawnExecContext.Parameters);
+			UpdateScript->RapidIterationParameters.Bind(&UpdateExecContext.Parameters);
 
-		SpawnScript->RapidIterationParameters.Bind(&SpawnExecContext.Parameters);
-		UpdateScript->RapidIterationParameters.Bind(&UpdateExecContext.Parameters);
+			// If this simulation is not solo than we have bind the source system parameters to the system simulation contexts so that
+			// the system and emitter scripts use the default shared data interfaces.
+			if (!bIsSolo)
+			{
+				GetSystem()->GetExposedParameters().Bind(&GetSpawnExecutionContext().Parameters);
+				GetSystem()->GetExposedParameters().Bind(&GetUpdateExecutionContext().Parameters);
+			}
+		}
 
-		SystemExecutionStateAccessor.Create(&MainDataSet, FNiagaraVariable(EnumPtr, TEXT("System.ExecutionState")));
-		EmitterSpawnInfoAccessors.Reset();
-		EmitterExecutionStateAccessors.Reset();
-		EmitterSpawnInfoAccessors.SetNum(System->GetNumEmitters());
 
-		for (int32 EmitterIdx = 0; EmitterIdx < System->GetNumEmitters(); ++EmitterIdx)
 		{
-			FNiagaraEmitterHandle& EmitterHandle = System->GetEmitterHandle(EmitterIdx);
-			UNiagaraEmitter* Emitter = EmitterHandle.GetInstance();
-			if (Emitter)
-			{
-				FString EmitterName = Emitter->GetUniqueEmitterName();
-				EmitterExecutionStateAccessors.Emplace(MainDataSet, FNiagaraVariable(EnumPtr, *(EmitterName + TEXT(".ExecutionState"))));
-				const TArray<FNiagaraEmitterCompiledData>& EmitterCompiledData = System->GetEmitterCompiledData();
+			SCOPE_CYCLE_COUNTER(STAT_NiagaraSystemSim_Init_DatasetAccessors);
+
+			SystemExecutionStateAccessor.Create(&MainDataSet, FNiagaraVariable(EnumPtr, TEXT("System.ExecutionState")));
+			EmitterSpawnInfoAccessors.Reset();
+			EmitterExecutionStateAccessors.Reset();
+			EmitterSpawnInfoAccessors.SetNum(System->GetNumEmitters());
 
-				check(EmitterCompiledData.Num() == System->GetNumEmitters());
-				for (FName AttrName : EmitterCompiledData[EmitterIdx].SpawnAttributes)
+			for (int32 EmitterIdx = 0; EmitterIdx < System->GetNumEmitters(); ++EmitterIdx)
+			{
+				FNiagaraEmitterHandle& EmitterHandle = System->GetEmitterHandle(EmitterIdx);
+				UNiagaraEmitter* Emitter = EmitterHandle.GetInstance();
+				if (Emitter)
 				{
-					EmitterSpawnInfoAccessors[EmitterIdx].Emplace(MainDataSet, FNiagaraVariable(FNiagaraTypeDefinition(FNiagaraSpawnInfo::StaticStruct()), AttrName));
-				}
+					FString EmitterName = Emitter->GetUniqueEmitterName();
+					EmitterExecutionStateAccessors.Emplace(MainDataSet, FNiagaraVariable(EnumPtr, *(EmitterName + TEXT(".ExecutionState"))));
+					const TArray<FNiagaraEmitterCompiledData>& EmitterCompiledData = System->GetEmitterCompiledData();
+
+					check(EmitterCompiledData.Num() == System->GetNumEmitters());
+					for (FName AttrName : EmitterCompiledData[EmitterIdx].SpawnAttributes)
+					{
+						EmitterSpawnInfoAccessors[EmitterIdx].Emplace(MainDataSet, FNiagaraVariable(FNiagaraTypeDefinition(FNiagaraSpawnInfo::StaticStruct()), AttrName));
+					}
 
-				if (Emitter->bLimitDeltaTime)
+					if (Emitter->bLimitDeltaTime)
+					{
+						MaxDeltaTime = MaxDeltaTime.IsSet() ? FMath::Min(MaxDeltaTime.GetValue(), Emitter->MaxDeltaTimePerTick) : Emitter->MaxDeltaTimePerTick;
+					}
+				}
+				else
 				{
-					MaxDeltaTime = MaxDeltaTime.IsSet() ? FMath::Min(MaxDeltaTime.GetValue(), Emitter->MaxDeltaTimePerTick) : Emitter->MaxDeltaTimePerTick;
+					EmitterExecutionStateAccessors.Emplace();
 				}
 			}
-			else
-			{
-				EmitterExecutionStateAccessors.Emplace();
-			}
 		}
 
-		SpawnTimeParam.Init(SpawnExecContext.Parameters, SYS_PARAM_ENGINE_TIME);
-		UpdateTimeParam.Init(UpdateExecContext.Parameters, SYS_PARAM_ENGINE_TIME);
-		SpawnDeltaTimeParam.Init(SpawnExecContext.Parameters, SYS_PARAM_ENGINE_DELTA_TIME);
-		UpdateDeltaTimeParam.Init(UpdateExecContext.Parameters, SYS_PARAM_ENGINE_DELTA_TIME);
-		SpawnInvDeltaTimeParam.Init(SpawnExecContext.Parameters, SYS_PARAM_ENGINE_INV_DELTA_TIME);
-		UpdateInvDeltaTimeParam.Init(UpdateExecContext.Parameters, SYS_PARAM_ENGINE_INV_DELTA_TIME);
-		SpawnNumSystemInstancesParam.Init(SpawnExecContext.Parameters, SYS_PARAM_ENGINE_NUM_SYSTEM_INSTANCES);
-		UpdateNumSystemInstancesParam.Init(UpdateExecContext.Parameters, SYS_PARAM_ENGINE_NUM_SYSTEM_INSTANCES);
-		SpawnGlobalSpawnCountScaleParam.Init(SpawnExecContext.Parameters, SYS_PARAM_ENGINE_GLOBAL_SPAWN_COUNT_SCALE);
-		UpdateGlobalSpawnCountScaleParam.Init(UpdateExecContext.Parameters, SYS_PARAM_ENGINE_GLOBAL_SPAWN_COUNT_SCALE);
-		SpawnGlobalSystemCountScaleParam.Init(SpawnExecContext.Parameters, SYS_PARAM_ENGINE_GLOBAL_SYSTEM_COUNT_SCALE);
-		UpdateGlobalSystemCountScaleParam.Init(UpdateExecContext.Parameters, SYS_PARAM_ENGINE_GLOBAL_SYSTEM_COUNT_SCALE);
+
+		{
+			SCOPE_CYCLE_COUNTER(STAT_NiagaraSystemSim_Init_DirectBindings);
+
+			SpawnTimeParam.Init(SpawnExecContext.Parameters, SYS_PARAM_ENGINE_TIME);
+			UpdateTimeParam.Init(UpdateExecContext.Parameters, SYS_PARAM_ENGINE_TIME);
+			SpawnDeltaTimeParam.Init(SpawnExecContext.Parameters, SYS_PARAM_ENGINE_DELTA_TIME);
+			UpdateDeltaTimeParam.Init(UpdateExecContext.Parameters, SYS_PARAM_ENGINE_DELTA_TIME);
+			SpawnInvDeltaTimeParam.Init(SpawnExecContext.Parameters, SYS_PARAM_ENGINE_INV_DELTA_TIME);
+			UpdateInvDeltaTimeParam.Init(UpdateExecContext.Parameters, SYS_PARAM_ENGINE_INV_DELTA_TIME);
+			SpawnNumSystemInstancesParam.Init(SpawnExecContext.Parameters, SYS_PARAM_ENGINE_NUM_SYSTEM_INSTANCES);
+			UpdateNumSystemInstancesParam.Init(UpdateExecContext.Parameters, SYS_PARAM_ENGINE_NUM_SYSTEM_INSTANCES);
+			SpawnGlobalSpawnCountScaleParam.Init(SpawnExecContext.Parameters, SYS_PARAM_ENGINE_GLOBAL_SPAWN_COUNT_SCALE);
+			UpdateGlobalSpawnCountScaleParam.Init(UpdateExecContext.Parameters, SYS_PARAM_ENGINE_GLOBAL_SPAWN_COUNT_SCALE);
+			SpawnGlobalSystemCountScaleParam.Init(SpawnExecContext.Parameters, SYS_PARAM_ENGINE_GLOBAL_SYSTEM_COUNT_SCALE);
+			UpdateGlobalSystemCountScaleParam.Init(UpdateExecContext.Parameters, SYS_PARAM_ENGINE_GLOBAL_SYSTEM_COUNT_SCALE);
+		}
 	}
 
 	return true;
@@ -429,8 +429,10 @@ UNiagaraParameterCollectionInstance* FNiagaraSystemSimulation::GetParameterColle
 	//If no explicit override from the system, just get the current instance set on the world.
 	if (!Ret)
 	{
-		FNiagaraWorldManager* WorldMan = FNiagaraWorldManager::Get(World);
-		Ret = WorldMan->GetParameterCollection(Collection);
+		if(FNiagaraWorldManager* WorldMan = FNiagaraWorldManager::Get(World))
+		{
+			Ret = WorldMan->GetParameterCollection(Collection);
+		}
 	}
 
 	return Ret;
diff --git a/Engine/Plugins/FX/Niagara/Source/Niagara/Public/NiagaraSystemInstance.h b/Engine/Plugins/FX/Niagara/Source/Niagara/Public/NiagaraSystemInstance.h
--- a/Engine/Plugins/FX/Niagara/Source/Niagara/Public/NiagaraSystemInstance.h
+++ b/Engine/Plugins/FX/Niagara/Source/Niagara/Public/NiagaraSystemInstance.h
@@ -74,7 +74,7 @@ public:
 
 	//void RebindParameterCollection(UNiagaraParameterCollectionInstance* OldInstance, UNiagaraParameterCollectionInstance* NewInstance);
 	void BindParameters();
-	void UnbindParameters();
+	void UnbindParameters(bool bFromComplete = false);
 
 	FORCEINLINE FNiagaraParameterStore& GetInstanceParameters() { return InstanceParameters; }
 	
@@ -377,6 +377,8 @@ private:
 
 	uint32 bDataInterfacesInitialized : 1;
 
+	uint32 bAlreadyBound : 1;
+
 	/** True if we have async work in flight. */
 	volatile bool bAsyncWorkInProgress;
 
diff --git a/Engine/Plugins/FX/Niagara/Source/NiagaraCore/Private/NiagaraCustomVersion.cpp b/Engine/Plugins/FX/Niagara/Source/NiagaraCore/Private/NiagaraCustomVersion.cpp
--- a/Engine/Plugins/FX/Niagara/Source/NiagaraCore/Private/NiagaraCustomVersion.cpp
+++ b/Engine/Plugins/FX/Niagara/Source/NiagaraCore/Private/NiagaraCustomVersion.cpp
@@ -8,4 +8,4 @@ const FGuid FNiagaraCustomVersion::GUID(0xFCF57AFA, 0x50764283, 0xB9A9E658, 0xFF
 // Register the custom version with core
 FCustomVersionRegistration GRegisterNiagaraCustomVersion(FNiagaraCustomVersion::GUID, FNiagaraCustomVersion::LatestVersion, TEXT("NiagaraVer"));
 
-const FGuid FNiagaraCustomVersion::LatestScriptCompileVersion(0x66E6CB72, 0xE08011E9, 0x94E21072, 0xE014E1EA);
+const FGuid FNiagaraCustomVersion::LatestScriptCompileVersion(0x908C3A78, 0x6E39F14C, 0x8E1C1A93, 0xDD667695);
\ No newline at end of file
diff --git a/Engine/Plugins/FX/Niagara/Source/NiagaraEditor/Private/Widgets/SNiagaraSpreadsheetView.cpp b/Engine/Plugins/FX/Niagara/Source/NiagaraEditor/Private/Widgets/SNiagaraSpreadsheetView.cpp
--- a/Engine/Plugins/FX/Niagara/Source/NiagaraEditor/Private/Widgets/SNiagaraSpreadsheetView.cpp
+++ b/Engine/Plugins/FX/Niagara/Source/NiagaraEditor/Private/Widgets/SNiagaraSpreadsheetView.cpp
@@ -710,7 +710,7 @@ void SNiagaraSpreadsheetView::SetTarget(UNiagaraComponent* InComponent)
 
 	for (int32 i = 0; i < (int32)UIMax; i++)
 	{
-		CaptureData[i].DataSet.Init(FNiagaraDataSetID(), ENiagaraSimTarget::CPUSim, InComponent->GetFullName() + TEXT("/Spreadsheet"));
+		CaptureData[i].DataSet.Init(nullptr);
 		CaptureData[i].InputParams.Reset();
 		CaptureData[i].CaptureData.Reset();
 	}
@@ -902,7 +902,7 @@ void SNiagaraSpreadsheetView::SystemSelectionChanged(UNiagaraSystemSelectionView
 	// Need to reset the attributes list...
 	for (int32 i = 0; i < (int32)UIMax; i++)
 	{
-		CaptureData[i].DataSet.Init(FNiagaraDataSetID(), ENiagaraSimTarget::CPUSim, TEXT("/Spreadsheet"));
+		CaptureData[i].DataSet.Init(nullptr);
 		CaptureData[i].SupportedInputIndices.SetNum(0);
 		CaptureData[i].SupportedOutputIndices.SetNum(0);
 		CaptureData[i].OutputsListView->RequestTreeRefresh();
@@ -1035,8 +1035,7 @@ void SNiagaraSpreadsheetView::HandleTimeChange()
 					else
 					{
 						CaptureData[i].CaptureData.Reset();
-						CaptureData[i].DataSet.Init(FNiagaraDataSetID(), ENiagaraSimTarget::CPUSim, TargetComponent->GetFullName() + TEXT("/Spreadsheet"));
-						CaptureData[i].DataSet.Finalize();
+						CaptureData[i].DataSet.Init(nullptr);
 
 						ResetColumns((EUITab)i);
 						ResetEntries((EUITab)i);
