Manually merging 9325681 to 11.10
Compile fix for non STATS builds.

diff --git a/Engine/Plugins/FX/Niagara/Source/Niagara/Private/NiagaraComponent.cpp b/Engine/Plugins/FX/Niagara/Source/Niagara/Private/NiagaraComponent.cpp
--- a/Engine/Plugins/FX/Niagara/Source/Niagara/Private/NiagaraComponent.cpp
+++ b/Engine/Plugins/FX/Niagara/Source/Niagara/Private/NiagaraComponent.cpp
@@ -123,7 +123,11 @@ FNiagaraSceneProxy::FNiagaraSceneProxy(const UNiagaraComponent* InComponent)
 		CreateRenderers(InComponent);
 		bAlwaysHasVelocity = true;
 		Batcher = SystemInst->GetBatcher();
+
+#if STATS
 		SystemStatID = InComponent->GetAsset()->GetStatID(false, false);
+#endif
+
 	}
 }
 
@@ -311,7 +315,9 @@ void FNiagaraSceneProxy::GetDynamicMeshElements(const TArray<const FSceneView*>&
 	SCOPE_CYCLE_COUNTER(STAT_NiagaraOverview_RT);
 	SCOPE_CYCLE_COUNTER(STAT_NiagaraComponentGetDynamicMeshElements);
 
+#if STATS
 	FScopeCycleCounter SystemStatCounter(SystemStatID);
+#endif
 
 	for (int32 RendererIdx : RendererDrawOrder)
 	{
