Moved initialization of ScriptExecutionParamStoreCPU and ScriptExecutionParamStoreGPU to postload.

diff --git a/Engine/Plugins/FX/Niagara/Source/Niagara/Private/NiagaraScript.cpp b/Engine/Plugins/FX/Niagara/Source/Niagara/Private/NiagaraScript.cpp
--- a/Engine/Plugins/FX/Niagara/Source/Niagara/Private/NiagaraScript.cpp
+++ b/Engine/Plugins/FX/Niagara/Source/Niagara/Private/NiagaraScript.cpp
@@ -662,6 +662,31 @@ void UNiagaraScript::PostLoad()
 
 	//FNiagaraUtilities::DumpHLSLText(RapidIterationParameters.ToString(), *GetPathName());
 
+	// Init ScriptExecutionParamStoreCPU or ScriptExecutionParamStoreGPU to avoid first time cost
+	if (FPlatformProperties::RequiresCookedData())
+	{
+	    switch (Usage)
+	    {
+	    case ENiagaraScriptUsage::ParticleSpawnScript:
+	    case ENiagaraScriptUsage::ParticleSpawnScriptInterpolated:
+	    case ENiagaraScriptUsage::ParticleUpdateScript:
+	    case ENiagaraScriptUsage::ParticleEventScript:
+	    case ENiagaraScriptUsage::ParticleGPUComputeScript:
+		    if (UNiagaraEmitter* OwningEmitter = GetTypedOuter<UNiagaraEmitter>())
+		    {
+				GetExecutionReadyParameterStore(OwningEmitter->SimTarget);
+		    }
+		    break;
+	    case ENiagaraScriptUsage::EmitterSpawnScript:
+	    case ENiagaraScriptUsage::EmitterUpdateScript:
+	    case ENiagaraScriptUsage::SystemSpawnScript:
+	    case ENiagaraScriptUsage::SystemUpdateScript:
+			GetExecutionReadyParameterStore(ENiagaraSimTarget::CPUSim);
+		    break;
+	    default:
+		    break;
+		};
+	}
 }
 
 bool UNiagaraScript::IsReadyToRun(ENiagaraSimTarget SimTarget) const
