Fixed bad merge disabling Niagara GPU particles

diff --git a/Engine/Plugins/FX/Niagara/Source/Niagara/Private/NiagaraEmitterInstanceBatcher.cpp b/Engine/Plugins/FX/Niagara/Source/Niagara/Private/NiagaraEmitterInstanceBatcher.cpp
--- a/Engine/Plugins/FX/Niagara/Source/Niagara/Private/NiagaraEmitterInstanceBatcher.cpp
+++ b/Engine/Plugins/FX/Niagara/Source/Niagara/Private/NiagaraEmitterInstanceBatcher.cpp
@@ -456,7 +456,7 @@ void NiagaraEmitterInstanceBatcher::DispatchAllOnCompute(FOverlappableTicks& Ove
 					FNiagaraComputeExecutionContext::TickCounter++;
 
 					// run shader, sim and spawn in a single dispatch
-				DispatchMultipleStages(*Tick, &Instance, RHICmdList, ViewUniformBuffer, Context->GPUScript_RT->GetShader());
+					DispatchMultipleStages(*Tick, &Instance, RHICmdList, ViewUniformBuffer, Context->GPUScript_RT->GetShader());
 
 					FNiagaraDataBuffer* CurrentData = Instance.CurrentData;
 					if (bSetReadback && Tick->bIsFinalTick)
@@ -1157,98 +1157,97 @@ void NiagaraEmitterInstanceBatcher::Run(const FNiagaraGPUSystemTick& Tick, const
 		{
 			RHICmdList.SetShaderParameter(Shader->GetComputeShader(), Shader->IterationInterfaceCount.GetBufferIndex(), Shader->IterationInterfaceCount.GetBaseIndex(), Shader->IterationInterfaceCount.GetNumBytes(), &TotalNumInstances);					// 0, except if several stages are defined
 		}
+	}
 
+	uint32 NumThreadGroups = 1;
+	if (TotalNumInstances > NIAGARA_COMPUTE_THREADGROUP_SIZE)
+	{
+		NumThreadGroups = FMath::Min(NIAGARA_MAX_COMPUTE_THREADGROUPS, FMath::DivideAndRoundUp(TotalNumInstances, NIAGARA_COMPUTE_THREADGROUP_SIZE));
+	}
 
-		uint32 NumThreadGroups = 1;
-		if (TotalNumInstances > NIAGARA_COMPUTE_THREADGROUP_SIZE)
-		{
-			NumThreadGroups = FMath::Min(NIAGARA_MAX_COMPUTE_THREADGROUPS, FMath::DivideAndRoundUp(TotalNumInstances, NIAGARA_COMPUTE_THREADGROUP_SIZE));
-		}
-
-		// setup script parameters
-		if (CBufferLayout.ConstantBufferSize)
-		{
-			check(CBufferLayout.Resources.Num() == 0);
-			const uint8* ParamData = Instance->ParamData;
-			FUniformBufferRHIRef CBuffer = RHICreateUniformBuffer(ParamData, CBufferLayout, EUniformBufferUsage::UniformBuffer_SingleDraw);
-			RHICmdList.SetShaderUniformBuffer(Shader->GetComputeShader(), Shader->EmitterConstantBufferParam.GetBaseIndex(), CBuffer);
-		}
-		else
-		{
-			ensure(!Shader->EmitterConstantBufferParam.IsBound());
-		}
+	// setup script parameters
+	if (CBufferLayout.ConstantBufferSize)
+	{
+		check(CBufferLayout.Resources.Num() == 0);
+		const uint8* ParamData = Instance->ParamData;
+		FUniformBufferRHIRef CBuffer = RHICreateUniformBuffer(ParamData, CBufferLayout, EUniformBufferUsage::UniformBuffer_SingleDraw);
+		RHICmdList.SetShaderUniformBuffer(Shader->GetComputeShader(), Shader->EmitterConstantBufferParam.GetBaseIndex(), CBuffer);
+	}
+	else
+	{
+		ensure(!Shader->EmitterConstantBufferParam.IsBound());
+	}
 
-		// #todo(dmp): temporary hack -- unbind UAVs if we have a valid iteration DI.  This way, when we are outputting with a different iteration count, we don't
-		// mess up particle state
-		if (IterationInterface)
-		{
-			CurrentData.UnsetShaderParams(Shader, RHICmdList);
-			DestinationData.UnsetShaderParams(Shader, RHICmdList);
-		}
+	// #todo(dmp): temporary hack -- unbind UAVs if we have a valid iteration DI.  This way, when we are outputting with a different iteration count, we don't
+	// mess up particle state
+	if (IterationInterface)
+	{
+		CurrentData.UnsetShaderParams(Shader, RHICmdList);
+		DestinationData.UnsetShaderParams(Shader, RHICmdList);
+	}
 
-		//UE_LOG(LogNiagara, Log, TEXT("Num Instance : %d | Num Group : %d | Spawned Istance : %d | Start Instance : %d | Num Indices : %d | Stage Index : %d"), 
-			//TotalNumInstances, NumThreadGroups, InstancesToSpawnThisFrame, UpdateStartInstance, Context->NumIndicesPerInstance, ShaderStageIndex);
+	//UE_LOG(LogNiagara, Log, TEXT("Num Instance : %d | Num Group : %d | Spawned Istance : %d | Start Instance : %d | Num Indices : %d | Stage Index : %d"), 
+		//TotalNumInstances, NumThreadGroups, InstancesToSpawnThisFrame, UpdateStartInstance, Context->NumIndicesPerInstance, ShaderStageIndex);
 
-		// Dispatch, if anything needs to be done
-		if (TotalNumInstances)
-		{
-			SCOPED_DRAW_EVENTF(RHICmdList, NiagaraGPUSimulationCS, TEXT("Niagara Gpu Sim - %s - NumInstances: %u - StageNumber: %u"),
-				*Context->DebugSimName,
-				TotalNumInstances,
-				ShaderStageIndex);
-			DispatchComputeShader(RHICmdList, Shader, NumThreadGroups, 1, 1);
-		}
+	// Dispatch, if anything needs to be done
+	if (TotalNumInstances)
+	{
+		SCOPED_DRAW_EVENTF(RHICmdList, NiagaraGPUSimulationCS, TEXT("Niagara Gpu Sim - %s - NumInstances: %u - StageNumber: %u"),
+			*Context->DebugSimName,
+			TotalNumInstances,
+			ShaderStageIndex);
+		DispatchComputeShader(RHICmdList, Shader, NumThreadGroups, 1, 1);
+	}
 
-		// reset iteration count
-		if (IterationInterface)
-		{
-			RHICmdList.SetShaderParameter(Shader->GetComputeShader(), Shader->IterationInterfaceCount.GetBufferIndex(), Shader->IterationInterfaceCount.GetBaseIndex(), Shader->IterationInterfaceCount.GetNumBytes(), &DefaultIterationCount);					// 0, except if several stages are defined
-		}
+	// reset iteration count
+	if (IterationInterface)
+	{
+		RHICmdList.SetShaderParameter(Shader->GetComputeShader(), Shader->IterationInterfaceCount.GetBufferIndex(), Shader->IterationInterfaceCount.GetBaseIndex(), Shader->IterationInterfaceCount.GetNumBytes(), &DefaultIterationCount);					// 0, except if several stages are defined
+	}
 
 #if WITH_EDITORONLY_DATA
-		// Check to see if we need to queue up a debug dump..
-		if (Context->DebugInfo.IsValid())
+	// Check to see if we need to queue up a debug dump..
+	if (Context->DebugInfo.IsValid())
+	{
+		//UE_LOG(LogNiagara, Warning, TEXT("Queued up!"));
+
+		if (!Context->GPUDebugDataReadbackFloat && !Context->GPUDebugDataReadbackInt && !Context->GPUDebugDataReadbackCounts && DestinationData.GetGPUInstanceCountBufferOffset() != INDEX_NONE && ShaderStageIndex == Context->MaxUpdateIterations - 1)
 		{
-			//UE_LOG(LogNiagara, Warning, TEXT("Queued up!"));
+			Context->GPUDebugDataFloatSize = 0;
+			Context->GPUDebugDataIntSize = 0;
+			Context->GPUDebugDataFloatStride = 0;
+			Context->GPUDebugDataIntStride = 0;
 
-			if (!Context->GPUDebugDataReadbackFloat && !Context->GPUDebugDataReadbackInt && !Context->GPUDebugDataReadbackCounts && DestinationData.GetGPUInstanceCountBufferOffset() != INDEX_NONE && ShaderStageIndex == Context->MaxUpdateIterations - 1)
+			if (DestinationData.GetGPUBufferFloat().NumBytes > 0)
 			{
-				Context->GPUDebugDataFloatSize = 0;
-				Context->GPUDebugDataIntSize = 0;
-				Context->GPUDebugDataFloatStride = 0;
-				Context->GPUDebugDataIntStride = 0;
-
-				if (DestinationData.GetGPUBufferFloat().NumBytes > 0)
-				{
-					static const FName ReadbackFloatName(TEXT("Niagara GPU Debug Info Float Emitter Readback"));
-					Context->GPUDebugDataReadbackFloat = new FRHIGPUBufferReadback(ReadbackFloatName);
-					Context->GPUDebugDataReadbackFloat->EnqueueCopy(RHICmdList, DestinationData.GetGPUBufferFloat().Buffer);
-					Context->GPUDebugDataFloatSize = DestinationData.GetGPUBufferFloat().NumBytes;
-					Context->GPUDebugDataFloatStride = DestinationData.GetFloatStride();
-				}
-
-				if (DestinationData.GetGPUBufferInt().NumBytes > 0)
-				{
-					static const FName ReadbackIntName(TEXT("Niagara GPU Debug Info Int Emitter Readback"));
-					Context->GPUDebugDataReadbackInt = new FRHIGPUBufferReadback(ReadbackIntName);
-					Context->GPUDebugDataReadbackInt->EnqueueCopy(RHICmdList, DestinationData.GetGPUBufferInt().Buffer);
-					Context->GPUDebugDataIntSize = DestinationData.GetGPUBufferInt().NumBytes;
-					Context->GPUDebugDataIntStride = DestinationData.GetInt32Stride();
-				}
+				static const FName ReadbackFloatName(TEXT("Niagara GPU Debug Info Float Emitter Readback"));
+				Context->GPUDebugDataReadbackFloat = new FRHIGPUBufferReadback(ReadbackFloatName);
+				Context->GPUDebugDataReadbackFloat->EnqueueCopy(RHICmdList, DestinationData.GetGPUBufferFloat().Buffer);
+				Context->GPUDebugDataFloatSize = DestinationData.GetGPUBufferFloat().NumBytes;
+				Context->GPUDebugDataFloatStride = DestinationData.GetFloatStride();
+			}
 
-				static const FName ReadbackCountsName(TEXT("Niagara GPU Emitter Readback"));
-				Context->GPUDebugDataReadbackCounts = new FRHIGPUBufferReadback(ReadbackCountsName);
-				Context->GPUDebugDataReadbackCounts->EnqueueCopy(RHICmdList, GPUInstanceCounterManager.GetInstanceCountBuffer().Buffer);
-				Context->GPUDebugDataCountOffset = DestinationData.GetGPUInstanceCountBufferOffset();
+			if (DestinationData.GetGPUBufferInt().NumBytes > 0)
+			{
+				static const FName ReadbackIntName(TEXT("Niagara GPU Debug Info Int Emitter Readback"));
+				Context->GPUDebugDataReadbackInt = new FRHIGPUBufferReadback(ReadbackIntName);
+				Context->GPUDebugDataReadbackInt->EnqueueCopy(RHICmdList, DestinationData.GetGPUBufferInt().Buffer);
+				Context->GPUDebugDataIntSize = DestinationData.GetGPUBufferInt().NumBytes;
+				Context->GPUDebugDataIntStride = DestinationData.GetInt32Stride();
 			}
+
+			static const FName ReadbackCountsName(TEXT("Niagara GPU Emitter Readback"));
+			Context->GPUDebugDataReadbackCounts = new FRHIGPUBufferReadback(ReadbackCountsName);
+			Context->GPUDebugDataReadbackCounts->EnqueueCopy(RHICmdList, GPUInstanceCounterManager.GetInstanceCountBuffer().Buffer);
+			Context->GPUDebugDataCountOffset = DestinationData.GetGPUInstanceCountBufferOffset();
 		}
+	}
 #endif // WITH_EDITORONLY_DATA
 
-		// Unset UAV parameters and transition resources (TODO: resource transition should be moved to the renderer)
-		// 
-		UnsetDataInterfaceParameters(DataInterfaceProxies, Shader, RHICmdList, Instance, Tick);
-		CurrentData.UnsetShaderParams(Shader, RHICmdList);
-		DestinationData.UnsetShaderParams(Shader, RHICmdList);
-		Shader->InstanceCountsParam.UnsetUAV(RHICmdList, Shader->GetComputeShader());
-	}
+	// Unset UAV parameters and transition resources (TODO: resource transition should be moved to the renderer)
+	// 
+	UnsetDataInterfaceParameters(DataInterfaceProxies, Shader, RHICmdList, Instance, Tick);
+	CurrentData.UnsetShaderParams(Shader, RHICmdList);
+	DestinationData.UnsetShaderParams(Shader, RHICmdList);
+	Shader->InstanceCountsParam.UnsetUAV(RHICmdList, Shader->GetComputeShader());
 }
