Precompiled the binding data for ScriptExecutionParamStore (during cook) to remove the cost related to ScriptExecutionParamStoreCPU.InitFromOwningScript(), or the GPU one.

diff --git a/Engine/Plugins/FX/Niagara/Source/Niagara/Classes/NiagaraScript.h b/Engine/Plugins/FX/Niagara/Source/Niagara/Classes/NiagaraScript.h
--- a/Engine/Plugins/FX/Niagara/Source/Niagara/Classes/NiagaraScript.h
+++ b/Engine/Plugins/FX/Niagara/Source/Niagara/Classes/NiagaraScript.h
@@ -479,6 +479,7 @@ public:
 #endif
 
 	//~ Begin UObject interface
+	void PreSave(const class ITargetPlatform* TargetPlatform) override;
 	void Serialize(FArchive& Ar)override;
 	virtual void PostLoad() override;
 	virtual bool IsDestructionThreadSafe() const override { return false; }
@@ -586,14 +587,28 @@ public:
 
 	FNiagaraScriptExecutionParameterStore* GetExecutionReadyParameterStore(ENiagaraSimTarget SimTarget);
 	void InvalidateExecutionReadyParameterStores();
+
 private:
+
 	bool LegacyCanBeRunOnGpu()const;
 
+	/** Return the expected SimTarget for this script. Only returns a valid target if there is valid data to run with. */
+	TOptional<ENiagaraSimTarget> GetSimTarget() const;
+
+#if WITH_EDITORONLY_DATA
 	UPROPERTY(Transient)
 	FNiagaraScriptExecutionParameterStore ScriptExecutionParamStoreCPU;
 
 	UPROPERTY(Transient)
 	FNiagaraScriptExecutionParameterStore ScriptExecutionParamStoreGPU;
+#endif // WITH_EDITORONLY_DATA
+
+	/** The equivalent of ScriptExecutionParamStoreCPU (or GPU) cooked for the given platform.*/
+	UPROPERTY()
+	FNiagaraScriptExecutionParameterStore ScriptExecutionParamStore;
+	/** The cooked binding data between ScriptExecutionParamStore and RapidIterationParameters.*/
+	UPROPERTY()
+	TArray<FNiagaraBoundParameter> ScriptExecutionBoundParameters;
 
 #if WITH_EDITORONLY_DATA
 	class UNiagaraSystem* FindRootSystem();
diff --git a/Engine/Plugins/FX/Niagara/Source/Niagara/Private/NiagaraParameterStore.cpp b/Engine/Plugins/FX/Niagara/Source/Niagara/Private/NiagaraParameterStore.cpp
--- a/Engine/Plugins/FX/Niagara/Source/Niagara/Private/NiagaraParameterStore.cpp
+++ b/Engine/Plugins/FX/Niagara/Source/Niagara/Private/NiagaraParameterStore.cpp
@@ -153,7 +153,7 @@ FNiagaraParameterStore::~FNiagaraParameterStore()
 	Bindings.Empty();
 }
 
-void FNiagaraParameterStore::Bind(FNiagaraParameterStore* DestStore)
+void FNiagaraParameterStore::Bind(FNiagaraParameterStore* DestStore, const FNiagaraBoundParameterArray* BoundParameters)
 {
 	check(DestStore);
 	SCOPE_CYCLE_COUNTER(STAT_NiagaraParameterStoreBind);
@@ -161,7 +161,7 @@ void FNiagaraParameterStore::Bind(FNiagaraParameterStore* DestStore)
 	{
 		// Bind the parameter stores only if they have variables in common.
 		FNiagaraParameterStoreBinding HeapBinding;
-		if  (HeapBinding.Initialize(DestStore, this))
+		if  (HeapBinding.Initialize(DestStore, this, BoundParameters))
 		{
 			FNiagaraParameterStoreBinding& Binding = Bindings.FindOrAdd(DestStore);
 			FMemory::Memswap(&Binding, &HeapBinding, sizeof(FNiagaraParameterStoreBinding));
@@ -169,35 +169,9 @@ void FNiagaraParameterStore::Bind(FNiagaraParameterStore* DestStore)
 	}
 }
 
-bool FNiagaraParameterStoreBinding::BindParameters(FNiagaraParameterStore* DestStore, FNiagaraParameterStore* SrcStore)
+template <typename TVisitor>
+void FNiagaraParameterStoreBinding::MatchParameters(FNiagaraParameterStore* DestStore, FNiagaraParameterStore* SrcStore, TVisitor Visitor)
 {
-	InterfaceBindings.Reset();
-	ParameterBindings.Reset();
-	UObjectBindings.Reset();
-
-	bool bAnyBinding = false;
-
-	auto BindVariable = [&](const FNiagaraVariable& InParameter, int32 SrcOffset, int32 DestOffset)
-	{
-		if (SrcOffset != INDEX_NONE && DestOffset != INDEX_NONE)
-		{
-			bAnyBinding = true;
-
-			if (InParameter.IsDataInterface())
-			{
-				InterfaceBindings.Add(FInterfaceBinding(SrcOffset, DestOffset));
-			}
-			else if (InParameter.IsUObject())
-			{
-				UObjectBindings.Add(FUObjectBinding(SrcOffset, DestOffset));
-			}
-			else
-			{
-				ParameterBindings.Add(FParameterBinding(SrcOffset, DestOffset, InParameter.GetSizeInBytes()));
-			}
-		}
-	};
-
 	const int32 SrcNum = SrcStore->GetSortedParameterOffsets().Num();
 	const int32 DestNum = DestStore->GetSortedParameterOffsets().Num();
 	const int32 BinarySearchComplexity = FMath::Min<int32>(SrcNum, DestNum) * FMath::RoundToInt(FMath::Log2((float)FMath::Max<int32>(SrcNum, DestNum)));
@@ -224,7 +198,7 @@ bool FNiagaraParameterStoreBinding::BindParameters(FNiagaraParameterStore* DestS
 			}
 			else // CompValue == 0
 			{
-				BindVariable(SrcParamWithOffset, SrcParamWithOffset.Offset, DestParamWithOffset.Offset);
+				Visitor(SrcParamWithOffset, SrcParamWithOffset.Offset, DestParamWithOffset.Offset);
 				++SrcIndex;
 				++DestIndex;
 			}
@@ -237,15 +211,79 @@ bool FNiagaraParameterStoreBinding::BindParameters(FNiagaraParameterStore* DestS
 	{
 		for (const FNiagaraVariableWithOffset& ParamWithOffset : DestStore->GetSortedParameterOffsets())
 		{
-			BindVariable(ParamWithOffset, SrcStore->IndexOf(ParamWithOffset), ParamWithOffset.Offset);
+			Visitor(ParamWithOffset, SrcStore->IndexOf(ParamWithOffset), ParamWithOffset.Offset);
 		}
 	}
 	else
 	{
 		for (const FNiagaraVariableWithOffset& ParamWithOffset : SrcStore->GetSortedParameterOffsets())
 		{
-			BindVariable(ParamWithOffset, ParamWithOffset.Offset, DestStore->IndexOf(ParamWithOffset));
+			Visitor(ParamWithOffset, ParamWithOffset.Offset, DestStore->IndexOf(ParamWithOffset));
+		}
+	}
+}
+
+void FNiagaraParameterStoreBinding::GetBindingData(FNiagaraParameterStore* DestStore, FNiagaraParameterStore* SrcStore, FNiagaraBoundParameterArray& OutBoundParameters)
+{
+	OutBoundParameters.Empty();
+
+	auto AddVariable = [&](const FNiagaraVariable& InParameter, int32 SrcOffset, int32 DestOffset)
+	{
+		if (SrcOffset != INDEX_NONE && DestOffset != INDEX_NONE)
+		{
+			FNiagaraBoundParameter BoundParameter;
+			BoundParameter.Parameter = InParameter;
+			BoundParameter.SrcOffset = SrcOffset;
+			BoundParameter.DestOffset = DestOffset;
+
+			OutBoundParameters.Add(BoundParameter);
+		}
+	};
+
+	MatchParameters(DestStore, SrcStore, AddVariable);
+}
+
+bool FNiagaraParameterStoreBinding::BindParameters(FNiagaraParameterStore* DestStore, FNiagaraParameterStore* SrcStore, const FNiagaraBoundParameterArray* BoundParameters)
+{
+	InterfaceBindings.Reset();
+	ParameterBindings.Reset();
+	UObjectBindings.Reset();
+
+	bool bAnyBinding = false;
+
+	auto BindVariable = [&](const FNiagaraVariable& InParameter, int32 SrcOffset, int32 DestOffset)
+	{
+		if (SrcOffset != INDEX_NONE && DestOffset != INDEX_NONE)
+		{
+			bAnyBinding = true;
+
+			if (InParameter.IsDataInterface())
+			{
+				InterfaceBindings.Add(FInterfaceBinding(SrcOffset, DestOffset));
+			}
+			else if (InParameter.IsUObject())
+			{
+				UObjectBindings.Add(FUObjectBinding(SrcOffset, DestOffset));
+			}
+			else
+			{
+				ParameterBindings.Add(FParameterBinding(SrcOffset, DestOffset, InParameter.GetSizeInBytes()));
+			}
+		}
+	};
+
+	if (!BoundParameters)
+	{
+		MatchParameters(DestStore, SrcStore, BindVariable);
+	}
+	else if (BoundParameters->Num())
+	{
+		for (const FNiagaraBoundParameter& BoundParameter : *BoundParameters)
+		{
+			checkSlow(SrcStore->IndexOf(BoundParameter.Parameter) == BoundParameter.SrcOffset && DestStore->IndexOf(BoundParameter.Parameter) == BoundParameter.DestOffset);
+			BindVariable(BoundParameter.Parameter, BoundParameter.SrcOffset, BoundParameter.DestOffset);
 		}
+		bAnyBinding = true;
 	}
 
 	if (bAnyBinding)
diff --git a/Engine/Plugins/FX/Niagara/Source/Niagara/Private/NiagaraScript.cpp b/Engine/Plugins/FX/Niagara/Source/Niagara/Private/NiagaraScript.cpp
--- a/Engine/Plugins/FX/Niagara/Source/Niagara/Private/NiagaraScript.cpp
+++ b/Engine/Plugins/FX/Niagara/Source/Niagara/Private/NiagaraScript.cpp
@@ -399,6 +399,7 @@ bool UNiagaraScript::ContainsUsage(ENiagaraScriptUsage InUsage) const
 
 FNiagaraScriptExecutionParameterStore* UNiagaraScript::GetExecutionReadyParameterStore(ENiagaraSimTarget SimTarget)
 {
+#if WITH_EDITORONLY_DATA
 	if (SimTarget == ENiagaraSimTarget::CPUSim && IsReadyToRun(ENiagaraSimTarget::CPUSim))
 	{
 		if (ScriptExecutionParamStoreCPU.IsInitialized() == false)
@@ -419,8 +420,64 @@ FNiagaraScriptExecutionParameterStore* UNiagaraScript::GetExecutionReadyParamete
 	{
 		return nullptr;
 	}
+#else
+	check(GetSimTarget().GetValue() == SimTarget);
+	return &ScriptExecutionParamStore;
+#endif
+}
+
+TOptional<ENiagaraSimTarget> UNiagaraScript::GetSimTarget() const
+{
+	switch (Usage)
+	{
+	case ENiagaraScriptUsage::ParticleSpawnScript:
+	case ENiagaraScriptUsage::ParticleSpawnScriptInterpolated:
+	case ENiagaraScriptUsage::ParticleUpdateScript:
+	case ENiagaraScriptUsage::ParticleEventScript:
+	case ENiagaraScriptUsage::ParticleGPUComputeScript:
+		if (UNiagaraEmitter* OwningEmitter = GetTypedOuter<UNiagaraEmitter>())
+		{
+			if (OwningEmitter->SimTarget != ENiagaraSimTarget::CPUSim || CachedScriptVM.IsValid())
+			{
+				return OwningEmitter->SimTarget;
+			}
+		}
+		break;
+	case ENiagaraScriptUsage::EmitterSpawnScript:
+	case ENiagaraScriptUsage::EmitterUpdateScript:
+	case ENiagaraScriptUsage::SystemSpawnScript:
+	case ENiagaraScriptUsage::SystemUpdateScript:
+		if (CachedScriptVM.IsValid())
+		{
+			return ENiagaraSimTarget::CPUSim;
+		}
+		break;
+	default:
+		break;
+	};
+	return TOptional<ENiagaraSimTarget>();
 }
 
+void UNiagaraScript::PreSave(const class ITargetPlatform* TargetPlatform)
+{
+	Super::PreSave(TargetPlatform);
+
+#if WITH_EDITORONLY_DATA
+	ScriptExecutionParamStore.Empty();
+	ScriptExecutionBoundParameters.Empty();
+
+	if (TargetPlatform && TargetPlatform->RequiresCookedData())
+	{
+		TOptional<ENiagaraSimTarget> SimTarget = GetSimTarget();
+		if (SimTarget)
+		{
+			// Partial execution of InitFromOwningScript()
+			ScriptExecutionParamStore.AddScriptParams(this, SimTarget.GetValue(), false);
+			FNiagaraParameterStoreBinding::GetBindingData(&ScriptExecutionParamStore, &RapidIterationParameters, ScriptExecutionBoundParameters);
+		}
+	}
+#endif
+}
 
 void UNiagaraScript::Serialize(FArchive& Ar)
 {
@@ -434,7 +491,6 @@ void UNiagaraScript::Serialize(FArchive& Ar)
 		bool bUsesRapidIterationParams = true;
 
 #if WITH_EDITORONLY_DATA
-
 		if (UNiagaraEmitter* Emitter = Cast<UNiagaraEmitter>(GetOuter()))
 		{
 			UNiagaraSystem* EmitterOwner = Cast<UNiagaraSystem>(Emitter->GetOuter());
@@ -598,6 +654,14 @@ void UNiagaraScript::PostLoad()
 	
 	RapidIterationParameters.PostLoad();
 
+	if (FPlatformProperties::RequiresCookedData())
+	{
+		ScriptExecutionParamStore.PostLoad();
+		RapidIterationParameters.Bind(&ScriptExecutionParamStore, &ScriptExecutionBoundParameters);
+		ScriptExecutionParamStore.SetAsInitialized();
+		ScriptExecutionBoundParameters.Empty();
+	}
+
 	bool bNeedsRecompile = false;
 	const int32 NiagaraVer = GetLinkerCustomVersion(FNiagaraCustomVersion::GUID);
 
@@ -661,32 +725,6 @@ void UNiagaraScript::PostLoad()
 #endif
 
 	//FNiagaraUtilities::DumpHLSLText(RapidIterationParameters.ToString(), *GetPathName());
-
-	// Init ScriptExecutionParamStoreCPU or ScriptExecutionParamStoreGPU to avoid first time cost
-	if (FPlatformProperties::RequiresCookedData())
-	{
-	    switch (Usage)
-	    {
-	    case ENiagaraScriptUsage::ParticleSpawnScript:
-	    case ENiagaraScriptUsage::ParticleSpawnScriptInterpolated:
-	    case ENiagaraScriptUsage::ParticleUpdateScript:
-	    case ENiagaraScriptUsage::ParticleEventScript:
-	    case ENiagaraScriptUsage::ParticleGPUComputeScript:
-		    if (UNiagaraEmitter* OwningEmitter = GetTypedOuter<UNiagaraEmitter>())
-		    {
-				GetExecutionReadyParameterStore(OwningEmitter->SimTarget);
-		    }
-		    break;
-	    case ENiagaraScriptUsage::EmitterSpawnScript:
-	    case ENiagaraScriptUsage::EmitterUpdateScript:
-	    case ENiagaraScriptUsage::SystemSpawnScript:
-	    case ENiagaraScriptUsage::SystemUpdateScript:
-			GetExecutionReadyParameterStore(ENiagaraSimTarget::CPUSim);
-		    break;
-	    default:
-		    break;
-		};
-	}
 }
 
 bool UNiagaraScript::IsReadyToRun(ENiagaraSimTarget SimTarget) const
@@ -1099,9 +1137,11 @@ void UNiagaraScript::SetVMCompilationResults(const FNiagaraVMExecutableDataId& I
 
 void UNiagaraScript::InvalidateExecutionReadyParameterStores()
 {
+#if WITH_EDITORONLY_DATA
 	// Make sure that we regenerate any parameter stores, since they must be kept in sync with the layout from script compilation.
 	ScriptExecutionParamStoreCPU.Empty();
 	ScriptExecutionParamStoreGPU.Empty();
+#endif
 }
 
 void UNiagaraScript::InvalidateCachedCompileIds()
diff --git a/Engine/Plugins/FX/Niagara/Source/Niagara/Public/NiagaraParameterStore.h b/Engine/Plugins/FX/Niagara/Source/Niagara/Public/NiagaraParameterStore.h
--- a/Engine/Plugins/FX/Niagara/Source/Niagara/Public/NiagaraParameterStore.h
+++ b/Engine/Plugins/FX/Niagara/Source/Niagara/Public/NiagaraParameterStore.h
@@ -36,6 +36,22 @@ struct FNiagaraVariableSearch
 
 struct FNiagaraParameterStore;
 
+USTRUCT()
+struct FNiagaraBoundParameter
+{
+	GENERATED_USTRUCT_BODY()
+
+	UPROPERTY()
+	FNiagaraVariable Parameter;
+	UPROPERTY()
+	int32 SrcOffset;
+	UPROPERTY()
+	int32 DestOffset;
+
+};
+
+typedef TArray<FNiagaraBoundParameter> FNiagaraBoundParameterArray;
+
 //Binding from one parameter store to another.
 //This does no tracking of lifetimes etc so the owner must ensure safe use and rebinding when needed etc.
 struct FNiagaraParameterStoreBinding
@@ -104,15 +120,21 @@ struct FNiagaraParameterStoreBinding
 	}
 	
 	FORCEINLINE_DEBUGGABLE void Empty(FNiagaraParameterStore* DestStore, FNiagaraParameterStore* SrcStore);
-	FORCEINLINE_DEBUGGABLE bool Initialize(FNiagaraParameterStore* DestStore, FNiagaraParameterStore* SrcStore);
+	FORCEINLINE_DEBUGGABLE bool Initialize(FNiagaraParameterStore* DestStore, FNiagaraParameterStore* SrcStore, const FNiagaraBoundParameterArray* BoundParameters = nullptr);
 	FORCEINLINE_DEBUGGABLE bool VerifyBinding(const FNiagaraParameterStore* DestStore, const FNiagaraParameterStore* SrcStore)const;
 	FORCEINLINE_DEBUGGABLE void Tick(FNiagaraParameterStore* DestStore, FNiagaraParameterStore* SrcStore, bool bForce = false);
 	FORCEINLINE_DEBUGGABLE void Dump(const FNiagaraParameterStore* DestStore, const FNiagaraParameterStore* SrcStore)const;
 	/** TODO: Merge contiguous ranges into a single binding? */
 	//FORCEINLINE_DEBUGGABLE void Optimize();
 
+	static void GetBindingData(FNiagaraParameterStore* DestStore, FNiagaraParameterStore* SrcStore, FNiagaraBoundParameterArray& OutBoundParameters);
+
 private:
-	bool BindParameters(FNiagaraParameterStore* DestStore, FNiagaraParameterStore* SrcStore);
+	
+	template <typename TVisitor>
+	FORCEINLINE_DEBUGGABLE static void MatchParameters(FNiagaraParameterStore* DestStore, FNiagaraParameterStore* SrcStore, TVisitor Visitor);
+
+	bool BindParameters(FNiagaraParameterStore* DestStore, FNiagaraParameterStore* SrcStore, const FNiagaraBoundParameterArray* BoundParameters = nullptr);
 };
 
 
@@ -215,7 +237,7 @@ public:
 	FORCEINLINE uint32 GetLayoutVersion() const { return LayoutVersion; }
 
 	/** Binds this parameter store to another, by default if we find no matching parameters we will not maintain a pointer to the store. */
-	void Bind(FNiagaraParameterStore* DestStore);
+	void Bind(FNiagaraParameterStore* DestStore, const FNiagaraBoundParameterArray* BoundParameters = nullptr);
 	/** Unbinds this store form one it's bound to. */
 	void Unbind(FNiagaraParameterStore* DestStore);
 	/** Recreates any bindings to reflect a layout change etc. */
@@ -625,12 +647,12 @@ FORCEINLINE_DEBUGGABLE void FNiagaraParameterStoreBinding::Empty(FNiagaraParamet
 	UObjectBindings.Reset();
 }
 
-FORCEINLINE_DEBUGGABLE bool FNiagaraParameterStoreBinding::Initialize(FNiagaraParameterStore* DestStore, FNiagaraParameterStore* SrcStore)
+FORCEINLINE_DEBUGGABLE bool FNiagaraParameterStoreBinding::Initialize(FNiagaraParameterStore* DestStore, FNiagaraParameterStore* SrcStore, const FNiagaraBoundParameterArray* BoundParameters)
 {
 	checkSlow(DestStore);
 	checkSlow(SrcStore);
 
-	if (BindParameters(DestStore, SrcStore))
+	if (BindParameters(DestStore, SrcStore, BoundParameters))
 	{
 		DestStore->GetSourceParameterStores().AddUnique(SrcStore);
 		return true;
diff --git a/Engine/Plugins/FX/Niagara/Source/Niagara/Public/NiagaraScriptExecutionParameterStore.h b/Engine/Plugins/FX/Niagara/Source/Niagara/Public/NiagaraScriptExecutionParameterStore.h
--- a/Engine/Plugins/FX/Niagara/Source/Niagara/Public/NiagaraScriptExecutionParameterStore.h
+++ b/Engine/Plugins/FX/Niagara/Source/Niagara/Public/NiagaraScriptExecutionParameterStore.h
@@ -14,9 +14,14 @@ struct FNiagaraScriptExecutionPaddingInfo
 public:
 	FNiagaraScriptExecutionPaddingInfo() : SrcOffset(0), DestOffset(0), SrcSize(0), DestSize(0) {}
 	FNiagaraScriptExecutionPaddingInfo(uint32 InSrcOffset, uint32 InDestOffset, uint32 InSrcSize, uint32 InDestSize) : SrcOffset(InSrcOffset), DestOffset(InDestOffset), SrcSize(InSrcSize), DestSize(InDestSize) {}
+
+	UPROPERTY()
 	uint32 SrcOffset;
+	UPROPERTY()
 	uint32 DestOffset;
+	UPROPERTY()
 	uint32 SrcSize;
+	UPROPERTY()
 	uint32 DestSize;
 };
 
@@ -89,7 +94,7 @@ public:
 
 
 	bool IsInitialized() const {return bInitialized;}
-
+	void SetAsInitialized() { bInitialized = true; }
 protected:
 	void AddPaddedParamSize(const FNiagaraTypeDefinition& InParamType, uint32 InOffset);
 
