Adds the option to override the Niagara global spawn count scale in the Emitter Properties for different EffectsQuality Scalability levels.
The spawn count scale overrides are initially seeded by the the EmitterQuality sections of the Scalability .ini hierarchy.
If the scales are not overriden the emitter will fall back to using fx.NiagaraGlobalSpawnCountScale instead.

Introduces two new engine variables:
 - Engine.Emitter.SpawnCountScale
 - Engine.Emitter.bOverrideGlobalSpawnCountScale
The latter is a constant and can be used in static switches.

Updates the SpawnRate and SpawnBurst_Instantaneous scripts to use these new overrides.

Updates the fast pass to use the new overrides.

This deprecates Engine.GlobalSpawnCountScale.

Bumps the Niagara custom script version.

diff --git a/Engine/Plugins/FX/Niagara/Source/Niagara/Classes/NiagaraConstants.h b/Engine/Plugins/FX/Niagara/Source/Niagara/Classes/NiagaraConstants.h
--- a/Engine/Plugins/FX/Niagara/Source/Niagara/Classes/NiagaraConstants.h
+++ b/Engine/Plugins/FX/Niagara/Source/Niagara/Classes/NiagaraConstants.h
@@ -50,6 +50,7 @@
 #define SYS_PARAM_ENGINE_EXEC_COUNT                      INiagaraModule::GetVar_Engine_ExecutionCount()
 #define SYS_PARAM_ENGINE_EMITTER_NUM_PARTICLES           INiagaraModule::GetVar_Engine_Emitter_NumParticles()
 #define SYS_PARAM_ENGINE_EMITTER_TOTAL_SPAWNED_PARTICLES INiagaraModule::GetVar_Engine_Emitter_TotalSpawnedParticles()
+#define SYS_PARAM_ENGINE_EMITTER_SPAWN_COUNT_SCALE       INiagaraModule::GetVar_Engine_Emitter_SpawnCountScale()
 #define SYS_PARAM_ENGINE_SYSTEM_NUM_EMITTERS_ALIVE       INiagaraModule::GetVar_Engine_System_NumEmittersAlive()
 #define SYS_PARAM_ENGINE_SYSTEM_NUM_EMITTERS             INiagaraModule::GetVar_Engine_System_NumEmitters()
 #define SYS_PARAM_ENGINE_NUM_SYSTEM_INSTANCES            INiagaraModule::GetVar_Engine_NumSystemInstances()
@@ -63,6 +64,7 @@
 #define SYS_PARAM_EMITTER_AGE                            INiagaraModule::GetVar_Emitter_Age()
 #define SYS_PARAM_EMITTER_LOCALSPACE                     INiagaraModule::GetVar_Emitter_LocalSpace()
 #define SYS_PARAM_EMITTER_DETERMINISM                    INiagaraModule::GetVar_Emitter_Determinism()
+#define SYS_PARAM_EMITTER_OVERRIDE_GLOBAL_SPAWN_COUNT_SCALE       INiagaraModule::GetVar_Emitter_OverrideGlobalSpawnCountScale()
 #define SYS_PARAM_EMITTER_RANDOM_SEED                    INiagaraModule::GetVar_Emitter_RandomSeed()
 #define SYS_PARAM_EMITTER_SPAWNRATE                      INiagaraModule::GetVar_Emitter_SpawnRate()
 #define SYS_PARAM_EMITTER_SPAWN_INTERVAL                 INiagaraModule::GetVar_Emitter_SpawnInterval()
diff --git a/Engine/Plugins/FX/Niagara/Source/Niagara/Classes/NiagaraEmitter.h b/Engine/Plugins/FX/Niagara/Source/Niagara/Classes/NiagaraEmitter.h
--- a/Engine/Plugins/FX/Niagara/Source/Niagara/Classes/NiagaraEmitter.h
+++ b/Engine/Plugins/FX/Niagara/Source/Niagara/Classes/NiagaraEmitter.h
@@ -167,6 +167,34 @@ struct FNiagaraEventScriptProperties : public FNiagaraEmitterScriptProperties
 	uint32 MinSpawnNumber;
 };
 
+USTRUCT()
+struct FNiagaraDetailsLevelScaleOverrides 
+{
+	GENERATED_BODY()
+
+	FNiagaraDetailsLevelScaleOverrides();
+
+	/** The Scalability spawn count scale used when Effects Quality level is set to Low instead of fx.NiagaraGlobalSpawnCountScale. */
+	UPROPERTY(EditAnywhere, Category = "Scalability")
+	float Low;
+
+	/** The Scalability spawn count scale used when Effects Quality level is set to Medium instead of fx.NiagaraGlobalSpawnCountScale. */
+	UPROPERTY(EditAnywhere, Category = "Scalability")
+	float Medium;
+
+	/** The Scalability spawn count scale used when Effects Quality level is set to High instead of fx.NiagaraGlobalSpawnCountScale. */
+	UPROPERTY(EditAnywhere, Category = "Scalability")
+	float High;
+
+	/** The Scalability spawn count scale used when Effects Quality level is set to Epic instead of fx.NiagaraGlobalSpawnCountScale. */
+	UPROPERTY(EditAnywhere, Category = "Scalability")
+	float Epic;
+
+	/** The Scalability spawn count scale used when Effects Quality level is set to Cine instead of fx.NiagaraGlobalSpawnCountScale. */
+	UPROPERTY(EditAnywhere, Category = "Scalability")
+	float Cine;
+};
+
 /** 
  *	UNiagaraEmitter stores the attributes of an FNiagaraEmitterInstance
  *	that need to be serialized and are used for its initialization 
@@ -253,6 +281,10 @@ public:
 	UPROPERTY(EditAnywhere, Category = "Scalability", meta=(EditCondition = "bUseMaxDetailLevel"))
 	int32 MaxDetailLevel;
 
+	/** Custom Scalability spawn count scales to override the ones defined by fx.NiagaraGlobalSpawnCountScale in the GScalabilityIni config hierarchy. */
+	UPROPERTY(EditAnywhere, Category = "Scalability", meta = (EditCondition = "bOverrideGlobalSpawnCountScale"))
+	FNiagaraDetailsLevelScaleOverrides GlobalSpawnCountScaleOverrides;
+
 	/** When enabled, this will spawn using interpolated parameter values and perform a partial update at spawn time. This adds significant additional cost for spawning but will produce much smoother spawning for high spawn rates, erratic frame rates and fast moving emitters. */
 	UPROPERTY(EditAnywhere, Category = "Emitter")
 	uint32 bInterpolatedSpawning : 1;
@@ -269,6 +301,10 @@ public:
 	UPROPERTY(EditAnywhere, Category = "Scalability", meta = (InlineEditConditionToggle))
 	uint32 bUseMaxDetailLevel : 1;
 
+	/** Whether or not to override the global spawn count scale based on Effects Quality levels, as defined by fx.NiagaraGlobalSpawnCountScale in the GScalabilityIni config hierarchy. */
+	UPROPERTY(EditAnywhere, Category = "Scalability", meta = (InlineEditConditionToggle))
+	uint32 bOverrideGlobalSpawnCountScale : 1;
+
 	/** Do particles in this emitter require a persistent ID? */
 	UPROPERTY(EditAnywhere, Category = "Emitter")
 	uint32 bRequiresPersistentIDs : 1;
@@ -350,6 +386,8 @@ public:
 	NIAGARA_API UNiagaraEmitter* DuplicateWithoutMerging(UObject* InOuter);
 #endif
 
+	float GetSpawnCountScale(int EffectsQuality = -1); 
+
 	/** Is this emitter allowed to be enabled by the current system detail level. */
 	NIAGARA_API bool IsAllowedByDetailLevel(int32 DetailLevel)const;
 	NIAGARA_API bool RequiresPersistantIDs()const;
diff --git a/Engine/Plugins/FX/Niagara/Source/Niagara/Classes/NiagaraEmitterInstance.h b/Engine/Plugins/FX/Niagara/Source/Niagara/Classes/NiagaraEmitterInstance.h
--- a/Engine/Plugins/FX/Niagara/Source/Niagara/Classes/NiagaraEmitterInstance.h
+++ b/Engine/Plugins/FX/Niagara/Source/Niagara/Classes/NiagaraEmitterInstance.h
@@ -83,6 +83,7 @@ public:
 		return 0;
 	}
 	FORCEINLINE int32 GetTotalSpawnedParticles()const { return TotalSpawnedParticles; }
+	FORCEINLINE float GetSpawnCountScale(int32 EffectsQuality = -1)const { return CachedEmitter->GetSpawnCountScale(EffectsQuality); }
 
 	NIAGARA_API const FNiagaraEmitterHandle& GetEmitterHandle() const;
 
diff --git a/Engine/Plugins/FX/Niagara/Source/Niagara/Classes/NiagaraSystem.h b/Engine/Plugins/FX/Niagara/Source/Niagara/Classes/NiagaraSystem.h
--- a/Engine/Plugins/FX/Niagara/Source/Niagara/Classes/NiagaraSystem.h
+++ b/Engine/Plugins/FX/Niagara/Source/Niagara/Classes/NiagaraSystem.h
@@ -70,6 +70,9 @@ struct FNiagaraSystemCompiledData
 	UPROPERTY()
 	FNiagaraParameterStore InstanceParamStore;
 
+	UPROPERTY()
+	TArray<FNiagaraVariable> SpawnCountScaleVars;
+
 	UPROPERTY()
 	FNiagaraDataSetCompiledData DataSetCompiledData;
 
diff --git a/Engine/Plugins/FX/Niagara/Source/Niagara/Private/NiagaraConstants.cpp b/Engine/Plugins/FX/Niagara/Source/Niagara/Private/NiagaraConstants.cpp
--- a/Engine/Plugins/FX/Niagara/Source/Niagara/Private/NiagaraConstants.cpp
+++ b/Engine/Plugins/FX/Niagara/Source/Niagara/Private/NiagaraConstants.cpp
@@ -53,6 +53,7 @@ void FNiagaraConstants::Init()
 		SystemParameters.Add(SYS_PARAM_ENGINE_EXEC_COUNT);
 		SystemParameters.Add(SYS_PARAM_ENGINE_EMITTER_NUM_PARTICLES);
 		SystemParameters.Add(SYS_PARAM_ENGINE_EMITTER_TOTAL_SPAWNED_PARTICLES);
+		SystemParameters.Add(SYS_PARAM_ENGINE_EMITTER_SPAWN_COUNT_SCALE);
 		SystemParameters.Add(SYS_PARAM_ENGINE_SYSTEM_NUM_EMITTERS_ALIVE);
 		SystemParameters.Add(SYS_PARAM_ENGINE_SYSTEM_NUM_EMITTERS);
 		SystemParameters.Add(SYS_PARAM_ENGINE_NUM_SYSTEM_INSTANCES);
@@ -63,6 +64,7 @@ void FNiagaraConstants::Init()
 		SystemParameters.Add(SYS_PARAM_EMITTER_AGE);
 		SystemParameters.Add(SYS_PARAM_EMITTER_LOCALSPACE);
 		SystemParameters.Add(SYS_PARAM_EMITTER_DETERMINISM);
+		SystemParameters.Add(SYS_PARAM_EMITTER_OVERRIDE_GLOBAL_SPAWN_COUNT_SCALE);
 		SystemParameters.Add(SYS_PARAM_EMITTER_RANDOM_SEED);
 		SystemParameters.Add(SYS_PARAM_EMITTER_SPAWN_GROUP);
 	}
@@ -76,6 +78,7 @@ void FNiagaraConstants::Init()
 	{
 		SwitchParameters.Add(SYS_PARAM_EMITTER_LOCALSPACE);
 		SwitchParameters.Add(SYS_PARAM_EMITTER_DETERMINISM);
+		SwitchParameters.Add(SYS_PARAM_EMITTER_OVERRIDE_GLOBAL_SPAWN_COUNT_SCALE);
 		SwitchParameters.Add(SYS_PARAM_EMITTER_SIMULATION_TARGET);
 		SwitchParameters.Add(SYS_PARAM_SCRIPT_USAGE);
 	}
@@ -111,6 +114,7 @@ void FNiagaraConstants::Init()
 		UpdatedSystemParameters.Add(FName(TEXT("Emitter Local Space")), SYS_PARAM_EMITTER_LOCALSPACE);
 		UpdatedSystemParameters.Add(FName(TEXT("Emitter Random Seed")), SYS_PARAM_EMITTER_RANDOM_SEED);
 		UpdatedSystemParameters.Add(FName(TEXT("Emitter Determinism")), SYS_PARAM_EMITTER_DETERMINISM);
+		UpdatedSystemParameters.Add(FName(TEXT("Emitter Override Global Spawn Count Scale")), SYS_PARAM_EMITTER_OVERRIDE_GLOBAL_SPAWN_COUNT_SCALE);
 		UpdatedSystemParameters.Add(FName(TEXT("Effect Position")), SYS_PARAM_ENGINE_POSITION);
 		UpdatedSystemParameters.Add(FName(TEXT("Effect Velocity")), SYS_PARAM_ENGINE_VELOCITY);
 		UpdatedSystemParameters.Add(FName(TEXT("Effect X Axis")), SYS_PARAM_ENGINE_X_AXIS);
@@ -161,6 +165,7 @@ void FNiagaraConstants::Init()
 		SystemStrMap.Add(SYS_PARAM_ENGINE_EXEC_COUNT, LOCTEXT("ExecCountDesc", "The index of this particle in the read buffer."));
 		SystemStrMap.Add(SYS_PARAM_ENGINE_EMITTER_NUM_PARTICLES, LOCTEXT("EmitterNumParticles", "The number of particles for this emitter at the beginning of simulation. Should only be used in Emitter scripts."));
 		SystemStrMap.Add(SYS_PARAM_ENGINE_EMITTER_TOTAL_SPAWNED_PARTICLES, LOCTEXT("EmitterTotalSpawnedParticles", "The total number of particles spawned for this emitter at the beginning of this simulation. Should only be used by the particle spawn script the assign unique IDs."));
+		SystemStrMap.Add(SYS_PARAM_ENGINE_EMITTER_SPAWN_COUNT_SCALE, LOCTEXT("EmitterSpawnCountScale", "The global spawn count scale, based on fx.NiagaraGlobalSpawnCountScale."));
 		SystemStrMap.Add(SYS_PARAM_ENGINE_SYSTEM_NUM_EMITTERS_ALIVE, LOCTEXT("SystemNumEmittersAlive", "The number of emitters still alive attached to this system. Should only be used in System scripts."));
 		SystemStrMap.Add(SYS_PARAM_ENGINE_SYSTEM_NUM_EMITTERS, LOCTEXT("SystemNumEmitters", "The number of emitters attached to this system. Should only be used in System scripts."));
 		SystemStrMap.Add(SYS_PARAM_ENGINE_NUM_SYSTEM_INSTANCES, LOCTEXT("SystemNumInstances", "The number of instances of the this system currently ticking. Should only be used in System scripts."));
diff --git a/Engine/Plugins/FX/Niagara/Source/Niagara/Private/NiagaraEmitter.cpp b/Engine/Plugins/FX/Niagara/Source/Niagara/Private/NiagaraEmitter.cpp
--- a/Engine/Plugins/FX/Niagara/Source/Niagara/Private/NiagaraEmitter.cpp
+++ b/Engine/Plugins/FX/Niagara/Source/Niagara/Private/NiagaraEmitter.cpp
@@ -51,6 +51,32 @@ static FAutoConsoleVariableRef CVarEnableEmitterChangeIdMergeLogging(
 	ECVF_Default
 );
 
+float UNiagaraEmitter::GetSpawnCountScale(int EffectsQuality)
+{
+	if (bOverrideGlobalSpawnCountScale)
+	{
+		switch (EffectsQuality)
+		{
+		case 0: return GlobalSpawnCountScaleOverrides.Low;
+		case 1: return GlobalSpawnCountScaleOverrides.Medium;
+		case 2: return GlobalSpawnCountScaleOverrides.High;
+		case 3: return GlobalSpawnCountScaleOverrides.Epic;
+		case 4: return GlobalSpawnCountScaleOverrides.Cine;
+		}
+	}
+	
+	return INiagaraModule::GetGlobalSpawnCountScale(); // default to the global spawn scale when invalid effects quality
+}
+
+FNiagaraDetailsLevelScaleOverrides::FNiagaraDetailsLevelScaleOverrides()
+{
+	GConfig->GetFloat(TEXT("EffectsQuality@0"), TEXT("fx.NiagaraGlobalSpawnCountScale"), Low, GScalabilityIni);
+	GConfig->GetFloat(TEXT("EffectsQuality@1"), TEXT("fx.NiagaraGlobalSpawnCountScale"), Medium, GScalabilityIni);
+	GConfig->GetFloat(TEXT("EffectsQuality@2"), TEXT("fx.NiagaraGlobalSpawnCountScale"), High, GScalabilityIni);
+	GConfig->GetFloat(TEXT("EffectsQuality@3"), TEXT("fx.NiagaraGlobalSpawnCountScale"), Epic, GScalabilityIni);
+	GConfig->GetFloat(TEXT("EffectsQuality@Cine"), TEXT("fx.NiagaraGlobalSpawnCountScale"), Cine, GScalabilityIni);
+}
+
 void FNiagaraEmitterScriptProperties::InitDataSetAccess()
 {
 	EventReceivers.Empty();
@@ -530,6 +556,17 @@ void UNiagaraEmitter::PostEditChangeProperty(struct FPropertyChangedEvent& Prope
 			GraphSource->MarkNotSynchronized(TEXT("Emitter Determinism changed."));
 		}
 
+#if WITH_EDITORONLY_DATA
+		UNiagaraSystem::RequestCompileForEmitter(this);
+#endif
+	}
+	if (PropertyName == GET_MEMBER_NAME_CHECKED(UNiagaraEmitter, bOverrideGlobalSpawnCountScale))
+	{
+		if (GraphSource != nullptr)
+		{
+			GraphSource->MarkNotSynchronized(TEXT("Emitter Override Global Spawn Count Scale changed."));
+		}
+
 #if WITH_EDITORONLY_DATA
 		UNiagaraSystem::RequestCompileForEmitter(this);
 #endif
diff --git a/Engine/Plugins/FX/Niagara/Source/Niagara/Private/NiagaraModule.cpp b/Engine/Plugins/FX/Niagara/Source/Niagara/Private/NiagaraModule.cpp
--- a/Engine/Plugins/FX/Niagara/Source/Niagara/Private/NiagaraModule.cpp
+++ b/Engine/Plugins/FX/Niagara/Source/Niagara/Private/NiagaraModule.cpp
@@ -117,6 +117,7 @@ FNiagaraVariable INiagaraModule::Engine_Owner_ExecutionState;
 FNiagaraVariable INiagaraModule::Engine_ExecutionCount;
 FNiagaraVariable INiagaraModule::Engine_Emitter_NumParticles;
 FNiagaraVariable INiagaraModule::Engine_Emitter_TotalSpawnedParticles;
+FNiagaraVariable INiagaraModule::Engine_Emitter_SpawnCountScale;
 FNiagaraVariable INiagaraModule::Engine_System_TickCount;
 FNiagaraVariable INiagaraModule::Engine_System_NumEmittersAlive;
 FNiagaraVariable INiagaraModule::Engine_System_NumEmitters;
@@ -130,6 +131,7 @@ FNiagaraVariable INiagaraModule::Engine_System_Age;
 FNiagaraVariable INiagaraModule::Emitter_Age;
 FNiagaraVariable INiagaraModule::Emitter_LocalSpace;
 FNiagaraVariable INiagaraModule::Emitter_Determinism;
+FNiagaraVariable INiagaraModule::Emitter_OverrideGlobalSpawnCountScale;
 FNiagaraVariable INiagaraModule::Emitter_SimulationTarget;
 FNiagaraVariable INiagaraModule::Emitter_RandomSeed;
 FNiagaraVariable INiagaraModule::Emitter_SpawnRate;
@@ -220,6 +222,7 @@ void INiagaraModule::StartupModule()
 	Engine_ExecutionCount = FNiagaraVariable(FNiagaraTypeDefinition::GetIntDef(), TEXT("Engine.ExecutionCount"));
 	Engine_Emitter_NumParticles = FNiagaraVariable(FNiagaraTypeDefinition::GetIntDef(), TEXT("Engine.Emitter.NumParticles"));
 	Engine_Emitter_TotalSpawnedParticles = FNiagaraVariable(FNiagaraTypeDefinition::GetIntDef(), TEXT("Engine.Emitter.TotalSpawnedParticles"));
+	Engine_Emitter_SpawnCountScale = FNiagaraVariable(FNiagaraTypeDefinition::GetFloatDef(), TEXT("Engine.Emitter.SpawnCountScale"));
 	Engine_System_TickCount = FNiagaraVariable(FNiagaraTypeDefinition::GetIntDef(), TEXT("Engine.System.TickCount"));
 	Engine_System_NumEmittersAlive = FNiagaraVariable(FNiagaraTypeDefinition::GetIntDef(), TEXT("Engine.System.NumEmittersAlive"));
 	Engine_System_NumEmitters = FNiagaraVariable(FNiagaraTypeDefinition::GetIntDef(), TEXT("Engine.System.NumEmitters"));
@@ -233,6 +236,7 @@ void INiagaraModule::StartupModule()
 	Emitter_LocalSpace = FNiagaraVariable(FNiagaraTypeDefinition::GetBoolDef(), TEXT("Emitter.LocalSpace"));
 	Emitter_RandomSeed = FNiagaraVariable(FNiagaraTypeDefinition::GetIntDef(), TEXT("Emitter.RandomSeed"));
 	Emitter_Determinism = FNiagaraVariable(FNiagaraTypeDefinition::GetBoolDef(), TEXT("Emitter.Determinism"));
+	Emitter_OverrideGlobalSpawnCountScale = FNiagaraVariable(FNiagaraTypeDefinition::GetBoolDef(), TEXT("Emitter.OverrideGlobalSpawnCountScale"));
 	Emitter_SimulationTarget = FNiagaraVariable(FNiagaraTypeDefinition::GetSimulationTargetEnum(), TEXT("Emitter.SimulationTarget"));
 	Emitter_SpawnRate = FNiagaraVariable(FNiagaraTypeDefinition::GetFloatDef(), TEXT("Emitter.SpawnRate"));
 	Emitter_SpawnInterval = FNiagaraVariable(FNiagaraTypeDefinition::GetFloatDef(), TEXT("Emitter.SpawnInterval"));
diff --git a/Engine/Plugins/FX/Niagara/Source/Niagara/Private/NiagaraScript.cpp b/Engine/Plugins/FX/Niagara/Source/Niagara/Private/NiagaraScript.cpp
--- a/Engine/Plugins/FX/Niagara/Source/Niagara/Private/NiagaraScript.cpp
+++ b/Engine/Plugins/FX/Niagara/Source/Niagara/Private/NiagaraScript.cpp
@@ -284,6 +284,10 @@ void UNiagaraScript::ComputeVMCompilationId(FNiagaraVMExecutableDataId& Id) cons
 		{
 			Id.AdditionalDefines.Add(TEXT("Emitter.Determinism"));
 		}
+		if (Emitter->bOverrideGlobalSpawnCountScale)
+		{
+			Id.AdditionalDefines.Add(TEXT("Emitter.OverrideGlobalSpawnCountScale"));
+		}
 
 		if (!Emitter->bBakeOutRapidIteration)
 		{
diff --git a/Engine/Plugins/FX/Niagara/Source/Niagara/Private/NiagaraSystem.cpp b/Engine/Plugins/FX/Niagara/Source/Niagara/Private/NiagaraSystem.cpp
--- a/Engine/Plugins/FX/Niagara/Source/Niagara/Private/NiagaraSystem.cpp
+++ b/Engine/Plugins/FX/Niagara/Source/Niagara/Private/NiagaraSystem.cpp
@@ -517,7 +517,7 @@ bool UNiagaraSystem::IsReadyToRun() const
 	/* Check that our post compile data is in sync with the current emitter handles count. If we have just added a new emitter handle, we will not have any outstanding compilation requests as the new compile
 	 * will not be added to the outstanding compilation requests until the next tick.
 	 */
-	if (EmitterHandles.Num() != EmitterCompiledData.Num() || EmitterHandles.Num() != SystemCompiledData.NumParticleVars.Num() || EmitterHandles.Num() != SystemCompiledData.TotalSpawnedParticlesVars.Num())
+	if (EmitterHandles.Num() != EmitterCompiledData.Num() || EmitterHandles.Num() != SystemCompiledData.NumParticleVars.Num() || EmitterHandles.Num() != SystemCompiledData.TotalSpawnedParticlesVars.Num() || EmitterHandles.Num() != SystemCompiledData.SpawnCountScaleVars.Num())
 	{
 		return false;
 	}
@@ -1149,6 +1149,7 @@ void UNiagaraSystem::InitSystemCompiledData()
 {
 	SystemCompiledData.NumParticleVars.Empty();
 	SystemCompiledData.TotalSpawnedParticlesVars.Empty();
+	SystemCompiledData.SpawnCountScaleVars.Empty();
 	SystemCompiledData.InstanceParamStore.Empty();
 
 	SystemCompiledData.InstanceParamStore = INiagaraModule::GetFixedSystemInstanceParameterStore();
@@ -1174,6 +1175,13 @@ void UNiagaraSystem::InitSystemCompiledData()
 				SystemCompiledData.InstanceParamStore.AddParameter(Var, true, false);
 				SystemCompiledData.TotalSpawnedParticlesVars.Add(Var);
 			}
+			{
+				FNiagaraVariable Var = SYS_PARAM_ENGINE_EMITTER_SPAWN_COUNT_SCALE;
+				const FString ParamName = Var.GetName().ToString().Replace(TEXT("Emitter"), *EmitterName);
+				Var.SetName(*ParamName);
+				SystemCompiledData.InstanceParamStore.AddParameter(Var, true, false);
+				SystemCompiledData.SpawnCountScaleVars.Add(Var);
+			}
 		}
 	}
 
diff --git a/Engine/Plugins/FX/Niagara/Source/Niagara/Private/NiagaraSystemFastPath.cpp b/Engine/Plugins/FX/Niagara/Source/Niagara/Private/NiagaraSystemFastPath.cpp
--- a/Engine/Plugins/FX/Niagara/Source/Niagara/Private/NiagaraSystemFastPath.cpp
+++ b/Engine/Plugins/FX/Niagara/Source/Niagara/Private/NiagaraSystemFastPath.cpp
@@ -245,6 +245,7 @@ void FNiagaraEmitterFastPath::SetSpawnMapDefaults(const UNiagaraEmitter& Emitter
 	Map.Emitter.LoopedAge = 0.0f;
 	Map.Emitter.NormalizedLoopAge = 0.0f;
 	Map.Emitter.DistanceTraveled = 0;
+	//Map.Emitter.SpawnCountScale = 1.0f; // TODO: This seems to always set the value to 1.0 due to the order of operations in FNiagaraSystemSimulation::TickFastPath
 	Map.Emitter.Scalability.DistanceFraction = 1;
 	Map.Scalability.Emitter.ExecutionState = (ENiagaraExecutionState)0;
 	Map.Scalability.Emitter.SpawnCountScale = 1;
@@ -522,7 +523,7 @@ void FNiagaraEmitterFastPath::Module_SpawnRate(FParamMap0& Context_Map)
 		const FParamMap0_SpawnRate& Context_Map_SpawnRate = Context_Map.SpawnRate[SpawnRateIndex];
 		FParamMap0_Emitter_SpawnRate& Context_Map_Emitter_SpawnRate = Context_Map.Emitter.SpawnRate[SpawnRateIndex];
 
-		float Result57 = Context_Map_SpawnRate.SpawnRate * Context_Map.Scalability.Emitter.SpawnCountScale * Context_Map.Engine.GlobalSpawnCountScale;
+		float Result57 = Context_Map_SpawnRate.SpawnRate * Context_Map.Scalability.Emitter.SpawnCountScale * Context_Map.Emitter.SpawnCountScale;
 		float Result58 = 1.0f / Result57;// Reciprocal(Result57);
 		float Result59 = 1 - Context_Map_Emitter_SpawnRate.SpawnRemainder;
 		float Result60 = Result58 * Result59;
@@ -676,7 +677,7 @@ void FNiagaraEmitterFastPath::Module_SpawnBurstInstantaneous(FParamMap0& Context
 		int32 Constant60 = 0;
 		bool Result89 = Context_Map_SpawnBurst_Instantaneous.SpawnCount == Constant60;
 		float Constant61 = 0;
-		float Result90 = Context_Map_SpawnBurst_Instantaneous.SpawnCount * Context_Map.Scalability.Emitter.SpawnCountScale * Context_Map.Engine.GlobalSpawnCountScale;
+		float Result90 = Context_Map_SpawnBurst_Instantaneous.SpawnCount * Context_Map.Scalability.Emitter.SpawnCountScale * Context_Map.Emitter.SpawnCountScale; 
 		float Constant62 = 1;
 		float Result91 = FMath::Max(Result90, Constant62);
 		float float_IfResult11;
diff --git a/Engine/Plugins/FX/Niagara/Source/Niagara/Private/NiagaraSystemInstance.cpp b/Engine/Plugins/FX/Niagara/Source/Niagara/Private/NiagaraSystemInstance.cpp
--- a/Engine/Plugins/FX/Niagara/Source/Niagara/Private/NiagaraSystemInstance.cpp
+++ b/Engine/Plugins/FX/Niagara/Source/Niagara/Private/NiagaraSystemInstance.cpp
@@ -885,6 +885,7 @@ void FNiagaraSystemInstance::ReInitInternal()
 	InstanceParameters = SystemCompiledData.InstanceParamStore;
 	const TArray<FNiagaraVariable>& NumParticleVars = SystemCompiledData.NumParticleVars;
 	const TArray<FNiagaraVariable>& TotalSpawnedParticlesVars = SystemCompiledData.TotalSpawnedParticlesVars;
+	const TArray<FNiagaraVariable>& SpawnCountScaleVars = SystemCompiledData.SpawnCountScaleVars;
 
 	// Make sure all parameters are added before initializing the bindings, otherwise parameter store layout changes might invalidate the bindings.
 	OwnerPositionParam.Init(InstanceParameters, SYS_PARAM_ENGINE_POSITION);
@@ -931,6 +932,12 @@ void FNiagaraSystemInstance::ReInitInternal()
 		ParameterTotalSpawnedParticlesBindings[i].Init(InstanceParameters, TotalSpawnedParticlesVars[i]);
 	}
 
+	ParameterSpawnCountScaleBindings.SetNum(SpawnCountScaleVars.Num());
+	for (int32 i = 0; i < SpawnCountScaleVars.Num(); i++)
+	{
+		ParameterSpawnCountScaleBindings[i].Init(InstanceParameters, SpawnCountScaleVars[i]);
+	}
+
 	// rebind now after all parameters have been added
 	InstanceParameters.Rebind();
 
@@ -1370,6 +1377,8 @@ ETickingGroup FNiagaraSystemInstance::CalculateTickGroup()
 
 void FNiagaraSystemInstance::TickInstanceParameters_GameThread(float DeltaSeconds)
 {
+	static const auto EffectsQualityLevelCVar = IConsoleManager::Get().FindConsoleVariable(TEXT("sg.EffectsQuality"));
+
 	if (!Component)
 	{
 		return;
@@ -1397,6 +1406,8 @@ void FNiagaraSystemInstance::TickInstanceParameters_GameThread(float DeltaSecond
 	GatheredInstanceParameters.Age = Age;
 	GatheredInstanceParameters.TickCount = TickCount;
 
+	int EffectsQualityLevel = EffectsQualityLevelCVar->GetInt();
+
 	GatheredInstanceParameters.NumAlive = 0;
 	for (int32 i = 0; i < Emitters.Num(); ++i)
 	{
@@ -1405,6 +1416,7 @@ void FNiagaraSystemInstance::TickInstanceParameters_GameThread(float DeltaSecond
 		{
 			GatheredInstanceParameters.EmitterNumParticles[i] = Emitter->GetNumParticles();
 			GatheredInstanceParameters.EmitterTotalSpawnedParticles[i] = Emitter->GetTotalSpawnedParticles();
+			GatheredInstanceParameters.EmitterSpawnCountScale[i] = Emitter->GetSpawnCountScale(EffectsQualityLevel);
 			if (!Emitter->IsComplete())
 			{
 				GatheredInstanceParameters.NumAlive++;
@@ -1468,6 +1480,7 @@ void FNiagaraSystemInstance::TickInstanceParameters_Concurrent()
 		{
 			ParameterNumParticleBindings[i].SetValue(GatheredInstanceParameters.EmitterNumParticles[i]);
 			ParameterTotalSpawnedParticlesBindings[i].SetValue(GatheredInstanceParameters.EmitterTotalSpawnedParticles[i]);
+			ParameterSpawnCountScaleBindings[i].SetValue(GatheredInstanceParameters.EmitterSpawnCountScale[i]);
 		}
 	}
 
diff --git a/Engine/Plugins/FX/Niagara/Source/Niagara/Private/NiagaraSystemSimulation.cpp b/Engine/Plugins/FX/Niagara/Source/Niagara/Private/NiagaraSystemSimulation.cpp
--- a/Engine/Plugins/FX/Niagara/Source/Niagara/Private/NiagaraSystemSimulation.cpp
+++ b/Engine/Plugins/FX/Niagara/Source/Niagara/Private/NiagaraSystemSimulation.cpp
@@ -93,13 +93,14 @@ static FAutoConsoleVariableRef CVarSystemSimTransferParamsParallelThreshold(
 
 //////////////////////////////////////////////////////////////////////////
 
-FNiagaraSystemSimulationTickContext::FNiagaraSystemSimulationTickContext(FNiagaraSystemSimulation* InOwner, TArray<FNiagaraSystemInstance*>& InInstances, FNiagaraDataSet& InDataSet, float InDeltaSeconds, int32 InSpawnNum, const FGraphEventRef& InMyCompletionGraphEvent)
+FNiagaraSystemSimulationTickContext::FNiagaraSystemSimulationTickContext(FNiagaraSystemSimulation* InOwner, TArray<FNiagaraSystemInstance*>& InInstances, FNiagaraDataSet& InDataSet, float InDeltaSeconds, int32 InSpawnNum, int32 InEffectsQuality, const FGraphEventRef& InMyCompletionGraphEvent)
 	: Owner(InOwner)
 	, System(InOwner->GetSystem())
 	, Instances(InInstances)
 	, DataSet(InDataSet)
 	, DeltaSeconds(InDeltaSeconds)
 	, SpawnNum(InSpawnNum)
+	, EffectsQuality(InEffectsQuality)
 	, MyCompletionGraphEvent(InMyCompletionGraphEvent)
 	, FinalizeEvents(nullptr)
 	, bTickAsync(GbParallelSystemSimTick && FApp::ShouldUseThreadingForPerformance() && InMyCompletionGraphEvent.IsValid())
@@ -780,7 +781,8 @@ void FNiagaraSystemSimulation::Tick_GameThread(float DeltaSeconds, const FGraphE
 			}
 		}
 
-		FNiagaraSystemSimulationTickContext Context(this, SystemInstances, MainDataSet, DeltaSeconds, SpawnNum, MyCompletionGraphEvent);
+		static const auto EffectsQualityCVar = IConsoleManager::Get().FindConsoleVariable(TEXT("sg.EffectsQuality"));
+		FNiagaraSystemSimulationTickContext Context(this, SystemInstances, MainDataSet, DeltaSeconds, SpawnNum, EffectsQualityCVar->GetInt(), MyCompletionGraphEvent);
 
 		//Now kick of the concurrent tick.
 		if (Context.bTickAsync)
@@ -933,8 +935,8 @@ void FNiagaraSystemSimulation::Spawn_GameThread(float DeltaSeconds)
 	if ( SpawningInstances.Num() > 0 )
 	{
 		//-OPT: This can be async :)
-
-		FNiagaraSystemSimulationTickContext Context(this, SpawningInstances, SpawningDataSet, DeltaSeconds, SpawningInstances.Num(), nullptr);
+		static const auto EffectsQualityCVar = IConsoleManager::Get().FindConsoleVariable(TEXT("sg.EffectsQuality"));
+		FNiagaraSystemSimulationTickContext Context(this, SpawningInstances, SpawningDataSet, DeltaSeconds, EffectsQualityCVar->GetInt(), SpawningInstances.Num(), nullptr);
 		Tick_Concurrent(Context);
 
 		check(MainDataSet.GetCurrentDataChecked().GetNumInstances() == SystemInstances.Num());
@@ -1084,7 +1086,7 @@ void FNiagaraSystemSimulation::Tick_Concurrent(FNiagaraSystemSimulationTickConte
 void FNiagaraSystemSimulation::TickFastPath(FNiagaraSystemSimulationTickContext& Context)
 {
 	SCOPE_CYCLE_COUNTER(STAT_NiagaraSystemSim_FastPathCNC);
-	
+
 	// PrepareForSystemSimulate
 	for (FNiagaraSystemInstance* SystemInstance : Context.Instances)
 	{
@@ -1105,8 +1107,13 @@ void FNiagaraSystemSimulation::TickFastPath(FNiagaraSystemSimulationTickContext&
 			EmitterMap.Engine.DeltaTime = Context.DeltaSeconds;
 			EmitterMap.Engine.Emitter.NumParticles = SystemInstance->GetNumParticles(EmitterIndex);
 			EmitterMap.Engine.Owner.Velocity = SystemInstance->GetOwnerVelocity();
-			EmitterMap.Engine.GlobalSpawnCountScale = INiagaraModule::GetGlobalSpawnCountScale();
 			EmitterMap.Emitter.ExecutionState = EmitterInstance->GetExecutionState();
+			EmitterMap.Engine.GlobalSpawnCountScale = INiagaraModule::GetGlobalSpawnCountScale();
+
+			if (UNiagaraEmitter* CachedEmitter = EmitterInstance->GetCachedEmitter())
+			{
+				EmitterMap.Emitter.SpawnCountScale = CachedEmitter->GetSpawnCountScale(Context.EffectsQuality);
+			}
 		}
 	}
 
diff --git a/Engine/Plugins/FX/Niagara/Source/Niagara/Public/NiagaraModule.h b/Engine/Plugins/FX/Niagara/Source/Niagara/Public/NiagaraModule.h
--- a/Engine/Plugins/FX/Niagara/Source/Niagara/Public/NiagaraModule.h
+++ b/Engine/Plugins/FX/Niagara/Source/Niagara/Public/NiagaraModule.h
@@ -102,6 +102,7 @@ public:
 	FORCEINLINE static const FNiagaraVariable&  GetVar_Engine_ExecutionCount() { return Engine_ExecutionCount; }
 	FORCEINLINE static const FNiagaraVariable&  GetVar_Engine_Emitter_NumParticles() { return Engine_Emitter_NumParticles; }
 	FORCEINLINE static const FNiagaraVariable&  GetVar_Engine_Emitter_TotalSpawnedParticles() { return Engine_Emitter_TotalSpawnedParticles; }
+	FORCEINLINE static const FNiagaraVariable&  GetVar_Engine_Emitter_SpawnCountScale() { return Engine_Emitter_SpawnCountScale; }
 	FORCEINLINE static const FNiagaraVariable&  GetVar_Engine_System_TickCount() { return Engine_System_TickCount; }
 	FORCEINLINE static const FNiagaraVariable&  GetVar_Engine_System_NumEmittersAlive() { return Engine_System_NumEmittersAlive; }
 	FORCEINLINE static const FNiagaraVariable&  GetVar_Engine_System_NumEmitters() { return Engine_System_NumEmitters; }
@@ -114,6 +115,7 @@ public:
 	FORCEINLINE static const FNiagaraVariable&  GetVar_Emitter_Age() { return Emitter_Age; }
 	FORCEINLINE static const FNiagaraVariable&  GetVar_Emitter_LocalSpace() { return Emitter_LocalSpace; }
 	FORCEINLINE static const FNiagaraVariable&  GetVar_Emitter_Determinism() { return Emitter_Determinism; }
+	FORCEINLINE static const FNiagaraVariable&  GetVar_Emitter_OverrideGlobalSpawnCountScale() { return Emitter_OverrideGlobalSpawnCountScale; }
 	FORCEINLINE static const FNiagaraVariable&  GetVar_Emitter_RandomSeed() { return Emitter_RandomSeed; }
 	FORCEINLINE static const FNiagaraVariable&  GetVar_Emitter_SpawnRate() { return Emitter_SpawnRate; }
 	FORCEINLINE static const FNiagaraVariable&  GetVar_Emitter_SpawnInterval() { return Emitter_SpawnInterval; }
@@ -207,6 +209,7 @@ private:
 	static FNiagaraVariable Engine_ExecutionCount;
 	static FNiagaraVariable Engine_Emitter_NumParticles;
 	static FNiagaraVariable Engine_Emitter_TotalSpawnedParticles;
+	static FNiagaraVariable Engine_Emitter_SpawnCountScale;
 	static FNiagaraVariable Engine_System_TickCount;
 	static FNiagaraVariable Engine_System_NumEmittersAlive;
 	static FNiagaraVariable Engine_System_NumEmitters;
@@ -219,6 +222,7 @@ private:
 	static FNiagaraVariable Emitter_Age;
 	static FNiagaraVariable Emitter_LocalSpace;
 	static FNiagaraVariable Emitter_Determinism;
+	static FNiagaraVariable Emitter_OverrideGlobalSpawnCountScale;
 	static FNiagaraVariable Emitter_SimulationTarget;
 	static FNiagaraVariable Emitter_RandomSeed;
 	static FNiagaraVariable Emitter_SpawnRate;
diff --git a/Engine/Plugins/FX/Niagara/Source/Niagara/Public/NiagaraSystemFastPath.h b/Engine/Plugins/FX/Niagara/Source/Niagara/Public/NiagaraSystemFastPath.h
--- a/Engine/Plugins/FX/Niagara/Source/Niagara/Public/NiagaraSystemFastPath.h
+++ b/Engine/Plugins/FX/Niagara/Source/Niagara/Public/NiagaraSystemFastPath.h
@@ -516,6 +516,7 @@ namespace FNiagaraEmitterFastPath
 		float LoopedAge;
 		float NormalizedLoopAge;
 		float DistanceTraveled;
+		float SpawnCountScale;
 		FParamMap0_Emitter_Scalability Scalability;
 		TArray<FParamMap0_Emitter_SpawnRate> SpawnRate;
 		TArray<FParamMap0_Emitter_SpawnPerUnit> SpawnPerUnit;
diff --git a/Engine/Plugins/FX/Niagara/Source/Niagara/Public/NiagaraSystemInstance.h b/Engine/Plugins/FX/Niagara/Source/Niagara/Public/NiagaraSystemInstance.h
--- a/Engine/Plugins/FX/Niagara/Source/Niagara/Public/NiagaraSystemInstance.h
+++ b/Engine/Plugins/FX/Niagara/Source/Niagara/Public/NiagaraSystemInstance.h
@@ -115,6 +115,7 @@ public:
 	float GetOwnerLODDistance() const { return OwnerLODDistanceParam.GetValue(); }
 
 	int32 GetNumParticles(int32 EmitterIndex) const { return ParameterNumParticleBindings[EmitterIndex].GetValue(); }
+	float GetSpawnCountScale(int32 EmitterIndex) const { return ParameterSpawnCountScaleBindings[EmitterIndex].GetValue(); }
 
 	FVector GetOwnerVelocity() const { return OwnerVelocityParam.GetValue(); }
 
@@ -349,6 +350,7 @@ private:
 
 	FNiagaraParameterDirectBinding<int32> OwnerExecutionStateParam;
 
+	TArray<FNiagaraParameterDirectBinding<float>> ParameterSpawnCountScaleBindings;
 	TArray<FNiagaraParameterDirectBinding<int32>> ParameterNumParticleBindings;
 	TArray<FNiagaraParameterDirectBinding<int32>> ParameterTotalSpawnedParticlesBindings;
 
@@ -421,6 +423,7 @@ public:
 
 		TArray<int32> EmitterNumParticles;
 		TArray<int32> EmitterTotalSpawnedParticles;
+		TArray<float> EmitterSpawnCountScale;
 		int32 NumAlive;
 
 		float SafeTimeSinceRendererd;
@@ -431,6 +434,7 @@ public:
 		{
 			EmitterNumParticles.AddUninitialized(NumEmitters);
 			EmitterTotalSpawnedParticles.AddUninitialized(NumEmitters);
+			EmitterSpawnCountScale.AddUninitialized(NumEmitters);
 		}
 	};
 
diff --git a/Engine/Plugins/FX/Niagara/Source/Niagara/Public/NiagaraSystemSimulation.h b/Engine/Plugins/FX/Niagara/Source/Niagara/Public/NiagaraSystemSimulation.h
--- a/Engine/Plugins/FX/Niagara/Source/Niagara/Public/NiagaraSystemSimulation.h
+++ b/Engine/Plugins/FX/Niagara/Source/Niagara/Public/NiagaraSystemSimulation.h
@@ -125,7 +125,7 @@ struct FNiagaraParameterStoreToDataSetBinding
 
 struct FNiagaraSystemSimulationTickContext
 {
-	FNiagaraSystemSimulationTickContext(class FNiagaraSystemSimulation* Owner, TArray<FNiagaraSystemInstance*>& Instances, FNiagaraDataSet& DataSet, float DeltaSeconds, int32 SpawnNum, const FGraphEventRef& MyCompletionGraphEvent);
+	FNiagaraSystemSimulationTickContext(class FNiagaraSystemSimulation* Owner, TArray<FNiagaraSystemInstance*>& Instances, FNiagaraDataSet& DataSet, float DeltaSeconds, int32 SpawnNum, int EffectsQuality, const FGraphEventRef& MyCompletionGraphEvent);
 
 	class FNiagaraSystemSimulation*		Owner;
 	UNiagaraSystem*						System;
@@ -136,6 +136,8 @@ struct FNiagaraSystemSimulationTickContext
 	float								DeltaSeconds;
 	int32								SpawnNum;
 
+	int									EffectsQuality;
+
 	FGraphEventRef						MyCompletionGraphEvent;
 	FGraphEventArray*					FinalizeEvents;
 
diff --git a/Engine/Plugins/FX/Niagara/Source/NiagaraEditor/Private/NiagaraHlslTranslator.cpp b/Engine/Plugins/FX/Niagara/Source/NiagaraEditor/Private/NiagaraHlslTranslator.cpp
--- a/Engine/Plugins/FX/Niagara/Source/NiagaraEditor/Private/NiagaraHlslTranslator.cpp
+++ b/Engine/Plugins/FX/Niagara/Source/NiagaraEditor/Private/NiagaraHlslTranslator.cpp
@@ -2725,6 +2725,7 @@ bool FHlslNiagaraTranslator::ShouldInterpolateParameter(const FNiagaraVariable&
 		Parameter == SYS_PARAM_EMITTER_SPAWN_INTERVAL ||
 		Parameter == SYS_PARAM_EMITTER_INTERP_SPAWN_START_DT ||
 		Parameter == SYS_PARAM_ENGINE_EMITTER_TOTAL_SPAWNED_PARTICLES ||
+		Parameter == SYS_PARAM_ENGINE_EMITTER_SPAWN_COUNT_SCALE ||
 		Parameter == SYS_PARAM_EMITTER_RANDOM_SEED ||
 		Parameter == SYS_PARAM_ENGINE_SYSTEM_TICK_COUNT)
 	{
@@ -3487,6 +3488,12 @@ bool FHlslNiagaraTranslator::GetLiteralConstantVariable(FNiagaraVariable& OutVar
 			OutVar.SetValue(bEmitterDeterminism ? FNiagaraBool(true) : FNiagaraBool(false));
 			return true;
 		}
+		if (OutVar == FNiagaraVariable(FNiagaraTypeDefinition::GetBoolDef(), TEXT("Emitter.OverrideGlobalSpawnCountScale")))
+		{
+			bool bOverrideGlobalSpawnCountScale = CompileOptions.AdditionalDefines.Contains(ResolvedVar.GetName().ToString());
+			OutVar.SetValue(bOverrideGlobalSpawnCountScale ? FNiagaraBool(true) : FNiagaraBool(false));
+			return true;
+		}
 		if (OutVar == FNiagaraVariable(FNiagaraTypeDefinition::GetSimulationTargetEnum(), TEXT("Emitter.SimulationTarget")))
 		{
 			FNiagaraInt32 EnumValue;
diff --git a/Engine/Plugins/FX/Niagara/Source/NiagaraEditor/Private/NiagaraParameterMapHistory.cpp b/Engine/Plugins/FX/Niagara/Source/NiagaraEditor/Private/NiagaraParameterMapHistory.cpp
--- a/Engine/Plugins/FX/Niagara/Source/NiagaraEditor/Private/NiagaraParameterMapHistory.cpp
+++ b/Engine/Plugins/FX/Niagara/Source/NiagaraEditor/Private/NiagaraParameterMapHistory.cpp
@@ -1413,6 +1413,11 @@ bool FCompileConstantResolver::ResolveConstant(FNiagaraVariable& OutConstant) co
 		OutConstant.SetValue(Emitter->bDeterminism ? FNiagaraBool(true) : FNiagaraBool(false));
 		return true;
 	}
+	if (Emitter && OutConstant == FNiagaraVariable(FNiagaraTypeDefinition::GetBoolDef(), TEXT("Emitter.OverrideGlobalSpawnCountScale")))
+	{
+		OutConstant.SetValue(Emitter->bOverrideGlobalSpawnCountScale ? FNiagaraBool(true) : FNiagaraBool(false));
+		return true;
+	}
 	if (Emitter && OutConstant == FNiagaraVariable(FNiagaraTypeDefinition::GetSimulationTargetEnum(), TEXT("Emitter.SimulationTarget")))
 	{
 		FNiagaraInt32 EnumValue;
