Add missing include for non-unity build.

diff --git a/Engine/Plugins/FX/Niagara/Source/Niagara/Private/NiagaraEmitter.cpp b/Engine/Plugins/FX/Niagara/Source/Niagara/Private/NiagaraEmitter.cpp
--- a/Engine/Plugins/FX/Niagara/Source/Niagara/Private/NiagaraEmitter.cpp
+++ b/Engine/Plugins/FX/Niagara/Source/Niagara/Private/NiagaraEmitter.cpp
@@ -12,6 +12,7 @@
 #include "NiagaraSystem.h"
 #include "NiagaraStats.h"
 #include "Modules/ModuleManager.h"
+#include "Misc/ConfigCacheIni.h"
 
 #if WITH_EDITOR
 const FName UNiagaraEmitter::PrivateMemberNames::EventHandlerScriptProps = GET_MEMBER_NAME_CHECKED(UNiagaraEmitter, EventHandlerScriptProps);
