- Cache IsReadyToRun in PostLoad for cooked platforms as it can not change
- Saves 20us - 30us for some of our VFX when the system has completed

diff --git a/Engine/Plugins/FX/Niagara/Source/Niagara/Classes/NiagaraSystem.h b/Engine/Plugins/FX/Niagara/Source/Niagara/Classes/NiagaraSystem.h
--- a/Engine/Plugins/FX/Niagara/Source/Niagara/Classes/NiagaraSystem.h
+++ b/Engine/Plugins/FX/Niagara/Source/Niagara/Classes/NiagaraSystem.h
@@ -182,6 +182,10 @@ public:
 	UNiagaraScript* GetSystemSpawnScript();
 	UNiagaraScript* GetSystemUpdateScript();
 
+private:
+	bool IsReadyToRunInternal() const;
+	bool bIsReadyToRunCached = false;
+public:
 	bool IsReadyToRun() const;
 
 	FORCEINLINE bool NeedsWarmup()const { return WarmupTickCount > 0 && WarmupTickDelta > SMALL_NUMBER; }
@@ -283,7 +287,7 @@ public:
 
 	/** Whether or not fixed bounds are enabled. */
 	UPROPERTY(EditAnywhere, Category = "System", meta = (InlineEditConditionToggle))
-		uint32 bFixedBounds : 1;
+	uint32 bFixedBounds : 1;
 
 	TStatId GetStatID(bool bGameThread, bool bConcurrent)const;
 
diff --git a/Engine/Plugins/FX/Niagara/Source/Niagara/Private/NiagaraSystem.cpp b/Engine/Plugins/FX/Niagara/Source/Niagara/Private/NiagaraSystem.cpp
--- a/Engine/Plugins/FX/Niagara/Source/Niagara/Private/NiagaraSystem.cpp
+++ b/Engine/Plugins/FX/Niagara/Source/Niagara/Private/NiagaraSystem.cpp
@@ -428,6 +428,10 @@ void UNiagaraSystem::PostLoad()
 	}
 #endif // WITH_EDITORONLY_DATA
 
+	if ( FPlatformProperties::RequiresCookedData() )
+	{
+		bIsReadyToRunCached = IsReadyToRunInternal();
+	}
 }
 
 #if WITH_EDITORONLY_DATA
@@ -501,7 +505,7 @@ const TArray<FNiagaraEmitterHandle>& UNiagaraSystem::GetEmitterHandles()const
 	return EmitterHandles;
 }
 
-bool UNiagaraSystem::IsReadyToRun() const
+bool UNiagaraSystem::IsReadyToRunInternal() const
 {
 	if (!SystemSpawnScript || !SystemUpdateScript)
 	{
@@ -539,6 +543,18 @@ bool UNiagaraSystem::IsReadyToRun() const
 	return true;
 }
 
+bool UNiagaraSystem::IsReadyToRun() const
+{
+	if (FPlatformProperties::RequiresCookedData())
+	{
+		return bIsReadyToRunCached;
+	}
+	else
+	{
+		return IsReadyToRunInternal();
+	}
+}
+
 #if WITH_EDITORONLY_DATA
 bool UNiagaraSystem::HasOutstandingCompilationRequests() const
 {
