- Fix for SpawningInstances & EffectsQualityCVar being swapped inside Spawn_GameThread

#rnx

diff --git a/Engine/Plugins/FX/Niagara/Source/Niagara/Private/NiagaraSystemSimulation.cpp b/Engine/Plugins/FX/Niagara/Source/Niagara/Private/NiagaraSystemSimulation.cpp
--- a/Engine/Plugins/FX/Niagara/Source/Niagara/Private/NiagaraSystemSimulation.cpp
+++ b/Engine/Plugins/FX/Niagara/Source/Niagara/Private/NiagaraSystemSimulation.cpp
@@ -936,7 +936,7 @@ void FNiagaraSystemSimulation::Spawn_GameThread(float DeltaSeconds)
 	{
 		//-OPT: This can be async :)
 		static const auto EffectsQualityCVar = IConsoleManager::Get().FindConsoleVariable(TEXT("sg.EffectsQuality"));
-		FNiagaraSystemSimulationTickContext Context(this, SpawningInstances, SpawningDataSet, DeltaSeconds, EffectsQualityCVar->GetInt(), SpawningInstances.Num(), nullptr);
+		FNiagaraSystemSimulationTickContext Context(this, SpawningInstances, SpawningDataSet, DeltaSeconds, SpawningInstances.Num(), EffectsQualityCVar->GetInt(), nullptr);
 		Tick_Concurrent(Context);
 
 		check(MainDataSet.GetCurrentDataChecked().GetNumInstances() == SystemInstances.Num());
