Fixed crash around memstack map and unbound parameters when using Niagara GPU particles.
Removed FNiagaraComputeExecutionContext::DebugSimName in Shipping.

diff --git a/Engine/Plugins/FX/Niagara/Source/Niagara/Classes/NiagaraScriptExecutionContext.h b/Engine/Plugins/FX/Niagara/Source/Niagara/Classes/NiagaraScriptExecutionContext.h
--- a/Engine/Plugins/FX/Niagara/Source/Niagara/Classes/NiagaraScriptExecutionContext.h
+++ b/Engine/Plugins/FX/Niagara/Source/Niagara/Classes/NiagaraScriptExecutionContext.h
@@ -185,8 +185,11 @@ private:
 public:
 	static uint32 TickCounter;
 
+#if !UE_BUILD_SHIPPING
 	//Persistent state 
 	FString DebugSimName;
+#endif
+
 	class FNiagaraDataSet *MainDataSet;
 	UNiagaraScript* GPUScript;
 	class FNiagaraShaderScript*  GPUScript_RT;
@@ -210,6 +213,9 @@ public:
 	uint32 MaxUpdateIterations;
 	TSet<uint32> SpawnStages;
 
+	/** Temp data used in NiagaraEmitterInstanceBatcher::ExecuteAll() to avoid creating a map per FNiagaraComputeExecutionContext */
+	mutable int32 ScratchIndex = INDEX_NONE;
+
 #if WITH_EDITORONLY_DATA
 	mutable FRHIGPUMemoryReadback *GPUDebugDataReadbackFloat;
 	mutable FRHIGPUMemoryReadback *GPUDebugDataReadbackInt;
diff --git a/Engine/Plugins/FX/Niagara/Source/Niagara/Private/NiagaraEmitterInstanceBatcher.cpp b/Engine/Plugins/FX/Niagara/Source/Niagara/Private/NiagaraEmitterInstanceBatcher.cpp
--- a/Engine/Plugins/FX/Niagara/Source/Niagara/Private/NiagaraEmitterInstanceBatcher.cpp
+++ b/Engine/Plugins/FX/Niagara/Source/Niagara/Private/NiagaraEmitterInstanceBatcher.cpp
@@ -443,36 +443,36 @@ void NiagaraEmitterInstanceBatcher::DispatchAllOnCompute(FOverlappableTicks& Ove
 	RHICmdList.TransitionResources(EResourceTransitionAccess::EReadable, EResourceTransitionPipeline::EComputeToCompute, CurrBufferIntFloat.GetData(), CurrBufferIntFloat.Num());
 	RHICmdList.FlushComputeShaderCache();
 
-		for (FNiagaraGPUSystemTick* Tick : OverlappableTick)
+	for (FNiagaraGPUSystemTick* Tick : OverlappableTick)
+	{
+		uint32 DispatchCount = Tick->Count;
+		FNiagaraComputeInstanceData* Instances = Tick->GetInstanceData();
+		for (uint32 Index = 0; Index < DispatchCount; Index++)
 		{
-			uint32 DispatchCount = Tick->Count;
-			FNiagaraComputeInstanceData* Instances = Tick->GetInstanceData();
-			for (uint32 Index = 0; Index < DispatchCount; Index++)
+			FNiagaraComputeInstanceData& Instance = Instances[Index];
+			FNiagaraComputeExecutionContext* Context = Instance.Context;
+			if (Context && Context->GPUScript_RT->GetShader())
 			{
-				FNiagaraComputeInstanceData& Instance = Instances[Index];
-				FNiagaraComputeExecutionContext* Context = Instance.Context;
-				if (Context && Context->GPUScript_RT->GetShader())
-				{
-					FNiagaraComputeExecutionContext::TickCounter++;
+				FNiagaraComputeExecutionContext::TickCounter++;
 
-					// run shader, sim and spawn in a single dispatch
-					DispatchMultipleStages(*Tick, &Instance, RHICmdList, ViewUniformBuffer, Context->GPUScript_RT->GetShader());
+				// run shader, sim and spawn in a single dispatch
+				DispatchMultipleStages(*Tick, &Instance, RHICmdList, ViewUniformBuffer, Context->GPUScript_RT->GetShader());
 
-					FNiagaraDataBuffer* CurrentData = Instance.CurrentData;
-					if (bSetReadback && Tick->bIsFinalTick)
+				FNiagaraDataBuffer* CurrentData = Instance.CurrentData;
+				if (bSetReadback && Tick->bIsFinalTick)
+				{
+					// Now that the current data is not required anymore, stage it for readback.
+					if (CurrentData->GetNumInstances() && Context->EmitterInstanceReadback.GPUCountOffset == INDEX_NONE && CurrentData->GetGPUInstanceCountBufferOffset() != INDEX_NONE)
 					{
-						// Now that the current data is not required anymore, stage it for readback.
-						if (CurrentData->GetNumInstances() && Context->EmitterInstanceReadback.GPUCountOffset == INDEX_NONE && CurrentData->GetGPUInstanceCountBufferOffset() != INDEX_NONE)
-						{
-							// Transfer the GPU instance counter ownership to the context. Note that when bSetReadback is true, a readback request will be performed later in the tick update.
-							Context->EmitterInstanceReadback.GPUCountOffset = CurrentData->GetGPUInstanceCountBufferOffset();
-							Context->EmitterInstanceReadback.CPUCount = CurrentData->GetNumInstances();
-							CurrentData->ClearGPUInstanceCountBufferOffset();
-						}
+						// Transfer the GPU instance counter ownership to the context. Note that when bSetReadback is true, a readback request will be performed later in the tick update.
+						Context->EmitterInstanceReadback.GPUCountOffset = CurrentData->GetGPUInstanceCountBufferOffset();
+						Context->EmitterInstanceReadback.CPUCount = CurrentData->GetNumInstances();
+						CurrentData->ClearGPUInstanceCountBufferOffset();
 					}
 				}
 			}
 		}
+	}
 
 	//
 	//	Now Copy to staging buffer the data we want to read back (alive particle count). And make buffer ready for that and draw commands on the graphics pipe too.
@@ -491,7 +491,7 @@ void NiagaraEmitterInstanceBatcher::DispatchAllOnCompute(FOverlappableTicks& Ove
 	RHICmdList.FlushComputeShaderCache();
 	if (GNiagaraSubmitCommands)
 	{
-	RHICmdList.SubmitCommandsHint();
+		RHICmdList.SubmitCommandsHint();
 	}
 }
 
@@ -523,82 +523,101 @@ void NiagaraEmitterInstanceBatcher::ExecuteAll(FRHICommandList &RHICmdList, FRHI
 
 	SCOPED_DRAW_EVENT(RHICmdList, NiagaraEmitterInstanceBatcher_ExecuteAll);
 
-		TArray<FOverlappableTicks, TInlineAllocator<2>> SimPasses;
+	FMemMark Mark(FMemStack::Get());
+	TArray<FOverlappableTicks, TMemStackAllocator<> > SimPasses;
+	{
+		TArray< FNiagaraComputeExecutionContext* , TMemStackAllocator<> > RelevantContexts;
+		TArray< FNiagaraGPUSystemTick* , TMemStackAllocator<> > RelevantTicks;
+		for (FNiagaraGPUSystemTick& Tick : Ticks_RT)
 		{
-			TMap<FNiagaraComputeExecutionContext*, FOverlappableTicks> ContextToTicks;
-
-			// Ticks are added in order. Two tick with the same context should not overlap so should be in two different batch.
-			// Those ticks should still be executed in order.
-			for (FNiagaraGPUSystemTick& Tick : Ticks_RT)
+			FNiagaraComputeInstanceData* Data = Tick.GetInstanceData();
+			FNiagaraComputeExecutionContext* Context = Data->Context;
+			// This assumes all emitters fallback to the same FNiagaraShaderScript*.
+			FNiagaraShader* ComputeShader = Context->GPUScript_RT->GetShader();
+			if (!ComputeShader)
 			{
-				FNiagaraComputeInstanceData* Data = Tick.GetInstanceData();
-
-				// This assumes all emitters fallback to the same FNiagaraShaderScript*.
-				FNiagaraShader* ComputeShader = Data->Context->GPUScript_RT->GetShader();
-				if (!ComputeShader)
-				{
-					continue;
-				}
-				else if (GNiagaraAllowTickBeforeRender && ComputeShader->ViewUniformBufferParam.IsBound() != (ViewUniformBuffer != nullptr))
-				{   // When allowing tick before render, skip this emitter if it is not in the right pass.
-					continue;
-				}
-
-				ContextToTicks.FindOrAdd(Data->Context).Add(&Tick);
+				continue;
+			}
+			else if (GNiagaraAllowTickBeforeRender && ComputeShader->ViewUniformBufferParam.IsBound() != (ViewUniformBuffer != nullptr))
+			{   // When allowing tick before render, skip this emitter if it is not in the right pass.
+				continue;
 			}
 
-			// Count the maximum number of tick per context to know the number of passes we need to do
-			uint32 NumSimPass = 0;
-			for (auto& Ticks : ContextToTicks)
+			Tick.bIsFinalTick = false; // @todo : this is true sometimes, needs investigation
+			if (Context->ScratchIndex == INDEX_NONE)
 			{
-				NumSimPass = FMath::Max(NumSimPass, (uint32)Ticks.Value.Num());
-				Ticks.Value.Last()->bIsFinalTick = true;
+				RelevantContexts.Add(Context);
 			}
+			// Here scratch index represent the index of the last tick
+			Context->ScratchIndex = RelevantTicks.Add(&Tick);
+		}
 
-		// Force dispatches to run individually, this should only be used for debugging as it is highly inefficient on the GPU
-		if (GNiagaraOverlapCompute == 0)
+		// Set bIsFinalTick for the last tick of each context and reset the scratch index.
+		const int32 ScrachIndexReset = GNiagaraOverlapCompute ? 0 : INDEX_NONE;
+		for (FNiagaraComputeExecutionContext* Context : RelevantContexts)
 		{
-			for (auto& Ticks : ContextToTicks)
+			RelevantTicks[Context->ScratchIndex]->bIsFinalTick = true;
+			Context->ScratchIndex = ScrachIndexReset;
+		}
+
+		if (GNiagaraOverlapCompute)
+		{
+			// Transpose now only once the data to get all independent tick per pass
+			SimPasses.Reserve(2); // Safe bet!
+
+			for (FNiagaraGPUSystemTick* Tick : RelevantTicks)
 			{
-				SimPasses.Reserve(Ticks.Value.Num() - SimPasses.GetSlack());
+				FNiagaraComputeExecutionContext* Context = Tick->GetInstanceData()->Context;
+				const int32 ScratchIndex = Context->ScratchIndex;
+				check(ScratchIndex != INDEX_NONE);
 
-				for (auto& Tick : Ticks.Value)
+				if (ScratchIndex >= SimPasses.Num())
 				{
-					SimPasses.AddDefaulted_GetRef().Add(Tick);
+					SimPasses.AddDefaulted(SimPasses.Num() - ScratchIndex + 1);
+					if (ScratchIndex == 0)
+					{
+						SimPasses[0].Reserve(RelevantContexts.Num()); // Guarantied!
+					}
+				}
+				SimPasses[ScratchIndex].Add(Tick);
+				// Scratch index is now the number of passes for this context.
+				if (Tick->bIsFinalTick)
+				{
+					// Reset to default as it will no longer be used.
+					Context->ScratchIndex = INDEX_NONE;
+				}
+				else
+				{
+					Context->ScratchIndex += 1;
 				}
 			}
 		}
 		else
 		{
-			// Transpose now only once the data to get all independent tick per pass
-			SimPasses.AddDefaulted(NumSimPass);
-			for (auto& Ticks : ContextToTicks)
+			// Force dispatches to run individually, this should only be used for debugging as it is highly inefficient on the GPU
+			SimPasses.Reserve(RelevantTicks.Num()); // Guarantied!
+			for (FNiagaraGPUSystemTick* Tick : RelevantTicks)
 			{
-				uint32 tickPass = 0;
-				for (auto& Tick : Ticks.Value)
-				{
-					SimPasses[tickPass].Add(Tick);
-					tickPass++;
-				}
+				SimPasses.AddDefaulted_GetRef().Add(Tick);
 			}
 		}
 	}
 
-		for (auto& SimPass : SimPasses)
-		{
-			FNiagaraBufferArray DestDataBuffers;
-			FNiagaraBufferArray CurrDataBuffers;
-			FNiagaraBufferArray DestBufferIntFloat;
-			FNiagaraBufferArray CurrBufferIntFloat;
+	for (auto& SimPass : SimPasses)
+	{
+		FNiagaraBufferArray DestDataBuffers;
+		FNiagaraBufferArray CurrDataBuffers;
+		FNiagaraBufferArray DestBufferIntFloat;
+		FNiagaraBufferArray CurrBufferIntFloat;
 
-			// This initial pass gathers all the buffers that are read from and written to so we can do batch resource transitions.
-			// It also ensures the GPU buffers are large enough to hold everything.
-			ResizeBuffersAndGatherResources(SimPass, RHICmdList, DestDataBuffers, CurrDataBuffers, DestBufferIntFloat, CurrBufferIntFloat);
+		// This initial pass gathers all the buffers that are read from and written to so we can do batch resource transitions.
+		// It also ensures the GPU buffers are large enough to hold everything.
+		ResizeBuffersAndGatherResources(SimPass, RHICmdList, DestDataBuffers, CurrDataBuffers, DestBufferIntFloat, CurrBufferIntFloat);
 
 		SCOPED_DRAW_EVENT(RHICmdList, NiagaraGPUSimulation);
 		SCOPED_GPU_STAT(RHICmdList, NiagaraGPUSimulation);
-			DispatchAllOnCompute(SimPass, RHICmdList, ViewUniformBuffer, DestDataBuffers, CurrDataBuffers, DestBufferIntFloat, CurrBufferIntFloat, bSetReadback);
-		}
+		DispatchAllOnCompute(SimPass, RHICmdList, ViewUniformBuffer, DestDataBuffers, CurrDataBuffers, DestBufferIntFloat, CurrBufferIntFloat, bSetReadback);
+	}
 }
 
 void NiagaraEmitterInstanceBatcher::PreInitViews(FRHICommandListImmediate& RHICmdList, bool bAllowGPUParticleUpdate)
@@ -853,7 +872,6 @@ void NiagaraEmitterInstanceBatcher::SortGPUParticles(FRHICommandListImmediate& R
 		{
 			NumFramesRequiringShrinking = 0;
 		}
-
 	}
 	else // If the are no sort task, we don't need any of the sort buffers.
 	{
@@ -1085,14 +1103,9 @@ void NiagaraEmitterInstanceBatcher::Run(const FNiagaraGPUSystemTick& Tick, const
 	RHICmdList.SetComputeShader(Shader->GetComputeShader());
 
 	// #todo(dmp): clean up this logic for shader stages on first frame
-	if (Tick.bNeedsReset)
+	if (Shader->SimStartParam.IsBound())
 	{
-		int v = 1;
-		RHICmdList.SetShaderParameter(Shader->GetComputeShader(), Shader->SimStartParam.GetBufferIndex(), Shader->SimStartParam.GetBaseIndex(), Shader->SimStartParam.GetNumBytes(), &v);
-	}
-	else
-	{
-		int v = 0;
+		int v = Tick.bNeedsReset ? 1 : 0;
 		RHICmdList.SetShaderParameter(Shader->GetComputeShader(), Shader->SimStartParam.GetBufferIndex(), Shader->SimStartParam.GetBaseIndex(), Shader->SimStartParam.GetNumBytes(), &v);
 	}
 	
@@ -1146,10 +1159,15 @@ void NiagaraEmitterInstanceBatcher::Run(const FNiagaraGPUSystemTick& Tick, const
 	}
 
 	RHICmdList.SetShaderParameter(Shader->GetComputeShader(), Shader->NumSpawnedInstancesParam.GetBufferIndex(), Shader->NumSpawnedInstancesParam.GetBaseIndex(), Shader->NumSpawnedInstancesParam.GetNumBytes(), &InstancesToSpawnThisFrame);				// number of instances in the spawn run
-	RHICmdList.SetShaderParameter(Shader->GetComputeShader(), Shader->ShaderStageIndexParam.GetBufferIndex(), Shader->ShaderStageIndexParam.GetBaseIndex(), Shader->ShaderStageIndexParam.GetNumBytes(), &ShaderStageIndex);					// 0, except if several stages are defined
-
+	if (Shader->ShaderStageIndexParam.IsBound())
+	{
+		RHICmdList.SetShaderParameter(Shader->GetComputeShader(), Shader->ShaderStageIndexParam.GetBufferIndex(), Shader->ShaderStageIndexParam.GetBaseIndex(), Shader->ShaderStageIndexParam.GetNumBytes(), &ShaderStageIndex);					// 0, except if several stages are defined
+	}
 	const int32 DefaultIterationCount = -1;
-	RHICmdList.SetShaderParameter(Shader->GetComputeShader(), Shader->IterationInterfaceCount.GetBufferIndex(), Shader->IterationInterfaceCount.GetBaseIndex(), Shader->IterationInterfaceCount.GetNumBytes(), &DefaultIterationCount);					// 0, except if several stages are defined
+	if (Shader->IterationInterfaceCount.IsBound())
+	{
+		RHICmdList.SetShaderParameter(Shader->GetComputeShader(), Shader->IterationInterfaceCount.GetBufferIndex(), Shader->IterationInterfaceCount.GetBaseIndex(), Shader->IterationInterfaceCount.GetNumBytes(), &DefaultIterationCount);					// 0, except if several stages are defined
+	}
 
 	if (IterationInterface)
 	{
diff --git a/Engine/Plugins/FX/Niagara/Source/Niagara/Private/NiagaraScriptExecutionContext.cpp b/Engine/Plugins/FX/Niagara/Source/Niagara/Private/NiagaraScriptExecutionContext.cpp
--- a/Engine/Plugins/FX/Niagara/Source/Niagara/Private/NiagaraScriptExecutionContext.cpp
+++ b/Engine/Plugins/FX/Niagara/Source/Niagara/Private/NiagaraScriptExecutionContext.cpp
@@ -454,7 +454,9 @@ void FNiagaraComputeExecutionContext::Reset(NiagaraEmitterInstanceBatcher* Batch
 
 void FNiagaraComputeExecutionContext::InitParams(UNiagaraScript* InGPUComputeScript, ENiagaraSimTarget InSimTarget, const FString& InDebugSimName, const int32 InMaxUpdateIterations, const TSet<uint32> InSpawnStages)
 {
+#if !UE_BUILD_SHIPPING
 	DebugSimName = InDebugSimName;
+#endif
 	GPUScript = InGPUComputeScript;
 	CombinedParamStore.InitFromOwningContext(InGPUComputeScript, InSimTarget, true);
 	MaxUpdateIterations = InMaxUpdateIterations;
