Remove duplicate if/else if conditions introduced in CL# 7986625
#rnx

diff --git a/Engine/Plugins/FX/Niagara/Source/Niagara/Private/NiagaraDataInterfaceRW.cpp b/Engine/Plugins/FX/Niagara/Source/Niagara/Private/NiagaraDataInterfaceRW.cpp
--- a/Engine/Plugins/FX/Niagara/Source/Niagara/Private/NiagaraDataInterfaceRW.cpp
+++ b/Engine/Plugins/FX/Niagara/Source/Niagara/Private/NiagaraDataInterfaceRW.cpp
@@ -146,7 +146,6 @@ void UNiagaraDataInterfaceGrid3D::GetVMExternalFunction(const FVMExternalFunctio
 	else if (BindingInfo.Name == UnitToIndexFunctionName) { OutFunc = FVMExternalFunction::CreateUObject(this, &UNiagaraDataInterfaceRWBase::EmptyVMFunction); }
 	else if (BindingInfo.Name == IndexToLinearFunctionName) { OutFunc = FVMExternalFunction::CreateUObject(this, &UNiagaraDataInterfaceRWBase::EmptyVMFunction); }
 	else if (BindingInfo.Name == VoxelSizeFunctionName) { OutFunc = FVMExternalFunction::CreateUObject(this, &UNiagaraDataInterfaceRWBase::EmptyVMFunction); }
-	else if (BindingInfo.Name == NumVoxelsFunctionName) { OutFunc = FVMExternalFunction::CreateUObject(this, &UNiagaraDataInterfaceRWBase::EmptyVMFunction); }
 }
 
 bool UNiagaraDataInterfaceGrid3D::Equals(const UNiagaraDataInterface* Other) const
@@ -472,7 +471,6 @@ void UNiagaraDataInterfaceGrid2D::GetVMExternalFunction(const FVMExternalFunctio
 	else if (BindingInfo.Name == IndexToLinearFunctionName) { OutFunc = FVMExternalFunction::CreateUObject(this, &UNiagaraDataInterfaceRWBase::EmptyVMFunction); }
 	else if (BindingInfo.Name == LinearToIndexFunctionName) { OutFunc = FVMExternalFunction::CreateUObject(this, &UNiagaraDataInterfaceRWBase::EmptyVMFunction); }
 	else if (BindingInfo.Name == CellSizeFunctionName) { OutFunc = FVMExternalFunction::CreateUObject(this, &UNiagaraDataInterfaceRWBase::EmptyVMFunction); }
-	else if (BindingInfo.Name == NumCellsFunctionName) { OutFunc = FVMExternalFunction::CreateUObject(this, &UNiagaraDataInterfaceRWBase::EmptyVMFunction); }
 }
 
 bool UNiagaraDataInterfaceGrid2D::Equals(const UNiagaraDataInterface* Other) const
