Fix use of DebugSimName when UE_BUILD_SHIPPING
#rnx

diff --git a/Engine/Plugins/FX/Niagara/Source/Niagara/Private/NiagaraEmitterInstanceBatcher.cpp b/Engine/Plugins/FX/Niagara/Source/Niagara/Private/NiagaraEmitterInstanceBatcher.cpp
--- a/Engine/Plugins/FX/Niagara/Source/Niagara/Private/NiagaraEmitterInstanceBatcher.cpp
+++ b/Engine/Plugins/FX/Niagara/Source/Niagara/Private/NiagaraEmitterInstanceBatcher.cpp
@@ -1085,10 +1085,12 @@ void NiagaraEmitterInstanceBatcher::Run(const FNiagaraGPUSystemTick& Tick, const
 	FNiagaraComputeExecutionContext* Context = Instance->Context;
 	if (TotalNumInstances == 0)
 	{
+#if !UE_BUILD_SHIPPING
 		SCOPED_DRAW_EVENTF(RHICmdList, NiagaraGPUSimulationCS, TEXT("Niagara Gpu Sim - %s - NumInstances: %u - StageNumber: %u"),
 			*Context->DebugSimName,
 			TotalNumInstances,
 			ShaderStageIndex);
+#endif
 		return;
 	}
 
@@ -1210,10 +1212,12 @@ void NiagaraEmitterInstanceBatcher::Run(const FNiagaraGPUSystemTick& Tick, const
 	// Dispatch, if anything needs to be done
 	if (TotalNumInstances)
 	{
+#if !UE_BUILD_SHIPPING
 		SCOPED_DRAW_EVENTF(RHICmdList, NiagaraGPUSimulationCS, TEXT("Niagara Gpu Sim - %s - NumInstances: %u - StageNumber: %u"),
 			*Context->DebugSimName,
 			TotalNumInstances,
 			ShaderStageIndex);
+#endif
 		DispatchComputeShader(RHICmdList, Shader, NumThreadGroups, 1, 1);
 	}
 
