Fixed hit proxies generating warning with FNiagaraGPUInstanceCountManager in the editor

diff --git a/Engine/Plugins/FX/Niagara/Source/Niagara/Classes/NiagaraScriptExecutionContext.h b/Engine/Plugins/FX/Niagara/Source/Niagara/Classes/NiagaraScriptExecutionContext.h
--- a/Engine/Plugins/FX/Niagara/Source/Niagara/Classes/NiagaraScriptExecutionContext.h
+++ b/Engine/Plugins/FX/Niagara/Source/Niagara/Classes/NiagaraScriptExecutionContext.h
@@ -188,6 +188,9 @@ public:
 #if !UE_BUILD_SHIPPING
 	//Persistent state 
 	FString DebugSimName;
+	FORCEINLINE const TCHAR* GetDebugSimName() const { return *DebugSimName; }
+#else
+	FORCEINLINE const TCHAR* GetDebugSimName() const { TEXT(""); }
 #endif
 
 	class FNiagaraDataSet *MainDataSet;
diff --git a/Engine/Plugins/FX/Niagara/Source/Niagara/Private/NiagaraEmitterInstanceBatcher.cpp b/Engine/Plugins/FX/Niagara/Source/Niagara/Private/NiagaraEmitterInstanceBatcher.cpp
--- a/Engine/Plugins/FX/Niagara/Source/Niagara/Private/NiagaraEmitterInstanceBatcher.cpp
+++ b/Engine/Plugins/FX/Niagara/Source/Niagara/Private/NiagaraEmitterInstanceBatcher.cpp
@@ -703,6 +703,10 @@ void NiagaraEmitterInstanceBatcher::PreInitViews(FRHICommandListImmediate& RHICm
 			ExecuteAll(RHICmdList, nullptr, !GPUInstanceCounterManager.HasPendingGPUReadback());
 		}
 	}
+	else
+	{
+		GPUInstanceCounterManager.ResizeBuffers(0);
+	}
 }
 
 bool NiagaraEmitterInstanceBatcher::UsesGlobalDistanceField() const
@@ -1087,7 +1091,7 @@ void NiagaraEmitterInstanceBatcher::Run(const FNiagaraGPUSystemTick& Tick, const
 	{
 #if !UE_BUILD_SHIPPING
 		SCOPED_DRAW_EVENTF(RHICmdList, NiagaraGPUSimulationCS, TEXT("Niagara Gpu Sim - %s - NumInstances: %u - StageNumber: %u"),
-			*Context->DebugSimName,
+			Context->GetDebugSimName(),
 			TotalNumInstances,
 			ShaderStageIndex);
 #endif
@@ -1214,7 +1218,7 @@ void NiagaraEmitterInstanceBatcher::Run(const FNiagaraGPUSystemTick& Tick, const
 	{
 #if !UE_BUILD_SHIPPING
 		SCOPED_DRAW_EVENTF(RHICmdList, NiagaraGPUSimulationCS, TEXT("Niagara Gpu Sim - %s - NumInstances: %u - StageNumber: %u"),
-			*Context->DebugSimName,
+			Context->GetDebugSimName(),
 			TotalNumInstances,
 			ShaderStageIndex);
 #endif
diff --git a/Engine/Source/Runtime/Renderer/Private/SceneHitProxyRendering.cpp b/Engine/Source/Runtime/Renderer/Private/SceneHitProxyRendering.cpp
--- a/Engine/Source/Runtime/Renderer/Private/SceneHitProxyRendering.cpp
+++ b/Engine/Source/Runtime/Renderer/Private/SceneHitProxyRendering.cpp
@@ -21,6 +21,7 @@
 #include "MeshPassProcessor.inl"
 #include "GPUScene.h"
 #include "Rendering/ColorVertexBuffer.h"
+#include "FXSystem.h"
 
 class FHitProxyShaderElementData : public FMeshMaterialShaderElementData
 {
@@ -602,6 +603,12 @@ void FDeferredShadingSceneRenderer::RenderHitProxies(FRHICommandListImmediate& R
 		DynamicVertexBufferForInitViews.Commit();
 		DynamicReadBufferForInitViews.Commit();
 
+		// Notify the FX system that the scene is about to be rendered.
+		if (Scene->FXSystem && Views.IsValidIndex(0))
+		{
+			Scene->FXSystem->PreRender(RHICmdList, &Views[0].GlobalDistanceFieldInfo.ParameterData, false);
+		}
+
 		::DoRenderHitProxies(RHICmdList, this, HitProxyRT, HitProxyDepthRT);
 		ClearPrimitiveSingleFrameIndirectLightingCacheBuffers();
 	}
diff --git a/Engine/Source/Runtime/Renderer/Private/SceneVisibility.cpp b/Engine/Source/Runtime/Renderer/Private/SceneVisibility.cpp
--- a/Engine/Source/Runtime/Renderer/Private/SceneVisibility.cpp
+++ b/Engine/Source/Runtime/Renderer/Private/SceneVisibility.cpp
@@ -2990,7 +2990,7 @@ void FSceneRenderer::PreVisibilityFrameSetup(FRHICommandListImmediate& RHICmdLis
 
 	if (Scene->FXSystem && Views.IsValidIndex(0))
 	{
-		Scene->FXSystem->PreInitViews(RHICmdList, Views[0].AllowGPUParticleUpdate());
+		Scene->FXSystem->PreInitViews(RHICmdList, Views[0].AllowGPUParticleUpdate() && !ViewFamily.EngineShowFlags.HitProxies);
 	}
 
 	// Draw lines to lights affecting this mesh if its selected.
