Fixed missing return statement

diff --git a/Engine/Plugins/FX/Niagara/Source/Niagara/Classes/NiagaraScriptExecutionContext.h b/Engine/Plugins/FX/Niagara/Source/Niagara/Classes/NiagaraScriptExecutionContext.h
--- a/Engine/Plugins/FX/Niagara/Source/Niagara/Classes/NiagaraScriptExecutionContext.h
+++ b/Engine/Plugins/FX/Niagara/Source/Niagara/Classes/NiagaraScriptExecutionContext.h
@@ -190,7 +190,7 @@ public:
 	FString DebugSimName;
 	FORCEINLINE const TCHAR* GetDebugSimName() const { return *DebugSimName; }
 #else
-	FORCEINLINE const TCHAR* GetDebugSimName() const { TEXT(""); }
+	FORCEINLINE const TCHAR* GetDebugSimName() const { return TEXT(""); }
 #endif
 
 	class FNiagaraDataSet *MainDataSet;
