- Tick Behavior controls for Niagara to allow components to ignore prereqs if they know that it is safe to do so
- This does not impact spawning, as in they will be queued until the end of the frame if the spawn did not occur before the regular tick

#rnx

diff --git a/Engine/Plugins/FX/Niagara/Source/Niagara/Private/NiagaraSystemInstance.cpp b/Engine/Plugins/FX/Niagara/Source/Niagara/Private/NiagaraSystemInstance.cpp
--- a/Engine/Plugins/FX/Niagara/Source/Niagara/Private/NiagaraSystemInstance.cpp
+++ b/Engine/Plugins/FX/Niagara/Source/Niagara/Private/NiagaraSystemInstance.cpp
@@ -68,6 +68,7 @@ FNiagaraSystemInstance::FNiagaraSystemInstance(UNiagaraComponent* InComponent)
 	: SystemInstanceIndex(INDEX_NONE)
 	, Component(InComponent)
 	, PrereqComponent(nullptr)
+	, TickBehavior(InComponent ? InComponent->GetTickBehavior() : ENiagaraTickBehavior::UsePrereqs)
 	, Age(0.0f)
 	, TickCount(0)
 	, InstanceParameters(Component)
@@ -1344,34 +1345,54 @@ ETickingGroup FNiagaraSystemInstance::CalculateTickGroup()
 	ETickingGroup NewTickGroup = (ETickingGroup)0;
 
 	// Debugging feature to force last tick group
-	if ( GNiagaraForceLastTickGroup )
+	if (GNiagaraForceLastTickGroup)
 	{
 		return NiagaraLastTickGroup;
 	}
 
-	// Handle attached component tick group
-	if (PrereqComponent != nullptr)
+	// Determine tick group
+	switch ( TickBehavior )
 	{
-		//-TODO: This doesn't deal with 'DontCompleteUntil' on the prereq's tick, if we have to handle that it could mean continual TG demotion
-		ETickingGroup PrereqTG = ETickingGroup(FMath::Max(PrereqComponent->PrimaryComponentTick.TickGroup, PrereqComponent->PrimaryComponentTick.EndTickGroup) + 1);
-		NewTickGroup = FMath::Max(NewTickGroup, PrereqTG);
-	}
-
-	// Handle data interfaces that have tick dependencies
-	if ( bDataInterfacesHaveTickPrereqs )
-	{
-		for (TPair<TWeakObjectPtr<UNiagaraDataInterface>, int32>& Pair : DataInterfaceInstanceDataOffsets)
-		{
-			if (UNiagaraDataInterface* Interface = Pair.Key.Get())
+		default:
+		case ENiagaraTickBehavior::UsePrereqs:
+			// Handle attached component tick group
+			if (PrereqComponent != nullptr)
 			{
-				ETickingGroup PrereqTG = Interface->CalculateTickGroup(&DataInterfaceInstanceData[Pair.Value]);
+				//-TODO: This doesn't deal with 'DontCompleteUntil' on the prereq's tick, if we have to handle that it could mean continual TG demotion
+				ETickingGroup PrereqTG = ETickingGroup(FMath::Max(PrereqComponent->PrimaryComponentTick.TickGroup, PrereqComponent->PrimaryComponentTick.EndTickGroup) + 1);
 				NewTickGroup = FMath::Max(NewTickGroup, PrereqTG);
 			}
-		}
+
+			// Handle data interfaces that have tick dependencies
+			if ( bDataInterfacesHaveTickPrereqs )
+			{
+				for (TPair<TWeakObjectPtr<UNiagaraDataInterface>, int32>& Pair : DataInterfaceInstanceDataOffsets)
+				{
+					if (UNiagaraDataInterface* Interface = Pair.Key.Get())
+					{
+						ETickingGroup PrereqTG = Interface->CalculateTickGroup(&DataInterfaceInstanceData[Pair.Value]);
+						NewTickGroup = FMath::Max(NewTickGroup, PrereqTG);
+					}
+				}
+			}
+
+			// Clamp tick group to our range
+			NewTickGroup = FMath::Clamp(NewTickGroup, NiagaraFirstTickGroup, NiagaraLastTickGroup);
+			break;
+
+		case ENiagaraTickBehavior::UseComponentTickGroup:
+			NewTickGroup = FMath::Clamp((ETickingGroup)Component->PrimaryComponentTick.TickGroup, NiagaraFirstTickGroup, NiagaraLastTickGroup);
+			break;
+
+		case ENiagaraTickBehavior::ForceTickFirst:
+			NewTickGroup = NiagaraFirstTickGroup;
+			break;
+
+		case ENiagaraTickBehavior::ForceTickLast:
+			NewTickGroup = NiagaraLastTickGroup;
+			break;
 	}
 
-	// Clamp tick group to our range
-	NewTickGroup = FMath::Clamp(NewTickGroup, NiagaraFirstTickGroup, NiagaraLastTickGroup);
 	return NewTickGroup;
 }
 
diff --git a/Engine/Plugins/FX/Niagara/Source/Niagara/Public/NiagaraCommon.h b/Engine/Plugins/FX/Niagara/Source/Niagara/Public/NiagaraCommon.h
--- a/Engine/Plugins/FX/Niagara/Source/Niagara/Public/NiagaraCommon.h
+++ b/Engine/Plugins/FX/Niagara/Source/Niagara/Public/NiagaraCommon.h
@@ -36,6 +36,20 @@ constexpr ETickingGroup NiagaraFirstTickGroup = TG_PrePhysics;
 constexpr ETickingGroup NiagaraLastTickGroup = TG_LastDemotable;
 constexpr int NiagaraNumTickGroups = NiagaraLastTickGroup - NiagaraFirstTickGroup + 1;
 
+/** Niagara ticking behaviour */
+UENUM()
+enum class ENiagaraTickBehavior : uint8
+{
+	/** Niagara will tick after all prereqs have ticked for attachements / data interfaces, this is the safest option. */
+	UsePrereqs,
+	/** Niagara will ignore prereqs (attachments / data interface dependencies) and use the tick group set on the component. */
+	UseComponentTickGroup,
+	/** Niagara will tick in the first tick group (default is TG_PrePhysics). */
+	ForceTickFirst,
+	/** Niagara will tick in the last tick group (default is TG_LastDemotable). */
+	ForceTickLast,
+};
+
 enum ENiagaraBaseTypes
 {
 	NBT_Float,
diff --git a/Engine/Plugins/FX/Niagara/Source/Niagara/Public/NiagaraComponent.h b/Engine/Plugins/FX/Niagara/Source/Niagara/Public/NiagaraComponent.h
--- a/Engine/Plugins/FX/Niagara/Source/Niagara/Public/NiagaraComponent.h
+++ b/Engine/Plugins/FX/Niagara/Source/Niagara/Public/NiagaraComponent.h
@@ -58,6 +58,10 @@ private:
 	UPROPERTY(EditAnywhere, Category="Niagara", meta = (DisplayName = "Niagara System Asset"))
 	UNiagaraSystem* Asset;
 
+	/** Allows you to control how Niagara selects the tick group, changing this while an instance is active will result in not change as it is cached. */
+	UPROPERTY(EditAnywhere, Category = "Niagara", meta = (DisplayName = "Niagara Tick Behavior"))
+	ENiagaraTickBehavior TickBehavior = ENiagaraTickBehavior::UsePrereqs;
+
 	/** Initial values for parameter overrides. 
 	TODO: This should be a minimal set of explicitly override parameters similar to how parameter collection instances override their collections store. 
 	Should expose anything in the "User" namespace.
@@ -228,6 +232,8 @@ public:
 
 	FNiagaraSystemInstance* GetSystemInstance() const;
 
+	ENiagaraTickBehavior GetTickBehavior() const { return TickBehavior; }
+
 	/** Returns true if this component forces it's instances to run in "Solo" mode. A sub optimal path required in some situations. */
 	bool ForcesSolo()const;
 
diff --git a/Engine/Plugins/FX/Niagara/Source/Niagara/Public/NiagaraSystemInstance.h b/Engine/Plugins/FX/Niagara/Source/Niagara/Public/NiagaraSystemInstance.h
--- a/Engine/Plugins/FX/Niagara/Source/Niagara/Public/NiagaraSystemInstance.h
+++ b/Engine/Plugins/FX/Niagara/Source/Niagara/Public/NiagaraSystemInstance.h
@@ -282,6 +282,8 @@ private:
 
 	UActorComponent* PrereqComponent;
 
+	ENiagaraTickBehavior TickBehavior;
+
 	/** The age of the System instance. */
 	float Age;
 
