Fix check for NiagaraRibbonVertexShader to compile particle time in shader if requested for ES3.1.

diff --git a/Engine/Plugins/FX/Niagara/Shaders/Private/NiagaraRibbonVertexFactory.ush b/Engine/Plugins/FX/Niagara/Shaders/Private/NiagaraRibbonVertexFactory.ush
--- a/Engine/Plugins/FX/Niagara/Shaders/Private/NiagaraRibbonVertexFactory.ush
+++ b/Engine/Plugins/FX/Niagara/Shaders/Private/NiagaraRibbonVertexFactory.ush
@@ -9,7 +9,7 @@
 
 #define USE_PARTICLE_POSITION (FEATURE_LEVEL >= FEATURE_LEVEL_SM5 && NEEDS_PARTICLE_POSITION)
 #define USE_PARTICLE_VELOCITY (FEATURE_LEVEL >= FEATURE_LEVEL_SM5 && NEEDS_PARTICLE_VELOCITY)
-#define USE_PARTICLE_TIME (FEATURE_LEVEL >= FEATURE_LEVEL_SM5 && NEEDS_PARTICLE_TIME)
+#define USE_PARTICLE_TIME ((FEATURE_LEVEL >= FEATURE_LEVEL_SM5 || FEATURE_LEVEL == FEATURE_LEVEL_ES3_1) && NEEDS_PARTICLE_TIME)
 #define USE_PARTICLE_SIZE (FEATURE_LEVEL >= FEATURE_LEVEL_SM5 && NEEDS_PARTICLE_SIZE)
 #define USE_PARTICLE_RANDOM (FEATURE_LEVEL >= FEATURE_LEVEL_SM5 && NEEDS_PARTICLE_RANDOM)
 #define USE_PARTICLE_INTERPOLATION (FEATURE_LEVEL >= FEATURE_LEVEL_SM5)
