Enable per particle material params in niagara vertex shaders for ES3.1 targets.

diff --git a/Engine/Plugins/FX/Niagara/Shaders/Private/NiagaraMeshVertexFactory.ush b/Engine/Plugins/FX/Niagara/Shaders/Private/NiagaraMeshVertexFactory.ush
--- a/Engine/Plugins/FX/Niagara/Shaders/Private/NiagaraMeshVertexFactory.ush
+++ b/Engine/Plugins/FX/Niagara/Shaders/Private/NiagaraMeshVertexFactory.ush
@@ -7,12 +7,11 @@
 #include "/Engine/Private/VertexFactoryCommon.ush"
 #include "NiagaraVFParticleAccess.usf"
 
-// These are only enabled for SM5 due to limited interpolators elsewhere
-#define USE_PARTICLE_POSITION (FEATURE_LEVEL >= FEATURE_LEVEL_SM5 && NEEDS_PARTICLE_POSITION)
-#define USE_PARTICLE_VELOCITY (FEATURE_LEVEL >= FEATURE_LEVEL_SM5 && NEEDS_PARTICLE_VELOCITY)
-#define USE_PARTICLE_TIME (FEATURE_LEVEL >= FEATURE_LEVEL_SM5 && NEEDS_PARTICLE_TIME)
-#define USE_PARTICLE_RANDOM (FEATURE_LEVEL >= FEATURE_LEVEL_SM5 && NEEDS_PARTICLE_RANDOM)
-#define USE_PARTICLE_TRANSFORM (FEATURE_LEVEL >= FEATURE_LEVEL_SM5 && NEEDS_PARTICLE_TRANSFORM)
+#define USE_PARTICLE_POSITION ((FEATURE_LEVEL >= FEATURE_LEVEL_SM5 || FEATURE_LEVEL == FEATURE_LEVEL_ES3_1) && NEEDS_PARTICLE_POSITION)
+#define USE_PARTICLE_VELOCITY ((FEATURE_LEVEL >= FEATURE_LEVEL_SM5 || FEATURE_LEVEL == FEATURE_LEVEL_ES3_1) && NEEDS_PARTICLE_VELOCITY)
+#define USE_PARTICLE_TIME ((FEATURE_LEVEL >= FEATURE_LEVEL_SM5 || FEATURE_LEVEL == FEATURE_LEVEL_ES3_1) && NEEDS_PARTICLE_TIME)
+#define USE_PARTICLE_RANDOM ((FEATURE_LEVEL >= FEATURE_LEVEL_SM5 || FEATURE_LEVEL == FEATURE_LEVEL_ES3_1) && NEEDS_PARTICLE_RANDOM)
+#define USE_PARTICLE_TRANSFORM ((FEATURE_LEVEL >= FEATURE_LEVEL_SM5 || FEATURE_LEVEL == FEATURE_LEVEL_ES3_1) && NEEDS_PARTICLE_TRANSFORM)
 
 float3 TransformPosition(float3 InPosition)
 {
diff --git a/Engine/Plugins/FX/Niagara/Shaders/Private/NiagaraRibbonVertexFactory.ush b/Engine/Plugins/FX/Niagara/Shaders/Private/NiagaraRibbonVertexFactory.ush
--- a/Engine/Plugins/FX/Niagara/Shaders/Private/NiagaraRibbonVertexFactory.ush
+++ b/Engine/Plugins/FX/Niagara/Shaders/Private/NiagaraRibbonVertexFactory.ush
@@ -7,12 +7,12 @@
 #include "/Engine/Private/VertexFactoryCommon.ush"
 #include "NiagaraVFParticleAccess.usf"
 
-#define USE_PARTICLE_POSITION (FEATURE_LEVEL >= FEATURE_LEVEL_SM5 && NEEDS_PARTICLE_POSITION)
-#define USE_PARTICLE_VELOCITY (FEATURE_LEVEL >= FEATURE_LEVEL_SM5 && NEEDS_PARTICLE_VELOCITY)
+#define USE_PARTICLE_POSITION ((FEATURE_LEVEL >= FEATURE_LEVEL_SM5 || FEATURE_LEVEL == FEATURE_LEVEL_ES3_1) && NEEDS_PARTICLE_POSITION)
+#define USE_PARTICLE_VELOCITY ((FEATURE_LEVEL >= FEATURE_LEVEL_SM5 || FEATURE_LEVEL == FEATURE_LEVEL_ES3_1) && NEEDS_PARTICLE_VELOCITY)
 #define USE_PARTICLE_TIME ((FEATURE_LEVEL >= FEATURE_LEVEL_SM5 || FEATURE_LEVEL == FEATURE_LEVEL_ES3_1) && NEEDS_PARTICLE_TIME)
-#define USE_PARTICLE_SIZE (FEATURE_LEVEL >= FEATURE_LEVEL_SM5 && NEEDS_PARTICLE_SIZE)
-#define USE_PARTICLE_RANDOM (FEATURE_LEVEL >= FEATURE_LEVEL_SM5 && NEEDS_PARTICLE_RANDOM)
-#define USE_PARTICLE_INTERPOLATION (FEATURE_LEVEL >= FEATURE_LEVEL_SM5)
+#define USE_PARTICLE_SIZE ((FEATURE_LEVEL >= FEATURE_LEVEL_SM5 || FEATURE_LEVEL == FEATURE_LEVEL_ES3_1) && NEEDS_PARTICLE_SIZE)
+#define USE_PARTICLE_RANDOM ((FEATURE_LEVEL >= FEATURE_LEVEL_SM5 || FEATURE_LEVEL == FEATURE_LEVEL_ES3_1) && NEEDS_PARTICLE_RANDOM)
+#define USE_PARTICLE_INTERPOLATION (FEATURE_LEVEL >= FEATURE_LEVEL_SM5 || FEATURE_LEVEL == FEATURE_LEVEL_ES3_1)
 
 float3 TransformPosition(float3 InPosition)
 {
diff --git a/Engine/Plugins/FX/Niagara/Shaders/Private/NiagaraSpriteVertexFactory.ush b/Engine/Plugins/FX/Niagara/Shaders/Private/NiagaraSpriteVertexFactory.ush
--- a/Engine/Plugins/FX/Niagara/Shaders/Private/NiagaraSpriteVertexFactory.ush
+++ b/Engine/Plugins/FX/Niagara/Shaders/Private/NiagaraSpriteVertexFactory.ush
@@ -9,13 +9,12 @@
 #include "/Engine/Private/ParticleVertexFactoryCommon.ush"
 #include "NiagaraVFParticleAccess.usf"
 
-// These are only enabled for SM5 due to limited interpolators elsewhere
-#define USE_PARTICLE_LIGHTING_OFFSET (FEATURE_LEVEL >= FEATURE_LEVEL_SM5 && !MATERIAL_SHADINGMODEL_UNLIT)
-#define USE_PARTICLE_POSITION (FEATURE_LEVEL >= FEATURE_LEVEL_SM5 && NEEDS_PARTICLE_POSITION)
-#define USE_PARTICLE_VELOCITY (FEATURE_LEVEL >= FEATURE_LEVEL_SM5 && NEEDS_PARTICLE_VELOCITY)
-#define USE_PARTICLE_TIME (FEATURE_LEVEL >= FEATURE_LEVEL_SM5 && NEEDS_PARTICLE_TIME)
-#define USE_PARTICLE_SIZE (FEATURE_LEVEL >= FEATURE_LEVEL_SM5 && NEEDS_PARTICLE_SIZE)
-#define USE_PARTICLE_RANDOM (FEATURE_LEVEL >= FEATURE_LEVEL_SM5 && NEEDS_PARTICLE_RANDOM)
+#define USE_PARTICLE_LIGHTING_OFFSET ((FEATURE_LEVEL >= FEATURE_LEVEL_SM5 || FEATURE_LEVEL == FEATURE_LEVEL_ES3_1) && !MATERIAL_SHADINGMODEL_UNLIT)
+#define USE_PARTICLE_POSITION ((FEATURE_LEVEL >= FEATURE_LEVEL_SM5 || FEATURE_LEVEL == FEATURE_LEVEL_ES3_1) && NEEDS_PARTICLE_POSITION)
+#define USE_PARTICLE_VELOCITY ((FEATURE_LEVEL >= FEATURE_LEVEL_SM5 || FEATURE_LEVEL == FEATURE_LEVEL_ES3_1) && NEEDS_PARTICLE_VELOCITY)
+#define USE_PARTICLE_TIME ((FEATURE_LEVEL >= FEATURE_LEVEL_SM5 || FEATURE_LEVEL == FEATURE_LEVEL_ES3_1) && NEEDS_PARTICLE_TIME)
+#define USE_PARTICLE_SIZE ((FEATURE_LEVEL >= FEATURE_LEVEL_SM5 || FEATURE_LEVEL == FEATURE_LEVEL_ES3_1) && NEEDS_PARTICLE_SIZE)
+#define USE_PARTICLE_RANDOM ((FEATURE_LEVEL >= FEATURE_LEVEL_SM5 || FEATURE_LEVEL == FEATURE_LEVEL_ES3_1) && NEEDS_PARTICLE_RANDOM)
 
 // Must keep these in sync with ENiagaraSpriteAlignment in NiagaraSpriteRendererProperties.h
 /** Only Particles.SpriteRotation and FacingMode impact the alignment of the particle.*/
