Fixed localization conflicts

#rnx

diff --git a/Engine/Plugins/Editor/CurveEditorTools/Source/CurveEditorTools/Private/Tools/CurveEditorMultiScaleTool.cpp b/Engine/Plugins/Editor/CurveEditorTools/Source/CurveEditorTools/Private/Tools/CurveEditorMultiScaleTool.cpp
--- a/Engine/Plugins/Editor/CurveEditorTools/Source/CurveEditorTools/Private/Tools/CurveEditorMultiScaleTool.cpp
+++ b/Engine/Plugins/Editor/CurveEditorTools/Source/CurveEditorTools/Private/Tools/CurveEditorMultiScaleTool.cpp
@@ -680,4 +680,6 @@ FVector2D FCurveEditorMultiScaleTool::GetPivot(FCurveModel* InCurve, const TArra
 		return Pivot;
 	}
 	return Pivot;
-}
\ No newline at end of file
+}
+
+#undef LOCTEXT_NAMESPACE
diff --git a/Engine/Plugins/Enterprise/StaticMeshEditorExtension/Source/StaticMeshEditorExtension/Private/PolygonEditingToolbar.cpp b/Engine/Plugins/Enterprise/StaticMeshEditorExtension/Source/StaticMeshEditorExtension/Private/PolygonEditingToolbar.cpp
--- a/Engine/Plugins/Enterprise/StaticMeshEditorExtension/Source/StaticMeshEditorExtension/Private/PolygonEditingToolbar.cpp
+++ b/Engine/Plugins/Enterprise/StaticMeshEditorExtension/Source/StaticMeshEditorExtension/Private/PolygonEditingToolbar.cpp
@@ -598,11 +598,11 @@ void FPolygonEditingToolbar::OnLODModelChanged()
 	{
 		if (StaticMeshEditor->GetCurrentLODLevel() == 0 && !EditableLODs[0])
 		{
-			FMessageDialog::Open(EAppMsgType::Ok, LOCTEXT("FPolygonEditingToolbarExitEdit", "LOD Auto is not editable on a static mesh with more than one LOD.\nExiting Edit Mode."));
+			FMessageDialog::Open(EAppMsgType::Ok, LOCTEXT("FPolygonEditingToolbarExitEdit_LODAutoNotEditable", "LOD Auto is not editable on a static mesh with more than one LOD.\nExiting Edit Mode."));
 		}
 		else
 		{
-			FMessageDialog::Open( EAppMsgType::Ok, LOCTEXT("FPolygonEditingToolbarExitEdit", "Non editable LOD has been selected.\nExiting Edit Mode.") );
+			FMessageDialog::Open( EAppMsgType::Ok, LOCTEXT("FPolygonEditingToolbarExitEdit_NonEditableLOD", "Non editable LOD has been selected.\nExiting Edit Mode.") );
 		}
 		OnToggleEditMode();
 		return;
@@ -629,11 +629,11 @@ void FPolygonEditingToolbar::OnMeshChanged()
 		{
 			if (StaticMeshEditor->GetCurrentLODLevel() == 0 && !EditableLODs[0])
 			{
-				FMessageDialog::Open(EAppMsgType::Ok, LOCTEXT("FPolygonEditingToolbarExitEdit", "LOD Auto is not editable on a static mesh with more than one LOD.\nExiting Edit Mode."));
+				FMessageDialog::Open(EAppMsgType::Ok, LOCTEXT("FPolygonEditingToolbarExitEdit_LODAutoNotEditable", "LOD Auto is not editable on a static mesh with more than one LOD.\nExiting Edit Mode."));
 			}
 			else
 			{
-				FMessageDialog::Open(EAppMsgType::Ok, LOCTEXT("FPolygonEditingToolbarExitEdit", "Selected LOD is not editable anymore.\nExiting Edit Mode."));
+				FMessageDialog::Open(EAppMsgType::Ok, LOCTEXT("FPolygonEditingToolbarExitEdit_LODNoLongerEditable", "Selected LOD is not editable anymore.\nExiting Edit Mode."));
 			}
 			OnToggleEditMode();
 		}
diff --git a/Engine/Plugins/Experimental/ChaosClothEditor/Source/Chaos/Private/ChaosClothEditor/ChaosSimulationEditorExtender.cpp b/Engine/Plugins/Experimental/ChaosClothEditor/Source/Chaos/Private/ChaosClothEditor/ChaosSimulationEditorExtender.cpp
--- a/Engine/Plugins/Experimental/ChaosClothEditor/Source/Chaos/Private/ChaosClothEditor/ChaosSimulationEditorExtender.cpp
+++ b/Engine/Plugins/Experimental/ChaosClothEditor/Source/Chaos/Private/ChaosClothEditor/ChaosSimulationEditorExtender.cpp
@@ -165,4 +165,6 @@ void FSimulationEditorExtender::ShowClothSections(USkeletalMeshComponent* MeshCo
 	}
 }
 
+#undef LOCTEXT_NAMESPACE
+
 #endif  // #if WITH_CHAOS
\ No newline at end of file
diff --git a/Engine/Plugins/Experimental/ControlRig/Source/ControlRigEditor/Private/Graph/SControlRigGraphPinNameList.cpp b/Engine/Plugins/Experimental/ControlRig/Source/ControlRigEditor/Private/Graph/SControlRigGraphPinNameList.cpp
--- a/Engine/Plugins/Experimental/ControlRig/Source/ControlRigEditor/Private/Graph/SControlRigGraphPinNameList.cpp
+++ b/Engine/Plugins/Experimental/ControlRig/Source/ControlRigEditor/Private/Graph/SControlRigGraphPinNameList.cpp
@@ -64,7 +64,7 @@ void SControlRigGraphPinNameList::SetNameListText(const FText& NewTypeInValue, E
 {
 	if(!GraphPinObj->GetDefaultAsString().Equals(NewTypeInValue.ToString()))
 	{
-		const FScopedTransaction Transaction( NSLOCTEXT("GraphEditor", "ChangeNameListPinValue", "Change Bone Name Pin Value" ) );
+		const FScopedTransaction Transaction( NSLOCTEXT("GraphEditor", "ChangeBoneNameListPinValue", "Change Bone Name Pin Value" ) );
 		GraphPinObj->Modify();
 		GraphPinObj->GetSchema()->TrySetDefaultValue(*GraphPinObj, NewTypeInValue.ToString());
 	}
diff --git a/Engine/Plugins/Experimental/ControlRig/Source/ControlRigEditor/Private/Sequencer/ControlRigParameterTrackEditor.cpp b/Engine/Plugins/Experimental/ControlRig/Source/ControlRigEditor/Private/Sequencer/ControlRigParameterTrackEditor.cpp
--- a/Engine/Plugins/Experimental/ControlRig/Source/ControlRigEditor/Private/Sequencer/ControlRigParameterTrackEditor.cpp
+++ b/Engine/Plugins/Experimental/ControlRig/Source/ControlRigEditor/Private/Sequencer/ControlRigParameterTrackEditor.cpp
@@ -1119,7 +1119,7 @@ void FControlRigParameterSection::BuildSectionContextMenu(FMenuBuilder& MenuBuil
 			return FUIAction(
 				FExecuteAction::CreateLambda([=]
 			{
-				FScopedTransaction Transaction(LOCTEXT("RigSectionFilter", "Toggle Rig Control Filters"));
+				FScopedTransaction Transaction(LOCTEXT("ToggleRigControlFiltersTransaction", "Toggle Rig Control Filters"));
 				ParameterSection->Modify();
 				if (Index >= 0)
 				{
@@ -1172,7 +1172,7 @@ void FControlRigParameterSection::BuildSectionContextMenu(FMenuBuilder& MenuBuil
 			})
 				);
 		};
-		MenuBuilder.BeginSection(NAME_None, LOCTEXT("RigSectionFilter", "Filter Controls"));
+		MenuBuilder.BeginSection(NAME_None, LOCTEXT("RigSectionFilterControls", "Filter Controls"));
 		{
 			MenuBuilder.AddSubMenu(
 				LOCTEXT("ToggleRigControlsText", "Toggle Rig Controls"), LOCTEXT("ToggleRigControlsText_Tooltip", "Toggle Rig Controls"),
diff --git a/Engine/Plugins/Experimental/DataPrepEditor/Source/DataPrepCore/Private/Blueprint/K2Node_DataprepProducer.cpp b/Engine/Plugins/Experimental/DataPrepEditor/Source/DataPrepCore/Private/Blueprint/K2Node_DataprepProducer.cpp
--- a/Engine/Plugins/Experimental/DataPrepEditor/Source/DataPrepCore/Private/Blueprint/K2Node_DataprepProducer.cpp
+++ b/Engine/Plugins/Experimental/DataPrepEditor/Source/DataPrepCore/Private/Blueprint/K2Node_DataprepProducer.cpp
@@ -54,5 +54,3 @@ TSharedPtr<SGraphNode> UK2Node_DataprepProducer::CreateVisualWidget()
 	return SNew(SGraphNodeDataprepProducer, this);
 }
 #endif
-
-#undef LOCTEXT_NAMESPACE
diff --git a/Engine/Plugins/Experimental/DataPrepEditor/Source/DataPrepEditor/Private/DataPrepEditor.cpp b/Engine/Plugins/Experimental/DataPrepEditor/Source/DataPrepEditor/Private/DataPrepEditor.cpp
--- a/Engine/Plugins/Experimental/DataPrepEditor/Source/DataPrepEditor/Private/DataPrepEditor.cpp
+++ b/Engine/Plugins/Experimental/DataPrepEditor/Source/DataPrepEditor/Private/DataPrepEditor.cpp
@@ -410,7 +410,7 @@ void FDataprepEditor::RegisterTabSpawners(const TSharedRef<class FTabManager>& I
 		.SetIcon(FSlateIcon(FDataprepEditorStyle::GetStyleSetName(), "DataprepEditor.Tabs.AssetPreview"));
 
 	InTabManager->RegisterTabSpawner(SceneViewportTabId, FOnSpawnTab::CreateSP(this, &FDataprepEditor::SpawnTabSceneViewport))
-		.SetDisplayName(LOCTEXT("AssetPreviewTab", "Scene Viewport"))
+		.SetDisplayName(LOCTEXT("SceneViewportTab", "Scene Viewport"))
 		.SetGroup(WorkspaceMenuCategoryRef)
 		.SetIcon(FSlateIcon(FDataprepEditorStyle::GetStyleSetName(), "DataprepEditor.Tabs.SceneViewport"));
 
diff --git a/Engine/Plugins/Experimental/DataPrepEditor/Source/DataPrepEditor/Private/Widgets/SDataprepEditorViewport.cpp b/Engine/Plugins/Experimental/DataPrepEditor/Source/DataPrepEditor/Private/Widgets/SDataprepEditorViewport.cpp
--- a/Engine/Plugins/Experimental/DataPrepEditor/Source/DataPrepEditor/Private/Widgets/SDataprepEditorViewport.cpp
+++ b/Engine/Plugins/Experimental/DataPrepEditor/Source/DataPrepEditor/Private/Widgets/SDataprepEditorViewport.cpp
@@ -567,7 +567,7 @@ void SDataprepEditorViewport::PopulateViewportOverlays(TSharedRef<SOverlay> Over
 	SEditorViewport::PopulateViewportOverlays(Overlay);
 
 	ScreenSizeText = SNew(STextBlock)
-	.Text( LOCTEXT( "ScreenSize_F", "Current Screen Size:"))
+	.Text( LOCTEXT( "ScreenSize", "Current Screen Size:"))
 	.TextStyle(FEditorStyle::Get(), "TextBlock.ShadowedText");
 
 	Overlay->AddSlot()
@@ -1384,7 +1384,7 @@ void SDataprepEditorViewportToolbar::ExtendLeftAlignedToolbarSlots(TSharedPtr<SH
 	[
 		SNew(SEditorViewportToolbarMenu)
 		.Label( LOCTEXT("DataprepEditor_Experimental", "Experimental") )
-		.ToolTipText(LOCTEXT("DataprepEditor_RenderingTooltip", "Experimental viewing modes or actions."))
+		.ToolTipText(LOCTEXT("DataprepEditor_ExperimentalTooltip", "Experimental viewing modes or actions."))
 		.ParentToolBar( ParentToolBarPtr )
 		.Cursor(EMouseCursor::Default)
 		.AddMetaData<FTagMetaData>(FTagMetaData(TEXT("ExperimentalMenuButton")))
diff --git a/Engine/Plugins/Experimental/HairStrands/Source/HairStrandsEditor/Private/HairStrandsFactory.cpp b/Engine/Plugins/Experimental/HairStrands/Source/HairStrandsEditor/Private/HairStrandsFactory.cpp
--- a/Engine/Plugins/Experimental/HairStrands/Source/HairStrandsEditor/Private/HairStrandsFactory.cpp
+++ b/Engine/Plugins/Experimental/HairStrands/Source/HairStrandsEditor/Private/HairStrandsFactory.cpp
@@ -141,4 +141,6 @@ TSharedPtr<IHairStrandsTranslator> UHairStrandsFactory::GetTranslator(const FStr
 		}
 	}
 	return {};
-}
\ No newline at end of file
+}
+
+#undef LOCTEXT_NAMESPACE
diff --git a/Engine/Plugins/Experimental/HairStrands/Source/HairStrandsEditor/Private/ReimportHairStrandsFactory.cpp b/Engine/Plugins/Experimental/HairStrands/Source/HairStrandsEditor/Private/ReimportHairStrandsFactory.cpp
--- a/Engine/Plugins/Experimental/HairStrands/Source/HairStrandsEditor/Private/ReimportHairStrandsFactory.cpp
+++ b/Engine/Plugins/Experimental/HairStrands/Source/HairStrandsEditor/Private/ReimportHairStrandsFactory.cpp
@@ -142,3 +142,5 @@ EReimportResult::Type UReimportHairStrandsFactory::Reimport(UObject* Obj)
 
 	return EReimportResult::Succeeded;
 }
+
+#undef LOCTEXT_NAMESPACE
diff --git a/Engine/Plugins/FX/Niagara/Source/NiagaraEditor/Private/Customizations/NiagaraTypeCustomizations.cpp b/Engine/Plugins/FX/Niagara/Source/NiagaraEditor/Private/Customizations/NiagaraTypeCustomizations.cpp
--- a/Engine/Plugins/FX/Niagara/Source/NiagaraEditor/Private/Customizations/NiagaraTypeCustomizations.cpp
+++ b/Engine/Plugins/FX/Niagara/Source/NiagaraEditor/Private/Customizations/NiagaraTypeCustomizations.cpp
@@ -125,7 +125,7 @@ FText FNiagaraVariableAttributeBindingCustomization::GetTooltipText() const
 			DefaultValueStr.TrimEndInline();
 		}
 
-		FText TooltipDesc = FText::Format(LOCTEXT("BindingTooltip", "Use the variable \"{0}\" if it exists, otherwise use the default \"{1}\" "), FText::FromName(TargetVariableBinding->BoundVariable.GetName()),
+		FText TooltipDesc = FText::Format(LOCTEXT("AttributeBindingTooltip", "Use the variable \"{0}\" if it exists, otherwise use the default \"{1}\" "), FText::FromName(TargetVariableBinding->BoundVariable.GetName()),
 			FText::FromString(DefaultValueStr));
 		return TooltipDesc;
 	}
@@ -224,7 +224,7 @@ void FNiagaraVariableAttributeBindingCustomization::OnActionSelected(const TArra
 
 void FNiagaraVariableAttributeBindingCustomization::ChangeSource(FName InVarName)
 {
-	FScopedTransaction Transaction(FText::Format(LOCTEXT("ChangeSource", " Change Variable Source to \"{0}\" "), FText::FromName(InVarName)));
+	FScopedTransaction Transaction(FText::Format(LOCTEXT("ChangeVariableSource", " Change Variable Source to \"{0}\" "), FText::FromName(InVarName)));
 	TArray<UObject*> Objects;
 	PropertyHandle->GetOuterObjects(Objects);
 	for (UObject* Obj : Objects)
@@ -314,7 +314,7 @@ FText FNiagaraUserParameterBindingCustomization::GetTooltipText() const
 {
 	if (BaseSystem && TargetUserParameterBinding && TargetUserParameterBinding->Parameter.IsValid())
 	{
-		FText TooltipDesc = FText::Format(LOCTEXT("BindingTooltip", "Bound to the user parameter \"{0}\""), FText::FromName(TargetUserParameterBinding->Parameter.GetName()));
+		FText TooltipDesc = FText::Format(LOCTEXT("ParameterBindingTooltip", "Bound to the user parameter \"{0}\""), FText::FromName(TargetUserParameterBinding->Parameter.GetName()));
 		return TooltipDesc;
 	}
 	return FText::FromString(TEXT("Missing"));
@@ -366,7 +366,7 @@ void FNiagaraUserParameterBindingCustomization::CollectAllActions(FGraphActionLi
 		FText CategoryName = FText();
 		FString DisplayNameString = FName::NameToDisplayString(UserParamName.ToString(), false);
 		const FText NameText = FText::FromString(DisplayNameString);
-		const FText TooltipDesc = FText::Format(LOCTEXT("SetFunctionPopupTooltip", "Bind to the User Parameter \"{0}\" "), FText::FromString(DisplayNameString));
+		const FText TooltipDesc = FText::Format(LOCTEXT("BindToUserParameter", "Bind to the User Parameter \"{0}\" "), FText::FromString(DisplayNameString));
 		TSharedPtr<FNiagaraStackAssetAction_VarBind> NewNodeAction(new FNiagaraStackAssetAction_VarBind(UserParamName, CategoryName, NameText,
 			TooltipDesc, 0, FText()));
 		OutAllActions.AddAction(NewNodeAction);
@@ -406,7 +406,7 @@ void FNiagaraUserParameterBindingCustomization::OnActionSelected(const TArray< T
 
 void FNiagaraUserParameterBindingCustomization::ChangeSource(FName InVarName)
 {
-	FScopedTransaction Transaction(FText::Format(LOCTEXT("ChangeSource", " Change User Parameter Source to \"{0}\" "), FText::FromName(InVarName)));
+	FScopedTransaction Transaction(FText::Format(LOCTEXT("ChangeParameterSource", " Change User Parameter Source to \"{0}\" "), FText::FromName(InVarName)));
 	TArray<UObject*> Objects;
 	PropertyHandle->GetOuterObjects(Objects);
 	for (UObject* Obj : Objects)
diff --git a/Engine/Plugins/FX/Niagara/Source/NiagaraEditor/Private/EdGraphSchema_Niagara.cpp b/Engine/Plugins/FX/Niagara/Source/NiagaraEditor/Private/EdGraphSchema_Niagara.cpp
--- a/Engine/Plugins/FX/Niagara/Source/NiagaraEditor/Private/EdGraphSchema_Niagara.cpp
+++ b/Engine/Plugins/FX/Niagara/Source/NiagaraEditor/Private/EdGraphSchema_Niagara.cpp
@@ -1753,9 +1753,9 @@ void UEdGraphSchema_Niagara::GetContextMenuActions(UToolMenu* Menu, UGraphNodeCo
 			const FName SectionName = "EdGraphSchema_NiagaraNodeActions";
 			FToolMenuSection& Section = Menu->AddSection(SectionName, LOCTEXT("PinConversionMenuHeader", "Convert Pins"));
 			Section.AddSubMenu(
-				"ConvertNumericSpecific",
-				LOCTEXT("ConvertNumericSpecific", "Convert All Numerics To..."),
-				LOCTEXT("ConvertNumericSpecificToolTip", "Convert all Numeric pins to the specific typed pin."),
+				"ConvertAllNumericSpecific",
+				LOCTEXT("ConvertAllNumericSpecific", "Convert All Numerics To..."),
+				LOCTEXT("ConvertAllNumericSpecificToolTip", "Convert all Numeric pins to the specific typed pin."),
 				FNewToolMenuDelegate::CreateUObject((UEdGraphSchema_Niagara*const)this, &UEdGraphSchema_Niagara::GetNumericConversionToSubMenuActionsAll, SectionName, const_cast<UNiagaraNode*>(Node)));
 		}
 		
diff --git a/Engine/Plugins/FX/Niagara/Source/NiagaraEditorWidgets/Private/SNiagaraOverviewGraph.cpp b/Engine/Plugins/FX/Niagara/Source/NiagaraEditorWidgets/Private/SNiagaraOverviewGraph.cpp
--- a/Engine/Plugins/FX/Niagara/Source/NiagaraEditorWidgets/Private/SNiagaraOverviewGraph.cpp
+++ b/Engine/Plugins/FX/Niagara/Source/NiagaraEditorWidgets/Private/SNiagaraOverviewGraph.cpp
@@ -185,7 +185,7 @@ FActionMenuContent SNiagaraOverviewGraph::OnCreateGraphActionMenu(UEdGraph* InGr
 		}
 		MenuBuilder.EndSection();
 
-		MenuBuilder.BeginSection(TEXT("NiagaraOverview_View"), LOCTEXT("EditGraph", "View"));
+		MenuBuilder.BeginSection(TEXT("NiagaraOverview_View"), LOCTEXT("View", "View"));
 		{
 			MenuBuilder.AddMenuEntry(FNiagaraEditorCommands::Get().ZoomToFit);
 			MenuBuilder.AddMenuEntry(FNiagaraEditorCommands::Get().ZoomToFitAll);
diff --git a/Engine/Plugins/Importers/USDImporter/Source/USDStageEditor/Private/Widgets/SUSDPrimPropertiesList.cpp b/Engine/Plugins/Importers/USDImporter/Source/USDStageEditor/Private/Widgets/SUSDPrimPropertiesList.cpp
--- a/Engine/Plugins/Importers/USDImporter/Source/USDStageEditor/Private/Widgets/SUSDPrimPropertiesList.cpp
+++ b/Engine/Plugins/Importers/USDImporter/Source/USDStageEditor/Private/Widgets/SUSDPrimPropertiesList.cpp
@@ -182,4 +182,6 @@ void SUsdPrimPropertiesList::SetPrimPath( const TCHAR* InPrimPath )
 	RequestListRefresh();
 }
 
+#undef LOCTEXT_NAMESPACE
+
 #endif // #if USE_USD_SDK
diff --git a/Engine/Plugins/Runtime/Oculus/OculusAudio/Source/OculusAudioEditor/Private/OculusAudioSourceSettingsFactory.cpp b/Engine/Plugins/Runtime/Oculus/OculusAudio/Source/OculusAudioEditor/Private/OculusAudioSourceSettingsFactory.cpp
--- a/Engine/Plugins/Runtime/Oculus/OculusAudio/Source/OculusAudioEditor/Private/OculusAudioSourceSettingsFactory.cpp
+++ b/Engine/Plugins/Runtime/Oculus/OculusAudio/Source/OculusAudioEditor/Private/OculusAudioSourceSettingsFactory.cpp
@@ -57,4 +57,6 @@ UObject* UOculusAudioSourceSettingsFactory::FactoryCreateNew(UClass* Class, UObj
 uint32 UOculusAudioSourceSettingsFactory::GetMenuCategories() const
 {
     return EAssetTypeCategories::Sounds;
-}
\ No newline at end of file
+}
+
+#undef LOCTEXT_NAMESPACE
diff --git a/Engine/Source/Editor/AudioEditor/Private/AssetTypeActions/AssetTypeActions_SoundClass.cpp b/Engine/Source/Editor/AudioEditor/Private/AssetTypeActions/AssetTypeActions_SoundClass.cpp
--- a/Engine/Source/Editor/AudioEditor/Private/AssetTypeActions/AssetTypeActions_SoundClass.cpp
+++ b/Engine/Source/Editor/AudioEditor/Private/AssetTypeActions/AssetTypeActions_SoundClass.cpp
@@ -45,8 +45,8 @@ void FAssetTypeActions_SoundClass::GetActions(const TArray<UObject*>& InObjects,
 	auto Sounds = GetTypedWeakObjectPtrs<USoundClass>(InObjects);
 
 	MenuBuilder.AddMenuEntry(
-		LOCTEXT("Sound_MuteSound", "Mute"),
-		LOCTEXT("Sound_MuteSoundTooltip", "Mutes anything using this SoundClass"),
+		LOCTEXT("Sound_MuteSoundClass", "Mute"),
+		LOCTEXT("Sound_MuteSoundClassTooltip", "Mutes anything using this SoundClass"),
 		FSlateIcon(FEditorStyle::GetStyleSetName(), "MediaAsset.AssetActions.Mute.Small"),
 		FUIAction(
 			FExecuteAction::CreateSP(this, &FAssetTypeActions_SoundClass::ExecuteMute, Sounds),
@@ -58,8 +58,8 @@ void FAssetTypeActions_SoundClass::GetActions(const TArray<UObject*>& InObjects,
 	);
 
 	MenuBuilder.AddMenuEntry(
-		LOCTEXT("Sound_SoloSound", "Solo"),
-		LOCTEXT("Sound_SoloSoundTooltip", "Mutes anything using this SoundClass"),
+		LOCTEXT("Sound_SoloSoundClass", "Solo"),
+		LOCTEXT("Sound_SoloSoundClassTooltip", "Mutes anything using this SoundClass"),
 		FSlateIcon(FEditorStyle::GetStyleSetName(), "MediaAsset.AssetActions.Solo.Small"),
 		FUIAction(
 			FExecuteAction::CreateSP(this, &FAssetTypeActions_SoundClass::ExecuteSolo, Sounds),
diff --git a/Engine/Source/Editor/BlueprintGraph/Private/K2Node_CallFunction.cpp b/Engine/Source/Editor/BlueprintGraph/Private/K2Node_CallFunction.cpp
--- a/Engine/Source/Editor/BlueprintGraph/Private/K2Node_CallFunction.cpp
+++ b/Engine/Source/Editor/BlueprintGraph/Private/K2Node_CallFunction.cpp
@@ -2057,11 +2057,11 @@ void UK2Node_CallFunction::ValidateNodeDuringCompilation(class FCompilerResultsL
 			{
 				if(bCanTreatAsError)
 				{
-					MessageLog.Error(*LOCTEXT("FunctionPrivateAccessed", "Function '@@' is protected and can't be accessed outside of its hierarchy.").ToString(), this);
+					MessageLog.Error(*LOCTEXT("FunctionProtectedAccessed", "Function '@@' is protected and can't be accessed outside of its hierarchy.").ToString(), this);
 				}
 				else
 				{
-					MessageLog.Note(*LOCTEXT("FunctionPrivateAccessedNote", "Function '@@' is protected and can't be accessed outside of its hierarchy - this will be an error if the asset is resaved.").ToString(), this);
+					MessageLog.Note(*LOCTEXT("FunctionProtectedAccessedNote", "Function '@@' is protected and can't be accessed outside of its hierarchy - this will be an error if the asset is resaved.").ToString(), this);
 				}
 			}
 
diff --git a/Engine/Source/Editor/ComponentVisualizers/Private/SplineComponentVisualizer.cpp b/Engine/Source/Editor/ComponentVisualizers/Private/SplineComponentVisualizer.cpp
--- a/Engine/Source/Editor/ComponentVisualizers/Private/SplineComponentVisualizer.cpp
+++ b/Engine/Source/Editor/ComponentVisualizers/Private/SplineComponentVisualizer.cpp
@@ -2482,7 +2482,7 @@ void FSplineComponentVisualizer::GenerateContextMenuSections(FMenuBuilder& InMen
 
 			InMenuBuilder.AddSubMenu(
 				LOCTEXT("SplinePointType", "Spline Point Type"),
-				LOCTEXT("KeyTypeTooltip", "Define the type of the spline point."),
+				LOCTEXT("SplinePointTypeTooltip", "Define the type of the spline point."),
 				FNewMenuDelegate::CreateSP(this, &FSplineComponentVisualizer::GenerateSplinePointTypeSubMenu));
 
 			// Only add the Automatic Tangents submenu if any of the keys is a curve type
@@ -2524,13 +2524,13 @@ void FSplineComponentVisualizer::GenerateContextMenuSections(FMenuBuilder& InMen
 
 		InMenuBuilder.AddSubMenu(
 			LOCTEXT("SnapAlign", "Snap/Align"),
-			LOCTEXT("KeyTypeTooltip", "Snap align options."),
+			LOCTEXT("SnapAlignTooltip", "Snap align options."),
 			FNewMenuDelegate::CreateSP(this, &FSplineComponentVisualizer::GenerateSnapAlignSubMenu));
 
 		/* temporarily disabled
 		InMenuBuilder.AddSubMenu(
 			LOCTEXT("LockAxis", "Lock Axis"),
-			LOCTEXT("KeyTypeTooltip", "Axis to lock when adding new spline points."),
+			LOCTEXT("LockAxisTooltip", "Axis to lock when adding new spline points."),
 			FNewMenuDelegate::CreateSP(this, &FSplineComponentVisualizer::GenerateLockAxisSubMenu));
 			*/
 	}
diff --git a/Engine/Source/Editor/GameProjectGeneration/Private/SGameProjectDialog.cpp b/Engine/Source/Editor/GameProjectGeneration/Private/SGameProjectDialog.cpp
--- a/Engine/Source/Editor/GameProjectGeneration/Private/SGameProjectDialog.cpp
+++ b/Engine/Source/Editor/GameProjectGeneration/Private/SGameProjectDialog.cpp
@@ -164,7 +164,7 @@ TSharedRef<SWidget> SGameProjectDialog::CreateLandingPage()
 	}
 	else if (DialogMode == EMode::Both)
 	{
-		Description = LOCTEXT("LandingPageDesc_Open", "Choose a recent <RichTextBlock.BoldHighlight>project</> to load, press <RichTextBlock.BoldHighlight>More</> to see all projects on your computer, or choose a template <RichTextBlock.BoldHighlight>category</> to create a new project.");
+		Description = LOCTEXT("LandingPageDesc_Both", "Choose a recent <RichTextBlock.BoldHighlight>project</> to load, press <RichTextBlock.BoldHighlight>More</> to see all projects on your computer, or choose a template <RichTextBlock.BoldHighlight>category</> to create a new project.");
 	}
 
 	TSharedRef<SWidget> Widget = SNew(SVerticalBox)
diff --git a/Engine/Source/Editor/Kismet/Private/BlueprintEditor.cpp b/Engine/Source/Editor/Kismet/Private/BlueprintEditor.cpp
--- a/Engine/Source/Editor/Kismet/Private/BlueprintEditor.cpp
+++ b/Engine/Source/Editor/Kismet/Private/BlueprintEditor.cpp
@@ -5420,7 +5420,7 @@ void FBlueprintEditor::OnConvertFunctionToEvent()
 
 		FFormatNamedArguments Arguments;
 		Arguments.Add(TEXT("SpecificErrorMessage"), SpecificErrorMessage);
-		MessageLog.Notify(FText::Format(LOCTEXT("OnConvertEventToFunctionError", "Convert Event to Function Failed!\n{SpecificErrorMessage}"), Arguments));
+		MessageLog.Notify(FText::Format(LOCTEXT("OnConvertEventToFunctionErrorMsg", "Convert Event to Function Failed!\n{SpecificErrorMessage}"), Arguments));
 	}
 }
 
@@ -5602,7 +5602,7 @@ void FBlueprintEditor::OnConvertEventToFunction()
 		// Format the title node
 		FFormatNamedArguments Arguments;
 		Arguments.Add(TEXT("SpecificErrorMessage"), SpecificErrorMessage);
-		MessageLog.Notify(FText::Format(LOCTEXT("OnConvertEventToFunctionError", "Convert Event to Function Failed!\n{SpecificErrorMessage}"), Arguments));
+		MessageLog.Notify(FText::Format(LOCTEXT("OnConvertEventToFunctionErrorMsg", "Convert Event to Function Failed!\n{SpecificErrorMessage}"), Arguments));
 	}
 }
 
diff --git a/Engine/Source/Editor/LandscapeEditor/Private/LandscapeEditorDetailCustomization_LayersBrushStack.cpp b/Engine/Source/Editor/LandscapeEditor/Private/LandscapeEditorDetailCustomization_LayersBrushStack.cpp
--- a/Engine/Source/Editor/LandscapeEditor/Private/LandscapeEditorDetailCustomization_LayersBrushStack.cpp
+++ b/Engine/Source/Editor/LandscapeEditor/Private/LandscapeEditorDetailCustomization_LayersBrushStack.cpp
@@ -277,7 +277,7 @@ void FLandscapeEditorCustomNodeBuilder_LayersBrushStack::FillAddBrushMenu(FMenuB
 	{
 		FUIAction AddAction = FUIAction(FExecuteAction::CreateLambda([SharedThis, Brush]()
 		{
-			const FScopedTransaction Transaction(LOCTEXT("LandscapeBrushAddToCurrentLayerTransaction", "Add brush to selected edit layer"));
+			const FScopedTransaction Transaction(LOCTEXT("LandscapeBrushAddToSelectedLayerTransaction", "Add brush to selected edit layer"));
 			GetEditorMode()->AddBrushToCurrentLayer(Brush);
 		}));
 		MenuBuilder.AddMenuEntry(FText::FromString(Brush->GetActorLabel()), FText(), FSlateIcon(), AddAction);
diff --git a/Engine/Source/Editor/MainFrame/Private/Menus/LayoutsMenu.cpp b/Engine/Source/Editor/MainFrame/Private/Menus/LayoutsMenu.cpp
--- a/Engine/Source/Editor/MainFrame/Private/Menus/LayoutsMenu.cpp
+++ b/Engine/Source/Editor/MainFrame/Private/Menus/LayoutsMenu.cpp
@@ -211,16 +211,16 @@ FText GetDisplayTextInternal(const FString& InString)
 	const FText DisplayName = FText::FromString(FName::NameToDisplayString(DisplayNameText.ToString(), bIsBool));
 	return DisplayName;
 }
-FText GetTooltipTextInternal(const FText& InDisplayName, const FString& InLayoutFilePath, const FText& InLayoutName, const int32 LayoutIndex)
+FText GetTooltipTextInternal(const FText& InDisplayName, const FString& InLayoutFilePath, const FText& InLayoutName)
 {
 	FText Tooltip;
 	if (InLayoutName.IsEmpty())
 	{
-		Tooltip = FText::Format(LOCTEXT("DisplayName{0}", "Layout name:\n{1}\n\nFull file path:\n{2}"), FText::AsNumber(LayoutIndex), InDisplayName, FText::FromString(InLayoutFilePath));
+		Tooltip = FText::Format(LOCTEXT("DisplayNameFmt", "Layout name:\n{0}\n\nFull file path:\n{1}"), InDisplayName, FText::FromString(InLayoutFilePath));
 	}
 	else
 	{
-		Tooltip = FText::Format(LOCTEXT("DisplayName{0}", "Description:\n{1}.\n\nFull file path:\n{2}"), FText::AsNumber(LayoutIndex), InLayoutName, FText::FromString(InLayoutFilePath));
+		Tooltip = FText::Format(LOCTEXT("LayoutNameFmt", "Description:\n{0}.\n\nFull file path:\n{1}"), InLayoutName, FText::FromString(InLayoutFilePath));
 	}
 	return Tooltip;
 }
@@ -240,7 +240,7 @@ void DisplayLayoutsInternal(FToolMenuSection& InSection, const TArray<TSharedPtr
 			const FText LayoutDescription = FLayoutSaveRestore::LoadSectionFromConfig(LayoutFilePath, "LayoutDescription");
 			// If no localization name, then display the file name
 			const FText DisplayName = (!LayoutName.IsEmpty() ? LayoutName : GetDisplayTextInternal(InLayoutIniFileNames[LayoutIndex]));
-			const FText Tooltip = GetTooltipTextInternal(DisplayName, LayoutFilePath, LayoutDescription, LayoutIndex);
+			const FText Tooltip = GetTooltipTextInternal(DisplayName, LayoutFilePath, LayoutDescription);
 			InSection.AddMenuEntry(InXLayoutCommands[LayoutIndex], DisplayName, Tooltip).Name = NAME_None;
 		}
 	}
@@ -579,7 +579,7 @@ void FLayoutsMenuLoad::ImportLayout()
 			{
 				const FText TextBody = LOCTEXT("SuccessfulImportBody",
 					"The layout(s) were successfully imported into the \"User Layouts\" section. However, no layout has been loaded into your current Unreal Editor UI because PIE is currently running. In order to do so, you must stop PIE and then load the layout from the \"User Layouts\" section.");
-				const FText TextTitle = LOCTEXT("UnsuccessfulImportHeader", "Successful Import!");
+				const FText TextTitle = LOCTEXT("SuccessfulImportHeader", "Successful Import!");
 				OpenMsgDlgInt(EAppMsgType::Ok, TextBody, TextTitle);
 				return;
 			}
diff --git a/Engine/Source/Editor/MovieSceneTools/Private/TrackEditors/AttachTrackEditor.cpp b/Engine/Source/Editor/MovieSceneTools/Private/TrackEditors/AttachTrackEditor.cpp
--- a/Engine/Source/Editor/MovieSceneTools/Private/TrackEditors/AttachTrackEditor.cpp
+++ b/Engine/Source/Editor/MovieSceneTools/Private/TrackEditors/AttachTrackEditor.cpp
@@ -87,7 +87,7 @@ public:
 		TArray<FGuid> ObjectBindings;
 		ObjectBindings.Add(ObjectBinding);
 
-		MenuBuilder.BeginSection(NAME_None, LOCTEXT("AttachOptions", "Attach Section Options"));
+		MenuBuilder.BeginSection(NAME_None, LOCTEXT("AttachSectionOptions", "Attach Section Options"));
 
 		MenuBuilder.AddSubMenu(
 			LOCTEXT("SetAttach", "Attach"), LOCTEXT("SetAttachTooltip", "Set attach"),
diff --git a/Engine/Source/Editor/Persona/Private/PersonaMeshDetails.cpp b/Engine/Source/Editor/Persona/Private/PersonaMeshDetails.cpp
--- a/Engine/Source/Editor/Persona/Private/PersonaMeshDetails.cpp
+++ b/Engine/Source/Editor/Persona/Private/PersonaMeshDetails.cpp
@@ -489,7 +489,7 @@ void FSkeletalMeshReductionSettingsLayout::GenerateChildContent(IDetailChildrenB
 		AddFloatRow(ChildrenBuilder, 
 			LOCTEXT("PercentTriangles_Row", "Triangle Percentage"),
 			LOCTEXT("PercentTriangles", "Percent of Triangles"),
-			LOCTEXT("PercentTriangles_ToolTip", "The percentage of triangles to retain as a ratio, e.g. 0.1 indicates 10 percent."),
+			LOCTEXT("PercentTriangles_DeviationToolTip", "The percentage of triangles to retain as a ratio, e.g. 0.1 indicates 10 percent."),
 			0.0f,
 			1.0f,
 			FGetFloatDelegate::CreateRaw(this, &FSkeletalMeshReductionSettingsLayout::GetNumTrianglesPercentage),
diff --git a/Engine/Source/Editor/Sequencer/Private/DisplayNodes/SequencerDisplayNode.cpp b/Engine/Source/Editor/Sequencer/Private/DisplayNodes/SequencerDisplayNode.cpp
--- a/Engine/Source/Editor/Sequencer/Private/DisplayNodes/SequencerDisplayNode.cpp
+++ b/Engine/Source/Editor/Sequencer/Private/DisplayNodes/SequencerDisplayNode.cpp
@@ -1015,7 +1015,7 @@ void FSequencerDisplayNode::BuildContextMenu(FMenuBuilder& MenuBuilder)
 		{
 			MenuBuilder.AddMenuEntry(
 				LOCTEXT("ToggleNodeSolo", "Solo"),
-				LOCTEXT("ToggleNodeLockTooltip", "Solo or unsolo this node or selected tracks"),
+				LOCTEXT("ToggleNodeSoloTooltip", "Solo or unsolo this node or selected tracks"),
 				FSlateIcon(),
 				FUIAction(
 					FExecuteAction::CreateRaw(&ParentTree, &FSequencerNodeTree::ToggleSelectedNodesSolo),
@@ -1028,7 +1028,7 @@ void FSequencerDisplayNode::BuildContextMenu(FMenuBuilder& MenuBuilder)
 
 			MenuBuilder.AddMenuEntry(
 				LOCTEXT("ToggleNodeMute", "Mute"),
-				LOCTEXT("ToggleNodeLockTooltip", "Mute or unmute this node or selected tracks"),
+				LOCTEXT("ToggleNodeMuteTooltip", "Mute or unmute this node or selected tracks"),
 				FSlateIcon(),
 				FUIAction(
 					FExecuteAction::CreateRaw(&ParentTree, &FSequencerNodeTree::ToggleSelectedNodesMute),
diff --git a/Engine/Source/Editor/Sequencer/Private/SExposeBindingWidget.cpp b/Engine/Source/Editor/Sequencer/Private/SExposeBindingWidget.cpp
--- a/Engine/Source/Editor/Sequencer/Private/SExposeBindingWidget.cpp
+++ b/Engine/Source/Editor/Sequencer/Private/SExposeBindingWidget.cpp
@@ -206,7 +206,7 @@ FReply SExposeBindingWidget::RemoveFromExposedName(FName InNameToRemove)
 	UMovieSceneSequence*   Sequence   = Sequencer ? Sequencer->GetFocusedMovieSceneSequence() : nullptr;
 	UMovieScene*           MovieScene = Sequence  ? Sequence->GetMovieScene()                 : nullptr;
 
-	FScopedTransaction Transaction(FText::Format(LOCTEXT("ExposeBinding_Transaction", "Stop exposing binding as {0}"), FText::FromName(InNameToRemove)));
+	FScopedTransaction Transaction(FText::Format(LOCTEXT("StopExposingBinding_Transaction", "Stop exposing binding as {0}"), FText::FromName(InNameToRemove)));
 
 	MovieScene->Modify();
 
diff --git a/Engine/Source/Editor/Sequencer/Private/SExposedBindingsWidget.cpp b/Engine/Source/Editor/Sequencer/Private/SExposedBindingsWidget.cpp
--- a/Engine/Source/Editor/Sequencer/Private/SExposedBindingsWidget.cpp
+++ b/Engine/Source/Editor/Sequencer/Private/SExposedBindingsWidget.cpp
@@ -360,7 +360,7 @@ FReply SExposedBindingsWidget::RemoveExposedName(FName InNameToRemove)
 	UMovieSceneSequence*   Sequence   = Sequencer ? Sequencer->GetFocusedMovieSceneSequence() : nullptr;
 	UMovieScene*           MovieScene = Sequence  ? Sequence->GetMovieScene()                 : nullptr;
 
-	FScopedTransaction Transaction(FText::Format(LOCTEXT("ExposeBinding_Transaction", "Stop exposing binding {0}"), FText::FromName(InNameToRemove)));
+	FScopedTransaction Transaction(FText::Format(LOCTEXT("StopExposingBinding_Transaction", "Stop exposing binding {0}"), FText::FromName(InNameToRemove)));
 
 	MovieScene->Modify();
 	MovieScene->RemoveExposedBinding(InNameToRemove);
diff --git a/Engine/Source/Editor/UnrealEd/Private/StaticLightingSystem/StaticLightingSystem.cpp b/Engine/Source/Editor/UnrealEd/Private/StaticLightingSystem/StaticLightingSystem.cpp
--- a/Engine/Source/Editor/UnrealEd/Private/StaticLightingSystem/StaticLightingSystem.cpp
+++ b/Engine/Source/Editor/UnrealEd/Private/StaticLightingSystem/StaticLightingSystem.cpp
@@ -2066,7 +2066,7 @@ void FStaticLightingSystem::GatherScene()
 		{
 			if (LegacyAtmosphericFogRegistered)
 			{
-				FMessageLog("LightingResults").Warning(LOCTEXT("LightmassError_BuildSelected", "Both a legacy AtmosphericFog and a new SkyAtmosphere wants to register. Lightmass will not consider the legacy component."));
+				FMessageLog("LightingResults").Warning(LOCTEXT("LightmassError_BothAtmosphericFogAndSkyAtmosphereSelected", "Both a legacy AtmosphericFog and a new SkyAtmosphere wants to register. Lightmass will not consider the legacy component."));
 			}
 			LightmassExporter->SetAtmosphericComponent(nullptr);
 			LightmassExporter->SetSkyAtmosphereComponent(SkyAtmosphere);
