Niagara - Rework stack and overview styling for disabled emitters and modules.

diff --git a/Engine/Plugins/FX/Niagara/Source/NiagaraEditor/Private/NiagaraOverviewNode.cpp b/Engine/Plugins/FX/Niagara/Source/NiagaraEditor/Private/NiagaraOverviewNode.cpp
--- a/Engine/Plugins/FX/Niagara/Source/NiagaraEditor/Private/NiagaraOverviewNode.cpp
+++ b/Engine/Plugins/FX/Niagara/Source/NiagaraEditor/Private/NiagaraOverviewNode.cpp
@@ -2,9 +2,17 @@
 
 #include "NiagaraOverviewNode.h"
 #include "NiagaraSystem.h"
+#include "NiagaraEditorModule.h"
+#include "ViewModels/Stack/NiagaraStackEntry.h"
+
+#include "Modules/ModuleManager.h"
 
 #define LOCTEXT_NAMESPACE "NiagaraOverviewNodeStackItem"
 
+bool UNiagaraOverviewNode::bColorsAreInitialized = false;
+FLinearColor UNiagaraOverviewNode::EmitterColor;
+FLinearColor UNiagaraOverviewNode::SystemColor;
+
 UNiagaraOverviewNode::UNiagaraOverviewNode()
 	: OwningSystem(nullptr)
 	, EmitterHandleGuid(FGuid())
@@ -52,6 +60,17 @@ FText UNiagaraOverviewNode::GetNodeTitle(ENodeTitleType::Type TitleType) const
 	}
 }
 
+FLinearColor UNiagaraOverviewNode::GetNodeTitleColor() const
+{
+	if (bColorsAreInitialized == false)
+	{
+		FNiagaraEditorModule& NiagaraEditorModule = FModuleManager::LoadModuleChecked<FNiagaraEditorModule>("NiagaraEditor");
+		EmitterColor = NiagaraEditorModule.GetWidgetProvider()->GetColorForExecutionCategory(UNiagaraStackEntry::FExecutionCategoryNames::Emitter);
+		SystemColor = NiagaraEditorModule.GetWidgetProvider()->GetColorForExecutionCategory(UNiagaraStackEntry::FExecutionCategoryNames::System);
+	}
+	return EmitterHandleGuid.IsValid() ? EmitterColor : SystemColor;
+}
+
 bool UNiagaraOverviewNode::CanUserDeleteNode() const
 {
 	return EmitterHandleGuid.IsValid();
diff --git a/Engine/Plugins/FX/Niagara/Source/NiagaraEditor/Private/NiagaraSystemEditorData.cpp b/Engine/Plugins/FX/Niagara/Source/NiagaraEditor/Private/NiagaraSystemEditorData.cpp
--- a/Engine/Plugins/FX/Niagara/Source/NiagaraEditor/Private/NiagaraSystemEditorData.cpp
+++ b/Engine/Plugins/FX/Niagara/Source/NiagaraEditor/Private/NiagaraSystemEditorData.cpp
@@ -256,9 +256,11 @@ void UNiagaraSystemEditorData::SynchronizeOverviewGraphWithSystem(UNiagaraSystem
 	for (const FNiagaraEmitterHandle& EmitterHandle : OwnerSystem.GetEmitterHandles())
 	{
 		int32 EmitterNodeIndex = OverviewNodes.IndexOfByPredicate([&EmitterHandle] (UNiagaraOverviewNode* OverviewNode) { return OverviewNode->GetEmitterHandleGuid() == EmitterHandle.GetId(); });
+		UNiagaraOverviewNode* EmitterNode = nullptr;
 		if (EmitterNodeIndex != INDEX_NONE)
 		{
 			// If a node for this emitter exists already, remove it from the collection so that we can track nodes for emitters which no longer exist.
+			EmitterNode = OverviewNodes[EmitterNodeIndex];
 			OverviewNodes.RemoveAt(EmitterNodeIndex);
 		}
 		else
@@ -278,7 +280,10 @@ void UNiagaraSystemEditorData::SynchronizeOverviewGraphWithSystem(UNiagaraSystem
 			EmitterOverviewNodeCreator.Finalize();
 
 			RightNode = EmitterOverviewNode;
+			EmitterNode = EmitterOverviewNode;
 		}
+
+		EmitterNode->SetEnabledState(EmitterHandle.GetIsEnabled() ? ENodeEnabledState::Enabled : ENodeEnabledState::Disabled);
 	}
 
 	// If there are any nodes remaining in the list they're no longer being used so destroy them.
diff --git a/Engine/Plugins/FX/Niagara/Source/NiagaraEditor/Private/Sequencer/NiagaraSequence/NiagaraEmitterTrackEditor.cpp b/Engine/Plugins/FX/Niagara/Source/NiagaraEditor/Private/Sequencer/NiagaraSequence/NiagaraEmitterTrackEditor.cpp
--- a/Engine/Plugins/FX/Niagara/Source/NiagaraEditor/Private/Sequencer/NiagaraSequence/NiagaraEmitterTrackEditor.cpp
+++ b/Engine/Plugins/FX/Niagara/Source/NiagaraEditor/Private/Sequencer/NiagaraSequence/NiagaraEmitterTrackEditor.cpp
@@ -42,18 +42,6 @@ public:
 				.Image(FEditorStyle::GetBrush("Icons.Info"))
 				.ToolTipText(this, &SEmitterTrackWidget::GetTrackErrorIconToolTip)
 			]
-			// Enabled checkbox.
-			+ SHorizontalBox::Slot()
-			.AutoWidth()
-			.VAlign(VAlign_Center)
-			.Padding(3, 0, 0, 0)
-			[
-				SNew(SCheckBox)
-				.ToolTipText(LOCTEXT("EnabledTooltip", "Toggle whether or not this emitter is enabled."))
-				.IsChecked(this, &SEmitterTrackWidget::GetEnabledCheckState)
-				.OnCheckStateChanged(this, &SEmitterTrackWidget::OnEnabledCheckStateChanged)
-				.Visibility(this, &SEmitterTrackWidget::GetEnableCheckboxVisibility)
-			]
 			// Isolate toggle
 			+ SHorizontalBox::Slot()
 			.AutoWidth()
@@ -95,6 +83,19 @@ public:
 				];
 		}
 
+		// Enabled checkbox.
+		TrackBox->AddSlot()
+			.AutoWidth()
+			.VAlign(VAlign_Center)
+			.Padding(3, 0, 0, 0)
+			[
+				SNew(SCheckBox)
+				.ToolTipText(LOCTEXT("EnabledTooltip", "Toggle whether or not this emitter is enabled."))
+				.IsChecked(this, &SEmitterTrackWidget::GetEnabledCheckState)
+				.OnCheckStateChanged(this, &SEmitterTrackWidget::OnEnabledCheckStateChanged)
+				.Visibility(this, &SEmitterTrackWidget::GetEnableCheckboxVisibility)
+			];
+
 		ChildSlot
 		[
 			TrackBox
diff --git a/Engine/Plugins/FX/Niagara/Source/NiagaraEditor/Private/ViewModels/NiagaraSystemViewModel.cpp b/Engine/Plugins/FX/Niagara/Source/NiagaraEditor/Private/ViewModels/NiagaraSystemViewModel.cpp
--- a/Engine/Plugins/FX/Niagara/Source/NiagaraEditor/Private/ViewModels/NiagaraSystemViewModel.cpp
+++ b/Engine/Plugins/FX/Niagara/Source/NiagaraEditor/Private/ViewModels/NiagaraSystemViewModel.cpp
@@ -1215,6 +1215,18 @@ void FNiagaraSystemViewModel::EmitterHandlePropertyChanged(FGuid OwningEmitterHa
 		}
 		Sequencer->NotifyMovieSceneDataChanged(EMovieSceneDataChangeType::MovieSceneStructureItemsChanged);
 	}
+
+	// Refresh the overview nodes and the emitter stacks just in case the emitter enabled state changed.
+	GetEditorData().SynchronizeOverviewGraphWithSystem(GetSystem());
+	TSharedPtr<FNiagaraEmitterHandleViewModel> EmitterHandleViewModel = GetEmitterHandleViewModelById(OwningEmitterHandleId);
+	if (EmitterHandleViewModel.IsValid())
+	{
+		for (UNiagaraStackEntry* RootEntry : EmitterHandleViewModel->GetEmitterStackViewModel()->GetRootEntries())
+		{
+			RootEntry->RefreshChildren();
+		}
+	}
+
 	ResetSystem(ETimeResetMode::AllowResetTime, EMultiResetMode::ResetThisInstance, EReinitMode::ReinitializeSystem);
 }
 
diff --git a/Engine/Plugins/FX/Niagara/Source/NiagaraEditor/Private/ViewModels/Stack/NiagaraStackEmitterSettingsGroup.cpp b/Engine/Plugins/FX/Niagara/Source/NiagaraEditor/Private/ViewModels/Stack/NiagaraStackEmitterSettingsGroup.cpp
--- a/Engine/Plugins/FX/Niagara/Source/NiagaraEditor/Private/ViewModels/Stack/NiagaraStackEmitterSettingsGroup.cpp
+++ b/Engine/Plugins/FX/Niagara/Source/NiagaraEditor/Private/ViewModels/Stack/NiagaraStackEmitterSettingsGroup.cpp
@@ -99,18 +99,6 @@ UNiagaraStackEmitterSettingsGroup::UNiagaraStackEmitterSettingsGroup()
 
 void UNiagaraStackEmitterSettingsGroup::RefreshChildrenInternal(const TArray<UNiagaraStackEntry*>& CurrentChildren, TArray<UNiagaraStackEntry*>& NewChildren, TArray<FStackIssue>& NewIssues)
 {
-	FName PropertiesSpacerKey = "PropertiesSpacer";
-	UNiagaraStackSpacer* PropertiesSpacer = FindCurrentChildOfTypeByPredicate<UNiagaraStackSpacer>(CurrentChildren,
-		[=](UNiagaraStackSpacer* CurrentPropertiesSpacer) { return CurrentPropertiesSpacer->GetSpacerKey() == PropertiesSpacerKey; });
-
-	if (PropertiesSpacer == nullptr)
-	{
-		PropertiesSpacer = NewObject<UNiagaraStackSpacer>(this);
-		PropertiesSpacer->Initialize(CreateDefaultChildRequiredData(), PropertiesSpacerKey);
-	}
-
-	NewChildren.Add(PropertiesSpacer);
-
 	if (PropertiesItem == nullptr)
 	{
 		PropertiesItem = NewObject<UNiagaraStackEmitterPropertiesItem>(this);
diff --git a/Engine/Plugins/FX/Niagara/Source/NiagaraEditor/Private/ViewModels/Stack/NiagaraStackEntry.cpp b/Engine/Plugins/FX/Niagara/Source/NiagaraEditor/Private/ViewModels/Stack/NiagaraStackEntry.cpp
--- a/Engine/Plugins/FX/Niagara/Source/NiagaraEditor/Private/ViewModels/Stack/NiagaraStackEntry.cpp
+++ b/Engine/Plugins/FX/Niagara/Source/NiagaraEditor/Private/ViewModels/Stack/NiagaraStackEntry.cpp
@@ -130,6 +130,7 @@ UNiagaraStackEntry::UNiagaraStackEntry()
 	: IndentLevel(0)
 	, bIsFinalized(false)
 	, bIsSearchResult(false)
+	, bOwnerIsEnabled(true)
 {
 }
 
@@ -230,6 +231,11 @@ bool UNiagaraStackEntry::GetIsEnabled() const
 	return true;
 }
 
+bool UNiagaraStackEntry::GetOwnerIsEnabled() const
+{
+	return bOwnerIsEnabled;
+}
+
 FName UNiagaraStackEntry::GetExecutionCategoryName() const
 {
 	return ExecutionCategoryName;
@@ -476,7 +482,9 @@ void UNiagaraStackEntry::RefreshChildren()
 
 	for (UNiagaraStackEntry* Child : Children)
 	{
+		UNiagaraStackEntry* OuterOwner = Cast<UNiagaraStackEntry>(Child->GetOuter());
 		Child->IndentLevel = GetChildIndentLevel();
+		Child->bOwnerIsEnabled = OuterOwner == nullptr || (OuterOwner->GetIsEnabled() && OuterOwner->GetOwnerIsEnabled());
 		Child->RefreshChildren();
 		Child->OnStructureChanged().AddUObject(this, &UNiagaraStackEntry::ChildStructureChanged);
 		Child->OnDataObjectModified().AddUObject(this, &UNiagaraStackEntry::ChildDataObjectModified);
diff --git a/Engine/Plugins/FX/Niagara/Source/NiagaraEditor/Private/ViewModels/Stack/NiagaraStackItemGroup.cpp b/Engine/Plugins/FX/Niagara/Source/NiagaraEditor/Private/ViewModels/Stack/NiagaraStackItemGroup.cpp
--- a/Engine/Plugins/FX/Niagara/Source/NiagaraEditor/Private/ViewModels/Stack/NiagaraStackItemGroup.cpp
+++ b/Engine/Plugins/FX/Niagara/Source/NiagaraEditor/Private/ViewModels/Stack/NiagaraStackItemGroup.cpp
@@ -7,6 +7,9 @@
 #include "NiagaraStackEditorData.h"
 #include "ViewModels/Stack/NiagaraStackSpacer.h"
 #include "ViewModels/Stack/NiagaraStackGraphUtilities.h"
+#include "ViewModels/NiagaraSystemViewModel.h"
+#include "ViewModels/NiagaraEmitterHandleViewModel.h"
+#include "ViewModels/NiagaraEmitterViewModel.h"
 
 
 void UNiagaraStackItemGroup::Initialize(FRequiredEntryData InRequiredEntryData, FText InDisplayName, FText InToolTip, INiagaraStackItemGroupAddUtilities* InAddUtilities)
@@ -32,6 +35,11 @@ FText UNiagaraStackItemGroup::GetTooltipText() const
 	return GroupToolTip;
 }
 
+bool UNiagaraStackItemGroup::GetIsEnabled() const
+{
+	return OwningEmitterHandleViewModel.IsValid() == false || OwningEmitterHandleViewModel->GetIsEnabled();
+}
+
 INiagaraStackItemGroupAddUtilities* UNiagaraStackItemGroup::GetAddUtilities() const
 {
 	return AddUtilities;
@@ -73,17 +81,21 @@ void UNiagaraStackItemGroup::SetDisplayName(FText InDisplayName)
 
 void UNiagaraStackItemGroup::RefreshChildrenInternal(const TArray<UNiagaraStackEntry*>& CurrentChildren, TArray<UNiagaraStackEntry*>& NewChildren, TArray<FStackIssue>& NewIssues)
 {
-	UNiagaraStackSpacer* SeparatorSpacer = FindCurrentChildOfTypeByPredicate<UNiagaraStackSpacer>(CurrentChildren,
-		[=](UNiagaraStackSpacer* CurrentSpacer) { return CurrentSpacer->GetSpacerKey() == "SeparatorSpacer"; });
-	if (SeparatorSpacer == nullptr)
-	{
-		SeparatorSpacer = NewObject<UNiagaraStackSpacer>(this);
-		FRequiredEntryData RequiredEntryData(GetSystemViewModel(), GetEmitterViewModel(), GetExecutionCategoryName(), NAME_None, GetStackEditorData());
-		SeparatorSpacer->Initialize(RequiredEntryData, "SeparatorSpacer");
-	}
-	NewChildren.Add(SeparatorSpacer);
 	RecursiveStackIssuesCount.Reset();
 	HighestIssueSeverity.Reset();
+
+	TSharedPtr<FNiagaraEmitterViewModel> OwningEmitterViewModel = GetEmitterViewModel();
+	if (OwningEmitterViewModel.IsValid())
+	{
+		if (OwningEmitterHandleViewModel.IsValid() == false || OwningEmitterHandleViewModel->GetEmitterViewModel() != OwningEmitterViewModel)
+		{
+			OwningEmitterHandleViewModel = GetSystemViewModel()->GetEmitterHandleViewModelForEmitter(OwningEmitterViewModel->GetEmitter());
+		}
+	}
+	else
+	{
+		OwningEmitterHandleViewModel.Reset();
+	}
 }
 
 int32 UNiagaraStackItemGroup::GetChildIndentLevel() const
diff --git a/Engine/Plugins/FX/Niagara/Source/NiagaraEditor/Private/ViewModels/Stack/NiagaraStackRendererItem.cpp b/Engine/Plugins/FX/Niagara/Source/NiagaraEditor/Private/ViewModels/Stack/NiagaraStackRendererItem.cpp
--- a/Engine/Plugins/FX/Niagara/Source/NiagaraEditor/Private/ViewModels/Stack/NiagaraStackRendererItem.cpp
+++ b/Engine/Plugins/FX/Niagara/Source/NiagaraEditor/Private/ViewModels/Stack/NiagaraStackRendererItem.cpp
@@ -214,6 +214,7 @@ void UNiagaraStackRendererItem::SetIsEnabled(bool bInIsEnabled)
 	RendererProperties->Modify();
 	RendererProperties->SetIsEnabled(bInIsEnabled);
 	OnDataObjectModified().Broadcast(RendererProperties.Get());
+	RefreshChildren();
 }
 
 void UNiagaraStackRendererItem::RefreshChildrenInternal(const TArray<UNiagaraStackEntry*>& CurrentChildren, TArray<UNiagaraStackEntry*>& NewChildren, TArray<FStackIssue>& NewIssues)
diff --git a/Engine/Plugins/FX/Niagara/Source/NiagaraEditor/Private/ViewModels/Stack/NiagaraStackRoot.cpp b/Engine/Plugins/FX/Niagara/Source/NiagaraEditor/Private/ViewModels/Stack/NiagaraStackRoot.cpp
--- a/Engine/Plugins/FX/Niagara/Source/NiagaraEditor/Private/ViewModels/Stack/NiagaraStackRoot.cpp
+++ b/Engine/Plugins/FX/Niagara/Source/NiagaraEditor/Private/ViewModels/Stack/NiagaraStackRoot.cpp
@@ -174,7 +174,7 @@ void UNiagaraStackRoot::RefreshChildrenInternal(const TArray<UNiagaraStackEntry*
 			FRequiredEntryData RequiredEntryData(GetSystemViewModel(), GetEmitterViewModel(),
 				SpacerExecutionCategory, NAME_None,
 				GetSystemViewModel()->GetEditorData().GetStackEditorData());
-			Spacer->Initialize(RequiredEntryData, SpacerKey, 1.0f);
+			Spacer->Initialize(RequiredEntryData, SpacerKey, 1.5f);
 		}
 
 		return Spacer;
@@ -184,22 +184,26 @@ void UNiagaraStackRoot::RefreshChildrenInternal(const TArray<UNiagaraStackEntry*
 	if (bIncludeSystemInformation)
 	{
 		NewChildren.Add(SystemSettingsGroup);
+		NewChildren.Add(GetOrCreateSpacer(FExecutionCategoryNames::System, "SystemSettingsSpacer"));
 		NewChildren.Add(SystemSpawnGroup);
+		NewChildren.Add(GetOrCreateSpacer(FExecutionCategoryNames::System, "SystemSpawnSpacer"));
 		NewChildren.Add(SystemUpdateGroup);
-		NewChildren.Add(GetOrCreateSpacer(FExecutionCategoryNames::System, "SystemFooter"));
 		NewChildren.Add(GetOrCreateSpacer(NAME_None, "SystemSpacer"));
 	}
 
 	if (bIncludeEmitterInformation)
 	{
 		NewChildren.Add(EmitterSettingsGroup);
+		NewChildren.Add(GetOrCreateSpacer(FExecutionCategoryNames::Emitter, "EmitterSettingsSpacer"));
 		NewChildren.Add(EmitterSpawnGroup);
+		NewChildren.Add(GetOrCreateSpacer(FExecutionCategoryNames::Emitter, "EmitterSpawnSpacer"));
 		NewChildren.Add(EmitterUpdateGroup);
-		NewChildren.Add(GetOrCreateSpacer(FExecutionCategoryNames::Emitter, "EmitterFooter"));
 		NewChildren.Add(GetOrCreateSpacer(NAME_None, "EmitterSpacer"));
 
 		NewChildren.Add(ParticleSpawnGroup);
+		NewChildren.Add(GetOrCreateSpacer(FExecutionCategoryNames::Particle, "ParticleSpawnSpacer"));
 		NewChildren.Add(ParticleUpdateGroup);
+		NewChildren.Add(GetOrCreateSpacer(FExecutionCategoryNames::Particle, "ParticleUpdateSpacer"));
 
 		for (const FNiagaraEventScriptProperties& EventScriptProperties : GetEmitterViewModel()->GetEmitter()->GetEventHandlers())
 		{
@@ -220,11 +224,9 @@ void UNiagaraStackRoot::RefreshChildrenInternal(const TArray<UNiagaraStackEntry*
 		}
 
 		NewChildren.Add(AddEventHandlerGroup);
-		NewChildren.Add(GetOrCreateSpacer(FExecutionCategoryNames::Particle, "ParticleFooter"));
 		NewChildren.Add(GetOrCreateSpacer(NAME_None, "ParticleSpacer"));
 
 		NewChildren.Add(RenderGroup);
-		NewChildren.Add(GetOrCreateSpacer(FExecutionCategoryNames::Render, "RenderFooter"));
 	}
 }
 
diff --git a/Engine/Plugins/FX/Niagara/Source/NiagaraEditor/Public/NiagaraEditorModule.h b/Engine/Plugins/FX/Niagara/Source/NiagaraEditor/Public/NiagaraEditorModule.h
--- a/Engine/Plugins/FX/Niagara/Source/NiagaraEditor/Public/NiagaraEditorModule.h
+++ b/Engine/Plugins/FX/Niagara/Source/NiagaraEditor/Public/NiagaraEditorModule.h
@@ -29,8 +29,9 @@ DECLARE_STATS_GROUP(TEXT("Niagara Editor"), STATGROUP_NiagaraEditor, STATCAT_Adv
 class NIAGARAEDITOR_API INiagaraEditorWidgetProvider
 {
 public:
-	virtual TSharedRef<SWidget> CreateStackView(UNiagaraStackViewModel& StackViewModel) = 0;
-	virtual TSharedRef<SWidget> CreateSystemOverview(TSharedRef<FNiagaraSystemViewModel> SystemViewModel) = 0;
+	virtual TSharedRef<SWidget> CreateStackView(UNiagaraStackViewModel& StackViewModel) const = 0;
+	virtual TSharedRef<SWidget> CreateSystemOverview(TSharedRef<FNiagaraSystemViewModel> SystemViewModel) const = 0;
+	virtual FLinearColor GetColorForExecutionCategory(FName ExecutionCategory) const = 0;
 };
 
 /** Niagara Editor module */
diff --git a/Engine/Plugins/FX/Niagara/Source/NiagaraEditor/Public/NiagaraOverviewNode.h b/Engine/Plugins/FX/Niagara/Source/NiagaraEditor/Public/NiagaraOverviewNode.h
--- a/Engine/Plugins/FX/Niagara/Source/NiagaraEditor/Public/NiagaraOverviewNode.h
+++ b/Engine/Plugins/FX/Niagara/Source/NiagaraEditor/Public/NiagaraOverviewNode.h
@@ -20,6 +20,8 @@ public:
 	/** Gets the name of this node, shown in title bar */
 	virtual FText GetNodeTitle(ENodeTitleType::Type TitleType) const override;
 
+	virtual FLinearColor GetNodeTitleColor() const override;
+
 	/** Whether or not this node can be deleted by user action */
 	virtual bool CanUserDeleteNode() const override;
 
@@ -35,4 +37,8 @@ private:
 
 	UPROPERTY()
 	FGuid EmitterHandleGuid;
+
+	static bool bColorsAreInitialized;
+	static FLinearColor SystemColor;
+	static FLinearColor EmitterColor;
 };
diff --git a/Engine/Plugins/FX/Niagara/Source/NiagaraEditor/Public/ViewModels/Stack/NiagaraStackEntry.h b/Engine/Plugins/FX/Niagara/Source/NiagaraEditor/Public/ViewModels/Stack/NiagaraStackEntry.h
--- a/Engine/Plugins/FX/Niagara/Source/NiagaraEditor/Public/ViewModels/Stack/NiagaraStackEntry.h
+++ b/Engine/Plugins/FX/Niagara/Source/NiagaraEditor/Public/ViewModels/Stack/NiagaraStackEntry.h
@@ -198,6 +198,10 @@ public:
 
 	virtual bool GetIsEnabled() const;
 
+	bool GetOwnerIsEnabled() const;
+
+	bool GetIsEnabledAndOwnerIsEnabled() const { return GetIsEnabled() && GetOwnerIsEnabled(); }
+
 	FName GetExecutionCategoryName() const;
 
 	FName GetExecutionSubcategoryName() const;
@@ -345,4 +349,6 @@ private:
 	bool bIsSearchResult;
 
 	mutable TOptional<bool> bHasBaseEmitterCache;
+
+	bool bOwnerIsEnabled;
 };
\ No newline at end of file
diff --git a/Engine/Plugins/FX/Niagara/Source/NiagaraEditor/Public/ViewModels/Stack/NiagaraStackItem.h b/Engine/Plugins/FX/Niagara/Source/NiagaraEditor/Public/ViewModels/Stack/NiagaraStackItem.h
--- a/Engine/Plugins/FX/Niagara/Source/NiagaraEditor/Public/ViewModels/Stack/NiagaraStackItem.h
+++ b/Engine/Plugins/FX/Niagara/Source/NiagaraEditor/Public/ViewModels/Stack/NiagaraStackItem.h
@@ -27,6 +27,9 @@ public:
 
 	uint32 GetRecursiveStackIssuesCount() const;
 	EStackIssueSeverity GetHighestStackIssueSeverity() const;
+
+	virtual bool CanChangeEnabled() const { return false; }
+	virtual void SetIsEnabled(bool bInIsEnabled) { }
 	
 protected:
 	virtual void RefreshChildrenInternal(const TArray<UNiagaraStackEntry*>& CurrentChildren, TArray<UNiagaraStackEntry*>& NewChildren, TArray<FStackIssue>& NewIssues) override;
diff --git a/Engine/Plugins/FX/Niagara/Source/NiagaraEditor/Public/ViewModels/Stack/NiagaraStackItemGroup.h b/Engine/Plugins/FX/Niagara/Source/NiagaraEditor/Public/ViewModels/Stack/NiagaraStackItemGroup.h
--- a/Engine/Plugins/FX/Niagara/Source/NiagaraEditor/Public/ViewModels/Stack/NiagaraStackItemGroup.h
+++ b/Engine/Plugins/FX/Niagara/Source/NiagaraEditor/Public/ViewModels/Stack/NiagaraStackItemGroup.h
@@ -7,6 +7,8 @@
 #include "Layout/Visibility.h"
 #include "NiagaraStackItemGroup.generated.h"
 
+class FNiagaraEmitterHandleViewModel;
+
 UCLASS()
 class NIAGARAEDITOR_API UNiagaraStackItemGroup : public UNiagaraStackEntry
 {
@@ -23,6 +25,8 @@ public:
 	virtual bool CanDelete() const { return false; }
 	virtual bool Delete() { return false; }
 
+	virtual bool GetIsEnabled() const override;
+
 	INiagaraStackItemGroupAddUtilities* GetAddUtilities() const;
 
 	uint32 GetRecursiveStackIssuesCount() const;
@@ -47,4 +51,6 @@ private:
 	mutable TOptional<uint32> RecursiveStackIssuesCount;
 	/** The highest severity of issues along this entry's tree. */
 	mutable TOptional<EStackIssueSeverity> HighestIssueSeverity;
+
+	TSharedPtr<FNiagaraEmitterHandleViewModel> OwningEmitterHandleViewModel;
 };
diff --git a/Engine/Plugins/FX/Niagara/Source/NiagaraEditor/Public/ViewModels/Stack/NiagaraStackModuleItem.h b/Engine/Plugins/FX/Niagara/Source/NiagaraEditor/Public/ViewModels/Stack/NiagaraStackModuleItem.h
--- a/Engine/Plugins/FX/Niagara/Source/NiagaraEditor/Public/ViewModels/Stack/NiagaraStackModuleItem.h
+++ b/Engine/Plugins/FX/Niagara/Source/NiagaraEditor/Public/ViewModels/Stack/NiagaraStackModuleItem.h
@@ -37,8 +37,9 @@ public:
 	bool CanRefresh() const;
 	void Refresh();
 
+	virtual bool CanChangeEnabled() const override { return true; }
 	virtual bool GetIsEnabled() const override;
-	void SetIsEnabled(bool bInIsEnabled);
+	virtual void SetIsEnabled(bool bInIsEnabled) override;
 
 	void Delete();
 
diff --git a/Engine/Plugins/FX/Niagara/Source/NiagaraEditor/Public/ViewModels/Stack/NiagaraStackRendererItem.h b/Engine/Plugins/FX/Niagara/Source/NiagaraEditor/Public/ViewModels/Stack/NiagaraStackRendererItem.h
--- a/Engine/Plugins/FX/Niagara/Source/NiagaraEditor/Public/ViewModels/Stack/NiagaraStackRendererItem.h
+++ b/Engine/Plugins/FX/Niagara/Source/NiagaraEditor/Public/ViewModels/Stack/NiagaraStackRendererItem.h
@@ -34,8 +34,9 @@ public:
 
 	void ResetToBase();
 
-	bool GetIsEnabled() const;
-	void SetIsEnabled(bool bInIsEnabled);
+	virtual bool CanChangeEnabled() const override { return true; }
+	virtual bool GetIsEnabled() const override;
+	virtual void SetIsEnabled(bool bInIsEnabled) override;
 
 	static TArray<FNiagaraVariable> GetMissingVariables(UNiagaraRendererProperties* RendererProperties, UNiagaraEmitter* Emitter);
 	static bool AddMissingVariable(UNiagaraEmitter* Emitter, const FNiagaraVariable& Variable);
diff --git a/Engine/Plugins/FX/Niagara/Source/NiagaraEditorWidgets/Private/NiagaraEditorWidgetsModule.cpp b/Engine/Plugins/FX/Niagara/Source/NiagaraEditorWidgets/Private/NiagaraEditorWidgetsModule.cpp
--- a/Engine/Plugins/FX/Niagara/Source/NiagaraEditorWidgets/Private/NiagaraEditorWidgetsModule.cpp
+++ b/Engine/Plugins/FX/Niagara/Source/NiagaraEditorWidgets/Private/NiagaraEditorWidgetsModule.cpp
@@ -7,12 +7,12 @@
 #include "DetailCustomizations/NiagaraDataInterfaceDetails.h"
 #include "DetailCustomizations/NiagaraDataInterfaceSkeletalMeshDetails.h"
 #include "ViewModels/NiagaraSystemViewModel.h"
+#include "SNiagaraOverviewGraph.h"
+#include "NiagaraEditorWidgetsUtilities.h"
 
 #include "Modules/ModuleManager.h"
 #include "PropertyEditorModule.h"
 
-#include "SNiagaraOverviewGraph.h"
-
 IMPLEMENT_MODULE(FNiagaraEditorWidgetsModule, NiagaraEditorWidgets);
 
 FNiagaraStackCurveEditorOptions::FNiagaraStackCurveEditorOptions()
@@ -151,12 +151,17 @@ TSharedRef<FNiagaraStackCurveEditorOptions> FNiagaraEditorWidgetsModule::GetOrCr
 	return *StackCurveEditorOptions;
 }
 
-TSharedRef<SWidget> FNiagaraEditorWidgetsModule::FNiagaraEditorWidgetProvider::CreateStackView(UNiagaraStackViewModel& StackViewModel)
+TSharedRef<SWidget> FNiagaraEditorWidgetsModule::FNiagaraEditorWidgetProvider::CreateStackView(UNiagaraStackViewModel& StackViewModel) const
 {
 	return SNew(SNiagaraStack, &StackViewModel);
 }
 
-TSharedRef<SWidget> FNiagaraEditorWidgetsModule::FNiagaraEditorWidgetProvider::CreateSystemOverview(TSharedRef<FNiagaraSystemViewModel> SystemViewModel)
+TSharedRef<SWidget> FNiagaraEditorWidgetsModule::FNiagaraEditorWidgetProvider::CreateSystemOverview(TSharedRef<FNiagaraSystemViewModel> SystemViewModel) const
 {
 	return SNew(SNiagaraOverviewGraph, SystemViewModel->GetOverviewGraphViewModel().ToSharedRef());
 }
+
+FLinearColor FNiagaraEditorWidgetsModule::FNiagaraEditorWidgetProvider::GetColorForExecutionCategory(FName ExecutionCategory) const
+{
+	return FNiagaraEditorWidgetsStyle::Get().GetColor(FNiagaraStackEditorWidgetsUtilities::GetIconColorNameForExecutionCategory(ExecutionCategory));
+}
diff --git a/Engine/Plugins/FX/Niagara/Source/NiagaraEditorWidgets/Private/NiagaraEditorWidgetsStyle.cpp b/Engine/Plugins/FX/Niagara/Source/NiagaraEditorWidgets/Private/NiagaraEditorWidgetsStyle.cpp
--- a/Engine/Plugins/FX/Niagara/Source/NiagaraEditorWidgets/Private/NiagaraEditorWidgetsStyle.cpp
+++ b/Engine/Plugins/FX/Niagara/Source/NiagaraEditorWidgets/Private/NiagaraEditorWidgetsStyle.cpp
@@ -9,6 +9,7 @@
 #include "Styling/SlateTypes.h"
 #include "Styling/CoreStyle.h"
 #include "Interfaces/IPluginManager.h"
+#include "Classes/EditorStyleSettings.h"
 
 TSharedPtr< FSlateStyleSet > FNiagaraEditorWidgetsStyle::NiagaraEditorWidgetsStyleInstance = NULL;
  
@@ -104,7 +105,9 @@ TSharedRef< FSlateStyleSet > FNiagaraEditorWidgetsStyle::Create()
 
 	FSlateFontInfo SystemOverviewGroupHeaderFont = DEFAULT_FONT("Bold", 9);
 	FTextBlockStyle SystemOverviewGroupHeaderText = FTextBlockStyle(NormalText)
-		.SetFont(SystemOverviewGroupHeaderFont);
+		.SetFont(SystemOverviewGroupHeaderFont)
+		.SetShadowOffset(FVector2D(0, 1))
+		.SetShadowColorAndOpacity(FLinearColor(0, 0, 0, 0.9f));;
 	Style->Set("NiagaraEditor.SystemOverview.GroupHeaderText", SystemOverviewGroupHeaderText);
 
 	FSlateFontInfo SystemOverviewItemFont = DEFAULT_FONT("Regular", 10);
@@ -112,6 +115,15 @@ TSharedRef< FSlateStyleSet > FNiagaraEditorWidgetsStyle::Create()
 		.SetFont(SystemOverviewItemFont);
 	Style->Set("NiagaraEditor.SystemOverview.ItemText", SystemOverviewItemText);
 
+	Style->Set("NiagaraEditor.SystemOverview.BackgroundColor", FLinearColor(FColor(80, 80, 80)));
+	Style->Set("NiagaraEditor.SystemOverview.CheckBoxColor", FLinearColor(FColor(160, 160, 160)));
+	Style->Set("NiagaraEditor.SystemOverview.CheckBoxBorder", new BOX_CORE_BRUSH("Common/GroupBorderLight", FMargin(4.0f / 16.0f)));
+	Style->Set("NiagaraEditor.SystemOverview.NodeBackgroundBorder", new BOX_PLUGIN_BRUSH("Icons/SystemOverviewNodeBackground", FMargin(1.0f / 4.0f)));
+
+	const FTableRowStyle& NormalTableRowStyle = FEditorStyle::Get().GetWidgetStyle<FTableRowStyle>("TableView.Row");
+	Style->Set("NiagaraEditor.SystemOverview.TableViewRow", FTableRowStyle(NormalTableRowStyle)
+		.SetInactiveBrush(IMAGE_CORE_BRUSH("Common/Selection", Icon8x8, GetDefault<UEditorStyleSettings>()->GetSubduedSelectionColor())));
+
 	Style->Set("NiagaraEditor.Stack.Group.BackgroundColor", FLinearColor(FColor(96, 96, 96)));
 	Style->Set("NiagaraEditor.Stack.Item.HeaderBackgroundColor", FLinearColor(FColor(48, 48, 48)));
 	Style->Set("NiagaraEditor.Stack.Item.ContentBackgroundColor", FLinearColor(FColor(62, 62, 62)));
@@ -124,7 +136,7 @@ TSharedRef< FSlateStyleSet > FNiagaraEditorWidgetsStyle::Create()
 
 	Style->Set("NiagaraEditor.Stack.ForegroundColor", FLinearColor(FColor(220, 220, 220)));
 	Style->Set("NiagaraEditor.Stack.GroupForegroundColor", FLinearColor(FColor(220, 220, 220)));
-	Style->Set("NiagaraEditor.Stack.FlatButtonColor", FLinearColor(FColor(191, 191, 191)));
+	Style->Set("NiagaraEditor.Stack.FlatButtonColor", FLinearColor(FColor(205, 205, 205)));
 
 	Style->Set("NiagaraEditor.Stack.AccentColor.System", FLinearColor(FColor(67, 105, 124)));
 	Style->Set("NiagaraEditor.Stack.AccentColor.Emitter", FLinearColor(FColor(126, 87, 67)));
diff --git a/Engine/Plugins/FX/Niagara/Source/NiagaraEditorWidgets/Private/SNiagaraOverviewStack.cpp b/Engine/Plugins/FX/Niagara/Source/NiagaraEditorWidgets/Private/SNiagaraOverviewStack.cpp
--- a/Engine/Plugins/FX/Niagara/Source/NiagaraEditorWidgets/Private/SNiagaraOverviewStack.cpp
+++ b/Engine/Plugins/FX/Niagara/Source/NiagaraEditorWidgets/Private/SNiagaraOverviewStack.cpp
@@ -18,6 +18,8 @@
 #include "EditorStyleSet.h"
 #include "Framework/Application/SlateApplication.h"
 
+#define LOCTEXT_NAMESPACE "NiagaraOverviewStack"
+
 class SNiagaraSystemOverviewEntryListRow : public STableRow<UNiagaraStackEntry*>
 {
 	SLATE_BEGIN_ARGS(SNiagaraSystemOverviewEntryListRow)
@@ -28,6 +30,8 @@ class SNiagaraSystemOverviewEntryListRow : public STableRow<UNiagaraStackEntry*>
 	void Construct(const FArguments& InArgs, UNiagaraStackEntry* StackEntry, const TSharedRef<STableViewBase>& InOwnerTableView)
 	{
 		FSlateColor IconColor = FNiagaraEditorWidgetsStyle::Get().GetColor(FNiagaraStackEditorWidgetsUtilities::GetColorNameForExecutionCategory(StackEntry->GetExecutionCategoryName()));
+		ItemBackgroundColor = FNiagaraEditorWidgetsStyle::Get().GetColor("NiagaraEditor.Stack.Item.HeaderBackgroundColor");
+		DisabledItemBackgroundColor = ItemBackgroundColor + FLinearColor(.02f, .02f, .02f, 0.0f);
 
 		TSharedPtr<SWidget> WrappedContent;
 		if (StackEntry->IsA<UNiagaraStackItem>())
@@ -35,12 +39,12 @@ class SNiagaraSystemOverviewEntryListRow : public STableRow<UNiagaraStackEntry*>
 			WrappedContent =
 				SNew(SBorder)
 				.BorderImage(FEditorStyle::GetBrush("WhiteBrush"))
-				.BorderBackgroundColor(FNiagaraEditorWidgetsStyle::Get().GetColor("NiagaraEditor.Stack.Item.HeaderBackgroundColor"))
+				.BorderBackgroundColor(this, &SNiagaraSystemOverviewEntryListRow::GetItemBackgroundColor, StackEntry)
 				.Padding(0)
 				[
 					SNew(SBorder)
 					.BorderImage(this, &SNiagaraSystemOverviewEntryListRow::GetBorder)
-					.Padding(FMargin(8, 4, 8, 4))
+					.Padding(FMargin(6, 4, 3, 4))
 					[
 						InArgs._Content.Widget
 					]
@@ -55,7 +59,7 @@ class SNiagaraSystemOverviewEntryListRow : public STableRow<UNiagaraStackEntry*>
 				[
 					SNew(SBorder)
 					.BorderImage(this, &SNiagaraSystemOverviewEntryListRow::GetBorder)
-					.Padding(FMargin(2, 6, 8, 2))
+					.Padding(FMargin(2, 4, 7, 4))
 					[
 						InArgs._Content.Widget
 					]
@@ -63,39 +67,83 @@ class SNiagaraSystemOverviewEntryListRow : public STableRow<UNiagaraStackEntry*>
 		}
 
 		STableRow<UNiagaraStackEntry*>::Construct(STableRow<UNiagaraStackEntry*>::FArguments()
+			.Style(FNiagaraEditorWidgetsStyle::Get(), "NiagaraEditor.SystemOverview.TableViewRow")
 		[
 			SNew(SBorder)
-			.BorderImage(FEditorStyle::GetBrush("WhiteBrush"))
-			.BorderBackgroundColor(FNiagaraEditorWidgetsStyle::Get().GetColor("NiagaraEditor.Stack.Item.ContentBackgroundColor"))
-			.Padding(0)
+			.BorderImage(this, &SNiagaraSystemOverviewEntryListRow::GetBorder)
+			.Padding(FMargin(5, 2, 5, 2))
 			[
-				SNew(SHorizontalBox)
-				+ SHorizontalBox::Slot()
-				.AutoWidth()
-				.VAlign(VAlign_Fill)
-				[
-					SNew(SBorder)
-					.BorderImage(FEditorStyle::GetBrush("WhiteBrush"))
-					.BorderBackgroundColor(FNiagaraEditorWidgetsStyle::Get().GetColor(FNiagaraStackEditorWidgetsUtilities::GetIconColorNameForExecutionCategory(StackEntry->GetExecutionCategoryName())))
-					.Padding(0)
-					[
-						SNew(SBox)
-						.WidthOverride(5)
-					]
-				]
-				+ SHorizontalBox::Slot()
-				[
-					SNew(SBorder)
-					.BorderImage(this, &SNiagaraSystemOverviewEntryListRow::GetBorder)
-					.Padding(FMargin(7, 0, 7, 4))
-					[
-						WrappedContent.ToSharedRef()
-					]
-				]
+				WrappedContent.ToSharedRef()
 			]
 		],
 		InOwnerTableView);
 	}
+
+private:
+	FSlateColor GetItemBackgroundColor(UNiagaraStackEntry* Entry) const
+	{
+		return Entry->GetIsEnabled() && Entry->GetOwnerIsEnabled() ? ItemBackgroundColor : DisabledItemBackgroundColor;
+	}
+
+private:
+	FLinearColor ItemBackgroundColor;
+	FLinearColor DisabledItemBackgroundColor;
+};
+
+class SNiagaraSystemOverviewEnabledCheckBox : public SCompoundWidget
+{
+public:
+	DECLARE_DELEGATE_OneParam(FOnCheckedChanged, bool /*bIsChecked*/);
+
+	SLATE_BEGIN_ARGS(SNiagaraSystemOverviewEnabledCheckBox) {}
+		SLATE_ATTRIBUTE(bool, IsChecked)
+		SLATE_EVENT(FOnCheckedChanged, OnCheckedChanged)
+	SLATE_END_ARGS();
+
+	void Construct(const FArguments& InArgs)
+	{
+		IsChecked = InArgs._IsChecked;
+		OnCheckedChanged = InArgs._OnCheckedChanged;
+
+		ChildSlot
+		[
+			SNew(SBorder)
+			.BorderImage(FNiagaraEditorWidgetsStyle::Get().GetBrush("NiagaraEditor.SystemOverview.CheckBoxBorder"))
+			.BorderBackgroundColor(FNiagaraEditorWidgetsStyle::Get().GetColor("NiagaraEditor.Stack.Item.HeaderBackgroundColor"))
+			.Padding(FMargin(0))
+			[
+				SNew(SButton)
+				.ButtonStyle(FEditorStyle::Get(), "HoverHintOnly")
+				.OnClicked(this, &SNiagaraSystemOverviewEnabledCheckBox::OnButtonClicked)
+				.ToolTipText(LOCTEXT("EnableCheckBoxToolTip", "Enable or disable this item."))
+				.ContentPadding(FMargin(3, 2, 3, 2))
+				[
+					SNew(STextBlock)
+					.Font(FEditorStyle::Get().GetFontStyle("FontAwesome.10"))
+					.Text(this, &SNiagaraSystemOverviewEnabledCheckBox::GetButtonText)
+					.ColorAndOpacity(FNiagaraEditorWidgetsStyle::Get().GetColor("NiagaraEditor.SystemOverview.CheckBoxColor"))
+				]
+			]
+		];
+	}
+
+private:
+	FReply OnButtonClicked()
+	{
+		OnCheckedChanged.ExecuteIfBound(IsChecked.IsBound() && !IsChecked.Get());
+		return FReply::Handled();
+	}
+
+	FText GetButtonText() const
+	{
+		return IsChecked.IsBound() && IsChecked.Get()
+			? LOCTEXT("CheckedText", "\xf14a")
+			: LOCTEXT("UncheckedText", "\xf0c8");
+	}
+
+private:
+	TAttribute<bool> IsChecked;
+	FOnCheckedChanged OnCheckedChanged;
 };
 
 void SNiagaraOverviewStack::Construct(const FArguments& InArgs, UNiagaraStackViewModel& InStackViewModel, UNiagaraSystemSelectionViewModel& InOverviewSelectionViewModel)
@@ -189,9 +237,25 @@ TSharedRef<ITableRow> SNiagaraOverviewStack::OnGenerateRowForEntry(UNiagaraStack
 	TSharedPtr<SWidget> Content;
 	if (Item->IsA<UNiagaraStackItem>())
 	{
-		Content = SNew(STextBlock)
-			.TextStyle(FNiagaraEditorWidgetsStyle::Get(), "NiagaraEditor.SystemOverview.ItemText")
-			.Text_UObject(Item, &UNiagaraStackEntry::GetDisplayName);
+		UNiagaraStackItem* StackItem = CastChecked<UNiagaraStackItem>(Item);
+		Content = SNew(SHorizontalBox)
+			+ SHorizontalBox::Slot()
+			[
+				SNew(STextBlock)
+				.TextStyle(FNiagaraEditorWidgetsStyle::Get(), "NiagaraEditor.SystemOverview.ItemText")
+				.Text_UObject(Item, &UNiagaraStackEntry::GetDisplayName)
+				.IsEnabled_UObject(Item, &UNiagaraStackEntry::GetIsEnabledAndOwnerIsEnabled)
+			]
+			+ SHorizontalBox::Slot()
+			.AutoWidth()
+			.Padding(3, 0, 0, 0)
+			[
+				SNew(SNiagaraSystemOverviewEnabledCheckBox)
+				.Visibility(this, &SNiagaraOverviewStack::GetEnabledCheckBoxVisibility, StackItem)
+				.IsEnabled_UObject(StackItem, &UNiagaraStackEntry::GetOwnerIsEnabled)
+				.IsChecked_UObject(StackItem, &UNiagaraStackEntry::GetIsEnabled)
+				.OnCheckedChanged_UObject(StackItem, &UNiagaraStackItem::SetIsEnabled)
+			];
 	}
 	else
 	{
@@ -202,15 +266,16 @@ TSharedRef<ITableRow> SNiagaraOverviewStack::OnGenerateRowForEntry(UNiagaraStack
 			.AutoWidth()
 			[
 				SNew(SImage)
-				//.Visibility(this, &SNiagaraStackTableRow::GetExecutionCategoryIconVisibility)
 				.Image(FNiagaraEditorWidgetsStyle::Get().GetBrush(FNiagaraStackEditorWidgetsUtilities::GetIconNameForExecutionSubcategory(Item->GetExecutionSubcategoryName(), true)))
 				.ColorAndOpacity(FNiagaraEditorWidgetsStyle::Get().GetColor(FNiagaraStackEditorWidgetsUtilities::GetIconColorNameForExecutionCategory(Item->GetExecutionCategoryName())))
+				.IsEnabled_UObject(Item, &UNiagaraStackEntry::GetIsEnabledAndOwnerIsEnabled)
 			]
 			+ SHorizontalBox::Slot()
 			[
 				SNew(STextBlock)
 				.TextStyle(FNiagaraEditorWidgetsStyle::Get(), "NiagaraEditor.SystemOverview.GroupHeaderText")
 				.Text_UObject(Item, &UNiagaraStackEntry::GetDisplayName)
+				.IsEnabled_UObject(Item, &UNiagaraStackEntry::GetIsEnabledAndOwnerIsEnabled)
 			];
 	}
 
@@ -221,6 +286,11 @@ TSharedRef<ITableRow> SNiagaraOverviewStack::OnGenerateRowForEntry(UNiagaraStack
 	];
 }
 
+EVisibility SNiagaraOverviewStack::GetEnabledCheckBoxVisibility(UNiagaraStackItem* Item) const
+{
+	return Item->CanChangeEnabled() ? EVisibility::Visible : EVisibility::Collapsed;
+}
+
 void SNiagaraOverviewStack::OnSelectionChanged(UNiagaraStackEntry* InNewSelection, ESelectInfo::Type SelectInfo)
 {
 	if (bUpdatingStackSelectionFromOverviewSelection == false)
@@ -274,4 +344,6 @@ void SNiagaraOverviewStack::SystemSelectionChanged(UNiagaraSystemSelectionViewMo
 			EntryListView->SetItemSelection(EntryToDeselect, false);
 		}
 	}
-}
\ No newline at end of file
+}
+
+#undef LOCTEXT_NAMESPACE
\ No newline at end of file
diff --git a/Engine/Plugins/FX/Niagara/Source/NiagaraEditorWidgets/Private/SNiagaraOverviewStack.h b/Engine/Plugins/FX/Niagara/Source/NiagaraEditorWidgets/Private/SNiagaraOverviewStack.h
--- a/Engine/Plugins/FX/Niagara/Source/NiagaraEditorWidgets/Private/SNiagaraOverviewStack.h
+++ b/Engine/Plugins/FX/Niagara/Source/NiagaraEditorWidgets/Private/SNiagaraOverviewStack.h
@@ -9,6 +9,7 @@
 #include "UObject/ObjectKey.h"
 
 class UNiagaraStackEntry;
+class UNiagaraStackItem;
 class UNiagaraStackViewModel;
 class UNiagaraSystemSelectionViewModel;
 
@@ -34,6 +35,8 @@ private:
 
 	TSharedRef<ITableRow> OnGenerateRowForEntry(UNiagaraStackEntry* Item, const TSharedRef<STableViewBase>& OwnerTable);
 
+	EVisibility GetEnabledCheckBoxVisibility(UNiagaraStackItem* Item) const;
+
 	void OnSelectionChanged(UNiagaraStackEntry* InNewSelection, ESelectInfo::Type SelectInfo);
 
 	void SystemSelectionChanged(UNiagaraSystemSelectionViewModel::ESelectionChangeSource SelectionChangeSource);
diff --git a/Engine/Plugins/FX/Niagara/Source/NiagaraEditorWidgets/Private/SNiagaraOverviewStackNode.cpp b/Engine/Plugins/FX/Niagara/Source/NiagaraEditorWidgets/Private/SNiagaraOverviewStackNode.cpp
--- a/Engine/Plugins/FX/Niagara/Source/NiagaraEditorWidgets/Private/SNiagaraOverviewStackNode.cpp
+++ b/Engine/Plugins/FX/Niagara/Source/NiagaraEditorWidgets/Private/SNiagaraOverviewStackNode.cpp
@@ -6,9 +6,10 @@
 #include "ViewModels/NiagaraEmitterHandleViewModel.h"
 #include "SNiagaraOverviewStack.h"
 #include "NiagaraEditorModule.h"
-#include "Widgets/Layout/SBox.h"
+#include "NiagaraEditorWidgetsStyle.h"
 
 #include "Modules/ModuleManager.h"
+#include "Widgets/Layout/SBox.h"
 
 void SNiagaraOverviewStackNode::Construct(const FArguments& InArgs, UNiagaraOverviewNode* InNode)
 {
@@ -26,10 +27,10 @@ void SNiagaraOverviewStackNode::Construct(const FArguments& InArgs, UNiagaraOver
 			}
 			else
 			{
-				TSharedPtr<FNiagaraEmitterHandleViewModel> EmitterHandleviewModel = OwningSystemViewModel->GetEmitterHandleViewModelById(OverviewStackNode->GetEmitterHandleGuid());
-				if (EmitterHandleviewModel.IsValid())
+				EmitterHandleViewModelWeak = OwningSystemViewModel->GetEmitterHandleViewModelById(OverviewStackNode->GetEmitterHandleGuid());
+				if (EmitterHandleViewModelWeak.IsValid())
 				{
-					StackViewModel = EmitterHandleviewModel->GetEmitterStackViewModel();
+					StackViewModel = EmitterHandleViewModelWeak.Pin()->GetEmitterStackViewModel();
 				}
 			}
 			OverviewSelectionViewModel = OwningSystemViewModel->GetSelectionViewModel();
@@ -38,6 +39,26 @@ void SNiagaraOverviewStackNode::Construct(const FArguments& InArgs, UNiagaraOver
 	UpdateGraphNode();
 }
 
+TSharedRef<SWidget> SNiagaraOverviewStackNode::CreateTitleWidget(TSharedPtr<SNodeTitle> NodeTitle)
+{
+	TSharedRef<SWidget> DefaultTitle = SGraphNode::CreateTitleWidget(NodeTitle);
+	
+	return SNew(SHorizontalBox)
+		+ SHorizontalBox::Slot()
+		.AutoWidth()
+		.Padding(0, 0, 5, 0)
+		[
+			SNew(SCheckBox)
+			.Visibility(this, &SNiagaraOverviewStackNode::GetEnabledCheckBoxVisibility)
+			.IsChecked(this, &SNiagaraOverviewStackNode::GetEnabledCheckState)
+			.OnCheckStateChanged(this, &SNiagaraOverviewStackNode::OnEnabledCheckStateChanged)
+		]
+		+ SHorizontalBox::Slot()
+		[
+			DefaultTitle
+		];
+}
+
 TSharedRef<SWidget> SNiagaraOverviewStackNode::CreateNodeContentArea()
 {
 	TSharedPtr<SWidget> ContentWidget;
@@ -56,10 +77,10 @@ TSharedRef<SWidget> SNiagaraOverviewStackNode::CreateNodeContentArea()
 
 	// NODE CONTENT AREA
 	return SNew(SBorder)
-		.BorderImage( FEditorStyle::GetBrush("NoBorder") )
+		.BorderImage(FEditorStyle::GetBrush("NoBorder"))
 		.HAlign(HAlign_Fill)
 		.VAlign(VAlign_Fill)
-		.Padding( FMargin(0,3) )
+		.Padding(FMargin(2, 3, 2, 2))
 		[
 			SNew(SHorizontalBox)
 			+SHorizontalBox::Slot()
@@ -71,7 +92,15 @@ TSharedRef<SWidget> SNiagaraOverviewStackNode::CreateNodeContentArea()
 			]
 			+SHorizontalBox::Slot()
 			[
-				ContentWidget.ToSharedRef()
+				SNew(SBorder)
+				.BorderImage(FNiagaraEditorWidgetsStyle::Get().GetBrush("NiagaraEditor.SystemOverview.NodeBackgroundBorder"))
+				.BorderBackgroundColor(FNiagaraEditorWidgetsStyle::Get().GetColor("NiagaraEditor.SystemOverview.BackgroundColor"))
+				.HAlign(HAlign_Fill)
+				.VAlign(VAlign_Fill)
+				.Padding(FMargin(0, 0, 0, 4))
+				[
+					ContentWidget.ToSharedRef()
+				]
 			]
 			+SHorizontalBox::Slot()
 			.AutoWidth()
@@ -81,4 +110,23 @@ TSharedRef<SWidget> SNiagaraOverviewStackNode::CreateNodeContentArea()
 				SAssignNew(RightNodeBox, SVerticalBox)
 			]
 		];
+}
+
+EVisibility SNiagaraOverviewStackNode::GetEnabledCheckBoxVisibility() const
+{
+	return EmitterHandleViewModelWeak.IsValid() ? EVisibility::Visible : EVisibility::Collapsed;
+}
+
+ECheckBoxState SNiagaraOverviewStackNode::GetEnabledCheckState() const
+{
+	return EmitterHandleViewModelWeak.IsValid() && EmitterHandleViewModelWeak.Pin()->GetIsEnabled() ? ECheckBoxState::Checked : ECheckBoxState::Unchecked;
+}
+
+void SNiagaraOverviewStackNode::OnEnabledCheckStateChanged(ECheckBoxState InCheckState)
+{
+	TSharedPtr<FNiagaraEmitterHandleViewModel> EmitterHandleViewModel = EmitterHandleViewModelWeak.Pin();
+	if (EmitterHandleViewModel.IsValid())
+	{
+		EmitterHandleViewModel->SetIsEnabled(InCheckState == ECheckBoxState::Checked);
+	}
 }
\ No newline at end of file
diff --git a/Engine/Plugins/FX/Niagara/Source/NiagaraEditorWidgets/Private/SNiagaraOverviewStackNode.h b/Engine/Plugins/FX/Niagara/Source/NiagaraEditorWidgets/Private/SNiagaraOverviewStackNode.h
--- a/Engine/Plugins/FX/Niagara/Source/NiagaraEditorWidgets/Private/SNiagaraOverviewStackNode.h
+++ b/Engine/Plugins/FX/Niagara/Source/NiagaraEditorWidgets/Private/SNiagaraOverviewStackNode.h
@@ -7,6 +7,7 @@
 class UNiagaraOverviewNode;
 class UNiagaraStackViewModel;
 class UNiagaraSystemSelectionViewModel;
+class FNiagaraEmitterHandleViewModel;
 
 class SNiagaraOverviewStackNode : public SGraphNode
 {
@@ -18,12 +19,17 @@ public:
 	void Construct(const FArguments& InArgs, UNiagaraOverviewNode* InNode);
 
 protected:
+	virtual TSharedRef<SWidget> CreateTitleWidget(TSharedPtr<SNodeTitle> NodeTitle) override;
 	virtual TSharedRef<SWidget> CreateNodeContentArea() override;
 
 private:
+	EVisibility GetEnabledCheckBoxVisibility() const;
+	ECheckBoxState GetEnabledCheckState() const;
+	void OnEnabledCheckStateChanged(ECheckBoxState InCheckState);
 
 private:
 	UNiagaraOverviewNode* OverviewStackNode;
 	UNiagaraStackViewModel* StackViewModel;
 	UNiagaraSystemSelectionViewModel* OverviewSelectionViewModel;
+	TWeakPtr<FNiagaraEmitterHandleViewModel> EmitterHandleViewModelWeak;
 };
\ No newline at end of file
diff --git a/Engine/Plugins/FX/Niagara/Source/NiagaraEditorWidgets/Private/SNiagaraStack.cpp b/Engine/Plugins/FX/Niagara/Source/NiagaraEditorWidgets/Private/SNiagaraStack.cpp
--- a/Engine/Plugins/FX/Niagara/Source/NiagaraEditorWidgets/Private/SNiagaraStack.cpp
+++ b/Engine/Plugins/FX/Niagara/Source/NiagaraEditorWidgets/Private/SNiagaraStack.cpp
@@ -254,28 +254,23 @@ void SNiagaraStack::Construct(const FArguments& InArgs, UNiagaraStackViewModel*
 		SNew(SVerticalBox)
 		+ SVerticalBox::Slot()
 		.AutoHeight()
-		.Padding(0, 0, 0, 3)
+		.Padding(1, 1, 1, 4)
 		[
 			SNew(SBorder)
 			.BorderImage(FEditorStyle::GetBrush("ToolPanel.GroupBorder"))
-			.Padding(3)
+			.Padding(5)
 			[
 				ConstructHeaderWidget()
 			]
 		]
 		+ SVerticalBox::Slot()
-		.Padding(0)
+		.Padding(1)
 		[
-			SNew(SBorder)
-			.BorderImage(FEditorStyle::GetBrush("ToolPanel.GroupBorder"))
-			.Padding(5)
-			[
-				SAssignNew(StackTree, STreeView<UNiagaraStackEntry*>)
-				.OnGenerateRow(this, &SNiagaraStack::OnGenerateRowForStackItem)
-				.OnGetChildren(this, &SNiagaraStack::OnGetChildren)
-				.TreeItemsSource(&StackViewModel->GetRootEntries())
-				.OnTreeViewScrolled(this, &SNiagaraStack::StackTreeScrolled)
-			]
+			SAssignNew(StackTree, STreeView<UNiagaraStackEntry*>)
+			.OnGenerateRow(this, &SNiagaraStack::OnGenerateRowForStackItem)
+			.OnGetChildren(this, &SNiagaraStack::OnGetChildren)
+			.TreeItemsSource(&StackViewModel->GetRootEntries())
+			.OnTreeViewScrolled(this, &SNiagaraStack::StackTreeScrolled)
 		]
 	];
 
@@ -824,7 +819,7 @@ TSharedRef<SNiagaraStackTableRow> SNiagaraStack::ConstructContainerForItem(UNiag
 		bShowExecutionCategoryIcon = false;
 		break;
 	case UNiagaraStackEntry::EStackRowStyle::GroupHeader:
-		ContentPadding = FMargin(LeftContentPadding, 4, 0, 0);
+		ContentPadding = FMargin(LeftContentPadding, 2, 0, 2);
 		ItemBackgroundColor = FLinearColor::Transparent;
 		ItemForegroundColor = FNiagaraEditorWidgetsStyle::Get().GetColor("NiagaraEditor.Stack.GroupForegroundColor");
 		bIsCategoryIconHighlighted = true;
@@ -977,7 +972,8 @@ SNiagaraStack::FRowWidgets SNiagaraStack::ConstructNameAndValueWidgetsForItem(UN
 			.ToolTipText_UObject(Item, &UNiagaraStackEntry::GetTooltipText)
 			.Text_UObject(Item, &UNiagaraStackEntry::GetDisplayName)
 			.ColorAndOpacity(this, &SNiagaraStack::GetTextColorForItem, Item)
-			.HighlightText_UObject(StackViewModel, &UNiagaraStackViewModel::GetCurrentSearchText),
+			.HighlightText_UObject(StackViewModel, &UNiagaraStackViewModel::GetCurrentSearchText)
+			.IsEnabled_UObject(Item, &UNiagaraStackEntry::GetOwnerIsEnabled),
 			SNullWidget::NullWidget);
 	}
 	else if (Item->IsA<UNiagaraStackModuleItemOutput>())
@@ -989,12 +985,14 @@ SNiagaraStack::FRowWidgets SNiagaraStack::ConstructNameAndValueWidgetsForItem(UN
 			.ToolTipText_UObject(Item, &UNiagaraStackEntry::GetTooltipText)
 			.Text_UObject(Item, &UNiagaraStackEntry::GetDisplayName)
 			.ColorAndOpacity(this, &SNiagaraStack::GetTextColorForItem, Item)
-			.HighlightText_UObject(StackViewModel, &UNiagaraStackViewModel::GetCurrentSearchText),
+			.HighlightText_UObject(StackViewModel, &UNiagaraStackViewModel::GetCurrentSearchText)
+			.IsEnabled_UObject(Item, &UNiagaraStackEntry::GetOwnerIsEnabled),
 			SNew(STextBlock)
 			.TextStyle(FNiagaraEditorWidgetsStyle::Get(), "NiagaraEditor.Stack.ParameterText")
 			.Text_UObject(ModuleItemOutput, &UNiagaraStackModuleItemOutput::GetOutputParameterHandleText)
 			.ColorAndOpacity(this, &SNiagaraStack::GetTextColorForItem, Item)
-			.HighlightText_UObject(StackViewModel, &UNiagaraStackViewModel::GetCurrentSearchText));
+			.HighlightText_UObject(StackViewModel, &UNiagaraStackViewModel::GetCurrentSearchText)
+			.IsEnabled_UObject(Item, &UNiagaraStackEntry::GetOwnerIsEnabled));
 	}
 	else if (Item->IsA<UNiagaraStackFunctionInputCollection>() ||
 		Item->IsA<UNiagaraStackModuleItemOutputCollection>() ||
@@ -1006,17 +1004,21 @@ SNiagaraStack::FRowWidgets SNiagaraStack::ConstructNameAndValueWidgetsForItem(UN
 			.ToolTipText_UObject(Item, &UNiagaraStackEntry::GetTooltipText)
 			.Text_UObject(Item, &UNiagaraStackEntry::GetDisplayName)
 			.ColorAndOpacity(this, &SNiagaraStack::GetTextColorForItem, Item)
-			.HighlightText_UObject(StackViewModel, &UNiagaraStackViewModel::GetCurrentSearchText),
+			.HighlightText_UObject(StackViewModel, &UNiagaraStackViewModel::GetCurrentSearchText)
+			.IsEnabled_UObject(Item, &UNiagaraStackEntry::GetOwnerIsEnabled),
 			SNullWidget::NullWidget);
 	}
 	else if (Item->IsA<UNiagaraStackPropertyRow>())
 	{
 		UNiagaraStackPropertyRow* PropertyRow = CastChecked<UNiagaraStackPropertyRow>(Item);
 		FNodeWidgets PropertyRowWidgets = PropertyRow->GetDetailTreeNode()->CreateNodeWidgets();
+		TAttribute<bool> IsEnabled;
+		IsEnabled.BindUObject(Item, &UNiagaraStackEntry::GetIsEnabledAndOwnerIsEnabled);
 		if (PropertyRowWidgets.WholeRowWidget.IsValid())
 		{
 			Container->SetOverrideNameWidth(PropertyRowWidgets.WholeRowWidgetLayoutData.MinWidth, PropertyRowWidgets.WholeRowWidgetLayoutData.MaxWidth);
 			Container->SetOverrideNameAlignment(PropertyRowWidgets.WholeRowWidgetLayoutData.HorizontalAlignment, PropertyRowWidgets.WholeRowWidgetLayoutData.VerticalAlignment);
+			PropertyRowWidgets.WholeRowWidget->SetEnabled(IsEnabled);
 			return FRowWidgets(PropertyRowWidgets.WholeRowWidget.ToSharedRef());
 		}
 		else
@@ -1025,6 +1027,8 @@ SNiagaraStack::FRowWidgets SNiagaraStack::ConstructNameAndValueWidgetsForItem(UN
 			Container->SetOverrideNameAlignment(PropertyRowWidgets.NameWidgetLayoutData.HorizontalAlignment, PropertyRowWidgets.NameWidgetLayoutData.VerticalAlignment);
 			Container->SetOverrideValueWidth(PropertyRowWidgets.ValueWidgetLayoutData.MinWidth, PropertyRowWidgets.ValueWidgetLayoutData.MaxWidth);
 			Container->SetOverrideValueAlignment(PropertyRowWidgets.ValueWidgetLayoutData.HorizontalAlignment, PropertyRowWidgets.ValueWidgetLayoutData.VerticalAlignment);
+			PropertyRowWidgets.NameWidget->SetEnabled(IsEnabled);
+			PropertyRowWidgets.ValueWidget->SetEnabled(IsEnabled);
 			return FRowWidgets(PropertyRowWidgets.NameWidget.ToSharedRef(), PropertyRowWidgets.ValueWidget.ToSharedRef());
 		}
 	}
diff --git a/Engine/Plugins/FX/Niagara/Source/NiagaraEditorWidgets/Private/SNiagaraStackEntryWidget.cpp b/Engine/Plugins/FX/Niagara/Source/NiagaraEditorWidgets/Private/SNiagaraStackEntryWidget.cpp
--- a/Engine/Plugins/FX/Niagara/Source/NiagaraEditorWidgets/Private/SNiagaraStackEntryWidget.cpp
+++ b/Engine/Plugins/FX/Niagara/Source/NiagaraEditorWidgets/Private/SNiagaraStackEntryWidget.cpp
@@ -35,14 +35,23 @@ TSharedRef<SWidget> SNiagaraStackDisplayName::ConstructChildren()
 		.Text_UObject(StackEntry, &UNiagaraStackEntry::GetDisplayName)
 		.HighlightText_UObject(StackViewModel, &UNiagaraStackViewModel::GetCurrentSearchText)
 		.ColorAndOpacity(ColorAndOpacity)
-		.IsEnabled_UObject(StackEntry, &UNiagaraStackEntry::GetIsEnabled);
+		.IsEnabled(this, &SNiagaraStackDisplayName::GetIsEnabled);
 
-	if (StackViewModel->GetTopLevelViewModels().Num() == 1)
+	int32 NumTopLevelEmitters = 0;
+	for (const TSharedRef<UNiagaraStackViewModel::FTopLevelViewModel>& TopLevelViewModel : StackViewModel->GetTopLevelViewModels())
+	{
+		if (TopLevelViewModel->EmitterHandleViewModel.IsValid())
+		{
+			NumTopLevelEmitters++;
+		}
+	}
+
+	if (NumTopLevelEmitters <= 1)
 	{
 		TopLevelViewModelCountAtLastConstruction = 1;
 		return BaseNameWidget;
 	}
-	else if (StackViewModel->GetTopLevelViewModels().Num() > 1)
+	else
 	{
 		TopLevelViewModelCountAtLastConstruction = StackViewModel->GetTopLevelViewModels().Num();
 		TSharedPtr<UNiagaraStackViewModel::FTopLevelViewModel> TopLevelViewModel = StackViewModel->GetTopLevelViewModelForEntry(*StackEntry);
@@ -59,7 +68,7 @@ TSharedRef<SWidget> SNiagaraStackDisplayName::ConstructChildren()
 					.Text(this, &SNiagaraStackDisplayName::GetTopLevelDisplayName, TWeakPtr<UNiagaraStackViewModel::FTopLevelViewModel>(TopLevelViewModel))
 					.HighlightText_UObject(StackViewModel, &UNiagaraStackViewModel::GetCurrentSearchText)
 					.ColorAndOpacity(ColorAndOpacity)
-					.IsEnabled_UObject(StackEntry, &UNiagaraStackEntry::GetIsEnabled)
+					.IsEnabled(this, &SNiagaraStackDisplayName::GetIsEnabled)
 				]
 				+ SWrapBox::Slot()
 				[
@@ -97,6 +106,11 @@ void SNiagaraStackDisplayName::StackViewModelStructureChanged()
 	}
 }
 
+bool SNiagaraStackDisplayName::GetIsEnabled() const
+{
+	return StackEntry->GetOwnerIsEnabled() && StackEntry->GetIsEnabled();
+}
+
 FSlateColor SNiagaraStackEntryWidget::GetTextColorForSearch() const
 {
 	if (IsCurrentSearchMatch())
diff --git a/Engine/Plugins/FX/Niagara/Source/NiagaraEditorWidgets/Private/SNiagaraStackEntryWidget.h b/Engine/Plugins/FX/Niagara/Source/NiagaraEditorWidgets/Private/SNiagaraStackEntryWidget.h
--- a/Engine/Plugins/FX/Niagara/Source/NiagaraEditorWidgets/Private/SNiagaraStackEntryWidget.h
+++ b/Engine/Plugins/FX/Niagara/Source/NiagaraEditorWidgets/Private/SNiagaraStackEntryWidget.h
@@ -26,6 +26,8 @@ private:
 
 	void StackViewModelStructureChanged();
 
+	bool GetIsEnabled() const;
+
 private:
 	UNiagaraStackEntry* StackEntry;
 	UNiagaraStackViewModel* StackViewModel;
diff --git a/Engine/Plugins/FX/Niagara/Source/NiagaraEditorWidgets/Private/Stack/SNiagaraStackFunctionInputName.cpp b/Engine/Plugins/FX/Niagara/Source/NiagaraEditorWidgets/Private/Stack/SNiagaraStackFunctionInputName.cpp
--- a/Engine/Plugins/FX/Niagara/Source/NiagaraEditorWidgets/Private/Stack/SNiagaraStackFunctionInputName.cpp
+++ b/Engine/Plugins/FX/Niagara/Source/NiagaraEditorWidgets/Private/Stack/SNiagaraStackFunctionInputName.cpp
@@ -18,6 +18,8 @@ void SNiagaraStackFunctionInputName::Construct(const FArguments& InArgs, UNiagar
 	ChildSlot
 	[
 		SNew(SHorizontalBox)
+		.IsEnabled_UObject(FunctionInput, &UNiagaraStackEntry::GetOwnerIsEnabled)
+
 		// Edit condition checkbox
 		+ SHorizontalBox::Slot()
 		.AutoWidth()
@@ -28,6 +30,7 @@ void SNiagaraStackFunctionInputName::Construct(const FArguments& InArgs, UNiagar
 			.IsChecked(this, &SNiagaraStackFunctionInputName::GetEditConditionCheckState)
 			.OnCheckStateChanged(this, &SNiagaraStackFunctionInputName::OnEditConditionCheckStateChanged)
 		]
+
 		// Name Label
 		+ SHorizontalBox::Slot()
 		.VAlign(VAlign_Center)
@@ -83,7 +86,7 @@ bool SNiagaraStackFunctionInputName::GetIsNameWidgetSelected() const
 
 bool SNiagaraStackFunctionInputName::GetIsEnabled() const
 {
-	return FunctionInput->GetHasEditCondition() == false || FunctionInput->GetEditConditionEnabled();
+	return FunctionInput->GetOwnerIsEnabled() && (FunctionInput->GetHasEditCondition() == false || FunctionInput->GetEditConditionEnabled());
 }
 
 void SNiagaraStackFunctionInputName::OnNameTextCommitted(const FText& InText, ETextCommit::Type InCommitType)
diff --git a/Engine/Plugins/FX/Niagara/Source/NiagaraEditorWidgets/Private/Stack/SNiagaraStackFunctionInputValue.cpp b/Engine/Plugins/FX/Niagara/Source/NiagaraEditorWidgets/Private/Stack/SNiagaraStackFunctionInputValue.cpp
--- a/Engine/Plugins/FX/Niagara/Source/NiagaraEditorWidgets/Private/Stack/SNiagaraStackFunctionInputValue.cpp
+++ b/Engine/Plugins/FX/Niagara/Source/NiagaraEditorWidgets/Private/Stack/SNiagaraStackFunctionInputValue.cpp
@@ -53,6 +53,7 @@ void SNiagaraStackFunctionInputValue::Construct(const FArguments& InArgs, UNiaga
 		.VerticalImage(FNiagaraEditorWidgetsStyle::Get().GetBrush("NiagaraEditor.Stack.DropTarget.BorderVertical"))
 		.BackgroundColor(FNiagaraEditorWidgetsStyle::Get().GetColor("NiagaraEditor.Stack.DropTarget.BackgroundColor"))
 		.BackgroundColorHover(FNiagaraEditorWidgetsStyle::Get().GetColor("NiagaraEditor.Stack.DropTarget.BackgroundColorHover"))
+		.IsEnabled_UObject(FunctionInput, &UNiagaraStackEntry::GetOwnerIsEnabled)
 		.Content()
 		[
 			// Values
diff --git a/Engine/Plugins/FX/Niagara/Source/NiagaraEditorWidgets/Private/Stack/SNiagaraStackModuleItem.cpp b/Engine/Plugins/FX/Niagara/Source/NiagaraEditorWidgets/Private/Stack/SNiagaraStackModuleItem.cpp
--- a/Engine/Plugins/FX/Niagara/Source/NiagaraEditorWidgets/Private/Stack/SNiagaraStackModuleItem.cpp
+++ b/Engine/Plugins/FX/Niagara/Source/NiagaraEditorWidgets/Private/Stack/SNiagaraStackModuleItem.cpp
@@ -66,6 +66,7 @@ void SNiagaraStackModuleItem::Construct(const FArguments& InArgs, UNiagaraStackM
 				.IssueSeverity_UObject(ModuleItem, &UNiagaraStackModuleItem::GetHighestStackIssueSeverity)
 				.ErrorTooltip(this, &SNiagaraStackModuleItem::GetErrorButtonTooltipText)
 				.Visibility(this, &SNiagaraStackModuleItem::GetStackIssuesWarningVisibility)
+				.IsEnabled(this, &SNiagaraStackModuleItem::GetButtonsEnabled)
 				.OnButtonClicked(this, &SNiagaraStackModuleItem::ExpandEntry)
 			]
 			// Raise Action Menu button
@@ -81,6 +82,7 @@ void SNiagaraStackModuleItem::Construct(const FArguments& InArgs, UNiagaraStackM
 				.HAlign(HAlign_Center)
 				.VAlign(VAlign_Center)
 				.Visibility(this, &SNiagaraStackModuleItem::GetRaiseActionMenuVisibility)
+				.IsEnabled(this, &SNiagaraStackModuleItem::GetButtonsEnabled)
 			]
 			// Refresh button
 			+ SHorizontalBox::Slot()
@@ -93,6 +95,7 @@ void SNiagaraStackModuleItem::Construct(const FArguments& InArgs, UNiagaraStackM
 				.ForegroundColor(FNiagaraEditorWidgetsStyle::Get().GetColor("NiagaraEditor.Stack.FlatButtonColor"))
 				.ToolTipText(LOCTEXT("RefreshTooltip", "Refresh this module"))
 				.Visibility(this, &SNiagaraStackModuleItem::GetRefreshVisibility)
+				.IsEnabled(this, &SNiagaraStackModuleItem::GetButtonsEnabled)
 				.OnClicked(this, &SNiagaraStackModuleItem::RefreshClicked)
 				.Content()
 				[
@@ -110,8 +113,8 @@ void SNiagaraStackModuleItem::Construct(const FArguments& InArgs, UNiagaraStackM
 				.ButtonStyle(FEditorStyle::Get(), "HoverHintOnly")
 				.IsFocusable(false)
 				.ForegroundColor(FNiagaraEditorWidgetsStyle::Get().GetColor("NiagaraEditor.Stack.FlatButtonColor"))
-				.ToolTipText(LOCTEXT("DeleteToolTip", "Delete this module"))
-				.Visibility(this, &SNiagaraStackModuleItem::GetEditButtonVisibility)
+				.ToolTipText(this, &SNiagaraStackModuleItem::GetDeleteButtonToolTipText)
+				.IsEnabled(this, &SNiagaraStackModuleItem::GetDeleteButtonEnabled)
 				.OnClicked(this, &SNiagaraStackModuleItem::DeleteClicked)
 				.Content()
 				[
@@ -128,6 +131,7 @@ void SNiagaraStackModuleItem::Construct(const FArguments& InArgs, UNiagaraStackM
 				SNew(SCheckBox)
 				.IsChecked(this, &SNiagaraStackModuleItem::GetCheckState)
 				.OnCheckStateChanged(this, &SNiagaraStackModuleItem::OnCheckStateChanged)
+				.IsEnabled(this, &SNiagaraStackModuleItem::GetEnabledCheckBoxEnabled)
 			]
 		]
 	];
@@ -198,9 +202,32 @@ void SNiagaraStackModuleItem::OnCheckStateChanged(ECheckBoxState InCheckState)
 	ModuleItem->SetIsEnabled(InCheckState == ECheckBoxState::Checked);
 }
 
-EVisibility SNiagaraStackModuleItem::GetEditButtonVisibility() const
+bool SNiagaraStackModuleItem::GetButtonsEnabled() const
 {
-	return ModuleItem->CanMoveAndDelete() ? EVisibility::Visible : EVisibility::Collapsed;
+	return ModuleItem->GetOwnerIsEnabled() && ModuleItem->GetIsEnabled();
+}
+
+FText SNiagaraStackModuleItem::GetDeleteButtonToolTipText() const
+{
+	if (ModuleItem->CanMoveAndDelete())
+	{
+		return LOCTEXT("DeleteToolTip", "Delete this module");
+	}
+	else
+	{
+		// TODO: This message should be handled by the entry.
+		return LOCTEXT("CantDeleteToolTip", "This module can not be deleted becaue it is inherited.");
+	}
+}
+
+bool SNiagaraStackModuleItem::GetDeleteButtonEnabled() const
+{
+	return ModuleItem->GetOwnerIsEnabled() && ModuleItem->CanMoveAndDelete();
+}
+
+bool SNiagaraStackModuleItem::GetEnabledCheckBoxEnabled() const
+{
+	return ModuleItem->GetOwnerIsEnabled();
 }
 
 EVisibility SNiagaraStackModuleItem::GetRaiseActionMenuVisibility() const
diff --git a/Engine/Plugins/FX/Niagara/Source/NiagaraEditorWidgets/Private/Stack/SNiagaraStackModuleItem.h b/Engine/Plugins/FX/Niagara/Source/NiagaraEditorWidgets/Private/Stack/SNiagaraStackModuleItem.h
--- a/Engine/Plugins/FX/Niagara/Source/NiagaraEditorWidgets/Private/Stack/SNiagaraStackModuleItem.h
+++ b/Engine/Plugins/FX/Niagara/Source/NiagaraEditorWidgets/Private/Stack/SNiagaraStackModuleItem.h
@@ -32,7 +32,13 @@ private:
 
 	void OnCheckStateChanged(ECheckBoxState InCheckState);
 
-	EVisibility GetEditButtonVisibility() const;
+	bool GetButtonsEnabled() const;
+
+	FText GetDeleteButtonToolTipText() const;
+
+	bool GetDeleteButtonEnabled() const;
+
+	bool GetEnabledCheckBoxEnabled() const;
 
 	EVisibility GetRaiseActionMenuVisibility() const;
 
diff --git a/Engine/Plugins/FX/Niagara/Source/NiagaraEditorWidgets/Private/Stack/SNiagaraStackTableRow.cpp b/Engine/Plugins/FX/Niagara/Source/NiagaraEditorWidgets/Private/Stack/SNiagaraStackTableRow.cpp
--- a/Engine/Plugins/FX/Niagara/Source/NiagaraEditorWidgets/Private/Stack/SNiagaraStackTableRow.cpp
+++ b/Engine/Plugins/FX/Niagara/Source/NiagaraEditorWidgets/Private/Stack/SNiagaraStackTableRow.cpp
@@ -40,6 +40,7 @@ void SNiagaraStackTableRow::Construct(const FArguments& InArgs, UNiagaraStackVie
 
 	InactiveItemBackgroundColor = InArgs._ItemBackgroundColor;
 	ActiveItemBackgroundColor = InactiveItemBackgroundColor + FLinearColor(.05f, .05f, .05f, 0.0f);
+	DisabledItemBackgroundColor = InactiveItemBackgroundColor + FLinearColor(.02f, .02f, .02f, 0.0f);
 	ForegroundColor = InArgs._ItemForegroundColor;
 
 	ExecutionCategoryToolTipText = InStackEntry->GetExecutionSubcategoryName() != NAME_None
@@ -206,24 +207,42 @@ void SNiagaraStackTableRow::SetNameAndValueContent(TSharedRef<SWidget> InNameWid
 		];
 	}
 
+	FName AccentColorName = FNiagaraStackEditorWidgetsUtilities::GetIconColorNameForExecutionCategory(StackEntry->GetExecutionCategoryName());
+	FLinearColor AccentColor = AccentColorName != NAME_None ? FNiagaraEditorWidgetsStyle::Get().GetColor(AccentColorName) : FLinearColor::Transparent;
 	ChildSlot
 	[
 		SNew(SOverlay)
 		+ SOverlay::Slot()
 		[
 			SNew(SBorder)
-			.BorderImage(FEditorStyle::GetBrush("WhiteBrush"))
-			.BorderBackgroundColor(FNiagaraEditorWidgetsStyle::Get().GetColor(FNiagaraStackEditorWidgetsUtilities::GetColorNameForExecutionCategory(StackEntry->GetExecutionCategoryName())))
+			.BorderImage(FEditorStyle::GetBrush("NoBrush"))
 			.Visibility(this, &SNiagaraStackTableRow::GetRowVisibility)
-			.Padding(FMargin(9, 0, 9, 0))
+			.Padding(FMargin(0, 0, 0, 0))
 			[
-				SNew(SBorder)
-				.BorderImage(FEditorStyle::GetBrush("WhiteBrush"))
-				.BorderBackgroundColor(this, &SNiagaraStackTableRow::GetItemBackgroundColor)
-				.ForegroundColor(ForegroundColor)
-				.Padding(0)
+				SNew(SHorizontalBox)
+				+ SHorizontalBox::Slot()
+				.AutoWidth()
+				.Padding(1, 0, 6, 0)
 				[
-					ChildContent.ToSharedRef()
+					SNew(SBorder)
+					.BorderImage(FEditorStyle::GetBrush("WhiteBrush"))
+					.BorderBackgroundColor(AccentColor)
+					.Padding(0)
+					[
+						SNew(SBox)
+						.WidthOverride(4)
+					]
+				]
+				+ SHorizontalBox::Slot()
+				[
+					SNew(SBorder)
+					.BorderImage(FEditorStyle::GetBrush("WhiteBrush"))
+					.BorderBackgroundColor(this, &SNiagaraStackTableRow::GetItemBackgroundColor)
+					.ForegroundColor(ForegroundColor)
+					.Padding(0)
+					[
+						ChildContent.ToSharedRef()
+					]
 				]
 			]
 		]
@@ -426,7 +445,14 @@ void SNiagaraStackTableRow::OnValueColumnWidthChanged(float Width)
 
 FSlateColor SNiagaraStackTableRow::GetItemBackgroundColor() const
 {
-	return GetIsRowActive() ? ActiveItemBackgroundColor : InactiveItemBackgroundColor;
+	if (StackEntry->GetIsEnabled() && StackEntry->GetOwnerIsEnabled())
+	{
+		return GetIsRowActive() ? ActiveItemBackgroundColor : InactiveItemBackgroundColor;
+	}
+	else
+	{
+		return DisabledItemBackgroundColor;
+	}
 }
 
 EVisibility SNiagaraStackTableRow::GetSearchResultBorderVisibility() const
diff --git a/Engine/Plugins/FX/Niagara/Source/NiagaraEditorWidgets/Private/Stack/SNiagaraStackTableRow.h b/Engine/Plugins/FX/Niagara/Source/NiagaraEditorWidgets/Private/Stack/SNiagaraStackTableRow.h
--- a/Engine/Plugins/FX/Niagara/Source/NiagaraEditorWidgets/Private/Stack/SNiagaraStackTableRow.h
+++ b/Engine/Plugins/FX/Niagara/Source/NiagaraEditorWidgets/Private/Stack/SNiagaraStackTableRow.h
@@ -110,6 +110,7 @@ private:
 
 	FLinearColor InactiveItemBackgroundColor;
 	FLinearColor ActiveItemBackgroundColor;
+	FLinearColor DisabledItemBackgroundColor;
 	FLinearColor ForegroundColor;
 
 	FText ExecutionCategoryToolTipText;
diff --git a/Engine/Plugins/FX/Niagara/Source/NiagaraEditorWidgets/Public/NiagaraEditorWidgetsModule.h b/Engine/Plugins/FX/Niagara/Source/NiagaraEditorWidgets/Public/NiagaraEditorWidgetsModule.h
--- a/Engine/Plugins/FX/Niagara/Source/NiagaraEditorWidgets/Public/NiagaraEditorWidgetsModule.h
+++ b/Engine/Plugins/FX/Niagara/Source/NiagaraEditorWidgets/Public/NiagaraEditorWidgetsModule.h
@@ -47,8 +47,9 @@ private:
 	class FNiagaraEditorWidgetProvider : public INiagaraEditorWidgetProvider
 	{
 	public:
-		virtual TSharedRef<SWidget> CreateStackView(UNiagaraStackViewModel& StackViewModel) override;
-		virtual TSharedRef<SWidget> CreateSystemOverview(TSharedRef<FNiagaraSystemViewModel> SystemViewModel) override;
+		virtual TSharedRef<SWidget> CreateStackView(UNiagaraStackViewModel& StackViewModel) const override;
+		virtual TSharedRef<SWidget> CreateSystemOverview(TSharedRef<FNiagaraSystemViewModel> SystemViewModel) const override;
+		virtual FLinearColor GetColorForExecutionCategory(FName ExecutionCategory) const override;
 	};
 
 public:
