Niagara - Add the ability to add, delete, and insert renderers and modules from the system overview.

diff --git a/Engine/Plugins/FX/Niagara/Source/NiagaraEditor/Private/ViewModels/Stack/NiagaraStackModuleItem.cpp b/Engine/Plugins/FX/Niagara/Source/NiagaraEditor/Private/ViewModels/Stack/NiagaraStackModuleItem.cpp
--- a/Engine/Plugins/FX/Niagara/Source/NiagaraEditor/Private/ViewModels/Stack/NiagaraStackModuleItem.cpp
+++ b/Engine/Plugins/FX/Niagara/Source/NiagaraEditor/Private/ViewModels/Stack/NiagaraStackModuleItem.cpp
@@ -778,6 +778,25 @@ void UNiagaraStackModuleItem::SetIsEnabled(bool bInIsEnabled)
 	OnRequestFullRefresh().Broadcast();
 }
 
+bool UNiagaraStackModuleItem::TestCanDeleteWithMessage(FText& OutCanDeleteMessage) const
+{
+	if (GetOwnerIsEnabled() == false)
+	{
+		OutCanDeleteMessage = LOCTEXT("CantDeleteOwnerDisabledToolTip", "This module can not be deleted because its owner is disabled.");
+		return false;
+	}
+	else if (CanMoveAndDelete())
+	{
+		OutCanDeleteMessage = LOCTEXT("DeleteToolTip", "Delete this module.");
+		return true;
+	}
+	else
+	{
+		OutCanDeleteMessage = LOCTEXT("CantDeleteToolTip", "This module can not be deleted becaue it is inherited.");
+		return false;
+	}
+}
+
 void UNiagaraStackModuleItem::Delete()
 {
 	checkf(CanMoveAndDelete(), TEXT("This module can't be deleted"));
diff --git a/Engine/Plugins/FX/Niagara/Source/NiagaraEditor/Private/ViewModels/Stack/NiagaraStackRendererItem.cpp b/Engine/Plugins/FX/Niagara/Source/NiagaraEditor/Private/ViewModels/Stack/NiagaraStackRendererItem.cpp
--- a/Engine/Plugins/FX/Niagara/Source/NiagaraEditor/Private/ViewModels/Stack/NiagaraStackRendererItem.cpp
+++ b/Engine/Plugins/FX/Niagara/Source/NiagaraEditor/Private/ViewModels/Stack/NiagaraStackRendererItem.cpp
@@ -144,9 +144,23 @@ FText UNiagaraStackRendererItem::GetDisplayName() const
 	}
 }
 
-bool UNiagaraStackRendererItem::CanDelete() const
+bool UNiagaraStackRendererItem::TestCanDeleteWithMessage(FText& OutCanDeleteMessage) const
 {
-	return HasBaseRenderer() == false;
+	if (GetOwnerIsEnabled() == false)
+	{
+		OutCanDeleteMessage = LOCTEXT("CantDeleteOwnerDisabledToolTip", "This renderer can not be deleted because its owner is disabled.");
+		return false;
+	}
+	else if (HasBaseRenderer() == false)
+	{
+		OutCanDeleteMessage = LOCTEXT("DeleteToolTip", "Delete this renderer.");
+		return true;
+	}
+	else
+	{
+		OutCanDeleteMessage = LOCTEXT("CantDeleteToolTip", "This renderer can not be deleted becaue it is inherited.");
+		return false;
+	}
 }
 
 void UNiagaraStackRendererItem::Delete()
diff --git a/Engine/Plugins/FX/Niagara/Source/NiagaraEditor/Public/ViewModels/Stack/NiagaraStackItem.h b/Engine/Plugins/FX/Niagara/Source/NiagaraEditor/Public/ViewModels/Stack/NiagaraStackItem.h
--- a/Engine/Plugins/FX/Niagara/Source/NiagaraEditor/Public/ViewModels/Stack/NiagaraStackItem.h
+++ b/Engine/Plugins/FX/Niagara/Source/NiagaraEditor/Public/ViewModels/Stack/NiagaraStackItem.h
@@ -28,8 +28,12 @@ public:
 	uint32 GetRecursiveStackIssuesCount() const;
 	EStackIssueSeverity GetHighestStackIssueSeverity() const;
 
-	virtual bool CanChangeEnabled() const { return false; }
+	virtual bool SupportsChangeEnabled() const { return false; }
 	virtual void SetIsEnabled(bool bInIsEnabled) { }
+
+	virtual bool SupportsDelete() const { return false; }
+	virtual bool TestCanDeleteWithMessage(FText& OutCanDeleteMessage) const { return false; }
+	virtual void Delete() { }
 	
 protected:
 	virtual void RefreshChildrenInternal(const TArray<UNiagaraStackEntry*>& CurrentChildren, TArray<UNiagaraStackEntry*>& NewChildren, TArray<FStackIssue>& NewIssues) override;
diff --git a/Engine/Plugins/FX/Niagara/Source/NiagaraEditor/Public/ViewModels/Stack/NiagaraStackModuleItem.h b/Engine/Plugins/FX/Niagara/Source/NiagaraEditor/Public/ViewModels/Stack/NiagaraStackModuleItem.h
--- a/Engine/Plugins/FX/Niagara/Source/NiagaraEditor/Public/ViewModels/Stack/NiagaraStackModuleItem.h
+++ b/Engine/Plugins/FX/Niagara/Source/NiagaraEditor/Public/ViewModels/Stack/NiagaraStackModuleItem.h
@@ -37,11 +37,13 @@ public:
 	bool CanRefresh() const;
 	void Refresh();
 
-	virtual bool CanChangeEnabled() const override { return true; }
+	virtual bool SupportsChangeEnabled() const override { return true; }
 	virtual bool GetIsEnabled() const override;
 	virtual void SetIsEnabled(bool bInIsEnabled) override;
 
-	void Delete();
+	virtual bool SupportsDelete() const override { return true; }
+	virtual bool TestCanDeleteWithMessage(FText& OutCanDeleteMessage) const override;
+	virtual void Delete() override;
 
 	int32 GetModuleIndex();
 	
diff --git a/Engine/Plugins/FX/Niagara/Source/NiagaraEditor/Public/ViewModels/Stack/NiagaraStackRendererItem.h b/Engine/Plugins/FX/Niagara/Source/NiagaraEditor/Public/ViewModels/Stack/NiagaraStackRendererItem.h
--- a/Engine/Plugins/FX/Niagara/Source/NiagaraEditor/Public/ViewModels/Stack/NiagaraStackRendererItem.h
+++ b/Engine/Plugins/FX/Niagara/Source/NiagaraEditor/Public/ViewModels/Stack/NiagaraStackRendererItem.h
@@ -24,9 +24,9 @@ public:
 
 	virtual FText GetDisplayName() const override;
 
-	bool CanDelete() const;
-
-	void Delete();
+	virtual bool SupportsDelete() const override { return true; }
+	virtual bool TestCanDeleteWithMessage(FText& OutCanDeleteMessage) const;
+	virtual void Delete() override;
 
 	bool HasBaseRenderer() const;
 
@@ -34,7 +34,7 @@ public:
 
 	void ResetToBase();
 
-	virtual bool CanChangeEnabled() const override { return true; }
+	virtual bool SupportsChangeEnabled() const override { return true; }
 	virtual bool GetIsEnabled() const override;
 	virtual void SetIsEnabled(bool bInIsEnabled) override;
 
diff --git a/Engine/Plugins/FX/Niagara/Source/NiagaraEditorWidgets/Private/NiagaraEditorWidgetsStyle.cpp b/Engine/Plugins/FX/Niagara/Source/NiagaraEditorWidgets/Private/NiagaraEditorWidgetsStyle.cpp
--- a/Engine/Plugins/FX/Niagara/Source/NiagaraEditorWidgets/Private/NiagaraEditorWidgetsStyle.cpp
+++ b/Engine/Plugins/FX/Niagara/Source/NiagaraEditorWidgets/Private/NiagaraEditorWidgetsStyle.cpp
@@ -146,7 +146,7 @@ TSharedRef< FSlateStyleSet > FNiagaraEditorWidgetsStyle::Create()
 
 	Style->Set("NiagaraEditor.Stack.IconColor.System", FLinearColor(FColor(1, 202, 252)));
 	Style->Set("NiagaraEditor.Stack.IconColor.Emitter", FLinearColor(FColor(241, 99, 6)));
-	Style->Set("NiagaraEditor.Stack.IconColor.Particle", FLinearColor(FColor(131, 228, 9)));
+	Style->Set("NiagaraEditor.Stack.IconColor.Particle", FLinearColor(FColor(131, 218, 9)));
 	Style->Set("NiagaraEditor.Stack.IconColor.Render", FLinearColor(FColor(230, 102, 102)));
 
  	Style->Set("NiagaraEditor.Stack.DropTarget.BackgroundColor", FLinearColor(1.0f, 1.0f, 1.0f, 0.25f));
@@ -178,11 +178,17 @@ TSharedRef< FSlateStyleSet > FNiagaraEditorWidgetsStyle::Create()
 	Style->Set("NiagaraEditor.Stack.SearchResult", new BOX_PLUGIN_BRUSH("Icons/SearchResultBorder", FMargin(1.f/8.f)));
 
 	Style->Set("NiagaraEditor.Stack.AddButton", FButtonStyle()
-		.SetNormal(BOX_CORE_BRUSH("Common/FlatButton", 2.0f / 8.0f, FLinearColor(0, 0, 0, .25f)))
-		.SetHovered(BOX_CORE_BRUSH("Common/FlatButton", 2.0f / 8.0f, FEditorStyle::GetSlateColor("SelectionColor")))
-		.SetPressed(BOX_CORE_BRUSH("Common/FlatButton", 2.0f / 8.0f, FEditorStyle::GetSlateColor("SelectionColor_Pressed")))
+		.SetNormal(BOX_CORE_BRUSH("Common/FlatButton", 2.0f / 8.0f, FLinearColor(0.5f, 0.5f, 0.5f, 1.0f)))
+		.SetHovered(BOX_CORE_BRUSH("Common/FlatButton", 2.0f / 8.0f, FLinearColor(1.0f, 1.0f, 1.0f, 1.0f)))
+		.SetPressed(BOX_CORE_BRUSH("Common/FlatButton", 2.0f / 8.0f, FLinearColor(.8f, .8f, .8f, 1.0f)))
 	);
 
+	FTextBlockStyle AddButtonText = FTextBlockStyle(NormalText)
+		.SetColorAndOpacity(FLinearColor(0.8f, 0.8f, 0.8f, 1.0f))
+		.SetShadowOffset(FVector2D(0, 1))
+		.SetShadowColorAndOpacity(FLinearColor(0, 0, 0, 0.9f));
+	Style->Set("NiagaraEditor.Stack.AddButtonText", AddButtonText);
+
 	Style->Set("NiagaraEditor.ShowInCurveEditorIcon", new IMAGE_PLUGIN_BRUSH("Icons/ShowInCurveEditor", Icon16x16, FLinearColor::White));
 
 	return Style;
diff --git a/Engine/Plugins/FX/Niagara/Source/NiagaraEditorWidgets/Private/NiagaraEditorWidgetsUtilities.cpp b/Engine/Plugins/FX/Niagara/Source/NiagaraEditorWidgets/Private/NiagaraEditorWidgetsUtilities.cpp
--- a/Engine/Plugins/FX/Niagara/Source/NiagaraEditorWidgets/Private/NiagaraEditorWidgetsUtilities.cpp
+++ b/Engine/Plugins/FX/Niagara/Source/NiagaraEditorWidgets/Private/NiagaraEditorWidgetsUtilities.cpp
@@ -2,6 +2,16 @@
 
 #include "NiagaraEditorWidgetsUtilities.h"
 #include "ViewModels/Stack/NiagaraStackEntry.h"
+#include "ViewModels/Stack/NiagaraStackItem.h"
+#include "ViewModels/Stack/NiagaraStackModuleItem.h"
+#include "Stack/SNiagaraStackItemGroupAddMenu.h"
+#include "IContentBrowserSingleton.h"
+#include "ContentBrowserModule.h"
+#include "Modules/ModuleManager.h"
+#include "Subsystems/AssetEditorSubsystem.h"
+#include "Editor.h"
+
+#define LOCTEXT_NAMESPACE "NiagaraStackEditorWidgetsUtilities"
 
 FName FNiagaraStackEditorWidgetsUtilities::GetColorNameForExecutionCategory(FName ExecutionCategoryName)
 {
@@ -101,4 +111,143 @@ FName FNiagaraStackEditorWidgetsUtilities::GetIconColorNameForExecutionCategory(
 	{
 		return NAME_None;
 	}
-}
\ No newline at end of file
+}
+
+void OpenSourceAsset(TWeakObjectPtr<UNiagaraStackEntry> StackEntryWeak)
+{
+	UNiagaraStackEntry* StackEntry = StackEntryWeak.Get();
+	if (StackEntry != nullptr)
+	{
+		GEditor->GetEditorSubsystem<UAssetEditorSubsystem>()->OpenEditorForAsset(StackEntry->GetExternalAsset());
+	}
+}
+
+void ShowAssetInContentBrowser(TWeakObjectPtr<UNiagaraStackEntry> StackEntryWeak)
+{
+	UNiagaraStackEntry* StackEntry = StackEntryWeak.Get();
+	if (StackEntry != nullptr)
+	{
+		FContentBrowserModule& ContentBrowserModule = FModuleManager::Get().LoadModuleChecked<FContentBrowserModule>(TEXT("ContentBrowser"));
+		TArray<FAssetData> Assets;
+		Assets.Add(FAssetData(StackEntry->GetExternalAsset()));
+		ContentBrowserModule.Get().SyncBrowserToAssets(Assets);
+	}
+}
+
+bool FNiagaraStackEditorWidgetsUtilities::AddStackEntryAssetContextMenuActions(FMenuBuilder& MenuBuilder, UNiagaraStackEntry& StackEntry)
+{
+	if (StackEntry.GetExternalAsset() != nullptr)
+	{
+		MenuBuilder.BeginSection("AssetActions", LOCTEXT("AssetActions", "Asset Actions"));
+		{
+			MenuBuilder.AddMenuEntry(
+				LOCTEXT("OpenAndFocusAsset", "Open and Focus Asset"),
+				FText::Format(LOCTEXT("OpenAndFocusAssetTooltip", "Open {0} in separate editor"), StackEntry.GetDisplayName()),
+				FSlateIcon(),
+				FUIAction(FExecuteAction::CreateStatic(&OpenSourceAsset, TWeakObjectPtr<UNiagaraStackEntry>(&StackEntry))));
+			MenuBuilder.AddMenuEntry(
+				LOCTEXT("ShowAssetInContentBrowser", "Show in Content Browser"),
+				FText::Format(LOCTEXT("ShowAssetInContentBrowserToolTip", "Navigate to {0} in the Content Browser window"), StackEntry.GetDisplayName()),
+				FSlateIcon(),
+				FUIAction(FExecuteAction::CreateStatic(&ShowAssetInContentBrowser, TWeakObjectPtr<UNiagaraStackEntry>(&StackEntry))));
+		}
+		MenuBuilder.EndSection();
+		return true;
+	}
+	return false;
+}
+
+void DeleteItem(TWeakObjectPtr<UNiagaraStackItem> StackItemWeak)
+{
+	UNiagaraStackItem* StackItem = StackItemWeak.Get();
+	FText Unused;
+	if (StackItem != nullptr && StackItem->TestCanDeleteWithMessage(Unused))
+	{
+		StackItem->Delete();
+	}
+}
+
+void ToggleEnabledState(TWeakObjectPtr<UNiagaraStackItem> StackItemWeak)
+{
+	UNiagaraStackItem* StackItem = StackItemWeak.Get();
+	if (StackItem != nullptr)
+	{
+		StackItem->SetIsEnabled(!StackItem->GetIsEnabled());
+	}
+}
+
+bool FNiagaraStackEditorWidgetsUtilities::AddStackItemContextMenuActions(FMenuBuilder& MenuBuilder, UNiagaraStackItem& StackItem)
+{
+	if (StackItem.SupportsDelete() || StackItem.SupportsChangeEnabled())
+	{
+		MenuBuilder.BeginSection("ItemActions", LOCTEXT("ItemActions", "Item Actions"));
+		{
+			if (StackItem.SupportsChangeEnabled())
+			{
+				FUIAction Action(FExecuteAction::CreateStatic(&ToggleEnabledState, TWeakObjectPtr<UNiagaraStackItem>(&StackItem)),
+					FCanExecuteAction(),
+					FIsActionChecked::CreateUObject(&StackItem, &UNiagaraStackItem::GetIsEnabled));
+				MenuBuilder.AddMenuEntry(
+					LOCTEXT("IsEnabled", "Is Enabled"),
+					LOCTEXT("ToggleEnabledToolTip", "Toggle enabled/disabled state"),
+					FSlateIcon(),
+					Action,
+					NAME_None,
+					EUserInterfaceActionType::Check);
+			}
+
+			if (StackItem.SupportsDelete())
+			{
+				FText CanDeleteMessage;
+				bool bCanDelete = StackItem.TestCanDeleteWithMessage(CanDeleteMessage);
+				MenuBuilder.AddMenuEntry(
+					LOCTEXT("DeleteModule", "Delete Item"),
+					CanDeleteMessage,
+					FSlateIcon(),
+					FUIAction(
+						FExecuteAction::CreateStatic(&DeleteItem, TWeakObjectPtr<UNiagaraStackItem>(&StackItem)),
+						FCanExecuteAction::CreateLambda([=]() { return bCanDelete; })));
+
+			}
+		}
+		MenuBuilder.EndSection();
+		return true;
+	}
+	return false;
+}
+
+void ShowInsertModuleMenu(TWeakObjectPtr<UNiagaraStackModuleItem> StackModuleItemWeak, int32 InsertOffset, TWeakPtr<SWidget> TargetWidgetWeak)
+{
+	UNiagaraStackModuleItem* StackModuleItem = StackModuleItemWeak.Get();
+	TSharedPtr<SWidget> TargetWidget = TargetWidgetWeak.Pin();
+	if (StackModuleItem != nullptr && TargetWidget.IsValid())
+	{
+		TSharedRef<SWidget> MenuContent = SNew(SNiagaraStackItemGroupAddMenu, StackModuleItem->GetGroupAddUtilities(), StackModuleItem->GetModuleIndex() + InsertOffset);
+		FGeometry ThisGeometry = TargetWidget->GetCachedGeometry();
+		bool bAutoAdjustForDpiScale = false; // Don't adjust for dpi scale because the push menu command is expecting an unscaled position.
+		FVector2D MenuPosition = FSlateApplication::Get().CalculatePopupWindowPosition(ThisGeometry.GetLayoutBoundingRect(), MenuContent->GetDesiredSize(), bAutoAdjustForDpiScale);
+		FSlateApplication::Get().PushMenu(TargetWidget.ToSharedRef(), FWidgetPath(), MenuContent, MenuPosition, FPopupTransitionEffect::ContextMenu);
+	}
+}
+
+bool FNiagaraStackEditorWidgetsUtilities::AddStackModuleItemContextMenuActions(FMenuBuilder& MenuBuilder, UNiagaraStackModuleItem& StackModuleItem, TSharedRef<SWidget> TargetWidget)
+{
+	MenuBuilder.BeginSection("ModuleActions", LOCTEXT("ModuleActions", "Module Actions"));
+	{
+		MenuBuilder.AddMenuEntry(
+			LOCTEXT("InsertModuleAbove", "Insert Above"),
+			LOCTEXT("InsertModuleAboveToolTip", "Insert a new module above this module in the stack."),
+			FSlateIcon(),
+			FUIAction(FExecuteAction::CreateStatic(&ShowInsertModuleMenu, TWeakObjectPtr<UNiagaraStackModuleItem>(&StackModuleItem), 0, TWeakPtr<SWidget>(TargetWidget))));
+
+		MenuBuilder.AddMenuEntry(
+			LOCTEXT("InsertModuleBelow", "Insert Below"),
+			LOCTEXT("InsertModuleBelowToolTip", "Insert a new module below this module in the stack."),
+			FSlateIcon(),
+			FUIAction(FExecuteAction::CreateStatic(&ShowInsertModuleMenu, TWeakObjectPtr<UNiagaraStackModuleItem>(&StackModuleItem), 1, TWeakPtr<SWidget>(TargetWidget))));
+	}
+	MenuBuilder.EndSection();
+	return true;
+}
+
+#undef LOCTEXT_NAMESPACE
\ No newline at end of file
diff --git a/Engine/Plugins/FX/Niagara/Source/NiagaraEditorWidgets/Private/NiagaraEditorWidgetsUtilities.h b/Engine/Plugins/FX/Niagara/Source/NiagaraEditorWidgets/Private/NiagaraEditorWidgetsUtilities.h
--- a/Engine/Plugins/FX/Niagara/Source/NiagaraEditorWidgets/Private/NiagaraEditorWidgetsUtilities.h
+++ b/Engine/Plugins/FX/Niagara/Source/NiagaraEditorWidgets/Private/NiagaraEditorWidgetsUtilities.h
@@ -4,6 +4,10 @@
 
 #include "UObject/NameTypes.h"
 
+class UNiagaraStackEntry;
+class UNiagaraStackItem;
+class UNiagaraStackModuleItem;
+
 namespace FNiagaraStackEditorWidgetsUtilities
 {
 	FName GetColorNameForExecutionCategory(FName ExecutionCategoryName);
@@ -11,4 +15,11 @@ namespace FNiagaraStackEditorWidgetsUtilities
 	FName GetIconNameForExecutionSubcategory(FName ExecutionSubcategoryName, bool bIsHighlighted);
 
 	FName GetIconColorNameForExecutionCategory(FName ExecutionCategoryName);
+	
+	bool AddStackEntryAssetContextMenuActions(FMenuBuilder& MenuBuilder, UNiagaraStackEntry& StackEntry);
+
+	bool AddStackItemContextMenuActions(FMenuBuilder& MenuBuilder, UNiagaraStackItem& StackItem);
+
+	bool AddStackModuleItemContextMenuActions(FMenuBuilder& MenuBuilder, UNiagaraStackModuleItem& StackItem, TSharedRef<SWidget> TargetWidget);
+
 }
\ No newline at end of file
diff --git a/Engine/Plugins/FX/Niagara/Source/NiagaraEditorWidgets/Private/SNiagaraOverviewStack.cpp b/Engine/Plugins/FX/Niagara/Source/NiagaraEditorWidgets/Private/SNiagaraOverviewStack.cpp
--- a/Engine/Plugins/FX/Niagara/Source/NiagaraEditorWidgets/Private/SNiagaraOverviewStack.cpp
+++ b/Engine/Plugins/FX/Niagara/Source/NiagaraEditorWidgets/Private/SNiagaraOverviewStack.cpp
@@ -7,11 +7,13 @@
 #include "ViewModels/Stack/NiagaraStackEntry.h"
 #include "ViewModels/Stack/NiagaraStackItemGroup.h"
 #include "ViewModels/Stack/NiagaraStackItem.h"
+#include "ViewModels/Stack/NiagaraStackModuleItem.h"
 #include "ViewModels/Stack/NiagaraStackViewModel.h"
 #include "ViewModels/Stack/NiagaraStackSelection.h"
 #include "NiagaraEditorWidgetsStyle.h"
 #include "NiagaraEditorWidgetsUtilities.h"
 #include "SNiagaraStack.h"
+#include "Stack/SNiagaraStackItemGroupAddButton.h"
 
 #include "Widgets/Layout/SScrollBox.h"
 #include "Widgets/Images/SImage.h"
@@ -23,13 +25,13 @@
 
 class SNiagaraSystemOverviewEntryListRow : public STableRow<UNiagaraStackEntry*>
 {
-	SLATE_BEGIN_ARGS(SNiagaraSystemOverviewEntryListRow)
-	{}
+	SLATE_BEGIN_ARGS(SNiagaraSystemOverviewEntryListRow) {}
 		SLATE_DEFAULT_SLOT(FArguments, Content)
 	SLATE_END_ARGS();
 
-	void Construct(const FArguments& InArgs, UNiagaraStackEntry* StackEntry, const TSharedRef<STableViewBase>& InOwnerTableView)
+	void Construct(const FArguments& InArgs, UNiagaraStackEntry* InStackEntry, const TSharedRef<STableViewBase>& InOwnerTableView)
 	{
+		StackEntry = InStackEntry;
 		FSlateColor IconColor = FNiagaraEditorWidgetsStyle::Get().GetColor(FNiagaraStackEditorWidgetsUtilities::GetColorNameForExecutionCategory(StackEntry->GetExecutionCategoryName()));
 		ItemBackgroundColor = FNiagaraEditorWidgetsStyle::Get().GetColor("NiagaraEditor.Stack.Item.HeaderBackgroundColor");
 		DisabledItemBackgroundColor = ItemBackgroundColor + FLinearColor(.02f, .02f, .02f, 0.0f);
@@ -60,7 +62,7 @@ class SNiagaraSystemOverviewEntryListRow : public STableRow<UNiagaraStackEntry*>
 				[
 					SNew(SBorder)
 					.BorderImage(this, &SNiagaraSystemOverviewEntryListRow::GetBorder)
-					.Padding(FMargin(2, 4, 7, 4))
+					.Padding(FMargin(2, 4, 1, 4))
 					[
 						InArgs._Content.Widget
 					]
@@ -80,6 +82,36 @@ class SNiagaraSystemOverviewEntryListRow : public STableRow<UNiagaraStackEntry*>
 		InOwnerTableView);
 	}
 
+	// Need to use the mouse down event here since the graph eats the mouse up to show its context menu.
+	virtual FReply OnMouseButtonDown(const FGeometry& MyGeometry, const FPointerEvent& MouseEvent) override
+	{
+		if (MouseEvent.GetEffectingButton() == EKeys::RightMouseButton)
+		{
+			FMenuBuilder MenuBuilder(true, nullptr);
+			bool bMenuItemsAdded = false;
+
+			if (StackEntry->IsA<UNiagaraStackModuleItem>())
+			{
+				bMenuItemsAdded |= FNiagaraStackEditorWidgetsUtilities::AddStackModuleItemContextMenuActions(MenuBuilder, *CastChecked<UNiagaraStackModuleItem>(StackEntry), this->AsShared());
+			}
+
+			if (StackEntry->IsA<UNiagaraStackItem>())
+			{
+				bMenuItemsAdded |= FNiagaraStackEditorWidgetsUtilities::AddStackItemContextMenuActions(MenuBuilder, *CastChecked<UNiagaraStackItem>(StackEntry));
+			}
+
+			bMenuItemsAdded |= FNiagaraStackEditorWidgetsUtilities::AddStackEntryAssetContextMenuActions(MenuBuilder, *StackEntry);
+		
+			if (bMenuItemsAdded)
+			{
+				FWidgetPath WidgetPath = MouseEvent.GetEventPath() != nullptr ? *MouseEvent.GetEventPath() : FWidgetPath();
+				FSlateApplication::Get().PushMenu(AsShared(), WidgetPath, MenuBuilder.MakeWidget(), MouseEvent.GetScreenSpacePosition(), FPopupTransitionEffect(FPopupTransitionEffect::ContextMenu));
+				return FReply::Handled();
+			}
+		}
+		return STableRow<UNiagaraStackEntry*>::OnMouseButtonDown(MyGeometry, MouseEvent);
+	}
+
 private:
 	FSlateColor GetItemBackgroundColor(UNiagaraStackEntry* Entry) const
 	{
@@ -87,6 +119,7 @@ private:
 	}
 
 private:
+	UNiagaraStackEntry* StackEntry;
 	FLinearColor ItemBackgroundColor;
 	FLinearColor DisabledItemBackgroundColor;
 };
@@ -241,6 +274,7 @@ TSharedRef<ITableRow> SNiagaraOverviewStack::OnGenerateRowForEntry(UNiagaraStack
 		UNiagaraStackItem* StackItem = CastChecked<UNiagaraStackItem>(Item);
 		Content = SNew(SHorizontalBox)
 			+ SHorizontalBox::Slot()
+			.VAlign(VAlign_Center)
 			[
 				SNew(STextBlock)
 				.TextStyle(FNiagaraEditorWidgetsStyle::Get(), "NiagaraEditor.SystemOverview.ItemText")
@@ -248,6 +282,7 @@ TSharedRef<ITableRow> SNiagaraOverviewStack::OnGenerateRowForEntry(UNiagaraStack
 				.IsEnabled_UObject(Item, &UNiagaraStackEntry::GetIsEnabledAndOwnerIsEnabled)
 			]
 			+ SHorizontalBox::Slot()
+			.VAlign(VAlign_Center)
 			.AutoWidth()
 			.Padding(3, 0, 0, 0)
 			[
@@ -258,9 +293,10 @@ TSharedRef<ITableRow> SNiagaraOverviewStack::OnGenerateRowForEntry(UNiagaraStack
 				.OnCheckedChanged_UObject(StackItem, &UNiagaraStackItem::SetIsEnabled)
 			];
 	}
-	else
+	else if (Item->IsA<UNiagaraStackItemGroup>())
 	{
-		Content = SNew(SHorizontalBox)
+		UNiagaraStackItemGroup* StackItemGroup = CastChecked<UNiagaraStackItemGroup>(Item);
+		TSharedRef<SHorizontalBox> ContentBox = SNew(SHorizontalBox)
 			+ SHorizontalBox::Slot()
 			.Padding(2, 0, 6, 0)
 			.VAlign(VAlign_Center)
@@ -272,12 +308,29 @@ TSharedRef<ITableRow> SNiagaraOverviewStack::OnGenerateRowForEntry(UNiagaraStack
 				.IsEnabled_UObject(Item, &UNiagaraStackEntry::GetIsEnabledAndOwnerIsEnabled)
 			]
 			+ SHorizontalBox::Slot()
+			.VAlign(VAlign_Center)
 			[
 				SNew(STextBlock)
 				.TextStyle(FNiagaraEditorWidgetsStyle::Get(), "NiagaraEditor.SystemOverview.GroupHeaderText")
 				.Text_UObject(Item, &UNiagaraStackEntry::GetDisplayName)
 				.IsEnabled_UObject(Item, &UNiagaraStackEntry::GetIsEnabledAndOwnerIsEnabled)
 			];
+
+		if (StackItemGroup->GetAddUtilities() != nullptr)
+		{
+			ContentBox->AddSlot()
+			.AutoWidth()
+			[
+				SNew(SNiagaraStackItemGroupAddButton, *StackItemGroup)
+				.Width(22)
+			];
+		}
+
+		Content = ContentBox;
+	}
+	else
+	{
+		Content = SNullWidget::NullWidget;
 	}
 
 
@@ -289,7 +342,7 @@ TSharedRef<ITableRow> SNiagaraOverviewStack::OnGenerateRowForEntry(UNiagaraStack
 
 EVisibility SNiagaraOverviewStack::GetEnabledCheckBoxVisibility(UNiagaraStackItem* Item) const
 {
-	return Item->CanChangeEnabled() ? EVisibility::Visible : EVisibility::Collapsed;
+	return Item->SupportsChangeEnabled() ? EVisibility::Visible : EVisibility::Collapsed;
 }
 
 void SNiagaraOverviewStack::OnSelectionChanged(UNiagaraStackEntry* InNewSelection, ESelectInfo::Type SelectInfo)
diff --git a/Engine/Plugins/FX/Niagara/Source/NiagaraEditorWidgets/Private/Stack/SNiagaraStackItemGroup.cpp b/Engine/Plugins/FX/Niagara/Source/NiagaraEditorWidgets/Private/Stack/SNiagaraStackItemGroup.cpp
--- a/Engine/Plugins/FX/Niagara/Source/NiagaraEditorWidgets/Private/Stack/SNiagaraStackItemGroup.cpp
+++ b/Engine/Plugins/FX/Niagara/Source/NiagaraEditorWidgets/Private/Stack/SNiagaraStackItemGroup.cpp
@@ -4,16 +4,11 @@
 #include "NiagaraEditorWidgetsStyle.h"
 #include "NiagaraEditorStyle.h"
 #include "EditorStyleSet.h"
-#include "Stack/SNiagaraStackItemGroupAddMenu.h"
+#include "Stack/SNiagaraStackItemGroupAddButton.h"
 #include "ViewModels/Stack/NiagaraStackItemGroup.h"
 #include "ViewModels/Stack/NiagaraStackViewModel.h"
-#include "Widgets/Layout/SBox.h"
-#include "Widgets/SBoxPanel.h"
 #include "Widgets/Input/SButton.h"
-#include "Widgets/Input/SComboButton.h"
-#include "Widgets/Input/SEditableTextBox.h"
 #include "Widgets/Text/STextBlock.h"
-#include "Widgets/Images/SImage.h"
 #include "SNiagaraStackErrorButton.h"
 
 #define LOCTEXT_NAMESPACE "NiagaraStackItemGroup"
@@ -78,58 +73,10 @@ void SNiagaraStackItemGroup::Construct(const FArguments& InArgs, UNiagaraStackIt
 
 TSharedRef<SWidget> SNiagaraStackItemGroup::ConstructAddButton()
 {
-	AddActionButton.Reset();
 	INiagaraStackItemGroupAddUtilities* AddUtilities = Group->GetAddUtilities();
 	if (AddUtilities != nullptr)
 	{
-		if (AddUtilities->GetAddMode() == INiagaraStackItemGroupAddUtilities::AddFromAction)
-		{
-			return SAssignNew(AddActionButton, SComboButton)
-				.Visibility(this, &SNiagaraStackItemGroup::GetAddButtonVisibility)
-				.ButtonStyle(FNiagaraEditorWidgetsStyle::Get(), "NiagaraEditor.Stack.AddButton")
-				.ToolTipText(this, &SNiagaraStackItemGroup::GetAddButtonToolTipText)
-				.HasDownArrow(false)
-				.OnGetMenuContent(this, &SNiagaraStackItemGroup::GetAddMenu)
-				.ContentPadding(0)
-				.MenuPlacement(MenuPlacement_BelowRightAnchor)
-				.ButtonContent()
-				[
-					SNew(SBox)
-					.WidthOverride(TextIconSize * 2)
-					.HeightOverride(TextIconSize)
-					.HAlign(HAlign_Center)
-					.VAlign(VAlign_Center)
-					[
-						SNew(STextBlock)
-						.TextStyle(FEditorStyle::Get(), "NormalText.Important")
-						.Font(FEditorStyle::Get().GetFontStyle("FontAwesome.10"))
-						.Text(FText::FromString(FString(TEXT("\xf067"))) /*fa-plus*/)
-					]
-				];
-		}
-		else if (AddUtilities->GetAddMode() == INiagaraStackItemGroupAddUtilities::AddDirectly)
-		{
-			return SNew(SButton)
-				.Visibility(this, &SNiagaraStackItemGroup::GetAddButtonVisibility)
-				.ButtonStyle(FNiagaraEditorWidgetsStyle::Get(), "NiagaraEditor.Stack.AddButton")
-				.ToolTipText(this, &SNiagaraStackItemGroup::GetAddButtonToolTipText)
-				.ContentPadding(0)
-				.OnClicked(this, &SNiagaraStackItemGroup::AddDirectlyButtonClicked)
-				.Content()
-				[
-					SNew(SBox)
-					.WidthOverride(TextIconSize * 2)
-					.HeightOverride(TextIconSize)
-					.HAlign(HAlign_Center)
-					.VAlign(VAlign_Center)
-					[
-						SNew(STextBlock)
-						.TextStyle(FEditorStyle::Get(), "NormalText.Important")
-						.Font(FEditorStyle::Get().GetFontStyle("FontAwesome.10"))
-						.Text(FText::FromString(FString(TEXT("\xf067"))) /*fa-plus*/)
-					]
-				];
-		}
+		return SNew(SNiagaraStackItemGroupAddButton, *Group);
 	}
 	return SNullWidget::NullWidget;
 }
@@ -146,49 +93,12 @@ EVisibility SNiagaraStackItemGroup::GetDeleteButtonVisibility() const
 	}
 }
 
-EVisibility SNiagaraStackItemGroup::GetAddButtonVisibility() const
-{
-	if (Group->GetAddUtilities() != nullptr)
-	{
-		return EVisibility::Visible;
-	}
-	else
-	{
-		return EVisibility::Collapsed;
-	}
-}
-
-FText SNiagaraStackItemGroup::GetAddButtonToolTipText() const
-{
-	if (Group->GetAddUtilities() != nullptr)
-	{
-		return FText::Format(LOCTEXT("AddToGroupFormat", "Add a new {0} to this group."), Group->GetAddUtilities()->GetAddItemName());
-	}
-	else
-	{
-		return FText();
-	}
-}
-
-FReply SNiagaraStackItemGroup::AddDirectlyButtonClicked()
-{
-	Group->GetAddUtilities()->AddItemDirectly();
-	return FReply::Handled();
-}
-
 FReply SNiagaraStackItemGroup::DeleteClicked()
 {
 	Group->Delete();
 	return FReply::Handled();
 }
 
-TSharedRef<SWidget> SNiagaraStackItemGroup::GetAddMenu()
-{
-	TSharedRef<SNiagaraStackItemGroupAddMenu> AddMenu = SNew(SNiagaraStackItemGroupAddMenu, Group->GetAddUtilities(), INDEX_NONE);
-	AddActionButton->SetMenuContentWidgetToFocus(AddMenu->GetFilterTextBox()->AsShared());
-	return AddMenu;
-}
-
 EVisibility SNiagaraStackItemGroup::GetStackIssuesWarningVisibility() const
 {
 	return  Group->GetRecursiveStackIssuesCount() > 0 ? EVisibility::Visible : EVisibility::Collapsed;
@@ -199,5 +109,4 @@ FText SNiagaraStackItemGroup::GetErrorButtonTooltipText() const
 	return FText::Format(LOCTEXT("GroupIssuesTooltip", "This group contains items that have a total of {0} issues, click to expand."), Group->GetRecursiveStackIssuesCount());
 }
 
-
 #undef LOCTEXT_NAMESPACE
\ No newline at end of file
diff --git a/Engine/Plugins/FX/Niagara/Source/NiagaraEditorWidgets/Private/Stack/SNiagaraStackItemGroup.h b/Engine/Plugins/FX/Niagara/Source/NiagaraEditorWidgets/Private/Stack/SNiagaraStackItemGroup.h
--- a/Engine/Plugins/FX/Niagara/Source/NiagaraEditorWidgets/Private/Stack/SNiagaraStackItemGroup.h
+++ b/Engine/Plugins/FX/Niagara/Source/NiagaraEditorWidgets/Private/Stack/SNiagaraStackItemGroup.h
@@ -21,22 +21,12 @@ private:
 
 	EVisibility GetDeleteButtonVisibility() const;
 
-	EVisibility GetAddButtonVisibility() const;
-
-	FText GetAddButtonToolTipText() const;
-
-	FReply AddDirectlyButtonClicked();
-
 	FReply DeleteClicked();
 
-	TSharedRef<SWidget> GetAddMenu();
-
 	EVisibility GetStackIssuesWarningVisibility() const;
 
 	FText GetErrorButtonTooltipText() const;
 
 private:
 	UNiagaraStackItemGroup* Group;
-	TSharedPtr<class SComboButton> AddActionButton; 
-	const float TextIconSize = 16;
 };
diff --git a/Engine/Plugins/FX/Niagara/Source/NiagaraEditorWidgets/Private/Stack/SNiagaraStackItemGroupAddButton.cpp b/Engine/Plugins/FX/Niagara/Source/NiagaraEditorWidgets/Private/Stack/SNiagaraStackItemGroupAddButton.cpp
new file mode 100644
--- /dev/null
+++ b/Engine/Plugins/FX/Niagara/Source/NiagaraEditorWidgets/Private/Stack/SNiagaraStackItemGroupAddButton.cpp
@@ -0,0 +1,111 @@
+// Copyright 1998-2019 Epic Games, Inc. All Rights Reserved.
+
+#include "Stack/SNiagaraStackItemGroupAddButton.h"
+#include "Stack/SNiagaraStackItemGroupAddMenu.h"
+#include "ViewModels/Stack/INiagaraStackItemGroupAddUtilities.h"
+#include "ViewModels/Stack/NiagaraStackItemGroup.h"
+#include "NiagaraEditorWidgetsUtilities.h"
+#include "NiagaraEditorWidgetsStyle.h"
+
+#include "Widgets/Input/SComboButton.h"
+#include "Widgets/Input/SButton.h"
+#include "Widgets/Layout/SBox.h"
+#include "Widgets/Text/STextBlock.h"
+
+#define LOCTEXT_NAMESPACE "SNiagaraStackItemGroupAddButton"
+
+const float SNiagaraStackItemGroupAddButton::TextIconSize = 16;
+
+void SNiagaraStackItemGroupAddButton::Construct(const FArguments& InArgs, UNiagaraStackItemGroup& InStackItemGroup)
+{
+	TSharedPtr<SWidget> Content;
+	if(InStackItemGroup.GetAddUtilities() != nullptr)
+	{
+		INiagaraStackItemGroupAddUtilities* AddUtilities = InStackItemGroup.GetAddUtilities();
+		FText AddToGroupFormat = LOCTEXT("AddToGroupFormat", "Add a new {0} to this group.");
+		if (AddUtilities->GetAddMode() == INiagaraStackItemGroupAddUtilities::AddFromAction)
+		{
+			Content = SAssignNew(AddActionButton, SComboButton)
+				.ButtonStyle(FNiagaraEditorWidgetsStyle::Get(), "NiagaraEditor.Stack.AddButton")
+				.ButtonColorAndOpacity(FNiagaraEditorWidgetsStyle::Get().GetColor(FNiagaraStackEditorWidgetsUtilities::GetIconColorNameForExecutionCategory(InStackItemGroup.GetExecutionCategoryName())))
+				.ToolTipText(FText::Format(AddToGroupFormat, AddUtilities->GetAddItemName()))
+				.HasDownArrow(false)
+				.OnGetMenuContent(this, &SNiagaraStackItemGroupAddButton::GetAddMenu)
+				.ContentPadding(0)
+				.MenuPlacement(MenuPlacement_BelowRightAnchor)
+				.ButtonContent()
+				[
+					SNew(SBox)
+					.WidthOverride(InArgs._Width)
+					.HeightOverride(TextIconSize)
+					.HAlign(HAlign_Center)
+					.VAlign(VAlign_Center)
+					[
+						SNew(STextBlock)
+						//.TextStyle(FEditorStyle::Get(), "NormalText.Important")
+						.TextStyle(FNiagaraEditorWidgetsStyle::Get(), "NiagaraEditor.Stack.AddButtonText")
+						.Font(FEditorStyle::Get().GetFontStyle("FontAwesome.10"))
+						.Text(FText::FromString(FString(TEXT("\xf067"))) /*fa-plus*/)
+					]
+				];
+		}
+		else if (AddUtilities->GetAddMode() == INiagaraStackItemGroupAddUtilities::AddDirectly)
+		{
+			Content = SNew(SButton)
+				.ButtonStyle(FNiagaraEditorWidgetsStyle::Get(), "NiagaraEditor.Stack.AddButton")
+				.ButtonColorAndOpacity(FNiagaraEditorWidgetsStyle::Get().GetColor(FNiagaraStackEditorWidgetsUtilities::GetIconColorNameForExecutionCategory(InStackItemGroup.GetExecutionCategoryName())))
+				.ToolTipText(FText::Format(AddToGroupFormat, AddUtilities->GetAddItemName()))
+				.ContentPadding(0)
+				.OnClicked(this, &SNiagaraStackItemGroupAddButton::AddDirectlyButtonClicked)
+				.Content()
+				[
+					SNew(SBox)
+					.WidthOverride(InArgs._Width)
+					.HeightOverride(TextIconSize)
+					.HAlign(HAlign_Center)
+					.VAlign(VAlign_Center)
+					[
+						SNew(STextBlock)
+						//.TextStyle(FEditorStyle::Get(), "NormalText.Important")
+						.TextStyle(FNiagaraEditorWidgetsStyle::Get(), "NiagaraEditor.Stack.AddButtonText")
+						.Font(FEditorStyle::Get().GetFontStyle("FontAwesome.10"))
+						.Text(FText::FromString(FString(TEXT("\xf067"))) /*fa-plus*/)
+					]
+				];
+		}
+	}
+	else
+	{
+		// TODO Log error here.
+		Content = SNullWidget::NullWidget;
+	}
+
+	ChildSlot
+	[
+		Content.ToSharedRef()
+	];
+	
+	StackItemGroupWeak = &InStackItemGroup;
+}
+
+TSharedRef<SWidget> SNiagaraStackItemGroupAddButton::GetAddMenu()
+{
+	if(StackItemGroupWeak.IsValid())
+	{
+		TSharedRef<SNiagaraStackItemGroupAddMenu> AddMenu = SNew(SNiagaraStackItemGroupAddMenu, StackItemGroupWeak->GetAddUtilities(), INDEX_NONE);
+		AddActionButton->SetMenuContentWidgetToFocus(AddMenu->GetFilterTextBox()->AsShared());
+		return AddMenu;
+	}
+	return SNullWidget::NullWidget;
+}
+
+FReply SNiagaraStackItemGroupAddButton::AddDirectlyButtonClicked()
+{
+	if(StackItemGroupWeak.IsValid())
+	{
+		StackItemGroupWeak->GetAddUtilities()->AddItemDirectly();
+	}
+	return FReply::Handled();
+}
+
+#undef LOCTEXT_NAMESPACE
\ No newline at end of file
diff --git a/Engine/Plugins/FX/Niagara/Source/NiagaraEditorWidgets/Private/Stack/SNiagaraStackItemGroupAddButton.h b/Engine/Plugins/FX/Niagara/Source/NiagaraEditorWidgets/Private/Stack/SNiagaraStackItemGroupAddButton.h
new file mode 100644
--- /dev/null
+++ b/Engine/Plugins/FX/Niagara/Source/NiagaraEditorWidgets/Private/Stack/SNiagaraStackItemGroupAddButton.h
@@ -0,0 +1,30 @@
+// Copyright 1998-2019 Epic Games, Inc. All Rights Reserved.
+
+#pragma once
+
+#include "Widgets/SCompoundWidget.h"
+
+class UNiagaraStackItemGroup;
+class SComboButton;
+
+class SNiagaraStackItemGroupAddButton : public SCompoundWidget
+{
+public:
+	SLATE_BEGIN_ARGS(SNiagaraStackItemGroupAddButton)
+		:_Width(TextIconSize * 2)
+		{}
+		SLATE_ARGUMENT(float, Width)
+	SLATE_END_ARGS();
+
+	void Construct(const FArguments& InArgs, UNiagaraStackItemGroup& InStackItemGroup);
+
+private:
+	TSharedRef<SWidget> GetAddMenu();
+
+	FReply AddDirectlyButtonClicked();
+
+private:
+	TWeakObjectPtr<UNiagaraStackItemGroup> StackItemGroupWeak;
+	TSharedPtr<SComboButton> AddActionButton;
+	static const float TextIconSize;
+};
\ No newline at end of file
diff --git a/Engine/Plugins/FX/Niagara/Source/NiagaraEditorWidgets/Private/Stack/SNiagaraStackModuleItem.cpp b/Engine/Plugins/FX/Niagara/Source/NiagaraEditorWidgets/Private/Stack/SNiagaraStackModuleItem.cpp
--- a/Engine/Plugins/FX/Niagara/Source/NiagaraEditorWidgets/Private/Stack/SNiagaraStackModuleItem.cpp
+++ b/Engine/Plugins/FX/Niagara/Source/NiagaraEditorWidgets/Private/Stack/SNiagaraStackModuleItem.cpp
@@ -20,11 +20,11 @@
 #include "Widgets/Images/SImage.h"
 #include "Framework/MultiBox/MultiBoxBuilder.h"
 #include "Framework/Application/SlateApplication.h"
-
 #include "SDropTarget.h"
 #include "SNiagaraStackErrorButton.h"
 #include "NiagaraEditorUtilities.h"
 #include "SGraphActionMenu.h"
+#include "NiagaraEditorWidgetsUtilities.h"
 #include "Subsystems/AssetEditorSubsystem.h"
 #include "Editor.h"
 
@@ -149,27 +149,8 @@ bool SNiagaraStackModuleItem::CheckEnabledStatus(bool bIsEnabled)
 
 void SNiagaraStackModuleItem::FillRowContextMenu(FMenuBuilder& MenuBuilder)
 {
-	MenuBuilder.AddMenuEntry(
-		LOCTEXT("InsertModuleAbove", "Insert Above"),
-		LOCTEXT("InsertModuleAboveToolTip", "Insert a new module above this module in the stack."),
-		FSlateIcon(),
-		FUIAction(FExecuteAction::CreateSP(this, &SNiagaraStackModuleItem::InsertModuleAbove)));
-	MenuBuilder.AddMenuEntry(
-		LOCTEXT("InsertModuleBelow", "Insert Below"),
-		LOCTEXT("InsertModuleBelowToolTip", "Insert a new module below this module in the stack."),
-		FSlateIcon(),
-		FUIAction(FExecuteAction::CreateSP(this, &SNiagaraStackModuleItem::InsertModuleBelow)));
-
-	FUIAction Action(FExecuteAction::CreateSP(this, &SNiagaraStackModuleItem::SetEnabled, !ModuleItem->GetIsEnabled()),
-		FCanExecuteAction(),
-		FIsActionChecked::CreateSP(this, &SNiagaraStackModuleItem::CheckEnabledStatus, true));
-	MenuBuilder.AddMenuEntry(
-		LOCTEXT("IsEnabled", "Is Enabled"),
-		LOCTEXT("ToggleModuleEnabledToolTip", "Toggle module enabled/disabled state"),
-		FSlateIcon(),
-		Action,
-		NAME_None,
-		EUserInterfaceActionType::Check);
+	FNiagaraStackEditorWidgetsUtilities::AddStackModuleItemContextMenuActions(MenuBuilder, *ModuleItem, this->AsShared());
+	FNiagaraStackEditorWidgetsUtilities::AddStackItemContextMenuActions(MenuBuilder, *ModuleItem);
 }
 
 FReply SNiagaraStackModuleItem::OnMouseButtonDoubleClick(const FGeometry& InMyGeometry, const FPointerEvent& InMouseEvent)
@@ -209,20 +190,15 @@ bool SNiagaraStackModuleItem::GetButtonsEnabled() const
 
 FText SNiagaraStackModuleItem::GetDeleteButtonToolTipText() const
 {
-	if (ModuleItem->CanMoveAndDelete())
-	{
-		return LOCTEXT("DeleteToolTip", "Delete this module");
-	}
-	else
-	{
-		// TODO: This message should be handled by the entry.
-		return LOCTEXT("CantDeleteToolTip", "This module can not be deleted becaue it is inherited.");
-	}
+	FText CanDeleteMessage;
+	ModuleItem->TestCanDeleteWithMessage(CanDeleteMessage);
+	return CanDeleteMessage;
 }
 
 bool SNiagaraStackModuleItem::GetDeleteButtonEnabled() const
 {
-	return ModuleItem->GetOwnerIsEnabled() && ModuleItem->CanMoveAndDelete();
+	FText CanDeleteMessage;
+	return ModuleItem->TestCanDeleteWithMessage(CanDeleteMessage);
 }
 
 bool SNiagaraStackModuleItem::GetEnabledCheckBoxEnabled() const
@@ -294,25 +270,6 @@ FReply SNiagaraStackModuleItem::RefreshClicked()
 	return FReply::Handled();
 }
 
-void SNiagaraStackModuleItem::InsertModuleAbove()
-{
-	ShowInsertModuleMenu(ModuleItem->GetModuleIndex());
-}
-
-void SNiagaraStackModuleItem::InsertModuleBelow()
-{
-	ShowInsertModuleMenu(ModuleItem->GetModuleIndex() + 1);
-}
-
-void SNiagaraStackModuleItem::ShowInsertModuleMenu(int32 InsertIndex)
-{
-	TSharedRef<SWidget> MenuContent = SNew(SNiagaraStackItemGroupAddMenu, ModuleItem->GetGroupAddUtilities(), InsertIndex);
-	FGeometry ThisGeometry = GetCachedGeometry();
-	bool bAutoAdjustForDpiScale = false; // Don't adjust for dpi scale because the push menu command is expecting an unscaled position.
-	FVector2D MenuPosition = FSlateApplication::Get().CalculatePopupWindowPosition(ThisGeometry.GetLayoutBoundingRect(), MenuContent->GetDesiredSize(), bAutoAdjustForDpiScale);
-	FSlateApplication::Get().PushMenu(AsShared(), FWidgetPath(), MenuContent, MenuPosition, FPopupTransitionEffect::ContextMenu);
-}
-
 FReply SNiagaraStackModuleItem::OnModuleItemDrop(TSharedPtr<FDragDropOperation> DragDropOperation)
 {
 	if (DragDropOperation->IsOfType<FNiagaraStackDragOperation>())
diff --git a/Engine/Plugins/FX/Niagara/Source/NiagaraEditorWidgets/Private/Stack/SNiagaraStackModuleItem.h b/Engine/Plugins/FX/Niagara/Source/NiagaraEditorWidgets/Private/Stack/SNiagaraStackModuleItem.h
--- a/Engine/Plugins/FX/Niagara/Source/NiagaraEditorWidgets/Private/Stack/SNiagaraStackModuleItem.h
+++ b/Engine/Plugins/FX/Niagara/Source/NiagaraEditorWidgets/Private/Stack/SNiagaraStackModuleItem.h
@@ -52,12 +52,6 @@ private:
 
 	FReply RefreshClicked();
 
-	void InsertModuleAbove();
-
-	void InsertModuleBelow();
-
-	void ShowInsertModuleMenu(int32 InsertIndex);
-
 	FReply OnModuleItemDrop(TSharedPtr<class FDragDropOperation> DragDropOperation);
 
 	bool OnModuleItemAllowDrop(TSharedPtr<class FDragDropOperation> DragDropOperation);
diff --git a/Engine/Plugins/FX/Niagara/Source/NiagaraEditorWidgets/Private/Stack/SNiagaraStackRendererItem.cpp b/Engine/Plugins/FX/Niagara/Source/NiagaraEditorWidgets/Private/Stack/SNiagaraStackRendererItem.cpp
--- a/Engine/Plugins/FX/Niagara/Source/NiagaraEditorWidgets/Private/Stack/SNiagaraStackRendererItem.cpp
+++ b/Engine/Plugins/FX/Niagara/Source/NiagaraEditorWidgets/Private/Stack/SNiagaraStackRendererItem.cpp
@@ -56,24 +56,6 @@ void SNiagaraStackRendererItem::Construct(const FArguments& InArgs, UNiagaraStac
 			.Visibility(this, &SNiagaraStackRendererItem::GetStackIssuesWarningVisibility)
 			.OnButtonClicked(this, &SNiagaraStackRendererItem::ExpandEntry)
 		]
-		// Delete button
-		+ SHorizontalBox::Slot()
-		.AutoWidth()
-		[
-			SNew(SButton)
-			.ButtonStyle(FEditorStyle::Get(), "HoverHintOnly")
-			.IsFocusable(false)
-			.ForegroundColor(FNiagaraEditorWidgetsStyle::Get().GetColor("NiagaraEditor.Stack.ForegroundColor"))
-			.ToolTipText(LOCTEXT("DeleteRendererToolTip", "Delete this Renderer"))
-			.Visibility(this, &SNiagaraStackRendererItem::GetDeleteButtonVisibility)
-			.OnClicked(this, &SNiagaraStackRendererItem::DeleteClicked)
-			.Content()
-			[
-				SNew(STextBlock)
-				.Font(FEditorStyle::Get().GetFontStyle("FontAwesome.10"))
-				.Text(FText::FromString(FString(TEXT("\xf1f8"))))
-			]
-		]
 		// Reset to base Button
 		+ SHorizontalBox::Slot()
 		.AutoWidth()
@@ -94,6 +76,24 @@ void SNiagaraStackRendererItem::Construct(const FArguments& InArgs, UNiagaraStac
 				.ColorAndOpacity(FSlateColor(FLinearColor::Green))
 			]
 		]
+		// Delete button
+		+ SHorizontalBox::Slot()
+		.AutoWidth()
+		[
+			SNew(SButton)
+			.ButtonStyle(FEditorStyle::Get(), "HoverHintOnly")
+			.IsFocusable(false)
+			.ForegroundColor(FNiagaraEditorWidgetsStyle::Get().GetColor("NiagaraEditor.Stack.FlatButtonColor"))
+			.ToolTipText(this, &SNiagaraStackRendererItem::GetDeleteButtonToolTipText)
+			.IsEnabled(this, &SNiagaraStackRendererItem::GetDeleteButtonEnabled)
+			.OnClicked(this, &SNiagaraStackRendererItem::DeleteClicked)
+			.Content()
+			[
+				SNew(STextBlock)
+				.Font(FEditorStyle::Get().GetFontStyle("FontAwesome.10"))
+				.Text(FText::FromString(FString(TEXT("\xf1f8"))))
+			]
+		]
 		// Enabled checkbox
 		+ SHorizontalBox::Slot()
 		.Padding(2, 0, 0, 0)
@@ -106,9 +106,17 @@ void SNiagaraStackRendererItem::Construct(const FArguments& InArgs, UNiagaraStac
 	];
 }
 
-EVisibility SNiagaraStackRendererItem::GetDeleteButtonVisibility() const
+FText SNiagaraStackRendererItem::GetDeleteButtonToolTipText() const
+{
+	FText CanDeleteMessage;
+	RendererItem->TestCanDeleteWithMessage(CanDeleteMessage);
+	return CanDeleteMessage;
+}
+
+bool SNiagaraStackRendererItem::GetDeleteButtonEnabled() const
 {
-	return RendererItem->CanDelete() ? EVisibility::Visible : EVisibility::Collapsed;
+	FText CanDeleteMessage;
+	return RendererItem->TestCanDeleteWithMessage(CanDeleteMessage);
 }
 
 FReply SNiagaraStackRendererItem::DeleteClicked()
diff --git a/Engine/Plugins/FX/Niagara/Source/NiagaraEditorWidgets/Private/Stack/SNiagaraStackRendererItem.h b/Engine/Plugins/FX/Niagara/Source/NiagaraEditorWidgets/Private/Stack/SNiagaraStackRendererItem.h
--- a/Engine/Plugins/FX/Niagara/Source/NiagaraEditorWidgets/Private/Stack/SNiagaraStackRendererItem.h
+++ b/Engine/Plugins/FX/Niagara/Source/NiagaraEditorWidgets/Private/Stack/SNiagaraStackRendererItem.h
@@ -18,7 +18,9 @@ public:
 	void Construct(const FArguments& InArgs, UNiagaraStackRendererItem& InRendererItem, UNiagaraStackViewModel* InStackViewModel);
 	
 private:
-	EVisibility GetDeleteButtonVisibility() const;
+	FText GetDeleteButtonToolTipText() const;
+
+	bool GetDeleteButtonEnabled() const;
 
 	FReply DeleteClicked();
 
diff --git a/Engine/Plugins/FX/Niagara/Source/NiagaraEditorWidgets/Private/Stack/SNiagaraStackTableRow.cpp b/Engine/Plugins/FX/Niagara/Source/NiagaraEditorWidgets/Private/Stack/SNiagaraStackTableRow.cpp
--- a/Engine/Plugins/FX/Niagara/Source/NiagaraEditorWidgets/Private/Stack/SNiagaraStackTableRow.cpp
+++ b/Engine/Plugins/FX/Niagara/Source/NiagaraEditorWidgets/Private/Stack/SNiagaraStackTableRow.cpp
@@ -11,12 +11,7 @@
 #include "Widgets/Input/SButton.h"
 #include "Widgets/Images/SImage.h"
 #include "Widgets/Layout/SBox.h"
-#include "IContentBrowserSingleton.h"
-#include "ContentBrowserModule.h"
-
 #include "Framework/Application/SlateApplication.h"
-#include "Subsystems/AssetEditorSubsystem.h"
-#include "Editor.h"
 
 #define LOCTEXT_NAMESPACE "NiagaraStackTableRow"
 
@@ -280,27 +275,13 @@ FReply SNiagaraStackTableRow::OnMouseButtonUp(const FGeometry& MyGeometry, const
 	if (MouseEvent.GetEffectingButton() == EKeys::RightMouseButton)
 	{
 		FMenuBuilder MenuBuilder(true, nullptr);
-		MenuBuilder.BeginSection("ModuleActions", LOCTEXT("ModuleActions", "Module Actions"));
 		for (FOnFillRowContextMenu& OnFillRowContextMenuHandler : OnFillRowContextMenuHanders)
 		{
 			OnFillRowContextMenuHandler.ExecuteIfBound(MenuBuilder);
 		}
-		MenuBuilder.EndSection();
-		if (StackEntry->GetExternalAsset() != nullptr)
-		{
-			MenuBuilder.BeginSection("AssetActions", LOCTEXT("AssetActions", "Asset Actions"));
-			MenuBuilder.AddMenuEntry(
-				LOCTEXT("OpenAndFocusAsset", "Open and Focus Asset"),
-				FText::Format(LOCTEXT("OpenAndFocusAssetTooltip", "Open {0} in separate editor"), StackEntry->GetDisplayName()),
-				FSlateIcon(),
-				FUIAction(FExecuteAction::CreateSP(this, &SNiagaraStackTableRow::OpenSourceAsset)));
-			MenuBuilder.AddMenuEntry(
-				LOCTEXT("ShowAssetInContentBrowser", "Show in Content Browser"),
-				FText::Format(LOCTEXT("ShowAssetInContentBrowserToolTip", "Navigate to {0} in the Content Browser window"), StackEntry->GetDisplayName()),
-				FSlateIcon(),
-				FUIAction(FExecuteAction::CreateSP(this, &SNiagaraStackTableRow::ShowAssetInContentBrowser)));
-			MenuBuilder.EndSection();
-		}
+
+		FNiagaraStackEditorWidgetsUtilities::AddStackEntryAssetContextMenuActions(MenuBuilder, *StackEntry);
+
 		TArray<UNiagaraStackEntry*> EntriesToProcess;
 		TArray<UNiagaraStackEntry*> NavigationEntries;
 		StackViewModel->GetPathForEntry(StackEntry, EntriesToProcess);
@@ -467,17 +448,4 @@ void SNiagaraStackTableRow::NavigateTo(UNiagaraStackEntry* Item)
 	OwnerTree->RequestNavigateToItem(Item, 0);
 }
 
-void SNiagaraStackTableRow::OpenSourceAsset()
-{
-	GEditor->GetEditorSubsystem<UAssetEditorSubsystem>()->OpenEditorForAsset(StackEntry->GetExternalAsset());
-}
-
-void SNiagaraStackTableRow::ShowAssetInContentBrowser()
-{
-	FContentBrowserModule& ContentBrowserModule = FModuleManager::Get().LoadModuleChecked<FContentBrowserModule>(TEXT("ContentBrowser"));
-	TArray<FAssetData> Assets;
-	Assets.Add(FAssetData(StackEntry->GetExternalAsset()));
-	ContentBrowserModule.Get().SyncBrowserToAssets(Assets);
-}
-
 #undef LOCTEXT_NAMESPACE
\ No newline at end of file
diff --git a/Engine/Plugins/FX/Niagara/Source/NiagaraEditorWidgets/Private/Stack/SNiagaraStackTableRow.h b/Engine/Plugins/FX/Niagara/Source/NiagaraEditorWidgets/Private/Stack/SNiagaraStackTableRow.h
--- a/Engine/Plugins/FX/Niagara/Source/NiagaraEditorWidgets/Private/Stack/SNiagaraStackTableRow.h
+++ b/Engine/Plugins/FX/Niagara/Source/NiagaraEditorWidgets/Private/Stack/SNiagaraStackTableRow.h
@@ -91,10 +91,6 @@ private:
 
 	void NavigateTo(UNiagaraStackEntry* Item);
 
-	void OpenSourceAsset();
-
-	void ShowAssetInContentBrowser();
-
 private:
 	UNiagaraStackViewModel* StackViewModel;
 	UNiagaraStackEntry* StackEntry;
