Niagara - Fix enable/disable state for add buttons and stack category icons.

diff --git a/Engine/Plugins/FX/Niagara/Source/NiagaraEditorWidgets/Private/Stack/SNiagaraStackItemGroupAddButton.cpp b/Engine/Plugins/FX/Niagara/Source/NiagaraEditorWidgets/Private/Stack/SNiagaraStackItemGroupAddButton.cpp
--- a/Engine/Plugins/FX/Niagara/Source/NiagaraEditorWidgets/Private/Stack/SNiagaraStackItemGroupAddButton.cpp
+++ b/Engine/Plugins/FX/Niagara/Source/NiagaraEditorWidgets/Private/Stack/SNiagaraStackItemGroupAddButton.cpp
@@ -31,6 +31,7 @@ void SNiagaraStackItemGroupAddButton::Construct(const FArguments& InArgs, UNiaga
 				.ToolTipText(FText::Format(AddToGroupFormat, AddUtilities->GetAddItemName()))
 				.HasDownArrow(false)
 				.OnGetMenuContent(this, &SNiagaraStackItemGroupAddButton::GetAddMenu)
+				.IsEnabled_UObject(&InStackItemGroup, &UNiagaraStackEntry::GetIsEnabledAndOwnerIsEnabled)
 				.ContentPadding(0)
 				.MenuPlacement(MenuPlacement_BelowRightAnchor)
 				.ButtonContent()
@@ -55,6 +56,7 @@ void SNiagaraStackItemGroupAddButton::Construct(const FArguments& InArgs, UNiaga
 				.ButtonStyle(FNiagaraEditorWidgetsStyle::Get(), "NiagaraEditor.Stack.AddButton")
 				.ButtonColorAndOpacity(FNiagaraEditorWidgetsStyle::Get().GetColor(FNiagaraStackEditorWidgetsUtilities::GetIconColorNameForExecutionCategory(InStackItemGroup.GetExecutionCategoryName())))
 				.ToolTipText(FText::Format(AddToGroupFormat, AddUtilities->GetAddItemName()))
+				.IsEnabled_UObject(&InStackItemGroup, &UNiagaraStackEntry::GetIsEnabledAndOwnerIsEnabled)
 				.ContentPadding(0)
 				.OnClicked(this, &SNiagaraStackItemGroupAddButton::AddDirectlyButtonClicked)
 				.Content()
diff --git a/Engine/Plugins/FX/Niagara/Source/NiagaraEditorWidgets/Private/Stack/SNiagaraStackTableRow.cpp b/Engine/Plugins/FX/Niagara/Source/NiagaraEditorWidgets/Private/Stack/SNiagaraStackTableRow.cpp
--- a/Engine/Plugins/FX/Niagara/Source/NiagaraEditorWidgets/Private/Stack/SNiagaraStackTableRow.cpp
+++ b/Engine/Plugins/FX/Niagara/Source/NiagaraEditorWidgets/Private/Stack/SNiagaraStackTableRow.cpp
@@ -136,6 +136,7 @@ void SNiagaraStackTableRow::SetNameAndValueContent(TSharedRef<SWidget> InNameWid
 		.VAlign(EVerticalAlignment::VAlign_Center)
 		.ToolTipText(ExecutionCategoryToolTipText)
 		.Visibility(this, &SNiagaraStackTableRow::GetExecutionCategoryIconVisibility)
+		.IsEnabled_UObject(StackEntry, &UNiagaraStackEntry::GetIsEnabledAndOwnerIsEnabled)
 		[
 			SNew(SImage)
 			.Visibility(this, &SNiagaraStackTableRow::GetExecutionCategoryIconVisibility)
