#Niagara Nonunity fix

diff --git a/Engine/Plugins/FX/Niagara/Source/NiagaraEditorWidgets/Private/NiagaraEditorWidgetsUtilities.h b/Engine/Plugins/FX/Niagara/Source/NiagaraEditorWidgets/Private/NiagaraEditorWidgetsUtilities.h
--- a/Engine/Plugins/FX/Niagara/Source/NiagaraEditorWidgets/Private/NiagaraEditorWidgetsUtilities.h
+++ b/Engine/Plugins/FX/Niagara/Source/NiagaraEditorWidgets/Private/NiagaraEditorWidgetsUtilities.h
@@ -2,8 +2,11 @@
 
 #pragma once
 
+#include "Templates/SharedPointer.h"
 #include "UObject/NameTypes.h"
 
+class FMenuBuilder;
+class SWidget;
 class UNiagaraStackEntry;
 class UNiagaraStackItem;
 class UNiagaraStackModuleItem;
