- Implement undo/redo for NiagaraScriptVariableCustomization default values.
 - Make details panel headers more consistent with Blueprints.

diff --git a/Engine/Plugins/FX/Niagara/Source/Niagara/Public/NiagaraTypes.h b/Engine/Plugins/FX/Niagara/Source/Niagara/Public/NiagaraTypes.h
--- a/Engine/Plugins/FX/Niagara/Source/Niagara/Public/NiagaraTypes.h
+++ b/Engine/Plugins/FX/Niagara/Source/Niagara/Public/NiagaraTypes.h
@@ -357,11 +357,11 @@ public:
 	bool bInlineEditConditionToggle;
 
 	/** Declares the associated input should be conditionally editable based on the value of another input. */
-	UPROPERTY(EditAnywhere, Category = "Input Conditions")
+	UPROPERTY(EditAnywhere, Category = "Variable")
 	FNiagaraInputConditionMetadata EditCondition;
 
 	/** Declares the associated input should be conditionally visible based on the value of another input. */
-	UPROPERTY(EditAnywhere, Category = "Input Conditions")
+	UPROPERTY(EditAnywhere, Category = "Variable")
 	FNiagaraInputConditionMetadata VisibleCondition;
 
 	UPROPERTY(EditAnywhere, Category = "Variable", DisplayName = "Property Metadata", meta = (ToolTip = "Property Metadata"))
diff --git a/Engine/Plugins/FX/Niagara/Source/NiagaraEditor/Private/Customizations/NiagaraScriptVariableCustomization.cpp b/Engine/Plugins/FX/Niagara/Source/NiagaraEditor/Private/Customizations/NiagaraScriptVariableCustomization.cpp
--- a/Engine/Plugins/FX/Niagara/Source/NiagaraEditor/Private/Customizations/NiagaraScriptVariableCustomization.cpp
+++ b/Engine/Plugins/FX/Niagara/Source/NiagaraEditor/Private/Customizations/NiagaraScriptVariableCustomization.cpp
@@ -20,22 +20,81 @@
 #include "TypeEditorUtilities/NiagaraVectorTypeEditorUtilities.h"
 #include "TypeEditorUtilities/NiagaraColorTypeEditorUtilities.h"
 #include "NiagaraEditorStyle.h"
- 
+#include "ScopedTransaction.h"
+
 #define LOCTEXT_NAMESPACE "NiagaraScriptVariableVariableDetails"
- 
+
 TSharedRef<IDetailCustomization> FNiagaraScriptVariableDetails::MakeInstance()
 {
-	return MakeShareable(new FNiagaraScriptVariableDetails);
+	return MakeShareable(new FNiagaraScriptVariableDetails());
 }
  
 FNiagaraScriptVariableDetails::FNiagaraScriptVariableDetails()
 {
- 
+	GEditor->RegisterForUndo(this);
+}
+
+FNiagaraScriptVariableDetails::~FNiagaraScriptVariableDetails()
+{
+	GEditor->UnregisterForUndo(this);
+}
+
+
+UEdGraphPin* FNiagaraScriptVariableDetails::GetDefaultPin()
+{
+	// TODO: We don't know the usage at this point, so we'll try each script type in order
+	//       This could probably be made much more robust, but works for now.
+	if (UNiagaraGraph* Graph = Cast<UNiagaraGraph>(Variable->GetOuter())) 
+	{
+		UEdGraphPin* Pin = Graph->FindParameterMapDefaultValuePin(Variable->Variable.GetName(), ENiagaraScriptUsage::Module, ENiagaraScriptUsage::Module);
+		if (Pin == nullptr)
+		{
+			Pin = Graph->FindParameterMapDefaultValuePin(Variable->Variable.GetName(), ENiagaraScriptUsage::DynamicInput, ENiagaraScriptUsage::Module);
+		}
+		if (Pin == nullptr)
+		{
+			Pin = Graph->FindParameterMapDefaultValuePin(Variable->Variable.GetName(), ENiagaraScriptUsage::Function, ENiagaraScriptUsage::Module);
+		}
+		return Pin;
+	}
+	return nullptr;
+}
+
+void FNiagaraScriptVariableDetails::PostUndo(bool bSuccess)
+{
+	if (Variable == nullptr)
+	{
+		return;
+	}
+
+	if (Variable->Metadata.bIsStaticSwitch)
+	{
+		if (TypeUtilityStaticSwitchValue && ParameterEditorStaticSwitchValue)
+		{
+			TSharedPtr<FStructOnScope> ParameterValue = MakeShareable(new FStructOnScope(Variable->Variable.GetType().GetStruct()));
+			Variable->Variable.SetValue(Variable->Metadata.StaticSwitchDefaultValue);
+			Variable->Variable.CopyTo(ParameterValue->GetStructMemory());
+			ParameterEditorStaticSwitchValue->UpdateInternalValueFromStruct(ParameterValue.ToSharedRef());
+		}
+	}
+	else if (Variable->Variable.GetName().ToString().StartsWith("Module."))
+	{
+		if (UEdGraphPin* Pin = GetDefaultPin())
+		{
+			if (TypeUtilityValue && ParameterEditorValue)
+			{
+				TypeUtilityValue->SetValueFromPinDefaultString(Pin->DefaultValue, Variable->Variable);
+				TSharedPtr<FStructOnScope> ParameterValue = MakeShareable(new FStructOnScope(Variable->Variable.GetType().GetStruct()));
+				Variable->Variable.CopyTo(ParameterValue->GetStructMemory());
+				ParameterEditorValue->UpdateInternalValueFromStruct(ParameterValue.ToSharedRef());
+			}
+		}
+	}
 }
  
 void FNiagaraScriptVariableDetails::CustomizeDetails(IDetailLayoutBuilder& DetailBuilder)
 {
-	static const FName CategoryName = TEXT("Script Variable Data");
+	static const FName CategoryName = TEXT("Default Value");
  
 	TArray<TWeakObjectPtr<UObject>> ObjectsCustomized;
 	DetailBuilder.GetObjectsBeingCustomized(ObjectsCustomized);
@@ -65,7 +124,6 @@ void FNiagaraScriptVariableDetails::CustomizeDetails(IDetailLayoutBuilder& Detai
 		DetailBuilder.HideProperty(GET_MEMBER_NAME_CHECKED(UNiagaraScriptVariable, Metadata.EditCondition));
 		DetailBuilder.HideProperty(GET_MEMBER_NAME_CHECKED(UNiagaraScriptVariable, Metadata.VisibleCondition));
  
-		// define a custom parameter editor widget for the static switch data
 		TypeUtilityStaticSwitchValue = EditorModule.GetTypeUtilities(Variable->Variable.GetType());
 		if (TypeUtilityStaticSwitchValue && TypeUtilityStaticSwitchValue->CanCreateParameterEditor())
 		{
@@ -82,6 +140,7 @@ void FNiagaraScriptVariableDetails::CustomizeDetails(IDetailLayoutBuilder& Detai
 			.NameContent()
 			[
 				SNew(STextBlock)
+				.Font(FNiagaraEditorStyle::Get().GetFontStyle("NiagaraEditor.ParameterFont"))
 				.Text(FText::FromString(TEXT("Default Value")))
 			]
 			.ValueContent()
@@ -97,49 +156,38 @@ void FNiagaraScriptVariableDetails::CustomizeDetails(IDetailLayoutBuilder& Detai
 	}
 	else if (Variable->Variable.GetName().ToString().StartsWith("Module.")) 
 	{
-		if (UNiagaraGraph* Graph = Cast<UNiagaraGraph>(Variable->GetOuter())) {
-			// TODO: Probably make a custom function for this...
-			UEdGraphPin* Pin = Graph->FindParameterMapDefaultValuePin(Variable->Variable.GetName(), ENiagaraScriptUsage::Module, ENiagaraScriptUsage::Module);
-			if (Pin == nullptr) 
-			{
-				Pin = Graph->FindParameterMapDefaultValuePin(Variable->Variable.GetName(), ENiagaraScriptUsage::DynamicInput, ENiagaraScriptUsage::Module);
-			}
-			if (Pin == nullptr) 
+		if (UEdGraphPin* Pin = GetDefaultPin())
+		{
+			TypeUtilityValue = EditorModule.GetTypeUtilities(Variable->Variable.GetType());
+			if (TypeUtilityValue && TypeUtilityValue->CanCreateParameterEditor())
 			{
-				Pin = Graph->FindParameterMapDefaultValuePin(Variable->Variable.GetName(), ENiagaraScriptUsage::Function, ENiagaraScriptUsage::Module);
+				ParameterEditorValue = TypeUtilityValue->CreateParameterEditor(Variable->Variable.GetType());
+
+				TypeUtilityValue->SetValueFromPinDefaultString(Pin->DefaultValue, Variable->Variable);
+				TSharedPtr<FStructOnScope> ParameterValue = MakeShareable(new FStructOnScope(Variable->Variable.GetType().GetStruct()));
+				Variable->Variable.CopyTo(ParameterValue->GetStructMemory());
+				ParameterEditorValue->UpdateInternalValueFromStruct(ParameterValue.ToSharedRef());
+				ParameterEditorValue->SetOnValueChanged(SNiagaraParameterEditor::FOnValueChange::CreateSP(this, &FNiagaraScriptVariableDetails::OnValueChanged));
+				ParameterEditorValue->SetOnBeginValueChange(SNiagaraParameterEditor::FOnValueChange::CreateSP(this, &FNiagaraScriptVariableDetails::OnBeginValueChanged));
+				ParameterEditorValue->SetOnEndValueChange(SNiagaraParameterEditor::FOnValueChange::CreateSP(this, &FNiagaraScriptVariableDetails::OnEndValueChanged));
+				
+				FDetailWidgetRow& DefaultValueWidget = CategoryBuilder.AddCustomRow(LOCTEXT("DefaultValueFilterText", "Default Value"));
+				DefaultValueWidget
+				.NameContent()
+				[
+					SNew(STextBlock)
+					.Font(FNiagaraEditorStyle::Get().GetFontStyle("NiagaraEditor.ParameterFont"))
+					.Text(FText::FromString(TEXT("Default Value")))
+				]
+				.ValueContent()
+				.HAlign(HAlign_Fill)
+				[
+					ParameterEditorValue.ToSharedRef()
+				];
 			}
-			if (Pin)
+			else
 			{
-				// define a custom parameter editor widget for module variables
-				TypeUtilityValue = EditorModule.GetTypeUtilities(Variable->Variable.GetType());
-				if (TypeUtilityValue && TypeUtilityValue->CanCreateParameterEditor())
-				{
-					ParameterEditorValue = TypeUtilityValue->CreateParameterEditor(Variable->Variable.GetType());
- 
-					TypeUtilityValue->SetValueFromPinDefaultString(Pin->DefaultValue, Variable->Variable);
-					TSharedPtr<FStructOnScope> ParameterValue = MakeShareable(new FStructOnScope(Variable->Variable.GetType().GetStruct()));
-					Variable->Variable.CopyTo(ParameterValue->GetStructMemory());
-					ParameterEditorValue->UpdateInternalValueFromStruct(ParameterValue.ToSharedRef());
-					ParameterEditorValue->SetOnValueChanged(SNiagaraParameterEditor::FOnValueChange::CreateSP(this, &FNiagaraScriptVariableDetails::OnValueChanged));
-					
-					FDetailWidgetRow& DefaultValueWidget = CategoryBuilder.AddCustomRow(LOCTEXT("DefaultValueFilterText", "Default Value"));
-					DefaultValueWidget
-					.NameContent()
-					[
-						SNew(STextBlock)
-						.Font(FNiagaraEditorStyle::Get().GetFontStyle("NiagaraEditor.ParameterFont"))
-						.Text(FText::FromString(TEXT("Default Value")))
-					]
-					.ValueContent()
-					.HAlign(HAlign_Fill)
-					[
-						ParameterEditorValue.ToSharedRef()
-					];
-				}
-				else
-				{
-					TypeUtilityValue = nullptr;
-				}
+				TypeUtilityValue = nullptr;
 			}
 		}
 	}
@@ -164,35 +212,71 @@ void FNiagaraScriptVariableDetails::CustomizeDetails(IDetailLayoutBuilder& Detai
  
 void FNiagaraScriptVariableDetails::OnValueChanged()
 {
+
 	if (TypeUtilityValue && ParameterEditorValue)
 	{
-		if (UNiagaraGraph* Graph = Cast<UNiagaraGraph>(Variable->GetOuter()))
+		if (UEdGraphPin* Pin = GetDefaultPin())
 		{
-			UEdGraphPin* Pin = Graph->FindParameterMapDefaultValuePin(Variable->Variable.GetName(), ENiagaraScriptUsage::Module, ENiagaraScriptUsage::Module);
-			if (Pin == nullptr) 
-			{
-				Pin = Graph->FindParameterMapDefaultValuePin(Variable->Variable.GetName(), ENiagaraScriptUsage::DynamicInput, ENiagaraScriptUsage::Module);
-			}
-			if (Pin == nullptr) 
-			{
-				Pin = Graph->FindParameterMapDefaultValuePin(Variable->Variable.GetName(), ENiagaraScriptUsage::Function, ENiagaraScriptUsage::Module);
-			}
-			if (Pin)
+			if (!ParameterEditorValue->CanChangeContinuously())
 			{
+				const FScopedTransaction Transaction( NSLOCTEXT("ScriptVariableCustomization", "ChangeValue", "Change Default Value" ) );
+				Pin->Modify();
+				Variable->Modify();
+
 				TSharedPtr<FStructOnScope> ParameterValue = MakeShareable(new FStructOnScope(Variable->Variable.GetType().GetStruct()));
 				ParameterEditorValue->UpdateStructFromInternalValue(ParameterValue.ToSharedRef());
 				Variable->Variable.SetData(ParameterValue->GetStructMemory());
- 
 				Pin->DefaultValue = TypeUtilityValue->GetPinDefaultStringFromValue(Variable->Variable);
+			}		
+			else
+			{
+				TSharedPtr<FStructOnScope> ParameterValue = MakeShareable(new FStructOnScope(Variable->Variable.GetType().GetStruct()));
+				ParameterEditorValue->UpdateStructFromInternalValue(ParameterValue.ToSharedRef());
+				Variable->Variable.SetData(ParameterValue->GetStructMemory());
+				Pin->DefaultValue = TypeUtilityValue->GetPinDefaultStringFromValue(Variable->Variable);	
 			}
+			
 		}
 	} 
 }
  
+
+void FNiagaraScriptVariableDetails::OnBeginValueChanged()
+{
+	if (!ParameterEditorValue->CanChangeContinuously())
+	{
+		return;
+	}
+
+	if (TypeUtilityValue && ParameterEditorValue)
+	{
+		if (UEdGraphPin* Pin = GetDefaultPin())
+		{
+			GEditor->BeginTransaction(NSLOCTEXT("ScriptVariableCustomization", "ChangeValue", "Change Default Value"));
+			Variable->Modify();
+			Pin->Modify();
+			TSharedPtr<FStructOnScope> ParameterValue = MakeShareable(new FStructOnScope(Variable->Variable.GetType().GetStruct()));
+			ParameterEditorValue->UpdateStructFromInternalValue(ParameterValue.ToSharedRef());
+			Variable->Variable.SetData(ParameterValue->GetStructMemory());
+			Pin->DefaultValue = TypeUtilityValue->GetPinDefaultStringFromValue(Variable->Variable);
+		}
+	} 
+}
+
+void FNiagaraScriptVariableDetails::OnEndValueChanged()
+{
+	if (GEditor->IsTransactionActive())
+	{
+		GEditor->EndTransaction();
+	}
+}
+
 void FNiagaraScriptVariableDetails::OnStaticSwitchValueChanged()
 {
 	if (TypeUtilityStaticSwitchValue && ParameterEditorStaticSwitchValue)
 	{
+		const FScopedTransaction Transaction(NSLOCTEXT("ScriptVariableCustomization", "ChangeStaticSwitchValue", "Change Static Switch Default Value"));
+		Variable->Modify();
 		TSharedPtr<FStructOnScope> ParameterValue = MakeShareable(new FStructOnScope(Variable->Variable.GetType().GetStruct()));
 		ParameterEditorStaticSwitchValue->UpdateStructFromInternalValue(ParameterValue.ToSharedRef());
 		Variable->Variable.SetData(ParameterValue->GetStructMemory());
diff --git a/Engine/Plugins/FX/Niagara/Source/NiagaraEditor/Private/Customizations/NiagaraScriptVariableCustomization.h b/Engine/Plugins/FX/Niagara/Source/NiagaraEditor/Private/Customizations/NiagaraScriptVariableCustomization.h
--- a/Engine/Plugins/FX/Niagara/Source/NiagaraEditor/Private/Customizations/NiagaraScriptVariableCustomization.h
+++ b/Engine/Plugins/FX/Niagara/Source/NiagaraEditor/Private/Customizations/NiagaraScriptVariableCustomization.h
@@ -5,31 +5,35 @@
 #include "CoreMinimal.h"
 #include "IDetailCustomization.h"
 #include "IPropertyUtilities.h"
+#include "EditorUndoClient.h"
  
 /** This customization sets up a custom details panel for the static switch Variable in the niagara module graph. */
-class FNiagaraScriptVariableDetails : public IDetailCustomization
+class FNiagaraScriptVariableDetails : public IDetailCustomization, public FEditorUndoClient
 {
 public:
  
 	/** Makes a new instance of this detail layout class for a specific detail view requesting it */
 	static TSharedRef<IDetailCustomization> MakeInstance();
  
-	//~FNiagaraScriptVariableDetails();
+	~FNiagaraScriptVariableDetails();
 	FNiagaraScriptVariableDetails();
  
 	/** IDetailCustomization interface */
 	virtual void CustomizeDetails( IDetailLayoutBuilder& DetailBuilder ) override;
 	
-	/*
+	
 	//~ Begin FEditorUndoClient Interface
 	virtual void PostUndo(bool bSuccess) override;
-	virtual void PostRedo(bool bSuccess) override;
+	virtual void PostRedo(bool bSuccess) override { PostUndo(bSuccess); }
 	// End of FEditorUndoClient
-	*/
  
 private:
+	void OnBeginValueChanged();
+	void OnEndValueChanged();
 	void OnValueChanged();
 	void OnStaticSwitchValueChanged();
+
+	UEdGraphPin* GetDefaultPin();
  
 	TWeakObjectPtr<class UNiagaraScriptVariable> Variable;
 	TSharedPtr<class INiagaraEditorTypeUtilities, ESPMode::ThreadSafe> TypeUtilityValue;
diff --git a/Engine/Plugins/FX/Niagara/Source/NiagaraEditor/Private/TypeEditorUtilities/NiagaraColorTypeEditorUtilities.cpp b/Engine/Plugins/FX/Niagara/Source/NiagaraEditor/Private/TypeEditorUtilities/NiagaraColorTypeEditorUtilities.cpp
--- a/Engine/Plugins/FX/Niagara/Source/NiagaraEditor/Private/TypeEditorUtilities/NiagaraColorTypeEditorUtilities.cpp
+++ b/Engine/Plugins/FX/Niagara/Source/NiagaraEditor/Private/TypeEditorUtilities/NiagaraColorTypeEditorUtilities.cpp
@@ -74,6 +74,8 @@ public:
 		*((FLinearColor*)Struct->GetStructMemory()) = ColorValue;
 	}
 
+	virtual bool CanChangeContinuously() const override { return true; }
+
 private:
 	TSharedRef<SWidget> ConstructComponentWidget(int32 Index, FText ComponentLabel)
 	{
diff --git a/Engine/Plugins/FX/Niagara/Source/NiagaraEditor/Private/TypeEditorUtilities/NiagaraFloatTypeEditorUtilities.cpp b/Engine/Plugins/FX/Niagara/Source/NiagaraEditor/Private/TypeEditorUtilities/NiagaraFloatTypeEditorUtilities.cpp
--- a/Engine/Plugins/FX/Niagara/Source/NiagaraEditor/Private/TypeEditorUtilities/NiagaraFloatTypeEditorUtilities.cpp
+++ b/Engine/Plugins/FX/Niagara/Source/NiagaraEditor/Private/TypeEditorUtilities/NiagaraFloatTypeEditorUtilities.cpp
@@ -49,6 +49,8 @@ public:
 		((FNiagaraFloat*)Struct->GetStructMemory())->Value = FloatValue;
 	}
 
+	virtual bool CanChangeContinuously() const override { return true; }
+
 private:
 	void BeginSliderMovement()
 	{
diff --git a/Engine/Plugins/FX/Niagara/Source/NiagaraEditor/Private/TypeEditorUtilities/NiagaraIntegerTypeEditorUtilities.cpp b/Engine/Plugins/FX/Niagara/Source/NiagaraEditor/Private/TypeEditorUtilities/NiagaraIntegerTypeEditorUtilities.cpp
--- a/Engine/Plugins/FX/Niagara/Source/NiagaraEditor/Private/TypeEditorUtilities/NiagaraIntegerTypeEditorUtilities.cpp
+++ b/Engine/Plugins/FX/Niagara/Source/NiagaraEditor/Private/TypeEditorUtilities/NiagaraIntegerTypeEditorUtilities.cpp
@@ -50,6 +50,8 @@ public:
 		((FNiagaraInt32*)Struct->GetStructMemory())->Value = IntValue;
 	}
 
+	virtual bool CanChangeContinuously() const override { return true; }
+
 private:
 	void BeginSliderMovement()
 	{
diff --git a/Engine/Plugins/FX/Niagara/Source/NiagaraEditor/Private/TypeEditorUtilities/NiagaraVectorTypeEditorUtilities.cpp b/Engine/Plugins/FX/Niagara/Source/NiagaraEditor/Private/TypeEditorUtilities/NiagaraVectorTypeEditorUtilities.cpp
--- a/Engine/Plugins/FX/Niagara/Source/NiagaraEditor/Private/TypeEditorUtilities/NiagaraVectorTypeEditorUtilities.cpp
+++ b/Engine/Plugins/FX/Niagara/Source/NiagaraEditor/Private/TypeEditorUtilities/NiagaraVectorTypeEditorUtilities.cpp
@@ -42,6 +42,8 @@ public:
 		];
 	}
 
+	virtual bool CanChangeContinuously() const override { return true; }
+	
 protected:
 	virtual float GetValue(int32 Index) const = 0;
 
diff --git a/Engine/Plugins/FX/Niagara/Source/NiagaraEditor/Public/SNiagaraParameterEditor.h b/Engine/Plugins/FX/Niagara/Source/NiagaraEditor/Public/SNiagaraParameterEditor.h
--- a/Engine/Plugins/FX/Niagara/Source/NiagaraEditor/Public/SNiagaraParameterEditor.h
+++ b/Engine/Plugins/FX/Niagara/Source/NiagaraEditor/Public/SNiagaraParameterEditor.h
@@ -60,6 +60,8 @@ public:
 	//~ SWidget interface
 	virtual FVector2D ComputeDesiredSize(float LayoutScaleMultiplier) const override;
 
+	/** */
+	virtual bool CanChangeContinuously() const { return false; }
 protected:
 	/** Sets whether this is currently the exclusive editor of this parameter, meaning that the corresponding details view
 		should not be updated.  This hack is necessary because the details view closes all color pickers when
