Add missing forward declaration for non-unity build in CL 9546759

Update Niagara script variable categories.

diff --git a/Engine/Plugins/FX/Niagara/Source/NiagaraEditor/Private/Customizations/NiagaraScriptVariableCustomization.h b/Engine/Plugins/FX/Niagara/Source/NiagaraEditor/Private/Customizations/NiagaraScriptVariableCustomization.h
--- a/Engine/Plugins/FX/Niagara/Source/NiagaraEditor/Private/Customizations/NiagaraScriptVariableCustomization.h
+++ b/Engine/Plugins/FX/Niagara/Source/NiagaraEditor/Private/Customizations/NiagaraScriptVariableCustomization.h
@@ -7,6 +7,8 @@
 #include "IPropertyUtilities.h"
 #include "EditorUndoClient.h"
  
+class UEdGraphPin;
+
 /** This customization sets up a custom details panel for the static switch Variable in the niagara module graph. */
 class FNiagaraScriptVariableDetails : public IDetailCustomization, public FEditorUndoClient
 {
@@ -21,7 +23,6 @@ public:
 	/** IDetailCustomization interface */
 	virtual void CustomizeDetails( IDetailLayoutBuilder& DetailBuilder ) override;
 	
-	
 	//~ Begin FEditorUndoClient Interface
 	virtual void PostUndo(bool bSuccess) override;
 	virtual void PostRedo(bool bSuccess) override { PostUndo(bSuccess); }
@@ -34,7 +35,7 @@ private:
 	void OnStaticSwitchValueChanged();
 
 	UEdGraphPin* GetDefaultPin();
- 
+
 	TWeakObjectPtr<class UNiagaraScriptVariable> Variable;
 	TSharedPtr<class INiagaraEditorTypeUtilities, ESPMode::ThreadSafe> TypeUtilityValue;
 	TSharedPtr<class INiagaraEditorTypeUtilities, ESPMode::ThreadSafe> TypeUtilityStaticSwitchValue;
diff --git a/Engine/Plugins/FX/Niagara/Source/NiagaraEditor/Private/NiagaraScriptVariable.h b/Engine/Plugins/FX/Niagara/Source/NiagaraEditor/Private/NiagaraScriptVariable.h
--- a/Engine/Plugins/FX/Niagara/Source/NiagaraEditor/Private/NiagaraScriptVariable.h
+++ b/Engine/Plugins/FX/Niagara/Source/NiagaraEditor/Private/NiagaraScriptVariable.h
@@ -21,6 +21,6 @@ public:
 	UPROPERTY()
 	FNiagaraVariable Variable;
 	/** The metadata associated with this script variable. */
-	UPROPERTY(EditAnywhere, Category = "Script Variable Data")
+	UPROPERTY(EditAnywhere, Category = "Variable", meta=(ShowOnlyInnerProperties))
 	FNiagaraVariableMetaData Metadata;
-};
\ No newline at end of file
+}; 
\ No newline at end of file
