- Only reset the preview component when the timeline wraps, rather than all components in the scene
- This mirrors Cascade behavior

#rnx

diff --git a/Engine/Plugins/FX/Niagara/Source/NiagaraEditor/Private/ViewModels/NiagaraSystemViewModel.cpp b/Engine/Plugins/FX/Niagara/Source/NiagaraEditor/Private/ViewModels/NiagaraSystemViewModel.cpp
--- a/Engine/Plugins/FX/Niagara/Source/NiagaraEditor/Private/ViewModels/NiagaraSystemViewModel.cpp
+++ b/Engine/Plugins/FX/Niagara/Source/NiagaraEditor/Private/ViewModels/NiagaraSystemViewModel.cpp
@@ -999,7 +999,19 @@ void FNiagaraSystemViewModel::ResetSystem(ETimeResetMode TimeResetMode, EMultiRe
 		}
 	}
 
-	TArray<UNiagaraComponent*> ReferencingComponents = FNiagaraEditorUtilities::GetComponentsThatReferenceSystem(GetSystem());
+	TArray<UNiagaraComponent*> ReferencingComponents;
+	if (MultiResetMode == EMultiResetMode::ResetThisInstance)
+	{
+		if (PreviewComponent != nullptr)
+		{
+			ReferencingComponents.Add(PreviewComponent);
+		}
+	}
+	else
+	{
+		ReferencingComponents = FNiagaraEditorUtilities::GetComponentsThatReferenceSystem(GetSystem());
+	}
+
 	for (auto Component : ReferencingComponents)
 	{
 		if (ReinitMode == EReinitMode::ResetSystem)
