- Reregister systems that are impacted by adding a new renderer to the emitter
- Fixes a crash when adding a new renderer while a system is active in the main viewport

#rnx

diff --git a/Engine/Plugins/FX/Niagara/Source/NiagaraEditor/Private/ViewModels/Stack/NiagaraStackRenderItemGroup.cpp b/Engine/Plugins/FX/Niagara/Source/NiagaraEditor/Private/ViewModels/Stack/NiagaraStackRenderItemGroup.cpp
--- a/Engine/Plugins/FX/Niagara/Source/NiagaraEditor/Private/ViewModels/Stack/NiagaraStackRenderItemGroup.cpp
+++ b/Engine/Plugins/FX/Niagara/Source/NiagaraEditor/Private/ViewModels/Stack/NiagaraStackRenderItemGroup.cpp
@@ -101,6 +101,8 @@ public:
 			}
 		}
 
+		FNiagaraSystemUpdateContext SystemUpdate(Emitter, true);
+
 		if (bVarsAdded)
 		{
 			FNotificationInfo Info(LOCTEXT("AddedVariables", "One or more variables have been added to the Spawn script to support the added renderer."));
