Merging using MainToDevVR

diff --git a/Engine/Plugins/FX/Niagara/Source/Niagara/Private/NiagaraModule.cpp b/Engine/Plugins/FX/Niagara/Source/Niagara/Private/NiagaraModule.cpp
--- a/Engine/Plugins/FX/Niagara/Source/Niagara/Private/NiagaraModule.cpp
+++ b/Engine/Plugins/FX/Niagara/Source/Niagara/Private/NiagaraModule.cpp
@@ -69,6 +69,7 @@ FNiagaraVariable INiagaraModule::Engine_Owner_XAxis;
 FNiagaraVariable INiagaraModule::Engine_Owner_YAxis;
 FNiagaraVariable INiagaraModule::Engine_Owner_ZAxis;
 FNiagaraVariable INiagaraModule::Engine_Owner_Scale;
+FNiagaraVariable INiagaraModule::Engine_Owner_Rotation;
 
 FNiagaraVariable INiagaraModule::Engine_Owner_SystemLocalToWorld;
 FNiagaraVariable INiagaraModule::Engine_Owner_SystemWorldToLocal;
@@ -84,6 +85,8 @@ FNiagaraVariable INiagaraModule::Engine_Owner_ExecutionState;
 
 FNiagaraVariable INiagaraModule::Engine_ExecutionCount;
 FNiagaraVariable INiagaraModule::Engine_Emitter_NumParticles;
+FNiagaraVariable INiagaraModule::Engine_Emitter_TotalSpawnedParticles;
+FNiagaraVariable INiagaraModule::Engine_System_TickCount;
 FNiagaraVariable INiagaraModule::Engine_System_NumEmittersAlive;
 FNiagaraVariable INiagaraModule::Engine_System_NumEmitters;
 FNiagaraVariable INiagaraModule::Engine_NumSystemInstances;
@@ -95,11 +98,14 @@ FNiagaraVariable INiagaraModule::Engine_System_Age;
 
 FNiagaraVariable INiagaraModule::Emitter_Age;
 FNiagaraVariable INiagaraModule::Emitter_LocalSpace;
+FNiagaraVariable INiagaraModule::Emitter_Determinism;
+FNiagaraVariable INiagaraModule::Emitter_RandomSeed;
 FNiagaraVariable INiagaraModule::Emitter_SpawnRate;
 FNiagaraVariable INiagaraModule::Emitter_SpawnInterval;
 FNiagaraVariable INiagaraModule::Emitter_InterpSpawnStartDt;
 FNiagaraVariable INiagaraModule::Emitter_SpawnGroup;
 
+FNiagaraVariable INiagaraModule::Particles_UniqueID;
 FNiagaraVariable INiagaraModule::Particles_ID;
 FNiagaraVariable INiagaraModule::Particles_Position;
 FNiagaraVariable INiagaraModule::Particles_Velocity;
@@ -121,6 +127,9 @@ FNiagaraVariable INiagaraModule::Particles_UVScale;
 FNiagaraVariable INiagaraModule::Particles_CameraOffset;
 FNiagaraVariable INiagaraModule::Particles_MaterialRandom;
 FNiagaraVariable INiagaraModule::Particles_LightRadius;
+FNiagaraVariable INiagaraModule::Particles_LightExponent;
+FNiagaraVariable INiagaraModule::Particles_LightEnabled;
+FNiagaraVariable INiagaraModule::Particles_LightVolumetricScattering;
 FNiagaraVariable INiagaraModule::Particles_RibbonID;
 FNiagaraVariable INiagaraModule::Particles_RibbonWidth;
 FNiagaraVariable INiagaraModule::Particles_RibbonTwist;
@@ -167,6 +176,7 @@ void INiagaraModule::StartupModule()
 	Engine_Owner_YAxis = FNiagaraVariable(FNiagaraTypeDefinition::GetVec3Def(), TEXT("Engine.Owner.SystemYAxis"));
 	Engine_Owner_ZAxis = FNiagaraVariable(FNiagaraTypeDefinition::GetVec3Def(), TEXT("Engine.Owner.SystemZAxis"));
 	Engine_Owner_Scale = FNiagaraVariable(FNiagaraTypeDefinition::GetVec3Def(), TEXT("Engine.Owner.Scale"));
+	Engine_Owner_Rotation = FNiagaraVariable(FNiagaraTypeDefinition::GetQuatDef(), TEXT("Engine.Owner.Rotation"));
 
 	Engine_Owner_SystemLocalToWorld = FNiagaraVariable(FNiagaraTypeDefinition::GetMatrix4Def(), TEXT("Engine.Owner.SystemLocalToWorld"));
 	Engine_Owner_SystemWorldToLocal = FNiagaraVariable(FNiagaraTypeDefinition::GetMatrix4Def(), TEXT("Engine.Owner.SystemWorldToLocal"));
@@ -182,6 +192,8 @@ void INiagaraModule::StartupModule()
 	
 	Engine_ExecutionCount = FNiagaraVariable(FNiagaraTypeDefinition::GetIntDef(), TEXT("Engine.ExecutionCount"));
 	Engine_Emitter_NumParticles = FNiagaraVariable(FNiagaraTypeDefinition::GetIntDef(), TEXT("Engine.Emitter.NumParticles"));
+	Engine_Emitter_TotalSpawnedParticles = FNiagaraVariable(FNiagaraTypeDefinition::GetIntDef(), TEXT("Engine.Emitter.TotalSpawnedParticles"));
+	Engine_System_TickCount = FNiagaraVariable(FNiagaraTypeDefinition::GetIntDef(), TEXT("Engine.System.TickCount"));
 	Engine_System_NumEmittersAlive = FNiagaraVariable(FNiagaraTypeDefinition::GetIntDef(), TEXT("Engine.System.NumEmittersAlive"));
 	Engine_System_NumEmitters = FNiagaraVariable(FNiagaraTypeDefinition::GetIntDef(), TEXT("Engine.System.NumEmitters"));
 	Engine_NumSystemInstances = FNiagaraVariable(FNiagaraTypeDefinition::GetIntDef(), TEXT("Engine.NumSystemInstances"));
@@ -192,11 +204,14 @@ void INiagaraModule::StartupModule()
 	Engine_System_Age = FNiagaraVariable(FNiagaraTypeDefinition::GetFloatDef(), TEXT("Engine.System.Age"));
 	Emitter_Age = FNiagaraVariable(FNiagaraTypeDefinition::GetFloatDef(), TEXT("Emitter.Age"));
 	Emitter_LocalSpace = FNiagaraVariable(FNiagaraTypeDefinition::GetBoolDef(), TEXT("Emitter.LocalSpace"));
+	Emitter_RandomSeed = FNiagaraVariable(FNiagaraTypeDefinition::GetIntDef(), TEXT("Emitter.RandomSeed"));
+	Emitter_Determinism = FNiagaraVariable(FNiagaraTypeDefinition::GetBoolDef(), TEXT("Emitter.Determinism"));
 	Emitter_SpawnRate = FNiagaraVariable(FNiagaraTypeDefinition::GetFloatDef(), TEXT("Emitter.SpawnRate"));
 	Emitter_SpawnInterval = FNiagaraVariable(FNiagaraTypeDefinition::GetFloatDef(), TEXT("Emitter.SpawnInterval"));
 	Emitter_InterpSpawnStartDt = FNiagaraVariable(FNiagaraTypeDefinition::GetFloatDef(), TEXT("Emitter.InterpSpawnStartDt"));
 	Emitter_SpawnGroup = FNiagaraVariable(FNiagaraTypeDefinition::GetIntDef(), TEXT("Emitter.SpawnGroup"));
 
+	Particles_UniqueID = FNiagaraVariable(FNiagaraTypeDefinition::GetIntDef(), TEXT("Particles.UniqueID"));
 	Particles_ID = FNiagaraVariable(FNiagaraTypeDefinition::GetIDDef(), TEXT("Particles.ID"));
 	Particles_Position = FNiagaraVariable(FNiagaraTypeDefinition::GetVec3Def(), TEXT("Particles.Position"));
 	Particles_Velocity = FNiagaraVariable(FNiagaraTypeDefinition::GetVec3Def(), TEXT("Particles.Velocity"));
@@ -218,6 +233,9 @@ void INiagaraModule::StartupModule()
 	Particles_CameraOffset = FNiagaraVariable(FNiagaraTypeDefinition::GetFloatDef(), TEXT("Particles.CameraOffset"));
 	Particles_MaterialRandom = FNiagaraVariable(FNiagaraTypeDefinition::GetFloatDef(), TEXT("Particles.MaterialRandom"));
 	Particles_LightRadius = FNiagaraVariable(FNiagaraTypeDefinition::GetFloatDef(), TEXT("Particles.LightRadius"));
+	Particles_LightExponent = FNiagaraVariable(FNiagaraTypeDefinition::GetFloatDef(), TEXT("Particles.LightExponent"));
+	Particles_LightEnabled = FNiagaraVariable(FNiagaraTypeDefinition::GetBoolDef(), TEXT("Particles.LightEnabled"));
+	Particles_LightVolumetricScattering = FNiagaraVariable(FNiagaraTypeDefinition::GetFloatDef(), TEXT("Particles.LightVolumetricScattering"));
 	Particles_RibbonID = FNiagaraVariable(FNiagaraTypeDefinition::GetIDDef(), TEXT("Particles.RibbonID"));
 	Particles_RibbonWidth = FNiagaraVariable(FNiagaraTypeDefinition::GetFloatDef(), TEXT("Particles.RibbonWidth"));
 	Particles_RibbonTwist = FNiagaraVariable(FNiagaraTypeDefinition::GetFloatDef(), TEXT("Particles.RibbonTwist"));
@@ -598,6 +616,8 @@ void FNiagaraTypeDefinition::RecreateUserDefinedTypeRegistry()
 
 	for (FSoftObjectPath AssetRef : TotalStructAssets)
 	{
+		FName AssetRefPathNamePreResolve = AssetRef.GetAssetPathName();
+
 		UObject* Obj = AssetRef.ResolveObject();
 		if (Obj == nullptr)
 		{
@@ -613,6 +633,11 @@ void FNiagaraTypeDefinition::RecreateUserDefinedTypeRegistry()
 			{
 				FNiagaraTypeRegistry::Register(ScriptStruct, ParamRefFound != nullptr, PayloadRefFound != nullptr, true);
 			}
+			if (Obj->GetPathName() != AssetRefPathNamePreResolve.ToString())
+			{
+				UE_LOG(LogNiagara, Warning, TEXT("Additional parameter/payload enum has moved from where it was in settings (this may cause errors at runtime): Was: \"%s\" Now: \"%s\""), *AssetRefPathNamePreResolve.ToString(), *Obj->GetPathName());
+			}
+
 		}
 		else
 		{
@@ -623,6 +648,7 @@ void FNiagaraTypeDefinition::RecreateUserDefinedTypeRegistry()
 
 	for (FSoftObjectPath AssetRef : Settings->AdditionalParameterEnums)
 	{
+		FName AssetRefPathNamePreResolve = AssetRef.GetAssetPathName();
 		UObject* Obj = AssetRef.ResolveObject();
 		if (Obj == nullptr)
 		{
@@ -631,13 +657,18 @@ void FNiagaraTypeDefinition::RecreateUserDefinedTypeRegistry()
 
 		if (Obj != nullptr)
 		{
-			const FSoftObjectPath* ParamRefFound = Settings->AdditionalParameterEnums.FindByPredicate([&](const FStringAssetReference& Ref) { return Ref.ToString() == AssetRef.ToString(); });
+			const FSoftObjectPath* ParamRefFound = Settings->AdditionalParameterEnums.FindByPredicate([&](const FSoftObjectPath& Ref) { return Ref.ToString() == AssetRef.ToString(); });
 			const FSoftObjectPath* PayloadRefFound = nullptr;
 			UEnum* Enum = Cast<UEnum>(Obj);
 			if (Enum != nullptr)
 			{
 				FNiagaraTypeRegistry::Register(Enum, ParamRefFound != nullptr, PayloadRefFound != nullptr, true);
 			}
+
+			if (Obj->GetPathName() != AssetRefPathNamePreResolve.ToString())
+			{
+				UE_LOG(LogNiagara, Warning, TEXT("Additional parameter/payload enum has moved from where it was in settings (this may cause errors at runtime): Was: \"%s\" Now: \"%s\""), *AssetRefPathNamePreResolve.ToString(), *Obj->GetPathName());
+			}
 		}
 		else
 		{
@@ -757,27 +788,10 @@ FNiagaraTypeDefinition FNiagaraTypeDefinition::GetNumericOutputType(const TArray
 	{
 		return SortedTypeDefinitions[0];
 	}
-}
-
-//////////////////////////////////////////////////////////////////////////
 
-template<>
-void FNiagaraVariable::SetValue(const bool& Data)
-{
-	check(TypeDef == FNiagaraTypeDefinition::GetBoolDef());
-	AllocateData();
-	FNiagaraBool* BoolStruct = (FNiagaraBool*)GetData();
-	BoolStruct->SetValue(Data);
+	return FNiagaraTypeDefinition::GetGenericNumericDef();
 }
 
-template<>
-bool FNiagaraVariable::GetValue() const
-{
-	check(TypeDef == FNiagaraTypeDefinition::GetBoolDef());
-	check(IsDataAllocated());
-	FNiagaraBool* BoolStruct = (FNiagaraBool*)GetData();
-	return BoolStruct->GetValue();
-}
 
 //////////////////////////////////////////////////////////////////////////
 
diff --git a/Engine/Plugins/FX/Niagara/Source/NiagaraEditor/Private/NiagaraEmitterFactoryNew.cpp b/Engine/Plugins/FX/Niagara/Source/NiagaraEditor/Private/NiagaraEmitterFactoryNew.cpp
--- a/Engine/Plugins/FX/Niagara/Source/NiagaraEditor/Private/NiagaraEmitterFactoryNew.cpp
+++ b/Engine/Plugins/FX/Niagara/Source/NiagaraEditor/Private/NiagaraEmitterFactoryNew.cpp
@@ -82,7 +82,7 @@ bool UNiagaraEmitterFactoryNew::ConfigureProperties()
 
 UNiagaraNodeFunctionCall* AddModuleFromAssetPath(FString AssetPath, UNiagaraNodeOutput& TargetOutputNode)
 {
-	FStringAssetReference AssetRef(AssetPath);
+	FSoftObjectPath AssetRef(AssetPath);
 	UNiagaraScript* AssetScript = Cast<UNiagaraScript>(AssetRef.TryLoad());
 	FAssetData ScriptAssetData(AssetScript);
 	if (ScriptAssetData.IsValid())
@@ -196,6 +196,7 @@ UObject* UNiagaraEmitterFactoryNew::FactoryCreateNew(UClass* Class, UObject* InP
 
 		FNiagaraStackGraphUtilities::RelayoutGraph(*Source->NodeGraph);
 		NewEmitter->bInterpolatedSpawning = true;
+		NewEmitter->bDeterminism = false; // NOTE: Default to non-determinism
 		NewEmitter->SpawnScriptProps.Script->SetUsage(ENiagaraScriptUsage::ParticleSpawnScriptInterpolated);
 	}
 	
diff --git a/Engine/Plugins/FX/Niagara/Source/NiagaraEditor/Private/NiagaraScriptSource.cpp b/Engine/Plugins/FX/Niagara/Source/NiagaraEditor/Private/NiagaraScriptSource.cpp
--- a/Engine/Plugins/FX/Niagara/Source/NiagaraEditor/Private/NiagaraScriptSource.cpp
+++ b/Engine/Plugins/FX/Niagara/Source/NiagaraEditor/Private/NiagaraScriptSource.cpp
@@ -153,7 +153,7 @@ void UNiagaraScriptSource::PostLoadFromEmitter(UNiagaraEmitter& OwningEmitter)
 
 bool UNiagaraScriptSource::AddModuleIfMissing(FString ModulePath, ENiagaraScriptUsage Usage, bool& bOutFoundModule)
 {
-	FStringAssetReference SystemUpdateScriptRef(ModulePath);
+	FSoftObjectPath SystemUpdateScriptRef(ModulePath);
 	FAssetData ModuleScriptAsset;
 	ModuleScriptAsset.ObjectPath = SystemUpdateScriptRef.GetAssetPathName();
 	bOutFoundModule = false;
@@ -185,6 +185,13 @@ void InitializeNewRapidIterationParametersForNode(const UEdGraphSchema_Niagara*
 		for (const UEdGraphPin* FunctionInputPin : FunctionInputPins)
 		{
 			FNiagaraTypeDefinition InputType = Schema->PinToTypeDefinition(FunctionInputPin);
+			if (InputType.IsValid() == false)
+			{
+				UE_LOG(LogNiagaraEditor, Error, TEXT("Invalid input type found while attempting initialize new rapid iteration parameters. Function Node: %s %s Input Name: %s"),
+					*FunctionCallNode->GetPathName(), *FunctionCallNode->GetFunctionName(), *FunctionInputPin->GetName());
+				continue;
+			}
+
 			if (FNiagaraStackGraphUtilities::IsRapidIterationType(InputType))
 			{
 				FNiagaraParameterHandle AliasedFunctionInputHandle = FNiagaraParameterHandle::CreateAliasedModuleParameterHandle(FNiagaraParameterHandle(FunctionInputPin->PinName), FunctionCallNode);
diff --git a/Engine/Plugins/FX/Niagara/Source/NiagaraEditor/Private/ViewModels/Stack/NiagaraStackScriptItemGroup.cpp b/Engine/Plugins/FX/Niagara/Source/NiagaraEditor/Private/ViewModels/Stack/NiagaraStackScriptItemGroup.cpp
--- a/Engine/Plugins/FX/Niagara/Source/NiagaraEditor/Private/ViewModels/Stack/NiagaraStackScriptItemGroup.cpp
+++ b/Engine/Plugins/FX/Niagara/Source/NiagaraEditor/Private/ViewModels/Stack/NiagaraStackScriptItemGroup.cpp
@@ -397,7 +397,7 @@ void UNiagaraStackScriptItemGroup::RefreshIssues(TArray<FStackIssue>& NewIssues)
 			// The factor ensures this, but older assets may not have it or it may have been removed accidentally.
 			// For now, treat this as an error and allow them to resolve.
 			FString ModulePath = TEXT("/Niagara/Modules/System/SystemLifeCycle.SystemLifeCycle");
-			FStringAssetReference SystemUpdateScriptRef(ModulePath);
+			FSoftObjectPath SystemUpdateScriptRef(ModulePath);
 			FAssetData ModuleScriptAsset;
 			ModuleScriptAsset.ObjectPath = SystemUpdateScriptRef.GetAssetPathName();
 
@@ -643,13 +643,6 @@ void UNiagaraStackScriptItemGroup::OnScriptGraphChanged(const struct FEdGraphEdi
 
 void UNiagaraStackScriptItemGroup::AddParameterModuleToStack(const UNiagaraStackModuleSpacer* InModuleSpacer, const FNiagaraVariable &InVariable)
 {
-	TArray<FNiagaraVariable> Vars;
-	Vars.Add(InVariable);
-	TArray<FString> DefaultVals;
-	DefaultVals.Add(FNiagaraConstants::GetAttributeDefaultValue(InVariable));
-
-	UNiagaraNodeOutput* OutputNode = GetScriptOutputNode();
-
 	int32 TargetIndex = INDEX_NONE;
 	UNiagaraStackModuleItem** TargetModuleItemPtr = StackSpacerToModuleItemMap.Find(FObjectKey(InModuleSpacer));
 	if (*TargetModuleItemPtr != nullptr)
@@ -658,8 +651,8 @@ void UNiagaraStackScriptItemGroup::AddParameterModuleToStack(const UNiagaraStack
 		TargetIndex = TargetModuleItem->GetModuleIndex();
 	}
 
-	FNiagaraStackGraphUtilities::AddParameterModuleToStack(Vars, *OutputNode, TargetIndex, DefaultVals);
-	RefreshChildren();
+	TSharedRef<FScriptGroupAddAction> AddAction = FScriptGroupAddAction::CreateExistingParameterModuleAction(InVariable);
+	AddUtilities->ExecuteAddAction(AddAction, TargetIndex);
 }
 
 #undef LOCTEXT_NAMESPACE
diff --git a/Engine/Source/Runtime/Experimental/ChaosSolvers/Private/PBDRigidsSolver.cpp b/Engine/Source/Runtime/Experimental/ChaosSolvers/Private/PBDRigidsSolver.cpp
--- a/Engine/Source/Runtime/Experimental/ChaosSolvers/Private/PBDRigidsSolver.cpp
+++ b/Engine/Source/Runtime/Experimental/ChaosSolvers/Private/PBDRigidsSolver.cpp
@@ -31,8 +31,8 @@ namespace Chaos
 			: MScene(Scene)
 			, MDeltaTime(DeltaTime)
 			, PrevLock(PrevFrameLock)
-			, PrevEvent(PrevFrameEvent)
 			, CurrentLock(CurrentFrameLock)
+			, PrevEvent(PrevFrameEvent)
 			, CurrentEvent(CurrentFrameEvent)
 		{
 			UE_LOG(LogPBDRigidsSolverSolver, Verbose, TEXT("AdvanceOneTimeStepTask::AdvanceOneTimeStepTask()"));
@@ -99,9 +99,9 @@ namespace Chaos
 		, bHasFloor(true)
 		, bIsFloorAnalytic(false)
 		, FloorHeight(0.f)
-		, MCurrentLock(nullptr)
-		, MCurrentEvent(nullptr)
 		, MaxCollisionDataSize(1024)
+		, MCurrentEvent(nullptr)
+		, MCurrentLock(nullptr)
 		, CollisionDataTimeWindow(0.1f)
 		, DoCollisionDataSpatialHash(true)
 		, CollisionDataSpatialHashRadius(15.f)
