DR - Workaround for Vulkan driver crash

diff --git a/Engine/Plugins/FX/Niagara/Source/NiagaraVertexFactories/Private/NiagaraSortingGPU.cpp b/Engine/Plugins/FX/Niagara/Source/NiagaraVertexFactories/Private/NiagaraSortingGPU.cpp
--- a/Engine/Plugins/FX/Niagara/Source/NiagaraVertexFactories/Private/NiagaraSortingGPU.cpp
+++ b/Engine/Plugins/FX/Niagara/Source/NiagaraVertexFactories/Private/NiagaraSortingGPU.cpp
@@ -180,7 +180,7 @@ void FNiagaraCopyIntBufferRegionCS::SetParameters(
 	int32 DestCount)
 {
 	FComputeShaderRHIParamRef ComputeShaderRHI = GetComputeShader();
-	check(DestCount > 0);
+	check(DestCount > 0 && DestCount <= NIAGARA_COPY_BUFFER_BUFFER_COUNT);
 
 	RHICmdList.SetShaderResourceViewParameter(ComputeShaderRHI, SourceData.GetBaseIndex(), InSourceData);
 
@@ -190,6 +190,13 @@ void FNiagaraCopyIntBufferRegionCS::SetParameters(
 		RHICmdList.SetUAVParameter(ComputeShaderRHI, DestData[Index].GetBaseIndex(), InDestDatas[Index]);
 		CopyParamsValue[Index + 1] = InUsedIndexCounts[Index];
 	}
+
+	// Workaround for Vulkan driver crashes: replicate the last UAV over
+	for (int32 Index = DestCount; Index < NIAGARA_COPY_BUFFER_BUFFER_COUNT; ++Index)
+	{
+		RHICmdList.SetUAVParameter(ComputeShaderRHI, DestData[Index].GetBaseIndex(), InDestDatas[DestCount - 1]);
+	}
+
 	RHICmdList.SetShaderParameter(ComputeShaderRHI, CopyParams.GetBufferIndex(), CopyParams.GetBaseIndex(), CopyParams.GetNumBytes(), &CopyParamsValue);
 }
 
