Fix Niagara ClangEditor CIS

diff --git a/Engine/Plugins/FX/Niagara/Source/Niagara/Private/NiagaraRendererSprites.cpp b/Engine/Plugins/FX/Niagara/Source/Niagara/Private/NiagaraRendererSprites.cpp
--- a/Engine/Plugins/FX/Niagara/Source/Niagara/Private/NiagaraRendererSprites.cpp
+++ b/Engine/Plugins/FX/Niagara/Source/Niagara/Private/NiagaraRendererSprites.cpp
@@ -500,7 +500,7 @@ void NiagaraRendererSprites::GetDynamicRayTracingInstances(FRayTracingMaterialGa
 				MeshBatch,
 				MeshBatch.Elements[0].NumPrimitives == 0,
 				6 * DynamicDataSprites->RTParticleData.GetNumInstances(),
-				6 * DynamicDataSprites->RTParticleData.GetNumInstances() * sizeof(FVector),
+				6 * DynamicDataSprites->RTParticleData.GetNumInstances() * (uint32)sizeof(FVector),
 				2 * DynamicDataSprites->RTParticleData.GetNumInstances(),
 				&RayTracingGeometry,
 				&RayTracingDynamicVertexBuffer
