Fix Linux non ray tracing build

diff --git a/Engine/Plugins/FX/Niagara/Source/Niagara/Private/NiagaraRendererSprites.cpp b/Engine/Plugins/FX/Niagara/Source/Niagara/Private/NiagaraRendererSprites.cpp
--- a/Engine/Plugins/FX/Niagara/Source/Niagara/Private/NiagaraRendererSprites.cpp
+++ b/Engine/Plugins/FX/Niagara/Source/Niagara/Private/NiagaraRendererSprites.cpp
@@ -372,7 +372,9 @@ void NiagaraRendererSprites::CreateMeshBatchForView(
 
 	MeshBatch.VertexFactory = &CollectorResources.VertexFactory;
 	MeshBatch.CastShadow = SceneProxy->CastsDynamicShadow();
+#if RHI_RAYTRACING
 	MeshBatch.CastRayTracedShadow = SceneProxy->CastsDynamicShadow();
+#endif
 	MeshBatch.bUseAsOccluder = false;
 	MeshBatch.ReverseCulling = SceneProxy->IsLocalToWorldDeterminantNegative();
 	MeshBatch.Type = PT_TriangleList;
