Fixing merge build break.

diff --git a/Engine/Plugins/FX/Niagara/Source/Niagara/Classes/NiagaraScriptExecutionContext.h b/Engine/Plugins/FX/Niagara/Source/Niagara/Classes/NiagaraScriptExecutionContext.h
--- a/Engine/Plugins/FX/Niagara/Source/Niagara/Classes/NiagaraScriptExecutionContext.h
+++ b/Engine/Plugins/FX/Niagara/Source/Niagara/Classes/NiagaraScriptExecutionContext.h
@@ -186,6 +186,7 @@ public:
 	//Most current buffer that can be used for rendering.
 	FNiagaraDataBuffer* DataToRender;
 
+	FRHIGPUBufferReadback *GPUDataReadback;
 	uint32 AccumulatedSpawnRate;
 	uint32 NumIndicesPerInstance;	// how many vtx indices per instance the renderer is going to have for its draw call
 
diff --git a/Engine/Plugins/FX/Niagara/Source/NiagaraVertexFactories/Private/NiagaraMeshVertexFactory.cpp b/Engine/Plugins/FX/Niagara/Source/NiagaraVertexFactories/Private/NiagaraMeshVertexFactory.cpp
--- a/Engine/Plugins/FX/Niagara/Source/NiagaraVertexFactories/Private/NiagaraMeshVertexFactory.cpp
+++ b/Engine/Plugins/FX/Niagara/Source/NiagaraVertexFactories/Private/NiagaraMeshVertexFactory.cpp
@@ -106,7 +106,7 @@ public:
 		const FSceneInterface* Scene,
 		const FSceneView* View,
 		const FMeshMaterialShader* Shader,
-		bool bShaderRequiresPositionOnlyStream,
+		const EVertexInputStreamType InputStreamType,
 		ERHIFeatureLevel::Type FeatureLevel,
 		const FVertexFactory* VertexFactory,
 		const FMeshBatchElement& BatchElement,
