Fixed compilation warnings due usage of deprecated FShaderResourceViewRHIParamRef instances

diff --git a/Engine/Plugins/FX/Niagara/Source/Niagara/Classes/NiagaraDataInterfaceSkeletalMesh.h b/Engine/Plugins/FX/Niagara/Source/Niagara/Classes/NiagaraDataInterfaceSkeletalMesh.h
--- a/Engine/Plugins/FX/Niagara/Source/Niagara/Classes/NiagaraDataInterfaceSkeletalMesh.h
+++ b/Engine/Plugins/FX/Niagara/Source/Niagara/Classes/NiagaraDataInterfaceSkeletalMesh.h
@@ -250,10 +250,10 @@ public:
 	FShaderResourceViewRHIRef GetBufferTriangleMatricesOffsetSRV() const { return BufferTriangleMatricesOffsetSRV; }
 	uint32 GetTriangleCount() const { return TriangleCount; }
 
-	FShaderResourceViewRHIParamRef GetBufferPositionSRV() const { return MeshVertexBufferSrv; }
-	FShaderResourceViewRHIParamRef GetBufferIndexSRV() const { return MeshIndexBufferSrv; }
-	FShaderResourceViewRHIParamRef GetBufferTangentSRV() const { return MeshTangentBufferSRV; }
-	FShaderResourceViewRHIParamRef GetBufferTexCoordSRV() const { return MeshTexCoordBufferSrv; }
+	FRHIShaderResourceView* GetBufferPositionSRV() const { return MeshVertexBufferSrv; }
+	FRHIShaderResourceView* GetBufferIndexSRV() const { return MeshIndexBufferSrv; }
+	FRHIShaderResourceView* GetBufferTangentSRV() const { return MeshTangentBufferSRV; }
+	FRHIShaderResourceView* GetBufferTexCoordSRV() const { return MeshTexCoordBufferSrv; }
 
 	uint32 GetNumTexCoord() const { return NumTexCoord; }
 
@@ -267,10 +267,10 @@ protected:
 	FShaderResourceViewRHIRef BufferTriangleMatricesOffsetSRV = nullptr;
 
 	/** Cached SRV to gpu buffers of the mesh we spawn from */
-	FShaderResourceViewRHIParamRef MeshVertexBufferSrv;
-	FShaderResourceViewRHIParamRef MeshIndexBufferSrv;
-	FShaderResourceViewRHIParamRef MeshTangentBufferSRV;
-	FShaderResourceViewRHIParamRef MeshTexCoordBufferSrv;
+	FRHIShaderResourceView* MeshVertexBufferSrv;
+	FRHIShaderResourceView* MeshIndexBufferSrv;
+	FRHIShaderResourceView* MeshTangentBufferSRV;
+	FRHIShaderResourceView* MeshTexCoordBufferSrv;
 
 	uint32 NumTexCoord;
 
@@ -369,7 +369,7 @@ struct FNDISkeletalMesh_InstanceData
 	/** True if the mesh we're using allows area weighted sampling on GPU. */
 	uint32 bIsGpuUniformlyDistributedSampling : 1;
 
-	FShaderResourceViewRHIParamRef MeshSkinWeightBufferSrv;
+	FRHIShaderResourceView* MeshSkinWeightBufferSrv;
 	uint32 MeshWeightStrideByte;
 
 	/** Extra mesh data upload to GPU.*/
@@ -650,7 +650,7 @@ struct FNiagaraDISkeletalMeshPassedDataToRT
 {
 	FSkeletalMeshGpuSpawnStaticBuffers* StaticBuffers;
 	FSkeletalMeshGpuDynamicBufferProxy* DynamicBuffer;
-	FShaderResourceViewRHIParamRef MeshSkinWeightBufferSrv;
+	FRHIShaderResourceView* MeshSkinWeightBufferSrv;
 
 	bool bIsGpuUniformlyDistributedSampling;
 
diff --git a/Engine/Plugins/FX/Niagara/Source/Niagara/Classes/NiagaraDataInterfaceStaticMesh.h b/Engine/Plugins/FX/Niagara/Source/Niagara/Classes/NiagaraDataInterfaceStaticMesh.h
--- a/Engine/Plugins/FX/Niagara/Source/Niagara/Classes/NiagaraDataInterfaceStaticMesh.h
+++ b/Engine/Plugins/FX/Niagara/Source/Niagara/Classes/NiagaraDataInterfaceStaticMesh.h
@@ -54,16 +54,16 @@ public:
 
 	virtual FString GetFriendlyName() const override { return TEXT("FStaticMeshGpuSpawnBuffer"); }
 
-	FShaderResourceViewRHIParamRef GetBufferSectionSRV() const { return BufferSectionSRV; }
+	FRHIShaderResourceView* GetBufferSectionSRV() const { return BufferSectionSRV; }
 	uint32 GetValidSectionCount() const { return ValidSections.Num(); };
 
-	FShaderResourceViewRHIParamRef GetBufferPositionSRV() const { return MeshVertexBufferSrv; }
-	FShaderResourceViewRHIParamRef GetBufferTangentSRV() const { return MeshTangentBufferSrv; }
-	FShaderResourceViewRHIParamRef GetBufferTexCoordSRV() const { return MeshTexCoordBufferSrv; }
-	FShaderResourceViewRHIParamRef GetBufferIndexSRV() const { return MeshIndexBufferSrv; }
-	FShaderResourceViewRHIParamRef GetBufferColorSRV() const { return MeshColorBufferSRV; }
+	FRHIShaderResourceView* GetBufferPositionSRV() const { return MeshVertexBufferSrv; }
+	FRHIShaderResourceView* GetBufferTangentSRV() const { return MeshTangentBufferSrv; }
+	FRHIShaderResourceView* GetBufferTexCoordSRV() const { return MeshTexCoordBufferSrv; }
+	FRHIShaderResourceView* GetBufferIndexSRV() const { return MeshIndexBufferSrv; }
+	FRHIShaderResourceView* GetBufferColorSRV() const { return MeshColorBufferSRV; }
 
-	FShaderResourceViewRHIParamRef GetBufferUniformTriangleSamplingSRV() const { return BufferUniformTriangleSamplingSRV; }
+	FRHIShaderResourceView* GetBufferUniformTriangleSamplingSRV() const { return BufferUniformTriangleSamplingSRV; }
 
 	uint32 GetNumTexCoord() const { return NumTexCoord; }
 
diff --git a/Engine/Plugins/FX/Niagara/Source/Niagara/Classes/NiagaraEmitterInstanceBatcher.h b/Engine/Plugins/FX/Niagara/Source/Niagara/Classes/NiagaraEmitterInstanceBatcher.h
--- a/Engine/Plugins/FX/Niagara/Source/Niagara/Classes/NiagaraEmitterInstanceBatcher.h
+++ b/Engine/Plugins/FX/Niagara/Source/Niagara/Classes/NiagaraEmitterInstanceBatcher.h
@@ -40,7 +40,7 @@ public:
 class NiagaraEmitterInstanceBatcher : public FFXSystemInterface
 {
 public:
-	using FNiagaraBufferArray = TArray<FUnorderedAccessViewRHIParamRef, TMemStackAllocator<>>;
+	using FNiagaraBufferArray = TArray<FRHIUnorderedAccessView*, TMemStackAllocator<>>;
 	using FOverlappableTicks = TArray<FNiagaraGPUSystemTick*, TMemStackAllocator<>>;
 
 	NiagaraEmitterInstanceBatcher(ERHIFeatureLevel::Type InFeatureLevel, EShaderPlatform InShaderPlatform)
diff --git a/Engine/Source/Runtime/Engine/Private/Particles/ParticleGpuSimulation.cpp b/Engine/Source/Runtime/Engine/Private/Particles/ParticleGpuSimulation.cpp
--- a/Engine/Source/Runtime/Engine/Private/Particles/ParticleGpuSimulation.cpp
+++ b/Engine/Source/Runtime/Engine/Private/Particles/ParticleGpuSimulation.cpp
@@ -680,7 +680,7 @@ TGlobalResource<FGPUSpriteVertexDeclaration> GGPUSpriteVertexDeclaration;
 struct FGPUSpriteMeshDataUserData  : public FOneFrameResource
 {
 	int32 SortedOffset = 0;
-	FShaderResourceViewRHIParamRef SortedParticleIndicesSRV = nullptr;
+	FRHIShaderResourceView* SortedParticleIndicesSRV = nullptr;
 };
 
 /**
@@ -701,7 +701,7 @@ public:
 	/** Emitter uniform buffer for dynamic parameters. */
 	FUniformBufferRHIRef EmitterDynamicUniformBuffer;
 	/** Buffer containing unsorted particle indices. */
-	FShaderResourceViewRHIParamRef UnsortedParticleIndicesSRV;
+	FRHIShaderResourceView* UnsortedParticleIndicesSRV;
 	/** Texture containing positions for all particles. */
 	FTexture2DRHIParamRef PositionTextureRHI;
 	/** Texture containing velocities for all particles. */
@@ -743,7 +743,7 @@ public:
 	/**
 	 * Set the source vertex buffer that contains particle indices.
 	 */
-	void SetUnsortedParticleIndicesSRV(FShaderResourceViewRHIParamRef VertexBuffer )
+	void SetUnsortedParticleIndicesSRV(FRHIShaderResourceView* VertexBuffer )
 	{
 		UnsortedParticleIndicesSRV = VertexBuffer;
 	}
