Fixed compilation warnings due usage of deprecated FComputeShaderRHIParamRef instances

diff --git a/Engine/Plugins/FX/Niagara/Source/Niagara/Private/NiagaraDataInterfaceSkeletalMesh.cpp b/Engine/Plugins/FX/Niagara/Source/Niagara/Private/NiagaraDataInterfaceSkeletalMesh.cpp
--- a/Engine/Plugins/FX/Niagara/Source/Niagara/Private/NiagaraDataInterfaceSkeletalMesh.cpp
+++ b/Engine/Plugins/FX/Niagara/Source/Niagara/Private/NiagaraDataInterfaceSkeletalMesh.cpp
@@ -591,7 +591,7 @@ struct FNiagaraDataInterfaceParametersCS_SkeletalMesh : public FNiagaraDataInter
 	{
 		check(IsInRenderingThread());
 
-		FComputeShaderRHIParamRef ComputeShaderRHI = Context.Shader->GetComputeShader();
+		FRHIComputeShader* ComputeShaderRHI = Context.Shader->GetComputeShader();
 		FNiagaraDataInterfaceProxySkeletalMesh* InterfaceProxy = static_cast<FNiagaraDataInterfaceProxySkeletalMesh*>(Context.DataInterface);
 		FNiagaraDataInterfaceProxySkeletalMeshData* InstanceData = InterfaceProxy->SystemInstancesToData.Find(Context.SystemInstance);
 		if (InstanceData && InstanceData->StaticBuffers)
diff --git a/Engine/Plugins/FX/Niagara/Source/Niagara/Private/NiagaraDataInterfaceStaticMesh.cpp b/Engine/Plugins/FX/Niagara/Source/Niagara/Private/NiagaraDataInterfaceStaticMesh.cpp
--- a/Engine/Plugins/FX/Niagara/Source/Niagara/Private/NiagaraDataInterfaceStaticMesh.cpp
+++ b/Engine/Plugins/FX/Niagara/Source/Niagara/Private/NiagaraDataInterfaceStaticMesh.cpp
@@ -499,7 +499,7 @@ struct FNiagaraDataInterfaceParametersCS_StaticMesh : public FNiagaraDataInterfa
 	{
 		check(IsInRenderingThread());
 
-		FComputeShaderRHIParamRef ComputeShaderRHI = Context.Shader->GetComputeShader();
+		FRHIComputeShader* ComputeShaderRHI = Context.Shader->GetComputeShader();
 		
 		{
 			FNiagaraDataInterfaceProxyStaticMesh* InterfaceProxy = static_cast<FNiagaraDataInterfaceProxyStaticMesh*>(Context.DataInterface);
