Merge fixes

diff --git a/Engine/Plugins/FX/Niagara/Source/Niagara/Private/NiagaraDataInterfaceSkeletalMesh.cpp b/Engine/Plugins/FX/Niagara/Source/Niagara/Private/NiagaraDataInterfaceSkeletalMesh.cpp
--- a/Engine/Plugins/FX/Niagara/Source/Niagara/Private/NiagaraDataInterfaceSkeletalMesh.cpp
+++ b/Engine/Plugins/FX/Niagara/Source/Niagara/Private/NiagaraDataInterfaceSkeletalMesh.cpp
@@ -333,7 +333,7 @@ void FSkeletalMeshGpuSpawnStaticBuffers::InitRHI()
 
 	const FMultiSizeIndexContainer& IndexBuffer = LODRenderData->MultiSizeIndexContainer;
 	MeshIndexBufferSrv = IndexBuffer.GetIndexBuffer()->GetSRV();
-	if (!MeshIndexBufferSrv.IsValid())
+	if (!MeshIndexBufferSrv)
 	{
 		UE_LOG(LogNiagara, Warning, TEXT("Skeletal Mesh does not have an SRV for the index buffer, if you are using triangle sampling it will not work."));
 	}
@@ -430,11 +430,11 @@ void FSkeletalMeshGpuSpawnStaticBuffers::ReleaseRHI()
 	BufferTriangleUniformSamplerAliasRHI.SafeRelease();
 	BufferTriangleUniformSamplerAliasSRV.SafeRelease();
 
-	MeshVertexBufferSrv.SafeRelease();
-	MeshIndexBufferSrv.SafeRelease();
-	MeshTangentBufferSRV.SafeRelease();
-	MeshTexCoordBufferSrv.SafeRelease();
-	MeshColorBufferSrv.SafeRelease();
+	MeshVertexBufferSrv = nullptr;
+	MeshIndexBufferSrv = nullptr;
+	MeshTangentBufferSRV = nullptr;
+	MeshTexCoordBufferSrv = nullptr;
+	MeshColorBufferSrv = nullptr;
 }
 
 //////////////////////////////////////////////////////////////////////////
