Fixed more post merge compilation warnings

diff --git a/Engine/Plugins/FX/Niagara/Source/Niagara/Classes/NiagaraDataInterfaceSkeletalMesh.h b/Engine/Plugins/FX/Niagara/Source/Niagara/Classes/NiagaraDataInterfaceSkeletalMesh.h
--- a/Engine/Plugins/FX/Niagara/Source/Niagara/Classes/NiagaraDataInterfaceSkeletalMesh.h
+++ b/Engine/Plugins/FX/Niagara/Source/Niagara/Classes/NiagaraDataInterfaceSkeletalMesh.h
@@ -261,10 +261,10 @@ public:
 	uint32 GetNumWeights() const { return NumWeights; }
 
 	uint32 GetNumSpecificBones() const { return NumSpecificBones; }
-	FShaderResourceViewRHIParamRef GetSpecificBonesSRV() const { return SpecificBonesSRV; }
+	FRHIShaderResourceView* GetSpecificBonesSRV() const { return SpecificBonesSRV; }
 
 	uint32 GetNumSpecificSocketBones() const { return NumSpecificSocketBones; }
-	FShaderResourceViewRHIParamRef GetSpecificSocketBonesSRV() const { return SpecificSocketBonesSRV; }
+	FRHIShaderResourceView* GetSpecificSocketBonesSRV() const { return SpecificSocketBonesSRV; }
 
 protected:
 
diff --git a/Engine/Plugins/FX/Niagara/Source/Niagara/Private/NiagaraDataInterfaceSkeletalMesh.cpp b/Engine/Plugins/FX/Niagara/Source/Niagara/Private/NiagaraDataInterfaceSkeletalMesh.cpp
--- a/Engine/Plugins/FX/Niagara/Source/Niagara/Private/NiagaraDataInterfaceSkeletalMesh.cpp
+++ b/Engine/Plugins/FX/Niagara/Source/Niagara/Private/NiagaraDataInterfaceSkeletalMesh.cpp
@@ -761,11 +761,11 @@ struct FNiagaraDataInterfaceParametersCS_SkeletalMesh : public FNiagaraDataInter
 				SetSRVParameter(RHICmdList, ComputeShaderRHI, MeshTriangleMatricesOffsetBuffer, FNiagaraRenderer::GetDummyUIntBuffer().SRV);
 			}
 
-			FShaderResourceViewRHIParamRef SpecificBonesSRV = StaticBuffers->GetNumSpecificBones() > 0 ? StaticBuffers->GetSpecificBonesSRV() : FNiagaraRenderer::GetDummyUIntBuffer().SRV.GetReference();
+			FRHIShaderResourceView* SpecificBonesSRV = StaticBuffers->GetNumSpecificBones() > 0 ? StaticBuffers->GetSpecificBonesSRV() : FNiagaraRenderer::GetDummyUIntBuffer().SRV.GetReference();
 			SetShaderValue(RHICmdList, ComputeShaderRHI, NumSpecificBones, StaticBuffers->GetNumSpecificBones());
 			SetSRVParameter(RHICmdList, ComputeShaderRHI, SpecificBones, SpecificBonesSRV);
 
-			FShaderResourceViewRHIParamRef SpecificSocketBonesSRV = StaticBuffers->GetNumSpecificSocketBones() > 0 ? StaticBuffers->GetSpecificSocketBonesSRV() : FNiagaraRenderer::GetDummyUIntBuffer().SRV.GetReference();
+			FRHIShaderResourceView* SpecificSocketBonesSRV = StaticBuffers->GetNumSpecificSocketBones() > 0 ? StaticBuffers->GetSpecificSocketBonesSRV() : FNiagaraRenderer::GetDummyUIntBuffer().SRV.GetReference();
 			SetShaderValue(RHICmdList, ComputeShaderRHI, NumSpecificSocketBones, StaticBuffers->GetNumSpecificSocketBones());
 			SetSRVParameter(RHICmdList, ComputeShaderRHI, SpecificSocketBones, SpecificSocketBonesSRV);
 
