DR - Fix
#rnx

diff --git a/Engine/Plugins/FX/Niagara/Source/Niagara/Private/NiagaraGPUInstanceCountManager.cpp b/Engine/Plugins/FX/Niagara/Source/Niagara/Private/NiagaraGPUInstanceCountManager.cpp
--- a/Engine/Plugins/FX/Niagara/Source/Niagara/Private/NiagaraGPUInstanceCountManager.cpp
+++ b/Engine/Plugins/FX/Niagara/Source/Niagara/Private/NiagaraGPUInstanceCountManager.cpp
@@ -232,7 +232,7 @@ void FNiagaraGPUInstanceCountManager::EnqueueGPUReadback(FRHICommandListImmediat
 	{
 		if (!CountReadback)
 		{
-			CountReadback = new FRHIGPUMemoryReadback(TEXT("Niagara GPU Instance Count Readback"));
+			CountReadback = new FRHIGPUBufferReadback(TEXT("Niagara GPU Instance Count Readback"));
 		}
 		CountReadbackSize = UsedInstanceCounts;
 		CountReadback->EnqueueCopy(RHICmdList, CountBuffer.Buffer);
