DR - SM4 compile fixes

diff --git a/Engine/Plugins/Compositing/LensDistortion/Source/LensDistortion/Private/LensDistortionRendering.cpp b/Engine/Plugins/Compositing/LensDistortion/Source/LensDistortion/Private/LensDistortionRendering.cpp
--- a/Engine/Plugins/Compositing/LensDistortion/Source/LensDistortion/Private/LensDistortionRendering.cpp
+++ b/Engine/Plugins/Compositing/LensDistortion/Source/LensDistortion/Private/LensDistortionRendering.cpp
@@ -57,7 +57,7 @@ class FLensDistortionUVGenerationShader : public FGlobalShader
 public:
 	static bool ShouldCompilePermutation(const FGlobalShaderPermutationParameters& Parameters)
 	{
-		return IsFeatureLevelSupported(Parameters.Platform, ERHIFeatureLevel::SM4);
+		return IsFeatureLevelSupported(Parameters.Platform, ERHIFeatureLevel::SM5);
 	}
 
 	static void ModifyCompilationEnvironment(const FGlobalShaderPermutationParameters& Parameters, FShaderCompilerEnvironment& OutEnvironment)
diff --git a/Engine/Plugins/Compositing/OpenColorIO/Source/OpenColorIO/Public/OpenColorIOShared.h b/Engine/Plugins/Compositing/OpenColorIO/Source/OpenColorIO/Public/OpenColorIOShared.h
--- a/Engine/Plugins/Compositing/OpenColorIO/Source/OpenColorIO/Public/OpenColorIOShared.h
+++ b/Engine/Plugins/Compositing/OpenColorIO/Source/OpenColorIO/Public/OpenColorIOShared.h
@@ -55,7 +55,7 @@ public:
 
 public:
 	FOpenColorIOShaderMapId()
-		: FeatureLevel(ERHIFeatureLevel::SM4)
+		: FeatureLevel(ERHIFeatureLevel::SM5)
 	{ }
 
 	~FOpenColorIOShaderMapId()
@@ -295,7 +295,7 @@ public:
 	FOpenColorIOTransformResource() 
 		: GameThreadShaderMap(nullptr)
 		, RenderingThreadShaderMap(nullptr)
-		, FeatureLevel(ERHIFeatureLevel::SM4)
+		, FeatureLevel(ERHIFeatureLevel::SM5)
 		, bContainsInlineShaders(false)
 		, bLoadedCookedShaderMapId(false)
 	{}
diff --git a/Engine/Plugins/FX/Niagara/Source/Niagara/Public/NiagaraCommon.h b/Engine/Plugins/FX/Niagara/Source/Niagara/Public/NiagaraCommon.h
--- a/Engine/Plugins/FX/Niagara/Source/Niagara/Public/NiagaraCommon.h
+++ b/Engine/Plugins/FX/Niagara/Source/Niagara/Public/NiagaraCommon.h
@@ -578,14 +578,14 @@ namespace FNiagaraUtilities
 
 	inline bool SupportsNiagaraRendering(ERHIFeatureLevel::Type FeatureLevel)
 	{
-		return FeatureLevel == ERHIFeatureLevel::SM4 || FeatureLevel == ERHIFeatureLevel::SM5 || FeatureLevel == ERHIFeatureLevel::ES3_1;
+		return FeatureLevel == ERHIFeatureLevel::SM5 || FeatureLevel == ERHIFeatureLevel::ES3_1;
 	}
 
 	inline bool SupportsNiagaraRendering(EShaderPlatform ShaderPlatform)
 	{
 		// Note:
-		// IsFeatureLevelSupported does a FeatureLevel < MaxFeatureLevel(ShaderPlatform) so checking ES3.1 support will return true for SM4. I added it explicitly to be clear what we are doing.
-		return IsFeatureLevelSupported(ShaderPlatform, ERHIFeatureLevel::SM5) || IsFeatureLevelSupported(ShaderPlatform, ERHIFeatureLevel::ES3_1) || IsFeatureLevelSupported(ShaderPlatform, ERHIFeatureLevel::SM4);
+		// IsFeatureLevelSupported does a FeatureLevel < MaxFeatureLevel(ShaderPlatform) so checking ES3.1 support will return true for SM5. I added it explicitly to be clear what we are doing.
+		return IsFeatureLevelSupported(ShaderPlatform, ERHIFeatureLevel::SM5) || IsFeatureLevelSupported(ShaderPlatform, ERHIFeatureLevel::ES3_1);
 	}
 
 	inline bool SupportsGPUParticles(ERHIFeatureLevel::Type FeatureLevel)
diff --git a/Engine/Plugins/FX/Niagara/Source/NiagaraShader/Private/NiagaraShader.cpp b/Engine/Plugins/FX/Niagara/Source/NiagaraShader/Private/NiagaraShader.cpp
--- a/Engine/Plugins/FX/Niagara/Source/NiagaraShader/Private/NiagaraShader.cpp
+++ b/Engine/Plugins/FX/Niagara/Source/NiagaraShader/Private/NiagaraShader.cpp
@@ -390,7 +390,7 @@ void FNiagaraShaderType::AddReferencedUniformBufferIncludes(FShaderCompilerEnvir
 	GeneratedUniformBuffersInclude.Append(UniformBufferIncludes);
 
 	ERHIFeatureLevel::Type MaxFeatureLevel = GetMaxSupportedFeatureLevel(Platform);
-	if (MaxFeatureLevel >= ERHIFeatureLevel::SM4)
+	if (MaxFeatureLevel >= ERHIFeatureLevel::SM5)
 	{
 		OutEnvironment.SetDefine(TEXT("PLATFORM_SUPPORTS_SRV_UB"), TEXT("1"));
 	}
diff --git a/Engine/Plugins/Media/WmfMedia/Source/WmfMedia/Public/WmfMediaHardwareVideoDecodingShaders.h b/Engine/Plugins/Media/WmfMedia/Source/WmfMedia/Public/WmfMediaHardwareVideoDecodingShaders.h
--- a/Engine/Plugins/Media/WmfMedia/Source/WmfMedia/Public/WmfMediaHardwareVideoDecodingShaders.h
+++ b/Engine/Plugins/Media/WmfMedia/Source/WmfMedia/Public/WmfMediaHardwareVideoDecodingShaders.h
@@ -15,7 +15,7 @@ class FWmfMediaHardwareVideoDecodingShader : public FGlobalShader
 public:
 	static bool ShouldCompilePermutation(const FGlobalShaderPermutationParameters& Parameters)
 	{
-		return IsFeatureLevelSupported(Parameters.Platform, ERHIFeatureLevel::SM4);
+		return IsFeatureLevelSupported(Parameters.Platform, ERHIFeatureLevel::SM5);
 	}
 
 	static void ModifyCompilationEnvironment(const FGlobalShaderPermutationParameters& Parameters, FShaderCompilerEnvironment& OutEnvironment)
