Fix Niagara ribbon read buffers not being held inside the vertex factory which can result in use-after-free

* SRVs are only weak references to the buffers

diff --git a/Engine/Plugins/FX/Niagara/Source/Niagara/Private/NiagaraRendererRibbons.cpp b/Engine/Plugins/FX/Niagara/Source/Niagara/Private/NiagaraRendererRibbons.cpp
--- a/Engine/Plugins/FX/Niagara/Source/Niagara/Private/NiagaraRendererRibbons.cpp
+++ b/Engine/Plugins/FX/Niagara/Source/Niagara/Private/NiagaraRendererRibbons.cpp
@@ -480,7 +480,7 @@ void FNiagaraRendererRibbons::GetDynamicMeshElements(const TArray<const FSceneVi
 			void *IndexPtr = RHILockVertexBuffer(SortedIndicesBuffer.Buffer, 0, DynamicDataRibbon->SortedIndices.Num() * sizeof(int32), RLM_WriteOnly);
 			FMemory::Memcpy(IndexPtr, DynamicDataRibbon->SortedIndices.GetData(), DynamicDataRibbon->SortedIndices.Num() * sizeof(int32));
 			RHIUnlockVertexBuffer(SortedIndicesBuffer.Buffer);
-			CollectorResources.VertexFactory.SetSortedIndices(SortedIndicesBuffer.SRV, 0);
+			CollectorResources.VertexFactory.SetSortedIndices(SortedIndicesBuffer.Buffer, SortedIndicesBuffer.SRV, 0);
 
 			// pass in the CPU generated total segment distance (for tiling distance modes); needs to be a buffer so we can fetch them in the correct order based on Draw Direction (front->back or back->front)
 			//	otherwise UVs will pop when draw direction changes based on camera view point
@@ -489,7 +489,7 @@ void FNiagaraRendererRibbons::GetDynamicMeshElements(const TArray<const FSceneVi
 			void *TangentsAndDistancesPtr = RHILockVertexBuffer(TangentsAndDistancesBuffer.Buffer, 0, DynamicDataRibbon->TangentAndDistances.Num() * sizeof(FVector4), RLM_WriteOnly);
 			FMemory::Memcpy(TangentsAndDistancesPtr, DynamicDataRibbon->TangentAndDistances.GetData(), DynamicDataRibbon->TangentAndDistances.Num() * sizeof(FVector4));
 			RHIUnlockVertexBuffer(TangentsAndDistancesBuffer.Buffer);
-			CollectorResources.VertexFactory.SetTangentAndDistances(TangentsAndDistancesBuffer.SRV);
+			CollectorResources.VertexFactory.SetTangentAndDistances(TangentsAndDistancesBuffer.Buffer, TangentsAndDistancesBuffer.SRV);
 
 			// Copy a buffer which has the per particle multi ribbon index.
 			FReadBuffer MultiRibbonIndicesBuffer;
@@ -497,7 +497,7 @@ void FNiagaraRendererRibbons::GetDynamicMeshElements(const TArray<const FSceneVi
 			void* MultiRibbonIndexPtr = RHILockVertexBuffer(MultiRibbonIndicesBuffer.Buffer, 0, DynamicDataRibbon->MultiRibbonIndices.Num() * sizeof(uint32), RLM_WriteOnly);
 			FMemory::Memcpy(MultiRibbonIndexPtr, DynamicDataRibbon->MultiRibbonIndices.GetData(), DynamicDataRibbon->MultiRibbonIndices.Num() * sizeof(uint32));
 			RHIUnlockVertexBuffer(MultiRibbonIndicesBuffer.Buffer);
-			CollectorResources.VertexFactory.SetMultiRibbonIndicesSRV(MultiRibbonIndicesBuffer.SRV);
+			CollectorResources.VertexFactory.SetMultiRibbonIndicesSRV(MultiRibbonIndicesBuffer.Buffer, MultiRibbonIndicesBuffer.SRV);
 
 			// Copy the packed u data for stable age based uv generation.
 			FReadBuffer PackedPerRibbonDataByIndexBuffer;
@@ -505,7 +505,7 @@ void FNiagaraRendererRibbons::GetDynamicMeshElements(const TArray<const FSceneVi
 			void *PackedPerRibbonDataByIndexPtr = RHILockVertexBuffer(PackedPerRibbonDataByIndexBuffer.Buffer, 0, DynamicDataRibbon->PackedPerRibbonDataByIndex.Num() * sizeof(float), RLM_WriteOnly);
 			FMemory::Memcpy(PackedPerRibbonDataByIndexPtr, DynamicDataRibbon->PackedPerRibbonDataByIndex.GetData(), DynamicDataRibbon->PackedPerRibbonDataByIndex.Num() * sizeof(float));
 			RHIUnlockVertexBuffer(PackedPerRibbonDataByIndexBuffer.Buffer);
-			CollectorResources.VertexFactory.SetPackedPerRibbonDataByIndexSRV(PackedPerRibbonDataByIndexBuffer.SRV);
+			CollectorResources.VertexFactory.SetPackedPerRibbonDataByIndexSRV(PackedPerRibbonDataByIndexBuffer.Buffer, PackedPerRibbonDataByIndexBuffer.SRV);
 
 			FMeshBatch& MeshBatch = Collector.AllocateMesh();
 			MeshBatch.VertexFactory = &CollectorResources.VertexFactory;
diff --git a/Engine/Plugins/FX/Niagara/Source/NiagaraVertexFactories/Public/NiagaraRibbonVertexFactory.h b/Engine/Plugins/FX/Niagara/Source/NiagaraVertexFactories/Public/NiagaraRibbonVertexFactory.h
--- a/Engine/Plugins/FX/Niagara/Source/NiagaraVertexFactories/Public/NiagaraRibbonVertexFactory.h
+++ b/Engine/Plugins/FX/Niagara/Source/NiagaraVertexFactories/Public/NiagaraRibbonVertexFactory.h
@@ -127,25 +127,29 @@ public:
 		FloatDataStride = InFloatDataStride;
 	}
 
-	void SetSortedIndices(const FShaderResourceViewRHIRef& InSortedIndicesSRV, uint32 InSortedIndicesOffset)
+	void SetSortedIndices(const FVertexBufferRHIRef& InSortedIndicesBuffer, const FShaderResourceViewRHIRef& InSortedIndicesSRV, uint32 InSortedIndicesOffset)
 	{
+		SortedIndicesBuffer = InSortedIndicesBuffer;
 		SortedIndicesSRV = InSortedIndicesSRV;
 		SortedIndicesOffset = InSortedIndicesOffset;
 	}
 
-	void SetTangentAndDistances(const FShaderResourceViewRHIRef& InTangentAndDistancesSRV)
+	void SetTangentAndDistances(const FVertexBufferRHIRef& InTangentAndDistancesBuffer, const FShaderResourceViewRHIRef& InTangentAndDistancesSRV)
 	{
+		TangentAndDistancesBuffer = InTangentAndDistancesBuffer;
 		TangentAndDistancesSRV = InTangentAndDistancesSRV;
 	}
 
 
-	void SetMultiRibbonIndicesSRV(const FShaderResourceViewRHIRef& InMultiRibbonIndicesSRV)
+	void SetMultiRibbonIndicesSRV(const FVertexBufferRHIRef& InMultiRibbonIndicesBuffer, const FShaderResourceViewRHIRef& InMultiRibbonIndicesSRV)
 	{
+		MultiRibbonIndicesBuffer = InMultiRibbonIndicesBuffer;
 		MultiRibbonIndicesSRV = InMultiRibbonIndicesSRV;
 	}
 
-	void SetPackedPerRibbonDataByIndexSRV(const FShaderResourceViewRHIRef& InPackedPerRibbonDataByIndexSRV)
+	void SetPackedPerRibbonDataByIndexSRV(const FVertexBufferRHIRef& InPackedPerRibbonDataByIndexBuffer, const FShaderResourceViewRHIRef& InPackedPerRibbonDataByIndexSRV)
 	{
+		PackedPerRibbonDataByIndexBuffer = InPackedPerRibbonDataByIndexBuffer;
 		PackedPerRibbonDataByIndexSRV = InPackedPerRibbonDataByIndexSRV;
 	}
 
@@ -235,6 +239,11 @@ private:
 	uint32 FloatDataOffset;
 	uint32 FloatDataStride;
 
+	FVertexBufferRHIRef SortedIndicesBuffer;
+	FVertexBufferRHIRef TangentAndDistancesBuffer;
+	FVertexBufferRHIRef MultiRibbonIndicesBuffer;
+	FVertexBufferRHIRef PackedPerRibbonDataByIndexBuffer;
+
 	FShaderResourceViewRHIRef SortedIndicesSRV;
 	FShaderResourceViewRHIRef TangentAndDistancesSRV;
 	FShaderResourceViewRHIRef MultiRibbonIndicesSRV;
