Merging //UE4/Dev-Main@8580251 to Dev-RenderPlat-Staging

diff --git a/Engine/Plugins/FX/Niagara/Source/Niagara/Private/NiagaraRendererRibbons.cpp b/Engine/Plugins/FX/Niagara/Source/Niagara/Private/NiagaraRendererRibbons.cpp
--- a/Engine/Plugins/FX/Niagara/Source/Niagara/Private/NiagaraRendererRibbons.cpp
+++ b/Engine/Plugins/FX/Niagara/Source/Niagara/Private/NiagaraRendererRibbons.cpp
@@ -310,9 +310,6 @@ void FNiagaraRendererRibbons::GetDynamicMeshElements(const TArray<const FSceneVi
 	FScopeCycleCounter EmitterStatsCounter(EmitterStatID);
 #endif
 
-	InitPrimitiveUniformBufferIfUninitialized(SceneProxy);
-
-
 	// Compute the per-view uniform buffers.
 	for (int32 ViewIndex = 0; ViewIndex < Views.Num(); ViewIndex++)
 	{
@@ -358,11 +355,6 @@ int FNiagaraRendererRibbons::GetDynamicDataSize()const
 	return Size;
 }
 
-void FNiagaraRendererRibbons::TransformChanged()
-{
-	WorldSpacePrimitiveUniformBuffer.ReleaseResource();
-}
-
 void CalculateUVScaleAndOffsets(
 	const FNiagaraDataSetAccessor<float>& SortKeyData, const TArray<int32>& RibbonIndices,
 	bool bSortKeyIsAge, int32 StartIndex, int32 EndIndex, int32 NumSegments,
@@ -883,33 +875,6 @@ FNiagaraRendererRibbons::FCPUSimParticleDataAllocation FNiagaraRendererRibbons::
 	return CPUSimParticleDataAllocation;
 }
 
-void FNiagaraRendererRibbons::InitPrimitiveUniformBufferIfUninitialized(const FNiagaraSceneProxy *SceneProxy) const
-{
-	if (!WorldSpacePrimitiveUniformBuffer.IsInitialized())
-	{
-		FPrimitiveUniformShaderParameters PrimitiveUniformShaderParameters = GetPrimitiveUniformShaderParameters(
-			FMatrix::Identity,
-			FMatrix::Identity,
-			SceneProxy->GetActorPosition(),
-			SceneProxy->GetBounds(),
-			SceneProxy->GetLocalBounds(),
-			SceneProxy->ReceivesDecals(),
-			false,
-			false,
-			SceneProxy->UseSingleSampleShadowFromStationaryLights(),
-			SceneProxy->GetScene().HasPrecomputedVolumetricLightmap_RenderThread(),
-			SceneProxy->DrawsVelocity(),
-			SceneProxy->GetLightingChannelMask(),
-			0,
-			INDEX_NONE,
-			INDEX_NONE,
-			SceneProxy->AlwaysHasVelocity()
-		);
-		WorldSpacePrimitiveUniformBuffer.SetContents(PrimitiveUniformShaderParameters);
-		WorldSpacePrimitiveUniformBuffer.InitResource();
-	}
-}
-
 void FNiagaraRendererRibbons::CreatePerViewResources(
 	const FSceneView* View, const FSceneViewFamily& ViewFamily, const FNiagaraSceneProxy* SceneProxy, FMeshElementCollector& Collector,
 	uint16& OutVertexCount, uint32& OutNumberOfPrimitives, FNiagaraRibbonUniformBufferRef& OutUniformBuffer, FGlobalDynamicIndexBuffer::FAllocation& InOutIndexAllocation) const
@@ -994,8 +959,7 @@ void FNiagaraRendererRibbons::CreatePerViewResources(
 	GenerateIndexBuffer((uint16*)InOutIndexAllocation.Buffer, OutVertexCount, DynamicDataRibbon->SegmentData, SegmentTessellation, bInvertOrder);
 
 	FNiagaraRibbonUniformParameters PerViewUniformParameters;
-	PerViewUniformParameters.LocalToWorld = bLocalSpace ? SceneProxy->GetLocalToWorld() : FMatrix::Identity;//For now just handle local space like this but maybe in future have a VF variant to avoid the transform entirely?
-	PerViewUniformParameters.LocalToWorldInverseTransposed = bLocalSpace ? SceneProxy->GetLocalToWorld().Inverse().GetTransposed() : FMatrix::Identity;
+	PerViewUniformParameters.bLocalSpace = bLocalSpace;
 	PerViewUniformParameters.DeltaSeconds = ViewFamily.DeltaWorldTime;
 	PerViewUniformParameters.CameraUp = View->GetViewUp(); // FVector4(0.0f, 0.0f, 1.0f, 0.0f);
 	PerViewUniformParameters.CameraRight = View->GetViewRight();//	FVector4(1.0f, 0.0f, 0.0f, 0.0f);
diff --git a/Engine/Plugins/FX/Niagara/Source/Niagara/Public/NiagaraRendererRibbons.h b/Engine/Plugins/FX/Niagara/Source/Niagara/Public/NiagaraRendererRibbons.h
--- a/Engine/Plugins/FX/Niagara/Source/Niagara/Public/NiagaraRendererRibbons.h
+++ b/Engine/Plugins/FX/Niagara/Source/Niagara/Public/NiagaraRendererRibbons.h
@@ -61,8 +61,6 @@ private:
 
 	FCPUSimParticleDataAllocation AllocateParticleDataIfCPUSim(struct FNiagaraDynamicDataRibbon* DynamicDataRibbon, FGlobalDynamicReadBuffer& DynamicReadBuffer) const;
 
-	void InitPrimitiveUniformBufferIfUninitialized(const FNiagaraSceneProxy *SceneProxy) const;
-
 	class FNiagaraRibbonVertexFactory *VertexFactory;
 
 	ENiagaraRibbonFacingMode FacingMode;
