Removed references to WorldSpacePrimitiveUniformBuffer.
#rnx

diff --git a/Engine/Plugins/FX/Niagara/Source/Niagara/Private/NiagaraRendererRibbons.cpp b/Engine/Plugins/FX/Niagara/Source/Niagara/Private/NiagaraRendererRibbons.cpp
--- a/Engine/Plugins/FX/Niagara/Source/Niagara/Private/NiagaraRendererRibbons.cpp
+++ b/Engine/Plugins/FX/Niagara/Source/Niagara/Private/NiagaraRendererRibbons.cpp
@@ -210,7 +210,6 @@ void FNiagaraRendererRibbons::ReleaseRenderThreadResources(NiagaraEmitterInstanc
 {
 	FNiagaraRenderer::ReleaseRenderThreadResources(Batcher);
 	VertexFactory->ReleaseResource();
-	WorldSpacePrimitiveUniformBuffer.ReleaseResource();
 #if RHI_RAYTRACING
 	if (IsRayTracingEnabled())
 	{
@@ -850,7 +849,7 @@ void FNiagaraRendererRibbons::SetupMeshBatchAndCollectorResourceForView(
 	MeshElement.NumInstances = 1;
 	MeshElement.MinVertexIndex = 0;
 	MeshElement.MaxVertexIndex = 0;
-	MeshElement.PrimitiveUniformBuffer = WorldSpacePrimitiveUniformBuffer.GetUniformBufferRHI();
+	MeshElement.PrimitiveUniformBuffer = SceneProxy->GetUniformBuffer();
 
 }
 
