Removed reference to InitPrimitiveUniformBufferIfUninitialized()
#rnx

diff --git a/Engine/Plugins/FX/Niagara/Source/Niagara/Private/NiagaraRendererRibbons.cpp b/Engine/Plugins/FX/Niagara/Source/Niagara/Private/NiagaraRendererRibbons.cpp
--- a/Engine/Plugins/FX/Niagara/Source/Niagara/Private/NiagaraRendererRibbons.cpp
+++ b/Engine/Plugins/FX/Niagara/Source/Niagara/Private/NiagaraRendererRibbons.cpp
@@ -1030,8 +1030,6 @@ void FNiagaraRendererRibbons::GetDynamicRayTracingInstances(FRayTracingMaterialG
 		// Setup material for our ray tracing instance
 		FNiagaraMeshCollectorResourcesRibbon& CollectorResources = Context.RayTracingMeshResourceCollector.AllocateOneFrameResource<FNiagaraMeshCollectorResourcesRibbon>();
 
-		InitPrimitiveUniformBufferIfUninitialized(SceneProxy);
-
 		FGlobalDynamicIndexBuffer::FAllocation DynamicIndexAllocation;
 		uint16 VertexCount;
 
