Fixed niagara ribbon being rendered in local space eventhough they were configured for worldspace.
Cleaned up unused params after VFLooseParams was added the niagara ribbon VF.
#rnx

diff --git a/Engine/Plugins/FX/Niagara/Shaders/Private/NiagaraRibbonVertexFactory.ush b/Engine/Plugins/FX/Niagara/Shaders/Private/NiagaraRibbonVertexFactory.ush
--- a/Engine/Plugins/FX/Niagara/Shaders/Private/NiagaraRibbonVertexFactory.ush
+++ b/Engine/Plugins/FX/Niagara/Shaders/Private/NiagaraRibbonVertexFactory.ush
@@ -596,6 +596,7 @@ void GetNiagaraParticleTangents(in FInterpVtxData InterpVtxData, in out FVertexF
 
 	if (NiagaraRibbonVFLooseParameters.FacingMode == 2 /* CustomSide */)
 	{
+		// Swap axis
 		const float3 Temp = Intermediates.Right;
 		Intermediates.Right = -Intermediates.Facing;
 		Intermediates.Facing = Temp;
diff --git a/Engine/Plugins/FX/Niagara/Source/Niagara/Private/NiagaraRendererRibbons.cpp b/Engine/Plugins/FX/Niagara/Source/Niagara/Private/NiagaraRendererRibbons.cpp
--- a/Engine/Plugins/FX/Niagara/Source/Niagara/Private/NiagaraRendererRibbons.cpp
+++ b/Engine/Plugins/FX/Niagara/Source/Niagara/Private/NiagaraRendererRibbons.cpp
@@ -811,8 +811,7 @@ void FNiagaraRendererRibbons::SetupMeshBatchAndCollectorResourceForView(
 
 	// Collector.AllocateOneFrameResource uses default ctor, initialize the vertex factory
 	CollectorResources.VertexFactory.SetParticleFactoryType(NVFT_Ribbon);
-	CollectorResources.VertexFactory.LooseParameterUniformBuffer
-		= FNiagaraRibbonVFLooseParametersRef::CreateUniformBufferImmediate(VFLooseParams, UniformBuffer_SingleFrame);
+	CollectorResources.VertexFactory.LooseParameterUniformBuffer = FNiagaraRibbonVFLooseParametersRef::CreateUniformBufferImmediate(VFLooseParams, UniformBuffer_SingleFrame);
 	CollectorResources.VertexFactory.InitResource();
 	CollectorResources.VertexFactory.SetRibbonUniformBuffer(CollectorResources.UniformBuffer);
 	CollectorResources.VertexFactory.SetFacingMode(static_cast<uint32>(FacingMode));
@@ -958,6 +957,8 @@ void FNiagaraRendererRibbons::CreatePerViewResources(
 	GenerateIndexBuffer((uint16*)InOutIndexAllocation.Buffer, OutVertexCount, DynamicDataRibbon->SegmentData, SegmentTessellation, bInvertOrder);
 
 	FNiagaraRibbonUniformParameters PerViewUniformParameters;
+	FMemory::Memzero(&PerViewUniformParameters,sizeof(PerViewUniformParameters)); // Clear unset bytes
+
 	PerViewUniformParameters.bLocalSpace = bLocalSpace;
 	PerViewUniformParameters.DeltaSeconds = ViewFamily.DeltaWorldTime;
 	PerViewUniformParameters.CameraUp = View->GetViewUp(); // FVector4(0.0f, 0.0f, 1.0f, 0.0f);
diff --git a/Engine/Plugins/FX/Niagara/Source/NiagaraVertexFactories/Private/NiagaraRibbonVertexFactory.cpp b/Engine/Plugins/FX/Niagara/Source/NiagaraVertexFactories/Private/NiagaraRibbonVertexFactory.cpp
--- a/Engine/Plugins/FX/Niagara/Source/NiagaraVertexFactories/Private/NiagaraRibbonVertexFactory.cpp
+++ b/Engine/Plugins/FX/Niagara/Source/NiagaraVertexFactories/Private/NiagaraRibbonVertexFactory.cpp
@@ -37,12 +37,6 @@ public:
 		NiagaraParticleDataFloat.Bind(ParameterMap, TEXT("NiagaraParticleDataFloat"));
 		FloatDataOffset.Bind(ParameterMap, TEXT("NiagaraFloatDataOffset"));
 		FloatDataStride.Bind(ParameterMap, TEXT("NiagaraFloatDataStride"));
-		SortedIndices.Bind(ParameterMap, TEXT("SortedIndices"));
-		SortedIndicesOffset.Bind(ParameterMap, TEXT("SortedIndicesOffset"));
-		TangentsAndDistances.Bind(ParameterMap, TEXT("TangentsAndDistances"));
-		MultiRibbonIndices.Bind(ParameterMap, TEXT("MultiRibbonIndices"));
-		PackedPerRibbonDataByIndex.Bind(ParameterMap, TEXT("PackedPerRibbonDataByIndex"));
-		FacingMode.Bind(ParameterMap, TEXT("FacingMode"));
 	}
 
 	virtual void Serialize(FArchive& Ar) override
@@ -50,12 +44,6 @@ public:
 		Ar << NiagaraParticleDataFloat;
 		Ar << FloatDataOffset;
 		Ar << FloatDataStride;
-		Ar << SortedIndices;
-		Ar << SortedIndicesOffset;
-		Ar << TangentsAndDistances;
-		Ar << MultiRibbonIndices;
-		Ar << PackedPerRibbonDataByIndex;
-		Ar << FacingMode;
 	}
 
 	virtual void GetElementShaderBindings(
@@ -75,26 +63,12 @@ public:
 		ShaderBindings.Add(NiagaraParticleDataFloat, RibbonVF->GetParticleDataFloatSRV());
 		ShaderBindings.Add(FloatDataOffset, RibbonVF->GetFloatDataOffset());
 		ShaderBindings.Add(FloatDataStride, RibbonVF->GetFloatDataStride());
-
-		ShaderBindings.Add(SortedIndices, RibbonVF->GetSortedIndicesSRV());
-		ShaderBindings.Add(TangentsAndDistances, RibbonVF->GetTangentAndDistancesSRV());
-		ShaderBindings.Add(MultiRibbonIndices, RibbonVF->GetMultiRibbonIndicesSRV());
-		ShaderBindings.Add(PackedPerRibbonDataByIndex, RibbonVF->GetPackedPerRibbonDataByIndexSRV());
-		ShaderBindings.Add(SortedIndicesOffset, RibbonVF->GetSortedIndicesOffset());
-		ShaderBindings.Add(FacingMode, RibbonVF->GetFacingMode());
 	}
 
 private:
 	FShaderResourceParameter NiagaraParticleDataFloat;
 	FShaderParameter FloatDataOffset;
 	FShaderParameter FloatDataStride;
-
-	FShaderResourceParameter SortedIndices;
-	FShaderResourceParameter TangentsAndDistances;
-	FShaderResourceParameter MultiRibbonIndices;
-	FShaderResourceParameter PackedPerRibbonDataByIndex;
-	FShaderParameter SortedIndicesOffset;
-	FShaderParameter FacingMode;
 };
 
 
