Merging //UE4/Dev-Main@9314622 to Dev-RenderPlat-Staging (//UE4/Dev-RenderPlat-Staging)

diff --git a/Engine/Plugins/FX/Niagara/Source/Niagara/Private/NiagaraRendererRibbons.cpp b/Engine/Plugins/FX/Niagara/Source/Niagara/Private/NiagaraRendererRibbons.cpp
--- a/Engine/Plugins/FX/Niagara/Source/Niagara/Private/NiagaraRendererRibbons.cpp
+++ b/Engine/Plugins/FX/Niagara/Source/Niagara/Private/NiagaraRendererRibbons.cpp
@@ -125,6 +125,26 @@ public:
 	}
 };
 
+namespace ENiagaraRibbonVFLayout
+{
+	enum Type
+	{
+		Position,
+		Velocity,
+		Color,
+		Width,
+		Twist,
+		Facing,
+		NormalizedAge,
+		MaterialRandom,
+		MaterialParam0,
+		MaterialParam1,
+		MaterialParam2,
+		MaterialParam3,
+		Num,
+	};
+};
+
 
 FNiagaraRendererRibbons::FNiagaraRendererRibbons(ERHIFeatureLevel::Type FeatureLevel, const UNiagaraRendererProperties *InProps, const FNiagaraEmitterInstance* Emitter)
 	: FNiagaraRenderer(FeatureLevel, InProps, Emitter)
@@ -141,14 +161,6 @@ FNiagaraRendererRibbons::FNiagaraRendererRibbons(ERHIFeatureLevel::Type FeatureL
 	, bCustomUseConstantFactor(false)
 	, CustomTessellationMinAngle(15.f * PI / 180.f)
 	, bCustomUseScreenSpace(true)
-	, PositionDataOffset(INDEX_NONE)
-	, VelocityDataOffset(INDEX_NONE)
-	, WidthDataOffset(INDEX_NONE)
-	, TwistDataOffset(INDEX_NONE)
-	, FacingDataOffset(INDEX_NONE)
-	, ColorDataOffset(INDEX_NONE)
-	, NormalizedAgeDataOffset(INDEX_NONE)
-	, MaterialRandomDataOffset(INDEX_NONE)
 {
 	const UNiagaraRibbonRendererProperties* Properties = CastChecked<const UNiagaraRibbonRendererProperties>(InProps);
 	const FNiagaraDataSet& Data = Emitter->GetData();
@@ -171,28 +183,25 @@ FNiagaraRendererRibbons::FNiagaraRendererRibbons(ERHIFeatureLevel::Type FeatureL
 	CustomTessellationMinAngle *= PI / 180.f;
 	bCustomUseScreenSpace = Properties->bScreenSpaceTessellation;
 
+	TotalVFComponents = 0;
+	VFVariables.SetNum(ENiagaraRibbonVFLayout::Num);
 	// required attributes
-	int32 IntDummy;
-	Data.GetVariableComponentOffsets(Properties->PositionBinding.DataSetVariable, PositionDataOffset, IntDummy);
-	Data.GetVariableComponentOffsets(Properties->VelocityBinding.DataSetVariable, VelocityDataOffset, IntDummy);
-	Data.GetVariableComponentOffsets(Properties->ColorBinding.DataSetVariable, ColorDataOffset, IntDummy);
+	SetVertexFactoryVariable(Data, Properties->PositionBinding.DataSetVariable, ENiagaraRibbonVFLayout::Position);
+	SetVertexFactoryVariable(Data, Properties->VelocityBinding.DataSetVariable, ENiagaraRibbonVFLayout::Velocity);
+	SetVertexFactoryVariable(Data, Properties->ColorBinding.DataSetVariable, ENiagaraRibbonVFLayout::Color);
 
 	// optional attributes
-	Data.GetVariableComponentOffsets(Properties->RibbonWidthBinding.DataSetVariable, WidthDataOffset, IntDummy);
-	Data.GetVariableComponentOffsets(Properties->RibbonTwistBinding.DataSetVariable, TwistDataOffset, IntDummy);
-	Data.GetVariableComponentOffsets(Properties->RibbonFacingBinding.DataSetVariable, FacingDataOffset, IntDummy);
-	Data.GetVariableComponentOffsets(Properties->NormalizedAgeBinding.DataSetVariable, NormalizedAgeDataOffset, IntDummy);
-	Data.GetVariableComponentOffsets(Properties->MaterialRandomBinding.DataSetVariable, MaterialRandomDataOffset, IntDummy);
-
-	Data.GetVariableComponentOffsets(Properties->DynamicMaterialBinding.DataSetVariable, MaterialParamOffset, IntDummy);
-	Data.GetVariableComponentOffsets(Properties->DynamicMaterial1Binding.DataSetVariable, MaterialParamOffset1, IntDummy);
-	Data.GetVariableComponentOffsets(Properties->DynamicMaterial2Binding.DataSetVariable, MaterialParamOffset2, IntDummy);
-	Data.GetVariableComponentOffsets(Properties->DynamicMaterial3Binding.DataSetVariable, MaterialParamOffset3, IntDummy);
-
-	MaterialParamValidMask = MaterialParamOffset != -1 ? 1 : 0;
-	MaterialParamValidMask |= MaterialParamOffset1 != -1 ? 2 : 0;
-	MaterialParamValidMask |= MaterialParamOffset2 != -1 ? 4 : 0;
-	MaterialParamValidMask |= MaterialParamOffset3 != -1 ? 8 : 0;
+	SetVertexFactoryVariable(Data, Properties->RibbonWidthBinding.DataSetVariable, ENiagaraRibbonVFLayout::Width);
+	SetVertexFactoryVariable(Data, Properties->RibbonTwistBinding.DataSetVariable, ENiagaraRibbonVFLayout::Twist);
+	SetVertexFactoryVariable(Data, Properties->RibbonFacingBinding.DataSetVariable, ENiagaraRibbonVFLayout::Facing);
+	SetVertexFactoryVariable(Data, Properties->NormalizedAgeBinding.DataSetVariable, ENiagaraRibbonVFLayout::NormalizedAge);
+	SetVertexFactoryVariable(Data, Properties->MaterialRandomBinding.DataSetVariable, ENiagaraRibbonVFLayout::MaterialRandom);
+
+	MaterialParamValidMask = 0;
+	MaterialParamValidMask |= SetVertexFactoryVariable(Data, Properties->DynamicMaterialBinding.DataSetVariable, ENiagaraRibbonVFLayout::MaterialParam0) ? 1 : 0;
+	MaterialParamValidMask |= SetVertexFactoryVariable(Data, Properties->DynamicMaterial1Binding.DataSetVariable, ENiagaraRibbonVFLayout::MaterialParam1) ? 2 : 0;
+	MaterialParamValidMask |= SetVertexFactoryVariable(Data, Properties->DynamicMaterial2Binding.DataSetVariable, ENiagaraRibbonVFLayout::MaterialParam2) ? 4 : 0;
+	MaterialParamValidMask |= SetVertexFactoryVariable(Data, Properties->DynamicMaterial3Binding.DataSetVariable, ENiagaraRibbonVFLayout::MaterialParam3) ? 8 : 0;
 }
 
 FNiagaraRendererRibbons::~FNiagaraRendererRibbons()
@@ -750,7 +759,9 @@ void FNiagaraRendererRibbons::SetupMeshBatchAndCollectorResourceForView(
 
 	FCPUSimParticleDataAllocation CPUSimParticleDataAllocation = AllocateParticleDataIfCPUSim(DynamicDataRibbon, Collector.GetDynamicReadBuffer());
 	auto& ParticleData = CPUSimParticleDataAllocation.ParticleData;
-	CollectorResources.VertexFactory.SetParticleData(ParticleData.ReadBuffer->SRV, ParticleData.FirstIndex / sizeof(float), SourceParticleData->GetFloatStride() / sizeof(float));
+
+	int32 ParticleDataStride = GbEnableMinimalGPUBuffers ? SourceParticleData->GetNumInstances() : SourceParticleData->GetFloatStride() / sizeof(float);
+	CollectorResources.VertexFactory.SetParticleData(ParticleData.ReadBuffer->SRV, ParticleData.FirstIndex / sizeof(float), ParticleDataStride);
 
 	// TODO: need to make these two a global alloc buffer as well, not recreate
 				// pass in the sorted indices so the VS can fetch the particle data in order
@@ -845,12 +856,19 @@ FNiagaraRendererRibbons::FCPUSimParticleDataAllocation FNiagaraRendererRibbons::
 
 	FCPUSimParticleDataAllocation CPUSimParticleDataAllocation{ DynamicReadBuffer };
 
+	bool bShouldDoFacing = FacingMode == ENiagaraRibbonFacingMode::Custom || FacingMode == ENiagaraRibbonFacingMode::CustomSideVector;
+	VFVariables[ENiagaraRibbonVFLayout::Facing].bUpload = bShouldDoFacing;
+
 	if (SimTarget == ENiagaraSimTarget::CPUSim)
 	{
-		SCOPE_CYCLE_COUNTER(STAT_NiagaraRenderRibbonsCPUSimCopy);
-		CPUSimParticleDataAllocation.ParticleData = DynamicReadBuffer.AllocateFloat(TotalFloatSize);
+		SCOPE_CYCLE_COUNTER(STAT_NiagaraRenderRibbonsCPUSimMemCopy);
+		if (GbEnableMinimalGPUBuffers)
+		{
+			CPUSimParticleDataAllocation.ParticleData = TransferDataToGPU(DynamicReadBuffer, SourceParticleData);
+		}
+		else
 		{
-			SCOPE_CYCLE_COUNTER(STAT_NiagaraRenderRibbonsCPUSimMemCopy);
+			CPUSimParticleDataAllocation.ParticleData = DynamicReadBuffer.AllocateFloat(TotalFloatSize);
 			FMemory::Memcpy(CPUSimParticleDataAllocation.ParticleData.Buffer, SourceParticleData->GetFloatBuffer().GetData(), SourceParticleData->GetFloatBuffer().Num());
 		}
 	}
@@ -953,19 +971,26 @@ void FNiagaraRendererRibbons::CreatePerViewResources(
 	PerViewUniformParameters.InterpCount = SegmentTessellation;
 	PerViewUniformParameters.OneOverInterpCount = 1.f / (float)SegmentTessellation;
 
-	PerViewUniformParameters.PositionDataOffset = PositionDataOffset;
-	PerViewUniformParameters.VelocityDataOffset = VelocityDataOffset;
-	PerViewUniformParameters.ColorDataOffset = ColorDataOffset;
-	PerViewUniformParameters.WidthDataOffset = WidthDataOffset;
-	PerViewUniformParameters.TwistDataOffset = TwistDataOffset;
-	PerViewUniformParameters.FacingDataOffset = FacingMode == ENiagaraRibbonFacingMode::Custom || FacingMode == ENiagaraRibbonFacingMode::CustomSideVector ? FacingDataOffset : -1;
-	PerViewUniformParameters.NormalizedAgeDataOffset = NormalizedAgeDataOffset;
-	PerViewUniformParameters.MaterialRandomDataOffset = MaterialRandomDataOffset;
+	PerViewUniformParameters.PositionDataOffset = VFVariables[ENiagaraRibbonVFLayout::Position].GetGPUOffset();
+	PerViewUniformParameters.VelocityDataOffset = VFVariables[ENiagaraRibbonVFLayout::Velocity].GetGPUOffset();
+	PerViewUniformParameters.ColorDataOffset = VFVariables[ENiagaraRibbonVFLayout::Color].GetGPUOffset();
+	PerViewUniformParameters.WidthDataOffset = VFVariables[ENiagaraRibbonVFLayout::Width].GetGPUOffset();
+	PerViewUniformParameters.TwistDataOffset = VFVariables[ENiagaraRibbonVFLayout::Twist].GetGPUOffset();
+	PerViewUniformParameters.NormalizedAgeDataOffset = VFVariables[ENiagaraRibbonVFLayout::NormalizedAge].GetGPUOffset();
+	PerViewUniformParameters.MaterialRandomDataOffset = VFVariables[ENiagaraRibbonVFLayout::MaterialRandom].GetGPUOffset();
+	PerViewUniformParameters.MaterialParamDataOffset = VFVariables[ENiagaraRibbonVFLayout::MaterialParam0].GetGPUOffset();
+	PerViewUniformParameters.MaterialParam1DataOffset = VFVariables[ENiagaraRibbonVFLayout::MaterialParam1].GetGPUOffset();
+	PerViewUniformParameters.MaterialParam2DataOffset = VFVariables[ENiagaraRibbonVFLayout::MaterialParam2].GetGPUOffset();
+	PerViewUniformParameters.MaterialParam3DataOffset = VFVariables[ENiagaraRibbonVFLayout::MaterialParam3].GetGPUOffset();
+
 	PerViewUniformParameters.MaterialParamValidMask = MaterialParamValidMask;
-	PerViewUniformParameters.MaterialParamDataOffset = MaterialParamOffset;
-	PerViewUniformParameters.MaterialParam1DataOffset = MaterialParamOffset1;
-	PerViewUniformParameters.MaterialParam2DataOffset = MaterialParamOffset2;
-	PerViewUniformParameters.MaterialParam3DataOffset = MaterialParamOffset3;
+
+	bool bShouldDoFacing = FacingMode == ENiagaraRibbonFacingMode::Custom || FacingMode == ENiagaraRibbonFacingMode::CustomSideVector;
+	PerViewUniformParameters.FacingDataOffset = bShouldDoFacing ? VFVariables[ENiagaraRibbonVFLayout::Facing].GetGPUOffset() : -1;
+	PerViewUniformParameters.OneOverUV0TilingDistance = UV0TilingDistance ? 1.f / (UV0TilingDistance) : 0.f;
+	PerViewUniformParameters.OneOverUV1TilingDistance = UV1TilingDistance ? 1.f / (UV1TilingDistance) : 0.f;
+	PerViewUniformParameters.PackedVData = FVector4(UV0Scale.Y, UV0Offset.Y, UV1Scale.Y, UV1Offset.Y);
+
 	PerViewUniformParameters.OneOverUV0TilingDistance = UV0TilingDistance ? 1.f / (UV0TilingDistance) : 0.f;
 	PerViewUniformParameters.OneOverUV1TilingDistance = UV1TilingDistance ? 1.f / (UV1TilingDistance) : 0.f;
 	PerViewUniformParameters.PackedVData = FVector4(UV0Scale.Y, UV0Offset.Y, UV1Scale.Y, UV1Offset.Y);
