Fixed debug dump iterator.

diff --git a/Engine/Plugins/FX/Niagara/Source/Niagara/Classes/NiagaraDataSet.h b/Engine/Plugins/FX/Niagara/Source/Niagara/Classes/NiagaraDataSet.h
--- a/Engine/Plugins/FX/Niagara/Source/Niagara/Classes/NiagaraDataSet.h
+++ b/Engine/Plugins/FX/Niagara/Source/Niagara/Classes/NiagaraDataSet.h
@@ -1445,64 +1445,45 @@ struct FNiagaraDataVariableIterator
 		: Data(InData)
 		, CurrIdx(StartIdx)
 	{
+		Variables = Data->GetOwner()->GetVariables();
 	}
 
 	void Get()
 	{
+		const TArray<FNiagaraVariableLayoutInfo>& VarLayouts = Data->GetOwner()->GetVariableLayouts();
 		for (int32 VarIdx = 0; VarIdx < Variables.Num(); ++VarIdx)
 		{
-			const FNiagaraVariable* Var = Variables[VarIdx];
-			const FNiagaraVariableLayoutInfo* Layout = VarLayouts[VarIdx];
-			check(Var && Layout);
-			const uint8* ValuePtr = Var->GetData();
+			FNiagaraVariable& Var = Variables[VarIdx];
+			const FNiagaraVariableLayoutInfo& Layout = VarLayouts[VarIdx];
+			Var.AllocateData();
+			const uint8* ValuePtr = Var.GetData();
 
-			for (uint32 CompIdx = 0; CompIdx < Layout->GetNumFloatComponents(); ++CompIdx)
+			for (uint32 CompIdx = 0; CompIdx < Layout.GetNumFloatComponents(); ++CompIdx)
 			{
-				uint32 CompBufferOffset = Layout->FloatComponentStart + CompIdx;
+				uint32 CompBufferOffset = Layout.FloatComponentStart + CompIdx;
 				float* Src = Data->GetInstancePtrFloat(CompBufferOffset, CurrIdx);
-				float* Dst = (float*)(ValuePtr + Layout->LayoutInfo.FloatComponentByteOffsets[CompIdx]);
+				float* Dst = (float*)(ValuePtr + Layout.LayoutInfo.FloatComponentByteOffsets[CompIdx]);
 				*Dst = *Src;
 			}
 
-			for (uint32 CompIdx = 0; CompIdx < Layout->GetNumInt32Components(); ++CompIdx)
+			for (uint32 CompIdx = 0; CompIdx < Layout.GetNumInt32Components(); ++CompIdx)
 			{
-				uint32 CompBufferOffset = Layout->Int32ComponentStart + CompIdx;
+				uint32 CompBufferOffset = Layout.Int32ComponentStart + CompIdx;
 				int32* Src = Data->GetInstancePtrInt32(CompBufferOffset, CurrIdx);
-				int32* Dst = (int32*)(ValuePtr + Layout->LayoutInfo.Int32ComponentByteOffsets[CompIdx]);
+				int32* Dst = (int32*)(ValuePtr + Layout.LayoutInfo.Int32ComponentByteOffsets[CompIdx]);
 				*Dst = *Src;
 			}
 		}
 	}
 
 	void Advance() { ++CurrIdx; }
-
-	bool IsValid()const
-	{
-		return Data && CurrIdx < Data->GetNumInstances();
-	}
-
+	bool IsValid()const { return Data && CurrIdx < Data->GetNumInstances(); }
 	uint32 GetCurrIndex()const { return CurrIdx; }
-
-	void AddVariable(FNiagaraVariable* InVar)
-	{
-		check(InVar && Data);
-		Variables.AddUnique(InVar);
-		VarLayouts.AddUnique(Data->GetOwner()->GetVariableLayout(*InVar));
-		InVar->AllocateData();
-	}
-
-	void AddVariables(TArray<FNiagaraVariable>& Vars)
-	{
-		for (FNiagaraVariable& Var : Vars)
-		{
-			AddVariable(&Var);
-		}
-	}
+	const TArray<FNiagaraVariable>& GetVariables()const { return Variables; }
 private:
 
 	const FNiagaraDataBuffer* Data;
-	TArray<const FNiagaraVariable*> Variables;
-	TArray<const FNiagaraVariableLayoutInfo*> VarLayouts;
+	TArray<FNiagaraVariable> Variables;
 
 	uint32 CurrIdx;
 };
@@ -1574,4 +1555,4 @@ FORCEINLINE void FNiagaraDataBuffer::ClearRegisterTable(FNiagaraRegisterTable& R
 {
 	int32 RegisterCount = RegisterTable.Num();
 	Registers.AddZeroed(RegisterCount);
-}
+}
\ No newline at end of file
diff --git a/Engine/Plugins/FX/Niagara/Source/Niagara/Private/NiagaraDataSet.cpp b/Engine/Plugins/FX/Niagara/Source/Niagara/Private/NiagaraDataSet.cpp
--- a/Engine/Plugins/FX/Niagara/Source/Niagara/Private/NiagaraDataSet.cpp
+++ b/Engine/Plugins/FX/Niagara/Source/Niagara/Private/NiagaraDataSet.cpp
@@ -815,9 +815,7 @@ void FNiagaraDataBuffer::CopyTo(FNiagaraDataBuffer& DestBuffer)const
 
 void FNiagaraDataBuffer::Dump(int32 StartIndex, int32 InNumInstances, const FString& Label)const
 {
-	TArray<FNiagaraVariable> Variables = Owner->GetVariables();
 	FNiagaraDataVariableIterator Itr(this, StartIndex);
-	Itr.AddVariables(Variables);
 
 	if (InNumInstances == INDEX_NONE)
 	{
@@ -833,7 +831,7 @@ void FNiagaraDataBuffer::Dump(int32 StartIndex, int32 InNumInstances, const FStr
 		Itr.Get();
 
 		FString Line = TEXT("| ");
-		for (const FNiagaraVariable& Var : Owner->GetVariables())
+		for (const FNiagaraVariable& Var : Itr.GetVariables())
 		{
 			Line += Var.ToString() + TEXT(" | ");
 		}
@@ -1088,4 +1086,4 @@ void FNiagaraDataSetCompiledData::Empty()
 	VariableLayouts.Empty();
 	ID = FNiagaraDataSetID();
 	SimTarget = ENiagaraSimTarget::CPUSim;
-}
+}
\ No newline at end of file
