Niagara - Move the emitter compiled data into shared refs so that there aren't lifetime issues when they're regenerated in the editor while system instances are in the deferred deletion queue.

diff --git a/Engine/Plugins/FX/Niagara/Source/Niagara/Classes/NiagaraEmitterInstance.h b/Engine/Plugins/FX/Niagara/Source/Niagara/Classes/NiagaraEmitterInstance.h
--- a/Engine/Plugins/FX/Niagara/Source/Niagara/Classes/NiagaraEmitterInstance.h
+++ b/Engine/Plugins/FX/Niagara/Source/Niagara/Classes/NiagaraEmitterInstance.h
@@ -22,7 +22,7 @@ struct FNiagaraEmitterHandle;
 class UNiagaraParameterCollection;
 class UNiagaraParameterCollectionInstance;
 class NiagaraEmitterInstanceBatcher;
-
+struct FNiagaraEmitterCompiledData;
 
 /**
 * A Niagara particle simulation.
@@ -244,4 +244,6 @@ private:
 
 	TArray<TNiagaraFastPathAttributeBinding<int32>> FastPathIntAttributeBindings;
 	TArray<TNiagaraFastPathAttributeBinding<float>> FastPathFloatAttributeBindings;
+
+	TSharedPtr<const FNiagaraEmitterCompiledData> CachedEmitterCompiledData;
 };
diff --git a/Engine/Plugins/FX/Niagara/Source/Niagara/Classes/NiagaraSystem.h b/Engine/Plugins/FX/Niagara/Source/Niagara/Classes/NiagaraSystem.h
--- a/Engine/Plugins/FX/Niagara/Source/Niagara/Classes/NiagaraSystem.h
+++ b/Engine/Plugins/FX/Niagara/Source/Niagara/Classes/NiagaraSystem.h
@@ -245,7 +245,7 @@ public:
 	bool ShouldAutoDeactivate() const { return bAutoDeactivate; }
 	bool IsLooping() const;
 
-	const TArray<FNiagaraEmitterCompiledData>& GetEmitterCompiledData()const {	return EmitterCompiledData;	};
+	const TArray<TSharedRef<const FNiagaraEmitterCompiledData>>& GetEmitterCompiledData() const { return EmitterCompiledData; };
 
 	const FNiagaraSystemCompiledData& GetSystemCompiledData() const { return SystemCompiledData; };
 
@@ -348,8 +348,7 @@ protected:
 	UNiagaraScript* SystemUpdateScript;
 
 	//** Post compile generated data used for initializing Emitter Instances during runtime. */
-	UPROPERTY()
-	TArray<FNiagaraEmitterCompiledData> EmitterCompiledData;
+	TArray<TSharedRef<const FNiagaraEmitterCompiledData>> EmitterCompiledData;
 
 	//** Post compile generated data used for initializing System Instances during runtime. */
 	UPROPERTY()
diff --git a/Engine/Plugins/FX/Niagara/Source/Niagara/Private/NiagaraEmitterInstance.cpp b/Engine/Plugins/FX/Niagara/Source/Niagara/Private/NiagaraEmitterInstance.cpp
--- a/Engine/Plugins/FX/Niagara/Source/Niagara/Private/NiagaraEmitterInstance.cpp
+++ b/Engine/Plugins/FX/Niagara/Source/Niagara/Private/NiagaraEmitterInstance.cpp
@@ -238,7 +238,7 @@ void FNiagaraEmitterInstance::Init(int32 InEmitterIdx, FNiagaraSystemInstanceID
 		return;
 	}
 
-	const TArray<FNiagaraEmitterCompiledData>& EmitterCompiledData = ParentSystemInstance->GetSystem()->GetEmitterCompiledData();
+	const TArray<TSharedRef<const FNiagaraEmitterCompiledData>>& EmitterCompiledData = ParentSystemInstance->GetSystem()->GetEmitterCompiledData();
 	if (EmitterCompiledData.IsValidIndex(EmitterIdx) == false)
 	{
 		//@todo(message manager) Error bubbling here
@@ -246,6 +246,8 @@ void FNiagaraEmitterInstance::Init(int32 InEmitterIdx, FNiagaraSystemInstanceID
 		return;
 	}
 
+	CachedEmitterCompiledData = EmitterCompiledData[EmitterIdx];
+
 #if !UE_BUILD_SHIPPING && !UE_BUILD_TEST 
 	CheckForErrors();
 #endif
@@ -263,12 +265,12 @@ void FNiagaraEmitterInstance::Init(int32 InEmitterIdx, FNiagaraSystemInstanceID
 	}
 	else
 	{
-		SpawnInfos.SetNum(EmitterCompiledData[EmitterIdx].SpawnAttributes.Num());
+		SpawnInfos.SetNum(CachedEmitterCompiledData->SpawnAttributes.Num());
 	}
 
 	{
 		SCOPE_CYCLE_COUNTER(STAT_NiagaraEmitterInit_Datasets);
-		Data.Init(&EmitterCompiledData[EmitterIdx].DataSetCompiledData);
+		Data.Init(&CachedEmitterCompiledData->DataSetCompiledData);
 
 		ResetSimulation();
 
@@ -367,28 +369,28 @@ void FNiagaraEmitterInstance::Init(int32 InEmitterIdx, FNiagaraSystemInstanceID
 	{
 		SCOPE_CYCLE_COUNTER(STAT_NiagaraEmitterInit_DirectBindings);
 		//Setup direct bindings for setting parameter values.
-		SpawnIntervalBinding.Init(SpawnExecContext.Parameters, EmitterCompiledData[EmitterIdx].EmitterSpawnIntervalVar);
-		InterpSpawnStartBinding.Init(SpawnExecContext.Parameters, EmitterCompiledData[EmitterIdx].EmitterInterpSpawnStartDTVar);
-		SpawnGroupBinding.Init(SpawnExecContext.Parameters, EmitterCompiledData[EmitterIdx].EmitterSpawnGroupVar);
+		SpawnIntervalBinding.Init(SpawnExecContext.Parameters, CachedEmitterCompiledData->EmitterSpawnIntervalVar);
+		InterpSpawnStartBinding.Init(SpawnExecContext.Parameters, CachedEmitterCompiledData->EmitterInterpSpawnStartDTVar);
+		SpawnGroupBinding.Init(SpawnExecContext.Parameters, CachedEmitterCompiledData->EmitterSpawnGroupVar);
 
-		SpawnEmitterAgeBinding.Init(SpawnExecContext.Parameters, EmitterCompiledData[EmitterIdx].EmitterAgeVar);
-		UpdateEmitterAgeBinding.Init(UpdateExecContext.Parameters, EmitterCompiledData[EmitterIdx].EmitterAgeVar);
+		SpawnEmitterAgeBinding.Init(SpawnExecContext.Parameters, CachedEmitterCompiledData->EmitterAgeVar);
+		UpdateEmitterAgeBinding.Init(UpdateExecContext.Parameters, CachedEmitterCompiledData->EmitterAgeVar);
 		EventEmitterAgeBindings.SetNum(NumEvents);
 		for (int32 i = 0; i < NumEvents; i++)
 		{
-			EventEmitterAgeBindings[i].Init(EventExecContexts[i].Parameters, EmitterCompiledData[EmitterIdx].EmitterAgeVar);
+			EventEmitterAgeBindings[i].Init(EventExecContexts[i].Parameters, CachedEmitterCompiledData->EmitterAgeVar);
 		}
 
 		if (CachedEmitter->SimTarget == ENiagaraSimTarget::GPUComputeSim && GPUExecContext != nullptr)
 		{
-			EmitterAgeBindingGPU.Init(GPUExecContext->CombinedParamStore, EmitterCompiledData[EmitterIdx].EmitterAgeVar);
+			EmitterAgeBindingGPU.Init(GPUExecContext->CombinedParamStore, CachedEmitterCompiledData->EmitterAgeVar);
 		}
 
-		SpawnRandomSeedBinding.Init(SpawnExecContext.Parameters, EmitterCompiledData[EmitterIdx].EmitterRandomSeedVar);
-		UpdateRandomSeedBinding.Init(UpdateExecContext.Parameters, EmitterCompiledData[EmitterIdx].EmitterRandomSeedVar);
+		SpawnRandomSeedBinding.Init(SpawnExecContext.Parameters, CachedEmitterCompiledData->EmitterRandomSeedVar);
+		UpdateRandomSeedBinding.Init(UpdateExecContext.Parameters, CachedEmitterCompiledData->EmitterRandomSeedVar);
 		if (CachedEmitter->SimTarget == ENiagaraSimTarget::GPUComputeSim && GPUExecContext != nullptr)
 		{
-			GPURandomSeedBinding.Init(GPUExecContext->CombinedParamStore, EmitterCompiledData[EmitterIdx].EmitterRandomSeedVar);
+			GPURandomSeedBinding.Init(GPUExecContext->CombinedParamStore, CachedEmitterCompiledData->EmitterRandomSeedVar);
 		}
 
 		// Initialize the exec count
@@ -1157,14 +1159,12 @@ void FNiagaraEmitterInstance::Tick(float DeltaSeconds)
 			TSharedPtr<struct FNiagaraScriptDebuggerInfo, ESPMode::ThreadSafe> DebugInfo = ParentSystemInstance->GetActiveCaptureWrite(CachedIDName, ENiagaraScriptUsage::ParticleGPUComputeScript, FGuid());
 			if (DebugInfo)
 			{
-				const TArray<FNiagaraEmitterCompiledData>& EmitterCompiledData = ParentSystemInstance->GetSystem()->GetEmitterCompiledData();
-
 				//Data.Dump(DebugInfo->Frame, true, 0, OrigNumParticles);
 				//DebugInfo->Frame.Dump(true, 0, OrigNumParticles);
 				DebugInfo->Parameters = GPUExecContext->CombinedParamStore;
 				
 				//TODO: This layout info can be pulled into the emitter/systems etc and all sets just refer to them. They are becoming an annoyance here.
-				DebugInfo->Frame.Init(&EmitterCompiledData[EmitterIdx].GPUCaptureDataSetCompiledData);
+				DebugInfo->Frame.Init(&CachedEmitterCompiledData->GPUCaptureDataSetCompiledData);
 
 				GPUExecContext->DebugInfo = DebugInfo;
 			}
diff --git a/Engine/Plugins/FX/Niagara/Source/Niagara/Private/NiagaraSystem.cpp b/Engine/Plugins/FX/Niagara/Source/Niagara/Private/NiagaraSystem.cpp
--- a/Engine/Plugins/FX/Niagara/Source/Niagara/Private/NiagaraSystem.cpp
+++ b/Engine/Plugins/FX/Niagara/Source/Niagara/Private/NiagaraSystem.cpp
@@ -190,6 +190,27 @@ void UNiagaraSystem::Serialize(FArchive& Ar)
 	Super::Serialize(Ar);
 
 	Ar.UsingCustomVersion(FNiagaraCustomVersion::GUID);
+
+	if (Ar.CustomVer(FNiagaraCustomVersion::GUID) >= FNiagaraCustomVersion::ChangeEmitterCompiledDataToSharedRefs)
+	{
+		UScriptStruct* NiagaraEmitterCompiledDataStruct = FNiagaraEmitterCompiledData::StaticStruct();
+
+		int32 EmitterCompiledDataNum;
+		if (Ar.IsSaving())
+		{
+			EmitterCompiledDataNum = EmitterCompiledData.Num();
+		}
+		Ar << EmitterCompiledDataNum;
+
+		for (int32 EmitterIndex = 0; EmitterIndex < EmitterCompiledDataNum; ++EmitterIndex)
+		{
+			if (Ar.IsLoading())
+			{
+				EmitterCompiledData.Add(MakeShared<FNiagaraEmitterCompiledData>());
+			}
+			NiagaraEmitterCompiledDataStruct->SerializeTaggedProperties(Ar, (uint8*)&ConstCastSharedRef<FNiagaraEmitterCompiledData>(EmitterCompiledData[EmitterIndex]).Get(), NiagaraEmitterCompiledDataStruct, nullptr);
+		}
+	}
 }
 
 #if WITH_EDITOR
@@ -369,7 +390,7 @@ void UNiagaraSystem::PostLoad()
 			UE_LOG(LogNiagara, Log, TEXT("System %s being compiled because there were changes to an emitter script Change ID."), *GetPathName());
 		}
 
-		if (EmitterCompiledData.Num() == 0 || EmitterCompiledData[0].DataSetCompiledData.Variables.Num() == 0)
+		if (EmitterCompiledData.Num() == 0 || EmitterCompiledData[0]->DataSetCompiledData.Variables.Num() == 0)
 		{
 			InitEmitterCompiledData();
 		}
@@ -1106,8 +1127,13 @@ void UNiagaraSystem::InitEmitterCompiledData()
 	EmitterCompiledData.Empty();
 	if (SystemSpawnScript->GetVMExecutableData().IsValid() && SystemUpdateScript->GetVMExecutableData().IsValid())
 	{
-		EmitterCompiledData.SetNum(EmitterHandles.Num());
-	FNiagaraTypeDefinition SpawnInfoDef = FNiagaraTypeDefinition(FNiagaraSpawnInfo::StaticStruct());
+		TArray<TSharedRef<FNiagaraEmitterCompiledData>> NewEmitterCompiledData;
+		for (int32 EmitterIdx = 0; EmitterIdx < EmitterHandles.Num(); ++EmitterIdx)
+		{
+			NewEmitterCompiledData.Add(MakeShared<FNiagaraEmitterCompiledData>());
+		}
+
+		FNiagaraTypeDefinition SpawnInfoDef = FNiagaraTypeDefinition(FNiagaraSpawnInfo::StaticStruct());
 
 		for (FNiagaraVariable& Var : SystemSpawnScript->GetVMExecutableData().Attributes)
 		{
@@ -1119,7 +1145,7 @@ void UNiagaraSystem::InitEmitterCompiledData()
 					FString EmitterName = Emitter->GetUniqueEmitterName() + TEXT(".");
 					if (Var.GetType() == SpawnInfoDef && Var.GetName().ToString().StartsWith(EmitterName))
 					{
-						EmitterCompiledData[EmitterIdx].SpawnAttributes.AddUnique(Var.GetName());
+						NewEmitterCompiledData[EmitterIdx]->SpawnAttributes.AddUnique(Var.GetName());
 					}
 				}
 			}
@@ -1135,22 +1161,22 @@ void UNiagaraSystem::InitEmitterCompiledData()
 					FString EmitterName = Emitter->GetUniqueEmitterName() + TEXT(".");
 					if (Var.GetType() == SpawnInfoDef && Var.GetName().ToString().StartsWith(EmitterName))
 					{
-						EmitterCompiledData[EmitterIdx].SpawnAttributes.AddUnique(Var.GetName());
+						NewEmitterCompiledData[EmitterIdx]->SpawnAttributes.AddUnique(Var.GetName());
 					}
 				}
-				}
 			}
+		}
 
 		for (int32 EmitterIdx = 0; EmitterIdx < EmitterHandles.Num(); ++EmitterIdx)
 		{
 			const FNiagaraEmitterHandle& EmitterHandle = EmitterHandles[EmitterIdx];
 			const UNiagaraEmitter* Emitter = EmitterHandle.GetInstance();
-			FNiagaraDataSetCompiledData& EmitterDataSetCompiledData = EmitterCompiledData[EmitterIdx].DataSetCompiledData;
-			FNiagaraDataSetCompiledData& GPUCaptureCompiledData = EmitterCompiledData[EmitterIdx].GPUCaptureDataSetCompiledData;
+			FNiagaraDataSetCompiledData& EmitterDataSetCompiledData = NewEmitterCompiledData[EmitterIdx]->DataSetCompiledData;
+			FNiagaraDataSetCompiledData& GPUCaptureCompiledData = NewEmitterCompiledData[EmitterIdx]->GPUCaptureDataSetCompiledData;
 			if ensureMsgf(Emitter != nullptr, TEXT("Failed to get Emitter Instance from Emitter Handle in post compile, please investigate."))
 			{
 				static FName GPUCaptureDataSetName = TEXT("GPU Capture Dataset");
-				InitEmitterVariableAliasNames(EmitterCompiledData[EmitterIdx], Emitter);
+				InitEmitterVariableAliasNames(NewEmitterCompiledData[EmitterIdx].Get(), Emitter);
 				InitEmitterDataSetCompiledData(EmitterDataSetCompiledData, Emitter, EmitterHandle);
 				GPUCaptureCompiledData.ID = FNiagaraDataSetID(GPUCaptureDataSetName, ENiagaraDataSetType::ParticleData);
 				GPUCaptureCompiledData.Variables = EmitterDataSetCompiledData.Variables;
@@ -1158,6 +1184,11 @@ void UNiagaraSystem::InitEmitterCompiledData()
 				GPUCaptureCompiledData.BuildLayout();				
 			}
 		}
+
+		for (int32 EmitterIdx = 0; EmitterIdx < EmitterHandles.Num(); ++EmitterIdx)
+		{
+			EmitterCompiledData.Add(NewEmitterCompiledData[EmitterIdx]);
+		}
 	}
 }
 
diff --git a/Engine/Plugins/FX/Niagara/Source/Niagara/Private/NiagaraSystemSimulation.cpp b/Engine/Plugins/FX/Niagara/Source/Niagara/Private/NiagaraSystemSimulation.cpp
--- a/Engine/Plugins/FX/Niagara/Source/Niagara/Private/NiagaraSystemSimulation.cpp
+++ b/Engine/Plugins/FX/Niagara/Source/Niagara/Private/NiagaraSystemSimulation.cpp
@@ -352,10 +352,10 @@ bool FNiagaraSystemSimulation::Init(UNiagaraSystem* InSystem, UWorld* InWorld, b
 				{
 					FString EmitterName = Emitter->GetUniqueEmitterName();
 					EmitterExecutionStateAccessors.Emplace(MainDataSet, FNiagaraVariable(EnumPtr, *(EmitterName + TEXT(".ExecutionState"))));
-					const TArray<FNiagaraEmitterCompiledData>& EmitterCompiledData = System->GetEmitterCompiledData();
+					const TArray<TSharedRef<const FNiagaraEmitterCompiledData>>& EmitterCompiledData = System->GetEmitterCompiledData();
 
 					check(EmitterCompiledData.Num() == System->GetNumEmitters());
-					for (FName AttrName : EmitterCompiledData[EmitterIdx].SpawnAttributes)
+					for (FName AttrName : EmitterCompiledData[EmitterIdx]->SpawnAttributes)
 					{
 						EmitterSpawnInfoAccessors[EmitterIdx].Emplace(MainDataSet, FNiagaraVariable(FNiagaraTypeDefinition(FNiagaraSpawnInfo::StaticStruct()), AttrName));
 					}
diff --git a/Engine/Plugins/FX/Niagara/Source/NiagaraCore/Public/NiagaraCustomVersion.h b/Engine/Plugins/FX/Niagara/Source/NiagaraCore/Public/NiagaraCustomVersion.h
--- a/Engine/Plugins/FX/Niagara/Source/NiagaraCore/Public/NiagaraCustomVersion.h
+++ b/Engine/Plugins/FX/Niagara/Source/NiagaraCore/Public/NiagaraCustomVersion.h
@@ -130,6 +130,8 @@ struct FNiagaraCustomVersion
 			
 		AddRIAndDetailLevel, //Adding UseRapidIterationParams and DetailLevelMask to the GPU script
 
+		ChangeEmitterCompiledDataToSharedRefs, // Changing the system and emitter compiled data to shared pointers to deal with lifetime issues in the editor.  They now are handled directly in system serialize.
+
 		// DO NOT ADD A NEW VERSION UNLESS YOU HAVE TALKED TO THE NIAGARA LEAD. Mismanagement of these versions can lead to data loss if it is adjusted in multiple streams simultaneously.
 		// -----<new versions can be added above this line>  -------------------------------------------------
 		VersionPlusOne,
