- Add Niagara stat for setting data buffers on the render thread

#rnx

diff --git a/Engine/Plugins/FX/Niagara/Source/Niagara/Private/NiagaraComponent.cpp b/Engine/Plugins/FX/Niagara/Source/Niagara/Private/NiagaraComponent.cpp
--- a/Engine/Plugins/FX/Niagara/Source/Niagara/Private/NiagaraComponent.cpp
+++ b/Engine/Plugins/FX/Niagara/Source/Niagara/Private/NiagaraComponent.cpp
@@ -28,6 +28,7 @@ DECLARE_CYCLE_STAT(TEXT("Component Tick (GT)"), STAT_NiagaraComponentTick, STATG
 DECLARE_CYCLE_STAT(TEXT("Activate (GT)"), STAT_NiagaraComponentActivate, STATGROUP_Niagara);
 DECLARE_CYCLE_STAT(TEXT("Deactivate (GT)"), STAT_NiagaraComponentDeactivate, STATGROUP_Niagara);
 DECLARE_CYCLE_STAT(TEXT("Send Render Data (GT)"), STAT_NiagaraComponentSendRenderData, STATGROUP_Niagara);
+DECLARE_CYCLE_STAT(TEXT("Set Dynamic Data (RT)"), STAT_NiagaraSetDynamicData, STATGROUP_Niagara);
 DECLARE_CYCLE_STAT(TEXT("Get Dynamic Mesh Elements (RT)"), STAT_NiagaraComponentGetDynamicMeshElements, STATGROUP_Niagara);
 
 DEFINE_LOG_CATEGORY(LogNiagara);
@@ -1082,9 +1083,11 @@ void UNiagaraComponent::SendRenderDynamicData_Concurrent()
 			}
 		}
 		
-		ENQUEUE_RENDER_COMMAND(ResetDataSetBuffers)(
+		ENQUEUE_RENDER_COMMAND(NiagaraSetDynamicData)(
 			[NiagaraProxy, DynamicData = MoveTemp(NewDynamicData)](FRHICommandListImmediate& RHICmdList)
 		{
+			SCOPE_CYCLE_COUNTER(STAT_NiagaraSetDynamicData);
+
 			const TArray<FNiagaraRenderer*>& EmitterRenderers_RT = NiagaraProxy->GetEmitterRenderers();
 			for (int32 i = 0; i < EmitterRenderers_RT.Num(); ++i)
 			{
