- Cache NiagaraMeshVertexFactories inside renderer, saved ~30us on GDME

#rnx

diff --git a/Engine/Plugins/FX/Niagara/Source/Niagara/Private/NiagaraRendererMeshes.cpp b/Engine/Plugins/FX/Niagara/Source/Niagara/Private/NiagaraRendererMeshes.cpp
--- a/Engine/Plugins/FX/Niagara/Source/Niagara/Private/NiagaraRendererMeshes.cpp
+++ b/Engine/Plugins/FX/Niagara/Source/Niagara/Private/NiagaraRendererMeshes.cpp
@@ -43,12 +43,16 @@ struct FNiagaraDynamicDataMesh : public FNiagaraDynamicDataBase
 class FNiagaraMeshCollectorResourcesMesh : public FOneFrameResource
 {
 public:
-	FNiagaraMeshVertexFactory VertexFactory;
+	FNiagaraMeshVertexFactory* VertexFactory = nullptr;
 	FNiagaraMeshUniformBufferRef UniformBuffer;
 
 	virtual ~FNiagaraMeshCollectorResourcesMesh()
 	{
-		VertexFactory.ReleaseResource();
+		if ( VertexFactory )
+		{
+			VertexFactory->SetParticleData(nullptr, 0, 0);
+			VertexFactory->SetSortedIndices(nullptr, 0xFFFFFFFF);
+		}
 	}
 };
 
@@ -115,11 +119,22 @@ FNiagaraRendererMeshes::FNiagaraRendererMeshes(ERHIFeatureLevel::Type FeatureLev
 
 FNiagaraRendererMeshes::~FNiagaraRendererMeshes()
 {
+	for ( FNiagaraMeshVertexFactory* VertexFactory : VertexFactories )
+	{
+		delete VertexFactory;
+	}
+	VertexFactories.Empty();
 }
 
 void FNiagaraRendererMeshes::ReleaseRenderThreadResources(NiagaraEmitterInstanceBatcher* Batcher)
 {
 	FNiagaraRenderer::ReleaseRenderThreadResources(Batcher);
+	for (FNiagaraMeshVertexFactory* VertexFactory : VertexFactories)
+	{
+		VertexFactory->ReleaseResource();
+		delete VertexFactory;
+	}
+	VertexFactories.Empty();
 }
 
 void FNiagaraRendererMeshes::CreateRenderThreadResources(NiagaraEmitterInstanceBatcher* Batcher)
@@ -212,6 +227,8 @@ void FNiagaraRendererMeshes::GetDynamicMeshElements(const TArray<const FSceneVie
 	}
 
 	{
+		int32 VertexFactoryIndex = 0;
+
 		// Compute the per-view uniform buffers.
 		for (int32 ViewIndex = 0; ViewIndex < Views.Num(); ViewIndex++)
 		{
@@ -228,8 +245,35 @@ void FNiagaraRendererMeshes::GetDynamicMeshElements(const TArray<const FSceneVie
 				}
 				const FStaticMeshLODResources& LODModel = MeshRenderData->LODResources[LODIndex];
 
+				// Get the next vertex factory to use
+				FNiagaraMeshVertexFactory* VertexFactory = nullptr;
+				if ( VertexFactories.IsValidIndex(VertexFactoryIndex) )
+				{
+					VertexFactory = VertexFactories[VertexFactoryIndex];
+					++VertexFactoryIndex;
+
+					if ( VertexFactory->GetLODIndex() != LODIndex )
+					{
+						SetupVertexFactory(VertexFactory, LODModel);
+						VertexFactory->SetLODIndex(LODIndex);
+					}
+				}
+				else
+				{
+					check(VertexFactoryIndex == VertexFactories.Num());
+					VertexFactory = new FNiagaraMeshVertexFactory();
+					VertexFactories.Add(VertexFactory);
+
+					VertexFactory->SetParticleFactoryType(NVFT_Mesh);
+					VertexFactory->SetMeshFacingMode((uint32)FacingMode);
+					VertexFactory->SetLODIndex(LODIndex);
+					VertexFactory->InitResource();
+					SetupVertexFactory(VertexFactory, LODModel);
+				}
+
 				FNiagaraMeshCollectorResourcesMesh& CollectorResources = Collector.AllocateOneFrameResource<FNiagaraMeshCollectorResourcesMesh>();
-				SetupVertexFactory(&CollectorResources.VertexFactory, LODModel);
+				CollectorResources.VertexFactory = VertexFactory;
+
 				FNiagaraMeshUniformParameters PerViewUniformParameters;// = UniformParameters;
 				FMemory::Memzero(&PerViewUniformParameters, sizeof(PerViewUniformParameters)); // Clear unset bytes
 
@@ -254,7 +298,7 @@ void FNiagaraRendererMeshes::GetDynamicMeshElements(const TArray<const FSceneVie
 				PerViewUniformParameters.DefaultPos = bLocalSpace ? FVector4(0.0f, 0.0f, 0.0f, 1.0f) : FVector4(SceneProxy->GetLocalToWorld().GetOrigin());
 
 				//Sort particles if needed.
-				CollectorResources.VertexFactory.SetSortedIndices(nullptr, 0xFFFFFFFF);
+				CollectorResources.VertexFactory->SetSortedIndices(nullptr, 0xFFFFFFFF);
 
 				FNiagaraGPUSortInfo SortInfo;
 				SortInfo.SortAttributeOffset = INDEX_NONE;
@@ -300,7 +344,7 @@ void FNiagaraRendererMeshes::GetDynamicMeshElements(const TArray<const FSceneVie
 							const int32 IndexBufferOffset = Batcher->AddSortedGPUSimulation(SortInfo);
 							if (IndexBufferOffset != INDEX_NONE)
 							{
-								CollectorResources.VertexFactory.SetSortedIndices(Batcher->GetGPUSortedBuffer().VertexBufferSRV, IndexBufferOffset);
+								CollectorResources.VertexFactory->SetSortedIndices(Batcher->GetGPUSortedBuffer().VertexBufferSRV, IndexBufferOffset);
 							}
 						}
 						else
@@ -308,11 +352,11 @@ void FNiagaraRendererMeshes::GetDynamicMeshElements(const TArray<const FSceneVie
 							FGlobalDynamicReadBuffer::FAllocation SortedIndices;
 							SortedIndices = DynamicReadBuffer.AllocateInt32(NumInstances);
 							SortIndices(SortInfo, SortVarIdx, *SourceParticleData, SortedIndices);
-							CollectorResources.VertexFactory.SetSortedIndices(SortedIndices.ReadBuffer->SRV, SortedIndices.FirstIndex / sizeof(float));
+							CollectorResources.VertexFactory->SetSortedIndices(SortedIndices.ReadBuffer->SRV, SortedIndices.FirstIndex / sizeof(float));
 						}
 					}
 					int32 ParticleDataStride = GbEnableMinimalGPUBuffers ? SourceParticleData->GetNumInstances() : SourceParticleData->GetFloatStride() / sizeof(float);
-					CollectorResources.VertexFactory.SetParticleData(ParticleData.ReadBuffer->SRV, ParticleData.FirstIndex / sizeof(float), ParticleDataStride);
+					CollectorResources.VertexFactory->SetParticleData(ParticleData.ReadBuffer->SRV, ParticleData.FirstIndex / sizeof(float), ParticleDataStride);
 				}
 				else
 				{
@@ -328,26 +372,22 @@ void FNiagaraRendererMeshes::GetDynamicMeshElements(const TArray<const FSceneVie
 						const int32 IndexBufferOffset = Batcher->AddSortedGPUSimulation(SortInfo);
 						if (IndexBufferOffset != INDEX_NONE && SortInfo.GPUParticleCountOffset != INDEX_NONE)
 						{
-							CollectorResources.VertexFactory.SetSortedIndices(Batcher->GetGPUSortedBuffer().VertexBufferSRV, IndexBufferOffset);
+							CollectorResources.VertexFactory->SetSortedIndices(Batcher->GetGPUSortedBuffer().VertexBufferSRV, IndexBufferOffset);
 						}
 					}
 					if (SourceParticleData->GetGPUBufferFloat().SRV.IsValid())
 					{
-						CollectorResources.VertexFactory.SetParticleData(SourceParticleData->GetGPUBufferFloat().SRV, 0, SourceParticleData->GetFloatStride() / sizeof(float));
+						CollectorResources.VertexFactory->SetParticleData(SourceParticleData->GetGPUBufferFloat().SRV, 0, SourceParticleData->GetFloatStride() / sizeof(float));
 					}
 					else
 					{
-						CollectorResources.VertexFactory.SetParticleData(FNiagaraRenderer::GetDummyFloatBuffer().SRV, 0, 0);
+						CollectorResources.VertexFactory->SetParticleData(FNiagaraRenderer::GetDummyFloatBuffer().SRV, 0, 0);
 					}
 				}
 
 				// Collector.AllocateOneFrameResource uses default ctor, initialize the vertex factory
-				CollectorResources.VertexFactory.SetParticleFactoryType(NVFT_Mesh);
-				CollectorResources.VertexFactory.SetMeshFacingMode((uint32)FacingMode);
 				CollectorResources.UniformBuffer = FNiagaraMeshUniformBufferRef::CreateUniformBufferImmediate(PerViewUniformParameters, UniformBuffer_SingleFrame);
-
-				CollectorResources.VertexFactory.InitResource();
-				CollectorResources.VertexFactory.SetUniformBuffer(CollectorResources.UniformBuffer);
+				CollectorResources.VertexFactory->SetUniformBuffer(CollectorResources.UniformBuffer);
 			
 				// GPU mesh rendering currently only supports one mesh section.
 				// TODO: Add proper support for multiple mesh sections for GPU mesh particles.
@@ -364,7 +404,7 @@ void FNiagaraRendererMeshes::GetDynamicMeshElements(const TArray<const FSceneVie
 					}
 
 					FMeshBatch& Mesh = Collector.AllocateMesh();
-					Mesh.VertexFactory = &CollectorResources.VertexFactory;
+					Mesh.VertexFactory = CollectorResources.VertexFactory;
 					Mesh.LCI = NULL;
 					Mesh.ReverseCulling = SceneProxy->IsLocalToWorldDeterminantNegative();
 					Mesh.CastShadow = SceneProxy->CastsDynamicShadow();
diff --git a/Engine/Plugins/FX/Niagara/Source/Niagara/Public/NiagaraRendererMeshes.h b/Engine/Plugins/FX/Niagara/Source/Niagara/Public/NiagaraRendererMeshes.h
--- a/Engine/Plugins/FX/Niagara/Source/Niagara/Public/NiagaraRendererMeshes.h
+++ b/Engine/Plugins/FX/Niagara/Source/Niagara/Public/NiagaraRendererMeshes.h
@@ -32,7 +32,7 @@ public:
 	void SetupVertexFactory(FNiagaraMeshVertexFactory *InVertexFactory, const FStaticMeshLODResources& LODResources) const;
 
 private:
-
+	mutable TArray<class FNiagaraMeshVertexFactory*, TInlineAllocator<2>> VertexFactories;
 	/** Render data of the static mesh we use. */
 	FStaticMeshRenderData* MeshRenderData;
 
diff --git a/Engine/Plugins/FX/Niagara/Source/NiagaraVertexFactories/Private/NiagaraMeshVertexFactory.cpp b/Engine/Plugins/FX/Niagara/Source/NiagaraVertexFactories/Private/NiagaraMeshVertexFactory.cpp
--- a/Engine/Plugins/FX/Niagara/Source/NiagaraVertexFactories/Private/NiagaraMeshVertexFactory.cpp
+++ b/Engine/Plugins/FX/Niagara/Source/NiagaraVertexFactories/Private/NiagaraMeshVertexFactory.cpp
@@ -70,7 +70,8 @@ public:
 		ShaderBindings.Add(FloatDataOffset, NiagaraMeshVF->GetFloatDataOffset());
 		ShaderBindings.Add(FloatDataStride, NiagaraMeshVF->GetFloatDataStride());
 
-		ShaderBindings.Add(SortedIndices, NiagaraMeshVF->GetSortedIndicesSRV() ? NiagaraMeshVF->GetSortedIndicesSRV() : GFNiagaraNullSortedIndicesVertexBuffer.VertexBufferSRV);
+		FRHIShaderResourceView* SortedSRV = NiagaraMeshVF->GetSortedIndicesSRV();
+		ShaderBindings.Add(SortedIndices, SortedSRV != nullptr ? SortedSRV : GFNiagaraNullSortedIndicesVertexBuffer.VertexBufferSRV.GetReference());
 		ShaderBindings.Add(SortedIndicesOffset, NiagaraMeshVF->GetSortedIndicesOffset());
 	}
 
diff --git a/Engine/Plugins/FX/Niagara/Source/NiagaraVertexFactories/Public/NiagaraMeshVertexFactory.h b/Engine/Plugins/FX/Niagara/Source/NiagaraVertexFactories/Public/NiagaraMeshVertexFactory.h
--- a/Engine/Plugins/FX/Niagara/Source/NiagaraVertexFactories/Public/NiagaraMeshVertexFactory.h
+++ b/Engine/Plugins/FX/Niagara/Source/NiagaraVertexFactories/Public/NiagaraMeshVertexFactory.h
@@ -66,6 +66,7 @@ public:
 	/** Default constructor. */
 	FNiagaraMeshVertexFactory(ENiagaraVertexFactoryType InType, ERHIFeatureLevel::Type InFeatureLevel)
 		: FNiagaraVertexFactoryBase(InType, InFeatureLevel)
+		, LODIndex(-1)
 		, MeshFacingMode(0)
 		, InstanceVerticesCPU(nullptr)
 		, FloatDataOffset(0)
@@ -75,6 +76,7 @@ public:
 
 	FNiagaraMeshVertexFactory()
 		: FNiagaraVertexFactoryBase(NVFT_MAX, ERHIFeatureLevel::Num)
+		, LODIndex(-1)
 		, MeshFacingMode(0)
 		, InstanceVerticesCPU(nullptr)
 		, FloatDataOffset(0)
@@ -114,7 +116,7 @@ public:
 		SortedIndicesOffset = InSortedIndicesOffset;
 	}
 
-	FORCEINLINE FShaderResourceViewRHIRef GetParticleDataFloatSRV()
+	FORCEINLINE FRHIShaderResourceView* GetParticleDataFloatSRV()
 	{
 		return ParticleDataFloatSRV;
 	}
@@ -129,7 +131,7 @@ public:
 		return FloatDataStride;
 	}
 
-	FORCEINLINE FShaderResourceViewRHIRef GetSortedIndicesSRV()
+	FORCEINLINE FRHIShaderResourceView* GetSortedIndicesSRV()
 	{
 		return SortedIndicesSRV;
 	}
@@ -177,18 +179,15 @@ public:
 
 	static FVertexFactoryShaderParameters* ConstructShaderParameters(EShaderFrequency ShaderFrequency);
 	
-	uint32 GetMeshFacingMode() const
-	{
-		return MeshFacingMode;
-	}
+	int32 GetLODIndex() const { return LODIndex; }
+	void SetLODIndex(int32 InLODIndex) { LODIndex = InLODIndex; }
 
-	void SetMeshFacingMode(uint32 InMode)
-	{
-		MeshFacingMode = InMode;
-	}
+	uint32 GetMeshFacingMode() const { return MeshFacingMode; }
+	void SetMeshFacingMode(uint32 InMode) { MeshFacingMode = InMode; }
 
 protected:
 	FStaticMeshDataType Data;
+	int32 LODIndex;
 	uint32 MeshFacingMode;
 
 	/** Uniform buffer with mesh particle parameters. */
