Fixing uninitialized memory issue to fix map hitting a check on load due to CL 9584846

diff --git a/Engine/Plugins/FX/Niagara/Source/Niagara/Private/NiagaraSystem.cpp b/Engine/Plugins/FX/Niagara/Source/Niagara/Private/NiagaraSystem.cpp
--- a/Engine/Plugins/FX/Niagara/Source/Niagara/Private/NiagaraSystem.cpp
+++ b/Engine/Plugins/FX/Niagara/Source/Niagara/Private/NiagaraSystem.cpp
@@ -195,7 +195,7 @@ void UNiagaraSystem::Serialize(FArchive& Ar)
 	{
 		UScriptStruct* NiagaraEmitterCompiledDataStruct = FNiagaraEmitterCompiledData::StaticStruct();
 
-		int32 EmitterCompiledDataNum;
+		int32 EmitterCompiledDataNum = 0;
 		if (Ar.IsSaving())
 		{
 			EmitterCompiledDataNum = EmitterCompiledData.Num();
@@ -208,6 +208,7 @@ void UNiagaraSystem::Serialize(FArchive& Ar)
 			{
 				EmitterCompiledData.Add(MakeShared<FNiagaraEmitterCompiledData>());
 			}
+
 			NiagaraEmitterCompiledDataStruct->SerializeTaggedProperties(Ar, (uint8*)&ConstCastSharedRef<FNiagaraEmitterCompiledData>(EmitterCompiledData[EmitterIndex]).Get(), NiagaraEmitterCompiledDataStruct, nullptr);
 		}
 	}
