Remove uneccessary checkSlow() call and markup to checkfSlow()

diff --git a/Engine/Plugins/FX/Niagara/Source/NiagaraEditor/Private/Toolkits/NiagaraSystemToolkit.cpp b/Engine/Plugins/FX/Niagara/Source/NiagaraEditor/Private/Toolkits/NiagaraSystemToolkit.cpp
--- a/Engine/Plugins/FX/Niagara/Source/NiagaraEditor/Private/Toolkits/NiagaraSystemToolkit.cpp
+++ b/Engine/Plugins/FX/Niagara/Source/NiagaraEditor/Private/Toolkits/NiagaraSystemToolkit.cpp
@@ -1019,9 +1019,8 @@ void FNiagaraSystemToolkit::UpdateOriginalEmitter()
 		EditableEmitter->GetScripts(AllScripts, true);
 		for (UNiagaraScript* Script : AllScripts)
 		{
-			checkSlow(Script->AreScriptAndSourceSynchronized());
+			checkfSlow(Script->AreScriptAndSourceSynchronized(), TEXT("Editable Emitter Script change ID is out of date when applying to Original Emitter!"));
 		}
-		checkSlow(EditableEmitter->AreAllScriptAndSourcesSynchronized());
 
 		// overwrite the original script in place by constructing a new one with the same name
 		Emitter = (UNiagaraEmitter*)StaticDuplicateObject(EditableEmitter, Emitter->GetOuter(),
