- Integrate CL 9553034 & 9600626 from UE4/Main
- Only reset the preview component when the timeline wraps, rather than all components in the scene
- This mirrors Cascade behavior

#rnx

diff --git a/Engine/Plugins/FX/Niagara/Source/NiagaraEditor/Private/Toolkits/NiagaraSystemToolkit.cpp b/Engine/Plugins/FX/Niagara/Source/NiagaraEditor/Private/Toolkits/NiagaraSystemToolkit.cpp
--- a/Engine/Plugins/FX/Niagara/Source/NiagaraEditor/Private/Toolkits/NiagaraSystemToolkit.cpp
+++ b/Engine/Plugins/FX/Niagara/Source/NiagaraEditor/Private/Toolkits/NiagaraSystemToolkit.cpp
@@ -744,7 +744,7 @@ void FNiagaraSystemToolkit::OnThumbnailCaptured(UTexture2D* Thumbnail)
 
 void FNiagaraSystemToolkit::ResetSimulation()
 {
-	SystemViewModel->ResetSystem(FNiagaraSystemViewModel::ETimeResetMode::AllowResetTime, FNiagaraSystemViewModel::EMultiResetMode::AllowResetAllInstances, FNiagaraSystemViewModel::EReinitMode::ResetSystem);
+	SystemViewModel->ResetSystem(FNiagaraSystemViewModel::ETimeResetMode::AllowResetTime, FNiagaraSystemViewModel::EMultiResetMode::AllowResetAllInstances, FNiagaraSystemViewModel::EReinitMode::ReinitializeSystem);
 }
 
 void FNiagaraSystemToolkit::ExtendToolbar()
diff --git a/Engine/Plugins/FX/Niagara/Source/NiagaraEditor/Private/ViewModels/NiagaraSystemViewModel.cpp b/Engine/Plugins/FX/Niagara/Source/NiagaraEditor/Private/ViewModels/NiagaraSystemViewModel.cpp
--- a/Engine/Plugins/FX/Niagara/Source/NiagaraEditor/Private/ViewModels/NiagaraSystemViewModel.cpp
+++ b/Engine/Plugins/FX/Niagara/Source/NiagaraEditor/Private/ViewModels/NiagaraSystemViewModel.cpp
@@ -513,7 +513,7 @@ void FNiagaraSystemViewModel::Tick(float DeltaTime)
 
 			if (bResetRequestPending)
 			{
-				ResetSystem(ETimeResetMode::AllowResetTime, EMultiResetMode::ResetThisInstance, EReinitMode::ReinitializeSystem);
+				ResetSystem(ETimeResetMode::AllowResetTime, EMultiResetMode::AllowResetAllInstances, EReinitMode::ReinitializeSystem);
 			}
 		}
 	}
@@ -999,7 +999,19 @@ void FNiagaraSystemViewModel::ResetSystem(ETimeResetMode TimeResetMode, EMultiRe
 		}
 	}
 
-	TArray<UNiagaraComponent*> ReferencingComponents = FNiagaraEditorUtilities::GetComponentsThatReferenceSystem(GetSystem());
+	TArray<UNiagaraComponent*> ReferencingComponents;
+	if (MultiResetMode == EMultiResetMode::ResetThisInstance)
+	{
+		if (PreviewComponent != nullptr)
+		{
+			ReferencingComponents.Add(PreviewComponent);
+		}
+	}
+	else
+	{
+		ReferencingComponents = FNiagaraEditorUtilities::GetComponentsThatReferenceSystem(GetSystem());
+	}
+
 	for (auto Component : ReferencingComponents)
 	{
 		if (ReinitMode == EReinitMode::ResetSystem)
