Fixing a bug where multiple Naigara particles could have the same value for Particles.UniqueID, thereby generating identical deterministic random numbers.

Fixes Particles.UniqueID for Niagara Events and multiple SpawnInfo counts when multiple spawn modules are present.
Adds a new autotest to verify deterministic random numbers for events.
Updates the screenshots of the ActivateDeactivate and FastPath autotests.

diff --git a/Engine/Plugins/FX/Niagara/Source/Niagara/Classes/NiagaraEmitterInstance.h b/Engine/Plugins/FX/Niagara/Source/Niagara/Classes/NiagaraEmitterInstance.h
--- a/Engine/Plugins/FX/Niagara/Source/Niagara/Classes/NiagaraEmitterInstance.h
+++ b/Engine/Plugins/FX/Niagara/Source/Niagara/Classes/NiagaraEmitterInstance.h
@@ -190,15 +190,14 @@ private:
 	FNiagaraParameterDirectBinding<int32> UpdateExecCountBinding;
 	TArray<FNiagaraParameterDirectBinding<int32>> EventExecCountBindings;
 
-	/*
 	FNiagaraParameterDirectBinding<int32> SpawnTotalSpawnedParticlesBinding;
 	FNiagaraParameterDirectBinding<int32> UpdateTotalSpawnedParticlesBinding;
-	FNiagaraParameterDirectBinding<int32> GPUTotalSpawnedParticlesBinding;
-	*/
+	TArray<FNiagaraParameterDirectBinding<int32>> EventTotalSpawnedParticlesBindings;
 
 	FNiagaraParameterDirectBinding<int32> SpawnRandomSeedBinding;
 	FNiagaraParameterDirectBinding<int32> UpdateRandomSeedBinding;
 	FNiagaraParameterDirectBinding<int32> GPURandomSeedBinding;
+	TArray<FNiagaraParameterDirectBinding<int>> EventRandomSeedBindings;
 
 	/*
 	FNiagaraParameterDirectBinding<int32> SpawnDeterminismBinding;
diff --git a/Engine/Plugins/FX/Niagara/Source/Niagara/Classes/NiagaraSystem.h b/Engine/Plugins/FX/Niagara/Source/Niagara/Classes/NiagaraSystem.h
--- a/Engine/Plugins/FX/Niagara/Source/Niagara/Classes/NiagaraSystem.h
+++ b/Engine/Plugins/FX/Niagara/Source/Niagara/Classes/NiagaraSystem.h
@@ -46,6 +46,9 @@ struct FNiagaraEmitterCompiledData
 	UPROPERTY()
 	FNiagaraVariable EmitterRandomSeedVar;
 
+	UPROPERTY()
+	FNiagaraVariable EmitterTotalSpawnedParticlesVar;
+
 	/** Per-Emitter DataSet Data. */
 	UPROPERTY()
 	FNiagaraDataSetCompiledData DataSetCompiledData;
diff --git a/Engine/Plugins/FX/Niagara/Source/Niagara/Private/NiagaraEmitterInstance.cpp b/Engine/Plugins/FX/Niagara/Source/Niagara/Private/NiagaraEmitterInstance.cpp
--- a/Engine/Plugins/FX/Niagara/Source/Niagara/Private/NiagaraEmitterInstance.cpp
+++ b/Engine/Plugins/FX/Niagara/Source/Niagara/Private/NiagaraEmitterInstance.cpp
@@ -388,6 +388,11 @@ void FNiagaraEmitterInstance::Init(int32 InEmitterIdx, FNiagaraSystemInstanceID
 
 		SpawnRandomSeedBinding.Init(SpawnExecContext.Parameters, CachedEmitterCompiledData->EmitterRandomSeedVar);
 		UpdateRandomSeedBinding.Init(UpdateExecContext.Parameters, CachedEmitterCompiledData->EmitterRandomSeedVar);
+		EventRandomSeedBindings.SetNum(NumEvents);
+		for (int32 i = 0; i < NumEvents; i++)
+		{
+			EventRandomSeedBindings[i].Init(EventExecContexts[i].Parameters, CachedEmitterCompiledData->EmitterRandomSeedVar);
+		}
 		if (CachedEmitter->SimTarget == ENiagaraSimTarget::GPUComputeSim && GPUExecContext != nullptr)
 		{
 			GPURandomSeedBinding.Init(GPUExecContext->CombinedParamStore, CachedEmitterCompiledData->EmitterRandomSeedVar);
@@ -401,6 +406,14 @@ void FNiagaraEmitterInstance::Init(int32 InEmitterIdx, FNiagaraSystemInstanceID
 		{
 			EventExecCountBindings[i].Init(EventExecContexts[i].Parameters, SYS_PARAM_ENGINE_EXEC_COUNT);
 		}
+
+		SpawnTotalSpawnedParticlesBinding.Init(SpawnExecContext.Parameters, CachedEmitterCompiledData->EmitterTotalSpawnedParticlesVar);
+		UpdateTotalSpawnedParticlesBinding.Init(UpdateExecContext.Parameters, CachedEmitterCompiledData->EmitterTotalSpawnedParticlesVar);
+		EventTotalSpawnedParticlesBindings.SetNum(NumEvents);
+		for (int32 i = 0; i < NumEvents; i++)
+		{
+			EventTotalSpawnedParticlesBindings[i].Init(EventExecContexts[i].Parameters, CachedEmitterCompiledData->EmitterTotalSpawnedParticlesVar);
+		}
 	}
 
 	{
@@ -1125,12 +1138,23 @@ void FNiagaraEmitterInstance::Tick(float DeltaSeconds)
 		
 		SpawnRandomSeedBinding.SetValue(CachedEmitter->RandomSeed);
 		UpdateRandomSeedBinding.SetValue(CachedEmitter->RandomSeed);
+		for (FNiagaraParameterDirectBinding<int32>& Binding : EventRandomSeedBindings)
+		{
+			Binding.SetValue(CachedEmitter->RandomSeed);
+		}
 
 		if (CachedEmitter->SimTarget == ENiagaraSimTarget::GPUComputeSim && GPUExecContext != nullptr)
 		{
 			EmitterAgeBindingGPU.SetValue(Age);
 			GPURandomSeedBinding.SetValue(CachedEmitter->RandomSeed);
 		}
+
+		SpawnTotalSpawnedParticlesBinding.SetValue(TotalSpawnedParticles);
+		UpdateTotalSpawnedParticlesBinding.SetValue(TotalSpawnedParticles);
+		for (FNiagaraParameterDirectBinding<int32>& Binding : EventTotalSpawnedParticlesBindings)
+		{
+			Binding.SetValue(TotalSpawnedParticles);
+		}
 	}
 	
 	// Calculate number of new particles from regular spawning 
@@ -1356,6 +1380,9 @@ void FNiagaraEmitterInstance::Tick(float DeltaSeconds)
 			int32 OrigNum = Data.GetDestinationDataChecked().GetNumInstances();
 			Data.GetDestinationDataChecked().SetNumInstances(OrigNum + Num);
 
+			// We need to update Engine.Emitter.TotalSpawnedParticles for each event spawn invocation.
+			SpawnTotalSpawnedParticlesBinding.SetValue(TotalSpawnedParticles); 
+			
 			// NOTE(mv): Updates the count after setting the variable, such that the TotalSpawnedParticles value read 
 			//           in the script has the count at the start of the frame. 
 			//           This way UniqueID = TotalSpawnedParticles + ExecIndex provide unique and sequential identifiers. 
diff --git a/Engine/Plugins/FX/Niagara/Source/Niagara/Private/NiagaraSystem.cpp b/Engine/Plugins/FX/Niagara/Source/Niagara/Private/NiagaraSystem.cpp
--- a/Engine/Plugins/FX/Niagara/Source/Niagara/Private/NiagaraSystem.cpp
+++ b/Engine/Plugins/FX/Niagara/Source/Niagara/Private/NiagaraSystem.cpp
@@ -979,6 +979,7 @@ void UNiagaraSystem::InitEmitterVariableAliasNames(FNiagaraEmitterCompiledData&
 	EmitterCompiledDataToInit.EmitterAgeVar.SetName(GetEmitterVariableAliasName(SYS_PARAM_EMITTER_AGE, InAssociatedEmitter));
 	EmitterCompiledDataToInit.EmitterSpawnGroupVar.SetName(GetEmitterVariableAliasName(SYS_PARAM_EMITTER_SPAWN_GROUP, InAssociatedEmitter));
 	EmitterCompiledDataToInit.EmitterRandomSeedVar.SetName(GetEmitterVariableAliasName(SYS_PARAM_EMITTER_RANDOM_SEED, InAssociatedEmitter));
+	EmitterCompiledDataToInit.EmitterTotalSpawnedParticlesVar.SetName(GetEmitterVariableAliasName(SYS_PARAM_ENGINE_EMITTER_TOTAL_SPAWNED_PARTICLES, InAssociatedEmitter));
 }
 
 const FName UNiagaraSystem::GetEmitterVariableAliasName(const FNiagaraVariable& InEmitterVar, const UNiagaraEmitter* InEmitter) const
@@ -1310,4 +1311,5 @@ FNiagaraEmitterCompiledData::FNiagaraEmitterCompiledData()
 	EmitterAgeVar = SYS_PARAM_EMITTER_AGE;
 	EmitterSpawnGroupVar = SYS_PARAM_EMITTER_SPAWN_GROUP;
 	EmitterRandomSeedVar = SYS_PARAM_EMITTER_RANDOM_SEED;
+	EmitterTotalSpawnedParticlesVar = SYS_PARAM_ENGINE_EMITTER_TOTAL_SPAWNED_PARTICLES;
 }
