- Fix Niagara Mesh Particle Motion Blur not working correctly with local space systems

#rnx

diff --git a/Engine/Plugins/FX/Niagara/Shaders/Private/NiagaraMeshVertexFactory.ush b/Engine/Plugins/FX/Niagara/Shaders/Private/NiagaraMeshVertexFactory.ush
--- a/Engine/Plugins/FX/Niagara/Shaders/Private/NiagaraMeshVertexFactory.ush
+++ b/Engine/Plugins/FX/Niagara/Shaders/Private/NiagaraMeshVertexFactory.ush
@@ -192,6 +192,7 @@ struct FVertexFactoryIntermediates
 #endif
 
 	float4x4 LocalToWorld;
+    float4x4 PrevLocalToWorld;
 };
 
 float4x4 GetParticleTransform(FVertexFactoryIntermediates Inters)
@@ -199,9 +200,9 @@ float4x4 GetParticleTransform(FVertexFactoryIntermediates Inters)
 	return Inters.LocalToWorld;
 }
 
-#if USE_PREVIOUS_TRANSFORM
-float4x4 GetParticlePrevTransform(FVertexFactoryInput Input)
+float4x4 GetParticlePrevTransform(FVertexFactoryIntermediates Inters)
 {
+#if USE_PREVIOUS_TRANSFORM
 	uint TransformBaseIdx = Input.InstanceId * 3;
 
 	float4x4 PrevTransform4x4 = float4x4(
@@ -213,8 +214,10 @@ float4x4 GetParticlePrevTransform(FVertexFactoryInput Input)
 
 	//Transform is transposed on CPU and position is packed into W. Transpose again to get the correct transform back.
 	return transpose(PrevTransform4x4);
-}
+#else
+	return Inters.PrevLocalToWorld;
 #endif
+}
 
 /** Converts from vertex factory specific interpolants to a FMaterialPixelParameters, which is used by material inputs. */
 FMaterialPixelParameters GetMaterialPixelParameters(FVertexFactoryInterpolantsVSToPS Interpolants, float4 SvPosition)
@@ -379,11 +382,7 @@ FMaterialVertexParameters GetMaterialVertexParameters(FVertexFactoryInput Input,
 
 	Result.TangentToWorld = mul(TangentToLocal, (float3x3)Transform );
 
-#if USE_PREVIOUS_TRANSFORM
-	Result.PrevFrameLocalToWorld = GetParticlePrevTransform(Input);
-#else
-	Result.PrevFrameLocalToWorld = GetParticleTransform(Intermediates);
-#endif
+	Result.PrevFrameLocalToWorld = GetParticlePrevTransform(Intermediates);
 
 #if NUM_MATERIAL_TEXCOORDS_VERTEX
 	for(int CoordinateIndex = 0; CoordinateIndex < NUM_MATERIAL_TEXCOORDS_VERTEX; CoordinateIndex++)
@@ -404,28 +403,9 @@ float4 CalcWorldPosition(FVertexFactoryInput Input, FVertexFactoryIntermediates
 
 float4 CalcPrevWorldPosition(FVertexFactoryInput Input, FVertexFactoryIntermediates Intermediates)
 {
-#if USE_PREVIOUS_TRANSFORM
-	if (MeshParticleVF.PrevTransformAvailable > 0)
-	{
-		float4x4 PrevTransform = GetParticlePrevTransform(Input);
-		float3 PrevWorldPosition = mul(Input.Position, PrevTransform).xyz + ResolvedView.PrevPreViewTranslation;
-		return float4(PrevWorldPosition, Input.Position.w);
-	}
-	else
-#endif
-	// if we don't have a full previous transform, we can at least calculate the position to get some motion blur
-	{
-		float3 WorldPosition = CalcWorldPosition(Input, Intermediates).xyz;
-
-		// ParticleVelocity.xyz - normalized direction vector
-		// ParticleVelocity.w - velocity magnitude
-		float3 WorldDisplacement = -(Intermediates.ParticleVelocity.xyz * (Intermediates.ParticleVelocity.w * NiagaraMeshVF.DeltaSeconds));
-
-		float3 PrevWorldPosition = WorldPosition - ResolvedView.PreViewTranslation;
-		PrevWorldPosition += WorldDisplacement + ResolvedView.PrevPreViewTranslation;
-
-		return float4(PrevWorldPosition.xyz, Input.Position.w);
-	}
+    float4x4 Transform = GetParticlePrevTransform(Intermediates);
+    float3 WorldPosition = mul(Input.Position, Transform).xyz + ResolvedView.PrevPreViewTranslation;
+    return float4(WorldPosition, Input.Position.w);
 }
 
 // From UnrealMath::FindBetween
@@ -687,10 +667,21 @@ FVertexFactoryIntermediates GetVertexFactoryIntermediates(FVertexFactoryInput In
 	Intermediates.LocalToWorld = transpose(Intermediates.LocalToWorld);
 	if (NiagaraMeshVF.bLocalSpace)
 	{
+        Intermediates.PrevLocalToWorld = mul(Intermediates.LocalToWorld, Primitive.PreviousLocalToWorld);
 		Intermediates.LocalToWorld = mul(Intermediates.LocalToWorld, Primitive.LocalToWorld);
 		Intermediates.ParticlePosition = TransformLocalToWorld(Intermediates.ParticlePosition);
 	}
-	//Intermediates.ParticleVelocity.xyz = TransformVector(Intermediates.ParticleVelocity.xyz);
+    else
+    {
+        //-TODO: If we have per particle previous transform use them here
+        Intermediates.PrevLocalToWorld = Intermediates.LocalToWorld;
+    }
+
+    // If we don't have per particle transforms adjust the PrevLocalToWorld by velocity
+	if (MeshParticleVF.PrevTransformAvailable == 0)
+    {
+        Intermediates.PrevLocalToWorld[3].xyz -= (ParticleVelocity * NiagaraMeshVF.DeltaSeconds);
+    }
 
 #if NUM_MATERIAL_TEXCOORDS_VERTEX
 	for(int CoordinateIndex = 0; CoordinateIndex < NUM_MATERIAL_TEXCOORDS_VERTEX; CoordinateIndex++)
