Fixed shader compile errors due to reserved keyword usage

diff --git a/Engine/Plugins/FX/Niagara/Shaders/Private/NiagaraRibbonVertexFactory.ush b/Engine/Plugins/FX/Niagara/Shaders/Private/NiagaraRibbonVertexFactory.ush
--- a/Engine/Plugins/FX/Niagara/Shaders/Private/NiagaraRibbonVertexFactory.ush
+++ b/Engine/Plugins/FX/Niagara/Shaders/Private/NiagaraRibbonVertexFactory.ush
@@ -808,8 +808,8 @@ FVertexFactoryInput LoadVertexFactoryInputForHGS(uint TriangleIndex, int VertexI
 {
 	FVertexFactoryInput Input;
 
-	FTriangleBaseAttributes Triangle = LoadTriangleBaseAttributes(TriangleIndex);
-	uint VertexId = Triangle.Indices[VertexIndex];
+	FTriangleBaseAttributes Tri = LoadTriangleBaseAttributes(TriangleIndex);
+	uint VertexId = Tri.Indices[VertexIndex];
 	Input.InterpVtxId = VertexId;
 #if FEATURE_LEVEL >= FEATURE_LEVEL_SM4 && ONEPASS_POINTLIGHT_SHADOW && USING_VERTEX_SHADER_LAYER
 	Input.InstanceInd = 0;
