Niagara - Move drag/drop from the stack to the overview, also unify drag drop code paths and remove stack spacers.

diff --git a/Engine/Plugins/FX/Niagara/Source/NiagaraEditor/Private/NiagaraActions.cpp b/Engine/Plugins/FX/Niagara/Source/NiagaraEditor/Private/NiagaraActions.cpp
--- a/Engine/Plugins/FX/Niagara/Source/NiagaraEditor/Private/NiagaraActions.cpp
+++ b/Engine/Plugins/FX/Niagara/Source/NiagaraEditor/Private/NiagaraActions.cpp
@@ -204,87 +204,4 @@ EVisibility FNiagaraParameterGraphDragOperation::GetErrorIconVisible() const
 	return EVisibility::Collapsed;
 }
 
-/************************************************************************/
-/* FNiagaraStackDragOperation											*/
-/************************************************************************/
-TSharedRef<FNiagaraStackDragOperation> FNiagaraStackDragOperation::New(TSharedPtr<FEdGraphSchemaAction> InActionNode)
-{
-	TSharedRef<FNiagaraStackDragOperation> Operation = MakeShareable(new FNiagaraStackDragOperation);
-	Operation->SourceAction = InActionNode;
-	Operation->Construct();
-	return Operation;
-}
-
-void FNiagaraStackDragOperation::OnDrop(bool bDropWasHandled, const FPointerEvent& MouseEvent)
-{
-	FDragDropOperation::OnDrop(bDropWasHandled, MouseEvent);
-}
-
-void FNiagaraStackDragOperation::OnDragged(const class FDragDropEvent& DragDropEvent)
-{
-	if (SourceAction.IsValid())
-	{
-		SetSimpleFeedbackMessage(SourceAction->GetMenuDescription());
-	}
-
-	FDragDropOperation::OnDragged(DragDropEvent);
-}
-
-void FNiagaraStackDragOperation::Construct()
-{
-	// Create the drag-drop decorator window
-	CursorDecoratorWindow = SWindow::MakeCursorDecorator();
-	const bool bShowImmediately = false;
-	FSlateApplication::Get().AddWindow(CursorDecoratorWindow.ToSharedRef(), bShowImmediately);
-}
-
-FNiagaraStackDragOperation::FNiagaraStackDragOperation()
-	: bControlDrag(false)
-	, bAltDrag(false)
-{
-
-}
-
-bool FNiagaraStackDragOperation::HasFeedbackMessage()
-{
-	return CursorDecoratorWindow->GetContent() != SNullWidget::NullWidget;
-}
-
-void FNiagaraStackDragOperation::SetFeedbackMessage(const TSharedPtr<SWidget>& Message)
-{
-	if (Message.IsValid())
-	{
-		CursorDecoratorWindow->ShowWindow();
-		CursorDecoratorWindow->SetContent
-		(
-			SNew(SBorder)
-			.BorderImage(FEditorStyle::GetBrush("Graph.ConnectorFeedback.Border"))
-			[
-				Message.ToSharedRef()
-			]
-		);
-	}
-	else
-	{
-		CursorDecoratorWindow->HideWindow();
-		CursorDecoratorWindow->SetContent(SNullWidget::NullWidget);
-	}
-}
-
-void FNiagaraStackDragOperation::SetSimpleFeedbackMessage(const FText& Message)
-{
-	SetFeedbackMessage(
-		SNew(SHorizontalBox)
-		+ SHorizontalBox::Slot()
-		.AutoWidth()
-		.MaxWidth(500)
-		.Padding(3.0f)
-		.VAlign(VAlign_Center)
-		[
-			SNew(STextBlock)
-			.Text(Message)
-		]
-	);
-}
-
 #undef LOCTEXT_NAMESPACE
diff --git a/Engine/Plugins/FX/Niagara/Source/NiagaraEditor/Private/ViewModels/NiagaraSystemSelectionViewModel.cpp b/Engine/Plugins/FX/Niagara/Source/NiagaraEditor/Private/ViewModels/NiagaraSystemSelectionViewModel.cpp
--- a/Engine/Plugins/FX/Niagara/Source/NiagaraEditor/Private/ViewModels/NiagaraSystemSelectionViewModel.cpp
+++ b/Engine/Plugins/FX/Niagara/Source/NiagaraEditor/Private/ViewModels/NiagaraSystemSelectionViewModel.cpp
@@ -2,6 +2,8 @@
 
 #include "ViewModels/NiagaraSystemSelectionViewModel.h"
 #include "ViewModels/Stack/NiagaraStackEntry.h"
+#include "ViewModels/Stack/NiagaraStackItemGroup.h"
+#include "ViewModels/Stack/NiagaraStackItem.h"
 #include "ViewModels/Stack/NiagaraStackSelection.h"
 #include "ViewModels/Stack/NiagaraStackViewModel.h"
 #include "ViewModels/NiagaraSystemViewModel.h"
@@ -207,6 +209,47 @@ UNiagaraStackViewModel* UNiagaraSystemSelectionViewModel::GetSelectionStackViewM
 	return SelectionStackViewModel;
 }
 
+void UNiagaraSystemSelectionViewModel::RemoveEntryFromSelectionByDisplayedObject(const UObject* InObject)
+{
+	bool bSelectionChanged = false;
+
+	TArray<UNiagaraStackEntry*> EntriesToDeselect;
+	for (UNiagaraStackEntry* SelectedEntry : SelectedEntries)
+	{
+		if (SelectedEntry->GetDisplayedObject() == InObject)
+		{
+			EntriesToDeselect.Add(SelectedEntry);
+			bSelectionChanged = true;
+		}
+	}
+
+	for (UNiagaraStackEntry* EntryToDeselect : EntriesToDeselect)
+	{
+		RemoveEntryFromSelectionInternal(EntryToDeselect);
+	}
+
+	if (bSelectionChanged)
+	{
+		OnSelectionChangedDelegate.Broadcast(ESelectionChangeSource::EntrySelection);
+		UpdateStackSelectionEntry();
+	}
+}
+
+void UNiagaraSystemSelectionViewModel::AddEntriesToSelectionByDisplayedObjectsDeferred(const TArray<const UObject*>& InObjects)
+{
+	for (const UObject* InObject : InObjects)
+	{
+		DeferredDisplayedObjectKeysToAddToSelection.Add(FObjectKey(InObject));
+	}
+}
+
+void UNiagaraSystemSelectionViewModel::AddEntryToSelectionByDisplayedObjectDeferred(const UObject* InObject)
+{
+	TArray<const UObject*> InObjects;
+	InObjects.Add(InObject);
+	AddEntriesToSelectionByDisplayedObjectsDeferred(InObjects);
+}
+
 void UNiagaraSystemSelectionViewModel::Refresh()
 {
 	TSet<FGuid> ValidEmitterHandleIds;
@@ -233,6 +276,64 @@ UNiagaraSystemSelectionViewModel::FOnSelectionChanged& UNiagaraSystemSelectionVi
 	return OnSelectionChangedDelegate;
 }
 
+void FindStackGroupsAndItemsForDisplayedObjectKeysRecursive(UNiagaraStackEntry* StackEntry, const TArray<FObjectKey>& ObjectKeys, TArray<UNiagaraStackEntry*>& OutFoundStackEntries)
+{
+	if (StackEntry->IsA<UNiagaraStackItemGroup>() || StackEntry->IsA<UNiagaraStackItem>())
+	{
+		UObject* DisplayedObject = StackEntry->GetDisplayedObject();
+		if (DisplayedObject != nullptr && ObjectKeys.Contains(FObjectKey(DisplayedObject)))
+		{
+			OutFoundStackEntries.Add(StackEntry);
+		}
+	}
+
+	if (StackEntry->IsA<UNiagaraStackItem>() == false)
+	{
+		TArray<UNiagaraStackEntry*> Children;
+		StackEntry->GetUnfilteredChildren(Children);
+		for (UNiagaraStackEntry* ChildEntry : Children)
+		{
+			FindStackGroupsAndItemsForDisplayedObjectKeysRecursive(ChildEntry, ObjectKeys, OutFoundStackEntries);
+		}
+	}
+}
+
+void FindStackGroupsAndItemsForDisplayedObjectKeys(UNiagaraStackViewModel* StackViewModel, const TArray<FObjectKey>& ObjectKeys, TArray<UNiagaraStackEntry*>& OutFoundStackEntries)
+{
+	for (UNiagaraStackEntry* RootEntry : StackViewModel->GetRootEntries())
+	{
+		FindStackGroupsAndItemsForDisplayedObjectKeysRecursive(RootEntry, ObjectKeys, OutFoundStackEntries);
+	}
+}
+
+void UNiagaraSystemSelectionViewModel::Tick()
+{
+	TArray<UNiagaraStackEntry*> FoundStackEntries;
+	TSharedRef<FNiagaraSystemViewModel> SystemViewModel = GetSystemViewModel();
+	FindStackGroupsAndItemsForDisplayedObjectKeys(SystemViewModel->GetSystemStackViewModel(), DeferredDisplayedObjectKeysToAddToSelection, FoundStackEntries);
+	for (TSharedRef<FNiagaraEmitterHandleViewModel> EmitterViewModel : SystemViewModel->GetEmitterHandleViewModels())
+	{
+		FindStackGroupsAndItemsForDisplayedObjectKeys(EmitterViewModel->GetEmitterStackViewModel(), DeferredDisplayedObjectKeysToAddToSelection, FoundStackEntries);
+	}
+
+	int32 AddedEntries = 0;
+	for (UNiagaraStackEntry* FoundStackEntry : FoundStackEntries)
+	{
+		if (SelectedEntries.Contains(FoundStackEntry) == false)
+		{
+			SelectedEntries.Add(FoundStackEntry);
+			AddedEntries++;
+		}
+	}
+
+	DeferredDisplayedObjectKeysToAddToSelection.Empty();
+	if (AddedEntries > 0)
+	{
+		OnSelectionChangedDelegate.Broadcast(ESelectionChangeSource::EntrySelection);
+		UpdateStackSelectionEntry();
+	}
+}
+
 TSharedRef<FNiagaraSystemViewModel> UNiagaraSystemSelectionViewModel::GetSystemViewModel()
 {
 	TSharedPtr<FNiagaraSystemViewModel> SystemViewModel = SystemViewModelWeak.Pin();
@@ -240,6 +341,38 @@ TSharedRef<FNiagaraSystemViewModel> UNiagaraSystemSelectionViewModel::GetSystemV
 	return SystemViewModel.ToSharedRef();
 }
 
+void UNiagaraSystemSelectionViewModel::RemoveEntryFromSelectionInternal(UNiagaraStackEntry* DeselectedEntry)
+{
+	SelectedEntries.Remove(DeselectedEntry);
+	TSharedPtr<FNiagaraEmitterViewModel> DeselectedEmitterViewModel = DeselectedEntry->GetEmitterViewModel();
+	if (DeselectedEmitterViewModel.IsValid())
+	{
+		bool bEmitterIsStillSelected = SelectedEntries.ContainsByPredicate([DeselectedEmitterViewModel](const UNiagaraStackEntry* SelectedEntry)
+		{
+			return SelectedEntry->GetEmitterViewModel() == DeselectedEmitterViewModel;
+		});
+		if (bEmitterIsStillSelected == false)
+		{
+			const FNiagaraEmitterHandle* DeselectedEmitterHandle = FNiagaraEditorUtilities::GetEmitterHandleForEmitter(DeselectedEntry->GetSystemViewModel()->GetSystem(), *DeselectedEmitterViewModel->GetEmitter());
+			if (DeselectedEmitterHandle != nullptr)
+			{
+				SelectedEmitterHandleIds.Remove(DeselectedEmitterHandle->GetId());
+			}
+		}
+	}
+	else
+	{
+		bool bSystemIsStillSelected = SelectedEntries.ContainsByPredicate([](const UNiagaraStackEntry* SelectedEntry)
+		{
+			return SelectedEntry->GetEmitterViewModel().IsValid() == false;
+		});
+		if (bSystemIsStillSelected == false)
+		{
+			bSystemIsSelected = false;
+		}
+	}
+}
+
 void UNiagaraSystemSelectionViewModel::UpdateStackSelectionEntry()
 {
 	// Check the new selection to make sure that there are no entries which are owned by other entries and if there are
diff --git a/Engine/Plugins/FX/Niagara/Source/NiagaraEditor/Private/ViewModels/NiagaraSystemViewModel.cpp b/Engine/Plugins/FX/Niagara/Source/NiagaraEditor/Private/ViewModels/NiagaraSystemViewModel.cpp
--- a/Engine/Plugins/FX/Niagara/Source/NiagaraEditor/Private/ViewModels/NiagaraSystemViewModel.cpp
+++ b/Engine/Plugins/FX/Niagara/Source/NiagaraEditor/Private/ViewModels/NiagaraSystemViewModel.cpp
@@ -523,6 +523,21 @@ void FNiagaraSystemViewModel::Tick(float DeltaTime)
 		UpdateSequencerTracksForEmitters(EmitterIdsRequiringSequencerTrackUpdate);
 		EmitterIdsRequiringSequencerTrackUpdate.Empty();
 	}
+
+	if (SystemStackViewModel != nullptr)
+	{
+		SystemStackViewModel->Tick();
+	}
+
+	for (TSharedRef<FNiagaraEmitterHandleViewModel> EmitterHandleViewModel : EmitterHandleViewModels)
+	{
+		EmitterHandleViewModel->GetEmitterStackViewModel()->Tick();
+	}
+
+	if (SelectionViewModel != nullptr)
+	{
+		SelectionViewModel->Tick();
+	}
 }
 
 void FNiagaraSystemViewModel::OnPreSave()
diff --git a/Engine/Plugins/FX/Niagara/Source/NiagaraEditor/Private/ViewModels/Stack/NiagaraStackAdvancedExpander.cpp b/Engine/Plugins/FX/Niagara/Source/NiagaraEditor/Private/ViewModels/Stack/NiagaraStackAdvancedExpander.cpp
deleted file mode 100644
diff --git a/Engine/Plugins/FX/Niagara/Source/NiagaraEditor/Private/ViewModels/Stack/NiagaraStackEmitterSettingsGroup.cpp b/Engine/Plugins/FX/Niagara/Source/NiagaraEditor/Private/ViewModels/Stack/NiagaraStackEmitterSettingsGroup.cpp
--- a/Engine/Plugins/FX/Niagara/Source/NiagaraEditor/Private/ViewModels/Stack/NiagaraStackEmitterSettingsGroup.cpp
+++ b/Engine/Plugins/FX/Niagara/Source/NiagaraEditor/Private/ViewModels/Stack/NiagaraStackEmitterSettingsGroup.cpp
@@ -2,7 +2,6 @@
 
 #include "ViewModels/Stack/NiagaraStackEmitterSettingsGroup.h"
 #include "ViewModels/Stack/NiagaraStackObject.h"
-#include "ViewModels/Stack/NiagaraStackSpacer.h"
 #include "ViewModels/NiagaraSystemViewModel.h"
 #include "ViewModels/NiagaraEmitterViewModel.h"
 #include "NiagaraEmitter.h"
diff --git a/Engine/Plugins/FX/Niagara/Source/NiagaraEditor/Private/ViewModels/Stack/NiagaraStackEntry.cpp b/Engine/Plugins/FX/Niagara/Source/NiagaraEditor/Private/ViewModels/Stack/NiagaraStackEntry.cpp
--- a/Engine/Plugins/FX/Niagara/Source/NiagaraEditor/Private/ViewModels/Stack/NiagaraStackEntry.cpp
+++ b/Engine/Plugins/FX/Niagara/Source/NiagaraEditor/Private/ViewModels/Stack/NiagaraStackEntry.cpp
@@ -177,6 +177,11 @@ FText UNiagaraStackEntry::GetDisplayName() const
 	return FText();
 }
 
+UObject* UNiagaraStackEntry::GetDisplayedObject() const
+{
+	return nullptr;
+}
+
 UNiagaraStackEditorData& UNiagaraStackEntry::GetStackEditorData() const
 {
 	return *StackEditorData;
@@ -349,33 +354,33 @@ bool UNiagaraStackEntry::CanDrag() const
 	return false;
 }
 
-TOptional<UNiagaraStackEntry::FDropResult> UNiagaraStackEntry::CanDrop(const TArray<UNiagaraStackEntry*>& DraggedEntries, EDragOptions DragOptions)
+TOptional<UNiagaraStackEntry::FDropRequestResponse> UNiagaraStackEntry::CanDrop(const FDropRequest& DropRequest)
 {
-	TOptional<FDropResult> Result = CanDropInternal(DraggedEntries, DragOptions);
-	if (Result.IsSet())
+	TOptional<FDropRequestResponse> Response = CanDropInternal(DropRequest);
+	if (Response.IsSet())
 	{
-		return Result;
+		return Response;
 	}
 	else
 	{
 		return OnRequestCanDropDelegate.IsBound() 
-			? OnRequestCanDropDelegate.Execute(*this, DraggedEntries, DragOptions)
-			: TOptional<FDropResult>();
+			? OnRequestCanDropDelegate.Execute(*this, DropRequest)
+			: TOptional<FDropRequestResponse>();
 	}
 }
 
-TOptional<UNiagaraStackEntry::FDropResult> UNiagaraStackEntry::Drop(const TArray<UNiagaraStackEntry*>& DraggedEntries, EDragOptions DragOptions)
+TOptional<UNiagaraStackEntry::FDropRequestResponse> UNiagaraStackEntry::Drop(const FDropRequest& DropRequest)
 {
-	TOptional<FDropResult> Result = DropInternal(DraggedEntries, DragOptions);
-	if (Result.IsSet())
+	TOptional<FDropRequestResponse> Response = DropInternal(DropRequest);
+	if (Response.IsSet())
 	{
-		return Result;
+		return Response;
 	}
 	else
 	{
 		return OnRequestDropDelegate.IsBound()
-			? OnRequestDropDelegate.Execute(*this, DraggedEntries, DragOptions)
-			: TOptional<FDropResult>();
+			? OnRequestDropDelegate.Execute(*this, DropRequest)
+			: TOptional<FDropRequestResponse>();
 	}
 }
 
@@ -409,24 +414,24 @@ bool UNiagaraStackEntry::HasBaseEmitter() const
 	return bHasBaseEmitterCache.GetValue();
 }
 
-TOptional<UNiagaraStackEntry::FDropResult> UNiagaraStackEntry::CanDropInternal(const TArray<UNiagaraStackEntry*>& DraggedEntries, EDragOptions DragOptions)
+TOptional<UNiagaraStackEntry::FDropRequestResponse> UNiagaraStackEntry::CanDropInternal(const FDropRequest& DropRequest)
 {
-	return TOptional<FDropResult>();
+	return TOptional<FDropRequestResponse>();
 }
 
-TOptional<UNiagaraStackEntry::FDropResult> UNiagaraStackEntry::DropInternal(const TArray<UNiagaraStackEntry*>& DraggedEntries, EDragOptions DragOptions)
+TOptional<UNiagaraStackEntry::FDropRequestResponse> UNiagaraStackEntry::DropInternal(const FDropRequest& DropRequest)
 {
-	return TOptional<FDropResult>();
+	return TOptional<FDropRequestResponse>();
 }
 
-TOptional<UNiagaraStackEntry::FDropResult> UNiagaraStackEntry::ChildRequestCanDropInternal(const UNiagaraStackEntry& TargetChild, const TArray<UNiagaraStackEntry*>& DraggedEntries, EDragOptions DragOptions)
+TOptional<UNiagaraStackEntry::FDropRequestResponse> UNiagaraStackEntry::ChildRequestCanDropInternal(const UNiagaraStackEntry& TargetChild, const FDropRequest& DropRequest)
 {
-	return TOptional<FDropResult>();
+	return TOptional<FDropRequestResponse>();
 }
 
-TOptional<UNiagaraStackEntry::FDropResult> UNiagaraStackEntry::ChildRequestDropInternal(const UNiagaraStackEntry& TargetChild, const TArray<UNiagaraStackEntry*>& DraggedEntries, EDragOptions DragOptions)
+TOptional<UNiagaraStackEntry::FDropRequestResponse> UNiagaraStackEntry::ChildRequestDropInternal(const UNiagaraStackEntry& TargetChild, const FDropRequest& DropRequest)
 {
-	return TOptional<FDropResult>();
+	return TOptional<FDropRequestResponse>();
 }
 
 void UNiagaraStackEntry::ChlildStructureChangedInternal()
@@ -449,8 +454,12 @@ void UNiagaraStackEntry::RefreshChildren()
 		Child->OnDataObjectModified().RemoveAll(this);
 		Child->OnRequestFullRefresh().RemoveAll(this);
 		Child->OnRequestFullRefreshDeferred().RemoveAll(this);
-		Child->SetOnRequestCanDrop(FOnRequestDrop());
-		Child->SetOnRequestDrop(FOnRequestDrop());
+		UNiagaraStackEntry* OuterOwner = Cast<UNiagaraStackEntry>(Child->GetOuter());
+		if (OuterOwner == this)
+		{
+			Child->SetOnRequestCanDrop(FOnRequestDrop());
+			Child->SetOnRequestDrop(FOnRequestDrop());
+		}
 	}
 	for (UNiagaraStackErrorItem* ErrorChild : ErrorChildren)
 	{
@@ -490,8 +499,11 @@ void UNiagaraStackEntry::RefreshChildren()
 		Child->OnDataObjectModified().AddUObject(this, &UNiagaraStackEntry::ChildDataObjectModified);
 		Child->OnRequestFullRefresh().AddUObject(this, &UNiagaraStackEntry::ChildRequestFullRefresh);
 		Child->OnRequestFullRefreshDeferred().AddUObject(this, &UNiagaraStackEntry::ChildRequestFullRefreshDeferred);
-		Child->SetOnRequestCanDrop(FOnRequestDrop::CreateUObject(this, &UNiagaraStackEntry::ChildRequestCanDrop));
-		Child->SetOnRequestDrop(FOnRequestDrop::CreateUObject(this, &UNiagaraStackEntry::ChildRequestDrop));
+		if (OuterOwner == this)
+		{
+			Child->SetOnRequestCanDrop(FOnRequestDrop::CreateUObject(this, &UNiagaraStackEntry::ChildRequestCanDrop));
+			Child->SetOnRequestDrop(FOnRequestDrop::CreateUObject(this, &UNiagaraStackEntry::ChildRequestDrop));
+		}
 	}
 	
 	// Stack issues refresh
@@ -606,32 +618,32 @@ void UNiagaraStackEntry::ChildRequestFullRefreshDeferred()
 	RequestFullRefreshDeferredDelegate.Broadcast();
 }
 
-TOptional<UNiagaraStackEntry::FDropResult> UNiagaraStackEntry::ChildRequestCanDrop(const UNiagaraStackEntry& TargetChild, const TArray<UNiagaraStackEntry*>& DraggedEntries, EDragOptions DragOptions)
+TOptional<UNiagaraStackEntry::FDropRequestResponse> UNiagaraStackEntry::ChildRequestCanDrop(const UNiagaraStackEntry& TargetChild, const FDropRequest& DropRequest)
 {
-	TOptional<FDropResult> Result = ChildRequestCanDropInternal(TargetChild, DraggedEntries, DragOptions);
-	if (Result.IsSet())
+	TOptional<FDropRequestResponse> Response = ChildRequestCanDropInternal(TargetChild, DropRequest);
+	if (Response.IsSet())
 	{
-		return Result;
+		return Response;
 	}
 	else
 	{
 		return OnRequestCanDropDelegate.IsBound()
-			? OnRequestCanDropDelegate.Execute(TargetChild, DraggedEntries, DragOptions)
-			: TOptional<FDropResult>();
+			? OnRequestCanDropDelegate.Execute(TargetChild, DropRequest)
+			: TOptional<FDropRequestResponse>();
 	}
 }
 
-TOptional<UNiagaraStackEntry::FDropResult> UNiagaraStackEntry::ChildRequestDrop(const UNiagaraStackEntry& TargetChild, const TArray<UNiagaraStackEntry*>& DraggedEntries, EDragOptions DragOptions)
+TOptional<UNiagaraStackEntry::FDropRequestResponse> UNiagaraStackEntry::ChildRequestDrop(const UNiagaraStackEntry& TargetChild, const FDropRequest& DropRequest)
 {
-	TOptional<FDropResult> Result = ChildRequestDropInternal(TargetChild, DraggedEntries, DragOptions);
-	if (Result.IsSet())
+	TOptional<FDropRequestResponse> Response = ChildRequestDropInternal(TargetChild, DropRequest);
+	if (Response.IsSet())
 	{
-		return Result;
+		return Response;
 	}
 	else
 	{
 		return OnRequestDropDelegate.IsBound()
-			? OnRequestDropDelegate.Execute(TargetChild, DraggedEntries, DragOptions)
-			: TOptional<FDropResult>();
+			? OnRequestDropDelegate.Execute(TargetChild, DropRequest)
+			: TOptional<FDropRequestResponse>();
 	}
 }
diff --git a/Engine/Plugins/FX/Niagara/Source/NiagaraEditor/Private/ViewModels/Stack/NiagaraStackEventHandlerGroup.cpp b/Engine/Plugins/FX/Niagara/Source/NiagaraEditor/Private/ViewModels/Stack/NiagaraStackEventHandlerGroup.cpp
--- a/Engine/Plugins/FX/Niagara/Source/NiagaraEditor/Private/ViewModels/Stack/NiagaraStackEventHandlerGroup.cpp
+++ b/Engine/Plugins/FX/Niagara/Source/NiagaraEditor/Private/ViewModels/Stack/NiagaraStackEventHandlerGroup.cpp
@@ -9,7 +9,6 @@
 #include "NiagaraEmitter.h"
 #include "NiagaraGraph.h"
 #include "NiagaraScriptSource.h"
-#include "ViewModels/Stack/NiagaraStackSpacer.h"
 #include "ViewModels/Stack/INiagaraStackItemGroupAddUtilities.h"
 #include "ViewModels/Stack/NiagaraStackGraphUtilities.h"
 
diff --git a/Engine/Plugins/FX/Niagara/Source/NiagaraEditor/Private/ViewModels/Stack/NiagaraStackEventScriptItemGroup.cpp b/Engine/Plugins/FX/Niagara/Source/NiagaraEditor/Private/ViewModels/Stack/NiagaraStackEventScriptItemGroup.cpp
--- a/Engine/Plugins/FX/Niagara/Source/NiagaraEditor/Private/ViewModels/Stack/NiagaraStackEventScriptItemGroup.cpp
+++ b/Engine/Plugins/FX/Niagara/Source/NiagaraEditor/Private/ViewModels/Stack/NiagaraStackEventScriptItemGroup.cpp
@@ -4,7 +4,6 @@
 #include "NiagaraEmitter.h"
 #include "ViewModels/Stack/NiagaraStackObject.h"
 #include "ViewModels/Stack/NiagaraStackModuleItem.h"
-#include "ViewModels/Stack/NiagaraStackSpacer.h"
 #include "NiagaraScriptViewModel.h"
 #include "NiagaraScriptGraphViewModel.h"
 #include "ViewModels/NiagaraEmitterViewModel.h"
@@ -202,17 +201,6 @@ void UNiagaraStackEventScriptItemGroup::RefreshChildrenInternal(const TArray<UNi
 {
 	bHasBaseEventHandlerCache.Reset();
 
-	FName EventSpacerKey = *FString::Printf(TEXT("EventSpacer"));
-	UNiagaraStackSpacer* SeparatorSpacer = FindCurrentChildOfTypeByPredicate<UNiagaraStackSpacer>(CurrentChildren,
-		[=](UNiagaraStackSpacer* CurrentSpacer) { return CurrentSpacer->GetSpacerKey() == EventSpacerKey; });
-	if (SeparatorSpacer == nullptr)
-	{
-		SeparatorSpacer = NewObject<UNiagaraStackSpacer>(this);
-		FRequiredEntryData RequiredEntryData(GetSystemViewModel(), GetEmitterViewModel(), GetExecutionCategoryName(), NAME_None, GetStackEditorData());
-		SeparatorSpacer->Initialize(RequiredEntryData, EventSpacerKey);
-	}
-	NewChildren.Add(SeparatorSpacer);
-
 	UNiagaraEmitter* Emitter = GetEmitterViewModel()->GetEmitter();
 
 	const FNiagaraEventScriptProperties* EventScriptProperties = Emitter->GetEventHandlers().FindByPredicate(
diff --git a/Engine/Plugins/FX/Niagara/Source/NiagaraEditor/Private/ViewModels/Stack/NiagaraStackFunctionInput.cpp b/Engine/Plugins/FX/Niagara/Source/NiagaraEditor/Private/ViewModels/Stack/NiagaraStackFunctionInput.cpp
--- a/Engine/Plugins/FX/Niagara/Source/NiagaraEditor/Private/ViewModels/Stack/NiagaraStackFunctionInput.cpp
+++ b/Engine/Plugins/FX/Niagara/Source/NiagaraEditor/Private/ViewModels/Stack/NiagaraStackFunctionInput.cpp
@@ -3,7 +3,6 @@
 #include "ViewModels/Stack/NiagaraStackFunctionInput.h"
 #include "ViewModels/Stack/NiagaraStackFunctionInputCollection.h"
 #include "ViewModels/Stack/NiagaraStackObject.h"
-#include "ViewModels/Stack/NiagaraStackSpacer.h"
 #include "NiagaraNodeFunctionCall.h"
 #include "NiagaraNodeAssignment.h"
 #include "NiagaraNodeParameterMapGet.h"
@@ -398,20 +397,6 @@ void UNiagaraStackFunctionInput::RefreshChildrenInternal(const TArray<UNiagaraSt
 		NewChildren.Add(ValueObjectEntry);
 	}
 
-	if (NewChildren.Num() > 0)
-	{
-		UNiagaraStackSpacer* ChildrenSpacer = FindCurrentChildOfTypeByPredicate<UNiagaraStackSpacer>(CurrentChildren,
-			[=](UNiagaraStackSpacer* CurrentSpacer) { return true; });
-
-		if (ChildrenSpacer == nullptr)
-		{
-			ChildrenSpacer = NewObject<UNiagaraStackSpacer>(this);
-			ChildrenSpacer->Initialize(CreateDefaultChildRequiredData(), "ChildrenSpacer", .5f, UNiagaraStackEntry::EStackRowStyle::ItemContent);
-		}
-
-		NewChildren.Add(ChildrenSpacer);
-	}
-
 	DisplayNameOverride.Reset();
 
 	if (InputMetaData)
diff --git a/Engine/Plugins/FX/Niagara/Source/NiagaraEditor/Private/ViewModels/Stack/NiagaraStackGraphUtilities.cpp b/Engine/Plugins/FX/Niagara/Source/NiagaraEditor/Private/ViewModels/Stack/NiagaraStackGraphUtilities.cpp
--- a/Engine/Plugins/FX/Niagara/Source/NiagaraEditor/Private/ViewModels/Stack/NiagaraStackGraphUtilities.cpp
+++ b/Engine/Plugins/FX/Niagara/Source/NiagaraEditor/Private/ViewModels/Stack/NiagaraStackGraphUtilities.cpp
@@ -1847,7 +1847,7 @@ bool FNiagaraStackGraphUtilities::GetStackIssuesRecursively(const UNiagaraStackE
 	return OutIssues.Num() > 0;
 }
 
-void FNiagaraStackGraphUtilities::MoveModule(UNiagaraScript& SourceScript, UNiagaraNodeFunctionCall& ModuleToMove, UNiagaraSystem& TargetSystem, FGuid TargetEmitterHandleId, ENiagaraScriptUsage TargetUsage, FGuid TargetUsageId, int32 TargetModuleIndex, bool bForceCopy)
+void FNiagaraStackGraphUtilities::MoveModule(UNiagaraScript& SourceScript, UNiagaraNodeFunctionCall& ModuleToMove, UNiagaraSystem& TargetSystem, FGuid TargetEmitterHandleId, ENiagaraScriptUsage TargetUsage, FGuid TargetUsageId, int32 TargetModuleIndex, bool bForceCopy, UNiagaraNodeFunctionCall*& OutMovedModule)
 {
 	UNiagaraScript* TargetScript = FNiagaraEditorUtilities::GetScriptFromSystem(TargetSystem, TargetEmitterHandleId, TargetUsage, TargetUsageId);
 	checkf(TargetScript != nullptr, TEXT("Target script not found"));
@@ -1864,6 +1864,12 @@ void FNiagaraStackGraphUtilities::MoveModule(UNiagaraScript& SourceScript, UNiag
 	UNiagaraGraph* SourceGraph = ModuleToMove.GetNiagaraGraph();
 	UNiagaraGraph* TargetGraph = TargetOutputNode->GetNiagaraGraph();
 
+	if (SourceGraph != TargetGraph && bForceCopy == false)
+	{
+		SourceGraph->Modify();
+	}
+	TargetGraph->Modify();
+
 	// If the source and target scripts don't match, or we're forcing a copy we need to collect the rapid iteration parameter values for each function in the source group
 	// so we can restore them after moving.
 	TMap<FGuid, TArray<FNiagaraVariable>> SourceFunctionIdToRapidIterationParametersMap;
@@ -2038,6 +2044,8 @@ void FNiagaraStackGraphUtilities::MoveModule(UNiagaraScript& SourceScript, UNiag
 			}
 		}
 	}
+
+	OutMovedModule = Cast<UNiagaraNodeFunctionCall>(TargetGroup.EndNode);
 }
 
 bool FNiagaraStackGraphUtilities::ParameterAllowedInExecutionCategory(const FName InParameterName, const FName ExecutionCategory)
diff --git a/Engine/Plugins/FX/Niagara/Source/NiagaraEditor/Private/ViewModels/Stack/NiagaraStackItem.cpp b/Engine/Plugins/FX/Niagara/Source/NiagaraEditor/Private/ViewModels/Stack/NiagaraStackItem.cpp
--- a/Engine/Plugins/FX/Niagara/Source/NiagaraEditor/Private/ViewModels/Stack/NiagaraStackItem.cpp
+++ b/Engine/Plugins/FX/Niagara/Source/NiagaraEditor/Private/ViewModels/Stack/NiagaraStackItem.cpp
@@ -1,18 +1,18 @@
 // Copyright 1998-2019 Epic Games, Inc. All Rights Reserved.
 
 #include "ViewModels/Stack/NiagaraStackItem.h"
-#include "ViewModels/Stack/NiagaraStackSpacer.h"
-#include "ViewModels/Stack/NiagaraStackAdvancedExpander.h"
+#include "ViewModels/Stack/NiagaraStackItemFooter.h"
 #include "ViewModels/Stack/NiagaraStackGraphUtilities.h"
 #include "ViewModels/Stack/NiagaraStackErrorItem.h"
 #include "ViewModels/Stack/NiagaraStackEntry.h"
+#include "ViewModels/NiagaraSystemViewModel.h"
+#include "ViewModels/NiagaraSystemSelectionViewModel.h"
 #include "NiagaraStackEditorData.h"
 
 void UNiagaraStackItem::Initialize(FRequiredEntryData InRequiredEntryData, FString InStackEditorDataKey)
 {
 	Super::Initialize(InRequiredEntryData, InStackEditorDataKey);
 	AddChildFilter(FOnFilterChild::CreateUObject(this, &UNiagaraStackItem::FilterAdvancedChildren));
-	AddChildFilter(FOnFilterChild::CreateUObject(this, &UNiagaraStackItem::FilterShowAdvancedChild));
 }
 
 UNiagaraStackEntry::EStackRowStyle UNiagaraStackItem::GetStackRowStyle() const
@@ -54,23 +54,35 @@ EStackIssueSeverity UNiagaraStackItem::GetHighestStackIssueSeverity() const
 	return HighestIssueSeverity.GetValue();
 }
 
-void UNiagaraStackItem::RefreshChildrenInternal(const TArray<UNiagaraStackEntry*>& CurrentChildren, TArray<UNiagaraStackEntry*>& NewChildren, TArray<FStackIssue>& NewIssues)
+void UNiagaraStackItem::SetIsEnabled(bool bInIsEnabled)
+{
+	if (ItemFooter != nullptr)
+	{
+		ItemFooter->SetIsEnabled(bInIsEnabled);
+	}
+	SetIsEnabledInternal(bInIsEnabled);
+}
+
+void UNiagaraStackItem::Delete()
 {
-	if (FooterSpacer == nullptr)
+	if (GetDisplayedObject() != nullptr)
 	{
-		FooterSpacer = NewObject<UNiagaraStackSpacer>(this);
-		FooterSpacer->Initialize(CreateDefaultChildRequiredData(), "ItemFooterSpacer", 1, UNiagaraStackEntry::EStackRowStyle::ItemContent);
+		GetSystemViewModel()->GetSelectionViewModel()->RemoveEntryFromSelectionByDisplayedObject(GetDisplayedObject());
 	}
+	DeleteInternal();
+}
 
-	if (ShowAdvancedExpander == nullptr)
+void UNiagaraStackItem::RefreshChildrenInternal(const TArray<UNiagaraStackEntry*>& CurrentChildren, TArray<UNiagaraStackEntry*>& NewChildren, TArray<FStackIssue>& NewIssues)
+{
+	if (ItemFooter == nullptr)
 	{
-		ShowAdvancedExpander = NewObject<UNiagaraStackAdvancedExpander>(this);
-		ShowAdvancedExpander->Initialize(CreateDefaultChildRequiredData(), GetStackEditorDataKey(), GetOwningNiagaraNode());
-		ShowAdvancedExpander->SetOnToggleShowAdvanced(UNiagaraStackAdvancedExpander::FOnToggleShowAdvanced::CreateUObject(this, &UNiagaraStackItem::ToggleShowAdvanced));
+		ItemFooter = NewObject<UNiagaraStackItemFooter>(this);
+		ItemFooter->Initialize(CreateDefaultChildRequiredData(), GetStackEditorDataKey());
+		ItemFooter->SetOnToggleShowAdvanced(UNiagaraStackItemFooter::FOnToggleShowAdvanced::CreateUObject(this, &UNiagaraStackItem::ToggleShowAdvanced));
 	}
+	ItemFooter->SetIsEnabled(GetIsEnabled());
 
-	NewChildren.Add(FooterSpacer);
-	NewChildren.Add(ShowAdvancedExpander);
+	NewChildren.Add(ItemFooter);
 	RecursiveStackIssuesCount.Reset();
 	HighestIssueSeverity.Reset();
 }
@@ -93,7 +105,7 @@ void GetContentChildren(UNiagaraStackEntry& CurrentEntry, TArray<UNiagaraStackIt
 void UNiagaraStackItem::PostRefreshChildrenInternal()
 {
 	Super::PostRefreshChildrenInternal();
-	bHasAdvancedContent = false;
+	bool bHasAdvancedContent = false;
 	TArray<UNiagaraStackItemContent*> ContentChildren;
 	GetContentChildren(*this, ContentChildren);
 	for (UNiagaraStackItemContent* ContentChild : ContentChildren)
@@ -104,6 +116,7 @@ void UNiagaraStackItem::PostRefreshChildrenInternal()
 			break;
 		}
 	}
+	ItemFooter->SetHasAdvancedContent(bHasAdvancedContent);
 }
 
 int32 UNiagaraStackItem::GetChildIndentLevel() const
@@ -111,11 +124,6 @@ int32 UNiagaraStackItem::GetChildIndentLevel() const
 	return GetIndentLevel();
 }
 
-UNiagaraNode* UNiagaraStackItem::GetOwningNiagaraNode() const
-{
-	return nullptr;
-}
-
 void UNiagaraStackItem::ChlildStructureChangedInternal()
 {
 	Super::ChlildStructureChangedInternal();
@@ -136,22 +144,6 @@ bool UNiagaraStackItem::FilterAdvancedChildren(const UNiagaraStackEntry& Child)
 	}
 }
 
-bool UNiagaraStackItem::FilterShowAdvancedChild(const UNiagaraStackEntry& Child) const 
-{
-	if (&Child == ShowAdvancedExpander && bHasAdvancedContent == false)
-	{
-		return false;
-	}
-	else if (&Child == FooterSpacer && bHasAdvancedContent)
-	{
-		return false;
-	}
-	else
-	{
-		return true;
-	}
-}
-
 void UNiagaraStackItem::ToggleShowAdvanced()
 {
 	bool bCurrentShowAdvanced = GetStackEditorData().GetStackItemShowAdvanced(GetStackEditorDataKey(), false);
diff --git a/Engine/Plugins/FX/Niagara/Source/NiagaraEditor/Private/ViewModels/Stack/NiagaraStackItemFooter.cpp b/Engine/Plugins/FX/Niagara/Source/NiagaraEditor/Private/ViewModels/Stack/NiagaraStackItemFooter.cpp
new file mode 100644
--- /dev/null
+++ b/Engine/Plugins/FX/Niagara/Source/NiagaraEditor/Private/ViewModels/Stack/NiagaraStackItemFooter.cpp
@@ -0,0 +1,62 @@
+// Copyright 1998-2019 Epic Games, Inc. All Rights Reserved.
+
+#include "ViewModels/Stack/NiagaraStackItemFooter.h"
+#include "NiagaraStackEditorData.h"
+#include "ViewModels/NiagaraEmitterViewModel.h"
+#include "NiagaraNode.h"
+
+void UNiagaraStackItemFooter::Initialize(
+	FRequiredEntryData InRequiredEntryData,
+	FString InOwnerStackItemEditorDataKey)
+{
+	Super::Initialize(InRequiredEntryData, FString());
+	OwnerStackItemEditorDataKey = InOwnerStackItemEditorDataKey;
+	bHasAdvancedContent = false;
+	bIsEnabled = true;
+}
+
+bool UNiagaraStackItemFooter::GetCanExpand() const
+{
+	return false;
+}
+
+UNiagaraStackEntry::EStackRowStyle UNiagaraStackItemFooter::GetStackRowStyle() const
+{
+	return UNiagaraStackEntry::EStackRowStyle::ItemFooter;
+}
+
+bool UNiagaraStackItemFooter::GetIsEnabled() const
+{
+	return bIsEnabled;
+}
+
+void UNiagaraStackItemFooter::SetIsEnabled(bool bInIsEnabled)
+{
+	bIsEnabled = bInIsEnabled;
+}
+
+
+bool UNiagaraStackItemFooter::GetHasAdvancedContent() const
+{
+	return bHasAdvancedContent;
+}
+
+void UNiagaraStackItemFooter::SetHasAdvancedContent(bool bInHasAdvancedRows)
+{
+	bHasAdvancedContent = bInHasAdvancedRows;
+}
+
+void UNiagaraStackItemFooter::SetOnToggleShowAdvanced(FOnToggleShowAdvanced OnExpandedChanged)
+{
+	ToggleShowAdvancedDelegate = OnExpandedChanged;
+}
+
+bool UNiagaraStackItemFooter::GetShowAdvanced() const
+{
+	return GetStackEditorData().GetStackItemShowAdvanced(OwnerStackItemEditorDataKey, false);
+}
+
+void UNiagaraStackItemFooter::ToggleShowAdvanced()
+{
+	ToggleShowAdvancedDelegate.ExecuteIfBound();
+}
diff --git a/Engine/Plugins/FX/Niagara/Source/NiagaraEditor/Private/ViewModels/Stack/NiagaraStackItemGroup.cpp b/Engine/Plugins/FX/Niagara/Source/NiagaraEditor/Private/ViewModels/Stack/NiagaraStackItemGroup.cpp
--- a/Engine/Plugins/FX/Niagara/Source/NiagaraEditor/Private/ViewModels/Stack/NiagaraStackItemGroup.cpp
+++ b/Engine/Plugins/FX/Niagara/Source/NiagaraEditor/Private/ViewModels/Stack/NiagaraStackItemGroup.cpp
@@ -5,7 +5,6 @@
 #include "ViewModels/Stack/NiagaraStackErrorItem.h"
 #include "ViewModels/Stack/NiagaraStackEntry.h"
 #include "NiagaraStackEditorData.h"
-#include "ViewModels/Stack/NiagaraStackSpacer.h"
 #include "ViewModels/Stack/NiagaraStackGraphUtilities.h"
 #include "ViewModels/NiagaraSystemViewModel.h"
 #include "ViewModels/NiagaraEmitterHandleViewModel.h"
@@ -18,6 +17,7 @@ void UNiagaraStackItemGroup::Initialize(FRequiredEntryData InRequiredEntryData,
 	GroupDisplayName = InDisplayName;
 	GroupToolTip = InToolTip;
 	AddUtilities = InAddUtilities;
+	GroupFooter = nullptr;
 }
 
 FText UNiagaraStackItemGroup::GetDisplayName() const
@@ -81,6 +81,14 @@ void UNiagaraStackItemGroup::SetDisplayName(FText InDisplayName)
 
 void UNiagaraStackItemGroup::RefreshChildrenInternal(const TArray<UNiagaraStackEntry*>& CurrentChildren, TArray<UNiagaraStackEntry*>& NewChildren, TArray<FStackIssue>& NewIssues)
 {
+	if (GroupFooter == nullptr)
+	{
+		GroupFooter = NewObject<UNiagaraStackItemGroupFooter>(this);
+		GroupFooter->Initialize(FRequiredEntryData(GetSystemViewModel(), GetEmitterViewModel(), NAME_None, NAME_None, GetStackEditorData()));
+	}
+
+	NewChildren.Add(GroupFooter);
+
 	RecursiveStackIssuesCount.Reset();
 	HighestIssueSeverity.Reset();
 
@@ -109,3 +117,18 @@ void UNiagaraStackItemGroup::ChlildStructureChangedInternal()
 	RecursiveStackIssuesCount.Reset();
 	HighestIssueSeverity.Reset();
 }
+
+void UNiagaraStackItemGroupFooter::Initialize(FRequiredEntryData InRequiredEntryData)
+{
+	Super::Initialize(InRequiredEntryData, FString());
+}
+
+UNiagaraStackEntry::EStackRowStyle UNiagaraStackItemGroupFooter::GetStackRowStyle() const
+{
+	return UNiagaraStackEntry::EStackRowStyle::GroupFooter;
+}
+
+bool UNiagaraStackItemGroupFooter::GetCanExpand() const
+{
+	return false;
+}
diff --git a/Engine/Plugins/FX/Niagara/Source/NiagaraEditor/Private/ViewModels/Stack/NiagaraStackModuleItem.cpp b/Engine/Plugins/FX/Niagara/Source/NiagaraEditor/Private/ViewModels/Stack/NiagaraStackModuleItem.cpp
--- a/Engine/Plugins/FX/Niagara/Source/NiagaraEditor/Private/ViewModels/Stack/NiagaraStackModuleItem.cpp
+++ b/Engine/Plugins/FX/Niagara/Source/NiagaraEditor/Private/ViewModels/Stack/NiagaraStackModuleItem.cpp
@@ -30,6 +30,7 @@
 #include "NiagaraCommon.h"
 #include "NiagaraConstants.h"
 #include "Widgets/SWidget.h"
+#include "NiagaraActions.h"
 
 #include "ScopedTransaction.h"
 #include "NiagaraScriptVariable.h"
@@ -114,12 +115,7 @@ UNiagaraStackModuleItem::UNiagaraStackModuleItem()
 {
 }
 
-const UNiagaraNodeFunctionCall& UNiagaraStackModuleItem::GetModuleNode() const
-{
-	return *FunctionCallNode;
-}
-
-UNiagaraNodeFunctionCall& UNiagaraStackModuleItem::GetModuleNode()
+UNiagaraNodeFunctionCall& UNiagaraStackModuleItem::GetModuleNode() const
 {
 	return *FunctionCallNode;
 }
@@ -153,6 +149,10 @@ FText UNiagaraStackModuleItem::GetDisplayName() const
 	}
 }
 
+UObject* UNiagaraStackModuleItem::GetDisplayedObject() const
+{
+	return FunctionCallNode;
+}
 
 FText UNiagaraStackModuleItem::GetTooltipText() const
 {
@@ -246,6 +246,38 @@ void UNiagaraStackModuleItem::RefreshChildrenInternal(const TArray<UNiagaraStack
 	RefreshIssues(NewIssues);
 }
 
+TOptional<UNiagaraStackEntry::FDropRequestResponse> UNiagaraStackModuleItem::CanDropInternal(const FDropRequest& DropRequest)
+{
+	if (DropRequest.DragDropOperation->IsOfType<FNiagaraParameterDragOperation>() && DropRequest.DropOptions != UNiagaraStackEntry::EDropOptions::Overview && FunctionCallNode->IsA<UNiagaraNodeAssignment>())
+	{
+		TSharedRef<const FNiagaraParameterDragOperation> ParameterDragDropOp = StaticCastSharedRef<const FNiagaraParameterDragOperation>(DropRequest.DragDropOperation);
+		TSharedPtr<const FNiagaraParameterAction> ParameterAction = StaticCastSharedPtr<const FNiagaraParameterAction>(ParameterDragDropOp->GetSourceAction());
+		if (ParameterAction.IsValid())
+		{
+			if (FNiagaraStackGraphUtilities::ParameterIsCompatibleWithScriptUsage(ParameterAction->GetParameter(), OutputNode->GetUsage()))
+			{
+				return FDropRequestResponse(EItemDropZone::OntoItem, LOCTEXT("DropParameterToAdd", "Add this parameter to this 'Set Variables' node."));
+			}
+			else
+			{
+				return FDropRequestResponse(TOptional<EItemDropZone>(), LOCTEXT("CantDropParameterByUsage", "Can not drop this parameter here because\nit's not valid for this usage context."));
+			}
+		}
+	}
+	return TOptional<FDropRequestResponse>();
+}
+
+TOptional<UNiagaraStackEntry::FDropRequestResponse> UNiagaraStackModuleItem::DropInternal(const FDropRequest& DropRequest)
+{
+	// If the drop was allowed from the can drop just return the drop zone here since it will be handled by the drop target in the stack control.
+	// TODO: Unify the drop target dropping with the existing drag/drop code.
+	if (DropRequest.DropOptions != UNiagaraStackEntry::EDropOptions::Overview)
+	{
+		return FDropRequestResponse(DropRequest.DropZone);
+	}
+	return TOptional<FDropRequestResponse>();
+}
+
 void UNiagaraStackModuleItem::RefreshIssues(TArray<FStackIssue>& NewIssues)
 {
 	if (!GetIsEnabled())
@@ -511,7 +543,8 @@ void UNiagaraStackModuleItem::RefreshIssues(TArray<FStackIssue>& NewIssues)
 						int32 CorrectIndex = Dependency.Type == ENiagaraModuleDependencyType::PostDependency ? DisorderedNode.Index : DisorderedNode.Index + 1;
 						checkf(ModuleIndex != INDEX_NONE, TEXT("Module data wasn't found in system for current module!"));
 						UNiagaraScript& OwningScript = *FNiagaraEditorUtilities::GetScriptFromSystem(GetSystemViewModel()->GetSystem(), SystemModuleData[ModuleIndex].EmitterHandleId, SystemModuleData[ModuleIndex].Usage, SystemModuleData[ModuleIndex].UsageId);
-						FNiagaraStackGraphUtilities::MoveModule(OwningScript, *FunctionCallNode, GetSystemViewModel()->GetSystem(), DisorderedNode.EmitterHandleId, DisorderedNode.Usage, DisorderedNode.UsageId, CorrectIndex, false);
+						UNiagaraNodeFunctionCall* MovedNode;
+						FNiagaraStackGraphUtilities::MoveModule(OwningScript, *FunctionCallNode, GetSystemViewModel()->GetSystem(), DisorderedNode.EmitterHandleId, DisorderedNode.Usage, DisorderedNode.UsageId, CorrectIndex, false, MovedNode);
 						// enable if needed
 						if (bNeedsEnable)
 						{
@@ -770,7 +803,7 @@ bool UNiagaraStackModuleItem::GetIsEnabled() const
 	return bIsEnabled;
 }
 
-void UNiagaraStackModuleItem::SetIsEnabled(bool bInIsEnabled)
+void UNiagaraStackModuleItem::SetIsEnabledInternal(bool bInIsEnabled)
 {
 	FScopedTransaction ScopedTransaction(LOCTEXT("EnableDisableModule", "Enable/Disable Module"));
 	FNiagaraStackGraphUtilities::SetModuleIsEnabled(*FunctionCallNode, bInIsEnabled);
@@ -797,7 +830,7 @@ bool UNiagaraStackModuleItem::TestCanDeleteWithMessage(FText& OutCanDeleteMessag
 	}
 }
 
-void UNiagaraStackModuleItem::Delete()
+void UNiagaraStackModuleItem::DeleteInternal()
 {
 	checkf(CanMoveAndDelete(), TEXT("This module can't be deleted"));
 
@@ -827,7 +860,7 @@ void UNiagaraStackModuleItem::Delete()
 	}
 }
 
-int32 UNiagaraStackModuleItem::GetModuleIndex()
+int32 UNiagaraStackModuleItem::GetModuleIndex() const
 {
 	TArray<FNiagaraStackGraphUtilities::FStackNodeGroup> StackGroups;
 	FNiagaraStackGraphUtilities::GetStackNodeGroups(*FunctionCallNode, StackGroups);
@@ -851,11 +884,6 @@ UNiagaraNodeOutput* UNiagaraStackModuleItem::GetOutputNode() const
 	return OutputNode;
 }
 
-void UNiagaraStackModuleItem::NotifyModuleMoved()
-{
-	ModifiedGroupItemsDelegate.Broadcast();
-}
-
 bool UNiagaraStackModuleItem::CanAddInput(FNiagaraVariable InputParameter) const
 {
 	UNiagaraNodeAssignment* AssignmentModule = Cast<UNiagaraNodeAssignment>(FunctionCallNode);
diff --git a/Engine/Plugins/FX/Niagara/Source/NiagaraEditor/Private/ViewModels/Stack/NiagaraStackModuleSpacer.cpp b/Engine/Plugins/FX/Niagara/Source/NiagaraEditor/Private/ViewModels/Stack/NiagaraStackModuleSpacer.cpp
deleted file mode 100644
diff --git a/Engine/Plugins/FX/Niagara/Source/NiagaraEditor/Private/ViewModels/Stack/NiagaraStackRenderItemGroup.cpp b/Engine/Plugins/FX/Niagara/Source/NiagaraEditor/Private/ViewModels/Stack/NiagaraStackRenderItemGroup.cpp
--- a/Engine/Plugins/FX/Niagara/Source/NiagaraEditor/Private/ViewModels/Stack/NiagaraStackRenderItemGroup.cpp
+++ b/Engine/Plugins/FX/Niagara/Source/NiagaraEditor/Private/ViewModels/Stack/NiagaraStackRenderItemGroup.cpp
@@ -2,7 +2,6 @@
 
 #include "ViewModels/Stack/NiagaraStackRenderItemGroup.h"
 #include "ViewModels/Stack/NiagaraStackRendererItem.h"
-#include "ViewModels/Stack/NiagaraStackSpacer.h"
 #include "ViewModels/NiagaraEmitterViewModel.h"
 #include "NiagaraEmitter.h"
 #include "NiagaraEmitterEditorData.h"
@@ -133,18 +132,6 @@ void UNiagaraStackRenderItemGroup::RefreshChildrenInternal(const TArray<UNiagara
 	int32 RendererIndex = 0;
 	for (UNiagaraRendererProperties* RendererProperties : GetEmitterViewModel()->GetEmitter()->GetRenderers())
 	{
-		FName RendererSpacerKey = *FString::Printf(TEXT("Renderer%i"), RendererIndex);
-		UNiagaraStackSpacer* RendererSpacer = FindCurrentChildOfTypeByPredicate<UNiagaraStackSpacer>(CurrentChildren,
-			[=](UNiagaraStackSpacer* CurrentRendererSpacer) { return CurrentRendererSpacer->GetSpacerKey() == RendererSpacerKey; });
-
-		if (RendererSpacer == nullptr)
-		{
-			RendererSpacer = NewObject<UNiagaraStackSpacer>(this);
-			RendererSpacer->Initialize(CreateDefaultChildRequiredData(), RendererSpacerKey);
-		}
-
-		NewChildren.Add(RendererSpacer);
-
 		UNiagaraStackRendererItem* RendererItem = FindCurrentChildOfTypeByPredicate<UNiagaraStackRendererItem>(CurrentChildren,
 			[=](UNiagaraStackRendererItem* CurrentRendererItem) { return CurrentRendererItem->GetRendererProperties() == RendererProperties; });
 
diff --git a/Engine/Plugins/FX/Niagara/Source/NiagaraEditor/Private/ViewModels/Stack/NiagaraStackRendererItem.cpp b/Engine/Plugins/FX/Niagara/Source/NiagaraEditor/Private/ViewModels/Stack/NiagaraStackRendererItem.cpp
--- a/Engine/Plugins/FX/Niagara/Source/NiagaraEditor/Private/ViewModels/Stack/NiagaraStackRendererItem.cpp
+++ b/Engine/Plugins/FX/Niagara/Source/NiagaraEditor/Private/ViewModels/Stack/NiagaraStackRendererItem.cpp
@@ -163,7 +163,7 @@ bool UNiagaraStackRendererItem::TestCanDeleteWithMessage(FText& OutCanDeleteMess
 	}
 }
 
-void UNiagaraStackRendererItem::Delete()
+void UNiagaraStackRendererItem::DeleteInternal()
 {
 	const FScopedTransaction Transaction(LOCTEXT("DeleteRenderer", "Delete Renderer"));
 
@@ -222,7 +222,7 @@ bool UNiagaraStackRendererItem::GetIsEnabled() const
 	return RendererProperties->GetIsEnabled();
 }
 
-void UNiagaraStackRendererItem::SetIsEnabled(bool bInIsEnabled)
+void UNiagaraStackRendererItem::SetIsEnabledInternal(bool bInIsEnabled)
 {
 	FScopedTransaction ScopedTransaction(LOCTEXT("SetRendererEnabledState", "Set renderer enabled/disabled state."));
 	RendererProperties->Modify();
diff --git a/Engine/Plugins/FX/Niagara/Source/NiagaraEditor/Private/ViewModels/Stack/NiagaraStackRoot.cpp b/Engine/Plugins/FX/Niagara/Source/NiagaraEditor/Private/ViewModels/Stack/NiagaraStackRoot.cpp
--- a/Engine/Plugins/FX/Niagara/Source/NiagaraEditor/Private/ViewModels/Stack/NiagaraStackRoot.cpp
+++ b/Engine/Plugins/FX/Niagara/Source/NiagaraEditor/Private/ViewModels/Stack/NiagaraStackRoot.cpp
@@ -14,7 +14,6 @@
 #include "NiagaraSystemEditorData.h"
 #include "NiagaraEmitterEditorData.h"
 #include "ViewModels/Stack/NiagaraStackSystemSettingsGroup.h"
-#include "ViewModels/Stack/NiagaraStackSpacer.h"
 
 #define LOCTEXT_NAMESPACE "NiagaraStackViewModel"
 
@@ -163,47 +162,22 @@ void UNiagaraStackRoot::RefreshChildrenInternal(const TArray<UNiagaraStackEntry*
 		RenderGroup->Initialize(RequiredEntryData);
 	}
 
-	auto GetOrCreateSpacer = [&](FName SpacerExecutionCategory, FName SpacerKey)
-	{
-		UNiagaraStackSpacer* Spacer = FindCurrentChildOfTypeByPredicate<UNiagaraStackSpacer>(CurrentChildren,
-			[SpacerKey](UNiagaraStackSpacer* CurrentSpacer) { return CurrentSpacer->GetSpacerKey() == SpacerKey; });
-
-		if (Spacer == nullptr)
-		{
-			Spacer = NewObject<UNiagaraStackSpacer>(this);
-			FRequiredEntryData RequiredEntryData(GetSystemViewModel(), GetEmitterViewModel(),
-				SpacerExecutionCategory, NAME_None,
-				GetSystemViewModel()->GetEditorData().GetStackEditorData());
-			Spacer->Initialize(RequiredEntryData, SpacerKey, 1.5f);
-		}
-
-		return Spacer;
-	};
-
 	// Populate new children
 	if (bIncludeSystemInformation)
 	{
 		NewChildren.Add(SystemSettingsGroup);
-		NewChildren.Add(GetOrCreateSpacer(FExecutionCategoryNames::System, "SystemSettingsSpacer"));
 		NewChildren.Add(SystemSpawnGroup);
-		NewChildren.Add(GetOrCreateSpacer(FExecutionCategoryNames::System, "SystemSpawnSpacer"));
 		NewChildren.Add(SystemUpdateGroup);
-		NewChildren.Add(GetOrCreateSpacer(NAME_None, "SystemSpacer"));
 	}
 
 	if (bIncludeEmitterInformation)
 	{
 		NewChildren.Add(EmitterSettingsGroup);
-		NewChildren.Add(GetOrCreateSpacer(FExecutionCategoryNames::Emitter, "EmitterSettingsSpacer"));
 		NewChildren.Add(EmitterSpawnGroup);
-		NewChildren.Add(GetOrCreateSpacer(FExecutionCategoryNames::Emitter, "EmitterSpawnSpacer"));
 		NewChildren.Add(EmitterUpdateGroup);
-		NewChildren.Add(GetOrCreateSpacer(NAME_None, "EmitterSpacer"));
 
 		NewChildren.Add(ParticleSpawnGroup);
-		NewChildren.Add(GetOrCreateSpacer(FExecutionCategoryNames::Particle, "ParticleSpawnSpacer"));
 		NewChildren.Add(ParticleUpdateGroup);
-		NewChildren.Add(GetOrCreateSpacer(FExecutionCategoryNames::Particle, "ParticleUpdateSpacer"));
 
 		for (const FNiagaraEventScriptProperties& EventScriptProperties : GetEmitterViewModel()->GetEmitter()->GetEventHandlers())
 		{
@@ -224,7 +198,6 @@ void UNiagaraStackRoot::RefreshChildrenInternal(const TArray<UNiagaraStackEntry*
 		}
 
 		NewChildren.Add(AddEventHandlerGroup);
-		NewChildren.Add(GetOrCreateSpacer(NAME_None, "ParticleSpacer"));
 
 		NewChildren.Add(RenderGroup);
 	}
diff --git a/Engine/Plugins/FX/Niagara/Source/NiagaraEditor/Private/ViewModels/Stack/NiagaraStackScriptItemGroup.cpp b/Engine/Plugins/FX/Niagara/Source/NiagaraEditor/Private/ViewModels/Stack/NiagaraStackScriptItemGroup.cpp
--- a/Engine/Plugins/FX/Niagara/Source/NiagaraEditor/Private/ViewModels/Stack/NiagaraStackScriptItemGroup.cpp
+++ b/Engine/Plugins/FX/Niagara/Source/NiagaraEditor/Private/ViewModels/Stack/NiagaraStackScriptItemGroup.cpp
@@ -2,7 +2,6 @@
 
 #include "ViewModels/Stack/NiagaraStackScriptItemGroup.h"
 #include "ViewModels/Stack/NiagaraStackModuleItem.h"
-#include "ViewModels/Stack/NiagaraStackSpacer.h"
 #include "ViewModels/NiagaraSystemViewModel.h"
 #include "ViewModels/NiagaraEmitterViewModel.h"
 #include "NiagaraScriptViewModel.h"
@@ -24,7 +23,7 @@
 #include "Widgets/Notifications/SNotificationList.h"
 #include "Framework/Notifications/NotificationManager.h"
 #include "Internationalization/Internationalization.h"
-#include "ViewModels/Stack/NiagaraStackModuleSpacer.h"
+#include "DragAndDrop/AssetDragDropOp.h"
 
 #define LOCTEXT_NAMESPACE "UNiagaraStackScriptItemGroup"
 
@@ -329,8 +328,6 @@ void UNiagaraStackScriptItemGroup::RefreshChildrenInternal(const TArray<UNiagara
 	TSharedPtr<FNiagaraScriptViewModel> ScriptViewModelPinned = ScriptViewModel.Pin();
 	checkf(ScriptViewModelPinned.IsValid(), TEXT("Can not refresh children when the script view model has been deleted."));
 
-	StackSpacerToModuleItemMap.Empty();
-
 	UNiagaraGraph* Graph = ScriptViewModelPinned->GetGraphViewModel()->GetGraph();
 	FText ErrorMessage;
 	bIsValidForOutput = false;
@@ -343,24 +340,8 @@ void UNiagaraStackScriptItemGroup::RefreshChildrenInternal(const TArray<UNiagara
 
 		TArray<UNiagaraNodeFunctionCall*> ModuleNodes;
 		FNiagaraStackGraphUtilities::GetOrderedModuleNodes(*MatchingOutputNode, ModuleNodes);
-		int32 ModuleIndex = 0;
 		for (UNiagaraNodeFunctionCall* ModuleNode : ModuleNodes)
 		{
-			FName ModuleSpacerKey = *FString::Printf(TEXT("Module%i"), ModuleIndex);
-			UNiagaraStackModuleSpacer* ModuleSpacer = FindCurrentChildOfTypeByPredicate<UNiagaraStackModuleSpacer>(CurrentChildren,
-				[=](UNiagaraStackModuleSpacer* CurrentModuleSpacer) { return CurrentModuleSpacer->GetSpacerKey() == ModuleSpacerKey; });
-
-			if (ModuleSpacer == nullptr)
-			{
-				ModuleSpacer = NewObject<UNiagaraStackModuleSpacer>(this);
-				ModuleSpacer->Initialize(CreateDefaultChildRequiredData(), GetScriptUsage(), ModuleSpacerKey, 1.4f, UNiagaraStackEntry::EStackRowStyle::None);
-				ModuleSpacer->OnStackSpacerAcceptDrop.BindUObject(this, &UNiagaraStackScriptItemGroup::AddParameterModuleToStack);
-				ModuleSpacer->OnStackSpacerRequestAssetDrop.BindUObject(this, &UNiagaraStackScriptItemGroup::CanAddAssetModuleToStack);
-				ModuleSpacer->OnStackSpacerAssetDrop.BindUObject(this, &UNiagaraStackScriptItemGroup::AddAssetModuleToStack);
-			}
-
-			NewChildren.Add(ModuleSpacer);
-
 			UNiagaraStackModuleItem* ModuleItem = FindCurrentChildOfTypeByPredicate<UNiagaraStackModuleItem>(CurrentChildren, 
 				[=](UNiagaraStackModuleItem* CurrentModuleItem) { return &CurrentModuleItem->GetModuleNode() == ModuleNode; });
 
@@ -372,27 +353,7 @@ void UNiagaraStackScriptItemGroup::RefreshChildrenInternal(const TArray<UNiagara
 			}
 
 			NewChildren.Add(ModuleItem);
-			StackSpacerToModuleItemMap.Add(FObjectKey(ModuleSpacer), ModuleItem);
-
-			ModuleIndex++;
 		}
-
-		// Add the post items spacer.
-		FName ModuleSpacerKey = *FString::Printf(TEXT("Module%i"), ModuleIndex);
-		UNiagaraStackModuleSpacer* ModuleSpacer = FindCurrentChildOfTypeByPredicate<UNiagaraStackModuleSpacer>(CurrentChildren,
-			[=](UNiagaraStackModuleSpacer* CurrentModuleSpacer) { return CurrentModuleSpacer->GetSpacerKey() == ModuleSpacerKey; });
-
-		if (ModuleSpacer == nullptr)
-		{
-			ModuleSpacer = NewObject<UNiagaraStackModuleSpacer>(this);
-			ModuleSpacer->Initialize(CreateDefaultChildRequiredData(), GetScriptUsage(), ModuleSpacerKey, 1.4f, UNiagaraStackEntry::EStackRowStyle::None);
-			ModuleSpacer->OnStackSpacerAcceptDrop.BindUObject(this, &UNiagaraStackScriptItemGroup::AddParameterModuleToStack);
-			ModuleSpacer->OnStackSpacerRequestAssetDrop.BindUObject(this, &UNiagaraStackScriptItemGroup::CanAddAssetModuleToStack);
-			ModuleSpacer->OnStackSpacerAssetDrop.BindUObject(this, &UNiagaraStackScriptItemGroup::AddAssetModuleToStack);
-		}
-
-		NewChildren.Add(ModuleSpacer);
-		StackSpacerToModuleItemMap.Add(FObjectKey(ModuleSpacer), nullptr);
 	}
 	Super::RefreshChildrenInternal(CurrentChildren, NewChildren, NewIssues);
 	RefreshIssues(NewIssues);
@@ -495,16 +456,14 @@ void UNiagaraStackScriptItemGroup::RefreshIssues(TArray<FStackIssue>& NewIssues)
 }
 
 void GenerateDragDropData(
-	UNiagaraNodeFunctionCall* SourceModule,
-	UNiagaraNodeFunctionCall* TargetModule, const UNiagaraGraph* TargetGraph, ENiagaraScriptUsage TargetScriptUsage, FGuid TargetScriptUsageId,
+	UNiagaraNodeFunctionCall& SourceModule,
+	UNiagaraNodeFunctionCall& TargetModule, EItemDropZone TargetZone,
 	TArray<FNiagaraStackGraphUtilities::FStackNodeGroup>& OutSourceStackGroups, int32& OutSourceGroupIndex,
 	TArray<FNiagaraStackGraphUtilities::FStackNodeGroup>& OutTargetStackGroups, int32& OutTargetGroupIndex)
 {
 	// Find the output nodes for the source and target
-	UNiagaraNodeOutput* SourceOutputNode = FNiagaraStackGraphUtilities::GetEmitterOutputNodeForStackNode(*SourceModule);
-	UNiagaraNodeOutput* TargetOutputNode = TargetModule != nullptr
-		? FNiagaraStackGraphUtilities::GetEmitterOutputNodeForStackNode(*TargetModule)
-		: TargetGraph->FindOutputNode(TargetScriptUsage, TargetScriptUsageId);
+	UNiagaraNodeOutput* SourceOutputNode = FNiagaraStackGraphUtilities::GetEmitterOutputNodeForStackNode(SourceModule);
+	UNiagaraNodeOutput* TargetOutputNode = FNiagaraStackGraphUtilities::GetEmitterOutputNodeForStackNode(TargetModule);
 
 	// Collect the stack node groups for the source and target.
 	FNiagaraStackGraphUtilities::GetStackNodeGroups(*SourceOutputNode, OutSourceStackGroups);
@@ -521,7 +480,7 @@ void GenerateDragDropData(
 	OutSourceGroupIndex = INDEX_NONE;
 	for (int32 GroupIndex = 0; GroupIndex < OutSourceStackGroups.Num(); GroupIndex++)
 	{
-		if (OutSourceStackGroups[GroupIndex].EndNode == SourceModule)
+		if (OutSourceStackGroups[GroupIndex].EndNode == &SourceModule)
 		{
 			OutSourceGroupIndex = GroupIndex;
 			break;
@@ -529,212 +488,357 @@ void GenerateDragDropData(
 	}
 
 	OutTargetGroupIndex = INDEX_NONE;
-	if (TargetModule == SourceModule)
+	if (&TargetModule == &SourceModule)
 	{
 		OutTargetGroupIndex = OutSourceGroupIndex;
 	}
-	else if (TargetModule != nullptr)
+	else
 	{
 		for (int32 GroupIndex = 0; GroupIndex < OutTargetStackGroups.Num(); GroupIndex++)
 		{
-			if (OutTargetStackGroups[GroupIndex].EndNode == TargetModule)
+			if (OutTargetStackGroups[GroupIndex].EndNode == &TargetModule)
 			{
 				OutTargetGroupIndex = GroupIndex;
 				break;
 			}
 		}
 	}
-	else
+
+	if (TargetZone == EItemDropZone::BelowItem)
 	{
-		// If there is no target module then we need to insert at the end.  The last group is the output node and we want to
-		// insert before that.
-		OutTargetGroupIndex = OutTargetStackGroups.Num() - 1;
+		OutTargetGroupIndex++;
 	}
 }
 
-TOptional<UNiagaraStackEntry::FDropResult> UNiagaraStackScriptItemGroup::ChildRequestCanDropInternal(const UNiagaraStackEntry& TargetChild, const TArray<UNiagaraStackEntry*>& DraggedEntries, EDragOptions DragOptions)
+TOptional<UNiagaraStackEntry::FDropRequestResponse> UNiagaraStackScriptItemGroup::ChildRequestCanDropInternal(const UNiagaraStackEntry& TargetChild, const FDropRequest& DropRequest)
 {
-	if (bIsValidForOutput && DraggedEntries.Num() == 1 && (DragOptions == UNiagaraStackEntry::EDragOptions::Move || DragOptions == UNiagaraStackEntry::EDragOptions::Copy))
+	if (DropRequest.DragDropOperation->IsOfType<FNiagaraStackEntryDragDropOp>())
+	{
+		return ChildRequestCanDropEntries(TargetChild, DropRequest);
+	}
+	else if (DropRequest.DragDropOperation->IsOfType<FAssetDragDropOp>())
 	{
-		UNiagaraStackModuleItem* SourceModuleItem = Cast<UNiagaraStackModuleItem>(DraggedEntries[0]);
-		if (SourceModuleItem != nullptr)
+		return ChildRequestCanDropAssets(TargetChild, DropRequest);
+	}
+	else if (DropRequest.DragDropOperation->IsOfType<FNiagaraParameterDragOperation>())
+	{
+		return ChildRequestCanDropParameter(TargetChild, DropRequest);
+	}
+	return TOptional<FDropRequestResponse>();
+}
+
+TOptional<UNiagaraStackEntry::FDropRequestResponse> UNiagaraStackScriptItemGroup::ChildRequestCanDropEntries(const UNiagaraStackEntry& TargetChild, const FDropRequest& DropRequest)
+{
+	TSharedRef<const FNiagaraStackEntryDragDropOp> StackEntryDragDropOp = StaticCastSharedRef<const FNiagaraStackEntryDragDropOp>(DropRequest.DragDropOperation);
+	bool ModuleEntriesDragged = false;
+	for (UNiagaraStackEntry* DraggedEntry : StackEntryDragDropOp->GetDraggedEntries())
+	{
+		if (DraggedEntry->IsA<UNiagaraStackModuleItem>())
 		{
-			if (DragOptions == EDragOptions::Move && SourceModuleItem->CanMoveAndDelete() == false)
-			{
-				return FDropResult(false, LOCTEXT("CantMoveModuleError", "This inherited module can't be moved."));
-			}
+			ModuleEntriesDragged = true;
+			break;
+		}
+	}
 
-			TArray<ENiagaraScriptUsage> SourceUsages = SourceModuleItem->GetModuleNode().FunctionScript->GetSupportedUsageContexts();
-			if (SourceUsages.ContainsByPredicate([this](ENiagaraScriptUsage SourceUsage) { return UNiagaraScript::IsEquivalentUsage(ScriptUsage, SourceUsage); }) == false)
+	if (ModuleEntriesDragged == false)
+	{
+		// Only handle dragged module items.
+		return TOptional<FDropRequestResponse>();
+	}
+	if (TargetChild.IsA<UNiagaraStackModuleItem>() == false)
+	{
+		// Only handle child drop requests from module items.
+		return TOptional<FDropRequestResponse>();
+	}
+	if (DropRequest.DropOptions != UNiagaraStackEntry::EDropOptions::Overview)
+	{
+		// Only allow dropping in the overview stacks.
+		return FDropRequestResponse(TOptional<EItemDropZone>(), LOCTEXT("CantDropOnStack", "Modules can only be dropped into the overview."));
+	}
+	if (StackEntryDragDropOp->GetDraggedEntries().Num() != 1)
+	{
+		// Only handle a single module.
+		return FDropRequestResponse(TOptional<EItemDropZone>(), LOCTEXT("CantDropMultiple", "Only single modules can be dragged and dropped."));
+	}
+
+	UNiagaraStackModuleItem* SourceModuleItem = CastChecked<UNiagaraStackModuleItem>(StackEntryDragDropOp->GetDraggedEntries()[0]);
+	if (DropRequest.DragOptions != EDragOptions::Copy && SourceModuleItem->CanMoveAndDelete() == false)
+	{
+		return FDropRequestResponse(TOptional<EItemDropZone>(), LOCTEXT("CantMoveModuleError", "This inherited module can't be moved."));
+	}
+
+	TArray<ENiagaraScriptUsage> SourceUsages = SourceModuleItem->GetModuleNode().FunctionScript->GetSupportedUsageContexts();
+	if (SourceUsages.ContainsByPredicate([this](ENiagaraScriptUsage SourceUsage) { return UNiagaraScript::IsEquivalentUsage(ScriptUsage, SourceUsage); }) == false)
+	{
+		return FDropRequestResponse(TOptional<EItemDropZone>(), LOCTEXT("CantMoveByUsage", "This module can't be moved to this section of the\nstack because it's not valid for this usage context."));
+	}
+	if (DropRequest.DropZone != EItemDropZone::AboveItem && DropRequest.DropZone != EItemDropZone::BelowItem)
+	{
+		// Only handle drops before and after items.
+		return TOptional<FDropRequestResponse>();
+	}
+
+	const UNiagaraStackModuleItem* TargetModuleItem = Cast<UNiagaraStackModuleItem>(&TargetChild);
+	if (TargetModuleItem != nullptr)
+	{
+		TArray<FNiagaraStackGraphUtilities::FStackNodeGroup> SourceStackGroups;
+		TArray<FNiagaraStackGraphUtilities::FStackNodeGroup> TargetStackGroups;
+		int32 SourceGroupIndex;
+		int32 TargetGroupIndex;
+		GenerateDragDropData(
+			SourceModuleItem->GetModuleNode(), TargetModuleItem->GetModuleNode(), DropRequest.DropZone,
+			SourceStackGroups, SourceGroupIndex,
+			TargetStackGroups, TargetGroupIndex);
+
+		// Make sure the source and target indices are within safe ranges, and make sure that the insert target isn't the source target or the spot directly
+		// after the source target since that won't actually move the module.
+		bool bSourceWithinRange = SourceGroupIndex > 0 && SourceGroupIndex < SourceStackGroups.Num() - 1;
+		bool bTargetWithinRange = TargetGroupIndex > 0 && TargetGroupIndex < TargetStackGroups.Num();
+		bool bWillMove = DropRequest.DragOptions == EDragOptions::Copy ||
+			(SourceStackGroups[SourceGroupIndex].EndNode != TargetStackGroups[TargetGroupIndex].EndNode &&
+				SourceStackGroups[SourceGroupIndex].EndNode != TargetStackGroups[TargetGroupIndex - 1].EndNode);
+		if (bSourceWithinRange && bTargetWithinRange && bWillMove)
+		{
+			FText DropMessage;
+			if (DropRequest.DragOptions == UNiagaraStackEntry::EDragOptions::Copy)
 			{
-				return FDropResult(false, LOCTEXT("CantMoveByUsage", "This module can't be moved to this section of the stack because its it's not supported in this context."));
+				DropMessage = LOCTEXT("CopyModuleResult", "Copy this module here.");
 			}
-
-			const UNiagaraStackSpacer* TargetSpacer = Cast<UNiagaraStackSpacer>(&TargetChild);
-			if (TargetSpacer != nullptr)
+			else
 			{
-				UNiagaraStackModuleItem** TargetModuleItemPtr = StackSpacerToModuleItemMap.Find(FObjectKey(TargetSpacer));
-				if (TargetModuleItemPtr != nullptr)
-				{
-					UNiagaraStackModuleItem* TargetModuleItem = *TargetModuleItemPtr;
-					TArray<FNiagaraStackGraphUtilities::FStackNodeGroup> SourceStackGroups;
-					TArray<FNiagaraStackGraphUtilities::FStackNodeGroup> TargetStackGroups;
-					int32 SourceGroupIndex;
-					int32 TargetGroupIndex;
-					GenerateDragDropData(
-						&SourceModuleItem->GetModuleNode(), TargetModuleItem != nullptr ? &TargetModuleItem->GetModuleNode() : nullptr,
-						ScriptViewModel.Pin()->GetGraphViewModel()->GetGraph(), ScriptUsage, ScriptUsageId,
-						SourceStackGroups, SourceGroupIndex,
-						TargetStackGroups, TargetGroupIndex);
-
-					// Make sure the source and target indices are within safe ranges, and make sure that the insert target isn't the source target or the spot directly
-					// after the source target since that won't actually move the module.
-					bool bSourceWithinRange = SourceGroupIndex > 0 && SourceGroupIndex < SourceStackGroups.Num() - 1;
-					bool bTargetWithinRange = TargetGroupIndex > 0 && TargetGroupIndex < TargetStackGroups.Num();
-					bool bWillMove = DragOptions == EDragOptions::Copy ||
-						(SourceStackGroups[SourceGroupIndex].EndNode != TargetStackGroups[TargetGroupIndex].EndNode &&
-						SourceStackGroups[SourceGroupIndex].EndNode != TargetStackGroups[TargetGroupIndex - 1].EndNode);
-					if(bSourceWithinRange && bTargetWithinRange && bWillMove)
-					{
-						return FDropResult(true, LOCTEXT("MoveModuleResult", "Move this module here."));
-					}
-				}
+				DropMessage = LOCTEXT("MoveModuleResult", "Move this module here.");
 			}
+			return FDropRequestResponse(DropRequest.DropZone, DropMessage);
 		}
 	}
-	return TOptional<FDropResult>();
+	return TOptional<FDropRequestResponse>();
 }
 
-TOptional<UNiagaraStackEntry::FDropResult> UNiagaraStackScriptItemGroup::ChildRequestDropInternal(const UNiagaraStackEntry& TargetChild, const TArray<UNiagaraStackEntry*>& DraggedEntries, EDragOptions DragOptions)
+TOptional<UNiagaraStackEntry::FDropRequestResponse> UNiagaraStackScriptItemGroup::ChildRequestCanDropAssets(const UNiagaraStackEntry& TargetChild, const FDropRequest& DropRequest)
 {
-	if (DragOptions != EDragOptions::Move && DragOptions != EDragOptions::Copy)
+	TSharedRef<const FAssetDragDropOp> AssetDragDropOp = StaticCastSharedRef<const FAssetDragDropOp>(DropRequest.DragDropOperation);
+	const UEnum* NiagaraScriptUsageEnum = FindObjectChecked<UEnum>(ANY_PACKAGE, TEXT("ENiagaraScriptUsage"), true);
+
+	if (TargetChild.IsA<UNiagaraStackModuleItem>() == false)
+	{
+		// Only handle child drop requests from module items.
+		return TOptional<FDropRequestResponse>();
+	}
+	if (DropRequest.DropOptions != UNiagaraStackEntry::EDropOptions::Overview)
 	{
-		return TOptional<UNiagaraStackEntry::FDropResult>();
+		// Only allow dropping in the overview stacks.
+		return FDropRequestResponse(TOptional<EItemDropZone>(), LOCTEXT("AssetCantDropOnStack", "Assets can only be dropped into the overview."));
 	}
 
-	if (bIsValidForOutput && DraggedEntries.Num() == 1 && DraggedEntries[0]->IsA<UNiagaraStackModuleItem>())
+	for (const FAssetData& AssetData : AssetDragDropOp->GetAssets())
 	{
-		UNiagaraStackModuleItem* SourceModuleItem = CastChecked<UNiagaraStackModuleItem>(DraggedEntries[0]);
-		TArray<ENiagaraScriptUsage> SourceUsages;
-		if (SourceModuleItem != nullptr)
+		if (AssetData.GetClass() != UNiagaraScript::StaticClass())
 		{
-			SourceUsages = SourceModuleItem->GetModuleNode().FunctionScript->GetSupportedUsageContexts();
+			FString AssetName;
+			AssetData.GetFullName(AssetName);
+			return FDropRequestResponse(TOptional<EItemDropZone>(), FText::Format(LOCTEXT("CantDropNonScriptAssetFormat", "Can not drop asset {0} here because it is not a niagara script."), FText::FromString(AssetName)));
 		}
 
-		if (SourceModuleItem != nullptr &&
-			SourceModuleItem->CanMoveAndDelete() &&
-			SourceUsages.ContainsByPredicate([this](ENiagaraScriptUsage SourceUsage) { return UNiagaraScript::IsEquivalentUsage(ScriptUsage, SourceUsage); }))
+		FString AssetScriptUsageString = AssetData.GetTagValueRef<FString>(GET_MEMBER_NAME_CHECKED(UNiagaraScript, Usage));
+		int64 AssetScriptUsageValue = NiagaraScriptUsageEnum->GetValueByNameString(AssetScriptUsageString);
+		if (AssetScriptUsageValue == INDEX_NONE)
 		{
-			const UNiagaraStackSpacer* TargetSpacer = Cast<UNiagaraStackSpacer>(&TargetChild);
-			if (TargetSpacer != nullptr)
-			{
-				UNiagaraStackModuleItem** TargetModuleItemPtr = StackSpacerToModuleItemMap.Find(FObjectKey(TargetSpacer));
-				if (TargetModuleItemPtr != nullptr)
-				{
-					const FNiagaraEmitterHandle* SourceEmitterHandle = SourceModuleItem->GetEmitterViewModel().IsValid()
-						? FNiagaraEditorUtilities::GetEmitterHandleForEmitter(SourceModuleItem->GetSystemViewModel()->GetSystem(), *SourceModuleItem->GetEmitterViewModel()->GetEmitter())
-						: nullptr;
-					FGuid SourceEmitterHandleId = SourceEmitterHandle != nullptr
-						? SourceEmitterHandle->GetId()
-						: FGuid();
-
-					UNiagaraNodeOutput* SourceModuleOutputNode = FNiagaraStackGraphUtilities::GetEmitterOutputNodeForStackNode(SourceModuleItem->GetModuleNode());
-					UNiagaraScript* SourceModuleScript = FNiagaraEditorUtilities::GetScriptFromSystem(SourceModuleItem->GetSystemViewModel()->GetSystem(), SourceEmitterHandleId,
-						SourceModuleOutputNode->GetUsage(), SourceModuleOutputNode->GetUsageId());
-
-					const FNiagaraEmitterHandle* TargetEmitterHandle = GetEmitterViewModel().IsValid()
-						? FNiagaraEditorUtilities::GetEmitterHandleForEmitter(GetSystemViewModel()->GetSystem(), *GetEmitterViewModel()->GetEmitter())
-						: nullptr;
-					FGuid TargetEmitterHandleId = TargetEmitterHandle != nullptr
-						? TargetEmitterHandle->GetId()
-						: FGuid();
-
-					int32 TargetIndex = *TargetModuleItemPtr != nullptr ? (*TargetModuleItemPtr)->GetModuleIndex() : INDEX_NONE;
-
-					FScopedTransaction ScopedTransaction(LOCTEXT("DragAndDropModule", "Drag and drop module"));
-					{
-						FNiagaraStackGraphUtilities::MoveModule(*SourceModuleScript, SourceModuleItem->GetModuleNode(), GetSystemViewModel()->GetSystem(), TargetEmitterHandleId,
-							ScriptUsage, ScriptUsageId, TargetIndex, DragOptions == EDragOptions::Copy);
+			FString AssetName;
+			AssetData.GetFullName(AssetName);
+			return FDropRequestResponse(TOptional<EItemDropZone>(), FText::Format(LOCTEXT("CantDropAssetWithInvalidUsageFormat", "Can not drop asset {0} here because it doesn't have a valid usage."), FText::FromString(AssetName)));
+		}
 
-						UNiagaraGraph* TargetGraph = ScriptViewModel.Pin()->GetGraphViewModel()->GetGraph();
-						FNiagaraStackGraphUtilities::RelayoutGraph(*TargetGraph);
-						TargetGraph->NotifyGraphNeedsRecompile();
+		ENiagaraScriptUsage AssetScriptUsage = static_cast<ENiagaraScriptUsage>(AssetScriptUsageValue);
+		if (AssetScriptUsage != ENiagaraScriptUsage::Module)
+		{
+			FString AssetName;
+			AssetData.GetFullName(AssetName);
+			return FDropRequestResponse(TOptional<EItemDropZone>(), FText::Format(LOCTEXT("CantDropNonModuleAssetFormat", "Can not drop asset {0} here because it is not a module script."), FText::FromString(AssetName)));
+		}
 
-						SourceModuleItem->NotifyModuleMoved();
-						RefreshChildren();
-					}
-					return FDropResult(true);
-				}
+		FString BitfieldTagValue = AssetData.GetTagValueRef<FString>(GET_MEMBER_NAME_CHECKED(UNiagaraScript, ModuleUsageBitmask));
+		int32 BitfieldValue = FCString::Atoi(*BitfieldTagValue);
+		TArray<ENiagaraScriptUsage> SupportedUsages = UNiagaraScript::GetSupportedUsageContextsForBitmask(BitfieldValue);
+		if (SupportedUsages.Contains(GetScriptUsage()) == false)
+		{
+			FString AssetName;
+			AssetData.GetFullName(AssetName);
+			return FDropRequestResponse(TOptional<EItemDropZone>(), FText::Format(LOCTEXT("CantDropByUsageFormat", "Can not drop asset {0} in this part of the stack\nbecause it's not valid for this usage context."), FText::FromString(AssetName)));
+		}
+	}
+
+	if (DropRequest.DropZone != EItemDropZone::AboveItem && DropRequest.DropZone != EItemDropZone::BelowItem)
+	{
+		// Only handle drops before and after items.
+		return TOptional<FDropRequestResponse>();
+	}
+
+	FText DropMessage;
+	if (AssetDragDropOp->GetAssets().Num() > 1)
+	{
+		DropMessage = LOCTEXT("DropAssets", "Insert modules for these assets here.");
+	}
+	else
+	{
+		DropMessage = LOCTEXT("DropAsset", "Insert a module for this asset here.");
+	}
+	return FDropRequestResponse(DropRequest.DropZone, DropMessage);
+}
+
+TOptional<UNiagaraStackEntry::FDropRequestResponse> UNiagaraStackScriptItemGroup::ChildRequestCanDropParameter(const UNiagaraStackEntry& TargetChild, const FDropRequest& DropRequest)
+{
+	TSharedRef<const FNiagaraParameterDragOperation> ParameterDragDropOp = StaticCastSharedRef<const FNiagaraParameterDragOperation>(DropRequest.DragDropOperation);
+	TSharedPtr<const FNiagaraParameterAction> ParameterAction = StaticCastSharedPtr<const FNiagaraParameterAction>(ParameterDragDropOp->GetSourceAction());
+
+	if (TargetChild.IsA<UNiagaraStackModuleItem>() == false)
+	{
+		// Only handle child drop requests from module items.
+		return TOptional<FDropRequestResponse>();
+	}
+	if (DropRequest.DropOptions != UNiagaraStackEntry::EDropOptions::Overview)
+	{
+		// Only allow dropping in the overview stacks.
+		return FDropRequestResponse(TOptional<EItemDropZone>(), LOCTEXT("AssetCantDropOnStack", "Parameters can only be dropped onto 'Set Variables' modules in the selection view."));
+	}
+	if (ParameterAction.IsValid())
+	{
+		if (FNiagaraStackGraphUtilities::ParameterIsCompatibleWithScriptUsage(ParameterAction->GetParameter(), ScriptUsage) == false)
+		{
+			return FDropRequestResponse(TOptional<EItemDropZone>(), LOCTEXT("CantDropParameterByUsage", "Can not drop this parameter here because\nit's not valid for this usage context."));
+		}
+
+		{
+			if (DropRequest.DropZone != EItemDropZone::AboveItem && DropRequest.DropZone != EItemDropZone::BelowItem)
+			{
+				// Only handle drops before and after items.
+				return TOptional<FDropRequestResponse>();
+			}
+			else
+			{
+				return FDropRequestResponse(DropRequest.DropZone, LOCTEXT("DropParameterFormat", "Insert a module for this parameter here."));
 			}
 		}
+
 	}
-	return FDropResult(false);
+	return TOptional<FDropRequestResponse>();
 }
 
-void UNiagaraStackScriptItemGroup::ItemAdded(UNiagaraNodeFunctionCall* AddedModule)
+TOptional<UNiagaraStackEntry::FDropRequestResponse> UNiagaraStackScriptItemGroup::ChildRequestDropInternal(const UNiagaraStackEntry& TargetChild, const FDropRequest& DropRequest)
 {
-	RefreshChildren();
+	if (DropRequest.DragDropOperation->IsOfType<FNiagaraStackEntryDragDropOp>())
+	{
+		return ChildRequestDropEntries(TargetChild, DropRequest);
+	}
+	else if (DropRequest.DragDropOperation->IsOfType<FAssetDragDropOp>())
+	{
+		return ChildRequestDropAssets(TargetChild, DropRequest);
+	}
+	else if (DropRequest.DragDropOperation->IsOfType<FNiagaraParameterDragOperation>())
+	{
+		return ChildRequestDropParameter(TargetChild, DropRequest);
+	}
+	return TOptional<FDropRequestResponse>();
 }
 
-void UNiagaraStackScriptItemGroup::ChildModifiedGroupItems()
+TOptional<UNiagaraStackEntry::FDropRequestResponse> UNiagaraStackScriptItemGroup::ChildRequestDropEntries(const UNiagaraStackEntry& TargetChild, const FDropRequest& DropRequest)
 {
-	RefreshChildren();
+	TSharedRef<const FNiagaraStackEntryDragDropOp> StackEntryDragDropOp = StaticCastSharedRef<const FNiagaraStackEntryDragDropOp>(DropRequest.DragDropOperation);
+
+	UNiagaraStackModuleItem* SourceModuleItem = CastChecked<UNiagaraStackModuleItem>(StackEntryDragDropOp->GetDraggedEntries()[0]);
+	const FNiagaraEmitterHandle* SourceEmitterHandle = SourceModuleItem->GetEmitterViewModel().IsValid()
+		? FNiagaraEditorUtilities::GetEmitterHandleForEmitter(SourceModuleItem->GetSystemViewModel()->GetSystem(), *SourceModuleItem->GetEmitterViewModel()->GetEmitter())
+		: nullptr;
+	FGuid SourceEmitterHandleId = SourceEmitterHandle != nullptr
+		? SourceEmitterHandle->GetId()
+		: FGuid();
+	UNiagaraNodeOutput* SourceModuleOutputNode = FNiagaraStackGraphUtilities::GetEmitterOutputNodeForStackNode(SourceModuleItem->GetModuleNode());
+	UNiagaraScript* SourceModuleScript = FNiagaraEditorUtilities::GetScriptFromSystem(SourceModuleItem->GetSystemViewModel()->GetSystem(), SourceEmitterHandleId,
+		SourceModuleOutputNode->GetUsage(), SourceModuleOutputNode->GetUsageId());
+
+	const UNiagaraStackModuleItem* TargetModuleItem = CastChecked<UNiagaraStackModuleItem>(&TargetChild);
+	const FNiagaraEmitterHandle* TargetEmitterHandle = GetEmitterViewModel().IsValid()
+		? FNiagaraEditorUtilities::GetEmitterHandleForEmitter(GetSystemViewModel()->GetSystem(), *GetEmitterViewModel()->GetEmitter())
+		: nullptr;
+	FGuid TargetEmitterHandleId = TargetEmitterHandle != nullptr
+		? TargetEmitterHandle->GetId()
+		: FGuid();
+	int32 TargetIndex = TargetModuleItem->GetModuleIndex();
+	if (DropRequest.DropZone == EItemDropZone::BelowItem)
+	{
+		TargetIndex++;
+	}
+
+	FScopedTransaction ScopedTransaction(LOCTEXT("DragAndDropModule", "Drag and drop module"));
+	{
+		GetSystemViewModel()->GetSelectionViewModel()->RemoveEntryFromSelectionByDisplayedObject(SourceModuleItem->GetDisplayedObject());
+		UNiagaraNodeFunctionCall* MovedModule = nullptr;
+		FNiagaraStackGraphUtilities::MoveModule(*SourceModuleScript, SourceModuleItem->GetModuleNode(), GetSystemViewModel()->GetSystem(), TargetEmitterHandleId,
+			ScriptUsage, ScriptUsageId, TargetIndex, DropRequest.DragOptions == EDragOptions::Copy, MovedModule);
+
+		UNiagaraGraph* TargetGraph = ScriptViewModel.Pin()->GetGraphViewModel()->GetGraph();
+		FNiagaraStackGraphUtilities::RelayoutGraph(*TargetGraph);
+		TargetGraph->NotifyGraphNeedsRecompile();
+		GetSystemViewModel()->GetSelectionViewModel()->AddEntryToSelectionByDisplayedObjectDeferred(MovedModule);
+		RefreshChildren();
+	}
+	return FDropRequestResponse(DropRequest.DropZone);
 }
 
-void UNiagaraStackScriptItemGroup::OnScriptGraphChanged(const struct FEdGraphEditAction& InAction)
+TOptional<UNiagaraStackEntry::FDropRequestResponse> UNiagaraStackScriptItemGroup::ChildRequestDropAssets(const UNiagaraStackEntry& TargetChild, const FDropRequest& DropRequest)
 {
-	if (InAction.Action == GRAPHACTION_RemoveNode)
+	TSharedRef<const FAssetDragDropOp> AssetDragDropOp = StaticCastSharedRef<const FAssetDragDropOp>(DropRequest.DragDropOperation);
+
+	const UNiagaraStackModuleItem* TargetModuleItem = CastChecked<UNiagaraStackModuleItem>(&TargetChild);
+	int32 TargetIndex = TargetModuleItem->GetModuleIndex();
+	if (DropRequest.DropZone == EItemDropZone::BelowItem)
 	{
-		OnRequestFullRefreshDeferred().Broadcast();
+		TargetIndex++;
 	}
+
+	FScopedTransaction ScopedTransaction(LOCTEXT("DragAndDropAsset", "Drag and drop assets"));
+	for (const FAssetData& AssetData : AssetDragDropOp->GetAssets())
+	{
+		TSharedRef<FScriptGroupAddAction> AddAction = FScriptGroupAddAction::CreateAssetModuleAction(AssetData);
+		AddUtilities->ExecuteAddAction(AddAction, TargetIndex);
+		TargetIndex++;
+	}
+	return FDropRequestResponse(DropRequest.DropZone);
 }
 
-void UNiagaraStackScriptItemGroup::AddParameterModuleToStack(const UNiagaraStackModuleSpacer* InModuleSpacer, const FNiagaraVariable &InVariable)
+TOptional<UNiagaraStackEntry::FDropRequestResponse> UNiagaraStackScriptItemGroup::ChildRequestDropParameter(const UNiagaraStackEntry& TargetChild, const FDropRequest& DropRequest)
 {
-	int32 TargetIndex = INDEX_NONE;
-	UNiagaraStackModuleItem** TargetModuleItemPtr = StackSpacerToModuleItemMap.Find(FObjectKey(InModuleSpacer));
-	if (*TargetModuleItemPtr != nullptr)
+	TSharedRef<const FNiagaraParameterDragOperation> ParameterDragDropOp = StaticCastSharedRef<const FNiagaraParameterDragOperation>(DropRequest.DragDropOperation);
+	TSharedPtr<FNiagaraParameterAction> ParameterAction = StaticCastSharedPtr<FNiagaraParameterAction>(ParameterDragDropOp->GetSourceAction());
+
+	const UNiagaraStackModuleItem* TargetModuleItem = CastChecked<UNiagaraStackModuleItem>(&TargetChild);
+	int32 TargetIndex = TargetModuleItem->GetModuleIndex();
+	if (DropRequest.DropZone == EItemDropZone::BelowItem)
 	{
-		UNiagaraStackModuleItem* TargetModuleItem = *TargetModuleItemPtr;
-		TargetIndex = TargetModuleItem->GetModuleIndex();
+		TargetIndex++;
 	}
 
-	TSharedRef<FScriptGroupAddAction> AddAction = FScriptGroupAddAction::CreateExistingParameterModuleAction(InVariable);
+	TSharedRef<FScriptGroupAddAction> AddAction = FScriptGroupAddAction::CreateExistingParameterModuleAction(ParameterAction->GetParameter());
 	AddUtilities->ExecuteAddAction(AddAction, TargetIndex);
+
+	return FDropRequestResponse(DropRequest.DropZone);
 }
 
-void UNiagaraStackScriptItemGroup::AddAssetModuleToStack(const UNiagaraStackModuleSpacer* InModuleSpcaer, const FAssetData& InAsset)
+void UNiagaraStackScriptItemGroup::ItemAdded(UNiagaraNodeFunctionCall* AddedModule)
 {
-	int32 TargetIndex = INDEX_NONE;
-	UNiagaraStackModuleItem** TargetModuleItemPtr = StackSpacerToModuleItemMap.Find(FObjectKey(InModuleSpcaer));
-	if (*TargetModuleItemPtr != nullptr)
-	{
-		UNiagaraStackModuleItem* TargetModuleItem = *TargetModuleItemPtr;
-		TargetIndex = TargetModuleItem->GetModuleIndex();
-	}
+	RefreshChildren();
+}
 
-	TSharedRef<FScriptGroupAddAction> AddAction = FScriptGroupAddAction::CreateAssetModuleAction(InAsset);
-	AddUtilities->ExecuteAddAction(AddAction, TargetIndex);
+void UNiagaraStackScriptItemGroup::ChildModifiedGroupItems()
+{
+	RefreshChildren();
 }
 
-bool UNiagaraStackScriptItemGroup::CanAddAssetModuleToStack(const UNiagaraStackModuleSpacer* InModuleSpacer, const FAssetData& InAsset)
+void UNiagaraStackScriptItemGroup::OnScriptGraphChanged(const struct FEdGraphEditAction& InAction)
 {
-	const UEnum* NiagaraScriptUsageEnum = FindObjectChecked<UEnum>(ANY_PACKAGE, TEXT("ENiagaraScriptUsage"), true);
-	FString AssetScriptUsageString = InAsset.GetTagValueRef<FString>(GET_MEMBER_NAME_CHECKED(UNiagaraScript, Usage));
-	int64 AssetScriptUsageValue = NiagaraScriptUsageEnum->GetValueByNameString(AssetScriptUsageString);
-	if (AssetScriptUsageValue == INDEX_NONE)
-	{
-		return false;
-	}
-	ENiagaraScriptUsage AssetScriptUsage = static_cast<ENiagaraScriptUsage>(AssetScriptUsageValue);
-	if (AssetScriptUsage != ENiagaraScriptUsage::Module)
+	if (InAction.Action == GRAPHACTION_RemoveNode)
 	{
-		return false;
+		OnRequestFullRefreshDeferred().Broadcast();
 	}
-	FString BitfieldTagValue = InAsset.GetTagValueRef<FString>(GET_MEMBER_NAME_CHECKED(UNiagaraScript, ModuleUsageBitmask));
-	int32 BitfieldValue = FCString::Atoi(*BitfieldTagValue);
-	TArray<ENiagaraScriptUsage> SupportedUsages = UNiagaraScript::GetSupportedUsageContextsForBitmask(BitfieldValue);
-	return SupportedUsages.Contains(GetScriptUsage());
 }
 
 #undef LOCTEXT_NAMESPACE
diff --git a/Engine/Plugins/FX/Niagara/Source/NiagaraEditor/Private/ViewModels/Stack/NiagaraStackSelection.cpp b/Engine/Plugins/FX/Niagara/Source/NiagaraEditor/Private/ViewModels/Stack/NiagaraStackSelection.cpp
--- a/Engine/Plugins/FX/Niagara/Source/NiagaraEditor/Private/ViewModels/Stack/NiagaraStackSelection.cpp
+++ b/Engine/Plugins/FX/Niagara/Source/NiagaraEditor/Private/ViewModels/Stack/NiagaraStackSelection.cpp
@@ -2,8 +2,6 @@
 
 #include "ViewModels/Stack/NiagaraStackSelection.h"
 #include "ViewModels/Stack/NiagaraStackEntry.h"
-#include "ViewModels/Stack/NiagaraStackSpacer.h"
-#include "ViewModels/Stack/NiagaraStackItem.h"
 #include "ViewModels/NiagaraSystemViewModel.h"
 #include "NiagaraSystemEditorData.h"
 
@@ -34,52 +32,12 @@ void UNiagaraStackSelection::SetSelectedEntries(const TArray<UNiagaraStackEntry*
 
 void UNiagaraStackSelection::RefreshChildrenInternal(const TArray<UNiagaraStackEntry*>& CurrentChildren, TArray<UNiagaraStackEntry*>& NewChildren, TArray<FStackIssue>& NewIssues)
 {
-	for (UNiagaraStackEntry* CurrentChild : CurrentChildren)
-	{
-		UNiagaraStackItem* CurrentChildItem = Cast<UNiagaraStackItem>(CurrentChild);
-		if (CurrentChildItem != nullptr)
-		{
-			CurrentChildItem->OnModifiedGroupItems().RemoveAll(this);
-		}
-	}
-
-	auto GetOrCreateSpacer = [&](FName SpacerExecutionCategory, FName SpacerKey)
-	{
-		UNiagaraStackSpacer* Spacer = FindCurrentChildOfTypeByPredicate<UNiagaraStackSpacer>(CurrentChildren,
-			[SpacerKey](UNiagaraStackSpacer* CurrentSpacer) { return CurrentSpacer->GetSpacerKey() == SpacerKey; });
-
-		if (Spacer == nullptr)
-		{
-			Spacer = NewObject<UNiagaraStackSpacer>(this);
-			FRequiredEntryData RequiredEntryData(GetSystemViewModel(), GetEmitterViewModel(),
-				SpacerExecutionCategory, NAME_None,
-				GetSystemViewModel()->GetEditorData().GetStackEditorData());
-			Spacer->Initialize(RequiredEntryData, SpacerKey, 1.0f);
-		}
-
-		return Spacer;
-	};
-
-	int32 EntryIndex = 0;
 	for (TWeakObjectPtr<UNiagaraStackEntry> SelectedEntry : SelectedEntries)
 	{
 		if (SelectedEntry.IsValid() && SelectedEntry->IsFinalized() == false)
 		{
-			if (EntryIndex > 0)
-			{
-				NewChildren.Add(GetOrCreateSpacer(NAME_None, *FString::Printf(TEXT("SelectionSpacer_%i"), EntryIndex)));
-			}
-
-			UNiagaraStackItem* SelectedItem = Cast<UNiagaraStackItem>(SelectedEntry);
-			if (SelectedItem != nullptr)
-			{
-				SelectedItem->OnModifiedGroupItems().AddUObject(this, &UNiagaraStackEntry::RefreshChildren);
-			}
-
 			SelectedEntry->SetIsExpanded(true);
 			NewChildren.Add(SelectedEntry.Get());
-
-			EntryIndex++;
 		}
 	}
 }
\ No newline at end of file
diff --git a/Engine/Plugins/FX/Niagara/Source/NiagaraEditor/Private/ViewModels/Stack/NiagaraStackSpacer.cpp b/Engine/Plugins/FX/Niagara/Source/NiagaraEditor/Private/ViewModels/Stack/NiagaraStackSpacer.cpp
deleted file mode 100644
diff --git a/Engine/Plugins/FX/Niagara/Source/NiagaraEditor/Private/ViewModels/Stack/NiagaraStackSystemSettingsGroup.cpp b/Engine/Plugins/FX/Niagara/Source/NiagaraEditor/Private/ViewModels/Stack/NiagaraStackSystemSettingsGroup.cpp
--- a/Engine/Plugins/FX/Niagara/Source/NiagaraEditor/Private/ViewModels/Stack/NiagaraStackSystemSettingsGroup.cpp
+++ b/Engine/Plugins/FX/Niagara/Source/NiagaraEditor/Private/ViewModels/Stack/NiagaraStackSystemSettingsGroup.cpp
@@ -11,7 +11,6 @@
 #include "ViewModels/Stack/NiagaraStackErrorItem.h"
 #include "ViewModels/Stack/NiagaraStackParameterStoreEntry.h"
 #include "Internationalization/Internationalization.h"
-#include "ViewModels/Stack/NiagaraStackSpacer.h"
 #include "NiagaraStackEditorData.h"
 #include "ScopedTransaction.h"
 #include "ViewModels/Stack/NiagaraStackGraphUtilities.h"
@@ -123,18 +122,6 @@ void UNiagaraStackSystemSettingsGroup::RefreshChildrenInternal(const TArray<UNia
 {
 	if (UserParameterStore != nullptr)
 	{
-		FName UserParameterStoreItemSpacerKey = "UserParameterStoreSpacer";
-		UNiagaraStackSpacer* UserParameterStoreItemSpacer = FindCurrentChildOfTypeByPredicate<UNiagaraStackSpacer>(CurrentChildren,
-			[=](UNiagaraStackSpacer* CurrentParameterStoreItemSpacer) { return CurrentParameterStoreItemSpacer->GetSpacerKey() == UserParameterStoreItemSpacerKey; });
-
-		if (UserParameterStoreItemSpacer == nullptr)
-		{
-			UserParameterStoreItemSpacer = NewObject<UNiagaraStackSpacer>(this);
-			UserParameterStoreItemSpacer->Initialize(CreateDefaultChildRequiredData(), UserParameterStoreItemSpacerKey, 1.4f);
-		}
-
-		NewChildren.Add(UserParameterStoreItemSpacer);
-
 		UNiagaraStackParameterStoreItem* UserParameterStoreItem = FindCurrentChildOfTypeByPredicate<UNiagaraStackParameterStoreItem>(CurrentChildren,
 			[=](UNiagaraStackParameterStoreItem* CurrentItem) { return true; });
 
@@ -147,17 +134,6 @@ void UNiagaraStackSystemSettingsGroup::RefreshChildrenInternal(const TArray<UNia
 		NewChildren.Add(UserParameterStoreItem);
 	}
 
-	FName SystemPropertiesItemSpacerKey = "SystemPropertiesSpacer";
-	UNiagaraStackSpacer* SystemPropertiesItemSpacer = FindCurrentChildOfTypeByPredicate<UNiagaraStackSpacer>(CurrentChildren,
-		[=](UNiagaraStackSpacer* CurrentSystemPropertiesItemSpacer) { return CurrentSystemPropertiesItemSpacer->GetSpacerKey() == SystemPropertiesItemSpacerKey; });
-
-	if (SystemPropertiesItemSpacer == nullptr)
-	{
-		SystemPropertiesItemSpacer = NewObject<UNiagaraStackSpacer>(this);
-		SystemPropertiesItemSpacer->Initialize(CreateDefaultChildRequiredData(), SystemPropertiesItemSpacerKey, 1.4f);
-	}
-	NewChildren.Add(SystemPropertiesItemSpacer);
-
 	UNiagaraStackSystemPropertiesItem* SystemPropertiesItem = FindCurrentChildOfTypeByPredicate<UNiagaraStackSystemPropertiesItem>(CurrentChildren,
 		[=](UNiagaraStackSystemPropertiesItem* CurrentItem) { return true; });
 
diff --git a/Engine/Plugins/FX/Niagara/Source/NiagaraEditor/Private/ViewModels/Stack/NiagaraStackViewModel.cpp b/Engine/Plugins/FX/Niagara/Source/NiagaraEditor/Private/ViewModels/Stack/NiagaraStackViewModel.cpp
--- a/Engine/Plugins/FX/Niagara/Source/NiagaraEditor/Private/ViewModels/Stack/NiagaraStackViewModel.cpp
+++ b/Engine/Plugins/FX/Niagara/Source/NiagaraEditor/Private/ViewModels/Stack/NiagaraStackViewModel.cpp
@@ -2,7 +2,6 @@
 
 #include "ViewModels/Stack/NiagaraStackViewModel.h"
 #include "ViewModels/Stack/NiagaraStackRoot.h"
-#include "ViewModels/Stack/NiagaraStackSpacer.h"
 #include "ViewModels/NiagaraSystemViewModel.h"
 #include "ViewModels/NiagaraEmitterViewModel.h"
 #include "ViewModels/NiagaraEmitterHandleViewModel.h"
@@ -114,6 +113,8 @@ void UNiagaraStackViewModel::InitializeWithRootEntry(UNiagaraStackEntry* InRootE
 
 	RootEntry = InRootEntry;
 	RootEntry->OnStructureChanged().AddUObject(this, &UNiagaraStackViewModel::EntryStructureChanged);
+	RootEntry->OnRequestFullRefresh().AddUObject(this, &UNiagaraStackViewModel::EntryRequestFullRefresh);
+	RootEntry->OnRequestFullRefreshDeferred().AddUObject(this, &UNiagaraStackViewModel::EntryRequestFullRefreshDeferred);
 	RootEntries.Add(RootEntry);
 
 	GEditor->RegisterForUndo(this);
@@ -640,12 +641,6 @@ void UNiagaraStackViewModel::RefreshTopLevelViewModels()
 	RootEntry->GetUnfilteredChildren(RootChildren);
 	for (UNiagaraStackEntry* RootChild : RootChildren)
 	{
-		if (RootChild->IsA<UNiagaraStackSpacer>())
-		{
-			// Skip spacers they're not a good indication of top level object.  This will be unnecessary when spacers are removed.
-			continue;
-		}
-
 		TSharedPtr<FTopLevelViewModel> TopLevelViewModel;
 		if (RootChild->GetEmitterViewModel().IsValid())
 		{
diff --git a/Engine/Plugins/FX/Niagara/Source/NiagaraEditor/Private/Widgets/SNiagaraParameterMapView.cpp b/Engine/Plugins/FX/Niagara/Source/NiagaraEditor/Private/Widgets/SNiagaraParameterMapView.cpp
--- a/Engine/Plugins/FX/Niagara/Source/NiagaraEditor/Private/Widgets/SNiagaraParameterMapView.cpp
+++ b/Engine/Plugins/FX/Niagara/Source/NiagaraEditor/Private/Widgets/SNiagaraParameterMapView.cpp
@@ -424,7 +424,10 @@ FReply SNiagaraParameterMapView::OnActionDragged(const TArray<TSharedPtr<FEdGrap
 			}
 			else if (IsSystemToolkit())
 			{
-				TSharedRef<FNiagaraStackDragOperation> DragOperation = FNiagaraStackDragOperation::New(InAction);
+				TSharedRef<FNiagaraParameterDragOperation> DragOperation = MakeShared<FNiagaraParameterDragOperation>(InAction);
+				DragOperation->CurrentHoverText = InAction->GetMenuDescription();
+				DragOperation->SetupDefaults();
+				DragOperation->Construct();
 				return FReply::Handled().BeginDragDrop(DragOperation);
 			}
 		}
diff --git a/Engine/Plugins/FX/Niagara/Source/NiagaraEditor/Public/NiagaraActions.h b/Engine/Plugins/FX/Niagara/Source/NiagaraEditor/Public/NiagaraActions.h
--- a/Engine/Plugins/FX/Niagara/Source/NiagaraEditor/Public/NiagaraActions.h
+++ b/Engine/Plugins/FX/Niagara/Source/NiagaraEditor/Public/NiagaraActions.h
@@ -7,6 +7,7 @@
 #include "NiagaraGraph.h"
 #include "EdGraph/EdGraphSchema.h"
 #include "GraphEditorDragDropAction.h"
+#include "DragAndDrop/DecoratedDragDropOp.h"
 #include "NiagaraActions.generated.h"
 
 USTRUCT()
@@ -50,7 +51,7 @@ struct NIAGARAEDITOR_API FNiagaraParameterAction : public FEdGraphSchemaAction
 		const TArray<FNiagaraGraphParameterReferenceCollection>& InReferenceCollection,
 		FText InNodeCategory, FText InMenuDesc, FText InToolTip, const int32 InGrouping, FText InKeywords, int32 InSectionID = 0);
 
-	FNiagaraVariable& GetParameter() { return Parameter; }
+	const FNiagaraVariable& GetParameter() const { return Parameter; }
 
 	FNiagaraVariable Parameter;
 
@@ -100,36 +101,18 @@ protected:
 	bool bAltDrag;
 };
 
-class NIAGARAEDITOR_API FNiagaraStackDragOperation : public FDragDropOperation
+class NIAGARAEDITOR_API FNiagaraParameterDragOperation : public FDecoratedDragDropOp
 {
 public:
-	DRAG_DROP_OPERATOR_TYPE(FNiagaraStackDragOperation, FDragDropOperation)
+	DRAG_DROP_OPERATOR_TYPE(FNiagaraParameterDragOperation, FDecoratedDragDropOp)
 
-	static TSharedRef<FNiagaraStackDragOperation> New(TSharedPtr<FEdGraphSchemaAction> InActionNode);
-
-	void Construct();
-
-	/** Set if operation is modified by alt */
-	void SetAltDrag(bool InIsAltDrag) { bAltDrag = InIsAltDrag; }
-
-	/** Set if operation is modified by the ctrl key */
-	void SetCtrlDrag(bool InIsCtrlDrag) { bControlDrag = InIsCtrlDrag; }
-
-	TSharedPtr<FEdGraphSchemaAction> GetAction() { return SourceAction; }
+	FNiagaraParameterDragOperation(TSharedPtr<FEdGraphSchemaAction> InSourceAction)
+		: SourceAction(InSourceAction)
+	{
+	}
 
-	virtual void OnDrop(bool bDropWasHandled, const FPointerEvent& MouseEvent) override;
-	virtual void OnDragged(const class FDragDropEvent& DragDropEvent) override;
-protected:
-	FNiagaraStackDragOperation();
+	TSharedPtr<FEdGraphSchemaAction> GetSourceAction() const { return SourceAction; }
 
-	bool HasFeedbackMessage();
-	void SetFeedbackMessage(const TSharedPtr<SWidget>& Message);
-	void SetSimpleFeedbackMessage(const FText& Message);
-	
+private:
 	TSharedPtr<FEdGraphSchemaAction> SourceAction;
-
-	/** Was ctrl held down at start of drag */
-	bool bControlDrag;
-	/** Was alt held down at the start of drag */
-	bool bAltDrag;
 };
\ No newline at end of file
diff --git a/Engine/Plugins/FX/Niagara/Source/NiagaraEditor/Public/ViewModels/NiagaraSystemSelectionViewModel.h b/Engine/Plugins/FX/Niagara/Source/NiagaraEditor/Public/ViewModels/NiagaraSystemSelectionViewModel.h
--- a/Engine/Plugins/FX/Niagara/Source/NiagaraEditor/Public/ViewModels/NiagaraSystemSelectionViewModel.h
+++ b/Engine/Plugins/FX/Niagara/Source/NiagaraEditor/Public/ViewModels/NiagaraSystemSelectionViewModel.h
@@ -2,6 +2,7 @@
 
 #pragma once
 
+#include "UObject/ObjectKey.h"
 #include "NiagaraSystemSelectionViewModel.generated.h"
 
 
@@ -43,13 +44,23 @@ public:
 
 	UNiagaraStackViewModel* GetSelectionStackViewModel();
 
+	void RemoveEntryFromSelectionByDisplayedObject(const UObject* InObject);
+
+	void AddEntriesToSelectionByDisplayedObjectsDeferred(const TArray<const UObject*>& InObjects);
+
+	void AddEntryToSelectionByDisplayedObjectDeferred(const UObject* InObjects);
+
 	void Refresh();
 
 	FOnSelectionChanged& OnSelectionChanged();
 
+	void Tick();
+
 private:
 	TSharedRef<FNiagaraSystemViewModel> GetSystemViewModel();
 
+	void RemoveEntryFromSelectionInternal(UNiagaraStackEntry* DeselectedEntry);
+
 	void UpdateStackSelectionEntry();
 
 private:
@@ -68,4 +79,6 @@ private:
 	UNiagaraStackViewModel* SelectionStackViewModel;
 
 	FOnSelectionChanged OnSelectionChangedDelegate;
+
+	TArray<FObjectKey> DeferredDisplayedObjectKeysToAddToSelection;
 };
\ No newline at end of file
diff --git a/Engine/Plugins/FX/Niagara/Source/NiagaraEditor/Public/ViewModels/Stack/NiagaraStackAdvancedExpander.h b/Engine/Plugins/FX/Niagara/Source/NiagaraEditor/Public/ViewModels/Stack/NiagaraStackAdvancedExpander.h
deleted file mode 100644
diff --git a/Engine/Plugins/FX/Niagara/Source/NiagaraEditor/Public/ViewModels/Stack/NiagaraStackEmitterSettingsGroup.h b/Engine/Plugins/FX/Niagara/Source/NiagaraEditor/Public/ViewModels/Stack/NiagaraStackEmitterSettingsGroup.h
--- a/Engine/Plugins/FX/Niagara/Source/NiagaraEditor/Public/ViewModels/Stack/NiagaraStackEmitterSettingsGroup.h
+++ b/Engine/Plugins/FX/Niagara/Source/NiagaraEditor/Public/ViewModels/Stack/NiagaraStackEmitterSettingsGroup.h
@@ -9,7 +9,6 @@
 class FNiagaraEmitterViewModel;
 class FNiagaraScriptViewModel;
 class UNiagaraStackObject;
-class UNiagaraStackSpacer;
 
 UCLASS()
 class NIAGARAEDITOR_API UNiagaraStackEmitterPropertiesItem : public UNiagaraStackItem
diff --git a/Engine/Plugins/FX/Niagara/Source/NiagaraEditor/Public/ViewModels/Stack/NiagaraStackEntry.h b/Engine/Plugins/FX/Niagara/Source/NiagaraEditor/Public/ViewModels/Stack/NiagaraStackEntry.h
--- a/Engine/Plugins/FX/Niagara/Source/NiagaraEditor/Public/ViewModels/Stack/NiagaraStackEntry.h
+++ b/Engine/Plugins/FX/Niagara/Source/NiagaraEditor/Public/ViewModels/Stack/NiagaraStackEntry.h
@@ -2,6 +2,8 @@
 
 #pragma once
 
+#include "DragAndDrop/DecoratedDragDropOp.h"
+#include "Widgets/Views/STableRow.h"
 #include "NiagaraStackEntry.generated.h"
 
 class FNiagaraSystemViewModel;
@@ -18,39 +20,79 @@ enum class EStackIssueSeverity : uint8
 	Info
 };
 
+class FNiagaraStackEntryDragDropOp : public FDecoratedDragDropOp
+{
+public:
+	DRAG_DROP_OPERATOR_TYPE(FNiagaraStackEntryDragDropOp, FDecoratedDragDropOp)
+
+public:
+	FNiagaraStackEntryDragDropOp(TArray<UNiagaraStackEntry*> InDraggedEntries)
+	{
+		DraggedEntries = InDraggedEntries;
+	}
+
+	const TArray<UNiagaraStackEntry*> GetDraggedEntries() const
+	{
+		return DraggedEntries;
+	}
+
+private:
+	TArray<UNiagaraStackEntry*> DraggedEntries;
+};
+
 UCLASS()
 class NIAGARAEDITOR_API UNiagaraStackEntry : public UObject
 {
 	GENERATED_BODY()
 
 public:
-	struct FDropResult
+	enum class EDragOptions
 	{
-		explicit FDropResult(bool bInCanDrop, FText InDropMessage = FText())
-			: bCanDrop(bInCanDrop)
-			, DropMessage(InDropMessage)
+		Copy,
+		None
+	};
+
+	enum class EDropOptions
+	{
+		Overview,
+		None
+	};
+
+	struct FDropRequest
+	{
+		FDropRequest(TSharedRef<const FDragDropOperation> InDragDropOperation, EItemDropZone InDropZone, EDragOptions InDragOptions, EDropOptions InDropOptions)
+			: DragDropOperation(InDragDropOperation)
+			, DropZone(InDropZone)
+			, DragOptions(InDragOptions)
+			, DropOptions(InDropOptions)
 		{
 		}
 
-		const bool bCanDrop;
-		const FText DropMessage;
+		const TSharedRef<const FDragDropOperation> DragDropOperation;
+		const EItemDropZone DropZone;
+		const EDragOptions DragOptions;
+		const EDropOptions DropOptions;
 	};
 
-	enum class EDragOptions
+	struct FDropRequestResponse
 	{
-		Copy,
-		Move,
-		None
+		FDropRequestResponse(TOptional<EItemDropZone> InDropZone, FText InDropMessage = FText())
+			: DropZone(InDropZone)
+			, DropMessage(InDropMessage)
+		{
+		}
+
+		const TOptional<EItemDropZone> DropZone;
+		const FText DropMessage;
 	};
 
 	DECLARE_MULTICAST_DELEGATE(FOnStructureChanged);
 	DECLARE_MULTICAST_DELEGATE_OneParam(FOnDataObjectModified, UObject*);
 	DECLARE_MULTICAST_DELEGATE(FOnRequestFullRefresh);
 	DECLARE_MULTICAST_DELEGATE(FOnRequestFullRefreshDeferred);
-	DECLARE_DELEGATE_RetVal_ThreeParams(TOptional<FDropResult>, FOnRequestDrop, const UNiagaraStackEntry&, const TArray<UNiagaraStackEntry*>&, EDragOptions);
+	DECLARE_DELEGATE_RetVal_TwoParams(TOptional<FDropRequestResponse>, FOnRequestDrop, const UNiagaraStackEntry& /*TargetEntry*/, const FDropRequest& /*DropRequest*/);
 	DECLARE_DELEGATE_RetVal_OneParam(bool, FOnFilterChild, const UNiagaraStackEntry&);
 	DECLARE_DELEGATE(FStackIssueFixDelegate);
-	DECLARE_DELEGATE(FOnRequestDeferredOperation);
 
 public:
 	struct NIAGARAEDITOR_API FExecutionCategoryNames
@@ -74,6 +116,7 @@ public:
 	{
 		None,
 		GroupHeader,
+		GroupFooter,
 		ItemHeader,
 		ItemContent,
 		ItemContentAdvanced,
@@ -180,6 +223,8 @@ public:
 
 	virtual FText GetDisplayName() const;
 
+	virtual UObject* GetDisplayedObject() const;
+
 	UNiagaraStackEditorData& GetStackEditorData() const;
 
 	FString GetStackEditorDataKey() const;
@@ -254,9 +299,9 @@ public:
 
 	virtual bool CanDrag() const;
 
-	TOptional<FDropResult> CanDrop(const TArray<UNiagaraStackEntry*>& DraggedEntries, EDragOptions DragOptions);
+	TOptional<FDropRequestResponse> CanDrop(const FDropRequest& DropRequest);
 
-	TOptional<FDropResult> Drop(const TArray<UNiagaraStackEntry*>& DraggedEntries, EDragOptions DragOptions);
+	TOptional<FDropRequestResponse> Drop(const FDropRequest& DropRequest);
 
 	void SetOnRequestCanDrop(FOnRequestDrop InOnRequestCanDrop);
 
@@ -279,13 +324,13 @@ protected:
 
 	virtual int32 GetChildIndentLevel() const;
 
-	virtual TOptional<FDropResult> CanDropInternal(const TArray<UNiagaraStackEntry*>& DraggedEntries, EDragOptions DragOptions);
+	virtual TOptional<FDropRequestResponse> CanDropInternal(const FDropRequest& DropRequest);
 
-	virtual TOptional<FDropResult> DropInternal(const TArray<UNiagaraStackEntry*>& DraggedEntries, EDragOptions DragOptions);
+	virtual TOptional<FDropRequestResponse> DropInternal(const FDropRequest& DropRequest);
 
-	virtual TOptional<FDropResult> ChildRequestCanDropInternal(const UNiagaraStackEntry& TargetChild, const TArray<UNiagaraStackEntry*>& DraggedEntries, EDragOptions DragOptions);
+	virtual TOptional<FDropRequestResponse> ChildRequestCanDropInternal(const UNiagaraStackEntry& TargetChild, const FDropRequest& DropRequest);
 
-	virtual TOptional<FDropResult> ChildRequestDropInternal(const UNiagaraStackEntry& TargetChild, const TArray<UNiagaraStackEntry*>& DraggedEntries, EDragOptions DragOptions);
+	virtual TOptional<FDropRequestResponse> ChildRequestDropInternal(const UNiagaraStackEntry& TargetChild, const FDropRequest& DropRequest);
 
 	virtual void ChlildStructureChangedInternal();
 
@@ -300,9 +345,9 @@ private:
 
 	void ChildRequestFullRefreshDeferred();
 
-	TOptional<FDropResult> ChildRequestCanDrop(const UNiagaraStackEntry& TargetChild, const TArray<UNiagaraStackEntry*>& DraggedEntries, EDragOptions DragOptions);
+	TOptional<FDropRequestResponse> ChildRequestCanDrop(const UNiagaraStackEntry& TargetChild, const FDropRequest& DropRequest);
 
-	TOptional<FDropResult> ChildRequestDrop(const UNiagaraStackEntry& TargetChild, const TArray<UNiagaraStackEntry*>& DraggedEntries, EDragOptions DragOptions);
+	TOptional<FDropRequestResponse> ChildRequestDrop(const UNiagaraStackEntry& TargetChild, const FDropRequest& DropRequest);
 
 	void RefreshStackErrorChildren();
 
diff --git a/Engine/Plugins/FX/Niagara/Source/NiagaraEditor/Public/ViewModels/Stack/NiagaraStackGraphUtilities.h b/Engine/Plugins/FX/Niagara/Source/NiagaraEditor/Public/ViewModels/Stack/NiagaraStackGraphUtilities.h
--- a/Engine/Plugins/FX/Niagara/Source/NiagaraEditor/Public/ViewModels/Stack/NiagaraStackGraphUtilities.h
+++ b/Engine/Plugins/FX/Niagara/Source/NiagaraEditor/Public/ViewModels/Stack/NiagaraStackGraphUtilities.h
@@ -170,7 +170,7 @@ namespace FNiagaraStackGraphUtilities
 	
 	bool GetStackIssuesRecursively(const UNiagaraStackEntry* const Entry, TArray<UNiagaraStackErrorItem*>& OutIssues);
 
-	void MoveModule(UNiagaraScript& SourceScript, UNiagaraNodeFunctionCall& ModuleToMove, UNiagaraSystem& TargetSystem, FGuid TargetEmitterHandleId, ENiagaraScriptUsage TargetUsage, FGuid TargetUsageId, int32 TargetModuleIndex, bool bForceCopy);
+	void MoveModule(UNiagaraScript& SourceScript, UNiagaraNodeFunctionCall& ModuleToMove, UNiagaraSystem& TargetSystem, FGuid TargetEmitterHandleId, ENiagaraScriptUsage TargetUsage, FGuid TargetUsageId, int32 TargetModuleIndex, bool bForceCopy, UNiagaraNodeFunctionCall*& OutMovedModue);
 
 	/** Whether a parameter is allowed to be used in a certain execution category. 
 		Used to check if parameter can be dropped on a module or funciton stack entry. */
diff --git a/Engine/Plugins/FX/Niagara/Source/NiagaraEditor/Public/ViewModels/Stack/NiagaraStackItem.h b/Engine/Plugins/FX/Niagara/Source/NiagaraEditor/Public/ViewModels/Stack/NiagaraStackItem.h
--- a/Engine/Plugins/FX/Niagara/Source/NiagaraEditor/Public/ViewModels/Stack/NiagaraStackItem.h
+++ b/Engine/Plugins/FX/Niagara/Source/NiagaraEditor/Public/ViewModels/Stack/NiagaraStackItem.h
@@ -6,8 +6,7 @@
 #include "Layout/Visibility.h"
 #include "NiagaraStackItem.generated.h"
 
-class UNiagaraStackSpacer;
-class UNiagaraStackAdvancedExpander;
+class UNiagaraStackItemFooter;
 class UNiagaraNode;
 
 UCLASS()
@@ -29,11 +28,11 @@ public:
 	EStackIssueSeverity GetHighestStackIssueSeverity() const;
 
 	virtual bool SupportsChangeEnabled() const { return false; }
-	virtual void SetIsEnabled(bool bInIsEnabled) { }
+	void SetIsEnabled(bool bInIsEnabled);
 
 	virtual bool SupportsDelete() const { return false; }
 	virtual bool TestCanDeleteWithMessage(FText& OutCanDeleteMessage) const { return false; }
-	virtual void Delete() { }
+	void Delete();
 	
 protected:
 	virtual void RefreshChildrenInternal(const TArray<UNiagaraStackEntry*>& CurrentChildren, TArray<UNiagaraStackEntry*>& NewChildren, TArray<FStackIssue>& NewIssues) override;
@@ -42,28 +41,22 @@ protected:
 
 	virtual int32 GetChildIndentLevel() const override;
 
-	virtual UNiagaraNode* GetOwningNiagaraNode() const;
-
 	virtual void ChlildStructureChangedInternal() override;
 
+	virtual void SetIsEnabledInternal(bool bInIsEnabled) { }
+	virtual void DeleteInternal() { }
+
 private:
 	bool FilterAdvancedChildren(const UNiagaraStackEntry& Child) const;
 
-	bool FilterShowAdvancedChild(const UNiagaraStackEntry& Child) const;
-
 	void ToggleShowAdvanced();
 
 protected:
 	FOnModifiedGroupItems ModifiedGroupItemsDelegate;
 
 private:
-	bool bHasAdvancedContent;
-
-	UPROPERTY()
-	UNiagaraStackSpacer* FooterSpacer;
-
 	UPROPERTY()
-	UNiagaraStackAdvancedExpander* ShowAdvancedExpander;
+	UNiagaraStackItemFooter* ItemFooter;
 
 	/** How many errors this entry has along its tree. */
 	mutable TOptional<uint32> RecursiveStackIssuesCount;
diff --git a/Engine/Plugins/FX/Niagara/Source/NiagaraEditor/Public/ViewModels/Stack/NiagaraStackItemFooter.h b/Engine/Plugins/FX/Niagara/Source/NiagaraEditor/Public/ViewModels/Stack/NiagaraStackItemFooter.h
new file mode 100644
--- /dev/null
+++ b/Engine/Plugins/FX/Niagara/Source/NiagaraEditor/Public/ViewModels/Stack/NiagaraStackItemFooter.h
@@ -0,0 +1,49 @@
+// Copyright 1998-2019 Epic Games, Inc. All Rights Reserved.
+
+#pragma once
+
+#include "ViewModels/Stack/NiagaraStackEntry.h"
+#include "NiagaraStackItemFooter.generated.h"
+
+class UNiagaraNodeFunctionCall;
+class UNiagaraStackFunctionInputCollection;
+class UNiagaraStackModuleItemOutputCollection;
+class UNiagaraNode;
+
+UCLASS()
+class NIAGARAEDITOR_API UNiagaraStackItemFooter : public UNiagaraStackEntry
+{
+	GENERATED_BODY()
+
+public:
+	DECLARE_DELEGATE(FOnToggleShowAdvanced);
+
+public:
+	void Initialize(
+		FRequiredEntryData InRequiredEntryData,
+		FString InOwnerStackItemEditorDataKey);
+
+	virtual bool GetCanExpand() const override;
+	virtual EStackRowStyle GetStackRowStyle() const override;
+
+	virtual bool GetIsEnabled() const override;
+	void SetIsEnabled(bool bInIsEnabled);
+
+	bool GetHasAdvancedContent() const;
+	void SetHasAdvancedContent(bool bHInHasAdvancedRows);
+
+	void SetOnToggleShowAdvanced(FOnToggleShowAdvanced OnToggleShowAdvanced);
+
+	bool GetShowAdvanced() const;
+
+	void ToggleShowAdvanced();
+
+private:
+	FString OwnerStackItemEditorDataKey;
+
+	FOnToggleShowAdvanced ToggleShowAdvancedDelegate;
+
+	bool bIsEnabled;
+
+	bool bHasAdvancedContent;
+};
diff --git a/Engine/Plugins/FX/Niagara/Source/NiagaraEditor/Public/ViewModels/Stack/NiagaraStackItemGroup.h b/Engine/Plugins/FX/Niagara/Source/NiagaraEditor/Public/ViewModels/Stack/NiagaraStackItemGroup.h
--- a/Engine/Plugins/FX/Niagara/Source/NiagaraEditor/Public/ViewModels/Stack/NiagaraStackItemGroup.h
+++ b/Engine/Plugins/FX/Niagara/Source/NiagaraEditor/Public/ViewModels/Stack/NiagaraStackItemGroup.h
@@ -8,6 +8,7 @@
 #include "NiagaraStackItemGroup.generated.h"
 
 class FNiagaraEmitterHandleViewModel;
+class UNiagaraStackItemGroupFooter;
 
 UCLASS()
 class NIAGARAEDITOR_API UNiagaraStackItemGroup : public UNiagaraStackEntry
@@ -42,6 +43,9 @@ protected:
 	virtual void ChlildStructureChangedInternal() override;
 
 private:
+	UPROPERTY()
+	UNiagaraStackItemGroupFooter* GroupFooter;
+
 	INiagaraStackItemGroupAddUtilities* AddUtilities;
 
 	FText GroupDisplayName;
@@ -54,3 +58,16 @@ private:
 
 	TSharedPtr<FNiagaraEmitterHandleViewModel> OwningEmitterHandleViewModel;
 };
+
+UCLASS()
+class NIAGARAEDITOR_API UNiagaraStackItemGroupFooter : public UNiagaraStackEntry
+{
+	GENERATED_BODY()
+
+public:
+	void Initialize(FRequiredEntryData InRequiredEntryData);
+
+	virtual EStackRowStyle GetStackRowStyle() const override;
+
+	virtual bool GetCanExpand() const;
+};
\ No newline at end of file
diff --git a/Engine/Plugins/FX/Niagara/Source/NiagaraEditor/Public/ViewModels/Stack/NiagaraStackModuleItem.h b/Engine/Plugins/FX/Niagara/Source/NiagaraEditor/Public/ViewModels/Stack/NiagaraStackModuleItem.h
--- a/Engine/Plugins/FX/Niagara/Source/NiagaraEditor/Public/ViewModels/Stack/NiagaraStackModuleItem.h
+++ b/Engine/Plugins/FX/Niagara/Source/NiagaraEditor/Public/ViewModels/Stack/NiagaraStackModuleItem.h
@@ -22,13 +22,12 @@ class NIAGARAEDITOR_API UNiagaraStackModuleItem : public UNiagaraStackItem
 public:
 	UNiagaraStackModuleItem();
 
-	const UNiagaraNodeFunctionCall& GetModuleNode() const;
-
-	UNiagaraNodeFunctionCall& GetModuleNode();
+	UNiagaraNodeFunctionCall& GetModuleNode() const;
 
 	void Initialize(FRequiredEntryData InRequiredEntryData, INiagaraStackItemGroupAddUtilities* GroupAddUtilities, UNiagaraNodeFunctionCall& InFunctionCallNode);
 
 	virtual FText GetDisplayName() const override;
+	virtual UObject* GetDisplayedObject() const override;
 	virtual FText GetTooltipText() const override;
 
 	INiagaraStackItemGroupAddUtilities* GetGroupAddUtilities();
@@ -39,13 +38,11 @@ public:
 
 	virtual bool SupportsChangeEnabled() const override { return true; }
 	virtual bool GetIsEnabled() const override;
-	virtual void SetIsEnabled(bool bInIsEnabled) override;
 
 	virtual bool SupportsDelete() const override { return true; }
 	virtual bool TestCanDeleteWithMessage(FText& OutCanDeleteMessage) const override;
-	virtual void Delete() override;
 
-	int32 GetModuleIndex();
+	int32 GetModuleIndex() const;
 	
 	UObject* GetExternalAsset() const override;
 
@@ -54,8 +51,6 @@ public:
 	/** Gets the output node of this module. */
 	class UNiagaraNodeOutput* GetOutputNode() const;
 
-	void NotifyModuleMoved();
-
 	bool CanAddInput(FNiagaraVariable InputParameter) const;
 
 	void AddInput(FNiagaraVariable InputParameter);
@@ -71,6 +66,11 @@ public:
 
 protected:
 	virtual void RefreshChildrenInternal(const TArray<UNiagaraStackEntry*>& CurrentChildren, TArray<UNiagaraStackEntry*>& NewChildren, TArray<FStackIssue>& NewIssues) override;
+	virtual void DeleteInternal() override;
+	virtual void SetIsEnabledInternal(bool bInIsEnabled) override;
+
+	virtual TOptional<FDropRequestResponse> CanDropInternal(const FDropRequest& DropRequest) override;
+	virtual TOptional<FDropRequestResponse> DropInternal(const FDropRequest& DropRequest) override;
 
 private:
 	bool FilterOutputCollection(const UNiagaraStackEntry& Child) const;
diff --git a/Engine/Plugins/FX/Niagara/Source/NiagaraEditor/Public/ViewModels/Stack/NiagaraStackModuleSpacer.h b/Engine/Plugins/FX/Niagara/Source/NiagaraEditor/Public/ViewModels/Stack/NiagaraStackModuleSpacer.h
deleted file mode 100644
diff --git a/Engine/Plugins/FX/Niagara/Source/NiagaraEditor/Public/ViewModels/Stack/NiagaraStackRendererItem.h b/Engine/Plugins/FX/Niagara/Source/NiagaraEditor/Public/ViewModels/Stack/NiagaraStackRendererItem.h
--- a/Engine/Plugins/FX/Niagara/Source/NiagaraEditor/Public/ViewModels/Stack/NiagaraStackRendererItem.h
+++ b/Engine/Plugins/FX/Niagara/Source/NiagaraEditor/Public/ViewModels/Stack/NiagaraStackRendererItem.h
@@ -26,7 +26,6 @@ public:
 
 	virtual bool SupportsDelete() const override { return true; }
 	virtual bool TestCanDeleteWithMessage(FText& OutCanDeleteMessage) const;
-	virtual void Delete() override;
 
 	bool HasBaseRenderer() const;
 
@@ -36,7 +35,6 @@ public:
 
 	virtual bool SupportsChangeEnabled() const override { return true; }
 	virtual bool GetIsEnabled() const override;
-	virtual void SetIsEnabled(bool bInIsEnabled) override;
 
 	static TArray<FNiagaraVariable> GetMissingVariables(UNiagaraRendererProperties* RendererProperties, UNiagaraEmitter* Emitter);
 	static bool AddMissingVariable(UNiagaraEmitter* Emitter, const FNiagaraVariable& Variable);
@@ -46,6 +44,9 @@ protected:
 
 	virtual void RefreshChildrenInternal(const TArray<UNiagaraStackEntry*>& CurrentChildren, TArray<UNiagaraStackEntry*>& NewChildren, TArray<FStackIssue>& NewIssues) override;
 
+	virtual void SetIsEnabledInternal(bool bInIsEnabled) override;
+	virtual void DeleteInternal() override;
+
 private:
 	void RendererChanged();
 	void RefreshIssues(TArray<FStackIssue>& NewIssues);
diff --git a/Engine/Plugins/FX/Niagara/Source/NiagaraEditor/Public/ViewModels/Stack/NiagaraStackScriptItemGroup.h b/Engine/Plugins/FX/Niagara/Source/NiagaraEditor/Public/ViewModels/Stack/NiagaraStackScriptItemGroup.h
--- a/Engine/Plugins/FX/Niagara/Source/NiagaraEditor/Public/ViewModels/Stack/NiagaraStackScriptItemGroup.h
+++ b/Engine/Plugins/FX/Niagara/Source/NiagaraEditor/Public/ViewModels/Stack/NiagaraStackScriptItemGroup.h
@@ -4,18 +4,14 @@
 
 #include "ViewModels/Stack/NiagaraStackItemGroup.h"
 #include "NiagaraCommon.h"
-#include "UObject/ObjectKey.h"
 #include "AssetRegistry/Public/AssetData.h"
 #include "NiagaraStackScriptItemGroup.generated.h"
 
 class FNiagaraScriptViewModel;
-class UNiagaraStackSpacer;
 class UNiagaraNodeOutput;
 class UNiagaraNodeFunctionCall;
 class FScriptItemGroupAddUtilities;
-class UNiagaraStackModuleItem;
 class UEdGraph;
-class UNiagaraStackModuleSpacer;
 
 UCLASS()
 class NIAGARAEDITOR_API UNiagaraStackScriptItemGroup : public UNiagaraStackItemGroup
@@ -35,25 +31,14 @@ public:
 	FGuid GetScriptUsageId() const { return ScriptUsageId; }
 	UNiagaraNodeOutput* GetScriptOutputNode() const;
 
-	/**
-	 * Add a Set Variables module to the stack.
-	 * @InModuleSpacer: Target UNiagaraStackModuleSpacer which we use to derive the insertion index.
-	 * @InVariable: Initial FNiagaraVariable to populate the new Set Variables module.
-	 */
-	void AddParameterModuleToStack(const UNiagaraStackModuleSpacer* InModuleSpacer, const FNiagaraVariable& InVariable);
-
-	void AddAssetModuleToStack(const UNiagaraStackModuleSpacer* InModuleSpcaer, const FAssetData& InAsset);
-
-	bool CanAddAssetModuleToStack(const UNiagaraStackModuleSpacer* InModuleSpacer, const FAssetData& InAsset);
-
 protected:
 	virtual void RefreshChildrenInternal(const TArray<UNiagaraStackEntry*>& CurrentChildren, TArray<UNiagaraStackEntry*>& NewChildren, TArray<FStackIssue>& NewIssues) override;
 
 	virtual void FinalizeInternal() override;
 
-	virtual TOptional<FDropResult> ChildRequestCanDropInternal(const UNiagaraStackEntry& TargetChild, const TArray<UNiagaraStackEntry*>& DraggedEntries, EDragOptions DragOptions) override;
+	virtual TOptional<FDropRequestResponse> ChildRequestCanDropInternal(const UNiagaraStackEntry& TargetChild, const FDropRequest& DropRequest) override;
 
-	virtual TOptional<FDropResult> ChildRequestDropInternal(const UNiagaraStackEntry& TargetChild, const TArray<UNiagaraStackEntry*>& DraggedEntries, EDragOptions DragOptions) override;
+	virtual TOptional<FDropRequestResponse> ChildRequestDropInternal(const UNiagaraStackEntry& TargetChild, const FDropRequest& DropRequest) override;
 
 private:
 	void ItemAdded(UNiagaraNodeFunctionCall* AddedModule);
@@ -62,6 +47,18 @@ private:
 
 	void OnScriptGraphChanged(const struct FEdGraphEditAction& InAction);
 
+	TOptional<FDropRequestResponse> ChildRequestCanDropEntries(const UNiagaraStackEntry& TargetChild, const FDropRequest& DropRequest);
+
+	TOptional<FDropRequestResponse> ChildRequestCanDropAssets(const UNiagaraStackEntry& TargetChild, const FDropRequest& DropRequest);
+
+	TOptional<FDropRequestResponse> ChildRequestCanDropParameter(const UNiagaraStackEntry& TargetChild, const FDropRequest& DropRequest);
+
+	TOptional<FDropRequestResponse> ChildRequestDropEntries(const UNiagaraStackEntry& TargetChild, const FDropRequest& DropRequest);
+
+	TOptional<FDropRequestResponse> ChildRequestDropAssets(const UNiagaraStackEntry& TargetChild, const FDropRequest& DropRequest);
+
+	TOptional<FDropRequestResponse> ChildRequestDropParameter(const UNiagaraStackEntry& TargetChild, const FDropRequest& DropRequest);
+
 protected:
 	TWeakPtr<FNiagaraScriptViewModel> ScriptViewModel;
 	void RefreshIssues(TArray<FStackIssue>& NewIssues);
@@ -77,6 +74,4 @@ private:
 	TWeakObjectPtr<UEdGraph> ScriptGraph;
 
 	FDelegateHandle OnGraphChangedHandle;
-
-	TMap<FObjectKey, UNiagaraStackModuleItem*> StackSpacerToModuleItemMap;
 };
diff --git a/Engine/Plugins/FX/Niagara/Source/NiagaraEditor/Public/ViewModels/Stack/NiagaraStackSpacer.h b/Engine/Plugins/FX/Niagara/Source/NiagaraEditor/Public/ViewModels/Stack/NiagaraStackSpacer.h
deleted file mode 100644
diff --git a/Engine/Plugins/FX/Niagara/Source/NiagaraEditorWidgets/Private/NiagaraEditorWidgetsUtilities.cpp b/Engine/Plugins/FX/Niagara/Source/NiagaraEditorWidgets/Private/NiagaraEditorWidgetsUtilities.cpp
--- a/Engine/Plugins/FX/Niagara/Source/NiagaraEditorWidgets/Private/NiagaraEditorWidgetsUtilities.cpp
+++ b/Engine/Plugins/FX/Niagara/Source/NiagaraEditorWidgets/Private/NiagaraEditorWidgetsUtilities.cpp
@@ -5,6 +5,8 @@
 #include "ViewModels/Stack/NiagaraStackItem.h"
 #include "ViewModels/Stack/NiagaraStackModuleItem.h"
 #include "Stack/SNiagaraStackItemGroupAddMenu.h"
+#include "NiagaraEditorWidgetsStyle.h"
+
 #include "Framework/Application/SlateApplication.h"
 #include "Framework/MultiBox/MultiBoxBuilder.h"
 #include "IContentBrowserSingleton.h"
@@ -252,4 +254,85 @@ bool FNiagaraStackEditorWidgetsUtilities::AddStackModuleItemContextMenuActions(F
 	return true;
 }
 
+TSharedRef<FDragDropOperation> FNiagaraStackEditorWidgetsUtilities::ConstructDragDropOperationForStackEntries(const TArray<UNiagaraStackEntry*>& DraggedEntries)
+{
+	TSharedRef<FNiagaraStackEntryDragDropOp> DragDropOp = MakeShared<FNiagaraStackEntryDragDropOp>(DraggedEntries);
+	DragDropOp->CurrentHoverText = DraggedEntries.Num() == 1 
+		? DraggedEntries[0]->GetDisplayName()
+		: FText::Format(LOCTEXT("MultipleEntryDragFormat", "{0} (and {1} others)"), DraggedEntries[0]->GetDisplayName(), FText::AsNumber(DraggedEntries.Num() - 1));
+	DragDropOp->CurrentIconBrush = FNiagaraEditorWidgetsStyle::Get().GetBrush(
+		GetIconNameForExecutionSubcategory(DraggedEntries[0]->GetExecutionSubcategoryName(), true));
+	DragDropOp->CurrentIconColorAndOpacity = FNiagaraEditorWidgetsStyle::Get().GetColor(
+		GetIconColorNameForExecutionCategory(DraggedEntries[0]->GetExecutionCategoryName()));
+	DragDropOp->SetupDefaults();
+	DragDropOp->Construct();
+	return DragDropOp;
+}
+
+TOptional<EItemDropZone> FNiagaraStackEditorWidgetsUtilities::RequestDropForStackEntry(const FDragDropEvent& InDragDropEvent, EItemDropZone InDropZone, UNiagaraStackEntry* InTargetEntry, UNiagaraStackEntry::EDropOptions DropOptions)
+{
+	TOptional<EItemDropZone> DropZone;
+	if (InDragDropEvent.GetOperation().IsValid())
+	{
+		TSharedPtr<FDecoratedDragDropOp> DecoratedDragDropOp = InDragDropEvent.GetOperationAs<FDecoratedDragDropOp>();
+		if (DecoratedDragDropOp.IsValid())
+		{
+			DecoratedDragDropOp->ResetToDefaultToolTip();
+		}
+
+		UNiagaraStackEntry::EDragOptions DragOptions = UNiagaraStackEntry::EDragOptions::None;
+		if (InDragDropEvent.IsAltDown() &&
+			InDragDropEvent.IsShiftDown() == false &&
+			InDragDropEvent.IsControlDown() == false &&
+			InDragDropEvent.IsCommandDown() == false)
+		{
+			DragOptions = UNiagaraStackEntry::EDragOptions::Copy;
+		}
+
+		TOptional<UNiagaraStackEntry::FDropRequestResponse> Response = InTargetEntry->CanDrop(UNiagaraStackEntry::FDropRequest(InDragDropEvent.GetOperation().ToSharedRef(), InDropZone, DragOptions, DropOptions));
+		if (Response.IsSet())
+		{
+			if (DecoratedDragDropOp.IsValid() && Response.GetValue().DropMessage.IsEmptyOrWhitespace() == false)
+			{
+				DecoratedDragDropOp->CurrentHoverText = FText::Format(LOCTEXT("DropFormat", "{0} - {1}"), DecoratedDragDropOp->GetDefaultHoverText(), Response.GetValue().DropMessage);
+			}
+
+			if (Response.GetValue().DropZone.IsSet())
+			{
+				DropZone = Response.GetValue().DropZone.GetValue();
+			}
+			else
+			{
+				if (DecoratedDragDropOp.IsValid())
+				{
+					DecoratedDragDropOp->CurrentIconBrush = FEditorStyle::GetBrush("Icons.Error");
+					DecoratedDragDropOp->CurrentIconColorAndOpacity = FLinearColor::White;
+				}
+			}
+		}
+	}
+	return DropZone;
+}
+
+bool FNiagaraStackEditorWidgetsUtilities::HandleDropForStackEntry(const FDragDropEvent& InDragDropEvent, EItemDropZone InDropZone, UNiagaraStackEntry* InTargetEntry, UNiagaraStackEntry::EDropOptions DropOptions)
+{
+	bool bHandled = false;
+	if (InDragDropEvent.GetOperation().IsValid())
+	{
+		UNiagaraStackEntry::EDragOptions DragOptions = UNiagaraStackEntry::EDragOptions::None;
+		if (InDragDropEvent.IsAltDown() &&
+			InDragDropEvent.IsShiftDown() == false &&
+			InDragDropEvent.IsControlDown() == false &&
+			InDragDropEvent.IsCommandDown() == false)
+		{
+			DragOptions = UNiagaraStackEntry::EDragOptions::Copy;
+		}
+
+		UNiagaraStackEntry::FDropRequest DropRequest(InDragDropEvent.GetOperation().ToSharedRef(), InDropZone, DragOptions, DropOptions);
+		bHandled = ensureMsgf(InTargetEntry->Drop(DropRequest).IsSet(),
+			TEXT("Failed to drop stack entry when it was requested"));
+	}
+	return bHandled;
+}
+
 #undef LOCTEXT_NAMESPACE
\ No newline at end of file
diff --git a/Engine/Plugins/FX/Niagara/Source/NiagaraEditorWidgets/Private/NiagaraEditorWidgetsUtilities.h b/Engine/Plugins/FX/Niagara/Source/NiagaraEditorWidgets/Private/NiagaraEditorWidgetsUtilities.h
--- a/Engine/Plugins/FX/Niagara/Source/NiagaraEditorWidgets/Private/NiagaraEditorWidgetsUtilities.h
+++ b/Engine/Plugins/FX/Niagara/Source/NiagaraEditorWidgets/Private/NiagaraEditorWidgetsUtilities.h
@@ -4,6 +4,7 @@
 
 #include "Templates/SharedPointer.h"
 #include "UObject/NameTypes.h"
+#include "ViewModels/Stack/NiagaraStackEntry.h"
 
 class FMenuBuilder;
 class SWidget;
@@ -25,4 +26,9 @@ namespace FNiagaraStackEditorWidgetsUtilities
 
 	bool AddStackModuleItemContextMenuActions(FMenuBuilder& MenuBuilder, UNiagaraStackModuleItem& StackItem, TSharedRef<SWidget> TargetWidget);
 
+	TSharedRef<FDragDropOperation> ConstructDragDropOperationForStackEntries(const TArray<UNiagaraStackEntry*>& DraggedEntries);
+
+	TOptional<EItemDropZone> RequestDropForStackEntry(const FDragDropEvent& InDragDropEvent, EItemDropZone InDropZone, UNiagaraStackEntry* InTargetEntry, UNiagaraStackEntry::EDropOptions DropOptions);
+
+	bool HandleDropForStackEntry(const FDragDropEvent& InDragDropEvent, EItemDropZone InDropZone, UNiagaraStackEntry* InTargetEntry, UNiagaraStackEntry::EDropOptions DropOptions);
 }
\ No newline at end of file
diff --git a/Engine/Plugins/FX/Niagara/Source/NiagaraEditorWidgets/Private/SNiagaraOverviewStack.cpp b/Engine/Plugins/FX/Niagara/Source/NiagaraEditorWidgets/Private/SNiagaraOverviewStack.cpp
--- a/Engine/Plugins/FX/Niagara/Source/NiagaraEditorWidgets/Private/SNiagaraOverviewStack.cpp
+++ b/Engine/Plugins/FX/Niagara/Source/NiagaraEditorWidgets/Private/SNiagaraOverviewStack.cpp
@@ -27,6 +27,9 @@
 class SNiagaraSystemOverviewEntryListRow : public STableRow<UNiagaraStackEntry*>
 {
 	SLATE_BEGIN_ARGS(SNiagaraSystemOverviewEntryListRow) {}
+		SLATE_EVENT(FOnDragDetected, OnDragDetected)
+		SLATE_EVENT(FOnCanAcceptDrop, OnCanAcceptDrop)
+		SLATE_EVENT(FOnAcceptDrop, OnAcceptDrop)
 		SLATE_DEFAULT_SLOT(FArguments, Content)
 	SLATE_END_ARGS();
 
@@ -72,6 +75,9 @@ class SNiagaraSystemOverviewEntryListRow : public STableRow<UNiagaraStackEntry*>
 
 		STableRow<UNiagaraStackEntry*>::Construct(STableRow<UNiagaraStackEntry*>::FArguments()
 			.Style(FNiagaraEditorWidgetsStyle::Get(), "NiagaraEditor.SystemOverview.TableViewRow")
+			.OnDragDetected(InArgs._OnDragDetected)
+			.OnCanAcceptDrop(InArgs._OnCanAcceptDrop)
+			.OnAcceptDrop(InArgs._OnAcceptDrop)
 		[
 			SNew(SBorder)
 			.BorderImage(this, &SNiagaraSystemOverviewEntryListRow::GetBorder)
@@ -183,7 +189,6 @@ private:
 
 void SNiagaraOverviewStack::Construct(const FArguments& InArgs, UNiagaraStackViewModel& InStackViewModel, UNiagaraSystemSelectionViewModel& InOverviewSelectionViewModel)
 {
-	bRefreshEntryListPending = false;
 	bUpdatingOverviewSelectionFromStackSelection = false;
 	bUpdatingStackSelectionFromOverviewSelection = false;
 
@@ -216,11 +221,7 @@ SNiagaraOverviewStack::~SNiagaraOverviewStack()
 
 void SNiagaraOverviewStack::Tick(const FGeometry& AllottedGeometry, const double InCurrentTime, const float InDeltaTime)
 {
-	if (bRefreshEntryListPending)
-	{
-		RefreshEntryList();
-		bRefreshEntryListPending = false;
-	}
+	RefreshEntryList();
 }
 
 void SNiagaraOverviewStack::AddEntriesRecursive(UNiagaraStackEntry& EntryToAdd, TArray<UNiagaraStackEntry*>& EntryList, const TArray<UClass*>& AcceptableClasses, TArray<UNiagaraStackEntry*> ParentChain)
@@ -242,24 +243,28 @@ void SNiagaraOverviewStack::AddEntriesRecursive(UNiagaraStackEntry& EntryToAdd,
 
 void SNiagaraOverviewStack::RefreshEntryList()
 {
-	FlattenedEntryList.Empty();
-	EntryObjectKeyToParentChain.Empty();
-	TArray<UClass*> AcceptableClasses;
-	AcceptableClasses.Add(UNiagaraStackItemGroup::StaticClass());
-	AcceptableClasses.Add(UNiagaraStackItem::StaticClass());
-	for (UNiagaraStackEntry* RootEntry : StackViewModel->GetRootEntries())
+	if (bRefreshEntryListPending)
 	{
-		checkf(RootEntry != nullptr, TEXT("Root entry was null."));
-		TArray<UNiagaraStackEntry*> RootChildren;
-		RootEntry->GetFilteredChildren(RootChildren);
-		for (UNiagaraStackEntry* RootChild : RootChildren)
+		FlattenedEntryList.Empty();
+		EntryObjectKeyToParentChain.Empty();
+		TArray<UClass*> AcceptableClasses;
+		AcceptableClasses.Add(UNiagaraStackItemGroup::StaticClass());
+		AcceptableClasses.Add(UNiagaraStackItem::StaticClass());
+		for (UNiagaraStackEntry* RootEntry : StackViewModel->GetRootEntries())
 		{
-			checkf(RootEntry != nullptr, TEXT("Root entry child was null."));
-			TArray<UNiagaraStackEntry*> ParentChain;
-			AddEntriesRecursive(*RootChild, FlattenedEntryList, AcceptableClasses, ParentChain);
+			checkf(RootEntry != nullptr, TEXT("Root entry was null."));
+			TArray<UNiagaraStackEntry*> RootChildren;
+			RootEntry->GetFilteredChildren(RootChildren);
+			for (UNiagaraStackEntry* RootChild : RootChildren)
+			{
+				checkf(RootEntry != nullptr, TEXT("Root entry child was null."));
+				TArray<UNiagaraStackEntry*> ParentChain;
+				AddEntriesRecursive(*RootChild, FlattenedEntryList, AcceptableClasses, ParentChain);
+			}
 		}
+		bRefreshEntryListPending = false;
+		EntryListView->RequestListRefresh();
 	}
-	EntryListView->RequestListRefresh();
 }
 
 void SNiagaraOverviewStack::EntryStructureChanged()
@@ -336,6 +341,9 @@ TSharedRef<ITableRow> SNiagaraOverviewStack::OnGenerateRowForEntry(UNiagaraStack
 
 
 	return SNew(SNiagaraSystemOverviewEntryListRow, Item, OwnerTable)
+		.OnDragDetected(this, &SNiagaraOverviewStack::OnRowDragDetected, TWeakObjectPtr<UNiagaraStackEntry>(Item))
+		.OnCanAcceptDrop(this, &SNiagaraOverviewStack::OnRowCanAcceptDrop)
+		.OnAcceptDrop(this, &SNiagaraOverviewStack::OnRowAcceptDrop)
 	[
 		Content.ToSharedRef()
 	];
@@ -381,16 +389,26 @@ void SNiagaraOverviewStack::SystemSelectionChanged(UNiagaraSystemSelectionViewMo
 	{
 		TGuardValue<bool> UpdateGuard(bUpdatingStackSelectionFromOverviewSelection, true);
 
-		TArray<UNiagaraStackEntry*> SelectedStackEntries;
-		EntryListView->GetSelectedItems(SelectedStackEntries);
+		TArray<UNiagaraStackEntry*> SelectedListViewStackEntries;
+		EntryListView->GetSelectedItems(SelectedListViewStackEntries);
 		TArray<UNiagaraStackEntry*> SelectedOverviewEntries = OverviewSelectionViewModel->GetSelectedEntries();
 
 		TArray<UNiagaraStackEntry*> EntriesToDeselect;
-		for (UNiagaraStackEntry* SelectedStackEntry : SelectedStackEntries)
+		for (UNiagaraStackEntry* SelectedListViewStackEntry : SelectedListViewStackEntries)
+		{
+			if (SelectedOverviewEntries.Contains(SelectedListViewStackEntry) == false)
+			{
+				EntriesToDeselect.Add(SelectedListViewStackEntry);
+			}
+		}
+
+		TArray<UNiagaraStackEntry*> EntriesToSelect;
+		RefreshEntryList();
+		for (UNiagaraStackEntry* SelectedOverviewEntry : SelectedOverviewEntries)
 		{
-			if (SelectedOverviewEntries.Contains(SelectedStackEntry) == false)
+			if (FlattenedEntryList.Contains(SelectedOverviewEntry))
 			{
-				EntriesToDeselect.Add(SelectedStackEntry);
+				EntriesToSelect.Add(SelectedOverviewEntry);
 			}
 		}
 
@@ -398,7 +416,33 @@ void SNiagaraOverviewStack::SystemSelectionChanged(UNiagaraSystemSelectionViewMo
 		{
 			EntryListView->SetItemSelection(EntryToDeselect, false);
 		}
+
+		for (UNiagaraStackEntry* EntryToSelect : EntriesToSelect)
+		{
+			EntryListView->SetItemSelection(EntryToSelect, true);
+		}
 	}
 }
 
+FReply SNiagaraOverviewStack::OnRowDragDetected(const FGeometry& InGeometry, const FPointerEvent& InPointerEvent, TWeakObjectPtr<UNiagaraStackEntry> InStackEntryWeak)
+{
+	UNiagaraStackEntry* StackEntry = InStackEntryWeak.Get();
+	if (StackEntry != nullptr && StackEntry->CanDrag())
+	{
+		return FReply::Handled().BeginDragDrop(FNiagaraStackEditorWidgetsUtilities::ConstructDragDropOperationForStackEntries(StackEntry->GetSystemViewModel()->GetSelectionViewModel()->GetSelectedEntries()));
+	}
+	return FReply::Unhandled();
+}
+
+TOptional<EItemDropZone> SNiagaraOverviewStack::OnRowCanAcceptDrop(const FDragDropEvent& InDragDropEvent, EItemDropZone InDropZone, UNiagaraStackEntry* InTargetEntry)
+{
+	return FNiagaraStackEditorWidgetsUtilities::RequestDropForStackEntry(InDragDropEvent, InDropZone, InTargetEntry, UNiagaraStackEntry::EDropOptions::Overview);
+}
+
+FReply SNiagaraOverviewStack::OnRowAcceptDrop(const FDragDropEvent& InDragDropEvent, EItemDropZone InDropZone, UNiagaraStackEntry* InTargetEntry)
+{
+	bool bHandled = FNiagaraStackEditorWidgetsUtilities::HandleDropForStackEntry(InDragDropEvent, InDropZone, InTargetEntry, UNiagaraStackEntry::EDropOptions::Overview);
+	return bHandled ? FReply::Handled() : FReply::Unhandled();
+}
+
 #undef LOCTEXT_NAMESPACE
\ No newline at end of file
diff --git a/Engine/Plugins/FX/Niagara/Source/NiagaraEditorWidgets/Private/SNiagaraOverviewStack.h b/Engine/Plugins/FX/Niagara/Source/NiagaraEditorWidgets/Private/SNiagaraOverviewStack.h
--- a/Engine/Plugins/FX/Niagara/Source/NiagaraEditorWidgets/Private/SNiagaraOverviewStack.h
+++ b/Engine/Plugins/FX/Niagara/Source/NiagaraEditorWidgets/Private/SNiagaraOverviewStack.h
@@ -41,6 +41,12 @@ private:
 
 	void SystemSelectionChanged(UNiagaraSystemSelectionViewModel::ESelectionChangeSource SelectionChangeSource);
 
+	FReply OnRowDragDetected(const FGeometry& InGeometry, const FPointerEvent& InPointerEvent, TWeakObjectPtr<UNiagaraStackEntry> InStackEntryWeak);
+
+	TOptional<EItemDropZone> OnRowCanAcceptDrop(const FDragDropEvent& InDragDropEvent, EItemDropZone InDropZone, UNiagaraStackEntry* InTargetEntry);
+
+	FReply OnRowAcceptDrop(const FDragDropEvent& InDragDropEvent, EItemDropZone InDropZone, UNiagaraStackEntry* InTargetEntry);
+
 private:
 	UNiagaraStackViewModel* StackViewModel;
 	UNiagaraSystemSelectionViewModel* OverviewSelectionViewModel;
diff --git a/Engine/Plugins/FX/Niagara/Source/NiagaraEditorWidgets/Private/SNiagaraStack.cpp b/Engine/Plugins/FX/Niagara/Source/NiagaraEditorWidgets/Private/SNiagaraStack.cpp
--- a/Engine/Plugins/FX/Niagara/Source/NiagaraEditorWidgets/Private/SNiagaraStack.cpp
+++ b/Engine/Plugins/FX/Niagara/Source/NiagaraEditorWidgets/Private/SNiagaraStack.cpp
@@ -12,14 +12,13 @@
 #include "ViewModels/NiagaraEmitterHandleViewModel.h"
 #include "ViewModels/NiagaraSystemViewModel.h"
 #include "ViewModels/Stack/NiagaraStackViewModel.h"
-#include "ViewModels/Stack/NiagaraStackSpacer.h"
 #include "ViewModels/Stack/NiagaraStackModuleItemOutputCollection.h"
 #include "ViewModels/Stack/NiagaraStackModuleItemLinkedInputCollection.h"
 #include "ViewModels/Stack/NiagaraStackModuleItemOutput.h"
 #include "ViewModels/Stack/NiagaraStackErrorItem.h"
 #include "ViewModels/Stack/NiagaraStackInputCategory.h" 
 #include "ViewModels/Stack/NiagaraStackPropertyRow.h"
-#include "ViewModels/Stack/NiagaraStackAdvancedExpander.h"
+#include "ViewModels/Stack/NiagaraStackItemFooter.h"
 #include "ViewModels/Stack/NiagaraStackFunctionInputCollection.h"
 #include "ViewModels/Stack/NiagaraStackModuleItem.h"
 #include "ViewModels/Stack/NiagaraStackRendererItem.h"
@@ -40,7 +39,7 @@
 #include "Stack/SNiagaraStackFunctionInputValue.h"
 #include "Stack/SNiagaraStackEmitterPropertiesItem.h"
 #include "Stack/SNiagaraStackEventHandlerPropertiesItem.h"
-#include "Stack/SNiagaraStackItemExpander.h"
+#include "Stack/SNiagaraStackItemFooter.h"
 #include "Stack/SNiagaraStackItemGroup.h"
 #include "Stack/SNiagaraStackModuleItem.h"
 #include "Stack/SNiagaraStackParameterStoreEntryName.h"
@@ -53,7 +52,6 @@
 #include "NiagaraStackEditorData.h"
 #include "ScopedTransaction.h"
 #include "Widgets/Layout/SWrapBox.h"
-#include "Stack/SNiagaraStackSpacer.h"
 #include "ViewModels/Stack/NiagaraStackRoot.h"
 #include "Subsystems/AssetEditorSubsystem.h"
 
@@ -215,27 +213,6 @@ private:
 	TSharedPtr<SInlineEditableTextBlock> EmitterNameTextBlock;
 };
 
-/** Contains data for a socket drag and drop operation in the StackEntry node. */
-class FNiagaraStackEntryDragDropOp : public FDecoratedDragDropOp
-{
-public:
-	DRAG_DROP_OPERATOR_TYPE(FNiagaraStackEntryDragDropOp, FDecoratedDragDropOp)
-
-public:
-	FNiagaraStackEntryDragDropOp(TArray<UNiagaraStackEntry*> InDraggedEntries)
-	{
-		DraggedEntries = InDraggedEntries;
-	}
-
-	const TArray<UNiagaraStackEntry*> GetDraggedEntries() const
-	{
-		return DraggedEntries;
-	}
-
-private:
-	TArray<UNiagaraStackEntry*> DraggedEntries;
-};
-
 const float SpacerHeight = 6;
 
 const FText SNiagaraStack::OccurencesFormat = NSLOCTEXT("NiagaraStack", "OccurencesFound", "{0} / {1}");
@@ -547,92 +524,20 @@ FReply SNiagaraStack::OnRowDragDetected(const FGeometry& InGeometry, const FPoin
 	{
 		TArray<UNiagaraStackEntry*> DraggedEntries;
 		DraggedEntries.Add(InStackEntry);
-		TSharedRef<FNiagaraStackEntryDragDropOp> DragDropOp = MakeShared<FNiagaraStackEntryDragDropOp>(DraggedEntries);
-		DragDropOp->CurrentHoverText = InStackEntry->GetDisplayName();
-		DragDropOp->CurrentIconBrush = FNiagaraEditorWidgetsStyle::Get().GetBrush(
-			FNiagaraStackEditorWidgetsUtilities::GetIconNameForExecutionSubcategory(InStackEntry->GetExecutionSubcategoryName(), true));
-		DragDropOp->CurrentIconColorAndOpacity = FNiagaraEditorWidgetsStyle::Get().GetColor(
-			FNiagaraStackEditorWidgetsUtilities::GetIconColorNameForExecutionCategory(InStackEntry->GetExecutionCategoryName()));
-		DragDropOp->SetupDefaults();
-		DragDropOp->Construct();
-		return FReply::Handled().BeginDragDrop(DragDropOp);
+		return FReply::Handled().BeginDragDrop(FNiagaraStackEditorWidgetsUtilities::ConstructDragDropOperationForStackEntries(DraggedEntries));
 	}
 	return FReply::Unhandled();
 }
 
 TOptional<EItemDropZone> SNiagaraStack::OnRowCanAcceptDrop(const FDragDropEvent& InDragDropEvent, EItemDropZone InDropZone, UNiagaraStackEntry* InTargetEntry)
 {
-	TOptional<EItemDropZone> DropZone;
-	TSharedPtr<FNiagaraStackEntryDragDropOp> DragDropOp = InDragDropEvent.GetOperationAs<FNiagaraStackEntryDragDropOp>();
-	if (DragDropOp.IsValid())
-	{
-		DragDropOp->ResetToDefaultToolTip();
-
-		UNiagaraStackEntry::EDragOptions DragOptions = UNiagaraStackEntry::EDragOptions::None;
-		if (InDragDropEvent.IsAltDown() &&
-			InDragDropEvent.IsShiftDown() == false &&
-			InDragDropEvent.IsControlDown() == false &&
-			InDragDropEvent.IsCommandDown() == false)
-		{
-			DragOptions = UNiagaraStackEntry::EDragOptions::Copy;
-		}
-		else if (InDragDropEvent.IsAltDown() == false &&
-			InDragDropEvent.IsShiftDown() == false &&
-			InDragDropEvent.IsControlDown() == false &&
-			InDragDropEvent.IsCommandDown() == false)
-		{
-			DragOptions = UNiagaraStackEntry::EDragOptions::Move;
-		}
-		
-		TOptional<UNiagaraStackEntry::FDropResult> Result = InTargetEntry->CanDrop(DragDropOp->GetDraggedEntries(), DragOptions);
-		if (Result.IsSet())
-		{
-			if (Result.GetValue().DropMessage.IsEmptyOrWhitespace() == false)
-			{
-				DragDropOp->CurrentHoverText = FText::Format(LOCTEXT("DropFormat", "{0} - {1}"), DragDropOp->GetDefaultHoverText(), Result.GetValue().DropMessage);
-			}
-
-			if (Result.GetValue().bCanDrop)
-			{
-				DropZone = EItemDropZone::OntoItem;
-			}
-			else 
-			{
-				DragDropOp->CurrentIconBrush = FEditorStyle::GetBrush("Icons.Error");
-				DragDropOp->CurrentIconColorAndOpacity = FLinearColor::White;
-			}
-		}
-	}
-	return DropZone;
+	return FNiagaraStackEditorWidgetsUtilities::RequestDropForStackEntry(InDragDropEvent, InDropZone, InTargetEntry, UNiagaraStackEntry::EDropOptions::None);
 }
 
 FReply SNiagaraStack::OnRowAcceptDrop(const FDragDropEvent& InDragDropEvent, EItemDropZone InDropZone, UNiagaraStackEntry* InTargetEntry)
 {
-	TSharedPtr<FNiagaraStackEntryDragDropOp> DragDropOp = InDragDropEvent.GetOperationAs<FNiagaraStackEntryDragDropOp>();
-	if (DragDropOp.IsValid())
-	{
-		UNiagaraStackEntry::EDragOptions DragOptions = UNiagaraStackEntry::EDragOptions::None;
-		if (InDragDropEvent.IsAltDown() &&
-			InDragDropEvent.IsShiftDown() == false &&
-			InDragDropEvent.IsControlDown() == false &&
-			InDragDropEvent.IsCommandDown() == false)
-		{
-			DragOptions = UNiagaraStackEntry::EDragOptions::Copy;
-		}
-		else if (InDragDropEvent.IsAltDown() == false &&
-			InDragDropEvent.IsShiftDown() == false &&
-			InDragDropEvent.IsControlDown() == false &&
-			InDragDropEvent.IsCommandDown() == false)
-		{
-			DragOptions = UNiagaraStackEntry::EDragOptions::Move;
-		}
-
-		if (ensureMsgf(InTargetEntry->Drop(DragDropOp->GetDraggedEntries(), DragOptions).IsSet(), TEXT("Failed to drop stack entry when can drop returned true")))
-		{
-			return FReply::Handled();
-		}
-	}
-	return FReply::Unhandled();
+	bool bHandled = FNiagaraStackEditorWidgetsUtilities::HandleDropForStackEntry(InDragDropEvent, InDropZone, InTargetEntry, UNiagaraStackEntry::EDropOptions::None);
+	return bHandled ? FReply::Handled() : FReply::Unhandled();
 }
 
 void SNiagaraStack::OnStackSearchComplete()
@@ -806,6 +711,7 @@ TSharedRef<SNiagaraStackTableRow> SNiagaraStack::ConstructContainerForItem(UNiag
 {
 	float LeftContentPadding = 4;
 	float RightContentPadding = 6;
+	TAttribute<FMargin> RowPadding = FMargin(0, 0, 0, 0);
 	FMargin ContentPadding(LeftContentPadding, 0, RightContentPadding, 0);
 	FLinearColor ItemBackgroundColor;
 	FLinearColor ItemForegroundColor = FNiagaraEditorWidgetsStyle::Get().GetColor("NiagaraEditor.Stack.ForegroundColor");
@@ -819,13 +725,30 @@ TSharedRef<SNiagaraStackTableRow> SNiagaraStack::ConstructContainerForItem(UNiag
 		bShowExecutionCategoryIcon = false;
 		break;
 	case UNiagaraStackEntry::EStackRowStyle::GroupHeader:
-		ContentPadding = FMargin(LeftContentPadding, 2, 0, 2);
+		ContentPadding = FMargin(LeftContentPadding, 3, 0, 3);
 		ItemBackgroundColor = FLinearColor::Transparent;
 		ItemForegroundColor = FNiagaraEditorWidgetsStyle::Get().GetColor("NiagaraEditor.Stack.GroupForegroundColor");
 		bIsCategoryIconHighlighted = true;
 		bShowExecutionCategoryIcon = true;
 		break;
+	case UNiagaraStackEntry::EStackRowStyle::GroupFooter:
+		RowPadding = FMargin(0, 0, 0, 6);
+		ItemBackgroundColor = FLinearColor::Transparent;
+		bIsCategoryIconHighlighted = false;
+		bShowExecutionCategoryIcon = false;
+		break;
 	case UNiagaraStackEntry::EStackRowStyle::ItemHeader:
+		RowPadding.Bind(TAttribute<FMargin>::FGetter::CreateLambda([Item]()
+		{
+			if (Item->GetIsExpanded())
+			{
+				return FMargin(0, 2, 0, 0);
+			}
+			else
+			{
+				return FMargin(0, 2, 0, 2);
+			}
+		}));
 		ContentPadding = FMargin(LeftContentPadding, 2, 2, 2);
 		ItemBackgroundColor = FNiagaraEditorWidgetsStyle::Get().GetColor("NiagaraEditor.Stack.Item.HeaderBackgroundColor");
 		bIsCategoryIconHighlighted = false;
@@ -844,6 +767,7 @@ TSharedRef<SNiagaraStackTableRow> SNiagaraStack::ConstructContainerForItem(UNiag
 		bShowExecutionCategoryIcon = false;
 		break;
 	case UNiagaraStackEntry::EStackRowStyle::ItemFooter:
+		RowPadding = FMargin(0, 0, 0, 2);
 		ItemBackgroundColor = FNiagaraEditorWidgetsStyle::Get().GetColor("NiagaraEditor.Stack.Item.FooterBackgroundColor");
 		bIsCategoryIconHighlighted = false;
 		bShowExecutionCategoryIcon = false;
@@ -868,6 +792,7 @@ TSharedRef<SNiagaraStackTableRow> SNiagaraStack::ConstructContainerForItem(UNiag
 	}
 
 	return SNew(SNiagaraStackTableRow, StackViewModel, Item, StackTree.ToSharedRef())
+		.RowPadding(RowPadding)
 		.ContentPadding(ContentPadding)
 		.ItemBackgroundColor(ItemBackgroundColor)
 		.ItemForegroundColor(ItemForegroundColor)
@@ -895,15 +820,7 @@ void SNiagaraStack::Tick(const FGeometry& AllottedGeometry, const double InCurre
 
 SNiagaraStack::FRowWidgets SNiagaraStack::ConstructNameAndValueWidgetsForItem(UNiagaraStackEntry* Item, TSharedRef<SNiagaraStackTableRow> Container)
 {
-	if (Item->IsA<UNiagaraStackSpacer>())
-	{
-		FMargin ContentPadding = Container->GetContentPadding();
-		Container->SetContentPadding(FMargin(ContentPadding.Left, 0, ContentPadding.Right, 0));
-		UNiagaraStackSpacer* SpacerItem = CastChecked<UNiagaraStackSpacer>(Item); 
-		return FRowWidgets(SNew(SNiagaraStackSpacer, *SpacerItem)
-			.HeightOverride(SpacerHeight * SpacerItem->GetSpacerScale()));
-	}
-	else if (Item->IsA<UNiagaraStackItemGroup>())
+	if (Item->IsA<UNiagaraStackItemGroup>())
 	{
 		return FRowWidgets(SNew(SNiagaraStackItemGroup, *CastChecked<UNiagaraStackItemGroup>(Item), StackViewModel));
 	}
@@ -943,10 +860,10 @@ SNiagaraStack::FRowWidgets SNiagaraStack::ConstructNameAndValueWidgetsForItem(UN
 	{
 		return FRowWidgets(SNew(SNiagaraStackErrorItemFix, CastChecked<UNiagaraStackErrorItemFix>(Item), StackViewModel));
 	}
-	else if (Item->IsA<UNiagaraStackAdvancedExpander>())
+	else if (Item->IsA<UNiagaraStackItemFooter>())
 	{
-		UNiagaraStackAdvancedExpander* ItemExpander = CastChecked<UNiagaraStackAdvancedExpander>(Item);
-		return FRowWidgets(SNew(SNiagaraStackItemExpander, *ItemExpander));
+		UNiagaraStackItemFooter* ItemExpander = CastChecked<UNiagaraStackItemFooter>(Item);
+		return FRowWidgets(SNew(SNiagaraStackItemFooter, *ItemExpander));
 	}
 	else if (Item->IsA<UNiagaraStackEmitterPropertiesItem>())
 	{
@@ -1043,12 +960,7 @@ SNiagaraStack::FRowWidgets SNiagaraStack::ConstructNameAndValueWidgetsForItem(UN
 	}
 	else
 	{
-		return FRowWidgets(SNew(STextBlock)
-			.TextStyle(FNiagaraEditorWidgetsStyle::Get(), "NiagaraEditor.Stack.DefaultText")
-			.ToolTipText_UObject(Item, &UNiagaraStackEntry::GetTooltipText)
-			.Text_UObject(Item, &UNiagaraStackEntry::GetDisplayName)
-			.ColorAndOpacity(this, &SNiagaraStack::GetTextColorForItem, Item)
-			.HighlightText_UObject(StackViewModel, &UNiagaraStackViewModel::GetCurrentSearchText));
+		return FRowWidgets(SNullWidget::NullWidget);
 	}
 }
 
diff --git a/Engine/Plugins/FX/Niagara/Source/NiagaraEditorWidgets/Private/SNiagaraStackEntryWidget.cpp b/Engine/Plugins/FX/Niagara/Source/NiagaraEditorWidgets/Private/SNiagaraStackEntryWidget.cpp
--- a/Engine/Plugins/FX/Niagara/Source/NiagaraEditorWidgets/Private/SNiagaraStackEntryWidget.cpp
+++ b/Engine/Plugins/FX/Niagara/Source/NiagaraEditorWidgets/Private/SNiagaraStackEntryWidget.cpp
@@ -27,6 +27,11 @@ void SNiagaraStackDisplayName::Construct(const FArguments& InArgs, UNiagaraStack
 	];
 }
 
+SNiagaraStackDisplayName::~SNiagaraStackDisplayName()
+{
+	StackViewModel->OnStructureChanged().RemoveAll(this);
+}
+
 TSharedRef<SWidget> SNiagaraStackDisplayName::ConstructChildren()
 {
 	TSharedRef<STextBlock> BaseNameWidget = SNew(STextBlock)
@@ -100,7 +105,7 @@ FText SNiagaraStackDisplayName::GetTopLevelDisplayName(TWeakPtr<UNiagaraStackVie
 
 void SNiagaraStackDisplayName::StackViewModelStructureChanged()
 {
-	if (StackViewModel->GetTopLevelViewModels().Num() != TopLevelViewModelCountAtLastConstruction)
+	if (StackEntry->IsFinalized() == false && StackViewModel->GetTopLevelViewModels().Num() != TopLevelViewModelCountAtLastConstruction)
 	{
 		Container->SetContent(ConstructChildren());
 	}
diff --git a/Engine/Plugins/FX/Niagara/Source/NiagaraEditorWidgets/Private/SNiagaraStackEntryWidget.h b/Engine/Plugins/FX/Niagara/Source/NiagaraEditorWidgets/Private/SNiagaraStackEntryWidget.h
--- a/Engine/Plugins/FX/Niagara/Source/NiagaraEditorWidgets/Private/SNiagaraStackEntryWidget.h
+++ b/Engine/Plugins/FX/Niagara/Source/NiagaraEditorWidgets/Private/SNiagaraStackEntryWidget.h
@@ -19,6 +19,8 @@ public:
 
 	void Construct(const FArguments& InArgs, UNiagaraStackEntry& InStackEntry, UNiagaraStackViewModel& InStackViewModel, FName InTextStyleName);
 
+	~SNiagaraStackDisplayName();
+
 private:
 	TSharedRef<SWidget> ConstructChildren();
 
diff --git a/Engine/Plugins/FX/Niagara/Source/NiagaraEditorWidgets/Private/Stack/SNiagaraStackFunctionInputValue.cpp b/Engine/Plugins/FX/Niagara/Source/NiagaraEditorWidgets/Private/Stack/SNiagaraStackFunctionInputValue.cpp
--- a/Engine/Plugins/FX/Niagara/Source/NiagaraEditorWidgets/Private/Stack/SNiagaraStackFunctionInputValue.cpp
+++ b/Engine/Plugins/FX/Niagara/Source/NiagaraEditorWidgets/Private/Stack/SNiagaraStackFunctionInputValue.cpp
@@ -832,10 +832,10 @@ FSlateColor SNiagaraStackFunctionInputValue::GetInputIconColor() const
 
 FReply SNiagaraStackFunctionInputValue::OnFunctionInputDrop(TSharedPtr<FDragDropOperation> DragDropOperation)
 {
-	if (DragDropOperation->IsOfType<FNiagaraStackDragOperation>())
+	if (DragDropOperation->IsOfType<FNiagaraParameterDragOperation>())
 	{
-		TSharedPtr<FNiagaraStackDragOperation> InputDragDropOperation = StaticCastSharedPtr<FNiagaraStackDragOperation>(DragDropOperation);
-		TSharedPtr<FNiagaraParameterAction> Action = StaticCastSharedPtr<FNiagaraParameterAction>(InputDragDropOperation->GetAction());
+		TSharedPtr<FNiagaraParameterDragOperation> InputDragDropOperation = StaticCastSharedPtr<FNiagaraParameterDragOperation>(DragDropOperation);
+		TSharedPtr<FNiagaraParameterAction> Action = StaticCastSharedPtr<FNiagaraParameterAction>(InputDragDropOperation->GetSourceAction());
 		if (Action.IsValid())
 		{
 			FunctionInput->SetLinkedValueHandle(FNiagaraParameterHandle(Action->GetParameter().GetName()));
@@ -848,10 +848,10 @@ FReply SNiagaraStackFunctionInputValue::OnFunctionInputDrop(TSharedPtr<FDragDrop
 
 bool SNiagaraStackFunctionInputValue::OnFunctionInputAllowDrop(TSharedPtr<FDragDropOperation> DragDropOperation)
 {
-	if (FunctionInput && DragDropOperation->IsOfType<FNiagaraStackDragOperation>())
+	if (FunctionInput && DragDropOperation->IsOfType<FNiagaraParameterDragOperation>())
 	{
-		TSharedPtr<FNiagaraStackDragOperation> InputDragDropOperation = StaticCastSharedPtr<FNiagaraStackDragOperation>(DragDropOperation);
-		TSharedPtr<FNiagaraParameterAction> Action = StaticCastSharedPtr<FNiagaraParameterAction>(InputDragDropOperation->GetAction());
+		TSharedPtr<FNiagaraParameterDragOperation> InputDragDropOperation = StaticCastSharedPtr<FNiagaraParameterDragOperation>(DragDropOperation);
+		TSharedPtr<FNiagaraParameterAction> Action = StaticCastSharedPtr<FNiagaraParameterAction>(InputDragDropOperation->GetSourceAction());
 		if (Action->GetParameter().GetType() == FunctionInput->GetInputType() && FNiagaraStackGraphUtilities::ParameterAllowedInExecutionCategory(Action->GetParameter().GetName(), FunctionInput->GetExecutionCategoryName()))
 		{
 			return true;
diff --git a/Engine/Plugins/FX/Niagara/Source/NiagaraEditorWidgets/Private/Stack/SNiagaraStackItemExpander.cpp b/Engine/Plugins/FX/Niagara/Source/NiagaraEditorWidgets/Private/Stack/SNiagaraStackItemExpander.cpp
deleted file mode 100644
diff --git a/Engine/Plugins/FX/Niagara/Source/NiagaraEditorWidgets/Private/Stack/SNiagaraStackItemExpander.h b/Engine/Plugins/FX/Niagara/Source/NiagaraEditorWidgets/Private/Stack/SNiagaraStackItemExpander.h
deleted file mode 100644
diff --git a/Engine/Plugins/FX/Niagara/Source/NiagaraEditorWidgets/Private/Stack/SNiagaraStackItemFooter.cpp b/Engine/Plugins/FX/Niagara/Source/NiagaraEditorWidgets/Private/Stack/SNiagaraStackItemFooter.cpp
new file mode 100644
--- /dev/null
+++ b/Engine/Plugins/FX/Niagara/Source/NiagaraEditorWidgets/Private/Stack/SNiagaraStackItemFooter.cpp
@@ -0,0 +1,67 @@
+// Copyright 1998-2019 Epic Games, Inc. All Rights Reserved.
+
+#include "Stack/SNiagaraStackItemFooter.h"
+#include "ViewModels/Stack/NiagaraStackItemFooter.h"
+#include "EditorStyleSet.h"
+#include "Widgets/Input/SButton.h"
+#include "Widgets/Images/SImage.h"
+
+#define LOCTEXT_NAMESPACE "NiagaraStackItemExpander"
+
+void SNiagaraStackItemFooter::Construct(const FArguments& InArgs, UNiagaraStackItemFooter& InItemFooter)
+{
+	ItemFooter = &InItemFooter;
+	ExpandedToolTipText = LOCTEXT("HideAdvancedToolTip", "Hide Advanced");
+	CollapsedToolTipText = LOCTEXT("ShowAdvancedToolTip", "Show Advanced");
+
+	ChildSlot
+	[
+		SNew(SButton)
+		.ButtonStyle(FEditorStyle::Get(), "NoBorder")
+		.Visibility(this, &SNiagaraStackItemFooter::GetExpandButtonVisibility)
+		.HAlign(HAlign_Center)
+		.ContentPadding(2)
+		.ToolTipText(this, &SNiagaraStackItemFooter::GetToolTipText)
+		.OnClicked(this, &SNiagaraStackItemFooter::ExpandButtonClicked)
+		.IsFocusable(false)
+		.Content()
+		[
+			SNew(SImage)
+			.Image(this, &SNiagaraStackItemFooter::GetButtonBrush)
+		]
+	];
+}
+
+EVisibility SNiagaraStackItemFooter::GetExpandButtonVisibility() const
+{
+	return ItemFooter->GetHasAdvancedContent() ? EVisibility::Visible : EVisibility::Hidden;
+}
+
+const FSlateBrush* SNiagaraStackItemFooter::GetButtonBrush() const
+{
+	if (IsHovered())
+	{
+		return ItemFooter->GetShowAdvanced()
+			? FEditorStyle::GetBrush("DetailsView.PulldownArrow.Up.Hovered")
+			: FEditorStyle::GetBrush("DetailsView.PulldownArrow.Down.Hovered");
+	}
+	else
+	{
+		return ItemFooter->GetShowAdvanced()
+			? FEditorStyle::GetBrush("DetailsView.PulldownArrow.Up")
+			: FEditorStyle::GetBrush("DetailsView.PulldownArrow.Down");
+	}
+}
+
+FText SNiagaraStackItemFooter::GetToolTipText() const
+{
+	return ItemFooter->GetShowAdvanced() ? ExpandedToolTipText : CollapsedToolTipText;
+}
+
+FReply SNiagaraStackItemFooter::ExpandButtonClicked()
+{
+	ItemFooter->ToggleShowAdvanced();
+	return FReply::Handled();
+}
+
+#undef LOCTEXT_NAMESPACE
\ No newline at end of file
diff --git a/Engine/Plugins/FX/Niagara/Source/NiagaraEditorWidgets/Private/Stack/SNiagaraStackItemFooter.h b/Engine/Plugins/FX/Niagara/Source/NiagaraEditorWidgets/Private/Stack/SNiagaraStackItemFooter.h
new file mode 100644
--- /dev/null
+++ b/Engine/Plugins/FX/Niagara/Source/NiagaraEditorWidgets/Private/Stack/SNiagaraStackItemFooter.h
@@ -0,0 +1,31 @@
+// Copyright 1998-2019 Epic Games, Inc. All Rights Reserved.
+
+#pragma once
+
+#include "Widgets/SCompoundWidget.h"
+#include "Widgets/DeclarativeSyntaxSupport.h"
+
+class UNiagaraStackItemFooter;
+
+class SNiagaraStackItemFooter: public SCompoundWidget
+{
+public:
+	SLATE_BEGIN_ARGS(SNiagaraStackItemFooter) { }
+	SLATE_END_ARGS();
+
+	void Construct(const FArguments& InArgs, UNiagaraStackItemFooter& InItemFooter);
+
+private:
+	EVisibility GetExpandButtonVisibility() const;
+
+	const struct FSlateBrush* GetButtonBrush() const;
+
+	FText GetToolTipText() const;
+
+	FReply ExpandButtonClicked();
+
+private:
+	UNiagaraStackItemFooter* ItemFooter;
+	FText ExpandedToolTipText;
+	FText CollapsedToolTipText;
+};
\ No newline at end of file
diff --git a/Engine/Plugins/FX/Niagara/Source/NiagaraEditorWidgets/Private/Stack/SNiagaraStackModuleItem.cpp b/Engine/Plugins/FX/Niagara/Source/NiagaraEditorWidgets/Private/Stack/SNiagaraStackModuleItem.cpp
--- a/Engine/Plugins/FX/Niagara/Source/NiagaraEditorWidgets/Private/Stack/SNiagaraStackModuleItem.cpp
+++ b/Engine/Plugins/FX/Niagara/Source/NiagaraEditorWidgets/Private/Stack/SNiagaraStackModuleItem.cpp
@@ -272,10 +272,10 @@ FReply SNiagaraStackModuleItem::RefreshClicked()
 
 FReply SNiagaraStackModuleItem::OnModuleItemDrop(TSharedPtr<FDragDropOperation> DragDropOperation)
 {
-	if (DragDropOperation->IsOfType<FNiagaraStackDragOperation>())
+	if (DragDropOperation->IsOfType<FNiagaraParameterDragOperation>())
 	{
-		TSharedPtr<FNiagaraStackDragOperation> InputDragDropOperation = StaticCastSharedPtr<FNiagaraStackDragOperation>(DragDropOperation);
-		TSharedPtr<FNiagaraParameterAction> Action = StaticCastSharedPtr<FNiagaraParameterAction>(InputDragDropOperation->GetAction());
+		TSharedPtr<FNiagaraParameterDragOperation> InputDragDropOperation = StaticCastSharedPtr<FNiagaraParameterDragOperation>(DragDropOperation);
+		TSharedPtr<FNiagaraParameterAction> Action = StaticCastSharedPtr<FNiagaraParameterAction>(InputDragDropOperation->GetSourceAction());
 		if (Action.IsValid() && ModuleItem->CanAddInput(Action->GetParameter()))
 		{
 			ModuleItem->AddInput(Action->GetParameter());
@@ -288,10 +288,10 @@ FReply SNiagaraStackModuleItem::OnModuleItemDrop(TSharedPtr<FDragDropOperation>
 
 bool SNiagaraStackModuleItem::OnModuleItemAllowDrop(TSharedPtr<FDragDropOperation> DragDropOperation)
 {
-	if (DragDropOperation->IsOfType<FNiagaraStackDragOperation>())
+	if (DragDropOperation->IsOfType<FNiagaraParameterDragOperation>())
 	{
-		TSharedPtr<FNiagaraStackDragOperation> InputDragDropOperation = StaticCastSharedPtr<FNiagaraStackDragOperation>(DragDropOperation);
-		TSharedPtr<FNiagaraParameterAction> Action = StaticCastSharedPtr<FNiagaraParameterAction>(InputDragDropOperation->GetAction());
+		TSharedPtr<FNiagaraParameterDragOperation> InputDragDropOperation = StaticCastSharedPtr<FNiagaraParameterDragOperation>(DragDropOperation);
+		TSharedPtr<FNiagaraParameterAction> Action = StaticCastSharedPtr<FNiagaraParameterAction>(InputDragDropOperation->GetSourceAction());
 		return Action.IsValid() && ModuleItem->CanAddInput(Action->GetParameter());
 	}
 
diff --git a/Engine/Plugins/FX/Niagara/Source/NiagaraEditorWidgets/Private/Stack/SNiagaraStackSpacer.cpp b/Engine/Plugins/FX/Niagara/Source/NiagaraEditorWidgets/Private/Stack/SNiagaraStackSpacer.cpp
deleted file mode 100644
diff --git a/Engine/Plugins/FX/Niagara/Source/NiagaraEditorWidgets/Private/Stack/SNiagaraStackSpacer.h b/Engine/Plugins/FX/Niagara/Source/NiagaraEditorWidgets/Private/Stack/SNiagaraStackSpacer.h
deleted file mode 100644
diff --git a/Engine/Plugins/FX/Niagara/Source/NiagaraEditorWidgets/Private/Stack/SNiagaraStackTableRow.cpp b/Engine/Plugins/FX/Niagara/Source/NiagaraEditorWidgets/Private/Stack/SNiagaraStackTableRow.cpp
--- a/Engine/Plugins/FX/Niagara/Source/NiagaraEditorWidgets/Private/Stack/SNiagaraStackTableRow.cpp
+++ b/Engine/Plugins/FX/Niagara/Source/NiagaraEditorWidgets/Private/Stack/SNiagaraStackTableRow.cpp
@@ -26,6 +26,7 @@ void SNiagaraStackTableRow::Construct(const FArguments& InArgs, UNiagaraStackVie
 	ValueColumnWidth = InArgs._ValueColumnWidth;
 	NameColumnWidthChanged = InArgs._OnNameColumnWidthChanged;
 	ValueColumnWidthChanged = InArgs._OnValueColumnWidthChanged;
+	RowPadding = InArgs._RowPadding;
 	StackViewModel = InStackViewModel;
 	StackEntry = InStackEntry;
 	OwnerTree = InOwnerTree;
@@ -210,26 +211,25 @@ void SNiagaraStackTableRow::SetNameAndValueContent(TSharedRef<SWidget> InNameWid
 		SNew(SOverlay)
 		+ SOverlay::Slot()
 		[
-			SNew(SBorder)
-			.BorderImage(FEditorStyle::GetBrush("NoBrush"))
+			SNew(SHorizontalBox)
 			.Visibility(this, &SNiagaraStackTableRow::GetRowVisibility)
-			.Padding(FMargin(0, 0, 0, 0))
+			+ SHorizontalBox::Slot()
+			.AutoWidth()
+			.Padding(1, 0, 6, 0)
 			[
-				SNew(SHorizontalBox)
-				+ SHorizontalBox::Slot()
-				.AutoWidth()
-				.Padding(1, 0, 6, 0)
+				SNew(SBorder)
+				.BorderImage(FEditorStyle::GetBrush("WhiteBrush"))
+				.BorderBackgroundColor(AccentColor)
+				.Padding(0)
 				[
-					SNew(SBorder)
-					.BorderImage(FEditorStyle::GetBrush("WhiteBrush"))
-					.BorderBackgroundColor(AccentColor)
-					.Padding(0)
-					[
-						SNew(SBox)
-						.WidthOverride(4)
-					]
+					SNew(SBox)
+					.WidthOverride(4)
 				]
-				+ SHorizontalBox::Slot()
+			]
+			+ SHorizontalBox::Slot()
+			[
+				SNew(SBox)
+				.Padding(RowPadding)
 				[
 					SNew(SBorder)
 					.BorderImage(FEditorStyle::GetBrush("WhiteBrush"))
diff --git a/Engine/Plugins/FX/Niagara/Source/NiagaraEditorWidgets/Private/Stack/SNiagaraStackTableRow.h b/Engine/Plugins/FX/Niagara/Source/NiagaraEditorWidgets/Private/Stack/SNiagaraStackTableRow.h
--- a/Engine/Plugins/FX/Niagara/Source/NiagaraEditorWidgets/Private/Stack/SNiagaraStackTableRow.h
+++ b/Engine/Plugins/FX/Niagara/Source/NiagaraEditorWidgets/Private/Stack/SNiagaraStackTableRow.h
@@ -25,14 +25,16 @@ public:
 		, _ItemBackgroundColor(FLinearColor::Transparent)
 		, _IsCategoryIconHighlighted(false)
 		, _ShowExecutionCategoryIcon(false)
+		, _RowPadding(FMargin(0, 0, 0, 0))
 	{}
-	SLATE_ARGUMENT(FMargin, ContentPadding)
+		SLATE_ARGUMENT(FMargin, ContentPadding)
 		SLATE_ARGUMENT(FLinearColor, ItemBackgroundColor)
 		SLATE_ARGUMENT(FLinearColor, ItemForegroundColor)
 		SLATE_ARGUMENT(bool, IsCategoryIconHighlighted)
 		SLATE_ARGUMENT(bool, ShowExecutionCategoryIcon)
 		SLATE_ATTRIBUTE(float, NameColumnWidth)
 		SLATE_ATTRIBUTE(float, ValueColumnWidth)
+		SLATE_ATTRIBUTE(FMargin, RowPadding)
 		SLATE_EVENT(FOnColumnWidthChanged, OnNameColumnWidthChanged)
 		SLATE_EVENT(FOnColumnWidthChanged, OnValueColumnWidthChanged)
 		SLATE_EVENT(FOnDragDetected, OnDragDetected)
@@ -101,6 +103,8 @@ private:
 	FOnColumnWidthChanged NameColumnWidthChanged;
 	FOnColumnWidthChanged ValueColumnWidthChanged;
 
+	TAttribute<FMargin> RowPadding;
+
 	const FSlateBrush* ExpandedImage;
 	const FSlateBrush* CollapsedImage;
 
