Niagara - Make the drop-down arrow for stack function inputs focusable again so that it can be opened with the keyboard, and also fixes the issue where the menu's search box isn't focused when it's opened.

diff --git a/Engine/Plugins/FX/Niagara/Source/NiagaraEditorWidgets/Private/Stack/SNiagaraStackFunctionInputValue.cpp b/Engine/Plugins/FX/Niagara/Source/NiagaraEditorWidgets/Private/Stack/SNiagaraStackFunctionInputValue.cpp
--- a/Engine/Plugins/FX/Niagara/Source/NiagaraEditorWidgets/Private/Stack/SNiagaraStackFunctionInputValue.cpp
+++ b/Engine/Plugins/FX/Niagara/Source/NiagaraEditorWidgets/Private/Stack/SNiagaraStackFunctionInputValue.cpp
@@ -178,7 +178,6 @@ void SNiagaraStackFunctionInputValue::Construct(const FArguments& InArgs, UNiaga
 			.Padding(3, 0, 0, 0)
 			[
 				SAssignNew(SetFunctionInputButton, SComboButton)
-				.IsFocusable(false)
 				.ButtonStyle(FEditorStyle::Get(), "HoverHintOnly")
 				.ForegroundColor(FSlateColor::UseForeground())
 				.OnGetMenuContent(this, &SNiagaraStackFunctionInputValue::OnGetAvailableHandleMenu)
