Moved static variable outside of function definition in NiagaraInternalNoise.
This fixes the behavior on Vulkan.

diff --git a/Engine/Plugins/FX/Niagara/Shaders/Private/NiagaraEmitterInstanceShader.usf b/Engine/Plugins/FX/Niagara/Shaders/Private/NiagaraEmitterInstanceShader.usf
--- a/Engine/Plugins/FX/Niagara/Shaders/Private/NiagaraEmitterInstanceShader.usf
+++ b/Engine/Plugins/FX/Niagara/Shaders/Private/NiagaraEmitterInstanceShader.usf
@@ -69,15 +69,17 @@ float Reciprocal(float x){ return 1.0f/x; }
 		return HomogeneousWorldPosition.xyz / HomogeneousWorldPosition.w;
 	}
 
+
+	static uint GRandomSeedOffset;
+
 	// MASSIVE HACK - Tracked in JIRA UE-69298
 	// Hardcoded random function accessible from inner part of node implementation.
 	// It works for now at least and avoids exposing every random needed in the UI. 
 	// Temporary solution, it will be replaced when a design is validated.
 	float NiagaraInternalNoise(uint u, uint v, uint s)
 	{
-		static uint RandomSeedOffset = 0;
-		uint Seed = (u * 1664525u + v) + s + RandomSeedOffset;
-		RandomSeedOffset += Seed;
+		uint Seed = (u * 1664525u + v) + s + GRandomSeedOffset;
+		GRandomSeedOffset += Seed;
 		return float(Rand3DPCG32(int3(u,v,Seed)).x) / 4294967296.0f;
 	}
 
@@ -492,6 +494,7 @@ void SimulateMainComputeCS(
 	GGroupThreadId = GroupThreadId;
 	GEmitterTickCounter = EmitterTickCounter;
 	GSimStart = SimStart;
+	GRandomSeedOffset = 0;
 	/*
 	if(CopyInstancesBeforeStart == 1)
 	{
