FNiagaraGPUInstanceCountManager::ResizeBuffers() now supports growing the instance count buffer through a realloc and copy.

diff --git a/Engine/Plugins/FX/Niagara/Source/Niagara/Classes/NiagaraGPUInstanceCountManager.h b/Engine/Plugins/FX/Niagara/Source/Niagara/Classes/NiagaraGPUInstanceCountManager.h
--- a/Engine/Plugins/FX/Niagara/Source/Niagara/Classes/NiagaraGPUInstanceCountManager.h
+++ b/Engine/Plugins/FX/Niagara/Source/Niagara/Classes/NiagaraGPUInstanceCountManager.h
@@ -62,7 +62,7 @@ public:
 	void DecrementMaxDrawIndirectCount() { checkSlow(IsInRenderingThread()); --MaxDrawIndirectArgs; }
 
 	// Resize instance count and draw indirect buffers to ensure it is big enough to hold all draw indirect args.
-	void ResizeBuffers(int32 ReservedInstanceCounts);
+	void ResizeBuffers(FRHICommandListImmediate& RHICmdList, ERHIFeatureLevel::Type FeatureLevel, int32 ReservedInstanceCounts);
 
 	// Generate the draw indirect buffers, and reset all release counts.
 	void UpdateDrawIndirectBuffer(FRHICommandList& RHICmdList, ERHIFeatureLevel::Type FeatureLevel);
diff --git a/Engine/Plugins/FX/Niagara/Source/Niagara/Private/NiagaraEmitterInstanceBatcher.cpp b/Engine/Plugins/FX/Niagara/Source/Niagara/Private/NiagaraEmitterInstanceBatcher.cpp
--- a/Engine/Plugins/FX/Niagara/Source/Niagara/Private/NiagaraEmitterInstanceBatcher.cpp
+++ b/Engine/Plugins/FX/Niagara/Source/Niagara/Private/NiagaraEmitterInstanceBatcher.cpp
@@ -654,7 +654,7 @@ void NiagaraEmitterInstanceBatcher::PreInitViews(FRHICommandListImmediate& RHICm
 				}
 			}
 		}
-		GPUInstanceCounterManager.ResizeBuffers(TotalDispatchCount);
+		GPUInstanceCounterManager.ResizeBuffers(RHICmdList, FeatureLevel, TotalDispatchCount);
 
 		// Update the instance counts from the GPU readback.
 		{
@@ -705,7 +705,7 @@ void NiagaraEmitterInstanceBatcher::PreInitViews(FRHICommandListImmediate& RHICm
 	}
 	else
 	{
-		GPUInstanceCounterManager.ResizeBuffers(0);
+		GPUInstanceCounterManager.ResizeBuffers(RHICmdList, FeatureLevel,  0);
 	}
 }
 
diff --git a/Engine/Plugins/FX/Niagara/Source/Niagara/Private/NiagaraGPUInstanceCountManager.cpp b/Engine/Plugins/FX/Niagara/Source/Niagara/Private/NiagaraGPUInstanceCountManager.cpp
--- a/Engine/Plugins/FX/Niagara/Source/Niagara/Private/NiagaraGPUInstanceCountManager.cpp
+++ b/Engine/Plugins/FX/Niagara/Source/Niagara/Private/NiagaraGPUInstanceCountManager.cpp
@@ -3,6 +3,7 @@
 #include "NiagaraGPUInstanceCountManager.h"
 #include "NiagaraStats.h"
 #include "Containers/DynamicRHIResourceArray.h"
+#include "NiagaraSortingGPU.h" // FNiagaraCopyIntBufferRegionCS
 
 int32 GNiagaraMinGPUInstanceCount = 2048;
 static FAutoConsoleVariableRef CVarNiagaraMinGPUInstanceCount(
@@ -89,7 +90,7 @@ void FNiagaraGPUInstanceCountManager::FreeEntry(uint32& BufferOffset)
 	}
 }
 
-void FNiagaraGPUInstanceCountManager::ResizeBuffers(int32 ReservedInstanceCounts)
+void FNiagaraGPUInstanceCountManager::ResizeBuffers(FRHICommandListImmediate& RHICmdList, ERHIFeatureLevel::Type FeatureLevel, int32 ReservedInstanceCounts)
 {
 	const int32 RequiredInstanceCounts = UsedInstanceCounts + FMath::Max<int32>(ReservedInstanceCounts - FreeEntries.Num(), 0);
 	if (RequiredInstanceCounts > 0 || MaxDrawIndirectArgs > 0)
@@ -106,7 +107,32 @@ void FNiagaraGPUInstanceCountManager::ResizeBuffers(int32 ReservedInstanceCounts
 		// If we need to increase the buffer size to RecommendedInstanceCounts because the buffer is too small.
 		else if (RequiredInstanceCounts > AllocatedInstanceCounts)
 		{
-			check(false); // @TODO : reallocate and copy current data, filling new space with 0s.
+			// Init a bigger buffer filled with 0.
+			TResourceArray<uint32> InitData;
+			InitData.AddZeroed(RecommendedInstanceCounts);
+			FRWBuffer NextCountBuffer;
+			NextCountBuffer.Initialize(sizeof(uint32), RecommendedInstanceCounts, EPixelFormat::PF_R32_UINT, BUF_Static | BUF_KeepCPUAccessible, TEXT("NiagaraGPUInstanceCounts"), &InitData);
+
+			// Because the shader works with SINT, we need to temporarily create view for the shader to work. Those will be (deferred) deleted at end of the scope.
+			FUnorderedAccessViewRHIRef NextCountBufferUAVAsInt = RHICreateUnorderedAccessView(NextCountBuffer.Buffer, EPixelFormat::PF_R32_SINT);
+			FShaderResourceViewRHIRef CountBufferSRVAsInt = RHICreateShaderResourceView(CountBuffer.Buffer, sizeof(uint32), EPixelFormat::PF_R32_SINT);
+
+			// Copy the current buffer in the next buffer.
+			TShaderMapRef<FNiagaraCopyIntBufferRegionCS> CopyBufferCS(GetGlobalShaderMap(FeatureLevel));
+			RHICmdList.SetComputeShader(CopyBufferCS->GetComputeShader());
+			FRHIUnorderedAccessView* UAVs[NIAGARA_COPY_BUFFER_BUFFER_COUNT] = {};
+			int32 UsedIndexCounts[NIAGARA_COPY_BUFFER_BUFFER_COUNT] = {};
+			UAVs[0] = NextCountBufferUAVAsInt;
+			UsedIndexCounts[0] = AllocatedInstanceCounts;
+			RHICmdList.TransitionResources(EResourceTransitionAccess::ERWBarrier, EResourceTransitionPipeline::EComputeToCompute, UAVs, 1);
+			CopyBufferCS->SetParameters(RHICmdList, CountBufferSRVAsInt, UAVs, UsedIndexCounts, 0, 1);
+			DispatchComputeShader(RHICmdList, *CopyBufferCS, FMath::DivideAndRoundUp(AllocatedInstanceCounts, NIAGARA_COPY_BUFFER_THREAD_COUNT), 1, 1);
+			RHICmdList.TransitionResources(EResourceTransitionAccess::ERWBarrier, EResourceTransitionPipeline::EComputeToGfx, UAVs, 1);
+			CopyBufferCS->UnbindBuffers(RHICmdList);
+
+			// Swap the buffers
+			AllocatedInstanceCounts = RecommendedInstanceCounts;
+			FMemory::Memswap(&NextCountBuffer, &CountBuffer, sizeof(NextCountBuffer));
 		}
 		// If we need to shrink the buffer size because use way to much buffer size.
 		else if ((int32)(RecommendedInstanceCounts * BufferSlack) < AllocatedInstanceCounts)
