Niagara - Add support for deleting items in the overview with the delete key.

diff --git a/Engine/Plugins/FX/Niagara/Source/Niagara/Classes/NiagaraEmitter.h b/Engine/Plugins/FX/Niagara/Source/Niagara/Classes/NiagaraEmitter.h
--- a/Engine/Plugins/FX/Niagara/Source/Niagara/Classes/NiagaraEmitter.h
+++ b/Engine/Plugins/FX/Niagara/Source/Niagara/Classes/NiagaraEmitter.h
@@ -209,6 +209,7 @@ class UNiagaraEmitter : public UObject
 public:
 #if WITH_EDITOR
 	DECLARE_MULTICAST_DELEGATE(FOnPropertiesChanged);
+	DECLARE_MULTICAST_DELEGATE(FOnRenderersChanged);
 	DECLARE_MULTICAST_DELEGATE_OneParam(FOnEmitterCompiled, UNiagaraEmitter*);
 
 	struct NIAGARA_API PrivateMemberNames
@@ -228,6 +229,7 @@ public:
 	//Begin UObject Interface
 	virtual void PostEditChangeProperty(struct FPropertyChangedEvent& PropertyChangedEvent) override;
 	NIAGARA_API FOnPropertiesChanged& OnPropertiesChanged();
+	NIAGARA_API FOnRenderersChanged& OnRenderersChanged();
 #endif
 	void Serialize(FArchive& Ar)override;
 	virtual void PostInitProperties() override;
@@ -503,6 +505,7 @@ private:
 
 #if WITH_EDITOR
 	FOnPropertiesChanged OnPropertiesChangedDelegate;
+	FOnRenderersChanged OnRenderersChangedDelegate;
 #endif
 
 	void GenerateStatID()const;
diff --git a/Engine/Plugins/FX/Niagara/Source/Niagara/Private/NiagaraComponent.cpp b/Engine/Plugins/FX/Niagara/Source/Niagara/Private/NiagaraComponent.cpp
--- a/Engine/Plugins/FX/Niagara/Source/Niagara/Private/NiagaraComponent.cpp
+++ b/Engine/Plugins/FX/Niagara/Source/Niagara/Private/NiagaraComponent.cpp
@@ -1082,6 +1082,15 @@ void UNiagaraComponent::SendRenderDynamicData_Concurrent()
 				NewDynamicData.Add(NewData);
 			}
 		}
+
+#if WITH_EDITOR
+		if (EmitterRenderers.Num() != NewDynamicData.Num())
+		{
+			// This can happen in the editor when modifying the number or renderers while the system is running and the render thread is already processing the data.
+			// in this case we just skip drawing this frame since the system will be reinitialized.
+			return;
+		}
+#endif
 		
 		ENQUEUE_RENDER_COMMAND(NiagaraSetDynamicData)(
 			[NiagaraProxy, DynamicData = MoveTemp(NewDynamicData)](FRHICommandListImmediate& RHICmdList)
diff --git a/Engine/Plugins/FX/Niagara/Source/Niagara/Private/NiagaraEmitter.cpp b/Engine/Plugins/FX/Niagara/Source/Niagara/Private/NiagaraEmitter.cpp
--- a/Engine/Plugins/FX/Niagara/Source/Niagara/Private/NiagaraEmitter.cpp
+++ b/Engine/Plugins/FX/Niagara/Source/Niagara/Private/NiagaraEmitter.cpp
@@ -583,6 +583,11 @@ UNiagaraEmitter::FOnPropertiesChanged& UNiagaraEmitter::OnPropertiesChanged()
 {
 	return OnPropertiesChangedDelegate;
 }
+
+UNiagaraEmitter::FOnPropertiesChanged& UNiagaraEmitter::OnRenderersChanged()
+{
+	return OnRenderersChangedDelegate;
+}
 #endif
 
 
@@ -1080,6 +1085,7 @@ void UNiagaraEmitter::AddRenderer(UNiagaraRendererProperties* Renderer)
 #if WITH_EDITOR
 	Renderer->OnChanged().AddUObject(this, &UNiagaraEmitter::RendererChanged);
 	UpdateChangeId(TEXT("Renderer added"));
+	OnRenderersChangedDelegate.Broadcast();
 #endif
 }
 
@@ -1090,6 +1096,7 @@ void UNiagaraEmitter::RemoveRenderer(UNiagaraRendererProperties* Renderer)
 #if WITH_EDITOR
 	Renderer->OnChanged().RemoveAll(this);
 	UpdateChangeId(TEXT("Renderer removed"));
+	OnRenderersChangedDelegate.Broadcast();
 #endif
 }
 
diff --git a/Engine/Plugins/FX/Niagara/Source/NiagaraEditor/Private/ViewModels/NiagaraSystemViewModel.cpp b/Engine/Plugins/FX/Niagara/Source/NiagaraEditor/Private/ViewModels/NiagaraSystemViewModel.cpp
--- a/Engine/Plugins/FX/Niagara/Source/NiagaraEditor/Private/ViewModels/NiagaraSystemViewModel.cpp
+++ b/Engine/Plugins/FX/Niagara/Source/NiagaraEditor/Private/ViewModels/NiagaraSystemViewModel.cpp
@@ -730,7 +730,7 @@ void FNiagaraSystemViewModel::NotifyDataObjectChanged(UObject* ChangedObject)
 		}
 	}
 
-	ResetSystem(ETimeResetMode::AllowResetTime, EMultiResetMode::ResetThisInstance, EReinitMode::ReinitializeSystem);
+	ResetSystem(ETimeResetMode::AllowResetTime, EMultiResetMode::AllowResetAllInstances, EReinitMode::ReinitializeSystem);
 }
 
 void FNiagaraSystemViewModel::IsolateEmitters(TArray<FGuid> EmitterHandlesIdsToIsolate)
diff --git a/Engine/Plugins/FX/Niagara/Source/NiagaraEditor/Private/ViewModels/Stack/NiagaraStackRenderItemGroup.cpp b/Engine/Plugins/FX/Niagara/Source/NiagaraEditor/Private/ViewModels/Stack/NiagaraStackRenderItemGroup.cpp
--- a/Engine/Plugins/FX/Niagara/Source/NiagaraEditor/Private/ViewModels/Stack/NiagaraStackRenderItemGroup.cpp
+++ b/Engine/Plugins/FX/Niagara/Source/NiagaraEditor/Private/ViewModels/Stack/NiagaraStackRenderItemGroup.cpp
@@ -55,8 +55,8 @@ private:
 class FRenderItemGroupAddUtilities : public TNiagaraStackItemGroupAddUtilities<UNiagaraRendererProperties*>
 {
 public:
-	FRenderItemGroupAddUtilities(TSharedRef<FNiagaraEmitterViewModel> InEmitterViewModel, FOnItemAdded InOnItemAdded)
-		: TNiagaraStackItemGroupAddUtilities(LOCTEXT("RenderGroupAddItemName", "Renderer"), EAddMode::AddFromAction, true, InOnItemAdded)
+	FRenderItemGroupAddUtilities(TSharedRef<FNiagaraEmitterViewModel> InEmitterViewModel)
+		: TNiagaraStackItemGroupAddUtilities(LOCTEXT("RenderGroupAddItemName", "Renderer"), EAddMode::AddFromAction, true, FRenderItemGroupAddUtilities::FOnItemAdded())
 		, EmitterViewModel(InEmitterViewModel)
 	{
 	}
@@ -122,9 +122,9 @@ void UNiagaraStackRenderItemGroup::Initialize(FRequiredEntryData InRequiredEntry
 {
 	FText DisplayName = LOCTEXT("RenderGroupName", "Render");
 	FText ToolTip = LOCTEXT("RendererGroupTooltip", "Describes how we should display/present each particle. Note that this doesn't have to be visual. Multiple renderers are supported. Order in this stack is not necessarily relevant to draw order.");
-	AddUtilities = MakeShared<FRenderItemGroupAddUtilities>(InRequiredEntryData.EmitterViewModel.ToSharedRef(),
-		TNiagaraStackItemGroupAddUtilities<UNiagaraRendererProperties*>::FOnItemAdded::CreateUObject(this, &UNiagaraStackRenderItemGroup::ItemAdded));
+	AddUtilities = MakeShared<FRenderItemGroupAddUtilities>(InRequiredEntryData.EmitterViewModel.ToSharedRef());
 	Super::Initialize(InRequiredEntryData, DisplayName, ToolTip, AddUtilities.Get());
+	GetEmitterViewModel()->GetEmitter()->OnRenderersChanged().AddUObject(this, &UNiagaraStackRenderItemGroup::EmitterRenderersChanged);
 }
 
 void UNiagaraStackRenderItemGroup::RefreshChildrenInternal(const TArray<UNiagaraStackEntry*>& CurrentChildren, TArray<UNiagaraStackEntry*>& NewChildren, TArray<FStackIssue>& NewIssues)
@@ -139,7 +139,6 @@ void UNiagaraStackRenderItemGroup::RefreshChildrenInternal(const TArray<UNiagara
 		{
 			RendererItem = NewObject<UNiagaraStackRendererItem>(this);
 			RendererItem->Initialize(CreateDefaultChildRequiredData(), RendererProperties);
-			RendererItem->OnModifiedGroupItems().AddUObject(this, &UNiagaraStackRenderItemGroup::ChildModifiedGroupItems);
 		}
 
 		NewChildren.Add(RendererItem);
@@ -150,15 +149,16 @@ void UNiagaraStackRenderItemGroup::RefreshChildrenInternal(const TArray<UNiagara
 	Super::RefreshChildrenInternal(CurrentChildren, NewChildren, NewIssues);
 }
 
-void UNiagaraStackRenderItemGroup::ItemAdded(UNiagaraRendererProperties* AddedRenderer)
+void UNiagaraStackRenderItemGroup::EmitterRenderersChanged()
 {
+	OnDataObjectModified().Broadcast(nullptr);
 	RefreshChildren();
-	OnDataObjectModified().Broadcast(AddedRenderer);
 }
 
-void UNiagaraStackRenderItemGroup::ChildModifiedGroupItems()
+void UNiagaraStackRenderItemGroup::FinalizeInternal()
 {
-	RefreshChildren();
+	GetEmitterViewModel()->GetEmitter()->OnRenderersChanged().RemoveAll(this);
+	Super::FinalizeInternal();
 }
 
 #undef LOCTEXT_NAMESPACE
diff --git a/Engine/Plugins/FX/Niagara/Source/NiagaraEditor/Private/ViewModels/Stack/NiagaraStackRendererItem.cpp b/Engine/Plugins/FX/Niagara/Source/NiagaraEditor/Private/ViewModels/Stack/NiagaraStackRendererItem.cpp
--- a/Engine/Plugins/FX/Niagara/Source/NiagaraEditor/Private/ViewModels/Stack/NiagaraStackRendererItem.cpp
+++ b/Engine/Plugins/FX/Niagara/Source/NiagaraEditor/Private/ViewModels/Stack/NiagaraStackRendererItem.cpp
@@ -173,7 +173,6 @@ void UNiagaraStackRendererItem::DeleteInternal()
 
 	OnDataObjectModified().Broadcast(RendererProperties.Get());
 	Finalize();
-	ModifiedGroupItemsDelegate.Broadcast();
 }
 
 bool UNiagaraStackRendererItem::HasBaseRenderer() const
diff --git a/Engine/Plugins/FX/Niagara/Source/NiagaraEditor/Public/NiagaraEditorCommon.h b/Engine/Plugins/FX/Niagara/Source/NiagaraEditor/Public/NiagaraEditorCommon.h
--- a/Engine/Plugins/FX/Niagara/Source/NiagaraEditor/Public/NiagaraEditorCommon.h
+++ b/Engine/Plugins/FX/Niagara/Source/NiagaraEditor/Public/NiagaraEditorCommon.h
@@ -5,7 +5,7 @@
 #include "CoreMinimal.h"
 #include "NiagaraCommon.h"
 
-DECLARE_LOG_CATEGORY_EXTERN(LogNiagaraEditor, All, All);
+NIAGARAEDITOR_API DECLARE_LOG_CATEGORY_EXTERN(LogNiagaraEditor, All, All);
 
 /** Information about a Niagara operation. */
 class FNiagaraOpInfo
diff --git a/Engine/Plugins/FX/Niagara/Source/NiagaraEditor/Public/ViewModels/Stack/NiagaraStackRenderItemGroup.h b/Engine/Plugins/FX/Niagara/Source/NiagaraEditor/Public/ViewModels/Stack/NiagaraStackRenderItemGroup.h
--- a/Engine/Plugins/FX/Niagara/Source/NiagaraEditor/Public/ViewModels/Stack/NiagaraStackRenderItemGroup.h
+++ b/Engine/Plugins/FX/Niagara/Source/NiagaraEditor/Public/ViewModels/Stack/NiagaraStackRenderItemGroup.h
@@ -19,8 +19,9 @@ protected:
 	virtual void RefreshChildrenInternal(const TArray<UNiagaraStackEntry*>& CurrentChildren, TArray<UNiagaraStackEntry*>& NewChildren, TArray<FStackIssue>& NewIssues) override;
 
 private:
-	void ItemAdded(UNiagaraRendererProperties* AddedRenderer);
-	void ChildModifiedGroupItems();
+	void EmitterRenderersChanged();
+
+	virtual void FinalizeInternal() override;
 
 private:
 	TSharedPtr<INiagaraStackItemGroupAddUtilities> AddUtilities;
diff --git a/Engine/Plugins/FX/Niagara/Source/NiagaraEditor/Public/ViewModels/Stack/NiagaraStackViewModel.h b/Engine/Plugins/FX/Niagara/Source/NiagaraEditor/Public/ViewModels/Stack/NiagaraStackViewModel.h
--- a/Engine/Plugins/FX/Niagara/Source/NiagaraEditor/Public/ViewModels/Stack/NiagaraStackViewModel.h
+++ b/Engine/Plugins/FX/Niagara/Source/NiagaraEditor/Public/ViewModels/Stack/NiagaraStackViewModel.h
@@ -88,6 +88,8 @@ public:
 
 	virtual void BeginDestroy() override;
 
+	TSharedPtr<FNiagaraSystemViewModel> GetSystemViewModel();
+
 	TArray<UNiagaraStackEntry*>& GetRootEntries();
 
 	FOnStructureChanged& OnStructureChanged();
diff --git a/Engine/Plugins/FX/Niagara/Source/NiagaraEditorWidgets/Private/SNiagaraOverviewStack.cpp b/Engine/Plugins/FX/Niagara/Source/NiagaraEditorWidgets/Private/SNiagaraOverviewStack.cpp
--- a/Engine/Plugins/FX/Niagara/Source/NiagaraEditorWidgets/Private/SNiagaraOverviewStack.cpp
+++ b/Engine/Plugins/FX/Niagara/Source/NiagaraEditorWidgets/Private/SNiagaraOverviewStack.cpp
@@ -14,13 +14,20 @@
 #include "NiagaraEditorWidgetsUtilities.h"
 #include "SNiagaraStack.h"
 #include "Stack/SNiagaraStackItemGroupAddButton.h"
+#include "NiagaraEditorCommon.h"
 
 #include "Widgets/Layout/SScrollBox.h"
 #include "Widgets/Images/SImage.h"
 #include "Widgets/Input/SButton.h"
+#include "Widgets/Notifications/SNotificationList.h"
 #include "EditorStyleSet.h"
+#include "Styling/CoreStyle.h"
 #include "Framework/Application/SlateApplication.h"
 #include "Framework/MultiBox/MultiBoxBuilder.h"
+#include "Framework/Commands/GenericCommands.h"
+#include "Framework/Notifications/NotificationManager.h"
+#include "ScopedTransaction.h"
+#include "Logging/LogMacros.h"
 
 #define LOCTEXT_NAMESPACE "NiagaraOverviewStack"
 
@@ -30,7 +37,7 @@ class SNiagaraSystemOverviewEntryListRow : public STableRow<UNiagaraStackEntry*>
 		SLATE_EVENT(FOnDragDetected, OnDragDetected)
 		SLATE_EVENT(FOnCanAcceptDrop, OnCanAcceptDrop)
 		SLATE_EVENT(FOnAcceptDrop, OnAcceptDrop)
-		SLATE_DEFAULT_SLOT(FArguments, Content)
+		SLATE_DEFAULT_SLOT(FArguments, Content);
 	SLATE_END_ARGS();
 
 	void Construct(const FArguments& InArgs, UNiagaraStackEntry* InStackEntry, const TSharedRef<STableViewBase>& InOwnerTableView)
@@ -138,7 +145,7 @@ public:
 
 	SLATE_BEGIN_ARGS(SNiagaraSystemOverviewEnabledCheckBox) {}
 		SLATE_ATTRIBUTE(bool, IsChecked)
-		SLATE_EVENT(FOnCheckedChanged, OnCheckedChanged)
+		SLATE_EVENT(FOnCheckedChanged, OnCheckedChanged);
 	SLATE_END_ARGS();
 
 	void Construct(const FArguments& InArgs)
@@ -189,6 +196,10 @@ private:
 
 void SNiagaraOverviewStack::Construct(const FArguments& InArgs, UNiagaraStackViewModel& InStackViewModel, UNiagaraSystemSelectionViewModel& InOverviewSelectionViewModel)
 {
+	Commands = MakeShared<FUICommandList>();
+
+	SetupCommands();
+
 	bUpdatingOverviewSelectionFromStackSelection = false;
 	bUpdatingStackSelectionFromOverviewSelection = false;
 
@@ -219,11 +230,78 @@ SNiagaraOverviewStack::~SNiagaraOverviewStack()
 	}
 }
 
+bool SNiagaraOverviewStack::SupportsKeyboardFocus() const
+{
+	return true;
+}
+
+FReply SNiagaraOverviewStack::OnKeyDown(const FGeometry& MyGeometry, const FKeyEvent& InKeyEvent)
+{
+	if (Commands->ProcessCommandBindings(InKeyEvent))
+	{
+		return FReply::Handled();
+	}
+	return SCompoundWidget::OnKeyDown(MyGeometry, InKeyEvent);
+}
+
 void SNiagaraOverviewStack::Tick(const FGeometry& AllottedGeometry, const double InCurrentTime, const float InDeltaTime)
 {
 	RefreshEntryList();
 }
 
+void SNiagaraOverviewStack::SetupCommands()
+{
+	Commands->MapAction(FGenericCommands::Get().Delete, FUIAction(
+		FExecuteAction::CreateSP(this, &SNiagaraOverviewStack::DeleteSelectedEntries)));
+}
+
+void SNiagaraOverviewStack::DeleteSelectedEntries()
+{
+	bool bIncompleteDelete = false;
+	TArray<UNiagaraStackItem*> ItemsToDelete;
+	for (UNiagaraStackEntry* SelectedEntry : OverviewSelectionViewModel->GetSelectedEntries())
+	{
+		UNiagaraStackItem* SelectedItem = Cast<UNiagaraStackItem>(SelectedEntry);
+		if (SelectedItem != nullptr)
+		{
+			FText DeleteMessage;
+			if (SelectedItem->SupportsDelete() && SelectedItem->TestCanDeleteWithMessage(DeleteMessage))
+			{
+				ItemsToDelete.Add(SelectedItem);
+			}
+			else
+			{
+				bIncompleteDelete = true;
+			}
+		}
+		else
+		{
+			bIncompleteDelete = true;
+		}
+	}
+
+	if (ItemsToDelete.Num() > 0)
+	{
+		const FScopedTransaction Transaction(LOCTEXT("DeleteSelected", "Delete items from the system overview"));
+		for (UNiagaraStackItem* ItemToDelete : ItemsToDelete)
+		{
+			ItemToDelete->Delete();
+		}
+	}
+
+	if (bIncompleteDelete)
+	{
+		FText IncompleteDeleteMessage = LOCTEXT("DeleteIncompleteMessage", "Not all items could be deleted because they either\ndon't support being deleted or they are inherited.");
+		FNotificationInfo Warning(IncompleteDeleteMessage);
+		Warning.ExpireDuration = 5.0f;
+		Warning.bFireAndForget = true;
+		Warning.bUseLargeFont = false;
+		Warning.Image = FCoreStyle::Get().GetBrush(TEXT("MessageLog.Warning"));
+		FSlateNotificationManager::Get().AddNotification(Warning);
+		UE_LOG(LogNiagaraEditor, Warning, TEXT("%s"), *IncompleteDeleteMessage.ToString());
+	}
+}
+
 void SNiagaraOverviewStack::AddEntriesRecursive(UNiagaraStackEntry& EntryToAdd, TArray<UNiagaraStackEntry*>& EntryList, const TArray<UClass*>& AcceptableClasses, TArray<UNiagaraStackEntry*> ParentChain)
 {
 	if (AcceptableClasses.ContainsByPredicate([&] (UClass* Class) { return EntryToAdd.IsA(Class); }))
diff --git a/Engine/Plugins/FX/Niagara/Source/NiagaraEditorWidgets/Private/SNiagaraOverviewStack.h b/Engine/Plugins/FX/Niagara/Source/NiagaraEditorWidgets/Private/SNiagaraOverviewStack.h
--- a/Engine/Plugins/FX/Niagara/Source/NiagaraEditorWidgets/Private/SNiagaraOverviewStack.h
+++ b/Engine/Plugins/FX/Niagara/Source/NiagaraEditorWidgets/Private/SNiagaraOverviewStack.h
@@ -7,6 +7,7 @@
 #include "Widgets/SCompoundWidget.h"
 #include "Widgets/Views/SListView.h"
 #include "UObject/ObjectKey.h"
+#include "Framework/Commands/UICommandList.h"
 
 class UNiagaraStackEntry;
 class UNiagaraStackItem;
@@ -24,9 +25,17 @@ public:
 
 	~SNiagaraOverviewStack();
 
+	virtual bool SupportsKeyboardFocus() const override;
+
+	virtual FReply OnKeyDown(const FGeometry& MyGeometry, const FKeyEvent& InKeyEvent) override;
+
 	virtual void Tick(const FGeometry& AllottedGeometry, const double InCurrentTime, const float InDeltaTime) override;
 
 private:
+	void SetupCommands();
+
+	void DeleteSelectedEntries();
+
 	void AddEntriesRecursive(UNiagaraStackEntry& EntryToAdd, TArray<UNiagaraStackEntry*>& EntryList, const TArray<UClass*>& AcceptableClasses, TArray<UNiagaraStackEntry*> ParentChain);
 
 	void RefreshEntryList();
@@ -57,6 +66,8 @@ private:
 
 	TArray<TWeakObjectPtr<UNiagaraStackEntry>> PreviousSelection;
 
+	TSharedPtr<FUICommandList> Commands;
+
 	bool bRefreshEntryListPending;
 	bool bUpdatingOverviewSelectionFromStackSelection;
 	bool bUpdatingStackSelectionFromOverviewSelection;
