Added "fx.NiagaraGPUDataBufferBufferSlack=1.5" to control the dataset allocation slack for GPU particle system (still rounded to ALLOC_CHUNKSIZE).
Setting it to 0 falls back to the previous implementation where allocation would only be allowed increase in chuncks of ALLOC_CHUNKSIZE,
which is not convenient for systems with 100K+ particles. This also reduces the memory peak for such systems which can prevent OOM

diff --git a/Engine/Plugins/FX/Niagara/Source/Niagara/Private/NiagaraDataSet.cpp b/Engine/Plugins/FX/Niagara/Source/Niagara/Private/NiagaraDataSet.cpp
--- a/Engine/Plugins/FX/Niagara/Source/Niagara/Private/NiagaraDataSet.cpp
+++ b/Engine/Plugins/FX/Niagara/Source/Niagara/Private/NiagaraDataSet.cpp
@@ -66,6 +66,14 @@ static FAutoConsoleVariableRef CVarNiagaraRenderBufferShrinkFactor(
 	ECVF_Default
 );
 
+static float GNiagaraGPUDataBufferBufferSlack = 1.5;
+static FAutoConsoleVariableRef CVarNiagaraGPUDataBufferBufferSlack(
+	TEXT("fx.NiagaraGPUDataBufferBufferSlack"),
+	GNiagaraGPUDataBufferBufferSlack,
+	TEXT("Niagara GPU data buffer size threshold for resizing. <= 1 to disable shrinking. (Default=1.5) \n"),
+	ECVF_Default
+);
+
 FNiagaraDataSet::FNiagaraDataSet()
 	: CompiledData(&FNiagaraDataSetCompiledData::DummyCompiledData)
 	, NumFreeIDs(0)
@@ -601,32 +609,42 @@ void FNiagaraDataBuffer::AllocateGPU(uint32 InNumInstances, FNiagaraGPUInstanceC
 	FloatStride = PaddedNumInstances * sizeof(float);
 	Int32Stride = PaddedNumInstances * sizeof(int32);
 
-	// When the number of elements that we are going to need is greater than the number we have reserved, we need to expand it.
-	if (PaddedNumInstances > NumChunksAllocatedForGPU * ALLOC_CHUNKSIZE)
+	// This never seems to happen because of the +1 in NiagaraEmitterInstanceBatcher::ResizeBuffersAndGatherResources()
+	if (PaddedNumInstances == 0)
 	{
-		NumChunksAllocatedForGPU = FMath::DivideAndRoundUp(PaddedNumInstances, ALLOC_CHUNKSIZE);
-		const uint32 NumElementsToAlloc = NumChunksAllocatedForGPU * ALLOC_CHUNKSIZE;
-		if (NumElementsToAlloc == 0)
+		if (GPUBufferFloat.Buffer)
 		{
-			return;
+			GPUBufferFloat.Release();
 		}
-
-		if (Owner->GetNumFloatComponents())
+		if (GPUBufferInt.Buffer)
 		{
-			if (GPUBufferFloat.Buffer)
-			{
-				GPUBufferFloat.Release();
-			}
-			GPUBufferFloat.Initialize(sizeof(float), NumElementsToAlloc * Owner->GetNumFloatComponents(), EPixelFormat::PF_R32_FLOAT, BUF_Static);
-
+			GPUBufferInt.Release();
 		}
-		if (Owner->GetNumInt32Components())
+	}
+	else // Otherwise check for growing and possibly shrinking (if GNiagaraGPUDataBufferBufferSlack > 1) .
+	{
+		const uint32 RecommendedNumChunks = FMath::DivideAndRoundUp<uint32>((uint32)(PaddedNumInstances * FMath::Max<float>(GNiagaraGPUDataBufferBufferSlack, 1.f)), ALLOC_CHUNKSIZE); 
+		if (PaddedNumInstances > NumChunksAllocatedForGPU * ALLOC_CHUNKSIZE || (GNiagaraGPUDataBufferBufferSlack > 1.f && (uint32)(RecommendedNumChunks * GNiagaraGPUDataBufferBufferSlack) < NumChunksAllocatedForGPU))
 		{
-			if (GPUBufferInt.Buffer)
+			NumChunksAllocatedForGPU = RecommendedNumChunks; 
+			const uint32 NumElementsToAlloc = NumChunksAllocatedForGPU * ALLOC_CHUNKSIZE;
+
+			if (Owner->GetNumFloatComponents())
+			{
+				if (GPUBufferFloat.Buffer)
+				{
+					GPUBufferFloat.Release();
+				}
+				GPUBufferFloat.Initialize(sizeof(float), NumElementsToAlloc * Owner->GetNumFloatComponents(), EPixelFormat::PF_R32_FLOAT, BUF_Static);
+			}
+			if (Owner->GetNumInt32Components())
 			{
-				GPUBufferInt.Release();
+				if (GPUBufferInt.Buffer)
+				{
+					GPUBufferInt.Release();
+				}
+				GPUBufferInt.Initialize(sizeof(int32), NumElementsToAlloc * Owner->GetNumInt32Components(), EPixelFormat::PF_R32_SINT, BUF_Static);
 			}
-			GPUBufferInt.Initialize(sizeof(int32), NumElementsToAlloc * Owner->GetNumInt32Components(), EPixelFormat::PF_R32_SINT, BUF_Static);
 		}
 	}
 }
