Adding support for materials being set in BP for Sprites and Ribbons.
Also fixed nested property editor for material user properties.

diff --git a/Engine/Plugins/FX/Niagara/Source/Niagara/Private/NiagaraRibbonRendererProperties.cpp b/Engine/Plugins/FX/Niagara/Source/Niagara/Private/NiagaraRibbonRendererProperties.cpp
--- a/Engine/Plugins/FX/Niagara/Source/Niagara/Private/NiagaraRibbonRendererProperties.cpp
+++ b/Engine/Plugins/FX/Niagara/Source/Niagara/Private/NiagaraRibbonRendererProperties.cpp
@@ -21,6 +21,8 @@ UNiagaraRibbonRendererProperties::UNiagaraRibbonRendererProperties()
 	, TessellationAngle(15)
 	, bScreenSpaceTessellation(true)
 {
+	FNiagaraTypeDefinition MaterialDef(UMaterialInterface::StaticClass());
+	MaterialUserParamBinding.Parameter.SetType(MaterialDef);
 }
 
 FNiagaraRenderer* UNiagaraRibbonRendererProperties::CreateEmitterRenderer(ERHIFeatureLevel::Type FeatureLevel, const FNiagaraEmitterInstance* Emitter)
@@ -30,6 +32,21 @@ FNiagaraRenderer* UNiagaraRibbonRendererProperties::CreateEmitterRenderer(ERHIFe
 	return NewRenderer;
 }
 
+
+void UNiagaraRibbonRendererProperties::PostLoad()
+{
+	Super::PostLoad();
+
+#if WITH_EDITORONLY_DATA
+
+	if (MaterialUserParamBinding.Parameter.GetType().GetClass() != UMaterialInterface::StaticClass())
+	{
+		FNiagaraTypeDefinition MaterialDef(UMaterialInterface::StaticClass());
+		MaterialUserParamBinding.Parameter.SetType(MaterialDef);
+	}
+#endif
+}
+
 FNiagaraBoundsCalculator* UNiagaraRibbonRendererProperties::CreateBoundsCalculator()
 {
 	return new FNiagaraBoundsCalculatorHelper<false, false, true>();
@@ -37,7 +54,17 @@ FNiagaraBoundsCalculator* UNiagaraRibbonRendererProperties::CreateBoundsCalculat
 
 void UNiagaraRibbonRendererProperties::GetUsedMaterials(const FNiagaraEmitterInstance* InEmitter, TArray<UMaterialInterface*>& OutMaterials) const
 {
-	OutMaterials.Add(Material);
+
+	bool bSet = false;
+	if (MaterialUserParamBinding.Parameter.IsValid() && InEmitter->FindBinding(MaterialUserParamBinding, OutMaterials))
+	{
+		bSet = true;
+	}
+
+	if (!bSet)
+	{
+		OutMaterials.Add(Material);
+	}
 }
 
 void UNiagaraRibbonRendererProperties::PostInitProperties()
diff --git a/Engine/Plugins/FX/Niagara/Source/Niagara/Private/NiagaraSpriteRendererProperties.cpp b/Engine/Plugins/FX/Niagara/Source/Niagara/Private/NiagaraSpriteRendererProperties.cpp
--- a/Engine/Plugins/FX/Niagara/Source/Niagara/Private/NiagaraSpriteRendererProperties.cpp
+++ b/Engine/Plugins/FX/Niagara/Source/Niagara/Private/NiagaraSpriteRendererProperties.cpp
@@ -44,6 +44,8 @@ UNiagaraSpriteRendererProperties::UNiagaraSpriteRendererProperties()
 	, AlphaThreshold(0.1f)
 #endif // WITH_EDITORONLY_DATA
 {
+	FNiagaraTypeDefinition MaterialDef(UMaterialInterface::StaticClass());
+	MaterialUserParamBinding.Parameter.SetType(MaterialDef);
 }
 
 FNiagaraRenderer* UNiagaraSpriteRendererProperties::CreateEmitterRenderer(ERHIFeatureLevel::Type FeatureLevel, const FNiagaraEmitterInstance* Emitter)
@@ -60,7 +62,16 @@ FNiagaraBoundsCalculator* UNiagaraSpriteRendererProperties::CreateBoundsCalculat
 
 void UNiagaraSpriteRendererProperties::GetUsedMaterials(const FNiagaraEmitterInstance* InEmitter, TArray<UMaterialInterface*>& OutMaterials) const
 {
-	OutMaterials.Add(Material);
+	bool bSet = false;
+	if (MaterialUserParamBinding.Parameter.IsValid() && InEmitter->FindBinding(MaterialUserParamBinding, OutMaterials))
+	{
+		bSet = true;
+	}
+
+	if (!bSet)
+	{
+		OutMaterials.Add(Material);
+	}
 }
 
 void UNiagaraSpriteRendererProperties::PostLoad()
@@ -68,6 +79,13 @@ void UNiagaraSpriteRendererProperties::PostLoad()
 	Super::PostLoad();
 
 #if WITH_EDITORONLY_DATA
+
+	if (MaterialUserParamBinding.Parameter.GetType().GetClass() != UMaterialInterface::StaticClass())
+	{
+		FNiagaraTypeDefinition MaterialDef(UMaterialInterface::StaticClass());
+		MaterialUserParamBinding.Parameter.SetType(MaterialDef);
+	}
+
 	if (!FPlatformProperties::RequiresCookedData())
 	{
 		if (CutoutTexture)
diff --git a/Engine/Plugins/FX/Niagara/Source/Niagara/Public/NiagaraRibbonRendererProperties.h b/Engine/Plugins/FX/Niagara/Source/Niagara/Public/NiagaraRibbonRendererProperties.h
--- a/Engine/Plugins/FX/Niagara/Source/Niagara/Public/NiagaraRibbonRendererProperties.h
+++ b/Engine/Plugins/FX/Niagara/Source/Niagara/Public/NiagaraRibbonRendererProperties.h
@@ -63,6 +63,7 @@ public:
 	UNiagaraRibbonRendererProperties();
 
 	//UObject Interface
+	virtual void PostLoad() override;
 	virtual void PostInitProperties() override;
 #if WITH_EDITOR
 	virtual void PostEditChangeProperty(struct FPropertyChangedEvent& PropertyChangedEvent) override;
@@ -89,6 +90,10 @@ public:
 	UPROPERTY(EditAnywhere, Category = "Ribbon Rendering")
 	UMaterialInterface* Material;
 
+	/** Use the UMaterialInterface bound to this user variable if it is set to a valid value. If this is bound to a valid value and Material is also set, UserParamBinding wins.*/
+	UPROPERTY(EditAnywhere, Category = "Ribbon Rendering")
+	FNiagaraUserParameterBinding MaterialUserParamBinding;
+
 	UPROPERTY(EditAnywhere, Category = "Ribbon Rendering")
 	ENiagaraRibbonFacingMode FacingMode;
 
diff --git a/Engine/Plugins/FX/Niagara/Source/Niagara/Public/NiagaraSpriteRendererProperties.h b/Engine/Plugins/FX/Niagara/Source/Niagara/Public/NiagaraSpriteRendererProperties.h
--- a/Engine/Plugins/FX/Niagara/Source/Niagara/Public/NiagaraSpriteRendererProperties.h
+++ b/Engine/Plugins/FX/Niagara/Source/Niagara/Public/NiagaraSpriteRendererProperties.h
@@ -77,6 +77,10 @@ public:
 	/** The material used to render the particle. Note that it must have the Use with Niagara Sprites flag checked.*/
 	UPROPERTY(EditAnywhere, Category = "Sprite Rendering")
 	UMaterialInterface* Material;
+
+	/** Use the UMaterialInterface bound to this user variable if it is set to a valid value. If this is bound to a valid value and Material is also set, UserParamBinding wins.*/
+	UPROPERTY(EditAnywhere, Category = "Sprite Rendering")
+	FNiagaraUserParameterBinding MaterialUserParamBinding;
 	
 	/** Imagine the particle texture having an arrow pointing up, these modes define how the particle aligns that texture to other particle attributes.*/
 	UPROPERTY(EditAnywhere, Category = "Sprite Rendering")
diff --git a/Engine/Plugins/FX/Niagara/Source/NiagaraEditor/Private/Customizations/NiagaraComponentDetails.cpp b/Engine/Plugins/FX/Niagara/Source/NiagaraEditor/Private/Customizations/NiagaraComponentDetails.cpp
--- a/Engine/Plugins/FX/Niagara/Source/NiagaraEditor/Private/Customizations/NiagaraComponentDetails.cpp
+++ b/Engine/Plugins/FX/Niagara/Source/NiagaraEditor/Private/Customizations/NiagaraComponentDetails.cpp
@@ -169,7 +169,7 @@ public:
 
 				FAddPropertyParams Params = FAddPropertyParams()
 					.UniqueId(Parameter.GetName())
-					.AllowChildren(true)
+					.AllowChildren(false) //Don't show the material's properties
 					.CreateCategoryNodes(false);
 
 				Row = ChildrenBuilder.AddExternalObjectProperty(Objects, NAME_None, Params);
diff --git a/Engine/Plugins/FX/Niagara/Source/NiagaraEditor/Private/ViewModels/Stack/NiagaraStackParameterStoreEntry.cpp b/Engine/Plugins/FX/Niagara/Source/NiagaraEditor/Private/ViewModels/Stack/NiagaraStackParameterStoreEntry.cpp
--- a/Engine/Plugins/FX/Niagara/Source/NiagaraEditor/Private/ViewModels/Stack/NiagaraStackParameterStoreEntry.cpp
+++ b/Engine/Plugins/FX/Niagara/Source/NiagaraEditor/Private/ViewModels/Stack/NiagaraStackParameterStoreEntry.cpp
@@ -58,7 +58,7 @@ void UNiagaraStackParameterStoreEntry::RefreshChildrenInternal(const TArray<UNia
 {
 	RefreshValueAndHandle();
 
-	if (ValueObject != nullptr)
+	if (ValueObject != nullptr && ValueObject->IsA<UNiagaraDataInterface>())
 	{
 		if(ValueObjectEntry == nullptr || ValueObjectEntry->GetObject() != ValueObject)
 		{
