Removed expensive unnessesary call that was bloating SystemSim Transfer Results time. Saves 17us per system sim in my tests.
Reordered exec context tick to do least possible work in the most common case. Saves ~31us per instace tick in my tests.
Both savings are dependant on emitter count of which my test has 10 per system.

diff --git a/Engine/Plugins/FX/Niagara/Source/Niagara/Classes/NiagaraDataSet.h b/Engine/Plugins/FX/Niagara/Source/Niagara/Classes/NiagaraDataSet.h
--- a/Engine/Plugins/FX/Niagara/Source/Niagara/Classes/NiagaraDataSet.h
+++ b/Engine/Plugins/FX/Niagara/Source/Niagara/Classes/NiagaraDataSet.h
@@ -413,12 +413,12 @@ struct FNiagaraDataSetAccessorBase
 		VarLayout = DataSet->GetVariableLayout(InVar);
 	}
 
-	void SetDataSet(FNiagaraDataSet& InDataSet)
+	FORCEINLINE void SetDataSet(FNiagaraDataSet& InDataSet)
 	{
 		if (Var.IsValid())
 		{
 			DataSet = &InDataSet;
-			VarLayout = InDataSet.GetVariableLayout(Var);
+			checkSlow(VarLayout == InDataSet.GetVariableLayout(Var));
 		}
 		else
 		{
diff --git a/Engine/Plugins/FX/Niagara/Source/Niagara/Private/NiagaraParameterStore.cpp b/Engine/Plugins/FX/Niagara/Source/Niagara/Private/NiagaraParameterStore.cpp
--- a/Engine/Plugins/FX/Niagara/Source/Niagara/Private/NiagaraParameterStore.cpp
+++ b/Engine/Plugins/FX/Niagara/Source/Niagara/Private/NiagaraParameterStore.cpp
@@ -385,25 +385,14 @@ void FNiagaraParameterStore::CheckForNaNs()const
 	}
 }
 
-void FNiagaraParameterStore::Tick()
+void FNiagaraParameterStore::TickBindings()
 {
-#if NIAGARA_NAN_CHECKING
-	CheckForNaNs();
-#endif
-	if (Bindings.Num() > 0 && (bParametersDirty || bInterfacesDirty || bUObjectsDirty))
+	SCOPE_CYCLE_COUNTER(STAT_NiagaraParameterStoreTick);
+	for (TPair<FNiagaraParameterStore*, FNiagaraParameterStoreBinding>& Binding : Bindings)
 	{
-		SCOPE_CYCLE_COUNTER(STAT_NiagaraParameterStoreTick);
-		for (TPair<FNiagaraParameterStore*, FNiagaraParameterStoreBinding>& Binding : Bindings)
-		{
-			Binding.Value.Tick(Binding.Key, this);
-		}
-		Dump();
+		Binding.Value.Tick(Binding.Key, this);
 	}
-
-	//We have to have ticked all our source stores before now.
-	bParametersDirty = false;
-	bInterfacesDirty = false;
-	bUObjectsDirty = false;
+	Dump();
 }
 
 void FNiagaraParameterStore::UnbindFromSourceStores()
diff --git a/Engine/Plugins/FX/Niagara/Source/Niagara/Private/NiagaraScriptExecutionContext.cpp b/Engine/Plugins/FX/Niagara/Source/Niagara/Private/NiagaraScriptExecutionContext.cpp
--- a/Engine/Plugins/FX/Niagara/Source/Niagara/Private/NiagaraScriptExecutionContext.cpp
+++ b/Engine/Plugins/FX/Niagara/Source/Niagara/Private/NiagaraScriptExecutionContext.cpp
@@ -49,15 +49,14 @@ bool FNiagaraScriptExecutionContext::Init(UNiagaraScript* InScript, ENiagaraSimT
 
 bool FNiagaraScriptExecutionContext::Tick(FNiagaraSystemInstance* ParentSystemInstance, ENiagaraSimTarget SimTarget)
 {
-	SCOPE_CYCLE_COUNTER(STAT_NiagaraScriptExecContextTick);
-
-	if (Script && Script->IsReadyToRun(ENiagaraSimTarget::CPUSim))//TODO: Remove. Script can only be null for system instances that currently don't have their script exec context set up correctly.
+	//Bind data interfaces if needed.
+	if (Parameters.GetInterfacesDirty())
 	{
-		const TArray<UNiagaraDataInterface*>& DataInterfaces = GetDataInterfaces();
-		
-		//Bind data interfaces if needed.
-		if (Parameters.GetInterfacesDirty())
+		SCOPE_CYCLE_COUNTER(STAT_NiagaraScriptExecContextTick);
+		if (Script && Script->IsReadyToRun(ENiagaraSimTarget::CPUSim))//TODO: Remove. Script can only be null for system instances that currently don't have their script exec context set up correctly.
 		{
+			const TArray<UNiagaraDataInterface*>& DataInterfaces = GetDataInterfaces();
+
 			SCOPE_CYCLE_COUNTER(STAT_NiagaraRebindDataInterfaceFunctionTable);
 			// UE_LOG(LogNiagara, Log, TEXT("Updating data interfaces for script %s"), *Script->GetFullName());
 
diff --git a/Engine/Plugins/FX/Niagara/Source/Niagara/Public/NiagaraParameterStore.h b/Engine/Plugins/FX/Niagara/Source/Niagara/Public/NiagaraParameterStore.h
--- a/Engine/Plugins/FX/Niagara/Source/Niagara/Public/NiagaraParameterStore.h
+++ b/Engine/Plugins/FX/Niagara/Source/Niagara/Public/NiagaraParameterStore.h
@@ -245,7 +245,7 @@ public:
 	/** Recreates any bindings to reflect a layout change etc. */
 	void TransferBindings(FNiagaraParameterStore& OtherStore);
 	/** Handles any update such as pushing parameters to bound stores etc. */
-	void Tick();
+	FORCEINLINE_DEBUGGABLE void Tick();
 	/** Unbinds this store from all stores it's being driven by. */
 	void UnbindFromSourceStores();
 	
@@ -610,6 +610,7 @@ public:
 	void TriggerOnLayoutChanged() { OnLayoutChange(); }
 
 protected:
+	void TickBindings();
 	void OnLayoutChange();
 
 	/** Returns the parameter data at the passed offset. */
@@ -635,6 +636,22 @@ protected:
 	friend struct FNiagaraParameterStoreToDataSetBinding;    // this should be the only class calling SetParameterByOffset
 };
 
+FORCEINLINE_DEBUGGABLE void FNiagaraParameterStore::Tick()
+{
+#if NIAGARA_NAN_CHECKING
+	CheckForNaNs();
+#endif
+	if (Bindings.Num() > 0 && (bParametersDirty || bInterfacesDirty || bUObjectsDirty))
+	{
+		TickBindings();
+	}
+
+	//We have to have ticked all our source stores before now.
+	bParametersDirty = false;
+	bInterfacesDirty = false;
+	bUObjectsDirty = false;
+}
+
 FORCEINLINE_DEBUGGABLE void FNiagaraParameterStoreBinding::Empty(FNiagaraParameterStore* DestStore, FNiagaraParameterStore* SrcStore)
 {
 	if (DestStore)
