Niagara - Fix duplicate loctext.

diff --git a/Engine/Plugins/FX/Niagara/Source/NiagaraEditor/Private/ViewModels/Stack/NiagaraStackScriptItemGroup.cpp b/Engine/Plugins/FX/Niagara/Source/NiagaraEditor/Private/ViewModels/Stack/NiagaraStackScriptItemGroup.cpp
--- a/Engine/Plugins/FX/Niagara/Source/NiagaraEditor/Private/ViewModels/Stack/NiagaraStackScriptItemGroup.cpp
+++ b/Engine/Plugins/FX/Niagara/Source/NiagaraEditor/Private/ViewModels/Stack/NiagaraStackScriptItemGroup.cpp
@@ -553,12 +553,12 @@ TOptional<UNiagaraStackEntry::FDropRequestResponse> UNiagaraStackScriptItemGroup
 	if (DropRequest.DropOptions != UNiagaraStackEntry::EDropOptions::Overview)
 	{
 		// Only allow dropping in the overview stacks.
-		return FDropRequestResponse(TOptional<EItemDropZone>(), LOCTEXT("CantDropOnStack", "Modules can only be dropped into the overview."));
+		return FDropRequestResponse(TOptional<EItemDropZone>(), LOCTEXT("CantDropModuleOnStack", "Modules can only be dropped into the overview."));
 	}
 	if (StackEntryDragDropOp->GetDraggedEntries().Num() != 1)
 	{
 		// Only handle a single module.
-		return FDropRequestResponse(TOptional<EItemDropZone>(), LOCTEXT("CantDropMultiple", "Only single modules can be dragged and dropped."));
+		return FDropRequestResponse(TOptional<EItemDropZone>(), LOCTEXT("CantDropMultipleModules", "Only single modules can be dragged and dropped."));
 	}
 
 	UNiagaraStackModuleItem* SourceModuleItem = CastChecked<UNiagaraStackModuleItem>(StackEntryDragDropOp->GetDraggedEntries()[0]);
@@ -570,7 +570,7 @@ TOptional<UNiagaraStackEntry::FDropRequestResponse> UNiagaraStackScriptItemGroup
 	TArray<ENiagaraScriptUsage> SourceUsages = SourceModuleItem->GetModuleNode().FunctionScript->GetSupportedUsageContexts();
 	if (SourceUsages.ContainsByPredicate([this](ENiagaraScriptUsage SourceUsage) { return UNiagaraScript::IsEquivalentUsage(ScriptUsage, SourceUsage); }) == false)
 	{
-		return FDropRequestResponse(TOptional<EItemDropZone>(), LOCTEXT("CantMoveByUsage", "This module can't be moved to this section of the\nstack because it's not valid for this usage context."));
+		return FDropRequestResponse(TOptional<EItemDropZone>(), LOCTEXT("CantMoveModuleByUsage", "This module can't be moved to this section of the\nstack because it's not valid for this usage context."));
 	}
 	if (DropRequest.DropZone != EItemDropZone::AboveItem && DropRequest.DropZone != EItemDropZone::BelowItem)
 	{
@@ -663,7 +663,7 @@ TOptional<UNiagaraStackEntry::FDropRequestResponse> UNiagaraStackScriptItemGroup
 		{
 			FString AssetName;
 			AssetData.GetFullName(AssetName);
-			return FDropRequestResponse(TOptional<EItemDropZone>(), FText::Format(LOCTEXT("CantDropByUsageFormat", "Can not drop asset {0} in this part of the stack\nbecause it's not valid for this usage context."), FText::FromString(AssetName)));
+			return FDropRequestResponse(TOptional<EItemDropZone>(), FText::Format(LOCTEXT("CantDropAssetByUsageFormat", "Can not drop asset {0} in this part of the stack\nbecause it's not valid for this usage context."), FText::FromString(AssetName)));
 		}
 	}
 
@@ -698,7 +698,7 @@ TOptional<UNiagaraStackEntry::FDropRequestResponse> UNiagaraStackScriptItemGroup
 	if (DropRequest.DropOptions != UNiagaraStackEntry::EDropOptions::Overview)
 	{
 		// Only allow dropping in the overview stacks.
-		return FDropRequestResponse(TOptional<EItemDropZone>(), LOCTEXT("AssetCantDropOnStack", "Parameters can only be dropped onto 'Set Variables' modules in the selection view."));
+		return FDropRequestResponse(TOptional<EItemDropZone>(), LOCTEXT("CantDropParameterOnStack", "Parameters can only be dropped onto 'Set Variables' modules in the selection view."));
 	}
 	if (ParameterAction.IsValid())
 	{
