Niagara - Fix a crash when ctrl+dragging an item in the overview, and fix a crash when draging and item in the overview from one system to another.

diff --git a/Engine/Plugins/FX/Niagara/Source/NiagaraEditor/Private/ViewModels/NiagaraEmitterHandleViewModel.cpp b/Engine/Plugins/FX/Niagara/Source/NiagaraEditor/Private/ViewModels/NiagaraEmitterHandleViewModel.cpp
--- a/Engine/Plugins/FX/Niagara/Source/NiagaraEditor/Private/ViewModels/NiagaraEmitterHandleViewModel.cpp
+++ b/Engine/Plugins/FX/Niagara/Source/NiagaraEditor/Private/ViewModels/NiagaraEmitterHandleViewModel.cpp
@@ -25,7 +25,7 @@
 FNiagaraEmitterHandleViewModel::FNiagaraEmitterHandleViewModel()
 	: EmitterHandle(nullptr)
 	, EmitterViewModel(MakeShared<FNiagaraEmitterViewModel>())
-	, EmitterStackViewModel(nullptr)
+	, EmitterStackViewModel(NewObject<UNiagaraStackViewModel>(GetTransientPackage()))
 	, bIsRenamePending(false)
 {
 }
@@ -64,19 +64,12 @@ void FNiagaraEmitterHandleViewModel::Initialize(TSharedRef<FNiagaraSystemViewMod
 	EmitterHandle = InEmitterHandle;
 	UNiagaraEmitter* Emitter = EmitterHandle != nullptr ? EmitterHandle->GetInstance() : nullptr;
 	EmitterViewModel->Initialize(Emitter, InSimulation);
-	if (EmitterStackViewModel != nullptr)
-	{
-		EmitterStackViewModel->InitializeWithViewModels(InOwningSystemViewModel, this->AsShared(), FNiagaraStackViewModelOptions(false, true));
-	}
+	EmitterStackViewModel->InitializeWithViewModels(InOwningSystemViewModel, this->AsShared(), FNiagaraStackViewModelOptions(false, true));
 }
 
 void FNiagaraEmitterHandleViewModel::Reset()
 {
-	if (EmitterStackViewModel != nullptr)
-	{
-		EmitterStackViewModel->Finalize();
-		EmitterStackViewModel = nullptr;
-	}
+	EmitterStackViewModel->Reset();
 	OwningSystemViewModelWeak.Reset();
 	EmitterHandle = nullptr;
 	EmitterViewModel->Reset();
@@ -253,11 +246,6 @@ TSharedRef<FNiagaraEmitterViewModel> FNiagaraEmitterHandleViewModel::GetEmitterV
 
 UNiagaraStackViewModel* FNiagaraEmitterHandleViewModel::GetEmitterStackViewModel()
 {
-	if (EmitterStackViewModel == nullptr)
-	{
-		EmitterStackViewModel = NewObject<UNiagaraStackViewModel>(GetTransientPackage());
-		EmitterStackViewModel->InitializeWithViewModels(GetOwningSystemViewModel(), this->AsShared(), FNiagaraStackViewModelOptions(false, true));
-	}
 	return EmitterStackViewModel;
 }
 
diff --git a/Engine/Plugins/FX/Niagara/Source/NiagaraEditor/Private/ViewModels/NiagaraSystemSelectionViewModel.cpp b/Engine/Plugins/FX/Niagara/Source/NiagaraEditor/Private/ViewModels/NiagaraSystemSelectionViewModel.cpp
--- a/Engine/Plugins/FX/Niagara/Source/NiagaraEditor/Private/ViewModels/NiagaraSystemSelectionViewModel.cpp
+++ b/Engine/Plugins/FX/Niagara/Source/NiagaraEditor/Private/ViewModels/NiagaraSystemSelectionViewModel.cpp
@@ -29,6 +29,7 @@ void UNiagaraSystemSelectionViewModel::Initialize(TSharedRef<FNiagaraSystemViewM
 	SelectionStackViewModel->InitializeWithRootEntry(StackSelection);
 
 	bSystemIsSelected = false;
+	bRefreshIsPending = false;
 }
 
 void UNiagaraSystemSelectionViewModel::Finalize()
@@ -271,6 +272,11 @@ void UNiagaraSystemSelectionViewModel::Refresh()
 	}
 }
 
+void UNiagaraSystemSelectionViewModel::RefreshDeferred()
+{
+	bRefreshIsPending = true;
+}
+
 UNiagaraSystemSelectionViewModel::FOnSelectionChanged& UNiagaraSystemSelectionViewModel::OnSelectionChanged()
 {
 	return OnSelectionChangedDelegate;
@@ -308,6 +314,12 @@ void FindStackGroupsAndItemsForDisplayedObjectKeys(UNiagaraStackViewModel* Stack
 
 void UNiagaraSystemSelectionViewModel::Tick()
 {
+	if (bRefreshIsPending)
+	{
+		Refresh();
+		bRefreshIsPending = false;
+	}
+
 	TArray<UNiagaraStackEntry*> FoundStackEntries;
 	TSharedRef<FNiagaraSystemViewModel> SystemViewModel = GetSystemViewModel();
 	FindStackGroupsAndItemsForDisplayedObjectKeys(SystemViewModel->GetSystemStackViewModel(), DeferredDisplayedObjectKeysToAddToSelection, FoundStackEntries);
diff --git a/Engine/Plugins/FX/Niagara/Source/NiagaraEditor/Private/ViewModels/NiagaraSystemViewModel.cpp b/Engine/Plugins/FX/Niagara/Source/NiagaraEditor/Private/ViewModels/NiagaraSystemViewModel.cpp
--- a/Engine/Plugins/FX/Niagara/Source/NiagaraEditor/Private/ViewModels/NiagaraSystemViewModel.cpp
+++ b/Engine/Plugins/FX/Niagara/Source/NiagaraEditor/Private/ViewModels/NiagaraSystemViewModel.cpp
@@ -96,6 +96,10 @@ void FNiagaraSystemViewModel::Initialize(UNiagaraSystem& InSystem, FNiagaraSyste
 	OverviewGraphViewModel = MakeShared<FNiagaraOverviewGraphViewModel>();
 	OverviewGraphViewModel->Initialize(this->AsShared());
 
+	SystemStackViewModel = NewObject<UNiagaraStackViewModel>(GetTransientPackage());
+	SystemStackViewModel->InitializeWithViewModels(this->AsShared(), TSharedPtr<FNiagaraEmitterHandleViewModel>(), FNiagaraStackViewModelOptions(true, false));
+	SystemStackViewModel->OnStructureChanged().AddSP(this, &FNiagaraSystemViewModel::StackViewModelStructureChanged);
+
 	SetupPreviewComponentAndInstance();
 	SetupSequencer();
 	RefreshAll();
@@ -138,8 +142,12 @@ void FNiagaraSystemViewModel::Cleanup()
 	for (TSharedRef<FNiagaraEmitterHandleViewModel>& HandleRef : EmitterHandleViewModels)
 	{
 		HandleRef->OnPropertyChanged().RemoveAll(this);
+		HandleRef->OnNameChanged().RemoveAll(this);
 		HandleRef->GetEmitterViewModel()->OnPropertyChanged().RemoveAll(this);
 		HandleRef->GetEmitterViewModel()->OnScriptCompiled().RemoveAll(this);
+		HandleRef->GetEmitterViewModel()->OnScriptGraphChanged().RemoveAll(this);
+		HandleRef->GetEmitterViewModel()->OnScriptParameterStoreChanged().RemoveAll(this);
+		HandleRef->GetEmitterStackViewModel()->OnStructureChanged().RemoveAll(this);
 		HandleRef->Cleanup();
 	}
 	EmitterHandleViewModels.Empty();
@@ -158,6 +166,7 @@ void FNiagaraSystemViewModel::Cleanup()
 
 	if (SystemStackViewModel != nullptr)
 	{
+		SystemStackViewModel->OnStructureChanged().RemoveAll(this);
 		SystemStackViewModel->Finalize();
 		SystemStackViewModel = nullptr;
 	}
@@ -601,11 +610,6 @@ const TArray<FNiagaraStackModuleData>& FNiagaraSystemViewModel::GetStackModuleDa
 
 UNiagaraStackViewModel* FNiagaraSystemViewModel::GetSystemStackViewModel()
 {
-	if (SystemStackViewModel == nullptr)
-	{
-		SystemStackViewModel = NewObject<UNiagaraStackViewModel>(GetTransientPackage());
-		SystemStackViewModel->InitializeWithViewModels(this->AsShared(), TSharedPtr<FNiagaraEmitterHandleViewModel>(), FNiagaraStackViewModelOptions(true, false));
-	}
 	return SystemStackViewModel;
 }
 
@@ -813,6 +817,7 @@ void FNiagaraSystemViewModel::RefreshEmitterHandleViewModels()
 			ViewModel->GetEmitterViewModel()->OnScriptCompiled().AddRaw(this, &FNiagaraSystemViewModel::ScriptCompiled);
 			ViewModel->GetEmitterViewModel()->OnScriptGraphChanged().AddRaw(this, &FNiagaraSystemViewModel::EmitterScriptGraphChanged, EmitterHandle->GetId());
 			ViewModel->GetEmitterViewModel()->OnScriptParameterStoreChanged().AddRaw(this, &FNiagaraSystemViewModel::EmitterParameterStoreChanged);
+			ViewModel->GetEmitterStackViewModel()->OnStructureChanged().AddRaw(this, &FNiagaraSystemViewModel::StackViewModelStructureChanged);
 		}
 		else
 		{
@@ -830,10 +835,12 @@ void FNiagaraSystemViewModel::RefreshEmitterHandleViewModels()
 	{
 		TSharedRef<FNiagaraEmitterHandleViewModel> ViewModel = OldViewModels[i];
 		ViewModel->OnPropertyChanged().RemoveAll(this);
+		ViewModel->OnNameChanged().RemoveAll(this);
 		ViewModel->GetEmitterViewModel()->OnPropertyChanged().RemoveAll(this);
 		ViewModel->GetEmitterViewModel()->OnScriptCompiled().RemoveAll(this);
 		ViewModel->GetEmitterViewModel()->OnScriptGraphChanged().RemoveAll(this);
 		ViewModel->GetEmitterViewModel()->OnScriptParameterStoreChanged().RemoveAll(this);
+		ViewModel->GetEmitterStackViewModel()->OnStructureChanged().RemoveAll(this);
 		ViewModel->Reset();
 	}
 
@@ -1771,6 +1778,14 @@ void FNiagaraSystemViewModel::BuildStackModuleData(UNiagaraScript* Script, FGuid
 
 }
 
+void FNiagaraSystemViewModel::StackViewModelStructureChanged()
+{
+	if (SelectionViewModel != nullptr)
+	{
+		SelectionViewModel->RefreshDeferred();
+	}
+}
+
 TSharedPtr<FNiagaraOverviewGraphViewModel> FNiagaraSystemViewModel::GetOverviewGraphViewModel() const
 {
 	return OverviewGraphViewModel;
diff --git a/Engine/Plugins/FX/Niagara/Source/NiagaraEditor/Public/ViewModels/NiagaraSystemSelectionViewModel.h b/Engine/Plugins/FX/Niagara/Source/NiagaraEditor/Public/ViewModels/NiagaraSystemSelectionViewModel.h
--- a/Engine/Plugins/FX/Niagara/Source/NiagaraEditor/Public/ViewModels/NiagaraSystemSelectionViewModel.h
+++ b/Engine/Plugins/FX/Niagara/Source/NiagaraEditor/Public/ViewModels/NiagaraSystemSelectionViewModel.h
@@ -52,6 +52,8 @@ public:
 
 	void Refresh();
 
+	void RefreshDeferred();
+
 	FOnSelectionChanged& OnSelectionChanged();
 
 	void Tick();
@@ -81,4 +83,6 @@ private:
 	FOnSelectionChanged OnSelectionChangedDelegate;
 
 	TArray<FObjectKey> DeferredDisplayedObjectKeysToAddToSelection;
+
+	bool bRefreshIsPending;
 };
\ No newline at end of file
diff --git a/Engine/Plugins/FX/Niagara/Source/NiagaraEditor/Public/ViewModels/NiagaraSystemViewModel.h b/Engine/Plugins/FX/Niagara/Source/NiagaraEditor/Public/ViewModels/NiagaraSystemViewModel.h
--- a/Engine/Plugins/FX/Niagara/Source/NiagaraEditor/Public/ViewModels/NiagaraSystemViewModel.h
+++ b/Engine/Plugins/FX/Niagara/Source/NiagaraEditor/Public/ViewModels/NiagaraSystemViewModel.h
@@ -415,6 +415,9 @@ private:
 	/** builds stack module data for use in module dependencies */
 	void BuildStackModuleData(UNiagaraScript* Script, FGuid InEmitterHandleId, TArray<FNiagaraStackModuleData>& OutStackModuleData);
 
+	/** Called whenever one of the owned stack viewmodels structure changes. */
+	void StackViewModelStructureChanged();
+
 private:
 	/** The System being viewed and edited by this view model. */
 	UNiagaraSystem* System;
diff --git a/Engine/Plugins/FX/Niagara/Source/NiagaraEditorWidgets/Private/SNiagaraOverviewStack.cpp b/Engine/Plugins/FX/Niagara/Source/NiagaraEditorWidgets/Private/SNiagaraOverviewStack.cpp
--- a/Engine/Plugins/FX/Niagara/Source/NiagaraEditorWidgets/Private/SNiagaraOverviewStack.cpp
+++ b/Engine/Plugins/FX/Niagara/Source/NiagaraEditorWidgets/Private/SNiagaraOverviewStack.cpp
@@ -507,7 +507,9 @@ FReply SNiagaraOverviewStack::OnRowDragDetected(const FGeometry& InGeometry, con
 	UNiagaraStackEntry* StackEntry = InStackEntryWeak.Get();
 	if (StackEntry != nullptr && StackEntry->CanDrag())
 	{
-		return FReply::Handled().BeginDragDrop(FNiagaraStackEditorWidgetsUtilities::ConstructDragDropOperationForStackEntries(StackEntry->GetSystemViewModel()->GetSelectionViewModel()->GetSelectedEntries()));
+		TArray<UNiagaraStackEntry*> EntriesToDrag = StackEntry->GetSystemViewModel()->GetSelectionViewModel()->GetSelectedEntries();
+		EntriesToDrag.AddUnique(StackEntry);
+		return FReply::Handled().BeginDragDrop(FNiagaraStackEditorWidgetsUtilities::ConstructDragDropOperationForStackEntries(EntriesToDrag));
 	}
 	return FReply::Unhandled();
 }
