- Niagara Matrix to Quaternion CPU & GPU
- Hidden behind "fx.AllowFastPathFunctionLibrary"

#rnx

diff --git a/Engine/Plugins/FX/Niagara/Source/Niagara/Private/NiagaraFunctionLibrary.cpp b/Engine/Plugins/FX/Niagara/Source/Niagara/Private/NiagaraFunctionLibrary.cpp
--- a/Engine/Plugins/FX/Niagara/Source/Niagara/Private/NiagaraFunctionLibrary.cpp
+++ b/Engine/Plugins/FX/Niagara/Source/Niagara/Private/NiagaraFunctionLibrary.cpp
@@ -37,6 +37,7 @@
 );
 
 TArray<FNiagaraFunctionSignature> UNiagaraFunctionLibrary::VectorVMOps;
+TArray<FString> UNiagaraFunctionLibrary::VectorVMOpsHLSL;
 
 UNiagaraFunctionLibrary::UNiagaraFunctionLibrary(const FObjectInitializer& ObjectInitializer)
 : Super(ObjectInitializer)
@@ -373,8 +374,49 @@
 	return VectorVMOps;
 }
 
+bool UNiagaraFunctionLibrary::DefineFunctionHLSL(const FNiagaraFunctionSignature& FunctionSignature, FString& HlslOutput)
+{
+	if (VectorVMOps.Num() == 0)
+	{
+		InitVectorVMFastPathOps();
+	}
+
+	const int32 i = VectorVMOps.IndexOfByKey(FunctionSignature);
+	if ( i == INDEX_NONE )
+	{
+		return false;
+	}
+
+	HlslOutput += VectorVMOpsHLSL[i];
+	return true;
+}
+
+const FName FastPathLibraryName(TEXT("FastPathLibrary"));
+const FName FastPathDot4Name(TEXT("FastPathDot4"));
+const FName FastPathTransformPositionName(TEXT("FastPathTransformPosition"));
+const FName FastMatrixToQuaternionName(TEXT("FastMatrixToQuaternion"));
+
 struct FVectorKernelFastDot4
 {
+	static FNiagaraFunctionSignature GetFunctionSignature()
+	{
+		FNiagaraFunctionSignature Sig;
+		Sig.Name = FastPathDot4Name;
+		Sig.OwnerName = FastPathLibraryName;
+		Sig.bMemberFunction = false;
+		Sig.bRequiresContext = false;
+		Sig.SetDescription(LOCTEXT("FastPathDot4Desc", "Fast path for Vector4 dot product."));
+		Sig.Inputs.Add(FNiagaraVariable(FNiagaraTypeDefinition::GetVec4Def(), TEXT("A")));
+		Sig.Inputs.Add(FNiagaraVariable(FNiagaraTypeDefinition::GetVec4Def(), TEXT("B")));
+		Sig.Outputs.Add(FNiagaraVariable(FNiagaraTypeDefinition::GetFloatDef(), TEXT("Value")));
+		return Sig;
+	}
+
+	static FString GetFunctionHLSL()
+	{
+		return FString();
+	}
+
 	static void Exec(FVectorVMContext& Context)
 	{
 		//SCOPE_CYCLE_COUNTER(STAT_NiagaraFastDot4);
@@ -449,6 +491,25 @@
 
 struct FVectorKernelFastTransformPosition
 {
+	static FNiagaraFunctionSignature GetFunctionSignature()
+	{
+		FNiagaraFunctionSignature Sig;
+		Sig.Name = FastPathTransformPositionName;
+		Sig.OwnerName = FastPathLibraryName;
+		Sig.bMemberFunction = false;
+		Sig.bRequiresContext = false;
+		Sig.SetDescription(LOCTEXT("FastPathTransformPositionDesc", "Fast path for Matrix4 transforming a Vector3 position"));
+		Sig.Inputs.Add(FNiagaraVariable(FNiagaraTypeDefinition::GetMatrix4Def(), TEXT("Mat")));
+		Sig.Inputs.Add(FNiagaraVariable(FNiagaraTypeDefinition::GetVec3Def(), TEXT("Position")));
+		Sig.Outputs.Add(FNiagaraVariable(FNiagaraTypeDefinition::GetVec3Def(), TEXT("PositionTransformed")));
+		return Sig;
+	}
+
+	static FString GetFunctionHLSL()
+	{
+		return FString();
+	}
+
 	static void Exec(FVectorVMContext& Context)
 	{
 #if 0
@@ -558,9 +619,107 @@
 	}
 };
 
-const FName FastPathLibraryName(TEXT("FastPathLibrary"));
-const FName FastPathDot4Name(TEXT("FastPathDot4"));
-const FName FastPathTransformPositionName(TEXT("FastPathTransformPosition"));
+struct FVectorKernelFastMatrixToQuaternion
+{
+	static FNiagaraFunctionSignature GetFunctionSignature()
+	{
+		FNiagaraFunctionSignature Sig;
+		Sig.Name = FastMatrixToQuaternionName;
+		Sig.OwnerName = FastPathLibraryName;
+		Sig.bMemberFunction = false;
+		Sig.bRequiresContext = false;
+		Sig.SetDescription(LOCTEXT("FastMatrixToQuaternionDesc", "Fast path for Matrix4 to Quaternion"));
+		Sig.Inputs.Add(FNiagaraVariable(FNiagaraTypeDefinition::GetMatrix4Def(), TEXT("Mat")));
+		Sig.Outputs.Add(FNiagaraVariable(FNiagaraTypeDefinition::GetQuatDef(), TEXT("Quat")));
+		return Sig;
+	}
 
+	static FString GetFunctionHLSL()
+	{
+		static FString FunctionHLSL = TEXT(R"(
+	void FastMatrixToQuaternion_FastPathLibrary(float4x4 Mat, out float4 Quat)
+	{
+		float tr = Mat[0][0] + Mat[1][1] + Mat[2][2];
+		if (tr > 0.0f)
+		{
+			float InvS = rsqrt(tr + 1.f);
+			float s = 0.5f * InvS;
+
+			Quat.x = (Mat[2][1] - Mat[1][2]) * s;
+			Quat.y = (Mat[0][2] - Mat[2][0]) * s;
+			Quat.z = (Mat[1][0] - Mat[0][1]) * s;
+			Quat.w = 0.5f * rcp(InvS);
+		}
+		else if ( (Mat[0][0] > Mat[1][1]) && (Mat[0][0] > Mat[2][2]) )
+		{
+			float s = Mat[0][0] - Mat[1][1] - Mat[2][2] + 1.0f;
+			float InvS = rsqrt(s);
+			s = 0.5f * InvS;
+
+			Quat.x = 0.5f * rcp(InvS);
+			Quat.y = (Mat[1][0] + Mat[0][1]) * s;
+			Quat.z = (Mat[2][0] + Mat[0][2]) * s;
+			Quat.w = (Mat[2][1] - Mat[1][2]) * s;
+		}
+		else if ( Mat[1][1] > Mat[2][2] )
+		{
+			float s = Mat[1][1] - Mat[2][2] - Mat[0][0] + 1.0f;
+			float InvS = rsqrt(s);
+			s = 0.5f * InvS;
+
+			Quat.x = (Mat[0][1] + Mat[1][0]) * s;
+			Quat.y = 0.5f * rcp(InvS);
+			Quat.z = (Mat[2][1] + Mat[1][2]) * s;
+			Quat.w = (Mat[0][2] - Mat[2][0]) * s;
+
+		}
+		else
+		{
+			float s = Mat[2][2] - Mat[0][0] - Mat[1][1] + 1.0f;
+			float InvS = rsqrt(s);
+			s = 0.5f * InvS;
+
+			Quat.x = (Mat[0][2] + Mat[2][0]) * s;
+			Quat.y = (Mat[1][2] + Mat[2][1]) * s;
+			Quat.z = 0.5f * rcp(InvS);
+			Quat.w = (Mat[1][0] - Mat[0][1]) * s;
+		}
+	}
+)");
+
+		return FunctionHLSL;
+	}
+
+	static void Exec(FVectorVMContext& Context)
+	{
+		TArray<VectorVM::FExternalFuncInputHandler<float>, TInlineAllocator<16>> InMatrix;
+		for (int i=0; i < 16; i++)
+		{
+			InMatrix.Emplace(Context);
+		}
+
+		TArray<VectorVM::FExternalFuncRegisterHandler<float>, TInlineAllocator<4>> OutQuat;
+		for (int i=0; i < 4; ++i)
+		{
+			OutQuat.Emplace(Context);
+		}
+
+		for ( int32 i=0; i < Context.NumInstances; ++i )
+		{
+			FMatrix Mat;
+			for ( int32 j=0; j < 16; ++j )
+			{
+				Mat.M[j & 3][j >> 2] = InMatrix[j].GetAndAdvance();
+			}
+
+			FQuat Quat(Mat);
+			*OutQuat[0].GetDestAndAdvance() = Quat.X;
+			*OutQuat[1].GetDestAndAdvance() = Quat.Y;
+			*OutQuat[2].GetDestAndAdvance() = Quat.Z;
+			*OutQuat[3].GetDestAndAdvance() = Quat.W;
+		}
+	}
+};
+
 bool UNiagaraFunctionLibrary::GetVectorVMFastPathExternalFunction(const FVMExternalFunctionBindingInfo& BindingInfo, FVMExternalFunction &OutFunc)
 {
@@ -574,6 +733,11 @@
 		OutFunc = FVMExternalFunction::CreateStatic(FVectorKernelFastTransformPosition::Exec);
 		return true;
 	}
+	else if (BindingInfo.Name == FastMatrixToQuaternionName)
+	{
+		OutFunc = FVMExternalFunction::CreateStatic(FVectorKernelFastMatrixToQuaternion::Exec);
+		return true;
+	}
 	return false;
 }
 
@@ -583,30 +747,15 @@
 	if (GAllowFastPathFunctionLibrary == 0)
 		return;
 
-	{
-		FNiagaraFunctionSignature Sig;
-		Sig.Name = FastPathDot4Name;
-		Sig.OwnerName = FastPathLibraryName;
-		Sig.bMemberFunction = false;
-		Sig.bRequiresContext = false;
-		Sig.SetDescription(LOCTEXT("FastPathDot4Desc", "Fast path for Vector4 dot product."));
-		Sig.Inputs.Add(FNiagaraVariable(FNiagaraTypeDefinition::GetVec4Def(), TEXT("A")));
-		Sig.Inputs.Add(FNiagaraVariable(FNiagaraTypeDefinition::GetVec4Def(), TEXT("B")));
-		Sig.Outputs.Add(FNiagaraVariable(FNiagaraTypeDefinition::GetFloatDef(), TEXT("Value")));
-		VectorVMOps.Add(Sig);
-	}
-	{
-		FNiagaraFunctionSignature Sig;
-		Sig.Name = FastPathTransformPositionName;
-		Sig.OwnerName = FastPathLibraryName;
-		Sig.bMemberFunction = false;
-		Sig.bRequiresContext = false;
-		Sig.SetDescription(LOCTEXT("FastPathTransformPositionDesc", "Fast path for Matrix4 transforming a Vector3 position"));
-		Sig.Inputs.Add(FNiagaraVariable(FNiagaraTypeDefinition::GetMatrix4Def(), TEXT("Mat")));
-		Sig.Inputs.Add(FNiagaraVariable(FNiagaraTypeDefinition::GetVec3Def(), TEXT("Position")));
-		Sig.Outputs.Add(FNiagaraVariable(FNiagaraTypeDefinition::GetVec3Def(), TEXT("PositionTransformed")));
-		VectorVMOps.Add(Sig);
-	}
+	VectorVMOps.Emplace(FVectorKernelFastDot4::GetFunctionSignature());
+	VectorVMOps.Emplace(FVectorKernelFastTransformPosition::GetFunctionSignature());
+	VectorVMOps.Emplace(FVectorKernelFastMatrixToQuaternion::GetFunctionSignature());
+
+	VectorVMOpsHLSL.Emplace(FVectorKernelFastDot4::GetFunctionHLSL());
+	VectorVMOpsHLSL.Emplace(FVectorKernelFastTransformPosition::GetFunctionHLSL());
+	VectorVMOpsHLSL.Emplace(FVectorKernelFastMatrixToQuaternion::GetFunctionHLSL());
+
+	check(VectorVMOps.Num() == VectorVMOpsHLSL.Num());
 }
 
 #undef LOCTEXT_NAMESPACE
diff --git a/Engine/Plugins/FX/Niagara/Source/Niagara/Public/NiagaraFunctionLibrary.h b/Engine/Plugins/FX/Niagara/Source/Niagara/Public/NiagaraFunctionLibrary.h
--- a/Engine/Plugins/FX/Niagara/Source/Niagara/Public/NiagaraFunctionLibrary.h
+++ b/Engine/Plugins/FX/Niagara/Source/Niagara/Public/NiagaraFunctionLibrary.h
@@ -56,10 +56,12 @@ public:
 	static UNiagaraParameterCollectionInstance* GetNiagaraParameterCollection(UObject* WorldContextObject, UNiagaraParameterCollection* Collection);
 
 	static const TArray<FNiagaraFunctionSignature>& GetVectorVMFastPathOps();
+	static bool DefineFunctionHLSL(const FNiagaraFunctionSignature& FunctionSignature, FString& HlslOutput);
 
 	static bool GetVectorVMFastPathExternalFunction(const FVMExternalFunctionBindingInfo& BindingInfo, FVMExternalFunction &OutFunc);
 
 private:
 	static void InitVectorVMFastPathOps();
 	static TArray<FNiagaraFunctionSignature> VectorVMOps;
+	static TArray<FString> VectorVMOpsHLSL;
 };
diff --git a/Engine/Plugins/FX/Niagara/Source/NiagaraEditor/Private/NiagaraHlslTranslator.cpp b/Engine/Plugins/FX/Niagara/Source/NiagaraEditor/Private/NiagaraHlslTranslator.cpp
--- a/Engine/Plugins/FX/Niagara/Source/NiagaraEditor/Private/NiagaraHlslTranslator.cpp
+++ b/Engine/Plugins/FX/Niagara/Source/NiagaraEditor/Private/NiagaraHlslTranslator.cpp
@@ -31,6 +31,8 @@
 #include "NiagaraEditorModule.h"
 #include "NiagaraNodeReroute.h"
 
+#include "NiagaraFunctionLibrary.h"
+
 #include "NiagaraDataInterface.h"
 #include "NiagaraDataInterfaceCurve.h"
 #include "NiagaraDataInterfaceVector2DCurve.h"
@@ -1323,6 +1325,8 @@ const FNiagaraTranslateResults &FHlslNiagaraTranslator::Translate(const FNiagara
 			FString DataInterfaceHLSL;
 			DefineDataInterfaceHLSL(DataInterfaceHLSL);
 			HlslOutput += DataInterfaceHLSL;
+
+			DefineExternalFunctionsHLSL(HlslOutput);
 		}
 
 		//And finally, define the actual main function that handles the reading and writing of data and calls the shared per instance simulate function.
@@ -1778,6 +1782,16 @@ void FHlslNiagaraTranslator::DefineDataInterfaceHLSL(FString &InHlslOutput)
 	InHlslOutput += InterfaceCommonHLSL + InterfaceUniformHLSL + InterfaceFunctionHLSL;
 }
 
+void FHlslNiagaraTranslator::DefineExternalFunctionsHLSL(FString &InHlslOutput)
+{
+	for (FNiagaraFunctionSignature& FunctionSig : CompilationOutput.ScriptData.AdditionalExternalFunctions )
+	{
+		if ( UNiagaraFunctionLibrary::DefineFunctionHLSL(FunctionSig, InHlslOutput) == false )
+		{
+			Error(FText::Format(LOCTEXT("ExternFunctionMissingHLSL", "ExternalFunction {0} does not have a HLSL implementation for the GPU."), FText::FromName(FunctionSig.Name)), nullptr, nullptr);
+		}
+	}
+}
 
 void FHlslNiagaraTranslator::DefineMain(FString &OutHlslOutput,
 	TArray<TArray<FNiagaraVariable>*> &InstanceReadVars, TArray<FNiagaraDataSetID>& ReadIds,
diff --git a/Engine/Plugins/FX/Niagara/Source/NiagaraEditor/Private/NiagaraHlslTranslator.h b/Engine/Plugins/FX/Niagara/Source/NiagaraEditor/Private/NiagaraHlslTranslator.h
--- a/Engine/Plugins/FX/Niagara/Source/NiagaraEditor/Private/NiagaraHlslTranslator.h
+++ b/Engine/Plugins/FX/Niagara/Source/NiagaraEditor/Private/NiagaraHlslTranslator.h
@@ -464,6 +464,7 @@ public:
 	void DefineDataSetVariableReads(FString &HLSLOutput, FNiagaraDataSetID& Id, int32 DataSetIndex, TArray<FNiagaraVariable> & ReadVars);
 	void DefineDataSetVariableWrites(FString &HlslOutput, FNiagaraDataSetID& Id, int32 DataSetIndex, TArray<FNiagaraVariable>& WriteVars);
 	void DefineDataInterfaceHLSL(FString &HlslOutput);
+	void DefineExternalFunctionsHLSL(FString &HlslOutput);
 	TArray<FNiagaraDataInterfaceGPUParamInfo>& GetDataInterfaceParameters() { return DIParamInfo; }
 
 	// Format string should have up to 5 entries, {{0} = Computed Variable Suffix, {1} = Float or Int, {2} = Data Set Index, {3} = Register Index, {4} Default value for that type.
