Niagara - In the stack and in the overview add icons to show issues warnings and errors wherever it's relevant.

diff --git a/Engine/Plugins/FX/Niagara/Source/NiagaraEditor/Private/Sequencer/NiagaraSequence/NiagaraEmitterTrackEditor.cpp b/Engine/Plugins/FX/Niagara/Source/NiagaraEditor/Private/Sequencer/NiagaraSequence/NiagaraEmitterTrackEditor.cpp
--- a/Engine/Plugins/FX/Niagara/Source/NiagaraEditor/Private/Sequencer/NiagaraSequence/NiagaraEmitterTrackEditor.cpp
+++ b/Engine/Plugins/FX/Niagara/Source/NiagaraEditor/Private/Sequencer/NiagaraSequence/NiagaraEmitterTrackEditor.cpp
@@ -7,9 +7,11 @@
 #include "ViewModels/NiagaraEmitterHandleViewModel.h"
 #include "ViewModels/NiagaraEmitterViewModel.h"
 #include "ViewModels/NiagaraSystemSelectionViewModel.h"
+#include "ViewModels/Stack/NiagaraStackViewModel.h"
 #include "Sequencer/NiagaraSequence/NiagaraSequence.h"
 #include "Framework/MultiBox/MultiBoxBuilder.h"
 #include "NiagaraEditorStyle.h"
+#include "NiagaraEditorModule.h"
 
 #include "EditorStyleSet.h"
 #include "Styling/SlateIconFinder.h"
@@ -17,6 +19,7 @@
 #include "Widgets/Input/SButton.h"
 #include "Widgets/Images/SImage.h"
 #include "Widgets/Input/SCheckBox.h"
+#include "Modules/ModuleManager.h"
 
 #define LOCTEXT_NAMESPACE "NiagaraEmitterTrackEditor"
 
@@ -30,6 +33,7 @@ public:
 	{
 		EmitterTrack = &InEmitterTrack;
 
+		FNiagaraEditorModule& NiagaraEditorModule = FModuleManager::LoadModuleChecked<FNiagaraEditorModule>("NiagaraEditor");
 		TSharedRef<SHorizontalBox> TrackBox = SNew(SHorizontalBox)
 			// Track initialization error icon.
 			+ SHorizontalBox::Slot()
@@ -42,6 +46,16 @@ public:
 				.Image(FEditorStyle::GetBrush("Icons.Info"))
 				.ToolTipText(this, &SEmitterTrackWidget::GetTrackErrorIconToolTip)
 			]
+			// Stack issues icon
+			+ SHorizontalBox::Slot()
+			.AutoWidth()
+			.VAlign(VAlign_Center)
+			.Padding(3, 0, 0, 0)
+			[
+				NiagaraEditorModule.GetWidgetProvider()->CreateStackIssueIcon( 
+					*EmitterTrack->GetEmitterHandleViewModel()->GetEmitterStackViewModel(),
+					*EmitterTrack->GetEmitterHandleViewModel()->GetEmitterStackViewModel()->GetRootEntry())
+			]
 			// Isolate toggle
 			+ SHorizontalBox::Slot()
 			.AutoWidth()
diff --git a/Engine/Plugins/FX/Niagara/Source/NiagaraEditor/Private/ViewModels/NiagaraSystemSelectionViewModel.cpp b/Engine/Plugins/FX/Niagara/Source/NiagaraEditor/Private/ViewModels/NiagaraSystemSelectionViewModel.cpp
--- a/Engine/Plugins/FX/Niagara/Source/NiagaraEditor/Private/ViewModels/NiagaraSystemSelectionViewModel.cpp
+++ b/Engine/Plugins/FX/Niagara/Source/NiagaraEditor/Private/ViewModels/NiagaraSystemSelectionViewModel.cpp
@@ -160,12 +160,11 @@ void UNiagaraSystemSelectionViewModel::UpdateSelectionFromTopLevelObjects(bool b
 		bSystemIsSelected = bInSystemIsSelected;
 		if (bSystemIsSelected)
 		{
-			SelectedEntries.Append(GetSystemViewModel()->GetSystemStackViewModel()->GetRootEntries());
+			SelectedEntries.Add(GetSystemViewModel()->GetSystemStackViewModel()->GetRootEntry());
 		}
 		else
 		{
-			TArray<UNiagaraStackEntry*> SystemRootEntries = GetSystemViewModel()->GetSystemStackViewModel()->GetRootEntries();
-			SelectedEntries.RemoveAll([&SystemRootEntries](UNiagaraStackEntry* SelectedEntry) { return SystemRootEntries.Contains(SelectedEntry); });
+			SelectedEntries.Remove(GetSystemViewModel()->GetSystemStackViewModel()->GetRootEntry());
 		}
 		bSelectionChanged = true;
 	}
@@ -178,8 +177,7 @@ void UNiagaraSystemSelectionViewModel::UpdateSelectionFromTopLevelObjects(bool b
 		{
 			SelectedEmitterHandleIds.Remove(CurrentSelectedEmitterHandleId);
 			TSharedPtr<FNiagaraEmitterHandleViewModel> DeselectedEmitterHandleViewModel = GetSystemViewModel()->GetEmitterHandleViewModelById(CurrentSelectedEmitterHandleId);
-			TArray<UNiagaraStackEntry*> DeselectedEmitterRootEntries = DeselectedEmitterHandleViewModel->GetEmitterStackViewModel()->GetRootEntries();
-			SelectedEntries.RemoveAll([&DeselectedEmitterRootEntries](UNiagaraStackEntry* SelectedEntry) { return DeselectedEmitterRootEntries.Contains(SelectedEntry); });
+			SelectedEntries.Remove(DeselectedEmitterHandleViewModel->GetEmitterStackViewModel()->GetRootEntry());
 			bSelectionChanged = true;
 		}
 	}
@@ -192,7 +190,7 @@ void UNiagaraSystemSelectionViewModel::UpdateSelectionFromTopLevelObjects(bool b
 			TSharedPtr<FNiagaraEmitterHandleViewModel> SelectedEmitterHandleViewModel = GetSystemViewModel()->GetEmitterHandleViewModelById(InSelectedEmitterId);
 			if (SelectedEmitterHandleViewModel.IsValid())
 			{
-				SelectedEntries.Append(SelectedEmitterHandleViewModel->GetEmitterStackViewModel()->GetRootEntries());
+				SelectedEntries.Add(SelectedEmitterHandleViewModel->GetEmitterStackViewModel()->GetRootEntry());
 			}
 			bSelectionChanged = true;
 		}
@@ -270,6 +268,10 @@ void UNiagaraSystemSelectionViewModel::Refresh()
 		OnSelectionChangedDelegate.Broadcast(ESelectionChangeSource::Refresh);
 		UpdateStackSelectionEntry();
 	}
+	else
+	{
+		StackSelection->RefreshChildren();
+	}
 }
 
 void UNiagaraSystemSelectionViewModel::RefreshDeferred()
@@ -304,14 +306,6 @@ void FindStackGroupsAndItemsForDisplayedObjectKeysRecursive(UNiagaraStackEntry*
 	}
 }
 
-void FindStackGroupsAndItemsForDisplayedObjectKeys(UNiagaraStackViewModel* StackViewModel, const TArray<FObjectKey>& ObjectKeys, TArray<UNiagaraStackEntry*>& OutFoundStackEntries)
-{
-	for (UNiagaraStackEntry* RootEntry : StackViewModel->GetRootEntries())
-	{
-		FindStackGroupsAndItemsForDisplayedObjectKeysRecursive(RootEntry, ObjectKeys, OutFoundStackEntries);
-	}
-}
-
 void UNiagaraSystemSelectionViewModel::Tick()
 {
 	if (bRefreshIsPending)
@@ -322,10 +316,10 @@ void UNiagaraSystemSelectionViewModel::Tick()
 
 	TArray<UNiagaraStackEntry*> FoundStackEntries;
 	TSharedRef<FNiagaraSystemViewModel> SystemViewModel = GetSystemViewModel();
-	FindStackGroupsAndItemsForDisplayedObjectKeys(SystemViewModel->GetSystemStackViewModel(), DeferredDisplayedObjectKeysToAddToSelection, FoundStackEntries);
+	FindStackGroupsAndItemsForDisplayedObjectKeysRecursive(SystemViewModel->GetSystemStackViewModel()->GetRootEntry(), DeferredDisplayedObjectKeysToAddToSelection, FoundStackEntries);
 	for (TSharedRef<FNiagaraEmitterHandleViewModel> EmitterViewModel : SystemViewModel->GetEmitterHandleViewModels())
 	{
-		FindStackGroupsAndItemsForDisplayedObjectKeys(EmitterViewModel->GetEmitterStackViewModel(), DeferredDisplayedObjectKeysToAddToSelection, FoundStackEntries);
+		FindStackGroupsAndItemsForDisplayedObjectKeysRecursive(EmitterViewModel->GetEmitterStackViewModel()->GetRootEntry(), DeferredDisplayedObjectKeysToAddToSelection, FoundStackEntries);
 	}
 
 	int32 AddedEntries = 0;
diff --git a/Engine/Plugins/FX/Niagara/Source/NiagaraEditor/Private/ViewModels/NiagaraSystemViewModel.cpp b/Engine/Plugins/FX/Niagara/Source/NiagaraEditor/Private/ViewModels/NiagaraSystemViewModel.cpp
--- a/Engine/Plugins/FX/Niagara/Source/NiagaraEditor/Private/ViewModels/NiagaraSystemViewModel.cpp
+++ b/Engine/Plugins/FX/Niagara/Source/NiagaraEditor/Private/ViewModels/NiagaraSystemViewModel.cpp
@@ -1255,10 +1255,7 @@ void FNiagaraSystemViewModel::EmitterHandlePropertyChanged(FGuid OwningEmitterHa
 	TSharedPtr<FNiagaraEmitterHandleViewModel> EmitterHandleViewModel = GetEmitterHandleViewModelById(OwningEmitterHandleId);
 	if (EmitterHandleViewModel.IsValid())
 	{
-		for (UNiagaraStackEntry* RootEntry : EmitterHandleViewModel->GetEmitterStackViewModel()->GetRootEntries())
-		{
-			RootEntry->RefreshChildren();
-		}
+		EmitterHandleViewModel->GetEmitterStackViewModel()->GetRootEntry()->RefreshChildren();
 	}
 
 	ResetSystem(ETimeResetMode::AllowResetTime, EMultiResetMode::ResetThisInstance, EReinitMode::ReinitializeSystem);
diff --git a/Engine/Plugins/FX/Niagara/Source/NiagaraEditor/Private/ViewModels/Stack/NiagaraStackEntry.cpp b/Engine/Plugins/FX/Niagara/Source/NiagaraEditor/Private/ViewModels/Stack/NiagaraStackEntry.cpp
--- a/Engine/Plugins/FX/Niagara/Source/NiagaraEditor/Private/ViewModels/Stack/NiagaraStackEntry.cpp
+++ b/Engine/Plugins/FX/Niagara/Source/NiagaraEditor/Private/ViewModels/Stack/NiagaraStackEntry.cpp
@@ -414,6 +414,36 @@ bool UNiagaraStackEntry::HasBaseEmitter() const
 	return bHasBaseEmitterCache.GetValue();
 }
 
+bool UNiagaraStackEntry::HasIssuesOrAnyChildHasIssues() const
+{
+	return TotalNumberOfErrorIssues > 0 || TotalNumberOfWarningIssues > 0 || TotalNumberOfInfoIssues > 0;
+}
+
+int32 UNiagaraStackEntry::GetTotalNumberOfInfoIssues() const
+{
+	return TotalNumberOfInfoIssues;
+}
+
+int32 UNiagaraStackEntry::GetTotalNumberOfWarningIssues() const
+{
+	return TotalNumberOfWarningIssues;
+}
+
+int32 UNiagaraStackEntry::GetTotalNumberOfErrorIssues() const
+{
+	return TotalNumberOfErrorIssues;
+}
+
+const TArray<UNiagaraStackEntry::FStackIssue>& UNiagaraStackEntry::GetIssues() const
+{
+	return StackIssues;
+}
+
+const TArray<UNiagaraStackEntry*>& UNiagaraStackEntry::GetAllChildrenWithIssues() const
+{
+	return ChildrenWithIssues;
+}
+
 TOptional<UNiagaraStackEntry::FDropRequestResponse> UNiagaraStackEntry::CanDropInternal(const FDropRequest& DropRequest)
 {
 	return TOptional<FDropRequestResponse>();
@@ -489,6 +519,11 @@ void UNiagaraStackEntry::RefreshChildren()
 	Children.Empty();
 	Children.Append(NewChildren);
 
+	TotalNumberOfInfoIssues = 0;
+	TotalNumberOfWarningIssues = 0;
+	TotalNumberOfErrorIssues = 0;
+	ChildrenWithIssues.Empty();
+
 	for (UNiagaraStackEntry* Child : Children)
 	{
 		UNiagaraStackEntry* OuterOwner = Cast<UNiagaraStackEntry>(Child->GetOuter());
@@ -504,11 +539,37 @@ void UNiagaraStackEntry::RefreshChildren()
 			Child->SetOnRequestCanDrop(FOnRequestDrop::CreateUObject(this, &UNiagaraStackEntry::ChildRequestCanDrop));
 			Child->SetOnRequestDrop(FOnRequestDrop::CreateUObject(this, &UNiagaraStackEntry::ChildRequestDrop));
 		}
+
+		TotalNumberOfInfoIssues += Child->GetTotalNumberOfInfoIssues();
+		TotalNumberOfWarningIssues += Child->GetTotalNumberOfWarningIssues();
+		TotalNumberOfErrorIssues += Child->GetTotalNumberOfErrorIssues();
+
+		if (Child->GetIssues().Num() > 0)
+		{
+			ChildrenWithIssues.Add(Child);
+		}
+		ChildrenWithIssues.Append(Child->GetAllChildrenWithIssues());
 	}
 	
 	// Stack issues refresh
 	NewStackIssues.RemoveAll([=](const FStackIssue& Issue) { return Issue.GetCanBeDismissed() && GetStackEditorData().GetDismissedStackIssueIds().Contains(Issue.GetUniqueIdentifier()); }); 
 
+	for (const FStackIssue& Issue : NewStackIssues)
+	{
+		if (Issue.GetSeverity() == EStackIssueSeverity::Info)
+		{
+			TotalNumberOfInfoIssues++;
+		}
+		else if (Issue.GetSeverity() == EStackIssueSeverity::Warning)
+		{
+			TotalNumberOfWarningIssues++;
+		}
+		else if (Issue.GetSeverity() == EStackIssueSeverity::Error)
+		{
+			TotalNumberOfErrorIssues++;
+		}
+	}
+
 	StackIssues.Empty();
 	StackIssues.Append(NewStackIssues);
 	RefreshStackErrorChildren();
diff --git a/Engine/Plugins/FX/Niagara/Source/NiagaraEditor/Private/ViewModels/Stack/NiagaraStackItem.cpp b/Engine/Plugins/FX/Niagara/Source/NiagaraEditor/Private/ViewModels/Stack/NiagaraStackItem.cpp
--- a/Engine/Plugins/FX/Niagara/Source/NiagaraEditor/Private/ViewModels/Stack/NiagaraStackItem.cpp
+++ b/Engine/Plugins/FX/Niagara/Source/NiagaraEditor/Private/ViewModels/Stack/NiagaraStackItem.cpp
@@ -25,35 +25,6 @@ UNiagaraStackItem::FOnModifiedGroupItems& UNiagaraStackItem::OnModifiedGroupItem
 	return ModifiedGroupItemsDelegate;
 }
 
-uint32 UNiagaraStackItem::GetRecursiveStackIssuesCount() const
-{
-	if (RecursiveStackIssuesCount.IsSet() == false)
-	{
-		TArray<UNiagaraStackErrorItem*> RecursiveIssues;
-		FNiagaraStackGraphUtilities::GetStackIssuesRecursively(this, RecursiveIssues);
-		RecursiveStackIssuesCount = RecursiveIssues.Num();
-		EStackIssueSeverity MinSeverity = EStackIssueSeverity::Info;
-		for (auto Issue : RecursiveIssues)
-		{
-			if (Issue->GetStackIssue().GetSeverity() < MinSeverity)
-			{
-				MinSeverity = Issue->GetStackIssue().GetSeverity();
-			}
-		}
-		HighestIssueSeverity = MinSeverity;
-	}
-	return RecursiveStackIssuesCount.GetValue();
-}
-
-EStackIssueSeverity UNiagaraStackItem::GetHighestStackIssueSeverity() const
-{
-	if (HighestIssueSeverity.IsSet() == false)
-	{
-		GetRecursiveStackIssuesCount();
-	}
-	return HighestIssueSeverity.GetValue();
-}
-
 void UNiagaraStackItem::SetIsEnabled(bool bInIsEnabled)
 {
 	if (ItemFooter != nullptr)
@@ -83,8 +54,6 @@ void UNiagaraStackItem::RefreshChildrenInternal(const TArray<UNiagaraStackEntry*
 	ItemFooter->SetIsEnabled(GetIsEnabled());
 
 	NewChildren.Add(ItemFooter);
-	RecursiveStackIssuesCount.Reset();
-	HighestIssueSeverity.Reset();
 }
 
 void GetContentChildren(UNiagaraStackEntry& CurrentEntry, TArray<UNiagaraStackItemContent*>& ContentChildren)
@@ -124,13 +93,6 @@ int32 UNiagaraStackItem::GetChildIndentLevel() const
 	return GetIndentLevel();
 }
 
-void UNiagaraStackItem::ChlildStructureChangedInternal()
-{
-	Super::ChlildStructureChangedInternal();
-	RecursiveStackIssuesCount.Reset();
-	HighestIssueSeverity.Reset();
-}
-
 bool UNiagaraStackItem::FilterAdvancedChildren(const UNiagaraStackEntry& Child) const
 {
 	const UNiagaraStackItemContent* ItemContent = Cast<UNiagaraStackItemContent>(&Child);
diff --git a/Engine/Plugins/FX/Niagara/Source/NiagaraEditor/Private/ViewModels/Stack/NiagaraStackModuleItem.cpp b/Engine/Plugins/FX/Niagara/Source/NiagaraEditor/Private/ViewModels/Stack/NiagaraStackModuleItem.cpp
--- a/Engine/Plugins/FX/Niagara/Source/NiagaraEditor/Private/ViewModels/Stack/NiagaraStackModuleItem.cpp
+++ b/Engine/Plugins/FX/Niagara/Source/NiagaraEditor/Private/ViewModels/Stack/NiagaraStackModuleItem.cpp
@@ -522,7 +522,7 @@ void UNiagaraStackModuleItem::RefreshIssues(TArray<FStackIssue>& NewIssues)
 					{
 						FScopedTransaction ScopedTransaction(LOCTEXT("EnableDependencyModule", "Enable dependency module"));
 						FNiagaraStackGraphUtilities::SetModuleIsEnabled(*DisabledNode, true);
-						OnRequestFullRefresh().Broadcast();
+						OnRequestFullRefreshDeferred().Broadcast();
 
 					}));
 					Fixes.Add(Fix);
@@ -551,7 +551,7 @@ void UNiagaraStackModuleItem::RefreshIssues(TArray<FStackIssue>& NewIssues)
 							FNiagaraStackGraphUtilities::SetModuleIsEnabled(*DisorderedNode.ModuleNode, true);
 						}
 						FNiagaraStackGraphUtilities::RelayoutGraph(*OutputNode->GetGraph());
-						OnRequestFullRefresh().Broadcast();
+						OnRequestFullRefreshDeferred().Broadcast();
 					}));
 					Fixes.Add(Fix);
 				}
@@ -673,7 +673,7 @@ void UNiagaraStackModuleItem::RefreshIssues(TArray<FStackIssue>& NewIssues)
 							checkf(NewModuleNode != nullptr, TEXT("Add module action failed"));
 							FNiagaraStackGraphUtilities::InitializeStackFunctionInputs(GetSystemViewModel(), GetEmitterViewModel(), GetStackEditorData(), *NewModuleNode, *NewModuleNode);
 							FNiagaraStackGraphUtilities::RelayoutGraph(*TargetOutputNode->GetGraph());
-							OnRequestFullRefresh().Broadcast();
+							OnRequestFullRefreshDeferred().Broadcast();
 						}));
 						Fixes.Add(Fix);
 					}
@@ -808,7 +808,7 @@ void UNiagaraStackModuleItem::SetIsEnabledInternal(bool bInIsEnabled)
 	FScopedTransaction ScopedTransaction(LOCTEXT("EnableDisableModule", "Enable/Disable Module"));
 	FNiagaraStackGraphUtilities::SetModuleIsEnabled(*FunctionCallNode, bInIsEnabled);
 	bIsEnabled = bInIsEnabled;
-	OnRequestFullRefresh().Broadcast();
+	OnRequestFullRefreshDeferred().Broadcast();
 }
 
 bool UNiagaraStackModuleItem::TestCanDeleteWithMessage(FText& OutCanDeleteMessage) const
diff --git a/Engine/Plugins/FX/Niagara/Source/NiagaraEditor/Private/ViewModels/Stack/NiagaraStackViewModel.cpp b/Engine/Plugins/FX/Niagara/Source/NiagaraEditor/Private/ViewModels/Stack/NiagaraStackViewModel.cpp
--- a/Engine/Plugins/FX/Niagara/Source/NiagaraEditor/Private/ViewModels/Stack/NiagaraStackViewModel.cpp
+++ b/Engine/Plugins/FX/Niagara/Source/NiagaraEditor/Private/ViewModels/Stack/NiagaraStackViewModel.cpp
@@ -159,6 +159,11 @@ void UNiagaraStackViewModel::Reset()
 	bUsesTopLevelViewModels = false;
 }
 
+bool UNiagaraStackViewModel::HasIssues() const
+{
+	return bHasIssues;
+}
+
 void UNiagaraStackViewModel::Finalize()
 {
 	Reset();
@@ -232,7 +237,7 @@ void UNiagaraStackViewModel::AddSearchScrollOffset(int NumberOfSteps)
 
 void UNiagaraStackViewModel::CollapseToHeaders()
 {
-	CollapseToHeadersRecursive(GetRootEntries());
+	CollapseToHeadersRecursive(GetRootEntryAsArray());
 	NotifyStructureChanged();
 }
 
@@ -326,7 +331,7 @@ void UNiagaraStackViewModel::CollapseToHeadersRecursive(TArray<UNiagaraStackEntr
 	}
 }
 
-void UNiagaraStackViewModel::GetPathForEntry(UNiagaraStackEntry* Entry, TArray<UNiagaraStackEntry*>& EntryPath)
+void UNiagaraStackViewModel::GetPathForEntry(UNiagaraStackEntry* Entry, TArray<UNiagaraStackEntry*>& EntryPath) const
 {
 	GeneratePathForEntry(RootEntry, Entry, TArray<UNiagaraStackEntry*>(), EntryPath);
 }
@@ -435,7 +440,7 @@ bool UNiagaraStackViewModel::ItemMatchesSearchCriteria(UNiagaraStackEntry::FStac
 	return SearchItem.Value.ToString().Contains(CurrentSearchText.ToString());
 }
 
-void UNiagaraStackViewModel::GeneratePathForEntry(UNiagaraStackEntry* Root, UNiagaraStackEntry* Entry, TArray<UNiagaraStackEntry*> CurrentPath, TArray<UNiagaraStackEntry*>& EntryPath)
+void UNiagaraStackViewModel::GeneratePathForEntry(UNiagaraStackEntry* Root, UNiagaraStackEntry* Entry, TArray<UNiagaraStackEntry*> CurrentPath, TArray<UNiagaraStackEntry*>& EntryPath) const
 {
 	if (EntryPath.Num() > 0)
 	{
@@ -477,7 +482,12 @@ void UNiagaraStackViewModel::RestoreStackEntryExpansionPreSearch()
 	}
 }
 
-TArray<UNiagaraStackEntry*>& UNiagaraStackViewModel::GetRootEntries()
+UNiagaraStackEntry* UNiagaraStackViewModel::GetRootEntry()
+{
+	return RootEntry;
+}
+
+TArray<UNiagaraStackEntry*>& UNiagaraStackViewModel::GetRootEntryAsArray()
 {
 	return RootEntries;
 }
@@ -608,6 +618,7 @@ void UNiagaraStackViewModel::EntryStructureChanged()
 	{
 		RefreshTopLevelViewModels();
 	}
+	RefreshHasIssues();
 	StructureChangedDelegate.Broadcast();
 	OnSearchTextChanged(CurrentSearchText);
 }
@@ -676,4 +687,35 @@ void UNiagaraStackViewModel::RefreshTopLevelViewModels()
 	}
 }
 
+void UNiagaraStackViewModel::RefreshHasIssues()
+{
+	bHasIssues = false;
+	if (bUsesTopLevelViewModels)
+	{
+		for (TSharedRef<FTopLevelViewModel> TopLevelViewModel : TopLevelViewModels)
+		{
+			if (TopLevelViewModel->SystemViewModel.IsValid())
+			{
+				if (TopLevelViewModel->SystemViewModel->GetSystemStackViewModel()->GetRootEntry()->HasIssuesOrAnyChildHasIssues())
+				{
+					bHasIssues = true;
+					return;
+				}
+			}
+			else if (TopLevelViewModel->EmitterHandleViewModel.IsValid())
+			{
+				if (TopLevelViewModel->EmitterHandleViewModel->GetEmitterStackViewModel()->GetRootEntry()->HasIssuesOrAnyChildHasIssues())
+				{
+					bHasIssues = true;
+					return;
+				}
+			}
+		}
+	}
+	else
+	{
+		bHasIssues = RootEntry->HasIssuesOrAnyChildHasIssues();
+	}
+}
+
 #undef LOCTEXT_NAMESPACE
diff --git a/Engine/Plugins/FX/Niagara/Source/NiagaraEditor/Public/NiagaraEditorModule.h b/Engine/Plugins/FX/Niagara/Source/NiagaraEditor/Public/NiagaraEditorModule.h
--- a/Engine/Plugins/FX/Niagara/Source/NiagaraEditor/Public/NiagaraEditorModule.h
+++ b/Engine/Plugins/FX/Niagara/Source/NiagaraEditor/Public/NiagaraEditorModule.h
@@ -15,6 +15,8 @@ class INiagaraEditorTypeUtilities;
 class UNiagaraSettings;
 class USequencerSettings;
 class UNiagaraStackViewModel;
+class UNiagaraStackEntry;
+class UNiagaraStackIssue;
 class FNiagaraSystemViewModel;
 class FNiagaraScriptMergeManager;
 class FNiagaraCompileOptions;
@@ -31,6 +33,7 @@ class NIAGARAEDITOR_API INiagaraEditorWidgetProvider
 public:
 	virtual TSharedRef<SWidget> CreateStackView(UNiagaraStackViewModel& StackViewModel) const = 0;
 	virtual TSharedRef<SWidget> CreateSystemOverview(TSharedRef<FNiagaraSystemViewModel> SystemViewModel) const = 0;
+	virtual TSharedRef<SWidget> CreateStackIssueIcon(UNiagaraStackViewModel& StackViewModel, UNiagaraStackEntry& StackEntry) const = 0;
 	virtual FLinearColor GetColorForExecutionCategory(FName ExecutionCategory) const = 0;
 };
 
diff --git a/Engine/Plugins/FX/Niagara/Source/NiagaraEditor/Public/ViewModels/Stack/NiagaraStackEntry.h b/Engine/Plugins/FX/Niagara/Source/NiagaraEditor/Public/ViewModels/Stack/NiagaraStackEntry.h
--- a/Engine/Plugins/FX/Niagara/Source/NiagaraEditor/Public/ViewModels/Stack/NiagaraStackEntry.h
+++ b/Engine/Plugins/FX/Niagara/Source/NiagaraEditor/Public/ViewModels/Stack/NiagaraStackEntry.h
@@ -177,7 +177,7 @@ public:
 		FString UniqueIdentifier;
 	};
 
-	struct FStackIssue
+	struct NIAGARAEDITOR_API FStackIssue
 	{
 		FStackIssue();
 
@@ -313,6 +313,18 @@ public:
 
 	bool HasBaseEmitter() const;
 
+	bool HasIssuesOrAnyChildHasIssues() const;
+
+	int32 GetTotalNumberOfInfoIssues() const;
+
+	int32 GetTotalNumberOfWarningIssues() const;
+
+	int32 GetTotalNumberOfErrorIssues() const;
+
+	const TArray<FStackIssue>& GetIssues() const;
+
+	const TArray<UNiagaraStackEntry*>& GetAllChildrenWithIssues() const;
+
 protected:
 	virtual void BeginDestroy() override;
 
@@ -389,6 +401,8 @@ private:
 	
 	TArray<FStackIssue> StackIssues;
 
+	TArray<UNiagaraStackEntry*> ChildrenWithIssues;
+
 	bool bIsFinalized;
 
 	bool bIsSearchResult;
@@ -396,4 +410,8 @@ private:
 	mutable TOptional<bool> bHasBaseEmitterCache;
 
 	bool bOwnerIsEnabled;
+
+	int32 TotalNumberOfInfoIssues;
+	int32 TotalNumberOfWarningIssues;
+	int32 TotalNumberOfErrorIssues;
 };
\ No newline at end of file
diff --git a/Engine/Plugins/FX/Niagara/Source/NiagaraEditor/Public/ViewModels/Stack/NiagaraStackItem.h b/Engine/Plugins/FX/Niagara/Source/NiagaraEditor/Public/ViewModels/Stack/NiagaraStackItem.h
--- a/Engine/Plugins/FX/Niagara/Source/NiagaraEditor/Public/ViewModels/Stack/NiagaraStackItem.h
+++ b/Engine/Plugins/FX/Niagara/Source/NiagaraEditor/Public/ViewModels/Stack/NiagaraStackItem.h
@@ -24,9 +24,6 @@ public:
 
 	FOnModifiedGroupItems& OnModifiedGroupItems();
 
-	uint32 GetRecursiveStackIssuesCount() const;
-	EStackIssueSeverity GetHighestStackIssueSeverity() const;
-
 	virtual bool SupportsChangeEnabled() const { return false; }
 	void SetIsEnabled(bool bInIsEnabled);
 
@@ -41,8 +38,6 @@ protected:
 
 	virtual int32 GetChildIndentLevel() const override;
 
-	virtual void ChlildStructureChangedInternal() override;
-
 	virtual void SetIsEnabledInternal(bool bInIsEnabled) { }
 	virtual void DeleteInternal() { }
 
@@ -57,11 +52,6 @@ protected:
 private:
 	UPROPERTY()
 	UNiagaraStackItemFooter* ItemFooter;
-
-	/** How many errors this entry has along its tree. */
-	mutable TOptional<uint32> RecursiveStackIssuesCount;
-	/** The highest severity of issues along this entry's tree. */
-	mutable TOptional<EStackIssueSeverity> HighestIssueSeverity;
 };
 
 UCLASS()
diff --git a/Engine/Plugins/FX/Niagara/Source/NiagaraEditor/Public/ViewModels/Stack/NiagaraStackViewModel.h b/Engine/Plugins/FX/Niagara/Source/NiagaraEditor/Public/ViewModels/Stack/NiagaraStackViewModel.h
--- a/Engine/Plugins/FX/Niagara/Source/NiagaraEditor/Public/ViewModels/Stack/NiagaraStackViewModel.h
+++ b/Engine/Plugins/FX/Niagara/Source/NiagaraEditor/Public/ViewModels/Stack/NiagaraStackViewModel.h
@@ -90,7 +90,9 @@ public:
 
 	TSharedPtr<FNiagaraSystemViewModel> GetSystemViewModel();
 
-	TArray<UNiagaraStackEntry*>& GetRootEntries();
+	UNiagaraStackEntry* GetRootEntry();
+
+	TArray<UNiagaraStackEntry*>& GetRootEntryAsArray();
 
 	FOnStructureChanged& OnStructureChanged();
 	FOnSearchCompleted& OnSearchCompleted();
@@ -122,7 +124,7 @@ public:
 	UNiagaraStackEntry* GetCurrentFocusedEntry();
 	void AddSearchScrollOffset(int NumberOfSteps);
 
-	void GetPathForEntry(UNiagaraStackEntry* Entry, TArray<UNiagaraStackEntry*>& EntryPath);
+	void GetPathForEntry(UNiagaraStackEntry* Entry, TArray<UNiagaraStackEntry*>& EntryPath) const;
 
 	/** Starts recursing through all entries to expand all groups and collapse all items. */
 	void CollapseToHeaders();
@@ -137,6 +139,8 @@ public:
 
 	void Reset();
 
+	bool HasIssues() const;
+
 private:
 	/** Recursively Expands all groups and collapses all items in the stack. */
 	void CollapseToHeadersRecursive(TArray<UNiagaraStackEntry*> Entries);
@@ -157,6 +161,7 @@ private:
 	void EntryRequestFullRefresh();
 	void EntryRequestFullRefreshDeferred();
 	void RefreshTopLevelViewModels();
+	void RefreshHasIssues();
 	void OnSystemCompiled();
 	void OnEmitterCompiled();
 	void EmitterParentRemoved();
@@ -165,7 +170,7 @@ private:
 	void GenerateTraversalEntries(UNiagaraStackEntry* Root, TArray<UNiagaraStackEntry*> ParentChain, 
 		TArray<FSearchWorkItem>& TraversedArray);
 	bool ItemMatchesSearchCriteria(UNiagaraStackEntry::FStackSearchItem SearchItem);
-	void GeneratePathForEntry(UNiagaraStackEntry* Root, UNiagaraStackEntry* Entry, TArray<UNiagaraStackEntry*> CurrentPath, TArray<UNiagaraStackEntry*>& EntryPath);
+	void GeneratePathForEntry(UNiagaraStackEntry* Root, UNiagaraStackEntry* Entry, TArray<UNiagaraStackEntry*> CurrentPath, TArray<UNiagaraStackEntry*>& EntryPath) const;
 	void RestoreStackEntryExpansionPreSearch();
 
 private:
@@ -191,6 +196,7 @@ private:
 	static const double MaxSearchTime;
 	bool bRestartSearch;
 	bool bRefreshPending;
+	bool bHasIssues;
 
 	bool bUsesTopLevelViewModels;
 	TArray<TSharedRef<FTopLevelViewModel>> TopLevelViewModels;
diff --git a/Engine/Plugins/FX/Niagara/Source/NiagaraEditorWidgets/Private/NiagaraEditorWidgetsModule.cpp b/Engine/Plugins/FX/Niagara/Source/NiagaraEditorWidgets/Private/NiagaraEditorWidgetsModule.cpp
--- a/Engine/Plugins/FX/Niagara/Source/NiagaraEditorWidgets/Private/NiagaraEditorWidgetsModule.cpp
+++ b/Engine/Plugins/FX/Niagara/Source/NiagaraEditorWidgets/Private/NiagaraEditorWidgetsModule.cpp
@@ -9,6 +9,7 @@
 #include "ViewModels/NiagaraSystemViewModel.h"
 #include "SNiagaraOverviewGraph.h"
 #include "NiagaraEditorWidgetsUtilities.h"
+#include "Stack/SNiagaraStackIssueIcon.h"
 
 #include "Modules/ModuleManager.h"
 #include "PropertyEditorModule.h"
@@ -161,6 +162,11 @@ TSharedRef<SWidget> FNiagaraEditorWidgetsModule::FNiagaraEditorWidgetProvider::C
 	return SNew(SNiagaraOverviewGraph, SystemViewModel->GetOverviewGraphViewModel().ToSharedRef());
 }
 
+TSharedRef<SWidget> FNiagaraEditorWidgetsModule::FNiagaraEditorWidgetProvider::CreateStackIssueIcon(UNiagaraStackViewModel& StackViewModel, UNiagaraStackEntry& StackEntry) const
+{
+	return SNew(SNiagaraStackIssueIcon, &StackViewModel, &StackEntry);
+}
+
 FLinearColor FNiagaraEditorWidgetsModule::FNiagaraEditorWidgetProvider::GetColorForExecutionCategory(FName ExecutionCategory) const
 {
 	return FNiagaraEditorWidgetsStyle::Get().GetColor(FNiagaraStackEditorWidgetsUtilities::GetIconColorNameForExecutionCategory(ExecutionCategory));
diff --git a/Engine/Plugins/FX/Niagara/Source/NiagaraEditorWidgets/Private/SNiagaraOverviewStack.cpp b/Engine/Plugins/FX/Niagara/Source/NiagaraEditorWidgets/Private/SNiagaraOverviewStack.cpp
--- a/Engine/Plugins/FX/Niagara/Source/NiagaraEditorWidgets/Private/SNiagaraOverviewStack.cpp
+++ b/Engine/Plugins/FX/Niagara/Source/NiagaraEditorWidgets/Private/SNiagaraOverviewStack.cpp
@@ -14,6 +14,7 @@
 #include "NiagaraEditorWidgetsUtilities.h"
 #include "SNiagaraStack.h"
 #include "Stack/SNiagaraStackItemGroupAddButton.h"
+#include "Stack/SNiagaraStackIssueIcon.h"
 #include "NiagaraEditorCommon.h"
 
 #include "Widgets/Layout/SScrollBox.h"
@@ -34,14 +35,16 @@
 class SNiagaraSystemOverviewEntryListRow : public STableRow<UNiagaraStackEntry*>
 {
 	SLATE_BEGIN_ARGS(SNiagaraSystemOverviewEntryListRow) {}
+		SLATE_ATTRIBUTE(EVisibility, IssueIconVisibility)
 		SLATE_EVENT(FOnDragDetected, OnDragDetected)
 		SLATE_EVENT(FOnCanAcceptDrop, OnCanAcceptDrop)
 		SLATE_EVENT(FOnAcceptDrop, OnAcceptDrop)
 		SLATE_DEFAULT_SLOT(FArguments, Content);
 	SLATE_END_ARGS();
 
-	void Construct(const FArguments& InArgs, UNiagaraStackEntry* InStackEntry, const TSharedRef<STableViewBase>& InOwnerTableView)
+	void Construct(const FArguments& InArgs, UNiagaraStackViewModel* InStackViewModel, UNiagaraStackEntry* InStackEntry, const TSharedRef<STableViewBase>& InOwnerTableView)
 	{
+		StackViewModel = InStackViewModel;
 		StackEntry = InStackEntry;
 		FSlateColor IconColor = FNiagaraEditorWidgetsStyle::Get().GetColor(FNiagaraStackEditorWidgetsUtilities::GetColorNameForExecutionCategory(StackEntry->GetExecutionCategoryName()));
 		ItemBackgroundColor = FNiagaraEditorWidgetsStyle::Get().GetColor("NiagaraEditor.Stack.Item.HeaderBackgroundColor");
@@ -88,9 +91,22 @@ class SNiagaraSystemOverviewEntryListRow : public STableRow<UNiagaraStackEntry*>
 		[
 			SNew(SBorder)
 			.BorderImage(this, &SNiagaraSystemOverviewEntryListRow::GetBorder)
-			.Padding(FMargin(5, 2, 5, 2))
+			.Padding(FMargin(5, 2, 3, 2))
 			[
-				WrappedContent.ToSharedRef()
+				SNew(SHorizontalBox)
+				+ SHorizontalBox::Slot()
+				.Padding(0, 0, 2, 0)
+				[
+					WrappedContent.ToSharedRef()
+				]
+				+ SHorizontalBox::Slot()
+				.AutoWidth()
+				.VAlign(VAlign_Center)
+				.Padding(1, 0, 0, 0)
+				[
+					SNew(SNiagaraStackIssueIcon, StackViewModel, StackEntry)
+					.Visibility(InArgs._IssueIconVisibility)
+				]
 			]
 		],
 		InOwnerTableView);
@@ -133,6 +149,7 @@ private:
 	}
 
 private:
+	UNiagaraStackViewModel* StackViewModel;
 	UNiagaraStackEntry* StackEntry;
 	FLinearColor ItemBackgroundColor;
 	FLinearColor DisabledItemBackgroundColor;
@@ -328,18 +345,18 @@ void SNiagaraOverviewStack::RefreshEntryList()
 		TArray<UClass*> AcceptableClasses;
 		AcceptableClasses.Add(UNiagaraStackItemGroup::StaticClass());
 		AcceptableClasses.Add(UNiagaraStackItem::StaticClass());
-		for (UNiagaraStackEntry* RootEntry : StackViewModel->GetRootEntries())
+
+		UNiagaraStackEntry* RootEntry = StackViewModel->GetRootEntry();
+		checkf(RootEntry != nullptr, TEXT("Root entry was null."));
+		TArray<UNiagaraStackEntry*> RootChildren;
+		RootEntry->GetFilteredChildren(RootChildren);
+		for (UNiagaraStackEntry* RootChild : RootChildren)
 		{
-			checkf(RootEntry != nullptr, TEXT("Root entry was null."));
-			TArray<UNiagaraStackEntry*> RootChildren;
-			RootEntry->GetFilteredChildren(RootChildren);
-			for (UNiagaraStackEntry* RootChild : RootChildren)
-			{
-				checkf(RootEntry != nullptr, TEXT("Root entry child was null."));
-				TArray<UNiagaraStackEntry*> ParentChain;
-				AddEntriesRecursive(*RootChild, FlattenedEntryList, AcceptableClasses, ParentChain);
-			}
+			checkf(RootEntry != nullptr, TEXT("Root entry child was null."));
+			TArray<UNiagaraStackEntry*> ParentChain;
+			AddEntriesRecursive(*RootChild, FlattenedEntryList, AcceptableClasses, ParentChain);
 		}
+
 		bRefreshEntryListPending = false;
 		EntryListView->RequestListRefresh();
 	}
@@ -357,6 +374,7 @@ TSharedRef<ITableRow> SNiagaraOverviewStack::OnGenerateRowForEntry(UNiagaraStack
 	{
 		UNiagaraStackItem* StackItem = CastChecked<UNiagaraStackItem>(Item);
 		Content = SNew(SHorizontalBox)
+			// Name
 			+ SHorizontalBox::Slot()
 			.VAlign(VAlign_Center)
 			[
@@ -365,6 +383,7 @@ TSharedRef<ITableRow> SNiagaraOverviewStack::OnGenerateRowForEntry(UNiagaraStack
 				.Text_UObject(Item, &UNiagaraStackEntry::GetDisplayName)
 				.IsEnabled_UObject(Item, &UNiagaraStackEntry::GetIsEnabledAndOwnerIsEnabled)
 			]
+			// Enabled checkbox
 			+ SHorizontalBox::Slot()
 			.VAlign(VAlign_Center)
 			.AutoWidth()
@@ -381,8 +400,9 @@ TSharedRef<ITableRow> SNiagaraOverviewStack::OnGenerateRowForEntry(UNiagaraStack
 	{
 		UNiagaraStackItemGroup* StackItemGroup = CastChecked<UNiagaraStackItemGroup>(Item);
 		TSharedRef<SHorizontalBox> ContentBox = SNew(SHorizontalBox)
+			// Execution category icon
 			+ SHorizontalBox::Slot()
-			.Padding(2, 0, 6, 0)
+			.Padding(0, 0, 6, 0)
 			.VAlign(VAlign_Center)
 			.AutoWidth()
 			[
@@ -391,6 +411,7 @@ TSharedRef<ITableRow> SNiagaraOverviewStack::OnGenerateRowForEntry(UNiagaraStack
 				.ColorAndOpacity(FNiagaraEditorWidgetsStyle::Get().GetColor(FNiagaraStackEditorWidgetsUtilities::GetIconColorNameForExecutionCategory(Item->GetExecutionCategoryName())))
 				.IsEnabled_UObject(Item, &UNiagaraStackEntry::GetIsEnabledAndOwnerIsEnabled)
 			]
+			// Name
 			+ SHorizontalBox::Slot()
 			.VAlign(VAlign_Center)
 			[
@@ -418,10 +439,11 @@ TSharedRef<ITableRow> SNiagaraOverviewStack::OnGenerateRowForEntry(UNiagaraStack
 	}
 
 
-	return SNew(SNiagaraSystemOverviewEntryListRow, Item, OwnerTable)
+	return SNew(SNiagaraSystemOverviewEntryListRow, StackViewModel, Item, OwnerTable)
 		.OnDragDetected(this, &SNiagaraOverviewStack::OnRowDragDetected, TWeakObjectPtr<UNiagaraStackEntry>(Item))
 		.OnCanAcceptDrop(this, &SNiagaraOverviewStack::OnRowCanAcceptDrop)
 		.OnAcceptDrop(this, &SNiagaraOverviewStack::OnRowAcceptDrop)
+		.IssueIconVisibility(this, &SNiagaraOverviewStack::GetIssueIconVisibility)
 	[
 		Content.ToSharedRef()
 	];
@@ -525,4 +547,9 @@ FReply SNiagaraOverviewStack::OnRowAcceptDrop(const FDragDropEvent& InDragDropEv
 	return bHandled ? FReply::Handled() : FReply::Unhandled();
 }
 
+EVisibility SNiagaraOverviewStack::GetIssueIconVisibility() const
+{
+	return StackViewModel->HasIssues() ? EVisibility::Visible : EVisibility::Collapsed;
+}
+
 #undef LOCTEXT_NAMESPACE
\ No newline at end of file
diff --git a/Engine/Plugins/FX/Niagara/Source/NiagaraEditorWidgets/Private/SNiagaraOverviewStack.h b/Engine/Plugins/FX/Niagara/Source/NiagaraEditorWidgets/Private/SNiagaraOverviewStack.h
--- a/Engine/Plugins/FX/Niagara/Source/NiagaraEditorWidgets/Private/SNiagaraOverviewStack.h
+++ b/Engine/Plugins/FX/Niagara/Source/NiagaraEditorWidgets/Private/SNiagaraOverviewStack.h
@@ -56,6 +56,8 @@ private:
 
 	FReply OnRowAcceptDrop(const FDragDropEvent& InDragDropEvent, EItemDropZone InDropZone, UNiagaraStackEntry* InTargetEntry);
 
+	EVisibility GetIssueIconVisibility() const;
+
 private:
 	UNiagaraStackViewModel* StackViewModel;
 	UNiagaraSystemSelectionViewModel* OverviewSelectionViewModel;
diff --git a/Engine/Plugins/FX/Niagara/Source/NiagaraEditorWidgets/Private/SNiagaraOverviewStackNode.cpp b/Engine/Plugins/FX/Niagara/Source/NiagaraEditorWidgets/Private/SNiagaraOverviewStackNode.cpp
--- a/Engine/Plugins/FX/Niagara/Source/NiagaraEditorWidgets/Private/SNiagaraOverviewStackNode.cpp
+++ b/Engine/Plugins/FX/Niagara/Source/NiagaraEditorWidgets/Private/SNiagaraOverviewStackNode.cpp
@@ -4,9 +4,11 @@
 #include "NiagaraOverviewNode.h"
 #include "ViewModels/NiagaraSystemViewModel.h"
 #include "ViewModels/NiagaraEmitterHandleViewModel.h"
+#include "ViewModels/Stack/NiagaraStackViewModel.h"
 #include "SNiagaraOverviewStack.h"
 #include "NiagaraEditorModule.h"
 #include "NiagaraEditorWidgetsStyle.h"
+#include "Stack/SNiagaraStackIssueIcon.h"
 
 #include "Modules/ModuleManager.h"
 #include "Widgets/Layout/SBox.h"
@@ -45,18 +47,30 @@ TSharedRef<SWidget> SNiagaraOverviewStackNode::CreateTitleWidget(TSharedPtr<SNod
 	TSharedRef<SWidget> DefaultTitle = SGraphNode::CreateTitleWidget(NodeTitle);
 	
 	return SNew(SHorizontalBox)
+		// Enabled checkbox
 		+ SHorizontalBox::Slot()
 		.AutoWidth()
-		.Padding(0, 0, 5, 0)
+		.Padding(0)
 		[
 			SNew(SCheckBox)
 			.Visibility(this, &SNiagaraOverviewStackNode::GetEnabledCheckBoxVisibility)
 			.IsChecked(this, &SNiagaraOverviewStackNode::GetEnabledCheckState)
 			.OnCheckStateChanged(this, &SNiagaraOverviewStackNode::OnEnabledCheckStateChanged)
 		]
+		// Name
 		+ SHorizontalBox::Slot()
+		.Padding(4, 0, 0, 0)
 		[
 			DefaultTitle
+		]
+		// Stack issues icon
+		+ SHorizontalBox::Slot()
+		.AutoWidth()
+		.VAlign(VAlign_Center)
+		.Padding(5, 0, 0, 0)
+		[
+			SNew(SNiagaraStackIssueIcon, StackViewModel, StackViewModel->GetRootEntry())
+			.Visibility(this, &SNiagaraOverviewStackNode::GetIssueIconVisibility)
 		];
 }
 
@@ -113,6 +127,11 @@ TSharedRef<SWidget> SNiagaraOverviewStackNode::CreateNodeContentArea()
 		];
 }
 
+EVisibility SNiagaraOverviewStackNode::GetIssueIconVisibility() const
+{
+	return StackViewModel->HasIssues() ? EVisibility::Visible : EVisibility::Collapsed;
+}
+
 EVisibility SNiagaraOverviewStackNode::GetEnabledCheckBoxVisibility() const
 {
 	return EmitterHandleViewModelWeak.IsValid() ? EVisibility::Visible : EVisibility::Collapsed;
diff --git a/Engine/Plugins/FX/Niagara/Source/NiagaraEditorWidgets/Private/SNiagaraOverviewStackNode.h b/Engine/Plugins/FX/Niagara/Source/NiagaraEditorWidgets/Private/SNiagaraOverviewStackNode.h
--- a/Engine/Plugins/FX/Niagara/Source/NiagaraEditorWidgets/Private/SNiagaraOverviewStackNode.h
+++ b/Engine/Plugins/FX/Niagara/Source/NiagaraEditorWidgets/Private/SNiagaraOverviewStackNode.h
@@ -23,6 +23,7 @@ protected:
 	virtual TSharedRef<SWidget> CreateNodeContentArea() override;
 
 private:
+	EVisibility GetIssueIconVisibility() const;
 	EVisibility GetEnabledCheckBoxVisibility() const;
 	ECheckBoxState GetEnabledCheckState() const;
 	void OnEnabledCheckStateChanged(ECheckBoxState InCheckState);
diff --git a/Engine/Plugins/FX/Niagara/Source/NiagaraEditorWidgets/Private/SNiagaraStack.cpp b/Engine/Plugins/FX/Niagara/Source/NiagaraEditorWidgets/Private/SNiagaraStack.cpp
--- a/Engine/Plugins/FX/Niagara/Source/NiagaraEditorWidgets/Private/SNiagaraStack.cpp
+++ b/Engine/Plugins/FX/Niagara/Source/NiagaraEditorWidgets/Private/SNiagaraStack.cpp
@@ -46,6 +46,7 @@
 #include "Stack/SNiagaraStackParameterStoreEntryValue.h"
 #include "Stack/SNiagaraStackRendererItem.h"
 #include "Stack/SNiagaraStackTableRow.h"
+#include "Stack/SNiagaraStackIssueIcon.h"
 #include "NiagaraEditorWidgetsUtilities.h"
 #include "DragAndDrop/DecoratedDragDropOp.h"
 #include "Stack/SNiagaraStackErrorItem.h"
@@ -61,6 +62,7 @@ class SNiagaraStackEmitterHeader : public SCompoundWidget
 {
 public:
 	SLATE_BEGIN_ARGS(SNiagaraStackEmitterHeader) {}
+		SLATE_ATTRIBUTE(EVisibility, IssueIconVisibility);
 	SLATE_END_ARGS();
 
 	void Construct(const FArguments& InArgs, TSharedRef<FNiagaraEmitterHandleViewModel> InEmitterHandleViewModel)
@@ -76,11 +78,11 @@ public:
 			.AutoHeight()
 			.HAlign(HAlign_Fill)
 			[
-				//~ Enabled
 				SNew(SHorizontalBox)
+				//~ Enabled
 				+ SHorizontalBox::Slot()
 				.AutoWidth()
-				.Padding(2)
+				.Padding(2, 0, 0, 0)
 				[
 					SNew(SCheckBox)
 					.ToolTipText(LOCTEXT("EnabledToolTip", "Toggles whether this emitter is enabled. Disabled emitters don't simulate or render."))
@@ -96,18 +98,31 @@ public:
 					.UseAllottedWidth(true)
 					+ SWrapBox::Slot()
 					.Padding(3, 0)
-					[						
-						SAssignNew(EmitterNameTextBlock, SInlineEditableTextBlock)
-						.ToolTipText(this, &SNiagaraStackEmitterHeader::GetEmitterNameToolTip)
-						.Style(FNiagaraEditorStyle::Get(), "NiagaraEditor.HeadingInlineEditableText") 
-						.Clipping(EWidgetClipping::ClipToBoundsAlways)
-						.Text(EmitterHandleViewModel.ToSharedRef(), &FNiagaraEmitterHandleViewModel::GetNameText)
-						.OnTextCommitted(EmitterHandleViewModel.ToSharedRef(), &FNiagaraEmitterHandleViewModel::OnNameTextComitted)
-						.OnVerifyTextChanged(EmitterHandleViewModel.ToSharedRef(), &FNiagaraEmitterHandleViewModel::VerifyNameTextChanged)
-						.IsReadOnly(EmitterHandleViewModel->CanRenameEmitter() == false)
+					[
+						SNew(SHorizontalBox)
+						+ SHorizontalBox::Slot()
+						[
+							SAssignNew(EmitterNameTextBlock, SInlineEditableTextBlock)
+							.ToolTipText(this, &SNiagaraStackEmitterHeader::GetEmitterNameToolTip)
+							.Style(FNiagaraEditorStyle::Get(), "NiagaraEditor.HeadingInlineEditableText") 
+							.Clipping(EWidgetClipping::ClipToBoundsAlways)
+							.Text(EmitterHandleViewModel.ToSharedRef(), &FNiagaraEmitterHandleViewModel::GetNameText)
+							.OnTextCommitted(EmitterHandleViewModel.ToSharedRef(), &FNiagaraEmitterHandleViewModel::OnNameTextComitted)
+							.OnVerifyTextChanged(EmitterHandleViewModel.ToSharedRef(), &FNiagaraEmitterHandleViewModel::VerifyNameTextChanged)
+							.IsReadOnly(EmitterHandleViewModel->CanRenameEmitter() == false)
+						]
+						// Issue Icon
+						+ SHorizontalBox::Slot()
+						.AutoWidth()
+						.VAlign(VAlign_Center)
+						.Padding(4, 0, 0, 0)
+						[
+							SNew(SNiagaraStackIssueIcon, EmitterHandleViewModel->GetEmitterStackViewModel(), EmitterHandleViewModel->GetEmitterStackViewModel()->GetRootEntry())
+							.Visibility(InArgs._IssueIconVisibility)
+						]
 					]
 					+ SWrapBox::Slot()
-					.Padding(3, 0)
+					.Padding(4, 0)
 					[
 						SNew(STextBlock)
 						.ToolTipText(this, &SNiagaraStackEmitterHeader::GetEmitterNameToolTip)
@@ -246,8 +261,9 @@ void SNiagaraStack::Construct(const FArguments& InArgs, UNiagaraStackViewModel*
 			SAssignNew(StackTree, STreeView<UNiagaraStackEntry*>)
 			.OnGenerateRow(this, &SNiagaraStack::OnGenerateRowForStackItem)
 			.OnGetChildren(this, &SNiagaraStack::OnGetChildren)
-			.TreeItemsSource(&StackViewModel->GetRootEntries())
+			.TreeItemsSource(&StackViewModel->GetRootEntryAsArray())
 			.OnTreeViewScrolled(this, &SNiagaraStack::StackTreeScrolled)
+			.SelectionMode(ESelectionMode::None)
 		]
 	];
 
@@ -258,7 +274,7 @@ void SNiagaraStack::Construct(const FArguments& InArgs, UNiagaraStackViewModel*
 
 void SNiagaraStack::SynchronizeTreeExpansion()
 {
-	TArray<UNiagaraStackEntry*> EntriesToProcess(StackViewModel->GetRootEntries());
+	TArray<UNiagaraStackEntry*> EntriesToProcess(StackViewModel->GetRootEntryAsArray());
 	while (EntriesToProcess.Num() > 0)
 	{
 		UNiagaraStackEntry* EntryToProcess = EntriesToProcess[0];
@@ -477,7 +493,7 @@ void CollapseEntriesRecursive(TArray<UNiagaraStackEntry*> Entries)
 
 void SNiagaraStack::CollapseAll()
 {
-	CollapseEntriesRecursive(StackViewModel->GetRootEntries());
+	CollapseEntriesRecursive(StackViewModel->GetRootEntryAsArray());
 	StackViewModel->NotifyStructureChanged();
 }
 
@@ -608,13 +624,28 @@ TSharedRef<ITableRow> SNiagaraStack::OnGenerateRowForTopLevelObject(TSharedRef<U
 	TSharedPtr<SWidget> Content;
 	if (Item->SystemViewModel.IsValid())
 	{
-		Content = SNew(STextBlock)
-			.TextStyle(FNiagaraEditorStyle::Get(), "NiagaraEditor.HeadingTextBlock")
-			.Text(Item->SystemViewModel.ToSharedRef(), &FNiagaraSystemViewModel::GetDisplayName);
+		Content = SNew(SHorizontalBox)
+			// System name
+			+ SHorizontalBox::Slot()
+			[
+				SNew(STextBlock)
+				.TextStyle(FNiagaraEditorStyle::Get(), "NiagaraEditor.HeadingTextBlock")
+				.Text(Item->SystemViewModel.ToSharedRef(), &FNiagaraSystemViewModel::GetDisplayName)
+			]
+			// Issue Icon
+			+ SHorizontalBox::Slot()
+			.AutoWidth()
+			.VAlign(VAlign_Center)
+			.Padding(4, 0, 2, 0)
+			[
+				SNew(SNiagaraStackIssueIcon, Item->SystemViewModel->GetSystemStackViewModel(), Item->SystemViewModel->GetSystemStackViewModel()->GetRootEntry())
+				.Visibility(this, &SNiagaraStack::GetIssueIconVisibility)
+			];
 	}
 	else if(Item->EmitterHandleViewModel.IsValid())
 	{
-		Content = SNew(SNiagaraStackEmitterHeader, Item->EmitterHandleViewModel.ToSharedRef());
+		Content = SNew(SNiagaraStackEmitterHeader, Item->EmitterHandleViewModel.ToSharedRef())
+			.IssueIconVisibility(this, &SNiagaraStack::GetIssueIconVisibility);
 	}
 
 	Content->SetOnMouseButtonUp(FPointerEventHandler::CreateSP(this, &SNiagaraStack::OnTopLevelRowMouseButtonDown, TWeakPtr<UNiagaraStackViewModel::FTopLevelViewModel>(Item)));
@@ -804,7 +835,8 @@ TSharedRef<SNiagaraStackTableRow> SNiagaraStack::ConstructContainerForItem(UNiag
 		.OnValueColumnWidthChanged(this, &SNiagaraStack::OnContentColumnWidthChanged)
 		.OnDragDetected(this, &SNiagaraStack::OnRowDragDetected, Item)
 		.OnCanAcceptDrop(this, &SNiagaraStack::OnRowCanAcceptDrop)
-		.OnAcceptDrop(this, &SNiagaraStack::OnRowAcceptDrop);
+		.OnAcceptDrop(this, &SNiagaraStack::OnRowAcceptDrop)
+		.IssueIconVisibility(this, &SNiagaraStack::GetIssueIconVisibility);
 }
 
 void SNiagaraStack::Tick(const FGeometry& AllottedGeometry, const double InCurrentTime, const float InDeltaTime)
@@ -1001,4 +1033,9 @@ void SNiagaraStack::StackStructureChanged()
 	HeaderList->RequestListRefresh();
 }
 
+EVisibility SNiagaraStack::GetIssueIconVisibility() const
+{
+	return StackViewModel->HasIssues() ? EVisibility::Visible : EVisibility::Collapsed;
+}
+
 #undef LOCTEXT_NAMESPACE
diff --git a/Engine/Plugins/FX/Niagara/Source/NiagaraEditorWidgets/Private/SNiagaraStack.h b/Engine/Plugins/FX/Niagara/Source/NiagaraEditorWidgets/Private/SNiagaraStack.h
--- a/Engine/Plugins/FX/Niagara/Source/NiagaraEditorWidgets/Private/SNiagaraStack.h
+++ b/Engine/Plugins/FX/Niagara/Source/NiagaraEditorWidgets/Private/SNiagaraStack.h
@@ -102,6 +102,8 @@ private:
 
 	FReply OnRowAcceptDrop(const FDragDropEvent& InDragDropEvent, EItemDropZone InDropZone, UNiagaraStackEntry* InTargetEntry);
 
+	EVisibility GetIssueIconVisibility() const;
+
 private:
 	UNiagaraStackViewModel* StackViewModel;
 
@@ -117,5 +119,4 @@ private:
 	TSharedPtr<SSearchBox> SearchBox;
 	static const FText OccurencesFormat;
 	bool bNeedsJumpToNextOccurence;
-	
 };
diff --git a/Engine/Plugins/FX/Niagara/Source/NiagaraEditorWidgets/Private/Stack/SNiagaraStackErrorButton.cpp b/Engine/Plugins/FX/Niagara/Source/NiagaraEditorWidgets/Private/Stack/SNiagaraStackErrorButton.cpp
deleted file mode 100644
diff --git a/Engine/Plugins/FX/Niagara/Source/NiagaraEditorWidgets/Private/Stack/SNiagaraStackErrorButton.h b/Engine/Plugins/FX/Niagara/Source/NiagaraEditorWidgets/Private/Stack/SNiagaraStackErrorButton.h
deleted file mode 100644
diff --git a/Engine/Plugins/FX/Niagara/Source/NiagaraEditorWidgets/Private/Stack/SNiagaraStackErrorItem.cpp b/Engine/Plugins/FX/Niagara/Source/NiagaraEditorWidgets/Private/Stack/SNiagaraStackErrorItem.cpp
--- a/Engine/Plugins/FX/Niagara/Source/NiagaraEditorWidgets/Private/Stack/SNiagaraStackErrorItem.cpp
+++ b/Engine/Plugins/FX/Niagara/Source/NiagaraEditorWidgets/Private/Stack/SNiagaraStackErrorItem.cpp
@@ -18,6 +18,22 @@ void SNiagaraStackErrorItem::Construct(const FArguments& InArgs, UNiagaraStackEr
 	ErrorItem = InErrorItem;
 	StackViewModel = InStackViewModel;
 	TSharedPtr<SHorizontalBox> ErrorInternalBox = SNew(SHorizontalBox);
+	const FSlateBrush* IconBrush;
+	switch (ErrorItem->GetStackIssue().GetSeverity())
+	{
+	case EStackIssueSeverity::Error:
+		IconBrush = FEditorStyle::GetBrush("Icons.Error");
+		break;
+	case EStackIssueSeverity::Warning:
+		IconBrush = FEditorStyle::GetBrush("Icons.Warning");
+		break;
+	case EStackIssueSeverity::Info:
+		IconBrush = FEditorStyle::GetBrush("Icons.Info");
+		break;
+	default:
+		IconBrush = FEditorStyle::GetBrush("NoBrush");
+		break;
+	}
 	ErrorInternalBox->AddSlot()
 		.HAlign(HAlign_Left)
 		.VAlign(VAlign_Center)
@@ -38,7 +54,7 @@ void SNiagaraStackErrorItem::Construct(const FArguments& InArgs, UNiagaraStackEr
 		.VAlign(VAlign_Center)
 		[
 			SNew(SImage)
-			.Image(FEditorStyle::GetBrush("Icons.Error"))
+			.Image(IconBrush)
 		]
 		+ SHorizontalBox::Slot()
 		.AutoWidth()
diff --git a/Engine/Plugins/FX/Niagara/Source/NiagaraEditorWidgets/Private/Stack/SNiagaraStackIssueIcon.cpp b/Engine/Plugins/FX/Niagara/Source/NiagaraEditorWidgets/Private/Stack/SNiagaraStackIssueIcon.cpp
new file mode 100644
--- /dev/null
+++ b/Engine/Plugins/FX/Niagara/Source/NiagaraEditorWidgets/Private/Stack/SNiagaraStackIssueIcon.cpp
@@ -0,0 +1,195 @@
+﻿// Copyright 1998-2019 Epic Games, Inc. All Rights Reserved.
+
+#include "SNiagaraStackIssueIcon.h"
+#include "ViewModels/Stack/NiagaraStackViewModel.h"
+#include "ViewModels/Stack/NiagaraStackEntry.h"
+
+#include "Widgets/Images/SImage.h"
+#include "Widgets/Layout/SBox.h"
+#include "EditorStyleSet.h"
+#include "Internationalization/Text.h"
+
+#define LOCTEXT_NAMESPACE "SNiagaraStackIssueIcon"
+
+void SNiagaraStackIssueIcon::Construct(const FArguments& InArgs, UNiagaraStackViewModel* InStackViewModel, UNiagaraStackEntry* InStackEntry)
+{
+	StackViewModel = InStackViewModel;
+	StackEntry = InStackEntry;
+	if (StackEntry != nullptr)
+	{
+		StackEntry->OnStructureChanged().AddSP(this, &SNiagaraStackIssueIcon::UpdateFromEntry);
+	}
+
+	ChildSlot
+	[
+		SNew(SBox)
+		.IsEnabled(this, &SNiagaraStackIssueIcon::GetIconIsEnabled)
+		.ToolTipText(this, &SNiagaraStackIssueIcon::GetIconToolTip)
+		.HeightOverride(16)
+		.WidthOverride(16)
+		.HAlign(HAlign_Center)
+		.VAlign(VAlign_Center)
+		[
+			SNew(SImage)
+			.Image(this, &SNiagaraStackIssueIcon::GetIconBrush)
+		]
+	];
+	UpdateFromEntry();
+}
+
+SNiagaraStackIssueIcon::~SNiagaraStackIssueIcon()
+{
+	if (StackEntry != nullptr)
+	{
+		StackEntry->OnStructureChanged().RemoveAll(this);
+	}
+}
+
+bool SNiagaraStackIssueIcon::GetIconIsEnabled() const
+{
+	return StackEntry != nullptr && StackEntry->GetOwnerIsEnabled() && StackEntry->GetIsEnabled();
+}
+
+const FSlateBrush* SNiagaraStackIssueIcon::GetIconBrush() const
+{
+	return IconBrush;
+}
+
+FText SNiagaraStackIssueIcon::GetIconToolTip() const
+{
+	if (IconToolTipCache.IsSet() == false)
+	{
+		FTextBuilder ToolTipBuilder;
+		TArray<FText> ToolTipParts;
+		if (StackEntry->GetTotalNumberOfErrorIssues() > 0)
+		{
+			if (StackEntry->GetTotalNumberOfErrorIssues() == 1)
+			{
+				ToolTipParts.Add(LOCTEXT("ErrorFormatSingle", "1 error"));
+			}
+			else
+			{
+				ToolTipParts.Add(FText::Format(LOCTEXT("ErrorFormatMultiple", "{0} errors"), FText::AsNumber(StackEntry->GetTotalNumberOfErrorIssues())));
+			}
+		}
+		if (StackEntry->GetTotalNumberOfWarningIssues() > 0)
+		{
+			if (StackEntry->GetTotalNumberOfWarningIssues() == 1)
+			{
+				ToolTipParts.Add(LOCTEXT("WarningFormatSingle", "1 warning"));
+			}
+			else
+			{
+				ToolTipParts.Add(FText::Format(LOCTEXT("WarningFormatMultiple", "{0} warnings"), FText::AsNumber(StackEntry->GetTotalNumberOfWarningIssues())));
+			}
+		}
+		if (StackEntry->GetTotalNumberOfInfoIssues() > 0)
+		{
+			if (StackEntry->GetTotalNumberOfInfoIssues() == 1)
+			{
+				ToolTipParts.Add(LOCTEXT("InfoFormatSingle", "1 info"));
+			}
+			else
+			{
+				ToolTipParts.Add(FText::Format(LOCTEXT("InfoFormatMultiple", "{0} infos"), FText::AsNumber(StackEntry->GetTotalNumberOfInfoIssues())));
+			}
+		}
+
+		if (ToolTipParts.Num() == 3)
+		{
+			ToolTipBuilder.AppendLineFormat(LOCTEXT("ThreePartFormat", "{0}, {1}, and {2}"), ToolTipParts[0], ToolTipParts[1], ToolTipParts[2]);
+		}
+		else if (ToolTipParts.Num() == 2)
+		{
+			ToolTipBuilder.AppendLineFormat(LOCTEXT("ThreePartFormat", "{0} and {1}"), ToolTipParts[0], ToolTipParts[1]);
+		}
+		else if (ToolTipParts.Num() == 1)
+		{
+			ToolTipBuilder.AppendLine(ToolTipParts[0]);
+		}
+
+		FText IssueLineFormat = LOCTEXT("IssueLineFormat", "{0} - {1} - {2}");
+		auto SeverityToText = [](EStackIssueSeverity Severity)
+		{
+			switch (Severity)
+			{
+			case EStackIssueSeverity::Error:
+				return LOCTEXT("Error", "Error");
+			case EStackIssueSeverity::Warning:
+				return LOCTEXT("Warning", "Warning");
+			case EStackIssueSeverity::Info:
+				return LOCTEXT("Issue", "Issue");
+			default:
+				return FText();
+			}
+		};
+
+		auto GetFullDisplayName = [](const UNiagaraStackViewModel* StackViewModel, UNiagaraStackEntry* EntryWithIssue)
+		{
+			TArray<UNiagaraStackEntry*> EntryPath;
+			StackViewModel->GetPathForEntry(EntryWithIssue, EntryPath);
+			EntryPath.Add(EntryWithIssue);
+			TArray<FText> DisplayNameParts;
+			for (UNiagaraStackEntry* Entry : EntryPath)
+			{
+				if (Entry->GetDisplayName().IsEmpty() == false)
+				{
+					DisplayNameParts.Add(Entry->GetDisplayName());
+				}
+			}
+			return FText::Join(LOCTEXT("DisplayNameJoinDelimiter", " - "), DisplayNameParts);
+		};
+
+		int32 IssueLinesAppended = 0;
+		int32 MaxIssueLines = 5;
+		int32 TotalIssues = 0;
+		TArray<UNiagaraStackEntry*> EntriesToCheck;
+		EntriesToCheck.Add(StackEntry);
+		EntriesToCheck.Append(StackEntry->GetAllChildrenWithIssues());
+		for (UNiagaraStackEntry* EntryToCheck : EntriesToCheck)
+		{
+			for (int32 IssueIndex = 0; IssueIndex < EntryToCheck->GetIssues().Num() && IssueLinesAppended < MaxIssueLines; IssueIndex++, IssueLinesAppended++)
+			{
+				const UNiagaraStackEntry::FStackIssue& Issue = EntryToCheck->GetIssues()[IssueIndex];
+				ToolTipBuilder.AppendLineFormat(IssueLineFormat, GetFullDisplayName(StackViewModel, EntryToCheck), SeverityToText(Issue.GetSeverity()), Issue.GetShortDescription());
+			}
+			TotalIssues += EntryToCheck->GetIssues().Num();
+		}
+
+		if (TotalIssues > MaxIssueLines)
+		{
+			ToolTipBuilder.AppendLineFormat(LOCTEXT("MoreIssuesFormat", "(And {0} more {0}|plural(one=issue,other=issues)...)"), FText::AsNumber(TotalIssues - MaxIssueLines));
+			//"There {NumCats}|plural(one=is,other=are) {NumCats} {NumCats}|plural(one=cat,other=cats)"
+		}
+
+		IconToolTipCache = ToolTipBuilder.ToText();
+	}
+	return IconToolTipCache.GetValue();
+}
+
+void SNiagaraStackIssueIcon::UpdateFromEntry()
+{
+	if (StackEntry == nullptr || StackEntry->HasIssuesOrAnyChildHasIssues() == false)
+	{
+		IconBrush = FEditorStyle::GetBrush("NoBrush");
+		IconToolTipCache = FText();
+		return;
+	}
+
+	if (StackEntry->GetTotalNumberOfErrorIssues() > 0)
+	{
+		IconBrush = FEditorStyle::GetBrush("Icons.Error");
+	}
+	else if (StackEntry->GetTotalNumberOfWarningIssues() > 0)
+	{
+		IconBrush = FEditorStyle::GetBrush("Icons.Warning");
+	}
+	else if (StackEntry->GetTotalNumberOfInfoIssues() > 0)
+	{
+		IconBrush = FEditorStyle::GetBrush("Icons.Info");
+	}
+
+	IconToolTipCache.Reset();
+}
+
+#undef LOCTEXT_NAMESPACE
diff --git a/Engine/Plugins/FX/Niagara/Source/NiagaraEditorWidgets/Private/Stack/SNiagaraStackIssueIcon.h b/Engine/Plugins/FX/Niagara/Source/NiagaraEditorWidgets/Private/Stack/SNiagaraStackIssueIcon.h
new file mode 100644
--- /dev/null
+++ b/Engine/Plugins/FX/Niagara/Source/NiagaraEditorWidgets/Private/Stack/SNiagaraStackIssueIcon.h
@@ -0,0 +1,37 @@
+﻿// Copyright 1998-2019 Epic Games, Inc. All Rights Reserved.
+
+#pragma once
+
+#include "Widgets/SCompoundWidget.h"
+#include "Widgets/DeclarativeSyntaxSupport.h"
+
+class UNiagaraStackViewModel;
+class UNiagaraStackEntry;
+
+class SNiagaraStackIssueIcon : public SCompoundWidget
+{
+public:
+	SLATE_BEGIN_ARGS(SNiagaraStackIssueIcon) {}
+	SLATE_END_ARGS()
+
+	void Construct(const FArguments& InArgs, UNiagaraStackViewModel* InStackViewModel, UNiagaraStackEntry* InStackEntry);
+
+	~SNiagaraStackIssueIcon();
+	
+private:
+	bool GetIconIsEnabled() const;
+
+	const FSlateBrush* GetIconBrush() const;
+
+	FText GetIconToolTip() const;
+
+	void UpdateFromEntry();
+
+private:
+	const FSlateBrush* IconBrush;
+
+	mutable TOptional<FText> IconToolTipCache;
+
+	UNiagaraStackViewModel* StackViewModel;
+	UNiagaraStackEntry* StackEntry;
+};
\ No newline at end of file
diff --git a/Engine/Plugins/FX/Niagara/Source/NiagaraEditorWidgets/Private/Stack/SNiagaraStackItemGroup.cpp b/Engine/Plugins/FX/Niagara/Source/NiagaraEditorWidgets/Private/Stack/SNiagaraStackItemGroup.cpp
--- a/Engine/Plugins/FX/Niagara/Source/NiagaraEditorWidgets/Private/Stack/SNiagaraStackItemGroup.cpp
+++ b/Engine/Plugins/FX/Niagara/Source/NiagaraEditorWidgets/Private/Stack/SNiagaraStackItemGroup.cpp
@@ -10,7 +10,6 @@
 #include "Widgets/SBoxPanel.h"
 	#include "Widgets/Input/SButton.h"
 #include "Widgets/Text/STextBlock.h"
-#include "SNiagaraStackErrorButton.h"
 
 #define LOCTEXT_NAMESPACE "NiagaraStackItemGroup"
 
@@ -23,27 +22,15 @@ void SNiagaraStackItemGroup::Construct(const FArguments& InArgs, UNiagaraStackIt
 
 	ChildSlot
 	[
-		// Name
 		SNew(SHorizontalBox)
+		// Name
 		+ SHorizontalBox::Slot()
 		.VAlign(VAlign_Center)
+		.Padding(2, 0, 0, 0)
 		[
 			SNew(SNiagaraStackDisplayName, InGroup, *InStackViewModel, "NiagaraEditor.Stack.GroupText")
 			.ColorAndOpacity(this, &SNiagaraStackEntryWidget::GetTextColorForSearch)
 		]
-		// Stack issues icon
-		+ SHorizontalBox::Slot()
-		.Padding(2, 0, 0, 0)
-		.AutoWidth()
-		.VAlign(VAlign_Center)
-		.HAlign(HAlign_Right)
-		[
-			SNew(SNiagaraStackErrorButton)
-			.IssueSeverity_UObject(Group, &UNiagaraStackItemGroup::GetHighestStackIssueSeverity)
-			.ErrorTooltip(this, &SNiagaraStackItemGroup::GetErrorButtonTooltipText)
-			.Visibility(this, &SNiagaraStackItemGroup::GetStackIssuesWarningVisibility)
-			.OnButtonClicked(this, &SNiagaraStackItemGroup::ExpandEntry)
-		]
 		// Delete group button
 		+ SHorizontalBox::Slot()
 		.AutoWidth()
@@ -100,14 +87,4 @@ FReply SNiagaraStackItemGroup::DeleteClicked()
 	return FReply::Handled();
 }
 
-EVisibility SNiagaraStackItemGroup::GetStackIssuesWarningVisibility() const
-{
-	return  Group->GetRecursiveStackIssuesCount() > 0 ? EVisibility::Visible : EVisibility::Collapsed;
-}
-
-FText SNiagaraStackItemGroup::GetErrorButtonTooltipText() const
-{
-	return FText::Format(LOCTEXT("GroupIssuesTooltip", "This group contains items that have a total of {0} issues, click to expand."), Group->GetRecursiveStackIssuesCount());
-}
-
 #undef LOCTEXT_NAMESPACE
\ No newline at end of file
diff --git a/Engine/Plugins/FX/Niagara/Source/NiagaraEditorWidgets/Private/Stack/SNiagaraStackItemGroup.h b/Engine/Plugins/FX/Niagara/Source/NiagaraEditorWidgets/Private/Stack/SNiagaraStackItemGroup.h
--- a/Engine/Plugins/FX/Niagara/Source/NiagaraEditorWidgets/Private/Stack/SNiagaraStackItemGroup.h
+++ b/Engine/Plugins/FX/Niagara/Source/NiagaraEditorWidgets/Private/Stack/SNiagaraStackItemGroup.h
@@ -23,10 +23,6 @@ private:
 
 	FReply DeleteClicked();
 
-	EVisibility GetStackIssuesWarningVisibility() const;
-
-	FText GetErrorButtonTooltipText() const;
-
 private:
 	UNiagaraStackItemGroup* Group;
 };
diff --git a/Engine/Plugins/FX/Niagara/Source/NiagaraEditorWidgets/Private/Stack/SNiagaraStackModuleItem.cpp b/Engine/Plugins/FX/Niagara/Source/NiagaraEditorWidgets/Private/Stack/SNiagaraStackModuleItem.cpp
--- a/Engine/Plugins/FX/Niagara/Source/NiagaraEditorWidgets/Private/Stack/SNiagaraStackModuleItem.cpp
+++ b/Engine/Plugins/FX/Niagara/Source/NiagaraEditorWidgets/Private/Stack/SNiagaraStackModuleItem.cpp
@@ -21,7 +21,6 @@
 #include "Framework/MultiBox/MultiBoxBuilder.h"
 #include "Framework/Application/SlateApplication.h"
 #include "SDropTarget.h"
-#include "SNiagaraStackErrorButton.h"
 #include "NiagaraEditorUtilities.h"
 #include "SGraphActionMenu.h"
 #include "NiagaraEditorWidgetsUtilities.h"
@@ -51,24 +50,11 @@ void SNiagaraStackModuleItem::Construct(const FArguments& InArgs, UNiagaraStackM
 			// Name
 			+ SHorizontalBox::Slot()
 			.VAlign(VAlign_Center)
-			.Padding(0)
+			.Padding(2, 0, 0, 0)
 			[
 				SNew(SNiagaraStackDisplayName, InModuleItem, *InStackViewModel, "NiagaraEditor.Stack.ItemText")
 				.ColorAndOpacity(this, &SNiagaraStackModuleItem::GetTextColorForSearch)
 			]
-			// Stack issues icon
-			+ SHorizontalBox::Slot()
-			.Padding(2, 0, 0, 0)
-			.AutoWidth()
-			.VAlign(VAlign_Center)
-			[
-				SNew(SNiagaraStackErrorButton)
-				.IssueSeverity_UObject(ModuleItem, &UNiagaraStackModuleItem::GetHighestStackIssueSeverity)
-				.ErrorTooltip(this, &SNiagaraStackModuleItem::GetErrorButtonTooltipText)
-				.Visibility(this, &SNiagaraStackModuleItem::GetStackIssuesWarningVisibility)
-				.IsEnabled(this, &SNiagaraStackModuleItem::GetButtonsEnabled)
-				.OnButtonClicked(this, &SNiagaraStackModuleItem::ExpandEntry)
-			]
 			// Raise Action Menu button
 			+ SHorizontalBox::Slot()
 			.AutoWidth()
@@ -298,16 +284,6 @@ bool SNiagaraStackModuleItem::OnModuleItemAllowDrop(TSharedPtr<FDragDropOperatio
 	return false;
 }
 
-EVisibility SNiagaraStackModuleItem::GetStackIssuesWarningVisibility() const
-{
-	return ModuleItem->GetRecursiveStackIssuesCount() > 0 ? EVisibility::Visible : EVisibility::Collapsed;
-}
-
-FText SNiagaraStackModuleItem::GetErrorButtonTooltipText() const
-{
-	return FText::Format(LOCTEXT("ModuleIssuesTooltip", "This module has {0} issues, click to expand."), ModuleItem->GetRecursiveStackIssuesCount());
-}
-
 void ReassignModuleScript(UNiagaraStackModuleItem* ModuleItem, FAssetData NewModuleScriptAsset)
 {
 	UNiagaraScript* NewModuleScript = Cast<UNiagaraScript>(NewModuleScriptAsset.GetAsset());
diff --git a/Engine/Plugins/FX/Niagara/Source/NiagaraEditorWidgets/Private/Stack/SNiagaraStackModuleItem.h b/Engine/Plugins/FX/Niagara/Source/NiagaraEditorWidgets/Private/Stack/SNiagaraStackModuleItem.h
--- a/Engine/Plugins/FX/Niagara/Source/NiagaraEditorWidgets/Private/Stack/SNiagaraStackModuleItem.h
+++ b/Engine/Plugins/FX/Niagara/Source/NiagaraEditorWidgets/Private/Stack/SNiagaraStackModuleItem.h
@@ -56,10 +56,6 @@ private:
 
 	bool OnModuleItemAllowDrop(TSharedPtr<class FDragDropOperation> DragDropOperation);
 
-	EVisibility GetStackIssuesWarningVisibility() const;
-
-	FText GetErrorButtonTooltipText() const;
-
 	void ShowReassignModuleScriptMenu();
 
 private:
diff --git a/Engine/Plugins/FX/Niagara/Source/NiagaraEditorWidgets/Private/Stack/SNiagaraStackRendererItem.cpp b/Engine/Plugins/FX/Niagara/Source/NiagaraEditorWidgets/Private/Stack/SNiagaraStackRendererItem.cpp
--- a/Engine/Plugins/FX/Niagara/Source/NiagaraEditorWidgets/Private/Stack/SNiagaraStackRendererItem.cpp
+++ b/Engine/Plugins/FX/Niagara/Source/NiagaraEditorWidgets/Private/Stack/SNiagaraStackRendererItem.cpp
@@ -15,7 +15,7 @@
 #include "Styling/SlateIconFinder.h"
 #include "Widgets/Images/SImage.h"
 #include "Widgets/Layout/SBox.h"
-#include "SNiagaraStackErrorButton.h"
+#include "SNiagaraStackIssueIcon.h"
 
 #define LOCTEXT_NAMESPACE "NiagaraStackRendererItem"
 
@@ -27,8 +27,8 @@ void SNiagaraStackRendererItem::Construct(const FArguments& InArgs, UNiagaraStac
 
 	ChildSlot
 	[
-		// Renderer icon
 		SNew(SHorizontalBox)
+		// Renderer icon
 		+ SHorizontalBox::Slot()
 		.Padding(2, 0, 0, 0)
 		.AutoWidth()
@@ -39,23 +39,12 @@ void SNiagaraStackRendererItem::Construct(const FArguments& InArgs, UNiagaraStac
 		]
 		// Display name
 		+ SHorizontalBox::Slot()
-		.Padding(5, 0, 0, 0)
+		.Padding(2, 0, 0, 0)
 		.VAlign(VAlign_Center)
 		[
 			SNew(SNiagaraStackDisplayName, InRendererItem, *InStackViewModel, "NiagaraEditor.Stack.ItemText")
 			.ColorAndOpacity(this, &SNiagaraStackEntryWidget::GetTextColorForSearch)
 		]
-		// Stack issues icon
-		+ SHorizontalBox::Slot()
-		.Padding(2, 0, 0, 0)
-		.AutoWidth()
-		[
-			SNew(SNiagaraStackErrorButton)
-			.IssueSeverity_UObject(RendererItem, &UNiagaraStackRendererItem::GetHighestStackIssueSeverity)
-			.ErrorTooltip(this, &SNiagaraStackRendererItem::GetErrorButtonTooltipText)
-			.Visibility(this, &SNiagaraStackRendererItem::GetStackIssuesWarningVisibility)
-			.OnButtonClicked(this, &SNiagaraStackRendererItem::ExpandEntry)
-		]
 		// Reset to base Button
 		+ SHorizontalBox::Slot()
 		.AutoWidth()
@@ -153,14 +142,4 @@ ECheckBoxState SNiagaraStackRendererItem::CheckEnabledStatus() const
 	return RendererItem->GetIsEnabled() ? ECheckBoxState::Checked : ECheckBoxState::Unchecked;
 }
 
-EVisibility SNiagaraStackRendererItem::GetStackIssuesWarningVisibility() const
-{
-	return  RendererItem->GetRecursiveStackIssuesCount() > 0? EVisibility::Visible : EVisibility::Collapsed;
-}
-
-FText SNiagaraStackRendererItem::GetErrorButtonTooltipText() const
-{
-	return FText::Format(LOCTEXT("ModuleIssuesTooltip", "This renderer has {0} issues, click to expand."), RendererItem->GetRecursiveStackIssuesCount());
-}
-
 #undef LOCTEXT_NAMESPACE
\ No newline at end of file
diff --git a/Engine/Plugins/FX/Niagara/Source/NiagaraEditorWidgets/Private/Stack/SNiagaraStackRendererItem.h b/Engine/Plugins/FX/Niagara/Source/NiagaraEditorWidgets/Private/Stack/SNiagaraStackRendererItem.h
--- a/Engine/Plugins/FX/Niagara/Source/NiagaraEditorWidgets/Private/Stack/SNiagaraStackRendererItem.h
+++ b/Engine/Plugins/FX/Niagara/Source/NiagaraEditorWidgets/Private/Stack/SNiagaraStackRendererItem.h
@@ -32,10 +32,6 @@ private:
 
 	ECheckBoxState CheckEnabledStatus() const;
 
-	EVisibility GetStackIssuesWarningVisibility() const;
-
-	FText GetErrorButtonTooltipText() const;
-
 private:
 	UNiagaraStackRendererItem* RendererItem;
 };
\ No newline at end of file
diff --git a/Engine/Plugins/FX/Niagara/Source/NiagaraEditorWidgets/Private/Stack/SNiagaraStackTableRow.cpp b/Engine/Plugins/FX/Niagara/Source/NiagaraEditorWidgets/Private/Stack/SNiagaraStackTableRow.cpp
--- a/Engine/Plugins/FX/Niagara/Source/NiagaraEditorWidgets/Private/Stack/SNiagaraStackTableRow.cpp
+++ b/Engine/Plugins/FX/Niagara/Source/NiagaraEditorWidgets/Private/Stack/SNiagaraStackTableRow.cpp
@@ -8,6 +8,7 @@
 #include "ViewModels/Stack/NiagaraStackItemGroup.h"
 #include "ViewModels/Stack/NiagaraStackEntry.h"
 #include "NiagaraEditorWidgetsUtilities.h"
+#include "Stack/SNiagaraStackIssueIcon.h"
 #include "Widgets/Input/SButton.h"
 #include "Widgets/Images/SImage.h"
 #include "Widgets/Layout/SBox.h"
@@ -26,6 +27,7 @@ void SNiagaraStackTableRow::Construct(const FArguments& InArgs, UNiagaraStackVie
 	ValueColumnWidth = InArgs._ValueColumnWidth;
 	NameColumnWidthChanged = InArgs._OnNameColumnWidthChanged;
 	ValueColumnWidthChanged = InArgs._OnValueColumnWidthChanged;
+	IssueIconVisibility = InArgs._IssueIconVisibility;
 	RowPadding = InArgs._RowPadding;
 	StackViewModel = InStackViewModel;
 	StackEntry = InStackEntry;
@@ -241,6 +243,13 @@ void SNiagaraStackTableRow::SetNameAndValueContent(TSharedRef<SWidget> InNameWid
 					]
 				]
 			]
+			+ SHorizontalBox::Slot()
+			.AutoWidth()
+			.Padding(3, 0, 0, 0)
+			[
+				SNew(SNiagaraStackIssueIcon, StackViewModel, StackEntry)
+				.Visibility(IssueIconVisibility)
+			]
 		]
 		+ SOverlay::Slot()
 		[
diff --git a/Engine/Plugins/FX/Niagara/Source/NiagaraEditorWidgets/Private/Stack/SNiagaraStackTableRow.h b/Engine/Plugins/FX/Niagara/Source/NiagaraEditorWidgets/Private/Stack/SNiagaraStackTableRow.h
--- a/Engine/Plugins/FX/Niagara/Source/NiagaraEditorWidgets/Private/Stack/SNiagaraStackTableRow.h
+++ b/Engine/Plugins/FX/Niagara/Source/NiagaraEditorWidgets/Private/Stack/SNiagaraStackTableRow.h
@@ -34,6 +34,7 @@ public:
 		SLATE_ARGUMENT(bool, ShowExecutionCategoryIcon)
 		SLATE_ATTRIBUTE(float, NameColumnWidth)
 		SLATE_ATTRIBUTE(float, ValueColumnWidth)
+		SLATE_ATTRIBUTE(EVisibility, IssueIconVisibility)
 		SLATE_ATTRIBUTE(FMargin, RowPadding)
 		SLATE_EVENT(FOnColumnWidthChanged, OnNameColumnWidthChanged)
 		SLATE_EVENT(FOnColumnWidthChanged, OnValueColumnWidthChanged)
@@ -103,6 +104,7 @@ private:
 	FOnColumnWidthChanged NameColumnWidthChanged;
 	FOnColumnWidthChanged ValueColumnWidthChanged;
 
+	TAttribute<EVisibility> IssueIconVisibility;
 	TAttribute<FMargin> RowPadding;
 
 	const FSlateBrush* ExpandedImage;
diff --git a/Engine/Plugins/FX/Niagara/Source/NiagaraEditorWidgets/Public/NiagaraEditorWidgetsModule.h b/Engine/Plugins/FX/Niagara/Source/NiagaraEditorWidgets/Public/NiagaraEditorWidgetsModule.h
--- a/Engine/Plugins/FX/Niagara/Source/NiagaraEditorWidgets/Public/NiagaraEditorWidgetsModule.h
+++ b/Engine/Plugins/FX/Niagara/Source/NiagaraEditorWidgets/Public/NiagaraEditorWidgetsModule.h
@@ -49,6 +49,7 @@ private:
 	public:
 		virtual TSharedRef<SWidget> CreateStackView(UNiagaraStackViewModel& StackViewModel) const override;
 		virtual TSharedRef<SWidget> CreateSystemOverview(TSharedRef<FNiagaraSystemViewModel> SystemViewModel) const override;
+		virtual TSharedRef<SWidget> CreateStackIssueIcon(UNiagaraStackViewModel& StackViewModel, UNiagaraStackEntry& StackEntry) const override;
 		virtual FLinearColor GetColorForExecutionCategory(FName ExecutionCategory) const override;
 	};
 
