Niagara - Shadowed variable warning to fix Mac and Linux builds.

diff --git a/Engine/Plugins/FX/Niagara/Source/NiagaraEditorWidgets/Private/Stack/SNiagaraStackIssueIcon.cpp b/Engine/Plugins/FX/Niagara/Source/NiagaraEditorWidgets/Private/Stack/SNiagaraStackIssueIcon.cpp
--- a/Engine/Plugins/FX/Niagara/Source/NiagaraEditorWidgets/Private/Stack/SNiagaraStackIssueIcon.cpp
+++ b/Engine/Plugins/FX/Niagara/Source/NiagaraEditorWidgets/Private/Stack/SNiagaraStackIssueIcon.cpp
@@ -124,10 +124,10 @@ FText SNiagaraStackIssueIcon::GetIconToolTip() const
 			}
 		};
 
-		auto GetFullDisplayName = [](const UNiagaraStackViewModel* StackViewModel, UNiagaraStackEntry* EntryWithIssue)
+		auto GetFullDisplayName = [](const UNiagaraStackViewModel* InStackViewModel, UNiagaraStackEntry* EntryWithIssue)
 		{
 			TArray<UNiagaraStackEntry*> EntryPath;
-			StackViewModel->GetPathForEntry(EntryWithIssue, EntryPath);
+			InStackViewModel->GetPathForEntry(EntryWithIssue, EntryPath);
 			EntryPath.Add(EntryWithIssue);
 			TArray<FText> DisplayNameParts;
 			for (UNiagaraStackEntry* Entry : EntryPath)
