Niagara - Show proper tooltips for rows in the system overview.

diff --git a/Engine/Plugins/FX/Niagara/Source/NiagaraEditorWidgets/Private/SNiagaraOverviewStack.cpp b/Engine/Plugins/FX/Niagara/Source/NiagaraEditorWidgets/Private/SNiagaraOverviewStack.cpp
--- a/Engine/Plugins/FX/Niagara/Source/NiagaraEditorWidgets/Private/SNiagaraOverviewStack.cpp
+++ b/Engine/Plugins/FX/Niagara/Source/NiagaraEditorWidgets/Private/SNiagaraOverviewStack.cpp
@@ -91,6 +91,7 @@ class SNiagaraSystemOverviewEntryListRow : public STableRow<UNiagaraStackEntry*>
 		[
 			SNew(SBorder)
 			.BorderImage(this, &SNiagaraSystemOverviewEntryListRow::GetBorder)
+			.ToolTipText_UObject(StackEntry, &UNiagaraStackEntry::GetTooltipText)
 			.Padding(FMargin(5, 2, 3, 2))
 			[
 				SNew(SHorizontalBox)
