Niagara - Prevent a crash when undoing copy/pasting emitters in the timeline.

diff --git a/Engine/Plugins/FX/Niagara/Source/NiagaraEditor/Private/ViewModels/Stack/NiagaraStackRenderItemGroup.cpp b/Engine/Plugins/FX/Niagara/Source/NiagaraEditor/Private/ViewModels/Stack/NiagaraStackRenderItemGroup.cpp
--- a/Engine/Plugins/FX/Niagara/Source/NiagaraEditor/Private/ViewModels/Stack/NiagaraStackRenderItemGroup.cpp
+++ b/Engine/Plugins/FX/Niagara/Source/NiagaraEditor/Private/ViewModels/Stack/NiagaraStackRenderItemGroup.cpp
@@ -124,7 +124,8 @@ void UNiagaraStackRenderItemGroup::Initialize(FRequiredEntryData InRequiredEntry
 	FText ToolTip = LOCTEXT("RendererGroupTooltip", "Describes how we should display/present each particle. Note that this doesn't have to be visual. Multiple renderers are supported. Order in this stack is not necessarily relevant to draw order.");
 	AddUtilities = MakeShared<FRenderItemGroupAddUtilities>(InRequiredEntryData.EmitterViewModel.ToSharedRef());
 	Super::Initialize(InRequiredEntryData, DisplayName, ToolTip, AddUtilities.Get());
-	GetEmitterViewModel()->GetEmitter()->OnRenderersChanged().AddUObject(this, &UNiagaraStackRenderItemGroup::EmitterRenderersChanged);
+	EmitterWeak = GetEmitterViewModel()->GetEmitter();
+	EmitterWeak->OnRenderersChanged().AddUObject(this, &UNiagaraStackRenderItemGroup::EmitterRenderersChanged);
 }
 
 void UNiagaraStackRenderItemGroup::RefreshChildrenInternal(const TArray<UNiagaraStackEntry*>& CurrentChildren, TArray<UNiagaraStackEntry*>& NewChildren, TArray<FStackIssue>& NewIssues)
@@ -151,13 +152,20 @@ void UNiagaraStackRenderItemGroup::RefreshChildrenInternal(const TArray<UNiagara
 
 void UNiagaraStackRenderItemGroup::EmitterRenderersChanged()
 {
-	OnDataObjectModified().Broadcast(nullptr);
-	RefreshChildren();
+	if (IsFinalized() == false)
+	{
+		// With undo/redo sometimes it's not possible to unbind this delegate, so we have to check to insure safety in those cases.
+		OnDataObjectModified().Broadcast(nullptr);
+		RefreshChildren();
+	}
 }
 
 void UNiagaraStackRenderItemGroup::FinalizeInternal()
 {
-	GetEmitterViewModel()->GetEmitter()->OnRenderersChanged().RemoveAll(this);
+	if (EmitterWeak.IsValid())
+	{
+		EmitterWeak->OnRenderersChanged().RemoveAll(this);
+	}
 	Super::FinalizeInternal();
 }
 
diff --git a/Engine/Plugins/FX/Niagara/Source/NiagaraEditor/Public/ViewModels/Stack/NiagaraStackRenderItemGroup.h b/Engine/Plugins/FX/Niagara/Source/NiagaraEditor/Public/ViewModels/Stack/NiagaraStackRenderItemGroup.h
--- a/Engine/Plugins/FX/Niagara/Source/NiagaraEditor/Public/ViewModels/Stack/NiagaraStackRenderItemGroup.h
+++ b/Engine/Plugins/FX/Niagara/Source/NiagaraEditor/Public/ViewModels/Stack/NiagaraStackRenderItemGroup.h
@@ -6,6 +6,7 @@
 #include "NiagaraStackRenderItemGroup.generated.h"
 
 class UNiagaraRendererProperties;
+class UNiagaraEmitter;
 
 UCLASS()
 class NIAGARAEDITOR_API UNiagaraStackRenderItemGroup : public UNiagaraStackItemGroup
@@ -25,4 +26,6 @@ private:
 
 private:
 	TSharedPtr<INiagaraStackItemGroupAddUtilities> AddUtilities;
+
+	TWeakObjectPtr<UNiagaraEmitter> EmitterWeak;
 };
diff --git a/Engine/Plugins/FX/Niagara/Source/NiagaraEditorWidgets/Private/SNiagaraOverviewStackNode.cpp b/Engine/Plugins/FX/Niagara/Source/NiagaraEditorWidgets/Private/SNiagaraOverviewStackNode.cpp
--- a/Engine/Plugins/FX/Niagara/Source/NiagaraEditorWidgets/Private/SNiagaraOverviewStackNode.cpp
+++ b/Engine/Plugins/FX/Niagara/Source/NiagaraEditorWidgets/Private/SNiagaraOverviewStackNode.cpp
@@ -14,10 +14,16 @@
 #include "Widgets/Layout/SBox.h"
 #include "Widgets/Input/SCheckBox.h"
 
+#define LOCTEXT_NAMESPACE "NiagaraOverviewStackNode"
+
 void SNiagaraOverviewStackNode::Construct(const FArguments& InArgs, UNiagaraOverviewNode* InNode)
 {
 	GraphNode = InNode;
 	OverviewStackNode = InNode;
+	StackViewModel = nullptr;
+	OverviewSelectionViewModel = nullptr;
+	EmitterHandleViewModelWeak.Reset();
+
 	if (OverviewStackNode->GetOwningSystem() != nullptr)
 	{
 		FNiagaraEditorModule& NiagaraEditorModule = FModuleManager::Get().LoadModuleChecked<FNiagaraEditorModule>("NiagaraEditor");
@@ -45,7 +51,22 @@ void SNiagaraOverviewStackNode::Construct(const FArguments& InArgs, UNiagaraOver
 TSharedRef<SWidget> SNiagaraOverviewStackNode::CreateTitleWidget(TSharedPtr<SNodeTitle> NodeTitle)
 {
 	TSharedRef<SWidget> DefaultTitle = SGraphNode::CreateTitleWidget(NodeTitle);
-	
+
+	if (StackViewModel == nullptr)
+	{
+		return SNew(SHorizontalBox)
+			+ SHorizontalBox::Slot()
+			.Padding(0, 0, 5, 0)
+			[
+				SNew(STextBlock)
+				.Text(LOCTEXT("InvalidNode", "INVALID"))
+			]
+			+ SHorizontalBox::Slot()
+			[
+				DefaultTitle
+			];
+	}
+
 	return SNew(SHorizontalBox)
 		// Enabled checkbox
 		+ SHorizontalBox::Slot()
@@ -149,4 +170,6 @@ void SNiagaraOverviewStackNode::OnEnabledCheckStateChanged(ECheckBoxState InChec
 	{
 		EmitterHandleViewModel->SetIsEnabled(InCheckState == ECheckBoxState::Checked);
 	}
-}
\ No newline at end of file
+}
+
+#undef LOCTEXT_NAMESPACE
\ No newline at end of file
