Bump Niagara StandaloneDefaultLayout version to invalidate script tab layouts.

diff --git a/Engine/Plugins/FX/Niagara/Source/NiagaraEditor/Private/Toolkits/NiagaraScriptToolkit.cpp b/Engine/Plugins/FX/Niagara/Source/NiagaraEditor/Private/Toolkits/NiagaraScriptToolkit.cpp
--- a/Engine/Plugins/FX/Niagara/Source/NiagaraEditor/Private/Toolkits/NiagaraScriptToolkit.cpp
+++ b/Engine/Plugins/FX/Niagara/Source/NiagaraEditor/Private/Toolkits/NiagaraScriptToolkit.cpp
@@ -166,7 +166,7 @@ void FNiagaraScriptToolkit::Initialize( const EToolkitMode::Type Mode, const TSh
 	StatsListing = MessageLogModule.CreateLogListing("MaterialEditorStats", LogOptions);
 	Stats = MessageLogModule.CreateLogListingWidget(StatsListing.ToSharedRef());
 
-	TSharedRef<FTabManager::FLayout> StandaloneDefaultLayout = FTabManager::NewLayout("Standalone_Niagara_Layout_v8")
+	TSharedRef<FTabManager::FLayout> StandaloneDefaultLayout = FTabManager::NewLayout("Standalone_Niagara_Layout_v9")
 	->AddArea
 	(
 		FTabManager::NewPrimaryArea()->SetOrientation(Orient_Vertical)
