Niagara - Fix an assert when using the selection stack search due to the root selection entry having no editor data key.

diff --git a/Engine/Plugins/FX/Niagara/Source/NiagaraEditor/Private/ViewModels/Stack/NiagaraStackSelection.cpp b/Engine/Plugins/FX/Niagara/Source/NiagaraEditor/Private/ViewModels/Stack/NiagaraStackSelection.cpp
--- a/Engine/Plugins/FX/Niagara/Source/NiagaraEditor/Private/ViewModels/Stack/NiagaraStackSelection.cpp
+++ b/Engine/Plugins/FX/Niagara/Source/NiagaraEditor/Private/ViewModels/Stack/NiagaraStackSelection.cpp
@@ -7,7 +7,7 @@
 
 void UNiagaraStackSelection::Initialize(FRequiredEntryData InRequiredEntryData)
 {
-	Super::Initialize(InRequiredEntryData, FString());
+	Super::Initialize(InRequiredEntryData, TEXT("Selection"));
 }
 
 bool UNiagaraStackSelection::GetCanExpand() const
