Fixed localization conflict

#rnx

diff --git a/Engine/Plugins/FX/Niagara/Source/NiagaraEditorWidgets/Private/Stack/SNiagaraStackIssueIcon.cpp b/Engine/Plugins/FX/Niagara/Source/NiagaraEditorWidgets/Private/Stack/SNiagaraStackIssueIcon.cpp
--- a/Engine/Plugins/FX/Niagara/Source/NiagaraEditorWidgets/Private/Stack/SNiagaraStackIssueIcon.cpp
+++ b/Engine/Plugins/FX/Niagara/Source/NiagaraEditorWidgets/Private/Stack/SNiagaraStackIssueIcon.cpp
@@ -101,7 +101,7 @@ FText SNiagaraStackIssueIcon::GetIconToolTip() const
 		}
 		else if (ToolTipParts.Num() == 2)
 		{
-			ToolTipBuilder.AppendLineFormat(LOCTEXT("ThreePartFormat", "{0} and {1}"), ToolTipParts[0], ToolTipParts[1]);
+			ToolTipBuilder.AppendLineFormat(LOCTEXT("TwoPartFormat", "{0} and {1}"), ToolTipParts[0], ToolTipParts[1]);
 		}
 		else if (ToolTipParts.Num() == 1)
 		{
