Niagara - Fix a crash which would happen if a module was missing a dependency and there were multiple modules with the same name which could provide that dependency by using the full name in those cases.  Also add code to prevent this type of crash in the future where two fixes have the same unique id, and also add a debug name generator so that duplicate stack entries are easier to debug in the future.

diff --git a/Engine/Plugins/FX/Niagara/Source/NiagaraEditor/Private/ViewModels/Stack/NiagaraStackEntry.cpp b/Engine/Plugins/FX/Niagara/Source/NiagaraEditor/Private/ViewModels/Stack/NiagaraStackEntry.cpp
--- a/Engine/Plugins/FX/Niagara/Source/NiagaraEditor/Private/ViewModels/Stack/NiagaraStackEntry.cpp
+++ b/Engine/Plugins/FX/Niagara/Source/NiagaraEditor/Private/ViewModels/Stack/NiagaraStackEntry.cpp
@@ -608,7 +608,13 @@ void UNiagaraStackEntry::RefreshStackErrorChildren()
 			ErrorEntry = *Found;
 			ErrorEntry->SetStackIssue(Issue); // we found the entry by id but we want to properly refresh the subentries of the issue (specifically its fixes), too
 		}
-		NewErrorChildren.Add(ErrorEntry);
+		if (ensureMsgf(NewErrorChildren.Contains(ErrorEntry) == false,
+			TEXT("Duplicate stack issue rows detected, this is caused by two different issues generating the same unique id. Issue Short description: %s Issue Long description: %s.  This issue will not be shown in the UI."),
+			*Issue.GetShortDescription().ToString(),
+			*Issue.GetLongDescription().ToString()))
+		{
+			NewErrorChildren.Add(ErrorEntry);
+		}
 	}
 
 	// If any of the current error children were not moved to the new error children collection than finalize them since
diff --git a/Engine/Plugins/FX/Niagara/Source/NiagaraEditor/Private/ViewModels/Stack/NiagaraStackErrorItem.cpp b/Engine/Plugins/FX/Niagara/Source/NiagaraEditor/Private/ViewModels/Stack/NiagaraStackErrorItem.cpp
--- a/Engine/Plugins/FX/Niagara/Source/NiagaraEditor/Private/ViewModels/Stack/NiagaraStackErrorItem.cpp
+++ b/Engine/Plugins/FX/Niagara/Source/NiagaraEditor/Private/ViewModels/Stack/NiagaraStackErrorItem.cpp
@@ -76,8 +76,13 @@ void UNiagaraStackErrorItem::RefreshChildrenInternal(const TArray<UNiagaraStackE
 		{
 			ErrorEntryFix->SetFixDelegate(CurrentFix.GetFixDelegate());
 		}
-		ErrorEntryFix->OnIssueFixed().AddUObject(this, &UNiagaraStackErrorItem::IssueFixed);
-		NewChildren.Add(ErrorEntryFix);
+		if (ensureMsgf(NewChildren.Contains(ErrorEntryFix) == false,
+			TEXT("Duplicate stack issue fix rows detected, This is caused by two different issue fixes with the same description which is used to generate their unique id. Issue Fix description: %s.  This issue fix will not be shown in the UI."),
+			*CurrentFix.GetDescription().ToString()))
+		{
+			ErrorEntryFix->OnIssueFixed().AddUObject(this, &UNiagaraStackErrorItem::IssueFixed);
+			NewChildren.Add(ErrorEntryFix);
+		}
 	}
 	// dismiss button
 	if (StackIssue.GetCanBeDismissed())
@@ -127,11 +132,6 @@ void UNiagaraStackErrorItemFix::Initialize(FRequiredEntryData InRequiredEntryDat
 	IssueFix = InIssueFix;
 }
 
-FText UNiagaraStackErrorItemFix::FixDescription() const
-{
-	return IssueFix.GetDescription();
-}
-
 FReply UNiagaraStackErrorItemFix::OnTryFixError()
 {
 	IssueFix.GetFixDelegate().ExecuteIfBound();
@@ -139,6 +139,11 @@ FReply UNiagaraStackErrorItemFix::OnTryFixError()
 	return FReply::Handled();
 }
 
+FText UNiagaraStackErrorItemFix::GetDisplayName() const
+{
+	return IssueFix.GetDescription();
+}
+
 UNiagaraStackEntry::EStackRowStyle UNiagaraStackErrorItemFix::GetStackRowStyle() const
 {
 	return EStackRowStyle::StackIssue;
diff --git a/Engine/Plugins/FX/Niagara/Source/NiagaraEditor/Private/ViewModels/Stack/NiagaraStackModuleItem.cpp b/Engine/Plugins/FX/Niagara/Source/NiagaraEditor/Private/ViewModels/Stack/NiagaraStackModuleItem.cpp
--- a/Engine/Plugins/FX/Niagara/Source/NiagaraEditor/Private/ViewModels/Stack/NiagaraStackModuleItem.cpp
+++ b/Engine/Plugins/FX/Niagara/Source/NiagaraEditor/Private/ViewModels/Stack/NiagaraStackModuleItem.cpp
@@ -559,6 +559,21 @@ void UNiagaraStackModuleItem::RefreshIssues(TArray<FStackIssue>& NewIssues)
 				{
 					TArray<FAssetData> ModuleAssets;
 					FNiagaraStackGraphUtilities::GetScriptAssetsByDependencyProvided(ENiagaraScriptUsage::Module, Dependency.Id, ModuleAssets);
+
+					// Find duplicate module names in the fixes so that unique fix descriptions can be generated.
+					TArray<FName> ModuleNames;
+					TArray<FName> DuplicateModuleNames;
+					for (FAssetData ModuleAsset : ModuleAssets)
+					{
+						if (ModuleNames.Contains(ModuleAsset.AssetName))
+						{
+							DuplicateModuleNames.Add(ModuleAsset.AssetName);
+						}
+						else
+						{
+							ModuleNames.Add(ModuleAsset.AssetName);
+						}
+					}
 					for (FAssetData ModuleAsset : ModuleAssets)
 					{
 						UNiagaraScript* DependencyScript = Cast<UNiagaraScript>(ModuleAsset.GetAsset());
@@ -572,7 +587,8 @@ void UNiagaraStackModuleItem::RefreshIssues(TArray<FStackIssue>& NewIssues)
 							}
 						}
 						
-						FText FixDescription = FText::Format(LOCTEXT("AddDependency", "Add new dependency module {0}"), FText::FromName(ModuleAsset.AssetName));
+						FText AssetNameText = DuplicateModuleNames.Contains(ModuleAsset.AssetName) ? FText::FromName(ModuleAsset.PackageName) : FText::FromName(ModuleAsset.AssetName);
+						FText FixDescription = FText::Format(LOCTEXT("AddDependency", "Add new dependency module {0}"), AssetNameText);
 						UNiagaraStackEntry::FStackIssueFix Fix(
 							FixDescription,
 							FStackIssueFixDelegate::CreateLambda([=]()
diff --git a/Engine/Plugins/FX/Niagara/Source/NiagaraEditor/Public/ViewModels/Stack/NiagaraStackErrorItem.h b/Engine/Plugins/FX/Niagara/Source/NiagaraEditor/Public/ViewModels/Stack/NiagaraStackErrorItem.h
--- a/Engine/Plugins/FX/Niagara/Source/NiagaraEditor/Public/ViewModels/Stack/NiagaraStackErrorItem.h
+++ b/Engine/Plugins/FX/Niagara/Source/NiagaraEditor/Public/ViewModels/Stack/NiagaraStackErrorItem.h
@@ -57,8 +57,8 @@ public:
 	void Initialize(FRequiredEntryData InRequiredEntryData, FStackIssue InStackIssue, FStackIssueFix InIssueFix, FString InStackEditorDataKey);
 
 	FStackIssueFix GetStackIssueFix() const { return IssueFix; };
-	FText FixDescription() const;
 	FReply OnTryFixError();
+	virtual FText GetDisplayName() const override;
 	virtual EStackRowStyle GetStackRowStyle() const override;
 	virtual FText GetFixButtonText() const;
 	UNiagaraStackErrorItem::FOnIssueNotify& OnIssueFixed();
diff --git a/Engine/Plugins/FX/Niagara/Source/NiagaraEditorWidgets/Private/NiagaraEditorWidgetsUtilities.cpp b/Engine/Plugins/FX/Niagara/Source/NiagaraEditorWidgets/Private/NiagaraEditorWidgetsUtilities.cpp
--- a/Engine/Plugins/FX/Niagara/Source/NiagaraEditorWidgets/Private/NiagaraEditorWidgetsUtilities.cpp
+++ b/Engine/Plugins/FX/Niagara/Source/NiagaraEditorWidgets/Private/NiagaraEditorWidgetsUtilities.cpp
@@ -335,4 +335,9 @@ bool FNiagaraStackEditorWidgetsUtilities::HandleDropForStackEntry(const FDragDro
 	return bHandled;
 }
 
+FString FNiagaraStackEditorWidgetsUtilities::StackEntryToStringForListDebug(UNiagaraStackEntry* StackEntry)
+{
+	return FString::Printf(TEXT("0x%08x - %s - %s"), StackEntry, *StackEntry->GetClass()->GetName(), *StackEntry->GetDisplayName().ToString());
+}
+
 #undef LOCTEXT_NAMESPACE
\ No newline at end of file
diff --git a/Engine/Plugins/FX/Niagara/Source/NiagaraEditorWidgets/Private/NiagaraEditorWidgetsUtilities.h b/Engine/Plugins/FX/Niagara/Source/NiagaraEditorWidgets/Private/NiagaraEditorWidgetsUtilities.h
--- a/Engine/Plugins/FX/Niagara/Source/NiagaraEditorWidgets/Private/NiagaraEditorWidgetsUtilities.h
+++ b/Engine/Plugins/FX/Niagara/Source/NiagaraEditorWidgets/Private/NiagaraEditorWidgetsUtilities.h
@@ -31,4 +31,6 @@ namespace FNiagaraStackEditorWidgetsUtilities
 	TOptional<EItemDropZone> RequestDropForStackEntry(const FDragDropEvent& InDragDropEvent, EItemDropZone InDropZone, UNiagaraStackEntry* InTargetEntry, UNiagaraStackEntry::EDropOptions DropOptions);
 
 	bool HandleDropForStackEntry(const FDragDropEvent& InDragDropEvent, EItemDropZone InDropZone, UNiagaraStackEntry* InTargetEntry, UNiagaraStackEntry::EDropOptions DropOptions);
+
+	FString StackEntryToStringForListDebug(UNiagaraStackEntry* StackEntry);
 }
\ No newline at end of file
diff --git a/Engine/Plugins/FX/Niagara/Source/NiagaraEditorWidgets/Private/SNiagaraOverviewStack.cpp b/Engine/Plugins/FX/Niagara/Source/NiagaraEditorWidgets/Private/SNiagaraOverviewStack.cpp
--- a/Engine/Plugins/FX/Niagara/Source/NiagaraEditorWidgets/Private/SNiagaraOverviewStack.cpp
+++ b/Engine/Plugins/FX/Niagara/Source/NiagaraEditorWidgets/Private/SNiagaraOverviewStack.cpp
@@ -232,6 +232,7 @@ void SNiagaraOverviewStack::Construct(const FArguments& InArgs, UNiagaraStackVie
 		.OnGenerateRow(this, &SNiagaraOverviewStack::OnGenerateRowForEntry)
 		.OnSelectionChanged(this, &SNiagaraOverviewStack::OnSelectionChanged)
 		.SelectionMode(ESelectionMode::Multi)
+		.OnItemToString_Debug_Static(&FNiagaraStackEditorWidgetsUtilities::StackEntryToStringForListDebug)
 	];
 
 	InStackViewModel.OnStructureChanged().AddSP(this, &SNiagaraOverviewStack::EntryStructureChanged);
diff --git a/Engine/Plugins/FX/Niagara/Source/NiagaraEditorWidgets/Private/SNiagaraStack.cpp b/Engine/Plugins/FX/Niagara/Source/NiagaraEditorWidgets/Private/SNiagaraStack.cpp
--- a/Engine/Plugins/FX/Niagara/Source/NiagaraEditorWidgets/Private/SNiagaraStack.cpp
+++ b/Engine/Plugins/FX/Niagara/Source/NiagaraEditorWidgets/Private/SNiagaraStack.cpp
@@ -264,6 +264,7 @@ void SNiagaraStack::Construct(const FArguments& InArgs, UNiagaraStackViewModel*
 			.TreeItemsSource(&StackViewModel->GetRootEntryAsArray())
 			.OnTreeViewScrolled(this, &SNiagaraStack::StackTreeScrolled)
 			.SelectionMode(ESelectionMode::None)
+			.OnItemToString_Debug_Static(&FNiagaraStackEditorWidgetsUtilities::StackEntryToStringForListDebug)
 		]
 	];
 
diff --git a/Engine/Plugins/FX/Niagara/Source/NiagaraEditorWidgets/Private/Stack/SNiagaraStackErrorItem.cpp b/Engine/Plugins/FX/Niagara/Source/NiagaraEditorWidgets/Private/Stack/SNiagaraStackErrorItem.cpp
--- a/Engine/Plugins/FX/Niagara/Source/NiagaraEditorWidgets/Private/Stack/SNiagaraStackErrorItem.cpp
+++ b/Engine/Plugins/FX/Niagara/Source/NiagaraEditorWidgets/Private/Stack/SNiagaraStackErrorItem.cpp
@@ -78,7 +78,7 @@ void SNiagaraStackErrorItemFix::Construct(const FArguments& InArgs, UNiagaraStac
 		[
 			SNew(STextBlock)
 				.TextStyle(FNiagaraEditorStyle::Get(), "NiagaraEditor.ParameterText")
-				.Text_UObject(ErrorItem, &UNiagaraStackErrorItemFix::FixDescription)
+				.Text_UObject(ErrorItem, &UNiagaraStackErrorItemFix::GetDisplayName)
 				.ColorAndOpacity(this, &SNiagaraStackErrorItemFix::GetTextColorForSearch)
 				.HighlightText_UObject(StackViewModel, &UNiagaraStackViewModel::GetCurrentSearchText)
 		];
