Niagara - Hide the experimental (and hopefully short lived) fast path options from the UI with a CVAR.

diff --git a/Engine/Plugins/FX/Niagara/Source/NiagaraEditor/Private/Customizations/NiagaraSystemDetails.cpp b/Engine/Plugins/FX/Niagara/Source/NiagaraEditor/Private/Customizations/NiagaraSystemDetails.cpp
new file mode 100644
--- /dev/null
+++ b/Engine/Plugins/FX/Niagara/Source/NiagaraEditor/Private/Customizations/NiagaraSystemDetails.cpp
@@ -0,0 +1,20 @@
+// Copyright 1998-2019 Epic Games, Inc. All Rights Reserved.
+
+#include "NiagaraSystemDetails.h"
+#include "DetailLayoutBuilder.h"
+#include "DetailCategoryBuilder.h"
+#include "NiagaraEditorModule.h"
+
+TSharedRef<IDetailCustomization> FNiagaraSystemDetails::MakeInstance()
+{
+	return MakeShared<FNiagaraSystemDetails>();
+}
+
+void FNiagaraSystemDetails::CustomizeDetails(IDetailLayoutBuilder& InDetailLayout)
+{
+	if (GbShowFastPathOptions <= 0)
+	{
+		InDetailLayout.EditCategory("Script Fast Path").SetCategoryVisibility(false);
+	}
+}
+
diff --git a/Engine/Plugins/FX/Niagara/Source/NiagaraEditor/Private/Customizations/NiagaraSystemDetails.h b/Engine/Plugins/FX/Niagara/Source/NiagaraEditor/Private/Customizations/NiagaraSystemDetails.h
new file mode 100644
--- /dev/null
+++ b/Engine/Plugins/FX/Niagara/Source/NiagaraEditor/Private/Customizations/NiagaraSystemDetails.h
@@ -0,0 +1,18 @@
+// Copyright 1998-2019 Epic Games, Inc. All Rights Reserved.
+
+#pragma once
+
+#include "IDetailCustomization.h"
+
+class IDetailLayoutBuilder;
+
+class FNiagaraSystemDetails : public IDetailCustomization
+{
+public:
+	/** Makes a new instance of this detail layout class for a specific detail view requesting it */
+	static TSharedRef<IDetailCustomization> MakeInstance();
+
+	/** IDetailCustomization interface */
+	virtual void CustomizeDetails(IDetailLayoutBuilder& InDetailLayout) override;
+};
+
diff --git a/Engine/Plugins/FX/Niagara/Source/NiagaraEditor/Private/NiagaraEditorModule.cpp b/Engine/Plugins/FX/Niagara/Source/NiagaraEditor/Private/NiagaraEditorModule.cpp
--- a/Engine/Plugins/FX/Niagara/Source/NiagaraEditor/Private/NiagaraEditorModule.cpp
+++ b/Engine/Plugins/FX/Niagara/Source/NiagaraEditor/Private/NiagaraEditorModule.cpp
@@ -114,6 +114,14 @@ const FLinearColor FNiagaraEditorModule::WorldCentricTabColorScale(0.0f, 0.0f, 0
 
 EAssetTypeCategories::Type FNiagaraEditorModule::NiagaraAssetCategory;
 
+int32 GbShowFastPathOptions = 0;
+static FAutoConsoleVariableRef CVarShowFastPathOptions(
+	TEXT("fx.Niagara.ShowFastPathOptions"),
+	GbShowFastPathOptions,
+	TEXT("If > 0 the experimental fast path options will be shown in the system and emitter properties in the niagara system editor.\n"),
+	ECVF_Default
+);
+
 //////////////////////////////////////////////////////////////////////////
 
 class FNiagaraEditorOnlyDataUtilities : public INiagaraEditorOnlyDataUtilities
diff --git a/Engine/Plugins/FX/Niagara/Source/NiagaraEditor/Private/NiagaraEmitterDetailsCustomization.cpp b/Engine/Plugins/FX/Niagara/Source/NiagaraEditor/Private/NiagaraEmitterDetailsCustomization.cpp
--- a/Engine/Plugins/FX/Niagara/Source/NiagaraEditor/Private/NiagaraEmitterDetailsCustomization.cpp
+++ b/Engine/Plugins/FX/Niagara/Source/NiagaraEditor/Private/NiagaraEmitterDetailsCustomization.cpp
@@ -3,7 +3,9 @@
 #include "NiagaraEmitterDetailsCustomization.h"
 #include "PropertyHandle.h"
 #include "DetailLayoutBuilder.h"
+#include "DetailCategoryBuilder.h"
 #include "NiagaraEmitter.h"
+#include "NiagaraEditorModule.h"
 
 TSharedRef<IDetailCustomization> FNiagaraEmitterDetails::MakeInstance()
 {
@@ -14,5 +16,10 @@ void FNiagaraEmitterDetails::CustomizeDetails(IDetailLayoutBuilder& InDetailLayo
 {
 	TSharedPtr<IPropertyHandle> EventHandlersPropertyHandle = InDetailLayout.GetProperty(UNiagaraEmitter::PrivateMemberNames::EventHandlerScriptProps);
 	EventHandlersPropertyHandle->MarkHiddenByCustomization();
+
+	if (GbShowFastPathOptions <= 0)
+	{
+		InDetailLayout.EditCategory("Script Fast Path").SetCategoryVisibility(false);
+	}
 }
 
diff --git a/Engine/Plugins/FX/Niagara/Source/NiagaraEditor/Private/ViewModels/Stack/NiagaraStackSystemPropertiesItem.cpp b/Engine/Plugins/FX/Niagara/Source/NiagaraEditor/Private/ViewModels/Stack/NiagaraStackSystemPropertiesItem.cpp
--- a/Engine/Plugins/FX/Niagara/Source/NiagaraEditor/Private/ViewModels/Stack/NiagaraStackSystemPropertiesItem.cpp
+++ b/Engine/Plugins/FX/Niagara/Source/NiagaraEditor/Private/ViewModels/Stack/NiagaraStackSystemPropertiesItem.cpp
@@ -4,6 +4,7 @@
 #include "ViewModels/Stack/NiagaraStackObject.h"
 #include "NiagaraSystem.h"
 #include "ViewModels/NiagaraSystemViewModel.h"
+#include "Customizations/NiagaraSystemDetails.h"
 
 #define LOCTEXT_NAMESPACE "UNiagaraStackSystemItemGroup"
 
@@ -35,6 +36,7 @@ void UNiagaraStackSystemPropertiesItem::RefreshChildrenInternal(const TArray<UNi
 		SystemObject = NewObject<UNiagaraStackObject>(this);
 		FRequiredEntryData RequiredEntryData(GetSystemViewModel(), GetEmitterViewModel(), FExecutionCategoryNames::System, NAME_None, GetStackEditorData());
 		SystemObject->Initialize(RequiredEntryData, System.Get(), GetStackEditorDataKey());
+		SystemObject->RegisterInstancedCustomPropertyLayout(UNiagaraSystem::StaticClass(), FOnGetDetailCustomizationInstance::CreateStatic(&FNiagaraSystemDetails::MakeInstance));
 	}
 
 	NewChildren.Add(SystemObject);
diff --git a/Engine/Plugins/FX/Niagara/Source/NiagaraEditor/Public/NiagaraEditorModule.h b/Engine/Plugins/FX/Niagara/Source/NiagaraEditor/Public/NiagaraEditorModule.h
--- a/Engine/Plugins/FX/Niagara/Source/NiagaraEditor/Public/NiagaraEditorModule.h
+++ b/Engine/Plugins/FX/Niagara/Source/NiagaraEditor/Public/NiagaraEditorModule.h
@@ -37,6 +37,8 @@ public:
 	virtual FLinearColor GetColorForExecutionCategory(FName ExecutionCategory) const = 0;
 };
 
+extern int32 GbShowFastPathOptions;
+
 /** Niagara Editor module */
 class FNiagaraEditorModule : public IModuleInterface,
 	public IHasMenuExtensibility, public IHasToolBarExtensibility, public FGCObject
