Niagara - Fix an issue where data interfaces weren't handled correctly when dragging to move or copy modules between emitters and systems.

diff --git a/Engine/Plugins/FX/Niagara/Source/NiagaraEditor/Private/ViewModels/Stack/NiagaraStackGraphUtilities.cpp b/Engine/Plugins/FX/Niagara/Source/NiagaraEditor/Private/ViewModels/Stack/NiagaraStackGraphUtilities.cpp
--- a/Engine/Plugins/FX/Niagara/Source/NiagaraEditor/Private/ViewModels/Stack/NiagaraStackGraphUtilities.cpp
+++ b/Engine/Plugins/FX/Niagara/Source/NiagaraEditor/Private/ViewModels/Stack/NiagaraStackGraphUtilities.cpp
@@ -337,9 +337,34 @@ void GetGroupNodesRecursive(const TArray<UNiagaraNode*>& CurrentStartNodes, UNia
 		if (OutAllNodes.Contains(CurrentStartNode) == false)
 		{
 			OutAllNodes.Add(CurrentStartNode);
+
+			// Check input pins for this node to handle any UNiagaraNodeInput nodes which are wired directly to one of the group nodes.
+			UEdGraphPin* ParameterMapInputPin = FNiagaraStackGraphUtilities::GetParameterMapInputPin(*CurrentStartNode);
+			if (ParameterMapInputPin != nullptr)
+			{
+				TArray<UEdGraphPin*> InputPins;
+				CurrentStartNode->GetInputPins(InputPins);
+				for (UEdGraphPin* InputPin : InputPins)
+				{
+					if (InputPin != ParameterMapInputPin)
+					{
+						for (UEdGraphPin* InputLinkedPin : InputPin->LinkedTo)
+						{
+							UNiagaraNode* LinkedNode = Cast<UNiagaraNode>(InputLinkedPin->GetOwningNode());
+							if (LinkedNode != nullptr)
+							{
+								OutAllNodes.AddUnique(LinkedNode);
+							}
+						}
+					}
+				}
+			}
+
+			// Handle nodes connected to the output.
 			if (CurrentStartNode != EndNode)
 			{
 				TArray<UNiagaraNode*> LinkedNodes;
+				TArray<UEdGraphPin*> InputPins;
 				TArray<UEdGraphPin*> OutputPins;
 				CurrentStartNode->GetOutputPins(OutputPins);
 				for (UEdGraphPin* OutputPin : OutputPins)
@@ -1918,7 +1943,21 @@ void FNiagaraStackGraphUtilities::MoveModule(UNiagaraScript& SourceScript, UNiag
 	{
 		// If the module is being inserted into a different graph, or it's being copied, all of the nodes need to be copied into the target graph.
 		FStackNodeGroup SourceGroup = SourceGroups[SourceGroupIndex];
-		
+
+		// HACK! The following code and the code after the import/export is necessary since sub-objects with a "." in them will not be correctly imported from text!
+		TMap<FGuid, FString> NodeIdToOriginalName;
+		for (UNiagaraNode* SourceGroupNode : SourceGroupNodes)
+		{
+			UNiagaraNodeInput* InputSourceGroupNode = Cast<UNiagaraNodeInput>(SourceGroupNode);
+			if (InputSourceGroupNode != nullptr && InputSourceGroupNode->GetDataInterface() != nullptr)
+			{
+				NodeIdToOriginalName.Add(InputSourceGroupNode->NodeGuid, InputSourceGroupNode->GetDataInterface()->GetName());
+				FString NewSanitizedName = InputSourceGroupNode->GetDataInterface()->GetName().Replace(TEXT("."), TEXT("_"));
+				InputSourceGroupNode->GetDataInterface()->Rename(*NewSanitizedName);
+			}
+		}
+		// HACK end
+
 		TSet<UObject*> NodesToCopy;
 		for (UNiagaraNode* SourceGroupNode : SourceGroupNodes)
 		{
@@ -1933,6 +1972,20 @@ void FNiagaraStackGraphUtilities::MoveModule(UNiagaraScript& SourceScript, UNiag
 		FEdGraphUtilities::ImportNodesFromText(TargetGraph, ExportedText, CopiedNodesSet);
 		TArray<UEdGraphNode*> CopiedNodes = CopiedNodesSet.Array();
 
+		// HACK continued.
+		if (NodeIdToOriginalName.Num() > 0)
+		{
+			for (UEdGraphNode* CopiedNode : CopiedNodes)
+			{
+				FString* OriginalName = NodeIdToOriginalName.Find(CopiedNode->NodeGuid);
+				if (OriginalName != nullptr)
+				{
+					CastChecked<UNiagaraNodeInput>(CopiedNode)->GetDataInterface()->Rename(**OriginalName);
+				}
+			}
+		}
+		// HACK end
+
 		// Collect the start and end nodes for the group by ID before assigning the copied nodes new ids.
 		UEdGraphNode** CopiedEndNode = CopiedNodes.FindByPredicate([SourceGroup](UEdGraphNode* CopiedNode)
 			{ return CopiedNode->NodeGuid == SourceGroup.EndNode->NodeGuid; });
