Niagara - Fix a drag-drop message to let the user know how the functionality has changed.

diff --git a/Engine/Plugins/FX/Niagara/Source/NiagaraEditor/Private/ViewModels/Stack/NiagaraStackScriptItemGroup.cpp b/Engine/Plugins/FX/Niagara/Source/NiagaraEditor/Private/ViewModels/Stack/NiagaraStackScriptItemGroup.cpp
--- a/Engine/Plugins/FX/Niagara/Source/NiagaraEditor/Private/ViewModels/Stack/NiagaraStackScriptItemGroup.cpp
+++ b/Engine/Plugins/FX/Niagara/Source/NiagaraEditor/Private/ViewModels/Stack/NiagaraStackScriptItemGroup.cpp
@@ -698,7 +698,8 @@ TOptional<UNiagaraStackEntry::FDropRequestResponse> UNiagaraStackScriptItemGroup
 	if (DropRequest.DropOptions != UNiagaraStackEntry::EDropOptions::Overview)
 	{
 		// Only allow dropping in the overview stacks.
-		return FDropRequestResponse(TOptional<EItemDropZone>(), LOCTEXT("CantDropParameterOnStack", "Parameters can only be dropped onto 'Set Variables' modules in the selection view."));
+		return FDropRequestResponse(TOptional<EItemDropZone>(), LOCTEXT("CantDropParameterOnStack", 
+			"Parameters can only be dropped onto 'Set Variables' modules, or correctly\ntyped inputs in the selection view.  If you want to add a new 'Set Variables' module for\n this parameter, you can drop it into one of the nodes in the System Overview graph."));
 	}
 	if (ParameterAction.IsValid())
 	{
