Niagara - Fix an issue where the editor would crash when alt+draging items in the timeline.

diff --git a/Engine/Plugins/FX/Niagara/Source/NiagaraEditor/Private/Sequencer/NiagaraSequence/MovieSceneNiagaraEmitterTrack.cpp b/Engine/Plugins/FX/Niagara/Source/NiagaraEditor/Private/Sequencer/NiagaraSequence/MovieSceneNiagaraEmitterTrack.cpp
--- a/Engine/Plugins/FX/Niagara/Source/NiagaraEditor/Private/Sequencer/NiagaraSequence/MovieSceneNiagaraEmitterTrack.cpp
+++ b/Engine/Plugins/FX/Niagara/Source/NiagaraEditor/Private/Sequencer/NiagaraSequence/MovieSceneNiagaraEmitterTrack.cpp
@@ -160,6 +160,11 @@ bool UMovieSceneNiagaraEmitterTrack::HasSection(const UMovieSceneSection& Sectio
 	return Sections.Contains(&Section);
 }
 
+void UMovieSceneNiagaraEmitterTrack::AddSection(UMovieSceneSection& Section)
+{
+	Sections.Add(&Section);
+}
+
 void UMovieSceneNiagaraEmitterTrack::RemoveSection(UMovieSceneSection& Section)
 {
 	Sections.Remove(&Section);
diff --git a/Engine/Plugins/FX/Niagara/Source/NiagaraEditor/Private/Sequencer/NiagaraSequence/MovieSceneNiagaraEmitterTrack.h b/Engine/Plugins/FX/Niagara/Source/NiagaraEditor/Private/Sequencer/NiagaraSequence/MovieSceneNiagaraEmitterTrack.h
--- a/Engine/Plugins/FX/Niagara/Source/NiagaraEditor/Private/Sequencer/NiagaraSequence/MovieSceneNiagaraEmitterTrack.h
+++ b/Engine/Plugins/FX/Niagara/Source/NiagaraEditor/Private/Sequencer/NiagaraSequence/MovieSceneNiagaraEmitterTrack.h
@@ -73,6 +73,7 @@ public:
 	//~ UMovieSceneTrack interface
 	virtual void RemoveAllAnimationData() override { }
 	virtual bool HasSection(const UMovieSceneSection& Section) const override;
+	virtual void AddSection(UMovieSceneSection& Section) override;
 	virtual void RemoveSection(UMovieSceneSection& Section) override;
 	virtual void RemoveSectionAt(int32 SectionIndex) override;
 	virtual bool IsEmpty() const override;
diff --git a/Engine/Plugins/FX/Niagara/Source/NiagaraEditor/Private/ViewModels/NiagaraSystemViewModel.cpp b/Engine/Plugins/FX/Niagara/Source/NiagaraEditor/Private/ViewModels/NiagaraSystemViewModel.cpp
--- a/Engine/Plugins/FX/Niagara/Source/NiagaraEditor/Private/ViewModels/NiagaraSystemViewModel.cpp
+++ b/Engine/Plugins/FX/Niagara/Source/NiagaraEditor/Private/ViewModels/NiagaraSystemViewModel.cpp
@@ -1418,18 +1418,28 @@ void FNiagaraSystemViewModel::SequencerDataChanged(EMovieSceneDataChangeType Dat
 		TArray<FEmitterHandleToDuplicate> EmittersToDuplicate;
 		TArray<TTuple<TSharedPtr<FNiagaraEmitterHandleViewModel>, FName>> EmitterHandlesToRename;
 
+		bool bRefreshAllTracks = false;
 		for (UMovieSceneTrack* Track : NiagaraSequence->GetMovieScene()->GetMasterTracks())
 		{
 			UMovieSceneNiagaraEmitterTrack* EmitterTrack = CastChecked<UMovieSceneNiagaraEmitterTrack>(Track);
 			if (EmitterTrack->GetEmitterHandleViewModel().IsValid())
 			{
-				VaildTrackEmitterHandleIds.Add(EmitterTrack->GetEmitterHandleViewModel()->GetId());
-				EmitterTrack->UpdateEmitterHandleFromTrackChange(NiagaraSequence->GetMovieScene()->GetTickResolution());
-				EmitterTrack->GetEmitterHandleViewModel()->GetEmitterViewModel()->GetOrCreateEditorData().Modify();
-				EmitterTrack->GetEmitterHandleViewModel()->GetEmitterViewModel()->GetOrCreateEditorData().SetPlaybackRange(GetEditorData().GetPlaybackRange());
-				if (EmitterTrack->GetDisplayName().ToString() != EmitterTrack->GetEmitterHandleViewModel()->GetNameText().ToString())
+				if (EmitterTrack->GetAllSections().Num() > 0)
 				{
-					EmitterHandlesToRename.Add(TTuple<TSharedPtr<FNiagaraEmitterHandleViewModel>, FName>(EmitterTrack->GetEmitterHandleViewModel(), *EmitterTrack->GetDisplayName().ToString()));
+					VaildTrackEmitterHandleIds.Add(EmitterTrack->GetEmitterHandleViewModel()->GetId());
+					EmitterTrack->UpdateEmitterHandleFromTrackChange(NiagaraSequence->GetMovieScene()->GetTickResolution());
+					EmitterTrack->GetEmitterHandleViewModel()->GetEmitterViewModel()->GetOrCreateEditorData().Modify();
+					EmitterTrack->GetEmitterHandleViewModel()->GetEmitterViewModel()->GetOrCreateEditorData().SetPlaybackRange(GetEditorData().GetPlaybackRange());
+					if (EmitterTrack->GetDisplayName().ToString() != EmitterTrack->GetEmitterHandleViewModel()->GetNameText().ToString())
+					{
+						EmitterHandlesToRename.Add(TTuple<TSharedPtr<FNiagaraEmitterHandleViewModel>, FName>(EmitterTrack->GetEmitterHandleViewModel(), *EmitterTrack->GetDisplayName().ToString()));
+					}
+					if (EmitterTrack->GetAllSections().Num() > 1)
+					{
+						// If a section was duplicated, force a refresh to remove the duplicated section since that's not currently supported.
+						// TODO: Detect duplicated sections and create a duplicate emitter with the correct values.
+						bRefreshAllTracks = true;
+					}
 				}
 			}
 			else
@@ -1445,7 +1455,7 @@ void FNiagaraSystemViewModel::SequencerDataChanged(EMovieSceneDataChangeType Dat
 			}
 		}
 
-		bool bRefreshAllTracks = EmitterHandlesToRename.Num() > 0;
+		bRefreshAllTracks |= EmitterHandlesToRename.Num() > 0;
 
 		for (TTuple<TSharedPtr<FNiagaraEmitterHandleViewModel>, FName>& EmitterHandletoRename : EmitterHandlesToRename)
 		{
