Niagara - Fix system parameters not being shown in the parameters pane in the system and emitter editor.

diff --git a/Engine/Plugins/FX/Niagara/Source/NiagaraEditor/Private/Toolkits/NiagaraSystemToolkit.cpp b/Engine/Plugins/FX/Niagara/Source/NiagaraEditor/Private/Toolkits/NiagaraSystemToolkit.cpp
--- a/Engine/Plugins/FX/Niagara/Source/NiagaraEditor/Private/Toolkits/NiagaraSystemToolkit.cpp
+++ b/Engine/Plugins/FX/Niagara/Source/NiagaraEditor/Private/Toolkits/NiagaraSystemToolkit.cpp
@@ -289,6 +289,7 @@ void FNiagaraSystemToolkit::InitializeWithEmitter(const EToolkitMode::Type Mode,
 void FNiagaraSystemToolkit::InitializeInternal(const EToolkitMode::Type Mode, const TSharedPtr<IToolkitHost>& InitToolkitHost, const FGuid& MessageLogGuid)
 {
 	NiagaraMessageLogViewModel = MakeShared<FNiagaraMessageLogViewModel>(GetNiagaraSystemMessageLogName(System), MessageLogGuid, NiagaraMessageLog);
+	ObjectSelectionForParameterMapView = MakeShared<FNiagaraObjectSelection>();
 
 	SystemViewModel->OnEmitterHandleViewModelsChanged().AddSP(this, &FNiagaraSystemToolkit::RefreshParameters);
 	SystemViewModel->GetSelectionViewModel()->OnSelectionChanged().AddSP(this, &FNiagaraSystemToolkit::OnSystemSelectionChanged);
@@ -491,26 +492,16 @@ TSharedRef<SDockTab> FNiagaraSystemToolkit::SpawnTab_SystemParameters(const FSpa
 {
 	check(Args.GetTabId().TabType == SystemParametersTabID);
 
-	TSharedRef<FNiagaraObjectSelection> ObjectSelection = MakeShareable(new FNiagaraObjectSelection());
-	if (SystemToolkitMode == ESystemToolkitMode::Emitter)
-	{
-		TSharedPtr<FNiagaraEmitterViewModel> EditableEmitterViewModel = SystemViewModel->GetEmitterHandleViewModels()[0]->GetEmitterViewModel();
-		UNiagaraEmitter* EditableEmitter = EditableEmitterViewModel->GetEmitter();
-		ObjectSelection->SetSelectedObject(EditableEmitter);
-	}
-	else if (SystemToolkitMode == ESystemToolkitMode::System)
-	{
-		ObjectSelection->SetSelectedObject(System);
-	}
 
-	TArray<TSharedRef<FNiagaraObjectSelection>> Array;
-	Array.Push(ObjectSelection);
+	TArray<TSharedRef<FNiagaraObjectSelection>> ObjectSelections;
+	ObjectSelections.Add(ObjectSelectionForParameterMapView.ToSharedRef());
 
 	TSharedRef<SDockTab> SpawnedTab =
 		SNew(SDockTab)
 		[
-			SAssignNew(ParameterMapView, SNiagaraParameterMapView, Array, SNiagaraParameterMapView::EToolkitType::SYSTEM, GetToolkitCommands())
+			SAssignNew(ParameterMapView, SNiagaraParameterMapView, ObjectSelections, SNiagaraParameterMapView::EToolkitType::SYSTEM, GetToolkitCommands())
 		];
+	RefreshParameters();
 
 	return SpawnedTab;
 }
@@ -1284,22 +1275,32 @@ void FNiagaraSystemToolkit::OnPinnedCurvesChanged()
 
 void FNiagaraSystemToolkit::RefreshParameters()
 {
-	if (ParameterMapView.IsValid())
-	{
-		TArray<TSharedPtr<FNiagaraEmitterHandleViewModel>> EmitterHandlesToDisplay;
-		EmitterHandlesToDisplay.Append(SystemViewModel->GetPinnedEmitterHandles());
+	TArray<UObject*> NewParameterViewSelection;
+
+	// Always display the system parameters
+	NewParameterViewSelection.Add(&SystemViewModel->GetSystem());
+
+	TArray<TSharedRef<FNiagaraEmitterHandleViewModel>> EmitterHandlesToDisplay;
+	EmitterHandlesToDisplay.Append(SystemViewModel->GetPinnedEmitterHandles());
 		
-		TArray<FGuid> SelectedEmitterHandleIds = SystemViewModel->GetSelectionViewModel()->GetSelectedEmitterHandleIds();
-		for (TSharedRef<FNiagaraEmitterHandleViewModel> EmitterHandleViewModel : SystemViewModel->GetEmitterHandleViewModels())
+	TArray<FGuid> SelectedEmitterHandleIds = SystemViewModel->GetSelectionViewModel()->GetSelectedEmitterHandleIds();
+	for (TSharedRef<FNiagaraEmitterHandleViewModel> EmitterHandleViewModel : SystemViewModel->GetEmitterHandleViewModels())
+	{
+		if (SelectedEmitterHandleIds.Contains(EmitterHandleViewModel->GetId()))
 		{
-			if (SelectedEmitterHandleIds.Contains(EmitterHandleViewModel->GetId()))
-			{
-				EmitterHandlesToDisplay.AddUnique(EmitterHandleViewModel);
-			}
+			EmitterHandlesToDisplay.AddUnique(EmitterHandleViewModel);
 		}
+	}
 
-		ParameterMapView->RefreshEmitterHandles(EmitterHandlesToDisplay);
+	for (TSharedRef<FNiagaraEmitterHandleViewModel> EmitterHandleToDisplay : EmitterHandlesToDisplay)
+	{
+		if (EmitterHandleToDisplay->IsValid() && EmitterHandleToDisplay->GetEmitterViewModel()->GetEmitter() != nullptr)
+		{
+			NewParameterViewSelection.Add(EmitterHandleToDisplay->GetEmitterViewModel()->GetEmitter());
+		}
 	}
+
+	ObjectSelectionForParameterMapView->SetSelectedObjects(NewParameterViewSelection);
 }
 
 void FNiagaraSystemToolkit::OnSystemSelectionChanged(UNiagaraSystemSelectionViewModel::ESelectionChangeSource SelectionChangeSource)
diff --git a/Engine/Plugins/FX/Niagara/Source/NiagaraEditor/Private/Toolkits/NiagaraSystemToolkit.h b/Engine/Plugins/FX/Niagara/Source/NiagaraEditor/Private/Toolkits/NiagaraSystemToolkit.h
--- a/Engine/Plugins/FX/Niagara/Source/NiagaraEditor/Private/Toolkits/NiagaraSystemToolkit.h
+++ b/Engine/Plugins/FX/Niagara/Source/NiagaraEditor/Private/Toolkits/NiagaraSystemToolkit.h
@@ -18,6 +18,7 @@
 
 class FNiagaraSystemInstance;
 class FNiagaraSystemViewModel;
+class FNiagaraObjectSelection;
 class SNiagaraSystemEditorViewport;
 class SNiagaraSystemEditorWidget;
 class SNiagaraSystemViewport;
@@ -162,6 +163,8 @@ private:
 
 	TSharedPtr<class SNiagaraParameterMapView> ParameterMapView;
 
+	TSharedPtr<FNiagaraObjectSelection> ObjectSelectionForParameterMapView;
+
 	bool bChangesDiscarded;
 
 	static const FName ViewportTabID;
diff --git a/Engine/Plugins/FX/Niagara/Source/NiagaraEditor/Private/Widgets/SNiagaraParameterMapView.cpp b/Engine/Plugins/FX/Niagara/Source/NiagaraEditor/Private/Widgets/SNiagaraParameterMapView.cpp
--- a/Engine/Plugins/FX/Niagara/Source/NiagaraEditor/Private/Widgets/SNiagaraParameterMapView.cpp
+++ b/Engine/Plugins/FX/Niagara/Source/NiagaraEditor/Private/Widgets/SNiagaraParameterMapView.cpp
@@ -141,7 +141,8 @@ SNiagaraParameterMapView::~SNiagaraParameterMapView()
 		}
 	}
 
-	Graphs.Empty();
+	EmptyGraphs();
+
 	SelectedScriptObjects->OnSelectedObjectsChanged().RemoveAll(this);
 	if (SelectedVariableObjects)
 	{
@@ -584,7 +585,7 @@ EVisibility SNiagaraParameterMapView::OnAddButtonTextVisibility(TWeakPtr<SWidget
 
 void SNiagaraParameterMapView::Refresh(bool bRefreshMenu/* = true*/)
 {
-	Graphs.Empty();
+	EmptyGraphs();
 
 	TSet<UObject*> Objects = SelectedScriptObjects->GetSelectedObjects();
 	for (UObject* Object : Objects)
@@ -597,16 +598,10 @@ void SNiagaraParameterMapView::Refresh(bool bRefreshMenu/* = true*/)
 		else if (UNiagaraEmitter* Emitter = Cast<UNiagaraEmitter>(Object))
 		{
 			AddGraph(Emitter->GraphSource);
-			break;
 		}
 		else if (UNiagaraSystem* System = Cast<UNiagaraSystem>(Object))
 		{
-			for (const FNiagaraEmitterHandle& EmitterHandle : System->GetEmitterHandles())
-			{
-				AddGraph(EmitterHandle.GetInstance()->GraphSource);
-			}
-			System->GetExposedParameters().AddOnChangedHandler(FNiagaraParameterStore::FOnChanged::FDelegate::CreateSP(this, &SNiagaraParameterMapView::RefreshActions));
-			break;
+			AddGraph(System->GetSystemSpawnScript()->GetSource());
 		}
 	}
 
@@ -616,30 +611,39 @@ void SNiagaraParameterMapView::Refresh(bool bRefreshMenu/* = true*/)
 	}
 }
 
-void SNiagaraParameterMapView::RefreshEmitterHandles(const TArray<TSharedPtr<FNiagaraEmitterHandleViewModel>>& EmitterHandles)
+void SNiagaraParameterMapView::SelectedObjectsChanged()
 {
-	Graphs.Empty();
-	for (const TSharedPtr<FNiagaraEmitterHandleViewModel>& Handle : EmitterHandles)
-	{
-		AddGraph(Handle->GetEmitterHandle()->GetInstance()->GraphSource);
-	}
-	GraphActionMenu->RefreshAllActions(true);
+	Refresh(true);
 }
 
-void SNiagaraParameterMapView::SelectedObjectsChanged()
+void SNiagaraParameterMapView::EmptyGraphs()
 {
-	Refresh(true);
+	checkf(Graphs.Num() == OnGraphChangedHandles.Num() && Graphs.Num() == OnRecompileHandles.Num(), TEXT("Graphs and change delegates out of sync!"));
+	for (int GraphIndex = 0; GraphIndex < Graphs.Num(); ++GraphIndex)
+	{
+		if (Graphs[GraphIndex].IsValid())
+		{
+			Graphs[GraphIndex]->RemoveOnGraphChangedHandler(OnGraphChangedHandles[GraphIndex]);
+			Graphs[GraphIndex]->RemoveOnGraphNeedsRecompileHandler(OnRecompileHandles[GraphIndex]);
+		}
+	}
+	Graphs.Empty();
+	OnGraphChangedHandles.Empty();
+	OnRecompileHandles.Empty();
 }
 
 void SNiagaraParameterMapView::AddGraph(UNiagaraGraph* Graph)
 {
-	if (Graph)
+	if (Graph && Graphs.Contains(Graph) == false)
 	{
-		Graphs.AddUnique(Graph);
-		OnGraphChangedHandle = Graph->AddOnGraphChangedHandler(
+		Graphs.Add(Graph);
+		FDelegateHandle OnGraphChangedHandle = Graph->AddOnGraphChangedHandler(
 			FOnGraphChanged::FDelegate::CreateRaw(this, &SNiagaraParameterMapView::OnGraphChanged));
-		OnRecompileHandle = Graph->AddOnGraphNeedsRecompileHandler(
+		FDelegateHandle OnRecompileHandle = Graph->AddOnGraphNeedsRecompileHandler(
 			FOnGraphChanged::FDelegate::CreateRaw(this, &SNiagaraParameterMapView::OnGraphChanged));
+
+		OnGraphChangedHandles.Add(OnGraphChangedHandle);
+		OnRecompileHandles.Add(OnRecompileHandle);
 	}
 }
 
diff --git a/Engine/Plugins/FX/Niagara/Source/NiagaraEditor/Private/Widgets/SNiagaraParameterMapView.h b/Engine/Plugins/FX/Niagara/Source/NiagaraEditor/Private/Widgets/SNiagaraParameterMapView.h
--- a/Engine/Plugins/FX/Niagara/Source/NiagaraEditor/Private/Widgets/SNiagaraParameterMapView.h
+++ b/Engine/Plugins/FX/Niagara/Source/NiagaraEditor/Private/Widgets/SNiagaraParameterMapView.h
@@ -89,7 +89,6 @@ public:
 
 	/** Refreshes the graphs used for this menu. */
 	void Refresh(bool bRefreshMenu = true);
-	void RefreshEmitterHandles(const TArray<TSharedPtr<class FNiagaraEmitterHandleViewModel>>& EmitterHandles);
 
 	static TSharedRef<SExpanderArrow> CreateCustomActionExpander(const struct FCustomExpanderData& ActionMenuData);
 
@@ -122,6 +121,7 @@ private:
 	EVisibility OnAddButtonTextVisibility(TWeakPtr<SWidget> RowWidget, const NiagaraParameterMapSectionID::Type InSection) const;
 
 	void SelectedObjectsChanged();
+	void EmptyGraphs();
 	void AddGraph(UNiagaraGraph* Graph);
 	void AddGraph(class UNiagaraScriptSourceBase* SourceBase);
 	void OnGraphChanged(const struct FEdGraphEditAction& InAction);
@@ -161,8 +161,8 @@ private:
 	TArray<TWeakObjectPtr<UNiagaraGraph>> Graphs;
 
 	/** The handle to the graph changed delegate. */
-	FDelegateHandle OnGraphChangedHandle;
-	FDelegateHandle OnRecompileHandle;
+	TArray<FDelegateHandle> OnGraphChangedHandles;
+	TArray<FDelegateHandle> OnRecompileHandles;
 
 	EToolkitType ToolkitType;
 	TSharedPtr<FUICommandList> ToolkitCommands;
