Forced GPU emitters using distance field data to be ticked in the view dependent pass.

diff --git a/Engine/Plugins/FX/Niagara/Source/Niagara/Private/NiagaraEmitterInstanceBatcher.cpp b/Engine/Plugins/FX/Niagara/Source/Niagara/Private/NiagaraEmitterInstanceBatcher.cpp
--- a/Engine/Plugins/FX/Niagara/Source/Niagara/Private/NiagaraEmitterInstanceBatcher.cpp
+++ b/Engine/Plugins/FX/Niagara/Source/Niagara/Private/NiagaraEmitterInstanceBatcher.cpp
@@ -538,7 +538,7 @@ void NiagaraEmitterInstanceBatcher::ExecuteAll(FRHICommandList &RHICmdList, FRHI
 			{
 				continue;
 			}
-			else if (GNiagaraAllowTickBeforeRender && ComputeShader->ViewUniformBufferParam.IsBound() != (ViewUniformBuffer != nullptr))
+			else if (GNiagaraAllowTickBeforeRender && (ComputeShader->ViewUniformBufferParam.IsBound() || Tick.bRequiredDistanceFieldData) != (ViewUniformBuffer != nullptr))
 			{   // When allowing tick before render, skip this emitter if it is not in the right pass.
 				continue;
 			}
