Niagara - Fix a crash caused by UI holding on to objects which were being garbage collected.

diff --git a/Engine/Plugins/FX/Niagara/Source/NiagaraEditorWidgets/Private/Stack/SNiagaraStackIssueIcon.cpp b/Engine/Plugins/FX/Niagara/Source/NiagaraEditorWidgets/Private/Stack/SNiagaraStackIssueIcon.cpp
--- a/Engine/Plugins/FX/Niagara/Source/NiagaraEditorWidgets/Private/Stack/SNiagaraStackIssueIcon.cpp
+++ b/Engine/Plugins/FX/Niagara/Source/NiagaraEditorWidgets/Private/Stack/SNiagaraStackIssueIcon.cpp
@@ -39,7 +39,7 @@ void SNiagaraStackIssueIcon::Construct(const FArguments& InArgs, UNiagaraStackVi
 
 SNiagaraStackIssueIcon::~SNiagaraStackIssueIcon()
 {
-	if (StackEntry != nullptr)
+	if (StackEntry.IsValid())
 	{
 		StackEntry->OnStructureChanged().RemoveAll(this);
 	}
@@ -47,7 +47,7 @@ SNiagaraStackIssueIcon::~SNiagaraStackIssueIcon()
 
 bool SNiagaraStackIssueIcon::GetIconIsEnabled() const
 {
-	return StackEntry != nullptr && StackEntry->GetOwnerIsEnabled() && StackEntry->GetIsEnabled();
+	return StackEntry.IsValid() && StackEntry->IsFinalized() == false && StackEntry->GetOwnerIsEnabled() && StackEntry->GetIsEnabled();
 }
 
 const FSlateBrush* SNiagaraStackIssueIcon::GetIconBrush() const
@@ -59,117 +59,124 @@ FText SNiagaraStackIssueIcon::GetIconToolTip() const
 {
 	if (IconToolTipCache.IsSet() == false)
 	{
-		FTextBuilder ToolTipBuilder;
-		TArray<FText> ToolTipParts;
-		if (StackEntry->GetTotalNumberOfErrorIssues() > 0)
+		if (StackEntry.IsValid() == false || StackEntry->IsFinalized())
+		{
+			IconToolTipCache = FText();
+		}
+		else
 		{
-			if (StackEntry->GetTotalNumberOfErrorIssues() == 1)
+			FTextBuilder ToolTipBuilder;
+			TArray<FText> ToolTipParts;
+			if (StackEntry->GetTotalNumberOfErrorIssues() > 0)
 			{
-				ToolTipParts.Add(LOCTEXT("ErrorFormatSingle", "1 error"));
+				if (StackEntry->GetTotalNumberOfErrorIssues() == 1)
+				{
+					ToolTipParts.Add(LOCTEXT("ErrorFormatSingle", "1 error"));
+				}
+				else
+				{
+					ToolTipParts.Add(FText::Format(LOCTEXT("ErrorFormatMultiple", "{0} errors"), FText::AsNumber(StackEntry->GetTotalNumberOfErrorIssues())));
+				}
 			}
-			else
+			if (StackEntry->GetTotalNumberOfWarningIssues() > 0)
 			{
-				ToolTipParts.Add(FText::Format(LOCTEXT("ErrorFormatMultiple", "{0} errors"), FText::AsNumber(StackEntry->GetTotalNumberOfErrorIssues())));
+				if (StackEntry->GetTotalNumberOfWarningIssues() == 1)
+				{
+					ToolTipParts.Add(LOCTEXT("WarningFormatSingle", "1 warning"));
+				}
+				else
+				{
+					ToolTipParts.Add(FText::Format(LOCTEXT("WarningFormatMultiple", "{0} warnings"), FText::AsNumber(StackEntry->GetTotalNumberOfWarningIssues())));
+				}
 			}
-		}
-		if (StackEntry->GetTotalNumberOfWarningIssues() > 0)
-		{
-			if (StackEntry->GetTotalNumberOfWarningIssues() == 1)
+			if (StackEntry->GetTotalNumberOfInfoIssues() > 0)
 			{
-				ToolTipParts.Add(LOCTEXT("WarningFormatSingle", "1 warning"));
+				if (StackEntry->GetTotalNumberOfInfoIssues() == 1)
+				{
+					ToolTipParts.Add(LOCTEXT("InfoFormatSingle", "1 info"));
+				}
+				else
+				{
+					ToolTipParts.Add(FText::Format(LOCTEXT("InfoFormatMultiple", "{0} infos"), FText::AsNumber(StackEntry->GetTotalNumberOfInfoIssues())));
+				}
 			}
-			else
+
+			if (ToolTipParts.Num() == 3)
 			{
-				ToolTipParts.Add(FText::Format(LOCTEXT("WarningFormatMultiple", "{0} warnings"), FText::AsNumber(StackEntry->GetTotalNumberOfWarningIssues())));
+				ToolTipBuilder.AppendLineFormat(LOCTEXT("ThreePartFormat", "{0}, {1}, and {2}"), ToolTipParts[0], ToolTipParts[1], ToolTipParts[2]);
 			}
-		}
-		if (StackEntry->GetTotalNumberOfInfoIssues() > 0)
-		{
-			if (StackEntry->GetTotalNumberOfInfoIssues() == 1)
+			else if (ToolTipParts.Num() == 2)
 			{
-				ToolTipParts.Add(LOCTEXT("InfoFormatSingle", "1 info"));
+				ToolTipBuilder.AppendLineFormat(LOCTEXT("TwoPartFormat", "{0} and {1}"), ToolTipParts[0], ToolTipParts[1]);
 			}
-			else
+			else if (ToolTipParts.Num() == 1)
 			{
-				ToolTipParts.Add(FText::Format(LOCTEXT("InfoFormatMultiple", "{0} infos"), FText::AsNumber(StackEntry->GetTotalNumberOfInfoIssues())));
+				ToolTipBuilder.AppendLine(ToolTipParts[0]);
 			}
-		}
 
-		if (ToolTipParts.Num() == 3)
-		{
-			ToolTipBuilder.AppendLineFormat(LOCTEXT("ThreePartFormat", "{0}, {1}, and {2}"), ToolTipParts[0], ToolTipParts[1], ToolTipParts[2]);
-		}
-		else if (ToolTipParts.Num() == 2)
-		{
-			ToolTipBuilder.AppendLineFormat(LOCTEXT("TwoPartFormat", "{0} and {1}"), ToolTipParts[0], ToolTipParts[1]);
-		}
-		else if (ToolTipParts.Num() == 1)
-		{
-			ToolTipBuilder.AppendLine(ToolTipParts[0]);
-		}
-
-		FText IssueLineFormat = LOCTEXT("IssueLineFormat", "{0} - {1} - {2}");
-		auto SeverityToText = [](EStackIssueSeverity Severity)
-		{
-			switch (Severity)
+			FText IssueLineFormat = LOCTEXT("IssueLineFormat", "{0} - {1} - {2}");
+			auto SeverityToText = [](EStackIssueSeverity Severity)
 			{
-			case EStackIssueSeverity::Error:
-				return LOCTEXT("Error", "Error");
-			case EStackIssueSeverity::Warning:
-				return LOCTEXT("Warning", "Warning");
-			case EStackIssueSeverity::Info:
-				return LOCTEXT("Issue", "Issue");
-			default:
-				return FText();
-			}
-		};
+				switch (Severity)
+				{
+				case EStackIssueSeverity::Error:
+					return LOCTEXT("Error", "Error");
+				case EStackIssueSeverity::Warning:
+					return LOCTEXT("Warning", "Warning");
+				case EStackIssueSeverity::Info:
+					return LOCTEXT("Issue", "Issue");
+				default:
+					return FText();
+				}
+			};
 
-		auto GetFullDisplayName = [](const UNiagaraStackViewModel* InStackViewModel, UNiagaraStackEntry* EntryWithIssue)
-		{
-			TArray<UNiagaraStackEntry*> EntryPath;
-			InStackViewModel->GetPathForEntry(EntryWithIssue, EntryPath);
-			EntryPath.Add(EntryWithIssue);
-			TArray<FText> DisplayNameParts;
-			for (UNiagaraStackEntry* Entry : EntryPath)
+			auto GetFullDisplayName = [](const UNiagaraStackViewModel* InStackViewModel, UNiagaraStackEntry* EntryWithIssue)
+			{
+				TArray<UNiagaraStackEntry*> EntryPath;
+				InStackViewModel->GetPathForEntry(EntryWithIssue, EntryPath);
+				EntryPath.Add(EntryWithIssue);
+				TArray<FText> DisplayNameParts;
+				for (UNiagaraStackEntry* Entry : EntryPath)
+				{
+					if (Entry->GetDisplayName().IsEmpty() == false)
+					{
+						DisplayNameParts.Add(Entry->GetDisplayName());
+					}
+				}
+				return FText::Join(LOCTEXT("DisplayNameJoinDelimiter", " - "), DisplayNameParts);
+			};
+
+			int32 IssueLinesAppended = 0;
+			int32 MaxIssueLines = 5;
+			int32 TotalIssues = 0;
+			TArray<UNiagaraStackEntry*> EntriesToCheck;
+			EntriesToCheck.Add(StackEntry.Get());
+			EntriesToCheck.Append(StackEntry->GetAllChildrenWithIssues());
+			for (UNiagaraStackEntry* EntryToCheck : EntriesToCheck)
 			{
-				if (Entry->GetDisplayName().IsEmpty() == false)
+				for (int32 IssueIndex = 0; IssueIndex < EntryToCheck->GetIssues().Num() && IssueLinesAppended < MaxIssueLines; IssueIndex++, IssueLinesAppended++)
 				{
-					DisplayNameParts.Add(Entry->GetDisplayName());
+					const UNiagaraStackEntry::FStackIssue& Issue = EntryToCheck->GetIssues()[IssueIndex];
+					ToolTipBuilder.AppendLineFormat(IssueLineFormat, GetFullDisplayName(StackViewModel, EntryToCheck), SeverityToText(Issue.GetSeverity()), Issue.GetShortDescription());
 				}
+				TotalIssues += EntryToCheck->GetIssues().Num();
 			}
-			return FText::Join(LOCTEXT("DisplayNameJoinDelimiter", " - "), DisplayNameParts);
-		};
-
-		int32 IssueLinesAppended = 0;
-		int32 MaxIssueLines = 5;
-		int32 TotalIssues = 0;
-		TArray<UNiagaraStackEntry*> EntriesToCheck;
-		EntriesToCheck.Add(StackEntry);
-		EntriesToCheck.Append(StackEntry->GetAllChildrenWithIssues());
-		for (UNiagaraStackEntry* EntryToCheck : EntriesToCheck)
-		{
-			for (int32 IssueIndex = 0; IssueIndex < EntryToCheck->GetIssues().Num() && IssueLinesAppended < MaxIssueLines; IssueIndex++, IssueLinesAppended++)
+
+			if (TotalIssues > MaxIssueLines)
 			{
-				const UNiagaraStackEntry::FStackIssue& Issue = EntryToCheck->GetIssues()[IssueIndex];
-				ToolTipBuilder.AppendLineFormat(IssueLineFormat, GetFullDisplayName(StackViewModel, EntryToCheck), SeverityToText(Issue.GetSeverity()), Issue.GetShortDescription());
+				ToolTipBuilder.AppendLineFormat(LOCTEXT("MoreIssuesFormat", "(And {0} more {0}|plural(one=issue,other=issues)...)"), FText::AsNumber(TotalIssues - MaxIssueLines));
+				//"There {NumCats}|plural(one=is,other=are) {NumCats} {NumCats}|plural(one=cat,other=cats)"
 			}
-			TotalIssues += EntryToCheck->GetIssues().Num();
-		}
 
-		if (TotalIssues > MaxIssueLines)
-		{
-			ToolTipBuilder.AppendLineFormat(LOCTEXT("MoreIssuesFormat", "(And {0} more {0}|plural(one=issue,other=issues)...)"), FText::AsNumber(TotalIssues - MaxIssueLines));
-			//"There {NumCats}|plural(one=is,other=are) {NumCats} {NumCats}|plural(one=cat,other=cats)"
+			IconToolTipCache = ToolTipBuilder.ToText();
 		}
-
-		IconToolTipCache = ToolTipBuilder.ToText();
 	}
 	return IconToolTipCache.GetValue();
 }
 
 void SNiagaraStackIssueIcon::UpdateFromEntry()
 {
-	if (StackEntry == nullptr || StackEntry->HasIssuesOrAnyChildHasIssues() == false)
+	if (StackEntry.IsValid() == false || StackEntry->IsFinalized() || StackEntry->HasIssuesOrAnyChildHasIssues() == false)
 	{
 		IconBrush = FEditorStyle::GetBrush("NoBrush");
 		IconToolTipCache = FText();
diff --git a/Engine/Plugins/FX/Niagara/Source/NiagaraEditorWidgets/Private/Stack/SNiagaraStackIssueIcon.h b/Engine/Plugins/FX/Niagara/Source/NiagaraEditorWidgets/Private/Stack/SNiagaraStackIssueIcon.h
--- a/Engine/Plugins/FX/Niagara/Source/NiagaraEditorWidgets/Private/Stack/SNiagaraStackIssueIcon.h
+++ b/Engine/Plugins/FX/Niagara/Source/NiagaraEditorWidgets/Private/Stack/SNiagaraStackIssueIcon.h
@@ -33,5 +33,5 @@ private:
 	mutable TOptional<FText> IconToolTipCache;
 
 	UNiagaraStackViewModel* StackViewModel;
-	UNiagaraStackEntry* StackEntry;
+	TWeakObjectPtr<UNiagaraStackEntry> StackEntry;
 };
\ No newline at end of file
