* Base particle emitter memory allocation on a runtime estimation
* Added option to display current particle count in editor viewport

#tests niagara automation tests checked

diff --git a/Engine/Plugins/FX/Niagara/Source/Niagara/Classes/NiagaraEmitter.h b/Engine/Plugins/FX/Niagara/Source/Niagara/Classes/NiagaraEmitter.h
--- a/Engine/Plugins/FX/Niagara/Source/Niagara/Classes/NiagaraEmitter.h
+++ b/Engine/Plugins/FX/Niagara/Source/Niagara/Classes/NiagaraEmitter.h
@@ -100,6 +100,15 @@ enum class EScriptExecutionMode : uint8
 	SingleParticle UMETA(Hidden)
 };
 
+UENUM()
+enum class EParticleAllocationMode : uint8
+{
+	/** This mode tries to estimate the max particle count at runtime by using previous simulations as reference.*/
+	AutomaticEstimate = 0,
+	/** This mode is useful if the particle count can vary wildly at runtime (e.g. due to user parameters) and a lot of reallocations happen.*/
+	ManualEstimate
+};
+
 USTRUCT()
 struct FNiagaraEmitterScriptProperties
 {
@@ -195,6 +204,13 @@ struct FNiagaraDetailsLevelScaleOverrides
 	float Cine;
 };
 
+struct MemoryRuntimeEstimation
+{
+	TMap<uint64, int32> RuntimeAllocations;
+	bool IsEstimationDirty = false;
+	int32 AllocationEstimate = 0;
+};
+
 /** 
  *	UNiagaraEmitter stores the attributes of an FNiagaraEmitterInstance
  *	that need to be serialized and are used for its initialization 
@@ -247,12 +263,20 @@ public:
 	/** An emitter-based seed for the deterministic random number generator. */
 	UPROPERTY(EditAnywhere, Category = "Emitter", meta = (EditCondition = "bDeterminism"))
 	int32 RandomSeed;
+
+	/**
+	The emitter needs to allocate memory for the particles each tick.
+	To prevent reallocations, the emitter should allocate as much memory as is needed for the max particle count.
+	This setting controls if the allocation size should be automatically determined or manually entered.
+	*/
+	UPROPERTY(EditAnywhere, AdvancedDisplay, Category = "Emitter")
+	EParticleAllocationMode AllocationMode;
 	
 	/** 
 	The emitter will allocate at least this many particles on it's first tick.
 	This can aid performance by avoiding many allocations as an emitter ramps up to it's max size.
 	*/
-	UPROPERTY(EditAnywhere, Category = "Emitter")
+	UPROPERTY(EditAnywhere, AdvancedDisplay, Category = "Emitter", meta = (EditCondition = "AllocationMode == EParticleAllocationMode::ManualEstimate"))
 	int32 PreAllocationCount;
 
 	UPROPERTY()
@@ -423,7 +447,13 @@ public:
 	/* Gets whether or not the supplied event generator id matches an event generator which is shared between the particle spawn and update scrips. */
 	bool IsEventGeneratorShared(FName EventGeneratorId) const;
 
-	TStatId GetStatID(bool bGameThread, bool bConcurrent)const;
+	TStatId GetStatID(bool bGameThread, bool bConcurrent) const;
+
+	/* This is used by the emitter instances to report runtime allocations to reduce reallocation in future simulation runs. */
+	int32 AddRuntimeAllocation(uint64 ReporterHandle, int32 AllocationCount);
+
+	/* Returns the number of max expected particles for memory allocations. */
+	NIAGARA_API int32 GetMaxParticleCountEstimate();
 
 #if WITH_EDITORONLY_DATA
 	NIAGARA_API UNiagaraEmitter* GetParent() const;
@@ -515,6 +545,9 @@ private:
 	mutable TStatId StatID_RT;
 	mutable TStatId StatID_RT_CNC;
 #endif
+
+	MemoryRuntimeEstimation RuntimeEstimation;
+	FCriticalSection EstimationCriticalSection;
 };
 
 
diff --git a/Engine/Plugins/FX/Niagara/Source/Niagara/Classes/NiagaraEmitterInstance.h b/Engine/Plugins/FX/Niagara/Source/Niagara/Classes/NiagaraEmitterInstance.h
--- a/Engine/Plugins/FX/Niagara/Source/Niagara/Classes/NiagaraEmitterInstance.h
+++ b/Engine/Plugins/FX/Niagara/Source/Niagara/Classes/NiagaraEmitterInstance.h
@@ -152,8 +152,6 @@ private:
 
 	/* The age of the emitter*/
 	float Age;
-	/* how many loops this emitter has gone through */
-	uint32 Loops;
 
 	int32 TickCount;
 
@@ -169,6 +167,8 @@ private:
 	/* Emitter bounds */
 	FBox CachedBounds;
 
+	uint32 MaxRuntimeAllocation;
+
 	/** Array of all spawn info driven by our owning emitter script. */
 	TArray<FNiagaraSpawnInfo> SpawnInfos;
 
@@ -237,6 +237,10 @@ private:
 	TArray<FNiagaraEventHandlingInfo> EventHandlingInfo;
 	int32 EventSpawnTotal;
 
+	int32 MaxAllocationCount;
+	int32 MinOverallocation;
+	int32 ReallocationCount;
+
 	/** Optional list of bounds calculators. */
 	TArray<TUniquePtr<FNiagaraBoundsCalculator>, TInlineAllocator<1>> BoundsCalculators;
 
diff --git a/Engine/Plugins/FX/Niagara/Source/Niagara/Classes/NiagaraScriptExecutionContext.h b/Engine/Plugins/FX/Niagara/Source/Niagara/Classes/NiagaraScriptExecutionContext.h
--- a/Engine/Plugins/FX/Niagara/Source/Niagara/Classes/NiagaraScriptExecutionContext.h
+++ b/Engine/Plugins/FX/Niagara/Source/Niagara/Classes/NiagaraScriptExecutionContext.h
@@ -162,6 +162,7 @@ struct FNiagaraGpuSpawnInfo
 {
 	uint32		EventSpawnTotal = 0;
 	uint32		SpawnRateInstances = 0;
+	uint32		MaxParticleCount = 0;
 	FVector4	SpawnInfoStartOffsets[NIAGARA_MAX_GPU_SPAWN_INFOS_V4];
 	FVector4	SpawnInfoParams[NIAGARA_MAX_GPU_SPAWN_INFOS];
 };
diff --git a/Engine/Plugins/FX/Niagara/Source/Niagara/Private/NiagaraDataSet.cpp b/Engine/Plugins/FX/Niagara/Source/Niagara/Private/NiagaraDataSet.cpp
--- a/Engine/Plugins/FX/Niagara/Source/Niagara/Private/NiagaraDataSet.cpp
+++ b/Engine/Plugins/FX/Niagara/Source/Niagara/Private/NiagaraDataSet.cpp
@@ -66,11 +66,11 @@ static FAutoConsoleVariableRef CVarNiagaraRenderBufferShrinkFactor(
 	ECVF_Default
 );
 
-static float GNiagaraGPUDataBufferBufferSlack = 1.5;
+static float GNiagaraGPUDataBufferBufferSlack = 1.1;
 static FAutoConsoleVariableRef CVarNiagaraGPUDataBufferBufferSlack(
 	TEXT("fx.NiagaraGPUDataBufferBufferSlack"),
 	GNiagaraGPUDataBufferBufferSlack,
-	TEXT("Niagara GPU data buffer size threshold for resizing. <= 1 to disable shrinking. (Default=1.5) \n"),
+	TEXT("Niagara GPU data buffer size threshold for resizing. <= 1 to disable shrinking. (Default=1.1) \n"),
 	ECVF_Default
 );
 
diff --git a/Engine/Plugins/FX/Niagara/Source/Niagara/Private/NiagaraEmitter.cpp b/Engine/Plugins/FX/Niagara/Source/Niagara/Private/NiagaraEmitter.cpp
--- a/Engine/Plugins/FX/Niagara/Source/Niagara/Private/NiagaraEmitter.cpp
+++ b/Engine/Plugins/FX/Niagara/Source/Niagara/Private/NiagaraEmitter.cpp
@@ -826,6 +826,8 @@ void UNiagaraEmitter::OnPostCompile()
 		}
 	}
 
+	RuntimeEstimation = MemoryRuntimeEstimation();
+
 	OnEmitterVMCompiled().Broadcast(this);
 
 	InitFastPathAttributeNames();
@@ -1221,6 +1223,58 @@ TStatId UNiagaraEmitter::GetStatID(bool bGameThread, bool bConcurrent)const
 #endif
 	return TStatId();
 }
+
+int32 UNiagaraEmitter::AddRuntimeAllocation(uint64 ReporterHandle, int32 AllocationCount)
+{
+	FScopeLock lock(&EstimationCriticalSection);
+	int32* Estimate = RuntimeEstimation.RuntimeAllocations.Find(ReporterHandle);
+	if (!Estimate || *Estimate < AllocationCount)
+	{
+		RuntimeEstimation.RuntimeAllocations.Add(ReporterHandle, AllocationCount);
+		RuntimeEstimation.IsEstimationDirty = true;
+
+		// Remove a random entry when there are enough logged allocations already
+		if (RuntimeEstimation.RuntimeAllocations.Num() > 10)
+		{
+			TArray<uint64> Keys;
+			RuntimeEstimation.RuntimeAllocations.GetKeys(Keys);
+			RuntimeEstimation.RuntimeAllocations.Remove(Keys[FMath::RandHelper(Keys.Num())]);
+		}
+	}
+	return RuntimeEstimation.RuntimeAllocations.Num();
+}
+
+int32 UNiagaraEmitter::GetMaxParticleCountEstimate()
+{
+	if (AllocationMode == EParticleAllocationMode::ManualEstimate)
+	{
+		return PreAllocationCount;
+	}
+	
+	if (RuntimeEstimation.IsEstimationDirty)
+	{
+		FScopeLock lock(&EstimationCriticalSection);
+		int32 EstimationCount = RuntimeEstimation.RuntimeAllocations.Num();
+		if (EstimationCount > 0)
+		{
+			RuntimeEstimation.RuntimeAllocations.ValueSort(TGreater<int32>());
+			int32 i = 0;
+			for (TPair<uint64, int32> pair : RuntimeEstimation.RuntimeAllocations)
+			{
+				if (i >= (EstimationCount - 1) / 2)
+				{
+					// to prevent overallocation from outliers we take the median instead of the global max
+					RuntimeEstimation.AllocationEstimate = pair.Value;
+					break;
+				}
+				i++;
+			}
+			RuntimeEstimation.IsEstimationDirty = false;
+		}
+	}
+	return RuntimeEstimation.AllocationEstimate;
+}
+
 void UNiagaraEmitter::GenerateStatID()const
 {
 #if STATS
diff --git a/Engine/Plugins/FX/Niagara/Source/Niagara/Private/NiagaraEmitterInstance.cpp b/Engine/Plugins/FX/Niagara/Source/Niagara/Private/NiagaraEmitterInstance.cpp
--- a/Engine/Plugins/FX/Niagara/Source/Niagara/Private/NiagaraEmitterInstance.cpp
+++ b/Engine/Plugins/FX/Niagara/Source/Niagara/Private/NiagaraEmitterInstance.cpp
@@ -42,7 +42,7 @@ static FAutoConsoleVariableRef CVarNiagaraDumpParticleData(
 static int32 GbNiagaraDumpNans = 0;
 static FAutoConsoleVariableRef CVarNiagaraDumpNans(
 	TEXT("fx.Niagara.DumpNans"),
-	GbDumpParticleData,
+	GbNiagaraDumpNans,
 	TEXT("If not 0 any NaNs will be dumped always.\n"),
 	ECVF_Default
 );
@@ -55,6 +55,14 @@ static FAutoConsoleVariableRef CVarNiagaraDumpNansOnce(
 	ECVF_Default
 );
 
+static int32 GbNiagaraShowAllocationWarnings = 1;
+static FAutoConsoleVariableRef CVarNiagaraShowAllocationWarnings(
+	TEXT("fx.Niagara.ShowAllocationWarnings"),
+	GbNiagaraShowAllocationWarnings,
+	TEXT("If not 0 then frequent reallocations and over-allocations of particle memory will cause warnings in the log.\n"),
+	ECVF_Default
+);
+
 /**
 TODO: This is mainly to avoid hard limits in our storage/alloc code etc rather than for perf reasons.
 We should improve our hard limit/safety code and possibly add a max for perf reasons.
@@ -247,6 +255,9 @@ void FNiagaraEmitterInstance::Init(int32 InEmitterIdx, FNiagaraSystemInstanceID
 	}
 
 	CachedEmitterCompiledData = EmitterCompiledData[EmitterIdx];
+	MaxAllocationCount = CachedEmitter->GetMaxParticleCountEstimate();
+	ReallocationCount = 0;
+	MinOverallocation = -1;
 
 #if !UE_BUILD_SHIPPING && !UE_BUILD_TEST 
 	CheckForErrors();
@@ -471,10 +482,16 @@ void FNiagaraEmitterInstance::Init(int32 InEmitterIdx, FNiagaraSystemInstanceID
 void FNiagaraEmitterInstance::ResetSimulation(bool bKillExisting /*= true*/)
 {
 	Age = 0;
-	Loops = 0;
 	TickCount = 0;
 	CachedBounds.Init();
 
+	if (MinOverallocation > 100 && GbNiagaraShowAllocationWarnings)
+	{
+		FString SystemName = this->GetParentSystemInstance()->GetSystem()->GetName();
+		FString FullName = SystemName + "::" + this->GetEmitterHandle().GetName().ToString();
+		UE_LOG(LogNiagara, Warning, TEXT("The emitter %s over-allocated %i particles during its runtime. If this happens frequently, consider setting the emitter's AllocationMode property to 'manual' to improve runtime performance."), *FullName, MinOverallocation);
+	}
+
 	if (IsDisabled())
 	{
 		return;
@@ -1169,6 +1186,27 @@ void FNiagaraEmitterInstance::Tick(float DeltaSeconds)
 			}
 		}
 	}
+	int32 OrigNumParticles = GetNumParticles();
+	int32 AllocationEstimate = CachedEmitter->GetMaxParticleCountEstimate();
+	int32 RequiredSize = OrigNumParticles + SpawnTotal + EventSpawnTotal;
+	int32 AllocationSize = FMath::Max<int32>(AllocationEstimate, RequiredSize);
+	if (AllocationSize > MaxAllocationCount)
+	{
+		ReallocationCount++;
+		MaxAllocationCount = AllocationSize;
+		int32 Estimations = CachedEmitter->AddRuntimeAllocation((uint64)this, MaxAllocationCount);
+		if (GbNiagaraShowAllocationWarnings && Estimations >= 5 && ReallocationCount == 3)
+		{
+			FString SystemName = this->GetParentSystemInstance()->GetSystem()->GetName();
+			FString FullName = SystemName + "::" + this->GetEmitterHandle().GetName().ToString();
+			UE_LOG(LogNiagara, Warning, TEXT("The emitter %s required many memory reallocation due to changing particle counts. Consider setting the emitter's AllocationMode property to 'manual' to improve runtime performance."), *FullName);
+		}
+	}
+	int32 Overallocation = AllocationSize - RequiredSize;
+	if (Overallocation >= 0 && (MinOverallocation < 0 || Overallocation < MinOverallocation))
+	{
+		MinOverallocation = Overallocation;
+	}
 
 	/* GPU simulation -  we just create an FNiagaraComputeExecutionContext, queue it, and let the batcher take care of the rest
 	 */
@@ -1202,6 +1240,7 @@ void FNiagaraEmitterInstance::Tick(float DeltaSeconds)
 			FNiagaraGpuSpawnInfo& GpuSpawnInfo = GPUExecContext->GpuSpawnInfo_GT;
 			GpuSpawnInfo.EventSpawnTotal = EventSpawnTotal;
 			GpuSpawnInfo.SpawnRateInstances = 0;
+			GpuSpawnInfo.MaxParticleCount = AllocationSize;
 
 			int NumSpawnInfos = 0;
 			if (ExecutionState == ENiagaraExecutionState::Active)
@@ -1284,9 +1323,6 @@ void FNiagaraEmitterInstance::Tick(float DeltaSeconds)
 		return;
 	}
 
-	int32 OrigNumParticles = Data.GetCurrentDataChecked().GetNumInstances();
-
-	int32 AllocationSize = FMath::Max<int32>(CachedEmitter->PreAllocationCount, OrigNumParticles + SpawnTotal + EventSpawnTotal);
 	//Ensure we don't blow our current hard limits on cpu particle count.
 	//TODO: These current limits can be improved relatively easily. Though perf in at these counts will obviously be an issue anyway.
 	if (CachedEmitter->SimTarget == ENiagaraSimTarget::CPUSim && AllocationSize > GMaxNiagaraCPUParticlesPerEmitter)
@@ -1298,6 +1334,7 @@ void FNiagaraEmitterInstance::Tick(float DeltaSeconds)
 		EventSpawnTotal = 0;
 	}
 
+
 	Data.BeginSimulate();
 	Data.Allocate(AllocationSize);
 
@@ -1714,7 +1751,6 @@ void FNiagaraEmitterInstance::Tick(float DeltaSeconds)
 	INC_DWORD_STAT_BY(STAT_NiagaraNumParticles, Data.GetCurrentDataChecked().GetNumInstances());
 }
 
-
 /** Calculate total number of spawned particles from events; these all come from event handler script with the SpawnedParticles execution mode
  *  We get the counts ahead of event processing time so we only have to allocate new particles once
  *  TODO: augment for multiple spawning event scripts
diff --git a/Engine/Plugins/FX/Niagara/Source/Niagara/Private/NiagaraEmitterInstanceBatcher.cpp b/Engine/Plugins/FX/Niagara/Source/Niagara/Private/NiagaraEmitterInstanceBatcher.cpp
--- a/Engine/Plugins/FX/Niagara/Source/Niagara/Private/NiagaraEmitterInstanceBatcher.cpp
+++ b/Engine/Plugins/FX/Niagara/Source/Niagara/Private/NiagaraEmitterInstanceBatcher.cpp
@@ -366,8 +366,9 @@ void NiagaraEmitterInstanceBatcher::ResizeBuffersAndGatherResources(FOverlappabl
 			//We must assume all particles survive when allocating here. 
 			//If this is not true, the read back in ResolveDatasetWrites will shrink the buffers.
 			const uint32 RequiredInstances = FMath::Max(PrevNumInstances, NewNumInstances);
+			const uint32 AllocatedInstances = FMath::Max(RequiredInstances, Instance.SpawnInfo.MaxParticleCount);
 
-			DestinationData.AllocateGPU(RequiredInstances + 1, GPUInstanceCounterManager, RHICmdList);
+			DestinationData.AllocateGPU(AllocatedInstances + 1, GPUInstanceCounterManager, RHICmdList);
 			DestinationData.SetNumInstances(RequiredInstances);
 
 			if ( Shader->FloatInputBufferParam.IsBound() )
diff --git a/Engine/Plugins/FX/Niagara/Source/NiagaraEditor/Private/NiagaraEditorCommands.cpp b/Engine/Plugins/FX/Niagara/Source/NiagaraEditor/Private/NiagaraEditorCommands.cpp
--- a/Engine/Plugins/FX/Niagara/Source/NiagaraEditor/Private/NiagaraEditorCommands.cpp
+++ b/Engine/Plugins/FX/Niagara/Source/NiagaraEditor/Private/NiagaraEditorCommands.cpp
@@ -14,6 +14,7 @@ void FNiagaraEditorCommands::RegisterCommands()
 	UI_COMMAND(ResetSimulation, "Reset", "Resets the current simulation", EUserInterfaceActionType::Button, FInputChord(EKeys::SpaceBar));
 	UI_COMMAND(TogglePreviewGrid, "Grid", "Toggles the preview pane's grid.", EUserInterfaceActionType::ToggleButton, FInputChord() );
 	UI_COMMAND(ToggleInstructionCounts, "InstructionCounts", "Display Instruction Counts", EUserInterfaceActionType::ToggleButton, FInputChord());
+	UI_COMMAND(ToggleParticleCounts, "ParticleCounts", "Display Particle Counts", EUserInterfaceActionType::ToggleButton, FInputChord());
 	UI_COMMAND(TogglePreviewBackground, "Background", "Toggles the preview pane's background.", EUserInterfaceActionType::ToggleButton, FInputChord());
 	UI_COMMAND(ToggleUnlockToChanges, "Lock/Unlock To Changes", "Toggles whether or not changes in the source asset get pulled into this asset automatically.", EUserInterfaceActionType::Button, FInputChord());
 	UI_COMMAND(ToggleBounds, "Bounds", "Display Bounds", EUserInterfaceActionType::ToggleButton, FInputChord());
diff --git a/Engine/Plugins/FX/Niagara/Source/NiagaraEditor/Private/Widgets/SNiagaraSystemViewport.cpp b/Engine/Plugins/FX/Niagara/Source/NiagaraEditor/Private/Widgets/SNiagaraSystemViewport.cpp
--- a/Engine/Plugins/FX/Niagara/Source/NiagaraEditor/Private/Widgets/SNiagaraSystemViewport.cpp
+++ b/Engine/Plugins/FX/Niagara/Source/NiagaraEditor/Private/Widgets/SNiagaraSystemViewport.cpp
@@ -47,6 +47,7 @@ public:
 	virtual void SetIsSimulateInEditorViewport(bool bInIsSimulateInEditorViewport)override;
 
 	void DrawInstructionCounts(UNiagaraSystem* ParticleSystem, FCanvas* Canvas, float& CurrentX, float& CurrentY, UFont* Font, const float FontHeight);
+	void DrawParticleCounts(UNiagaraComponent* Component, FCanvas* Canvas, float& CurrentX, float& CurrentY, UFont* Font, const float FontHeight);
 
 	TWeakPtr<SNiagaraSystemViewport> NiagaraViewportPtr;
 	bool bCaptureScreenShot;
@@ -101,6 +102,7 @@ void FNiagaraSystemViewportClient::Draw(FViewport* InViewport, FCanvas* Canvas)
 {
 	TSharedPtr<SNiagaraSystemViewport> NiagaraViewport = NiagaraViewportPtr.Pin();
 	UNiagaraSystem* ParticleSystem = NiagaraViewport.IsValid() ? NiagaraViewport->GetPreviewComponent()->GetAsset() : nullptr;
+	UNiagaraComponent* Component = NiagaraViewport.IsValid() ? NiagaraViewport->GetPreviewComponent() : nullptr;
 
 	if (NiagaraViewport.IsValid() && NiagaraViewport->GetDrawElement(SNiagaraSystemViewport::EDrawElements::Bounds))
 	{
@@ -122,9 +124,14 @@ void FNiagaraSystemViewportClient::Draw(FViewport* InViewport, FCanvas* Canvas)
 		UFont* Font = GEngine->GetSmallFont();
 		const float FontHeight = Font->GetMaxCharHeight() * 1.1f;
 
-		if ( NiagaraViewport->GetDrawElement(SNiagaraSystemViewport::EDrawElements::InstructionCounts) )
+		if (NiagaraViewport->GetDrawElement(SNiagaraSystemViewport::EDrawElements::InstructionCounts) && ParticleSystem)
 		{
 			DrawInstructionCounts(ParticleSystem, Canvas, CurrentX, CurrentY, Font, FontHeight);
+			CurrentY += FontHeight;
+		}
+		if (NiagaraViewport->GetDrawElement(SNiagaraSystemViewport::EDrawElements::ParticleCounts) && Component)
+		{
+			DrawParticleCounts(Component, Canvas, CurrentX, CurrentY, Font, FontHeight);
 		}
 	}
 
@@ -203,6 +210,31 @@ void FNiagaraSystemViewportClient::DrawInstructionCounts(UNiagaraSystem* Particl
 	}
 }
 
+void FNiagaraSystemViewportClient::DrawParticleCounts(UNiagaraComponent* Component, FCanvas* Canvas, float& CurrentX, float& CurrentY, UFont* Font, const float FontHeight)
+{
+	FCanvasTextItem TextItem(FVector2D(CurrentX, CurrentY), FText::FromString(TEXT("Particle Counts")), Font, FLinearColor::White);
+	TextItem.EnableShadow(FLinearColor::Black);
+	TextItem.Draw(Canvas);
+	CurrentY += FontHeight;
+
+	FNiagaraSystemInstance* SystemInstance = Component->GetSystemInstance();
+	if (!SystemInstance)
+	{
+		return;
+	}
+
+	for (TSharedRef<FNiagaraEmitterInstance, ESPMode::ThreadSafe> EmitterInstance : SystemInstance->GetEmitters())
+	{
+		FName EmitterName = EmitterInstance->GetEmitterHandle().GetName();
+		int32 CurrentCount = EmitterInstance->GetNumParticles();
+		int32 MaxCount = EmitterInstance->GetEmitterHandle().GetInstance()->GetMaxParticleCountEstimate();
+		TextItem.Text = FText::FromString(FString::Printf(TEXT("%i Current, %i Max (est.) - [%s]"), CurrentCount, MaxCount, *EmitterName.ToString()));
+		TextItem.Position = FVector2D(CurrentX, CurrentY);
+		TextItem.Draw(Canvas);
+		CurrentY += FontHeight;
+	}
+}
+
 bool FNiagaraSystemViewportClient::ShouldOrbitCamera() const
 {
 	return bUsingOrbitCamera;
@@ -387,6 +419,13 @@ void SNiagaraSystemViewport::BindCommands()
 		FIsActionChecked::CreateLambda([Viewport=this]() -> bool { return Viewport->GetDrawElement(EDrawElements::InstructionCounts); })
 	);
 
+	CommandList->MapAction(
+		Commands.ToggleParticleCounts,
+		FExecuteAction::CreateLambda([Viewport = this]() { Viewport->ToggleDrawElement(EDrawElements::ParticleCounts); Viewport->RefreshViewport(); }),
+		FCanExecuteAction(),
+		FIsActionChecked::CreateLambda([Viewport = this]() -> bool { return Viewport->GetDrawElement(EDrawElements::ParticleCounts); })
+	);
+
 	CommandList->MapAction(
 		Commands.TogglePreviewBackground,
 		FExecuteAction::CreateSP( this, &SNiagaraSystemViewport::TogglePreviewBackground ),
diff --git a/Engine/Plugins/FX/Niagara/Source/NiagaraEditor/Private/Widgets/SNiagaraSystemViewport.h b/Engine/Plugins/FX/Niagara/Source/NiagaraEditor/Private/Widgets/SNiagaraSystemViewport.h
--- a/Engine/Plugins/FX/Niagara/Source/NiagaraEditor/Private/Widgets/SNiagaraSystemViewport.h
+++ b/Engine/Plugins/FX/Niagara/Source/NiagaraEditor/Private/Widgets/SNiagaraSystemViewport.h
@@ -75,6 +75,7 @@ public:
 	{
 		Bounds = 0x020,
 		InstructionCounts = 0x040,
+		ParticleCounts = 0x080,
 	};
 
 	bool GetDrawElement(EDrawElements Element) const;
diff --git a/Engine/Plugins/FX/Niagara/Source/NiagaraEditor/Private/Widgets/SNiagaraSystemViewportToolBar.cpp b/Engine/Plugins/FX/Niagara/Source/NiagaraEditor/Private/Widgets/SNiagaraSystemViewportToolBar.cpp
--- a/Engine/Plugins/FX/Niagara/Source/NiagaraEditor/Private/Widgets/SNiagaraSystemViewportToolBar.cpp
+++ b/Engine/Plugins/FX/Niagara/Source/NiagaraEditor/Private/Widgets/SNiagaraSystemViewportToolBar.cpp
@@ -35,6 +35,7 @@ TSharedRef<SWidget> SNiagaraSystemViewportToolBar::GenerateShowMenu() const
 
 		ShowMenuBuilder.AddMenuEntry(Commands.TogglePreviewGrid);
 		ShowMenuBuilder.AddMenuEntry(Commands.ToggleInstructionCounts);
+		ShowMenuBuilder.AddMenuEntry(Commands.ToggleParticleCounts);
 		//ShowMenuBuilder.AddMenuEntry(Commands.TogglePreviewBackground);
 	}
 
diff --git a/Engine/Plugins/FX/Niagara/Source/NiagaraEditor/Public/NiagaraEditorCommands.h b/Engine/Plugins/FX/Niagara/Source/NiagaraEditor/Public/NiagaraEditorCommands.h
--- a/Engine/Plugins/FX/Niagara/Source/NiagaraEditor/Public/NiagaraEditorCommands.h
+++ b/Engine/Plugins/FX/Niagara/Source/NiagaraEditor/Public/NiagaraEditorCommands.h
@@ -31,6 +31,7 @@ public:
 	/** Toggles the preview pane's grid */
 	TSharedPtr<FUICommandInfo> TogglePreviewGrid;
 	TSharedPtr<FUICommandInfo> ToggleInstructionCounts;
+	TSharedPtr<FUICommandInfo> ToggleParticleCounts;
 
 	/** Toggles the preview pane's background */
 	TSharedPtr< FUICommandInfo > TogglePreviewBackground;
