Niagara - Guard against null script variables in the stack.

diff --git a/Engine/Plugins/FX/Niagara/Source/NiagaraEditor/Private/ViewModels/Stack/NiagaraStackModuleItem.cpp b/Engine/Plugins/FX/Niagara/Source/NiagaraEditor/Private/ViewModels/Stack/NiagaraStackModuleItem.cpp
--- a/Engine/Plugins/FX/Niagara/Source/NiagaraEditor/Private/ViewModels/Stack/NiagaraStackModuleItem.cpp
+++ b/Engine/Plugins/FX/Niagara/Source/NiagaraEditor/Private/ViewModels/Stack/NiagaraStackModuleItem.cpp
@@ -187,16 +187,21 @@ void UNiagaraStackModuleItem::RefreshChildrenInternal(const TArray<UNiagaraStack
 			auto Iter = MetaDataMap.CreateConstIterator();
 			while (Iter)
 			{
-				auto PropertyIter = Iter.Value()->Metadata.PropertyMetaData.CreateConstIterator();
-				while (PropertyIter)
+				// TODO: This should never be null, but somehow it is in some assets so guard this to prevent crashes
+				// until we have better repro steps.
+				if (Iter.Value() != nullptr)
 				{
-					if (PropertyIter.Key() == (TEXT("DisplayNameArg0")))
+					auto PropertyIter = Iter.Value()->Metadata.PropertyMetaData.CreateConstIterator();
+					while (PropertyIter)
 					{
-						bCanRefresh = true;
+						if (PropertyIter.Key() == (TEXT("DisplayNameArg0")))
+						{
+							bCanRefresh = true;
+						}
+						++PropertyIter;
 					}
-					++PropertyIter;
+					++Iter;
 				}
-				++Iter;
 			}
 		}
 
