Merging 9906999 from main - Niagara - Fix an issue where data interfaces weren't handled correctly when dragging to move or copy modules between emitters and systems.

diff --git a/Engine/Plugins/FX/Niagara/Source/NiagaraEditor/Private/ViewModels/Stack/NiagaraStackGraphUtilities.cpp b/Engine/Plugins/FX/Niagara/Source/NiagaraEditor/Private/ViewModels/Stack/NiagaraStackGraphUtilities.cpp
--- a/Engine/Plugins/FX/Niagara/Source/NiagaraEditor/Private/ViewModels/Stack/NiagaraStackGraphUtilities.cpp
+++ b/Engine/Plugins/FX/Niagara/Source/NiagaraEditor/Private/ViewModels/Stack/NiagaraStackGraphUtilities.cpp
@@ -364,7 +364,6 @@ void GetGroupNodesRecursive(const TArray<UNiagaraNode*>& CurrentStartNodes, UNia
 			if (CurrentStartNode != EndNode)
 			{
 				TArray<UNiagaraNode*> LinkedNodes;
-				TArray<UEdGraphPin*> InputPins;
 				TArray<UEdGraphPin*> OutputPins;
 				CurrentStartNode->GetOutputPins(OutputPins);
 				for (UEdGraphPin* OutputPin : OutputPins)
