Niagara - Fix a crash when opening the system editor on mac and linux.

diff --git a/Engine/Plugins/FX/Niagara/Source/NiagaraEditor/Private/NiagaraEditorModule.cpp b/Engine/Plugins/FX/Niagara/Source/NiagaraEditor/Private/NiagaraEditorModule.cpp
--- a/Engine/Plugins/FX/Niagara/Source/NiagaraEditor/Private/NiagaraEditorModule.cpp
+++ b/Engine/Plugins/FX/Niagara/Source/NiagaraEditor/Private/NiagaraEditorModule.cpp
@@ -70,6 +70,7 @@
 #include "NiagaraTypes.h"
 #include "NiagaraSystemFactoryNew.h"
 #include "NiagaraSystemEditorData.h"
+#include "NiagaraEditorCommands.h"
 
 #include "MovieScene/Parameters/MovieSceneNiagaraBoolParameterTrack.h"
 #include "MovieScene/Parameters/MovieSceneNiagaraFloatParameterTrack.h"
@@ -1043,6 +1044,11 @@ TSharedPtr<FNiagaraSystemViewModel> FNiagaraEditorModule::GetExistingViewModelFo
 	return FNiagaraSystemViewModel::GetExistingViewModelForObject(InSystem);
 }
 
+const FNiagaraEditorCommands& FNiagaraEditorModule::GetCommands() const
+{
+	return FNiagaraEditorCommands::Get();
+}
+
 void FNiagaraEditorModule::RegisterAssetTypeAction(IAssetTools& AssetTools, TSharedRef<IAssetTypeActions> Action)
 {
 	AssetTools.RegisterAssetTypeActions(Action);
diff --git a/Engine/Plugins/FX/Niagara/Source/NiagaraEditor/Public/NiagaraEditorModule.h b/Engine/Plugins/FX/Niagara/Source/NiagaraEditor/Public/NiagaraEditorModule.h
--- a/Engine/Plugins/FX/Niagara/Source/NiagaraEditor/Public/NiagaraEditorModule.h
+++ b/Engine/Plugins/FX/Niagara/Source/NiagaraEditor/Public/NiagaraEditorModule.h
@@ -24,6 +24,7 @@ class FNiagaraCompileRequestDataBase;
 class UMovieSceneNiagaraParameterTrack;
 struct IConsoleCommand;
 class INiagaraEditorOnlyDataUtilities;
+class FNiagaraEditorCommands;
 
 DECLARE_STATS_GROUP(TEXT("Niagara Editor"), STATGROUP_NiagaraEditor, STATCAT_Advanced);
 
@@ -103,6 +104,8 @@ public:
 
 	NIAGARAEDITOR_API TSharedPtr<FNiagaraSystemViewModel> GetExistingViewModelForSystem(UNiagaraSystem* InSystem);
 
+	NIAGARAEDITOR_API const FNiagaraEditorCommands& GetCommands() const;
+
 private:
 	void RegisterAssetTypeAction(IAssetTools& AssetTools, TSharedRef<IAssetTypeActions> Action);
 	void OnNiagaraSettingsChangedEvent(const FString& PropertyName, const UNiagaraSettings* Settings);
diff --git a/Engine/Plugins/FX/Niagara/Source/NiagaraEditorWidgets/Private/SNiagaraOverviewGraph.cpp b/Engine/Plugins/FX/Niagara/Source/NiagaraEditorWidgets/Private/SNiagaraOverviewGraph.cpp
--- a/Engine/Plugins/FX/Niagara/Source/NiagaraEditorWidgets/Private/SNiagaraOverviewGraph.cpp
+++ b/Engine/Plugins/FX/Niagara/Source/NiagaraEditorWidgets/Private/SNiagaraOverviewGraph.cpp
@@ -16,6 +16,7 @@
 #include "ViewModels/NiagaraOverviewGraphViewModel.h"
 #include "EdGraphSchema_Niagara.h"
 #include "NiagaraEditorCommands.h"
+#include "NiagaraEditorModule.h"
 #include "GraphEditorActions.h"
 #include "Framework/MultiBox/MultiBoxBuilder.h"
 #include "ScopedTransaction.h"
@@ -70,10 +71,10 @@ void SNiagaraOverviewGraph::Construct(const FArguments& InArgs, TSharedRef<FNiag
 		FGraphEditorCommands::Get().CreateComment,
 		FExecuteAction::CreateSP(this, &SNiagaraOverviewGraph::OnCreateComment));
 	Commands->MapAction(
-		FNiagaraEditorCommands::Get().ZoomToFit,
+		FNiagaraEditorModule::Get().GetCommands().ZoomToFit,
 		FExecuteAction::CreateSP(this, &SNiagaraOverviewGraph::ZoomToFit));
 	Commands->MapAction(
-		FNiagaraEditorCommands::Get().ZoomToFitAll,
+		FNiagaraEditorModule::Get().GetCommands().ZoomToFitAll,
 		FExecuteAction::CreateSP(this, &SNiagaraOverviewGraph::ZoomToFitAll));
 	
 	GraphEditor = SNew(SGraphEditor)
