Niagara - Fix another issue where data interfaces weren't being handled correctly for module drag and drop.

diff --git a/Engine/Plugins/FX/Niagara/Source/NiagaraEditor/Private/NiagaraEditorUtilities.cpp b/Engine/Plugins/FX/Niagara/Source/NiagaraEditor/Private/NiagaraEditorUtilities.cpp
--- a/Engine/Plugins/FX/Niagara/Source/NiagaraEditor/Private/NiagaraEditorUtilities.cpp
+++ b/Engine/Plugins/FX/Niagara/Source/NiagaraEditor/Private/NiagaraEditorUtilities.cpp
@@ -269,8 +269,9 @@ void FNiagaraEditorUtilities::FixUpPastedNodes(UEdGraph* Graph, TSet<UEdGraphNod
 				FName FunctionCallName = *PastedFunctionCallNode->GetFunctionName();
 				FName UniqueFunctionCallName = FNiagaraUtilities::GetUniqueName(FunctionCallName, ExistingNames);
 				PastedFunctionCallNode->SuggestName(UniqueFunctionCallName.ToString());
-				ExistingNames.Add(UniqueFunctionCallName);
-				OldFunctionToNewFunctionNameMap.Add(FunctionCallName, UniqueFunctionCallName);
+				FName ActualPastedFunctionCallName = *PastedFunctionCallNode->GetFunctionName();
+				ExistingNames.Add(ActualPastedFunctionCallName);
+				OldFunctionToNewFunctionNameMap.Add(FunctionCallName, ActualPastedFunctionCallName);
 			}
 		}
 	}
