- Allow fast ops HLSL to be used once the node is placed
- Before you needed to have "fx.AllowFastPathFunctionLibrary 1" enabled in order for the HLSL to be picked up

#rnx

diff --git a/Engine/Plugins/FX/Niagara/Source/Niagara/Private/NiagaraFunctionLibrary.cpp b/Engine/Plugins/FX/Niagara/Source/Niagara/Private/NiagaraFunctionLibrary.cpp
--- a/Engine/Plugins/FX/Niagara/Source/Niagara/Private/NiagaraFunctionLibrary.cpp
+++ b/Engine/Plugins/FX/Niagara/Source/Niagara/Private/NiagaraFunctionLibrary.cpp
@@ -367,19 +367,19 @@ UNiagaraParameterCollectionInstance* UNiagaraFunctionLibrary::GetNiagaraParamete
 
 const TArray<FNiagaraFunctionSignature>& UNiagaraFunctionLibrary::GetVectorVMFastPathOps()
 {
-	if (VectorVMOps.Num() == 0)
+	if (GAllowFastPathFunctionLibrary == 0)
 	{
-		InitVectorVMFastPathOps();
+		static TArray<FNiagaraFunctionSignature> Empty;
+		return Empty;
 	}
+
+	InitVectorVMFastPathOps();
 	return VectorVMOps;
 }
 
 bool UNiagaraFunctionLibrary::DefineFunctionHLSL(const FNiagaraFunctionSignature& FunctionSignature, FString& HlslOutput)
 {
-	if (VectorVMOps.Num() == 0)
-	{
-		InitVectorVMFastPathOps();
-	}
+	InitVectorVMFastPathOps();
 
 	const int32 i = VectorVMOps.IndexOfByKey(FunctionSignature);
 	if ( i == INDEX_NONE )
@@ -744,7 +744,7 @@ bool UNiagaraFunctionLibrary::GetVectorVMFastPathExternalFunction(const FVMExter
 
 void UNiagaraFunctionLibrary::InitVectorVMFastPathOps()
 {
-	if (GAllowFastPathFunctionLibrary == 0)
+	if (VectorVMOps.Num() > 0)
 		return;
 
 	VectorVMOps.Emplace(FVectorKernelFastDot4::GetFunctionSignature());
