query for the entire grid
remove specifying resolution on water bp renderer
updated old tests to work

diff --git a/Engine/Plugins/FX/Niagara/Source/Niagara/Classes/NiagaraDataInterfaceGrid2DCollection.h b/Engine/Plugins/FX/Niagara/Source/Niagara/Classes/NiagaraDataInterfaceGrid2DCollection.h
--- a/Engine/Plugins/FX/Niagara/Source/Niagara/Classes/NiagaraDataInterfaceGrid2DCollection.h
+++ b/Engine/Plugins/FX/Niagara/Source/Niagara/Classes/NiagaraDataInterfaceGrid2DCollection.h
@@ -136,6 +136,15 @@ public:
 	// #todo(dmp): this will eventually go away when we formalize how data makes it out of Niagara
 	UFUNCTION(BlueprintCallable, Category = Niagara)
 	virtual void FillTexture2D(const UNiagaraComponent *Component, UTextureRenderTarget2D *dest, int AttributeIndex);
+	
+	UFUNCTION(BlueprintCallable, Category = Niagara)
+	virtual void FillRawTexture2D(const UNiagaraComponent *Component, UTextureRenderTarget2D *Dest, int &TilesX, int &TilesY);
+	
+	UFUNCTION(BlueprintCallable, Category = Niagara)
+	virtual void GetRawTextureSize(const UNiagaraComponent *Component, int &SizeX, int &SizeY);
+
+	UFUNCTION(BlueprintCallable, Category = Niagara)
+	virtual void GetTextureSize(const UNiagaraComponent *Component, int &SizeX, int &SizeY);
 
 	void GetCellSize(FVectorVMContext& Context);
 
diff --git a/Engine/Plugins/FX/Niagara/Source/Niagara/Private/NiagaraDataInterfaceGrid2DCollection.cpp b/Engine/Plugins/FX/Niagara/Source/Niagara/Private/NiagaraDataInterfaceGrid2DCollection.cpp
--- a/Engine/Plugins/FX/Niagara/Source/Niagara/Private/NiagaraDataInterfaceGrid2DCollection.cpp
+++ b/Engine/Plugins/FX/Niagara/Source/Niagara/Private/NiagaraDataInterfaceGrid2DCollection.cpp
@@ -470,6 +470,119 @@ void UNiagaraDataInterfaceGrid2DCollection::FillTexture2D(const UNiagaraComponen
 	});
 }
 
+UFUNCTION(BlueprintCallable, Category = Niagara)
+void UNiagaraDataInterfaceGrid2DCollection::FillRawTexture2D(const UNiagaraComponent *Component, UTextureRenderTarget2D *Dest, int &TilesX, int &TilesY)
+{
+	FNiagaraDataInterfaceProxyGrid2DCollection* TProxy = GetProxyAs<FNiagaraDataInterfaceProxyGrid2DCollection>();
+
+	if (!Component)
+	{
+		TilesX = -1;
+		TilesY = -1;
+		return;
+	}
+
+	FNiagaraSystemInstance *SystemInstance = Component->GetSystemInstance();
+	if (!SystemInstance)
+	{
+		TilesX = -1;
+		TilesY = -1;
+		return;
+	}
+	FNiagaraSystemInstanceID InstanceID = SystemInstance->GetId();
+
+	Grid2DCollectionRWInstanceData* Grid2DInstanceData = TProxy->SystemInstancesToProxyData.Find(InstanceID);
+	if (!Grid2DInstanceData)
+	{
+		TilesX = -1;
+		TilesY = -1;
+		return;
+	}
+
+	TilesX = Grid2DInstanceData->NumTilesX;
+	TilesY = Grid2DInstanceData->NumTilesY;
+
+	ENQUEUE_RENDER_COMMAND(FUpdateDIColorCurve)(
+		[Dest, TProxy, InstanceID](FRHICommandListImmediate& RHICmdList)
+	{
+		Grid2DCollectionRWInstanceData* Grid2DInstanceData = TProxy->SystemInstancesToProxyData.Find(InstanceID);
+
+		if (Dest && Dest->Resource && Grid2DInstanceData && Grid2DInstanceData->CurrentData)
+		{
+			
+			FRHICopyTextureInfo CopyInfo;
+			RHICmdList.CopyTexture(Grid2DInstanceData->CurrentData->GridBuffer.Buffer, Dest->Resource->TextureRHI, CopyInfo);
+		}
+	});
+}
+
+UFUNCTION(BlueprintCallable, Category = Niagara)
+void UNiagaraDataInterfaceGrid2DCollection::GetRawTextureSize(const UNiagaraComponent *Component, int &SizeX, int &SizeY)
+{
+	FNiagaraDataInterfaceProxyGrid2DCollection* TProxy = GetProxyAs<FNiagaraDataInterfaceProxyGrid2DCollection>();
+
+	if (!Component)
+	{
+		SizeX = -1;
+		SizeY = -1;
+		return;
+	}
+
+	FNiagaraSystemInstance *SystemInstance = Component->GetSystemInstance();
+	if (!SystemInstance)
+	{
+		SizeX = -1;
+		SizeY = -1;
+		return;
+	}
+	FNiagaraSystemInstanceID InstanceID = SystemInstance->GetId();
+
+	Grid2DCollectionRWInstanceData* Grid2DInstanceData = TProxy->SystemInstancesToProxyData.Find(InstanceID);
+	if (!Grid2DInstanceData)
+	{
+		SizeX = -1;
+		SizeY = -1;
+		return;
+	}
+	
+	SizeX = Grid2DInstanceData->Buffers[0]->GridBuffer.Buffer->GetSizeX();
+	SizeY = Grid2DInstanceData->Buffers[0]->GridBuffer.Buffer->GetSizeY();
+}
+
+UFUNCTION(BlueprintCallable, Category = Niagara)
+void UNiagaraDataInterfaceGrid2DCollection::GetTextureSize(const UNiagaraComponent *Component, int &SizeX, int &SizeY)
+{
+	FNiagaraDataInterfaceProxyGrid2DCollection* TProxy = GetProxyAs<FNiagaraDataInterfaceProxyGrid2DCollection>();
+
+	if (!Component)
+	{
+		SizeX = -1;
+		SizeY = -1;
+		return;
+	}
+
+	FNiagaraSystemInstance *SystemInstance = Component->GetSystemInstance();
+	if (!SystemInstance)
+	{
+		SizeX = -1;
+		SizeY = -1;
+		return;
+	}
+	FNiagaraSystemInstanceID InstanceID = SystemInstance->GetId();
+
+	Grid2DCollectionRWInstanceData* Grid2DInstanceData = TProxy->SystemInstancesToProxyData.Find(InstanceID);
+	if (!Grid2DInstanceData)
+	{
+		SizeX = -1;
+		SizeY = -1;
+		return;
+	}
+
+	SizeX = Grid2DInstanceData->NumCellsX;
+	SizeY = Grid2DInstanceData->NumCellsY;
+}
+
+
 void UNiagaraDataInterfaceGrid2DCollection::GetCellSize(FVectorVMContext& Context)
 {
 	VectorVM::FUserPtrHandler<Grid2DCollectionRWInstanceData> InstData(Context);
