Deferred Init Renderer Properties to avoid failed init when Niagara Module has not loaded but the Properties object has.
Move Niagara plugin load up to PreDefault to avoid scenario where packaged game module  loads Niagara UObject in Default phase before we are ready to hande initialization.

diff --git a/Engine/Plugins/FX/Niagara/Source/Niagara/Private/NiagaraLightRendererProperties.cpp b/Engine/Plugins/FX/Niagara/Source/Niagara/Private/NiagaraLightRendererProperties.cpp
--- a/Engine/Plugins/FX/Niagara/Source/Niagara/Private/NiagaraLightRendererProperties.cpp
+++ b/Engine/Plugins/FX/Niagara/Source/Niagara/Private/NiagaraLightRendererProperties.cpp
@@ -5,6 +5,8 @@
 #include "NiagaraConstants.h"
 #include "NiagaraRendererLights.h"
 
+TArray<TWeakObjectPtr<UNiagaraLightRendererProperties>> UNiagaraLightRendererProperties::LightRendererPropertiesToDeferredInit;
+
 UNiagaraLightRendererProperties::UNiagaraLightRendererProperties()
 	: bUseInverseSquaredFalloff(1), bAffectsTranslucency(0), bOverrideRenderingEnabled(0), RadiusScale(1.0f), ColorAdd(FVector(0.0f, 0.0f, 0.0f))
 {
@@ -13,14 +15,24 @@ UNiagaraLightRendererProperties::UNiagaraLightRendererProperties()
 void UNiagaraLightRendererProperties::PostInitProperties()
 {
 	Super::PostInitProperties();
-	if (HasAnyFlags(RF_ClassDefaultObject) == false && PositionBinding.BoundVariable.GetName() == NAME_None)
+
+	if (HasAnyFlags(RF_ClassDefaultObject) == false)
 	{
-		PositionBinding = FNiagaraConstants::GetAttributeDefaultBinding(SYS_PARAM_PARTICLES_POSITION);
-		ColorBinding = FNiagaraConstants::GetAttributeDefaultBinding(SYS_PARAM_PARTICLES_COLOR);
-		RadiusBinding = FNiagaraConstants::GetAttributeDefaultBinding(SYS_PARAM_PARTICLES_LIGHT_RADIUS);
-		LightExponentBinding = FNiagaraConstants::GetAttributeDefaultBinding(SYS_PARAM_PARTICLES_LIGHT_EXPONENT);
-		LightRenderingEnabledBinding = FNiagaraConstants::GetAttributeDefaultBinding(SYS_PARAM_PARTICLES_LIGHT_ENABLED);
-		VolumetricScatteringBinding = FNiagaraConstants::GetAttributeDefaultBinding(SYS_PARAM_PARTICLES_LIGHT_VOLUMETRIC_SCATTERING);
+		// We can end up hitting PostInitProperties before the Niagara Module has initialized bindings this needs, mark this object for deferred init and early out.
+		if (FModuleManager::Get().IsModuleLoaded("Niagara") == false)
+		{
+			LightRendererPropertiesToDeferredInit.Add(this);
+			return;
+		}
+		else if (PositionBinding.BoundVariable.GetName() == NAME_None)
+		{
+			PositionBinding = FNiagaraConstants::GetAttributeDefaultBinding(SYS_PARAM_PARTICLES_POSITION);
+			ColorBinding = FNiagaraConstants::GetAttributeDefaultBinding(SYS_PARAM_PARTICLES_COLOR);
+			RadiusBinding = FNiagaraConstants::GetAttributeDefaultBinding(SYS_PARAM_PARTICLES_LIGHT_RADIUS);
+			LightExponentBinding = FNiagaraConstants::GetAttributeDefaultBinding(SYS_PARAM_PARTICLES_LIGHT_EXPONENT);
+			LightRenderingEnabledBinding = FNiagaraConstants::GetAttributeDefaultBinding(SYS_PARAM_PARTICLES_LIGHT_ENABLED);
+			VolumetricScatteringBinding = FNiagaraConstants::GetAttributeDefaultBinding(SYS_PARAM_PARTICLES_LIGHT_VOLUMETRIC_SCATTERING);
+		}
 	}
 }
 
@@ -28,13 +40,28 @@ void UNiagaraLightRendererProperties::PostInitProperties()
 void UNiagaraLightRendererProperties::InitCDOPropertiesAfterModuleStartup()
 {
 	UNiagaraLightRendererProperties* CDO = CastChecked<UNiagaraLightRendererProperties>(UNiagaraLightRendererProperties::StaticClass()->GetDefaultObject());
-
 	CDO->PositionBinding = FNiagaraConstants::GetAttributeDefaultBinding(SYS_PARAM_PARTICLES_POSITION);
 	CDO->ColorBinding = FNiagaraConstants::GetAttributeDefaultBinding(SYS_PARAM_PARTICLES_COLOR);
 	CDO->RadiusBinding = FNiagaraConstants::GetAttributeDefaultBinding(SYS_PARAM_PARTICLES_LIGHT_RADIUS);
 	CDO->LightExponentBinding = FNiagaraConstants::GetAttributeDefaultBinding(SYS_PARAM_PARTICLES_LIGHT_EXPONENT);
 	CDO->LightRenderingEnabledBinding = FNiagaraConstants::GetAttributeDefaultBinding(SYS_PARAM_PARTICLES_LIGHT_ENABLED);
 	CDO->VolumetricScatteringBinding = FNiagaraConstants::GetAttributeDefaultBinding(SYS_PARAM_PARTICLES_LIGHT_VOLUMETRIC_SCATTERING);
+
+	for (TWeakObjectPtr<UNiagaraLightRendererProperties>& WeakLightRendererProperties : LightRendererPropertiesToDeferredInit)
+	{
+		if (WeakLightRendererProperties.Get())
+		{
+			if (WeakLightRendererProperties->PositionBinding.BoundVariable.GetName() == NAME_None)
+			{
+				WeakLightRendererProperties->PositionBinding = FNiagaraConstants::GetAttributeDefaultBinding(SYS_PARAM_PARTICLES_POSITION);
+				WeakLightRendererProperties->ColorBinding = FNiagaraConstants::GetAttributeDefaultBinding(SYS_PARAM_PARTICLES_COLOR);
+				WeakLightRendererProperties->RadiusBinding = FNiagaraConstants::GetAttributeDefaultBinding(SYS_PARAM_PARTICLES_LIGHT_RADIUS);
+				WeakLightRendererProperties->LightExponentBinding = FNiagaraConstants::GetAttributeDefaultBinding(SYS_PARAM_PARTICLES_LIGHT_EXPONENT);
+				WeakLightRendererProperties->LightRenderingEnabledBinding = FNiagaraConstants::GetAttributeDefaultBinding(SYS_PARAM_PARTICLES_LIGHT_ENABLED);
+				WeakLightRendererProperties->VolumetricScatteringBinding = FNiagaraConstants::GetAttributeDefaultBinding(SYS_PARAM_PARTICLES_LIGHT_VOLUMETRIC_SCATTERING);
+			}
+		}
+	}
 }
 
 FNiagaraRenderer* UNiagaraLightRendererProperties::CreateEmitterRenderer(ERHIFeatureLevel::Type FeatureLevel, const FNiagaraEmitterInstance* Emitter)
diff --git a/Engine/Plugins/FX/Niagara/Source/Niagara/Private/NiagaraMeshRendererProperties.cpp b/Engine/Plugins/FX/Niagara/Source/Niagara/Private/NiagaraMeshRendererProperties.cpp
--- a/Engine/Plugins/FX/Niagara/Source/Niagara/Private/NiagaraMeshRendererProperties.cpp
+++ b/Engine/Plugins/FX/Niagara/Source/Niagara/Private/NiagaraMeshRendererProperties.cpp
@@ -6,6 +6,8 @@
 #include "NiagaraConstants.h"
 #include "NiagaraBoundsCalculatorHelper.h"
 
+TArray<TWeakObjectPtr<UNiagaraMeshRendererProperties>> UNiagaraMeshRendererProperties::MeshRendererPropertiesToDeferredInit;
+
 FNiagaraMeshMaterialOverride::FNiagaraMeshMaterialOverride()
 	: ExplicitMat(nullptr)
 {
@@ -61,8 +63,15 @@ FNiagaraBoundsCalculator* UNiagaraMeshRendererProperties::CreateBoundsCalculator
 void UNiagaraMeshRendererProperties::PostInitProperties()
 {
 	Super::PostInitProperties();
+
 	if (HasAnyFlags(RF_ClassDefaultObject) == false)
 	{
+		// We can end up hitting PostInitProperties before the Niagara Module has initialized bindings this needs, mark this object for deferred init and early out.
+		if (FModuleManager::Get().IsModuleLoaded("Niagara") == false)
+		{
+			MeshRendererPropertiesToDeferredInit.Add(this);
+			return;
+		}
 		InitBindings();
 	}
 }
@@ -72,6 +81,14 @@ void UNiagaraMeshRendererProperties::InitCDOPropertiesAfterModuleStartup()
 {
 	UNiagaraMeshRendererProperties* CDO = CastChecked<UNiagaraMeshRendererProperties>(UNiagaraMeshRendererProperties::StaticClass()->GetDefaultObject());
 	CDO->InitBindings();
+
+	for (TWeakObjectPtr<UNiagaraMeshRendererProperties>& WeakMeshRendererProperties : MeshRendererPropertiesToDeferredInit)
+	{
+		if (WeakMeshRendererProperties.Get())
+		{
+			WeakMeshRendererProperties->InitBindings();
+		}
+	}
 }
 
 void UNiagaraMeshRendererProperties::InitBindings()
diff --git a/Engine/Plugins/FX/Niagara/Source/Niagara/Private/NiagaraRibbonRendererProperties.cpp b/Engine/Plugins/FX/Niagara/Source/Niagara/Private/NiagaraRibbonRendererProperties.cpp
--- a/Engine/Plugins/FX/Niagara/Source/Niagara/Private/NiagaraRibbonRendererProperties.cpp
+++ b/Engine/Plugins/FX/Niagara/Source/Niagara/Private/NiagaraRibbonRendererProperties.cpp
@@ -5,6 +5,8 @@
 #include "NiagaraConstants.h"
 #include "NiagaraBoundsCalculatorHelper.h"
 
+TArray<TWeakObjectPtr<UNiagaraRibbonRendererProperties>> UNiagaraRibbonRendererProperties::RibbonRendererPropertiesToDeferredInit;
+
 UNiagaraRibbonRendererProperties::UNiagaraRibbonRendererProperties()
 	: Material(nullptr)
 	, FacingMode(ENiagaraRibbonFacingMode::Screen)
@@ -70,8 +72,15 @@ void UNiagaraRibbonRendererProperties::GetUsedMaterials(const FNiagaraEmitterIns
 void UNiagaraRibbonRendererProperties::PostInitProperties()
 {
 	Super::PostInitProperties();
+
 	if (HasAnyFlags(RF_ClassDefaultObject) == false)
 	{
+		// We can end up hitting PostInitProperties before the Niagara Module has initialized bindings this needs, mark this object for deferred init and early out.
+		if (FModuleManager::Get().IsModuleLoaded("Niagara") == false)
+		{
+			RibbonRendererPropertiesToDeferredInit.Add(this);
+			return;
+		}
 		InitBindings();
 	}
 }
@@ -81,6 +90,14 @@ void UNiagaraRibbonRendererProperties::InitCDOPropertiesAfterModuleStartup()
 {
 	UNiagaraRibbonRendererProperties* CDO = CastChecked<UNiagaraRibbonRendererProperties>(UNiagaraRibbonRendererProperties::StaticClass()->GetDefaultObject());
 	CDO->InitBindings();
+
+	for (TWeakObjectPtr<UNiagaraRibbonRendererProperties>& WeakRibbonRendererProperties : RibbonRendererPropertiesToDeferredInit)
+	{
+		if (WeakRibbonRendererProperties.Get())
+		{
+			WeakRibbonRendererProperties->InitBindings();
+		}
+	}
 }
 
 void UNiagaraRibbonRendererProperties::InitBindings()
diff --git a/Engine/Plugins/FX/Niagara/Source/Niagara/Private/NiagaraSpriteRendererProperties.cpp b/Engine/Plugins/FX/Niagara/Source/Niagara/Private/NiagaraSpriteRendererProperties.cpp
--- a/Engine/Plugins/FX/Niagara/Source/Niagara/Private/NiagaraSpriteRendererProperties.cpp
+++ b/Engine/Plugins/FX/Niagara/Source/Niagara/Private/NiagaraSpriteRendererProperties.cpp
@@ -11,6 +11,8 @@
 #endif
 #define LOCTEXT_NAMESPACE "UNiagaraSpriteRendererProperties"
 
+TArray<TWeakObjectPtr<UNiagaraSpriteRendererProperties>> UNiagaraSpriteRendererProperties::SpriteRendererPropertiesToDeferredInit;
+
 #if ENABLE_COOK_STATS
 class NiagaraCutoutCookStats
 {
@@ -26,7 +28,6 @@ FCookStatsManager::FAutoRegisterCallback NiagaraCutoutCookStats::RegisterCookSta
 });
 #endif // ENABLE_COOK_STATS
 
-
 UNiagaraSpriteRendererProperties::UNiagaraSpriteRendererProperties()
 	: Alignment(ENiagaraSpriteAlignment::Unaligned)
 	, FacingMode(ENiagaraSpriteFacingMode::FaceCamera)
@@ -100,8 +101,15 @@ void UNiagaraSpriteRendererProperties::PostLoad()
 void UNiagaraSpriteRendererProperties::PostInitProperties()
 {
 	Super::PostInitProperties();
+
 	if (HasAnyFlags(RF_ClassDefaultObject) == false)
 	{
+		// We can end up hitting PostInitProperties before the Niagara Module has initialized bindings this needs, mark this object for deferred init and early out.
+		if (FModuleManager::Get().IsModuleLoaded("Niagara") == false)
+		{
+			SpriteRendererPropertiesToDeferredInit.Add(this);
+			return;
+		}
 		InitBindings();
 	}
 }
@@ -127,6 +135,14 @@ void UNiagaraSpriteRendererProperties::InitCDOPropertiesAfterModuleStartup()
 {
 	UNiagaraSpriteRendererProperties* CDO = CastChecked<UNiagaraSpriteRendererProperties>(UNiagaraSpriteRendererProperties::StaticClass()->GetDefaultObject());
 	CDO->InitBindings();
+
+	for (TWeakObjectPtr<UNiagaraSpriteRendererProperties>& WeakSpriteRendererProperties : SpriteRendererPropertiesToDeferredInit)
+	{
+		if (WeakSpriteRendererProperties.Get())
+		{
+			WeakSpriteRendererProperties->InitBindings();
+		}
+	}
 }
 
 void UNiagaraSpriteRendererProperties::InitBindings()
diff --git a/Engine/Plugins/FX/Niagara/Source/Niagara/Public/NiagaraLightRendererProperties.h b/Engine/Plugins/FX/Niagara/Source/Niagara/Public/NiagaraLightRendererProperties.h
--- a/Engine/Plugins/FX/Niagara/Source/Niagara/Public/NiagaraLightRendererProperties.h
+++ b/Engine/Plugins/FX/Niagara/Source/Niagara/Public/NiagaraLightRendererProperties.h
@@ -82,4 +82,7 @@ public:
 	/** Which attribute should we use for the intensity of the volumetric scattering from this light? This scales the light's intensity and color. */
 	UPROPERTY(EditAnywhere, AdvancedDisplay, Category = "Bindings")
 	FNiagaraVariableAttributeBinding VolumetricScatteringBinding;
+
+private:
+	static TArray<TWeakObjectPtr<UNiagaraLightRendererProperties>> LightRendererPropertiesToDeferredInit;
 };
diff --git a/Engine/Plugins/FX/Niagara/Source/Niagara/Public/NiagaraMeshRendererProperties.h b/Engine/Plugins/FX/Niagara/Source/Niagara/Public/NiagaraMeshRendererProperties.h
--- a/Engine/Plugins/FX/Niagara/Source/Niagara/Public/NiagaraMeshRendererProperties.h
+++ b/Engine/Plugins/FX/Niagara/Source/Niagara/Public/NiagaraMeshRendererProperties.h
@@ -168,4 +168,7 @@ public:
 protected:
 	bool FindBinding(const FNiagaraUserParameterBinding& InBinding, const FNiagaraEmitterInstance* InEmitter, TArray<UMaterialInterface*>& OutMaterials);
 	void InitBindings();
+
+private:
+	static TArray<TWeakObjectPtr<UNiagaraMeshRendererProperties>> MeshRendererPropertiesToDeferredInit;
 };
diff --git a/Engine/Plugins/FX/Niagara/Source/Niagara/Public/NiagaraRibbonRendererProperties.h b/Engine/Plugins/FX/Niagara/Source/Niagara/Public/NiagaraRibbonRendererProperties.h
--- a/Engine/Plugins/FX/Niagara/Source/Niagara/Public/NiagaraRibbonRendererProperties.h
+++ b/Engine/Plugins/FX/Niagara/Source/Niagara/Public/NiagaraRibbonRendererProperties.h
@@ -214,4 +214,7 @@ public:
 
 protected:
 	void InitBindings();
+
+private: 
+	static TArray<TWeakObjectPtr<UNiagaraRibbonRendererProperties>> RibbonRendererPropertiesToDeferredInit;
 };
diff --git a/Engine/Plugins/FX/Niagara/Source/Niagara/Public/NiagaraSpriteRendererProperties.h b/Engine/Plugins/FX/Niagara/Source/Niagara/Public/NiagaraSpriteRendererProperties.h
--- a/Engine/Plugins/FX/Niagara/Source/Niagara/Public/NiagaraSpriteRendererProperties.h
+++ b/Engine/Plugins/FX/Niagara/Source/Niagara/Public/NiagaraSpriteRendererProperties.h
@@ -234,4 +234,6 @@ private:
 
 	/** Derived data for this asset, generated off of SubUVTexture. */
 	FSubUVDerivedData DerivedData;
+
+	static TArray<TWeakObjectPtr<UNiagaraSpriteRendererProperties>> SpriteRendererPropertiesToDeferredInit;
 };
