Non unity includes fix.

diff --git a/Engine/Plugins/FX/Niagara/Source/Niagara/Private/NiagaraLightRendererProperties.cpp b/Engine/Plugins/FX/Niagara/Source/Niagara/Private/NiagaraLightRendererProperties.cpp
--- a/Engine/Plugins/FX/Niagara/Source/Niagara/Private/NiagaraLightRendererProperties.cpp
+++ b/Engine/Plugins/FX/Niagara/Source/Niagara/Private/NiagaraLightRendererProperties.cpp
@@ -4,6 +4,7 @@
 #include "NiagaraRenderer.h"
 #include "NiagaraConstants.h"
 #include "NiagaraRendererLights.h"
+#include "Modules/ModuleManager.h"
 
 TArray<TWeakObjectPtr<UNiagaraLightRendererProperties>> UNiagaraLightRendererProperties::LightRendererPropertiesToDeferredInit;
 
diff --git a/Engine/Plugins/FX/Niagara/Source/Niagara/Private/NiagaraMeshRendererProperties.cpp b/Engine/Plugins/FX/Niagara/Source/Niagara/Private/NiagaraMeshRendererProperties.cpp
--- a/Engine/Plugins/FX/Niagara/Source/Niagara/Private/NiagaraMeshRendererProperties.cpp
+++ b/Engine/Plugins/FX/Niagara/Source/Niagara/Private/NiagaraMeshRendererProperties.cpp
@@ -5,6 +5,7 @@
 #include "Engine/StaticMesh.h"
 #include "NiagaraConstants.h"
 #include "NiagaraBoundsCalculatorHelper.h"
+#include "Modules/ModuleManager.h"
 
 TArray<TWeakObjectPtr<UNiagaraMeshRendererProperties>> UNiagaraMeshRendererProperties::MeshRendererPropertiesToDeferredInit;
 
diff --git a/Engine/Plugins/FX/Niagara/Source/Niagara/Private/NiagaraRibbonRendererProperties.cpp b/Engine/Plugins/FX/Niagara/Source/Niagara/Private/NiagaraRibbonRendererProperties.cpp
--- a/Engine/Plugins/FX/Niagara/Source/Niagara/Private/NiagaraRibbonRendererProperties.cpp
+++ b/Engine/Plugins/FX/Niagara/Source/Niagara/Private/NiagaraRibbonRendererProperties.cpp
@@ -4,6 +4,7 @@
 #include "NiagaraRendererRibbons.h"
 #include "NiagaraConstants.h"
 #include "NiagaraBoundsCalculatorHelper.h"
+#include "Modules/ModuleManager.h"
 
 TArray<TWeakObjectPtr<UNiagaraRibbonRendererProperties>> UNiagaraRibbonRendererProperties::RibbonRendererPropertiesToDeferredInit;
 
diff --git a/Engine/Plugins/FX/Niagara/Source/Niagara/Private/NiagaraSpriteRendererProperties.cpp b/Engine/Plugins/FX/Niagara/Source/Niagara/Private/NiagaraSpriteRendererProperties.cpp
--- a/Engine/Plugins/FX/Niagara/Source/Niagara/Private/NiagaraSpriteRendererProperties.cpp
+++ b/Engine/Plugins/FX/Niagara/Source/Niagara/Private/NiagaraSpriteRendererProperties.cpp
@@ -6,6 +6,7 @@
 #include "NiagaraConstants.h"
 #include "NiagaraRendererSprites.h"
 #include "NiagaraBoundsCalculatorHelper.h"
+#include "Modules/ModuleManager.h"
 #if WITH_EDITOR
 #include "DerivedDataCacheInterface.h"
 #endif
