Niagara - Selection fixes.
+ Fixed loss of selection when the UI is refreshed to to undo, applying changes from emitters, appyling changes from modules, etc.
+ Select new modules after they've been added.
+ Fixed a fiew other issues cause by selection changes which were affecting stability.

diff --git a/Engine/Plugins/FX/Niagara/Source/NiagaraEditor/Private/Toolkits/NiagaraSystemToolkit.cpp b/Engine/Plugins/FX/Niagara/Source/NiagaraEditor/Private/Toolkits/NiagaraSystemToolkit.cpp
--- a/Engine/Plugins/FX/Niagara/Source/NiagaraEditor/Private/Toolkits/NiagaraSystemToolkit.cpp
+++ b/Engine/Plugins/FX/Niagara/Source/NiagaraEditor/Private/Toolkits/NiagaraSystemToolkit.cpp
@@ -169,8 +169,13 @@ FNiagaraSystemToolkit::~FNiagaraSystemToolkit()
 {
 	if (SystemViewModel.IsValid())
 	{
-		SystemViewModel->Cleanup();
+		if (SystemViewModel->GetSelectionViewModel() != nullptr)
+		{
+			SystemViewModel->GetSelectionViewModel()->OnSystemIsSelectedChanged().RemoveAll(this);
+			SystemViewModel->GetSelectionViewModel()->OnEmitterHandleIdSelectionChanged().RemoveAll(this);
+		}
 		SystemViewModel->GetOnPinnedCurvesChanged().RemoveAll(this);
+		SystemViewModel->Cleanup();
 	}
 	SystemViewModel.Reset();
 }
@@ -292,7 +297,8 @@ void FNiagaraSystemToolkit::InitializeInternal(const EToolkitMode::Type Mode, co
 	ObjectSelectionForParameterMapView = MakeShared<FNiagaraObjectSelection>();
 
 	SystemViewModel->OnEmitterHandleViewModelsChanged().AddSP(this, &FNiagaraSystemToolkit::RefreshParameters);
-	SystemViewModel->GetSelectionViewModel()->OnSelectionChanged().AddSP(this, &FNiagaraSystemToolkit::OnSystemSelectionChanged);
+	SystemViewModel->GetSelectionViewModel()->OnSystemIsSelectedChanged().AddSP(this, &FNiagaraSystemToolkit::OnSystemSelectionChanged);
+	SystemViewModel->GetSelectionViewModel()->OnEmitterHandleIdSelectionChanged().AddSP(this, &FNiagaraSystemToolkit::OnSystemSelectionChanged);
 	SystemViewModel->GetOnPinnedEmittersChanged().AddSP(this, &FNiagaraSystemToolkit::RefreshParameters);
 	SystemViewModel->GetOnPinnedCurvesChanged().AddSP(this, &FNiagaraSystemToolkit::OnPinnedCurvesChanged);
 
@@ -530,7 +536,7 @@ public:
 	void Construct(const FArguments& InArgs, TSharedRef<FNiagaraSystemViewModel> InSystemViewModel)
 	{
 		SystemViewModel = InSystemViewModel;
-		SystemViewModel->GetSelectionViewModel()->OnSelectionChanged().AddSP(this, &SNiagaraSelectedEmitterGraph::SystemSelectionChanged);
+		SystemViewModel->GetSelectionViewModel()->OnEmitterHandleIdSelectionChanged().AddSP(this, &SNiagaraSelectedEmitterGraph::SystemSelectionChanged);
 		ChildSlot
 		[
 			SAssignNew(GraphWidgetContainer, SBox)
@@ -542,12 +548,12 @@ public:
 	{
 		if (SystemViewModel.IsValid() && SystemViewModel->GetSelectionViewModel())
 		{
-			SystemViewModel->GetSelectionViewModel()->OnSelectionChanged().RemoveAll(this);
+			SystemViewModel->GetSelectionViewModel()->OnEmitterHandleIdSelectionChanged().RemoveAll(this);
 		}
 	}
 
 private:
-	void SystemSelectionChanged(UNiagaraSystemSelectionViewModel::ESelectionChangeSource SelectionChangeSource)
+	void SystemSelectionChanged()
 	{
 		UpdateGraphWidget();
 	}
@@ -1303,7 +1309,7 @@ void FNiagaraSystemToolkit::RefreshParameters()
 	ObjectSelectionForParameterMapView->SetSelectedObjects(NewParameterViewSelection);
 }
 
-void FNiagaraSystemToolkit::OnSystemSelectionChanged(UNiagaraSystemSelectionViewModel::ESelectionChangeSource SelectionChangeSource)
+void FNiagaraSystemToolkit::OnSystemSelectionChanged()
 {
 	RefreshParameters();
 }
diff --git a/Engine/Plugins/FX/Niagara/Source/NiagaraEditor/Private/Toolkits/NiagaraSystemToolkit.h b/Engine/Plugins/FX/Niagara/Source/NiagaraEditor/Private/Toolkits/NiagaraSystemToolkit.h
--- a/Engine/Plugins/FX/Niagara/Source/NiagaraEditor/Private/Toolkits/NiagaraSystemToolkit.h
+++ b/Engine/Plugins/FX/Niagara/Source/NiagaraEditor/Private/Toolkits/NiagaraSystemToolkit.h
@@ -14,7 +14,6 @@
 #include "ISequencerTrackEditor.h"
 
 #include "NiagaraScript.h"
-#include "ViewModels/NiagaraSystemSelectionViewModel.h"
 
 class FNiagaraSystemInstance;
 class FNiagaraSystemViewModel;
@@ -125,7 +124,7 @@ private:
 
 	void OnPinnedCurvesChanged();
 	void RefreshParameters();
-	void OnSystemSelectionChanged(UNiagaraSystemSelectionViewModel::ESelectionChangeSource SelectionChangeSource);
+	void OnSystemSelectionChanged();
 
 	TSharedRef<SWidget> GenerateBoundsMenuContent(TSharedRef<FUICommandList> InCommandList);
 	const FName GetNiagaraSystemMessageLogName(UNiagaraSystem* InSystem) const;
diff --git a/Engine/Plugins/FX/Niagara/Source/NiagaraEditor/Private/ViewModels/NiagaraOverviewGraphViewModel.cpp b/Engine/Plugins/FX/Niagara/Source/NiagaraEditor/Private/ViewModels/NiagaraOverviewGraphViewModel.cpp
--- a/Engine/Plugins/FX/Niagara/Source/NiagaraEditor/Private/ViewModels/NiagaraOverviewGraphViewModel.cpp
+++ b/Engine/Plugins/FX/Niagara/Source/NiagaraEditor/Private/ViewModels/NiagaraOverviewGraphViewModel.cpp
@@ -10,6 +10,7 @@
 #include "NiagaraObjectSelection.h"
 #include "ViewModels/NiagaraSystemViewModel.h"
 #include "ViewModels/NiagaraEmitterHandleViewModel.h"
+#include "ViewModels/NiagaraSystemSelectionViewModel.h"
 #include "ViewModels/Stack/NiagaraStackViewModel.h"
 #include "ViewModels/Stack/NiagaraStackEntry.h"
 #include "NiagaraSystem.h"
@@ -29,6 +30,14 @@ FNiagaraOverviewGraphViewModel::FNiagaraOverviewGraphViewModel()
 FNiagaraOverviewGraphViewModel::~FNiagaraOverviewGraphViewModel()
 {
 	GEditor->UnregisterForUndo(this);
+	TSharedPtr<FNiagaraSystemViewModel> SystemViewModelPinned = SystemViewModel.Pin();
+	if (SystemViewModelPinned.IsValid())
+	{
+		if (SystemViewModelPinned->GetSelectionViewModel() != nullptr)
+		{
+			SystemViewModelPinned->GetSelectionViewModel()->OnEntrySelectionChanged().RemoveAll(this);
+		}
+	}
 }
 
 void FNiagaraOverviewGraphViewModel::Initialize(TSharedRef<FNiagaraSystemViewModel> InSystemViewModel)
@@ -40,7 +49,7 @@ void FNiagaraOverviewGraphViewModel::Initialize(TSharedRef<FNiagaraSystemViewMod
 	GEditor->RegisterForUndo(this);
 
 	NodeSelection->OnSelectedObjectsChanged().AddSP(this, &FNiagaraOverviewGraphViewModel::GraphSelectionChanged);
-	InSystemViewModel->GetSelectionViewModel()->OnSelectionChanged().AddSP(this, &FNiagaraOverviewGraphViewModel::SystemSelectionChanged);
+	InSystemViewModel->GetSelectionViewModel()->OnEntrySelectionChanged().AddSP(this, &FNiagaraOverviewGraphViewModel::SystemSelectionChanged);
 }
 
 FText FNiagaraOverviewGraphViewModel::GetDisplayName() const
@@ -362,57 +371,75 @@ void FNiagaraOverviewGraphViewModel::GraphSelectionChanged()
 	{
 		TGuardValue<bool> UpdateGuard(bUpdatingSystemSelectionFromGraph, true);
 
-		bool bSystemIsSelected = false;
-		TArray<FGuid> SelectedEmitterHandleGuids;
-		for (UObject* SelectedNode : NodeSelection->GetSelectedObjects())
+		TArray<FGuid> SelectedGuids;
+		for (UObject* SelectedObject : NodeSelection->GetSelectedObjects())
 		{
-			UNiagaraOverviewNode* OverviewNode = Cast<UNiagaraOverviewNode>(SelectedNode);
-			if (OverviewNode != nullptr)
+			UNiagaraOverviewNode* SelectedOverviewNode = Cast<UNiagaraOverviewNode>(SelectedObject);
+			if (SelectedOverviewNode != nullptr)
 			{
-				if (OverviewNode->GetEmitterHandleGuid().IsValid())
-				{
-					SelectedEmitterHandleGuids.Add(OverviewNode->GetEmitterHandleGuid());
-				}
-				else
-				{
-					bSystemIsSelected = true;
-				}
+				SelectedGuids.AddUnique(SelectedOverviewNode->GetEmitterHandleGuid());
+			}
+		}
+
+		TArray<UNiagaraStackEntry*> EntriesToSelect;
+		TArray<UNiagaraStackEntry*> EntriesToDeselect;
+		UNiagaraStackEntry* SystemRootEntry = GetSystemViewModel()->GetSystemStackViewModel()->GetRootEntry();
+		if (SelectedGuids.Contains(FGuid()))
+		{
+			EntriesToSelect.Add(SystemRootEntry);
+		}
+		else
+		{
+			EntriesToDeselect.Add(SystemRootEntry);
+		}
+
+		for (TSharedRef<FNiagaraEmitterHandleViewModel> EmitterHandleViewModel : GetSystemViewModel()->GetEmitterHandleViewModels())
+		{
+			UNiagaraStackEntry* EmitterRootEntry = EmitterHandleViewModel->GetEmitterStackViewModel()->GetRootEntry();
+			if (SelectedGuids.Contains(EmitterHandleViewModel->GetId()))
+			{
+				EntriesToSelect.Add(EmitterRootEntry);
+			}
+			else
+			{
+				EntriesToDeselect.Add(EmitterRootEntry);
 			}
 		}
 
 		bool bClearCurrentSelection = FSlateApplication::Get().GetModifierKeys().IsControlDown() == false;
-		GetSystemViewModel()->GetSelectionViewModel()->UpdateSelectionFromTopLevelObjects(bSystemIsSelected, SelectedEmitterHandleGuids, bClearCurrentSelection);
+		GetSystemViewModel()->GetSelectionViewModel()->UpdateSelectedEntries(EntriesToSelect, EntriesToDeselect, bClearCurrentSelection);
 	}
 }
 
-void FNiagaraOverviewGraphViewModel::SystemSelectionChanged(UNiagaraSystemSelectionViewModel::ESelectionChangeSource SelectionChangeSource)
+void FNiagaraOverviewGraphViewModel::SystemSelectionChanged()
 {
-	if (bUpdatingSystemSelectionFromGraph == false && 
-		SelectionChangeSource != UNiagaraSystemSelectionViewModel::ESelectionChangeSource::EntrySelection)
+	if (bUpdatingSystemSelectionFromGraph == false)
 	{
 		TGuardValue<bool> UpdateGuard(bUpdatingGraphSelectionFromSystem, true);
 
 		TArray<UObject*> SelectedNodes;
 		TArray<UNiagaraOverviewNode*> OverviewNodes;
 		OverviewGraph->GetNodesOfClass<UNiagaraOverviewNode>(OverviewNodes);
+		TArray<UNiagaraStackEntry*> SelectedEntries;
+		GetSystemViewModel()->GetSelectionViewModel()->GetSelectedEntries(SelectedEntries);
 		for (UNiagaraOverviewNode* OverviewNode : OverviewNodes)
 		{
 			if (OverviewNode->GetEmitterHandleGuid().IsValid())
 			{
-				if (GetSystemViewModel()->GetSelectionViewModel()->GetSelectedEmitterHandleIds().Contains(OverviewNode->GetEmitterHandleGuid()))
+				TSharedPtr<FNiagaraEmitterHandleViewModel> EmitterHandleviewModel = GetSystemViewModel()->GetEmitterHandleViewModelById(OverviewNode->GetEmitterHandleGuid());
+				if (EmitterHandleviewModel.IsValid() && SelectedEntries.Contains(EmitterHandleviewModel->GetEmitterStackViewModel()->GetRootEntry()))
 				{
 					SelectedNodes.Add(OverviewNode);
 				}
 			}
 			else
 			{
-				if (GetSystemViewModel()->GetSelectionViewModel()->GetSystemIsSelected())
+				if (SelectedEntries.Contains(GetSystemViewModel()->GetSystemStackViewModel()->GetRootEntry()))
 				{
 					SelectedNodes.Add(OverviewNode);
 				}
 			}
 		}
-
 		NodeSelection->SetSelectedObjects(SelectedNodes);
 	}
 }
diff --git a/Engine/Plugins/FX/Niagara/Source/NiagaraEditor/Private/ViewModels/NiagaraSystemSelectionViewModel.cpp b/Engine/Plugins/FX/Niagara/Source/NiagaraEditor/Private/ViewModels/NiagaraSystemSelectionViewModel.cpp
--- a/Engine/Plugins/FX/Niagara/Source/NiagaraEditor/Private/ViewModels/NiagaraSystemSelectionViewModel.cpp
+++ b/Engine/Plugins/FX/Niagara/Source/NiagaraEditor/Private/ViewModels/NiagaraSystemSelectionViewModel.cpp
@@ -13,6 +13,26 @@
 #include "NiagaraSystemEditorData.h"
 #include "NiagaraEmitterHandle.h"
 
+FGuid StackEntryToEmitterGuid(UNiagaraStackEntry* StackEntry)
+{
+	if (StackEntry->GetEmitterViewModel().IsValid())
+	{
+		TSharedPtr<FNiagaraEmitterHandleViewModel> EmitterHandleViewModel = StackEntry->GetSystemViewModel()->GetEmitterHandleViewModelForEmitter(StackEntry->GetEmitterViewModel()->GetEmitter());
+		if (EmitterHandleViewModel.IsValid())
+		{
+			return EmitterHandleViewModel->GetId();
+		}
+	}
+	return FGuid();
+}
+
+UNiagaraSystemSelectionViewModel::FSelectionEntry::FSelectionEntry(UNiagaraStackEntry* SelectedEntry)
+	: Entry(SelectedEntry)
+	, EmitterGuid(StackEntryToEmitterGuid(SelectedEntry))
+	, EditorDataKey(SelectedEntry->GetStackEditorDataKey())
+{
+}
+
 void UNiagaraSystemSelectionViewModel::Initialize(TSharedRef<FNiagaraSystemViewModel> InSystemViewModel)
 {
 	SystemViewModelWeak = InSystemViewModel;
@@ -49,9 +69,24 @@ void UNiagaraSystemSelectionViewModel::Finalize()
 	}
 }
 
-const TArray<UNiagaraStackEntry*> UNiagaraSystemSelectionViewModel::GetSelectedEntries() const
+bool UNiagaraSystemSelectionViewModel::ContainsEntry(UNiagaraStackEntry* StackEntry) const
 {
-	return SelectedEntries;
+	for (const FSelectionEntry& SelectionEntry : SelectionEntries)
+	{
+		if (SelectionEntry.Entry == StackEntry)
+		{
+			return true;
+		}
+	}
+	return false;
+}
+
+void UNiagaraSystemSelectionViewModel::GetSelectedEntries(TArray<UNiagaraStackEntry*>& OutSelectedEntries) const
+{
+	for (const FSelectionEntry& SelectionEntry : SelectionEntries)
+	{
+		OutSelectedEntries.Add(SelectionEntry.Entry.Get());
+	}
 }
 
 bool UNiagaraSystemSelectionViewModel::GetSystemIsSelected() const
@@ -64,142 +99,38 @@ const TArray<FGuid>& UNiagaraSystemSelectionViewModel::GetSelectedEmitterHandleI
 	return SelectedEmitterHandleIds;
 }
 
-void UNiagaraSystemSelectionViewModel::UpdateSelectionFromEntries(const TArray<UNiagaraStackEntry*> InSelectedEntries, const TArray<UNiagaraStackEntry*> InDeselectedEntries, bool bClearCurrentSelection)
+void UNiagaraSystemSelectionViewModel::UpdateSelectedEntries(const TArray<UNiagaraStackEntry*> InSelectedEntries, const TArray<UNiagaraStackEntry*> InDeselectedEntries, bool bClearCurrentSelection)
 {
+	bool bSelectionChanged = false;
 	if (bClearCurrentSelection)
 	{
-		SelectedEntries.Empty();
-		bSystemIsSelected = false;
-		SelectedEmitterHandleIds.Empty();
-		OnSelectionChangedDelegate.Broadcast(ESelectionChangeSource::Clear);
-		UpdateStackSelectionEntry();
+		SelectionEntries.Empty();
+		bSelectionChanged = true;
 	}
-
-	bool bSelectionChanged = false;
-	TSet<UNiagaraStackEntry*> CurrentSelection(SelectedEntries);
-
-	for (UNiagaraStackEntry* DeselectedEntry : InDeselectedEntries)
+	else
 	{
-		if (CurrentSelection.Contains(DeselectedEntry))
+		for (UNiagaraStackEntry* DeselectedEntry : InDeselectedEntries)
 		{
-			SelectedEntries.Remove(DeselectedEntry);
-			TSharedPtr<FNiagaraEmitterViewModel> DeselectedEmitterViewModel = DeselectedEntry->GetEmitterViewModel();
-			if (DeselectedEmitterViewModel.IsValid())
-			{
-				bool bEmitterIsStillSelected = SelectedEntries.ContainsByPredicate([DeselectedEmitterViewModel](const UNiagaraStackEntry* SelectedEntry)
-				{ 
-					return SelectedEntry->GetEmitterViewModel() == DeselectedEmitterViewModel; 
-				});
-				if (bEmitterIsStillSelected == false)
-				{
-					const FNiagaraEmitterHandle* DeselectedEmitterHandle = FNiagaraEditorUtilities::GetEmitterHandleForEmitter(DeselectedEntry->GetSystemViewModel()->GetSystem(), *DeselectedEmitterViewModel->GetEmitter());
-					if (DeselectedEmitterHandle != nullptr)
-					{
-						SelectedEmitterHandleIds.Remove(DeselectedEmitterHandle->GetId());
-					}
-				}
-			}
-			else
+			if (ContainsEntry(DeselectedEntry))
 			{
-				bool bSystemIsStillSelected = SelectedEntries.ContainsByPredicate([](const UNiagaraStackEntry* SelectedEntry)
-				{
-					return SelectedEntry->GetEmitterViewModel().IsValid() == false;
-				});
-				if (bSystemIsStillSelected == false)
-				{
-					bSystemIsSelected = false;
-				}
+				RemoveEntryFromSelectionInternal(DeselectedEntry);
+				bSelectionChanged = true;
 			}
-			bSelectionChanged = true;
 		}
 	}
 
 	for (UNiagaraStackEntry* SelectedEntry : InSelectedEntries)
 	{
-		if (CurrentSelection.Contains(SelectedEntry) == false)
-		{
-			SelectedEntries.Add(SelectedEntry);
-			TSharedPtr<FNiagaraEmitterViewModel> SelectedEmitterViewModel = SelectedEntry->GetEmitterViewModel();
-			if (SelectedEmitterViewModel.IsValid())
-			{
-				const FNiagaraEmitterHandle* SelectedEmitterHandle = FNiagaraEditorUtilities::GetEmitterHandleForEmitter(SelectedEntry->GetSystemViewModel()->GetSystem(), *SelectedEmitterViewModel->GetEmitter());
-				if (SelectedEmitterHandle != nullptr)
-				{
-					SelectedEmitterHandleIds.AddUnique(SelectedEmitterHandle->GetId());
-				}
-			}
-			else
-			{
-				bSystemIsSelected = true;
-			}
-			bSelectionChanged = true;
-		}
-	}
-
-	if (bSelectionChanged)
-	{
-		OnSelectionChangedDelegate.Broadcast(ESelectionChangeSource::EntrySelection);
-		UpdateStackSelectionEntry();
-	}
-}
-
-void UNiagaraSystemSelectionViewModel::UpdateSelectionFromTopLevelObjects(bool bInSystemIsSelected, const TArray<FGuid> InSelectedEmitterIds, bool bClearCurrentSelection)
-{
-	if (bClearCurrentSelection)
-	{
-		SelectedEntries.Empty();
-		bSystemIsSelected = false;
-		SelectedEmitterHandleIds.Empty();
-		OnSelectionChangedDelegate.Broadcast(ESelectionChangeSource::Clear);
-		UpdateStackSelectionEntry();
-	}
-
-	bool bSelectionChanged = false;
-	if (bSystemIsSelected != bInSystemIsSelected)
-	{
-		bSystemIsSelected = bInSystemIsSelected;
-		if (bSystemIsSelected)
-		{
-			SelectedEntries.Add(GetSystemViewModel()->GetSystemStackViewModel()->GetRootEntry());
-		}
-		else
-		{
-			SelectedEntries.Remove(GetSystemViewModel()->GetSystemStackViewModel()->GetRootEntry());
-		}
-		bSelectionChanged = true;
-	}
-
-	TSet<FGuid> CurrentSelectedEmitterHandleIds(SelectedEmitterHandleIds);
-
-	for (const FGuid& CurrentSelectedEmitterHandleId : CurrentSelectedEmitterHandleIds)
-	{
-		if (InSelectedEmitterIds.Contains(CurrentSelectedEmitterHandleId) == false)
-		{
-			SelectedEmitterHandleIds.Remove(CurrentSelectedEmitterHandleId);
-			TSharedPtr<FNiagaraEmitterHandleViewModel> DeselectedEmitterHandleViewModel = GetSystemViewModel()->GetEmitterHandleViewModelById(CurrentSelectedEmitterHandleId);
-			SelectedEntries.Remove(DeselectedEmitterHandleViewModel->GetEmitterStackViewModel()->GetRootEntry());
-			bSelectionChanged = true;
-		}
-	}
-
-	for (const FGuid& InSelectedEmitterId : InSelectedEmitterIds)
-	{
-		if (CurrentSelectedEmitterHandleIds.Contains(InSelectedEmitterId) == false)
+		if (ContainsEntry(SelectedEntry) == false)
 		{
-			SelectedEmitterHandleIds.Add(InSelectedEmitterId);
-			TSharedPtr<FNiagaraEmitterHandleViewModel> SelectedEmitterHandleViewModel = GetSystemViewModel()->GetEmitterHandleViewModelById(InSelectedEmitterId);
-			if (SelectedEmitterHandleViewModel.IsValid())
-			{
-				SelectedEntries.Add(SelectedEmitterHandleViewModel->GetEmitterStackViewModel()->GetRootEntry());
-			}
+			SelectionEntries.Add(FSelectionEntry(SelectedEntry));
 			bSelectionChanged = true;
 		}
 	}
 
 	if (bSelectionChanged)
 	{
-		OnSelectionChangedDelegate.Broadcast(ESelectionChangeSource::TopObjectLevelSelection);
-		UpdateStackSelectionEntry();
+		UpdateExternalSelectionState();
 	}
 }
 
@@ -213,24 +144,23 @@ void UNiagaraSystemSelectionViewModel::RemoveEntryFromSelectionByDisplayedObject
 	bool bSelectionChanged = false;
 
 	TArray<UNiagaraStackEntry*> EntriesToDeselect;
-	for (UNiagaraStackEntry* SelectedEntry : SelectedEntries)
+	for (FSelectionEntry& SelectionEntry : SelectionEntries)
 	{
-		if (SelectedEntry->GetDisplayedObject() == InObject)
+		if (SelectionEntry.Entry->GetDisplayedObject() == InObject)
 		{
-			EntriesToDeselect.Add(SelectedEntry);
-			bSelectionChanged = true;
+			EntriesToDeselect.Add(SelectionEntry.Entry.Get());
 		}
 	}
 
 	for (UNiagaraStackEntry* EntryToDeselect : EntriesToDeselect)
 	{
 		RemoveEntryFromSelectionInternal(EntryToDeselect);
+		bSelectionChanged = true;
 	}
 
 	if (bSelectionChanged)
 	{
-		OnSelectionChangedDelegate.Broadcast(ESelectionChangeSource::EntrySelection);
-		UpdateStackSelectionEntry();
+		UpdateExternalSelectionState();
 	}
 }
 
@@ -249,29 +179,78 @@ void UNiagaraSystemSelectionViewModel::AddEntryToSelectionByDisplayedObjectDefer
 	AddEntriesToSelectionByDisplayedObjectsDeferred(InObjects);
 }
 
-void UNiagaraSystemSelectionViewModel::Refresh()
+void CollectGroupAndItemEntries(TSharedRef<FNiagaraSystemViewModel> SystemViewModel, TMap<FGuid, TMap<FString, UNiagaraStackEntry*>>& OutEmitterGuidToEntryKeyToEntryMap)
 {
-	TSet<FGuid> ValidEmitterHandleIds;
-	for (const TSharedRef<FNiagaraEmitterHandleViewModel>& EmitterHandleViewModel : GetSystemViewModel()->GetEmitterHandleViewModels())
+	auto CollectFromStackViewModel = [](FGuid EmitterGuid, UNiagaraStackViewModel* StackViewModel, TMap<FGuid, TMap<FString, UNiagaraStackEntry*>>& OutEmitterGuidToEntryKeyToEntryMap)
+	{
+		TArray<UNiagaraStackEntry*> EntriesToCheck;
+		EntriesToCheck.Add(StackViewModel->GetRootEntry());
+
+		TMap<FString, UNiagaraStackEntry*>& EntryKeyToEntryMap = OutEmitterGuidToEntryKeyToEntryMap.FindOrAdd(EmitterGuid);
+		while (EntriesToCheck.Num() > 0)
+		{
+			UNiagaraStackEntry* Entry = EntriesToCheck[EntriesToCheck.Num() - 1];
+			EntriesToCheck.RemoveAt(EntriesToCheck.Num() - 1);
+
+			if (Entry->IsA<UNiagaraStackItemGroup>() || Entry->IsA<UNiagaraStackItem>())
+			{
+				EntryKeyToEntryMap.Add(Entry->GetStackEditorDataKey(), Entry);
+			}
+
+			Entry->GetUnfilteredChildren(EntriesToCheck);
+		}
+	};
+
+	CollectFromStackViewModel(FGuid(), SystemViewModel->GetSystemStackViewModel(), OutEmitterGuidToEntryKeyToEntryMap);
+	for (TSharedRef<FNiagaraEmitterHandleViewModel> EmitterHandleViewModel : SystemViewModel->GetEmitterHandleViewModels())
 	{
-		ValidEmitterHandleIds.Add(EmitterHandleViewModel->GetId());
+		CollectFromStackViewModel(EmitterHandleViewModel->GetId(), EmitterHandleViewModel->GetEmitterStackViewModel(), OutEmitterGuidToEntryKeyToEntryMap);
 	}
+}
 
-	int32 NumEmitterHandleIdsRemoved = SelectedEmitterHandleIds.RemoveAll(
-		[&ValidEmitterHandleIds](const FGuid& SelectedEmitterHandleId) { return ValidEmitterHandleIds.Contains(SelectedEmitterHandleId) == false; });
+bool UNiagaraSystemSelectionViewModel::Refresh()
+{
+	bool bSelectionChanged = false;
 
-	int32 NumEntriesRemoved = SelectedEntries.RemoveAll(
-		[](const UNiagaraStackEntry* SelectedEntry) { return SelectedEntry->IsFinalized(); });
+	TMap<FGuid, TMap<FString, UNiagaraStackEntry*>> EmitterGuidToEntryKeyToEntryMap;
+	TArray<FSelectionEntry> NewSelectionEntries;
+	for (FSelectionEntry& SelectionEntry : SelectionEntries)
+	{
+		if (SelectionEntry.Entry.IsValid() == false || SelectionEntry.Entry->IsFinalized())
+		{
+			// If the entry is no longer valid, try to reselect it by it's emitter guid and stack editor data key.
+			if (EmitterGuidToEntryKeyToEntryMap.Num() == 0)
+			{
+				// This is expensive, so only do it when necessary.
+				CollectGroupAndItemEntries(GetSystemViewModel(), EmitterGuidToEntryKeyToEntryMap);
+			}
+			TMap<FString, UNiagaraStackEntry*>* EntryKeyToEntryMap = EmitterGuidToEntryKeyToEntryMap.Find(SelectionEntry.EmitterGuid);
+			if (EntryKeyToEntryMap != nullptr)
+			{
+				UNiagaraStackEntry** Entry = EntryKeyToEntryMap->Find(SelectionEntry.EditorDataKey);
+				if (Entry != nullptr)
+				{
+					NewSelectionEntries.Add(FSelectionEntry(*Entry));
+				}
+			}
+			bSelectionChanged = true;
+		}
+		else
+		{
+			NewSelectionEntries.Add(SelectionEntry);
+		}
+	}
 
-	if (NumEmitterHandleIdsRemoved > 0 || NumEntriesRemoved > 0)
+	if (bSelectionChanged)
 	{
-		OnSelectionChangedDelegate.Broadcast(ESelectionChangeSource::Refresh);
-		UpdateStackSelectionEntry();
+		SelectionEntries = NewSelectionEntries;
+		UpdateExternalSelectionState();
 	}
 	else
 	{
 		StackSelection->RefreshChildren();
 	}
+	return bSelectionChanged;
 }
 
 void UNiagaraSystemSelectionViewModel::RefreshDeferred()
@@ -279,9 +258,19 @@ void UNiagaraSystemSelectionViewModel::RefreshDeferred()
 	bRefreshIsPending = true;
 }
 
-UNiagaraSystemSelectionViewModel::FOnSelectionChanged& UNiagaraSystemSelectionViewModel::OnSelectionChanged()
+UNiagaraSystemSelectionViewModel::FOnSelectionChanged& UNiagaraSystemSelectionViewModel::OnEntrySelectionChanged()
 {
-	return OnSelectionChangedDelegate;
+	return OnEntrySelectionChangedDelegate;
+}
+
+UNiagaraSystemSelectionViewModel::FOnSelectionChanged& UNiagaraSystemSelectionViewModel::OnEmitterHandleIdSelectionChanged()
+{
+	return OnEmitterHandleIdSelectionChangedDelegate;
+}
+
+UNiagaraSystemSelectionViewModel::FOnSelectionChanged& UNiagaraSystemSelectionViewModel::OnSystemIsSelectedChanged()
+{
+	return OnSystemIsSelectedChangedDelegate;
 }
 
 void FindStackGroupsAndItemsForDisplayedObjectKeysRecursive(UNiagaraStackEntry* StackEntry, const TArray<FObjectKey>& ObjectKeys, TArray<UNiagaraStackEntry*>& OutFoundStackEntries)
@@ -308,12 +297,6 @@ void FindStackGroupsAndItemsForDisplayedObjectKeysRecursive(UNiagaraStackEntry*
 
 void UNiagaraSystemSelectionViewModel::Tick()
 {
-	if (bRefreshIsPending)
-	{
-		Refresh();
-		bRefreshIsPending = false;
-	}
-
 	TArray<UNiagaraStackEntry*> FoundStackEntries;
 	TSharedRef<FNiagaraSystemViewModel> SystemViewModel = GetSystemViewModel();
 	FindStackGroupsAndItemsForDisplayedObjectKeysRecursive(SystemViewModel->GetSystemStackViewModel()->GetRootEntry(), DeferredDisplayedObjectKeysToAddToSelection, FoundStackEntries);
@@ -322,21 +305,28 @@ void UNiagaraSystemSelectionViewModel::Tick()
 		FindStackGroupsAndItemsForDisplayedObjectKeysRecursive(EmitterViewModel->GetEmitterStackViewModel()->GetRootEntry(), DeferredDisplayedObjectKeysToAddToSelection, FoundStackEntries);
 	}
 
-	int32 AddedEntries = 0;
+	bool bSelectionChanged = false;
 	for (UNiagaraStackEntry* FoundStackEntry : FoundStackEntries)
 	{
-		if (SelectedEntries.Contains(FoundStackEntry) == false)
+		if (ContainsEntry(FoundStackEntry) == false)
 		{
-			SelectedEntries.Add(FoundStackEntry);
-			AddedEntries++;
+			SelectionEntries.Add(FSelectionEntry(FoundStackEntry));
+			bSelectionChanged = true;
 		}
 	}
 
 	DeferredDisplayedObjectKeysToAddToSelection.Empty();
-	if (AddedEntries > 0)
+
+	bool bHasBeenRefreshed = false;
+	if (bRefreshIsPending)
 	{
-		OnSelectionChangedDelegate.Broadcast(ESelectionChangeSource::EntrySelection);
-		UpdateStackSelectionEntry();
+		bHasBeenRefreshed = Refresh();
+		bRefreshIsPending = false;
+	}
+
+	if (bSelectionChanged && bHasBeenRefreshed == false)
+	{
+		UpdateExternalSelectionState();
 	}
 }
 
@@ -349,48 +339,48 @@ TSharedRef<FNiagaraSystemViewModel> UNiagaraSystemSelectionViewModel::GetSystemV
 
 void UNiagaraSystemSelectionViewModel::RemoveEntryFromSelectionInternal(UNiagaraStackEntry* DeselectedEntry)
 {
-	SelectedEntries.Remove(DeselectedEntry);
-	TSharedPtr<FNiagaraEmitterViewModel> DeselectedEmitterViewModel = DeselectedEntry->GetEmitterViewModel();
-	if (DeselectedEmitterViewModel.IsValid())
+	SelectionEntries.RemoveAll([DeselectedEntry](const FSelectionEntry& SelectionEntry) { return SelectionEntry.Entry == DeselectedEntry; });
+}
+
+void UNiagaraSystemSelectionViewModel::UpdateExternalSelectionState()
+{
+	bSystemIsSelected = false;
+	TArray<FGuid> OldSelectedEmitterHandleIds = SelectedEmitterHandleIds;
+	SelectedEmitterHandleIds.Empty();
+	TArray<UNiagaraEmitter*> SelectedEmitters;
+	for (const FSelectionEntry& SelectionEntry : SelectionEntries)
 	{
-		bool bEmitterIsStillSelected = SelectedEntries.ContainsByPredicate([DeselectedEmitterViewModel](const UNiagaraStackEntry* SelectedEntry)
+		if (SelectionEntry.Entry->GetEmitterViewModel().IsValid())
 		{
-			return SelectedEntry->GetEmitterViewModel() == DeselectedEmitterViewModel;
-		});
-		if (bEmitterIsStillSelected == false)
+			SelectedEmitters.AddUnique(SelectionEntry.Entry->GetEmitterViewModel()->GetEmitter());
+		}
+		else
 		{
-			const FNiagaraEmitterHandle* DeselectedEmitterHandle = FNiagaraEditorUtilities::GetEmitterHandleForEmitter(DeselectedEntry->GetSystemViewModel()->GetSystem(), *DeselectedEmitterViewModel->GetEmitter());
-			if (DeselectedEmitterHandle != nullptr)
-			{
-				SelectedEmitterHandleIds.Remove(DeselectedEmitterHandle->GetId());
-			}
+			bSystemIsSelected = true;
 		}
 	}
-	else
+
+	int32 NumRemoved = 0;
+	for (UNiagaraEmitter* SelectedEmitter : SelectedEmitters)
 	{
-		bool bSystemIsStillSelected = SelectedEntries.ContainsByPredicate([](const UNiagaraStackEntry* SelectedEntry)
-		{
-			return SelectedEntry->GetEmitterViewModel().IsValid() == false;
-		});
-		if (bSystemIsStillSelected == false)
-		{
-			bSystemIsSelected = false;
-		}
+		FGuid EmitterId = GetSystemViewModel()->GetEmitterHandleViewModelForEmitter(SelectedEmitter)->GetId();
+		NumRemoved += OldSelectedEmitterHandleIds.Remove(EmitterId);
+		SelectedEmitterHandleIds.Add(EmitterId);
 	}
-}
 
-void UNiagaraSystemSelectionViewModel::UpdateStackSelectionEntry()
-{
+	// If there were ids left over, or we added more than we removed than the id selection has changed.
+	bool bEmitterHandleIdSelectionChanged = OldSelectedEmitterHandleIds.Num() != 0 || SelectedEmitterHandleIds.Num() != NumRemoved;
+
 	// Check the new selection to make sure that there are no entries which are owned by other entries and if there are
 	// remove them.
 	TArray<UNiagaraStackEntry*> StackSelectionEntries;
-	for (UNiagaraStackEntry* SelectedEntry : SelectedEntries)
+	for (const FSelectionEntry& SelectionEntry : SelectionEntries)
 	{
 		bool bIsChildEntry = false;
-		UNiagaraStackEntry* EntryInOuterChain = SelectedEntry->GetTypedOuter<UNiagaraStackEntry>();
+		UNiagaraStackEntry* EntryInOuterChain = SelectionEntry.Entry->GetTypedOuter<UNiagaraStackEntry>();
 		while (EntryInOuterChain != nullptr)
 		{
-			if (SelectedEntries.Contains(EntryInOuterChain))
+			if (ContainsEntry(EntryInOuterChain))
 			{
 				bIsChildEntry = true;
 				EntryInOuterChain = nullptr;
@@ -403,9 +393,14 @@ void UNiagaraSystemSelectionViewModel::UpdateStackSelectionEntry()
 
 		if (bIsChildEntry == false)
 		{
-			StackSelectionEntries.Add(SelectedEntry);
+			StackSelectionEntries.Add(SelectionEntry.Entry.Get());
 		}
 	}
 
 	StackSelection->SetSelectedEntries(StackSelectionEntries);
+	OnEntrySelectionChangedDelegate.Broadcast();
+	if (bEmitterHandleIdSelectionChanged)
+	{
+		OnEmitterHandleIdSelectionChangedDelegate.Broadcast();
+	}
 }
\ No newline at end of file
diff --git a/Engine/Plugins/FX/Niagara/Source/NiagaraEditor/Private/ViewModels/NiagaraSystemViewModel.cpp b/Engine/Plugins/FX/Niagara/Source/NiagaraEditor/Private/ViewModels/NiagaraSystemViewModel.cpp
--- a/Engine/Plugins/FX/Niagara/Source/NiagaraEditor/Private/ViewModels/NiagaraSystemViewModel.cpp
+++ b/Engine/Plugins/FX/Niagara/Source/NiagaraEditor/Private/ViewModels/NiagaraSystemViewModel.cpp
@@ -88,7 +88,7 @@ void FNiagaraSystemViewModel::Initialize(UNiagaraSystem& InSystem, FNiagaraSyste
 	
 	SelectionViewModel = NewObject<UNiagaraSystemSelectionViewModel>(GetTransientPackage());
 	SelectionViewModel->Initialize(this->AsShared());
-	SelectionViewModel->OnSelectionChanged().AddSP(this, &FNiagaraSystemViewModel::SystemSelectionChanged);
+	SelectionViewModel->OnEmitterHandleIdSelectionChanged().AddSP(this, &FNiagaraSystemViewModel::SystemSelectionChanged);
 
 	SystemScriptViewModel = MakeShared<FNiagaraSystemScriptViewModel>();
 	SystemScriptViewModel->Initialize(GetSystem());
@@ -173,6 +173,7 @@ void FNiagaraSystemViewModel::Cleanup()
 
 	if (SelectionViewModel != nullptr)
 	{
+		SelectionViewModel->OnEmitterHandleIdSelectionChanged().RemoveAll(this);
 		SelectionViewModel->Finalize();
 		SelectionViewModel = nullptr;
 	}
@@ -327,9 +328,10 @@ void FNiagaraSystemViewModel::AddEmitter(UNiagaraEmitter& Emitter)
 
 	NiagaraSequence->GetMovieScene()->SetPlaybackRange(NewStartFrame.RoundToFrame(), NewDuration);
 
-	TArray<FGuid> SelectedEmitterHandles;
-	SelectedEmitterHandles.Add(NewEmitterHandleId);
-	SelectionViewModel->UpdateSelectionFromTopLevelObjects(false, SelectedEmitterHandles, true);
+	TSharedPtr<FNiagaraEmitterHandleViewModel> NewEmitterHandleViewModel = GetEmitterHandleViewModelById(NewEmitterHandleId);
+	TArray<UNiagaraStackEntry*> SelectedStackEntries;
+	SelectedStackEntries.Add(NewEmitterHandleViewModel->GetEmitterStackViewModel()->GetRootEntry());
+	SelectionViewModel->UpdateSelectedEntries(SelectedStackEntries, TArray<UNiagaraStackEntry*>(), true);
 
 	bForceAutoCompileOnce = true;
 }
@@ -1561,7 +1563,7 @@ void FNiagaraSystemViewModel::SequencerTimeChanged()
 	OnPostSequencerTimeChangeDelegate.Broadcast();
 }
 
-void FNiagaraSystemViewModel::SystemSelectionChanged(UNiagaraSystemSelectionViewModel::ESelectionChangeSource SelectionChangeSource)
+void FNiagaraSystemViewModel::SystemSelectionChanged()
 {
 	if (bUpdatingSystemSelectionFromSequencer == false)
 	{
@@ -1614,9 +1616,31 @@ void FNiagaraSystemViewModel::UpdateEmitterHandleSelectionFromSequencer()
 	TGuardValue<bool> UpdateGuard(bUpdatingSystemSelectionFromSequencer, true);
 	if (SelectionViewModel != nullptr)
 	{
+		TArray<UNiagaraStackEntry*> EntriesToSelect;
+		TArray<UNiagaraStackEntry*> EntriesToDeselect;
+		UNiagaraStackEntry* SystemRootEntry = SystemStackViewModel->GetRootEntry();
+		if (NewSelectedEmitterHandleIds.Num() > 0 && GetEditorData().GetOwningSystemIsPlaceholder() == false)
+		{
+			EntriesToSelect.Add(SystemRootEntry);
+		}
+		else
+		{
+			EntriesToDeselect.Add(SystemRootEntry);
+		}
+		for (TSharedRef<FNiagaraEmitterHandleViewModel> EmitterHandleViewModel : EmitterHandleViewModels)
+		{
+			UNiagaraStackEntry* EmitterRootEntry = EmitterHandleViewModel->GetEmitterStackViewModel()->GetRootEntry();
+			if (NewSelectedEmitterHandleIds.Contains(EmitterHandleViewModel->GetId()))
+			{
+				EntriesToSelect.Add(EmitterRootEntry);
+			}
+			else
+			{
+				EntriesToDeselect.Add(EmitterRootEntry);
+			}
+		}
 		bool bClearCurrentSelection = FSlateApplication::Get().GetModifierKeys().IsControlDown() == false;
-		bool bSelectSystem = NewSelectedEmitterHandleIds.Num() > 0 && GetEditorData().GetOwningSystemIsPlaceholder() == false;
-		SelectionViewModel->UpdateSelectionFromTopLevelObjects(bSelectSystem, NewSelectedEmitterHandleIds, bClearCurrentSelection);
+		SelectionViewModel->UpdateSelectedEntries(EntriesToSelect, EntriesToDeselect, bClearCurrentSelection);
 	}
 }
 
diff --git a/Engine/Plugins/FX/Niagara/Source/NiagaraEditor/Private/ViewModels/Stack/NiagaraStackScriptItemGroup.cpp b/Engine/Plugins/FX/Niagara/Source/NiagaraEditor/Private/ViewModels/Stack/NiagaraStackScriptItemGroup.cpp
--- a/Engine/Plugins/FX/Niagara/Source/NiagaraEditor/Private/ViewModels/Stack/NiagaraStackScriptItemGroup.cpp
+++ b/Engine/Plugins/FX/Niagara/Source/NiagaraEditor/Private/ViewModels/Stack/NiagaraStackScriptItemGroup.cpp
@@ -3,6 +3,7 @@
 #include "ViewModels/Stack/NiagaraStackScriptItemGroup.h"
 #include "ViewModels/Stack/NiagaraStackModuleItem.h"
 #include "ViewModels/NiagaraSystemViewModel.h"
+#include "ViewModels/NiagaraSystemSelectionViewModel.h"
 #include "ViewModels/NiagaraEmitterViewModel.h"
 #include "NiagaraScriptViewModel.h"
 #include "NiagaraScriptGraphViewModel.h"
@@ -826,6 +827,9 @@ TOptional<UNiagaraStackEntry::FDropRequestResponse> UNiagaraStackScriptItemGroup
 
 void UNiagaraStackScriptItemGroup::ItemAdded(UNiagaraNodeFunctionCall* AddedModule)
 {
+	TArray<const UObject*> DisplayedObjectsToSelect;
+	DisplayedObjectsToSelect.Add(AddedModule);
+	GetSystemViewModel()->GetSelectionViewModel()->AddEntriesToSelectionByDisplayedObjectsDeferred(DisplayedObjectsToSelect);
 	RefreshChildren();
 }
 
diff --git a/Engine/Plugins/FX/Niagara/Source/NiagaraEditor/Private/ViewModels/Stack/NiagaraStackViewModel.cpp b/Engine/Plugins/FX/Niagara/Source/NiagaraEditor/Private/ViewModels/Stack/NiagaraStackViewModel.cpp
--- a/Engine/Plugins/FX/Niagara/Source/NiagaraEditor/Private/ViewModels/Stack/NiagaraStackViewModel.cpp
+++ b/Engine/Plugins/FX/Niagara/Source/NiagaraEditor/Private/ViewModels/Stack/NiagaraStackViewModel.cpp
@@ -22,6 +22,18 @@
 #define LOCTEXT_NAMESPACE "NiagaraStackViewModel"
 const double UNiagaraStackViewModel::MaxSearchTime = .02f; // search at 50 fps
 
+UNiagaraStackViewModel::FTopLevelViewModel::FTopLevelViewModel(TSharedPtr<FNiagaraSystemViewModel> InSystemViewModel)
+	: SystemViewModel(InSystemViewModel)
+	, RootEntry(InSystemViewModel->GetSystemStackViewModel()->GetRootEntry())
+{
+}
+
+UNiagaraStackViewModel::FTopLevelViewModel::FTopLevelViewModel(TSharedPtr<FNiagaraEmitterHandleViewModel> InEmitterHandleViewModel)
+	: EmitterHandleViewModel(InEmitterHandleViewModel)
+	, RootEntry(InEmitterHandleViewModel->GetEmitterStackViewModel()->GetRootEntry())
+{
+}
+
 bool UNiagaraStackViewModel::FTopLevelViewModel::IsValid() const
 {
 	return (SystemViewModel.IsValid() && EmitterHandleViewModel.IsValid() == false) ||
@@ -62,7 +74,7 @@ FText UNiagaraStackViewModel::FTopLevelViewModel::GetDisplayName() const
 
 bool UNiagaraStackViewModel::FTopLevelViewModel::operator==(const FTopLevelViewModel& Other) const
 {
-	return Other.SystemViewModel == SystemViewModel && Other.EmitterHandleViewModel == EmitterHandleViewModel;
+	return Other.SystemViewModel == SystemViewModel && Other.EmitterHandleViewModel == EmitterHandleViewModel && Other.RootEntry == RootEntry;
 }
 
 void UNiagaraStackViewModel::InitializeWithViewModels(TSharedPtr<FNiagaraSystemViewModel> InSystemViewModel, TSharedPtr<FNiagaraEmitterHandleViewModel> InEmitterHandleViewModel, FNiagaraStackViewModelOptions InOptions)
@@ -656,8 +668,11 @@ void UNiagaraStackViewModel::RefreshTopLevelViewModels()
 		if (RootChild->GetEmitterViewModel().IsValid())
 		{
 			TSharedPtr<FNiagaraEmitterHandleViewModel> RootChildEmitterHandleViewModel = RootChild->GetSystemViewModel()->GetEmitterHandleViewModelForEmitter(RootChild->GetEmitterViewModel()->GetEmitter());
-			TSharedRef<FTopLevelViewModel>* CurrentTopLevelViewModelPtr = CurrentTopLevelViewModels.FindByPredicate(
-				[&RootChildEmitterHandleViewModel](const TSharedRef<FTopLevelViewModel>& TopLevelViewModel) { return TopLevelViewModel->EmitterHandleViewModel == RootChildEmitterHandleViewModel; });
+			TSharedRef<FTopLevelViewModel>* CurrentTopLevelViewModelPtr = CurrentTopLevelViewModels.FindByPredicate([&RootChildEmitterHandleViewModel](const TSharedRef<FTopLevelViewModel>& TopLevelViewModel) 
+			{ 
+				return TopLevelViewModel->EmitterHandleViewModel == RootChildEmitterHandleViewModel &&
+					TopLevelViewModel->RootEntry == RootChildEmitterHandleViewModel->GetEmitterStackViewModel()->GetRootEntry(); 
+			});
 			if (CurrentTopLevelViewModelPtr != nullptr)
 			{
 				TopLevelViewModel = *CurrentTopLevelViewModelPtr;
@@ -669,8 +684,11 @@ void UNiagaraStackViewModel::RefreshTopLevelViewModels()
 		}
 		else
 		{
-			TSharedRef<FTopLevelViewModel>* CurrentTopLevelViewModelPtr = CurrentTopLevelViewModels.FindByPredicate(
-				[RootChild](const TSharedRef<FTopLevelViewModel>& TopLevelViewModel) { return TopLevelViewModel->SystemViewModel == RootChild->GetSystemViewModel(); });
+			TSharedRef<FTopLevelViewModel>* CurrentTopLevelViewModelPtr = CurrentTopLevelViewModels.FindByPredicate([RootChild](const TSharedRef<FTopLevelViewModel>& TopLevelViewModel)
+			{ 
+				return TopLevelViewModel->SystemViewModel == RootChild->GetSystemViewModel() &&
+					TopLevelViewModel->RootEntry == RootChild->GetSystemViewModel()->GetSystemStackViewModel()->GetRootEntry();
+			});
 			if (CurrentTopLevelViewModelPtr != nullptr)
 			{
 				TopLevelViewModel = *CurrentTopLevelViewModelPtr;
diff --git a/Engine/Plugins/FX/Niagara/Source/NiagaraEditor/Private/Widgets/SNiagaraGeneratedCodeView.cpp b/Engine/Plugins/FX/Niagara/Source/NiagaraEditor/Private/Widgets/SNiagaraGeneratedCodeView.cpp
--- a/Engine/Plugins/FX/Niagara/Source/NiagaraEditor/Private/Widgets/SNiagaraGeneratedCodeView.cpp
+++ b/Engine/Plugins/FX/Niagara/Source/NiagaraEditor/Private/Widgets/SNiagaraGeneratedCodeView.cpp
@@ -9,6 +9,7 @@
 #include "ISequencer.h"
 #include "ViewModels/NiagaraSystemViewModel.h"
 #include "ViewModels/NiagaraEmitterHandleViewModel.h"
+#include "ViewModels/NiagaraSystemSelectionViewModel.h"
 #include "NiagaraEmitterHandle.h"
 #include "NiagaraEmitter.h"
 #include "NiagaraScript.h"
@@ -33,7 +34,7 @@ void SNiagaraGeneratedCodeView::Construct(const FArguments& InArgs, TSharedRef<F
 	ensure(ScriptEnum);
 
 	SystemViewModel = InSystemViewModel;
-	SystemViewModel->GetSelectionViewModel()->OnSelectionChanged().AddSP(this, &SNiagaraGeneratedCodeView::SystemSelectionChanged);
+	SystemViewModel->GetSelectionViewModel()->OnEmitterHandleIdSelectionChanged().AddSP(this, &SNiagaraGeneratedCodeView::SystemSelectionChanged);
 	SystemViewModel->GetSystemScriptViewModel()->OnSystemCompiled().AddRaw(this, &SNiagaraGeneratedCodeView::OnCodeCompiled);
 
 	TSharedRef<SWidget> HeaderContentsFirstLine = SNew(SHorizontalBox)
@@ -329,7 +330,7 @@ void SNiagaraGeneratedCodeView::OnCodeCompiled()
 	DoSearch(SearchBox->GetText());
 }
 
-void SNiagaraGeneratedCodeView::SystemSelectionChanged(UNiagaraSystemSelectionViewModel::ESelectionChangeSource SelectionChangeSource)
+void SNiagaraGeneratedCodeView::SystemSelectionChanged()
 {
 	UpdateUI();
 	DoSearch(SearchBox->GetText());
@@ -533,7 +534,7 @@ SNiagaraGeneratedCodeView::~SNiagaraGeneratedCodeView()
 	{
 		if (SystemViewModel->GetSelectionViewModel())
 		{
-			SystemViewModel->GetSelectionViewModel()->OnSelectionChanged().RemoveAll(this);
+			SystemViewModel->GetSelectionViewModel()->OnEmitterHandleIdSelectionChanged().RemoveAll(this);
 		}
 
 		if (SystemViewModel->GetSystemScriptViewModel().IsValid())
diff --git a/Engine/Plugins/FX/Niagara/Source/NiagaraEditor/Private/Widgets/SNiagaraGeneratedCodeView.h b/Engine/Plugins/FX/Niagara/Source/NiagaraEditor/Private/Widgets/SNiagaraGeneratedCodeView.h
--- a/Engine/Plugins/FX/Niagara/Source/NiagaraEditor/Private/Widgets/SNiagaraGeneratedCodeView.h
+++ b/Engine/Plugins/FX/Niagara/Source/NiagaraEditor/Private/Widgets/SNiagaraGeneratedCodeView.h
@@ -9,7 +9,6 @@
 #include "Widgets/Views/STableViewBase.h"
 #include "Widgets/Views/STreeView.h"
 #include "ViewModels/NiagaraSystemViewModel.h"
-#include "ViewModels/NiagaraSystemSelectionViewModel.h"
 #include "TickableEditorObject.h"
 #include "UObject/WeakObjectPtrTemplates.h"
 #include "NiagaraDataSet.h"
@@ -62,7 +61,7 @@ protected:
 
 	void SetSearchMofN();
 
-	void SystemSelectionChanged(UNiagaraSystemSelectionViewModel::ESelectionChangeSource SelectionChangeSource);
+	void SystemSelectionChanged();
 	void OnTabChanged(uint32 Tab);
 	bool GetTabCheckedState(uint32 Tab) const;
 	EVisibility GetViewVisibility(uint32 Tab) const;
diff --git a/Engine/Plugins/FX/Niagara/Source/NiagaraEditor/Private/Widgets/SNiagaraSpreadsheetView.cpp b/Engine/Plugins/FX/Niagara/Source/NiagaraEditor/Private/Widgets/SNiagaraSpreadsheetView.cpp
--- a/Engine/Plugins/FX/Niagara/Source/NiagaraEditor/Private/Widgets/SNiagaraSpreadsheetView.cpp
+++ b/Engine/Plugins/FX/Niagara/Source/NiagaraEditor/Private/Widgets/SNiagaraSpreadsheetView.cpp
@@ -10,6 +10,7 @@
 #include "ISequencer.h"
 #include "ViewModels/NiagaraSystemViewModel.h"
 #include "ViewModels/NiagaraEmitterHandleViewModel.h"
+#include "ViewModels/NiagaraSystemSelectionViewModel.h"
 #include "NiagaraEmitterHandle.h"
 #include "NiagaraEmitter.h"
 #include "NiagaraScript.h"
@@ -224,7 +225,7 @@ void SNiagaraSpreadsheetView::Construct(const FArguments& InArgs, TSharedRef<FNi
 	CaptureData.SetNum(UIMax);
 
 	SystemViewModel = InSystemViewModel;
-	SystemViewModel->GetSelectionViewModel()->OnSelectionChanged().AddSP(this, &SNiagaraSpreadsheetView::SystemSelectionChanged);
+	SystemViewModel->GetSelectionViewModel()->OnEmitterHandleIdSelectionChanged().AddSP(this, &SNiagaraSpreadsheetView::SystemSelectionChanged);
 	SystemViewModel->OnPostSequencerTimeChanged().AddRaw(this, &SNiagaraSpreadsheetView::OnSequencerTimeChanged);
 
 	bInitialColumns = true;
@@ -620,7 +621,7 @@ SNiagaraSpreadsheetView::~SNiagaraSpreadsheetView()
 	{
 		if (SystemViewModel->GetSelectionViewModel() != nullptr)
 		{
-			SystemViewModel->GetSelectionViewModel()->OnSelectionChanged().RemoveAll(this);
+			SystemViewModel->GetSelectionViewModel()->OnEmitterHandleIdSelectionChanged().RemoveAll(this);
 		}
 		SystemViewModel->OnPostSequencerTimeChanged().RemoveAll(this);
 	}
@@ -897,7 +898,7 @@ bool SNiagaraSpreadsheetView::IsOutputAttributeEnabled(EUITab Tab, FName Item)
 	return CaptureData[(int32)Tab].FilteredOutputFields.Find(Item) != INDEX_NONE;
 }
 
-void SNiagaraSpreadsheetView::SystemSelectionChanged(UNiagaraSystemSelectionViewModel::ESelectionChangeSource SelectionChangeSource)
+void SNiagaraSpreadsheetView::SystemSelectionChanged()
 {
 	// Need to reset the attributes list...
 	for (int32 i = 0; i < (int32)UIMax; i++)
diff --git a/Engine/Plugins/FX/Niagara/Source/NiagaraEditor/Private/Widgets/SNiagaraSpreadsheetView.h b/Engine/Plugins/FX/Niagara/Source/NiagaraEditor/Private/Widgets/SNiagaraSpreadsheetView.h
--- a/Engine/Plugins/FX/Niagara/Source/NiagaraEditor/Private/Widgets/SNiagaraSpreadsheetView.h
+++ b/Engine/Plugins/FX/Niagara/Source/NiagaraEditor/Private/Widgets/SNiagaraSpreadsheetView.h
@@ -10,7 +10,6 @@
 #include "Widgets/Views/STreeView.h"
 #include "Widgets/Input/SComboBox.h"
 #include "ViewModels/NiagaraSystemViewModel.h"
-#include "ViewModels/NiagaraSystemSelectionViewModel.h"
 #include "TickableEditorObject.h"
 #include "UObject/WeakObjectPtrTemplates.h"
 #include "NiagaraDataSet.h"
@@ -133,7 +132,7 @@ protected:
 	bool CanCapture() const;
 	bool IsPausedAtRightTimeOnRightHandle() const;
 
-	void SystemSelectionChanged(UNiagaraSystemSelectionViewModel::ESelectionChangeSource SelectionChangeSource);
+	void SystemSelectionChanged();
 
 	void ResetColumns(EUITab Tab);
 	void ResetEntries(EUITab Tab);
diff --git a/Engine/Plugins/FX/Niagara/Source/NiagaraEditor/Public/ViewModels/NiagaraOverviewGraphViewModel.h b/Engine/Plugins/FX/Niagara/Source/NiagaraEditor/Public/ViewModels/NiagaraOverviewGraphViewModel.h
--- a/Engine/Plugins/FX/Niagara/Source/NiagaraEditor/Public/ViewModels/NiagaraOverviewGraphViewModel.h
+++ b/Engine/Plugins/FX/Niagara/Source/NiagaraEditor/Public/ViewModels/NiagaraOverviewGraphViewModel.h
@@ -7,7 +7,6 @@
 #include "EditorUndoClient.h"
 #include "Styling/SlateColor.h"
 #include "Framework/Commands/UICommandList.h"
-#include "ViewModels/NiagaraSystemSelectionViewModel.h"
 
 class UNiagaraGraph;
 class FNiagaraSystemViewModel;
@@ -68,7 +67,7 @@ private:
 
 	void GraphSelectionChanged();
 
-	void SystemSelectionChanged(UNiagaraSystemSelectionViewModel::ESelectionChangeSource SelectionChangeSource);
+	void SystemSelectionChanged();
 
 private:
 
diff --git a/Engine/Plugins/FX/Niagara/Source/NiagaraEditor/Public/ViewModels/NiagaraSystemSelectionViewModel.h b/Engine/Plugins/FX/Niagara/Source/NiagaraEditor/Public/ViewModels/NiagaraSystemSelectionViewModel.h
--- a/Engine/Plugins/FX/Niagara/Source/NiagaraEditor/Public/ViewModels/NiagaraSystemSelectionViewModel.h
+++ b/Engine/Plugins/FX/Niagara/Source/NiagaraEditor/Public/ViewModels/NiagaraSystemSelectionViewModel.h
@@ -15,15 +15,17 @@ UCLASS()
 class NIAGARAEDITOR_API UNiagaraSystemSelectionViewModel : public UObject
 {
 public:
-	enum class ESelectionChangeSource
+	DECLARE_MULTICAST_DELEGATE(FOnSelectionChanged);
+
+private:
+	struct FSelectionEntry
 	{
-		EntrySelection,
-		TopObjectLevelSelection,
-		Refresh,
-		Clear
-	};
+		explicit FSelectionEntry(UNiagaraStackEntry* SelectedEntry);
 
-	DECLARE_MULTICAST_DELEGATE_OneParam(FOnSelectionChanged, ESelectionChangeSource);
+		const TWeakObjectPtr<UNiagaraStackEntry> Entry;
+		const FGuid EmitterGuid;
+		const FString EditorDataKey;
+	};
 
 public:
 	GENERATED_BODY()
@@ -32,15 +34,15 @@ public:
 
 	void Finalize();
 
-	const TArray<UNiagaraStackEntry*> GetSelectedEntries() const;
+	bool ContainsEntry(UNiagaraStackEntry* StackEntry) const;
+
+	void GetSelectedEntries(TArray<UNiagaraStackEntry*>& OutSelectedEntries) const;
 
 	bool GetSystemIsSelected() const;
 
 	const TArray<FGuid>& GetSelectedEmitterHandleIds() const;
 
-	void UpdateSelectionFromEntries(const TArray<UNiagaraStackEntry*> InSelectedEntries, const TArray<UNiagaraStackEntry*> InDeselectedEntries, bool bClearCurrentSelection);
-
-	void UpdateSelectionFromTopLevelObjects(bool bInSystemIsSelected, const TArray<FGuid> InSelectedEmitterHandleIds, bool bClearCurrentSelection);
+	void UpdateSelectedEntries(const TArray<UNiagaraStackEntry*> InSelectedEntries, const TArray<UNiagaraStackEntry*> InDeselectedEntries, bool bClearCurrentSelection);
 
 	UNiagaraStackViewModel* GetSelectionStackViewModel();
 
@@ -50,11 +52,15 @@ public:
 
 	void AddEntryToSelectionByDisplayedObjectDeferred(const UObject* InObjects);
 
-	void Refresh();
+	bool Refresh();
 
 	void RefreshDeferred();
 
-	FOnSelectionChanged& OnSelectionChanged();
+	FOnSelectionChanged& OnEntrySelectionChanged();
+
+	FOnSelectionChanged& OnEmitterHandleIdSelectionChanged();
+
+	FOnSelectionChanged& OnSystemIsSelectedChanged();
 
 	void Tick();
 
@@ -63,12 +69,12 @@ private:
 
 	void RemoveEntryFromSelectionInternal(UNiagaraStackEntry* DeselectedEntry);
 
-	void UpdateStackSelectionEntry();
+	void UpdateExternalSelectionState();
 
 private:
 	TWeakPtr<FNiagaraSystemViewModel> SystemViewModelWeak;
 
-	TArray<UNiagaraStackEntry*> SelectedEntries;
+	TArray<FSelectionEntry> SelectionEntries;
 
 	bool bSystemIsSelected;
 
@@ -80,7 +86,11 @@ private:
 	UPROPERTY()
 	UNiagaraStackViewModel* SelectionStackViewModel;
 
-	FOnSelectionChanged OnSelectionChangedDelegate;
+	FOnSelectionChanged OnEntrySelectionChangedDelegate;
+
+	FOnSelectionChanged OnEmitterHandleIdSelectionChangedDelegate;
+
+	FOnSelectionChanged OnSystemIsSelectedChangedDelegate;
 
 	TArray<FObjectKey> DeferredDisplayedObjectKeysToAddToSelection;
 
diff --git a/Engine/Plugins/FX/Niagara/Source/NiagaraEditor/Public/ViewModels/NiagaraSystemViewModel.h b/Engine/Plugins/FX/Niagara/Source/NiagaraEditor/Public/ViewModels/NiagaraSystemViewModel.h
--- a/Engine/Plugins/FX/Niagara/Source/NiagaraEditor/Public/ViewModels/NiagaraSystemViewModel.h
+++ b/Engine/Plugins/FX/Niagara/Source/NiagaraEditor/Public/ViewModels/NiagaraSystemViewModel.h
@@ -8,7 +8,6 @@
 #include "UObject/GCObject.h"
 #include "NiagaraCurveOwner.h"
 #include "ViewModels/TNiagaraViewModelManager.h"
-#include "ViewModels/NiagaraSystemSelectionViewModel.h"
 #include "ISequencer.h"
 
 #include "TickableEditorObject.h"
@@ -374,7 +373,7 @@ private:
 	void SequencerTimeChanged();
 
 	/** Called whenever the current selection in the system changes. */
-	void SystemSelectionChanged(UNiagaraSystemSelectionViewModel::ESelectionChangeSource SelectionChangeSource);
+	void SystemSelectionChanged();
 
 	/** Called whenever the track selection in sequencer changes. */
 	void SequencerTrackSelectionChanged(TArray<UMovieSceneTrack*> SelectedTracks);
diff --git a/Engine/Plugins/FX/Niagara/Source/NiagaraEditor/Public/ViewModels/Stack/NiagaraStackViewModel.h b/Engine/Plugins/FX/Niagara/Source/NiagaraEditor/Public/ViewModels/Stack/NiagaraStackViewModel.h
--- a/Engine/Plugins/FX/Niagara/Source/NiagaraEditor/Public/ViewModels/Stack/NiagaraStackViewModel.h
+++ b/Engine/Plugins/FX/Niagara/Source/NiagaraEditor/Public/ViewModels/Stack/NiagaraStackViewModel.h
@@ -58,15 +58,9 @@ public:
 
 	struct NIAGARAEDITOR_API FTopLevelViewModel
 	{
-		FTopLevelViewModel(TSharedPtr<FNiagaraSystemViewModel> InSystemViewModel)
-			: SystemViewModel(InSystemViewModel)
-		{
-		}
+		FTopLevelViewModel(TSharedPtr<FNiagaraSystemViewModel> InSystemViewModel);
 
-		FTopLevelViewModel(TSharedPtr<FNiagaraEmitterHandleViewModel> InEmitterHandleViewModel)
-			: EmitterHandleViewModel(InEmitterHandleViewModel)
-		{
-		}
+		FTopLevelViewModel(TSharedPtr<FNiagaraEmitterHandleViewModel> InEmitterHandleViewModel);
 
 		bool IsValid() const;
 
@@ -78,6 +72,7 @@ public:
 
 		const TSharedPtr<FNiagaraSystemViewModel> SystemViewModel;
 		const TSharedPtr<FNiagaraEmitterHandleViewModel> EmitterHandleViewModel;
+		const TWeakObjectPtr<UNiagaraStackEntry> RootEntry;
 	};
 
 public:
@@ -88,8 +83,6 @@ public:
 
 	virtual void BeginDestroy() override;
 
-	TSharedPtr<FNiagaraSystemViewModel> GetSystemViewModel();
-
 	UNiagaraStackEntry* GetRootEntry();
 
 	TArray<UNiagaraStackEntry*>& GetRootEntryAsArray();
diff --git a/Engine/Plugins/FX/Niagara/Source/NiagaraEditorWidgets/Private/SNiagaraOverviewStack.cpp b/Engine/Plugins/FX/Niagara/Source/NiagaraEditorWidgets/Private/SNiagaraOverviewStack.cpp
--- a/Engine/Plugins/FX/Niagara/Source/NiagaraEditorWidgets/Private/SNiagaraOverviewStack.cpp
+++ b/Engine/Plugins/FX/Niagara/Source/NiagaraEditorWidgets/Private/SNiagaraOverviewStack.cpp
@@ -4,6 +4,7 @@
 #include "NiagaraSystemEditorData.h"
 #include "ViewModels/NiagaraSystemViewModel.h"
 #include "ViewModels/NiagaraEmitterHandleViewModel.h"
+#include "ViewModels/NiagaraSystemSelectionViewModel.h"
 #include "ViewModels/Stack/NiagaraStackEntry.h"
 #include "ViewModels/Stack/NiagaraStackItemGroup.h"
 #include "ViewModels/Stack/NiagaraStackItem.h"
@@ -223,7 +224,7 @@ void SNiagaraOverviewStack::Construct(const FArguments& InArgs, UNiagaraStackVie
 
 	StackViewModel = &InStackViewModel;
 	OverviewSelectionViewModel = &InOverviewSelectionViewModel;
-	OverviewSelectionViewModel->OnSelectionChanged().AddSP(this, &SNiagaraOverviewStack::SystemSelectionChanged);
+	OverviewSelectionViewModel->OnEntrySelectionChanged().AddSP(this, &SNiagaraOverviewStack::SystemSelectionChanged);
 
 	ChildSlot
 	[
@@ -239,6 +240,7 @@ void SNiagaraOverviewStack::Construct(const FArguments& InArgs, UNiagaraStackVie
 		
 	bRefreshEntryListPending = true;
 	RefreshEntryList();
+	SystemSelectionChanged();
 }
 
 SNiagaraOverviewStack::~SNiagaraOverviewStack()
@@ -247,6 +249,11 @@ SNiagaraOverviewStack::~SNiagaraOverviewStack()
 	{
 		StackViewModel->OnStructureChanged().RemoveAll(this);
 	}
+
+	if (OverviewSelectionViewModel != nullptr)
+	{
+		OverviewSelectionViewModel->OnEntrySelectionChanged().RemoveAll(this);
+	}
 }
 
 bool SNiagaraOverviewStack::SupportsKeyboardFocus() const
@@ -278,7 +285,9 @@ void SNiagaraOverviewStack::DeleteSelectedEntries()
 {
 	bool bIncompleteDelete = false;
 	TArray<UNiagaraStackItem*> ItemsToDelete;
-	for (UNiagaraStackEntry* SelectedEntry : OverviewSelectionViewModel->GetSelectedEntries())
+	TArray<UNiagaraStackEntry*> SelectedEntries;
+	OverviewSelectionViewModel->GetSelectedEntries(SelectedEntries);
+	for (UNiagaraStackEntry* SelectedEntry : SelectedEntries)
 	{
 		UNiagaraStackItem* SelectedItem = Cast<UNiagaraStackItem>(SelectedEntry);
 		if (SelectedItem != nullptr)
@@ -475,7 +484,7 @@ void SNiagaraOverviewStack::OnSelectionChanged(UNiagaraStackEntry* InNewSelectio
 		}
 
 		bool bClearCurrentSelection = FSlateApplication::Get().GetModifierKeys().IsControlDown() == false;
-		OverviewSelectionViewModel->UpdateSelectionFromEntries(SelectedEntries, DeselectedEntries, bClearCurrentSelection);
+		OverviewSelectionViewModel->UpdateSelectedEntries(SelectedEntries, DeselectedEntries, bClearCurrentSelection);
 
 		PreviousSelection.Empty();
 		for (UNiagaraStackEntry* SelectedEntry : SelectedEntries)
@@ -485,7 +494,7 @@ void SNiagaraOverviewStack::OnSelectionChanged(UNiagaraStackEntry* InNewSelectio
 	}
 }
 
-void SNiagaraOverviewStack::SystemSelectionChanged(UNiagaraSystemSelectionViewModel::ESelectionChangeSource SelectionChangeSource)
+void SNiagaraOverviewStack::SystemSelectionChanged()
 {
 	if (bUpdatingOverviewSelectionFromStackSelection == false)
 	{
@@ -493,7 +502,8 @@ void SNiagaraOverviewStack::SystemSelectionChanged(UNiagaraSystemSelectionViewMo
 
 		TArray<UNiagaraStackEntry*> SelectedListViewStackEntries;
 		EntryListView->GetSelectedItems(SelectedListViewStackEntries);
-		TArray<UNiagaraStackEntry*> SelectedOverviewEntries = OverviewSelectionViewModel->GetSelectedEntries();
+		TArray<UNiagaraStackEntry*> SelectedOverviewEntries;
+		OverviewSelectionViewModel->GetSelectedEntries(SelectedOverviewEntries);
 
 		TArray<UNiagaraStackEntry*> EntriesToDeselect;
 		for (UNiagaraStackEntry* SelectedListViewStackEntry : SelectedListViewStackEntries)
@@ -531,7 +541,8 @@ FReply SNiagaraOverviewStack::OnRowDragDetected(const FGeometry& InGeometry, con
 	UNiagaraStackEntry* StackEntry = InStackEntryWeak.Get();
 	if (StackEntry != nullptr && StackEntry->CanDrag())
 	{
-		TArray<UNiagaraStackEntry*> EntriesToDrag = StackEntry->GetSystemViewModel()->GetSelectionViewModel()->GetSelectedEntries();
+		TArray<UNiagaraStackEntry*> EntriesToDrag;
+		StackEntry->GetSystemViewModel()->GetSelectionViewModel()->GetSelectedEntries(EntriesToDrag);
 		EntriesToDrag.AddUnique(StackEntry);
 		return FReply::Handled().BeginDragDrop(FNiagaraStackEditorWidgetsUtilities::ConstructDragDropOperationForStackEntries(EntriesToDrag));
 	}
diff --git a/Engine/Plugins/FX/Niagara/Source/NiagaraEditorWidgets/Private/SNiagaraOverviewStack.h b/Engine/Plugins/FX/Niagara/Source/NiagaraEditorWidgets/Private/SNiagaraOverviewStack.h
--- a/Engine/Plugins/FX/Niagara/Source/NiagaraEditorWidgets/Private/SNiagaraOverviewStack.h
+++ b/Engine/Plugins/FX/Niagara/Source/NiagaraEditorWidgets/Private/SNiagaraOverviewStack.h
@@ -3,7 +3,6 @@
 #pragma once
 
 #include "ViewModels/Stack/NiagaraStackEntry.h"
-#include "ViewModels/NiagaraSystemSelectionViewModel.h"
 #include "Widgets/SCompoundWidget.h"
 #include "Widgets/Views/SListView.h"
 #include "UObject/ObjectKey.h"
@@ -48,7 +47,7 @@ private:
 
 	void OnSelectionChanged(UNiagaraStackEntry* InNewSelection, ESelectInfo::Type SelectInfo);
 
-	void SystemSelectionChanged(UNiagaraSystemSelectionViewModel::ESelectionChangeSource SelectionChangeSource);
+	void SystemSelectionChanged();
 
 	FReply OnRowDragDetected(const FGeometry& InGeometry, const FPointerEvent& InPointerEvent, TWeakObjectPtr<UNiagaraStackEntry> InStackEntryWeak);
 
diff --git a/Engine/Plugins/FX/Niagara/Source/NiagaraEditorWidgets/Private/SNiagaraStack.cpp b/Engine/Plugins/FX/Niagara/Source/NiagaraEditorWidgets/Private/SNiagaraStack.cpp
--- a/Engine/Plugins/FX/Niagara/Source/NiagaraEditorWidgets/Private/SNiagaraStack.cpp
+++ b/Engine/Plugins/FX/Niagara/Source/NiagaraEditorWidgets/Private/SNiagaraStack.cpp
@@ -65,7 +65,7 @@ public:
 		SLATE_ATTRIBUTE(EVisibility, IssueIconVisibility);
 	SLATE_END_ARGS();
 
-	void Construct(const FArguments& InArgs, TSharedRef<FNiagaraEmitterHandleViewModel> InEmitterHandleViewModel)
+	void Construct(const FArguments& InArgs, TSharedRef<FNiagaraEmitterHandleViewModel> InEmitterHandleViewModel, UNiagaraStackEntry* InRootEntry)
 	{
 		EmitterHandleViewModel = InEmitterHandleViewModel;
 
@@ -117,7 +117,7 @@ public:
 						.VAlign(VAlign_Center)
 						.Padding(4, 0, 0, 0)
 						[
-							SNew(SNiagaraStackIssueIcon, EmitterHandleViewModel->GetEmitterStackViewModel(), EmitterHandleViewModel->GetEmitterStackViewModel()->GetRootEntry())
+							SNew(SNiagaraStackIssueIcon, EmitterHandleViewModel->GetEmitterStackViewModel(), InRootEntry)
 							.Visibility(InArgs._IssueIconVisibility)
 						]
 					]
@@ -639,13 +639,13 @@ TSharedRef<ITableRow> SNiagaraStack::OnGenerateRowForTopLevelObject(TSharedRef<U
 			.VAlign(VAlign_Center)
 			.Padding(4, 0, 2, 0)
 			[
-				SNew(SNiagaraStackIssueIcon, Item->SystemViewModel->GetSystemStackViewModel(), Item->SystemViewModel->GetSystemStackViewModel()->GetRootEntry())
+				SNew(SNiagaraStackIssueIcon, Item->SystemViewModel->GetSystemStackViewModel(), Item->RootEntry.Get())
 				.Visibility(this, &SNiagaraStack::GetIssueIconVisibility)
 			];
 	}
 	else if(Item->EmitterHandleViewModel.IsValid())
 	{
-		Content = SNew(SNiagaraStackEmitterHeader, Item->EmitterHandleViewModel.ToSharedRef())
+		Content = SNew(SNiagaraStackEmitterHeader, Item->EmitterHandleViewModel.ToSharedRef(), Item->RootEntry.Get())
 			.IssueIconVisibility(this, &SNiagaraStack::GetIssueIconVisibility);
 	}
 
