Fix shadowing causing Android build to fail to compile
#rnx

diff --git a/Engine/Plugins/FX/Niagara/Source/Niagara/Private/NiagaraDataInterfaceGrid2DCollection.cpp b/Engine/Plugins/FX/Niagara/Source/Niagara/Private/NiagaraDataInterfaceGrid2DCollection.cpp
--- a/Engine/Plugins/FX/Niagara/Source/Niagara/Private/NiagaraDataInterfaceGrid2DCollection.cpp
+++ b/Engine/Plugins/FX/Niagara/Source/Niagara/Private/NiagaraDataInterfaceGrid2DCollection.cpp
@@ -491,16 +491,18 @@ void UNiagaraDataInterfaceGrid2DCollection::FillRawTexture2D(const UNiagaraCompo
 	}
 	FNiagaraSystemInstanceID InstanceID = SystemInstance->GetId();
 
-	Grid2DCollectionRWInstanceData* Grid2DInstanceData = TProxy->SystemInstancesToProxyData.Find(InstanceID);
-	if (!Grid2DInstanceData)
 	{
-		TilesX = -1;
-		TilesY = -1;
-		return;
-	}
+		Grid2DCollectionRWInstanceData* Grid2DInstanceData = TProxy->SystemInstancesToProxyData.Find(InstanceID);
+		if (!Grid2DInstanceData)
+		{
+			TilesX = -1;
+			TilesY = -1;
+			return;
+		}
 
-	TilesX = Grid2DInstanceData->NumTilesX;
-	TilesY = Grid2DInstanceData->NumTilesY;
+		TilesX = Grid2DInstanceData->NumTilesX;
+		TilesY = Grid2DInstanceData->NumTilesY;
+	}
 
 	ENQUEUE_RENDER_COMMAND(FUpdateDIColorCurve)(
 		[Dest, TProxy, InstanceID](FRHICommandListImmediate& RHICmdList)
