Niagara - Fix module dependency issues.
+ Fix an issue which was preventing system module dependencies from being evaluated.
+ Fix a bug where when trying to fix missing dependencies, modules would be added to the wrong script.

diff --git a/Engine/Plugins/FX/Niagara/Source/NiagaraEditor/Private/ViewModels/NiagaraSystemViewModel.cpp b/Engine/Plugins/FX/Niagara/Source/NiagaraEditor/Private/ViewModels/NiagaraSystemViewModel.cpp
--- a/Engine/Plugins/FX/Niagara/Source/NiagaraEditor/Private/ViewModels/NiagaraSystemViewModel.cpp
+++ b/Engine/Plugins/FX/Niagara/Source/NiagaraEditor/Private/ViewModels/NiagaraSystemViewModel.cpp
@@ -12,6 +12,7 @@
 #include "ViewModels/NiagaraSystemSelectionViewModel.h"
 #include "ViewModels/Stack/NiagaraStackViewModel.h"
 #include "ViewModels/Stack/NiagaraStackGraphUtilities.h"
+#include "ViewModels/Stack/NiagaraStackEntry.h"
 #include "NiagaraSystemScriptViewModel.h"
 #include "NiagaraScriptGraphViewModel.h"
 #include "NiagaraSequence.h"
@@ -584,30 +585,41 @@ void FNiagaraSystemViewModel::SetToolkitCommands(const TSharedRef<FUICommandList
 	ToolkitCommands = InToolkitCommands;
 }
 
-const TArray<FNiagaraStackModuleData>& FNiagaraSystemViewModel::GetStackModuleDataForEmitter(TSharedRef<FNiagaraEmitterViewModel> EmitterViewModel)
+const TArray<FNiagaraStackModuleData>& FNiagaraSystemViewModel::GetStackModuleData(UNiagaraStackEntry* ModuleEntry)
 {
-	FGuid EmitterHandleId;
-	TSharedRef<FNiagaraEmitterHandleViewModel>* FoundModel = EmitterHandleViewModels.FindByPredicate([EmitterViewModel](TSharedRef<FNiagaraEmitterHandleViewModel> CurrentViewModel)
-	{ return CurrentViewModel->GetEmitterViewModel() == EmitterViewModel; });
-	checkf(FoundModel != nullptr, TEXT("Couldn't get stack module data for emitter"));
-	if (FoundModel)
+	FGuid StackDataId;
+	UNiagaraEmitter* Emitter = nullptr;
+	if (ModuleEntry->GetEmitterViewModel().IsValid())
 	{
-		EmitterHandleId = (*FoundModel)->GetEmitterHandle()->GetId();
-		if (!EmitterToCachedStackModuleData.Contains(EmitterHandleId))
+		TSharedRef<FNiagaraEmitterHandleViewModel>* FoundModel = EmitterHandleViewModels.FindByPredicate([ModuleEntry](TSharedRef<FNiagaraEmitterHandleViewModel> CurrentViewModel)
 		{
-			// If not cached, rebuild
-			UNiagaraEmitter* Emitter = (*FoundModel)->GetEmitterViewModel()->GetEmitter();
-			TArray<FNiagaraStackModuleData> StackModuleData;
-			BuildStackModuleData(GetSystem().GetSystemSpawnScript(), EmitterHandleId, StackModuleData);
-			BuildStackModuleData(GetSystem().GetSystemUpdateScript(), EmitterHandleId, StackModuleData);
-			BuildStackModuleData(Emitter->EmitterSpawnScriptProps.Script, EmitterHandleId, StackModuleData);
-			BuildStackModuleData(Emitter->EmitterUpdateScriptProps.Script, EmitterHandleId, StackModuleData);
-			BuildStackModuleData(Emitter->SpawnScriptProps.Script, EmitterHandleId, StackModuleData);
-			BuildStackModuleData(Emitter->UpdateScriptProps.Script, EmitterHandleId, StackModuleData);
-			EmitterToCachedStackModuleData.Add(EmitterHandleId) = StackModuleData;
+			return CurrentViewModel->GetEmitterViewModel() == ModuleEntry->GetEmitterViewModel(); 
+		});
+		checkf(FoundModel != nullptr, TEXT("Couldn't get stack module data for emitter"));
+		StackDataId = (*FoundModel)->GetEmitterHandle()->GetId();
+		Emitter = ModuleEntry->GetEmitterViewModel()->GetEmitter();
+	}
+	else
+	{
+		StackDataId = FGuid();
+	}
+
+	TArray<FNiagaraStackModuleData>* StackModuleData = GuidToCachedStackModuleData.Find(StackDataId);
+	if (StackModuleData == nullptr)
+	{
+		// If not cached, rebuild
+		StackModuleData = &GuidToCachedStackModuleData.Add(StackDataId);
+		BuildStackModuleData(GetSystem().GetSystemSpawnScript(), StackDataId, *StackModuleData);
+		BuildStackModuleData(GetSystem().GetSystemUpdateScript(), StackDataId, *StackModuleData);
+		if (StackDataId.IsValid() && Emitter != nullptr)
+		{
+			BuildStackModuleData(Emitter->EmitterSpawnScriptProps.Script, StackDataId, *StackModuleData);
+			BuildStackModuleData(Emitter->EmitterUpdateScriptProps.Script, StackDataId, *StackModuleData);
+			BuildStackModuleData(Emitter->SpawnScriptProps.Script, StackDataId, *StackModuleData);
+			BuildStackModuleData(Emitter->UpdateScriptProps.Script, StackDataId, *StackModuleData);
 		}
 	}
-	return EmitterToCachedStackModuleData[EmitterHandleId];
+	return *StackModuleData;
 }
 
 UNiagaraStackViewModel* FNiagaraSystemViewModel::GetSystemStackViewModel()
@@ -797,7 +809,7 @@ void FNiagaraSystemViewModel::RefreshEmitterHandleViewModels()
 {
 	TArray<TSharedRef<FNiagaraEmitterHandleViewModel>> OldViewModels = EmitterHandleViewModels;
 	EmitterHandleViewModels.Empty();
-	EmitterToCachedStackModuleData.Empty();
+	GuidToCachedStackModuleData.Empty();
 
 	// Map existing view models to the real instances that now exist. Reuse if we can. Create a new one if we cannot.
 	TArray<FGuid> ValidEmitterHandleIds;
@@ -1292,12 +1304,12 @@ void FNiagaraSystemViewModel::EmitterScriptGraphChanged(const FEdGraphEditAction
 		EmitterIdsRequiringSequencerTrackUpdate.AddUnique(OwningEmitterHandleId);
 	}
 	// Remove from cache on graph change 
-	EmitterToCachedStackModuleData.Remove(OwningEmitterHandleId);
+	GuidToCachedStackModuleData.Remove(OwningEmitterHandleId);
 }
 
 void FNiagaraSystemViewModel::SystemScriptGraphChanged(const FEdGraphEditAction& InAction)
 {
-	EmitterToCachedStackModuleData.Empty();
+	GuidToCachedStackModuleData.Empty();
 }
 
 void FNiagaraSystemViewModel::EmitterParameterStoreChanged(const FNiagaraParameterStore& ChangedParameterStore, const UNiagaraScript& OwningScript)
diff --git a/Engine/Plugins/FX/Niagara/Source/NiagaraEditor/Private/ViewModels/Stack/NiagaraStackModuleItem.cpp b/Engine/Plugins/FX/Niagara/Source/NiagaraEditor/Private/ViewModels/Stack/NiagaraStackModuleItem.cpp
--- a/Engine/Plugins/FX/Niagara/Source/NiagaraEditor/Private/ViewModels/Stack/NiagaraStackModuleItem.cpp
+++ b/Engine/Plugins/FX/Niagara/Source/NiagaraEditor/Private/ViewModels/Stack/NiagaraStackModuleItem.cpp
@@ -431,14 +431,9 @@ void UNiagaraStackModuleItem::RefreshIssues(TArray<FStackIssue>& NewIssues)
 	// Generate dependency errors with their fixes
 	TArray<UNiagaraNodeFunctionCall*> FoundCalls;
 	TArray<FNiagaraModuleDependency> DependenciesNeeded;
-
-
-	TArray<FNiagaraStackModuleData> SystemModuleData;
-	if (GetEmitterViewModel().IsValid())
-	{
-		SystemModuleData = GetSystemViewModel()->GetStackModuleDataForEmitter(GetEmitterViewModel().ToSharedRef());
-	}
+	TArray<FNiagaraStackModuleData> SystemModuleData = GetSystemViewModel()->GetStackModuleData(this);
 	int32 ModuleIndex = INDEX_NONE;
+
 	for (int i = 0; i < SystemModuleData.Num(); i++)
 	{
 		auto ModuleData = SystemModuleData[i];
@@ -612,7 +607,9 @@ void UNiagaraStackModuleItem::RefreshIssues(TArray<FStackIssue>& NewIssues)
 								{
 									return RequiredDependency.Id == Dependency.Id;
 								});
-								if (bRequiredDependencyFound) // check for multiple dependents along the way, and stop adjacent to the last one
+								if (bRequiredDependencyFound && 
+									(Dependency.ScriptConstraint != ENiagaraModuleDependencyScriptConstraint::SameScript ||
+									SystemModuleData[ModuleIndex].Usage == SystemModuleData[i].Usage)) // check for multiple dependents along the way, and stop adjacent to the last one
 								{
 									ENiagaraScriptUsage DependencyUsage = SystemModuleData[i].Usage;
 									const FNiagaraEmitterHandle* EmitterHandle = GetEmitterViewModel().IsValid()
diff --git a/Engine/Plugins/FX/Niagara/Source/NiagaraEditor/Public/ViewModels/NiagaraSystemViewModel.h b/Engine/Plugins/FX/Niagara/Source/NiagaraEditor/Public/ViewModels/NiagaraSystemViewModel.h
--- a/Engine/Plugins/FX/Niagara/Source/NiagaraEditor/Public/ViewModels/NiagaraSystemViewModel.h
+++ b/Engine/Plugins/FX/Niagara/Source/NiagaraEditor/Public/ViewModels/NiagaraSystemViewModel.h
@@ -23,6 +23,7 @@ struct FNiagaraEmitterHandle;
 class FNiagaraEmitterHandleViewModel;
 class FNiagaraSystemScriptViewModel;
 class UNiagaraStackViewModel;
+class UNiagaraStackEntry;
 class UNiagaraSystemSelectionViewModel;
 class FNiagaraSystemInstance;
 class ISequencer;
@@ -278,8 +279,8 @@ public:
 	/** Set the system toolkit command list. */
 	void SetToolkitCommands(const TSharedRef<FUICommandList>& InToolkitCommands);
 
-	/** Gets the stack module data for the provided emitter, for use in module dependencies. */
-	const TArray<FNiagaraStackModuleData>& GetStackModuleDataForEmitter(TSharedRef<FNiagaraEmitterViewModel> EmitterViewModel);
+	/** Gets the stack module data for the provided module, for use in determining dependencies. */
+	const TArray<FNiagaraStackModuleData>& GetStackModuleData(UNiagaraStackEntry* ModuleEntry);
 
 	/** Gets the ViewModel for the system overview graph. */
 	NIAGARAEDITOR_API TSharedPtr<FNiagaraOverviewGraphViewModel> GetOverviewGraphViewModel() const;
@@ -525,8 +526,8 @@ private:
 	/** A flag which indicates that a compile has been requested, but has not completed. */
 	bool bCompilePendingCompletion;
 
-	/** The cache of stack module data for each emitter */
-	TMap<FGuid, TArray<FNiagaraStackModuleData>> EmitterToCachedStackModuleData;
+	/** The cache of stack module data for each emitter and for the system */
+	TMap<FGuid, TArray<FNiagaraStackModuleData>> GuidToCachedStackModuleData;
 	
 	/** A handle to the on graph changed delegate for the system script. */
 	FDelegateHandle SystemScriptGraphChangedHandler;
