Niagara - Zoom to fit curves on open in the stack view.

diff --git a/Engine/Plugins/FX/Niagara/Source/NiagaraEditorWidgets/Private/DetailCustomizations/NiagaraDataInterfaceCurveDetails.cpp b/Engine/Plugins/FX/Niagara/Source/NiagaraEditorWidgets/Private/DetailCustomizations/NiagaraDataInterfaceCurveDetails.cpp
--- a/Engine/Plugins/FX/Niagara/Source/NiagaraEditorWidgets/Private/DetailCustomizations/NiagaraDataInterfaceCurveDetails.cpp
+++ b/Engine/Plugins/FX/Niagara/Source/NiagaraEditorWidgets/Private/DetailCustomizations/NiagaraDataInterfaceCurveDetails.cpp
@@ -224,17 +224,9 @@ public:
 			}
 		}
 
-		ViewMinInput = 0;
-		ViewMaxOutput = 1;
-
-		for (const FRichCurveEditInfo& CurveEditInfo : CurveOwner->GetCurves())
+		if (StackCurveEditorOptions->GetNeedsInitializeView())
 		{
-			FRealCurve* Curve = CurveEditInfo.CurveToEdit;
-			if (Curve->GetNumKeys())
-			{
-				ViewMinInput = FMath::Min(ViewMinInput, Curve->GetKeyTime(Curve->GetFirstKeyHandle()));
-				ViewMaxOutput = FMath::Max(ViewMaxOutput, Curve->GetKeyTime(Curve->GetLastKeyHandle()));
-			}
+			InitializeView();
 		}
 
 		ChildSlot
@@ -260,6 +252,49 @@ public:
 	}
 
 private:
+	void InitializeView()
+	{
+		bool bHasKeys = false;
+		float ViewMinInput = TNumericLimits<float>::Max();
+		float ViewMaxInput = TNumericLimits<float>::Lowest();
+		float ViewMinOutput = TNumericLimits<float>::Max();
+		float ViewMaxOutput = TNumericLimits<float>::Lowest();
+
+		for (const FRichCurveEditInfo& CurveEditInfo : CurveOwner->GetCurves())
+		{
+			FRealCurve* Curve = CurveEditInfo.CurveToEdit;
+			for(auto KeyIterator = Curve->GetKeyHandleIterator(); KeyIterator; ++KeyIterator)
+			{
+				float KeyTime = Curve->GetKeyTime(*KeyIterator);
+				float KeyValue = Curve->GetKeyValue(*KeyIterator);
+				ViewMinInput = FMath::Min(ViewMinInput, KeyTime);
+				ViewMaxInput = FMath::Max(ViewMaxInput, KeyTime);
+				ViewMinOutput = FMath::Min(ViewMinOutput, KeyValue);
+				ViewMaxOutput = FMath::Max(ViewMaxOutput, KeyValue);
+				bHasKeys = true;
+			}
+		}
+
+		if (bHasKeys == false)
+		{
+			ViewMinInput = 0;
+			ViewMaxInput = 1;
+			ViewMinOutput = 0;
+			ViewMaxOutput = 1;
+		}
+
+		float ViewInputRange = ViewMaxInput - ViewMinInput;
+		float ViewOutputRange = ViewMaxOutput - ViewMinOutput;
+		float ViewInputPadding = ViewInputRange * .05f;
+		float ViewOutputPadding = ViewOutputRange * .05f;
+
+		StackCurveEditorOptions->InitializeView(
+			ViewMinInput - ViewInputPadding,
+			ViewMaxInput + ViewInputPadding,
+			ViewMinOutput - ViewOutputPadding,
+			ViewMaxOutput + ViewOutputPadding);
+	}
+
 	void CurveChanged(FRichCurve* ChangedCurve, UObject* CurveOwnerObject)
 	{
 		UNiagaraDataInterfaceCurveBase* EditedCurve = Cast<UNiagaraDataInterfaceCurveBase>(CurveOwnerObject);
@@ -275,8 +310,6 @@ private:
 	}
 
 private:
-	float ViewMinInput;
-	float ViewMaxOutput;
 	TArray<TSharedRef<IPropertyHandle>> CurveProperties;
 	TSharedPtr<FNiagaraStackCurveEditorOptions> StackCurveEditorOptions;
 	TSharedPtr<FNiagaraCurveOwner> CurveOwner;
diff --git a/Engine/Plugins/FX/Niagara/Source/NiagaraEditorWidgets/Private/NiagaraEditorWidgetsModule.cpp b/Engine/Plugins/FX/Niagara/Source/NiagaraEditorWidgets/Private/NiagaraEditorWidgetsModule.cpp
--- a/Engine/Plugins/FX/Niagara/Source/NiagaraEditorWidgets/Private/NiagaraEditorWidgetsModule.cpp
+++ b/Engine/Plugins/FX/Niagara/Source/NiagaraEditorWidgets/Private/NiagaraEditorWidgetsModule.cpp
@@ -22,10 +22,25 @@ FNiagaraStackCurveEditorOptions::FNiagaraStackCurveEditorOptions()
 	, ViewMinOutput(0)
 	, ViewMaxOutput(1)
 	, bAreCurvesVisible(true)
+	, bNeedsInitializeView(true)
 	, Height(100)
 {
 }
 
+bool FNiagaraStackCurveEditorOptions::GetNeedsInitializeView() const
+{
+	return bNeedsInitializeView;
+}
+
+void FNiagaraStackCurveEditorOptions::InitializeView(float InViewMinInput, float InViewMaxInput, float InViewMinOutput, float InViewMaxOutput)
+{
+	ViewMinInput = InViewMinInput;
+	ViewMaxInput = InViewMaxInput;
+	ViewMinOutput = InViewMinOutput;
+	ViewMaxOutput = InViewMaxOutput;
+	bNeedsInitializeView = false;
+}
+
 float FNiagaraStackCurveEditorOptions::GetViewMinInput() const
 {
 	return ViewMinInput;
diff --git a/Engine/Plugins/FX/Niagara/Source/NiagaraEditorWidgets/Public/NiagaraEditorWidgetsModule.h b/Engine/Plugins/FX/Niagara/Source/NiagaraEditorWidgets/Public/NiagaraEditorWidgetsModule.h
--- a/Engine/Plugins/FX/Niagara/Source/NiagaraEditorWidgets/Public/NiagaraEditorWidgetsModule.h
+++ b/Engine/Plugins/FX/Niagara/Source/NiagaraEditorWidgets/Public/NiagaraEditorWidgetsModule.h
@@ -13,6 +13,9 @@ class FNiagaraStackCurveEditorOptions
 public:
 	FNiagaraStackCurveEditorOptions();
 
+	bool GetNeedsInitializeView() const;
+	void InitializeView(float InViewMinInput, float InViewMaxInput, float InViewMinOutput, float InViewMaxOutput);
+
 	float GetViewMinInput() const;
 	float GetViewMaxInput() const;
 	void SetInputViewRange(float InViewMinInput, float InViewMaxInput);
@@ -35,6 +38,7 @@ private:
 	float ViewMinOutput;
 	float ViewMaxOutput;
 	bool bAreCurvesVisible;
+	bool bNeedsInitializeView;
 	float Height;
 };
 
