Niagara - Fix a crash when undoing while hovering over the advanced expander or function input.

diff --git a/Engine/Plugins/FX/Niagara/Source/NiagaraEditorWidgets/Private/Stack/SNiagaraStackFunctionInputName.cpp b/Engine/Plugins/FX/Niagara/Source/NiagaraEditorWidgets/Private/Stack/SNiagaraStackFunctionInputName.cpp
--- a/Engine/Plugins/FX/Niagara/Source/NiagaraEditorWidgets/Private/Stack/SNiagaraStackFunctionInputName.cpp
+++ b/Engine/Plugins/FX/Niagara/Source/NiagaraEditorWidgets/Private/Stack/SNiagaraStackFunctionInputName.cpp
@@ -44,7 +44,7 @@ void SNiagaraStackFunctionInputName::Construct(const FArguments& InArgs, UNiagar
 			.OnTextCommitted(this, &SNiagaraStackFunctionInputName::OnNameTextCommitted)
 			.HighlightText_UObject(InStackViewModel, &UNiagaraStackViewModel::GetCurrentSearchText)
 			.ColorAndOpacity(this, &SNiagaraStackFunctionInputName::GetTextColorForSearch)
-			.ToolTipText_UObject(FunctionInput, &UNiagaraStackFunctionInput::GetTooltipText, UNiagaraStackFunctionInput::EValueMode::Local)
+			.ToolTipText(this, &SNiagaraStackFunctionInputName::GetToolTipText)
 		]
 	];
 }
@@ -89,6 +89,17 @@ bool SNiagaraStackFunctionInputName::GetIsEnabled() const
 	return FunctionInput->GetOwnerIsEnabled() && (FunctionInput->GetHasEditCondition() == false || FunctionInput->GetEditConditionEnabled());
 }
 
+FText SNiagaraStackFunctionInputName::GetToolTipText() const
+{
+	// The engine ticks tooltips before widgets so it's possible for the footer to be finalized when
+	// the widgets haven't been recreated.
+	if (FunctionInput->IsFinalized())
+	{
+		return FText();
+	}
+	return FunctionInput->GetTooltipText(UNiagaraStackFunctionInput::EValueMode::Local);
+}
+
 void SNiagaraStackFunctionInputName::OnNameTextCommitted(const FText& InText, ETextCommit::Type InCommitType)
 {
 	FunctionInput->RenameInput(*InText.ToString());
diff --git a/Engine/Plugins/FX/Niagara/Source/NiagaraEditorWidgets/Private/Stack/SNiagaraStackFunctionInputName.h b/Engine/Plugins/FX/Niagara/Source/NiagaraEditorWidgets/Private/Stack/SNiagaraStackFunctionInputName.h
--- a/Engine/Plugins/FX/Niagara/Source/NiagaraEditorWidgets/Private/Stack/SNiagaraStackFunctionInputName.h
+++ b/Engine/Plugins/FX/Niagara/Source/NiagaraEditorWidgets/Private/Stack/SNiagaraStackFunctionInputName.h
@@ -35,6 +35,8 @@ private:
 
 	bool GetIsEnabled() const;
 
+	FText GetToolTipText() const;
+
 	void OnNameTextCommitted(const FText& InText, ETextCommit::Type InCommitType);
 
 private:
diff --git a/Engine/Plugins/FX/Niagara/Source/NiagaraEditorWidgets/Private/Stack/SNiagaraStackItemFooter.cpp b/Engine/Plugins/FX/Niagara/Source/NiagaraEditorWidgets/Private/Stack/SNiagaraStackItemFooter.cpp
--- a/Engine/Plugins/FX/Niagara/Source/NiagaraEditorWidgets/Private/Stack/SNiagaraStackItemFooter.cpp
+++ b/Engine/Plugins/FX/Niagara/Source/NiagaraEditorWidgets/Private/Stack/SNiagaraStackItemFooter.cpp
@@ -55,6 +55,12 @@ const FSlateBrush* SNiagaraStackItemFooter::GetButtonBrush() const
 
 FText SNiagaraStackItemFooter::GetToolTipText() const
 {
+	// The engine ticks tooltips before widgets so it's possible for the footer to be finalized when
+	// the widgets haven't been recreated.
+	if (ItemFooter->IsFinalized())
+	{
+		return FText();
+	}
 	return ItemFooter->GetShowAdvanced() ? ExpandedToolTipText : CollapsedToolTipText;
 }
 
