Niagara - Fix some styling issues due to the new background styling in the editor.  We should revisit for 4.25.

diff --git a/Engine/Plugins/FX/Niagara/Source/NiagaraEditorWidgets/Private/NiagaraEditorWidgetsStyle.cpp b/Engine/Plugins/FX/Niagara/Source/NiagaraEditorWidgets/Private/NiagaraEditorWidgetsStyle.cpp
--- a/Engine/Plugins/FX/Niagara/Source/NiagaraEditorWidgets/Private/NiagaraEditorWidgetsStyle.cpp
+++ b/Engine/Plugins/FX/Niagara/Source/NiagaraEditorWidgets/Private/NiagaraEditorWidgetsStyle.cpp
@@ -124,11 +124,11 @@ TSharedRef< FSlateStyleSet > FNiagaraEditorWidgetsStyle::Create()
 	Style->Set("NiagaraEditor.SystemOverview.TableViewRow", FTableRowStyle(NormalTableRowStyle)
 		.SetInactiveBrush(IMAGE_CORE_BRUSH("Common/Selection", Icon8x8, GetDefault<UEditorStyleSettings>()->GetSubduedSelectionColor())));
 
-	Style->Set("NiagaraEditor.Stack.Group.BackgroundColor", FLinearColor(FColor(96, 96, 96)));
+	Style->Set("NiagaraEditor.Stack.BackgroundColor", FLinearColor(FColor(96, 96, 96)));
 	Style->Set("NiagaraEditor.Stack.Item.HeaderBackgroundColor", FLinearColor(FColor(48, 48, 48)));
 	Style->Set("NiagaraEditor.Stack.Item.ContentBackgroundColor", FLinearColor(FColor(62, 62, 62)));
 	Style->Set("NiagaraEditor.Stack.Item.ContentAdvancedBackgroundColor", FLinearColor(FColor(53, 53, 53)));
-	Style->Set("NiagaraEditor.Stack.Item.FooterBackgroundColor", FLinearColor(FColor(71, 71, 71)));
+	Style->Set("NiagaraEditor.Stack.Item.FooterBackgroundColor", FLinearColor(FColor(75, 75, 75)));
 	Style->Set("NiagaraEditor.Stack.Item.IssueBackgroundColor", FLinearColor(FColor(120, 120, 62)));
 	Style->Set("NiagaraEditor.Stack.UnknownColor", FLinearColor(1, 0, 1));
 
diff --git a/Engine/Plugins/FX/Niagara/Source/NiagaraEditorWidgets/Private/SNiagaraStack.cpp b/Engine/Plugins/FX/Niagara/Source/NiagaraEditorWidgets/Private/SNiagaraStack.cpp
--- a/Engine/Plugins/FX/Niagara/Source/NiagaraEditorWidgets/Private/SNiagaraStack.cpp
+++ b/Engine/Plugins/FX/Niagara/Source/NiagaraEditorWidgets/Private/SNiagaraStack.cpp
@@ -258,13 +258,18 @@ void SNiagaraStack::Construct(const FArguments& InArgs, UNiagaraStackViewModel*
 		+ SVerticalBox::Slot()
 		.Padding(1)
 		[
-			SAssignNew(StackTree, STreeView<UNiagaraStackEntry*>)
-			.OnGenerateRow(this, &SNiagaraStack::OnGenerateRowForStackItem)
-			.OnGetChildren(this, &SNiagaraStack::OnGetChildren)
-			.TreeItemsSource(&StackViewModel->GetRootEntryAsArray())
-			.OnTreeViewScrolled(this, &SNiagaraStack::StackTreeScrolled)
-			.SelectionMode(ESelectionMode::None)
-			.OnItemToString_Debug_Static(&FNiagaraStackEditorWidgetsUtilities::StackEntryToStringForListDebug)
+			SNew(SBorder)
+			.BorderImage(FEditorStyle::GetBrush("WhiteBrush"))
+			.BorderBackgroundColor(FNiagaraEditorWidgetsStyle::Get().GetColor("NiagaraEditor.Stack.BackgroundColor"))
+			[
+				SAssignNew(StackTree, STreeView<UNiagaraStackEntry*>)
+				.OnGenerateRow(this, &SNiagaraStack::OnGenerateRowForStackItem)
+				.OnGetChildren(this, &SNiagaraStack::OnGetChildren)
+				.TreeItemsSource(&StackViewModel->GetRootEntryAsArray())
+				.OnTreeViewScrolled(this, &SNiagaraStack::StackTreeScrolled)
+				.SelectionMode(ESelectionMode::None)
+				.OnItemToString_Debug_Static(&FNiagaraStackEditorWidgetsUtilities::StackEntryToStringForListDebug)
+			]
 		]
 	];
 
