Move Niagara plugin load order back up to PreDefault to avoid Niagara objects loading in the core packages Default phase failing init due to the Niagara module not being started.

diff --git a/Engine/Plugins/FX/Niagara/Niagara.uplugin b/Engine/Plugins/FX/Niagara/Niagara.uplugin
--- a/Engine/Plugins/FX/Niagara/Niagara.uplugin
+++ b/Engine/Plugins/FX/Niagara/Niagara.uplugin
@@ -24,7 +24,7 @@
 		{
 		  "Name": "Niagara",
 		  "Type": "Runtime",
-		  "LoadingPhase": "Default"
+		  "LoadingPhase": "PreDefault"
 		},
 		{
 			"Name" : "NiagaraAnimNotifies",
@@ -44,12 +44,12 @@
 		{
 			"Name" : "NiagaraEditor",
 			"Type" : "Editor",
-			"LoadingPhase" : "PostDefault"
+			"LoadingPhase" : "Default"
 		},
 		{
 			"Name" : "NiagaraEditorWidgets",
 			"Type" : "Editor",
-			"LoadingPhase" : "PostDefault"
+			"LoadingPhase" : "Default"
 		}
 	] 
 }
