Fixing custom hlsl node to support DataInterfaces. Once you name the DI, you can invoke any supported DI method with a "." followed by the function invocation. We detect this behind the scenes and setup all the necessary connections.
Can also now rename on pin creation.

diff --git a/Engine/Plugins/FX/Niagara/Source/NiagaraCore/Private/NiagaraCustomVersion.cpp b/Engine/Plugins/FX/Niagara/Source/NiagaraCore/Private/NiagaraCustomVersion.cpp
--- a/Engine/Plugins/FX/Niagara/Source/NiagaraCore/Private/NiagaraCustomVersion.cpp
+++ b/Engine/Plugins/FX/Niagara/Source/NiagaraCore/Private/NiagaraCustomVersion.cpp
@@ -8,4 +8,4 @@ const FGuid FNiagaraCustomVersion::GUID(0xFCF57AFA, 0x50764283, 0xB9A9E658, 0xFF
 // Register the custom version with core
 FCustomVersionRegistration GRegisterNiagaraCustomVersion(FNiagaraCustomVersion::GUID, FNiagaraCustomVersion::LatestVersion, TEXT("NiagaraVer"));
 
-const FGuid FNiagaraCustomVersion::LatestScriptCompileVersion(0x79007C85, 0x2A9CA54A, 0xA7DABF24, 0x30A4EB07);
\ No newline at end of file
+const FGuid FNiagaraCustomVersion::LatestScriptCompileVersion(0x3A9CFB8E, 0x7E17484D, 0xA42624B4, 0xF2F4D750);
\ No newline at end of file
diff --git a/Engine/Plugins/FX/Niagara/Source/NiagaraEditor/Private/NiagaraHlslTranslator.cpp b/Engine/Plugins/FX/Niagara/Source/NiagaraEditor/Private/NiagaraHlslTranslator.cpp
--- a/Engine/Plugins/FX/Niagara/Source/NiagaraEditor/Private/NiagaraHlslTranslator.cpp
+++ b/Engine/Plugins/FX/Niagara/Source/NiagaraEditor/Private/NiagaraHlslTranslator.cpp
@@ -1370,7 +1370,11 @@ void FHlslNiagaraTranslator::GatherVariableForDataSetAccess(const FNiagaraVariab
 {
 	TArray<FString> Components;
 	UScriptStruct* Struct = Var.GetType().GetScriptStruct();
-	check(Struct);
+	if (!Struct)
+	{
+		Error(FText::Format(LOCTEXT("BadStructDef", "Variable {0} missing struct definition."), FText::FromName(Var.GetName())), nullptr, nullptr);
+		return;
+	}
 
 	TArray<ENiagaraBaseTypes> Types;
 	GatherComponentsForDataSetAccess(Struct, TEXT(""), false, Components, Types);
@@ -4694,7 +4698,7 @@ void FHlslNiagaraTranslator::FinalResolveNamespacedTokens(const FString& Paramet
 								break;
 							}
 						}
-						if (!bAdded)
+						if (!bAdded && !UNiagaraScript::IsStandaloneScript(CompileOptions.TargetUsage)) // Don't warn in modules, they don't have enough context.
 						{
 							Error(FText::Format(LOCTEXT("GetCustomFail1", "Cannot use variable in custom expression, it hasn't been encountered yet: {0}"), FText::FromName(Var.GetName())), nullptr, nullptr);
 						}
@@ -4744,14 +4748,64 @@ void FHlslNiagaraTranslator::HandleCustomHlslNode(UNiagaraNodeCustomHlsl* Custom
 		if (Input.GetType() == FNiagaraTypeDefinition::GetParameterMapDef())
 		{
 			FString ParameterMapInstanceName = GetParameterMapInstanceName(0);
-			FString ReplaceSrc = Input.GetName().ToString() + TEXT(".");
-			FString ReplaceDest = ParameterMapInstanceName + TEXT(".");
-			CustomFunctionHlsl->ReplaceExactMatchTokens(Tokens, ReplaceSrc, ReplaceDest, false);
+			FString ReplaceSrc = Input.GetName().ToString();
+			FString ReplaceDest = ParameterMapInstanceName;
+			CustomFunctionHlsl->ReplaceExactMatchTokens(Tokens, ReplaceSrc, ReplaceDest, true);
 			SigInputs.Add(Input);
 			OutSignature.bRequiresContext = true;
 			ParamMapHistoryIdx = Inputs[i];
 			bHasParamMapInputs = true;
 		}
+		else if (Input.GetType().IsDataInterface())
+		{
+			UObject* const* FoundCDO = CompileData->CDOs.Find(Input.GetType().GetClass());
+			UNiagaraDataInterface* CDO = Cast<UNiagaraDataInterface>(*FoundCDO);
+
+			if (FoundCDO == nullptr || CDO == nullptr)
+			{
+				// If the cdo wasn't found, the data interface was not passed through a parameter map and so it won't be bound correctly, so add a compile error
+				// and invalidate the signature.
+				Error(LOCTEXT("DataInterfaceNotFoundCustomHLSL", "Data interface used by custom hlsl, but not found in precompiled data. Please notify Niagara team of bug."), nullptr, nullptr);
+				return;
+			}
+			int32 OwnerIdx = Inputs[i];
+			if (OwnerIdx < 0 || OwnerIdx >= CompilationOutput.ScriptData.DataInterfaceInfo.Num())
+			{
+				Error(LOCTEXT("FunctionCallDataInterfaceMissingRegistration", "Function call signature does not match to a registered DataInterface. Valid DataInterfaces should be wired into a DataInterface function call."), nullptr, nullptr);
+				return;
+			}
+			
+			// Go over all the supported functions in the DI and look to see if they occur in the 
+			// actual custom hlsl source. If they do, then add them to the function table that we need to map.
+			FNiagaraScriptDataInterfaceCompileInfo& Info = CompilationOutput.ScriptData.DataInterfaceInfo[OwnerIdx];
+			TArray<FNiagaraFunctionSignature> Funcs;
+			TArray<FNiagaraFunctionSignature> AddedFuncs;
+			CDO->GetFunctions(Funcs);
+			for (int32 FuncIdx = 0; FuncIdx < Funcs.Num(); FuncIdx++)
+			{
+				FNiagaraFunctionSignature Sig = Funcs[FuncIdx];
+				FString ReplaceSrc = Input.GetName().ToString() + TEXT(".") + Sig.GetName();
+				FString ReplaceDest = GetSanitizedSymbolName(Sig.GetName() + TEXT("_") + (Info.Name.ToString().Replace(TEXT("."), TEXT(""))));
+				uint32 NumFound = CustomFunctionHlsl->ReplaceExactMatchTokens(Tokens, ReplaceSrc, ReplaceDest, false);
+				if (NumFound != 0)
+				{
+					AddedFuncs.Add(Sig);
+					DataInterfaceRegisteredFunctions.FindOrAdd(Input.GetType().GetFName()).Add(Sig);
+
+					if (Info.UserPtrIdx != INDEX_NONE && CompilationTarget != ENiagaraSimTarget::GPUComputeSim)
+					{
+						//This interface requires per instance data via a user ptr so place the index to it at the end of the inputs.
+						//Skip this for now... Inputs.Add(AddSourceChunk(LexToString(Info.UserPtrIdx), FNiagaraTypeDefinition::GetIntDef(), false));
+						Sig.Inputs.Add(FNiagaraVariable(FNiagaraTypeDefinition::GetIntDef(), TEXT("InstanceData")));
+					}
+					//Override the owner id of the signature with the actual caller.
+					Sig.OwnerName = Info.Name;
+					Info.RegisteredFunctions.Add(Sig);
+					Functions.FindOrAdd(Sig);
+				}
+			}
+			SigInputs.Add(Input);
+		}
 		else
 		{
 			FString ReplaceSrc = Input.GetName().ToString();
@@ -4769,9 +4823,9 @@ void FHlslNiagaraTranslator::HandleCustomHlslNode(UNiagaraNodeCustomHlsl* Custom
 		if (Output.GetType() == FNiagaraTypeDefinition::GetParameterMapDef())
 		{
 			FString ParameterMapInstanceName = GetParameterMapInstanceName(0);
-			FString ReplaceSrc = Output.GetName().ToString() + TEXT(".");
-			FString ReplaceDest = ParameterMapInstanceName + TEXT(".");
-			CustomFunctionHlsl->ReplaceExactMatchTokens(Tokens, ReplaceSrc, ReplaceDest, false);
+			FString ReplaceSrc = Output.GetName().ToString();
+			FString ReplaceDest = ParameterMapInstanceName;
+			CustomFunctionHlsl->ReplaceExactMatchTokens(Tokens, ReplaceSrc, ReplaceDest, true);
 			SigOutputs.Add(Output);
 			OutSignature.bRequiresContext = true;
 			bHasParamMapOutputs = true;
diff --git a/Engine/Plugins/FX/Niagara/Source/NiagaraEditor/Private/NiagaraNodeCustomHlsl.cpp b/Engine/Plugins/FX/Niagara/Source/NiagaraEditor/Private/NiagaraNodeCustomHlsl.cpp
--- a/Engine/Plugins/FX/Niagara/Source/NiagaraEditor/Private/NiagaraNodeCustomHlsl.cpp
+++ b/Engine/Plugins/FX/Niagara/Source/NiagaraEditor/Private/NiagaraNodeCustomHlsl.cpp
@@ -13,6 +13,7 @@
 UNiagaraNodeCustomHlsl::UNiagaraNodeCustomHlsl(const FObjectInitializer& ObjectInitializer)
 	: Super(ObjectInitializer)
 {
+	PinPendingRename = nullptr;
 	bCanRenameNode = true;
 	ScriptUsage = ENiagaraScriptUsage::Function;
 
@@ -164,11 +165,101 @@ void UNiagaraNodeCustomHlsl::InitAsCustomHlslDynamicInput(const FNiagaraTypeDefi
 
 #endif
 
+bool UNiagaraNodeCustomHlsl::IsPinNameEditableUponCreation(const UEdGraphPin* GraphPinObj) const
+{
+	if (GraphPinObj == PinPendingRename && ScriptUsage != ENiagaraScriptUsage::DynamicInput)
+	{
+		return true;
+	}
+	else
+	{
+		return false;
+	}
+}
+
+bool UNiagaraNodeCustomHlsl::IsPinNameEditable(const UEdGraphPin* GraphPinObj) const
+{
+
+	const UEdGraphSchema_Niagara* Schema = GetDefault<UEdGraphSchema_Niagara>();
+	FNiagaraTypeDefinition TypeDef = Schema->PinToTypeDefinition(GraphPinObj);
+	if (TypeDef.IsValid() && GraphPinObj &&  CanRenamePin(GraphPinObj) && ScriptUsage != ENiagaraScriptUsage::DynamicInput)
+	{
+		return true;
+	}
+	else
+	{
+		return false;
+	}
+}
+
+bool UNiagaraNodeCustomHlsl::VerifyEditablePinName(const FText& InName, FText& OutErrorMessage, const UEdGraphPin* InGraphPinObj)  const
+{
+	// Check to see if the symbol has to be mangled to be valid hlsl. If it does, then prevent it from being 
+	// valid. This helps clear up any ambiguity downstream in the translator.
+	FString NewName = InName.ToString();
+	FString SanitizedNewName = FHlslNiagaraTranslator::GetSanitizedSymbolName(NewName);
+
+	if (NewName != SanitizedNewName || NewName.Len() == 0)
+	{
+		OutErrorMessage = FText::Format(LOCTEXT("InvalidPinName", "Pin \"{0}\" cannot be renamed to \"{1}\". Certain words are restricted, as are spaces and special characters. Suggestion: \"{2}\""), InGraphPinObj->GetDisplayName(), InName, FText::FromString(SanitizedNewName));
+		return false;
+	}
+	TSet<FName> Names;
+	for (int32 i = 0; i < Pins.Num(); i++)
+	{
+		if (Pins[i] != InGraphPinObj)
+			Names.Add(Pins[i]->GetFName());
+	}
+	if (Names.Contains(*NewName))
+	{
+		OutErrorMessage = FText::Format(LOCTEXT("InvalidPinName", "Pin \"{0}\" cannot be renamed to \"{1}\" as it conflicts with another name in use. Suggestion: \"{2}\""), InGraphPinObj->GetDisplayName(), InName, FText::FromName(FNiagaraUtilities::GetUniqueName(*SanitizedNewName, Names)));
+		return false;
+	}
+	OutErrorMessage = FText::GetEmpty();
+	return true;
+}
+
+bool UNiagaraNodeCustomHlsl::CommitEditablePinName(const FText& InName, UEdGraphPin* InGraphPinObj)
+{
+	if (Pins.Contains(InGraphPinObj))
+	{
+		FScopedTransaction AddNewPinTransaction(LOCTEXT("Rename Pin", "Renamed pin"));
+		Modify();
+		InGraphPinObj->Modify();
+
+		FString OldPinName = InGraphPinObj->PinName.ToString();
+		InGraphPinObj->PinName = *InName.ToString();
+		OnPinRenamed(InGraphPinObj, OldPinName);
+
+		return true;
+	}
+	return false;
+}
+
+bool UNiagaraNodeCustomHlsl::CancelEditablePinName(const FText& InName, UEdGraphPin* InGraphPinObj)
+{
+	if (InGraphPinObj == PinPendingRename)
+	{
+		PinPendingRename = nullptr;
+	}
+	return true;
+}
+
+
 /** Called when a new typed pin is added by the user. */
 void UNiagaraNodeCustomHlsl::OnNewTypedPinAdded(UEdGraphPin* NewPin)
 {
+	TSet<FName> Names;
+	for (int32 i = 0; i < Pins.Num(); i++)
+	{
+		if (Pins[i] != NewPin)
+			Names.Add(Pins[i]->GetFName());
+	}
+	FName Name = FNiagaraUtilities::GetUniqueName(NewPin->GetFName(), Names);
+	NewPin->PinName = Name;
 	UNiagaraNodeWithDynamicPins::OnNewTypedPinAdded(NewPin);
 	RebuildSignatureFromPins();
+	PinPendingRename = NewPin;
 }
 
 /** Called when a pin is renamed. */
@@ -249,6 +340,7 @@ void UNiagaraNodeCustomHlsl::BuildParameterMapHistory(FNiagaraParameterMapHistor
 				bHasParamMapOutput = true;
 				FString ReplaceSrc = Output.GetName().ToString() + TEXT(".");
 				ReplaceExactMatchTokens(Tokens, ReplaceSrc, TEXT(""), false);
+				OutHistory.RegisterParameterMapPin(ParamMapIdx, OutputPins[i]);
 			}
 			else
 			{
@@ -287,8 +379,9 @@ void UNiagaraNodeCustomHlsl::BuildParameterMapHistory(FNiagaraParameterMapHistor
 }
 
 // Replace items in the tokens array if they start with the src string or optionally src string and a namespace delimiter
-void UNiagaraNodeCustomHlsl::ReplaceExactMatchTokens(TArray<FString>& Tokens, const FString& SrcString, const FString& ReplaceString, bool bAllowNamespaceSeparation)
+uint32 UNiagaraNodeCustomHlsl::ReplaceExactMatchTokens(TArray<FString>& Tokens, const FString& SrcString, const FString& ReplaceString, bool bAllowNamespaceSeparation)
 {
+	uint32 Count = 0;
 	for (int32 i = 0; i < Tokens.Num(); i++)
 	{
 		if (Tokens[i].StartsWith(SrcString, ESearchCase::CaseSensitive))
@@ -296,13 +389,25 @@ void UNiagaraNodeCustomHlsl::ReplaceExactMatchTokens(TArray<FString>& Tokens, co
 			if (Tokens[i].Len() > SrcString.Len() && Tokens[i][SrcString.Len()] == '.' && bAllowNamespaceSeparation)
 			{
 				Tokens[i] = ReplaceString + Tokens[i].Mid(SrcString.Len());
+				Count++;
 			}
 			else if (Tokens[i].Len() == SrcString.Len())
 			{
 				Tokens[i] = ReplaceString;
+				Count++;
 			}
 		}
 	}
+	return Count;
+}
+
+
+bool UNiagaraNodeCustomHlsl::AllowNiagaraTypeForAddPin(const FNiagaraTypeDefinition& InType)
+{
+	if (Super::AllowNiagaraTypeForAddPin(InType) || InType.IsDataInterface())
+		return true;
+	else
+		return false;
 }
 
 
diff --git a/Engine/Plugins/FX/Niagara/Source/NiagaraEditor/Private/NiagaraNodeFunctionCall.cpp b/Engine/Plugins/FX/Niagara/Source/NiagaraEditor/Private/NiagaraNodeFunctionCall.cpp
--- a/Engine/Plugins/FX/Niagara/Source/NiagaraEditor/Private/NiagaraNodeFunctionCall.cpp
+++ b/Engine/Plugins/FX/Niagara/Source/NiagaraEditor/Private/NiagaraNodeFunctionCall.cpp
@@ -551,7 +551,7 @@ void UNiagaraNodeFunctionCall::Compile(class FHlslNiagaraTranslator* Translator,
 	{
 		if (Signature.Inputs.Num() > 0)
 		{
-			if (Signature.Inputs[0].GetType().IsDataInterface())
+			if (Signature.Inputs[0].GetType().IsDataInterface() && GetValidateDataInterfaces())
 			{
 				UClass* DIClass = Signature.Inputs[0].GetType().GetClass();
 				if (UNiagaraDataInterface* DataInterfaceCDO = Cast<UNiagaraDataInterface>(DIClass->GetDefaultObject()))
diff --git a/Engine/Plugins/FX/Niagara/Source/NiagaraEditor/Private/NiagaraNodeWithDynamicPins.cpp b/Engine/Plugins/FX/Niagara/Source/NiagaraEditor/Private/NiagaraNodeWithDynamicPins.cpp
--- a/Engine/Plugins/FX/Niagara/Source/NiagaraEditor/Private/NiagaraNodeWithDynamicPins.cpp
+++ b/Engine/Plugins/FX/Niagara/Source/NiagaraEditor/Private/NiagaraNodeWithDynamicPins.cpp
@@ -15,6 +15,7 @@
 #include "SNiagaraParameterMapView.h"
 #include "Framework/Application/SlateApplication.h"
 #include "NiagaraNodeParameterMapBase.h"
+#include "Widgets/SEditableTextBoxWithVerification.h"
 
 #define LOCTEXT_NAMESPACE "NiagaraNodeWithDynamicPins"
 
@@ -176,9 +177,10 @@ void UNiagaraNodeWithDynamicPins::GetNodeContextMenuActions(UToolMenu* Menu, UGr
 				.WidthOverride(100)
 				.Padding(FMargin(5, 0, 0, 0))
 				[
-					SNew(SEditableTextBox)
+					SNew(SEditableTextBoxWithVerification)
 					.Text_UObject(this, &UNiagaraNodeWithDynamicPins::GetPinNameText, Pin)
 					.OnTextCommitted_UObject(const_cast<UNiagaraNodeWithDynamicPins*>(this), &UNiagaraNodeWithDynamicPins::PinNameTextCommitted, Pin)
+					.OnVerifyTextChanged_UObject(this, &UNiagaraNodeWithDynamicPins::VerifyEditablePinName, Context->Pin)
 				];
 			Section.AddEntry(FToolMenuEntry::InitWidget("RenameWidget", RenameWidget, LOCTEXT("NameMenuItem", "Name")));
 		}
@@ -307,6 +309,8 @@ void UNiagaraNodeWithDynamicPins::PinNameTextCommitted(const FText& Text, ETextC
 	}
 }
 
+
+
 void UNiagaraNodeWithDynamicPins::RemoveDynamicPinFromMenu(UEdGraphPin* Pin)
 {
 	FScopedTransaction RemovePinTransaction(LOCTEXT("RemovePinTransaction", "Remove pin"));
diff --git a/Engine/Plugins/FX/Niagara/Source/NiagaraEditor/Public/NiagaraNodeCustomHlsl.h b/Engine/Plugins/FX/Niagara/Source/NiagaraEditor/Public/NiagaraNodeCustomHlsl.h
--- a/Engine/Plugins/FX/Niagara/Source/NiagaraEditor/Public/NiagaraNodeCustomHlsl.h
+++ b/Engine/Plugins/FX/Niagara/Source/NiagaraEditor/Public/NiagaraNodeCustomHlsl.h
@@ -35,8 +35,9 @@ public:
 	virtual void BuildParameterMapHistory(FNiagaraParameterMapHistoryBuilder& OutHistory, bool bRecursive = true, bool bFilterForCompilation = true) const override;
 
 	// Replace items in the tokens array if they start with the src string or optionally src string and a namespace delimiter
-	static void ReplaceExactMatchTokens(TArray<FString>& Tokens, const FString& SrcString, const FString& ReplaceString, bool bAllowNamespaceSeparation);
+	static uint32 ReplaceExactMatchTokens(TArray<FString>& Tokens, const FString& SrcString, const FString& ReplaceString, bool bAllowNamespaceSeparation);
 	static FNiagaraVariable StripVariableToBaseType(const FNiagaraVariable& InVar);
+	virtual bool AllowNiagaraTypeForAddPin(const FNiagaraTypeDefinition& InType);
 
 #if WITH_EDITOR
 	virtual void PostEditChangeProperty(struct FPropertyChangedEvent& PropertyChangedEvent) override;
@@ -46,7 +47,13 @@ public:
 #endif
 protected:
 	virtual bool AllowDynamicPins() const override { return true; }
-
+	virtual bool GetValidateDataInterfaces() const override { return false; }
+	virtual bool VerifyEditablePinName(const FText& InName, FText& OutErrorMessage, const UEdGraphPin* InGraphPinObj) const override;
+	virtual bool IsPinNameEditableUponCreation(const UEdGraphPin* Pin) const override;
+	virtual bool IsPinNameEditable(const UEdGraphPin* Pin) const override;
+	virtual bool CommitEditablePinName(const FText& InName, UEdGraphPin* InGraphPinObj) override;
+	virtual bool CancelEditablePinName(const FText& InName, UEdGraphPin* InGraphPinObj) override;
+	
 	virtual bool CanRenamePin(const UEdGraphPin* Pin) const override { return UNiagaraNodeWithDynamicPins::CanRenamePin(Pin); }
 	virtual bool CanRemovePin(const UEdGraphPin* Pin) const override {
 		return UNiagaraNodeWithDynamicPins::CanRemovePin(Pin);
@@ -67,5 +74,6 @@ protected:
 	virtual void MoveDynamicPin(UEdGraphPin* Pin, int32 DirectionToMove) override;
 
 	void RebuildSignatureFromPins();
+	UEdGraphPin* PinPendingRename;
 
 };
diff --git a/Engine/Plugins/FX/Niagara/Source/NiagaraEditor/Public/NiagaraNodeFunctionCall.h b/Engine/Plugins/FX/Niagara/Source/NiagaraEditor/Public/NiagaraNodeFunctionCall.h
--- a/Engine/Plugins/FX/Niagara/Source/NiagaraEditor/Public/NiagaraNodeFunctionCall.h
+++ b/Engine/Plugins/FX/Niagara/Source/NiagaraEditor/Public/NiagaraNodeFunctionCall.h
@@ -137,6 +137,7 @@ public:
 protected:
 
 	virtual bool IsValidPinToCompile(UEdGraphPin* Pin) const;
+	virtual bool GetValidateDataInterfaces() const { return true; };
 
 	virtual bool AllowDynamicPins() const override { return false; }
 	virtual bool CanRenamePin(const UEdGraphPin* Pin) const override { return false; }
diff --git a/Engine/Plugins/FX/Niagara/Source/NiagaraEditor/Public/NiagaraNodeWithDynamicPins.h b/Engine/Plugins/FX/Niagara/Source/NiagaraEditor/Public/NiagaraNodeWithDynamicPins.h
--- a/Engine/Plugins/FX/Niagara/Source/NiagaraEditor/Public/NiagaraNodeWithDynamicPins.h
+++ b/Engine/Plugins/FX/Niagara/Source/NiagaraEditor/Public/NiagaraNodeWithDynamicPins.h
@@ -52,6 +52,10 @@ protected:
 
 	/** Called when a new typed pin is added by the user. */
 	virtual void OnNewTypedPinAdded(UEdGraphPin* NewPin) { }
+	
+	/** Called in subclasses to restrict renaming.*/
+	/** Verify that the potential rename has produced acceptable results for a pin.*/
+	virtual bool VerifyEditablePinName(const FText& InName, FText& OutErrorMessage, const UEdGraphPin* InGraphPinObj) const { OutErrorMessage = FText::GetEmpty(); return true; }
 
 	/** Called when a pin is renamed. */
 	virtual void OnPinRenamed(UEdGraphPin* RenamedPin, const FString& OldPinName) { }
@@ -70,6 +74,8 @@ protected:
 
 	virtual void MoveDynamicPin(UEdGraphPin* Pin, int32 DirectionToMove);
 
+	virtual bool OnVerifyTextChanged(const FText& NewText, FText& OutMessage) { return true; };
+
 private:
 
 	/** Gets the display text for a pin. */
@@ -77,7 +83,7 @@ private:
 
 	/** Called when a pin's name text is committed. */
 	void PinNameTextCommitted(const FText& Text, ETextCommit::Type CommitType, UEdGraphPin* Pin);
-
+	
 	void RemoveDynamicPinFromMenu(UEdGraphPin* Pin);
 
 	void MoveDynamicPinFromMenu(UEdGraphPin* Pin, int32 DirectionToMove);
