Niagara - Fix button layout in the script editor.

diff --git a/Engine/Plugins/FX/Niagara/Source/NiagaraEditor/Private/Customizations/NiagaraScriptDetails.cpp b/Engine/Plugins/FX/Niagara/Source/NiagaraEditor/Private/Customizations/NiagaraScriptDetails.cpp
--- a/Engine/Plugins/FX/Niagara/Source/NiagaraEditor/Private/Customizations/NiagaraScriptDetails.cpp
+++ b/Engine/Plugins/FX/Niagara/Source/NiagaraEditor/Private/Customizations/NiagaraScriptDetails.cpp
@@ -144,7 +144,13 @@ void FNiagaraScriptDetails::CustomizeDetails(IDetailLayoutBuilder& DetailBuilder
 		IDetailCategoryBuilder& InputParamCategory = DetailBuilder.EditCategory(InputParamCategoryName, LOCTEXT("InputParamCategoryName", "Input Parameters"));
 		if (bAddParameters)
 		{
-			InputParamCategory.HeaderContent(SNew(SAddParameterButton, InputCollectionViewModel.ToSharedRef()));
+			InputParamCategory.HeaderContent(
+				SNew(SBox)
+				.Padding(FMargin(2, 2, 0, 2))
+				.VAlign(VAlign_Center)
+				[
+					SNew(SAddParameterButton, InputCollectionViewModel.ToSharedRef())
+				]);
 		}
 		InputParamCategory.AddCustomBuilder(MakeShared<FNiagaraParameterCollectionCustomNodeBuilder>(InputCollectionViewModel.ToSharedRef()));
 	}
@@ -154,7 +160,13 @@ void FNiagaraScriptDetails::CustomizeDetails(IDetailLayoutBuilder& DetailBuilder
 		IDetailCategoryBuilder& OutputParamCategory = DetailBuilder.EditCategory(OutputParamCategoryName, LOCTEXT("OutputParamCategoryName", "Output Parameters"));
 		if (bAddParameters)
 		{
-			OutputParamCategory.HeaderContent(SNew(SAddParameterButton, OutputCollectionViewModel.ToSharedRef()));
+			OutputParamCategory.HeaderContent(
+				SNew(SBox)
+				.Padding(FMargin(2, 2, 0, 2))
+				.VAlign(VAlign_Center)
+				[
+					SNew(SAddParameterButton, OutputCollectionViewModel.ToSharedRef())
+				]);
 		}
 		OutputParamCategory.AddCustomBuilder(MakeShared<FNiagaraParameterCollectionCustomNodeBuilder>(OutputCollectionViewModel.ToSharedRef()));
 	}
