filling raw grid 2d data into a rt for use in materials

diff --git a/Engine/Plugins/FX/Niagara/Source/Niagara/Private/NiagaraDataInterfaceGrid2DCollection.cpp b/Engine/Plugins/FX/Niagara/Source/Niagara/Private/NiagaraDataInterfaceGrid2DCollection.cpp
--- a/Engine/Plugins/FX/Niagara/Source/Niagara/Private/NiagaraDataInterfaceGrid2DCollection.cpp
+++ b/Engine/Plugins/FX/Niagara/Source/Niagara/Private/NiagaraDataInterfaceGrid2DCollection.cpp
@@ -547,6 +547,13 @@ void UNiagaraDataInterfaceGrid2DCollection::GetRawTextureSize(const UNiagaraComp
 		return;
 	}
 	
+	if (Grid2DInstanceData->Buffers.Num() <= 0)
+	{
+		SizeX = -1;
+		SizeY = -1;
+		return;
+	}
+
 	SizeX = Grid2DInstanceData->Buffers[0]->GridBuffer.Buffer->GetSizeX();
 	SizeY = Grid2DInstanceData->Buffers[0]->GridBuffer.Buffer->GetSizeY();
 }
