Niagara Parameter Map pins now look like BP exec pins. This is put in place to start making Niagara users raise questions about how it looks and runs compared to BP. May be reverted in a future release.

diff --git a/Engine/Plugins/FX/Niagara/Source/NiagaraEditor/Private/NiagaraEditorModule.cpp b/Engine/Plugins/FX/Niagara/Source/NiagaraEditor/Private/NiagaraEditorModule.cpp
--- a/Engine/Plugins/FX/Niagara/Source/NiagaraEditor/Private/NiagaraEditorModule.cpp
+++ b/Engine/Plugins/FX/Niagara/Source/NiagaraEditor/Private/NiagaraEditorModule.cpp
@@ -23,6 +23,7 @@
 #include "SGraphPin.h"
 #include "KismetPins/SGraphPinVector4.h"
 #include "KismetPins/SGraphPinNum.h"
+#include "KismetPins/SGraphPinExec.h"
 #include "KismetPins/SGraphPinInteger.h"
 #include "KismetPins/SGraphPinVector.h"
 #include "KismetPins/SGraphPinVector2D.h"
@@ -691,6 +692,10 @@ void FNiagaraEditorModule::StartupModule()
 	GraphPanelPinFactory->RegisterMiscSubCategoryPin(UNiagaraNodeWithDynamicPins::AddPinSubCategory, FNiagaraScriptGraphPanelPinFactory::FCreateGraphPin::CreateLambda(
 		[](UEdGraphPin* GraphPin) -> TSharedRef<SGraphPin> { return SNew(SNiagaraGraphPinAdd, GraphPin); }));
 
+	GraphPanelPinFactory->RegisterTypePin(FNiagaraTypeDefinition::GetParameterMapStruct(), FNiagaraScriptGraphPanelPinFactory::FCreateGraphPin::CreateLambda(
+		[](UEdGraphPin* GraphPin) -> TSharedRef<SGraphPin> { return SNew(SGraphPinExec, GraphPin); }));
+
+
 	EnumTypeUtilities = MakeShareable(new FNiagaraEditorEnumTypeUtilities());
 	RegisterTypeUtilities(FNiagaraTypeDefinition::GetFloatDef(), MakeShareable(new FNiagaraEditorFloatTypeUtilities()));
 	RegisterTypeUtilities(FNiagaraTypeDefinition::GetIntDef(), MakeShareable(new FNiagaraEditorIntegerTypeUtilities()));
