Niagara - Fix an issue where emitters would be incorrectly deleted when modifying the timeline.

diff --git a/Engine/Plugins/FX/Niagara/Source/NiagaraEditor/Private/Sequencer/NiagaraSequence/MovieSceneNiagaraEmitterTrack.cpp b/Engine/Plugins/FX/Niagara/Source/NiagaraEditor/Private/Sequencer/NiagaraSequence/MovieSceneNiagaraEmitterTrack.cpp
--- a/Engine/Plugins/FX/Niagara/Source/NiagaraEditor/Private/Sequencer/NiagaraSequence/MovieSceneNiagaraEmitterTrack.cpp
+++ b/Engine/Plugins/FX/Niagara/Source/NiagaraEditor/Private/Sequencer/NiagaraSequence/MovieSceneNiagaraEmitterTrack.cpp
@@ -72,6 +72,7 @@ TSharedRef<ISequencerSection> UMovieSceneNiagaraEmitterSectionBase::MakeInvalidS
 UMovieSceneNiagaraEmitterTrack::UMovieSceneNiagaraEmitterTrack(const FObjectInitializer& ObjectInitializer)
 	: Super(ObjectInitializer)
 {
+	bSectionsWereModified = false;
 }
 
 void UMovieSceneNiagaraEmitterTrack::Initialize(FNiagaraSystemViewModel& InSystemViewModel, TSharedRef<FNiagaraEmitterHandleViewModel> InEmitterHandleViewModel, const FFrameRate& InFrameResolution)
@@ -155,6 +156,11 @@ void UMovieSceneNiagaraEmitterTrack::UpdateEmitterHandleFromTrackChange(const FF
 	}
 }
 
+bool UMovieSceneNiagaraEmitterTrack::GetSectionsWereModified() const
+{
+	return bSectionsWereModified;
+}
+
 bool UMovieSceneNiagaraEmitterTrack::HasSection(const UMovieSceneSection& Section) const
 {
 	return Sections.Contains(&Section);
@@ -163,16 +169,19 @@ bool UMovieSceneNiagaraEmitterTrack::HasSection(const UMovieSceneSection& Sectio
 void UMovieSceneNiagaraEmitterTrack::AddSection(UMovieSceneSection& Section)
 {
 	Sections.Add(&Section);
+	bSectionsWereModified = true;
 }
 
 void UMovieSceneNiagaraEmitterTrack::RemoveSection(UMovieSceneSection& Section)
 {
 	Sections.Remove(&Section);
+	bSectionsWereModified = true;
 }
 
 void UMovieSceneNiagaraEmitterTrack::RemoveSectionAt(int32 SectionIndex)
 {
 	Sections.RemoveAt(SectionIndex);
+	bSectionsWereModified = true;
 }
 
 bool UMovieSceneNiagaraEmitterTrack::IsEmpty() const
@@ -270,6 +279,7 @@ void UMovieSceneNiagaraEmitterTrack::CreateSections(const FFrameRate& InFrameRes
 	}
 
 	UpdateTrackFromEmitterParameterChange(InFrameResolution);
+	bSectionsWereModified = false;
 }
 
 #undef LOCTEXT_NAMESPACE
diff --git a/Engine/Plugins/FX/Niagara/Source/NiagaraEditor/Private/Sequencer/NiagaraSequence/MovieSceneNiagaraEmitterTrack.h b/Engine/Plugins/FX/Niagara/Source/NiagaraEditor/Private/Sequencer/NiagaraSequence/MovieSceneNiagaraEmitterTrack.h
--- a/Engine/Plugins/FX/Niagara/Source/NiagaraEditor/Private/Sequencer/NiagaraSequence/MovieSceneNiagaraEmitterTrack.h
+++ b/Engine/Plugins/FX/Niagara/Source/NiagaraEditor/Private/Sequencer/NiagaraSequence/MovieSceneNiagaraEmitterTrack.h
@@ -70,6 +70,8 @@ public:
 
 	void UpdateEmitterHandleFromTrackChange(const FFrameRate& InFrameResolution);
 
+	bool GetSectionsWereModified() const;
+
 	//~ UMovieSceneTrack interface
 	virtual void RemoveAllAnimationData() override { }
 	virtual bool HasSection(const UMovieSceneSection& Section) const override;
@@ -99,6 +101,9 @@ private:
 	UPROPERTY()
 	TArray<UMovieSceneSection*> Sections;
 
+	UPROPERTY()
+	bool bSectionsWereModified;
+
 	// Used for detecting copy/paste 
 	UPROPERTY()
 	FGuid EmitterHandleId;
diff --git a/Engine/Plugins/FX/Niagara/Source/NiagaraEditor/Private/ViewModels/NiagaraSystemViewModel.cpp b/Engine/Plugins/FX/Niagara/Source/NiagaraEditor/Private/ViewModels/NiagaraSystemViewModel.cpp
--- a/Engine/Plugins/FX/Niagara/Source/NiagaraEditor/Private/ViewModels/NiagaraSystemViewModel.cpp
+++ b/Engine/Plugins/FX/Niagara/Source/NiagaraEditor/Private/ViewModels/NiagaraSystemViewModel.cpp
@@ -1438,22 +1438,27 @@ void FNiagaraSystemViewModel::SequencerDataChanged(EMovieSceneDataChangeType Dat
 			UMovieSceneNiagaraEmitterTrack* EmitterTrack = CastChecked<UMovieSceneNiagaraEmitterTrack>(Track);
 			if (EmitterTrack->GetEmitterHandleViewModel().IsValid())
 			{
-				if (EmitterTrack->GetAllSections().Num() > 0)
+				if (EmitterTrack->GetAllSections().Num() == 0 && EmitterTrack->GetSectionsWereModified())
 				{
-					VaildTrackEmitterHandleIds.Add(EmitterTrack->GetEmitterHandleViewModel()->GetId());
-					EmitterTrack->UpdateEmitterHandleFromTrackChange(NiagaraSequence->GetMovieScene()->GetTickResolution());
-					EmitterTrack->GetEmitterHandleViewModel()->GetEmitterViewModel()->GetOrCreateEditorData().Modify();
-					EmitterTrack->GetEmitterHandleViewModel()->GetEmitterViewModel()->GetOrCreateEditorData().SetPlaybackRange(GetEditorData().GetPlaybackRange());
-					if (EmitterTrack->GetDisplayName().ToString() != EmitterTrack->GetEmitterHandleViewModel()->GetNameText().ToString())
-					{
-						EmitterHandlesToRename.Add(TTuple<TSharedPtr<FNiagaraEmitterHandleViewModel>, FName>(EmitterTrack->GetEmitterHandleViewModel(), *EmitterTrack->GetDisplayName().ToString()));
-					}
-					if (EmitterTrack->GetAllSections().Num() > 1)
-					{
-						// If a section was duplicated, force a refresh to remove the duplicated section since that's not currently supported.
-						// TODO: Detect duplicated sections and create a duplicate emitter with the correct values.
-						bRefreshAllTracks = true;
-					}
+					// If there are no sections and the section collection was modified, the section was deleted so skip adding
+					// it's handle to the valid collection so that it's emitter is deleted.
+					continue;
+				}
+
+				VaildTrackEmitterHandleIds.Add(EmitterTrack->GetEmitterHandleViewModel()->GetId());
+				EmitterTrack->UpdateEmitterHandleFromTrackChange(NiagaraSequence->GetMovieScene()->GetTickResolution());
+				EmitterTrack->GetEmitterHandleViewModel()->GetEmitterViewModel()->GetOrCreateEditorData().Modify();
+				EmitterTrack->GetEmitterHandleViewModel()->GetEmitterViewModel()->GetOrCreateEditorData().SetPlaybackRange(GetEditorData().GetPlaybackRange());
+				if (EmitterTrack->GetDisplayName().ToString() != EmitterTrack->GetEmitterHandleViewModel()->GetNameText().ToString())
+				{
+					EmitterHandlesToRename.Add(TTuple<TSharedPtr<FNiagaraEmitterHandleViewModel>, FName>(EmitterTrack->GetEmitterHandleViewModel(), *EmitterTrack->GetDisplayName().ToString()));
+				}
+
+				if (EmitterTrack->GetAllSections().Num() > 1)
+				{
+					// If a section was duplicated, force a refresh to remove the duplicated section since that's not currently supported.
+					// TODO: Detect duplicated sections and create a duplicate emitter with the correct values.
+					bRefreshAllTracks = true;
 				}
 			}
 			else
