- Disable interpolation on dynamic parameters for sprites / meshes, they are constant across the whole triangle

#rnx

diff --git a/Engine/Plugins/FX/Niagara/Shaders/Private/NiagaraMeshVertexFactory.ush b/Engine/Plugins/FX/Niagara/Shaders/Private/NiagaraMeshVertexFactory.ush
--- a/Engine/Plugins/FX/Niagara/Shaders/Private/NiagaraMeshVertexFactory.ush
+++ b/Engine/Plugins/FX/Niagara/Shaders/Private/NiagaraMeshVertexFactory.ush
@@ -101,16 +101,16 @@ struct FVertexFactoryInterpolantsVSToPS
 #endif
 
 #if (DYNAMIC_PARAMETERS_MASK & 1)
-	float4 DynamicParameter : COLOR2;
+	nointerpolation float4 DynamicParameter : COLOR2;
 #endif
 #if (DYNAMIC_PARAMETERS_MASK & 2)
-	float4 DynamicParameter1 : COLOR3;
+	nointerpolation float4 DynamicParameter1 : COLOR3;
 #endif
 #if (DYNAMIC_PARAMETERS_MASK & 4)
-	float4 DynamicParameter2 : COLOR4;
+	nointerpolation float4 DynamicParameter2 : COLOR4;
 #endif
 #if (DYNAMIC_PARAMETERS_MASK & 8)
-	float4 DynamicParameter3 : COLOR5;
+	nointerpolation float4 DynamicParameter3 : COLOR5;
 #endif
 
 #if USE_PARTICLE_POSITION
diff --git a/Engine/Plugins/FX/Niagara/Shaders/Private/NiagaraSpriteVertexFactory.ush b/Engine/Plugins/FX/Niagara/Shaders/Private/NiagaraSpriteVertexFactory.ush
--- a/Engine/Plugins/FX/Niagara/Shaders/Private/NiagaraSpriteVertexFactory.ush
+++ b/Engine/Plugins/FX/Niagara/Shaders/Private/NiagaraSpriteVertexFactory.ush
@@ -76,16 +76,16 @@ struct FVertexFactoryInterpolantsVSToPS
 	float4	TangentToWorld2	: TANGENTTOWORLD2;
 	
 #if (DYNAMIC_PARAMETERS_MASK & 1)
-	float4 DynamicParameter	: PARTICLE_DYNAMIC_PARAM0;
+	nointerpolation float4 DynamicParameter	: PARTICLE_DYNAMIC_PARAM0;
 #endif
 #if (DYNAMIC_PARAMETERS_MASK & 2)
-	float4 DynamicParameter1	: PARTICLE_DYNAMIC_PARAM1;
+	nointerpolation float4 DynamicParameter1 : PARTICLE_DYNAMIC_PARAM1;
 #endif
 #if (DYNAMIC_PARAMETERS_MASK & 4)
-	float4 DynamicParameter2	: PARTICLE_DYNAMIC_PARAM2;
+	nointerpolation float4 DynamicParameter2 : PARTICLE_DYNAMIC_PARAM2;
 #endif
 #if (DYNAMIC_PARAMETERS_MASK & 8)
-	float4 DynamicParameter3	: PARTICLE_DYNAMIC_PARAM3;
+	nointerpolation float4 DynamicParameter3 : PARTICLE_DYNAMIC_PARAM3;
 #endif
 
 #if NEEDS_PARTICLE_COLOR
