Niagara - Renable dragging from the selection stack, and fix the drag tooltips so that they reset when they leave items.

diff --git a/Engine/Plugins/FX/Niagara/Source/NiagaraEditorWidgets/Private/NiagaraEditorWidgetsStyle.cpp b/Engine/Plugins/FX/Niagara/Source/NiagaraEditorWidgets/Private/NiagaraEditorWidgetsStyle.cpp
--- a/Engine/Plugins/FX/Niagara/Source/NiagaraEditorWidgets/Private/NiagaraEditorWidgetsStyle.cpp
+++ b/Engine/Plugins/FX/Niagara/Source/NiagaraEditorWidgets/Private/NiagaraEditorWidgetsStyle.cpp
@@ -121,6 +121,13 @@ TSharedRef< FSlateStyleSet > FNiagaraEditorWidgetsStyle::Create()
 	Style->Set("NiagaraEditor.SystemOverview.NodeBackgroundBorder", new BOX_PLUGIN_BRUSH("Icons/SystemOverviewNodeBackground", FMargin(1.0f / 4.0f)));
 
 	const FTableRowStyle& NormalTableRowStyle = FEditorStyle::Get().GetWidgetStyle<FTableRowStyle>("TableView.Row");
+
+	Style->Set("NiagaraEditor.Stack.TableViewRow", FTableRowStyle(NormalTableRowStyle)
+		.SetActiveBrush(*FEditorStyle::GetBrush("NoBrush"))
+		.SetActiveHoveredBrush(*FEditorStyle::GetBrush("NoBrush"))
+		.SetInactiveBrush(*FEditorStyle::GetBrush("NoBrush"))
+		.SetInactiveHoveredBrush(*FEditorStyle::GetBrush("NoBrush")));
+
 	Style->Set("NiagaraEditor.SystemOverview.TableViewRow", FTableRowStyle(NormalTableRowStyle)
 		.SetInactiveBrush(IMAGE_CORE_BRUSH("Common/Selection", Icon8x8, GetDefault<UEditorStyleSettings>()->GetSubduedSelectionColor())));
 
diff --git a/Engine/Plugins/FX/Niagara/Source/NiagaraEditorWidgets/Private/NiagaraEditorWidgetsUtilities.cpp b/Engine/Plugins/FX/Niagara/Source/NiagaraEditorWidgets/Private/NiagaraEditorWidgetsUtilities.cpp
--- a/Engine/Plugins/FX/Niagara/Source/NiagaraEditorWidgets/Private/NiagaraEditorWidgetsUtilities.cpp
+++ b/Engine/Plugins/FX/Niagara/Source/NiagaraEditorWidgets/Private/NiagaraEditorWidgetsUtilities.cpp
@@ -269,6 +269,18 @@ TSharedRef<FDragDropOperation> FNiagaraStackEditorWidgetsUtilities::ConstructDra
 	return DragDropOp;
 }
 
+void FNiagaraStackEditorWidgetsUtilities::HandleDragLeave(const FDragDropEvent& InDragDropEvent)
+{
+	if (InDragDropEvent.GetOperation().IsValid())
+	{
+		TSharedPtr<FDecoratedDragDropOp> DecoratedDragDropOp = InDragDropEvent.GetOperationAs<FDecoratedDragDropOp>();
+		if (DecoratedDragDropOp.IsValid())
+		{
+			DecoratedDragDropOp->ResetToDefaultToolTip();
+		}
+	}
+}
+
 TOptional<EItemDropZone> FNiagaraStackEditorWidgetsUtilities::RequestDropForStackEntry(const FDragDropEvent& InDragDropEvent, EItemDropZone InDropZone, UNiagaraStackEntry* InTargetEntry, UNiagaraStackEntry::EDropOptions DropOptions)
 {
 	TOptional<EItemDropZone> DropZone;
diff --git a/Engine/Plugins/FX/Niagara/Source/NiagaraEditorWidgets/Private/NiagaraEditorWidgetsUtilities.h b/Engine/Plugins/FX/Niagara/Source/NiagaraEditorWidgets/Private/NiagaraEditorWidgetsUtilities.h
--- a/Engine/Plugins/FX/Niagara/Source/NiagaraEditorWidgets/Private/NiagaraEditorWidgetsUtilities.h
+++ b/Engine/Plugins/FX/Niagara/Source/NiagaraEditorWidgets/Private/NiagaraEditorWidgetsUtilities.h
@@ -28,6 +28,8 @@ namespace FNiagaraStackEditorWidgetsUtilities
 
 	TSharedRef<FDragDropOperation> ConstructDragDropOperationForStackEntries(const TArray<UNiagaraStackEntry*>& DraggedEntries);
 
+	void HandleDragLeave(const FDragDropEvent& InDragDropEvent);
+
 	TOptional<EItemDropZone> RequestDropForStackEntry(const FDragDropEvent& InDragDropEvent, EItemDropZone InDropZone, UNiagaraStackEntry* InTargetEntry, UNiagaraStackEntry::EDropOptions DropOptions);
 
 	bool HandleDropForStackEntry(const FDragDropEvent& InDragDropEvent, EItemDropZone InDropZone, UNiagaraStackEntry* InTargetEntry, UNiagaraStackEntry::EDropOptions DropOptions);
diff --git a/Engine/Plugins/FX/Niagara/Source/NiagaraEditorWidgets/Private/SNiagaraOverviewStack.cpp b/Engine/Plugins/FX/Niagara/Source/NiagaraEditorWidgets/Private/SNiagaraOverviewStack.cpp
--- a/Engine/Plugins/FX/Niagara/Source/NiagaraEditorWidgets/Private/SNiagaraOverviewStack.cpp
+++ b/Engine/Plugins/FX/Niagara/Source/NiagaraEditorWidgets/Private/SNiagaraOverviewStack.cpp
@@ -38,6 +38,7 @@ class SNiagaraSystemOverviewEntryListRow : public STableRow<UNiagaraStackEntry*>
 	SLATE_BEGIN_ARGS(SNiagaraSystemOverviewEntryListRow) {}
 		SLATE_ATTRIBUTE(EVisibility, IssueIconVisibility)
 		SLATE_EVENT(FOnDragDetected, OnDragDetected)
+		SLATE_EVENT(FOnTableRowDragLeave, OnDragLeave)
 		SLATE_EVENT(FOnCanAcceptDrop, OnCanAcceptDrop)
 		SLATE_EVENT(FOnAcceptDrop, OnAcceptDrop)
 		SLATE_DEFAULT_SLOT(FArguments, Content);
@@ -87,6 +88,7 @@ class SNiagaraSystemOverviewEntryListRow : public STableRow<UNiagaraStackEntry*>
 		STableRow<UNiagaraStackEntry*>::Construct(STableRow<UNiagaraStackEntry*>::FArguments()
 			.Style(FNiagaraEditorWidgetsStyle::Get(), "NiagaraEditor.SystemOverview.TableViewRow")
 			.OnDragDetected(InArgs._OnDragDetected)
+			.OnDragLeave(InArgs._OnDragLeave)
 			.OnCanAcceptDrop(InArgs._OnCanAcceptDrop)
 			.OnAcceptDrop(InArgs._OnAcceptDrop)
 		[
@@ -452,6 +454,7 @@ TSharedRef<ITableRow> SNiagaraOverviewStack::OnGenerateRowForEntry(UNiagaraStack
 
 	return SNew(SNiagaraSystemOverviewEntryListRow, StackViewModel, Item, OwnerTable)
 		.OnDragDetected(this, &SNiagaraOverviewStack::OnRowDragDetected, TWeakObjectPtr<UNiagaraStackEntry>(Item))
+		.OnDragLeave(this, &SNiagaraOverviewStack::OnRowDragLeave)
 		.OnCanAcceptDrop(this, &SNiagaraOverviewStack::OnRowCanAcceptDrop)
 		.OnAcceptDrop(this, &SNiagaraOverviewStack::OnRowAcceptDrop)
 		.IssueIconVisibility(this, &SNiagaraOverviewStack::GetIssueIconVisibility)
@@ -549,6 +552,11 @@ FReply SNiagaraOverviewStack::OnRowDragDetected(const FGeometry& InGeometry, con
 	return FReply::Unhandled();
 }
 
+void SNiagaraOverviewStack::OnRowDragLeave(const FDragDropEvent& InDragDropEvent)
+{
+	FNiagaraStackEditorWidgetsUtilities::HandleDragLeave(InDragDropEvent);
+}
+
 TOptional<EItemDropZone> SNiagaraOverviewStack::OnRowCanAcceptDrop(const FDragDropEvent& InDragDropEvent, EItemDropZone InDropZone, UNiagaraStackEntry* InTargetEntry)
 {
 	return FNiagaraStackEditorWidgetsUtilities::RequestDropForStackEntry(InDragDropEvent, InDropZone, InTargetEntry, UNiagaraStackEntry::EDropOptions::Overview);
diff --git a/Engine/Plugins/FX/Niagara/Source/NiagaraEditorWidgets/Private/SNiagaraOverviewStack.h b/Engine/Plugins/FX/Niagara/Source/NiagaraEditorWidgets/Private/SNiagaraOverviewStack.h
--- a/Engine/Plugins/FX/Niagara/Source/NiagaraEditorWidgets/Private/SNiagaraOverviewStack.h
+++ b/Engine/Plugins/FX/Niagara/Source/NiagaraEditorWidgets/Private/SNiagaraOverviewStack.h
@@ -51,6 +51,8 @@ private:
 
 	FReply OnRowDragDetected(const FGeometry& InGeometry, const FPointerEvent& InPointerEvent, TWeakObjectPtr<UNiagaraStackEntry> InStackEntryWeak);
 
+	void OnRowDragLeave(const FDragDropEvent& InDragDropEvent);
+
 	TOptional<EItemDropZone> OnRowCanAcceptDrop(const FDragDropEvent& InDragDropEvent, EItemDropZone InDropZone, UNiagaraStackEntry* InTargetEntry);
 
 	FReply OnRowAcceptDrop(const FDragDropEvent& InDragDropEvent, EItemDropZone InDropZone, UNiagaraStackEntry* InTargetEntry);
diff --git a/Engine/Plugins/FX/Niagara/Source/NiagaraEditorWidgets/Private/SNiagaraStack.cpp b/Engine/Plugins/FX/Niagara/Source/NiagaraEditorWidgets/Private/SNiagaraStack.cpp
--- a/Engine/Plugins/FX/Niagara/Source/NiagaraEditorWidgets/Private/SNiagaraStack.cpp
+++ b/Engine/Plugins/FX/Niagara/Source/NiagaraEditorWidgets/Private/SNiagaraStack.cpp
@@ -267,7 +267,7 @@ void SNiagaraStack::Construct(const FArguments& InArgs, UNiagaraStackViewModel*
 				.OnGetChildren(this, &SNiagaraStack::OnGetChildren)
 				.TreeItemsSource(&StackViewModel->GetRootEntryAsArray())
 				.OnTreeViewScrolled(this, &SNiagaraStack::StackTreeScrolled)
-				.SelectionMode(ESelectionMode::None)
+				.SelectionMode(ESelectionMode::Single)
 				.OnItemToString_Debug_Static(&FNiagaraStackEditorWidgetsUtilities::StackEntryToStringForListDebug)
 			]
 		]
@@ -551,6 +551,11 @@ FReply SNiagaraStack::OnRowDragDetected(const FGeometry& InGeometry, const FPoin
 	return FReply::Unhandled();
 }
 
+void  SNiagaraStack::OnRowDragLeave(FDragDropEvent const& InDragDropEvent)
+{
+	FNiagaraStackEditorWidgetsUtilities::HandleDragLeave(InDragDropEvent);
+}
+
 TOptional<EItemDropZone> SNiagaraStack::OnRowCanAcceptDrop(const FDragDropEvent& InDragDropEvent, EItemDropZone InDropZone, UNiagaraStackEntry* InTargetEntry)
 {
 	return FNiagaraStackEditorWidgetsUtilities::RequestDropForStackEntry(InDragDropEvent, InDropZone, InTargetEntry, UNiagaraStackEntry::EDropOptions::None);
@@ -840,6 +845,7 @@ TSharedRef<SNiagaraStackTableRow> SNiagaraStack::ConstructContainerForItem(UNiag
 		.ValueColumnWidth(this, &SNiagaraStack::GetContentColumnWidth)
 		.OnValueColumnWidthChanged(this, &SNiagaraStack::OnContentColumnWidthChanged)
 		.OnDragDetected(this, &SNiagaraStack::OnRowDragDetected, Item)
+		.OnDragLeave(this, &SNiagaraStack::OnRowDragLeave)
 		.OnCanAcceptDrop(this, &SNiagaraStack::OnRowCanAcceptDrop)
 		.OnAcceptDrop(this, &SNiagaraStack::OnRowAcceptDrop)
 		.IssueIconVisibility(this, &SNiagaraStack::GetIssueIconVisibility);
diff --git a/Engine/Plugins/FX/Niagara/Source/NiagaraEditorWidgets/Private/SNiagaraStack.h b/Engine/Plugins/FX/Niagara/Source/NiagaraEditorWidgets/Private/SNiagaraStack.h
--- a/Engine/Plugins/FX/Niagara/Source/NiagaraEditorWidgets/Private/SNiagaraStack.h
+++ b/Engine/Plugins/FX/Niagara/Source/NiagaraEditorWidgets/Private/SNiagaraStack.h
@@ -98,6 +98,8 @@ private:
 	// Drag/Drop
 	FReply OnRowDragDetected(const FGeometry& InGeometry, const FPointerEvent& InPointerEvent, UNiagaraStackEntry* InStackEntry);
 
+	void OnRowDragLeave(FDragDropEvent const& InDragDropEvent);
+
 	TOptional<EItemDropZone> OnRowCanAcceptDrop(const FDragDropEvent& InDragDropEvent, EItemDropZone InDropZone, UNiagaraStackEntry* InTargetEntry);
 
 	FReply OnRowAcceptDrop(const FDragDropEvent& InDragDropEvent, EItemDropZone InDropZone, UNiagaraStackEntry* InTargetEntry);
diff --git a/Engine/Plugins/FX/Niagara/Source/NiagaraEditorWidgets/Private/Stack/SNiagaraStackTableRow.cpp b/Engine/Plugins/FX/Niagara/Source/NiagaraEditorWidgets/Private/Stack/SNiagaraStackTableRow.cpp
--- a/Engine/Plugins/FX/Niagara/Source/NiagaraEditorWidgets/Private/Stack/SNiagaraStackTableRow.cpp
+++ b/Engine/Plugins/FX/Niagara/Source/NiagaraEditorWidgets/Private/Stack/SNiagaraStackTableRow.cpp
@@ -47,7 +47,9 @@ void SNiagaraStackTableRow::Construct(const FArguments& InArgs, UNiagaraStackVie
 
 	ConstructInternal(
 		STableRow<UNiagaraStackEntry*>::FArguments()
+			.Style(FNiagaraEditorWidgetsStyle::Get(), "NiagaraEditor.Stack.TableViewRow")
 			.OnDragDetected(InArgs._OnDragDetected)
+			.OnDragLeave(InArgs._OnDragLeave)
 			.OnCanAcceptDrop(InArgs._OnCanAcceptDrop)
 			.OnAcceptDrop(InArgs._OnAcceptDrop)
 		, OwnerTree.ToSharedRef());
diff --git a/Engine/Plugins/FX/Niagara/Source/NiagaraEditorWidgets/Private/Stack/SNiagaraStackTableRow.h b/Engine/Plugins/FX/Niagara/Source/NiagaraEditorWidgets/Private/Stack/SNiagaraStackTableRow.h
--- a/Engine/Plugins/FX/Niagara/Source/NiagaraEditorWidgets/Private/Stack/SNiagaraStackTableRow.h
+++ b/Engine/Plugins/FX/Niagara/Source/NiagaraEditorWidgets/Private/Stack/SNiagaraStackTableRow.h
@@ -39,6 +39,7 @@ public:
 		SLATE_EVENT(FOnColumnWidthChanged, OnNameColumnWidthChanged)
 		SLATE_EVENT(FOnColumnWidthChanged, OnValueColumnWidthChanged)
 		SLATE_EVENT(FOnDragDetected, OnDragDetected)
+		SLATE_EVENT(FOnTableRowDragLeave, OnDragLeave)
 		SLATE_EVENT(FOnCanAcceptDrop, OnCanAcceptDrop)
 		SLATE_EVENT(FOnAcceptDrop, OnAcceptDrop)
 	SLATE_END_ARGS();
