Niagara - Fix a crash when dragging and dropping modules in the same stack.

diff --git a/Engine/Plugins/FX/Niagara/Source/NiagaraEditor/Private/ViewModels/Stack/NiagaraStackEntry.cpp b/Engine/Plugins/FX/Niagara/Source/NiagaraEditor/Private/ViewModels/Stack/NiagaraStackEntry.cpp
--- a/Engine/Plugins/FX/Niagara/Source/NiagaraEditor/Private/ViewModels/Stack/NiagaraStackEntry.cpp
+++ b/Engine/Plugins/FX/Niagara/Source/NiagaraEditor/Private/ViewModels/Stack/NiagaraStackEntry.cpp
@@ -329,6 +329,11 @@ UNiagaraStackEntry::FOnRequestFullRefresh& UNiagaraStackEntry::OnRequestFullRefr
 	return RequestFullRefreshDelegate;
 }
 
+const UNiagaraStackEntry::FOnRequestFullRefresh& UNiagaraStackEntry::OnRequestFullRefreshDeferred() const
+{
+	return RequestFullRefreshDeferredDelegate;
+}
+
 UNiagaraStackEntry::FOnRequestFullRefresh& UNiagaraStackEntry::OnRequestFullRefreshDeferred()
 {
 	return RequestFullRefreshDeferredDelegate;
diff --git a/Engine/Plugins/FX/Niagara/Source/NiagaraEditor/Private/ViewModels/Stack/NiagaraStackScriptItemGroup.cpp b/Engine/Plugins/FX/Niagara/Source/NiagaraEditor/Private/ViewModels/Stack/NiagaraStackScriptItemGroup.cpp
--- a/Engine/Plugins/FX/Niagara/Source/NiagaraEditor/Private/ViewModels/Stack/NiagaraStackScriptItemGroup.cpp
+++ b/Engine/Plugins/FX/Niagara/Source/NiagaraEditor/Private/ViewModels/Stack/NiagaraStackScriptItemGroup.cpp
@@ -781,8 +781,10 @@ TOptional<UNiagaraStackEntry::FDropRequestResponse> UNiagaraStackScriptItemGroup
 		FNiagaraStackGraphUtilities::RelayoutGraph(*TargetGraph);
 		TargetGraph->NotifyGraphNeedsRecompile();
 		GetSystemViewModel()->GetSelectionViewModel()->AddEntryToSelectionByDisplayedObjectDeferred(MovedModule);
-		RefreshChildren();
 	}
+
+	SourceModuleItem->OnRequestFullRefreshDeferred().Broadcast();
+	TargetChild.OnRequestFullRefreshDeferred().Broadcast();
 	return FDropRequestResponse(DropRequest.DropZone);
 }
 
diff --git a/Engine/Plugins/FX/Niagara/Source/NiagaraEditor/Private/ViewModels/Stack/NiagaraStackViewModel.cpp b/Engine/Plugins/FX/Niagara/Source/NiagaraEditor/Private/ViewModels/Stack/NiagaraStackViewModel.cpp
--- a/Engine/Plugins/FX/Niagara/Source/NiagaraEditor/Private/ViewModels/Stack/NiagaraStackViewModel.cpp
+++ b/Engine/Plugins/FX/Niagara/Source/NiagaraEditor/Private/ViewModels/Stack/NiagaraStackViewModel.cpp
@@ -664,6 +664,13 @@ void UNiagaraStackViewModel::RefreshTopLevelViewModels()
 	RootEntry->GetUnfilteredChildren(RootChildren);
 	for (UNiagaraStackEntry* RootChild : RootChildren)
 	{
+		if (RootChild->IsFinalized())
+		{
+			// It's possible for this to run when a system or emitter stack view model has updated it's children, but
+			// before the selection view model with the top level view models has refreshed and removed the finalized
+			// children in the selection so we need to guard against that here.
+			continue;
+		}
 		TSharedPtr<FTopLevelViewModel> TopLevelViewModel;
 		if (RootChild->GetEmitterViewModel().IsValid())
 		{
diff --git a/Engine/Plugins/FX/Niagara/Source/NiagaraEditor/Public/ViewModels/Stack/NiagaraStackEntry.h b/Engine/Plugins/FX/Niagara/Source/NiagaraEditor/Public/ViewModels/Stack/NiagaraStackEntry.h
--- a/Engine/Plugins/FX/Niagara/Source/NiagaraEditor/Public/ViewModels/Stack/NiagaraStackEntry.h
+++ b/Engine/Plugins/FX/Niagara/Source/NiagaraEditor/Public/ViewModels/Stack/NiagaraStackEntry.h
@@ -267,6 +267,8 @@ public:
 
 	FOnRequestFullRefresh& OnRequestFullRefresh();
 
+	const FOnRequestFullRefresh& OnRequestFullRefreshDeferred() const;
+
 	FOnRequestFullRefresh& OnRequestFullRefreshDeferred();
 
 	void RefreshChildren();
