better error reporting for fill texture methods

diff --git a/Engine/Plugins/FX/Niagara/Source/Niagara/Classes/NiagaraDataInterfaceGrid2DCollection.h b/Engine/Plugins/FX/Niagara/Source/Niagara/Classes/NiagaraDataInterfaceGrid2DCollection.h
--- a/Engine/Plugins/FX/Niagara/Source/Niagara/Classes/NiagaraDataInterfaceGrid2DCollection.h
+++ b/Engine/Plugins/FX/Niagara/Source/Niagara/Classes/NiagaraDataInterfaceGrid2DCollection.h
@@ -135,10 +135,10 @@ public:
 	// Fills a texture render target 2d with the current data from the simulation
 	// #todo(dmp): this will eventually go away when we formalize how data makes it out of Niagara
 	UFUNCTION(BlueprintCallable, Category = Niagara)
-	virtual void FillTexture2D(const UNiagaraComponent *Component, UTextureRenderTarget2D *dest, int AttributeIndex);
+	virtual bool FillTexture2D(const UNiagaraComponent *Component, UTextureRenderTarget2D *dest, int AttributeIndex);
 	
 	UFUNCTION(BlueprintCallable, Category = Niagara)
-	virtual void FillRawTexture2D(const UNiagaraComponent *Component, UTextureRenderTarget2D *Dest, int &TilesX, int &TilesY);
+	virtual bool FillRawTexture2D(const UNiagaraComponent *Component, UTextureRenderTarget2D *Dest, int &TilesX, int &TilesY);
 	
 	UFUNCTION(BlueprintCallable, Category = Niagara)
 	virtual void GetRawTextureSize(const UNiagaraComponent *Component, int &SizeX, int &SizeY);
diff --git a/Engine/Plugins/FX/Niagara/Source/Niagara/Private/NiagaraDataInterfaceGrid2DCollection.cpp b/Engine/Plugins/FX/Niagara/Source/Niagara/Private/NiagaraDataInterfaceGrid2DCollection.cpp
--- a/Engine/Plugins/FX/Niagara/Source/Niagara/Private/NiagaraDataInterfaceGrid2DCollection.cpp
+++ b/Engine/Plugins/FX/Niagara/Source/Niagara/Private/NiagaraDataInterfaceGrid2DCollection.cpp
@@ -438,20 +438,35 @@ void UNiagaraDataInterfaceGrid2DCollection::DestroyPerInstanceData(void* PerInst
 }
 
 UFUNCTION(BlueprintCallable, Category = Niagara)
-void UNiagaraDataInterfaceGrid2DCollection::FillTexture2D(const UNiagaraComponent *Component, UTextureRenderTarget2D *Dest, int AttributeIndex)
+bool UNiagaraDataInterfaceGrid2DCollection::FillTexture2D(const UNiagaraComponent *Component, UTextureRenderTarget2D *Dest, int AttributeIndex)
 {
 	FNiagaraDataInterfaceProxyGrid2DCollection* TProxy = GetProxyAs<FNiagaraDataInterfaceProxyGrid2DCollection>();
 
 	if (!Component)
 	{
-		return;
+		return false;
 	}
 
 	FNiagaraSystemInstance *SystemInstance = Component->GetSystemInstance();
 	if (!SystemInstance)
 	{
-		return;
+		return false;
+	}
+
+	// check valid attribute index
+	if (AttributeIndex < 0 || AttributeIndex >=NumAttributes)
+	{
+		return false;
+	}
+
+	// check dest size and type needs to be float
+	// #todo(dmp): don't hardcode float since we might do other stuff in the future
+	EPixelFormat RequiredTye = PF_R32_FLOAT;
+	if (Dest->SizeX != NumCellsX || Dest->SizeY != NumCellsY || Dest->GetFormat() != RequiredTye)
+	{
+		return false;
 	}
+
 	FNiagaraSystemInstanceID InstanceID = SystemInstance->GetId();
 
 	ENQUEUE_RENDER_COMMAND(FUpdateDIColorCurve)(
@@ -473,10 +488,12 @@ void UNiagaraDataInterfaceGrid2DCollection::FillTexture2D(const UNiagaraComponen
 			RHICmdList.CopyTexture(Grid2DInstanceData->CurrentData->GridBuffer.Buffer, Dest->Resource->TextureRHI, CopyInfo);
 		}
 	});
+
+	return true;
 }
 
 UFUNCTION(BlueprintCallable, Category = Niagara)
-void UNiagaraDataInterfaceGrid2DCollection::FillRawTexture2D(const UNiagaraComponent *Component, UTextureRenderTarget2D *Dest, int &TilesX, int &TilesY)
+bool UNiagaraDataInterfaceGrid2DCollection::FillRawTexture2D(const UNiagaraComponent *Component, UTextureRenderTarget2D *Dest, int &TilesX, int &TilesY)
 {
 	FNiagaraDataInterfaceProxyGrid2DCollection* TProxy = GetProxyAs<FNiagaraDataInterfaceProxyGrid2DCollection>();
 
@@ -484,7 +501,7 @@ void UNiagaraDataInterfaceGrid2DCollection::FillRawTexture2D(const UNiagaraCompo
 	{
 		TilesX = -1;
 		TilesY = -1;
-		return;
+		return false;
 	}
 
 	FNiagaraSystemInstance *SystemInstance = Component->GetSystemInstance();
@@ -492,21 +509,27 @@ void UNiagaraDataInterfaceGrid2DCollection::FillRawTexture2D(const UNiagaraCompo
 	{
 		TilesX = -1;
 		TilesY = -1;
-		return;
+		return false;
 	}
 	FNiagaraSystemInstanceID InstanceID = SystemInstance->GetId();
-
+	
+	Grid2DCollectionRWInstanceData* Grid2DInstanceData = TProxy->SystemInstancesToProxyData.Find(InstanceID);
+	if (!Grid2DInstanceData)
 	{
-		Grid2DCollectionRWInstanceData* Grid2DInstanceData = TProxy->SystemInstancesToProxyData.Find(InstanceID);
-		if (!Grid2DInstanceData)
-		{
-			TilesX = -1;
-			TilesY = -1;
-			return;
-		}
+		TilesX = -1;
+		TilesY = -1;
+		return false;
+	}
+
+	TilesX = Grid2DInstanceData->NumTilesX;
+	TilesY = Grid2DInstanceData->NumTilesY;	
 
-		TilesX = Grid2DInstanceData->NumTilesX;
-		TilesY = Grid2DInstanceData->NumTilesY;
+	// check dest size and type needs to be float
+	// #todo(dmp): don't hardcode float since we might do other stuff in the future
+	EPixelFormat RequiredTye = PF_R32_FLOAT;
+	if (Dest->SizeX != NumCellsX * TilesX || Dest->SizeY != NumCellsY * TilesY || Dest->GetFormat() != RequiredTye)
+	{
+		return false;
 	}
 
 	ENQUEUE_RENDER_COMMAND(FUpdateDIColorCurve)(
@@ -521,6 +544,8 @@ void UNiagaraDataInterfaceGrid2DCollection::FillRawTexture2D(const UNiagaraCompo
 			RHICmdList.CopyTexture(Grid2DInstanceData->CurrentData->GridBuffer.Buffer, Dest->Resource->TextureRHI, CopyInfo);
 		}
 	});
+
+	return true;
 }
 
 UFUNCTION(BlueprintCallable, Category = Niagara)
