CIS fix

diff --git a/Engine/Plugins/FX/Niagara/Source/Niagara/Private/NiagaraDataInterfaceGrid2DCollection.cpp b/Engine/Plugins/FX/Niagara/Source/Niagara/Private/NiagaraDataInterfaceGrid2DCollection.cpp
--- a/Engine/Plugins/FX/Niagara/Source/Niagara/Private/NiagaraDataInterfaceGrid2DCollection.cpp
+++ b/Engine/Plugins/FX/Niagara/Source/Niagara/Private/NiagaraDataInterfaceGrid2DCollection.cpp
@@ -532,16 +532,17 @@ bool UNiagaraDataInterfaceGrid2DCollection::FillRawTexture2D(const UNiagaraCompo
 		return false;
 	}
 
+	// #todo(dmp): pass grid2dinstancedata through?
 	ENQUEUE_RENDER_COMMAND(FUpdateDIColorCurve)(
 		[Dest, TProxy, InstanceID](FRHICommandListImmediate& RHICmdList)
 	{
-		Grid2DCollectionRWInstanceData* Grid2DInstanceData = TProxy->SystemInstancesToProxyData.Find(InstanceID);
+		Grid2DCollectionRWInstanceData* Grid2DInstanceDataTmp = TProxy->SystemInstancesToProxyData.Find(InstanceID);
 
-		if (Dest && Dest->Resource && Grid2DInstanceData && Grid2DInstanceData->CurrentData)
+		if (Dest && Dest->Resource && Grid2DInstanceDataTmp && Grid2DInstanceDataTmp->CurrentData)
 		{
 			
 			FRHICopyTextureInfo CopyInfo;
-			RHICmdList.CopyTexture(Grid2DInstanceData->CurrentData->GridBuffer.Buffer, Dest->Resource->TextureRHI, CopyInfo);
+			RHICmdList.CopyTexture(Grid2DInstanceDataTmp->CurrentData->GridBuffer.Buffer, Dest->Resource->TextureRHI, CopyInfo);
 		}
 	});
 
