remove clear logic on non iteration stages for grid2d DIs

diff --git a/Engine/Plugins/FX/Niagara/Source/Niagara/Private/NiagaraDataInterfaceGrid2DCollection.cpp b/Engine/Plugins/FX/Niagara/Source/Niagara/Private/NiagaraDataInterfaceGrid2DCollection.cpp
--- a/Engine/Plugins/FX/Niagara/Source/Niagara/Private/NiagaraDataInterfaceGrid2DCollection.cpp
+++ b/Engine/Plugins/FX/Niagara/Source/Niagara/Private/NiagaraDataInterfaceGrid2DCollection.cpp
@@ -671,10 +671,11 @@ void FNiagaraDataInterfaceProxyGrid2DCollection::PreStage(FRHICommandList& RHICm
 
 		// #todo(dmp): we might want to expose an option where we have buffers that are write only and need a clear (ie: no buffering like the neighbor grid).  They would be considered transient perhaps?  It'd be more
 		// memory efficient since it would theoretically not require any double buffering.
-		if (!Context.IsIterationStage)
-		{
-			ClearUAV(RHICmdList, ProxyData->DestinationData->GridBuffer, FLinearColor(0, 0, 0, 0));
-		}
+		// #todo(dmp): for now, if someone needs to clear then you have to do an iteration stage before to set clear values.  In the future we could have a hardware clear stage that we expose
+		// if (!Context.IsIterationStage)
+		// {
+		//	ClearUAV(RHICmdList, ProxyData->DestinationData->GridBuffer, FLinearColor(0, 0, 0, 0));
+		// }
 	}
 }
 
