Register thumbnail renders after engine loop init to avoid forcing early init of other modules using thumbnail renderer.

diff --git a/Engine/Plugins/FX/Niagara/Source/NiagaraEditor/Private/NiagaraEditorModule.cpp b/Engine/Plugins/FX/Niagara/Source/NiagaraEditor/Private/NiagaraEditorModule.cpp
--- a/Engine/Plugins/FX/Niagara/Source/NiagaraEditor/Private/NiagaraEditorModule.cpp
+++ b/Engine/Plugins/FX/Niagara/Source/NiagaraEditor/Private/NiagaraEditorModule.cpp
@@ -596,6 +596,8 @@ class FNiagaraSystemColorParameterTrackEditor : public FNiagaraSystemParameterTr
 
 void FNiagaraEditorModule::StartupModule()
 {
+	bThumbnailRenderersRegistered = false;
+
 	FHlslNiagaraTranslator::Init();
 	MenuExtensibilityManager = MakeShareable(new FExtensibilityManager);
 	ToolBarExtensibilityManager = MakeShareable(new FExtensibilityManager);
@@ -825,11 +827,6 @@ void FNiagaraEditorModule::StartupModule()
 		TEXT("Dumps data relevant to generating the compile id for an asset."),
 		FConsoleCommandWithArgsDelegate::CreateStatic(&DumpCompileIdDataForAsset));
 
-	if (GIsEditor)
-	{
-		UThumbnailManager::Get().RegisterCustomRenderer(UNiagaraEmitter::StaticClass(), UNiagaraEmitterThumbnailRenderer::StaticClass());
-		UThumbnailManager::Get().RegisterCustomRenderer(UNiagaraSystem::StaticClass(), UNiagaraSystemThumbnailRenderer::StaticClass());
-	}
 }
 
 
@@ -930,16 +927,22 @@ void FNiagaraEditorModule::ShutdownModule()
 		DumpCompileIdDataForAssetCommand = nullptr;
 	}
 
-	if (UObjectInitialized() && GIsEditor)
+	if (UObjectInitialized() && GIsEditor && bThumbnailRenderersRegistered)
 	{
 		UThumbnailManager::Get().UnregisterCustomRenderer(UNiagaraEmitter::StaticClass());
 		UThumbnailManager::Get().UnregisterCustomRenderer(UNiagaraSystem::StaticClass());
 	}
 }
 
-
 void FNiagaraEditorModule::OnPostEngineInit()
 {
+	if (GIsEditor)
+	{
+		UThumbnailManager::Get().RegisterCustomRenderer(UNiagaraEmitter::StaticClass(), UNiagaraEmitterThumbnailRenderer::StaticClass());
+		UThumbnailManager::Get().RegisterCustomRenderer(UNiagaraSystem::StaticClass(), UNiagaraSystemThumbnailRenderer::StaticClass());
+		bThumbnailRenderersRegistered = true;
+	}
+
 	// The editor should be valid at this point.. log a warning if not!
 	if (GEditor)
 	{
diff --git a/Engine/Plugins/FX/Niagara/Source/NiagaraEditor/Public/NiagaraEditorModule.h b/Engine/Plugins/FX/Niagara/Source/NiagaraEditor/Public/NiagaraEditorModule.h
--- a/Engine/Plugins/FX/Niagara/Source/NiagaraEditor/Public/NiagaraEditorModule.h
+++ b/Engine/Plugins/FX/Niagara/Source/NiagaraEditor/Public/NiagaraEditorModule.h
@@ -161,4 +161,6 @@ private:
 	IConsoleCommand* DumpCompileIdDataForAssetCommand;
 
 	FOnCheckScriptToolkitsShouldFocusGraphElement OnCheckScriptToolkitsShouldFocusGraphElement;
+
+	bool bThumbnailRenderersRegistered;
 };
