Fix graph editor commands init failing to register on Mac and Linux.

diff --git a/Engine/Plugins/FX/Niagara/Source/NiagaraEditorWidgets/Private/NiagaraEditorWidgetsModule.cpp b/Engine/Plugins/FX/Niagara/Source/NiagaraEditorWidgets/Private/NiagaraEditorWidgetsModule.cpp
--- a/Engine/Plugins/FX/Niagara/Source/NiagaraEditorWidgets/Private/NiagaraEditorWidgetsModule.cpp
+++ b/Engine/Plugins/FX/Niagara/Source/NiagaraEditorWidgets/Private/NiagaraEditorWidgetsModule.cpp
@@ -13,6 +13,8 @@
 
 #include "Modules/ModuleManager.h"
 #include "PropertyEditorModule.h"
+#include "GraphEditorActions.h"
+#include "NiagaraEditorCommands.h"
 
 IMPLEMENT_MODULE(FNiagaraEditorWidgetsModule, NiagaraEditorWidgets);
 
@@ -119,6 +121,9 @@ void FNiagaraEditorWidgetsModule::StartupModule()
 		TEXT("fx.NiagaraEditorWidgets.ReinitializeStyle"),
 		TEXT("Reinitializes the style for the niagara editor widgets module.  Used in conjuction with live coding for UI tweaks.  May crash the editor if style objects are in use."),
 		FConsoleCommandDelegate::CreateRaw(this, &FNiagaraEditorWidgetsModule::ReinitializeStyle));
+
+	FGraphEditorCommands::Register();
+	FNiagaraEditorCommands::Register();
 }
 
 void FNiagaraEditorWidgetsModule::ShutdownModule()
