Improve new parameter creation through ParameterMapView, now ensure we get unique names for new user parameters.
Fix bug with duplicate CompiledEmitterData being generated in undo/redo transactions, preventing the preview system from running.
Niagara parameter store entry items now check for reuse whether the entry to recycle has the same type and name as the incoming variable.

diff --git a/Engine/Plugins/FX/Niagara/Source/Niagara/Private/NiagaraSystem.cpp b/Engine/Plugins/FX/Niagara/Source/Niagara/Private/NiagaraSystem.cpp
--- a/Engine/Plugins/FX/Niagara/Source/Niagara/Private/NiagaraSystem.cpp
+++ b/Engine/Plugins/FX/Niagara/Source/Niagara/Private/NiagaraSystem.cpp
@@ -202,6 +202,11 @@ void UNiagaraSystem::Serialize(FArchive& Ar)
 		}
 		Ar << EmitterCompiledDataNum;
 
+		if (Ar.IsLoading())
+		{
+			// Clear out EmitterCompiledData when loading or else we will end up with duplicate entries. 
+			EmitterCompiledData.Reset();
+		}
 		for (int32 EmitterIndex = 0; EmitterIndex < EmitterCompiledDataNum; ++EmitterIndex)
 		{
 			if (Ar.IsLoading())
diff --git a/Engine/Plugins/FX/Niagara/Source/NiagaraEditor/Private/NiagaraEditorUtilities.cpp b/Engine/Plugins/FX/Niagara/Source/NiagaraEditor/Private/NiagaraEditorUtilities.cpp
--- a/Engine/Plugins/FX/Niagara/Source/NiagaraEditor/Private/NiagaraEditorUtilities.cpp
+++ b/Engine/Plugins/FX/Niagara/Source/NiagaraEditor/Private/NiagaraEditorUtilities.cpp
@@ -40,6 +40,7 @@
 #include "NiagaraNodeFunctionCall.h"
 #include "NiagaraParameterMapHistory.h"
 #include "ScopedTransaction.h"
+#include "NiagaraStackEditorData.h"
 
 #define LOCTEXT_NAMESPACE "FNiagaraEditorUtilities"
 
@@ -1238,4 +1239,28 @@ bool FNiagaraEditorUtilities::VerifyNameChangeForInputOrOutputNode(const UNiagar
 	return true;
 }
 
+bool FNiagaraEditorUtilities::AddParameter(FNiagaraVariable& NewParameterVariable, FNiagaraParameterStore& TargetParameterStore, UObject& ParameterStoreOwner, UNiagaraStackEditorData& StackEditorData)
+{
+	FScopedTransaction AddTransaction(LOCTEXT("AddParameter", "Add Parameter"));
+	ParameterStoreOwner.Modify();
+
+	TSet<FName> ExistingParameterStoreNames;
+	TArray<FNiagaraVariable> ParameterStoreVariables;
+	TargetParameterStore.GetParameters(ParameterStoreVariables);
+	for (const FNiagaraVariable& Var : ParameterStoreVariables)
+	{
+		ExistingParameterStoreNames.Add(Var.GetName());
+	}
+
+	FNiagaraEditorUtilities::ResetVariableToDefaultValue(NewParameterVariable);
+	NewParameterVariable.SetName(FNiagaraUtilities::GetUniqueName(NewParameterVariable.GetName(), ExistingParameterStoreNames));
+
+	bool bSuccess = TargetParameterStore.AddParameter(NewParameterVariable);
+	if (bSuccess)
+	{
+		StackEditorData.SetModuleInputIsRenamePending(NewParameterVariable.GetName().ToString(), true);
+	}
+	return bSuccess;
+}
+
 #undef LOCTEXT_NAMESPACE
diff --git a/Engine/Plugins/FX/Niagara/Source/NiagaraEditor/Private/ViewModels/Stack/NiagaraStackSystemSettingsGroup.cpp b/Engine/Plugins/FX/Niagara/Source/NiagaraEditor/Private/ViewModels/Stack/NiagaraStackSystemSettingsGroup.cpp
--- a/Engine/Plugins/FX/Niagara/Source/NiagaraEditor/Private/ViewModels/Stack/NiagaraStackSystemSettingsGroup.cpp
+++ b/Engine/Plugins/FX/Niagara/Source/NiagaraEditor/Private/ViewModels/Stack/NiagaraStackSystemSettingsGroup.cpp
@@ -84,17 +84,12 @@ public:
 	virtual void ExecuteAddAction(TSharedRef<INiagaraStackItemGroupAddAction> AddAction, int32 TargetIndex) override
 	{
 		TSharedRef<FParameterStoreGroupAddAction> ParameterAddAction = StaticCastSharedRef<FParameterStoreGroupAddAction>(AddAction);
-
-		FScopedTransaction AddTransaction(LOCTEXT("AddParameter", "Add Parameter"));
-		ParameterStoreOwner.Modify();
-
-		FNiagaraVariable ParameterVariable = ParameterAddAction->GetNewParameterVariable();
-		FNiagaraEditorUtilities::ResetVariableToDefaultValue(ParameterVariable);
-
-		ParameterStore.AddParameter(ParameterVariable);
-		StackEditorData.SetModuleInputIsRenamePending(ParameterVariable.GetName().ToString(), true);
-
-		OnItemAdded.ExecuteIfBound(ParameterVariable);
+		FNiagaraVariable NewParameterVariable = ParameterAddAction->GetNewParameterVariable();
+		bool bSuccess = FNiagaraEditorUtilities::AddParameter(NewParameterVariable, ParameterStore, ParameterStoreOwner, StackEditorData);
+		if (bSuccess)
+		{
+			OnItemAdded.ExecuteIfBound(NewParameterVariable);
+		}
 	}
 
 private:
@@ -183,7 +178,11 @@ void UNiagaraStackParameterStoreItem::RefreshChildrenInternal(const TArray<UNiag
 		for (FNiagaraVariable& Var : Variables)
 		{
 			UNiagaraStackParameterStoreEntry* ValueObjectEntry = FindCurrentChildOfTypeByPredicate<UNiagaraStackParameterStoreEntry>(CurrentChildren,
-				[=](UNiagaraStackParameterStoreEntry* CurrentEntry) { return CurrentEntry->GetDisplayName().ToString() == Var.GetName().ToString(); });
+				[=](UNiagaraStackParameterStoreEntry* CurrentEntry) { 
+				bool bSameName = CurrentEntry->GetDisplayName().ToString() == Var.GetName().ToString();
+				bool bSameType = CurrentEntry->GetInputType() == Var.GetType();
+				return bSameName && bSameType;
+			});
 
 			if (ValueObjectEntry == nullptr)
 			{
diff --git a/Engine/Plugins/FX/Niagara/Source/NiagaraEditor/Private/Widgets/SNiagaraParameterMapView.cpp b/Engine/Plugins/FX/Niagara/Source/NiagaraEditor/Private/Widgets/SNiagaraParameterMapView.cpp
--- a/Engine/Plugins/FX/Niagara/Source/NiagaraEditor/Private/Widgets/SNiagaraParameterMapView.cpp
+++ b/Engine/Plugins/FX/Niagara/Source/NiagaraEditor/Private/Widgets/SNiagaraParameterMapView.cpp
@@ -37,6 +37,8 @@
 #include "NiagaraScriptVariable.h"
 #include "Widgets/SPanel.h"
 #include "Widgets/Layout/SBox.h"
+#include "NiagaraSystemEditorData.h"
+#include "ViewModels/Stack/NiagaraStackSystemSettingsGroup.h"
 
 #define LOCTEXT_NAMESPACE "NiagaraParameterMapView"
 
@@ -258,43 +260,49 @@ bool SNiagaraParameterMapView::ParameterAddEnabled() const
 
 void SNiagaraParameterMapView::AddParameter(FNiagaraVariable NewVariable)
 {
-	TSet<FName> Names;
-	for (auto& GraphWeakPtr : Graphs)
-	{
-		UNiagaraGraph* Graph = GraphWeakPtr.Get();
-		for (const auto& ParameterElement : Graph->GetParameterReferenceMap())
-		{
-			Names.Add(ParameterElement.Key.GetName());
-		}
-	}
-	const FName NewUniqueName = FNiagaraUtilities::GetUniqueName(NewVariable.GetName(), Names);
-	NewVariable.SetName(NewUniqueName);
-
-	bool bAddedParameter = false;
-	// Check whether we have to add this parameter to the user exposed system parameters.
 	FNiagaraParameterHandle ParameterHandle;
+	bool bSuccess = false;
 	if (NiagaraParameterMapSectionID::OnGetSectionFromVariable(NewVariable, IsStaticSwitchParameter(NewVariable, Graphs), ParameterHandle) == NiagaraParameterMapSectionID::USER)
 	{
 		for (UObject* Object : SelectedScriptObjects->GetSelectedObjects())
 		{
 			if (UNiagaraSystem* System = Cast<UNiagaraSystem>(Object))
 			{
-				bAddedParameter = System->GetExposedParameters().AddParameter(NewVariable);
-				break;
+				UNiagaraSystemEditorData* SystemEditorData = CastChecked<UNiagaraSystemEditorData>(System->GetEditorData(), ECastCheckedType::NullChecked);
+				bSuccess = FNiagaraEditorUtilities::AddParameter(NewVariable, System->GetExposedParameters(), *System, SystemEditorData->GetStackEditorData());
 			}
 		}
 	}
-
-	if (!bAddedParameter && Graphs.Num() > 0)
+	else if (Graphs.Num() > 0)
 	{
+		TSet<FName> Names;
+		for (const TWeakObjectPtr<UNiagaraGraph>& GraphWeakPtr : Graphs)
+		{
+			if (GraphWeakPtr.IsValid())
+			{
+				UNiagaraGraph* Graph = GraphWeakPtr.Get();
+				for (const auto& ParameterElement : Graph->GetParameterReferenceMap())
+				{
+					Names.Add(ParameterElement.Key.GetName());
+				}
+			}
+		}
+		const FName NewUniqueName = FNiagaraUtilities::GetUniqueName(NewVariable.GetName(), Names);
+		NewVariable.SetName(NewUniqueName);
+
 		for (const TWeakObjectPtr<UNiagaraGraph>& GraphWeakPtr : Graphs)
 		{
 			if (GraphWeakPtr.IsValid())
 			{
 				UNiagaraGraph* Graph = GraphWeakPtr.Get();
 				Graph->AddParameter(NewVariable);
+				bSuccess = true;
 			}
 		}
+	}
+
+	if (bSuccess)
+	{
 		GraphActionMenu->RefreshAllActions(true);
 		GraphActionMenu->SelectItemByName(NewVariable.GetName());
 		GraphActionMenu->OnRequestRenameOnActionNode();
diff --git a/Engine/Plugins/FX/Niagara/Source/NiagaraEditor/Public/NiagaraEditorUtilities.h b/Engine/Plugins/FX/Niagara/Source/NiagaraEditor/Public/NiagaraEditorUtilities.h
--- a/Engine/Plugins/FX/Niagara/Source/NiagaraEditor/Public/NiagaraEditorUtilities.h
+++ b/Engine/Plugins/FX/Niagara/Source/NiagaraEditor/Public/NiagaraEditorUtilities.h
@@ -27,6 +27,7 @@ class UNiagaraNode;
 class UEdGraphSchema_Niagara;
 class UEdGraphPin;
 class FCompileConstantResolver;
+class UNiagaraStackEditorData;
 
 namespace FNiagaraEditorUtilities
 {
@@ -199,4 +200,14 @@ namespace FNiagaraEditorUtilities
 
 
 	bool VerifyNameChangeForInputOrOutputNode(const UNiagaraNode& NodeBeingChanged, FName OldName, FName NewName, FText& OutErrorMessage);
+
+	/**
+	 * Adds a new Parameter to a target ParameterStore with an undo/redo transaction and name collision handling.
+	 * @param NewParameterVariable The FNiagaraVariable to be added to TargetParameterStore. MUST be a unique object, do not pass an existing reference.
+	 * @param TargetParameterStore The ParameterStore to receive NewVariable.
+	 * @param ParameterStoreOwner The UObject to call Modify() on for the undo/redo transaction of adding NewVariable.
+	 * @param StackEditorData The editor data used to mark the newly added FNiagaraVariable in the Stack for renaming.
+	 * @returns Bool for whether adding the parameter succeeded.
+	 */
+	bool AddParameter(FNiagaraVariable& NewParameterVariable, FNiagaraParameterStore& TargetParameterStore, UObject& ParameterStoreOwner, UNiagaraStackEditorData& StackEditorData);
 };
