Niagara - Fix an issue where undoing adding an emitter would leave an invalid node in the overview graph.  Also fix other parts of the code which were following the same incorrect pattern.

diff --git a/Engine/Plugins/FX/Niagara/Source/NiagaraEditor/Private/NiagaraActions.cpp b/Engine/Plugins/FX/Niagara/Source/NiagaraEditor/Private/NiagaraActions.cpp
--- a/Engine/Plugins/FX/Niagara/Source/NiagaraEditor/Private/NiagaraActions.cpp
+++ b/Engine/Plugins/FX/Niagara/Source/NiagaraEditor/Private/NiagaraActions.cpp
@@ -175,6 +175,7 @@ FReply FNiagaraParameterGraphDragOperation::DroppedOnPanel(const TSharedRef<SWid
 void FNiagaraParameterGraphDragOperation::MakeGetMap(FNiagaraParameterNodeConstructionParams InParams)
 {
 	check(InParams.Graph);
+	InParams.Graph->Modify();
 	FGraphNodeCreator<UNiagaraNodeParameterMapGet> GetNodeCreator(*InParams.Graph);
 	UNiagaraNodeParameterMapGet* GetNode = GetNodeCreator.CreateNode();
 	GetNode->NodePosX = InParams.GraphPosition.X;
@@ -186,6 +187,7 @@ void FNiagaraParameterGraphDragOperation::MakeGetMap(FNiagaraParameterNodeConstr
 void FNiagaraParameterGraphDragOperation::MakeSetMap(FNiagaraParameterNodeConstructionParams InParams)
 {
 	check(InParams.Graph);
+	InParams.Graph->Modify();
 	FGraphNodeCreator<UNiagaraNodeParameterMapSet> SetNodeCreator(*InParams.Graph);
 	UNiagaraNodeParameterMapSet* SetNode = SetNodeCreator.CreateNode();
 	SetNode->NodePosX = InParams.GraphPosition.X;
diff --git a/Engine/Plugins/FX/Niagara/Source/NiagaraEditor/Private/NiagaraNodeAssignment.cpp b/Engine/Plugins/FX/Niagara/Source/NiagaraEditor/Private/NiagaraNodeAssignment.cpp
--- a/Engine/Plugins/FX/Niagara/Source/NiagaraEditor/Private/NiagaraNodeAssignment.cpp
+++ b/Engine/Plugins/FX/Niagara/Source/NiagaraEditor/Private/NiagaraNodeAssignment.cpp
@@ -355,6 +355,7 @@ void UNiagaraNodeAssignment::InitializeScript(UNiagaraScript* NewScript)
 			CreatedGraph = NewObject<UNiagaraGraph>(Source, NAME_None, RF_Transactional);
 			Source->NodeGraph = CreatedGraph;
 		}
+		CreatedGraph->Modify();
 		
 		TArray<UNiagaraNodeInput*> InputNodes;
 		CreatedGraph->FindInputNodes(InputNodes);
diff --git a/Engine/Plugins/FX/Niagara/Source/NiagaraEditor/Private/NiagaraSystemEditorData.cpp b/Engine/Plugins/FX/Niagara/Source/NiagaraEditor/Private/NiagaraSystemEditorData.cpp
--- a/Engine/Plugins/FX/Niagara/Source/NiagaraEditor/Private/NiagaraSystemEditorData.cpp
+++ b/Engine/Plugins/FX/Niagara/Source/NiagaraEditor/Private/NiagaraSystemEditorData.cpp
@@ -233,6 +233,7 @@ void UNiagaraSystemEditorData::SynchronizeOverviewGraphWithSystem(UNiagaraSystem
 	{
 		if(bSystemIsPlaceholder == false)
 		{
+			SystemOverviewGraph->Modify();
 			FGraphNodeCreator<UNiagaraOverviewNode> SystemOverviewNodeCreator(*SystemOverviewGraph);
 			UNiagaraOverviewNode* SystemOverviewNode = SystemOverviewNodeCreator.CreateNode(false);
 			SystemOverviewNode->Initialize(&OwnerSystem);
@@ -265,6 +266,7 @@ void UNiagaraSystemEditorData::SynchronizeOverviewGraphWithSystem(UNiagaraSystem
 		}
 		else
 		{
+			SystemOverviewGraph->Modify();
 			FGraphNodeCreator<UNiagaraOverviewNode> EmitterOverviewNodeCreator(*SystemOverviewGraph);
 			UNiagaraOverviewNode* EmitterOverviewNode = EmitterOverviewNodeCreator.CreateNode(false);
 			EmitterOverviewNode->Initialize(&OwnerSystem, EmitterHandle.GetId());
@@ -282,7 +284,7 @@ void UNiagaraSystemEditorData::SynchronizeOverviewGraphWithSystem(UNiagaraSystem
 			RightNode = EmitterOverviewNode;
 			EmitterNode = EmitterOverviewNode;
 		}
-
+		EmitterNode->Modify();
 		EmitterNode->SetEnabledState(EmitterHandle.GetIsEnabled() ? ENodeEnabledState::Enabled : ENodeEnabledState::Disabled);
 	}
 
diff --git a/Engine/Plugins/FX/Niagara/Source/NiagaraEditor/Private/ViewModels/Stack/NiagaraStackGraphUtilities.cpp b/Engine/Plugins/FX/Niagara/Source/NiagaraEditor/Private/ViewModels/Stack/NiagaraStackGraphUtilities.cpp
--- a/Engine/Plugins/FX/Niagara/Source/NiagaraEditor/Private/ViewModels/Stack/NiagaraStackGraphUtilities.cpp
+++ b/Engine/Plugins/FX/Niagara/Source/NiagaraEditor/Private/ViewModels/Stack/NiagaraStackGraphUtilities.cpp
@@ -718,6 +718,7 @@ UNiagaraNodeParameterMapSet& FNiagaraStackGraphUtilities::GetOrCreateStackFuncti
 	if (OverrideNode == nullptr)
 	{
 		UEdGraph* Graph = StackFunctionCall.GetGraph();
+		Graph->Modify();
 		FGraphNodeCreator<UNiagaraNodeParameterMapSet> ParameterMapSetNodeCreator(*Graph);
 		OverrideNode = ParameterMapSetNodeCreator.CreateNode();
 		ParameterMapSetNodeCreator.Finalize();
@@ -877,6 +878,7 @@ void FNiagaraStackGraphUtilities::SetLinkedValueHandleForFunctionInput(UEdGraphP
 
 	UNiagaraNodeParameterMapSet* OverrideNode = CastChecked<UNiagaraNodeParameterMapSet>(OverridePin.GetOwningNode());
 	UEdGraph* Graph = OverrideNode->GetGraph();
+	Graph->Modify();
 	FGraphNodeCreator<UNiagaraNodeParameterMapGet> GetNodeCreator(*Graph);
 	UNiagaraNodeParameterMapGet* GetNode = GetNodeCreator.CreateNode();
 	GetNodeCreator.Finalize();
@@ -908,6 +910,7 @@ void FNiagaraStackGraphUtilities::SetDataValueObjectForFunctionInput(UEdGraphPin
 
 	UNiagaraNodeParameterMapSet* OverrideNode = CastChecked<UNiagaraNodeParameterMapSet>(OverridePin.GetOwningNode());
 	UEdGraph* Graph = OverrideNode->GetGraph();
+	Graph->Modify();
 	FGraphNodeCreator<UNiagaraNodeInput> InputNodeCreator(*Graph);
 	UNiagaraNodeInput* InputNode = InputNodeCreator.CreateNode();
 	FNiagaraEditorUtilities::InitializeParameterInputNode(*InputNode, FNiagaraTypeDefinition(DataObjectType), CastChecked<UNiagaraGraph>(Graph), *DataObjectName);
@@ -930,6 +933,7 @@ void FNiagaraStackGraphUtilities::SetDynamicInputForFunctionInput(UEdGraphPin& O
 
 	UNiagaraNodeParameterMapSet* OverrideNode = CastChecked<UNiagaraNodeParameterMapSet>(OverridePin.GetOwningNode());
 	UEdGraph* Graph = OverrideNode->GetGraph();
+	Graph->Modify();
 	FGraphNodeCreator<UNiagaraNodeFunctionCall> FunctionCallNodeCreator(*Graph);
 	UNiagaraNodeFunctionCall* FunctionCallNode = FunctionCallNodeCreator.CreateNode();
 	FunctionCallNode->FunctionScript = DynamicInput;
@@ -983,6 +987,7 @@ void FNiagaraStackGraphUtilities::SetCustomExpressionForFunctionInput(UEdGraphPi
 	UEdGraph* Graph = OverrideNode->GetGraph();
 	const UEdGraphSchema_Niagara* Schema = CastChecked<UEdGraphSchema_Niagara>(OverrideNode->GetSchema());
 
+	Graph->Modify();
 	FGraphNodeCreator<UNiagaraNodeCustomHlsl> FunctionCallNodeCreator(*Graph);
 	UNiagaraNodeCustomHlsl* FunctionCallNode = FunctionCallNodeCreator.CreateNode();
 	FunctionCallNode->InitAsCustomHlslDynamicInput(Schema->PinToTypeDefinition(&OverridePin));
@@ -2130,6 +2135,7 @@ void FNiagaraStackGraphUtilities::RebuildEmitterNodes(UNiagaraSystem& System)
 	{
 		return;
 	}
+	SystemGraph->Modify();
 
 	TArray<UNiagaraNodeEmitter*> CurrentEmitterNodes;
 	SystemGraph->GetNodesOfClass<UNiagaraNodeEmitter>(CurrentEmitterNodes);
