re-enable auto clear on non iteration stages

diff --git a/Engine/Plugins/FX/Niagara/Source/Niagara/Private/NiagaraDataInterfaceGrid2DCollection.cpp b/Engine/Plugins/FX/Niagara/Source/Niagara/Private/NiagaraDataInterfaceGrid2DCollection.cpp
--- a/Engine/Plugins/FX/Niagara/Source/Niagara/Private/NiagaraDataInterfaceGrid2DCollection.cpp
+++ b/Engine/Plugins/FX/Niagara/Source/Niagara/Private/NiagaraDataInterfaceGrid2DCollection.cpp
@@ -670,12 +670,11 @@ void FNiagaraDataInterfaceProxyGrid2DCollection::PreStage(FRHICommandList& RHICm
 		// we start.  If a user wants to access data from the previous stage, then they can read from the current data.
 
 		// #todo(dmp): we might want to expose an option where we have buffers that are write only and need a clear (ie: no buffering like the neighbor grid).  They would be considered transient perhaps?  It'd be more
-		// memory efficient since it would theoretically not require any double buffering.
-		// #todo(dmp): for now, if someone needs to clear then you have to do an iteration stage before to set clear values.  In the future we could have a hardware clear stage that we expose
-		// if (!Context.IsIterationStage)
-		// {
-		//	ClearUAV(RHICmdList, ProxyData->DestinationData->GridBuffer, FLinearColor(0, 0, 0, 0));
-		// }
+		// memory efficient since it would theoretically not require any double buffering.		
+		if (!Context.IsIterationStage)
+		{
+			ClearUAV(RHICmdList, ProxyData->DestinationData->GridBuffer, FLinearColor(0, 0, 0, 0));
+		}
 	}
 }
 
