fixed grid2d data format due to merge issue

diff --git a/Engine/Plugins/FX/Niagara/Source/Niagara/Classes/NiagaraDataInterfaceGrid2DCollection.h b/Engine/Plugins/FX/Niagara/Source/Niagara/Classes/NiagaraDataInterfaceGrid2DCollection.h
--- a/Engine/Plugins/FX/Niagara/Source/Niagara/Classes/NiagaraDataInterfaceGrid2DCollection.h
+++ b/Engine/Plugins/FX/Niagara/Source/Niagara/Classes/NiagaraDataInterfaceGrid2DCollection.h
@@ -16,7 +16,7 @@ class Grid2DBuffer
 public:
 	Grid2DBuffer(int NumX, int NumY)
 	{		
-		GridBuffer.Initialize(16, NumX, NumY, EPixelFormat::PF_A32B32G32R32F);
+		GridBuffer.Initialize(4, NumX, NumY, EPixelFormat::PF_R32_FLOAT);
 	}
 
 	FTextureRWBuffer2D GridBuffer;	
