Fix potential crash identified by static analysis
#rnx

diff --git a/Engine/Plugins/FX/Niagara/Source/NiagaraEditor/Private/NiagaraHlslTranslator.cpp b/Engine/Plugins/FX/Niagara/Source/NiagaraEditor/Private/NiagaraHlslTranslator.cpp
--- a/Engine/Plugins/FX/Niagara/Source/NiagaraEditor/Private/NiagaraHlslTranslator.cpp
+++ b/Engine/Plugins/FX/Niagara/Source/NiagaraEditor/Private/NiagaraHlslTranslator.cpp
@@ -4759,9 +4759,9 @@ void FHlslNiagaraTranslator::HandleCustomHlslNode(UNiagaraNodeCustomHlsl* Custom
 		else if (Input.GetType().IsDataInterface())
 		{
 			UObject* const* FoundCDO = CompileData->CDOs.Find(Input.GetType().GetClass());
-			UNiagaraDataInterface* CDO = Cast<UNiagaraDataInterface>(*FoundCDO);
+			UNiagaraDataInterface* CDO = (FoundCDO ? Cast<UNiagaraDataInterface>(*FoundCDO) : nullptr);
 
-			if (FoundCDO == nullptr || CDO == nullptr)
+			if (CDO == nullptr)
 			{
 				// If the cdo wasn't found, the data interface was not passed through a parameter map and so it won't be bound correctly, so add a compile error
 				// and invalidate the signature.
