Fixed the matrix access in the spline DI

diff --git a/Engine/Plugins/FX/Niagara/Source/Niagara/Private/NiagaraDataInterfaceSpline.cpp b/Engine/Plugins/FX/Niagara/Source/Niagara/Private/NiagaraDataInterfaceSpline.cpp
--- a/Engine/Plugins/FX/Niagara/Source/Niagara/Private/NiagaraDataInterfaceSpline.cpp
+++ b/Engine/Plugins/FX/Niagara/Source/Niagara/Private/NiagaraDataInterfaceSpline.cpp
@@ -641,21 +641,21 @@ void UNiagaraDataInterfaceSpline::WriteTransform(const FMatrix& ToWrite, FVector
 	for (int32 i = 0; i < Context.NumInstances; ++i)
 	{
 		*Out00.GetDest() = ToWrite.M[0][0]; Out00.Advance();
-		*Out01.GetDest() = ToWrite.M[0][0]; Out01.Advance();
-		*Out02.GetDest() = ToWrite.M[0][0]; Out02.Advance();
-		*Out03.GetDest() = ToWrite.M[0][0]; Out03.Advance();
-		*Out04.GetDest() = ToWrite.M[0][0]; Out04.Advance();
-		*Out05.GetDest() = ToWrite.M[0][0]; Out05.Advance();
-		*Out06.GetDest() = ToWrite.M[0][0]; Out06.Advance();
-		*Out07.GetDest() = ToWrite.M[0][0]; Out07.Advance();
-		*Out08.GetDest() = ToWrite.M[0][0]; Out08.Advance();
-		*Out09.GetDest() = ToWrite.M[0][0]; Out09.Advance();
-		*Out10.GetDest() = ToWrite.M[0][0]; Out10.Advance();
-		*Out11.GetDest() = ToWrite.M[0][0]; Out11.Advance();
-		*Out12.GetDest() = ToWrite.M[0][0]; Out12.Advance();
-		*Out13.GetDest() = ToWrite.M[0][0]; Out13.Advance();
-		*Out14.GetDest() = ToWrite.M[0][0]; Out14.Advance();
-		*Out15.GetDest() = ToWrite.M[0][0]; Out15.Advance();
+		*Out01.GetDest() = ToWrite.M[0][1]; Out01.Advance();
+		*Out02.GetDest() = ToWrite.M[0][2]; Out02.Advance();
+		*Out03.GetDest() = ToWrite.M[0][3]; Out03.Advance();
+		*Out04.GetDest() = ToWrite.M[1][0]; Out04.Advance();
+		*Out05.GetDest() = ToWrite.M[1][1]; Out05.Advance();
+		*Out06.GetDest() = ToWrite.M[1][2]; Out06.Advance();
+		*Out07.GetDest() = ToWrite.M[1][3]; Out07.Advance();
+		*Out08.GetDest() = ToWrite.M[2][0]; Out08.Advance();
+		*Out09.GetDest() = ToWrite.M[2][1]; Out09.Advance();
+		*Out10.GetDest() = ToWrite.M[2][2]; Out10.Advance();
+		*Out11.GetDest() = ToWrite.M[2][3]; Out11.Advance();
+		*Out12.GetDest() = ToWrite.M[3][0]; Out12.Advance();
+		*Out13.GetDest() = ToWrite.M[3][1]; Out13.Advance();
+		*Out14.GetDest() = ToWrite.M[3][2]; Out14.Advance();
+		*Out15.GetDest() = ToWrite.M[3][3]; Out15.Advance();
 	}
 }
 
