Adds a new default initialization mode to Niagara module inputs where you can bind to available compatible parameters via a dropdown menu in the Selected Details panel inside the graph editor.

By selecting a Module parameter in the Parameters window there's now a new DefaultMode widget that can be "Value", "Binding" or "Custom".

When "Value" is selected it will display an appropriate widget below it to set the default value.
When "Binding" is selected it will display a dropdown widget of a list of available names to bind. Only built-in names and names explicitly added to the graph or Parameters list will show up
When "Custom is selected it will display neither of the widgets, indicating that the intiialization is to be done via a sub-graph. This will be streamlined later when the default pins are removed in favor of a separate initialization node.

The robustness of how all of this interacts with rapid iteration parameters requires a bit more work, but this is a good checkpoint as any. With respect to UE-82847, the bindings behave the same as any other non-value overrides for now.

Includes an accompanying autotest to test different permutations of initializing in the graph and overriding in the stack.

diff --git a/Engine/Plugins/FX/Niagara/Source/Niagara/Public/NiagaraCommon.h b/Engine/Plugins/FX/Niagara/Source/Niagara/Public/NiagaraCommon.h
--- a/Engine/Plugins/FX/Niagara/Source/Niagara/Public/NiagaraCommon.h
+++ b/Engine/Plugins/FX/Niagara/Source/Niagara/Public/NiagaraCommon.h
@@ -58,6 +58,22 @@ enum ENiagaraBaseTypes
 	NBT_Max,
 };
 
+// TODO: Custom will eventually mean that the default value or binding will be overridden by a subgraph default, i.e. expose it to a "Initialize variable" node. 
+// TODO: Should we add an "Uninitialized" entry, or is that too much friction? 
+// TODO: Why is there a "Max" entry on most enums, and do we need it?
+
+UENUM()
+enum class ENiagaraDefaultMode : uint8
+{
+	// Default initialize using a value widget in the Selected Details panel. 
+	Value = 0, 
+	// Default initialize using a dropdown widget in the Selected Details panel. 
+	Binding,   
+	// Default initialization is done using a sub-graph.
+	Custom,    
+	Max,       
+};
+
 UENUM()
 enum class ENiagaraSimTarget : uint8
 {
@@ -591,6 +607,31 @@ struct FNiagaraVariableDataInterfaceBinding
 	FNiagaraVariable BoundVariable;
 };
 
+/** Primarily a wrapper around an FName to be used for customizations in the Selected Details panel 
+    to select a default binding to initialize module inputs. The customization implementation
+    is FNiagaraScriptVariableBindingCustomization. */
+USTRUCT()
+struct FNiagaraScriptVariableBinding
+{
+	GENERATED_USTRUCT_BODY();
+
+	FNiagaraScriptVariableBinding() {}
+	FNiagaraScriptVariableBinding(const FNiagaraVariable& InVar) : Name(InVar.GetName())
+	{
+		
+	}
+	FNiagaraScriptVariableBinding(const FName& InName) : Name(InName)
+	{
+		
+	}
+
+	UPROPERTY(EditAnywhere, Category = "Variable")
+	FName Name;
+
+	FName GetName() const { return Name; }
+	void SetName(FName InName) { Name = InName; }
+	bool IsValid() const { return Name != NAME_None; }
+};
 
 namespace FNiagaraUtilities
 {
diff --git a/Engine/Plugins/FX/Niagara/Source/NiagaraEditor/Private/Customizations/NiagaraScriptVariableCustomization.cpp b/Engine/Plugins/FX/Niagara/Source/NiagaraEditor/Private/Customizations/NiagaraScriptVariableCustomization.cpp
--- a/Engine/Plugins/FX/Niagara/Source/NiagaraEditor/Private/Customizations/NiagaraScriptVariableCustomization.cpp
+++ b/Engine/Plugins/FX/Niagara/Source/NiagaraEditor/Private/Customizations/NiagaraScriptVariableCustomization.cpp
@@ -39,7 +39,6 @@ FNiagaraScriptVariableDetails::~FNiagaraScriptVariableDetails()
 	GEditor->UnregisterForUndo(this);
 }
 
-
 UEdGraphPin* FNiagaraScriptVariableDetails::GetDefaultPin()
 {
 	// TODO: We don't know the usage at this point, so we'll try each script type in order
@@ -92,6 +91,27 @@ void FNiagaraScriptVariableDetails::PostUndo(bool bSuccess)
 	}
 }
  
+void FNiagaraScriptVariableDetails::CustomizeDetails(const TSharedPtr<IDetailLayoutBuilder>& DetailBuilder)
+{
+	CachedDetailBuilder = DetailBuilder;
+	CustomizeDetails(*DetailBuilder);
+}
+
+void FNiagaraScriptVariableDetails::OnComboValueChanged()
+{
+	IDetailLayoutBuilder* DetailBuilderPtr = nullptr;
+	if (TSharedPtr<IDetailLayoutBuilder> LockedDetailBuilder = CachedDetailBuilder.Pin())
+	{
+		// WARNING: We do this because ForceRefresh will lock while pinning...
+		DetailBuilderPtr = LockedDetailBuilder.Get();
+	}
+	if (DetailBuilderPtr != nullptr)
+	{
+		TSharedRef<IPropertyUtilities> PropertyUtilities = DetailBuilderPtr->GetPropertyUtilities();
+		PropertyUtilities->ForceRefresh();
+	}
+}
+
 void FNiagaraScriptVariableDetails::CustomizeDetails(IDetailLayoutBuilder& DetailBuilder)
 {
 	static const FName CategoryName = TEXT("Default Value");
@@ -116,14 +136,16 @@ void FNiagaraScriptVariableDetails::CustomizeDetails(IDetailLayoutBuilder& Detai
 	}
  
 	IDetailCategoryBuilder& CategoryBuilder = DetailBuilder.EditCategory(CategoryName);		
+ 	
+ 	// NOTE: Automatically generated widgets from UProperties are placed below custom properties by default. 
+	//       In this case DefaultMode is just a built in combo box, while value widget is custom and added afterwards. 
+	//       This guarantees that the combo box always show above the value widget instead of at the bottom of the window. 
+	const TSharedPtr<IPropertyHandle> DefaultModeHandle = DetailBuilder.GetProperty(GET_MEMBER_NAME_CHECKED(UNiagaraScriptVariable, DefaultMode));
+	DefaultModeHandle->SetOnPropertyValueChanged(FSimpleDelegate::CreateSP(this, &FNiagaraScriptVariableDetails::OnComboValueChanged));
+
 	FNiagaraEditorModule& EditorModule = FNiagaraEditorModule::Get();
- 
 	if (Variable->Metadata.bIsStaticSwitch)
 	{
-		// Hide metadata UProperties that isn't useful for static switches
-		DetailBuilder.HideProperty(GET_MEMBER_NAME_CHECKED(UNiagaraScriptVariable, Metadata.EditCondition));
-		DetailBuilder.HideProperty(GET_MEMBER_NAME_CHECKED(UNiagaraScriptVariable, Metadata.VisibleCondition));
- 
 		TypeUtilityStaticSwitchValue = EditorModule.GetTypeUtilities(Variable->Variable.GetType());
 		if (TypeUtilityStaticSwitchValue && TypeUtilityStaticSwitchValue->CanCreateParameterEditor())
 		{
@@ -156,10 +178,13 @@ void FNiagaraScriptVariableDetails::CustomizeDetails(IDetailLayoutBuilder& Detai
 	}
 	else if (Variable->Variable.GetName().ToString().StartsWith("Module.")) 
 	{
+		DetailBuilder.HideProperty(DefaultModeHandle);
+		DetailBuilder.AddPropertyToCategory(DefaultModeHandle);
+
 		if (UEdGraphPin* Pin = GetDefaultPin())
 		{
 			TypeUtilityValue = EditorModule.GetTypeUtilities(Variable->Variable.GetType());
-			if (TypeUtilityValue && TypeUtilityValue->CanCreateParameterEditor())
+			if (TypeUtilityValue && TypeUtilityValue->CanCreateParameterEditor() && Variable->DefaultMode == ENiagaraDefaultMode::Value)
 			{
 				ParameterEditorValue = TypeUtilityValue->CreateParameterEditor(Variable->Variable.GetType());
 
@@ -191,22 +216,42 @@ void FNiagaraScriptVariableDetails::CustomizeDetails(IDetailLayoutBuilder& Detai
 			}
 		}
 	}
-	else
-	{
-		// Hide metadata UProperties that isn't useful for non module and non static switch variables
-		DetailBuilder.HideProperty(GET_MEMBER_NAME_CHECKED(UNiagaraScriptVariable, Metadata.CategoryName));
-		DetailBuilder.HideProperty(GET_MEMBER_NAME_CHECKED(UNiagaraScriptVariable, Metadata.bAdvancedDisplay));
-		DetailBuilder.HideProperty(GET_MEMBER_NAME_CHECKED(UNiagaraScriptVariable, Metadata.EditorSortPriority));
-		DetailBuilder.HideProperty(GET_MEMBER_NAME_CHECKED(UNiagaraScriptVariable, Metadata.bInlineEditConditionToggle));
+	
+
+	if (Variable->Metadata.bIsStaticSwitch) {
+		// Hide metadata UProperties that aren't useful for static switch variables
 		DetailBuilder.HideProperty(GET_MEMBER_NAME_CHECKED(UNiagaraScriptVariable, Metadata.EditCondition));
 		DetailBuilder.HideProperty(GET_MEMBER_NAME_CHECKED(UNiagaraScriptVariable, Metadata.VisibleCondition));
-		DetailBuilder.HideProperty(GET_MEMBER_NAME_CHECKED(UNiagaraScriptVariable, Metadata.PropertyMetaData));
+		DetailBuilder.HideProperty(GET_MEMBER_NAME_CHECKED(UNiagaraScriptVariable, DefaultBinding));
+		DetailBuilder.HideProperty(GET_MEMBER_NAME_CHECKED(UNiagaraScriptVariable, DefaultMode));
+		DetailBuilder.HideProperty(DefaultModeHandle); // TODO: Redundant?
 	}
- 
-	// bInlineEditConditionToggle is only useful for bool types
-	if (Variable->Variable.GetType() != FNiagaraTypeDefinition::GetBoolDef())
+	else
 	{
-		DetailBuilder.HideProperty(GET_MEMBER_NAME_CHECKED(UNiagaraScriptVariable, Metadata.bInlineEditConditionToggle));
+		if (!Variable->Variable.GetName().ToString().Contains("Module."))
+		{
+			DetailBuilder.HideProperty(GET_MEMBER_NAME_CHECKED(UNiagaraScriptVariable, DefaultMode));
+			DetailBuilder.HideProperty(GET_MEMBER_NAME_CHECKED(UNiagaraScriptVariable, Metadata.CategoryName));
+			DetailBuilder.HideProperty(GET_MEMBER_NAME_CHECKED(UNiagaraScriptVariable, Metadata.bAdvancedDisplay));
+			DetailBuilder.HideProperty(GET_MEMBER_NAME_CHECKED(UNiagaraScriptVariable, Metadata.EditorSortPriority));
+			DetailBuilder.HideProperty(GET_MEMBER_NAME_CHECKED(UNiagaraScriptVariable, Metadata.bInlineEditConditionToggle));
+			DetailBuilder.HideProperty(GET_MEMBER_NAME_CHECKED(UNiagaraScriptVariable, Metadata.EditCondition));
+			DetailBuilder.HideProperty(GET_MEMBER_NAME_CHECKED(UNiagaraScriptVariable, Metadata.VisibleCondition));
+			DetailBuilder.HideProperty(GET_MEMBER_NAME_CHECKED(UNiagaraScriptVariable, Metadata.PropertyMetaData));
+			DetailBuilder.HideProperty(GET_MEMBER_NAME_CHECKED(UNiagaraScriptVariable, DefaultBinding));
+			DetailBuilder.HideProperty(DefaultModeHandle);
+		}
+		
+		if (Variable->DefaultMode != ENiagaraDefaultMode::Binding)
+		{
+			DetailBuilder.HideProperty(GET_MEMBER_NAME_CHECKED(UNiagaraScriptVariable, DefaultBinding));
+		}
+
+		// bInlineEditConditionToggle is only useful for bool types
+		if (Variable->Variable.GetType() != FNiagaraTypeDefinition::GetBoolDef())
+		{
+			DetailBuilder.HideProperty(GET_MEMBER_NAME_CHECKED(UNiagaraScriptVariable, Metadata.bInlineEditConditionToggle));
+		}
 	}
 }
  
diff --git a/Engine/Plugins/FX/Niagara/Source/NiagaraEditor/Private/Customizations/NiagaraScriptVariableCustomization.h b/Engine/Plugins/FX/Niagara/Source/NiagaraEditor/Private/Customizations/NiagaraScriptVariableCustomization.h
--- a/Engine/Plugins/FX/Niagara/Source/NiagaraEditor/Private/Customizations/NiagaraScriptVariableCustomization.h
+++ b/Engine/Plugins/FX/Niagara/Source/NiagaraEditor/Private/Customizations/NiagaraScriptVariableCustomization.h
@@ -22,12 +22,14 @@ public:
  
 	/** IDetailCustomization interface */
 	virtual void CustomizeDetails( IDetailLayoutBuilder& DetailBuilder ) override;
-	
+	virtual void CustomizeDetails( const TSharedPtr<IDetailLayoutBuilder>& DetailBuilder ) override;
+
 	//~ Begin FEditorUndoClient Interface
 	virtual void PostUndo(bool bSuccess) override;
 	virtual void PostRedo(bool bSuccess) override { PostUndo(bSuccess); }
 	// End of FEditorUndoClient
  
+	void OnComboValueChanged();
 private:
 	void OnBeginValueChanged();
 	void OnEndValueChanged();
@@ -36,6 +38,8 @@ private:
 
 	UEdGraphPin* GetDefaultPin();
 
+	TWeakPtr<class IDetailLayoutBuilder> CachedDetailBuilder;
+
 	TWeakObjectPtr<class UNiagaraScriptVariable> Variable;
 	TSharedPtr<class INiagaraEditorTypeUtilities, ESPMode::ThreadSafe> TypeUtilityValue;
 	TSharedPtr<class INiagaraEditorTypeUtilities, ESPMode::ThreadSafe> TypeUtilityStaticSwitchValue;
diff --git a/Engine/Plugins/FX/Niagara/Source/NiagaraEditor/Private/Customizations/NiagaraTypeCustomizations.cpp b/Engine/Plugins/FX/Niagara/Source/NiagaraEditor/Private/Customizations/NiagaraTypeCustomizations.cpp
--- a/Engine/Plugins/FX/Niagara/Source/NiagaraEditor/Private/Customizations/NiagaraTypeCustomizations.cpp
+++ b/Engine/Plugins/FX/Niagara/Source/NiagaraEditor/Private/Customizations/NiagaraTypeCustomizations.cpp
@@ -22,6 +22,9 @@
 #include "Widgets/Input/STextComboBox.h"
 #include "NiagaraSystem.h"
 
+#include "NiagaraScriptVariable.h"
+#include "NiagaraConstants.h"
+
 #define LOCTEXT_NAMESPACE "FNiagaraVariableAttributeBindingCustomization"
 
 
@@ -477,6 +480,227 @@ void FNiagaraUserParameterBindingCustomization::CustomizeHeader(TSharedRef<IProp
 }
 
 
+FText FNiagaraScriptVariableBindingCustomization::GetCurrentText() const
+{
+	if (BaseGraph && TargetVariableBinding && TargetVariableBinding->IsValid())
+	{
+		return FText::FromName(TargetVariableBinding->Name);
+	}
+	return FText::FromString(TEXT("None"));
+}
+
+FText FNiagaraScriptVariableBindingCustomization::GetTooltipText() const
+{
+	if (BaseGraph && TargetVariableBinding && TargetVariableBinding->IsValid())
+	{
+		FText TooltipDesc = FText::Format(LOCTEXT("BindingTooltip", "Use the variable \"{0}\" if it is defined, otherwise use the type's default value."), FText::FromName(TargetVariableBinding->Name));
+		return TooltipDesc;
+	}
+	return FText::FromString(TEXT("There is no default binding selected."));
+}
+
+TSharedRef<SWidget> FNiagaraScriptVariableBindingCustomization::OnGetMenuContent() const
+{
+	FGraphActionMenuBuilder MenuBuilder; // TODO: Is this necessary? It's included in all the other implementations above, but it's never used. Spooky
+
+	return SNew(SBorder)
+		.BorderImage(FEditorStyle::GetBrush("Menu.Background"))
+		.Padding(5)
+		[
+			SNew(SBox)
+			[
+				SNew(SGraphActionMenu)
+				.OnActionSelected(this, &FNiagaraScriptVariableBindingCustomization::OnActionSelected)
+				.OnCreateWidgetForAction(SGraphActionMenu::FOnCreateWidgetForAction::CreateSP(this, &FNiagaraScriptVariableBindingCustomization::OnCreateWidgetForAction))
+				.OnCollectAllActions(this, &FNiagaraScriptVariableBindingCustomization::CollectAllActions)
+				.AutoExpandActionMenu(false)
+				.ShowFilterTextBox(true)
+			]
+		];
+}
+
+TArray<FName> FNiagaraScriptVariableBindingCustomization::GetNames() const
+{
+	// TODO: Only show Particles attributes for valid graphs,
+	//       i.e. only show Particles attributes for Particle scripts
+	//       and only show Emitter attributes for Emitter and Particle scripts.
+	TArray<FName> Names;
+
+	for (const FNiagaraParameterMapHistory& History : UNiagaraNodeParameterMapBase::GetParameterMaps(BaseGraph))
+	{
+		for (const FNiagaraVariable& Var : History.Variables)
+		{
+			FString Namespace = FNiagaraParameterMapHistory::GetNamespace(Var);
+			if (Namespace == TEXT("Module."))
+			{
+				// TODO: Skip module inputs for now. Does it make sense to bind module inputs to module inputs?
+				continue;
+			}
+			if (Var.GetType() == BaseScriptVariable->Variable.GetType())
+			{
+				Names.AddUnique(Var.GetName());
+			}
+		}
+	}
+
+	for (const auto& Var : BaseGraph->GetParameterReferenceMap())
+	{
+		FString Namespace = FNiagaraParameterMapHistory::GetNamespace(Var.Key);
+		if (Namespace == TEXT("Module."))
+		{
+			// TODO: Skip module inputs for now. Does it make sense to bind module inputs to module inputs?
+			continue;
+		}
+		if (Var.Key.GetType() == BaseScriptVariable->Variable.GetType())
+		{
+			Names.AddUnique(Var.Key.GetName());
+		}
+	}
+
+	for (const FNiagaraVariable& Var : FNiagaraConstants::GetEngineConstants())
+	{
+		if (Var.GetType() == BaseScriptVariable->Variable.GetType())
+		{
+			Names.AddUnique(Var.GetName());
+		}
+	}
+
+	for (const FNiagaraVariable& Var : FNiagaraConstants::GetCommonParticleAttributes())
+	{
+		if (Var.GetType() == BaseScriptVariable->Variable.GetType())
+		{
+			Names.AddUnique(Var.GetName());
+		}
+	}
+
+	return Names;
+}
+
+void FNiagaraScriptVariableBindingCustomization::CollectAllActions(FGraphActionListBuilderBase& OutAllActions)
+{
+	if (BaseGraph)
+	{
+		for (FName Name : GetNames())
+		{
+			const FText NameText = FText::FromName(Name);
+			const FText TooltipDesc = FText::Format(LOCTEXT("SetFunctionPopupTooltip", "Use the variable \"{0}\" "), NameText);
+			TSharedPtr<FNiagaraStackAssetAction_VarBind> NewNodeAction(
+				new FNiagaraStackAssetAction_VarBind(Name, FText(), NameText, TooltipDesc, 0, FText())
+			);
+			OutAllActions.AddAction(NewNodeAction);
+		}
+	}
+}
+
+TSharedRef<SWidget> FNiagaraScriptVariableBindingCustomization::OnCreateWidgetForAction(struct FCreateWidgetForActionData* const InCreateData)
+{
+	return SNew(SVerticalBox)
+		+ SVerticalBox::Slot()
+		.AutoHeight()
+		[
+			SNew(STextBlock)
+			.Text(InCreateData->Action->GetMenuDescription())
+			.ToolTipText(InCreateData->Action->GetTooltipDescription())
+		];
+}
+
+void FNiagaraScriptVariableBindingCustomization::OnActionSelected(const TArray< TSharedPtr<FEdGraphSchemaAction> >& SelectedActions, ESelectInfo::Type InSelectionType)
+{
+	if (InSelectionType == ESelectInfo::OnMouseClick || InSelectionType == ESelectInfo::OnKeyPress || SelectedActions.Num() == 0)
+	{
+		for (auto& CurrentAction : SelectedActions)
+		{
+			if (CurrentAction.IsValid())
+			{
+				FSlateApplication::Get().DismissAllMenus();
+				FNiagaraStackAssetAction_VarBind* EventSourceAction = (FNiagaraStackAssetAction_VarBind*)CurrentAction.Get();
+				ChangeSource(EventSourceAction->VarName);
+			}
+		}
+	}
+}
+
+void FNiagaraScriptVariableBindingCustomization::ChangeSource(FName InVarName)
+{
+	FScopedTransaction Transaction(FText::Format(LOCTEXT("ChangeBinding", " Change default binding to \"{0}\" "), FText::FromName(InVarName)));
+	TArray<UObject*> Objects;
+	PropertyHandle->GetOuterObjects(Objects);
+	for (UObject* Obj : Objects)
+	{
+		Obj->Modify();
+	}
+
+	PropertyHandle->NotifyPreChange();
+	TargetVariableBinding->Name = InVarName;
+	PropertyHandle->NotifyPostChange();
+	PropertyHandle->NotifyFinishedChangingProperties();
+}
+
+void FNiagaraScriptVariableBindingCustomization::CustomizeHeader(TSharedRef<IPropertyHandle> InPropertyHandle, FDetailWidgetRow& HeaderRow, IPropertyTypeCustomizationUtils& CustomizationUtils)
+{
+	PropertyHandle = InPropertyHandle;
+	TArray<UObject*> Objects;
+	PropertyHandle->GetOuterObjects(Objects);
+	bool bAddDefault = true;
+	if (Objects.Num() == 1)
+	{
+		BaseScriptVariable = Cast<UNiagaraScriptVariable>(Objects[0]);
+		if (BaseScriptVariable)
+		{
+		    BaseGraph = Cast<UNiagaraGraph>(BaseScriptVariable->GetOuter());
+			if (BaseGraph)
+			{
+				TargetVariableBinding = (FNiagaraScriptVariableBinding*)PropertyHandle->GetValueBaseAddress((uint8*)Objects[0]);
+
+				HeaderRow
+					.NameContent()
+					[
+						PropertyHandle->CreatePropertyNameWidget()
+					]
+				.ValueContent()
+					.MaxDesiredWidth(200.f)
+					[
+						SNew(SComboButton)
+						.OnGetMenuContent(this, &FNiagaraScriptVariableBindingCustomization::OnGetMenuContent)
+					.ContentPadding(1)
+					.ToolTipText(this, &FNiagaraScriptVariableBindingCustomization::GetTooltipText)
+					.ButtonContent()
+					[
+						SNew(STextBlock)
+						.Text(this, &FNiagaraScriptVariableBindingCustomization::GetCurrentText)
+					.Font(IDetailLayoutBuilder::GetDetailFont())
+					]
+					];
+				bAddDefault = false;
+			}
+			else
+			{
+				BaseScriptVariable = nullptr;
+			}
+		}
+		else
+		{
+			BaseGraph = nullptr;
+		}
+	}
+	
+	if (bAddDefault)
+	{
+		HeaderRow
+			.NameContent()
+			[
+				PropertyHandle->CreatePropertyNameWidget()
+			]
+		.ValueContent()
+			.MaxDesiredWidth(200.f)
+			[
+				SNew(STextBlock)
+				.Text(FText::FromString(FName::NameToDisplayString(Cast<UStructProperty>(PropertyHandle->GetProperty())->Struct->GetName(), false)))
+				.Font(IDetailLayoutBuilder::GetDetailFont())
+			];
+	}
+}
+
 
 
 #undef LOCTEXT_NAMESPACE
diff --git a/Engine/Plugins/FX/Niagara/Source/NiagaraEditor/Private/Customizations/NiagaraTypeCustomizations.h b/Engine/Plugins/FX/Niagara/Source/NiagaraEditor/Private/Customizations/NiagaraTypeCustomizations.h
--- a/Engine/Plugins/FX/Niagara/Source/NiagaraEditor/Private/Customizations/NiagaraTypeCustomizations.h
+++ b/Engine/Plugins/FX/Niagara/Source/NiagaraEditor/Private/Customizations/NiagaraTypeCustomizations.h
@@ -155,4 +155,36 @@ private:
 	class UNiagaraSystem* BaseSystem;
 	struct FNiagaraUserParameterBinding* TargetUserParameterBinding;
 
+};
+
+/** The primary goal of this class is to search through type matched and defined Niagara variables 
+    in the UNiagaraScriptVariable customization panel to provide a default binding for module inputs. */
+class FNiagaraScriptVariableBindingCustomization : public IPropertyTypeCustomization
+{
+public:
+	static TSharedRef<IPropertyTypeCustomization> MakeInstance()
+	{
+		return MakeShared<FNiagaraScriptVariableBindingCustomization>();
+	}
+
+	/** IPropertyTypeCustomization interface begin */
+	virtual void CustomizeHeader(TSharedRef<IPropertyHandle> StructPropertyHandle, class FDetailWidgetRow& HeaderRow, IPropertyTypeCustomizationUtils& StructCustomizationUtils) override;
+	virtual void CustomizeChildren(TSharedRef<IPropertyHandle> StructPropertyHandle, class IDetailChildrenBuilder& ChildBuilder, IPropertyTypeCustomizationUtils& StructCustomizationUtils) override {};
+	/** IPropertyTypeCustomization interface end */
+private:
+   /** Helpers */
+	FText GetCurrentText() const;
+	FText GetTooltipText() const;
+	TSharedRef<SWidget> OnGetMenuContent() const;
+	TArray<FName> GetNames() const;
+	void ChangeSource(FName InVarName);
+	void CollectAllActions(FGraphActionListBuilderBase& OutAllActions);
+	TSharedRef<SWidget> OnCreateWidgetForAction(struct FCreateWidgetForActionData* const InCreateData);
+	void OnActionSelected(const TArray< TSharedPtr<FEdGraphSchemaAction> >& SelectedActions, ESelectInfo::Type InSelectionType);
+
+	/** State */
+	TSharedPtr<IPropertyHandle> PropertyHandle;
+	class UNiagaraGraph* BaseGraph;
+	class UNiagaraScriptVariable* BaseScriptVariable;
+	struct FNiagaraScriptVariableBinding* TargetVariableBinding;
 };
\ No newline at end of file
diff --git a/Engine/Plugins/FX/Niagara/Source/NiagaraEditor/Private/NiagaraEditorModule.cpp b/Engine/Plugins/FX/Niagara/Source/NiagaraEditor/Private/NiagaraEditorModule.cpp
--- a/Engine/Plugins/FX/Niagara/Source/NiagaraEditor/Private/NiagaraEditorModule.cpp
+++ b/Engine/Plugins/FX/Niagara/Source/NiagaraEditor/Private/NiagaraEditorModule.cpp
@@ -657,6 +657,10 @@ void FNiagaraEditorModule::StartupModule()
 		FOnGetPropertyTypeCustomizationInstance::CreateStatic(&FNiagaraVariableAttributeBindingCustomization::MakeInstance)
 	);
 	
+	PropertyModule.RegisterCustomPropertyTypeLayout("NiagaraScriptVariableBinding",
+		FOnGetPropertyTypeCustomizationInstance::CreateStatic(&FNiagaraScriptVariableBindingCustomization::MakeInstance)
+	);
+		
 	PropertyModule.RegisterCustomPropertyTypeLayout("NiagaraUserParameterBinding",
 		FOnGetPropertyTypeCustomizationInstance::CreateStatic(&FNiagaraUserParameterBindingCustomization::MakeInstance)
 	);
diff --git a/Engine/Plugins/FX/Niagara/Source/NiagaraEditor/Private/NiagaraGraph.cpp b/Engine/Plugins/FX/Niagara/Source/NiagaraEditor/Private/NiagaraGraph.cpp
--- a/Engine/Plugins/FX/Niagara/Source/NiagaraEditor/Private/NiagaraGraph.cpp
+++ b/Engine/Plugins/FX/Niagara/Source/NiagaraEditor/Private/NiagaraGraph.cpp
@@ -936,6 +936,15 @@ bool UNiagaraGraph::RenameParameter(const FNiagaraVariable& Parameter, FName New
 	{
 		if (!bFromStaticSwitch || (bFromStaticSwitch && NumReferences <= 1))
 		{
+			// Rename all the bindings that point to the old parameter 
+			for (auto It : VariableToScriptVariable)
+			{
+				UNiagaraScriptVariable* Variable = It.Value;
+				if (Variable && Variable->DefaultBinding.GetName() == Parameter.GetName())
+				{
+					Variable->DefaultBinding.SetName(NewParameter.GetName());
+				}
+			}
 			VariableToScriptVariable.Remove(Parameter);
 		}
 	}
@@ -1528,6 +1537,18 @@ void UNiagaraGraph::RefreshParameterReferences() const
 		}
 	}
 
+	// Add reference to all variables that are default bound to
+	for (auto It = VariableToScriptVariable.CreateConstIterator(); It; ++It)
+	{
+		UNiagaraScriptVariable* Variable = It.Value();
+		if (Variable->DefaultMode != ENiagaraDefaultMode::Binding || !Variable->DefaultBinding.IsValid())
+		{
+			continue;
+		}
+
+		AddBindingParameterReference(FNiagaraVariable(Variable->Variable.GetType(), Variable->DefaultBinding.GetName()));
+	}
+
 	// If there were any previous parameters which didn't have any references added, remove them here.
 	for (const FNiagaraVariable& UnreferencedParameterToRemove : CandidateUnreferencedParametersToRemove)
 	{
@@ -1550,11 +1571,10 @@ void UNiagaraGraph::RefreshParameterReferences() const
 	for (auto& ParameterToReferences : ParameterToReferencesMap)
 	{
 		UNiagaraScriptVariable* Variable = GetScriptVariable(ParameterToReferences.Key);
-		if (Variable)
+		if (!Variable)
 		{
-			continue;
+			AddMissingScriptVariable(ParameterToReferences.Key);
 		}
-		AddMissingScriptVariable(ParameterToReferences.Key);
 	}
 	
 	static const auto UseShaderStagesCVar = IConsoleManager::Get().FindConsoleVariable(TEXT("fx.UseShaderStages"));
diff --git a/Engine/Plugins/FX/Niagara/Source/NiagaraEditor/Private/NiagaraHlslTranslator.cpp b/Engine/Plugins/FX/Niagara/Source/NiagaraEditor/Private/NiagaraHlslTranslator.cpp
--- a/Engine/Plugins/FX/Niagara/Source/NiagaraEditor/Private/NiagaraHlslTranslator.cpp
+++ b/Engine/Plugins/FX/Niagara/Source/NiagaraEditor/Private/NiagaraHlslTranslator.cpp
@@ -49,6 +49,7 @@
 
 #include "NiagaraEditorSettings.h"
 #include "NiagaraNodeStaticSwitch.h"
+#include "NiagaraScriptVariable.h"
 
 #define LOCTEXT_NAMESPACE "NiagaraCompiler"
 
@@ -454,7 +455,7 @@ void FHlslNiagaraTranslator::BuildMissingDefaults()
 			const UEdGraphPin* DefaultPin = UniqueVarToDefaultPin.FindChecked(Var);
 			bool bWriteToParamMapEntries = UniqueVarToWriteToParamMap.FindChecked(Var);
 			int32 OutputChunkId = INDEX_NONE;
-			HandleParameterRead(ActiveStageIdx, Var, DefaultPin, DefaultPin != nullptr ? Cast<UNiagaraNode>(DefaultPin->GetOwningNode()) : nullptr, OutputChunkId, !bWriteToParamMapEntries);
+			HandleParameterRead(ActiveStageIdx, Var, DefaultPin, DefaultPin != nullptr ? Cast<UNiagaraNode>(DefaultPin->GetOwningNode()) : nullptr, OutputChunkId, nullptr, !bWriteToParamMapEntries);
 		}
 
 		DeferredVariablesMissingDefault.Empty();
@@ -3149,7 +3150,7 @@ void FHlslNiagaraTranslator::InitializeParameterMapDefaults(int32 ParamMapHistor
 			// During the initial pass, only support constants for the default pin.
 			if (!FNiagaraParameterMapHistory::IsInitialValue(Var) && (DefaultPin == nullptr || DefaultPin->LinkedTo.Num() == 0))
 			{
-				HandleParameterRead(ParamMapHistoryIdx, Var, DefaultPin, DefaultPin != nullptr ? Cast<UNiagaraNode>(DefaultPin->GetOwningNode()) : nullptr, OutputChunkId, !bWriteToParamMapEntries);
+				HandleParameterRead(ParamMapHistoryIdx, Var, DefaultPin, DefaultPin != nullptr ? Cast<UNiagaraNode>(DefaultPin->GetOwningNode()) : nullptr, OutputChunkId, nullptr, !bWriteToParamMapEntries);
 				UniqueVarToChunk.Add(Var, OutputChunkId);
 			}
 			else if (FNiagaraParameterMapHistory::IsInitialValue(Var))
@@ -3903,12 +3904,106 @@ void FHlslNiagaraTranslator::ParameterMapGet(UNiagaraNodeParameterMapGet* GetNod
 				OutputTypeDefinition.IsDataInterface() == false;
 			FNiagaraVariable Var = Schema->PinToNiagaraVariable(OutputPins[i], bNeedsValue);
 
-			HandleParameterRead(ParamMapHistoryIdx, Var, GetNode->GetDefaultPin(OutputPins[i]), GetNode, Outputs[i]);
+			UNiagaraGraph* SourceGraph = GetNode->GetNiagaraGraph();
+			UNiagaraScriptVariable* Variable = SourceGraph->GetScriptVariable(Var);
+
+			bool bFoundBinding = false;
+			if (Var.GetName().ToString().Contains(TEXT("Module.")))
+			{
+				if (Variable && Variable->DefaultMode == ENiagaraDefaultMode::Binding)
+				{
+					FNiagaraScriptVariableBinding Bind = Variable->DefaultBinding;
+					if (Bind.IsValid())
+					{
+						int LastSetChunkIdx = INDEX_NONE;
+						for (auto& UniqueVar : UniqueVarToChunk)
+						{
+							if (UniqueVar.Key.IsEquivalent(FNiagaraVariable(Var.GetType(), Var.GetName())))
+							{
+								LastSetChunkIdx = UniqueVar.Value;
+								break;
+							}
+						}
+
+						Var = ActiveHistoryForFunctionCalls.ResolveAliases(Var);
+
+						if (LastSetChunkIdx == INDEX_NONE)
+						{
+							for (auto HistoryVariable : ParamMapHistories[ParamMapHistoryIdx].Variables)
+							{
+								if (HistoryVariable.IsEquivalent(FNiagaraVariable(Var.GetType(), Bind.GetName())))
+								{
+									FString SanitizedName = GetParameterMapInstanceName(ActiveStageIdx) + TEXT(".") + GetSanitizedSymbolName(Var.GetName().ToString());
+									LastSetChunkIdx = AddSourceChunk(SanitizedName, Var.GetType());
+									break;
+								}
+							}
+						}
+
+						if (LastSetChunkIdx == INDEX_NONE)
+						{
+							FString SanitizedName = GetParameterMapInstanceName(ActiveStageIdx) + TEXT(".") + GetSanitizedSymbolName(Var.GetName().ToString());
+							LastSetChunkIdx = CodeChunks.IndexOfByPredicate(
+								[&](const FNiagaraCodeChunk& Chunk)
+								{
+									return Chunk.Mode == ENiagaraCodeChunkMode::Source && Chunk.SymbolName == SanitizedName && Chunk.Type == Var.GetType();
+								}
+							);
+						}
+
+						if (LastSetChunkIdx == INDEX_NONE)
+						{
+							FString SanitizedName = GetParameterMapInstanceName(ActiveStageIdx) + TEXT(".") + GetSanitizedSymbolName(Var.GetName().ToString());
+							LastSetChunkIdx = CodeChunks.IndexOfByPredicate(
+								[&](const FNiagaraCodeChunk& Chunk)
+								{
+									return Chunk.Mode == ENiagaraCodeChunkMode::Body && Chunk.SymbolName == SanitizedName && Chunk.Type == Var.GetType();
+								}
+							);
+						}
+
+						if (LastSetChunkIdx != INDEX_NONE && Var.GetType().GetClass() == nullptr)
+						{
+							int32 VarIdx = ParamMapHistories[ParamMapHistoryIdx].FindVariableByName(Var.GetName());
+							if (VarIdx != INDEX_NONE && VarIdx < ParamMapSetVariablesToChunks[ParamMapHistoryIdx].Num())
+							{
+								ParamMapSetVariablesToChunks[ParamMapHistoryIdx][VarIdx] = LastSetChunkIdx;
+								Outputs[i] = LastSetChunkIdx;
+								ParamMapDefinedAttributesToNamespaceVars.FindOrAdd(Var.GetName()) = Var;
+							}
+							
+							bFoundBinding = true;
+						}
+						else
+						{
+							int Out;
+							if (ParameterMapRegisterExternalConstantNamespaceVariable(FNiagaraVariable(Var.GetType(), Bind.GetName()), GetNode, ActiveStageIdx, Out, GetNode->GetDefaultPin(OutputPins[i])))
+							{
+								LastSetChunkIdx = Out;
+								int32 VarIdx = ParamMapHistories[ParamMapHistoryIdx].FindVariableByName(Var.GetName());
+								if (VarIdx != INDEX_NONE && VarIdx < ParamMapSetVariablesToChunks[ParamMapHistoryIdx].Num())
+								{
+									ParamMapSetVariablesToChunks[ParamMapHistoryIdx][VarIdx] = LastSetChunkIdx;
+									Outputs[i] = LastSetChunkIdx;
+									ParamMapDefinedAttributesToNamespaceVars.FindOrAdd(Var.GetName()) = Var;
+								}
+
+								bFoundBinding = true;
+							}
+						}
+					}
+				}
+			}
+
+			if (!bFoundBinding)
+			{
+				HandleParameterRead(ParamMapHistoryIdx, Var, GetNode->GetDefaultPin(OutputPins[i]), GetNode, Outputs[i], nullptr);
+			}
 		}
 	}
 }
 
-void FHlslNiagaraTranslator::HandleParameterRead(int32 ParamMapHistoryIdx, const FNiagaraVariable& InVar, const UEdGraphPin* DefaultPin, UNiagaraNode* ErrorNode, int32& OutputChunkId, bool bTreatAsUnknownParameterMap)
+void FHlslNiagaraTranslator::HandleParameterRead(int32 ParamMapHistoryIdx, const FNiagaraVariable& InVar, const UEdGraphPin* DefaultPin, UNiagaraNode* ErrorNode, int32& OutputChunkId, UNiagaraScriptVariable* Variable, bool bTreatAsUnknownParameterMap)
 {
 	FString ParameterMapInstanceName = GetParameterMapInstanceName(ParamMapHistoryIdx);
 	FNiagaraVariable Var = InVar;
@@ -3918,9 +4013,6 @@ void FHlslNiagaraTranslator::HandleParameterRead(int32 ParamMapHistoryIdx, const
 	}
 
 	// If this is a System parameter, just wire in the system appropriate system attribute.
-	FString VarName = Var.GetName().ToString();
-	FString SymbolName = GetSanitizedSymbolName(VarName);
-
 	bool bIsPerInstanceAttribute = false;
 	bool bIsCandidateForRapidIteration = false;
 	const UEdGraphPin* InputPin = DefaultPin;
@@ -3948,7 +4040,12 @@ void FHlslNiagaraTranslator::HandleParameterRead(int32 ParamMapHistoryIdx, const
 	}
 	else if (FNiagaraParameterMapHistory::IsAliasedModuleParameter(Var) && ActiveHistoryForFunctionCalls.InTopLevelFunctionCall(CompileOptions.TargetUsage))
 	{
-		if (InputPin != nullptr && InputPin->LinkedTo.Num() == 0 && Var.GetType() != FNiagaraTypeDefinition::GetBoolDef() && !Var.GetType().IsEnum() && !Var.GetType().IsDataInterface())
+		if (Variable && Variable->DefaultMode == ENiagaraDefaultMode::Binding && Variable->DefaultBinding.IsValid())
+		{
+			bIsCandidateForRapidIteration = false;
+			
+		}
+		else if (InputPin != nullptr && InputPin->LinkedTo.Num() == 0 && Var.GetType() != FNiagaraTypeDefinition::GetBoolDef() && !Var.GetType().IsEnum() && !Var.GetType().IsDataInterface())
 		{
 			bIsCandidateForRapidIteration = true;
 		}
@@ -4111,8 +4208,79 @@ void FHlslNiagaraTranslator::HandleParameterRead(int32 ParamMapHistoryIdx, const
 
 		if (LastSetChunkIdx == INDEX_NONE && !bIgnoreDefaultValue)
 		{
-			// Default was found, trace back its inputs.
-			if (InputPin != nullptr)
+			if (Variable && Variable->DefaultMode == ENiagaraDefaultMode::Binding && Variable->DefaultBinding.IsValid())
+			{
+				bool bMatchingNames = false;
+				bool bIsCorrectType = false;
+				FNiagaraVariable FoundVar;
+				for (auto& UniqueVar : UniqueVars)
+				{
+					if (Variable->DefaultBinding.GetName() == UniqueVar.GetName())
+					{
+						bMatchingNames = true;
+						bIsCorrectType = Variable->Variable.GetType() == UniqueVar.GetType();
+						FoundVar = UniqueVar;
+						break;
+					}
+				}
+
+				if (!bMatchingNames)
+				{
+					for (auto& EngineConstant : FNiagaraConstants::GetEngineConstants())
+					{
+						if (Variable->DefaultBinding.GetName() == EngineConstant.GetName())
+						{
+							bMatchingNames = true;
+							bIsCorrectType = Variable->Variable.GetType() == EngineConstant.GetType();
+							FoundVar = EngineConstant;
+							break;
+						}
+					}
+
+					for (auto& TranslatorConstant : FNiagaraConstants::GetTranslatorConstants())
+					{
+						if (Variable->DefaultBinding.GetName() == TranslatorConstant.GetName())
+						{
+							bMatchingNames = true;
+							bIsCorrectType = Variable->Variable.GetType() == TranslatorConstant.GetType();
+							FoundVar = TranslatorConstant;
+							break;
+						}
+					}
+				}
+
+				if (bMatchingNames && bIsCorrectType)
+				{
+					for (auto& UniqueVar : UniqueVarToChunk)
+					{
+						if (UniqueVar.Key.IsEquivalent(FNiagaraVariable(Variable->Variable.GetType(), Variable->DefaultBinding.GetName())))
+						{
+							LastSetChunkIdx = UniqueVar.Value;
+							break;
+						}
+					}
+
+				}
+				else
+				{
+					if (!bMatchingNames)
+					{
+						Error(FText::Format(LOCTEXT("CannotFindBinding", "The module input {0} is bound to {1}, but {1} is not yet defined. Make sure {1} is defined prior to this module call."),
+							FText::FromName(Var.GetName()), 
+							FText::FromName(Variable->DefaultBinding.GetName())), ErrorNode, nullptr);
+					}
+					else
+					{
+						Error(FText::Format(LOCTEXT("IncorrectBindingType", "The type of module input {0} ({1}) does not match the type of the bound parameter {2} ({3}). This module expects {2} to be of type {1}."),
+							FText::FromName(Var.GetName()),
+							Var.GetType().GetNameText(),
+							FText::FromName(Variable->DefaultBinding.GetName()),
+							FoundVar.GetType().GetNameText()), ErrorNode, nullptr);
+					}
+					
+				}
+			}
+			else if (InputPin != nullptr) // Default was found, trace back its inputs.
 			{
 				// Check to see if there are any overrides passed in to the translator. This allows us to bake in rapid iteration variables for performance.
 				if (InputPin->LinkedTo.Num() == 0 && bIsCandidateForRapidIteration && !TranslationOptions.bParameterRapidIteration)
@@ -5060,8 +5228,70 @@ void FHlslNiagaraTranslator::RegisterFunctionCall(ENiagaraScriptUsage ScriptUsag
 									if (LastSetChunkIdx == INDEX_NONE)
 									{
 										const UEdGraphPin* DefaultPin = History.GetDefaultValuePin(VarIdx);
-										HandleParameterRead(ActiveStageIdx, AliasedVar, DefaultPin, ParamNode, LastSetChunkIdx);
+										UNiagaraScriptVariable* ScriptVariable = SourceGraph->GetScriptVariable(AliasedVar);
 
+										bool bFoundBinding = false;
+										if (AliasedVar.GetName().ToString().Contains(TEXT("Module.")))
+										{
+											
+											if (ScriptVariable->DefaultMode == ENiagaraDefaultMode::Binding)
+											{
+												FNiagaraScriptVariableBinding Bind = ScriptVariable->DefaultBinding;
+												if (Bind.IsValid())
+												{
+													for (auto& UniqueVar : UniqueVarToChunk)
+													{
+														if (UniqueVar.Key.IsEquivalent(FNiagaraVariable(Var.GetType(), Bind.GetName())))
+														{
+															LastSetChunkIdx = UniqueVar.Value;
+															break;
+														}
+													}
+													
+													if (LastSetChunkIdx == INDEX_NONE)
+													{
+														for (auto HistoryVariable : History.Variables)
+														{
+															if (HistoryVariable.IsEquivalent(FNiagaraVariable(Var.GetType(), Bind.GetName())))
+															{
+																FString SanitizedName = GetParameterMapInstanceName(ActiveStageIdx) + TEXT(".") + GetSanitizedSymbolName(Bind.GetName().ToString());
+																LastSetChunkIdx = AddSourceChunk(SanitizedName, Var.GetType());
+																break;
+															}
+														}
+													}
+
+													if (LastSetChunkIdx != INDEX_NONE && Var.GetType().GetClass() == nullptr)
+													{
+														FString SanitizedName = GetParameterMapInstanceName(ActiveStageIdx) + TEXT(".") + GetSanitizedSymbolName(Var.GetName().ToString());
+
+														AddBodyChunk(SanitizedName, TEXT("{0}"), Var.GetType(), LastSetChunkIdx, false);
+														LastSetChunkIdx = AddSourceChunk(SanitizedName, Var.GetType());
+														ParamMapDefinedAttributesToNamespaceVars.FindOrAdd(Var.GetName()) = Var;
+														bFoundBinding = true;
+													}
+													else
+													{
+														int Out;
+														if (ParameterMapRegisterExternalConstantNamespaceVariable(FNiagaraVariable(Var.GetType(), Bind.GetName()), ParamNode, ActiveStageIdx, Out, nullptr))
+														{
+															FString SanitizedName = GetParameterMapInstanceName(ActiveStageIdx) + TEXT(".") + GetSanitizedSymbolName(Var.GetName().ToString());
+
+															AddBodyChunk(SanitizedName, TEXT("{0}"), Var.GetType(), Out, false);
+															LastSetChunkIdx = AddSourceChunk(SanitizedName, Var.GetType());
+															ParamMapDefinedAttributesToNamespaceVars.FindOrAdd(Var.GetName()) = Var;
+															bFoundBinding = true;
+														}
+													}
+												}
+											}
+										}
+									
+										if (!bFoundBinding)
+										{
+											HandleParameterRead(ActiveStageIdx, AliasedVar, DefaultPin, ParamNode, LastSetChunkIdx, ScriptVariable);
+										}
+										
 										// If this variable was in the pending defaults list, go ahead and remove it
 										// as we added it before first use...
 										if (DeferredVariablesMissingDefault.Contains(Var))
diff --git a/Engine/Plugins/FX/Niagara/Source/NiagaraEditor/Private/NiagaraHlslTranslator.h b/Engine/Plugins/FX/Niagara/Source/NiagaraEditor/Private/NiagaraHlslTranslator.h
--- a/Engine/Plugins/FX/Niagara/Source/NiagaraEditor/Private/NiagaraHlslTranslator.h
+++ b/Engine/Plugins/FX/Niagara/Source/NiagaraEditor/Private/NiagaraHlslTranslator.h
@@ -22,6 +22,7 @@ class UNiagaraNodeCustomHlsl;
 class UNiagaraNodeOutput;
 class UNiagaraScriptSource;
 class UNiagaraNodeEmitter;
+class UNiagaraScriptVariable;
 
 
 // handles finished shader compile jobs, applying of the shaders to their scripts, and some error handling
@@ -524,7 +525,7 @@ public:
 
 private:
 	void InitializeParameterMapDefaults(int32 ParamMapHistoryIdx);
-	void HandleParameterRead(int32 ParamMapHistoryIdx, const FNiagaraVariable& Var, const UEdGraphPin* DefaultPin, UNiagaraNode* ErrorNode, int32& OutputChunkId,  bool bTreatAsUnknownParameterMap = false);
+	void HandleParameterRead(int32 ParamMapHistoryIdx, const FNiagaraVariable& Var, const UEdGraphPin* DefaultPin, UNiagaraNode* ErrorNode, int32& OutputChunkId, UNiagaraScriptVariable* Variable, bool bTreatAsUnknownParameterMap = false);
 	bool ShouldConsiderTargetParameterMap(ENiagaraScriptUsage InUsage) const;
 	FString BuildParameterMapHlslDefinitions(TArray<FNiagaraVariable>& PrimaryDataSetOutputEntries);
 	void BuildMissingDefaults();
diff --git a/Engine/Plugins/FX/Niagara/Source/NiagaraEditor/Private/NiagaraParameterMapHistory.cpp b/Engine/Plugins/FX/Niagara/Source/NiagaraEditor/Private/NiagaraParameterMapHistory.cpp
--- a/Engine/Plugins/FX/Niagara/Source/NiagaraEditor/Private/NiagaraParameterMapHistory.cpp
+++ b/Engine/Plugins/FX/Niagara/Source/NiagaraEditor/Private/NiagaraParameterMapHistory.cpp
@@ -21,6 +21,8 @@
 #include "NiagaraConstants.h"
 #include "NiagaraNodeStaticSwitch.h"
 
+#include "NiagaraScriptVariable.h"
+
 #define LOCTEXT_NAMESPACE "NiagaraEditor"
 
 
@@ -1203,6 +1205,21 @@ int32 FNiagaraParameterMapHistoryBuilder::HandleVariableRead(int32 ParamMapIdx,
 				Histories[ParamMapIdx].PerVariableWriteHistory.AddDefaulted(1);
 				Histories[ParamMapIdx].PerVariableReadHistory.AddDefaulted(1);
 
+				// Add the default binding as well, if used.
+				if (UNiagaraGraph* Graph = Cast<UNiagaraGraph>(InPin->GetOwningNode()->GetGraph()))
+				{
+					UNiagaraScriptVariable* Variable = Graph->GetScriptVariable(AliasedVar);
+					if (Variable && Variable->DefaultMode == ENiagaraDefaultMode::Binding && Variable->DefaultBinding.IsValid())
+					{
+						int FoundIdx2 = Histories[ParamMapIdx].Variables.Add(FNiagaraVariable(Var.GetType(), Variable->DefaultBinding.GetName()));
+						Histories[ParamMapIdx].VariablesWithOriginalAliasesIntact.Add(FNiagaraVariable(Var.GetType(), Variable->DefaultBinding.GetName()));
+						Histories[ParamMapIdx].PerVariableWarnings.AddDefaulted(1);
+						Histories[ParamMapIdx].PerVariableWriteHistory.AddDefaulted(1);
+						Histories[ParamMapIdx].PerVariableReadHistory.AddDefaulted(1);
+						
+						Histories[ParamMapIdx].PerVariableReadHistory[FoundIdx2].Add(TTuple<const UEdGraphPin*, const UEdGraphPin*>(InDefaultPin, nullptr));
+					}
+				}
 			}
 			Histories[ParamMapIdx].PerVariableReadHistory[FoundIdx].Add(TTuple<const UEdGraphPin*, const UEdGraphPin*>(InPin, nullptr));
 
diff --git a/Engine/Plugins/FX/Niagara/Source/NiagaraEditor/Private/NiagaraScriptVariable.cpp b/Engine/Plugins/FX/Niagara/Source/NiagaraEditor/Private/NiagaraScriptVariable.cpp
--- a/Engine/Plugins/FX/Niagara/Source/NiagaraEditor/Private/NiagaraScriptVariable.cpp
+++ b/Engine/Plugins/FX/Niagara/Source/NiagaraEditor/Private/NiagaraScriptVariable.cpp
@@ -1,9 +1,11 @@
 // Copyright 1998-2019 Epic Games, Inc. All Rights Reserved.
 
 #include "NiagaraScriptVariable.h"
+#include "NiagaraCommon.h"
 
 UNiagaraScriptVariable::UNiagaraScriptVariable(const FObjectInitializer& ObjectInitializer)
 	: Super(ObjectInitializer)
+	, DefaultMode(ENiagaraDefaultMode::Value)
 {
 	
 }
diff --git a/Engine/Plugins/FX/Niagara/Source/NiagaraEditor/Private/NiagaraScriptVariable.h b/Engine/Plugins/FX/Niagara/Source/NiagaraEditor/Private/NiagaraScriptVariable.h
--- a/Engine/Plugins/FX/Niagara/Source/NiagaraEditor/Private/NiagaraScriptVariable.h
+++ b/Engine/Plugins/FX/Niagara/Source/NiagaraEditor/Private/NiagaraScriptVariable.h
@@ -17,9 +17,18 @@ public:
 
 	virtual void PostLoad() override;
 	
+	/** The default mode. Can be Value, Binding or Custom. */
+	UPROPERTY(EditAnywhere, Category = "Default Value")
+	ENiagaraDefaultMode DefaultMode; 
+
+	/** The default binding. Only used if DefaultMode == ENiagaraDefaultMode::Binding. */
+	UPROPERTY(EditAnywhere, Category = "Default Value")
+	FNiagaraScriptVariableBinding DefaultBinding; 
+
 	/** Variable type, name and data. The data is not persistent, but used as a buffer when interfacing elsewhere. */
 	UPROPERTY()
 	FNiagaraVariable Variable;
+
 	/** The metadata associated with this script variable. */
 	UPROPERTY(EditAnywhere, Category = "Variable", meta=(ShowOnlyInnerProperties))
 	FNiagaraVariableMetaData Metadata;
diff --git a/Engine/Plugins/FX/Niagara/Source/NiagaraEditor/Private/ViewModels/Stack/NiagaraStackFunctionInput.cpp b/Engine/Plugins/FX/Niagara/Source/NiagaraEditor/Private/ViewModels/Stack/NiagaraStackFunctionInput.cpp
--- a/Engine/Plugins/FX/Niagara/Source/NiagaraEditor/Private/ViewModels/Stack/NiagaraStackFunctionInput.cpp
+++ b/Engine/Plugins/FX/Niagara/Source/NiagaraEditor/Private/ViewModels/Stack/NiagaraStackFunctionInput.cpp
@@ -41,6 +41,8 @@
 #include "INiagaraEditorTypeUtilities.h"
 #include "ViewModels/Stack/NiagaraStackInputCategory.h"
 
+#include "NiagaraScriptVariable.h"
+
 #define LOCTEXT_NAMESPACE "NiagaraStackViewModel"
 
 UNiagaraStackFunctionInput::UNiagaraStackFunctionInput()
@@ -51,6 +53,7 @@ UNiagaraStackFunctionInput::UNiagaraStackFunctionInput()
 	, bShowEditConditionInline(false)
 	, bIsInlineEditConditionToggle(false)
 	, bIsDynamicInputScriptReassignmentPending(false)
+	, bIsLocalOverride(false)
 {
 }
 
@@ -549,7 +552,23 @@ UNiagaraDataInterface* UNiagaraStackFunctionInput::FInputValues::GetDataDefaultV
 		: nullptr;
 }
 
-void UNiagaraStackFunctionInput::RefreshValues()
+bool UNiagaraStackFunctionInput::TryGetDefaultBinding(FNiagaraParameterHandle& LinkedValueHandle, UNiagaraScriptVariable* InVariable, UEdGraphPin& ValuePin)
+{
+	if (!InVariable)
+	{
+		return false;
+	}
+
+	if (InVariable->DefaultMode != ENiagaraDefaultMode::Binding || !InVariable->DefaultBinding.IsValid())
+	{
+		return false;
+	}
+
+	LinkedValueHandle = FNiagaraParameterHandle(InVariable->DefaultBinding.GetName());
+	return true;
+}
+
+void UNiagaraStackFunctionInput::RefreshValues(bool bFromSetLocalValue)
 {
 	if (ensureMsgf(IsStaticParameter() || InputParameterHandle.IsModuleHandle(), TEXT("Function inputs can only be generated for module paramters.")) == false)
 	{
@@ -565,25 +584,54 @@ void UNiagaraStackFunctionInput::RefreshValues()
 		UEdGraphPin* OverridePin = GetOverridePin();
 		UEdGraphPin* ValuePin = OverridePin != nullptr ? OverridePin : DefaultPin;
 
-		if (TryGetCurrentLocalValue(InputValues.LocalStruct, *DefaultPin, *ValuePin, OldValues.GetLocalStructToReuse()))
+		if (UNiagaraGraph* FunctionGraph = Cast<UNiagaraScriptSource>(OwningFunctionCallNode->FunctionScript->GetSource())->NodeGraph)
 		{
-			InputValues.Mode = EValueMode::Local;
+			Variable = FunctionGraph->GetScriptVariable(*(TEXT("Module.") + DisplayName.ToString()));
 		}
-		else if (TryGetCurrentLinkedValue(InputValues.LinkedHandle, *ValuePin))
-		{
-			InputValues.Mode = EValueMode::Linked;
-		}
-		else if (TryGetCurrentDataValue(InputValues.DataObjects, OverridePin, *DefaultPin, OldValues.GetDataDefaultValueObjectToReuse()))
+
+		if (TryGetCurrentDataValue(InputValues.DataObjects, OverridePin, *DefaultPin, OldValues.GetDataDefaultValueObjectToReuse()))
 		{
 			InputValues.Mode = EValueMode::Data;
+			bIsLocalOverride = false;
 		}
 		else if (TryGetCurrentExpressionValue(InputValues.ExpressionNode, OverridePin))
 		{
 			InputValues.Mode = EValueMode::Expression;
+			bIsLocalOverride = false;
 		}
 		else if (TryGetCurrentDynamicValue(InputValues.DynamicNode, OverridePin))
 		{
 			InputValues.Mode = EValueMode::Dynamic;
+			bIsLocalOverride = false;
+		}
+		else if (TryGetCurrentLinkedValue(InputValues.LinkedHandle, *ValuePin))
+		{
+			InputValues.Mode = EValueMode::Linked;
+			bIsLocalOverride = false;
+		}
+		else if (TryGetCurrentLocalValue(InputValues.LocalStruct, *DefaultPin, *ValuePin, OldValues.GetLocalStructToReuse(), Variable))
+		{
+			if (OverridePin)
+			{
+				InputValues.Mode = EValueMode::Local;
+			}
+			else
+			{
+				if (!bIsLocalOverride && TryGetDefaultBinding(InputValues.LinkedHandle, Variable, *ValuePin))
+				{
+					InputValues.Mode = EValueMode::Linked;
+					bIsLocalOverride = false; 
+				}
+				else
+				{
+					InputValues.Mode = EValueMode::Local;
+				}
+			}
+		}
+		else if (TryGetDefaultBinding(InputValues.LinkedHandle, Variable, *ValuePin))
+		{
+			InputValues.Mode = EValueMode::Linked;
+			bIsLocalOverride = false;
 		}
 	}
 
@@ -809,9 +857,9 @@ void UNiagaraStackFunctionInput::GetAvailableParameterHandles(TArray<FNiagaraPar
 				{
 					for (int32 j = 0; j < Builder.Histories[0].Variables.Num(); j++)
 					{
-						FNiagaraVariable& Variable = Builder.Histories[0].Variables[j];
-						FNiagaraParameterHandle AvailableHandle = FNiagaraParameterHandle(Variable.GetName());
-						if (Variable.GetType() == InputType)
+						FNiagaraVariable& HistoryVariable = Builder.Histories[0].Variables[j];
+						FNiagaraParameterHandle AvailableHandle = FNiagaraParameterHandle(HistoryVariable.GetName());
+						if (HistoryVariable.GetType() == InputType)
 						{
 							TArray<const UEdGraphPin*>& WriteHistory = Builder.Histories[0].PerVariableWriteHistory[j];
 							for (const UEdGraphPin* WritePin : WriteHistory)
@@ -977,8 +1025,12 @@ bool UNiagaraStackFunctionInput::IsRapidIterationCandidate() const
 	return !IsStaticParameter() && FNiagaraStackGraphUtilities::IsRapidIterationType(InputType);
 }
 
-void UNiagaraStackFunctionInput::SetLocalValue(TSharedRef<FStructOnScope> InLocalValue)
+void UNiagaraStackFunctionInput::SetLocalValue(TSharedRef<FStructOnScope> InLocalValue, bool bIsOverride)
 {
+	if (bIsOverride)
+	{
+		bIsLocalOverride = bIsOverride;
+	}
 	TGuardValue<bool> UpdateGuard(bUpdatingLocalValueDirectly, true);
 
 	const UEdGraphSchema_Niagara* NiagaraSchema = GetDefault<UEdGraphSchema_Niagara>();
@@ -999,13 +1051,18 @@ void UNiagaraStackFunctionInput::SetLocalValue(TSharedRef<FStructOnScope> InLoca
 			DefaultPin->DefaultValue = PinDefaultValue;
 			Cast<UNiagaraNode>(DefaultPin->GetOwningNode())->MarkNodeRequiresSynchronization(TEXT("Default Value Changed"), true);
 		}
-		RefreshValues();
+		RefreshValues(true);
 		return;
 	}
 	
 	// If the default pin in the function graph is connected internally, rapid iteration parameters can't be used since
 	// the compilation currently won't use them.
 	bool bCanUseRapidIterationParameter = IsRapidIterationCandidate() && DefaultPin->LinkedTo.Num() == 0;
+	if (Variable && Variable->DefaultMode == ENiagaraDefaultMode::Binding)
+	{
+		bCanUseRapidIterationParameter = false;
+	}
+
 	if (bCanUseRapidIterationParameter == false)
 	{
 		ValuePin = OverridePin != nullptr ? OverridePin : DefaultPin;
@@ -1013,7 +1070,7 @@ void UNiagaraStackFunctionInput::SetLocalValue(TSharedRef<FStructOnScope> InLoca
 
 	TSharedPtr<FStructOnScope> CurrentValue;
 	bool bCanHaveLocalValue = ValuePin != nullptr;
-	bool bHasLocalValue = bCanHaveLocalValue && InputValues.Mode == EValueMode::Local && TryGetCurrentLocalValue(CurrentValue, *DefaultPin, *ValuePin, TSharedPtr<FStructOnScope>());
+	bool bHasLocalValue = bCanHaveLocalValue && InputValues.Mode == EValueMode::Local && TryGetCurrentLocalValue(CurrentValue, *DefaultPin, *ValuePin, TSharedPtr<FStructOnScope>(), nullptr);
 	bool bLocalValueMatchesSetValue = bHasLocalValue && FNiagaraEditorUtilities::DataMatches(*CurrentValue.Get(), InLocalValue.Get());
 
 	if (bCanHaveLocalValue == false || bLocalValueMatchesSetValue)
@@ -1077,7 +1134,7 @@ void UNiagaraStackFunctionInput::SetLocalValue(TSharedRef<FStructOnScope> InLoca
 		FNiagaraStackGraphUtilities::RelayoutGraph(*EmitterGraph);
 	}
 
-	RefreshValues();
+	RefreshValues(true);
 }
 
 bool UNiagaraStackFunctionInput::CanReset() const
@@ -1108,9 +1165,18 @@ bool UNiagaraStackFunctionInput::CanReset() const
 			{			
 				if(DefaultPin->LinkedTo.Num() == 0)
 				{
-					if (GetOverridePin() != nullptr)
+					if (UEdGraphPin* OverridePin = GetOverridePin())
 					{
 						bNewCanReset = true;
+						UNiagaraGraph* FunctionGraph = CastChecked<UNiagaraScriptSource>(OwningFunctionCallNode->FunctionScript->GetSource())->NodeGraph;
+						if (FunctionGraph && OverridePin && OverridePin->LinkedTo.Num() == 1)
+						{
+							UNiagaraScriptVariable* ScriptVariable = FunctionGraph->GetScriptVariable(*(TEXT("Module.") + DisplayName.ToString()));
+							if (ScriptVariable && OverridePin->LinkedTo[0]->PinName == ScriptVariable->DefaultBinding.GetName())
+							{
+								bNewCanReset = false;
+							}
+						}
 					}
 					else if (IsRapidIterationCandidate())
 					{
@@ -1193,6 +1259,7 @@ bool UNiagaraStackFunctionInput::RemoveRapidIterationParametersForAffectedScript
 
 void UNiagaraStackFunctionInput::Reset()
 {
+	bIsLocalOverride = false;
 	if(InputValues.Mode == EValueMode::Data)
 	{
 		// For data values they are reset by making sure the data object owned by this input matches the default
@@ -1737,7 +1804,7 @@ UEdGraphPin& UNiagaraStackFunctionInput::GetOrCreateOverridePin()
 	return *OverridePin;
 }
 
-bool UNiagaraStackFunctionInput::TryGetCurrentLocalValue(TSharedPtr<FStructOnScope>& LocalValue, UEdGraphPin& DefaultPin, UEdGraphPin& ValuePin, TSharedPtr<FStructOnScope> OldValueToReuse)
+bool UNiagaraStackFunctionInput::TryGetCurrentLocalValue(TSharedPtr<FStructOnScope>& LocalValue, UEdGraphPin& DefaultPin, UEdGraphPin& ValuePin, TSharedPtr<FStructOnScope> OldValueToReuse, UNiagaraScriptVariable* InVariable)
 {
 	if (InputType.IsUObject() == false && ValuePin.LinkedTo.Num() == 0)
 	{
@@ -1756,6 +1823,10 @@ bool UNiagaraStackFunctionInput::TryGetCurrentLocalValue(TSharedPtr<FStructOnSco
 		// the compilation currently won't use them.
 		bool bCanUseRapidIterationParameter = IsRapidIterationCandidate() && DefaultPin.LinkedTo.Num() == 0;
 		bool bFoundRapidIterationParameter = false;
+		if (InVariable && InVariable->DefaultMode == ENiagaraDefaultMode::Binding)
+		{
+			bCanUseRapidIterationParameter = false;
+		}
 		if (bCanUseRapidIterationParameter)
 		{
 			const uint8* RapidIterationParameterData = SourceScript->RapidIterationParameters.GetParameterData(RapidIterationParameter);
diff --git a/Engine/Plugins/FX/Niagara/Source/NiagaraEditor/Public/ViewModels/Stack/NiagaraStackFunctionInput.h b/Engine/Plugins/FX/Niagara/Source/NiagaraEditor/Public/ViewModels/Stack/NiagaraStackFunctionInput.h
--- a/Engine/Plugins/FX/Niagara/Source/NiagaraEditor/Public/ViewModels/Stack/NiagaraStackFunctionInput.h
+++ b/Engine/Plugins/FX/Niagara/Source/NiagaraEditor/Public/ViewModels/Stack/NiagaraStackFunctionInput.h
@@ -19,6 +19,7 @@ class UNiagaraScript;
 class UEdGraphPin;
 class UNiagaraDataInterface;
 enum class EStackParameterBehavior;
+class UNiagaraScriptVariable;
 
 /** Represents a single module input in the module stack view model. */
 UCLASS()
@@ -136,7 +137,7 @@ public:
 	bool IsEnabled() const;
 
 	/** Sets this input's local value. */
-	void SetLocalValue(TSharedRef<FStructOnScope> InLocalValue);
+	void SetLocalValue(TSharedRef<FStructOnScope> InLocalValue, bool bIsOverride = false);
 	
 	/** Returns whether or not the value or handle of this input has been overridden and can be reset. */
 	bool CanReset() const;
@@ -312,7 +313,7 @@ private:
 
 private:
 	/** Refreshes the current values for this input from the state of the graph. */
-	void RefreshValues();
+	void RefreshValues(bool bFromSetLocalValue=false);
 
 	/** Refreshes additional state for this input which comes from input metadata. */
 	void RefreshFromMetaData();
@@ -343,8 +344,10 @@ private:
 	  * pin don't exist, they will be created. */
 	UEdGraphPin& GetOrCreateOverridePin();
 
+	bool TryGetDefaultBinding(FNiagaraParameterHandle& LinkedValueHandle, UNiagaraScriptVariable* Variable, UEdGraphPin& ValuePin);
+
 	/** Tries to get a local value for this input if it exists by checking the graph data directly. */
-	bool TryGetCurrentLocalValue(TSharedPtr<FStructOnScope>& LocalValue, UEdGraphPin& DefaultPin, UEdGraphPin& ValuePin, TSharedPtr<FStructOnScope> OldValueToReuse);
+	bool TryGetCurrentLocalValue(TSharedPtr<FStructOnScope>& LocalValue, UEdGraphPin& DefaultPin, UEdGraphPin& ValuePin, TSharedPtr<FStructOnScope> OldValueToReuse, UNiagaraScriptVariable* Variable);
 
 	/** Tries to get a data interface value for this input if it exists by checking the graph data directly .*/
 	bool TryGetCurrentDataValue(FDataValues& DataValues, UEdGraphPin* OverrideValuePin, UEdGraphPin& DefaultValuePin, UNiagaraDataInterface* LocallyOwnedDefaultDataValueObjectToReuse);
@@ -462,4 +465,7 @@ private:
 
 	/** Whether or not the dynamic input for this input has a function script reassignment pending due to a request to fix a missing script. */
 	bool bIsDynamicInputScriptReassignmentPending;
+
+	bool bIsLocalOverride;
+	UNiagaraScriptVariable* Variable;
 };
\ No newline at end of file
diff --git a/Engine/Plugins/FX/Niagara/Source/NiagaraEditorWidgets/Private/Stack/SNiagaraStackFunctionInputValue.cpp b/Engine/Plugins/FX/Niagara/Source/NiagaraEditorWidgets/Private/Stack/SNiagaraStackFunctionInputValue.cpp
--- a/Engine/Plugins/FX/Niagara/Source/NiagaraEditorWidgets/Private/Stack/SNiagaraStackFunctionInputValue.cpp
+++ b/Engine/Plugins/FX/Niagara/Source/NiagaraEditorWidgets/Private/Stack/SNiagaraStackFunctionInputValue.cpp
@@ -707,7 +707,7 @@ void SNiagaraStackFunctionInputValue::SetToLocalValue()
 		if (DefaultValueData.Num() == LocalValueStruct->GetStructureSize())
 		{
 			FMemory::Memcpy(LocalValue->GetStructMemory(), DefaultValueData.GetData(), DefaultValueData.Num());
-			FunctionInput->SetLocalValue(LocalValue);
+			FunctionInput->SetLocalValue(LocalValue, true);
 		}
 	}
 }
