non rectangular domain support

diff --git a/Engine/Plugins/FX/Niagara/Source/Niagara/Classes/NiagaraDataInterfaceRW.h b/Engine/Plugins/FX/Niagara/Source/Niagara/Classes/NiagaraDataInterfaceRW.h
--- a/Engine/Plugins/FX/Niagara/Source/Niagara/Classes/NiagaraDataInterfaceRW.h
+++ b/Engine/Plugins/FX/Niagara/Source/Niagara/Classes/NiagaraDataInterfaceRW.h
@@ -173,20 +173,20 @@ class NIAGARA_API UNiagaraDataInterfaceGrid2D : public UNiagaraDataInterfaceRWBa
 	GENERATED_UCLASS_BODY()
 
 public:
-	UPROPERTY(EditAnywhere, Category = "Grid", meta = (EditCondition = "!SetGridFromCellSize"))
+	UPROPERTY(EditAnywhere, Category = "Grid", meta = (EditCondition = "!SetGridFromMaxAxis"))
 	int32 NumCellsX;
 
-	UPROPERTY(EditAnywhere, Category = "Grid", meta = (EditCondition = "!SetGridFromCellSize"))
+	UPROPERTY(EditAnywhere, Category = "Grid", meta = (EditCondition = "!SetGridFromMaxAxis"))
 	int32 NumCellsY;
 	
-	UPROPERTY(EditAnywhere, Category = "Grid", meta = (EditCondition = "SetGridFromCellSize"))
-	float CellSize;
+	UPROPERTY(EditAnywhere, Category = "Grid", meta = (EditCondition = "SetGridFromMaxAxis"))
+	int32 NumCellsMaxAxis;
 
 	UPROPERTY(EditAnywhere, Category = "Grid")
 	int32 NumAttributes;
 
 	UPROPERTY(EditAnywhere, Category = "Grid")
-	bool SetGridFromCellSize;
+	bool SetGridFromMaxAxis;
 
 	UPROPERTY(EditAnywhere, Category = "Grid")
 	FVector WorldBBoxMin;
diff --git a/Engine/Plugins/FX/Niagara/Source/Niagara/Private/NiagaraDataInterfaceGrid2DCollection.cpp b/Engine/Plugins/FX/Niagara/Source/Niagara/Private/NiagaraDataInterfaceGrid2DCollection.cpp
--- a/Engine/Plugins/FX/Niagara/Source/Niagara/Private/NiagaraDataInterfaceGrid2DCollection.cpp
+++ b/Engine/Plugins/FX/Niagara/Source/Niagara/Private/NiagaraDataInterfaceGrid2DCollection.cpp
@@ -366,18 +366,31 @@ bool UNiagaraDataInterfaceGrid2DCollection::InitPerInstanceData(void* PerInstanc
 	FVector2D RT_WorldBBoxSize = WorldBBoxSize;
 
 	// If we are setting the grid from the voxel size, then recompute NumVoxels and change bbox	
-	if (SetGridFromCellSize)
+	if (SetGridFromMaxAxis)
 	{
+		float CellSize = CellSize = WorldBBoxSize.Y / NumCellsMaxAxis;
+		if (WorldBBoxSize.X > WorldBBoxSize.Y)
+		{
+			CellSize = WorldBBoxSize.X / NumCellsMaxAxis;
+
+		}
+
 		RT_NumCellsX = WorldBBoxSize.X / CellSize;
 		RT_NumCellsY = WorldBBoxSize.Y / CellSize;
 
 		// Pad grid by 1 voxel if our computed bounding box is too small
-		if (!FMath::IsNearlyEqual(CellSize * RT_NumCellsX, WorldBBoxSize.X))
+		if (WorldBBoxSize.X > WorldBBoxSize.Y && !FMath::IsNearlyEqual(CellSize * RT_NumCellsY, WorldBBoxSize.Y))
 		{
-			RT_NumCellsX = NumCellsX + 1;
-			RT_NumCellsY = NumCellsY + 1;
-			RT_WorldBBoxSize = FVector2D(RT_NumCellsX, RT_NumCellsY) * CellSize;
+			RT_NumCellsY++;
 		}
+		else if (WorldBBoxSize.X < WorldBBoxSize.Y && !FMath::IsNearlyEqual(CellSize * RT_NumCellsX, WorldBBoxSize.X))
+		{
+			RT_NumCellsX++;
+		}		
+		
+		RT_WorldBBoxSize = FVector2D(RT_NumCellsX, RT_NumCellsY) * CellSize;
+		NumCellsX = RT_NumCellsX;
+		NumCellsY = RT_NumCellsY;
 	}	
 
 	TSet<int> RT_OutputShaderStages = OutputShaderStages;
diff --git a/Engine/Plugins/FX/Niagara/Source/Niagara/Private/NiagaraDataInterfaceRW.cpp b/Engine/Plugins/FX/Niagara/Source/Niagara/Private/NiagaraDataInterfaceRW.cpp
--- a/Engine/Plugins/FX/Niagara/Source/Niagara/Private/NiagaraDataInterfaceRW.cpp
+++ b/Engine/Plugins/FX/Niagara/Source/Niagara/Private/NiagaraDataInterfaceRW.cpp
@@ -305,9 +305,9 @@ UNiagaraDataInterfaceGrid2D::UNiagaraDataInterfaceGrid2D(FObjectInitializer cons
 	: Super(ObjectInitializer)
 	, NumCellsX(3)
 	, NumCellsY(3)
-	, CellSize(1.)
+	, NumCellsMaxAxis(3)
 	, NumAttributes(1)
-	, SetGridFromCellSize(false)
+	, SetGridFromMaxAxis(false)
 	, WorldBBoxMin(0., 0., 0.)
 	, WorldBBoxSize(100., 100.)
 {
@@ -485,7 +485,7 @@ bool UNiagaraDataInterfaceGrid2D::Equals(const UNiagaraDataInterface* Other) con
 		OtherTyped->NumCellsX == NumCellsX &&
 		OtherTyped->NumCellsY == NumCellsY &&
 		OtherTyped->NumAttributes == NumAttributes &&
-		FMath::IsNearlyEqual(OtherTyped->CellSize, CellSize) &&
+		OtherTyped->NumCellsMaxAxis == NumCellsMaxAxis &&
 		OtherTyped->WorldBBoxMin.Equals(WorldBBoxMin) &&
 		OtherTyped->WorldBBoxSize.Equals(WorldBBoxSize);
 }
@@ -690,8 +690,8 @@ bool UNiagaraDataInterfaceGrid2D::CopyToInternal(UNiagaraDataInterface* Destinat
 	OtherTyped->NumCellsX = NumCellsX;
 	OtherTyped->NumCellsY = NumCellsY;
 	OtherTyped->NumAttributes = NumAttributes;
-	OtherTyped->CellSize = CellSize;
-	OtherTyped->SetGridFromCellSize = SetGridFromCellSize;
+	OtherTyped->NumCellsMaxAxis = NumCellsMaxAxis;
+	OtherTyped->SetGridFromMaxAxis = SetGridFromMaxAxis;
 	OtherTyped->WorldBBoxMin = WorldBBoxMin;
 	OtherTyped->WorldBBoxSize = WorldBBoxSize;
 
