fixed potenteial nullptr bugs

diff --git a/Engine/Plugins/FX/Niagara/Source/Niagara/Private/NiagaraDataInterfaceGrid2DCollection.cpp b/Engine/Plugins/FX/Niagara/Source/Niagara/Private/NiagaraDataInterfaceGrid2DCollection.cpp
--- a/Engine/Plugins/FX/Niagara/Source/Niagara/Private/NiagaraDataInterfaceGrid2DCollection.cpp
+++ b/Engine/Plugins/FX/Niagara/Source/Niagara/Private/NiagaraDataInterfaceGrid2DCollection.cpp
@@ -705,9 +705,17 @@ void FNiagaraDataInterfaceProxyGrid2DCollection::ResetData(FRHICommandList& RHIC
 {	
 	Grid2DCollectionRWInstanceData* ProxyData = SystemInstancesToProxyData.Find(Context.SystemInstance);
 
+	if (!ProxyData)
+	{
+		return;
+	}
+
 	for (Grid2DBuffer* Buffer : ProxyData->Buffers)
 	{
-		ClearUAV(RHICmdList, Buffer->GridBuffer, FLinearColor(0, 0, 0, 0));
+		if (Buffer)
+		{
+			ClearUAV(RHICmdList, Buffer->GridBuffer, FLinearColor(0, 0, 0, 0));
+		}		
 	}	
 }
 #undef LOCTEXT_NAMESPACE
\ No newline at end of file
