fix for WorldToUnit

diff --git a/Engine/Plugins/FX/Niagara/Source/Niagara/Private/NiagaraDataInterfaceRW.cpp b/Engine/Plugins/FX/Niagara/Source/Niagara/Private/NiagaraDataInterfaceRW.cpp
--- a/Engine/Plugins/FX/Niagara/Source/Niagara/Private/NiagaraDataInterfaceRW.cpp
+++ b/Engine/Plugins/FX/Niagara/Source/Niagara/Private/NiagaraDataInterfaceRW.cpp
@@ -529,7 +529,7 @@ bool UNiagaraDataInterfaceGrid2D::GetFunctionHLSL(const FName& DefinitionFunctio
 	else if (DefinitionFunctionName == WorldToUnitFunctionName)
 	{
 		static const TCHAR *FormatSample = TEXT(R"(
-			void {FunctionName}(float3 In_World, out float2 Out_Unit)
+			void {FunctionName}(float3 In_World, out float3 Out_Unit)
 			{
 				Out_Unit = (In_World - {WorldBBoxMinName}) / float3({WorldBBoxSizeName}, 1);
 			}
