- Disable sorting in Niagara by default on sprite / mesh renderers
- Existing system will use the previous default to avoid undesireable behavour

diff --git a/Engine/Plugins/FX/Niagara/Source/Niagara/Private/NiagaraMeshRendererProperties.cpp b/Engine/Plugins/FX/Niagara/Source/Niagara/Private/NiagaraMeshRendererProperties.cpp
--- a/Engine/Plugins/FX/Niagara/Source/Niagara/Private/NiagaraMeshRendererProperties.cpp
+++ b/Engine/Plugins/FX/Niagara/Source/Niagara/Private/NiagaraMeshRendererProperties.cpp
@@ -31,7 +31,7 @@ bool FNiagaraMeshMaterialOverride::SerializeFromMismatchedTag(const struct FProp
 
 UNiagaraMeshRendererProperties::UNiagaraMeshRendererProperties()
 	: ParticleMesh(nullptr)
-	, SortMode(ENiagaraSortMode::ViewDistance)
+	, SortMode(ENiagaraSortMode::None)
 	, bSortOnlyWhenTranslucent(true)
 {
 }
@@ -77,6 +77,18 @@ void UNiagaraMeshRendererProperties::PostInitProperties()
 	}
 }
 
+void UNiagaraMeshRendererProperties::Serialize(FArchive& Ar)
+{
+	Ar.UsingCustomVersion(FNiagaraCustomVersion::GUID);
+	const int32 NiagaraVersion = Ar.CustomVer(FNiagaraCustomVersion::GUID);
+
+	if (Ar.IsLoading() && (NiagaraVersion < FNiagaraCustomVersion::DisableSortingByDefault))
+	{
+		SortMode = ENiagaraSortMode::ViewDistance;
+	}
+	Super::Serialize(Ar);
+}
+
 /** The bindings depend on variables that are created during the NiagaraModule startup. However, the CDO's are build prior to this being initialized, so we defer setting these values until later.*/
 void UNiagaraMeshRendererProperties::InitCDOPropertiesAfterModuleStartup()
 {
diff --git a/Engine/Plugins/FX/Niagara/Source/Niagara/Private/NiagaraSpriteRendererProperties.cpp b/Engine/Plugins/FX/Niagara/Source/Niagara/Private/NiagaraSpriteRendererProperties.cpp
--- a/Engine/Plugins/FX/Niagara/Source/Niagara/Private/NiagaraSpriteRendererProperties.cpp
+++ b/Engine/Plugins/FX/Niagara/Source/Niagara/Private/NiagaraSpriteRendererProperties.cpp
@@ -34,7 +34,7 @@ UNiagaraSpriteRendererProperties::UNiagaraSpriteRendererProperties()
 	, FacingMode(ENiagaraSpriteFacingMode::FaceCamera)
 	, CustomFacingVectorMask(ForceInitToZero)
 	, PivotInUVSpace(0.5f, 0.5f)
-	, SortMode(ENiagaraSortMode::ViewDistance)
+	, SortMode(ENiagaraSortMode::None)
 	, SubImageSize(1.0f, 1.0f)
 	, bSubImageBlend(false)
 	, bRemoveHMDRollInVR(false)
@@ -117,6 +117,15 @@ void UNiagaraSpriteRendererProperties::PostInitProperties()
 
 void UNiagaraSpriteRendererProperties::Serialize(FStructuredArchive::FRecord Record)
 {
+	FArchive& Ar = Record.GetUnderlyingArchive();
+	Ar.UsingCustomVersion(FNiagaraCustomVersion::GUID);
+	const int32 NiagaraVersion = Ar.CustomVer(FNiagaraCustomVersion::GUID);
+
+	if (Ar.IsLoading() && (NiagaraVersion < FNiagaraCustomVersion::DisableSortingByDefault))
+	{
+		SortMode = ENiagaraSortMode::ViewDistance;
+	}
+
 	Super::Serialize(Record);
 
 	bool bIsCookedForEditor = false;
diff --git a/Engine/Plugins/FX/Niagara/Source/Niagara/Public/NiagaraMeshRendererProperties.h b/Engine/Plugins/FX/Niagara/Source/Niagara/Public/NiagaraMeshRendererProperties.h
--- a/Engine/Plugins/FX/Niagara/Source/Niagara/Public/NiagaraMeshRendererProperties.h
+++ b/Engine/Plugins/FX/Niagara/Source/Niagara/Public/NiagaraMeshRendererProperties.h
@@ -63,6 +63,7 @@ public:
 
 	//UObject Interface
 	virtual void PostInitProperties() override;
+	virtual void Serialize(FArchive& Ar) override;
 #if WITH_EDITORONLY_DATA
 	virtual void BeginDestroy() override;
 	virtual void PreEditChange(class UProperty* PropertyThatWillChange) override;
diff --git a/Engine/Plugins/FX/Niagara/Source/NiagaraCore/Public/NiagaraCustomVersion.h b/Engine/Plugins/FX/Niagara/Source/NiagaraCore/Public/NiagaraCustomVersion.h
--- a/Engine/Plugins/FX/Niagara/Source/NiagaraCore/Public/NiagaraCustomVersion.h
+++ b/Engine/Plugins/FX/Niagara/Source/NiagaraCore/Public/NiagaraCustomVersion.h
@@ -132,6 +132,8 @@ struct FNiagaraCustomVersion
 
 		ChangeEmitterCompiledDataToSharedRefs, // Changing the system and emitter compiled data to shared pointers to deal with lifetime issues in the editor.  They now are handled directly in system serialize.
 
+		DisableSortingByDefault, // Sorting on Renderers is disabled by default, we add a version to maintain existing systems that expected sorting to be enabled
+
 		// DO NOT ADD A NEW VERSION UNLESS YOU HAVE TALKED TO THE NIAGARA LEAD. Mismanagement of these versions can lead to data loss if it is adjusted in multiple streams simultaneously.
 		// -----<new versions can be added above this line>  -------------------------------------------------
 		VersionPlusOne,
