Niagara - Update overview node stack styling to better match other patterns in the editor.

diff --git a/Engine/Plugins/FX/Niagara/Source/NiagaraEditor/Private/ViewModels/NiagaraSystemSelectionViewModel.cpp b/Engine/Plugins/FX/Niagara/Source/NiagaraEditor/Private/ViewModels/NiagaraSystemSelectionViewModel.cpp
--- a/Engine/Plugins/FX/Niagara/Source/NiagaraEditor/Private/ViewModels/NiagaraSystemSelectionViewModel.cpp
+++ b/Engine/Plugins/FX/Niagara/Source/NiagaraEditor/Private/ViewModels/NiagaraSystemSelectionViewModel.cpp
@@ -139,6 +139,15 @@ UNiagaraStackViewModel* UNiagaraSystemSelectionViewModel::GetSelectionStackViewM
 	return SelectionStackViewModel;
 }
 
+void  UNiagaraSystemSelectionViewModel::EmptySelection()
+{
+	if (SelectionEntries.Num() > 0)
+	{
+		SelectionEntries.Empty();
+		UpdateExternalSelectionState();
+	}
+}
+
 void UNiagaraSystemSelectionViewModel::RemoveEntryFromSelectionByDisplayedObject(const UObject* InObject)
 {
 	bool bSelectionChanged = false;
diff --git a/Engine/Plugins/FX/Niagara/Source/NiagaraEditor/Private/ViewModels/Stack/NiagaraStackScriptItemGroup.cpp b/Engine/Plugins/FX/Niagara/Source/NiagaraEditor/Private/ViewModels/Stack/NiagaraStackScriptItemGroup.cpp
--- a/Engine/Plugins/FX/Niagara/Source/NiagaraEditor/Private/ViewModels/Stack/NiagaraStackScriptItemGroup.cpp
+++ b/Engine/Plugins/FX/Niagara/Source/NiagaraEditor/Private/ViewModels/Stack/NiagaraStackScriptItemGroup.cpp
@@ -829,9 +829,8 @@ TOptional<UNiagaraStackEntry::FDropRequestResponse> UNiagaraStackScriptItemGroup
 
 void UNiagaraStackScriptItemGroup::ItemAdded(UNiagaraNodeFunctionCall* AddedModule)
 {
-	TArray<const UObject*> DisplayedObjectsToSelect;
-	DisplayedObjectsToSelect.Add(AddedModule);
-	GetSystemViewModel()->GetSelectionViewModel()->AddEntriesToSelectionByDisplayedObjectsDeferred(DisplayedObjectsToSelect);
+	GetSystemViewModel()->GetSelectionViewModel()->EmptySelection();
+	GetSystemViewModel()->GetSelectionViewModel()->AddEntryToSelectionByDisplayedObjectDeferred(AddedModule);
 	RefreshChildren();
 }
 
diff --git a/Engine/Plugins/FX/Niagara/Source/NiagaraEditor/Public/ViewModels/NiagaraSystemSelectionViewModel.h b/Engine/Plugins/FX/Niagara/Source/NiagaraEditor/Public/ViewModels/NiagaraSystemSelectionViewModel.h
--- a/Engine/Plugins/FX/Niagara/Source/NiagaraEditor/Public/ViewModels/NiagaraSystemSelectionViewModel.h
+++ b/Engine/Plugins/FX/Niagara/Source/NiagaraEditor/Public/ViewModels/NiagaraSystemSelectionViewModel.h
@@ -46,6 +46,8 @@ public:
 
 	UNiagaraStackViewModel* GetSelectionStackViewModel();
 
+	void EmptySelection();
+
 	void RemoveEntryFromSelectionByDisplayedObject(const UObject* InObject);
 
 	void AddEntriesToSelectionByDisplayedObjectsDeferred(const TArray<const UObject*>& InObjects);
diff --git a/Engine/Plugins/FX/Niagara/Source/NiagaraEditorWidgets/Private/NiagaraEditorWidgetsStyle.cpp b/Engine/Plugins/FX/Niagara/Source/NiagaraEditorWidgets/Private/NiagaraEditorWidgetsStyle.cpp
--- a/Engine/Plugins/FX/Niagara/Source/NiagaraEditorWidgets/Private/NiagaraEditorWidgetsStyle.cpp
+++ b/Engine/Plugins/FX/Niagara/Source/NiagaraEditorWidgets/Private/NiagaraEditorWidgetsStyle.cpp
@@ -110,15 +110,17 @@ TSharedRef< FSlateStyleSet > FNiagaraEditorWidgetsStyle::Create()
 		.SetShadowColorAndOpacity(FLinearColor(0, 0, 0, 0.9f));;
 	Style->Set("NiagaraEditor.SystemOverview.GroupHeaderText", SystemOverviewGroupHeaderText);
 
-	FSlateFontInfo SystemOverviewItemFont = DEFAULT_FONT("Regular", 10);
+	FSlateFontInfo SystemOverviewItemFont = DEFAULT_FONT("Regular", 9);
 	FTextBlockStyle SystemOverviewItemText = FTextBlockStyle(NormalText)
 		.SetFont(SystemOverviewItemFont);
 	Style->Set("NiagaraEditor.SystemOverview.ItemText", SystemOverviewItemText);
 
-	Style->Set("NiagaraEditor.SystemOverview.BackgroundColor", FLinearColor(FColor(80, 80, 80)));
+	Style->Set("NiagaraEditor.SystemOverview.Item.BackgroundColor", FLinearColor(FColor(62, 62, 62)));
+	Style->Set("NiagaraEditor.SystemOverview.Group.BackgroundColor", FLinearColor::Transparent);
 	Style->Set("NiagaraEditor.SystemOverview.CheckBoxColor", FLinearColor(FColor(160, 160, 160)));
 	Style->Set("NiagaraEditor.SystemOverview.CheckBoxBorder", new BOX_CORE_BRUSH("Common/GroupBorderLight", FMargin(4.0f / 16.0f)));
 	Style->Set("NiagaraEditor.SystemOverview.NodeBackgroundBorder", new BOX_PLUGIN_BRUSH("Icons/SystemOverviewNodeBackground", FMargin(1.0f / 4.0f)));
+	Style->Set("NiagaraEditor.SystemOverview.NodeBackgroundColor", FLinearColor(FColor(48, 48, 48)));
 
 	const FTableRowStyle& NormalTableRowStyle = FEditorStyle::Get().GetWidgetStyle<FTableRowStyle>("TableView.Row");
 
diff --git a/Engine/Plugins/FX/Niagara/Source/NiagaraEditorWidgets/Private/SNiagaraOverviewStack.cpp b/Engine/Plugins/FX/Niagara/Source/NiagaraEditorWidgets/Private/SNiagaraOverviewStack.cpp
--- a/Engine/Plugins/FX/Niagara/Source/NiagaraEditorWidgets/Private/SNiagaraOverviewStack.cpp
+++ b/Engine/Plugins/FX/Niagara/Source/NiagaraEditorWidgets/Private/SNiagaraOverviewStack.cpp
@@ -48,42 +48,21 @@ class SNiagaraSystemOverviewEntryListRow : public STableRow<UNiagaraStackEntry*>
 	{
 		StackViewModel = InStackViewModel;
 		StackEntry = InStackEntry;
+		IssueIconVisibility = InArgs._IssueIconVisibility;
 		FSlateColor IconColor = FNiagaraEditorWidgetsStyle::Get().GetColor(FNiagaraStackEditorWidgetsUtilities::GetColorNameForExecutionCategory(StackEntry->GetExecutionCategoryName()));
-		ItemBackgroundColor = FNiagaraEditorWidgetsStyle::Get().GetColor("NiagaraEditor.Stack.Item.HeaderBackgroundColor");
-		DisabledItemBackgroundColor = ItemBackgroundColor + FLinearColor(.02f, .02f, .02f, 0.0f);
 
-		TSharedPtr<SWidget> WrappedContent;
+		FMargin ContentPadding;
 		if (StackEntry->IsA<UNiagaraStackItem>())
 		{
-			WrappedContent =
-				SNew(SBorder)
-				.BorderImage(FEditorStyle::GetBrush("WhiteBrush"))
-				.BorderBackgroundColor(this, &SNiagaraSystemOverviewEntryListRow::GetItemBackgroundColor, StackEntry)
-				.Padding(0)
-				[
-					SNew(SBorder)
-					.BorderImage(this, &SNiagaraSystemOverviewEntryListRow::GetBorder)
-					.Padding(FMargin(6, 4, 3, 4))
-					[
-						InArgs._Content.Widget
-					]
-				];
+			BackgroundColor = FNiagaraEditorWidgetsStyle::Get().GetColor("NiagaraEditor.SystemOverview.Item.BackgroundColor");
+			ContentPadding = FMargin(0, 2, 0, 2);
 		}
 		else
 		{
-			WrappedContent =
-				SNew(SBorder)
-				.BorderImage(FEditorStyle::GetBrush("NoBrush"))
-				.Padding(0)
-				[
-					SNew(SBorder)
-					.BorderImage(this, &SNiagaraSystemOverviewEntryListRow::GetBorder)
-					.Padding(FMargin(2, 4, 1, 4))
-					[
-						InArgs._Content.Widget
-					]
-				];
+			BackgroundColor = FNiagaraEditorWidgetsStyle::Get().GetColor("NiagaraEditor.SystemOverview.Group.BackgroundColor");
+			ContentPadding = FMargin(0, 5, 0, 5);
 		}
+		DisabledBackgroundColor = BackgroundColor + FLinearColor(.02f, .02f, .02f, 0.0f);
 
 		STableRow<UNiagaraStackEntry*>::Construct(STableRow<UNiagaraStackEntry*>::FArguments()
 			.Style(FNiagaraEditorWidgetsStyle::Get(), "NiagaraEditor.SystemOverview.TableViewRow")
@@ -92,24 +71,34 @@ class SNiagaraSystemOverviewEntryListRow : public STableRow<UNiagaraStackEntry*>
 			.OnCanAcceptDrop(InArgs._OnCanAcceptDrop)
 			.OnAcceptDrop(InArgs._OnAcceptDrop)
 		[
-			SNew(SBorder)
-			.BorderImage(this, &SNiagaraSystemOverviewEntryListRow::GetBorder)
-			.ToolTipText_UObject(StackEntry, &UNiagaraStackEntry::GetTooltipText)
-			.Padding(FMargin(5, 2, 3, 2))
+			SNew(SBox)
+			.Padding(FMargin(0, 1, 0, 1))
 			[
-				SNew(SHorizontalBox)
-				+ SHorizontalBox::Slot()
-				.Padding(0, 0, 2, 0)
-				[
-					WrappedContent.ToSharedRef()
-				]
-				+ SHorizontalBox::Slot()
-				.AutoWidth()
-				.VAlign(VAlign_Center)
-				.Padding(1, 0, 0, 0)
+				SNew(SBorder)
+				.BorderImage(FEditorStyle::GetBrush("WhiteBrush"))
+				.BorderBackgroundColor(this, &SNiagaraSystemOverviewEntryListRow::GetBackgroundColor, StackEntry)
+				.ToolTipText_UObject(StackEntry, &UNiagaraStackEntry::GetTooltipText)
+				.Padding(FMargin(0))
 				[
-					SNew(SNiagaraStackIssueIcon, StackViewModel, StackEntry)
-					.Visibility(InArgs._IssueIconVisibility)
+					SNew(SBorder)
+					.BorderImage(this, &SNiagaraSystemOverviewEntryListRow::GetBorder)
+					.Padding(ContentPadding)
+					[
+						SNew(SHorizontalBox)
+						+ SHorizontalBox::Slot()
+						.Padding(TAttribute<FMargin>(this, &SNiagaraSystemOverviewEntryListRow::GetInnerContentPadding))
+						[
+							InArgs._Content.Widget
+						]
+						+ SHorizontalBox::Slot()
+						.AutoWidth()
+						.VAlign(VAlign_Center)
+						.Padding(2, 0, 2, 0)
+						[
+							SNew(SNiagaraStackIssueIcon, StackViewModel, StackEntry)
+							.Visibility(IssueIconVisibility)
+						]
+					]
 				]
 			]
 		],
@@ -147,16 +136,29 @@ class SNiagaraSystemOverviewEntryListRow : public STableRow<UNiagaraStackEntry*>
 	}
 
 private:
-	FSlateColor GetItemBackgroundColor(UNiagaraStackEntry* Entry) const
+	FSlateColor GetBackgroundColor(UNiagaraStackEntry* Entry) const
+	{
+		return Entry->GetIsEnabled() && Entry->GetOwnerIsEnabled() ? BackgroundColor : DisabledBackgroundColor;
+	}
+
+	FMargin GetInnerContentPadding() const
 	{
-		return Entry->GetIsEnabled() && Entry->GetOwnerIsEnabled() ? ItemBackgroundColor : DisabledItemBackgroundColor;
+		if (IssueIconVisibility.Get() == EVisibility::Visible)
+		{
+			return FMargin(6, 0, 1, 0);
+		}
+		else
+		{
+			return FMargin(6, 0, 5, 0);
+		}
 	}
 
 private:
 	UNiagaraStackViewModel* StackViewModel;
 	UNiagaraStackEntry* StackEntry;
-	FLinearColor ItemBackgroundColor;
-	FLinearColor DisabledItemBackgroundColor;
+	FLinearColor BackgroundColor;
+	FLinearColor DisabledBackgroundColor;
+	TAttribute<EVisibility> IssueIconVisibility;
 };
 
 class SNiagaraSystemOverviewEnabledCheckBox : public SCompoundWidget
@@ -176,22 +178,16 @@ public:
 
 		ChildSlot
 		[
-			SNew(SBorder)
-			.BorderImage(FNiagaraEditorWidgetsStyle::Get().GetBrush("NiagaraEditor.SystemOverview.CheckBoxBorder"))
-			.BorderBackgroundColor(FNiagaraEditorWidgetsStyle::Get().GetColor("NiagaraEditor.Stack.Item.HeaderBackgroundColor"))
-			.Padding(FMargin(0))
+			SNew(SButton)
+			.ButtonStyle(FEditorStyle::Get(), "HoverHintOnly")
+			.ForegroundColor(FSlateColor::UseSubduedForeground())
+			.OnClicked(this, &SNiagaraSystemOverviewEnabledCheckBox::OnButtonClicked)
+			.ToolTipText(LOCTEXT("EnableCheckBoxToolTip", "Enable or disable this item."))
+			.ContentPadding(FMargin(3, 2, 2, 2))
 			[
-				SNew(SButton)
-				.ButtonStyle(FEditorStyle::Get(), "HoverHintOnly")
-				.OnClicked(this, &SNiagaraSystemOverviewEnabledCheckBox::OnButtonClicked)
-				.ToolTipText(LOCTEXT("EnableCheckBoxToolTip", "Enable or disable this item."))
-				.ContentPadding(FMargin(3, 2, 3, 2))
-				[
-					SNew(STextBlock)
-					.Font(FEditorStyle::Get().GetFontStyle("FontAwesome.10"))
-					.Text(this, &SNiagaraSystemOverviewEnabledCheckBox::GetButtonText)
-					.ColorAndOpacity(FNiagaraEditorWidgetsStyle::Get().GetColor("NiagaraEditor.SystemOverview.CheckBoxColor"))
-				]
+				SNew(STextBlock)
+				.Font(FEditorStyle::Get().GetFontStyle("FontAwesome.10"))
+				.Text(this, &SNiagaraSystemOverviewEnabledCheckBox::GetButtonText)
 			]
 		];
 	}
@@ -390,6 +386,7 @@ TSharedRef<ITableRow> SNiagaraOverviewStack::OnGenerateRowForEntry(UNiagaraStack
 			// Name
 			+ SHorizontalBox::Slot()
 			.VAlign(VAlign_Center)
+			.Padding(2, 3, 0, 3)
 			[
 				SNew(STextBlock)
 				.TextStyle(FNiagaraEditorWidgetsStyle::Get(), "NiagaraEditor.SystemOverview.ItemText")
@@ -415,7 +412,7 @@ TSharedRef<ITableRow> SNiagaraOverviewStack::OnGenerateRowForEntry(UNiagaraStack
 		TSharedRef<SHorizontalBox> ContentBox = SNew(SHorizontalBox)
 			// Execution category icon
 			+ SHorizontalBox::Slot()
-			.Padding(0, 0, 6, 0)
+			.Padding(0, 0, 4, 0)
 			.VAlign(VAlign_Center)
 			.AutoWidth()
 			[
@@ -427,6 +424,7 @@ TSharedRef<ITableRow> SNiagaraOverviewStack::OnGenerateRowForEntry(UNiagaraStack
 			// Name
 			+ SHorizontalBox::Slot()
 			.VAlign(VAlign_Center)
+			.Padding(0, 2, 0, 2)
 			[
 				SNew(STextBlock)
 				.TextStyle(FNiagaraEditorWidgetsStyle::Get(), "NiagaraEditor.SystemOverview.GroupHeaderText")
@@ -438,6 +436,7 @@ TSharedRef<ITableRow> SNiagaraOverviewStack::OnGenerateRowForEntry(UNiagaraStack
 		{
 			ContentBox->AddSlot()
 			.AutoWidth()
+			.Padding(0, 0, 1, 0)
 			[
 				SNew(SNiagaraStackItemGroupAddButton, *StackItemGroup)
 				.Width(22)
diff --git a/Engine/Plugins/FX/Niagara/Source/NiagaraEditorWidgets/Private/SNiagaraOverviewStackNode.cpp b/Engine/Plugins/FX/Niagara/Source/NiagaraEditorWidgets/Private/SNiagaraOverviewStackNode.cpp
--- a/Engine/Plugins/FX/Niagara/Source/NiagaraEditorWidgets/Private/SNiagaraOverviewStackNode.cpp
+++ b/Engine/Plugins/FX/Niagara/Source/NiagaraEditorWidgets/Private/SNiagaraOverviewStackNode.cpp
@@ -116,7 +116,7 @@ TSharedRef<SWidget> SNiagaraOverviewStackNode::CreateNodeContentArea()
 		.BorderImage(FEditorStyle::GetBrush("NoBorder"))
 		.HAlign(HAlign_Fill)
 		.VAlign(VAlign_Fill)
-		.Padding(FMargin(2, 3, 2, 2))
+		.Padding(FMargin(2, 2, 2, 4))
 		[
 			SNew(SHorizontalBox)
 			+SHorizontalBox::Slot()
@@ -130,7 +130,7 @@ TSharedRef<SWidget> SNiagaraOverviewStackNode::CreateNodeContentArea()
 			[
 				SNew(SBorder)
 				.BorderImage(FNiagaraEditorWidgetsStyle::Get().GetBrush("NiagaraEditor.SystemOverview.NodeBackgroundBorder"))
-				.BorderBackgroundColor(FNiagaraEditorWidgetsStyle::Get().GetColor("NiagaraEditor.SystemOverview.BackgroundColor"))
+				.BorderBackgroundColor(FNiagaraEditorWidgetsStyle::Get().GetColor("NiagaraEditor.SystemOverview.NodeBackgroundColor"))
 				.HAlign(HAlign_Fill)
 				.VAlign(VAlign_Fill)
 				.Padding(FMargin(0, 0, 0, 4))
