- Fix for potential divide by zero in both spawn rate and spawn per unit

#rnx

diff --git a/Engine/Plugins/FX/Niagara/Source/Niagara/Private/NiagaraSystemFastPath.cpp b/Engine/Plugins/FX/Niagara/Source/Niagara/Private/NiagaraSystemFastPath.cpp
--- a/Engine/Plugins/FX/Niagara/Source/Niagara/Private/NiagaraSystemFastPath.cpp
+++ b/Engine/Plugins/FX/Niagara/Source/Niagara/Private/NiagaraSystemFastPath.cpp
@@ -3,6 +3,14 @@
 #include "NiagaraSystemFastPath.h"
 #include "NiagaraEmitter.h"
 
+namespace FNiagaraSystemFastPath
+{
+	FORCEINLINE float SafeReciprocal(float v)
+	{
+		return FMath::Abs(v) > SMALL_NUMBER ? 1.0f / v : 0.0f;
+	}
+}
+
 FName FNiagaraSystemFastPath::FParameterNames::ExecutionState = "ExecutionState";
 FName FNiagaraSystemFastPath::FParameterNames::ExecutionStateSource = "ExecutionStateSource";
 
@@ -524,7 +532,7 @@ void FNiagaraEmitterFastPath::Module_SpawnRate(FParamMap0& Context_Map)
 		FParamMap0_Emitter_SpawnRate& Context_Map_Emitter_SpawnRate = Context_Map.Emitter.SpawnRate[SpawnRateIndex];
 
 		float Result57 = Context_Map_SpawnRate.SpawnRate * Context_Map.Scalability.Emitter.SpawnCountScale * Context_Map.Emitter.SpawnCountScale;
-		float Result58 = 1.0f / Result57;// Reciprocal(Result57);
+		float Result58 = FNiagaraSystemFastPath::SafeReciprocal(Result57);
 		float Result59 = 1 - Context_Map_Emitter_SpawnRate.SpawnRemainder;
 		float Result60 = Result58 * Result59;
 		float Context_Map_Local_SpawnRate_SpawnRate = Result57;
@@ -603,7 +611,7 @@ void FNiagaraEmitterFastPath::Module_SpawnPerUnit(FParamMap0& Context_Map)
 		float Constant54 = 500000;
 		float Result70 = FMath::Fmod(Context_Map.Emitter.DistanceTraveled, Constant54);
 		float Result71 = Result69 + Result70;
-		float Result72 = 1.0f / Context_Map_SpawnPerUnit.SpawnPerUnit; // Reciprocal(Context_Map_SpawnPerUnit.SpawnPerUnit);
+		float Result72 = FNiagaraSystemFastPath::SafeReciprocal(Context_Map_SpawnPerUnit.SpawnPerUnit);
 		float Result73 = Context_Map_Local_SpawnPerUnit_MovementThresholdVectorLength * Result72;
 		Context_Map.Emitter.DistanceTraveled = Result71;
 		float Context_Map_Local_SpawnPerUnit_SpawnSpacing = Result73;
@@ -611,7 +619,7 @@ void FNiagaraEmitterFastPath::Module_SpawnPerUnit(FParamMap0& Context_Map)
 		float Result75 = Result74 * Context_Map.Engine.DeltaTime + Context_Map_Emitter_SpawnPerUnit.SpawnRemainder;
 		int32 Result76 = FMath::FloorToInt(Result75);
 		float Result77 = Result75 - Result76;
-		float Result78 = 1.0f / Context_Map_Local_SpawnPerUnit_SpawnSpacing; // Reciprocal(Context.Map.Local.SpawnPerUnit.SpawnSpacing);
+		float Result78 = FNiagaraSystemFastPath::SafeReciprocal(Context_Map_Local_SpawnPerUnit_SpawnSpacing);
 		float Result79 = 1 - Context_Map_Emitter_SpawnPerUnit.SpawnRemainder;
 		float Result80 = Result79 * Result78;
 		int32 Context_Map_Local_SpawnPerUnit_SpawnCountInt = Result76;
