Fix static analysis warning V700: Consider inspecting the 'T foo = foo = ...' expression. It is odd that variable is initialized through itself.
#rnx

diff --git a/Engine/Plugins/FX/Niagara/Source/Niagara/Private/NiagaraDataInterfaceGrid2DCollection.cpp b/Engine/Plugins/FX/Niagara/Source/Niagara/Private/NiagaraDataInterfaceGrid2DCollection.cpp
--- a/Engine/Plugins/FX/Niagara/Source/Niagara/Private/NiagaraDataInterfaceGrid2DCollection.cpp
+++ b/Engine/Plugins/FX/Niagara/Source/Niagara/Private/NiagaraDataInterfaceGrid2DCollection.cpp
@@ -368,12 +368,7 @@ bool UNiagaraDataInterfaceGrid2DCollection::InitPerInstanceData(void* PerInstanc
 	// If we are setting the grid from the voxel size, then recompute NumVoxels and change bbox	
 	if (SetGridFromMaxAxis)
 	{
-		float CellSize = CellSize = WorldBBoxSize.Y / NumCellsMaxAxis;
-		if (WorldBBoxSize.X > WorldBBoxSize.Y)
-		{
-			CellSize = WorldBBoxSize.X / NumCellsMaxAxis;
-
-		}
+		float CellSize = FMath::Max(WorldBBoxSize.X, WorldBBoxSize.Y) / NumCellsMaxAxis;
 
 		RT_NumCellsX = WorldBBoxSize.X / CellSize;
 		RT_NumCellsY = WorldBBoxSize.Y / CellSize;
