Added a shared ref count to track the number of GPU emitters

#rnx

diff --git a/Engine/Plugins/FX/Niagara/Source/Niagara/Classes/NiagaraGPUInstanceCountManager.h b/Engine/Plugins/FX/Niagara/Source/Niagara/Classes/NiagaraGPUInstanceCountManager.h
--- a/Engine/Plugins/FX/Niagara/Source/Niagara/Classes/NiagaraGPUInstanceCountManager.h
+++ b/Engine/Plugins/FX/Niagara/Source/Niagara/Classes/NiagaraGPUInstanceCountManager.h
@@ -10,6 +10,7 @@ NiagaraGPUInstanceCountManager.h: GPU particle count handling
 #include "RHIUtilities.h"
 #include "NiagaraDrawIndirect.h"
 #include "RHIGPUReadback.h"
+#include "NiagaraRenderer.h"
 
 class FRHIGPUMemoryReadback;
 
@@ -58,8 +59,7 @@ public:
 	 * Called on the renderthread from FNiagaraRenderer::CreateRenderThreadResources()
 	 * or FNiagaraRenderer::ReleaseRenderThreadResources()
 	 */
-	void IncrementMaxDrawIndirectCount() { checkSlow(IsInRenderingThread()); ++MaxDrawIndirectArgs; }
-	void DecrementMaxDrawIndirectCount() { checkSlow(IsInRenderingThread()); --MaxDrawIndirectArgs; }
+	FORCEINLINE const TRefCountPtr<FNiagaraGPURendererCount>& GetGPURendererCount() { checkSlow(IsInRenderingThread()); return NumRegisteredGPURenderers; }
 
 	// Resize instance count and draw indirect buffers to ensure it is big enough to hold all draw indirect args.
 	void ResizeBuffers(FRHICommandListImmediate& RHICmdList, ERHIFeatureLevel::Type FeatureLevel, int32 ReservedInstanceCounts);
@@ -83,8 +83,8 @@ protected:
 	FRHIGPUMemoryReadback* CountReadback = nullptr;
 	int32 CountReadbackSize = 0;
 
-	/** The max possible required draw indirect args count based on the number of GPU renderer current registered */
-	int32 MaxDrawIndirectArgs = 0;
+	/** The number of GPU renderer. It defines the max possible required draw indirect args count */
+	TRefCountPtr<FNiagaraGPURendererCount> NumRegisteredGPURenderers;
 	/** The allocated indirect args in DrawIndirectBuffer (each being 5 uint32) */
 	int32 AllocatedDrawIndirectArgs = 0;
 
diff --git a/Engine/Plugins/FX/Niagara/Source/Niagara/Private/NiagaraComponent.cpp b/Engine/Plugins/FX/Niagara/Source/Niagara/Private/NiagaraComponent.cpp
--- a/Engine/Plugins/FX/Niagara/Source/Niagara/Private/NiagaraComponent.cpp
+++ b/Engine/Plugins/FX/Niagara/Source/Niagara/Private/NiagaraComponent.cpp
@@ -142,17 +142,15 @@ void FNiagaraSceneProxy::ReleaseRenderers()
 {
 	if (EmitterRenderers.Num() > 0)
 	{
-		NiagaraEmitterInstanceBatcher* TheBatcher = Batcher && !Batcher->IsPendingKill() ? Batcher : nullptr;
-
 		//Renderers must be freed on the render thread.
 		ENQUEUE_RENDER_COMMAND(ReleaseRenderersCommand)(
-			[ToDeleteEmitterRenderers = MoveTemp(EmitterRenderers), TheBatcher](FRHICommandListImmediate& RHICmdList)
+			[ToDeleteEmitterRenderers = MoveTemp(EmitterRenderers)](FRHICommandListImmediate& RHICmdList)
 		{
 			for (FNiagaraRenderer* EmitterRenderer : ToDeleteEmitterRenderers)
 			{
 				if (EmitterRenderer)
 				{
-					EmitterRenderer->ReleaseRenderThreadResources(TheBatcher);
+					EmitterRenderer->ReleaseRenderThreadResources();
 					delete EmitterRenderer;
 				}
 			}
@@ -215,13 +213,15 @@ void FNiagaraSceneProxy::CreateRenderers(const UNiagaraComponent* Component)
 
 FNiagaraSceneProxy::~FNiagaraSceneProxy()
 {
+	Batcher = nullptr;
+
 	//UE_LOG(LogNiagara, Warning, TEXT("~FNiagaraSceneProxy %p"), this);
 	check(IsInRenderingThread());
 	for (FNiagaraRenderer* EmitterRenderer : EmitterRenderers)
 	{
 		if (EmitterRenderer)
 		{
-			EmitterRenderer->ReleaseRenderThreadResources(Batcher);
+			EmitterRenderer->ReleaseRenderThreadResources();
 			delete EmitterRenderer;
 		}
 	}
@@ -234,10 +234,9 @@ void FNiagaraSceneProxy::ReleaseRenderThreadResources()
 	{
 		if (Renderer)
 		{
-			Renderer->ReleaseRenderThreadResources(Batcher);
+			Renderer->ReleaseRenderThreadResources();
 		}
 	}
-	return;
 }
 
 // FPrimitiveSceneProxy interface.
diff --git a/Engine/Plugins/FX/Niagara/Source/Niagara/Private/NiagaraGPUInstanceCountManager.cpp b/Engine/Plugins/FX/Niagara/Source/Niagara/Private/NiagaraGPUInstanceCountManager.cpp
--- a/Engine/Plugins/FX/Niagara/Source/Niagara/Private/NiagaraGPUInstanceCountManager.cpp
+++ b/Engine/Plugins/FX/Niagara/Source/Niagara/Private/NiagaraGPUInstanceCountManager.cpp
@@ -28,7 +28,9 @@ DECLARE_DWORD_COUNTER_STAT(TEXT("Max Num GPU Renderers"), STAT_NiagaraMaxNumGPUR
 //*****************************************************************************
 
 FNiagaraGPUInstanceCountManager::FNiagaraGPUInstanceCountManager() 
-{}
+{
+	NumRegisteredGPURenderers = new FNiagaraGPURendererCount();
+}
 
 FNiagaraGPUInstanceCountManager::~FNiagaraGPUInstanceCountManager()
 {
@@ -93,6 +95,7 @@ void FNiagaraGPUInstanceCountManager::FreeEntry(uint32& BufferOffset)
 void FNiagaraGPUInstanceCountManager::ResizeBuffers(FRHICommandListImmediate& RHICmdList, ERHIFeatureLevel::Type FeatureLevel, int32 ReservedInstanceCounts)
 {
 	const int32 RequiredInstanceCounts = UsedInstanceCounts + FMath::Max<int32>(ReservedInstanceCounts - FreeEntries.Num(), 0);
+	const int32 MaxDrawIndirectArgs = NumRegisteredGPURenderers->Value;
 	if (RequiredInstanceCounts > 0 || MaxDrawIndirectArgs > 0)
 	{
 		const int32 RecommendedInstanceCounts = FMath::Max(GNiagaraMinGPUInstanceCount, (int32)(RequiredInstanceCounts * BufferSlack));
@@ -172,6 +175,7 @@ uint32 FNiagaraGPUInstanceCountManager::AddDrawIndirect(uint32 InstanceCountBuff
 	else if (DrawIndirectArgGenTasks.Num() < AllocatedDrawIndirectArgs)
 	{
 #if !UE_BUILD_SHIPPING
+		const int32 MaxDrawIndirectArgs = NumRegisteredGPURenderers->Value;
 		if (DrawIndirectArgGenTasks.Num() >= MaxDrawIndirectArgs)
 		{
 			UE_LOG(LogNiagara, Warning, TEXT("More draw indirect args then expected (%d / %d)"), DrawIndirectArgGenTasks.Num() + 1, MaxDrawIndirectArgs);
diff --git a/Engine/Plugins/FX/Niagara/Source/Niagara/Private/NiagaraRenderer.cpp b/Engine/Plugins/FX/Niagara/Source/Niagara/Private/NiagaraRenderer.cpp
--- a/Engine/Plugins/FX/Niagara/Source/Niagara/Private/NiagaraRenderer.cpp
+++ b/Engine/Plugins/FX/Niagara/Source/Niagara/Private/NiagaraRenderer.cpp
@@ -272,7 +272,7 @@ void FNiagaraRenderer::Initialize(const UNiagaraRendererProperties *InProps, con
 
 FNiagaraRenderer::~FNiagaraRenderer()
 {
-	ReleaseRenderThreadResources(nullptr);
+	ReleaseRenderThreadResources();
 	SetDynamicData_RenderThread(nullptr);
 }
 
@@ -280,15 +280,20 @@ void FNiagaraRenderer::CreateRenderThreadResources(NiagaraEmitterInstanceBatcher
 {
 	if (Batcher && SimTarget == ENiagaraSimTarget::GPUComputeSim)
 	{
-		Batcher->GetGPUInstanceCounterManager().IncrementMaxDrawIndirectCount();
+		NumRegisteredGPURenderers = Batcher->GetGPUInstanceCounterManager().GetGPURendererCount();
+		if (NumRegisteredGPURenderers)
+		{
+			NumRegisteredGPURenderers->Value += 1;
+		}
 	}
 }
 
-void FNiagaraRenderer::ReleaseRenderThreadResources(NiagaraEmitterInstanceBatcher* Batcher)
+void FNiagaraRenderer::ReleaseRenderThreadResources()
 {
-	if (Batcher && SimTarget == ENiagaraSimTarget::GPUComputeSim)
+	if (NumRegisteredGPURenderers)
 	{
-		Batcher->GetGPUInstanceCounterManager().DecrementMaxDrawIndirectCount();
+		NumRegisteredGPURenderers->Value -= 1;
+		NumRegisteredGPURenderers.SafeRelease();
 	}
 }
 
diff --git a/Engine/Plugins/FX/Niagara/Source/Niagara/Private/NiagaraRendererMeshes.cpp b/Engine/Plugins/FX/Niagara/Source/Niagara/Private/NiagaraRendererMeshes.cpp
--- a/Engine/Plugins/FX/Niagara/Source/Niagara/Private/NiagaraRendererMeshes.cpp
+++ b/Engine/Plugins/FX/Niagara/Source/Niagara/Private/NiagaraRendererMeshes.cpp
@@ -126,9 +126,9 @@ FNiagaraRendererMeshes::~FNiagaraRendererMeshes()
 	VertexFactories.Empty();
 }
 
-void FNiagaraRendererMeshes::ReleaseRenderThreadResources(NiagaraEmitterInstanceBatcher* Batcher)
+void FNiagaraRendererMeshes::ReleaseRenderThreadResources()
 {
-	FNiagaraRenderer::ReleaseRenderThreadResources(Batcher);
+	FNiagaraRenderer::ReleaseRenderThreadResources();
 	for (FNiagaraMeshVertexFactory* VertexFactory : VertexFactories)
 	{
 		VertexFactory->ReleaseResource();
diff --git a/Engine/Plugins/FX/Niagara/Source/Niagara/Private/NiagaraRendererRibbons.cpp b/Engine/Plugins/FX/Niagara/Source/Niagara/Private/NiagaraRendererRibbons.cpp
--- a/Engine/Plugins/FX/Niagara/Source/Niagara/Private/NiagaraRendererRibbons.cpp
+++ b/Engine/Plugins/FX/Niagara/Source/Niagara/Private/NiagaraRendererRibbons.cpp
@@ -208,9 +208,9 @@ FNiagaraRendererRibbons::~FNiagaraRendererRibbons()
 {
 }
 
-void FNiagaraRendererRibbons::ReleaseRenderThreadResources(NiagaraEmitterInstanceBatcher* Batcher)
+void FNiagaraRendererRibbons::ReleaseRenderThreadResources()
 {
-	FNiagaraRenderer::ReleaseRenderThreadResources(Batcher);
+	FNiagaraRenderer::ReleaseRenderThreadResources();
 #if RHI_RAYTRACING
 	if (IsRayTracingEnabled())
 	{
diff --git a/Engine/Plugins/FX/Niagara/Source/Niagara/Private/NiagaraRendererSprites.cpp b/Engine/Plugins/FX/Niagara/Source/Niagara/Private/NiagaraRendererSprites.cpp
--- a/Engine/Plugins/FX/Niagara/Source/Niagara/Private/NiagaraRendererSprites.cpp
+++ b/Engine/Plugins/FX/Niagara/Source/Niagara/Private/NiagaraRendererSprites.cpp
@@ -155,9 +155,9 @@ FNiagaraRendererSprites::~FNiagaraRendererSprites()
 {
 }
 
-void FNiagaraRendererSprites::ReleaseRenderThreadResources(NiagaraEmitterInstanceBatcher* Batcher)
+void FNiagaraRendererSprites::ReleaseRenderThreadResources()
 {
-	FNiagaraRenderer::ReleaseRenderThreadResources(Batcher);
+	FNiagaraRenderer::ReleaseRenderThreadResources();
 
 	CutoutVertexBuffer.ReleaseResource();
 #if RHI_RAYTRACING
diff --git a/Engine/Plugins/FX/Niagara/Source/Niagara/Public/NiagaraRenderer.h b/Engine/Plugins/FX/Niagara/Source/Niagara/Public/NiagaraRenderer.h
--- a/Engine/Plugins/FX/Niagara/Source/Niagara/Public/NiagaraRenderer.h
+++ b/Engine/Plugins/FX/Niagara/Source/Niagara/Public/NiagaraRenderer.h
@@ -76,6 +76,14 @@ struct FNiagaraRendererVariableInfo
 	bool bUpload;
 };
 
+// The number of GPU renderers registered in the instance count manager.
+// Shared between the manager and the renderers.
+class FNiagaraGPURendererCount : public FRefCountedObject
+{
+public:
+	int32 Value = 0;
+};
+
 /**
 * Base class for Niagara System renderers.
 */
@@ -91,7 +99,7 @@ public:
 
 	virtual void Initialize(const UNiagaraRendererProperties *InProps, const FNiagaraEmitterInstance* Emitter);
 	virtual void CreateRenderThreadResources(NiagaraEmitterInstanceBatcher* Batcher);
-	virtual void ReleaseRenderThreadResources(NiagaraEmitterInstanceBatcher* Batcher);
+	virtual void ReleaseRenderThreadResources();
 
 	virtual FPrimitiveViewRelevance GetViewRelevance(const FSceneView* View, const FNiagaraSceneProxy *SceneProxy)const;
 	virtual void GetDynamicMeshElements(const TArray<const FSceneView*>& Views, const FSceneViewFamily& ViewFamily, uint32 VisibilityMap, FMeshElementCollector& Collector, const FNiagaraSceneProxy *SceneProxy) const {}
@@ -130,7 +138,7 @@ protected:
 
 	uint32 bLocalSpace : 1;
 	uint32 bHasLights : 1;
-	ENiagaraSimTarget SimTarget;
+	const ENiagaraSimTarget SimTarget;
 	uint32 NumIndicesPerInstance;
 
 	ERHIFeatureLevel::Type FeatureLevel;
@@ -149,5 +157,6 @@ protected:
 	TArray<UMaterialInterface*> BaseMaterials_GT;
 	FMaterialRelevance BaseMaterialRelevance_GT;
 
+	TRefCountPtr<FNiagaraGPURendererCount> NumRegisteredGPURenderers;
 };
 
diff --git a/Engine/Plugins/FX/Niagara/Source/Niagara/Public/NiagaraRendererMeshes.h b/Engine/Plugins/FX/Niagara/Source/Niagara/Public/NiagaraRendererMeshes.h
--- a/Engine/Plugins/FX/Niagara/Source/Niagara/Public/NiagaraRendererMeshes.h
+++ b/Engine/Plugins/FX/Niagara/Source/Niagara/Public/NiagaraRendererMeshes.h
@@ -21,7 +21,7 @@ public:
 	
 	//FNiagaraRenderer Interface
 	virtual void CreateRenderThreadResources(NiagaraEmitterInstanceBatcher* Batcher) override;
-	virtual void ReleaseRenderThreadResources(NiagaraEmitterInstanceBatcher* Batcher) override;
+	virtual void ReleaseRenderThreadResources() override;
 
 	virtual void GetDynamicMeshElements(const TArray<const FSceneView*>& Views, const FSceneViewFamily& ViewFamily, uint32 VisibilityMap, FMeshElementCollector& Collector, const FNiagaraSceneProxy *SceneProxy) const override;
 	virtual FNiagaraDynamicDataBase* GenerateDynamicData(const FNiagaraSceneProxy* Proxy, const UNiagaraRendererProperties* InProperties, const FNiagaraEmitterInstance* Emitter) const override;
diff --git a/Engine/Plugins/FX/Niagara/Source/Niagara/Public/NiagaraRendererRibbons.h b/Engine/Plugins/FX/Niagara/Source/Niagara/Public/NiagaraRendererRibbons.h
--- a/Engine/Plugins/FX/Niagara/Source/Niagara/Public/NiagaraRendererRibbons.h
+++ b/Engine/Plugins/FX/Niagara/Source/Niagara/Public/NiagaraRendererRibbons.h
@@ -22,7 +22,7 @@ public:
 
 	// FNiagaraRenderer Interface 
 	virtual void CreateRenderThreadResources(NiagaraEmitterInstanceBatcher* Batcher) override;
-	virtual void ReleaseRenderThreadResources(NiagaraEmitterInstanceBatcher* Batcher) override;
+	virtual void ReleaseRenderThreadResources() override;
 
 	virtual void GetDynamicMeshElements(const TArray<const FSceneView*>& Views, const FSceneViewFamily& ViewFamily, uint32 VisibilityMap, FMeshElementCollector& Collector, const FNiagaraSceneProxy *SceneProxy) const override;
 	virtual FNiagaraDynamicDataBase *GenerateDynamicData(const FNiagaraSceneProxy* Proxy, const UNiagaraRendererProperties* InProperties, const FNiagaraEmitterInstance* Emitteride) const override;
diff --git a/Engine/Plugins/FX/Niagara/Source/Niagara/Public/NiagaraRendererSprites.h b/Engine/Plugins/FX/Niagara/Source/Niagara/Public/NiagaraRendererSprites.h
--- a/Engine/Plugins/FX/Niagara/Source/Niagara/Public/NiagaraRendererSprites.h
+++ b/Engine/Plugins/FX/Niagara/Source/Niagara/Public/NiagaraRendererSprites.h
@@ -21,7 +21,7 @@ public:
 
 	//FNiagaraRenderer interface
 	virtual void CreateRenderThreadResources(NiagaraEmitterInstanceBatcher* Batcher)override;
-	virtual void ReleaseRenderThreadResources(NiagaraEmitterInstanceBatcher* Batcher)override;
+	virtual void ReleaseRenderThreadResources()override;
 
 	virtual void GetDynamicMeshElements(const TArray<const FSceneView*>& Views, const FSceneViewFamily& ViewFamily, uint32 VisibilityMap, FMeshElementCollector& Collector, const FNiagaraSceneProxy *SceneProxy) const override;
 	virtual FNiagaraDynamicDataBase* GenerateDynamicData(const FNiagaraSceneProxy* Proxy, const UNiagaraRendererProperties* InProperties, const FNiagaraEmitterInstance* Emitter) const override;
