Fixes editor crashes when creating a circular reference in the niagara graph.

diff --git a/Engine/Plugins/FX/Niagara/Source/NiagaraEditor/Private/EdGraphSchema_Niagara.cpp b/Engine/Plugins/FX/Niagara/Source/NiagaraEditor/Private/EdGraphSchema_Niagara.cpp
--- a/Engine/Plugins/FX/Niagara/Source/NiagaraEditor/Private/EdGraphSchema_Niagara.cpp
+++ b/Engine/Plugins/FX/Niagara/Source/NiagaraEditor/Private/EdGraphSchema_Niagara.cpp
@@ -1011,8 +1011,8 @@ const FPinConnectionResponse UEdGraphSchema_Niagara::CanCreateConnection(const U
 		}
 	}
 
-	int32 Depth = 0;
-	if (UEdGraphSchema_Niagara::CheckCircularConnection(PinB->GetOwningNode(), PinB->Direction, PinA, Depth))
+	TSet<const UEdGraphNode*> VisitedNodes;
+	if (UEdGraphSchema_Niagara::CheckCircularConnection(VisitedNodes, OutputPin->GetOwningNode(), InputPin->GetOwningNode()))
 	{
 		return FPinConnectionResponse(CONNECT_RESPONSE_DISALLOW, TEXT("Circular connection found"));
 	}
@@ -1523,41 +1523,42 @@ void UEdGraphSchema_Niagara::ConvertNumericPinToType(UEdGraphPin* InGraphPin, FN
 	}
 }
 
-bool UEdGraphSchema_Niagara::CheckCircularConnection(const UEdGraphNode* InRootNode, const EEdGraphPinDirection InRootPinDirection, const UEdGraphPin* InPin, int32& OutDepth)
+bool UEdGraphSchema_Niagara::CheckCircularConnection(TSet<const UEdGraphNode*>& VisitedNodes, const UEdGraphNode* InNode, const UEdGraphNode* InTestNode)
 {
-	if (InPin->GetOwningNode() == InRootNode)
+	bool AlreadyAdded = false;
+
+	VisitedNodes.Add(InNode, &AlreadyAdded);
+	if (AlreadyAdded)
 	{
-		return true;
+		// node is already in our set, so return so we don't reprocess it
+		return false;
 	}
 
-	static const int32 MaxDepth = 3;
-	OutDepth++;
-	if (OutDepth > MaxDepth)
+	if (InNode == InTestNode)
 	{
-		return false;
+		// we've found a match, so we have a circular reference
+		return true;
 	}
 
-	for (const UEdGraphPin* Pin : InPin->GetOwningNode()->GetAllPins())
+	// iterate over all of the nodes that are inputs to InNode
+	for (const UEdGraphPin* Pin : InNode->GetAllPins())
 	{
-		if (Pin->Direction == InRootPinDirection && Pin != InPin)
+		if (Pin && Pin->Direction == EGPD_Input)
 		{
-			for (const UEdGraphPin* LinkedPin : Pin->LinkedTo)
+			for (const UEdGraphPin* OutputPin : Pin->LinkedTo)
 			{
-				if (UEdGraphSchema_Niagara::CheckCircularConnection(InRootNode, InRootPinDirection, LinkedPin, OutDepth))
+				if (const UEdGraphNode* InputNode = OutputPin ? OutputPin->GetOwningNode() : nullptr)
 				{
-					return true;
-				}
-
-				// If the CheckCircularConnection call above returned without finding the root node and was too deep.
-				if (OutDepth > MaxDepth)
-				{
-					return false;
+					if (CheckCircularConnection(VisitedNodes, InputNode, InTestNode))
+					{
+						return true;
+					}
 				}
 			}
+
 		}
 	}
 
-	OutDepth--;
 	return false;
 }
 
diff --git a/Engine/Plugins/FX/Niagara/Source/NiagaraEditor/Public/EdGraphSchema_Niagara.h b/Engine/Plugins/FX/Niagara/Source/NiagaraEditor/Public/EdGraphSchema_Niagara.h
--- a/Engine/Plugins/FX/Niagara/Source/NiagaraEditor/Public/EdGraphSchema_Niagara.h
+++ b/Engine/Plugins/FX/Niagara/Source/NiagaraEditor/Public/EdGraphSchema_Niagara.h
@@ -154,7 +154,8 @@ private:
 	void GetNumericConversionToSubMenuActionsAll(class UToolMenu* Menu, const FName SectionName, UNiagaraNode* InGraphPin);
 	void ConvertNumericPinToType(UEdGraphPin* InPin, FNiagaraTypeDefinition TypeDef);
 	void ConvertNumericPinToTypeAll(UNiagaraNode* InPin, FNiagaraTypeDefinition TypeDef);
-	static bool CheckCircularConnection(const UEdGraphNode* InRootNode, const EEdGraphPinDirection InRootPinDirection, const UEdGraphPin* InPin, int32& OutDepth);
+
+	static bool CheckCircularConnection(TSet<const UEdGraphNode*>& VisitedNodes, const UEdGraphNode* InNode, const UEdGraphNode* InTestNode);
 };
 
 class FNiagaraConnectionDrawingPolicy : public FConnectionDrawingPolicy
