Niagara - Fix an issue where modules, parameters, and assets couldn't be dropped onto an empty stack.

diff --git a/Engine/Plugins/FX/Niagara/Source/NiagaraEditor/Private/ViewModels/Stack/NiagaraStackScriptItemGroup.cpp b/Engine/Plugins/FX/Niagara/Source/NiagaraEditor/Private/ViewModels/Stack/NiagaraStackScriptItemGroup.cpp
--- a/Engine/Plugins/FX/Niagara/Source/NiagaraEditor/Private/ViewModels/Stack/NiagaraStackScriptItemGroup.cpp
+++ b/Engine/Plugins/FX/Niagara/Source/NiagaraEditor/Private/ViewModels/Stack/NiagaraStackScriptItemGroup.cpp
@@ -458,13 +458,14 @@ void UNiagaraStackScriptItemGroup::RefreshIssues(TArray<FStackIssue>& NewIssues)
 
 void GenerateDragDropData(
 	UNiagaraNodeFunctionCall& SourceModule,
-	UNiagaraNodeFunctionCall& TargetModule, EItemDropZone TargetZone,
+	UNiagaraNodeFunctionCall* TargetModule,
+	UNiagaraNodeOutput* TargetOutputNode,
+	EItemDropZone TargetZone,
 	TArray<FNiagaraStackGraphUtilities::FStackNodeGroup>& OutSourceStackGroups, int32& OutSourceGroupIndex,
 	TArray<FNiagaraStackGraphUtilities::FStackNodeGroup>& OutTargetStackGroups, int32& OutTargetGroupIndex)
 {
-	// Find the output nodes for the source and target
+	// Find the output node for the source
 	UNiagaraNodeOutput* SourceOutputNode = FNiagaraStackGraphUtilities::GetEmitterOutputNodeForStackNode(SourceModule);
-	UNiagaraNodeOutput* TargetOutputNode = FNiagaraStackGraphUtilities::GetEmitterOutputNodeForStackNode(TargetModule);
 
 	// Collect the stack node groups for the source and target.
 	FNiagaraStackGraphUtilities::GetStackNodeGroups(*SourceOutputNode, OutSourceStackGroups);
@@ -488,47 +489,93 @@ void GenerateDragDropData(
 		}
 	}
 
-	OutTargetGroupIndex = INDEX_NONE;
-	if (&TargetModule == &SourceModule)
+	if (TargetModule == nullptr)
 	{
-		OutTargetGroupIndex = OutSourceGroupIndex;
+		// If no target module was supplied then the module should be inserted at the beginning.
+		OutTargetGroupIndex = 1;
 	}
 	else
 	{
-		for (int32 GroupIndex = 0; GroupIndex < OutTargetStackGroups.Num(); GroupIndex++)
+		if (TargetModule == &SourceModule)
 		{
-			if (OutTargetStackGroups[GroupIndex].EndNode == &TargetModule)
+			OutTargetGroupIndex = OutSourceGroupIndex;
+		}
+		else
+		{
+			for (int32 GroupIndex = 0; GroupIndex < OutTargetStackGroups.Num(); GroupIndex++)
 			{
-				OutTargetGroupIndex = GroupIndex;
-				break;
+				if (OutTargetStackGroups[GroupIndex].EndNode == TargetModule)
+				{
+					OutTargetGroupIndex = GroupIndex;
+					break;
+				}
 			}
 		}
-	}
 
-	if (TargetZone == EItemDropZone::BelowItem)
-	{
-		OutTargetGroupIndex++;
+		if (TargetZone == EItemDropZone::BelowItem)
+		{
+			OutTargetGroupIndex++;
+		}
 	}
 }
 
+TOptional<UNiagaraStackEntry::FDropRequestResponse> UNiagaraStackScriptItemGroup::CanDropInternal(const FDropRequest& DropRequest)
+{
+	return CanDropOnTarget(*this, DropRequest);
+}
+
+TOptional<UNiagaraStackEntry::FDropRequestResponse> UNiagaraStackScriptItemGroup::DropInternal(const FDropRequest& DropRequest)
+{
+	return DropOnTarget(*this, DropRequest);
+}
+
 TOptional<UNiagaraStackEntry::FDropRequestResponse> UNiagaraStackScriptItemGroup::ChildRequestCanDropInternal(const UNiagaraStackEntry& TargetChild, const FDropRequest& DropRequest)
+{
+	return CanDropOnTarget(TargetChild, DropRequest);
+}
+
+TOptional<UNiagaraStackEntry::FDropRequestResponse> UNiagaraStackScriptItemGroup::ChildRequestDropInternal(const UNiagaraStackEntry& TargetChild, const FDropRequest& DropRequest)
+{
+	return DropOnTarget(TargetChild, DropRequest);
+}
+
+TOptional<UNiagaraStackEntry::FDropRequestResponse> UNiagaraStackScriptItemGroup::CanDropOnTarget(const UNiagaraStackEntry& TargetEntry, const FDropRequest& DropRequest)
 {
 	if (DropRequest.DragDropOperation->IsOfType<FNiagaraStackEntryDragDropOp>())
 	{
-		return ChildRequestCanDropEntries(TargetChild, DropRequest);
+		return CanDropEntriesOnTarget(TargetEntry, DropRequest);
 	}
 	else if (DropRequest.DragDropOperation->IsOfType<FAssetDragDropOp>())
 	{
-		return ChildRequestCanDropAssets(TargetChild, DropRequest);
+		return CanDropAssetsOnTarget(TargetEntry, DropRequest);
 	}
 	else if (DropRequest.DragDropOperation->IsOfType<FNiagaraParameterDragOperation>())
 	{
-		return ChildRequestCanDropParameter(TargetChild, DropRequest);
+		return CanDropParameterOnTarget(TargetEntry, DropRequest);
 	}
 	return TOptional<FDropRequestResponse>();
 }
 
-TOptional<UNiagaraStackEntry::FDropRequestResponse> UNiagaraStackScriptItemGroup::ChildRequestCanDropEntries(const UNiagaraStackEntry& TargetChild, const FDropRequest& DropRequest)
+TOptional<EItemDropZone> GetTargetDropZoneForTargetEntry(UNiagaraStackScriptItemGroup* ThisEntry, const UNiagaraStackEntry& TargetEntry, EItemDropZone RequestedDropZone)
+{
+	TOptional<EItemDropZone> TargetDropZone;
+	if (&TargetEntry == ThisEntry)
+	{
+		// Items dragged onto the group entry are always inserted below it.
+		TargetDropZone = EItemDropZone::BelowItem;
+	}
+	else
+	{
+		// Otherwise only allow drops above or below items.
+		if (RequestedDropZone == EItemDropZone::AboveItem || RequestedDropZone == EItemDropZone::BelowItem)
+		{
+			TargetDropZone = RequestedDropZone;
+		}
+	}
+	return TargetDropZone;
+}
+
+TOptional<UNiagaraStackEntry::FDropRequestResponse> UNiagaraStackScriptItemGroup::CanDropEntriesOnTarget(const UNiagaraStackEntry& TargetEntry, const FDropRequest& DropRequest)
 {
 	TSharedRef<const FNiagaraStackEntryDragDropOp> StackEntryDragDropOp = StaticCastSharedRef<const FNiagaraStackEntryDragDropOp>(DropRequest.DragDropOperation);
 	bool ModuleEntriesDragged = false;
@@ -546,9 +593,9 @@ TOptional<UNiagaraStackEntry::FDropRequestResponse> UNiagaraStackScriptItemGroup
 		// Only handle dragged module items.
 		return TOptional<FDropRequestResponse>();
 	}
-	if (TargetChild.IsA<UNiagaraStackModuleItem>() == false)
+	if (&TargetEntry != this && TargetEntry.IsA<UNiagaraStackModuleItem>() == false)
 	{
-		// Only handle child drop requests from module items.
+		// Only handle drops onto this script group, or child drop requests from module items.
 		return TOptional<FDropRequestResponse>();
 	}
 	if (DropRequest.DropOptions != UNiagaraStackEntry::EDropOptions::Overview)
@@ -573,56 +620,59 @@ TOptional<UNiagaraStackEntry::FDropRequestResponse> UNiagaraStackScriptItemGroup
 	{
 		return FDropRequestResponse(TOptional<EItemDropZone>(), LOCTEXT("CantMoveModuleByUsage", "This module can't be moved to this section of the\nstack because it's not valid for this usage context."));
 	}
-	if (DropRequest.DropZone != EItemDropZone::AboveItem && DropRequest.DropZone != EItemDropZone::BelowItem)
+
+	TOptional<EItemDropZone> TargetDropZone = GetTargetDropZoneForTargetEntry(this, TargetEntry, DropRequest.DropZone);
+	if(TargetDropZone.IsSet() == false)
 	{
-		// Only handle drops before and after items.
 		return TOptional<FDropRequestResponse>();
 	}
 
-	const UNiagaraStackModuleItem* TargetModuleItem = Cast<UNiagaraStackModuleItem>(&TargetChild);
-	if (TargetModuleItem != nullptr)
-	{
-		TArray<FNiagaraStackGraphUtilities::FStackNodeGroup> SourceStackGroups;
-		TArray<FNiagaraStackGraphUtilities::FStackNodeGroup> TargetStackGroups;
-		int32 SourceGroupIndex;
-		int32 TargetGroupIndex;
-		GenerateDragDropData(
-			SourceModuleItem->GetModuleNode(), TargetModuleItem->GetModuleNode(), DropRequest.DropZone,
-			SourceStackGroups, SourceGroupIndex,
-			TargetStackGroups, TargetGroupIndex);
-
-		// Make sure the source and target indices are within safe ranges, and make sure that the insert target isn't the source target or the spot directly
-		// after the source target since that won't actually move the module.
-		bool bSourceWithinRange = SourceGroupIndex > 0 && SourceGroupIndex < SourceStackGroups.Num() - 1;
-		bool bTargetWithinRange = TargetGroupIndex > 0 && TargetGroupIndex < TargetStackGroups.Num();
-		bool bWillMove = DropRequest.DragOptions == EDragOptions::Copy ||
-			(SourceStackGroups[SourceGroupIndex].EndNode != TargetStackGroups[TargetGroupIndex].EndNode &&
-				SourceStackGroups[SourceGroupIndex].EndNode != TargetStackGroups[TargetGroupIndex - 1].EndNode);
-		if (bSourceWithinRange && bTargetWithinRange && bWillMove)
+	const UNiagaraStackModuleItem* TargetModuleItem = Cast<UNiagaraStackModuleItem>(&TargetEntry);
+
+	TArray<FNiagaraStackGraphUtilities::FStackNodeGroup> SourceStackGroups;
+	TArray<FNiagaraStackGraphUtilities::FStackNodeGroup> TargetStackGroups;
+	int32 SourceGroupIndex;
+	int32 TargetGroupIndex;
+	UNiagaraNodeFunctionCall* TargetModuleNode = TargetModuleItem != nullptr ? &TargetModuleItem->GetModuleNode() : nullptr;
+	UNiagaraNodeOutput* TargetOutputNode = GetScriptOutputNode();
+	GenerateDragDropData(
+		SourceModuleItem->GetModuleNode(), TargetModuleNode, 
+		TargetOutputNode, DropRequest.DropZone,
+		SourceStackGroups, SourceGroupIndex,
+		TargetStackGroups, TargetGroupIndex);
+
+	// Make sure the source and target indices are within safe ranges, and make sure that the insert target isn't the source target or the spot directly
+	// after the source target since that won't actually move the module.
+	bool bSourceWithinRange = SourceGroupIndex > 0 && SourceGroupIndex < SourceStackGroups.Num() - 1;
+	bool bTargetWithinRange = TargetGroupIndex > 0 && TargetGroupIndex < TargetStackGroups.Num();
+	bool bWillMove = DropRequest.DragOptions == EDragOptions::Copy ||
+		(SourceStackGroups[SourceGroupIndex].EndNode != TargetStackGroups[TargetGroupIndex].EndNode &&
+			SourceStackGroups[SourceGroupIndex].EndNode != TargetStackGroups[TargetGroupIndex - 1].EndNode);
+	if (bSourceWithinRange && bTargetWithinRange && bWillMove)
+	{
+		FText DropMessage;
+		if (DropRequest.DragOptions == UNiagaraStackEntry::EDragOptions::Copy)
 		{
-			FText DropMessage;
-			if (DropRequest.DragOptions == UNiagaraStackEntry::EDragOptions::Copy)
-			{
-				DropMessage = LOCTEXT("CopyModuleResult", "Copy this module here.");
-			}
-			else
-			{
-				DropMessage = LOCTEXT("MoveModuleResult", "Move this module here.");
-			}
-			return FDropRequestResponse(DropRequest.DropZone, DropMessage);
+			DropMessage = LOCTEXT("CopyModuleResult", "Copy this module here.");
+		}
+		else
+		{
+			DropMessage = LOCTEXT("MoveModuleResult", "Move this module here.");
 		}
+		return FDropRequestResponse(TargetDropZone, DropMessage);
 	}
+
 	return TOptional<FDropRequestResponse>();
 }
 
-TOptional<UNiagaraStackEntry::FDropRequestResponse> UNiagaraStackScriptItemGroup::ChildRequestCanDropAssets(const UNiagaraStackEntry& TargetChild, const FDropRequest& DropRequest)
+TOptional<UNiagaraStackEntry::FDropRequestResponse> UNiagaraStackScriptItemGroup::CanDropAssetsOnTarget(const UNiagaraStackEntry& TargetEntry, const FDropRequest& DropRequest)
 {
 	TSharedRef<const FAssetDragDropOp> AssetDragDropOp = StaticCastSharedRef<const FAssetDragDropOp>(DropRequest.DragDropOperation);
 	const UEnum* NiagaraScriptUsageEnum = FindObjectChecked<UEnum>(ANY_PACKAGE, TEXT("ENiagaraScriptUsage"), true);
 
-	if (TargetChild.IsA<UNiagaraStackModuleItem>() == false)
+	if (&TargetEntry != this && TargetEntry.IsA<UNiagaraStackModuleItem>() == false)
 	{
-		// Only handle child drop requests from module items.
+		// Only handle drops onto this script group, or child drop requests from module items.
 		return TOptional<FDropRequestResponse>();
 	}
 	if (DropRequest.DropOptions != UNiagaraStackEntry::EDropOptions::Overview)
@@ -668,9 +718,9 @@ TOptional<UNiagaraStackEntry::FDropRequestResponse> UNiagaraStackScriptItemGroup
 		}
 	}
 
-	if (DropRequest.DropZone != EItemDropZone::AboveItem && DropRequest.DropZone != EItemDropZone::BelowItem)
+	TOptional<EItemDropZone> TargetDropZone = GetTargetDropZoneForTargetEntry(this, TargetEntry, DropRequest.DropZone);
+	if (TargetDropZone.IsSet() == false)
 	{
-		// Only handle drops before and after items.
 		return TOptional<FDropRequestResponse>();
 	}
 
@@ -683,17 +733,17 @@ TOptional<UNiagaraStackEntry::FDropRequestResponse> UNiagaraStackScriptItemGroup
 	{
 		DropMessage = LOCTEXT("DropAsset", "Insert a module for this asset here.");
 	}
-	return FDropRequestResponse(DropRequest.DropZone, DropMessage);
+	return FDropRequestResponse(TargetDropZone, DropMessage);
 }
 
-TOptional<UNiagaraStackEntry::FDropRequestResponse> UNiagaraStackScriptItemGroup::ChildRequestCanDropParameter(const UNiagaraStackEntry& TargetChild, const FDropRequest& DropRequest)
+TOptional<UNiagaraStackEntry::FDropRequestResponse> UNiagaraStackScriptItemGroup::CanDropParameterOnTarget(const UNiagaraStackEntry& TargetEntry, const FDropRequest& DropRequest)
 {
 	TSharedRef<const FNiagaraParameterDragOperation> ParameterDragDropOp = StaticCastSharedRef<const FNiagaraParameterDragOperation>(DropRequest.DragDropOperation);
 	TSharedPtr<const FNiagaraParameterAction> ParameterAction = StaticCastSharedPtr<const FNiagaraParameterAction>(ParameterDragDropOp->GetSourceAction());
 
-	if (TargetChild.IsA<UNiagaraStackModuleItem>() == false)
+	if (&TargetEntry != this && TargetEntry.IsA<UNiagaraStackModuleItem>() == false)
 	{
-		// Only handle child drop requests from module items.
+		// Only handle drops onto this script group, or child drop requests from module items.
 		return TOptional<FDropRequestResponse>();
 	}
 	if (DropRequest.DropOptions != UNiagaraStackEntry::EDropOptions::Overview)
@@ -710,14 +760,14 @@ TOptional<UNiagaraStackEntry::FDropRequestResponse> UNiagaraStackScriptItemGroup
 		}
 
 		{
-			if (DropRequest.DropZone != EItemDropZone::AboveItem && DropRequest.DropZone != EItemDropZone::BelowItem)
+			TOptional<EItemDropZone> TargetDropZone = GetTargetDropZoneForTargetEntry(this, TargetEntry, DropRequest.DropZone);
+			if (TargetDropZone.IsSet() == false)
 			{
-				// Only handle drops before and after items.
 				return TOptional<FDropRequestResponse>();
 			}
 			else
 			{
-				return FDropRequestResponse(DropRequest.DropZone, LOCTEXT("DropParameterFormat", "Insert a module for this parameter here."));
+				return FDropRequestResponse(TargetDropZone, LOCTEXT("DropParameterFormat", "Insert a module for this parameter here."));
 			}
 		}
 
@@ -725,24 +775,45 @@ TOptional<UNiagaraStackEntry::FDropRequestResponse> UNiagaraStackScriptItemGroup
 	return TOptional<FDropRequestResponse>();
 }
 
-TOptional<UNiagaraStackEntry::FDropRequestResponse> UNiagaraStackScriptItemGroup::ChildRequestDropInternal(const UNiagaraStackEntry& TargetChild, const FDropRequest& DropRequest)
+TOptional<UNiagaraStackEntry::FDropRequestResponse> UNiagaraStackScriptItemGroup::DropOnTarget(const UNiagaraStackEntry& TargetEntry, const FDropRequest& DropRequest)
 {
 	if (DropRequest.DragDropOperation->IsOfType<FNiagaraStackEntryDragDropOp>())
 	{
-		return ChildRequestDropEntries(TargetChild, DropRequest);
+		return DropEntriesOnTarget(TargetEntry, DropRequest);
 	}
 	else if (DropRequest.DragDropOperation->IsOfType<FAssetDragDropOp>())
 	{
-		return ChildRequestDropAssets(TargetChild, DropRequest);
+		return DropAssetsOnTarget(TargetEntry, DropRequest);
 	}
 	else if (DropRequest.DragDropOperation->IsOfType<FNiagaraParameterDragOperation>())
 	{
-		return ChildRequestDropParameter(TargetChild, DropRequest);
+		return DropParameterOnTarget(TargetEntry, DropRequest);
 	}
 	return TOptional<FDropRequestResponse>();
 }
 
-TOptional<UNiagaraStackEntry::FDropRequestResponse> UNiagaraStackScriptItemGroup::ChildRequestDropEntries(const UNiagaraStackEntry& TargetChild, const FDropRequest& DropRequest)
+int32 GetTargetIndexForTargetEntry(UNiagaraStackScriptItemGroup* ThisEntry, const UNiagaraStackEntry& TargetEntry, EItemDropZone DropZone)
+{
+	int32 TargetIndex;
+	if (&TargetEntry == ThisEntry)
+	{
+		// Items dragged onto this group entry are always inserted at index 0.
+		TargetIndex = 0;
+	}
+	else
+	{
+		// Otherwise get the index from the module and the drop zone.
+		const UNiagaraStackModuleItem* TargetModuleItem = CastChecked<UNiagaraStackModuleItem>(&TargetEntry);
+		TargetIndex = TargetModuleItem->GetModuleIndex();
+		if (DropZone == EItemDropZone::BelowItem)
+		{
+			TargetIndex++;
+		}
+	}
+	return TargetIndex;
+}
+
+TOptional<UNiagaraStackEntry::FDropRequestResponse> UNiagaraStackScriptItemGroup::DropEntriesOnTarget(const UNiagaraStackEntry& TargetEntry, const FDropRequest& DropRequest)
 {
 	TSharedRef<const FNiagaraStackEntryDragDropOp> StackEntryDragDropOp = StaticCastSharedRef<const FNiagaraStackEntryDragDropOp>(DropRequest.DragDropOperation);
 
@@ -757,18 +828,14 @@ TOptional<UNiagaraStackEntry::FDropRequestResponse> UNiagaraStackScriptItemGroup
 	UNiagaraScript* SourceModuleScript = FNiagaraEditorUtilities::GetScriptFromSystem(SourceModuleItem->GetSystemViewModel()->GetSystem(), SourceEmitterHandleId,
 		SourceModuleOutputNode->GetUsage(), SourceModuleOutputNode->GetUsageId());
 
-	const UNiagaraStackModuleItem* TargetModuleItem = CastChecked<UNiagaraStackModuleItem>(&TargetChild);
+	const UNiagaraStackModuleItem* TargetModuleItem = Cast<UNiagaraStackModuleItem>(&TargetEntry);
 	const FNiagaraEmitterHandle* TargetEmitterHandle = GetEmitterViewModel().IsValid()
 		? FNiagaraEditorUtilities::GetEmitterHandleForEmitter(GetSystemViewModel()->GetSystem(), *GetEmitterViewModel()->GetEmitter())
 		: nullptr;
 	FGuid TargetEmitterHandleId = TargetEmitterHandle != nullptr
 		? TargetEmitterHandle->GetId()
 		: FGuid();
-	int32 TargetIndex = TargetModuleItem->GetModuleIndex();
-	if (DropRequest.DropZone == EItemDropZone::BelowItem)
-	{
-		TargetIndex++;
-	}
+	int32 TargetIndex = GetTargetIndexForTargetEntry(this, TargetEntry, DropRequest.DropZone);
 
 	FScopedTransaction ScopedTransaction(LOCTEXT("DragAndDropModule", "Drag and drop module"));
 	{
@@ -784,22 +851,17 @@ TOptional<UNiagaraStackEntry::FDropRequestResponse> UNiagaraStackScriptItemGroup
 	}
 
 	SourceModuleItem->OnRequestFullRefreshDeferred().Broadcast();
-	TargetChild.OnRequestFullRefreshDeferred().Broadcast();
+	TargetEntry.OnRequestFullRefreshDeferred().Broadcast();
 	return FDropRequestResponse(DropRequest.DropZone);
 }
 
-TOptional<UNiagaraStackEntry::FDropRequestResponse> UNiagaraStackScriptItemGroup::ChildRequestDropAssets(const UNiagaraStackEntry& TargetChild, const FDropRequest& DropRequest)
+TOptional<UNiagaraStackEntry::FDropRequestResponse> UNiagaraStackScriptItemGroup::DropAssetsOnTarget(const UNiagaraStackEntry& TargetEntry, const FDropRequest& DropRequest)
 {
 	TSharedRef<const FAssetDragDropOp> AssetDragDropOp = StaticCastSharedRef<const FAssetDragDropOp>(DropRequest.DragDropOperation);
 
-	const UNiagaraStackModuleItem* TargetModuleItem = CastChecked<UNiagaraStackModuleItem>(&TargetChild);
-	int32 TargetIndex = TargetModuleItem->GetModuleIndex();
-	if (DropRequest.DropZone == EItemDropZone::BelowItem)
-	{
-		TargetIndex++;
-	}
+	int32 TargetIndex = GetTargetIndexForTargetEntry(this, TargetEntry, DropRequest.DropZone);
 
-	FScopedTransaction ScopedTransaction(LOCTEXT("DragAndDropAsset", "Drag and drop assets"));
+	FScopedTransaction ScopedTransaction(LOCTEXT("DragAndDropAsset", "Insert modules for assets"));
 	for (const FAssetData& AssetData : AssetDragDropOp->GetAssets())
 	{
 		TSharedRef<FScriptGroupAddAction> AddAction = FScriptGroupAddAction::CreateAssetModuleAction(AssetData);
@@ -809,18 +871,12 @@ TOptional<UNiagaraStackEntry::FDropRequestResponse> UNiagaraStackScriptItemGroup
 	return FDropRequestResponse(DropRequest.DropZone);
 }
 
-TOptional<UNiagaraStackEntry::FDropRequestResponse> UNiagaraStackScriptItemGroup::ChildRequestDropParameter(const UNiagaraStackEntry& TargetChild, const FDropRequest& DropRequest)
+TOptional<UNiagaraStackEntry::FDropRequestResponse> UNiagaraStackScriptItemGroup::DropParameterOnTarget(const UNiagaraStackEntry& TargetEntry, const FDropRequest& DropRequest)
 {
 	TSharedRef<const FNiagaraParameterDragOperation> ParameterDragDropOp = StaticCastSharedRef<const FNiagaraParameterDragOperation>(DropRequest.DragDropOperation);
 	TSharedPtr<FNiagaraParameterAction> ParameterAction = StaticCastSharedPtr<FNiagaraParameterAction>(ParameterDragDropOp->GetSourceAction());
 
-	const UNiagaraStackModuleItem* TargetModuleItem = CastChecked<UNiagaraStackModuleItem>(&TargetChild);
-	int32 TargetIndex = TargetModuleItem->GetModuleIndex();
-	if (DropRequest.DropZone == EItemDropZone::BelowItem)
-	{
-		TargetIndex++;
-	}
-
+	int32 TargetIndex = GetTargetIndexForTargetEntry(this, TargetEntry, DropRequest.DropZone);
 	TSharedRef<FScriptGroupAddAction> AddAction = FScriptGroupAddAction::CreateExistingParameterModuleAction(ParameterAction->GetParameter());
 	AddUtilities->ExecuteAddAction(AddAction, TargetIndex);
 
diff --git a/Engine/Plugins/FX/Niagara/Source/NiagaraEditor/Public/ViewModels/Stack/NiagaraStackScriptItemGroup.h b/Engine/Plugins/FX/Niagara/Source/NiagaraEditor/Public/ViewModels/Stack/NiagaraStackScriptItemGroup.h
--- a/Engine/Plugins/FX/Niagara/Source/NiagaraEditor/Public/ViewModels/Stack/NiagaraStackScriptItemGroup.h
+++ b/Engine/Plugins/FX/Niagara/Source/NiagaraEditor/Public/ViewModels/Stack/NiagaraStackScriptItemGroup.h
@@ -36,6 +36,10 @@ protected:
 
 	virtual void FinalizeInternal() override;
 
+	virtual TOptional<FDropRequestResponse> CanDropInternal(const FDropRequest& DropRequest) override;
+
+	virtual TOptional<FDropRequestResponse> DropInternal(const FDropRequest& DropRequest) override;
+
 	virtual TOptional<FDropRequestResponse> ChildRequestCanDropInternal(const UNiagaraStackEntry& TargetChild, const FDropRequest& DropRequest) override;
 
 	virtual TOptional<FDropRequestResponse> ChildRequestDropInternal(const UNiagaraStackEntry& TargetChild, const FDropRequest& DropRequest) override;
@@ -47,17 +51,21 @@ private:
 
 	void OnScriptGraphChanged(const struct FEdGraphEditAction& InAction);
 
-	TOptional<FDropRequestResponse> ChildRequestCanDropEntries(const UNiagaraStackEntry& TargetChild, const FDropRequest& DropRequest);
+	TOptional<FDropRequestResponse> CanDropOnTarget(const UNiagaraStackEntry& TargetEntry, const FDropRequest& DropRequest);
+
+	TOptional<FDropRequestResponse> CanDropEntriesOnTarget(const UNiagaraStackEntry& TargetEntry, const FDropRequest& DropRequest);
+
+	TOptional<FDropRequestResponse> CanDropAssetsOnTarget(const UNiagaraStackEntry& TargetEntry, const FDropRequest& DropRequest);
 
-	TOptional<FDropRequestResponse> ChildRequestCanDropAssets(const UNiagaraStackEntry& TargetChild, const FDropRequest& DropRequest);
+	TOptional<FDropRequestResponse> CanDropParameterOnTarget(const UNiagaraStackEntry& TargetEntry, const FDropRequest& DropRequest);
 
-	TOptional<FDropRequestResponse> ChildRequestCanDropParameter(const UNiagaraStackEntry& TargetChild, const FDropRequest& DropRequest);
+	TOptional<FDropRequestResponse> DropOnTarget(const UNiagaraStackEntry& TargetEntry, const FDropRequest& DropRequest);
 
-	TOptional<FDropRequestResponse> ChildRequestDropEntries(const UNiagaraStackEntry& TargetChild, const FDropRequest& DropRequest);
+	TOptional<FDropRequestResponse> DropEntriesOnTarget(const UNiagaraStackEntry& TargetEntry, const FDropRequest& DropRequest);
 
-	TOptional<FDropRequestResponse> ChildRequestDropAssets(const UNiagaraStackEntry& TargetChild, const FDropRequest& DropRequest);
+	TOptional<FDropRequestResponse> DropAssetsOnTarget(const UNiagaraStackEntry& TargetEntry, const FDropRequest& DropRequest);
 
-	TOptional<FDropRequestResponse> ChildRequestDropParameter(const UNiagaraStackEntry& TargetChild, const FDropRequest& DropRequest);
+	TOptional<FDropRequestResponse> DropParameterOnTarget(const UNiagaraStackEntry& TargetEntry, const FDropRequest& DropRequest);
 
 protected:
 	TWeakPtr<FNiagaraScriptViewModel> ScriptViewModel;
