Camera DI first implementation; work in progress

diff --git a/Engine/Plugins/FX/Niagara/Source/Niagara/Classes/NiagaraDataInterfaceCamera.h b/Engine/Plugins/FX/Niagara/Source/Niagara/Classes/NiagaraDataInterfaceCamera.h
new file mode 100644
--- /dev/null
+++ b/Engine/Plugins/FX/Niagara/Source/Niagara/Classes/NiagaraDataInterfaceCamera.h
@@ -0,0 +1,68 @@
+// Copyright 1998-2019 Epic Games, Inc. All Rights Reserved.
+#pragma once
+
+#include "NiagaraCommon.h"
+#include "NiagaraShared.h"
+#include "NiagaraDataInterface.h"
+#include "NiagaraParameterStore.h"
+#include "Camera/PlayerCameraManager.h"
+#include "NiagaraDataInterfaceCamera.generated.h"
+
+struct CameraDataInterface_InstanceData
+{
+	TWeakObjectPtr<APlayerCameraManager> CameraObject = nullptr;
+
+	/** A binding to the user ptr we're reading the camera from. */
+	//FNiagaraParameterDirectBinding<UObject*> UserParamBinding;
+};
+
+UCLASS(EditInlineNew, Category = "Camera", meta = (DisplayName = "Camera Query"))
+class NIAGARA_API UNiagaraDataInterfaceCamera : public UNiagaraDataInterface
+{
+	GENERATED_UCLASS_BODY()
+
+public:
+	/** Reference to a user parameter that should receive the particle data after the simulation tick. The supplied parameter object needs to implement the INiagaraParticleCallbackHandler interface. */
+	//UPROPERTY(EditAnywhere, Category = "Camera")
+	//FNiagaraUserParameterBinding CameraParameter;
+
+	//UObject Interface
+	virtual void PostInitProperties() override;
+	//UObject Interface End
+
+	//UNiagaraDataInterface Interface
+	virtual void GetFunctions(TArray<FNiagaraFunctionSignature>& OutFunctions)override;
+	virtual void GetVMExternalFunction(const FVMExternalFunctionBindingInfo& BindingInfo, void* InstanceData, FVMExternalFunction &OutFunc) override;
+	virtual bool InitPerInstanceData(void* PerInstanceData, FNiagaraSystemInstance* SystemInstance) override;
+	virtual int32 PerInstanceDataSize() const override { return sizeof(CameraDataInterface_InstanceData); }
+	virtual bool PerInstanceTick(void* PerInstanceData, FNiagaraSystemInstance* SystemInstance, float DeltaSeconds) override;
+	virtual bool GetFunctionHLSL(const FName&  DefinitionFunctionName, FString InstanceFunctionName, FNiagaraDataInterfaceGPUParamInfo& ParamInfo, FString& OutHLSL) override;
+	virtual bool CanExecuteOnTarget(ENiagaraSimTarget Target) const override { return true; }
+	virtual FNiagaraDataInterfaceParametersCS* ConstructComputeParameters() const override;
+	//UNiagaraDataInterface Interface
+
+	void QueryOcclusionFactorGPU(FVectorVMContext& Context);
+	void GetCameraFOV(FVectorVMContext& Context);
+	void GetCameraPosition(FVectorVMContext& Context);
+	void GetViewPropertiesGPU(FVectorVMContext& Context);
+	void GetClipSpaceTransformsGPU(FVectorVMContext& Context);
+	void GetViewSpaceTransformsGPU(FVectorVMContext& Context);
+	
+private:
+	static const FName GetCameraOcclusionName;
+	static const FName GetViewPropertiesName;
+	static const FName GetClipSpaceTransformsName;
+	static const FName GetViewSpaceTransformsName;
+	static const FName GetCameraPositionsName;
+	static const FName GetFieldOfViewName;
+};
+
+struct FNiagaraDataIntefaceProxyCameraQuery : public FNiagaraDataInterfaceProxy
+{
+	// There's nothing in this proxy. It just reads from scene textures.
+
+	virtual int32 PerInstanceDataPassedToRenderThreadSize() const override
+	{
+		return 0;
+	}
+};
\ No newline at end of file
diff --git a/Engine/Plugins/FX/Niagara/Source/Niagara/Private/NiagaraDataInterfaceCamera.cpp b/Engine/Plugins/FX/Niagara/Source/Niagara/Private/NiagaraDataInterfaceCamera.cpp
new file mode 100644
--- /dev/null
+++ b/Engine/Plugins/FX/Niagara/Source/Niagara/Private/NiagaraDataInterfaceCamera.cpp
@@ -0,0 +1,511 @@
+// Copyright 1998-2019 Epic Games, Inc. All Rights Reserved.
+
+#include "NiagaraDataInterfaceCamera.h"
+#include "NiagaraTypes.h"
+#include "NiagaraCustomVersion.h"
+#include "NiagaraWorldManager.h"
+#include "Internationalization/Internationalization.h"
+#include "NiagaraSystemInstance.h"
+
+const FName UNiagaraDataInterfaceCamera::GetCameraOcclusionName(TEXT("QueryOcclusionFactorGPU"));
+const FName UNiagaraDataInterfaceCamera::GetViewPropertiesName(TEXT("GetViewPropertiesGPU"));
+const FName UNiagaraDataInterfaceCamera::GetClipSpaceTransformsName(TEXT("GetClipSpaceTransformsGPU"));
+const FName UNiagaraDataInterfaceCamera::GetViewSpaceTransformsName(TEXT("GetViewSpaceTransformsGPU"));
+const FName UNiagaraDataInterfaceCamera::GetCameraPositionsName(TEXT("GetCameraPositionCPU/GPU"));
+const FName UNiagaraDataInterfaceCamera::GetFieldOfViewName(TEXT("GetFieldOfView"));
+
+UNiagaraDataInterfaceCamera::UNiagaraDataInterfaceCamera(FObjectInitializer const& ObjectInitializer)
+	: Super(ObjectInitializer)
+{
+	Proxy = MakeShared<FNiagaraDataIntefaceProxyCameraQuery, ESPMode::ThreadSafe>();
+}
+
+void UNiagaraDataInterfaceCamera::PostInitProperties()
+{
+	Super::PostInitProperties();
+
+	if (HasAnyFlags(RF_ClassDefaultObject))
+	{
+		FNiagaraTypeRegistry::Register(FNiagaraTypeDefinition(GetClass()), true, false, false);
+	}
+}
+
+bool UNiagaraDataInterfaceCamera::InitPerInstanceData(void* PerInstanceData, FNiagaraSystemInstance* SystemInstance)
+{
+	CameraDataInterface_InstanceData* PIData = new (PerInstanceData) CameraDataInterface_InstanceData;
+	return true;
+}
+
+bool UNiagaraDataInterfaceCamera::PerInstanceTick(void* PerInstanceData, FNiagaraSystemInstance* SystemInstance, float DeltaSeconds)
+{
+	CameraDataInterface_InstanceData* PIData = (CameraDataInterface_InstanceData*)PerInstanceData;
+	if (!PIData)
+	{
+		return true;
+	}
+	
+	// todo get the correct camera here
+	UWorld* World = SystemInstance->GetWorldManager()->World;
+	if (World->GetFirstPlayerController())
+	{
+		PIData->CameraObject = World->GetFirstPlayerController()->PlayerCameraManager;
+	}
+	else
+	{
+		PIData->CameraObject = nullptr;
+	}
+	
+	return false;
+}
+
+void UNiagaraDataInterfaceCamera::GetFunctions(TArray<FNiagaraFunctionSignature>& OutFunctions)
+{
+	FNiagaraFunctionSignature Sig;
+	Sig.Name = GetCameraOcclusionName;
+#if WITH_EDITORONLY_DATA
+	Sig.Description = NSLOCTEXT("Niagara", "GetCameraOcclusionFunctionDescription", "This function returns the occlusion factor of the camera.");
+#endif
+	Sig.bMemberFunction = true;
+	Sig.bRequiresContext = false;
+	Sig.Inputs.Add(FNiagaraVariable(FNiagaraTypeDefinition(GetClass()), TEXT("Camera interface")));
+	Sig.Inputs.Add(FNiagaraVariable(FNiagaraTypeDefinition::GetVec3Def(), TEXT("Sample Center World Position")));
+	Sig.Inputs.Add(FNiagaraVariable(FNiagaraTypeDefinition::GetFloatDef(), TEXT("Sample Window Size World")));
+	Sig.Inputs.Add(FNiagaraVariable(FNiagaraTypeDefinition::GetFloatDef(), TEXT("Sample Steps Per Line")));
+	Sig.Outputs.Add(FNiagaraVariable(FNiagaraTypeDefinition::GetFloatDef(), TEXT("Occlusion Factor")));
+	Sig.Outputs.Add(FNiagaraVariable(FNiagaraTypeDefinition::GetFloatDef(), TEXT("Sample Factor")));
+	OutFunctions.Add(Sig);
+
+
+	Sig = FNiagaraFunctionSignature();
+	Sig.Name = GetViewPropertiesName;
+#if WITH_EDITORONLY_DATA
+	Sig.Description = NSLOCTEXT("Niagara", "GetCameraPropertiesDescription", "This function returns the properties of the current camera. Only valid for gpu particles.");
+#endif
+	Sig.bMemberFunction = true;
+	Sig.bRequiresContext = false;
+	Sig.Inputs.Add(FNiagaraVariable(FNiagaraTypeDefinition(GetClass()), TEXT("Camera interface")));
+	Sig.Outputs.Add(FNiagaraVariable(FNiagaraTypeDefinition::GetVec3Def(), TEXT("View Position World")));
+	Sig.Outputs.Add(FNiagaraVariable(FNiagaraTypeDefinition::GetVec3Def(), TEXT("View Forward Vector")));
+	Sig.Outputs.Add(FNiagaraVariable(FNiagaraTypeDefinition::GetVec3Def(), TEXT("View Up Vector")));
+	Sig.Outputs.Add(FNiagaraVariable(FNiagaraTypeDefinition::GetVec3Def(), TEXT("View Right Vector")));
+	Sig.Outputs.Add(FNiagaraVariable(FNiagaraTypeDefinition::GetVec4Def(), TEXT("View Size And Inverse Size")));
+	Sig.Outputs.Add(FNiagaraVariable(FNiagaraTypeDefinition::GetVec4Def(), TEXT("Screen To View Space")));
+	OutFunctions.Add(Sig);
+
+
+	Sig = FNiagaraFunctionSignature();
+	Sig.Name = GetClipSpaceTransformsName;
+#if WITH_EDITORONLY_DATA
+	Sig.Description = NSLOCTEXT("Niagara", "GetCameraPropertiesDescription", "This function returns the properties of the current camera. Only valid for gpu particles.");
+#endif
+	Sig.bMemberFunction = true;
+	Sig.bRequiresContext = false;
+	Sig.Inputs.Add(FNiagaraVariable(FNiagaraTypeDefinition(GetClass()), TEXT("Camera interface")));
+	Sig.Outputs.Add(FNiagaraVariable(FNiagaraTypeDefinition::GetMatrix4Def(), TEXT("World To Clip Transform")));
+	Sig.Outputs.Add(FNiagaraVariable(FNiagaraTypeDefinition::GetMatrix4Def(), TEXT("Translated World To Clip Transform")));
+	Sig.Outputs.Add(FNiagaraVariable(FNiagaraTypeDefinition::GetMatrix4Def(), TEXT("Clip To World Transform")));
+	Sig.Outputs.Add(FNiagaraVariable(FNiagaraTypeDefinition::GetMatrix4Def(), TEXT("Clip To View Transform")));
+	Sig.Outputs.Add(FNiagaraVariable(FNiagaraTypeDefinition::GetMatrix4Def(), TEXT("Clip To Translated World Transform")));
+	Sig.Outputs.Add(FNiagaraVariable(FNiagaraTypeDefinition::GetMatrix4Def(), TEXT("Screen To World Transform")));
+	Sig.Outputs.Add(FNiagaraVariable(FNiagaraTypeDefinition::GetMatrix4Def(), TEXT("Screen To Translated World Transform")));
+	Sig.Outputs.Add(FNiagaraVariable(FNiagaraTypeDefinition::GetMatrix4Def(), TEXT("Clip To Previous Clip Transform")));
+	OutFunctions.Add(Sig);
+
+
+	Sig = FNiagaraFunctionSignature();
+	Sig.Name = GetViewSpaceTransformsName;
+#if WITH_EDITORONLY_DATA
+	Sig.Description = NSLOCTEXT("Niagara", "GetCameraPropertiesDescription", "This function returns the properties of the current camera. Only valid for gpu particles.");
+#endif
+	Sig.bMemberFunction = true;
+	Sig.bRequiresContext = false;
+	Sig.Inputs.Add(FNiagaraVariable(FNiagaraTypeDefinition(GetClass()), TEXT("Camera interface")));
+	Sig.Outputs.Add(FNiagaraVariable(FNiagaraTypeDefinition::GetMatrix4Def(), TEXT("Translated World To View Transform")));
+	Sig.Outputs.Add(FNiagaraVariable(FNiagaraTypeDefinition::GetMatrix4Def(), TEXT("View To Translated World Transform")));
+	Sig.Outputs.Add(FNiagaraVariable(FNiagaraTypeDefinition::GetMatrix4Def(), TEXT("Translated World To Camera View Transform")));
+	Sig.Outputs.Add(FNiagaraVariable(FNiagaraTypeDefinition::GetMatrix4Def(), TEXT("Camera View To Translated World Transform")));
+	Sig.Outputs.Add(FNiagaraVariable(FNiagaraTypeDefinition::GetMatrix4Def(), TEXT("View To Clip Transform")));
+	Sig.Outputs.Add(FNiagaraVariable(FNiagaraTypeDefinition::GetMatrix4Def(), TEXT("View To ClipNoAA Transform")));
+	OutFunctions.Add(Sig);
+
+
+	Sig = FNiagaraFunctionSignature();
+	Sig.Name = GetFieldOfViewName;
+#if WITH_EDITORONLY_DATA
+	Sig.Description = NSLOCTEXT("Niagara", "GetCameraPropertiesDescription", "This function returns the properties of the current camera. Only valid for gpu particles.");
+#endif
+	Sig.bMemberFunction = true;
+	Sig.bRequiresContext = false;
+	Sig.Inputs.Add(FNiagaraVariable(FNiagaraTypeDefinition(GetClass()), TEXT("Camera interface")));
+	Sig.Outputs.Add(FNiagaraVariable(FNiagaraTypeDefinition::GetFloatDef(), TEXT("Field Of View Angle")));
+	OutFunctions.Add(Sig);
+
+
+	Sig = FNiagaraFunctionSignature();
+	Sig.Name = GetCameraPositionsName;
+#if WITH_EDITORONLY_DATA
+	Sig.Description = NSLOCTEXT("Niagara", "GetCameraPropertiesDescription", "This function returns the properties of the current camera. Only valid for gpu particles.");
+#endif
+	Sig.bMemberFunction = true;
+	Sig.bRequiresContext = false;
+	Sig.Inputs.Add(FNiagaraVariable(FNiagaraTypeDefinition(GetClass()), TEXT("Camera interface")));
+	Sig.Outputs.Add(FNiagaraVariable(FNiagaraTypeDefinition::GetVec3Def(), TEXT("Camera Position World")));
+	OutFunctions.Add(Sig);
+}
+
+bool UNiagaraDataInterfaceCamera::GetFunctionHLSL(const FName& DefinitionFunctionName, FString InstanceFunctionName, FNiagaraDataInterfaceGPUParamInfo& ParamInfo, FString& OutHLSL)
+{
+	TMap<FString, FStringFormatArg> ArgsSample;
+	ArgsSample.Add(TEXT("FunctionName"), InstanceFunctionName);
+
+	if (DefinitionFunctionName == GetCameraOcclusionName)
+	{
+		/*
+		OutHLSL += TEXT("\
+			Out_Occlusion = -1;\n\
+			Out_CameraPosWorld.xyz = View.WorldCameraOrigin.xyz;\n\
+			float4 SamplePosition = float4(In_SamplePos + View.PreViewTranslation, 1);\n\
+			float4 ClipPosition = mul(SamplePosition, View.TranslatedWorldToClip);\n\
+			float2 ScreenPosition = ClipPosition.xy / ClipPosition.w;\n\
+			// Check if the sample is inside the view.\n\
+			if (all(abs(ScreenPosition.xy) <= float2(1, 1)))\n\
+			{\n\
+				// Sample the depth buffer to get a world position near the sample position.\n\
+				float2 ScreenUV = ScreenPosition * View.ScreenPositionScaleBias.xy + View.ScreenPositionScaleBias.wz;\n\
+				float SceneDepth = CalcSceneDepth(ScreenUV);\n\
+				Out_SceneDepth = SceneDepth;\n\
+				// Reconstruct world position.\n\
+				Out_WorldPos = WorldPositionFromSceneDepth(ScreenPosition.xy, SceneDepth);\n\
+				// Sample the normal buffer\n\
+				Out_WorldNormal = Texture2DSampleLevel(SceneTexturesStruct.GBufferATexture, SceneTexturesStruct.GBufferATextureSampler, ScreenUV, 0).xyz * 2.0 - 1.0;\n\
+			}\n\
+			else\n\
+			{\n\
+				Out_IsInsideView = false;\n\
+			} \n}\n\n");*/
+
+		static const TCHAR *FormatSample = TEXT(R"(
+			void {FunctionName}(in float3 In_SampleCenterWorldPos, in float In_SampleWindowSizeWorld, in float In_SampleSteps, out float Out_OcclusionFactor, out float Out_SampleFactor)
+			{
+				float4 SamplePosition = float4(In_SampleCenterWorldPos + View.PreViewTranslation, 1);
+				float CameraDistance = length(SamplePosition.xyz - View.WorldCameraOrigin.xyz);	
+				float4 ClipPosition = mul(SamplePosition, View.TranslatedWorldToClip);
+				float2 ScreenPosition = ClipPosition.xy / ClipPosition.w;
+				float2 ScreenUV = ScreenPosition * View.ScreenPositionScaleBias.xy + View.ScreenPositionScaleBias.wz;
+
+				float Steps = In_SampleSteps <= 1 ? 0 : In_SampleSteps;
+				float2 SampleUV = float2(0, 0);
+				float TotalSamples = 0;
+				float OccludedSamples = 0;
+				float SampleWidth = 0.2; // get from In_SampleWindowSizeWorld
+
+				if (Steps > 0) 
+				{
+					for (int ys = 0; ys < Steps; ys++)
+					{
+						SampleUV.y = ScreenUV.y - 0.5 * SampleWidth + ys * SampleWidth / (Steps - 1);
+						if (SampleUV.y > 1 || SampleUV.y < 0)
+						{
+							continue;
+						}
+						for (int xs = 0; xs < Steps; xs++)
+						{
+							SampleUV.x = ScreenUV.x - 0.5 * SampleWidth + xs * SampleWidth / (Steps - 1);
+							if (SampleUV.x > 1 || SampleUV.x < 0)
+							{
+								continue;
+							}
+			
+							float Depth = CalcSceneDepth(SampleUV);
+							if (Depth < CameraDistance) 
+							{
+								OccludedSamples++;
+							}
+							TotalSamples++;
+						} 
+					}
+				}	
+
+				Out_OcclusionFactor = TotalSamples > 0 ? OccludedSamples / TotalSamples : 1;
+				Out_SampleFactor = Steps == 0 ? 0 : (TotalSamples / (Steps * Steps));
+			}
+		)");
+		OutHLSL += FString::Format(FormatSample, ArgsSample);
+		return true;
+	}
+	if (DefinitionFunctionName == GetViewPropertiesName)
+	{
+		static const TCHAR *FormatSample = TEXT(R"(
+			void {FunctionName}(out float3 Out_ViewPositionWorld, out float3 Out_ViewForwardVector, out float3 Out_ViewUpVector, out float3 Out_ViewRightVector, out float4 Out_ViewSizeAndInverseSize, out float4 Out_ScreenToViewSpace)
+			{
+				Out_ViewPositionWorld.xyz = View.WorldViewOrigin.xyz;
+				Out_ViewForwardVector.xyz = View.ViewForward.xyz;
+				Out_ViewUpVector.xyz = View.ViewUp.xyz;
+				Out_ViewRightVector.xyz = View.ViewRight.xyz;
+				Out_ViewSizeAndInverseSize = View.ViewSizeAndInvSize;
+				Out_ScreenToViewSpace = View.ScreenToViewSpace;
+			}
+		)");
+		OutHLSL += FString::Format(FormatSample, ArgsSample);
+		return true;
+	}
+	if (DefinitionFunctionName == GetFieldOfViewName)
+	{
+		static const TCHAR *FormatSample = TEXT(R"(
+			void {FunctionName}(out float Out_FieldOfViewAngle)
+			{
+				Out_FieldOfViewAngle = degrees(View.FieldOfViewWideAngles.x);
+			}
+		)");
+		OutHLSL += FString::Format(FormatSample, ArgsSample);
+		return true;
+	}
+	if (DefinitionFunctionName == GetClipSpaceTransformsName)
+	{
+		static const TCHAR *FormatSample = TEXT(R"(
+			void {FunctionName}out float4x4 Out_WorldToClipTransform, out float4x4 Out_TranslatedWorldToClipTransform, out float4x4 Out_ClipToWorldTransform, out float4x4 Out_ClipToViewTransform,
+				out float4x4 Out_ClipToTranslatedWorldTransform, out float4x4 Out_ScreenToWorldTransform, out float4x4 Out_ScreenToTranslatedWorldTransform, out float4x4 Out_ClipToPreviousClipTransform)
+			{
+				Out_WorldToClipTransform = View.WorldToClip;
+				Out_TranslatedWorldToClipTransform = View.TranslatedWorldToClip;
+				Out_ClipToWorldTransform = View.ClipToWorld;
+				Out_ClipToViewTransform = View.ClipToView;
+				Out_ClipToTranslatedWorldTransform = View.ClipToTranslatedWorld;
+				Out_ScreenToWorldTransform = View. ScreenToWorld;
+				Out_ScreenToTranslatedWorldTransform = View.ScreenToTranslatedWorld;
+				Out_ClipToPreviousClipTransform = View.ClipToPrevClip;
+			}
+		)");
+		OutHLSL += FString::Format(FormatSample, ArgsSample);
+		return true;
+	}
+	if (DefinitionFunctionName == GetViewSpaceTransformsName)
+	{
+		static const TCHAR *FormatSample = TEXT(R"(
+			void {FunctionName}(out float4x4 Out_TranslatedWorldToViewTransform, out float4x4 Out_ViewToTranslatedWorldTransform, out float4x4 Out_TranslatedWorldToCameraViewTransform,
+				out float4x4 Out_CameraViewToTranslatedWorldTransform, out float4x4 Out_ViewToClipTransform, out float4x4 Out_ViewToClipNoAATransform)
+			{
+				Out_TranslatedWorldToViewTransform = View.TranslatedWorldToView;
+				Out_ViewToTranslatedWorldTransform = View.ViewToTranslatedWorld;
+				Out_TranslatedWorldToCameraViewTransform = View.TranslatedWorldToCameraView;
+				Out_CameraViewToTranslatedWorldTransform = View.CameraViewToTranslatedWorld;
+				Out_ViewToClipTransform = View.ViewToClip;
+				Out_ViewToClipNoAATransform = View.ViewToClipNoAA;
+			}
+		)");
+		OutHLSL += FString::Format(FormatSample, ArgsSample);
+		return true;
+	}
+	if (DefinitionFunctionName == GetCameraPositionsName)
+	{
+		static const TCHAR *FormatSample = TEXT(R"(
+			void {FunctionName}(out float3 Out_CameraPositionWorld)
+			{				
+				Out_CameraPositionWorld.xyz = View.WorldCameraOrigin.xyz;
+			}
+		)");
+		OutHLSL += FString::Format(FormatSample, ArgsSample);
+		return true;
+	}
+	return false;
+}
+
+DEFINE_NDI_DIRECT_FUNC_BINDER(UNiagaraDataInterfaceCamera, GetCameraFOV);
+DEFINE_NDI_DIRECT_FUNC_BINDER(UNiagaraDataInterfaceCamera, GetCameraPosition);
+DEFINE_NDI_DIRECT_FUNC_BINDER(UNiagaraDataInterfaceCamera, QueryOcclusionFactorGPU);
+DEFINE_NDI_DIRECT_FUNC_BINDER(UNiagaraDataInterfaceCamera, GetViewPropertiesGPU);
+DEFINE_NDI_DIRECT_FUNC_BINDER(UNiagaraDataInterfaceCamera, GetClipSpaceTransformsGPU);
+DEFINE_NDI_DIRECT_FUNC_BINDER(UNiagaraDataInterfaceCamera, GetViewSpaceTransformsGPU);
+void UNiagaraDataInterfaceCamera::GetVMExternalFunction(const FVMExternalFunctionBindingInfo& BindingInfo, void* InstanceData, FVMExternalFunction &OutFunc)
+{
+	if (BindingInfo.Name == GetCameraOcclusionName)
+	{
+		NDI_FUNC_BINDER(UNiagaraDataInterfaceCamera, QueryOcclusionFactorGPU)::Bind(this, OutFunc);
+	}
+	else if (BindingInfo.Name == GetFieldOfViewName)
+	{
+		NDI_FUNC_BINDER(UNiagaraDataInterfaceCamera, GetCameraFOV)::Bind(this, OutFunc);
+	}
+	else if (BindingInfo.Name == GetCameraPositionsName)
+	{
+		NDI_FUNC_BINDER(UNiagaraDataInterfaceCamera, GetCameraPosition)::Bind(this, OutFunc);
+	}
+	else if (BindingInfo.Name == GetViewSpaceTransformsName)
+	{
+		NDI_FUNC_BINDER(UNiagaraDataInterfaceCamera, GetViewSpaceTransformsGPU)::Bind(this, OutFunc);
+	}
+	else if (BindingInfo.Name == GetClipSpaceTransformsName)
+	{
+		NDI_FUNC_BINDER(UNiagaraDataInterfaceCamera, GetClipSpaceTransformsGPU)::Bind(this, OutFunc);
+	}
+	else if (BindingInfo.Name == GetViewPropertiesName)
+	{
+		NDI_FUNC_BINDER(UNiagaraDataInterfaceCamera, GetViewPropertiesGPU)::Bind(this, OutFunc);
+	}
+	else
+	{
+		UE_LOG(LogNiagara, Error, TEXT("Could not find data interface external function. Received Name: %s"), *BindingInfo.Name.ToString());
+	}
+}
+
+void UNiagaraDataInterfaceCamera::GetCameraFOV(FVectorVMContext& Context)
+{
+	VectorVM::FUserPtrHandler<CameraDataInterface_InstanceData> InstData(Context);
+	VectorVM::FExternalFuncRegisterHandler<float> OutFov(Context);
+
+	float Fov = 0;
+	if (InstData.Get()->CameraObject.IsValid(false))
+	{
+		Fov = InstData.Get()->CameraObject->GetFOVAngle();
+	}
+
+	for (int32 i = 0; i < Context.NumInstances; ++i)
+	{
+		*OutFov.GetDestAndAdvance() = Fov;
+	}
+}
+
+void UNiagaraDataInterfaceCamera::GetCameraPosition(FVectorVMContext& Context)
+{
+	VectorVM::FUserPtrHandler<CameraDataInterface_InstanceData> InstData(Context);
+
+	VectorVM::FExternalFuncRegisterHandler<float> CamPosX(Context);
+	VectorVM::FExternalFuncRegisterHandler<float> CamPosY(Context);
+	VectorVM::FExternalFuncRegisterHandler<float> CamPosZ(Context);
+
+	float XPos = 0;
+	float YPos = 0;
+	float ZPos = 0;
+	if (InstData.Get()->CameraObject.IsValid(false))
+	{
+		XPos = InstData.Get()->CameraObject->GetCameraLocation().X;
+		YPos = InstData.Get()->CameraObject->GetCameraLocation().Y;
+		ZPos = InstData.Get()->CameraObject->GetCameraLocation().Z;
+	}
+
+	for (int32 i = 0; i < Context.NumInstances; ++i)
+	{
+		*CamPosX.GetDestAndAdvance() = XPos;
+		*CamPosY.GetDestAndAdvance() = YPos;
+		*CamPosZ.GetDestAndAdvance() = ZPos;
+	}
+}
+
+// ------- Dummy implementations for CPU execution ------------
+
+void UNiagaraDataInterfaceCamera::QueryOcclusionFactorGPU(FVectorVMContext& Context)
+{
+	VectorVM::FExternalFuncInputHandler<float> PosParamX(Context);
+	VectorVM::FExternalFuncInputHandler<float> PosParamY(Context);
+	VectorVM::FExternalFuncInputHandler<float> PosParamZ(Context);
+	VectorVM::FExternalFuncInputHandler<float> SizeParam(Context);
+	VectorVM::FExternalFuncInputHandler<float> StepsParam(Context);
+
+	VectorVM::FUserPtrHandler<CameraDataInterface_InstanceData> InstData(Context);
+
+	VectorVM::FExternalFuncRegisterHandler<float> OutOcclusion(Context);
+	VectorVM::FExternalFuncRegisterHandler<float> OutFactor(Context);
+
+	for (int32 i = 0; i < Context.NumInstances; ++i)
+	{
+		PosParamX.GetAndAdvance();
+		PosParamY.GetAndAdvance();
+		PosParamZ.GetAndAdvance();
+		SizeParam.GetAndAdvance();
+		StepsParam.GetAndAdvance();
+
+		*OutOcclusion.GetDestAndAdvance() = 0;
+		*OutFactor.GetDestAndAdvance() = 0;
+	}
+}
+
+void UNiagaraDataInterfaceCamera::GetViewPropertiesGPU(FVectorVMContext& Context)
+{
+	VectorVM::FUserPtrHandler<CameraDataInterface_InstanceData> InstData(Context);
+	TArray<VectorVM::FExternalFuncRegisterHandler<float>> OutParams;
+
+	for (int i = 0; i < 20; i++)
+	{
+		OutParams.Emplace(Context);
+	}
+
+	for (int32 k = 0; k < Context.NumInstances; ++k)
+	{
+		for (int i = 0; i < 20; i++)
+		{
+			*OutParams[i].GetDestAndAdvance() = 0;
+		}
+	}
+}
+
+void UNiagaraDataInterfaceCamera::GetClipSpaceTransformsGPU(FVectorVMContext& Context)
+{
+	VectorVM::FUserPtrHandler<CameraDataInterface_InstanceData> InstData(Context);
+	TArray<VectorVM::FExternalFuncRegisterHandler<float>> OutParams;
+
+	for (int i = 0; i < 128; i++)
+	{
+		OutParams.Emplace(Context);
+	}
+
+	for (int32 k = 0; k < Context.NumInstances; ++k)
+	{
+		for (int i = 0; i < 128; i++)
+		{
+			*OutParams[i].GetDestAndAdvance() = 0;
+		}
+	}
+}
+
+void UNiagaraDataInterfaceCamera::GetViewSpaceTransformsGPU(FVectorVMContext& Context)
+{
+	VectorVM::FUserPtrHandler<CameraDataInterface_InstanceData> InstData(Context);
+	TArray<VectorVM::FExternalFuncRegisterHandler<float>> OutParams;
+
+	for (int i = 0; i < 96; i++)
+	{
+		OutParams.Emplace(Context);
+	}
+
+	for (int32 k = 0; k < Context.NumInstances; ++k)
+	{
+		for (int i = 0; i < 96; i++)
+		{
+			*OutParams[i].GetDestAndAdvance() = 0;
+		}
+	}
+}
+
+// ------------------------------------------------------------
+
+struct FNiagaraDataInterfaceParametersCS_CameraQuery : public FNiagaraDataInterfaceParametersCS
+{
+	virtual void Bind(const FNiagaraDataInterfaceParamRef& ParamRef, const class FShaderParameterMap& ParameterMap) override
+	{
+		PassUniformBuffer.Bind(ParameterMap, FSceneTexturesUniformParameters::StaticStructMetadata.GetShaderVariableName());
+	}
+
+	virtual void Serialize(FArchive& Ar) override
+	{
+		Ar << PassUniformBuffer;
+	}
+
+	virtual void Set(FRHICommandList& RHICmdList, const FNiagaraDataInterfaceSetArgs& Context) const override
+	{
+		check(IsInRenderingThread());
+		FRHIComputeShader* ComputeShaderRHI = Context.Shader->GetComputeShader();
+
+		TUniformBufferRef<FSceneTexturesUniformParameters> SceneTextureUniformParams = GNiagaraViewDataManager.GetSceneTextureUniformParameters();
+		SetUniformBufferParameter(RHICmdList, ComputeShaderRHI, PassUniformBuffer, SceneTextureUniformParams);
+	}
+
+private:
+
+	/** The SceneDepthTexture parameter for depth buffer query. */
+	FShaderUniformBufferParameter PassUniformBuffer;
+};
+
+FNiagaraDataInterfaceParametersCS* UNiagaraDataInterfaceCamera::ConstructComputeParameters() const
+{
+	return new FNiagaraDataInterfaceParametersCS_CameraQuery();
+}
\ No newline at end of file
