Build fix for CL 10249850

diff --git a/Engine/Plugins/FX/Niagara/Source/Niagara/Private/NiagaraDataInterfaceCamera.cpp b/Engine/Plugins/FX/Niagara/Source/Niagara/Private/NiagaraDataInterfaceCamera.cpp
--- a/Engine/Plugins/FX/Niagara/Source/Niagara/Private/NiagaraDataInterfaceCamera.cpp
+++ b/Engine/Plugins/FX/Niagara/Source/Niagara/Private/NiagaraDataInterfaceCamera.cpp
@@ -4,8 +4,10 @@
 #include "NiagaraTypes.h"
 #include "NiagaraCustomVersion.h"
 #include "NiagaraWorldManager.h"
+#include "ShaderParameterUtils.h"
 #include "Internationalization/Internationalization.h"
 #include "NiagaraSystemInstance.h"
+#include "GameFramework/PlayerController.h"
 
 const FName UNiagaraDataInterfaceCamera::GetCameraOcclusionName(TEXT("QueryOcclusionFactorGPU"));
 const FName UNiagaraDataInterfaceCamera::GetViewPropertiesName(TEXT("GetViewPropertiesGPU"));
