- Add missing stat for tracking RenderThread compute execution cost in Niagara

#rnx

diff --git a/Engine/Plugins/FX/Niagara/Source/Niagara/Private/NiagaraEmitterInstanceBatcher.cpp b/Engine/Plugins/FX/Niagara/Source/Niagara/Private/NiagaraEmitterInstanceBatcher.cpp
--- a/Engine/Plugins/FX/Niagara/Source/Niagara/Private/NiagaraEmitterInstanceBatcher.cpp
+++ b/Engine/Plugins/FX/Niagara/Source/Niagara/Private/NiagaraEmitterInstanceBatcher.cpp
@@ -516,6 +516,8 @@ void NiagaraEmitterInstanceBatcher::PostRenderOpaque(FRHICommandListImmediate& R
 
 void NiagaraEmitterInstanceBatcher::ExecuteAll(FRHICommandList &RHICmdList, FRHIUniformBuffer* ViewUniformBuffer, bool bSetReadback)
 {
+	SCOPE_CYCLE_COUNTER(STAT_NiagaraGPUSimTick_RT);
+
 	// This is always called by the renderer so early out if we have no work.
 	if (Ticks_RT.Num() == 0)
 	{
