default shader stage index on emitters for modules that don't expose shader stage parameters

diff --git a/Engine/Plugins/FX/Niagara/Shaders/Private/NiagaraEmitterInstanceShader.usf b/Engine/Plugins/FX/Niagara/Shaders/Private/NiagaraEmitterInstanceShader.usf
--- a/Engine/Plugins/FX/Niagara/Shaders/Private/NiagaraEmitterInstanceShader.usf
+++ b/Engine/Plugins/FX/Niagara/Shaders/Private/NiagaraEmitterInstanceShader.usf
@@ -333,6 +333,7 @@ int rand(int x, int Seed1, int Seed2, int Seed3)
 	/* Parameters for spawning and simulations
 	 */
 	uint UpdateStartInstance;
+	uint DefaultShaderStageIndex;
 	uint ShaderStageIndex;
 	int IterationInterfaceCount;
 	uint SpawnedInstances;
diff --git a/Engine/Plugins/FX/Niagara/Source/Niagara/Classes/NiagaraEmitter.h b/Engine/Plugins/FX/Niagara/Source/Niagara/Classes/NiagaraEmitter.h
--- a/Engine/Plugins/FX/Niagara/Source/Niagara/Classes/NiagaraEmitter.h
+++ b/Engine/Plugins/FX/Niagara/Source/Niagara/Classes/NiagaraEmitter.h
@@ -339,6 +339,10 @@ public:
 	UPROPERTY(EditAnywhere, AdvancedDisplay, Category = "Emitter", meta = (EditCondition = "bLimitDeltaTime"))
 	float MaxDeltaTimePerTick;
 
+	/** Get the max number of iteration that the CS is going to be launched. */
+	UPROPERTY(EditAnywhere, AdvancedDisplay, Category = "Emitter")
+	uint32 DefaultShaderStageIndex;
+
 	/** Get the max number of iteration that the CS is going to be launched. */
 	UPROPERTY(EditAnywhere, AdvancedDisplay, Category = "Emitter")
 	uint32 MaxUpdateIterations;
diff --git a/Engine/Plugins/FX/Niagara/Source/Niagara/Classes/NiagaraEmitterInstanceBatcher.h b/Engine/Plugins/FX/Niagara/Source/Niagara/Classes/NiagaraEmitterInstanceBatcher.h
--- a/Engine/Plugins/FX/Niagara/Source/Niagara/Classes/NiagaraEmitterInstanceBatcher.h
+++ b/Engine/Plugins/FX/Niagara/Source/Niagara/Classes/NiagaraEmitterInstanceBatcher.h
@@ -126,6 +126,7 @@ public:
 				FRHIUniformBuffer* ViewUniformBuffer, 
 				const FNiagaraGpuSpawnInfo& SpawnInfo,
 				bool bCopyBeforeStart = false,
+				uint32 DefaultShaderStageIndex = 0,
 				uint32 ShaderStageIndex = 0,
 				FNiagaraDataInterfaceProxy *IterationInterface = nullptr,
 				bool HasRunParticleStage = false
diff --git a/Engine/Plugins/FX/Niagara/Source/Niagara/Classes/NiagaraScriptExecutionContext.h b/Engine/Plugins/FX/Niagara/Source/Niagara/Classes/NiagaraScriptExecutionContext.h
--- a/Engine/Plugins/FX/Niagara/Source/Niagara/Classes/NiagaraScriptExecutionContext.h
+++ b/Engine/Plugins/FX/Niagara/Source/Niagara/Classes/NiagaraScriptExecutionContext.h
@@ -174,7 +174,7 @@ struct FNiagaraComputeExecutionContext
 
 	void Reset(NiagaraEmitterInstanceBatcher* Batcher);
 
-	void InitParams(UNiagaraScript* InGPUComputeScript, ENiagaraSimTarget InSimTarget, const FString& InDebugSimName, const int32 InMaxUpdateIterations, const TSet<uint32> InSpawnStages);
+	void InitParams(UNiagaraScript* InGPUComputeScript, ENiagaraSimTarget InSimTarget, const FString& InDebugSimName, const uint32 InDefaultShaderStageIndex, int32 InMaxUpdateIterations, const TSet<uint32> InSpawnStages);
 	void DirtyDataInterfaces();
 	bool Tick(FNiagaraSystemInstance* ParentSystemInstance);
 
@@ -225,6 +225,7 @@ public:
 	// Game thread spawn info will be sent to the render thread inside FNiagaraComputeInstanceData
 	FNiagaraGpuSpawnInfo GpuSpawnInfo_GT;
 
+	uint32 DefaultShaderStageIndex;
 	uint32 MaxUpdateIterations;
 	TSet<uint32> SpawnStages;
 
diff --git a/Engine/Plugins/FX/Niagara/Source/Niagara/Private/NiagaraEmitter.cpp b/Engine/Plugins/FX/Niagara/Source/Niagara/Private/NiagaraEmitter.cpp
--- a/Engine/Plugins/FX/Niagara/Source/Niagara/Private/NiagaraEmitter.cpp
+++ b/Engine/Plugins/FX/Niagara/Source/Niagara/Private/NiagaraEmitter.cpp
@@ -135,6 +135,7 @@ UNiagaraEmitter::UNiagaraEmitter(const FObjectInitializer& Initializer)
 , bUseMaxDetailLevel(false)
 , bRequiresPersistentIDs(false)
 , MaxDeltaTimePerTick(0.125)
+, DefaultShaderStageIndex(0)
 , MaxUpdateIterations(1)
 , bLimitDeltaTime(true)
 #if WITH_EDITORONLY_DATA
diff --git a/Engine/Plugins/FX/Niagara/Source/Niagara/Private/NiagaraEmitterInstance.cpp b/Engine/Plugins/FX/Niagara/Source/Niagara/Private/NiagaraEmitterInstance.cpp
--- a/Engine/Plugins/FX/Niagara/Source/Niagara/Private/NiagaraEmitterInstance.cpp
+++ b/Engine/Plugins/FX/Niagara/Source/Niagara/Private/NiagaraEmitterInstance.cpp
@@ -356,7 +356,7 @@ void FNiagaraEmitterInstance::Init(int32 InEmitterIdx, FNiagaraSystemInstanceID
 		if (CachedEmitter->SimTarget == ENiagaraSimTarget::GPUComputeSim)
 		{
 			GPUExecContext = new FNiagaraComputeExecutionContext();
-			GPUExecContext->InitParams(CachedEmitter->GetGPUComputeScript(), CachedEmitter->SimTarget, CachedEmitter->GetUniqueEmitterName(), CachedEmitter->MaxUpdateIterations, CachedEmitter->SpawnStages);
+			GPUExecContext->InitParams(CachedEmitter->GetGPUComputeScript(), CachedEmitter->SimTarget, CachedEmitter->GetUniqueEmitterName(), CachedEmitter->DefaultShaderStageIndex, CachedEmitter->MaxUpdateIterations, CachedEmitter->SpawnStages);
 			GPUExecContext->MainDataSet = &Data;
 			GPUExecContext->GPUScript_RT = CachedEmitter->GetGPUComputeScript()->GetRenderThreadScript();
 
diff --git a/Engine/Plugins/FX/Niagara/Source/Niagara/Private/NiagaraEmitterInstanceBatcher.cpp b/Engine/Plugins/FX/Niagara/Source/Niagara/Private/NiagaraEmitterInstanceBatcher.cpp
--- a/Engine/Plugins/FX/Niagara/Source/Niagara/Private/NiagaraEmitterInstanceBatcher.cpp
+++ b/Engine/Plugins/FX/Niagara/Source/Niagara/Private/NiagaraEmitterInstanceBatcher.cpp
@@ -282,6 +282,8 @@ void NiagaraEmitterInstanceBatcher::DispatchMultipleStages(const FNiagaraGPUSyst
 		FNiagaraDataBuffer* DestinationData = Instance->DestinationData;
 
 		const uint32 NumStages = Instance->Context->MaxUpdateIterations;
+		const uint32 DefaultShaderStageIndex = Instance->Context->DefaultShaderStageIndex;
+
 		for (uint32 ShaderStageIndex = 0; ShaderStageIndex < NumStages; ++ShaderStageIndex)
 		{
 			// Determine if the iteration is outputting to a custom data size
@@ -304,13 +306,13 @@ void NiagaraEmitterInstanceBatcher::DispatchMultipleStages(const FNiagaraGPUSyst
 					Instance->CurrentData = Instance->DestinationData;
 					Instance->DestinationData = StoreData;
 				}
-				Run(Tick, Instance, 0, Instance->DestinationData->GetNumInstances(), ComputeShader, RHICmdList, ViewUniformBuffer, Instance->SpawnInfo, false, ShaderStageIndex,  nullptr, HasRunParticleStage);
+				Run(Tick, Instance, 0, Instance->DestinationData->GetNumInstances(), ComputeShader, RHICmdList, ViewUniformBuffer, Instance->SpawnInfo, false, DefaultShaderStageIndex, ShaderStageIndex,  nullptr, HasRunParticleStage);
 				HasRunParticleStage = true;
 			}
 			else
 			{
 				// run with correct number of instances.  This will make curr data junk or empty
-				Run(Tick, Instance, 0, IterationInterface->ElementCount, ComputeShader, RHICmdList, ViewUniformBuffer, Instance->SpawnInfo, false, ShaderStageIndex, IterationInterface);
+				Run(Tick, Instance, 0, IterationInterface->ElementCount, ComputeShader, RHICmdList, ViewUniformBuffer, Instance->SpawnInfo, false, DefaultShaderStageIndex, ShaderStageIndex, IterationInterface);
 			}
 			PostStageInterface(Tick, Instance, RHICmdList, ComputeShader, ShaderStageIndex);
 		}
@@ -1087,7 +1089,7 @@ void NiagaraEmitterInstanceBatcher::UnsetDataInterfaceParameters(const TArray<FN
 /* Kick off a simulation/spawn run
  */
 void NiagaraEmitterInstanceBatcher::Run(const FNiagaraGPUSystemTick& Tick, const FNiagaraComputeInstanceData* Instance, uint32 UpdateStartInstance, const uint32 TotalNumInstances, FNiagaraShader* Shader,
-	FRHICommandList &RHICmdList, FRHIUniformBuffer* ViewUniformBuffer, const FNiagaraGpuSpawnInfo& SpawnInfo, bool bCopyBeforeStart, uint32 ShaderStageIndex, FNiagaraDataInterfaceProxy *IterationInterface, bool HasRunParticleStage) const
+	FRHICommandList &RHICmdList, FRHIUniformBuffer* ViewUniformBuffer, const FNiagaraGpuSpawnInfo& SpawnInfo, bool bCopyBeforeStart, uint32 DefaultShaderStageIndex, uint32 ShaderStageIndex, FNiagaraDataInterfaceProxy *IterationInterface, bool HasRunParticleStage) const
 {
 	FNiagaraComputeExecutionContext* Context = Instance->Context;
 	if (TotalNumInstances == 0)
@@ -1168,6 +1170,12 @@ void NiagaraEmitterInstanceBatcher::Run(const FNiagaraGPUSystemTick& Tick, const
 	}
 
 	RHICmdList.SetShaderParameter(Shader->GetComputeShader(), Shader->NumSpawnedInstancesParam.GetBufferIndex(), Shader->NumSpawnedInstancesParam.GetBaseIndex(), Shader->NumSpawnedInstancesParam.GetNumBytes(), &InstancesToSpawnThisFrame);				// number of instances in the spawn run
+	
+	if (Shader->DefaultShaderStageIndexParam.IsBound())
+	{
+		RHICmdList.SetShaderParameter(Shader->GetComputeShader(), Shader->DefaultShaderStageIndexParam.GetBufferIndex(), Shader->DefaultShaderStageIndexParam.GetBaseIndex(), Shader->DefaultShaderStageIndexParam.GetNumBytes(), &DefaultShaderStageIndex);					// 0, except if several stages are defined
+	}
+	
 	if (Shader->ShaderStageIndexParam.IsBound())
 	{
 		RHICmdList.SetShaderParameter(Shader->GetComputeShader(), Shader->ShaderStageIndexParam.GetBufferIndex(), Shader->ShaderStageIndexParam.GetBaseIndex(), Shader->ShaderStageIndexParam.GetNumBytes(), &ShaderStageIndex);					// 0, except if several stages are defined
diff --git a/Engine/Plugins/FX/Niagara/Source/Niagara/Private/NiagaraScriptExecutionContext.cpp b/Engine/Plugins/FX/Niagara/Source/Niagara/Private/NiagaraScriptExecutionContext.cpp
--- a/Engine/Plugins/FX/Niagara/Source/Niagara/Private/NiagaraScriptExecutionContext.cpp
+++ b/Engine/Plugins/FX/Niagara/Source/Niagara/Private/NiagaraScriptExecutionContext.cpp
@@ -426,13 +426,14 @@ void FNiagaraComputeExecutionContext::Reset(NiagaraEmitterInstanceBatcher* Batch
 	);
 }
 
-void FNiagaraComputeExecutionContext::InitParams(UNiagaraScript* InGPUComputeScript, ENiagaraSimTarget InSimTarget, const FString& InDebugSimName, const int32 InMaxUpdateIterations, const TSet<uint32> InSpawnStages)
+void FNiagaraComputeExecutionContext::InitParams(UNiagaraScript* InGPUComputeScript, ENiagaraSimTarget InSimTarget, const FString& InDebugSimName, const uint32 InDefaultShaderStageIndex, const int32 InMaxUpdateIterations, const TSet<uint32> InSpawnStages)
 {
 #if !UE_BUILD_SHIPPING
 	DebugSimName = InDebugSimName;
 #endif
 	GPUScript = InGPUComputeScript;
 	CombinedParamStore.InitFromOwningContext(InGPUComputeScript, InSimTarget, true);
+	DefaultShaderStageIndex = InDefaultShaderStageIndex;
 	MaxUpdateIterations = InMaxUpdateIterations;
 	SpawnStages.Empty();
 
diff --git a/Engine/Plugins/FX/Niagara/Source/NiagaraEditor/Private/NiagaraHlslTranslator.cpp b/Engine/Plugins/FX/Niagara/Source/NiagaraEditor/Private/NiagaraHlslTranslator.cpp
--- a/Engine/Plugins/FX/Niagara/Source/NiagaraEditor/Private/NiagaraHlslTranslator.cpp
+++ b/Engine/Plugins/FX/Niagara/Source/NiagaraEditor/Private/NiagaraHlslTranslator.cpp
@@ -5114,7 +5114,7 @@ void FHlslNiagaraTranslator::RegisterFunctionCall(ENiagaraScriptUsage ScriptUsag
 						}
 						else
 						{
-							FunctionDefStr += TEXT("if (ShaderStageIndex == 0)\n{\n");
+							FunctionDefStr += TEXT("if (ShaderStageIndex == DefaultShaderStageIndex)\n{\n");
 						}
 					}
 				}
diff --git a/Engine/Plugins/FX/Niagara/Source/NiagaraShader/Private/NiagaraShader.cpp b/Engine/Plugins/FX/Niagara/Source/NiagaraShader/Private/NiagaraShader.cpp
--- a/Engine/Plugins/FX/Niagara/Source/NiagaraShader/Private/NiagaraShader.cpp
+++ b/Engine/Plugins/FX/Niagara/Source/NiagaraShader/Private/NiagaraShader.cpp
@@ -1326,6 +1326,7 @@ void FNiagaraShader::BindParams(const FShaderParameterMap &ParameterMap)
 
 	NumSpawnedInstancesParam.Bind(ParameterMap, TEXT("SpawnedInstances"));
 	UpdateStartInstanceParam.Bind(ParameterMap, TEXT("UpdateStartInstance"));
+	DefaultShaderStageIndexParam.Bind(ParameterMap, TEXT("DefaultShaderStageIndex"));
 	ShaderStageIndexParam.Bind(ParameterMap, TEXT("ShaderStageIndex"));
 	IterationInterfaceCount.Bind(ParameterMap, TEXT("IterationInterfaceCount"));
 
@@ -1393,6 +1394,7 @@ bool FNiagaraShader::Serialize(FArchive& Ar)
 
 	Ar << NumSpawnedInstancesParam;
 	Ar << UpdateStartInstanceParam;
+	Ar << DefaultShaderStageIndexParam;
 	Ar << ShaderStageIndexParam;
 	Ar << IterationInterfaceCount;
 	Ar << ComponentBufferSizeReadParam;
diff --git a/Engine/Plugins/FX/Niagara/Source/NiagaraShader/Public/NiagaraShader.h b/Engine/Plugins/FX/Niagara/Source/NiagaraShader/Public/NiagaraShader.h
--- a/Engine/Plugins/FX/Niagara/Source/NiagaraShader/Public/NiagaraShader.h
+++ b/Engine/Plugins/FX/Niagara/Source/NiagaraShader/Public/NiagaraShader.h
@@ -106,6 +106,7 @@ public:
 	FShaderParameter CopyInstancesBeforeStartParam;
 	FShaderParameter NumSpawnedInstancesParam;
 	FShaderParameter UpdateStartInstanceParam;
+	FShaderParameter DefaultShaderStageIndexParam;
 	FShaderParameter ShaderStageIndexParam;
 	FShaderParameter IterationInterfaceCount;
 	FShaderParameter ComponentBufferSizeReadParam;
