particle spawn while using default shader stage fix

diff --git a/Engine/Plugins/FX/Niagara/Source/NiagaraEditor/Private/NiagaraHlslTranslator.cpp b/Engine/Plugins/FX/Niagara/Source/NiagaraEditor/Private/NiagaraHlslTranslator.cpp
--- a/Engine/Plugins/FX/Niagara/Source/NiagaraEditor/Private/NiagaraHlslTranslator.cpp
+++ b/Engine/Plugins/FX/Niagara/Source/NiagaraEditor/Private/NiagaraHlslTranslator.cpp
@@ -5114,7 +5114,7 @@ void FHlslNiagaraTranslator::RegisterFunctionCall(ENiagaraScriptUsage ScriptUsag
 						}
 						else
 						{
-							FunctionDefStr += TEXT("if (ShaderStageIndex == DefaultShaderStageIndex)\n{\n");
+							FunctionDefStr += TEXT("if ((GCurrentPhase == 0 && ShaderStageIndex == 0) || (GCurrentPhase == 1 && ShaderStageIndex == DefaultShaderStageIndex))\n{\n");
 						}
 					}
 				}
