Fixes code generation of Niagara shaders which were holding onto a reference to a TMap value, but the underlying array of elements was getting reallocated.

diff --git a/Engine/Plugins/FX/Niagara/Source/NiagaraEditor/Private/NiagaraHlslTranslator.cpp b/Engine/Plugins/FX/Niagara/Source/NiagaraEditor/Private/NiagaraHlslTranslator.cpp
--- a/Engine/Plugins/FX/Niagara/Source/NiagaraEditor/Private/NiagaraHlslTranslator.cpp
+++ b/Engine/Plugins/FX/Niagara/Source/NiagaraEditor/Private/NiagaraHlslTranslator.cpp
@@ -1143,21 +1143,18 @@ const FNiagaraTranslateResults &FHlslNiagaraTranslator::Translate(const FNiagara
 		WriteDataSetStructDeclarations(DataSetWriteInfo[0], false, HlslOutput);
 
 		//Map of all variables accessed by all datasets.
-		TMap<FNiagaraDataSetID, TArray<FNiagaraVariable>> DataSetReads;
-		TMap<FNiagaraDataSetID, TArray<FNiagaraVariable>> DataSetWrites;
+		TArray<TArray<FNiagaraVariable>> DataSetVariables;
 
-		TArray<TArray<FNiagaraVariable>* > DataSetReadVars;
-		TArray<TArray<FNiagaraVariable>* > DataSetWriteVars;
-		TArray<FNiagaraDataSetID> DataSetReadIds;
-		TArray<FNiagaraDataSetID> DataSetWriteIds;
+		TMap<FNiagaraDataSetID, int32> DataSetReads;
+		TMap<FNiagaraDataSetID, int32> DataSetWrites;
 
-		TArray<FNiagaraVariable>& InstanceReadVars = DataSetReads.Add(GetInstanceDataSetID());
-		TArray<FNiagaraVariable>& InstanceWriteVars = DataSetWrites.Add(GetInstanceDataSetID());
+		const FNiagaraDataSetID InstanceDataSetID = GetInstanceDataSetID();
 
-		DataSetReadVars.Add(&InstanceReadVars);
-		DataSetReadIds.Add(GetInstanceDataSetID());
-		DataSetWriteVars.Add(&InstanceWriteVars);
-		DataSetWriteIds.Add(GetInstanceDataSetID());
+		const int32 InstanceReadVarsIndex = DataSetVariables.AddDefaulted();
+		const int32 InstanceWriteVarsIndex = DataSetVariables.AddDefaulted();
+
+		DataSetReads.Add(InstanceDataSetID, InstanceReadVarsIndex);
+		DataSetWrites.Add(InstanceDataSetID, InstanceWriteVarsIndex);
 
 		if (IsBulkSystemScript())
 		{
@@ -1165,24 +1162,21 @@ const FNiagaraTranslateResults &FHlslNiagaraTranslator::Translate(const FNiagara
 			// that are essentially set once per system. The constants are rapid iteration variables
 			// that exist per emitter and change infrequently. Since they are so different, putting
 			// them in two distinct read data sets seems warranted.
-			TArray<FNiagaraVariable>& SystemEngineReadVars = DataSetReads.Add(GetSystemEngineDataSetID());
-			TArray<FNiagaraVariable>& SystemConstantReadVars = DataSetReads.Add(GetSystemConstantDataSetID());
+			const FNiagaraDataSetID SystemEngineDataSetID = GetSystemEngineDataSetID();
 
-			DataSetReadVars.Add(&SystemEngineReadVars);
-			DataSetReadIds.Add(GetSystemEngineDataSetID());
+			const int32 SystemEngineReadVarsIndex = DataSetVariables.Num();
+			DataSetReads.Add(SystemEngineDataSetID, SystemEngineReadVarsIndex);
+			TArray<FNiagaraVariable>& SystemEngineReadVars = DataSetVariables.AddDefaulted_GetRef();
 
 			HandleNamespacedExternalVariablesToDataSetRead(SystemEngineReadVars, TEXT("Engine"));
-			HandleNamespacedExternalVariablesToDataSetRead(SystemEngineReadVars/*SystemUserReadVars*/, TEXT("User"));
+			HandleNamespacedExternalVariablesToDataSetRead(SystemEngineReadVars, TEXT("User"));
 
 			// We sort the variables so that they end up in the same ordering between Spawn & Update...
-			SystemEngineReadVars.Sort([&](const FNiagaraVariable& A, const FNiagaraVariable& B)
-			{
-				return A.GetName().LexicalLess(B.GetName());
-			});
+			Algo::SortBy(SystemEngineReadVars, &FNiagaraVariable::GetName, FNameLexicalLess());
 
 			{
 				FNiagaraParameters ExternalParams;
-				ExternalParams.Parameters = SystemEngineReadVars;
+				ExternalParams.Parameters = DataSetVariables[SystemEngineReadVarsIndex];
 				CompilationOutput.ScriptData.DataSetToParameters.Add(GetSystemEngineDataSetID().Name, ExternalParams);
 			}
 		}
@@ -1219,20 +1213,12 @@ const FNiagaraTranslateResults &FHlslNiagaraTranslator::Translate(const FNiagara
 			}
 		}
 
-		InstanceReadVars = BasicAttributes;
-		InstanceWriteVars = BasicAttributes;
+		// We sort the variables so that they end up in the same ordering between Spawn & Update...
+		Algo::SortBy(BasicAttributes, &FNiagaraVariable::GetName, FNameLexicalLess());
 
+		DataSetVariables[InstanceReadVarsIndex] = BasicAttributes;
+		DataSetVariables[InstanceWriteVarsIndex] = BasicAttributes;
 
-		// We sort the variables so that they end up in the same ordering between Spawn & Update...
-		InstanceReadVars.Sort([&](const FNiagaraVariable& A, const FNiagaraVariable& B)
-		{
-			return A.GetName().LexicalLess(B.GetName());
-		});
-		// We sort the variables so that they end up in the same ordering between Spawn & Update...
-		InstanceWriteVars.Sort([&](const FNiagaraVariable& A, const FNiagaraVariable& B)
-		{
-			return A.GetName().LexicalLess(B.GetName());
-		});
 		//Define the simulation context. Which is a helper struct containing all the input, result and intermediate data needed for a single simulation.
 		//Allows us to reuse the same simulate function but provide different wrappers for final IO between GPU and CPU sims.
 		{
@@ -1331,11 +1317,10 @@ const FNiagaraTranslateResults &FHlslNiagaraTranslator::Translate(const FNiagara
 
 		//And finally, define the actual main function that handles the reading and writing of data and calls the shared per instance simulate function.
 		//TODO: Different wrappers for GPU and CPU sims. 
-		DefineMain(HlslOutput, DataSetReadVars, DataSetReadIds, DataSetWriteVars, DataSetWriteIds);
-
+		DefineMain(HlslOutput, DataSetVariables, DataSetReads, DataSetWrites);
 
 		//Get full list of instance data accessed by the script as the VM binding assumes same for input and output.
-		for (FNiagaraVariable& Var : InstanceReadVars)
+		for (FNiagaraVariable& Var : DataSetVariables[InstanceReadVarsIndex])
 		{
 			if (FNiagaraParameterMapHistory::IsAttribute(Var))
 			{
@@ -1798,8 +1783,9 @@ void FHlslNiagaraTranslator::DefineExternalFunctionsHLSL(FString &InHlslOutput)
 }
 
 void FHlslNiagaraTranslator::DefineMain(FString &OutHlslOutput,
-	TArray<TArray<FNiagaraVariable>*> &InstanceReadVars, TArray<FNiagaraDataSetID>& ReadIds,
-	TArray<TArray<FNiagaraVariable>*> &InstanceWriteVars, TArray<FNiagaraDataSetID>& WriteIds)
+	const TArray<TArray<FNiagaraVariable>>& DataSetVariables,
+	const TMap<FNiagaraDataSetID, int32>& DataSetReads,
+	const TMap<FNiagaraDataSetID, int32>& DataSetWrites)
 {
 	if (CompilationTarget == ENiagaraSimTarget::GPUComputeSim)
 	{
@@ -1838,11 +1824,18 @@ void FHlslNiagaraTranslator::DefineMain(FString &OutHlslOutput,
 		OutHlslOutput += TEXT("\t}\n");
 	}
 
+	TArray<FNiagaraDataSetID> ReadDataSetIDs;
+	TArray<FNiagaraDataSetID> WriteDataSetIDs;
+
+	DataSetReads.GetKeys(ReadDataSetIDs);
+	DataSetWrites.GetKeys(WriteDataSetIDs);
+
 	//The VM register binding assumes the same inputs as outputs which is obviously not always the case.
-	for (int32 VarArrayIdx = 0; VarArrayIdx < InstanceReadVars.Num(); VarArrayIdx++)
+	for (int32 VarArrayIdx = 0; VarArrayIdx < DataSetReads.Num(); VarArrayIdx++)
 	{
-		TArray<FNiagaraVariable>* ArrayRef = InstanceReadVars[VarArrayIdx];
-		DefineDataSetVariableReads(HlslOutput, ReadIds[VarArrayIdx], VarArrayIdx, *ArrayRef);
+		const FNiagaraDataSetID DataSetID = ReadDataSetIDs[VarArrayIdx];
+		const TArray<FNiagaraVariable>& ArrayRef = DataSetVariables[DataSetReads[DataSetID]];
+		DefineDataSetVariableReads(HlslOutput, DataSetID, VarArrayIdx, ArrayRef);
 	}
 
 	bool bNeedsPersistentIDs = CompileOptions.AdditionalDefines.Contains(TEXT("RequiresPersistentIDs"));
@@ -1892,11 +1885,11 @@ void FHlslNiagaraTranslator::DefineMain(FString &OutHlslOutput,
 		{
 			// Determine 
 			TArray<FName> DataSetNames;
-			for (const FNiagaraDataSetID& ReadId : ReadIds)
+			for (const FNiagaraDataSetID& ReadId : ReadDataSetIDs)
 			{
 				DataSetNames.AddUnique(ReadId.Name);
 			}
-			for (const FNiagaraDataSetID& WriteId : WriteIds)
+			for (const FNiagaraDataSetID& WriteId : WriteDataSetIDs)
 			{
 				DataSetNames.AddUnique(WriteId.Name);
 			}
@@ -2002,10 +1995,11 @@ void FHlslNiagaraTranslator::DefineMain(FString &OutHlslOutput,
 
 	//The VM register binding assumes the same inputs as outputs which is obviously not always the case.
 	//We should separate inputs and outputs in the script.
-	for (int32 VarArrayIdx = 0; VarArrayIdx < InstanceWriteVars.Num(); VarArrayIdx++)
+	for (int32 VarArrayIdx = 0; VarArrayIdx < DataSetWrites.Num(); VarArrayIdx++)
 	{
-		TArray<FNiagaraVariable>* ArrayRef = InstanceWriteVars[VarArrayIdx];
-		DefineDataSetVariableWrites(HlslOutput, WriteIds[VarArrayIdx], VarArrayIdx, *ArrayRef);
+		const FNiagaraDataSetID DataSetID = WriteDataSetIDs[VarArrayIdx];
+		const TArray<FNiagaraVariable> ArrayRef = DataSetVariables[DataSetWrites[DataSetID]];
+		DefineDataSetVariableWrites(HlslOutput, DataSetID, VarArrayIdx, ArrayRef);
 	}
 
 	ExitStatsScope(OutHlslOutput);
@@ -2022,10 +2016,11 @@ void FHlslNiagaraTranslator::DefineMain(FString &OutHlslOutput,
 		OutHlslOutput += TEXT("void CopyInstance(in int InstanceIdx)\n{\n\tFSimulationContext Context = (FSimulationContext)0;\n");
 		if (UNiagaraScript::IsParticleEventScript(CompileOptions.TargetUsage))
 		{
-			for (int32 VarArrayIdx = 0; VarArrayIdx < InstanceReadVars.Num(); VarArrayIdx++)
+			for (int32 VarArrayIdx = 0; VarArrayIdx < DataSetReads.Num(); VarArrayIdx++)
 			{
-				TArray<FNiagaraVariable>* ArrayRef = InstanceReadVars[VarArrayIdx];
-				DefineDataSetVariableReads(HlslOutput, ReadIds[VarArrayIdx], VarArrayIdx, *ArrayRef);
+				const FNiagaraDataSetID DataSetID = ReadDataSetIDs[VarArrayIdx];
+				const TArray<FNiagaraVariable>& ArrayRef = DataSetVariables[DataSetReads[DataSetID]];
+				DefineDataSetVariableReads(HlslOutput, DataSetID, VarArrayIdx, ArrayRef);
 			}
 
 			if (bGpuUsesAlive)
@@ -2033,10 +2028,11 @@ void FHlslNiagaraTranslator::DefineMain(FString &OutHlslOutput,
 				OutHlslOutput += TEXT("\tContext.Map.DataInstance.Alive = true;\n");
 			}
 
-			for (int32 VarArrayIdx = 0; VarArrayIdx < InstanceWriteVars.Num(); VarArrayIdx++)
+			for (int32 VarArrayIdx = 0; VarArrayIdx < DataSetWrites.Num(); VarArrayIdx++)
 			{
-				TArray<FNiagaraVariable>* ArrayRef = InstanceWriteVars[VarArrayIdx];
-				DefineDataSetVariableWrites(HlslOutput, WriteIds[VarArrayIdx], VarArrayIdx, *ArrayRef);
+				const FNiagaraDataSetID DataSetID = WriteDataSetIDs[VarArrayIdx];
+				const TArray<FNiagaraVariable>& ArrayRef = DataSetVariables[DataSetWrites[DataSetID]];
+				DefineDataSetVariableWrites(HlslOutput, DataSetID, VarArrayIdx, ArrayRef);
 			}
 		}
 		OutHlslOutput += TEXT("}\n");
@@ -2044,7 +2040,7 @@ void FHlslNiagaraTranslator::DefineMain(FString &OutHlslOutput,
 
 }
 
-void FHlslNiagaraTranslator::DefineDataSetVariableWrites(FString &OutHlslOutput, FNiagaraDataSetID& Id, int32 DataSetIndex, TArray<FNiagaraVariable>& WriteVars)
+void FHlslNiagaraTranslator::DefineDataSetVariableWrites(FString &OutHlslOutput, const FNiagaraDataSetID& Id, int32 DataSetIndex, const TArray<FNiagaraVariable>& WriteVars)
 {
 	//TODO Grab indices for data set writes (inc output) and do the write. Need to rewrite this for events interleaved..
 	OutHlslOutput += "\t{\n";
@@ -2123,7 +2119,7 @@ void FHlslNiagaraTranslator::DefineDataSetVariableWrites(FString &OutHlslOutput,
 	int32 WriteOffsetFloat = 0;
 	int32 &FloatCounter = WriteOffsetFloat;
 	int32 &IntCounter = CompilationTarget == ENiagaraSimTarget::GPUComputeSim ? WriteOffsetInt : WriteOffsetFloat;
-	for (FNiagaraVariable &Var : WriteVars)
+	for (const FNiagaraVariable &Var : WriteVars)
 	{
 		// If coming from a parameter map, use the one on the context, otherwise use the output.
 		FString Fmt;
@@ -2140,7 +2136,7 @@ void FHlslNiagaraTranslator::DefineDataSetVariableWrites(FString &OutHlslOutput,
 	OutHlslOutput += "\t}\n";
 }
 
-void FHlslNiagaraTranslator::DefineDataSetVariableReads(FString &OutHlslOutput, FNiagaraDataSetID& Id, int32 DataSetIndex, TArray<FNiagaraVariable>& ReadVars)
+void FHlslNiagaraTranslator::DefineDataSetVariableReads(FString &OutHlslOutput, const FNiagaraDataSetID& Id, int32 DataSetIndex, const TArray<FNiagaraVariable>& ReadVars)
 {
 	int32 ReadOffsetInt = 0;
 	int32 ReadOffsetFloat = 0;
@@ -2183,7 +2179,7 @@ void FHlslNiagaraTranslator::DefineDataSetVariableReads(FString &OutHlslOutput,
 
 		FString VarReads;
 
-		for (FNiagaraVariable &Var : ReadVars)
+		for (const FNiagaraVariable &Var : ReadVars)
 		{
 			Fmt = ContextName + GetSanitizedSymbolName(Var.GetName().ToString()) + TEXT("{0} = {4};\n");
 			GatherVariableForDataSetAccess(Var, Fmt, IntCounter, FloatCounter, DataSetIndex, TEXT(""), VarReads);
@@ -2212,7 +2208,7 @@ void FHlslNiagaraTranslator::DefineDataSetVariableReads(FString &OutHlslOutput,
 
 		FString VarReads;
 
-		for (FNiagaraVariable &Var : ReadVars)
+		for (const FNiagaraVariable &Var : ReadVars)
 		{
 			const FString VariableName = ContextName + GetSanitizedSymbolName(Var.GetName().ToString());
 			// If the NiagaraClearEachFrame value is set on the data set, we don't bother reading it in each frame as we know that it is is invalid. However,
diff --git a/Engine/Plugins/FX/Niagara/Source/NiagaraEditor/Private/NiagaraHlslTranslator.h b/Engine/Plugins/FX/Niagara/Source/NiagaraEditor/Private/NiagaraHlslTranslator.h
--- a/Engine/Plugins/FX/Niagara/Source/NiagaraEditor/Private/NiagaraHlslTranslator.h
+++ b/Engine/Plugins/FX/Niagara/Source/NiagaraEditor/Private/NiagaraHlslTranslator.h
@@ -460,9 +460,10 @@ public:
 	void DefineInterpolatedParametersFunction(FString &HlslOutput);
 	void DefineDataSetReadFunction(FString &HlslOutput, TArray<FNiagaraDataSetID> &ReadDataSets);
 	void DefineDataSetWriteFunction(FString &HlslOutput, TArray<FNiagaraDataSetProperties> &WriteDataSets, TArray<int32>& WriteConditionVarIndices);
-	void DefineMain(FString &HLSLOutput, TArray<TArray<FNiagaraVariable>* > &InstanceReadVars, TArray<FNiagaraDataSetID>& ReadIds, TArray<TArray<FNiagaraVariable>* > &InstanceWriteVars, TArray<FNiagaraDataSetID>& WriteIds);
-	void DefineDataSetVariableReads(FString &HLSLOutput, FNiagaraDataSetID& Id, int32 DataSetIndex, TArray<FNiagaraVariable> & ReadVars);
-	void DefineDataSetVariableWrites(FString &HlslOutput, FNiagaraDataSetID& Id, int32 DataSetIndex, TArray<FNiagaraVariable>& WriteVars);
+	void DefineMain(FString &HLSLOutput, const TArray<TArray<FNiagaraVariable>>& DataSetVariables, const TMap<FNiagaraDataSetID, int32>& DataSetReads, const TMap<FNiagaraDataSetID, int32>& DataSetWrites);
+
+	void DefineDataSetVariableReads(FString &HLSLOutput, const FNiagaraDataSetID& Id, int32 DataSetIndex, const TArray<FNiagaraVariable>& ReadVars);
+	void DefineDataSetVariableWrites(FString &HlslOutput, const FNiagaraDataSetID& Id, int32 DataSetIndex, const TArray<FNiagaraVariable>& WriteVars);
 	void DefineDataInterfaceHLSL(FString &HlslOutput);
 	void DefineExternalFunctionsHLSL(FString &HlslOutput);
 	TArray<FNiagaraDataInterfaceGPUParamInfo>& GetDataInterfaceParameters() { return DIParamInfo; }
