- Print component path when we fail to activate a NiagaraComponent

#rnx

diff --git a/Engine/Plugins/FX/Niagara/Source/Niagara/Private/NiagaraComponent.cpp b/Engine/Plugins/FX/Niagara/Source/Niagara/Private/NiagaraComponent.cpp
--- a/Engine/Plugins/FX/Niagara/Source/Niagara/Private/NiagaraComponent.cpp
+++ b/Engine/Plugins/FX/Niagara/Source/Niagara/Private/NiagaraComponent.cpp
@@ -719,7 +719,7 @@ void UNiagaraComponent::Activate(bool bReset /* = false */)
 		DestroyInstance();
 		if (!HasAnyFlags(RF_DefaultSubObject | RF_ArchetypeObject | RF_ClassDefaultObject))
 		{
-			UE_LOG(LogNiagara, Warning, TEXT("Failed to activate Niagara Component due to missing or invalid asset!"));
+			UE_LOG(LogNiagara, Warning, TEXT("Failed to activate Niagara Component due to missing or invalid asset! (%s)"), *GetFullName());
 		}
 		SetComponentTickEnabled(false);
 		return;
