Fixed localization conflicts

diff --git a/Engine/Plugins/Editor/DataValidation/Source/DataValidation/Private/EditorValidatorBase.cpp b/Engine/Plugins/Editor/DataValidation/Source/DataValidation/Private/EditorValidatorBase.cpp
--- a/Engine/Plugins/Editor/DataValidation/Source/DataValidation/Private/EditorValidatorBase.cpp
+++ b/Engine/Plugins/Editor/DataValidation/Source/DataValidation/Private/EditorValidatorBase.cpp
@@ -40,7 +40,7 @@ void UEditorValidatorBase::AssetFails(UObject* InAsset, const FText& InMessage,
 	Arguments.Add(TEXT("CustomMessage"), InMessage);
 	Arguments.Add(TEXT("ValidatorName"), FText::FromString(GetClass()->GetName()));
 
-	FText FailureMessage = FText::Format(LOCTEXT("AssetCheck_Message_Display", "{AssetName} failed: {CustomMessage}. ({ValidatorName})"), Arguments);
+	FText FailureMessage = FText::Format(LOCTEXT("AssetCheck_Message_Error", "{AssetName} failed: {CustomMessage}. ({ValidatorName})"), Arguments);
 
 	if(LogContentValidation.GetVerbosity() >= ELogVerbosity::Verbose)
 	{
@@ -63,7 +63,7 @@ void UEditorValidatorBase::AssetWarning(UObject* InAsset, const FText& InMessage
 	Arguments.Add(TEXT("CustomMessage"), InMessage);
 	Arguments.Add(TEXT("ValidatorName"), FText::FromString(GetClass()->GetName()));
 
-	FText WarningMessage = FText::Format(LOCTEXT("AssetCheck_Message_Display", "{AssetName} warning: {CustomMessage}. ({ValidatorName})"), Arguments);
+	FText WarningMessage = FText::Format(LOCTEXT("AssetCheck_Message_Warning", "{AssetName} warning: {CustomMessage}. ({ValidatorName})"), Arguments);
 	AllWarnings.Add(WarningMessage);
 }
 
diff --git a/Engine/Plugins/Editor/DataValidation/Source/DataValidation/Private/EditorValidatorSubsystem.cpp b/Engine/Plugins/Editor/DataValidation/Source/DataValidation/Private/EditorValidatorSubsystem.cpp
--- a/Engine/Plugins/Editor/DataValidation/Source/DataValidation/Private/EditorValidatorSubsystem.cpp
+++ b/Engine/Plugins/Editor/DataValidation/Source/DataValidation/Private/EditorValidatorSubsystem.cpp
@@ -247,7 +247,7 @@ int32 UEditorValidatorSubsystem::ValidateAssets(TArray<FAssetData> AssetDataList
 			if (ValidationWarnings.Num() > 0)
 			{
 				DataValidationLog.Info()->AddToken(FAssetNameToken::Create(Data.PackageName.ToString()))
-					->AddToken(FTextToken::Create(LOCTEXT("InvalidDataResult", "contains valid data, but has warnings.")));
+					->AddToken(FTextToken::Create(LOCTEXT("ContainsWarningsResult", "contains valid data, but has warnings.")));
 			}
 			++NumValidFiles;
 		}
diff --git a/Engine/Plugins/Editor/GeometryMode/Source/GeometryMode/Private/GeometryModeModule.cpp b/Engine/Plugins/Editor/GeometryMode/Source/GeometryMode/Private/GeometryModeModule.cpp
--- a/Engine/Plugins/Editor/GeometryMode/Source/GeometryMode/Private/GeometryModeModule.cpp
+++ b/Engine/Plugins/Editor/GeometryMode/Source/GeometryMode/Private/GeometryModeModule.cpp
@@ -29,5 +29,3 @@ void FGeometryModeModule::ShutdownModule()
 }
 
 IMPLEMENT_MODULE(FGeometryModeModule, GeometryMode);
-
-#undef LOCTEXT_NAMESPACE
diff --git a/Engine/Plugins/FX/Niagara/Source/NiagaraEditor/Private/NiagaraNodeCustomHlsl.cpp b/Engine/Plugins/FX/Niagara/Source/NiagaraEditor/Private/NiagaraNodeCustomHlsl.cpp
--- a/Engine/Plugins/FX/Niagara/Source/NiagaraEditor/Private/NiagaraNodeCustomHlsl.cpp
+++ b/Engine/Plugins/FX/Niagara/Source/NiagaraEditor/Private/NiagaraNodeCustomHlsl.cpp
@@ -201,7 +201,7 @@ bool UNiagaraNodeCustomHlsl::VerifyEditablePinName(const FText& InName, FText& O
 
 	if (NewName != SanitizedNewName || NewName.Len() == 0)
 	{
-		OutErrorMessage = FText::Format(LOCTEXT("InvalidPinName", "Pin \"{0}\" cannot be renamed to \"{1}\". Certain words are restricted, as are spaces and special characters. Suggestion: \"{2}\""), InGraphPinObj->GetDisplayName(), InName, FText::FromString(SanitizedNewName));
+		OutErrorMessage = FText::Format(LOCTEXT("InvalidPinName_Restricted", "Pin \"{0}\" cannot be renamed to \"{1}\". Certain words are restricted, as are spaces and special characters. Suggestion: \"{2}\""), InGraphPinObj->GetDisplayName(), InName, FText::FromString(SanitizedNewName));
 		return false;
 	}
 	TSet<FName> Names;
@@ -212,7 +212,7 @@ bool UNiagaraNodeCustomHlsl::VerifyEditablePinName(const FText& InName, FText& O
 	}
 	if (Names.Contains(*NewName))
 	{
-		OutErrorMessage = FText::Format(LOCTEXT("InvalidPinName", "Pin \"{0}\" cannot be renamed to \"{1}\" as it conflicts with another name in use. Suggestion: \"{2}\""), InGraphPinObj->GetDisplayName(), InName, FText::FromName(FNiagaraUtilities::GetUniqueName(*SanitizedNewName, Names)));
+		OutErrorMessage = FText::Format(LOCTEXT("InvalidPinName_Conflicts", "Pin \"{0}\" cannot be renamed to \"{1}\" as it conflicts with another name in use. Suggestion: \"{2}\""), InGraphPinObj->GetDisplayName(), InName, FText::FromName(FNiagaraUtilities::GetUniqueName(*SanitizedNewName, Names)));
 		return false;
 	}
 	OutErrorMessage = FText::GetEmpty();
diff --git a/Engine/Source/Editor/LevelEditor/Private/LevelEditorToolBar.cpp b/Engine/Source/Editor/LevelEditor/Private/LevelEditorToolBar.cpp
--- a/Engine/Source/Editor/LevelEditor/Private/LevelEditorToolBar.cpp
+++ b/Engine/Source/Editor/LevelEditor/Private/LevelEditorToolBar.cpp
@@ -2163,6 +2163,7 @@ TSharedRef< SWidget > FLevelEditorToolBar::GenerateOpenBlueprintMenuContent( TSh
 	MenuContext.AddObject(LevelEditorMenuContext);
 
 	return UToolMenus::Get()->GenerateWidget("LevelEditor.LevelEditorToolBar.OpenBlueprint", MenuContext);
+#undef LOCTEXT_NAMESPACE
 }
 
 void FLevelEditorToolBar::RegisterOpenBlueprintMenu()
