- Fix deadlock if TickDataInterfaces destroy's the system instance it will still flag for async work / finalize even though it's complete

diff --git a/Engine/Plugins/FX/Niagara/Source/Niagara/Private/NiagaraSystemInstance.cpp b/Engine/Plugins/FX/Niagara/Source/Niagara/Private/NiagaraSystemInstance.cpp
--- a/Engine/Plugins/FX/Niagara/Source/Niagara/Private/NiagaraSystemInstance.cpp
+++ b/Engine/Plugins/FX/Niagara/Source/Niagara/Private/NiagaraSystemInstance.cpp
@@ -1732,7 +1732,10 @@ void FNiagaraSystemInstance::Tick_GameThread(float DeltaSeconds)
 	Age += DeltaSeconds;
 	TickCount += 1;
 	
-	BeginAsyncWork();
+	if ( !IsComplete() )
+	{
+		BeginAsyncWork();
+	}
 }
 
 void FNiagaraSystemInstance::Tick_Concurrent()
