Fix for lifetime issues with UNiagaraSystem and deferred deletion queues.

diff --git a/Engine/Plugins/FX/Niagara/Source/Niagara/Classes/NiagaraEmitterInstanceBatcher.h b/Engine/Plugins/FX/Niagara/Source/Niagara/Classes/NiagaraEmitterInstanceBatcher.h
--- a/Engine/Plugins/FX/Niagara/Source/Niagara/Classes/NiagaraEmitterInstanceBatcher.h
+++ b/Engine/Plugins/FX/Niagara/Source/Niagara/Classes/NiagaraEmitterInstanceBatcher.h
@@ -23,6 +23,15 @@ the same VectorVM byte code / compute shader code
 #include "NiagaraScriptExecutionContext.h"
 #include "NiagaraGPUInstanceCountManager.h"
 
+struct FNiagaraDeferredDeletionFence;
+
+struct FNiagaraInstanceBatcherDeferredDeletionFence
+{
+	FNiagaraDeferredDeletionFence* Fence;
+	FNiagaraInstanceBatcherDeferredDeletionFence(FNiagaraDeferredDeletionFence* InFence);
+	~FNiagaraInstanceBatcherDeferredDeletionFence();
+};
+
 class FNiagaraIndicesVertexBuffer : public FParticleIndicesVertexBuffer
 {
 public:
@@ -69,6 +78,8 @@ public:
 		DIProxyDeferredDeletes_RT.Add(MoveTemp(Proxy));
 	}
 
+	void AddFence_RenderThread(FNiagaraDeferredDeletionFence* Fence);
+
 #if WITH_EDITOR
 	virtual void Suspend() override {}
 	virtual void Resume() override {}
@@ -203,4 +214,6 @@ private:
 	TSet<FNiagaraDataSet*> DataSetsToDestroy_RT;
 
 	TSet<TSharedPtr<FNiagaraDataInterfaceProxy, ESPMode::ThreadSafe>> DIProxyDeferredDeletes_RT;
+
+	TArray<FNiagaraInstanceBatcherDeferredDeletionFence> Fences_RT;
 };
diff --git a/Engine/Plugins/FX/Niagara/Source/Niagara/Classes/NiagaraSystem.h b/Engine/Plugins/FX/Niagara/Source/Niagara/Classes/NiagaraSystem.h
--- a/Engine/Plugins/FX/Niagara/Source/Niagara/Classes/NiagaraSystem.h
+++ b/Engine/Plugins/FX/Niagara/Source/Niagara/Classes/NiagaraSystem.h
@@ -13,12 +13,26 @@
 #include "NiagaraParameterCollection.h"
 #include "NiagaraUserRedirectionParameterStore.h"
 #include "NiagaraSystemFastPath.h"
+#include "NiagaraEmitterInstanceBatcher.h"
 #include "NiagaraSystem.generated.h"
 
 #if WITH_EDITORONLY_DATA
 class UNiagaraEditorDataBase;
 #endif
 
+struct FNiagaraDeferredDeletionFence
+{
+	FNiagaraDeferredDeletionFence();
+	bool IsComplete();
+
+	/** Flag to signal that the instance batcher deferred deletion queues have been cleared. */
+	TAtomic<bool> bInstanceBatcherComplete;
+
+	/** Flag to signal that the world manager deferred deletion queues have been cleared. */
+	TAtomic<bool> bWorldManagerComplete;
+};
+
+
 USTRUCT()
 struct FNiagaraEmitterCompiledData
 {
@@ -131,6 +145,7 @@ public:
 	virtual void PostLoad() override; 
 	virtual void BeginDestroy() override;
 	virtual void PreSave(const class ITargetPlatform * TargetPlatform) override;
+	virtual bool IsReadyForFinishDestroy()override;
 #if WITH_EDITOR
 	virtual void PostEditChangeProperty(struct FPropertyChangedEvent& PropertyChangedEvent) override; 
 	virtual void BeginCacheForCookedPlatformData(const ITargetPlatform *TargetPlatform) override;
@@ -404,4 +419,7 @@ protected:
 	mutable TStatId StatID_RT;
 	mutable TStatId StatID_RT_CNC;
 #endif
+
+	//We submit one deletion fence for each world to ensure this System lives longer than any worlds deferred deletion queue.
+	TArray<FNiagaraDeferredDeletionFence> DeletionFences;
 };
diff --git a/Engine/Plugins/FX/Niagara/Source/Niagara/Private/NiagaraEmitterInstanceBatcher.cpp b/Engine/Plugins/FX/Niagara/Source/Niagara/Private/NiagaraEmitterInstanceBatcher.cpp
--- a/Engine/Plugins/FX/Niagara/Source/Niagara/Private/NiagaraEmitterInstanceBatcher.cpp
+++ b/Engine/Plugins/FX/Niagara/Source/Niagara/Private/NiagaraEmitterInstanceBatcher.cpp
@@ -131,6 +131,12 @@ void NiagaraEmitterInstanceBatcher::GiveDataSetToDestroy_RenderThread(FNiagaraDa
 	DataSetsToDestroy_RT.Add(DataSet);
 }
 
+void NiagaraEmitterInstanceBatcher::AddFence_RenderThread(FNiagaraDeferredDeletionFence* Fence)
+{
+	LLM_SCOPE(ELLMTag::Niagara);
+	Fences_RT.Emplace(Fence);
+}
+
 void NiagaraEmitterInstanceBatcher::FinishDispatches()
 {
 	ReleaseTicks();
@@ -159,6 +165,8 @@ void NiagaraEmitterInstanceBatcher::FinishDispatches()
 	}
 
 	DIProxyDeferredDeletes_RT.Empty();
+
+	Fences_RT.Empty();
 }
 
 void NiagaraEmitterInstanceBatcher::ReleaseTicks()
diff --git a/Engine/Plugins/FX/Niagara/Source/Niagara/Private/NiagaraSystem.cpp b/Engine/Plugins/FX/Niagara/Source/Niagara/Private/NiagaraSystem.cpp
--- a/Engine/Plugins/FX/Niagara/Source/Niagara/Private/NiagaraSystem.cpp
+++ b/Engine/Plugins/FX/Niagara/Source/Niagara/Private/NiagaraSystem.cpp
@@ -16,6 +16,7 @@
 #include "NiagaraStats.h"
 #include "NiagaraEditorDataBase.h"
 #include "INiagaraEditorOnlyDataUtlities.h"
+#include "NiagaraWorldManager.h"
 
 #if WITH_EDITOR
 #include "NiagaraScriptDerivedData.h"
@@ -62,6 +63,10 @@ UNiagaraSystem::UNiagaraSystem(const FObjectInitializer& ObjectInitializer)
 void UNiagaraSystem::BeginDestroy()
 {
 	Super::BeginDestroy();
+
+	DeletionFences.Reset();
+	FNiagaraWorldManager::EnqueueDeferredDeletionFences(DeletionFences);
+
 #if WITH_EDITORONLY_DATA
 	while (ActiveCompilations.Num() > 0)
 	{
@@ -70,6 +75,20 @@ void UNiagaraSystem::BeginDestroy()
 #endif
 }
 
+bool UNiagaraSystem::IsReadyForFinishDestroy()
+{
+ 	bool bAllComplete = true;
+	
+	//We have to wait until the existing items in the deferred deletion queue are cleared as they have direct refs to the layout info in our System and Emitter compiled data.
+	//We have to do this for each world.
+	for (FNiagaraDeferredDeletionFence& Fence : DeletionFences)
+	{
+		bAllComplete &= Fence.IsComplete();
+	}
+	
+	return Super::IsReadyForFinishDestroy() && bAllComplete;
+}
+
 void UNiagaraSystem::PreSave(const class ITargetPlatform * TargetPlatform)
 {
 	Super::PreSave(TargetPlatform);
@@ -1318,3 +1337,53 @@ FNiagaraEmitterCompiledData::FNiagaraEmitterCompiledData()
 	EmitterRandomSeedVar = SYS_PARAM_EMITTER_RANDOM_SEED;
 	EmitterTotalSpawnedParticlesVar = SYS_PARAM_ENGINE_EMITTER_TOTAL_SPAWNED_PARTICLES;
 }
+
+
+//////////////////////////////////////////////////////////////////////////
+
+FNiagaraDeferredDeletionFence::FNiagaraDeferredDeletionFence()
+{
+	bInstanceBatcherComplete = false;
+	bWorldManagerComplete = false;
+}
+
+bool FNiagaraDeferredDeletionFence::IsComplete()
+{
+	return bInstanceBatcherComplete.Load() && bWorldManagerComplete.Load();
+}
+
+FNiagaraWorldManagerDeferredDeletionFence::FNiagaraWorldManagerDeferredDeletionFence(FNiagaraDeferredDeletionFence* InFence)
+	: Fence(InFence)
+{
+	if (Fence)
+	{
+		Fence->bWorldManagerComplete = false;
+	}
+}
+
+FNiagaraWorldManagerDeferredDeletionFence::~FNiagaraWorldManagerDeferredDeletionFence()
+{
+	if (Fence)
+	{
+		Fence->bWorldManagerComplete = true;
+	}
+}
+
+FNiagaraInstanceBatcherDeferredDeletionFence::FNiagaraInstanceBatcherDeferredDeletionFence(FNiagaraDeferredDeletionFence* InFence)
+	: Fence(InFence)
+{
+	if (Fence)
+	{
+		Fence->bInstanceBatcherComplete = false;
+	}
+}
+FNiagaraInstanceBatcherDeferredDeletionFence::~FNiagaraInstanceBatcherDeferredDeletionFence()
+{
+	if (Fence)
+	{
+		Fence->bInstanceBatcherComplete = true;
+	}
+}
+
+//////////////////////////////////////////////////////////////////////////
+//////////////////////////////////////////////////////////////////////////
\ No newline at end of file
diff --git a/Engine/Plugins/FX/Niagara/Source/Niagara/Private/NiagaraWorldManager.cpp b/Engine/Plugins/FX/Niagara/Source/Niagara/Private/NiagaraWorldManager.cpp
--- a/Engine/Plugins/FX/Niagara/Source/Niagara/Private/NiagaraWorldManager.cpp
+++ b/Engine/Plugins/FX/Niagara/Source/Niagara/Private/NiagaraWorldManager.cpp
@@ -313,6 +313,7 @@ void FNiagaraWorldManager::OnBatcherDestroyed_Internal(NiagaraEmitterInstanceBat
 		{
 			DeferredDeletionQueue[DeferredDeletionQueueIndex].Fence.Wait();
 			DeferredDeletionQueue[i].Queue.Empty();
+			DeferredDeletionQueue[i].DeletionFences.Empty();
 		}
 	}
 }
@@ -335,6 +336,7 @@ void FNiagaraWorldManager::OnWorldCleanup(bool bSessionEnded, bool bCleanupResou
 	{
 		DeferredDeletionQueue[i].Fence.Wait();
 		DeferredDeletionQueue[i].Queue.Empty();
+		DeferredDeletionQueue[i].DeletionFences.Empty();
 	}
 }
 
@@ -453,6 +455,7 @@ void FNiagaraWorldManager::PostActorTick(float DeltaSeconds)
 			DeferredDeletionQueue[DeferredDeletionQueueIndex].Fence.Wait();
 		}
 		DeferredDeletionQueue[DeferredDeletionQueueIndex].Queue.Empty();
+		DeferredDeletionQueue[DeferredDeletionQueueIndex].DeletionFences.Empty();
 	}
 
 	// Update tick groups
@@ -561,3 +564,40 @@ void FNiagaraWorldManager::DumpDetails(FOutputDevice& Ar)
 		}
 	}
 }
+
+void FNiagaraWorldManager::EnqueueDeferredDeletionFences(TArray<FNiagaraDeferredDeletionFence>& OutFences)
+{
+	for (auto& Pair : WorldManagers)
+	{
+		UWorld* World = Pair.Key;
+		FNiagaraWorldManager* Manager = Pair.Value;
+
+		if (World && Manager)
+		{
+			FNiagaraDeferredDeletionFence* Fence = &OutFences.AddDefaulted_GetRef();
+			//Add the fence to the deferred deletion queue.
+			Manager->DeferredDeletionQueue[Manager->DeferredDeletionQueueIndex].DeletionFences.Emplace(Fence);
+
+			//Add the fence to the instance batcher
+			if (World->Scene)
+			{
+				FFXSystemInterface*  FXSystemInterface = World->Scene->GetFXSystem();
+				if (FXSystemInterface)
+				{
+					if (NiagaraEmitterInstanceBatcher* Batcher = static_cast<NiagaraEmitterInstanceBatcher*>(FXSystemInterface->GetInterface(NiagaraEmitterInstanceBatcher::Name)))
+					{
+						ENQUEUE_RENDER_COMMAND(FDeletionFenceCommand)(
+							[Batcher, Fence](FRHICommandListImmediate& RHICmdList)
+						{
+							if (ensure(Batcher))
+							{
+								Batcher->AddFence_RenderThread(Fence);
+							}
+						}
+						);
+					}
+				}
+			}
+		}
+	}
+}
\ No newline at end of file
diff --git a/Engine/Plugins/FX/Niagara/Source/Niagara/Public/NiagaraWorldManager.h b/Engine/Plugins/FX/Niagara/Source/Niagara/Public/NiagaraWorldManager.h
--- a/Engine/Plugins/FX/Niagara/Source/Niagara/Public/NiagaraWorldManager.h
+++ b/Engine/Plugins/FX/Niagara/Source/Niagara/Public/NiagaraWorldManager.h
@@ -24,6 +24,14 @@ class UNiagaraParameterCollection;
 class UNiagaraParameterCollectionInstance;
 class UNiagaraComponentPool;
 
+struct FNiagaraDeferredDeletionFence;
+
+struct FNiagaraWorldManagerDeferredDeletionFence
+{
+	FNiagaraDeferredDeletionFence* Fence;
+	FNiagaraWorldManagerDeferredDeletionFence(FNiagaraDeferredDeletionFence* InFence);
+	~FNiagaraWorldManagerDeferredDeletionFence();
+};
 
 class FNiagaraViewDataMgr : public FRenderResource
 {
@@ -132,6 +140,9 @@ public:
 	// Dump details about what's inside the world manager
 	void DumpDetails(FOutputDevice& Ar);
 
+	//Creates fences for deferred deletion for each world.
+	static void EnqueueDeferredDeletionFences(TArray<FNiagaraDeferredDeletionFence>& OutFence);
+
 	UWorld* World;
 
 private:
@@ -184,6 +195,9 @@ private:
 	{
 		FRenderCommandFence							Fence;
 		TArray<TUniquePtr<FNiagaraSystemInstance>>	Queue;
+
+		/** Fences to signal to other Niagara code when this queue has been cleared. */
+		TArray<FNiagaraWorldManagerDeferredDeletionFence>		DeletionFences;
 	};
 
 	static constexpr int NumDeferredQueues = 3;
