Compile fix.

diff --git a/Engine/Plugins/FX/Niagara/Source/Niagara/Classes/NiagaraSystem.h b/Engine/Plugins/FX/Niagara/Source/Niagara/Classes/NiagaraSystem.h
--- a/Engine/Plugins/FX/Niagara/Source/Niagara/Classes/NiagaraSystem.h
+++ b/Engine/Plugins/FX/Niagara/Source/Niagara/Classes/NiagaraSystem.h
@@ -13,7 +13,6 @@
 #include "NiagaraParameterCollection.h"
 #include "NiagaraUserRedirectionParameterStore.h"
 #include "NiagaraSystemFastPath.h"
-#include "NiagaraEmitterInstanceBatcher.h"
 #include "NiagaraSystem.generated.h"
 
 #if WITH_EDITORONLY_DATA
