Compile fix.

diff --git a/Engine/Plugins/FX/Niagara/Source/Niagara/Private/NiagaraSystem.cpp b/Engine/Plugins/FX/Niagara/Source/Niagara/Private/NiagaraSystem.cpp
--- a/Engine/Plugins/FX/Niagara/Source/Niagara/Private/NiagaraSystem.cpp
+++ b/Engine/Plugins/FX/Niagara/Source/Niagara/Private/NiagaraSystem.cpp
@@ -17,6 +17,7 @@
 #include "NiagaraEditorDataBase.h"
 #include "INiagaraEditorOnlyDataUtlities.h"
 #include "NiagaraWorldManager.h"
+#include "NiagaraEmitterInstanceBatcher.h"
 
 #if WITH_EDITOR
 #include "NiagaraScriptDerivedData.h"
