Fix include order related compile errors in Niagara system.

diff --git a/Engine/Plugins/FX/Niagara/Source/Niagara/Classes/NiagaraGPUInstanceCountManager.h b/Engine/Plugins/FX/Niagara/Source/Niagara/Classes/NiagaraGPUInstanceCountManager.h
--- a/Engine/Plugins/FX/Niagara/Source/Niagara/Classes/NiagaraGPUInstanceCountManager.h
+++ b/Engine/Plugins/FX/Niagara/Source/Niagara/Classes/NiagaraGPUInstanceCountManager.h
@@ -10,10 +10,17 @@ NiagaraGPUInstanceCountManager.h: GPU particle count handling
 #include "RHIUtilities.h"
 #include "NiagaraDrawIndirect.h"
 #include "RHIGPUReadback.h"
-#include "NiagaraRenderer.h"
 
 class FRHIGPUMemoryReadback;
 
+// The number of GPU renderers registered in the instance count manager.
+// Shared between the manager and the renderers.
+class FNiagaraGPURendererCount : public FRefCountedObject
+{
+public:
+	int32 Value = 0;
+};
+
 FORCEINLINE uint32 GetTypeHash(const FNiagaraDrawIndirectArgsGenCS::FArgGenTaskInfo& Info)
 {
 	return (Info.InstanceCountBufferOffset << 8) | Info.NumIndicesPerInstance;
diff --git a/Engine/Plugins/FX/Niagara/Source/Niagara/Private/NiagaraComponent.cpp b/Engine/Plugins/FX/Niagara/Source/Niagara/Private/NiagaraComponent.cpp
--- a/Engine/Plugins/FX/Niagara/Source/Niagara/Private/NiagaraComponent.cpp
+++ b/Engine/Plugins/FX/Niagara/Source/Niagara/Private/NiagaraComponent.cpp
@@ -428,6 +428,10 @@ void UNiagaraComponent::SetActorParameter(FName ParameterName, class AActor* Par
 	SetVariableActor(ParameterName, Param);
 }
 
+UFXSystemAsset* UNiagaraComponent::GetFXSystemAsset() const
+{
+	return Asset;
+}
 
 void UNiagaraComponent::SetEmitterEnable(FName EmitterName, bool bNewEnableState)
 {
diff --git a/Engine/Plugins/FX/Niagara/Source/Niagara/Public/NiagaraComponent.h b/Engine/Plugins/FX/Niagara/Source/Niagara/Public/NiagaraComponent.h
--- a/Engine/Plugins/FX/Niagara/Source/Niagara/Public/NiagaraComponent.h
+++ b/Engine/Plugins/FX/Niagara/Source/Niagara/Public/NiagaraComponent.h
@@ -48,7 +48,7 @@ public:
 	void SetVectorParameter(FName ParameterName, FVector Param) override;
 	void SetColorParameter(FName ParameterName, FLinearColor Param) override;
 	void SetActorParameter(FName ParameterName, class AActor* Param) override;
-	virtual UFXSystemAsset* GetFXSystemAsset() const override { return Asset; };
+	virtual UFXSystemAsset* GetFXSystemAsset() const override;
 	void SetEmitterEnable(FName EmitterName, bool bNewEnableState) override;
 	void ReleaseToPool() override;
 	uint32 GetApproxMemoryUsage() const override;
diff --git a/Engine/Plugins/FX/Niagara/Source/Niagara/Public/NiagaraRenderer.h b/Engine/Plugins/FX/Niagara/Source/Niagara/Public/NiagaraRenderer.h
--- a/Engine/Plugins/FX/Niagara/Source/Niagara/Public/NiagaraRenderer.h
+++ b/Engine/Plugins/FX/Niagara/Source/Niagara/Public/NiagaraRenderer.h
@@ -23,6 +23,8 @@ NiagaraRenderer.h: Base class for Niagara render modules
 #include "NiagaraBoundsCalculator.h"
 
 class FNiagaraDataSet;
+class FNiagaraSceneProxy;
+class FNiagaraGPURendererCount;
 
 /** Struct used to pass dynamic data from game thread to render thread */
 struct FNiagaraDynamicDataBase
@@ -76,14 +78,6 @@ struct FNiagaraRendererVariableInfo
 	bool bUpload;
 };
 
-// The number of GPU renderers registered in the instance count manager.
-// Shared between the manager and the renderers.
-class FNiagaraGPURendererCount : public FRefCountedObject
-{
-public:
-	int32 Value = 0;
-};
-
 /**
 * Base class for Niagara System renderers.
 */
