Niagara - Change the indenting of items in the overview stacks and add the icons for renderers, and optional "highlight" support for modules.

diff --git a/Engine/Plugins/FX/Niagara/Source/Niagara/Classes/NiagaraScript.h b/Engine/Plugins/FX/Niagara/Source/Niagara/Classes/NiagaraScript.h
--- a/Engine/Plugins/FX/Niagara/Source/Niagara/Classes/NiagaraScript.h
+++ b/Engine/Plugins/FX/Niagara/Source/Niagara/Classes/NiagaraScript.h
@@ -12,6 +12,7 @@
 #include "NiagaraDataSet.h"
 #include "NiagaraShared.h"
 #include "NiagaraScriptExecutionParameterStore.h"
+#include "NiagaraScriptHighlight.h"
 #include "NiagaraCustomVersion.h"
 
 #include "NiagaraScript.generated.h"
@@ -395,6 +396,9 @@ public:
 	UPROPERTY(AssetRegistrySearchable, EditAnywhere, Category = Script)
 	FText Keywords;
 
+	UPROPERTY(EditAnywhere, Category = Script)
+	TArray<FNiagaraScriptHighlight> Highlights;
+
 	UPROPERTY(EditAnywhere, Category = Script, DisplayName = "Script Metadata", meta = (ToolTip = "Script Metadata"))
 	TMap<FName, FString> ScriptMetaData;
 #endif
@@ -625,6 +629,8 @@ private:
 	FOnScriptCompiled OnVMScriptCompiledDelegate;
 
 	mutable FNiagaraVMExecutableDataId LastReportedVMId;
+
+	mutable TOptional<TMap<FName, FString>> CustomAssetRegistryTagCache;
 #endif
 
 	/** Adjusted every time that we compile this script. Lets us know that we might differ from any cached versions.*/
diff --git a/Engine/Plugins/FX/Niagara/Source/Niagara/Classes/NiagaraScriptHighlight.h b/Engine/Plugins/FX/Niagara/Source/Niagara/Classes/NiagaraScriptHighlight.h
new file mode 100644
--- /dev/null
+++ b/Engine/Plugins/FX/Niagara/Source/Niagara/Classes/NiagaraScriptHighlight.h
@@ -0,0 +1,28 @@
+// Copyright 1998-2019 Epic Games, Inc. All Rights Reserved.
+
+#pragma once
+
+#include "NiagaraScriptHighlight.generated.h"
+
+USTRUCT()
+struct NIAGARA_API FNiagaraScriptHighlight
+{
+	GENERATED_BODY()
+
+	FNiagaraScriptHighlight();
+
+	bool IsValid() const;
+
+	UPROPERTY(EditAnywhere, Category = Highlight)
+	FLinearColor Color;
+
+	UPROPERTY(EditAnywhere, Category = Highlight)
+	FText DisplayName;
+
+	bool operator==(const FNiagaraScriptHighlight& Other) const;
+
+#if WITH_EDITORONLY_DATA
+	static void ArrayToJson(const TArray<FNiagaraScriptHighlight>& InHighlights, FString& OutJson);
+	static void JsonToArray(const FString& InJson, TArray<FNiagaraScriptHighlight>& OutHighlights);
+#endif
+};
\ No newline at end of file
diff --git a/Engine/Plugins/FX/Niagara/Source/Niagara/Niagara.Build.cs b/Engine/Plugins/FX/Niagara/Source/Niagara/Niagara.Build.cs
--- a/Engine/Plugins/FX/Niagara/Source/Niagara/Niagara.Build.cs
+++ b/Engine/Plugins/FX/Niagara/Source/Niagara/Niagara.Build.cs
@@ -16,6 +16,8 @@ public class Niagara : ModuleRules
                 "Engine",
                 "TimeManagement",
                 "Renderer",
+                "JsonUtilities",
+				"Json",
             }
         );
 
diff --git a/Engine/Plugins/FX/Niagara/Source/Niagara/Private/NiagaraScript.cpp b/Engine/Plugins/FX/Niagara/Source/Niagara/Private/NiagaraScript.cpp
--- a/Engine/Plugins/FX/Niagara/Source/Niagara/Private/NiagaraScript.cpp
+++ b/Engine/Plugins/FX/Niagara/Source/Niagara/Private/NiagaraScript.cpp
@@ -814,6 +814,8 @@ void UNiagaraScript::PostEditChangeProperty(FPropertyChangedEvent& PropertyChang
 			Source->MarkNotSynchronized(TEXT("Deprecation changed."));
 		}
 	}
+
+	CustomAssetRegistryTagCache.Reset();
 }
 
 #endif
@@ -1284,15 +1286,48 @@ void UNiagaraScript::GetAssetRegistryTags(TArray<FAssetRegistryTag>& OutTags) co
 {
 	Super::GetAssetRegistryTags(OutTags);
 #if WITH_EDITORONLY_DATA
-	FName TagName = GET_MEMBER_NAME_CHECKED(UNiagaraScript, ProvidedDependencies);
-	FString DependenciesProvidedString;
-	for (FName DependencyProvided : ProvidedDependencies)
+
+	if (ProvidedDependencies.Num() > 0)
 	{
-		DependenciesProvidedString.Append(DependencyProvided.ToString() + ",");
+		FName ProvidedDependenciesName = GET_MEMBER_NAME_CHECKED(UNiagaraScript, ProvidedDependencies);
+		FString* ProvidedDependenciesTags;
+
+		if (CustomAssetRegistryTagCache.IsSet() == false)
+		{
+			CustomAssetRegistryTagCache = TMap<FName, FString>();
+		}
+
+		ProvidedDependenciesTags = CustomAssetRegistryTagCache->Find(ProvidedDependenciesName);
+		if(ProvidedDependenciesTags == nullptr)
+		{
+			ProvidedDependenciesTags = &CustomAssetRegistryTagCache->Add(ProvidedDependenciesName);
+			for (FName ProvidedDependency : ProvidedDependencies)
+			{
+				ProvidedDependenciesTags->Append(ProvidedDependency.ToString() + ",");
+			}
+		}
+
+		OutTags.Add(FAssetRegistryTag(ProvidedDependenciesName, *ProvidedDependenciesTags, UObject::FAssetRegistryTag::TT_Hidden));
 	}
-	if (ProvidedDependencies.Num() > 0)
+
+	if (Highlights.Num() > 0)
 	{
-		OutTags.Add(FAssetRegistryTag(TagName, DependenciesProvidedString, UObject::FAssetRegistryTag::TT_Hidden));
+		FName HighlightsName = GET_MEMBER_NAME_CHECKED(UNiagaraScript, Highlights);
+		FString* HighlightsTags;
+
+		if (CustomAssetRegistryTagCache.IsSet() == false)
+		{
+			CustomAssetRegistryTagCache = TMap<FName, FString>();
+		}
+
+		HighlightsTags = CustomAssetRegistryTagCache->Find(HighlightsName);
+		if (HighlightsTags == nullptr)
+		{
+			HighlightsTags = &CustomAssetRegistryTagCache->Add(HighlightsName);
+			FNiagaraScriptHighlight::ArrayToJson(Highlights, *HighlightsTags);
+		}
+
+		OutTags.Add(FAssetRegistryTag(HighlightsName, *HighlightsTags, UObject::FAssetRegistryTag::TT_Hidden));
 	}
 #endif
 }
diff --git a/Engine/Plugins/FX/Niagara/Source/Niagara/Private/NiagaraScriptHighlight.cpp b/Engine/Plugins/FX/Niagara/Source/Niagara/Private/NiagaraScriptHighlight.cpp
new file mode 100644
--- /dev/null
+++ b/Engine/Plugins/FX/Niagara/Source/Niagara/Private/NiagaraScriptHighlight.cpp
@@ -0,0 +1,48 @@
+// Copyright 1998-2019 Epic Games, Inc. All Rights Reserved.
+
+#include "NiagaraScriptHighlight.h"
+
+#if WITH_EDITORONLY_DATA
+#include "JsonObjectConverter.h"
+#include "Serialization/JsonWriter.h"
+#include "Serialization/JsonSerializer.h"
+#endif
+
+FNiagaraScriptHighlight::FNiagaraScriptHighlight()
+	: Color(FLinearColor::Transparent)
+{
+}
+
+bool FNiagaraScriptHighlight::IsValid() const
+{
+	return Color != FLinearColor::Transparent && DisplayName.IsEmpty() == false;
+}
+
+bool FNiagaraScriptHighlight::operator==(const FNiagaraScriptHighlight& Other) const
+{
+	return Color == Other.Color && 
+		(DisplayName.IdenticalTo(Other.DisplayName) || DisplayName.CompareTo(Other.DisplayName) == 0);
+}
+
+#if WITH_EDITORONLY_DATA
+void FNiagaraScriptHighlight::ArrayToJson(const TArray<FNiagaraScriptHighlight>& InHighlights, FString& OutJson)
+{
+	TArray<TSharedPtr<FJsonValue>> HighlightValues;
+	for (const FNiagaraScriptHighlight& Highlight : InHighlights)
+	{
+		TSharedPtr<FJsonObject> HighlightObject = FJsonObjectConverter::UStructToJsonObject(Highlight);
+		if (HighlightObject.IsValid())
+		{
+			HighlightValues.Add(MakeShared<FJsonValueObject>(HighlightObject));
+		}
+	}
+	TSharedRef<FJsonValueArray> HighlightValuesArrayValue = MakeShared<FJsonValueArray>(HighlightValues);
+	TSharedRef<TJsonWriter<TCHAR, TCondensedJsonPrintPolicy<TCHAR>>> JsonWriter = TJsonWriterFactory<TCHAR, TCondensedJsonPrintPolicy<TCHAR>>::Create(&OutJson);
+	FJsonSerializer::Serialize(HighlightValuesArrayValue, FString(), JsonWriter);
+}
+
+void FNiagaraScriptHighlight::JsonToArray(const FString& InJson, TArray<FNiagaraScriptHighlight>& OutHighlights)
+{
+	FJsonObjectConverter::JsonArrayStringToUStruct(InJson, &OutHighlights, 0, 0);
+}
+#endif
\ No newline at end of file
diff --git a/Engine/Plugins/FX/Niagara/Source/NiagaraEditor/Private/Customizations/NiagaraScriptDetails.cpp b/Engine/Plugins/FX/Niagara/Source/NiagaraEditor/Private/Customizations/NiagaraScriptDetails.cpp
--- a/Engine/Plugins/FX/Niagara/Source/NiagaraEditor/Private/Customizations/NiagaraScriptDetails.cpp
+++ b/Engine/Plugins/FX/Niagara/Source/NiagaraEditor/Private/Customizations/NiagaraScriptDetails.cpp
@@ -18,13 +18,18 @@
 #include "Widgets/Input/SComboButton.h"
 #include "Widgets/Input/SButton.h"
 #include "Widgets/Images/SImage.h"
+#include "Widgets/Colors/SColorPicker.h"
+#include "Widgets/Colors/SColorBlock.h"
+#include "Widgets/Layout/SBox.h"
+#include "Widgets/SOverlay.h"
 #include "NiagaraEditorModule.h"
 #include "Modules/ModuleManager.h"
 #include "INiagaraEditorTypeUtilities.h"
-#include "Widgets/Layout/SBox.h"
 #include "NiagaraScript.h"
 #include "NiagaraParameterCollectionCustomNodeBuilder.h"
 #include "NiagaraMetaDataCustomNodeBuilder.h"
+#include "SItemSelector.h"
+#include "PropertyCustomizationHelpers.h"
 
 #define LOCTEXT_NAMESPACE "NiagaraScriptDetails"
 
@@ -211,5 +216,641 @@ FReply FNiagaraScriptDetails::OnRefreshMetadata()
     return FReply::Handled();
 }
 
+float ButtonWidth = 50;
+float DisplayColorHeight = 16;
+float DisplayColorWidth = 24;
+float EditColorHeight = 16;
+float EditColorWidth = 30;
+
+class SNiagaraDetailsButton : public SButton
+{
+	SLATE_BEGIN_ARGS(SNiagaraDetailsButton) {}
+		SLATE_EVENT(FOnClicked, OnClicked)
+		SLATE_ARGUMENT(FText, Text);
+	SLATE_END_ARGS()
+
+	void Construct(const FArguments& InArgs)
+	{
+		SButton::Construct(SButton::FArguments()
+			.OnClicked(InArgs._OnClicked)
+			.Content()
+			[
+				SNew(STextBlock)
+				.TextStyle(FNiagaraEditorStyle::Get(), "NiagaraEditor.ParameterText")
+				.Text(InArgs._Text)
+				.Justification(ETextJustify::Center)
+				.MinDesiredWidth(ButtonWidth)
+			]);
+	}
+};
+
+class SNiagaraScriptHighlightPicker : public SCompoundWidget
+{
+private:
+	enum class EEditMode
+	{
+		Picker,
+		EditCurrent,
+		EditNew
+	};
+
+public:
+	SLATE_BEGIN_ARGS(SNiagaraScriptHighlightPicker)
+	{}
+	SLATE_END_ARGS()
+
+	void Construct(const FArguments& InArgs, TSharedRef<IPropertyHandle> InPropertyHandle)
+	{
+		PropertyHandle = InPropertyHandle;
+		PropertyHandle->SetOnPropertyValueChanged(FSimpleDelegate::CreateSP(this, &SNiagaraScriptHighlightPicker::RefreshDisplayedValue));
+		EditMode = EEditMode::Picker;
+		RefreshDisplayedValue();
+		ChildSlot
+		[
+			SNew(SOverlay)
+			// Highlight picker
+			+ SOverlay::Slot()
+			[
+				ConstructPickerWidgets()
+			]
+			// Highlight Editor
+			+ SOverlay::Slot()
+			[
+				ConstructEditWidgets()
+			]
+		];
+	}
+
+private:
+
+	TSharedRef<SWidget> ConstructPickerWidgets()
+	{
+		return SNew(SHorizontalBox)
+			// Hightlight picker combo
+			+ SHorizontalBox::Slot()
+			[
+				SAssignNew(HighlightPickerCombo, SComboButton)
+				.Visibility(this, &SNiagaraScriptHighlightPicker::GetPickerVisibility)
+				.OnGetMenuContent(this, &SNiagaraScriptHighlightPicker::OnGetMenuContent)
+				.ContentPadding(FMargin(4, 2, 4, 2))
+				.ButtonContent()
+				[
+					SNew(SHorizontalBox)
+					+ SHorizontalBox::Slot()
+					.AutoWidth()
+					.Padding(0, 0, 5, 0)
+					[
+						SNew(SBorder)
+						.BorderImage(FEditorStyle::GetBrush("WhiteBrush"))
+						.BorderBackgroundColor(this, &SNiagaraScriptHighlightPicker::GetHighlightSlateColor)
+						.Padding(0)
+						[
+							SNew(SBox)
+							.WidthOverride(DisplayColorWidth)
+							.HeightOverride(DisplayColorHeight)
+						]
+					]
+					+ SHorizontalBox::Slot()
+					.VAlign(VAlign_Center)
+					[
+						SNew(STextBlock)
+						.TextStyle(FNiagaraEditorStyle::Get(), "NiagaraEditor.ParameterText")
+						.Text(this, &SNiagaraScriptHighlightPicker::GetHighlightText)
+					]
+				]
+			]
+			+ SHorizontalBox::Slot()
+			.AutoWidth()
+			.VAlign(VAlign_Center)
+			[
+				PropertyHandle->CreateDefaultPropertyButtonWidgets()
+			];
+	}
+
+	TSharedRef<SWidget> ConstructEditWidgets()
+	{
+		return SNew(SBorder)
+			.BorderImage(FEditorStyle::GetBrush("ToolPanel.DarkGroupBorder"))
+			.Visibility(this, &SNiagaraScriptHighlightPicker::GetEditVisibility)
+			.Padding(10)
+			[
+				SNew(SVerticalBox)
+
+				// Edit Header
+				+ SVerticalBox::Slot()
+				.AutoHeight()
+				[
+					SNew(STextBlock)
+					.TextStyle(FNiagaraEditorStyle::Get(), "NiagaraEditor.DetailsHeadingText")
+					.Text(this, &SNiagaraScriptHighlightPicker::GetEditModeText)
+				]
+				// Edit controls.
+				+ SVerticalBox::Slot()
+				.AutoHeight()
+				.Padding(0, 10, 0, 0)
+				[
+					ConstructEditHighlightWidgets()
+				]
+				// Error display.
+				+ SVerticalBox::Slot()
+				.AutoHeight()
+				.Padding(0, 5, 0, 0)
+				[
+					ConstructEditErrorWidgets()
+				]
+				// Edit buttons
+				+ SVerticalBox::Slot()
+				.AutoHeight()
+				.HAlign(HAlign_Right)
+				.Padding(0, 10, 0, 0)
+				[
+					ConstructEditButtons()
+				]
+			];
+	}
+
+	TSharedRef<SWidget> ConstructEditHighlightWidgets()
+	{
+		return SNew(SHorizontalBox)
+			+ SHorizontalBox::Slot()
+			.AutoWidth()
+			.Padding(0, 0, 5, 0)
+			[
+				SAssignNew(ColorBlock, SColorBlock)
+				.Color(this, &SNiagaraScriptHighlightPicker::GetHighlightColor)
+				.ShowBackgroundForAlpha(true)
+				.IgnoreAlpha(false)
+				.OnMouseButtonDown(this, &SNiagaraScriptHighlightPicker::OnMouseButtonDownColorBlock)
+				.Size(FVector2D(EditColorWidth, EditColorHeight))
+			]
+			+ SHorizontalBox::Slot()
+			[
+				SNew(SEditableTextBox)
+				.Text(this, &SNiagaraScriptHighlightPicker::GetHighlightText)
+				.OnTextCommitted(this, &SNiagaraScriptHighlightPicker::HighlightTextCommitted)
+			];
+	};
+
+	TSharedRef<SWidget> ConstructEditButtons()
+	{
+		return SNew(SHorizontalBox)
+			// Add new button.
+			+ SHorizontalBox::Slot()
+			[
+				SNew(SNiagaraDetailsButton)
+				.Visibility(this, &SNiagaraScriptHighlightPicker::GetAddNewButtonVisibility)
+				.Text(LOCTEXT("AddNewHighlight", "Add"))
+				.OnClicked(this, &SNiagaraScriptHighlightPicker::AddNewRequested)
+			]
+			// Apply button
+			+ SHorizontalBox::Slot()
+			.Padding(5, 0, 0, 0)
+			[
+				SNew(SNiagaraDetailsButton)
+				.Visibility(this, &SNiagaraScriptHighlightPicker::GetApplyButtonVisibility)
+				.Text(LOCTEXT("ApplyToCurrent", "Apply"))
+				.ToolTipText(LOCTEXT("ApplyToCurrentToolTip", "Applies the current highlight changes to this script."))
+				.OnClicked(this, &SNiagaraScriptHighlightPicker::ApplyCurrentRequested)
+			]
+			// Apply all button
+			+ SHorizontalBox::Slot()
+			.Padding(5, 0, 0, 0)
+			[
+				SNew(SNiagaraDetailsButton)
+				.Visibility(this, &SNiagaraScriptHighlightPicker::GetApplyButtonVisibility)
+				.Text(LOCTEXT("ApplyToAll", "Apply All"))
+				.ToolTipText(LOCTEXT("ApplyToAllCurrentToolTip", "Applies the current highlight changes to this script,\nand any other scripts that use this same highlight."))
+				.OnClicked(this, &SNiagaraScriptHighlightPicker::ApplyAllRequested)
+			]
+			// Cancel button
+			+ SHorizontalBox::Slot()
+			.Padding(5, 0, 0, 0)
+			[
+				SNew(SNiagaraDetailsButton)
+				.Text(LOCTEXT("Cancel", "Cancel"))
+				.OnClicked(this, &SNiagaraScriptHighlightPicker::AddOrEditCancelled)
+			];
+	}
+
+	TSharedRef<SWidget> ConstructEditErrorWidgets()
+	{
+		return SNew(SHorizontalBox)
+			.Visibility(this, &SNiagaraScriptHighlightPicker::GetErrorVisibility)
+			+ SHorizontalBox::Slot()
+			.AutoWidth()
+			.Padding(0, 0, 5, 0)
+			.VAlign(VAlign_Center)
+			[
+				SNew(SImage)
+				.Image(FEditorStyle::GetBrush("Icons.Error"))
+			]
+			+ SHorizontalBox::Slot()
+			[
+				SNew(STextBlock)
+				.Text(this, &SNiagaraScriptHighlightPicker::GetErrorText)
+				.AutoWrapText(true)
+			];
+	}
+
+	FNiagaraScriptHighlight GetCurrentPropertyValue()
+	{
+		FNiagaraScriptHighlight CurrentPropertyValue;
+		void* PropertyValuePtr;
+		if (PropertyHandle->GetValueData(PropertyValuePtr) == FPropertyAccess::Success)
+		{
+			CurrentPropertyValue = *(FNiagaraScriptHighlight*)PropertyValuePtr;
+		}
+		return CurrentPropertyValue;
+	}
+
+	void RefreshDisplayedValue()
+	{
+		DisplayedHighlight = GetCurrentPropertyValue();
+	}
+
+	EVisibility GetPickerVisibility() const
+	{
+		return EditMode == EEditMode::Picker ? EVisibility::Visible : EVisibility::Collapsed;
+	}
+
+	EVisibility GetEditVisibility() const
+	{
+		return EditMode == EEditMode::EditCurrent || EditMode == EEditMode::EditNew ? EVisibility::Visible : EVisibility::Collapsed;
+	}
+
+	FText GetEditModeText() const
+	{
+		if (EditMode == EEditMode::EditNew)
+		{
+			return LOCTEXT("EditNewLabel", "Add a new highlight value");
+		}
+		else if (EditMode == EEditMode::EditCurrent)
+		{
+			return LOCTEXT("EditCurrentLabel", "Edit an existing highlight value");
+		}
+		else
+		{
+			return FText();
+		}
+	}
+
+	void DeleteItem()
+	{
+		TSharedPtr<IPropertyHandleArray> ParentArray = PropertyHandle->GetParentHandle()->AsArray();
+		if (ParentArray.IsValid())
+		{
+			int32 PropertyIndex = PropertyHandle->GetIndexInArray();
+			if (PropertyIndex != INDEX_NONE)
+			{
+				ParentArray->DeleteItem(PropertyHandle->GetIndexInArray());
+			}
+		}
+	}
+
+	typedef SItemSelector<FText, FNiagaraScriptHighlight> SNiagaraScriptHighlightSelector;
+
+	TSharedRef<SWidget> OnGetMenuContent()
+	{
+		return SNew(SBox)
+			.MaxDesiredHeight(200)
+			[
+				SNew(SVerticalBox)
+				// Existing Items
+				+ SVerticalBox::Slot()
+				.Padding(5, 10, 5, 5)
+				[
+					SNew(SNiagaraScriptHighlightSelector)
+					.Items(FNiagaraEditorModule::Get().GetCachedScriptAssetHighlights())
+					.ClickActivateMode(EItemSelectorClickActivateMode::SingleClick)
+					.OnCompareItemsForSorting(this, &SNiagaraScriptHighlightPicker::CompareHighlights)
+					.OnDoesItemMatchFilterText(this, &SNiagaraScriptHighlightPicker::DoesHightlightMatchFilterText)
+					.OnGenerateWidgetForItem(this, &SNiagaraScriptHighlightPicker::ConstructHighlightWidget)
+					.OnItemActivated(this, &SNiagaraScriptHighlightPicker::HighlightActivated)
+				]
+				// Add and edit buttons
+				+ SVerticalBox::Slot()
+				.AutoHeight()
+				.Padding(5, 0, 5, 5)
+				.HAlign(HAlign_Right)
+				[
+					SNew(SHorizontalBox)
+					+ SHorizontalBox::Slot()
+					.AutoWidth()
+					.Padding(0, 0, 5, 0)
+					[
+						SNew(SNiagaraDetailsButton)
+						.Text(LOCTEXT("AddNewText", "Add New"))
+						.ToolTipText(LOCTEXT("AddNewToolTipText", "Add a new highlight."))
+						.OnClicked(this, &SNiagaraScriptHighlightPicker::RequestNewHighlight)
+					]
+					+ SHorizontalBox::Slot()
+					.AutoWidth()
+					[
+						SNew(SNiagaraDetailsButton)
+						.Text(LOCTEXT("EditCurrentText", "Edit"))
+						.ToolTipText(LOCTEXT("EditCurrentToolTipText", "Edit the current highlight."))
+						.OnClicked(this, &SNiagaraScriptHighlightPicker::RequestEditHighlight)
+					]
+				]
+			];
+	}
+
+	bool CompareHighlights(const FNiagaraScriptHighlight& HighlightA, const FNiagaraScriptHighlight& HighlightB) const
+	{
+		return HighlightA.DisplayName.CompareTo(HighlightB.DisplayName) > 0;
+	}
+
+	bool DoesHightlightMatchFilterText(const FText& FilterText, const FNiagaraScriptHighlight& Highlight)
+	{
+		return Highlight.DisplayName.ToString().Contains(FilterText.ToString());
+	}
+
+	TSharedRef<SWidget> ConstructHighlightWidget(const FNiagaraScriptHighlight& Highlight) const
+	{
+		return SNew(SHorizontalBox)
+			+ SHorizontalBox::Slot()
+			.AutoWidth()
+			.Padding(0, 5, 5, 5)
+			[
+				SNew(SBorder)
+				.BorderImage(FEditorStyle::GetBrush("WhiteBrush"))
+				.BorderBackgroundColor(Highlight.Color)
+				.Padding(0)
+				[
+					SNew(SBox)
+					.WidthOverride(DisplayColorWidth)
+					.HeightOverride(DisplayColorHeight)
+				]
+			]
+			+ SHorizontalBox::Slot()
+			.VAlign(VAlign_Center)
+			[
+				SNew(STextBlock)
+				.Text(Highlight.DisplayName)
+			];
+	}
+
+	void HighlightActivated(const FNiagaraScriptHighlight& ActivatedHighlight)
+	{
+		HighlightPickerCombo->SetIsOpen(false);
+		ApplyHighlightChange(ActivatedHighlight);
+	}
+
+	void ApplyHighlightChange(const FNiagaraScriptHighlight& NewHighlight)
+	{
+		void* PropertyValuePtr;
+		if (PropertyHandle->GetValueData(PropertyValuePtr) == FPropertyAccess::Success)
+		{
+			TArray<UObject*> OuterObjects;
+			PropertyHandle->GetOuterObjects(OuterObjects);
+			FScopedTransaction EditHightlightTransaction(LOCTEXT("EditHighlight", "Edit Highlight"));
+			for (UObject* OuterObject : OuterObjects)
+			{
+				OuterObject->Modify();
+			}
+			PropertyHandle->NotifyPreChange();
+			FNiagaraScriptHighlight* CurrentScriptHighlight = (FNiagaraScriptHighlight*)PropertyValuePtr;
+			*CurrentScriptHighlight = NewHighlight;
+			DisplayedHighlight = NewHighlight;
+			PropertyHandle->NotifyPostChange(EPropertyChangeType::ValueSet);
+			PropertyHandle->NotifyFinishedChangingProperties();
+		}
+	}
+
+	FReply RequestNewHighlight()
+	{
+		EditMode = EEditMode::EditNew;
+		HighlightPickerCombo->SetIsOpen(false);
+		return FReply::Handled();
+	}
+
+	FReply RequestEditHighlight()
+	{
+		EditMode = EEditMode::EditCurrent;
+		HighlightPickerCombo->SetIsOpen(false);
+		return FReply::Handled();
+	}
+
+	FSlateColor GetHighlightSlateColor() const
+	{
+		return DisplayedHighlight.Color;
+	}
+
+	FLinearColor GetHighlightColor() const
+	{
+		return DisplayedHighlight.Color;
+	}
+
+	FText GetHighlightText() const
+	{
+		return DisplayedHighlight.DisplayName;
+	}
+
+	void HighlightTextCommitted(const FText& Text, ETextCommit::Type CommitType)
+	{
+		DisplayedHighlight.DisplayName = Text;
+	}
+
+	FReply OnMouseButtonDownColorBlock(const FGeometry& MyGeometry, const FPointerEvent& MouseEvent)
+	{
+		if (MouseEvent.GetEffectingButton() != EKeys::LeftMouseButton)
+		{
+			return FReply::Unhandled();
+		}
+
+		FColorPickerArgs PickerArgs;
+		{
+			PickerArgs.bUseAlpha = true;
+			PickerArgs.bOnlyRefreshOnMouseUp = false;
+			PickerArgs.bOnlyRefreshOnOk = false;
+			PickerArgs.DisplayGamma = TAttribute<float>::Create(TAttribute<float>::FGetter::CreateUObject(GEngine, &UEngine::GetDisplayGamma));
+			PickerArgs.OnColorCommitted = FOnLinearColorValueChanged::CreateSP(this, &SNiagaraScriptHighlightPicker::HighlightColorCommitted);
+			PickerArgs.OnColorPickerCancelled = FOnColorPickerCancelled::CreateSP(this, &SNiagaraScriptHighlightPicker::HighlightColorCancelled);
+			PickerArgs.InitialColorOverride = DisplayedHighlight.Color;
+			PickerArgs.ParentWidget = ColorBlock;
+		}
+
+		OpenColorPicker(PickerArgs);
+		return FReply::Handled();
+	}
+
+	void HighlightColorCommitted(FLinearColor CommittedColor)
+	{
+		DisplayedHighlight.Color = CommittedColor;
+	}
+
+	void HighlightColorCancelled(FLinearColor OriginalColor)
+	{
+		DisplayedHighlight.Color = OriginalColor;
+	}
+
+	EVisibility GetAddNewButtonVisibility() const
+	{
+		return EditMode == EEditMode::EditNew ? EVisibility::Visible : EVisibility::Collapsed;
+	}
+
+	EVisibility GetApplyButtonVisibility() const
+	{
+		return EditMode == EEditMode::EditCurrent ? EVisibility::Visible : EVisibility::Collapsed;
+	}
+
+	bool ValidateNewHighlight(const FNiagaraScriptHighlight& HighlightToCheck, TOptional<FNiagaraScriptHighlight> CurrentHightlightToIgnore, TOptional<FNiagaraScriptHighlight>& OutFoundExistingHighlight, FText& OutErrorMessage)
+	{
+		if (HighlightToCheck.DisplayName.IsEmptyOrWhitespace())
+		{
+			OutErrorMessage = LOCTEXT("HighlightWithEmptyDisplay", "Display name can not be empty or white space.");
+			return false;
+		}
+
+		for (const FNiagaraScriptHighlight& ExistingHighlight : FNiagaraEditorModule::Get().GetCachedScriptAssetHighlights())
+		{
+			if (CurrentHightlightToIgnore.IsSet() && ExistingHighlight == CurrentHightlightToIgnore)
+			{
+				continue;
+			}
+
+			bool bColorMatches = ExistingHighlight.Color.Equals(HighlightToCheck.Color);
+			bool bDisplayNameMatches = ExistingHighlight.DisplayName.CompareTo(HighlightToCheck.DisplayName) == 0;
+			if (bColorMatches && bDisplayNameMatches)
+			{
+				// A match was found so return it so the color will match exactly.
+				OutFoundExistingHighlight = ExistingHighlight;
+				return true;
+			}
+			if (bColorMatches && bDisplayNameMatches == false)
+			{
+				// Only the color or display name matched which is invalid.
+				OutErrorMessage = LOCTEXT("HightlightMismatchDisplayName", "This new highlight matches an existing highlight's color but has a different display name.");
+				return false;
+			}
+			if (bColorMatches == false && bDisplayNameMatches)
+			{
+				// Only the color or display name matched which is invalid.
+				OutErrorMessage = LOCTEXT("HightlightMismatchColor", "This new highlight matches an existing highlight's display name but has a different color.");
+				return false;
+			}
+		}
+		// No matches were found which is also valid.
+		OutFoundExistingHighlight.Reset();
+		return true;
+	}
+
+	FReply AddNewRequested()
+	{
+		TOptional<FNiagaraScriptHighlight> ExistingHighlight;
+		FText ValidateErrorMessage;
+		if (ValidateNewHighlight(DisplayedHighlight, TOptional<FNiagaraScriptHighlight>(), ExistingHighlight, ValidateErrorMessage))
+		{
+			ApplyHighlightChange(ExistingHighlight.IsSet() ? ExistingHighlight.GetValue() : DisplayedHighlight);
+			EditMode = EEditMode::Picker;
+			bShowError = false;
+			ErrorMessage = FText();
+		}
+		else
+		{
+			bShowError = true;
+			ErrorMessage = ValidateErrorMessage;
+		}
+		return FReply::Handled();
+	}
+
+	FReply ApplyCurrentRequested()
+	{
+		TOptional<FNiagaraScriptHighlight> ExistingHighlight;
+		FText ValidateErrorMessage;
+		if (ValidateNewHighlight(DisplayedHighlight, TOptional<FNiagaraScriptHighlight>(), ExistingHighlight, ValidateErrorMessage))
+		{
+			ApplyHighlightChange(ExistingHighlight.IsSet() ? ExistingHighlight.GetValue() : DisplayedHighlight);
+			EditMode = EEditMode::Picker;
+			bShowError = false;
+			ErrorMessage = FText();
+		}
+		else
+		{
+			bShowError = true;
+			ErrorMessage = ValidateErrorMessage;
+		}
+		return FReply::Handled();
+	}
+
+	FReply ApplyAllRequested()
+	{
+		TOptional<FNiagaraScriptHighlight> ExistingHighlight;
+		FText ValidateErrorMessage;
+		FNiagaraScriptHighlight CurrentHighlight = GetCurrentPropertyValue();
+		if (ValidateNewHighlight(DisplayedHighlight, CurrentHighlight, ExistingHighlight, ValidateErrorMessage))
+		{
+			TArray<FAssetData> MatchingScriptAssets;
+			FNiagaraEditorModule::Get().GetScriptAssetsMatchingHighlight(CurrentHighlight, MatchingScriptAssets);
+
+			FScopedTransaction ApplyAllTransaction(LOCTEXT("ApplyAllTransaction", "Apply highlight change to all scripts."));
+
+			FNiagaraScriptHighlight NewHighlight = ExistingHighlight.IsSet() ? ExistingHighlight.GetValue() : DisplayedHighlight;
+			ApplyHighlightChange(NewHighlight);
+
+			for (const FAssetData& MatchingScriptAsset : MatchingScriptAssets)
+			{
+				UNiagaraScript* Script = CastChecked<UNiagaraScript>(MatchingScriptAsset.GetAsset());
+				Script->Modify();
+				for (FNiagaraScriptHighlight& Highlight : Script->Highlights)
+				{
+					if (Highlight == CurrentHighlight)
+					{
+						Highlight = NewHighlight;
+					}
+				}
+			}
+
+			EditMode = EEditMode::Picker;
+			bShowError = false;
+			ErrorMessage = FText();
+		}
+		else
+		{
+			bShowError = true;
+			ErrorMessage = ValidateErrorMessage;
+		}
+		return FReply::Handled();
+	}
+
+	FReply AddOrEditCancelled()
+	{
+		EditMode = EEditMode::Picker;
+		bShowError = false;
+		ErrorMessage = FText();
+		RefreshDisplayedValue();
+		return FReply::Handled();
+	}
+
+	EVisibility GetErrorVisibility() const
+	{
+		return bShowError ? EVisibility::Visible : EVisibility::Collapsed;
+	}
+
+	FText GetErrorText() const
+	{
+		return ErrorMessage;
+	}
+
+private:
+	TSharedPtr<IPropertyHandle> PropertyHandle;
+	FNiagaraScriptHighlight DisplayedHighlight;
+	EEditMode EditMode;
+	TSharedPtr<SComboButton> HighlightPickerCombo;
+	TSharedPtr<SColorBlock> ColorBlock;
+	bool bShowError;
+	FText ErrorMessage;
+};
+
+void FNiagaraScriptHighlightDetails::CustomizeHeader(TSharedRef<IPropertyHandle> PropertyHandle, FDetailWidgetRow& HeaderRow, IPropertyTypeCustomizationUtils& CustomizationUtils)
+{
+	//HeaderRow.ValueContent()
+
+	HeaderRow.WholeRowContent()
+		[
+			SNew(SNiagaraScriptHighlightPicker, PropertyHandle)
+		];
+}
+
 
 #undef LOCTEXT_NAMESPACE
diff --git a/Engine/Plugins/FX/Niagara/Source/NiagaraEditor/Private/Customizations/NiagaraScriptDetails.h b/Engine/Plugins/FX/Niagara/Source/NiagaraEditor/Private/Customizations/NiagaraScriptDetails.h
--- a/Engine/Plugins/FX/Niagara/Source/NiagaraEditor/Private/Customizations/NiagaraScriptDetails.h
+++ b/Engine/Plugins/FX/Niagara/Source/NiagaraEditor/Private/Customizations/NiagaraScriptDetails.h
@@ -1,4 +1,4 @@
-﻿// Copyright 1998-2019 Epic Games, Inc. All Rights Reserved.
+// Copyright 1998-2019 Epic Games, Inc. All Rights Reserved.
 
 #pragma once
 
@@ -26,3 +26,16 @@ private:
 	TSharedPtr<FNiagaraMetaDataCustomNodeBuilder> MetaDataBuilder;
 };
 
+class FNiagaraScriptHighlightDetails : public IPropertyTypeCustomization
+{
+public:
+	static TSharedRef<IPropertyTypeCustomization> MakeInstance()
+	{
+		return MakeShared<FNiagaraScriptHighlightDetails>();
+	}
+
+	virtual void CustomizeHeader(TSharedRef<IPropertyHandle> PropertyHandle, FDetailWidgetRow& HeaderRow, IPropertyTypeCustomizationUtils& CustomizationUtils) override;
+
+	virtual void CustomizeChildren(TSharedRef<IPropertyHandle> PropertyHandle, IDetailChildrenBuilder& ChildBuilder, IPropertyTypeCustomizationUtils& CustomizationUtils) override { }
+};
+
diff --git a/Engine/Plugins/FX/Niagara/Source/NiagaraEditor/Private/NiagaraEditorModule.cpp b/Engine/Plugins/FX/Niagara/Source/NiagaraEditor/Private/NiagaraEditorModule.cpp
--- a/Engine/Plugins/FX/Niagara/Source/NiagaraEditor/Private/NiagaraEditorModule.cpp
+++ b/Engine/Plugins/FX/Niagara/Source/NiagaraEditor/Private/NiagaraEditorModule.cpp
@@ -1057,6 +1057,92 @@ const FNiagaraEditorCommands& FNiagaraEditorModule::GetCommands() const
 	return FNiagaraEditorCommands::Get();
 }
 
+void FNiagaraEditorModule::InvalidateCachedScriptAssetData()
+{
+	CachedScriptAssetHighlights.Reset();
+}
+
+const TArray<FNiagaraScriptHighlight>& FNiagaraEditorModule::GetCachedScriptAssetHighlights() const
+{
+	if (CachedScriptAssetHighlights.IsSet() == false)
+	{
+		CachedScriptAssetHighlights = TArray<FNiagaraScriptHighlight>();
+		FAssetRegistryModule& AssetRegistryModule = FModuleManager::LoadModuleChecked<FAssetRegistryModule>("AssetRegistry");
+		TArray<FAssetData> ScriptAssets;
+		AssetRegistryModule.Get().GetAssetsByClass(UNiagaraScript::StaticClass()->GetFName(), ScriptAssets);
+		for (const FAssetData& ScriptAsset : ScriptAssets)
+		{
+			if (ScriptAsset.IsAssetLoaded())
+			{
+				UNiagaraScript* Script = CastChecked<UNiagaraScript>(ScriptAsset.GetAsset());
+				for (const FNiagaraScriptHighlight& Highlight : Script->Highlights)
+				{
+					if (Highlight.IsValid())
+					{
+						CachedScriptAssetHighlights->AddUnique(Highlight);
+					}
+				}
+			}
+			else
+			{
+				FString HighlightsString;
+				if (ScriptAsset.GetTagValue(GET_MEMBER_NAME_CHECKED(UNiagaraScript, Highlights), HighlightsString))
+				{
+					TArray<FNiagaraScriptHighlight> Highlights;
+					FNiagaraScriptHighlight::JsonToArray(HighlightsString, Highlights);
+					for (const FNiagaraScriptHighlight& Highlight : Highlights)
+					{
+						if (Highlight.IsValid())
+						{
+							CachedScriptAssetHighlights->AddUnique(Highlight);
+						}
+					}
+				}
+			}
+		}
+	}
+	return CachedScriptAssetHighlights.GetValue();
+}
+
+void FNiagaraEditorModule::GetScriptAssetsMatchingHighlight(const FNiagaraScriptHighlight& InHighlight, TArray<FAssetData>& OutMatchingScriptAssets) const
+{
+	FAssetRegistryModule& AssetRegistryModule = FModuleManager::LoadModuleChecked<FAssetRegistryModule>("AssetRegistry");
+	TArray<FAssetData> ScriptAssets;
+	AssetRegistryModule.Get().GetAssetsByClass(UNiagaraScript::StaticClass()->GetFName(), ScriptAssets);
+	for (const FAssetData& ScriptAsset : ScriptAssets)
+	{
+		if (ScriptAsset.IsAssetLoaded())
+		{
+			UNiagaraScript* Script = CastChecked<UNiagaraScript>(ScriptAsset.GetAsset());
+			for (const FNiagaraScriptHighlight& Highlight : Script->Highlights)
+			{
+				if (Highlight == InHighlight)
+				{
+					OutMatchingScriptAssets.Add(ScriptAsset);
+					break;
+				}
+			}
+		}
+		else
+		{
+			FString HighlightsString;
+			if (ScriptAsset.GetTagValue(GET_MEMBER_NAME_CHECKED(UNiagaraScript, Highlights), HighlightsString))
+			{
+				TArray<FNiagaraScriptHighlight> Highlights;
+				FNiagaraScriptHighlight::JsonToArray(HighlightsString, Highlights);
+				for (const FNiagaraScriptHighlight& Highlight : Highlights)
+				{
+					if (Highlight == InHighlight)
+					{
+						OutMatchingScriptAssets.Add(ScriptAsset);
+						break;
+					}
+				}
+			}
+		}
+	}
+}
+
 void FNiagaraEditorModule::RegisterAssetTypeAction(IAssetTools& AssetTools, TSharedRef<IAssetTypeActions> Action)
 {
 	AssetTools.RegisterAssetTypeActions(Action);
diff --git a/Engine/Plugins/FX/Niagara/Source/NiagaraEditor/Private/NiagaraEditorStyle.cpp b/Engine/Plugins/FX/Niagara/Source/NiagaraEditor/Private/NiagaraEditorStyle.cpp
--- a/Engine/Plugins/FX/Niagara/Source/NiagaraEditor/Private/NiagaraEditorStyle.cpp
+++ b/Engine/Plugins/FX/Niagara/Source/NiagaraEditor/Private/NiagaraEditorStyle.cpp
@@ -148,6 +148,14 @@ TSharedRef< FSlateStyleSet > FNiagaraEditorStyle::Create()
 	
 	Style->Set("NiagaraEditor.SubduedHeadingTextBox", SubduedHeadingText);
 
+	// Details
+	FTextBlockStyle DetailsHeadingText = FTextBlockStyle(NormalText)
+		.SetFont(DEFAULT_FONT("Bold", 9))
+		.SetShadowOffset(FVector2D(0, 1))
+		.SetShadowColorAndOpacity(FLinearColor(0, 0, 0, 0.9f));
+
+	Style->Set("NiagaraEditor.DetailsHeadingText", DetailsHeadingText);
+
 	// Parameters
 	FSlateFontInfo ParameterFont = DEFAULT_FONT("Regular", 8);
 
diff --git a/Engine/Plugins/FX/Niagara/Source/NiagaraEditor/Private/SItemSelector.h b/Engine/Plugins/FX/Niagara/Source/NiagaraEditor/Private/SItemSelector.h
--- a/Engine/Plugins/FX/Niagara/Source/NiagaraEditor/Private/SItemSelector.h
+++ b/Engine/Plugins/FX/Niagara/Source/NiagaraEditor/Private/SItemSelector.h
@@ -6,6 +6,12 @@
 #include "Widgets/Input/SSearchBox.h"
 #include "Widgets/Views/STreeView.h"
 
+enum class EItemSelectorClickActivateMode
+{
+	SingleClick,
+	DoubleClick
+};
+
 /** A generic widget for selecting an item from an array of items including optional filtering and categorization. */
 template<typename CategoryType, typename ItemType>
 class SItemSelector : public SCompoundWidget
@@ -23,12 +29,16 @@ public:
 public:
 	SLATE_BEGIN_ARGS(SItemSelector)
 		: _AllowMultiselect(false)
+		, _ClickActivateMode(EItemSelectorClickActivateMode::DoubleClick)
 	{}
 
 	SLATE_ARGUMENT(TArray<ItemType>, Items)
 		/** Whether or not this item selector should allow multiple items to be selected. */
 		SLATE_ARGUMENT(bool, AllowMultiselect)
 
+		/** Whether or not a single click activates an item. */
+		SLATE_ARGUMENT(EItemSelectorClickActivateMode, ClickActivateMode)
+
 		/** An optional delegate to get an array of categories for the specified item. Each category in the returned array represents one level of nested categories. 
 		NOTE: The OnCompareCategoriesForEquality, and OnGenerateWidgetForCategory delegates must be bound if this delegate is bound. */
 		SLATE_EVENT(FOnGetCategoriesForItem, OnGetCategoriesForItem)
@@ -388,6 +398,7 @@ public:
 	void Construct(const FArguments& InArgs)
 	{
 		Items = InArgs._Items;
+		ClickActivateMode = InArgs._ClickActivateMode;
 		OnGetCategoriesForItem = InArgs._OnGetCategoriesForItem;
 		OnCompareCategoriesForEquality = InArgs._OnCompareCategoriesForEquality;
 		OnCompareCategoriesForSorting = InArgs._OnCompareCategoriesForSorting;
@@ -421,6 +432,7 @@ public:
 				.SelectionMode(InArgs._AllowMultiselect ? ESelectionMode::Multi : ESelectionMode::SingleToggle)
 				.OnGenerateRow(this, &SItemSelector::OnGenerateRow)
 				.OnGetChildren(this, &SItemSelector::OnGetChildren)
+				.OnMouseButtonClick(this, &SItemSelector::OnMouseClick)
 				.OnMouseButtonDoubleClick(this, &SItemSelector::OnMouseDoubleClick)
 				.TreeItemsSource(ViewModel->GetRootItems())
 			]
@@ -510,9 +522,18 @@ private:
 		Item->GetChildren(OutChildren);
 	}
 
+	void OnMouseClick(TSharedRef<FItemSelectorItemViewModel> ItemClicked)
+	{
+		if (ClickActivateMode == EItemSelectorClickActivateMode::SingleClick && OnItemActivated.IsBound() && ItemClicked->GetType() == EItemSelectorItemViewModelType::Item)
+		{
+			TSharedRef<FItemSelectorItemContainerViewModel> ItemContainer = StaticCastSharedRef<FItemSelectorItemContainerViewModel>(ItemClicked);
+			OnItemActivated.Execute(ItemContainer->GetItem());
+		}
+	}
+
 	void OnMouseDoubleClick(TSharedRef<FItemSelectorItemViewModel> ItemDoubleClicked)
 	{
-		if (OnItemActivated.IsBound() && ItemDoubleClicked->GetType() == EItemSelectorItemViewModelType::Item)
+		if (ClickActivateMode == EItemSelectorClickActivateMode::DoubleClick && OnItemActivated.IsBound() && ItemDoubleClicked->GetType() == EItemSelectorItemViewModelType::Item)
 		{
 			TSharedRef<FItemSelectorItemContainerViewModel> ItemContainer = StaticCastSharedRef<FItemSelectorItemContainerViewModel>(ItemDoubleClicked);
 			OnItemActivated.Execute(ItemContainer->GetItem());
@@ -535,6 +556,8 @@ private:
 private:
 	TArray<ItemType> Items;
 
+	EItemSelectorClickActivateMode ClickActivateMode;
+
 	FOnGetCategoriesForItem OnGetCategoriesForItem;
 	FOnCompareCategoriesForEquality OnCompareCategoriesForEquality;
 	FOnCompareCategoriesForSorting OnCompareCategoriesForSorting;
diff --git a/Engine/Plugins/FX/Niagara/Source/NiagaraEditor/Private/Toolkits/NiagaraScriptToolkit.cpp b/Engine/Plugins/FX/Niagara/Source/NiagaraEditor/Private/Toolkits/NiagaraScriptToolkit.cpp
--- a/Engine/Plugins/FX/Niagara/Source/NiagaraEditor/Private/Toolkits/NiagaraScriptToolkit.cpp
+++ b/Engine/Plugins/FX/Niagara/Source/NiagaraEditor/Private/Toolkits/NiagaraScriptToolkit.cpp
@@ -343,9 +343,11 @@ TSharedRef<SDockTab> FNiagaraScriptToolkit::SpawnTabScriptDetails(const FSpawnTa
 	DetailsView->OnFinishedChangingProperties().AddRaw(this, &FNiagaraScriptToolkit::OnEditedScriptPropertyFinishedChanging);
 	DetailsView->RegisterInstancedCustomPropertyLayout(
 		UNiagaraScript::StaticClass(),
-		FOnGetDetailCustomizationInstance::CreateStatic(&FNiagaraScriptDetails::MakeInstance, ScriptViewModelWeakPtr)
-	);
-
+		FOnGetDetailCustomizationInstance::CreateStatic(&FNiagaraScriptDetails::MakeInstance, ScriptViewModelWeakPtr));
+	DetailsView->RegisterInstancedCustomPropertyTypeLayout(
+		FNiagaraScriptHighlight::StaticStruct()->GetFName(),
+		FOnGetPropertyTypeCustomizationInstance::CreateStatic(&FNiagaraScriptHighlightDetails::MakeInstance));
+	
 	DetailsView->SetObjects(DetailsScriptSelection->GetSelectedObjects().Array());
 
 	return SNew(SDockTab)
@@ -710,6 +712,7 @@ void FNiagaraScriptToolkit::UpdateOriginalNiagaraScript()
 
 	GWarn->EndSlowTask();
 	bEditedScriptHasPendingChanges = false;
+	FNiagaraEditorModule::Get().InvalidateCachedScriptAssetData();
 }
 
 
diff --git a/Engine/Plugins/FX/Niagara/Source/NiagaraEditor/Private/ViewModels/Stack/NiagaraStackItem.cpp b/Engine/Plugins/FX/Niagara/Source/NiagaraEditor/Private/ViewModels/Stack/NiagaraStackItem.cpp
--- a/Engine/Plugins/FX/Niagara/Source/NiagaraEditor/Private/ViewModels/Stack/NiagaraStackItem.cpp
+++ b/Engine/Plugins/FX/Niagara/Source/NiagaraEditor/Private/ViewModels/Stack/NiagaraStackItem.cpp
@@ -9,6 +9,8 @@
 #include "ViewModels/NiagaraSystemSelectionViewModel.h"
 #include "NiagaraStackEditorData.h"
 
+#include "EditorStyleSet.h"
+
 void UNiagaraStackItem::Initialize(FRequiredEntryData InRequiredEntryData, FString InStackEditorDataKey)
 {
 	Super::Initialize(InRequiredEntryData, InStackEditorDataKey);
@@ -43,6 +45,17 @@ void UNiagaraStackItem::Delete()
 	DeleteInternal();
 }
 
+const TArray<FNiagaraScriptHighlight>& UNiagaraStackItem::GetHighlights() const
+{
+	static TArray<FNiagaraScriptHighlight> Empty;
+	return Empty;
+}
+
+const FSlateBrush* UNiagaraStackItem::GetIconBrush() const
+{
+	return FEditorStyle::GetBrush("NoBrush");
+}
+
 void UNiagaraStackItem::RefreshChildrenInternal(const TArray<UNiagaraStackEntry*>& CurrentChildren, TArray<UNiagaraStackEntry*>& NewChildren, TArray<FStackIssue>& NewIssues)
 {
 	if (ItemFooter == nullptr)
diff --git a/Engine/Plugins/FX/Niagara/Source/NiagaraEditor/Private/ViewModels/Stack/NiagaraStackModuleItem.cpp b/Engine/Plugins/FX/Niagara/Source/NiagaraEditor/Private/ViewModels/Stack/NiagaraStackModuleItem.cpp
--- a/Engine/Plugins/FX/Niagara/Source/NiagaraEditor/Private/ViewModels/Stack/NiagaraStackModuleItem.cpp
+++ b/Engine/Plugins/FX/Niagara/Source/NiagaraEditor/Private/ViewModels/Stack/NiagaraStackModuleItem.cpp
@@ -878,6 +878,16 @@ void UNiagaraStackModuleItem::DeleteInternal()
 	}
 }
 
+bool UNiagaraStackModuleItem::SupportsHighlights() const
+{
+	return FunctionCallNode != nullptr && FunctionCallNode->FunctionScript != nullptr;
+}
+
+const TArray<FNiagaraScriptHighlight>& UNiagaraStackModuleItem::GetHighlights() const
+{
+	return FunctionCallNode->FunctionScript->Highlights;
+}
+
 int32 UNiagaraStackModuleItem::GetModuleIndex() const
 {
 	TArray<FNiagaraStackGraphUtilities::FStackNodeGroup> StackGroups;
diff --git a/Engine/Plugins/FX/Niagara/Source/NiagaraEditor/Private/ViewModels/Stack/NiagaraStackRendererItem.cpp b/Engine/Plugins/FX/Niagara/Source/NiagaraEditor/Private/ViewModels/Stack/NiagaraStackRendererItem.cpp
--- a/Engine/Plugins/FX/Niagara/Source/NiagaraEditor/Private/ViewModels/Stack/NiagaraStackRendererItem.cpp
+++ b/Engine/Plugins/FX/Niagara/Source/NiagaraEditor/Private/ViewModels/Stack/NiagaraStackRendererItem.cpp
@@ -25,6 +25,8 @@
 #include "ViewModels/Stack/NiagaraStackGraphUtilities.h"
 #include "NiagaraScriptMergeManager.h"
 
+#include "Styling/SlateIconFinder.h"
+
 #define LOCTEXT_NAMESPACE "UNiagaraStackRendererItem"
 
 UNiagaraStackRendererItem::UNiagaraStackRendererItem()
@@ -230,6 +232,11 @@ void UNiagaraStackRendererItem::SetIsEnabledInternal(bool bInIsEnabled)
 	RefreshChildren();
 }
 
+const FSlateBrush* UNiagaraStackRendererItem::GetIconBrush() const
+{
+	return FSlateIconFinder::FindIconBrushForClass(RendererProperties->GetClass());
+}
+
 void UNiagaraStackRendererItem::RefreshChildrenInternal(const TArray<UNiagaraStackEntry*>& CurrentChildren, TArray<UNiagaraStackEntry*>& NewChildren, TArray<FStackIssue>& NewIssues)
 {
 	if (RendererObject == nullptr)
diff --git a/Engine/Plugins/FX/Niagara/Source/NiagaraEditor/Public/NiagaraEditorModule.h b/Engine/Plugins/FX/Niagara/Source/NiagaraEditor/Public/NiagaraEditorModule.h
--- a/Engine/Plugins/FX/Niagara/Source/NiagaraEditor/Public/NiagaraEditorModule.h
+++ b/Engine/Plugins/FX/Niagara/Source/NiagaraEditor/Public/NiagaraEditorModule.h
@@ -25,6 +25,7 @@ class UMovieSceneNiagaraParameterTrack;
 struct IConsoleCommand;
 class INiagaraEditorOnlyDataUtilities;
 class FNiagaraEditorCommands;
+struct FNiagaraScriptHighlight;
 
 DECLARE_STATS_GROUP(TEXT("Niagara Editor"), STATGROUP_NiagaraEditor, STATCAT_Advanced);
 
@@ -106,6 +107,12 @@ public:
 
 	NIAGARAEDITOR_API const FNiagaraEditorCommands& GetCommands() const;
 
+	void InvalidateCachedScriptAssetData();
+
+	const TArray<FNiagaraScriptHighlight>& GetCachedScriptAssetHighlights() const;
+
+	void GetScriptAssetsMatchingHighlight(const FNiagaraScriptHighlight& InHighlight, TArray<FAssetData>& OutMatchingScriptAssets) const;
+
 private:
 	void RegisterAssetTypeAction(IAssetTools& AssetTools, TSharedRef<IAssetTypeActions> Action);
 	void OnNiagaraSettingsChangedEvent(const FString& PropertyName, const UNiagaraSettings* Settings);
@@ -162,5 +169,7 @@ private:
 
 	FOnCheckScriptToolkitsShouldFocusGraphElement OnCheckScriptToolkitsShouldFocusGraphElement;
 
+	mutable TOptional<TArray<FNiagaraScriptHighlight>> CachedScriptAssetHighlights;
+
 	bool bThumbnailRenderersRegistered;
 };
diff --git a/Engine/Plugins/FX/Niagara/Source/NiagaraEditor/Public/ViewModels/Stack/NiagaraStackItem.h b/Engine/Plugins/FX/Niagara/Source/NiagaraEditor/Public/ViewModels/Stack/NiagaraStackItem.h
--- a/Engine/Plugins/FX/Niagara/Source/NiagaraEditor/Public/ViewModels/Stack/NiagaraStackItem.h
+++ b/Engine/Plugins/FX/Niagara/Source/NiagaraEditor/Public/ViewModels/Stack/NiagaraStackItem.h
@@ -3,6 +3,7 @@
 #pragma once
 
 #include "ViewModels/Stack/NiagaraStackEntry.h"
+#include "NiagaraScriptHighlight.h"
 #include "Layout/Visibility.h"
 #include "NiagaraStackItem.generated.h"
 
@@ -30,6 +31,12 @@ public:
 	virtual bool SupportsDelete() const { return false; }
 	virtual bool TestCanDeleteWithMessage(FText& OutCanDeleteMessage) const { return false; }
 	void Delete();
+
+	virtual bool SupportsHighlights() const { return false; }
+	virtual const TArray<FNiagaraScriptHighlight>& GetHighlights() const;
+
+	virtual bool SupportsIcon() const { return false; }
+	virtual const FSlateBrush* GetIconBrush() const;
 	
 protected:
 	virtual void RefreshChildrenInternal(const TArray<UNiagaraStackEntry*>& CurrentChildren, TArray<UNiagaraStackEntry*>& NewChildren, TArray<FStackIssue>& NewIssues) override;
diff --git a/Engine/Plugins/FX/Niagara/Source/NiagaraEditor/Public/ViewModels/Stack/NiagaraStackModuleItem.h b/Engine/Plugins/FX/Niagara/Source/NiagaraEditor/Public/ViewModels/Stack/NiagaraStackModuleItem.h
--- a/Engine/Plugins/FX/Niagara/Source/NiagaraEditor/Public/ViewModels/Stack/NiagaraStackModuleItem.h
+++ b/Engine/Plugins/FX/Niagara/Source/NiagaraEditor/Public/ViewModels/Stack/NiagaraStackModuleItem.h
@@ -42,6 +42,9 @@ public:
 	virtual bool SupportsDelete() const override { return true; }
 	virtual bool TestCanDeleteWithMessage(FText& OutCanDeleteMessage) const override;
 
+	virtual bool SupportsHighlights() const override;
+	virtual const TArray<FNiagaraScriptHighlight>& GetHighlights() const override;
+
 	int32 GetModuleIndex() const;
 	
 	UObject* GetExternalAsset() const override;
diff --git a/Engine/Plugins/FX/Niagara/Source/NiagaraEditor/Public/ViewModels/Stack/NiagaraStackRendererItem.h b/Engine/Plugins/FX/Niagara/Source/NiagaraEditor/Public/ViewModels/Stack/NiagaraStackRendererItem.h
--- a/Engine/Plugins/FX/Niagara/Source/NiagaraEditor/Public/ViewModels/Stack/NiagaraStackRendererItem.h
+++ b/Engine/Plugins/FX/Niagara/Source/NiagaraEditor/Public/ViewModels/Stack/NiagaraStackRendererItem.h
@@ -36,6 +36,9 @@ public:
 	virtual bool SupportsChangeEnabled() const override { return true; }
 	virtual bool GetIsEnabled() const override;
 
+	virtual bool SupportsIcon() const override { return true; }
+	virtual const FSlateBrush* GetIconBrush() const override;
+
 	static TArray<FNiagaraVariable> GetMissingVariables(UNiagaraRendererProperties* RendererProperties, UNiagaraEmitter* Emitter);
 	static bool AddMissingVariable(UNiagaraEmitter* Emitter, const FNiagaraVariable& Variable);
 
diff --git a/Engine/Plugins/FX/Niagara/Source/NiagaraEditorWidgets/Private/NiagaraEditorWidgetsStyle.cpp b/Engine/Plugins/FX/Niagara/Source/NiagaraEditorWidgets/Private/NiagaraEditorWidgetsStyle.cpp
--- a/Engine/Plugins/FX/Niagara/Source/NiagaraEditorWidgets/Private/NiagaraEditorWidgetsStyle.cpp
+++ b/Engine/Plugins/FX/Niagara/Source/NiagaraEditorWidgets/Private/NiagaraEditorWidgetsStyle.cpp
@@ -64,6 +64,8 @@ TSharedRef< FSlateStyleSet > FNiagaraEditorWidgetsStyle::Create()
 	Style->SetContentRoot(FPaths::EngineContentDir() / TEXT("Editor/Slate/Niagara"));
 
 	// Stack
+	Style->Set("NiagaraEditor.Stack.IconSize", FVector2D(18.0f, 18.0f));
+
 	FSlateFontInfo StackGroupFont = DEFAULT_FONT("Bold", 10);
 	FTextBlockStyle StackGroupText = FTextBlockStyle(NormalText)
 		.SetFont(StackGroupFont)
@@ -198,6 +200,8 @@ TSharedRef< FSlateStyleSet > FNiagaraEditorWidgetsStyle::Create()
 		.SetShadowColorAndOpacity(FLinearColor(0, 0, 0, 0.9f));
 	Style->Set("NiagaraEditor.Stack.AddButtonText", AddButtonText);
 
+	Style->Set("NiagaraEditor.Stack.ModuleHighlight",  new IMAGE_PLUGIN_BRUSH("Icons/ModuleHighlight", Icon8x8, FLinearColor::White));
+
 	Style->Set("NiagaraEditor.ShowInCurveEditorIcon", new IMAGE_PLUGIN_BRUSH("Icons/ShowInCurveEditor", Icon16x16, FLinearColor::White));
 
 	return Style;
diff --git a/Engine/Plugins/FX/Niagara/Source/NiagaraEditorWidgets/Private/SNiagaraOverviewStack.cpp b/Engine/Plugins/FX/Niagara/Source/NiagaraEditorWidgets/Private/SNiagaraOverviewStack.cpp
--- a/Engine/Plugins/FX/Niagara/Source/NiagaraEditorWidgets/Private/SNiagaraOverviewStack.cpp
+++ b/Engine/Plugins/FX/Niagara/Source/NiagaraEditorWidgets/Private/SNiagaraOverviewStack.cpp
@@ -19,6 +19,7 @@
 #include "NiagaraEditorCommon.h"
 
 #include "Widgets/Layout/SScrollBox.h"
+#include "Widgets/Layout/SGridPanel.h"
 #include "Widgets/Images/SImage.h"
 #include "Widgets/Input/SButton.h"
 #include "Widgets/Notifications/SNotificationList.h"
@@ -145,11 +146,11 @@ private:
 	{
 		if (IssueIconVisibility.Get() == EVisibility::Visible)
 		{
-			return FMargin(6, 0, 1, 0);
+			return FMargin(5, 0, 1, 0);
 		}
 		else
 		{
-			return FMargin(6, 0, 5, 0);
+			return FMargin(5, 0, 5, 0);
 		}
 	}
 
@@ -378,11 +379,98 @@ void SNiagaraOverviewStack::EntryStructureChanged()
 
 TSharedRef<ITableRow> SNiagaraOverviewStack::OnGenerateRowForEntry(UNiagaraStackEntry* Item, const TSharedRef<STableViewBase>& OwnerTable)
 {
+	FVector2D IconSize = FNiagaraEditorWidgetsStyle::Get().GetVector("NiagaraEditor.Stack.IconSize");
 	TSharedPtr<SWidget> Content;
 	if (Item->IsA<UNiagaraStackItem>())
 	{
 		UNiagaraStackItem* StackItem = CastChecked<UNiagaraStackItem>(Item);
+		TSharedPtr<SWidget> IndentContent;
+		if (StackItem->SupportsHighlights())
+		{
+			TArray<FNiagaraScriptHighlight> ScriptHighlights = StackItem->GetHighlights();
+			Algo::Sort(ScriptHighlights, [](const FNiagaraScriptHighlight& A, const FNiagaraScriptHighlight& B)
+			{
+				return A.DisplayName.CompareTo(B.DisplayName);
+			});
+
+			FText HighlightToolTip;
+			if (ScriptHighlights.Num() > 0)
+			{
+				TArray<FText> DisplayNames;
+				for (const FNiagaraScriptHighlight& ScriptHighlight : ScriptHighlights)
+				{
+					if (ScriptHighlight.IsValid())
+					{
+						DisplayNames.Add(ScriptHighlight.DisplayName);
+					}
+				}
+				HighlightToolTip = FText::Join(LOCTEXT("HighlightDelimiter", ", "), DisplayNames);
+			}
+
+			TSharedRef<SBox> IndentBox = SNew(SBox)
+				.ToolTipText(HighlightToolTip)
+				.IsEnabled_UObject(Item, &UNiagaraStackEntry::GetIsEnabledAndOwnerIsEnabled)
+				.Visibility(EVisibility::Visible)
+				.WidthOverride(IconSize.X)
+				.HeightOverride(IconSize.Y);
+			
+			if (ScriptHighlights.Num() > 0)
+			{
+				TSharedPtr<SGridPanel> HighlightsGrid;
+				IndentBox->SetContent(SAssignNew(HighlightsGrid, SGridPanel)
+					.FillRow(0, .5f)
+					.FillRow(1, .5f)
+					.FillColumn(0, .5f)
+					.FillColumn(1, .5f));
+				int32 HighlightsAdded = 0;
+				for (int32 HighlightIndex = 0; HighlightIndex < ScriptHighlights.Num() && HighlightsAdded < 4; HighlightIndex++)
+				{
+					if (ScriptHighlights[HighlightIndex].IsValid())
+					{
+						int32 Column = HighlightIndex / 2;
+						int32 Row = HighlightIndex % 2;
+						HighlightsGrid->AddSlot(Column, Row)
+							.Padding(Column, Row, 1 - Column, 1 - Row)
+							[
+								SNew(SImage)
+								.Image(FNiagaraEditorWidgetsStyle::Get().GetBrush("NiagaraEditor.Stack.ModuleHighlight"))
+								.ColorAndOpacity(ScriptHighlights[HighlightIndex].Color)
+							];
+						HighlightsAdded++;
+					}
+				}
+			}
+			IndentContent = IndentBox;
+		}
+		else if (StackItem->SupportsIcon())
+		{
+			IndentContent = SNew(SBox)
+				.IsEnabled_UObject(Item, &UNiagaraStackEntry::GetIsEnabledAndOwnerIsEnabled)
+				.WidthOverride(IconSize.X)
+				.HeightOverride(IconSize.Y)
+				.HAlign(HAlign_Center)
+				.VAlign(VAlign_Center)
+				[
+					SNew(SImage)
+					.Image(StackItem->GetIconBrush())
+				];
+		}
+		else
+		{
+			IndentContent = SNew(SBox)
+				.WidthOverride(IconSize.X)
+				.HeightOverride(IconSize.Y);
+		}
+
 		Content = SNew(SHorizontalBox)
+			// Indent content
+			+ SHorizontalBox::Slot()
+			.Padding(0, 1, 2, 1)
+			.AutoWidth()
+			.VAlign(VAlign_Center)
+			[
+				IndentContent.ToSharedRef()
+			]
 			// Name
 			+ SHorizontalBox::Slot()
 			.VAlign(VAlign_Center)
@@ -416,10 +504,17 @@ TSharedRef<ITableRow> SNiagaraOverviewStack::OnGenerateRowForEntry(UNiagaraStack
 			.VAlign(VAlign_Center)
 			.AutoWidth()
 			[
-				SNew(SImage)
-				.Image(FNiagaraEditorWidgetsStyle::Get().GetBrush(FNiagaraStackEditorWidgetsUtilities::GetIconNameForExecutionSubcategory(Item->GetExecutionSubcategoryName(), true)))
-				.ColorAndOpacity(FNiagaraEditorWidgetsStyle::Get().GetColor(FNiagaraStackEditorWidgetsUtilities::GetIconColorNameForExecutionCategory(Item->GetExecutionCategoryName())))
-				.IsEnabled_UObject(Item, &UNiagaraStackEntry::GetIsEnabledAndOwnerIsEnabled)
+				SNew(SBox)
+				.WidthOverride(IconSize.X)
+				.HeightOverride(IconSize.Y)
+				.HAlign(HAlign_Center)
+				.VAlign(VAlign_Center)
+				[
+					SNew(SImage)
+					.Image(FNiagaraEditorWidgetsStyle::Get().GetBrush(FNiagaraStackEditorWidgetsUtilities::GetIconNameForExecutionSubcategory(Item->GetExecutionSubcategoryName(), true)))
+					.ColorAndOpacity(FNiagaraEditorWidgetsStyle::Get().GetColor(FNiagaraStackEditorWidgetsUtilities::GetIconColorNameForExecutionCategory(Item->GetExecutionCategoryName())))
+					.IsEnabled_UObject(Item, &UNiagaraStackEntry::GetIsEnabledAndOwnerIsEnabled)
+				]
 			]
 			// Name
 			+ SHorizontalBox::Slot()
diff --git a/Engine/Plugins/FX/Niagara/Source/NiagaraEditorWidgets/Private/Stack/SNiagaraStackRendererItem.cpp b/Engine/Plugins/FX/Niagara/Source/NiagaraEditorWidgets/Private/Stack/SNiagaraStackRendererItem.cpp
--- a/Engine/Plugins/FX/Niagara/Source/NiagaraEditorWidgets/Private/Stack/SNiagaraStackRendererItem.cpp
+++ b/Engine/Plugins/FX/Niagara/Source/NiagaraEditorWidgets/Private/Stack/SNiagaraStackRendererItem.cpp
@@ -12,7 +12,6 @@
 #include "Widgets/Input/SCheckBox.h"
 #include "Widgets/Images/SImage.h"
 #include "Widgets/Text/STextBlock.h"
-#include "Styling/SlateIconFinder.h"
 #include "Widgets/Images/SImage.h"
 #include "Widgets/Layout/SBox.h"
 #include "SNiagaraStackIssueIcon.h"
@@ -35,7 +34,7 @@ void SNiagaraStackRendererItem::Construct(const FArguments& InArgs, UNiagaraStac
 		.VAlign(VAlign_Center)
 		[
 			SNew(SImage)
-			.Image(FSlateIconFinder::FindIconBrushForClass(RendererItem->GetRendererProperties()->GetClass()))
+			.Image(RendererItem->GetIconBrush())
 		]
 		// Display name
 		+ SHorizontalBox::Slot()
