fixed shader stage issue with execindex having duplicate values due to particle emission.  Using GDispatchThreadID instead

diff --git a/Engine/Plugins/FX/Niagara/Shaders/Private/NiagaraEmitterInstanceShader.usf b/Engine/Plugins/FX/Niagara/Shaders/Private/NiagaraEmitterInstanceShader.usf
--- a/Engine/Plugins/FX/Niagara/Shaders/Private/NiagaraEmitterInstanceShader.usf
+++ b/Engine/Plugins/FX/Niagara/Shaders/Private/NiagaraEmitterInstanceShader.usf
@@ -468,7 +468,7 @@ int rand(int x, int Seed1, int Seed2, int Seed3)
 	{
 		#if USE_SHADER_STAGES
 			const int DenseIndex = AcquireIndex(DataSetID,bIsValid);
-			const int SparseIndex = ExecIndex();
+			const int SparseIndex = GDispatchThreadId.x;
 			return (bIsPersistent && (DenseIndex != SparseIndex) ) ? SparseIndex : DenseIndex;
 		#else
 			return AcquireIndex(DataSetID,bIsValid);
