Niagara - Fix cis warning about int to bool coversions.  Also sorts in the correct order now.

diff --git a/Engine/Plugins/FX/Niagara/Source/NiagaraEditorWidgets/Private/SNiagaraOverviewStack.cpp b/Engine/Plugins/FX/Niagara/Source/NiagaraEditorWidgets/Private/SNiagaraOverviewStack.cpp
--- a/Engine/Plugins/FX/Niagara/Source/NiagaraEditorWidgets/Private/SNiagaraOverviewStack.cpp
+++ b/Engine/Plugins/FX/Niagara/Source/NiagaraEditorWidgets/Private/SNiagaraOverviewStack.cpp
@@ -390,7 +390,7 @@ TSharedRef<ITableRow> SNiagaraOverviewStack::OnGenerateRowForEntry(UNiagaraStack
 			TArray<FNiagaraScriptHighlight> ScriptHighlights = StackItem->GetHighlights();
 			Algo::Sort(ScriptHighlights, [](const FNiagaraScriptHighlight& A, const FNiagaraScriptHighlight& B)
 			{
-				return A.DisplayName.CompareTo(B.DisplayName);
+				return A.DisplayName.CompareTo(B.DisplayName) < 0;
 			});
 
 			FText HighlightToolTip;
