Export DI signature and bug fixes

diff --git a/Engine/Plugins/FX/Niagara/Source/Niagara/Classes/NiagaraDataInterfaceExport.h b/Engine/Plugins/FX/Niagara/Source/Niagara/Classes/NiagaraDataInterfaceExport.h
--- a/Engine/Plugins/FX/Niagara/Source/Niagara/Classes/NiagaraDataInterfaceExport.h
+++ b/Engine/Plugins/FX/Niagara/Source/Niagara/Classes/NiagaraDataInterfaceExport.h
@@ -66,14 +66,15 @@ public:
 	//UObject Interface End
 
 	//UNiagaraDataInterface Interface
-	virtual void GetFunctions(TArray<FNiagaraFunctionSignature>& OutFunctions)override;
+	virtual void GetFunctions(TArray<FNiagaraFunctionSignature>& OutFunctions) override;
 	virtual void GetVMExternalFunction(const FVMExternalFunctionBindingInfo& BindingInfo, void* InstanceData, FVMExternalFunction &OutFunc) override;
 	virtual bool InitPerInstanceData(void* PerInstanceData, FNiagaraSystemInstance* SystemInstance) override;
-	virtual int32 PerInstanceDataSize() const { return sizeof(ExportInterface_InstanceData); }
+	virtual int32 PerInstanceDataSize() const override { return sizeof(ExportInterface_InstanceData); }
 	virtual bool PerInstanceTick(void* PerInstanceData, FNiagaraSystemInstance* SystemInstance, float DeltaSeconds) override;
-	virtual bool PerInstanceTickPostSimulate(void* PerInstanceData, FNiagaraSystemInstance* SystemInstance, float DeltaSeconds);
+	virtual bool PerInstanceTickPostSimulate(void* PerInstanceData, FNiagaraSystemInstance* SystemInstance, float DeltaSeconds) override;
 	virtual bool Equals(const UNiagaraDataInterface* Other) const override;
 	virtual bool GetFunctionHLSL(const FName&  DefinitionFunctionName, FString InstanceFunctionName, FNiagaraDataInterfaceGPUParamInfo& ParamInfo, FString& OutHLSL) override;
+	virtual bool CanExecuteOnTarget(ENiagaraSimTarget Target) const override { return Target == ENiagaraSimTarget::CPUSim; }
 	//UNiagaraDataInterface Interface
 
 	virtual void StoreData(FVectorVMContext& Context);
diff --git a/Engine/Plugins/FX/Niagara/Source/Niagara/Private/NiagaraDataInterfaceExport.cpp b/Engine/Plugins/FX/Niagara/Source/Niagara/Private/NiagaraDataInterfaceExport.cpp
--- a/Engine/Plugins/FX/Niagara/Source/Niagara/Private/NiagaraDataInterfaceExport.cpp
+++ b/Engine/Plugins/FX/Niagara/Source/Niagara/Private/NiagaraDataInterfaceExport.cpp
@@ -138,7 +138,7 @@ void UNiagaraDataInterfaceExport::StoreData(FVectorVMContext& Context)
 		FBasicParticleData Data;
 		Data.Position = FVector(PositionParamX.GetAndAdvance(), PositionParamY.GetAndAdvance(), PositionParamZ.GetAndAdvance());
 		Data.Size = SizeParam.GetAndAdvance();
-		Data.Position = FVector(VelocityParamX.GetAndAdvance(), VelocityParamY.GetAndAdvance(), VelocityParamZ.GetAndAdvance());
+		Data.Velocity = FVector(VelocityParamX.GetAndAdvance(), VelocityParamY.GetAndAdvance(), VelocityParamZ.GetAndAdvance());
 
 		FNiagaraBool Valid;
 		if (ValidHandlerData && ShouldStore)
