Removed useless log warning

diff --git a/Engine/Plugins/FX/Niagara/Source/Niagara/Private/NiagaraComponentPool.cpp b/Engine/Plugins/FX/Niagara/Source/Niagara/Private/NiagaraComponentPool.cpp
--- a/Engine/Plugins/FX/Niagara/Source/Niagara/Private/NiagaraComponentPool.cpp
+++ b/Engine/Plugins/FX/Niagara/Source/Niagara/Private/NiagaraComponentPool.cpp
@@ -65,19 +65,17 @@ void FNCPool::Cleanup()
 		//It's possible for people to manually destroy these so we have to guard against it. Though we warn about it in UNiagaraComponent::BeginDestroy
 		if (NC)
 		{
-			NC->PoolingMethod = ENCPoolMethod::None;//Reset so we don't trigger warnings about destroying pooled NCs.
+			NC->PoolingMethod = ENCPoolMethod::None; //Reset so we don't trigger warnings about destroying pooled NCs.
 			NC->DestroyComponent();
 		}
 	}
 
-	//Warn if there are any manually released NCs still in the world at cleanup time.
 	for (UNiagaraComponent* NC : InUseComponents_Manual)
 	{
 		//It's possible for people to manually destroy these so we have to guard against it. Though we warn about it in UNiagaraComponent::BeginDestroy
 		if (NC)
 		{
-			UE_LOG(LogNiagara, Warning, TEXT("Pooled NC set to manual release is still in use as the pool is being cleaned up. %s"), *NC->GetAsset()->GetFullName());
-			NC->PoolingMethod = ENCPoolMethod::None;//Reset so we don't trigger warnings about destroying pooled NCs.
+			NC->PoolingMethod = ENCPoolMethod::None; //Reset so we don't trigger warnings about destroying pooled NCs.
 			NC->DestroyComponent();
 		}
 	}
