Niagara - Fix non-unity issues.

diff --git a/Engine/Plugins/FX/Niagara/Source/Niagara/Private/NiagaraScriptHighlight.cpp b/Engine/Plugins/FX/Niagara/Source/Niagara/Private/NiagaraScriptHighlight.cpp
--- a/Engine/Plugins/FX/Niagara/Source/Niagara/Private/NiagaraScriptHighlight.cpp
+++ b/Engine/Plugins/FX/Niagara/Source/Niagara/Private/NiagaraScriptHighlight.cpp
@@ -6,6 +6,7 @@
 #include "JsonObjectConverter.h"
 #include "Serialization/JsonWriter.h"
 #include "Serialization/JsonSerializer.h"
+#include "Policies/CondensedJsonPrintPolicy.h"
 #endif
 
 FNiagaraScriptHighlight::FNiagaraScriptHighlight()
diff --git a/Engine/Plugins/FX/Niagara/Source/NiagaraEditor/Private/Customizations/NiagaraScriptDetails.cpp b/Engine/Plugins/FX/Niagara/Source/NiagaraEditor/Private/Customizations/NiagaraScriptDetails.cpp
--- a/Engine/Plugins/FX/Niagara/Source/NiagaraEditor/Private/Customizations/NiagaraScriptDetails.cpp
+++ b/Engine/Plugins/FX/Niagara/Source/NiagaraEditor/Private/Customizations/NiagaraScriptDetails.cpp
@@ -30,6 +30,7 @@
 #include "NiagaraMetaDataCustomNodeBuilder.h"
 #include "SItemSelector.h"
 #include "PropertyCustomizationHelpers.h"
+#include "ScopedTransaction.h"
 
 #define LOCTEXT_NAMESPACE "NiagaraScriptDetails"
 
@@ -844,8 +845,6 @@ private:
 
 void FNiagaraScriptHighlightDetails::CustomizeHeader(TSharedRef<IPropertyHandle> PropertyHandle, FDetailWidgetRow& HeaderRow, IPropertyTypeCustomizationUtils& CustomizationUtils)
 {
-	//HeaderRow.ValueContent()
-
 	HeaderRow.WholeRowContent()
 		[
 			SNew(SNiagaraScriptHighlightPicker, PropertyHandle)
diff --git a/Engine/Plugins/FX/Niagara/Source/NiagaraEditor/Private/Customizations/NiagaraScriptDetails.h b/Engine/Plugins/FX/Niagara/Source/NiagaraEditor/Private/Customizations/NiagaraScriptDetails.h
--- a/Engine/Plugins/FX/Niagara/Source/NiagaraEditor/Private/Customizations/NiagaraScriptDetails.h
+++ b/Engine/Plugins/FX/Niagara/Source/NiagaraEditor/Private/Customizations/NiagaraScriptDetails.h
@@ -4,11 +4,16 @@
 
 #include "CoreMinimal.h"
 #include "IDetailCustomization.h"
+#include "IPropertyTypeCustomization.h"
 #include "Input/Reply.h"
 
 class FNiagaraScriptViewModel;
 class INiagaraParameterCollectionViewModel;
 class FNiagaraMetaDataCustomNodeBuilder;
+class IPropertyHandle;
+class FDetailWidgetRow;
+class IDetailChildrenBuilder;
+class IPropertyTypeCustomizationUtils;
 
 class FNiagaraScriptDetails : public IDetailCustomization
 {
diff --git a/Engine/Plugins/FX/Niagara/Source/NiagaraEditorWidgets/Private/NiagaraEditorWidgetsStyle.cpp b/Engine/Plugins/FX/Niagara/Source/NiagaraEditorWidgets/Private/NiagaraEditorWidgetsStyle.cpp
--- a/Engine/Plugins/FX/Niagara/Source/NiagaraEditorWidgets/Private/NiagaraEditorWidgetsStyle.cpp
+++ b/Engine/Plugins/FX/Niagara/Source/NiagaraEditorWidgets/Private/NiagaraEditorWidgetsStyle.cpp
@@ -200,7 +200,8 @@ TSharedRef< FSlateStyleSet > FNiagaraEditorWidgetsStyle::Create()
 		.SetShadowColorAndOpacity(FLinearColor(0, 0, 0, 0.9f));
 	Style->Set("NiagaraEditor.Stack.AddButtonText", AddButtonText);
 
-	Style->Set("NiagaraEditor.Stack.ModuleHighlight",  new IMAGE_PLUGIN_BRUSH("Icons/ModuleHighlight", Icon8x8, FLinearColor::White));
+	Style->Set("NiagaraEditor.Stack.ModuleHighlight", new IMAGE_PLUGIN_BRUSH("Icons/ModuleHighlight", Icon8x8, FLinearColor::White));
+	Style->Set("NiagaraEditor.Stack.ModuleHighlightMore", new IMAGE_PLUGIN_BRUSH("Icons/ModuleHighlightMore", Icon8x8, FLinearColor::White));
 
 	Style->Set("NiagaraEditor.ShowInCurveEditorIcon", new IMAGE_PLUGIN_BRUSH("Icons/ShowInCurveEditor", Icon16x16, FLinearColor::White));
 
