- Skip NaN checks in TEST builds
- Skip some event code if no events, reduces instruction count and saves performance

#rnx
#release 11.30

diff --git a/Engine/Plugins/FX/Niagara/Source/Niagara/Private/NiagaraEmitterInstance.cpp b/Engine/Plugins/FX/Niagara/Source/Niagara/Private/NiagaraEmitterInstance.cpp
--- a/Engine/Plugins/FX/Niagara/Source/Niagara/Private/NiagaraEmitterInstance.cpp
+++ b/Engine/Plugins/FX/Niagara/Source/Niagara/Private/NiagaraEmitterInstance.cpp
@@ -854,7 +854,7 @@ FBox FNiagaraEmitterInstance::CalculateDynamicBounds(const bool bReadGPUSimulati
 		Ret += BoundsCalculator->CalculateBounds(NumInstances, bContainsNaN);
 	}
 
-#if !UE_BUILD_SHIPPING
+#if !UE_BUILD_SHIPPING && !UE_BUILD_TEST
 	if (bContainsNaN && ParentSystemInstance != nullptr && CachedEmitter != nullptr && ParentSystemInstance->GetSystem() != nullptr)
 	{
 		const FString OnScreenMessage = FString::Printf(TEXT("Niagara Particle position data contains NaNs. Likely a divide by zero somewhere in your modules. Emitter \"%s\" in System \"%s\".  Use fx.Niagara.DumpNansOnce to get full log."), *CachedEmitter->GetName(), *ParentSystemInstance->GetSystem()->GetName());
@@ -1047,24 +1047,28 @@ void FNiagaraEmitterInstance::PreTick()
 
 	//Gather events we're going to be reading from / handling this frame.
 	//We must do this in pre-tick so we can gather (and mark in use) all sets from other emitters.
-	EventHandlingInfo.Reset();
-	EventHandlingInfo.SetNum(CachedEmitter->GetEventHandlers().Num());
-	EventSpawnTotal = 0;
-	for (int32 i = 0; i < CachedEmitter->GetEventHandlers().Num(); i++)
+	const int32 NumEventHandlers = CachedEmitter->GetEventHandlers().Num();
+	if (NumEventHandlers > 0)
 	{
-		const FNiagaraEventScriptProperties &EventHandlerProps = CachedEmitter->GetEventHandlers()[i];
-		FNiagaraEventHandlingInfo& Info = EventHandlingInfo[i];
-		Info.SourceEmitterGuid = EventHandlerProps.SourceEmitterID;
-		Info.SourceEmitterName = Info.SourceEmitterGuid.IsValid() ? *Info.SourceEmitterGuid.ToString() : CachedIDName;
-		Info.SpawnCounts.Reset();
-		Info.TotalSpawnCount = 0;
-		Info.EventData = nullptr;
-		if (FNiagaraDataSet* EventSet = FNiagaraEventDataSetMgr::GetEventDataSet(ParentSystemInstance->GetId(), Info.SourceEmitterName, EventHandlerProps.SourceEventName))
-		{
-			Info.SetEventData(&EventSet->GetCurrentDataChecked());
-			uint32 EventSpawnNum = CalculateEventSpawnCount(EventHandlerProps, Info.SpawnCounts, EventSet);
-			Info.TotalSpawnCount += EventSpawnNum;
-			EventSpawnTotal += EventSpawnNum;
+		EventHandlingInfo.Reset();
+		EventHandlingInfo.SetNum(CachedEmitter->GetEventHandlers().Num());
+		EventSpawnTotal = 0;
+		for (int32 i = 0; i < CachedEmitter->GetEventHandlers().Num(); i++)
+		{
+			const FNiagaraEventScriptProperties &EventHandlerProps = CachedEmitter->GetEventHandlers()[i];
+			FNiagaraEventHandlingInfo& Info = EventHandlingInfo[i];
+			Info.SourceEmitterGuid = EventHandlerProps.SourceEmitterID;
+			Info.SourceEmitterName = Info.SourceEmitterGuid.IsValid() ? *Info.SourceEmitterGuid.ToString() : CachedIDName;
+			Info.SpawnCounts.Reset();
+			Info.TotalSpawnCount = 0;
+			Info.EventData = nullptr;
+			if (FNiagaraDataSet* EventSet = FNiagaraEventDataSetMgr::GetEventDataSet(ParentSystemInstance->GetId(), Info.SourceEmitterName, EventHandlerProps.SourceEventName))
+			{
+				Info.SetEventData(&EventSet->GetCurrentDataChecked());
+				uint32 EventSpawnNum = CalculateEventSpawnCount(EventHandlerProps, Info.SpawnCounts, EventSet);
+				Info.TotalSpawnCount += EventSpawnNum;
+				EventSpawnTotal += EventSpawnNum;
+			}
 		}
 	}
 
diff --git a/Engine/Plugins/FX/Niagara/Source/Niagara/Public/NiagaraBoundsCalculatorHelper.h b/Engine/Plugins/FX/Niagara/Source/Niagara/Public/NiagaraBoundsCalculatorHelper.h
--- a/Engine/Plugins/FX/Niagara/Source/Niagara/Public/NiagaraBoundsCalculatorHelper.h
+++ b/Engine/Plugins/FX/Niagara/Source/Niagara/Public/NiagaraBoundsCalculatorHelper.h
@@ -71,7 +71,7 @@ public:
 		float MaxSpriteSize(SpriteSizeAccessor.IsValid() ? 0.0f : kDefaultSize);
 		float MaxRibbonWidth(RibbonWidthAccessor.IsValid() ? 0.0f : kDefaultSize);
 
-#if !UE_BUILD_SHIPPING
+#if !UE_BUILD_SHIPPING && !UE_BUILD_TEST
 		bOutContainsNaN = false;
 #endif
 
@@ -83,6 +83,7 @@ public:
 			FVector Position;
 			PositionAccessor.Get(InstIdx, Position);
 
+#if !UE_BUILD_SHIPPING && !UE_BUILD_TEST
 			// Some graphs have a tendency to divide by zero. This ContainsNaN has been added prophylactically
 			// to keep us safe during GDC. It should be removed as soon as we feel safe that scripts are appropriately warned.
 			if (Position.ContainsNaN())
@@ -90,6 +91,7 @@ public:
 				bOutContainsNaN = true;
 				continue;
 			}
+#endif
 
 			Ret += Position;
 
