- Optimize Curve LUTs to save memory and disk space
- Remove Rich Curve data if we are using a LUT when we are cooking
- In my simple test this brought the overall curve memory from 193.83kb -> 40.67kb

#rnx

diff --git a/Engine/Plugins/FX/Niagara/Source/Niagara/Classes/NiagaraDataInterface.h b/Engine/Plugins/FX/Niagara/Source/Niagara/Classes/NiagaraDataInterface.h
--- a/Engine/Plugins/FX/Niagara/Source/Niagara/Classes/NiagaraDataInterface.h
+++ b/Engine/Plugins/FX/Niagara/Source/Niagara/Classes/NiagaraDataInterface.h
@@ -185,28 +185,6 @@ struct FNiagaraDataInterfaceProxy : TSharedFromThis<FNiagaraDataInterfaceProxy,
 	virtual void PostStage(FRHICommandList& RHICmdList, const FNiagaraDataInterfaceSetArgs& Context) {}
 };
 
-struct FNiagaraDataInterfaceProxyCurveBase : public FNiagaraDataInterfaceProxy
-{
-	~FNiagaraDataInterfaceProxyCurveBase()
-	{
-		check(IsInRenderingThread());
-		CurveLUT.Release();
-	}
-
-	float LUTMinTime;
-	float LUTMaxTime;
-	float LUTInvTimeRange;
-	FReadBuffer CurveLUT;
-
-	virtual int32 PerInstanceDataPassedToRenderThreadSize() const override
-	{
-		return 0;
-	}
-
-	// @todo REMOVEME
-	virtual void ConsumePerInstanceDataFromGameThread(void* PerInstanceData, const FNiagaraSystemInstanceID& Instance) override { check(false); }
-};
-
 //////////////////////////////////////////////////////////////////////////
 
 /** Base class for all Niagara data interfaces. */
@@ -335,141 +313,6 @@ protected:
 	TSharedPtr<FNiagaraDataInterfaceProxy, ESPMode::ThreadSafe> Proxy;
 };
 
-/** Base class for curve data interfaces which facilitates handling the curve data in a standardized way. */
-UCLASS(EditInlineNew, Category = "Curves", meta = (DisplayName = "Float Curve"))
-class NIAGARA_API UNiagaraDataInterfaceCurveBase : public UNiagaraDataInterface
-{
-protected:
-	GENERATED_BODY()
-	/*UPROPERTY()
-	bool GPUBufferDirty;*/
-	UPROPERTY()
-	TArray<float> ShaderLUT;
-	UPROPERTY()
-	float LUTMinTime;
-	UPROPERTY()
-	float LUTMaxTime;
-	UPROPERTY()
-	float LUTInvTimeRange;
-
-	/** Remap a sample time for this curve to 0 to 1 between first and last keys for LUT access.*/
-	FORCEINLINE float NormalizeTime(float T)
-	{
-		return (T - LUTMinTime) * LUTInvTimeRange;
-	}
-
-	/** Remap a 0 to 1 value between the first and last keys to a real sample time for this curve. */
-	FORCEINLINE float UnnormalizeTime(float T)
-	{
-		return (T / LUTInvTimeRange) + LUTMinTime;
-	}
-
-public:
-	UNiagaraDataInterfaceCurveBase()
-		: LUTMinTime(0.0f)
-		, LUTMaxTime(1.0f)
-		, LUTInvTimeRange(1.0f)
-		, bUseLUT(true)
-#if WITH_EDITORONLY_DATA
-		, ShowInCurveEditor(false)
-#endif
-	{
-		Proxy = MakeShared<FNiagaraDataInterfaceProxyCurveBase, ESPMode::ThreadSafe>();
-	}
-
-	UNiagaraDataInterfaceCurveBase(FObjectInitializer const& ObjectInitializer)
-		: LUTMinTime(0.0f)
-		, LUTMaxTime(1.0f)
-		, LUTInvTimeRange(1.0f)
-		, bUseLUT(true)
-#if WITH_EDITORONLY_DATA
-		, ShowInCurveEditor(false)
-#endif
-	{
-		Proxy = MakeShared<FNiagaraDataInterfaceProxyCurveBase, ESPMode::ThreadSafe>();
-	}
-
-	UPROPERTY(EditAnywhere, AdvancedDisplay, Category = "Curve")
-	uint32 bUseLUT : 1;
-	
-#if WITH_EDITORONLY_DATA
-	UPROPERTY(EditAnywhere, Transient, Category = "Curve")
-	bool ShowInCurveEditor;
-#endif
-
-	enum
-	{
-		CurveLUTWidth = 128,
-		CurveLUTWidthMinusOne = 127,
-	};
-	/** Structure to facilitate getting standardized curve information from a curve data interface. */
-	struct FCurveData
-	{
-		FCurveData(FRichCurve* InCurve, FName InName, FLinearColor InColor)
-			: Curve(InCurve)
-			, Name(InName)
-			, Color(InColor)
-		{
-		}
-		/** A pointer to the curve. */
-		FRichCurve* Curve;
-		/** The name of the curve, unique within the data interface, which identifies the curve in the UI. */
-		FName Name;
-		/** The color to use when displaying this curve in the UI. */
-		FLinearColor Color;
-	};
-
-	/** Gets information for all of the curves owned by this curve data interface. */
-	virtual void GetCurveData(TArray<FCurveData>& OutCurveData) { }
-
-	virtual void GetParameterDefinitionHLSL(FNiagaraDataInterfaceGPUParamInfo& ParamInfo, FString& OutHLSL)override;
-	virtual FNiagaraDataInterfaceParametersCS* ConstructComputeParameters()const override;
-
-	virtual int32 GetCurveNumElems()const 
-	{
-		checkf(false, TEXT("You must implement this function so the GPU buffer can be created of the correct size."));
-		return 0; 
-	}
-
-	//TODO: Make this a texture and get HW filter + clamping?
-	//FReadBuffer& GetCurveLUTGPUBuffer();
-
-	//UNiagaraDataInterface interface
-	virtual bool Equals(const UNiagaraDataInterface* Other) const override;
-	virtual bool CanExecuteOnTarget(ENiagaraSimTarget Target)const override { return true; }
-	virtual void UpdateLUT() { };
-
-	FORCEINLINE float GetMinTime()const { return LUTMinTime; }
-	FORCEINLINE float GetMaxTime()const { return LUTMaxTime; }
-	FORCEINLINE float GetInvTimeRange()const { return LUTInvTimeRange; }
-
-protected:
-	void PushToRenderThread();
-	virtual bool CopyToInternal(UNiagaraDataInterface* Destination) const override;
-	virtual bool CompareLUTS(const TArray<float>& OtherLUT) const;
-	//UNiagaraDataInterface interface END
-};
-
-//External function binder choosing between template specializations based on if a curve should use the LUT over full evaluation.
-template<typename NextBinder>
-struct TCurveUseLUTBinder
-{
-	template<typename... ParamTypes>
-	static void Bind(UNiagaraDataInterface* Interface, const FVMExternalFunctionBindingInfo& BindingInfo, void* InstanceData, FVMExternalFunction &OutFunc)
-	{
-		UNiagaraDataInterfaceCurveBase* CurveInterface = CastChecked<UNiagaraDataInterfaceCurveBase>(Interface);
-		if (CurveInterface->bUseLUT)
-		{
-			NextBinder::template Bind<ParamTypes..., TIntegralConstant<bool, true>>(Interface, BindingInfo, InstanceData, OutFunc);
-		}
-		else
-		{
-			NextBinder::template Bind<ParamTypes..., TIntegralConstant<bool, false>>(Interface, BindingInfo, InstanceData, OutFunc);
-		}
-	}
-};
-
-
 /** Helper class for decoding NDI parameters into a usable struct type. */
 template<typename T>
 struct FNDIParameter
@@ -647,4 +490,4 @@ struct FNDIRandomHelper
 	FNDIParameter<FNiagaraRandInfo> RandParam;
 
 	FNiagaraRandInfo RandInfo;
-};
\ No newline at end of file
+};
diff --git a/Engine/Plugins/FX/Niagara/Source/Niagara/Classes/NiagaraDataInterfaceColorCurve.h b/Engine/Plugins/FX/Niagara/Source/Niagara/Classes/NiagaraDataInterfaceColorCurve.h
--- a/Engine/Plugins/FX/Niagara/Source/Niagara/Classes/NiagaraDataInterfaceColorCurve.h
+++ b/Engine/Plugins/FX/Niagara/Source/Niagara/Classes/NiagaraDataInterfaceColorCurve.h
@@ -6,7 +6,7 @@
 #include "VectorVM.h"
 #include "StaticMeshResources.h"
 #include "Curves/RichCurve.h"
-#include "NiagaraDataInterface.h"
+#include "NiagaraDataInterfaceCurveBase.h"
 #include "NiagaraDataInterfaceColorCurve.generated.h"
 
 class INiagaraCompiler;
@@ -36,16 +36,16 @@ public:
 
 	//UObject Interface
 	virtual void PostInitProperties() override;
-	virtual void PostLoad() override;
+	virtual void Serialize(FArchive& Ar) override;
 	//UObject Interface End
 
 	enum
 	{
 		CurveLUTNumElems = 4,
-		CurveLUTMax = (CurveLUTWidth * CurveLUTNumElems) - 1,
 	};
 
-	void UpdateLUT();
+	virtual void UpdateTimeRanges() override;
+	virtual TArray<float> BuildLUT(int32 NumEntries) const override;
 
 	virtual void GetFunctions(TArray<FNiagaraFunctionSignature>& OutFunctions)override;
 	virtual void GetVMExternalFunction(const FVMExternalFunctionBindingInfo& BindingInfo, void* InstanceData, FVMExternalFunction &OutFunc) override;
diff --git a/Engine/Plugins/FX/Niagara/Source/Niagara/Classes/NiagaraDataInterfaceCurve.h b/Engine/Plugins/FX/Niagara/Source/Niagara/Classes/NiagaraDataInterfaceCurve.h
--- a/Engine/Plugins/FX/Niagara/Source/Niagara/Classes/NiagaraDataInterfaceCurve.h
+++ b/Engine/Plugins/FX/Niagara/Source/Niagara/Classes/NiagaraDataInterfaceCurve.h
@@ -6,7 +6,7 @@
 #include "VectorVM.h"
 #include "StaticMeshResources.h"
 #include "Curves/RichCurve.h"
-#include "NiagaraDataInterface.h"
+#include "NiagaraDataInterfaceCurveBase.h"
 #include "NiagaraDataInterfaceCurve.generated.h"
 
 
@@ -23,16 +23,16 @@ public:
 	enum
 	{
 		CurveLUTNumElems = 1,
-		CurveLUTMax = (CurveLUTWidth * CurveLUTNumElems) - 1,
 	};
 
-	void UpdateLUT();
-
 	//UObject Interface
 	virtual void PostInitProperties() override;
-	virtual void PostLoad() override;
+	virtual void Serialize(FArchive& Ar) override;
 	//UObject Interface End
 
+	virtual void UpdateTimeRanges() override;
+	virtual TArray<float> BuildLUT(int32 NumEntries) const override;
+
 	virtual void GetFunctions(TArray<FNiagaraFunctionSignature>& OutFunctions)override;
 	virtual void GetVMExternalFunction(const FVMExternalFunctionBindingInfo& BindingInfo, void* InstanceData, FVMExternalFunction &OutFunc) override;
 
diff --git a/Engine/Plugins/FX/Niagara/Source/Niagara/Classes/NiagaraDataInterfaceCurveBase.h b/Engine/Plugins/FX/Niagara/Source/Niagara/Classes/NiagaraDataInterfaceCurveBase.h
new file mode 100644
--- /dev/null
+++ b/Engine/Plugins/FX/Niagara/Source/Niagara/Classes/NiagaraDataInterfaceCurveBase.h
@@ -0,0 +1,196 @@
+// Copyright 1998-2019 Epic Games, Inc. All Rights Reserved.
+#pragma once
+
+#include "NiagaraCommon.h"
+#include "NiagaraShared.h"
+#include "VectorVM.h"
+#include "StaticMeshResources.h"
+#include "Curves/RichCurve.h"
+#include "NiagaraDataInterface.h"
+#include "NiagaraDataInterfaceCurveBase.generated.h"
+
+/** Base class for curve data proxy data. */
+struct FNiagaraDataInterfaceProxyCurveBase : public FNiagaraDataInterfaceProxy
+{
+	virtual ~FNiagaraDataInterfaceProxyCurveBase()
+	{
+		check(IsInRenderingThread());
+		CurveLUT.Release();
+	}
+
+	float LUTMinTime;
+	float LUTMaxTime;
+	float LUTInvTimeRange;
+	float CurveLUTNumMinusOne;
+	FReadBuffer CurveLUT;
+
+	virtual int32 PerInstanceDataPassedToRenderThreadSize() const override
+	{
+		return 0;
+	}
+
+	// @todo REMOVEME
+	virtual void ConsumePerInstanceDataFromGameThread(void* PerInstanceData, const FNiagaraSystemInstanceID& Instance) override { check(false); }
+};
+
+/** Base class for curve data interfaces which facilitates handling the curve data in a standardized way. */
+UCLASS(EditInlineNew, Category = "Curves", meta = (DisplayName = "Float Curve"))
+class NIAGARA_API UNiagaraDataInterfaceCurveBase : public UNiagaraDataInterface
+{
+public:
+	static constexpr float DefaultOptimizeThreshold = 0.01f;
+
+protected:
+	GENERATED_BODY()
+
+	UPROPERTY()
+	TArray<float> ShaderLUT;
+
+	UPROPERTY()
+	float LUTMinTime;
+
+	UPROPERTY()
+	float LUTMaxTime;
+
+	UPROPERTY()
+	float LUTInvTimeRange;
+
+	UPROPERTY()
+	float LUTNumSamplesMinusOne;
+
+	/** Remap a sample time for this curve to 0 to 1 between first and last keys for LUT access.*/
+	FORCEINLINE float NormalizeTime(float T) const
+	{
+		return (T - LUTMinTime) * LUTInvTimeRange;
+	}
+
+	/** Remap a 0 to 1 value between the first and last keys to a real sample time for this curve. */
+	FORCEINLINE float UnnormalizeTime(float T) const
+	{
+		return (T / LUTInvTimeRange) + LUTMinTime;
+	}
+
+public:
+	UPROPERTY(EditAnywhere, AdvancedDisplay, Category = "Curve")
+	uint32 bUseLUT : 1;
+
+	/** Do we optimize the LUT, this saves memory but may introduce errors.  Errors can be reduced modifying the threshold. */
+	UPROPERTY(EditAnywhere, AdvancedDisplay, Category = "Curve")
+	uint32 bOptimizeLUT : 1;
+
+	UPROPERTY(EditAnywhere, AdvancedDisplay, Category = "Curve")
+	uint32 bOverrideOptimizeThreshold : 1;
+
+	/** Threshold used to optimize the LUT. */
+	UPROPERTY(EditAnywhere, AdvancedDisplay, Category = "Curve", meta=(EditCondition="bOverrideOptimizeThreshold"))
+	float OptimizeThreshold;
+
+#if WITH_EDITORONLY_DATA
+	UPROPERTY(EditAnywhere, Transient, Category = "Curve")
+	bool ShowInCurveEditor;
+#endif
+
+public:
+	UNiagaraDataInterfaceCurveBase()
+		: LUTMinTime(0.0f)
+		, LUTMaxTime(1.0f)
+		, LUTInvTimeRange(1.0f)
+		, bUseLUT(true)
+		, bOptimizeLUT(true)
+		, bOverrideOptimizeThreshold(false)
+		, OptimizeThreshold(DefaultOptimizeThreshold)
+#if WITH_EDITORONLY_DATA
+		, ShowInCurveEditor(false)
+#endif
+	{
+		Proxy = MakeShared<FNiagaraDataInterfaceProxyCurveBase, ESPMode::ThreadSafe>();
+	}
+
+	UNiagaraDataInterfaceCurveBase(FObjectInitializer const& ObjectInitializer)
+		: LUTMinTime(0.0f)
+		, LUTMaxTime(1.0f)
+		, LUTInvTimeRange(1.0f)
+		, bUseLUT(true)
+		, bOptimizeLUT(true)
+		, bOverrideOptimizeThreshold(false)
+		, OptimizeThreshold(DefaultOptimizeThreshold)
+#if WITH_EDITORONLY_DATA
+		, ShowInCurveEditor(false)
+#endif
+	{
+		Proxy = MakeShared<FNiagaraDataInterfaceProxyCurveBase, ESPMode::ThreadSafe>();
+	}
+
+	enum
+	{
+		CurveLUTDefaultWidth = 128,
+	};
+
+	/** Structure to facilitate getting standardized curve information from a curve data interface. */
+	struct FCurveData
+	{
+		FCurveData(FRichCurve* InCurve, FName InName, FLinearColor InColor)
+			: Curve(InCurve)
+			, Name(InName)
+			, Color(InColor)
+		{
+		}
+		/** A pointer to the curve. */
+		FRichCurve* Curve;
+		/** The name of the curve, unique within the data interface, which identifies the curve in the UI. */
+		FName Name;
+		/** The color to use when displaying this curve in the UI. */
+		FLinearColor Color;
+	};
+
+	//UObject Interface
+	virtual void PostLoad() override;
+	//UObject Interface End
+
+	/** Gets information for all of the curves owned by this curve data interface. */
+	virtual void GetCurveData(TArray<FCurveData>& OutCurveData) { }
+
+	virtual void GetParameterDefinitionHLSL(FNiagaraDataInterfaceGPUParamInfo& ParamInfo, FString& OutHLSL) override;
+	virtual FNiagaraDataInterfaceParametersCS* ConstructComputeParameters() const override;
+
+	void UpdateLUT();
+	void OptimizeLUT();
+
+	//UNiagaraDataInterface interface
+	virtual bool Equals(const UNiagaraDataInterface* Other) const override;
+	virtual bool CanExecuteOnTarget(ENiagaraSimTarget Target) const override { return true; }
+
+	virtual int32 GetCurveNumElems() const { checkf(false, TEXT("You must implement this function in your derived class")); return 0; }
+
+	virtual void UpdateTimeRanges() { checkf(false, TEXT("You must implement this function in your derived class")); }
+	virtual TArray<float> BuildLUT(int32 NumEntries) const { checkf(false, TEXT("You must implement this function in your derived class")); return TArray<float>(); }
+
+	FORCEINLINE float GetMinTime()const { return LUTMinTime; }
+	FORCEINLINE float GetMaxTime()const { return LUTMaxTime; }
+	FORCEINLINE float GetInvTimeRange()const { return LUTInvTimeRange; }
+
+protected:
+	void PushToRenderThread();
+	virtual bool CopyToInternal(UNiagaraDataInterface* Destination) const override;
+	virtual bool CompareLUTS(const TArray<float>& OtherLUT) const;
+	//UNiagaraDataInterface interface END
+};
+
+//External function binder choosing between template specializations based on if a curve should use the LUT over full evaluation.
+template<typename NextBinder>
+struct TCurveUseLUTBinder
+{
+	template<typename... ParamTypes>
+	static void Bind(UNiagaraDataInterface* Interface, const FVMExternalFunctionBindingInfo& BindingInfo, void* InstanceData, FVMExternalFunction &OutFunc)
+	{
+		UNiagaraDataInterfaceCurveBase* CurveInterface = CastChecked<UNiagaraDataInterfaceCurveBase>(Interface);
+		if (CurveInterface->bUseLUT)
+		{
+			NextBinder::template Bind<ParamTypes..., TIntegralConstant<bool, true>>(Interface, BindingInfo, InstanceData, OutFunc);
+		}
+		else
+		{
+			NextBinder::template Bind<ParamTypes..., TIntegralConstant<bool, false>>(Interface, BindingInfo, InstanceData, OutFunc);
+		}
+	}
+};
diff --git a/Engine/Plugins/FX/Niagara/Source/Niagara/Classes/NiagaraDataInterfaceVector2DCurve.h b/Engine/Plugins/FX/Niagara/Source/Niagara/Classes/NiagaraDataInterfaceVector2DCurve.h
--- a/Engine/Plugins/FX/Niagara/Source/Niagara/Classes/NiagaraDataInterfaceVector2DCurve.h
+++ b/Engine/Plugins/FX/Niagara/Source/Niagara/Classes/NiagaraDataInterfaceVector2DCurve.h
@@ -6,8 +6,7 @@
 #include "VectorVM.h"
 #include "StaticMeshResources.h"
 #include "Curves/RichCurve.h"
-#include "NiagaraDataInterface.h"
-
+#include "NiagaraDataInterfaceCurveBase.h"
 #include "NiagaraDataInterfaceVector2DCurve.generated.h"
 
 
@@ -27,16 +26,15 @@ public:
 	enum
 	{
 		CurveLUTNumElems = 2,
-		CurveLUTMax = (CurveLUTWidth * CurveLUTNumElems) - 1,
 	};
 
-	void UpdateLUT();
-
 	//UObject Interface
 	virtual void PostInitProperties() override;
-	virtual void PostLoad() override;
-
+	virtual void Serialize(FArchive& Ar) override;
 	//UObject Interface End
+	
+	virtual void UpdateTimeRanges() override;
+	virtual TArray<float> BuildLUT(int32 NumEntries) const override;
 
 	virtual bool Equals(const UNiagaraDataInterface* Other) const override;
 
diff --git a/Engine/Plugins/FX/Niagara/Source/Niagara/Classes/NiagaraDataInterfaceVector4Curve.h b/Engine/Plugins/FX/Niagara/Source/Niagara/Classes/NiagaraDataInterfaceVector4Curve.h
--- a/Engine/Plugins/FX/Niagara/Source/Niagara/Classes/NiagaraDataInterfaceVector4Curve.h
+++ b/Engine/Plugins/FX/Niagara/Source/Niagara/Classes/NiagaraDataInterfaceVector4Curve.h
@@ -6,7 +6,7 @@
 #include "VectorVM.h"
 #include "StaticMeshResources.h"
 #include "Curves/RichCurve.h"
-#include "NiagaraDataInterface.h"
+#include "NiagaraDataInterfaceCurveBase.h"
 #include "NiagaraDataInterfaceVector4Curve.generated.h"
 
 class INiagaraCompiler;
@@ -37,16 +37,16 @@ public:
 	enum
 	{
 		CurveLUTNumElems = 4,
-		CurveLUTMax = (CurveLUTWidth * CurveLUTNumElems) - 1,
 	};
 
-	void UpdateLUT();
-
 	//UObject Interface
 	virtual void PostInitProperties() override;
-	virtual void PostLoad() override;
+	virtual void Serialize(FArchive& Ar) override;
 	//UObject Interface End
 
+	virtual void UpdateTimeRanges() override;
+	virtual TArray<float> BuildLUT(int32 NumEntries) const override;
+
 	virtual void GetFunctions(TArray<FNiagaraFunctionSignature>& OutFunctions)override;
 	virtual void GetVMExternalFunction(const FVMExternalFunctionBindingInfo& BindingInfo, void* InstanceData, FVMExternalFunction &OutFunc) override;
 
diff --git a/Engine/Plugins/FX/Niagara/Source/Niagara/Classes/NiagaraDataInterfaceVectorCurve.h b/Engine/Plugins/FX/Niagara/Source/Niagara/Classes/NiagaraDataInterfaceVectorCurve.h
--- a/Engine/Plugins/FX/Niagara/Source/Niagara/Classes/NiagaraDataInterfaceVectorCurve.h
+++ b/Engine/Plugins/FX/Niagara/Source/Niagara/Classes/NiagaraDataInterfaceVectorCurve.h
@@ -6,7 +6,7 @@
 #include "VectorVM.h"
 #include "StaticMeshResources.h"
 #include "Curves/RichCurve.h"
-#include "NiagaraDataInterface.h"
+#include "NiagaraDataInterfaceCurveBase.h"
 #include "NiagaraDataInterfaceVectorCurve.generated.h"
 
 class INiagaraCompiler;
@@ -34,16 +34,16 @@ public:
 	enum
 	{
 		CurveLUTNumElems = 3,
-		CurveLUTMax = (CurveLUTWidth * CurveLUTNumElems) - 1,
 	};
 
-	void UpdateLUT();
-
 	//UObject Interface
 	virtual void PostInitProperties() override;
-	virtual void PostLoad() override;
+	virtual void Serialize(FArchive& Ar) override;
 	//UObject Interface End
 
+	virtual void UpdateTimeRanges() override;
+	virtual TArray<float> BuildLUT(int32 NumEntries) const override;
+
 	virtual void GetFunctions(TArray<FNiagaraFunctionSignature>& OutFunctions)override;
 	virtual void GetVMExternalFunction(const FVMExternalFunctionBindingInfo& BindingInfo, void* InstanceData, FVMExternalFunction &OutFunc) override;
 
diff --git a/Engine/Plugins/FX/Niagara/Source/Niagara/Private/NiagaraDataInterface.cpp b/Engine/Plugins/FX/Niagara/Source/Niagara/Private/NiagaraDataInterface.cpp
--- a/Engine/Plugins/FX/Niagara/Source/Niagara/Private/NiagaraDataInterface.cpp
+++ b/Engine/Plugins/FX/Niagara/Source/Niagara/Private/NiagaraDataInterface.cpp
@@ -97,178 +97,4 @@ void UNiagaraDataInterface::ValidateFunction(const FNiagaraFunctionSignature& Fu
 
 #endif
 
-//////////////////////////////////////////////////////////////////////////
-
-bool UNiagaraDataInterfaceCurveBase::CopyToInternal(UNiagaraDataInterface* Destination) const
-{
-	if (!Super::CopyToInternal(Destination))
-	{
-		return false;
-	}
-
-	UNiagaraDataInterfaceCurveBase* DestinationTyped = CastChecked<UNiagaraDataInterfaceCurveBase>(Destination);
-	DestinationTyped->bUseLUT = bUseLUT;
-	return true;
-}
-
-
-bool UNiagaraDataInterfaceCurveBase::CompareLUTS(const TArray<float>& OtherLUT) const
-{
-	if (ShaderLUT.Num() == OtherLUT.Num())
-	{
-		bool bMatched = true;
-		for (int32 i = 0; i < ShaderLUT.Num(); i++)
-		{
-			if (false == FMath::IsNearlyEqual(ShaderLUT[i], OtherLUT[i], 0.0001f))
-			{
-				bMatched = false;
-				UE_LOG(LogNiagara, Log, TEXT("First LUT mismatch found on comparison - LUT[%d] = %.9f  Other = %.9f \t%.9f"), i, ShaderLUT[i], OtherLUT[i], fabsf(ShaderLUT[i] - OtherLUT[i]));
-				break;
-			}
-		}
-		return bMatched;
-	}
-	else
-	{
-		UE_LOG(LogNiagara, Log, TEXT("Table sizes don't match"));
-		return false;
-	}
-}
-
-bool UNiagaraDataInterfaceCurveBase::Equals(const UNiagaraDataInterface* Other) const
-{
-	if (!Super::Equals(Other))
-	{
-		return false;
-	}
-
-	const UNiagaraDataInterfaceCurveBase* OtherTyped = CastChecked<UNiagaraDataInterfaceCurveBase>(Other);
-	bool bEqual = OtherTyped->bUseLUT == bUseLUT;
-	return bEqual;
-}
-
-
-void UNiagaraDataInterfaceCurveBase::GetParameterDefinitionHLSL(FNiagaraDataInterfaceGPUParamInfo& ParamInfo, FString& OutHLSL)
-{
-	OutHLSL += TEXT("\n");
-	FString MinTimeStr = TEXT("MinTime_") + ParamInfo.DataInterfaceHLSLSymbol;
-	FString MaxTimeStr = TEXT("MaxTime_") + ParamInfo.DataInterfaceHLSLSymbol;
-	FString InvTimeRangeStr = TEXT("InvTimeRange_") + ParamInfo.DataInterfaceHLSLSymbol;
-
-	FString BufferName = "CurveLUT_" + ParamInfo.DataInterfaceHLSLSymbol;
-	OutHLSL += TEXT("Buffer<float> ") + BufferName + TEXT(";\n");
-
-	OutHLSL += TEXT("float ") + MinTimeStr + TEXT(";\n");
-	OutHLSL += TEXT("float ") + MaxTimeStr + TEXT(";\n");
-	OutHLSL += TEXT("float ") + InvTimeRangeStr + TEXT(";\n");
-	OutHLSL += TEXT("\n");
-
-	//TODO: Create a Unitiliy/Common funcitons hlsl def shared between all instances of the same data interface class for these.
-	OutHLSL += FString::Printf(TEXT("float TimeToLUTFraction_%s(float T)\n{\n\treturn saturate((T - %s) * %s);\n}\n"), *ParamInfo.DataInterfaceHLSLSymbol, *MinTimeStr, *InvTimeRangeStr);
-	OutHLSL += FString::Printf(TEXT("float SampleCurve_%s(float T)\n{\n\treturn %s[(uint)T];\n}\n"), *ParamInfo.DataInterfaceHLSLSymbol, *BufferName);
-	OutHLSL += TEXT("\n");
-}
-
-//FReadBuffer& UNiagaraDataInterfaceCurveBase::GetCurveLUTGPUBuffer()
-//{
-//	//TODO: This isn't really very thread safe. Need to move to a proxy like system where DIs can push data to the RT safely.
-//	if (GPUBufferDirty)
-//	{
-//		int32 ElemSize = GetCurveNumElems();
-//		CurveLUT.Release();
-//		CurveLUT.Initialize(sizeof(float), CurveLUTWidth * ElemSize, EPixelFormat::PF_R32_FLOAT, BUF_Static);
-//		uint32 BufferSize = ShaderLUT.Num() * sizeof(float);
-//		int32 *BufferData = static_cast<int32*>(RHILockVertexBuffer(CurveLUT.Buffer, 0, BufferSize, EResourceLockMode::RLM_WriteOnly));
-//		FPlatformMemory::Memcpy(BufferData, ShaderLUT.GetData(), BufferSize);
-//		RHIUnlockVertexBuffer(CurveLUT.Buffer);
-//		GPUBufferDirty = false;
-//	}
-//
-//	return CurveLUT;
-//}
-
-void UNiagaraDataInterfaceCurveBase::PushToRenderThread()
-{
-	if (!GSupportsResourceView)
-	{
-		return;
-	}
-
-	FNiagaraDataInterfaceProxyCurveBase* RT_Proxy = GetProxyAs<FNiagaraDataInterfaceProxyCurveBase>();
-
-	int32 rtNumElems = GetCurveNumElems();
-	float rtLUTMinTime = this->LUTMinTime;
-	float rtLUTMaxTime = this->LUTMaxTime;
-	float rtLUTInvTimeRange = this->LUTInvTimeRange;
-
-	const uint32 rtCurveLUTWidth = CurveLUTWidth;
-
-	// @todo-threadsafety Gross.
-	TArray<float>* ShaderLut_RT = new TArray<float>(ShaderLUT);
-
-	// Push Updates to Proxy.
-	ENQUEUE_RENDER_COMMAND(FUpdateDIColorCurve)(
-		[RT_Proxy, rtCurveLUTWidth, rtNumElems, rtLUTMinTime, rtLUTMaxTime, rtLUTInvTimeRange, ShaderLut_RT](FRHICommandListImmediate& RHICmdList)
-	{
-		RT_Proxy->LUTMinTime = rtLUTMinTime;
-		RT_Proxy->LUTMaxTime = rtLUTMaxTime;
-		RT_Proxy->LUTInvTimeRange = rtLUTInvTimeRange;
-
-		int32 ElemSize = rtNumElems;
-		RT_Proxy->CurveLUT.Release();
-		RT_Proxy->CurveLUT.Initialize(sizeof(float), rtCurveLUTWidth * ElemSize, EPixelFormat::PF_R32_FLOAT, BUF_Static);
-		uint32 BufferSize = ShaderLut_RT->Num() * sizeof(float);
-		int32 *BufferData = static_cast<int32*>(RHILockVertexBuffer(RT_Proxy->CurveLUT.Buffer, 0, BufferSize, EResourceLockMode::RLM_WriteOnly));
-		FPlatformMemory::Memcpy(BufferData, ShaderLut_RT->GetData(), BufferSize);
-		RHIUnlockVertexBuffer(RT_Proxy->CurveLUT.Buffer);
-
-		delete ShaderLut_RT;
-	});
-}
-
-struct FNiagaraDataInterfaceParametersCS_Curve : public FNiagaraDataInterfaceParametersCS
-{
-	virtual ~FNiagaraDataInterfaceParametersCS_Curve() {}
-	virtual void Bind(const FNiagaraDataInterfaceParamRef& ParamRef, const class FShaderParameterMap& ParameterMap) override
-	{
-		MinTime.Bind(ParameterMap, *(TEXT("MinTime_") + ParamRef.ParameterInfo.DataInterfaceHLSLSymbol));
-		MaxTime.Bind(ParameterMap, *(TEXT("MaxTime_") + ParamRef.ParameterInfo.DataInterfaceHLSLSymbol));
-		InvTimeRange.Bind(ParameterMap, *(TEXT("InvTimeRange_") + ParamRef.ParameterInfo.DataInterfaceHLSLSymbol));
-		CurveLUT.Bind(ParameterMap, *(TEXT("CurveLUT_") + ParamRef.ParameterInfo.DataInterfaceHLSLSymbol));
-	}
-
-	virtual void Serialize(FArchive& Ar) override
-	{
-		Ar << MinTime;
-		Ar << MaxTime;
-		Ar << InvTimeRange;
-		Ar << CurveLUT;
-	}
-
-	virtual void Set(FRHICommandList& RHICmdList, const FNiagaraDataInterfaceSetArgs& Context) const override
-	{
-		check(IsInRenderingThread());
-
-		FRHIComputeShader* ComputeShaderRHI = Context.Shader->GetComputeShader();
-		FNiagaraDataInterfaceProxyCurveBase* CurveDI = (FNiagaraDataInterfaceProxyCurveBase*) Context.DataInterface;
-		FReadBuffer& CurveLUTBuffer = CurveDI->CurveLUT;
-
-		SetShaderValue(RHICmdList, ComputeShaderRHI, MinTime, CurveDI->LUTMinTime);
-		SetShaderValue(RHICmdList, ComputeShaderRHI, MaxTime, CurveDI->LUTMaxTime);
-		SetShaderValue(RHICmdList, ComputeShaderRHI, InvTimeRange, CurveDI->LUTInvTimeRange);
-		RHICmdList.SetShaderResourceViewParameter(ComputeShaderRHI, CurveLUT.GetBaseIndex(), CurveLUTBuffer.SRV);
-	}
-
-	FShaderParameter MinTime;
-	FShaderParameter MaxTime;
-	FShaderParameter InvTimeRange;
-	FShaderResourceParameter CurveLUT;
-};
-
-FNiagaraDataInterfaceParametersCS* UNiagaraDataInterfaceCurveBase::ConstructComputeParameters()const
-{
-	return new FNiagaraDataInterfaceParametersCS_Curve();
-}
-
-
 #undef LOCTEXT_NAMESPACE
diff --git a/Engine/Plugins/FX/Niagara/Source/Niagara/Private/NiagaraDataInterfaceColorCurve.cpp b/Engine/Plugins/FX/Niagara/Source/Niagara/Private/NiagaraDataInterfaceColorCurve.cpp
--- a/Engine/Plugins/FX/Niagara/Source/Niagara/Private/NiagaraDataInterfaceColorCurve.cpp
+++ b/Engine/Plugins/FX/Niagara/Source/Niagara/Private/NiagaraDataInterfaceColorCurve.cpp
@@ -5,7 +5,10 @@
 #include "Curves/CurveLinearColor.h"
 #include "Curves/CurveFloat.h"
 #include "NiagaraTypes.h"
-#include "NiagaraCustomVersion.h"
+
+#if WITH_EDITORONLY_DATA
+#include "Interfaces/ITargetPlatform.h"
+#endif
 
 //////////////////////////////////////////////////////////////////////////
 //Color Curve
@@ -31,37 +34,38 @@ void UNiagaraDataInterfaceColorCurve::PostInitProperties()
 	UpdateLUT();
 }
 
-void UNiagaraDataInterfaceColorCurve::PostLoad()
+void UNiagaraDataInterfaceColorCurve::Serialize(FArchive& Ar)
 {
-	Super::PostLoad();
-
-	const int32 NiagaraVer = GetLinkerCustomVersion(FNiagaraCustomVersion::GUID);
-
-	if (NiagaraVer < FNiagaraCustomVersion::LatestVersion)
+#if WITH_EDITORONLY_DATA
+	if (bUseLUT && Ar.IsCooking() && Ar.CookingTarget()->RequiresCookedData())
 	{
 		UpdateLUT();
+
+		FRichCurve TempRCurve;
+		FRichCurve TempGCurve;
+		FRichCurve TempBCurve;
+		FRichCurve TempACurve;
+		Exchange(RedCurve, TempRCurve);
+		Exchange(GreenCurve, TempGCurve);
+		Exchange(BlueCurve, TempBCurve);
+		Exchange(AlphaCurve, TempACurve);
+
+		Super::Serialize(Ar);
+
+		Exchange(RedCurve, TempRCurve);
+		Exchange(GreenCurve, TempGCurve);
+		Exchange(BlueCurve, TempBCurve);
+		Exchange(AlphaCurve, TempACurve);
 	}
 	else
-	{
-
-#if !UE_BUILD_SHIPPING
-		TArray<float> OldLUT = ShaderLUT;
-#endif
-		UpdateLUT();
-
-#if !UE_BUILD_SHIPPING
-		if (!CompareLUTS(OldLUT))
-		{
-			UE_LOG(LogNiagara, Log, TEXT("PostLoad LUT generation is out of sync. Please investigate. %s"), *GetPathName());
-		}
 #endif
+	{
+		Super::Serialize(Ar);
 	}
 }
 
-void UNiagaraDataInterfaceColorCurve::UpdateLUT()
+void UNiagaraDataInterfaceColorCurve::UpdateTimeRanges()
 {
-	ShaderLUT.Empty();
-
 	if ((RedCurve.GetNumKeys() > 0 || GreenCurve.GetNumKeys() > 0 || BlueCurve.GetNumKeys() > 0 || AlphaCurve.GetNumKeys() > 0))
 	{
 		LUTMinTime = FLT_MAX;
@@ -83,22 +87,24 @@ void UNiagaraDataInterfaceColorCurve::UpdateLUT()
 		LUTMaxTime = 1.0f;
 		LUTInvTimeRange = 1.0f;
 	}
+}
 
-	for (uint32 i = 0; i < CurveLUTWidth; i++)
+TArray<float> UNiagaraDataInterfaceColorCurve::BuildLUT(int32 NumEntries) const
+{
+	TArray<float> OutputLUT;
+	const float NumEntriesMinusOne = NumEntries - 1;
+
+	OutputLUT.Reserve(NumEntries * 4);
+	for (int32 i = 0; i < NumEntries; i++)
 	{
-		float X = UnnormalizeTime(i / (float)CurveLUTWidthMinusOne);
+		float X = UnnormalizeTime(float(i) / NumEntriesMinusOne);
 		FLinearColor C(RedCurve.Eval(X), GreenCurve.Eval(X), BlueCurve.Eval(X), AlphaCurve.Eval(X));
-		ShaderLUT.Add(C.R);
-		ShaderLUT.Add(C.G);
-		ShaderLUT.Add(C.B);
-		ShaderLUT.Add(C.A);
+		OutputLUT.Add(C.R);
+		OutputLUT.Add(C.G);
+		OutputLUT.Add(C.B);
+		OutputLUT.Add(C.A);
 	}
-
-	check(Proxy);
-
-	// @todo-threadsafety Not the best way to do this. Ideally we'd have a funnel where we can update all this data for all 
-	// interfaces at the same time.
-	Super::PushToRenderThread();
+	return OutputLUT;
 }
 
 bool UNiagaraDataInterfaceColorCurve::CopyToInternal(UNiagaraDataInterface* Destination) const 
@@ -164,12 +170,13 @@ bool UNiagaraDataInterfaceColorCurve::GetFunctionHLSL(const FName& DefinitionFun
 {
 	FString TimeToLUTFrac = TEXT("TimeToLUTFraction_") + ParamInfo.DataInterfaceHLSLSymbol;
 	FString Sample = TEXT("SampleCurve_") + ParamInfo.DataInterfaceHLSLSymbol;
+	FString NumSamples = TEXT("CurveLUTNumMinusOne_") + ParamInfo.DataInterfaceHLSLSymbol;
 	OutHLSL += FString::Printf(TEXT("\
 void %s(in float In_X, out float4 Out_Value) \n\
 { \n\
-	float RemappedX = %s(In_X) * %u; \n\
+	float RemappedX = %s(In_X) * %s; \n\
 	float Prev = floor(RemappedX); \n\
-	float Next = Prev < %u ? Prev + 1.0 : Prev; \n\
+	float Next = Prev < %s ? Prev + 1.0 : Prev; \n\
 	float Interp = RemappedX - Prev; \n\
 	Prev *= %u; \n\
 	Next *= %u; \n\
@@ -177,7 +184,7 @@ void %s(in float In_X, out float4 Out_Value) \n\
 	float4 B = float4(%s(Next), %s(Next + 1), %s(Next + 2), %s(Next + 3)); \n\
 	Out_Value = lerp(A, B, Interp); \n\
 }\n")
-, *InstanceFunctionName, *TimeToLUTFrac, CurveLUTWidthMinusOne, CurveLUTWidthMinusOne, CurveLUTNumElems, CurveLUTNumElems
+, *InstanceFunctionName, *TimeToLUTFrac, *NumSamples, *NumSamples, CurveLUTNumElems, CurveLUTNumElems
 , *Sample, *Sample, *Sample, *Sample, *Sample, *Sample, *Sample, *Sample);
 
 	return true;
@@ -201,9 +208,9 @@ void UNiagaraDataInterfaceColorCurve::GetVMExternalFunction(const FVMExternalFun
 template<>
 FORCEINLINE_DEBUGGABLE FLinearColor UNiagaraDataInterfaceColorCurve::SampleCurveInternal<TIntegralConstant<bool, true>>(float X)
 {
-	float RemappedX = FMath::Clamp(NormalizeTime(X) * CurveLUTWidthMinusOne, 0.0f, (float)CurveLUTWidthMinusOne);
+	float RemappedX = FMath::Clamp(NormalizeTime(X) * LUTNumSamplesMinusOne, 0.0f, LUTNumSamplesMinusOne);
 	float PrevEntry = FMath::TruncToFloat(RemappedX);
-	float NextEntry = PrevEntry < (float)CurveLUTWidthMinusOne ? PrevEntry + 1.0f : PrevEntry;
+	float NextEntry = PrevEntry < LUTNumSamplesMinusOne ? PrevEntry + 1.0f : PrevEntry;
 	float Interp = RemappedX - PrevEntry;
 	
 	int32 AIndex = PrevEntry * CurveLUTNumElems;
diff --git a/Engine/Plugins/FX/Niagara/Source/Niagara/Private/NiagaraDataInterfaceCurve.cpp b/Engine/Plugins/FX/Niagara/Source/Niagara/Private/NiagaraDataInterfaceCurve.cpp
--- a/Engine/Plugins/FX/Niagara/Source/Niagara/Private/NiagaraDataInterfaceCurve.cpp
+++ b/Engine/Plugins/FX/Niagara/Source/Niagara/Private/NiagaraDataInterfaceCurve.cpp
@@ -5,7 +5,10 @@
 #include "Curves/CurveLinearColor.h"
 #include "Curves/CurveFloat.h"
 #include "NiagaraTypes.h"
-#include "NiagaraCustomVersion.h"
+
+#if WITH_EDITORONLY_DATA
+#include "Interfaces/ITargetPlatform.h"
+#endif
 
 const FName UNiagaraDataInterfaceCurve::SampleCurveName(TEXT("SampleCurve"));
 
@@ -27,7 +30,28 @@ void UNiagaraDataInterfaceCurve::PostInitProperties()
 	UpdateLUT();
 }
 
-bool UNiagaraDataInterfaceCurve::CopyToInternal(UNiagaraDataInterface* Destination) const 
+void UNiagaraDataInterfaceCurve::Serialize(FArchive& Ar)
+{
+#if WITH_EDITORONLY_DATA
+	if (bUseLUT && Ar.IsCooking() && Ar.CookingTarget()->RequiresCookedData())
+	{
+		UpdateLUT();
+
+		FRichCurve TempCurve;
+		Exchange(Curve, TempCurve);
+
+		Super::Serialize(Ar);
+
+		Exchange(Curve, TempCurve);
+	}
+	else
+#endif
+	{
+		Super::Serialize(Ar);
+	}
+}
+
+bool UNiagaraDataInterfaceCurve::CopyToInternal(UNiagaraDataInterface* Destination) const
 {
 	if (!Super::CopyToInternal(Destination))
 	{
@@ -71,37 +95,8 @@ void UNiagaraDataInterfaceCurve::GetFunctions(TArray<FNiagaraFunctionSignature>&
 	OutFunctions.Add(Sig);
 }
 
-void UNiagaraDataInterfaceCurve::PostLoad()
+void UNiagaraDataInterfaceCurve::UpdateTimeRanges()
 {
-	Super::PostLoad();
-	const int32 NiagaraVer = GetLinkerCustomVersion(FNiagaraCustomVersion::GUID);
-
-	if (NiagaraVer < FNiagaraCustomVersion::LatestVersion)
-	{
-		UpdateLUT();
-	}
-	else
-	{
-
-#if !UE_BUILD_SHIPPING
-		TArray<float> OldLUT = ShaderLUT;
-#endif
-		UpdateLUT();
-
-
-#if !UE_BUILD_SHIPPING
-		if (!CompareLUTS(OldLUT))
-		{
-			UE_LOG(LogNiagara, Log, TEXT("PostLoad LUT generation is out of sync. Please investigate. %s"), *GetPathName());
-		}
-#endif
-	}
-}
-
-
-void UNiagaraDataInterfaceCurve::UpdateLUT()
-{
-	ShaderLUT.Empty();
 	if (Curve.GetNumKeys() > 0)
 	{
 		LUTMinTime = Curve.GetFirstKey().Time;
@@ -114,14 +109,21 @@ void UNiagaraDataInterfaceCurve::UpdateLUT()
 		LUTMaxTime = 1.0f;
 		LUTInvTimeRange = 1.0f;
 	}
+}
+
+TArray<float> UNiagaraDataInterfaceCurve::BuildLUT(int32 NumEntries) const
+{
+	TArray<float> OutputLUT;
+	const float NumEntriesMinusOne = NumEntries - 1;
 
-	for (uint32 i = 0; i < CurveLUTWidth; i++)
+	OutputLUT.Reserve(NumEntries);
+	for (int32 i = 0; i < NumEntries; i++)
 	{
-		float X = UnnormalizeTime(i / (float)CurveLUTWidthMinusOne);
+		float X = UnnormalizeTime(i / NumEntriesMinusOne);
 		float C = Curve.Eval(X);
-		ShaderLUT.Add(C);
+		OutputLUT.Add(C);
 	}
-	Super::PushToRenderThread();
+	return OutputLUT;
 }
 
 // build the shader function HLSL; function name is passed in, as it's defined per-DI; that way, configuration could change
@@ -132,18 +134,19 @@ bool UNiagaraDataInterfaceCurve::GetFunctionHLSL(const FName& DefinitionFunction
 {
 	FString TimeToLUTFrac = TEXT("TimeToLUTFraction_") + ParamInfo.DataInterfaceHLSLSymbol;
 	FString Sample = TEXT("SampleCurve_") + ParamInfo.DataInterfaceHLSLSymbol;
+	FString NumSamples = TEXT("CurveLUTNumMinusOne_") + ParamInfo.DataInterfaceHLSLSymbol;
 	OutHLSL += FString::Printf(TEXT("\
 void %s(in float In_X, out float Out_Value) \n\
 { \n\
-	float RemappedX = %s(In_X) * %u; \n\
+	float RemappedX = %s(In_X) * %s; \n\
 	float Prev = floor(RemappedX); \n\
-	float Next = Prev < %u ? Prev + 1.0 : Prev; \n\
+	float Next = Prev < %s ? Prev + 1.0 : Prev; \n\
 	float Interp = RemappedX - Prev; \n\
 	float A = %s(Prev); \n\
 	float B = %s(Next); \n\
 	Out_Value = lerp(A, B, Interp); \n\
 }\n")
-, *InstanceFunctionName, *TimeToLUTFrac, CurveLUTWidthMinusOne, CurveLUTWidthMinusOne, *Sample, *Sample);
+, *InstanceFunctionName, *TimeToLUTFrac, *NumSamples, *NumSamples, *Sample, *Sample);
 
 	return true;
 }
@@ -165,9 +168,9 @@ void UNiagaraDataInterfaceCurve::GetVMExternalFunction(const FVMExternalFunction
 template<>
 FORCEINLINE_DEBUGGABLE float UNiagaraDataInterfaceCurve::SampleCurveInternal<TIntegralConstant<bool, true>>(float X)
 {
-	float RemappedX = FMath::Clamp(NormalizeTime(X) * CurveLUTWidthMinusOne, 0.0f, (float)CurveLUTWidthMinusOne);
+	float RemappedX = FMath::Clamp(NormalizeTime(X) * LUTNumSamplesMinusOne, 0.0f, LUTNumSamplesMinusOne);
 	float PrevEntry = FMath::TruncToFloat(RemappedX);
-	float NextEntry = PrevEntry < (float)CurveLUTWidthMinusOne ? PrevEntry + 1.0f : PrevEntry;
+	float NextEntry = PrevEntry < LUTNumSamplesMinusOne ? PrevEntry + 1.0f : PrevEntry;
 	float Interp = RemappedX - PrevEntry;
 
 	int32 AIndex = PrevEntry * CurveLUTNumElems;
diff --git a/Engine/Plugins/FX/Niagara/Source/Niagara/Private/NiagaraDataInterfaceCurveBase.cpp b/Engine/Plugins/FX/Niagara/Source/Niagara/Private/NiagaraDataInterfaceCurveBase.cpp
new file mode 100644
--- /dev/null
+++ b/Engine/Plugins/FX/Niagara/Source/Niagara/Private/NiagaraDataInterfaceCurveBase.cpp
@@ -0,0 +1,314 @@
+// Copyright 1998-2019 Epic Games, Inc. All Rights Reserved.
+
+#include "NiagaraDataInterfaceCurveBase.h"
+#include "Curves/CurveVector.h"
+#include "Curves/CurveLinearColor.h"
+#include "Curves/CurveFloat.h"
+#include "NiagaraTypes.h"
+#include "ShaderParameterUtils.h"
+#include "NiagaraShader.h"
+#include "NiagaraCustomVersion.h"
+
+float GNiagaraLUTOptimizeThreshold = UNiagaraDataInterfaceCurveBase::DefaultOptimizeThreshold;
+static FAutoConsoleVariableRef CVarNiagaraLUTOptimizeThreshold(
+	TEXT("fx.Niagara.LUT.OptimizeThreshold"),
+	GNiagaraLUTOptimizeThreshold,
+	TEXT("Error Threshold used when optimizing Curve LUTs, setting to 0.0 or below will result in no optimization\n"),
+	ECVF_Default
+);
+
+int32 GNiagaraLUTVerifyPostLoad = 0;
+static FAutoConsoleVariableRef CVarNiagaraLUTVerifyPostLoad(
+	TEXT("fx.Niagara.LUT.VerifyPostLoad"),
+	GNiagaraLUTVerifyPostLoad,
+	TEXT("Enable to verify LUTs match in PostLoad vs the Loaded Data\n"),
+	ECVF_Default
+);
+
+void UNiagaraDataInterfaceCurveBase::PostLoad()
+{
+	Super::PostLoad();
+
+	const int32 NiagaraVer = GetLinkerCustomVersion(FNiagaraCustomVersion::GUID);
+
+	if (NiagaraVer < FNiagaraCustomVersion::LatestVersion)
+	{
+		UpdateLUT();
+	}
+	else
+	{
+		if (GIsEditor)
+		{
+			TArray<float> OldLUT = ShaderLUT;
+			UpdateLUT();
+			if (GNiagaraLUTVerifyPostLoad && !CompareLUTS(OldLUT))
+			{
+				UE_LOG(LogNiagara, Log, TEXT("PostLoad LUT generation is out of sync. Please investigate. %s"), *GetPathName());
+			}
+		}
+	}
+}
+
+bool UNiagaraDataInterfaceCurveBase::CopyToInternal(UNiagaraDataInterface* Destination) const
+{
+	if (!Super::CopyToInternal(Destination))
+	{
+		return false;
+	}
+
+	UNiagaraDataInterfaceCurveBase* DestinationTyped = CastChecked<UNiagaraDataInterfaceCurveBase>(Destination);
+	DestinationTyped->bUseLUT = bUseLUT;
+	DestinationTyped->bOptimizeLUT = bOptimizeLUT;
+	DestinationTyped->bOverrideOptimizeThreshold = bOverrideOptimizeThreshold;
+	DestinationTyped->OptimizeThreshold = OptimizeThreshold;
+
+	return true;
+}
+
+
+bool UNiagaraDataInterfaceCurveBase::CompareLUTS(const TArray<float>& OtherLUT) const
+{
+	if (ShaderLUT.Num() == OtherLUT.Num())
+	{
+		bool bMatched = true;
+		for (int32 i = 0; i < ShaderLUT.Num(); i++)
+		{
+			if (false == FMath::IsNearlyEqual(ShaderLUT[i], OtherLUT[i], 0.0001f))
+			{
+				bMatched = false;
+				UE_LOG(LogNiagara, Log, TEXT("First LUT mismatch found on comparison - LUT[%d] = %.9f  Other = %.9f \t%.9f"), i, ShaderLUT[i], OtherLUT[i], fabsf(ShaderLUT[i] - OtherLUT[i]));
+				break;
+			}
+		}
+		return bMatched;
+	}
+	else
+	{
+		UE_LOG(LogNiagara, Log, TEXT("Table sizes don't match"));
+		return false;
+	}
+}
+
+void UNiagaraDataInterfaceCurveBase::UpdateLUT()
+{
+	UpdateTimeRanges();
+	if (bUseLUT)
+	{
+		ShaderLUT = BuildLUT(CurveLUTDefaultWidth);
+		OptimizeLUT();
+		LUTNumSamplesMinusOne = float(ShaderLUT.Num() / GetCurveNumElems()) - 1;
+	}
+
+	if (!bUseLUT || (ShaderLUT.Num() == 0))
+	{
+		ShaderLUT.Empty(GetCurveNumElems());
+		ShaderLUT.AddDefaulted(GetCurveNumElems());
+		LUTNumSamplesMinusOne = 0;
+	}
+
+	PushToRenderThread();
+}
+
+void UNiagaraDataInterfaceCurveBase::OptimizeLUT()
+{
+	// Do we optimize this LUT?
+	if (!bOptimizeLUT)
+	{
+		return;
+	}
+
+	// Check error threshold is valid for us to optimize
+	const float ErrorThreshold = bOverrideOptimizeThreshold ? OptimizeThreshold : GNiagaraLUTOptimizeThreshold;
+	if (ErrorThreshold <= 0.0f)
+	{
+		return;
+	}
+
+	const int32 NumElements = GetCurveNumElems();
+	check((ShaderLUT.Num() % NumElements) == 0);
+
+	const int CurrNumSamples = ShaderLUT.Num() / NumElements;
+
+	for (int32 NewNumSamples = 1; NewNumSamples < CurrNumSamples; ++NewNumSamples)
+	{
+		TArray<float> ResampledLUT = BuildLUT(NewNumSamples);
+
+		bool bCanUseLUT = true;
+		for (int iSample = 0; iSample < CurrNumSamples; ++iSample)
+		{
+			const float NormalizedSampleTime = float(iSample) / float(CurrNumSamples - 1);
+
+			const float LhsInterp = FMath::Frac(NormalizedSampleTime * CurrNumSamples);
+			const int LhsSampleA = FMath::Min(FMath::FloorToInt(NormalizedSampleTime * CurrNumSamples), CurrNumSamples - 1);
+			const int LhsSampleB = FMath::Min(LhsSampleA + 1, CurrNumSamples - 1);
+
+			const float RhsInterp = FMath::Frac(NormalizedSampleTime * NewNumSamples);
+			const int RhsSampleA = FMath::Min(FMath::FloorToInt(NormalizedSampleTime * NewNumSamples), NewNumSamples - 1);
+			const int RhsSampleB = FMath::Min(RhsSampleA + 1, NewNumSamples - 1);
+
+			for (int iElement = 0; iElement < NumElements; ++iElement)
+			{
+				const float LhsValue = FMath::Lerp(ShaderLUT[LhsSampleA * NumElements + iElement], ShaderLUT[LhsSampleB * NumElements + iElement], LhsInterp);
+				const float RhsValue = FMath::Lerp(ResampledLUT[RhsSampleA * NumElements + iElement], ResampledLUT[RhsSampleB * NumElements + iElement], RhsInterp);
+				const float Error = FMath::Abs(LhsValue - RhsValue);
+				if (Error > ErrorThreshold)
+				{
+					bCanUseLUT = false;
+					break;
+				}
+			}
+
+			if (!bCanUseLUT)
+			{
+				break;
+			}
+		}
+
+		if (bCanUseLUT)
+		{
+			ShaderLUT = ResampledLUT;
+			break;
+		}
+	}
+}
+
+bool UNiagaraDataInterfaceCurveBase::Equals(const UNiagaraDataInterface* Other) const
+{
+	if (!Super::Equals(Other))
+	{
+		return false;
+	}
+
+	const UNiagaraDataInterfaceCurveBase* OtherTyped = CastChecked<UNiagaraDataInterfaceCurveBase>(Other);
+	bool bEqual = OtherTyped->bUseLUT == bUseLUT;
+	bEqual &= OtherTyped->bOptimizeLUT == bOptimizeLUT;
+	bEqual &= OtherTyped->bOverrideOptimizeThreshold == bOverrideOptimizeThreshold;
+	if (bOverrideOptimizeThreshold)
+	{
+		bEqual &= OtherTyped->OptimizeThreshold == OptimizeThreshold;
+	}
+
+	return bEqual;
+}
+
+
+void UNiagaraDataInterfaceCurveBase::GetParameterDefinitionHLSL(FNiagaraDataInterfaceGPUParamInfo& ParamInfo, FString& OutHLSL)
+{
+	OutHLSL += TEXT("\n");
+	FString MinTimeStr = TEXT("MinTime_") + ParamInfo.DataInterfaceHLSLSymbol;
+	FString MaxTimeStr = TEXT("MaxTime_") + ParamInfo.DataInterfaceHLSLSymbol;
+	FString InvTimeRangeStr = TEXT("InvTimeRange_") + ParamInfo.DataInterfaceHLSLSymbol;
+	FString CurveLUTNumMinusOneStr = TEXT("CurveLUTNumMinusOne_") + ParamInfo.DataInterfaceHLSLSymbol;
+
+	FString BufferName = "CurveLUT_" + ParamInfo.DataInterfaceHLSLSymbol;
+	OutHLSL += TEXT("Buffer<float> ") + BufferName + TEXT(";\n");
+
+	OutHLSL += TEXT("float ") + MinTimeStr + TEXT(";\n");
+	OutHLSL += TEXT("float ") + MaxTimeStr + TEXT(";\n");
+	OutHLSL += TEXT("float ") + InvTimeRangeStr + TEXT(";\n");
+	OutHLSL += TEXT("float ") + CurveLUTNumMinusOneStr + TEXT(";\n");
+	OutHLSL += TEXT("\n");
+
+	//TODO: Create a Unitiliy/Common funcitons hlsl def shared between all instances of the same data interface class for these.
+	OutHLSL += FString::Printf(TEXT("float TimeToLUTFraction_%s(float T)\n{\n\treturn saturate((T - %s) * %s);\n}\n"), *ParamInfo.DataInterfaceHLSLSymbol, *MinTimeStr, *InvTimeRangeStr);
+	OutHLSL += FString::Printf(TEXT("float SampleCurve_%s(float T)\n{\n\treturn %s[(uint)T];\n}\n"), *ParamInfo.DataInterfaceHLSLSymbol, *BufferName);
+	OutHLSL += TEXT("\n");
+}
+
+//FReadBuffer& UNiagaraDataInterfaceCurveBase::GetCurveLUTGPUBuffer()
+//{
+//	//TODO: This isn't really very thread safe. Need to move to a proxy like system where DIs can push data to the RT safely.
+//	if (GPUBufferDirty)
+//	{
+//		int32 ElemSize = GetCurveNumElems();
+//		CurveLUT.Release();
+//		CurveLUT.Initialize(sizeof(float), CurveLUTWidth * ElemSize, EPixelFormat::PF_R32_FLOAT, BUF_Static);
+//		uint32 BufferSize = ShaderLUT.Num() * sizeof(float);
+//		int32 *BufferData = static_cast<int32*>(RHILockVertexBuffer(CurveLUT.Buffer, 0, BufferSize, EResourceLockMode::RLM_WriteOnly));
+//		FPlatformMemory::Memcpy(BufferData, ShaderLUT.GetData(), BufferSize);
+//		RHIUnlockVertexBuffer(CurveLUT.Buffer);
+//		GPUBufferDirty = false;
+//	}
+//
+//	return CurveLUT;
+//}
+
+void UNiagaraDataInterfaceCurveBase::PushToRenderThread()
+{
+	if (!GSupportsResourceView)
+	{
+		return;
+	}
+
+	FNiagaraDataInterfaceProxyCurveBase* RT_Proxy = GetProxyAs<FNiagaraDataInterfaceProxyCurveBase>();
+
+	const float rtLUTMinTime = this->LUTMinTime;
+	const float rtLUTMaxTime = this->LUTMaxTime;
+	const float rtLUTInvTimeRange = this->LUTInvTimeRange;
+	const float rtCurveLUTNumMinusOne = this->LUTNumSamplesMinusOne;;
+
+	// Push Updates to Proxy.
+	ENQUEUE_RENDER_COMMAND(FUpdateDIColorCurve)(
+		[RT_Proxy, rtLUTMinTime, rtLUTMaxTime, rtLUTInvTimeRange, rtCurveLUTNumMinusOne, rtShaderLUT = ShaderLUT](FRHICommandListImmediate& RHICmdList)
+	{
+		RT_Proxy->LUTMinTime = rtLUTMinTime;
+		RT_Proxy->LUTMaxTime = rtLUTMaxTime;
+		RT_Proxy->LUTInvTimeRange = rtLUTInvTimeRange;
+		RT_Proxy->CurveLUTNumMinusOne = rtCurveLUTNumMinusOne;
+
+		RT_Proxy->CurveLUT.Release();
+		RT_Proxy->CurveLUT.Initialize(sizeof(float), rtShaderLUT.Num(), EPixelFormat::PF_R32_FLOAT, BUF_Static);
+		uint32 BufferSize = rtShaderLUT.Num() * sizeof(float);
+		int32 *BufferData = static_cast<int32*>(RHILockVertexBuffer(RT_Proxy->CurveLUT.Buffer, 0, BufferSize, EResourceLockMode::RLM_WriteOnly));
+		FPlatformMemory::Memcpy(BufferData, rtShaderLUT.GetData(), BufferSize);
+		RHIUnlockVertexBuffer(RT_Proxy->CurveLUT.Buffer);
+	});
+}
+
+struct FNiagaraDataInterfaceParametersCS_Curve : public FNiagaraDataInterfaceParametersCS
+{
+	virtual ~FNiagaraDataInterfaceParametersCS_Curve() {}
+	virtual void Bind(const FNiagaraDataInterfaceParamRef& ParamRef, const class FShaderParameterMap& ParameterMap) override
+	{
+		MinTime.Bind(ParameterMap, *(TEXT("MinTime_") + ParamRef.ParameterInfo.DataInterfaceHLSLSymbol));
+		MaxTime.Bind(ParameterMap, *(TEXT("MaxTime_") + ParamRef.ParameterInfo.DataInterfaceHLSLSymbol));
+		InvTimeRange.Bind(ParameterMap, *(TEXT("InvTimeRange_") + ParamRef.ParameterInfo.DataInterfaceHLSLSymbol));
+		CurveLUTNumMinusOne.Bind(ParameterMap, *(TEXT("CurveLUTNumMinusOne_") + ParamRef.ParameterInfo.DataInterfaceHLSLSymbol));
+		CurveLUT.Bind(ParameterMap, *(TEXT("CurveLUT_") + ParamRef.ParameterInfo.DataInterfaceHLSLSymbol));
+	}
+
+	virtual void Serialize(FArchive& Ar) override
+	{
+		Ar << MinTime;
+		Ar << MaxTime;
+		Ar << InvTimeRange;
+		Ar << CurveLUTNumMinusOne;
+		Ar << CurveLUT;
+	}
+
+	virtual void Set(FRHICommandList& RHICmdList, const FNiagaraDataInterfaceSetArgs& Context) const override
+	{
+		check(IsInRenderingThread());
+
+		FRHIComputeShader* ComputeShaderRHI = Context.Shader->GetComputeShader();
+		FNiagaraDataInterfaceProxyCurveBase* CurveDI = (FNiagaraDataInterfaceProxyCurveBase*)Context.DataInterface;
+		FReadBuffer& CurveLUTBuffer = CurveDI->CurveLUT;
+
+		SetShaderValue(RHICmdList, ComputeShaderRHI, MinTime, CurveDI->LUTMinTime);
+		SetShaderValue(RHICmdList, ComputeShaderRHI, MaxTime, CurveDI->LUTMaxTime);
+		SetShaderValue(RHICmdList, ComputeShaderRHI, InvTimeRange, CurveDI->LUTInvTimeRange);
+		SetShaderValue(RHICmdList, ComputeShaderRHI, CurveLUTNumMinusOne, CurveDI->CurveLUTNumMinusOne);
+		RHICmdList.SetShaderResourceViewParameter(ComputeShaderRHI, CurveLUT.GetBaseIndex(), CurveLUTBuffer.SRV);
+	}
+
+	FShaderParameter MinTime;
+	FShaderParameter MaxTime;
+	FShaderParameter InvTimeRange;
+	FShaderParameter CurveLUTNumMinusOne;
+	FShaderResourceParameter CurveLUT;
+};
+
+FNiagaraDataInterfaceParametersCS* UNiagaraDataInterfaceCurveBase::ConstructComputeParameters()const
+{
+	return new FNiagaraDataInterfaceParametersCS_Curve();
+}
diff --git a/Engine/Plugins/FX/Niagara/Source/Niagara/Private/NiagaraDataInterfaceVector2DCurve.cpp b/Engine/Plugins/FX/Niagara/Source/Niagara/Private/NiagaraDataInterfaceVector2DCurve.cpp
--- a/Engine/Plugins/FX/Niagara/Source/Niagara/Private/NiagaraDataInterfaceVector2DCurve.cpp
+++ b/Engine/Plugins/FX/Niagara/Source/Niagara/Private/NiagaraDataInterfaceVector2DCurve.cpp
@@ -5,7 +5,10 @@
 #include "Curves/CurveLinearColor.h"
 #include "Curves/CurveFloat.h"
 #include "NiagaraTypes.h"
-#include "NiagaraCustomVersion.h"
+
+#if WITH_EDITORONLY_DATA
+#include "Interfaces/ITargetPlatform.h"
+#endif
 
 //////////////////////////////////////////////////////////////////////////
 //Vector2D Curve
@@ -17,6 +20,7 @@ UNiagaraDataInterfaceVector2DCurve::UNiagaraDataInterfaceVector2DCurve(FObjectIn
 {
 	UpdateLUT();
 }
+
 void UNiagaraDataInterfaceVector2DCurve::PostInitProperties()
 {
 	Super::PostInitProperties();
@@ -29,35 +33,32 @@ void UNiagaraDataInterfaceVector2DCurve::PostInitProperties()
 	UpdateLUT();
 }
 
-void UNiagaraDataInterfaceVector2DCurve::PostLoad()
+void UNiagaraDataInterfaceVector2DCurve::Serialize(FArchive& Ar)
 {
-	Super::PostLoad();
-
-	const int32 NiagaraVer = GetLinkerCustomVersion(FNiagaraCustomVersion::GUID);
-
-	if (NiagaraVer < FNiagaraCustomVersion::LatestVersion)
+#if WITH_EDITORONLY_DATA
+	if (bUseLUT && Ar.IsCooking() && Ar.CookingTarget()->RequiresCookedData())
 	{
 		UpdateLUT();
+
+		FRichCurve TempXCurve;
+		FRichCurve TempYCurve;
+		Exchange(XCurve, TempXCurve);
+		Exchange(YCurve, TempYCurve);
+
+		Super::Serialize(Ar);
+
+		Exchange(XCurve, TempXCurve);
+		Exchange(YCurve, TempYCurve);
 	}
 	else
-	{
-#if !UE_BUILD_SHIPPING
-		TArray<float> OldLUT = ShaderLUT;
-#endif
-		UpdateLUT();
-
-#if !UE_BUILD_SHIPPING
-		if (!CompareLUTS(OldLUT))
-		{
-			UE_LOG(LogNiagara, Log, TEXT("PostLoad LUT generation is out of sync. Please investigate. %s"), *GetPathName());
-		}
 #endif
+	{
+		Super::Serialize(Ar);
 	}
 }
 
-void UNiagaraDataInterfaceVector2DCurve::UpdateLUT()
+void UNiagaraDataInterfaceVector2DCurve::UpdateTimeRanges()
 {
-	ShaderLUT.Empty();
 	if ((XCurve.GetNumKeys() > 0 || YCurve.GetNumKeys() > 0))
 	{
 		LUTMinTime = FLT_MAX;
@@ -75,16 +76,22 @@ void UNiagaraDataInterfaceVector2DCurve::UpdateLUT()
 		LUTMaxTime = 1.0f;
 		LUTInvTimeRange = 1.0f;
 	}
+}
+
+TArray<float> UNiagaraDataInterfaceVector2DCurve::BuildLUT(int32 NumEntries) const
+{
+	TArray<float> OutputLUT;
+	const float NumEntriesMinusOne = NumEntries - 1;
 
-	for (uint32 i = 0; i < CurveLUTWidth; i++)
+	OutputLUT.Reserve(NumEntries * 2);
+	for (int32 i = 0; i < NumEntries; i++)
 	{
-		float X = UnnormalizeTime(i / (float)CurveLUTWidthMinusOne);
+		float X = UnnormalizeTime(i / NumEntriesMinusOne);
 		FVector2D C(XCurve.Eval(X), YCurve.Eval(X));
-		ShaderLUT.Add(C.X);
-		ShaderLUT.Add(C.Y);
+		OutputLUT.Add(C.X);
+		OutputLUT.Add(C.Y);
 	}
-
-	Super::PushToRenderThread();
+	return OutputLUT;
 }
 
 bool UNiagaraDataInterfaceVector2DCurve::CopyToInternal(UNiagaraDataInterface* Destination) const
@@ -142,12 +149,13 @@ bool UNiagaraDataInterfaceVector2DCurve::GetFunctionHLSL(const FName& Definition
 {
 	FString TimeToLUTFrac = TEXT("TimeToLUTFraction_") + ParamInfo.DataInterfaceHLSLSymbol;
 	FString Sample = TEXT("SampleCurve_") + ParamInfo.DataInterfaceHLSLSymbol;
+	FString NumSamples = TEXT("CurveLUTNumMinusOne_") + ParamInfo.DataInterfaceHLSLSymbol;
 	OutHLSL += FString::Printf(TEXT("\
 void %s(in float In_X, out float2 Out_Value) \n\
 { \n\
-	float RemappedX = %s(In_X) * %u; \n\
+	float RemappedX = %s(In_X) * %s; \n\
 	float Prev = floor(RemappedX); \n\
-	float Next = Prev < %u ? Prev + 1.0 : Prev; \n\
+	float Next = Prev < %s ? Prev + 1.0 : Prev; \n\
 	float Interp = RemappedX - Prev; \n\
 	Prev *= %u; \n\
 	Next *= %u; \n\
@@ -155,7 +163,7 @@ void %s(in float In_X, out float2 Out_Value) \n\
 	float2 B = float2(%s(Next), %s(Next + 1)); \n\
 	Out_Value = lerp(A, B, Interp); \n\
 }\n")
-, *InstanceFunctionName, *TimeToLUTFrac, CurveLUTWidthMinusOne, CurveLUTWidthMinusOne, CurveLUTNumElems, CurveLUTNumElems, *Sample, *Sample, *Sample, *Sample);
+, *InstanceFunctionName, *TimeToLUTFrac, *NumSamples, *NumSamples, CurveLUTNumElems, CurveLUTNumElems, *Sample, *Sample, *Sample, *Sample);
 
 	return true;
 }
@@ -177,9 +185,9 @@ void UNiagaraDataInterfaceVector2DCurve::GetVMExternalFunction(const FVMExternal
 template<>
 FORCEINLINE_DEBUGGABLE FVector2D UNiagaraDataInterfaceVector2DCurve::SampleCurveInternal<TIntegralConstant<bool, true>>(float X)
 {
-	float RemappedX = FMath::Clamp(NormalizeTime(X) * CurveLUTWidthMinusOne, 0.0f, (float)CurveLUTWidthMinusOne);
+	float RemappedX = FMath::Clamp(NormalizeTime(X) * LUTNumSamplesMinusOne, 0.0f, LUTNumSamplesMinusOne);
 	float PrevEntry = FMath::TruncToFloat(RemappedX);
-	float NextEntry = PrevEntry < (float)CurveLUTWidthMinusOne ? PrevEntry + 1.0f : PrevEntry;
+	float NextEntry = PrevEntry < LUTNumSamplesMinusOne ? PrevEntry + 1.0f : PrevEntry;
 	float Interp = RemappedX - PrevEntry;
 
 	int32 AIndex = PrevEntry * CurveLUTNumElems;
diff --git a/Engine/Plugins/FX/Niagara/Source/Niagara/Private/NiagaraDataInterfaceVector4Curve.cpp b/Engine/Plugins/FX/Niagara/Source/Niagara/Private/NiagaraDataInterfaceVector4Curve.cpp
--- a/Engine/Plugins/FX/Niagara/Source/Niagara/Private/NiagaraDataInterfaceVector4Curve.cpp
+++ b/Engine/Plugins/FX/Niagara/Source/Niagara/Private/NiagaraDataInterfaceVector4Curve.cpp
@@ -5,7 +5,10 @@
 #include "Curves/CurveLinearColor.h"
 #include "Curves/CurveFloat.h"
 #include "NiagaraTypes.h"
-#include "NiagaraCustomVersion.h"
+
+#if WITH_EDITORONLY_DATA
+#include "Interfaces/ITargetPlatform.h"
+#endif
 
 //////////////////////////////////////////////////////////////////////////
 //Color Curve
@@ -22,7 +25,7 @@ void UNiagaraDataInterfaceVector4Curve::PostInitProperties()
 {
 	Super::PostInitProperties();
 
-	//Can we regitser data interfaces as regular types and fold them into the FNiagaraVariable framework for UI and function calls etc?
+	//Can we register data interfaces as regular types and fold them into the FNiagaraVariable framework for UI and function calls etc?
 	if (HasAnyFlags(RF_ClassDefaultObject))
 	{
 		FNiagaraTypeRegistry::Register(FNiagaraTypeDefinition(GetClass()), true, false, false);
@@ -31,36 +34,38 @@ void UNiagaraDataInterfaceVector4Curve::PostInitProperties()
 	UpdateLUT();
 }
 
-void UNiagaraDataInterfaceVector4Curve::PostLoad()
+void UNiagaraDataInterfaceVector4Curve::Serialize(FArchive& Ar)
 {
-	Super::PostLoad();
-
-	const int32 NiagaraVer = GetLinkerCustomVersion(FNiagaraCustomVersion::GUID);
-
-	if (NiagaraVer < FNiagaraCustomVersion::LatestVersion)
+#if WITH_EDITORONLY_DATA
+	if (bUseLUT && Ar.IsCooking() && Ar.CookingTarget()->RequiresCookedData())
 	{
 		UpdateLUT();
+
+		FRichCurve TempXCurve;
+		FRichCurve TempYCurve;
+		FRichCurve TempZCurve;
+		FRichCurve TempWCurve;
+		Exchange(XCurve, TempXCurve);
+		Exchange(YCurve, TempYCurve);
+		Exchange(ZCurve, TempZCurve);
+		Exchange(WCurve, TempWCurve);
+
+		Super::Serialize(Ar);
+
+		Exchange(XCurve, TempXCurve);
+		Exchange(YCurve, TempYCurve);
+		Exchange(ZCurve, TempZCurve);
+		Exchange(WCurve, TempWCurve);
 	}
 	else
-	{
-#if !UE_BUILD_SHIPPING
-		TArray<float> OldLUT = ShaderLUT;
-#endif
-		UpdateLUT();
-
-#if !UE_BUILD_SHIPPING
-		if (!CompareLUTS(OldLUT))
-		{
-			UE_LOG(LogNiagara, Log, TEXT("PostLoad LUT generation is out of sync. Please investigate. %s"), *GetPathName());
-		}
 #endif
+	{
+		Super::Serialize(Ar);
 	}
 }
 
-void UNiagaraDataInterfaceVector4Curve::UpdateLUT()
+void UNiagaraDataInterfaceVector4Curve::UpdateTimeRanges()
 {
-	ShaderLUT.Empty();
-
 	if ((XCurve.GetNumKeys() > 0 || YCurve.GetNumKeys() > 0 || ZCurve.GetNumKeys() > 0 || WCurve.GetNumKeys() > 0))
 	{
 		LUTMinTime = FLT_MAX;
@@ -82,17 +87,24 @@ void UNiagaraDataInterfaceVector4Curve::UpdateLUT()
 		LUTMaxTime = 1.0f;
 		LUTInvTimeRange = 1.0f;
 	}
+}
 
-	for (uint32 i = 0; i < CurveLUTWidth; i++)
+TArray<float> UNiagaraDataInterfaceVector4Curve::BuildLUT(int32 NumEntries) const
+{
+	TArray<float> OutputLUT;
+	const float NumEntriesMinusOne = NumEntries - 1;
+
+	OutputLUT.Reserve(NumEntries * 4);
+	for (int32 i = 0; i < NumEntries; i++)
 	{
-		float X = UnnormalizeTime(i / (float)CurveLUTWidthMinusOne);
+		float X = UnnormalizeTime(i / NumEntriesMinusOne);
 		FLinearColor C(XCurve.Eval(X), YCurve.Eval(X), ZCurve.Eval(X), WCurve.Eval(X));
-		ShaderLUT.Add(C.R);
-		ShaderLUT.Add(C.G);
-		ShaderLUT.Add(C.B);
-		ShaderLUT.Add(C.A);
+		OutputLUT.Add(C.R);
+		OutputLUT.Add(C.G);
+		OutputLUT.Add(C.B);
+		OutputLUT.Add(C.A);
 	}
-	Super::PushToRenderThread();
+	return OutputLUT;
 }
 
 bool UNiagaraDataInterfaceVector4Curve::CopyToInternal(UNiagaraDataInterface* Destination) const 
@@ -107,7 +119,10 @@ bool UNiagaraDataInterfaceVector4Curve::CopyToInternal(UNiagaraDataInterface* De
 	DestinationColorCurve->ZCurve = ZCurve;
 	DestinationColorCurve->WCurve = WCurve;
 	DestinationColorCurve->UpdateLUT();
-	ensure(CompareLUTS(DestinationColorCurve->ShaderLUT));
+	if (!CompareLUTS(DestinationColorCurve->ShaderLUT))
+	{
+		UE_LOG(LogNiagara, Log, TEXT("Post CopyToInternal LUT generation is out of sync. Please investigate. %s"), *GetPathName());
+	}
 
 	return true;
 }
@@ -155,12 +170,13 @@ bool UNiagaraDataInterfaceVector4Curve::GetFunctionHLSL(const FName& DefinitionF
 {
 	FString TimeToLUTFrac = TEXT("TimeToLUTFraction_") + ParamInfo.DataInterfaceHLSLSymbol;
 	FString Sample = TEXT("SampleCurve_") + ParamInfo.DataInterfaceHLSLSymbol;
+	FString NumSamples = TEXT("CurveLUTNumMinusOne_") + ParamInfo.DataInterfaceHLSLSymbol;
 	OutHLSL += FString::Printf(TEXT("\
 void %s(in float In_X, out float4 Out_Value) \n\
 { \n\
-	float RemappedX = %s(In_X) * %u; \n\
+	float RemappedX = %s(In_X) * %s; \n\
 	float Prev = floor(RemappedX); \n\
-	float Next = Prev < %u ? Prev + 1.0 : Prev; \n\
+	float Next = Prev < %s ? Prev + 1.0 : Prev; \n\
 	float Interp = RemappedX - Prev; \n\
 	Prev *= %u; \n\
 	Next *= %u; \n\
@@ -168,7 +184,7 @@ void %s(in float In_X, out float4 Out_Value) \n\
 	float4 B = float4(%s(Next), %s(Next + 1), %s(Next + 2), %s(Next + 3)); \n\
 	Out_Value = lerp(A, B, Interp); \n\
 }\n")
-, *InstanceFunctionName, *TimeToLUTFrac, CurveLUTWidthMinusOne, CurveLUTWidthMinusOne, CurveLUTNumElems, CurveLUTNumElems
+, *InstanceFunctionName, *TimeToLUTFrac, *NumSamples, *NumSamples, CurveLUTNumElems, CurveLUTNumElems
 , *Sample, *Sample, *Sample, *Sample, *Sample, *Sample, *Sample, *Sample);
 
 	return true;
@@ -191,9 +207,9 @@ void UNiagaraDataInterfaceVector4Curve::GetVMExternalFunction(const FVMExternalF
 template<>
 FORCEINLINE_DEBUGGABLE FVector4 UNiagaraDataInterfaceVector4Curve::SampleCurveInternal<TIntegralConstant<bool, true>>(float X)
 {
-	float RemappedX = FMath::Clamp(NormalizeTime(X) * CurveLUTWidthMinusOne, 0.0f, (float)CurveLUTWidthMinusOne);
+	float RemappedX = FMath::Clamp(NormalizeTime(X) * LUTNumSamplesMinusOne, 0.0f, LUTNumSamplesMinusOne);
 	float PrevEntry = FMath::TruncToFloat(RemappedX);
-	float NextEntry = PrevEntry < (float)CurveLUTWidthMinusOne ? PrevEntry + 1.0f : PrevEntry;
+	float NextEntry = PrevEntry < LUTNumSamplesMinusOne ? PrevEntry + 1.0f : PrevEntry;
 	float Interp = RemappedX - PrevEntry;
 
 	int32 AIndex = PrevEntry * CurveLUTNumElems;
diff --git a/Engine/Plugins/FX/Niagara/Source/Niagara/Private/NiagaraDataInterfaceVectorCurve.cpp b/Engine/Plugins/FX/Niagara/Source/Niagara/Private/NiagaraDataInterfaceVectorCurve.cpp
--- a/Engine/Plugins/FX/Niagara/Source/Niagara/Private/NiagaraDataInterfaceVectorCurve.cpp
+++ b/Engine/Plugins/FX/Niagara/Source/Niagara/Private/NiagaraDataInterfaceVectorCurve.cpp
@@ -5,7 +5,10 @@
 #include "Curves/CurveLinearColor.h"
 #include "Curves/CurveFloat.h"
 #include "NiagaraTypes.h"
-#include "NiagaraCustomVersion.h"
+
+#if WITH_EDITORONLY_DATA
+#include "Interfaces/ITargetPlatform.h"
+#endif
 
 //////////////////////////////////////////////////////////////////////////
 //Vector Curve
@@ -22,7 +25,7 @@ void UNiagaraDataInterfaceVectorCurve::PostInitProperties()
 {
 	Super::PostInitProperties();
 
-	//Can we regitser data interfaces as regular types and fold them into the FNiagaraVariable framework for UI and function calls etc?
+	//Can we register data interfaces as regular types and fold them into the FNiagaraVariable framework for UI and function calls etc?
 	if (HasAnyFlags(RF_ClassDefaultObject))
 	{
 		FNiagaraTypeRegistry::Register(FNiagaraTypeDefinition(GetClass()), true, false, false);
@@ -31,34 +34,35 @@ void UNiagaraDataInterfaceVectorCurve::PostInitProperties()
 	UpdateLUT();
 }
 
-void UNiagaraDataInterfaceVectorCurve::PostLoad()
+void UNiagaraDataInterfaceVectorCurve::Serialize(FArchive& Ar)
 {
-	Super::PostLoad();
-
-	const int32 NiagaraVer = GetLinkerCustomVersion(FNiagaraCustomVersion::GUID);
-
-	if (NiagaraVer < FNiagaraCustomVersion::LatestVersion)
+#if WITH_EDITORONLY_DATA
+	if ( bUseLUT && Ar.IsCooking() && Ar.CookingTarget()->RequiresCookedData() )
 	{
 		UpdateLUT();
+
+		FRichCurve TempXCurve;
+		FRichCurve TempYCurve;
+		FRichCurve TempZCurve;
+		Exchange(XCurve, TempXCurve);
+		Exchange(YCurve, TempYCurve);
+		Exchange(ZCurve, TempZCurve);
+
+		Super::Serialize(Ar);
+
+		Exchange(XCurve, TempXCurve);
+		Exchange(YCurve, TempYCurve);
+		Exchange(ZCurve, TempZCurve);
 	}
 	else
-	{
-#if !UE_BUILD_SHIPPING
-		TArray<float> OldLUT = ShaderLUT;
-#endif
-		UpdateLUT();
-#if !UE_BUILD_SHIPPING
-		if (!CompareLUTS(OldLUT))
-		{
-			UE_LOG(LogNiagara, Log, TEXT("PostLoad LUT generation is out of sync. Please investigate. %s"), *GetPathName());
-		}
 #endif
+	{
+		Super::Serialize(Ar);
 	}
 }
 
-void UNiagaraDataInterfaceVectorCurve::UpdateLUT()
+void UNiagaraDataInterfaceVectorCurve::UpdateTimeRanges()
 {
-	ShaderLUT.Empty();
 	if ((XCurve.GetNumKeys() > 0 || YCurve.GetNumKeys() > 0 || ZCurve.GetNumKeys() > 0))
 	{
 		LUTMinTime = FLT_MAX;
@@ -78,16 +82,23 @@ void UNiagaraDataInterfaceVectorCurve::UpdateLUT()
 		LUTMaxTime = 1.0f;
 		LUTInvTimeRange = 1.0f;
 	}
+}
+
+TArray<float> UNiagaraDataInterfaceVectorCurve::BuildLUT(int32 NumEntries) const
+{
+	TArray<float> OutputLUT;
+	const float NumEntriesMinusOne = NumEntries - 1;
 
-	for (uint32 i = 0; i < CurveLUTWidth; i++)
+	OutputLUT.Reserve(NumEntries * 3);
+	for (int32 i=0; i < NumEntries; i++)
 	{
-		float X = UnnormalizeTime(i / (float)CurveLUTWidthMinusOne);
+		float X = UnnormalizeTime(i / NumEntriesMinusOne);
 		FVector C(XCurve.Eval(X), YCurve.Eval(X), ZCurve.Eval(X));
-		ShaderLUT.Add(C.X);
-		ShaderLUT.Add(C.Y);
-		ShaderLUT.Add(C.Z);
+		OutputLUT.Add(C.X);
+		OutputLUT.Add(C.Y);
+		OutputLUT.Add(C.Z);
 	}
-	Super::PushToRenderThread();
+	return OutputLUT;
 }
 
 bool UNiagaraDataInterfaceVectorCurve::CopyToInternal(UNiagaraDataInterface* Destination) const 
@@ -150,12 +161,13 @@ bool UNiagaraDataInterfaceVectorCurve::GetFunctionHLSL(const FName& DefinitionFu
 {
 	FString TimeToLUTFrac = TEXT("TimeToLUTFraction_") + ParamInfo.DataInterfaceHLSLSymbol;
 	FString Sample = TEXT("SampleCurve_") + ParamInfo.DataInterfaceHLSLSymbol;
+	FString NumSamples = TEXT("CurveLUTNumMinusOne_") + ParamInfo.DataInterfaceHLSLSymbol;
 	OutHLSL += FString::Printf(TEXT("\
 void %s(in float In_X, out float3 Out_Value) \n\
 { \n\
-	float RemappedX = %s(In_X) * %u; \n\
+	float RemappedX = %s(In_X) * %s; \n\
 	float Prev = floor(RemappedX); \n\
-	float Next = Prev < %u ? Prev + 1.0 : Prev; \n\
+	float Next = Prev < %s ? Prev + 1.0 : Prev; \n\
 	float Interp = RemappedX - Prev; \n\
 	Prev *= %u; \n\
 	Next *= %u; \n\
@@ -163,7 +175,7 @@ void %s(in float In_X, out float3 Out_Value) \n\
 	float3 B = float3(%s(Next), %s(Next + 1), %s(Next + 2)); \n\
 	Out_Value = lerp(A, B, Interp); \n\
 }\n") 
-, *InstanceFunctionName, *TimeToLUTFrac, CurveLUTWidthMinusOne, CurveLUTWidthMinusOne, CurveLUTNumElems, CurveLUTNumElems
+, *InstanceFunctionName, *TimeToLUTFrac, *NumSamples, *NumSamples, CurveLUTNumElems, CurveLUTNumElems
 , *Sample, *Sample, *Sample, *Sample, *Sample, *Sample);
 
 	return true;
@@ -186,9 +198,9 @@ void UNiagaraDataInterfaceVectorCurve::GetVMExternalFunction(const FVMExternalFu
 template<>
 FORCEINLINE_DEBUGGABLE FVector UNiagaraDataInterfaceVectorCurve::SampleCurveInternal<TIntegralConstant<bool, true>>(float X)
 {
-	float RemappedX = FMath::Clamp(NormalizeTime(X) * CurveLUTWidthMinusOne, 0.0f, (float)CurveLUTWidthMinusOne);
+	float RemappedX = FMath::Clamp(NormalizeTime(X) * LUTNumSamplesMinusOne, 0.0f, LUTNumSamplesMinusOne);
 	float PrevEntry = FMath::TruncToFloat(RemappedX);
-	float NextEntry = PrevEntry < (float)CurveLUTWidthMinusOne ? PrevEntry + 1.0f : PrevEntry;
+	float NextEntry = PrevEntry < LUTNumSamplesMinusOne ? PrevEntry + 1.0f : PrevEntry;
 	float Interp = RemappedX - PrevEntry;
 
 	int32 AIndex = PrevEntry * CurveLUTNumElems;
diff --git a/Engine/Plugins/FX/Niagara/Source/NiagaraCore/Private/NiagaraCustomVersion.cpp b/Engine/Plugins/FX/Niagara/Source/NiagaraCore/Private/NiagaraCustomVersion.cpp
--- a/Engine/Plugins/FX/Niagara/Source/NiagaraCore/Private/NiagaraCustomVersion.cpp
+++ b/Engine/Plugins/FX/Niagara/Source/NiagaraCore/Private/NiagaraCustomVersion.cpp
@@ -8,5 +8,5 @@ const FGuid FNiagaraCustomVersion::GUID(0xFCF57AFA, 0x50764283, 0xB9A9E658, 0xFF
 // Register the custom version with core
 FCustomVersionRegistration GRegisterNiagaraCustomVersion(FNiagaraCustomVersion::GUID, FNiagaraCustomVersion::LatestVersion, TEXT("NiagaraVer"));
 
-const FGuid FNiagaraCustomVersion::LatestScriptCompileVersion(0x61E361A4, 0xD3C94ADB, 0xA93C3D8A, 0x507707CC);
+const FGuid FNiagaraCustomVersion::LatestScriptCompileVersion(0x6D2E1207, 0x63F99C45, 0xBC426396, 0xFBF80045);
 
diff --git a/Engine/Plugins/FX/Niagara/Source/NiagaraEditor/Public/NiagaraCurveOwner.h b/Engine/Plugins/FX/Niagara/Source/NiagaraEditor/Public/NiagaraCurveOwner.h
--- a/Engine/Plugins/FX/Niagara/Source/NiagaraEditor/Public/NiagaraCurveOwner.h
+++ b/Engine/Plugins/FX/Niagara/Source/NiagaraEditor/Public/NiagaraCurveOwner.h
@@ -4,6 +4,7 @@
 
 #include "Curves/CurveOwnerInterface.h"
 #include "NiagaraDataInterface.h"
+#include "NiagaraDataInterfaceCurveBase.h"
 
 /** A curve owner for curves in a niagara System. */
 class NIAGARAEDITOR_API FNiagaraCurveOwner : public FCurveOwnerInterface
