- VM natviz of Life Cycle / Spawn Rate / SolveVelocitiesAndForces

#rnx

diff --git a/Engine/Plugins/FX/Niagara/Source/Niagara/Private/NiagaraFunctionLibrary.cpp b/Engine/Plugins/FX/Niagara/Source/Niagara/Private/NiagaraFunctionLibrary.cpp
--- a/Engine/Plugins/FX/Niagara/Source/Niagara/Private/NiagaraFunctionLibrary.cpp
+++ b/Engine/Plugins/FX/Niagara/Source/Niagara/Private/NiagaraFunctionLibrary.cpp
@@ -395,6 +395,10 @@ const FName FastPathLibraryName(TEXT("FastPathLibrary"));
 const FName FastPathDot4Name(TEXT("FastPathDot4"));
 const FName FastPathTransformPositionName(TEXT("FastPathTransformPosition"));
 const FName FastMatrixToQuaternionName(TEXT("FastMatrixToQuaternion"));
+const FName FastPathEmitterLifeCycleName(TEXT("FastPathEmitterLifeCycle"));
+const FName FastPathSpawnRateName(TEXT("FastPathSpawnRate"));
+const FName FastPathSpawnBurstInstantaneousName(TEXT("FastPathSpawnBurstInstantaneous"));
+const FName FastPathSolveVelocitiesAndForces(TEXT("FastPathSolveVelocitiesAndForces"));
 
 struct FVectorKernelFastDot4
 {
@@ -451,8 +455,8 @@ struct FVectorKernelFastDot4
 		VectorRegister* RESTRICT BW = (VectorRegister *)InVecB[3].GetDest();
 		VectorRegister* RESTRICT Out = (VectorRegister *)OutValue.GetDest();
 
-		int32 LoopInstances = Align(Context.NumInstances, 4) / 4;
-		for (int32 i = 0; i < LoopInstances; ++i)
+		const int32 Loops = Context.GetNumLoops<4>();
+		for (int32 i = 0; i < Loops; ++i)
 		{
 
 			VectorRegister AVX0= VectorLoadAligned(&AX[i]);
@@ -704,7 +708,7 @@ struct FVectorKernelFastMatrixToQuaternion
 			OutQuat.Emplace(Context);
 		}
 
-		for ( int32 i=0; i < Context.NumInstances; ++i )
+		for ( int32 i=0; i < Context.GetNumInstances(); ++i )
 		{
 			FMatrix Mat;
 			for ( int32 j=0; j < 16; ++j )
@@ -721,6 +725,706 @@ struct FVectorKernelFastMatrixToQuaternion
 	}
 };
 
+struct FVectorKernel_EmitterLifeCycle
+{
+	static FNiagaraFunctionSignature GetFunctionSignature()
+	{
+		FNiagaraFunctionSignature Sig;
+		Sig.Name = FastPathEmitterLifeCycleName;
+		Sig.OwnerName = FastPathLibraryName;
+		Sig.bMemberFunction = false;
+		Sig.bRequiresContext = false;
+		Sig.SetDescription(LOCTEXT("FastPathEmitterLifeCycleDesc", "Fast path for life cycle"));
+		Sig.Inputs.Add(FNiagaraVariable(FNiagaraTypeDefinition::GetFloatDef(), TEXT("EngineDeltaTime")));
+		Sig.Inputs.Add(FNiagaraVariable(FNiagaraTypeDefinition::GetIntDef(), TEXT("EngineNumParticles")));
+		Sig.Inputs.Add(FNiagaraVariable(FNiagaraTypeDefinition::GetExecutionStateEnum(), TEXT("ScalabilityEmitterExecutionState")));
+		Sig.Inputs.Add(FNiagaraVariable(FNiagaraTypeDefinition::GetExecutionStateEnum(), TEXT("SystemExecutionState")));
+		Sig.Inputs.Add(FNiagaraVariable(FNiagaraTypeDefinition::GetExecutionStateSouceEnum(), TEXT("SystemExecutionStateSource")));
+
+		Sig.Inputs.Add(FNiagaraVariable(FNiagaraTypeDefinition::GetFloatDef(), TEXT("ModuleNextLoopDuration")));
+		Sig.Inputs.Add(FNiagaraVariable(FNiagaraTypeDefinition::GetFloatDef(), TEXT("ModuleNextLoopDelay")));
+		Sig.Inputs.Add(FNiagaraVariable(FNiagaraTypeDefinition::GetBoolDef(), TEXT("ModuleDurationRecalcEachLoop")));
+		Sig.Inputs.Add(FNiagaraVariable(FNiagaraTypeDefinition::GetBoolDef(), TEXT("ModuleDelayFirstLoopOnly")));
+		Sig.Inputs.Add(FNiagaraVariable(FNiagaraTypeDefinition::GetIntDef(), TEXT("ModuleMaxLoopCount")));
+		Sig.Inputs.Add(FNiagaraVariable(FNiagaraTypeDefinition::GetBoolDef(), TEXT("ModuleAutoComplete")));
+		Sig.Inputs.Add(FNiagaraVariable(FNiagaraTypeDefinition::GetBoolDef(), TEXT("ModuleCompleteOnInactive")));
+
+		Sig.Inputs.Add(FNiagaraVariable(FNiagaraTypeDefinition::GetExecutionStateEnum(), TEXT("EmitterExecutionState")));
+		Sig.Inputs.Add(FNiagaraVariable(FNiagaraTypeDefinition::GetExecutionStateSouceEnum(), TEXT("EmitterExecutionStateSource")));
+		Sig.Inputs.Add(FNiagaraVariable(FNiagaraTypeDefinition::GetFloatDef(), TEXT("EmitterAge")));
+		Sig.Inputs.Add(FNiagaraVariable(FNiagaraTypeDefinition::GetFloatDef(), TEXT("EmitterLoopedAge")));
+		Sig.Inputs.Add(FNiagaraVariable(FNiagaraTypeDefinition::GetFloatDef(), TEXT("EmitterCurrentLoopDuration")));
+		Sig.Inputs.Add(FNiagaraVariable(FNiagaraTypeDefinition::GetFloatDef(), TEXT("EmitterCurrentLoopDelay")));
+		Sig.Inputs.Add(FNiagaraVariable(FNiagaraTypeDefinition::GetIntDef(), TEXT("EmitterLoopCount")));
+		Sig.Inputs.Add(FNiagaraVariable(FNiagaraTypeDefinition::GetFloatDef(), TEXT("EmitterNormalizedLoopAge")));
+
+		Sig.Outputs.Add(FNiagaraVariable(FNiagaraTypeDefinition::GetExecutionStateEnum(), TEXT("EmitterExecutionState")));
+		Sig.Outputs.Add(FNiagaraVariable(FNiagaraTypeDefinition::GetExecutionStateSouceEnum(), TEXT("EmitterExecutionStateSource")));
+		Sig.Outputs.Add(FNiagaraVariable(FNiagaraTypeDefinition::GetFloatDef(), TEXT("EmitterAge")));
+		Sig.Outputs.Add(FNiagaraVariable(FNiagaraTypeDefinition::GetFloatDef(), TEXT("EmitterLoopedAge")));
+		Sig.Outputs.Add(FNiagaraVariable(FNiagaraTypeDefinition::GetFloatDef(), TEXT("EmitterCurrentLoopDuration")));
+		Sig.Outputs.Add(FNiagaraVariable(FNiagaraTypeDefinition::GetFloatDef(), TEXT("EmitterCurrentLoopDelay")));
+		Sig.Outputs.Add(FNiagaraVariable(FNiagaraTypeDefinition::GetIntDef(), TEXT("EmitterLoopCount")));
+		Sig.Outputs.Add(FNiagaraVariable(FNiagaraTypeDefinition::GetFloatDef(), TEXT("EmitterNormalizedLoopAge")));
+
+		return Sig;
+	}
+
+	static FString GetFunctionHLSL()
+	{
+		return FString();
+	}
+
+	static void Exec(FVectorVMContext& Context)
+	{
+		VectorVM::FExternalFuncInputHandler<float> InEngineDeltaTime(Context);
+		VectorVM::FExternalFuncInputHandler<int>   InEngineNumParticles(Context);
+		VectorVM::FExternalFuncInputHandler<ENiagaraExecutionState> InScalabilityEmitterExecutionState(Context);
+		VectorVM::FExternalFuncInputHandler<ENiagaraExecutionState> InSystemExecutionState(Context);
+		VectorVM::FExternalFuncInputHandler<ENiagaraExecutionStateSource> InSystemExecutionStateSource(Context);
+
+		VectorVM::FExternalFuncInputHandler<float> InModuleNextLoopDuration(Context);
+		VectorVM::FExternalFuncInputHandler<float> InModuleNextLoopDelay(Context);
+		VectorVM::FExternalFuncInputHandler<bool>  InModuleDurationRecalcEachLoop(Context);
+		VectorVM::FExternalFuncInputHandler<bool>  InModuleDelayFirstLoopOnly(Context);
+		VectorVM::FExternalFuncInputHandler<int>   InModuleMaxLoopCount(Context);
+		VectorVM::FExternalFuncInputHandler<bool>  InModuleAutoComplete(Context);
+		VectorVM::FExternalFuncInputHandler<bool>  InModuleCompleteOnInactive(Context);
+
+		VectorVM::FExternalFuncInputHandler<ENiagaraExecutionState> InEmitterExecutionState(Context);
+		VectorVM::FExternalFuncInputHandler<ENiagaraExecutionStateSource> InEmitterExecutionStateSource(Context);
+		VectorVM::FExternalFuncInputHandler<float> InEmitterAge(Context);
+		VectorVM::FExternalFuncInputHandler<float> InEmitterLoopedAge(Context);
+		VectorVM::FExternalFuncInputHandler<float> InEmitterCurrentLoopDuration(Context);
+		VectorVM::FExternalFuncInputHandler<float> InEmitterCurrentLoopDelay(Context);
+		VectorVM::FExternalFuncInputHandler<int>   InEmitterLoopCount(Context);
+		VectorVM::FExternalFuncInputHandler<float> InEmitterNormalizedLoopAge(Context);
+
+		VectorVM::FExternalFuncRegisterHandler<ENiagaraExecutionState> OutEmitterExecutionState(Context);
+		VectorVM::FExternalFuncRegisterHandler<ENiagaraExecutionStateSource> OutEmitterExecutionStateSource(Context);
+		VectorVM::FExternalFuncRegisterHandler<float> OutEmitterAge(Context);
+		VectorVM::FExternalFuncRegisterHandler<float> OutEmitterLoopedAge(Context);
+		VectorVM::FExternalFuncRegisterHandler<float> OutEmitterCurrentLoopDuration(Context);
+		VectorVM::FExternalFuncRegisterHandler<float> OutEmitterCurrentLoopDelay(Context);
+		VectorVM::FExternalFuncRegisterHandler<int>   OutEmitterLoopCount(Context);
+		VectorVM::FExternalFuncRegisterHandler<float> OutEmitterNormalizedLoopAge(Context);
+
+		for (int i = 0; i < Context.GetNumInstances(); ++i)
+		{
+			const float EngineDeltaTime = InEngineDeltaTime.GetAndAdvance();
+			const int EngineNumParticles = InEngineNumParticles.GetAndAdvance();
+
+			const ENiagaraExecutionState ScalabilityEmitterExecutionState = InScalabilityEmitterExecutionState.GetAndAdvance();
+
+			const ENiagaraExecutionState SystemExecutionState = InSystemExecutionState.GetAndAdvance();
+			const ENiagaraExecutionStateSource SystemExecutionStateSource = InSystemExecutionStateSource.GetAndAdvance();
+
+			const float ModuleNextLoopDuration = InModuleNextLoopDuration.GetAndAdvance();
+			const float ModuleNextLoopDelay = InModuleNextLoopDelay.GetAndAdvance();
+			const bool ModuleDurationRecalcEachLoop = InModuleDurationRecalcEachLoop.GetAndAdvance();
+			const bool ModuleDelayFirstLoopOnly = InModuleDelayFirstLoopOnly.GetAndAdvance();
+			const int ModuleMaxLoopCount = InModuleMaxLoopCount.GetAndAdvance();
+			const bool ModuleAutoComplete = InModuleAutoComplete.GetAndAdvance();
+			const bool ModuleCompleteOnInactive = InModuleCompleteOnInactive.GetAndAdvance();
+
+			ENiagaraExecutionState EmitterExecutionState = InEmitterExecutionState.GetAndAdvance();
+			ENiagaraExecutionStateSource EmitterExecutionStateSource = InEmitterExecutionStateSource.GetAndAdvance();
+			float EmitterAge = InEmitterAge.GetAndAdvance();
+			float EmitterLoopedAge = InEmitterLoopedAge.GetAndAdvance();
+			float EmitterCurrentLoopDuration = InEmitterCurrentLoopDuration.GetAndAdvance();
+			float EmitterCurrentLoopDelay = InEmitterCurrentLoopDelay.GetAndAdvance();
+			int EmitterLoopCount = InEmitterLoopCount.GetAndAdvance();
+			float EmitterNormalizedLoopAge = InEmitterNormalizedLoopAge.GetAndAdvance();
+
+			// Skip disabled emitters
+			if ( EmitterExecutionState != ENiagaraExecutionState::Disabled )
+			{
+				// Initialize parameters
+				if (EmitterAge == 0.0f)
+				{
+					EmitterLoopedAge = -ModuleNextLoopDelay;
+					EmitterCurrentLoopDuration = ModuleNextLoopDuration;
+					EmitterCurrentLoopDelay = ModuleNextLoopDelay;
+				}
+
+				// Handle emitter looping
+				EmitterAge += EngineDeltaTime;
+				EmitterLoopedAge += EngineDeltaTime;
+				const int32 LoopsPerformed = FMath::FloorToInt(EmitterLoopedAge / EmitterCurrentLoopDuration);
+				if (LoopsPerformed > 0)
+				{
+					EmitterLoopedAge -= float(LoopsPerformed) * EmitterCurrentLoopDuration;
+					EmitterLoopCount += LoopsPerformed;
+
+					if (ModuleDurationRecalcEachLoop)
+					{
+						EmitterCurrentLoopDuration = ModuleNextLoopDuration;
+					}
+					if (ModuleDelayFirstLoopOnly)
+					{
+						EmitterCurrentLoopDelay = 0.0f;
+					}
+					EmitterNormalizedLoopAge = EmitterLoopedAge / EmitterCurrentLoopDuration;
+				}
+
+				// Set emitter state from scalability (if allowed)
+				if ( EmitterExecutionStateSource <= ENiagaraExecutionStateSource::Scalability)
+				{
+					EmitterExecutionState = ScalabilityEmitterExecutionState;
+					EmitterExecutionStateSource = ENiagaraExecutionStateSource::Scalability;
+				}
+
+				// Exceeded maximum loops?
+				if (ModuleMaxLoopCount > 0 && EmitterLoopCount >= ModuleMaxLoopCount)
+				{
+					if (EmitterExecutionStateSource <= ENiagaraExecutionStateSource::Internal)
+					{
+						EmitterExecutionState = ENiagaraExecutionState::Inactive;
+						EmitterExecutionStateSource = ENiagaraExecutionStateSource::Internal;
+					}
+				}
+
+				// Are we complete?
+				if (EmitterExecutionState != ENiagaraExecutionState::Active && (ModuleCompleteOnInactive || (EngineNumParticles == 0 && ModuleAutoComplete)))
+				{
+					if ( EmitterExecutionStateSource <= ENiagaraExecutionStateSource::InternalCompletion )
+					{
+						EmitterExecutionState = ENiagaraExecutionState::Complete;
+						EmitterExecutionStateSource = ENiagaraExecutionStateSource::InternalCompletion;
+					}
+				}
+			}
+
+			// Set values
+			*OutEmitterExecutionState.GetDestAndAdvance() = EmitterExecutionState;
+			*OutEmitterExecutionStateSource.GetDestAndAdvance() = EmitterExecutionStateSource;
+			*OutEmitterAge.GetDestAndAdvance() = EmitterAge;
+			*OutEmitterLoopedAge.GetDestAndAdvance() = EmitterLoopedAge;
+			*OutEmitterCurrentLoopDuration.GetDestAndAdvance() = EmitterCurrentLoopDuration;
+			*OutEmitterCurrentLoopDelay.GetDestAndAdvance() = EmitterCurrentLoopDelay;
+			*OutEmitterLoopCount.GetDestAndAdvance() = EmitterLoopCount;
+			*OutEmitterNormalizedLoopAge.GetDestAndAdvance() = EmitterNormalizedLoopAge;
+		}
+	}
+};
+
+struct FVectorKernel_SpawnRate
+{
+	static FNiagaraFunctionSignature GetFunctionSignature()
+	{
+		FNiagaraFunctionSignature Sig;
+		Sig.Name = FastPathSpawnRateName;
+		Sig.OwnerName = FastPathLibraryName;
+		Sig.bMemberFunction = false;
+		Sig.bRequiresContext = false;
+		Sig.SetDescription(LOCTEXT("FastPathSpawnRateDesc", "Fast path for spawn rate"));
+		Sig.Inputs.Add(FNiagaraVariable(FNiagaraTypeDefinition::GetFloatDef(), TEXT("EngineDeltaTime")));
+		Sig.Inputs.Add(FNiagaraVariable(FNiagaraTypeDefinition::GetFloatDef(), TEXT("ModuleSpawnRate")));
+		Sig.Inputs.Add(FNiagaraVariable(FNiagaraTypeDefinition::GetFloatDef(), TEXT("ScalabilityEmitterSpawnCountScale")));
+		Sig.Inputs.Add(FNiagaraVariable(FNiagaraTypeDefinition::GetFloatDef(), TEXT("EngineEmitterSpawnCountScale")));
+		Sig.Inputs.Add(FNiagaraVariable(FNiagaraTypeDefinition::GetFloatDef(), TEXT("EmitterModuleSpawnRemainder")));
+		Sig.Inputs.Add(FNiagaraVariable(FNiagaraTypeDefinition::GetFloatDef(), TEXT("EmitterLoopedAge")));
+		Sig.Inputs.Add(FNiagaraVariable(FNiagaraTypeDefinition::GetIntDef(), TEXT("EmitterSpawnGroup")));
+
+		Sig.Outputs.Add(FNiagaraVariable(FNiagaraTypeDefinition::GetBoolDef(), TEXT("SpawningCanEverSpawn")));
+		Sig.Outputs.Add(FNiagaraVariable(FNiagaraTypeDefinition::GetFloatDef(), TEXT("EmitterModuleSpawnRemainder")));
+		Sig.Outputs.Add(FNiagaraVariable(FNiagaraTypeDefinition::GetIntDef(), TEXT("EmitterModuleSpawnInfoCount")));
+		Sig.Outputs.Add(FNiagaraVariable(FNiagaraTypeDefinition::GetFloatDef(), TEXT("EmitterModuleSpawnInfoInterpStartDt")));
+		Sig.Outputs.Add(FNiagaraVariable(FNiagaraTypeDefinition::GetFloatDef(), TEXT("EmitterModuleSpawnInfoIntervalDt")));
+		Sig.Outputs.Add(FNiagaraVariable(FNiagaraTypeDefinition::GetIntDef(), TEXT("EmitterModuleSpawnInfoSpawnGroup")));
+
+		return Sig;
+	}
+
+	static FString GetFunctionHLSL()
+	{
+		return FString();
+	}
+
+	static void Exec(FVectorVMContext& Context)
+	{
+		VectorVM::FExternalFuncInputHandler<float> InEngineDeltaTime(Context);
+		VectorVM::FExternalFuncInputHandler<float> InModuleSpawnRate(Context);
+		VectorVM::FExternalFuncInputHandler<float> InScalabilityEmitterSpawnCountScale(Context);
+		VectorVM::FExternalFuncInputHandler<float> InEngineEmitterSpawnCountScale(Context);
+		VectorVM::FExternalFuncInputHandler<float> InEmitterModuleSpawnRemainder(Context);
+		VectorVM::FExternalFuncInputHandler<float> InEmitterLoopedAge(Context);
+		VectorVM::FExternalFuncInputHandler<int32> InEmitterSpawnGroup(Context);
+
+		VectorVM::FExternalFuncRegisterHandler<bool> OutSpawningCanEverSpawn(Context);
+		VectorVM::FExternalFuncRegisterHandler<float> OutEmitterModuleSpawnRemainder(Context);
+		VectorVM::FExternalFuncRegisterHandler<int32> OutEmitterModuleSpawnInfoCount(Context);
+		VectorVM::FExternalFuncRegisterHandler<float> OutEmitterModuleSpawnInfoInterpStartDt(Context);
+		VectorVM::FExternalFuncRegisterHandler<float> OutEmitterModuleSpawnInfoIntervalDt(Context);
+		VectorVM::FExternalFuncRegisterHandler<int32> OutEmitterModuleSpawnInfoSpawnGroup(Context);
+
+		for (int i = 0; i < Context.GetNumInstances(); ++i)
+		{
+			// Gather values
+			float EngineDeltaTime = InEngineDeltaTime.GetAndAdvance();
+			float ModuleSpawnRate = InModuleSpawnRate.GetAndAdvance();
+			float ScalabilityEmitterSpawnCountScale = InScalabilityEmitterSpawnCountScale.GetAndAdvance();
+			float EngineEmitterSpawnCountScale = InEngineEmitterSpawnCountScale.GetAndAdvance();
+			float EmitterModuleSpawnRemainder = InEmitterModuleSpawnRemainder.GetAndAdvance();
+			float EmitterLoopedAge = InEmitterLoopedAge.GetAndAdvance();
+			int32 EmitterSpawnGroup = InEmitterSpawnGroup.GetAndAdvance();
+
+			//
+			float LocalModuleSpawnRate = ModuleSpawnRate * ScalabilityEmitterSpawnCountScale * EngineEmitterSpawnCountScale;
+			float LocalModuleIntervalDT = 1.0f / LocalModuleSpawnRate;
+			float LocalModuleInterpStartDT = LocalModuleIntervalDT * (1.0f - EmitterModuleSpawnRemainder);
+
+			FNiagaraSpawnInfo SpawnInfo;
+			if (EmitterLoopedAge > 0.0f)
+			{
+				float fSpawnCount = (LocalModuleSpawnRate * EngineDeltaTime) + EmitterModuleSpawnRemainder;
+				int32 LocalModuleSpawnCount = FMath::FloorToInt(fSpawnCount);
+				EmitterModuleSpawnRemainder = fSpawnCount - float(LocalModuleSpawnCount);
+
+				SpawnInfo.Count = LocalModuleSpawnCount;
+				SpawnInfo.InterpStartDt = LocalModuleInterpStartDT;
+				SpawnInfo.IntervalDt = LocalModuleIntervalDT;
+				SpawnInfo.SpawnGroup = EmitterSpawnGroup;
+			}
+
+			// Write values
+			*OutEmitterModuleSpawnInfoCount.GetDestAndAdvance() = SpawnInfo.Count;
+			*OutEmitterModuleSpawnInfoInterpStartDt.GetDestAndAdvance() = SpawnInfo.InterpStartDt;
+			*OutEmitterModuleSpawnInfoIntervalDt.GetDestAndAdvance() = SpawnInfo.IntervalDt;
+			*OutEmitterModuleSpawnInfoSpawnGroup.GetDestAndAdvance() = SpawnInfo.SpawnGroup;
+			*OutSpawningCanEverSpawn.GetDestAndAdvance() = true;
+			*OutEmitterModuleSpawnRemainder.GetDestAndAdvance() = EmitterModuleSpawnRemainder;
+		}
+	}
+};
+
+struct FVectorKernel_SpawnBurstInstantaneous
+{
+	static FNiagaraFunctionSignature GetFunctionSignature()
+	{
+		FNiagaraFunctionSignature Sig;
+		Sig.Name = FastPathSpawnBurstInstantaneousName;
+		Sig.OwnerName = FastPathLibraryName;
+		Sig.bMemberFunction = false;
+		Sig.bRequiresContext = false;
+		Sig.SetDescription(LOCTEXT("FastPathSpawnBurstInstantaneous", "Fast path for spawn burst instantaneous"));
+		Sig.Inputs.Add(FNiagaraVariable(FNiagaraTypeDefinition::GetFloatDef(), TEXT("EngineDeltaTime")));
+		Sig.Inputs.Add(FNiagaraVariable(FNiagaraTypeDefinition::GetFloatDef(), TEXT("ScalabilityEmitterSpawnCountScale")));
+		Sig.Inputs.Add(FNiagaraVariable(FNiagaraTypeDefinition::GetFloatDef(), TEXT("EmitterLoopedAge")));
+		Sig.Inputs.Add(FNiagaraVariable(FNiagaraTypeDefinition::GetFloatDef(), TEXT("ModuleSpawnTime")));
+		Sig.Inputs.Add(FNiagaraVariable(FNiagaraTypeDefinition::GetIntDef(), TEXT("ModuleSpawnCount")));
+		Sig.Inputs.Add(FNiagaraVariable(FNiagaraTypeDefinition::GetIntDef(), TEXT("ModuleSpawnGroup")));
+
+		Sig.Outputs.Add(FNiagaraVariable(FNiagaraTypeDefinition::GetBoolDef(), TEXT("SpawningCanEverSpawn")));
+		Sig.Outputs.Add(FNiagaraVariable(FNiagaraTypeDefinition::GetIntDef(), TEXT("EmitterModuleSpawnInfoCount")));
+		Sig.Outputs.Add(FNiagaraVariable(FNiagaraTypeDefinition::GetFloatDef(), TEXT("EmitterModuleSpawnInfoInterpStartDt")));
+		Sig.Outputs.Add(FNiagaraVariable(FNiagaraTypeDefinition::GetFloatDef(), TEXT("EmitterModuleSpawnInfoIntervalDt")));
+		Sig.Outputs.Add(FNiagaraVariable(FNiagaraTypeDefinition::GetIntDef(), TEXT("EmitterModuleSpawnInfoSpawnGroup")));
+
+		return Sig;
+	}
+
+	static FString GetFunctionHLSL()
+	{
+		return FString();
+	}
+
+	static void Exec(FVectorVMContext& Context)
+	{
+		VectorVM::FExternalFuncInputHandler<float> InEngineDeltaTime(Context);
+		VectorVM::FExternalFuncInputHandler<float> InScalabilityEmitterSpawnCountScale(Context);
+		VectorVM::FExternalFuncInputHandler<float> InEmitterLoopedAge(Context);
+		VectorVM::FExternalFuncInputHandler<float> InModuleSpawnTime(Context);
+		VectorVM::FExternalFuncInputHandler<int32> InModuleSpawnCount(Context);
+		VectorVM::FExternalFuncInputHandler<int32> InModuleSpawnGroup(Context);
+
+		VectorVM::FExternalFuncRegisterHandler<bool> OutSpawningCanEverSpawn(Context);
+		VectorVM::FExternalFuncRegisterHandler<int32> OutEmitterModuleSpawnInfoCount(Context);
+		VectorVM::FExternalFuncRegisterHandler<float> OutEmitterModuleSpawnInfoInterpStartDt(Context);
+		VectorVM::FExternalFuncRegisterHandler<float> OutEmitterModuleSpawnInfoIntervalDt(Context);
+		VectorVM::FExternalFuncRegisterHandler<int32> OutEmitterModuleSpawnInfoSpawnGroup(Context);
+
+		for (int i=0; i < Context.GetNumInstances(); ++i)
+		{
+			// Gather values
+			float EngineDeltaTime = InEngineDeltaTime.GetAndAdvance();
+			float ScalabilityEmitterSpawnCountScale = InScalabilityEmitterSpawnCountScale.GetAndAdvance();
+			float EmitterLoopedAge = InEmitterLoopedAge.GetAndAdvance();
+			float ModuleSpawnTime = InModuleSpawnTime.GetAndAdvance();
+			int32 ModuleSpawnCount = InModuleSpawnCount.GetAndAdvance();
+			int32 ModuleSpawnGroup = InModuleSpawnGroup.GetAndAdvance();
+
+			const float PreviousTime = EmitterLoopedAge - EngineDeltaTime;
+
+			*OutSpawningCanEverSpawn.GetDestAndAdvance() = EmitterLoopedAge <= ModuleSpawnTime;
+			if ( (ModuleSpawnTime >= PreviousTime) && (ModuleSpawnTime < EmitterLoopedAge) )
+			{
+				*OutEmitterModuleSpawnInfoCount.GetDestAndAdvance() = ModuleSpawnCount;
+				*OutEmitterModuleSpawnInfoInterpStartDt.GetDestAndAdvance() = ModuleSpawnTime - PreviousTime;
+			}
+			else
+			{
+				*OutEmitterModuleSpawnInfoCount.GetDestAndAdvance() = 0;
+				*OutEmitterModuleSpawnInfoInterpStartDt.GetDestAndAdvance() = 0.0f;
+			}
+			*OutEmitterModuleSpawnInfoIntervalDt.GetDestAndAdvance() = 0.0f;
+			*OutEmitterModuleSpawnInfoSpawnGroup.GetDestAndAdvance() = ModuleSpawnGroup;
+		}
+	}
+};
+
+struct FVectorKernel_SolveVelocitiesAndForces
+{
+	static FNiagaraFunctionSignature GetFunctionSignature()
+	{
+		FNiagaraFunctionSignature Sig;
+		Sig.Name = FastPathSolveVelocitiesAndForces;
+		Sig.OwnerName = FastPathLibraryName;
+		Sig.bMemberFunction = false;
+		Sig.bRequiresContext = false;
+		Sig.SetDescription(LOCTEXT("FastPathSolveVelocitiesAndForces", "Fast path for SolveVelocitiesAndForces"));
+
+		Sig.Inputs.Add(FNiagaraVariable(FNiagaraTypeDefinition::GetFloatDef(), TEXT("EngineDeltaTime")));
+		Sig.Inputs.Add(FNiagaraVariable(FNiagaraTypeDefinition::GetVec3Def(), TEXT("PhysicsForce")));
+		Sig.Inputs.Add(FNiagaraVariable(FNiagaraTypeDefinition::GetFloatDef(), TEXT("PhysicsDrag")));
+		Sig.Inputs.Add(FNiagaraVariable(FNiagaraTypeDefinition::GetFloatDef(), TEXT("ParticlesMass")));
+		Sig.Inputs.Add(FNiagaraVariable(FNiagaraTypeDefinition::GetVec3Def(), TEXT("ParticlesPosition")));
+		Sig.Inputs.Add(FNiagaraVariable(FNiagaraTypeDefinition::GetVec3Def(), TEXT("ParticlesVelocity")));
+
+		Sig.Outputs.Add(FNiagaraVariable(FNiagaraTypeDefinition::GetVec3Def(), TEXT("ParticlesPosition")));
+		Sig.Outputs.Add(FNiagaraVariable(FNiagaraTypeDefinition::GetVec3Def(), TEXT("ParticlesVelocity")));
+		Sig.Outputs.Add(FNiagaraVariable(FNiagaraTypeDefinition::GetVec3Def(), TEXT("ParticlesPreviousVelocity")));
+
+		return Sig;
+	}
+
+	static FString GetFunctionHLSL()
+	{
+		return FString();
+	}
+
+	static void Exec(FVectorVMContext& Context)
+	{
+		VectorVM::FExternalFuncInputHandler<float> InEngineDeltaTime(Context);
+		VectorVM::FExternalFuncInputHandler<float> InPhysicsForceX(Context);
+		VectorVM::FExternalFuncInputHandler<float> InPhysicsForceY(Context);
+		VectorVM::FExternalFuncInputHandler<float> InPhysicsForceZ(Context);
+		VectorVM::FExternalFuncInputHandler<float> InPhysicsDrag(Context);
+		VectorVM::FExternalFuncInputHandler<float> InParticlesMass(Context);
+		VectorVM::FExternalFuncInputHandler<float> InParticlesPositionX(Context);
+		VectorVM::FExternalFuncInputHandler<float> InParticlesPositionY(Context);
+		VectorVM::FExternalFuncInputHandler<float> InParticlesPositionZ(Context);
+		VectorVM::FExternalFuncInputHandler<float> InParticlesVelocityX(Context);
+		VectorVM::FExternalFuncInputHandler<float> InParticlesVelocityY(Context);
+		VectorVM::FExternalFuncInputHandler<float> InParticlesVelocityZ(Context);
+
+		VectorVM::FExternalFuncRegisterHandler<float> OutParticlesPositionX(Context);
+		VectorVM::FExternalFuncRegisterHandler<float> OutParticlesPositionY(Context);
+		VectorVM::FExternalFuncRegisterHandler<float> OutParticlesPositionZ(Context);
+		VectorVM::FExternalFuncRegisterHandler<float> OutParticlesVelocityX(Context);
+		VectorVM::FExternalFuncRegisterHandler<float> OutParticlesVelocityY(Context);
+		VectorVM::FExternalFuncRegisterHandler<float> OutParticlesVelocityZ(Context);
+		VectorVM::FExternalFuncRegisterHandler<float> OutParticlesPreviousVelocityX(Context);
+		VectorVM::FExternalFuncRegisterHandler<float> OutParticlesPreviousVelocityY(Context);
+		VectorVM::FExternalFuncRegisterHandler<float> OutParticlesPreviousVelocityZ(Context);
+#if 1
+		const float EngineDeltaTime = InEngineDeltaTime.Get();
+		const float MassMin = 0.0001f;
+
+		for (int i = 0; i < Context.GetNumInstances(); ++i)
+		{
+			// Gather values
+			FVector PhysicsForce(InPhysicsForceX.GetAndAdvance(), InPhysicsForceY.GetAndAdvance(), InPhysicsForceZ.GetAndAdvance());
+			float PhysicsDrag = InPhysicsDrag.GetAndAdvance();
+
+			float ParticleMass = InParticlesMass.GetAndAdvance();
+			FVector ParticlePosition(InParticlesPositionX.GetAndAdvance(), InParticlesPositionY.GetAndAdvance(), InParticlesPositionZ.GetAndAdvance());
+			FVector ParticleVelocity(InParticlesVelocityX.GetAndAdvance(), InParticlesVelocityY.GetAndAdvance(), InParticlesVelocityZ.GetAndAdvance());
+
+			*OutParticlesPreviousVelocityX.GetDestAndAdvance() = ParticleVelocity.X;
+			*OutParticlesPreviousVelocityY.GetDestAndAdvance() = ParticleVelocity.Y;
+			*OutParticlesPreviousVelocityZ.GetDestAndAdvance() = ParticleVelocity.Z;
+
+			// Apply velocity
+			const float OOParticleMassDT = (1.0f / FMath::Max(ParticleMass, MassMin)) * EngineDeltaTime;
+			ParticleVelocity += PhysicsForce * OOParticleMassDT;
+
+			// Apply Drag
+			float ClampedDrag = FMath::Clamp(PhysicsDrag * EngineDeltaTime, 0.0f, 1.0f);
+			ParticleVelocity -= ParticleVelocity * ClampedDrag;
+
+			// Velocity Clamp
+			//-TODO: Not used
+
+			// Limit Acceleration
+			//-TODO: Not used
+
+			// Apply velocity
+			ParticlePosition += ParticleVelocity * EngineDeltaTime;
+
+			// Write parameters
+			*OutParticlesPositionX.GetDestAndAdvance() = ParticlePosition.X;
+			*OutParticlesPositionY.GetDestAndAdvance() = ParticlePosition.Y;
+			*OutParticlesPositionZ.GetDestAndAdvance() = ParticlePosition.Z;
+
+			*OutParticlesVelocityX.GetDestAndAdvance() = ParticleVelocity.X;
+			*OutParticlesVelocityY.GetDestAndAdvance() = ParticleVelocity.Y;
+			*OutParticlesVelocityZ.GetDestAndAdvance() = ParticleVelocity.Z;
+		}
+#else
+		const VectorRegister EngineDeltaTime = VectorSetFloat1(InEngineDeltaTime.Get());
+		const VectorRegister MassMin = VectorSetFloat1(0.0001f);
+
+		for (int i=0; i < Context.GetNumLoops<4>(); ++i)
+		{
+			// Gather values
+			VectorRegister PhysicsForceX		= VectorLoad(InPhysicsForceX.GetDestAndAdvance());
+			VectorRegister PhysicsForceY		= VectorLoad(InPhysicsForceY.GetDestAndAdvance());
+			VectorRegister PhysicsForceZ		= VectorLoad(InPhysicsForceZ.GetDestAndAdvance());
+			VectorRegister PhysicsDrag			= VectorLoad(InPhysicsDrag.GetDestAndAdvance());
+
+			VectorRegister ParticlesMass		= VectorLoad(InParticlesMass.GetDestAndAdvance());
+			VectorRegister ParticlesPositionX	= VectorLoad(InParticlesPositionX.GetDestAndAdvance());
+			VectorRegister ParticlesPositionY	= VectorLoad(InParticlesPositionY.GetDestAndAdvance());
+			VectorRegister ParticlesPositionZ	= VectorLoad(InParticlesPositionZ.GetDestAndAdvance());
+			VectorRegister ParticlesVelocityX	= VectorLoad(InParticlesVelocityX.GetDestAndAdvance());
+			VectorRegister ParticlesVelocityY	= VectorLoad(InParticlesVelocityY.GetDestAndAdvance());
+			VectorRegister ParticlesVelocityZ	= VectorLoad(InParticlesVelocityZ.GetDestAndAdvance());
+
+			VectorStore(ParticlesVelocityX, OutParticlesPreviousVelocityX.GetDestAndAdvance());
+			VectorStore(ParticlesVelocityY, OutParticlesPreviousVelocityY.GetDestAndAdvance());
+			VectorStore(ParticlesVelocityZ, OutParticlesPreviousVelocityZ.GetDestAndAdvance());
+
+			// Apply velocity
+			const VectorRegister OOParticleMassDT = VectorMultiply(VectorReciprocal(VectorMax(ParticlesMass, MassMin)), EngineDeltaTime);
+			ParticlesVelocityX = VectorMultiplyAdd(PhysicsForceX, OOParticleMassDT, ParticlesVelocityX);
+			ParticlesVelocityY = VectorMultiplyAdd(PhysicsForceY, OOParticleMassDT, ParticlesVelocityY);
+			ParticlesVelocityZ = VectorMultiplyAdd(PhysicsForceZ, OOParticleMassDT, ParticlesVelocityZ);
+
+			// Apply Drag
+			VectorRegister ClampedDrag = VectorMultiply(PhysicsDrag, EngineDeltaTime);
+			ClampedDrag = VectorMax(VectorMin(ClampedDrag, VectorOne()), VectorZero());
+			ClampedDrag = VectorNegate(ClampedDrag);
+
+			ParticlesVelocityX = VectorMultiplyAdd(ParticlesVelocityX, ClampedDrag, ParticlesVelocityX);
+			ParticlesVelocityY = VectorMultiplyAdd(ParticlesVelocityY, ClampedDrag, ParticlesVelocityY);
+			ParticlesVelocityZ = VectorMultiplyAdd(ParticlesVelocityZ, ClampedDrag, ParticlesVelocityZ);
+
+			// Velocity Clamp
+			//-TODO: Not used
+
+			// Limit Acceleration
+			//-TODO: Not used
+
+			// Apply velocity
+			ParticlesPositionX = VectorMultiplyAdd(ParticlesVelocityX, EngineDeltaTime, ParticlesPositionX);
+			ParticlesPositionY = VectorMultiplyAdd(ParticlesVelocityY, EngineDeltaTime, ParticlesPositionY);
+			ParticlesPositionZ = VectorMultiplyAdd(ParticlesVelocityZ, EngineDeltaTime, ParticlesPositionZ);
+
+			// Write parameters
+			VectorStore(ParticlesPositionX, OutParticlesPositionX.GetDestAndAdvance());
+			VectorStore(ParticlesPositionY, OutParticlesPositionY.GetDestAndAdvance());
+			VectorStore(ParticlesPositionZ, OutParticlesPositionZ.GetDestAndAdvance());
+
+			VectorStore(ParticlesVelocityX, OutParticlesVelocityX.GetDestAndAdvance());
+			VectorStore(ParticlesVelocityY, OutParticlesVelocityY.GetDestAndAdvance());
+			VectorStore(ParticlesVelocityZ, OutParticlesVelocityZ.GetDestAndAdvance());
+		}
+#endif
+	}
+
+	template<bool bForceConstant, bool bDragConstant, bool bMassConstant>
+	FORCEINLINE static void ExecOptimized(FVectorVMContext& Context)
+	{
+#if 0
+		const VectorRegister MassMin = VectorSetFloat1(0.0001f);
+		const VectorRegister EngineDeltaTime = VectorSetFloat1(VectorVM::FExternalFuncInputHandler<float>(Context).Get());
+		VectorRegister PhysicsForceX = VectorLoadFloat1(VectorVM::FExternalFuncInputHandler<VectorRegister>(Context).GetDest());
+		VectorRegister PhysicsForceY = VectorLoadFloat1(VectorVM::FExternalFuncInputHandler<VectorRegister>(Context).GetDest());
+		VectorRegister PhysicsForceZ = VectorLoadFloat1(VectorVM::FExternalFuncInputHandler<VectorRegister>(Context).GetDest());
+		VectorRegister PhysicsDrag = VectorLoadFloat1(VectorVM::FExternalFuncInputHandler<VectorRegister>(Context).GetDest());
+		VectorRegister ParticlesMass = VectorLoadFloat1(VectorVM::FExternalFuncInputHandler<VectorRegister>(Context).GetDest());
+		const VectorRegister* RESTRICT InParticlesPositionX = VectorVM::FExternalFuncInputHandler<VectorRegister>(Context).GetDest();
+		const VectorRegister* RESTRICT InParticlesPositionY = VectorVM::FExternalFuncInputHandler<VectorRegister>(Context).GetDest();
+		const VectorRegister* RESTRICT InParticlesPositionZ = VectorVM::FExternalFuncInputHandler<VectorRegister>(Context).GetDest();
+		const VectorRegister* RESTRICT InParticlesVelocityX = VectorVM::FExternalFuncInputHandler<VectorRegister>(Context).GetDest();
+		const VectorRegister* RESTRICT InParticlesVelocityY = VectorVM::FExternalFuncInputHandler<VectorRegister>(Context).GetDest();
+		const VectorRegister* RESTRICT InParticlesVelocityZ = VectorVM::FExternalFuncInputHandler<VectorRegister>(Context).GetDest();
+
+		VectorRegister* RESTRICT OutParticlesPositionX = VectorVM::FExternalFuncRegisterHandler<VectorRegister>(Context).GetDest();
+		VectorRegister* RESTRICT OutParticlesPositionY = VectorVM::FExternalFuncRegisterHandler<VectorRegister>(Context).GetDest();
+		VectorRegister* RESTRICT OutParticlesPositionZ = VectorVM::FExternalFuncRegisterHandler<VectorRegister>(Context).GetDest();
+		VectorRegister* RESTRICT OutParticlesVelocityX = VectorVM::FExternalFuncRegisterHandler<VectorRegister>(Context).GetDest();
+		VectorRegister* RESTRICT OutParticlesVelocityY = VectorVM::FExternalFuncRegisterHandler<VectorRegister>(Context).GetDest();
+		VectorRegister* RESTRICT OutParticlesVelocityZ = VectorVM::FExternalFuncRegisterHandler<VectorRegister>(Context).GetDest();
+		VectorRegister* RESTRICT OutParticlesPreviousVelocityX = VectorVM::FExternalFuncRegisterHandler<VectorRegister>(Context).GetDest();
+		VectorRegister* RESTRICT OutParticlesPreviousVelocityY = VectorVM::FExternalFuncRegisterHandler<VectorRegister>(Context).GetDest();
+		VectorRegister* RESTRICT OutParticlesPreviousVelocityZ = VectorVM::FExternalFuncRegisterHandler<VectorRegister>(Context).GetDest();
+
+		const VectorRegister OOParticleMassDT = VectorMultiply(VectorReciprocal(VectorMax(ParticlesMass, MassMin)), EngineDeltaTime);
+
+		VectorRegister ClampedDrag = VectorMultiply(PhysicsDrag, EngineDeltaTime);
+		ClampedDrag = VectorMax(VectorMin(ClampedDrag, VectorOne()), VectorZero());
+		ClampedDrag = VectorNegate(ClampedDrag);
+
+		for (int i = 0; i < Context.GetNumLoops<4>(); ++i)
+		{
+			// Gather values
+			VectorRegister ParticlesPositionX = VectorLoad(InParticlesPositionX + i);
+			VectorRegister ParticlesPositionY = VectorLoad(InParticlesPositionY + i);
+			VectorRegister ParticlesPositionZ = VectorLoad(InParticlesPositionZ + i);
+			VectorRegister ParticlesVelocityX = VectorLoad(InParticlesVelocityX + i);
+			VectorRegister ParticlesVelocityY = VectorLoad(InParticlesVelocityY + i);
+			VectorRegister ParticlesVelocityZ = VectorLoad(InParticlesVelocityZ + i);
+
+			VectorStore(ParticlesVelocityX, OutParticlesPreviousVelocityX + i);
+			VectorStore(ParticlesVelocityY, OutParticlesPreviousVelocityY + i);
+			VectorStore(ParticlesVelocityZ, OutParticlesPreviousVelocityZ + i);
+
+			// Apply velocity
+			ParticlesVelocityX = VectorMultiplyAdd(PhysicsForceX, OOParticleMassDT, ParticlesVelocityX);
+			ParticlesVelocityY = VectorMultiplyAdd(PhysicsForceY, OOParticleMassDT, ParticlesVelocityY);
+			ParticlesVelocityZ = VectorMultiplyAdd(PhysicsForceZ, OOParticleMassDT, ParticlesVelocityZ);
+
+			// Apply Drag
+			ParticlesVelocityX = VectorMultiplyAdd(ParticlesVelocityX, ClampedDrag, ParticlesVelocityX);
+			ParticlesVelocityY = VectorMultiplyAdd(ParticlesVelocityY, ClampedDrag, ParticlesVelocityY);
+			ParticlesVelocityZ = VectorMultiplyAdd(ParticlesVelocityZ, ClampedDrag, ParticlesVelocityZ);
+
+			// Velocity Clamp
+			//-TODO: Not used
+
+			// Limit Acceleration
+			//-TODO: Not used
+
+			// Apply velocity
+			ParticlesPositionX = VectorMultiplyAdd(ParticlesVelocityX, EngineDeltaTime, ParticlesPositionX);
+			ParticlesPositionY = VectorMultiplyAdd(ParticlesVelocityY, EngineDeltaTime, ParticlesPositionY);
+			ParticlesPositionZ = VectorMultiplyAdd(ParticlesVelocityZ, EngineDeltaTime, ParticlesPositionZ);
+
+			// Write parameters
+			VectorStore(ParticlesPositionX, OutParticlesPositionX + i);
+			VectorStore(ParticlesPositionY, OutParticlesPositionY + i);
+			VectorStore(ParticlesPositionZ, OutParticlesPositionZ + i);
+
+			VectorStore(ParticlesVelocityX, OutParticlesVelocityX + i);
+			VectorStore(ParticlesVelocityY, OutParticlesVelocityY + i);
+			VectorStore(ParticlesVelocityZ, OutParticlesVelocityZ + i);
+		}
+#else
+		VectorVM::FExternalFuncInputHandler<float> InEngineDeltaTime(Context);
+		VectorVM::FExternalFuncInputHandler<VectorRegister> InPhysicsForceXHandler(Context);
+		VectorVM::FExternalFuncInputHandler<VectorRegister> InPhysicsForceYHandler(Context);
+		VectorVM::FExternalFuncInputHandler<VectorRegister> InPhysicsForceZHandler(Context);
+		VectorVM::FExternalFuncInputHandler<VectorRegister> InPhysicsDragHandler(Context);
+		VectorVM::FExternalFuncInputHandler<VectorRegister> InParticlesMassHandler(Context);
+		VectorVM::FExternalFuncInputHandler<VectorRegister> InParticlesPositionXHandler(Context);
+		VectorVM::FExternalFuncInputHandler<VectorRegister> InParticlesPositionYHandler(Context);
+		VectorVM::FExternalFuncInputHandler<VectorRegister> InParticlesPositionZHandler(Context);
+		VectorVM::FExternalFuncInputHandler<VectorRegister> InParticlesVelocityXHandler(Context);
+		VectorVM::FExternalFuncInputHandler<VectorRegister> InParticlesVelocityYHandler(Context);
+		VectorVM::FExternalFuncInputHandler<VectorRegister> InParticlesVelocityZHandler(Context);
+
+		VectorVM::FExternalFuncRegisterHandler<VectorRegister> OutParticlesPositionXHandler(Context);
+		VectorVM::FExternalFuncRegisterHandler<VectorRegister> OutParticlesPositionYHandler(Context);
+		VectorVM::FExternalFuncRegisterHandler<VectorRegister> OutParticlesPositionZHandler(Context);
+		VectorVM::FExternalFuncRegisterHandler<VectorRegister> OutParticlesVelocityXHandler(Context);
+		VectorVM::FExternalFuncRegisterHandler<VectorRegister> OutParticlesVelocityYHandler(Context);
+		VectorVM::FExternalFuncRegisterHandler<VectorRegister> OutParticlesVelocityZHandler(Context);
+		VectorVM::FExternalFuncRegisterHandler<VectorRegister> OutParticlesPreviousVelocityXHandler(Context);
+		VectorVM::FExternalFuncRegisterHandler<VectorRegister> OutParticlesPreviousVelocityYHandler(Context);
+		VectorVM::FExternalFuncRegisterHandler<VectorRegister> OutParticlesPreviousVelocityZHandler(Context);
+
+		const VectorRegister EngineDeltaTime = VectorSetFloat1(InEngineDeltaTime.Get());
+		const VectorRegister* RESTRICT InPhysicsForceX = InPhysicsForceXHandler.GetDest();
+		const VectorRegister* RESTRICT InPhysicsForceY = InPhysicsForceYHandler.GetDest();
+		const VectorRegister* RESTRICT InPhysicsForceZ = InPhysicsForceZHandler.GetDest();
+		const VectorRegister* RESTRICT InPhysicsDrag = InPhysicsDragHandler.GetDest();
+		const VectorRegister* RESTRICT InParticlesMass = InParticlesMassHandler.GetDest();
+		const VectorRegister* RESTRICT InParticlesPositionX = InParticlesPositionXHandler.GetDest();
+		const VectorRegister* RESTRICT InParticlesPositionY = InParticlesPositionYHandler.GetDest();
+		const VectorRegister* RESTRICT InParticlesPositionZ = InParticlesPositionZHandler.GetDest();
+		const VectorRegister* RESTRICT InParticlesVelocityX = InParticlesVelocityXHandler.GetDest();
+		const VectorRegister* RESTRICT InParticlesVelocityY = InParticlesVelocityYHandler.GetDest();
+		const VectorRegister* RESTRICT InParticlesVelocityZ = InParticlesVelocityZHandler.GetDest();
+
+		VectorRegister* RESTRICT OutParticlesPositionX = OutParticlesPositionXHandler.GetDest();
+		VectorRegister* RESTRICT OutParticlesPositionY = OutParticlesPositionYHandler.GetDest();
+		VectorRegister* RESTRICT OutParticlesPositionZ = OutParticlesPositionZHandler.GetDest();
+		VectorRegister* RESTRICT OutParticlesVelocityX = OutParticlesVelocityXHandler.GetDest();
+		VectorRegister* RESTRICT OutParticlesVelocityY = OutParticlesVelocityYHandler.GetDest();
+		VectorRegister* RESTRICT OutParticlesVelocityZ = OutParticlesVelocityZHandler.GetDest();
+		VectorRegister* RESTRICT OutParticlesPreviousVelocityX = OutParticlesPreviousVelocityXHandler.GetDest();
+		VectorRegister* RESTRICT OutParticlesPreviousVelocityY = OutParticlesPreviousVelocityYHandler.GetDest();
+		VectorRegister* RESTRICT OutParticlesPreviousVelocityZ = OutParticlesPreviousVelocityZHandler.GetDest();
+
+		const VectorRegister MassMin = VectorSetFloat1(0.0001f);
+
+		const VectorRegister ConstantPhysicsForceX = bForceConstant ? VectorLoadFloat1(InPhysicsForceX) : VectorZero();
+		const VectorRegister ConstantPhysicsForceY = bForceConstant ? VectorLoadFloat1(InPhysicsForceY) : VectorZero();
+		const VectorRegister ConstantPhysicsForceZ = bForceConstant ? VectorLoadFloat1(InPhysicsForceZ) : VectorZero();
+		const VectorRegister ConstantPhysicsDrag = bDragConstant ? VectorLoadFloat1(InPhysicsDrag) : VectorZero();
+		const VectorRegister ConstantParticlesMass = bMassConstant ? VectorLoadFloat1(InParticlesMass) : VectorZero();
+
+		for (int i = 0; i < Context.GetNumLoops<4>(); ++i)
+		{
+			// Gather values
+			VectorRegister PhysicsForceX = bForceConstant ? ConstantPhysicsForceX : VectorLoad(InPhysicsForceX + i);
+			VectorRegister PhysicsForceY = bForceConstant ? ConstantPhysicsForceY : VectorLoad(InPhysicsForceY + i);
+			VectorRegister PhysicsForceZ = bForceConstant ? ConstantPhysicsForceZ : VectorLoad(InPhysicsForceZ + i);
+			VectorRegister PhysicsDrag = bDragConstant ? ConstantPhysicsDrag : VectorLoad(InPhysicsDrag + i);
+
+			VectorRegister ParticlesMass = bMassConstant ? ConstantParticlesMass : VectorLoad(InParticlesMass + i);
+			VectorRegister ParticlesPositionX = VectorLoad(InParticlesPositionX + i);
+			VectorRegister ParticlesPositionY = VectorLoad(InParticlesPositionY + i);
+			VectorRegister ParticlesPositionZ = VectorLoad(InParticlesPositionZ + i);
+			VectorRegister ParticlesVelocityX = VectorLoad(InParticlesVelocityX + i);
+			VectorRegister ParticlesVelocityY = VectorLoad(InParticlesVelocityY + i);
+			VectorRegister ParticlesVelocityZ = VectorLoad(InParticlesVelocityZ + i);
+
+			VectorStore(ParticlesVelocityX, OutParticlesPreviousVelocityX + i);
+			VectorStore(ParticlesVelocityY, OutParticlesPreviousVelocityY + i);
+			VectorStore(ParticlesVelocityZ, OutParticlesPreviousVelocityZ + i);
+
+			// Apply velocity
+			const VectorRegister OOParticleMassDT = VectorMultiply(VectorReciprocal(VectorMax(ParticlesMass, MassMin)), EngineDeltaTime);
+			ParticlesVelocityX = VectorMultiplyAdd(PhysicsForceX, OOParticleMassDT, ParticlesVelocityX);
+			ParticlesVelocityY = VectorMultiplyAdd(PhysicsForceY, OOParticleMassDT, ParticlesVelocityY);
+			ParticlesVelocityZ = VectorMultiplyAdd(PhysicsForceZ, OOParticleMassDT, ParticlesVelocityZ);
+
+			// Apply Drag
+			VectorRegister ClampedDrag = VectorMultiply(PhysicsDrag, EngineDeltaTime);
+			ClampedDrag = VectorMax(VectorMin(ClampedDrag, VectorOne()), VectorZero());
+			ClampedDrag = VectorNegate(ClampedDrag);
+
+			ParticlesVelocityX = VectorMultiplyAdd(ParticlesVelocityX, ClampedDrag, ParticlesVelocityX);
+			ParticlesVelocityY = VectorMultiplyAdd(ParticlesVelocityY, ClampedDrag, ParticlesVelocityY);
+			ParticlesVelocityZ = VectorMultiplyAdd(ParticlesVelocityZ, ClampedDrag, ParticlesVelocityZ);
+
+			// Velocity Clamp
+			//-TODO: Not used
+
+			// Limit Acceleration
+			//-TODO: Not used
+
+			// Apply velocity
+			ParticlesPositionX = VectorMultiplyAdd(ParticlesVelocityX, EngineDeltaTime, ParticlesPositionX);
+			ParticlesPositionY = VectorMultiplyAdd(ParticlesVelocityY, EngineDeltaTime, ParticlesPositionY);
+			ParticlesPositionZ = VectorMultiplyAdd(ParticlesVelocityZ, EngineDeltaTime, ParticlesPositionZ);
+
+			// Write parameters
+			VectorStore(ParticlesPositionX, OutParticlesPositionX + i);
+			VectorStore(ParticlesPositionY, OutParticlesPositionY + i);
+			VectorStore(ParticlesPositionZ, OutParticlesPositionZ + i);
+
+			VectorStore(ParticlesVelocityX, OutParticlesVelocityX + i);
+			VectorStore(ParticlesVelocityY, OutParticlesVelocityY + i);
+			VectorStore(ParticlesVelocityZ, OutParticlesVelocityZ + i);
+		}
+#endif
+	}
+};
+
 bool UNiagaraFunctionLibrary::GetVectorVMFastPathExternalFunction(const FVMExternalFunctionBindingInfo& BindingInfo, FVMExternalFunction &OutFunc)
 {
 	if (BindingInfo.Name == FastPathDot4Name)
@@ -738,6 +1442,68 @@ bool UNiagaraFunctionLibrary::GetVectorVMFastPathExternalFunction(const FVMExter
 		OutFunc = FVMExternalFunction::CreateStatic(FVectorKernelFastMatrixToQuaternion::Exec);
 		return true;
 	}
+	else if (BindingInfo.Name == FastPathEmitterLifeCycleName)
+	{
+		OutFunc = FVMExternalFunction::CreateStatic(FVectorKernel_EmitterLifeCycle::Exec);
+		return true;
+	}
+	else if (BindingInfo.Name == FastPathSpawnRateName)
+	{
+		OutFunc = FVMExternalFunction::CreateStatic(FVectorKernel_SpawnRate::Exec);
+		return true;
+	}
+	else if (BindingInfo.Name == FastPathSpawnBurstInstantaneousName)
+	{
+		OutFunc = FVMExternalFunction::CreateStatic(FVectorKernel_SpawnBurstInstantaneous::Exec);
+		return true;
+	}
+	else if (BindingInfo.Name == FastPathSolveVelocitiesAndForces)
+	{
+#if 0
+		OutFunc = FVMExternalFunction::CreateStatic(FVectorKernel_SolveVelocitiesAndForces::Exec);
+#else
+		const bool bForceConstant = BindingInfo.InputParamLocations[1] && BindingInfo.InputParamLocations[2] && BindingInfo.InputParamLocations[3];
+		const bool bDragConstant = BindingInfo.InputParamLocations[4];
+		const bool bMassConstant = BindingInfo.InputParamLocations[5];
+
+		if ( bForceConstant )
+		{
+			if ( bDragConstant )
+			{
+				if ( bMassConstant )
+					OutFunc = FVMExternalFunction::CreateStatic(FVectorKernel_SolveVelocitiesAndForces::ExecOptimized<true, true, true>);
+				else
+					OutFunc = FVMExternalFunction::CreateStatic(FVectorKernel_SolveVelocitiesAndForces::ExecOptimized<true, true, false>);
+			}
+			else
+			{
+				if (bMassConstant)
+					OutFunc = FVMExternalFunction::CreateStatic(FVectorKernel_SolveVelocitiesAndForces::ExecOptimized<true, false, true>);
+				else
+					OutFunc = FVMExternalFunction::CreateStatic(FVectorKernel_SolveVelocitiesAndForces::ExecOptimized<true, false, false>);
+			}
+		}
+		else
+		{
+			if (bDragConstant)
+			{
+				if (bMassConstant)
+					OutFunc = FVMExternalFunction::CreateStatic(FVectorKernel_SolveVelocitiesAndForces::ExecOptimized<false, true, true>);
+				else
+					OutFunc = FVMExternalFunction::CreateStatic(FVectorKernel_SolveVelocitiesAndForces::ExecOptimized<false, true, false>);
+			}
+			else
+			{
+				if (bMassConstant)
+					OutFunc = FVMExternalFunction::CreateStatic(FVectorKernel_SolveVelocitiesAndForces::ExecOptimized<false, false, true>);
+				else
+					OutFunc = FVMExternalFunction::CreateStatic(FVectorKernel_SolveVelocitiesAndForces::Exec);
+					//OutFunc = FVMExternalFunction::CreateStatic(FVectorKernel_SolveVelocitiesAndForces::ExecOptimized<false, false, false>);
+			}
+		}
+#endif
+		return true;
+	}
 	return false;
 }
 
@@ -750,10 +1516,18 @@ void UNiagaraFunctionLibrary::InitVectorVMFastPathOps()
 	VectorVMOps.Emplace(FVectorKernelFastDot4::GetFunctionSignature());
 	VectorVMOps.Emplace(FVectorKernelFastTransformPosition::GetFunctionSignature());
 	VectorVMOps.Emplace(FVectorKernelFastMatrixToQuaternion::GetFunctionSignature());
+	VectorVMOps.Emplace(FVectorKernel_EmitterLifeCycle::GetFunctionSignature());
+	VectorVMOps.Emplace(FVectorKernel_SpawnRate::GetFunctionSignature());
+	VectorVMOps.Emplace(FVectorKernel_SpawnBurstInstantaneous::GetFunctionSignature());
+	VectorVMOps.Emplace(FVectorKernel_SolveVelocitiesAndForces::GetFunctionSignature());
 
 	VectorVMOpsHLSL.Emplace(FVectorKernelFastDot4::GetFunctionHLSL());
 	VectorVMOpsHLSL.Emplace(FVectorKernelFastTransformPosition::GetFunctionHLSL());
 	VectorVMOpsHLSL.Emplace(FVectorKernelFastMatrixToQuaternion::GetFunctionHLSL());
+	VectorVMOpsHLSL.Emplace(FVectorKernel_EmitterLifeCycle::GetFunctionHLSL());
+	VectorVMOpsHLSL.Emplace(FVectorKernel_SpawnRate::GetFunctionHLSL());
+	VectorVMOpsHLSL.Emplace(FVectorKernel_SpawnBurstInstantaneous::GetFunctionHLSL());
+	VectorVMOpsHLSL.Emplace(FVectorKernel_SolveVelocitiesAndForces::GetFunctionHLSL());
 
 	check(VectorVMOps.Num() == VectorVMOpsHLSL.Num());
 }
