Fix for issue with niagara deferred deletion fences during cook.

#rnx

diff --git a/Engine/Plugins/FX/Niagara/Source/Niagara/Classes/NiagaraEmitterInstanceBatcher.h b/Engine/Plugins/FX/Niagara/Source/Niagara/Classes/NiagaraEmitterInstanceBatcher.h
--- a/Engine/Plugins/FX/Niagara/Source/Niagara/Classes/NiagaraEmitterInstanceBatcher.h
+++ b/Engine/Plugins/FX/Niagara/Source/Niagara/Classes/NiagaraEmitterInstanceBatcher.h
@@ -30,6 +30,12 @@ struct FNiagaraInstanceBatcherDeferredDeletionFence
 	FNiagaraDeferredDeletionFence* Fence;
 	FNiagaraInstanceBatcherDeferredDeletionFence(FNiagaraDeferredDeletionFence* InFence);
 	~FNiagaraInstanceBatcherDeferredDeletionFence();
+
+	FNiagaraInstanceBatcherDeferredDeletionFence(FNiagaraInstanceBatcherDeferredDeletionFence&& Other);
+	FNiagaraInstanceBatcherDeferredDeletionFence& operator=(FNiagaraInstanceBatcherDeferredDeletionFence&& Other);
+
+	FNiagaraInstanceBatcherDeferredDeletionFence(const FNiagaraInstanceBatcherDeferredDeletionFence& Other) = delete;
+	FNiagaraInstanceBatcherDeferredDeletionFence& operator=(const FNiagaraInstanceBatcherDeferredDeletionFence& Other) = delete;
 };
 
 class FNiagaraIndicesVertexBuffer : public FParticleIndicesVertexBuffer
@@ -78,7 +84,7 @@ public:
 		DIProxyDeferredDeletes_RT.Add(MoveTemp(Proxy));
 	}
 
-	void AddFence_RenderThread(FNiagaraDeferredDeletionFence* Fence);
+	void AddFence_RenderThread(FNiagaraInstanceBatcherDeferredDeletionFence& Fence);
 
 #if WITH_EDITOR
 	virtual void Suspend() override {}
diff --git a/Engine/Plugins/FX/Niagara/Source/Niagara/Private/NiagaraEmitterInstanceBatcher.cpp b/Engine/Plugins/FX/Niagara/Source/Niagara/Private/NiagaraEmitterInstanceBatcher.cpp
--- a/Engine/Plugins/FX/Niagara/Source/Niagara/Private/NiagaraEmitterInstanceBatcher.cpp
+++ b/Engine/Plugins/FX/Niagara/Source/Niagara/Private/NiagaraEmitterInstanceBatcher.cpp
@@ -131,10 +131,10 @@ void NiagaraEmitterInstanceBatcher::GiveDataSetToDestroy_RenderThread(FNiagaraDa
 	DataSetsToDestroy_RT.Add(DataSet);
 }
 
-void NiagaraEmitterInstanceBatcher::AddFence_RenderThread(FNiagaraDeferredDeletionFence* Fence)
+void NiagaraEmitterInstanceBatcher::AddFence_RenderThread(FNiagaraInstanceBatcherDeferredDeletionFence& Fence)
 {
 	LLM_SCOPE(ELLMTag::Niagara);
-	Fences_RT.Emplace(Fence);
+	Fences_RT.Emplace(MoveTemp(Fence));
 }
 
 void NiagaraEmitterInstanceBatcher::FinishDispatches()
diff --git a/Engine/Plugins/FX/Niagara/Source/Niagara/Private/NiagaraSystem.cpp b/Engine/Plugins/FX/Niagara/Source/Niagara/Private/NiagaraSystem.cpp
--- a/Engine/Plugins/FX/Niagara/Source/Niagara/Private/NiagaraSystem.cpp
+++ b/Engine/Plugins/FX/Niagara/Source/Niagara/Private/NiagaraSystem.cpp
@@ -1344,8 +1344,13 @@ FNiagaraEmitterCompiledData::FNiagaraEmitterCompiledData()
 
 FNiagaraDeferredDeletionFence::FNiagaraDeferredDeletionFence()
 {
-	bInstanceBatcherComplete = false;
-	bWorldManagerComplete = false;
+	bInstanceBatcherComplete = true;
+	bWorldManagerComplete = true;
+}
+
+FNiagaraDeferredDeletionFence::~FNiagaraDeferredDeletionFence()
+{
+	check(IsComplete());
 }
 
 bool FNiagaraDeferredDeletionFence::IsComplete()
@@ -1358,6 +1363,8 @@ FNiagaraWorldManagerDeferredDeletionFence::FNiagaraWorldManagerDeferredDeletionF
 {
 	if (Fence)
 	{
+		//UE_LOG(LogNiagara, Display, TEXT("%0xP - Create WorldManFence"), Fence);
+		check(Fence->bWorldManagerComplete == true);
 		Fence->bWorldManagerComplete = false;
 	}
 }
@@ -1366,25 +1373,68 @@ FNiagaraWorldManagerDeferredDeletionFence::~FNiagaraWorldManagerDeferredDeletion
 {
 	if (Fence)
 	{
+		//UE_LOG(LogNiagara, Display, TEXT("%0xP - Destroy WorldManFence"), Fence);
+		check(Fence->bWorldManagerComplete == false);
 		Fence->bWorldManagerComplete = true;
 	}
 }
 
+FNiagaraWorldManagerDeferredDeletionFence::FNiagaraWorldManagerDeferredDeletionFence(FNiagaraWorldManagerDeferredDeletionFence&& Other)
+	: Fence(Other.Fence)
+{
+	check(Fence->bWorldManagerComplete == false);
+	Other.Fence = nullptr;
+	//UE_LOG(LogNiagara, Display, TEXT("%0xP - MoveCtor WorldManFence"), Fence);
+}
+
+FNiagaraWorldManagerDeferredDeletionFence& FNiagaraWorldManagerDeferredDeletionFence::operator=(FNiagaraWorldManagerDeferredDeletionFence&& Other)
+{
+	check(Fence->bWorldManagerComplete == false);
+	Fence = Other.Fence;
+	Other.Fence = nullptr;
+	//UE_LOG(LogNiagara, Display, TEXT("%0xP - Move WorldManFence"), Fence);
+	return *this;
+}
+
+
 FNiagaraInstanceBatcherDeferredDeletionFence::FNiagaraInstanceBatcherDeferredDeletionFence(FNiagaraDeferredDeletionFence* InFence)
 	: Fence(InFence)
 {
 	if (Fence)
 	{
+		UE_LOG(LogNiagara, Display, TEXT("%0xP - Create BatcherFence"), Fence);
+		check(Fence->bInstanceBatcherComplete == true);
 		Fence->bInstanceBatcherComplete = false;
 	}
 }
+
 FNiagaraInstanceBatcherDeferredDeletionFence::~FNiagaraInstanceBatcherDeferredDeletionFence()
 {
 	if (Fence)
 	{
+		UE_LOG(LogNiagara, Display, TEXT("%0xP - Destroy BatcherFence"), Fence);
+		check(Fence->bInstanceBatcherComplete == false);
 		Fence->bInstanceBatcherComplete = true;
 	}
 }
 
+FNiagaraInstanceBatcherDeferredDeletionFence::FNiagaraInstanceBatcherDeferredDeletionFence(FNiagaraInstanceBatcherDeferredDeletionFence&& Other)
+	: Fence(Other.Fence)
+{
+	check(Fence->bInstanceBatcherComplete == false);
+	Other.Fence = nullptr;
+	UE_LOG(LogNiagara, Display, TEXT("%0xP - MoveCTor BatcherFence"), Fence);
+}
+
+FNiagaraInstanceBatcherDeferredDeletionFence& FNiagaraInstanceBatcherDeferredDeletionFence::operator=(FNiagaraInstanceBatcherDeferredDeletionFence&& Other)
+{
+	check(Fence->bInstanceBatcherComplete == false);
+	Fence = Other.Fence;
+	Other.Fence = nullptr;
+	UE_LOG(LogNiagara, Display, TEXT("%0xP - Move BatcherFence"), Fence);
+	return *this;
+}
+
+
 //////////////////////////////////////////////////////////////////////////
 //////////////////////////////////////////////////////////////////////////
\ No newline at end of file
diff --git a/Engine/Plugins/FX/Niagara/Source/Niagara/Private/NiagaraWorldManager.cpp b/Engine/Plugins/FX/Niagara/Source/Niagara/Private/NiagaraWorldManager.cpp
--- a/Engine/Plugins/FX/Niagara/Source/Niagara/Private/NiagaraWorldManager.cpp
+++ b/Engine/Plugins/FX/Niagara/Source/Niagara/Private/NiagaraWorldManager.cpp
@@ -567,6 +567,7 @@ void FNiagaraWorldManager::DumpDetails(FOutputDevice& Ar)
 
 void FNiagaraWorldManager::EnqueueDeferredDeletionFences(TArray<FNiagaraDeferredDeletionFence>& OutFences)
 {
+	OutFences.Reserve(WorldManagers.Num());
 	for (auto& Pair : WorldManagers)
 	{
 		UWorld* World = Pair.Key;
@@ -586,12 +587,13 @@ void FNiagaraWorldManager::EnqueueDeferredDeletionFences(TArray<FNiagaraDeferred
 				{
 					if (NiagaraEmitterInstanceBatcher* Batcher = static_cast<NiagaraEmitterInstanceBatcher*>(FXSystemInterface->GetInterface(NiagaraEmitterInstanceBatcher::Name)))
 					{
+						FNiagaraInstanceBatcherDeferredDeletionFence BatcherFence(Fence);
 						ENQUEUE_RENDER_COMMAND(FDeletionFenceCommand)(
-							[Batcher, Fence](FRHICommandListImmediate& RHICmdList)
+							[Batcher, BatcherFence_RT = MoveTemp(BatcherFence)](FRHICommandListImmediate& RHICmdList) mutable
 						{
 							if (ensure(Batcher))
 							{
-								Batcher->AddFence_RenderThread(Fence);
+								Batcher->AddFence_RenderThread(BatcherFence_RT);
 							}
 						}
 						);
diff --git a/Engine/Plugins/FX/Niagara/Source/Niagara/Public/NiagaraWorldManager.h b/Engine/Plugins/FX/Niagara/Source/Niagara/Public/NiagaraWorldManager.h
--- a/Engine/Plugins/FX/Niagara/Source/Niagara/Public/NiagaraWorldManager.h
+++ b/Engine/Plugins/FX/Niagara/Source/Niagara/Public/NiagaraWorldManager.h
@@ -31,6 +31,12 @@ struct FNiagaraWorldManagerDeferredDeletionFence
 	FNiagaraDeferredDeletionFence* Fence;
 	FNiagaraWorldManagerDeferredDeletionFence(FNiagaraDeferredDeletionFence* InFence);
 	~FNiagaraWorldManagerDeferredDeletionFence();
+
+	FNiagaraWorldManagerDeferredDeletionFence(FNiagaraWorldManagerDeferredDeletionFence&& Other);
+	FNiagaraWorldManagerDeferredDeletionFence& operator=(FNiagaraWorldManagerDeferredDeletionFence&& Other);
+
+	FNiagaraWorldManagerDeferredDeletionFence(const FNiagaraWorldManagerDeferredDeletionFence& Other) = delete;
+	FNiagaraWorldManagerDeferredDeletionFence& operator=(const FNiagaraWorldManagerDeferredDeletionFence& Other) = delete;
 };
 
 class FNiagaraViewDataMgr : public FRenderResource
