Speculative fix for compile error in FNiagaraDeferredDeletionFence
#rnx

diff --git a/Engine/Plugins/FX/Niagara/Source/Niagara/Classes/NiagaraSystem.h b/Engine/Plugins/FX/Niagara/Source/Niagara/Classes/NiagaraSystem.h
--- a/Engine/Plugins/FX/Niagara/Source/Niagara/Classes/NiagaraSystem.h
+++ b/Engine/Plugins/FX/Niagara/Source/Niagara/Classes/NiagaraSystem.h
@@ -22,6 +22,7 @@ class UNiagaraEditorDataBase;
 struct FNiagaraDeferredDeletionFence
 {
 	FNiagaraDeferredDeletionFence();
+	~FNiagaraDeferredDeletionFence();
 	bool IsComplete();
 
 	/** Flag to signal that the instance batcher deferred deletion queues have been cleared. */
