Set volumetric scattering default for light renderer to 0 (same as Cascade)

diff --git a/Engine/Plugins/FX/Niagara/Source/Niagara/Private/NiagaraConstants.cpp b/Engine/Plugins/FX/Niagara/Source/Niagara/Private/NiagaraConstants.cpp
--- a/Engine/Plugins/FX/Niagara/Source/Niagara/Private/NiagaraConstants.cpp
+++ b/Engine/Plugins/FX/Niagara/Source/Niagara/Private/NiagaraConstants.cpp
@@ -344,9 +344,9 @@ void FNiagaraConstants::Init()
 		Var.SetValue<float>(10.0f);
 		AttrDefaultsValueMap.Add(SYS_PARAM_PARTICLES_LIGHT_EXPONENT, Var);
 
-		AttrDefaultsStrMap.Add(SYS_PARAM_PARTICLES_LIGHT_VOLUMETRIC_SCATTERING, TEXT("1.0"));
+		AttrDefaultsStrMap.Add(SYS_PARAM_PARTICLES_LIGHT_VOLUMETRIC_SCATTERING, TEXT("0.0"));
 		Var = SYS_PARAM_PARTICLES_LIGHT_VOLUMETRIC_SCATTERING;
-		Var.SetValue<float>(1.0f);
+		Var.SetValue<float>(0.0f);
 		AttrDefaultsValueMap.Add(SYS_PARAM_PARTICLES_LIGHT_VOLUMETRIC_SCATTERING, Var);
 
 		AttrDefaultsStrMap.Add(SYS_PARAM_PARTICLES_RIBBONID, TEXT("0"));
