Niagara dynamic inputs now display a tooltip in the stack view.

diff --git a/Engine/Plugins/FX/Niagara/Source/NiagaraEditorWidgets/Private/Stack/SNiagaraStackFunctionInputValue.cpp b/Engine/Plugins/FX/Niagara/Source/NiagaraEditorWidgets/Private/Stack/SNiagaraStackFunctionInputValue.cpp
--- a/Engine/Plugins/FX/Niagara/Source/NiagaraEditorWidgets/Private/Stack/SNiagaraStackFunctionInputValue.cpp
+++ b/Engine/Plugins/FX/Niagara/Source/NiagaraEditorWidgets/Private/Stack/SNiagaraStackFunctionInputValue.cpp
@@ -132,6 +132,7 @@ void SNiagaraStackFunctionInputValue::Construct(const FArguments& InArgs, UNiaga
 						SNew(STextBlock)
 						.TextStyle(FNiagaraEditorStyle::Get(), "NiagaraEditor.ParameterText")
 						.Text(this, &SNiagaraStackFunctionInputValue::GetDynamicValueText)
+						.ToolTipText(this, &SNiagaraStackFunctionInputValue::GetDynamicValueToolTip)
 						.OnDoubleClicked(this, &SNiagaraStackFunctionInputValue::DynamicInputTextDoubleClicked)
 					]
 				]
@@ -377,6 +378,18 @@ FText SNiagaraStackFunctionInputValue::GetDynamicValueText() const
 	}
 }
 
+FText SNiagaraStackFunctionInputValue::GetDynamicValueToolTip() const
+{
+	if (FunctionInput->GetDynamicInputNode() != nullptr)
+	{
+		return FunctionInput->GetDynamicInputNode()->GetTooltipText();
+	}
+	else
+	{
+		return LOCTEXT("InvalidDynamicDisplayName", "(Invalid)");
+	}
+}
+
 FText SNiagaraStackFunctionInputValue::GetExpressionValueText() const
 {
 	if (FunctionInput->GetExpressionNode() != nullptr)
diff --git a/Engine/Plugins/FX/Niagara/Source/NiagaraEditorWidgets/Private/Stack/SNiagaraStackFunctionInputValue.h b/Engine/Plugins/FX/Niagara/Source/NiagaraEditorWidgets/Private/Stack/SNiagaraStackFunctionInputValue.h
--- a/Engine/Plugins/FX/Niagara/Source/NiagaraEditorWidgets/Private/Stack/SNiagaraStackFunctionInputValue.h
+++ b/Engine/Plugins/FX/Niagara/Source/NiagaraEditorWidgets/Private/Stack/SNiagaraStackFunctionInputValue.h
@@ -51,6 +51,8 @@ private:
 
 	FText GetDynamicValueText() const;
 
+	FText GetDynamicValueToolTip() const;
+
 	FText GetExpressionValueText() const;
 
 	void OnExpressionTextCommitted(const FText& Name, ETextCommit::Type CommitInfo);
