Niagara emitter templates should now compare correctly on all platforms.

diff --git a/Engine/Plugins/FX/Niagara/Source/NiagaraEditor/Private/SNiagaraTemplateAssetPicker.cpp b/Engine/Plugins/FX/Niagara/Source/NiagaraEditor/Private/SNiagaraTemplateAssetPicker.cpp
--- a/Engine/Plugins/FX/Niagara/Source/NiagaraEditor/Private/SNiagaraTemplateAssetPicker.cpp
+++ b/Engine/Plugins/FX/Niagara/Source/NiagaraEditor/Private/SNiagaraTemplateAssetPicker.cpp
@@ -118,12 +118,12 @@ bool SNiagaraTemplateAssetPicker::OnCompareCategoriesForSorting(const FText& Cat
 		}
 	}
 	// Otherwise just return the actual result.
-	return CompareResult == -1;
+	return CompareResult < 0;
 }
 
 bool SNiagaraTemplateAssetPicker::OnCompareItemsForSorting(const FAssetData& ItemA, const FAssetData& ItemB) const
 {
-	return ItemA.AssetName.ToString().Compare(ItemB.AssetName.ToString()) == -1;
+	return ItemA.AssetName.ToString().Compare(ItemB.AssetName.ToString()) < 0;
 }
 
 bool SNiagaraTemplateAssetPicker::OnDoesItemMatchFilterText(const FText& FilterText, const FAssetData& Item)
