Shift-clicking on an expander arrow in the niagara stack UI will now recursively expand or collapse children.
#rnx

diff --git a/Engine/Plugins/FX/Niagara/Source/NiagaraEditor/Private/ViewModels/Stack/NiagaraStackEntry.cpp b/Engine/Plugins/FX/Niagara/Source/NiagaraEditor/Private/ViewModels/Stack/NiagaraStackEntry.cpp
--- a/Engine/Plugins/FX/Niagara/Source/NiagaraEditor/Private/ViewModels/Stack/NiagaraStackEntry.cpp
+++ b/Engine/Plugins/FX/Niagara/Source/NiagaraEditor/Private/ViewModels/Stack/NiagaraStackEntry.cpp
@@ -224,13 +224,30 @@ bool UNiagaraStackEntry::GetIsExpanded() const
 
 void UNiagaraStackEntry::SetIsExpanded(bool bInExpanded)
 {
-	if (StackEditorData)
+	if (StackEditorData && GetCanExpand())
 	{
 		StackEditorData->SetStackEntryIsExpanded(GetStackEditorDataKey(), bInExpanded);
 	}
 	bIsExpandedCache.Reset();
 }
 
+void UNiagaraStackEntry::SetIsExpanded_Recursive(bool bInExpanded)
+{
+	if (StackEditorData && GetCanExpand())
+	{
+		StackEditorData->SetStackEntryIsExpanded(GetStackEditorDataKey(), bInExpanded);
+	}
+
+	TArray<UNiagaraStackEntry*> UnfilteredChildren;
+	GetUnfilteredChildren(UnfilteredChildren);
+	for (UNiagaraStackEntry* Child : UnfilteredChildren)
+	{
+		Child->SetIsExpanded_Recursive(bInExpanded);
+	}
+
+	bIsExpandedCache.Reset();
+}
+
 bool UNiagaraStackEntry::GetIsEnabled() const
 {
 	return true;
diff --git a/Engine/Plugins/FX/Niagara/Source/NiagaraEditor/Public/ViewModels/Stack/NiagaraStackEntry.h b/Engine/Plugins/FX/Niagara/Source/NiagaraEditor/Public/ViewModels/Stack/NiagaraStackEntry.h
--- a/Engine/Plugins/FX/Niagara/Source/NiagaraEditor/Public/ViewModels/Stack/NiagaraStackEntry.h
+++ b/Engine/Plugins/FX/Niagara/Source/NiagaraEditor/Public/ViewModels/Stack/NiagaraStackEntry.h
@@ -241,6 +241,8 @@ public:
 	// expanded state with the tree which is done to prevent items being expanded on tick.
 	void SetIsExpanded(bool bInExpanded);
 
+	void SetIsExpanded_Recursive(bool bInExpanded);
+
 	virtual bool GetIsEnabled() const;
 
 	bool GetOwnerIsEnabled() const;
diff --git a/Engine/Plugins/FX/Niagara/Source/NiagaraEditorWidgets/Private/Stack/SNiagaraStackTableRow.cpp b/Engine/Plugins/FX/Niagara/Source/NiagaraEditorWidgets/Private/Stack/SNiagaraStackTableRow.cpp
--- a/Engine/Plugins/FX/Niagara/Source/NiagaraEditorWidgets/Private/Stack/SNiagaraStackTableRow.cpp
+++ b/Engine/Plugins/FX/Niagara/Source/NiagaraEditorWidgets/Private/Stack/SNiagaraStackTableRow.cpp
@@ -414,7 +414,17 @@ EVisibility SNiagaraStackTableRow::GetExpanderVisibility() const
 
 FReply SNiagaraStackTableRow::ExpandButtonClicked()
 {
-	StackEntry->SetIsExpanded(!StackEntry->GetIsExpanded());
+	const bool bWillBeExpanded = !StackEntry->GetIsExpanded();
+	// Recurse the expansion if "shift" is being pressed
+	const FModifierKeysState ModKeyState = FSlateApplication::Get().GetModifierKeys();
+	if (ModKeyState.IsShiftDown())
+	{
+		StackEntry->SetIsExpanded_Recursive(bWillBeExpanded);
+	}
+	else
+	{
+		StackEntry->SetIsExpanded(bWillBeExpanded);
+	}
 	// Calling SetIsExpanded doesn't broadcast structure change automatically due to the expense of synchronizing
 	// expanded state with the tree to prevent items being expanded on tick, so we call this manually here.
 	StackEntry->OnStructureChanged().Broadcast();
