Removed problematic fence constructs and alt fix for orignial issue of compiled data lifetime.

#rnx

diff --git a/Engine/Plugins/FX/Niagara/Source/Niagara/Classes/NiagaraDataSet.h b/Engine/Plugins/FX/Niagara/Source/Niagara/Classes/NiagaraDataSet.h
--- a/Engine/Plugins/FX/Niagara/Source/Niagara/Classes/NiagaraDataSet.h
+++ b/Engine/Plugins/FX/Niagara/Source/Niagara/Classes/NiagaraDataSet.h
@@ -269,8 +269,10 @@ public:
 
 	FORCEINLINE void Init(const FNiagaraDataSetCompiledData* InDataSetCompiledData)
 	{
-		CompiledData = InDataSetCompiledData != nullptr ? InDataSetCompiledData : &FNiagaraDataSetCompiledData::DummyCompiledData;
-	
+		//CompiledData = InDataSetCompiledData != nullptr ? InDataSetCompiledData : &FNiagaraDataSetCompiledData::DummyCompiledData;
+		//Temporarily taking a copy of the compiled data to avoid lifetime issues in some cases.
+		CompiledData = InDataSetCompiledData != nullptr ? *InDataSetCompiledData : FNiagaraDataSetCompiledData::DummyCompiledData;
+		bInitialized = true;
 		Reset();
 	}
 
@@ -289,10 +291,10 @@ public:
 	/** Returns size in bytes for all data buffers currently allocated by this dataset. */
 	uint32 GetSizeBytes()const;
 
-	FORCEINLINE bool IsInitialized()const { return CompiledData != nullptr; }
-	FORCEINLINE ENiagaraSimTarget GetSimTarget()const { return CompiledData->SimTarget; }
-	FORCEINLINE FNiagaraDataSetID GetID()const { return CompiledData->ID; }	
-	FORCEINLINE bool GetNeedsPersistentIDs()const { return CompiledData->bNeedsPersistentIDs; }
+	FORCEINLINE bool IsInitialized()const { return bInitialized; }
+	FORCEINLINE ENiagaraSimTarget GetSimTarget()const { return CompiledData.SimTarget; }
+	FORCEINLINE FNiagaraDataSetID GetID()const { return CompiledData.ID; }	
+	FORCEINLINE bool GetNeedsPersistentIDs()const { return CompiledData.bNeedsPersistentIDs; }
 
 	FORCEINLINE TArray<int32>& GetFreeIDTable() { return FreeIDsTable; }
 	FORCEINLINE int32& GetNumFreeIDs() { return NumFreeIDs; }
@@ -300,13 +302,13 @@ public:
 	FORCEINLINE int32& GetIDAcquireTag() { return IDAcquireTag; }
 	FORCEINLINE void SetIDAcquireTag(int32 InTag) { IDAcquireTag = InTag; }
 
-	FORCEINLINE const TArray<FNiagaraVariable>& GetVariables()const { return CompiledData->Variables; }
-	FORCEINLINE uint32 GetNumVariables()const { return CompiledData->Variables.Num(); }
-	FORCEINLINE bool HasVariable(const FNiagaraVariable& Var)const { return CompiledData->Variables.Contains(Var); }
-	FORCEINLINE uint32 GetNumFloatComponents()const { return CompiledData->TotalFloatComponents; }
-	FORCEINLINE uint32 GetNumInt32Components()const { return CompiledData->TotalInt32Components; }
+	FORCEINLINE const TArray<FNiagaraVariable>& GetVariables()const { return CompiledData.Variables; }
+	FORCEINLINE uint32 GetNumVariables()const { return CompiledData.Variables.Num(); }
+	FORCEINLINE bool HasVariable(const FNiagaraVariable& Var)const { return CompiledData.Variables.Contains(Var); }
+	FORCEINLINE uint32 GetNumFloatComponents()const { return CompiledData.TotalFloatComponents; }
+	FORCEINLINE uint32 GetNumInt32Components()const { return CompiledData.TotalInt32Components; }
 
-	const TArray<FNiagaraVariableLayoutInfo>& GetVariableLayouts()const { return CompiledData->VariableLayouts; }
+	const TArray<FNiagaraVariableLayoutInfo>& GetVariableLayouts()const { return CompiledData.VariableLayouts; }
 	const FNiagaraVariableLayoutInfo* GetVariableLayout(const FNiagaraVariable& Var)const;
 	bool GetVariableComponentOffsets(const FNiagaraVariable& Var, int32 &FloatStart, int32 &IntStart) const;
 
@@ -358,8 +360,9 @@ private:
 #endif
 	}
 
-	const FNiagaraDataSetCompiledData* CompiledData;
-
+	//const FNiagaraDataSetCompiledData* CompiledData;
+	//For safety we're temporarily taking a copy of the compiled data. In certain cases the lifetime of the CompiledData ptr cannot be guaranteed. 
+	FNiagaraDataSetCompiledData CompiledData;
 
 	/** Table of free IDs available to allocate next tick. */
 	TArray<int32> FreeIDsTable;
@@ -387,6 +390,8 @@ private:
 	Additional buffers may be in here if they are currently being used by the render thread.
 	*/
 	TArray<FNiagaraDataBuffer*, TInlineAllocator<2>> Data;
+
+	bool bInitialized;
 };
 
 /**
diff --git a/Engine/Plugins/FX/Niagara/Source/Niagara/Classes/NiagaraEmitterInstanceBatcher.h b/Engine/Plugins/FX/Niagara/Source/Niagara/Classes/NiagaraEmitterInstanceBatcher.h
--- a/Engine/Plugins/FX/Niagara/Source/Niagara/Classes/NiagaraEmitterInstanceBatcher.h
+++ b/Engine/Plugins/FX/Niagara/Source/Niagara/Classes/NiagaraEmitterInstanceBatcher.h
@@ -23,21 +23,6 @@ the same VectorVM byte code / compute shader code
 #include "NiagaraScriptExecutionContext.h"
 #include "NiagaraGPUInstanceCountManager.h"
 
-struct FNiagaraDeferredDeletionFence;
-
-struct FNiagaraInstanceBatcherDeferredDeletionFence
-{
-	FNiagaraDeferredDeletionFence* Fence;
-	FNiagaraInstanceBatcherDeferredDeletionFence(FNiagaraDeferredDeletionFence* InFence);
-	~FNiagaraInstanceBatcherDeferredDeletionFence();
-
-	FNiagaraInstanceBatcherDeferredDeletionFence(FNiagaraInstanceBatcherDeferredDeletionFence&& Other);
-	FNiagaraInstanceBatcherDeferredDeletionFence& operator=(FNiagaraInstanceBatcherDeferredDeletionFence&& Other);
-
-	FNiagaraInstanceBatcherDeferredDeletionFence(const FNiagaraInstanceBatcherDeferredDeletionFence& Other) = delete;
-	FNiagaraInstanceBatcherDeferredDeletionFence& operator=(const FNiagaraInstanceBatcherDeferredDeletionFence& Other) = delete;
-};
-
 class FNiagaraIndicesVertexBuffer : public FParticleIndicesVertexBuffer
 {
 public:
@@ -84,8 +69,6 @@ public:
 		DIProxyDeferredDeletes_RT.Add(MoveTemp(Proxy));
 	}
 
-	void AddFence_RenderThread(FNiagaraInstanceBatcherDeferredDeletionFence& Fence);
-
 #if WITH_EDITOR
 	virtual void Suspend() override {}
 	virtual void Resume() override {}
@@ -220,6 +203,4 @@ private:
 	TSet<FNiagaraDataSet*> DataSetsToDestroy_RT;
 
 	TSet<TSharedPtr<FNiagaraDataInterfaceProxy, ESPMode::ThreadSafe>> DIProxyDeferredDeletes_RT;
-
-	TArray<FNiagaraInstanceBatcherDeferredDeletionFence> Fences_RT;
 };
diff --git a/Engine/Plugins/FX/Niagara/Source/Niagara/Classes/NiagaraSystem.h b/Engine/Plugins/FX/Niagara/Source/Niagara/Classes/NiagaraSystem.h
--- a/Engine/Plugins/FX/Niagara/Source/Niagara/Classes/NiagaraSystem.h
+++ b/Engine/Plugins/FX/Niagara/Source/Niagara/Classes/NiagaraSystem.h
@@ -19,20 +19,6 @@
 class UNiagaraEditorDataBase;
 #endif
 
-struct FNiagaraDeferredDeletionFence
-{
-	FNiagaraDeferredDeletionFence();
-	~FNiagaraDeferredDeletionFence();
-	bool IsComplete();
-
-	/** Flag to signal that the instance batcher deferred deletion queues have been cleared. */
-	TAtomic<bool> bInstanceBatcherComplete;
-
-	/** Flag to signal that the world manager deferred deletion queues have been cleared. */
-	TAtomic<bool> bWorldManagerComplete;
-};
-
-
 USTRUCT()
 struct FNiagaraEmitterCompiledData
 {
@@ -145,7 +131,6 @@ public:
 	virtual void PostLoad() override; 
 	virtual void BeginDestroy() override;
 	virtual void PreSave(const class ITargetPlatform * TargetPlatform) override;
-	virtual bool IsReadyForFinishDestroy()override;
 #if WITH_EDITOR
 	virtual void PostEditChangeProperty(struct FPropertyChangedEvent& PropertyChangedEvent) override; 
 	virtual void BeginCacheForCookedPlatformData(const ITargetPlatform *TargetPlatform) override;
@@ -419,7 +404,4 @@ protected:
 	mutable TStatId StatID_RT;
 	mutable TStatId StatID_RT_CNC;
 #endif
-
-	//We submit one deletion fence for each world to ensure this System lives longer than any worlds deferred deletion queue.
-	TArray<FNiagaraDeferredDeletionFence> DeletionFences;
 };
diff --git a/Engine/Plugins/FX/Niagara/Source/Niagara/Private/NiagaraDataSet.cpp b/Engine/Plugins/FX/Niagara/Source/Niagara/Private/NiagaraDataSet.cpp
--- a/Engine/Plugins/FX/Niagara/Source/Niagara/Private/NiagaraDataSet.cpp
+++ b/Engine/Plugins/FX/Niagara/Source/Niagara/Private/NiagaraDataSet.cpp
@@ -75,12 +75,13 @@ static FAutoConsoleVariableRef CVarNiagaraGPUDataBufferBufferSlack(
 );
 
 FNiagaraDataSet::FNiagaraDataSet()
-	: CompiledData(&FNiagaraDataSetCompiledData::DummyCompiledData)
+	: CompiledData(FNiagaraDataSetCompiledData::DummyCompiledData)
 	, NumFreeIDs(0)
 	, MaxUsedID(0)
 	, IDAcquireTag(0)
 	, CurrentData(nullptr)
 	, DestinationData(nullptr)
+	, bInitialized(false)
 {
 }
 
@@ -347,7 +348,7 @@ void FNiagaraDataSet::ReleaseGPUInstanceCounts(FNiagaraGPUInstanceCountManager&
 const FNiagaraVariableLayoutInfo* FNiagaraDataSet::GetVariableLayout(const FNiagaraVariable& Var)const
 {
 	int32 VarLayoutIndex = GetVariables().IndexOfByKey(Var);
-	return VarLayoutIndex != INDEX_NONE ? &CompiledData->VariableLayouts[VarLayoutIndex] : nullptr;
+	return VarLayoutIndex != INDEX_NONE ? &CompiledData.VariableLayouts[VarLayoutIndex] : nullptr;
 }
 
 bool FNiagaraDataSet::GetVariableComponentOffsets(const FNiagaraVariable& Var, int32 &FloatStart, int32 &IntStart) const
@@ -367,7 +368,7 @@ bool FNiagaraDataSet::GetVariableComponentOffsets(const FNiagaraVariable& Var, i
 
 void FNiagaraDataSet::CopyTo(FNiagaraDataSet& Other, int32 StartIdx, int32 NumInstances, bool bResetOther)const
 {
-	check(CompiledData);
+	//check(CompiledData);
 
 	CheckCorrectThread();
 
diff --git a/Engine/Plugins/FX/Niagara/Source/Niagara/Private/NiagaraEmitterInstanceBatcher.cpp b/Engine/Plugins/FX/Niagara/Source/Niagara/Private/NiagaraEmitterInstanceBatcher.cpp
--- a/Engine/Plugins/FX/Niagara/Source/Niagara/Private/NiagaraEmitterInstanceBatcher.cpp
+++ b/Engine/Plugins/FX/Niagara/Source/Niagara/Private/NiagaraEmitterInstanceBatcher.cpp
@@ -131,12 +131,6 @@ void NiagaraEmitterInstanceBatcher::GiveDataSetToDestroy_RenderThread(FNiagaraDa
 	DataSetsToDestroy_RT.Add(DataSet);
 }
 
-void NiagaraEmitterInstanceBatcher::AddFence_RenderThread(FNiagaraInstanceBatcherDeferredDeletionFence& Fence)
-{
-	LLM_SCOPE(ELLMTag::Niagara);
-	Fences_RT.Emplace(MoveTemp(Fence));
-}
-
 void NiagaraEmitterInstanceBatcher::FinishDispatches()
 {
 	ReleaseTicks();
@@ -165,8 +159,6 @@ void NiagaraEmitterInstanceBatcher::FinishDispatches()
 	}
 
 	DIProxyDeferredDeletes_RT.Empty();
-
-	Fences_RT.Empty();
 }
 
 void NiagaraEmitterInstanceBatcher::ReleaseTicks()
diff --git a/Engine/Plugins/FX/Niagara/Source/Niagara/Private/NiagaraSystem.cpp b/Engine/Plugins/FX/Niagara/Source/Niagara/Private/NiagaraSystem.cpp
--- a/Engine/Plugins/FX/Niagara/Source/Niagara/Private/NiagaraSystem.cpp
+++ b/Engine/Plugins/FX/Niagara/Source/Niagara/Private/NiagaraSystem.cpp
@@ -65,9 +65,6 @@ void UNiagaraSystem::BeginDestroy()
 {
 	Super::BeginDestroy();
 
-	DeletionFences.Reset();
-	FNiagaraWorldManager::EnqueueDeferredDeletionFences(DeletionFences);
-
 #if WITH_EDITORONLY_DATA
 	while (ActiveCompilations.Num() > 0)
 	{
@@ -76,20 +73,6 @@ void UNiagaraSystem::BeginDestroy()
 #endif
 }
 
-bool UNiagaraSystem::IsReadyForFinishDestroy()
-{
- 	bool bAllComplete = true;
-	
-	//We have to wait until the existing items in the deferred deletion queue are cleared as they have direct refs to the layout info in our System and Emitter compiled data.
-	//We have to do this for each world.
-	for (FNiagaraDeferredDeletionFence& Fence : DeletionFences)
-	{
-		bAllComplete &= Fence.IsComplete();
-	}
-	
-	return Super::IsReadyForFinishDestroy() && bAllComplete;
-}
-
 void UNiagaraSystem::PreSave(const class ITargetPlatform * TargetPlatform)
 {
 	Super::PreSave(TargetPlatform);
@@ -1337,104 +1320,4 @@ FNiagaraEmitterCompiledData::FNiagaraEmitterCompiledData()
 	EmitterSpawnGroupVar = SYS_PARAM_EMITTER_SPAWN_GROUP;
 	EmitterRandomSeedVar = SYS_PARAM_EMITTER_RANDOM_SEED;
 	EmitterTotalSpawnedParticlesVar = SYS_PARAM_ENGINE_EMITTER_TOTAL_SPAWNED_PARTICLES;
-}
-
-
-//////////////////////////////////////////////////////////////////////////
-
-FNiagaraDeferredDeletionFence::FNiagaraDeferredDeletionFence()
-{
-	bInstanceBatcherComplete = true;
-	bWorldManagerComplete = true;
-}
-
-FNiagaraDeferredDeletionFence::~FNiagaraDeferredDeletionFence()
-{
-	check(IsComplete());
-}
-
-bool FNiagaraDeferredDeletionFence::IsComplete()
-{
-	return bInstanceBatcherComplete.Load() && bWorldManagerComplete.Load();
-}
-
-FNiagaraWorldManagerDeferredDeletionFence::FNiagaraWorldManagerDeferredDeletionFence(FNiagaraDeferredDeletionFence* InFence)
-	: Fence(InFence)
-{
-	if (Fence)
-	{
-		//UE_LOG(LogNiagara, Display, TEXT("%0xP - Create WorldManFence"), Fence);
-		check(Fence->bWorldManagerComplete == true);
-		Fence->bWorldManagerComplete = false;
-	}
-}
-
-FNiagaraWorldManagerDeferredDeletionFence::~FNiagaraWorldManagerDeferredDeletionFence()
-{
-	if (Fence)
-	{
-		//UE_LOG(LogNiagara, Display, TEXT("%0xP - Destroy WorldManFence"), Fence);
-		check(Fence->bWorldManagerComplete == false);
-		Fence->bWorldManagerComplete = true;
-	}
-}
-
-FNiagaraWorldManagerDeferredDeletionFence::FNiagaraWorldManagerDeferredDeletionFence(FNiagaraWorldManagerDeferredDeletionFence&& Other)
-	: Fence(Other.Fence)
-{
-	check(Fence->bWorldManagerComplete == false);
-	Other.Fence = nullptr;
-	//UE_LOG(LogNiagara, Display, TEXT("%0xP - MoveCtor WorldManFence"), Fence);
-}
-
-FNiagaraWorldManagerDeferredDeletionFence& FNiagaraWorldManagerDeferredDeletionFence::operator=(FNiagaraWorldManagerDeferredDeletionFence&& Other)
-{
-	check(Fence->bWorldManagerComplete == false);
-	Fence = Other.Fence;
-	Other.Fence = nullptr;
-	//UE_LOG(LogNiagara, Display, TEXT("%0xP - Move WorldManFence"), Fence);
-	return *this;
-}
-
-
-FNiagaraInstanceBatcherDeferredDeletionFence::FNiagaraInstanceBatcherDeferredDeletionFence(FNiagaraDeferredDeletionFence* InFence)
-	: Fence(InFence)
-{
-	if (Fence)
-	{
-		UE_LOG(LogNiagara, Display, TEXT("%0xP - Create BatcherFence"), Fence);
-		check(Fence->bInstanceBatcherComplete == true);
-		Fence->bInstanceBatcherComplete = false;
-	}
-}
-
-FNiagaraInstanceBatcherDeferredDeletionFence::~FNiagaraInstanceBatcherDeferredDeletionFence()
-{
-	if (Fence)
-	{
-		UE_LOG(LogNiagara, Display, TEXT("%0xP - Destroy BatcherFence"), Fence);
-		check(Fence->bInstanceBatcherComplete == false);
-		Fence->bInstanceBatcherComplete = true;
-	}
-}
-
-FNiagaraInstanceBatcherDeferredDeletionFence::FNiagaraInstanceBatcherDeferredDeletionFence(FNiagaraInstanceBatcherDeferredDeletionFence&& Other)
-	: Fence(Other.Fence)
-{
-	check(Fence->bInstanceBatcherComplete == false);
-	Other.Fence = nullptr;
-	UE_LOG(LogNiagara, Display, TEXT("%0xP - MoveCTor BatcherFence"), Fence);
-}
-
-FNiagaraInstanceBatcherDeferredDeletionFence& FNiagaraInstanceBatcherDeferredDeletionFence::operator=(FNiagaraInstanceBatcherDeferredDeletionFence&& Other)
-{
-	check(Fence->bInstanceBatcherComplete == false);
-	Fence = Other.Fence;
-	Other.Fence = nullptr;
-	UE_LOG(LogNiagara, Display, TEXT("%0xP - Move BatcherFence"), Fence);
-	return *this;
-}
-
-
-//////////////////////////////////////////////////////////////////////////
-//////////////////////////////////////////////////////////////////////////
\ No newline at end of file
+}
\ No newline at end of file
diff --git a/Engine/Plugins/FX/Niagara/Source/Niagara/Private/NiagaraWorldManager.cpp b/Engine/Plugins/FX/Niagara/Source/Niagara/Private/NiagaraWorldManager.cpp
--- a/Engine/Plugins/FX/Niagara/Source/Niagara/Private/NiagaraWorldManager.cpp
+++ b/Engine/Plugins/FX/Niagara/Source/Niagara/Private/NiagaraWorldManager.cpp
@@ -313,7 +313,6 @@ void FNiagaraWorldManager::OnBatcherDestroyed_Internal(NiagaraEmitterInstanceBat
 		{
 			DeferredDeletionQueue[DeferredDeletionQueueIndex].Fence.Wait();
 			DeferredDeletionQueue[i].Queue.Empty();
-			DeferredDeletionQueue[i].DeletionFences.Empty();
 		}
 	}
 }
@@ -336,7 +335,6 @@ void FNiagaraWorldManager::OnWorldCleanup(bool bSessionEnded, bool bCleanupResou
 	{
 		DeferredDeletionQueue[i].Fence.Wait();
 		DeferredDeletionQueue[i].Queue.Empty();
-		DeferredDeletionQueue[i].DeletionFences.Empty();
 	}
 }
 
@@ -455,7 +453,6 @@ void FNiagaraWorldManager::PostActorTick(float DeltaSeconds)
 			DeferredDeletionQueue[DeferredDeletionQueueIndex].Fence.Wait();
 		}
 		DeferredDeletionQueue[DeferredDeletionQueueIndex].Queue.Empty();
-		DeferredDeletionQueue[DeferredDeletionQueueIndex].DeletionFences.Empty();
 	}
 
 	// Update tick groups
@@ -568,43 +565,4 @@ void FNiagaraWorldManager::DumpDetails(FOutputDevice& Ar)
 UWorld* FNiagaraWorldManager::GetWorld()
 {
 	return World;
-}
-
-void FNiagaraWorldManager::EnqueueDeferredDeletionFences(TArray<FNiagaraDeferredDeletionFence>& OutFences)
-{
-	OutFences.Reserve(WorldManagers.Num());
-	for (auto& Pair : WorldManagers)
-	{
-		UWorld* World = Pair.Key;
-		FNiagaraWorldManager* Manager = Pair.Value;
-
-		if (World && Manager)
-		{
-			FNiagaraDeferredDeletionFence* Fence = &OutFences.AddDefaulted_GetRef();
-			//Add the fence to the deferred deletion queue.
-			Manager->DeferredDeletionQueue[Manager->DeferredDeletionQueueIndex].DeletionFences.Emplace(Fence);
-
-			//Add the fence to the instance batcher
-			if (World->Scene)
-			{
-				FFXSystemInterface*  FXSystemInterface = World->Scene->GetFXSystem();
-				if (FXSystemInterface)
-				{
-					if (NiagaraEmitterInstanceBatcher* Batcher = static_cast<NiagaraEmitterInstanceBatcher*>(FXSystemInterface->GetInterface(NiagaraEmitterInstanceBatcher::Name)))
-					{
-						FNiagaraInstanceBatcherDeferredDeletionFence BatcherFence(Fence);
-						ENQUEUE_RENDER_COMMAND(FDeletionFenceCommand)(
-							[Batcher, BatcherFence_RT = MoveTemp(BatcherFence)](FRHICommandListImmediate& RHICmdList) mutable
-						{
-							if (ensure(Batcher))
-							{
-								Batcher->AddFence_RenderThread(BatcherFence_RT);
-							}
-						}
-						);
-					}
-				}
-			}
-		}
-	}
 }
\ No newline at end of file
diff --git a/Engine/Plugins/FX/Niagara/Source/Niagara/Public/NiagaraWorldManager.h b/Engine/Plugins/FX/Niagara/Source/Niagara/Public/NiagaraWorldManager.h
--- a/Engine/Plugins/FX/Niagara/Source/Niagara/Public/NiagaraWorldManager.h
+++ b/Engine/Plugins/FX/Niagara/Source/Niagara/Public/NiagaraWorldManager.h
@@ -24,21 +24,6 @@ class UNiagaraParameterCollection;
 class UNiagaraParameterCollectionInstance;
 class UNiagaraComponentPool;
 
-struct FNiagaraDeferredDeletionFence;
-
-struct FNiagaraWorldManagerDeferredDeletionFence
-{
-	FNiagaraDeferredDeletionFence* Fence;
-	FNiagaraWorldManagerDeferredDeletionFence(FNiagaraDeferredDeletionFence* InFence);
-	~FNiagaraWorldManagerDeferredDeletionFence();
-
-	FNiagaraWorldManagerDeferredDeletionFence(FNiagaraWorldManagerDeferredDeletionFence&& Other);
-	FNiagaraWorldManagerDeferredDeletionFence& operator=(FNiagaraWorldManagerDeferredDeletionFence&& Other);
-
-	FNiagaraWorldManagerDeferredDeletionFence(const FNiagaraWorldManagerDeferredDeletionFence& Other) = delete;
-	FNiagaraWorldManagerDeferredDeletionFence& operator=(const FNiagaraWorldManagerDeferredDeletionFence& Other) = delete;
-};
-
 class FNiagaraViewDataMgr : public FRenderResource
 {
 public:
@@ -145,12 +130,8 @@ public:
 
 	// Dump details about what's inside the world manager
 	void DumpDetails(FOutputDevice& Ar);
-
+	
 	UWorld* GetWorld();
-
-	//Creates fences for deferred deletion for each world.
-	static void EnqueueDeferredDeletionFences(TArray<FNiagaraDeferredDeletionFence>& OutFence);
-
 private:
 
 	// Callback function registered with global world delegates to instantiate world manager when a game world is created
@@ -203,9 +184,6 @@ private:
 	{
 		FRenderCommandFence							Fence;
 		TArray<TUniquePtr<FNiagaraSystemInstance>>	Queue;
-
-		/** Fences to signal to other Niagara code when this queue has been cleared. */
-		TArray<FNiagaraWorldManagerDeferredDeletionFence>		DeletionFences;
 	};
 
 	static constexpr int NumDeferredQueues = 3;
