Fix checkSlow no longer against a pointer
#rnx

diff --git a/Engine/Plugins/FX/Niagara/Source/Niagara/Private/NiagaraDataSet.cpp b/Engine/Plugins/FX/Niagara/Source/Niagara/Private/NiagaraDataSet.cpp
--- a/Engine/Plugins/FX/Niagara/Source/Niagara/Private/NiagaraDataSet.cpp
+++ b/Engine/Plugins/FX/Niagara/Source/Niagara/Private/NiagaraDataSet.cpp
@@ -99,7 +99,7 @@ void FNiagaraDataSet::Reset()
 
 void FNiagaraDataSet::ResetBuffers()
 {
-	checkSlow(CompiledData);
+	//checkSlow(CompiledData);
 	
 	if (GetSimTarget() == ENiagaraSimTarget::CPUSim)
 	{
