#Niagara Bumping Niagara DDC version due to native serialization change in 10269905

diff --git a/Engine/Plugins/FX/Niagara/Source/NiagaraShader/Private/NiagaraShaderDerivedDataVersion.h b/Engine/Plugins/FX/Niagara/Source/NiagaraShader/Private/NiagaraShaderDerivedDataVersion.h
--- a/Engine/Plugins/FX/Niagara/Source/NiagaraShader/Private/NiagaraShaderDerivedDataVersion.h
+++ b/Engine/Plugins/FX/Niagara/Source/NiagaraShader/Private/NiagaraShaderDerivedDataVersion.h
@@ -11,4 +11,4 @@ NiagaraShaderDerivedDataVersion.h: Shader derived data version for Niagara.
 // In case of merge conflicts with DDC versions, you MUST generate a new GUID and set this new
 // guid as version
 
-#define NIAGARASHADERMAP_DERIVEDDATA_VER		TEXT("068D842C69C776419BF530A17E1A388B")
+#define NIAGARASHADERMAP_DERIVEDDATA_VER		TEXT("A27AE38E639D4954AC529C61BD67058A")
