Added ability to mark Niagara modules, dynamic input classes and functions as Experimental, which will display an Info icon next to their name in the stack, node and parameters menu.

diff --git a/Engine/Plugins/FX/Niagara/Source/Niagara/Classes/NiagaraScript.h b/Engine/Plugins/FX/Niagara/Source/Niagara/Classes/NiagaraScript.h
--- a/Engine/Plugins/FX/Niagara/Source/Niagara/Classes/NiagaraScript.h
+++ b/Engine/Plugins/FX/Niagara/Source/Niagara/Classes/NiagaraScript.h
@@ -375,6 +375,9 @@ public:
 	UPROPERTY(EditAnywhere, Category = Script, meta = (EditCondition = "bDeprecated"))
 	UNiagaraScript* DeprecationRecommendation;
 
+	UPROPERTY(EditAnywhere, Category = Script)
+	uint32 bExperimental : 1;
+
 	/* If this script is exposed to the library. */
 	UPROPERTY(AssetRegistrySearchable, EditAnywhere, Category = Script)
 	uint32 bExposeToLibrary : 1;
diff --git a/Engine/Plugins/FX/Niagara/Source/Niagara/Public/NiagaraCommon.h b/Engine/Plugins/FX/Niagara/Source/Niagara/Public/NiagaraCommon.h
--- a/Engine/Plugins/FX/Niagara/Source/Niagara/Public/NiagaraCommon.h
+++ b/Engine/Plugins/FX/Niagara/Source/Niagara/Public/NiagaraCommon.h
@@ -244,6 +244,9 @@ struct NIAGARA_API FNiagaraFunctionSignature
 	/** True if this is the signature for a "member" function of a data interface. If this is true, the first input is the owner. */
 	UPROPERTY()
 	bool bMemberFunction;
+	/** Is this function experimental? */
+	UPROPERTY()
+	bool bExperimental;
 	/** Function specifiers verified at bind time. */
 	UPROPERTY()
 	TMap<FName, FName> FunctionSpecifiers;
diff --git a/Engine/Plugins/FX/Niagara/Source/NiagaraEditor/Private/NiagaraNodeFunctionCall.cpp b/Engine/Plugins/FX/Niagara/Source/NiagaraEditor/Private/NiagaraNodeFunctionCall.cpp
--- a/Engine/Plugins/FX/Niagara/Source/NiagaraEditor/Private/NiagaraNodeFunctionCall.cpp
+++ b/Engine/Plugins/FX/Niagara/Source/NiagaraEditor/Private/NiagaraNodeFunctionCall.cpp
@@ -607,44 +607,64 @@ UObject*  UNiagaraNodeFunctionCall::GetReferencedAsset() const
 	}
 }
 
-bool UNiagaraNodeFunctionCall::RefreshFromExternalChanges()
+void UNiagaraNodeFunctionCall::UpdateNodeErrorMessage()
 {
-	bool bReload = false;
 	if (FunctionScript)
 	{
-		UNiagaraScriptSource* Source = CastChecked<UNiagaraScriptSource>(FunctionScript->GetSource());
-		if (Source != nullptr)
-		{
-			bReload = CachedChangeId != Source->NodeGraph->GetChangeID();
-			if (bReload)
-			{
-				bReload = true;
-			}
-		}
-
 		if (FunctionScript->bDeprecated)
 		{
-			bHasCompilerMessage = true;
+			UEdGraphNode::bHasCompilerMessage = true;
 			ErrorType = EMessageSeverity::Warning;
 			FFormatNamedArguments Args;
 			Args.Add(TEXT("NodeName"), GetNodeTitle(ENodeTitleType::ListView));
 			Args.Add(TEXT("Recommendation"), FunctionScript->DeprecationRecommendation != nullptr ? FText::FromString(FunctionScript->DeprecationRecommendation->GetPathName()) : LOCTEXT("UnspecifiedName","Unspecified"));
-			ErrorMsg = FText::Format(LOCTEXT("DeprecationWarning", "{NodeName} is deprecated. Suggested replacement: {Recommendation}"), Args).ToString();
+			UEdGraphNode::ErrorMsg = FText::Format(LOCTEXT("DeprecationWarning", "{NodeName} is deprecated. Suggested replacement: {Recommendation}"), Args).ToString();
+		}
+		else if (FunctionScript->bExperimental)
+		{
+			UEdGraphNode::bHasCompilerMessage = true;
+			UEdGraphNode::ErrorType = EMessageSeverity::Info;
+			UEdGraphNode::NodeUpgradeMessage = LOCTEXT("FunctionExperimental", "This function is marked as experimental, use with care!");
 		}
 		else
 		{
-			bHasCompilerMessage = false;
-			ErrorMsg = FString();
+			UEdGraphNode::bHasCompilerMessage = false;
+			UEdGraphNode::ErrorMsg = FString();
 		}
-
 	}
-	else
+	else if (Signature.IsValid())
 	{
-		if (Signature.IsValid())
+		if (Signature.bExperimental)
 		{
-			bReload = true;
+			UEdGraphNode::bHasCompilerMessage = true;
+			UEdGraphNode::ErrorType = EMessageSeverity::Info;
+			UEdGraphNode::NodeUpgradeMessage = LOCTEXT("FunctionExperimental", "This function is marked as experimental, use with care!");
 		}
 	}
+}
+
+void UNiagaraNodeFunctionCall::PostPlacedNewNode()
+{
+	UpdateNodeErrorMessage();
+}
+
+bool UNiagaraNodeFunctionCall::RefreshFromExternalChanges()
+{
+	bool bReload = false;
+	if (FunctionScript)
+	{
+		UNiagaraScriptSource* Source = CastChecked<UNiagaraScriptSource>(FunctionScript->GetSource());
+		if (Source != nullptr)
+		{
+			bReload = CachedChangeId != Source->NodeGraph->GetChangeID();
+		}
+	}
+	else if (Signature.IsValid())
+	{
+		bReload = true;		
+	}
+
+	UpdateNodeErrorMessage();
 
 	// Go over the static switch parameters to set their propagation status on the pins
 	UNiagaraGraph* CalledGraph = GetCalledGraph();
diff --git a/Engine/Plugins/FX/Niagara/Source/NiagaraEditor/Private/NiagaraNodeInput.cpp b/Engine/Plugins/FX/Niagara/Source/NiagaraEditor/Private/NiagaraNodeInput.cpp
--- a/Engine/Plugins/FX/Niagara/Source/NiagaraEditor/Private/NiagaraNodeInput.cpp
+++ b/Engine/Plugins/FX/Niagara/Source/NiagaraEditor/Private/NiagaraNodeInput.cpp
@@ -70,6 +70,8 @@ void UNiagaraNodeInput::PostLoad()
 	{
 		DataInterface->OnChanged().AddUObject(this, &UNiagaraNodeInput::DataInterfaceChanged);
 	}
+
+	ValidateDataInterface();
 }
 
 void UNiagaraNodeInput::BuildParameterMapHistory(FNiagaraParameterMapHistoryBuilder& OutHistory, bool bRecursive /*= true*/, bool bFilterForCompilation /*= true*/) const
@@ -453,6 +455,8 @@ void UNiagaraNodeInput::SetDataInterface(UNiagaraDataInterface* InDataInterface)
 	{
 		DataInterface->OnChanged().AddUObject(this, &UNiagaraNodeInput::DataInterfaceChanged);
 	}
+
+	ValidateDataInterface();
 	DataInterfaceChanged();
 }
 
@@ -467,4 +471,15 @@ void UNiagaraNodeInput::DataInterfaceChanged()
 	}
 }
 
+void UNiagaraNodeInput::ValidateDataInterface()
+{
+	if (DataInterface != nullptr &&
+		DataInterface->GetClass()->GetMetaData("DevelopmentStatus") == TEXT("Experimental"))
+	{
+		UEdGraphNode::bHasCompilerMessage = true;
+		UEdGraphNode::ErrorType = EMessageSeverity::Info;
+		UEdGraphNode::NodeUpgradeMessage = LOCTEXT("DataInterfaceExperimental", "This data interface is marked as experimental, use with care!");
+	}
+}
+
 #undef LOCTEXT_NAMESPACE
diff --git a/Engine/Plugins/FX/Niagara/Source/NiagaraEditor/Private/ViewModels/Stack/NiagaraStackFunctionInput.cpp b/Engine/Plugins/FX/Niagara/Source/NiagaraEditor/Private/ViewModels/Stack/NiagaraStackFunctionInput.cpp
--- a/Engine/Plugins/FX/Niagara/Source/NiagaraEditor/Private/ViewModels/Stack/NiagaraStackFunctionInput.cpp
+++ b/Engine/Plugins/FX/Niagara/Source/NiagaraEditor/Private/ViewModels/Stack/NiagaraStackFunctionInput.cpp
@@ -332,36 +332,49 @@ void UNiagaraStackFunctionInput::RefreshChildrenInternal(const TArray<UNiagaraSt
 				DynamicInputEntry->Initialize(CreateDefaultChildRequiredData(), *OwningModuleNode, *InputValues.DynamicNode.Get(), GetOwnerStackItemEditorDataKey());
 				DynamicInputEntry->SetShouldShowInStack(false);
 			}
-			if (InputValues.DynamicNode->FunctionScript != nullptr && InputValues.DynamicNode->FunctionScript->bDeprecated)
+
+			if (InputValues.DynamicNode->FunctionScript != nullptr)
 			{
-				FText LongMessage = InputValues.DynamicNode->FunctionScript->DeprecationRecommendation != nullptr ?
-					FText::Format(LOCTEXT("ModuleScriptDeprecationLong", "The script asset for the assigned module {0} has been deprecated. Suggested replacement: {1}"), FText::FromString(InputValues.DynamicNode->GetName()), FText::FromString(InputValues.DynamicNode->FunctionScript->DeprecationRecommendation->GetPathName())) :
-					FText::Format(LOCTEXT("ModuleScriptDeprecationUnknownLong", "The script asset for the assigned module {0} has been deprecated."), FText::FromString(InputValues.DynamicNode->GetName()));
-
-				int32 AddIdx = NewIssues.Add(FStackIssue(
-					EStackIssueSeverity::Warning,
-					LOCTEXT("ModuleScriptDeprecationShort", "Deprecated module"),
-					LongMessage,
-					GetStackEditorDataKey(),
-					false,
-					{
-						FStackIssueFix(
-							LOCTEXT("SelectNewDynamicInputScriptFix", "Select a new dynamic input script"),
-							FStackIssueFixDelegate::CreateLambda([this]() { this->bIsDynamicInputScriptReassignmentPending = true; })),
-						FStackIssueFix(
-							LOCTEXT("ResetFix", "Reset this input to it's default value"),
-							FStackIssueFixDelegate::CreateLambda([this]() { this->Reset(); }))
-					}));
-
-				if (InputValues.DynamicNode->FunctionScript->DeprecationRecommendation != nullptr)
+				if (InputValues.DynamicNode->FunctionScript->bDeprecated)
 				{
-					NewIssues[AddIdx].InsertFix(0,
-						FStackIssueFix(
+					FText LongMessage = InputValues.DynamicNode->FunctionScript->DeprecationRecommendation != nullptr ?
+						FText::Format(LOCTEXT("ModuleScriptDeprecationLong", "The script asset for the assigned module {0} has been deprecated. Suggested replacement: {1}"), FText::FromString(InputValues.DynamicNode->GetName()), FText::FromString(InputValues.DynamicNode->FunctionScript->DeprecationRecommendation->GetPathName())) :
+						FText::Format(LOCTEXT("ModuleScriptDeprecationUnknownLong", "The script asset for the assigned module {0} has been deprecated."), FText::FromString(InputValues.DynamicNode->GetName()));
+
+					int32 AddIdx = NewIssues.Add(FStackIssue(
+						EStackIssueSeverity::Warning,
+						LOCTEXT("ModuleScriptDeprecationShort", "Deprecated module"),
+						LongMessage,
+						GetStackEditorDataKey(),
+						false,
+						{
+							FStackIssueFix(
+								LOCTEXT("SelectNewDynamicInputScriptFix", "Select a new dynamic input script"),
+								FStackIssueFixDelegate::CreateLambda([this]() { this->bIsDynamicInputScriptReassignmentPending = true; })),
+							FStackIssueFix(
+								LOCTEXT("ResetFix", "Reset this input to it's default value"),
+								FStackIssueFixDelegate::CreateLambda([this]() { this->Reset(); }))
+						}));
+
+					if (InputValues.DynamicNode->FunctionScript->DeprecationRecommendation != nullptr)
+					{
+						NewIssues[AddIdx].InsertFix(0,
+							FStackIssueFix(
 							LOCTEXT("SelectNewModuleScriptFixUseRecommended", "Use recommended replacement"),
 							FStackIssueFixDelegate::CreateLambda([this]() { ReassignDynamicInputScript(InputValues.DynamicNode->FunctionScript->DeprecationRecommendation); })));
+					}
 				}
-			}
 
+				if (InputValues.DynamicNode->FunctionScript->bExperimental)
+				{
+					NewIssues.Add(FStackIssue(
+						EStackIssueSeverity::Info,
+						LOCTEXT("ModuleScriptExperimentalShort", "Experimental module"),
+						FText::Format(LOCTEXT("ModuleScriptExperimental", "The script asset for the assigned module {0} is experimental, use with care!"), FText::FromString(InputValues.DynamicNode->GetName())),
+						GetStackEditorDataKey(),
+						true));
+				}
+			}
 
 			NewChildren.Add(DynamicInputEntry);
 		}
diff --git a/Engine/Plugins/FX/Niagara/Source/NiagaraEditor/Private/ViewModels/Stack/NiagaraStackModuleItem.cpp b/Engine/Plugins/FX/Niagara/Source/NiagaraEditor/Private/ViewModels/Stack/NiagaraStackModuleItem.cpp
--- a/Engine/Plugins/FX/Niagara/Source/NiagaraEditor/Private/ViewModels/Stack/NiagaraStackModuleItem.cpp
+++ b/Engine/Plugins/FX/Niagara/Source/NiagaraEditor/Private/ViewModels/Stack/NiagaraStackModuleItem.cpp
@@ -290,52 +290,67 @@ void UNiagaraStackModuleItem::RefreshIssues(TArray<FStackIssue>& NewIssues)
 		NewIssues.Empty();
 		return;
 	}
+
 	if (FunctionCallNode != nullptr)
 	{
-		if (FunctionCallNode->FunctionScript != nullptr && FunctionCallNode->FunctionScript->bDeprecated)
+		if (FunctionCallNode->FunctionScript != nullptr)
 		{
-			FText ModuleScriptDeprecationShort = LOCTEXT("ModuleScriptDeprecationShort", "Deprecated module");
-			if (CanMoveAndDelete())
+			if (FunctionCallNode->FunctionScript->bDeprecated)
 			{
-				FText LongMessage = FunctionCallNode->FunctionScript->DeprecationRecommendation != nullptr ?
-					FText::Format(LOCTEXT("ModuleScriptDeprecationLong", "The script asset for the assigned module {0} has been deprecated. Suggested replacement: {1}"), FText::FromString(FunctionCallNode->GetFunctionName()), FText::FromString(FunctionCallNode->FunctionScript->DeprecationRecommendation->GetPathName())) :
-					FText::Format(LOCTEXT("ModuleScriptDeprecationUnknownLong", "The script asset for the assigned module {0} has been deprecated."), FText::FromString(FunctionCallNode->GetFunctionName()));
-
-				int32 AddIdx = NewIssues.Add(FStackIssue(
-					EStackIssueSeverity::Warning,
-					ModuleScriptDeprecationShort,
-					LongMessage,
-					GetStackEditorDataKey(),
-					false,
-					{
-						FStackIssueFix(
-							LOCTEXT("SelectNewModuleScriptFix", "Select a new module script"),
-							FStackIssueFixDelegate::CreateLambda([this]() { this->bIsModuleScriptReassignmentPending = true; })),
-						FStackIssueFix(
-							LOCTEXT("DeleteFix", "Delete this module"),
-							FStackIssueFixDelegate::CreateLambda([this]() { this->Delete(); }))
-					}));
-
-				if (FunctionCallNode->FunctionScript->DeprecationRecommendation != nullptr)
+				FText ModuleScriptDeprecationShort = LOCTEXT("ModuleScriptDeprecationShort", "Deprecated module");
+				if (CanMoveAndDelete())
 				{
-					NewIssues[AddIdx].InsertFix(0,
-						FStackIssueFix(
+					FText LongMessage = FunctionCallNode->FunctionScript->DeprecationRecommendation != nullptr ?
+						FText::Format(LOCTEXT("ModuleScriptDeprecationLong", "The script asset for the assigned module {0} has been deprecated. Suggested replacement: {1}"), FText::FromString(FunctionCallNode->GetFunctionName()), FText::FromString(FunctionCallNode->FunctionScript->DeprecationRecommendation->GetPathName())) :
+						FText::Format(LOCTEXT("ModuleScriptDeprecationUnknownLong", "The script asset for the assigned module {0} has been deprecated."), FText::FromString(FunctionCallNode->GetFunctionName()));
+
+					int32 AddIdx = NewIssues.Add(FStackIssue(
+						EStackIssueSeverity::Warning,
+						ModuleScriptDeprecationShort,
+						LongMessage,
+						GetStackEditorDataKey(),
+						false,
+						{
+							FStackIssueFix(
+								LOCTEXT("SelectNewModuleScriptFix", "Select a new module script"),
+								FStackIssueFixDelegate::CreateLambda([this]() { this->bIsModuleScriptReassignmentPending = true; })),
+							FStackIssueFix(
+								LOCTEXT("DeleteFix", "Delete this module"),
+								FStackIssueFixDelegate::CreateLambda([this]() { this->Delete(); }))
+						}));
+
+					if (FunctionCallNode->FunctionScript->DeprecationRecommendation != nullptr)
+					{
+						NewIssues[AddIdx].InsertFix(0,
+							FStackIssueFix(
 							LOCTEXT("SelectNewModuleScriptFixUseRecommended", "Use recommended replacement"),
 							FStackIssueFixDelegate::CreateLambda([this]() { ReassignModuleScript(FunctionCallNode->FunctionScript->DeprecationRecommendation); })));
+					}
+				}
+				else
+				{
+					NewIssues.Add(FStackIssue(
+						EStackIssueSeverity::Warning,
+						ModuleScriptDeprecationShort,
+						FText::Format(LOCTEXT("ModuleScriptDeprecationFixParentLong", "The script asset for the assigned module {0} has been deprecated.  This module is inherited and this issue must be fixed in the parent emitter."),
+						FText::FromString(FunctionCallNode->GetFunctionName())),
+						GetStackEditorDataKey(),
+						false));
 				}
 			}
-			else
+
+			if (FunctionCallNode->FunctionScript->bExperimental)
 			{
 				NewIssues.Add(FStackIssue(
-					EStackIssueSeverity::Warning,
-					ModuleScriptDeprecationShort,
-					FText::Format(LOCTEXT("ModuleScriptDeprecationFixParentLong", "The script asset for the assigned module {0} has been deprecated.  This module is inherited and this issue must be fixed in the parent emitter."),
-						FText::FromString(FunctionCallNode->GetFunctionName())),
+					EStackIssueSeverity::Info,
+					LOCTEXT("ModuleScriptExperimentalShort", "Experimental module"),
+					FText::Format(LOCTEXT("ModuleScriptExperimental", "The script asset for the assigned module {0} is experimental, use with care!"), FText::FromString(FunctionCallNode->GetName())),
 					GetStackEditorDataKey(),
-					false));
+					true));
 			}
 		}
-		if(FunctionCallNode->FunctionScript == nullptr && FunctionCallNode->GetClass() == UNiagaraNodeFunctionCall::StaticClass())
+
+		if (FunctionCallNode->FunctionScript == nullptr && FunctionCallNode->GetClass() == UNiagaraNodeFunctionCall::StaticClass())
 		{
 			FText ModuleScriptMissingShort = LOCTEXT("ModuleScriptMissingShort", "Missing module script");
 			if (CanMoveAndDelete())
diff --git a/Engine/Plugins/FX/Niagara/Source/NiagaraEditor/Private/Widgets/SNiagaraGraphActionWidget.cpp b/Engine/Plugins/FX/Niagara/Source/NiagaraEditor/Private/Widgets/SNiagaraGraphActionWidget.cpp
new file mode 100644
--- /dev/null
+++ b/Engine/Plugins/FX/Niagara/Source/NiagaraEditor/Private/Widgets/SNiagaraGraphActionWidget.cpp
@@ -0,0 +1,56 @@
+// Copyright 1998-2019 Epic Games, Inc. All Rights Reserved.
+
+#include "SNiagaraGraphActionWidget.h"
+
+#include "NiagaraActions.h"
+#include "Styling/CoreStyle.h"
+#include "Widgets/SBoxPanel.h"
+#include "Widgets/Images/SImage.h"
+#include "Widgets/Text/STextBlock.h"
+
+#define LOCTEXT_NAMESPACE "NiagaraGraphActionWidget"
+
+void SNiagaraGraphActionWidget::Construct(const FArguments& InArgs, const FCreateWidgetForActionData* InCreateData)
+{
+	ActionPtr = InCreateData->Action;
+	MouseButtonDownDelegate = InCreateData->MouseButtonDownDelegate;
+
+	TSharedPtr<FNiagaraMenuAction> NiagaraAction = StaticCastSharedPtr<FNiagaraMenuAction>(InCreateData->Action);
+
+	this->ChildSlot
+	[
+		SNew(SHorizontalBox)
+		.ToolTipText(InCreateData->Action->GetTooltipDescription())
+		+ SHorizontalBox::Slot()
+		.FillWidth(1)
+		.VAlign(VAlign_Center)
+		[
+			SNew(STextBlock)
+			.Font(FCoreStyle::GetDefaultFontStyle("Regular", 9))
+			.Text(InCreateData->Action->GetMenuDescription())
+			.HighlightText(InArgs._HighlightText)
+		]
+		+ SHorizontalBox::Slot()
+		.AutoWidth()
+		.VAlign(VAlign_Center)
+		.HAlign(HAlign_Right)
+		[
+			SNew(SImage)
+			.Image(NiagaraAction->IsExperimental ? FEditorStyle::GetBrush("Icons.Info") : nullptr)
+			.Visibility(NiagaraAction->IsExperimental ? EVisibility::Visible : EVisibility::Collapsed)
+			.ToolTipText(NiagaraAction->IsExperimental ? LOCTEXT("ScriptExperimentalToolTip", "This script is experimental, use with care!") : FText::GetEmpty())
+		]
+	];
+}
+
+FReply SNiagaraGraphActionWidget::OnMouseButtonDown(const FGeometry& MyGeometry, const FPointerEvent& MouseEvent)
+{
+	if (MouseButtonDownDelegate.Execute(ActionPtr))
+	{
+		return FReply::Handled();
+	}
+
+	return FReply::Unhandled();
+}
+
+#undef LOCTEXT_NAMESPACE
diff --git a/Engine/Plugins/FX/Niagara/Source/NiagaraEditor/Private/Widgets/SNiagaraParameterMapView.cpp b/Engine/Plugins/FX/Niagara/Source/NiagaraEditor/Private/Widgets/SNiagaraParameterMapView.cpp
--- a/Engine/Plugins/FX/Niagara/Source/NiagaraEditor/Private/Widgets/SNiagaraParameterMapView.cpp
+++ b/Engine/Plugins/FX/Niagara/Source/NiagaraEditor/Private/Widgets/SNiagaraParameterMapView.cpp
@@ -39,6 +39,7 @@
 #include "Widgets/Layout/SBox.h"
 #include "NiagaraSystemEditorData.h"
 #include "ViewModels/Stack/NiagaraStackSystemSettingsGroup.h"
+#include "SNiagaraGraphActionWidget.h"
 
 #define LOCTEXT_NAMESPACE "NiagaraParameterMapView"
 
@@ -838,6 +839,10 @@ void SNiagaraAddParameterMenu::Construct(const FArguments& InArgs, TArray<TWeakO
 				.AutoExpandActionMenu(AutoExpandMenu.Get())
 				.ShowFilterTextBox(true)
 				.OnCreateCustomRowExpander_Static(&SNiagaraParameterMapView::CreateCustomActionExpander)
+				.OnCreateWidgetForAction_Lambda([](const FCreateWidgetForActionData* InData)
+				{
+					return SNew(SNiagaraGraphActionWidget, InData);
+				})
 			]
 		]
 	];
@@ -1053,12 +1058,27 @@ void SNiagaraAddParameterMenu::AddParameterGroup(FGraphActionListBuilderBase& Ou
 		Variables.Sort([](const FNiagaraVariable& A, const FNiagaraVariable& B) { return A.GetName().LexicalLess(B.GetName()); });
 	}
 
-	for (FNiagaraVariable& Variable : Variables)
+	for (const FNiagaraVariable& Variable : Variables)
 	{
 		const FText DisplayName = FText::FromName(Variable.GetName());
-		const FText Tooltip = FText::GetEmpty();
+		FText Tooltip = FText::GetEmpty();
+
+		if (const UStruct* VariableStruct = Variable.GetType().GetStruct())
+		{
+			Tooltip = VariableStruct->GetToolTipText(true);
+		}
+
 		TSharedPtr<FNiagaraMenuAction> Action(new FNiagaraMenuAction(Category, DisplayName, Tooltip, 0, FText(),
 			FNiagaraMenuAction::FOnExecuteStackAction::CreateSP(this, &SNiagaraAddParameterMenu::AddParameterSelected, Variable, bCustomName, InSection)));
+
+		if (Variable.IsDataInterface())
+		{
+			if (const UClass* DataInterfaceClass = Variable.GetType().GetClass())
+			{
+				Action->IsExperimental = DataInterfaceClass->GetMetaData("DevelopmentStatus") == TEXT("Experimental");
+			}
+		}
+
 		OutActions.AddAction(Action, RootCategory);
 	}
 }
diff --git a/Engine/Plugins/FX/Niagara/Source/NiagaraEditor/Public/NiagaraActions.h b/Engine/Plugins/FX/Niagara/Source/NiagaraEditor/Public/NiagaraActions.h
--- a/Engine/Plugins/FX/Niagara/Source/NiagaraEditor/Public/NiagaraActions.h
+++ b/Engine/Plugins/FX/Niagara/Source/NiagaraEditor/Public/NiagaraActions.h
@@ -18,10 +18,7 @@ struct NIAGARAEDITOR_API FNiagaraMenuAction : public FEdGraphSchemaAction
 	DECLARE_DELEGATE(FOnExecuteStackAction);
 	DECLARE_DELEGATE_RetVal(bool, FCanExecuteStackAction);
 
-	FNiagaraMenuAction()
-	{
-	}
-
+	FNiagaraMenuAction() {}
 	FNiagaraMenuAction(FText InNodeCategory, FText InMenuDesc, FText InToolTip, const int32 InGrouping, FText InKeywords, FOnExecuteStackAction InAction, int32 InSectionID = 0);
 	FNiagaraMenuAction(FText InNodeCategory, FText InMenuDesc, FText InToolTip, const int32 InGrouping, FText InKeywords, FOnExecuteStackAction InAction, FCanExecuteStackAction InCanPerformAction, int32 InSectionID = 0);
 
@@ -39,6 +36,8 @@ struct NIAGARAEDITOR_API FNiagaraMenuAction : public FEdGraphSchemaAction
 		return CanPerformAction.IsBound() ? CanPerformAction.Execute() : true;
 	}
 
+	bool IsExperimental = false;
+
 private:
 	FOnExecuteStackAction Action;
 	FCanExecuteStackAction CanPerformAction;
diff --git a/Engine/Plugins/FX/Niagara/Source/NiagaraEditor/Public/NiagaraNodeFunctionCall.h b/Engine/Plugins/FX/Niagara/Source/NiagaraEditor/Public/NiagaraNodeFunctionCall.h
--- a/Engine/Plugins/FX/Niagara/Source/NiagaraEditor/Public/NiagaraNodeFunctionCall.h
+++ b/Engine/Plugins/FX/Niagara/Source/NiagaraEditor/Public/NiagaraNodeFunctionCall.h
@@ -81,11 +81,12 @@ public:
 
 	//Begin UObject interface
 	virtual void PostLoad()override;
-	virtual void PostEditChangeProperty(struct FPropertyChangedEvent& PropertyChangedEvent)override;
+	virtual void PostEditChangeProperty(struct FPropertyChangedEvent& PropertyChangedEvent) override;
 	//End UObject interface
 
 	//~ Begin UNiagaraNode Interface
-	virtual void Compile(class FHlslNiagaraTranslator* Translator, TArray<int32>& Outputs)override;
+	virtual void PostPlacedNewNode() override;
+	virtual void Compile(class FHlslNiagaraTranslator* Translator, TArray<int32>& Outputs) override;
 	virtual UObject* GetReferencedAsset() const override;
 	virtual bool RefreshFromExternalChanges() override;
 	virtual ENiagaraNumericOutputTypeSelectionMode GetNumericOutputTypeSelectionMode() const override;
@@ -96,7 +97,7 @@ public:
 	//End UNiagaraNode interface
 
 	//~ Begin EdGraphNode Interface
-	virtual void AutowireNewNode(UEdGraphPin* FromPin)override;
+	virtual void AutowireNewNode(UEdGraphPin* FromPin) override;
 	virtual void AllocateDefaultPins() override;
 	virtual FText GetNodeTitle(ENodeTitleType::Type TitleType) const override;
 	virtual FText GetTooltipText() const override;
@@ -151,6 +152,8 @@ protected:
 
 	bool IsValidPropagatedVariable(const FNiagaraVariable& Variable) const;
 
+	void UpdateNodeErrorMessage();
+
 	/** Adjusted every time that we compile this script. Lets us know that we might differ from any cached versions.*/
 	UPROPERTY()
 	FGuid CachedChangeId;
diff --git a/Engine/Plugins/FX/Niagara/Source/NiagaraEditor/Public/NiagaraNodeInput.h b/Engine/Plugins/FX/Niagara/Source/NiagaraEditor/Public/NiagaraNodeInput.h
--- a/Engine/Plugins/FX/Niagara/Source/NiagaraEditor/Public/NiagaraNodeInput.h
+++ b/Engine/Plugins/FX/Niagara/Source/NiagaraEditor/Public/NiagaraNodeInput.h
@@ -112,6 +112,7 @@ public:
 
 private:
 	void DataInterfaceChanged();
+	void ValidateDataInterface();
 
 private:
 	UPROPERTY()
diff --git a/Engine/Plugins/FX/Niagara/Source/NiagaraEditor/Public/SNiagaraGraphActionWidget.h b/Engine/Plugins/FX/Niagara/Source/NiagaraEditor/Public/SNiagaraGraphActionWidget.h
new file mode 100644
--- /dev/null
+++ b/Engine/Plugins/FX/Niagara/Source/NiagaraEditor/Public/SNiagaraGraphActionWidget.h
@@ -0,0 +1,31 @@
+// Copyright 1998-2019 Epic Games, Inc. All Rights Reserved.
+
+#pragma once
+
+#include "CoreMinimal.h"
+#include "SlateFwd.h"
+#include "SGraphActionMenu.h"
+#include "EdGraph/EdGraphSchema.h"
+#include "Widgets/DeclarativeSyntaxSupport.h"
+#include "Widgets/SCompoundWidget.h"
+
+struct FCreateWidgetForActionData;
+
+/** Custom widget for GraphActionMenu */
+class NIAGARAEDITOR_API SNiagaraGraphActionWidget : public SCompoundWidget
+{
+public:
+
+	SLATE_BEGIN_ARGS( SNiagaraGraphActionWidget ) {}
+		SLATE_ATTRIBUTE(FText, HighlightText)
+	SLATE_END_ARGS()
+
+	void Construct(const FArguments& InArgs, const FCreateWidgetForActionData* InCreateData);
+	virtual FReply OnMouseButtonDown( const FGeometry& MyGeometry, const FPointerEvent& MouseEvent ) override;
+
+private:
+	/** The item that we want to display with this widget */
+	TWeakPtr<struct FEdGraphSchemaAction> ActionPtr;
+	/** Delegate executed when mouse button goes down */
+	FCreateWidgetMouseButtonDown MouseButtonDownDelegate;
+};
diff --git a/Engine/Plugins/FX/Niagara/Source/NiagaraEditorWidgets/Private/SNiagaraOverviewStackNode.cpp b/Engine/Plugins/FX/Niagara/Source/NiagaraEditorWidgets/Private/SNiagaraOverviewStackNode.cpp
--- a/Engine/Plugins/FX/Niagara/Source/NiagaraEditorWidgets/Private/SNiagaraOverviewStackNode.cpp
+++ b/Engine/Plugins/FX/Niagara/Source/NiagaraEditorWidgets/Private/SNiagaraOverviewStackNode.cpp
@@ -81,6 +81,7 @@ TSharedRef<SWidget> SNiagaraOverviewStackNode::CreateTitleWidget(TSharedPtr<SNod
 		// Name
 		+ SHorizontalBox::Slot()
 		.Padding(4, 0, 0, 0)
+		.FillWidth(1.0f)
 		[
 			DefaultTitle
 		]
@@ -88,6 +89,7 @@ TSharedRef<SWidget> SNiagaraOverviewStackNode::CreateTitleWidget(TSharedPtr<SNod
 		+ SHorizontalBox::Slot()
 		.AutoWidth()
 		.VAlign(VAlign_Center)
+		.HAlign(HAlign_Right)
 		.Padding(5, 0, 0, 0)
 		[
 			SNew(SNiagaraStackIssueIcon, StackViewModel, StackViewModel->GetRootEntry())
diff --git a/Engine/Plugins/FX/Niagara/Source/NiagaraEditorWidgets/Private/Stack/SNiagaraStackFunctionInputValue.cpp b/Engine/Plugins/FX/Niagara/Source/NiagaraEditorWidgets/Private/Stack/SNiagaraStackFunctionInputValue.cpp
--- a/Engine/Plugins/FX/Niagara/Source/NiagaraEditorWidgets/Private/Stack/SNiagaraStackFunctionInputValue.cpp
+++ b/Engine/Plugins/FX/Niagara/Source/NiagaraEditorWidgets/Private/Stack/SNiagaraStackFunctionInputValue.cpp
@@ -28,6 +28,7 @@
 #include "NiagaraParameterCollection.h"
 #include "ViewModels/NiagaraSystemViewModel.h"
 #include "NiagaraSystem.h"
+#include "SNiagaraGraphActionWidget.h"
 #include "Subsystems/AssetEditorSubsystem.h"
 
 #define LOCTEXT_NAMESPACE "NiagaraStackFunctionInputValue"
@@ -530,6 +531,10 @@ TSharedRef<SWidget> SNiagaraStackFunctionInputValue::OnGetAvailableHandleMenu()
 				.AutoExpandActionMenu(false)
 				.ShowFilterTextBox(true)
 				.OnCreateCustomRowExpander_Static(&CreateCustomNiagaraFunctionInputActionExpander)
+				.OnCreateWidgetForAction_Lambda([](const FCreateWidgetForActionData* InData)
+				{
+					return SNew(SNiagaraGraphActionWidget, InData);
+				})
 			]
 		]
 	];
@@ -585,6 +590,8 @@ void SNiagaraStackFunctionInputValue::CollectAllActions(FGraphActionListBuilderB
 			const FText Tooltip = FNiagaraEditorUtilities::FormatScriptAssetDescription(DynamicInputScript->Description, *(DynamicInputScript->GetPathName() + (DynamicInputScript->bExposeToLibrary ? "" : "\n*Not exposed to library")));
 			TSharedPtr<FNiagaraMenuAction> DynamicInputAction(new FNiagaraMenuAction(CategoryName, DynamicInputText, Tooltip, 0, DynamicInputScript->Keywords,
 				FNiagaraMenuAction::FOnExecuteStackAction::CreateSP(this, &SNiagaraStackFunctionInputValue::DynamicInputScriptSelected, DynamicInputScript)));
+
+			DynamicInputAction->IsExperimental = DynamicInputScript->bExperimental;
 			OutAllActions.AddAction(DynamicInputAction);
 		}
 	}
