Removed code comments for Niagara floats that were being picked up by UHT and displayed to the end user.

diff --git a/Engine/Plugins/FX/Niagara/Source/Niagara/Public/NiagaraTypes.h b/Engine/Plugins/FX/Niagara/Source/Niagara/Public/NiagaraTypes.h
--- a/Engine/Plugins/FX/Niagara/Source/Niagara/Public/NiagaraTypes.h
+++ b/Engine/Plugins/FX/Niagara/Source/Niagara/Public/NiagaraTypes.h
@@ -18,7 +18,7 @@ struct FNiagaraFloat
 {
 	GENERATED_USTRUCT_BODY()
 
-	UPROPERTY(EditAnywhere, Category=Parameters)//Parameters? These are used for attrs too.
+	UPROPERTY(EditAnywhere, Category=Parameters)
 	float Value;
 };
 
@@ -27,7 +27,7 @@ struct FNiagaraInt32
 {
 	GENERATED_USTRUCT_BODY()
 
-	UPROPERTY(EditAnywhere, Category = Parameters)//Parameters? These are used for attrs too.
+	UPROPERTY(EditAnywhere, Category = Parameters)
 	int32 Value;
 };
 
@@ -36,14 +36,12 @@ struct FNiagaraBool
 {
 	GENERATED_USTRUCT_BODY()
 
-		// The Niagara VM expects this bitmask for its compare and select operators for false.
+	// The Niagara VM expects this bitmask for its compare and select operators for false.
 	enum BoolValues { 
 		True = INDEX_NONE,
 		False = 0
 	}; 
 
-	
-
 	void SetValue(bool bValue) { Value = bValue ? True : False; }
 	bool GetValue() const { return Value != False; }
 
@@ -55,12 +53,12 @@ struct FNiagaraBool
 
 	bool IsValid() const { return Value == True || Value == False; }
 	
-	FNiagaraBool():Value(False) {}
+	FNiagaraBool() : Value(False) {}
 	FNiagaraBool(bool bInValue) : Value(bInValue ? True : False) {}
 	FORCEINLINE operator bool() { return GetValue(); }
 
 private:
-	UPROPERTY(EditAnywhere, Category = Parameters)//Parameters? These are used for attrs too. Must be either FNiagaraBool::True or FNiagaraBool::False.
+	UPROPERTY(EditAnywhere, Category = Parameters)// Must be either FNiagaraBool::True or FNiagaraBool::False.
 	int32 Value;
 };
 
