#Niagara Fix initialization issue

diff --git a/Engine/Plugins/FX/Niagara/Source/Niagara/Public/NiagaraCommon.h b/Engine/Plugins/FX/Niagara/Source/Niagara/Public/NiagaraCommon.h
--- a/Engine/Plugins/FX/Niagara/Source/Niagara/Public/NiagaraCommon.h
+++ b/Engine/Plugins/FX/Niagara/Source/Niagara/Public/NiagaraCommon.h
@@ -265,6 +265,7 @@ struct NIAGARA_API FNiagaraFunctionSignature
 	FNiagaraFunctionSignature() 
 		: bRequiresContext(false)
 		, bMemberFunction(false)
+		, bExperimental(false)
 	{
 	}
 
@@ -274,6 +275,7 @@ struct NIAGARA_API FNiagaraFunctionSignature
 		, Outputs(InOutputs)
 		, bRequiresContext(bInRequiresContext)
 		, bMemberFunction(bInMemberFunction)
+		, bExperimental(false)
 		, FunctionSpecifiers()
 	{
 
@@ -285,6 +287,7 @@ struct NIAGARA_API FNiagaraFunctionSignature
 		, Outputs(InOutputs)
 		, bRequiresContext(bInRequiresContext)
 		, bMemberFunction(bInMemberFunction)
+		, bExperimental(false)
 		, FunctionSpecifiers(InFunctionSpecifiers)
 	{
 
