Fixed circular occlusion query hlsl

diff --git a/Engine/Plugins/FX/Niagara/Source/Niagara/Private/NiagaraDataInterfaceCamera.cpp b/Engine/Plugins/FX/Niagara/Source/Niagara/Private/NiagaraDataInterfaceCamera.cpp
--- a/Engine/Plugins/FX/Niagara/Source/Niagara/Private/NiagaraDataInterfaceCamera.cpp
+++ b/Engine/Plugins/FX/Niagara/Source/Niagara/Private/NiagaraDataInterfaceCamera.cpp
@@ -261,12 +261,9 @@ bool UNiagaraDataInterfaceCamera::GetFunctionHLSL(const FName& DefinitionFunctio
 				float2 ScreenUV = ScreenPosition * View.ScreenPositionScaleBias.xy + View.ScreenPositionScaleBias.wz;
 
 				float Rays = In_SampleRays <= 1 ? 0 : In_SampleRays;
-				float Steps = In_SampleSteps < 1 ? 0 : In_SampleSteps;
+				float Steps = In_SampleStepsPerRay < 1 ? 0 : In_SampleStepsPerRay;
 				float TotalSamples = 0;
 				float OccludedSamples = 0;
-				
-				float SampleWidth = ScreenUV.x > 1 ? 0 : SampleWidthUV.x - ScreenUV.x;
-				float SampleHeight = ScreenUV.y > 1 ? 0 : SampleHeightUV.y - ScreenUV.y;
 
 				if (ScreenUV.x <= 1 && ScreenUV.x >= 0 && ScreenUV.y <= 1 && ScreenUV.y >= 0)
 				{
@@ -287,7 +284,7 @@ bool UNiagaraDataInterfaceCamera::GetFunctionHLSL(const FName& DefinitionFunctio
 						float4 RayClip = mul(float4(RayDirection * In_SampleWindowRadiusWorld, 0) + SamplePosition, View.TranslatedWorldToClip);
 						float2 RayUV = RayClip.xy / RayClip.w * View.ScreenPositionScaleBias.xy + View.ScreenPositionScaleBias.wz;
 
-						if ((ScreenUV.x > 1 && RayUV.x < 0) || (ScreenUV.y > 1 && RayUV.y < 0) || (ScreenUV.x < 0 && RayUV.x > 1) || (ScreenUV.< < 0 && RayUV.y > 1))
+						if ((ScreenUV.x > 1 && RayUV.x < 0) || (ScreenUV.y > 1 && RayUV.y < 0) || (ScreenUV.x < 0 && RayUV.x > 1) || (ScreenUV.y < 0 && RayUV.y > 1))
 						{
 							continue;
 						}
