The stack search is focused on stack refresh, and clicking on a stack item will also focus the search. Clicking on text entry (either numeric or text field) isn't impacted.
#rnx

diff --git a/Engine/Plugins/FX/Niagara/Source/NiagaraEditorWidgets/Private/SNiagaraStack.cpp b/Engine/Plugins/FX/Niagara/Source/NiagaraEditorWidgets/Private/SNiagaraStack.cpp
--- a/Engine/Plugins/FX/Niagara/Source/NiagaraEditorWidgets/Private/SNiagaraStack.cpp
+++ b/Engine/Plugins/FX/Niagara/Source/NiagaraEditorWidgets/Private/SNiagaraStack.cpp
@@ -269,6 +269,7 @@ void SNiagaraStack::Construct(const FArguments& InArgs, UNiagaraStackViewModel*
 				.OnTreeViewScrolled(this, &SNiagaraStack::StackTreeScrolled)
 				.SelectionMode(ESelectionMode::Single)
 				.OnItemToString_Debug_Static(&FNiagaraStackEditorWidgetsUtilities::StackEntryToStringForListDebug)
+				.OnMouseButtonClick(this, &SNiagaraStack::OnStackItemClicked)
 			]
 		]
 	];
@@ -1043,6 +1044,8 @@ void SNiagaraStack::StackStructureChanged()
 	SynchronizeTreeExpansion();
 	StackTree->RequestTreeRefresh();
 	HeaderList->RequestListRefresh();
+	// Keep the stack search focused by default
+	FSlateApplication::Get().SetKeyboardFocus(SearchBox, EFocusCause::OtherWidgetLostFocus);
 }
 
 EVisibility SNiagaraStack::GetIssueIconVisibility() const
@@ -1050,4 +1053,11 @@ EVisibility SNiagaraStack::GetIssueIconVisibility() const
 	return StackViewModel->HasIssues() ? EVisibility::Visible : EVisibility::Collapsed;
 }
 
+void SNiagaraStack::OnStackItemClicked(UNiagaraStackEntry* Item)
+{
+	// Focus the stack search if we are just clicking on an item
+	// This won't be hit by a bubbling-up click if a widget with text entry is clicked
+	FSlateApplication::Get().SetKeyboardFocus(SearchBox, EFocusCause::OtherWidgetLostFocus);
+}
+
 #undef LOCTEXT_NAMESPACE
diff --git a/Engine/Plugins/FX/Niagara/Source/NiagaraEditorWidgets/Private/SNiagaraStack.h b/Engine/Plugins/FX/Niagara/Source/NiagaraEditorWidgets/Private/SNiagaraStack.h
--- a/Engine/Plugins/FX/Niagara/Source/NiagaraEditorWidgets/Private/SNiagaraStack.h
+++ b/Engine/Plugins/FX/Niagara/Source/NiagaraEditorWidgets/Private/SNiagaraStack.h
@@ -106,6 +106,7 @@ private:
 
 	EVisibility GetIssueIconVisibility() const;
 
+	void OnStackItemClicked(UNiagaraStackEntry* Item);
 private:
 	UNiagaraStackViewModel* StackViewModel;
 
