Niagara - Highlight visual tweaks.

diff --git a/Engine/Plugins/FX/Niagara/Source/NiagaraEditorWidgets/Private/NiagaraEditorWidgetsStyle.cpp b/Engine/Plugins/FX/Niagara/Source/NiagaraEditorWidgets/Private/NiagaraEditorWidgetsStyle.cpp
--- a/Engine/Plugins/FX/Niagara/Source/NiagaraEditorWidgets/Private/NiagaraEditorWidgetsStyle.cpp
+++ b/Engine/Plugins/FX/Niagara/Source/NiagaraEditorWidgets/Private/NiagaraEditorWidgetsStyle.cpp
@@ -46,6 +46,7 @@ NIAGARAEDITOR_API FString RelativePathToPluginPath(const FString& RelativePath,
 #define IMAGE_CORE_BRUSH( RelativePath, ... ) FSlateImageBrush( FPaths::EngineContentDir() / "Editor/Slate" / RelativePath + TEXT(".png") , __VA_ARGS__ )
 #define BOX_CORE_BRUSH( RelativePath, ... ) FSlateBoxBrush( FPaths::EngineContentDir() / "Editor/Slate" / RelativePath + TEXT(".png"), __VA_ARGS__ )
 
+const FVector2D Icon6x6(6.0f, 6.0f);
 const FVector2D Icon8x8(8.0f, 8.0f);
 const FVector2D Icon8x16(8.0f, 16.0f);
 const FVector2D Icon12x12(12.0f, 12.0f);
@@ -200,8 +201,9 @@ TSharedRef< FSlateStyleSet > FNiagaraEditorWidgetsStyle::Create()
 		.SetShadowColorAndOpacity(FLinearColor(0, 0, 0, 0.9f));
 	Style->Set("NiagaraEditor.Stack.AddButtonText", AddButtonText);
 
-	Style->Set("NiagaraEditor.Stack.ModuleHighlight", new IMAGE_PLUGIN_BRUSH("Icons/ModuleHighlight", Icon8x8, FLinearColor::White));
-	Style->Set("NiagaraEditor.Stack.ModuleHighlightMore", new IMAGE_PLUGIN_BRUSH("Icons/ModuleHighlightMore", Icon8x8, FLinearColor::White));
+	Style->Set("NiagaraEditor.Stack.ModuleHighlight", new IMAGE_PLUGIN_BRUSH("Icons/ModuleHighlight", Icon6x6, FLinearColor::White));
+	Style->Set("NiagaraEditor.Stack.ModuleHighlightMore", new IMAGE_PLUGIN_BRUSH("Icons/ModuleHighlightMore", Icon6x6, FLinearColor::White));
+	Style->Set("NiagaraEditor.Stack.ModuleHighlightLarge", new IMAGE_PLUGIN_BRUSH("Icons/ModuleHighlightLarge", Icon8x8, FLinearColor::White));
 
 	Style->Set("NiagaraEditor.ShowInCurveEditorIcon", new IMAGE_PLUGIN_BRUSH("Icons/ShowInCurveEditor", Icon16x16, FLinearColor::White));
 
diff --git a/Engine/Plugins/FX/Niagara/Source/NiagaraEditorWidgets/Private/SNiagaraOverviewStack.cpp b/Engine/Plugins/FX/Niagara/Source/NiagaraEditorWidgets/Private/SNiagaraOverviewStack.cpp
--- a/Engine/Plugins/FX/Niagara/Source/NiagaraEditorWidgets/Private/SNiagaraOverviewStack.cpp
+++ b/Engine/Plugins/FX/Niagara/Source/NiagaraEditorWidgets/Private/SNiagaraOverviewStack.cpp
@@ -387,28 +387,43 @@ TSharedRef<ITableRow> SNiagaraOverviewStack::OnGenerateRowForEntry(UNiagaraStack
 		TSharedPtr<SWidget> IndentContent;
 		if (StackItem->SupportsHighlights())
 		{
-			TArray<FNiagaraScriptHighlight> ScriptHighlights = StackItem->GetHighlights();
-			Algo::Sort(ScriptHighlights, [](const FNiagaraScriptHighlight& A, const FNiagaraScriptHighlight& B)
+
+			TArray<FNiagaraScriptHighlight> ScriptHighlights;
+			for (const FNiagaraScriptHighlight& ScriptHighlight : StackItem->GetHighlights())
 			{
-				return A.DisplayName.CompareTo(B.DisplayName) < 0;
-			});
+				if (ScriptHighlight.IsValid())
+				{
+					ScriptHighlights.Add(ScriptHighlight);
+				}
+			}
 
-			FText HighlightToolTip;
+			TSharedRef<SHorizontalBox> ToolTipBox = SNew(SHorizontalBox);
 			if (ScriptHighlights.Num() > 0)
 			{
-				TArray<FText> DisplayNames;
 				for (const FNiagaraScriptHighlight& ScriptHighlight : ScriptHighlights)
 				{
-					if (ScriptHighlight.IsValid())
-					{
-						DisplayNames.Add(ScriptHighlight.DisplayName);
-					}
+					ToolTipBox->AddSlot()
+						.AutoWidth()
+						.VAlign(VAlign_Center)
+						.Padding(0, 0, 5, 0)
+						[
+							SNew(SImage)
+							.Image(FNiagaraEditorWidgetsStyle::Get().GetBrush("NiagaraEditor.Stack.ModuleHighlightLarge"))
+							.ColorAndOpacity(ScriptHighlight.Color)
+						];
+					ToolTipBox->AddSlot()
+						.AutoWidth()
+						.VAlign(VAlign_Center)
+						.Padding(0, 0, 10, 0)
+						[
+							SNew(STextBlock)
+							.Text(ScriptHighlight.DisplayName)
+						];
 				}
-				HighlightToolTip = FText::Join(LOCTEXT("HighlightDelimiter", ", "), DisplayNames);
 			}
 
 			TSharedRef<SBox> IndentBox = SNew(SBox)
-				.ToolTipText(HighlightToolTip)
+				.ToolTip(SNew(SToolTip) [ ToolTipBox ])
 				.IsEnabled_UObject(Item, &UNiagaraStackEntry::GetIsEnabledAndOwnerIsEnabled)
 				.Visibility(EVisibility::Visible)
 				.WidthOverride(IconSize.X)
@@ -416,25 +431,38 @@ TSharedRef<ITableRow> SNiagaraOverviewStack::OnGenerateRowForEntry(UNiagaraStack
 			
 			if (ScriptHighlights.Num() > 0)
 			{
+				IndentBox->SetHAlign(HAlign_Left);
+				IndentBox->SetVAlign(VAlign_Top);
+
 				TSharedPtr<SGridPanel> HighlightsGrid;
-				IndentBox->SetContent(SAssignNew(HighlightsGrid, SGridPanel)
-					.FillRow(0, .5f)
-					.FillRow(1, .5f)
-					.FillColumn(0, .5f)
-					.FillColumn(1, .5f));
+				IndentBox->SetContent(SAssignNew(HighlightsGrid, SGridPanel));
+
 				int32 HighlightsAdded = 0;
 				for (int32 HighlightIndex = 0; HighlightIndex < ScriptHighlights.Num() && HighlightsAdded < 4; HighlightIndex++)
 				{
 					if (ScriptHighlights[HighlightIndex].IsValid())
 					{
-						int32 Column = HighlightIndex / 2;
-						int32 Row = HighlightIndex % 2;
+						FName HighlightImageBrushName;
+						FLinearColor HighlightImageColor;
+						if (HighlightIndex < 3)
+						{
+							HighlightImageBrushName = "NiagaraEditor.Stack.ModuleHighlight";
+							HighlightImageColor = ScriptHighlights[HighlightIndex].Color;
+						}
+						else
+						{
+							HighlightImageBrushName = "NiagaraEditor.Stack.ModuleHighlightMore";
+							HighlightImageColor = FLinearColor::White;
+						}
+
+						int32 Column = HighlightIndex % 2;
+						int32 Row = HighlightIndex / 2;
 						HighlightsGrid->AddSlot(Column, Row)
 							.Padding(Column, Row, 1 - Column, 1 - Row)
 							[
 								SNew(SImage)
-								.Image(FNiagaraEditorWidgetsStyle::Get().GetBrush("NiagaraEditor.Stack.ModuleHighlight"))
-								.ColorAndOpacity(ScriptHighlights[HighlightIndex].Color)
+								.Image(FNiagaraEditorWidgetsStyle::Get().GetBrush(HighlightImageBrushName))
+								.ColorAndOpacity(HighlightImageColor)
 							];
 						HighlightsAdded++;
 					}
@@ -474,7 +502,7 @@ TSharedRef<ITableRow> SNiagaraOverviewStack::OnGenerateRowForEntry(UNiagaraStack
 			// Name
 			+ SHorizontalBox::Slot()
 			.VAlign(VAlign_Center)
-			.Padding(2, 3, 0, 3)
+			.Padding(3, 3, 0, 3)
 			[
 				SNew(STextBlock)
 				.TextStyle(FNiagaraEditorWidgetsStyle::Get(), "NiagaraEditor.SystemOverview.ItemText")
@@ -500,7 +528,7 @@ TSharedRef<ITableRow> SNiagaraOverviewStack::OnGenerateRowForEntry(UNiagaraStack
 		TSharedRef<SHorizontalBox> ContentBox = SNew(SHorizontalBox)
 			// Execution category icon
 			+ SHorizontalBox::Slot()
-			.Padding(0, 0, 4, 0)
+			.Padding(0, 0, 2, 0)
 			.VAlign(VAlign_Center)
 			.AutoWidth()
 			[
@@ -519,7 +547,7 @@ TSharedRef<ITableRow> SNiagaraOverviewStack::OnGenerateRowForEntry(UNiagaraStack
 			// Name
 			+ SHorizontalBox::Slot()
 			.VAlign(VAlign_Center)
-			.Padding(0, 2, 0, 2)
+			.Padding(3, 2, 0, 2)
 			[
 				SNew(STextBlock)
 				.TextStyle(FNiagaraEditorWidgetsStyle::Get(), "NiagaraEditor.SystemOverview.GroupHeaderText")
