switched shaderstageindex to int because of lots of casting warnings.  Eventually this will go back to a uint, but it likely won't matter one way or another

diff --git a/Engine/Plugins/FX/Niagara/Shaders/Private/NiagaraEmitterInstanceShader.usf b/Engine/Plugins/FX/Niagara/Shaders/Private/NiagaraEmitterInstanceShader.usf
--- a/Engine/Plugins/FX/Niagara/Shaders/Private/NiagaraEmitterInstanceShader.usf
+++ b/Engine/Plugins/FX/Niagara/Shaders/Private/NiagaraEmitterInstanceShader.usf
@@ -333,8 +333,15 @@ int rand(int x, int Seed1, int Seed2, int Seed3)
 	/* Parameters for spawning and simulations
 	 */
 	uint UpdateStartInstance;
-	uint DefaultShaderStageIndex;
-	uint ShaderStageIndex;
+	
+	/*
+	 * declared as int instead of uint to avoid compiler warnings because we can't expose uint on Niagara module parameters.  
+	 * This is still declared as a uint in the corresponding C++ classes
+	 */
+	int DefaultShaderStageIndex;
+	int ShaderStageIndex;
+
+
 	int IterationInterfaceCount;
 	uint SpawnedInstances;
 
