#Niagara IWYU fix

diff --git a/Engine/Plugins/FX/Niagara/Source/NiagaraEditorWidgets/Private/SNiagaraOverviewStack.cpp b/Engine/Plugins/FX/Niagara/Source/NiagaraEditorWidgets/Private/SNiagaraOverviewStack.cpp
--- a/Engine/Plugins/FX/Niagara/Source/NiagaraEditorWidgets/Private/SNiagaraOverviewStack.cpp
+++ b/Engine/Plugins/FX/Niagara/Source/NiagaraEditorWidgets/Private/SNiagaraOverviewStack.cpp
@@ -23,6 +23,7 @@
 #include "Widgets/Images/SImage.h"
 #include "Widgets/Input/SButton.h"
 #include "Widgets/Notifications/SNotificationList.h"
+#include "Widgets/SToolTip.h"
 #include "EditorStyleSet.h"
 #include "Styling/CoreStyle.h"
 #include "Framework/Application/SlateApplication.h"
