fixed null ptr bug with filling textures

diff --git a/Engine/Plugins/FX/Niagara/Source/Niagara/Private/NiagaraDataInterfaceGrid2DCollection.cpp b/Engine/Plugins/FX/Niagara/Source/Niagara/Private/NiagaraDataInterfaceGrid2DCollection.cpp
--- a/Engine/Plugins/FX/Niagara/Source/Niagara/Private/NiagaraDataInterfaceGrid2DCollection.cpp
+++ b/Engine/Plugins/FX/Niagara/Source/Niagara/Private/NiagaraDataInterfaceGrid2DCollection.cpp
@@ -535,7 +535,7 @@ bool UNiagaraDataInterfaceGrid2DCollection::FillRawTexture2D(const UNiagaraCompo
 	// check dest size and type needs to be float
 	// #todo(dmp): don't hardcode float since we might do other stuff in the future
 	EPixelFormat RequiredTye = PF_R32_FLOAT;
-	if (Dest->SizeX != NumCellsX * TilesX || Dest->SizeY != NumCellsY * TilesY || Dest->GetFormat() != RequiredTye)
+	if (!Dest || Dest->SizeX != NumCellsX * TilesX || Dest->SizeY != NumCellsY * TilesY || Dest->GetFormat() != RequiredTye)
 	{
 		return false;
 	}
