Fixes for using cooked LUT with niagara color curves, ensures that the cooked versions do get propagated to the render thread, and removes the UpdateLut code from non-editor data code

diff --git a/Engine/Plugins/FX/Niagara/Source/Niagara/Classes/NiagaraDataInterfaceColorCurve.h b/Engine/Plugins/FX/Niagara/Source/Niagara/Classes/NiagaraDataInterfaceColorCurve.h
--- a/Engine/Plugins/FX/Niagara/Source/Niagara/Classes/NiagaraDataInterfaceColorCurve.h
+++ b/Engine/Plugins/FX/Niagara/Source/Niagara/Classes/NiagaraDataInterfaceColorCurve.h
@@ -44,8 +44,10 @@ public:
 		CurveLUTNumElems = 4,
 	};
 
+#if WITH_EDITORONLY_DATA
 	virtual void UpdateTimeRanges() override;
 	virtual TArray<float> BuildLUT(int32 NumEntries) const override;
+#endif
 
 	virtual void GetFunctions(TArray<FNiagaraFunctionSignature>& OutFunctions)override;
 	virtual void GetVMExternalFunction(const FVMExternalFunctionBindingInfo& BindingInfo, void* InstanceData, FVMExternalFunction &OutFunc) override;
diff --git a/Engine/Plugins/FX/Niagara/Source/Niagara/Classes/NiagaraDataInterfaceCurveBase.h b/Engine/Plugins/FX/Niagara/Source/Niagara/Classes/NiagaraDataInterfaceCurveBase.h
--- a/Engine/Plugins/FX/Niagara/Source/Niagara/Classes/NiagaraDataInterfaceCurveBase.h
+++ b/Engine/Plugins/FX/Niagara/Source/Niagara/Classes/NiagaraDataInterfaceCurveBase.h
@@ -74,6 +74,7 @@ public:
 	UPROPERTY(EditAnywhere, AdvancedDisplay, Category = "Curve")
 	uint32 bUseLUT : 1;
 
+#if WITH_EDITORONLY_DATA
 	/** Do we optimize the LUT, this saves memory but may introduce errors.  Errors can be reduced modifying the threshold. */
 	UPROPERTY(EditAnywhere, AdvancedDisplay, Category = "Curve")
 	uint32 bOptimizeLUT : 1;
@@ -85,7 +86,6 @@ public:
 	UPROPERTY(EditAnywhere, AdvancedDisplay, Category = "Curve", meta=(EditCondition="bOverrideOptimizeThreshold"))
 	float OptimizeThreshold;
 
-#if WITH_EDITORONLY_DATA
 	UPROPERTY(EditAnywhere, Transient, Category = "Curve")
 	bool ShowInCurveEditor;
 #endif
@@ -96,10 +96,10 @@ public:
 		, LUTMaxTime(1.0f)
 		, LUTInvTimeRange(1.0f)
 		, bUseLUT(true)
+#if WITH_EDITORONLY_DATA
 		, bOptimizeLUT(true)
 		, bOverrideOptimizeThreshold(false)
 		, OptimizeThreshold(DefaultOptimizeThreshold)
-#if WITH_EDITORONLY_DATA
 		, ShowInCurveEditor(false)
 #endif
 	{
@@ -111,10 +111,10 @@ public:
 		, LUTMaxTime(1.0f)
 		, LUTInvTimeRange(1.0f)
 		, bUseLUT(true)
+#if WITH_EDITORONLY_DATA
 		, bOptimizeLUT(true)
 		, bOverrideOptimizeThreshold(false)
 		, OptimizeThreshold(DefaultOptimizeThreshold)
-#if WITH_EDITORONLY_DATA
 		, ShowInCurveEditor(false)
 #endif
 	{
@@ -153,8 +153,10 @@ public:
 	virtual void GetParameterDefinitionHLSL(FNiagaraDataInterfaceGPUParamInfo& ParamInfo, FString& OutHLSL) override;
 	virtual FNiagaraDataInterfaceParametersCS* ConstructComputeParameters() const override;
 
+#if WITH_EDITORONLY_DATA
 	void UpdateLUT();
 	void OptimizeLUT();
+#endif
 
 	//UNiagaraDataInterface interface
 	virtual bool Equals(const UNiagaraDataInterface* Other) const override;
diff --git a/Engine/Plugins/FX/Niagara/Source/Niagara/Private/NiagaraDataInterfaceColorCurve.cpp b/Engine/Plugins/FX/Niagara/Source/Niagara/Private/NiagaraDataInterfaceColorCurve.cpp
--- a/Engine/Plugins/FX/Niagara/Source/Niagara/Private/NiagaraDataInterfaceColorCurve.cpp
+++ b/Engine/Plugins/FX/Niagara/Source/Niagara/Private/NiagaraDataInterfaceColorCurve.cpp
@@ -18,7 +18,9 @@ const FName UNiagaraDataInterfaceColorCurve::SampleCurveName(TEXT("SampleColorCu
 UNiagaraDataInterfaceColorCurve::UNiagaraDataInterfaceColorCurve(FObjectInitializer const& ObjectInitializer)
 	: Super(ObjectInitializer)
 {
+#if WITH_EDITORONLY_DATA
 	UpdateLUT();
+#endif
 }
 
 void UNiagaraDataInterfaceColorCurve::PostInitProperties()
@@ -31,7 +33,9 @@ void UNiagaraDataInterfaceColorCurve::PostInitProperties()
 		FNiagaraTypeRegistry::Register(FNiagaraTypeDefinition(GetClass()), true, false, false);
 	}
 
+#if WITH_EDITORONLY_DATA
 	UpdateLUT();
+#endif
 }
 
 void UNiagaraDataInterfaceColorCurve::Serialize(FArchive& Ar)
@@ -61,9 +65,11 @@ void UNiagaraDataInterfaceColorCurve::Serialize(FArchive& Ar)
 #endif
 	{
 		Super::Serialize(Ar);
+		PushToRenderThread();
 	}
 }
 
+#if WITH_EDITORONLY_DATA
 void UNiagaraDataInterfaceColorCurve::UpdateTimeRanges()
 {
 	if ((RedCurve.GetNumKeys() > 0 || GreenCurve.GetNumKeys() > 0 || BlueCurve.GetNumKeys() > 0 || AlphaCurve.GetNumKeys() > 0))
@@ -106,6 +112,7 @@ TArray<float> UNiagaraDataInterfaceColorCurve::BuildLUT(int32 NumEntries) const
 	}
 	return OutputLUT;
 }
+#endif
 
 bool UNiagaraDataInterfaceColorCurve::CopyToInternal(UNiagaraDataInterface* Destination) const 
 {
@@ -118,12 +125,14 @@ bool UNiagaraDataInterfaceColorCurve::CopyToInternal(UNiagaraDataInterface* Dest
 	DestinationColorCurve->GreenCurve = GreenCurve;
 	DestinationColorCurve->BlueCurve = BlueCurve;
 	DestinationColorCurve->AlphaCurve = AlphaCurve;
+#if WITH_EDITORONLY_DATA
 	DestinationColorCurve->UpdateLUT();
 
 	if (!CompareLUTS(DestinationColorCurve->ShaderLUT))
 	{
 		UE_LOG(LogNiagara, Log, TEXT("CopyToInternal LUT generation is out of sync. Please investigate. %s to %s"), *GetPathName(), *DestinationColorCurve->GetPathName());
 	}
+#endif
 	return true;
 }
 
diff --git a/Engine/Plugins/FX/Niagara/Source/Niagara/Private/NiagaraDataInterfaceCurve.cpp b/Engine/Plugins/FX/Niagara/Source/Niagara/Private/NiagaraDataInterfaceCurve.cpp
--- a/Engine/Plugins/FX/Niagara/Source/Niagara/Private/NiagaraDataInterfaceCurve.cpp
+++ b/Engine/Plugins/FX/Niagara/Source/Niagara/Private/NiagaraDataInterfaceCurve.cpp
@@ -15,7 +15,9 @@ const FName UNiagaraDataInterfaceCurve::SampleCurveName(TEXT("SampleCurve"));
 UNiagaraDataInterfaceCurve::UNiagaraDataInterfaceCurve(FObjectInitializer const& ObjectInitializer)
 	: Super(ObjectInitializer)
 {
+#if WITH_EDITORONLY_DATA
 	UpdateLUT();
+#endif
 }
 
 void UNiagaraDataInterfaceCurve::PostInitProperties()
@@ -27,7 +29,9 @@ void UNiagaraDataInterfaceCurve::PostInitProperties()
 		FNiagaraTypeRegistry::Register(FNiagaraTypeDefinition(GetClass()), true, false, false);
 	}
 
+#if WITH_EDITORONLY_DATA
 	UpdateLUT();
+#endif
 }
 
 void UNiagaraDataInterfaceCurve::Serialize(FArchive& Ar)
@@ -48,6 +52,7 @@ void UNiagaraDataInterfaceCurve::Serialize(FArchive& Ar)
 #endif
 	{
 		Super::Serialize(Ar);
+		PushToRenderThread();
 	}
 }
 
@@ -58,11 +63,13 @@ bool UNiagaraDataInterfaceCurve::CopyToInternal(UNiagaraDataInterface* Destinati
 		return false;
 	}
 	CastChecked<UNiagaraDataInterfaceCurve>(Destination)->Curve = Curve;
+#if WITH_EDITORONLY_DATA
 	CastChecked<UNiagaraDataInterfaceCurve>(Destination)->UpdateLUT();
 	if (!CompareLUTS(CastChecked<UNiagaraDataInterfaceCurve>(Destination)->ShaderLUT))
 	{
 		UE_LOG(LogNiagara, Log, TEXT("Post CopyToInternal LUT generation is out of sync. Please investigate. %s"), *GetPathName());
 	}
+#endif
 	return true;
 }
 
diff --git a/Engine/Plugins/FX/Niagara/Source/Niagara/Private/NiagaraDataInterfaceCurveBase.cpp b/Engine/Plugins/FX/Niagara/Source/Niagara/Private/NiagaraDataInterfaceCurveBase.cpp
--- a/Engine/Plugins/FX/Niagara/Source/Niagara/Private/NiagaraDataInterfaceCurveBase.cpp
+++ b/Engine/Plugins/FX/Niagara/Source/Niagara/Private/NiagaraDataInterfaceCurveBase.cpp
@@ -29,6 +29,7 @@ void UNiagaraDataInterfaceCurveBase::PostLoad()
 {
 	Super::PostLoad();
 
+#if WITH_EDITORONLY_DATA
 	const int32 NiagaraVer = GetLinkerCustomVersion(FNiagaraCustomVersion::GUID);
 
 	if (NiagaraVer < FNiagaraCustomVersion::LatestVersion)
@@ -47,6 +48,7 @@ void UNiagaraDataInterfaceCurveBase::PostLoad()
 			}
 		}
 	}
+#endif
 }
 
 bool UNiagaraDataInterfaceCurveBase::CopyToInternal(UNiagaraDataInterface* Destination) const
@@ -58,9 +60,12 @@ bool UNiagaraDataInterfaceCurveBase::CopyToInternal(UNiagaraDataInterface* Desti
 
 	UNiagaraDataInterfaceCurveBase* DestinationTyped = CastChecked<UNiagaraDataInterfaceCurveBase>(Destination);
 	DestinationTyped->bUseLUT = bUseLUT;
+	DestinationTyped->ShaderLUT = ShaderLUT;
+#if WITH_EDITORONLY_DATA
 	DestinationTyped->bOptimizeLUT = bOptimizeLUT;
 	DestinationTyped->bOverrideOptimizeThreshold = bOverrideOptimizeThreshold;
 	DestinationTyped->OptimizeThreshold = OptimizeThreshold;
+#endif
 
 	return true;
 }
@@ -89,6 +94,7 @@ bool UNiagaraDataInterfaceCurveBase::CompareLUTS(const TArray<float>& OtherLUT)
 	}
 }
 
+#if WITH_EDITORONLY_DATA
 void UNiagaraDataInterfaceCurveBase::UpdateLUT()
 {
 	UpdateTimeRanges();
@@ -171,6 +177,7 @@ void UNiagaraDataInterfaceCurveBase::OptimizeLUT()
 		}
 	}
 }
+#endif
 
 bool UNiagaraDataInterfaceCurveBase::Equals(const UNiagaraDataInterface* Other) const
 {
@@ -181,12 +188,14 @@ bool UNiagaraDataInterfaceCurveBase::Equals(const UNiagaraDataInterface* Other)
 
 	const UNiagaraDataInterfaceCurveBase* OtherTyped = CastChecked<UNiagaraDataInterfaceCurveBase>(Other);
 	bool bEqual = OtherTyped->bUseLUT == bUseLUT;
+#if WITH_EDITORONLY_DATA
 	bEqual &= OtherTyped->bOptimizeLUT == bOptimizeLUT;
 	bEqual &= OtherTyped->bOverrideOptimizeThreshold == bOverrideOptimizeThreshold;
 	if (bOverrideOptimizeThreshold)
 	{
 		bEqual &= OtherTyped->OptimizeThreshold == OptimizeThreshold;
 	}
+#endif
 
 	return bEqual;
 }
@@ -257,6 +266,8 @@ void UNiagaraDataInterfaceCurveBase::PushToRenderThread()
 		RT_Proxy->CurveLUTNumMinusOne = rtCurveLUTNumMinusOne;
 
 		RT_Proxy->CurveLUT.Release();
+
+		check(rtShaderLUT.Num());
 		RT_Proxy->CurveLUT.Initialize(sizeof(float), rtShaderLUT.Num(), EPixelFormat::PF_R32_FLOAT, BUF_Static);
 		uint32 BufferSize = rtShaderLUT.Num() * sizeof(float);
 		int32 *BufferData = static_cast<int32*>(RHILockVertexBuffer(RT_Proxy->CurveLUT.Buffer, 0, BufferSize, EResourceLockMode::RLM_WriteOnly));
diff --git a/Engine/Plugins/FX/Niagara/Source/Niagara/Private/NiagaraDataInterfaceVector2DCurve.cpp b/Engine/Plugins/FX/Niagara/Source/Niagara/Private/NiagaraDataInterfaceVector2DCurve.cpp
--- a/Engine/Plugins/FX/Niagara/Source/Niagara/Private/NiagaraDataInterfaceVector2DCurve.cpp
+++ b/Engine/Plugins/FX/Niagara/Source/Niagara/Private/NiagaraDataInterfaceVector2DCurve.cpp
@@ -18,7 +18,9 @@ const FName UNiagaraDataInterfaceVector2DCurve::SampleCurveName("SampleVector2DC
 UNiagaraDataInterfaceVector2DCurve::UNiagaraDataInterfaceVector2DCurve(FObjectInitializer const& ObjectInitializer)
 	: Super(ObjectInitializer)
 {
+#if WITH_EDITORONLY_DATA
 	UpdateLUT();
+#endif
 }
 
 void UNiagaraDataInterfaceVector2DCurve::PostInitProperties()
@@ -30,7 +32,9 @@ void UNiagaraDataInterfaceVector2DCurve::PostInitProperties()
 		FNiagaraTypeRegistry::Register(FNiagaraTypeDefinition(GetClass()), true, false, false);
 	}
 
+#if WITH_EDITORONLY_DATA
 	UpdateLUT();
+#endif
 }
 
 void UNiagaraDataInterfaceVector2DCurve::Serialize(FArchive& Ar)
@@ -54,6 +58,7 @@ void UNiagaraDataInterfaceVector2DCurve::Serialize(FArchive& Ar)
 #endif
 	{
 		Super::Serialize(Ar);
+		PushToRenderThread();
 	}
 }
 
@@ -103,11 +108,13 @@ bool UNiagaraDataInterfaceVector2DCurve::CopyToInternal(UNiagaraDataInterface* D
 	UNiagaraDataInterfaceVector2DCurve* DestinationVector2DCurve = CastChecked<UNiagaraDataInterfaceVector2DCurve>(Destination);
 	DestinationVector2DCurve->XCurve = XCurve;
 	DestinationVector2DCurve->YCurve = YCurve;
+#if WITH_EDITORONLY_DATA
 	DestinationVector2DCurve->UpdateLUT();
 	if (!CompareLUTS(DestinationVector2DCurve->ShaderLUT))
 	{
 		UE_LOG(LogNiagara, Log, TEXT("Post CopyToInternal LUT generation is out of sync. Please investigate. %s"), *GetPathName());
 	}
+#endif
 	return true;
 }
 
diff --git a/Engine/Plugins/FX/Niagara/Source/Niagara/Private/NiagaraDataInterfaceVector4Curve.cpp b/Engine/Plugins/FX/Niagara/Source/Niagara/Private/NiagaraDataInterfaceVector4Curve.cpp
--- a/Engine/Plugins/FX/Niagara/Source/Niagara/Private/NiagaraDataInterfaceVector4Curve.cpp
+++ b/Engine/Plugins/FX/Niagara/Source/Niagara/Private/NiagaraDataInterfaceVector4Curve.cpp
@@ -18,7 +18,9 @@ const FName UNiagaraDataInterfaceVector4Curve::SampleCurveName(TEXT("SampleColor
 UNiagaraDataInterfaceVector4Curve::UNiagaraDataInterfaceVector4Curve(FObjectInitializer const& ObjectInitializer)
 	: Super(ObjectInitializer)
 {
+#if WITH_EDITORONLY_DATA
 	UpdateLUT();
+#endif
 }
 
 void UNiagaraDataInterfaceVector4Curve::PostInitProperties()
@@ -31,7 +33,9 @@ void UNiagaraDataInterfaceVector4Curve::PostInitProperties()
 		FNiagaraTypeRegistry::Register(FNiagaraTypeDefinition(GetClass()), true, false, false);
 	}
 
+#if WITH_EDITORONLY_DATA
 	UpdateLUT();
+#endif
 }
 
 void UNiagaraDataInterfaceVector4Curve::Serialize(FArchive& Ar)
@@ -61,6 +65,7 @@ void UNiagaraDataInterfaceVector4Curve::Serialize(FArchive& Ar)
 #endif
 	{
 		Super::Serialize(Ar);
+		PushToRenderThread();
 	}
 }
 
@@ -118,11 +123,13 @@ bool UNiagaraDataInterfaceVector4Curve::CopyToInternal(UNiagaraDataInterface* De
 	DestinationColorCurve->YCurve = YCurve;
 	DestinationColorCurve->ZCurve = ZCurve;
 	DestinationColorCurve->WCurve = WCurve;
+#if WITH_EDITORONLY_DATA
 	DestinationColorCurve->UpdateLUT();
 	if (!CompareLUTS(DestinationColorCurve->ShaderLUT))
 	{
 		UE_LOG(LogNiagara, Log, TEXT("Post CopyToInternal LUT generation is out of sync. Please investigate. %s"), *GetPathName());
 	}
+#endif
 
 	return true;
 }
diff --git a/Engine/Plugins/FX/Niagara/Source/Niagara/Private/NiagaraDataInterfaceVectorCurve.cpp b/Engine/Plugins/FX/Niagara/Source/Niagara/Private/NiagaraDataInterfaceVectorCurve.cpp
--- a/Engine/Plugins/FX/Niagara/Source/Niagara/Private/NiagaraDataInterfaceVectorCurve.cpp
+++ b/Engine/Plugins/FX/Niagara/Source/Niagara/Private/NiagaraDataInterfaceVectorCurve.cpp
@@ -18,7 +18,9 @@ const FName UNiagaraDataInterfaceVectorCurve::SampleCurveName(TEXT("SampleVector
 UNiagaraDataInterfaceVectorCurve::UNiagaraDataInterfaceVectorCurve(FObjectInitializer const& ObjectInitializer)
 	: Super(ObjectInitializer)
 {
+#if WITH_EDITORONLY_DATA
 	UpdateLUT();
+#endif
 }
 
 void UNiagaraDataInterfaceVectorCurve::PostInitProperties()
@@ -31,7 +33,9 @@ void UNiagaraDataInterfaceVectorCurve::PostInitProperties()
 		FNiagaraTypeRegistry::Register(FNiagaraTypeDefinition(GetClass()), true, false, false);
 	}
 
+#if WITH_EDITORONLY_DATA
 	UpdateLUT();
+#endif
 }
 
 void UNiagaraDataInterfaceVectorCurve::Serialize(FArchive& Ar)
@@ -58,6 +62,7 @@ void UNiagaraDataInterfaceVectorCurve::Serialize(FArchive& Ar)
 #endif
 	{
 		Super::Serialize(Ar);
+		PushToRenderThread();
 	}
 }
 
@@ -111,12 +116,14 @@ bool UNiagaraDataInterfaceVectorCurve::CopyToInternal(UNiagaraDataInterface* Des
 	DestinationVectorCurve->XCurve = XCurve;
 	DestinationVectorCurve->YCurve = YCurve;
 	DestinationVectorCurve->ZCurve = ZCurve;
+#if WITH_EDITORONLY_DATA
 	DestinationVectorCurve->UpdateLUT();
 
 	if (!CompareLUTS(DestinationVectorCurve->ShaderLUT))
 	{
 		UE_LOG(LogNiagara, Log, TEXT("Post CopyToInternal LUT generation is out of sync. Please investigate. %s"), *GetPathName());
 	}
+#endif
 	return true;
 }
 
