- Skip building the register table if we didn't actual allocate anything

#rnx

diff --git a/Engine/Plugins/FX/Niagara/Source/Niagara/Classes/NiagaraDataSet.h b/Engine/Plugins/FX/Niagara/Source/Niagara/Classes/NiagaraDataSet.h
--- a/Engine/Plugins/FX/Niagara/Source/Niagara/Classes/NiagaraDataSet.h
+++ b/Engine/Plugins/FX/Niagara/Source/Niagara/Classes/NiagaraDataSet.h
@@ -107,7 +107,6 @@ public:
 	void SwapInstances(uint32 OldIndex, uint32 NewIndex);
 	void KillInstance(uint32 InstanceIdx);
 	void CopyTo(FNiagaraDataBuffer& DestBuffer, int32 SrcStartIdx, int32 DestStartIdx, int32 NumInstances)const;
-	void CopyTo(FNiagaraDataBuffer& DestBuffer)const;
 	void GPUCopyFrom(float* GPUReadBackFloat, int* GPUReadBackInt, int32 StartIdx, int32 NumInstances, uint32 InSrcFloatStride, uint32 InSrcIntStride);
 	void Dump(int32 StartIndex, int32 NumInstances, const FString& Label)const;
 
@@ -1499,7 +1498,14 @@ class NIAGARA_API FScopedNiagaraDataSetGPUReadback
 {
 public:
 	FScopedNiagaraDataSetGPUReadback() {}
-	~FScopedNiagaraDataSetGPUReadback();
+	FORCEINLINE ~FScopedNiagaraDataSetGPUReadback()
+	{
+		if (DataBuffer != nullptr)
+		{
+			DataBuffer->FloatData.Empty();
+			DataBuffer->Int32Data.Empty();
+		}
+	}
 
 	void ReadbackData(class NiagaraEmitterInstanceBatcher* Batcher, FNiagaraDataSet* InDataSet);
 	uint32 GetNumInstances() const { check(DataSet != nullptr); return NumInstances; }
diff --git a/Engine/Plugins/FX/Niagara/Source/Niagara/Private/NiagaraDataSet.cpp b/Engine/Plugins/FX/Niagara/Source/Niagara/Private/NiagaraDataSet.cpp
--- a/Engine/Plugins/FX/Niagara/Source/Niagara/Private/NiagaraDataSet.cpp
+++ b/Engine/Plugins/FX/Niagara/Source/Niagara/Private/NiagaraDataSet.cpp
@@ -129,11 +129,6 @@ void FNiagaraDataSet::ResetBuffersInternal()
 	//Ensure we have a valid current buffer
 	BeginSimulate();
 	EndSimulate();
-
-	if (GetSimTarget() == ENiagaraSimTarget::CPUSim)
-	{
-		GetCurrentDataChecked().BuildRegisterTable();
-	}
 }
 
 void FNiagaraDataSet::ReleaseBuffers()
@@ -202,14 +197,12 @@ void FNiagaraDataSet::Allocate(int32 NumInstances, bool bMaintainExisting)
 	CheckCorrectThread();
 	checkSlow(DestinationData);
 
-	DestinationData->GetIDTable().Reset();
+	DestinationData->Allocate(NumInstances);
 	if (bMaintainExisting)
 	{
-		CurrentData->CopyTo(*DestinationData);
+		CurrentData->CopyTo(*DestinationData, 0, 0, CurrentData->GetNumInstances());
 	}
 
-	DestinationData->Allocate(NumInstances, bMaintainExisting);
-
 #if NIAGARA_NAN_CHECKING
 	CheckForNaNs();
 #endif
@@ -517,73 +510,92 @@ void FNiagaraDataBuffer::Allocate(uint32 InNumInstances, bool bMaintainExisting)
 	//CheckUsage(false);
 	checkSlow(Owner->GetSimTarget() == ENiagaraSimTarget::CPUSim);
 
-	NumInstancesAllocated = InNumInstances;
 	NumInstances = 0;
+	if (!bMaintainExisting)
+	{
+		IDToIndexTable.Reset();
+	}
 
-	DEC_MEMORY_STAT_BY(STAT_NiagaraParticleMemory, FloatData.GetAllocatedSize() + Int32Data.GetAllocatedSize());
+	// Calculate allocation size
+	uint32 NewFloatStride = GetSafeComponentBufferSize(InNumInstances * sizeof(float));
+	int32 NewFloatNum = NewFloatStride * Owner->GetNumFloatComponents();
 
-	const uint32 OldFloatStride = FloatStride;
-	TArray<uint8> OldFloatData;
-	const uint32 OldInt32Stride = Int32Stride;
-	TArray<uint8> OldIntData;
+	uint32 NewInt32Stride = GetSafeComponentBufferSize(InNumInstances * sizeof(int32));
+	int32 NewInt32Num = NewInt32Stride * Owner->GetNumInt32Components();
 
-	if (bMaintainExisting)
+	// Do sizes match?
+	if (NewFloatNum != FloatData.Num() && NewInt32Num != Int32Data.Num())
 	{
-		//Need to copy off old data so we can copy it back into the newly laid out buffers. TODO: Avoid this needless copying.
-		OldFloatData = FloatData;
-		OldIntData = Int32Data;
-	}
+		// Do we need to grow or shrink?
+		const bool bGrowData = NewFloatNum > FloatData.Num() || NewInt32Num > Int32Data.Num();
+		const bool bShrinkFloatData = !bGrowData && (GNiagaraDataBufferShrinkFactor * FMath::Max(GNiagaraDataBufferMinSize, NewFloatNum) < FloatData.Max() || !NewFloatNum);
+		const bool bShrinkIntData = !bGrowData && (GNiagaraDataBufferShrinkFactor * FMath::Max(GNiagaraDataBufferMinSize, NewInt32Num) < Int32Data.Max() || !NewInt32Num);
+		if ( bGrowData || bShrinkFloatData || bShrinkIntData )
+		{
+			NumInstancesAllocated = InNumInstances;
 
-	FloatStride = GetSafeComponentBufferSize(NumInstancesAllocated * sizeof(float));
-	{
-		const int32 NewNum = FloatStride * Owner->GetNumFloatComponents();
-		const bool bAllowShrink = GNiagaraDataBufferShrinkFactor * FMath::Max(GNiagaraDataBufferMinSize, NewNum) < FloatData.Max() || !NewNum;
-		FloatData.SetNum(NewNum, bAllowShrink);
-	}
+			DEC_MEMORY_STAT_BY(STAT_NiagaraParticleMemory, FloatData.GetAllocatedSize() + Int32Data.GetAllocatedSize());
+			if (bMaintainExisting)
+			{
+				TArray<uint8> NewFloatData;
+				TArray<uint8> NewInt32Data;
+				NewFloatData.SetNum(NewFloatNum);
+				NewInt32Data.SetNum(NewInt32Num);
 
-	Int32Stride = GetSafeComponentBufferSize(NumInstancesAllocated * sizeof(int32));
-	{
-		const int32 NewNum = Int32Stride * Owner->GetNumInt32Components();
-		const bool bAllowShrink = GNiagaraDataBufferShrinkFactor * FMath::Max(GNiagaraDataBufferMinSize, NewNum) < Int32Data.Max() || !NewNum;
-		Int32Data.SetNum(NewNum, bAllowShrink);
-	}
+				if (NewFloatStride > 0 && FloatStride > 0)
+				{
+					const int32 FloatComponents = Owner->GetNumFloatComponents();
+					const uint32 BytesToCopy = FMath::Min(NewFloatStride, FloatStride);
+					for (int32 CompIdx=FloatComponents-1; CompIdx >= 0; --CompIdx)
+					{
+						const uint8* Src = FloatData.GetData() + FloatStride * CompIdx;
+						uint8* Dst = NewFloatData.GetData() + NewFloatStride * CompIdx;
+						FMemory::Memcpy(Dst, Src, BytesToCopy);
+					}
+				}
 
-	INC_MEMORY_STAT_BY(STAT_NiagaraParticleMemory, FloatData.GetAllocatedSize() + Int32Data.GetAllocatedSize());
+				if (NewInt32Stride > 0 && Int32Stride > 0)
+				{
+					const int32 IntComponents = Owner->GetNumInt32Components();
+					const uint32 BytesToCopy = FMath::Min(NewInt32Stride, Int32Stride);
+					for (int32 CompIdx=IntComponents - 1; CompIdx >= 0; --CompIdx)
+					{
+						const uint8* Src = Int32Data.GetData() + Int32Stride * CompIdx;
+						uint8* Dst = NewInt32Data.GetData() + NewInt32Stride * CompIdx;
+						FMemory::Memcpy(Dst, Src, BytesToCopy);
+					}
+				}
 
-	//In some cases we want the existing data in the buffer to be maintained which due to the data layout requires some fix up.
-	if (bMaintainExisting)
-	{
-		if (FloatStride != OldFloatStride && FloatStride > 0 && OldFloatStride > 0)
-		{
-			int32 FloatComponents = Owner->GetNumFloatComponents();
-			const uint32 BytesToCopy = FMath::Min(OldFloatStride, FloatStride);
-			for (int32 CompIdx = FloatComponents - 1; CompIdx >= 0; --CompIdx)
+				FloatData = MoveTemp(NewFloatData);
+				Int32Data = MoveTemp(NewInt32Data);
+			}
+			else
 			{
-				uint8* Src = OldFloatData.GetData() + OldFloatStride * CompIdx;
-				uint8* Dst = FloatData.GetData() + FloatStride * CompIdx;
-				FMemory::Memcpy(Dst, Src, BytesToCopy);
+				FloatData.SetNum(NewFloatNum, bShrinkFloatData);
+				Int32Data.SetNum(NewInt32Num, bShrinkIntData);
 			}
+			INC_MEMORY_STAT_BY(STAT_NiagaraParticleMemory, FloatData.GetAllocatedSize() + Int32Data.GetAllocatedSize());
 		}
-
-		if (Int32Stride != OldInt32Stride && Int32Stride > 0 && OldInt32Stride > 0)
+		// Calculate strides based upon max of instance counts
+		// This allows us to skip building the register table when shrinking
+		else
 		{
-			int32 IntComponents = Owner->GetNumInt32Components();
-			const uint32 BytesToCopy = FMath::Min(OldInt32Stride, Int32Stride);
-			for (int32 CompIdx = IntComponents - 1; CompIdx >= 0; --CompIdx)
-			{
-				uint8* Src = OldIntData.GetData() + OldInt32Stride * CompIdx;
-				uint8* Dst = Int32Data.GetData() + Int32Stride * CompIdx;
-				FMemory::Memcpy(Dst, Src, BytesToCopy);
-			}
+			NumInstancesAllocated = FMath::Max(NumInstancesAllocated, InNumInstances);
+			NewFloatStride = GetSafeComponentBufferSize(NumInstancesAllocated * sizeof(float));
+			NewInt32Stride = GetSafeComponentBufferSize(NumInstancesAllocated * sizeof(int32));
 		}
 	}
 	else
 	{
-		IDToIndexTable.Reset();
+		NumInstancesAllocated = InNumInstances;
 	}
 
-	//Must rebuild our register table every time we allocate new buffers.
-	BuildRegisterTable();
+	if ( (NewFloatStride != FloatStride) || (NewInt32Stride != Int32Stride) )
+	{
+		FloatStride = NewFloatStride;
+		Int32Stride = NewInt32Stride;
+		BuildRegisterTable();
+	}
 }
 
 void FNiagaraDataBuffer::AllocateGPU(uint32 InNumInstances, FNiagaraGPUInstanceCountManager& GPUInstanceCountManager, FRHICommandList &RHICmdList)
@@ -819,19 +831,6 @@ void FNiagaraDataBuffer::GPUCopyFrom(float* GPUReadBackFloat, int* GPUReadBackIn
 	}
 }
 
-void FNiagaraDataBuffer::CopyTo(FNiagaraDataBuffer& DestBuffer)const
-{
-	CheckUsage(true);
-	DestBuffer.CheckUsage(false);
-	DestBuffer.FloatStride = FloatStride;
-	DestBuffer.FloatData = FloatData;
-	DestBuffer.Int32Stride = Int32Stride;
-	DestBuffer.Int32Data = Int32Data;
-	DestBuffer.NumInstancesAllocated = NumInstancesAllocated;
-	DestBuffer.NumInstances = NumInstances;
-	DestBuffer.IDToIndexTable = IDToIndexTable;
-}
-
 void FNiagaraDataBuffer::Dump(int32 StartIndex, int32 InNumInstances, const FString& Label)const
 {
 	FNiagaraDataVariableIterator Itr(this, StartIndex);
@@ -973,15 +972,6 @@ void FNiagaraDataBuffer::ReleaseGPUInstanceCount(FNiagaraGPUInstanceCountManager
 }
 
 
-FScopedNiagaraDataSetGPUReadback::~FScopedNiagaraDataSetGPUReadback()
-{
-	if (DataBuffer != nullptr)
-	{
-		DataBuffer->FloatData.Empty();
-		DataBuffer->Int32Data.Empty();
-	}
-}
-
 void FScopedNiagaraDataSetGPUReadback::ReadbackData(NiagaraEmitterInstanceBatcher* InBatcher, FNiagaraDataSet* InDataSet)
 {
 	check(DataSet == nullptr);
