- Update LUTs when during editor load and move PushToRenderThread to be during base class serialization
- Fix an issue where we didn't set the object default to be the same in none editor builds leading to 0 entry curves

#rnx

diff --git a/Engine/Plugins/FX/Niagara/Source/Niagara/Classes/NiagaraDataInterfaceCurveBase.h b/Engine/Plugins/FX/Niagara/Source/Niagara/Classes/NiagaraDataInterfaceCurveBase.h
--- a/Engine/Plugins/FX/Niagara/Source/Niagara/Classes/NiagaraDataInterfaceCurveBase.h
+++ b/Engine/Plugins/FX/Niagara/Source/Niagara/Classes/NiagaraDataInterfaceCurveBase.h
@@ -145,6 +145,7 @@ public:
 
 	//UObject Interface
 	virtual void PostLoad() override;
+	virtual void Serialize(FArchive& Ar) override;
 	//UObject Interface End
 
 	/** Gets information for all of the curves owned by this curve data interface. */
@@ -153,6 +154,7 @@ public:
 	virtual void GetParameterDefinitionHLSL(FNiagaraDataInterfaceGPUParamInfo& ParamInfo, FString& OutHLSL) override;
 	virtual FNiagaraDataInterfaceParametersCS* ConstructComputeParameters() const override;
 
+	void SetDefaultLUT();
 #if WITH_EDITORONLY_DATA
 	void UpdateLUT();
 	void OptimizeLUT();
diff --git a/Engine/Plugins/FX/Niagara/Source/Niagara/Private/NiagaraDataInterfaceColorCurve.cpp b/Engine/Plugins/FX/Niagara/Source/Niagara/Private/NiagaraDataInterfaceColorCurve.cpp
--- a/Engine/Plugins/FX/Niagara/Source/Niagara/Private/NiagaraDataInterfaceColorCurve.cpp
+++ b/Engine/Plugins/FX/Niagara/Source/Niagara/Private/NiagaraDataInterfaceColorCurve.cpp
@@ -18,9 +18,7 @@ const FName UNiagaraDataInterfaceColorCurve::SampleCurveName(TEXT("SampleColorCu
 UNiagaraDataInterfaceColorCurve::UNiagaraDataInterfaceColorCurve(FObjectInitializer const& ObjectInitializer)
 	: Super(ObjectInitializer)
 {
-#if WITH_EDITORONLY_DATA
-	UpdateLUT();
-#endif
+	SetDefaultLUT();
 }
 
 void UNiagaraDataInterfaceColorCurve::PostInitProperties()
@@ -65,7 +63,6 @@ void UNiagaraDataInterfaceColorCurve::Serialize(FArchive& Ar)
 #endif
 	{
 		Super::Serialize(Ar);
-		PushToRenderThread();
 	}
 }
 
diff --git a/Engine/Plugins/FX/Niagara/Source/Niagara/Private/NiagaraDataInterfaceCurve.cpp b/Engine/Plugins/FX/Niagara/Source/Niagara/Private/NiagaraDataInterfaceCurve.cpp
--- a/Engine/Plugins/FX/Niagara/Source/Niagara/Private/NiagaraDataInterfaceCurve.cpp
+++ b/Engine/Plugins/FX/Niagara/Source/Niagara/Private/NiagaraDataInterfaceCurve.cpp
@@ -15,9 +15,7 @@ const FName UNiagaraDataInterfaceCurve::SampleCurveName(TEXT("SampleCurve"));
 UNiagaraDataInterfaceCurve::UNiagaraDataInterfaceCurve(FObjectInitializer const& ObjectInitializer)
 	: Super(ObjectInitializer)
 {
-#if WITH_EDITORONLY_DATA
-	UpdateLUT();
-#endif
+	SetDefaultLUT();
 }
 
 void UNiagaraDataInterfaceCurve::PostInitProperties()
@@ -52,7 +50,6 @@ void UNiagaraDataInterfaceCurve::Serialize(FArchive& Ar)
 #endif
 	{
 		Super::Serialize(Ar);
-		PushToRenderThread();
 	}
 }
 
diff --git a/Engine/Plugins/FX/Niagara/Source/Niagara/Private/NiagaraDataInterfaceCurveBase.cpp b/Engine/Plugins/FX/Niagara/Source/Niagara/Private/NiagaraDataInterfaceCurveBase.cpp
--- a/Engine/Plugins/FX/Niagara/Source/Niagara/Private/NiagaraDataInterfaceCurveBase.cpp
+++ b/Engine/Plugins/FX/Niagara/Source/Niagara/Private/NiagaraDataInterfaceCurveBase.cpp
@@ -51,6 +51,28 @@ void UNiagaraDataInterfaceCurveBase::PostLoad()
 #endif
 }
 
+void UNiagaraDataInterfaceCurveBase::Serialize(FArchive& Ar)
+{
+	Super::Serialize(Ar);
+
+	// Push to render thread if loading
+	if (Ar.IsLoading())
+	{
+#if WITH_EDITORONLY_DATA
+		// Sometimes curves are out of date which needs to be tracked down
+		// Temporarily we will make sure they are up to date in editor builds
+		if (GetClass() != UNiagaraDataInterfaceCurveBase::StaticClass())
+		{
+			UpdateLUT();
+		}
+		else
+#endif
+		{
+			PushToRenderThread();
+		}
+	}
+}
+
 bool UNiagaraDataInterfaceCurveBase::CopyToInternal(UNiagaraDataInterface* Destination) const
 {
 	if (!Super::CopyToInternal(Destination))
@@ -94,6 +116,13 @@ bool UNiagaraDataInterfaceCurveBase::CompareLUTS(const TArray<float>& OtherLUT)
 	}
 }
 
+void UNiagaraDataInterfaceCurveBase::SetDefaultLUT()
+{
+	ShaderLUT.Empty(GetCurveNumElems());
+	ShaderLUT.AddDefaulted(GetCurveNumElems());
+	LUTNumSamplesMinusOne = 0;
+}
+
 #if WITH_EDITORONLY_DATA
 void UNiagaraDataInterfaceCurveBase::UpdateLUT()
 {
@@ -107,9 +136,7 @@ void UNiagaraDataInterfaceCurveBase::UpdateLUT()
 
 	if (!bUseLUT || (ShaderLUT.Num() == 0))
 	{
-		ShaderLUT.Empty(GetCurveNumElems());
-		ShaderLUT.AddDefaulted(GetCurveNumElems());
-		LUTNumSamplesMinusOne = 0;
+		SetDefaultLUT();
 	}
 
 	PushToRenderThread();
@@ -196,6 +223,10 @@ bool UNiagaraDataInterfaceCurveBase::Equals(const UNiagaraDataInterface* Other)
 		bEqual &= OtherTyped->OptimizeThreshold == OptimizeThreshold;
 	}
 #endif
+	if ( bEqual && bUseLUT )
+	{
+		bEqual &= OtherTyped->ShaderLUT == ShaderLUT;
+	}
 
 	return bEqual;
 }
diff --git a/Engine/Plugins/FX/Niagara/Source/Niagara/Private/NiagaraDataInterfaceVector2DCurve.cpp b/Engine/Plugins/FX/Niagara/Source/Niagara/Private/NiagaraDataInterfaceVector2DCurve.cpp
--- a/Engine/Plugins/FX/Niagara/Source/Niagara/Private/NiagaraDataInterfaceVector2DCurve.cpp
+++ b/Engine/Plugins/FX/Niagara/Source/Niagara/Private/NiagaraDataInterfaceVector2DCurve.cpp
@@ -18,9 +18,7 @@ const FName UNiagaraDataInterfaceVector2DCurve::SampleCurveName("SampleVector2DC
 UNiagaraDataInterfaceVector2DCurve::UNiagaraDataInterfaceVector2DCurve(FObjectInitializer const& ObjectInitializer)
 	: Super(ObjectInitializer)
 {
-#if WITH_EDITORONLY_DATA
-	UpdateLUT();
-#endif
+	SetDefaultLUT();
 }
 
 void UNiagaraDataInterfaceVector2DCurve::PostInitProperties()
@@ -58,7 +56,6 @@ void UNiagaraDataInterfaceVector2DCurve::Serialize(FArchive& Ar)
 #endif
 	{
 		Super::Serialize(Ar);
-		PushToRenderThread();
 	}
 }
 
diff --git a/Engine/Plugins/FX/Niagara/Source/Niagara/Private/NiagaraDataInterfaceVector4Curve.cpp b/Engine/Plugins/FX/Niagara/Source/Niagara/Private/NiagaraDataInterfaceVector4Curve.cpp
--- a/Engine/Plugins/FX/Niagara/Source/Niagara/Private/NiagaraDataInterfaceVector4Curve.cpp
+++ b/Engine/Plugins/FX/Niagara/Source/Niagara/Private/NiagaraDataInterfaceVector4Curve.cpp
@@ -18,9 +18,7 @@ const FName UNiagaraDataInterfaceVector4Curve::SampleCurveName(TEXT("SampleColor
 UNiagaraDataInterfaceVector4Curve::UNiagaraDataInterfaceVector4Curve(FObjectInitializer const& ObjectInitializer)
 	: Super(ObjectInitializer)
 {
-#if WITH_EDITORONLY_DATA
-	UpdateLUT();
-#endif
+	SetDefaultLUT();
 }
 
 void UNiagaraDataInterfaceVector4Curve::PostInitProperties()
@@ -65,7 +63,6 @@ void UNiagaraDataInterfaceVector4Curve::Serialize(FArchive& Ar)
 #endif
 	{
 		Super::Serialize(Ar);
-		PushToRenderThread();
 	}
 }
 
diff --git a/Engine/Plugins/FX/Niagara/Source/Niagara/Private/NiagaraDataInterfaceVectorCurve.cpp b/Engine/Plugins/FX/Niagara/Source/Niagara/Private/NiagaraDataInterfaceVectorCurve.cpp
--- a/Engine/Plugins/FX/Niagara/Source/Niagara/Private/NiagaraDataInterfaceVectorCurve.cpp
+++ b/Engine/Plugins/FX/Niagara/Source/Niagara/Private/NiagaraDataInterfaceVectorCurve.cpp
@@ -18,9 +18,7 @@ const FName UNiagaraDataInterfaceVectorCurve::SampleCurveName(TEXT("SampleVector
 UNiagaraDataInterfaceVectorCurve::UNiagaraDataInterfaceVectorCurve(FObjectInitializer const& ObjectInitializer)
 	: Super(ObjectInitializer)
 {
-#if WITH_EDITORONLY_DATA
-	UpdateLUT();
-#endif
+	SetDefaultLUT();
 }
 
 void UNiagaraDataInterfaceVectorCurve::PostInitProperties()
@@ -62,7 +60,6 @@ void UNiagaraDataInterfaceVectorCurve::Serialize(FArchive& Ar)
 #endif
 	{
 		Super::Serialize(Ar);
-		PushToRenderThread();
 	}
 }
 
