Fix for bug when allocating in datasets with only integer or float variables.

#rnx

diff --git a/Engine/Plugins/FX/Niagara/Source/Niagara/Private/NiagaraDataSet.cpp b/Engine/Plugins/FX/Niagara/Source/Niagara/Private/NiagaraDataSet.cpp
--- a/Engine/Plugins/FX/Niagara/Source/Niagara/Private/NiagaraDataSet.cpp
+++ b/Engine/Plugins/FX/Niagara/Source/Niagara/Private/NiagaraDataSet.cpp
@@ -524,7 +524,7 @@ void FNiagaraDataBuffer::Allocate(uint32 InNumInstances, bool bMaintainExisting)
 	int32 NewInt32Num = NewInt32Stride * Owner->GetNumInt32Components();
 
 	// Do sizes match?
-	if (NewFloatNum != FloatData.Num() && NewInt32Num != Int32Data.Num())
+	if (NewFloatNum != FloatData.Num() || NewInt32Num != Int32Data.Num())
 	{
 		// Do we need to grow or shrink?
 		const bool bGrowData = NewFloatNum > FloatData.Num() || NewInt32Num > Int32Data.Num();
