Niagara High Level Scalability as a first class feature.

This CL provides features for managing very high/component level scaling.
These features will cull whole systems. For sub system, emitter level scaling we still need the emitter scalability module code.
This feature entirely replaces system script level scalability modules which are now deprecated.

New "Niagara Effect Type" asset
Allows easy centralized place for leads and tech artists to define scalability control for the different effect types a game may have. eg impact vs environment vs player vs weapon.
Means 99% case for artists is simple selecting the correct effect type. We could intergrate this into the templates to make even easier.
These are assets rather than project settings to allow quick and easy management and migration between projects.

Scalability features provided:
Each of these is usable in isolation or together and can be overriden by each UNiagaraSystem.
All settings driving these features can be set per detail level.

Distance Cull - Systems beyond max distance are culled.
Visibility Cull - Systems not rendered for N seconds are culled.
OwnerLOD Cull - Systems who's owner's LOD goes beyond a set value will be culled. Useful for hooking into player FX in games where the LOD is driven by the main significance manager for example.
Instance Count - New instances will fail to spawn when there are too many instances of the effect type.
Runtime Perf Checks - Code stubs for tracking runtime perf are included but are disabled until I can work out the best way to apply the results.

Effect Types can choose how they respond to being culled. Wether to deactivate or deactivate immediate and also whether to stay in the manager and be reactivated when they pass the cull checks in future.
The frequency of cull checks can also be controled between spawn only, low, med, high  (1, .5 and .25s respectively) or and every frame.
Cull checks are performed on activation.
SpawnSystemAtLocation/Attached perform cull checks before ever even creating a component and so should incur near zero cost.

TODO:
- Incorporate spawn rate scaling.
- Per System Overrides could use some UI TLC via a details customization.
- Add a Paused cull state? Pausing is a bit heavy so I'm on the fence here.
- Enable and apply runtime perf data.

diff --git a/Engine/Plugins/FX/Niagara/Source/Niagara/Classes/NiagaraConstants.h b/Engine/Plugins/FX/Niagara/Source/Niagara/Classes/NiagaraConstants.h
--- a/Engine/Plugins/FX/Niagara/Source/Niagara/Classes/NiagaraConstants.h
+++ b/Engine/Plugins/FX/Niagara/Source/Niagara/Classes/NiagaraConstants.h
@@ -41,8 +41,9 @@
 #define SYS_PARAM_ENGINE_WORLD_TO_LOCAL_NO_SCALE         INiagaraModule::GetVar_Engine_Owner_SystemWorldToLocalNoScale()
 
 
-#define SYS_PARAM_ENGINE_TIME_SINCE_RENDERED		INiagaraModule::GetVar_Engine_Owner_TimeSinceRendered()
-#define SYS_PARAM_ENGINE_LOD_DISTANCE				INiagaraModule::GetVar_Engine_Owner_LODDistance()
+#define SYS_PARAM_ENGINE_TIME_SINCE_RENDERED			INiagaraModule::GetVar_Engine_Owner_TimeSinceRendered()
+#define SYS_PARAM_ENGINE_LOD_DISTANCE					INiagaraModule::GetVar_Engine_Owner_LODDistance()
+#define SYS_PARAM_ENGINE_LOD_DISTANCE_FRACTION			INiagaraModule::GetVar_Engine_Owner_LODDistanceFraction()
 
 
 #define SYS_PARAM_ENGINE_EXECUTION_STATE                 INiagaraModule::GetVar_Engine_Owner_ExecutionState()
diff --git a/Engine/Plugins/FX/Niagara/Source/Niagara/Classes/NiagaraEffectType.h b/Engine/Plugins/FX/Niagara/Source/Niagara/Classes/NiagaraEffectType.h
new file mode 100644
--- /dev/null
+++ b/Engine/Plugins/FX/Niagara/Source/Niagara/Classes/NiagaraEffectType.h
@@ -0,0 +1,185 @@
+// Copyright 1998-2019 Epic Games, Inc. All Rights Reserved.
+
+#pragma once
+
+#include "InGamePerformanceTracker.h"
+#include "RenderCommandFence.h"
+#include "NiagaraEffectType.generated.h"
+
+#define DEBUG_SCALABILITY_STATE (!UE_BUILD_SHIPPING)
+
+/** Controls what action is taken by a Niagara system that fails it's cull checks. */
+UENUM()
+enum class ENiagaraCullReaction
+{
+	/** The system instance will be deactivated. It will not be reactivated automatically by the scalability system. */
+	Deactivate,
+	/** The system instance will be deactivated and particles killed immediately. It will not be reactivated automatically by the scalability system. */
+	DeactivateImmediate,
+	/** The system instance will be deactivated. Will reactivate when it passes cull tests again. */
+	DeactivateResume,
+	/** The system instance will be deactivated and particles killed immediately. Will reactivate when it passes cull tests again. */
+	DeactivateImmediateResume,
+	/** The system instance will be paused but will resume ticking when it passes cull tests again. */
+	//PauseResume,
+};
+
+/** Controls how often should we update scalability states for these effects. */
+UENUM()
+enum class ENiagaraScalabilityUpdateFrequency
+{
+	/** Scalability will be checked only on spawn. */
+	SpawnOnly,
+	/** Scalability will be checked infrequently.*/
+	Low,
+	/** Scalability will be occasionally. */
+	Medium,
+	/** Scalability will be checked regularly. */
+	High,
+	/** Scalability will be checked every frame. */
+	Continuous,
+};
+
+/** Settings which control high level scalability of Niagara effects. Effects culled at this level stop ticking and where possible are not spawned at all. Costing the absolute minimum at runtime. */
+USTRUCT()
+struct FNiagaraScalabilitySettings
+{
+	GENERATED_USTRUCT_BODY()
+
+	UPROPERTY(EditAnywhere, Category = "Scalability", meta = (InlineEditConditionToggle))
+	uint32 bCullByDistance : 1;
+	UPROPERTY(EditAnywhere, Category = "Scalability", meta = (InlineEditConditionToggle))
+	uint32 bCullMaxInstanceCount : 1;
+	UPROPERTY(EditAnywhere, Category = "Scalability", meta = (InlineEditConditionToggle))
+	uint32 bCullByMaxOwnerLOD : 1;
+	UPROPERTY(EditAnywhere, Category = "Scalability", meta = (InlineEditConditionToggle))
+	uint32 bCullByMaxTimeWithoutRender : 1;
+
+	/** Effects of this type are culled beyond this distance. */
+	UPROPERTY(EditAnywhere, Category = "Scalability", meta = (EditCondition = "bCullByDistance"))
+	float MaxDistance;
+
+	/** Effects of this type will fail to spawn when total active instance count exceeds this number. */
+	UPROPERTY(EditAnywhere, Category = "Scalability", meta = (EditCondition = "bCullMaxInstanceCount"))
+	float MaxInstances;
+
+	/** The effect is culled beyond this LOD level. Useful for FX that are linked to a specific actor with an LOD level driven be some other system. */
+	UPROPERTY(EditAnywhere, Category = "Scalability", meta = (EditCondition = "bCullByMaxOwnerLOD"))
+	int32 MaxOwnerLOD;
+
+	//TODO:
+	/** The effect is culled when it's bounds take up less that this fraction of the total screen area. Only usable with fixed bounds. */
+	//float ScreenFraction;
+
+	/** Effects will be culled if they go more than this length of time without being rendered. */
+	UPROPERTY(EditAnywhere, Category = "Scalability", meta = (EditCondition = "bCullByMaxTimeWithoutRender"))
+	float MaxTimeWithoutRender;
+
+	/** Scale factor applied to all spawn counts for effects of this type. */
+	//UPROPERTY(EditAnywhere, Category = "Scalability", meta = (EditCondition = "bCullByMaxTimeWithoutRender"))
+	//float SpawnCountScale;
+
+	//TODO: Assess viability of progressive scaling via biasing the DetailLevel that effects use and reducing this with significance etc.
+	/** Settings that control whether we should bias the DetailLevel used by effects of this type. */
+	//TArray<FNiagaraDetailLevelBiasSettings> DetailLevelBiasSettings;
+};
+
+USTRUCT()
+struct FNiagaraScalabilityOverrides : public FNiagaraScalabilitySettings
+{
+	GENERATED_USTRUCT_BODY()
+		
+	//TODO: Detail customization that effectively allows these values to be edit conditions for their respective properties inside FNiagaraScalabilitySettings.
+	UPROPERTY(EditAnywhere, Category = "Override")
+	uint32 bOverrideDistanceSettings : 1;
+	UPROPERTY(EditAnywhere, Category = "Override")
+	uint32 bOverrideInstanceCountSettings : 1;
+	UPROPERTY(EditAnywhere, Category = "Override")
+	uint32 bOverrideOwnerLODSettings : 1;
+	UPROPERTY(EditAnywhere, Category = "Override")
+	uint32 bOverrideTimeSinceRendererSettings : 1;
+	UPROPERTY(EditAnywhere, Category = "Override")
+	uint32 bOverrideSpawnCountScale : 1;
+};
+
+/** Contains settings and working data shared among many NiagaraSystems that share some commonality of type. For example ImpactFX vs EnvironmentalFX. */
+UCLASS()
+class NIAGARA_API UNiagaraEffectType : public UObject
+{
+	GENERATED_UCLASS_BODY()
+
+	//UObject Interface
+	virtual void BeginDestroy()override;
+	virtual bool IsReadyForFinishDestroy()override;
+	virtual void PostLoad()override;
+#if WITH_EDITOR
+	virtual void PostEditChangeProperty(struct FPropertyChangedEvent& PropertyChangedEvent);
+#endif
+	//UObject Interface END
+
+	/** How regularly effects of this type are checked for scalability. */
+	UPROPERTY(EditAnywhere, Category = "Effect Type")
+	ENiagaraScalabilityUpdateFrequency UpdateFrequency;
+
+	/** How effects of this type react when they fail the cull checks. */
+	UPROPERTY(EditAnywhere, Category = "Effect Type")
+	ENiagaraCullReaction CullReaction;
+
+	/** Cull settings to use at each detail level. */
+	UPROPERTY(EditAnywhere, Category = "Effect Type")
+	TArray<FNiagaraScalabilitySettings> DetailLevelScalabilitySettings;
+	
+	FORCEINLINE int32* GetCycleCounter(bool bGameThread, bool bConcurrent);
+	void ProcessLastFrameCycleCounts();
+
+	//TODO: Dynamic budgetting from perf data.
+	//void ApplyDynamicBudget(float InDynamicBudget_GT, float InDynamicBudget_GT_CNC, float InDynamicBudget_RT);
+
+	void SetDetailLevel(int32 DetailLevel);
+
+	FORCEINLINE const TArray<FNiagaraScalabilitySettings>& GetScalabilitySettings()const { return DetailLevelScalabilitySettings; }
+
+	float GetAverageFrameTime_GT() { return AvgTimeMS_GT; }
+	float GetAverageFrameTime_GT_CNC() { return AvgTimeMS_GT_CNC; }
+	float GetAverageFrameTime_RT() { return AvgTimeMS_RT; }
+
+	/** Total number of instances across all systems for this effect type. */
+	uint32 NumInstances;
+
+private:
+	float AvgTimeMS_GT;
+	float AvgTimeMS_GT_CNC;
+	float AvgTimeMS_RT;
+
+	//TODO: Budgets from runtime perf.
+	//The result of runtime perf calcs and dynamic budget is a bias to the minimum significance required for FX of ths type.
+	//float MinSignificanceFromPerf;
+
+	FInGameCycleHistory CycleHistory_GT;
+	FInGameCycleHistory CycleHistory_GT_CNC;
+	FInGameCycleHistory CycleHistory_RT;
+
+	//Number of frames since we last sampled perf. We need not sample runtime perf every frame to get usable data.
+	int32 FramesSincePerfSampled;
+	bool bSampleRunTimePerfThisFrame;
+
+	/** Fence used to guarantee that the RT is finished using our cycle counters in the case we're gathering RT cycle counts. */
+	FRenderCommandFence ReleaseFence;
+};
+
+FORCEINLINE int32* UNiagaraEffectType::GetCycleCounter(bool bGameThread, bool bConcurrent)
+{
+	if (bSampleRunTimePerfThisFrame)
+	{
+		if (bGameThread)
+		{
+			return bConcurrent ? &CycleHistory_GT_CNC.CurrFrameCycles : &CycleHistory_GT.CurrFrameCycles;
+		}
+		else
+		{
+			//Just use the one for RT. Can split later if we'd like. We currently don't have any RT task work anyway.s
+			return &CycleHistory_RT.CurrFrameCycles;
+		}
+	}
+	return nullptr;
+}
\ No newline at end of file
diff --git a/Engine/Plugins/FX/Niagara/Source/Niagara/Classes/NiagaraSystem.h b/Engine/Plugins/FX/Niagara/Source/Niagara/Classes/NiagaraSystem.h
--- a/Engine/Plugins/FX/Niagara/Source/Niagara/Classes/NiagaraSystem.h
+++ b/Engine/Plugins/FX/Niagara/Source/Niagara/Classes/NiagaraSystem.h
@@ -13,6 +13,8 @@
 #include "NiagaraParameterCollection.h"
 #include "NiagaraUserRedirectionParameterStore.h"
 #include "NiagaraSystemFastPath.h"
+#include "NiagaraEffectType.h"
+
 #include "NiagaraSystem.generated.h"
 
 #if WITH_EDITORONLY_DATA
@@ -124,6 +126,9 @@ public:
 #if WITH_EDITOR
 	DECLARE_MULTICAST_DELEGATE_OneParam(FOnSystemCompiled, UNiagaraSystem*);
 #endif
+	//TestChange
+
+	UNiagaraSystem(FVTableHelper& Helper);
 
 	//~ UObject interface
 	void PostInitProperties();
@@ -132,6 +137,7 @@ public:
 	virtual void BeginDestroy() override;
 	virtual void PreSave(const class ITargetPlatform * TargetPlatform) override;
 #if WITH_EDITOR
+	virtual void PreEditChange(UProperty* PropertyThatWillChange)override;
 	virtual void PostEditChangeProperty(struct FPropertyChangedEvent& PropertyChangedEvent) override; 
 	virtual void BeginCacheForCookedPlatformData(const ITargetPlatform *TargetPlatform) override;
 #endif
@@ -288,6 +294,14 @@ public:
 
 	FBox GetFixedBounds() const;
 
+	FORCEINLINE int32* GetCycleCounter(bool bGameThread, bool bConcurrent);
+
+	UNiagaraEffectType* GetEffectType()const;
+	const FNiagaraScalabilitySettings& GetScalabilitySettings(int32 DetailLevel=INDEX_NONE);
+	
+	void ResolveScalabilityOverrides();
+	void OnDetailLevelChanges(int32 DetailLevel);
+
 	/** Whether or not fixed bounds are enabled. */
 	UPROPERTY(EditAnywhere, Category = "System", meta = (InlineEditConditionToggle))
 	uint32 bFixedBounds : 1;
@@ -322,8 +336,18 @@ private:
 #endif
 
 	void UpdatePostCompileDIInfo();
+
 protected:
 
+	UPROPERTY(EditAnywhere, Category = "System")
+	UNiagaraEffectType* EffectType;
+
+	UPROPERTY(EditAnywhere, Category = "System", meta=(InlineEditConditionToggle))
+	bool bOverrideScalabilitySettings;
+
+	UPROPERTY(EditAnywhere, Category = "System", meta = (EditCondition="bOverrideScalabilitySettings"))
+	TArray<FNiagaraScalabilityOverrides> ScalabilityOverrides;
+
 	/** Handles to the emitter this System will simulate. */
 	UPROPERTY()
 	TArray<FNiagaraEmitterHandle> EmitterHandles;
@@ -404,4 +428,20 @@ protected:
 	mutable TStatId StatID_RT;
 	mutable TStatId StatID_RT_CNC;
 #endif
+
+	/** Resolved results of this system's overrides applied on top of it's effect type settings. */
+	UPROPERTY()
+	TArray<FNiagaraScalabilitySettings> ResolvedScalabilitySettings;
+
+	FNiagaraScalabilitySettings CurrentScalabilitySettings;
 };
+
+extern int32 GEnableNiagaraRuntimeCycleCounts;
+FORCEINLINE int32* UNiagaraSystem::GetCycleCounter(bool bGameThread, bool bConcurrent)
+{
+	if (GEnableNiagaraRuntimeCycleCounts && EffectType)
+	{
+		return EffectType->GetCycleCounter(bGameThread, bConcurrent);
+	}
+	return nullptr;
+}
\ No newline at end of file
diff --git a/Engine/Plugins/FX/Niagara/Source/Niagara/Private/NiagaraCommon.cpp b/Engine/Plugins/FX/Niagara/Source/Niagara/Private/NiagaraCommon.cpp
--- a/Engine/Plugins/FX/Niagara/Source/Niagara/Private/NiagaraCommon.cpp
+++ b/Engine/Plugins/FX/Niagara/Source/Niagara/Private/NiagaraCommon.cpp
@@ -171,6 +171,11 @@ void FNiagaraSystemUpdateContext::Add(const UNiagaraParameterCollection* Collect
 
 void FNiagaraSystemUpdateContext::AddInternal(UNiagaraComponent* Comp, bool bReInit)
 {
+	if (bDestroyOnAdd)
+	{
+		Comp->DeactivateImmediate();
+	}
+
 	if (bReInit)
 	{
 		ComponentsToReInit.AddUnique(Comp);
diff --git a/Engine/Plugins/FX/Niagara/Source/Niagara/Private/NiagaraComponent.cpp b/Engine/Plugins/FX/Niagara/Source/Niagara/Private/NiagaraComponent.cpp
--- a/Engine/Plugins/FX/Niagara/Source/Niagara/Private/NiagaraComponent.cpp
+++ b/Engine/Plugins/FX/Niagara/Source/Niagara/Private/NiagaraComponent.cpp
@@ -113,6 +113,7 @@ FAutoConsoleCommandWithWorld DumpNiagaraComponentsCommand(
 FNiagaraSceneProxy::FNiagaraSceneProxy(const UNiagaraComponent* InComponent)
 		: FPrimitiveSceneProxy(InComponent, InComponent->GetAsset() ? InComponent->GetAsset()->GetFName() : FName())
 		, bRenderingEnabled(true)
+		, RuntimeCycleCount(nullptr)
 {
 	// In this case only, update the System renderers on the game thread.
 	check(IsInGameThread());
@@ -129,6 +130,7 @@ FNiagaraSceneProxy::FNiagaraSceneProxy(const UNiagaraComponent* InComponent)
 		SystemStatID = InComponent->GetAsset()->GetStatID(false, false);
 #endif
 
+		RuntimeCycleCount = InComponent->GetAsset()->GetCycleCounter(false, false);
 	}
 }
 
@@ -319,6 +321,8 @@ void FNiagaraSceneProxy::GetDynamicMeshElements(const TArray<const FSceneView*>&
 	FScopeCycleCounter SystemStatCounter(SystemStatID);
 #endif
 
+	FNiagaraScopedRuntimeCycleCounter RuntimeScope(RuntimeCycleCount);
+
 	for (int32 RendererIdx : RendererDrawOrder)
 	{
 		FNiagaraRenderer* Renderer = EmitterRenderers[RendererIdx];
@@ -347,6 +351,7 @@ void FNiagaraSceneProxy::GetDynamicMeshElements(const TArray<const FSceneView*>&
 #if RHI_RAYTRACING
 void FNiagaraSceneProxy::GetDynamicRayTracingInstances(FRayTracingMaterialGatheringContext& Context, TArray<FRayTracingInstance>& OutRayTracingInstances)
 {
+	FNiagaraScopedRuntimeCycleCounter RuntimeScope(RuntimeCycleCount);
 	for (FNiagaraRenderer* Renderer : EmitterRenderers)
 	{
 		if (Renderer)
@@ -359,6 +364,7 @@ void FNiagaraSceneProxy::GetDynamicRayTracingInstances(FRayTracingMaterialGather
 
 void FNiagaraSceneProxy::GatherSimpleLights(const FSceneViewFamily& ViewFamily, FSimpleLightArray& OutParticleLights) const
 {
+	FNiagaraScopedRuntimeCycleCounter RuntimeScope(RuntimeCycleCount);
 	for (int32 Idx = 0; Idx < EmitterRenderers.Num(); Idx++)
 	{
 		FNiagaraRenderer *Renderer = EmitterRenderers[Idx];
@@ -391,7 +397,12 @@ UNiagaraComponent::UNiagaraComponent(const FObjectInitializer& ObjectInitializer
 	, bWaitForCompilationOnActivate(false)
 #endif
 	, bAwaitingActivationDueToNotReady(false)
+	, bActivateShouldResetWhenReady(false)
+	, bDidAutoAttach(false)
+	, bAllowScalability(true)
 	//, bIsChangingAutoAttachment(false)
+	, ScalabilityManagerHandle(INDEX_NONE)
+	, OwnerLOD(0)
 {
 	PrimaryComponentTick.bCanEverTick = true;
 	PrimaryComponentTick.TickGroup = GNiagaraSoloTickEarly ? TG_PrePhysics : TG_DuringPhysics;
@@ -701,6 +712,11 @@ bool UNiagaraComponent::InitializeSystem()
 }
 
 void UNiagaraComponent::Activate(bool bReset /* = false */)
+{
+	ActivateInternal(bReset, false);
+}
+
+void UNiagaraComponent::ActivateInternal(bool bReset /* = false */, bool bIsScalabilityCull)
 {
 	bAwaitingActivationDueToNotReady = false;
 
@@ -766,7 +782,24 @@ void UNiagaraComponent::Activate(bool bReset /* = false */)
 		return;
 	}
 
-	
+	if (bReset)
+	{
+		UnregisterWithScalabilityManager();
+	}
+
+	if (!bIsScalabilityCull && ScalabilityManagerHandle != INDEX_NONE)
+	{
+		//If this is a non scalability activate call and we're still registered with the manager.
+		//If we reach this point then we must have been previously culled by scalability so bail here.
+		return;
+	}
+
+	if (RegisterWithScalabilityManagerOrPreCull())
+	{
+		//We have decided to pre cull the system.
+		return;
+	}
+
 	Super::Activate(bReset);
 
 	// Early out if we're not forcing a reset, and both the component and system instance are already active.
@@ -829,6 +862,11 @@ void UNiagaraComponent::Activate(bool bReset /* = false */)
 }
 
 void UNiagaraComponent::Deactivate()
+{
+	DeactivateInternal(false);
+}
+
+void UNiagaraComponent::DeactivateInternal(bool bIsScalabilityCull /* = false */)
 {
 	SCOPE_CYCLE_COUNTER(STAT_NiagaraComponentDeactivate);
 	CSV_SCOPED_TIMING_STAT_EXCLUSIVE(Niagara);
@@ -844,6 +882,13 @@ void UNiagaraComponent::Deactivate()
 			Super::Deactivate();
 		}
 
+		//Unregister with the scalability manager if this is a genuine deactivation from outside.
+		//The scalability manager itself can call this function when culling systems.
+		if (bIsScalabilityCull == false)
+		{
+			UnregisterWithScalabilityManager();
+		}
+
 		SystemInstance->Deactivate();
 
 		// We are considered active until we are complete
@@ -857,6 +902,11 @@ void UNiagaraComponent::Deactivate()
 }
 
 void UNiagaraComponent::DeactivateImmediate()
+{
+	DeactivateImmediateInternal(false);
+}
+
+void UNiagaraComponent::DeactivateImmediateInternal(bool bIsScalabilityCull)
 {
 	SCOPE_CYCLE_COUNTER(STAT_NiagaraComponentDeactivate);
 	Super::Deactivate();
@@ -865,6 +915,13 @@ void UNiagaraComponent::DeactivateImmediate()
 
 	//UE_LOG(LogNiagara, Log, TEXT("Deactivate %s"), *GetName());
 
+	//Unregister with the scalability manager if this is a genuine deactivation from outside.
+	//The scalability manager itself can call this function when culling systems.
+	if (bIsScalabilityCull == false)
+	{
+		UnregisterWithScalabilityManager();
+	}
+
 	SetActiveFlag(false);
 
 	if (SystemInstance)
@@ -873,6 +930,44 @@ void UNiagaraComponent::DeactivateImmediate()
 	}
 }
 
+bool UNiagaraComponent::RegisterWithScalabilityManagerOrPreCull()
+{
+	if (ScalabilityManagerHandle == INDEX_NONE && bAllowScalability)
+	{
+		if (UNiagaraSystem* System = GetAsset())
+		{
+			if (UNiagaraEffectType* EffectType = System->GetEffectType())
+			{
+				if (FNiagaraWorldManager* WorldMan = FNiagaraWorldManager::Get(GetWorld()))
+				{
+					if (EffectType->UpdateFrequency == ENiagaraScalabilityUpdateFrequency::SpawnOnly)
+					{
+						//If we're just set to check on spawn then check for precull here.
+						return WorldMan->ShouldPreCull(GetAsset(), this);
+					}
+					else
+					{
+						WorldMan->RegisterWithScalabilityManager(this);
+					}
+				}
+			}
+		}
+	}
+
+	return false;
+}
+
+void UNiagaraComponent::UnregisterWithScalabilityManager()
+{
+	if (ScalabilityManagerHandle != INDEX_NONE)
+	{
+		if (FNiagaraWorldManager* WorldMan = FNiagaraWorldManager::Get(GetWorld()))
+		{
+			WorldMan->UnregisterWithScalabilityManager(this);
+		}
+	}
+}
+
 void UNiagaraComponent::OnSystemComplete()
 {
 	//UE_LOG(LogNiagara, Log, TEXT("OnSystemComplete: %p - %s"), SystemInstance.Get(), *Asset->GetName());
@@ -1046,6 +1141,9 @@ void UNiagaraComponent::SendRenderDynamicData_Concurrent()
 		TStatId SystemStatID = GetAsset() ? GetAsset()->GetStatID(true, true) : TStatId();
 		FScopeCycleCounter SystemStatCounter(SystemStatID);
 #endif
+
+		FNiagaraScopedRuntimeCycleCounter RuntimeScope(GetAsset(), true, false);
+
 		FNiagaraSceneProxy* NiagaraProxy = static_cast<FNiagaraSceneProxy*>(SceneProxy);
 		const TArray<FNiagaraRenderer*>& EmitterRenderers = NiagaraProxy->GetEmitterRenderers();
 
@@ -1625,6 +1723,7 @@ void UNiagaraComponent::SetMaxSimTime(float InMaxTime)
 	MaxSimTime = InMaxTime;
 }
 
+#if WITH_NIAGARA_COMPONENT_PREVIEW_DATA
 void UNiagaraComponent::SetPreviewDetailLevel(bool bInEnablePreviewDetailLevel, int32 InPreviewDetailLevel)
 {
 	bool bReInit = bEnablePreviewDetailLevel != bInEnablePreviewDetailLevel || (bEnablePreviewDetailLevel && PreviewDetailLevel != InPreviewDetailLevel);
@@ -1642,6 +1741,15 @@ void UNiagaraComponent::SetPreviewLODDistance(bool bInEnablePreviewLODDistance,
 	bEnablePreviewLODDistance = bInEnablePreviewLODDistance;
 	PreviewLODDistance = InPreviewLODDistance;
 }
+#else
+void UNiagaraComponent::SetPreviewDetailLevel(bool bInEnablePreviewDetailLevel, int32 InPreviewDetailLevel){}
+void UNiagaraComponent::SetPreviewLODDistance(bool bInEnablePreviewLODDistance, float InPreviewLODDistance){}
+#endif
+
+void UNiagaraComponent::SetOwnerLOD(int32 InOwnerLOD)
+{
+	OwnerLOD = InOwnerLOD;  
+}
 
 #if WITH_EDITOR
 
diff --git a/Engine/Plugins/FX/Niagara/Source/Niagara/Private/NiagaraConstants.cpp b/Engine/Plugins/FX/Niagara/Source/Niagara/Private/NiagaraConstants.cpp
--- a/Engine/Plugins/FX/Niagara/Source/Niagara/Private/NiagaraConstants.cpp
+++ b/Engine/Plugins/FX/Niagara/Source/Niagara/Private/NiagaraConstants.cpp
@@ -46,6 +46,7 @@ void FNiagaraConstants::Init()
 		SystemParameters.Add(SYS_PARAM_ENGINE_WORLD_TO_LOCAL_NO_SCALE);
 
 		SystemParameters.Add(SYS_PARAM_ENGINE_LOD_DISTANCE);
+		SystemParameters.Add(SYS_PARAM_ENGINE_LOD_DISTANCE_FRACTION);
 		SystemParameters.Add(SYS_PARAM_ENGINE_TIME_SINCE_RENDERED);
 
 		SystemParameters.Add(SYS_PARAM_ENGINE_EXECUTION_STATE);
@@ -160,6 +161,7 @@ void FNiagaraConstants::Init()
 
 		SystemStrMap.Add(SYS_PARAM_ENGINE_TIME_SINCE_RENDERED, LOCTEXT("TimeSinceRendered", "The time in seconds that have passed since this system was last rendered."));
 		SystemStrMap.Add(SYS_PARAM_ENGINE_LOD_DISTANCE, LOCTEXT("LODDistance", "The distance to use in LOD scaling/culling. Typically, the distance to the nearest camera."));
+		SystemStrMap.Add(SYS_PARAM_ENGINE_LOD_DISTANCE_FRACTION, LOCTEXT("LODDistanceFraction", "The distance fraction between this system and it's max culling distance defined in it's EffectType scalabiltiy settings."));
 		SystemStrMap.Add(SYS_PARAM_ENGINE_EXECUTION_STATE, LOCTEXT("ExecutionState", "The execution state of the systems owner. Takes precedence over the systems internal execution state."));
 
 		SystemStrMap.Add(SYS_PARAM_ENGINE_EXEC_COUNT, LOCTEXT("ExecCountDesc", "The index of this particle in the read buffer."));
diff --git a/Engine/Plugins/FX/Niagara/Source/Niagara/Private/NiagaraEffectType.cpp b/Engine/Plugins/FX/Niagara/Source/Niagara/Private/NiagaraEffectType.cpp
new file mode 100644
--- /dev/null
+++ b/Engine/Plugins/FX/Niagara/Source/Niagara/Private/NiagaraEffectType.cpp
@@ -0,0 +1,102 @@
+// Copyright 1998-2019 Epic Games, Inc. All Rights Reserved.
+#include "NiagaraEffectType.h"
+#include "NiagaraModule.h"
+#include "NiagaraCommon.h"
+#include "NiagaraSystem.h"
+
+//In an effort to cut the impact of runtime perf tracking, I limit the number of fames we actually sample on.
+int32 GNumFramesBetweenRuntimePerfSamples = 5;
+static FAutoConsoleVariableRef CVarNumFramesBetweenRuntimePerfSamples(TEXT("fx.NumFramesBetweenRuntimePerfSamples"), GNumFramesBetweenRuntimePerfSamples, TEXT("How many frames between each sample of Niagara runtime perf. \n"), ECVF_ReadOnly);
+
+int32 GNiagaraRuntimeCycleHistorySize = 15;
+static FAutoConsoleVariableRef CVarNiagaraRuntimeCycleHistorySize(TEXT("fx.NiagaraRuntimeCycleHistorySize"), GNiagaraRuntimeCycleHistorySize, TEXT("How many frames history to use in Niagara's runtime performance trackers. \n"), ECVF_ReadOnly);
+
+UNiagaraEffectType::UNiagaraEffectType(const FObjectInitializer& ObjectInitializer)
+	: Super(ObjectInitializer)
+	, UpdateFrequency(ENiagaraScalabilityUpdateFrequency::SpawnOnly)
+	, CullReaction(ENiagaraCullReaction::DeactivateImmediate)
+	, NumInstances(0)
+	, AvgTimeMS_GT(0.0f)
+	, AvgTimeMS_GT_CNC(0.0f)
+	, AvgTimeMS_RT(0.0f)
+	, CycleHistory_GT(GNiagaraRuntimeCycleHistorySize)
+	, CycleHistory_GT_CNC(GNiagaraRuntimeCycleHistorySize)
+	, CycleHistory_RT(GNiagaraRuntimeCycleHistorySize)
+	, FramesSincePerfSampled(0)
+	, bSampleRunTimePerfThisFrame(false)
+{
+}
+
+void UNiagaraEffectType::BeginDestroy()
+{
+	Super::BeginDestroy();
+
+	ReleaseFence.BeginFence();
+}
+
+bool UNiagaraEffectType::IsReadyForFinishDestroy()
+{
+	return ReleaseFence.IsFenceComplete() && Super::IsReadyForFinishDestroy();
+}
+
+void UNiagaraEffectType::PostLoad()
+{
+	Super::PostLoad();
+}
+
+#if WITH_EDITOR
+
+void UNiagaraEffectType::PostEditChangeProperty(struct FPropertyChangedEvent& PropertyChangedEvent)
+{
+	Super::PostEditChangeProperty(PropertyChangedEvent);
+
+	FNiagaraSystemUpdateContext UpdateContext;
+	for (TObjectIterator<UNiagaraSystem> It; It; ++It)
+	{
+		UNiagaraSystem* System = *It;
+		if (System->GetEffectType() == this)
+		{
+			System->ResolveScalabilityOverrides();
+			UpdateContext.Add(System, true);
+		}
+	}
+}
+#endif
+
+void UNiagaraEffectType::ProcessLastFrameCycleCounts()
+{
+	if (FramesSincePerfSampled > GNumFramesBetweenRuntimePerfSamples)
+	{ 
+		FramesSincePerfSampled = 0;
+		bSampleRunTimePerfThisFrame = true;
+	}
+	else
+	{
+		++FramesSincePerfSampled;
+		bSampleRunTimePerfThisFrame = false;
+	}
+
+	CycleHistory_GT.NextFrame();
+	CycleHistory_GT_CNC.NextFrame();
+	CycleHistory_RT.NextFrame();
+
+	AvgTimeMS_GT = FPlatformTime::ToMilliseconds(CycleHistory_GT.GetAverageCycles());
+	AvgTimeMS_GT_CNC = FPlatformTime::ToMilliseconds(CycleHistory_GT_CNC.GetAverageCycles());
+	AvgTimeMS_RT = FPlatformTime::ToMilliseconds(CycleHistory_RT.GetAverageCycles());
+}
+
+// void UNiagaraEffectType::ApplyDynamicBudget(float InDynamicBudget_GT, float InDynamicBudget_GT_CNC, float InDynamicBudget_RT)
+// {
+// 	if (bApplyDynamicBudgetsToSignificance)
+// 	{
+// 		float MinSigificance_GTPerf = (InDynamicBudget_GT > SMALL_NUMBER) ? (AvgTimeMS_GT / InDynamicBudget_GT) : 1.0f;
+// 		float MinSigificance_GTTotalPerf = (InDynamicBudget_GT_CNC > SMALL_NUMBER) ? (AvgTimeMS_GT + AvgTimeMS_GT_CNC) / InDynamicBudget_GT_CNC : 1.0f;
+// 		float MinSignificance_RTPerf = (InDynamicBudget_RT > SMALL_NUMBER) ? AvgTimeMS_GT_CNC / InDynamicBudget_RT : 1.0f;
+// 
+// 		MinSignificanceFromPerf = FMath::Max3(MinSigificance_GTPerf, MinSigificance_GTTotalPerf, MinSignificance_RTPerf);
+// 	}
+// 	else
+// 	{
+// 		MinSignificanceFromPerf = 0.0f;
+// 	}
+// }
diff --git a/Engine/Plugins/FX/Niagara/Source/Niagara/Private/NiagaraFunctionLibrary.cpp b/Engine/Plugins/FX/Niagara/Source/Niagara/Private/NiagaraFunctionLibrary.cpp
--- a/Engine/Plugins/FX/Niagara/Source/Niagara/Private/NiagaraFunctionLibrary.cpp
+++ b/Engine/Plugins/FX/Niagara/Source/Niagara/Private/NiagaraFunctionLibrary.cpp
@@ -70,7 +70,7 @@ UNiagaraComponent* CreateNiagaraSystem(UNiagaraSystem* SystemTemplate, UWorld* W
 * Spawns a Niagara System at the specified world location/rotation
 * @return			The spawned UNiagaraComponent
 */
-UNiagaraComponent* UNiagaraFunctionLibrary::SpawnSystemAtLocation(const UObject* WorldContextObject, UNiagaraSystem* SystemTemplate, FVector SpawnLocation, FRotator SpawnRotation, FVector Scale, bool bAutoDestroy, bool bAutoActivate, ENCPoolMethod PoolingMethod)
+UNiagaraComponent* UNiagaraFunctionLibrary::SpawnSystemAtLocation(const UObject* WorldContextObject, UNiagaraSystem* SystemTemplate, FVector SpawnLocation, FRotator SpawnRotation, FVector Scale, bool bAutoDestroy, bool bAutoActivate, ENCPoolMethod PoolingMethod, bool bPreCullCheck)
 {
 	UNiagaraComponent* PSC = NULL;
 	if (SystemTemplate)
@@ -78,18 +78,28 @@ UNiagaraComponent* UNiagaraFunctionLibrary::SpawnSystemAtLocation(const UObject*
 		UWorld* World = GEngine->GetWorldFromContextObject(WorldContextObject, EGetWorldErrorMode::LogAndReturnNull);
 		if (World != nullptr)
 		{
-			PSC = CreateNiagaraSystem(SystemTemplate, World, World->GetWorldSettings(), bAutoDestroy, PoolingMethod);
+			bool bShouldCull = false;
+			if (bPreCullCheck)
+			{
+				FNiagaraWorldManager* WorldManager = FNiagaraWorldManager::Get(World);
+				bShouldCull = WorldManager->ShouldPreCull(SystemTemplate, SpawnLocation);
+			}
+
+			if (!bShouldCull)
+			{
+				PSC = CreateNiagaraSystem(SystemTemplate, World, World->GetWorldSettings(), bAutoDestroy, PoolingMethod);
 #if WITH_EDITORONLY_DATA
-			PSC->bWaitForCompilationOnActivate = true;
+				PSC->bWaitForCompilationOnActivate = true;
 #endif
-			PSC->RegisterComponentWithWorld(World);
+				PSC->RegisterComponentWithWorld(World);
 
-			PSC->SetAbsolute(true, true, true);
-			PSC->SetWorldLocationAndRotation(SpawnLocation, SpawnRotation);
-			PSC->SetRelativeScale3D(Scale);
-			if (bAutoActivate)
-			{
-				PSC->Activate(true);
+				PSC->SetAbsolute(true, true, true);
+				PSC->SetWorldLocationAndRotation(SpawnLocation, SpawnRotation);
+				PSC->SetRelativeScale3D(Scale);
+				if (bAutoActivate)
+				{
+					PSC->Activate(true);
+				}
 			}
 		}
 	}
@@ -104,7 +114,7 @@ UNiagaraComponent* UNiagaraFunctionLibrary::SpawnSystemAtLocation(const UObject*
 * Spawns a Niagara System attached to a component
 * @return			The spawned UNiagaraComponent
 */
-UNiagaraComponent* UNiagaraFunctionLibrary::SpawnSystemAttached(UNiagaraSystem* SystemTemplate, USceneComponent* AttachToComponent, FName AttachPointName, FVector Location, FRotator Rotation, EAttachLocation::Type LocationType, bool bAutoDestroy, bool bAutoActivate, ENCPoolMethod PoolingMethod)
+UNiagaraComponent* UNiagaraFunctionLibrary::SpawnSystemAttached(UNiagaraSystem* SystemTemplate, USceneComponent* AttachToComponent, FName AttachPointName, FVector Location, FRotator Rotation, EAttachLocation::Type LocationType, bool bAutoDestroy, bool bAutoActivate, ENCPoolMethod PoolingMethod, bool bPreCullCheck)
 {
 	UNiagaraComponent* PSC = nullptr;
 	if (SystemTemplate)
@@ -115,29 +125,40 @@ UNiagaraComponent* UNiagaraFunctionLibrary::SpawnSystemAttached(UNiagaraSystem*
 		}
 		else
 		{
-			PSC = CreateNiagaraSystem(SystemTemplate, AttachToComponent->GetWorld(), AttachToComponent->GetOwner(), bAutoDestroy, PoolingMethod);
-#if WITH_EDITOR
-			if (GForceNiagaraSpawnAttachedSolo > 0)
+			bool bShouldCull = false;
+			if (bPreCullCheck)
 			{
-				PSC->SetForceSolo(true);
+				FNiagaraWorldManager* WorldManager = FNiagaraWorldManager::Get(AttachToComponent->GetWorld());
+				//TODO: For now using the attach parent location and ignoring the emitters relative location which is clearly going to be a bit wrong in some cases.
+				bShouldCull = WorldManager->ShouldPreCull(SystemTemplate, AttachToComponent->GetComponentLocation());
 			}
-#endif
-			PSC->RegisterComponentWithWorld(AttachToComponent->GetWorld());
 
-			PSC->AttachToComponent(AttachToComponent, FAttachmentTransformRules::KeepRelativeTransform, AttachPointName);
-			if (LocationType == EAttachLocation::KeepWorldPosition)
-			{
-				PSC->SetWorldLocationAndRotation(Location, Rotation);
-			}
-			else
+			if (!bShouldCull)
 			{
-				PSC->SetRelativeLocationAndRotation(Location, Rotation);
-			}
-			PSC->SetRelativeScale3D(FVector(1.f));
+				PSC = CreateNiagaraSystem(SystemTemplate, AttachToComponent->GetWorld(), AttachToComponent->GetOwner(), bAutoDestroy, PoolingMethod);
+#if WITH_EDITOR
+				if (GForceNiagaraSpawnAttachedSolo > 0)
+				{
+					PSC->SetForceSolo(true);
+				}
+#endif
+				PSC->RegisterComponentWithWorld(AttachToComponent->GetWorld());
 
-			if (bAutoActivate)
-			{
-				PSC->Activate();
+				PSC->AttachToComponent(AttachToComponent, FAttachmentTransformRules::KeepRelativeTransform, AttachPointName);
+				if (LocationType == EAttachLocation::KeepWorldPosition)
+				{
+					PSC->SetWorldLocationAndRotation(Location, Rotation);
+				}
+				else
+				{
+					PSC->SetRelativeLocationAndRotation(Location, Rotation);
+				}
+				PSC->SetRelativeScale3D(FVector(1.f));
+
+				if (bAutoActivate)
+				{
+					PSC->Activate();
+				}
 			}
 		}
 	}
@@ -159,7 +180,8 @@ UNiagaraComponent* UNiagaraFunctionLibrary::SpawnSystemAttached(
 	EAttachLocation::Type LocationType,
 	bool bAutoDestroy,
 	ENCPoolMethod PoolingMethod,
-	bool bAutoActivate
+	bool bAutoActivate,
+	bool bPreCullCheck
 )
 {
 	UNiagaraComponent* PSC = nullptr;
@@ -174,7 +196,17 @@ UNiagaraComponent* UNiagaraFunctionLibrary::SpawnSystemAttached(
 			UWorld* const World = AttachToComponent->GetWorld();
 			if (World && !World->IsNetMode(NM_DedicatedServer))
 			{
-				PSC = CreateNiagaraSystem(SystemTemplate, World, AttachToComponent->GetOwner(), bAutoDestroy, PoolingMethod);
+				bool bShouldCull = false;
+				if (bPreCullCheck)
+				{
+					FNiagaraWorldManager* WorldManager = FNiagaraWorldManager::Get(World);
+					//TODO: For now using the attach parent location and ignoring the emitters relative location which is clearly going to be a bit wrong in some cases.
+					bShouldCull = WorldManager->ShouldPreCull(SystemTemplate, AttachToComponent->GetComponentLocation());
+				}
+
+				if (!bShouldCull)
+				{
+					PSC = CreateNiagaraSystem(SystemTemplate, World, AttachToComponent->GetOwner(), bAutoDestroy, PoolingMethod);
 				if (PSC)
 				{
 #if WITH_EDITOR
@@ -221,6 +253,7 @@ UNiagaraComponent* UNiagaraFunctionLibrary::SpawnSystemAttached(
 					// Notify the texture streamer so that PSC gets managed as a dynamic component.
 					IStreamingManager::Get().NotifyPrimitiveUpdated(PSC);
 				}
+				}
 			}
 		}
 	}
@@ -1532,4 +1565,4 @@ void UNiagaraFunctionLibrary::InitVectorVMFastPathOps()
 	check(VectorVMOps.Num() == VectorVMOpsHLSL.Num());
 }
 
-#undef LOCTEXT_NAMESPACE
+#undef LOCTEXT_NAMESPACE
\ No newline at end of file
diff --git a/Engine/Plugins/FX/Niagara/Source/Niagara/Private/NiagaraModule.cpp b/Engine/Plugins/FX/Niagara/Source/Niagara/Private/NiagaraModule.cpp
--- a/Engine/Plugins/FX/Niagara/Source/Niagara/Private/NiagaraModule.cpp
+++ b/Engine/Plugins/FX/Niagara/Source/Niagara/Private/NiagaraModule.cpp
@@ -111,6 +111,7 @@ FNiagaraVariable INiagaraModule::Engine_Owner_SystemWorldToLocalNoScale;
 
 FNiagaraVariable INiagaraModule::Engine_Owner_TimeSinceRendered;
 FNiagaraVariable INiagaraModule::Engine_Owner_LODDistance;
+FNiagaraVariable INiagaraModule::Engine_Owner_LODDistanceFraction;
 
 FNiagaraVariable INiagaraModule::Engine_Owner_ExecutionState;
 
@@ -216,6 +217,7 @@ void INiagaraModule::StartupModule()
 
 	Engine_Owner_TimeSinceRendered = FNiagaraVariable(FNiagaraTypeDefinition::GetFloatDef(), TEXT("Engine.Owner.TimeSinceRendered"));
 	Engine_Owner_LODDistance = FNiagaraVariable(FNiagaraTypeDefinition::GetFloatDef(), TEXT("Engine.Owner.LODDistance"));
+	Engine_Owner_LODDistanceFraction = FNiagaraVariable(FNiagaraTypeDefinition::GetFloatDef(), TEXT("Engine.Owner.LODDistanceFraction"));
 
 	Engine_Owner_ExecutionState = FNiagaraVariable(FNiagaraTypeDefinition::GetExecutionStateEnum(), TEXT("Engine.Owner.ExecutionState"));
 	
@@ -454,6 +456,7 @@ void INiagaraModule::InitFixedSystemInstanceParameterStore()
 	FixedSystemInstanceParameters.AddParameter(SYS_PARAM_ENGINE_TIME_SINCE_RENDERED, true, false);
 	FixedSystemInstanceParameters.AddParameter(SYS_PARAM_ENGINE_EXECUTION_STATE, true, false);
 	FixedSystemInstanceParameters.AddParameter(SYS_PARAM_ENGINE_LOD_DISTANCE, true, false);
+	FixedSystemInstanceParameters.AddParameter(SYS_PARAM_ENGINE_LOD_DISTANCE_FRACTION, true, false);
 	FixedSystemInstanceParameters.AddParameter(SYS_PARAM_ENGINE_SYSTEM_NUM_EMITTERS, true, false);
 	FixedSystemInstanceParameters.AddParameter(SYS_PARAM_ENGINE_SYSTEM_NUM_EMITTERS_ALIVE, true, false);
 	FixedSystemInstanceParameters.AddParameter(SYS_PARAM_ENGINE_SYSTEM_AGE);
@@ -463,13 +466,18 @@ void INiagaraModule::InitFixedSystemInstanceParameterStore()
 
 void INiagaraModule::OnChangeDetailLevel(class IConsoleVariable* CVar)
 {
-	//Can only change the detail level at runtime on when not cooked.
-#if WITH_EDITORONLY_DATA
 	int32 NewDetailLevel = CVar->GetInt();
 	if (EngineDetailLevel != NewDetailLevel)
 	{
 		EngineDetailLevel = NewDetailLevel;
 
+		for (TObjectIterator<UNiagaraSystem> It; It; ++It)
+		{
+			UNiagaraSystem* System = *It;
+			check(System);
+			System->OnDetailLevelChanges(NewDetailLevel);
+		}
+
 		// If the detail level has changed we have to reset all systems,
 		// and only activate ones that were previously active
 		for (TObjectIterator<UNiagaraComponent> It; It; ++It)
@@ -485,7 +493,6 @@ void INiagaraModule::OnChangeDetailLevel(class IConsoleVariable* CVar)
 			}
 		}
 	}
-#endif
 }
 
 //////////////////////////////////////////////////////////////////////////
diff --git a/Engine/Plugins/FX/Niagara/Source/Niagara/Private/NiagaraPreviewGrid.cpp b/Engine/Plugins/FX/Niagara/Source/Niagara/Private/NiagaraPreviewGrid.cpp
--- a/Engine/Plugins/FX/Niagara/Source/Niagara/Private/NiagaraPreviewGrid.cpp
+++ b/Engine/Plugins/FX/Niagara/Source/Niagara/Private/NiagaraPreviewGrid.cpp
@@ -140,6 +140,9 @@ void ANiagaraPreviewGrid::ActivatePreviews(bool bReset)
 		GeneratePreviews();
 	}
 
+	FString XLabel;
+	FString YLabel;
+
 	for (int32 X = 0; X < NumX; ++X)
 	{
 		float XLocation = X / NumX - 1;
@@ -156,6 +159,9 @@ void ANiagaraPreviewGrid::ActivatePreviews(bool bReset)
 			PreviewActor->GetComponents<UNiagaraComponent>(NiagaraComponents);
 			for (UNiagaraComponent* Component : NiagaraComponents)
 			{
+				PreviewAxisX->ApplyToPreview(Component, X, true, XLabel);
+				PreviewAxisY->ApplyToPreview(Component, Y, false, YLabel);
+				
 				Component->Activate(true);
 			}
 		}
@@ -250,6 +256,13 @@ void ANiagaraPreviewGrid::GeneratePreviews()
 					PreviewActor->PrimaryActorTick.bCanEverTick = true;
 					PreviewActor->PrimaryActorTick.bStartWithTickEnabled = true;
 					PreviewActor->SetSystem(System);
+
+					TArray<UNiagaraComponent*, TInlineAllocator<4>> NiagaraComponents;
+					PreviewActor->GetComponents<UNiagaraComponent>(NiagaraComponents);
+					for (UNiagaraComponent* Component : NiagaraComponents)
+					{
+						Component->Activate(true);
+					}
 				}
 			}
 		}
diff --git a/Engine/Plugins/FX/Niagara/Source/Niagara/Private/NiagaraScalabilityManager.cpp b/Engine/Plugins/FX/Niagara/Source/Niagara/Private/NiagaraScalabilityManager.cpp
new file mode 100644
--- /dev/null
+++ b/Engine/Plugins/FX/Niagara/Source/Niagara/Private/NiagaraScalabilityManager.cpp
@@ -0,0 +1,254 @@
+// Copyright 1998-2019 Epic Games, Inc. All Rights Reserved.
+#include "NiagaraScalabilityManager.h"
+#include "NiagaraWorldManager.h"
+#include "NiagaraComponent.h"
+
+static float GScalabilityUpdateTime_Low = 1.0f;
+static float GScalabilityUpdateTime_Medium = 0.5f;
+static float GScalabilityUpdateTime_High = 0.25f;
+static FAutoConsoleVariableRef CVarScalabilityUpdateTime_Low(TEXT("fx.NiagaraScalabilityUpdateTime_Low"), GScalabilityUpdateTime_Low, TEXT("Time in seconds between updates to scalability states for Niagara systems set to update at Low frequency. \n"), ECVF_Default);
+static FAutoConsoleVariableRef CVarScalabilityUpdateTime_Medium(TEXT("fx.NiagaraScalabilityUpdateTime_Medium"), GScalabilityUpdateTime_Medium, TEXT("Time in seconds between updates to scalability states for Niagara systems set to update at Medium frequency. \n"), ECVF_Default);
+static FAutoConsoleVariableRef CVarScalabilityUpdateTime_High(TEXT("fx.NiagaraScalabilityUpdateTime_High"), GScalabilityUpdateTime_High, TEXT("Time in seconds between updates to scalability states for Niagara systems set to update at High frequency. \n"), ECVF_Default);
+
+static int32 GScalabilityManParallelThreshold = 50;
+static FAutoConsoleVariableRef CVarScalabilityManParallelThreshold(TEXT("fx.ScalabilityManParallelThreshold"), GScalabilityManParallelThreshold, TEXT("Number of instances required for a niagara significance manger to go parallel for it's update. \n"), ECVF_Default);
+
+FNiagaraScalabilityManager::FNiagaraScalabilityManager()
+	: EffectType(nullptr)
+	, LastUpdateTime(0.0f)
+{
+
+}
+
+void FNiagaraScalabilityManager::Register(UNiagaraComponent* Component)
+{
+	check(Component->ScalabilityManagerHandle == INDEX_NONE);
+	check(ManagedComponents.Num() == State.Num());
+
+	Component->ScalabilityManagerHandle = ManagedComponents.Add(Component);
+	State.AddDefaulted();
+}
+
+void FNiagaraScalabilityManager::Unregister(UNiagaraComponent* Component)
+{
+	check(Component->ScalabilityManagerHandle != INDEX_NONE);
+	check(ManagedComponents.Num() == State.Num());
+
+	int32 IndexToRemove = Component->ScalabilityManagerHandle;
+	Component->ScalabilityManagerHandle = INDEX_NONE;
+	ManagedComponents.RemoveAtSwap(IndexToRemove);
+	State.RemoveAtSwap(IndexToRemove);
+
+	if (ManagedComponents.IsValidIndex(IndexToRemove))
+	{
+		check(ManagedComponents[IndexToRemove] != nullptr);
+		ManagedComponents[IndexToRemove]->ScalabilityManagerHandle = IndexToRemove;
+	}
+}
+
+void FNiagaraScalabilityManager::Update(FNiagaraWorldManager* WorldMan)
+{
+	float WorldTime = WorldMan->GetWorld()->GetTimeSeconds();
+	bool bShouldUpdateScalabilityStates = false;
+	switch (EffectType->UpdateFrequency)
+	{
+	case ENiagaraScalabilityUpdateFrequency::Continuous: bShouldUpdateScalabilityStates = true; break;
+	case ENiagaraScalabilityUpdateFrequency::High: bShouldUpdateScalabilityStates = WorldTime >= LastUpdateTime + GScalabilityUpdateTime_High; break;
+	case ENiagaraScalabilityUpdateFrequency::Medium: bShouldUpdateScalabilityStates = WorldTime >= LastUpdateTime + GScalabilityUpdateTime_Medium; break;
+	case ENiagaraScalabilityUpdateFrequency::Low: bShouldUpdateScalabilityStates = WorldTime >= LastUpdateTime + GScalabilityUpdateTime_Low; break;
+	};
+
+	if (!bShouldUpdateScalabilityStates)
+	{
+		return;
+	}
+
+	LastUpdateTime = WorldTime;
+
+	int32 CompIdx = 0;
+#if WITH_EDITOR
+	//In cases such as undo/redo we can get components and system asset references pulled from under us.
+	while (CompIdx < ManagedComponents.Num())
+	{
+		UNiagaraComponent* Component = ManagedComponents[CompIdx];
+		if (Component && Component->GetAsset())
+		{
+			++CompIdx;
+		}
+		else
+		{
+			Unregister(Component);
+		}
+	}
+#endif
+
+	bool bNeedSortedSignificanceCull = false;
+	SignificanceSortedIndices.Reset(ManagedComponents.Num());
+
+	//TODO parallelize if we exceed GScalabilityManParallelThreshold instances.
+	CompIdx = 0;
+	bool bAnyDirty = false;
+	for (int32 i = 0; i < ManagedComponents.Num(); ++i)
+	{
+		UNiagaraComponent* Component = ManagedComponents[i];
+		FNiagaraScalabilityState& CompState = State[i];
+#if DEBUG_SCALABILITY_STATE
+		CompState.bCulledByInstanceCount = false;
+		CompState.bCulledByMaxOwnerLOD = false;
+		CompState.bCulledBySignificance = false;
+		CompState.bCulledByVisibility = false;
+#endif
+
+		UNiagaraSystem* System = Component->GetAsset();
+		const FNiagaraScalabilitySettings& ScalabilitySettings = System->GetScalabilitySettings(Component->GetPreviewDetailLevel());
+
+		SignificanceSortedIndices.Add(i);
+		bNeedSortedSignificanceCull = ScalabilitySettings.bCullMaxInstanceCount && ScalabilitySettings.MaxInstances > 0;
+
+		WorldMan->CalculateScalabilityState(System, ScalabilitySettings, EffectType, Component, false, CompState);
+
+		bAnyDirty |= CompState.bDirty;
+	}
+
+	if (bNeedSortedSignificanceCull)
+	{
+		SignificanceSortedIndices.Sort([&](int32 A, int32 B) { return State[A].Significance > State[B].Significance; });
+
+		for (int32 i = 0; i < SignificanceSortedIndices.Num(); ++i)
+		{
+			int32 SortedIdx = SignificanceSortedIndices[i];
+			UNiagaraComponent* Component = ManagedComponents[SortedIdx];
+			FNiagaraScalabilityState& CompState = State[SortedIdx];
+			UNiagaraSystem* System = Component->GetAsset();
+
+			bool bOldCulled = CompState.bCulled;
+
+			const FNiagaraScalabilitySettings& ScalabilitySettings = System->GetScalabilitySettings(Component->GetPreviewDetailLevel());
+			WorldMan->SortedSignificanceCull(EffectType, ScalabilitySettings, CompState.Significance, i, CompState);
+
+			CompState.bDirty |= CompState.bCulled != bOldCulled;
+			bAnyDirty |= CompState.bDirty;
+		}
+	}
+
+	if (bAnyDirty)
+	{
+		CompIdx = 0;
+		//As we'll be activating and deactivating here, this must be done on the game thread.
+		while (CompIdx < ManagedComponents.Num())
+		{
+			FNiagaraScalabilityState& CompState = State[CompIdx];
+			bool bRepeatIndex = false;
+			if (CompState.bDirty)
+			{
+				UNiagaraComponent* Component = ManagedComponents[CompIdx];
+				if (CompState.bCulled)
+				{
+					switch (EffectType->CullReaction)
+					{
+					case ENiagaraCullReaction::Deactivate:					Component->DeactivateInternal(false); bRepeatIndex = true; break;//We don't increment CompIdx here as this call will remove an entry from ManagedObjects;
+					case ENiagaraCullReaction::DeactivateImmediate:			Component->DeactivateImmediateInternal(false); bRepeatIndex = true;  break; //We don't increment CompIdx here as this call will remove an entry from ManagedObjects;
+					case ENiagaraCullReaction::DeactivateResume:			Component->DeactivateInternal(true); break;
+					case ENiagaraCullReaction::DeactivateImmediateResume:	Component->DeactivateImmediateInternal(true); break;
+					};
+				}
+				else
+				{
+					if (EffectType->CullReaction == ENiagaraCullReaction::Deactivate || EffectType->CullReaction == ENiagaraCullReaction::DeactivateImmediate)
+					{
+						UE_LOG(LogNiagara, Error, TEXT("Niagara Component is incorrectly still registered with the scalability manager. %d - %s "), (int32)EffectType->CullReaction, *Component->GetAsset()->GetFullName());
+					}
+					Component->ActivateInternal(false, true);
+				}
+
+				//TODO: Beyond culling by hard limits here we could progressively scale down fx by biasing detail levels they use. Could also introduce some budgeting here like N at lvl 0, M at lvl 1 etc.
+				//TODO: Possibly also limiting the rate at which their instances can tick. Ofc system sims still need to run but instances can skip ticks.
+			}
+
+			//If we are making a call that will unregister this component from the manager and remove it from ManagedComponents then we need to visit the new component that is now at this index.
+			if (bRepeatIndex == false)
+			{
+				++CompIdx;
+			}
+		}
+	}
+}
+
+#if DEBUG_SCALABILITY_STATE
+
+void FNiagaraScalabilityManager::Dump()
+{
+	FString DetailString;
+
+	struct FSummary
+	{
+		FSummary()
+			: NumCulled(0)
+			, NumCulledBySignificance(0)
+			, NumCulledByInstanceCount(0)
+			, NumCulledByVisibility(0)
+			, NumCulledByMaxOwnerLOD(0)
+		{}
+
+		int32 NumCulled;
+		int32 NumCulledBySignificance;
+		int32 NumCulledByInstanceCount;
+		int32 NumCulledByVisibility;
+		int32 NumCulledByMaxOwnerLOD;
+	}Summary;
+
+	for (int32 i = 0; i < ManagedComponents.Num(); ++i)
+	{
+		UNiagaraComponent* Comp = ManagedComponents[i];
+		FNiagaraScalabilityState& CompState = State[i];
+
+		FString CulledStr = TEXT("Active:");
+		if (CompState.bCulled)
+		{
+			CulledStr = TEXT("Culled:");
+			++Summary.NumCulled;
+		}
+		if (CompState.bCulledBySignificance)
+		{
+			CulledStr += TEXT("-Significance-");
+			++Summary.NumCulledBySignificance;
+		}
+		if (CompState.bCulledByInstanceCount)
+		{
+			CulledStr += TEXT("-Inst Count-");
+			++Summary.NumCulledByInstanceCount;
+		}
+		if (CompState.bCulledByVisibility)
+		{
+			CulledStr += TEXT("-Visibility-");
+			++Summary.NumCulledByVisibility;
+		}
+		if (CompState.bCulledByMaxOwnerLOD)
+		{
+			CulledStr += TEXT("-Owner LOD-");
+			++Summary.NumCulledByMaxOwnerLOD;
+		}
+
+		DetailString += FString::Printf(TEXT("| %s | Sig: %2.4f | %0xP | %s | %s |\n"), *CulledStr, CompState.Significance, Comp, *Comp->GetAsset()->GetPathName(), *Comp->GetPathName());
+	}
+
+	UE_LOG(LogNiagara, Display, TEXT("-------------------------------------------------------------------------------"));
+	UE_LOG(LogNiagara, Display, TEXT("Effect Type: %s"), *EffectType->GetPathName());
+	UE_LOG(LogNiagara, Display, TEXT("-------------------------------------------------------------------------------"));
+	UE_LOG(LogNiagara, Display, TEXT("| Summary for managed systems of this effect type. Does NOT inclued all possible Niagara FX in scene. |"));
+	UE_LOG(LogNiagara, Display, TEXT("| Num Managed Components: %d |"), ManagedComponents.Num());
+	UE_LOG(LogNiagara, Display, TEXT("| Num Active: %d |"), ManagedComponents.Num() - Summary.NumCulled);
+	UE_LOG(LogNiagara, Display, TEXT("| Num Culled: %d |"), Summary.NumCulled);
+	UE_LOG(LogNiagara, Display, TEXT("| Num Culled By Significance: %d |"), Summary.NumCulledBySignificance);
+	UE_LOG(LogNiagara, Display, TEXT("| Num Culled By Instance Count: %d |"), Summary.NumCulledByInstanceCount);
+	UE_LOG(LogNiagara, Display, TEXT("| Num Culled By Max Owner LOD: %d |"), Summary.NumCulledByMaxOwnerLOD);
+	UE_LOG(LogNiagara, Display, TEXT("| Num Culled By Visibility: %d |"), Summary.NumCulledByVisibility);
+	UE_LOG(LogNiagara, Display, TEXT("| Avg Frame GT: %d |"), EffectType->GetAverageFrameTime_GT());
+	UE_LOG(LogNiagara, Display, TEXT("| Avg Frame GT + CNC: %d |"), EffectType->GetAverageFrameTime_GT_CNC());
+	UE_LOG(LogNiagara, Display, TEXT("| Avg Frame RT: %d |"), EffectType->GetAverageFrameTime_RT());
+	UE_LOG(LogNiagara, Display, TEXT("-------------------------------------------------------------------------------"));
+	UE_LOG(LogNiagara, Display, TEXT("| Details |"));
+	UE_LOG(LogNiagara, Display, TEXT("-------------------------------------------------------------------------------\n%s"), *DetailString);
+}
+
+#endif
\ No newline at end of file
diff --git a/Engine/Plugins/FX/Niagara/Source/Niagara/Private/NiagaraSettings.cpp b/Engine/Plugins/FX/Niagara/Source/Niagara/Private/NiagaraSettings.cpp
--- a/Engine/Plugins/FX/Niagara/Source/Niagara/Private/NiagaraSettings.cpp
+++ b/Engine/Plugins/FX/Niagara/Source/Niagara/Private/NiagaraSettings.cpp
@@ -4,6 +4,7 @@
 
 UNiagaraSettings::UNiagaraSettings(const FObjectInitializer& ObjectInitlaizer)
 	: Super(ObjectInitlaizer)
+	, DefaultEffectTypePtr(nullptr)
 {
 
 }
@@ -20,6 +21,13 @@ FText UNiagaraSettings::GetSectionText() const
 }
 #endif
 
+void UNiagaraSettings::PostLoad()
+{
+	Super::PostLoad();
+
+	DefaultEffectTypePtr = Cast<UNiagaraEffectType>(DefaultEffectType.TryLoad());
+}
+
 #if WITH_EDITOR
 void UNiagaraSettings::PostEditChangeProperty(FPropertyChangedEvent& PropertyChangedEvent)
 {
@@ -27,6 +35,8 @@ void UNiagaraSettings::PostEditChangeProperty(FPropertyChangedEvent& PropertyCha
 	{
 		SettingsChangedDelegate.Broadcast(PropertyChangedEvent.Property->GetName(), this);
 	}
+
+	DefaultEffectTypePtr = Cast<UNiagaraEffectType>(DefaultEffectType.TryLoad());
 }
 
 UNiagaraSettings::FOnNiagaraSettingsChanged& UNiagaraSettings::OnSettingsChanged()
@@ -37,4 +47,7 @@ UNiagaraSettings::FOnNiagaraSettingsChanged& UNiagaraSettings::OnSettingsChanged
 UNiagaraSettings::FOnNiagaraSettingsChanged UNiagaraSettings::SettingsChangedDelegate;
 #endif
 
-
+UNiagaraEffectType* UNiagaraSettings::GetDefaultEffectType()const
+{
+	return DefaultEffectTypePtr;
+}
diff --git a/Engine/Plugins/FX/Niagara/Source/Niagara/Private/NiagaraSystem.cpp b/Engine/Plugins/FX/Niagara/Source/Niagara/Private/NiagaraSystem.cpp
--- a/Engine/Plugins/FX/Niagara/Source/Niagara/Private/NiagaraSystem.cpp
+++ b/Engine/Plugins/FX/Niagara/Source/Niagara/Private/NiagaraSystem.cpp
@@ -18,6 +18,8 @@
 #include "INiagaraEditorOnlyDataUtlities.h"
 #include "NiagaraWorldManager.h"
 #include "NiagaraEmitterInstanceBatcher.h"
+#include "NiagaraSettings.h"
+
 
 #if WITH_EDITOR
 #include "NiagaraScriptDerivedData.h"
@@ -28,10 +30,9 @@ DECLARE_CYCLE_STAT(TEXT("Niagara - System - Precompile"), STAT_Niagara_System_Pr
 DECLARE_CYCLE_STAT(TEXT("Niagara - System - CompileScript"), STAT_Niagara_System_CompileScript, STATGROUP_Niagara);
 DECLARE_CYCLE_STAT(TEXT("Niagara - System - CompileScript_ResetAfter"), STAT_Niagara_System_CompileScriptResetAfter, STATGROUP_Niagara);
 
-// UNiagaraSystemCategory::UNiagaraSystemCategory(const FObjectInitializer& ObjectInitializer)
-// 	: Super(ObjectInitializer)
-// {
-// }
+//Disable for now until we can spend more time on a good method of applying the data gathered.
+int32 GEnableNiagaraRuntimeCycleCounts = 0;
+static FAutoConsoleVariableRef CVarEnableNiagaraRuntimeCycleCounts(TEXT("fx.EnableNiagaraRuntimeCycleCounts"), GEnableNiagaraRuntimeCycleCounts, TEXT("Toggle for runtime cylce counts tracking Niagara's frame time. \n"), ECVF_ReadOnly);
 
 static int GNiagaraForceSystemsToCookOutRapidIterationOnLoad = 0;
 static FAutoConsoleVariableRef CVarNiagaraForceSystemsToCookOutRapidIterationOnLoad(
@@ -61,6 +62,11 @@ UNiagaraSystem::UNiagaraSystem(const FObjectInitializer& ObjectInitializer)
 	MaxPoolSize = 32;
 }
 
+UNiagaraSystem::UNiagaraSystem(FVTableHelper& Helper)
+	: Super(Helper)
+{
+}
+
 void UNiagaraSystem::BeginDestroy()
 {
 	Super::BeginDestroy();
@@ -71,6 +77,9 @@ void UNiagaraSystem::BeginDestroy()
 		QueryCompileComplete(true, false, true);
 	}
 #endif
+
+	//Should we just destroy all system sims here to simplify cleanup?
+	//FNiagaraWorldManager::DestroyAllSystemSimulations(this);
 }
 
 void UNiagaraSystem::PreSave(const class ITargetPlatform * TargetPlatform)
@@ -223,6 +232,22 @@ void UNiagaraSystem::Serialize(FArchive& Ar)
 }
 
 #if WITH_EDITOR
+
+void UNiagaraSystem::PreEditChange(UProperty* PropertyThatWillChange)
+{
+	Super::PreEditChange(PropertyThatWillChange);
+
+	if (PropertyThatWillChange && PropertyThatWillChange->GetFName() == GET_MEMBER_NAME_CHECKED(UNiagaraSystem, EffectType))
+	{
+		FNiagaraSystemUpdateContext UpdateContext;
+		UpdateContext.SetDestroyOnAdd(true);
+		UpdateContext.Add(this, false);
+		
+		//We have to clear this here and reset so that we unregister properly from the scalability manager.
+		EffectType = nullptr;
+	}
+}
+
 void UNiagaraSystem::PostEditChangeProperty(struct FPropertyChangedEvent& PropertyChangedEvent)
 {
 	Super::PostEditChangeProperty(PropertyChangedEvent);
@@ -250,6 +275,8 @@ void UNiagaraSystem::PostEditChangeProperty(struct FPropertyChangedEvent& Proper
 			}
 		}
 	}
+
+	ResolveScalabilityOverrides();
 }
 #endif 
 
@@ -462,6 +489,8 @@ void UNiagaraSystem::PostLoad()
 	{
 		bIsReadyToRunCached = IsReadyToRunInternal();
 	}
+
+	ResolveScalabilityOverrides();
 }
 
 #if WITH_EDITORONLY_DATA
@@ -1312,6 +1341,110 @@ void UNiagaraSystem::GenerateStatID()const
 #endif
 }
 
+UNiagaraEffectType* UNiagaraSystem::GetEffectType()const
+{
+	if (EffectType)
+	{
+		return EffectType;
+	}
+	
+	const UNiagaraSettings* Settings = GetDefault<UNiagaraSettings>();
+	check(Settings);
+	return Settings->GetDefaultEffectType();
+}
+
+const FNiagaraScalabilitySettings& UNiagaraSystem::GetScalabilitySettings(int32 DetailLevel)
+{
+#if WITH_NIAGARA_COMPONENT_PREVIEW_DATA
+	int32 UseDetailLevel = DetailLevel == INDEX_NONE ? INiagaraModule::GetDetailLevel() : DetailLevel;
+
+	int32 SettingsIndex = ResolvedScalabilitySettings.IsValidIndex(UseDetailLevel) ? UseDetailLevel : ResolvedScalabilitySettings.Num() - 1;
+	check(SettingsIndex != INDEX_NONE);//We should always have at least one entry in our settings array.
+	return ResolvedScalabilitySettings[SettingsIndex];
+#else
+	return CurrentScalabilitySettings;
+#endif
+}
+
+void UNiagaraSystem::ResolveScalabilityOverrides()
+{
+	ResolvedScalabilitySettings.Empty();
+
+	UNiagaraEffectType* ActualEffectType = GetEffectType();
+	int32 NumEffectTypeSettings = ActualEffectType ? ActualEffectType->GetScalabilitySettings().Num() : 0;
+	int32 NumOverrides = ScalabilityOverrides.Num();
+
+	int32 ResolvedCount = FMath::Max(NumEffectTypeSettings, NumOverrides);
+	if (ResolvedCount == 0)
+	{
+		//Just add a single defaulted entry to make other calling code simpler.
+		ResolvedScalabilitySettings.AddDefaulted();
+	}
+	else
+	{
+
+		ResolvedScalabilitySettings.Reserve(ResolvedCount);
+		for (int32 i = 0; i < ResolvedCount; ++i)
+		{
+			int32 NewIdx = ResolvedScalabilitySettings.AddDefaulted();
+			FNiagaraScalabilitySettings& NewSettings = ResolvedScalabilitySettings[NewIdx];
+			if (ActualEffectType)
+			{
+				const TArray<FNiagaraScalabilitySettings>& EffectTypeSettings = EffectType->GetScalabilitySettings();
+				int32 LastIndex = EffectTypeSettings.Num() - 1;
+				if (LastIndex >= 0)
+				{
+					NewSettings = EffectTypeSettings[FMath::Min(i, LastIndex)];
+				}
+			}
+
+			if (ScalabilityOverrides.IsValidIndex(i) && bOverrideScalabilitySettings)
+			{
+				const FNiagaraScalabilityOverrides& Overrides = ScalabilityOverrides[i];
+
+				if (Overrides.bOverrideDistanceSettings)
+				{
+					NewSettings.bCullByDistance = Overrides.bCullByDistance;
+					NewSettings.MaxDistance = Overrides.MaxDistance;
+				}
+
+				if (Overrides.bOverrideOwnerLODSettings)
+				{
+					NewSettings.bCullByMaxOwnerLOD = Overrides.bCullByMaxOwnerLOD;
+					NewSettings.MaxOwnerLOD = Overrides.MaxOwnerLOD;
+				}
+
+				if (Overrides.bOverrideInstanceCountSettings)
+				{
+					NewSettings.bCullMaxInstanceCount = Overrides.bCullMaxInstanceCount;
+					NewSettings.MaxInstances = Overrides.MaxInstances;
+				}
+
+				if (Overrides.bOverrideTimeSinceRendererSettings)
+				{
+					NewSettings.bCullByMaxTimeWithoutRender = Overrides.bCullByMaxTimeWithoutRender;
+					NewSettings.MaxTimeWithoutRender = Overrides.MaxTimeWithoutRender;
+				}
+
+// 				if (Overrides.bOverrideSpawnCountScale)
+// 				{
+// 					NewSettings.SpawnCountScale = Overrides.SpawnCountScale;
+// 				}
+			}
+		}
+	}
+
+	OnDetailLevelChanges(INiagaraModule::GetDetailLevel());
+}
+
+void UNiagaraSystem::OnDetailLevelChanges(int32 DetailLevel)
+{
+	int32 SettingsIndex = ResolvedScalabilitySettings.IsValidIndex(DetailLevel) ? DetailLevel : ResolvedScalabilitySettings.Num() - 1;
+	check(SettingsIndex != INDEX_NONE);//We should always have at least one entry in our settings array.
+	CurrentScalabilitySettings = ResolvedScalabilitySettings[SettingsIndex];
+}
+
+
 FNiagaraEmitterCompiledData::FNiagaraEmitterCompiledData()
 {
 	EmitterSpawnIntervalVar = SYS_PARAM_EMITTER_SPAWN_INTERVAL;
@@ -1320,4 +1453,4 @@ FNiagaraEmitterCompiledData::FNiagaraEmitterCompiledData()
 	EmitterSpawnGroupVar = SYS_PARAM_EMITTER_SPAWN_GROUP;
 	EmitterRandomSeedVar = SYS_PARAM_EMITTER_RANDOM_SEED;
 	EmitterTotalSpawnedParticlesVar = SYS_PARAM_ENGINE_EMITTER_TOTAL_SPAWNED_PARTICLES;
-}
\ No newline at end of file
+}
diff --git a/Engine/Plugins/FX/Niagara/Source/Niagara/Private/NiagaraSystemFastPath.cpp b/Engine/Plugins/FX/Niagara/Source/Niagara/Private/NiagaraSystemFastPath.cpp
--- a/Engine/Plugins/FX/Niagara/Source/Niagara/Private/NiagaraSystemFastPath.cpp
+++ b/Engine/Plugins/FX/Niagara/Source/Niagara/Private/NiagaraSystemFastPath.cpp
@@ -27,18 +27,18 @@ void FNiagaraSystemFastPath::SetUpdateMapDefaults(FParamMap0& Map)
 
 void FNiagaraSystemFastPath::Module_SystemScalability(const FNiagaraFastPath_Module_SystemScalability& Context_Map_SystemScalability, FParamMap0& Context_Map)
 {
-	bool Result9 = Context_Map.Engine.Owner.TimeSinceRendered > Context_Map_SystemScalability.VisibilityCullDelay;
-	bool Result10 = Context_Map_SystemScalability.bEnableVisibilityCulling && Result9;
-	ENiagaraExecutionState ENiagaraExecutionState_IfResult3;
-	if (Result10)
-	{
-		ENiagaraExecutionState_IfResult3 = Context_Map_SystemScalability.CulledState;
-	}
-	else
-	{
-		ENiagaraExecutionState_IfResult3 = Context_Map.Scalability.ExecutionState;
-	}
-	Context_Map.Scalability.ExecutionState = ENiagaraExecutionState_IfResult3;
+// 	bool Result9 = Context_Map.Engine.Owner.TimeSinceRendered > Context_Map_SystemScalability.VisibilityCullDelay;
+// 	bool Result10 = Context_Map_SystemScalability.bEnableVisibilityCulling && Result9;
+// 	ENiagaraExecutionState ENiagaraExecutionState_IfResult3;
+// 	if (Result10)
+// 	{
+// 		ENiagaraExecutionState_IfResult3 = Context_Map_SystemScalability.CulledState;
+// 	}
+// 	else
+// 	{
+// 		ENiagaraExecutionState_IfResult3 = Context_Map.Scalability.ExecutionState;
+// 	}
+// 	Context_Map.Scalability.ExecutionState = ENiagaraExecutionState_IfResult3;
 }
 
 void FNiagaraSystemFastPath::Function_SystemChangeState(ENiagaraExecutionState In_NewState, bool In_Condition, ENiagaraExecutionStateSource In_NewStateSource, FParamMap0& Context_Map)
@@ -293,7 +293,7 @@ void FNiagaraEmitterFastPath::Module_EmitterScalability(const FNiagaraFastPath_M
 		float_IfResult = Constant13;
 	}
 	float Result9 = Context_Map.Engine.Owner.LODDistance - float_IfResult;
-	float Constant14 = 1e+12;
+	float Constant14 = Context_Map.Engine.Owner.MaxLODDistance;
 	float float_IfResult1;
 	if (Context_Map_EmitterScalability.bUseMaxDistance)
 	{
@@ -479,21 +479,24 @@ void FNiagaraEmitterFastPath::Module_EmitterLifeCycle(const FNiagaraFastPath_Mod
 	ENiagaraExecutionStateSource Constant32 = ENiagaraExecutionStateSource::Scalability;
 	bool Result31 = Context_Map.System->ExecutionStateSource != Constant32;
 	Function_EmitterLifeCycle_EmitterChangeState(Context_Map.System->ExecutionState, Result31, Context_Map.System->ExecutionStateSource, Context_Map);
-	ENiagaraExecutionStateSource Constant35 = ENiagaraExecutionStateSource::Scalability;
-	bool Result38 = Context_Map.System->ExecutionStateSource == Constant35;
-	ENiagaraExecutionState Constant36 = (ENiagaraExecutionState::Active);
-	bool Result39 = Context_Map.System->ExecutionState != Constant36;
-	bool Result40 = Result38 && Result39;
-	ENiagaraExecutionState ENiagaraExecutionState_IfResult7;
-	if (Result40)
-	{
-		ENiagaraExecutionState_IfResult7 = Context_Map.System->ExecutionState;
-	}
-	else
-	{
-		ENiagaraExecutionState_IfResult7 = Context_Map.Scalability.Emitter.ExecutionState;
-	}
-	Context_Map.Scalability.Emitter.ExecutionState = ENiagaraExecutionState_IfResult7;
+
+	//System level scalability is now handled in higher level C++ code.
+// 	ENiagaraExecutionStateSource Constant35 = ENiagaraExecutionStateSource::Scalability;
+// 	bool Result38 = Context_Map.System->ExecutionStateSource == Constant35;
+// 	ENiagaraExecutionState Constant36 = (ENiagaraExecutionState::Active);
+// 	bool Result39 = Context_Map.System->ExecutionState != Constant36;
+// 	bool Result40 = Result38 && Result39;
+// 	ENiagaraExecutionState ENiagaraExecutionState_IfResult7;
+// 	if (Result40)
+// 	{
+// 		ENiagaraExecutionState_IfResult7 = Context_Map.System->ExecutionState;
+// 	}
+// 	else
+// 	{
+// 		ENiagaraExecutionState_IfResult7 = Context_Map.Scalability.Emitter.ExecutionState;
+// 	}
+// 	Context_Map.Scalability.Emitter.ExecutionState = ENiagaraExecutionState_IfResult7;
+
 	bool Constant37 = true;
 	ENiagaraExecutionStateSource Constant38 = ENiagaraExecutionStateSource::Scalability;
 	Function_EmitterLifeCycle_EmitterChangeState(Context_Map.Scalability.Emitter.ExecutionState, Constant37, Constant38, Context_Map);
@@ -515,7 +518,10 @@ void FNiagaraEmitterFastPath::Module_EmitterLifeCycle(const FNiagaraFastPath_Mod
 	int32 Constant46 = 0;
 	bool Result50 = Context_Map.Engine.Emitter.NumParticles == Constant46;
 	ENiagaraExecutionState Constant47 = ENiagaraExecutionState::Active;
-	bool Result51 = Context_Map.Emitter.ExecutionState != Constant47;
+
+	//Manual edit to stop auto complete when we become inactive due to scalability.
+	bool Result51 = Context_Map.Emitter.ExecutionState != Constant47 && Context_Map.Emitter.ExecutionStateSource != ENiagaraExecutionStateSource::Scalability;
+
 	bool Result52 = Result50 && Result51;
 	bool Result53 = Result52 && Context_Map_EmitterLifeCycle.bAutoComplete;
 	bool Result54 = Context_Map_EmitterLifeCycle.bCompleteOnInactive && Result51;
diff --git a/Engine/Plugins/FX/Niagara/Source/Niagara/Private/NiagaraSystemInstance.cpp b/Engine/Plugins/FX/Niagara/Source/Niagara/Private/NiagaraSystemInstance.cpp
--- a/Engine/Plugins/FX/Niagara/Source/Niagara/Private/NiagaraSystemInstance.cpp
+++ b/Engine/Plugins/FX/Niagara/Source/Niagara/Private/NiagaraSystemInstance.cpp
@@ -40,7 +40,7 @@ static FAutoConsoleVariableRef CVarWaitForAsyncStallWarnThresholdMS(
 );
 
 /** Safety time to allow for the LastRenderTime coming back from the RT. This is overkill but that's ok.*/
-static float GLastRenderTimeSafetyBias = 0.1f;
+float GLastRenderTimeSafetyBias = 0.1f;
 static FAutoConsoleVariableRef CVarLastRenderTimeSafetyBias(
 	TEXT("fx.LastRenderTimeSafetyBias"),
 	GLastRenderTimeSafetyBias,
@@ -82,6 +82,7 @@ FNiagaraSystemInstance::FNiagaraSystemInstance(UNiagaraComponent* InComponent)
 	, bNeedsFinalize(false)
 	, bDataInterfacesInitialized(false)
 	, bAlreadyBound(false)
+	, bLODDistanceIsValid(false)
 	, bAsyncWorkInProgress(false)
 	, CachedDeltaSeconds(0.0f)
 	, RequestedExecutionState(ENiagaraExecutionState::Complete)
@@ -92,6 +93,8 @@ FNiagaraSystemInstance::FNiagaraSystemInstance(UNiagaraComponent* InComponent)
 
 	LocalBounds = FBox(FVector::ZeroVector, FVector::ZeroVector);
 
+	LODDistance = 0.0f;
+
 	if (Component)
 	{
 		UWorld* World = Component->GetWorld();
@@ -698,6 +701,7 @@ void FNiagaraSystemInstance::ResetInternal(bool bResetSimulations)
 	TickCount = 0;
 	bHasTickingEmitters = true;
 	CachedDeltaSeconds = 0.0f;
+	bLODDistanceIsValid = false;
 	// Note: We do not need to update our bounds here as they are still valid
 
 	UNiagaraSystem* System = GetSystem();
@@ -914,6 +918,7 @@ void FNiagaraSystemInstance::ReInitInternal()
 	OwnerEngineRealtimeParam.Init(InstanceParameters, SYS_PARAM_ENGINE_REAL_TIME);
 
 	OwnerLODDistanceParam.Init(InstanceParameters, SYS_PARAM_ENGINE_LOD_DISTANCE);
+	OwnerLODDistanceFractionParam.Init(InstanceParameters, SYS_PARAM_ENGINE_LOD_DISTANCE_FRACTION);
 	SystemNumEmittersParam.Init(InstanceParameters, SYS_PARAM_ENGINE_SYSTEM_NUM_EMITTERS);
 	SystemNumEmittersAliveParam.Init(InstanceParameters, SYS_PARAM_ENGINE_SYSTEM_NUM_EMITTERS_ALIVE);
 
@@ -1273,6 +1278,12 @@ float FNiagaraSystemInstance::GetLODDistance()
 	}
 #endif
 
+	//In most cases this will have been set externally by the scalability manager.
+	if (bLODDistanceIsValid)
+	{
+		return LODDistance;
+	}
+	
 	constexpr float DefaultLODDistance = 0.0f;
 
 	FNiagaraWorldManager* WorldManager = GetWorldManager();
@@ -1280,8 +1291,11 @@ float FNiagaraSystemInstance::GetLODDistance()
 	{
 		return DefaultLODDistance;
 	}
-
+	
+	UWorld* World = Component->GetWorld();
+	check(World);
 	const FVector EffectLocation = Component->GetComponentLocation();
+	LODDistance = DefaultLODDistance;
 
 	// If we are inside the WorldManager tick we will use the cache player view locations as we can be ticked on different threads
 	if (WorldManager->CachedPlayerViewLocationsValid())
@@ -1289,24 +1303,25 @@ float FNiagaraSystemInstance::GetLODDistance()
 		TArrayView<const FVector> PlayerViewLocations = WorldManager->GetCachedPlayerViewLocations();
 		if (PlayerViewLocations.Num() == 0)
 		{
-			return DefaultLODDistance;
+			LODDistance = DefaultLODDistance;
 		}
-
-		// We are being ticked inside the WorldManager and can safely use the list of cached player view locations
-		float LODDistanceSqr = FMath::Square(WORLD_MAX);
-		for (const FVector& ViewLocation : PlayerViewLocations)
+		else
 		{
-			const float DistanceToEffectSqr = FVector(ViewLocation - EffectLocation).SizeSquared();
-			LODDistanceSqr = FMath::Min(LODDistanceSqr, DistanceToEffectSqr);
+			// We are being ticked inside the WorldManager and can safely use the list of cached player view locations
+			float LODDistanceSqr = FMath::Square(WORLD_MAX);
+			for (const FVector& ViewLocation : PlayerViewLocations)
+			{
+				const float DistanceToEffectSqr = FVector(ViewLocation - EffectLocation).SizeSquared();
+				LODDistanceSqr = FMath::Min(LODDistanceSqr, DistanceToEffectSqr);
+			}
+			LODDistance = FMath::Sqrt(LODDistanceSqr);
 		}
-		return FMath::Sqrt(LODDistanceSqr);
 	}
-
-	// If we are not inside the WorldManager tick (solo tick) we must look over the player view locations manually
-	ensureMsgf(IsInGameThread(), TEXT("FNiagaraSystemInstance::GetLODDistance called in potentially thread unsafe way"));
-
-	if ( UWorld* World = Component->GetWorld() )
+	else
 	{
+		// If we are not inside the WorldManager tick (solo tick) we must look over the player view locations manually
+		ensureMsgf(IsInGameThread(), TEXT("FNiagaraSystemInstance::GetLODDistance called in potentially thread unsafe way"));
+
 		TArray<FVector, TInlineAllocator<8> > PlayerViewLocations;
 		if (World->GetPlayerControllerIterator())
 		{
@@ -1334,10 +1349,10 @@ float FNiagaraSystemInstance::GetLODDistance()
 				const float DistanceToEffectSqr = FVector(ViewLocation - EffectLocation).SizeSquared();
 				LODDistanceSqr = FMath::Min(LODDistanceSqr, DistanceToEffectSqr);
 			}
-			return FMath::Sqrt(LODDistanceSqr);
+			LODDistance = FMath::Sqrt(LODDistanceSqr);
 		}
 	}
-	return DefaultLODDistance;
+	return LODDistance;
 }
 
 ETickingGroup FNiagaraSystemInstance::CalculateTickGroup()
@@ -1404,6 +1419,8 @@ void FNiagaraSystemInstance::TickInstanceParameters_GameThread(float DeltaSecond
 	{
 		return;
 	}
+	check(Component->GetAsset());
+	UNiagaraSystem* System = Component->GetAsset();
 
 	GatheredInstanceParameters.DeltaSeconds = DeltaSeconds;
 
@@ -1414,12 +1431,14 @@ void FNiagaraSystemInstance::TickInstanceParameters_GameThread(float DeltaSecond
 	if (World != NULL)
 	{
 		GatheredInstanceParameters.LODDistance = GetLODDistance();
+		GatheredInstanceParameters.MaxLODDistance = MaxLODDistance;
 		GatheredInstanceParameters.TimeSeconds = World->TimeSeconds;
 		GatheredInstanceParameters.RealTimeSeconds = World->RealTimeSeconds;
 	}
 	else
 	{
 		GatheredInstanceParameters.LODDistance = 0;
+		GatheredInstanceParameters.MaxLODDistance = 1;
 		GatheredInstanceParameters.TimeSeconds = Age;
 		GatheredInstanceParameters.RealTimeSeconds = Age;
 	}
@@ -1446,7 +1465,7 @@ void FNiagaraSystemInstance::TickInstanceParameters_GameThread(float DeltaSecond
 	}
 
 	//Bias the LastRenderTime slightly to account for any delay as it's written by the RT.
-	GatheredInstanceParameters.SafeTimeSinceRendererd = FMath::Max(0.0f, GatheredInstanceParameters.TimeSeconds - Component->GetLastRenderTime() - GLastRenderTimeSafetyBias);
+	GatheredInstanceParameters.SafeTimeSinceRendererd = Component->GetSafeTimeSinceRendered(GatheredInstanceParameters.TimeSeconds);
 
 	GatheredInstanceParameters.RequestedExecutionState = RequestedExecutionState;
 
@@ -1488,6 +1507,7 @@ void FNiagaraSystemInstance::TickInstanceParameters_Concurrent()
 	OwnerInverseDeltaSecondsParam.SetValue(1.0f / GatheredInstanceParameters.DeltaSeconds);
 
 	OwnerLODDistanceParam.SetValue(GatheredInstanceParameters.LODDistance);
+	OwnerLODDistanceFractionParam.SetValue(GatheredInstanceParameters.LODDistance / GatheredInstanceParameters.MaxLODDistance);
 	OwnerEngineTimeParam.SetValue(GatheredInstanceParameters.TimeSeconds);
 	OwnerEngineRealtimeParam.SetValue(GatheredInstanceParameters.RealTimeSeconds);
 
diff --git a/Engine/Plugins/FX/Niagara/Source/Niagara/Private/NiagaraSystemSimulation.cpp b/Engine/Plugins/FX/Niagara/Source/Niagara/Private/NiagaraSystemSimulation.cpp
--- a/Engine/Plugins/FX/Niagara/Source/Niagara/Private/NiagaraSystemSimulation.cpp
+++ b/Engine/Plugins/FX/Niagara/Source/Niagara/Private/NiagaraSystemSimulation.cpp
@@ -192,6 +192,7 @@ public:
 	void DoTask(ENamedThreads::Type CurrentThread, const FGraphEventRef& MyCompletionGraphEvent)
 	{
 		check(CurrentThread == ENamedThreads::GameThread);
+		FNiagaraScopedRuntimeCycleCounter RuntimeScope(SystemSim->GetSystem(), true, false);
 
 		for (FNiagaraSystemInstance* Inst : Batch)
 		{
@@ -230,6 +231,7 @@ public:
 
 	void DoTask(ENamedThreads::Type CurrentThread, const FGraphEventRef& MyCompletionGraphEvent)
 	{
+		FNiagaraScopedRuntimeCycleCounter RuntimeScope(SystemSim->GetSystem(), true, true);
 		for (FNiagaraSystemInstance* Inst : Batch)
 		{
 			Inst->Tick_Concurrent();
@@ -239,6 +241,25 @@ public:
 
 //////////////////////////////////////////////////////////////////////////
 
+void FNiagaraSystemSimulation::AddReferencedObjects(FReferenceCollector& Collector)
+{
+	//We keep a hard ref to the system.
+	Collector.AddReferencedObject(EffectType);
+}
+
+FNiagaraSystemSimulation::FNiagaraSystemSimulation()
+	: EffectType(nullptr)
+	, SystemTickGroup(TG_MAX)
+	, World(nullptr)
+	, bCanExecute(false)
+	, bBindingsInitialized(false)
+	, bInSpawnPhase(false)
+	, bIsSolo(false)
+	, bHasEverTicked(false)
+{
+
+}
+
 FNiagaraSystemSimulation::~FNiagaraSystemSimulation()
 {
 	Destroy();
@@ -250,6 +271,7 @@ bool FNiagaraSystemSimulation::Init(UNiagaraSystem* InSystem, UWorld* InWorld, b
 	UNiagaraSystem* System = InSystem;
 	WeakSystem = System;
 
+	EffectType = InSystem->GetEffectType();
 	SystemTickGroup = InTickGroup;
 
 	World = InWorld;
@@ -424,8 +446,12 @@ void FNiagaraSystemSimulation::Destroy()
 UNiagaraParameterCollectionInstance* FNiagaraSystemSimulation::GetParameterCollectionInstance(UNiagaraParameterCollection* Collection)
 {
 	UNiagaraSystem* System = WeakSystem.Get();
-	check(System != nullptr);
-	UNiagaraParameterCollectionInstance* Ret = System->GetParameterCollectionOverride(Collection);
+	UNiagaraParameterCollectionInstance* Ret = nullptr;
+
+	if (System)
+	{
+		System->GetParameterCollectionOverride(Collection);
+	}
 
 	//If no explicit override from the system, just get the current instance set on the world.
 	if (!Ret)
@@ -608,7 +634,7 @@ void FNiagaraSystemSimulation::Tick_GameThread(float DeltaSeconds, const FGraphE
 	check(PausedSystemInstances.Num() == PausedInstanceData.GetCurrentDataChecked().GetNumInstances());
 
 	UNiagaraSystem* System = WeakSystem.Get();
-
+	FNiagaraScopedRuntimeCycleCounter RuntimeScope(System, true, false);
 	FScopeCycleCounter SystemStatCounter(System->GetStatID(true, false));
 
 	if (MaxDeltaTime.IsSet())
@@ -835,6 +861,8 @@ void FNiagaraSystemSimulation::UpdateTickGroups_GameThread()
 	UNiagaraSystem* System = WeakSystem.Get();
 	check(System != nullptr);
 
+	FNiagaraScopedRuntimeCycleCounter RuntimeScope(System, true, false);
+
 	// Transfer promoted instances to the new tick group
 	//-OPT: This can be done async
 	for (FNiagaraSystemInstance* Instance : PendingTickGroupPromotions)
@@ -888,7 +916,9 @@ void FNiagaraSystemSimulation::Spawn_GameThread(float DeltaSeconds)
 	SCOPE_CYCLE_COUNTER(STAT_NiagaraOverview_GT);
 	CSV_SCOPED_TIMING_STAT_EXCLUSIVE(Niagara);
 	LLM_SCOPE(ELLMTag::Niagara);
-	FScopeCycleCounterUObject AdditionalScope(GetSystem(), GET_STATID(STAT_NiagaraOverview_GT_CNC));
+
+	UNiagaraSystem* System = WeakSystem.Get();
+	FScopeCycleCounterUObject AdditionalScope(System, GET_STATID(STAT_NiagaraOverview_GT_CNC));
 
 	WaitForSystemTickComplete(true);
 
@@ -899,6 +929,8 @@ void FNiagaraSystemSimulation::Spawn_GameThread(float DeltaSeconds)
 		Tick_GameThread(DeltaSeconds, nullptr);
 	}
 
+	FNiagaraScopedRuntimeCycleCounter RuntimeScope(System, true, false);
+
 	// Spawn instances
 	SpawningInstances.Reserve(PendingSystemInstances.Num());
 
@@ -1006,8 +1038,10 @@ void FNiagaraSystemSimulation::Tick_Concurrent(FNiagaraSystemSimulationTickConte
 	SCOPE_CYCLE_COUNTER(STAT_NiagaraOverview_GT_CNC);
 	CSV_SCOPED_TIMING_STAT_EXCLUSIVE(Niagara);
 	LLM_SCOPE(ELLMTag::Niagara);
-	FScopeCycleCounterUObject AdditionalScope(GetSystem(), GET_STATID(STAT_NiagaraOverview_GT_CNC));
 
+	FScopeCycleCounterUObject AdditionalScope(Context.System, GET_STATID(STAT_NiagaraOverview_GT_CNC));
+
+	FNiagaraScopedRuntimeCycleCounter RuntimeScope(Context.System, true, true);
 	FNiagaraSystemInstance* SoloSystemInstance = bIsSolo && Context.Instances.Num() == 1 ? Context.Instances[0] : nullptr;
 
 	if (bCanExecute && Context.Instances.Num())
@@ -1104,6 +1138,7 @@ void FNiagaraSystemSimulation::TickFastPath(FNiagaraSystemSimulationTickContext&
 			}
 			FNiagaraEmitterFastPath::FParamMap0& EmitterMap = EmitterInstance->GetFastPathMap();
 			EmitterMap.Engine.Owner.LODDistance = SystemInstance->GetOwnerLODDistance();
+			EmitterMap.Engine.Owner.MaxLODDistance = SystemInstance->GetOwnerMaxLODDistance();
 			EmitterMap.Engine.DeltaTime = Context.DeltaSeconds;
 			EmitterMap.Engine.Emitter.NumParticles = SystemInstance->GetNumParticles(EmitterIndex);
 			EmitterMap.Engine.Owner.Velocity = SystemInstance->GetOwnerVelocity();
@@ -1515,12 +1550,12 @@ void FNiagaraSystemSimulation::TransferSystemSimResults(FNiagaraSystemSimulation
 
 	for (int32 SystemIndex = 0; SystemIndex < Context.Instances.Num(); ++SystemIndex)
 	{
-		ENiagaraExecutionState ExecutionState = (ENiagaraExecutionState)SystemExecutionStateAccessor.GetSafe(SystemIndex, (int32)ENiagaraExecutionState::Disabled);
 		FNiagaraSystemInstance* SystemInst = Context.Instances[SystemIndex];
 
 		if (bIsUsingFastPath == false)
 		{
 			//Apply the systems requested execution state to it's actual execution state.
+			ENiagaraExecutionState ExecutionState = (ENiagaraExecutionState)SystemExecutionStateAccessor.GetSafe(SystemIndex, (int32)ENiagaraExecutionState::Disabled);
 			SystemInst->SetActualExecutionState(ExecutionState);
 		}
 
@@ -1596,6 +1631,11 @@ void FNiagaraSystemSimulation::RemoveInstance(FNiagaraSystemInstance* Instance)
 	check(PausedSystemInstances.Num() == PausedInstanceData.GetCurrentDataChecked().GetNumInstances());
 
 	check(IsInGameThread());
+	if (EffectType)
+	{
+		--EffectType->NumInstances;
+	}
+
 	UNiagaraSystem* System = WeakSystem.Get();
 	if (Instance->IsPendingSpawn())
 	{
@@ -1702,6 +1742,11 @@ void FNiagaraSystemSimulation::AddInstance(FNiagaraSystemInstance* Instance)
 		UE_LOG(LogNiagara, Log, TEXT("=== Adding To Pending Spawn %d ==="), Instance->SystemInstanceIndex);
 		//MainDataSet.Dump(true, Instance->SystemInstanceIndex, 1);
 	}
+	
+	if (EffectType)
+	{
+		++EffectType->NumInstances;
+	}
 
 	check(SystemInstances.Num() == MainDataSet.GetCurrentDataChecked().GetNumInstances());
 	check(PausedSystemInstances.Num() == PausedInstanceData.GetCurrentDataChecked().GetNumInstances());
diff --git a/Engine/Plugins/FX/Niagara/Source/Niagara/Private/NiagaraWorldManager.cpp b/Engine/Plugins/FX/Niagara/Source/Niagara/Private/NiagaraWorldManager.cpp
--- a/Engine/Plugins/FX/Niagara/Source/Niagara/Private/NiagaraWorldManager.cpp
+++ b/Engine/Plugins/FX/Niagara/Source/Niagara/Private/NiagaraWorldManager.cpp
@@ -16,7 +16,10 @@
 #include "NiagaraStats.h"
 #include "NiagaraEmitterInstanceBatcher.h"
 #include "NiagaraComponentPool.h"
+#include "NiagaraComponent.h"
+#include "NiagaraEffectType.h"
 
+DECLARE_CYCLE_STAT(TEXT("Niagara Manager Update Scalability Managers [GT]"), STAT_UpdateScalabilityManagers, STATGROUP_Niagara);
 DECLARE_CYCLE_STAT(TEXT("Niagara Manager Tick [GT]"), STAT_NiagaraWorldManTick, STATGROUP_Niagara);
 DECLARE_CYCLE_STAT(TEXT("Niagara Manager Wait On Render [GT]"), STAT_NiagaraWorldManWaitOnRender, STATGROUP_Niagara);
 
@@ -187,6 +190,7 @@ void FNiagaraWorldManager::AddReferencedObjects(FReferenceCollector& Collector)
 
 	Collector.AddReferencedObjects(ParameterCollections);
 	Collector.AddReferencedObject(ComponentPool);
+	Collector.AddReferencedObjects(ScalabilityManagers);
 }
 
 FString FNiagaraWorldManager::GetReferencerName() const
@@ -497,6 +501,9 @@ void FNiagaraWorldManager::Tick(ETickingGroup TickGroup, float DeltaSeconds, ELe
 			CachedPlayerViewLocations.Append(World->ViewLocationsRenderedLastFrame);
 		}
 
+		UpdateScalabilityManagers();
+
+
 		//Tick our collections to push any changes to bound stores.
 		//-TODO: Do we need to do this per tick group?
 		for (TPair<UNiagaraParameterCollection*, UNiagaraParameterCollectionInstance*> CollectionInstPair : ParameterCollections)
@@ -565,4 +572,305 @@ void FNiagaraWorldManager::DumpDetails(FOutputDevice& Ar)
 UWorld* FNiagaraWorldManager::GetWorld()
 {
 	return World;
-}
\ No newline at end of file
+}
+
+//////////////////////////////////////////////////////////////////////////
+
+void FNiagaraWorldManager::UpdateScalabilityManagers()
+{
+	SCOPE_CYCLE_COUNTER(STAT_UpdateScalabilityManagers);
+
+	for (auto& Pair : ScalabilityManagers)
+	{
+		FNiagaraScalabilityManager& ScalabilityMan = Pair.Value;
+		UNiagaraEffectType* EffectType = Pair.Key;
+
+		EffectType->ProcessLastFrameCycleCounts();
+
+		//TODO: Work out how best to budget each effect type.
+		//EffectType->ApplyDynamicBudget(DynamicBudget_GT, DynamicBudget_GT_CNC, DynamicBudget_RT);
+
+		ScalabilityMan.Update(this);
+	}
+}
+
+void FNiagaraWorldManager::RegisterWithScalabilityManager(UNiagaraComponent* Component)
+{
+	if (UNiagaraEffectType* EffectType = Component->GetAsset()->GetEffectType())
+	{
+		FNiagaraScalabilityManager* ScalabilityManager = ScalabilityManagers.Find(EffectType);
+
+		if (!ScalabilityManager)
+		{
+			ScalabilityManager = &ScalabilityManagers.Add(EffectType);
+			ScalabilityManager->EffectType = EffectType;
+		}
+
+		ScalabilityManager->Register(Component);
+	}
+}
+
+void FNiagaraWorldManager::UnregisterWithScalabilityManager(UNiagaraComponent* Component)
+{
+	if (UNiagaraEffectType* EffectType = Component->GetAsset()->GetEffectType())
+	{
+		FNiagaraScalabilityManager& ScalabilityManager = ScalabilityManagers.FindChecked(EffectType);
+		ScalabilityManager.Unregister(Component);
+	}
+}
+
+bool FNiagaraWorldManager::ShouldPreCull(UNiagaraSystem* System, UNiagaraComponent* Component)
+{
+	if (System)
+	{
+		if (UNiagaraEffectType* EffectType = System->GetEffectType())
+		{
+			if (CanPreCull(EffectType))
+			{
+				FNiagaraScalabilityState State;
+				const FNiagaraScalabilitySettings& ScalabilitySettings = System->GetScalabilitySettings(Component->GetPreviewDetailLevel());
+				CalculateScalabilityState(System, ScalabilitySettings, EffectType, Component, true, State);
+				return State.bCulled;
+			}
+		}
+	}
+	return false;
+}
+
+bool FNiagaraWorldManager::ShouldPreCull(UNiagaraSystem* System, FVector Location)
+{
+	if (System)
+	{
+		if (UNiagaraEffectType* EffectType = System->GetEffectType())
+		{
+			if (CanPreCull(EffectType))
+			{
+				FNiagaraScalabilityState State;
+				const FNiagaraScalabilitySettings& ScalabilitySettings = System->GetScalabilitySettings();
+				CalculateScalabilityState(System, ScalabilitySettings, EffectType, Location, true, State);
+				return State.bCulled;
+			}
+		}
+	}
+	return false;
+}
+
+void FNiagaraWorldManager::CalculateScalabilityState(UNiagaraSystem* System, const FNiagaraScalabilitySettings& ScalabilitySettings, UNiagaraEffectType* EffectType, FVector Location, bool bIsPreCull, FNiagaraScalabilityState& OutState)
+{
+	float DistSignificance = DistanceSignificance(EffectType, ScalabilitySettings, Location);
+
+	float Significance = DistSignificance;
+
+	//TODO: Other significance metrics? 
+	//TODO: Provide hook into game code for special case significance calcs?
+	OutState.Significance = Significance;
+
+	bool bOldCulled = OutState.bCulled;
+	SignificanceCull(EffectType, ScalabilitySettings, Significance, OutState);
+
+	//Only apply hard instance count cull limit for precull + spawn only fx. We can apply instance count via significance cull for managed fx.
+	if (bIsPreCull && EffectType->UpdateFrequency == ENiagaraScalabilityUpdateFrequency::SpawnOnly)
+	{
+		InstanceCountCull(EffectType, ScalabilitySettings, OutState);
+	}
+
+	OutState.bDirty = OutState.bCulled != bOldCulled;
+
+	//TODO: More progressive scalability options?
+	//We can possibly adjust the DetailLevel each effect can use so we can more gradually scale down effects at lower significance.
+	//Possibly do similar to the Significance manager and define a set budget of instances that are allowed at each level.
+	//E.g. 5 max quality environmental effects, 10 bias to -1 detail level, 20 biased to -2 etc.
+}
+
+void FNiagaraWorldManager::CalculateScalabilityState(UNiagaraSystem* System, const FNiagaraScalabilitySettings& ScalabilitySettings, UNiagaraEffectType* EffectType, UNiagaraComponent* Component, bool bIsPreCull, FNiagaraScalabilityState& OutState)
+{
+	float DistSignificance = DistanceSignificance(EffectType, ScalabilitySettings, Component);
+	
+	//If/when we do have multiple drivers of significance, how best to combine them?
+	float Significance = DistSignificance;
+
+	//TODO: Other significance metrics? 
+	//TODO: Provide hook into game code for special case significance calcs?
+	OutState.Significance = Significance;
+
+	bool bOldCulled = OutState.bCulled;
+	OutState.bCulled = false;
+	SignificanceCull(EffectType, ScalabilitySettings, Significance, OutState);
+	OwnerLODCull(EffectType, ScalabilitySettings, Component, OutState);
+	
+	//Can't cull dynamic bounds by visibility
+	if (System->bFixedBounds)
+	{
+		VisibilityCull(EffectType, ScalabilitySettings, Component, OutState);
+	}
+
+	//Only apply hard instance count cull limit for precull + spawn only fx. We can apply instance count via significance cull for managed fx.
+	if (bIsPreCull && EffectType->UpdateFrequency == ENiagaraScalabilityUpdateFrequency::SpawnOnly)
+	{
+		InstanceCountCull(EffectType, ScalabilitySettings, OutState);
+	}
+
+	OutState.bDirty = OutState.bCulled != bOldCulled;
+
+	//TODO: More progressive scalability options?
+	//We can possibly adjust the DetailLevel each effect can use so we can more gradually scale down effects at lower significance.
+	//Possibly do similar to the Significance manager and define a set budget of instances that are allowed at each level.
+	//E.g. 5 max quality environmental effects, 10 bias to -1 detail level, 20 biased to -2 etc.
+}
+
+bool FNiagaraWorldManager::CanPreCull(UNiagaraEffectType* EffectType)
+{
+	checkSlow(EffectType);
+	return EffectType->CullReaction == ENiagaraCullReaction::Deactivate || EffectType->CullReaction == ENiagaraCullReaction::DeactivateImmediate;
+}
+
+void FNiagaraWorldManager::SortedSignificanceCull(UNiagaraEffectType* EffectType, const FNiagaraScalabilitySettings& ScalabilitySettings, float Significance, int32 Index, FNiagaraScalabilityState& OutState)
+{
+	//Cull all but the N most significance FX.
+	bool bCull = ScalabilitySettings.bCullMaxInstanceCount && Index >= ScalabilitySettings.MaxInstances;
+	OutState.bCulled |= bCull;
+#if DEBUG_SCALABILITY_STATE
+	OutState.bCulledByInstanceCount = bCull;
+#endif
+}
+
+void FNiagaraWorldManager::SignificanceCull(UNiagaraEffectType* EffectType, const FNiagaraScalabilitySettings& ScalabilitySettings, float Significance, FNiagaraScalabilityState& OutState)
+{
+	float MinSignificance = 0.0f;
+
+// 	//Could We adjust the minimum significance needed by how much of this effect types budget is being used?
+// 	if (ScalabilitySettings.bCullByRuntimePerf)
+// 	{
+// 		MinSignificance = MinSignificanceFromPerf;
+// 
+// 		//TODO: Other factors raising the min significance?
+// 	}
+
+
+
+	bool bCull = Significance <= MinSignificance;
+	OutState.bCulled |= bCull;
+#if DEBUG_SCALABILITY_STATE
+	OutState.bCulledBySignificance = bCull;
+#endif
+}
+
+void FNiagaraWorldManager::VisibilityCull(UNiagaraEffectType* EffectType, const FNiagaraScalabilitySettings& ScalabilitySettings, UNiagaraComponent* Component, FNiagaraScalabilityState& OutState)
+{
+	float TimeSinceRendered = Component->GetSafeTimeSinceRendered(World->TimeSeconds);
+	bool bCull = ScalabilitySettings.bCullByMaxTimeWithoutRender && TimeSinceRendered > ScalabilitySettings.MaxTimeWithoutRender;
+
+	OutState.bCulled |= bCull;
+#if DEBUG_SCALABILITY_STATE
+	OutState.bCulledByVisibility = bCull;
+#endif
+}
+
+void FNiagaraWorldManager::OwnerLODCull(UNiagaraEffectType* EffectType, const FNiagaraScalabilitySettings& ScalabilitySettings, UNiagaraComponent* Component, FNiagaraScalabilityState& OutState)
+{
+	bool bCull = ScalabilitySettings.bCullByMaxOwnerLOD && Component->GetOwnerLOD() > ScalabilitySettings.MaxOwnerLOD;
+	OutState.bCulled |= bCull;
+#if DEBUG_SCALABILITY_STATE
+	OutState.bCulledByMaxOwnerLOD = bCull;
+#endif
+}
+
+void FNiagaraWorldManager::InstanceCountCull(UNiagaraEffectType* EffectType, const FNiagaraScalabilitySettings& ScalabilitySettings, FNiagaraScalabilityState& OutState)
+{
+	bool bCull = ScalabilitySettings.bCullMaxInstanceCount && EffectType->NumInstances > ScalabilitySettings.MaxInstances;
+	OutState.bCulled |= bCull;
+#if DEBUG_SCALABILITY_STATE
+	OutState.bCulledByInstanceCount = bCull;
+#endif
+}
+
+float FNiagaraWorldManager::DistanceSignificance(UNiagaraEffectType* EffectType, const FNiagaraScalabilitySettings& ScalabilitySettings, UNiagaraComponent* Component)
+{
+	float LODDistance = 0.0f;
+
+	//Directly drive the system lod distance from here.
+	float MaxDist = ScalabilitySettings.MaxDistance;
+#if WITH_NIAGARA_COMPONENT_PREVIEW_DATA
+	if (Component->bEnablePreviewLODDistance)
+	{
+		LODDistance = Component->PreviewLODDistance;
+	}
+	else
+#endif
+	if(bCachedPlayerViewLocationsValid)
+	{
+		float ClosestDistSq = FLT_MAX;
+		FVector Location = Component->GetComponentLocation();
+		for (FVector ViewLocation : CachedPlayerViewLocations)
+		{
+			ClosestDistSq = FMath::Min(ClosestDistSq, FVector::DistSquared(ViewLocation, Location));
+		}
+
+		LODDistance = FMath::Sqrt(ClosestDistSq);
+	}
+
+	Component->SetLODDistance(LODDistance, FMath::Max(MaxDist, 1.0f));
+
+	if (ScalabilitySettings.bCullByDistance)
+	{
+		if (LODDistance >= ScalabilitySettings.MaxDistance)
+		{
+			return 0.0f;
+		}
+
+		return 1.0f - (LODDistance / ScalabilitySettings.MaxDistance);
+	}
+	return 1.0f;
+}
+
+float FNiagaraWorldManager::DistanceSignificance(UNiagaraEffectType* EffectType, const FNiagaraScalabilitySettings& ScalabilitySettings, FVector Location)
+{
+	if (ScalabilitySettings.bCullByDistance)
+	{
+		float ClosestDistSq = FLT_MAX;
+		for (FVector ViewLocation : CachedPlayerViewLocations)
+		{
+			ClosestDistSq = FMath::Min(ClosestDistSq, FVector::DistSquared(ViewLocation, Location));
+		}
+
+		float ClosestDist = FMath::Sqrt(ClosestDistSq);
+		if (ClosestDist >= ScalabilitySettings.MaxDistance)
+		{
+			return 0.0f;
+		}
+
+		return ClosestDist / ScalabilitySettings.MaxDistance;
+	}
+	return 1.0f;
+}
+
+#if DEBUG_SCALABILITY_STATE
+
+void FNiagaraWorldManager::DumpScalabilityState()
+{
+	UE_LOG(LogNiagara, Display, TEXT("========================================================================"));
+	UE_LOG(LogNiagara, Display, TEXT("Niagara World Manager Scalability State. %0xP - %s"), World, *World->GetPathName());
+	UE_LOG(LogNiagara, Display, TEXT("========================================================================"));
+
+	for (auto& Pair : ScalabilityManagers)
+	{
+		FNiagaraScalabilityManager& ScalabilityMan = Pair.Value;
+		ScalabilityMan.Dump();
+	}
+
+	UE_LOG(LogNiagara, Display, TEXT("========================================================================"));
+}
+
+
+FAutoConsoleCommandWithWorld GDumpNiagaraScalabilityData(
+	TEXT("fx.DumpNiagaraScalabilityState"),
+	TEXT("Dumps state information for all Niagara Scalability Mangers."),
+	FConsoleCommandWithWorldDelegate::CreateStatic(
+		[](UWorld* World)
+{
+	FNiagaraWorldManager* WorldMan = FNiagaraWorldManager::Get(World);
+	WorldMan->DumpScalabilityState();
+}));
+
+
+#endif
diff --git a/Engine/Plugins/FX/Niagara/Source/Niagara/Public/NiagaraCommon.h b/Engine/Plugins/FX/Niagara/Source/Niagara/Public/NiagaraCommon.h
--- a/Engine/Plugins/FX/Niagara/Source/Niagara/Public/NiagaraCommon.h
+++ b/Engine/Plugins/FX/Niagara/Source/Niagara/Public/NiagaraCommon.h
@@ -466,17 +466,19 @@ struct FVMExternalFunctionBindingInfo
 
 struct NIAGARA_API FNiagaraSystemUpdateContext
 {
-	FNiagaraSystemUpdateContext(const UNiagaraSystem* System, bool bReInit) { Add(System, bReInit); }
+	FNiagaraSystemUpdateContext(const UNiagaraSystem* System, bool bReInit) :bDestroyOnAdd(false) { Add(System, bReInit); }
 #if WITH_EDITORONLY_DATA
-	FNiagaraSystemUpdateContext(const UNiagaraEmitter* Emitter, bool bReInit) { Add(Emitter, bReInit); }
-	FNiagaraSystemUpdateContext(const UNiagaraScript* Script, bool bReInit) { Add(Script, bReInit); }
+	FNiagaraSystemUpdateContext(const UNiagaraEmitter* Emitter, bool bReInit) : bDestroyOnAdd(false) { Add(Emitter, bReInit); }
+	FNiagaraSystemUpdateContext(const UNiagaraScript* Script, bool bReInit) :bDestroyOnAdd(false) { Add(Script, bReInit); }
 	//FNiagaraSystemUpdateContext(UNiagaraDataInterface* Interface, bool bReinit) : Add(Interface, bReinit) {}
-	FNiagaraSystemUpdateContext(const UNiagaraParameterCollection* Collection, bool bReInit) { Add(Collection, bReInit); }
+	FNiagaraSystemUpdateContext(const UNiagaraParameterCollection* Collection, bool bReInit) :bDestroyOnAdd(false) { Add(Collection, bReInit); }
 #endif
-	FNiagaraSystemUpdateContext() { }
+	FNiagaraSystemUpdateContext():bDestroyOnAdd(false){ }
 
 	~FNiagaraSystemUpdateContext();
 
+	void SetDestroyOnAdd(bool bInDestroyOnAdd) { bDestroyOnAdd = bInDestroyOnAdd; }
+
 	void Add(const UNiagaraSystem* System, bool bReInit);
 #if WITH_EDITORONLY_DATA
 	void Add(const UNiagaraEmitter* Emitter, bool bReInit);
@@ -490,13 +492,14 @@ struct NIAGARA_API FNiagaraSystemUpdateContext
 
 private:
 	void AddInternal(class UNiagaraComponent* Comp, bool bReInit);
-	FNiagaraSystemUpdateContext(FNiagaraSystemUpdateContext& Other) { }
+	FNiagaraSystemUpdateContext(FNiagaraSystemUpdateContext& Other) :bDestroyOnAdd(false) { }
 
 	TArray<UNiagaraComponent*> ComponentsToReset;
 	TArray<UNiagaraComponent*> ComponentsToReInit;
 
 	TArray<UNiagaraSystem*> SystemSimsToDestroy;
 
+	bool bDestroyOnAdd;
 	//TODO: When we allow component less systems we'll also want to find and reset those.
 };
 
diff --git a/Engine/Plugins/FX/Niagara/Source/Niagara/Public/NiagaraComponent.h b/Engine/Plugins/FX/Niagara/Source/Niagara/Public/NiagaraComponent.h
--- a/Engine/Plugins/FX/Niagara/Source/Niagara/Public/NiagaraComponent.h
+++ b/Engine/Plugins/FX/Niagara/Source/Niagara/Public/NiagaraComponent.h
@@ -25,6 +25,7 @@ class NiagaraEmitterInstanceBatcher;
 // Called when the particle system is done
 DECLARE_DYNAMIC_MULTICAST_DELEGATE_OneParam(FOnNiagaraSystemFinished, class UNiagaraComponent*, PSystem);
 
+#define WITH_NIAGARA_COMPONENT_PREVIEW_DATA (!UE_BUILD_SHIPPING)
 
 /**
 * UNiagaraComponent is the primitive component for a Niagara System.
@@ -34,6 +35,7 @@ DECLARE_DYNAMIC_MULTICAST_DELEGATE_OneParam(FOnNiagaraSystemFinished, class UNia
 UCLASS(ClassGroup = (Rendering, Common), hidecategories = Object, hidecategories = Physics, hidecategories = Collision, showcategories = Trigger, editinlinenew, meta = (BlueprintSpawnableComponent, DisplayName = "Niagara Particle System"))
 class NIAGARA_API UNiagaraComponent : public UFXSystemComponent
 {
+	friend struct FNiagaraScalabilityManager;
 	GENERATED_UCLASS_BODY()
 
 #if WITH_EDITORONLY_DATA
@@ -144,6 +146,22 @@ public:
 	virtual void Deactivate()override;
 	void DeactivateImmediate();
 
+	FORCEINLINE ENiagaraExecutionState GetRequestedExecutionState() { return SystemInstance ? SystemInstance->GetRequestedExecutionState() : ENiagaraExecutionState::Complete; }
+
+	FORCEINLINE float GetSafeTimeSinceRendered(float WorldTime)const;
+	private:
+
+	//Internal versions that can be called from the scalability code.
+	//These will behave as expected but will keep the component registered with the scalability manager.
+	void ActivateInternal(bool bReset, bool bIsScalabilityCull);
+	void DeactivateInternal(bool bIsScalabilityCull);
+	void DeactivateImmediateInternal(bool bIsScalabilityCull);
+
+	bool RegisterWithScalabilityManagerOrPreCull();
+	void UnregisterWithScalabilityManager();
+
+	public:
+
 	virtual void SetComponentTickEnabled(bool bEnabled) override;
 
 	virtual void TickComponent(float DeltaTime, enum ELevelTick TickType, FActorComponentTickFunction* ThisTickFunction) override;
@@ -372,13 +390,27 @@ public:
 #endif
 	//~ End UObject Interface
 
-
 	UFUNCTION(BlueprintCallable, Category = Preview, meta = (Keywords = "preview detail level scalability"))
 	void SetPreviewDetailLevel(bool bEnablePreviewDetailLevel, int32 PreviewDetailLevel);
 
 	UFUNCTION(BlueprintCallable, Category = Preview, meta = (Keywords = "preview LOD Distance scalability"))
 	void SetPreviewLODDistance(bool bEnablePreviewLODDistance, float PreviewLODDistance);
 
+	UFUNCTION(BlueprintCallable, Category = Preview, meta = (Keywords = "preview detail level scalability"))
+	FORCEINLINE bool GetPreviewDetailLevelEnabled()const;
+
+	UFUNCTION(BlueprintCallable, Category = Preview, meta = (Keywords = "preview LOD Distance scalability"))
+	FORCEINLINE bool GetPreviewLODDistanceEnabled()const;
+
+	UFUNCTION(BlueprintCallable, Category = Preview, meta = (Keywords = "preview detail level scalability"))
+	FORCEINLINE int32 GetPreviewDetailLevel()const;
+
+	UFUNCTION(BlueprintCallable, Category = Preview, meta = (Keywords = "preview LOD Distance scalability"))
+	FORCEINLINE int32 GetPreviewLODDistance()const;
+
+
+	FORCEINLINE void SetLODDistance(float InLODDistance, float InMaxLODDistance) { if (SystemInstance) SystemInstance->SetLODDistance(InLODDistance, InMaxLODDistance); }
+
 #if WITH_EDITOR
 	void PostLoadNormalizeOverrideNames();
 	bool IsParameterValueOverriddenLocally(const FName& InParamName);
@@ -461,23 +493,27 @@ public:
 
 	virtual void SetUseAutoManageAttachment(bool bAutoManage) override { bAutoManageAttachment = bAutoManage; }
 
-	UPROPERTY(EditAnywhere, Category = Preview, Transient, meta=(EditCondition=bEnablePreviewDetailLevel))
+#if WITH_NIAGARA_COMPONENT_PREVIEW_DATA
 	int32 PreviewDetailLevel;
-
-	UPROPERTY(EditAnywhere, Category = Preview, Transient, meta=(EditCondition= bEnablePreviewLODDistance))
 	float PreviewLODDistance;
-
-	UPROPERTY(EditAnywhere, Category = Preview, Transient)
 	uint32 bEnablePreviewDetailLevel : 1;
-
-	UPROPERTY(EditAnywhere, Category = Preview, Transient)
 	uint32 bEnablePreviewLODDistance : 1;
+#endif
 
 #if WITH_EDITORONLY_DATA
 	UPROPERTY(EditAnywhere, Category = Compilation)
 	uint32 bWaitForCompilationOnActivate : 1;
 #endif
 
+	UFUNCTION(BlueprintCallable, Category = Preview, meta = (Keywords = "LOD scalability"))
+	void SetOwnerLOD(int32 InOwnerLOD);
+
+	UFUNCTION(BlueprintCallable, Category = Preview, meta = (Keywords = "LOD scalability"))
+	FORCEINLINE int32 GetOwnerLOD()const { return OwnerLOD; }
+
+	/** Set whether this component is allowed to perform scalability checks and potentially be culled etc. Occasionally it is useful to disable this for specific components. E.g. Effects on the local player. */
+	FORCEINLINE void SetAllowScalability(bool bAllow) { bAllowScalability = bAllow; }
+
 private:
 	/** Did we try and activate but fail due to the asset being not yet ready. Keep looping.*/
 	uint32 bAwaitingActivationDueToNotReady : 1;
@@ -487,6 +523,9 @@ private:
 	/** Did we auto attach during activation? Used to determine if we should restore the relative transform during detachment. */
 	uint32 bDidAutoAttach : 1;
 
+	/** True if this component is allowed to perform scalability checks and potentially be culled etc. Occasionally it is useful to disable this for specific components. E.g. Effects on the local player. */
+	uint32 bAllowScalability : 1;
+
 	/** Flag to mark us as currently changing auto attachment as part of Activate/Deactivate so we don't reset in the OnAttachmentChanged() callback. */
 	//uint32 bIsChangingAutoAttachment : 1;
 
@@ -500,9 +539,31 @@ private:
 	FVector SavedAutoAttachRelativeScale3D;
 
 	FDelegateHandle AssetExposedParametersChangedHandle;
+
+	int32 ScalabilityManagerHandle;
+
+	/**
+	LOD level of our owning actor / component if it's been provided. 
+	Can be useful for scalability calculations for actor based FX on actors who have some other system determining their LOD level.
+	Would be nice if we just had a virtual on the actor Actor->GetLODLevel() for example that would return 0 by default and the correct level for any actors implementing a LOD level.
+	*/
+	int32 OwnerLOD;
 };
 
 
+#if WITH_NIAGARA_COMPONENT_PREVIEW_DATA
+FORCEINLINE bool UNiagaraComponent::GetPreviewDetailLevelEnabled()const { return bEnablePreviewDetailLevel; }
+FORCEINLINE bool UNiagaraComponent::GetPreviewLODDistanceEnabled()const { return bEnablePreviewLODDistance; }
+FORCEINLINE int32 UNiagaraComponent::GetPreviewDetailLevel()const { return bEnablePreviewDetailLevel ? PreviewDetailLevel : INDEX_NONE; }
+FORCEINLINE int32 UNiagaraComponent::GetPreviewLODDistance()const { return bEnablePreviewLODDistance ? PreviewLODDistance : 0.0f; }
+#else
+FORCEINLINE bool UNiagaraComponent::GetPreviewDetailLevelEnabled()const { return false; }
+FORCEINLINE bool UNiagaraComponent::GetPreviewLODDistanceEnabled()const { return false; }
+FORCEINLINE int32 UNiagaraComponent::GetPreviewDetailLevel()const { return INDEX_NONE; }
+FORCEINLINE int32 UNiagaraComponent::GetPreviewLODDistance()const { return 0.0f; }
+#endif
+
+
 
 
 
@@ -579,7 +640,16 @@ private:
 
 	FMatrix LocalToWorldInverse;
 
+	/** Ptr to the cycle count for this systems effect type. Lifetime is guaranteed to be longer than the proxy. */
+	int32* RuntimeCycleCount;
+
 #if STATS
 	TStatId SystemStatID;
 #endif
 };
+
+extern float GLastRenderTimeSafetyBias;
+FORCEINLINE float UNiagaraComponent::GetSafeTimeSinceRendered(float WorldTime)const
+{
+	return FMath::Max(0.0f, WorldTime - GetLastRenderTime() - GLastRenderTimeSafetyBias);
+}
\ No newline at end of file
diff --git a/Engine/Plugins/FX/Niagara/Source/Niagara/Public/NiagaraFunctionLibrary.h b/Engine/Plugins/FX/Niagara/Source/Niagara/Public/NiagaraFunctionLibrary.h
--- a/Engine/Plugins/FX/Niagara/Source/Niagara/Public/NiagaraFunctionLibrary.h
+++ b/Engine/Plugins/FX/Niagara/Source/Niagara/Public/NiagaraFunctionLibrary.h
@@ -30,12 +30,12 @@ public:
 	* @return			The spawned UNiagaraComponent
 	*/
 	UFUNCTION(BlueprintCallable, Category = Niagara, meta = (Keywords = "niagara System", WorldContext = "WorldContextObject", UnsafeDuringActorConstruction = "true"))
-	static UNiagaraComponent* SpawnSystemAtLocation(const UObject* WorldContextObject, class UNiagaraSystem* SystemTemplate, FVector Location, FRotator Rotation = FRotator::ZeroRotator, FVector Scale = FVector(1.f), bool bAutoDestroy = true, bool bAutoActivate = true, ENCPoolMethod PoolingMethod = ENCPoolMethod::None);
+	static UNiagaraComponent* SpawnSystemAtLocation(const UObject* WorldContextObject, class UNiagaraSystem* SystemTemplate, FVector Location, FRotator Rotation = FRotator::ZeroRotator, FVector Scale = FVector(1.f), bool bAutoDestroy = true, bool bAutoActivate = true, ENCPoolMethod PoolingMethod = ENCPoolMethod::None, bool bPreCullCheck = true);
 
 	UFUNCTION(BlueprintCallable, Category = Niagara, meta = (Keywords = "niagara System", UnsafeDuringActorConstruction = "true"))
-	static UNiagaraComponent* SpawnSystemAttached(UNiagaraSystem* SystemTemplate, USceneComponent* AttachToComponent, FName AttachPointName, FVector Location, FRotator Rotation, EAttachLocation::Type LocationType, bool bAutoDestroy, bool bAutoActivate = true, ENCPoolMethod PoolingMethod = ENCPoolMethod::None);
+	static UNiagaraComponent* SpawnSystemAttached(UNiagaraSystem* SystemTemplate, USceneComponent* AttachToComponent, FName AttachPointName, FVector Location, FRotator Rotation, EAttachLocation::Type LocationType, bool bAutoDestroy, bool bAutoActivate = true, ENCPoolMethod PoolingMethod = ENCPoolMethod::None, bool bPreCullCheck = true);
 
-	static UNiagaraComponent* SpawnSystemAttached(UNiagaraSystem* SystemTemplate, USceneComponent* AttachToComponent, FName AttachPointName, FVector Location, FRotator Rotation, FVector Scale, EAttachLocation::Type LocationType, bool bAutoDestroy, ENCPoolMethod PoolingMethod, bool bAutoActivate = true);
+	static UNiagaraComponent* SpawnSystemAttached(UNiagaraSystem* SystemTemplate, USceneComponent* AttachToComponent, FName AttachPointName, FVector Location, FRotator Rotation, FVector Scale, EAttachLocation::Type LocationType, bool bAutoDestroy, ENCPoolMethod PoolingMethod, bool bAutoActivate = true, bool bPreCullCheck = true);
 
 	/** Sets a Niagara StaticMesh parameter by name, overriding locally if necessary.*/
 	UFUNCTION(BlueprintCallable, Category = Niagara, meta = (DisplayName = "Set Niagara Static Mesh Component"))
diff --git a/Engine/Plugins/FX/Niagara/Source/Niagara/Public/NiagaraModule.h b/Engine/Plugins/FX/Niagara/Source/Niagara/Public/NiagaraModule.h
--- a/Engine/Plugins/FX/Niagara/Source/Niagara/Public/NiagaraModule.h
+++ b/Engine/Plugins/FX/Niagara/Source/Niagara/Public/NiagaraModule.h
@@ -96,6 +96,7 @@ public:
 
 	FORCEINLINE static const FNiagaraVariable&  GetVar_Engine_Owner_TimeSinceRendered() { return Engine_Owner_TimeSinceRendered; }
 	FORCEINLINE static const FNiagaraVariable&  GetVar_Engine_Owner_LODDistance() { return Engine_Owner_LODDistance; }
+	FORCEINLINE static const FNiagaraVariable&  GetVar_Engine_Owner_LODDistanceFraction() { return Engine_Owner_LODDistanceFraction; }
 
 	FORCEINLINE static const FNiagaraVariable&  GetVar_Engine_Owner_ExecutionState() { return Engine_Owner_ExecutionState; }
 
@@ -203,6 +204,7 @@ private:
 
 	static FNiagaraVariable Engine_Owner_TimeSinceRendered;
 	static FNiagaraVariable Engine_Owner_LODDistance;
+	static FNiagaraVariable Engine_Owner_LODDistanceFraction;
 	
 	static FNiagaraVariable Engine_Owner_ExecutionState;
 
diff --git a/Engine/Plugins/FX/Niagara/Source/Niagara/Public/NiagaraScalabilityManager.h b/Engine/Plugins/FX/Niagara/Source/Niagara/Public/NiagaraScalabilityManager.h
new file mode 100644
--- /dev/null
+++ b/Engine/Plugins/FX/Niagara/Source/Niagara/Public/NiagaraScalabilityManager.h
@@ -0,0 +1,134 @@
+// Copyright 1998-2019 Epic Games, Inc. All Rights Reserved.
+
+#pragma once
+
+#include "NiagaraEffectType.h"
+
+#include "NiagaraScalabilityManager.generated.h"
+
+class FNiagaraWorldManager;
+class UNiagaraSystem;
+class UNiagaraComponent;
+
+//Disabled for now until we can spend more time on a good method of applying the data that is gathered.
+#define ENABLE_NIAGARA_RUNTIME_CYCLE_COUNTERS (0)
+
+#if ENABLE_NIAGARA_RUNTIME_CYCLE_COUNTERS
+struct FNiagaraRuntimeCycleCounter
+{
+	FORCEINLINE FNiagaraRuntimeCycleCounter(int32* InCycleDest)
+		: StartCycles(0)
+		, CycleDest(InCycleDest)
+	{
+	}
+
+	FORCEINLINE FNiagaraRuntimeCycleCounter(UNiagaraSystem* System, bool bGameThread, bool bConcurrent)
+		: StartCycles(0)
+		, CycleDest(nullptr)
+	{
+		CycleDest = System->GetCycleCounter(bGameThread, bConcurrent);
+	}
+
+	FORCEINLINE void Begin()
+	{
+		if (CycleDest)
+		{
+			StartCycles = FPlatformTime::Cycles();
+		}
+	}
+
+	FORCEINLINE void End()
+	{
+		if (CycleDest)
+		{
+			FPlatformAtomics::InterlockedAdd(CycleDest, FPlatformTime::Cycles() - StartCycles);
+		}
+	}
+
+private:
+	uint32 StartCycles;
+	int32* CycleDest;
+};
+#else
+struct FNiagaraRuntimeCycleCounter
+{
+	FORCEINLINE FNiagaraRuntimeCycleCounter(int32* InCycleDest) {}
+	FORCEINLINE FNiagaraRuntimeCycleCounter(UNiagaraSystem* System, bool bGameThread, bool bConcurrent) {}
+	FORCEINLINE void Begin() {}
+	FORCEINLINE void End() {}
+};
+#endif
+
+struct FNiagaraScopedRuntimeCycleCounter : public FNiagaraRuntimeCycleCounter
+{
+	FORCEINLINE FNiagaraScopedRuntimeCycleCounter(int32* InCycleDest)
+		: FNiagaraRuntimeCycleCounter(InCycleDest)
+	{
+		Begin();
+	}
+
+	FORCEINLINE FNiagaraScopedRuntimeCycleCounter(UNiagaraSystem* System, bool bGameThread, bool bConcurrent)
+		: FNiagaraRuntimeCycleCounter(System, bGameThread, bConcurrent)
+	{
+		Begin();
+	}
+
+	FORCEINLINE ~FNiagaraScopedRuntimeCycleCounter()
+	{
+		End();
+	}
+};
+
+struct FNiagaraScalabilityState
+{
+	FNiagaraScalabilityState()
+		: Significance(1.0f)
+		, bDirty(0)
+		, bCulled(0)
+#if DEBUG_SCALABILITY_STATE
+		, bCulledBySignificance(0)
+		, bCulledByInstanceCount(0)
+		, bCulledByVisibility(0)
+		, bCulledByMaxOwnerLOD(0)
+#endif
+	{
+	}
+
+	float Significance;
+	uint8 bDirty : 1;
+	uint8 bCulled : 1;
+
+#if DEBUG_SCALABILITY_STATE
+	uint8 bCulledBySignificance : 1;
+	uint8 bCulledByInstanceCount : 1;
+	uint8 bCulledByVisibility : 1;
+	uint8 bCulledByMaxOwnerLOD : 1;
+#endif
+};
+
+USTRUCT()
+struct FNiagaraScalabilityManager
+{
+	GENERATED_USTRUCT_BODY()
+
+	UPROPERTY(transient)
+	UNiagaraEffectType* EffectType;
+
+	UPROPERTY(transient)
+	TArray<UNiagaraComponent*>  ManagedComponents;
+
+	TArray<FNiagaraScalabilityState> State;
+
+	TArray<int32> SignificanceSortedIndices;
+
+	float LastUpdateTime;
+
+	FNiagaraScalabilityManager();
+	void Update(FNiagaraWorldManager* Owner);
+	void Register(UNiagaraComponent* Component);
+	void Unregister(UNiagaraComponent* Component);
+
+#if DEBUG_SCALABILITY_STATE
+	void Dump();
+#endif
+};
diff --git a/Engine/Plugins/FX/Niagara/Source/Niagara/Public/NiagaraSettings.h b/Engine/Plugins/FX/Niagara/Source/Niagara/Public/NiagaraSettings.h
--- a/Engine/Plugins/FX/Niagara/Source/Niagara/Public/NiagaraSettings.h
+++ b/Engine/Plugins/FX/Niagara/Source/Niagara/Public/NiagaraSettings.h
@@ -23,7 +23,11 @@ class NIAGARA_API UNiagaraSettings : public UDeveloperSettings
 
 	UPROPERTY(config, EditAnywhere, Category = Niagara, meta = (AllowedClasses = "Enum"))
 	TArray<FSoftObjectPath> AdditionalParameterEnums;
-	
+
+	/** Default effect type to use for effects that don't define their own. Can be null. */
+	UPROPERTY(config, EditAnywhere, Category = Niagara, meta = (AllowedClasses = "NiagaraEffectType"))
+	FSoftObjectPath DefaultEffectType;
+
 	// Begin UDeveloperSettings Interface
 	virtual FName GetCategoryName() const override;
 #if WITH_EDITOR
@@ -31,6 +35,9 @@ class NIAGARA_API UNiagaraSettings : public UDeveloperSettings
 #endif
 	// END UDeveloperSettings Interface
 
+	UNiagaraEffectType* GetDefaultEffectType()const;
+
+	virtual void PostLoad();
 #if WITH_EDITOR
 	virtual void PostEditChangeProperty(FPropertyChangedEvent& PropertyChangedEvent) override;
 
@@ -43,4 +50,9 @@ public:
 protected:
 	static FOnNiagaraSettingsChanged SettingsChangedDelegate;
 #endif
+
+
+private:
+	UPROPERTY(transient)
+	UNiagaraEffectType* DefaultEffectTypePtr;
 };
\ No newline at end of file
diff --git a/Engine/Plugins/FX/Niagara/Source/Niagara/Public/NiagaraSystemFastPath.h b/Engine/Plugins/FX/Niagara/Source/Niagara/Public/NiagaraSystemFastPath.h
--- a/Engine/Plugins/FX/Niagara/Source/Niagara/Public/NiagaraSystemFastPath.h
+++ b/Engine/Plugins/FX/Niagara/Source/Niagara/Public/NiagaraSystemFastPath.h
@@ -396,6 +396,7 @@ namespace FNiagaraEmitterFastPath
 	struct FParamMap0_Engine_Owner
 	{
 		float LODDistance;
+		float MaxLODDistance;
 		FVector Velocity;
 	};
 
diff --git a/Engine/Plugins/FX/Niagara/Source/Niagara/Public/NiagaraSystemInstance.h b/Engine/Plugins/FX/Niagara/Source/Niagara/Public/NiagaraSystemInstance.h
--- a/Engine/Plugins/FX/Niagara/Source/Niagara/Public/NiagaraSystemInstance.h
+++ b/Engine/Plugins/FX/Niagara/Source/Niagara/Public/NiagaraSystemInstance.h
@@ -112,7 +112,8 @@ public:
 
 	float GetSystemTimeSinceRendered() const { return SystemTimeSinceRenderedParam.GetValue(); }
 
-	float GetOwnerLODDistance() const { return OwnerLODDistanceParam.GetValue(); }
+	float GetOwnerLODDistance() const { return GatheredInstanceParameters.LODDistance; }
+	float GetOwnerMaxLODDistance() const { return GatheredInstanceParameters.MaxLODDistance; }
 
 	int32 GetNumParticles(int32 EmitterIndex) const { return ParameterNumParticleBindings[EmitterIndex].GetValue(); }
 	float GetSpawnCountScale(int32 EmitterIndex) const { return ParameterSpawnCountScaleBindings[EmitterIndex].GetValue(); }
@@ -252,7 +253,7 @@ public:
 
 	FNiagaraSystemFastPath::FParamMap0& GetFastPathMap() { return FastPathMap; }
 
-private:
+	FORCEINLINE void SetLODDistance(float InLODDistance, float InMaxLODDistance);
 
 	void DestroyDataInterfaceInstanceData();
 
@@ -290,6 +291,10 @@ private:
 	/** The tick count of the System instance. */
 	int32 TickCount;
 
+	/** LODDistance driven by our component. */
+	float LODDistance;
+	float MaxLODDistance;
+
 	TMap<FNiagaraDataSetID, FNiagaraDataSet> ExternalEvents;
 
 	TArray< TSharedRef<FNiagaraEmitterInstance, ESPMode::ThreadSafe> > Emitters;
@@ -345,6 +350,7 @@ private:
 	FNiagaraParameterDirectBinding<int32> SystemTickCountParam;
 
 	FNiagaraParameterDirectBinding<float> OwnerLODDistanceParam;
+	FNiagaraParameterDirectBinding<float> OwnerLODDistanceFractionParam;
 	FNiagaraParameterDirectBinding<int32> SystemNumEmittersParam;
 	FNiagaraParameterDirectBinding<int32> SystemNumEmittersAliveParam;
 
@@ -383,6 +389,8 @@ private:
 
 	uint32 bAlreadyBound : 1;
 
+	uint32 bLODDistanceIsValid : 1;
+
 	/** True if we have async work in flight. */
 	volatile bool bAsyncWorkInProgress;
 
@@ -417,6 +425,7 @@ public:
 		FVector OldPos;
 
 		float LODDistance;
+		float MaxLODDistance;
 		float TimeSeconds;
 		float RealTimeSeconds;
 
@@ -449,3 +458,10 @@ public:
 	TArray<TNiagaraFastPathUserParameterInputBinding<int32>> FastPathIntUserParameterInputBindings;
 	TArray<TNiagaraFastPathUserParameterInputBinding<float>> FastPathFloatUserParameterInputBindings;
 };
+
+FORCEINLINE void FNiagaraSystemInstance::SetLODDistance(float InLODDistance, float InMaxLODDistance)
+{
+	bLODDistanceIsValid = true;
+	LODDistance = InLODDistance; 
+	MaxLODDistance = InMaxLODDistance;
+}
\ No newline at end of file
diff --git a/Engine/Plugins/FX/Niagara/Source/Niagara/Public/NiagaraSystemSimulation.h b/Engine/Plugins/FX/Niagara/Source/Niagara/Public/NiagaraSystemSimulation.h
--- a/Engine/Plugins/FX/Niagara/Source/Niagara/Public/NiagaraSystemSimulation.h
+++ b/Engine/Plugins/FX/Niagara/Source/Niagara/Public/NiagaraSystemSimulation.h
@@ -8,6 +8,7 @@
 #include "NiagaraDataSet.h"
 #include "NiagaraScriptExecutionContext.h"
 
+class UNiagaraEffectType;
 class UWorld;
 class UNiagaraParameterCollection;
 class UNiagaraParameterCollectionInstance;
@@ -146,10 +147,16 @@ struct FNiagaraSystemSimulationTickContext
 };
 
 /** Simulation performing all system and emitter scripts for a instances of a UNiagaraSystem in a world. */
-class FNiagaraSystemSimulation : public TSharedFromThis<FNiagaraSystemSimulation, ESPMode::ThreadSafe>
+class FNiagaraSystemSimulation : public TSharedFromThis<FNiagaraSystemSimulation, ESPMode::ThreadSafe>, FGCObject
 {
 	friend FNiagaraSystemSimulationTickContext;
 public:
+
+	//FGCObject Interface
+	virtual void AddReferencedObjects(FReferenceCollector& Collector)override;
+	//FGCObject Interface END
+
+	FNiagaraSystemSimulation();
 	~FNiagaraSystemSimulation();
 	bool Init(UNiagaraSystem* InSystem, UWorld* InWorld, bool bInIsSolo, ETickingGroup TickGroup);
 	void Destroy();
@@ -224,6 +231,9 @@ protected:
 	/** System of instances being simulated.  We use a weak object ptr here because once the last referencing object goes away this system may be come invalid at runtime. */
 	TWeakObjectPtr<UNiagaraSystem> WeakSystem;
 
+	/** We cache off the effect type in the unlikely even that someone GCs the System from under us so that we can keep the effect types instance count etc accurate. */
+	UNiagaraEffectType* EffectType;
+
 	/** Which tick group we are in, only valid when not in Solo mode. */
 	ETickingGroup SystemTickGroup = TG_MAX;
 
diff --git a/Engine/Plugins/FX/Niagara/Source/Niagara/Public/NiagaraWorldManager.h b/Engine/Plugins/FX/Niagara/Source/Niagara/Public/NiagaraWorldManager.h
--- a/Engine/Plugins/FX/Niagara/Source/Niagara/Public/NiagaraWorldManager.h
+++ b/Engine/Plugins/FX/Niagara/Source/Niagara/Public/NiagaraWorldManager.h
@@ -16,6 +16,8 @@
 #include "GlobalDistanceFieldParameters.h"
 #include "NiagaraDataInterfaceSkeletalMesh.h"
 #include "NiagaraComponentPool.h"
+#include "NiagaraEffectType.h"
+#include "NiagaraScalabilityManager.h"
 
 #include "NiagaraWorldManager.generated.h"
 
@@ -23,6 +25,7 @@ class UWorld;
 class UNiagaraParameterCollection;
 class UNiagaraParameterCollectionInstance;
 class UNiagaraComponentPool;
+struct FNiagaraScalabilityState;
 
 class FNiagaraViewDataMgr : public FRenderResource
 {
@@ -128,10 +131,32 @@ public:
 
 	UNiagaraComponentPool* GetComponentPool() { return ComponentPool; }
 
+	void UpdateScalabilityManagers();
+
 	// Dump details about what's inside the world manager
 	void DumpDetails(FOutputDevice& Ar);
 	
 	UWorld* GetWorld();
+	FORCEINLINE UWorld* GetWorld()const { return World; }
+
+	//Various helper functions for scalability culling.
+	
+	void RegisterWithScalabilityManager(UNiagaraComponent* Component);
+	void UnregisterWithScalabilityManager(UNiagaraComponent* Component);
+
+	/** Should we cull an instance of this system at the passed location before it's even been spawned? */
+	NIAGARA_API bool ShouldPreCull(UNiagaraSystem* System, UNiagaraComponent* Component);
+	NIAGARA_API bool ShouldPreCull(UNiagaraSystem* System, FVector Location);
+
+	void CalculateScalabilityState(UNiagaraSystem* System, const FNiagaraScalabilitySettings& ScalabilitySettings, UNiagaraEffectType* EffectType, UNiagaraComponent* Component, bool bIsPreCull, FNiagaraScalabilityState& OutState);
+	void CalculateScalabilityState(UNiagaraSystem* System, const FNiagaraScalabilitySettings& ScalabilitySettings, UNiagaraEffectType* EffectType, FVector Location, bool bIsPreCull, FNiagaraScalabilityState& OutState);
+
+	/*FORCEINLINE_DEBUGGABLE*/ void SortedSignificanceCull(UNiagaraEffectType* EffectType, const FNiagaraScalabilitySettings& ScalabilitySettings, float Significance, int32 Index, FNiagaraScalabilityState& OutState);
+
+#if DEBUG_SCALABILITY_STATE
+	void DumpScalabilityState();
+#endif
+
 private:
 
 	// Callback function registered with global world delegates to instantiate world manager when a game world is created
@@ -152,6 +177,17 @@ private:
 	// Gamethread callback to cleanup references to the given batcher before it gets deleted on the renderthread.
 	void OnBatcherDestroyed_Internal(NiagaraEmitterInstanceBatcher* InBatcher);
 
+	FORCEINLINE_DEBUGGABLE bool CanPreCull(UNiagaraEffectType* EffectType);
+
+	FORCEINLINE_DEBUGGABLE void SignificanceCull(UNiagaraEffectType* EffectType, const FNiagaraScalabilitySettings& ScalabilitySettings, float Significance, FNiagaraScalabilityState& OutState);
+	FORCEINLINE_DEBUGGABLE void VisibilityCull(UNiagaraEffectType* EffectType, const FNiagaraScalabilitySettings& ScalabilitySettings, UNiagaraComponent* Component, FNiagaraScalabilityState& OutState);
+	FORCEINLINE_DEBUGGABLE void OwnerLODCull(UNiagaraEffectType* EffectType, const FNiagaraScalabilitySettings& ScalabilitySettings, UNiagaraComponent* Component, FNiagaraScalabilityState& OutState);
+	FORCEINLINE_DEBUGGABLE void InstanceCountCull(UNiagaraEffectType* EffectType, const FNiagaraScalabilitySettings& ScalabilitySettings, FNiagaraScalabilityState& OutState);
+
+	/** Calculate significance contribution from the distance to nearest view. */
+	FORCEINLINE_DEBUGGABLE float DistanceSignificance(UNiagaraEffectType* EffectType, const FNiagaraScalabilitySettings& ScalabilitySettings, FVector Location);
+	FORCEINLINE_DEBUGGABLE float DistanceSignificance(UNiagaraEffectType* EffectType, const FNiagaraScalabilitySettings& ScalabilitySettings, UNiagaraComponent* Component);
+	
 	static FDelegateHandle OnWorldInitHandle;
 	static FDelegateHandle OnWorldCleanupHandle;
 	static FDelegateHandle OnPreWorldFinishDestroyHandle;
@@ -189,5 +225,8 @@ private:
 	static constexpr int NumDeferredQueues = 3;
 	int DeferredDeletionQueueIndex = 0;
 	FDeferredDeletionQueue DeferredDeletionQueue[NumDeferredQueues];
+
+	UPROPERTY(transient)
+	TMap<UNiagaraEffectType*, FNiagaraScalabilityManager> ScalabilityManagers;
 };
 
diff --git a/Engine/Plugins/FX/Niagara/Source/NiagaraEditor/Private/AssetTypeActions/AssetTypeActions_NiagaraEffectType.cpp b/Engine/Plugins/FX/Niagara/Source/NiagaraEditor/Private/AssetTypeActions/AssetTypeActions_NiagaraEffectType.cpp
new file mode 100644
--- /dev/null
+++ b/Engine/Plugins/FX/Niagara/Source/NiagaraEditor/Private/AssetTypeActions/AssetTypeActions_NiagaraEffectType.cpp
@@ -0,0 +1,40 @@
+// Copyright 1998-2019 Epic Games, Inc. All Rights Reserved.
+#include "AssetTypeActions_NiagaraEffectType.h"
+#include "NiagaraEffectType.h"
+#include "NiagaraEffectTypeFactoryNew.h"
+#include "NiagaraEditorStyle.h"
+
+#include "IContentBrowserSingleton.h"
+#include "ContentBrowserModule.h"
+
+#include "ToolMenus.h"
+#include "EditorStyleSet.h"
+
+#define LOCTEXT_NAMESPACE "AssetTypeActions_NiagaraEffectType"
+
+FColor FAssetTypeActions_NiagaraEffectType::GetTypeColor() const
+{
+	return FNiagaraEditorStyle::Get().GetColor("NiagaraEditor.AssetColors.EffectType").ToFColor(true);
+}
+
+// void FAssetTypeActions_NiagaraEffectType::OpenAssetEditor(const TArray<UObject*>& InObjects, TSharedPtr<IToolkitHost> EditWithinLevelEditor)
+// {
+// 	EToolkitMode::Type Mode = EditWithinLevelEditor.IsValid() ? EToolkitMode::WorldCentric : EToolkitMode::Standalone;
+// 
+// 	for (auto ObjIt = InObjects.CreateConstIterator(); ObjIt; ++ObjIt)
+// 	{
+// 		auto Instance = Cast<UNiagaraParameterCollectionInstance>(*ObjIt);
+// 		if (Instance != NULL)
+// 		{
+// 			TSharedRef< FNiagaraParameterCollectionToolkit > NewNiagaraNPCToolkit(new FNiagaraParameterCollectionToolkit());
+// 			NewNiagaraNPCToolkit->Initialize(Mode, EditWithinLevelEditor, Instance);
+// 		}
+// 	}
+// }
+
+UClass* FAssetTypeActions_NiagaraEffectType::GetSupportedClass() const
+{
+	return UNiagaraEffectType::StaticClass();
+}
+
+#undef LOCTEXT_NAMESPACE
diff --git a/Engine/Plugins/FX/Niagara/Source/NiagaraEditor/Private/AssetTypeActions/AssetTypeActions_NiagaraEffectType.h b/Engine/Plugins/FX/Niagara/Source/NiagaraEditor/Private/AssetTypeActions/AssetTypeActions_NiagaraEffectType.h
new file mode 100644
--- /dev/null
+++ b/Engine/Plugins/FX/Niagara/Source/NiagaraEditor/Private/AssetTypeActions/AssetTypeActions_NiagaraEffectType.h
@@ -0,0 +1,19 @@
+// Copyright 1998-2019 Epic Games, Inc. All Rights Reserved.
+
+#pragma once
+
+#include "CoreMinimal.h"
+#include "Toolkits/IToolkitHost.h"
+#include "AssetTypeActions_Base.h"
+#include "NiagaraEditorModule.h"
+
+class FAssetTypeActions_NiagaraEffectType : public FAssetTypeActions_Base
+{
+public:
+	// IAssetTypeActions Implementation
+	virtual FText GetName() const override { return NSLOCTEXT("AssetTypeActions", "AssetTypeActions_NiagaraEffectType", "Niagara Effect Type"); }
+	virtual FColor GetTypeColor() const override;
+	virtual UClass* GetSupportedClass() const override;
+	//virtual void OpenAssetEditor(const TArray<UObject*>& InObjects, TSharedPtr<class IToolkitHost> EditWithinLevelEditor = TSharedPtr<IToolkitHost>()) override;
+	virtual uint32 GetCategories() override { return FNiagaraEditorModule::GetAssetCategory(); }
+};
\ No newline at end of file
diff --git a/Engine/Plugins/FX/Niagara/Source/NiagaraEditor/Private/NiagaraEditorModule.cpp b/Engine/Plugins/FX/Niagara/Source/NiagaraEditor/Private/NiagaraEditorModule.cpp
--- a/Engine/Plugins/FX/Niagara/Source/NiagaraEditor/Private/NiagaraEditorModule.cpp
+++ b/Engine/Plugins/FX/Niagara/Source/NiagaraEditor/Private/NiagaraEditorModule.cpp
@@ -18,6 +18,7 @@
 #include "AssetTypeActions/AssetTypeActions_NiagaraEmitter.h"
 #include "AssetTypeActions/AssetTypeActions_NiagaraScript.h"
 #include "AssetTypeActions/AssetTypeActions_NiagaraParameterCollection.h"
+#include "AssetTypeActions/AssetTypeActions_NiagaraEffectType.h"
 
 #include "EdGraphUtilities.h"
 #include "SGraphPin.h"
@@ -609,8 +610,9 @@ void FNiagaraEditorModule::StartupModule()
 	RegisterAssetTypeAction(AssetTools, MakeShareable(new FAssetTypeActions_NiagaraScriptFunctions()));
 	RegisterAssetTypeAction(AssetTools, MakeShareable(new FAssetTypeActions_NiagaraScriptModules()));
 	RegisterAssetTypeAction(AssetTools, MakeShareable(new FAssetTypeActions_NiagaraScriptDynamicInputs()));
-	RegisterAssetTypeAction(AssetTools, MakeShareable(new FAssetTypeActions_NiagaraParameterCollection()));
+	RegisterAssetTypeAction(AssetTools, MakeShareable(new FAssetTypeActions_NiagaraParameterCollection())); 
 	RegisterAssetTypeAction(AssetTools, MakeShareable(new FAssetTypeActions_NiagaraParameterCollectionInstance()));
+	RegisterAssetTypeAction(AssetTools, MakeShareable(new FAssetTypeActions_NiagaraEffectType()));
 
 	UNiagaraSettings::OnSettingsChanged().AddRaw(this, &FNiagaraEditorModule::OnNiagaraSettingsChangedEvent);
 	FCoreUObjectDelegates::GetPreGarbageCollectDelegate().AddRaw(this, &FNiagaraEditorModule::OnPreGarbageCollection);
diff --git a/Engine/Plugins/FX/Niagara/Source/NiagaraEditor/Private/NiagaraEffectTypeFactoryNew.cpp b/Engine/Plugins/FX/Niagara/Source/NiagaraEditor/Private/NiagaraEffectTypeFactoryNew.cpp
new file mode 100644
--- /dev/null
+++ b/Engine/Plugins/FX/Niagara/Source/NiagaraEditor/Private/NiagaraEffectTypeFactoryNew.cpp
@@ -0,0 +1,29 @@
+// Copyright 1998-2019 Epic Games, Inc. All Rights Reserved.
+
+#include "NiagaraEffectTypeFactoryNew.h"
+#include "CoreMinimal.h"
+#include "Misc/ConfigCacheIni.h"
+#include "NiagaraSettings.h"
+#include "NiagaraEffectType.h"
+
+#define LOCTEXT_NAMESPACE "NiagaraEffectTypeFactory"
+
+UNiagaraEffectTypeFactoryNew::UNiagaraEffectTypeFactoryNew(const FObjectInitializer& ObjectInitializer)
+	: Super(ObjectInitializer)
+{
+	SupportedClass = UNiagaraEffectType::StaticClass();
+	bCreateNew = false;
+	bEditAfterNew = true;
+	bCreateNew = true;
+}
+
+UObject* UNiagaraEffectTypeFactoryNew::FactoryCreateNew(UClass* Class, UObject* InParent, FName Name, EObjectFlags Flags, UObject* Context, FFeedbackContext* Warn)
+{
+	check(Class->IsChildOf(UNiagaraEffectType::StaticClass()));
+
+	UNiagaraEffectType* NewEffectType = NewObject<UNiagaraEffectType>(InParent, Class, Name, Flags | RF_Transactional);
+
+	return NewEffectType;
+}
+
+#undef LOCTEXT_NAMESPACE
\ No newline at end of file
diff --git a/Engine/Plugins/FX/Niagara/Source/NiagaraEditor/Public/NiagaraEffectTypeFactoryNew.h b/Engine/Plugins/FX/Niagara/Source/NiagaraEditor/Public/NiagaraEffectTypeFactoryNew.h
new file mode 100644
--- /dev/null
+++ b/Engine/Plugins/FX/Niagara/Source/NiagaraEditor/Public/NiagaraEffectTypeFactoryNew.h
@@ -0,0 +1,18 @@
+// Copyright 1998-2019 Epic Games, Inc. All Rights Reserved.
+
+#pragma once
+
+#include "CoreMinimal.h"
+#include "UObject/ObjectMacros.h"
+#include "Factories/Factory.h"
+#include "NiagaraEffectTypeFactoryNew.generated.h"
+
+UCLASS(hidecategories = Object)
+class UNiagaraEffectTypeFactoryNew : public UFactory
+{
+	GENERATED_UCLASS_BODY()
+
+	//~ Begin UFactory Interface
+	virtual UObject* FactoryCreateNew(UClass* Class, UObject* InParent, FName Name, EObjectFlags Flags, UObject* Context, FFeedbackContext* Warn) override;
+	//~ Begin UFactory Interface	
+};
\ No newline at end of file
diff --git a/Engine/Source/Runtime/Engine/Public/InGamePerformanceTracker.h b/Engine/Source/Runtime/Engine/Public/InGamePerformanceTracker.h
--- a/Engine/Source/Runtime/Engine/Public/InGamePerformanceTracker.h
+++ b/Engine/Source/Runtime/Engine/Public/InGamePerformanceTracker.h
@@ -5,8 +5,16 @@
 #include "CoreMinimal.h"
 #include "HAL/PlatformTime.h"
 
-struct FInGameCycleHistory
+struct ENGINE_API FInGameCycleHistory
 {
+	FInGameCycleHistory()
+		: FrameIdx(0)
+		, TotalCycles(0)
+		, ValidFrames(0)
+		, CurrFrameCycles(0)
+	{
+	}
+
 	FInGameCycleHistory(int32 InNumSamples)
 	: FrameIdx(0)
 	, TotalCycles(0)
@@ -64,7 +72,7 @@ struct FInGameCycleHistory
 Helper class to track code timings.
 Similar to stats but much more minimal as they're intended to be used at InGame. Though it's still not free so use wisely.
 */
-class FInGamePerformanceTracker
+class ENGINE_API FInGamePerformanceTracker
 {
 private:
 
@@ -100,9 +108,9 @@ public:
 	/** Enters a timed section of code. Can be re-entrant and will take the final exit point time. Only callable from game thread. */
 	void ExitTimedSection();
 	
-	ENGINE_API static class IConsoleVariable* Enabled;
-	ENGINE_API static int32 CachedEnabled;
-	ENGINE_API static class IConsoleVariable* HistorySize;
+	static class IConsoleVariable* Enabled;
+	static int32 CachedEnabled;
+	static class IConsoleVariable* HistorySize;
 };
 
 class FInGameCycleCounter
