Make Niagara material loading correctly pass the "loaded from cooked material" parameter around so we can load cooked data

diff --git a/Engine/Plugins/FX/Niagara/Source/NiagaraShader/Private/NiagaraShader.cpp b/Engine/Plugins/FX/Niagara/Source/NiagaraShader/Private/NiagaraShader.cpp
--- a/Engine/Plugins/FX/Niagara/Source/NiagaraShader/Private/NiagaraShader.cpp
+++ b/Engine/Plugins/FX/Niagara/Source/NiagaraShader/Private/NiagaraShader.cpp
@@ -1184,7 +1184,7 @@ void FNiagaraShaderMap::FlushShadersByShaderType(FShaderType* ShaderType)
 
 
 
-void FNiagaraShaderMap::Serialize(FArchive& Ar, bool bInlineShaderResources)
+void FNiagaraShaderMap::Serialize(FArchive& Ar, bool bInlineShaderResources, bool bLoadedByCookedMaterial)
 {
 	// Note: This is saved to the DDC, not into packages (except when cooked)
 	// Backwards compatibility therefore will not work based on the version of Ar
@@ -1205,7 +1205,7 @@ void FNiagaraShaderMap::Serialize(FArchive& Ar, bool bInlineShaderResources)
 
 	if (Ar.IsSaving() || Ar.IsLoading())
 	{
-		TShaderMap<FNiagaraShaderType>::SerializeInline(Ar, bInlineShaderResources, false, false);
+		TShaderMap<FNiagaraShaderType>::SerializeInline(Ar, bInlineShaderResources, false, bLoadedByCookedMaterial);
 	}
 }
 
diff --git a/Engine/Plugins/FX/Niagara/Source/NiagaraShader/Private/NiagaraShared.cpp b/Engine/Plugins/FX/Niagara/Source/NiagaraShader/Private/NiagaraShared.cpp
--- a/Engine/Plugins/FX/Niagara/Source/NiagaraShader/Private/NiagaraShared.cpp
+++ b/Engine/Plugins/FX/Niagara/Source/NiagaraShader/Private/NiagaraShared.cpp
@@ -145,7 +145,7 @@ void FNiagaraShaderScript::RegisterShaderMap()
 {
 	if (GameThreadShaderMap)
 	{
-		GameThreadShaderMap->RegisterSerializedShaders(false);
+		GameThreadShaderMap->RegisterSerializedShaders(bLoadedFromCookedMaterial);
 	}
 }
 
@@ -177,6 +177,11 @@ void FNiagaraShaderScript::SerializeShaderMap(FArchive& Ar)
 	bool bCooked = Ar.IsCooking();
 	Ar << bCooked;
 
+	if (Ar.IsLoading())
+	{
+		bLoadedFromCookedMaterial = bCooked;
+	}
+
 	if (FPlatformProperties::RequiresCookedData() && !bCooked && Ar.IsLoading())
 	{
 		UE_LOG(LogShaders, Fatal, TEXT("This platform requires cooked packages, and shaders were not cooked into this Niagara script %s."), *GetFriendlyName());
@@ -211,7 +216,7 @@ void FNiagaraShaderScript::SerializeShaderMap(FArchive& Ar)
 			if (bValid)
 			{
 				TRefCountPtr<FNiagaraShaderMap> LoadedShaderMap = new FNiagaraShaderMap();
-				LoadedShaderMap->Serialize(Ar);
+				LoadedShaderMap->Serialize(Ar, true, true);
 
 				// Toss the loaded shader data if this is a server only instance
 				//@todo - don't cook it in the first place
diff --git a/Engine/Plugins/FX/Niagara/Source/NiagaraShader/Public/NiagaraShared.h b/Engine/Plugins/FX/Niagara/Source/NiagaraShader/Public/NiagaraShared.h
--- a/Engine/Plugins/FX/Niagara/Source/NiagaraShader/Public/NiagaraShared.h
+++ b/Engine/Plugins/FX/Niagara/Source/NiagaraShader/Public/NiagaraShared.h
@@ -389,7 +389,7 @@ public:
 	static const FNiagaraShaderMap* GetShaderMapBeingCompiled(const FNiagaraShaderScript* Script);
 
 	/** Serializes the shader map. */
-	void Serialize(FArchive& Ar, bool bInlineShaderResources = true);
+	void Serialize(FArchive& Ar, bool bInlineShaderResources = true, bool bLoadedByCookedMaterial = false);
 
 	/** Saves this shader map to the derived data cache. */
 	void SaveToDerivedDataCache();
@@ -530,7 +530,8 @@ public:
 		GameThreadShaderMap(NULL),
 		RenderingThreadShaderMap(NULL),
 		FeatureLevel(GMaxRHIFeatureLevel),
-		bLoadedCookedShaderMapId(false)
+		bLoadedCookedShaderMapId(false),
+		bLoadedFromCookedMaterial(false)
 	{}
 
 	/**
@@ -777,6 +778,8 @@ private:
 	ERHIFeatureLevel::Type FeatureLevel;
 
 	uint32 bLoadedCookedShaderMapId : 1;
+	uint32 bLoadedFromCookedMaterial : 1;
+
 	FNiagaraShaderMapId CookedShaderMapId;
 
 	FOnNiagaraScriptCompilationComplete OnCompilationCompleteDelegate;
