Added alignment and distribution functionality for nodes in the niagara script graph
#rnx

diff --git a/Engine/Plugins/FX/Niagara/Source/NiagaraEditor/Private/NiagaraNode.cpp b/Engine/Plugins/FX/Niagara/Source/NiagaraEditor/Private/NiagaraNode.cpp
--- a/Engine/Plugins/FX/Niagara/Source/NiagaraEditor/Private/NiagaraNode.cpp
+++ b/Engine/Plugins/FX/Niagara/Source/NiagaraEditor/Private/NiagaraNode.cpp
@@ -7,8 +7,12 @@
 #include "GraphEditAction.h"
 #include "SNiagaraGraphNode.h"
 #include "Misc/SecureHash.h"
+#include "ToolMenuSection.h"
+#include "ToolMenu.h"
+#include "GraphEditorActions.h"
 #include "Serialization/PropertyLocalizationDataGathering.h"
 
+
 #define LOCTEXT_NAMESPACE "NiagaraNode"
 
 UNiagaraNode::UNiagaraNode(const FObjectInitializer& ObjectInitializer)
@@ -245,6 +249,36 @@ void UNiagaraNode::GetPinHoverText(const UEdGraphPin& Pin, FString& HoverTextOut
 	}
 }
 
+void UNiagaraNode::GetNodeContextMenuActions(class UToolMenu* Menu, class UGraphNodeContextMenuContext* Context) const
+{
+	{
+		FToolMenuSection& Section = Menu->AddSection("Alignment");
+		Section.AddSubMenu(
+			"Alignment",
+			LOCTEXT("AlignmentHeader", "Alignment"),
+			FText(),
+			FNewToolMenuDelegate::CreateLambda([](UToolMenu* InMenu)
+		{
+			{
+				FToolMenuSection& SubMenuSection = InMenu->AddSection("EdGraphSchemaAlignment", LOCTEXT("AlignHeader", "Align"));
+				SubMenuSection.AddMenuEntry(FGraphEditorCommands::Get().AlignNodesTop);
+				SubMenuSection.AddMenuEntry(FGraphEditorCommands::Get().AlignNodesMiddle);
+				SubMenuSection.AddMenuEntry(FGraphEditorCommands::Get().AlignNodesBottom);
+				SubMenuSection.AddMenuEntry(FGraphEditorCommands::Get().AlignNodesLeft);
+				SubMenuSection.AddMenuEntry(FGraphEditorCommands::Get().AlignNodesCenter);
+				SubMenuSection.AddMenuEntry(FGraphEditorCommands::Get().AlignNodesRight);
+				SubMenuSection.AddMenuEntry(FGraphEditorCommands::Get().StraightenConnections);
+			}
+
+			{
+				FToolMenuSection& SubMenuSection = InMenu->AddSection("EdGraphSchemaDistribution", LOCTEXT("DistributionHeader", "Distribution"));
+				SubMenuSection.AddMenuEntry(FGraphEditorCommands::Get().DistributeNodesHorizontally);
+				SubMenuSection.AddMenuEntry(FGraphEditorCommands::Get().DistributeNodesVertically);
+			}
+		}));
+	}
+}
+
 void UNiagaraNode::MarkNodeRequiresSynchronization(FString Reason, bool bRaiseGraphNeedsRecompile)
 {
 	Modify();
diff --git a/Engine/Plugins/FX/Niagara/Source/NiagaraEditor/Private/Widgets/SNiagaraScriptGraph.cpp b/Engine/Plugins/FX/Niagara/Source/NiagaraEditor/Private/Widgets/SNiagaraScriptGraph.cpp
--- a/Engine/Plugins/FX/Niagara/Source/NiagaraEditor/Private/Widgets/SNiagaraScriptGraph.cpp
+++ b/Engine/Plugins/FX/Niagara/Source/NiagaraEditor/Private/Widgets/SNiagaraScriptGraph.cpp
@@ -155,6 +155,43 @@ TSharedRef<SGraphEditor> SNiagaraScriptGraph::ConstructGraphEditor()
 	Commands->MapAction(
 		FGraphEditorCommands::Get().CreateComment,
 		FExecuteAction::CreateRaw(this, &SNiagaraScriptGraph::OnCreateComment));
+	// Alignment Commands
+	Commands->MapAction(FGraphEditorCommands::Get().AlignNodesTop,
+		FExecuteAction::CreateSP(this, &SNiagaraScriptGraph::OnAlignTop)
+	);
+
+	Commands->MapAction(FGraphEditorCommands::Get().AlignNodesMiddle,
+		FExecuteAction::CreateSP(this, &SNiagaraScriptGraph::OnAlignMiddle)
+	);
+
+	Commands->MapAction(FGraphEditorCommands::Get().AlignNodesBottom,
+		FExecuteAction::CreateSP(this, &SNiagaraScriptGraph::OnAlignBottom)
+	);
+
+	Commands->MapAction(FGraphEditorCommands::Get().AlignNodesLeft,
+		FExecuteAction::CreateSP(this, &SNiagaraScriptGraph::OnAlignLeft)
+	);
+
+	Commands->MapAction(FGraphEditorCommands::Get().AlignNodesCenter,
+		FExecuteAction::CreateSP(this, &SNiagaraScriptGraph::OnAlignCenter)
+	);
+
+	Commands->MapAction(FGraphEditorCommands::Get().AlignNodesRight,
+		FExecuteAction::CreateSP(this, &SNiagaraScriptGraph::OnAlignRight)
+	);
+
+	Commands->MapAction(FGraphEditorCommands::Get().StraightenConnections,
+		FExecuteAction::CreateSP(this, &SNiagaraScriptGraph::OnStraightenConnections)
+	);
+
+	// Distribution Commands
+	Commands->MapAction(FGraphEditorCommands::Get().DistributeNodesHorizontally,
+		FExecuteAction::CreateSP(this, &SNiagaraScriptGraph::OnDistributeNodesH)
+	);
+
+	Commands->MapAction(FGraphEditorCommands::Get().DistributeNodesVertically,
+		FExecuteAction::CreateSP(this, &SNiagaraScriptGraph::OnDistributeNodesV)
+	);
 	
 	TSharedRef<SGraphEditor> CreatedGraphEditor = SNew(SGraphEditor)
 		.AdditionalCommands(Commands.ToSharedRef())
@@ -501,6 +538,80 @@ FReply SNiagaraScriptGraph::HandleGraphSearchBoxKeyDown(const FGeometry& MyGeome
 	return FReply::Unhandled();
 }
 
+
+void SNiagaraScriptGraph::OnAlignTop()
+{
+	if (GraphEditor.IsValid())
+	{
+		GraphEditor->OnAlignTop();
+	}
+}
+
+void SNiagaraScriptGraph::OnAlignMiddle()
+{
+	if (GraphEditor.IsValid())
+	{
+		GraphEditor->OnAlignMiddle();
+	}
+}
+
+void SNiagaraScriptGraph::OnAlignBottom()
+{
+	if (GraphEditor.IsValid())
+	{
+		GraphEditor->OnAlignBottom();
+	}
+}
+
+void SNiagaraScriptGraph::OnAlignLeft()
+{
+	if (GraphEditor.IsValid())
+	{
+		GraphEditor->OnAlignLeft();
+	}
+}
+
+void SNiagaraScriptGraph::OnAlignCenter()
+{
+	if (GraphEditor.IsValid())
+	{
+		GraphEditor->OnAlignCenter();
+	}
+}
+
+void SNiagaraScriptGraph::OnAlignRight()
+{
+	if (GraphEditor.IsValid())
+	{
+		GraphEditor->OnAlignRight();
+	}
+}
+
+void SNiagaraScriptGraph::OnStraightenConnections()
+{
+	if (GraphEditor.IsValid())
+	{
+		GraphEditor->OnStraightenConnections();
+	}
+}
+
+void SNiagaraScriptGraph::OnDistributeNodesH()
+{
+	if (GraphEditor.IsValid())
+	{
+		GraphEditor->OnDistributeNodesH();
+	}
+}
+
+void SNiagaraScriptGraph::OnDistributeNodesV()
+{
+	if (GraphEditor.IsValid())
+	{
+		GraphEditor->OnDistributeNodesV();
+	}
+}
+
+
 void SNiagaraScriptGraph::FocusGraphSearchBox()
 {
 	bGraphSearchBoxActive = true;
diff --git a/Engine/Plugins/FX/Niagara/Source/NiagaraEditor/Private/Widgets/SNiagaraScriptGraph.h b/Engine/Plugins/FX/Niagara/Source/NiagaraEditor/Private/Widgets/SNiagaraScriptGraph.h
--- a/Engine/Plugins/FX/Niagara/Source/NiagaraEditor/Private/Widgets/SNiagaraScriptGraph.h
+++ b/Engine/Plugins/FX/Niagara/Source/NiagaraEditor/Private/Widgets/SNiagaraScriptGraph.h
@@ -81,6 +81,19 @@ private:
 
 	FReply HandleGraphSearchBoxKeyDown(const FGeometry& MyGeometry, const FKeyEvent& InKeyEvent);
 
+	void OnAlignTop();
+	void OnAlignMiddle();
+	void OnAlignBottom();
+	void OnAlignLeft();
+	void OnAlignCenter();
+	void OnAlignRight();
+
+	void OnStraightenConnections();
+
+	void OnDistributeNodesH();
+	void OnDistributeNodesV();
+
+
 private:
 	/** The combined commands used by the graph editor. */
 	TSharedPtr<FUICommandList> Commands;
diff --git a/Engine/Plugins/FX/Niagara/Source/NiagaraEditor/Public/NiagaraNode.h b/Engine/Plugins/FX/Niagara/Source/NiagaraEditor/Public/NiagaraNode.h
--- a/Engine/Plugins/FX/Niagara/Source/NiagaraEditor/Public/NiagaraNode.h
+++ b/Engine/Plugins/FX/Niagara/Source/NiagaraEditor/Public/NiagaraNode.h
@@ -45,6 +45,7 @@ public:
 	virtual void NodeConnectionListChanged() override;	
 	virtual TSharedPtr<SGraphNode> CreateVisualWidget() override; 
 	virtual void GetPinHoverText(const UEdGraphPin& Pin, FString& HoverTextOut) const override;
+	virtual void GetNodeContextMenuActions(class UToolMenu* Menu, class UGraphNodeContextMenuContext* Context) const override;
 	//~ End EdGraphNode Interface
 
 	/** Get the Niagara graph that owns this node */
