Alignment and distribution of system overview nodes in Niagara
#rnx

diff --git a/Engine/Plugins/FX/Niagara/Source/NiagaraEditor/Private/EdGraphSchema_NiagaraSystemOverview.cpp b/Engine/Plugins/FX/Niagara/Source/NiagaraEditor/Private/EdGraphSchema_NiagaraSystemOverview.cpp
--- a/Engine/Plugins/FX/Niagara/Source/NiagaraEditor/Private/EdGraphSchema_NiagaraSystemOverview.cpp
+++ b/Engine/Plugins/FX/Niagara/Source/NiagaraEditor/Private/EdGraphSchema_NiagaraSystemOverview.cpp
@@ -15,7 +15,6 @@ UEdGraphSchema_NiagaraSystemOverview::UEdGraphSchema_NiagaraSystemOverview(const
 void UEdGraphSchema_NiagaraSystemOverview::GetContextMenuActions(UToolMenu* Menu, UGraphNodeContextMenuContext* Context) const
 {
 	//@TODO System Overview: enable/disable emitter, solo, etc. 
-	return;
 }
 
 bool UEdGraphSchema_NiagaraSystemOverview::IsCacheVisualizationOutOfDate(int32 InVisualizationCacheID) const
diff --git a/Engine/Plugins/FX/Niagara/Source/NiagaraEditor/Private/NiagaraOverviewNode.cpp b/Engine/Plugins/FX/Niagara/Source/NiagaraEditor/Private/NiagaraOverviewNode.cpp
--- a/Engine/Plugins/FX/Niagara/Source/NiagaraEditor/Private/NiagaraOverviewNode.cpp
+++ b/Engine/Plugins/FX/Niagara/Source/NiagaraEditor/Private/NiagaraOverviewNode.cpp
@@ -6,6 +6,9 @@
 #include "ViewModels/Stack/NiagaraStackEntry.h"
 
 #include "Modules/ModuleManager.h"
+#include "ToolMenuSection.h"
+#include "ToolMenu.h"
+#include "GraphEditorActions.h"
 
 #define LOCTEXT_NAMESPACE "NiagaraOverviewNodeStackItem"
 
@@ -82,6 +85,30 @@ bool UNiagaraOverviewNode::CanDuplicateNode() const
 	return HasAnyFlags(RF_ClassDefaultObject) || EmitterHandleGuid.IsValid();
 }
 
+void UNiagaraOverviewNode::GetNodeContextMenuActions(class UToolMenu* Menu, class UGraphNodeContextMenuContext* Context) const
+{
+	FToolMenuSection& Section = Menu->AddSection("Alignment");
+	Section.AddSubMenu(
+		"Alignment",
+		LOCTEXT("AlignmentHeader", "Alignment"),
+		FText(),
+		FNewToolMenuDelegate::CreateLambda([](UToolMenu* InMenu)
+	{
+		{
+			FToolMenuSection& SubMenuSection = InMenu->AddSection("EdGraphSchemaAlignment", LOCTEXT("AlignHeader", "Align"));
+			SubMenuSection.AddMenuEntry(FGraphEditorCommands::Get().AlignNodesTop);
+			SubMenuSection.AddMenuEntry(FGraphEditorCommands::Get().AlignNodesMiddle);
+			SubMenuSection.AddMenuEntry(FGraphEditorCommands::Get().AlignNodesBottom);
+		}
+
+		{
+			FToolMenuSection& SubMenuSection = InMenu->AddSection("EdGraphSchemaDistribution", LOCTEXT("DistributionHeader", "Distribution"));
+			SubMenuSection.AddMenuEntry(FGraphEditorCommands::Get().DistributeNodesHorizontally);
+			SubMenuSection.AddMenuEntry(FGraphEditorCommands::Get().DistributeNodesVertically);
+		}
+	}));
+}
+
 UNiagaraSystem* UNiagaraOverviewNode::GetOwningSystem()
 {
 	return OwningSystem;
diff --git a/Engine/Plugins/FX/Niagara/Source/NiagaraEditor/Public/EdGraphSchema_NiagaraSystemOverview.h b/Engine/Plugins/FX/Niagara/Source/NiagaraEditor/Public/EdGraphSchema_NiagaraSystemOverview.h
--- a/Engine/Plugins/FX/Niagara/Source/NiagaraEditor/Public/EdGraphSchema_NiagaraSystemOverview.h
+++ b/Engine/Plugins/FX/Niagara/Source/NiagaraEditor/Public/EdGraphSchema_NiagaraSystemOverview.h
@@ -16,7 +16,6 @@ class NIAGARAEDITOR_API UEdGraphSchema_NiagaraSystemOverview : public UEdGraphSc
 
 	//~ Begin EdGraphSchema Interface
 	virtual void GetContextMenuActions(class UToolMenu* Menu, class UGraphNodeContextMenuContext* Context) const override;
-	virtual FName GetParentContextMenuName() const override { return NAME_None; }
 	virtual const FPinConnectionResponse CanCreateConnection(const UEdGraphPin* A, const UEdGraphPin* B) const override { return FPinConnectionResponse(CONNECT_RESPONSE_DISALLOW, FText()); }; //@TODO System Overview: write text response
 	virtual bool TryCreateConnection(UEdGraphPin* A, UEdGraphPin* B) const override { return false; };
 	virtual bool ShouldAlwaysPurgeOnModification() const override { return false; }
diff --git a/Engine/Plugins/FX/Niagara/Source/NiagaraEditor/Public/NiagaraOverviewNode.h b/Engine/Plugins/FX/Niagara/Source/NiagaraEditor/Public/NiagaraOverviewNode.h
--- a/Engine/Plugins/FX/Niagara/Source/NiagaraEditor/Public/NiagaraOverviewNode.h
+++ b/Engine/Plugins/FX/Niagara/Source/NiagaraEditor/Public/NiagaraOverviewNode.h
@@ -27,6 +27,8 @@ public:
 
 	/** Whether or not this node can be safely duplicated (via copy/paste, etc...) in the graph */
 	virtual bool CanDuplicateNode() const override;
+
+	virtual void GetNodeContextMenuActions(class UToolMenu* Menu, class UGraphNodeContextMenuContext* Context) const override;
 	//~ End UEdGraphNode Interface
 
 	UNiagaraSystem* GetOwningSystem();
diff --git a/Engine/Plugins/FX/Niagara/Source/NiagaraEditorWidgets/Private/SNiagaraOverviewGraph.cpp b/Engine/Plugins/FX/Niagara/Source/NiagaraEditorWidgets/Private/SNiagaraOverviewGraph.cpp
--- a/Engine/Plugins/FX/Niagara/Source/NiagaraEditorWidgets/Private/SNiagaraOverviewGraph.cpp
+++ b/Engine/Plugins/FX/Niagara/Source/NiagaraEditorWidgets/Private/SNiagaraOverviewGraph.cpp
@@ -76,7 +76,29 @@ void SNiagaraOverviewGraph::Construct(const FArguments& InArgs, TSharedRef<FNiag
 	Commands->MapAction(
 		FNiagaraEditorModule::Get().GetCommands().ZoomToFitAll,
 		FExecuteAction::CreateSP(this, &SNiagaraOverviewGraph::ZoomToFitAll));
-	
+	// Alignment Commands
+	Commands->MapAction(FGraphEditorCommands::Get().AlignNodesTop,
+		FExecuteAction::CreateSP(this, &SNiagaraOverviewGraph::OnAlignTop)
+	);
+
+	Commands->MapAction(FGraphEditorCommands::Get().AlignNodesMiddle,
+		FExecuteAction::CreateSP(this, &SNiagaraOverviewGraph::OnAlignMiddle)
+	);
+
+	Commands->MapAction(FGraphEditorCommands::Get().AlignNodesBottom,
+		FExecuteAction::CreateSP(this, &SNiagaraOverviewGraph::OnAlignBottom)
+	);
+
+	// Distribution Commands
+	Commands->MapAction(FGraphEditorCommands::Get().DistributeNodesHorizontally,
+		FExecuteAction::CreateSP(this, &SNiagaraOverviewGraph::OnDistributeNodesH)
+	);
+
+	Commands->MapAction(FGraphEditorCommands::Get().DistributeNodesVertically,
+		FExecuteAction::CreateSP(this, &SNiagaraOverviewGraph::OnDistributeNodesV)
+	);
+
+
 	GraphEditor = SNew(SGraphEditor)
 		.AdditionalCommands(Commands)
 		.Appearance(AppearanceInfo)
@@ -283,4 +305,45 @@ void SNiagaraOverviewGraph::ZoomToFitAll()
 	GraphEditor->ZoomToFit(false);
 }
 
+
+void SNiagaraOverviewGraph::OnAlignTop()
+{
+	if (GraphEditor.IsValid())
+	{
+		GraphEditor->OnAlignTop();
+	}
+}
+
+void SNiagaraOverviewGraph::OnAlignMiddle()
+{
+	if (GraphEditor.IsValid())
+	{
+		GraphEditor->OnAlignMiddle();
+	}
+}
+
+void SNiagaraOverviewGraph::OnAlignBottom()
+{
+	if (GraphEditor.IsValid())
+	{
+		GraphEditor->OnAlignBottom();
+	}
+}
+
+void SNiagaraOverviewGraph::OnDistributeNodesH()
+{
+	if (GraphEditor.IsValid())
+	{
+		GraphEditor->OnDistributeNodesH();
+	}
+}
+
+void SNiagaraOverviewGraph::OnDistributeNodesV()
+{
+	if (GraphEditor.IsValid())
+	{
+		GraphEditor->OnDistributeNodesV();
+	}
+}
+
 #undef LOCTEXT_NAMESPACE // "NiagaraOverviewGraph"
\ No newline at end of file
diff --git a/Engine/Plugins/FX/Niagara/Source/NiagaraEditorWidgets/Private/SNiagaraOverviewGraph.h b/Engine/Plugins/FX/Niagara/Source/NiagaraEditorWidgets/Private/SNiagaraOverviewGraph.h
--- a/Engine/Plugins/FX/Niagara/Source/NiagaraEditorWidgets/Private/SNiagaraOverviewGraph.h
+++ b/Engine/Plugins/FX/Niagara/Source/NiagaraEditorWidgets/Private/SNiagaraOverviewGraph.h
@@ -41,6 +41,14 @@ private:
 
 	void ZoomToFitAll();
 
+	void OnAlignTop();
+	void OnAlignMiddle();
+	void OnAlignBottom();
+
+	void OnDistributeNodesH();
+	void OnDistributeNodesV();
+
+
 private:
 	TSharedPtr<FNiagaraOverviewGraphViewModel> ViewModel;
 	TSharedPtr<SGraphEditor> GraphEditor;
