FNM - implemented niagarashader atomic compaticion for Console using wave intrinsics. Performance improvement on XB1S is about 33% and PS4P shows 10%. (The XB1 implementation was provided by Stu)

diff --git a/Engine/Plugins/FX/Niagara/Shaders/Private/NiagaraEmitterInstanceShader.usf b/Engine/Plugins/FX/Niagara/Shaders/Private/NiagaraEmitterInstanceShader.usf
--- a/Engine/Plugins/FX/Niagara/Shaders/Private/NiagaraEmitterInstanceShader.usf
+++ b/Engine/Plugins/FX/Niagara/Shaders/Private/NiagaraEmitterInstanceShader.usf
@@ -397,6 +397,21 @@ int rand(int x, int Seed1, int Seed2, int Seed3)
 			if(ShaderStageIndex == 0)
 			{
 				int PrevIdx = GSpawnStartInstance + SpawnedInstances;	// scratch instance as default; write to that for dead particles
+
+#if PS4_PROFILE || XBOXONE_PROFILE
+				uint NumCounters = WaveActiveCountBits(bDoAcquire);
+				if(NumCounters)
+				{
+					if(WaveIsFirstLane())
+					{
+						uint RetPrevIdx;
+						InterlockedAdd(RWInstanceCounts[WriteInstanceCountOffset], NumCounters, RetPrevIdx);
+						PrevIdx = (int)RetPrevIdx;
+					}
+					PrevIdx = WaveReadFirstLane(PrevIdx);
+					PrevIdx += WavePrefixCountBits(bDoAcquire);
+				}		
+#else
 				if(bDoAcquire == true)
 				{
 					// Have to use uint's here to avoid PS4 compiler warnings about InterlockedAdd, cannot propagate uint due to CPU VM limitations...
@@ -405,6 +420,7 @@ int rand(int x, int Seed1, int Seed2, int Seed3)
 					InterlockedAdd(RWInstanceCounts[WriteInstanceCountOffset], (uint)1U, RetPrevIdx);
 					PrevIdx = (int)RetPrevIdx;
 				}
+#endif
 				return PrevIdx;
 			}
 			else
@@ -413,6 +429,21 @@ int rand(int x, int Seed1, int Seed2, int Seed3)
 			}
 		#else
 			int PrevIdx = GSpawnStartInstance + SpawnedInstances;	// scratch instance as default; write to that for dead particles
+
+#if PS4_PROFILE || XBOXONE_PROFILE
+			uint NumCounters = WaveActiveCountBits(bDoAcquire);
+			if(NumCounters)
+			{
+				if(WaveIsFirstLane())
+				{
+					uint RetPrevIdx;
+					InterlockedAdd(RWInstanceCounts[WriteInstanceCountOffset], NumCounters, RetPrevIdx);
+					PrevIdx = (int)RetPrevIdx;
+				}
+				PrevIdx = WaveReadFirstLane(PrevIdx);
+				PrevIdx += WavePrefixCountBits(bDoAcquire);
+			}	
+#else
 			if(bDoAcquire == true)
 			{
 				// Have to use uint's here to avoid PS4 compiler warnings about InterlockedAdd, cannot propagate uint due to CPU VM limitations...
@@ -421,6 +452,7 @@ int rand(int x, int Seed1, int Seed2, int Seed3)
 				InterlockedAdd(RWInstanceCounts[WriteInstanceCountOffset], (uint)1U, RetPrevIdx);
 				PrevIdx = (int)RetPrevIdx;
 			}
+#endif
 			return PrevIdx;
 		#endif
 	}
diff --git a/Engine/Plugins/FX/Niagara/Source/Niagara/Private/NiagaraGPUInstanceCountManager.cpp b/Engine/Plugins/FX/Niagara/Source/Niagara/Private/NiagaraGPUInstanceCountManager.cpp
--- a/Engine/Plugins/FX/Niagara/Source/Niagara/Private/NiagaraGPUInstanceCountManager.cpp
+++ b/Engine/Plugins/FX/Niagara/Source/Niagara/Private/NiagaraGPUInstanceCountManager.cpp
@@ -105,7 +105,7 @@ void FNiagaraGPUInstanceCountManager::ResizeBuffers(FRHICommandListImmediate& RH
 			AllocatedInstanceCounts = RecommendedInstanceCounts;
 			TResourceArray<uint32> InitData;
 			InitData.AddZeroed(AllocatedInstanceCounts);
-			CountBuffer.Initialize(sizeof(uint32), AllocatedInstanceCounts, EPixelFormat::PF_R32_UINT, BUF_Static | BUF_KeepCPUAccessible, TEXT("NiagaraGPUInstanceCounts"), &InitData);
+			CountBuffer.Initialize(sizeof(uint32), AllocatedInstanceCounts, EPixelFormat::PF_R32_UINT, BUF_Static, TEXT("NiagaraGPUInstanceCounts"), &InitData);
 		}
 		// If we need to increase the buffer size to RecommendedInstanceCounts because the buffer is too small.
 		else if (RequiredInstanceCounts > AllocatedInstanceCounts)
@@ -114,7 +114,7 @@ void FNiagaraGPUInstanceCountManager::ResizeBuffers(FRHICommandListImmediate& RH
 			TResourceArray<uint32> InitData;
 			InitData.AddZeroed(RecommendedInstanceCounts);
 			FRWBuffer NextCountBuffer;
-			NextCountBuffer.Initialize(sizeof(uint32), RecommendedInstanceCounts, EPixelFormat::PF_R32_UINT, BUF_Static | BUF_KeepCPUAccessible, TEXT("NiagaraGPUInstanceCounts"), &InitData);
+			NextCountBuffer.Initialize(sizeof(uint32), RecommendedInstanceCounts, EPixelFormat::PF_R32_UINT, BUF_Static, TEXT("NiagaraGPUInstanceCounts"), &InitData);
 
 			// Because the shader works with SINT, we need to temporarily create view for the shader to work. Those will be (deferred) deleted at end of the scope.
 			FUnorderedAccessViewRHIRef NextCountBufferUAVAsInt = RHICreateUnorderedAccessView(NextCountBuffer.Buffer, EPixelFormat::PF_R32_SINT);
diff --git a/Engine/Plugins/FX/Niagara/Source/NiagaraCore/Private/NiagaraCustomVersion.cpp b/Engine/Plugins/FX/Niagara/Source/NiagaraCore/Private/NiagaraCustomVersion.cpp
--- a/Engine/Plugins/FX/Niagara/Source/NiagaraCore/Private/NiagaraCustomVersion.cpp
+++ b/Engine/Plugins/FX/Niagara/Source/NiagaraCore/Private/NiagaraCustomVersion.cpp
@@ -8,5 +8,5 @@ const FGuid FNiagaraCustomVersion::GUID(0xFCF57AFA, 0x50764283, 0xB9A9E658, 0xFF
 // Register the custom version with core
 FCustomVersionRegistration GRegisterNiagaraCustomVersion(FNiagaraCustomVersion::GUID, FNiagaraCustomVersion::LatestVersion, TEXT("NiagaraVer"));
 
-const FGuid FNiagaraCustomVersion::LatestScriptCompileVersion(0x6D2E1207, 0x63F99C45, 0xBC426396, 0xFBF80045);
+const FGuid FNiagaraCustomVersion::LatestScriptCompileVersion(0x02ABFEE7, 0x88B02847, 0x85046386, 0x8BD6F349);
 
