Null check to fix crash in UI

#rnx

diff --git a/Engine/Plugins/FX/Niagara/Source/NiagaraEditor/Private/ViewModels/Stack/NiagaraStackModuleItem.cpp b/Engine/Plugins/FX/Niagara/Source/NiagaraEditor/Private/ViewModels/Stack/NiagaraStackModuleItem.cpp
--- a/Engine/Plugins/FX/Niagara/Source/NiagaraEditor/Private/ViewModels/Stack/NiagaraStackModuleItem.cpp
+++ b/Engine/Plugins/FX/Niagara/Source/NiagaraEditor/Private/ViewModels/Stack/NiagaraStackModuleItem.cpp
@@ -592,7 +592,7 @@ void UNiagaraStackModuleItem::RefreshIssues(TArray<FStackIssue>& NewIssues)
 					for (FAssetData ModuleAsset : ModuleAssets)
 					{
 						UNiagaraScript* DependencyScript = Cast<UNiagaraScript>(ModuleAsset.GetAsset());
-						if (Dependency.ScriptConstraint == ENiagaraModuleDependencyScriptConstraint::SameScript)
+						if (Dependency.ScriptConstraint == ENiagaraModuleDependencyScriptConstraint::SameScript && DependencyScript)
 						{
 							TArray<ENiagaraScriptUsage> SupportedUsages = UNiagaraScript::GetSupportedUsageContextsForBitmask(DependencyScript->ModuleUsageBitmask);
 							if (SupportedUsages.Contains(OutputNode->GetUsage()) == false)
