- Keeping attached components attached to their parents when completing due to scalabilty culling so their positions are correct for subsequent reactivation checks.
- Fixed accessing default effect type.
- New helper functions to get the completion state of a niagara component.
- Fixed typo in scalabilty setting code.

#rnx

diff --git a/Engine/Plugins/FX/Niagara/Source/Niagara/Private/NiagaraComponent.cpp b/Engine/Plugins/FX/Niagara/Source/Niagara/Private/NiagaraComponent.cpp
--- a/Engine/Plugins/FX/Niagara/Source/Niagara/Private/NiagaraComponent.cpp
+++ b/Engine/Plugins/FX/Niagara/Source/Niagara/Private/NiagaraComponent.cpp
@@ -996,7 +996,7 @@ void UNiagaraComponent::OnSystemComplete()
 		DestroyComponent();
 		//UE_LOG(LogNiagara, Log, TEXT("OnSystemComplete DestroyComponent();: } %p - %s"), SystemInstance.Get(), *Asset->GetName());
 	}
-	else if (bAutoManageAttachment)
+	else if (bAutoManageAttachment && ScalabilityManagerHandle == INDEX_NONE)//Do not detach from our parent if we were deactivated by scalability and we need to be considered for reactivation.
 	{
 		CancelAutoAttachment(/*bDetachFromParent=*/ true);
 	}
@@ -1751,6 +1751,15 @@ void UNiagaraComponent::SetOwnerLOD(int32 InOwnerLOD)
 	OwnerLOD = InOwnerLOD;  
 }
 
+void UNiagaraComponent::SetAllowScalability(bool bAllow)
+{
+	bAllowScalability = bAllow; 
+	if (!bAllow)
+	{
+		UnregisterWithScalabilityManager();
+	}
+}
+
 #if WITH_EDITOR
 
 void UNiagaraComponent::PostLoadNormalizeOverrideNames()
diff --git a/Engine/Plugins/FX/Niagara/Source/Niagara/Private/NiagaraSettings.cpp b/Engine/Plugins/FX/Niagara/Source/Niagara/Private/NiagaraSettings.cpp
--- a/Engine/Plugins/FX/Niagara/Source/Niagara/Private/NiagaraSettings.cpp
+++ b/Engine/Plugins/FX/Niagara/Source/Niagara/Private/NiagaraSettings.cpp
@@ -49,5 +49,9 @@ UNiagaraSettings::FOnNiagaraSettingsChanged UNiagaraSettings::SettingsChangedDel
 
 UNiagaraEffectType* UNiagaraSettings::GetDefaultEffectType()const
 {
+	if (DefaultEffectTypePtr == nullptr && DefaultEffectType.IsValid())
+	{
+		DefaultEffectTypePtr = Cast<UNiagaraEffectType>(DefaultEffectType.TryLoad());
+	}
 	return DefaultEffectTypePtr;
 }
diff --git a/Engine/Plugins/FX/Niagara/Source/Niagara/Private/NiagaraSystem.cpp b/Engine/Plugins/FX/Niagara/Source/Niagara/Private/NiagaraSystem.cpp
--- a/Engine/Plugins/FX/Niagara/Source/Niagara/Private/NiagaraSystem.cpp
+++ b/Engine/Plugins/FX/Niagara/Source/Niagara/Private/NiagaraSystem.cpp
@@ -1390,7 +1390,7 @@ void UNiagaraSystem::ResolveScalabilityOverrides()
 			FNiagaraScalabilitySettings& NewSettings = ResolvedScalabilitySettings[NewIdx];
 			if (ActualEffectType)
 			{
-				const TArray<FNiagaraScalabilitySettings>& EffectTypeSettings = EffectType->GetScalabilitySettings();
+				const TArray<FNiagaraScalabilitySettings>& EffectTypeSettings = ActualEffectType->GetScalabilitySettings();
 				int32 LastIndex = EffectTypeSettings.Num() - 1;
 				if (LastIndex >= 0)
 				{
diff --git a/Engine/Plugins/FX/Niagara/Source/Niagara/Public/NiagaraComponent.h b/Engine/Plugins/FX/Niagara/Source/Niagara/Public/NiagaraComponent.h
--- a/Engine/Plugins/FX/Niagara/Source/Niagara/Public/NiagaraComponent.h
+++ b/Engine/Plugins/FX/Niagara/Source/Niagara/Public/NiagaraComponent.h
@@ -146,7 +146,10 @@ public:
 	virtual void Deactivate()override;
 	void DeactivateImmediate();
 
-	FORCEINLINE ENiagaraExecutionState GetRequestedExecutionState() { return SystemInstance ? SystemInstance->GetRequestedExecutionState() : ENiagaraExecutionState::Complete; }
+	FORCEINLINE ENiagaraExecutionState GetRequestedExecutionState()const { return SystemInstance ? SystemInstance->GetRequestedExecutionState() : ENiagaraExecutionState::Complete; }
+	FORCEINLINE ENiagaraExecutionState GetExecutionState()const { return SystemInstance ? SystemInstance->GetActualExecutionState() : ENiagaraExecutionState::Complete; }
+
+	FORCEINLINE bool IsComplete()const { return SystemInstance ? SystemInstance->IsComplete() : true; }
 
 	FORCEINLINE float GetSafeTimeSinceRendered(float WorldTime)const;
 	private:
@@ -505,14 +508,14 @@ public:
 	uint32 bWaitForCompilationOnActivate : 1;
 #endif
 
-	UFUNCTION(BlueprintCallable, Category = Preview, meta = (Keywords = "LOD scalability"))
-	void SetOwnerLOD(int32 InOwnerLOD);
+	virtual void SetOwnerLOD(int32 InOwnerLOD);
 
-	UFUNCTION(BlueprintCallable, Category = Preview, meta = (Keywords = "LOD scalability"))
+	UFUNCTION(BlueprintCallable, Category = Scalability, meta = (Keywords = "LOD scalability"))
 	FORCEINLINE int32 GetOwnerLOD()const { return OwnerLOD; }
 
 	/** Set whether this component is allowed to perform scalability checks and potentially be culled etc. Occasionally it is useful to disable this for specific components. E.g. Effects on the local player. */
-	FORCEINLINE void SetAllowScalability(bool bAllow) { bAllowScalability = bAllow; }
+	UFUNCTION(BlueprintCallable, Category = Scalability, meta = (Keywords = "LOD scalability"))
+	void SetAllowScalability(bool bAllow);
 
 private:
 	/** Did we try and activate but fail due to the asset being not yet ready. Keep looping.*/
@@ -550,7 +553,6 @@ private:
 	int32 OwnerLOD;
 };
 
-
 #if WITH_NIAGARA_COMPONENT_PREVIEW_DATA
 FORCEINLINE bool UNiagaraComponent::GetPreviewDetailLevelEnabled()const { return bEnablePreviewDetailLevel; }
 FORCEINLINE bool UNiagaraComponent::GetPreviewLODDistanceEnabled()const { return bEnablePreviewLODDistance; }
diff --git a/Engine/Plugins/FX/Niagara/Source/Niagara/Public/NiagaraSettings.h b/Engine/Plugins/FX/Niagara/Source/Niagara/Public/NiagaraSettings.h
--- a/Engine/Plugins/FX/Niagara/Source/Niagara/Public/NiagaraSettings.h
+++ b/Engine/Plugins/FX/Niagara/Source/Niagara/Public/NiagaraSettings.h
@@ -54,5 +54,5 @@ protected:
 
 private:
 	UPROPERTY(transient)
-	UNiagaraEffectType* DefaultEffectTypePtr;
+	mutable UNiagaraEffectType* DefaultEffectTypePtr;
 };
\ No newline at end of file
diff --git a/Engine/Plugins/FX/Niagara/Source/Niagara/Public/NiagaraSystemInstance.h b/Engine/Plugins/FX/Niagara/Source/Niagara/Public/NiagaraSystemInstance.h
--- a/Engine/Plugins/FX/Niagara/Source/Niagara/Public/NiagaraSystemInstance.h
+++ b/Engine/Plugins/FX/Niagara/Source/Niagara/Public/NiagaraSystemInstance.h
@@ -104,7 +104,7 @@ public:
 	/** Perform per-tick updates on data interfaces that need it. This can cause systems to complete so cannot be parallelized. */
 	void TickDataInterfaces(float DeltaSeconds, bool bPostSimulate);
 
-	ENiagaraExecutionState GetRequestedExecutionState() { return RequestedExecutionState; }
+	ENiagaraExecutionState GetRequestedExecutionState()const { return RequestedExecutionState; }
 	void SetRequestedExecutionState(ENiagaraExecutionState InState);
 
 	ENiagaraExecutionState GetActualExecutionState() { return ActualExecutionState; }
diff --git a/Engine/Source/Runtime/Engine/Classes/Particles/ParticleSystemComponent.h b/Engine/Source/Runtime/Engine/Classes/Particles/ParticleSystemComponent.h
--- a/Engine/Source/Runtime/Engine/Classes/Particles/ParticleSystemComponent.h
+++ b/Engine/Source/Runtime/Engine/Classes/Particles/ParticleSystemComponent.h
@@ -406,6 +406,9 @@ public:
 
 	/** Returns an approximate memory usage value for this component. */
 	virtual uint32 GetApproxMemoryUsage() const { return 0; }
+
+	UFUNCTION(BlueprintCallable, Category = Preview, meta = (Keywords = "LOD scalability"))
+	virtual void SetOwnerLOD(int32 LOD) {}
 };
 
 
