Niagara - Missed in my checkin from last week.

diff --git a/Engine/Plugins/FX/Niagara/Source/NiagaraEditorWidgets/Private/SNiagaraOverviewStack.cpp b/Engine/Plugins/FX/Niagara/Source/NiagaraEditorWidgets/Private/SNiagaraOverviewStack.cpp
--- a/Engine/Plugins/FX/Niagara/Source/NiagaraEditorWidgets/Private/SNiagaraOverviewStack.cpp
+++ b/Engine/Plugins/FX/Niagara/Source/NiagaraEditorWidgets/Private/SNiagaraOverviewStack.cpp
@@ -459,7 +459,7 @@ TSharedRef<ITableRow> SNiagaraOverviewStack::OnGenerateRowForEntry(UNiagaraStack
 						int32 Column = HighlightIndex % 2;
 						int32 Row = HighlightIndex / 2;
 						HighlightsGrid->AddSlot(Column, Row)
-							.Padding(Column, Row, 1 - Column, 1 - Row)
+							.Padding(0, 0, 1 - Column, 1 - Row)
 							[
 								SNew(SImage)
 								.Image(FNiagaraEditorWidgetsStyle::Get().GetBrush(HighlightImageBrushName))
