Focusing the selected emitter in the Niagara preview viewport no longer disables the orbit camera.

diff --git a/Engine/Plugins/FX/Niagara/Source/NiagaraEditor/Private/Widgets/SNiagaraSystemViewport.cpp b/Engine/Plugins/FX/Niagara/Source/NiagaraEditor/Private/Widgets/SNiagaraSystemViewport.cpp
--- a/Engine/Plugins/FX/Niagara/Source/NiagaraEditor/Private/Widgets/SNiagaraSystemViewport.cpp
+++ b/Engine/Plugins/FX/Niagara/Source/NiagaraEditor/Private/Widgets/SNiagaraSystemViewport.cpp
@@ -445,7 +445,12 @@ void SNiagaraSystemViewport::OnFocusViewportToSelection()
 {
 	if( PreviewComponent )
 	{
-		SystemViewportClient->FocusViewportOnBox( PreviewComponent->Bounds.GetBox() );
+		// FocusViewportOnBox disables orbit, so remember our state
+		bool bIsOrbit = SystemViewportClient->ShouldOrbitCamera();
+
+		SystemViewportClient->FocusViewportOnBox(PreviewComponent->Bounds.GetBox());
+
+		SystemViewportClient->ToggleOrbitCamera(bIsOrbit);
 	}
 }
 
