- Ensure allocation tracking variables are initialized all the time
- In the case where an emitter will fail the IsAllowedToExecute check we would not set the allocation tracking members, if a subsequent reset of the system passed this check these were uninitialized causing potentially large memory allocations

#rnx

diff --git a/Engine/Plugins/FX/Niagara/Source/Niagara/Private/NiagaraEmitterInstance.cpp b/Engine/Plugins/FX/Niagara/Source/Niagara/Private/NiagaraEmitterInstance.cpp
--- a/Engine/Plugins/FX/Niagara/Source/Niagara/Private/NiagaraEmitterInstance.cpp
+++ b/Engine/Plugins/FX/Niagara/Source/Niagara/Private/NiagaraEmitterInstance.cpp
@@ -239,13 +239,24 @@ void FNiagaraEmitterInstance::Init(int32 InEmitterIdx, FNiagaraSystemInstanceID
 	CachedEmitter = EmitterHandle.GetInstance();
 	CachedIDName = EmitterHandle.GetIdName();
 
-	if (CachedEmitter == nullptr || !IsAllowedToExecute())
+	MaxAllocationCount = 0;
+	ReallocationCount = 0;
+	MinOverallocation = -1;
+
+	if (CachedEmitter == nullptr)
 	{
 		//@todo(message manager) Error bubbling here
 		ExecutionState = ENiagaraExecutionState::Disabled;
 		return;
 	}
 
+	MaxAllocationCount = CachedEmitter->GetMaxParticleCountEstimate();
+	if (!IsAllowedToExecute())
+	{
+		ExecutionState = ENiagaraExecutionState::Disabled;
+		return;
+	}
+
 	const TArray<TSharedRef<const FNiagaraEmitterCompiledData>>& EmitterCompiledData = ParentSystemInstance->GetSystem()->GetEmitterCompiledData();
 	if (EmitterCompiledData.IsValidIndex(EmitterIdx) == false)
 	{
@@ -255,9 +266,6 @@ void FNiagaraEmitterInstance::Init(int32 InEmitterIdx, FNiagaraSystemInstanceID
 	}
 
 	CachedEmitterCompiledData = EmitterCompiledData[EmitterIdx];
-	MaxAllocationCount = CachedEmitter->GetMaxParticleCountEstimate();
-	ReallocationCount = 0;
-	MinOverallocation = -1;
 
 #if !UE_BUILD_SHIPPING && !UE_BUILD_TEST 
 	CheckForErrors();
