Default initialize particle allocation variables to guard against un-initialized emitters being executed

diff --git a/Engine/Plugins/FX/Niagara/Source/Niagara/Classes/NiagaraEmitterInstance.h b/Engine/Plugins/FX/Niagara/Source/Niagara/Classes/NiagaraEmitterInstance.h
--- a/Engine/Plugins/FX/Niagara/Source/Niagara/Classes/NiagaraEmitterInstance.h
+++ b/Engine/Plugins/FX/Niagara/Source/Niagara/Classes/NiagaraEmitterInstance.h
@@ -237,9 +237,9 @@ private:
 	TArray<FNiagaraEventHandlingInfo> EventHandlingInfo;
 	int32 EventSpawnTotal;
 
-	int32 MaxAllocationCount;
-	int32 MinOverallocation;
-	int32 ReallocationCount;
+	int32 MaxAllocationCount = 0;
+	int32 MinOverallocation = -1;
+	int32 ReallocationCount = 0;
 
 	/** Optional list of bounds calculators. */
 	TArray<TUniquePtr<FNiagaraBoundsCalculator>, TInlineAllocator<1>> BoundsCalculators;
