Deleting an isolated emitter now clears the isolated status of the system itself.

diff --git a/Engine/Plugins/FX/Niagara/Source/NiagaraEditor/Private/ViewModels/NiagaraSystemViewModel.cpp b/Engine/Plugins/FX/Niagara/Source/NiagaraEditor/Private/ViewModels/NiagaraSystemViewModel.cpp
--- a/Engine/Plugins/FX/Niagara/Source/NiagaraEditor/Private/ViewModels/NiagaraSystemViewModel.cpp
+++ b/Engine/Plugins/FX/Niagara/Source/NiagaraEditor/Private/ViewModels/NiagaraSystemViewModel.cpp
@@ -858,6 +858,16 @@ void FNiagaraSystemViewModel::RefreshEmitterHandleViewModels()
 		ViewModel->Reset();
 	}
 
+	bool bAnyEmitterIsolated = false;
+	for (TSharedRef<FNiagaraEmitterHandleViewModel> EmitterHandle : EmitterHandleViewModels)
+	{
+		if (EmitterHandle->GetIsIsolated())
+		{
+			bAnyEmitterIsolated = true;
+		}
+	}
+	GetSystem().SetIsolateEnabled(bAnyEmitterIsolated);
+
 	if (SelectionViewModel != nullptr)
 	{
 		SelectionViewModel->Refresh();
