Set default values on unitialized Niagara variables

diff --git a/Engine/Plugins/FX/Niagara/Source/Niagara/Public/NiagaraTypes.h b/Engine/Plugins/FX/Niagara/Source/Niagara/Public/NiagaraTypes.h
--- a/Engine/Plugins/FX/Niagara/Source/Niagara/Public/NiagaraTypes.h
+++ b/Engine/Plugins/FX/Niagara/Source/Niagara/Public/NiagaraTypes.h
@@ -19,7 +19,7 @@ struct FNiagaraFloat
 	GENERATED_USTRUCT_BODY()
 
 	UPROPERTY(EditAnywhere, Category=Parameters)
-	float Value;
+	float Value = 0;
 };
 
 USTRUCT(meta = (DisplayName = "int32"))
@@ -28,7 +28,7 @@ struct FNiagaraInt32
 	GENERATED_USTRUCT_BODY()
 
 	UPROPERTY(EditAnywhere, Category = Parameters)
-	int32 Value;
+	int32 Value = 0;
 };
 
 USTRUCT(meta=(DisplayName="bool"))
diff --git a/Engine/Plugins/FX/Niagara/Source/NiagaraEditor/Private/TypeEditorUtilities/NiagaraBoolTypeEditorUtilities.cpp b/Engine/Plugins/FX/Niagara/Source/NiagaraEditor/Private/TypeEditorUtilities/NiagaraBoolTypeEditorUtilities.cpp
--- a/Engine/Plugins/FX/Niagara/Source/NiagaraEditor/Private/TypeEditorUtilities/NiagaraBoolTypeEditorUtilities.cpp
+++ b/Engine/Plugins/FX/Niagara/Source/NiagaraEditor/Private/TypeEditorUtilities/NiagaraBoolTypeEditorUtilities.cpp
@@ -77,8 +77,8 @@ FString FNiagaraEditorBoolTypeUtilities::GetPinDefaultStringFromValue(const FNia
 
 bool FNiagaraEditorBoolTypeUtilities::SetValueFromPinDefaultString(const FString& StringValue, FNiagaraVariable& Variable) const
 {
-	bool bBoolValue;
-	if (LexTryParseString(bBoolValue, *StringValue))
+	bool bBoolValue = false;
+	if (LexTryParseString(bBoolValue, *StringValue) || !Variable.IsDataAllocated())
 	{
 		FNiagaraBool BoolValue;
 		BoolValue.SetValue(bBoolValue);
diff --git a/Engine/Plugins/FX/Niagara/Source/NiagaraEditor/Private/TypeEditorUtilities/NiagaraColorTypeEditorUtilities.cpp b/Engine/Plugins/FX/Niagara/Source/NiagaraEditor/Private/TypeEditorUtilities/NiagaraColorTypeEditorUtilities.cpp
--- a/Engine/Plugins/FX/Niagara/Source/NiagaraEditor/Private/TypeEditorUtilities/NiagaraColorTypeEditorUtilities.cpp
+++ b/Engine/Plugins/FX/Niagara/Source/NiagaraEditor/Private/TypeEditorUtilities/NiagaraColorTypeEditorUtilities.cpp
@@ -222,8 +222,8 @@ FString FNiagaraEditorColorTypeUtilities::GetPinDefaultStringFromValue(const FNi
 
 bool FNiagaraEditorColorTypeUtilities::SetValueFromPinDefaultString(const FString& StringValue, FNiagaraVariable& Variable) const
 {
-	FLinearColor ColorValue;
-	if (ColorValue.InitFromString(StringValue))
+	FLinearColor ColorValue = FLinearColor::Black;
+	if (ColorValue.InitFromString(StringValue) || !Variable.IsDataAllocated())
 	{
 		Variable.SetValue<FLinearColor>(ColorValue);
 		return true;
diff --git a/Engine/Plugins/FX/Niagara/Source/NiagaraEditor/Private/TypeEditorUtilities/NiagaraFloatTypeEditorUtilities.cpp b/Engine/Plugins/FX/Niagara/Source/NiagaraEditor/Private/TypeEditorUtilities/NiagaraFloatTypeEditorUtilities.cpp
--- a/Engine/Plugins/FX/Niagara/Source/NiagaraEditor/Private/TypeEditorUtilities/NiagaraFloatTypeEditorUtilities.cpp
+++ b/Engine/Plugins/FX/Niagara/Source/NiagaraEditor/Private/TypeEditorUtilities/NiagaraFloatTypeEditorUtilities.cpp
@@ -104,7 +104,7 @@ FString FNiagaraEditorFloatTypeUtilities::GetPinDefaultStringFromValue(const FNi
 bool FNiagaraEditorFloatTypeUtilities::SetValueFromPinDefaultString(const FString& StringValue, FNiagaraVariable& Variable) const
 {
 	FNiagaraFloat FloatValue;
-	if (LexTryParseString(FloatValue.Value, *StringValue))
+	if (LexTryParseString(FloatValue.Value, *StringValue) || !Variable.IsDataAllocated())
 	{
 		Variable.SetValue<FNiagaraFloat>(FloatValue);
 		return true;
diff --git a/Engine/Plugins/FX/Niagara/Source/NiagaraEditor/Private/TypeEditorUtilities/NiagaraIntegerTypeEditorUtilities.cpp b/Engine/Plugins/FX/Niagara/Source/NiagaraEditor/Private/TypeEditorUtilities/NiagaraIntegerTypeEditorUtilities.cpp
--- a/Engine/Plugins/FX/Niagara/Source/NiagaraEditor/Private/TypeEditorUtilities/NiagaraIntegerTypeEditorUtilities.cpp
+++ b/Engine/Plugins/FX/Niagara/Source/NiagaraEditor/Private/TypeEditorUtilities/NiagaraIntegerTypeEditorUtilities.cpp
@@ -105,7 +105,7 @@ FString FNiagaraEditorIntegerTypeUtilities::GetPinDefaultStringFromValue(const F
 bool FNiagaraEditorIntegerTypeUtilities::SetValueFromPinDefaultString(const FString& StringValue, FNiagaraVariable& Variable) const
 {
 	FNiagaraInt32 IntegerValue;
-	if(LexTryParseString(IntegerValue.Value, *StringValue))
+	if(LexTryParseString(IntegerValue.Value, *StringValue) || !Variable.IsDataAllocated())
 	{
 		Variable.SetValue<FNiagaraInt32>(IntegerValue);
 		return true;
diff --git a/Engine/Plugins/FX/Niagara/Source/NiagaraEditor/Private/TypeEditorUtilities/NiagaraVectorTypeEditorUtilities.cpp b/Engine/Plugins/FX/Niagara/Source/NiagaraEditor/Private/TypeEditorUtilities/NiagaraVectorTypeEditorUtilities.cpp
--- a/Engine/Plugins/FX/Niagara/Source/NiagaraEditor/Private/TypeEditorUtilities/NiagaraVectorTypeEditorUtilities.cpp
+++ b/Engine/Plugins/FX/Niagara/Source/NiagaraEditor/Private/TypeEditorUtilities/NiagaraVectorTypeEditorUtilities.cpp
@@ -167,8 +167,8 @@ FString FNiagaraEditorVector2TypeUtilities::GetPinDefaultStringFromValue(const F
 
 bool FNiagaraEditorVector2TypeUtilities::SetValueFromPinDefaultString(const FString& StringValue, FNiagaraVariable& Variable) const
 {
-	FVector2D VectorValue;
-	if (VectorValue.InitFromString(StringValue))
+	FVector2D VectorValue = FVector2D::ZeroVector;
+	if (VectorValue.InitFromString(StringValue) || !Variable.IsDataAllocated())
 	{
 		Variable.SetValue<FVector2D>(VectorValue);
 		return true;
@@ -243,8 +243,8 @@ FString FNiagaraEditorVector3TypeUtilities::GetPinDefaultStringFromValue(const F
 bool FNiagaraEditorVector3TypeUtilities::SetValueFromPinDefaultString(const FString& StringValue, FNiagaraVariable& Variable) const
 {
 	// NOTE: We can not use InitFromString() here since the vector pin control doesn't use the standard 'X=0,Y=0,Z=0' syntax.
-	FVector Value;
-	if (FDefaultValueHelper::ParseVector(StringValue, Value))
+	FVector Value = FVector::ZeroVector;
+	if (FDefaultValueHelper::ParseVector(StringValue, Value) || !Variable.IsDataAllocated())
 	{
 		Variable.SetValue<FVector>(Value);
 		return true;
@@ -318,8 +318,8 @@ FString FNiagaraEditorVector4TypeUtilities::GetPinDefaultStringFromValue(const F
 bool FNiagaraEditorVector4TypeUtilities::SetValueFromPinDefaultString(const FString& StringValue, FNiagaraVariable& Variable) const
 {
 	// NOTE: We can not use InitFromString() here since the vector pin control doesn't use the standard 'X=0,Y=0,Z=0,W=0' syntax.
-	FVector4 Value;
-	if (FDefaultValueHelper::ParseVector4(StringValue, Value))
+	FVector4 Value(0, 0, 0, 0);
+	if (FDefaultValueHelper::ParseVector4(StringValue, Value) || !Variable.IsDataAllocated())
 	{
 		Variable.SetValue<FVector4>(Value);
 		return true;
@@ -423,8 +423,8 @@ FString FNiagaraEditorQuatTypeUtilities::GetPinDefaultStringFromValue(const FNia
 bool FNiagaraEditorQuatTypeUtilities::SetValueFromPinDefaultString(const FString& StringValue, FNiagaraVariable& Variable) const
 {
 	// NOTE: We can not use InitFromString() here since the vector pin control doesn't use the standard 'X=0,Y=0,Z=0,W=0' syntax.
-	FVector4 Value;
-	if (FDefaultValueHelper::ParseVector4(StringValue, Value))
+	FVector4 Value(0, 0, 0, 0);
+	if (FDefaultValueHelper::ParseVector4(StringValue, Value) || !Variable.IsDataAllocated())
 	{
 		FQuat Quat(Value.X, Value.Y, Value.Z, Value.W);
 		Variable.SetValue<FQuat>(Quat);
