Niagara - Implement copy/paste for inputs in the stack including local values, linked inputs, data interfaces, expressions, and dynamic inputs.

diff --git a/Engine/Plugins/FX/Niagara/Source/NiagaraEditor/Private/EdGraphSchema_Niagara.cpp b/Engine/Plugins/FX/Niagara/Source/NiagaraEditor/Private/EdGraphSchema_Niagara.cpp
--- a/Engine/Plugins/FX/Niagara/Source/NiagaraEditor/Private/EdGraphSchema_Niagara.cpp
+++ b/Engine/Plugins/FX/Niagara/Source/NiagaraEditor/Private/EdGraphSchema_Niagara.cpp
@@ -390,7 +390,7 @@ TArray<TSharedPtr<FNiagaraSchemaAction_NewNode> > UEdGraphSchema_Niagara::GetGra
 		const FText TooltipDesc = LOCTEXT("CustomHlslPopupTooltip", "Add a node with custom hlsl content");
 		TSharedPtr<FNiagaraSchemaAction_NewNode> FunctionCallAction = AddNewNodeAction(NewActions, LOCTEXT("Function Menu Title", "Functions"), MenuDesc, TEXT("CustomHLSL"), TooltipDesc);
 		UNiagaraNodeCustomHlsl* CustomHlslNode = NewObject<UNiagaraNodeCustomHlsl>(OwnerOfTemporaries);
-		CustomHlslNode->CustomHlsl = TEXT("// Insert the body of the function here and add any inputs\r\n// and outputs by name using the add pins above.\r\n// Currently, complicated branches, for loops, switches, etc are not advised.");
+		CustomHlslNode->SetCustomHlsl(TEXT("// Insert the body of the function here and add any inputs\r\n// and outputs by name using the add pins above.\r\n// Currently, complicated branches, for loops, switches, etc are not advised."));
 		FunctionCallAction->NodeTemplate = CustomHlslNode;
 	}
 
diff --git a/Engine/Plugins/FX/Niagara/Source/NiagaraEditor/Private/NiagaraClipboard.cpp b/Engine/Plugins/FX/Niagara/Source/NiagaraEditor/Private/NiagaraClipboard.cpp
new file mode 100644
--- /dev/null
+++ b/Engine/Plugins/FX/Niagara/Source/NiagaraEditor/Private/NiagaraClipboard.cpp
@@ -0,0 +1,143 @@
+// Copyright 1998-2019 Epic Games, Inc. All Rights Reserved.
+
+#include "NiagaraClipboard.h"
+#include "NiagaraDataInterface.h"
+
+#include "Factories.h"
+#include "UObject/UObjectMarks.h"
+#include "UObject/PropertyPortFlags.h"
+#include "Containers/UnrealString.h"
+#include "UnrealExporter.h"
+#include "Exporters/Exporter.h"
+#include "HAL/PlatformApplicationMisc.h"
+
+struct FNiagaraClipboardContentTextObjectFactory : public FCustomizableTextObjectFactory
+{
+public:
+	UNiagaraClipboardContent* ClipboardContent;
+
+public:
+	FNiagaraClipboardContentTextObjectFactory()
+		: FCustomizableTextObjectFactory(GWarn)
+		, ClipboardContent(nullptr)
+	{
+	}
+
+protected:
+	virtual bool CanCreateClass(UClass* ObjectClass, bool& bOmitSubObjs) const override
+	{
+		return ObjectClass == UNiagaraClipboardContent::StaticClass();
+	}
+
+	virtual void ProcessConstructedObject(UObject* CreatedObject) override
+	{
+		if (CreatedObject->IsA<UNiagaraClipboardContent>())
+		{
+			ClipboardContent = CastChecked<UNiagaraClipboardContent>(CreatedObject);
+		}
+	}
+};
+
+UNiagaraClipboardFunctionInput* MakeNewInput(UObject* InOuter, FName InInputName, FNiagaraTypeDefinition InInputType, ENiagaraClipboardFunctionInputValueMode InValueMode)
+{
+	UNiagaraClipboardFunctionInput* NewInput = Cast<UNiagaraClipboardFunctionInput>(NewObject<UNiagaraClipboardFunctionInput>(InOuter));
+	NewInput->InputName = InInputName;
+	NewInput->InputType = InInputType;
+	NewInput->ValueMode = InValueMode;
+	return NewInput;
+}
+
+const UNiagaraClipboardFunctionInput* UNiagaraClipboardFunctionInput::CreateLocalValue(UObject* InOuter, FName InInputName, FNiagaraTypeDefinition InInputType, TArray<uint8>& InLocalValueData)
+{
+	checkf(InLocalValueData.Num() == InInputType.GetSize(), TEXT("Input data size didn't match type size."))
+	UNiagaraClipboardFunctionInput* NewInput = MakeNewInput(InOuter, InInputName, InInputType, ENiagaraClipboardFunctionInputValueMode::Local);
+	NewInput->Local = InLocalValueData;
+	return NewInput;
+}
+
+const UNiagaraClipboardFunctionInput* UNiagaraClipboardFunctionInput::CreateLinkedValue(UObject* InOuter, FName InInputName, FNiagaraTypeDefinition InInputType, FName InLinkedValue)
+{
+	UNiagaraClipboardFunctionInput* NewInput = MakeNewInput(InOuter, InInputName, InInputType, ENiagaraClipboardFunctionInputValueMode::Linked);
+	NewInput->Linked = InLinkedValue;
+	return NewInput;
+}
+
+const UNiagaraClipboardFunctionInput* UNiagaraClipboardFunctionInput::CreateDataValue(UObject* InOuter, FName InInputName, FNiagaraTypeDefinition InInputType, UNiagaraDataInterface* InDataValue)
+{
+	UNiagaraClipboardFunctionInput* NewInput = MakeNewInput(InOuter, InInputName, InInputType, ENiagaraClipboardFunctionInputValueMode::Data);
+	NewInput->Data = NewObject<UNiagaraDataInterface>(NewInput, InDataValue->GetClass());
+	InDataValue->CopyTo(NewInput->Data);
+	return NewInput;
+}
+
+const UNiagaraClipboardFunctionInput* UNiagaraClipboardFunctionInput::CreateExpressionValue(UObject* InOuter, FName InInputName, FNiagaraTypeDefinition InInputType, const FString& InExpressionValue)
+{
+	UNiagaraClipboardFunctionInput* NewInput = MakeNewInput(InOuter, InInputName, InInputType, ENiagaraClipboardFunctionInputValueMode::Expression);
+	NewInput->Expression = InExpressionValue;
+	return NewInput;
+}
+
+const UNiagaraClipboardFunctionInput* UNiagaraClipboardFunctionInput::CreateDynamicValue(UObject* InOuter, FName InInputName, FNiagaraTypeDefinition InInputType, UNiagaraScript* InDynamicValue)
+{
+	UNiagaraClipboardFunctionInput* NewInput = MakeNewInput(InOuter, InInputName, InInputType, ENiagaraClipboardFunctionInputValueMode::Dynamic);
+	NewInput->Dynamic = UNiagaraClipboardFunction::CreateScriptFunction(NewInput, InDynamicValue);
+	return NewInput;
+}
+
+UNiagaraClipboardFunction* UNiagaraClipboardFunction::CreateScriptFunction(UObject* InOuter, UNiagaraScript* InScript)
+{
+	UNiagaraClipboardFunction* NewFunction = Cast<UNiagaraClipboardFunction>(NewObject<UNiagaraClipboardFunction>(InOuter));
+	NewFunction->ScriptMode = ENiagaraClipboardFunctionScriptMode::ScriptAsset;
+	NewFunction->Script = InScript;
+	return NewFunction;
+}
+
+FNiagaraClipboard::FNiagaraClipboard()
+{
+}
+
+void FNiagaraClipboard::Copy(const UNiagaraClipboardFunctionInput* Input)
+{
+	UNiagaraClipboardContent* ClipboardContent = NewObject<UNiagaraClipboardContent>(GetTransientPackage());
+	ClipboardContent->FunctionInputs.Add(CastChecked<UNiagaraClipboardFunctionInput>(StaticDuplicateObject(Input, ClipboardContent)));
+	SetClipboardContentInternal(ClipboardContent);
+}
+
+const UNiagaraClipboardFunctionInput* FNiagaraClipboard::GetCopiedInput() const
+{
+	const UNiagaraClipboardContent* ClipboardContent = GetClipboardContentInternal();
+	if (ClipboardContent != nullptr &&ClipboardContent->FunctionInputs.Num() == 1)
+	{
+		return ClipboardContent->FunctionInputs[0];
+	}
+	return nullptr;
+}
+
+void FNiagaraClipboard::SetClipboardContentInternal(UNiagaraClipboardContent* ClipboardContent)
+{
+	// Clear the mark state for saving.
+	UnMarkAllObjects(EObjectMark(OBJECTMARK_TagExp | OBJECTMARK_TagImp));
+
+	// Export the clipboard to text.
+	FStringOutputDevice Archive;
+	const FExportObjectInnerContext Context;
+	UExporter::ExportToOutputDevice(&Context, ClipboardContent, nullptr, Archive, TEXT("copy"), 0, PPF_ExportsNotFullyQualified | PPF_Copy | PPF_Delimited, false, ClipboardContent->GetOuter());
+	FPlatformApplicationMisc::ClipboardCopy(*Archive);
+}
+
+const UNiagaraClipboardContent* FNiagaraClipboard::GetClipboardContentInternal() const
+{
+	// Get the text from the clipboard.
+	FString ClipboardText;
+	FPlatformApplicationMisc::ClipboardPaste(ClipboardText);
+
+	// Try to create niagara clipboard content from that.
+	FNiagaraClipboardContentTextObjectFactory ClipboardContentFactory;
+	if (ClipboardContentFactory.CanCreateObjectsFromText(ClipboardText))
+	{
+		ClipboardContentFactory.ProcessBuffer(GetTransientPackage(), RF_Transient, ClipboardText);
+		return ClipboardContentFactory.ClipboardContent;
+	}
+
+	return nullptr;
+}
\ No newline at end of file
diff --git a/Engine/Plugins/FX/Niagara/Source/NiagaraEditor/Private/NiagaraEditorModule.cpp b/Engine/Plugins/FX/Niagara/Source/NiagaraEditor/Private/NiagaraEditorModule.cpp
--- a/Engine/Plugins/FX/Niagara/Source/NiagaraEditor/Private/NiagaraEditorModule.cpp
+++ b/Engine/Plugins/FX/Niagara/Source/NiagaraEditor/Private/NiagaraEditorModule.cpp
@@ -73,6 +73,7 @@
 #include "NiagaraSystemFactoryNew.h"
 #include "NiagaraSystemEditorData.h"
 #include "NiagaraEditorCommands.h"
+#include "NiagaraClipboard.h"
 
 #include "MovieScene/Parameters/MovieSceneNiagaraBoolParameterTrack.h"
 #include "MovieScene/Parameters/MovieSceneNiagaraFloatParameterTrack.h"
@@ -230,6 +231,7 @@ FNiagaraEditorModule::FNiagaraEditorModule()
 	: SequencerSettings(nullptr)
 	, TestCompileScriptCommand(nullptr)
 	, DumpCompileIdDataForAssetCommand(nullptr)
+	, Clipboard(MakeShared<FNiagaraClipboard>())
 {
 }
 
@@ -1145,6 +1147,11 @@ void FNiagaraEditorModule::GetScriptAssetsMatchingHighlight(const FNiagaraScript
 	}
 }
 
+FNiagaraClipboard& FNiagaraEditorModule::GetClipboard() const
+{
+	return Clipboard.Get();
+}
+
 void FNiagaraEditorModule::RegisterAssetTypeAction(IAssetTools& AssetTools, TSharedRef<IAssetTypeActions> Action)
 {
 	AssetTools.RegisterAssetTypeActions(Action);
diff --git a/Engine/Plugins/FX/Niagara/Source/NiagaraEditor/Private/NiagaraHlslTranslator.cpp b/Engine/Plugins/FX/Niagara/Source/NiagaraEditor/Private/NiagaraHlslTranslator.cpp
--- a/Engine/Plugins/FX/Niagara/Source/NiagaraEditor/Private/NiagaraHlslTranslator.cpp
+++ b/Engine/Plugins/FX/Niagara/Source/NiagaraEditor/Private/NiagaraHlslTranslator.cpp
@@ -4726,7 +4726,7 @@ void FHlslNiagaraTranslator::HandleCustomHlslNode(UNiagaraNodeCustomHlsl* Custom
 	OutFullName = CustomFunctionHlsl->GetFullName();
 	OutSignature.Name = *OutName; // Force the name to be set to include the node guid for safety...
 	bOutCustomHlsl = true;
-	OutCustomHlsl = CustomFunctionHlsl->CustomHlsl;
+	OutCustomHlsl = CustomFunctionHlsl->GetCustomHlsl();
 
 	// Split up the hlsl into constituent tokens
 	TArray<FString> Tokens;
diff --git a/Engine/Plugins/FX/Niagara/Source/NiagaraEditor/Private/NiagaraNodeCustomHlsl.cpp b/Engine/Plugins/FX/Niagara/Source/NiagaraEditor/Private/NiagaraNodeCustomHlsl.cpp
--- a/Engine/Plugins/FX/Niagara/Source/NiagaraEditor/Private/NiagaraNodeCustomHlsl.cpp
+++ b/Engine/Plugins/FX/Niagara/Source/NiagaraEditor/Private/NiagaraNodeCustomHlsl.cpp
@@ -21,6 +21,19 @@ UNiagaraNodeCustomHlsl::UNiagaraNodeCustomHlsl(const FObjectInitializer& ObjectI
 	FunctionDisplayName = Signature.Name.ToString();
 }
 
+const FString& UNiagaraNodeCustomHlsl::GetCustomHlsl() const
+{
+	return CustomHlsl;
+}
+
+void UNiagaraNodeCustomHlsl::SetCustomHlsl(const FString& InCustomHlsl)
+{
+	Modify();
+	CustomHlsl = InCustomHlsl;
+	RefreshFromExternalChanges();
+	MarkNodeRequiresSynchronization(__FUNCTION__, true);
+}
+
 TSharedPtr<SGraphNode> UNiagaraNodeCustomHlsl::CreateVisualWidget()
 {
 	return SNew(SNiagaraGraphNodeCustomHlsl, this);
@@ -43,10 +56,7 @@ void UNiagaraNodeCustomHlsl::OnCustomHlslTextCommitted(const FText& InText, ETex
 	if (!NewValue.Equals(CustomHlsl, ESearchCase::CaseSensitive))
 	{
 		FScopedTransaction Transaction(LOCTEXT("CustomHlslCommit", "Edited Custom Hlsl"));
-		Modify();
-		CustomHlsl = NewValue;
-		RefreshFromExternalChanges();			
-		MarkNodeRequiresSynchronization(__FUNCTION__, true);
+		SetCustomHlsl(NewValue);
 	}
 }
 
diff --git a/Engine/Plugins/FX/Niagara/Source/NiagaraEditor/Private/NiagaraScriptMergeManager.cpp b/Engine/Plugins/FX/Niagara/Source/NiagaraEditor/Private/NiagaraScriptMergeManager.cpp
--- a/Engine/Plugins/FX/Niagara/Source/NiagaraEditor/Private/NiagaraScriptMergeManager.cpp
+++ b/Engine/Plugins/FX/Niagara/Source/NiagaraEditor/Private/NiagaraScriptMergeManager.cpp
@@ -1600,7 +1600,7 @@ TOptional<bool> FNiagaraScriptMergeManager::DoFunctionInputOverridesMatch(TShare
 
 		if (BaseCustomHlsl && OtherCustomHlsl)
 		{
-			if (BaseCustomHlsl->CustomHlsl != OtherCustomHlsl->CustomHlsl || BaseCustomHlsl->ScriptUsage != OtherCustomHlsl->ScriptUsage)
+			if (BaseCustomHlsl->GetCustomHlsl() != OtherCustomHlsl->GetCustomHlsl() || BaseCustomHlsl->ScriptUsage != OtherCustomHlsl->ScriptUsage)
 			{
 				return false;
 			}
@@ -1782,8 +1782,7 @@ FNiagaraScriptMergeManager::FApplyDiffResults FNiagaraScriptMergeManager::AddInp
 				if (CustomHlslFunction != nullptr)
 				{
 					UNiagaraNodeCustomHlsl* DynamicInputFunctionCall;
-					FNiagaraStackGraphUtilities::SetCustomExpressionForFunctionInput(InputOverridePin, DynamicInputFunctionCall, OverrideToAdd->GetOverrideNodeId());
-					DynamicInputFunctionCall->CustomHlsl = CustomHlslFunction->CustomHlsl;
+					FNiagaraStackGraphUtilities::SetCustomExpressionForFunctionInput(InputOverridePin, CustomHlslFunction->GetCustomHlsl(), DynamicInputFunctionCall, OverrideToAdd->GetOverrideNodeId());
 					for (TSharedRef<FNiagaraStackFunctionInputOverrideMergeAdapter> DynamicInputInputOverride : OverrideToAdd->GetDynamicValueFunction()->GetInputOverrides())
 					{
 						FApplyDiffResults AddResults = AddInputOverride(UniqueEmitterName, OwningScript, *((UNiagaraNodeFunctionCall*)DynamicInputFunctionCall), DynamicInputInputOverride);
diff --git a/Engine/Plugins/FX/Niagara/Source/NiagaraEditor/Private/ViewModels/Stack/NiagaraStackFunctionInput.cpp b/Engine/Plugins/FX/Niagara/Source/NiagaraEditor/Private/ViewModels/Stack/NiagaraStackFunctionInput.cpp
--- a/Engine/Plugins/FX/Niagara/Source/NiagaraEditor/Private/ViewModels/Stack/NiagaraStackFunctionInput.cpp
+++ b/Engine/Plugins/FX/Niagara/Source/NiagaraEditor/Private/ViewModels/Stack/NiagaraStackFunctionInput.cpp
@@ -29,6 +29,7 @@
 #include "NiagaraRendererProperties.h"
 #include "ViewModels/Stack/NiagaraParameterHandle.h"
 #include "NiagaraEmitter.h"
+#include "NiagaraClipboard.h"
 
 #include "Materials/Material.h"
 #include "Materials/MaterialExpressionDynamicParameter.h"
@@ -256,6 +257,65 @@ UObject* UNiagaraStackFunctionInput::GetExternalAsset() const
 	return nullptr;
 }
 
+bool UNiagaraStackFunctionInput::TestCanCutWithMessage(FText& OutMessage) const
+{
+	if (InputValues.Mode == EValueMode::Invalid)
+	{
+		OutMessage = LOCTEXT("CantCutInvalidMessage", "The current input state doesn't support cutting.");
+		return false;
+	}
+	OutMessage = LOCTEXT("CanCutMessage", "Cut will copy the value of this input including\nany data objects and dynamic inputs, and will reset it to default.");
+	return true;
+}
+
+void UNiagaraStackFunctionInput::Cut()
+{
+	FScopedTransaction ScopedTransaction(LOCTEXT("CutInput", "Cut niagara input"));
+	Copy();
+	Reset();
+}
+
+bool UNiagaraStackFunctionInput::TestCanCopyWithMessage(FText& OutMessage) const
+{
+	if (InputValues.Mode == EValueMode::Invalid)
+	{
+		OutMessage = LOCTEXT("CantCopyInvalidMessage", "The current input state doesn't support copying.");
+		return false;
+	}
+	OutMessage = LOCTEXT("CanCopyMessage", "Copy the value of this input including\nany data objects and dynamic inputs.");
+	return true;
+}
+
+void UNiagaraStackFunctionInput::Copy() const
+{
+	const UNiagaraClipboardFunctionInput* ClipboardInput = ToClipboardFunctionInput(GetTransientPackage());
+	if (ClipboardInput != nullptr)
+	{
+		FNiagaraEditorModule::Get().GetClipboard().Copy(ClipboardInput);
+	}
+}
+
+bool UNiagaraStackFunctionInput::TestCanPasteWithMessage(FText& OutMessage) const
+{
+	const UNiagaraClipboardFunctionInput* ClipboardInput = FNiagaraEditorModule::Get().GetClipboard().GetCopiedInput();
+	if (GetIsEnabledAndOwnerIsEnabled() &&  ClipboardInput != nullptr && ClipboardInput->InputType == InputType)
+	{
+		OutMessage = LOCTEXT("PastMessage", "Paste the input from the clipboard here.");
+		return true;
+	}
+	return false;
+}
+
+void UNiagaraStackFunctionInput::Paste()
+{
+	const UNiagaraClipboardFunctionInput* ClipboardInput = FNiagaraEditorModule::Get().GetClipboard().GetCopiedInput();
+	if (ClipboardInput != nullptr && ClipboardInput->InputType == InputType)
+	{
+		FScopedTransaction ScopedTransaction(LOCTEXT("PasteInput", "Paste niagara input"));
+		SetValueFromClipboardFunctionInput(*ClipboardInput);
+	}
+}
+
 FText UNiagaraStackFunctionInput::GetTooltipText(EValueMode InValueMode) const
 {
 	FNiagaraVariable ValueVariable;
@@ -884,11 +944,6 @@ UNiagaraNodeFunctionCall* UNiagaraStackFunctionInput::GetDynamicInputNode() cons
 	return InputValues.DynamicNode.Get();
 }
 
-UNiagaraNodeCustomHlsl* UNiagaraStackFunctionInput::GetExpressionNode() const
-{
-	return InputValues.ExpressionNode.Get();
-}
-
 void UNiagaraStackFunctionInput::GetAvailableDynamicInputs(TArray<UNiagaraScript*>& AvailableDynamicInputs, bool bIncludeNonLibraryInputs)
 {
 	TArray<FAssetData> DynamicInputAssets;
@@ -942,10 +997,14 @@ void UNiagaraStackFunctionInput::SetDynamicInput(UNiagaraScript* DynamicInput)
 	RefreshChildren();
 }
 
-void UNiagaraStackFunctionInput::SetCustomExpression(const FString& InputText)
+FText UNiagaraStackFunctionInput::GetCustomExpressionText() const
 {
-	FScopedTransaction ScopedTransaction(LOCTEXT("SetCustomExpressionInput", "Make custom expression input"));
+	return InputValues.ExpressionNode != nullptr ? FText::FromString(InputValues.ExpressionNode->GetCustomHlsl()) : FText();
+}
 
+void UNiagaraStackFunctionInput::SetCustomExpression(const FString& InCustomExpression)
+{
+	FScopedTransaction ScopedTransaction(LOCTEXT("SetCustomExpressionInput", "Make custom expression input"));
 
 	UEdGraphPin& OverridePin = GetOrCreateOverridePin();
 	RemoveNodesForOverridePin(OverridePin);
@@ -955,7 +1014,7 @@ void UNiagaraStackFunctionInput::SetCustomExpression(const FString& InputText)
 	}
 
 	UNiagaraNodeCustomHlsl* FunctionCallNode;
-	FNiagaraStackGraphUtilities::SetCustomExpressionForFunctionInput(OverridePin, FunctionCallNode);
+	FNiagaraStackGraphUtilities::SetCustomExpressionForFunctionInput(OverridePin, InCustomExpression, FunctionCallNode);
 	FNiagaraStackGraphUtilities::InitializeStackFunctionInputs(GetSystemViewModel(), GetEmitterViewModel(), GetStackEditorData(), *OwningModuleNode, *FunctionCallNode);
 	FNiagaraStackGraphUtilities::RelayoutGraph(*OwningFunctionCallNode->GetGraph());
 
@@ -1608,6 +1667,91 @@ bool UNiagaraStackFunctionInput::GetShouldPassFilterForVisibleCondition() const
 	return bIsVisible && (GetHasVisibleCondition() == false || GetVisibleConditionEnabled());
 }
 
+const UNiagaraClipboardFunctionInput* UNiagaraStackFunctionInput::ToClipboardFunctionInput(UObject* InOuter) const
+{
+	const UNiagaraClipboardFunctionInput* ClipboardInput = nullptr;
+	FName InputName = InputParameterHandle.GetName();
+	switch (InputValues.Mode)
+	{
+	case EValueMode::Local:
+	{
+		TArray<uint8> LocalValueData;
+		LocalValueData.AddUninitialized(InputType.GetSize());
+		FMemory::Memcpy(LocalValueData.GetData(), InputValues.LocalStruct->GetStructMemory(), InputType.GetSize());
+		ClipboardInput = UNiagaraClipboardFunctionInput::CreateLocalValue(InOuter, InputName, InputType, LocalValueData);
+		break;
+	}
+	case EValueMode::Linked:
+		ClipboardInput = UNiagaraClipboardFunctionInput::CreateLinkedValue(InOuter, InputName, InputType, InputValues.LinkedHandle.GetParameterHandleString());
+		break;
+	case EValueMode::Data:
+		ClipboardInput = UNiagaraClipboardFunctionInput::CreateDataValue(InOuter, InputName, InputType, InputValues.DataObjects.GetValueObject());
+		break;
+	case EValueMode::Expression:
+		ClipboardInput = UNiagaraClipboardFunctionInput::CreateExpressionValue(InOuter, InputName, InputType, InputValues.ExpressionNode->GetHlslText().ToString());
+		break;
+	case EValueMode::Dynamic:
+	{
+		ClipboardInput = UNiagaraClipboardFunctionInput::CreateDynamicValue(InOuter, InputName, InputType, InputValues.DynamicNode->FunctionScript);
+
+		TArray<UNiagaraStackFunctionInputCollection*> DynamicInputCollections;
+		GetUnfilteredChildrenOfType(DynamicInputCollections);
+		for (UNiagaraStackFunctionInputCollection* DynamicInputCollection : DynamicInputCollections)
+		{
+			DynamicInputCollection->ToClipboardFunctionInputs(ClipboardInput->Dynamic, ClipboardInput->Dynamic->Inputs);
+		}
+
+		break;
+	}
+	default:
+		ensureMsgf(false, TEXT("A new value mode was added without adding support for copy paste."));
+		break;
+	}
+	return ClipboardInput;
+}
+
+void UNiagaraStackFunctionInput::SetValueFromClipboardFunctionInput(const UNiagaraClipboardFunctionInput& ClipboardFunctionInput)
+{
+	if (ensureMsgf(ClipboardFunctionInput.InputType == InputType, TEXT("Can not set input value from clipboard, input types don't match.")))
+	{
+		switch (ClipboardFunctionInput.ValueMode)
+		{
+		case ENiagaraClipboardFunctionInputValueMode::Local:
+		{
+			TSharedRef<FStructOnScope> ValueStruct = MakeShared<FStructOnScope>(InputType.GetStruct());
+			FMemory::Memcpy(ValueStruct->GetStructMemory(), ClipboardFunctionInput.Local.GetData(), InputType.GetSize());
+			SetLocalValue(ValueStruct);
+			break;
+		}
+		case ENiagaraClipboardFunctionInputValueMode::Linked:
+			SetLinkedValueHandle(FNiagaraParameterHandle(ClipboardFunctionInput.Linked));
+			break;
+		case ENiagaraClipboardFunctionInputValueMode::Data:
+			ClipboardFunctionInput.Data->CopyTo(GetDataValueObject());
+			break;
+		case ENiagaraClipboardFunctionInputValueMode::Expression:
+			SetCustomExpression(ClipboardFunctionInput.Expression);
+			break;
+		case ENiagaraClipboardFunctionInputValueMode::Dynamic:
+			if (ensureMsgf(ClipboardFunctionInput.Dynamic->ScriptMode == ENiagaraClipboardFunctionScriptMode::ScriptAsset,
+				TEXT("Can not set dynamic input value from clipboard, only script asset funcitons can be set.")))
+			{
+				SetDynamicInput(ClipboardFunctionInput.Dynamic->Script);
+				TArray<UNiagaraStackFunctionInputCollection*> DynamicInputCollections;
+				GetUnfilteredChildrenOfType(DynamicInputCollections);
+				for (UNiagaraStackFunctionInputCollection* DynamicInputCollection : DynamicInputCollections)
+				{
+					DynamicInputCollection->SetValuesFromClipboardFunctionInputs(ClipboardFunctionInput.Dynamic->Inputs);
+				}
+			}
+			break;
+		default:
+			ensureMsgf(false, TEXT("A new value mode was added without adding support for copy paste."));
+			break;
+		}
+	}
+}
+
 void UNiagaraStackFunctionInput::GetSearchItems(TArray<FStackSearchItem>& SearchItems) const
 {
 	if (GetShouldPassFilterForVisibleCondition() && GetIsInlineEditConditionToggle() == false)
diff --git a/Engine/Plugins/FX/Niagara/Source/NiagaraEditor/Private/ViewModels/Stack/NiagaraStackFunctionInputCollection.cpp b/Engine/Plugins/FX/Niagara/Source/NiagaraEditor/Private/ViewModels/Stack/NiagaraStackFunctionInputCollection.cpp
--- a/Engine/Plugins/FX/Niagara/Source/NiagaraEditor/Private/ViewModels/Stack/NiagaraStackFunctionInputCollection.cpp
+++ b/Engine/Plugins/FX/Niagara/Source/NiagaraEditor/Private/ViewModels/Stack/NiagaraStackFunctionInputCollection.cpp
@@ -13,6 +13,7 @@
 #include "NiagaraConstants.h"
 #include "NiagaraGraph.h"
 #include "NiagaraStackEditorData.h"
+#include "NiagaraClipboard.h"
 
 #include "EdGraph/EdGraphPin.h"
 #include "ScopedTransaction.h"
@@ -77,6 +78,26 @@ void UNiagaraStackFunctionInputCollection::SetShouldShowInStack(bool bInShouldSh
 	bShouldShowInStack = bInShouldShowInStack;
 }
 
+void UNiagaraStackFunctionInputCollection::ToClipboardFunctionInputs(UObject* InOuter, TArray<const UNiagaraClipboardFunctionInput*>& OutClipboardFunctionInputs) const
+{
+	TArray<UNiagaraStackInputCategory*> ChildCategories;
+	GetUnfilteredChildrenOfType(ChildCategories);
+	for (UNiagaraStackInputCategory* ChildCategory : ChildCategories)
+	{
+		ChildCategory->ToClipboardFunctionInputs(InOuter, OutClipboardFunctionInputs);
+	}
+}
+
+void UNiagaraStackFunctionInputCollection::SetValuesFromClipboardFunctionInputs(const TArray<const UNiagaraClipboardFunctionInput*>& ClipboardFunctionInputs)
+{
+	TArray<UNiagaraStackInputCategory*> ChildCategories;
+	GetUnfilteredChildrenOfType(ChildCategories);
+	for (UNiagaraStackInputCategory* ChildCategory : ChildCategories)
+	{
+		ChildCategory->SetValuesFromClipboardFunctionInputs(ClipboardFunctionInputs);
+	}
+}
+
 void UNiagaraStackFunctionInputCollection::RefreshChildrenInternal(const TArray<UNiagaraStackEntry*>& CurrentChildren, TArray<UNiagaraStackEntry*>& NewChildren, TArray<FStackIssue>& NewIssues)
 {
 	TSet<const UEdGraphPin*> HiddenPins;
diff --git a/Engine/Plugins/FX/Niagara/Source/NiagaraEditor/Private/ViewModels/Stack/NiagaraStackGraphUtilities.cpp b/Engine/Plugins/FX/Niagara/Source/NiagaraEditor/Private/ViewModels/Stack/NiagaraStackGraphUtilities.cpp
--- a/Engine/Plugins/FX/Niagara/Source/NiagaraEditor/Private/ViewModels/Stack/NiagaraStackGraphUtilities.cpp
+++ b/Engine/Plugins/FX/Niagara/Source/NiagaraEditor/Private/ViewModels/Stack/NiagaraStackGraphUtilities.cpp
@@ -979,7 +979,7 @@ void FNiagaraStackGraphUtilities::SetDynamicInputForFunctionInput(UEdGraphPin& O
 	}
 }
 
-void FNiagaraStackGraphUtilities::SetCustomExpressionForFunctionInput(UEdGraphPin& OverridePin, UNiagaraNodeCustomHlsl*& OutDynamicInputFunctionCall, const FGuid& NewNodePersistentId)
+void FNiagaraStackGraphUtilities::SetCustomExpressionForFunctionInput(UEdGraphPin& OverridePin, const FString& CustomExpression, UNiagaraNodeCustomHlsl*& OutDynamicInputFunctionCall, const FGuid& NewNodePersistentId)
 {
 	checkf(OverridePin.LinkedTo.Num() == 0, TEXT("Can't set a data value when the override pin already has a value."));
 
@@ -1017,6 +1017,8 @@ void FNiagaraStackGraphUtilities::SetCustomExpressionForFunctionInput(UEdGraphPi
 	{
 		FunctionCallNode->NodeGuid = NewNodePersistentId;
 	}
+
+	FunctionCallNode->SetCustomHlsl(CustomExpression);
 }
 
 bool FNiagaraStackGraphUtilities::RemoveModuleFromStack(UNiagaraSystem& OwningSystem, FGuid OwningEmitterId, UNiagaraNodeFunctionCall& ModuleNode)
diff --git a/Engine/Plugins/FX/Niagara/Source/NiagaraEditor/Private/ViewModels/Stack/NiagaraStackInputCategory.cpp b/Engine/Plugins/FX/Niagara/Source/NiagaraEditor/Private/ViewModels/Stack/NiagaraStackInputCategory.cpp
--- a/Engine/Plugins/FX/Niagara/Source/NiagaraEditor/Private/ViewModels/Stack/NiagaraStackInputCategory.cpp
+++ b/Engine/Plugins/FX/Niagara/Source/NiagaraEditor/Private/ViewModels/Stack/NiagaraStackInputCategory.cpp
@@ -3,6 +3,7 @@
 #include "ViewModels/Stack/NiagaraStackInputCategory.h"
 #include "ViewModels/Stack/NiagaraStackFunctionInput.h"
 #include "NiagaraNodeFunctionCall.h"
+#include "NiagaraClipboard.h"
 
 void UNiagaraStackInputCategory::Initialize(
 	FRequiredEntryData InRequiredEntryData,
@@ -90,6 +91,32 @@ void UNiagaraStackInputCategory::SetShouldShowInStack(bool bInShouldShowInStack)
 	bShouldShowInStack = bInShouldShowInStack;
 }
 
+void UNiagaraStackInputCategory::ToClipboardFunctionInputs(UObject* InOuter, TArray<const UNiagaraClipboardFunctionInput*>& OutClipboardFunctionInputs) const
+{
+	TArray<UNiagaraStackFunctionInput*> ChildInputs;
+	GetUnfilteredChildrenOfType(ChildInputs);
+	for (UNiagaraStackFunctionInput* ChildInput : ChildInputs)
+	{
+		OutClipboardFunctionInputs.Add(ChildInput->ToClipboardFunctionInput(InOuter));
+	}
+}
+
+void  UNiagaraStackInputCategory::SetValuesFromClipboardFunctionInputs(const TArray<const UNiagaraClipboardFunctionInput*>& ClipboardFunctionInputs)
+{
+	TArray<UNiagaraStackFunctionInput*> ChildInputs;
+	GetUnfilteredChildrenOfType(ChildInputs);
+	for (UNiagaraStackFunctionInput* ChildInput : ChildInputs)
+	{
+		for (const UNiagaraClipboardFunctionInput* ClipboardFunctionInput : ClipboardFunctionInputs)
+		{
+			if (ChildInput->GetInputParameterHandle().GetName() == ClipboardFunctionInput->InputName && ChildInput->GetInputType() == ClipboardFunctionInput->InputType)
+			{
+				ChildInput->SetValueFromClipboardFunctionInput(*ClipboardFunctionInput);
+			}
+		}
+	}
+}
+
 bool UNiagaraStackInputCategory::FilterForVisibleCondition(const UNiagaraStackEntry& Child) const
 {
 	const UNiagaraStackFunctionInput* StackFunctionInputChild = Cast<UNiagaraStackFunctionInput>(&Child);
diff --git a/Engine/Plugins/FX/Niagara/Source/NiagaraEditor/Public/NiagaraClipboard.h b/Engine/Plugins/FX/Niagara/Source/NiagaraEditor/Public/NiagaraClipboard.h
new file mode 100644
--- /dev/null
+++ b/Engine/Plugins/FX/Niagara/Source/NiagaraEditor/Public/NiagaraClipboard.h
@@ -0,0 +1,118 @@
+// Copyright 1998-2019 Epic Games, Inc. All Rights Reserved.
+
+#pragma once
+
+#include "NiagaraTypes.h"
+#include "UObject/GCObject.h"
+#include "NiagaraClipboard.generated.h"
+
+class UNiagaraDataInterface;
+class UNiagaraScript;
+
+UENUM()
+enum class ENiagaraClipboardFunctionInputValueMode
+{
+	Local,
+	Linked,
+	Data,
+	Expression,
+	Dynamic
+};
+
+class UNiagaraClipboardFunction;
+
+UCLASS()
+class NIAGARAEDITOR_API UNiagaraClipboardFunctionInput : public UObject
+{
+	GENERATED_BODY()
+
+public:
+	static const UNiagaraClipboardFunctionInput* CreateLocalValue(UObject* InOuter, FName InInputName, FNiagaraTypeDefinition InInputType, TArray<uint8>& InLocalValueData);
+
+	static const UNiagaraClipboardFunctionInput* CreateLinkedValue(UObject* InOuter, FName InInputName, FNiagaraTypeDefinition InInputType, FName InLinkedValue);
+
+	static const UNiagaraClipboardFunctionInput* CreateDataValue(UObject* InOuter, FName InInputName, FNiagaraTypeDefinition InInputType, UNiagaraDataInterface* InDataValue);
+
+	static const UNiagaraClipboardFunctionInput* CreateExpressionValue(UObject* InOuter, FName InInputName, FNiagaraTypeDefinition InInputType, const FString& InExpressionValue);
+
+	static const UNiagaraClipboardFunctionInput* CreateDynamicValue(UObject* InOuter, FName InInputName, FNiagaraTypeDefinition InInputType, UNiagaraScript* InDynamicValue);
+
+	UPROPERTY()
+	FName InputName;
+
+	UPROPERTY()
+	FNiagaraTypeDefinition InputType;
+
+	UPROPERTY()
+	ENiagaraClipboardFunctionInputValueMode ValueMode;
+
+	UPROPERTY()
+	TArray<uint8> Local;
+
+	UPROPERTY()
+	FName Linked;
+
+	UPROPERTY()
+	UNiagaraDataInterface* Data;
+
+	UPROPERTY()
+	FString Expression;
+
+	UPROPERTY()
+	UNiagaraClipboardFunction* Dynamic;
+};
+
+UENUM()
+enum class ENiagaraClipboardFunctionScriptMode
+{
+	ScriptAsset,
+	Assignment
+};
+
+UCLASS()
+class NIAGARAEDITOR_API UNiagaraClipboardFunction : public UObject
+{
+	GENERATED_BODY()
+
+public:
+	static UNiagaraClipboardFunction* CreateScriptFunction(UObject* InOuter, UNiagaraScript* InScript);
+
+	UPROPERTY()
+	ENiagaraClipboardFunctionScriptMode ScriptMode;
+
+	UPROPERTY()
+	UNiagaraScript* Script;
+
+	UPROPERTY()
+	TArray<FNiagaraVariable> AssignmentTargets;
+
+	UPROPERTY()
+	TArray<const UNiagaraClipboardFunctionInput*> Inputs;
+};
+
+UCLASS()
+class UNiagaraClipboardContent : public UObject
+{
+	GENERATED_BODY()
+
+public:
+	UPROPERTY()
+	TArray<const UNiagaraClipboardFunction*> Functions;
+
+	UPROPERTY()
+	TArray<const UNiagaraClipboardFunctionInput*> FunctionInputs;
+};
+
+class FNiagaraClipboard
+{
+public:
+	FNiagaraClipboard();
+
+	void Copy(const UNiagaraClipboardFunctionInput* Input);
+
+	const UNiagaraClipboardFunctionInput* GetCopiedInput() const;
+
+private:
+	void SetClipboardContentInternal(UNiagaraClipboardContent* ClipboardContent);
+	const UNiagaraClipboardContent* GetClipboardContentInternal() const;
+};
\ No newline at end of file
diff --git a/Engine/Plugins/FX/Niagara/Source/NiagaraEditor/Public/NiagaraEditorModule.h b/Engine/Plugins/FX/Niagara/Source/NiagaraEditor/Public/NiagaraEditorModule.h
--- a/Engine/Plugins/FX/Niagara/Source/NiagaraEditor/Public/NiagaraEditorModule.h
+++ b/Engine/Plugins/FX/Niagara/Source/NiagaraEditor/Public/NiagaraEditorModule.h
@@ -26,6 +26,7 @@ struct IConsoleCommand;
 class INiagaraEditorOnlyDataUtilities;
 class FNiagaraEditorCommands;
 struct FNiagaraScriptHighlight;
+class FNiagaraClipboard;
 
 DECLARE_STATS_GROUP(TEXT("Niagara Editor"), STATGROUP_NiagaraEditor, STATCAT_Advanced);
 
@@ -113,6 +114,8 @@ public:
 
 	void GetScriptAssetsMatchingHighlight(const FNiagaraScriptHighlight& InHighlight, TArray<FAssetData>& OutMatchingScriptAssets) const;
 
+	FNiagaraClipboard& GetClipboard() const;
+
 private:
 	void RegisterAssetTypeAction(IAssetTools& AssetTools, TSharedRef<IAssetTypeActions> Action);
 	void OnNiagaraSettingsChangedEvent(const FString& PropertyName, const UNiagaraSettings* Settings);
@@ -176,4 +179,6 @@ private:
 	mutable TOptional<TArray<FNiagaraScriptHighlight>> CachedScriptAssetHighlights;
 
 	bool bThumbnailRenderersRegistered;
+
+	TSharedRef<FNiagaraClipboard> Clipboard;
 };
diff --git a/Engine/Plugins/FX/Niagara/Source/NiagaraEditor/Public/NiagaraNodeCustomHlsl.h b/Engine/Plugins/FX/Niagara/Source/NiagaraEditor/Public/NiagaraNodeCustomHlsl.h
--- a/Engine/Plugins/FX/Niagara/Source/NiagaraEditor/Public/NiagaraNodeCustomHlsl.h
+++ b/Engine/Plugins/FX/Niagara/Source/NiagaraEditor/Public/NiagaraNodeCustomHlsl.h
@@ -17,8 +17,8 @@ class UNiagaraNodeCustomHlsl : public UNiagaraNodeFunctionCall
 	GENERATED_UCLASS_BODY()
 
 public:
-	UPROPERTY(EditAnywhere, Category = "Function", meta = (MultiLine = true))
-	FString CustomHlsl;
+	const FString& GetCustomHlsl() const;
+	void SetCustomHlsl(const FString& InCustomHlsl);
 
 	UPROPERTY()
 	ENiagaraScriptUsage ScriptUsage;
@@ -27,8 +27,8 @@ public:
 	virtual void OnRenameNode(const FString& NewName) override;
 	virtual FLinearColor GetNodeTitleColor() const override;
 
-	FText NIAGARAEDITOR_API GetHlslText() const;
-	void NIAGARAEDITOR_API OnCustomHlslTextCommitted(const FText& InText, ETextCommit::Type InType);
+	FText GetHlslText() const;
+	void OnCustomHlslTextCommitted(const FText& InText, ETextCommit::Type InType);
 
 	bool GetTokens(TArray<FString>& OutTokens) const;
 
@@ -76,4 +76,7 @@ protected:
 	void RebuildSignatureFromPins();
 	UEdGraphPin* PinPendingRename;
 
+private:
+	UPROPERTY(EditAnywhere, Category = "Function", meta = (MultiLine = true))
+	FString CustomHlsl;
 };
diff --git a/Engine/Plugins/FX/Niagara/Source/NiagaraEditor/Public/ViewModels/Stack/NiagaraStackEntry.h b/Engine/Plugins/FX/Niagara/Source/NiagaraEditor/Public/ViewModels/Stack/NiagaraStackEntry.h
--- a/Engine/Plugins/FX/Niagara/Source/NiagaraEditor/Public/ViewModels/Stack/NiagaraStackEntry.h
+++ b/Engine/Plugins/FX/Niagara/Source/NiagaraEditor/Public/ViewModels/Stack/NiagaraStackEntry.h
@@ -263,6 +263,21 @@ public:
 
 	void GetUnfilteredChildren(TArray<UNiagaraStackEntry*>& OutUnfilteredChildren) const;
 
+	template<typename T>
+	void GetUnfilteredChildrenOfType(TArray<T*>& OutUnfilteredChldrenOfType) const
+	{
+		TArray<UNiagaraStackEntry*> UnfilteredChildren;
+		GetUnfilteredChildren(UnfilteredChildren);
+		for (UNiagaraStackEntry* UnfilteredChild : UnfilteredChildren)
+		{
+			T* UnfilteredChildOfType = Cast<T>(UnfilteredChild);
+			if (UnfilteredChildOfType != nullptr)
+			{
+				OutUnfilteredChldrenOfType.Add(UnfilteredChildOfType);
+			}
+		}
+	}
+
 	FOnStructureChanged& OnStructureChanged();
 
 	FOnDataObjectModified& OnDataObjectModified();
@@ -329,6 +344,24 @@ public:
 
 	const TArray<UNiagaraStackEntry*>& GetAllChildrenWithIssues() const;
 
+	virtual bool SupportsCut() const { return false; }
+
+	virtual bool TestCanCutWithMessage(FText& OutMessage) const { return false; }
+
+	virtual void Cut() { }
+
+	virtual bool SupportsCopy() const { return false; }
+
+	virtual bool TestCanCopyWithMessage(FText& OutMessage) const { return false; }
+
+	virtual void Copy() const { }
+
+	virtual bool SupportsPaste() const { return false; }
+
+	virtual bool TestCanPasteWithMessage(FText& OutMessage) const { return false; }
+
+	virtual void Paste() { }
+
 protected:
 	virtual void BeginDestroy() override;
 
diff --git a/Engine/Plugins/FX/Niagara/Source/NiagaraEditor/Public/ViewModels/Stack/NiagaraStackFunctionInput.h b/Engine/Plugins/FX/Niagara/Source/NiagaraEditor/Public/ViewModels/Stack/NiagaraStackFunctionInput.h
--- a/Engine/Plugins/FX/Niagara/Source/NiagaraEditor/Public/ViewModels/Stack/NiagaraStackFunctionInput.h
+++ b/Engine/Plugins/FX/Niagara/Source/NiagaraEditor/Public/ViewModels/Stack/NiagaraStackFunctionInput.h
@@ -19,6 +19,8 @@ class UNiagaraScript;
 class UEdGraphPin;
 class UNiagaraDataInterface;
 enum class EStackParameterBehavior;
+class UNiagaraClipboardFunctionInput;
+class UNiagaraClipboardFunction;
 
 /** Represents a single module input in the module stack view model. */
 UCLASS()
@@ -85,8 +87,16 @@ public:
 	virtual FText GetDisplayName() const override;
 	virtual FText GetTooltipText() const override;
 	virtual bool GetIsEnabled() const override;
-	
-	UObject* GetExternalAsset() const override;
+	virtual UObject* GetExternalAsset() const override;
+	virtual bool SupportsCut() const override { return true; }
+	virtual bool TestCanCutWithMessage(FText& OutMessage) const override;
+	virtual void Cut() override;
+	virtual bool SupportsCopy() const override { return true; }
+	virtual bool TestCanCopyWithMessage(FText& OutMessage) const override;
+	virtual void Copy() const override;
+	virtual bool SupportsPaste() const override { return true; }
+	virtual bool TestCanPasteWithMessage(FText& OutMessage) const override;
+	virtual void Paste() override;
 
 	FText GetTooltipText(EValueMode InValueMode) const;
 
@@ -108,17 +118,17 @@ public:
 	/** Gets the dynamic input node providing the value for this input, if one is available. */
 	UNiagaraNodeFunctionCall* GetDynamicInputNode() const;
 
-	/** Gets the expression input node providing the value for this input, if one is available. */
-	UNiagaraNodeCustomHlsl* GetExpressionNode() const;
-
 	/** Gets the dynamic inputs available for this input. */
 	void GetAvailableDynamicInputs(TArray<UNiagaraScript*>& AvailableDynamicInputs, bool bIncludeNonLibraryInputs = false);
 
 	/** Sets the dynamic input script for this input. */
 	void SetDynamicInput(UNiagaraScript* DynamicInput);
 
+	/** Gets the expression providing the value for this input, if one is available. */
+	FText GetCustomExpressionText() const;
+
 	/** Sets the dynamic custom expression script for this input. */
-	void SetCustomExpression(const FString& InputText);
+	void SetCustomExpression(const FString& InCustomExpression);
 
 	/** Gets the current struct value of this input is there is one. */
 	TSharedPtr<FStructOnScope> GetLocalValueStruct();
@@ -216,6 +226,10 @@ public:
 	/** Gets whether or not this input is filtered from search results and appearing in stack due to visibility metadata*/
 	bool GetShouldPassFilterForVisibleCondition() const;
 
+	const UNiagaraClipboardFunctionInput* ToClipboardFunctionInput(UObject* InOuter) const;
+
+	void SetValueFromClipboardFunctionInput(const UNiagaraClipboardFunctionInput& ClipboardFunctionInput);
+
 public:
 	//~ UNiagaraStackEntry interface
 	virtual void GetSearchItems(TArray<FStackSearchItem>& SearchItems) const override;
@@ -366,6 +380,8 @@ private:
 
 	FNiagaraVariable CreateRapidIterationVariable(const FName& InName);
 
+	void PrepareFunctionInputForNewValue();
+
 private:
 	/** The module function call which owns this input entry. NOTE: This input might not be an input to the module function
 		call, it may be an input to a dynamic input function call which is owned by the module. */
diff --git a/Engine/Plugins/FX/Niagara/Source/NiagaraEditor/Public/ViewModels/Stack/NiagaraStackFunctionInputCollection.h b/Engine/Plugins/FX/Niagara/Source/NiagaraEditor/Public/ViewModels/Stack/NiagaraStackFunctionInputCollection.h
--- a/Engine/Plugins/FX/Niagara/Source/NiagaraEditor/Public/ViewModels/Stack/NiagaraStackFunctionInputCollection.h
+++ b/Engine/Plugins/FX/Niagara/Source/NiagaraEditor/Public/ViewModels/Stack/NiagaraStackFunctionInputCollection.h
@@ -8,6 +8,7 @@
 
 class UNiagaraNodeFunctionCall;
 class UNiagaraStackFunctionInput;
+class UNiagaraClipboardFunctionInput;
 class UEdGraphPin;
 
 UCLASS()
@@ -35,6 +36,10 @@ public:
 
 	void SetShouldShowInStack(bool bInShouldShowInStack);
 
+	void ToClipboardFunctionInputs(UObject* InOuter, TArray<const UNiagaraClipboardFunctionInput*>& OutClipboardFunctionInputs) const;
+
+	void SetValuesFromClipboardFunctionInputs(const TArray<const UNiagaraClipboardFunctionInput*>& ClipboardFunctionInputs);
+
 protected:
 	virtual void FinalizeInternal() override;
 
diff --git a/Engine/Plugins/FX/Niagara/Source/NiagaraEditor/Public/ViewModels/Stack/NiagaraStackGraphUtilities.h b/Engine/Plugins/FX/Niagara/Source/NiagaraEditor/Public/ViewModels/Stack/NiagaraStackGraphUtilities.h
--- a/Engine/Plugins/FX/Niagara/Source/NiagaraEditor/Public/ViewModels/Stack/NiagaraStackGraphUtilities.h
+++ b/Engine/Plugins/FX/Niagara/Source/NiagaraEditor/Public/ViewModels/Stack/NiagaraStackGraphUtilities.h
@@ -103,7 +103,7 @@ namespace FNiagaraStackGraphUtilities
 
 	void SetDynamicInputForFunctionInput(UEdGraphPin& OverridePin, UNiagaraScript* DynamicInput, UNiagaraNodeFunctionCall*& OutDynamicInputFunctionCall, const FGuid& NewNodePersistentId = FGuid(), FString SuggestedName = FString());
 
-	void SetCustomExpressionForFunctionInput(UEdGraphPin& OverridePin, UNiagaraNodeCustomHlsl*& OutDynamicInputFunctionCall, const FGuid& NewNodePersistentId = FGuid());
+	void SetCustomExpressionForFunctionInput(UEdGraphPin& OverridePin, const FString& CustomExpression, UNiagaraNodeCustomHlsl*& OutDynamicInputFunctionCall, const FGuid& NewNodePersistentId = FGuid());
 
 	bool RemoveModuleFromStack(UNiagaraSystem& OwningSystem, FGuid OwningEmitterId, UNiagaraNodeFunctionCall& ModuleNode);
 
diff --git a/Engine/Plugins/FX/Niagara/Source/NiagaraEditor/Public/ViewModels/Stack/NiagaraStackInputCategory.h b/Engine/Plugins/FX/Niagara/Source/NiagaraEditor/Public/ViewModels/Stack/NiagaraStackInputCategory.h
--- a/Engine/Plugins/FX/Niagara/Source/NiagaraEditor/Public/ViewModels/Stack/NiagaraStackInputCategory.h
+++ b/Engine/Plugins/FX/Niagara/Source/NiagaraEditor/Public/ViewModels/Stack/NiagaraStackInputCategory.h
@@ -8,6 +8,7 @@
 
 class UNiagaraStackFunctionInput;
 class UNiagaraNodeFunctionCall;
+class UNiagaraClipboardFunctionInput;
 
 UENUM()
 enum class EStackParameterBehavior
@@ -42,6 +43,10 @@ public:
 
 	void SetShouldShowInStack(bool bInShouldShowInStack);
 
+	void ToClipboardFunctionInputs(UObject* InOuter, TArray<const UNiagaraClipboardFunctionInput*>& OutClipboardFunctionInputs) const;
+
+	void SetValuesFromClipboardFunctionInputs(const TArray<const UNiagaraClipboardFunctionInput*>& ClipboardFunctionInputs);
+
 protected:
 	//~ UNiagaraStackEntry interface
 	virtual void RefreshChildrenInternal(const TArray<UNiagaraStackEntry*>& CurrentChildren, TArray<UNiagaraStackEntry*>& NewChildren, TArray<FStackIssue>& NewIssues) override;
diff --git a/Engine/Plugins/FX/Niagara/Source/NiagaraEditorWidgets/Private/NiagaraEditorWidgetsUtilities.cpp b/Engine/Plugins/FX/Niagara/Source/NiagaraEditorWidgets/Private/NiagaraEditorWidgetsUtilities.cpp
--- a/Engine/Plugins/FX/Niagara/Source/NiagaraEditorWidgets/Private/NiagaraEditorWidgetsUtilities.cpp
+++ b/Engine/Plugins/FX/Niagara/Source/NiagaraEditorWidgets/Private/NiagaraEditorWidgetsUtilities.cpp
@@ -161,6 +161,73 @@ bool FNiagaraStackEditorWidgetsUtilities::AddStackEntryAssetContextMenuActions(F
 	return false;
 }
 
+void CopyEntry(TWeakObjectPtr<UNiagaraStackEntry> StackEntryWeak)
+{
+	UNiagaraStackEntry* StackEntry = StackEntryWeak.Get();
+	if (StackEntry != nullptr)
+	{
+		StackEntry->Copy();
+	}
+}
+
+void PasteEntry(TWeakObjectPtr<UNiagaraStackEntry> StackEntryWeak)
+{
+	UNiagaraStackEntry* StackEntry = StackEntryWeak.Get();
+	if (StackEntry != nullptr)
+	{
+		StackEntry->Paste();
+	}
+}
+
+bool FNiagaraStackEditorWidgetsUtilities::AddStackEntryContextMenuActions(FMenuBuilder& MenuBuilder, UNiagaraStackEntry& StackEntry)
+{
+	if (StackEntry.SupportsCut() || StackEntry.SupportsCopy() || StackEntry.SupportsPaste())
+	{
+		MenuBuilder.BeginSection("EntryEdit", LOCTEXT("EntryEditActions", "Edit"));
+		{
+			if (StackEntry.SupportsCut())
+			{
+				FText CutMessage;
+				bool bCanCut = StackEntry.TestCanCutWithMessage(CutMessage);
+				MenuBuilder.AddMenuEntry(
+					LOCTEXT("CutAction", "Cut"),
+					CutMessage,
+					FSlateIcon(),
+					FUIAction(
+						FExecuteAction::CreateUObject(&StackEntry, &UNiagaraStackEntry::Cut),
+						FCanExecuteAction::CreateLambda([bCanCut]() { return bCanCut; })));
+			}
+			if (StackEntry.SupportsCopy())
+			{
+				FText CopyMessage;
+				bool bCanCopy = StackEntry.TestCanCopyWithMessage(CopyMessage);
+				MenuBuilder.AddMenuEntry(
+					LOCTEXT("CopyAction", "Copy"),
+					CopyMessage,
+					FSlateIcon(),
+					FUIAction(
+						FExecuteAction::CreateStatic(&CopyEntry, TWeakObjectPtr<UNiagaraStackEntry>(&StackEntry)),
+						FCanExecuteAction::CreateLambda([bCanCopy]() { return bCanCopy; })));
+			}
+			if (StackEntry.SupportsPaste())
+			{
+				FText PasteMessage;
+				bool bCanPaste = StackEntry.TestCanPasteWithMessage(PasteMessage);
+				MenuBuilder.AddMenuEntry(
+					LOCTEXT("PasteAction", "Paste"),
+					PasteMessage,
+					FSlateIcon(),
+					FUIAction(
+						FExecuteAction::CreateStatic(&PasteEntry, TWeakObjectPtr<UNiagaraStackEntry>(&StackEntry)),
+						FCanExecuteAction::CreateLambda([bCanPaste]() { return bCanPaste; })));
+			}
+		}
+		MenuBuilder.EndSection();
+		return true;
+	}
+	return false;
+}
+
 void DeleteItem(TWeakObjectPtr<UNiagaraStackItem> StackItemWeak)
 {
 	UNiagaraStackItem* StackItem = StackItemWeak.Get();
diff --git a/Engine/Plugins/FX/Niagara/Source/NiagaraEditorWidgets/Private/NiagaraEditorWidgetsUtilities.h b/Engine/Plugins/FX/Niagara/Source/NiagaraEditorWidgets/Private/NiagaraEditorWidgetsUtilities.h
--- a/Engine/Plugins/FX/Niagara/Source/NiagaraEditorWidgets/Private/NiagaraEditorWidgetsUtilities.h
+++ b/Engine/Plugins/FX/Niagara/Source/NiagaraEditorWidgets/Private/NiagaraEditorWidgetsUtilities.h
@@ -22,6 +22,8 @@ namespace FNiagaraStackEditorWidgetsUtilities
 	
 	bool AddStackEntryAssetContextMenuActions(FMenuBuilder& MenuBuilder, UNiagaraStackEntry& StackEntry);
 
+	bool AddStackEntryContextMenuActions(FMenuBuilder& MenuBuilder, UNiagaraStackEntry& StackEntry);
+
 	bool AddStackItemContextMenuActions(FMenuBuilder& MenuBuilder, UNiagaraStackItem& StackItem);
 
 	bool AddStackModuleItemContextMenuActions(FMenuBuilder& MenuBuilder, UNiagaraStackModuleItem& StackItem, TSharedRef<SWidget> TargetWidget);
diff --git a/Engine/Plugins/FX/Niagara/Source/NiagaraEditorWidgets/Private/SNiagaraOverviewStack.cpp b/Engine/Plugins/FX/Niagara/Source/NiagaraEditorWidgets/Private/SNiagaraOverviewStack.cpp
--- a/Engine/Plugins/FX/Niagara/Source/NiagaraEditorWidgets/Private/SNiagaraOverviewStack.cpp
+++ b/Engine/Plugins/FX/Niagara/Source/NiagaraEditorWidgets/Private/SNiagaraOverviewStack.cpp
@@ -126,6 +126,7 @@ class SNiagaraSystemOverviewEntryListRow : public STableRow<UNiagaraStackEntry*>
 			}
 
 			bMenuItemsAdded |= FNiagaraStackEditorWidgetsUtilities::AddStackEntryAssetContextMenuActions(MenuBuilder, *StackEntry);
+			bMenuItemsAdded |= FNiagaraStackEditorWidgetsUtilities::AddStackEntryContextMenuActions(MenuBuilder, *StackEntry);
 		
 			if (bMenuItemsAdded)
 			{
diff --git a/Engine/Plugins/FX/Niagara/Source/NiagaraEditorWidgets/Private/Stack/SNiagaraStackFunctionInputValue.cpp b/Engine/Plugins/FX/Niagara/Source/NiagaraEditorWidgets/Private/Stack/SNiagaraStackFunctionInputValue.cpp
--- a/Engine/Plugins/FX/Niagara/Source/NiagaraEditorWidgets/Private/Stack/SNiagaraStackFunctionInputValue.cpp
+++ b/Engine/Plugins/FX/Niagara/Source/NiagaraEditorWidgets/Private/Stack/SNiagaraStackFunctionInputValue.cpp
@@ -145,15 +145,9 @@ void SNiagaraStackFunctionInputValue::Construct(const FArguments& InArgs, UNiaga
 					.Visibility(this, &SNiagaraStackFunctionInputValue::GetValueWidgetVisibility, UNiagaraStackFunctionInput::EValueMode::Expression)
 					.VAlign(VAlign_Center)
 					[
-						//SNew(SInlineEditableTextBlock)
-						//.Style(FNiagaraEditorStyle::Get(), "NiagaraEditor.ParameterInlineEditableText")
-						//.Text(this, &SNiagaraStackFunctionInputValue::GetExpressionValueText)
-						//.OnTextCommitted(this, &SNiagaraStackFunctionInputValue::OnExpressionTextCommitted)
-						// SNew(SMultiLineEditableTextBox)
-						//.AutoWrapText(false)
 						SNew(SEditableTextBox)
 						.IsReadOnly(false)
-						.Text(this, &SNiagaraStackFunctionInputValue::GetExpressionValueText)
+						.Text_UObject(FunctionInput, &UNiagaraStackFunctionInput::GetCustomExpressionText)
 						.OnTextCommitted(this, &SNiagaraStackFunctionInputValue::OnExpressionTextCommitted)
 					]
 				]
@@ -391,24 +385,9 @@ FText SNiagaraStackFunctionInputValue::GetDynamicValueToolTip() const
 	}
 }
 
-FText SNiagaraStackFunctionInputValue::GetExpressionValueText() const
-{
-	if (FunctionInput->GetExpressionNode() != nullptr)
-	{
-		return FunctionInput->GetExpressionNode()->GetHlslText();
-	}
-	else
-	{
-		return LOCTEXT("InvalidDynamicDisplayName", "(Invalid)");
-	}
-}
-
 void SNiagaraStackFunctionInputValue::OnExpressionTextCommitted(const FText& Name, ETextCommit::Type CommitInfo)
 {
-	if (FunctionInput->GetExpressionNode() != nullptr)
-	{
-		FunctionInput->GetExpressionNode()->OnCustomHlslTextCommitted(Name, CommitInfo);
-	}
+	FunctionInput->SetCustomExpression(Name.ToString());
 }
 
 FText SNiagaraStackFunctionInputValue::GetInvalidValueText() const
diff --git a/Engine/Plugins/FX/Niagara/Source/NiagaraEditorWidgets/Private/Stack/SNiagaraStackFunctionInputValue.h b/Engine/Plugins/FX/Niagara/Source/NiagaraEditorWidgets/Private/Stack/SNiagaraStackFunctionInputValue.h
--- a/Engine/Plugins/FX/Niagara/Source/NiagaraEditorWidgets/Private/Stack/SNiagaraStackFunctionInputValue.h
+++ b/Engine/Plugins/FX/Niagara/Source/NiagaraEditorWidgets/Private/Stack/SNiagaraStackFunctionInputValue.h
@@ -53,8 +53,6 @@ private:
 
 	FText GetDynamicValueToolTip() const;
 
-	FText GetExpressionValueText() const;
-
 	void OnExpressionTextCommitted(const FText& Name, ETextCommit::Type CommitInfo);
 
 	FText GetInvalidValueText() const;
diff --git a/Engine/Plugins/FX/Niagara/Source/NiagaraEditorWidgets/Private/Stack/SNiagaraStackTableRow.cpp b/Engine/Plugins/FX/Niagara/Source/NiagaraEditorWidgets/Private/Stack/SNiagaraStackTableRow.cpp
--- a/Engine/Plugins/FX/Niagara/Source/NiagaraEditorWidgets/Private/Stack/SNiagaraStackTableRow.cpp
+++ b/Engine/Plugins/FX/Niagara/Source/NiagaraEditorWidgets/Private/Stack/SNiagaraStackTableRow.cpp
@@ -293,6 +293,7 @@ FReply SNiagaraStackTableRow::OnMouseButtonUp(const FGeometry& MyGeometry, const
 		}
 
 		FNiagaraStackEditorWidgetsUtilities::AddStackEntryAssetContextMenuActions(MenuBuilder, *StackEntry);
+		FNiagaraStackEditorWidgetsUtilities::AddStackEntryContextMenuActions(MenuBuilder, *StackEntry);
 
 		TArray<UNiagaraStackEntry*> EntriesToProcess;
 		TArray<UNiagaraStackEntry*> NavigationEntries;
@@ -301,25 +302,27 @@ FReply SNiagaraStackTableRow::OnMouseButtonUp(const FGeometry& MyGeometry, const
 		{
 			UNiagaraStackItemGroup* GroupParent = Cast<UNiagaraStackItemGroup>(Parent);
 			UNiagaraStackItem* ItemParent = Cast<UNiagaraStackItem>(Parent);
-			if (GroupParent != nullptr)
-			{
-				MenuBuilder.BeginSection("StackRowNavigateTo", LOCTEXT("NavigateToSection", "Navigate to:"));
-				MenuBuilder.AddMenuEntry(
-					LOCTEXT("TopOfSection", "Top of Section"),
-					FText::Format(LOCTEXT("NavigateToFormatted", "Navigate to {0}"), Parent->GetDisplayName()),
-					FSlateIcon(),
-					FUIAction(FExecuteAction::CreateSP(this, &SNiagaraStackTableRow::NavigateTo, Parent)));
-			}
-			if (ItemParent != nullptr)
-			{
-				MenuBuilder.AddMenuEntry(
-					LOCTEXT("TopOfModule", "Top of Module"),
-					FText::Format(LOCTEXT("NavigateToFormatted", "Navigate to {0}"), Parent->GetDisplayName()),
-					FSlateIcon(),
-					FUIAction(FExecuteAction::CreateSP(this, &SNiagaraStackTableRow::NavigateTo, Parent)));
-			}
-			if (GroupParent != nullptr)
+			if (GroupParent != nullptr || ItemParent != nullptr)
 			{
+				MenuBuilder.BeginSection("StackRowNavigation", LOCTEXT("NavigationMenuSection", "Navigation"));
+				{
+					if (GroupParent != nullptr)
+					{
+						MenuBuilder.AddMenuEntry(
+							LOCTEXT("TopOfSection", "Top of Section"),
+							FText::Format(LOCTEXT("NavigateToFormatted", "Navigate to {0}"), Parent->GetDisplayName()),
+							FSlateIcon(),
+							FUIAction(FExecuteAction::CreateSP(this, &SNiagaraStackTableRow::NavigateTo, Parent)));
+					}
+					if (ItemParent != nullptr)
+					{
+						MenuBuilder.AddMenuEntry(
+							LOCTEXT("TopOfModule", "Top of Module"),
+							FText::Format(LOCTEXT("NavigateToFormatted", "Navigate to {0}"), Parent->GetDisplayName()),
+							FSlateIcon(),
+							FUIAction(FExecuteAction::CreateSP(this, &SNiagaraStackTableRow::NavigateTo, Parent)));
+					}
+				}
 				MenuBuilder.EndSection();
 			}
 		}
