Niagara - Fix an issue with changes related to copy/paste which was causing the niagara script editor to crash.

diff --git a/Engine/Plugins/FX/Niagara/Source/NiagaraEditor/Private/NiagaraNodeCustomHlsl.cpp b/Engine/Plugins/FX/Niagara/Source/NiagaraEditor/Private/NiagaraNodeCustomHlsl.cpp
--- a/Engine/Plugins/FX/Niagara/Source/NiagaraEditor/Private/NiagaraNodeCustomHlsl.cpp
+++ b/Engine/Plugins/FX/Niagara/Source/NiagaraEditor/Private/NiagaraNodeCustomHlsl.cpp
@@ -31,7 +31,12 @@ void UNiagaraNodeCustomHlsl::SetCustomHlsl(const FString& InCustomHlsl)
 	Modify();
 	CustomHlsl = InCustomHlsl;
 	RefreshFromExternalChanges();
-	MarkNodeRequiresSynchronization(__FUNCTION__, true);
+	if (GetOuter()->IsA<UNiagaraGraph>())
+	{
+		// This is needed to guard against a crash when setting this value before the node has actually been
+		// added to a graph.
+		MarkNodeRequiresSynchronization(__FUNCTION__, true);
+	}
 }
 
 TSharedPtr<SGraphNode> UNiagaraNodeCustomHlsl::CreateVisualWidget()
