* Changed DI function parameter names
* Sample circular occlusion in a spiral pattern
* Changed DI tick group to remove frame latency on the cpu

diff --git a/Engine/Plugins/FX/Niagara/Source/Niagara/Classes/NiagaraDataInterfaceCamera.h b/Engine/Plugins/FX/Niagara/Source/Niagara/Classes/NiagaraDataInterfaceCamera.h
--- a/Engine/Plugins/FX/Niagara/Source/Niagara/Classes/NiagaraDataInterfaceCamera.h
+++ b/Engine/Plugins/FX/Niagara/Source/Niagara/Classes/NiagaraDataInterfaceCamera.h
@@ -11,6 +11,7 @@
 struct CameraDataInterface_InstanceData
 {
 	FVector CameraLocation;
+	FRotator CameraRotation;
 	float CameraFOV;
 };
 
@@ -37,12 +38,14 @@ public:
 	virtual bool GetFunctionHLSL(const FName&  DefinitionFunctionName, FString InstanceFunctionName, FNiagaraDataInterfaceGPUParamInfo& ParamInfo, FString& OutHLSL) override;
 	virtual bool CanExecuteOnTarget(ENiagaraSimTarget Target) const override { return true; }
 	virtual FNiagaraDataInterfaceParametersCS* ConstructComputeParameters() const override;
+	virtual bool HasTickGroupPrereqs() const override { return true; }
+	virtual ETickingGroup CalculateTickGroup(void* PerInstanceData) const override;
 	//UNiagaraDataInterface Interface
 
 	void QueryOcclusionFactorGPU(FVectorVMContext& Context);
 	void QueryOcclusionFactorCircleGPU(FVectorVMContext& Context);
 	void GetCameraFOV(FVectorVMContext& Context);
-	void GetCameraPosition(FVectorVMContext& Context);
+	void GetCameraProperties(FVectorVMContext& Context);
 	void GetViewPropertiesGPU(FVectorVMContext& Context);
 	void GetClipSpaceTransformsGPU(FVectorVMContext& Context);
 	void GetViewSpaceTransformsGPU(FVectorVMContext& Context);
@@ -53,7 +56,7 @@ private:
 	static const FName GetViewPropertiesName;
 	static const FName GetClipSpaceTransformsName;
 	static const FName GetViewSpaceTransformsName;
-	static const FName GetCameraPositionsName;
+	static const FName GetCameraPropertiesName;
 	static const FName GetFieldOfViewName;
 };
 
diff --git a/Engine/Plugins/FX/Niagara/Source/Niagara/Private/NiagaraDataInterfaceCamera.cpp b/Engine/Plugins/FX/Niagara/Source/Niagara/Private/NiagaraDataInterfaceCamera.cpp
--- a/Engine/Plugins/FX/Niagara/Source/Niagara/Private/NiagaraDataInterfaceCamera.cpp
+++ b/Engine/Plugins/FX/Niagara/Source/Niagara/Private/NiagaraDataInterfaceCamera.cpp
@@ -14,7 +14,7 @@ const FName UNiagaraDataInterfaceCamera::GetCameraOcclusionCircleName(TEXT("Quer
 const FName UNiagaraDataInterfaceCamera::GetViewPropertiesName(TEXT("GetViewPropertiesGPU"));
 const FName UNiagaraDataInterfaceCamera::GetClipSpaceTransformsName(TEXT("GetClipSpaceTransformsGPU"));
 const FName UNiagaraDataInterfaceCamera::GetViewSpaceTransformsName(TEXT("GetViewSpaceTransformsGPU"));
-const FName UNiagaraDataInterfaceCamera::GetCameraPositionsName(TEXT("GetCameraPositionCPU/GPU"));
+const FName UNiagaraDataInterfaceCamera::GetCameraPropertiesName(TEXT("GetCameraPropertiesCPU/GPU"));
 const FName UNiagaraDataInterfaceCamera::GetFieldOfViewName(TEXT("GetFieldOfView"));
 
 UNiagaraDataInterfaceCamera::UNiagaraDataInterfaceCamera(FObjectInitializer const& ObjectInitializer)
@@ -57,13 +57,15 @@ bool UNiagaraDataInterfaceCamera::PerInstanceTick(void* PerInstanceData, FNiagar
 			if (i == PlayerControllerIndex && PlayerController)
 			{
 				PIData->CameraLocation = PlayerController->PlayerCameraManager->GetCameraLocation();
+				PIData->CameraRotation = PlayerController->PlayerCameraManager->GetCameraRotation();
 				PIData->CameraFOV = PlayerController->PlayerCameraManager->GetFOVAngle();
 			}
 			i++;
 		}
 		return false;
 	}
-	PIData->CameraLocation = FVector();
+	PIData->CameraLocation = FVector::ZeroVector;
+	PIData->CameraRotation = FRotator(0);
 	PIData->CameraFOV = 0;
 	
 	return false;
@@ -84,25 +86,25 @@ void UNiagaraDataInterfaceCamera::GetFunctions(TArray<FNiagaraFunctionSignature>
 	Sig.Inputs.Add(FNiagaraVariable(FNiagaraTypeDefinition::GetFloatDef(), TEXT("Sample Window Width World")));
 	Sig.Inputs.Add(FNiagaraVariable(FNiagaraTypeDefinition::GetFloatDef(), TEXT("Sample Window Height World")));
 	Sig.Inputs.Add(FNiagaraVariable(FNiagaraTypeDefinition::GetFloatDef(), TEXT("Sample Steps Per Line")));
-	Sig.Outputs.Add(FNiagaraVariable(FNiagaraTypeDefinition::GetFloatDef(), TEXT("Occlusion Factor")));
-	Sig.Outputs.Add(FNiagaraVariable(FNiagaraTypeDefinition::GetFloatDef(), TEXT("Sample Factor")));
+	Sig.Outputs.Add(FNiagaraVariable(FNiagaraTypeDefinition::GetFloatDef(), TEXT("Visibility Fraction")));
+	Sig.Outputs.Add(FNiagaraVariable(FNiagaraTypeDefinition::GetFloatDef(), TEXT("Sample Fraction")));
 	OutFunctions.Add(Sig);
 
 	Sig = FNiagaraFunctionSignature();
 	Sig.Name = GetCameraOcclusionCircleName;
 #if WITH_EDITORONLY_DATA
-	Sig.Description = NSLOCTEXT("Niagara", "GetCameraOcclusionCircleFunctionDescription", "This function returns the occlusion factor of a sprite. It samples the depth buffer in a circular window around the given world position and compares each sample with the camera distance.");
+	Sig.Description = NSLOCTEXT("Niagara", "GetCameraOcclusionCircleFunctionDescription", "This function returns the occlusion factor of a sprite. It samples the depth buffer in concentric rings around the given world position and compares each sample with the camera distance.");
 #endif
 	Sig.bMemberFunction = true;
 	Sig.bRequiresContext = false;
 	Sig.bSupportsCPU = false;
 	Sig.Inputs.Add(FNiagaraVariable(FNiagaraTypeDefinition(GetClass()), TEXT("Camera interface")));
 	Sig.Inputs.Add(FNiagaraVariable(FNiagaraTypeDefinition::GetVec3Def(), TEXT("Sample Center World Position")));
-	Sig.Inputs.Add(FNiagaraVariable(FNiagaraTypeDefinition::GetFloatDef(), TEXT("Sample Window Radius World")));
-	Sig.Inputs.Add(FNiagaraVariable(FNiagaraTypeDefinition::GetFloatDef(), TEXT("Number of sample rays")));
-	Sig.Inputs.Add(FNiagaraVariable(FNiagaraTypeDefinition::GetFloatDef(), TEXT("Sample Steps Per Ray")));
-	Sig.Outputs.Add(FNiagaraVariable(FNiagaraTypeDefinition::GetFloatDef(), TEXT("Occlusion Factor")));
-	Sig.Outputs.Add(FNiagaraVariable(FNiagaraTypeDefinition::GetFloatDef(), TEXT("Sample Factor")));
+	Sig.Inputs.Add(FNiagaraVariable(FNiagaraTypeDefinition::GetFloatDef(), TEXT("Sample Window Diameter World")));
+	Sig.Inputs.Add(FNiagaraVariable(FNiagaraTypeDefinition::GetFloatDef(), TEXT("Samples per ring")));
+	Sig.Inputs.Add(FNiagaraVariable(FNiagaraTypeDefinition::GetFloatDef(), TEXT("Number of sample rings")));
+	Sig.Outputs.Add(FNiagaraVariable(FNiagaraTypeDefinition::GetFloatDef(), TEXT("Visibility Fraction")));
+	Sig.Outputs.Add(FNiagaraVariable(FNiagaraTypeDefinition::GetFloatDef(), TEXT("Sample Fraction")));
 	OutFunctions.Add(Sig);
 
 
@@ -147,7 +149,7 @@ void UNiagaraDataInterfaceCamera::GetFunctions(TArray<FNiagaraFunctionSignature>
 	Sig = FNiagaraFunctionSignature();
 	Sig.Name = GetViewSpaceTransformsName;
 #if WITH_EDITORONLY_DATA
-	Sig.Description = NSLOCTEXT("Niagara", "GetViewSpaceTransformsDescription", "This function returns the relevant tranforms for the current view. Only valid for gpu particles.");
+	Sig.Description = NSLOCTEXT("Niagara", "GetViewSpaceTransformsDescription", "This function returns the relevant transforms for the current view. Only valid for gpu particles.");
 #endif
 	Sig.bMemberFunction = true;
 	Sig.bRequiresContext = false;
@@ -175,7 +177,7 @@ void UNiagaraDataInterfaceCamera::GetFunctions(TArray<FNiagaraFunctionSignature>
 
 
 	Sig = FNiagaraFunctionSignature();
-	Sig.Name = GetCameraPositionsName;
+	Sig.Name = GetCameraPropertiesName;
 #if WITH_EDITORONLY_DATA
 	Sig.Description = NSLOCTEXT("Niagara", "GetCameraPositionDescription", "This function returns the position of the currently active camera.");
 #endif
@@ -183,6 +185,9 @@ void UNiagaraDataInterfaceCamera::GetFunctions(TArray<FNiagaraFunctionSignature>
 	Sig.bRequiresContext = false;
 	Sig.Inputs.Add(FNiagaraVariable(FNiagaraTypeDefinition(GetClass()), TEXT("Camera interface")));
 	Sig.Outputs.Add(FNiagaraVariable(FNiagaraTypeDefinition::GetVec3Def(), TEXT("Camera Position World")));
+	Sig.Outputs.Add(FNiagaraVariable(FNiagaraTypeDefinition::GetVec3Def(), TEXT("Forward Vector World")));
+	Sig.Outputs.Add(FNiagaraVariable(FNiagaraTypeDefinition::GetVec3Def(), TEXT("Up Vector World")));
+	Sig.Outputs.Add(FNiagaraVariable(FNiagaraTypeDefinition::GetVec3Def(), TEXT("Right Vector World")));
 	OutFunctions.Add(Sig);
 }
 
@@ -194,7 +199,7 @@ bool UNiagaraDataInterfaceCamera::GetFunctionHLSL(const FName& DefinitionFunctio
 	if (DefinitionFunctionName == GetCameraOcclusionRectangleName)
 	{
 		static const TCHAR *FormatSample = TEXT(R"(
-			void {FunctionName}(in float3 In_SampleCenterWorldPos, in float In_SampleWindowWidthWorld, in float In_SampleWindowHeightWorld, in float In_SampleSteps, out float Out_OcclusionFactor, out float Out_SampleFactor)
+			void {FunctionName}(in float3 In_SampleCenterWorldPos, in float In_SampleWindowWidthWorld, in float In_SampleWindowHeightWorld, in float In_SampleSteps, out float Out_VisibilityFraction, out float Out_SampleFraction)
 			{
 				float CameraDistance = length(In_SampleCenterWorldPos.xyz - View.WorldViewOrigin.xyz);
 				float4 SamplePosition = float4(In_SampleCenterWorldPos + View.PreViewTranslation, 1);
@@ -241,8 +246,8 @@ bool UNiagaraDataInterfaceCamera::GetFunctionHLSL(const FName& DefinitionFunctio
 						} 
 					}
 				}
-				Out_OcclusionFactor = TotalSamples > 0 ? OccludedSamples / TotalSamples : 1;
-				Out_SampleFactor = Steps == 0 ? 0 : (TotalSamples / (Steps * Steps));
+				Out_VisibilityFraction = TotalSamples > 0 ? 1 - OccludedSamples / TotalSamples : 0;
+				Out_SampleFraction = Steps == 0 ? 0 : (TotalSamples / (Steps * Steps));
 			}
 		)");
 		OutHLSL += FString::Format(FormatSample, ArgsSample);
@@ -251,9 +256,10 @@ bool UNiagaraDataInterfaceCamera::GetFunctionHLSL(const FName& DefinitionFunctio
 	if (DefinitionFunctionName == GetCameraOcclusionCircleName)
 	{
 		static const TCHAR *FormatSample = TEXT(R"(
-			void {FunctionName}(in float3 In_SampleCenterWorldPos, in float In_SampleWindowRadiusWorld, in float In_SampleRays, in float In_SampleStepsPerRay, out float Out_OcclusionFactor, out float Out_SampleFactor)
+			void {FunctionName}(in float3 In_SampleCenterWorldPos, in float In_SampleWindowDiameterWorld, in float In_SampleRays, in float In_SampleStepsPerRay, out float Out_VisibilityFraction, out float Out_SampleFraction)
 			{
 				const float PI = 3.14159265;
+				const float SPIRAL_TURN = 2 * PI * 0.61803399; // use golden ratio to rotate sample pattern each ring so we get a spiral
 				float CameraDistance = length(In_SampleCenterWorldPos.xyz - View.WorldViewOrigin.xyz);
 				float4 SamplePosition = float4(In_SampleCenterWorldPos + View.PreViewTranslation, 1);
 				float4 ClipPosition = mul(SamplePosition, View.TranslatedWorldToClip);
@@ -277,22 +283,24 @@ bool UNiagaraDataInterfaceCamera::GetFunctionHLSL(const FName& DefinitionFunctio
 				if (Steps > 0) 
 				{
 					float Degrees = 0;
-					for (int ray = 0; ray < Rays; ray++)
+					for (int Step = 1; Step <= Steps; Step++)
 					{
-						// calc ray direction vector
-						float3 RayDirection = cos(Degrees) * View.ViewUp + sin(Degrees) * View.ViewRight;
-						float4 RayClip = mul(float4(RayDirection * In_SampleWindowRadiusWorld, 0) + SamplePosition, View.TranslatedWorldToClip);
-						float2 RayUV = RayClip.xy / RayClip.w * View.ScreenPositionScaleBias.xy + View.ScreenPositionScaleBias.wz;
-
-						if ((ScreenUV.x > 1 && RayUV.x < 0) || (ScreenUV.y > 1 && RayUV.y < 0) || (ScreenUV.x < 0 && RayUV.x > 1) || (ScreenUV.y < 0 && RayUV.y > 1))
+						float LerpFactor = Step / Steps;
+						Degrees += SPIRAL_TURN;
+						for (int ray = 0; ray < Rays; ray++)
 						{
-							continue;
-						}
+							// calc ray direction vector
+							float3 RayDirection = cos(Degrees) * View.ViewUp + sin(Degrees) * View.ViewRight;
+							float4 RayClip = mul(float4(RayDirection * In_SampleWindowDiameterWorld / 2, 0) + SamplePosition, View.TranslatedWorldToClip);
+							float2 RayUV = RayClip.xy / RayClip.w * View.ScreenPositionScaleBias.xy + View.ScreenPositionScaleBias.wz;
 
-						for (int Step = 1; Step <= Steps; Step++)
-						{
-							float LerpFactor = Step / Steps;
+							if ((ScreenUV.x > 1 && RayUV.x < 0) || (ScreenUV.y > 1 && RayUV.y < 0) || (ScreenUV.x < 0 && RayUV.x > 1) || (ScreenUV.y < 0 && RayUV.y > 1))
+							{
+								continue;
+							}
+						
 							float2 SampleUV = lerp(ScreenUV, RayUV, float2(LerpFactor, LerpFactor));
+							
 							if (SampleUV.x > 1 || SampleUV.x < 0 || SampleUV.y > 1 || SampleUV.y < 0)
 							{
 								continue;
@@ -304,12 +312,12 @@ bool UNiagaraDataInterfaceCamera::GetFunctionHLSL(const FName& DefinitionFunctio
 								OccludedSamples++;
 							}
 							TotalSamples++;
-						}
-						Degrees += 2 * PI / Rays;
+							Degrees += 2 * PI / Rays;
+						}						
 					}
 				}
-				Out_OcclusionFactor = TotalSamples > 0 ? OccludedSamples / TotalSamples : 1;
-				Out_SampleFactor = Steps == 0 ? 0 : (TotalSamples / (Rays * Steps + 1));
+				Out_VisibilityFraction = TotalSamples > 0 ? 1 - OccludedSamples / TotalSamples : 0;
+				Out_SampleFraction = Steps == 0 ? 0 : (TotalSamples / (Rays * Steps + 1));
 			}
 		)");
 		OutHLSL += FString::Format(FormatSample, ArgsSample);
@@ -378,12 +386,15 @@ bool UNiagaraDataInterfaceCamera::GetFunctionHLSL(const FName& DefinitionFunctio
 		OutHLSL += FString::Format(FormatSample, ArgsSample);
 		return true;
 	}
-	if (DefinitionFunctionName == GetCameraPositionsName)
+	if (DefinitionFunctionName == GetCameraPropertiesName)
 	{
 		static const TCHAR *FormatSample = TEXT(R"(
-			void {FunctionName}(out float3 Out_CameraPositionWorld)
+			void {FunctionName}(out float3 Out_CameraPositionWorld, out float3 Out_ViewForwardVector, out float3 Out_ViewUpVector, out float3 Out_ViewRightVector)
 			{				
 				Out_CameraPositionWorld.xyz = View.WorldCameraOrigin.xyz;
+				Out_ViewForwardVector.xyz = View.ViewForward.xyz;
+				Out_ViewUpVector.xyz = View.ViewUp.xyz;
+				Out_ViewRightVector.xyz = View.ViewRight.xyz;
 			}
 		)");
 		OutHLSL += FString::Format(FormatSample, ArgsSample);
@@ -393,7 +404,7 @@ bool UNiagaraDataInterfaceCamera::GetFunctionHLSL(const FName& DefinitionFunctio
 }
 
 DEFINE_NDI_DIRECT_FUNC_BINDER(UNiagaraDataInterfaceCamera, GetCameraFOV);
-DEFINE_NDI_DIRECT_FUNC_BINDER(UNiagaraDataInterfaceCamera, GetCameraPosition);
+DEFINE_NDI_DIRECT_FUNC_BINDER(UNiagaraDataInterfaceCamera, GetCameraProperties);
 DEFINE_NDI_DIRECT_FUNC_BINDER(UNiagaraDataInterfaceCamera, QueryOcclusionFactorGPU);
 DEFINE_NDI_DIRECT_FUNC_BINDER(UNiagaraDataInterfaceCamera, QueryOcclusionFactorCircleGPU);
 DEFINE_NDI_DIRECT_FUNC_BINDER(UNiagaraDataInterfaceCamera, GetViewPropertiesGPU);
@@ -413,9 +424,9 @@ void UNiagaraDataInterfaceCamera::GetVMExternalFunction(const FVMExternalFunctio
 	{
 		NDI_FUNC_BINDER(UNiagaraDataInterfaceCamera, GetCameraFOV)::Bind(this, OutFunc);
 	}
-	else if (BindingInfo.Name == GetCameraPositionsName)
+	else if (BindingInfo.Name == GetCameraPropertiesName)
 	{
-		NDI_FUNC_BINDER(UNiagaraDataInterfaceCamera, GetCameraPosition)::Bind(this, OutFunc);
+		NDI_FUNC_BINDER(UNiagaraDataInterfaceCamera, GetCameraProperties)::Bind(this, OutFunc);
 	}
 	else if (BindingInfo.Name == GetViewSpaceTransformsName)
 	{
@@ -448,7 +459,7 @@ void UNiagaraDataInterfaceCamera::GetCameraFOV(FVectorVMContext& Context)
 	}
 }
 
-void UNiagaraDataInterfaceCamera::GetCameraPosition(FVectorVMContext& Context)
+void UNiagaraDataInterfaceCamera::GetCameraProperties(FVectorVMContext& Context)
 {
 	VectorVM::FUserPtrHandler<CameraDataInterface_InstanceData> InstData(Context);
 
@@ -456,15 +467,44 @@ void UNiagaraDataInterfaceCamera::GetCameraPosition(FVectorVMContext& Context)
 	VectorVM::FExternalFuncRegisterHandler<float> CamPosY(Context);
 	VectorVM::FExternalFuncRegisterHandler<float> CamPosZ(Context);
 
-	float XPos = InstData.Get()->CameraLocation.X;
-	float YPos = InstData.Get()->CameraLocation.Y;
-	float ZPos = InstData.Get()->CameraLocation.Z;
+	VectorVM::FExternalFuncRegisterHandler<float> CamForwardX(Context);
+	VectorVM::FExternalFuncRegisterHandler<float> CamForwardY(Context);
+	VectorVM::FExternalFuncRegisterHandler<float> CamForwardZ(Context);
+
+	VectorVM::FExternalFuncRegisterHandler<float> CamUpX(Context);
+	VectorVM::FExternalFuncRegisterHandler<float> CamUpY(Context);
+	VectorVM::FExternalFuncRegisterHandler<float> CamUpZ(Context);
+
+	VectorVM::FExternalFuncRegisterHandler<float> CamRightX(Context);
+	VectorVM::FExternalFuncRegisterHandler<float> CamRightY(Context);
+	VectorVM::FExternalFuncRegisterHandler<float> CamRightZ(Context);
+
+	CameraDataInterface_InstanceData* CamData = InstData.Get();
+	float XPos = CamData->CameraLocation.X;
+	float YPos = CamData->CameraLocation.Y;
+	float ZPos = CamData->CameraLocation.Z;
+
+	FVector Forward = CamData->CameraRotation.RotateVector(FVector::ForwardVector);
+	FVector Up = CamData->CameraRotation.RotateVector(FVector::UpVector);
+	FVector Right = CamData->CameraRotation.RotateVector(FVector::RightVector);
 
 	for (int32 i = 0; i < Context.NumInstances; ++i)
 	{
 		*CamPosX.GetDestAndAdvance() = XPos;
 		*CamPosY.GetDestAndAdvance() = YPos;
 		*CamPosZ.GetDestAndAdvance() = ZPos;
+
+		*CamForwardX.GetDestAndAdvance() = Forward.X;
+		*CamForwardY.GetDestAndAdvance() = Forward.Y;
+		*CamForwardZ.GetDestAndAdvance() = Forward.Z;
+
+		*CamUpX.GetDestAndAdvance() = Up.X;
+		*CamUpY.GetDestAndAdvance() = Up.Y;
+		*CamUpZ.GetDestAndAdvance() = Up.Z;
+
+		*CamRightX.GetDestAndAdvance() = Right.X;
+		*CamRightY.GetDestAndAdvance() = Right.Y;
+		*CamRightZ.GetDestAndAdvance() = Right.Z;
 	}
 }
 
@@ -615,4 +655,9 @@ private:
 FNiagaraDataInterfaceParametersCS* UNiagaraDataInterfaceCamera::ConstructComputeParameters() const
 {
 	return new FNiagaraDataInterfaceParametersCS_CameraQuery();
-}
\ No newline at end of file
+}
+
+ETickingGroup UNiagaraDataInterfaceCamera::CalculateTickGroup(void * PerInstanceData) const
+{
+	return ETickingGroup::TG_PostUpdateWork;
+}
