New display names for Cull Reaction enum to align with Wyeth's script level terminology.

#rnx

diff --git a/Engine/Plugins/FX/Niagara/Source/Niagara/Classes/NiagaraEffectType.h b/Engine/Plugins/FX/Niagara/Source/Niagara/Classes/NiagaraEffectType.h
--- a/Engine/Plugins/FX/Niagara/Source/Niagara/Classes/NiagaraEffectType.h
+++ b/Engine/Plugins/FX/Niagara/Source/Niagara/Classes/NiagaraEffectType.h
@@ -12,14 +12,14 @@
 UENUM()
 enum class ENiagaraCullReaction
 {
-	/** The system instance will be deactivated. It will not be reactivated automatically by the scalability system. */
-	Deactivate,
+	/** The system instance will be deactivated. Particles will be allowed to die naturally. It will not be reactivated automatically by the scalability system. */
+	Deactivate UMETA(DisplayName = "Kill"),
 	/** The system instance will be deactivated and particles killed immediately. It will not be reactivated automatically by the scalability system. */
-	DeactivateImmediate,
-	/** The system instance will be deactivated. Will reactivate when it passes cull tests again. */
-	DeactivateResume,
+	DeactivateImmediate UMETA(DisplayName = "Kill and Clear"),
+	/** The system instance will be deactivated. Particles will be allowed to die naturally. Will reactivate when it passes cull tests again. */
+	DeactivateResume UMETA(DisplayName = "Asleep"),
 	/** The system instance will be deactivated and particles killed immediately. Will reactivate when it passes cull tests again. */
-	DeactivateImmediateResume,
+	DeactivateImmediateResume UMETA(DisplayName = "Asleep and Clear"),
 	/** The system instance will be paused but will resume ticking when it passes cull tests again. */
 	//PauseResume,
 };
