Added a button to isolate a selected Niagara emitter to the System Overview panel. Emitters that are not shown when the system is isolated are greyed out in the System Overview panel.

diff --git a/Engine/Plugins/FX/Niagara/Source/NiagaraEditor/Private/NiagaraEditorStyle.cpp b/Engine/Plugins/FX/Niagara/Source/NiagaraEditor/Private/NiagaraEditorStyle.cpp
--- a/Engine/Plugins/FX/Niagara/Source/NiagaraEditor/Private/NiagaraEditorStyle.cpp
+++ b/Engine/Plugins/FX/Niagara/Source/NiagaraEditor/Private/NiagaraEditorStyle.cpp
@@ -230,6 +230,9 @@ TSharedRef< FSlateStyleSet > FNiagaraEditorStyle::Create()
 	Style->Set("Niagara.CompileStatus.Warning.Small", new IMAGE_BRUSH("Icons/CompileStatus_Warning", Icon20x20));
 	Style->Set("Niagara.Asset.ReimportAsset.Needed", new IMAGE_BRUSH("Icons/icon_Reimport_Needed_40x", Icon40x40));
 	Style->Set("Niagara.Asset.ReimportAsset.Default", new IMAGE_BRUSH("Icons/icon_Reimport_40x", Icon40x40));
+	
+	Style->Set("Niagaraeditor.OverviewNode.IsolatedColor", FLinearColor::Yellow);
+	Style->Set("Niagaraeditor.OverviewNode.NotIsolatedColor", FLinearColor::Transparent);
 
 	// Icons
 	Style->Set("NiagaraEditor.Isolate", new IMAGE_PLUGIN_BRUSH("Icons/Isolate", Icon16x16));
diff --git a/Engine/Plugins/FX/Niagara/Source/NiagaraEditor/Private/NiagaraOverviewNode.cpp b/Engine/Plugins/FX/Niagara/Source/NiagaraEditor/Private/NiagaraOverviewNode.cpp
--- a/Engine/Plugins/FX/Niagara/Source/NiagaraEditor/Private/NiagaraOverviewNode.cpp
+++ b/Engine/Plugins/FX/Niagara/Source/NiagaraEditor/Private/NiagaraOverviewNode.cpp
@@ -3,6 +3,7 @@
 #include "NiagaraOverviewNode.h"
 #include "NiagaraSystem.h"
 #include "NiagaraEditorModule.h"
+#include "NiagaraEditorStyle.h"
 #include "ViewModels/Stack/NiagaraStackEntry.h"
 
 #include "Modules/ModuleManager.h"
@@ -15,6 +16,8 @@
 bool UNiagaraOverviewNode::bColorsAreInitialized = false;
 FLinearColor UNiagaraOverviewNode::EmitterColor;
 FLinearColor UNiagaraOverviewNode::SystemColor;
+FLinearColor UNiagaraOverviewNode::IsolatedColor;
+FLinearColor UNiagaraOverviewNode::NotIsolatedColor;
 
 UNiagaraOverviewNode::UNiagaraOverviewNode()
 	: OwningSystem(nullptr)
@@ -38,6 +41,21 @@ const FGuid UNiagaraOverviewNode::GetEmitterHandleGuid() const
 	return EmitterHandleGuid;
 }
 
+static const FNiagaraEmitterHandle* FindEmitterHandleByID(const UNiagaraSystem* System, const FGuid& Guid)
+{
+	check(System != nullptr);
+
+	for (const FNiagaraEmitterHandle& Handle : System->GetEmitterHandles())
+	{
+		if (Handle.GetId() == Guid)
+		{
+			return &Handle;
+		}
+	}
+
+	return nullptr;
+}
+
 FText UNiagaraOverviewNode::GetNodeTitle(ENodeTitleType::Type TitleType) const
 {
 	if (OwningSystem == nullptr)
@@ -47,14 +65,12 @@ FText UNiagaraOverviewNode::GetNodeTitle(ENodeTitleType::Type TitleType) const
 
 	if (EmitterHandleGuid.IsValid())
 	{
-		for (const FNiagaraEmitterHandle& Handle : OwningSystem->GetEmitterHandles())
+		const FNiagaraEmitterHandle* Handle = FindEmitterHandleByID(OwningSystem, EmitterHandleGuid);
+		if (ensureMsgf(Handle != nullptr, TEXT("Failed to find matching emitter handle for existing overview node!")))
 		{
-			if (Handle.GetId() == EmitterHandleGuid)
-			{
-				return FText::FromName(Handle.GetName());
-			}
+			return FText::FromName(Handle->GetName());
 		}
-		ensureMsgf(false, TEXT("Failed to find matching emitter handle for existing overview node!"));
+
 		return LOCTEXT("UnknownEmitterName", "Unknown Emitter");
 	}
 	else
@@ -70,7 +86,19 @@ FLinearColor UNiagaraOverviewNode::GetNodeTitleColor() const
 		FNiagaraEditorModule& NiagaraEditorModule = FModuleManager::LoadModuleChecked<FNiagaraEditorModule>("NiagaraEditor");
 		EmitterColor = NiagaraEditorModule.GetWidgetProvider()->GetColorForExecutionCategory(UNiagaraStackEntry::FExecutionCategoryNames::Emitter);
 		SystemColor = NiagaraEditorModule.GetWidgetProvider()->GetColorForExecutionCategory(UNiagaraStackEntry::FExecutionCategoryNames::System);
+		IsolatedColor = FNiagaraEditorStyle::Get().GetColor("NiagaraEditor.OverviewNode.IsolatedColor");
+		NotIsolatedColor = FNiagaraEditorStyle::Get().GetColor("NiagaraEditor.OverviewNode.NotIsolatedColor");
+	}
+
+	if (EmitterHandleGuid.IsValid() && OwningSystem->GetIsolateEnabled())
+	{
+		const FNiagaraEmitterHandle* Handle = FindEmitterHandleByID(OwningSystem, EmitterHandleGuid);
+		if (ensureMsgf(Handle != nullptr, TEXT("Failed to find matching emitter handle for existing overview node!")))
+		{
+			return Handle->IsIsolated() ? IsolatedColor : NotIsolatedColor;
+		}
 	}
+
 	return EmitterHandleGuid.IsValid() ? EmitterColor : SystemColor;
 }
 
diff --git a/Engine/Plugins/FX/Niagara/Source/NiagaraEditor/Private/ViewModels/NiagaraEmitterHandleViewModel.cpp b/Engine/Plugins/FX/Niagara/Source/NiagaraEditor/Private/ViewModels/NiagaraEmitterHandleViewModel.cpp
--- a/Engine/Plugins/FX/Niagara/Source/NiagaraEditor/Private/ViewModels/NiagaraEmitterHandleViewModel.cpp
+++ b/Engine/Plugins/FX/Niagara/Source/NiagaraEditor/Private/ViewModels/NiagaraEmitterHandleViewModel.cpp
@@ -220,6 +220,25 @@ bool FNiagaraEmitterHandleViewModel::GetIsIsolated() const
 	return EmitterHandle != nullptr && EmitterHandle->IsIsolated();
 }
 
+void FNiagaraEmitterHandleViewModel::SetIsIsolated(bool bInIsIsolated)
+{
+	bool bWasIsolated = GetIsIsolated();
+
+	if (bWasIsolated == false && bInIsIsolated == true)
+	{
+		GetOwningSystemViewModel()->IsolateEmitters(TArray<FGuid> { EmitterHandle->GetId() });
+	}
+	else if (bWasIsolated == true && bInIsIsolated == false)
+	{
+		GetOwningSystemViewModel()->IsolateEmitters(TArray<FGuid> {});
+	}
+}
+
+ENiagaraSystemViewModelEditMode FNiagaraEmitterHandleViewModel::GetOwningSystemEditMode() const
+{
+	return GetOwningSystemViewModel()->GetEditMode();
+}
+
 ECheckBoxState FNiagaraEmitterHandleViewModel::GetIsEnabledCheckState() const
 {
 	if (EmitterHandle)
diff --git a/Engine/Plugins/FX/Niagara/Source/NiagaraEditor/Public/NiagaraOverviewNode.h b/Engine/Plugins/FX/Niagara/Source/NiagaraEditor/Public/NiagaraOverviewNode.h
--- a/Engine/Plugins/FX/Niagara/Source/NiagaraEditor/Public/NiagaraOverviewNode.h
+++ b/Engine/Plugins/FX/Niagara/Source/NiagaraEditor/Public/NiagaraOverviewNode.h
@@ -5,6 +5,7 @@
 #include "NiagaraOverviewNode.generated.h"
 
 class UNiagaraSystem;
+class FNiagaraEmitterHandleViewModel;
 
 UCLASS()
 class NIAGARAEDITOR_API UNiagaraOverviewNode : public UEdGraphNode
@@ -43,4 +44,6 @@ private:
 	static bool bColorsAreInitialized;
 	static FLinearColor SystemColor;
 	static FLinearColor EmitterColor;
+	static FLinearColor IsolatedColor;
+	static FLinearColor NotIsolatedColor;
 };
diff --git a/Engine/Plugins/FX/Niagara/Source/NiagaraEditor/Public/ViewModels/NiagaraEmitterHandleViewModel.h b/Engine/Plugins/FX/Niagara/Source/NiagaraEditor/Public/ViewModels/NiagaraEmitterHandleViewModel.h
--- a/Engine/Plugins/FX/Niagara/Source/NiagaraEditor/Public/ViewModels/NiagaraEmitterHandleViewModel.h
+++ b/Engine/Plugins/FX/Niagara/Source/NiagaraEditor/Public/ViewModels/NiagaraEmitterHandleViewModel.h
@@ -14,6 +14,7 @@ struct FNiagaraEmitterHandle;
 class FNiagaraEmitterInstance;
 class FNiagaraEmitterViewModel;
 class UNiagaraStackViewModel;
+enum class ENiagaraSystemViewModelEditMode;
 
 /** The view model for the FNiagaraEmitterEditorWidget. */
 class FNiagaraEmitterHandleViewModel : public TSharedFromThis<FNiagaraEmitterHandleViewModel>, public FGCObject
@@ -81,6 +82,9 @@ public:
 	/** Gets whether or not the emitter for this handle has been isolated in the UI. */
 	NIAGARAEDITOR_API bool GetIsIsolated() const;
 
+	/** Sets whether or not this emitter is isolated. May affect other emitters in the system. */
+	NIAGARAEDITOR_API void SetIsIsolated(bool InIsIsolated);
+
 	/** Gets the check state for the is enabled check box. */
 	NIAGARAEDITOR_API ECheckBoxState GetIsEnabledCheckState() const;
 
@@ -96,6 +100,9 @@ public:
 	/** Gets the stack view model which represents the emitter pointed to by this handle. */
 	NIAGARAEDITOR_API UNiagaraStackViewModel* GetEmitterStackViewModel();
 
+	/** Gets the current edit mode of the emitter's owning system. */
+	NIAGARAEDITOR_API ENiagaraSystemViewModelEditMode GetOwningSystemEditMode() const;
+
 	/** Gets whether or not this emitter handle has a rename pending. */
 	NIAGARAEDITOR_API bool GetIsRenamePending() const;
 
diff --git a/Engine/Plugins/FX/Niagara/Source/NiagaraEditor/Public/ViewModels/NiagaraSystemViewModel.h b/Engine/Plugins/FX/Niagara/Source/NiagaraEditor/Public/ViewModels/NiagaraSystemViewModel.h
--- a/Engine/Plugins/FX/Niagara/Source/NiagaraEditor/Public/ViewModels/NiagaraSystemViewModel.h
+++ b/Engine/Plugins/FX/Niagara/Source/NiagaraEditor/Public/ViewModels/NiagaraSystemViewModel.h
@@ -150,10 +150,11 @@ public:
 	/** Gets an array of the view models for the emitter handles owned by this System. */
 	NIAGARAEDITOR_API const TArray<TSharedRef<FNiagaraEmitterHandleViewModel>>& GetEmitterHandleViewModels();
 
-	/** Gets an emitter handle view model by id.  Returns an invalid shared ptr if it can't be found. */
+	/** Gets an emitter handle view model by ID. Returns an invalid shared ptr if it can't be found. */
 	NIAGARAEDITOR_API TSharedPtr<FNiagaraEmitterHandleViewModel> GetEmitterHandleViewModelById(FGuid InEmitterHandleId);
 
-	TSharedPtr<FNiagaraEmitterHandleViewModel> GetEmitterHandleViewModelForEmitter(UNiagaraEmitter* InEmitter) const;
+	/** Gets an emitter handle view model for the given emitter. Returns an invalid shared ptr if it can't be found. */
+	NIAGARAEDITOR_API TSharedPtr<FNiagaraEmitterHandleViewModel> GetEmitterHandleViewModelForEmitter(UNiagaraEmitter* InEmitter) const;
 
 	/** Gets the view model for the System script. */
 	TSharedPtr<FNiagaraSystemScriptViewModel> GetSystemScriptViewModel();
diff --git a/Engine/Plugins/FX/Niagara/Source/NiagaraEditorWidgets/Private/SNiagaraOverviewStack.cpp b/Engine/Plugins/FX/Niagara/Source/NiagaraEditorWidgets/Private/SNiagaraOverviewStack.cpp
--- a/Engine/Plugins/FX/Niagara/Source/NiagaraEditorWidgets/Private/SNiagaraOverviewStack.cpp
+++ b/Engine/Plugins/FX/Niagara/Source/NiagaraEditorWidgets/Private/SNiagaraOverviewStack.cpp
@@ -1,7 +1,9 @@
 // Copyright 1998-2019 Epic Games, Inc. All Rights Reserved.
 
 #include "SNiagaraOverviewStack.h"
+#include "NiagaraSystem.h"
 #include "NiagaraSystemEditorData.h"
+#include "ViewModels/NiagaraEmitterViewModel.h"
 #include "ViewModels/NiagaraSystemViewModel.h"
 #include "ViewModels/NiagaraEmitterHandleViewModel.h"
 #include "ViewModels/NiagaraSystemSelectionViewModel.h"
@@ -64,7 +66,16 @@ class SNiagaraSystemOverviewEntryListRow : public STableRow<UNiagaraStackEntry*>
 			BackgroundColor = FNiagaraEditorWidgetsStyle::Get().GetColor("NiagaraEditor.SystemOverview.Group.BackgroundColor");
 			ContentPadding = FMargin(0, 5, 0, 5);
 		}
+
 		DisabledBackgroundColor = BackgroundColor + FLinearColor(.02f, .02f, .02f, 0.0f);
+		IsolatedBackgroundColor = BackgroundColor - FLinearColor(.03f, .03f, .03f, 0.0f);
+
+		TSharedRef<FNiagaraSystemViewModel> SystemViewModel = StackEntry->GetSystemViewModel();
+		TSharedPtr<FNiagaraEmitterViewModel> EmitterViewModel = StackEntry->GetEmitterViewModel();
+		if (EmitterViewModel.IsValid())
+		{
+			EmitterHandleViewModel = SystemViewModel->GetEmitterHandleViewModelForEmitter(EmitterViewModel->GetEmitter());
+		}
 
 		STableRow<UNiagaraStackEntry*>::Construct(STableRow<UNiagaraStackEntry*>::FArguments()
 			.Style(FNiagaraEditorWidgetsStyle::Get(), "NiagaraEditor.SystemOverview.TableViewRow")
@@ -78,7 +89,7 @@ class SNiagaraSystemOverviewEntryListRow : public STableRow<UNiagaraStackEntry*>
 			[
 				SNew(SBorder)
 				.BorderImage(FEditorStyle::GetBrush("WhiteBrush"))
-				.BorderBackgroundColor(this, &SNiagaraSystemOverviewEntryListRow::GetBackgroundColor, StackEntry)
+				.BorderBackgroundColor(this, &SNiagaraSystemOverviewEntryListRow::GetBackgroundColor)
 				.ToolTipText_UObject(StackEntry, &UNiagaraStackEntry::GetTooltipText)
 				.Padding(FMargin(0))
 				[
@@ -139,9 +150,19 @@ class SNiagaraSystemOverviewEntryListRow : public STableRow<UNiagaraStackEntry*>
 	}
 
 private:
-	FSlateColor GetBackgroundColor(UNiagaraStackEntry* Entry) const
+	FSlateColor GetBackgroundColor() const
 	{
-		return Entry->GetIsEnabled() && Entry->GetOwnerIsEnabled() ? BackgroundColor : DisabledBackgroundColor;
+		TSharedRef<FNiagaraSystemViewModel> SystemViewModel = StackEntry->GetSystemViewModel();
+		if (SystemViewModel->GetSystem().GetIsolateEnabled())
+		{
+			TSharedPtr<FNiagaraEmitterHandleViewModel> EmitterHandleViewModelPtr = EmitterHandleViewModel.Pin();
+			if (EmitterHandleViewModelPtr.IsValid() && !EmitterHandleViewModelPtr->GetIsIsolated())
+			{
+				return IsolatedBackgroundColor;
+			}
+		}
+
+		return StackEntry->GetIsEnabledAndOwnerIsEnabled() ? BackgroundColor : DisabledBackgroundColor;
 	}
 
 	FMargin GetInnerContentPadding() const
@@ -161,7 +182,9 @@ private:
 	UNiagaraStackEntry* StackEntry;
 	FLinearColor BackgroundColor;
 	FLinearColor DisabledBackgroundColor;
+	FLinearColor IsolatedBackgroundColor;
 	TAttribute<EVisibility> IssueIconVisibility;
+	TWeakPtr<FNiagaraEmitterHandleViewModel> EmitterHandleViewModel;
 };
 
 class SNiagaraSystemOverviewEnabledCheckBox : public SCompoundWidget
diff --git a/Engine/Plugins/FX/Niagara/Source/NiagaraEditorWidgets/Private/SNiagaraOverviewStackNode.cpp b/Engine/Plugins/FX/Niagara/Source/NiagaraEditorWidgets/Private/SNiagaraOverviewStackNode.cpp
--- a/Engine/Plugins/FX/Niagara/Source/NiagaraEditorWidgets/Private/SNiagaraOverviewStackNode.cpp
+++ b/Engine/Plugins/FX/Niagara/Source/NiagaraEditorWidgets/Private/SNiagaraOverviewStackNode.cpp
@@ -7,11 +7,12 @@
 #include "ViewModels/Stack/NiagaraStackViewModel.h"
 #include "SNiagaraOverviewStack.h"
 #include "NiagaraEditorModule.h"
+#include "NiagaraEditorStyle.h"
 #include "NiagaraEditorWidgetsStyle.h"
 #include "Stack/SNiagaraStackIssueIcon.h"
-
 #include "Modules/ModuleManager.h"
 #include "Widgets/Layout/SBox.h"
+#include "Widgets/Input/SButton.h"
 #include "Widgets/Input/SCheckBox.h"
 
 #define LOCTEXT_NAMESPACE "NiagaraOverviewStackNode"
@@ -45,6 +46,7 @@ void SNiagaraOverviewStackNode::Construct(const FArguments& InArgs, UNiagaraOver
 			OverviewSelectionViewModel = OwningSystemViewModel->GetSelectionViewModel();
 		}
 	}
+
 	UpdateGraphNode();
 }
 
@@ -78,9 +80,29 @@ TSharedRef<SWidget> SNiagaraOverviewStackNode::CreateTitleWidget(TSharedPtr<SNod
 			.IsChecked(this, &SNiagaraOverviewStackNode::GetEnabledCheckState)
 			.OnCheckStateChanged(this, &SNiagaraOverviewStackNode::OnEnabledCheckStateChanged)
 		]
+		// Isolate toggle
+		+ SHorizontalBox::Slot()
+		.AutoWidth()
+		.VAlign(VAlign_Center)
+		.Padding(3, 0, 0, 0)
+		[
+			SNew(SButton)
+			.ButtonStyle(FEditorStyle::Get(), "HoverHintOnly")
+			.HAlign(HAlign_Center)
+			.ContentPadding(1)
+			.ToolTipText(this, &SNiagaraOverviewStackNode::GetToggleIsolateToolTip)
+			.OnClicked(this, &SNiagaraOverviewStackNode::OnToggleIsolateButtonClicked)
+			.Visibility(this, &SNiagaraOverviewStackNode::GetToggleIsolateVisibility)
+			.Content()
+			[
+				SNew(SImage)
+				.Image(FNiagaraEditorStyle::Get().GetBrush("NiagaraEditor.Isolate"))
+				.ColorAndOpacity(this, &SNiagaraOverviewStackNode::GetToggleIsolateImageColor)
+			]
+		]
 		// Name
 		+ SHorizontalBox::Slot()
-		.Padding(4, 0, 0, 0)
+		.Padding(3, 0, 0, 0)
 		.FillWidth(1.0f)
 		[
 			DefaultTitle
@@ -174,4 +196,41 @@ void SNiagaraOverviewStackNode::OnEnabledCheckStateChanged(ECheckBoxState InChec
 	}
 }
 
+FReply SNiagaraOverviewStackNode::OnToggleIsolateButtonClicked()
+{
+	TSharedPtr<FNiagaraEmitterHandleViewModel> EmitterHandleViewModel = EmitterHandleViewModelWeak.Pin();
+	if (EmitterHandleViewModel.IsValid())
+	{
+		bool bShouldBeIsolated = !EmitterHandleViewModel->GetIsIsolated();
+		EmitterHandleViewModel->SetIsIsolated(bShouldBeIsolated);
+	}
+	return FReply::Handled();
+}
+
+FText SNiagaraOverviewStackNode::GetToggleIsolateToolTip() const
+{
+	TSharedPtr<FNiagaraEmitterHandleViewModel> EmitterHandleViewModel = EmitterHandleViewModelWeak.Pin();
+	return EmitterHandleViewModel.IsValid() && EmitterHandleViewModel->GetIsIsolated()
+		? LOCTEXT("TurnOffEmitterIsolation", "Disable emitter isolation.")
+		: LOCTEXT("IsolateThisEmitter", "Enable isolation for this emitter.");
+}
+
+EVisibility SNiagaraOverviewStackNode::GetToggleIsolateVisibility() const
+{
+	TSharedPtr<FNiagaraEmitterHandleViewModel> EmitterHandleViewModel = EmitterHandleViewModelWeak.Pin();
+	return EmitterHandleViewModel.IsValid() &&
+		EmitterHandleViewModel->GetOwningSystemEditMode() == ENiagaraSystemViewModelEditMode::SystemAsset 
+		? EVisibility::Visible 
+		: EVisibility::Collapsed;
+}
+
+FSlateColor SNiagaraOverviewStackNode::GetToggleIsolateImageColor() const
+{
+	TSharedPtr<FNiagaraEmitterHandleViewModel> EmitterHandleViewModel = EmitterHandleViewModelWeak.Pin();
+	return EmitterHandleViewModel.IsValid() && 
+		EmitterHandleViewModel->GetIsIsolated()
+		? FEditorStyle::GetSlateColor("SelectionColor")
+		: FLinearColor::Gray;
+}
+
 #undef LOCTEXT_NAMESPACE
\ No newline at end of file
diff --git a/Engine/Plugins/FX/Niagara/Source/NiagaraEditorWidgets/Private/SNiagaraOverviewStackNode.h b/Engine/Plugins/FX/Niagara/Source/NiagaraEditorWidgets/Private/SNiagaraOverviewStackNode.h
--- a/Engine/Plugins/FX/Niagara/Source/NiagaraEditorWidgets/Private/SNiagaraOverviewStackNode.h
+++ b/Engine/Plugins/FX/Niagara/Source/NiagaraEditorWidgets/Private/SNiagaraOverviewStackNode.h
@@ -28,6 +28,11 @@ private:
 	ECheckBoxState GetEnabledCheckState() const;
 	void OnEnabledCheckStateChanged(ECheckBoxState InCheckState);
 
+	FText GetToggleIsolateToolTip() const;
+	FReply OnToggleIsolateButtonClicked();
+	EVisibility GetToggleIsolateVisibility() const;
+	FSlateColor GetToggleIsolateImageColor() const;
+
 private:
 	UNiagaraOverviewNode* OverviewStackNode;
 	UNiagaraStackViewModel* StackViewModel;
