The "Fix issue" and "Dismiss issue" buttons in the Niagara stack overview now wrap to the following line when the error message is longer than the width of the panel.

diff --git a/Engine/Plugins/FX/Niagara/Source/NiagaraEditorWidgets/Private/Stack/SNiagaraStackErrorItem.cpp b/Engine/Plugins/FX/Niagara/Source/NiagaraEditorWidgets/Private/Stack/SNiagaraStackErrorItem.cpp
--- a/Engine/Plugins/FX/Niagara/Source/NiagaraEditorWidgets/Private/Stack/SNiagaraStackErrorItem.cpp
+++ b/Engine/Plugins/FX/Niagara/Source/NiagaraEditorWidgets/Private/Stack/SNiagaraStackErrorItem.cpp
@@ -17,7 +17,7 @@ void SNiagaraStackErrorItem::Construct(const FArguments& InArgs, UNiagaraStackEr
 {
 	ErrorItem = InErrorItem;
 	StackViewModel = InStackViewModel;
-	TSharedPtr<SHorizontalBox> ErrorInternalBox = SNew(SHorizontalBox);
+
 	const FSlateBrush* IconBrush;
 	switch (ErrorItem->GetStackIssue().GetSeverity())
 	{
@@ -34,34 +34,28 @@ void SNiagaraStackErrorItem::Construct(const FArguments& InArgs, UNiagaraStackEr
 		IconBrush = FEditorStyle::GetBrush("NoBrush");
 		break;
 	}
-	ErrorInternalBox->AddSlot()
-		.HAlign(HAlign_Left)
-		.VAlign(VAlign_Center)
-		.AutoWidth()
-		[
-			SNew(STextBlock)
-			.TextStyle(FNiagaraEditorStyle::Get(), "NiagaraEditor.ParameterText")
-			.Text_UObject(ErrorItem, &UNiagaraStackErrorItem::GetDisplayName)
-			.ColorAndOpacity(this, &SNiagaraStackErrorItem::GetTextColorForSearch)
-			.HighlightText_UObject(StackViewModel, &UNiagaraStackViewModel::GetCurrentSearchText)
-		];
+
 	ChildSlot
 	[
 		SNew(SHorizontalBox)
 		+ SHorizontalBox::Slot()
-		.AutoWidth()
-		.HAlign(HAlign_Left)
 		.VAlign(VAlign_Center)
+		.HAlign(HAlign_Left)
+		.AutoWidth()
 		[
 			SNew(SImage)
 			.Image(IconBrush)
 		]
 		+ SHorizontalBox::Slot()
-		.AutoWidth()
-		.HAlign(HAlign_Left)
 		.VAlign(VAlign_Center)
+		.HAlign(HAlign_Left)
+		.AutoWidth()
 		[
-			ErrorInternalBox.ToSharedRef()
+			SNew(STextBlock)
+			.TextStyle(FNiagaraEditorStyle::Get(), "NiagaraEditor.ParameterText")
+			.Text_UObject(ErrorItem, &UNiagaraStackErrorItem::GetDisplayName)
+			.ColorAndOpacity(this, &SNiagaraStackErrorItem::GetTextColorForSearch)
+			.HighlightText_UObject(StackViewModel, &UNiagaraStackViewModel::GetCurrentSearchText)
 		]
 	];
 }
@@ -70,43 +64,31 @@ void SNiagaraStackErrorItemFix::Construct(const FArguments& InArgs, UNiagaraStac
 {
 	ErrorItem = InErrorItem;
 	StackViewModel = InStackViewModel;
-	TSharedPtr<SHorizontalBox> ErrorInternalBox = SNew(SHorizontalBox);
-	ErrorInternalBox->AddSlot()
-		.HAlign(HAlign_Left)
+
+	ChildSlot
+	[
+		SNew(SHorizontalBox)
+		+ SHorizontalBox::Slot()
+		.Padding(0, 4, 0, 0)
 		.VAlign(VAlign_Center)
-		.AutoWidth()
 		[
 			SNew(STextBlock)
-				.TextStyle(FNiagaraEditorStyle::Get(), "NiagaraEditor.ParameterText")
-				.Text_UObject(ErrorItem, &UNiagaraStackErrorItemFix::GetDisplayName)
-				.ColorAndOpacity(this, &SNiagaraStackErrorItemFix::GetTextColorForSearch)
-				.HighlightText_UObject(StackViewModel, &UNiagaraStackViewModel::GetCurrentSearchText)
-		];
-	ErrorInternalBox->AddSlot()
+			.TextStyle(FNiagaraEditorStyle::Get(), "NiagaraEditor.ParameterText")
+			.Text_UObject(ErrorItem, &UNiagaraStackErrorItemFix::GetDisplayName)
+			.ColorAndOpacity(this, &SNiagaraStackErrorItemFix::GetTextColorForSearch)
+			.HighlightText_UObject(StackViewModel, &UNiagaraStackViewModel::GetCurrentSearchText)
+			.AutoWrapText(true)
+		]
+		+ SHorizontalBox::Slot()
+		.Padding(5, 0, 0, 0)
 		.HAlign(HAlign_Right)
 		.VAlign(VAlign_Center)
-		.Padding(10, 0)
 		[
-			SNew(SHorizontalBox)
-				+ SHorizontalBox::Slot()
-				.VAlign(VAlign_Center)
-				[
-					SNew(SButton)
-					.TextStyle(FNiagaraEditorStyle::Get(), "NiagaraEditor.ParameterText")
-					.Text_UObject(ErrorItem, &UNiagaraStackErrorItemFix::GetFixButtonText)
-					.OnClicked_UObject(ErrorItem, &UNiagaraStackErrorItemFix::OnTryFixError)
-				]
-		];
-	ChildSlot
-	[
-		SNew(SHorizontalBox)
-			+ SHorizontalBox::Slot()
-			.AutoWidth()
-			.HAlign(HAlign_Left)
-			.VAlign(VAlign_Center)
-			[
-				ErrorInternalBox.ToSharedRef()
-			]
+			SNew(SButton)
+			.TextStyle(FNiagaraEditorStyle::Get(), "NiagaraEditor.ParameterText")
+			.Text_UObject(ErrorItem, &UNiagaraStackErrorItemFix::GetFixButtonText)
+			.OnClicked_UObject(ErrorItem, &UNiagaraStackErrorItemFix::OnTryFixError)
+		]
 	];
 }
 
