FNM - Fix WaveOp Niagaracode where dead particles could accidently be assigned an index in the active-particle list because the index assignment was not taking the "dead lanes" into account.

diff --git a/Engine/Plugins/FX/Niagara/Shaders/Private/NiagaraEmitterInstanceShader.usf b/Engine/Plugins/FX/Niagara/Shaders/Private/NiagaraEmitterInstanceShader.usf
--- a/Engine/Plugins/FX/Niagara/Shaders/Private/NiagaraEmitterInstanceShader.usf
+++ b/Engine/Plugins/FX/Niagara/Shaders/Private/NiagaraEmitterInstanceShader.usf
@@ -400,6 +400,7 @@ int rand(int x, int Seed1, int Seed2, int Seed3)
 
 #if PS4_PROFILE || XBOXONE_PROFILE
 				uint NumCounters = WaveActiveCountBits(bDoAcquire);
+				uint PrefixCounts = WavePrefixCountBits(bDoAcquire);
 				if(NumCounters)
 				{
 					if(WaveIsFirstLane())
@@ -408,8 +409,12 @@ int rand(int x, int Seed1, int Seed2, int Seed3)
 						InterlockedAdd(RWInstanceCounts[WriteInstanceCountOffset], NumCounters, RetPrevIdx);
 						PrevIdx = (int)RetPrevIdx;
 					}
-					PrevIdx = WaveReadFirstLane(PrevIdx);
-					PrevIdx += WavePrefixCountBits(bDoAcquire);
+
+					if(bDoAcquire)
+					{
+						PrevIdx = WaveReadFirstLane(PrevIdx);
+						PrevIdx += PrefixCounts;
+					}
 				}		
 #else
 				if(bDoAcquire == true)
@@ -432,6 +437,7 @@ int rand(int x, int Seed1, int Seed2, int Seed3)
 
 #if PS4_PROFILE || XBOXONE_PROFILE
 			uint NumCounters = WaveActiveCountBits(bDoAcquire);
+			uint PrefixCounts = WavePrefixCountBits(bDoAcquire);
 			if(NumCounters)
 			{
 				if(WaveIsFirstLane())
@@ -440,8 +446,12 @@ int rand(int x, int Seed1, int Seed2, int Seed3)
 					InterlockedAdd(RWInstanceCounts[WriteInstanceCountOffset], NumCounters, RetPrevIdx);
 					PrevIdx = (int)RetPrevIdx;
 				}
-				PrevIdx = WaveReadFirstLane(PrevIdx);
-				PrevIdx += WavePrefixCountBits(bDoAcquire);
+
+				if(bDoAcquire)
+				{
+					PrevIdx = WaveReadFirstLane(PrevIdx);
+					PrevIdx += PrefixCounts;
+				}
 			}	
 #else
 			if(bDoAcquire == true)
diff --git a/Engine/Plugins/FX/Niagara/Source/NiagaraCore/Private/NiagaraCustomVersion.cpp b/Engine/Plugins/FX/Niagara/Source/NiagaraCore/Private/NiagaraCustomVersion.cpp
--- a/Engine/Plugins/FX/Niagara/Source/NiagaraCore/Private/NiagaraCustomVersion.cpp
+++ b/Engine/Plugins/FX/Niagara/Source/NiagaraCore/Private/NiagaraCustomVersion.cpp
@@ -8,5 +8,5 @@ const FGuid FNiagaraCustomVersion::GUID(0xFCF57AFA, 0x50764283, 0xB9A9E658, 0xFF
 // Register the custom version with core
 FCustomVersionRegistration GRegisterNiagaraCustomVersion(FNiagaraCustomVersion::GUID, FNiagaraCustomVersion::LatestVersion, TEXT("NiagaraVer"));
 
-const FGuid FNiagaraCustomVersion::LatestScriptCompileVersion(0x02ABFEE7, 0x88B02847, 0x85046386, 0x8BD6F349);
+const FGuid FNiagaraCustomVersion::LatestScriptCompileVersion(0x02EFFEE7, 0x8852B128, 0x46385086, 0x84BD6F39);
 
