Integrating missing changes from 11.30 into 11.40.

#rnx

diff --git a/Engine/Plugins/FX/Niagara/Source/Niagara/Private/NiagaraComponent.cpp b/Engine/Plugins/FX/Niagara/Source/Niagara/Private/NiagaraComponent.cpp
--- a/Engine/Plugins/FX/Niagara/Source/Niagara/Private/NiagaraComponent.cpp
+++ b/Engine/Plugins/FX/Niagara/Source/Niagara/Private/NiagaraComponent.cpp
@@ -868,18 +868,31 @@ void UNiagaraComponent::Deactivate()
 
 void UNiagaraComponent::DeactivateInternal(bool bIsScalabilityCull /* = false */)
 {
-	SCOPE_CYCLE_COUNTER(STAT_NiagaraComponentDeactivate);
-	CSV_SCOPED_TIMING_STAT_EXCLUSIVE(Niagara);
+	if (IsActive())
+	{
+		SCOPE_CYCLE_COUNTER(STAT_NiagaraComponentDeactivate);
+		CSV_SCOPED_TIMING_STAT_EXCLUSIVE(Niagara);
 
-	//UE_LOG(LogNiagara, Log, TEXT("Deactivate: %u - %s"), this, *Asset->GetName());
+		//UE_LOG(LogNiagara, Log, TEXT("Deactivate: %u - %s"), this, *Asset->GetName());
 
-	if (SystemInstance)
-	{
-		// Don't deactivate in solo mode as we are not ticked by the world but rather the component
-		// Deactivating will cause the system to never Complete
-		if (SystemInstance->IsSolo() == false)
+		if (SystemInstance)
+		{
+			// Don't deactivate in solo mode as we are not ticked by the world but rather the component
+			// Deactivating will cause the system to never Complete
+			if (SystemInstance->IsSolo() == false)
+			{
+				Super::Deactivate();
+			}
+
+			SystemInstance->Deactivate();
+
+			// We are considered active until we are complete
+			SetActiveFlag(!SystemInstance->IsComplete());
+		}
+		else
 		{
 			Super::Deactivate();
+			SetActiveFlag(false);
 		}
 
 		//Unregister with the scalability manager if this is a genuine deactivation from outside.
@@ -890,9 +903,6 @@ void UNiagaraComponent::DeactivateInternal(bool bIsScalabilityCull /* = false */
 		}
 
 		SystemInstance->Deactivate();
-
-		// We are considered active until we are complete
-		SetActiveFlag(!SystemInstance->IsComplete());
 	}
 	else
 	{
