Niagara - Fix an intermittent crash when closing the system or emitter editor.

diff --git a/Engine/Plugins/FX/Niagara/Source/NiagaraEditor/Private/Toolkits/NiagaraSystemToolkit.cpp b/Engine/Plugins/FX/Niagara/Source/NiagaraEditor/Private/Toolkits/NiagaraSystemToolkit.cpp
--- a/Engine/Plugins/FX/Niagara/Source/NiagaraEditor/Private/Toolkits/NiagaraSystemToolkit.cpp
+++ b/Engine/Plugins/FX/Niagara/Source/NiagaraEditor/Private/Toolkits/NiagaraSystemToolkit.cpp
@@ -1166,7 +1166,7 @@ void FNiagaraSystemToolkit::SaveAsset_Execute()
 		UE_LOG(LogNiagaraEditor, Log, TEXT("Saving and Compiling NiagaraEmitter %s"), *GetEditingObjects()[0]->GetName());
 		UpdateOriginalEmitter();
 	}
-	SystemViewModel->OnPreSave();
+	SystemViewModel->NotifyPreSave();
 	FAssetEditorToolkit::SaveAsset_Execute();
 }
 
@@ -1177,7 +1177,7 @@ void FNiagaraSystemToolkit::SaveAssetAs_Execute()
 		UE_LOG(LogNiagaraEditor, Log, TEXT("Saving and Compiling NiagaraEmitter %s"), *GetEditingObjects()[0]->GetName());
 		UpdateOriginalEmitter();
 	}
-	SystemViewModel->OnPreSave();
+	SystemViewModel->NotifyPreSave();
 	FAssetEditorToolkit::SaveAssetAs_Execute();
 }
 
@@ -1204,7 +1204,7 @@ bool FNiagaraSystemToolkit::OnRequestClose()
 		FNiagaraEditorUtilities::WriteTextFileToDisk(FPaths::ProjectLogDir(), FilenamePart + TEXT(".onClose.txt"), ExportText, true);
 	}
 
-	SystemViewModel->OnPreClose();
+	SystemViewModel->NotifyPreClose();
 
 	if (SystemToolkitMode == ESystemToolkitMode::Emitter)
 	{
diff --git a/Engine/Plugins/FX/Niagara/Source/NiagaraEditor/Private/ViewModels/NiagaraSystemViewModel.cpp b/Engine/Plugins/FX/Niagara/Source/NiagaraEditor/Private/ViewModels/NiagaraSystemViewModel.cpp
--- a/Engine/Plugins/FX/Niagara/Source/NiagaraEditor/Private/ViewModels/NiagaraSystemViewModel.cpp
+++ b/Engine/Plugins/FX/Niagara/Source/NiagaraEditor/Private/ViewModels/NiagaraSystemViewModel.cpp
@@ -552,7 +552,7 @@ void FNiagaraSystemViewModel::Tick(float DeltaTime)
 	}
 }
 
-void FNiagaraSystemViewModel::OnPreSave()
+void FNiagaraSystemViewModel::NotifyPreSave()
 {
 	if (GetSystem().HasOutstandingCompilationRequests())
 	{
@@ -561,13 +561,19 @@ void FNiagaraSystemViewModel::OnPreSave()
 	}
 }
 
-void FNiagaraSystemViewModel::OnPreClose()
+void FNiagaraSystemViewModel::NotifyPreClose()
 {
 	if (GetSystem().HasOutstandingCompilationRequests())
 	{
 		UE_LOG(LogNiagaraEditor, Log, TEXT("System %s has pending compile jobs. Waiting for that code to complete before Closing.."), *GetSystem().GetName());
 		GetSystem().WaitForCompilationComplete();
 	}
+	OnPreCloseDelegate.Broadcast();
+}
+
+FNiagaraSystemViewModel::FOnPreClose& FNiagaraSystemViewModel::OnPreClose()
+{
+	return OnPreCloseDelegate;
 }
 
 TSharedPtr<FUICommandList> FNiagaraSystemViewModel::GetToolkitCommands()
diff --git a/Engine/Plugins/FX/Niagara/Source/NiagaraEditor/Public/ViewModels/NiagaraSystemViewModel.h b/Engine/Plugins/FX/Niagara/Source/NiagaraEditor/Public/ViewModels/NiagaraSystemViewModel.h
--- a/Engine/Plugins/FX/Niagara/Source/NiagaraEditor/Public/ViewModels/NiagaraSystemViewModel.h
+++ b/Engine/Plugins/FX/Niagara/Source/NiagaraEditor/Public/ViewModels/NiagaraSystemViewModel.h
@@ -100,6 +100,8 @@ public:
 
 	DECLARE_MULTICAST_DELEGATE(FOnPinnedCurvesChanged);
 
+	DECLARE_MULTICAST_DELEGATE(FOnPreClose);
+
 public:
 	struct FEmitterHandleToDuplicate
 	{
@@ -268,8 +270,11 @@ public:
 	/** Checks whether or not an emitter is pinned in the stack UI*/
 	NIAGARAEDITOR_API bool GetIsEmitterPinned(TSharedRef<FNiagaraEmitterHandleViewModel> EmitterHandleModel);
 
-	NIAGARAEDITOR_API void OnPreSave();
-	NIAGARAEDITOR_API void OnPreClose();
+	NIAGARAEDITOR_API void NotifyPreSave();
+
+	NIAGARAEDITOR_API void NotifyPreClose();
+
+	NIAGARAEDITOR_API FOnPreClose& OnPreClose();
 	
 	/** Gets the system toolkit command list. */
 	NIAGARAEDITOR_API TSharedPtr<FUICommandList> GetToolkitCommands();
@@ -480,6 +485,9 @@ private:
 	/** A multicast delegate which is called whenever the system has been compiled. */
 	FOnSystemCompiled OnSystemCompiledDelegate;
 
+	/** A multicast delegate which is called whenever this has been notified it's owner will be closing. */
+	FOnPreClose OnPreCloseDelegate;
+
 	/** A flag for preventing reentrancy when syncrhonizing sequencer data. */
 	bool bUpdatingEmittersFromSequencerDataChange;
 
diff --git a/Engine/Plugins/FX/Niagara/Source/NiagaraEditorWidgets/Private/SNiagaraOverviewGraph.cpp b/Engine/Plugins/FX/Niagara/Source/NiagaraEditorWidgets/Private/SNiagaraOverviewGraph.cpp
--- a/Engine/Plugins/FX/Niagara/Source/NiagaraEditorWidgets/Private/SNiagaraOverviewGraph.cpp
+++ b/Engine/Plugins/FX/Niagara/Source/NiagaraEditorWidgets/Private/SNiagaraOverviewGraph.cpp
@@ -31,6 +31,7 @@ void SNiagaraOverviewGraph::Construct(const FArguments& InArgs, TSharedRef<FNiag
 {
 	ViewModel = InViewModel;
 	ViewModel->GetNodeSelection()->OnSelectedObjectsChanged().AddSP(this, &SNiagaraOverviewGraph::ViewModelSelectionChanged);
+	ViewModel->GetSystemViewModel()->OnPreClose().AddSP(this, &SNiagaraOverviewGraph::PreClose);
 
 	bUpdatingViewModelSelectionFromGraph = false;
 	bUpdatingGraphSelectionFromViewModel = false;
@@ -128,17 +129,6 @@ void SNiagaraOverviewGraph::Construct(const FArguments& InArgs, TSharedRef<FNiag
 	];
 }
 
-SNiagaraOverviewGraph::~SNiagaraOverviewGraph()
-{
-	if (ViewModel.IsValid() && GraphEditor.IsValid())
-	{
-		FVector2D Location;
-		float Zoom;
-		GraphEditor->GetViewLocation(Location, Zoom);
-		ViewModel->SetViewSettings(FNiagaraGraphViewSettings(Location, Zoom));
-	}
-}
-
 void SNiagaraOverviewGraph::Tick(const FGeometry& AllottedGeometry, const double InCurrentTime, const float InDeltaTime)
 {
 	if (ZoomToFitFrameDelay > 0)
@@ -187,6 +177,17 @@ void SNiagaraOverviewGraph::GraphSelectionChanged(const TSet<UObject*>& Selected
 	}
 }
 
+void SNiagaraOverviewGraph::PreClose()
+{
+	if (ViewModel.IsValid() && GraphEditor.IsValid())
+	{
+		FVector2D Location;
+		float Zoom;
+		GraphEditor->GetViewLocation(Location, Zoom);
+		ViewModel->SetViewSettings(FNiagaraGraphViewSettings(Location, Zoom));
+	}
+}
+
 FActionMenuContent SNiagaraOverviewGraph::OnCreateGraphActionMenu(UEdGraph* InGraph, const FVector2D& InNodePosition, const TArray<UEdGraphPin*>& InDraggedPins, bool bAutoExpand, SGraphEditor::FActionMenuClosed InOnMenuClosed)
 {
 	if (ViewModel->GetSystemViewModel()->GetEditMode() == ENiagaraSystemViewModelEditMode::SystemAsset)
diff --git a/Engine/Plugins/FX/Niagara/Source/NiagaraEditorWidgets/Private/SNiagaraOverviewGraph.h b/Engine/Plugins/FX/Niagara/Source/NiagaraEditorWidgets/Private/SNiagaraOverviewGraph.h
--- a/Engine/Plugins/FX/Niagara/Source/NiagaraEditorWidgets/Private/SNiagaraOverviewGraph.h
+++ b/Engine/Plugins/FX/Niagara/Source/NiagaraEditorWidgets/Private/SNiagaraOverviewGraph.h
@@ -19,8 +19,6 @@ public:
 
 	void Construct(const FArguments& InArgs, TSharedRef<FNiagaraOverviewGraphViewModel> InViewModel);
 
-	~SNiagaraOverviewGraph();
-
 	virtual void Tick(const FGeometry& AllottedGeometry, const double InCurrentTime, const float InDeltaTime) override;
 
 private:
@@ -28,6 +26,8 @@ private:
 
 	void GraphSelectionChanged(const TSet<UObject*>& SelectedNodes);
 
+	void PreClose();
+
 	/** Called to create context menu when right-clicking on graph */
 	FActionMenuContent OnCreateGraphActionMenu(UEdGraph* InGraph, const FVector2D& InNodePosition, const TArray<UEdGraphPin*>& InDraggedPins, bool bAutoExpand, SGraphEditor::FActionMenuClosed InOnMenuClosed);
 
