Niagara - Fix a selection issue which was preventing the user from selecting and deleting modules and renderers in the system overview.

diff --git a/Engine/Plugins/FX/Niagara/Source/NiagaraEditorWidgets/Private/SNiagaraStack.cpp b/Engine/Plugins/FX/Niagara/Source/NiagaraEditorWidgets/Private/SNiagaraStack.cpp
--- a/Engine/Plugins/FX/Niagara/Source/NiagaraEditorWidgets/Private/SNiagaraStack.cpp
+++ b/Engine/Plugins/FX/Niagara/Source/NiagaraEditorWidgets/Private/SNiagaraStack.cpp
@@ -1044,8 +1044,6 @@ void SNiagaraStack::StackStructureChanged()
 	SynchronizeTreeExpansion();
 	StackTree->RequestTreeRefresh();
 	HeaderList->RequestListRefresh();
-	// Keep the stack search focused by default
-	FSlateApplication::Get().SetKeyboardFocus(SearchBox, EFocusCause::OtherWidgetLostFocus);
 }
 
 EVisibility SNiagaraStack::GetIssueIconVisibility() const
