- world transform integration
    * removed UnitToWorld and WorldToUnit and renamed with Simulation instead of Unit
    * transform functions take a matrix now to use
    * removed world position from DIs...this must be specified when the transform is created
    * added blueprint functions and modules to create transforms
    * updated all examples and emitters for water and sand
- fixed normal generation

diff --git a/Engine/Plugins/FX/Niagara/Source/Niagara/Classes/NiagaraDataInterfaceGrid2DCollection.h b/Engine/Plugins/FX/Niagara/Source/Niagara/Classes/NiagaraDataInterfaceGrid2DCollection.h
--- a/Engine/Plugins/FX/Niagara/Source/Niagara/Classes/NiagaraDataInterfaceGrid2DCollection.h
+++ b/Engine/Plugins/FX/Niagara/Source/Niagara/Classes/NiagaraDataInterfaceGrid2DCollection.h
@@ -34,8 +34,7 @@ public:
 
 	FVector2D CellSize;
 	bool SetGridFromVoxelSize;
-
-	FVector WorldBBoxMin;
+	
 	FVector2D WorldBBoxSize;
 
 	TArray<Grid2DBuffer*> Buffers;
@@ -146,6 +145,7 @@ public:
 	UFUNCTION(BlueprintCallable, Category = Niagara)
 	virtual void GetTextureSize(const UNiagaraComponent *Component, int &SizeX, int &SizeY);
 
+	void GetWorldBBoxSize(FVectorVMContext& Context);
 	void GetCellSize(FVectorVMContext& Context);
 
 protected:
diff --git a/Engine/Plugins/FX/Niagara/Source/Niagara/Classes/NiagaraDataInterfaceNeighborGrid3D.h b/Engine/Plugins/FX/Niagara/Source/Niagara/Classes/NiagaraDataInterfaceNeighborGrid3D.h
--- a/Engine/Plugins/FX/Niagara/Source/Niagara/Classes/NiagaraDataInterfaceNeighborGrid3D.h
+++ b/Engine/Plugins/FX/Niagara/Source/Niagara/Classes/NiagaraDataInterfaceNeighborGrid3D.h
@@ -22,8 +22,7 @@ public:
 	FIntVector NumVoxels;
 	float VoxelSize;
 	bool SetGridFromVoxelSize;
-	uint32 MaxNeighborsPerVoxel;
-	FVector WorldBBoxMin;
+	uint32 MaxNeighborsPerVoxel;	
 	FVector WorldBBoxSize;
 
 	FRWBuffer NeighborhoodBuffer;
@@ -89,6 +88,8 @@ public:
 	virtual int32 PerInstanceDataSize()const override { return sizeof(NeighborGrid3DRWInstanceData); }
 	//~ UNiagaraDataInterface interface END
 
+	void GetWorldBBoxSize(FVectorVMContext& Context);
+
 protected:
 	//~ UNiagaraDataInterface interface
 	virtual bool CopyToInternal(UNiagaraDataInterface* Destination) const override;
diff --git a/Engine/Plugins/FX/Niagara/Source/Niagara/Classes/NiagaraDataInterfaceRW.h b/Engine/Plugins/FX/Niagara/Source/Niagara/Classes/NiagaraDataInterfaceRW.h
--- a/Engine/Plugins/FX/Niagara/Source/Niagara/Classes/NiagaraDataInterfaceRW.h
+++ b/Engine/Plugins/FX/Niagara/Source/Niagara/Classes/NiagaraDataInterfaceRW.h
@@ -9,7 +9,6 @@ class FNiagaraSystemInstance;
 // Global HLSL variable base names, used by HLSL.
 static const FString NumVoxelsName(TEXT("NumVoxels_"));
 static const FString VoxelSizeName(TEXT("VoxelSize_"));
-static const FString WorldBBoxMinName(TEXT("WorldBBoxMin_"));
 static const FString WorldBBoxSizeName(TEXT("WorldBBoxSize_"));
 
 static const FString NumCellsName(TEXT("NumCells_"));
@@ -22,8 +21,10 @@ static const FName VoxelSizeFunctionName("GetVoxelSize");
 static const FName NumCellsFunctionName("GetNumCells");
 static const FName CellSizeFunctionName("GetCellSize");
 
-static const FName WorldToUnitFunctionName("WorldToUnit");
-static const FName UnitToWorldFunctionName("UnitToWorld");
+static const FName WorldBBoxSizeFunctionName("GetWorldBBoxSize");
+
+static const FName SimulationToUnitFunctionName("SimulationToUnit");
+static const FName UnitToSimulationFunctionName("UnitToSimulation");
 static const FName UnitToIndexFunctionName("UnitToIndex");
 static const FName IndexToUnitFunctionName("IndexToUnit");
 static const FName IndexToUnitStaggeredXFunctionName("IndexToUnitStaggeredX");
@@ -134,10 +135,7 @@ public:
 	float VoxelSize;
 
 	UPROPERTY(EditAnywhere, Category = "Grid")
-	bool SetGridFromVoxelSize;	
-
-	UPROPERTY(EditAnywhere, Category = "Grid")
-	FVector WorldBBoxMin;
+	bool SetGridFromVoxelSize;		
 
 	UPROPERTY(EditAnywhere, Category = "Grid")
 	FVector WorldBBoxSize;
@@ -147,7 +145,8 @@ public:
 	//~ UNiagaraDataInterface interface
 	// VM functionality
 	virtual void GetFunctions(TArray<FNiagaraFunctionSignature>& OutFunctions) override;
-	virtual void GetVMExternalFunction(const FVMExternalFunctionBindingInfo& BindingInfo, void* InstanceData, FVMExternalFunction &OutFunc) override;
+	virtual void GetVMExternalFunction(const FVMExternalFunctionBindingInfo& BindingInfo, void* InstanceData, FVMExternalFunction &OutFunc) override;	
+
 
 	virtual bool Equals(const UNiagaraDataInterface* Other) const override;
 
@@ -186,10 +185,7 @@ public:
 	int32 NumAttributes;
 
 	UPROPERTY(EditAnywhere, Category = "Grid")
-	bool SetGridFromMaxAxis;
-
-	UPROPERTY(EditAnywhere, Category = "Grid")
-	FVector WorldBBoxMin;
+	bool SetGridFromMaxAxis;	
 
 	UPROPERTY(EditAnywhere, Category = "Grid")
 	FVector2D WorldBBoxSize;
diff --git a/Engine/Plugins/FX/Niagara/Source/Niagara/Private/NiagaraDataInterfaceGrid2DCollection.cpp b/Engine/Plugins/FX/Niagara/Source/Niagara/Private/NiagaraDataInterfaceGrid2DCollection.cpp
--- a/Engine/Plugins/FX/Niagara/Source/Niagara/Private/NiagaraDataInterfaceGrid2DCollection.cpp
+++ b/Engine/Plugins/FX/Niagara/Source/Niagara/Private/NiagaraDataInterfaceGrid2DCollection.cpp
@@ -34,8 +34,7 @@ struct FNiagaraDataInterfaceParametersCS_Grid2DCollection : public FNiagaraDataI
 		NumTilesParam.Bind(ParameterMap, *(NumTilesName + ParamRef.ParameterInfo.DataInterfaceHLSLSymbol));
 
 		CellSizeParam.Bind(ParameterMap, *(CellSizeName + ParamRef.ParameterInfo.DataInterfaceHLSLSymbol));
-
-		WorldBBoxMinParam.Bind(ParameterMap, *(WorldBBoxMinName + ParamRef.ParameterInfo.DataInterfaceHLSLSymbol));
+		
 		WorldBBoxSizeParam.Bind(ParameterMap, *(WorldBBoxSizeName + ParamRef.ParameterInfo.DataInterfaceHLSLSymbol));
 
 		GridParam.Bind(ParameterMap, *(GridName + ParamRef.ParameterInfo.DataInterfaceHLSLSymbol));
@@ -48,8 +47,7 @@ struct FNiagaraDataInterfaceParametersCS_Grid2DCollection : public FNiagaraDataI
 	{		
 		Ar << NumCellsParam;		
 		Ar << NumTilesParam;	
-		Ar << CellSizeParam;	
-		Ar << WorldBBoxMinParam;
+		Ar << CellSizeParam;			
 		Ar << WorldBBoxSizeParam;
 
 		Ar << GridParam;
@@ -80,8 +78,7 @@ struct FNiagaraDataInterfaceParametersCS_Grid2DCollection : public FNiagaraDataI
 		SetShaderValue(RHICmdList, ComputeShaderRHI, NumTilesParam, NumTilesTmp);		
 
 		SetShaderValue(RHICmdList, ComputeShaderRHI, CellSizeParam, ProxyData->CellSize);		
-		
-		SetShaderValue(RHICmdList, ComputeShaderRHI, WorldBBoxMinParam, ProxyData->WorldBBoxMin);
+				
 		SetShaderValue(RHICmdList, ComputeShaderRHI, WorldBBoxSizeParam, ProxyData->WorldBBoxSize);
 
 		FRHISamplerState *SamplerState = TStaticSamplerState<SF_Bilinear, AM_Clamp, AM_Clamp, AM_Clamp>::GetRHI();
@@ -112,8 +109,7 @@ private:
 
 	FShaderParameter NumCellsParam;
 	FShaderParameter NumTilesParam;
-	FShaderParameter CellSizeParam;
-	FShaderParameter WorldBBoxMinParam;
+	FShaderParameter CellSizeParam;	
 	FShaderParameter WorldBBoxSizeParam;
 
 	FShaderResourceParameter GridParam;
@@ -193,21 +189,25 @@ void UNiagaraDataInterfaceGrid2DCollection::GetFunctions(TArray<FNiagaraFunction
 
 }
 
-// #todo(dmp): expose more CPU functinality
+// #todo(dmp): expose more CPU functionality
+// #todo(dmp): ideally these would be exposed on the parent class, but we can't bind functions of parent classes but need to work on the interface
+// for sharing an instance data object with the super class
+DEFINE_NDI_DIRECT_FUNC_BINDER(UNiagaraDataInterfaceGrid2DCollection, GetWorldBBoxSize);
 DEFINE_NDI_DIRECT_FUNC_BINDER(UNiagaraDataInterfaceGrid2DCollection, GetCellSize);
 void UNiagaraDataInterfaceGrid2DCollection::GetVMExternalFunction(const FVMExternalFunctionBindingInfo& BindingInfo, void* InstanceData, FVMExternalFunction &OutFunc)
 {
-	if (BindingInfo.Name == NumCellsFunctionName) { OutFunc = FVMExternalFunction::CreateUObject(this, &UNiagaraDataInterfaceRWBase::EmptyVMFunction); }
-	else if (BindingInfo.Name == WorldToUnitFunctionName) { OutFunc = FVMExternalFunction::CreateUObject(this, &UNiagaraDataInterfaceRWBase::EmptyVMFunction); }
-	else if (BindingInfo.Name == UnitToWorldFunctionName) { OutFunc = FVMExternalFunction::CreateUObject(this, &UNiagaraDataInterfaceRWBase::EmptyVMFunction); }
-	else if (BindingInfo.Name == UnitToIndexFunctionName) { OutFunc = FVMExternalFunction::CreateUObject(this, &UNiagaraDataInterfaceRWBase::EmptyVMFunction); }
-	else if (BindingInfo.Name == IndexToUnitFunctionName) { OutFunc = FVMExternalFunction::CreateUObject(this, &UNiagaraDataInterfaceRWBase::EmptyVMFunction); }
-	else if (BindingInfo.Name == IndexToUnitStaggeredXFunctionName) { OutFunc = FVMExternalFunction::CreateUObject(this, &UNiagaraDataInterfaceRWBase::EmptyVMFunction); }
-	else if (BindingInfo.Name == IndexToUnitStaggeredYFunctionName) { OutFunc = FVMExternalFunction::CreateUObject(this, &UNiagaraDataInterfaceRWBase::EmptyVMFunction); }
-	else if (BindingInfo.Name == IndexToLinearFunctionName) { OutFunc = FVMExternalFunction::CreateUObject(this, &UNiagaraDataInterfaceRWBase::EmptyVMFunction); }
-	else if (BindingInfo.Name == LinearToIndexFunctionName) { OutFunc = FVMExternalFunction::CreateUObject(this, &UNiagaraDataInterfaceRWBase::EmptyVMFunction); }
-	else if (BindingInfo.Name == CellSizeFunctionName) 
-	{ 
+	Super::GetVMExternalFunction(BindingInfo, InstanceData, OutFunc);
+
+	
+	if (BindingInfo.Name == WorldBBoxSizeFunctionName)
+	{
+		check(BindingInfo.GetNumInputs() == 1 && BindingInfo.GetNumOutputs() == 2);
+		NDI_FUNC_BINDER(UNiagaraDataInterfaceGrid2DCollection, GetWorldBBoxSize)::Bind(this, OutFunc);
+	}
+	else if (BindingInfo.Name == CellSizeFunctionName)
+	{
+		// #todo(dmp): this will override the base class definition for GetCellSize because the data interface instance data computes cell size
+		// it would be nice to refactor this so it can be part of the super class
 		check(BindingInfo.GetNumInputs() == 1 && BindingInfo.GetNumOutputs() == 2);
 		NDI_FUNC_BINDER(UNiagaraDataInterfaceGrid2DCollection, GetCellSize)::Bind(this, OutFunc);
 	}
@@ -361,8 +361,7 @@ bool UNiagaraDataInterfaceGrid2DCollection::InitPerInstanceData(void* PerInstanc
 
 	int RT_NumTilesX = NumTilesX;
 	int RT_NumTilesY = NumTilesY;
-
-	FVector RT_WorldBBoxMin = WorldBBoxMin;
+	
 	FVector2D RT_WorldBBoxSize = WorldBBoxSize;
 
 	// If we are setting the grid from the voxel size, then recompute NumVoxels and change bbox	
@@ -393,10 +392,11 @@ bool UNiagaraDataInterfaceGrid2DCollection::InitPerInstanceData(void* PerInstanc
 	
 	FVector2D RT_CellSize = RT_WorldBBoxSize / FVector2D(RT_NumCellsX, RT_NumCellsY);
 	InstanceData->CellSize = RT_CellSize;
+	InstanceData->WorldBBoxSize = RT_WorldBBoxSize;
 	
 	// Push Updates to Proxy.
 	ENQUEUE_RENDER_COMMAND(FUpdateData)(
-		[RT_Proxy, RT_NumCellsX, RT_NumCellsY, RT_NumTilesX, RT_NumTilesY, RT_CellSize, RT_WorldBBoxMin, RT_WorldBBoxSize, RT_OutputShaderStages, RT_IterationShaderStages, InstanceID = SystemInstance->GetId()](FRHICommandListImmediate& RHICmdList)
+		[RT_Proxy, RT_NumCellsX, RT_NumCellsY, RT_NumTilesX, RT_NumTilesY, RT_CellSize, RT_WorldBBoxSize, RT_OutputShaderStages, RT_IterationShaderStages, InstanceID = SystemInstance->GetId()](FRHICommandListImmediate& RHICmdList)
 	{
 		Grid2DCollectionRWInstanceData* TargetData = RT_Proxy->SystemInstancesToProxyData.Find(InstanceID);
 		if (TargetData != nullptr)
@@ -415,8 +415,7 @@ bool UNiagaraDataInterfaceGrid2DCollection::InitPerInstanceData(void* PerInstanc
 		TargetData->NumTilesY = RT_NumTilesY;
 
 		TargetData->CellSize = RT_CellSize;
-
-		TargetData->WorldBBoxMin = RT_WorldBBoxMin;
+		
 		TargetData->WorldBBoxSize = RT_WorldBBoxSize;
 
 		RT_Proxy->OutputShaderStages = RT_OutputShaderStages;
@@ -631,6 +630,20 @@ void UNiagaraDataInterfaceGrid2DCollection::GetTextureSize(const UNiagaraCompone
 }
 
 
+void UNiagaraDataInterfaceGrid2DCollection::GetWorldBBoxSize(FVectorVMContext& Context)
+{
+	VectorVM::FUserPtrHandler<Grid2DCollectionRWInstanceData> InstData(Context);
+	VectorVM::FExternalFuncRegisterHandler<float> OutWorldBoundsX(Context);
+	VectorVM::FExternalFuncRegisterHandler<float> OutWorldBoundsY(Context);
+
+	for (int32 InstanceIdx = 0; InstanceIdx < Context.NumInstances; ++InstanceIdx)
+	{
+		*OutWorldBoundsX.GetDestAndAdvance() = InstData->WorldBBoxSize.X;
+		*OutWorldBoundsY.GetDestAndAdvance() = InstData->WorldBBoxSize.Y;
+	}
+}
+
+
 void UNiagaraDataInterfaceGrid2DCollection::GetCellSize(FVectorVMContext& Context)
 {
 	VectorVM::FUserPtrHandler<Grid2DCollectionRWInstanceData> InstData(Context);
diff --git a/Engine/Plugins/FX/Niagara/Source/Niagara/Private/NiagaraDataInterfaceNeighborGrid3D.cpp b/Engine/Plugins/FX/Niagara/Source/Niagara/Private/NiagaraDataInterfaceNeighborGrid3D.cpp
--- a/Engine/Plugins/FX/Niagara/Source/Niagara/Private/NiagaraDataInterfaceNeighborGrid3D.cpp
+++ b/Engine/Plugins/FX/Niagara/Source/Niagara/Private/NiagaraDataInterfaceNeighborGrid3D.cpp
@@ -33,8 +33,7 @@ struct FNiagaraDataInterfaceParametersCS_NeighborGrid3D : public FNiagaraDataInt
 		NumVoxelsParam.Bind(ParameterMap, *(NumVoxelsName + ParamRef.ParameterInfo.DataInterfaceHLSLSymbol));		
 		VoxelSizeParam.Bind(ParameterMap, *(VoxelSizeName + ParamRef.ParameterInfo.DataInterfaceHLSLSymbol));
 
-		MaxNeighborsPerVoxelParam.Bind(ParameterMap, *(MaxNeighborsPerVoxelName + ParamRef.ParameterInfo.DataInterfaceHLSLSymbol));
-		WorldBBoxMinParam.Bind(ParameterMap, *(WorldBBoxMinName + ParamRef.ParameterInfo.DataInterfaceHLSLSymbol));
+		MaxNeighborsPerVoxelParam.Bind(ParameterMap, *(MaxNeighborsPerVoxelName + ParamRef.ParameterInfo.DataInterfaceHLSLSymbol));		
 		WorldBBoxSizeParam.Bind(ParameterMap, *(WorldBBoxSizeName + ParamRef.ParameterInfo.DataInterfaceHLSLSymbol));
 		
 		ParticleNeighborsGridParam.Bind(ParameterMap,  *(ParticleNeighborsName + ParamRef.ParameterInfo.DataInterfaceHLSLSymbol));
@@ -48,8 +47,7 @@ struct FNiagaraDataInterfaceParametersCS_NeighborGrid3D : public FNiagaraDataInt
 	{
 		Ar << NumVoxelsParam;
 		Ar << VoxelSizeParam;
-		Ar << MaxNeighborsPerVoxelParam;
-		Ar << WorldBBoxMinParam;
+		Ar << MaxNeighborsPerVoxelParam;		
 		Ar << WorldBBoxSizeParam;
 		
 		Ar << ParticleNeighborsGridParam;
@@ -79,8 +77,7 @@ struct FNiagaraDataInterfaceParametersCS_NeighborGrid3D : public FNiagaraDataInt
 		VoxelSizeTmp[2] = ProxyData->WorldBBoxSize.Z / ProxyData->NumVoxels.Z;
 		SetShaderValue(RHICmdList, ComputeShaderRHI, VoxelSizeParam, VoxelSizeTmp);
 
-		SetShaderValue(RHICmdList, ComputeShaderRHI, MaxNeighborsPerVoxelParam, ProxyData->MaxNeighborsPerVoxel);
-		SetShaderValue(RHICmdList, ComputeShaderRHI, WorldBBoxMinParam, ProxyData->WorldBBoxMin);
+		SetShaderValue(RHICmdList, ComputeShaderRHI, MaxNeighborsPerVoxelParam, ProxyData->MaxNeighborsPerVoxel);		
 		SetShaderValue(RHICmdList, ComputeShaderRHI, WorldBBoxSizeParam, ProxyData->WorldBBoxSize);
 		
 		if (!Context.IsOutputStage)
@@ -130,8 +127,7 @@ struct FNiagaraDataInterfaceParametersCS_NeighborGrid3D : public FNiagaraDataInt
 private:
 	FShaderParameter NumVoxelsParam;
 	FShaderParameter VoxelSizeParam;
-	FShaderParameter MaxNeighborsPerVoxelParam;
-	FShaderParameter WorldBBoxMinParam;
+	FShaderParameter MaxNeighborsPerVoxelParam;	
 	FShaderParameter WorldBBoxSizeParam;
 	FShaderResourceParameter ParticleNeighborsGridParam;
 	FShaderResourceParameter ParticleNeighborCountGridParam;
@@ -230,12 +226,18 @@ void UNiagaraDataInterfaceNeighborGrid3D::GetFunctions(TArray<FNiagaraFunctionSi
 	}
 }
 
+DEFINE_NDI_DIRECT_FUNC_BINDER(UNiagaraDataInterfaceNeighborGrid3D, GetWorldBBoxSize);
 
 void UNiagaraDataInterfaceNeighborGrid3D::GetVMExternalFunction(const FVMExternalFunctionBindingInfo& BindingInfo, void* InstanceData, FVMExternalFunction &OutFunc)
 {
 	Super::GetVMExternalFunction(BindingInfo, InstanceData, OutFunc);
 
-	if (BindingInfo.Name == MaxNeighborsPerVoxelFunctionName) { OutFunc = FVMExternalFunction::CreateUObject(this, &UNiagaraDataInterfaceRWBase::EmptyVMFunction); }	
+	// #todo(dmp): this overrides the empty function set by the super class
+	if (BindingInfo.Name == WorldBBoxSizeFunctionName) {
+		check(BindingInfo.GetNumInputs() == 1 && BindingInfo.GetNumOutputs() == 3);
+		NDI_FUNC_BINDER(UNiagaraDataInterfaceNeighborGrid3D, GetWorldBBoxSize)::Bind(this, OutFunc);
+	}
+	else if (BindingInfo.Name == MaxNeighborsPerVoxelFunctionName) { OutFunc = FVMExternalFunction::CreateUObject(this, &UNiagaraDataInterfaceRWBase::EmptyVMFunction); }	
 	else if (BindingInfo.Name == NeighborGridIndexToLinearFunctionName) { OutFunc = FVMExternalFunction::CreateUObject(this, &UNiagaraDataInterfaceRWBase::EmptyVMFunction); }
 	else if (BindingInfo.Name == GetParticleNeighborFunctionName) { OutFunc = FVMExternalFunction::CreateUObject(this, &UNiagaraDataInterfaceRWBase::EmptyVMFunction); }
 	else if (BindingInfo.Name == SetParticleNeighborFunctionName) { OutFunc = FVMExternalFunction::CreateUObject(this, &UNiagaraDataInterfaceRWBase::EmptyVMFunction); }
@@ -243,6 +245,22 @@ void UNiagaraDataInterfaceNeighborGrid3D::GetVMExternalFunction(const FVMExterna
 	else if (BindingInfo.Name == SetParticleNeighborCountFunctionName) { OutFunc = FVMExternalFunction::CreateUObject(this, &UNiagaraDataInterfaceRWBase::EmptyVMFunction); }
 }
 
+void UNiagaraDataInterfaceNeighborGrid3D::GetWorldBBoxSize(FVectorVMContext& Context)
+{
+	VectorVM::FUserPtrHandler<NeighborGrid3DRWInstanceData> InstData(Context);
+
+	VectorVM::FExternalFuncRegisterHandler<float> OutWorldBoundsX(Context);
+	VectorVM::FExternalFuncRegisterHandler<float> OutWorldBoundsY(Context);
+	VectorVM::FExternalFuncRegisterHandler<float> OutWorldBoundsZ(Context);
+
+	for (int32 InstanceIdx = 0; InstanceIdx < Context.NumInstances; ++InstanceIdx)
+	{
+		*OutWorldBoundsX.GetDestAndAdvance() = WorldBBoxSize.X;
+		*OutWorldBoundsY.GetDestAndAdvance() = WorldBBoxSize.Y;
+		*OutWorldBoundsZ.GetDestAndAdvance() = WorldBBoxSize.Z;
+	}
+}
+
 bool UNiagaraDataInterfaceNeighborGrid3D::Equals(const UNiagaraDataInterface* Other) const
 {
 	if (!Super::Equals(Other))
@@ -393,8 +411,7 @@ bool UNiagaraDataInterfaceNeighborGrid3D::InitPerInstanceData(void* PerInstanceD
 	FNiagaraDataInterfaceProxyNeighborGrid3D* RT_Proxy = GetProxyAs<FNiagaraDataInterfaceProxyNeighborGrid3D>();
 
 	FIntVector RT_NumVoxels = NumVoxels;
-	uint32 RT_MaxNeighborsPerVoxel = MaxNeighborsPerVoxel;
-	FVector RT_WorldBBoxMin = WorldBBoxMin;
+	uint32 RT_MaxNeighborsPerVoxel = MaxNeighborsPerVoxel;	
 	FVector RT_WorldBBoxSize = WorldBBoxSize;
 	TSet<int> RT_OutputShaderStages = OutputShaderStages;
 	TSet<int> RT_IterationShaderStages = IterationShaderStages;
@@ -413,14 +430,13 @@ bool UNiagaraDataInterfaceNeighborGrid3D::InitPerInstanceData(void* PerInstanceD
 	}
 
 	InstanceData->NumVoxels = RT_NumVoxels;
-	InstanceData->MaxNeighborsPerVoxel = RT_MaxNeighborsPerVoxel;
-	InstanceData->WorldBBoxMin = RT_WorldBBoxMin;
+	InstanceData->MaxNeighborsPerVoxel = RT_MaxNeighborsPerVoxel;	
 	InstanceData->WorldBBoxSize = RT_WorldBBoxSize;
 
 	// @todo-threadsafety. This would be a race but I'm taking a ref here. Not ideal in the long term.
 	// Push Updates to Proxy.
 	ENQUEUE_RENDER_COMMAND(FUpdateData)(
-		[RT_Proxy, RT_NumVoxels, RT_MaxNeighborsPerVoxel, RT_WorldBBoxMin, RT_WorldBBoxSize, RT_OutputShaderStages, RT_IterationShaderStages, InstanceID = SystemInstance->GetId()](FRHICommandListImmediate& RHICmdList)
+		[RT_Proxy, RT_NumVoxels, RT_MaxNeighborsPerVoxel, RT_WorldBBoxSize, RT_OutputShaderStages, RT_IterationShaderStages, InstanceID = SystemInstance->GetId()](FRHICommandListImmediate& RHICmdList)
 	{
 		NeighborGrid3DRWInstanceData* TargetData = RT_Proxy->SystemInstancesToProxyData.Find(InstanceID);
 		if (TargetData != nullptr)
@@ -433,8 +449,7 @@ bool UNiagaraDataInterfaceNeighborGrid3D::InitPerInstanceData(void* PerInstanceD
 		}
 
 		TargetData->NumVoxels = RT_NumVoxels;
-		TargetData->MaxNeighborsPerVoxel = RT_MaxNeighborsPerVoxel;
-		TargetData->WorldBBoxMin = RT_WorldBBoxMin;
+		TargetData->MaxNeighborsPerVoxel = RT_MaxNeighborsPerVoxel;		
 		TargetData->WorldBBoxSize = RT_WorldBBoxSize;
 
 		RT_Proxy->OutputShaderStages = RT_OutputShaderStages;
diff --git a/Engine/Plugins/FX/Niagara/Source/Niagara/Private/NiagaraDataInterfaceRW.cpp b/Engine/Plugins/FX/Niagara/Source/Niagara/Private/NiagaraDataInterfaceRW.cpp
--- a/Engine/Plugins/FX/Niagara/Source/Niagara/Private/NiagaraDataInterfaceRW.cpp
+++ b/Engine/Plugins/FX/Niagara/Source/Niagara/Private/NiagaraDataInterfaceRW.cpp
@@ -48,7 +48,6 @@ UNiagaraDataInterfaceGrid3D::UNiagaraDataInterfaceGrid3D(FObjectInitializer cons
 	, NumVoxels(3, 3, 3)
 	, VoxelSize(1.)
 	, SetGridFromVoxelSize(false)	
-	, WorldBBoxMin(0., 0., 0.)
 	, WorldBBoxSize(100., 100., 100.)
 {
 	Proxy = MakeShared<FNiagaraDataInterfaceProxyRW, ESPMode::ThreadSafe>();
@@ -60,6 +59,17 @@ UNiagaraDataInterfaceGrid3D::UNiagaraDataInterfaceGrid3D(FObjectInitializer cons
 void UNiagaraDataInterfaceGrid3D::GetFunctions(TArray<FNiagaraFunctionSignature>& OutFunctions)
 {
 
+	{
+		FNiagaraFunctionSignature Sig;
+		Sig.Name = WorldBBoxSizeFunctionName;
+		Sig.Inputs.Add(FNiagaraVariable(FNiagaraTypeDefinition(GetClass()), TEXT("Grid")));
+		Sig.Outputs.Add(FNiagaraVariable(FNiagaraTypeDefinition::GetVec3Def(), TEXT("WorldBBoxSize")));
+
+		Sig.bMemberFunction = true;
+		Sig.bRequiresContext = false;
+		OutFunctions.Add(Sig);
+	}
+
 	{
 		FNiagaraFunctionSignature Sig;
 		Sig.Name = NumVoxelsFunctionName;
@@ -75,9 +85,10 @@ void UNiagaraDataInterfaceGrid3D::GetFunctions(TArray<FNiagaraFunctionSignature>
 
 	{
 		FNiagaraFunctionSignature Sig;
-		Sig.Name = WorldToUnitFunctionName;
+		Sig.Name = SimulationToUnitFunctionName;
 		Sig.Inputs.Add(FNiagaraVariable(FNiagaraTypeDefinition(GetClass()), TEXT("Grid")));
-		Sig.Inputs.Add(FNiagaraVariable(FNiagaraTypeDefinition::GetVec3Def(), TEXT("World")));
+		Sig.Inputs.Add(FNiagaraVariable(FNiagaraTypeDefinition::GetVec3Def(), TEXT("Simulation")));
+		Sig.Inputs.Add(FNiagaraVariable(FNiagaraTypeDefinition::GetMatrix4Def(), TEXT("SimulationToUnitTransform")));
 		Sig.Outputs.Add(FNiagaraVariable(FNiagaraTypeDefinition::GetVec3Def(), TEXT("Unit")));
 
 		Sig.bMemberFunction = true;
@@ -138,11 +149,11 @@ void UNiagaraDataInterfaceGrid3D::GetFunctions(TArray<FNiagaraFunctionSignature>
 	}
 }
 
-
 void UNiagaraDataInterfaceGrid3D::GetVMExternalFunction(const FVMExternalFunctionBindingInfo& BindingInfo, void* InstanceData, FVMExternalFunction &OutFunc)
 {
-	if (BindingInfo.Name == NumVoxelsFunctionName) { OutFunc = FVMExternalFunction::CreateUObject(this, &UNiagaraDataInterfaceRWBase::EmptyVMFunction); }	
-	else if (BindingInfo.Name == WorldToUnitFunctionName) { OutFunc = FVMExternalFunction::CreateUObject(this, &UNiagaraDataInterfaceRWBase::EmptyVMFunction); }
+	if (BindingInfo.Name == WorldBBoxSizeFunctionName) { OutFunc = FVMExternalFunction::CreateUObject(this, &UNiagaraDataInterfaceRWBase::EmptyVMFunction); }
+	else if (BindingInfo.Name == NumVoxelsFunctionName) { OutFunc = FVMExternalFunction::CreateUObject(this, &UNiagaraDataInterfaceRWBase::EmptyVMFunction); }	
+	else if (BindingInfo.Name == SimulationToUnitFunctionName) { OutFunc = FVMExternalFunction::CreateUObject(this, &UNiagaraDataInterfaceRWBase::EmptyVMFunction); }
 	else if (BindingInfo.Name == UnitToIndexFunctionName) { OutFunc = FVMExternalFunction::CreateUObject(this, &UNiagaraDataInterfaceRWBase::EmptyVMFunction); }
 	else if (BindingInfo.Name == IndexToLinearFunctionName) { OutFunc = FVMExternalFunction::CreateUObject(this, &UNiagaraDataInterfaceRWBase::EmptyVMFunction); }
 	else if (BindingInfo.Name == VoxelSizeFunctionName) { OutFunc = FVMExternalFunction::CreateUObject(this, &UNiagaraDataInterfaceRWBase::EmptyVMFunction); }
@@ -158,22 +169,19 @@ bool UNiagaraDataInterfaceGrid3D::Equals(const UNiagaraDataInterface* Other) con
 
 	return OtherTyped->NumVoxels == NumVoxels &&
 		FMath::IsNearlyEqual(OtherTyped->VoxelSize, VoxelSize) &&		
-		OtherTyped->WorldBBoxMin.Equals(WorldBBoxMin) &&
 		OtherTyped->WorldBBoxSize.Equals(WorldBBoxSize);
 }
 
 void UNiagaraDataInterfaceGrid3D::GetParameterDefinitionHLSL(FNiagaraDataInterfaceGPUParamInfo& ParamInfo, FString& OutHLSL)
 {
 	static const TCHAR *FormatDeclarations = TEXT(R"(
-		int3 {NumVoxels};
-		float3 {VoxelSize};
-		float3 {WorldBBoxMinName};
+		int3 {NumVoxelsName};
+		float3 {VoxelSizeName};		
 		float3 {WorldBBoxSizeName};
 	)");
 	TMap<FString, FStringFormatArg> ArgsDeclarations = {
-		{ TEXT("NumVoxels"), NumVoxelsName + ParamInfo.DataInterfaceHLSLSymbol },
-		{ TEXT("VoxelSize"), VoxelSizeName + ParamInfo.DataInterfaceHLSLSymbol },
-		{ TEXT("WorldBBoxMinName"),  WorldBBoxMinName + ParamInfo.DataInterfaceHLSLSymbol },
+		{ TEXT("NumVoxelsName"), NumVoxelsName + ParamInfo.DataInterfaceHLSLSymbol },
+		{ TEXT("VoxelSizeName"), VoxelSizeName + ParamInfo.DataInterfaceHLSLSymbol },		
 		{ TEXT("WorldBBoxSizeName"),    WorldBBoxSizeName + ParamInfo.DataInterfaceHLSLSymbol },
 	};
 	OutHLSL += FString::Format(FormatDeclarations, ArgsDeclarations);
@@ -181,7 +189,26 @@ void UNiagaraDataInterfaceGrid3D::GetParameterDefinitionHLSL(FNiagaraDataInterfa
 
 bool UNiagaraDataInterfaceGrid3D::GetFunctionHLSL(const FName& DefinitionFunctionName, FString InstanceFunctionName, FNiagaraDataInterfaceGPUParamInfo& ParamInfo, FString& OutHLSL)
 {
-	if (DefinitionFunctionName == NumVoxelsFunctionName)
+	TMap<FString, FStringFormatArg> ArgsDeclarations = {
+		{ TEXT("FunctionName"), InstanceFunctionName},
+		{ TEXT("NumVoxelsName"), NumVoxelsName + ParamInfo.DataInterfaceHLSLSymbol },
+		{ TEXT("VoxelSizeName"), VoxelSizeName + ParamInfo.DataInterfaceHLSLSymbol },
+		{ TEXT("WorldBBoxSizeName"),    WorldBBoxSizeName + ParamInfo.DataInterfaceHLSLSymbol },
+	};
+
+	if (DefinitionFunctionName == WorldBBoxSizeFunctionName)
+	{
+		static const TCHAR *FormatSample = TEXT(R"(
+			void {FunctionName}(out float3 Out_WorldBBox)
+			{
+				Out_WorldBBox = {WorldBBoxSizeName};				
+			}
+		)");
+
+		OutHLSL += FString::Format(FormatSample, ArgsDeclarations);
+		return true;
+	}
+	else if (DefinitionFunctionName == NumVoxelsFunctionName)
 	{
 		static const TCHAR *FormatSample = TEXT(R"(
 			void {FunctionName}(out int Out_NumVoxelsX, out int Out_NumVoxelsY, out int Out_NumVoxelsZ)
@@ -191,28 +218,20 @@ bool UNiagaraDataInterfaceGrid3D::GetFunctionHLSL(const FName& DefinitionFunctio
 				Out_NumVoxelsZ = {NumVoxelsName}.z;
 			}
 		)");
-		TMap<FString, FStringFormatArg> ArgsSample = {
-			{TEXT("FunctionName"), InstanceFunctionName},
-			{TEXT("NumVoxelsName"), NumVoxelsName + ParamInfo.DataInterfaceHLSLSymbol},
 
-		};
-		OutHLSL += FString::Format(FormatSample, ArgsSample);
+		OutHLSL += FString::Format(FormatSample, ArgsDeclarations);
 		return true;
 	}
-	else if (DefinitionFunctionName == WorldToUnitFunctionName)
+	else if (DefinitionFunctionName == SimulationToUnitFunctionName)
 	{
 		static const TCHAR *FormatSample = TEXT(R"(
-			void {FunctionName}(float3 In_World, out float3 Out_Unit)
+			void {FunctionName}(float3 In_Simulation, float4x4 In_SimulationToUnitTransform, out float3 Out_Unit)
 			{
-				Out_Unit = (In_World - {WorldBBoxMinName}) / {WorldBBoxSizeName};
+				Out_Unit = mul(float4(In_Simulation, 1.0), In_SimulationToUnitTransform).xyz;
 			}
 		)");
-		TMap<FString, FStringFormatArg> ArgsSample = {
-			{TEXT("FunctionName"), InstanceFunctionName},
-			{TEXT("WorldBBoxMinName"), WorldBBoxMinName + ParamInfo.DataInterfaceHLSLSymbol},
-			{TEXT("WorldBBoxSizeName"), WorldBBoxSizeName + ParamInfo.DataInterfaceHLSLSymbol},
-		};
-		OutHLSL += FString::Format(FormatSample, ArgsSample);
+
+		OutHLSL += FString::Format(FormatSample, ArgsDeclarations);
 		return true;
 	}
 	else if (DefinitionFunctionName == UnitToIndexFunctionName)
@@ -226,11 +245,8 @@ bool UNiagaraDataInterfaceGrid3D::GetFunctionHLSL(const FName& DefinitionFunctio
 				Out_IndexZ = Out_IndexTmp.z;
 			}
 		)");
-		TMap<FString, FStringFormatArg> ArgsSample = {
-			{TEXT("FunctionName"), InstanceFunctionName},
-			{TEXT("NumVoxelsName"), NumVoxelsName + ParamInfo.DataInterfaceHLSLSymbol},
-		};
-		OutHLSL += FString::Format(FormatSample, ArgsSample);
+
+		OutHLSL += FString::Format(FormatSample, ArgsDeclarations);
 		return true;
 	}
 	else if (DefinitionFunctionName == IndexToLinearFunctionName)
@@ -241,11 +257,8 @@ bool UNiagaraDataInterfaceGrid3D::GetFunctionHLSL(const FName& DefinitionFunctio
 				Out_Linear = In_IndexX + In_IndexY * {NumVoxelsName}.x + In_IndexZ * {NumVoxelsName}.x * {NumVoxelsName}.y;
 			}
 		)");
-		TMap<FString, FStringFormatArg> ArgsSample = {
-			{TEXT("FunctionName"), InstanceFunctionName},
-			{TEXT("NumVoxelsName"), NumVoxelsName + ParamInfo.DataInterfaceHLSLSymbol},
-		};
-		OutHLSL += FString::Format(FormatSample, ArgsSample);
+
+		OutHLSL += FString::Format(FormatSample, ArgsDeclarations);
 		return true;
 	}
 	else if (DefinitionFunctionName == VoxelSizeFunctionName)
@@ -256,11 +269,8 @@ bool UNiagaraDataInterfaceGrid3D::GetFunctionHLSL(const FName& DefinitionFunctio
 				Out_VoxelSize = {VoxelSizeName};
 			}
 		)");
-		TMap<FString, FStringFormatArg> ArgsSample = {
-			{TEXT("FunctionName"), InstanceFunctionName},
-			{TEXT("VoxelSizeName"), VoxelSizeName + ParamInfo.DataInterfaceHLSLSymbol},
-		};
-		OutHLSL += FString::Format(FormatSample, ArgsSample);
+
+		OutHLSL += FString::Format(FormatSample, ArgsDeclarations);
 		return true;
 	}
 
@@ -286,8 +296,7 @@ bool UNiagaraDataInterfaceGrid3D::CopyToInternal(UNiagaraDataInterface* Destinat
 
 	OtherTyped->NumVoxels = NumVoxels;
 	OtherTyped->VoxelSize = VoxelSize;
-	OtherTyped->SetGridFromVoxelSize = SetGridFromVoxelSize;	
-	OtherTyped->WorldBBoxMin = WorldBBoxMin;
+	OtherTyped->SetGridFromVoxelSize = SetGridFromVoxelSize;		
 	OtherTyped->WorldBBoxSize = WorldBBoxSize;
 
 	return true;
@@ -307,8 +316,7 @@ UNiagaraDataInterfaceGrid2D::UNiagaraDataInterfaceGrid2D(FObjectInitializer cons
 	, NumCellsY(3)
 	, NumCellsMaxAxis(3)
 	, NumAttributes(1)
-	, SetGridFromMaxAxis(false)
-	, WorldBBoxMin(0., 0., 0.)
+	, SetGridFromMaxAxis(false)	
 	, WorldBBoxSize(100., 100.)
 {
 	Proxy = MakeShared<FNiagaraDataInterfaceProxyRW, ESPMode::ThreadSafe>();
@@ -317,7 +325,17 @@ UNiagaraDataInterfaceGrid2D::UNiagaraDataInterfaceGrid2D(FObjectInitializer cons
 
 
 void UNiagaraDataInterfaceGrid2D::GetFunctions(TArray<FNiagaraFunctionSignature>& OutFunctions)
-{
+{	
+	{
+		FNiagaraFunctionSignature Sig;
+		Sig.Name = WorldBBoxSizeFunctionName;
+		Sig.Inputs.Add(FNiagaraVariable(FNiagaraTypeDefinition(GetClass()), TEXT("Grid")));
+		Sig.Outputs.Add(FNiagaraVariable(FNiagaraTypeDefinition::GetVec2Def(), TEXT("WorldBBoxSize")));		
+
+		Sig.bMemberFunction = true;
+		Sig.bRequiresContext = false;
+		OutFunctions.Add(Sig);
+	}
 
 	{
 		FNiagaraFunctionSignature Sig;
@@ -333,9 +351,10 @@ void UNiagaraDataInterfaceGrid2D::GetFunctions(TArray<FNiagaraFunctionSignature>
 
 	{
 		FNiagaraFunctionSignature Sig;
-		Sig.Name = WorldToUnitFunctionName;
+		Sig.Name = SimulationToUnitFunctionName;
 		Sig.Inputs.Add(FNiagaraVariable(FNiagaraTypeDefinition(GetClass()), TEXT("Grid")));
-		Sig.Inputs.Add(FNiagaraVariable(FNiagaraTypeDefinition::GetVec3Def(), TEXT("World")));
+		Sig.Inputs.Add(FNiagaraVariable(FNiagaraTypeDefinition::GetVec3Def(), TEXT("Simulation")));
+		Sig.Inputs.Add(FNiagaraVariable(FNiagaraTypeDefinition::GetMatrix4Def(), TEXT("SimulationToUnitTransform")));
 		Sig.Outputs.Add(FNiagaraVariable(FNiagaraTypeDefinition::GetVec3Def(), TEXT("Unit")));
 
 		Sig.bMemberFunction = true;
@@ -345,10 +364,11 @@ void UNiagaraDataInterfaceGrid2D::GetFunctions(TArray<FNiagaraFunctionSignature>
 
 	{
 		FNiagaraFunctionSignature Sig;
-		Sig.Name = UnitToWorldFunctionName;
+		Sig.Name = UnitToSimulationFunctionName;
 		Sig.Inputs.Add(FNiagaraVariable(FNiagaraTypeDefinition(GetClass()), TEXT("Grid")));
 		Sig.Inputs.Add(FNiagaraVariable(FNiagaraTypeDefinition::GetVec3Def(), TEXT("Unit")));
-		Sig.Outputs.Add(FNiagaraVariable(FNiagaraTypeDefinition::GetVec3Def(), TEXT("World")));
+		Sig.Inputs.Add(FNiagaraVariable(FNiagaraTypeDefinition::GetMatrix4Def(), TEXT("UnitToSimulationTransform")));
+		Sig.Outputs.Add(FNiagaraVariable(FNiagaraTypeDefinition::GetVec3Def(), TEXT("Simulation")));
 
 		Sig.bMemberFunction = true;
 		Sig.bRequiresContext = false;
@@ -461,9 +481,10 @@ void UNiagaraDataInterfaceGrid2D::GetFunctions(TArray<FNiagaraFunctionSignature>
 
 void UNiagaraDataInterfaceGrid2D::GetVMExternalFunction(const FVMExternalFunctionBindingInfo& BindingInfo, void* InstanceData, FVMExternalFunction &OutFunc)
 {
-	if (BindingInfo.Name == NumCellsFunctionName) { OutFunc = FVMExternalFunction::CreateUObject(this, &UNiagaraDataInterfaceRWBase::EmptyVMFunction); }
-	else if (BindingInfo.Name == WorldToUnitFunctionName) { OutFunc = FVMExternalFunction::CreateUObject(this, &UNiagaraDataInterfaceRWBase::EmptyVMFunction); }
-	else if (BindingInfo.Name == UnitToWorldFunctionName) { OutFunc = FVMExternalFunction::CreateUObject(this, &UNiagaraDataInterfaceRWBase::EmptyVMFunction); }
+	if (BindingInfo.Name == WorldBBoxSizeFunctionName) { OutFunc = FVMExternalFunction::CreateUObject(this, &UNiagaraDataInterfaceRWBase::EmptyVMFunction); }
+	else if (BindingInfo.Name == NumCellsFunctionName) { OutFunc = FVMExternalFunction::CreateUObject(this, &UNiagaraDataInterfaceRWBase::EmptyVMFunction); }
+	else if (BindingInfo.Name == SimulationToUnitFunctionName) { OutFunc = FVMExternalFunction::CreateUObject(this, &UNiagaraDataInterfaceRWBase::EmptyVMFunction); }
+	else if (BindingInfo.Name == UnitToSimulationFunctionName) { OutFunc = FVMExternalFunction::CreateUObject(this, &UNiagaraDataInterfaceRWBase::EmptyVMFunction); }
 	else if (BindingInfo.Name == UnitToIndexFunctionName) { OutFunc = FVMExternalFunction::CreateUObject(this, &UNiagaraDataInterfaceRWBase::EmptyVMFunction); }
 	else if (BindingInfo.Name == IndexToUnitFunctionName) { OutFunc = FVMExternalFunction::CreateUObject(this, &UNiagaraDataInterfaceRWBase::EmptyVMFunction); }
 	else if (BindingInfo.Name == IndexToUnitStaggeredXFunctionName) { OutFunc = FVMExternalFunction::CreateUObject(this, &UNiagaraDataInterfaceRWBase::EmptyVMFunction); }
@@ -473,6 +494,7 @@ void UNiagaraDataInterfaceGrid2D::GetVMExternalFunction(const FVMExternalFunctio
 	else if (BindingInfo.Name == CellSizeFunctionName) { OutFunc = FVMExternalFunction::CreateUObject(this, &UNiagaraDataInterfaceRWBase::EmptyVMFunction); }
 }
 
+
 bool UNiagaraDataInterfaceGrid2D::Equals(const UNiagaraDataInterface* Other) const
 {
 	if (!Super::Equals(Other))
@@ -485,23 +507,20 @@ bool UNiagaraDataInterfaceGrid2D::Equals(const UNiagaraDataInterface* Other) con
 		OtherTyped->NumCellsX == NumCellsX &&
 		OtherTyped->NumCellsY == NumCellsY &&
 		OtherTyped->NumAttributes == NumAttributes &&
-		OtherTyped->NumCellsMaxAxis == NumCellsMaxAxis &&
-		OtherTyped->WorldBBoxMin.Equals(WorldBBoxMin) &&
+		OtherTyped->NumCellsMaxAxis == NumCellsMaxAxis &&		
 		OtherTyped->WorldBBoxSize.Equals(WorldBBoxSize);
 }
 
 void UNiagaraDataInterfaceGrid2D::GetParameterDefinitionHLSL(FNiagaraDataInterfaceGPUParamInfo& ParamInfo, FString& OutHLSL)
 {
 	static const TCHAR *FormatDeclarations = TEXT(R"(
-		int2 {NumCells};
-		float2 {CellSize};
-		float3 {WorldBBoxMinName};
+		int2 {NumCellsName};
+		float2 {CellSizeName};		
 		float2 {WorldBBoxSizeName};
 	)");
 	TMap<FString, FStringFormatArg> ArgsDeclarations = {
-		{ TEXT("NumCells"), NumCellsName + ParamInfo.DataInterfaceHLSLSymbol },
-		{ TEXT("CellSize"), CellSizeName + ParamInfo.DataInterfaceHLSLSymbol },
-		{ TEXT("WorldBBoxMinName"),  WorldBBoxMinName + ParamInfo.DataInterfaceHLSLSymbol },
+		{ TEXT("NumCellsName"), NumCellsName + ParamInfo.DataInterfaceHLSLSymbol },
+		{ TEXT("CellSizeName"), CellSizeName + ParamInfo.DataInterfaceHLSLSymbol },		
 		{ TEXT("WorldBBoxSizeName"),    WorldBBoxSizeName + ParamInfo.DataInterfaceHLSLSymbol },
 	};
 	OutHLSL += FString::Format(FormatDeclarations, ArgsDeclarations);
@@ -509,7 +528,26 @@ void UNiagaraDataInterfaceGrid2D::GetParameterDefinitionHLSL(FNiagaraDataInterfa
 
 bool UNiagaraDataInterfaceGrid2D::GetFunctionHLSL(const FName& DefinitionFunctionName, FString InstanceFunctionName, FNiagaraDataInterfaceGPUParamInfo& ParamInfo, FString& OutHLSL)
 {
-	if (DefinitionFunctionName == NumCellsFunctionName)
+	TMap<FString, FStringFormatArg> ArgsDeclarations = {
+		{ TEXT("FunctionName"), InstanceFunctionName},
+		{ TEXT("NumCellsName"), NumCellsName + ParamInfo.DataInterfaceHLSLSymbol },
+		{ TEXT("CellSizeName"), CellSizeName + ParamInfo.DataInterfaceHLSLSymbol },		
+		{ TEXT("WorldBBoxSizeName"),    WorldBBoxSizeName + ParamInfo.DataInterfaceHLSLSymbol },
+	};
+
+	if (DefinitionFunctionName == WorldBBoxSizeFunctionName)
+	{
+		static const TCHAR *FormatSample = TEXT(R"(
+			void {FunctionName}(out float2 Out_WorldBBox)
+			{
+				Out_WorldBBox = {WorldBBoxSizeName};				
+			}
+		)");
+
+		OutHLSL += FString::Format(FormatSample, ArgsDeclarations);
+		return true;
+	}
+	else if (DefinitionFunctionName == NumCellsFunctionName)
 	{
 		static const TCHAR *FormatSample = TEXT(R"(
 			void {FunctionName}(out int Out_NumCellsX, out int Out_NumCellsY)
@@ -518,44 +556,32 @@ bool UNiagaraDataInterfaceGrid2D::GetFunctionHLSL(const FName& DefinitionFunctio
 				Out_NumCellsY = {NumCellsName}.y;
 			}
 		)");
-		TMap<FString, FStringFormatArg> ArgsSample = {
-			{TEXT("FunctionName"), InstanceFunctionName},
-			{TEXT("NumCellsName"), NumCellsName + ParamInfo.DataInterfaceHLSLSymbol},
 
-		};
-		OutHLSL += FString::Format(FormatSample, ArgsSample);
+		OutHLSL += FString::Format(FormatSample, ArgsDeclarations);
 		return true;
 	}
-	else if (DefinitionFunctionName == WorldToUnitFunctionName)
+	else if (DefinitionFunctionName == SimulationToUnitFunctionName)
 	{
 		static const TCHAR *FormatSample = TEXT(R"(
-			void {FunctionName}(float3 In_World, out float3 Out_Unit)
+			void {FunctionName}(float3 In_Simulation, float4x4 In_SimulationToUnitTransform, out float3 Out_Unit)
 			{
-				Out_Unit = (In_World - {WorldBBoxMinName}) / float3({WorldBBoxSizeName}, 1);
+				Out_Unit = mul(float4(In_Simulation, 1.0), In_SimulationToUnitTransform).xyz;
 			}
 		)");
-		TMap<FString, FStringFormatArg> ArgsSample = {
-			{TEXT("FunctionName"), InstanceFunctionName},
-			{TEXT("WorldBBoxMinName"), WorldBBoxMinName + ParamInfo.DataInterfaceHLSLSymbol},
-			{TEXT("WorldBBoxSizeName"), WorldBBoxSizeName + ParamInfo.DataInterfaceHLSLSymbol},
-		};
-		OutHLSL += FString::Format(FormatSample, ArgsSample);
+
+		OutHLSL += FString::Format(FormatSample, ArgsDeclarations);
 		return true;
 	}
-	else if (DefinitionFunctionName == UnitToWorldFunctionName)
+	else if (DefinitionFunctionName == UnitToSimulationFunctionName)
 	{
 		static const TCHAR *FormatSample = TEXT(R"(
-			void {FunctionName}(float3 In_Unit, out float3 Out_World)
+			void {FunctionName}(float3 In_Unit, float4x4 In_UnitToSimulationTransform, out float3 Out_Simulation)
 			{
-				Out_World = In_Unit * float3({WorldBBoxSizeName}, 1) + {WorldBBoxMinName};
+				Out_Simulation = mul(float4(In_Unit, 1.0), In_UnitToSimulationTransform).xyz;
 			}
 		)");
-		TMap<FString, FStringFormatArg> ArgsSample = {
-			{TEXT("FunctionName"), InstanceFunctionName},
-			{TEXT("WorldBBoxMinName"), WorldBBoxMinName + ParamInfo.DataInterfaceHLSLSymbol},
-			{TEXT("WorldBBoxSizeName"), WorldBBoxSizeName + ParamInfo.DataInterfaceHLSLSymbol},
-		};
-		OutHLSL += FString::Format(FormatSample, ArgsSample);
+
+		OutHLSL += FString::Format(FormatSample, ArgsDeclarations);
 		return true;
 	}
 	else if (DefinitionFunctionName == UnitToIndexFunctionName)
@@ -568,11 +594,8 @@ bool UNiagaraDataInterfaceGrid2D::GetFunctionHLSL(const FName& DefinitionFunctio
 				Out_IndexY = Out_IndexTmp.y;				
 			}
 		)");
-		TMap<FString, FStringFormatArg> ArgsSample = {
-			{TEXT("FunctionName"), InstanceFunctionName},
-			{TEXT("NumCellsName"), NumCellsName + ParamInfo.DataInterfaceHLSLSymbol},			
-		};
-		OutHLSL += FString::Format(FormatSample, ArgsSample);
+
+		OutHLSL += FString::Format(FormatSample, ArgsDeclarations);
 		return true;
 	}
 	else if (DefinitionFunctionName == IndexToUnitFunctionName)
@@ -583,11 +606,8 @@ bool UNiagaraDataInterfaceGrid2D::GetFunctionHLSL(const FName& DefinitionFunctio
 				Out_Unit = float3((float2(In_IndexX, In_IndexY) + .5) / {NumCellsName}, 0);
 			}
 		)");
-		TMap<FString, FStringFormatArg> ArgsSample = {
-			{TEXT("FunctionName"), InstanceFunctionName},
-			{TEXT("NumCellsName"), NumCellsName + ParamInfo.DataInterfaceHLSLSymbol},
-		};
-		OutHLSL += FString::Format(FormatSample, ArgsSample);
+
+		OutHLSL += FString::Format(FormatSample, ArgsDeclarations);
 		return true;
 	}
 	else if (DefinitionFunctionName == IndexToUnitStaggeredXFunctionName)
@@ -598,11 +618,8 @@ bool UNiagaraDataInterfaceGrid2D::GetFunctionHLSL(const FName& DefinitionFunctio
 				Out_Unit = float3((float2(In_IndexX, In_IndexY) + float2(0.0, 0.5)) / {NumCellsName}, 0);
 			}
 		)");
-		TMap<FString, FStringFormatArg> ArgsSample = {
-			{TEXT("FunctionName"), InstanceFunctionName},
-			{TEXT("NumCellsName"), NumCellsName + ParamInfo.DataInterfaceHLSLSymbol},
-		};
-		OutHLSL += FString::Format(FormatSample, ArgsSample);
+
+		OutHLSL += FString::Format(FormatSample, ArgsDeclarations);
 		return true;
 	}
 	else if (DefinitionFunctionName == IndexToUnitStaggeredYFunctionName)
@@ -613,11 +630,8 @@ bool UNiagaraDataInterfaceGrid2D::GetFunctionHLSL(const FName& DefinitionFunctio
 				Out_Unit = float3((float2(In_IndexX, In_IndexY) +  + float2(0.5, 0.0)) / {NumCellsName}, 0);
 			}
 		)");
-		TMap<FString, FStringFormatArg> ArgsSample = {
-			{TEXT("FunctionName"), InstanceFunctionName},
-			{TEXT("NumCellsName"), NumCellsName + ParamInfo.DataInterfaceHLSLSymbol},
-		};
-		OutHLSL += FString::Format(FormatSample, ArgsSample);
+		
+		OutHLSL += FString::Format(FormatSample, ArgsDeclarations);
 		return true;
 	}
 	else if (DefinitionFunctionName == IndexToLinearFunctionName)
@@ -628,11 +642,8 @@ bool UNiagaraDataInterfaceGrid2D::GetFunctionHLSL(const FName& DefinitionFunctio
 				Out_Linear = In_IndexX + In_IndexY * {NumCellsName}.x;
 			}
 		)");
-		TMap<FString, FStringFormatArg> ArgsSample = {
-			{TEXT("FunctionName"), InstanceFunctionName},
-			{TEXT("NumCellsName"), NumCellsName + ParamInfo.DataInterfaceHLSLSymbol},
-		};
-		OutHLSL += FString::Format(FormatSample, ArgsSample);
+
+		OutHLSL += FString::Format(FormatSample, ArgsDeclarations);
 		return true;
 	}
 	else if (DefinitionFunctionName == LinearToIndexFunctionName)
@@ -640,15 +651,12 @@ bool UNiagaraDataInterfaceGrid2D::GetFunctionHLSL(const FName& DefinitionFunctio
 		static const TCHAR *FormatSample = TEXT(R"(
 			void {FunctionName}(int In_Linear, out int Out_IndexX, out int Out_IndexY)
 			{
-				Out_IndexX = In_Linear % {NumCells}.x;
-				Out_IndexY = In_Linear / {NumCells}.x;				
+				Out_IndexX = In_Linear % {NumCellsName}.x;
+				Out_IndexY = In_Linear / {NumCellsName}.x;				
 			}
 		)");
-		TMap<FString, FStringFormatArg> ArgsSample = {
-			{TEXT("FunctionName"), InstanceFunctionName},
-			{TEXT("NumCells"), NumCellsName + ParamInfo.DataInterfaceHLSLSymbol},
-		};
-		OutHLSL += FString::Format(FormatSample, ArgsSample);
+
+		OutHLSL += FString::Format(FormatSample, ArgsDeclarations);
 		return true;
 	}
 	else if (DefinitionFunctionName == CellSizeFunctionName)
@@ -659,11 +667,8 @@ bool UNiagaraDataInterfaceGrid2D::GetFunctionHLSL(const FName& DefinitionFunctio
 				Out_CellSize = {CellSizeName};
 			}
 		)");
-		TMap<FString, FStringFormatArg> ArgsSample = {
-			{TEXT("FunctionName"), InstanceFunctionName},
-			{TEXT("CellSizeName"), CellSizeName + ParamInfo.DataInterfaceHLSLSymbol},
-		};
-		OutHLSL += FString::Format(FormatSample, ArgsSample);
+
+		OutHLSL += FString::Format(FormatSample, ArgsDeclarations);
 		return true;
 	}
 
@@ -691,8 +696,7 @@ bool UNiagaraDataInterfaceGrid2D::CopyToInternal(UNiagaraDataInterface* Destinat
 	OtherTyped->NumCellsY = NumCellsY;
 	OtherTyped->NumAttributes = NumAttributes;
 	OtherTyped->NumCellsMaxAxis = NumCellsMaxAxis;
-	OtherTyped->SetGridFromMaxAxis = SetGridFromMaxAxis;
-	OtherTyped->WorldBBoxMin = WorldBBoxMin;
+	OtherTyped->SetGridFromMaxAxis = SetGridFromMaxAxis;	
 	OtherTyped->WorldBBoxSize = WorldBBoxSize;
 
 	return true;
diff --git a/Engine/Plugins/FX/Niagara/Source/NiagaraCore/Private/NiagaraCustomVersion.cpp b/Engine/Plugins/FX/Niagara/Source/NiagaraCore/Private/NiagaraCustomVersion.cpp
--- a/Engine/Plugins/FX/Niagara/Source/NiagaraCore/Private/NiagaraCustomVersion.cpp
+++ b/Engine/Plugins/FX/Niagara/Source/NiagaraCore/Private/NiagaraCustomVersion.cpp
@@ -8,5 +8,5 @@ const FGuid FNiagaraCustomVersion::GUID(0xFCF57AFA, 0x50764283, 0xB9A9E658, 0xFF
 // Register the custom version with core
 FCustomVersionRegistration GRegisterNiagaraCustomVersion(FNiagaraCustomVersion::GUID, FNiagaraCustomVersion::LatestVersion, TEXT("NiagaraVer"));
 
-const FGuid FNiagaraCustomVersion::LatestScriptCompileVersion(0x02EFFEE7, 0x8852B128, 0x46385086, 0x84BD6F39);
+const FGuid FNiagaraCustomVersion::LatestScriptCompileVersion(0xAD86D23D, 0xA2C94302, 0xB4865706, 0x82E01289);
 
