Removed some logspam.

#rnx

diff --git a/Engine/Plugins/FX/Niagara/Source/Niagara/Private/NiagaraScript.cpp b/Engine/Plugins/FX/Niagara/Source/Niagara/Private/NiagaraScript.cpp
--- a/Engine/Plugins/FX/Niagara/Source/Niagara/Private/NiagaraScript.cpp
+++ b/Engine/Plugins/FX/Niagara/Source/Niagara/Private/NiagaraScript.cpp
@@ -538,7 +538,7 @@ void UNiagaraScript::Serialize(FArchive& Ar)
 				NumRemoved++;
 			}
 
-			UE_LOG(LogNiagara, Display, TEXT("Pruned %d/%d parameters from script %s"), NumRemoved, TemporaryStore.GetNumParameters(), *GetFullName());
+			UE_LOG(LogNiagara, Verbose, TEXT("Pruned %d/%d parameters from script %s"), NumRemoved, TemporaryStore.GetNumParameters(), *GetFullName());
 		}
 	}
 
diff --git a/Engine/Plugins/FX/Niagara/Source/NiagaraShader/Private/NiagaraShared.cpp b/Engine/Plugins/FX/Niagara/Source/NiagaraShader/Private/NiagaraShared.cpp
--- a/Engine/Plugins/FX/Niagara/Source/NiagaraShader/Private/NiagaraShared.cpp
+++ b/Engine/Plugins/FX/Niagara/Source/NiagaraShader/Private/NiagaraShared.cpp
@@ -299,7 +299,7 @@ bool FNiagaraShaderScript::CacheShaders(const FNiagaraShaderMapId& ShaderMapId,
 			FNiagaraShaderMap::LoadFromDerivedDataCache(this, ShaderMapId, Platform, GameThreadShaderMap);
 			if (GameThreadShaderMap && GameThreadShaderMap->IsValid())
 			{
-				UE_LOG(LogTemp, Display, TEXT("Loaded shader %s for Niagara script %s from DDC"), *GameThreadShaderMap->GetFriendlyName(), *GetFriendlyName());
+				UE_LOG(LogTemp, Verbose, TEXT("Loaded shader %s for Niagara script %s from DDC"), *GameThreadShaderMap->GetFriendlyName(), *GetFriendlyName());
 			}
 			else
 			{
