CIS Fix
Missing Files

#rnx

diff --git a/Engine/Plugins/FX/Niagara/Source/Niagara/Private/NiagaraWorldManager.cpp b/Engine/Plugins/FX/Niagara/Source/Niagara/Private/NiagaraWorldManager.cpp
--- a/Engine/Plugins/FX/Niagara/Source/Niagara/Private/NiagaraWorldManager.cpp
+++ b/Engine/Plugins/FX/Niagara/Source/Niagara/Private/NiagaraWorldManager.cpp
@@ -145,7 +145,7 @@ FNiagaraWorldManager::~FNiagaraWorldManager()
 	OnWorldCleanup(true, true);
 }
 
-FNiagaraWorldManager* FNiagaraWorldManager::Get(UWorld* World)
+FNiagaraWorldManager* FNiagaraWorldManager::Get(const UWorld* World)
 {
 	FNiagaraWorldManager** OutWorld = WorldManagers.Find(World);
 	if (OutWorld == nullptr)
diff --git a/Engine/Plugins/FX/Niagara/Source/Niagara/Public/NiagaraWorldManager.h b/Engine/Plugins/FX/Niagara/Source/Niagara/Public/NiagaraWorldManager.h
--- a/Engine/Plugins/FX/Niagara/Source/Niagara/Public/NiagaraWorldManager.h
+++ b/Engine/Plugins/FX/Niagara/Source/Niagara/Public/NiagaraWorldManager.h
@@ -96,7 +96,7 @@ public:
 	FNiagaraWorldManager(UWorld* InWorld);
 	~FNiagaraWorldManager();
 
-	static NIAGARA_API FNiagaraWorldManager* Get(UWorld* World);
+	static NIAGARA_API FNiagaraWorldManager* Get(const UWorld* World);
 	static void OnStartup();
 	static void OnShutdown();
 
