Added a shared ref count to track the number of GPU emitters

diff --git a/Engine/Plugins/FX/Niagara/Source/Niagara/Classes/NiagaraGPUInstanceCountManager.h b/Engine/Plugins/FX/Niagara/Source/Niagara/Classes/NiagaraGPUInstanceCountManager.h
--- a/Engine/Plugins/FX/Niagara/Source/Niagara/Classes/NiagaraGPUInstanceCountManager.h
+++ b/Engine/Plugins/FX/Niagara/Source/Niagara/Classes/NiagaraGPUInstanceCountManager.h
@@ -10,6 +10,7 @@ NiagaraGPUInstanceCountManager.h: GPU particle count handling
 #include "RHIUtilities.h"
 #include "NiagaraDrawIndirect.h"
 #include "RHIGPUReadback.h"
+#include "NiagaraRenderer.h"
 
 class FRHIGPUMemoryReadback;
 
diff --git a/Engine/Plugins/FX/Niagara/Source/Niagara/Public/NiagaraRenderer.h b/Engine/Plugins/FX/Niagara/Source/Niagara/Public/NiagaraRenderer.h
--- a/Engine/Plugins/FX/Niagara/Source/Niagara/Public/NiagaraRenderer.h
+++ b/Engine/Plugins/FX/Niagara/Source/Niagara/Public/NiagaraRenderer.h
@@ -78,6 +78,14 @@ struct FNiagaraRendererVariableInfo
 	bool bUpload;
 };
 
+// The number of GPU renderers registered in the instance count manager.
+// Shared between the manager and the renderers.
+class FNiagaraGPURendererCount : public FRefCountedObject
+{
+public:
+	int32 Value = 0;
+};
+
 /**
 * Base class for Niagara System renderers.
 */
