Refactored Niagara GPU particles simulation aroun a few new functions :
- InitConstants() : init constant accessed in the spawn and update logic.
- InitSpawnVariables() : init properties for newly spawned particles.
- LoadUpdateVariables() : Load properties from last simulation result.
- ConditionalInterpolateParameters() : call InterpolateParameters() if required.
- TransferAttibutes() : transfer data from MapSpawn to MapUpdate.
- StoreUpdateVariables() : acquire a write index and store the resulting properties.
- CopyInstance() : not yet implemented on GPU particles.

Implemented a non intrinsics fallback for acquiring destination indices using group shared.

Niagara GPU simulation shader THREADGROUP_SIZE is now platform dependent to optimize group shared syncs.

diff --git a/Engine/Plugins/FX/Niagara/Shaders/Private/NiagaraEmitterInstanceShader.usf b/Engine/Plugins/FX/Niagara/Shaders/Private/NiagaraEmitterInstanceShader.usf
--- a/Engine/Plugins/FX/Niagara/Shaders/Private/NiagaraEmitterInstanceShader.usf
+++ b/Engine/Plugins/FX/Niagara/Shaders/Private/NiagaraEmitterInstanceShader.usf
@@ -37,13 +37,12 @@ float Reciprocal(float x){ return 1.0f/x; }
  */
 #if GPU_SIMULATION
 
-	#define THREADGROUP_SIZE 64
-
 	static uint3 GDispatchThreadId;
 	static uint3 GGroupThreadId;
 	static uint GCurrentPhase;
 	static uint GEmitterTickCounter;
 	static uint GSimStart;
+	static uint GRandomSeedOffset = 0;
 
 	// To be removed as well but don't know who is using it
 	#include "/Engine/Private/SceneTexturesCommon.ush"
@@ -382,6 +381,22 @@ int rand(int x, int Seed1, int Seed2, int Seed3)
 		}
 	}
 
+#define USE_WAVE_INTRISICS (PS4_PROFILE || XBOXONE_PROFILE)
+#define USE_GROUP_SHARED ((THREADGROUP_SIZE == 64 || THREADGROUP_SIZE == 32)  && !USE_WAVE_INTRISICS)
+
+#if USE_GROUP_SHARED
+	#if THREADGROUP_SIZE == 64
+		groupshared uint GroupSharedIndex[64];
+		groupshared uint GroupSharedIndex4[16];
+		groupshared uint GroupSharedIndex16[4];
+		groupshared uint GroupSharedIndex64;
+	#elif THREADGROUP_SIZE == 32
+		groupshared uint GroupSharedIndex[32];
+		groupshared uint GroupSharedIndex4[8];
+		groupshared uint GroupSharedIndex16[2];
+		groupshared uint GroupSharedIndex64;
+	#endif
+#endif // USE_GROUP_SHARED
 
 	/* Acquire an output index - the default index is the scratch instance; one additional instance is allocated 
 	 *	at the end of the buffer, so no branching on -1 is necessary during OutputData operations
@@ -394,48 +409,16 @@ int rand(int x, int Seed1, int Seed2, int Seed3)
 		// End static assert
 
 		#if USE_SHADER_STAGES
-			if(ShaderStageIndex == 0)
+			if(ShaderStageIndex != 0)
 			{
-				int PrevIdx = GSpawnStartInstance + SpawnedInstances;	// scratch instance as default; write to that for dead particles
-
-#if PS4_PROFILE || XBOXONE_PROFILE
-				uint NumCounters = WaveActiveCountBits(bDoAcquire);
-				uint PrefixCounts = WavePrefixCountBits(bDoAcquire);
-				if(NumCounters)
-				{
-					if(WaveIsFirstLane())
-					{
-						uint RetPrevIdx;
-						InterlockedAdd(RWInstanceCounts[WriteInstanceCountOffset], NumCounters, RetPrevIdx);
-						PrevIdx = (int)RetPrevIdx;
-					}
-
-					if(bDoAcquire)
-					{
-						PrevIdx = WaveReadFirstLane(PrevIdx);
-						PrevIdx += PrefixCounts;
-					}
-				}		
-#else
-				if(bDoAcquire == true)
-				{
-					// Have to use uint's here to avoid PS4 compiler warnings about InterlockedAdd, cannot propagate uint due to CPU VM limitations...
-					uint RetPrevIdx;
-					// @TODO : add some TLS logic to avoid thread group for doing atomic for each thread. (gathering the actual required count)
-					InterlockedAdd(RWInstanceCounts[WriteInstanceCountOffset], (uint)1U, RetPrevIdx);
-					PrevIdx = (int)RetPrevIdx;
-				}
-#endif
-				return PrevIdx;
+				bDoAcquire = false; // Prevent any change to the atomic
 			}
-			else
-			{
-				return ExecIndex();
-			}
-		#else
-			int PrevIdx = GSpawnStartInstance + SpawnedInstances;	// scratch instance as default; write to that for dead particles
+		#endif // USE_SHADER_STAGES
+
+		int PrevIdx = GSpawnStartInstance + SpawnedInstances;	// scratch instance as default; write to that for dead particles
+
+		#if USE_WAVE_INTRISICS
 
-#if PS4_PROFILE || XBOXONE_PROFILE
 			uint NumCounters = WaveActiveCountBits(bDoAcquire);
 			uint PrefixCounts = WavePrefixCountBits(bDoAcquire);
 			if(NumCounters)
@@ -453,7 +436,77 @@ int rand(int x, int Seed1, int Seed2, int Seed3)
 					PrevIdx += PrefixCounts;
 				}
 			}	
-#else
+
+		#elif USE_GROUP_SHARED
+
+			GroupSharedIndex[GGroupThreadId.x] = bDoAcquire ? 1 : 0;
+			GroupMemoryBarrierWithGroupSync(); 
+
+			// Group by 4
+			if ((GGroupThreadId.x & 0x3) == 0)
+			{
+				const uint Index = GGroupThreadId.x;
+
+				const uint ActiveCount1 = GroupSharedIndex[Index];
+				const uint ActiveCount2 = ActiveCount1 + GroupSharedIndex[Index + 1];
+				const uint ActiveCount3 = ActiveCount2 + GroupSharedIndex[Index + 2];
+				const uint ActiveCount4 = ActiveCount3 + GroupSharedIndex[Index + 3];
+				
+				GroupSharedIndex[Index] = 0;
+				GroupSharedIndex[Index + 1] = ActiveCount1;
+				GroupSharedIndex[Index + 2] = ActiveCount2;
+				GroupSharedIndex[Index + 3] = ActiveCount3;
+				GroupSharedIndex4[Index / 4] = ActiveCount4;
+			}
+			GroupMemoryBarrierWithGroupSync(); 
+
+			// Group by 16
+			if ((GGroupThreadId.x & 0xF) == 0)
+			{
+				const uint Index = GGroupThreadId.x / 4;
+
+				const uint ActiveCount1 = GroupSharedIndex4[Index];
+				const uint ActiveCount2 = ActiveCount1 + GroupSharedIndex4[Index + 1];
+				const uint ActiveCount3 = ActiveCount2 + GroupSharedIndex4[Index + 2];
+				const uint ActiveCount4 = ActiveCount3 + GroupSharedIndex4[Index + 3];
+
+				GroupSharedIndex4[Index] = 0;
+				GroupSharedIndex4[Index + 1] = ActiveCount1;
+				GroupSharedIndex4[Index + 2] = ActiveCount2;
+				GroupSharedIndex4[Index + 3] = ActiveCount3;
+				GroupSharedIndex16[Index / 4] = ActiveCount4;
+			}
+			GroupMemoryBarrierWithGroupSync(); 
+
+			// Group by 64
+			if (GGroupThreadId.x == 0)
+			{
+				const uint ActiveCount1 = GroupSharedIndex16[0];
+				const uint ActiveCount2 = ActiveCount1 + GroupSharedIndex16[1];
+		#if THREADGROUP_SIZE == 64
+				const uint ActiveCount3 = ActiveCount2 + GroupSharedIndex16[2];
+				const uint ActiveCount4 = ActiveCount3 + GroupSharedIndex16[3];
+		#endif
+				GroupSharedIndex16[0] = 0;
+				GroupSharedIndex16[1] = ActiveCount1;
+		#if THREADGROUP_SIZE == 64
+				GroupSharedIndex16[2] = ActiveCount2;
+				GroupSharedIndex16[3] = ActiveCount3;
+		#endif
+				uint RetPrevIdx = 0;
+		#if THREADGROUP_SIZE == 64
+				InterlockedAdd(RWInstanceCounts[WriteInstanceCountOffset], ActiveCount4, RetPrevIdx);
+		#elif THREADGROUP_SIZE == 32
+				InterlockedAdd(RWInstanceCounts[WriteInstanceCountOffset], ActiveCount2, RetPrevIdx);
+		#endif
+				GroupSharedIndex64 = RetPrevIdx;
+			}
+			GroupMemoryBarrierWithGroupSync(); 
+
+			PrevIdx = GroupSharedIndex64 + GroupSharedIndex16[GGroupThreadId.x / 16] + GroupSharedIndex4[GGroupThreadId.x / 4] + GroupSharedIndex[GGroupThreadId.x];
+
+		#else // !USE_WAVE_INTRISICS && !USE_GROUP_SHARED
+
 			if(bDoAcquire == true)
 			{
 				// Have to use uint's here to avoid PS4 compiler warnings about InterlockedAdd, cannot propagate uint due to CPU VM limitations...
@@ -462,9 +515,17 @@ int rand(int x, int Seed1, int Seed2, int Seed3)
 				InterlockedAdd(RWInstanceCounts[WriteInstanceCountOffset], (uint)1U, RetPrevIdx);
 				PrevIdx = (int)RetPrevIdx;
 			}
-#endif
-			return PrevIdx;
-		#endif
+
+		#endif // USE_WAVE_INTRISICS || USE_GROUP_SHARED
+
+		#if USE_SHADER_STAGES
+			if(ShaderStageIndex != 0)
+			{
+				return ExecIndex();
+			}
+		#endif // USE_SHADER_STAGES
+
+		return PrevIdx;
 	}
 
 	/* ---------------------------------------------------------------------
@@ -508,7 +569,7 @@ int rand(int x, int Seed1, int Seed2, int Seed3)
 
 	void EnterStatScope(int ID)	{}
 	void ExitStatScope()	{}
-#endif
+#endif // GPU_SIMULATION
 
 	/* 
 	 * Get the index to write onto the output buffer
@@ -521,7 +582,7 @@ int rand(int x, int Seed1, int Seed2, int Seed3)
 			return (bIsPersistent && (DenseIndex != SparseIndex) ) ? SparseIndex : DenseIndex;
 		#else
 			return AcquireIndex(DataSetID,bIsValid);
-		#endif
+		#endif // USE_SHADER_STAGES
 	}
 
 //Include the simulation shader code generated by the node graph.
@@ -542,9 +603,11 @@ void SimulateMainComputeCS(
 {
 	GDispatchThreadId = DispatchThreadId;
 	GGroupThreadId = GroupThreadId;
+	GCurrentPhase = -1;
 	GEmitterTickCounter = EmitterTickCounter;
 	GSimStart = SimStart;
 	GRandomSeedOffset = 0;
+
 	/*
 	if(CopyInstancesBeforeStart == 1)
 	{
@@ -552,7 +615,7 @@ void SimulateMainComputeCS(
 	}
 	*/
 
-	uint InstanceID = UpdateStartInstance + DispatchThreadId.x;
+	const uint InstanceID = UpdateStartInstance + DispatchThreadId.x;
 	if (ReadInstanceCountOffset == 0xFFFFFFFF)
 	{
 		GSpawnStartInstance = 0;
@@ -561,59 +624,59 @@ void SimulateMainComputeCS(
 	{
 		GSpawnStartInstance = RWInstanceCounts[ReadInstanceCountOffset];				// needed by ExecIndex()
 	}
-	int MaxInstances = GSpawnStartInstance + SpawnedInstances;
+	const int MaxInstances = GSpawnStartInstance + SpawnedInstances;
 
-	bool UseInterface = false;
+	bool bRunPhase1 = InstanceID < GSpawnStartInstance && InstanceID < UpdateStartInstance + MaxInstances;
+	bool bRunPhase0 = InstanceID >= GSpawnStartInstance && InstanceID < UpdateStartInstance + MaxInstances;
+	
 	#if USE_SHADER_STAGES
 	if (IterationInterfaceCount > 0)
 	{
-		if (InstanceID < IterationInterfaceCount)
-		{
-			if (GSimStart == 1)
-			{
-				SimulateMain(InstanceID, 0, 0);
-			}
-			else
-			{
-				SimulateMain(InstanceID, 0, 1);
-			}
-		}
-		UseInterface = true;
+		bRunPhase1 = InstanceID < IterationInterfaceCount && GSimStart != 1;
+		bRunPhase0 = InstanceID < IterationInterfaceCount && GSimStart == 1;
 	}
-	#endif
-	if( !UseInterface && InstanceID < UpdateStartInstance + MaxInstances )
+	#endif // USE_SHADER_STAGES
+
+	const float RandomSeedInitialisation = NiagaraInternalNoise(InstanceID*16384, 0 * 8196, (bRunPhase1 ? 4096 : 0) + EmitterTickCounter);	// initialise the random state seed
+
+	FSimulationContext Context = (FSimulationContext)0;
+	InitConstants(Context);
+
+	BRANCH
+	if (bRunPhase1)
 	{
-		// calculate event index for the current instance; this is primarily for spawn events, where we do a single read from a secondary DataSet for a subset of particles
-		//  for all-particle events, this will always end up being 0 because NumParticlesPerEvent = total particles in the set
-		uint EventIndex = 0; //(InstanceID-SimulateStartInstance) / NumParticlesPerEvent;
+		GCurrentPhase = 1;
 
-		// before the simulation start - either update or copy
-		//
-		if(InstanceID < GSpawnStartInstance)
-		{
-			// copy instances before simulation start inst (e.g. for event spawns)
-			/*
-			if(CopyInstancesBeforeStart)
-			{
-				CopyInstance(InstanceID);
-			}
-			else
-			*/
-			{
-				GCurrentPhase = 1;
-				SimulateMain(InstanceID, EventIndex, 1);
-			}
+		LoadUpdateVariables(Context, InstanceID);
+		ReadDataSets(Context);
 
-		}
+		SimulateMapUpdate(Context);
+		WriteDataSets(Context);
+	}
+	else if (bRunPhase0)
+	{
+		GCurrentPhase = 0;
 
-		// after simulation start - spawn phase
-		//
-		else if(InstanceID >= GSpawnStartInstance)
-		{
-			GCurrentPhase = 0;
-			SimulateMain(InstanceID, EventIndex, 0);		// SimulateMain spawn phase
-		}
+		GetEmitterSpawnInfoForParticle(ExecIndex());
+		// Part of InitConstants() but updated by the previous statement
+		Context.MapSpawn.Emitter.SpawnGroup = Emitter_SpawnGroup;
+
+		InitSpawnVariables(Context);
+		ReadDataSets(Context);
+
+		Context.MapSpawn.Particles.UniqueID = Engine_Emitter_TotalSpawnedParticles + ExecIndex();
+		ConditionalInterpolateParameters(Context);
+		SimulateMapSpawn(Context);
+
+		GCurrentPhase = 1;
+
+		TransferAttributes(Context);
+
+		SimulateMapUpdate(Context);
+		WriteDataSets(Context);
 	}
+
+	StoreUpdateVariables(Context);
 }
 
 #endif//GPU_SIMULATION
diff --git a/Engine/Plugins/FX/Niagara/Source/Niagara/Private/NiagaraEmitterInstanceBatcher.cpp b/Engine/Plugins/FX/Niagara/Source/Niagara/Private/NiagaraEmitterInstanceBatcher.cpp
--- a/Engine/Plugins/FX/Niagara/Source/Niagara/Private/NiagaraEmitterInstanceBatcher.cpp
+++ b/Engine/Plugins/FX/Niagara/Source/Niagara/Private/NiagaraEmitterInstanceBatcher.cpp
@@ -1186,18 +1186,19 @@ void NiagaraEmitterInstanceBatcher::Run(const FNiagaraGPUSystemTick& Tick, const
 		RHICmdList.SetShaderParameter(Shader->GetComputeShader(), Shader->IterationInterfaceCount.GetBufferIndex(), Shader->IterationInterfaceCount.GetBaseIndex(), Shader->IterationInterfaceCount.GetNumBytes(), &DefaultIterationCount);					// 0, except if several stages are defined
 	}
 
+	const uint32 ShaderThreadGroupSize = FNiagaraShader::GetGroupSize(ShaderPlatform);
 	if (IterationInterface)
 	{
-		if (TotalNumInstances > NIAGARA_COMPUTE_THREADGROUP_SIZE)
+		if (TotalNumInstances > ShaderThreadGroupSize)
 		{
 			RHICmdList.SetShaderParameter(Shader->GetComputeShader(), Shader->IterationInterfaceCount.GetBufferIndex(), Shader->IterationInterfaceCount.GetBaseIndex(), Shader->IterationInterfaceCount.GetNumBytes(), &TotalNumInstances);					// 0, except if several stages are defined
 		}
 	}
 
 	uint32 NumThreadGroups = 1;
-	if (TotalNumInstances > NIAGARA_COMPUTE_THREADGROUP_SIZE)
+	if (TotalNumInstances > ShaderThreadGroupSize)
 	{
-		NumThreadGroups = FMath::Min(NIAGARA_MAX_COMPUTE_THREADGROUPS, FMath::DivideAndRoundUp(TotalNumInstances, NIAGARA_COMPUTE_THREADGROUP_SIZE));
+		NumThreadGroups = FMath::Min(NIAGARA_MAX_COMPUTE_THREADGROUPS, FMath::DivideAndRoundUp(TotalNumInstances, ShaderThreadGroupSize));
 	}
 
 	// setup script parameters
diff --git a/Engine/Plugins/FX/Niagara/Source/NiagaraCore/Private/NiagaraCustomVersion.cpp b/Engine/Plugins/FX/Niagara/Source/NiagaraCore/Private/NiagaraCustomVersion.cpp
--- a/Engine/Plugins/FX/Niagara/Source/NiagaraCore/Private/NiagaraCustomVersion.cpp
+++ b/Engine/Plugins/FX/Niagara/Source/NiagaraCore/Private/NiagaraCustomVersion.cpp
@@ -8,5 +8,4 @@ const FGuid FNiagaraCustomVersion::GUID(0xFCF57AFA, 0x50764283, 0xB9A9E658, 0xFF
 // Register the custom version with core
 FCustomVersionRegistration GRegisterNiagaraCustomVersion(FNiagaraCustomVersion::GUID, FNiagaraCustomVersion::LatestVersion, TEXT("NiagaraVer"));
 
-const FGuid FNiagaraCustomVersion::LatestScriptCompileVersion(0xAD86D23D, 0xA2C94302, 0xB4865706, 0x82E01289);
-
+const FGuid FNiagaraCustomVersion::LatestScriptCompileVersion(0xBCB9A1E7, 0xB6F54078, 0xBC9895E2, 0x1D5D4C0F);
\ No newline at end of file
diff --git a/Engine/Plugins/FX/Niagara/Source/NiagaraEditor/Private/NiagaraHlslTranslator.cpp b/Engine/Plugins/FX/Niagara/Source/NiagaraEditor/Private/NiagaraHlslTranslator.cpp
--- a/Engine/Plugins/FX/Niagara/Source/NiagaraEditor/Private/NiagaraHlslTranslator.cpp
+++ b/Engine/Plugins/FX/Niagara/Source/NiagaraEditor/Private/NiagaraHlslTranslator.cpp
@@ -1328,7 +1328,14 @@ const FNiagaraTranslateResults &FHlslNiagaraTranslator::Translate(const FNiagara
 
 		//And finally, define the actual main function that handles the reading and writing of data and calls the shared per instance simulate function.
 		//TODO: Different wrappers for GPU and CPU sims. 
-		DefineMain(HlslOutput, DataSetVariables, DataSetReads, DataSetWrites);
+		if (TranslationOptions.SimTarget == ENiagaraSimTarget::GPUComputeSim)
+		{
+			DefineMainGPUFunctions(DataSetVariables, DataSetReads, DataSetWrites);
+		}
+		else
+		{
+			DefineMain(HlslOutput, DataSetVariables, DataSetReads, DataSetWrites);
+		}
 
 		//Get full list of instance data accessed by the script as the VM binding assumes same for input and output.
 		for (FNiagaraVariable& Var : DataSetVariables[InstanceReadVarsIndex])
@@ -1793,28 +1800,245 @@ void FHlslNiagaraTranslator::DefineExternalFunctionsHLSL(FString &InHlslOutput)
 	}
 }
 
-void FHlslNiagaraTranslator::DefineMain(FString &OutHlslOutput,
+void FHlslNiagaraTranslator::DefineMainGPUFunctions(
 	const TArray<TArray<FNiagaraVariable>>& DataSetVariables,
 	const TMap<FNiagaraDataSetID, int32>& DataSetReads,
 	const TMap<FNiagaraDataSetID, int32>& DataSetWrites)
 {
-	if (CompilationTarget == ENiagaraSimTarget::GPUComputeSim)
+	static const auto UseShaderStagesCVar = IConsoleManager::Get().FindConsoleVariable(TEXT("fx.UseShaderStages"));
+	const bool bUseShaderStages = UseShaderStagesCVar->GetInt() == 1;
+ 
+	TArray<FNiagaraDataSetID> ReadDataSetIDs;
+	TArray<FNiagaraDataSetID> WriteDataSetIDs;
+	
+	DataSetReads.GetKeys(ReadDataSetIDs);
+	DataSetWrites.GetKeys(WriteDataSetIDs);
+
+	// Whether Alive is used and must be set at each run
+	const bool bUsesAlive = [&]
 	{
-		OutHlslOutput += TEXT("void SimulateMain(in int InstanceIdx, in int InEventIndex, in int Phase)\n{\n");
+		TArray<FName> DataSetNames;
+		for (const FNiagaraDataSetID& ReadId : ReadDataSetIDs)
+		{
+			DataSetNames.AddUnique(ReadId.Name);
+		}
+		for (const FNiagaraDataSetID& WriteId : WriteDataSetIDs)
+		{
+			DataSetNames.AddUnique(WriteId.Name);
+		}
+		for (int32 i = 0; i < ParamMapHistories.Num(); i++)
+		{
+			for (FName DataSetName : DataSetNames)
+			{
+				if (ParamMapHistories[i].FindVariable(*(DataSetName.ToString() + TEXT(".Alive")), FNiagaraTypeDefinition::GetBoolDef()) != INDEX_NONE)
+				{
+					return true;
+				}
+			}
+		}
+		return false;
+	}();
 
-		// MASSIVE HACK
-		HlslOutput += TEXT(R"(
-			float RandomSeedInitialisation = NiagaraInternalNoise(InstanceIdx*16384, InEventIndex*8196, Phase*4096 + EmitterTickCounter);	// initialise the random state seed
-		)");
-		
+	///////////////////////
+	// InitConstants()
+	HlslOutput += TEXT("void InitConstants(inout FSimulationContext Context)\n{\n");
+	{
+		// Fill in the defaults for parameters.
+		for (int32 i = 0; i < MainPreSimulateChunks.Num(); ++i)
+		{
+			HlslOutput += TEXT("\t") + MainPreSimulateChunks[i] + TEXT("\n");
+		}
 	}
-	else
+	HlslOutput += TEXT("}\n\n");
+
+	///////////////////////
+	// InitSpawnVariables()
+	HlslOutput += TEXT("void InitSpawnVariables(inout FSimulationContext Context)\n{\n");
 	{
-		OutHlslOutput += TEXT("void SimulateMain()\n{\n");
+		FString ContextName = TEXT("\tContext.Map.");
+		if (TranslationStages.Num() > 1) // First context 0 is "MapSpawn"
+		{
+			ContextName = FString::Printf(TEXT("\tContext.%s."), *TranslationStages[0].PassNamespace);
+		}
+
+		//The VM register binding assumes the same inputs as outputs which is obviously not always the case.
+		for (int32 DataSetIndex = 0, IntCounter = 0, FloatCounter = 0; DataSetIndex < DataSetReads.Num(); ++DataSetIndex)
+		{
+			const FNiagaraDataSetID DataSetID = ReadDataSetIDs[DataSetIndex];
+			const TArray<FNiagaraVariable>& NiagaraVariables = DataSetVariables[DataSetReads[DataSetID]];
+			for (const FNiagaraVariable& Var : NiagaraVariables)
+			{
+				FString VarFmt = ContextName + GetSanitizedSymbolName(Var.GetName().ToString()) + TEXT("{0} = {4};\n");
+				GatherVariableForDataSetAccess(Var, VarFmt, IntCounter, FloatCounter, DataSetIndex, TEXT(""), HlslOutput);
+			}
+		}
+
+		if (bUsesAlive)
+		{
+			HlslOutput += TEXT("\n") + ContextName + TEXT("DataInstance.Alive=true;\n");
+		}
 	}
+	HlslOutput += TEXT("}\n\n");
 
-	EnterStatsScope(FNiagaraStatScope(*(CompileOptions.GetName() + TEXT("_Main")), TEXT("Main")), OutHlslOutput);
+	////////////////////////
+	// LoadUpdateVariables()
+	HlslOutput += TEXT("void LoadUpdateVariables(inout FSimulationContext Context, int InstanceIdx)\n{\n");
+	{
+		FString ContextName = TEXT("\tContext.Map.");
+		if (TranslationStages.Num() > 1) // Last context is "MapUpdate"
+		{
+			ContextName = FString::Printf(TEXT("\tContext.%s."), *TranslationStages.Last().PassNamespace);
+		}
+
+		for (int32 DataSetIndex = 0, IntCounter = 0, FloatCounter = 0; DataSetIndex < DataSetReads.Num(); ++DataSetIndex)
+		{
+			const FNiagaraDataSetID DataSetID = ReadDataSetIDs[DataSetIndex];
+			const TArray<FNiagaraVariable>& NiagaraVariables = DataSetVariables[DataSetReads[DataSetID]];
+			for (const FNiagaraVariable& Var : NiagaraVariables)
+			{
+				const FString VarName = ContextName + GetSanitizedSymbolName(Var.GetName().ToString());
+				FString VarFmt;
+
+				// If the NiagaraClearEachFrame value is set on the data set, we don't bother reading it in each frame as we know that it is is invalid. However,
+				// this is only used for the base data set. Other reads are potentially from events and are therefore perfectly valid.
+				if (DataSetIndex == 0 && Var.GetType().GetScriptStruct() != nullptr && Var.GetType().GetScriptStruct()->GetMetaData(TEXT("NiagaraClearEachFrame")).Equals(TEXT("true"), ESearchCase::IgnoreCase))
+				{
+					VarFmt = VarName + TEXT("{0} = {4};\n");
+				}
+				else
+				{
+					VarFmt = VarName + TEXT("{0} = InputData{1}({2}, {3}, InstanceIdx);\n");
+
+					if (bUseShaderStages)
+					{
+						if (FNiagaraParameterMapHistory::IsAttribute(Var))
+						{
+							FString RegisterName = VarName;
+							RegisterName.ReplaceInline(PARAM_MAP_ATTRIBUTE_STR, PARAM_MAP_INDICES_STR);
+							const int32 RegisterValue = Var.GetType().IsFloatPrimitive() ? FloatCounter : IntCounter;
+							HlslOutput += RegisterName + FString::Printf(TEXT(" = %d;\n"), RegisterValue);
+						}
+					}
+				}
+				GatherVariableForDataSetAccess(Var, VarFmt, IntCounter, FloatCounter, DataSetIndex, TEXT(""), HlslOutput);
+			}
+		}
+		if (bUsesAlive)
+		{
+			HlslOutput += TEXT("\n") + ContextName + TEXT("DataInstance.Alive=true;\n");
+		}
+	}
+	HlslOutput += TEXT("}\n\n");
+
+	/////////////////////////////////////
+	// ConditionalInterpolateParameters()
+	HlslOutput += TEXT("void ConditionalInterpolateParameters(inout FSimulationContext Context)\n{\n");
+	{
+		if (RequiresInterpolation())
+		{
+			HlslOutput += TEXT("\tInterpolateParameters(Context);\n"); // Requires ExecIndex, which needs to be in a stage.
+		}
+	}
+	HlslOutput += TEXT("}\n\n");
+
+	//////////////////////
+	// TransferAttibutes()
+	HlslOutput += TEXT("void TransferAttributes(inout FSimulationContext Context)\n{\n");
+	{
+		if (TranslationStages.Last().bCopyPreviousParams)
+		{
+			if (ParamMapDefinedAttributesToNamespaceVars.Num() != 0)
+			{
+				HlslOutput += TEXT("\tContext.") + TranslationStages.Last().PassNamespace + TEXT(".Particles = Context.") + TranslationStages[0].PassNamespace + TEXT(".Particles;\n");
+				if (bUsesAlive)
+				{
+					HlslOutput += TEXT("\tContext.") + TranslationStages.Last().PassNamespace + TEXT(".DataInstance = Context.") + TranslationStages[0].PassNamespace + TEXT(".DataInstance;\n");
+				}
+			}
+		}
+	}
+	HlslOutput += TEXT("}\n\n");
+
+	/////////////////////////
+	// StoreUpdateVariables()
+	HlslOutput += TEXT("void StoreUpdateVariables(in FSimulationContext Context)\n{\n");
+	{
+		if (bUsesAlive)
+		{
+			HlslOutput += TEXT("\tconst bool bValid = Context.") + TranslationStages.Last().PassNamespace + TEXT(".DataInstance.Alive;\n");
+			HlslOutput += TEXT("\tconst int WriteIndex = OutputIndex(0, false, bValid);\n");
+		}
+		else
+		{
+			HlslOutput += TEXT("\tconst bool bValid = GCurrentPhase != -1;\n");
+			HlslOutput += TEXT("\tconst int WriteIndex = OutputIndex(0, true, bValid);\n");
+		}
+
+		HlslOutput += TEXT("\tif (bValid)\n\t{\n");
+		FString ContextName = TEXT("Context.Map.");
+		if (TranslationStages.Num() > 1) // Last context is "MapUpdate"
+		{
+			ContextName = FString::Printf(TEXT("Context.%s."), *TranslationStages.Last().PassNamespace);
+		}
+
+		for (int32 DataSetIndex = 0, IntCounter = 0, FloatCounter = 0; DataSetIndex < DataSetWrites.Num(); ++DataSetIndex)
+		{
+			const FNiagaraDataSetID DataSetID = ReadDataSetIDs[DataSetIndex];
+			const TArray<FNiagaraVariable>& NiagaraVariables = DataSetVariables[DataSetWrites[DataSetID]];
+			for (const FNiagaraVariable& Var : NiagaraVariables)
+			{
+				// If coming from a parameter map, use the one on the context, otherwise use the output.
+				FString VarFmt = TEXT("\t\tOutputData{1}(0, {2}, {3}, ") + ContextName + GetSanitizedSymbolName(Var.GetName().ToString()) + TEXT("{0});\n");
+				GatherVariableForDataSetAccess(Var, VarFmt, IntCounter, FloatCounter, -1, TEXT("WriteIndex"), HlslOutput);
+			}
+		}
+	}
+	HlslOutput += TEXT("\t}\n}\n\n");
+
+	/////////////////////////
+	// CopyInstance()
+	HlslOutput += TEXT("void CopyInstance(in int InstanceIdx)\n{\n");
+	{
+#if 0
+		HlslOutput += TEXT("\tFSimulationContext Context = (FSimulationContext)0;\n");
+		if (UNiagaraScript::IsParticleEventScript(CompileOptions.TargetUsage))
+		{
+			for (int32 VarArrayIdx = 0; VarArrayIdx < DataSetReads.Num(); VarArrayIdx++)
+			{
+				const FNiagaraDataSetID DataSetID = ReadDataSetIDs[VarArrayIdx];
+				const TArray<FNiagaraVariable>& ArrayRef = DataSetVariables[DataSetReads[DataSetID]];
+				DefineDataSetVariableReads(HlslOutput, DataSetID, VarArrayIdx, ArrayRef);
+			}
+
+			if (bGpuUsesAlive)
+			{
+				HlslOutput += TEXT("\tContext.Map.DataInstance.Alive = true;\n");
+			}
 
+			for (int32 VarArrayIdx = 0; VarArrayIdx < DataSetWrites.Num(); VarArrayIdx++)
+			{
+				const FNiagaraDataSetID DataSetID = WriteDataSetIDs[VarArrayIdx];
+				const TArray<FNiagaraVariable>& ArrayRef = DataSetVariables[DataSetWrites[DataSetID]];
+				DefineDataSetVariableWrites(HlslOutput, DataSetID, VarArrayIdx, ArrayRef);
+			}
+		}
+#else
+		HlslOutput += TEXT("\t// TODO!\n");
+#endif
+	}
+	HlslOutput += TEXT("}\n");
+}
+
+void FHlslNiagaraTranslator::DefineMain(FString &OutHlslOutput,
+	const TArray<TArray<FNiagaraVariable>>& DataSetVariables,
+	const TMap<FNiagaraDataSetID, int32>& DataSetReads,
+	const TMap<FNiagaraDataSetID, int32>& DataSetWrites)
+{
+	check(CompilationTarget != ENiagaraSimTarget::GPUComputeSim);
+
+	OutHlslOutput += TEXT("void SimulateMain()\n{\n");
+	
+	EnterStatsScope(FNiagaraStatScope(*(CompileOptions.GetName() + TEXT("_Main")), TEXT("Main")), OutHlslOutput);
 
 	OutHlslOutput += TEXT("\n\tFSimulationContext Context = (FSimulationContext)0;\n");
 	TMap<FName, int32> InputRegisterAllocations;
@@ -1826,15 +2050,6 @@ void FHlslNiagaraTranslator::DefineMain(FString &OutHlslOutput,
 	//TODO: Grab indices for reading data sets and do the read.
 	//read input.
 
-	// If we are GPU we need to set the spawn info global variables
-	if (CompilationTarget == ENiagaraSimTarget::GPUComputeSim)
-	{
-		OutHlslOutput += TEXT("\tif (Phase == 0)\n");
-		OutHlslOutput += TEXT("\t{\n");
-		OutHlslOutput += TEXT("\t\tGetEmitterSpawnInfoForParticle(ExecIndex());\n");
-		OutHlslOutput += TEXT("\t}\n");
-	}
-
 	TArray<FNiagaraDataSetID> ReadDataSetIDs;
 	TArray<FNiagaraDataSetID> WriteDataSetIDs;
 
@@ -1886,118 +2101,30 @@ void FHlslNiagaraTranslator::DefineMain(FString &OutHlslOutput,
 		OutHlslOutput += TEXT("\t") + MainPreSimulateChunks[i] + TEXT("\n");
 	}
 
-	bool bGpuUsesAlive = false;
+	// call the read data set function
+	OutHlslOutput += TEXT("\tReadDataSets(Context);\n");
+	for (int32 StageIdx = 0; StageIdx < TranslationStages.Num(); StageIdx++)
 	{
-		// call the read data set function
-		OutHlslOutput += TEXT("\tReadDataSets(Context);\n");
-
-		// branch between spawn and update
-		if (CompilationTarget == ENiagaraSimTarget::GPUComputeSim)
+		if (StageIdx == 0)
 		{
-			// Determine 
-			TArray<FName> DataSetNames;
-			for (const FNiagaraDataSetID& ReadId : ReadDataSetIDs)
-			{
-				DataSetNames.AddUnique(ReadId.Name);
-			}
-			for (const FNiagaraDataSetID& WriteId : WriteDataSetIDs)
-			{
-				DataSetNames.AddUnique(WriteId.Name);
-			}
-
-			for (int32 i = 0; i < ParamMapHistories.Num(); i++)
+			// Either go on to the next phase, or write to the final output context.
+			if (RequiresInterpolation())
 			{
-				for (FName DataSetName : DataSetNames)
-				{
-					if (ParamMapHistories[i].FindVariable(*(DataSetName.ToString() + TEXT(".Alive")), FNiagaraTypeDefinition::GetBoolDef()) != INDEX_NONE)
-					{
-						bGpuUsesAlive = true;
-						break;
-					}
-				}
-				if (bGpuUsesAlive)
-				{
-					break;
-				}
+				OutHlslOutput += TEXT("\tInterpolateParameters(Context);\n"); // Requires ExecIndex, which needs to be in a stage.
 			}
+		}
 
-			OutHlslOutput += TEXT("\tint StartingPhase = Phase;\n");
-			OutHlslOutput += TEXT("\tGCurrentPhase = Phase;\n");
-
-			for (int32 StageIdx = 0; StageIdx < TranslationStages.Num(); StageIdx++)
-			{
-				OutHlslOutput += FString::Printf(TEXT("\tif(Phase==%d)\n\t{\n"), StageIdx);
-				if (bGpuUsesAlive)
-				{
-					OutHlslOutput += FString::Printf(TEXT("\t\tif (StartingPhase == %d)\n\t\t{\n\t\t\tContext.%s.DataInstance.Alive=true;\n\t\t}\n"), StageIdx, *TranslationStages[StageIdx].PassNamespace);
-				}
-
-				if (StageIdx == 0)
-				{
-					// Either go on to the next phase, or write to the final output context.
-					if (RequiresInterpolation())
-					{
-						OutHlslOutput += TEXT("\t\tInterpolateParameters(Context);\n"); // Requires ExecIndex, which needs to be in a stage.
-					}
-				}
-
-				OutHlslOutput += FString::Printf(TEXT("\t\tSimulate%s(Context);\n"), TranslationStages.Num() > 1 ? *TranslationStages[StageIdx].PassNamespace : TEXT(""));
-
-				if (StageIdx + 1 < TranslationStages.Num() && TranslationStages[StageIdx + 1].bCopyPreviousParams)
-				{
-					OutHlslOutput += TEXT("\t\t//Begin Transfer of Attributes!\n");
-					if (ParamMapDefinedAttributesToNamespaceVars.Num() != 0)
-					{
-						FString CopyParticlesStr =
-							TEXT("\t\tContext.") + TranslationStages[StageIdx + 1].PassNamespace + TEXT(".Particles = Context.") + TranslationStages[StageIdx].PassNamespace + TEXT(".Particles;\n");
-						OutHlslOutput += CopyParticlesStr;
-
-						if (bGpuUsesAlive)
-						{
-							FString CopyDataInstanceStr = TEXT("\t\tContext.") + TranslationStages[StageIdx + 1].PassNamespace + TEXT(".DataInstance = Context.") + TranslationStages[StageIdx].PassNamespace + TEXT(".DataInstance;\n");
-							OutHlslOutput += CopyDataInstanceStr;
-						}
-					}
-					OutHlslOutput += TEXT("\t\t//End Transfer of Attributes!\n\n");
-				}
-
-
-				// Either go on to the next phase, or write to the final output context.
-				if (StageIdx + 1 < TranslationStages.Num() && TranslationStages[StageIdx + 1].bInterpolatePreviousParams)
-				{
-					OutHlslOutput += FString::Printf(TEXT("\t\tPhase = %d;\n"), StageIdx + 1);
-					OutHlslOutput += FString::Printf(TEXT("\t\tGCurrentPhase = %d;\n"), StageIdx + 1);
-				}
+		OutHlslOutput += FString::Printf(TEXT("\tSimulate%s(Context);\n"), TranslationStages.Num() > 1 ? *TranslationStages[StageIdx].PassNamespace : TEXT(""));
 
-				OutHlslOutput += TEXT("\t}\n");
-			}
-		}
-		else
+		if (StageIdx + 1 < TranslationStages.Num() && TranslationStages[StageIdx + 1].bCopyPreviousParams)
 		{
-			for (int32 StageIdx = 0; StageIdx < TranslationStages.Num(); StageIdx++)
+			OutHlslOutput += TEXT("\t//Begin Transfer of Attributes!\n");
+			if (ParamMapDefinedAttributesToNamespaceVars.Num() != 0)
 			{
-				if (StageIdx == 0)
-				{
-					// Either go on to the next phase, or write to the final output context.
-					if (RequiresInterpolation())
-					{
-						OutHlslOutput += TEXT("\tInterpolateParameters(Context);\n"); // Requires ExecIndex, which needs to be in a stage.
-					}
-				}
-
-				OutHlslOutput += FString::Printf(TEXT("\tSimulate%s(Context);\n"), TranslationStages.Num() > 1 ? *TranslationStages[StageIdx].PassNamespace : TEXT(""));
-
-				if (StageIdx + 1 < TranslationStages.Num() && TranslationStages[StageIdx + 1].bCopyPreviousParams)
-				{
-					OutHlslOutput += TEXT("\t//Begin Transfer of Attributes!\n");
-					if (ParamMapDefinedAttributesToNamespaceVars.Num() != 0)
-					{
-						FString CopyStr = TEXT("\tContext.") + TranslationStages[StageIdx + 1].PassNamespace + TEXT(".Particles = Context.") + TranslationStages[StageIdx].PassNamespace + TEXT(".Particles;\n");
-						OutHlslOutput += CopyStr;
-					}
-					OutHlslOutput += TEXT("\t//End Transfer of Attributes!\n\n");
-				}
+				FString CopyStr = TEXT("\tContext.") + TranslationStages[StageIdx + 1].PassNamespace + TEXT(".Particles = Context.") + TranslationStages[StageIdx].PassNamespace + TEXT(".Particles;\n");
+				OutHlslOutput += CopyStr;
 			}
+			OutHlslOutput += TEXT("\t//End Transfer of Attributes!\n\n");
 		}
 	}
 
@@ -2015,44 +2142,12 @@ void FHlslNiagaraTranslator::DefineMain(FString &OutHlslOutput,
 
 	ExitStatsScope(OutHlslOutput);
 	OutHlslOutput += TEXT("}\n");
-
-
-	// define a simple copy function to run on a section of the buffer for GPU event scripts; 
-	//	SimulateMainComputeCS in the usf will decide which one to call for which instances
-	// TODO: we'll want to combine spawn and update here soon so we'll need multiple main functions to be called from NiagaraEmitterInstanceShader.usf
-	//	this will need SpawnMain and UpdateMain for the regular particle scripts; all spawn events should be a single dispatch as well, each with its own function
-	//
-	if (CompilationTarget == ENiagaraSimTarget::GPUComputeSim)
-	{
-		OutHlslOutput += TEXT("void CopyInstance(in int InstanceIdx)\n{\n\tFSimulationContext Context = (FSimulationContext)0;\n");
-		if (UNiagaraScript::IsParticleEventScript(CompileOptions.TargetUsage))
-		{
-			for (int32 VarArrayIdx = 0; VarArrayIdx < DataSetReads.Num(); VarArrayIdx++)
-			{
-				const FNiagaraDataSetID DataSetID = ReadDataSetIDs[VarArrayIdx];
-				const TArray<FNiagaraVariable>& ArrayRef = DataSetVariables[DataSetReads[DataSetID]];
-				DefineDataSetVariableReads(HlslOutput, DataSetID, VarArrayIdx, ArrayRef);
-			}
-
-			if (bGpuUsesAlive)
-			{
-				OutHlslOutput += TEXT("\tContext.Map.DataInstance.Alive = true;\n");
-			}
-
-			for (int32 VarArrayIdx = 0; VarArrayIdx < DataSetWrites.Num(); VarArrayIdx++)
-			{
-				const FNiagaraDataSetID DataSetID = WriteDataSetIDs[VarArrayIdx];
-				const TArray<FNiagaraVariable>& ArrayRef = DataSetVariables[DataSetWrites[DataSetID]];
-				DefineDataSetVariableWrites(HlslOutput, DataSetID, VarArrayIdx, ArrayRef);
-			}
-		}
-		OutHlslOutput += TEXT("}\n");
-	}
-
 }
 
 void FHlslNiagaraTranslator::DefineDataSetVariableWrites(FString &OutHlslOutput, const FNiagaraDataSetID& Id, int32 DataSetIndex, const TArray<FNiagaraVariable>& WriteVars)
 {
+	check(CompilationTarget != ENiagaraSimTarget::GPUComputeSim);
+
 	//TODO Grab indices for data set writes (inc output) and do the write. Need to rewrite this for events interleaved..
 	OutHlslOutput += "\t{\n";
 	bool bUsesAlive = false;
@@ -2129,7 +2224,7 @@ void FHlslNiagaraTranslator::DefineDataSetVariableWrites(FString &OutHlslOutput,
 	int32 WriteOffsetInt = 0;
 	int32 WriteOffsetFloat = 0;
 	int32 &FloatCounter = WriteOffsetFloat;
-	int32 &IntCounter = CompilationTarget == ENiagaraSimTarget::GPUComputeSim ? WriteOffsetInt : WriteOffsetFloat;
+	int32 &IntCounter = WriteOffsetFloat;
 	for (const FNiagaraVariable &Var : WriteVars)
 	{
 		// If coming from a parameter map, use the one on the context, otherwise use the output.
@@ -2149,20 +2244,15 @@ void FHlslNiagaraTranslator::DefineDataSetVariableWrites(FString &OutHlslOutput,
 
 void FHlslNiagaraTranslator::DefineDataSetVariableReads(FString &OutHlslOutput, const FNiagaraDataSetID& Id, int32 DataSetIndex, const TArray<FNiagaraVariable>& ReadVars)
 {
+	check(CompilationTarget != ENiagaraSimTarget::GPUComputeSim);
+
 	int32 ReadOffsetInt = 0;
 	int32 ReadOffsetFloat = 0;
 	int32 &FloatCounter = ReadOffsetFloat;
-	int32 &IntCounter = CompilationTarget == ENiagaraSimTarget::GPUComputeSim ? ReadOffsetInt : ReadOffsetFloat;
+	int32 &IntCounter = ReadOffsetFloat;
 
 	FString DataSetName = Id.Name.ToString();
 	FString Fmt;
-	FString SpawnCondition = "\t{\n", UpdateCondition = "\t{\n";
-
-	if (TranslationOptions.SimTarget == ENiagaraSimTarget::GPUComputeSim)
-	{
-		SpawnCondition = "\tif(Phase==0)\n\t{\n";
-		UpdateCondition = "\tif(Phase==1)\n\t{\n";
-	}
 
 	bool bIsGPUScript = UNiagaraScript::IsGPUScript(CompileOptions.TargetUsage);
 	bool bIsSpawnScript =
@@ -2196,8 +2286,7 @@ void FHlslNiagaraTranslator::DefineDataSetVariableReads(FString &OutHlslOutput,
 			GatherVariableForDataSetAccess(Var, Fmt, IntCounter, FloatCounter, DataSetIndex, TEXT(""), VarReads);
 		}
 
-		OutHlslOutput += SpawnCondition;
-		OutHlslOutput += VarReads + "\t}\n";
+		OutHlslOutput += VarReads;
 	}
 
 	// This will initialize parameters to their correct initial values from constants or data sets for update, and will also initialize parameters
@@ -2230,14 +2319,8 @@ void FHlslNiagaraTranslator::DefineDataSetVariableReads(FString &OutHlslOutput,
 			}
 			else
 			{
-				if (CompilationTarget == ENiagaraSimTarget::GPUComputeSim)
-				{
-					Fmt = VariableName + TEXT("{0} = InputData{1}({2}, {3}, InstanceIdx);\n");
-				}
-				else
-				{
-					Fmt = VariableName + TEXT("{0} = InputData{1}({2}, {3});\n");
-				}
+				Fmt = VariableName + TEXT("{0} = InputData{1}({2}, {3});\n");
+
 				static const auto UseShaderStagesCVar = IConsoleManager::Get().FindConsoleVariable(TEXT("fx.UseShaderStages"));
 				if (UseShaderStagesCVar->GetInt() == 1)
 				{
@@ -2253,8 +2336,7 @@ void FHlslNiagaraTranslator::DefineDataSetVariableReads(FString &OutHlslOutput,
 			GatherVariableForDataSetAccess(Var, Fmt, IntCounter, FloatCounter, DataSetIndex, TEXT(""), VarReads);
 		}
 
-		OutHlslOutput += UpdateCondition;
-		OutHlslOutput += VarReads + "\n\t}\n";
+		OutHlslOutput += VarReads;
 	}
 }
 
diff --git a/Engine/Plugins/FX/Niagara/Source/NiagaraEditor/Private/NiagaraHlslTranslator.h b/Engine/Plugins/FX/Niagara/Source/NiagaraEditor/Private/NiagaraHlslTranslator.h
--- a/Engine/Plugins/FX/Niagara/Source/NiagaraEditor/Private/NiagaraHlslTranslator.h
+++ b/Engine/Plugins/FX/Niagara/Source/NiagaraEditor/Private/NiagaraHlslTranslator.h
@@ -462,6 +462,7 @@ public:
 	void DefineDataSetReadFunction(FString &HlslOutput, TArray<FNiagaraDataSetID> &ReadDataSets);
 	void DefineDataSetWriteFunction(FString &HlslOutput, TArray<FNiagaraDataSetProperties> &WriteDataSets, TArray<int32>& WriteConditionVarIndices);
 	void DefineMain(FString &HLSLOutput, const TArray<TArray<FNiagaraVariable>>& DataSetVariables, const TMap<FNiagaraDataSetID, int32>& DataSetReads, const TMap<FNiagaraDataSetID, int32>& DataSetWrites);
+	void DefineMainGPUFunctions(const TArray<TArray<FNiagaraVariable>>& DataSetVariables, const TMap<FNiagaraDataSetID, int32>& DataSetReads, const TMap<FNiagaraDataSetID, int32>& DataSetWrites);
 
 	void DefineDataSetVariableReads(FString &HLSLOutput, const FNiagaraDataSetID& Id, int32 DataSetIndex, const TArray<FNiagaraVariable>& ReadVars);
 	void DefineDataSetVariableWrites(FString &HlslOutput, const FNiagaraDataSetID& Id, int32 DataSetIndex, const TArray<FNiagaraVariable>& WriteVars);
diff --git a/Engine/Plugins/FX/Niagara/Source/NiagaraShader/Public/NiagaraShader.h b/Engine/Plugins/FX/Niagara/Source/NiagaraShader/Public/NiagaraShader.h
--- a/Engine/Plugins/FX/Niagara/Source/NiagaraShader/Public/NiagaraShader.h
+++ b/Engine/Plugins/FX/Niagara/Source/NiagaraShader/Public/NiagaraShader.h
@@ -46,8 +46,21 @@ public:
 
 	typedef void (*ModifyCompilationEnvironmentType)(EShaderPlatform, const FNiagaraShaderScript* , FShaderCompilerEnvironment&);
 
+	static uint32 GetGroupSize(EShaderPlatform Platform)
+	{
+		if (Platform == SP_XBOXONE_D3D12 || Platform == SP_PS4)
+		{
+			return 64;
+		}
+		else
+		{
+			return 32;
+		}
+	}
+
 	static void ModifyCompilationEnvironment(EShaderPlatform Platform, const FNiagaraShaderScript*  Script, FShaderCompilerEnvironment& OutEnvironment)
 	{
+		OutEnvironment.SetDefine(TEXT("THREADGROUP_SIZE"), GetGroupSize(Platform));
 	}
 
 	static bool ValidateCompiledResult(EShaderPlatform Platform, const FShaderParameterMap& ParameterMap, TArray<FString>& OutError)
