Fixed niagara shader compilation warning around spawngroup

diff --git a/Engine/Plugins/FX/Niagara/Shaders/Private/NiagaraEmitterInstanceShader.usf b/Engine/Plugins/FX/Niagara/Shaders/Private/NiagaraEmitterInstanceShader.usf
--- a/Engine/Plugins/FX/Niagara/Shaders/Private/NiagaraEmitterInstanceShader.usf
+++ b/Engine/Plugins/FX/Niagara/Shaders/Private/NiagaraEmitterInstanceShader.usf
@@ -659,7 +659,7 @@ void SimulateMainComputeCS(
 
 		GetEmitterSpawnInfoForParticle(ExecIndex());
 		// Part of InitConstants() but updated by the previous statement
-		Context.MapSpawn.Emitter.SpawnGroup = Emitter_SpawnGroup;
+		// Context.MapSpawn.Emitter.SpawnGroup = Emitter_SpawnGroup;
 
 		InitSpawnVariables(Context);
 		ReadDataSets(Context);
