Remove duplicated definition
#rnx

diff --git a/Engine/Plugins/FX/Niagara/Shaders/Private/NiagaraEmitterInstanceShader.usf b/Engine/Plugins/FX/Niagara/Shaders/Private/NiagaraEmitterInstanceShader.usf
--- a/Engine/Plugins/FX/Niagara/Shaders/Private/NiagaraEmitterInstanceShader.usf
+++ b/Engine/Plugins/FX/Niagara/Shaders/Private/NiagaraEmitterInstanceShader.usf
@@ -68,9 +68,6 @@ float Reciprocal(float x){ return 1.0f/x; }
 		return HomogeneousWorldPosition.xyz / HomogeneousWorldPosition.w;
 	}
 
-
-	static uint GRandomSeedOffset;
-
 	// MASSIVE HACK - Tracked in JIRA UE-69298
 	// Hardcoded random function accessible from inner part of node implementation.
 	// It works for now at least and avoids exposing every random needed in the UI. 
