Fixed issues with constant being initialized to Emitter_SpawnGroup before it gets updated by GetEmitterSpawnInfoForParticle()

diff --git a/Engine/Plugins/FX/Niagara/Shaders/Private/NiagaraEmitterInstanceShader.usf b/Engine/Plugins/FX/Niagara/Shaders/Private/NiagaraEmitterInstanceShader.usf
--- a/Engine/Plugins/FX/Niagara/Shaders/Private/NiagaraEmitterInstanceShader.usf
+++ b/Engine/Plugins/FX/Niagara/Shaders/Private/NiagaraEmitterInstanceShader.usf
@@ -655,8 +655,6 @@ void SimulateMainComputeCS(
 		GCurrentPhase = 0;
 
 		GetEmitterSpawnInfoForParticle(ExecIndex());
-		// Part of InitConstants() but updated by the previous statement
-		// Context.MapSpawn.Emitter.SpawnGroup = Emitter_SpawnGroup;
 
 		InitSpawnVariables(Context);
 		ReadDataSets(Context);
diff --git a/Engine/Plugins/FX/Niagara/Source/NiagaraCore/Private/NiagaraCustomVersion.cpp b/Engine/Plugins/FX/Niagara/Source/NiagaraCore/Private/NiagaraCustomVersion.cpp
--- a/Engine/Plugins/FX/Niagara/Source/NiagaraCore/Private/NiagaraCustomVersion.cpp
+++ b/Engine/Plugins/FX/Niagara/Source/NiagaraCore/Private/NiagaraCustomVersion.cpp
@@ -8,4 +8,4 @@ const FGuid FNiagaraCustomVersion::GUID(0xFCF57AFA, 0x50764283, 0xB9A9E658, 0xFF
 // Register the custom version with core
 FCustomVersionRegistration GRegisterNiagaraCustomVersion(FNiagaraCustomVersion::GUID, FNiagaraCustomVersion::LatestVersion, TEXT("NiagaraVer"));
 
-const FGuid FNiagaraCustomVersion::LatestScriptCompileVersion(0xBCB9A1E7, 0xB6F54078, 0xBC9895E2, 0x1D5D4C0F);
\ No newline at end of file
+const FGuid FNiagaraCustomVersion::LatestScriptCompileVersion(0x7D607C23, 0x5E344266, 0x81518861, 0x9278156B);
\ No newline at end of file
diff --git a/Engine/Plugins/FX/Niagara/Source/NiagaraEditor/Private/NiagaraHlslTranslator.cpp b/Engine/Plugins/FX/Niagara/Source/NiagaraEditor/Private/NiagaraHlslTranslator.cpp
--- a/Engine/Plugins/FX/Niagara/Source/NiagaraEditor/Private/NiagaraHlslTranslator.cpp
+++ b/Engine/Plugins/FX/Niagara/Source/NiagaraEditor/Private/NiagaraHlslTranslator.cpp
@@ -1839,14 +1839,22 @@ void FHlslNiagaraTranslator::DefineMainGPUFunctions(
 		return false;
 	}();
 
+	// A list of constant to reset after Emitter_SpawnGroup gets modified by GetEmitterSpawnInfoForParticle()
+	TArray<FString> EmitterSpawnGroupReinit;
+
 	///////////////////////
 	// InitConstants()
 	HlslOutput += TEXT("void InitConstants(inout FSimulationContext Context)\n{\n");
 	{
 		// Fill in the defaults for parameters.
-		for (int32 i = 0; i < MainPreSimulateChunks.Num(); ++i)
+		for (const FString& InitChunk : MainPreSimulateChunks)
 		{
-			HlslOutput += TEXT("\t") + MainPreSimulateChunks[i] + TEXT("\n");
+			HlslOutput += TEXT("\t") + InitChunk + TEXT("\n");
+
+			if (InitChunk.Contains(TEXT("Emitter_SpawnGroup;")))
+			{
+				EmitterSpawnGroupReinit.Add(InitChunk);
+			}
 		}
 	}
 	HlslOutput += TEXT("}\n\n");
@@ -1855,6 +1863,16 @@ void FHlslNiagaraTranslator::DefineMainGPUFunctions(
 	// InitSpawnVariables()
 	HlslOutput += TEXT("void InitSpawnVariables(inout FSimulationContext Context)\n{\n");
 	{
+		// Reset constant that have been modified by GetEmitterSpawnInfoForParticle()
+		if (EmitterSpawnGroupReinit.Num())
+		{
+			for (const FString& ReinitChunk : EmitterSpawnGroupReinit)
+			{
+				HlslOutput += TEXT("\t") + ReinitChunk + TEXT("\n");
+			}
+			HlslOutput += TEXT("\n");
+		}
+
 		FString ContextName = TEXT("\tContext.Map.");
 		if (TranslationStages.Num() > 1) // First context 0 is "MapSpawn"
 		{
