Fixed non deterministic cook issues caused by FNiagaraParameterStore::SortedParameterOffsets.
The fix implies using lexical sorting in non cooked builds, while using comparison indices in cooked (for FName)

diff --git a/Engine/Plugins/FX/Niagara/Source/Niagara/Public/NiagaraParameterStore.h b/Engine/Plugins/FX/Niagara/Source/Niagara/Public/NiagaraParameterStore.h
--- a/Engine/Plugins/FX/Niagara/Source/Niagara/Public/NiagaraParameterStore.h
+++ b/Engine/Plugins/FX/Niagara/Source/Niagara/Public/NiagaraParameterStore.h
@@ -6,7 +6,8 @@
 #include "NiagaraDataInterface.h"
 #include "NiagaraParameterStore.generated.h"
 
-#define NIAGARA_VARIABLE_LEXICAL_SORTING 0
+// When not cooked, sort by actual name to ensure deterministic cooked data
+#define NIAGARA_VARIABLE_LEXICAL_SORTING WITH_EDITORONLY_DATA
 
 struct FNiagaraVariableSearch
 {
