Fixing issue with loading DefaultEffectType.
Settings objects do not get PostLoad called. Instead they're loaded from config files and their state is ready in by PostInitProperites.
Also removing the lazy load in Get().

#rnx

diff --git a/Engine/Plugins/FX/Niagara/Source/Niagara/Private/NiagaraSettings.cpp b/Engine/Plugins/FX/Niagara/Source/Niagara/Private/NiagaraSettings.cpp
--- a/Engine/Plugins/FX/Niagara/Source/Niagara/Private/NiagaraSettings.cpp
+++ b/Engine/Plugins/FX/Niagara/Source/Niagara/Private/NiagaraSettings.cpp
@@ -21,9 +21,9 @@ FText UNiagaraSettings::GetSectionText() const
 }
 #endif
 
-void UNiagaraSettings::PostLoad()
+void UNiagaraSettings::PostInitProperties()
 {
-	Super::PostLoad();
+	Super::PostInitProperties();
 
 	DefaultEffectTypePtr = Cast<UNiagaraEffectType>(DefaultEffectType.TryLoad());
 }
@@ -49,9 +49,5 @@ UNiagaraSettings::FOnNiagaraSettingsChanged UNiagaraSettings::SettingsChangedDel
 
 UNiagaraEffectType* UNiagaraSettings::GetDefaultEffectType()const
 {
-	if (DefaultEffectTypePtr == nullptr && DefaultEffectType.IsValid())
-	{
-		DefaultEffectTypePtr = Cast<UNiagaraEffectType>(DefaultEffectType.TryLoad());
-	}
 	return DefaultEffectTypePtr;
 }
diff --git a/Engine/Plugins/FX/Niagara/Source/Niagara/Private/NiagaraSystem.cpp b/Engine/Plugins/FX/Niagara/Source/Niagara/Private/NiagaraSystem.cpp
--- a/Engine/Plugins/FX/Niagara/Source/Niagara/Private/NiagaraSystem.cpp
+++ b/Engine/Plugins/FX/Niagara/Source/Niagara/Private/NiagaraSystem.cpp
@@ -1371,7 +1371,11 @@ void UNiagaraSystem::ResolveScalabilityOverrides()
 	ResolvedScalabilitySettings.Empty();
 
 	UNiagaraEffectType* ActualEffectType = GetEffectType();
-	int32 NumEffectTypeSettings = ActualEffectType ? ActualEffectType->GetScalabilitySettings().Num() : 0;
+	int32 NumEffectTypeSettings = 0;
+	if (ActualEffectType)
+	{
+		ActualEffectType->GetScalabilitySettings().Num();
+	}
 	int32 NumOverrides = ScalabilityOverrides.Num();
 
 	int32 ResolvedCount = FMath::Max(NumEffectTypeSettings, NumOverrides);
diff --git a/Engine/Plugins/FX/Niagara/Source/Niagara/Public/NiagaraSettings.h b/Engine/Plugins/FX/Niagara/Source/Niagara/Public/NiagaraSettings.h
--- a/Engine/Plugins/FX/Niagara/Source/Niagara/Public/NiagaraSettings.h
+++ b/Engine/Plugins/FX/Niagara/Source/Niagara/Public/NiagaraSettings.h
@@ -37,7 +37,7 @@ class NIAGARA_API UNiagaraSettings : public UDeveloperSettings
 
 	UNiagaraEffectType* GetDefaultEffectType()const;
 
-	virtual void PostLoad();
+	virtual void PostInitProperties();
 #if WITH_EDITOR
 	virtual void PostEditChangeProperty(FPropertyChangedEvent& PropertyChangedEvent) override;
 
