Clamping incomming DistanceFraction from the scalabiltiy system a Wyeth's request.

#rnx

diff --git a/Engine/Plugins/FX/Niagara/Source/Niagara/Private/NiagaraSystemInstance.cpp b/Engine/Plugins/FX/Niagara/Source/Niagara/Private/NiagaraSystemInstance.cpp
--- a/Engine/Plugins/FX/Niagara/Source/Niagara/Private/NiagaraSystemInstance.cpp
+++ b/Engine/Plugins/FX/Niagara/Source/Niagara/Private/NiagaraSystemInstance.cpp
@@ -1507,7 +1507,8 @@ void FNiagaraSystemInstance::TickInstanceParameters_Concurrent()
 	OwnerInverseDeltaSecondsParam.SetValue(1.0f / GatheredInstanceParameters.DeltaSeconds);
 
 	OwnerLODDistanceParam.SetValue(GatheredInstanceParameters.LODDistance);
-	OwnerLODDistanceFractionParam.SetValue(GatheredInstanceParameters.LODDistance / GatheredInstanceParameters.MaxLODDistance);
+	float SystemDistanceFraction = FMath::Clamp(GatheredInstanceParameters.LODDistance / GatheredInstanceParameters.MaxLODDistance, 0.0f, 1.0f);
+	OwnerLODDistanceFractionParam.SetValue(SystemDistanceFraction);
 	OwnerEngineTimeParam.SetValue(GatheredInstanceParameters.TimeSeconds);
 	OwnerEngineRealtimeParam.SetValue(GatheredInstanceParameters.RealTimeSeconds);
 
