Removing erroneous check.

#rnx

diff --git a/Engine/Plugins/FX/Niagara/Source/Niagara/Classes/NiagaraDataSet.h b/Engine/Plugins/FX/Niagara/Source/Niagara/Classes/NiagaraDataSet.h
--- a/Engine/Plugins/FX/Niagara/Source/Niagara/Classes/NiagaraDataSet.h
+++ b/Engine/Plugins/FX/Niagara/Source/Niagara/Classes/NiagaraDataSet.h
@@ -422,7 +422,7 @@ struct FNiagaraDataSetAccessorBase
 		if (Var.IsValid())
 		{
 			DataSet = &InDataSet;
-			checkSlow(VarLayout == InDataSet.GetVariableLayout(Var));
+			//checkSlow(VarLayout == InDataSet.GetVariableLayout(Var));
 		}
 		else
 		{
