Fixes to the CPU async colliison queries within Niagara
Creation and reading the results are not thread safe, and needs to be done on the game thread.  Individual DI will now queue up requests in a threadsafe maner, and then the gamethread will queue them up (and read the results) for the individual DI to gather their results.

diff --git a/Engine/Plugins/FX/Niagara/Source/Niagara/Classes/NiagaraCollision.h b/Engine/Plugins/FX/Niagara/Source/Niagara/Classes/NiagaraCollision.h
--- a/Engine/Plugins/FX/Niagara/Source/Niagara/Classes/NiagaraCollision.h
+++ b/Engine/Plugins/FX/Niagara/Source/Niagara/Classes/NiagaraCollision.h
@@ -14,19 +14,26 @@ enum class ENiagaraCollisionMode : uint8
 	DistanceField
 };
 
-
 struct FNiagaraCollisionTrace
 {
+	FNiagaraCollisionTrace(const FVector& InStartPos, const FVector& InEndPos, ECollisionChannel InChannel, const FCollisionQueryParams& InQueryParams)
+		: HitIndex(INDEX_NONE)
+		, CollisionQueryParams(InQueryParams)
+		, StartPos(InStartPos)
+		, EndPos(InEndPos)
+		, Channel(InChannel)
+	{}
+
 	FTraceHandle CollisionTraceHandle;
-	uint32 SourceParticleIndex;
-	float CollisionSize;
-	float DeltaSeconds;
+	int32 HitIndex;
+	const FCollisionQueryParams CollisionQueryParams;
+	const FVector StartPos;
+	const FVector EndPos;
+	const ECollisionChannel Channel;
 };
 
-
 struct FNiagaraDICollsionQueryResult
 {
-	uint32 TraceID;
 	FVector CollisionPos;
 	FVector CollisionNormal;
 	FVector CollisionVelocity;
@@ -59,16 +66,13 @@ public:
 		return CurrBuffer ^ 0x1;
 	}
 
-	void Tick(ENiagaraSimTarget Target)
-	{
-		CurrBuffer = CurrBuffer ^ 0x1;
-	}
+	void DispatchQueries();
+	void CollectResults();
 
 	void ClearWrite()
 	{
 		uint32 PrevNum = CollisionTraces[CurrBuffer].Num();
 		CollisionTraces[CurrBuffer].Empty(PrevNum);
-		IdToTraceIdx[CurrBuffer].Empty(PrevNum);
 	}
 
 
@@ -78,10 +82,7 @@ public:
 		CollisionWorld = InCollisionWorld;
 		CollisionTraces[0].Empty();
 		CollisionTraces[1].Empty();
-		IdToTraceIdx[0].Empty();
-		IdToTraceIdx[1].Empty();
 		CurrBuffer = 0;
-		TraceID = 0;
 	}
 
 	int32 SubmitQuery(FVector Position, FVector Direction, float CollisionSize, float DeltaSeconds);
@@ -90,15 +91,21 @@ public:
 	bool GetQueryResult(uint32 TraceID, FNiagaraDICollsionQueryResult &Result);
 
 private:
+	void FlipBuffers()
+	{
+		CurrBuffer = CurrBuffer ^ 0x1;
+	}
+
 	//TArray<FTraceHandle> CollisionTraceHandles;
+	const static FName CollisionTagName;
 
+	FRWLock CollisionTraceLock;
 	TArray<FNiagaraCollisionEventPayload> CollisionEvents;
 	FNiagaraDataSet *CollisionEventDataSet;
 
 	FNiagaraSystemInstanceID BatchID;
 	TArray<FNiagaraCollisionTrace> CollisionTraces[2];
-	TMap<uint32, int32> IdToTraceIdx[2];
+	TArray<FNiagaraDICollsionQueryResult> CollisionResults;
 	uint32 CurrBuffer;
-	uint32 TraceID;
 	UWorld *CollisionWorld;
 };
\ No newline at end of file
diff --git a/Engine/Plugins/FX/Niagara/Source/Niagara/Private/NiagaraCollision.cpp b/Engine/Plugins/FX/Niagara/Source/Niagara/Private/NiagaraCollision.cpp
--- a/Engine/Plugins/FX/Niagara/Source/Niagara/Private/NiagaraCollision.cpp
+++ b/Engine/Plugins/FX/Niagara/Source/Niagara/Private/NiagaraCollision.cpp
@@ -10,77 +10,151 @@
 DECLARE_CYCLE_STAT(TEXT("Collision"), STAT_NiagaraCollision, STATGROUP_Niagara);
 DECLARE_CYCLE_STAT(TEXT("Event Emission"), STAT_NiagaraEventWrite, STATGROUP_Niagara);
 
+const FName FNiagaraDICollisionQueryBatch::CollisionTagName = FName("Niagara");
 
-int32 FNiagaraDICollisionQueryBatch::SubmitQuery(FVector Position, FVector Direction, float CollisionSize, float DeltaSeconds)
+void FNiagaraDICollisionQueryBatch::DispatchQueries()
 {
+	check(IsInGameThread());
+
 	SCOPE_CYCLE_COUNTER(STAT_NiagaraCollision);
-	int32 Ret = INDEX_NONE;
+
+	// swap the buffers, and use our new read buffer for issuing all of our accumulated queries now that we're on the main thread
+	FlipBuffers();
+
+	// locks are not used here because we assume that the per instance ticking is done from the main thread
 	if (CollisionWorld)
 	{
-		//int32 TestCollision = *TstIt;
-		//if (TestCollision)
+		const int32 ReadBufferIdx = GetReadBufferIdx();
+		const int32 TraceCount = CollisionTraces[ReadBufferIdx].Num();
+
+		for (int32 TraceIt = 0; TraceIt < TraceCount; ++TraceIt)
 		{
-			FVector EndPosition = Position + Direction * DeltaSeconds;
-
-			float Length;
-			FVector NormDir;
-			Direction.ToDirectionAndLength(NormDir, Length);
-			Position -= NormDir * (CollisionSize / 2);
-			EndPosition += NormDir * (CollisionSize/ 2);
-
-			FCollisionQueryParams QueryParams(SCENE_QUERY_STAT(NiagraAsync));
-			QueryParams.OwnerTag = "Niagara";
-			QueryParams.bFindInitialOverlaps = false;
-			QueryParams.bReturnFaceIndex = false;
-			QueryParams.bReturnPhysicalMaterial = true;
-			QueryParams.bTraceComplex = false;
-			QueryParams.bIgnoreTouches = true;
-			FTraceHandle Handle = CollisionWorld->AsyncLineTraceByChannel(EAsyncTraceType::Single, Position, EndPosition, ECollisionChannel::ECC_WorldStatic, QueryParams, FCollisionResponseParams::DefaultResponseParam, nullptr, TraceID);
-			FNiagaraCollisionTrace Trace;
-			Trace.CollisionTraceHandle = Handle;
-			Trace.SourceParticleIndex = TraceID;
-			Trace.CollisionSize = CollisionSize;
-			Trace.DeltaSeconds = DeltaSeconds;
-
-			int32 TraceIdx = CollisionTraces[GetWriteBufferIdx()].Add(Trace);
-			IdToTraceIdx[GetWriteBufferIdx()].Add(TraceID) = TraceIdx;
-			Ret = TraceID;
-			TraceID++;
+			FNiagaraCollisionTrace& CollisionTrace = CollisionTraces[ReadBufferIdx][TraceIt];
+
+			CollisionTrace.CollisionTraceHandle = CollisionWorld->AsyncLineTraceByChannel(
+				EAsyncTraceType::Single,
+				CollisionTrace.StartPos,
+				CollisionTrace.EndPos,
+				CollisionTrace.Channel,
+				CollisionTrace.CollisionQueryParams,
+				FCollisionResponseParams::DefaultResponseParam,
+				nullptr,
+				TraceIt);
 		}
 	}
+}
+
+void FNiagaraDICollisionQueryBatch::CollectResults()
+{
+	check(IsInGameThread());
+
+	SCOPE_CYCLE_COUNTER(STAT_NiagaraCollision);
 
-	return Ret;
+	// locks are not used here because we assume that the per instance ticking is done from the main thread
+	if (CollisionWorld)
+	{
+		const int32 ReadBufferIdx = GetReadBufferIdx();
+		const int32 TraceCount = CollisionTraces[ReadBufferIdx].Num();
+
+		CollisionResults.Reset(TraceCount);
+
+		for (int32 TraceIt = 0; TraceIt < TraceCount; ++TraceIt)
+		{
+			FNiagaraCollisionTrace& CollisionTrace = CollisionTraces[ReadBufferIdx][TraceIt];
+
+			FTraceDatum TraceResult;
+			const bool TraceReady = CollisionWorld->QueryTraceData(CollisionTrace.CollisionTraceHandle, TraceResult);
+
+			if (TraceReady && TraceResult.OutHits.Num())
+			{
+				FHitResult* Hit = FHitResult::GetFirstBlockingHit(TraceResult.OutHits);
+				if (Hit && Hit->bBlockingHit)
+				{
+					CollisionTrace.HitIndex = CollisionResults.AddUninitialized();
+					FNiagaraDICollsionQueryResult& Result = CollisionResults[CollisionTrace.HitIndex];
+
+					Result.IsInsideMesh = Hit->bStartPenetrating;
+					Result.CollisionPos = Hit->ImpactPoint;// -NormVel*(CurTrace.CollisionSize / 2);
+					Result.CollisionNormal = Hit->ImpactNormal;
+					if (Hit->PhysMaterial.IsValid())
+					{
+						Result.PhysicalMaterialIdx = Hit->PhysMaterial->GetUniqueID();
+						Result.Friction = Hit->PhysMaterial->Friction;
+						Result.Restitution = Hit->PhysMaterial->Restitution;
+					}
+					else
+					{
+						Result.PhysicalMaterialIdx = -1;
+						Result.Friction = 0.0f;
+						Result.Restitution = 0.0f;
+					}
+				}
+			}
+		}
+	}
 }
 
-int32 FNiagaraDICollisionQueryBatch::SubmitQuery(FVector StartPos, FVector EndPos, ECollisionChannel TraceChannel)
+int32 FNiagaraDICollisionQueryBatch::SubmitQuery(FVector StartPos, FVector Direction, float CollisionSize, float DeltaSeconds)
 {
 	SCOPE_CYCLE_COUNTER(STAT_NiagaraCollision);
-	if (!CollisionWorld)
+
+	FVector EndPos = StartPos + Direction * DeltaSeconds;
+
+	float Length;
+	FVector NormDir;
+	Direction.ToDirectionAndLength(NormDir, Length);
+	StartPos -= NormDir * (CollisionSize / 2);
+	EndPos += NormDir * (CollisionSize / 2);
+
+	FCollisionQueryParams CollisionQueryParams;
+	CollisionQueryParams.OwnerTag = CollisionTagName;
+	CollisionQueryParams.bFindInitialOverlaps = false;
+	CollisionQueryParams.bReturnFaceIndex = false;
+	CollisionQueryParams.bReturnPhysicalMaterial = true;
+	CollisionQueryParams.bTraceComplex = false;
+	CollisionQueryParams.bIgnoreTouches = true;
+
+	int32 TraceIdx = INDEX_NONE;
+
+	if (Length > SMALL_NUMBER)
 	{
-		return INDEX_NONE;
+		FRWScopeLock WriteScope(CollisionTraceLock, SLT_Write);
+
+		const int32 WriteBufferIdx = GetWriteBufferIdx();
+
+		TraceIdx = CollisionTraces[WriteBufferIdx].Emplace(StartPos, EndPos, ECC_WorldStatic, CollisionQueryParams);
 	}
 
-	FCollisionQueryParams QueryParams(SCENE_QUERY_STAT(NiagaraAsync));
-	static FName NiagaraName("Niagara");
-	QueryParams.OwnerTag = NiagaraName;
-	QueryParams.bFindInitialOverlaps = false;
-	QueryParams.bReturnFaceIndex = false;
-	QueryParams.bReturnPhysicalMaterial = true;
-	QueryParams.bTraceComplex = false;
-	QueryParams.bIgnoreTouches = true;
-	FTraceHandle Handle = CollisionWorld->AsyncLineTraceByChannel(EAsyncTraceType::Single, StartPos, EndPos, TraceChannel, QueryParams, FCollisionResponseParams::DefaultResponseParam, nullptr, TraceID);
-	FNiagaraCollisionTrace Trace;
-	Trace.CollisionTraceHandle = Handle;
-	Trace.SourceParticleIndex = TraceID;
+	return TraceIdx;
+}
 
-	int32 TraceIdx = CollisionTraces[GetWriteBufferIdx()].Add(Trace);
-	IdToTraceIdx[GetWriteBufferIdx()].Add(TraceID) = TraceIdx;
+int32 FNiagaraDICollisionQueryBatch::SubmitQuery(FVector StartPos, FVector EndPos, ECollisionChannel TraceChannel)
+{
+	SCOPE_CYCLE_COUNTER(STAT_NiagaraCollision);
+	
+	FCollisionQueryParams CollisionQueryParams;
+	CollisionQueryParams.OwnerTag = CollisionTagName;
+	CollisionQueryParams.bFindInitialOverlaps = false;
+	CollisionQueryParams.bReturnFaceIndex = false;
+	CollisionQueryParams.bReturnPhysicalMaterial = true;
+	CollisionQueryParams.bTraceComplex = false;
+	CollisionQueryParams.bIgnoreTouches = true;
 
-	int32 Ret = TraceID;
-	TraceID++;
-	return Ret;
+	int32 TraceIdx = INDEX_NONE;
+
+	if ((EndPos - StartPos).SizeSquared() > SMALL_NUMBER)
+	{
+		FRWScopeLock WriteScope(CollisionTraceLock, SLT_Write);
+
+		const int32 WriteBufferIdx = GetWriteBufferIdx();
+
+		TraceIdx = CollisionTraces[WriteBufferIdx].Emplace(StartPos, EndPos, TraceChannel, CollisionQueryParams);
+	}
+
+	return TraceIdx;
 }
 
+// Work has been done on the collision world side of things to support synchronous queries from multiple threads
 bool FNiagaraDICollisionQueryBatch::PerformQuery(FVector StartPos, FVector EndPos, FNiagaraDICollsionQueryResult &Result, ECollisionChannel TraceChannel)
 {
 	SCOPE_CYCLE_COUNTER(STAT_NiagaraCollision);
@@ -90,7 +164,7 @@ bool FNiagaraDICollisionQueryBatch::PerformQuery(FVector StartPos, FVector EndPo
 	}
 
 	FCollisionQueryParams QueryParams(SCENE_QUERY_STAT(NiagaraSync));
-	QueryParams.OwnerTag = "Niagara";
+	QueryParams.OwnerTag = CollisionTagName;
 	QueryParams.bFindInitialOverlaps = false;
 	QueryParams.bReturnFaceIndex = false;
 	QueryParams.bReturnPhysicalMaterial = true;
@@ -122,82 +196,18 @@ bool FNiagaraDICollisionQueryBatch::PerformQuery(FVector StartPos, FVector EndPo
 
 bool FNiagaraDICollisionQueryBatch::GetQueryResult(uint32 InTraceID, FNiagaraDICollsionQueryResult &Result)
 {
-	int32 *TraceIdxPtr = IdToTraceIdx[GetReadBufferIdx()].Find(InTraceID);
-	int32 TraceIdx = INDEX_NONE;
+	const int32 ReadBufferIdx = GetReadBufferIdx();
 
-	if (TraceIdxPtr)
+	if (CollisionTraces[ReadBufferIdx].IsValidIndex(InTraceID))
 	{
-		TraceIdx = *TraceIdxPtr;
-		check(CollisionTraces[GetReadBufferIdx()].IsValidIndex(TraceIdx));
-		FNiagaraCollisionTrace &CurTrace = CollisionTraces[GetReadBufferIdx()][TraceIdx];
-		FTraceHandle Handle = CurTrace.CollisionTraceHandle;
-		FTraceDatum CurData;
-		// wait for trace handles; this should block rarely to never
-		bool bReady = CollisionWorld->QueryTraceData(Handle, CurData);
-		/*
-		if (!bReady)
-		{
-			// if the query came back false, it's possible that the hanle is invalid for some reason; skip in that case
-			// TODO: handle this more gracefully
-			if (!CollisionWorld->IsTraceHandleValid(Handle, false))
-			{
-				break;
-			}
-			break;
-			}
-		}
-		*/
+		const int32 ResultIndex = CollisionTraces[ReadBufferIdx][InTraceID].HitIndex;
 
-		if (bReady && CurData.OutHits.Num())
+		if (CollisionResults.IsValidIndex(ResultIndex))
 		{
-			// grab the first hit that blocks
-			FHitResult *Hit = FHitResult::GetFirstBlockingHit(CurData.OutHits);
-			if (Hit && Hit->bBlockingHit)
-			{
-				/*
-				FNiagaraCollisionEventPayload Event;
-				Event.CollisionNormal = Hit->ImpactNormal;
-				Event.CollisionPos = Hit->ImpactPoint;
-				Event.CollisionVelocity = CurCheck.OriginalVelocity;
-				Event.ParticleIndex = CurData.UserData;
-				check(!Event.CollisionNormal.ContainsNaN());
-				check(Event.CollisionNormal.IsNormalized());
-				check(!Event.CollisionPos.ContainsNaN());
-				check(!Event.CollisionVelocity.ContainsNaN());
-
-				// TODO add to unique list of physical materials for Blueprint
-				Event.PhysicalMaterialIndex = 0;// Hit->PhysMaterial->GetUniqueID();
-
-				Payloads.Add(Event);
-				*/
-				Result.IsInsideMesh = Hit->bStartPenetrating;
-				Result.CollisionPos = Hit->ImpactPoint;// -NormVel*(CurTrace.CollisionSize / 2);
-				Result.CollisionNormal = Hit->ImpactNormal;
-				Result.TraceID = InTraceID;
-				if (Hit->PhysMaterial.IsValid())
-				{
-					Result.PhysicalMaterialIdx = Hit->PhysMaterial->GetUniqueID();
-					Result.Friction = Hit->PhysMaterial->Friction;
-					Result.Restitution = Hit->PhysMaterial->Restitution;
-				}
-				else
-				{
-					Result.PhysicalMaterialIdx = -1;
-					Result.Friction = 0.0f;
-					Result.Restitution = 0.0f;
-				}
-				return true;
-			}
+			Result = CollisionResults[ResultIndex];
+			return true;
 		}
 	}
 
 	return false;
-}
-
-
-/*
-int32 FNiagaraDICollisionQueryBatch::Tick()
-{
-
-}
-/*/
\ No newline at end of file
+}
\ No newline at end of file
diff --git a/Engine/Plugins/FX/Niagara/Source/Niagara/Private/NiagaraDataInterfaceCollisionQuery.cpp b/Engine/Plugins/FX/Niagara/Source/Niagara/Private/NiagaraDataInterfaceCollisionQuery.cpp
--- a/Engine/Plugins/FX/Niagara/Source/Niagara/Private/NiagaraDataInterfaceCollisionQuery.cpp
+++ b/Engine/Plugins/FX/Niagara/Source/Niagara/Private/NiagaraDataInterfaceCollisionQuery.cpp
@@ -596,10 +596,11 @@ void UNiagaraDataInterfaceCollisionQuery::PerformQueryAsyncCPU(FVectorVMContext
 	for (int32 i = 0; i < Context.NumInstances; ++i)
 	{
 		FVector Pos(StartPosParamX.GetAndAdvance(), StartPosParamY.GetAndAdvance(), StartPosParamZ.GetAndAdvance());
-		FVector Dir(EndPosParamX.GetAndAdvance(), EndPosParamY.GetAndAdvance(), EndPosParamZ.GetAndAdvance());
+		FVector End(EndPosParamX.GetAndAdvance(), EndPosParamY.GetAndAdvance(), EndPosParamZ.GetAndAdvance());
 		ECollisionChannel TraceChannel = TraceChannelParam.GetAndAdvance();
 		ensure(!Pos.ContainsNaN());
-		*OutQueryID.GetDestAndAdvance() = InstanceData->CollisionBatch.SubmitQuery(Pos, Dir, TraceChannel);
+
+		*OutQueryID.GetDestAndAdvance() = InstanceData->CollisionBatch.SubmitQuery(Pos, End, TraceChannel);
 
 		// try to retrieve a query with the supplied query ID
 		FNiagaraDICollsionQueryResult Res;
@@ -829,13 +830,16 @@ void UNiagaraDataInterfaceCollisionQuery::ReadQuery(FVectorVMContext& Context)
 
 bool UNiagaraDataInterfaceCollisionQuery::PerInstanceTick(void* PerInstanceData, FNiagaraSystemInstance* InSystemInstance, float DeltaSeconds)
 {
+	CQDIPerInstanceData* PIData = static_cast<CQDIPerInstanceData*>(PerInstanceData);
+	PIData->CollisionBatch.CollectResults();
+
 	return false;
 }
 
 bool UNiagaraDataInterfaceCollisionQuery::PerInstanceTickPostSimulate(void* PerInstanceData, FNiagaraSystemInstance* InSystemInstance, float DeltaSeconds)
 {
-	CQDIPerInstanceData *PIData = static_cast<CQDIPerInstanceData*>(PerInstanceData);
-	PIData->CollisionBatch.Tick(ENiagaraSimTarget::CPUSim);
+	CQDIPerInstanceData* PIData = static_cast<CQDIPerInstanceData*>(PerInstanceData);
+	PIData->CollisionBatch.DispatchQueries();
 	PIData->CollisionBatch.ClearWrite();
 	return false;
 }
