Niagara - At a setting to the niagara editor settings which defines the required system update module rather than having it hard coded in c++.

diff --git a/Engine/Plugins/FX/Niagara/Source/NiagaraEditor/Private/ViewModels/Stack/NiagaraStackScriptItemGroup.cpp b/Engine/Plugins/FX/Niagara/Source/NiagaraEditor/Private/ViewModels/Stack/NiagaraStackScriptItemGroup.cpp
--- a/Engine/Plugins/FX/Niagara/Source/NiagaraEditor/Private/ViewModels/Stack/NiagaraStackScriptItemGroup.cpp
+++ b/Engine/Plugins/FX/Niagara/Source/NiagaraEditor/Private/ViewModels/Stack/NiagaraStackScriptItemGroup.cpp
@@ -18,6 +18,7 @@
 #include "NiagaraStackEditorData.h"
 #include "NiagaraSystem.h"
 #include "NiagaraEmitterHandle.h"
+#include "NiagaraEditorSettings.h"
 
 #include "Internationalization/Internationalization.h"
 #include "ScopedTransaction.h"
@@ -25,6 +26,8 @@
 #include "Framework/Notifications/NotificationManager.h"
 #include "Internationalization/Internationalization.h"
 #include "DragAndDrop/AssetDragDropOp.h"
+#include "AssetRegistryModule.h"
+#include "Modules/ModuleManager.h"
 
 #define LOCTEXT_NAMESPACE "UNiagaraStackScriptItemGroup"
 
@@ -398,19 +401,17 @@ void UNiagaraStackScriptItemGroup::RefreshIssues(TArray<FStackIssue>& NewIssues)
 			// We need to make sure that System Update Scripts have the SystemLifecycle script for now.
 			// The factor ensures this, but older assets may not have it or it may have been removed accidentally.
 			// For now, treat this as an error and allow them to resolve.
-			FString ModulePath = TEXT("/Niagara/Modules/System/SystemLifeCycle.SystemLifeCycle");
-			FSoftObjectPath SystemUpdateScriptRef(ModulePath);
-			FAssetData ModuleScriptAsset;
-			ModuleScriptAsset.ObjectPath = SystemUpdateScriptRef.GetAssetPathName();
+			const FAssetRegistryModule& AssetRegistryModule = FModuleManager::LoadModuleChecked<FAssetRegistryModule>(TEXT("AssetRegistry"));
+			const FAssetData ModuleScriptAsset = AssetRegistryModule.Get().GetAssetByObjectPath(GetDefault<UNiagaraEditorSettings>()->RequiredSystemUpdateScript.GetAssetPathName());
 
 			TArray<UNiagaraNodeFunctionCall*> FoundCalls;
 			UNiagaraNodeOutput* MatchingOutputNode = Graph->FindOutputNode(ScriptUsage, ScriptUsageId);
-			if (!FNiagaraStackGraphUtilities::FindScriptModulesInStack(ModuleScriptAsset, *MatchingOutputNode, FoundCalls))
+			if (ModuleScriptAsset.IsValid() && !FNiagaraStackGraphUtilities::FindScriptModulesInStack(ModuleScriptAsset, *MatchingOutputNode, FoundCalls))
 			{
 				bForcedError = true;
 
-				FText FixDescription = LOCTEXT("AddingSystemLifecycleModule", "Adding System Lifecycle Module.");
-				FStackIssueFix AddLifeCycleFix(
+				FText FixDescription = FText::Format(LOCTEXT("AddingRequiredSystemUpdateModuleFormat", "Add {0} module."), FText::FromName(ModuleScriptAsset.AssetName));
+				FStackIssueFix AddRequiredSystemUpdateModule(
 					FixDescription,
 					FStackIssueFixDelegate::CreateLambda([=]()
 				{
@@ -418,7 +419,7 @@ void UNiagaraStackScriptItemGroup::RefreshIssues(TArray<FStackIssue>& NewIssues)
 					UNiagaraNodeFunctionCall* AddedModuleNode = FNiagaraStackGraphUtilities::AddScriptModuleToStack(ModuleScriptAsset, *MatchingOutputNode);
 					if (AddedModuleNode == nullptr)
 					{
-						FNotificationInfo Info(LOCTEXT("FailedToAddSystemLifecycle", "Failed to add system life cycle module.\nCheck the log for errors."));
+						FNotificationInfo Info(LOCTEXT("FailedToAddSystemLifecycle", "Failed to add required module.\nCheck the log for errors."));
 						Info.ExpireDuration = 5.0f;
 						Info.bFireAndForget = true;
 						Info.Image = FCoreStyle::Get().GetBrush(TEXT("MessageLog.Error"));
@@ -426,15 +427,15 @@ void UNiagaraStackScriptItemGroup::RefreshIssues(TArray<FStackIssue>& NewIssues)
 					}
 				}));
 
-				FStackIssue MissingLifeCycleError(
+				FStackIssue MissingRequiredSystemUpdateModuleError(
 					EStackIssueSeverity::Error,
-					LOCTEXT("SystemLifeCycleWarning", "The stack needs a SystemLifeCycle module."),
+					FText::Format(LOCTEXT("RequiredSystemUpdateModuleWarningFormat", "The stack needs a {0} module."), FText::FromName(ModuleScriptAsset.AssetName)),
 					LOCTEXT("MissingRequiredMode", "Missing required module."),
 					GetStackEditorDataKey(),
 					false,
-					AddLifeCycleFix);
+					AddRequiredSystemUpdateModule);
 
-				NewIssues.Add(MissingLifeCycleError);
+				NewIssues.Add(MissingRequiredSystemUpdateModuleError);
 			}
 		}
 
diff --git a/Engine/Plugins/FX/Niagara/Source/NiagaraEditor/Public/NiagaraEditorSettings.h b/Engine/Plugins/FX/Niagara/Source/NiagaraEditor/Public/NiagaraEditorSettings.h
--- a/Engine/Plugins/FX/Niagara/Source/NiagaraEditor/Public/NiagaraEditorSettings.h
+++ b/Engine/Plugins/FX/Niagara/Source/NiagaraEditor/Public/NiagaraEditorSettings.h
@@ -49,6 +49,10 @@ public:
 	UPROPERTY(config, EditAnywhere, Category = Niagara)
 	FSoftObjectPath DefaultModuleScript;
 
+	/** Niagara script which is required in the system update script to control system state. */
+	UPROPERTY(config, EditAnywhere, Category = Niagara)
+	FSoftObjectPath RequiredSystemUpdateScript;
+
 	/** Shortcut key bindings that if held down while doing a mouse click, will spawn the specified type of Niagara node.*/
 	UPROPERTY(config, EditAnywhere, Category = Niagara)
 	TArray<FNiagaraSpawnShortcut> GraphCreationShortcuts;
