Niagara - Fix an issue where the dynamic material parameter names could get out of sync.

diff --git a/Engine/Plugins/FX/Niagara/Source/Niagara/Private/NiagaraMeshRendererProperties.cpp b/Engine/Plugins/FX/Niagara/Source/Niagara/Private/NiagaraMeshRendererProperties.cpp
--- a/Engine/Plugins/FX/Niagara/Source/Niagara/Private/NiagaraMeshRendererProperties.cpp
+++ b/Engine/Plugins/FX/Niagara/Source/Niagara/Private/NiagaraMeshRendererProperties.cpp
@@ -146,7 +146,7 @@ void UNiagaraMeshRendererProperties::GetUsedMaterials(const FNiagaraEmitterInsta
 					// UserParamBinding, if mapped to a real value, always wins. Otherwise, use the ExplictMat if it is set. Finally, fall
 					// back to the particle mesh material. This allows the user to effectively optionally bind to a Material binding
 					// and still have good defaults if it isn't set to anything.
-					if (OverrideMaterials[Section.MaterialIndex].UserParamBinding.Parameter.IsValid() && InEmitter->FindBinding(OverrideMaterials[Section.MaterialIndex].UserParamBinding, OutMaterials))
+					if (InEmitter != nullptr && OverrideMaterials[Section.MaterialIndex].UserParamBinding.Parameter.IsValid() && InEmitter->FindBinding(OverrideMaterials[Section.MaterialIndex].UserParamBinding, OutMaterials))
 					{
 						bSet = true;
 					}
diff --git a/Engine/Plugins/FX/Niagara/Source/Niagara/Private/NiagaraRibbonRendererProperties.cpp b/Engine/Plugins/FX/Niagara/Source/Niagara/Private/NiagaraRibbonRendererProperties.cpp
--- a/Engine/Plugins/FX/Niagara/Source/Niagara/Private/NiagaraRibbonRendererProperties.cpp
+++ b/Engine/Plugins/FX/Niagara/Source/Niagara/Private/NiagaraRibbonRendererProperties.cpp
@@ -59,7 +59,7 @@ void UNiagaraRibbonRendererProperties::GetUsedMaterials(const FNiagaraEmitterIns
 {
 
 	bool bSet = false;
-	if (MaterialUserParamBinding.Parameter.IsValid() && InEmitter->FindBinding(MaterialUserParamBinding, OutMaterials))
+	if (InEmitter != nullptr && MaterialUserParamBinding.Parameter.IsValid() && InEmitter->FindBinding(MaterialUserParamBinding, OutMaterials))
 	{
 		bSet = true;
 	}
diff --git a/Engine/Plugins/FX/Niagara/Source/Niagara/Private/NiagaraSpriteRendererProperties.cpp b/Engine/Plugins/FX/Niagara/Source/Niagara/Private/NiagaraSpriteRendererProperties.cpp
--- a/Engine/Plugins/FX/Niagara/Source/Niagara/Private/NiagaraSpriteRendererProperties.cpp
+++ b/Engine/Plugins/FX/Niagara/Source/Niagara/Private/NiagaraSpriteRendererProperties.cpp
@@ -65,7 +65,7 @@ FNiagaraBoundsCalculator* UNiagaraSpriteRendererProperties::CreateBoundsCalculat
 void UNiagaraSpriteRendererProperties::GetUsedMaterials(const FNiagaraEmitterInstance* InEmitter, TArray<UMaterialInterface*>& OutMaterials) const
 {
 	bool bSet = false;
-	if (MaterialUserParamBinding.Parameter.IsValid() && InEmitter->FindBinding(MaterialUserParamBinding, OutMaterials))
+	if (InEmitter != nullptr && MaterialUserParamBinding.Parameter.IsValid() && InEmitter->FindBinding(MaterialUserParamBinding, OutMaterials))
 	{
 		bSet = true;
 	}
diff --git a/Engine/Plugins/FX/Niagara/Source/NiagaraEditor/Private/ViewModels/Stack/NiagaraStackFunctionInput.cpp b/Engine/Plugins/FX/Niagara/Source/NiagaraEditor/Private/ViewModels/Stack/NiagaraStackFunctionInput.cpp
--- a/Engine/Plugins/FX/Niagara/Source/NiagaraEditor/Private/ViewModels/Stack/NiagaraStackFunctionInput.cpp
+++ b/Engine/Plugins/FX/Niagara/Source/NiagaraEditor/Private/ViewModels/Stack/NiagaraStackFunctionInput.cpp
@@ -496,9 +496,9 @@ void UNiagaraStackFunctionInput::RefreshChildrenInternal(const TArray<UNiagaraSt
 class UNiagaraStackFunctionInputUtilities
 {
 public:
-	static bool GetMaterialExpressionDynamicParameter(FNiagaraEmitterInstance* InEmitterInstance, TArray<UMaterialExpressionDynamicParameter*>& OutDynamicParameterExpressions)
+	static bool GetMaterialExpressionDynamicParameter(UNiagaraEmitter* InEmitter, FNiagaraEmitterInstance* InEmitterInstance, TArray<UMaterialExpressionDynamicParameter*>& OutDynamicParameterExpressions)
 	{
-		TArray<UMaterial*> Materials = GetMaterialFromEmitter(InEmitterInstance);
+		TArray<UMaterial*> Materials = GetMaterialFromEmitter(InEmitter, InEmitterInstance);
 		
 		OutDynamicParameterExpressions.Reset();
 
@@ -523,28 +523,24 @@ public:
 		return OutDynamicParameterExpressions.Num() > 0;
 	}
 
-	static TArray<UMaterial*> GetMaterialFromEmitter(FNiagaraEmitterInstance* InEmitterInstance)
+	static TArray<UMaterial*> GetMaterialFromEmitter(UNiagaraEmitter* InEmitter, FNiagaraEmitterInstance* InEmitterInstance)
 	{
 		TArray<UMaterial*> ResultMaterials;
-		if (InEmitterInstance)
+		if (InEmitter != nullptr)
 		{
-			UNiagaraEmitter* InEmitter = InEmitterInstance->GetCachedEmitter();
-			if (InEmitter && InEmitter->GetRenderers().Num() > 0)
+			for (UNiagaraRendererProperties* RenderProperties : InEmitter->GetRenderers())
 			{
-				for (UNiagaraRendererProperties* RenderProperties : InEmitter->GetRenderers())
+				TArray<UMaterialInterface*> UsedMaterialInteraces;
+				RenderProperties->GetUsedMaterials(InEmitterInstance, UsedMaterialInteraces);
+				for (UMaterialInterface* UsedMaterialInterface : UsedMaterialInteraces)
 				{
-					TArray<UMaterialInterface*> UsedMaterialInteraces;
-					RenderProperties->GetUsedMaterials(InEmitterInstance, UsedMaterialInteraces);
-					for (UMaterialInterface* UsedMaterialInterface : UsedMaterialInteraces)
+					if (UsedMaterialInterface != nullptr)
 					{
-						if (UsedMaterialInterface != nullptr)
+						UMaterial* UsedMaterial = UsedMaterialInterface->GetBaseMaterial();
+						if (UsedMaterial != nullptr)
 						{
-							UMaterial* UsedMaterial = UsedMaterialInterface->GetBaseMaterial();
-							if (UsedMaterial != nullptr)
-							{
-								ResultMaterials.AddUnique(UsedMaterial);
-								break;
-							}
+							ResultMaterials.AddUnique(UsedMaterial);
+							break;
 						}
 					}
 				}
@@ -561,7 +557,7 @@ FString UNiagaraStackFunctionInput::ResolveDisplayNameArgument(const FString& In
 		TSharedPtr<FNiagaraEmitterViewModel> ThisEmitterViewModel = GetEmitterViewModel();
 		TArray<UMaterialExpressionDynamicParameter*> ExpressionParams;
 		FNiagaraEmitterInstance* Instance = ThisEmitterViewModel->GetSimulation().Pin().Get();
-		if (ThisEmitterViewModel.IsValid() == false || !UNiagaraStackFunctionInputUtilities::GetMaterialExpressionDynamicParameter(Instance, ExpressionParams))
+		if (ThisEmitterViewModel.IsValid() == false || !UNiagaraStackFunctionInputUtilities::GetMaterialExpressionDynamicParameter(ThisEmitterViewModel->GetEmitter(), Instance, ExpressionParams))
 		{
 			return InArg.Replace(TEXT("MaterialDynamic"), TEXT("")) + TEXT(" (No material found using dynamic params)");
 		}
