- Generate optimized VM ByteCode async

#rnx

diff --git a/Engine/Plugins/FX/Niagara/Source/Niagara/Classes/NiagaraScript.h b/Engine/Plugins/FX/Niagara/Source/Niagara/Classes/NiagaraScript.h
--- a/Engine/Plugins/FX/Niagara/Source/Niagara/Classes/NiagaraScript.h
+++ b/Engine/Plugins/FX/Niagara/Source/Niagara/Classes/NiagaraScript.h
@@ -195,7 +195,6 @@ public:
 	void GetScriptHash(FSHAHash& OutHash) const;
 #endif
 
-#if WITH_EDITORONLY_DATA
 	/**
 	* Tests this set against another for equality, disregarding override settings.
 	*
@@ -209,6 +208,7 @@ public:
 		return !(*this == ReferenceSet);
 	}
 
+#if WITH_EDITORONLY_DATA
 	/** Appends string representations of this Id to a key string. */
 	void AppendKeyString(FString& KeyString) const;
 #endif
@@ -606,6 +606,9 @@ private:
 	/** Return the expected SimTarget for this script. Only returns a valid target if there is valid data to run with. */
 	TOptional<ENiagaraSimTarget> GetSimTarget() const;
 
+	/** Kicks off an async job to convert the ByteCode into an optimized version for the platform we are running on. */
+	void AsyncOptimizeByteCode();
+
 #if WITH_EDITORONLY_DATA
 	UPROPERTY(Transient)
 	FNiagaraScriptExecutionParameterStore ScriptExecutionParamStoreCPU;
diff --git a/Engine/Plugins/FX/Niagara/Source/Niagara/Private/NiagaraScript.cpp b/Engine/Plugins/FX/Niagara/Source/Niagara/Private/NiagaraScript.cpp
--- a/Engine/Plugins/FX/Niagara/Source/Niagara/Private/NiagaraScript.cpp
+++ b/Engine/Plugins/FX/Niagara/Source/Niagara/Private/NiagaraScript.cpp
@@ -32,6 +32,8 @@
 #include "NiagaraFunctionLibrary.h"
 #include "VectorVM.h"
 
+#include "Async/Async.h"
+
 DECLARE_STATS_GROUP(TEXT("Niagara Detailed"), STATGROUP_NiagaraDetailed, STATCAT_Advanced);
 
 FNiagaraScriptDebuggerInfo::FNiagaraScriptDebuggerInfo() : bWaitForGPU(false), FrameLastWriteId(-1), bWritten(false)
@@ -109,7 +111,6 @@ bool FNiagaraVMExecutableDataId::RequiresPersistentIDs() const
 	return AdditionalDefines.Contains("RequiresPersistentIDs");
 }
 
-#if WITH_EDITORONLY_DATA
 /**
 * Tests this set against another for equality, disregarding override settings.
 *
@@ -122,13 +123,16 @@ bool FNiagaraVMExecutableDataId::operator==(const FNiagaraVMExecutableDataId& Re
 		ScriptUsageType != ReferenceSet.ScriptUsageType || 
 		ScriptUsageTypeID != ReferenceSet.ScriptUsageTypeID ||
 		BaseScriptID != ReferenceSet.BaseScriptID ||
+#if WITH_EDITORONLY_DATA
 		BaseScriptCompileHash != ReferenceSet.BaseScriptCompileHash ||
+#endif
 		DetailLevelMask != ReferenceSet.DetailLevelMask ||
 		bUsesRapidIterationParams != ReferenceSet.bUsesRapidIterationParams)
 	{
 		return false;
 	}
 
+#if WITH_EDITORONLY_DATA
 	if (ReferencedCompileHashes.Num() != ReferenceSet.ReferencedCompileHashes.Num())
 	{
 		return false;
@@ -141,6 +145,7 @@ bool FNiagaraVMExecutableDataId::operator==(const FNiagaraVMExecutableDataId& Re
 			return false;
 		}
 	}
+#endif
 
 	if (AdditionalDefines.Num() != ReferenceSet.AdditionalDefines.Num())
 	{
@@ -161,6 +166,7 @@ bool FNiagaraVMExecutableDataId::operator==(const FNiagaraVMExecutableDataId& Re
 	return true;
 }
 
+#if WITH_EDITORONLY_DATA
 void FNiagaraVMExecutableDataId::AppendKeyString(FString& KeyString) const
 {
 	KeyString += FString::Printf(TEXT("%d_"), (int32)ScriptUsageType);
@@ -465,6 +471,60 @@ TOptional<ENiagaraSimTarget> UNiagaraScript::GetSimTarget() const
 	return TOptional<ENiagaraSimTarget>();
 }
 
+void UNiagaraScript::AsyncOptimizeByteCode()
+{
+	if ( !CachedScriptVM.IsValid() || (CachedScriptVM.OptimizedByteCode.Num() > 0) || (CachedScriptVM.ByteCode.Num() == 0) )
+	{
+		return;
+	}
+
+	static const IConsoleVariable* CVarOptimizeVMCode = IConsoleManager::Get().FindConsoleVariable(TEXT("vm.OptimizeVMByteCode"));
+	if (!CVarOptimizeVMCode || CVarOptimizeVMCode->GetInt() == 0 )
+	{
+		return;
+	}
+
+	// This has to be done game code side as we can not access anything in CachedScriptVM
+	TArray<uint8, TInlineAllocator<32>> ExternalFunctionRegisterCounts;
+	ExternalFunctionRegisterCounts.Reserve(CachedScriptVM.CalledVMExternalFunctions.Num());
+	for (const FVMExternalFunctionBindingInfo FunctionBindingInfo : CachedScriptVM.CalledVMExternalFunctions)
+	{
+		const uint8 RegisterCount = FunctionBindingInfo.GetNumInputs() + FunctionBindingInfo.GetNumOutputs();
+		ExternalFunctionRegisterCounts.Add(RegisterCount);
+	}
+
+	// Kick off async task to optimize the byte code, note we must take a copy of the ByteCode & VMId currently as we must guard the async task against changes in the editor
+	AsyncTask(
+		ENamedThreads::AnyThread,
+		[WeakScript=TWeakObjectPtr<UNiagaraScript>(this), InExternalFunctionRegisterCounts=MoveTemp(ExternalFunctionRegisterCounts), InByteCode=CachedScriptVM.ByteCode, InCachedScriptVMId=CachedScriptVMId]() mutable
+		{
+			// Generate optimized byte code on any thread
+			TArray<uint8> OptimizedByteCode;
+			VectorVM::OptimizeByteCode(InByteCode.GetData(), OptimizedByteCode, MakeArrayView(InExternalFunctionRegisterCounts));
+
+			// Kick off task to set optimized byte code on game thread
+			AsyncTask(
+				ENamedThreads::GameThread,
+				[WeakScript, InOptimizedByteCode = MoveTemp(OptimizedByteCode), InCachedScriptVMId]() mutable
+				{
+					UNiagaraScript* NiagaraScript = WeakScript.Get();
+					if ( (NiagaraScript != nullptr) && (NiagaraScript->CachedScriptVMId == InCachedScriptVMId) )
+					{
+						NiagaraScript->CachedScriptVM.OptimizedByteCode = MoveTemp(InOptimizedByteCode);
+
+						static const IConsoleVariable* CVarFreeUnoptimizedByteCode = IConsoleManager::Get().FindConsoleVariable(TEXT("vm.FreeUnoptimizedByteCode"));
+						if ( CVarFreeUnoptimizedByteCode && (CVarFreeUnoptimizedByteCode->GetInt() != 0) && (NiagaraScript->CachedScriptVM.OptimizedByteCode.Num() > 0) )
+						{
+							//-TODO: Support this, we will have to be careful that the VM is not currently executing the ByteCode before we free it
+							//NiagaraScript->CachedScriptVM.ByteCode.Empty();
+						}
+					}
+				}
+			);
+		}
+	);
+}
+
 void UNiagaraScript::PreSave(const class ITargetPlatform* TargetPlatform)
 {
 	Super::PreSave(TargetPlatform);
@@ -732,19 +792,7 @@ void UNiagaraScript::PostLoad()
 #endif
 
 	// Optimize the VM script for runtime usage
-	if ( FPlatformProperties::RequiresCookedData() && CachedScriptVM.IsValid() )
-	{
-		TArray<uint8, TInlineAllocator<32>> ExternalFunctionRegisterCounts;
-		ExternalFunctionRegisterCounts.Reserve(CachedScriptVM.CalledVMExternalFunctions.Num());
-
-		for ( const FVMExternalFunctionBindingInfo FunctionBindingInfo : CachedScriptVM.CalledVMExternalFunctions )
-		{
-			const uint8 RegisterCount = FunctionBindingInfo.GetNumInputs() + FunctionBindingInfo.GetNumOutputs();
-			ExternalFunctionRegisterCounts.Add(RegisterCount);
-		}
-
-		VectorVM::OptimizeByteCode(CachedScriptVM.ByteCode.GetData(), CachedScriptVM.OptimizedByteCode, MakeArrayView(ExternalFunctionRegisterCounts));
-	}
+	AsyncOptimizeByteCode();
 
 	//FNiagaraUtilities::DumpHLSLText(RapidIterationParameters.ToString(), *GetPathName());
 }
@@ -1156,6 +1204,8 @@ void UNiagaraScript::SetVMCompilationResults(const FNiagaraVMExecutableDataId& I
 
 	InvalidateExecutionReadyParameterStores();
 	
+	AsyncOptimizeByteCode();
+
 	OnVMScriptCompiled().Broadcast(this);
 }
 
diff --git a/Engine/Source/Runtime/VectorVM/Private/VectorVM.cpp b/Engine/Source/Runtime/VectorVM/Private/VectorVM.cpp
--- a/Engine/Source/Runtime/VectorVM/Private/VectorVM.cpp
+++ b/Engine/Source/Runtime/VectorVM/Private/VectorVM.cpp
@@ -182,6 +182,14 @@ static FAutoConsoleVariableRef CVarbOptimizeVMByteCode(
 	ECVF_Default
 );
 
+static int32 GbFreeUnoptimizedVMByteCode = 0;
+static FAutoConsoleVariableRef CVarbFreeUnoptimizedVMByteCode(
+	TEXT("vm.FreeUnoptimizedByteCode"),
+	GbFreeUnoptimizedVMByteCode,
+	TEXT("When we have optimized the VM byte code should we free the original unoptimized byte code?"),
+	ECVF_Default
+);
+
 static int32 GbUseOptimizedVMByteCode = 0;
 static FAutoConsoleVariableRef CVarbUseOptimizedVMByteCode(
 	TEXT("vm.UseOptimizedVMByteCode"),
