- Various memory savings for NiagaraScript
- Saves ~33% in a simple test level (713.61kb -> 478.70kb)

diff --git a/Engine/Plugins/FX/Niagara/Source/Niagara/Classes/NiagaraScript.h b/Engine/Plugins/FX/Niagara/Source/Niagara/Classes/NiagaraScript.h
--- a/Engine/Plugins/FX/Niagara/Source/Niagara/Classes/NiagaraScript.h
+++ b/Engine/Plugins/FX/Niagara/Source/Niagara/Classes/NiagaraScript.h
@@ -130,9 +130,11 @@ public:
 	UPROPERTY()
 	FGuid ScriptUsageTypeID;
 
+#if WITH_EDITORONLY_DATA
 	/** Configuration options*/
 	UPROPERTY()
 	TArray<FString> AdditionalDefines;
+#endif
 
 	/** Bitfield of supported detail levels in this compile.*/
 	UPROPERTY()
@@ -140,7 +142,15 @@ public:
 
 	/** Whether or not we need to bake Rapid Iteration params. True to keep params, false to bake.*/
 	UPROPERTY()
-	bool bUsesRapidIterationParams;
+	uint32 bUsesRapidIterationParams : 1;
+
+	/** Do we require interpolated spawning */
+	UPROPERTY()
+	uint32 bInterpolatedSpawn : 1;
+
+	/** Do we require persistent IDs */
+	UPROPERTY()
+	uint32 bRequiresPersistentIDs : 1;
 
 	/**
 	* The GUID of the subgraph this shader primarily represents.
@@ -168,6 +178,8 @@ public:
 		, ScriptUsageType(ENiagaraScriptUsage::Function)
 		, DetailLevelMask(0xFFFFFFFF)
 		, bUsesRapidIterationParams(true)
+		, bInterpolatedSpawn(false)
+		, bRequiresPersistentIDs(false)
 		, BaseScriptID(0, 0, 0, 0)
 	{ }
 
@@ -256,8 +268,10 @@ public:
 	UPROPERTY()
 	FNiagaraScriptDataUsageInfo DataUsage;
 
+#if WITH_EDITORONLY_DATA
 	UPROPERTY()
 	TArray<FNiagaraFunctionSignature> AdditionalExternalFunctions;
+#endif
 
 	/** Information about all data interfaces used by this script. */
 	UPROPERTY()
@@ -407,7 +421,11 @@ public:
 #endif
 
 	NIAGARA_API void ComputeVMCompilationId(FNiagaraVMExecutableDataId& Id) const;
+#if WITH_EDITORONLY_DATA
 	NIAGARA_API const FNiagaraVMExecutableDataId& GetComputedVMCompilationId() const { return LastGeneratedVMId; }
+#else
+	NIAGARA_API const FNiagaraVMExecutableDataId& GetComputedVMCompilationId() const { return CachedScriptVMId; }
+#endif
 
 	void SetUsage(ENiagaraScriptUsage InUsage) { Usage = InUsage; }
 	ENiagaraScriptUsage GetUsage() const { return Usage; }
@@ -643,9 +661,11 @@ private:
 	UPROPERTY()
 	FNiagaraVMExecutableDataId CachedScriptVMId;
 
+#if WITH_EDITORONLY_DATA
 	/** Adjusted every time ComputeVMCompilationId is called.*/
 	UPROPERTY()
 	mutable FNiagaraVMExecutableDataId LastGeneratedVMId;
+#endif
 
 	TArray<FNiagaraShaderScript> LoadedScriptResources;
 
diff --git a/Engine/Plugins/FX/Niagara/Source/Niagara/Private/NiagaraCommon.cpp b/Engine/Plugins/FX/Niagara/Source/Niagara/Private/NiagaraCommon.cpp
--- a/Engine/Plugins/FX/Niagara/Source/Niagara/Private/NiagaraCommon.cpp
+++ b/Engine/Plugins/FX/Niagara/Source/Niagara/Private/NiagaraCommon.cpp
@@ -407,4 +407,31 @@ FNiagaraUserParameterBinding::FNiagaraUserParameterBinding()
 	: Parameter(FNiagaraTypeDefinition::GetUObjectDef(), NAME_None)
 	{
 
-	}
\ No newline at end of file
+	}
+
+//////////////////////////////////////////////////////////////////////////
+
+bool FVMExternalFunctionBindingInfo::Serialize(FArchive& Ar)
+{
+	Ar.UsingCustomVersion(FNiagaraCustomVersion::GUID);
+
+	if (Ar.IsLoading() || Ar.IsSaving())
+	{
+		UScriptStruct* Struct = FVMExternalFunctionBindingInfo::StaticStruct();
+		Struct->SerializeTaggedProperties(Ar, (uint8*)this, Struct, nullptr);
+	}
+
+#if WITH_EDITORONLY_DATA
+	const int32 NiagaraVersion = Ar.CustomVer(FNiagaraCustomVersion::GUID);
+
+	if (NiagaraVersion < FNiagaraCustomVersion::MemorySaving)
+	{
+		for (auto it = Specifiers_DEPRECATED.CreateConstIterator(); it; ++it)
+		{
+			FunctionSpecifiers.Emplace(it->Key, it->Value);
+		}
+	}
+#endif
+
+	return true;
+}
diff --git a/Engine/Plugins/FX/Niagara/Source/Niagara/Private/NiagaraDataInterfaceParticleRead.cpp b/Engine/Plugins/FX/Niagara/Source/Niagara/Private/NiagaraDataInterfaceParticleRead.cpp
--- a/Engine/Plugins/FX/Niagara/Source/Niagara/Private/NiagaraDataInterfaceParticleRead.cpp
+++ b/Engine/Plugins/FX/Niagara/Source/Niagara/Private/NiagaraDataInterfaceParticleRead.cpp
@@ -215,123 +215,92 @@ void UNiagaraDataInterfaceParticleRead::GetVMExternalFunction(const FVMExternalF
 {
 	bool bBindSuccessful = false;
 	FNDIParticleRead_InstanceData* PIData = static_cast<FNDIParticleRead_InstanceData*>(InstanceData);
-	const FName* AttributeToRead = BindingInfo.Specifiers.Find(FName("Attribute"));
+	static const FName NAME_Attribute("Attribute");
 
-	if (!AttributeToRead)
+	const FVMFunctionSpecifier* FunctionSpecifier = BindingInfo.FindSpecifier(NAME_Attribute);
+	if (FunctionSpecifier == nullptr)
 	{
 		UE_LOG(LogNiagara, Error, TEXT("VMExternalFunction '%s' does not have a function specifier 'attribute'!"), *BindingInfo.Name.ToString());
 		return;
 	}
 
+	const FName AttributeToRead = FunctionSpecifier->Value;
 	if (BindingInfo.Name == GetFloatAttributeFunctionName)
 	{
-		if (AttributeToRead)
+		FNiagaraVariable VariableToRead(FNiagaraTypeDefinition::GetFloatDef(), AttributeToRead);
+		if (PIData->EmitterInstance->GetData().GetVariables().Find(VariableToRead) != INDEX_NONE)
 		{
-			FNiagaraVariable VariableToRead(FNiagaraTypeDefinition::GetFloatDef(), *AttributeToRead);
-			FName Attr = *AttributeToRead;
-			if (PIData->EmitterInstance->GetData().GetVariables().Find(VariableToRead) != INDEX_NONE)
-			{
-				NDI_FUNC_BINDER(UNiagaraDataInterfaceParticleRead, ReadFloat)::Bind(this, OutFunc, Attr);
-				bBindSuccessful = true;
-			}
+			NDI_FUNC_BINDER(UNiagaraDataInterfaceParticleRead, ReadFloat)::Bind(this, OutFunc, AttributeToRead);
+			bBindSuccessful = true;
 		}
 	}
 	else if (BindingInfo.Name == GetVec2AttributeFunctionName)
 	{
-		if (AttributeToRead)
+		FNiagaraVariable VariableToRead(FNiagaraTypeDefinition::GetVec2Def(), AttributeToRead);
+		if (PIData->EmitterInstance->GetData().GetVariables().Find(VariableToRead) != INDEX_NONE)
 		{
-			FNiagaraVariable VariableToRead(FNiagaraTypeDefinition::GetVec2Def(), *AttributeToRead);
-			FName Attr = *AttributeToRead;
-			if (PIData->EmitterInstance->GetData().GetVariables().Find(VariableToRead) != INDEX_NONE)
-			{
-				NDI_FUNC_BINDER(UNiagaraDataInterfaceParticleRead, ReadVector2)::Bind(this, OutFunc, Attr);
-				bBindSuccessful = true;
-			}
+			NDI_FUNC_BINDER(UNiagaraDataInterfaceParticleRead, ReadVector2)::Bind(this, OutFunc, AttributeToRead);
+			bBindSuccessful = true;
 		}
 	}
 	else if (BindingInfo.Name == GetVec3AttributeFunctionName)
 	{
-		if (AttributeToRead)
+		FNiagaraVariable VariableToRead(FNiagaraTypeDefinition::GetVec3Def(), AttributeToRead);
+		if (PIData->EmitterInstance->GetData().GetVariables().Find(VariableToRead) != INDEX_NONE)
 		{
-			FNiagaraVariable VariableToRead(FNiagaraTypeDefinition::GetVec3Def(), *AttributeToRead);
-			FName Attr = *AttributeToRead;
-			if (PIData->EmitterInstance->GetData().GetVariables().Find(VariableToRead) != INDEX_NONE)
-			{
-				NDI_FUNC_BINDER(UNiagaraDataInterfaceParticleRead, ReadVector3)::Bind(this, OutFunc, Attr);
-				bBindSuccessful = true;
-			}
+			NDI_FUNC_BINDER(UNiagaraDataInterfaceParticleRead, ReadVector3)::Bind(this, OutFunc, AttributeToRead);
+			bBindSuccessful = true;
 		}
 	}
 	else if (BindingInfo.Name == GetVec4AttributeFunctionName)
 	{
-		if (AttributeToRead)
+		FNiagaraVariable VariableToRead(FNiagaraTypeDefinition::GetVec4Def(), AttributeToRead);
+		if (PIData->EmitterInstance->GetData().GetVariables().Find(VariableToRead) != INDEX_NONE)
 		{
-			FNiagaraVariable VariableToRead(FNiagaraTypeDefinition::GetVec4Def(), *AttributeToRead);
-			FName Attr = *AttributeToRead;
-			if (PIData->EmitterInstance->GetData().GetVariables().Find(VariableToRead) != INDEX_NONE)
-			{
-				NDI_FUNC_BINDER(UNiagaraDataInterfaceParticleRead, ReadVector4)::Bind(this, OutFunc, Attr);
-				bBindSuccessful = true;
-			}
+			NDI_FUNC_BINDER(UNiagaraDataInterfaceParticleRead, ReadVector4)::Bind(this, OutFunc, AttributeToRead);
+			bBindSuccessful = true;
 		}
 	}
 	else if (BindingInfo.Name == GetIntAttributeFunctionName)
 	{
-		if (AttributeToRead)
+		FNiagaraVariable VariableToRead(FNiagaraTypeDefinition::GetIntDef(), AttributeToRead);
+		if (PIData->EmitterInstance->GetData().GetVariables().Find(VariableToRead) != INDEX_NONE)
 		{
-			FNiagaraVariable VariableToRead(FNiagaraTypeDefinition::GetIntDef(), *AttributeToRead);
-			FName Attr = *AttributeToRead;
-			if (PIData->EmitterInstance->GetData().GetVariables().Find(VariableToRead) != INDEX_NONE)
-			{
-				NDI_FUNC_BINDER(UNiagaraDataInterfaceParticleRead, ReadInt)::Bind(this, OutFunc, Attr);
-				bBindSuccessful = true;
-			}
+			NDI_FUNC_BINDER(UNiagaraDataInterfaceParticleRead, ReadInt)::Bind(this, OutFunc, AttributeToRead);
+			bBindSuccessful = true;
 		}
 	}
 	else if (BindingInfo.Name == GetBoolAttributeFunctionName)
 	{
-		if (AttributeToRead)
+		FNiagaraVariable VariableToRead(FNiagaraTypeDefinition::GetBoolDef(), AttributeToRead);
+		if (PIData->EmitterInstance->GetData().GetVariables().Find(VariableToRead) != INDEX_NONE)
 		{
-			FNiagaraVariable VariableToRead(FNiagaraTypeDefinition::GetBoolDef(), *AttributeToRead);
-			FName Attr = *AttributeToRead;
-			if (PIData->EmitterInstance->GetData().GetVariables().Find(VariableToRead) != INDEX_NONE)
-			{
-				NDI_FUNC_BINDER(UNiagaraDataInterfaceParticleRead, ReadBool)::Bind(this, OutFunc, Attr);
-				bBindSuccessful = true;
-			}
+			NDI_FUNC_BINDER(UNiagaraDataInterfaceParticleRead, ReadBool)::Bind(this, OutFunc, AttributeToRead);
+			bBindSuccessful = true;
 		}
 	}
 	else if (BindingInfo.Name == GetColorAttributeFunctionName)
 	{
-		if (AttributeToRead)
+		FNiagaraVariable VariableToRead(FNiagaraTypeDefinition::GetColorDef(), AttributeToRead);
+		if (PIData->EmitterInstance->GetData().GetVariables().Find(VariableToRead) != INDEX_NONE)
 		{
-			FNiagaraVariable VariableToRead(FNiagaraTypeDefinition::GetColorDef(), *AttributeToRead);
-			FName Attr = *AttributeToRead;
-			if (PIData->EmitterInstance->GetData().GetVariables().Find(VariableToRead) != INDEX_NONE)
-			{
-				NDI_FUNC_BINDER(UNiagaraDataInterfaceParticleRead, ReadColor)::Bind(this, OutFunc, Attr);
-				bBindSuccessful = true;
-			}
+			NDI_FUNC_BINDER(UNiagaraDataInterfaceParticleRead, ReadColor)::Bind(this, OutFunc, AttributeToRead);
+			bBindSuccessful = true;
 		}
 	}
 	else if (BindingInfo.Name == GetQuatAttributeFunctionName)
 	{
-		if (AttributeToRead)
+		FNiagaraVariable VariableToRead(FNiagaraTypeDefinition::GetQuatDef(), AttributeToRead);
+		if (PIData->EmitterInstance->GetData().GetVariables().Find(VariableToRead) != INDEX_NONE)
 		{
-			FNiagaraVariable VariableToRead(FNiagaraTypeDefinition::GetQuatDef(), *AttributeToRead);
-			FName Attr = *AttributeToRead;
-			if (PIData->EmitterInstance->GetData().GetVariables().Find(VariableToRead) != INDEX_NONE)
-			{
-				NDI_FUNC_BINDER(UNiagaraDataInterfaceParticleRead, ReadQuat)::Bind(this, OutFunc, Attr);
-				bBindSuccessful = true;
-			}
+			NDI_FUNC_BINDER(UNiagaraDataInterfaceParticleRead, ReadQuat)::Bind(this, OutFunc, AttributeToRead);
+			bBindSuccessful = true;
 		}
 	}
 
-
 	if (!bBindSuccessful)
 	{
-		UE_LOG(LogNiagara, Error, TEXT("Failed to bind VMExternalFunction '%s' with attribute '%s'! Check that the attribute is named correctly."), *BindingInfo.Name.ToString(), *(*AttributeToRead).ToString());
+		UE_LOG(LogNiagara, Error, TEXT("Failed to bind VMExternalFunction '%s' with attribute '%s'! Check that the attribute is named correctly."), *BindingInfo.Name.ToString(), *AttributeToRead.ToString());
 	}
 }
 
diff --git a/Engine/Plugins/FX/Niagara/Source/Niagara/Private/NiagaraModule.cpp b/Engine/Plugins/FX/Niagara/Source/Niagara/Private/NiagaraModule.cpp
--- a/Engine/Plugins/FX/Niagara/Source/Niagara/Private/NiagaraModule.cpp
+++ b/Engine/Plugins/FX/Niagara/Source/Niagara/Private/NiagaraModule.cpp
@@ -555,7 +555,7 @@ TArray<FNiagaraTypeDefinition> FNiagaraTypeRegistry::RegisteredNumericTypes;
 
 bool FNiagaraTypeDefinition::IsDataInterface()const
 {
-	return Struct->IsChildOf(UNiagaraDataInterface::StaticClass());
+	return GetStruct()->IsChildOf(UNiagaraDataInterface::StaticClass());
 }
 
 void FNiagaraTypeDefinition::Init()
@@ -876,6 +876,35 @@ FNiagaraTypeDefinition FNiagaraTypeDefinition::GetNumericOutputType(const TArray
 	return FNiagaraTypeDefinition::GetGenericNumericDef();
 }
 
+bool FNiagaraTypeDefinition::Serialize(FArchive& Ar)
+{
+	Ar.UsingCustomVersion(FNiagaraCustomVersion::GUID);
+	return false;
+}
+
+void FNiagaraTypeDefinition::PostSerialize(const FArchive& Ar)
+{
+#if WITH_EDITORONLY_DATA
+	if (Ar.IsLoading() && Ar.CustomVer(FNiagaraCustomVersion::GUID) < FNiagaraCustomVersion::MemorySaving)
+	{
+		if (Enum_DEPRECATED != nullptr)
+		{
+			UnderlyingType = UT_Enum;
+			ClassStructOrEnum = Enum_DEPRECATED;
+		}
+		else if (Struct_DEPRECATED != nullptr)
+		{
+			UnderlyingType = Struct_DEPRECATED->IsA<UClass>() ? UT_Class : UT_Struct;
+			ClassStructOrEnum = Struct_DEPRECATED;
+		}
+		else
+		{
+			UnderlyingType = UT_None;
+			ClassStructOrEnum = nullptr;
+		}
+	}
+#endif
+}
 
 //////////////////////////////////////////////////////////////////////////
 
diff --git a/Engine/Plugins/FX/Niagara/Source/Niagara/Private/NiagaraParameterStore.cpp b/Engine/Plugins/FX/Niagara/Source/Niagara/Private/NiagaraParameterStore.cpp
--- a/Engine/Plugins/FX/Niagara/Source/Niagara/Private/NiagaraParameterStore.cpp
+++ b/Engine/Plugins/FX/Niagara/Source/Niagara/Private/NiagaraParameterStore.cpp
@@ -24,6 +24,55 @@ static FAutoConsoleVariableRef CVarNiagaraDumpParticleParameterStores(
 );
 #endif
 
+struct FNiagaraVariableSearch
+{
+	static FORCEINLINE int32 Compare(const FNiagaraVariableBase& A, const FNiagaraVariableBase& B)
+	{
+#if NIAGARA_VARIABLE_LEXICAL_SORTING
+		int32 ComparisonDiff = A.GetName().Compare(B.GetName());
+#else
+		int32 ComparisonDiff = A.GetName().CompareIndexes(B.GetName());
+#endif
+		if (ComparisonDiff != 0)
+		{
+			return ComparisonDiff;
+		}
+		else
+		{
+#if NIAGARA_VARIABLE_LEXICAL_SORTING
+			return ComparisonDiff = A.GetType().GetFName().Compare(B.GetType().GetFName());
+#else
+			return ComparisonDiff = A.GetType().GetFName().CompareIndexes(B.GetType().GetFName());
+#endif
+		}
+	}
+
+	static bool Find(const FNiagaraVariableWithOffset* Variables, const FNiagaraVariableBase& Ref, int32 Start, int32 Num, int32& CheckIndex)
+	{
+		while (Num)
+		{
+			const int32 LeftoverSize = Num % 2;
+			Num = Num / 2;
+
+			CheckIndex = Start + Num;
+			const int32 StartIfLess = CheckIndex + LeftoverSize;
+
+			const int32 ComparisonDiff = Compare(Variables[CheckIndex], Ref);
+			if (ComparisonDiff < 0)
+			{
+				Start = CheckIndex + 1;
+				Num += LeftoverSize - 1;
+			}
+			else if (ComparisonDiff == 0)
+			{
+				return true;
+			}
+		}
+		CheckIndex = Start;
+		return false;
+	}
+};
+
 //////////////////////////////////////////////////////////////////////////
 
 int32 GNiagaraAllowQuickSortedParameterOffetsCopy = 1;
@@ -63,31 +112,6 @@ namespace
 
 //////////////////////////////////////////////////////////////////////////
 
-bool FNiagaraVariableSearch::Find(const FNiagaraVariableWithOffset* Variables, const FNiagaraVariable& Ref, int32 Start, int32 Num, int32& CheckIndex)
-{
-	while (Num)
-	{
-		const int32 LeftoverSize = Num % 2;
-		Num = Num / 2;
-
-		CheckIndex = Start + Num;
-		const int32 StartIfLess = CheckIndex + LeftoverSize;
-
-		const int32 ComparisonDiff = Compare(Variables[CheckIndex], Ref);
-		if (ComparisonDiff < 0)
-		{
-			Start = CheckIndex + 1;
-			Num += LeftoverSize - 1;
-		}
-		else if (ComparisonDiff == 0)
-		{
-			return true;
-		}
-	}
-	CheckIndex = Start;
-	return false;
-}
-
 FNiagaraParameterStore::FNiagaraParameterStore()
 	: Owner(nullptr)
 	, bParametersDirty(true)
@@ -344,10 +368,9 @@ bool FNiagaraParameterStore::VerifyBinding(const FNiagaraParameterStore* DestSto
 
 void FNiagaraParameterStore::CheckForNaNs()const
 {
-	for (const FNiagaraVariableWithOffset& ParamWithOffset : SortedParameterOffsets)
+	for (const FNiagaraVariableWithOffset& Var : SortedParameterOffsets)
 	{
-		const FNiagaraVariable& Var = ParamWithOffset;
-		const int32 Offset = ParamWithOffset.Offset;
+		const int32 Offset = Var.Offset;
 
 		bool bContainsNans = false;
 		if (Var.GetType() == FNiagaraTypeDefinition::GetFloatDef())
@@ -876,7 +899,7 @@ void FNiagaraParameterStore::OnLayoutChange()
 #endif
 }
 
-const FNiagaraVariable* FNiagaraParameterStore::FindVariable(UNiagaraDataInterface* Interface)const
+const FNiagaraVariableBase* FNiagaraParameterStore::FindVariable(UNiagaraDataInterface* Interface)const
 {
 	SCOPE_CYCLE_COUNTER(STAT_NiagaraParameterStoreFindVar);
 	int32 Idx = DataInterfaces.IndexOfByKey(Interface);
@@ -922,11 +945,6 @@ void FNiagaraParameterStore::PostLoad()
 		}
 		ParameterOffsets.Empty();
 	}
-	// Safeguard while we don't have yet the FNiagaraVariable type without data
-	for (FNiagaraVariableWithOffset& ParamWithOffset : SortedParameterOffsets)
-	{
-		ParamWithOffset.ClearData();
-	}
 #endif
 
 	// Not always required if NIAGARA_VARIABLE_LEXICAL_SORTING
diff --git a/Engine/Plugins/FX/Niagara/Source/Niagara/Private/NiagaraScript.cpp b/Engine/Plugins/FX/Niagara/Source/Niagara/Private/NiagaraScript.cpp
--- a/Engine/Plugins/FX/Niagara/Source/Niagara/Private/NiagaraScript.cpp
+++ b/Engine/Plugins/FX/Niagara/Source/Niagara/Private/NiagaraScript.cpp
@@ -103,12 +103,12 @@ void FNiagaraVMExecutableDataId::Invalidate()
 
 bool FNiagaraVMExecutableDataId::HasInterpolatedParameters() const
 {
-	return AdditionalDefines.Contains("InterpolatedSpawn");
+	return bInterpolatedSpawn;
 }
 
 bool FNiagaraVMExecutableDataId::RequiresPersistentIDs() const
 {
-	return AdditionalDefines.Contains("RequiresPersistentIDs");
+	return bRequiresPersistentIDs;
 }
 
 /**
@@ -127,7 +127,9 @@ bool FNiagaraVMExecutableDataId::operator==(const FNiagaraVMExecutableDataId& Re
 		BaseScriptCompileHash != ReferenceSet.BaseScriptCompileHash ||
 #endif
 		DetailLevelMask != ReferenceSet.DetailLevelMask ||
-		bUsesRapidIterationParams != ReferenceSet.bUsesRapidIterationParams)
+		bUsesRapidIterationParams != ReferenceSet.bUsesRapidIterationParams ||
+		bInterpolatedSpawn != ReferenceSet.bInterpolatedSpawn ||
+		bRequiresPersistentIDs != ReferenceSet.bRequiresPersistentIDs )
 	{
 		return false;
 	}
@@ -145,7 +147,6 @@ bool FNiagaraVMExecutableDataId::operator==(const FNiagaraVMExecutableDataId& Re
 			return false;
 		}
 	}
-#endif
 
 	if (AdditionalDefines.Num() != ReferenceSet.AdditionalDefines.Num())
 	{
@@ -161,6 +162,7 @@ bool FNiagaraVMExecutableDataId::operator==(const FNiagaraVMExecutableDataId& Re
 			return false;
 		}
 	}
+#endif
 
 
 	return true;
@@ -262,6 +264,8 @@ void UNiagaraScript::ComputeVMCompilationId(FNiagaraVMExecutableDataId& Id) cons
 	Id = FNiagaraVMExecutableDataId();
 
 	Id.bUsesRapidIterationParams = true;
+	Id.bInterpolatedSpawn = false;
+	Id.bRequiresPersistentIDs = false;
 	Id.DetailLevelMask = 0xFFFFFFFF; // Unused for now.
 	
 	// Ideally we wouldn't want to do this but rather than push the data down
@@ -279,10 +283,12 @@ void UNiagaraScript::ComputeVMCompilationId(FNiagaraVMExecutableDataId& Id) cons
 			(Emitter->bInterpolatedSpawning && Usage == ENiagaraScriptUsage::ParticleSpawnScript) ||
 			Usage == ENiagaraScriptUsage::ParticleSpawnScriptInterpolated)
 		{
+			Id.bInterpolatedSpawn = true;
 			Id.AdditionalDefines.Add(TEXT("InterpolatedSpawn"));
 		}
 		if (Emitter->RequiresPersistantIDs())
 		{
+			Id.bRequiresPersistentIDs = true;
 			Id.AdditionalDefines.Add(TEXT("RequiresPersistentIDs"));
 		}
 		if (Emitter->bLocalSpace)
@@ -1693,7 +1699,7 @@ NIAGARA_API bool UNiagaraScript::IsScriptCompilationPending(bool bGPUScript) con
 	{
 		if (CachedScriptVM.IsValid())
 		{
-			return CachedScriptVM.ByteCode.Num() == 0 && (CachedScriptVM.LastCompileStatus == ENiagaraScriptCompileStatus::NCS_BeingCreated || CachedScriptVM.LastCompileStatus == ENiagaraScriptCompileStatus::NCS_Unknown);
+			return (CachedScriptVM.ByteCode.Num() == 0) && (CachedScriptVM.OptimizedByteCode.Num() == 0) && (CachedScriptVM.LastCompileStatus == ENiagaraScriptCompileStatus::NCS_BeingCreated || CachedScriptVM.LastCompileStatus == ENiagaraScriptCompileStatus::NCS_Unknown);
 		}
 		return false;
 	}
@@ -1719,7 +1725,7 @@ NIAGARA_API bool UNiagaraScript::DidScriptCompilationSucceed(bool bGPUScript) co
 	{
 		if (CachedScriptVM.IsValid())
 		{
-			return CachedScriptVM.ByteCode.Num() != 0;
+			return (CachedScriptVM.ByteCode.Num() != 0) || (CachedScriptVM.OptimizedByteCode.Num() != 0);
 		}
 	}
 
diff --git a/Engine/Plugins/FX/Niagara/Source/Niagara/Private/NiagaraScriptExecutionParameterStore.cpp b/Engine/Plugins/FX/Niagara/Source/Niagara/Private/NiagaraScriptExecutionParameterStore.cpp
--- a/Engine/Plugins/FX/Niagara/Source/Niagara/Private/NiagaraScriptExecutionParameterStore.cpp
+++ b/Engine/Plugins/FX/Niagara/Source/Niagara/Private/NiagaraScriptExecutionParameterStore.cpp
@@ -295,8 +295,8 @@ void FNiagaraScriptExecutionParameterStore::CopyParameterDataToPaddedBuffer(uint
 	const uint8* SrcData = GetParameterDataArray().GetData();
 	for (int32 i = 0; i < PaddingInfo.Num(); i++)
 	{
-		check((PaddingInfo[i].DestOffset + PaddingInfo[i].DestSize) <= InTargetBufferSizeInBytes);
-		check((PaddingInfo[i].SrcOffset + PaddingInfo[i].SrcSize) <= (uint32)GetParameterDataArray().Num());
+		check(uint32(PaddingInfo[i].DestOffset + PaddingInfo[i].DestSize) <= InTargetBufferSizeInBytes);
+		check(uint32(PaddingInfo[i].SrcOffset + PaddingInfo[i].SrcSize) <= (uint32)GetParameterDataArray().Num());
 		FMemory::Memcpy(InTargetBuffer + PaddingInfo[i].DestOffset, SrcData + PaddingInfo[i].SrcOffset, PaddingInfo[i].SrcSize);
 	}
 }
diff --git a/Engine/Plugins/FX/Niagara/Source/Niagara/Public/NiagaraCommon.h b/Engine/Plugins/FX/Niagara/Source/Niagara/Public/NiagaraCommon.h
--- a/Engine/Plugins/FX/Niagara/Source/Niagara/Public/NiagaraCommon.h
+++ b/Engine/Plugins/FX/Niagara/Source/Niagara/Public/NiagaraCommon.h
@@ -252,21 +252,23 @@ struct NIAGARA_API FNiagaraFunctionSignature
 	UPROPERTY()
 	FName OwnerName;
 	UPROPERTY()
-	bool bRequiresContext;
+	uint32 bRequiresContext : 1;
 	/** True if this is the signature for a "member" function of a data interface. If this is true, the first input is the owner. */
 	UPROPERTY()
-	bool bMemberFunction;
+	uint32 bMemberFunction : 1;
 	/** Is this function experimental? */
 	UPROPERTY()
-	bool bExperimental;
-	/** Function specifiers verified at bind time. */
+	uint32 bExperimental : 1;
+	/** Support running on the CPU. */
 	UPROPERTY()
-	TMap<FName, FName> FunctionSpecifiers;
-
+	uint32 bSupportsCPU : 1;
+	/** Support running on the GPU. */
 	UPROPERTY()
-	bool bSupportsCPU = true;
+	uint32 bSupportsGPU : 1;
+
+	/** Function specifiers verified at bind time. */
 	UPROPERTY()
-	bool bSupportsGPU = true;
+	TMap<FName, FName> FunctionSpecifiers;
 
 	/** Localized description of this node. Note that this is *not* used during the operator == below since it may vary from culture to culture.*/
 #if WITH_EDITORONLY_DATA
@@ -278,6 +280,8 @@ struct NIAGARA_API FNiagaraFunctionSignature
 		: bRequiresContext(false)
 		, bMemberFunction(false)
 		, bExperimental(false)
+		, bSupportsCPU(true)
+		, bSupportsGPU(true)
 	{
 	}
 
@@ -288,7 +292,8 @@ struct NIAGARA_API FNiagaraFunctionSignature
 		, bRequiresContext(bInRequiresContext)
 		, bMemberFunction(bInMemberFunction)
 		, bExperimental(false)
-		, FunctionSpecifiers()
+		, bSupportsCPU(true)
+		, bSupportsGPU(true)
 	{
 
 	}
@@ -300,6 +305,8 @@ struct NIAGARA_API FNiagaraFunctionSignature
 		, bRequiresContext(bInRequiresContext)
 		, bMemberFunction(bInMemberFunction)
 		, bExperimental(false)
+		, bSupportsCPU(true)
+		, bSupportsGPU(true)
 		, FunctionSpecifiers(InFunctionSpecifiers)
 	{
 
@@ -416,8 +423,11 @@ public:
 	UPROPERTY()
 	FNiagaraTypeDefinition Type;
 
+	// Removed from cooked builds, if we need to add this back the TMap<FName, FName> FunctionSpecifiers should be replaced with an array
+#if WITH_EDITORONLY_DATA
 	UPROPERTY()
 	TArray<FNiagaraFunctionSignature> RegisteredFunctions;
+#endif
 
 	UPROPERTY()
 	FName RegisteredParameterMapRead;
@@ -452,6 +462,21 @@ struct FNiagaraStatScope
 	bool operator==(const FNiagaraStatScope& Other) const { return FullName == Other.FullName; }
 };
 
+USTRUCT()
+struct FVMFunctionSpecifier
+{
+	GENERATED_USTRUCT_BODY();
+
+	FVMFunctionSpecifier() {}
+	explicit FVMFunctionSpecifier(FName InKey, FName InValue) : Key(InKey), Value(InValue) {}
+
+	UPROPERTY()
+	FName Key;
+
+	UPROPERTY()
+	FName Value;
+};
+
 USTRUCT()
 struct FVMExternalFunctionBindingInfo
 {
@@ -470,10 +495,32 @@ struct FVMExternalFunctionBindingInfo
 	int32 NumOutputs;
 
 	UPROPERTY()
-	TMap<FName, FName> Specifiers;
+	TArray<FVMFunctionSpecifier> FunctionSpecifiers;
+
+	FORCEINLINE int32 GetNumInputs() const { return InputParamLocations.Num(); }
+	FORCEINLINE int32 GetNumOutputs() const { return NumOutputs; }
+
+	const FVMFunctionSpecifier* FindSpecifier(const FName& Key) const
+	{
+		return FunctionSpecifiers.FindByPredicate([&](const FVMFunctionSpecifier& v) -> bool { return v.Key == Key; });
+	}
 
-	FORCEINLINE int32 GetNumInputs()const { return InputParamLocations.Num(); }
-	FORCEINLINE int32 GetNumOutputs()const { return NumOutputs; }
+	bool Serialize(FArchive& Ar);
+
+#if WITH_EDITORONLY_DATA
+private:
+	UPROPERTY()
+	TMap<FName, FName> Specifiers_DEPRECATED;
+#endif
+};
+
+template<>
+struct TStructOpsTypeTraits<FVMExternalFunctionBindingInfo> : public TStructOpsTypeTraitsBase2<FVMExternalFunctionBindingInfo>
+{
+	enum
+	{
+		WithSerializer = true,
+	};
 };
 
 struct NIAGARA_API FNiagaraSystemUpdateContext
diff --git a/Engine/Plugins/FX/Niagara/Source/Niagara/Public/NiagaraParameterStore.h b/Engine/Plugins/FX/Niagara/Source/Niagara/Public/NiagaraParameterStore.h
--- a/Engine/Plugins/FX/Niagara/Source/Niagara/Public/NiagaraParameterStore.h
+++ b/Engine/Plugins/FX/Niagara/Source/Niagara/Public/NiagaraParameterStore.h
@@ -9,32 +9,6 @@
 // When not cooked, sort by actual name to ensure deterministic cooked data
 #define NIAGARA_VARIABLE_LEXICAL_SORTING WITH_EDITORONLY_DATA
 
-struct FNiagaraVariableSearch
-{
-	static FORCEINLINE int32 Compare(const FNiagaraVariable& A, const FNiagaraVariable& B)
-	{
-#if NIAGARA_VARIABLE_LEXICAL_SORTING
-		int32 ComparisonDiff = A.GetName().Compare(B.GetName());
-#else
-		int32 ComparisonDiff = A.GetName().CompareIndexes(B.GetName());
-#endif
-		if (ComparisonDiff != 0)
-		{
-			return ComparisonDiff;
-		}
-		else
-		{
-#if NIAGARA_VARIABLE_LEXICAL_SORTING
-			return ComparisonDiff = A.GetType().GetFName().Compare(B.GetType().GetFName());
-#else
-			return ComparisonDiff = A.GetType().GetFName().CompareIndexes(B.GetType().GetFName());
-#endif
-		}
-	}
-
-	static bool Find(const struct FNiagaraVariableWithOffset* Variables, const FNiagaraVariable& Ref, int32 Start, int32 Num, int32& CheckIndex);
-};
-
 struct FNiagaraParameterStore;
 
 USTRUCT()
@@ -140,14 +114,14 @@ private:
 
 
 USTRUCT()
-struct FNiagaraVariableWithOffset : public FNiagaraVariable
+struct FNiagaraVariableWithOffset : public FNiagaraVariableBase
 {
 	GENERATED_USTRUCT_BODY()
 
 	// Those constructor enforce that there are no data allocated.
 	FORCEINLINE FNiagaraVariableWithOffset() : Offset(INDEX_NONE) {}
-	FORCEINLINE FNiagaraVariableWithOffset(const FNiagaraVariableWithOffset& InRef) : FNiagaraVariable(InRef.GetType(), InRef.GetName()), Offset(InRef.Offset) {}
-	FORCEINLINE FNiagaraVariableWithOffset(const FNiagaraVariable& InVariable, int32 InOffset) : FNiagaraVariable(InVariable.GetType(), InVariable.GetName()), Offset(InOffset) {}
+	FORCEINLINE FNiagaraVariableWithOffset(const FNiagaraVariableWithOffset& InRef) : FNiagaraVariableBase(InRef.GetType(), InRef.GetName()), Offset(InRef.Offset) {}
+	FORCEINLINE FNiagaraVariableWithOffset(const FNiagaraVariableBase& InVariable, int32 InOffset) : FNiagaraVariableBase(InVariable.GetType(), InVariable.GetName()), Offset(InOffset) {}
 
 	UPROPERTY()
 	int32 Offset;
@@ -367,7 +341,7 @@ public:
 	}
 
 	/** Returns the associated FNiagaraVariable for the passed data interface if it exists in the store. Null if not.*/
-	const FNiagaraVariable* FindVariable(UNiagaraDataInterface* Interface)const;
+	const FNiagaraVariableBase* FindVariable(UNiagaraDataInterface* Interface)const;
 
 	virtual const int32* FindParameterOffset(const FNiagaraVariable& Parameter) const;
 
diff --git a/Engine/Plugins/FX/Niagara/Source/Niagara/Public/NiagaraParameters.h b/Engine/Plugins/FX/Niagara/Source/Niagara/Public/NiagaraParameters.h
--- a/Engine/Plugins/FX/Niagara/Source/Niagara/Public/NiagaraParameters.h
+++ b/Engine/Plugins/FX/Niagara/Source/Niagara/Public/NiagaraParameters.h
@@ -14,7 +14,7 @@ struct FNiagaraParameters
 	GENERATED_USTRUCT_BODY()
 	
 public:
-	//TODO: Make this a map?
+	//TODO: Sort the array so we can binary search, do not change to a TMap to avoid memory bloat!
 	UPROPERTY(EditAnywhere, EditFixedSize, Category = "Uniform")
 	TArray<FNiagaraVariable> Parameters;
 
diff --git a/Engine/Plugins/FX/Niagara/Source/Niagara/Public/NiagaraScriptExecutionParameterStore.h b/Engine/Plugins/FX/Niagara/Source/Niagara/Public/NiagaraScriptExecutionParameterStore.h
--- a/Engine/Plugins/FX/Niagara/Source/Niagara/Public/NiagaraScriptExecutionParameterStore.h
+++ b/Engine/Plugins/FX/Niagara/Source/Niagara/Public/NiagaraScriptExecutionParameterStore.h
@@ -16,13 +16,13 @@ public:
 	FNiagaraScriptExecutionPaddingInfo(uint32 InSrcOffset, uint32 InDestOffset, uint32 InSrcSize, uint32 InDestSize) : SrcOffset(InSrcOffset), DestOffset(InDestOffset), SrcSize(InSrcSize), DestSize(InDestSize) {}
 
 	UPROPERTY()
-	uint32 SrcOffset;
+	uint16 SrcOffset;
 	UPROPERTY()
-	uint32 DestOffset;
+	uint16 DestOffset;
 	UPROPERTY()
-	uint32 SrcSize;
+	uint16 SrcSize;
 	UPROPERTY()
-	uint32 DestSize;
+	uint16 DestSize;
 };
 
 /**
diff --git a/Engine/Plugins/FX/Niagara/Source/Niagara/Public/NiagaraTypes.h b/Engine/Plugins/FX/Niagara/Source/Niagara/Public/NiagaraTypes.h
--- a/Engine/Plugins/FX/Niagara/Source/Niagara/Public/NiagaraTypes.h
+++ b/Engine/Plugins/FX/Niagara/Source/Niagara/Public/NiagaraTypes.h
@@ -378,35 +378,59 @@ USTRUCT()
 struct NIAGARA_API FNiagaraTypeDefinition
 {
 	GENERATED_USTRUCT_BODY()
+
+	enum FUnderlyingType
+	{
+		UT_None = 0,
+		UT_Class,
+		UT_Struct,
+		UT_Enum
+	};
+
 public:
 
 	// Construct blank raw type definition 
 	FORCEINLINE FNiagaraTypeDefinition(UClass *ClassDef)
-		: Struct(ClassDef), Enum(nullptr), Size(INDEX_NONE), Alignment(INDEX_NONE)
+		: ClassStructOrEnum(ClassDef), UnderlyingType(UT_Class), Size(INDEX_NONE), Alignment(INDEX_NONE)
+#if WITH_EDITORONLY_DATA
+		, Struct_DEPRECATED(nullptr), Enum_DEPRECATED(nullptr)
+#endif
 	{
-		checkSlow(Struct != nullptr);
+		checkSlow(ClassDef != nullptr);
 	}
 
 	FORCEINLINE FNiagaraTypeDefinition(UEnum *EnumDef)
-		: Struct(IntStruct), Enum(EnumDef), Size(INDEX_NONE), Alignment(INDEX_NONE)
+		: ClassStructOrEnum(EnumDef), UnderlyingType(UT_Enum), Size(INDEX_NONE), Alignment(INDEX_NONE)
+#if WITH_EDITORONLY_DATA
+		, Struct_DEPRECATED(nullptr), Enum_DEPRECATED(nullptr)
+#endif
 	{
-		checkSlow(Struct != nullptr);
+		checkSlow(ClassDef != nullptr);
 	}
 
 	FORCEINLINE FNiagaraTypeDefinition(UScriptStruct *StructDef)
-		: Struct(StructDef), Enum(nullptr), Size(INDEX_NONE), Alignment(INDEX_NONE)
+		: ClassStructOrEnum(StructDef), UnderlyingType(UT_Struct), Size(INDEX_NONE), Alignment(INDEX_NONE)
+#if WITH_EDITORONLY_DATA
+		, Struct_DEPRECATED(nullptr), Enum_DEPRECATED(nullptr)
+#endif
 	{
-		checkSlow(Struct != nullptr);
+		checkSlow(ClassDef != nullptr);
 	}
 
 	FORCEINLINE FNiagaraTypeDefinition(const FNiagaraTypeDefinition &Other)
-		: Struct(Other.Struct), Enum(Other.Enum), Size(INDEX_NONE), Alignment(INDEX_NONE)
+		: ClassStructOrEnum(Other.ClassStructOrEnum), UnderlyingType(Other.UnderlyingType), Size(INDEX_NONE), Alignment(INDEX_NONE)
+#if WITH_EDITORONLY_DATA
+		, Struct_DEPRECATED(nullptr), Enum_DEPRECATED(nullptr)
+#endif
 	{
 	}
 
 	// Construct a blank raw type definition
 	FORCEINLINE FNiagaraTypeDefinition()
-		: Struct(nullptr), Enum(nullptr), Size(INDEX_NONE), Alignment(INDEX_NONE)
+		: ClassStructOrEnum(nullptr), UnderlyingType(UT_None), Size(INDEX_NONE), Alignment(INDEX_NONE)
+#if WITH_EDITORONLY_DATA
+		, Struct_DEPRECATED(nullptr), Enum_DEPRECATED(nullptr)
+#endif
 	{}
 
 	FORCEINLINE bool operator !=(const FNiagaraTypeDefinition &Other) const
@@ -416,7 +440,7 @@ public:
 
 	FORCEINLINE bool operator == (const FNiagaraTypeDefinition &Other) const
 	{
-		return Struct == Other.Struct && Enum == Other.Enum;
+		return ClassStructOrEnum == Other.ClassStructOrEnum && UnderlyingType == Other.UnderlyingType;
 	}
 
 	FText GetNameText()const
@@ -426,16 +450,13 @@ public:
 			return NSLOCTEXT("NiagaraTypeDefinition", "InvalidNameText", "Invalid (null type)");
 		}
 
-		if (Enum != nullptr)
+#if WITH_EDITOR
+		if ( UnderlyingType != UT_Enum )
 		{
-			return  FText::FromString(Enum->GetName());
+			return GetStruct()->GetDisplayNameText();
 		}
-
-#if WITH_EDITOR
-		return GetStruct()->GetDisplayNameText();
-#else
-		return FText::FromString( GetStruct()->GetName() );
 #endif
+		return FText::FromString(ClassStructOrEnum->GetName());
 	}
 
 	FName GetFName() const
@@ -444,11 +465,7 @@ public:
 		{
 			return FName();
 		}
-		if ( Enum != nullptr )
-		{
-			return Enum->GetFName();
-		}
-		return GetStruct()->GetFName();
+		return ClassStructOrEnum->GetFName();
 	}
 
 	FString GetName()const
@@ -457,45 +474,40 @@ public:
 		{
 			return TEXT("Invalid");
 		}
-
-		if (Enum != nullptr)
-		{
-			return  Enum->GetName();
-		}
-		return GetStruct()->GetName();
+		return ClassStructOrEnum->GetName();
 	}
 
-	UStruct* GetStruct()const
+	UStruct* GetStruct() const
 	{
-		return Struct;
+		return UnderlyingType == UT_Enum ? IntStruct : CastChecked<UStruct>(ClassStructOrEnum);
 	}
 
 	UScriptStruct* GetScriptStruct()const
 	{
-		return Cast<UScriptStruct>(Struct);
+		return Cast<UScriptStruct>(GetStruct());
 	}
 
 	/** Gets the class ptr for this type if it is a class. */
-	UClass* GetClass()const
+	UClass* GetClass() const
 	{
-		return Cast<UClass>(Struct);
+		return UnderlyingType == UT_Class ? CastChecked<UClass>(ClassStructOrEnum) : nullptr;
 	}
 
 	UEnum* GetEnum() const
 	{
-		return Enum;
+		return UnderlyingType == UT_Enum ? CastChecked<UEnum>(ClassStructOrEnum) : nullptr;
 	}
 
-	bool IsDataInterface()const;
+	bool IsDataInterface() const;
 
-	FORCEINLINE bool IsUObject()const
+	FORCEINLINE bool IsUObject() const
 	{
-		return Struct->IsChildOf<UObject>();
+		return GetStruct()->IsChildOf<UObject>();
 	}
 
-	bool IsEnum() const { return Enum != nullptr; }
+	bool IsEnum() const { return UnderlyingType == UT_Enum; }
 	
-	int32 GetSize()const
+	int32 GetSize() const
 	{
 		if (Size == INDEX_NONE)
 		{
@@ -506,13 +518,13 @@ public:
 			}
 			else
 			{
-				Size = CastChecked<UScriptStruct>(Struct)->GetStructureSize();	// TODO: cache this here?
+				Size = CastChecked<UScriptStruct>(GetStruct())->GetStructureSize();
 			}
 		}
 		return Size;
 	}
 
-	int32 GetAlignment()const
+	int32 GetAlignment() const
 	{
 		if (Alignment == INDEX_NONE)
 		{
@@ -523,7 +535,7 @@ public:
 			}
 			else
 			{
-				Alignment = CastChecked<UScriptStruct>(Struct)->GetMinAlignment();
+				Alignment = CastChecked<UScriptStruct>(GetStruct())->GetMinAlignment();
 			}
 		}
 		return Alignment;
@@ -531,32 +543,43 @@ public:
 
 	bool IsFloatPrimitive() const
 	{
-		return Struct == FNiagaraTypeDefinition::GetFloatStruct() || Struct == FNiagaraTypeDefinition::GetVec2Struct() || Struct == FNiagaraTypeDefinition::GetVec3Struct() || Struct == FNiagaraTypeDefinition::GetVec4Struct() ||
-			Struct == FNiagaraTypeDefinition::GetMatrix4Struct() || Struct == FNiagaraTypeDefinition::GetColorStruct() || Struct == FNiagaraTypeDefinition::GetQuatStruct();
+		return ClassStructOrEnum == FNiagaraTypeDefinition::GetFloatStruct() || ClassStructOrEnum == FNiagaraTypeDefinition::GetVec2Struct() || ClassStructOrEnum == FNiagaraTypeDefinition::GetVec3Struct() || ClassStructOrEnum == FNiagaraTypeDefinition::GetVec4Struct() ||
+			ClassStructOrEnum == FNiagaraTypeDefinition::GetMatrix4Struct() || ClassStructOrEnum == FNiagaraTypeDefinition::GetColorStruct() || ClassStructOrEnum == FNiagaraTypeDefinition::GetQuatStruct();
  	}
 
 	bool IsValid() const 
 	{ 
-		return Struct != nullptr;
+		return ClassStructOrEnum != nullptr;
 	}
 
-	/**
-	UStruct specifying the type for this variable.
-	For most types this will be a UScriptStruct pointing to a something like the struct for an FVector etc.
+	/*
+	Underlying type for this variable, use FUnderlyingType to determine type without casting
+	This can be a UClass, UStruct or UEnum.  Pointing to something like the struct for an FVector, etc.
 	In occasional situations this may be a UClass when we're dealing with DataInterface etc.
 	*/
 	UPROPERTY()
-	UStruct *Struct;
+	UObject* ClassStructOrEnum;
 
+	// See enumeration FUnderlyingType for possible values
 	UPROPERTY()
-	UEnum* Enum;
+	uint16 UnderlyingType;
+
+	bool Serialize(FArchive& Ar);
+	void PostSerialize(const FArchive& Ar);
 
 private:
 	mutable int16 Size;
 	mutable int16 Alignment;
 
-public:
+#if WITH_EDITORONLY_DATA
+	UPROPERTY()
+	UStruct* Struct_DEPRECATED;
+
+	UPROPERTY()
+	UEnum* Enum_DEPRECATED;
+#endif
 
+public:
 	static void Init();
 	static void RecreateUserDefinedTypeRegistry();
 	static const FNiagaraTypeDefinition& GetFloatDef() { return FloatDef; }
@@ -603,7 +626,7 @@ public:
 		{
 			return TEXT("(null)");
 		}
-		return FNiagaraTypeHelper::ToString(ValueData, CastChecked<UScriptStruct>(Struct));
+		return FNiagaraTypeHelper::ToString(ValueData, CastChecked<UScriptStruct>(ClassStructOrEnum));
 	}
 
 	static bool TypesAreAssignable(const FNiagaraTypeDefinition& TypeA, const FNiagaraTypeDefinition& TypeB);
@@ -661,7 +684,16 @@ private:
 	static TSet<UStruct*> BoolStructs;
 
 	static FNiagaraTypeDefinition CollisionEventDef;
+};
 
+template<>
+struct TStructOpsTypeTraits<FNiagaraTypeDefinition> : public TStructOpsTypeTraitsBase2<FNiagaraTypeDefinition>
+{
+	enum
+	{
+		WithSerializer = true,
+		WithPostSerialize = true,
+	};
 };
 
 /* Contains all types currently available for use in Niagara
@@ -775,19 +807,80 @@ FORCEINLINE uint32 GetTypeHash(const FNiagaraTypeDefinition& Type)
 //////////////////////////////////////////////////////////////////////////
 
 USTRUCT()
-struct FNiagaraVariable
+struct FNiagaraVariableBase
+{
+	GENERATED_USTRUCT_BODY()
+
+	FORCEINLINE FNiagaraVariableBase() : Name(NAME_None), TypeDef(FNiagaraTypeDefinition::GetVec4Def()) { }
+	FORCEINLINE FNiagaraVariableBase(const FNiagaraVariableBase &Other) : Name(Other.Name), TypeDef(Other.TypeDef) { }
+	FORCEINLINE FNiagaraVariableBase(const FNiagaraTypeDefinition& InType, const FName& InName) : Name(InName), TypeDef(InType) { }
+	
+	/** Check if Name and Type definition are the same. The actual stored value is not checked here.*/
+	bool operator==(const FNiagaraVariableBase& Other)const
+	{
+		return Name == Other.Name && TypeDef == Other.TypeDef;
+	}
+
+	/** Check if Name and Type definition are not the same. The actual stored value is not checked here.*/
+	bool operator!=(const FNiagaraVariableBase& Other)const
+	{
+		return !(*this == Other);
+	}
+
+	/** Variables are the same name but if types are auto-assignable, allow them to match. */
+	bool IsEquivalent(const FNiagaraVariableBase& Other, bool bAllowAssignableTypes = true)const
+	{
+		return Name == Other.Name && (TypeDef == Other.TypeDef || (bAllowAssignableTypes && FNiagaraTypeDefinition::TypesAreAssignable(TypeDef, Other.TypeDef)));
+	}
+	
+	FORCEINLINE void SetName(FName InName) { Name = InName; }
+	FORCEINLINE const FName& GetName() const { return Name; }
+
+	void SetType(const FNiagaraTypeDefinition& InTypeDef) { TypeDef = InTypeDef; }
+	const FNiagaraTypeDefinition& GetType()const { return TypeDef; }
+
+	FORCEINLINE bool IsDataInterface()const { return GetType().IsDataInterface(); }
+	FORCEINLINE bool IsUObject()const { return GetType().IsUObject(); }
+
+	int32 GetSizeInBytes() const
+	{
+		return TypeDef.GetSize();
+	}
+
+	int32 GetAlignment()const
+	{
+		return TypeDef.GetAlignment();
+	}
+
+	bool IsValid() const
+	{
+		return Name != NAME_None && TypeDef.IsValid();
+	}
+
+	FORCEINLINE bool IsInNameSpace(FString Namespace) const
+	{
+		return Name.ToString().StartsWith(Namespace + TEXT("."));
+	}
+
+protected:
+	UPROPERTY(EditAnywhere, Category = "Variable")
+	FName Name;
+
+	UPROPERTY(EditAnywhere, Category = "Variable")
+	FNiagaraTypeDefinition TypeDef;
+};
+
+USTRUCT()
+struct FNiagaraVariable : public FNiagaraVariableBase
 {
 	GENERATED_USTRUCT_BODY()
 
 	FNiagaraVariable()
-		: Name(NAME_None)
-		, TypeDef(FNiagaraTypeDefinition::GetVec4Def())
 	{
 	}
 
 	FNiagaraVariable(const FNiagaraVariable &Other)
-		: Name(Other.Name)
-		, TypeDef(Other.TypeDef)
+		: FNiagaraVariableBase(Other)
 	{
 		if (Other.IsDataAllocated())
 		{
@@ -795,15 +888,20 @@ struct FNiagaraVariable
 		}
 	}
 
+	FNiagaraVariable(const FNiagaraVariableBase& Other)
+		: FNiagaraVariableBase(Other)
+	{
+	}
+
 	FORCEINLINE FNiagaraVariable(const FNiagaraTypeDefinition& InType, const FName& InName)
-		: Name(InName)
-		, TypeDef(InType)
+		: FNiagaraVariableBase(InType, InName)
 	{
 	}
 	
 	/** Check if Name and Type definition are the same. The actual stored value is not checked here.*/
 	bool operator==(const FNiagaraVariable& Other)const
 	{
+		//-TODO: Should this check the value???
 		return Name == Other.Name && TypeDef == Other.TypeDef;
 	}
 
@@ -813,21 +911,7 @@ struct FNiagaraVariable
 		return !(*this == Other);
 	}
 
-	/** Variables are the same name but if types are auto-assignable, allow them to match. */
-	bool IsEquivalent(const FNiagaraVariable& Other, bool bAllowAssignableTypes = true)const
-	{
-		return Name == Other.Name && (TypeDef == Other.TypeDef || (bAllowAssignableTypes && FNiagaraTypeDefinition::TypesAreAssignable(TypeDef, Other.TypeDef)));
-	}
-	
-	FORCEINLINE void SetName(FName InName) { Name = InName; }
-	FORCEINLINE const FName& GetName() const { return Name; }
-
-	void SetType(const FNiagaraTypeDefinition& InTypeDef) { TypeDef = InTypeDef; }
-	const FNiagaraTypeDefinition& GetType()const { return TypeDef; }
-
-	FORCEINLINE bool IsDataInterface()const { return GetType().IsDataInterface(); }
-	FORCEINLINE bool IsUObject()const { return GetType().IsUObject(); }
-
+	// Var data operations
 	void AllocateData()
 	{
 		if (VarData.Num() != TypeDef.GetSize())
@@ -885,16 +969,6 @@ struct FNiagaraVariable
 		VarData.Empty();
 	}
 
-	int32 GetSizeInBytes() const
-	{
-		return TypeDef.GetSize();
-	}
-
-	int32 GetAlignment()const
-	{
-		return TypeDef.GetAlignment();
-	}
-
 	int32 GetAllocatedSizeInBytes() const
 	{
 		return VarData.Num();
@@ -908,16 +982,6 @@ struct FNiagaraVariable
 		return Ret;
 	}
 
-	bool IsValid() const
-	{
-		return Name != NAME_None && TypeDef.IsValid();
-	}
-
-	FORCEINLINE bool IsInNameSpace(FString Namespace) const
-	{
-		return Name.ToString().StartsWith(Namespace + TEXT("."));
-	}
-
 	static FNiagaraVariable ResolveAliases(const FNiagaraVariable& InVar, const TMap<FString, FString>& InAliases, const TCHAR* InJoinSeparator = TEXT("."))
 	{
 		FNiagaraVariable OutVar = InVar;
@@ -970,10 +1034,6 @@ struct FNiagaraVariable
 	}
 
 private:
-	UPROPERTY(EditAnywhere, Category = "Variable")
-	FName Name;
-	UPROPERTY(EditAnywhere, Category = "Variable")
-	FNiagaraTypeDefinition TypeDef;
 	//This gets serialized but do we need to worry about endianness doing things like this? If not, where does that get handled?
 	//TODO: Remove storage here entirely and move everything to an FNiagaraParameterStore.
 	UPROPERTY()
diff --git a/Engine/Plugins/FX/Niagara/Source/NiagaraCore/Private/NiagaraCustomVersion.cpp b/Engine/Plugins/FX/Niagara/Source/NiagaraCore/Private/NiagaraCustomVersion.cpp
--- a/Engine/Plugins/FX/Niagara/Source/NiagaraCore/Private/NiagaraCustomVersion.cpp
+++ b/Engine/Plugins/FX/Niagara/Source/NiagaraCore/Private/NiagaraCustomVersion.cpp
@@ -8,4 +8,5 @@ const FGuid FNiagaraCustomVersion::GUID(0xFCF57AFA, 0x50764283, 0xB9A9E658, 0xFF
 // Register the custom version with core
 FCustomVersionRegistration GRegisterNiagaraCustomVersion(FNiagaraCustomVersion::GUID, FNiagaraCustomVersion::LatestVersion, TEXT("NiagaraVer"));
 
-const FGuid FNiagaraCustomVersion::LatestScriptCompileVersion(0x7D607C23, 0x5E344266, 0x81518861, 0x9278156B);
\ No newline at end of file
+const FGuid FNiagaraCustomVersion::LatestScriptCompileVersion(0xDEA735C8, 0xD11A3B4A, 0xA4841FAF, 0xDF7D991E);
+
diff --git a/Engine/Plugins/FX/Niagara/Source/NiagaraCore/Public/NiagaraCustomVersion.h b/Engine/Plugins/FX/Niagara/Source/NiagaraCore/Public/NiagaraCustomVersion.h
--- a/Engine/Plugins/FX/Niagara/Source/NiagaraCore/Public/NiagaraCustomVersion.h
+++ b/Engine/Plugins/FX/Niagara/Source/NiagaraCore/Public/NiagaraCustomVersion.h
@@ -134,6 +134,8 @@ struct FNiagaraCustomVersion
 
 		DisableSortingByDefault, // Sorting on Renderers is disabled by default, we add a version to maintain existing systems that expected sorting to be enabled
 
+		MemorySaving, // Convert TMap into TArray to save memory, TMap contains an inline allocator which pushes the size to 80 bytes
+
 		// DO NOT ADD A NEW VERSION UNLESS YOU HAVE TALKED TO THE NIAGARA LEAD. Mismanagement of these versions can lead to data loss if it is adjusted in multiple streams simultaneously.
 		// -----<new versions can be added above this line>  -------------------------------------------------
 		VersionPlusOne,
diff --git a/Engine/Plugins/FX/Niagara/Source/NiagaraEditor/Private/Customizations/NiagaraStaticSwitchNodeDetails.cpp b/Engine/Plugins/FX/Niagara/Source/NiagaraEditor/Private/Customizations/NiagaraStaticSwitchNodeDetails.cpp
--- a/Engine/Plugins/FX/Niagara/Source/NiagaraEditor/Private/Customizations/NiagaraStaticSwitchNodeDetails.cpp
+++ b/Engine/Plugins/FX/Niagara/Source/NiagaraEditor/Private/Customizations/NiagaraStaticSwitchNodeDetails.cpp
@@ -361,7 +361,7 @@ void FNiagaraStaticSwitchNodeDetails::OnSelectionChanged(TSharedPtr<ConstantDrop
 	// in case of an enum constant we also need the enum type
 	if (SelectedDropdownItem == DropdownOptions[4])
 	{
-		Node->SwitchTypeData.Enum = NewValue->Constant.GetType().Enum;
+		Node->SwitchTypeData.Enum = NewValue->Constant.GetType().GetEnum();
 	}
 	else
 	{
diff --git a/Engine/Plugins/FX/Niagara/Source/NiagaraEditor/Private/NiagaraCompiler.cpp b/Engine/Plugins/FX/Niagara/Source/NiagaraEditor/Private/NiagaraCompiler.cpp
--- a/Engine/Plugins/FX/Niagara/Source/NiagaraEditor/Private/NiagaraCompiler.cpp
+++ b/Engine/Plugins/FX/Niagara/Source/NiagaraEditor/Private/NiagaraCompiler.cpp
@@ -990,7 +990,11 @@ FNiagaraCompileResults FHlslNiagaraCompiler::CompileScript(const FNiagaraCompile
 					int32 NewBindingIdx = CompileResults.Data->CalledVMExternalFunctions.AddDefaulted();
 					CompileResults.Data->CalledVMExternalFunctions[NewBindingIdx].Name = *Sig->GetName();
 					CompileResults.Data->CalledVMExternalFunctions[NewBindingIdx].OwnerName = Sig->OwnerName;
-					CompileResults.Data->CalledVMExternalFunctions[NewBindingIdx].Specifiers = Sig->FunctionSpecifiers;
+					CompileResults.Data->CalledVMExternalFunctions[NewBindingIdx].FunctionSpecifiers.Empty();
+					for ( auto it = Sig->FunctionSpecifiers.CreateConstIterator(); it; ++it )
+					{
+						CompileResults.Data->CalledVMExternalFunctions[NewBindingIdx].FunctionSpecifiers.Emplace(it->Key, it->Value);
+					}
 
 					CompileResults.Data->CalledVMExternalFunctions[NewBindingIdx].InputParamLocations = FuncInfo.InputParamLocations;
 					CompileResults.Data->CalledVMExternalFunctions[NewBindingIdx].NumOutputs = FuncInfo.NumOutputs;
diff --git a/Engine/Plugins/FX/Niagara/Source/NiagaraEditor/Private/ViewModels/NiagaraSystemViewModel.cpp b/Engine/Plugins/FX/Niagara/Source/NiagaraEditor/Private/ViewModels/NiagaraSystemViewModel.cpp
--- a/Engine/Plugins/FX/Niagara/Source/NiagaraEditor/Private/ViewModels/NiagaraSystemViewModel.cpp
+++ b/Engine/Plugins/FX/Niagara/Source/NiagaraEditor/Private/ViewModels/NiagaraSystemViewModel.cpp
@@ -1254,7 +1254,7 @@ void FNiagaraSystemViewModel::UpdateCompiledDataInterfaces(UNiagaraDataInterface
 	else
 	{
 		// If the data interface wasn't owned by a script, try to find it in the exposed parameter data interfaces.
-		const FNiagaraVariable* FoundExposedDataInterface = GetSystem().GetExposedParameters().FindVariable(ChangedDataInterface);
+		const FNiagaraVariableBase* FoundExposedDataInterface = GetSystem().GetExposedParameters().FindVariable(ChangedDataInterface);
 		if (FoundExposedDataInterface != nullptr)
 		{
 			GetSystem().GetExposedParameters().OnInterfaceChange();
