- Fix debug build

#rnx

diff --git a/Engine/Plugins/FX/Niagara/Source/Niagara/Private/NiagaraParameterStore.cpp b/Engine/Plugins/FX/Niagara/Source/Niagara/Private/NiagaraParameterStore.cpp
--- a/Engine/Plugins/FX/Niagara/Source/Niagara/Private/NiagaraParameterStore.cpp
+++ b/Engine/Plugins/FX/Niagara/Source/Niagara/Private/NiagaraParameterStore.cpp
@@ -89,17 +89,6 @@ namespace
 	{
 		if (GNiagaraAllowQuickSortedParameterOffetsCopy)
 		{
-		#if DO_GUARD_SLOW
-			// Safeguard while we don't have yet the FNiagaraVariable type without data.
-			for (const FNiagaraVariableWithOffset& ParamWithOffset : Dest)
-			{
-				checkSlow(!ParamWithOffset.IsDataAllocated());
-			}
-			for (const FNiagaraVariableWithOffset& ParamWithOffset : Src)
-			{
-				checkSlow(!ParamWithOffset.IsDataAllocated());
-			}
-		#endif
 			Dest.SetNumUninitialized(Src.Num());
 			FMemory::Memcpy(Dest.GetData(), Src.GetData(), Dest.GetTypeSize() * Dest.Num());
 		}
diff --git a/Engine/Plugins/FX/Niagara/Source/Niagara/Public/NiagaraTypes.h b/Engine/Plugins/FX/Niagara/Source/Niagara/Public/NiagaraTypes.h
--- a/Engine/Plugins/FX/Niagara/Source/Niagara/Public/NiagaraTypes.h
+++ b/Engine/Plugins/FX/Niagara/Source/Niagara/Public/NiagaraTypes.h
@@ -396,7 +396,7 @@ public:
 		, Struct_DEPRECATED(nullptr), Enum_DEPRECATED(nullptr)
 #endif
 	{
-		checkSlow(ClassDef != nullptr);
+		checkSlow(ClassStructOrEnum != nullptr);
 	}
 
 	FORCEINLINE FNiagaraTypeDefinition(UEnum *EnumDef)
@@ -405,7 +405,7 @@ public:
 		, Struct_DEPRECATED(nullptr), Enum_DEPRECATED(nullptr)
 #endif
 	{
-		checkSlow(ClassDef != nullptr);
+		checkSlow(ClassStructOrEnum != nullptr);
 	}
 
 	FORCEINLINE FNiagaraTypeDefinition(UScriptStruct *StructDef)
@@ -414,7 +414,7 @@ public:
 		, Struct_DEPRECATED(nullptr), Enum_DEPRECATED(nullptr)
 #endif
 	{
-		checkSlow(ClassDef != nullptr);
+		checkSlow(ClassStructOrEnum != nullptr);
 	}
 
 	FORCEINLINE FNiagaraTypeDefinition(const FNiagaraTypeDefinition &Other)
