Adds two new pairs of UNiagaraNodeOp nodes for generating integer and floating point random numbers.
This is because generating integer and floating point random numbers are semantically different due to how it handles the upper range of an interval, and the old "Random" and "Seeded Random" were inconsistent about this across CPU, GPU and determinism modes.
Floating point random number ranges are typically half-open, i.e. they do not include the upper limit. Integer random number ranges on the other hand do typically include the upper limit.

This adds the following ops that are consistent with the equivalent Blueprints functions:
 - Random Integer: Non-deterministic, produces numbers between 0 and Max-1.
 - Random Float: Non-deterministic, produces numbers between 0 and Max, but not including Max.
 - Seeded Random Integer: Deterministic, produces numbers between 0 and Max-1.
 - Seeded Random Float: Determinstic, produces numbers between 0 and Max, but not including Max.

In addition it includes two new helper functions, also consistent with the equivalent Blueprints functions:
 - Random Range Float, produces numbers between Min and Max, but not including Max.
 - Random Range Integer, produces numbers between Min and Max.

The old Random Range function or the old Random and Seeded Random ops are not touched.

diff --git a/Engine/Plugins/FX/Niagara/Shaders/Private/NiagaraEmitterInstanceShader.usf b/Engine/Plugins/FX/Niagara/Shaders/Private/NiagaraEmitterInstanceShader.usf
--- a/Engine/Plugins/FX/Niagara/Shaders/Private/NiagaraEmitterInstanceShader.usf
+++ b/Engine/Plugins/FX/Niagara/Shaders/Private/NiagaraEmitterInstanceShader.usf
@@ -257,6 +257,9 @@ int rand(int x, int Seed1, int Seed2, int Seed3)
 	}
 #else
 	// Old unseeded, passthrough to FRandomStream
+
+	float rand(float x); // Invokes EVectorVMOp::random
+
 	float2 rand(float2 x)
 	{
 		return float2(rand(x.x), rand(x.y));
@@ -272,10 +275,65 @@ int rand(int x, int Seed1, int Seed2, int Seed3)
 		return float4(rand(x.x), rand(x.y), rand(x.z), rand(x.w));
 	}
 	
-	// NOTE(mv): Where's rand(float)? Seems to work regardless..
-	int rand(int x);
+	int rand(int x); // Invokes EVectorVMOp::randomi and is semi-open. This is inconsistent with the rest of the functions above. As a result this function and the ones above should be deprecated in favor of the functions below
 #endif
 
+// Explicit non-deterministic random overrides used by Random Float/Integer and Seeded Random Float/Integer op nodes
+float rand_float(float x)
+{
+	return rand(x.x);
+}
+
+float2 rand_float(float2 x)
+{
+	return float2(rand_float(x.x), rand_float(x.y));
+}
+
+float3 rand_float(float3 x)
+{
+	return float3(rand_float(x.x), rand_float(x.y), rand_float(x.z));
+}
+
+float4 rand_float(float4 x) 
+{
+	return float4(rand_float(x.x), rand_float(x.y), rand_float(x.z), rand_float(x.w));
+}
+
+int rand_int(int x)
+{
+	// Going through the float function also give us a better distribution than using modulo 
+	// to get an integer range.
+	// This will not include the upper range as rand_float returns [0, max), not [0, max].
+	return rand_float(x.x);
+}
+
+// Explicit deterministic random overrides used by Random Float/Integer and Seeded Random Float/Integer op nodes
+float rand_float(float x, int Seed1, int Seed2, int Seed3)
+{
+	return rand(x.x, Seed1, Seed2, Seed3);
+}
+
+float2 rand_float(float2 x, int Seed1, int Seed2, int Seed3)
+{
+	return rand(x, Seed1, Seed2, Seed3);
+}
+
+float3 rand_float(float3 x, int Seed1, int Seed2, int Seed3)
+{
+	return rand(x, Seed1, Seed2, Seed3);
+}
+
+float4 rand_float(float4 x, int Seed1, int Seed2, int Seed3) 
+{
+	return rand(x, Seed1, Seed2, Seed3);
+}
+
+int rand_int(int x, int Seed1, int Seed2, int Seed3)
+{
+	// This will not include the upper range as rand_float returns [0, max), not [0, max]
+	// The deterministic rand call will include the upper range, so we subtract a one to compensate
+	return rand(x.x-1, Seed1, Seed2, Seed3);
+}
 
 /* -----------------------------------------------------------------
  * VM simulation function declarations
diff --git a/Engine/Plugins/FX/Niagara/Source/NiagaraEditor/Private/NiagaraEditorCommon.cpp b/Engine/Plugins/FX/Niagara/Source/NiagaraEditor/Private/NiagaraEditorCommon.cpp
--- a/Engine/Plugins/FX/Niagara/Source/NiagaraEditor/Private/NiagaraEditorCommon.cpp
+++ b/Engine/Plugins/FX/Niagara/Source/NiagaraEditor/Private/NiagaraEditorCommon.cpp
@@ -786,6 +786,31 @@ void FNiagaraOpInfo::Init()
 		Op->BuildName(TEXT("Rand"), CategoryName);
 		OpInfoMap.Add(Op->Name) = Idx;
 
+		// Non-deterministic integer random number generation. Calls FRandomStream on the CPU. 
+		Idx = OpInfos.AddDefaulted();
+		Op = &OpInfos[Idx];
+		Op->Category = CategoryText;
+		Op->FriendlyName = NSLOCTEXT("NiagaraOpInfo", "Rand Integer Name", "Random Integer");
+		Op->Description = NSLOCTEXT("NiagaraOpInfo", "Rand Integer Desc", "Returns a non-deterministic random integer value between 0 and Max-1");
+		Op->Inputs.Add(FNiagaraOpInOutInfo(Max, Type, MaxText, MaxText, DefaultStr_One));
+		Op->Outputs.Add(FNiagaraOpInOutInfo(Result, Type, ResultText, ResultText, DefaultStr_One, TEXT("rand_int({0})")));
+		Op->BuildName(TEXT("Rand Integer"), CategoryName);
+		Op->bNumericsCanBeIntegers = true;
+		Op->bNumericsCanBeFloats = false;
+		OpInfoMap.Add(Op->Name) = Idx;
+
+		// Non-deterministic float random number generation. Calls FRandomStream on the CPU. 
+		Idx = OpInfos.AddDefaulted();
+		Op = &OpInfos[Idx];
+		Op->Category = CategoryText;
+		Op->FriendlyName = NSLOCTEXT("NiagaraOpInfo", "Rand Float Name", "Random Float");
+		Op->Description = NSLOCTEXT("NiagaraOpInfo", "Rand Float Desc", "Returns a non-deterministic random float value between 0 and Max");
+		Op->Inputs.Add(FNiagaraOpInOutInfo(Max, Type, MaxText, MaxText, DefaultStr_One));
+		Op->Outputs.Add(FNiagaraOpInOutInfo(Result, Type, ResultText, ResultText, DefaultStr_One, TEXT("rand_float({0})")));
+		Op->BuildName(TEXT("Rand Float"), CategoryName);
+		Op->bNumericsCanBeIntegers = false;
+		Op->bNumericsCanBeFloats = true;
+		OpInfoMap.Add(Op->Name) = Idx;
 
 		// Deterministic/seeded random number generation. 
 		static FName SeedName1(TEXT("Seed 1"));
@@ -810,6 +835,36 @@ void FNiagaraOpInfo::Init()
 		Op->BuildName(TEXT("SeededRand"), CategoryName);
 		OpInfoMap.Add(Op->Name) = Idx;
 
+		Idx = OpInfos.AddDefaulted();
+		Op = &OpInfos[Idx];
+		Op->Category = CategoryText;
+		Op->FriendlyName = NSLOCTEXT("NiagaraOpInfo", "Seeded Integer Rand Name", "Seeded Integer Random");
+		Op->Description = NSLOCTEXT("NiagaraOpInfo", "Seeded Integer Rand Desc", "Returns a deterministic random integer value between 0 and Max-1.");
+		Op->Inputs.Add(FNiagaraOpInOutInfo(Max, Type, MaxText, MaxText, DefaultStr_One));
+		Op->Inputs.Add(FNiagaraOpInOutInfo(SeedName1, SeedType, SeedText1, SeedText1, DefaultSeed_Zero));
+		Op->Inputs.Add(FNiagaraOpInOutInfo(SeedName2, SeedType, SeedText2, SeedText2, DefaultSeed_Zero));
+		Op->Inputs.Add(FNiagaraOpInOutInfo(SeedName3, SeedType, SeedText3, SeedText3, DefaultSeed_Zero));
+		Op->Outputs.Add(FNiagaraOpInOutInfo(Result, Type, ResultText, ResultText, DefaultStr_One, TEXT("rand_int({0}, {1}, {2}, {3})")));
+		Op->BuildName(TEXT("SeededRand Integer"), CategoryName);
+		Op->bNumericsCanBeIntegers = true;
+		Op->bNumericsCanBeFloats = false;
+		OpInfoMap.Add(Op->Name) = Idx;
+
+		Idx = OpInfos.AddDefaulted();
+		Op = &OpInfos[Idx];
+		Op->Category = CategoryText;
+		Op->FriendlyName = NSLOCTEXT("NiagaraOpInfo", "Seeded Float Rand Name", "Seeded Float Random");
+		Op->Description = NSLOCTEXT("NiagaraOpInfo", "Seeded Float Rand Desc", "Returns a deterministic random float value between 0 and Max.");
+		Op->Inputs.Add(FNiagaraOpInOutInfo(Max, Type, MaxText, MaxText, DefaultStr_One));
+		Op->Inputs.Add(FNiagaraOpInOutInfo(SeedName1, SeedType, SeedText1, SeedText1, DefaultSeed_Zero));
+		Op->Inputs.Add(FNiagaraOpInOutInfo(SeedName2, SeedType, SeedText2, SeedText2, DefaultSeed_Zero));
+		Op->Inputs.Add(FNiagaraOpInOutInfo(SeedName3, SeedType, SeedText3, SeedText3, DefaultSeed_Zero));
+		Op->Outputs.Add(FNiagaraOpInOutInfo(Result, Type, ResultText, ResultText, DefaultStr_One, TEXT("rand_float({0}, {1}, {2}, {3})")));
+		Op->BuildName(TEXT("SeededRand Float"), CategoryName);
+		Op->bNumericsCanBeIntegers = false;
+		Op->bNumericsCanBeFloats = true;
+		OpInfoMap.Add(Op->Name) = Idx;
+
 		//Comparison ops
 		Idx = OpInfos.AddDefaulted();
 		Op = &OpInfos[Idx];
diff --git a/Engine/Plugins/FX/Niagara/Source/NiagaraEditor/Private/NiagaraNodeOp.cpp b/Engine/Plugins/FX/Niagara/Source/NiagaraEditor/Private/NiagaraNodeOp.cpp
--- a/Engine/Plugins/FX/Niagara/Source/NiagaraEditor/Private/NiagaraNodeOp.cpp
+++ b/Engine/Plugins/FX/Niagara/Source/NiagaraEditor/Private/NiagaraNodeOp.cpp
@@ -123,6 +123,31 @@ void UNiagaraNodeOp::Compile(class FHlslNiagaraTranslator* Translator, TArray<in
 			FText Format = LOCTEXT("InputErrorFormat", "Error compiling input on {OpName} node.");
 			Translator->Error(FText::Format(Format, Args), this, Pin);
 		}
+		else if (i < OpInfo->Inputs.Num() && OpInfo->Inputs[i].DataType == FNiagaraTypeDefinition::GetGenericNumericDef()) 
+		{
+			// Some nodes disallow integer or floating numeric input pins, so we guard against them here. 
+			// This will catch both implicitly and explicitly set pin types.
+			// Currently this is for the Random Float/Integer and Seeded Random Float/Integer ops, but might be useful for others in the future. 
+
+			const UEdGraphSchema_Niagara* Schema = CastChecked<UEdGraphSchema_Niagara>(GetSchema());
+			FNiagaraTypeDefinition TypeDef = Schema->PinToTypeDefinition(Pin);
+			if (TypeDef.IsFloatPrimitive() && !OpInfo->bNumericsCanBeFloats)
+			{
+				bError = true;
+				FFormatNamedArguments Args;
+				Args.Add(TEXT("OpName"), GetNodeTitle(ENodeTitleType::FullTitle));
+				FText Format = LOCTEXT("InputTypeErrorFormat", "The {OpName} node cannot have float based numeric input pins.");
+				Translator->Error(FText::Format(Format, Args), this, Pin);
+			}
+			else if (!TypeDef.IsFloatPrimitive() && !OpInfo->bNumericsCanBeIntegers)
+			{
+				bError = true;
+				FFormatNamedArguments Args;
+				Args.Add(TEXT("OpName"), GetNodeTitle(ENodeTitleType::FullTitle));
+				FText Format = LOCTEXT("InputTypeErrorFormat", "The {OpName} node cannot have integer based numeric input pins.");
+				Translator->Error(FText::Format(Format, Args), this, Pin);
+			}
+		}
 		Inputs.Add(CompiledInput);
 	}
 	
diff --git a/Engine/Plugins/FX/Niagara/Source/NiagaraEditor/Public/NiagaraEditorCommon.h b/Engine/Plugins/FX/Niagara/Source/NiagaraEditor/Public/NiagaraEditorCommon.h
--- a/Engine/Plugins/FX/Niagara/Source/NiagaraEditor/Public/NiagaraEditorCommon.h
+++ b/Engine/Plugins/FX/Niagara/Source/NiagaraEditor/Public/NiagaraEditorCommon.h
@@ -12,7 +12,11 @@ class FNiagaraOpInfo
 {
 public:
 	FNiagaraOpInfo()
-		: Keywords(FText()), NumericOuputTypeSelectionMode(ENiagaraNumericOutputTypeSelectionMode::Largest), bSupportsAddedInputs(false)
+		: Keywords(FText())
+		, NumericOuputTypeSelectionMode(ENiagaraNumericOutputTypeSelectionMode::Largest)
+		, bSupportsAddedInputs(false)
+		, bNumericsCanBeIntegers(true)
+		, bNumericsCanBeFloats(true)
 	{}
 
 	FName Name;
@@ -27,6 +31,12 @@ public:
 	/** If true then this operation supports a variable number of inputs */
 	bool bSupportsAddedInputs;
 
+	/** If integer pins are allowed on this op's numeric pins. */
+	bool bNumericsCanBeIntegers;
+
+	/** If float pins are allowed on this op's numeric pins. */
+	bool bNumericsCanBeFloats;
+
 	/** 
 	* The format that can generate the hlsl for the given number of inputs.
 	* Used the placeholder {A} and {B} to chain the inputs together.
