Niagara Override Parameters are now searchable in the Actor Details view.

The DisplayName parameter of the row is used as the search key, but wasn't being set by the Niagara Details customization.

diff --git a/Engine/Plugins/FX/Niagara/Source/NiagaraEditor/Private/Customizations/NiagaraComponentDetails.cpp b/Engine/Plugins/FX/Niagara/Source/NiagaraEditor/Private/Customizations/NiagaraComponentDetails.cpp
--- a/Engine/Plugins/FX/Niagara/Source/NiagaraEditor/Private/Customizations/NiagaraComponentDetails.cpp
+++ b/Engine/Plugins/FX/Niagara/Source/NiagaraEditor/Private/Customizations/NiagaraComponentDetails.cpp
@@ -222,6 +222,8 @@ public:
 				TSharedPtr<SWidget> DefaultNameWidget;
 				TSharedPtr<SWidget> DefaultValueWidget;
 
+				Row->DisplayName(FText::FromName(Parameter.GetName()));
+
 				TSharedPtr<IPropertyHandle> PropertyHandle = Row->GetPropertyHandle();
 
 				FDetailWidgetRow& CustomWidget = Row->CustomWidget(true);
