Fix crash when creating a float constant node in a Niagara module.

FNiagaraTypeDefinition::GetStruct() is expected to return null when ClassStructOrEnum is null, so it can't use CastChecked.

diff --git a/Engine/Plugins/FX/Niagara/Source/Niagara/Public/NiagaraTypes.h b/Engine/Plugins/FX/Niagara/Source/Niagara/Public/NiagaraTypes.h
--- a/Engine/Plugins/FX/Niagara/Source/Niagara/Public/NiagaraTypes.h
+++ b/Engine/Plugins/FX/Niagara/Source/Niagara/Public/NiagaraTypes.h
@@ -479,7 +479,7 @@ public:
 
 	UStruct* GetStruct() const
 	{
-		return UnderlyingType == UT_Enum ? IntStruct : CastChecked<UStruct>(ClassStructOrEnum);
+		return UnderlyingType == UT_Enum ? IntStruct : Cast<UStruct>(ClassStructOrEnum);
 	}
 
 	UScriptStruct* GetScriptStruct()const
