- Remove members inside FNiagaraParameterDirectBinding that are only used when checks are enabled
- Saves ~5 cache misses per FNiagaraEmitterInstance tick, which over Niagara test is ~3k misses

fyi shaun.kime,simon.tovey,rob.krajcarski
#rnx

diff --git a/Engine/Plugins/FX/Niagara/Source/Niagara/Classes/NiagaraEmitterInstance.h b/Engine/Plugins/FX/Niagara/Source/Niagara/Classes/NiagaraEmitterInstance.h
--- a/Engine/Plugins/FX/Niagara/Source/Niagara/Classes/NiagaraEmitterInstance.h
+++ b/Engine/Plugins/FX/Niagara/Source/Niagara/Classes/NiagaraEmitterInstance.h
@@ -187,18 +187,18 @@ private:
 	FNiagaraParameterDirectBinding<float> UpdateEmitterAgeBinding;
 	TArray<FNiagaraParameterDirectBinding<float>> EventEmitterAgeBindings;
 
-	FNiagaraParameterDirectBinding<int32> SpawnExecCountBinding;
-	FNiagaraParameterDirectBinding<int32> UpdateExecCountBinding;
-	TArray<FNiagaraParameterDirectBinding<int32>> EventExecCountBindings;
+	FNiagaraParameterDirectBinding<int32> SpawnRandomSeedBinding;
+	FNiagaraParameterDirectBinding<int32> UpdateRandomSeedBinding;
+	FNiagaraParameterDirectBinding<int32> GPURandomSeedBinding;
+	TArray<FNiagaraParameterDirectBinding<int>> EventRandomSeedBindings;
 
 	FNiagaraParameterDirectBinding<int32> SpawnTotalSpawnedParticlesBinding;
 	FNiagaraParameterDirectBinding<int32> UpdateTotalSpawnedParticlesBinding;
 	TArray<FNiagaraParameterDirectBinding<int32>> EventTotalSpawnedParticlesBindings;
 
-	FNiagaraParameterDirectBinding<int32> SpawnRandomSeedBinding;
-	FNiagaraParameterDirectBinding<int32> UpdateRandomSeedBinding;
-	FNiagaraParameterDirectBinding<int32> GPURandomSeedBinding;
-	TArray<FNiagaraParameterDirectBinding<int>> EventRandomSeedBindings;
+	FNiagaraParameterDirectBinding<int32> SpawnExecCountBinding;
+	FNiagaraParameterDirectBinding<int32> UpdateExecCountBinding;
+	TArray<FNiagaraParameterDirectBinding<int32>> EventExecCountBindings;
 
 	/*
 	FNiagaraParameterDirectBinding<int32> SpawnDeterminismBinding;
diff --git a/Engine/Plugins/FX/Niagara/Source/Niagara/Private/NiagaraEmitterInstance.cpp b/Engine/Plugins/FX/Niagara/Source/Niagara/Private/NiagaraEmitterInstance.cpp
--- a/Engine/Plugins/FX/Niagara/Source/Niagara/Private/NiagaraEmitterInstance.cpp
+++ b/Engine/Plugins/FX/Niagara/Source/Niagara/Private/NiagaraEmitterInstance.cpp
@@ -1209,7 +1209,7 @@ void FNiagaraEmitterInstance::Tick(float DeltaSeconds)
 		int32 Estimations = CachedEmitter->AddRuntimeAllocation((uint64)this, MaxAllocationCount);
 		if (GbNiagaraShowAllocationWarnings && Estimations >= 5 && ReallocationCount == 3)
 		{
-			FString SystemName = this->GetParentSystemInstance()->GetSystem()->GetName();
+			FString SystemName = System->GetName();
 			FString FullName = SystemName + "::" + this->GetEmitterHandle().GetName().ToString();
 			UE_LOG(LogNiagara, Warning, TEXT("The emitter %s required many memory reallocation due to changing particle counts. Consider setting the emitter's AllocationMode property to 'manual' to improve runtime performance."), *FullName);
 		}
diff --git a/Engine/Plugins/FX/Niagara/Source/Niagara/Public/NiagaraParameterStore.h b/Engine/Plugins/FX/Niagara/Source/Niagara/Public/NiagaraParameterStore.h
--- a/Engine/Plugins/FX/Niagara/Source/Niagara/Public/NiagaraParameterStore.h
+++ b/Engine/Plugins/FX/Niagara/Source/Niagara/Public/NiagaraParameterStore.h
@@ -843,6 +843,8 @@ FORCEINLINE_DEBUGGABLE void FNiagaraParameterStoreBinding::Dump(const FNiagaraPa
 
 //////////////////////////////////////////////////////////////////////////
 
+#define NIAGARA_VALIDATE_DIRECT_BINDINGS	DO_CHECK
+
 /**
 Direct binding to a parameter store to allow efficient gets/sets from code etc. 
 Does no tracking of lifetimes etc so users are responsible for safety.
@@ -851,20 +853,26 @@ template<typename T>
 struct FNiagaraParameterDirectBinding
 {
 	mutable T* ValuePtr;
+#if NIAGARA_VALIDATE_DIRECT_BINDINGS
 	FNiagaraParameterStore* BoundStore;
 	FNiagaraVariable BoundVariable;
 	uint32 LayoutVersion;
+#endif
 
 	FNiagaraParameterDirectBinding()
-		: ValuePtr(nullptr), BoundStore(nullptr)
+		: ValuePtr(nullptr)
+#if NIAGARA_VALIDATE_DIRECT_BINDINGS
+		, BoundStore(nullptr)
+#endif
 	{}
 
 	T* Init(FNiagaraParameterStore& InStore, const FNiagaraVariable& DestVariable)
 	{
+#if NIAGARA_VALIDATE_DIRECT_BINDINGS
 		BoundStore = &InStore;
 		BoundVariable = DestVariable;
 		LayoutVersion = BoundStore->GetLayoutVersion();
-
+#endif
 		check(DestVariable.GetSizeInBytes() == sizeof(T));
 		ValuePtr = (T*)InStore.GetParameterData(DestVariable);
 		return ValuePtr;
@@ -872,9 +880,10 @@ struct FNiagaraParameterDirectBinding
 
 	FORCEINLINE void SetValue(const T& InValue)
 	{
+#if NIAGARA_VALIDATE_DIRECT_BINDINGS
 		checkSlow(BoundVariable.GetSizeInBytes() == sizeof(T));
 		checkfSlow(LayoutVersion == BoundStore->GetLayoutVersion(), TEXT("This binding is invalid, its bound parameter store's layout was changed since it was created"));
-
+#endif
 		if (ValuePtr)
 		{
 			*ValuePtr = InValue;
@@ -883,9 +892,10 @@ struct FNiagaraParameterDirectBinding
 
 	FORCEINLINE T GetValue()const 
 	{
+#if NIAGARA_VALIDATE_DIRECT_BINDINGS
 		checkSlow(BoundVariable.GetSizeInBytes() == sizeof(T));
 		checkfSlow(LayoutVersion == BoundStore->GetLayoutVersion(), TEXT("This binding is invalid, its bound parameter store's layout was changed since it was created"));
-
+#endif
 		if (ValuePtr)
 		{
 			return *ValuePtr;
@@ -898,18 +908,26 @@ template<>
 struct FNiagaraParameterDirectBinding<FMatrix>
 {
 	mutable FMatrix* ValuePtr;
+#if NIAGARA_VALIDATE_DIRECT_BINDINGS
 	FNiagaraParameterStore* BoundStore;
 	FNiagaraVariable BoundVariable;
 	uint32 LayoutVersion;
+#endif
 
-	FNiagaraParameterDirectBinding():ValuePtr(nullptr), BoundStore(nullptr) {}
+	FNiagaraParameterDirectBinding()
+		: ValuePtr(nullptr)
+#if NIAGARA_VALIDATE_DIRECT_BINDINGS
+		, BoundStore(nullptr)
+#endif
+	{}
 
 	FMatrix* Init(FNiagaraParameterStore& InStore, const FNiagaraVariable& DestVariable)
 	{
+#if NIAGARA_VALIDATE_DIRECT_BINDINGS
 		BoundStore = &InStore;
 		BoundVariable = DestVariable;
 		LayoutVersion = BoundStore->GetLayoutVersion();
-
+#endif
 		check(DestVariable.GetSizeInBytes() == sizeof(FMatrix));
 		ValuePtr = (FMatrix*)InStore.GetParameterData(DestVariable);
 		return ValuePtr;
@@ -917,9 +935,10 @@ struct FNiagaraParameterDirectBinding<FMatrix>
 
 	FORCEINLINE void SetValue(const FMatrix& InValue)
 	{
+#if NIAGARA_VALIDATE_DIRECT_BINDINGS
 		checkSlow(BoundVariable.GetSizeInBytes() == sizeof(FMatrix));
 		checkfSlow(LayoutVersion == BoundStore->GetLayoutVersion(), TEXT("This binding is invalid, its bound parameter store's layout was changed since it was created"));
-
+#endif
 		if (ValuePtr)
 		{
 			FMemory::Memcpy(ValuePtr, &InValue, sizeof(FMatrix));//Temp annoyance until we fix the alignment issues with parameter stores.
@@ -928,9 +947,10 @@ struct FNiagaraParameterDirectBinding<FMatrix>
 
 	FORCEINLINE FMatrix GetValue()const
 	{
+#if NIAGARA_VALIDATE_DIRECT_BINDINGS
 		checkSlow(BoundVariable.GetSizeInBytes() == sizeof(FMatrix));
 		checkfSlow(LayoutVersion == BoundStore->GetLayoutVersion(), TEXT("This binding is invalid, its bound parameter store's layout was changed since it was created"));
-
+#endif
 		FMatrix Ret;
 		if (ValuePtr)
 		{
@@ -944,18 +964,26 @@ template<>
 struct FNiagaraParameterDirectBinding<FVector4>
 {
 	mutable FVector4* ValuePtr;
+#if NIAGARA_VALIDATE_DIRECT_BINDINGS
 	FNiagaraParameterStore* BoundStore;
 	FNiagaraVariable BoundVariable;
 	uint32 LayoutVersion;
+#endif
 
-	FNiagaraParameterDirectBinding() :ValuePtr(nullptr), BoundStore(nullptr){}
+	FNiagaraParameterDirectBinding()
+		: ValuePtr(nullptr)
+#if NIAGARA_VALIDATE_DIRECT_BINDINGS
+		, BoundStore(nullptr)
+#endif
+	{}
 
 	FVector4* Init(FNiagaraParameterStore& InStore, const FNiagaraVariable& DestVariable)
 	{
+#if NIAGARA_VALIDATE_DIRECT_BINDINGS
 		BoundStore = &InStore;
 		BoundVariable = DestVariable;
 		LayoutVersion = BoundStore->GetLayoutVersion();
-
+#endif
 		check(DestVariable.GetSizeInBytes() == sizeof(FVector4));
 		ValuePtr = (FVector4*)InStore.GetParameterData(DestVariable);
 		return ValuePtr;
@@ -963,9 +991,10 @@ struct FNiagaraParameterDirectBinding<FVector4>
 
 	FORCEINLINE void SetValue(const FVector4& InValue)
 	{
+#if NIAGARA_VALIDATE_DIRECT_BINDINGS
 		checkSlow(BoundVariable.GetSizeInBytes() == sizeof(FVector4));
 		checkfSlow(LayoutVersion == BoundStore->GetLayoutVersion(), TEXT("This binding is invalid, its bound parameter store's layout was changed since it was created"));
-
+#endif
 		if (ValuePtr)
 		{
 			FMemory::Memcpy(ValuePtr, &InValue, sizeof(FVector4));//Temp annoyance until we fix the alignment issues with parameter stores.
@@ -974,9 +1003,10 @@ struct FNiagaraParameterDirectBinding<FVector4>
 
 	FORCEINLINE FVector4 GetValue()const
 	{
+#if NIAGARA_VALIDATE_DIRECT_BINDINGS
 		check(BoundVariable.GetSizeInBytes() == sizeof(FVector4));
 		checkfSlow(LayoutVersion == BoundStore->GetLayoutVersion(), TEXT("This binding is invalid, its bound parameter store's layout was changed since it was created"));
-
+#endif
 		FVector4 Ret;
 		if (ValuePtr)
 		{
@@ -990,18 +1020,26 @@ template<>
 struct FNiagaraParameterDirectBinding<FQuat>
 {
 	mutable FQuat* ValuePtr;
+#if NIAGARA_VALIDATE_DIRECT_BINDINGS
 	FNiagaraParameterStore* BoundStore;
 	FNiagaraVariable BoundVariable;
 	uint32 LayoutVersion;
+#endif
 
-	FNiagaraParameterDirectBinding() :ValuePtr(nullptr), BoundStore(nullptr) {}
+	FNiagaraParameterDirectBinding()
+		: ValuePtr(nullptr)
+#if NIAGARA_VALIDATE_DIRECT_BINDINGS
+		, BoundStore(nullptr)
+#endif
+	{}
 
 	FQuat* Init(FNiagaraParameterStore& InStore, const FNiagaraVariable& DestVariable)
 	{
+#if NIAGARA_VALIDATE_DIRECT_BINDINGS
 		BoundStore = &InStore;
 		BoundVariable = DestVariable;
 		LayoutVersion = BoundStore->GetLayoutVersion();
-
+#endif
 		check(DestVariable.GetSizeInBytes() == sizeof(FQuat));
 		ValuePtr = (FQuat*)InStore.GetParameterData(DestVariable);
 		return ValuePtr;
@@ -1009,9 +1047,10 @@ struct FNiagaraParameterDirectBinding<FQuat>
 
 	FORCEINLINE void SetValue(const FQuat& InValue)
 	{
+#if NIAGARA_VALIDATE_DIRECT_BINDINGS
 		checkSlow(BoundVariable.GetSizeInBytes() == sizeof(FQuat));
 		checkfSlow(LayoutVersion == BoundStore->GetLayoutVersion(), TEXT("This binding is invalid, its bound parameter store's layout was changed since it was created"));
-
+#endif
 		if (ValuePtr)
 		{
 			FMemory::Memcpy(ValuePtr, &InValue, sizeof(FQuat));//Temp annoyance until we fix the alignment issues with parameter stores.
@@ -1020,9 +1059,10 @@ struct FNiagaraParameterDirectBinding<FQuat>
 
 	FORCEINLINE FQuat GetValue()const
 	{
+#if NIAGARA_VALIDATE_DIRECT_BINDINGS
 		checkSlow(BoundVariable.GetSizeInBytes() == sizeof(FQuat));
 		checkfSlow(LayoutVersion == BoundStore->GetLayoutVersion(), TEXT("This binding is invalid, its bound parameter store's layout was changed since it was created"));
-
+#endif
 		FQuat Ret;
 		if (ValuePtr)
 		{
@@ -1036,18 +1076,26 @@ template<>
 struct FNiagaraParameterDirectBinding<FNiagaraBool>
 {
 	mutable uint32* ValuePtr;
+#if NIAGARA_VALIDATE_DIRECT_BINDINGS
 	FNiagaraParameterStore* BoundStore;
 	FNiagaraVariable BoundVariable;
 	uint32 LayoutVersion;
+#endif
 
-	FNiagaraParameterDirectBinding() :ValuePtr(nullptr), BoundStore(nullptr) {}
+	FNiagaraParameterDirectBinding()
+		: ValuePtr(nullptr)
+#if NIAGARA_VALIDATE_DIRECT_BINDINGS
+		, BoundStore(nullptr)
+#endif
+	{}
 
 	uint32* Init(FNiagaraParameterStore& InStore, const FNiagaraVariable& DestVariable)
 	{
+#if NIAGARA_VALIDATE_DIRECT_BINDINGS
 		BoundStore = &InStore;
 		BoundVariable = DestVariable;
 		LayoutVersion = BoundStore->GetLayoutVersion();
-
+#endif
 		check(DestVariable.GetSizeInBytes() == sizeof(FNiagaraBool));
 		check(sizeof(uint32) == sizeof(FNiagaraBool));
 		ValuePtr = (uint32*)InStore.GetParameterData(DestVariable);
@@ -1056,10 +1104,11 @@ struct FNiagaraParameterDirectBinding<FNiagaraBool>
 
 	FORCEINLINE void SetValue(const FNiagaraBool& InValue)
 	{
+#if NIAGARA_VALIDATE_DIRECT_BINDINGS
 		checkSlow(BoundVariable.GetSizeInBytes() == sizeof(FNiagaraBool));
 		checkSlow(sizeof(uint32) == sizeof(FNiagaraBool));
 		checkfSlow(LayoutVersion == BoundStore->GetLayoutVersion(), TEXT("This binding is invalid, its bound parameter store's layout was changed since it was created"));
-
+#endif
 		if (ValuePtr)
 		{
 			FMemory::Memcpy(ValuePtr, &InValue, sizeof(FNiagaraBool));
@@ -1068,9 +1117,10 @@ struct FNiagaraParameterDirectBinding<FNiagaraBool>
 
 	FORCEINLINE void SetValue(const bool& InValue)
 	{
+#if NIAGARA_VALIDATE_DIRECT_BINDINGS
 		checkSlow(BoundVariable.GetSizeInBytes() == sizeof(FNiagaraBool));
 		checkfSlow(LayoutVersion == BoundStore->GetLayoutVersion(), TEXT("This binding is invalid, its bound parameter store's layout was changed since it was created"));
-
+#endif
 		if (ValuePtr)
 		{
 			if (!InValue)
@@ -1086,9 +1136,10 @@ struct FNiagaraParameterDirectBinding<FNiagaraBool>
 
 	FORCEINLINE FNiagaraBool GetValue()const
 	{
+#if NIAGARA_VALIDATE_DIRECT_BINDINGS
 		checkSlow(BoundVariable.GetSizeInBytes() == sizeof(FNiagaraBool));
 		checkfSlow(LayoutVersion == BoundStore->GetLayoutVersion(), TEXT("This binding is invalid, its bound parameter store's layout was changed since it was created"));
-
+#endif
 		FNiagaraBool Ret(false);
 		if (ValuePtr)
 		{
