Adding thumbnails to the emitter stack node.
#rnx

diff --git a/Engine/Plugins/FX/Niagara/Source/NiagaraEditor/Private/ViewModels/NiagaraEmitterHandleViewModel.cpp b/Engine/Plugins/FX/Niagara/Source/NiagaraEditor/Private/ViewModels/NiagaraEmitterHandleViewModel.cpp
--- a/Engine/Plugins/FX/Niagara/Source/NiagaraEditor/Private/ViewModels/NiagaraEmitterHandleViewModel.cpp
+++ b/Engine/Plugins/FX/Niagara/Source/NiagaraEditor/Private/ViewModels/NiagaraEmitterHandleViewModel.cpp
@@ -18,7 +18,10 @@
 #include "Widgets/Notifications/SNotificationList.h"
 
 #include "ScopedTransaction.h"
-
+#include "NiagaraRendererProperties.h"
+#include "ViewModels/Stack/NiagaraStackRoot.h"
+#include "ViewModels/Stack/NiagaraStackRenderItemGroup.h"
+#include "ViewModels/Stack/NiagaraStackRendererItem.h"
 
 #define LOCTEXT_NAMESPACE "EmitterHandleViewModel"
 
@@ -45,6 +48,31 @@ void FNiagaraEmitterHandleViewModel::Cleanup()
 	}
 }
 
+void FNiagaraEmitterHandleViewModel::GetRendererPreviewData(TArray<FRendererPreviewData>& InRendererPreviewData)
+{
+	InRendererPreviewData.Empty();
+	UNiagaraEmitter* Emitter = GetEmitterHandle()->GetInstance();
+	UNiagaraStackRoot* StackRoot = Cast<UNiagaraStackRoot>(EmitterStackViewModel->GetRootEntry());
+	if (StackRoot)
+	{
+		TArray<UNiagaraStackEntry*> Children;
+		StackRoot->GetRenderGroup()->GetUnfilteredChildren(Children);
+		for (UNiagaraStackEntry* Child : Children)
+		{
+			if (UNiagaraStackRendererItem* RendererItem = Cast<UNiagaraStackRendererItem>(Child))
+			{
+				TArray<UMaterialInterface*> Materials;
+				FNiagaraEmitterInstance* InInstance = GetEmitterViewModel()->GetSimulation().IsValid() ? GetEmitterViewModel()->GetSimulation().Pin().Get() : nullptr;
+				RendererItem->GetRendererProperties()->GetUsedMaterials(InInstance, Materials);
+				for (UMaterialInterface* Material : Materials)
+				{
+					InRendererPreviewData.Add(FRendererPreviewData(Child, Material));
+				}
+			}
+		}
+	}
+}
+
 TSharedRef<FNiagaraSystemViewModel> FNiagaraEmitterHandleViewModel::GetOwningSystemViewModel() const
 {
 	TSharedPtr<FNiagaraSystemViewModel> OwningSystemViewModelPinned = OwningSystemViewModelWeak.Pin();
diff --git a/Engine/Plugins/FX/Niagara/Source/NiagaraEditor/Private/ViewModels/Stack/NiagaraStackViewModel.cpp b/Engine/Plugins/FX/Niagara/Source/NiagaraEditor/Private/ViewModels/Stack/NiagaraStackViewModel.cpp
--- a/Engine/Plugins/FX/Niagara/Source/NiagaraEditor/Private/ViewModels/Stack/NiagaraStackViewModel.cpp
+++ b/Engine/Plugins/FX/Niagara/Source/NiagaraEditor/Private/ViewModels/Stack/NiagaraStackViewModel.cpp
@@ -514,6 +514,11 @@ UNiagaraStackViewModel::FOnSearchCompleted& UNiagaraStackViewModel::OnSearchComp
 	return SearchCompletedDelegate;
 }
 
+UNiagaraStackViewModel::FOnDataObjectChanged& UNiagaraStackViewModel::OnDataObjectChanged()
+{
+	return DataObjectChangedDelegate;
+}
+
 bool UNiagaraStackViewModel::GetShowAllAdvanced() const
 {
 	for (TSharedRef<FTopLevelViewModel> TopLevelViewModel : TopLevelViewModels)
@@ -642,6 +647,7 @@ void UNiagaraStackViewModel::EntryDataObjectModified(UObject* ChangedObject)
 		SystemViewModel.Pin()->NotifyDataObjectChanged(ChangedObject);
 	}
 	OnSearchTextChanged(CurrentSearchText);
+	DataObjectChangedDelegate.Broadcast(ChangedObject);
 }
 
 void UNiagaraStackViewModel::EntryRequestFullRefresh()
diff --git a/Engine/Plugins/FX/Niagara/Source/NiagaraEditor/Public/ViewModels/NiagaraEmitterHandleViewModel.h b/Engine/Plugins/FX/Niagara/Source/NiagaraEditor/Public/ViewModels/NiagaraEmitterHandleViewModel.h
--- a/Engine/Plugins/FX/Niagara/Source/NiagaraEditor/Public/ViewModels/NiagaraEmitterHandleViewModel.h
+++ b/Engine/Plugins/FX/Niagara/Source/NiagaraEditor/Public/ViewModels/NiagaraEmitterHandleViewModel.h
@@ -14,8 +14,23 @@ struct FNiagaraEmitterHandle;
 class FNiagaraEmitterInstance;
 class FNiagaraEmitterViewModel;
 class UNiagaraStackViewModel;
+class UNiagaraStackEntry;
 enum class ENiagaraSystemViewModelEditMode;
 
+struct FRendererPreviewData
+{
+	UNiagaraStackEntry* RenderingEntry;
+	UObject* RenderingObject;
+
+
+	FRendererPreviewData(UNiagaraStackEntry* InRenderingEntry, UObject* InRenderingObject)
+		: RenderingEntry(InRenderingEntry)
+		, RenderingObject(InRenderingObject)
+	{
+
+	};
+};
+
 /** The view model for the FNiagaraEmitterEditorWidget. */
 class FNiagaraEmitterHandleViewModel : public TSharedFromThis<FNiagaraEmitterHandleViewModel>, public FGCObject
 {
@@ -117,6 +132,7 @@ public:
 
 	void Cleanup();
 
+	NIAGARAEDITOR_API void GetRendererPreviewData(TArray<FRendererPreviewData>& InRendererPreviewData);
 private:
 	TSharedRef<FNiagaraSystemViewModel> GetOwningSystemViewModel() const;
 
diff --git a/Engine/Plugins/FX/Niagara/Source/NiagaraEditor/Public/ViewModels/Stack/NiagaraStackRoot.h b/Engine/Plugins/FX/Niagara/Source/NiagaraEditor/Public/ViewModels/Stack/NiagaraStackRoot.h
--- a/Engine/Plugins/FX/Niagara/Source/NiagaraEditor/Public/ViewModels/Stack/NiagaraStackRoot.h
+++ b/Engine/Plugins/FX/Niagara/Source/NiagaraEditor/Public/ViewModels/Stack/NiagaraStackRoot.h
@@ -24,7 +24,10 @@ public:
 
 	virtual bool GetCanExpand() const override;
 	virtual bool GetShouldShowInStack() const override;
-
+	UNiagaraStackRenderItemGroup* GetRenderGroup() const
+	{
+		return RenderGroup;
+	}
 protected:
 	virtual void RefreshChildrenInternal(const TArray<UNiagaraStackEntry*>& CurrentChildren, TArray<UNiagaraStackEntry*>& NewChildren, TArray<FStackIssue>& NewIssues) override;
 
diff --git a/Engine/Plugins/FX/Niagara/Source/NiagaraEditor/Public/ViewModels/Stack/NiagaraStackViewModel.h b/Engine/Plugins/FX/Niagara/Source/NiagaraEditor/Public/ViewModels/Stack/NiagaraStackViewModel.h
--- a/Engine/Plugins/FX/Niagara/Source/NiagaraEditor/Public/ViewModels/Stack/NiagaraStackViewModel.h
+++ b/Engine/Plugins/FX/Niagara/Source/NiagaraEditor/Public/ViewModels/Stack/NiagaraStackViewModel.h
@@ -43,6 +43,7 @@ class NIAGARAEDITOR_API UNiagaraStackViewModel : public UObject, public FEditorU
 public:
 	DECLARE_MULTICAST_DELEGATE(FOnStructureChanged);
 	DECLARE_MULTICAST_DELEGATE(FOnSearchCompleted);
+	DECLARE_MULTICAST_DELEGATE_OneParam(FOnDataObjectChanged, UObject* /** Changed Object */);
 public:
 	struct FSearchResult
 	{
@@ -89,6 +90,7 @@ public:
 
 	FOnStructureChanged& OnStructureChanged();
 	FOnSearchCompleted& OnSearchCompleted();
+	FOnDataObjectChanged& OnDataObjectChanged();
 
 	bool GetShowAllAdvanced() const;
 	void SetShowAllAdvanced(bool bInShowAllAdvanced);
@@ -178,6 +180,7 @@ private:
 	bool bExternalRootEntry;
 
 	FOnStructureChanged StructureChangedDelegate;
+	FOnDataObjectChanged DataObjectChangedDelegate;
 
 	// ~Search stuff
 	FText CurrentSearchText;
diff --git a/Engine/Plugins/FX/Niagara/Source/NiagaraEditorWidgets/Private/SNiagaraOverviewStackNode.cpp b/Engine/Plugins/FX/Niagara/Source/NiagaraEditorWidgets/Private/SNiagaraOverviewStackNode.cpp
--- a/Engine/Plugins/FX/Niagara/Source/NiagaraEditorWidgets/Private/SNiagaraOverviewStackNode.cpp
+++ b/Engine/Plugins/FX/Niagara/Source/NiagaraEditorWidgets/Private/SNiagaraOverviewStackNode.cpp
@@ -14,9 +14,17 @@
 #include "Widgets/Layout/SBox.h"
 #include "Widgets/Input/SButton.h"
 #include "Widgets/Input/SCheckBox.h"
+#include "AssetThumbnail.h"
+#include "ViewModels/NiagaraEmitterViewModel.h"
+#include "Widgets/SToolTip.h"
+#include "Widgets/SBoxPanel.h"
+#include "NiagaraRendererProperties.h"
+#include "ViewModels/NiagaraSystemSelectionViewModel.h"
 
 #define LOCTEXT_NAMESPACE "NiagaraOverviewStackNode"
 
+const float ThumbnailSize = 24.0f;
+
 void SNiagaraOverviewStackNode::Construct(const FArguments& InArgs, UNiagaraOverviewNode* InNode)
 {
 	GraphNode = InNode;
@@ -24,7 +32,6 @@ void SNiagaraOverviewStackNode::Construct(const FArguments& InArgs, UNiagaraOver
 	StackViewModel = nullptr;
 	OverviewSelectionViewModel = nullptr;
 	EmitterHandleViewModelWeak.Reset();
-
 	if (OverviewStackNode->GetOwningSystem() != nullptr)
 	{
 		FNiagaraEditorModule& NiagaraEditorModule = FModuleManager::Get().LoadModuleChecked<FNiagaraEditorModule>("NiagaraEditor");
@@ -43,10 +50,13 @@ void SNiagaraOverviewStackNode::Construct(const FArguments& InArgs, UNiagaraOver
 					StackViewModel = EmitterHandleViewModelWeak.Pin()->GetEmitterStackViewModel();
 				}
 			}
+			if (StackViewModel)
+			{
+				StackViewModel->OnDataObjectChanged().AddRaw(this, &SNiagaraOverviewStackNode::FillThumbnailBar, true);
+			}
 			OverviewSelectionViewModel = OwningSystemViewModel->GetSelectionViewModel();
 		}
 	}
-
 	UpdateGraphNode();
 }
 
@@ -119,6 +129,41 @@ TSharedRef<SWidget> SNiagaraOverviewStackNode::CreateTitleWidget(TSharedPtr<SNod
 		];
 }
 
+TSharedRef<SWidget> SNiagaraOverviewStackNode::CreateThumbnailWidget(float InThumbnailSize, FRendererPreviewData InData)
+{
+	TSharedPtr<FAssetThumbnail> AssetThumbnail = MakeShareable(new FAssetThumbnail(InData.RenderingObject, InThumbnailSize, InThumbnailSize, ThumbnailPool));
+	TSharedRef<SWidget> ThumbnailWidget = AssetThumbnail->MakeThumbnailWidget();
+	TSharedPtr<FAssetThumbnail> TooltipThumbnail = MakeShareable(new FAssetThumbnail(InData.RenderingObject, 64.0f, 64.0f, ThumbnailPool));
+	TSharedRef<SToolTip> ThumbnailTooltipWidget = SNew(SToolTip)
+		.Content()
+		[
+			SNew(SBox)
+			.MinDesiredHeight(64.0f)
+			.MaxDesiredHeight(64.0f)
+			.MinDesiredWidth(64.0f)
+			.MaxDesiredWidth(64.0f)
+			[
+				TooltipThumbnail->MakeThumbnailWidget()
+			]
+		];
+	ThumbnailWidget->SetOnMouseButtonDown(FPointerEventHandler::CreateSP(this, &SNiagaraOverviewStackNode::OnClickedRenderingPreview, InData.RenderingEntry));
+	ThumbnailWidget->SetToolTip(ThumbnailTooltipWidget);
+	return ThumbnailWidget;
+}
+
+FReply SNiagaraOverviewStackNode::OnClickedRenderingPreview(const FGeometry& InGeometry, const FPointerEvent& InEvent, UNiagaraStackEntry* InEntry)
+{
+	if (InEvent.GetEffectingButton() == EKeys::LeftMouseButton)
+	{
+		TArray<UNiagaraStackEntry*> SelectedEntries;
+		SelectedEntries.Add(InEntry);
+		TArray<UNiagaraStackEntry*> DeselectedEntries;
+		OverviewSelectionViewModel->UpdateSelectedEntries(SelectedEntries, DeselectedEntries, true);
+		return FReply::Handled();
+	}
+	return FReply::Unhandled();
+}
+
 TSharedRef<SWidget> SNiagaraOverviewStackNode::CreateNodeContentArea()
 {
 	TSharedPtr<SWidget> ContentWidget;
@@ -135,41 +180,109 @@ TSharedRef<SWidget> SNiagaraOverviewStackNode::CreateNodeContentArea()
 		ContentWidget = SNullWidget::NullWidget;
 	}
 
+	ThumbnailBar = SNew(SHorizontalBox);
+	FillThumbnailBar(nullptr, false);
+
 	// NODE CONTENT AREA
-	return SNew(SBorder)
+	TSharedRef<SWidget> NodeWidget = SNew(SBorder)
 		.BorderImage(FEditorStyle::GetBrush("NoBorder"))
 		.HAlign(HAlign_Fill)
 		.VAlign(VAlign_Fill)
 		.Padding(FMargin(2, 2, 2, 4))
 		[
-			SNew(SHorizontalBox)
-			+SHorizontalBox::Slot()
-			.AutoWidth()
-			.HAlign(HAlign_Left)
+			SNew(SVerticalBox)
+			+ SVerticalBox::Slot()
+			.AutoHeight()
+			.HAlign(HAlign_Fill)
+			.VAlign(VAlign_Center)
+			.Padding(2.0f, 2.0f)
 			[
-				// LEFT
-				SAssignNew(LeftNodeBox, SVerticalBox)
+				ThumbnailBar.ToSharedRef()
 			]
-			+SHorizontalBox::Slot()
+			+ SVerticalBox::Slot()
+			.FillHeight(1.0f)
+			.HAlign(HAlign_Fill)
+			.VAlign(VAlign_Center)
+			.Padding(0.0f)
 			[
-				SNew(SBorder)
-				.BorderImage(FNiagaraEditorWidgetsStyle::Get().GetBrush("NiagaraEditor.SystemOverview.NodeBackgroundBorder"))
-				.BorderBackgroundColor(FNiagaraEditorWidgetsStyle::Get().GetColor("NiagaraEditor.SystemOverview.NodeBackgroundColor"))
-				.HAlign(HAlign_Fill)
-				.VAlign(VAlign_Fill)
-				.Padding(FMargin(0, 0, 0, 4))
+				SNew(SHorizontalBox)
+				+SHorizontalBox::Slot()
+				.AutoWidth()
+				.HAlign(HAlign_Left)
 				[
-					ContentWidget.ToSharedRef()
+					// LEFT
+					SAssignNew(LeftNodeBox, SVerticalBox)
+				]
+				+SHorizontalBox::Slot()
+				[
+					SNew(SBorder)
+					.BorderImage(FNiagaraEditorWidgetsStyle::Get().GetBrush("NiagaraEditor.SystemOverview.NodeBackgroundBorder"))
+					.BorderBackgroundColor(FNiagaraEditorWidgetsStyle::Get().GetColor("NiagaraEditor.SystemOverview.NodeBackgroundColor"))
+					.HAlign(HAlign_Fill)
+					.VAlign(VAlign_Fill)
+					.Padding(FMargin(0, 0, 0, 4))
+					[
+						ContentWidget.ToSharedRef()
+					]
+				]
+				+SHorizontalBox::Slot()
+				.AutoWidth()
+				.HAlign(HAlign_Right)
+				[
+					// RIGHT
+					SAssignNew(RightNodeBox, SVerticalBox)
 				]
-			]
-			+SHorizontalBox::Slot()
-			.AutoWidth()
-			.HAlign(HAlign_Right)
-			[
-				// RIGHT
-				SAssignNew(RightNodeBox, SVerticalBox)
 			]
 		];
+
+
+	return NodeWidget;
+
+}
+
+void SNiagaraOverviewStackNode::FillThumbnailBar(UObject* ChangedObject, const bool bIsTriggeredByObjectUpdate)
+{
+	UNiagaraRendererProperties* RendererProperties = Cast< UNiagaraRendererProperties>(ChangedObject);
+	if (!bIsTriggeredByObjectUpdate || RendererProperties != nullptr)
+	{
+		ThumbnailBar->ClearChildren();
+		ThumbnailPool = MakeShared<FAssetThumbnailPool>(100, /*InAreRealTileThumbnailsAllowed=*/true);
+		if (EmitterHandleViewModelWeak.IsValid())
+		{
+			EmitterHandleViewModelWeak.Pin()->GetRendererPreviewData(PreviewData);
+
+			for (FRendererPreviewData Preview : PreviewData)
+			{
+				if (Preview.RenderingObject)
+				{
+					ThumbnailBar->AddSlot()
+						.AutoWidth()
+						.HAlign(HAlign_Left)
+						.VAlign(VAlign_Center)
+						.Padding(2.0f, 2.0f, 4.0f, 4.0f)
+						.MaxWidth(ThumbnailSize)
+						[
+							SNew(SBox)
+							.MinDesiredHeight(ThumbnailSize)
+							.MaxDesiredHeight(ThumbnailSize)
+							.MinDesiredWidth(ThumbnailSize)
+							.Visibility(this, &SNiagaraOverviewStackNode::GetEnabledCheckBoxVisibility)
+							[
+								CreateThumbnailWidget(ThumbnailSize, Preview)
+							]
+						];
+				}
+			}
+		}
+
+		ThumbnailBar->AddSlot()
+			.HAlign(HAlign_Fill)
+			.VAlign(VAlign_Center)
+			.FillWidth(1.0f)
+			[
+				SNullWidget::NullWidget
+			];
+	}
 }
 
 EVisibility SNiagaraOverviewStackNode::GetIssueIconVisibility() const
diff --git a/Engine/Plugins/FX/Niagara/Source/NiagaraEditorWidgets/Private/SNiagaraOverviewStackNode.h b/Engine/Plugins/FX/Niagara/Source/NiagaraEditorWidgets/Private/SNiagaraOverviewStackNode.h
--- a/Engine/Plugins/FX/Niagara/Source/NiagaraEditorWidgets/Private/SNiagaraOverviewStackNode.h
+++ b/Engine/Plugins/FX/Niagara/Source/NiagaraEditorWidgets/Private/SNiagaraOverviewStackNode.h
@@ -3,11 +3,17 @@
 #pragma once
 
 #include "SGraphNode.h"
+#include "Widgets/SBoxPanel.h"
+#include "SNiagaraOverviewStack.h"
 
 class UNiagaraOverviewNode;
 class UNiagaraStackViewModel;
 class UNiagaraSystemSelectionViewModel;
 class FNiagaraEmitterHandleViewModel;
+class FAssetThumbnailPool;
+class FAssetThumbnail;
+struct FRendererPreviewData;
+
 
 class SNiagaraOverviewStackNode : public SGraphNode
 {
@@ -21,13 +27,14 @@ public:
 protected:
 	virtual TSharedRef<SWidget> CreateTitleWidget(TSharedPtr<SNodeTitle> NodeTitle) override;
 	virtual TSharedRef<SWidget> CreateNodeContentArea() override;
-
+	void FillThumbnailBar(UObject* ChangedObject, const bool bIsTriggeredByObjectUpdate);
 private:
 	EVisibility GetIssueIconVisibility() const;
 	EVisibility GetEnabledCheckBoxVisibility() const;
 	ECheckBoxState GetEnabledCheckState() const;
 	void OnEnabledCheckStateChanged(ECheckBoxState InCheckState);
-
+	TSharedRef<SWidget> CreateThumbnailWidget(float InThumbnailSize, FRendererPreviewData InData);
+	FReply OnClickedRenderingPreview(const FGeometry& InGeometry, const FPointerEvent& InEvent, class UNiagaraStackEntry* InEntry);
 	FText GetToggleIsolateToolTip() const;
 	FReply OnToggleIsolateButtonClicked();
 	EVisibility GetToggleIsolateVisibility() const;
@@ -38,4 +45,8 @@ private:
 	UNiagaraStackViewModel* StackViewModel;
 	UNiagaraSystemSelectionViewModel* OverviewSelectionViewModel;
 	TWeakPtr<FNiagaraEmitterHandleViewModel> EmitterHandleViewModelWeak;
+	TSharedPtr<FAssetThumbnailPool> ThumbnailPool;
+	/** Thumbnail widget containers */
+	TSharedPtr<SHorizontalBox> ThumbnailBar;
+	TArray<FRendererPreviewData> PreviewData;
 };
\ No newline at end of file
diff --git a/Engine/Source/Runtime/Engine/Private/Materials/Material.cpp b/Engine/Source/Runtime/Engine/Private/Materials/Material.cpp
--- a/Engine/Source/Runtime/Engine/Private/Materials/Material.cpp
+++ b/Engine/Source/Runtime/Engine/Private/Materials/Material.cpp
@@ -5782,7 +5782,8 @@ bool UMaterial::ShouldForcePlanePreview()
 		MaterialThumbnailInfo = USceneThumbnailInfoWithPrimitive::StaticClass()->GetDefaultObject<USceneThumbnailInfoWithPrimitive>();
 	}
 	// UI and particle sprite material thumbnails always get a 2D plane centered at the camera which is a better representation of the what the material will look like
-	return Super::ShouldForcePlanePreview() || IsUIMaterial() || (bUsedWithParticleSprites && !MaterialThumbnailInfo->bUserModifiedShape);
+	const bool bUsedWithNiagara = bUsedWithNiagaraSprites || bUsedWithNiagaraRibbons || bUsedWithNiagaraMeshParticles; 
+	return Super::ShouldForcePlanePreview() || IsUIMaterial() || (bUsedWithParticleSprites && !MaterialThumbnailInfo->bUserModifiedShape) || (bUsedWithNiagara && !MaterialThumbnailInfo->bUserModifiedShape);
 }
 
 void UMaterial::NotifyCompilationFinished(UMaterialInterface* Material)
