FNM - Streamline Transitions for NiagaraSim to avoud Back2Back CacheFlushes.

diff --git a/Engine/Plugins/FX/Niagara/Source/Niagara/Private/NiagaraEmitterInstanceBatcher.cpp b/Engine/Plugins/FX/Niagara/Source/Niagara/Private/NiagaraEmitterInstanceBatcher.cpp
--- a/Engine/Plugins/FX/Niagara/Source/Niagara/Private/NiagaraEmitterInstanceBatcher.cpp
+++ b/Engine/Plugins/FX/Niagara/Source/Niagara/Private/NiagaraEmitterInstanceBatcher.cpp
@@ -408,14 +408,13 @@ void NiagaraEmitterInstanceBatcher::DispatchAllOnCompute(FOverlappableTicks& Ove
 	// Disable automatic cache flush so that we can have our compute work overlapping. Barrier will be used as a sync mechanism.
 	RHICmdList.AutomaticCacheFlushAfterComputeShader(false);
 
-	//
-	//	Transition current index buffer ready for compute and clear then all using overlapping compute work items.
-	//
-
-	RHICmdList.TransitionResources(EResourceTransitionAccess::ERWBarrier, EResourceTransitionPipeline::EGfxToCompute, CurrDataBuffers.GetData(), CurrDataBuffers.Num());
-
 #if WITH_EDITORONLY_DATA
 	{
+		//
+		//	Transition current index buffer ready for compute and clear then all using overlapping compute work items.
+		//
+		RHICmdList.TransitionResources(EResourceTransitionAccess::ERWBarrier, EResourceTransitionPipeline::EGfxToCompute, CurrDataBuffers.GetData(), CurrDataBuffers.Num());
+
 		for (FNiagaraGPUSystemTick* Tick : OverlappableTick)
 		{
 			uint32 DispatchCount = Tick->Count;
@@ -439,12 +438,15 @@ void NiagaraEmitterInstanceBatcher::DispatchAllOnCompute(FOverlappableTicks& Ove
 	//
 	//	Add a rw barrier for the destination data buffers we just cleared and mark others as read/write as needed for particles simulation.
 	//
-	RHICmdList.TransitionResource(EResourceTransitionAccess::ERWBarrier, EResourceTransitionPipeline::EComputeToCompute, GPUInstanceCounterManager.GetInstanceCountBuffer().UAV);		// transition to readable; we'll be using this next frame
-	RHICmdList.TransitionResources(EResourceTransitionAccess::ERWBarrier, EResourceTransitionPipeline::EComputeToCompute, DestDataBuffers.GetData(), DestDataBuffers.Num());
-	RHICmdList.TransitionResources(EResourceTransitionAccess::ERWBarrier, EResourceTransitionPipeline::EComputeToCompute, DestBufferIntFloat.GetData(), DestBufferIntFloat.Num());
-	RHICmdList.TransitionResources(EResourceTransitionAccess::EReadable, EResourceTransitionPipeline::EComputeToCompute, CurrDataBuffers.GetData(), CurrDataBuffers.Num());
-	RHICmdList.TransitionResources(EResourceTransitionAccess::EReadable, EResourceTransitionPipeline::EComputeToCompute, CurrBufferIntFloat.GetData(), CurrBufferIntFloat.Num());
-	RHICmdList.FlushComputeShaderCache();
+	{
+		TArray<FRHIUnorderedAccessView*, TInlineAllocator<5>> TransitionUAVs;
+		TransitionUAVs.Add(GPUInstanceCounterManager.GetInstanceCountBuffer().UAV);
+		TransitionUAVs.Append(DestDataBuffers.GetData(), DestDataBuffers.Num());
+		TransitionUAVs.Append(DestBufferIntFloat.GetData(), DestBufferIntFloat.Num());
+		TransitionUAVs.Append(CurrDataBuffers.GetData(), CurrDataBuffers.Num());
+		TransitionUAVs.Append(CurrBufferIntFloat.GetData(), CurrBufferIntFloat.Num());
+		RHICmdList.TransitionResources(EResourceTransitionAccess::ERWBarrier, EResourceTransitionPipeline::EComputeToCompute, TransitionUAVs.GetData(), TransitionUAVs.Num());
+	}
 
 	for (FNiagaraGPUSystemTick* Tick : OverlappableTick)
 	{
@@ -480,18 +482,18 @@ void NiagaraEmitterInstanceBatcher::DispatchAllOnCompute(FOverlappableTicks& Ove
 	//
 	//	Now Copy to staging buffer the data we want to read back (alive particle count). And make buffer ready for that and draw commands on the graphics pipe too.
 	//
-	RHICmdList.TransitionResource(EResourceTransitionAccess::ERWBarrier, EResourceTransitionPipeline::EComputeToCompute, GPUInstanceCounterManager.GetInstanceCountBuffer().UAV);		// transition to readable; we'll be using this next frame
-	RHICmdList.TransitionResources(EResourceTransitionAccess::ERWBarrier, EResourceTransitionPipeline::EComputeToCompute, DestDataBuffers.GetData(), DestDataBuffers.Num());
-	RHICmdList.TransitionResources(EResourceTransitionAccess::ERWBarrier, EResourceTransitionPipeline::EComputeToCompute, DestBufferIntFloat.GetData(), DestBufferIntFloat.Num());
-	RHICmdList.TransitionResources(EResourceTransitionAccess::EReadable, EResourceTransitionPipeline::EComputeToGfx, DestDataBuffers.GetData(), DestDataBuffers.Num());
-	RHICmdList.TransitionResources(EResourceTransitionAccess::EReadable, EResourceTransitionPipeline::EComputeToGfx, DestBufferIntFloat.GetData(), DestBufferIntFloat.Num());
-	RHICmdList.FlushComputeShaderCache();
-
+	{
+		TArray<FRHIUnorderedAccessView*, TInlineAllocator<5>> TransitionUAVs;
+		TransitionUAVs.Add(GPUInstanceCounterManager.GetInstanceCountBuffer().UAV);
+		TransitionUAVs.Append(DestDataBuffers.GetData(), DestDataBuffers.Num());
+		TransitionUAVs.Append(DestBufferIntFloat.GetData(), DestBufferIntFloat.Num());
+		TransitionUAVs.Append(DestDataBuffers.GetData(), DestDataBuffers.Num());
+		TransitionUAVs.Append(DestBufferIntFloat.GetData(), DestBufferIntFloat.Num());
+		RHICmdList.TransitionResources(EResourceTransitionAccess::EReadable, EResourceTransitionPipeline::EComputeToCompute, TransitionUAVs.GetData(), TransitionUAVs.Num());
+	}
 	// We have done all our overlapping compute work on this list so go back to default behavior and flush.
 	RHICmdList.AutomaticCacheFlushAfterComputeShader(true);
 
-	// We have done all our compute work
-	RHICmdList.FlushComputeShaderCache();
 	if (GNiagaraSubmitCommands)
 	{
 		RHICmdList.SubmitCommandsHint();
