Not using FNiagaraVariable names in UNiagaraScriptVariable names as periods can lead to bad behavior.

diff --git a/Engine/Plugins/FX/Niagara/Source/NiagaraEditor/Private/NiagaraGraph.cpp b/Engine/Plugins/FX/Niagara/Source/NiagaraEditor/Private/NiagaraGraph.cpp
--- a/Engine/Plugins/FX/Niagara/Source/NiagaraEditor/Private/NiagaraGraph.cpp
+++ b/Engine/Plugins/FX/Niagara/Source/NiagaraEditor/Private/NiagaraGraph.cpp
@@ -823,7 +823,7 @@ void UNiagaraGraph::AddParameter(const FNiagaraVariable& Parameter, bool bIsStat
 	UNiagaraScriptVariable** FoundScriptVariable = VariableToScriptVariable.Find(Parameter);
 	if (!FoundScriptVariable)
 	{
-		UNiagaraScriptVariable* NewScriptVariable = NewObject<UNiagaraScriptVariable>(this, Parameter.GetName(), RF_Transactional);
+		UNiagaraScriptVariable* NewScriptVariable = NewObject<UNiagaraScriptVariable>(this, FName(), RF_Transactional);
 		NewScriptVariable->Variable = Parameter;
 		NewScriptVariable->Metadata.bIsStaticSwitch = bIsStaticSwitch;
 		VariableToScriptVariable.Add(Parameter, NewScriptVariable);
@@ -1412,7 +1412,7 @@ void UNiagaraGraph::SetMetaData(const FNiagaraVariable& InVar, const FNiagaraVar
 	}
 	else 
 	{
-		UNiagaraScriptVariable*& NewScriptVariable = VariableToScriptVariable.Add(InVar, NewObject<UNiagaraScriptVariable>(this, InVar.GetName(), RF_Transactional));
+		UNiagaraScriptVariable*& NewScriptVariable = VariableToScriptVariable.Add(InVar, NewObject<UNiagaraScriptVariable>(this, FName(), RF_Transactional));
 		NewScriptVariable->Variable = InVar;
 		NewScriptVariable->Metadata = InMetaData;
 	}
