Changing how the render preview data is referenced to fix include issue
#rnx

diff --git a/Engine/Plugins/FX/Niagara/Source/NiagaraEditor/Private/ViewModels/NiagaraEmitterHandleViewModel.cpp b/Engine/Plugins/FX/Niagara/Source/NiagaraEditor/Private/ViewModels/NiagaraEmitterHandleViewModel.cpp
--- a/Engine/Plugins/FX/Niagara/Source/NiagaraEditor/Private/ViewModels/NiagaraEmitterHandleViewModel.cpp
+++ b/Engine/Plugins/FX/Niagara/Source/NiagaraEditor/Private/ViewModels/NiagaraEmitterHandleViewModel.cpp
@@ -48,7 +48,7 @@ void FNiagaraEmitterHandleViewModel::Cleanup()
 	}
 }
 
-void FNiagaraEmitterHandleViewModel::GetRendererPreviewData(TArray<FRendererPreviewData>& InRendererPreviewData)
+void FNiagaraEmitterHandleViewModel::GetRendererPreviewData(TArray<FRendererPreviewData*>& InRendererPreviewData)
 {
 	InRendererPreviewData.Empty();
 	UNiagaraEmitter* Emitter = GetEmitterHandle()->GetInstance();
@@ -66,7 +66,7 @@ void FNiagaraEmitterHandleViewModel::GetRendererPreviewData(TArray<FRendererPrev
 				RendererItem->GetRendererProperties()->GetUsedMaterials(InInstance, Materials);
 				for (UMaterialInterface* Material : Materials)
 				{
-					InRendererPreviewData.Add(FRendererPreviewData(Child, Material));
+					InRendererPreviewData.Add(new FRendererPreviewData(Child, Material));
 				}
 			}
 		}
diff --git a/Engine/Plugins/FX/Niagara/Source/NiagaraEditor/Public/ViewModels/NiagaraEmitterHandleViewModel.h b/Engine/Plugins/FX/Niagara/Source/NiagaraEditor/Public/ViewModels/NiagaraEmitterHandleViewModel.h
--- a/Engine/Plugins/FX/Niagara/Source/NiagaraEditor/Public/ViewModels/NiagaraEmitterHandleViewModel.h
+++ b/Engine/Plugins/FX/Niagara/Source/NiagaraEditor/Public/ViewModels/NiagaraEmitterHandleViewModel.h
@@ -132,7 +132,7 @@ public:
 
 	void Cleanup();
 
-	NIAGARAEDITOR_API void GetRendererPreviewData(TArray<FRendererPreviewData>& InRendererPreviewData);
+	NIAGARAEDITOR_API void GetRendererPreviewData(TArray<FRendererPreviewData*>& InRendererPreviewData);
 private:
 	TSharedRef<FNiagaraSystemViewModel> GetOwningSystemViewModel() const;
 
diff --git a/Engine/Plugins/FX/Niagara/Source/NiagaraEditorWidgets/Private/SNiagaraOverviewStackNode.cpp b/Engine/Plugins/FX/Niagara/Source/NiagaraEditorWidgets/Private/SNiagaraOverviewStackNode.cpp
--- a/Engine/Plugins/FX/Niagara/Source/NiagaraEditorWidgets/Private/SNiagaraOverviewStackNode.cpp
+++ b/Engine/Plugins/FX/Niagara/Source/NiagaraEditorWidgets/Private/SNiagaraOverviewStackNode.cpp
@@ -129,11 +129,11 @@ TSharedRef<SWidget> SNiagaraOverviewStackNode::CreateTitleWidget(TSharedPtr<SNod
 		];
 }
 
-TSharedRef<SWidget> SNiagaraOverviewStackNode::CreateThumbnailWidget(float InThumbnailSize, FRendererPreviewData InData)
+TSharedRef<SWidget> SNiagaraOverviewStackNode::CreateThumbnailWidget(float InThumbnailSize, FRendererPreviewData* InData)
 {
-	TSharedPtr<FAssetThumbnail> AssetThumbnail = MakeShareable(new FAssetThumbnail(InData.RenderingObject, InThumbnailSize, InThumbnailSize, ThumbnailPool));
+	TSharedPtr<FAssetThumbnail> AssetThumbnail = MakeShareable(new FAssetThumbnail(InData->RenderingObject, InThumbnailSize, InThumbnailSize, ThumbnailPool));
 	TSharedRef<SWidget> ThumbnailWidget = AssetThumbnail->MakeThumbnailWidget();
-	TSharedPtr<FAssetThumbnail> TooltipThumbnail = MakeShareable(new FAssetThumbnail(InData.RenderingObject, 64.0f, 64.0f, ThumbnailPool));
+	TSharedPtr<FAssetThumbnail> TooltipThumbnail = MakeShareable(new FAssetThumbnail(InData->RenderingObject, 64.0f, 64.0f, ThumbnailPool));
 	TSharedRef<SToolTip> ThumbnailTooltipWidget = SNew(SToolTip)
 		.Content()
 		[
@@ -146,7 +146,7 @@ TSharedRef<SWidget> SNiagaraOverviewStackNode::CreateThumbnailWidget(float InThu
 				TooltipThumbnail->MakeThumbnailWidget()
 			]
 		];
-	ThumbnailWidget->SetOnMouseButtonDown(FPointerEventHandler::CreateSP(this, &SNiagaraOverviewStackNode::OnClickedRenderingPreview, InData.RenderingEntry));
+	ThumbnailWidget->SetOnMouseButtonDown(FPointerEventHandler::CreateSP(this, &SNiagaraOverviewStackNode::OnClickedRenderingPreview, InData->RenderingEntry));
 	ThumbnailWidget->SetToolTip(ThumbnailTooltipWidget);
 	return ThumbnailWidget;
 }
@@ -251,9 +251,9 @@ void SNiagaraOverviewStackNode::FillThumbnailBar(UObject* ChangedObject, const b
 		{
 			EmitterHandleViewModelWeak.Pin()->GetRendererPreviewData(PreviewData);
 
-			for (FRendererPreviewData Preview : PreviewData)
+			for (FRendererPreviewData* Preview : PreviewData)
 			{
-				if (Preview.RenderingObject)
+				if (Preview->RenderingObject)
 				{
 					ThumbnailBar->AddSlot()
 						.AutoWidth()
diff --git a/Engine/Plugins/FX/Niagara/Source/NiagaraEditorWidgets/Private/SNiagaraOverviewStackNode.h b/Engine/Plugins/FX/Niagara/Source/NiagaraEditorWidgets/Private/SNiagaraOverviewStackNode.h
--- a/Engine/Plugins/FX/Niagara/Source/NiagaraEditorWidgets/Private/SNiagaraOverviewStackNode.h
+++ b/Engine/Plugins/FX/Niagara/Source/NiagaraEditorWidgets/Private/SNiagaraOverviewStackNode.h
@@ -33,7 +33,7 @@ private:
 	EVisibility GetEnabledCheckBoxVisibility() const;
 	ECheckBoxState GetEnabledCheckState() const;
 	void OnEnabledCheckStateChanged(ECheckBoxState InCheckState);
-	TSharedRef<SWidget> CreateThumbnailWidget(float InThumbnailSize, FRendererPreviewData InData);
+	TSharedRef<SWidget> CreateThumbnailWidget(float InThumbnailSize, FRendererPreviewData* InData);
 	FReply OnClickedRenderingPreview(const FGeometry& InGeometry, const FPointerEvent& InEvent, class UNiagaraStackEntry* InEntry);
 	FText GetToggleIsolateToolTip() const;
 	FReply OnToggleIsolateButtonClicked();
@@ -48,5 +48,6 @@ private:
 	TSharedPtr<FAssetThumbnailPool> ThumbnailPool;
 	/** Thumbnail widget containers */
 	TSharedPtr<SHorizontalBox> ThumbnailBar;
-	TArray<FRendererPreviewData> PreviewData;
+	TArray<FRendererPreviewData*> PreviewData;
+	
 };
\ No newline at end of file
