Added a context menu to Niagara emitter nodes in the System Overview panel.

Unified the context menus of all the places where you can summon it for emitters: Timeline, Selection and System Overview views.

diff --git a/Engine/Plugins/FX/Niagara/Source/NiagaraEditor/Private/NiagaraEditorUtilities.cpp b/Engine/Plugins/FX/Niagara/Source/NiagaraEditor/Private/NiagaraEditorUtilities.cpp
--- a/Engine/Plugins/FX/Niagara/Source/NiagaraEditor/Private/NiagaraEditorUtilities.cpp
+++ b/Engine/Plugins/FX/Niagara/Source/NiagaraEditor/Private/NiagaraEditorUtilities.cpp
@@ -27,6 +27,7 @@
 #include "EditorStyleSet.h"
 #include "ViewModels/NiagaraSystemViewModel.h"
 #include "ViewModels/NiagaraEmitterViewModel.h"
+#include "ViewModels/NiagaraEmitterHandleViewModel.h"
 #include "HAL/PlatformApplicationMisc.h"
 #include "AssetRegistryModule.h"
 #include "Misc/FeedbackContext.h"
@@ -41,6 +42,12 @@
 #include "NiagaraParameterMapHistory.h"
 #include "ScopedTransaction.h"
 #include "NiagaraStackEditorData.h"
+#include "Framework/MultiBox/MultiBoxBuilder.h"
+#include "IContentBrowserSingleton.h"
+#include "ContentBrowserModule.h"
+#include "Toolkits/AssetEditorManager.h"
+#include "Modules/ModuleManager.h"
+#include "ViewModels/NiagaraSystemSelectionViewModel.h"
 
 #define LOCTEXT_NAMESPACE "FNiagaraEditorUtilities"
 
@@ -1263,4 +1270,199 @@ bool FNiagaraEditorUtilities::AddParameter(FNiagaraVariable& NewParameterVariabl
 	return bSuccess;
 }
 
+void FNiagaraEditorUtilities::ShowParentEmitterInContentBrowser(TSharedRef<FNiagaraEmitterViewModel> Emitter)
+{
+	FContentBrowserModule& ContentBrowserModule = FModuleManager::Get().LoadModuleChecked<FContentBrowserModule>(TEXT("ContentBrowser"));
+	ContentBrowserModule.Get().SyncBrowserToAssets(TArray<FAssetData> { FAssetData(Emitter->GetParentEmitter()) });
+}
+
+ECheckBoxState FNiagaraEditorUtilities::GetSelectedEmittersEnabledCheckState(TSharedRef<FNiagaraSystemViewModel> SystemViewModel)
+{
+	bool bFirst = true;
+	ECheckBoxState CurrentState = ECheckBoxState::Undetermined;
+
+	const TArray<FGuid>& SelectedHandleIds = SystemViewModel->GetSelectionViewModel()->GetSelectedEmitterHandleIds();
+	for (const TSharedRef<FNiagaraEmitterHandleViewModel>& EmitterHandle : SystemViewModel->GetEmitterHandleViewModels())
+	{
+		if (SelectedHandleIds.Contains(EmitterHandle->GetId()))
+		{
+			ECheckBoxState EmitterState = EmitterHandle->GetIsEnabled() ? ECheckBoxState::Checked : ECheckBoxState::Unchecked;
+			if (bFirst)
+			{
+				CurrentState = EmitterState;
+				bFirst = false;
+				continue;
+			}
+
+			if (CurrentState != EmitterState)
+			{
+				return ECheckBoxState::Undetermined;
+			}
+		}
+	}
+
+	return CurrentState;
+}
+
+void FNiagaraEditorUtilities::ToggleSelectedEmittersEnabled(TSharedRef<FNiagaraSystemViewModel> SystemViewModel)
+{
+	bool bEnabled = GetSelectedEmittersEnabledCheckState(SystemViewModel) != ECheckBoxState::Checked;
+
+	const TArray<FGuid>& SelectedHandleIds = SystemViewModel->GetSelectionViewModel()->GetSelectedEmitterHandleIds();
+	for (const FGuid& HandleId : SelectedHandleIds)
+	{
+		TSharedPtr<FNiagaraEmitterHandleViewModel> EmitterHandleViewModel = SystemViewModel->GetEmitterHandleViewModelById(HandleId);
+		if (EmitterHandleViewModel.IsValid())
+		{
+			EmitterHandleViewModel->SetIsEnabled(bEnabled);
+		}
+	}
+}
+
+ECheckBoxState FNiagaraEditorUtilities::GetSelectedEmittersIsolatedCheckState(TSharedRef<FNiagaraSystemViewModel> SystemViewModel)
+{
+	bool bFirst = true;
+	ECheckBoxState CurrentState = ECheckBoxState::Undetermined;
+
+	const TArray<FGuid>& SelectedHandleIds = SystemViewModel->GetSelectionViewModel()->GetSelectedEmitterHandleIds();
+	for (const TSharedRef<FNiagaraEmitterHandleViewModel>& EmitterHandle : SystemViewModel->GetEmitterHandleViewModels())
+	{
+		if (SelectedHandleIds.Contains(EmitterHandle->GetId()))
+		{
+			ECheckBoxState EmitterState = EmitterHandle->GetIsIsolated() ? ECheckBoxState::Checked : ECheckBoxState::Unchecked;
+			if (bFirst)
+			{
+				CurrentState = EmitterState;
+				bFirst = false;
+				continue;
+			}
+
+			if (CurrentState != EmitterState)
+			{
+				return ECheckBoxState::Undetermined;
+			}
+		}
+	}
+
+	return CurrentState;
+}
+
+void FNiagaraEditorUtilities::ToggleSelectedEmittersIsolated(TSharedRef<FNiagaraSystemViewModel> SystemViewModel)
+{
+	bool bIsolated = GetSelectedEmittersIsolatedCheckState(SystemViewModel) != ECheckBoxState::Checked;
+
+	TArray<FGuid> EmittersToIsolate;
+	for (const TSharedRef<FNiagaraEmitterHandleViewModel>& EmitterHandle : SystemViewModel->GetEmitterHandleViewModels())
+	{
+		if (EmitterHandle->GetIsIsolated())
+		{
+			EmittersToIsolate.Add(EmitterHandle->GetId());
+		}
+	}
+
+	const TArray<FGuid>& SelectedHandleIds = SystemViewModel->GetSelectionViewModel()->GetSelectedEmitterHandleIds();
+	for (const FGuid& HandleId : SelectedHandleIds)
+	{
+		if (bIsolated)
+		{
+			EmittersToIsolate.Add(HandleId);
+		}
+		else
+		{
+			EmittersToIsolate.Remove(HandleId);
+		}
+	}
+
+	SystemViewModel->IsolateEmitters(EmittersToIsolate);
+}
+
+bool FNiagaraEditorUtilities::AddEmitterContextMenuActions(FMenuBuilder& MenuBuilder, const TSharedPtr<FNiagaraEmitterHandleViewModel>& EmitterHandleViewModelPtr)
+{
+	if (EmitterHandleViewModelPtr.IsValid())
+	{
+		TSharedRef<FNiagaraEmitterHandleViewModel> EmitterHandleViewModel = EmitterHandleViewModelPtr.ToSharedRef();
+		TSharedRef<FNiagaraSystemViewModel> OwningSystemViewModel = EmitterHandleViewModel->GetOwningSystemViewModel();
+
+		bool bSingleSelection = OwningSystemViewModel->GetSelectionViewModel()->GetSelectedEmitterHandleIds().Num() == 1;
+
+		MenuBuilder.BeginSection("EmitterActions", LOCTEXT("EmitterActions", "Emitter Actions"));
+		{
+			TSharedRef<FNiagaraEmitterViewModel> EmitterViewModel = EmitterHandleViewModel->GetEmitterViewModel();
+
+			if (OwningSystemViewModel->GetEditMode() == ENiagaraSystemViewModelEditMode::SystemAsset)
+			{
+				MenuBuilder.AddMenuEntry(
+					LOCTEXT("ToggleEmittersEnabled", "Enabled"),
+					LOCTEXT("ToggleEmittersEnabledToolTip", "Toggle whether or not the selected emitters are enabled."),
+					FSlateIcon(),
+					FUIAction(
+						FExecuteAction::CreateStatic(&FNiagaraEditorUtilities::ToggleSelectedEmittersEnabled, OwningSystemViewModel),
+						FCanExecuteAction(),
+						FGetActionCheckState::CreateStatic(&FNiagaraEditorUtilities::GetSelectedEmittersEnabledCheckState, OwningSystemViewModel)
+					),
+					NAME_None,
+					EUserInterfaceActionType::ToggleButton
+				);
+
+				MenuBuilder.AddMenuEntry(
+					LOCTEXT("ToggleEmittersIsolated", "Isolated"),
+					LOCTEXT("ToggleEmittersIsolatedToolTip", "Toggle whether or not the selected emitters are isolated."),
+					FSlateIcon(),
+					FUIAction(
+						FExecuteAction::CreateStatic(&FNiagaraEditorUtilities::ToggleSelectedEmittersIsolated, OwningSystemViewModel),
+						FCanExecuteAction(),
+						FGetActionCheckState::CreateStatic(&FNiagaraEditorUtilities::GetSelectedEmittersIsolatedCheckState, OwningSystemViewModel)
+					),
+					NAME_None,
+					EUserInterfaceActionType::ToggleButton
+				);
+
+				MenuBuilder.AddMenuEntry(
+					LOCTEXT("RenameEmitter", "Rename Emitter"),
+					LOCTEXT("RenameEmitterToolTip", "Rename this local emitter copy."),
+					FSlateIcon(),
+					FUIAction(
+						FExecuteAction::CreateSP(EmitterHandleViewModel, &FNiagaraEmitterHandleViewModel::SetIsRenamePending, true)
+					)
+				);
+			}
+
+			MenuBuilder.AddMenuEntry(
+				LOCTEXT("RemoveParentEmitter", "Remove Parent Emitter"),
+				LOCTEXT("RemoveParentEmitterToolTip", "Remove this emitter's parent, preventing inheritance of any further changes."),
+				FSlateIcon(),
+				FUIAction(
+					FExecuteAction::CreateSP(EmitterViewModel, &FNiagaraEmitterViewModel::RemoveParentEmitter),
+					FCanExecuteAction::CreateLambda(
+						[bSingleSelection, bHasParent = EmitterViewModel->HasParentEmitter()]()
+						{
+							return bSingleSelection && bHasParent;
+						}
+					)
+				)
+			);
+
+			MenuBuilder.AddMenuEntry(
+				LOCTEXT("ShowEmitterInContentBrowser", "Show in Content Browser"),
+				LOCTEXT("ShowEmitterInContentBrowserToolTip", "Show the selected emitter in the Content Browser."),
+				FSlateIcon(),
+				FUIAction(
+					FExecuteAction::CreateStatic(&FNiagaraEditorUtilities::ShowParentEmitterInContentBrowser, EmitterViewModel),
+					FCanExecuteAction::CreateLambda(
+						[bSingleSelection, bHasParent = EmitterViewModel->HasParentEmitter()]()
+						{
+							return bSingleSelection && bHasParent;
+						}
+					)
+				)
+			);
+		}
+		MenuBuilder.EndSection();
+
+		return true;
+	}
+
+	return false;
+}
+
 #undef LOCTEXT_NAMESPACE
diff --git a/Engine/Plugins/FX/Niagara/Source/NiagaraEditor/Private/NiagaraOverviewNode.cpp b/Engine/Plugins/FX/Niagara/Source/NiagaraEditor/Private/NiagaraOverviewNode.cpp
--- a/Engine/Plugins/FX/Niagara/Source/NiagaraEditor/Private/NiagaraOverviewNode.cpp
+++ b/Engine/Plugins/FX/Niagara/Source/NiagaraEditor/Private/NiagaraOverviewNode.cpp
@@ -4,8 +4,12 @@
 #include "NiagaraSystem.h"
 #include "NiagaraEditorModule.h"
 #include "NiagaraEditorStyle.h"
+#include "NiagaraEditorUtilities.h"
 #include "ViewModels/Stack/NiagaraStackEntry.h"
-
+#include "ViewModels/NiagaraEmitterViewModel.h"
+#include "ViewModels/NiagaraEmitterHandleViewModel.h"
+#include "ViewModels/NiagaraSystemViewModel.h"
+#include "ViewModels/NiagaraSystemSelectionViewModel.h"
 #include "Modules/ModuleManager.h"
 #include "ToolMenuSection.h"
 #include "ToolMenu.h"
@@ -22,7 +26,9 @@ FLinearColor UNiagaraOverviewNode::NotIsolatedColor;
 UNiagaraOverviewNode::UNiagaraOverviewNode()
 	: OwningSystem(nullptr)
 	, EmitterHandleGuid(FGuid())
+	, bRenamePending(false)
 {
+	UEdGraphNode::bCanRenameNode = true;
 };
 
 void UNiagaraOverviewNode::Initialize(UNiagaraSystem* InOwningSystem)
@@ -41,15 +47,16 @@ const FGuid UNiagaraOverviewNode::GetEmitterHandleGuid() const
 	return EmitterHandleGuid;
 }
 
-static const FNiagaraEmitterHandle* FindEmitterHandleByID(const UNiagaraSystem* System, const FGuid& Guid)
+static FNiagaraEmitterHandle* FindEmitterHandleByID(UNiagaraSystem* System, const FGuid& Guid)
 {
 	check(System != nullptr);
 
-	for (const FNiagaraEmitterHandle& Handle : System->GetEmitterHandles())
+	for (int Idx = 0; Idx < System->GetNumEmitters(); ++Idx)
 	{
-		if (Handle.GetId() == Guid)
+		FNiagaraEmitterHandle& EmitterHandle = System->GetEmitterHandle(Idx);
+		if (EmitterHandle.GetId() == Guid)
 		{
-			return &Handle;
+			return &EmitterHandle;
 		}
 	}
 
@@ -113,28 +120,138 @@ bool UNiagaraOverviewNode::CanDuplicateNode() const
 	return HasAnyFlags(RF_ClassDefaultObject) || EmitterHandleGuid.IsValid();
 }
 
-void UNiagaraOverviewNode::GetNodeContextMenuActions(class UToolMenu* Menu, class UGraphNodeContextMenuContext* Context) const
+void UNiagaraOverviewNode::OnRenameNode(const FString& NewName)
 {
-	FToolMenuSection& Section = Menu->AddSection("Alignment");
-	Section.AddSubMenu(
-		"Alignment",
-		LOCTEXT("AlignmentHeader", "Alignment"),
-		FText(),
-		FNewToolMenuDelegate::CreateLambda([](UToolMenu* InMenu)
+	bRenamePending = false;
+
+	FNiagaraEditorModule& NiagaraEditorModule = FModuleManager::Get().LoadModuleChecked<FNiagaraEditorModule>("NiagaraEditor");
+	TSharedPtr<FNiagaraSystemViewModel> OwningSystemViewModel = NiagaraEditorModule.GetExistingViewModelForSystem(OwningSystem);
+	if (OwningSystemViewModel.IsValid())
 	{
+		TSharedPtr<FNiagaraEmitterHandleViewModel> EmitterHandleViewModelPtr = OwningSystemViewModel->GetEmitterHandleViewModelById(EmitterHandleGuid);
+		if (EmitterHandleViewModelPtr.IsValid())
 		{
-			FToolMenuSection& SubMenuSection = InMenu->AddSection("EdGraphSchemaAlignment", LOCTEXT("AlignHeader", "Align"));
-			SubMenuSection.AddMenuEntry(FGraphEditorCommands::Get().AlignNodesTop);
-			SubMenuSection.AddMenuEntry(FGraphEditorCommands::Get().AlignNodesMiddle);
-			SubMenuSection.AddMenuEntry(FGraphEditorCommands::Get().AlignNodesBottom);
+			EmitterHandleViewModelPtr->SetName(*NewName);
 		}
+	}
+}
+
+void UNiagaraOverviewNode::GetNodeContextMenuActions(class UToolMenu* Menu, class UGraphNodeContextMenuContext* Context) const
+{
+	FNiagaraEditorModule& NiagaraEditorModule = FModuleManager::Get().LoadModuleChecked<FNiagaraEditorModule>("NiagaraEditor");
+	TSharedPtr<FNiagaraSystemViewModel> OwningSystemViewModel = NiagaraEditorModule.GetExistingViewModelForSystem(OwningSystem);
+	if (OwningSystemViewModel.IsValid())
+	{
+		FToolMenuSection& Section = Menu->AddSection("Emitter Actions");
+
+		bool bSingleSelection = OwningSystemViewModel->GetSelectionViewModel()->GetSelectedEmitterHandleIds().Num() == 1;
 
+		TSharedPtr<FNiagaraEmitterHandleViewModel> EmitterHandleViewModelPtr = OwningSystemViewModel->GetEmitterHandleViewModelById(EmitterHandleGuid);
+		if (EmitterHandleViewModelPtr.IsValid())
 		{
-			FToolMenuSection& SubMenuSection = InMenu->AddSection("EdGraphSchemaDistribution", LOCTEXT("DistributionHeader", "Distribution"));
-			SubMenuSection.AddMenuEntry(FGraphEditorCommands::Get().DistributeNodesHorizontally);
-			SubMenuSection.AddMenuEntry(FGraphEditorCommands::Get().DistributeNodesVertically);
+			TSharedRef<FNiagaraEmitterHandleViewModel> EmitterHandleViewModel = EmitterHandleViewModelPtr.ToSharedRef();
+			TSharedRef<FNiagaraEmitterViewModel> EmitterViewModel = EmitterHandleViewModel->GetEmitterViewModel();
+			{
+				if (OwningSystemViewModel->GetEditMode() == ENiagaraSystemViewModelEditMode::SystemAsset)
+				{
+					Section.AddMenuEntry(
+						"ToggleEmittersEnabled",
+						LOCTEXT("ToggleEmittersEnabled", "Enabled"),
+						LOCTEXT("ToggleEmittersEnabledToolTip", "Toggle whether or not the selected emitters are enabled."),
+						FSlateIcon(),
+						FUIAction(
+							FExecuteAction::CreateStatic(&FNiagaraEditorUtilities::ToggleSelectedEmittersEnabled, OwningSystemViewModel.ToSharedRef()),
+							FCanExecuteAction(),
+							FGetActionCheckState::CreateStatic(&FNiagaraEditorUtilities::GetSelectedEmittersEnabledCheckState, OwningSystemViewModel.ToSharedRef())
+						),
+						EUserInterfaceActionType::ToggleButton
+					);
+
+					Section.AddMenuEntry(
+						"ToggleEmittersIsolated",
+						LOCTEXT("ToggleEmittersIsolated", "Isolated"),
+						LOCTEXT("ToggleEmittersIsolatedToolTip", "Toggle whether or not the selected emitters are isolated."),
+						FSlateIcon(),
+						FUIAction(
+							FExecuteAction::CreateStatic(&FNiagaraEditorUtilities::ToggleSelectedEmittersIsolated, OwningSystemViewModel.ToSharedRef()),
+							FCanExecuteAction(),
+							FGetActionCheckState::CreateStatic(&FNiagaraEditorUtilities::GetSelectedEmittersIsolatedCheckState, OwningSystemViewModel.ToSharedRef())
+						),
+						EUserInterfaceActionType::ToggleButton
+					);
+
+					Section.AddMenuEntry(
+						"RenameEmitter",
+						LOCTEXT("RenameEmitter", "Rename Emitter"),
+						LOCTEXT("RenameEmitterToolTip", "Rename this local emitter copy."),
+						FSlateIcon(),
+						FUIAction(
+							FExecuteAction::CreateLambda([this]()
+							{
+								const_cast<UNiagaraOverviewNode*>(this)->RequestRename();
+							}),
+							FCanExecuteAction::CreateLambda([bSingleSelection]() { return bSingleSelection; })
+						)
+					);
+				}
+
+				Section.AddMenuEntry(
+					"RemoveParentEmitter",
+					LOCTEXT("RemoveParentEmitter", "Remove Parent Emitter"),
+					LOCTEXT("RemoveParentEmitterToolTip", "Removes this emitter's parent, preventing inheritance of any further changes."),
+					FSlateIcon(),
+					FUIAction(
+						FExecuteAction::CreateSP(EmitterViewModel, &FNiagaraEmitterViewModel::RemoveParentEmitter),
+						FCanExecuteAction::CreateLambda(
+							[bSingleSelection, bHasParent = EmitterViewModel->HasParentEmitter()]()
+							{
+								return bSingleSelection && bHasParent;
+							}
+						)
+					)
+				);
+
+				Section.AddMenuEntry(
+					"ShowEmitterInContentBrowser",
+					LOCTEXT("ShowEmitterInContentBrowser", "Show in Content Browser"),
+					LOCTEXT("ShowEmitterInContentBrowserToolTip", "Show the selected emitter in the Content Browser."),
+					FSlateIcon(),
+					FUIAction(
+						FExecuteAction::CreateStatic(&FNiagaraEditorUtilities::ShowParentEmitterInContentBrowser, EmitterViewModel),
+						FCanExecuteAction::CreateLambda(
+							[bSingleSelection, bHasParent = EmitterViewModel->HasParentEmitter()]()
+							{
+								return bSingleSelection && bHasParent;
+							}
+						)
+					)
+				);
+			}
 		}
-	}));
+	}
+
+	{
+		FToolMenuSection& Section = Menu->AddSection("Alignment");
+		Section.AddSubMenu(
+			"Alignment",
+			LOCTEXT("AlignmentHeader", "Alignment"),
+			FText(),
+			FNewToolMenuDelegate::CreateLambda([](UToolMenu* InMenu)
+		{
+			{
+				FToolMenuSection& SubMenuSection = InMenu->AddSection("EdGraphSchemaAlignment", LOCTEXT("AlignHeader", "Align"));
+				SubMenuSection.AddMenuEntry(FGraphEditorCommands::Get().AlignNodesTop);
+				SubMenuSection.AddMenuEntry(FGraphEditorCommands::Get().AlignNodesMiddle);
+				SubMenuSection.AddMenuEntry(FGraphEditorCommands::Get().AlignNodesBottom);
+			}
+
+			{
+				FToolMenuSection& SubMenuSection = InMenu->AddSection("EdGraphSchemaDistribution", LOCTEXT("DistributionHeader", "Distribution"));
+				SubMenuSection.AddMenuEntry(FGraphEditorCommands::Get().DistributeNodesHorizontally);
+				SubMenuSection.AddMenuEntry(FGraphEditorCommands::Get().DistributeNodesVertically);
+			}
+		}));
+	}
 }
 
 UNiagaraSystem* UNiagaraOverviewNode::GetOwningSystem()
diff --git a/Engine/Plugins/FX/Niagara/Source/NiagaraEditor/Private/Sequencer/NiagaraSequence/NiagaraEmitterTrackEditor.cpp b/Engine/Plugins/FX/Niagara/Source/NiagaraEditor/Private/Sequencer/NiagaraSequence/NiagaraEmitterTrackEditor.cpp
--- a/Engine/Plugins/FX/Niagara/Source/NiagaraEditor/Private/Sequencer/NiagaraSequence/NiagaraEmitterTrackEditor.cpp
+++ b/Engine/Plugins/FX/Niagara/Source/NiagaraEditor/Private/Sequencer/NiagaraSequence/NiagaraEmitterTrackEditor.cpp
@@ -12,7 +12,7 @@
 #include "Framework/MultiBox/MultiBoxBuilder.h"
 #include "NiagaraEditorStyle.h"
 #include "NiagaraEditorModule.h"
-
+#include "NiagaraEditorUtilities.h"
 #include "EditorStyleSet.h"
 #include "Styling/SlateIconFinder.h"
 #include "Widgets/SBoxPanel.h"
@@ -246,26 +246,7 @@ void FNiagaraEmitterTrackEditor::BuildTrackContextMenu( FMenuBuilder& MenuBuilde
 
 	if (SystemViewModel.GetEditMode() == ENiagaraSystemViewModelEditMode::SystemAsset)
 	{
-		MenuBuilder.BeginSection("Niagara", LOCTEXT("NiagaraContextMenuSectionName", "Niagara"));
-		{
-			MenuBuilder.AddMenuEntry(
-				EmitterTrack->GetEmitterHandleViewModel()->GetIsIsolated()
-					? LOCTEXT("RemoveFromIsolation", "Remove this from isolation.")
-					: LOCTEXT("AddToIsolation", "Add this to isolation"),
-				EmitterTrack->GetEmitterHandleViewModel()->GetIsIsolated()
-					? LOCTEXT("RemoveFromIsolation_NoChangeOthers", "Remove this emitter from isolation, without changing other emitters.")
-					: LOCTEXT("AddToIsolation_NoChangeOthers", "Add this emitter to isolation, without changing other emitters."),
-				FSlateIcon(),
-				FUIAction(FExecuteAction::CreateRaw(&SystemViewModel, &FNiagaraSystemViewModel::ToggleEmitterIsolation, EmitterTrack->GetEmitterHandleViewModel().ToSharedRef())));
-			
-			TArray<FGuid> SelectedEmitterHandleIds = SystemViewModel.GetSelectionViewModel()->GetSelectedEmitterHandleIds();
-			MenuBuilder.AddMenuEntry(
-				LOCTEXT("IsolateSelected", "Isolate all selected"),
-				LOCTEXT("IsolateSelectedToolTip", "Add all of the selected emitters to isloation"),
-				FSlateIcon(),
-				FUIAction(FExecuteAction::CreateRaw(&SystemViewModel, &FNiagaraSystemViewModel::IsolateEmitters, SelectedEmitterHandleIds)));
-		}
-		MenuBuilder.EndSection();
+		FNiagaraEditorUtilities::AddEmitterContextMenuActions(MenuBuilder, EmitterTrack->GetEmitterHandleViewModel());
 	}
 }
 
diff --git a/Engine/Plugins/FX/Niagara/Source/NiagaraEditor/Public/NiagaraEditorUtilities.h b/Engine/Plugins/FX/Niagara/Source/NiagaraEditor/Public/NiagaraEditorUtilities.h
--- a/Engine/Plugins/FX/Niagara/Source/NiagaraEditor/Public/NiagaraEditorUtilities.h
+++ b/Engine/Plugins/FX/Niagara/Source/NiagaraEditor/Public/NiagaraEditorUtilities.h
@@ -28,6 +28,10 @@ class UEdGraphSchema_Niagara;
 class UEdGraphPin;
 class FCompileConstantResolver;
 class UNiagaraStackEditorData;
+class FMenuBuilder;
+class FNiagaraEmitterViewModel;
+class FNiagaraEmitterHandleViewModel;
+enum class ECheckBoxState : uint8;
 
 namespace FNiagaraEditorUtilities
 {
@@ -210,4 +214,14 @@ namespace FNiagaraEditorUtilities
 	 * @returns Bool for whether adding the parameter succeeded.
 	 */
 	bool AddParameter(FNiagaraVariable& NewParameterVariable, FNiagaraParameterStore& TargetParameterStore, UObject& ParameterStoreOwner, UNiagaraStackEditorData& StackEditorData);
+
+	NIAGARAEDITOR_API bool AddEmitterContextMenuActions(FMenuBuilder& MenuBuilder, const TSharedPtr<FNiagaraEmitterHandleViewModel>& EmitterHandleViewModel);
+
+	void ShowParentEmitterInContentBrowser(TSharedRef<FNiagaraEmitterViewModel> EmitterViewModel);
+
+	ECheckBoxState GetSelectedEmittersEnabledCheckState(TSharedRef<FNiagaraSystemViewModel> SystemViewModel);
+	void ToggleSelectedEmittersEnabled(TSharedRef<FNiagaraSystemViewModel> SystemViewModel);
+
+	ECheckBoxState GetSelectedEmittersIsolatedCheckState(TSharedRef<FNiagaraSystemViewModel> SystemViewModel);
+	void ToggleSelectedEmittersIsolated(TSharedRef<FNiagaraSystemViewModel> SystemViewModel);
 };
diff --git a/Engine/Plugins/FX/Niagara/Source/NiagaraEditor/Public/NiagaraOverviewNode.h b/Engine/Plugins/FX/Niagara/Source/NiagaraEditor/Public/NiagaraOverviewNode.h
--- a/Engine/Plugins/FX/Niagara/Source/NiagaraEditor/Public/NiagaraOverviewNode.h
+++ b/Engine/Plugins/FX/Niagara/Source/NiagaraEditor/Public/NiagaraOverviewNode.h
@@ -29,11 +29,17 @@ public:
 	/** Whether or not this node can be safely duplicated (via copy/paste, etc...) in the graph */
 	virtual bool CanDuplicateNode() const override;
 
+	virtual void OnRenameNode(const FString& NewName) override;
+
 	virtual void GetNodeContextMenuActions(class UToolMenu* Menu, class UGraphNodeContextMenuContext* Context) const override;
 	//~ End UEdGraphNode Interface
 
 	UNiagaraSystem* GetOwningSystem();
 
+	void RequestRename() { bRenamePending = true; }
+	void RenameStarted() { bRenamePending = false; }
+	bool IsRenamePending() const { return bRenamePending; }
+
 private:
 	UPROPERTY()
 	UNiagaraSystem* OwningSystem;
@@ -41,6 +47,8 @@ private:
 	UPROPERTY()
 	FGuid EmitterHandleGuid;
 
+	bool bRenamePending;
+
 	static bool bColorsAreInitialized;
 	static FLinearColor SystemColor;
 	static FLinearColor EmitterColor;
diff --git a/Engine/Plugins/FX/Niagara/Source/NiagaraEditor/Public/ViewModels/NiagaraEmitterHandleViewModel.h b/Engine/Plugins/FX/Niagara/Source/NiagaraEditor/Public/ViewModels/NiagaraEmitterHandleViewModel.h
--- a/Engine/Plugins/FX/Niagara/Source/NiagaraEditor/Public/ViewModels/NiagaraEmitterHandleViewModel.h
+++ b/Engine/Plugins/FX/Niagara/Source/NiagaraEditor/Public/ViewModels/NiagaraEmitterHandleViewModel.h
@@ -133,8 +133,8 @@ public:
 	void Cleanup();
 
 	NIAGARAEDITOR_API void GetRendererPreviewData(TArray<FRendererPreviewData*>& InRendererPreviewData);
-private:
-	TSharedRef<FNiagaraSystemViewModel> GetOwningSystemViewModel() const;
+
+	NIAGARAEDITOR_API TSharedRef<FNiagaraSystemViewModel> GetOwningSystemViewModel() const;
 
 private:
 	/** The system view model which owns this emitter handle view model. */
diff --git a/Engine/Plugins/FX/Niagara/Source/NiagaraEditor/Public/ViewModels/NiagaraSystemViewModel.h b/Engine/Plugins/FX/Niagara/Source/NiagaraEditor/Public/ViewModels/NiagaraSystemViewModel.h
--- a/Engine/Plugins/FX/Niagara/Source/NiagaraEditor/Public/ViewModels/NiagaraSystemViewModel.h
+++ b/Engine/Plugins/FX/Niagara/Source/NiagaraEditor/Public/ViewModels/NiagaraSystemViewModel.h
@@ -250,7 +250,7 @@ public:
 	void UpdateEmitterFixedBounds();
 
 	/** Isolates the supplied emitters.  This will remove all other emitters from isolation. */
-	void IsolateEmitters(TArray<FGuid> EmitterHandlesIdsToIsolate);
+	NIAGARAEDITOR_API void IsolateEmitters(TArray<FGuid> EmitterHandlesIdsToIsolate);
 
 	/** Toggles the isolation state of a single emitter. */
 	void ToggleEmitterIsolation(TSharedRef<FNiagaraEmitterHandleViewModel> InEmitterHandle);
diff --git a/Engine/Plugins/FX/Niagara/Source/NiagaraEditorWidgets/Private/SNiagaraOverviewGraph.cpp b/Engine/Plugins/FX/Niagara/Source/NiagaraEditorWidgets/Private/SNiagaraOverviewGraph.cpp
--- a/Engine/Plugins/FX/Niagara/Source/NiagaraEditorWidgets/Private/SNiagaraOverviewGraph.cpp
+++ b/Engine/Plugins/FX/Niagara/Source/NiagaraEditorWidgets/Private/SNiagaraOverviewGraph.cpp
@@ -206,6 +206,12 @@ FActionMenuContent SNiagaraOverviewGraph::OnCreateGraphActionMenu(UEdGraph* InGr
 				LOCTEXT("CommentsToolTip", "Add a comment box"),
 				FSlateIcon(),
 				FExecuteAction::CreateSP(this, &SNiagaraOverviewGraph::OnCreateComment));
+
+			MenuBuilder.AddMenuEntry(
+				LOCTEXT("ClearIsolatedLabel", "Clear Isolated"),
+				LOCTEXT("ClearIsolatedToolTip", "Clear the current set of isolated emitters."),
+				FSlateIcon(),
+				FExecuteAction::CreateSP(this, &SNiagaraOverviewGraph::OnClearIsolated));
 		}
 		MenuBuilder.EndSection();
 
@@ -229,6 +235,11 @@ void SNiagaraOverviewGraph::OnCreateComment()
 	CommentAction.PerformAction(ViewModel->GetGraph(), nullptr, GraphEditor->GetPasteLocation(), false);
 }
 
+void SNiagaraOverviewGraph::OnClearIsolated()
+{
+	ViewModel->GetSystemViewModel()->IsolateEmitters(TArray<FGuid>());
+}
+
 void SNiagaraOverviewGraph::OnNodeTitleCommitted(const FText& NewText, ETextCommit::Type CommitInfo, UEdGraphNode* NodeBeingChanged)
 {
 	if (NodeBeingChanged)
diff --git a/Engine/Plugins/FX/Niagara/Source/NiagaraEditorWidgets/Private/SNiagaraOverviewGraph.h b/Engine/Plugins/FX/Niagara/Source/NiagaraEditorWidgets/Private/SNiagaraOverviewGraph.h
--- a/Engine/Plugins/FX/Niagara/Source/NiagaraEditorWidgets/Private/SNiagaraOverviewGraph.h
+++ b/Engine/Plugins/FX/Niagara/Source/NiagaraEditorWidgets/Private/SNiagaraOverviewGraph.h
@@ -32,13 +32,13 @@ private:
 	FActionMenuContent OnCreateGraphActionMenu(UEdGraph* InGraph, const FVector2D& InNodePosition, const TArray<UEdGraphPin*>& InDraggedPins, bool bAutoExpand, SGraphEditor::FActionMenuClosed InOnMenuClosed);
 
 	void OnCreateComment();
+	void OnClearIsolated();
 
 	void OnNodeTitleCommitted(const FText& NewText, ETextCommit::Type CommitInfo, UEdGraphNode* NodeBeingChanged);
 
 	void CreateAddEmitterMenuContent(FMenuBuilder& MenuBuilder, UEdGraph* InGraph);
 
 	void ZoomToFit();
-
 	void ZoomToFitAll();
 
 	void OnAlignTop();
@@ -48,7 +48,6 @@ private:
 	void OnDistributeNodesH();
 	void OnDistributeNodesV();
 
-
 private:
 	TSharedPtr<FNiagaraOverviewGraphViewModel> ViewModel;
 	TSharedPtr<SGraphEditor> GraphEditor;
diff --git a/Engine/Plugins/FX/Niagara/Source/NiagaraEditorWidgets/Private/SNiagaraOverviewStackNode.cpp b/Engine/Plugins/FX/Niagara/Source/NiagaraEditorWidgets/Private/SNiagaraOverviewStackNode.cpp
--- a/Engine/Plugins/FX/Niagara/Source/NiagaraEditorWidgets/Private/SNiagaraOverviewStackNode.cpp
+++ b/Engine/Plugins/FX/Niagara/Source/NiagaraEditorWidgets/Private/SNiagaraOverviewStackNode.cpp
@@ -32,6 +32,7 @@ void SNiagaraOverviewStackNode::Construct(const FArguments& InArgs, UNiagaraOver
 	StackViewModel = nullptr;
 	OverviewSelectionViewModel = nullptr;
 	EmitterHandleViewModelWeak.Reset();
+
 	if (OverviewStackNode->GetOwningSystem() != nullptr)
 	{
 		FNiagaraEditorModule& NiagaraEditorModule = FModuleManager::Get().LoadModuleChecked<FNiagaraEditorModule>("NiagaraEditor");
@@ -57,6 +58,7 @@ void SNiagaraOverviewStackNode::Construct(const FArguments& InArgs, UNiagaraOver
 			OverviewSelectionViewModel = OwningSystemViewModel->GetSelectionViewModel();
 		}
 	}
+
 	UpdateGraphNode();
 }
 
@@ -164,6 +166,18 @@ FReply SNiagaraOverviewStackNode::OnClickedRenderingPreview(const FGeometry& InG
 	return FReply::Unhandled();
 }
 
+void SNiagaraOverviewStackNode::Tick(const FGeometry& AllottedGeometry, const double InCurrentTime, const float InDeltaTime)
+{
+	if (OverviewStackNode != nullptr)
+	{
+		if (OverviewStackNode->IsRenamePending() && !SGraphNode::IsRenamePending())
+		{
+			SGraphNode::RequestRename();
+			OverviewStackNode->RenameStarted();
+		}
+	}
+}
+
 TSharedRef<SWidget> SNiagaraOverviewStackNode::CreateNodeContentArea()
 {
 	TSharedPtr<SWidget> ContentWidget;
@@ -205,34 +219,34 @@ TSharedRef<SWidget> SNiagaraOverviewStackNode::CreateNodeContentArea()
 			.VAlign(VAlign_Center)
 			.Padding(0.0f)
 			[
-				SNew(SHorizontalBox)
-				+SHorizontalBox::Slot()
-				.AutoWidth()
-				.HAlign(HAlign_Left)
-				[
-					// LEFT
-					SAssignNew(LeftNodeBox, SVerticalBox)
-				]
-				+SHorizontalBox::Slot()
-				[
-					SNew(SBorder)
-					.BorderImage(FNiagaraEditorWidgetsStyle::Get().GetBrush("NiagaraEditor.SystemOverview.NodeBackgroundBorder"))
-					.BorderBackgroundColor(FNiagaraEditorWidgetsStyle::Get().GetColor("NiagaraEditor.SystemOverview.NodeBackgroundColor"))
-					.HAlign(HAlign_Fill)
-					.VAlign(VAlign_Fill)
-					.Padding(FMargin(0, 0, 0, 4))
-					[
-						ContentWidget.ToSharedRef()
-					]
-				]
-				+SHorizontalBox::Slot()
-				.AutoWidth()
-				.HAlign(HAlign_Right)
+			SNew(SHorizontalBox)
+			+SHorizontalBox::Slot()
+			.AutoWidth()
+			.HAlign(HAlign_Left)
+			[
+				// LEFT
+				SAssignNew(LeftNodeBox, SVerticalBox)
+			]
+			+SHorizontalBox::Slot()
+			[
+				SNew(SBorder)
+				.BorderImage(FNiagaraEditorWidgetsStyle::Get().GetBrush("NiagaraEditor.SystemOverview.NodeBackgroundBorder"))
+				.BorderBackgroundColor(FNiagaraEditorWidgetsStyle::Get().GetColor("NiagaraEditor.SystemOverview.NodeBackgroundColor"))
+				.HAlign(HAlign_Fill)
+				.VAlign(VAlign_Fill)
+				.Padding(FMargin(0, 0, 0, 4))
 				[
-					// RIGHT
-					SAssignNew(RightNodeBox, SVerticalBox)
+					ContentWidget.ToSharedRef()
 				]
 			]
+			+SHorizontalBox::Slot()
+			.AutoWidth()
+			.HAlign(HAlign_Right)
+			[
+				// RIGHT
+				SAssignNew(RightNodeBox, SVerticalBox)
+			]
+			]
 		];
 
 
diff --git a/Engine/Plugins/FX/Niagara/Source/NiagaraEditorWidgets/Private/SNiagaraOverviewStackNode.h b/Engine/Plugins/FX/Niagara/Source/NiagaraEditorWidgets/Private/SNiagaraOverviewStackNode.h
--- a/Engine/Plugins/FX/Niagara/Source/NiagaraEditorWidgets/Private/SNiagaraOverviewStackNode.h
+++ b/Engine/Plugins/FX/Niagara/Source/NiagaraEditorWidgets/Private/SNiagaraOverviewStackNode.h
@@ -23,6 +23,7 @@ public:
 
 	/** Constructs this widget with InArgs */
 	void Construct(const FArguments& InArgs, UNiagaraOverviewNode* InNode);
+	virtual void Tick(const FGeometry& AllottedGeometry, const double InCurrentTime, const float InDeltaTime) override;
 
 protected:
 	virtual TSharedRef<SWidget> CreateTitleWidget(TSharedPtr<SNodeTitle> NodeTitle) override;
diff --git a/Engine/Plugins/FX/Niagara/Source/NiagaraEditorWidgets/Private/SNiagaraStack.cpp b/Engine/Plugins/FX/Niagara/Source/NiagaraEditorWidgets/Private/SNiagaraStack.cpp
--- a/Engine/Plugins/FX/Niagara/Source/NiagaraEditorWidgets/Private/SNiagaraStack.cpp
+++ b/Engine/Plugins/FX/Niagara/Source/NiagaraEditorWidgets/Private/SNiagaraStack.cpp
@@ -54,6 +54,7 @@
 #include "ScopedTransaction.h"
 #include "Widgets/Layout/SWrapBox.h"
 #include "ViewModels/Stack/NiagaraStackRoot.h"
+#include "NiagaraEditorUtilities.h"
 #include "Subsystems/AssetEditorSubsystem.h"
 
 #define LOCTEXT_NAMESPACE "NiagaraStack"
@@ -262,15 +263,15 @@ void SNiagaraStack::Construct(const FArguments& InArgs, UNiagaraStackViewModel*
 			.BorderImage(FEditorStyle::GetBrush("WhiteBrush"))
 			.BorderBackgroundColor(FNiagaraEditorWidgetsStyle::Get().GetColor("NiagaraEditor.Stack.BackgroundColor"))
 			[
-				SAssignNew(StackTree, STreeView<UNiagaraStackEntry*>)
-				.OnGenerateRow(this, &SNiagaraStack::OnGenerateRowForStackItem)
-				.OnGetChildren(this, &SNiagaraStack::OnGetChildren)
-				.TreeItemsSource(&StackViewModel->GetRootEntryAsArray())
-				.OnTreeViewScrolled(this, &SNiagaraStack::StackTreeScrolled)
-				.SelectionMode(ESelectionMode::Single)
-				.OnItemToString_Debug_Static(&FNiagaraStackEditorWidgetsUtilities::StackEntryToStringForListDebug)
-				.OnMouseButtonClick(this, &SNiagaraStack::OnStackItemClicked)
-			]
+			SAssignNew(StackTree, STreeView<UNiagaraStackEntry*>)
+			.OnGenerateRow(this, &SNiagaraStack::OnGenerateRowForStackItem)
+			.OnGetChildren(this, &SNiagaraStack::OnGetChildren)
+			.TreeItemsSource(&StackViewModel->GetRootEntryAsArray())
+			.OnTreeViewScrolled(this, &SNiagaraStack::StackTreeScrolled)
+			.SelectionMode(ESelectionMode::Single)
+			.OnItemToString_Debug_Static(&FNiagaraStackEditorWidgetsUtilities::StackEntryToStringForListDebug)
+			.OnMouseButtonClick(this, &SNiagaraStack::OnStackItemClicked)
+		]
 		]
 	];
 
@@ -675,47 +676,7 @@ FReply SNiagaraStack::OnTopLevelRowMouseButtonDown(const FGeometry&, const FPoin
 	{
 		FMenuBuilder MenuBuilder(true, nullptr);
 
-		if (TopLevelViewModel->EmitterHandleViewModel.IsValid())
-		{
-			TSharedRef<FNiagaraEmitterHandleViewModel> EmitterHandleViewModel = TopLevelViewModel->EmitterHandleViewModel.ToSharedRef();
-			MenuBuilder.BeginSection("EmitterInlineMenuActions", LOCTEXT("EmitterActions", "Emitter Actions"));
-			{
-				FUIAction Action(FExecuteAction::CreateSP(EmitterHandleViewModel, &FNiagaraEmitterHandleViewModel::SetIsEnabled, !EmitterHandleViewModel->GetIsEnabled()),
-					FCanExecuteAction(),
-					FIsActionChecked::CreateSP(EmitterHandleViewModel, &FNiagaraEmitterHandleViewModel::GetIsEnabled));
-				MenuBuilder.AddMenuEntry(
-					LOCTEXT("IsEnabled", "Is Enabled"),
-					LOCTEXT("ToggleEmitterEnabledToolTip", "Toggle emitter enabled/disabled state"),
-					FSlateIcon(),
-					Action,
-					NAME_None,
-					EUserInterfaceActionType::Check);
-
-				if (EmitterHandleViewModel->CanRenameEmitter()) // Only allow renaming local copies of Emitters in Systems
-				{
-					MenuBuilder.AddMenuEntry(
-						LOCTEXT("RenameEmitter", "Rename Emitter"),
-						LOCTEXT("RenameEmitterToolTip", "Rename this local emitter copy"),
-						FSlateIcon(),
-						FUIAction(FExecuteAction::CreateSP(EmitterHandleViewModel, &FNiagaraEmitterHandleViewModel::SetIsRenamePending, true)));
-				}
-
-				MenuBuilder.AddMenuEntry(
-					LOCTEXT("RemoveSourceEmitter", "Remove Source Emitter"),
-					LOCTEXT("RemoveSourceEmitterToolTip", "Removes source information from this emitter, preventing inheritance of any further changes."),
-					FSlateIcon(),
-					FUIAction(
-						FExecuteAction::CreateSP(EmitterHandleViewModel->GetEmitterViewModel(), &FNiagaraEmitterViewModel::RemoveParentEmitter),
-						FCanExecuteAction::CreateSP(EmitterHandleViewModel->GetEmitterViewModel(), &FNiagaraEmitterViewModel::HasParentEmitter)));
-
-				MenuBuilder.AddMenuEntry(
-					LOCTEXT("ShowEmitterInContentBrowser", "Show in Content Browser"),
-					LOCTEXT("ShowEmitterInContentBrowserToolTip", "Show the emitter in this stack in the Content Browser"),
-					FSlateIcon(),
-					FUIAction(FExecuteAction::CreateSP(this, &SNiagaraStack::ShowEmitterInContentBrowser, TWeakPtr<FNiagaraEmitterHandleViewModel>(EmitterHandleViewModel))));
-			}
-			MenuBuilder.EndSection();
-		}
+		FNiagaraEditorUtilities::AddEmitterContextMenuActions(MenuBuilder, TopLevelViewModel->EmitterHandleViewModel);
 
 		MenuBuilder.BeginSection("StackActions", LOCTEXT("StackActions", "Stack Actions"));
 		{
