Added option to the stack view to filter the view to only display modules that have issues. Clicking the issue icon in the stack view header now cycles through all the issues in the view.

Right-aligned the issue icon in the stack view header.

diff --git a/Engine/Plugins/FX/Niagara/Source/NiagaraEditor/Private/NiagaraStackEditorData.cpp b/Engine/Plugins/FX/Niagara/Source/NiagaraEditor/Private/NiagaraStackEditorData.cpp
--- a/Engine/Plugins/FX/Niagara/Source/NiagaraEditor/Private/NiagaraStackEditorData.cpp
+++ b/Engine/Plugins/FX/Niagara/Source/NiagaraEditor/Private/NiagaraStackEditorData.cpp
@@ -85,6 +85,16 @@ void UNiagaraStackEditorData::SetShowLinkedInputs(bool bInShowLinkedInputs)
 	bShowLinkedInputs = bInShowLinkedInputs;
 }
 
+bool UNiagaraStackEditorData::GetShowOnlyIssues() const
+{
+	return bShowOnlyIssues;
+}
+
+void UNiagaraStackEditorData::SetShowOnlyIssues(bool bInShowOnlyIssues)
+{
+	bShowOnlyIssues = bInShowOnlyIssues;
+}
+
 double UNiagaraStackEditorData::GetLastScrollPosition() const
 {
 	return LastScrollPosition;
diff --git a/Engine/Plugins/FX/Niagara/Source/NiagaraEditor/Private/ViewModels/Stack/NiagaraStackItemGroup.cpp b/Engine/Plugins/FX/Niagara/Source/NiagaraEditor/Private/ViewModels/Stack/NiagaraStackItemGroup.cpp
--- a/Engine/Plugins/FX/Niagara/Source/NiagaraEditor/Private/ViewModels/Stack/NiagaraStackItemGroup.cpp
+++ b/Engine/Plugins/FX/Niagara/Source/NiagaraEditor/Private/ViewModels/Stack/NiagaraStackItemGroup.cpp
@@ -4,13 +4,13 @@
 #include "ViewModels/Stack/NiagaraStackItem.h"
 #include "ViewModels/Stack/NiagaraStackErrorItem.h"
 #include "ViewModels/Stack/NiagaraStackEntry.h"
+#include "ViewModels/Stack/NiagaraStackModuleItem.h"
 #include "NiagaraStackEditorData.h"
 #include "ViewModels/Stack/NiagaraStackGraphUtilities.h"
 #include "ViewModels/NiagaraSystemViewModel.h"
 #include "ViewModels/NiagaraEmitterHandleViewModel.h"
 #include "ViewModels/NiagaraEmitterViewModel.h"
 
-
 void UNiagaraStackItemGroup::Initialize(FRequiredEntryData InRequiredEntryData, FText InDisplayName, FText InToolTip, INiagaraStackItemGroupAddUtilities* InAddUtilities)
 {
 	Super::Initialize(InRequiredEntryData, InDisplayName.ToString());
@@ -18,6 +18,8 @@ void UNiagaraStackItemGroup::Initialize(FRequiredEntryData InRequiredEntryData,
 	GroupToolTip = InToolTip;
 	AddUtilities = InAddUtilities;
 	GroupFooter = nullptr;
+
+	AddChildFilter(FOnFilterChild::CreateUObject(this, &UNiagaraStackItemGroup::FilterChildrenWithIssues));
 }
 
 FText UNiagaraStackItemGroup::GetDisplayName() const
@@ -118,6 +120,22 @@ void UNiagaraStackItemGroup::ChlildStructureChangedInternal()
 	HighestIssueSeverity.Reset();
 }
 
+bool UNiagaraStackItemGroup::FilterChildrenWithIssues(const UNiagaraStackEntry& Child) const
+{
+	if (GetStackEditorData().GetShowOnlyIssues() == false)
+	{
+		return true;
+	}
+
+	const UNiagaraStackModuleItem* Module = Cast<UNiagaraStackModuleItem>(&Child);
+	if (Module != nullptr)
+	{
+		return Module->HasIssuesOrAnyChildHasIssues();
+	}
+
+	return false;
+}
+
 void UNiagaraStackItemGroupFooter::Initialize(FRequiredEntryData InRequiredEntryData)
 {
 	Super::Initialize(InRequiredEntryData, FString());
diff --git a/Engine/Plugins/FX/Niagara/Source/NiagaraEditor/Private/ViewModels/Stack/NiagaraStackViewModel.cpp b/Engine/Plugins/FX/Niagara/Source/NiagaraEditor/Private/ViewModels/Stack/NiagaraStackViewModel.cpp
--- a/Engine/Plugins/FX/Niagara/Source/NiagaraEditor/Private/ViewModels/Stack/NiagaraStackViewModel.cpp
+++ b/Engine/Plugins/FX/Niagara/Source/NiagaraEditor/Private/ViewModels/Stack/NiagaraStackViewModel.cpp
@@ -165,6 +165,7 @@ void UNiagaraStackViewModel::Reset()
 	
 	GEditor->UnregisterForUndo(this);
 
+	CurrentIssueCycleIndex = -1;
 	CurrentFocusedSearchMatchIndex = -1;
 	bRestartSearch = false;
 	bRefreshPending = false;
@@ -599,6 +600,32 @@ void UNiagaraStackViewModel::SetShowLinkedInputs(bool bInShowLinkedInputs)
 	RootEntry->RefreshChildren();
 }
 
+bool UNiagaraStackViewModel::GetShowOnlyIssues() const
+{
+	for (TSharedRef<FTopLevelViewModel> TopLevelViewModel : TopLevelViewModels)
+	{
+		if (TopLevelViewModel->IsValid() && TopLevelViewModel->GetStackEditorData()->GetShowOnlyIssues())
+		{
+			return true;
+		}
+	}
+	return false;
+}
+
+void UNiagaraStackViewModel::SetShowOnlyIssues(bool bInShowOnlyIssues)
+{
+	for (TSharedRef<FTopLevelViewModel> TopLevelViewModel : TopLevelViewModels)
+	{
+		if (TopLevelViewModel->IsValid())
+		{
+			TopLevelViewModel->GetStackEditorData()->SetShowOnlyIssues(bInShowOnlyIssues);
+		}
+	}
+
+	OnSearchTextChanged(CurrentSearchText);
+	RootEntry->RefreshChildren();
+}
+
 double UNiagaraStackViewModel::GetLastScrollPosition() const
 {
 	// TODO: Fix this with the new overview paradigm.
@@ -749,4 +776,34 @@ void UNiagaraStackViewModel::RefreshHasIssues()
 	}
 }
 
+UNiagaraStackEntry* UNiagaraStackViewModel::GetCurrentFocusedIssue() const
+{
+	if (CurrentIssueCycleIndex >= 0)
+	{
+		const TArray<UNiagaraStackEntry*>& Issues = RootEntry->GetAllChildrenWithIssues();
+		return Issues[CurrentIssueCycleIndex];
+	}
+	
+	return nullptr;
+}
+
+void UNiagaraStackViewModel::OnCycleThroughIssues()
+{
+	if (RootEntry == nullptr)
+	{
+		CurrentIssueCycleIndex = -1;
+		return;
+	}
+
+	const TArray<UNiagaraStackEntry*>& Issues = RootEntry->GetAllChildrenWithIssues();
+	if (Issues.Num() > 0)
+	{
+		++CurrentIssueCycleIndex;
+		if (CurrentIssueCycleIndex >= Issues.Num())
+		{
+			CurrentIssueCycleIndex = 0;
+		}
+	}
+}
+
 #undef LOCTEXT_NAMESPACE
diff --git a/Engine/Plugins/FX/Niagara/Source/NiagaraEditor/Public/NiagaraStackEditorData.h b/Engine/Plugins/FX/Niagara/Source/NiagaraEditor/Public/NiagaraStackEditorData.h
--- a/Engine/Plugins/FX/Niagara/Source/NiagaraEditor/Public/NiagaraStackEditorData.h
+++ b/Engine/Plugins/FX/Niagara/Source/NiagaraEditor/Public/NiagaraStackEditorData.h
@@ -61,7 +61,7 @@ public:
 	bool GetStackItemShowAdvanced(const FString& StackEntryKey, bool bShowAdvancedDefault) const;
 
 	/*
-	* Sets whether or not a stack entry is Expanded.
+	* Sets whether or not a stack entry is showing advanced items.
 	* @param StackEntryKey A unique key for the entry.
 	* @param bIsExpanded Whether or not the entry is expanded.
 	*/
@@ -85,6 +85,12 @@ public:
 	/* Sets whether or not item linked script inputs should be shown in the stack. */
 	void SetShowLinkedInputs(bool bInShowLinkedInputs);
 
+	/* Gets whether or not only modules that have issues should be shown in the stack. */
+	bool GetShowOnlyIssues() const;
+
+	/* Sets whether or not only modules that haves issues should be shown in the stack. */
+	void SetShowOnlyIssues(bool bInShowIssues);
+
 	/* Gets the last scroll position for the associated stack. */
 	double GetLastScrollPosition() const;
 
@@ -103,28 +109,31 @@ public:
 	NIAGARAEDITOR_API const TArray<FString>& GetDismissedStackIssueIds();
 
 private:
-	UPROPERTY()
+	UPROPERTY(Transient)
 	TMap<FString, bool> ModuleInputKeyToRenamePendingMap;
 
-	UPROPERTY()
+	UPROPERTY(Transient)
 	TMap<FString, bool> StackEntryKeyToExpandedMap;
 
-	UPROPERTY()
+	UPROPERTY(Transient)
 	TMap<FString, bool> StackEntryKeyToPreSearchExpandedMap;
 
-	UPROPERTY()
+	UPROPERTY(Transient)
 	TMap<FString, bool> StackItemKeyToShowAdvancedMap;
 
-	UPROPERTY()
+	UPROPERTY(Transient)
 	bool bShowAllAdvanced;
 
-	UPROPERTY()
+	UPROPERTY(Transient)
 	bool bShowOutputs;
 
-	UPROPERTY()
+	UPROPERTY(Transient)
 	bool bShowLinkedInputs;
 
-	UPROPERTY()
+	UPROPERTY(Transient)
+	bool bShowOnlyIssues;
+
+	UPROPERTY(Transient)
 	double LastScrollPosition;
 
 	UPROPERTY()
diff --git a/Engine/Plugins/FX/Niagara/Source/NiagaraEditor/Public/ViewModels/Stack/NiagaraStackItemGroup.h b/Engine/Plugins/FX/Niagara/Source/NiagaraEditor/Public/ViewModels/Stack/NiagaraStackItemGroup.h
--- a/Engine/Plugins/FX/Niagara/Source/NiagaraEditor/Public/ViewModels/Stack/NiagaraStackItemGroup.h
+++ b/Engine/Plugins/FX/Niagara/Source/NiagaraEditor/Public/ViewModels/Stack/NiagaraStackItemGroup.h
@@ -42,6 +42,9 @@ protected:
 
 	virtual void ChlildStructureChangedInternal() override;
 
+private:
+	bool FilterChildrenWithIssues(const UNiagaraStackEntry& Child) const;
+
 private:
 	UPROPERTY()
 	UNiagaraStackItemGroupFooter* GroupFooter;
diff --git a/Engine/Plugins/FX/Niagara/Source/NiagaraEditor/Public/ViewModels/Stack/NiagaraStackViewModel.h b/Engine/Plugins/FX/Niagara/Source/NiagaraEditor/Public/ViewModels/Stack/NiagaraStackViewModel.h
--- a/Engine/Plugins/FX/Niagara/Source/NiagaraEditor/Public/ViewModels/Stack/NiagaraStackViewModel.h
+++ b/Engine/Plugins/FX/Niagara/Source/NiagaraEditor/Public/ViewModels/Stack/NiagaraStackViewModel.h
@@ -97,9 +97,13 @@ public:
 
 	bool GetShowOutputs() const;
 	void SetShowOutputs(bool bInShowOutputs);
+
 	bool GetShowLinkedInputs() const;
 	void SetShowLinkedInputs(bool bInShowLinkedInputs);
 
+	bool GetShowOnlyIssues() const;
+	void SetShowOnlyIssues(bool bInShowOnlyIssues);
+
 	double GetLastScrollPosition() const;
 	void SetLastScrollPosition(double InLastScrollPosition);
 
@@ -119,6 +123,9 @@ public:
 	UNiagaraStackEntry* GetCurrentFocusedEntry();
 	void AddSearchScrollOffset(int NumberOfSteps);
 
+	void OnCycleThroughIssues();
+	UNiagaraStackEntry* GetCurrentFocusedIssue() const;
+
 	void GetPathForEntry(UNiagaraStackEntry* Entry, TArray<UNiagaraStackEntry*>& EntryPath) const;
 
 	/** Starts recursing through all entries to expand all groups and collapse all items. */
@@ -193,6 +200,7 @@ private:
 	bool bRestartSearch;
 	bool bRefreshPending;
 	bool bHasIssues;
+	int32 CurrentIssueCycleIndex;
 
 	bool bUsesTopLevelViewModels;
 	TArray<TSharedRef<FTopLevelViewModel>> TopLevelViewModels;
diff --git a/Engine/Plugins/FX/Niagara/Source/NiagaraEditorWidgets/Private/SNiagaraStack.cpp b/Engine/Plugins/FX/Niagara/Source/NiagaraEditorWidgets/Private/SNiagaraStack.cpp
--- a/Engine/Plugins/FX/Niagara/Source/NiagaraEditorWidgets/Private/SNiagaraStack.cpp
+++ b/Engine/Plugins/FX/Niagara/Source/NiagaraEditorWidgets/Private/SNiagaraStack.cpp
@@ -64,11 +64,14 @@ class SNiagaraStackEmitterHeader : public SCompoundWidget
 public:
 	SLATE_BEGIN_ARGS(SNiagaraStackEmitterHeader) {}
 		SLATE_ATTRIBUTE(EVisibility, IssueIconVisibility);
+		SLATE_EVENT(FSimpleDelegate, OnCycleThroughIssues);
 	SLATE_END_ARGS();
 
-	void Construct(const FArguments& InArgs, TSharedRef<FNiagaraEmitterHandleViewModel> InEmitterHandleViewModel, UNiagaraStackEntry* InRootEntry)
+	void Construct(const FArguments& InArgs, TSharedRef<FNiagaraEmitterHandleViewModel> InEmitterHandleViewModel, UNiagaraStackEntry* InRootEntry, UNiagaraStackViewModel* InStackViewModel)
 	{
 		EmitterHandleViewModel = InEmitterHandleViewModel;
+		OnCycleThroughIssues = InArgs._OnCycleThroughIssues;
+		StackViewModel = InStackViewModel;
 
 		ChildSlot
 		[
@@ -83,47 +86,38 @@ public:
 				//~ Enabled
 				+ SHorizontalBox::Slot()
 				.AutoWidth()
-				.Padding(2, 0, 0, 0)
+				.Padding(2, 4, 0, 0)
+				.VAlign(VAlign_Top)
 				[
 					SNew(SCheckBox)
 					.ToolTipText(LOCTEXT("EnabledToolTip", "Toggles whether this emitter is enabled. Disabled emitters don't simulate or render."))
 					.IsChecked(EmitterHandleViewModel.ToSharedRef(), &FNiagaraEmitterHandleViewModel::GetIsEnabledCheckState)
 					.OnCheckStateChanged(EmitterHandleViewModel.ToSharedRef(), &FNiagaraEmitterHandleViewModel::OnIsEnabledCheckStateChanged)
 				]
-				// Name and Source Emitter Name
 				+ SHorizontalBox::Slot()
 				.Padding(2)
+				.HAlign(HAlign_Fill)
+				.VAlign(VAlign_Center)
+				.FillWidth(1.0f)
 				[
+					// Name and Source Emitter Name
 					SNew(SWrapBox)
 					.Clipping(EWidgetClipping::ClipToBoundsAlways) 
 					.UseAllottedWidth(true)
 					+ SWrapBox::Slot()
-					.Padding(3, 0)
 					[
-						SNew(SHorizontalBox)
-						+ SHorizontalBox::Slot()
-						[
-							SAssignNew(EmitterNameTextBlock, SInlineEditableTextBlock)
-							.ToolTipText(this, &SNiagaraStackEmitterHeader::GetEmitterNameToolTip)
-							.Style(FNiagaraEditorStyle::Get(), "NiagaraEditor.HeadingInlineEditableText") 
-							.Clipping(EWidgetClipping::ClipToBoundsAlways)
-							.Text(EmitterHandleViewModel.ToSharedRef(), &FNiagaraEmitterHandleViewModel::GetNameText)
-							.OnTextCommitted(EmitterHandleViewModel.ToSharedRef(), &FNiagaraEmitterHandleViewModel::OnNameTextComitted)
-							.OnVerifyTextChanged(EmitterHandleViewModel.ToSharedRef(), &FNiagaraEmitterHandleViewModel::VerifyNameTextChanged)
-							.IsReadOnly(EmitterHandleViewModel->CanRenameEmitter() == false)
-						]
-						// Issue Icon
-						+ SHorizontalBox::Slot()
-						.AutoWidth()
-						.VAlign(VAlign_Center)
-						.Padding(4, 0, 0, 0)
-						[
-							SNew(SNiagaraStackIssueIcon, EmitterHandleViewModel->GetEmitterStackViewModel(), InRootEntry)
-							.Visibility(InArgs._IssueIconVisibility)
-						]
+				
+						SAssignNew(EmitterNameTextBlock, SInlineEditableTextBlock)
+						.ToolTipText(this, &SNiagaraStackEmitterHeader::GetEmitterNameToolTip)
+						.Style(FNiagaraEditorStyle::Get(), "NiagaraEditor.HeadingInlineEditableText") 
+						.Clipping(EWidgetClipping::ClipToBoundsAlways)
+						.Text(EmitterHandleViewModel.ToSharedRef(), &FNiagaraEmitterHandleViewModel::GetNameText)
+						.OnTextCommitted(EmitterHandleViewModel.ToSharedRef(), &FNiagaraEmitterHandleViewModel::OnNameTextComitted)
+						.OnVerifyTextChanged(EmitterHandleViewModel.ToSharedRef(), &FNiagaraEmitterHandleViewModel::VerifyNameTextChanged)
+						.IsReadOnly(EmitterHandleViewModel->CanRenameEmitter() == false)
 					]
 					+ SWrapBox::Slot()
-					.Padding(4, 0)
+					.Padding(4, 0, 0, 0)
 					[
 						SNew(STextBlock)
 						.ToolTipText(this, &SNiagaraStackEmitterHeader::GetEmitterNameToolTip)
@@ -132,13 +126,24 @@ public:
 						.Text(EmitterHandleViewModel->GetEmitterViewModel(), &FNiagaraEmitterViewModel::GetParentNameText)
 						.Visibility(this, &SNiagaraStackEmitterHeader::GetSourceEmitterNameVisibility) 
 					]
- 				]
+				]
+				// Issue Icon
+				+ SHorizontalBox::Slot()
+				.AutoWidth()
+				.HAlign(HAlign_Right)
+				.VAlign(VAlign_Top)
+				.Padding(0, 4, 2, 0)
+				[
+					SNew(SNiagaraStackIssueIcon, StackViewModel, InRootEntry)
+					.Visibility(InArgs._IssueIconVisibility)
+					.OnClicked(this, &SNiagaraStackEmitterHeader::OnIssueIconClicked)
+				]
 				// Open and Focus Source Emitter
 				+ SHorizontalBox::Slot()
+				.AutoWidth()
 				.HAlign(HAlign_Fill)
 				.VAlign(VAlign_Center)
-				.AutoWidth()
-				.Padding(2)
+				.Padding(0, 4, 2, 0)
 				[
 					SNew(SButton)
 					.IsFocusable(false)
@@ -223,16 +228,23 @@ private:
 		return EmitterHandleViewModel->GetEmitterViewModel()->GetEmitter()->GetParent() != nullptr ? EVisibility::Visible : EVisibility::Collapsed;
 	}
 
+	FReply OnIssueIconClicked() const
+	{
+		StackViewModel->OnCycleThroughIssues();
+		OnCycleThroughIssues.ExecuteIfBound();
+		return FReply::Handled();
+	}
+
 private:
 	TSharedPtr<FNiagaraEmitterHandleViewModel> EmitterHandleViewModel;
 
 	TSharedPtr<SInlineEditableTextBlock> EmitterNameTextBlock;
+	FSimpleDelegate OnCycleThroughIssues;
+	UNiagaraStackViewModel* StackViewModel;
 };
 
 const float SpacerHeight = 6;
 
-const FText SNiagaraStack::OccurencesFormat = NSLOCTEXT("NiagaraStack", "OccurencesFound", "{0} / {1}");
-
 void SNiagaraStack::Construct(const FArguments& InArgs, UNiagaraStackViewModel* InStackViewModel)
 {
 	StackViewModel = InStackViewModel;
@@ -263,15 +275,15 @@ void SNiagaraStack::Construct(const FArguments& InArgs, UNiagaraStackViewModel*
 			.BorderImage(FEditorStyle::GetBrush("WhiteBrush"))
 			.BorderBackgroundColor(FNiagaraEditorWidgetsStyle::Get().GetColor("NiagaraEditor.Stack.BackgroundColor"))
 			[
-			SAssignNew(StackTree, STreeView<UNiagaraStackEntry*>)
-			.OnGenerateRow(this, &SNiagaraStack::OnGenerateRowForStackItem)
-			.OnGetChildren(this, &SNiagaraStack::OnGetChildren)
-			.TreeItemsSource(&StackViewModel->GetRootEntryAsArray())
-			.OnTreeViewScrolled(this, &SNiagaraStack::StackTreeScrolled)
-			.SelectionMode(ESelectionMode::Single)
-			.OnItemToString_Debug_Static(&FNiagaraStackEditorWidgetsUtilities::StackEntryToStringForListDebug)
-			.OnMouseButtonClick(this, &SNiagaraStack::OnStackItemClicked)
-		]
+				SAssignNew(StackTree, STreeView<UNiagaraStackEntry*>)
+				.OnGenerateRow(this, &SNiagaraStack::OnGenerateRowForStackItem)
+				.OnGetChildren(this, &SNiagaraStack::OnGetChildren)
+				.TreeItemsSource(&StackViewModel->GetRootEntryAsArray())
+				.OnTreeViewScrolled(this, &SNiagaraStack::StackTreeScrolled)
+				.SelectionMode(ESelectionMode::Single)
+				.OnItemToString_Debug_Static(&FNiagaraStackEditorWidgetsUtilities::StackEntryToStringForListDebug)
+				.OnMouseButtonClick(this, &SNiagaraStack::OnStackItemClicked)
+			]
 		]
 	];
 
@@ -366,33 +378,29 @@ void SNiagaraStack::OnSearchTextChanged(const FText& SearchText)
 	StackViewModel->OnSearchTextChanged(SearchText);
 }
 
+static void ExpandAllInPath(const TArray<UNiagaraStackEntry*>& EntryPath)
+{
+	for (UNiagaraStackEntry* Entry : EntryPath)
+	{
+		if (!Entry->IsA<UNiagaraStackRoot>())
+		{
+			Entry->SetIsExpanded(true);
+		}
+	}
+}
+
 FReply SNiagaraStack::ScrollToNextMatch()
 {
-	
-	const int NextMatchIndex = StackViewModel->GetCurrentFocusedMatchIndex() + 1;
+	int NextMatchIndex = StackViewModel->GetCurrentFocusedMatchIndex() + 1;
 	TArray<UNiagaraStackViewModel::FSearchResult> CurrentSearchResults = StackViewModel->GetCurrentSearchResults();
 	if (CurrentSearchResults.Num() != 0)
 	{
-		if (NextMatchIndex < CurrentSearchResults.Num())
+		if (NextMatchIndex >= CurrentSearchResults.Num())
 		{
-			for (auto SearchResultEntry : CurrentSearchResults[NextMatchIndex].EntryPath)
-			{
-				if (!SearchResultEntry->IsA<UNiagaraStackRoot>())
-				{
-					SearchResultEntry->SetIsExpanded(true);
-				}
-			}
-		}
-		else
-		{
-			for (auto SearchResultEntry : CurrentSearchResults[0].EntryPath)
-			{
-				if (!SearchResultEntry->IsA<UNiagaraStackRoot>())
-				{
-					SearchResultEntry->SetIsExpanded(true);
-				}
-			}
+			NextMatchIndex = 0;
 		}
+
+		ExpandAllInPath(CurrentSearchResults[NextMatchIndex].EntryPath);
 		SynchronizeTreeExpansion();
 	}
 
@@ -539,6 +547,16 @@ TSharedRef<SWidget> SNiagaraStack::GetViewOptionsMenu() const
 			FGetActionCheckState::CreateLambda([=]() { return StackViewModel->GetShowOutputs() ? ECheckBoxState::Checked : ECheckBoxState::Unchecked; })),
 		NAME_None, EUserInterfaceActionType::Check);
 
+	MenuBuilder.AddMenuEntry(
+		LOCTEXT("ShowIssuesLabel", "Show Only Issues"),
+		LOCTEXT("ShowIssuesToolTip", "Hides all modules except those that have unresolved issues."),
+		FSlateIcon(),
+		FUIAction(
+			FExecuteAction::CreateLambda([=]() { StackViewModel->SetShowOnlyIssues(!StackViewModel->GetShowOnlyIssues()); }),
+			FCanExecuteAction(),
+			FGetActionCheckState::CreateLambda([=]() { return StackViewModel->GetShowOnlyIssues() ? ECheckBoxState::Checked : ECheckBoxState::Unchecked; })),
+		NAME_None, EUserInterfaceActionType::Check);
+
 	return MenuBuilder.MakeWidget();
 }
 
@@ -657,8 +675,9 @@ TSharedRef<ITableRow> SNiagaraStack::OnGenerateRowForTopLevelObject(TSharedRef<U
 	}
 	else if(Item->EmitterHandleViewModel.IsValid())
 	{
-		Content = SNew(SNiagaraStackEmitterHeader, Item->EmitterHandleViewModel.ToSharedRef(), Item->RootEntry.Get())
-			.IssueIconVisibility(this, &SNiagaraStack::GetIssueIconVisibility);
+		Content = SNew(SNiagaraStackEmitterHeader, Item->EmitterHandleViewModel.ToSharedRef(), Item->RootEntry.Get(), StackViewModel)
+			.IssueIconVisibility(this, &SNiagaraStack::GetIssueIconVisibility)
+			.OnCycleThroughIssues(this, &SNiagaraStack::OnCycleThroughIssues);
 	}
 
 	Content->SetOnMouseButtonUp(FPointerEventHandler::CreateSP(this, &SNiagaraStack::OnTopLevelRowMouseButtonDown, TWeakPtr<UNiagaraStackViewModel::FTopLevelViewModel>(Item)));
@@ -1019,4 +1038,18 @@ void SNiagaraStack::OnStackItemClicked(UNiagaraStackEntry* Item)
 	FSlateApplication::Get().SetKeyboardFocus(SearchBox, EFocusCause::OtherWidgetLostFocus);
 }
 
+void SNiagaraStack::OnCycleThroughIssues()
+{
+	UNiagaraStackEntry* StackEntry = StackViewModel->GetCurrentFocusedIssue();
+
+	TArray<UNiagaraStackEntry*> EntryPath;
+	StackViewModel->GetPathForEntry(StackEntry, EntryPath);
+	EntryPath.Add(StackEntry);
+
+	ExpandAllInPath(EntryPath);
+	SynchronizeTreeExpansion();
+
+	StackTree->RequestScrollIntoView(StackEntry);
+}
+
 #undef LOCTEXT_NAMESPACE
diff --git a/Engine/Plugins/FX/Niagara/Source/NiagaraEditorWidgets/Private/SNiagaraStack.h b/Engine/Plugins/FX/Niagara/Source/NiagaraEditorWidgets/Private/SNiagaraStack.h
--- a/Engine/Plugins/FX/Niagara/Source/NiagaraEditorWidgets/Private/SNiagaraStack.h
+++ b/Engine/Plugins/FX/Niagara/Source/NiagaraEditorWidgets/Private/SNiagaraStack.h
@@ -107,6 +107,8 @@ private:
 	EVisibility GetIssueIconVisibility() const;
 
 	void OnStackItemClicked(UNiagaraStackEntry* Item);
+	void OnCycleThroughIssues();
+
 private:
 	UNiagaraStackViewModel* StackViewModel;
 
@@ -120,6 +122,5 @@ private:
 
 	// ~ search stuff
 	TSharedPtr<SSearchBox> SearchBox;
-	static const FText OccurencesFormat;
 	bool bNeedsJumpToNextOccurence;
 };
diff --git a/Engine/Plugins/FX/Niagara/Source/NiagaraEditorWidgets/Private/Stack/SNiagaraStackIssueIcon.cpp b/Engine/Plugins/FX/Niagara/Source/NiagaraEditorWidgets/Private/Stack/SNiagaraStackIssueIcon.cpp
--- a/Engine/Plugins/FX/Niagara/Source/NiagaraEditorWidgets/Private/Stack/SNiagaraStackIssueIcon.cpp
+++ b/Engine/Plugins/FX/Niagara/Source/NiagaraEditorWidgets/Private/Stack/SNiagaraStackIssueIcon.cpp
@@ -5,9 +5,11 @@
 #include "ViewModels/Stack/NiagaraStackEntry.h"
 
 #include "Widgets/Images/SImage.h"
+#include "Widgets/Input/SButton.h"
 #include "Widgets/Layout/SBox.h"
 #include "EditorStyleSet.h"
 #include "Internationalization/Text.h"
+#include "NiagaraEditorWidgetsStyle.h"
 
 #define LOCTEXT_NAMESPACE "SNiagaraStackIssueIcon"
 
@@ -20,8 +22,7 @@ void SNiagaraStackIssueIcon::Construct(const FArguments& InArgs, UNiagaraStackVi
 		StackEntry->OnStructureChanged().AddSP(this, &SNiagaraStackIssueIcon::UpdateFromEntry);
 	}
 
-	ChildSlot
-	[
+	TSharedRef<SWidget> IconWidget =
 		SNew(SBox)
 		.IsEnabled(this, &SNiagaraStackIssueIcon::GetIconIsEnabled)
 		.ToolTipText(this, &SNiagaraStackIssueIcon::GetIconToolTip)
@@ -32,7 +33,27 @@ void SNiagaraStackIssueIcon::Construct(const FArguments& InArgs, UNiagaraStackVi
 		[
 			SNew(SImage)
 			.Image(this, &SNiagaraStackIssueIcon::GetIconBrush)
-		]
+		];
+
+	if (InArgs._OnClicked.IsBound())
+	{
+		TSharedRef<SButton> IconButton =
+			SNew(SButton)
+			.ButtonStyle(FEditorStyle::Get(), "RoundButton")
+			.ForegroundColor(FEditorStyle::GetSlateColor("DefaultForeground"))
+			.ContentPadding(FMargin(1,0,0,0))
+			.OnClicked(InArgs._OnClicked)
+			.Content()
+			[
+				IconWidget
+			];
+
+		IconWidget = IconButton;
+	}
+
+	ChildSlot
+	[
+		IconWidget
 	];
 	UpdateFromEntry();
 }
diff --git a/Engine/Plugins/FX/Niagara/Source/NiagaraEditorWidgets/Private/Stack/SNiagaraStackIssueIcon.h b/Engine/Plugins/FX/Niagara/Source/NiagaraEditorWidgets/Private/Stack/SNiagaraStackIssueIcon.h
--- a/Engine/Plugins/FX/Niagara/Source/NiagaraEditorWidgets/Private/Stack/SNiagaraStackIssueIcon.h
+++ b/Engine/Plugins/FX/Niagara/Source/NiagaraEditorWidgets/Private/Stack/SNiagaraStackIssueIcon.h
@@ -4,6 +4,7 @@
 
 #include "Widgets/SCompoundWidget.h"
 #include "Widgets/DeclarativeSyntaxSupport.h"
+#include "Framework/SlateDelegates.h"
 
 class UNiagaraStackViewModel;
 class UNiagaraStackEntry;
@@ -12,6 +13,7 @@ class SNiagaraStackIssueIcon : public SCompoundWidget
 {
 public:
 	SLATE_BEGIN_ARGS(SNiagaraStackIssueIcon) {}
+		SLATE_EVENT(FOnClicked, OnClicked)
 	SLATE_END_ARGS()
 
 	void Construct(const FArguments& InArgs, UNiagaraStackViewModel* InStackViewModel, UNiagaraStackEntry* InStackEntry);
