Opening a Niagara emitter from the content browser no longer causes a crash due to scalability settings being out of date.

diff --git a/Engine/Plugins/FX/Niagara/Source/Niagara/Private/NiagaraSystem.cpp b/Engine/Plugins/FX/Niagara/Source/Niagara/Private/NiagaraSystem.cpp
--- a/Engine/Plugins/FX/Niagara/Source/Niagara/Private/NiagaraSystem.cpp
+++ b/Engine/Plugins/FX/Niagara/Source/Niagara/Private/NiagaraSystem.cpp
@@ -122,6 +122,8 @@ void UNiagaraSystem::PostInitProperties()
 		EditorData = NiagaraModule.GetEditorOnlyDataUtilities().CreateDefaultEditorData(this);
 #endif
 	}
+
+	ResolveScalabilityOverrides();
 }
 
 bool UNiagaraSystem::IsLooping() const
