Aligned and unified the styles of the issues icon and open parent buttons in the Niagara stack panel header.

#rnx

diff --git a/Engine/Plugins/FX/Niagara/Source/NiagaraEditorWidgets/Private/SNiagaraStack.cpp b/Engine/Plugins/FX/Niagara/Source/NiagaraEditorWidgets/Private/SNiagaraStack.cpp
--- a/Engine/Plugins/FX/Niagara/Source/NiagaraEditorWidgets/Private/SNiagaraStack.cpp
+++ b/Engine/Plugins/FX/Niagara/Source/NiagaraEditorWidgets/Private/SNiagaraStack.cpp
@@ -131,8 +131,8 @@ public:
 				+ SHorizontalBox::Slot()
 				.AutoWidth()
 				.HAlign(HAlign_Right)
-				.VAlign(VAlign_Top)
-				.Padding(0, 4, 2, 0)
+				.VAlign(VAlign_Center)
+				.Padding(0, 0, 2, 0)
 				[
 					SNew(SNiagaraStackIssueIcon, StackViewModel, InRootEntry)
 					.Visibility(InArgs._IssueIconVisibility)
@@ -143,7 +143,7 @@ public:
 				.AutoWidth()
 				.HAlign(HAlign_Fill)
 				.VAlign(VAlign_Center)
-				.Padding(0, 4, 2, 0)
+				.Padding(0, 0, 2, 0)
 				[
 					SNew(SButton)
 					.IsFocusable(false)
diff --git a/Engine/Plugins/FX/Niagara/Source/NiagaraEditorWidgets/Private/Stack/SNiagaraStackIssueIcon.cpp b/Engine/Plugins/FX/Niagara/Source/NiagaraEditorWidgets/Private/Stack/SNiagaraStackIssueIcon.cpp
--- a/Engine/Plugins/FX/Niagara/Source/NiagaraEditorWidgets/Private/Stack/SNiagaraStackIssueIcon.cpp
+++ b/Engine/Plugins/FX/Niagara/Source/NiagaraEditorWidgets/Private/Stack/SNiagaraStackIssueIcon.cpp
@@ -39,9 +39,9 @@ void SNiagaraStackIssueIcon::Construct(const FArguments& InArgs, UNiagaraStackVi
 	{
 		TSharedRef<SButton> IconButton =
 			SNew(SButton)
-			.ButtonStyle(FEditorStyle::Get(), "RoundButton")
-			.ForegroundColor(FEditorStyle::GetSlateColor("DefaultForeground"))
-			.ContentPadding(FMargin(1,0,0,0))
+			.ButtonStyle(FEditorStyle::Get(), "HoverHintOnly")
+			.ForegroundColor(FNiagaraEditorWidgetsStyle::Get().GetColor("NiagaraEditor.Stack.ForegroundColor"))
+			.ContentPadding(FMargin(1, 0, 0, 0))
 			.OnClicked(InArgs._OnClicked)
 			.Content()
 			[
