Migrate stuff over to RHI helper functions to ensure we properly handed unbounded parameters

diff --git a/Engine/Plugins/Experimental/ChaosNiagara/Source/ChaosNiagara/Private/NiagaraDataInterfaceChaosDestruction.cpp b/Engine/Plugins/Experimental/ChaosNiagara/Source/ChaosNiagara/Private/NiagaraDataInterfaceChaosDestruction.cpp
--- a/Engine/Plugins/Experimental/ChaosNiagara/Source/ChaosNiagara/Private/NiagaraDataInterfaceChaosDestruction.cpp
+++ b/Engine/Plugins/Experimental/ChaosNiagara/Source/ChaosNiagara/Private/NiagaraDataInterfaceChaosDestruction.cpp
@@ -3858,17 +3858,6 @@ void LoadGPUBufferFromArray(FDynamicReadBuffer& Buffer,
 	Buffer.Unlock();
 }
 
-static void SetBuffer(FRHICommandList& CmdList,
-	const FShaderResourceParameter& Param,
-	FRHIComputeShader* Shader,
-	FDynamicReadBuffer& Buffer)
-{
-	// Skip unbound parameters, since we won't be reading any of them
-	if (!Param.IsBound()) return;
-
-	CmdList.SetShaderResourceViewParameter(Shader, Param.GetBaseIndex(), Buffer.SRV);
-}
-
 template <typename T>
 static void SetBuffer(FRHICommandList& CmdList,
 	const FShaderResourceParameter& Param,
@@ -3928,7 +3917,7 @@ static void SetBuffer(FRHICommandList& CmdList,
 
 	Buffer.Unlock();
 
-	CmdList.SetShaderResourceViewParameter(Shader, Param.GetBaseIndex(), Buffer.SRV);
+	SetSRVParameter(CmdList, Shader, Param, Buffer.SRV);
 }
 
 void UNiagaraDataInterfaceChaosDestruction::PushToRenderThread()
diff --git a/Engine/Plugins/FX/HoudiniNiagara/Source/HoudiniNiagara/Private/NiagaraDataInterfaceHoudiniCSV.cpp b/Engine/Plugins/FX/HoudiniNiagara/Source/HoudiniNiagara/Private/NiagaraDataInterfaceHoudiniCSV.cpp
--- a/Engine/Plugins/FX/HoudiniNiagara/Source/HoudiniNiagara/Private/NiagaraDataInterfaceHoudiniCSV.cpp
+++ b/Engine/Plugins/FX/HoudiniNiagara/Source/HoudiniNiagara/Private/NiagaraDataInterfaceHoudiniCSV.cpp
@@ -2269,23 +2269,15 @@ struct FNiagaraDataInterfaceParametersCS_HoudiniCSV : public FNiagaraDataInterfa
 		SetShaderValue(RHICmdList, ComputeShaderRHI, NumberOfColumns, HoudiniDI->NumColumns);
 		SetShaderValue(RHICmdList, ComputeShaderRHI, NumberOfPoints, HoudiniDI->NumColumns);
 
-		FRWBuffer& FloatRWBuffer = HoudiniDI->FloatValuesGPUBuffer;
- 		RHICmdList.SetShaderResourceViewParameter( ComputeShaderRHI, FloatValuesBuffer.GetBaseIndex(), FloatRWBuffer.SRV );
-
-		FRWBuffer& AttributeColumnRWBuffer = HoudiniDI->SpecialAttributesColumnIndexesGPUBuffer;
-		RHICmdList.SetShaderResourceViewParameter(ComputeShaderRHI, SpecialAttributesColumnIndexesBuffer.GetBaseIndex(), AttributeColumnRWBuffer.SRV);
-
-		FRWBuffer& SpawnRWBuffer = HoudiniDI->SpawnTimesGPUBuffer;
-		RHICmdList.SetShaderResourceViewParameter(ComputeShaderRHI, SpawnTimesBuffer.GetBaseIndex(), SpawnRWBuffer.SRV);
-		FRWBuffer& LifeRWBuffer = HoudiniDI->LifeValuesGPUBuffer;
-		RHICmdList.SetShaderResourceViewParameter(ComputeShaderRHI, LifeValuesBuffer.GetBaseIndex(), LifeRWBuffer.SRV);
-		FRWBuffer& TypesRWBuffer = HoudiniDI->PointTypesGPUBuffer;
-		RHICmdList.SetShaderResourceViewParameter(ComputeShaderRHI, PointTypesBuffer.GetBaseIndex(), TypesRWBuffer.SRV);
+		SetSRVParameter(RHICmdList, ComputeShaderRHI, FloatValuesBuffer, HoudiniDI->FloatValuesGPUBuffer.SRV);
+		SetSRVParameter(RHICmdList, ComputeShaderRHI, SpecialAttributesColumnIndexesBuffer, HoudiniDI->SpecialAttributesColumnIndexesGPUBuffer.SRV);
+		SetSRVParameter(RHICmdList, ComputeShaderRHI, SpawnTimesBuffer, HoudiniDI->SpawnTimesGPUBuffer.SRV);
+		SetSRVParameter(RHICmdList, ComputeShaderRHI, LifeValuesBuffer, HoudiniDI->LifeValuesGPUBuffer.SRV);
+		SetSRVParameter(RHICmdList, ComputeShaderRHI, PointTypesBuffer, HoudiniDI->PointTypesGPUBuffer.SRV);
 
 		SetShaderValue(RHICmdList, ComputeShaderRHI, MaxNumberOfIndexesPerPoint, HoudiniDI->MaxNumberOfIndexesPerPoint);
 
-		FRWBuffer& PointValuesIndexesRWBuffer = HoudiniDI->PointValueIndexesGPUBuffer;
-		RHICmdList.SetShaderResourceViewParameter(ComputeShaderRHI, PointValueIndexesBuffer.GetBaseIndex(), PointValuesIndexesRWBuffer.SRV);
+		SetSRVParameter(RHICmdList, ComputeShaderRHI, PointValueIndexesBuffer, HoudiniDI->PointValueIndexesGPUBuffer.SRV);
 
 		SetShaderValue(RHICmdList, ComputeShaderRHI, LastSpawnedPointId, -1);
 		SetShaderValue(RHICmdList, ComputeShaderRHI, LastSpawnTime, -1);
diff --git a/Engine/Plugins/FX/Niagara/Source/Niagara/Private/NiagaraDataInterfaceCurveBase.cpp b/Engine/Plugins/FX/Niagara/Source/Niagara/Private/NiagaraDataInterfaceCurveBase.cpp
--- a/Engine/Plugins/FX/Niagara/Source/Niagara/Private/NiagaraDataInterfaceCurveBase.cpp
+++ b/Engine/Plugins/FX/Niagara/Source/Niagara/Private/NiagaraDataInterfaceCurveBase.cpp
@@ -340,7 +340,7 @@ struct FNiagaraDataInterfaceParametersCS_Curve : public FNiagaraDataInterfacePar
 		SetShaderValue(RHICmdList, ComputeShaderRHI, MaxTime, CurveDI->LUTMaxTime);
 		SetShaderValue(RHICmdList, ComputeShaderRHI, InvTimeRange, CurveDI->LUTInvTimeRange);
 		SetShaderValue(RHICmdList, ComputeShaderRHI, CurveLUTNumMinusOne, CurveDI->CurveLUTNumMinusOne);
-		RHICmdList.SetShaderResourceViewParameter(ComputeShaderRHI, CurveLUT.GetBaseIndex(), CurveLUTBuffer.SRV);
+		SetSRVParameter(RHICmdList, ComputeShaderRHI, CurveLUT, CurveLUTBuffer.SRV);
 	}
 
 	FShaderParameter MinTime;
diff --git a/Engine/Plugins/FX/Niagara/Source/Niagara/Private/NiagaraDataInterfaceGrid2DCollection.cpp b/Engine/Plugins/FX/Niagara/Source/Niagara/Private/NiagaraDataInterfaceGrid2DCollection.cpp
--- a/Engine/Plugins/FX/Niagara/Source/Niagara/Private/NiagaraDataInterfaceGrid2DCollection.cpp
+++ b/Engine/Plugins/FX/Niagara/Source/Niagara/Private/NiagaraDataInterfaceGrid2DCollection.cpp
@@ -87,13 +87,14 @@ struct FNiagaraDataInterfaceParametersCS_Grid2DCollection : public FNiagaraDataI
 		if (GridParam.IsBound() && ProxyData->GetCurrentData() != NULL)
 		{
 			RHICmdList.TransitionResource(EResourceTransitionAccess::EReadable, EResourceTransitionPipeline::EComputeToCompute, ProxyData->GetCurrentData()->GridBuffer.UAV);
-			RHICmdList.SetShaderResourceViewParameter(Context.Shader->GetComputeShader(), GridParam.GetBaseIndex(), ProxyData->GetCurrentData()->GridBuffer.SRV);
+			SetSRVParameter(RHICmdList, Context.Shader->GetComputeShader(), GridParam, ProxyData->GetCurrentData()->GridBuffer.SRV);
 		}
 
 		if (Context.IsOutputStage && OutputGridParam.IsBound() && ProxyData->GetDestinationData() != NULL)
 		{
+			const FTextureRWBuffer& TextureBuffer = ProxyData->GetDestinationData()->GridBuffer;
 			RHICmdList.TransitionResource(EResourceTransitionAccess::EWritable, EResourceTransitionPipeline::EComputeToCompute, ProxyData->GetDestinationData()->GridBuffer.UAV);
-			RHICmdList.SetUAVParameter(Context.Shader->GetComputeShader(), OutputGridParam.GetUAVIndex(), ProxyData->GetDestinationData()->GridBuffer.UAV);		
+			OutputGridParam.SetTexture(RHICmdList, Context.Shader->GetComputeShader(), TextureBuffer.Buffer, TextureBuffer.UAV);
 		}
 	}
 
diff --git a/Engine/Plugins/FX/Niagara/Source/Niagara/Private/NiagaraDataInterfaceNeighborGrid3D.cpp b/Engine/Plugins/FX/Niagara/Source/Niagara/Private/NiagaraDataInterfaceNeighborGrid3D.cpp
--- a/Engine/Plugins/FX/Niagara/Source/Niagara/Private/NiagaraDataInterfaceNeighborGrid3D.cpp
+++ b/Engine/Plugins/FX/Niagara/Source/Niagara/Private/NiagaraDataInterfaceNeighborGrid3D.cpp
@@ -85,13 +85,13 @@ struct FNiagaraDataInterfaceParametersCS_NeighborGrid3D : public FNiagaraDataInt
 			if (ParticleNeighborsGridParam.IsBound())
 			{
 				RHICmdList.TransitionResource(EResourceTransitionAccess::EReadable, EResourceTransitionPipeline::EComputeToCompute, ProxyData->NeighborhoodBuffer.UAV);
-				RHICmdList.SetShaderResourceViewParameter(Context.Shader->GetComputeShader(), ParticleNeighborsGridParam.GetBaseIndex(), ProxyData->NeighborhoodBuffer.SRV);
+				SetSRVParameter(RHICmdList, Context.Shader->GetComputeShader(), ParticleNeighborsGridParam, ProxyData->NeighborhoodBuffer.SRV);
 			}
 
 			if (ParticleNeighborCountGridParam.IsBound())
 			{
 				RHICmdList.TransitionResource(EResourceTransitionAccess::EReadable, EResourceTransitionPipeline::EComputeToCompute, ProxyData->NeighborhoodCountBuffer.UAV);
-				RHICmdList.SetShaderResourceViewParameter(Context.Shader->GetComputeShader(), ParticleNeighborCountGridParam.GetBaseIndex(), ProxyData->NeighborhoodCountBuffer.SRV);
+				SetSRVParameter(RHICmdList, Context.Shader->GetComputeShader(), ParticleNeighborCountGridParam, ProxyData->NeighborhoodCountBuffer.SRV);
 			}
 		}
 		else
@@ -99,13 +99,13 @@ struct FNiagaraDataInterfaceParametersCS_NeighborGrid3D : public FNiagaraDataInt
 			if (OutputParticleNeighborsGridParam.IsBound())
 			{
 				RHICmdList.TransitionResource(EResourceTransitionAccess::EWritable, EResourceTransitionPipeline::EComputeToCompute, ProxyData->NeighborhoodBuffer.UAV);
-				RHICmdList.SetUAVParameter(Context.Shader->GetComputeShader(), OutputParticleNeighborsGridParam.GetUAVIndex(), ProxyData->NeighborhoodBuffer.UAV);
+				OutputParticleNeighborsGridParam.SetBuffer(RHICmdList, Context.Shader->GetComputeShader(), ProxyData->NeighborhoodBuffer);
 			}
 
 			if (OutputParticleNeighborCountGridParam.IsBound())
 			{
 				RHICmdList.TransitionResource(EResourceTransitionAccess::EWritable, EResourceTransitionPipeline::EComputeToCompute, ProxyData->NeighborhoodCountBuffer.UAV);
-				RHICmdList.SetUAVParameter(Context.Shader->GetComputeShader(), OutputParticleNeighborCountGridParam.GetUAVIndex(), ProxyData->NeighborhoodCountBuffer.UAV);
+				OutputParticleNeighborCountGridParam.SetBuffer(RHICmdList, Context.Shader->GetComputeShader(), ProxyData->NeighborhoodCountBuffer);
 			}
 		}
 		// Note: There is a flush in PreEditChange to make sure everything is synced up at this point 
diff --git a/Engine/Plugins/FX/Niagara/Source/Niagara/Private/NiagaraDataSet.cpp b/Engine/Plugins/FX/Niagara/Source/Niagara/Private/NiagaraDataSet.cpp
--- a/Engine/Plugins/FX/Niagara/Source/Niagara/Private/NiagaraDataSet.cpp
+++ b/Engine/Plugins/FX/Niagara/Source/Niagara/Private/NiagaraDataSet.cpp
@@ -893,55 +893,23 @@ void FNiagaraDataBuffer::Dump(int32 StartIndex, int32 InNumInstances, const FStr
 void FNiagaraDataBuffer::SetShaderParams(FNiagaraShader *Shader, FRHICommandList &CommandList, bool bInput)
 {
 	check(IsInRenderingThread());
+
+	const uint32 SafeBufferSize = GetFloatStride() / sizeof(float);
+	FRHIComputeShader* ComputeShader = Shader->GetComputeShader();
+
 	if (bInput)
 	{
-		if (Shader->FloatInputBufferParam.IsBound())
-		{
-			if (GetNumInstancesAllocated() > 0)
-			{
-				CommandList.SetShaderResourceViewParameter(Shader->GetComputeShader(), Shader->FloatInputBufferParam.GetBaseIndex(), GetGPUBufferFloat().SRV);
-			}
-			else
-			{
-				CommandList.SetShaderResourceViewParameter(Shader->GetComputeShader(), Shader->FloatInputBufferParam.GetBaseIndex(), FNiagaraRenderer::GetDummyFloatBuffer().SRV);
-			}
-		}
+		const bool InstancesAllocated = GetNumInstancesAllocated() > 0;
 
-		if (Shader->IntInputBufferParam.IsBound())
-		{
-			if (GetNumInstancesAllocated() > 0)
-			{
-				CommandList.SetShaderResourceViewParameter(Shader->GetComputeShader(), Shader->IntInputBufferParam.GetBaseIndex(), GetGPUBufferInt().SRV);
-			}
-			else
-			{
-				CommandList.SetShaderResourceViewParameter(Shader->GetComputeShader(), Shader->IntInputBufferParam.GetBaseIndex(), FNiagaraRenderer::GetDummyIntBuffer().SRV);
-			}
-		}
-
-		if (Shader->ComponentBufferSizeReadParam.IsBound())
-		{
-			uint32 SafeBufferSize = GetFloatStride() / sizeof(float);
-			CommandList.SetShaderParameter(Shader->GetComputeShader(), Shader->ComponentBufferSizeReadParam.GetBufferIndex(), Shader->ComponentBufferSizeReadParam.GetBaseIndex(), Shader->ComponentBufferSizeReadParam.GetNumBytes(), &SafeBufferSize);
-		}
+		SetSRVParameter(CommandList, ComputeShader, Shader->FloatInputBufferParam, InstancesAllocated ? GetGPUBufferFloat().SRV : FNiagaraRenderer::GetDummyFloatBuffer().SRV);
+		SetSRVParameter(CommandList, ComputeShader, Shader->IntInputBufferParam, InstancesAllocated ? GetGPUBufferInt().SRV : FNiagaraRenderer::GetDummyIntBuffer().SRV);
+		SetShaderValue(CommandList, ComputeShader, Shader->ComponentBufferSizeReadParam, SafeBufferSize);
 	}
 	else
 	{
-		if (Shader->FloatOutputBufferParam.IsUAVBound())
-		{
-			CommandList.SetUAVParameter(Shader->GetComputeShader(), Shader->FloatOutputBufferParam.GetUAVIndex(), GetGPUBufferFloat().UAV);
-		}
-
-		if (Shader->IntOutputBufferParam.IsUAVBound())
-		{
-			CommandList.SetUAVParameter(Shader->GetComputeShader(), Shader->IntOutputBufferParam.GetUAVIndex(), GetGPUBufferInt().UAV);
-		}
-
-		if (Shader->ComponentBufferSizeWriteParam.IsBound())
-		{
-			uint32 SafeBufferSize = GetFloatStride() / sizeof(float);
-			CommandList.SetShaderParameter(Shader->GetComputeShader(), Shader->ComponentBufferSizeWriteParam.GetBufferIndex(), Shader->ComponentBufferSizeWriteParam.GetBaseIndex(), Shader->ComponentBufferSizeWriteParam.GetNumBytes(), &SafeBufferSize);
-		}
+		Shader->FloatOutputBufferParam.SetBuffer(CommandList, ComputeShader, GetGPUBufferFloat());
+		Shader->IntOutputBufferParam.SetBuffer(CommandList, ComputeShader, GetGPUBufferInt());
+		SetShaderValue(CommandList, ComputeShader, Shader->ComponentBufferSizeWriteParam, SafeBufferSize);
 	}
 }
 
diff --git a/Engine/Plugins/FX/Niagara/Source/NiagaraVertexFactories/Private/NiagaraDrawIndirect.cpp b/Engine/Plugins/FX/Niagara/Source/NiagaraVertexFactories/Private/NiagaraDrawIndirect.cpp
--- a/Engine/Plugins/FX/Niagara/Source/NiagaraVertexFactories/Private/NiagaraDrawIndirect.cpp
+++ b/Engine/Plugins/FX/Niagara/Source/NiagaraVertexFactories/Private/NiagaraDrawIndirect.cpp
@@ -6,7 +6,7 @@ NiagaraDrawIndirect.cpp : Niagara shader to generate the draw indirect args for
 
 #include "NiagaraDrawIndirect.h"
 #include "NiagaraGPUSortInfo.h"
-
+#include "ShaderParameterUtils.h"
 
 IMPLEMENT_GLOBAL_SHADER(FNiagaraDrawIndirectArgsGenCS, "/Plugin/FX/Niagara/Private/NiagaraDrawIndirectArgsGen.usf", "MainCS", SF_Compute);
 
@@ -53,19 +53,17 @@ void FNiagaraDrawIndirectArgsGenCS::SetParameters(FRHICommandList& RHICmdList, F
 {
 	FRHIComputeShader* ComputeShaderRHI = GetComputeShader();
 
-	RHICmdList.SetShaderResourceViewParameter(ComputeShaderRHI, TaskInfosParam.GetBaseIndex(), TaskInfosBuffer);
+	SetSRVParameter(RHICmdList, ComputeShaderRHI, TaskInfosParam, TaskInfosBuffer);
 
 	const FUintVector4 TaskCountValue((int32)NumArgGenTasks, (int32)NumInstanceCountClearTasks, (int32)(NumArgGenTasks + NumInstanceCountClearTasks), 0);
-	RHICmdList.SetShaderParameter(ComputeShaderRHI, TaskCountParam.GetBufferIndex(), TaskCountParam.GetBaseIndex(), TaskCountParam.GetNumBytes(), &TaskCountValue);
+	SetShaderValue(RHICmdList, ComputeShaderRHI, TaskCountParam, TaskCountValue);
 }
 
 void FNiagaraDrawIndirectArgsGenCS::UnbindBuffers(FRHICommandList& RHICmdList)
 {
 	FRHIComputeShader* ComputeShaderRHI = GetComputeShader();
-	if (TaskInfosParam.IsBound())
-	{
-		RHICmdList.SetShaderResourceViewParameter(ComputeShaderRHI, TaskInfosParam.GetBaseIndex(), nullptr);
-	}
+
+	SetSRVParameter(RHICmdList, ComputeShaderRHI, TaskInfosParam, nullptr);
 	if (DrawIndirectArgsParam.IsBound())
 	{
 		RHICmdList.SetUAVParameter(ComputeShaderRHI, DrawIndirectArgsParam.GetUAVIndex(), nullptr);
@@ -74,4 +72,4 @@ void FNiagaraDrawIndirectArgsGenCS::UnbindBuffers(FRHICommandList& RHICmdList)
 	{
 		RHICmdList.SetUAVParameter(ComputeShaderRHI, InstanceCountsParam.GetUAVIndex(), nullptr);
 	}
-}
+}
\ No newline at end of file
diff --git a/Engine/Plugins/FX/Niagara/Source/NiagaraVertexFactories/Private/NiagaraSortingGPU.cpp b/Engine/Plugins/FX/Niagara/Source/NiagaraVertexFactories/Private/NiagaraSortingGPU.cpp
--- a/Engine/Plugins/FX/Niagara/Source/NiagaraVertexFactories/Private/NiagaraSortingGPU.cpp
+++ b/Engine/Plugins/FX/Niagara/Source/NiagaraVertexFactories/Private/NiagaraSortingGPU.cpp
@@ -6,7 +6,7 @@ NiagaraSortingGPU.cpp: Niagara sorting shaders
 
 #include "NiagaraSortingGPU.h"
 #include "NiagaraGPUSortInfo.h"
-
+#include "ShaderParameterUtils.h"
 
 int32 GNiagaraGPUSorting = 1;
 static FAutoConsoleVariableRef CVarNiagaraGPUSorting(
@@ -121,60 +121,40 @@ bool FNiagaraSortKeyGenCS::Serialize(FArchive& Ar)
 void FNiagaraSortKeyGenCS::SetOutput(FRHICommandList& RHICmdList, FRHIUnorderedAccessView* OutKeysUAV, FRHIUnorderedAccessView* OutIndicesUAV)
 {
 	FRHIComputeShader* ComputeShaderRHI = GetComputeShader();
-	if (OutKeys.IsBound())
-	{
-		RHICmdList.SetUAVParameter(ComputeShaderRHI, OutKeys.GetBaseIndex(), OutKeysUAV);
-	}
-	if (OutParticleIndices.IsBound())
-	{
-		RHICmdList.SetUAVParameter(ComputeShaderRHI, OutParticleIndices.GetBaseIndex(), OutIndicesUAV);
-	}
+	SetUAVParameter(RHICmdList, ComputeShaderRHI, OutKeys, OutKeysUAV);
+	SetUAVParameter(RHICmdList, ComputeShaderRHI, OutParticleIndices, OutIndicesUAV);
 }
 
 void FNiagaraSortKeyGenCS::SetParameters(FRHICommandList& RHICmdList, const FNiagaraGPUSortInfo& SortInfo, uint32 EmitterKey, int32 OutputOffset, const FUintVector4& SortKeyParamsValue)
 {
 	FRHIComputeShader* ComputeShaderRHI = GetComputeShader();
 
-	RHICmdList.SetShaderResourceViewParameter(ComputeShaderRHI, NiagaraParticleDataFloat.GetBaseIndex(), SortInfo.ParticleDataFloatSRV);
-	RHICmdList.SetShaderParameter(ComputeShaderRHI, FloatDataOffset.GetBufferIndex(), FloatDataOffset.GetBaseIndex(), FloatDataOffset.GetNumBytes(), &SortInfo.FloatDataOffset);
-	RHICmdList.SetShaderParameter(ComputeShaderRHI, FloatDataStride.GetBufferIndex(), FloatDataStride.GetBaseIndex(), FloatDataStride.GetNumBytes(), &SortInfo.FloatDataStride);
+	SetSRVParameter(RHICmdList, ComputeShaderRHI, NiagaraParticleDataFloat, SortInfo.ParticleDataFloatSRV);
+	SetShaderValue(RHICmdList, ComputeShaderRHI, FloatDataOffset, SortInfo.FloatDataOffset);
+	SetShaderValue(RHICmdList, ComputeShaderRHI, FloatDataStride, SortInfo.FloatDataStride);
+
+	SetSRVParameter(RHICmdList, ComputeShaderRHI, GPUParticleCountBuffer, SortInfo.GPUParticleCountSRV);
 
-	RHICmdList.SetShaderResourceViewParameter(ComputeShaderRHI, GPUParticleCountBuffer.GetBaseIndex(), SortInfo.GPUParticleCountSRV);
 	const FUintVector4 ParticleCountParamsValue(SortInfo.ParticleCount, SortInfo.GPUParticleCountOffset, 0, 0);
-	RHICmdList.SetShaderParameter(ComputeShaderRHI, ParticleCountParams.GetBufferIndex(), ParticleCountParams.GetBaseIndex(), ParticleCountParams.GetNumBytes(), &ParticleCountParamsValue);
+	SetShaderValue(RHICmdList, ComputeShaderRHI, ParticleCountParams, ParticleCountParamsValue);
 
 	// (EmitterKey, OutputOffset, SortMode, SortAttributeOffset)
 	const FUintVector4 SortParamsValue(EmitterKey, OutputOffset, (uint8)SortInfo.SortMode, SortInfo.SortAttributeOffset);
-	RHICmdList.SetShaderParameter(ComputeShaderRHI, SortParams.GetBufferIndex(), SortParams.GetBaseIndex(), SortParams.GetNumBytes(), &SortParamsValue);
+	SetShaderValue(RHICmdList, ComputeShaderRHI, SortParams, SortParamsValue);
 
 	// SortKeyParams only exists in the permutation with SORT_MAX_PRECISION set.
-	if(SortKeyParams.IsBound())
-	{
-		RHICmdList.SetShaderParameter(ComputeShaderRHI, SortKeyParams.GetBufferIndex(), SortKeyParams.GetBaseIndex(), SortKeyParams.GetNumBytes(), &SortKeyParamsValue);
-	}
-	RHICmdList.SetShaderParameter(ComputeShaderRHI, CameraPosition.GetBufferIndex(), CameraPosition.GetBaseIndex(), CameraPosition.GetNumBytes(), &SortInfo.ViewOrigin);
-	RHICmdList.SetShaderParameter(ComputeShaderRHI, CameraDirection.GetBufferIndex(), CameraDirection.GetBaseIndex(), CameraDirection.GetNumBytes(), &SortInfo.ViewDirection);
+	SetShaderValue(RHICmdList, ComputeShaderRHI, SortKeyParams, SortKeyParamsValue);
+	SetShaderValue(RHICmdList, ComputeShaderRHI, CameraPosition, SortInfo.ViewOrigin);
+	SetShaderValue(RHICmdList, ComputeShaderRHI, CameraDirection, SortInfo.ViewDirection);
 }
 
 void FNiagaraSortKeyGenCS::UnbindBuffers(FRHICommandList& RHICmdList)
 {
 	FRHIComputeShader* ComputeShaderRHI = GetComputeShader();
-	if (NiagaraParticleDataFloat.IsBound())
-	{
-		RHICmdList.SetShaderResourceViewParameter(ComputeShaderRHI, NiagaraParticleDataFloat.GetBaseIndex(), nullptr);
-	}
-	if (GPUParticleCountBuffer.IsBound())
-	{
-		RHICmdList.SetShaderResourceViewParameter(ComputeShaderRHI, GPUParticleCountBuffer.GetBaseIndex(), nullptr);
-	}
-	if (OutKeys.IsBound())
-	{
-		RHICmdList.SetUAVParameter(ComputeShaderRHI, OutKeys.GetBaseIndex(), nullptr);
-	}
-	if (OutParticleIndices.IsBound())
-	{
-		RHICmdList.SetUAVParameter(ComputeShaderRHI, OutParticleIndices.GetBaseIndex(), nullptr);
-	}
+	SetUAVParameter(RHICmdList, ComputeShaderRHI, NiagaraParticleDataFloat, nullptr);
+	SetUAVParameter(RHICmdList, ComputeShaderRHI, GPUParticleCountBuffer, nullptr);
+	SetUAVParameter(RHICmdList, ComputeShaderRHI, OutKeys, nullptr);
+	SetUAVParameter(RHICmdList, ComputeShaderRHI, OutParticleIndices, nullptr);
 }
 
 //////////////////////////////////////////////////////////////////////////////////////////////////////////////////////
@@ -220,32 +200,26 @@ void FNiagaraCopyIntBufferRegionCS::SetParameters(
 	FUintVector4 CopyParamsValue(StartingIndex, 0, 0, 0);
 	for (int32 Index = 0; Index < DestCount; ++Index)
 	{
-		RHICmdList.SetUAVParameter(ComputeShaderRHI, DestData[Index].GetBaseIndex(), InDestDatas[Index]);
+		SetUAVParameter(RHICmdList, ComputeShaderRHI, DestData[Index], InDestDatas[Index]);
 		CopyParamsValue[Index + 1] = InUsedIndexCounts[Index];
 	}
 
 	for (int32 Index = DestCount; Index < NIAGARA_COPY_BUFFER_BUFFER_COUNT; ++Index)
 	{
 		RHICmdList.TransitionResource(EResourceTransitionAccess::ERWNoBarrier, EResourceTransitionPipeline::EComputeToCompute, NiagaraSortingDummyUAV[Index].Buffer.UAV);
-		RHICmdList.SetUAVParameter(ComputeShaderRHI, DestData[Index].GetBaseIndex(), NiagaraSortingDummyUAV[Index].Buffer.UAV);
+		SetUAVParameter(RHICmdList, ComputeShaderRHI, DestData[Index], NiagaraSortingDummyUAV[Index].Buffer.UAV);
 	}
 
-	RHICmdList.SetShaderParameter(ComputeShaderRHI, CopyParams.GetBufferIndex(), CopyParams.GetBaseIndex(), CopyParams.GetNumBytes(), &CopyParamsValue);
+	SetShaderValue(RHICmdList, ComputeShaderRHI, CopyParams, CopyParamsValue);
 }
 
 void FNiagaraCopyIntBufferRegionCS::UnbindBuffers(FRHICommandList& RHICmdList)
 {
 	FRHIComputeShader* ComputeShaderRHI = GetComputeShader();
-	if (SourceData.IsBound())
-	{
-		RHICmdList.SetShaderResourceViewParameter(ComputeShaderRHI, SourceData.GetBaseIndex(), nullptr);
-	}
+	SetSRVParameter(RHICmdList, ComputeShaderRHI, SourceData, nullptr);
 	for (int32 Index = 0; Index < NIAGARA_COPY_BUFFER_BUFFER_COUNT; ++Index)
 	{
-		if (DestData[Index].IsBound())
-		{
-			RHICmdList.SetUAVParameter(ComputeShaderRHI, DestData[Index].GetBaseIndex(), nullptr);
-		}
+		SetUAVParameter(RHICmdList, ComputeShaderRHI, DestData[Index], nullptr);
 	}
 }
 
