Migrate stuff over to RHI helper functions to ensure we properly handed unbounded parameters

diff --git a/Engine/Plugins/FX/Niagara/Source/Niagara/Private/NiagaraDataInterfaceGrid2DCollection.cpp b/Engine/Plugins/FX/Niagara/Source/Niagara/Private/NiagaraDataInterfaceGrid2DCollection.cpp
--- a/Engine/Plugins/FX/Niagara/Source/Niagara/Private/NiagaraDataInterfaceGrid2DCollection.cpp
+++ b/Engine/Plugins/FX/Niagara/Source/Niagara/Private/NiagaraDataInterfaceGrid2DCollection.cpp
@@ -92,7 +92,7 @@ struct FNiagaraDataInterfaceParametersCS_Grid2DCollection : public FNiagaraDataI
 
 		if (Context.IsOutputStage && OutputGridParam.IsBound() && ProxyData->GetDestinationData() != NULL)
 		{
-			const FTextureRWBuffer& TextureBuffer = ProxyData->GetDestinationData()->GridBuffer;
+			const FTextureRWBuffer2D& TextureBuffer = ProxyData->GetDestinationData()->GridBuffer;
 			RHICmdList.TransitionResource(EResourceTransitionAccess::EWritable, EResourceTransitionPipeline::EComputeToCompute, ProxyData->GetDestinationData()->GridBuffer.UAV);
 			OutputGridParam.SetTexture(RHICmdList, Context.Shader->GetComputeShader(), TextureBuffer.Buffer, TextureBuffer.UAV);
 		}
