FNM - remove double array entries in NiagaraTransitions

diff --git a/Engine/Plugins/FX/Niagara/Source/Niagara/Private/NiagaraEmitterInstanceBatcher.cpp b/Engine/Plugins/FX/Niagara/Source/Niagara/Private/NiagaraEmitterInstanceBatcher.cpp
--- a/Engine/Plugins/FX/Niagara/Source/Niagara/Private/NiagaraEmitterInstanceBatcher.cpp
+++ b/Engine/Plugins/FX/Niagara/Source/Niagara/Private/NiagaraEmitterInstanceBatcher.cpp
@@ -445,7 +445,7 @@ void NiagaraEmitterInstanceBatcher::DispatchAllOnCompute(FOverlappableTicks& Ove
 		TransitionUAVs.Append(DestBufferIntFloat.GetData(), DestBufferIntFloat.Num());
 		TransitionUAVs.Append(CurrDataBuffers.GetData(), CurrDataBuffers.Num());
 		TransitionUAVs.Append(CurrBufferIntFloat.GetData(), CurrBufferIntFloat.Num());
-		RHICmdList.TransitionResources(EResourceTransitionAccess::ERWBarrier, EResourceTransitionPipeline::EComputeToCompute, TransitionUAVs.GetData(), TransitionUAVs.Num());
+		RHICmdList.TransitionResources(EResourceTransitionAccess::ERWBarrier, EResourceTransitionPipeline::EGfxToCompute, TransitionUAVs.GetData(), TransitionUAVs.Num());
 	}
 
 	for (FNiagaraGPUSystemTick* Tick : OverlappableTick)
@@ -487,9 +487,7 @@ void NiagaraEmitterInstanceBatcher::DispatchAllOnCompute(FOverlappableTicks& Ove
 		TransitionUAVs.Add(GPUInstanceCounterManager.GetInstanceCountBuffer().UAV);
 		TransitionUAVs.Append(DestDataBuffers.GetData(), DestDataBuffers.Num());
 		TransitionUAVs.Append(DestBufferIntFloat.GetData(), DestBufferIntFloat.Num());
-		TransitionUAVs.Append(DestDataBuffers.GetData(), DestDataBuffers.Num());
-		TransitionUAVs.Append(DestBufferIntFloat.GetData(), DestBufferIntFloat.Num());
-		RHICmdList.TransitionResources(EResourceTransitionAccess::EReadable, EResourceTransitionPipeline::EComputeToCompute, TransitionUAVs.GetData(), TransitionUAVs.Num());
+		RHICmdList.TransitionResources(EResourceTransitionAccess::EReadable, EResourceTransitionPipeline::EComputeToGfx, TransitionUAVs.GetData(), TransitionUAVs.Num());
 	}
 	// We have done all our overlapping compute work on this list so go back to default behavior and flush.
 	RHICmdList.AutomaticCacheFlushAfterComputeShader(true);
