Niagara renderer preview thumbnails are now only realtime when hovered. Fixes a crash when creating the thumbnail bar
#rnx

diff --git a/Engine/Plugins/FX/Niagara/Source/NiagaraEditorWidgets/Private/SNiagaraOverviewStackNode.cpp b/Engine/Plugins/FX/Niagara/Source/NiagaraEditorWidgets/Private/SNiagaraOverviewStackNode.cpp
--- a/Engine/Plugins/FX/Niagara/Source/NiagaraEditorWidgets/Private/SNiagaraOverviewStackNode.cpp
+++ b/Engine/Plugins/FX/Niagara/Source/NiagaraEditorWidgets/Private/SNiagaraOverviewStackNode.cpp
@@ -20,6 +20,7 @@
 #include "Widgets/SBoxPanel.h"
 #include "NiagaraRendererProperties.h"
 #include "ViewModels/NiagaraSystemSelectionViewModel.h"
+#include "Widgets/SWidget.h"
 
 #define LOCTEXT_NAMESPACE "NiagaraOverviewStackNode"
 
@@ -32,7 +33,8 @@ void SNiagaraOverviewStackNode::Construct(const FArguments& InArgs, UNiagaraOver
 	StackViewModel = nullptr;
 	OverviewSelectionViewModel = nullptr;
 	EmitterHandleViewModelWeak.Reset();
-
+	ThumbnailPool = MakeShared<FAssetThumbnailPool>(100, TAttribute<bool>::Create(TAttribute<bool>::FGetter::CreateSP(this, &SNiagaraOverviewStackNode::IsHoveringThumbnail)));
+	ThumbnailBar = SNew(SHorizontalBox);
 	if (OverviewStackNode->GetOwningSystem() != nullptr)
 	{
 		FNiagaraEditorModule& NiagaraEditorModule = FModuleManager::Get().LoadModuleChecked<FNiagaraEditorModule>("NiagaraEditor");
@@ -53,7 +55,7 @@ void SNiagaraOverviewStackNode::Construct(const FArguments& InArgs, UNiagaraOver
 			}
 			if (StackViewModel)
 			{
-				StackViewModel->OnDataObjectChanged().AddRaw(this, &SNiagaraOverviewStackNode::FillThumbnailBar, true);
+				StackViewModel->OnDataObjectChanged().AddSP(this, &SNiagaraOverviewStackNode::FillThumbnailBar, true);
 			}
 			OverviewSelectionViewModel = OwningSystemViewModel->GetSelectionViewModel();
 		}
@@ -133,23 +135,29 @@ TSharedRef<SWidget> SNiagaraOverviewStackNode::CreateTitleWidget(TSharedPtr<SNod
 
 TSharedRef<SWidget> SNiagaraOverviewStackNode::CreateThumbnailWidget(float InThumbnailSize, FRendererPreviewData* InData)
 {
+	TSharedRef<SWidget> ThumbnailWidget = SNullWidget::NullWidget;
 	TSharedPtr<FAssetThumbnail> AssetThumbnail = MakeShareable(new FAssetThumbnail(InData->RenderingObject, InThumbnailSize, InThumbnailSize, ThumbnailPool));
-	TSharedRef<SWidget> ThumbnailWidget = AssetThumbnail->MakeThumbnailWidget();
-	TSharedPtr<FAssetThumbnail> TooltipThumbnail = MakeShareable(new FAssetThumbnail(InData->RenderingObject, 64.0f, 64.0f, ThumbnailPool));
-	TSharedRef<SToolTip> ThumbnailTooltipWidget = SNew(SToolTip)
-		.Content()
-		[
-			SNew(SBox)
-			.MinDesiredHeight(64.0f)
-			.MaxDesiredHeight(64.0f)
-			.MinDesiredWidth(64.0f)
-			.MaxDesiredWidth(64.0f)
+	if (AssetThumbnail)
+	{
+		ThumbnailWidget = AssetThumbnail->MakeThumbnailWidget();
+		TSharedPtr<FAssetThumbnail> TooltipThumbnail = MakeShareable(new FAssetThumbnail(InData->RenderingObject, 64.0f, 64.0f, ThumbnailPool));
+		TSharedRef<SToolTip> ThumbnailTooltipWidget = SNew(SToolTip)
+			.Content()
 			[
-				TooltipThumbnail->MakeThumbnailWidget()
-			]
-		];
-	ThumbnailWidget->SetOnMouseButtonDown(FPointerEventHandler::CreateSP(this, &SNiagaraOverviewStackNode::OnClickedRenderingPreview, InData->RenderingEntry));
-	ThumbnailWidget->SetToolTip(ThumbnailTooltipWidget);
+				SNew(SBox)
+				.MinDesiredHeight(64.0f)
+				.MaxDesiredHeight(64.0f)
+				.MinDesiredWidth(64.0f)
+				.MaxDesiredWidth(64.0f)
+				[
+					TooltipThumbnail->MakeThumbnailWidget()
+				]
+			];
+		ThumbnailWidget->SetOnMouseButtonDown(FPointerEventHandler::CreateSP(this, &SNiagaraOverviewStackNode::OnClickedRenderingPreview, InData->RenderingEntry));
+		ThumbnailWidget->SetOnMouseEnter(FNoReplyPointerEventHandler::CreateSP(this, &SNiagaraOverviewStackNode::SetIsHoveringThumbnail, true));
+		ThumbnailWidget->SetOnMouseLeave(FSimpleNoReplyPointerEventHandler::CreateSP(this, &SNiagaraOverviewStackNode::SetIsHoveringThumbnail, false));
+		ThumbnailWidget->SetToolTip(ThumbnailTooltipWidget);
+	}
 	return ThumbnailWidget;
 }
 
@@ -194,7 +202,6 @@ TSharedRef<SWidget> SNiagaraOverviewStackNode::CreateNodeContentArea()
 		ContentWidget = SNullWidget::NullWidget;
 	}
 
-	ThumbnailBar = SNew(SHorizontalBox);
 	FillThumbnailBar(nullptr, false);
 
 	// NODE CONTENT AREA
@@ -259,8 +266,10 @@ void SNiagaraOverviewStackNode::FillThumbnailBar(UObject* ChangedObject, const b
 	UNiagaraRendererProperties* RendererProperties = Cast< UNiagaraRendererProperties>(ChangedObject);
 	if (!bIsTriggeredByObjectUpdate || RendererProperties != nullptr)
 	{
-		ThumbnailBar->ClearChildren();
-		ThumbnailPool = MakeShared<FAssetThumbnailPool>(100, /*InAreRealTileThumbnailsAllowed=*/true);
+		if (ThumbnailBar.IsValid() && ThumbnailBar->GetChildren())
+		{
+			ThumbnailBar->ClearChildren();
+		}
 		if (EmitterHandleViewModelWeak.IsValid())
 		{
 			EmitterHandleViewModelWeak.Pin()->GetRendererPreviewData(PreviewData);
diff --git a/Engine/Plugins/FX/Niagara/Source/NiagaraEditorWidgets/Private/SNiagaraOverviewStackNode.h b/Engine/Plugins/FX/Niagara/Source/NiagaraEditorWidgets/Private/SNiagaraOverviewStackNode.h
--- a/Engine/Plugins/FX/Niagara/Source/NiagaraEditorWidgets/Private/SNiagaraOverviewStackNode.h
+++ b/Engine/Plugins/FX/Niagara/Source/NiagaraEditorWidgets/Private/SNiagaraOverviewStackNode.h
@@ -40,6 +40,18 @@ private:
 	FReply OnToggleIsolateButtonClicked();
 	EVisibility GetToggleIsolateVisibility() const;
 	FSlateColor GetToggleIsolateImageColor() const;
+	void SetIsHoveringThumbnail(const FGeometry& InGeometry, const FPointerEvent& InEvent, const bool bInHoveringThumbnail)
+	{
+		SetIsHoveringThumbnail(InEvent, bInHoveringThumbnail);
+	}
+	void SetIsHoveringThumbnail(const FPointerEvent& InEvent, const bool bInHoveringThumbnail)
+	{
+		bIsHoveringThumbnail = bInHoveringThumbnail;
+	}
+	bool IsHoveringThumbnail()
+	{
+		return bIsHoveringThumbnail;
+	}
 
 private:
 	UNiagaraOverviewNode* OverviewStackNode;
@@ -50,5 +62,5 @@ private:
 	/** Thumbnail widget containers */
 	TSharedPtr<SHorizontalBox> ThumbnailBar;
 	TArray<FRendererPreviewData*> PreviewData;
-	
+	bool bIsHoveringThumbnail;
 };
\ No newline at end of file
