Niagara - Bump the niagara compiler version GUID to hopefully fix some cooking issues where ddc data generated in main is being incorrectly picked up by compiles in the current release branch.

diff --git a/Engine/Plugins/FX/Niagara/Source/NiagaraCore/Private/NiagaraCustomVersion.cpp b/Engine/Plugins/FX/Niagara/Source/NiagaraCore/Private/NiagaraCustomVersion.cpp
--- a/Engine/Plugins/FX/Niagara/Source/NiagaraCore/Private/NiagaraCustomVersion.cpp
+++ b/Engine/Plugins/FX/Niagara/Source/NiagaraCore/Private/NiagaraCustomVersion.cpp
@@ -8,5 +8,5 @@ const FGuid FNiagaraCustomVersion::GUID(0xFCF57AFA, 0x50764283, 0xB9A9E658, 0xFF
 // Register the custom version with core
 FCustomVersionRegistration GRegisterNiagaraCustomVersion(FNiagaraCustomVersion::GUID, FNiagaraCustomVersion::LatestVersion, TEXT("NiagaraVer"));
 
-const FGuid FNiagaraCustomVersion::LatestScriptCompileVersion(0xDEA735C8, 0xD11A3B4A, 0xA4841FAF, 0xDF7D991E);
+const FGuid FNiagaraCustomVersion::LatestScriptCompileVersion(0x7698EEE8, 0x01EA483B, 0xAA2841E4, 0xCC6BED0D);
 
