Updated handling of component lifecycle and GC in relation the the scalability manager.
#rnx

diff --git a/Engine/Plugins/FX/Niagara/Source/Niagara/Private/NiagaraComponent.cpp b/Engine/Plugins/FX/Niagara/Source/Niagara/Private/NiagaraComponent.cpp
--- a/Engine/Plugins/FX/Niagara/Source/Niagara/Private/NiagaraComponent.cpp
+++ b/Engine/Plugins/FX/Niagara/Source/Niagara/Private/NiagaraComponent.cpp
@@ -1020,6 +1020,8 @@ void UNiagaraComponent::DestroyInstance()
 	//UE_LOG(LogNiagara, Log, TEXT("UNiagaraComponent::DestroyInstance: %p  %s\n"), SystemInstance.Get(), *GetAsset()->GetFullName());
 	//UE_LOG(LogNiagara, Log, TEXT("DestroyInstance: %u - %s"), this, *Asset->GetName());
 	SetActiveFlag(false);
+
+	UnregisterWithScalabilityManager();
 	
 	// Rather than setting the unique ptr to null here, we allow it to transition ownership to the system's deferred deletion queue. This allows us to safely
 	// get rid of the system interface should we be doing this in response to a callback invoked during the system interface's lifetime completion cycle.
@@ -1089,6 +1091,8 @@ void UNiagaraComponent::OnUnregister()
 
 	SetActiveFlag(false);
 
+	UnregisterWithScalabilityManager();
+
 	if (SystemInstance)
 	{
 		//UE_LOG(LogNiagara, Log, TEXT("UNiagaraComponent::OnUnregister: %p  %s\n"), SystemInstance.Get(), *GetAsset()->GetFullName());
@@ -1107,7 +1111,7 @@ void UNiagaraComponent::OnUnregister()
 
 void UNiagaraComponent::BeginDestroy()
 {
-	//UE_LOG(LogNiagara, Log, TEXT("~UNiagaraComponent: %p  %s\n"), SystemInstance.Get(), *GetAsset()->GetFullName());
+	//UE_LOG(LogNiagara, Log, TEXT("UNiagaraComponent::BeginDestroy(): %0xP  %s\n"), this, *GetAsset()->GetFullName());
 	DestroyInstance();
 
 	Super::BeginDestroy();
diff --git a/Engine/Plugins/FX/Niagara/Source/Niagara/Private/NiagaraScalabilityManager.cpp b/Engine/Plugins/FX/Niagara/Source/Niagara/Private/NiagaraScalabilityManager.cpp
--- a/Engine/Plugins/FX/Niagara/Source/Niagara/Private/NiagaraScalabilityManager.cpp
+++ b/Engine/Plugins/FX/Niagara/Source/Niagara/Private/NiagaraScalabilityManager.cpp
@@ -20,6 +20,12 @@ FNiagaraScalabilityManager::FNiagaraScalabilityManager()
 
 }
 
+void FNiagaraScalabilityManager::AddReferencedObjects(FReferenceCollector& Collector)
+{
+	Collector.AddReferencedObject(EffectType);
+	Collector.AddReferencedObjects(ManagedComponents);
+}
+
 void FNiagaraScalabilityManager::Register(UNiagaraComponent* Component)
 {
 	check(Component->ScalabilityManagerHandle == INDEX_NONE);
@@ -27,6 +33,8 @@ void FNiagaraScalabilityManager::Register(UNiagaraComponent* Component)
 
 	Component->ScalabilityManagerHandle = ManagedComponents.Add(Component);
 	State.AddDefaulted();
+
+	//UE_LOG(LogNiagara, Warning, TEXT("Registered Component %0xP at index %d"), Component, Component->ScalabilityManagerHandle);
 }
 
 void FNiagaraScalabilityManager::Unregister(UNiagaraComponent* Component)
@@ -36,6 +44,13 @@ void FNiagaraScalabilityManager::Unregister(UNiagaraComponent* Component)
 
 	int32 IndexToRemove = Component->ScalabilityManagerHandle;
 	Component->ScalabilityManagerHandle = INDEX_NONE;
+	UnregisterAt(IndexToRemove);
+}
+
+void FNiagaraScalabilityManager::UnregisterAt(int32 IndexToRemove)
+{
+	//UE_LOG(LogNiagara, Warning, TEXT("Unregistering Component %0xP at index %d (Replaced with %0xP)"), ManagedComponents[IndexToRemove], IndexToRemove, ManagedComponents.Num() > 1 ? ManagedComponents.Last() : nullptr);
+
 	ManagedComponents.RemoveAtSwap(IndexToRemove);
 	State.RemoveAtSwap(IndexToRemove);
 
@@ -48,6 +63,14 @@ void FNiagaraScalabilityManager::Unregister(UNiagaraComponent* Component)
 
 void FNiagaraScalabilityManager::Update(FNiagaraWorldManager* WorldMan)
 {
+	//Paranoia code in case the EffectType is GCd from under us.
+	if (EffectType == nullptr)
+	{
+		ManagedComponents.Empty();
+		State.Empty();
+		LastUpdateTime = 0.0f;
+	}
+
 	float WorldTime = WorldMan->GetWorld()->GetTimeSeconds();
 	bool bShouldUpdateScalabilityStates = false;
 	switch (EffectType->UpdateFrequency)
@@ -65,22 +88,31 @@ void FNiagaraScalabilityManager::Update(FNiagaraWorldManager* WorldMan)
 
 	LastUpdateTime = WorldTime;
 
+	//Belt and braces paranoia code to ensure we're safe if a component or System is GCd but the component isn't unregistered for whatever reason.
 	int32 CompIdx = 0;
-#if WITH_EDITOR
-	//In cases such as undo/redo we can get components and system asset references pulled from under us.
 	while (CompIdx < ManagedComponents.Num())
 	{
 		UNiagaraComponent* Component = ManagedComponents[CompIdx];
-		if (Component && Component->GetAsset())
+		if (Component)
 		{
-			++CompIdx;
+			if (Component->GetAsset())
+			{
+				++CompIdx;
+			}
+			else
+			{
+				UE_LOG(LogNiagara, Warning, TEXT("Niagara System has been destroyed with components still registered to the scalability manager. Unregistering this component.\nComponent: %0xP - %s\nEffectType: %0xP - %s"),
+					Component, *Component->GetName(), EffectType, *EffectType->GetName());
+				Unregister(Component);
+			}
 		}
 		else
 		{
-			Unregister(Component);
+			UE_LOG(LogNiagara, Warning, TEXT("Niagara Component has been destroyed while still registered to the scalability manager. Unregistering this component.\nComponent: %0xP \nEffectType: %0xP - %s"),
+				Component, EffectType, *EffectType->GetName());
+			UnregisterAt(CompIdx);
 		}
 	}
-#endif
 
 	bool bNeedSortedSignificanceCull = false;
 	SignificanceSortedIndices.Reset(ManagedComponents.Num());
diff --git a/Engine/Plugins/FX/Niagara/Source/Niagara/Private/NiagaraWorldManager.cpp b/Engine/Plugins/FX/Niagara/Source/Niagara/Private/NiagaraWorldManager.cpp
--- a/Engine/Plugins/FX/Niagara/Source/Niagara/Private/NiagaraWorldManager.cpp
+++ b/Engine/Plugins/FX/Niagara/Source/Niagara/Private/NiagaraWorldManager.cpp
@@ -191,7 +191,14 @@ void FNiagaraWorldManager::AddReferencedObjects(FReferenceCollector& Collector)
 
 	Collector.AddReferencedObjects(ParameterCollections);
 	Collector.AddReferencedObject(ComponentPool);
-	Collector.AddReferencedObjects(ScalabilityManagers);
+	for (auto& Pair : ScalabilityManagers)
+	{
+		UNiagaraEffectType* EffectType = Pair.Key;
+		FNiagaraScalabilityManager& ScalabilityMan = Pair.Value;
+
+		Collector.AddReferencedObject(EffectType);
+		ScalabilityMan.AddReferencedObjects(Collector);
+	}
 }
 
 FString FNiagaraWorldManager::GetReferencerName() const
diff --git a/Engine/Plugins/FX/Niagara/Source/Niagara/Public/NiagaraScalabilityManager.h b/Engine/Plugins/FX/Niagara/Source/Niagara/Public/NiagaraScalabilityManager.h
--- a/Engine/Plugins/FX/Niagara/Source/Niagara/Public/NiagaraScalabilityManager.h
+++ b/Engine/Plugins/FX/Niagara/Source/Niagara/Public/NiagaraScalabilityManager.h
@@ -9,6 +9,7 @@
 class FNiagaraWorldManager;
 class UNiagaraSystem;
 class UNiagaraComponent;
+class FReferenceCollector;
 
 //Disabled for now until we can spend more time on a good method of applying the data that is gathered.
 #define ENABLE_NIAGARA_RUNTIME_CYCLE_COUNTERS (0)
@@ -128,7 +129,12 @@ struct FNiagaraScalabilityManager
 	void Register(UNiagaraComponent* Component);
 	void Unregister(UNiagaraComponent* Component);
 
+	void AddReferencedObjects(FReferenceCollector& Collector);
+
 #if DEBUG_SCALABILITY_STATE
 	void Dump();
 #endif
+
+private: 
+	void UnregisterAt(int32 IndexToRemove);
 };
