Allowing FNiagaraUpdateContext to be split between PreEditChagne and PostEditChange so that we can safely teardown/reinit system instances without editing data from underneath them.
Should also help in some Undo/Redo cases.

diff --git a/Engine/Plugins/FX/Niagara/Source/Niagara/Classes/NiagaraSystem.h b/Engine/Plugins/FX/Niagara/Source/Niagara/Classes/NiagaraSystem.h
--- a/Engine/Plugins/FX/Niagara/Source/Niagara/Classes/NiagaraSystem.h
+++ b/Engine/Plugins/FX/Niagara/Source/Niagara/Classes/NiagaraSystem.h
@@ -250,6 +250,9 @@ public:
 
 	bool GetIsolateEnabled() const;
 	void SetIsolateEnabled(bool bIsolate);
+	
+	UPROPERTY(transient)
+	FNiagaraSystemUpdateContext UpdateContext;
 #endif
 
 	bool ShouldAutoDeactivate() const { return bAutoDeactivate; }
diff --git a/Engine/Plugins/FX/Niagara/Source/Niagara/Private/NiagaraCommon.cpp b/Engine/Plugins/FX/Niagara/Source/Niagara/Private/NiagaraCommon.cpp
--- a/Engine/Plugins/FX/Niagara/Source/Niagara/Private/NiagaraCommon.cpp
+++ b/Engine/Plugins/FX/Niagara/Source/Niagara/Private/NiagaraCommon.cpp
@@ -70,20 +70,38 @@ FString FNiagaraTypeHelper::ToString(const uint8* ValueData, const UScriptStruct
 //////////////////////////////////////////////////////////////////////////
 
 FNiagaraSystemUpdateContext::~FNiagaraSystemUpdateContext()
+{
+	CommitUpdate();
+}
+
+void FNiagaraSystemUpdateContext::CommitUpdate()
 {
 	for (UNiagaraSystem* Sys : SystemSimsToDestroy)
 	{
-		FNiagaraWorldManager::DestroyAllSystemSimulations(Sys);
+		if(Sys)
+		{
+			FNiagaraWorldManager::DestroyAllSystemSimulations(Sys);
+		}		
 	}
+	SystemSimsToDestroy.Empty();
 
 	for (UNiagaraComponent* Comp : ComponentsToReInit)
 	{
-		Comp->ReinitializeSystem();
+		if (Comp)
+		{
+			Comp->ReinitializeSystem();
+		}
 	}
+	ComponentsToReInit.Empty();
+
 	for (UNiagaraComponent* Comp : ComponentsToReset)
 	{
-		Comp->ResetSystem();
+		if (Comp)
+		{
+			Comp->ResetSystem();
+		}
 	}
+	ComponentsToReset.Empty();
 }
 
 void FNiagaraSystemUpdateContext::AddAll(bool bReInit)
diff --git a/Engine/Plugins/FX/Niagara/Source/Niagara/Private/NiagaraSystem.cpp b/Engine/Plugins/FX/Niagara/Source/Niagara/Private/NiagaraSystem.cpp
--- a/Engine/Plugins/FX/Niagara/Source/Niagara/Private/NiagaraSystem.cpp
+++ b/Engine/Plugins/FX/Niagara/Source/Niagara/Private/NiagaraSystem.cpp
@@ -241,12 +241,8 @@ void UNiagaraSystem::PreEditChange(UProperty* PropertyThatWillChange)
 
 	if (PropertyThatWillChange && PropertyThatWillChange->GetFName() == GET_MEMBER_NAME_CHECKED(UNiagaraSystem, EffectType))
 	{
-		FNiagaraSystemUpdateContext UpdateContext;
 		UpdateContext.SetDestroyOnAdd(true);
 		UpdateContext.Add(this, false);
-		
-		//We have to clear this here and reset so that we unregister properly from the scalability manager.
-		EffectType = nullptr;
 	}
 }
 
@@ -279,6 +275,8 @@ void UNiagaraSystem::PostEditChangeProperty(struct FPropertyChangedEvent& Proper
 	}
 
 	ResolveScalabilityOverrides();
+
+	UpdateContext.CommitUpdate();
 }
 #endif 
 
@@ -1376,7 +1374,7 @@ void UNiagaraSystem::ResolveScalabilityOverrides()
 	int32 NumEffectTypeSettings = 0;
 	if (ActualEffectType)
 	{
-		ActualEffectType->GetScalabilitySettings().Num();
+		NumEffectTypeSettings = ActualEffectType->GetScalabilitySettings().Num();
 	}
 	int32 NumOverrides = ScalabilityOverrides.Num();
 
diff --git a/Engine/Plugins/FX/Niagara/Source/Niagara/Public/NiagaraCommon.h b/Engine/Plugins/FX/Niagara/Source/Niagara/Public/NiagaraCommon.h
--- a/Engine/Plugins/FX/Niagara/Source/Niagara/Public/NiagaraCommon.h
+++ b/Engine/Plugins/FX/Niagara/Source/Niagara/Public/NiagaraCommon.h
@@ -523,8 +523,17 @@ struct TStructOpsTypeTraits<FVMExternalFunctionBindingInfo> : public TStructOpsT
 	};
 };
 
+/**
+Helper for reseting/reinitializing Niagara systems currently active when they are being edited. 
+Can be used inside a scope with Systems being reinitialized on destruction or you can store the context and use CommitUpdate() to trigger reinitialization.
+For example, this can be split between PreEditChange and PostEditChange to ensure problematic data is not modified during execution of a system.
+This can be made a UPROPERTY() to ensure safey in cases where a GC could be possible between Add() and CommitUpdate().
+*/
+USTRUCT()
 struct NIAGARA_API FNiagaraSystemUpdateContext
 {
+	GENERATED_BODY()
+
 	FNiagaraSystemUpdateContext(const UNiagaraSystem* System, bool bReInit) :bDestroyOnAdd(false) { Add(System, bReInit); }
 #if WITH_EDITORONLY_DATA
 	FNiagaraSystemUpdateContext(const UNiagaraEmitter* Emitter, bool bReInit) : bDestroyOnAdd(false) { Add(Emitter, bReInit); }
@@ -549,13 +558,18 @@ struct NIAGARA_API FNiagaraSystemUpdateContext
 	/** Adds all currently active systems.*/
 	void AddAll(bool bReInit);
 
+	/** Handles any pending reinits or resets of system instances in this update context. */
+	void CommitUpdate();
+
 private:
 	void AddInternal(class UNiagaraComponent* Comp, bool bReInit);
 	FNiagaraSystemUpdateContext(FNiagaraSystemUpdateContext& Other) :bDestroyOnAdd(false) { }
 
+	UPROPERTY(transient)
 	TArray<UNiagaraComponent*> ComponentsToReset;
+	UPROPERTY(transient)
 	TArray<UNiagaraComponent*> ComponentsToReInit;
-
+	UPROPERTY(transient)
 	TArray<UNiagaraSystem*> SystemSimsToDestroy;
 
 	bool bDestroyOnAdd;
