- Enable optimized VM code by default
- Fallback to synchronous optimizing when free is enabled so that we can free the source until I figure out a safe way to do this

#rnx

diff --git a/Engine/Plugins/FX/Niagara/Source/Niagara/Private/NiagaraScript.cpp b/Engine/Plugins/FX/Niagara/Source/Niagara/Private/NiagaraScript.cpp
--- a/Engine/Plugins/FX/Niagara/Source/Niagara/Private/NiagaraScript.cpp
+++ b/Engine/Plugins/FX/Niagara/Source/Niagara/Private/NiagaraScript.cpp
@@ -499,36 +499,43 @@ void UNiagaraScript::AsyncOptimizeByteCode()
 		ExternalFunctionRegisterCounts.Add(RegisterCount);
 	}
 
-	// Kick off async task to optimize the byte code, note we must take a copy of the ByteCode & VMId currently as we must guard the async task against changes in the editor
-	AsyncTask(
-		ENamedThreads::AnyThread,
-		[WeakScript=TWeakObjectPtr<UNiagaraScript>(this), InExternalFunctionRegisterCounts=MoveTemp(ExternalFunctionRegisterCounts), InByteCode=CachedScriptVM.ByteCode, InCachedScriptVMId=CachedScriptVMId]() mutable
-		{
-			// Generate optimized byte code on any thread
-			TArray<uint8> OptimizedByteCode;
-			VectorVM::OptimizeByteCode(InByteCode.GetData(), OptimizedByteCode, MakeArrayView(InExternalFunctionRegisterCounts));
-
-			// Kick off task to set optimized byte code on game thread
-			AsyncTask(
-				ENamedThreads::GameThread,
-				[WeakScript, InOptimizedByteCode = MoveTemp(OptimizedByteCode), InCachedScriptVMId]() mutable
-				{
-					UNiagaraScript* NiagaraScript = WeakScript.Get();
-					if ( (NiagaraScript != nullptr) && (NiagaraScript->CachedScriptVMId == InCachedScriptVMId) )
+	// If we wish to release the original ByteCode we must optimize synchronously currently
+	//-TODO: Find a safe point where we can release the original ByteCode
+	static const IConsoleVariable* CVarFreeUnoptimizedByteCode = IConsoleManager::Get().FindConsoleVariable(TEXT("vm.FreeUnoptimizedByteCode"));
+	if ( FPlatformProperties::RequiresCookedData() && CVarFreeUnoptimizedByteCode && (CVarFreeUnoptimizedByteCode->GetInt() != 0) )
+	{
+		VectorVM::OptimizeByteCode(CachedScriptVM.ByteCode.GetData(), CachedScriptVM.OptimizedByteCode, MakeArrayView(ExternalFunctionRegisterCounts));
+		if (CachedScriptVM.OptimizedByteCode.Num() > 0)
+		{
+			CachedScriptVM.ByteCode.Empty();
+		}
+	}
+	else
+	{
+		// Async optimize the ByteCode
+		AsyncTask(
+			ENamedThreads::AnyThread,
+			[WeakScript=TWeakObjectPtr<UNiagaraScript>(this), InExternalFunctionRegisterCounts=MoveTemp(ExternalFunctionRegisterCounts), InByteCode=CachedScriptVM.ByteCode, InCachedScriptVMId=CachedScriptVMId]() mutable
+			{
+				// Generate optimized byte code on any thread
+				TArray<uint8> OptimizedByteCode;
+				VectorVM::OptimizeByteCode(InByteCode.GetData(), OptimizedByteCode, MakeArrayView(InExternalFunctionRegisterCounts));
+
+				// Kick off task to set optimized byte code on game thread
+				AsyncTask(
+					ENamedThreads::GameThread,
+					[WeakScript, InOptimizedByteCode = MoveTemp(OptimizedByteCode), InCachedScriptVMId]() mutable
 					{
-						NiagaraScript->CachedScriptVM.OptimizedByteCode = MoveTemp(InOptimizedByteCode);
-
-						static const IConsoleVariable* CVarFreeUnoptimizedByteCode = IConsoleManager::Get().FindConsoleVariable(TEXT("vm.FreeUnoptimizedByteCode"));
-						if ( CVarFreeUnoptimizedByteCode && (CVarFreeUnoptimizedByteCode->GetInt() != 0) && (NiagaraScript->CachedScriptVM.OptimizedByteCode.Num() > 0) )
+						UNiagaraScript* NiagaraScript = WeakScript.Get();
+						if ( (NiagaraScript != nullptr) && (NiagaraScript->CachedScriptVMId == InCachedScriptVMId) )
 						{
-							//-TODO: Support this, we will have to be careful that the VM is not currently executing the ByteCode before we free it
-							//NiagaraScript->CachedScriptVM.ByteCode.Empty();
+							NiagaraScript->CachedScriptVM.OptimizedByteCode = MoveTemp(InOptimizedByteCode);
 						}
 					}
-				}
-			);
-		}
-	);
+				);
+			}
+		);
+	}
 }
 
 void UNiagaraScript::PreSave(const class ITargetPlatform* TargetPlatform)
diff --git a/Engine/Source/Runtime/VectorVM/Private/VectorVM.cpp b/Engine/Source/Runtime/VectorVM/Private/VectorVM.cpp
--- a/Engine/Source/Runtime/VectorVM/Private/VectorVM.cpp
+++ b/Engine/Source/Runtime/VectorVM/Private/VectorVM.cpp
@@ -174,7 +174,7 @@ static FAutoConsoleVariableRef CVarParallelVVMInstancesPerChunk(
 	ECVF_ReadOnly
 );
 
-static int32 GbOptimizeVMByteCode = 0;
+static int32 GbOptimizeVMByteCode = 1;
 static FAutoConsoleVariableRef CVarbOptimizeVMByteCode(
 	TEXT("vm.OptimizeVMByteCode"),
 	GbOptimizeVMByteCode,
@@ -182,7 +182,7 @@ static FAutoConsoleVariableRef CVarbOptimizeVMByteCode(
 	ECVF_Default
 );
 
-static int32 GbFreeUnoptimizedVMByteCode = 0;
+static int32 GbFreeUnoptimizedVMByteCode = 1;
 static FAutoConsoleVariableRef CVarbFreeUnoptimizedVMByteCode(
 	TEXT("vm.FreeUnoptimizedByteCode"),
 	GbFreeUnoptimizedVMByteCode,
@@ -190,7 +190,7 @@ static FAutoConsoleVariableRef CVarbFreeUnoptimizedVMByteCode(
 	ECVF_Default
 );
 
-static int32 GbUseOptimizedVMByteCode = 0;
+static int32 GbUseOptimizedVMByteCode = 1;
 static FAutoConsoleVariableRef CVarbUseOptimizedVMByteCode(
 	TEXT("vm.UseOptimizedVMByteCode"),
 	GbUseOptimizedVMByteCode,
