Removed incorrec unregister call from scalability manager.

#rnx

diff --git a/Engine/Plugins/FX/Niagara/Source/Niagara/Private/NiagaraComponent.cpp b/Engine/Plugins/FX/Niagara/Source/Niagara/Private/NiagaraComponent.cpp
--- a/Engine/Plugins/FX/Niagara/Source/Niagara/Private/NiagaraComponent.cpp
+++ b/Engine/Plugins/FX/Niagara/Source/Niagara/Private/NiagaraComponent.cpp
@@ -1091,8 +1091,6 @@ void UNiagaraComponent::OnUnregister()
 
 	SetActiveFlag(false);
 
-	UnregisterWithScalabilityManager();
-
 	if (SystemInstance)
 	{
 		//UE_LOG(LogNiagara, Log, TEXT("UNiagaraComponent::OnUnregister: %p  %s\n"), SystemInstance.Get(), *GetAsset()->GetFullName());
