Clicking on the renderer icon on the timeline track now selects the renderer in the stack
#rnx

diff --git a/Engine/Plugins/FX/Niagara/Source/NiagaraEditor/Private/Sequencer/NiagaraSequence/NiagaraEmitterTrackEditor.cpp b/Engine/Plugins/FX/Niagara/Source/NiagaraEditor/Private/Sequencer/NiagaraSequence/NiagaraEmitterTrackEditor.cpp
--- a/Engine/Plugins/FX/Niagara/Source/NiagaraEditor/Private/Sequencer/NiagaraSequence/NiagaraEmitterTrackEditor.cpp
+++ b/Engine/Plugins/FX/Niagara/Source/NiagaraEditor/Private/Sequencer/NiagaraSequence/NiagaraEmitterTrackEditor.cpp
@@ -20,6 +20,7 @@
 #include "Widgets/Images/SImage.h"
 #include "Widgets/Input/SCheckBox.h"
 #include "Modules/ModuleManager.h"
+#include "ViewModels/Stack/NiagaraStackRendererItem.h"
 
 #define LOCTEXT_NAMESPACE "NiagaraEmitterTrackEditor"
 
@@ -78,23 +79,29 @@ public:
 			];
 
 		// Renderer buttons
-		for (UNiagaraRendererProperties* Renderer : EmitterTrack->GetEmitterHandleViewModel()->GetEmitterViewModel()->GetEmitter()->GetRenderers())
+
+		EmitterTrack->GetEmitterHandleViewModel()->GetRendererPreviewData(RendererPreviewData);
+		for (FRendererPreviewData* RendererPreview : RendererPreviewData)
 		{
-			TrackBox->AddSlot()
-				.AutoWidth()
-				.VAlign(VAlign_Center)
-				.Padding(3, 0, 0, 0)
-				[
-					SNew(SButton)
-					.ButtonStyle(FEditorStyle::Get(), "HoverHintOnly")
-					.IsFocusable(false)
-					.ToolTipText(FText::Format(LOCTEXT("RenderButtonToolTip", "{0} - Press to select."), FText::FromString(FName::NameToDisplayString(Renderer->GetName(), false))))
-					.OnClicked(this, &SEmitterTrackWidget::OnRenderButtonClicked)
+			if (UNiagaraStackRendererItem* RendererItem = Cast<UNiagaraStackRendererItem>(RendererPreview->RenderingEntry))
+			{
+				UNiagaraRendererProperties* Renderer = RendererItem->GetRendererProperties();
+				TrackBox->AddSlot()
+					.AutoWidth()
+					.VAlign(VAlign_Center)
+					.Padding(3, 0, 0, 0)
 					[
-						SNew(SImage)
-						.Image(FSlateIconFinder::FindIconBrushForClass(Renderer->GetClass()))
-					]
-				];
+						SNew(SButton)
+						.ButtonStyle(FEditorStyle::Get(), "HoverHintOnly")
+						.IsFocusable(false)
+						.ToolTipText(FText::Format(LOCTEXT("RenderButtonToolTip", "{0} - Press to select."), FText::FromString(FName::NameToDisplayString(Renderer->GetName(), false))))
+						.OnClicked(this, &SEmitterTrackWidget::OnRenderButtonClicked, RendererPreview)
+						[
+							SNew(SImage)
+							.Image(FSlateIconFinder::FindIconBrushForClass(Renderer->GetClass()))
+						]
+					];
+			}
 		}
 
 		// Enabled checkbox.
@@ -183,8 +190,15 @@ private:
 			: FLinearColor::Gray;
 	}
 
-	FReply OnRenderButtonClicked()
+	FReply OnRenderButtonClicked(FRendererPreviewData* InRendererPreview)
 	{
+		if (EmitterTrack.IsValid())
+		{
+			TArray<UNiagaraStackEntry*> SelectedEntries;
+			SelectedEntries.Add(InRendererPreview->RenderingEntry);
+			TArray<UNiagaraStackEntry*> DeselectedEntries;
+			EmitterTrack->GetSystemViewModel().GetSelectionViewModel()->UpdateSelectedEntries(SelectedEntries, DeselectedEntries, true);
+		}
 		return FReply::Handled();
 	}
 
@@ -201,6 +215,7 @@ private:
 private:
 	TWeakObjectPtr<UMovieSceneNiagaraEmitterTrack> EmitterTrack;
 	mutable TOptional<FText> TrackErrorIconToolTip;
+	TArray<FRendererPreviewData*> RendererPreviewData;
 };
 
 FNiagaraEmitterTrackEditor::FNiagaraEmitterTrackEditor(TSharedPtr<ISequencer> Sequencer) 
