Locally bump Niagara DDC key to resolve issue with some autotest script assets failing to deserialize.

diff --git a/Engine/Plugins/FX/Niagara/Source/Niagara/Private/NiagaraScriptDerivedData.h b/Engine/Plugins/FX/Niagara/Source/Niagara/Private/NiagaraScriptDerivedData.h
--- a/Engine/Plugins/FX/Niagara/Source/Niagara/Private/NiagaraScriptDerivedData.h
+++ b/Engine/Plugins/FX/Niagara/Source/Niagara/Private/NiagaraScriptDerivedData.h
@@ -54,7 +54,7 @@ public:
 		// This is a version string that mimics the old versioning scheme. If you
 		// want to bump this version, generate a new guid using VS->Tools->Create GUID and
 		// return it here. Ex.
-		return TEXT("FB23DB34150247469E8B70EBE84A30C9");
+		return TEXT("B575C44C549E4AE5A86049B662D2EFFA");
 	}
 
 	virtual FString GetPluginSpecificCacheKeySuffix() const override;
