Sequencer: Fix crash disabling all filters

diff --git a/Engine/Plugins/FX/Niagara/Source/NiagaraEditor/Private/Sequencer/NiagaraSequence/NiagaraSequencerFilters.cpp b/Engine/Plugins/FX/Niagara/Source/NiagaraEditor/Private/Sequencer/NiagaraSequence/NiagaraSequencerFilters.cpp
--- a/Engine/Plugins/FX/Niagara/Source/NiagaraEditor/Private/Sequencer/NiagaraSequence/NiagaraSequencerFilters.cpp
+++ b/Engine/Plugins/FX/Niagara/Source/NiagaraEditor/Private/Sequencer/NiagaraSequence/NiagaraSequencerFilters.cpp
@@ -21,6 +21,11 @@ class FSequencerTrackFilter_NiagaraTracks : public FSequencerTrackFilter
 	// IFilter implementation
 	virtual bool PassesFilter(FTrackFilterType InItem) const override
 	{
+		if (!InItem)
+		{
+			return false;
+		}
+
 		if (InItem->IsA(UMovieSceneNiagaraTrack::StaticClass()) || InItem->IsA(ANiagaraActor::StaticClass()) || InItem->IsA(UNiagaraComponent::StaticClass()))
 		{
 			return true;
diff --git a/Engine/Source/Editor/Sequencer/Private/SequencerTrackFilters.cpp b/Engine/Source/Editor/Sequencer/Private/SequencerTrackFilters.cpp
--- a/Engine/Source/Editor/Sequencer/Private/SequencerTrackFilters.cpp
+++ b/Engine/Source/Editor/Sequencer/Private/SequencerTrackFilters.cpp
@@ -14,6 +14,11 @@ FSequencerTrackFilter_LevelFilter::~FSequencerTrackFilter_LevelFilter()
 
 bool FSequencerTrackFilter_LevelFilter::PassesFilter(FTrackFilterType InItem) const
 {
+	if (!InItem || !InItem->GetOutermost())
+	{
+		return false;
+	}
+
 	// For anything in a level, outermost should refer to the ULevel that contains it
 	FString OutermostName = FPackageName::GetShortName(InItem->GetOutermost()->GetName());
 	
diff --git a/Engine/Source/Editor/Sequencer/Public/SequencerTrackFilterBase.h b/Engine/Source/Editor/Sequencer/Public/SequencerTrackFilterBase.h
--- a/Engine/Source/Editor/Sequencer/Public/SequencerTrackFilterBase.h
+++ b/Engine/Source/Editor/Sequencer/Public/SequencerTrackFilterBase.h
@@ -67,6 +67,11 @@ class FSequencerTrackFilter_ComponentType : public FSequencerTrackFilter
 	// IFilter implementation
 	virtual bool PassesFilter(FTrackFilterType InItem) const override
 	{
+		if (!InItem)
+		{
+			return false;
+		}
+
 		if (InItem->IsA(ComponentType::StaticClass()))
 		{
 			return true;
