Cleaned up registering niagara emitter events to resolve multithreaded access
-Registered events are now stored on the system instance (rather than globally) and cleaned up appropriately

diff --git a/Engine/Plugins/FX/Niagara/Source/Niagara/Classes/NiagaraCollision.h b/Engine/Plugins/FX/Niagara/Source/Niagara/Classes/NiagaraCollision.h
--- a/Engine/Plugins/FX/Niagara/Source/Niagara/Classes/NiagaraCollision.h
+++ b/Engine/Plugins/FX/Niagara/Source/Niagara/Classes/NiagaraCollision.h
@@ -5,6 +5,8 @@
 #include "NiagaraEvents.h"
 #include "WorldCollision.h"
 
+class FNiagaraDataSet;
+
 UENUM()
 enum class ENiagaraCollisionMode : uint8
 {
diff --git a/Engine/Plugins/FX/Niagara/Source/Niagara/Classes/NiagaraEmitterInstance.h b/Engine/Plugins/FX/Niagara/Source/Niagara/Classes/NiagaraEmitterInstance.h
--- a/Engine/Plugins/FX/Niagara/Source/Niagara/Classes/NiagaraEmitterInstance.h
+++ b/Engine/Plugins/FX/Niagara/Source/Niagara/Classes/NiagaraEmitterInstance.h
@@ -38,9 +38,6 @@ public:
 
 	void ResetSimulation(bool bKillExisting = true);
 
-	/** Called after all emitters in an System have been initialized, allows emitters to access information from one another. */
-	void PostInitSimulation();
-
 	void DirtyDataInterfaces();
 
 	/** Replaces the binding for a single parameter colleciton instance. If for example the component begins to override the global instance. */
diff --git a/Engine/Plugins/FX/Niagara/Source/Niagara/Classes/NiagaraEvents.h b/Engine/Plugins/FX/Niagara/Source/Niagara/Classes/NiagaraEvents.h
--- a/Engine/Plugins/FX/Niagara/Source/Niagara/Classes/NiagaraEvents.h
+++ b/Engine/Plugins/FX/Niagara/Source/Niagara/Classes/NiagaraEvents.h
@@ -10,7 +10,6 @@ NiagaraEmitterInstance.h: Niagara emitter simulation class
 #include "UObject/ObjectMacros.h"
 #include "UObject/Object.h"
 #include "NiagaraCommon.h"
-#include "NiagaraDataSet.h"
 #include "NiagaraEvents.generated.h"
 
 struct FNiagaraEventReceiverProperties;
@@ -67,126 +66,3 @@ public:
 
 	virtual void PerformAction(FNiagaraEmitterInstance& OwningSim, const FNiagaraEventReceiverProperties& OwningEventReceiver)override;
 };
-
-
-
-typedef TMap<FName, FNiagaraDataSet*> PerEmitterEventDataSetMap;
-typedef TMap<FName, PerEmitterEventDataSetMap> PerSystemInstanceDataSetMap;
-
-class FNiagaraEventDataSetMgr
-{
-public:
-	FNiagaraEventDataSetMgr()
-	{
-	}
-
-	// call this to allocate a data set for a specific emitter and event
-	// if the data set already exists, returns the existing one
-	static FNiagaraDataSet *CreateEventDataSet(FNiagaraSystemInstanceID OwnerSystemInstanceID, FName EmitterName, FName EventName)
-	{
-		PerSystemInstanceDataSetMap *SystemInstanceMap = EmitterEventDataSets.Find(OwnerSystemInstanceID);
-		if (!SystemInstanceMap)
-		{
-			SystemInstanceMap = &EmitterEventDataSets.Add(OwnerSystemInstanceID);
-		}
-
-		// find the emitter's map
-		PerEmitterEventDataSetMap *Map = SystemInstanceMap->Find(EmitterName);
-		if (!Map)
-		{
-			Map = &SystemInstanceMap->Add(EmitterName);
-		}
-
-		// TODO: find a better way of multiple events trying to write to the same data set; 
-		// for example, if two analytical collision primitives want to send collision events, they need to push to the same data set
-
-		FNiagaraDataSet **DataSetPtr = Map->Find(EventName);
-		FNiagaraDataSet *OutSet = nullptr;
-		if (DataSetPtr == nullptr)
-		{
-			OutSet = new FNiagaraDataSet();
-			Map->Add(EventName) = OutSet;
-		}
-		else
-		{
-			OutSet = *DataSetPtr;
-		}
-		return OutSet;
-	}
-
-	// deletes and cleans up all event data sets for an emitter
-	//  should be called when the emitter is destroyed
-	static void Reset(FNiagaraSystemInstanceID OwnerSystemInstanceID, FName EmitterName)
-	{
-		PerSystemInstanceDataSetMap *SystemInstanceMap = EmitterEventDataSets.Find(OwnerSystemInstanceID);
-		if (SystemInstanceMap)
-		{
-			PerEmitterEventDataSetMap *Map = SystemInstanceMap->Find(EmitterName);
-			if (Map)
-			{
-				for (TPair<FName, FNiagaraDataSet *> Pair : *Map)
-				{
-					delete Pair.Value;
-				}
-				Map->Empty();
-				SystemInstanceMap->Remove(EmitterName);
-			}
-			if (SystemInstanceMap->Num() == 0)
-			{
-				EmitterEventDataSets.Remove(OwnerSystemInstanceID);
-			}
-		}
-	}
-
-	/*
-	static void AddEventSet(FName EmitterName, FName EventName, const FNiagaraDataSet *DataSet)
-	{
-		// find the emitter's map
-		PerEmitterEventDataSetMap *Map = EmitterEventDataSets.Find(EmitterName);
-		if(!Map)
-		{
-			Map = &EmitterEventDataSets.Add(EmitterName);
-		}
-
-		// TODO: find a better way of multiple events trying to write to the same data set; 
-		// for example, if two analytical collision primitives want to send collision events, they need to push to the same data set
-		ensure(Map->Find(EventName) == nullptr);
-
-		if (Map->Find(EventName) == nullptr)
-		{
-			Map->Add(EventName) = DataSet;
-		}
-	}
-	*/
-
-	static PerEmitterEventDataSetMap * GetEmitterMap(FNiagaraSystemInstanceID OwnerSystemInstanceID, FName EmitterName)
-	{
-		TMap<FName, PerEmitterEventDataSetMap> *SystemMap = EmitterEventDataSets.Find(OwnerSystemInstanceID);
-		if (SystemMap)
-		{
-			return SystemMap->Find(EmitterName);
-		}
-
-		return nullptr;
-	}
-
-	static FNiagaraDataSet * GetEventDataSet(FNiagaraSystemInstanceID OwnerSystemInstanceID, FName EmitterName, FName EventName)
-	{
-		PerEmitterEventDataSetMap *SetMap = GetEmitterMap(OwnerSystemInstanceID, EmitterName);
-		if (!SetMap)
-		{
-			return nullptr;
-		}
-
-		FNiagaraDataSet **FoundSet = SetMap->Find(EventName);
-		if (!FoundSet)
-		{
-			return nullptr;
-		}
-
-		return *FoundSet;
-	}
-
-private:
-	static TMap<FNiagaraSystemInstanceID, TMap<FName, PerEmitterEventDataSetMap>> EmitterEventDataSets;
-};
\ No newline at end of file
diff --git a/Engine/Plugins/FX/Niagara/Source/Niagara/Private/NiagaraEmitterInstance.cpp b/Engine/Plugins/FX/Niagara/Source/Niagara/Private/NiagaraEmitterInstance.cpp
--- a/Engine/Plugins/FX/Niagara/Source/Niagara/Private/NiagaraEmitterInstance.cpp
+++ b/Engine/Plugins/FX/Niagara/Source/Niagara/Private/NiagaraEmitterInstance.cpp
@@ -330,7 +330,7 @@ void FNiagaraEmitterInstance::Init(int32 InEmitterIdx, FNiagaraSystemInstanceID
 		int32 UpdateEventGeneratorIndex = 0;
 		for (const FNiagaraEventGeneratorProperties &GeneratorProps : CachedEmitter->UpdateScriptProps.EventGenerators)
 		{
-			FNiagaraDataSet *Set = FNiagaraEventDataSetMgr::CreateEventDataSet(ParentSystemInstance->GetId(), EmitterHandle.GetIdName(), GeneratorProps.ID);
+			FNiagaraDataSet *Set = ParentSystemInstance->CreateEventDataSet(EmitterHandle.GetIdName(), GeneratorProps.ID);
 			Set->Init(&GeneratorProps.DataSetCompiledData);
 
 			UpdateScriptEventDataSets.Add(Set);
@@ -344,7 +344,7 @@ void FNiagaraEmitterInstance::Init(int32 InEmitterIdx, FNiagaraSystemInstanceID
 		int32 SpawnEventGeneratorIndex = 0;
 		for (const FNiagaraEventGeneratorProperties &GeneratorProps : CachedEmitter->SpawnScriptProps.EventGenerators)
 		{
-			FNiagaraDataSet *Set = FNiagaraEventDataSetMgr::CreateEventDataSet(ParentSystemInstance->GetId(), EmitterHandle.GetIdName(), GeneratorProps.ID);
+			FNiagaraDataSet *Set = ParentSystemInstance->CreateEventDataSet(EmitterHandle.GetIdName(), GeneratorProps.ID);
 
 			Set->Init(&GeneratorProps.DataSetCompiledData);
 
@@ -754,28 +754,6 @@ void FNiagaraEmitterInstance::BindParameters(bool bExternalOnly)
 	}
 }
 
-void FNiagaraEmitterInstance::PostInitSimulation()
-{
-	if (!IsDisabled())
-	{
-		check(ParentSystemInstance);
-
-		//Go through all our receivers and grab their generator sets so that the source emitters can do any init work they need to do.
-		for (const FNiagaraEventReceiverProperties& Receiver : CachedEmitter->SpawnScriptProps.EventReceivers)
-		{
-			//FNiagaraDataSet* ReceiverSet = ParentSystemInstance->GetDataSet(FNiagaraDataSetID(Receiver.SourceEventGenerator, ENiagaraDataSetType::Event), Receiver.SourceEmitter);
-			const FNiagaraDataSet* ReceiverSet = FNiagaraEventDataSetMgr::GetEventDataSet(ParentSystemInstance->GetId(), Receiver.SourceEmitter, Receiver.SourceEventGenerator);
-
-		}
-
-		for (const FNiagaraEventReceiverProperties& Receiver : CachedEmitter->UpdateScriptProps.EventReceivers)
-		{
-			//FNiagaraDataSet* ReceiverSet = ParentSystemInstance->GetDataSet(FNiagaraDataSetID(Receiver.SourceEventGenerator, ENiagaraDataSetType::Event), Receiver.SourceEmitter);
-			const FNiagaraDataSet* ReceiverSet = FNiagaraEventDataSetMgr::GetEventDataSet(ParentSystemInstance->GetId(), Receiver.SourceEmitter, Receiver.SourceEventGenerator);
-		}
-	}
-}
-
 FNiagaraDataSet* FNiagaraEmitterInstance::GetDataSet(FNiagaraDataSetID SetID)
 {
 	FNiagaraDataSet** SetPtr = DataSetMap.Find(SetID);
@@ -1059,9 +1037,9 @@ void FNiagaraEmitterInstance::PreTick()
 	if (NumEventHandlers > 0)
 	{
 		EventHandlingInfo.Reset();
-		EventHandlingInfo.SetNum(CachedEmitter->GetEventHandlers().Num());
+		EventHandlingInfo.SetNum(NumEventHandlers);
 		EventSpawnTotal = 0;
-		for (int32 i = 0; i < CachedEmitter->GetEventHandlers().Num(); i++)
+		for (int32 i = 0; i < NumEventHandlers; i++)
 		{
 			const FNiagaraEventScriptProperties &EventHandlerProps = CachedEmitter->GetEventHandlers()[i];
 			FNiagaraEventHandlingInfo& Info = EventHandlingInfo[i];
@@ -1070,7 +1048,7 @@ void FNiagaraEmitterInstance::PreTick()
 			Info.SpawnCounts.Reset();
 			Info.TotalSpawnCount = 0;
 			Info.EventData = nullptr;
-			if (FNiagaraDataSet* EventSet = FNiagaraEventDataSetMgr::GetEventDataSet(ParentSystemInstance->GetId(), Info.SourceEmitterName, EventHandlerProps.SourceEventName))
+			if (FNiagaraDataSet* EventSet = ParentSystemInstance->GetEventDataSet(Info.SourceEmitterName, EventHandlerProps.SourceEventName))
 			{
 				Info.SetEventData(&EventSet->GetCurrentDataChecked());
 				uint32 EventSpawnNum = CalculateEventSpawnCount(EventHandlerProps, Info.SpawnCounts, EventSet);
diff --git a/Engine/Plugins/FX/Niagara/Source/Niagara/Private/NiagaraEvents.cpp b/Engine/Plugins/FX/Niagara/Source/Niagara/Private/NiagaraEvents.cpp
--- a/Engine/Plugins/FX/Niagara/Source/Niagara/Private/NiagaraEvents.cpp
+++ b/Engine/Plugins/FX/Niagara/Source/Niagara/Private/NiagaraEvents.cpp
@@ -8,10 +8,6 @@
 DECLARE_DWORD_COUNTER_STAT(TEXT("Num Death Events"), STAT_NiagaraNumDeathEvents, STATGROUP_Niagara);
 DECLARE_DWORD_COUNTER_STAT(TEXT("Num Spawn Events"), STAT_NiagaraNumSpawnEvents, STATGROUP_Niagara);
 
-
-TMap<FNiagaraSystemInstanceID, TMap<FName, PerEmitterEventDataSetMap>> FNiagaraEventDataSetMgr::EmitterEventDataSets;
-
-
 void UNiagaraEventReceiverEmitterAction_SpawnParticles::PerformAction(FNiagaraEmitterInstance& OwningSim, const FNiagaraEventReceiverProperties& OwningEventReceiver)
 {
 	//Find the generator set we're bound to and spawn NumParticles particles for every event it generated last frame.
diff --git a/Engine/Plugins/FX/Niagara/Source/Niagara/Private/NiagaraSystemInstance.cpp b/Engine/Plugins/FX/Niagara/Source/Niagara/Private/NiagaraSystemInstance.cpp
--- a/Engine/Plugins/FX/Niagara/Source/Niagara/Private/NiagaraSystemInstance.cpp
+++ b/Engine/Plugins/FX/Niagara/Source/Niagara/Private/NiagaraSystemInstance.cpp
@@ -989,6 +989,9 @@ void FNiagaraSystemInstance::Cleanup()
 
 	// Clear out the emitters.
 	Emitters.Empty(0);
+
+	// clean up any event datasets that we're holding onto for our child emitters
+	ClearEventDataSets();
 }
 
 //Unsure on usage of this atm. Possibly useful in future.
@@ -1568,6 +1571,40 @@ void FNiagaraSystemInstance::ResetFastPathBindings()
 	FastPathFloatUpdateRangedInputBindings.Empty();
 }
 
+void
+FNiagaraSystemInstance::ClearEventDataSets()
+{
+	for (auto& EventDataSetIt : EmitterEventDataSetMap)
+	{
+		delete EventDataSetIt.Value;
+	}
+
+	EmitterEventDataSetMap.Empty();
+}
+
+FNiagaraDataSet*
+FNiagaraSystemInstance::CreateEventDataSet(FName EmitterName, FName EventName)
+{
+	// TODO: find a better way of multiple events trying to write to the same data set; 
+	// for example, if two analytical collision primitives want to send collision events, they need to push to the same data set
+	FNiagaraDataSet*& OutSet = EmitterEventDataSetMap.FindOrAdd(EmitterEventKey(EmitterName, EventName));
+
+	if (!OutSet)
+	{
+		OutSet = new FNiagaraDataSet();
+	}
+
+	return OutSet;
+}
+
+FNiagaraDataSet*
+FNiagaraSystemInstance::GetEventDataSet(FName EmitterName, FName EventName) const
+{
+	FNiagaraDataSet* const* OutDataSet = EmitterEventDataSetMap.Find(EmitterEventKey(EmitterName, EventName));
+
+	return OutDataSet ? *OutDataSet : nullptr;
+}
+
 #if WITH_EDITORONLY_DATA
 
 bool FNiagaraSystemInstance::UsesEmitter(const UNiagaraEmitter* Emitter)const
@@ -1644,11 +1681,9 @@ void FNiagaraSystemInstance::InitEmitters()
 
 		for (TSharedRef<FNiagaraEmitterInstance, ESPMode::ThreadSafe>& Simulation : Emitters)
 		{
-			FNiagaraEmitterInstance& Inst = Simulation.Get();
-			if (Inst.GetCachedEmitter() != nullptr)
+			if (const UNiagaraEmitter* Emitter = Simulation->GetCachedEmitter())
 			{
-				bHasGPUEmitters |= Inst.GetCachedEmitter()->SimTarget == ENiagaraSimTarget::GPUComputeSim;
-				Inst.PostInitSimulation();
+				bHasGPUEmitters |= Emitter->SimTarget == ENiagaraSimTarget::GPUComputeSim;
 			}
 		}
 
diff --git a/Engine/Plugins/FX/Niagara/Source/Niagara/Public/NiagaraSystemInstance.h b/Engine/Plugins/FX/Niagara/Source/Niagara/Public/NiagaraSystemInstance.h
--- a/Engine/Plugins/FX/Niagara/Source/Niagara/Public/NiagaraSystemInstance.h
+++ b/Engine/Plugins/FX/Niagara/Source/Niagara/Public/NiagaraSystemInstance.h
@@ -251,6 +251,10 @@ public:
 
 	void ResetFastPathBindings();
 
+	FNiagaraDataSet* CreateEventDataSet(FName EmitterName, FName EventName);
+	FNiagaraDataSet* GetEventDataSet(FName EmitterName, FName EventName) const;
+	void ClearEventDataSets();
+
 	FNiagaraSystemFastPath::FParamMap0& GetFastPathMap() { return FastPathMap; }
 
 	FORCEINLINE void SetLODDistance(float InLODDistance, float InMaxLODDistance);
@@ -362,6 +366,11 @@ public:
 	TArray<FNiagaraParameterDirectBinding<int32>> ParameterNumParticleBindings;
 	TArray<FNiagaraParameterDirectBinding<int32>> ParameterTotalSpawnedParticlesBindings;
 
+	// registered events for each of the emitters
+	typedef TPair<FName, FName> EmitterEventKey;
+	typedef TMap<EmitterEventKey, FNiagaraDataSet*> EventDataSetMap;
+	EventDataSetMap EmitterEventDataSetMap;
+
 	/** Indicates whether this instance must update itself rather than being batched up as most instances are. */
 	uint32 bSolo : 1;
 	uint32 bForceSolo : 1;
