Fixed duplicate location key

diff --git a/Engine/Plugins/FX/Niagara/Source/NiagaraEditor/Private/NiagaraNodeOp.cpp b/Engine/Plugins/FX/Niagara/Source/NiagaraEditor/Private/NiagaraNodeOp.cpp
--- a/Engine/Plugins/FX/Niagara/Source/NiagaraEditor/Private/NiagaraNodeOp.cpp
+++ b/Engine/Plugins/FX/Niagara/Source/NiagaraEditor/Private/NiagaraNodeOp.cpp
@@ -136,7 +136,7 @@ void UNiagaraNodeOp::Compile(class FHlslNiagaraTranslator* Translator, TArray<in
 				bError = true;
 				FFormatNamedArguments Args;
 				Args.Add(TEXT("OpName"), GetNodeTitle(ENodeTitleType::FullTitle));
-				FText Format = LOCTEXT("InputTypeErrorFormat", "The {OpName} node cannot have float based numeric input pins.");
+				FText Format = LOCTEXT("InputTypeErrorFormatFloat", "The {OpName} node cannot have float based numeric input pins.");
 				Translator->Error(FText::Format(Format, Args), this, Pin);
 			}
 			else if (!TypeDef.IsFloatPrimitive() && !OpInfo->bNumericsCanBeIntegers)
@@ -144,7 +144,7 @@ void UNiagaraNodeOp::Compile(class FHlslNiagaraTranslator* Translator, TArray<in
 				bError = true;
 				FFormatNamedArguments Args;
 				Args.Add(TEXT("OpName"), GetNodeTitle(ENodeTitleType::FullTitle));
-				FText Format = LOCTEXT("InputTypeErrorFormat", "The {OpName} node cannot have integer based numeric input pins.");
+				FText Format = LOCTEXT("InputTypeErrorFormatInt", "The {OpName} node cannot have integer based numeric input pins.");
 				Translator->Error(FText::Format(Format, Args), this, Pin);
 			}
 		}
