Removed const modifier from callback handler interface function

diff --git a/Engine/Plugins/FX/Niagara/Source/Niagara/Classes/NiagaraDataInterfaceExport.h b/Engine/Plugins/FX/Niagara/Source/Niagara/Classes/NiagaraDataInterfaceExport.h
--- a/Engine/Plugins/FX/Niagara/Source/Niagara/Classes/NiagaraDataInterfaceExport.h
+++ b/Engine/Plugins/FX/Niagara/Source/Niagara/Classes/NiagaraDataInterfaceExport.h
@@ -35,7 +35,7 @@ class INiagaraParticleCallbackHandler
 public:
 	/** This function is called once per tick with the gathered particle data. It will not be called if there is no particle data to call it with. */
 	UFUNCTION(BlueprintCallable, BlueprintNativeEvent, Category = "Niagara")
-	void ReceiveParticleData(const TArray<FBasicParticleData>& Data, UNiagaraSystem* NiagaraSystem) const;
+	void ReceiveParticleData(const TArray<FBasicParticleData>& Data, UNiagaraSystem* NiagaraSystem);
 };
 
 struct ExportInterface_InstanceData
