* Fix frame latency on gpu with camera DI
* Separated the occlusion from the camera DI into a separate DI

#tests niagara automation tests passed

diff --git a/Engine/Plugins/FX/Niagara/Source/Niagara/Classes/NiagaraDataInterface.h b/Engine/Plugins/FX/Niagara/Source/Niagara/Classes/NiagaraDataInterface.h
--- a/Engine/Plugins/FX/Niagara/Source/Niagara/Classes/NiagaraDataInterface.h
+++ b/Engine/Plugins/FX/Niagara/Source/Niagara/Classes/NiagaraDataInterface.h
@@ -262,6 +262,8 @@ public:
 	virtual bool CanExecuteOnTarget(ENiagaraSimTarget Target)const { return false; }
 
 	virtual bool RequiresDistanceFieldData() const { return false; }
+	virtual bool RequiresDepthBuffer() const { return false; }
+	virtual bool RequiresEarlyViewData() const { return false; }
 
 	virtual bool HasTickGroupPrereqs() const { return false; }
 	virtual ETickingGroup CalculateTickGroup(void* PerInstanceData) const { return NiagaraFirstTickGroup; }
diff --git a/Engine/Plugins/FX/Niagara/Source/Niagara/Classes/NiagaraDataInterfaceCamera.h b/Engine/Plugins/FX/Niagara/Source/Niagara/Classes/NiagaraDataInterfaceCamera.h
--- a/Engine/Plugins/FX/Niagara/Source/Niagara/Classes/NiagaraDataInterfaceCamera.h
+++ b/Engine/Plugins/FX/Niagara/Source/Niagara/Classes/NiagaraDataInterfaceCamera.h
@@ -4,7 +4,6 @@
 #include "NiagaraCommon.h"
 #include "NiagaraShared.h"
 #include "NiagaraDataInterface.h"
-#include "NiagaraParameterStore.h"
 #include "Camera/PlayerCameraManager.h"
 #include "NiagaraDataInterfaceCamera.generated.h"
 
@@ -37,13 +36,11 @@ public:
 	virtual bool PerInstanceTick(void* PerInstanceData, FNiagaraSystemInstance* SystemInstance, float DeltaSeconds) override;
 	virtual bool GetFunctionHLSL(const FName&  DefinitionFunctionName, FString InstanceFunctionName, FNiagaraDataInterfaceGPUParamInfo& ParamInfo, FString& OutHLSL) override;
 	virtual bool CanExecuteOnTarget(ENiagaraSimTarget Target) const override { return true; }
-	virtual FNiagaraDataInterfaceParametersCS* ConstructComputeParameters() const override;
 	virtual bool HasTickGroupPrereqs() const override { return true; }
 	virtual ETickingGroup CalculateTickGroup(void* PerInstanceData) const override;
+	virtual bool RequiresEarlyViewData() const override { return true; }
 	//UNiagaraDataInterface Interface
 
-	void QueryOcclusionFactorGPU(FVectorVMContext& Context);
-	void QueryOcclusionFactorCircleGPU(FVectorVMContext& Context);
 	void GetCameraFOV(FVectorVMContext& Context);
 	void GetCameraProperties(FVectorVMContext& Context);
 	void GetViewPropertiesGPU(FVectorVMContext& Context);
@@ -51,8 +48,6 @@ public:
 	void GetViewSpaceTransformsGPU(FVectorVMContext& Context);
 	
 private:
-	static const FName GetCameraOcclusionRectangleName;
-	static const FName GetCameraOcclusionCircleName;
 	static const FName GetViewPropertiesName;
 	static const FName GetClipSpaceTransformsName;
 	static const FName GetViewSpaceTransformsName;
@@ -62,8 +57,7 @@ private:
 
 struct FNiagaraDataIntefaceProxyCameraQuery : public FNiagaraDataInterfaceProxy
 {
-	// There's nothing in this proxy. It just reads from scene textures.
-
+	// There's nothing in this proxy.
 	virtual int32 PerInstanceDataPassedToRenderThreadSize() const override
 	{
 		return 0;
diff --git a/Engine/Plugins/FX/Niagara/Source/Niagara/Classes/NiagaraDataInterfaceCollisionQuery.h b/Engine/Plugins/FX/Niagara/Source/Niagara/Classes/NiagaraDataInterfaceCollisionQuery.h
--- a/Engine/Plugins/FX/Niagara/Source/Niagara/Classes/NiagaraDataInterfaceCollisionQuery.h
+++ b/Engine/Plugins/FX/Niagara/Source/Niagara/Classes/NiagaraDataInterfaceCollisionQuery.h
@@ -62,6 +62,7 @@ public:
 
 	virtual bool CanExecuteOnTarget(ENiagaraSimTarget Target) const override { return true; }
 	virtual bool RequiresDistanceFieldData() const override { return true; }
+	virtual bool RequiresDepthBuffer() const override { return true; }
 
 	virtual bool GetFunctionHLSL(const FName&  DefinitionFunctionName, FString InstanceFunctionName, FNiagaraDataInterfaceGPUParamInfo& ParamInfo, FString& OutHLSL) override;
 	virtual void GetParameterDefinitionHLSL(FNiagaraDataInterfaceGPUParamInfo& ParamInfo, FString& OutHLSL) override;
diff --git a/Engine/Plugins/FX/Niagara/Source/Niagara/Classes/NiagaraDataInterfaceOcclusion.h b/Engine/Plugins/FX/Niagara/Source/Niagara/Classes/NiagaraDataInterfaceOcclusion.h
new file mode 100644
--- /dev/null
+++ b/Engine/Plugins/FX/Niagara/Source/Niagara/Classes/NiagaraDataInterfaceOcclusion.h
@@ -0,0 +1,44 @@
+// Copyright 1998-2019 Epic Games, Inc. All Rights Reserved.
+#pragma once
+
+#include "NiagaraCommon.h"
+#include "NiagaraShared.h"
+#include "NiagaraDataInterface.h"
+#include "NiagaraDataInterfaceOcclusion.generated.h"
+
+UCLASS(EditInlineNew, Category = "Camera", meta = (DisplayName = "Occlusion Query"))
+class NIAGARA_API UNiagaraDataInterfaceOcclusion : public UNiagaraDataInterface
+{
+	GENERATED_UCLASS_BODY()
+
+public:
+	//UObject Interface
+	virtual void PostInitProperties() override;
+	//UObject Interface End
+
+	//UNiagaraDataInterface Interface
+	virtual void GetFunctions(TArray<FNiagaraFunctionSignature>& OutFunctions)override;
+	virtual void GetVMExternalFunction(const FVMExternalFunctionBindingInfo& BindingInfo, void* InstanceData, FVMExternalFunction &OutFunc) override;
+	virtual bool GetFunctionHLSL(const FName&  DefinitionFunctionName, FString InstanceFunctionName, FNiagaraDataInterfaceGPUParamInfo& ParamInfo, FString& OutHLSL) override;
+	virtual bool CanExecuteOnTarget(ENiagaraSimTarget Target) const override { return Target == ENiagaraSimTarget::GPUComputeSim; }
+	virtual FNiagaraDataInterfaceParametersCS* ConstructComputeParameters() const override;
+	virtual bool RequiresDepthBuffer() const override { return true; }
+	//UNiagaraDataInterface Interface
+
+	void QueryOcclusionFactorGPU(FVectorVMContext& Context);
+	void QueryOcclusionFactorCircleGPU(FVectorVMContext& Context);
+	
+private:
+	static const FName GetCameraOcclusionRectangleName;
+	static const FName GetCameraOcclusionCircleName;
+};
+
+struct FNiagaraDataIntefaceProxyOcclusionQuery : public FNiagaraDataInterfaceProxy
+{
+	// There's nothing in this proxy. It just reads from scene textures.
+
+	virtual int32 PerInstanceDataPassedToRenderThreadSize() const override
+	{
+		return 0;
+	}
+};
\ No newline at end of file
diff --git a/Engine/Plugins/FX/Niagara/Source/Niagara/Classes/NiagaraEmitterInstanceBatcher.h b/Engine/Plugins/FX/Niagara/Source/Niagara/Classes/NiagaraEmitterInstanceBatcher.h
--- a/Engine/Plugins/FX/Niagara/Source/Niagara/Classes/NiagaraEmitterInstanceBatcher.h
+++ b/Engine/Plugins/FX/Niagara/Source/Niagara/Classes/NiagaraEmitterInstanceBatcher.h
@@ -38,6 +38,13 @@ public:
 	int32 UsedIndexCount = 0;
 };
 
+enum class ETickStage
+{
+	PreInitViews,
+	PostInitViews,
+	PostOpaqueRender
+};
+
 class NiagaraEmitterInstanceBatcher : public FFXSystemInterface
 {
 public:
@@ -79,7 +86,10 @@ public:
 	virtual void RemoveVectorField(UVectorFieldComponent* VectorFieldComponent) override {}
 	virtual void UpdateVectorField(UVectorFieldComponent* VectorFieldComponent) override {}
 	virtual void PreInitViews(FRHICommandListImmediate& RHICmdList, bool bAllowGPUParticleUpdate) override;
+	virtual void PostInitViews(FRHICommandListImmediate& RHICmdList, FRHIUniformBuffer* ViewUniformBuffer, bool bAllowGPUParticleUpdate) override;
 	virtual bool UsesGlobalDistanceField() const override;
+	virtual bool UsesDepthBuffer() const override;
+	virtual bool RequiresEarlyViewUniformBuffer() const override;
 	virtual void PreRender(FRHICommandListImmediate& RHICmdList, const class FGlobalDistanceFieldParameterData* GlobalDistanceFieldParameterData, bool bAllowGPUParticleUpdate) override;
 	virtual void OnDestroy() override; // Called on the gamethread to delete the batcher on the renderthread.
 
@@ -155,10 +165,12 @@ public:
 	void DispatchMultipleStages(const FNiagaraGPUSystemTick& Tick, FNiagaraComputeInstanceData *Instance, FRHICommandList &RHICmdList, FRHIUniformBuffer* ViewUniformBuffer, FNiagaraShader* ComputeShader) const;
 
 private:
-	void ExecuteAll(FRHICommandList &RHICmdList, FRHIUniformBuffer* ViewUniformBuffer, bool bSetReadback);
+	void ExecuteAll(FRHICommandList &RHICmdList, FRHIUniformBuffer* ViewUniformBuffer, bool bSetReadback, ETickStage TickStage);
 	void ResizeBuffersAndGatherResources(FOverlappableTicks& OverlappableTick, FRHICommandList& RHICmdList, FNiagaraBufferArray& DestDataBuffers, FNiagaraBufferArray& CurrDataBuffers, FNiagaraBufferArray& DestBufferIntFloat, FNiagaraBufferArray& CurrBufferIntFloat);
 	void DispatchAllOnCompute(FOverlappableTicks& OverlappableTick, FRHICommandList& RHICmdList, FRHIUniformBuffer* ViewUniformBuffer, FNiagaraBufferArray& DestDataBuffers, FNiagaraBufferArray& CurrDataBuffers, FNiagaraBufferArray& DestBufferIntFloat, FNiagaraBufferArray& CurrBufferIntFloat, bool bSetReadback);
 
+	bool ShouldTickForStage(const FNiagaraGPUSystemTick& Tick, ETickStage TickStage) const;
+
 //	void SimStepReadback(const FNiagaraComputeInstanceData& Instance, FRHICommandList& RHICmdList) const;
 
 	inline uint32 UnpackEmitterDispatchCount(uint8* PackedData)
diff --git a/Engine/Plugins/FX/Niagara/Source/Niagara/Classes/NiagaraScriptExecutionContext.h b/Engine/Plugins/FX/Niagara/Source/Niagara/Classes/NiagaraScriptExecutionContext.h
--- a/Engine/Plugins/FX/Niagara/Source/Niagara/Classes/NiagaraScriptExecutionContext.h
+++ b/Engine/Plugins/FX/Niagara/Source/Niagara/Classes/NiagaraScriptExecutionContext.h
@@ -313,7 +313,9 @@ public:
 	FNiagaraSystemInstanceID SystemInstanceID;
 	FNiagaraDataInterfaceInstanceData* DIInstanceData;
 	uint8* InstanceData_ParamData_Packed;
-	bool bRequiredDistanceFieldData = false;
+	bool bRequiresDistanceFieldData = false;
+	bool bRequiresDepthBuffer = false;
+	bool bRequiresEarlyViewData = false;
 	bool bNeedsReset = false;
 	bool bIsFinalTick = false;
 };
diff --git a/Engine/Plugins/FX/Niagara/Source/Niagara/Private/NiagaraDataInterfaceCamera.cpp b/Engine/Plugins/FX/Niagara/Source/Niagara/Private/NiagaraDataInterfaceCamera.cpp
--- a/Engine/Plugins/FX/Niagara/Source/Niagara/Private/NiagaraDataInterfaceCamera.cpp
+++ b/Engine/Plugins/FX/Niagara/Source/Niagara/Private/NiagaraDataInterfaceCamera.cpp
@@ -2,15 +2,11 @@
 
 #include "NiagaraDataInterfaceCamera.h"
 #include "NiagaraTypes.h"
-#include "NiagaraCustomVersion.h"
 #include "NiagaraWorldManager.h"
-#include "ShaderParameterUtils.h"
 #include "Internationalization/Internationalization.h"
 #include "NiagaraSystemInstance.h"
 #include "GameFramework/PlayerController.h"
 
-const FName UNiagaraDataInterfaceCamera::GetCameraOcclusionRectangleName(TEXT("QueryOcclusionFactorWithRectangleGPU"));
-const FName UNiagaraDataInterfaceCamera::GetCameraOcclusionCircleName(TEXT("QueryOcclusionFactorWithCircleGPU"));
 const FName UNiagaraDataInterfaceCamera::GetViewPropertiesName(TEXT("GetViewPropertiesGPU"));
 const FName UNiagaraDataInterfaceCamera::GetClipSpaceTransformsName(TEXT("GetClipSpaceTransformsGPU"));
 const FName UNiagaraDataInterfaceCamera::GetViewSpaceTransformsName(TEXT("GetViewSpaceTransformsGPU"));
@@ -74,41 +70,6 @@ bool UNiagaraDataInterfaceCamera::PerInstanceTick(void* PerInstanceData, FNiagar
 void UNiagaraDataInterfaceCamera::GetFunctions(TArray<FNiagaraFunctionSignature>& OutFunctions)
 {
 	FNiagaraFunctionSignature Sig;
-	Sig.Name = GetCameraOcclusionRectangleName;
-#if WITH_EDITORONLY_DATA
-	Sig.Description = NSLOCTEXT("Niagara", "GetCameraOcclusionRectFunctionDescription", "This function returns the occlusion factor of a sprite. It samples the depth buffer in a rectangular grid around the given world position and compares each sample with the camera distance.");
-#endif
-	Sig.bMemberFunction = true;
-	Sig.bRequiresContext = false;
-	Sig.bSupportsCPU = false;
-	Sig.Inputs.Add(FNiagaraVariable(FNiagaraTypeDefinition(GetClass()), TEXT("Camera interface")));
-	Sig.Inputs.Add(FNiagaraVariable(FNiagaraTypeDefinition::GetVec3Def(), TEXT("Sample Center World Position")));
-	Sig.Inputs.Add(FNiagaraVariable(FNiagaraTypeDefinition::GetFloatDef(), TEXT("Sample Window Width World")));
-	Sig.Inputs.Add(FNiagaraVariable(FNiagaraTypeDefinition::GetFloatDef(), TEXT("Sample Window Height World")));
-	Sig.Inputs.Add(FNiagaraVariable(FNiagaraTypeDefinition::GetFloatDef(), TEXT("Sample Steps Per Line")));
-	Sig.Outputs.Add(FNiagaraVariable(FNiagaraTypeDefinition::GetFloatDef(), TEXT("Visibility Fraction")));
-	Sig.Outputs.Add(FNiagaraVariable(FNiagaraTypeDefinition::GetFloatDef(), TEXT("Sample Fraction")));
-	OutFunctions.Add(Sig);
-
-	Sig = FNiagaraFunctionSignature();
-	Sig.Name = GetCameraOcclusionCircleName;
-#if WITH_EDITORONLY_DATA
-	Sig.Description = NSLOCTEXT("Niagara", "GetCameraOcclusionCircleFunctionDescription", "This function returns the occlusion factor of a sprite. It samples the depth buffer in concentric rings around the given world position and compares each sample with the camera distance.");
-#endif
-	Sig.bMemberFunction = true;
-	Sig.bRequiresContext = false;
-	Sig.bSupportsCPU = false;
-	Sig.Inputs.Add(FNiagaraVariable(FNiagaraTypeDefinition(GetClass()), TEXT("Camera interface")));
-	Sig.Inputs.Add(FNiagaraVariable(FNiagaraTypeDefinition::GetVec3Def(), TEXT("Sample Center World Position")));
-	Sig.Inputs.Add(FNiagaraVariable(FNiagaraTypeDefinition::GetFloatDef(), TEXT("Sample Window Diameter World")));
-	Sig.Inputs.Add(FNiagaraVariable(FNiagaraTypeDefinition::GetFloatDef(), TEXT("Samples per ring")));
-	Sig.Inputs.Add(FNiagaraVariable(FNiagaraTypeDefinition::GetFloatDef(), TEXT("Number of sample rings")));
-	Sig.Outputs.Add(FNiagaraVariable(FNiagaraTypeDefinition::GetFloatDef(), TEXT("Visibility Fraction")));
-	Sig.Outputs.Add(FNiagaraVariable(FNiagaraTypeDefinition::GetFloatDef(), TEXT("Sample Fraction")));
-	OutFunctions.Add(Sig);
-
-
-	Sig = FNiagaraFunctionSignature();
 	Sig.Name = GetViewPropertiesName;
 #if WITH_EDITORONLY_DATA
 	Sig.Description = NSLOCTEXT("Niagara", "GetViewPropertiesDescription", "This function returns the properties of the current view. Only valid for gpu particles.");
@@ -196,133 +157,6 @@ bool UNiagaraDataInterfaceCamera::GetFunctionHLSL(const FName& DefinitionFunctio
 	TMap<FString, FStringFormatArg> ArgsSample;
 	ArgsSample.Add(TEXT("FunctionName"), InstanceFunctionName);
 
-	if (DefinitionFunctionName == GetCameraOcclusionRectangleName)
-	{
-		static const TCHAR *FormatSample = TEXT(R"(
-			void {FunctionName}(in float3 In_SampleCenterWorldPos, in float In_SampleWindowWidthWorld, in float In_SampleWindowHeightWorld, in float In_SampleSteps, out float Out_VisibilityFraction, out float Out_SampleFraction)
-			{
-				float CameraDistance = length(In_SampleCenterWorldPos.xyz - View.WorldViewOrigin.xyz);
-				float4 SamplePosition = float4(In_SampleCenterWorldPos + View.PreViewTranslation, 1);
-				float4 ClipPosition = mul(SamplePosition, View.TranslatedWorldToClip);
-				float2 ScreenPosition = ClipPosition.xy / ClipPosition.w;
-				float2 ScreenUV = ScreenPosition * View.ScreenPositionScaleBias.xy + View.ScreenPositionScaleBias.wz;
-
-				float Steps = In_SampleSteps <= 1 ? 0 : In_SampleSteps;
-				float TotalSamples = 0;
-				float OccludedSamples = 0;
-
-				float4 SampleWidthClip = mul(float4(View.ViewRight * In_SampleWindowWidthWorld, 0) + SamplePosition, View.TranslatedWorldToClip);
-				float4 SampleHeightClip = mul(float4(View.ViewUp * In_SampleWindowHeightWorld, 0) + SamplePosition, View.TranslatedWorldToClip);
-				
-				float2 SampleWidthUV = SampleWidthClip.xy / SampleWidthClip.w * View.ScreenPositionScaleBias.xy + View.ScreenPositionScaleBias.wz;
-				float2 SampleHeightUV = SampleHeightClip.xy / SampleHeightClip.w * View.ScreenPositionScaleBias.xy + View.ScreenPositionScaleBias.wz;
-				
-				float SampleWidth = ScreenUV.x > 1 ? 0 : SampleWidthUV.x - ScreenUV.x;
-				float SampleHeight = ScreenUV.y > 1 ? 0 : SampleHeightUV.y - ScreenUV.y;
-
-				if (Steps > 0) 
-				{
-					for (int ys = 0; ys < Steps; ys++)
-					{
-						float SampleY = ScreenUV.y - 0.5 * SampleHeight + ys * SampleHeight / (Steps - 1);
-						if (SampleY > 1 || SampleY < 0)
-						{
-							continue;
-						}
-						for (int xs = 0; xs < Steps; xs++)
-						{
-							float SampleX = ScreenUV.x - 0.5 * SampleWidth + xs * SampleWidth / (Steps - 1);
-							if (SampleX > 1 || SampleX < 0)
-							{
-								continue;
-							}
-			
-							float Depth = CalcSceneDepth(float2(SampleX, SampleY));
-							if (Depth < CameraDistance) 
-							{
-								OccludedSamples++;
-							}
-							TotalSamples++;
-						} 
-					}
-				}
-				Out_VisibilityFraction = TotalSamples > 0 ? 1 - OccludedSamples / TotalSamples : 0;
-				Out_SampleFraction = Steps == 0 ? 0 : (TotalSamples / (Steps * Steps));
-			}
-		)");
-		OutHLSL += FString::Format(FormatSample, ArgsSample);
-		return true;
-	}
-	if (DefinitionFunctionName == GetCameraOcclusionCircleName)
-	{
-		static const TCHAR *FormatSample = TEXT(R"(
-			void {FunctionName}(in float3 In_SampleCenterWorldPos, in float In_SampleWindowDiameterWorld, in float In_SampleRays, in float In_SampleStepsPerRay, out float Out_VisibilityFraction, out float Out_SampleFraction)
-			{
-				const float PI = 3.14159265;
-				const float SPIRAL_TURN = 2 * PI * 0.61803399; // use golden ratio to rotate sample pattern each ring so we get a spiral
-				float CameraDistance = length(In_SampleCenterWorldPos.xyz - View.WorldViewOrigin.xyz);
-				float4 SamplePosition = float4(In_SampleCenterWorldPos + View.PreViewTranslation, 1);
-				float4 ClipPosition = mul(SamplePosition, View.TranslatedWorldToClip);
-				float2 ScreenPosition = ClipPosition.xy / ClipPosition.w;
-				float2 ScreenUV = ScreenPosition * View.ScreenPositionScaleBias.xy + View.ScreenPositionScaleBias.wz;
-
-				float Rays = In_SampleRays <= 1 ? 0 : In_SampleRays;
-				float Steps = In_SampleStepsPerRay < 1 ? 0 : In_SampleStepsPerRay;
-				float TotalSamples = 0;
-				float OccludedSamples = 0;
-
-				if (ScreenUV.x <= 1 && ScreenUV.x >= 0 && ScreenUV.y <= 1 && ScreenUV.y >= 0)
-				{
-					float Depth = CalcSceneDepth(ScreenUV);
-					if (Depth < CameraDistance) 
-					{
-						OccludedSamples++;
-					}
-					TotalSamples++;
-				}
-				if (Steps > 0) 
-				{
-					float Degrees = 0;
-					for (int Step = 1; Step <= Steps; Step++)
-					{
-						float LerpFactor = Step / Steps;
-						Degrees += SPIRAL_TURN;
-						for (int ray = 0; ray < Rays; ray++)
-						{
-							// calc ray direction vector
-							float3 RayDirection = cos(Degrees) * View.ViewUp + sin(Degrees) * View.ViewRight;
-							float4 RayClip = mul(float4(RayDirection * In_SampleWindowDiameterWorld / 2, 0) + SamplePosition, View.TranslatedWorldToClip);
-							float2 RayUV = RayClip.xy / RayClip.w * View.ScreenPositionScaleBias.xy + View.ScreenPositionScaleBias.wz;
-
-							if ((ScreenUV.x > 1 && RayUV.x < 0) || (ScreenUV.y > 1 && RayUV.y < 0) || (ScreenUV.x < 0 && RayUV.x > 1) || (ScreenUV.y < 0 && RayUV.y > 1))
-							{
-								continue;
-							}
-						
-							float2 SampleUV = lerp(ScreenUV, RayUV, float2(LerpFactor, LerpFactor));
-							
-							if (SampleUV.x > 1 || SampleUV.x < 0 || SampleUV.y > 1 || SampleUV.y < 0)
-							{
-								continue;
-							}
-			
-							float Depth = CalcSceneDepth(SampleUV);
-							if (Depth < CameraDistance) 
-							{
-								OccludedSamples++;
-							}
-							TotalSamples++;
-							Degrees += 2 * PI / Rays;
-						}						
-					}
-				}
-				Out_VisibilityFraction = TotalSamples > 0 ? 1 - OccludedSamples / TotalSamples : 0;
-				Out_SampleFraction = Steps == 0 ? 0 : (TotalSamples / (Rays * Steps + 1));
-			}
-		)");
-		OutHLSL += FString::Format(FormatSample, ArgsSample);
-		return true;
-	}
 	if (DefinitionFunctionName == GetViewPropertiesName)
 	{
 		static const TCHAR *FormatSample = TEXT(R"(
@@ -405,22 +239,12 @@ bool UNiagaraDataInterfaceCamera::GetFunctionHLSL(const FName& DefinitionFunctio
 
 DEFINE_NDI_DIRECT_FUNC_BINDER(UNiagaraDataInterfaceCamera, GetCameraFOV);
 DEFINE_NDI_DIRECT_FUNC_BINDER(UNiagaraDataInterfaceCamera, GetCameraProperties);
-DEFINE_NDI_DIRECT_FUNC_BINDER(UNiagaraDataInterfaceCamera, QueryOcclusionFactorGPU);
-DEFINE_NDI_DIRECT_FUNC_BINDER(UNiagaraDataInterfaceCamera, QueryOcclusionFactorCircleGPU);
 DEFINE_NDI_DIRECT_FUNC_BINDER(UNiagaraDataInterfaceCamera, GetViewPropertiesGPU);
 DEFINE_NDI_DIRECT_FUNC_BINDER(UNiagaraDataInterfaceCamera, GetClipSpaceTransformsGPU);
 DEFINE_NDI_DIRECT_FUNC_BINDER(UNiagaraDataInterfaceCamera, GetViewSpaceTransformsGPU);
 void UNiagaraDataInterfaceCamera::GetVMExternalFunction(const FVMExternalFunctionBindingInfo& BindingInfo, void* InstanceData, FVMExternalFunction &OutFunc)
 {
-	if (BindingInfo.Name == GetCameraOcclusionRectangleName)
-	{
-		NDI_FUNC_BINDER(UNiagaraDataInterfaceCamera, QueryOcclusionFactorGPU)::Bind(this, OutFunc);
-	}
-	else if (BindingInfo.Name == GetCameraOcclusionCircleName)
-	{
-		NDI_FUNC_BINDER(UNiagaraDataInterfaceCamera, QueryOcclusionFactorCircleGPU)::Bind(this, OutFunc);
-	}
-	else if (BindingInfo.Name == GetFieldOfViewName)
+	if (BindingInfo.Name == GetFieldOfViewName)
 	{
 		NDI_FUNC_BINDER(UNiagaraDataInterfaceCamera, GetCameraFOV)::Bind(this, OutFunc);
 	}
@@ -508,63 +332,12 @@ void UNiagaraDataInterfaceCamera::GetCameraProperties(FVectorVMContext& Context)
 	}
 }
 
-// ------- Dummy implementations for CPU execution ------------
-
-void UNiagaraDataInterfaceCamera::QueryOcclusionFactorGPU(FVectorVMContext& Context)
+ETickingGroup UNiagaraDataInterfaceCamera::CalculateTickGroup(void * PerInstanceData) const
 {
-	VectorVM::FExternalFuncInputHandler<float> PosParamX(Context);
-	VectorVM::FExternalFuncInputHandler<float> PosParamY(Context);
-	VectorVM::FExternalFuncInputHandler<float> PosParamZ(Context);
-	VectorVM::FExternalFuncInputHandler<float> WidthParam(Context);
-	VectorVM::FExternalFuncInputHandler<float> HeightParam(Context);
-	VectorVM::FExternalFuncInputHandler<float> StepsParam(Context);
-
-	VectorVM::FUserPtrHandler<CameraDataInterface_InstanceData> InstData(Context);
-
-	VectorVM::FExternalFuncRegisterHandler<float> OutOcclusion(Context);
-	VectorVM::FExternalFuncRegisterHandler<float> OutFactor(Context);
-
-	for (int32 i = 0; i < Context.NumInstances; ++i)
-	{
-		PosParamX.GetAndAdvance();
-		PosParamY.GetAndAdvance();
-		PosParamZ.GetAndAdvance();
-		WidthParam.GetAndAdvance();
-		HeightParam.GetAndAdvance();
-		StepsParam.GetAndAdvance();
-
-		*OutOcclusion.GetDestAndAdvance() = 0;
-		*OutFactor.GetDestAndAdvance() = 0;
-	}
+	return ETickingGroup::TG_PostUpdateWork;
 }
 
-void UNiagaraDataInterfaceCamera::QueryOcclusionFactorCircleGPU(FVectorVMContext& Context)
-{
-	VectorVM::FExternalFuncInputHandler<float> PosParamX(Context);
-	VectorVM::FExternalFuncInputHandler<float> PosParamY(Context);
-	VectorVM::FExternalFuncInputHandler<float> PosParamZ(Context);
-	VectorVM::FExternalFuncInputHandler<float> RadiusParam(Context);
-	VectorVM::FExternalFuncInputHandler<float> RaysParam(Context);
-	VectorVM::FExternalFuncInputHandler<float> StepsParam(Context);
-
-	VectorVM::FUserPtrHandler<CameraDataInterface_InstanceData> InstData(Context);
-
-	VectorVM::FExternalFuncRegisterHandler<float> OutOcclusion(Context);
-	VectorVM::FExternalFuncRegisterHandler<float> OutFactor(Context);
-
-	for (int32 i = 0; i < Context.NumInstances; ++i)
-	{
-		PosParamX.GetAndAdvance();
-		PosParamY.GetAndAdvance();
-		PosParamZ.GetAndAdvance();
-		RadiusParam.GetAndAdvance();
-		RaysParam.GetAndAdvance();
-		StepsParam.GetAndAdvance();
-
-		*OutOcclusion.GetDestAndAdvance() = 0;
-		*OutFactor.GetDestAndAdvance() = 0;
-	}
-}
+// ------- Dummy implementations for CPU execution ------------
 
 void UNiagaraDataInterfaceCamera::GetViewPropertiesGPU(FVectorVMContext& Context)
 {
@@ -625,39 +398,3 @@ void UNiagaraDataInterfaceCamera::GetViewSpaceTransformsGPU(FVectorVMContext& Co
 
 // ------------------------------------------------------------
 
-struct FNiagaraDataInterfaceParametersCS_CameraQuery : public FNiagaraDataInterfaceParametersCS
-{
-	virtual void Bind(const FNiagaraDataInterfaceParamRef& ParamRef, const class FShaderParameterMap& ParameterMap) override
-	{
-		PassUniformBuffer.Bind(ParameterMap, FSceneTexturesUniformParameters::StaticStructMetadata.GetShaderVariableName());
-	}
-
-	virtual void Serialize(FArchive& Ar) override
-	{
-		Ar << PassUniformBuffer;
-	}
-
-	virtual void Set(FRHICommandList& RHICmdList, const FNiagaraDataInterfaceSetArgs& Context) const override
-	{
-		check(IsInRenderingThread());
-		FRHIComputeShader* ComputeShaderRHI = Context.Shader->GetComputeShader();
-
-		TUniformBufferRef<FSceneTexturesUniformParameters> SceneTextureUniformParams = GNiagaraViewDataManager.GetSceneTextureUniformParameters();
-		SetUniformBufferParameter(RHICmdList, ComputeShaderRHI, PassUniformBuffer, SceneTextureUniformParams);
-	}
-
-private:
-
-	/** The SceneDepthTexture parameter for depth buffer query. */
-	FShaderUniformBufferParameter PassUniformBuffer;
-};
-
-FNiagaraDataInterfaceParametersCS* UNiagaraDataInterfaceCamera::ConstructComputeParameters() const
-{
-	return new FNiagaraDataInterfaceParametersCS_CameraQuery();
-}
-
-ETickingGroup UNiagaraDataInterfaceCamera::CalculateTickGroup(void * PerInstanceData) const
-{
-	return ETickingGroup::TG_PostUpdateWork;
-}
diff --git a/Engine/Plugins/FX/Niagara/Source/Niagara/Private/NiagaraDataInterfaceOcclusion.cpp b/Engine/Plugins/FX/Niagara/Source/Niagara/Private/NiagaraDataInterfaceOcclusion.cpp
new file mode 100644
--- /dev/null
+++ b/Engine/Plugins/FX/Niagara/Source/Niagara/Private/NiagaraDataInterfaceOcclusion.cpp
@@ -0,0 +1,305 @@
+// Copyright 1998-2019 Epic Games, Inc. All Rights Reserved.
+
+#include "NiagaraDataInterfaceOcclusion.h"
+#include "NiagaraTypes.h"
+#include "NiagaraWorldManager.h"
+#include "ShaderParameterUtils.h"
+#include "Internationalization/Internationalization.h"
+#include "NiagaraSystemInstance.h"
+
+const FName UNiagaraDataInterfaceOcclusion::GetCameraOcclusionRectangleName(TEXT("QueryOcclusionFactorWithRectangleGPU"));
+const FName UNiagaraDataInterfaceOcclusion::GetCameraOcclusionCircleName(TEXT("QueryOcclusionFactorWithCircleGPU"));
+
+UNiagaraDataInterfaceOcclusion::UNiagaraDataInterfaceOcclusion(FObjectInitializer const& ObjectInitializer)
+	: Super(ObjectInitializer)
+{
+	Proxy = MakeShared<FNiagaraDataIntefaceProxyOcclusionQuery, ESPMode::ThreadSafe>();
+}
+
+void UNiagaraDataInterfaceOcclusion::PostInitProperties()
+{
+	Super::PostInitProperties();
+
+	if (HasAnyFlags(RF_ClassDefaultObject))
+	{
+		FNiagaraTypeRegistry::Register(FNiagaraTypeDefinition(GetClass()), true, false, false);
+	}
+}
+
+void UNiagaraDataInterfaceOcclusion::GetFunctions(TArray<FNiagaraFunctionSignature>& OutFunctions)
+{
+	FNiagaraFunctionSignature Sig;
+	Sig.Name = GetCameraOcclusionRectangleName;
+#if WITH_EDITORONLY_DATA
+	Sig.Description = NSLOCTEXT("Niagara", "GetCameraOcclusionRectFunctionDescription", "This function returns the occlusion factor of a sprite. It samples the depth buffer in a rectangular grid around the given world position and compares each sample with the camera distance.");
+#endif
+	Sig.bMemberFunction = true;
+	Sig.bRequiresContext = false;
+	Sig.bSupportsCPU = false;
+	Sig.Inputs.Add(FNiagaraVariable(FNiagaraTypeDefinition(GetClass()), TEXT("Camera interface")));
+	Sig.Inputs.Add(FNiagaraVariable(FNiagaraTypeDefinition::GetVec3Def(), TEXT("Sample Center World Position")));
+	Sig.Inputs.Add(FNiagaraVariable(FNiagaraTypeDefinition::GetFloatDef(), TEXT("Sample Window Width World")));
+	Sig.Inputs.Add(FNiagaraVariable(FNiagaraTypeDefinition::GetFloatDef(), TEXT("Sample Window Height World")));
+	Sig.Inputs.Add(FNiagaraVariable(FNiagaraTypeDefinition::GetFloatDef(), TEXT("Sample Steps Per Line")));
+	Sig.Outputs.Add(FNiagaraVariable(FNiagaraTypeDefinition::GetFloatDef(), TEXT("Visibility Fraction")));
+	Sig.Outputs.Add(FNiagaraVariable(FNiagaraTypeDefinition::GetFloatDef(), TEXT("Sample Fraction")));
+	OutFunctions.Add(Sig);
+
+	Sig = FNiagaraFunctionSignature();
+	Sig.Name = GetCameraOcclusionCircleName;
+#if WITH_EDITORONLY_DATA
+	Sig.Description = NSLOCTEXT("Niagara", "GetCameraOcclusionCircleFunctionDescription", "This function returns the occlusion factor of a sprite. It samples the depth buffer in concentric rings around the given world position and compares each sample with the camera distance.");
+#endif
+	Sig.bMemberFunction = true;
+	Sig.bRequiresContext = false;
+	Sig.bSupportsCPU = false;
+	Sig.Inputs.Add(FNiagaraVariable(FNiagaraTypeDefinition(GetClass()), TEXT("Camera interface")));
+	Sig.Inputs.Add(FNiagaraVariable(FNiagaraTypeDefinition::GetVec3Def(), TEXT("Sample Center World Position")));
+	Sig.Inputs.Add(FNiagaraVariable(FNiagaraTypeDefinition::GetFloatDef(), TEXT("Sample Window Diameter World")));
+	Sig.Inputs.Add(FNiagaraVariable(FNiagaraTypeDefinition::GetFloatDef(), TEXT("Samples per ring")));
+	Sig.Inputs.Add(FNiagaraVariable(FNiagaraTypeDefinition::GetFloatDef(), TEXT("Number of sample rings")));
+	Sig.Outputs.Add(FNiagaraVariable(FNiagaraTypeDefinition::GetFloatDef(), TEXT("Visibility Fraction")));
+	Sig.Outputs.Add(FNiagaraVariable(FNiagaraTypeDefinition::GetFloatDef(), TEXT("Sample Fraction")));
+	OutFunctions.Add(Sig);
+}
+
+bool UNiagaraDataInterfaceOcclusion::GetFunctionHLSL(const FName& DefinitionFunctionName, FString InstanceFunctionName, FNiagaraDataInterfaceGPUParamInfo& ParamInfo, FString& OutHLSL)
+{
+	TMap<FString, FStringFormatArg> ArgsSample;
+	ArgsSample.Add(TEXT("FunctionName"), InstanceFunctionName);
+
+	if (DefinitionFunctionName == GetCameraOcclusionRectangleName)
+	{
+		static const TCHAR *FormatSample = TEXT(R"(
+			void {FunctionName}(in float3 In_SampleCenterWorldPos, in float In_SampleWindowWidthWorld, in float In_SampleWindowHeightWorld, in float In_SampleSteps, out float Out_VisibilityFraction, out float Out_SampleFraction)
+			{
+				float CameraDistance = length(In_SampleCenterWorldPos.xyz - View.WorldViewOrigin.xyz);
+				float4 SamplePosition = float4(In_SampleCenterWorldPos + View.PreViewTranslation, 1);
+				float4 ClipPosition = mul(SamplePosition, View.TranslatedWorldToClip);
+				float2 ScreenPosition = ClipPosition.xy / ClipPosition.w;
+				float2 ScreenUV = ScreenPosition * View.ScreenPositionScaleBias.xy + View.ScreenPositionScaleBias.wz;
+
+				float Steps = In_SampleSteps <= 1 ? 0 : In_SampleSteps;
+				float TotalSamples = 0;
+				float OccludedSamples = 0;
+
+				float4 SampleWidthClip = mul(float4(View.ViewRight * In_SampleWindowWidthWorld, 0) + SamplePosition, View.TranslatedWorldToClip);
+				float4 SampleHeightClip = mul(float4(View.ViewUp * In_SampleWindowHeightWorld, 0) + SamplePosition, View.TranslatedWorldToClip);
+				
+				float2 SampleWidthUV = SampleWidthClip.xy / SampleWidthClip.w * View.ScreenPositionScaleBias.xy + View.ScreenPositionScaleBias.wz;
+				float2 SampleHeightUV = SampleHeightClip.xy / SampleHeightClip.w * View.ScreenPositionScaleBias.xy + View.ScreenPositionScaleBias.wz;
+				
+				float SampleWidth = ScreenUV.x > 1 ? 0 : SampleWidthUV.x - ScreenUV.x;
+				float SampleHeight = ScreenUV.y > 1 ? 0 : SampleHeightUV.y - ScreenUV.y;
+
+				if (Steps > 0) 
+				{
+					for (int ys = 0; ys < Steps; ys++)
+					{
+						float SampleY = ScreenUV.y - 0.5 * SampleHeight + ys * SampleHeight / (Steps - 1);
+						if (SampleY > 1 || SampleY < 0)
+						{
+							continue;
+						}
+						for (int xs = 0; xs < Steps; xs++)
+						{
+							float SampleX = ScreenUV.x - 0.5 * SampleWidth + xs * SampleWidth / (Steps - 1);
+							if (SampleX > 1 || SampleX < 0)
+							{
+								continue;
+							}
+			
+							float Depth = CalcSceneDepth(float2(SampleX, SampleY));
+							if (Depth < CameraDistance) 
+							{
+								OccludedSamples++;
+							}
+							TotalSamples++;
+						} 
+					}
+				}
+				Out_VisibilityFraction = TotalSamples > 0 ? 1 - OccludedSamples / TotalSamples : 0;
+				Out_SampleFraction = Steps == 0 ? 0 : (TotalSamples / (Steps * Steps));
+			}
+		)");
+		OutHLSL += FString::Format(FormatSample, ArgsSample);
+		return true;
+	}
+	if (DefinitionFunctionName == GetCameraOcclusionCircleName)
+	{
+		static const TCHAR *FormatSample = TEXT(R"(
+			void {FunctionName}(in float3 In_SampleCenterWorldPos, in float In_SampleWindowDiameterWorld, in float In_SampleRays, in float In_SampleStepsPerRay, out float Out_VisibilityFraction, out float Out_SampleFraction)
+			{
+				const float PI = 3.14159265;
+				const float SPIRAL_TURN = 2 * PI * 0.61803399; // use golden ratio to rotate sample pattern each ring so we get a spiral
+				float CameraDistance = length(In_SampleCenterWorldPos.xyz - View.WorldViewOrigin.xyz);
+				float4 SamplePosition = float4(In_SampleCenterWorldPos + View.PreViewTranslation, 1);
+				float4 ClipPosition = mul(SamplePosition, View.TranslatedWorldToClip);
+				float2 ScreenPosition = ClipPosition.xy / ClipPosition.w;
+				float2 ScreenUV = ScreenPosition * View.ScreenPositionScaleBias.xy + View.ScreenPositionScaleBias.wz;
+
+				float Rays = In_SampleRays <= 1 ? 0 : In_SampleRays;
+				float Steps = In_SampleStepsPerRay < 1 ? 0 : In_SampleStepsPerRay;
+				float TotalSamples = 0;
+				float OccludedSamples = 0;
+
+				if (ScreenUV.x <= 1 && ScreenUV.x >= 0 && ScreenUV.y <= 1 && ScreenUV.y >= 0)
+				{
+					float Depth = CalcSceneDepth(ScreenUV);
+					if (Depth < CameraDistance) 
+					{
+						OccludedSamples++;
+					}
+					TotalSamples++;
+				}
+				if (Steps > 0) 
+				{
+					float Degrees = 0;
+					for (int Step = 1; Step <= Steps; Step++)
+					{
+						float LerpFactor = Step / Steps;
+						Degrees += SPIRAL_TURN;
+						for (int ray = 0; ray < Rays; ray++)
+						{
+							// calc ray direction vector
+							float3 RayDirection = cos(Degrees) * View.ViewUp + sin(Degrees) * View.ViewRight;
+							float4 RayClip = mul(float4(RayDirection * In_SampleWindowDiameterWorld / 2, 0) + SamplePosition, View.TranslatedWorldToClip);
+							float2 RayUV = RayClip.xy / RayClip.w * View.ScreenPositionScaleBias.xy + View.ScreenPositionScaleBias.wz;
+
+							if ((ScreenUV.x > 1 && RayUV.x < 0) || (ScreenUV.y > 1 && RayUV.y < 0) || (ScreenUV.x < 0 && RayUV.x > 1) || (ScreenUV.y < 0 && RayUV.y > 1))
+							{
+								continue;
+							}
+						
+							float2 SampleUV = lerp(ScreenUV, RayUV, float2(LerpFactor, LerpFactor));
+							
+							if (SampleUV.x > 1 || SampleUV.x < 0 || SampleUV.y > 1 || SampleUV.y < 0)
+							{
+								continue;
+							}
+			
+							float Depth = CalcSceneDepth(SampleUV);
+							if (Depth < CameraDistance) 
+							{
+								OccludedSamples++;
+							}
+							TotalSamples++;
+							Degrees += 2 * PI / Rays;
+						}						
+					}
+				}
+				Out_VisibilityFraction = TotalSamples > 0 ? 1 - OccludedSamples / TotalSamples : 0;
+				Out_SampleFraction = Steps == 0 ? 0 : (TotalSamples / (Rays * Steps + 1));
+			}
+		)");
+		OutHLSL += FString::Format(FormatSample, ArgsSample);
+		return true;
+	}
+	return false;
+}
+
+DEFINE_NDI_DIRECT_FUNC_BINDER(UNiagaraDataInterfaceOcclusion, QueryOcclusionFactorGPU);
+DEFINE_NDI_DIRECT_FUNC_BINDER(UNiagaraDataInterfaceOcclusion, QueryOcclusionFactorCircleGPU);
+void UNiagaraDataInterfaceOcclusion::GetVMExternalFunction(const FVMExternalFunctionBindingInfo& BindingInfo, void* InstanceData, FVMExternalFunction &OutFunc)
+{
+	if (BindingInfo.Name == GetCameraOcclusionRectangleName)
+	{
+		NDI_FUNC_BINDER(UNiagaraDataInterfaceOcclusion, QueryOcclusionFactorGPU)::Bind(this, OutFunc);
+	}
+	else if (BindingInfo.Name == GetCameraOcclusionCircleName)
+	{
+		NDI_FUNC_BINDER(UNiagaraDataInterfaceOcclusion, QueryOcclusionFactorCircleGPU)::Bind(this, OutFunc);
+	}
+	else
+	{
+		UE_LOG(LogNiagara, Error, TEXT("Could not find data interface external function. Received Name: %s"), *BindingInfo.Name.ToString());
+	}
+}
+
+// ------- Dummy implementations for CPU execution ------------
+
+void UNiagaraDataInterfaceOcclusion::QueryOcclusionFactorGPU(FVectorVMContext& Context)
+{
+	VectorVM::FExternalFuncInputHandler<float> PosParamX(Context);
+	VectorVM::FExternalFuncInputHandler<float> PosParamY(Context);
+	VectorVM::FExternalFuncInputHandler<float> PosParamZ(Context);
+	VectorVM::FExternalFuncInputHandler<float> WidthParam(Context);
+	VectorVM::FExternalFuncInputHandler<float> HeightParam(Context);
+	VectorVM::FExternalFuncInputHandler<float> StepsParam(Context);
+
+	VectorVM::FExternalFuncRegisterHandler<float> OutOcclusion(Context);
+	VectorVM::FExternalFuncRegisterHandler<float> OutFactor(Context);
+
+	for (int32 i = 0; i < Context.NumInstances; ++i)
+	{
+		PosParamX.GetAndAdvance();
+		PosParamY.GetAndAdvance();
+		PosParamZ.GetAndAdvance();
+		WidthParam.GetAndAdvance();
+		HeightParam.GetAndAdvance();
+		StepsParam.GetAndAdvance();
+
+		*OutOcclusion.GetDestAndAdvance() = 0;
+		*OutFactor.GetDestAndAdvance() = 0;
+	}
+}
+
+void UNiagaraDataInterfaceOcclusion::QueryOcclusionFactorCircleGPU(FVectorVMContext& Context)
+{
+	VectorVM::FExternalFuncInputHandler<float> PosParamX(Context);
+	VectorVM::FExternalFuncInputHandler<float> PosParamY(Context);
+	VectorVM::FExternalFuncInputHandler<float> PosParamZ(Context);
+	VectorVM::FExternalFuncInputHandler<float> RadiusParam(Context);
+	VectorVM::FExternalFuncInputHandler<float> RaysParam(Context);
+	VectorVM::FExternalFuncInputHandler<float> StepsParam(Context);
+
+	VectorVM::FExternalFuncRegisterHandler<float> OutOcclusion(Context);
+	VectorVM::FExternalFuncRegisterHandler<float> OutFactor(Context);
+
+	for (int32 i = 0; i < Context.NumInstances; ++i)
+	{
+		PosParamX.GetAndAdvance();
+		PosParamY.GetAndAdvance();
+		PosParamZ.GetAndAdvance();
+		RadiusParam.GetAndAdvance();
+		RaysParam.GetAndAdvance();
+		StepsParam.GetAndAdvance();
+
+		*OutOcclusion.GetDestAndAdvance() = 0;
+		*OutFactor.GetDestAndAdvance() = 0;
+	}
+}
+
+// ------------------------------------------------------------
+
+struct FNiagaraDataInterfaceParametersCS_OcclusionQuery : public FNiagaraDataInterfaceParametersCS
+{
+	virtual void Bind(const FNiagaraDataInterfaceParamRef& ParamRef, const class FShaderParameterMap& ParameterMap) override
+	{
+		PassUniformBuffer.Bind(ParameterMap, FSceneTexturesUniformParameters::StaticStructMetadata.GetShaderVariableName());
+	}
+
+	virtual void Serialize(FArchive& Ar) override
+	{
+		Ar << PassUniformBuffer;
+	}
+
+	virtual void Set(FRHICommandList& RHICmdList, const FNiagaraDataInterfaceSetArgs& Context) const override
+	{
+		check(IsInRenderingThread());
+		FRHIComputeShader* ComputeShaderRHI = Context.Shader->GetComputeShader();
+
+		TUniformBufferRef<FSceneTexturesUniformParameters> SceneTextureUniformParams = GNiagaraViewDataManager.GetSceneTextureUniformParameters();
+		SetUniformBufferParameter(RHICmdList, ComputeShaderRHI, PassUniformBuffer, SceneTextureUniformParams);
+	}
+
+private:
+
+	/** The SceneDepthTexture parameter for depth buffer query. */
+	FShaderUniformBufferParameter PassUniformBuffer;
+};
+
+FNiagaraDataInterfaceParametersCS* UNiagaraDataInterfaceOcclusion::ConstructComputeParameters() const
+{
+	return new FNiagaraDataInterfaceParametersCS_OcclusionQuery();
+}
diff --git a/Engine/Plugins/FX/Niagara/Source/Niagara/Private/NiagaraEmitterInstanceBatcher.cpp b/Engine/Plugins/FX/Niagara/Source/Niagara/Private/NiagaraEmitterInstanceBatcher.cpp
--- a/Engine/Plugins/FX/Niagara/Source/Niagara/Private/NiagaraEmitterInstanceBatcher.cpp
+++ b/Engine/Plugins/FX/Niagara/Source/Niagara/Private/NiagaraEmitterInstanceBatcher.cpp
@@ -505,7 +505,7 @@ void NiagaraEmitterInstanceBatcher::PostRenderOpaque(FRHICommandListImmediate& R
 	if (bAllowGPUParticleUpdate)
 	{
 		// Setup new readback since if there is no pending request, there is no risk of having invalid data read (offset being allocated after the readback was sent).
-		ExecuteAll(RHICmdList, ViewUniformBuffer, !GPUInstanceCounterManager.HasPendingGPUReadback());
+		ExecuteAll(RHICmdList, ViewUniformBuffer, !GPUInstanceCounterManager.HasPendingGPUReadback(), ETickStage::PostOpaqueRender);
 
 		FinishDispatches();
 	}
@@ -516,7 +516,27 @@ void NiagaraEmitterInstanceBatcher::PostRenderOpaque(FRHICommandListImmediate& R
 	}
 }
 
-void NiagaraEmitterInstanceBatcher::ExecuteAll(FRHICommandList &RHICmdList, FRHIUniformBuffer* ViewUniformBuffer, bool bSetReadback)
+bool NiagaraEmitterInstanceBatcher::ShouldTickForStage(const FNiagaraGPUSystemTick& Tick, ETickStage TickStage) const
+{
+	if (!GNiagaraAllowTickBeforeRender || Tick.bRequiresDistanceFieldData || Tick.bRequiresDepthBuffer)
+	{
+		return TickStage == ETickStage::PostOpaqueRender;
+	}
+
+	if (Tick.bRequiresEarlyViewData)
+	{
+		return TickStage == ETickStage::PostInitViews;
+	}
+
+	FNiagaraShader* ComputeShader = Tick.GetInstanceData()->Context->GPUScript_RT->GetShader();
+	if (ComputeShader->ViewUniformBufferParam.IsBound())
+	{
+		return TickStage == ETickStage::PostOpaqueRender;
+	}
+	return TickStage == ETickStage::PreInitViews;
+}
+
+void NiagaraEmitterInstanceBatcher::ExecuteAll(FRHICommandList &RHICmdList, FRHIUniformBuffer* ViewUniformBuffer, bool bSetReadback, ETickStage TickStage)
 {
 	SCOPE_CYCLE_COUNTER(STAT_NiagaraGPUSimTick_RT);
 
@@ -539,14 +559,10 @@ void NiagaraEmitterInstanceBatcher::ExecuteAll(FRHICommandList &RHICmdList, FRHI
 			FNiagaraComputeExecutionContext* Context = Data->Context;
 			// This assumes all emitters fallback to the same FNiagaraShaderScript*.
 			FNiagaraShader* ComputeShader = Context->GPUScript_RT->GetShader();
-			if (!ComputeShader)
+			if (!ComputeShader || !ShouldTickForStage(Tick, TickStage))
 			{
 				continue;
 			}
-			else if (GNiagaraAllowTickBeforeRender && (ComputeShader->ViewUniformBufferParam.IsBound() || Tick.bRequiredDistanceFieldData) != (ViewUniformBuffer != nullptr))
-			{   // When allowing tick before render, skip this emitter if it is not in the right pass.
-				continue;
-			}
 
 			Tick.bIsFinalTick = false; // @todo : this is true sometimes, needs investigation
 			if (Context->ScratchIndex == INDEX_NONE)
@@ -705,7 +721,7 @@ void NiagaraEmitterInstanceBatcher::PreInitViews(FRHICommandListImmediate& RHICm
 
 		if (GNiagaraAllowTickBeforeRender)
 		{
-			ExecuteAll(RHICmdList, nullptr, !GPUInstanceCounterManager.HasPendingGPUReadback());
+			ExecuteAll(RHICmdList, nullptr, !GPUInstanceCounterManager.HasPendingGPUReadback(), ETickStage::PreInitViews);
 		}
 	}
 	else
@@ -714,11 +730,47 @@ void NiagaraEmitterInstanceBatcher::PreInitViews(FRHICommandListImmediate& RHICm
 	}
 }
 
+void NiagaraEmitterInstanceBatcher::PostInitViews(FRHICommandListImmediate& RHICmdList, FRHIUniformBuffer* ViewUniformBuffer, bool bAllowGPUParticleUpdate)
+{
+	LLM_SCOPE(ELLMTag::Niagara);
+
+	if (bAllowGPUParticleUpdate)
+	{
+		ExecuteAll(RHICmdList, ViewUniformBuffer, false, ETickStage::PostInitViews);
+	}
+}
+
 bool NiagaraEmitterInstanceBatcher::UsesGlobalDistanceField() const
 {
 	for (const FNiagaraGPUSystemTick& Tick : Ticks_RT)
 	{
-		if (Tick.bRequiredDistanceFieldData)
+		if (Tick.bRequiresDistanceFieldData)
+		{
+			return true;
+		}
+	}
+
+	return false;
+}
+
+bool NiagaraEmitterInstanceBatcher::UsesDepthBuffer() const
+{
+	for (const FNiagaraGPUSystemTick& Tick : Ticks_RT)
+	{
+		if (Tick.bRequiresDepthBuffer)
+		{
+			return true;
+		}
+	}
+
+	return false;
+}
+
+bool NiagaraEmitterInstanceBatcher::RequiresEarlyViewUniformBuffer() const
+{
+	for (const FNiagaraGPUSystemTick& Tick : Ticks_RT)
+	{
+		if (Tick.bRequiresEarlyViewData)
 		{
 			return true;
 		}
diff --git a/Engine/Plugins/FX/Niagara/Source/Niagara/Private/NiagaraScriptExecutionContext.cpp b/Engine/Plugins/FX/Niagara/Source/Niagara/Private/NiagaraScriptExecutionContext.cpp
--- a/Engine/Plugins/FX/Niagara/Source/Niagara/Private/NiagaraScriptExecutionContext.cpp
+++ b/Engine/Plugins/FX/Niagara/Source/Niagara/Private/NiagaraScriptExecutionContext.cpp
@@ -249,7 +249,9 @@ void FNiagaraGPUSystemTick::Init(FNiagaraSystemInstance* InSystemInstance)
 	CA_ASSUME(InSystemInstance != nullptr);
 	ensure(!InSystemInstance->IsComplete());
 	SystemInstanceID = InSystemInstance->GetId();
-	bRequiredDistanceFieldData = InSystemInstance->RequiresDistanceFieldData();
+	bRequiresDistanceFieldData = InSystemInstance->RequiresDistanceFieldData();
+	bRequiresDepthBuffer = InSystemInstance->RequiresDepthBuffer();
+	bRequiresEarlyViewData = InSystemInstance->RequiresEarlyViewData();
 	uint32 DataSizeForGPU = InSystemInstance->GPUDataInterfaceInstanceDataSize;
 
 	if (DataSizeForGPU > 0)
diff --git a/Engine/Plugins/FX/Niagara/Source/Niagara/Private/NiagaraSystemInstance.cpp b/Engine/Plugins/FX/Niagara/Source/Niagara/Private/NiagaraSystemInstance.cpp
--- a/Engine/Plugins/FX/Niagara/Source/Niagara/Private/NiagaraSystemInstance.cpp
+++ b/Engine/Plugins/FX/Niagara/Source/Niagara/Private/NiagaraSystemInstance.cpp
@@ -1096,6 +1096,56 @@ bool FNiagaraSystemInstance::RequiresDistanceFieldData() const
 	return false;
 }
 
+bool FNiagaraSystemInstance::RequiresDepthBuffer() const
+{
+	if (!bHasGPUEmitters)
+	{
+		return false;
+	}
+
+	for (const TSharedRef<FNiagaraEmitterInstance, ESPMode::ThreadSafe>& Emitter : Emitters)
+	{
+		FNiagaraComputeExecutionContext* GPUContext = Emitter->GetGPUContext();
+		if (GPUContext)
+		{
+			for (UNiagaraDataInterface* DataInterface : GPUContext->CombinedParamStore.GetDataInterfaces())
+			{
+				if (DataInterface && DataInterface->RequiresDepthBuffer())
+				{
+					return true;
+				}
+			}
+		}
+	}
+
+	return false;
+}
+
+bool FNiagaraSystemInstance::RequiresEarlyViewData() const
+{
+	if (!bHasGPUEmitters)
+	{
+		return false;
+	}
+
+	for (const TSharedRef<FNiagaraEmitterInstance, ESPMode::ThreadSafe>& Emitter : Emitters)
+	{
+		FNiagaraComputeExecutionContext* GPUContext = Emitter->GetGPUContext();
+		if (GPUContext)
+		{
+			for (UNiagaraDataInterface* DataInterface : GPUContext->CombinedParamStore.GetDataInterfaces())
+			{
+				if (DataInterface && DataInterface->RequiresEarlyViewData())
+				{
+					return true;
+				}
+			}
+		}
+	}
+
+	return false;
+}
+
 void FNiagaraSystemInstance::InitDataInterfaces()
 {
 	bDataInterfacesHaveTickPrereqs = false;
diff --git a/Engine/Plugins/FX/Niagara/Source/Niagara/Public/NiagaraSystemInstance.h b/Engine/Plugins/FX/Niagara/Source/Niagara/Public/NiagaraSystemInstance.h
--- a/Engine/Plugins/FX/Niagara/Source/Niagara/Public/NiagaraSystemInstance.h
+++ b/Engine/Plugins/FX/Niagara/Source/Niagara/Public/NiagaraSystemInstance.h
@@ -80,6 +80,8 @@ public:
 	
 	FNiagaraWorldManager* GetWorldManager()const;
 	bool RequiresDistanceFieldData() const;
+	bool RequiresDepthBuffer() const;
+	bool RequiresEarlyViewData() const;
 
 	/** Requests the the simulation be reset on the next tick. */
 	void Reset(EResetMode Mode);
diff --git a/Engine/Source/Runtime/Engine/Private/Particles/FXSystem.cpp b/Engine/Source/Runtime/Engine/Private/Particles/FXSystem.cpp
--- a/Engine/Source/Runtime/Engine/Private/Particles/FXSystem.cpp
+++ b/Engine/Source/Runtime/Engine/Private/Particles/FXSystem.cpp
@@ -364,6 +364,11 @@ void FFXSystem::PreInitViews(FRHICommandListImmediate& RHICmdList, bool bAllowGP
 	}
 }
 
+void FFXSystem::PostInitViews(FRHICommandListImmediate& RHICmdList, FRHIUniformBuffer* ViewUniformBuffer, bool bAllowGPUParticleUpdate)
+{
+	// nothing to do here
+}
+
 bool FFXSystem::UsesGlobalDistanceField() const
 {
 	if (RHISupportsGPUParticles())
@@ -374,6 +379,26 @@ bool FFXSystem::UsesGlobalDistanceField() const
 	return false;
 }
 
+bool FFXSystem::UsesDepthBuffer() const
+{
+	if (RHISupportsGPUParticles())
+	{
+		return UsesDepthBufferInternal();
+	}
+
+	return false;
+}
+
+bool FFXSystem::RequiresEarlyViewUniformBuffer() const
+{
+	if (RHISupportsGPUParticles())
+	{
+		return RequiresEarlyViewUniformBufferInternal();
+	}
+
+	return false;
+}
+
 DECLARE_CYCLE_STAT(TEXT("FXPreRender_Prepare"), STAT_CLM_FXPreRender_Prepare, STATGROUP_CommandListMarkers);
 DECLARE_CYCLE_STAT(TEXT("FXPreRender_Simulate"), STAT_CLM_FXPreRender_Simulate, STATGROUP_CommandListMarkers);
 DECLARE_CYCLE_STAT(TEXT("FXPreRender_Finalize"), STAT_CLM_FXPreRender_Finalize, STATGROUP_CommandListMarkers);
@@ -405,7 +430,7 @@ void FFXSystem::PreRender(FRHICommandListImmediate& RHICmdList, const FGlobalDis
 
 			RHICmdList.SetCurrentStat(GET_STATID(STAT_CLM_FXPreRender_SimulateCDF));
 			SimulateGPUParticles(RHICmdList, EParticleSimulatePhase::CollisionDistanceField, nullptr, GlobalDistanceFieldParameterData, nullptr, nullptr);
-			//particles rendered during basepass may need to read pos/velocity buffers.  must finalize unless we know for sure that nothingin base pass will read.
+			//particles rendered during basepass may need to read pos/velocity buffers; must finalize unless we know for sure that nothing in base pass will read it.
 			RHICmdList.SetCurrentStat(GET_STATID(STAT_CLM_FXPreRender_FinalizeCDF));
 			FinalizeGPUSimulation(RHICmdList);
 		}
diff --git a/Engine/Source/Runtime/Engine/Private/Particles/FXSystemPrivate.h b/Engine/Source/Runtime/Engine/Private/Particles/FXSystemPrivate.h
--- a/Engine/Source/Runtime/Engine/Private/Particles/FXSystemPrivate.h
+++ b/Engine/Source/Runtime/Engine/Private/Particles/FXSystemPrivate.h
@@ -150,7 +150,10 @@ public:
 	virtual void UpdateVectorField(UVectorFieldComponent* VectorFieldComponent) override;
 	FParticleEmitterInstance* CreateGPUSpriteEmitterInstance(FGPUSpriteEmitterInfo& EmitterInfo);
 	virtual void PreInitViews(FRHICommandListImmediate& RHICmdList, bool bAllowGPUParticleUpdate) override;
+	virtual void PostInitViews(FRHICommandListImmediate& RHICmdList, FRHIUniformBuffer* ViewUniformBuffer, bool bAllowGPUParticleUpdate) override;
 	virtual bool UsesGlobalDistanceField() const override;
+	virtual bool UsesDepthBuffer() const override;
+	virtual bool RequiresEarlyViewUniformBuffer() const override;
 	virtual void PreRender(FRHICommandListImmediate& RHICmdList, const FGlobalDistanceFieldParameterData* GlobalDistanceFieldParameterData, bool bAllowGPUParticleUpdate) override;
 	virtual void PostRenderOpaque(
 		FRHICommandListImmediate& RHICmdList, 
@@ -242,6 +245,8 @@ private:
 	void SortGPUParticles(FRHICommandListImmediate& RHICmdList);
 
 	bool UsesGlobalDistanceFieldInternal() const;
+	bool UsesDepthBufferInternal() const;
+	bool RequiresEarlyViewUniformBufferInternal() const;
 
 	/**
 	* Updates resources used in a multi-GPU context
@@ -307,3 +312,4 @@ private:
 	EParticleSimulatePhase::Type PhaseToBroadcastTemporalEffect = EParticleSimulatePhase::Last;
 #endif
 };
+
diff --git a/Engine/Source/Runtime/Engine/Private/Particles/FXSystemSet.cpp b/Engine/Source/Runtime/Engine/Private/Particles/FXSystemSet.cpp
--- a/Engine/Source/Runtime/Engine/Private/Particles/FXSystemSet.cpp
+++ b/Engine/Source/Runtime/Engine/Private/Particles/FXSystemSet.cpp
@@ -96,6 +96,15 @@ void FFXSystemSet::PreInitViews(FRHICommandListImmediate& RHICmdList, bool bAllo
 	}
 }
 
+void FFXSystemSet::PostInitViews(FRHICommandListImmediate& RHICmdList, FRHIUniformBuffer* ViewUniformBuffer, bool bAllowGPUParticleUpdate)
+{
+	for (FFXSystemInterface* FXSystem : FXSystems)
+	{
+		check(FXSystem);
+		FXSystem->PostInitViews(RHICmdList, ViewUniformBuffer, bAllowGPUParticleUpdate);
+	}
+}
+
 bool FFXSystemSet::UsesGlobalDistanceField() const
 {
 	for (FFXSystemInterface* FXSystem : FXSystems)
@@ -109,6 +118,32 @@ bool FFXSystemSet::UsesGlobalDistanceField() const
 	return false;
 }
 
+bool FFXSystemSet::UsesDepthBuffer() const
+{
+	for (FFXSystemInterface* FXSystem : FXSystems)
+	{
+		check(FXSystem);
+		if (FXSystem->UsesDepthBuffer())
+		{
+			return true;
+		}
+	}
+	return false;
+}
+
+bool FFXSystemSet::RequiresEarlyViewUniformBuffer() const
+{
+	for (FFXSystemInterface* FXSystem : FXSystems)
+	{
+		check(FXSystem);
+		if (FXSystem->RequiresEarlyViewUniformBuffer())
+		{
+			return true;
+		}
+	}
+	return false;
+}
+
 void FFXSystemSet::PreRender(FRHICommandListImmediate& RHICmdList, const class FGlobalDistanceFieldParameterData* GlobalDistanceFieldParameterData, bool bAllowGPUParticleSceneUpdate)
 {
 	for (FFXSystemInterface* FXSystem : FXSystems)
diff --git a/Engine/Source/Runtime/Engine/Private/Particles/FXSystemSet.h b/Engine/Source/Runtime/Engine/Private/Particles/FXSystemSet.h
--- a/Engine/Source/Runtime/Engine/Private/Particles/FXSystemSet.h
+++ b/Engine/Source/Runtime/Engine/Private/Particles/FXSystemSet.h
@@ -31,7 +31,10 @@ public:
 	virtual void RemoveVectorField(UVectorFieldComponent* VectorFieldComponent) override;
 	virtual void UpdateVectorField(UVectorFieldComponent* VectorFieldComponent) override;
 	virtual void PreInitViews(FRHICommandListImmediate& RHICmdList, bool bAllowGPUParticleUpdate) override;
+	virtual void PostInitViews(FRHICommandListImmediate& RHICmdList, FRHIUniformBuffer* ViewUniformBuffer, bool bAllowGPUParticleUpdate) override;
 	virtual bool UsesGlobalDistanceField() const override;
+	virtual bool UsesDepthBuffer() const override;
+	virtual bool RequiresEarlyViewUniformBuffer() const override;
 	virtual void PreRender(FRHICommandListImmediate& RHICmdList, const class FGlobalDistanceFieldParameterData* GlobalDistanceFieldParameterData, bool bAllowGPUParticleSceneUpdate) override;
 	virtual void PostRenderOpaque(
 		FRHICommandListImmediate& RHICmdList, 
diff --git a/Engine/Source/Runtime/Engine/Private/Particles/ParticleGpuSimulation.cpp b/Engine/Source/Runtime/Engine/Private/Particles/ParticleGpuSimulation.cpp
--- a/Engine/Source/Runtime/Engine/Private/Particles/ParticleGpuSimulation.cpp
+++ b/Engine/Source/Runtime/Engine/Private/Particles/ParticleGpuSimulation.cpp
@@ -4684,6 +4684,27 @@ bool FFXSystem::UsesGlobalDistanceFieldInternal() const
 	return false;
 }
 
+bool FFXSystem::UsesDepthBufferInternal() const
+{
+	for (TSparseArray<FParticleSimulationGPU*>::TConstIterator It(GPUSimulations); It; ++It)
+	{
+		const FParticleSimulationGPU* Simulation = *It;
+
+		if (Simulation->SimulationPhase == EParticleSimulatePhase::CollisionDepthBuffer
+			&& Simulation->TileVertexBuffer.AlignedTileCount > 0)
+		{
+			return true;
+		}
+	}
+
+	return false;
+}
+
+bool FFXSystem::RequiresEarlyViewUniformBufferInternal() const
+{
+	return false;
+}
+
 void FFXSystem::PrepareGPUSimulation(FRHICommandListImmediate& RHICmdList)
 {
 	// Grab resources.
diff --git a/Engine/Source/Runtime/Engine/Public/FXSystem.h b/Engine/Source/Runtime/Engine/Public/FXSystem.h
--- a/Engine/Source/Runtime/Engine/Public/FXSystem.h
+++ b/Engine/Source/Runtime/Engine/Public/FXSystem.h
@@ -205,8 +205,14 @@ public:
 	 */
 	virtual void PreInitViews(FRHICommandListImmediate& RHICmdList, bool bAllowGPUParticleUpdate) = 0;
 
+	virtual void PostInitViews(FRHICommandListImmediate& RHICmdList, FRHIUniformBuffer* ViewUniformBuffer, bool bAllowGPUParticleUpdate) = 0;
+
 	virtual bool UsesGlobalDistanceField() const = 0;
 
+	virtual bool UsesDepthBuffer() const = 0;
+
+	virtual bool RequiresEarlyViewUniformBuffer() const = 0;
+
 	/**
 	 * Notification from the renderer that it is about to draw FX belonging to
 	 * this system.
diff --git a/Engine/Source/Runtime/Renderer/Private/SceneVisibility.cpp b/Engine/Source/Runtime/Renderer/Private/SceneVisibility.cpp
--- a/Engine/Source/Runtime/Renderer/Private/SceneVisibility.cpp
+++ b/Engine/Source/Runtime/Renderer/Private/SceneVisibility.cpp
@@ -4226,6 +4226,16 @@ bool FDeferredShadingSceneRenderer::InitViews(FRHICommandListImmediate& RHICmdLi
 
 	RHICmdList.ImmediateFlush(EImmediateFlushType::DispatchToRHIThread);
 
+	{
+		// This is to init the ViewUniformBuffer before rendering for the Niagara compute shader.
+		// This needs to run before ComputeViewVisibility() is called, but the views normally initialize the ViewUniformBuffer after that (at the end of this method).
+		if (Scene->FXSystem && Scene->FXSystem->RequiresEarlyViewUniformBuffer() && Views.IsValidIndex(0))
+		{
+			Views[0].InitRHIResources();
+			Scene->FXSystem->PostInitViews(RHICmdList, Views[0].ViewUniformBuffer, Views[0].AllowGPUParticleUpdate() && !ViewFamily.EngineShowFlags.HitProxies);
+		}
+	}
+	
 	FViewVisibleCommandsPerView ViewCommandsPerView;
 	ViewCommandsPerView.SetNum(Views.Num());
 
