Fixed crash that could occurr when multiple Niagara components are GCd from underneath the scalability manager.

diff --git a/Engine/Plugins/FX/Niagara/Source/Niagara/Private/NiagaraScalabilityManager.cpp b/Engine/Plugins/FX/Niagara/Source/Niagara/Private/NiagaraScalabilityManager.cpp
--- a/Engine/Plugins/FX/Niagara/Source/Niagara/Private/NiagaraScalabilityManager.cpp
+++ b/Engine/Plugins/FX/Niagara/Source/Niagara/Private/NiagaraScalabilityManager.cpp
@@ -56,8 +56,10 @@ void FNiagaraScalabilityManager::UnregisterAt(int32 IndexToRemove)
 
 	if (ManagedComponents.IsValidIndex(IndexToRemove))
 	{
-		check(ManagedComponents[IndexToRemove] != nullptr);
-		ManagedComponents[IndexToRemove]->ScalabilityManagerHandle = IndexToRemove;
+		if ((ManagedComponents[IndexToRemove] != nullptr))//Possible this has been GCd. It will be removed later if so.
+		{
+			ManagedComponents[IndexToRemove]->ScalabilityManagerHandle = IndexToRemove;
+		}
 	}
 }
 
