- More explicit information about NiagaraSystems that fail to activate due to an invalid system

#rnx

diff --git a/Engine/Plugins/FX/Niagara/Source/Niagara/Private/NiagaraSystemInstance.cpp b/Engine/Plugins/FX/Niagara/Source/Niagara/Private/NiagaraSystemInstance.cpp
--- a/Engine/Plugins/FX/Niagara/Source/Niagara/Private/NiagaraSystemInstance.cpp
+++ b/Engine/Plugins/FX/Niagara/Source/Niagara/Private/NiagaraSystemInstance.cpp
@@ -727,7 +727,7 @@ void FNiagaraSystemInstance::ResetInternal(bool bResetSimulations)
 	if (!System->IsValid())
 	{
 		SetRequestedExecutionState(ENiagaraExecutionState::Disabled);
-		UE_LOG(LogNiagara, Warning, TEXT("Failed to activate Niagara System due to invalid asset!"));
+		UE_LOG(LogNiagara, Warning, TEXT("Failed to activate Niagara System due to invalid asset! System(%s) Component(%s)"), *System->GetName(), *Component->GetFullName());
 		return;
 	}
 
@@ -861,7 +861,7 @@ void FNiagaraSystemInstance::ReInitInternal()
 	if (!System->IsValid())
 	{
 		SetRequestedExecutionState(ENiagaraExecutionState::Disabled);
-		UE_LOG(LogNiagara, Warning, TEXT("Failed to activate Niagara System due to invalid asset!"));
+		UE_LOG(LogNiagara, Warning, TEXT("Failed to activate Niagara System due to invalid asset! System(%s) Component(%s)"), *System->GetName(), *Component->GetFullName());
 		return;
 	}
 
