- Fix missing transition in UNiagaraDataInterfaceGrid2DCollection::FillTexture2D

#rnx

diff --git a/Engine/Plugins/FX/Niagara/Source/Niagara/Private/NiagaraDataInterfaceGrid2DCollection.cpp b/Engine/Plugins/FX/Niagara/Source/Niagara/Private/NiagaraDataInterfaceGrid2DCollection.cpp
--- a/Engine/Plugins/FX/Niagara/Source/Niagara/Private/NiagaraDataInterfaceGrid2DCollection.cpp
+++ b/Engine/Plugins/FX/Niagara/Source/Niagara/Private/NiagaraDataInterfaceGrid2DCollection.cpp
@@ -493,6 +493,8 @@ bool UNiagaraDataInterfaceGrid2DCollection::FillTexture2D(const UNiagaraComponen
 			int StartY = TileIndexY * Grid2DInstanceData->NumCellsY;
 			CopyInfo.SourcePosition = FIntVector(StartX, StartY, 0);
 
+			RHICmdList.TransitionResource(EResourceTransitionAccess::EReadable, Grid2DInstanceData->CurrentData->GridBuffer.Buffer);
+
 			RHICmdList.CopyTexture(Grid2DInstanceData->CurrentData->GridBuffer.Buffer, Dest->Resource->TextureRHI, CopyInfo);
 		}
 	});
@@ -548,7 +550,8 @@ bool UNiagaraDataInterfaceGrid2DCollection::FillRawTexture2D(const UNiagaraCompo
 
 		if (Dest && Dest->Resource && Grid2DInstanceDataTmp && Grid2DInstanceDataTmp->CurrentData)
 		{
-			
+			RHICmdList.TransitionResource(EResourceTransitionAccess::EReadable, Grid2DInstanceDataTmp->CurrentData->GridBuffer.Buffer);
+
 			FRHICopyTextureInfo CopyInfo;
 			RHICmdList.CopyTexture(Grid2DInstanceDataTmp->CurrentData->GridBuffer.Buffer, Dest->Resource->TextureRHI, CopyInfo);
 		}
