Fix for copying of DI curves not also copying the sampling limit, which was preventing the LUTs from working properly in a cooked build when the asset was being assigned dyanmically.

diff --git a/Engine/Plugins/FX/Niagara/Source/Niagara/Private/NiagaraDataInterfaceCurveBase.cpp b/Engine/Plugins/FX/Niagara/Source/Niagara/Private/NiagaraDataInterfaceCurveBase.cpp
--- a/Engine/Plugins/FX/Niagara/Source/Niagara/Private/NiagaraDataInterfaceCurveBase.cpp
+++ b/Engine/Plugins/FX/Niagara/Source/Niagara/Private/NiagaraDataInterfaceCurveBase.cpp
@@ -83,6 +83,7 @@ bool UNiagaraDataInterfaceCurveBase::CopyToInternal(UNiagaraDataInterface* Desti
 	UNiagaraDataInterfaceCurveBase* DestinationTyped = CastChecked<UNiagaraDataInterfaceCurveBase>(Destination);
 	DestinationTyped->bUseLUT = bUseLUT;
 	DestinationTyped->ShaderLUT = ShaderLUT;
+	DestinationTyped->LUTNumSamplesMinusOne = LUTNumSamplesMinusOne;
 #if WITH_EDITORONLY_DATA
 	DestinationTyped->bOptimizeLUT = bOptimizeLUT;
 	DestinationTyped->bOverrideOptimizeThreshold = bOverrideOptimizeThreshold;
