Access niagara shader compilation manager through a static Get() function rather than through a globally initialized variable. Stops the compilation manager crashing in monolithic builds where it attempts to access FCommandLine during global construction where it hasn't been initialized.

diff --git a/Engine/Plugins/FX/Niagara/Source/NiagaraEditor/Private/NiagaraHlslTranslator.h b/Engine/Plugins/FX/Niagara/Source/NiagaraEditor/Private/NiagaraHlslTranslator.h
--- a/Engine/Plugins/FX/Niagara/Source/NiagaraEditor/Private/NiagaraHlslTranslator.h
+++ b/Engine/Plugins/FX/Niagara/Source/NiagaraEditor/Private/NiagaraHlslTranslator.h
@@ -36,8 +36,8 @@ class FNiagaraShaderProcessorTickable : FTickableEditorObject
 
 	virtual void Tick(float DeltaSeconds) override
 	{
-		GNiagaraShaderCompilationManager.Tick(DeltaSeconds);
-		GNiagaraShaderCompilationManager.ProcessAsyncResults();
+		FNiagaraShaderCompilationManager::Get().Tick(DeltaSeconds);
+		FNiagaraShaderCompilationManager::Get().ProcessAsyncResults();
 	}
 
 	virtual TStatId GetStatId() const
diff --git a/Engine/Plugins/FX/Niagara/Source/NiagaraShader/Private/NiagaraShader.cpp b/Engine/Plugins/FX/Niagara/Source/NiagaraShader/Private/NiagaraShader.cpp
--- a/Engine/Plugins/FX/Niagara/Source/NiagaraShader/Private/NiagaraShader.cpp
+++ b/Engine/Plugins/FX/Niagara/Source/NiagaraShader/Private/NiagaraShader.cpp
@@ -913,14 +913,14 @@ void FNiagaraShaderMap::Compile(
 			// Mark as not having been compiled
 			bCompiledSuccessfully = false;
   
-			GNiagaraShaderCompilationManager.AddJobs(NewJobs);
+			FNiagaraShaderCompilationManager::Get().AddJobs(NewJobs);
   
 			// Compile the shaders for this shader map now if not deferring and deferred compiles are not enabled globally
 			if (bSynchronousCompile)
 			{
 				TArray<int32> CurrentShaderMapId;
 				CurrentShaderMapId.Add(CompilingId);
-				GNiagaraShaderCompilationManager.FinishCompilation(*FriendlyName, CurrentShaderMapId);
+				FNiagaraShaderCompilationManager::Get().FinishCompilation(*FriendlyName, CurrentShaderMapId);
 			}
 		}
 	}
diff --git a/Engine/Plugins/FX/Niagara/Source/NiagaraShader/Private/NiagaraShaderCompilationManager.cpp b/Engine/Plugins/FX/Niagara/Source/NiagaraShader/Private/NiagaraShaderCompilationManager.cpp
--- a/Engine/Plugins/FX/Niagara/Source/NiagaraShader/Private/NiagaraShaderCompilationManager.cpp
+++ b/Engine/Plugins/FX/Niagara/Source/NiagaraShader/Private/NiagaraShaderCompilationManager.cpp
@@ -78,8 +78,11 @@ public:
 	}
 };
 
-
-NIAGARASHADER_API FNiagaraShaderCompilationManager GNiagaraShaderCompilationManager;
+FNiagaraShaderCompilationManager& FNiagaraShaderCompilationManager::Get()
+{
+	static FNiagaraShaderCompilationManager Instance;
+	return Instance;
+}
 
 bool FNiagaraShaderCompilationManager::IsAsyncHackEnabled()const
 {
diff --git a/Engine/Plugins/FX/Niagara/Source/NiagaraShader/Private/NiagaraShared.cpp b/Engine/Plugins/FX/Niagara/Source/NiagaraShader/Private/NiagaraShared.cpp
--- a/Engine/Plugins/FX/Niagara/Source/NiagaraShader/Private/NiagaraShared.cpp
+++ b/Engine/Plugins/FX/Niagara/Source/NiagaraShader/Private/NiagaraShared.cpp
@@ -378,7 +378,7 @@ void FNiagaraShaderScript::FinishCompilation()
 		}
 		// Block until the shader maps that we will save have finished being compiled
 		// NIAGARATODO: implement when async compile works
-		GNiagaraShaderCompilationManager.FinishCompilation(*GetFriendlyName(), ShaderMapIdsToFinish);
+		FNiagaraShaderCompilationManager::Get().FinishCompilation(*GetFriendlyName(), ShaderMapIdsToFinish);
 
 		// Shouldn't have anything left to do...
 		TArray<int32> ShaderMapIdsToFinish2;
diff --git a/Engine/Plugins/FX/Niagara/Source/NiagaraShader/Public/NiagaraShaderCompilationManager.h b/Engine/Plugins/FX/Niagara/Source/NiagaraShader/Public/NiagaraShaderCompilationManager.h
--- a/Engine/Plugins/FX/Niagara/Source/NiagaraShader/Public/NiagaraShaderCompilationManager.h
+++ b/Engine/Plugins/FX/Niagara/Source/NiagaraShader/Public/NiagaraShaderCompilationManager.h
@@ -122,6 +122,7 @@ public:
 	NIAGARASHADER_API void Tick(float DeltaSeconds = 0.0f, bool bBlock=false);
 	NIAGARASHADER_API void AddJobs(TArray<class FShaderCommonCompileJob*> InNewJobs);
 	NIAGARASHADER_API void ProcessAsyncResults();
+	NIAGARASHADER_API static FNiagaraShaderCompilationManager& Get();
 
 	void FinishCompilation(const TCHAR* ScriptName, const TArray<int32>& ShaderMapIdsToFinishCompiling);
 
@@ -155,7 +156,4 @@ private:
 	//Handle to the last task we kicked off to compile previously.
 	//We can only have one compilation task going at a time as (apparently) the compilation is very non thread-safe
 	FGraphEventRef AsyncWork;
-};
-
-
-extern NIAGARASHADER_API FNiagaraShaderCompilationManager GNiagaraShaderCompilationManager;
+};
\ No newline at end of file
