Fixed a function definition typo

diff --git a/Engine/Plugins/FX/Niagara/Source/Niagara/Private/NiagaraDataInterfaceCamera.cpp b/Engine/Plugins/FX/Niagara/Source/Niagara/Private/NiagaraDataInterfaceCamera.cpp
--- a/Engine/Plugins/FX/Niagara/Source/Niagara/Private/NiagaraDataInterfaceCamera.cpp
+++ b/Engine/Plugins/FX/Niagara/Source/Niagara/Private/NiagaraDataInterfaceCamera.cpp
@@ -187,7 +187,7 @@ bool UNiagaraDataInterfaceCamera::GetFunctionHLSL(const FName& DefinitionFunctio
 	if (DefinitionFunctionName == GetClipSpaceTransformsName)
 	{
 		static const TCHAR *FormatSample = TEXT(R"(
-			void {FunctionName}out float4x4 Out_WorldToClipTransform, out float4x4 Out_TranslatedWorldToClipTransform, out float4x4 Out_ClipToWorldTransform, out float4x4 Out_ClipToViewTransform,
+			void {FunctionName}(out float4x4 Out_WorldToClipTransform, out float4x4 Out_TranslatedWorldToClipTransform, out float4x4 Out_ClipToWorldTransform, out float4x4 Out_ClipToViewTransform,
 				out float4x4 Out_ClipToTranslatedWorldTransform, out float4x4 Out_ScreenToWorldTransform, out float4x4 Out_ScreenToTranslatedWorldTransform, out float4x4 Out_ClipToPreviousClipTransform)
 			{
 				Out_WorldToClipTransform = View.WorldToClip;
