Added TAAJitter to camera DI for Wyeth's lens flare system

diff --git a/Engine/Plugins/FX/Niagara/Source/Niagara/Classes/NiagaraDataInterfaceCamera.h b/Engine/Plugins/FX/Niagara/Source/Niagara/Classes/NiagaraDataInterfaceCamera.h
--- a/Engine/Plugins/FX/Niagara/Source/Niagara/Classes/NiagaraDataInterfaceCamera.h
+++ b/Engine/Plugins/FX/Niagara/Source/Niagara/Classes/NiagaraDataInterfaceCamera.h
@@ -46,6 +46,7 @@ public:
 	void GetViewPropertiesGPU(FVectorVMContext& Context);
 	void GetClipSpaceTransformsGPU(FVectorVMContext& Context);
 	void GetViewSpaceTransformsGPU(FVectorVMContext& Context);
+	void GetTAAJitter(FVectorVMContext& Context);
 	
 private:
 	static const FName GetViewPropertiesName;
@@ -53,6 +54,7 @@ private:
 	static const FName GetViewSpaceTransformsName;
 	static const FName GetCameraPropertiesName;
 	static const FName GetFieldOfViewName;
+	static const FName GetTAAJitterName;
 };
 
 struct FNiagaraDataIntefaceProxyCameraQuery : public FNiagaraDataInterfaceProxy
diff --git a/Engine/Plugins/FX/Niagara/Source/Niagara/Private/NiagaraDataInterfaceCamera.cpp b/Engine/Plugins/FX/Niagara/Source/Niagara/Private/NiagaraDataInterfaceCamera.cpp
--- a/Engine/Plugins/FX/Niagara/Source/Niagara/Private/NiagaraDataInterfaceCamera.cpp
+++ b/Engine/Plugins/FX/Niagara/Source/Niagara/Private/NiagaraDataInterfaceCamera.cpp
@@ -12,6 +12,7 @@ const FName UNiagaraDataInterfaceCamera::GetClipSpaceTransformsName(TEXT("GetCli
 const FName UNiagaraDataInterfaceCamera::GetViewSpaceTransformsName(TEXT("GetViewSpaceTransformsGPU"));
 const FName UNiagaraDataInterfaceCamera::GetCameraPropertiesName(TEXT("GetCameraPropertiesCPU/GPU"));
 const FName UNiagaraDataInterfaceCamera::GetFieldOfViewName(TEXT("GetFieldOfView"));
+const FName UNiagaraDataInterfaceCamera::GetTAAJitterName(TEXT("GetTAAJitter"));
 
 UNiagaraDataInterfaceCamera::UNiagaraDataInterfaceCamera(FObjectInitializer const& ObjectInitializer)
 	: Super(ObjectInitializer)
@@ -150,6 +151,18 @@ void UNiagaraDataInterfaceCamera::GetFunctions(TArray<FNiagaraFunctionSignature>
 	Sig.Outputs.Add(FNiagaraVariable(FNiagaraTypeDefinition::GetVec3Def(), TEXT("Up Vector World")));
 	Sig.Outputs.Add(FNiagaraVariable(FNiagaraTypeDefinition::GetVec3Def(), TEXT("Right Vector World")));
 	OutFunctions.Add(Sig);
+
+	Sig = FNiagaraFunctionSignature();
+	Sig.Name = GetTAAJitterName;
+#if WITH_EDITORONLY_DATA
+	Sig.Description = NSLOCTEXT("Niagara", "GetTAAJitterDescription", "This function returns the TAA jitter values of the currently active camera.");
+#endif
+	Sig.bMemberFunction = true;
+	Sig.bRequiresContext = false;
+	Sig.Inputs.Add(FNiagaraVariable(FNiagaraTypeDefinition(GetClass()), TEXT("Camera interface")));
+	Sig.Outputs.Add(FNiagaraVariable(FNiagaraTypeDefinition::GetVec2Def(), TEXT("Current TAA Jitter (clip)")));
+	Sig.Outputs.Add(FNiagaraVariable(FNiagaraTypeDefinition::GetVec2Def(), TEXT("Previous TAA Jitter (clip)")));
+	OutFunctions.Add(Sig);
 }
 
 bool UNiagaraDataInterfaceCamera::GetFunctionHLSL(const FName& DefinitionFunctionName, FString InstanceFunctionName, FNiagaraDataInterfaceGPUParamInfo& ParamInfo, FString& OutHLSL)
@@ -234,6 +247,18 @@ bool UNiagaraDataInterfaceCamera::GetFunctionHLSL(const FName& DefinitionFunctio
 		OutHLSL += FString::Format(FormatSample, ArgsSample);
 		return true;
 	}
+	if (DefinitionFunctionName == GetTAAJitterName)
+	{
+		static const TCHAR *FormatSample = TEXT(R"(
+			void {FunctionName}(out float2 Out_CurrentJitter, out float2 Out_PreviousJitter)
+			{
+				Out_CurrentJitter = View.TemporalAAJitter.xy;
+				Out_PreviousJitter = View.TemporalAAJitter.zw;
+			}
+		)");
+		OutHLSL += FString::Format(FormatSample, ArgsSample);
+		return true;
+	}
 	return false;
 }
 
@@ -242,6 +267,7 @@ DEFINE_NDI_DIRECT_FUNC_BINDER(UNiagaraDataInterfaceCamera, GetCameraProperties);
 DEFINE_NDI_DIRECT_FUNC_BINDER(UNiagaraDataInterfaceCamera, GetViewPropertiesGPU);
 DEFINE_NDI_DIRECT_FUNC_BINDER(UNiagaraDataInterfaceCamera, GetClipSpaceTransformsGPU);
 DEFINE_NDI_DIRECT_FUNC_BINDER(UNiagaraDataInterfaceCamera, GetViewSpaceTransformsGPU);
+DEFINE_NDI_DIRECT_FUNC_BINDER(UNiagaraDataInterfaceCamera, GetTAAJitter);
 void UNiagaraDataInterfaceCamera::GetVMExternalFunction(const FVMExternalFunctionBindingInfo& BindingInfo, void* InstanceData, FVMExternalFunction &OutFunc)
 {
 	if (BindingInfo.Name == GetFieldOfViewName)
@@ -264,6 +290,10 @@ void UNiagaraDataInterfaceCamera::GetVMExternalFunction(const FVMExternalFunctio
 	{
 		NDI_FUNC_BINDER(UNiagaraDataInterfaceCamera, GetViewPropertiesGPU)::Bind(this, OutFunc);
 	}
+	else if (BindingInfo.Name == GetTAAJitterName)
+	{
+		NDI_FUNC_BINDER(UNiagaraDataInterfaceCamera, GetTAAJitter)::Bind(this, OutFunc);
+	}
 	else
 	{
 		UE_LOG(LogNiagara, Error, TEXT("Could not find data interface external function. Received Name: %s"), *BindingInfo.Name.ToString());
@@ -396,5 +426,26 @@ void UNiagaraDataInterfaceCamera::GetViewSpaceTransformsGPU(FVectorVMContext& Co
 	}
 }
 
-// ------------------------------------------------------------
+void UNiagaraDataInterfaceCamera::GetTAAJitter(FVectorVMContext& Context)
+{
+	VectorVM::FUserPtrHandler<CameraDataInterface_InstanceData> InstData(Context);
+	TArray<VectorVM::FExternalFuncRegisterHandler<float>> OutParams;
+
+	constexpr int32 ElementCount = 2 * sizeof(FVector2D) / sizeof(float);
+	OutParams.Reserve(ElementCount);
+	for (int i = 0; i < ElementCount; i++)
+	{
+		OutParams.Emplace(Context);
+	}
+
+	for (int32 k = 0; k < Context.NumInstances; ++k)
+	{
+		for (int i = 0; i < ElementCount; i++)
+		{
+			*OutParams[i].GetDestAndAdvance() = 0;
+		}
+	}
+}
+
+
 
