- Integrate Graham's CL 10707088 NiagaraDataInterfaceVectorField crashing while trying to bind vector field texture with a null render thread representation (likely a timing or init order problem)

fyi graham.wihlidal,shaun.kime
#rnx

diff --git a/Engine/Plugins/FX/Niagara/Source/Niagara/Private/NiagaraDataInterfaceVectorField.cpp b/Engine/Plugins/FX/Niagara/Source/Niagara/Private/NiagaraDataInterfaceVectorField.cpp
--- a/Engine/Plugins/FX/Niagara/Source/Niagara/Private/NiagaraDataInterfaceVectorField.cpp
+++ b/Engine/Plugins/FX/Niagara/Source/Niagara/Private/NiagaraDataInterfaceVectorField.cpp
@@ -360,7 +360,14 @@ struct FNiagaraDataInterfaceParametersCS_VectorField : public FNiagaraDataInterf
 		// Note: There is a flush in PreEditChange to make sure everything is synced up at this point 
 
 		// Get and set 3D texture handle from the currently bound vector field.
-		SetTextureParameter(RHICmdList, ComputeShaderRHI, VectorFieldTexture, VFDI->TextureRHI); 
+		if (VFDI->TextureRHI)
+		{
+			SetTextureParameter(RHICmdList, ComputeShaderRHI, VectorFieldTexture, VFDI->TextureRHI);
+		}
+		else
+		{
+			SetTextureParameter(RHICmdList, ComputeShaderRHI, VectorFieldTexture, GBlackVolumeTexture->TextureRHI);
+		}
 		
 		// Get and set sampler state
 		FRHISamplerState* SamplerState = SamplerStates[int(VFDI->bTileX) + 2 * int(VFDI->bTileY) + 4 * int(VFDI->bTileZ)];
