Added missing SupportType funciton to Niagara Emitter Track.

diff --git a/Engine/Plugins/FX/Niagara/Source/NiagaraEditor/Private/Sequencer/NiagaraSequence/MovieSceneNiagaraEmitterTrack.cpp b/Engine/Plugins/FX/Niagara/Source/NiagaraEditor/Private/Sequencer/NiagaraSequence/MovieSceneNiagaraEmitterTrack.cpp
--- a/Engine/Plugins/FX/Niagara/Source/NiagaraEditor/Private/Sequencer/NiagaraSequence/MovieSceneNiagaraEmitterTrack.cpp
+++ b/Engine/Plugins/FX/Niagara/Source/NiagaraEditor/Private/Sequencer/NiagaraSequence/MovieSceneNiagaraEmitterTrack.cpp
@@ -14,6 +14,7 @@
 #include "NiagaraScriptSource.h"
 #include "ViewModels/Stack/NiagaraStackGraphUtilities.h"
 #include "NiagaraEditorStyle.h"
+#include "Sections/MovieSceneNiagaraEmitterSection.h"
 
 #include "ISequencerSection.h"
 #include "SequencerSectionPainter.h"
@@ -172,6 +173,11 @@ void UMovieSceneNiagaraEmitterTrack::AddSection(UMovieSceneSection& Section)
 	bSectionsWereModified = true;
 }
 
+bool UMovieSceneNiagaraEmitterTrack::SupportsType(TSubclassOf<UMovieSceneSection> SectionClass) const
+{
+	return SectionClass == UMovieSceneNiagaraEmitterSection::StaticClass();
+}
+
 void UMovieSceneNiagaraEmitterTrack::RemoveSection(UMovieSceneSection& Section)
 {
 	Sections.Remove(&Section);
diff --git a/Engine/Plugins/FX/Niagara/Source/NiagaraEditor/Private/Sequencer/NiagaraSequence/MovieSceneNiagaraEmitterTrack.h b/Engine/Plugins/FX/Niagara/Source/NiagaraEditor/Private/Sequencer/NiagaraSequence/MovieSceneNiagaraEmitterTrack.h
--- a/Engine/Plugins/FX/Niagara/Source/NiagaraEditor/Private/Sequencer/NiagaraSequence/MovieSceneNiagaraEmitterTrack.h
+++ b/Engine/Plugins/FX/Niagara/Source/NiagaraEditor/Private/Sequencer/NiagaraSequence/MovieSceneNiagaraEmitterTrack.h
@@ -76,6 +76,7 @@ public:
 	virtual void RemoveAllAnimationData() override { }
 	virtual bool HasSection(const UMovieSceneSection& Section) const override;
 	virtual void AddSection(UMovieSceneSection& Section) override;
+	virtual bool SupportsType(TSubclassOf<UMovieSceneSection> SectionClass) const override;
 	virtual void RemoveSection(UMovieSceneSection& Section) override;
 	virtual void RemoveSectionAt(int32 SectionIndex) override;
 	virtual bool IsEmpty() const override;
