Fixed tessellation numerical issues with Niagara ribbons that are too straight.

diff --git a/Engine/Plugins/FX/Niagara/Source/Niagara/Private/NiagaraRendererRibbons.cpp b/Engine/Plugins/FX/Niagara/Source/Niagara/Private/NiagaraRendererRibbons.cpp
--- a/Engine/Plugins/FX/Niagara/Source/Niagara/Private/NiagaraRendererRibbons.cpp
+++ b/Engine/Plugins/FX/Niagara/Source/Niagara/Private/NiagaraRendererRibbons.cpp
@@ -922,14 +922,15 @@ void FNiagaraRendererRibbons::CreatePerViewResources(
 				return FMath::Max<float>(1.f, FMath::Max(TessellationTwistAngle, TessellationAngle) / FMath::Max<float>(SMALL_NUMBER, TessellationMinAngle));
 			}
 		}();
+		const float MAX_CURVATURE_FACTOR = 0.002f; // This will clamp the curvature to around 2.5 km and avoid numerical issues.
 		const float ViewDistance = SceneProxy->GetBounds().ComputeSquaredDistanceFromBoxToPoint(ViewOriginForDistanceCulling);
 		const float MaxDisplacementError = FMath::Max(GNiagaraRibbonTessellationMinDisplacementError, ScreenPercentage * FMath::Sqrt(ViewDistance) / View->LODDistanceFactor);
-		float Tess = TessellationAngle / FMath::Max(SMALL_NUMBER, AcosFast(TessellationCurvature / (TessellationCurvature + MaxDisplacementError)));
+		float Tess = TessellationAngle / FMath::Max(MAX_CURVATURE_FACTOR, AcosFast(TessellationCurvature / (TessellationCurvature + MaxDisplacementError)));
 		// FMath::RoundUpToPowerOfTwo ? This could avoid vertices moving around as tesselation increases
 
 		if (TessellationTwistAngle > 0 && TessellationTwistCurvature > 0)
 		{
-			const float TwistTess = TessellationTwistAngle / FMath::Max(SMALL_NUMBER, AcosFast(TessellationTwistCurvature / (TessellationTwistCurvature + MaxDisplacementError)));
+			const float TwistTess = TessellationTwistAngle / FMath::Max(MAX_CURVATURE_FACTOR, AcosFast(TessellationTwistCurvature / (TessellationTwistCurvature + MaxDisplacementError)));
 			Tess = FMath::Max(TwistTess, Tess);
 		}
 		SegmentTessellation = FMath::Clamp<int32>(FMath::RoundToInt(Tess), FMath::RoundToInt(MinTesselation), TessellationFactor);
