- Fix missing resource transitions for Niagara Data buffers when shader stages are used
- Ensure current / destination data are swapped after written so subsequent stages read from the correct data

#rnx

diff --git a/Engine/Plugins/FX/Niagara/Source/Niagara/Private/NiagaraEmitterInstanceBatcher.cpp b/Engine/Plugins/FX/Niagara/Source/Niagara/Private/NiagaraEmitterInstanceBatcher.cpp
--- a/Engine/Plugins/FX/Niagara/Source/Niagara/Private/NiagaraEmitterInstanceBatcher.cpp
+++ b/Engine/Plugins/FX/Niagara/Source/Niagara/Private/NiagaraEmitterInstanceBatcher.cpp
@@ -283,6 +283,7 @@ void NiagaraEmitterInstanceBatcher::DispatchMultipleStages(const FNiagaraGPUSyst
 
 		const uint32 NumStages = Instance->Context->MaxUpdateIterations;
 		const uint32 DefaultShaderStageIndex = Instance->Context->DefaultShaderStageIndex;
+		bool bTransitionCurrentBuffer = false;
 
 		for (uint32 ShaderStageIndex = 0; ShaderStageIndex < NumStages; ++ShaderStageIndex)
 		{
@@ -298,16 +299,35 @@ void NiagaraEmitterInstanceBatcher::DispatchMultipleStages(const FNiagaraGPUSyst
 
 			PreStageInterface(Tick, Instance, RHICmdList, ComputeShader, ShaderStageIndex);
 
-			if (!IterationInterface)
+			// If we are reading from current data we need to transition the resource if it was previously written to
+			if (bTransitionCurrentBuffer && (ComputeShader->FloatInputBufferParam.IsBound() || ComputeShader->IntInputBufferParam.IsBound()))
 			{
-				if (HasRunParticleStage)
+				bTransitionCurrentBuffer = false;
+
+				TArray<FRHIUnorderedAccessView*, TInlineAllocator<2>> Resources;
+				if (Instance->CurrentData->GetGPUBufferFloat().UAV.IsValid())
+				{
+					Resources.Add(Instance->CurrentData->GetGPUBufferFloat().UAV);
+				}
+				if (Instance->CurrentData->GetGPUBufferInt().UAV.IsValid())
+				{
+					Resources.Add(Instance->CurrentData->GetGPUBufferInt().UAV);
+				}
+				if (Resources.Num() > 0)
 				{
-					FNiagaraDataBuffer* StoreData = Instance->CurrentData;
-					Instance->CurrentData = Instance->DestinationData;
-					Instance->DestinationData = StoreData;
+					RHICmdList.TransitionResources(EResourceTransitionAccess::ERWBarrier, EResourceTransitionPipeline::EComputeToCompute, Resources.GetData(), Resources.Num());
 				}
+			}
+
+			if (!IterationInterface)
+			{
 				Run(Tick, Instance, 0, Instance->DestinationData->GetNumInstances(), ComputeShader, RHICmdList, ViewUniformBuffer, Instance->SpawnInfo, false, DefaultShaderStageIndex, ShaderStageIndex,  nullptr, HasRunParticleStage);
 				HasRunParticleStage = true;
+				bTransitionCurrentBuffer = true;
+
+				FNiagaraDataBuffer* StoreData = Instance->CurrentData;
+				Instance->CurrentData = Instance->DestinationData;
+				Instance->DestinationData = StoreData;
 			}
 			else
 			{
