Fixed GC safety code in the scalability manager.

diff --git a/Engine/Plugins/FX/Niagara/Source/Niagara/Private/NiagaraComponent.cpp b/Engine/Plugins/FX/Niagara/Source/Niagara/Private/NiagaraComponent.cpp
--- a/Engine/Plugins/FX/Niagara/Source/Niagara/Private/NiagaraComponent.cpp
+++ b/Engine/Plugins/FX/Niagara/Source/Niagara/Private/NiagaraComponent.cpp
@@ -979,6 +979,7 @@ void UNiagaraComponent::UnregisterWithScalabilityManager()
 			WorldMan->UnregisterWithScalabilityManager(this);
 		}
 	}
+	ScalabilityManagerHandle = INDEX_NONE;//Just to be sure our state is unregistered.
 }
 
 void UNiagaraComponent::OnSystemComplete()
diff --git a/Engine/Plugins/FX/Niagara/Source/Niagara/Private/NiagaraScalabilityManager.cpp b/Engine/Plugins/FX/Niagara/Source/Niagara/Private/NiagaraScalabilityManager.cpp
--- a/Engine/Plugins/FX/Niagara/Source/Niagara/Private/NiagaraScalabilityManager.cpp
+++ b/Engine/Plugins/FX/Niagara/Source/Niagara/Private/NiagaraScalabilityManager.cpp
@@ -20,6 +20,18 @@ FNiagaraScalabilityManager::FNiagaraScalabilityManager()
 
 }
 
+FNiagaraScalabilityManager::~FNiagaraScalabilityManager()
+{
+	for(UNiagaraComponent* Component : ManagedComponents)
+	{
+		if (Component)
+		{
+			Component->ScalabilityManagerHandle = INDEX_NONE;
+		}
+	}
+	ManagedComponents.Empty();
+}
+
 void FNiagaraScalabilityManager::AddReferencedObjects(FReferenceCollector& Collector)
 {
 	Collector.AddReferencedObject(EffectType);
@@ -40,7 +52,6 @@ void FNiagaraScalabilityManager::Register(UNiagaraComponent* Component)
 void FNiagaraScalabilityManager::Unregister(UNiagaraComponent* Component)
 {
 	check(Component->ScalabilityManagerHandle != INDEX_NONE);
-	check(ManagedComponents.Num() == State.Num());
 
 	int32 IndexToRemove = Component->ScalabilityManagerHandle;
 	Component->ScalabilityManagerHandle = INDEX_NONE;
@@ -51,9 +62,18 @@ void FNiagaraScalabilityManager::UnregisterAt(int32 IndexToRemove)
 {
 	//UE_LOG(LogNiagara, Warning, TEXT("Unregistering Component %0xP at index %d (Replaced with %0xP)"), ManagedComponents[IndexToRemove], IndexToRemove, ManagedComponents.Num() > 1 ? ManagedComponents.Last() : nullptr);
 
-	ManagedComponents.RemoveAtSwap(IndexToRemove);
-	State.RemoveAtSwap(IndexToRemove);
+	check(ManagedComponents.Num() == State.Num());
+	if (ManagedComponents.IsValidIndex(IndexToRemove))
+	{
+		ManagedComponents.RemoveAtSwap(IndexToRemove);
+		State.RemoveAtSwap(IndexToRemove);
+	}
+	else
+	{
+		UE_LOG(LogNiagara, Warning, TEXT("Attempting to unregister an invalid index from the Scalability Manager. Index: %d - Num: %d"), IndexToRemove, ManagedComponents.Num());
+	}
 
+	//Redirect the component that is now at IndexToRemove to the correct index.
 	if (ManagedComponents.IsValidIndex(IndexToRemove))
 	{
 		if ((ManagedComponents[IndexToRemove] != nullptr))//Possible this has been GCd. It will be removed later if so.
@@ -92,28 +112,22 @@ void FNiagaraScalabilityManager::Update(FNiagaraWorldManager* WorldMan)
 
 	//Belt and braces paranoia code to ensure we're safe if a component or System is GCd but the component isn't unregistered for whatever reason.
 	int32 CompIdx = 0;
+	check(State.Num() == ManagedComponents.Num());
 	while (CompIdx < ManagedComponents.Num())
 	{
 		UNiagaraComponent* Component = ManagedComponents[CompIdx];
 		if (Component)
 		{
-			if (Component->GetAsset())
-			{
-				++CompIdx;
-			}
-			else
+			if (Component->GetAsset() == nullptr)
 			{
 				UE_LOG(LogNiagara, Warning, TEXT("Niagara System has been destroyed with components still registered to the scalability manager. Unregistering this component.\nComponent: %0xP - %s\nEffectType: %0xP - %s"),
 					Component, *Component->GetName(), EffectType, *EffectType->GetName());
 				Unregister(Component);
+				continue;
 			}
 		}
-		else
-		{
-			UE_LOG(LogNiagara, Warning, TEXT("Niagara Component has been destroyed while still registered to the scalability manager. Unregistering this component.\nComponent: %0xP \nEffectType: %0xP - %s"),
-				Component, EffectType, *EffectType->GetName());
-			UnregisterAt(CompIdx);
-		}
+
+		++CompIdx;
 	}
 
 	bool bNeedSortedSignificanceCull = false;
@@ -125,6 +139,13 @@ void FNiagaraScalabilityManager::Update(FNiagaraWorldManager* WorldMan)
 	for (int32 i = 0; i < ManagedComponents.Num(); ++i)
 	{
 		UNiagaraComponent* Component = ManagedComponents[i];
+		
+		//The GC can pull this ref from underneath us before the component unregisters itself during BeginDestroy().
+		if(!Component)
+		{
+			continue;
+		}
+
 		FNiagaraScalabilityState& CompState = State[i];
 #if DEBUG_SCALABILITY_STATE
 		CompState.bCulledByInstanceCount = false;
diff --git a/Engine/Plugins/FX/Niagara/Source/Niagara/Private/NiagaraWorldManager.cpp b/Engine/Plugins/FX/Niagara/Source/Niagara/Private/NiagaraWorldManager.cpp
--- a/Engine/Plugins/FX/Niagara/Source/Niagara/Private/NiagaraWorldManager.cpp
+++ b/Engine/Plugins/FX/Niagara/Source/Niagara/Private/NiagaraWorldManager.cpp
@@ -52,6 +52,7 @@ FDelegateHandle FNiagaraWorldManager::OnWorldCleanupHandle;
 FDelegateHandle FNiagaraWorldManager::OnPreWorldFinishDestroyHandle;
 FDelegateHandle FNiagaraWorldManager::OnWorldBeginTearDownHandle;
 FDelegateHandle FNiagaraWorldManager::TickWorldHandle;
+FDelegateHandle FNiagaraWorldManager::PostGCHandle;
 TMap<class UWorld*, class FNiagaraWorldManager*> FNiagaraWorldManager::WorldManagers;
 
 TGlobalResource<FNiagaraViewDataMgr> GNiagaraViewDataManager;
@@ -166,6 +167,8 @@ void FNiagaraWorldManager::OnStartup()
 	OnPreWorldFinishDestroyHandle = FWorldDelegates::OnPreWorldFinishDestroy.AddStatic(&FNiagaraWorldManager::OnPreWorldFinishDestroy);
 	OnWorldBeginTearDownHandle = FWorldDelegates::OnWorldBeginTearDown.AddStatic(&FNiagaraWorldManager::OnWorldBeginTearDown);
 	TickWorldHandle = FWorldDelegates::OnWorldPostActorTick.AddStatic(&FNiagaraWorldManager::TickWorld);
+
+	PostGCHandle = FCoreUObjectDelegates::GetPostGarbageCollect().AddStatic(&FNiagaraWorldManager::OnPostGarbageCollect);
 }
 
 void FNiagaraWorldManager::OnShutdown()
@@ -176,6 +179,8 @@ void FNiagaraWorldManager::OnShutdown()
 	FWorldDelegates::OnWorldBeginTearDown.Remove(OnWorldBeginTearDownHandle);
 	FWorldDelegates::OnWorldPostActorTick.Remove(TickWorldHandle);
 
+	FCoreUObjectDelegates::GetPostGarbageCollect().Remove(PostGCHandle);
+
 	//Should have cleared up all world managers by now.
 	check(WorldManagers.Num() == 0);
 	for (TPair<UWorld*, FNiagaraWorldManager*> Pair : WorldManagers)
@@ -350,6 +355,12 @@ void FNiagaraWorldManager::OnWorldCleanup(bool bSessionEnded, bool bCleanupResou
 	}
 }
 
+void FNiagaraWorldManager::PostGarbageCollect()
+{
+	//Clear out and scalability managers who's EffectTypes have been GCd.
+	while (ScalabilityManagers.Remove(nullptr)) {}
+}
+
 void FNiagaraWorldManager::OnWorldInit(UWorld* World, const UWorld::InitializationValues IVS)
 {
 	check(WorldManagers.Find(World) == nullptr);
@@ -412,6 +423,14 @@ void FNiagaraWorldManager::TickWorld(UWorld* World, ELevelTick TickType, float D
 	Get(World)->PostActorTick(DeltaSeconds);
 }
 
+void FNiagaraWorldManager::OnPostGarbageCollect()
+{
+	for (TPair<UWorld*, FNiagaraWorldManager*>& Pair : WorldManagers)
+	{
+		Pair.Value->PostGarbageCollect();
+	}
+}
+
 void FNiagaraWorldManager::PostActorTick(float DeltaSeconds)
 {
 	CSV_SCOPED_TIMING_STAT_EXCLUSIVE(Niagara);
@@ -592,6 +611,7 @@ void FNiagaraWorldManager::UpdateScalabilityManagers()
 	{
 		FNiagaraScalabilityManager& ScalabilityMan = Pair.Value;
 		UNiagaraEffectType* EffectType = Pair.Key;
+		check(EffectType);
 
 		EffectType->ProcessLastFrameCycleCounts();
 
@@ -604,6 +624,7 @@ void FNiagaraWorldManager::UpdateScalabilityManagers()
 
 void FNiagaraWorldManager::RegisterWithScalabilityManager(UNiagaraComponent* Component)
 {
+	check(Component);
 	if (UNiagaraEffectType* EffectType = Component->GetAsset()->GetEffectType())
 	{
 		FNiagaraScalabilityManager* ScalabilityManager = ScalabilityManagers.Find(EffectType);
@@ -620,10 +641,19 @@ void FNiagaraWorldManager::RegisterWithScalabilityManager(UNiagaraComponent* Com
 
 void FNiagaraWorldManager::UnregisterWithScalabilityManager(UNiagaraComponent* Component)
 {
+	check(Component);
 	if (UNiagaraEffectType* EffectType = Component->GetAsset()->GetEffectType())
 	{
-		FNiagaraScalabilityManager& ScalabilityManager = ScalabilityManagers.FindChecked(EffectType);
-		ScalabilityManager.Unregister(Component);
+		//Possibly the manager has been GCd.
+		if (FNiagaraScalabilityManager* ScalabilityManager = ScalabilityManagers.Find(EffectType))
+		{
+			ScalabilityManager->Unregister(Component);
+		}
+		else
+		{
+			//The component does this itself in unregister.
+			//Component->ScalabilityManagerHandle = INDEX_NONE;
+		}
 	}
 }
 
diff --git a/Engine/Plugins/FX/Niagara/Source/Niagara/Public/NiagaraScalabilityManager.h b/Engine/Plugins/FX/Niagara/Source/Niagara/Public/NiagaraScalabilityManager.h
--- a/Engine/Plugins/FX/Niagara/Source/Niagara/Public/NiagaraScalabilityManager.h
+++ b/Engine/Plugins/FX/Niagara/Source/Niagara/Public/NiagaraScalabilityManager.h
@@ -125,6 +125,7 @@ struct FNiagaraScalabilityManager
 	float LastUpdateTime;
 
 	FNiagaraScalabilityManager();
+	~FNiagaraScalabilityManager();
 	void Update(FNiagaraWorldManager* Owner);
 	void Register(UNiagaraComponent* Component);
 	void Unregister(UNiagaraComponent* Component);
diff --git a/Engine/Plugins/FX/Niagara/Source/Niagara/Public/NiagaraWorldManager.h b/Engine/Plugins/FX/Niagara/Source/Niagara/Public/NiagaraWorldManager.h
--- a/Engine/Plugins/FX/Niagara/Source/Niagara/Public/NiagaraWorldManager.h
+++ b/Engine/Plugins/FX/Niagara/Source/Niagara/Public/NiagaraWorldManager.h
@@ -124,6 +124,8 @@ public:
 	void PostActorTick(float DeltaSeconds);
 
 	void OnWorldCleanup(bool bSessionEnded, bool bCleanupResources);
+
+	void PostGarbageCollect();
 	
 	FORCEINLINE FNDI_SkeletalMesh_GeneratedData& GetSkeletalMeshGeneratedData() { return SkeletalMeshGeneratedData; }
 
@@ -174,6 +176,9 @@ private:
 
 	// Callback for when a world is ticked.
 	static void TickWorld(UWorld* World, ELevelTick TickType, float DeltaSeconds);
+
+	// Callback to handle any post GC processing needed.
+	static void OnPostGarbageCollect();
 	
 	// Gamethread callback to cleanup references to the given batcher before it gets deleted on the renderthread.
 	void OnBatcherDestroyed_Internal(NiagaraEmitterInstanceBatcher* InBatcher);
@@ -194,6 +199,7 @@ private:
 	static FDelegateHandle OnPreWorldFinishDestroyHandle;
 	static FDelegateHandle OnWorldBeginTearDownHandle;
 	static FDelegateHandle TickWorldHandle;
+	static FDelegateHandle PostGCHandle;
 
 	static TMap<class UWorld*, class FNiagaraWorldManager*> WorldManagers;
 
