- Fix crash if we failed to spawn the Niagara System, this can happen because it was culled by scalability

#rnx

diff --git a/Engine/Plugins/FX/Niagara/Source/NiagaraAnimNotifies/Private/AnimNotify_PlayNiagaraEffect.cpp b/Engine/Plugins/FX/Niagara/Source/NiagaraAnimNotifies/Private/AnimNotify_PlayNiagaraEffect.cpp
--- a/Engine/Plugins/FX/Niagara/Source/NiagaraAnimNotifies/Private/AnimNotify_PlayNiagaraEffect.cpp
+++ b/Engine/Plugins/FX/Niagara/Source/NiagaraAnimNotifies/Private/AnimNotify_PlayNiagaraEffect.cpp
@@ -112,7 +112,10 @@ UFXSystemComponent* UAnimNotify_PlayNiagaraEffect::SpawnEffect(USkeletalMeshComp
 			ReturnComp = UNiagaraFunctionLibrary::SpawnSystemAtLocation(MeshComp->GetWorld(), Template, MeshTransform.TransformPosition(LocationOffset), (MeshTransform.GetRotation() * RotationOffsetQuat).Rotator(), FVector(1.0f),true);
 		}
 
-		ReturnComp->SetRelativeScale3D_Direct(Scale);
+		if (ReturnComp != nullptr)
+		{
+			ReturnComp->SetRelativeScale3D_Direct(Scale);
+		}
 	}
 
 	return ReturnComp;
