Added a "Create Asset from This" context menu option to emitters, which will duplicate the selected emitter in a Niagara system and re-parent the emitter to the new asset.

diff --git a/Engine/Plugins/FX/Niagara/Source/NiagaraEditor/Private/NiagaraEditorUtilities.cpp b/Engine/Plugins/FX/Niagara/Source/NiagaraEditor/Private/NiagaraEditorUtilities.cpp
--- a/Engine/Plugins/FX/Niagara/Source/NiagaraEditor/Private/NiagaraEditorUtilities.cpp
+++ b/Engine/Plugins/FX/Niagara/Source/NiagaraEditor/Private/NiagaraEditorUtilities.cpp
@@ -8,7 +8,6 @@
 #include "NiagaraNodeParameterMapSet.h"
 #include "NiagaraDataInterface.h"
 #include "NiagaraComponent.h"
-#include "Modules/ModuleManager.h"
 #include "UObject/StructOnScope.h"
 #include "NiagaraGraph.h"
 #include "NiagaraSystem.h"
@@ -16,6 +15,7 @@
 #include "NiagaraScriptSource.h"
 #include "NiagaraScript.h"
 #include "NiagaraNodeOutput.h"
+#include "NiagaraOverviewNode.h"
 #include "EdGraphUtilities.h"
 #include "NiagaraConstants.h"
 #include "Widgets/SWidget.h"
@@ -28,6 +28,8 @@
 #include "ViewModels/NiagaraSystemViewModel.h"
 #include "ViewModels/NiagaraEmitterViewModel.h"
 #include "ViewModels/NiagaraEmitterHandleViewModel.h"
+#include "ViewModels/NiagaraOverviewGraphViewModel.h"
+#include "ViewModels/NiagaraSystemSelectionViewModel.h"
 #include "HAL/PlatformApplicationMisc.h"
 #include "AssetRegistryModule.h"
 #include "Misc/FeedbackContext.h"
@@ -45,9 +47,9 @@
 #include "Framework/MultiBox/MultiBoxBuilder.h"
 #include "IContentBrowserSingleton.h"
 #include "ContentBrowserModule.h"
-#include "Toolkits/AssetEditorManager.h"
 #include "Modules/ModuleManager.h"
-#include "ViewModels/NiagaraSystemSelectionViewModel.h"
+#include "AssetToolsModule.h"
+#include "Subsystems/AssetEditorSubsystem.h"
 
 #define LOCTEXT_NAMESPACE "FNiagaraEditorUtilities"
 
@@ -1376,6 +1378,74 @@ void FNiagaraEditorUtilities::ToggleSelectedEmittersIsolated(TSharedRef<FNiagara
 	SystemViewModel->IsolateEmitters(EmittersToIsolate);
 }
 
+
+void FNiagaraEditorUtilities::CreateAssetFromEmitter(TSharedRef<FNiagaraEmitterHandleViewModel> EmitterHandleViewModel)
+{
+	TSharedRef<FNiagaraSystemViewModel> SystemViewModel = EmitterHandleViewModel->GetOwningSystemViewModel();
+	if (SystemViewModel->GetEditMode() != ENiagaraSystemViewModelEditMode::SystemAsset)
+	{
+		return;
+	}
+
+	UNiagaraEmitter* EmitterToCopy = EmitterHandleViewModel->GetEmitterViewModel()->GetEmitter();
+	const FString PackagePath = FPackageName::GetLongPackagePath(EmitterToCopy->GetOutermost()->GetName());
+	const FName EmitterName = EmitterToCopy->GetFName();
+	
+	FAssetToolsModule& AssetToolsModule = FModuleManager::LoadModuleChecked<FAssetToolsModule>("AssetTools");
+	UNiagaraEmitter* CreatedAsset = Cast<UNiagaraEmitter>(AssetToolsModule.Get().DuplicateAssetWithDialogAndTitle(EmitterName.GetPlainNameString(), PackagePath, EmitterToCopy, LOCTEXT("CreateEmitterAssetDialogTitle", "Create Emitter As")));
+	if (CreatedAsset != nullptr)
+	{
+		CreatedAsset->SetFlags(RF_Standalone | RF_Public);
+
+		CreatedAsset->SetUniqueEmitterName(CreatedAsset->GetName());
+
+		GEditor->GetEditorSubsystem<UAssetEditorSubsystem>()->OpenEditorForAsset(CreatedAsset);
+
+		// find the existing overview node to store the position
+		UEdGraph* OverviewGraph = SystemViewModel->GetOverviewGraphViewModel()->GetGraph();
+		
+		TArray<UNiagaraOverviewNode*> OverviewNodes;
+		OverviewGraph->GetNodesOfClass<UNiagaraOverviewNode>(OverviewNodes);
+
+		const UNiagaraOverviewNode* CurrentNode = *OverviewNodes.FindByPredicate(
+			[Guid = EmitterHandleViewModel->GetId()](const UNiagaraOverviewNode* Node)
+			{
+				return Node->GetEmitterHandleGuid() == Guid;
+			});
+
+		int32 CurrentX = CurrentNode->NodePosX;
+		int32 CurrentY = CurrentNode->NodePosY;
+		FString CurrentComment = CurrentNode->NodeComment;
+		bool bCommentBubbleVisible = CurrentNode->bCommentBubbleVisible;
+		bool bCommentBubblePinned = CurrentNode->bCommentBubblePinned;
+
+		CurrentNode = nullptr;
+
+		FScopedTransaction ScopedTransaction(LOCTEXT("CreateAssetFromEmitter", "Create asset from emitter"));
+		SystemViewModel->GetSystem().Modify();
+
+		// Replace existing emitter
+		SystemViewModel->DeleteEmitters(TSet<FGuid> { EmitterHandleViewModel->GetId() });
+		TSharedPtr<FNiagaraEmitterHandleViewModel> NewEmitterHandleViewModel = SystemViewModel->AddEmitter(*CreatedAsset);
+
+		NewEmitterHandleViewModel->SetName(EmitterName);
+
+		OverviewGraph->GetNodesOfClass<UNiagaraOverviewNode>(OverviewNodes);
+
+		UNiagaraOverviewNode* NewNode = *OverviewNodes.FindByPredicate(
+			[Guid = NewEmitterHandleViewModel->GetId()](const UNiagaraOverviewNode* Node)
+			{
+				return Node->GetEmitterHandleGuid() == Guid;
+			});
+
+		NewNode->NodePosX = CurrentX;
+		NewNode->NodePosY = CurrentY;
+		NewNode->NodeComment = CurrentComment;
+		NewNode->bCommentBubbleVisible = bCommentBubbleVisible;
+		NewNode->bCommentBubblePinned = bCommentBubblePinned;
+	}
+}
+
 bool FNiagaraEditorUtilities::AddEmitterContextMenuActions(FMenuBuilder& MenuBuilder, const TSharedPtr<FNiagaraEmitterHandleViewModel>& EmitterHandleViewModelPtr)
 {
 	if (EmitterHandleViewModelPtr.IsValid())
@@ -1456,6 +1526,21 @@ bool FNiagaraEditorUtilities::AddEmitterContextMenuActions(FMenuBuilder& MenuBui
 					)
 				)
 			);
+
+			MenuBuilder.AddMenuEntry(
+				LOCTEXT("CreateAssetFromEmitter", "Create Asset From This"),
+				LOCTEXT("CreateAssetFromEmitterToolTip", "Create an emitter asset from this emitter."),
+				FSlateIcon(),
+				FUIAction(
+					FExecuteAction::CreateStatic(&FNiagaraEditorUtilities::CreateAssetFromEmitter, EmitterHandleViewModel),
+					FCanExecuteAction::CreateLambda(
+						[bSingleSelection, EmitterHandleViewModel]()
+						{
+							return bSingleSelection && EmitterHandleViewModel->GetOwningSystemEditMode() == ENiagaraSystemViewModelEditMode::SystemAsset;
+						}
+					)
+				)
+			);
 		}
 		MenuBuilder.EndSection();
 
diff --git a/Engine/Plugins/FX/Niagara/Source/NiagaraEditor/Private/NiagaraOverviewNode.cpp b/Engine/Plugins/FX/Niagara/Source/NiagaraEditor/Private/NiagaraOverviewNode.cpp
--- a/Engine/Plugins/FX/Niagara/Source/NiagaraEditor/Private/NiagaraOverviewNode.cpp
+++ b/Engine/Plugins/FX/Niagara/Source/NiagaraEditor/Private/NiagaraOverviewNode.cpp
@@ -142,8 +142,6 @@ void UNiagaraOverviewNode::GetNodeContextMenuActions(class UToolMenu* Menu, clas
 	TSharedPtr<FNiagaraSystemViewModel> OwningSystemViewModel = NiagaraEditorModule.GetExistingViewModelForSystem(OwningSystem);
 	if (OwningSystemViewModel.IsValid())
 	{
-		FToolMenuSection& Section = Menu->AddSection("Emitter Actions");
-
 		bool bSingleSelection = OwningSystemViewModel->GetSelectionViewModel()->GetSelectedEmitterHandleIds().Num() == 1;
 
 		TSharedPtr<FNiagaraEmitterHandleViewModel> EmitterHandleViewModelPtr = OwningSystemViewModel->GetEmitterHandleViewModelById(EmitterHandleGuid);
@@ -151,82 +149,99 @@ void UNiagaraOverviewNode::GetNodeContextMenuActions(class UToolMenu* Menu, clas
 		{
 			TSharedRef<FNiagaraEmitterHandleViewModel> EmitterHandleViewModel = EmitterHandleViewModelPtr.ToSharedRef();
 			TSharedRef<FNiagaraEmitterViewModel> EmitterViewModel = EmitterHandleViewModel->GetEmitterViewModel();
+			
+			FToolMenuSection& Section = Menu->AddSection("Emitter Actions");
+
+			if (OwningSystemViewModel->GetEditMode() == ENiagaraSystemViewModelEditMode::SystemAsset)
 			{
-				if (OwningSystemViewModel->GetEditMode() == ENiagaraSystemViewModelEditMode::SystemAsset)
-				{
-					Section.AddMenuEntry(
-						"ToggleEmittersEnabled",
-						LOCTEXT("ToggleEmittersEnabled", "Enabled"),
-						LOCTEXT("ToggleEmittersEnabledToolTip", "Toggle whether or not the selected emitters are enabled."),
-						FSlateIcon(),
-						FUIAction(
-							FExecuteAction::CreateStatic(&FNiagaraEditorUtilities::ToggleSelectedEmittersEnabled, OwningSystemViewModel.ToSharedRef()),
-							FCanExecuteAction(),
-							FGetActionCheckState::CreateStatic(&FNiagaraEditorUtilities::GetSelectedEmittersEnabledCheckState, OwningSystemViewModel.ToSharedRef())
-						),
-						EUserInterfaceActionType::ToggleButton
-					);
-
-					Section.AddMenuEntry(
-						"ToggleEmittersIsolated",
-						LOCTEXT("ToggleEmittersIsolated", "Isolated"),
-						LOCTEXT("ToggleEmittersIsolatedToolTip", "Toggle whether or not the selected emitters are isolated."),
-						FSlateIcon(),
-						FUIAction(
-							FExecuteAction::CreateStatic(&FNiagaraEditorUtilities::ToggleSelectedEmittersIsolated, OwningSystemViewModel.ToSharedRef()),
-							FCanExecuteAction(),
-							FGetActionCheckState::CreateStatic(&FNiagaraEditorUtilities::GetSelectedEmittersIsolatedCheckState, OwningSystemViewModel.ToSharedRef())
-						),
-						EUserInterfaceActionType::ToggleButton
-					);
-
-					Section.AddMenuEntry(
-						"RenameEmitter",
-						LOCTEXT("RenameEmitter", "Rename Emitter"),
-						LOCTEXT("RenameEmitterToolTip", "Rename this local emitter copy."),
-						FSlateIcon(),
-						FUIAction(
-							FExecuteAction::CreateLambda([this]()
-							{
-								const_cast<UNiagaraOverviewNode*>(this)->RequestRename();
-							}),
-							FCanExecuteAction::CreateLambda([bSingleSelection]() { return bSingleSelection; })
-						)
-					);
-				}
+				Section.AddMenuEntry(
+					"ToggleEmittersEnabled",
+					LOCTEXT("ToggleEmittersEnabled", "Enabled"),
+					LOCTEXT("ToggleEmittersEnabledToolTip", "Toggle whether or not the selected emitters are enabled."),
+					FSlateIcon(),
+					FUIAction(
+						FExecuteAction::CreateStatic(&FNiagaraEditorUtilities::ToggleSelectedEmittersEnabled, OwningSystemViewModel.ToSharedRef()),
+						FCanExecuteAction(),
+						FGetActionCheckState::CreateStatic(&FNiagaraEditorUtilities::GetSelectedEmittersEnabledCheckState, OwningSystemViewModel.ToSharedRef())
+					),
+					EUserInterfaceActionType::ToggleButton
+				);
 
 				Section.AddMenuEntry(
-					"RemoveParentEmitter",
-					LOCTEXT("RemoveParentEmitter", "Remove Parent Emitter"),
-					LOCTEXT("RemoveParentEmitterToolTip", "Removes this emitter's parent, preventing inheritance of any further changes."),
+					"ToggleEmittersIsolated",
+					LOCTEXT("ToggleEmittersIsolated", "Isolated"),
+					LOCTEXT("ToggleEmittersIsolatedToolTip", "Toggle whether or not the selected emitters are isolated."),
 					FSlateIcon(),
 					FUIAction(
-						FExecuteAction::CreateSP(EmitterViewModel, &FNiagaraEmitterViewModel::RemoveParentEmitter),
-						FCanExecuteAction::CreateLambda(
-							[bSingleSelection, bHasParent = EmitterViewModel->HasParentEmitter()]()
-							{
-								return bSingleSelection && bHasParent;
-							}
-						)
-					)
+						FExecuteAction::CreateStatic(&FNiagaraEditorUtilities::ToggleSelectedEmittersIsolated, OwningSystemViewModel.ToSharedRef()),
+						FCanExecuteAction(),
+						FGetActionCheckState::CreateStatic(&FNiagaraEditorUtilities::GetSelectedEmittersIsolatedCheckState, OwningSystemViewModel.ToSharedRef())
+					),
+					EUserInterfaceActionType::ToggleButton
 				);
 
 				Section.AddMenuEntry(
-					"ShowEmitterInContentBrowser",
-					LOCTEXT("ShowEmitterInContentBrowser", "Show in Content Browser"),
-					LOCTEXT("ShowEmitterInContentBrowserToolTip", "Show the selected emitter in the Content Browser."),
+					"RenameEmitter",
+					LOCTEXT("RenameEmitter", "Rename Emitter"),
+					LOCTEXT("RenameEmitterToolTip", "Rename this local emitter copy."),
 					FSlateIcon(),
 					FUIAction(
-						FExecuteAction::CreateStatic(&FNiagaraEditorUtilities::ShowParentEmitterInContentBrowser, EmitterViewModel),
-						FCanExecuteAction::CreateLambda(
-							[bSingleSelection, bHasParent = EmitterViewModel->HasParentEmitter()]()
-							{
-								return bSingleSelection && bHasParent;
-							}
-						)
+						FExecuteAction::CreateLambda([this]()
+						{
+							const_cast<UNiagaraOverviewNode*>(this)->RequestRename();
+						}),
+						FCanExecuteAction::CreateLambda([bSingleSelection]() { return bSingleSelection; })
 					)
 				);
 			}
+
+			Section.AddMenuEntry(
+				"RemoveParentEmitter",
+				LOCTEXT("RemoveParentEmitter", "Remove Parent Emitter"),
+				LOCTEXT("RemoveParentEmitterToolTip", "Removes this emitter's parent, preventing inheritance of any further changes."),
+				FSlateIcon(),
+				FUIAction(
+					FExecuteAction::CreateSP(EmitterViewModel, &FNiagaraEmitterViewModel::RemoveParentEmitter),
+					FCanExecuteAction::CreateLambda(
+						[bSingleSelection, bHasParent = EmitterViewModel->HasParentEmitter()]()
+						{
+							return bSingleSelection && bHasParent;
+						}
+					)
+				)
+			);
+
+			Section.AddMenuEntry(
+				"ShowEmitterInContentBrowser",
+				LOCTEXT("ShowEmitterInContentBrowser", "Show in Content Browser"),
+				LOCTEXT("ShowEmitterInContentBrowserToolTip", "Show the selected emitter in the Content Browser."),
+				FSlateIcon(),
+				FUIAction(
+					FExecuteAction::CreateStatic(&FNiagaraEditorUtilities::ShowParentEmitterInContentBrowser, EmitterViewModel),
+					FCanExecuteAction::CreateLambda(
+						[bSingleSelection, bHasParent = EmitterViewModel->HasParentEmitter()]()
+						{
+							return bSingleSelection && bHasParent;
+						}
+					)
+				)
+			);
+
+			Section.AddMenuEntry(
+				"CreateAssetFromEmitter",
+				NSLOCTEXT("FNiagaraEditorUtilities", "CreateAssetFromEmitter", "Create Asset From This"),
+				NSLOCTEXT("FNiagaraEditorUtilities", "CreateAssetFromEmitterToolTip", "Create an emitter asset from this emitter."),
+				FSlateIcon(),
+				FUIAction(
+					FExecuteAction::CreateStatic(&FNiagaraEditorUtilities::CreateAssetFromEmitter, EmitterHandleViewModel),
+					FCanExecuteAction::CreateLambda(
+						[bSingleSelection, EmitterHandleViewModel]()
+						{
+							return bSingleSelection && EmitterHandleViewModel->GetOwningSystemEditMode() == ENiagaraSystemViewModelEditMode::SystemAsset;
+						}
+					)
+				)
+			);
 		}
 	}
 
diff --git a/Engine/Plugins/FX/Niagara/Source/NiagaraEditor/Private/ViewModels/NiagaraSystemViewModel.cpp b/Engine/Plugins/FX/Niagara/Source/NiagaraEditor/Private/ViewModels/NiagaraSystemViewModel.cpp
--- a/Engine/Plugins/FX/Niagara/Source/NiagaraEditor/Private/ViewModels/NiagaraSystemViewModel.cpp
+++ b/Engine/Plugins/FX/Niagara/Source/NiagaraEditor/Private/ViewModels/NiagaraSystemViewModel.cpp
@@ -278,16 +278,18 @@ ENiagaraSystemViewModelEditMode FNiagaraSystemViewModel::GetEditMode() const
 	return EditMode;
 }
 
-void FNiagaraSystemViewModel::AddEmitterFromAssetData(const FAssetData& AssetData)
+TSharedPtr<FNiagaraEmitterHandleViewModel> FNiagaraSystemViewModel::AddEmitterFromAssetData(const FAssetData& AssetData)
 {
 	UNiagaraEmitter* Emitter = Cast<UNiagaraEmitter>(AssetData.GetAsset());
 	if (Emitter != nullptr)
 	{
-		AddEmitter(*Emitter);
+		return AddEmitter(*Emitter);
 	}
+
+	return nullptr;
 }
 
-void FNiagaraSystemViewModel::AddEmitter(UNiagaraEmitter& Emitter)
+TSharedPtr<FNiagaraEmitterHandleViewModel> FNiagaraSystemViewModel::AddEmitter(UNiagaraEmitter& Emitter)
 {
 	// Reset view models before modifying the emitter handle list to prevent accessing deleted data.
 	ResetEmitterHandleViewModelsAndTracks();
@@ -335,6 +337,8 @@ void FNiagaraSystemViewModel::AddEmitter(UNiagaraEmitter& Emitter)
 	SelectionViewModel->UpdateSelectedEntries(SelectedStackEntries, TArray<UNiagaraStackEntry*>(), true);
 
 	bForceAutoCompileOnce = true;
+
+	return NewEmitterHandleViewModel;
 }
 
 void FNiagaraSystemViewModel::DuplicateEmitters(TArray<FEmitterHandleToDuplicate> EmitterHandlesToDuplicate)
diff --git a/Engine/Plugins/FX/Niagara/Source/NiagaraEditor/Public/NiagaraEditorUtilities.h b/Engine/Plugins/FX/Niagara/Source/NiagaraEditor/Public/NiagaraEditorUtilities.h
--- a/Engine/Plugins/FX/Niagara/Source/NiagaraEditor/Public/NiagaraEditorUtilities.h
+++ b/Engine/Plugins/FX/Niagara/Source/NiagaraEditor/Public/NiagaraEditorUtilities.h
@@ -224,4 +224,6 @@ namespace FNiagaraEditorUtilities
 
 	ECheckBoxState GetSelectedEmittersIsolatedCheckState(TSharedRef<FNiagaraSystemViewModel> SystemViewModel);
 	void ToggleSelectedEmittersIsolated(TSharedRef<FNiagaraSystemViewModel> SystemViewModel);
+
+	void CreateAssetFromEmitter(TSharedRef<FNiagaraEmitterHandleViewModel> EmitterHandleViewModel);
 };
diff --git a/Engine/Plugins/FX/Niagara/Source/NiagaraEditor/Public/ViewModels/NiagaraSystemViewModel.h b/Engine/Plugins/FX/Niagara/Source/NiagaraEditor/Public/ViewModels/NiagaraSystemViewModel.h
--- a/Engine/Plugins/FX/Niagara/Source/NiagaraEditor/Public/ViewModels/NiagaraSystemViewModel.h
+++ b/Engine/Plugins/FX/Niagara/Source/NiagaraEditor/Public/ViewModels/NiagaraSystemViewModel.h
@@ -180,13 +180,13 @@ public:
 	NIAGARAEDITOR_API ENiagaraSystemViewModelEditMode GetEditMode() const;
 
 	/** Adds a new emitter to the System from an emitter asset data. */
-	NIAGARAEDITOR_API void AddEmitterFromAssetData(const FAssetData& AssetData);
+	NIAGARAEDITOR_API TSharedPtr<FNiagaraEmitterHandleViewModel> AddEmitterFromAssetData(const FAssetData& AssetData);
 
 	/** Adds a new emitter to the System. */
-	void AddEmitter(UNiagaraEmitter& Emitter);
+	NIAGARAEDITOR_API TSharedPtr<FNiagaraEmitterHandleViewModel> AddEmitter(UNiagaraEmitter& Emitter);
 
 	/** Deletes the emitters with the supplied ids from the system */
-	void DeleteEmitters(TSet<FGuid> EmitterHandleIdsToDelete);
+	NIAGARAEDITOR_API void DeleteEmitters(TSet<FGuid> EmitterHandleIdsToDelete);
 
 	/** Gets a multicast delegate which is called any time the array of emitter handle view models changes. */
 	NIAGARAEDITOR_API FOnEmitterHandleViewModelsChanged& OnEmitterHandleViewModelsChanged();
