Fix for GC warnings about leaks. These would not have caused leaks for real but now cleaning things up earlier on world cleanup.

diff --git a/Engine/Plugins/FX/Niagara/Source/Niagara/Private/NiagaraWorldManager.cpp b/Engine/Plugins/FX/Niagara/Source/Niagara/Private/NiagaraWorldManager.cpp
--- a/Engine/Plugins/FX/Niagara/Source/Niagara/Private/NiagaraWorldManager.cpp
+++ b/Engine/Plugins/FX/Niagara/Source/Niagara/Private/NiagaraWorldManager.cpp
@@ -353,6 +353,8 @@ void FNiagaraWorldManager::OnWorldCleanup(bool bSessionEnded, bool bCleanupResou
 		DeferredDeletionQueue[i].Fence.Wait();
 		DeferredDeletionQueue[i].Queue.Empty();
 	}
+
+	ScalabilityManagers.Empty();
 }
 
 void FNiagaraWorldManager::PostGarbageCollect()
