Emitters now have the Edit section in their context menus. It is now possible to rename an emitter by pressing F2.

diff --git a/Engine/Plugins/FX/Niagara/Source/NiagaraEditor/Private/NiagaraEditorUtilities.cpp b/Engine/Plugins/FX/Niagara/Source/NiagaraEditor/Private/NiagaraEditorUtilities.cpp
--- a/Engine/Plugins/FX/Niagara/Source/NiagaraEditor/Private/NiagaraEditorUtilities.cpp
+++ b/Engine/Plugins/FX/Niagara/Source/NiagaraEditor/Private/NiagaraEditorUtilities.cpp
@@ -50,6 +50,7 @@
 #include "Modules/ModuleManager.h"
 #include "AssetToolsModule.h"
 #include "Subsystems/AssetEditorSubsystem.h"
+#include "Framework/Commands/GenericCommands.h"
 
 #define LOCTEXT_NAMESPACE "FNiagaraEditorUtilities"
 
@@ -1544,6 +1545,15 @@ bool FNiagaraEditorUtilities::AddEmitterContextMenuActions(FMenuBuilder& MenuBui
 		}
 		MenuBuilder.EndSection();
 
+		MenuBuilder.BeginSection("EmitterEditSection", LOCTEXT("Edit", "Edit"));
+
+		MenuBuilder.AddMenuEntry(FGenericCommands::Get().Cut);
+		MenuBuilder.AddMenuEntry(FGenericCommands::Get().Copy);
+		MenuBuilder.AddMenuEntry(FGenericCommands::Get().Delete);
+		MenuBuilder.AddMenuEntry(FGenericCommands::Get().Rename);
+
+		MenuBuilder.EndSection();
+
 		return true;
 	}
 
diff --git a/Engine/Plugins/FX/Niagara/Source/NiagaraEditor/Private/NiagaraOverviewNode.cpp b/Engine/Plugins/FX/Niagara/Source/NiagaraEditor/Private/NiagaraOverviewNode.cpp
--- a/Engine/Plugins/FX/Niagara/Source/NiagaraEditor/Private/NiagaraOverviewNode.cpp
+++ b/Engine/Plugins/FX/Niagara/Source/NiagaraEditor/Private/NiagaraOverviewNode.cpp
@@ -14,6 +14,7 @@
 #include "ToolMenuSection.h"
 #include "ToolMenu.h"
 #include "GraphEditorActions.h"
+#include "Framework/Commands/GenericCommands.h"
 
 #define LOCTEXT_NAMESPACE "NiagaraOverviewNodeStackItem"
 
@@ -142,6 +143,8 @@ void UNiagaraOverviewNode::GetNodeContextMenuActions(class UToolMenu* Menu, clas
 	TSharedPtr<FNiagaraSystemViewModel> OwningSystemViewModel = NiagaraEditorModule.GetExistingViewModelForSystem(OwningSystem);
 	if (OwningSystemViewModel.IsValid())
 	{
+		FToolMenuSection& Section = Menu->AddSection("Emitter Actions", LOCTEXT("EmitterActions", "Emitter Actions"));
+
 		bool bSingleSelection = OwningSystemViewModel->GetSelectionViewModel()->GetSelectedEmitterHandleIds().Num() == 1;
 
 		TSharedPtr<FNiagaraEmitterHandleViewModel> EmitterHandleViewModelPtr = OwningSystemViewModel->GetEmitterHandleViewModelById(EmitterHandleGuid);
@@ -149,104 +152,82 @@ void UNiagaraOverviewNode::GetNodeContextMenuActions(class UToolMenu* Menu, clas
 		{
 			TSharedRef<FNiagaraEmitterHandleViewModel> EmitterHandleViewModel = EmitterHandleViewModelPtr.ToSharedRef();
 			TSharedRef<FNiagaraEmitterViewModel> EmitterViewModel = EmitterHandleViewModel->GetEmitterViewModel();
-			
-			FToolMenuSection& Section = Menu->AddSection("Emitter Actions");
-
-			if (OwningSystemViewModel->GetEditMode() == ENiagaraSystemViewModelEditMode::SystemAsset)
 			{
-				Section.AddMenuEntry(
-					"ToggleEmittersEnabled",
-					LOCTEXT("ToggleEmittersEnabled", "Enabled"),
-					LOCTEXT("ToggleEmittersEnabledToolTip", "Toggle whether or not the selected emitters are enabled."),
-					FSlateIcon(),
-					FUIAction(
-						FExecuteAction::CreateStatic(&FNiagaraEditorUtilities::ToggleSelectedEmittersEnabled, OwningSystemViewModel.ToSharedRef()),
-						FCanExecuteAction(),
-						FGetActionCheckState::CreateStatic(&FNiagaraEditorUtilities::GetSelectedEmittersEnabledCheckState, OwningSystemViewModel.ToSharedRef())
-					),
-					EUserInterfaceActionType::ToggleButton
-				);
+				if (OwningSystemViewModel->GetEditMode() == ENiagaraSystemViewModelEditMode::SystemAsset)
+				{
+					Section.AddMenuEntry(
+						"ToggleEmittersEnabled",
+						LOCTEXT("ToggleEmittersEnabled", "Enabled"),
+						LOCTEXT("ToggleEmittersEnabledToolTip", "Toggle whether or not the selected emitters are enabled."),
+						FSlateIcon(),
+						FUIAction(
+							FExecuteAction::CreateStatic(&FNiagaraEditorUtilities::ToggleSelectedEmittersEnabled, OwningSystemViewModel.ToSharedRef()),
+							FCanExecuteAction(),
+							FGetActionCheckState::CreateStatic(&FNiagaraEditorUtilities::GetSelectedEmittersEnabledCheckState, OwningSystemViewModel.ToSharedRef())
+						),
+						EUserInterfaceActionType::ToggleButton
+					);
+
+					Section.AddMenuEntry(
+						"ToggleEmittersIsolated",
+						LOCTEXT("ToggleEmittersIsolated", "Isolated"),
+						LOCTEXT("ToggleEmittersIsolatedToolTip", "Toggle whether or not the selected emitters are isolated."),
+						FSlateIcon(),
+						FUIAction(
+							FExecuteAction::CreateStatic(&FNiagaraEditorUtilities::ToggleSelectedEmittersIsolated, OwningSystemViewModel.ToSharedRef()),
+							FCanExecuteAction(),
+							FGetActionCheckState::CreateStatic(&FNiagaraEditorUtilities::GetSelectedEmittersIsolatedCheckState, OwningSystemViewModel.ToSharedRef())
+						),
+						EUserInterfaceActionType::ToggleButton
+					);
+				}
 
 				Section.AddMenuEntry(
-					"ToggleEmittersIsolated",
-					LOCTEXT("ToggleEmittersIsolated", "Isolated"),
-					LOCTEXT("ToggleEmittersIsolatedToolTip", "Toggle whether or not the selected emitters are isolated."),
+					"RemoveParentEmitter",
+					LOCTEXT("RemoveParentEmitter", "Remove Parent Emitter"),
+					LOCTEXT("RemoveParentEmitterToolTip", "Removes this emitter's parent, preventing inheritance of any further changes."),
 					FSlateIcon(),
 					FUIAction(
-						FExecuteAction::CreateStatic(&FNiagaraEditorUtilities::ToggleSelectedEmittersIsolated, OwningSystemViewModel.ToSharedRef()),
-						FCanExecuteAction(),
-						FGetActionCheckState::CreateStatic(&FNiagaraEditorUtilities::GetSelectedEmittersIsolatedCheckState, OwningSystemViewModel.ToSharedRef())
-					),
-					EUserInterfaceActionType::ToggleButton
+						FExecuteAction::CreateSP(EmitterViewModel, &FNiagaraEmitterViewModel::RemoveParentEmitter),
+						FCanExecuteAction::CreateLambda(
+							[bSingleSelection, bHasParent = EmitterViewModel->HasParentEmitter()]()
+							{
+								return bSingleSelection && bHasParent;
+							}
+						)
+					)
 				);
 
 				Section.AddMenuEntry(
-					"RenameEmitter",
-					LOCTEXT("RenameEmitter", "Rename Emitter"),
-					LOCTEXT("RenameEmitterToolTip", "Rename this local emitter copy."),
+					"ShowEmitterInContentBrowser",
+					LOCTEXT("ShowEmitterInContentBrowser", "Show in Content Browser"),
+					LOCTEXT("ShowEmitterInContentBrowserToolTip", "Show the selected emitter in the Content Browser."),
 					FSlateIcon(),
 					FUIAction(
-						FExecuteAction::CreateLambda([this]()
-						{
-							const_cast<UNiagaraOverviewNode*>(this)->RequestRename();
-						}),
-						FCanExecuteAction::CreateLambda([bSingleSelection]() { return bSingleSelection; })
+						FExecuteAction::CreateStatic(&FNiagaraEditorUtilities::ShowParentEmitterInContentBrowser, EmitterViewModel),
+						FCanExecuteAction::CreateLambda(
+							[bSingleSelection, bHasParent = EmitterViewModel->HasParentEmitter()]()
+							{
+								return bSingleSelection && bHasParent;
+							}
+						)
 					)
 				);
 			}
-
-			Section.AddMenuEntry(
-				"RemoveParentEmitter",
-				LOCTEXT("RemoveParentEmitter", "Remove Parent Emitter"),
-				LOCTEXT("RemoveParentEmitterToolTip", "Removes this emitter's parent, preventing inheritance of any further changes."),
-				FSlateIcon(),
-				FUIAction(
-					FExecuteAction::CreateSP(EmitterViewModel, &FNiagaraEmitterViewModel::RemoveParentEmitter),
-					FCanExecuteAction::CreateLambda(
-						[bSingleSelection, bHasParent = EmitterViewModel->HasParentEmitter()]()
-						{
-							return bSingleSelection && bHasParent;
-						}
-					)
-				)
-			);
-
-			Section.AddMenuEntry(
-				"ShowEmitterInContentBrowser",
-				LOCTEXT("ShowEmitterInContentBrowser", "Show in Content Browser"),
-				LOCTEXT("ShowEmitterInContentBrowserToolTip", "Show the selected emitter in the Content Browser."),
-				FSlateIcon(),
-				FUIAction(
-					FExecuteAction::CreateStatic(&FNiagaraEditorUtilities::ShowParentEmitterInContentBrowser, EmitterViewModel),
-					FCanExecuteAction::CreateLambda(
-						[bSingleSelection, bHasParent = EmitterViewModel->HasParentEmitter()]()
-						{
-							return bSingleSelection && bHasParent;
-						}
-					)
-				)
-			);
-
-			Section.AddMenuEntry(
-				"CreateAssetFromEmitter",
-				NSLOCTEXT("FNiagaraEditorUtilities", "CreateAssetFromEmitter", "Create Asset From This"),
-				NSLOCTEXT("FNiagaraEditorUtilities", "CreateAssetFromEmitterToolTip", "Create an emitter asset from this emitter."),
-				FSlateIcon(),
-				FUIAction(
-					FExecuteAction::CreateStatic(&FNiagaraEditorUtilities::CreateAssetFromEmitter, EmitterHandleViewModel),
-					FCanExecuteAction::CreateLambda(
-						[bSingleSelection, EmitterHandleViewModel]()
-						{
-							return bSingleSelection && EmitterHandleViewModel->GetOwningSystemEditMode() == ENiagaraSystemViewModelEditMode::SystemAsset;
-						}
-					)
-				)
-			);
 		}
 	}
+	
+	{
+		FToolMenuSection& EditSection = Menu->AddSection("EmitterEdit", LOCTEXT("Edit", "Edit"));
+
+		EditSection.AddMenuEntry(FGenericCommands::Get().Cut);
+		EditSection.AddMenuEntry(FGenericCommands::Get().Copy);
+		EditSection.AddMenuEntry(FGenericCommands::Get().Delete);
+		EditSection.AddMenuEntry(FGenericCommands::Get().Rename);
+	}
 
 	{
-		FToolMenuSection& Section = Menu->AddSection("Alignment");
+		FToolMenuSection& Section = Menu->AddSection("Alignment", LOCTEXT("Alignment", "Alignment"));
 		Section.AddSubMenu(
 			"Alignment",
 			LOCTEXT("AlignmentHeader", "Alignment"),
diff --git a/Engine/Plugins/FX/Niagara/Source/NiagaraEditor/Private/ViewModels/NiagaraOverviewGraphViewModel.cpp b/Engine/Plugins/FX/Niagara/Source/NiagaraEditor/Private/ViewModels/NiagaraOverviewGraphViewModel.cpp
--- a/Engine/Plugins/FX/Niagara/Source/NiagaraEditor/Private/ViewModels/NiagaraOverviewGraphViewModel.cpp
+++ b/Engine/Plugins/FX/Niagara/Source/NiagaraEditor/Private/ViewModels/NiagaraOverviewGraphViewModel.cpp
@@ -134,6 +134,11 @@ void FNiagaraOverviewGraphViewModel::SetupCommands()
 		FGenericCommands::Get().Duplicate,
 		FExecuteAction::CreateRaw(this, &FNiagaraOverviewGraphViewModel::DuplicateNodes),
 		FCanExecuteAction::CreateRaw(this, &FNiagaraOverviewGraphViewModel::CanDuplicateNodes));
+
+	Commands->MapAction(
+		FGenericCommands::Get().Rename,
+		FExecuteAction::CreateRaw(this, &FNiagaraOverviewGraphViewModel::RenameNode),
+		FCanExecuteAction::CreateRaw(this, &FNiagaraOverviewGraphViewModel::CanRenameNode));
 }
 
 void FNiagaraOverviewGraphViewModel::SelectAllNodes()
@@ -341,6 +346,49 @@ bool FNiagaraOverviewGraphViewModel::CanDuplicateNodes() const
 	return CanCopyNodes();
 }
 
+void FNiagaraOverviewGraphViewModel::RenameNode()
+{
+	UEdGraph* Graph = GetGraph();
+	if (Graph != nullptr)
+	{
+		const TSet<UObject*>& SelectedNodes = NodeSelection->GetSelectedObjects();
+		if (SelectedNodes.Num() > 0)
+		{
+			for (UObject* SelectedNode : SelectedNodes)
+			{
+				UNiagaraOverviewNode* SelectedOverviewNode = Cast<UNiagaraOverviewNode>(SelectedNode);
+				if (SelectedOverviewNode != nullptr)
+				{
+					SelectedOverviewNode->RequestRename();
+					return;
+				}
+			}
+		}
+	}
+}
+
+bool FNiagaraOverviewGraphViewModel::CanRenameNode() const
+{
+	UEdGraph* Graph = GetGraph();
+	if (Graph != nullptr)
+	{
+		const TSet<UObject*>& SelectedNodes = NodeSelection->GetSelectedObjects();
+		if (SelectedNodes.Num() == 1)
+		{
+			for (UObject* SelectedNode : SelectedNodes)
+			{
+				UNiagaraOverviewNode* SelectedOverviewNode = Cast<UNiagaraOverviewNode>(SelectedNode);
+				if (SelectedOverviewNode != nullptr && SelectedOverviewNode->GetCanRenameNode())
+				{
+					return true;
+				}
+			}
+		}
+	}
+
+	return false;
+}
+
 FText FNiagaraOverviewGraphViewModel::GetDisplayNameInternal() const
 {
 	ensureMsgf(SystemViewModel.IsValid(), TEXT("SystemViewModel was not initialized before InitDisplayName!"));
diff --git a/Engine/Plugins/FX/Niagara/Source/NiagaraEditor/Public/ViewModels/NiagaraOverviewGraphViewModel.h b/Engine/Plugins/FX/Niagara/Source/NiagaraEditor/Public/ViewModels/NiagaraOverviewGraphViewModel.h
--- a/Engine/Plugins/FX/Niagara/Source/NiagaraEditor/Public/ViewModels/NiagaraOverviewGraphViewModel.h
+++ b/Engine/Plugins/FX/Niagara/Source/NiagaraEditor/Public/ViewModels/NiagaraOverviewGraphViewModel.h
@@ -62,6 +62,8 @@ private:
 	bool CanPasteNodes() const;
 	void DuplicateNodes();
 	bool CanDuplicateNodes() const;
+	void RenameNode();
+	bool CanRenameNode() const;
 
 	FText GetDisplayNameInternal() const;
 
diff --git a/Engine/Plugins/FX/Niagara/Source/NiagaraEditorWidgets/Private/SNiagaraOverviewGraph.cpp b/Engine/Plugins/FX/Niagara/Source/NiagaraEditorWidgets/Private/SNiagaraOverviewGraph.cpp
--- a/Engine/Plugins/FX/Niagara/Source/NiagaraEditorWidgets/Private/SNiagaraOverviewGraph.cpp
+++ b/Engine/Plugins/FX/Niagara/Source/NiagaraEditorWidgets/Private/SNiagaraOverviewGraph.cpp
@@ -24,6 +24,7 @@
 #include "Templates/SharedPointer.h"
 #include "EdGraph/EdGraph.h"
 #include "EdGraph/EdGraphNode.h"
+#include "Framework/Commands/GenericCommands.h"
 
 #define LOCTEXT_NAMESPACE "NiagaraOverviewGraph"
 
@@ -222,6 +223,12 @@ FActionMenuContent SNiagaraOverviewGraph::OnCreateGraphActionMenu(UEdGraph* InGr
 		}
 		MenuBuilder.EndSection();
 
+		MenuBuilder.BeginSection(TEXT("NiagaraOverview_Edit"), LOCTEXT("Edit", "Edit"));
+		{
+			MenuBuilder.AddMenuEntry(FGenericCommands::Get().Paste);
+		}
+		MenuBuilder.EndSection();
+
 		TSharedRef<SWidget> ActionMenu = MenuBuilder.MakeWidget();
 
 		return FActionMenuContent(ActionMenu, ActionMenu);
