Fixing Niagara debug dumps.

diff --git a/Engine/Plugins/FX/Niagara/Source/Niagara/Private/NiagaraCommon.cpp b/Engine/Plugins/FX/Niagara/Source/Niagara/Private/NiagaraCommon.cpp
--- a/Engine/Plugins/FX/Niagara/Source/Niagara/Private/NiagaraCommon.cpp
+++ b/Engine/Plugins/FX/Niagara/Source/Niagara/Private/NiagaraCommon.cpp
@@ -18,50 +18,57 @@ DECLARE_CYCLE_STAT(TEXT("Niagara - Utilities - PrepareRapidIterationParameters")
 
 //////////////////////////////////////////////////////////////////////////
 
-FString FNiagaraTypeHelper::ToString(const uint8* ValueData, const UScriptStruct* Struct)
+FString FNiagaraTypeHelper::ToString(const uint8* ValueData, const UObject* StructOrEnum)
 {
 	FString Ret;
-	if (Struct == FNiagaraTypeDefinition::GetFloatStruct())
+	if (const UEnum* Enum = Cast<const UEnum>(StructOrEnum))
 	{
-		Ret += FString::Printf(TEXT("%g "), *(float*)ValueData);
+		Ret = Enum->GetNameStringByValue(*(int32*)ValueData);
 	}
-	else if (Struct == FNiagaraTypeDefinition::GetIntStruct())
+	else if (const UScriptStruct* Struct = Cast<const UScriptStruct>(StructOrEnum))
 	{
-		Ret += FString::Printf(TEXT("%d "), *(int32*)ValueData);
-	}
-	else if (Struct == FNiagaraTypeDefinition::GetBoolStruct())
-	{
-		int32 Val = *(int32*)ValueData;
-		Ret += Val == 0xFFFFFFFF ? (TEXT("True")) : ( Val == 0x0 ? TEXT("False") : TEXT("Invalid"));		
-	}
-	else
-	{
-		for (TFieldIterator<FProperty> PropertyIt(Struct, EFieldIteratorFlags::IncludeSuper); PropertyIt; ++PropertyIt)
+		if (Struct == FNiagaraTypeDefinition::GetFloatStruct())
 		{
-			const FProperty* Property = *PropertyIt;
-			const uint8* PropPtr = ValueData + PropertyIt->GetOffset_ForInternal();
-			if (Property->IsA(FFloatProperty::StaticClass()))
-			{
-				Ret += FString::Printf(TEXT("%s: %g "), *Property->GetNameCPP(), *(float*)PropPtr);
-			}
-			else if (Property->IsA(FIntProperty::StaticClass()))
-			{
-				Ret += FString::Printf(TEXT("%s: %d "), *Property->GetNameCPP(), *(int32*)PropPtr);
-			}
-			else if (Property->IsA(FBoolProperty::StaticClass()))
-			{
-				int32 Val = *(int32*)ValueData;
-				FString BoolStr = Val == 0xFFFFFFFF ? (TEXT("True")) : (Val == 0x0 ? TEXT("False") : TEXT("Invalid"));
-				Ret += FString::Printf(TEXT("%s: %d "), *Property->GetNameCPP(), *BoolStr);
-			}
-			else if (const FStructProperty* StructProp = CastFieldChecked<const FStructProperty>(Property))
-			{
-				Ret += FString::Printf(TEXT("%s: (%s) "), *Property->GetNameCPP(), *FNiagaraTypeHelper::ToString(PropPtr, StructProp->Struct));
-			}
-			else
+			Ret += FString::Printf(TEXT("%g "), *(float*)ValueData);
+		}
+		else if (Struct == FNiagaraTypeDefinition::GetIntStruct())
+		{
+			Ret += FString::Printf(TEXT("%d "), *(int32*)ValueData);
+		}
+		else if (Struct == FNiagaraTypeDefinition::GetBoolStruct())
+		{
+			int32 Val = *(int32*)ValueData;
+			Ret += Val == 0xFFFFFFFF ? (TEXT("True")) : (Val == 0x0 ? TEXT("False") : TEXT("Invalid"));
+		}
+		else
+		{
+			for (TFieldIterator<FProperty> PropertyIt(Struct, EFieldIteratorFlags::IncludeSuper); PropertyIt; ++PropertyIt)
 			{
-				check(false);
-				Ret += TEXT("Unknown Type");
+				const FProperty* Property = *PropertyIt;
+				const uint8* PropPtr = ValueData + PropertyIt->GetOffset_ForInternal();
+				if (Property->IsA(FFloatProperty::StaticClass()))
+				{
+					Ret += FString::Printf(TEXT("%s: %g "), *Property->GetNameCPP(), *(float*)PropPtr);
+				}
+				else if (Property->IsA(FIntProperty::StaticClass()))
+				{
+					Ret += FString::Printf(TEXT("%s: %d "), *Property->GetNameCPP(), *(int32*)PropPtr);
+				}
+				else if (Property->IsA(FBoolProperty::StaticClass()))
+				{
+					int32 Val = *(int32*)ValueData;
+					FString BoolStr = Val == 0xFFFFFFFF ? (TEXT("True")) : (Val == 0x0 ? TEXT("False") : TEXT("Invalid"));
+					Ret += FString::Printf(TEXT("%s: %d "), *Property->GetNameCPP(), *BoolStr);
+				}
+				else if (const FStructProperty* StructProp = CastFieldChecked<const FStructProperty>(Property))
+				{
+					Ret += FString::Printf(TEXT("%s: (%s) "), *Property->GetNameCPP(), *FNiagaraTypeHelper::ToString(PropPtr, StructProp->Struct));
+				}
+				else
+				{
+					check(false);
+					Ret += TEXT("Unknown Type");
+				}
 			}
 		}
 	}
diff --git a/Engine/Plugins/FX/Niagara/Source/Niagara/Public/NiagaraTypes.h b/Engine/Plugins/FX/Niagara/Source/Niagara/Public/NiagaraTypes.h
--- a/Engine/Plugins/FX/Niagara/Source/Niagara/Public/NiagaraTypes.h
+++ b/Engine/Plugins/FX/Niagara/Source/Niagara/Public/NiagaraTypes.h
@@ -258,8 +258,7 @@ private:
 class NIAGARA_API FNiagaraTypeHelper
 {
 public:
-	static FString ToString(const uint8* ValueData, const UScriptStruct* Struct);
-
+	static FString ToString(const uint8* ValueData, const UObject* StructOrEnum);
 };
 
 /** Defines different modes for selecting the output numeric type of a function or operation based on the types of the inputs. */
@@ -626,7 +625,7 @@ public:
 		{
 			return TEXT("(null)");
 		}
-		return FNiagaraTypeHelper::ToString(ValueData, CastChecked<UScriptStruct>(ClassStructOrEnum));
+		return FNiagaraTypeHelper::ToString(ValueData, ClassStructOrEnum);
 	}
 
 	static bool TypesAreAssignable(const FNiagaraTypeDefinition& TypeA, const FNiagaraTypeDefinition& TypeB);
