Fixed concurrency problems by moving the callback to a separate task graph call

diff --git a/Engine/Plugins/FX/Niagara/Source/Niagara/Private/NiagaraDataInterfaceExport.cpp b/Engine/Plugins/FX/Niagara/Source/Niagara/Private/NiagaraDataInterfaceExport.cpp
--- a/Engine/Plugins/FX/Niagara/Source/Niagara/Private/NiagaraDataInterfaceExport.cpp
+++ b/Engine/Plugins/FX/Niagara/Source/Niagara/Private/NiagaraDataInterfaceExport.cpp
@@ -8,6 +8,38 @@
 
 const FName UNiagaraDataInterfaceExport::StoreDataName(TEXT("StoreParticleData"));
 
+/**
+Async task to call the blueprint callback on the game thread and isolate the niagara tick from the blueprint
+*/
+class FNiagaraExportCallbackAsyncTask
+{
+	UObject* CallbackHandler;
+	TArray<FBasicParticleData> Data;
+	UNiagaraSystem* System;
+
+public:
+	FNiagaraExportCallbackAsyncTask(UObject* CallbackHandler, TArray<FBasicParticleData>& Data, UNiagaraSystem* System)
+		: CallbackHandler(CallbackHandler)
+		, Data(Data)
+		, System(System)
+	{
+	}
+
+	FORCEINLINE TStatId GetStatId() const { RETURN_QUICK_DECLARE_CYCLE_STAT(FNiagaraExportCallbackAsyncTask, STATGROUP_TaskGraphTasks); }
+	static FORCEINLINE ENamedThreads::Type GetDesiredThread() { return ENamedThreads::GameThread; }
+	static FORCEINLINE ESubsequentsMode::Type GetSubsequentsMode() { return ESubsequentsMode::FireAndForget; }
+
+	void DoTask(ENamedThreads::Type CurrentThread, const FGraphEventRef& MyCompletionGraphEvent)
+	{
+		if (!CallbackHandler->IsValidLowLevelFast())
+		{
+			UE_LOG(LogNiagara, Warning, TEXT("Received invalid callback handle in data interface - was the callback object deleted? System %s"), *System->GetName());
+			return;
+		}
+		INiagaraParticleCallbackHandler::Execute_ReceiveParticleData(CallbackHandler, Data, System);
+	}
+};
+
 UNiagaraDataInterfaceExport::UNiagaraDataInterfaceExport(FObjectInitializer const& ObjectInitializer)
 	: Super(ObjectInitializer)
 {
@@ -60,7 +92,8 @@ bool UNiagaraDataInterfaceExport::PerInstanceTickPostSimulate(void* PerInstanceD
 		{
 			Data.Add(Value);
 		}
-		INiagaraParticleCallbackHandler::Execute_ReceiveParticleData(PIData->CallbackHandler, Data, SystemInstance->GetSystem());
+
+		TGraphTask<FNiagaraExportCallbackAsyncTask>::CreateTask().ConstructAndDispatchWhenReady(PIData->CallbackHandler, Data, SystemInstance->GetSystem());
 	}
 	return false;
 }
