Updated scalabiltiy engine test with Nick's helper object to hopefully get rid of all off by one frame errors on cooked/platforms.

diff --git a/Engine/Plugins/FX/Niagara/Source/Niagara/Private/NiagaraPreviewGrid.cpp b/Engine/Plugins/FX/Niagara/Source/Niagara/Private/NiagaraPreviewGrid.cpp
--- a/Engine/Plugins/FX/Niagara/Source/Niagara/Private/NiagaraPreviewGrid.cpp
+++ b/Engine/Plugins/FX/Niagara/Source/Niagara/Private/NiagaraPreviewGrid.cpp
@@ -192,6 +192,56 @@ void ANiagaraPreviewGrid::DeactivatePreviews()
 	}
 }
 
+void ANiagaraPreviewGrid::SetPaused(bool bPaused)
+{
+	for (int32 X = 0; X < NumX; ++X)
+	{
+		float XLocation = X / NumX - 1;
+
+		for (int32 Y = 0; Y < NumY; ++Y)
+		{
+			float YLocation = Y / NumY - 1;
+
+			int32 PreviewIdx = PreviewIndex(X, Y);
+			UChildActorComponent* PreviewComp = PreviewComponents[PreviewIdx];
+			ANiagaraPreviewBase* PreviewActor = CastChecked<ANiagaraPreviewBase>(PreviewComp->GetChildActor());
+
+			TArray<UNiagaraComponent*, TInlineAllocator<4>> NiagaraComponents;
+			PreviewActor->GetComponents<UNiagaraComponent>(NiagaraComponents);
+			for (UNiagaraComponent* Component : NiagaraComponents)
+			{
+				Component->SetPaused(bPaused);
+			}
+		}
+	}
+
+	SetActorTickEnabled(!bPaused);
+}
+
+void ANiagaraPreviewGrid::GetPreviews(TArray<UNiagaraComponent*>& OutPreviews)
+{
+	for (int32 X = 0; X < NumX; ++X)
+	{
+		float XLocation = X / NumX - 1;
+
+		for (int32 Y = 0; Y < NumY; ++Y)
+		{
+			float YLocation = Y / NumY - 1;
+
+			int32 PreviewIdx = PreviewIndex(X, Y);
+			UChildActorComponent* PreviewComp = PreviewComponents[PreviewIdx];
+			ANiagaraPreviewBase* PreviewActor = CastChecked<ANiagaraPreviewBase>(PreviewComp->GetChildActor());
+
+			TArray<UNiagaraComponent*, TInlineAllocator<4>> NiagaraComponents;
+			PreviewActor->GetComponents<UNiagaraComponent>(NiagaraComponents);
+			for (UNiagaraComponent* Component : NiagaraComponents)
+			{
+				OutPreviews.Add(Component);
+			}
+		}
+	}
+}
+
 void ANiagaraPreviewGrid::TickActor(float DeltaTime, enum ELevelTick TickType, FActorTickFunction& ThisTickFunction)
 {
 	Super::TickActor(DeltaTime, TickType, ThisTickFunction);
diff --git a/Engine/Plugins/FX/Niagara/Source/Niagara/Public/NiagaraPreviewGrid.h b/Engine/Plugins/FX/Niagara/Source/Niagara/Public/NiagaraPreviewGrid.h
--- a/Engine/Plugins/FX/Niagara/Source/Niagara/Public/NiagaraPreviewGrid.h
+++ b/Engine/Plugins/FX/Niagara/Source/Niagara/Public/NiagaraPreviewGrid.h
@@ -203,6 +203,13 @@ public:
 
 	UFUNCTION(BlueprintCallable, Category=Preview)
 	void DeactivatePreviews();
+
+	UFUNCTION(BlueprintCallable, Category = Preview)
+	void SetPaused(bool bPaused);
+
+	UFUNCTION(BlueprintCallable, Category = Preview)
+	void GetPreviews(TArray<UNiagaraComponent*>& OutPreviews);
+
 private:
 
 	void DestroyPreviews();
