Moved temporal AA jitter vector into view properties on camera DI

diff --git a/Engine/Plugins/FX/Niagara/Source/Niagara/Classes/NiagaraDataInterfaceCamera.h b/Engine/Plugins/FX/Niagara/Source/Niagara/Classes/NiagaraDataInterfaceCamera.h
--- a/Engine/Plugins/FX/Niagara/Source/Niagara/Classes/NiagaraDataInterfaceCamera.h
+++ b/Engine/Plugins/FX/Niagara/Source/Niagara/Classes/NiagaraDataInterfaceCamera.h
@@ -46,7 +46,6 @@ public:
 	void GetViewPropertiesGPU(FVectorVMContext& Context);
 	void GetClipSpaceTransformsGPU(FVectorVMContext& Context);
 	void GetViewSpaceTransformsGPU(FVectorVMContext& Context);
-	void GetTAAJitter(FVectorVMContext& Context);
 	
 private:
 	static const FName GetViewPropertiesName;
@@ -54,7 +53,6 @@ private:
 	static const FName GetViewSpaceTransformsName;
 	static const FName GetCameraPropertiesName;
 	static const FName GetFieldOfViewName;
-	static const FName GetTAAJitterName;
 };
 
 struct FNiagaraDataIntefaceProxyCameraQuery : public FNiagaraDataInterfaceProxy
diff --git a/Engine/Plugins/FX/Niagara/Source/Niagara/Private/NiagaraDataInterfaceCamera.cpp b/Engine/Plugins/FX/Niagara/Source/Niagara/Private/NiagaraDataInterfaceCamera.cpp
--- a/Engine/Plugins/FX/Niagara/Source/Niagara/Private/NiagaraDataInterfaceCamera.cpp
+++ b/Engine/Plugins/FX/Niagara/Source/Niagara/Private/NiagaraDataInterfaceCamera.cpp
@@ -12,7 +12,6 @@ const FName UNiagaraDataInterfaceCamera::GetClipSpaceTransformsName(TEXT("GetCli
 const FName UNiagaraDataInterfaceCamera::GetViewSpaceTransformsName(TEXT("GetViewSpaceTransformsGPU"));
 const FName UNiagaraDataInterfaceCamera::GetCameraPropertiesName(TEXT("GetCameraPropertiesCPU/GPU"));
 const FName UNiagaraDataInterfaceCamera::GetFieldOfViewName(TEXT("GetFieldOfView"));
-const FName UNiagaraDataInterfaceCamera::GetTAAJitterName(TEXT("GetTAAJitter"));
 
 UNiagaraDataInterfaceCamera::UNiagaraDataInterfaceCamera(FObjectInitializer const& ObjectInitializer)
 	: Super(ObjectInitializer)
@@ -85,6 +84,8 @@ void UNiagaraDataInterfaceCamera::GetFunctions(TArray<FNiagaraFunctionSignature>
 	Sig.Outputs.Add(FNiagaraVariable(FNiagaraTypeDefinition::GetVec3Def(), TEXT("View Right Vector")));
 	Sig.Outputs.Add(FNiagaraVariable(FNiagaraTypeDefinition::GetVec4Def(), TEXT("View Size And Inverse Size")));
 	Sig.Outputs.Add(FNiagaraVariable(FNiagaraTypeDefinition::GetVec4Def(), TEXT("Screen To View Space")));
+	Sig.Outputs.Add(FNiagaraVariable(FNiagaraTypeDefinition::GetVec2Def(), TEXT("Temporal AA Jitter (Current Frame)")));
+	Sig.Outputs.Add(FNiagaraVariable(FNiagaraTypeDefinition::GetVec2Def(), TEXT("Temporal AA Jitter (Previous Frame)")));
 	OutFunctions.Add(Sig);
 
 
@@ -151,18 +152,6 @@ void UNiagaraDataInterfaceCamera::GetFunctions(TArray<FNiagaraFunctionSignature>
 	Sig.Outputs.Add(FNiagaraVariable(FNiagaraTypeDefinition::GetVec3Def(), TEXT("Up Vector World")));
 	Sig.Outputs.Add(FNiagaraVariable(FNiagaraTypeDefinition::GetVec3Def(), TEXT("Right Vector World")));
 	OutFunctions.Add(Sig);
-
-	Sig = FNiagaraFunctionSignature();
-	Sig.Name = GetTAAJitterName;
-#if WITH_EDITORONLY_DATA
-	Sig.Description = NSLOCTEXT("Niagara", "GetTAAJitterDescription", "This function returns the TAA jitter values of the currently active camera.");
-#endif
-	Sig.bMemberFunction = true;
-	Sig.bRequiresContext = false;
-	Sig.Inputs.Add(FNiagaraVariable(FNiagaraTypeDefinition(GetClass()), TEXT("Camera interface")));
-	Sig.Outputs.Add(FNiagaraVariable(FNiagaraTypeDefinition::GetVec2Def(), TEXT("Current TAA Jitter (clip)")));
-	Sig.Outputs.Add(FNiagaraVariable(FNiagaraTypeDefinition::GetVec2Def(), TEXT("Previous TAA Jitter (clip)")));
-	OutFunctions.Add(Sig);
 }
 
 bool UNiagaraDataInterfaceCamera::GetFunctionHLSL(const FName& DefinitionFunctionName, FString InstanceFunctionName, FNiagaraDataInterfaceGPUParamInfo& ParamInfo, FString& OutHLSL)
@@ -173,7 +162,7 @@ bool UNiagaraDataInterfaceCamera::GetFunctionHLSL(const FName& DefinitionFunctio
 	if (DefinitionFunctionName == GetViewPropertiesName)
 	{
 		static const TCHAR *FormatSample = TEXT(R"(
-			void {FunctionName}(out float3 Out_ViewPositionWorld, out float3 Out_ViewForwardVector, out float3 Out_ViewUpVector, out float3 Out_ViewRightVector, out float4 Out_ViewSizeAndInverseSize, out float4 Out_ScreenToViewSpace)
+			void {FunctionName}(out float3 Out_ViewPositionWorld, out float3 Out_ViewForwardVector, out float3 Out_ViewUpVector, out float3 Out_ViewRightVector, out float4 Out_ViewSizeAndInverseSize, out float4 Out_ScreenToViewSpace, out float2 Out_Current_TAAJitter, out float2 Out_Previous_TAAJitter)
 			{
 				Out_ViewPositionWorld.xyz = View.WorldViewOrigin.xyz;
 				Out_ViewForwardVector.xyz = View.ViewForward.xyz;
@@ -181,6 +170,8 @@ bool UNiagaraDataInterfaceCamera::GetFunctionHLSL(const FName& DefinitionFunctio
 				Out_ViewRightVector.xyz = View.ViewRight.xyz;
 				Out_ViewSizeAndInverseSize = View.ViewSizeAndInvSize;
 				Out_ScreenToViewSpace = View.ScreenToViewSpace;
+				Out_Current_TAAJitter = View.TemporalAAJitter.xy;
+				Out_Previous_TAAJitter = View.TemporalAAJitter.zw;
 			}
 		)");
 		OutHLSL += FString::Format(FormatSample, ArgsSample);
@@ -247,18 +238,6 @@ bool UNiagaraDataInterfaceCamera::GetFunctionHLSL(const FName& DefinitionFunctio
 		OutHLSL += FString::Format(FormatSample, ArgsSample);
 		return true;
 	}
-	if (DefinitionFunctionName == GetTAAJitterName)
-	{
-		static const TCHAR *FormatSample = TEXT(R"(
-			void {FunctionName}(out float2 Out_CurrentJitter, out float2 Out_PreviousJitter)
-			{
-				Out_CurrentJitter = View.TemporalAAJitter.xy;
-				Out_PreviousJitter = View.TemporalAAJitter.zw;
-			}
-		)");
-		OutHLSL += FString::Format(FormatSample, ArgsSample);
-		return true;
-	}
 	return false;
 }
 
@@ -267,7 +246,6 @@ DEFINE_NDI_DIRECT_FUNC_BINDER(UNiagaraDataInterfaceCamera, GetCameraProperties);
 DEFINE_NDI_DIRECT_FUNC_BINDER(UNiagaraDataInterfaceCamera, GetViewPropertiesGPU);
 DEFINE_NDI_DIRECT_FUNC_BINDER(UNiagaraDataInterfaceCamera, GetClipSpaceTransformsGPU);
 DEFINE_NDI_DIRECT_FUNC_BINDER(UNiagaraDataInterfaceCamera, GetViewSpaceTransformsGPU);
-DEFINE_NDI_DIRECT_FUNC_BINDER(UNiagaraDataInterfaceCamera, GetTAAJitter);
 void UNiagaraDataInterfaceCamera::GetVMExternalFunction(const FVMExternalFunctionBindingInfo& BindingInfo, void* InstanceData, FVMExternalFunction &OutFunc)
 {
 	if (BindingInfo.Name == GetFieldOfViewName)
@@ -290,10 +268,6 @@ void UNiagaraDataInterfaceCamera::GetVMExternalFunction(const FVMExternalFunctio
 	{
 		NDI_FUNC_BINDER(UNiagaraDataInterfaceCamera, GetViewPropertiesGPU)::Bind(this, OutFunc);
 	}
-	else if (BindingInfo.Name == GetTAAJitterName)
-	{
-		NDI_FUNC_BINDER(UNiagaraDataInterfaceCamera, GetTAAJitter)::Bind(this, OutFunc);
-	}
 	else
 	{
 		UE_LOG(LogNiagara, Error, TEXT("Could not find data interface external function. Received Name: %s"), *BindingInfo.Name.ToString());
@@ -373,15 +347,15 @@ void UNiagaraDataInterfaceCamera::GetViewPropertiesGPU(FVectorVMContext& Context
 {
 	VectorVM::FUserPtrHandler<CameraDataInterface_InstanceData> InstData(Context);
 	TArray<VectorVM::FExternalFuncRegisterHandler<float>> OutParams;
-	OutParams.Reserve(20);
-	for (int i = 0; i < 20; i++)
+	OutParams.Reserve(24);
+	for (int i = 0; i < 24; i++)
 	{
 		OutParams.Emplace(Context);
 	}
 
 	for (int32 k = 0; k < Context.NumInstances; ++k)
 	{
-		for (int i = 0; i < 20; i++)
+		for (int i = 0; i < 24; i++)
 		{
 			*OutParams[i].GetDestAndAdvance() = 0;
 		}
@@ -426,26 +400,3 @@ void UNiagaraDataInterfaceCamera::GetViewSpaceTransformsGPU(FVectorVMContext& Co
 	}
 }
 
-void UNiagaraDataInterfaceCamera::GetTAAJitter(FVectorVMContext& Context)
-{
-	VectorVM::FUserPtrHandler<CameraDataInterface_InstanceData> InstData(Context);
-	TArray<VectorVM::FExternalFuncRegisterHandler<float>> OutParams;
-
-	constexpr int32 ElementCount = 2 * sizeof(FVector2D) / sizeof(float);
-	OutParams.Reserve(ElementCount);
-	for (int i = 0; i < ElementCount; i++)
-	{
-		OutParams.Emplace(Context);
-	}
-
-	for (int32 k = 0; k < Context.NumInstances; ++k)
-	{
-		for (int i = 0; i < ElementCount; i++)
-		{
-			*OutParams[i].GetDestAndAdvance() = 0;
-		}
-	}
-}
-
-
-
