SNiagaraStackDisplayName no longer has a ColorAndOpacity argument, since it was always set to SNiagaraStackEntryWidget::GetTextColorForSearch() - instead, it too inherits from that class and uses it.

Required to support display names in niagara stack entries.

diff --git a/Engine/Plugins/FX/Niagara/Source/NiagaraEditorWidgets/Private/SNiagaraStackEntryWidget.cpp b/Engine/Plugins/FX/Niagara/Source/NiagaraEditorWidgets/Private/SNiagaraStackEntryWidget.cpp
--- a/Engine/Plugins/FX/Niagara/Source/NiagaraEditorWidgets/Private/SNiagaraStackEntryWidget.cpp
+++ b/Engine/Plugins/FX/Niagara/Source/NiagaraEditorWidgets/Private/SNiagaraStackEntryWidget.cpp
@@ -10,12 +10,10 @@
 
 void SNiagaraStackDisplayName::Construct(const FArguments& InArgs, UNiagaraStackEntry& InStackEntry, UNiagaraStackViewModel& InStackViewModel, FName InTextStyleName)
 {
-	StackEntry = &InStackEntry;
+	StackEntryItem = &InStackEntry;
 	StackViewModel = &InStackViewModel;
 	TextStyleName = InTextStyleName;
 
-	ColorAndOpacity = InArgs._ColorAndOpacity;
-
 	StackViewModel->OnStructureChanged().AddSP(this, &SNiagaraStackDisplayName::StackViewModelStructureChanged);
 
 	ChildSlot
@@ -36,10 +34,10 @@ TSharedRef<SWidget> SNiagaraStackDisplayName::ConstructChildren()
 {
 	TSharedRef<STextBlock> BaseNameWidget = SNew(STextBlock)
 		.TextStyle(FNiagaraEditorWidgetsStyle::Get(), TextStyleName)
-		.ToolTipText_UObject(StackEntry, &UNiagaraStackEntry::GetTooltipText)
-		.Text_UObject(StackEntry, &UNiagaraStackEntry::GetDisplayName)
+		.ToolTipText_UObject(StackEntryItem, &UNiagaraStackEntry::GetTooltipText)
+		.Text_UObject(StackEntryItem, &UNiagaraStackEntry::GetDisplayName)
 		.HighlightText_UObject(StackViewModel, &UNiagaraStackViewModel::GetCurrentSearchText)
-		.ColorAndOpacity(ColorAndOpacity)
+		.ColorAndOpacity(this, &SNiagaraStackDisplayName::GetTextColorForSearch, FSlateColor::UseForeground())
 		.IsEnabled(this, &SNiagaraStackDisplayName::GetIsEnabled);
 
 	int32 NumTopLevelEmitters = 0;
@@ -59,7 +57,7 @@ TSharedRef<SWidget> SNiagaraStackDisplayName::ConstructChildren()
 	else
 	{
 		TopLevelViewModelCountAtLastConstruction = StackViewModel->GetTopLevelViewModels().Num();
-		TSharedPtr<UNiagaraStackViewModel::FTopLevelViewModel> TopLevelViewModel = StackViewModel->GetTopLevelViewModelForEntry(*StackEntry);
+		TSharedPtr<UNiagaraStackViewModel::FTopLevelViewModel> TopLevelViewModel = StackViewModel->GetTopLevelViewModelForEntry(*StackEntryItem);
 		if(TopLevelViewModel.IsValid())
 		{
 			return SNew(SWrapBox)
@@ -69,10 +67,10 @@ TSharedRef<SWidget> SNiagaraStackDisplayName::ConstructChildren()
 				[
 					SNew(STextBlock)
 					.TextStyle(FNiagaraEditorWidgetsStyle::Get(), TextStyleName)
-					.ToolTipText_UObject(StackEntry, &UNiagaraStackEntry::GetTooltipText)
+					.ToolTipText_UObject(StackEntryItem, &UNiagaraStackEntry::GetTooltipText)
 					.Text(this, &SNiagaraStackDisplayName::GetTopLevelDisplayName, TWeakPtr<UNiagaraStackViewModel::FTopLevelViewModel>(TopLevelViewModel))
 					.HighlightText_UObject(StackViewModel, &UNiagaraStackViewModel::GetCurrentSearchText)
-					.ColorAndOpacity(ColorAndOpacity)
+					.ColorAndOpacity(this, &SNiagaraStackDisplayName::GetTextColorForSearch, FSlateColor::UseForeground())
 					.IsEnabled(this, &SNiagaraStackDisplayName::GetIsEnabled)
 				]
 				+ SWrapBox::Slot()
@@ -105,7 +103,7 @@ FText SNiagaraStackDisplayName::GetTopLevelDisplayName(TWeakPtr<UNiagaraStackVie
 
 void SNiagaraStackDisplayName::StackViewModelStructureChanged()
 {
-	if (StackEntry->IsFinalized() == false && StackViewModel->GetTopLevelViewModels().Num() != TopLevelViewModelCountAtLastConstruction)
+	if (StackEntryItem->IsFinalized() == false && StackViewModel->GetTopLevelViewModels().Num() != TopLevelViewModelCountAtLastConstruction)
 	{
 		Container->SetContent(ConstructChildren());
 	}
@@ -113,17 +111,17 @@ void SNiagaraStackDisplayName::StackViewModelStructureChanged()
 
 bool SNiagaraStackDisplayName::GetIsEnabled() const
 {
-	return StackEntry->GetOwnerIsEnabled() && StackEntry->GetIsEnabled();
+	return StackEntryItem->GetOwnerIsEnabled() && StackEntryItem->GetIsEnabled();
 }
 
-FSlateColor SNiagaraStackEntryWidget::GetTextColorForSearch() const
+FSlateColor SNiagaraStackEntryWidget::GetTextColorForSearch(FSlateColor DefaultColor) const
 {
 	if (IsCurrentSearchMatch())
 	{
 		return FNiagaraEditorWidgetsStyle::Get().GetColor("NiagaraEditor.Stack.SearchHighlightColor");
 	} 
 	
-	return FSlateColor::UseForeground();
+	return DefaultColor;
 }
 
 bool SNiagaraStackEntryWidget::IsCurrentSearchMatch() const
diff --git a/Engine/Plugins/FX/Niagara/Source/NiagaraEditorWidgets/Private/SNiagaraStackEntryWidget.h b/Engine/Plugins/FX/Niagara/Source/NiagaraEditorWidgets/Private/SNiagaraStackEntryWidget.h
--- a/Engine/Plugins/FX/Niagara/Source/NiagaraEditorWidgets/Private/SNiagaraStackEntryWidget.h
+++ b/Engine/Plugins/FX/Niagara/Source/NiagaraEditorWidgets/Private/SNiagaraStackEntryWidget.h
@@ -10,11 +10,24 @@ class UNiagaraStackViewModel;
 class UNiagaraStackEntry;
 class SBox;
 
-class SNiagaraStackDisplayName : public SCompoundWidget
+class SNiagaraStackEntryWidget : public SCompoundWidget
+{
+public:
+	FSlateColor GetTextColorForSearch(FSlateColor DefaultColor) const;
+	FReply ExpandEntry();
+	
+protected:
+	bool IsCurrentSearchMatch() const;
+	
+protected:
+	UNiagaraStackViewModel* StackViewModel;
+	UNiagaraStackEntry* StackEntryItem;
+};
+
+class SNiagaraStackDisplayName : public SNiagaraStackEntryWidget
 {
 public:
 	SLATE_BEGIN_ARGS(SNiagaraStackDisplayName) { }
-		SLATE_ATTRIBUTE(FSlateColor, ColorAndOpacity);
 	SLATE_END_ARGS();
 
 	void Construct(const FArguments& InArgs, UNiagaraStackEntry& InStackEntry, UNiagaraStackViewModel& InStackViewModel, FName InTextStyleName);
@@ -31,30 +44,11 @@ private:
 	bool GetIsEnabled() const;
 
 private:
-	UNiagaraStackEntry* StackEntry;
-	UNiagaraStackViewModel* StackViewModel;
 	FName TextStyleName;
 
-	TAttribute<FSlateColor> ColorAndOpacity;
-
 	TSharedPtr<SBox> Container;
 
 	mutable FText TopLevelDisplayNameCache;
 	mutable FText TopLevelDisplayNameFormattedCache;
 	int32 TopLevelViewModelCountAtLastConstruction;
 };
-
-class SNiagaraStackEntryWidget : public SCompoundWidget
-{
-public:
-	FSlateColor GetTextColorForSearch() const;
-	FReply ExpandEntry();
-	
-protected:
-	bool IsCurrentSearchMatch() const;
-	
-protected:
-	UNiagaraStackViewModel * StackViewModel;
-	
-	UNiagaraStackEntry * StackEntryItem;
-};
\ No newline at end of file
diff --git a/Engine/Plugins/FX/Niagara/Source/NiagaraEditorWidgets/Private/Stack/SNiagaraStackEmitterPropertiesItem.cpp b/Engine/Plugins/FX/Niagara/Source/NiagaraEditorWidgets/Private/Stack/SNiagaraStackEmitterPropertiesItem.cpp
--- a/Engine/Plugins/FX/Niagara/Source/NiagaraEditorWidgets/Private/Stack/SNiagaraStackEmitterPropertiesItem.cpp
+++ b/Engine/Plugins/FX/Niagara/Source/NiagaraEditorWidgets/Private/Stack/SNiagaraStackEmitterPropertiesItem.cpp
@@ -26,7 +26,6 @@ void SNiagaraStackEmitterPropertiesItem::Construct(const FArguments& InArgs, UNi
 		.Padding(0)
 		[
 			SNew(SNiagaraStackDisplayName, InEmitterPropertiesItem, *InStackViewModel, "NiagaraEditor.Stack.ItemText")
-			.ColorAndOpacity(this, &SNiagaraStackEntryWidget::GetTextColorForSearch)
 		]
 		// Reset to base Button
 		+ SHorizontalBox::Slot()
diff --git a/Engine/Plugins/FX/Niagara/Source/NiagaraEditorWidgets/Private/Stack/SNiagaraStackErrorItem.cpp b/Engine/Plugins/FX/Niagara/Source/NiagaraEditorWidgets/Private/Stack/SNiagaraStackErrorItem.cpp
--- a/Engine/Plugins/FX/Niagara/Source/NiagaraEditorWidgets/Private/Stack/SNiagaraStackErrorItem.cpp
+++ b/Engine/Plugins/FX/Niagara/Source/NiagaraEditorWidgets/Private/Stack/SNiagaraStackErrorItem.cpp
@@ -54,7 +54,7 @@ void SNiagaraStackErrorItem::Construct(const FArguments& InArgs, UNiagaraStackEr
 			SNew(STextBlock)
 			.TextStyle(FNiagaraEditorStyle::Get(), "NiagaraEditor.ParameterText")
 			.Text_UObject(ErrorItem, &UNiagaraStackErrorItem::GetDisplayName)
-			.ColorAndOpacity(this, &SNiagaraStackErrorItem::GetTextColorForSearch)
+			.ColorAndOpacity(this, &SNiagaraStackErrorItem::GetTextColorForSearch, FSlateColor::UseForeground())
 			.HighlightText_UObject(StackViewModel, &UNiagaraStackViewModel::GetCurrentSearchText)
 		]
 	];
@@ -75,7 +75,7 @@ void SNiagaraStackErrorItemFix::Construct(const FArguments& InArgs, UNiagaraStac
 			SNew(STextBlock)
 			.TextStyle(FNiagaraEditorStyle::Get(), "NiagaraEditor.ParameterText")
 			.Text_UObject(ErrorItem, &UNiagaraStackErrorItemFix::GetDisplayName)
-			.ColorAndOpacity(this, &SNiagaraStackErrorItemFix::GetTextColorForSearch)
+			.ColorAndOpacity(this, &SNiagaraStackErrorItemFix::GetTextColorForSearch, FSlateColor::UseForeground())
 			.HighlightText_UObject(StackViewModel, &UNiagaraStackViewModel::GetCurrentSearchText)
 			.AutoWrapText(true)
 		]
diff --git a/Engine/Plugins/FX/Niagara/Source/NiagaraEditorWidgets/Private/Stack/SNiagaraStackEventHandlerPropertiesItem.cpp b/Engine/Plugins/FX/Niagara/Source/NiagaraEditorWidgets/Private/Stack/SNiagaraStackEventHandlerPropertiesItem.cpp
--- a/Engine/Plugins/FX/Niagara/Source/NiagaraEditorWidgets/Private/Stack/SNiagaraStackEventHandlerPropertiesItem.cpp
+++ b/Engine/Plugins/FX/Niagara/Source/NiagaraEditorWidgets/Private/Stack/SNiagaraStackEventHandlerPropertiesItem.cpp
@@ -26,7 +26,6 @@ void SNiagaraStackEventHandlerPropertiesItem::Construct(const FArguments& InArgs
 		.Padding(0)
 		[
 			SNew(SNiagaraStackDisplayName, InEventHandlerPropertiesItem, *InStackViewModel, "NiagaraEditor.Stack.ItemText")
-			.ColorAndOpacity(this, &SNiagaraStackEntryWidget::GetTextColorForSearch)
 		]
 		// Reset to base Button
 		+ SHorizontalBox::Slot()
diff --git a/Engine/Plugins/FX/Niagara/Source/NiagaraEditorWidgets/Private/Stack/SNiagaraStackFunctionInputName.cpp b/Engine/Plugins/FX/Niagara/Source/NiagaraEditorWidgets/Private/Stack/SNiagaraStackFunctionInputName.cpp
--- a/Engine/Plugins/FX/Niagara/Source/NiagaraEditorWidgets/Private/Stack/SNiagaraStackFunctionInputName.cpp
+++ b/Engine/Plugins/FX/Niagara/Source/NiagaraEditorWidgets/Private/Stack/SNiagaraStackFunctionInputName.cpp
@@ -43,7 +43,7 @@ void SNiagaraStackFunctionInputName::Construct(const FArguments& InArgs, UNiagar
 			.IsSelected(this, &SNiagaraStackFunctionInputName::GetIsNameWidgetSelected)
 			.OnTextCommitted(this, &SNiagaraStackFunctionInputName::OnNameTextCommitted)
 			.HighlightText_UObject(InStackViewModel, &UNiagaraStackViewModel::GetCurrentSearchText)
-			.ColorAndOpacity(this, &SNiagaraStackFunctionInputName::GetTextColorForSearch)
+			.ColorAndOpacity(this, &SNiagaraStackFunctionInputName::GetTextColorForSearch, FSlateColor::UseForeground())
 			.ToolTipText(this, &SNiagaraStackFunctionInputName::GetToolTipText)
 		]
 	];
diff --git a/Engine/Plugins/FX/Niagara/Source/NiagaraEditorWidgets/Private/Stack/SNiagaraStackItemGroup.cpp b/Engine/Plugins/FX/Niagara/Source/NiagaraEditorWidgets/Private/Stack/SNiagaraStackItemGroup.cpp
--- a/Engine/Plugins/FX/Niagara/Source/NiagaraEditorWidgets/Private/Stack/SNiagaraStackItemGroup.cpp
+++ b/Engine/Plugins/FX/Niagara/Source/NiagaraEditorWidgets/Private/Stack/SNiagaraStackItemGroup.cpp
@@ -29,7 +29,6 @@ void SNiagaraStackItemGroup::Construct(const FArguments& InArgs, UNiagaraStackIt
 		.Padding(2, 0, 0, 0)
 		[
 			SNew(SNiagaraStackDisplayName, InGroup, *InStackViewModel, "NiagaraEditor.Stack.GroupText")
-			.ColorAndOpacity(this, &SNiagaraStackEntryWidget::GetTextColorForSearch)
 		]
 		// Delete group button
 		+ SHorizontalBox::Slot()
diff --git a/Engine/Plugins/FX/Niagara/Source/NiagaraEditorWidgets/Private/Stack/SNiagaraStackModuleItem.cpp b/Engine/Plugins/FX/Niagara/Source/NiagaraEditorWidgets/Private/Stack/SNiagaraStackModuleItem.cpp
--- a/Engine/Plugins/FX/Niagara/Source/NiagaraEditorWidgets/Private/Stack/SNiagaraStackModuleItem.cpp
+++ b/Engine/Plugins/FX/Niagara/Source/NiagaraEditorWidgets/Private/Stack/SNiagaraStackModuleItem.cpp
@@ -53,7 +53,6 @@ void SNiagaraStackModuleItem::Construct(const FArguments& InArgs, UNiagaraStackM
 			.Padding(2, 0, 0, 0)
 			[
 				SNew(SNiagaraStackDisplayName, InModuleItem, *InStackViewModel, "NiagaraEditor.Stack.ItemText")
-				.ColorAndOpacity(this, &SNiagaraStackModuleItem::GetTextColorForSearch)
 			]
 			// Raise Action Menu button
 			+ SHorizontalBox::Slot()
diff --git a/Engine/Plugins/FX/Niagara/Source/NiagaraEditorWidgets/Private/Stack/SNiagaraStackParameterStoreEntryName.cpp b/Engine/Plugins/FX/Niagara/Source/NiagaraEditorWidgets/Private/Stack/SNiagaraStackParameterStoreEntryName.cpp
--- a/Engine/Plugins/FX/Niagara/Source/NiagaraEditorWidgets/Private/Stack/SNiagaraStackParameterStoreEntryName.cpp
+++ b/Engine/Plugins/FX/Niagara/Source/NiagaraEditorWidgets/Private/Stack/SNiagaraStackParameterStoreEntryName.cpp
@@ -31,7 +31,7 @@ void SNiagaraStackParameterStoreEntryName::Construct(const FArguments& InArgs, U
 			.OnTextCommitted(this, &SNiagaraStackParameterStoreEntryName::OnNameTextCommitted)
 			.OnVerifyTextChanged(this, &SNiagaraStackParameterStoreEntryName::VerifyNameTextChanged)
 			.HighlightText_UObject(InStackViewModel, &UNiagaraStackViewModel::GetCurrentSearchText)
-			.ColorAndOpacity(this, &SNiagaraStackParameterStoreEntryName::GetTextColorForSearch)
+			.ColorAndOpacity(this, &SNiagaraStackParameterStoreEntryName::GetTextColorForSearch, FSlateColor::UseForeground())
 		]
 	];
 }
diff --git a/Engine/Plugins/FX/Niagara/Source/NiagaraEditorWidgets/Private/Stack/SNiagaraStackRendererItem.cpp b/Engine/Plugins/FX/Niagara/Source/NiagaraEditorWidgets/Private/Stack/SNiagaraStackRendererItem.cpp
--- a/Engine/Plugins/FX/Niagara/Source/NiagaraEditorWidgets/Private/Stack/SNiagaraStackRendererItem.cpp
+++ b/Engine/Plugins/FX/Niagara/Source/NiagaraEditorWidgets/Private/Stack/SNiagaraStackRendererItem.cpp
@@ -42,7 +42,6 @@ void SNiagaraStackRendererItem::Construct(const FArguments& InArgs, UNiagaraStac
 		.VAlign(VAlign_Center)
 		[
 			SNew(SNiagaraStackDisplayName, InRendererItem, *InStackViewModel, "NiagaraEditor.Stack.ItemText")
-			.ColorAndOpacity(this, &SNiagaraStackEntryWidget::GetTextColorForSearch)
 		]
 		// Reset to base Button
 		+ SHorizontalBox::Slot()
