Niagara View Options combo button in the Stack view is now highlighted with an orange background if any options are set to non-defaults.

diff --git a/Engine/Plugins/FX/Niagara/Source/NiagaraEditorWidgets/Private/NiagaraEditorWidgetsStyle.cpp b/Engine/Plugins/FX/Niagara/Source/NiagaraEditorWidgets/Private/NiagaraEditorWidgetsStyle.cpp
--- a/Engine/Plugins/FX/Niagara/Source/NiagaraEditorWidgets/Private/NiagaraEditorWidgetsStyle.cpp
+++ b/Engine/Plugins/FX/Niagara/Source/NiagaraEditorWidgets/Private/NiagaraEditorWidgetsStyle.cpp
@@ -150,6 +150,8 @@ TSharedRef< FSlateStyleSet > FNiagaraEditorWidgetsStyle::Create()
 	Style->Set("NiagaraEditor.Stack.GroupForegroundColor", FLinearColor(FColor(220, 220, 220)));
 	Style->Set("NiagaraEditor.Stack.FlatButtonColor", FLinearColor(FColor(205, 205, 205)));
 
+	Style->Set("NiagaraEditor.Stack.HighlightedButtonBrush", new BOX_CORE_BRUSH("Common/ButtonHoverHint", FMargin(4 / 16.0f), FLinearColor(FColor(218, 130, 9))));
+
 	Style->Set("NiagaraEditor.Stack.AccentColor.System", FLinearColor(FColor(67, 105, 124)));
 	Style->Set("NiagaraEditor.Stack.AccentColor.Emitter", FLinearColor(FColor(126, 87, 67)));
 	Style->Set("NiagaraEditor.Stack.AccentColor.Particle", FLinearColor(FColor(87, 107, 61)));
diff --git a/Engine/Plugins/FX/Niagara/Source/NiagaraEditorWidgets/Private/SNiagaraStack.cpp b/Engine/Plugins/FX/Niagara/Source/NiagaraEditorWidgets/Private/SNiagaraStack.cpp
--- a/Engine/Plugins/FX/Niagara/Source/NiagaraEditorWidgets/Private/SNiagaraStack.cpp
+++ b/Engine/Plugins/FX/Niagara/Source/NiagaraEditorWidgets/Private/SNiagaraStack.cpp
@@ -350,28 +350,54 @@ TSharedRef<SWidget> SNiagaraStack::ConstructHeaderWidget()
 			.AutoWidth()
 			.Padding(4, 0, 0, 0)
 			[
-				SNew(SComboButton)
-				.ButtonStyle(FEditorStyle::Get(), "HoverHintOnly")
-				.ForegroundColor(FSlateColor::UseForeground())
-				.ToolTipText(LOCTEXT("ViewOptionsToolTip", "View Options"))
-				.OnGetMenuContent(this, &SNiagaraStack::GetViewOptionsMenu)
-				.ContentPadding(0)
-				.MenuPlacement(MenuPlacement_BelowRightAnchor)
-				.ButtonContent()
+				SNew(SBorder)
+				.Padding(0)
+				.BorderImage(this, &SNiagaraStack::GetViewOptionsIconBrush)
 				[
-					SNew(SBox)
-					.HAlign(HAlign_Center)
-					.VAlign(VAlign_Center)
+					SNew(SComboButton)
+					.ButtonStyle(FEditorStyle::Get(), "HoverHintOnly")
+					.ForegroundColor(FSlateColor::UseForeground())
+					.ToolTipText(LOCTEXT("ViewOptionsToolTip", "View Options"))
+					.OnGetMenuContent(this, &SNiagaraStack::GetViewOptionsMenu)
+					.ContentPadding(1)
+					.MenuPlacement(MenuPlacement_BelowRightAnchor)
+					.ButtonContent()
 					[
-						SNew(SImage)
-						.Image(FEditorStyle::GetBrush("GenericViewButton"))
-						.ColorAndOpacity(FNiagaraEditorWidgetsStyle::Get().GetColor("NiagaraEditor.Stack.FlatButtonColor"))
+						SNew(SOverlay)
+						// drop shadow
+						+ SOverlay::Slot()
+						.VAlign(VAlign_Top)
+						.Padding(0, 1, 0, 0)
+						[
+							SNew(SImage)
+							.Image(FEditorStyle::GetBrush("GenericViewButton"))
+							.ColorAndOpacity(FLinearColor::Black)
+						]
+						+ SOverlay::Slot()
+						.VAlign(VAlign_Top)
+						[
+							SNew(SImage)
+							.Image(FEditorStyle::GetBrush("GenericViewButton"))
+							.ColorAndOpacity(FNiagaraEditorWidgetsStyle::Get().GetColor("NiagaraEditor.Stack.FlatButtonColor"))
+						]
 					]
 				]
 			]
 		];
 }
 
+const FSlateBrush* SNiagaraStack::GetViewOptionsIconBrush() const
+{
+	bool bIsDefault = StackViewModel->GetShowAllAdvanced() == false &&
+		StackViewModel->GetShowLinkedInputs() == false &&
+		StackViewModel->GetShowOutputs() == false &&
+		StackViewModel->GetShowOnlyIssues() == false;
+
+	return bIsDefault ? 
+		FCoreStyle::Get().GetBrush("NoBrush") :
+		FNiagaraEditorWidgetsStyle::Get().GetBrush("NiagaraEditor.Stack.HighlightedButtonBrush");
+}
+
 void SNiagaraStack::OnSearchTextChanged(const FText& SearchText)
 {
 	bNeedsJumpToNextOccurence = true;
diff --git a/Engine/Plugins/FX/Niagara/Source/NiagaraEditorWidgets/Private/SNiagaraStack.h b/Engine/Plugins/FX/Niagara/Source/NiagaraEditorWidgets/Private/SNiagaraStack.h
--- a/Engine/Plugins/FX/Niagara/Source/NiagaraEditorWidgets/Private/SNiagaraStack.h
+++ b/Engine/Plugins/FX/Niagara/Source/NiagaraEditorWidgets/Private/SNiagaraStack.h
@@ -94,6 +94,7 @@ private:
 	void CollapseAll();
 
 	TSharedRef<SWidget> GetViewOptionsMenu() const;
+	const FSlateBrush* GetViewOptionsIconBrush() const;
 
 	// Drag/Drop
 	FReply OnRowDragDetected(const FGeometry& InGeometry, const FPointerEvent& InPointerEvent, UNiagaraStackEntry* InStackEntry);
