Niagara overview nodes now have their issue icon right-aligned. Clicking the icon cycles through modules in the selected emitter that contain issues.

diff --git a/Engine/Plugins/FX/Niagara/Source/NiagaraEditorWidgets/Private/SNiagaraOverviewStackNode.cpp b/Engine/Plugins/FX/Niagara/Source/NiagaraEditorWidgets/Private/SNiagaraOverviewStackNode.cpp
--- a/Engine/Plugins/FX/Niagara/Source/NiagaraEditorWidgets/Private/SNiagaraOverviewStackNode.cpp
+++ b/Engine/Plugins/FX/Niagara/Source/NiagaraEditorWidgets/Private/SNiagaraOverviewStackNode.cpp
@@ -3,7 +3,10 @@
 #include "SNiagaraOverviewStackNode.h"
 #include "NiagaraOverviewNode.h"
 #include "ViewModels/NiagaraSystemViewModel.h"
+#include "ViewModels/NiagaraEmitterViewModel.h"
 #include "ViewModels/NiagaraEmitterHandleViewModel.h"
+#include "ViewModels/NiagaraSystemSelectionViewModel.h"
+#include "ViewModels/Stack/NiagaraStackModuleItem.h"
 #include "ViewModels/Stack/NiagaraStackViewModel.h"
 #include "SNiagaraOverviewStack.h"
 #include "NiagaraEditorModule.h"
@@ -15,11 +18,9 @@
 #include "Widgets/Input/SButton.h"
 #include "Widgets/Input/SCheckBox.h"
 #include "AssetThumbnail.h"
-#include "ViewModels/NiagaraEmitterViewModel.h"
 #include "Widgets/SToolTip.h"
 #include "Widgets/SBoxPanel.h"
 #include "NiagaraRendererProperties.h"
-#include "ViewModels/NiagaraSystemSelectionViewModel.h"
 #include "Widgets/SWidget.h"
 
 #define LOCTEXT_NAMESPACE "NiagaraOverviewStackNode"
@@ -32,6 +33,8 @@ void SNiagaraOverviewStackNode::Construct(const FArguments& InArgs, UNiagaraOver
 	OverviewStackNode = InNode;
 	StackViewModel = nullptr;
 	OverviewSelectionViewModel = nullptr;
+	CurrentIssueIndex = -1;
+
 	EmitterHandleViewModelWeak.Reset();
 	ThumbnailPool = MakeShared<FAssetThumbnailPool>(100, TAttribute<bool>::Create(TAttribute<bool>::FGetter::CreateSP(this, &SNiagaraOverviewStackNode::IsHoveringThumbnail)));
 	ThumbnailBar = SNew(SHorizontalBox);
@@ -120,19 +123,58 @@ TSharedRef<SWidget> SNiagaraOverviewStackNode::CreateTitleWidget(TSharedPtr<SNod
 		.FillWidth(1.0f)
 		[
 			DefaultTitle
-		]
-		// Stack issues icon
-		+ SHorizontalBox::Slot()
-		.AutoWidth()
-		.VAlign(VAlign_Center)
-		.HAlign(HAlign_Right)
-		.Padding(5, 0, 0, 0)
+		];
+}
+
+TSharedRef<SWidget> SNiagaraOverviewStackNode::CreateTitleRightWidget()
+{
+	return SNew(SBox)
+		.Content()
 		[
 			SNew(SNiagaraStackIssueIcon, StackViewModel, StackViewModel->GetRootEntry())
 			.Visibility(this, &SNiagaraOverviewStackNode::GetIssueIconVisibility)
+			.OnClicked(this, &SNiagaraOverviewStackNode::OnCycleThroughIssues)
 		];
 }
 
+FReply SNiagaraOverviewStackNode::OnCycleThroughIssues()
+{
+	TSharedPtr<FNiagaraEmitterHandleViewModel> EmitterHandleViewModel = EmitterHandleViewModelWeak.Pin();
+	if (EmitterHandleViewModel.IsValid())
+	{
+		const TArray<UNiagaraStackEntry*>& ChildrenWithIssues = StackViewModel->GetRootEntry()->GetAllChildrenWithIssues();
+		if (ChildrenWithIssues.Num() > 0)
+		{
+			++CurrentIssueIndex;
+
+			if (CurrentIssueIndex >= ChildrenWithIssues.Num())
+			{
+				CurrentIssueIndex = 0;
+			}
+
+			if (ChildrenWithIssues.IsValidIndex(CurrentIssueIndex))
+			{
+				UNiagaraStackEntry* ChildToSelect = ChildrenWithIssues[CurrentIssueIndex];
+				UNiagaraStackModuleItem* ModuleToSelect = Cast<UNiagaraStackModuleItem>(ChildToSelect);
+				if (ModuleToSelect == nullptr)
+				{
+					ModuleToSelect = ChildToSelect->GetTypedOuter<UNiagaraStackModuleItem>();
+				}
+
+				if (ModuleToSelect != nullptr)
+				{
+					OverviewSelectionViewModel->UpdateSelectedEntries(TArray<UNiagaraStackEntry*> { ModuleToSelect },
+						TArray<UNiagaraStackEntry*>(), true /* bClearCurrentSelection */ );
+				}
+			}
+		}
+
+		return FReply::Handled();
+	}
+
+	return FReply::Unhandled();
+}
+
 TSharedRef<SWidget> SNiagaraOverviewStackNode::CreateThumbnailWidget(float InThumbnailSize, FRendererPreviewData* InData)
 {
 	TSharedRef<SWidget> ThumbnailWidget = SNullWidget::NullWidget;
@@ -169,6 +211,8 @@ FReply SNiagaraOverviewStackNode::OnClickedRenderingPreview(const FGeometry& InG
 		SelectedEntries.Add(InEntry);
 		TArray<UNiagaraStackEntry*> DeselectedEntries;
 		OverviewSelectionViewModel->UpdateSelectedEntries(SelectedEntries, DeselectedEntries, true);
+
+		CurrentIssueIndex = -1;
 		return FReply::Handled();
 	}
 	return FReply::Unhandled();
diff --git a/Engine/Plugins/FX/Niagara/Source/NiagaraEditorWidgets/Private/SNiagaraOverviewStackNode.h b/Engine/Plugins/FX/Niagara/Source/NiagaraEditorWidgets/Private/SNiagaraOverviewStackNode.h
--- a/Engine/Plugins/FX/Niagara/Source/NiagaraEditorWidgets/Private/SNiagaraOverviewStackNode.h
+++ b/Engine/Plugins/FX/Niagara/Source/NiagaraEditorWidgets/Private/SNiagaraOverviewStackNode.h
@@ -14,7 +14,6 @@ class FAssetThumbnailPool;
 class FAssetThumbnail;
 struct FRendererPreviewData;
 
-
 class SNiagaraOverviewStackNode : public SGraphNode
 {
 public:
@@ -27,8 +26,10 @@ public:
 
 protected:
 	virtual TSharedRef<SWidget> CreateTitleWidget(TSharedPtr<SNodeTitle> NodeTitle) override;
+	virtual TSharedRef<SWidget> CreateTitleRightWidget() override;
 	virtual TSharedRef<SWidget> CreateNodeContentArea() override;
 	void FillThumbnailBar(UObject* ChangedObject, const bool bIsTriggeredByObjectUpdate);
+
 private:
 	EVisibility GetIssueIconVisibility() const;
 	EVisibility GetEnabledCheckBoxVisibility() const;
@@ -53,6 +54,8 @@ private:
 		return bIsHoveringThumbnail;
 	}
 
+	FReply OnCycleThroughIssues();
+
 private:
 	UNiagaraOverviewNode* OverviewStackNode;
 	UNiagaraStackViewModel* StackViewModel;
@@ -63,4 +66,5 @@ private:
 	TSharedPtr<SHorizontalBox> ThumbnailBar;
 	TArray<FRendererPreviewData*> PreviewData;
 	bool bIsHoveringThumbnail;
+	int32 CurrentIssueIndex;
 };
\ No newline at end of file
