Niagara emitters no longer support Edit commands that make no sense inside an emitter node being edited (ie. Cut, Delete, Rename).

diff --git a/Engine/Plugins/FX/Niagara/Source/NiagaraEditor/Private/NiagaraOverviewNode.cpp b/Engine/Plugins/FX/Niagara/Source/NiagaraEditor/Private/NiagaraOverviewNode.cpp
--- a/Engine/Plugins/FX/Niagara/Source/NiagaraEditor/Private/NiagaraOverviewNode.cpp
+++ b/Engine/Plugins/FX/Niagara/Source/NiagaraEditor/Private/NiagaraOverviewNode.cpp
@@ -110,9 +110,29 @@ FLinearColor UNiagaraOverviewNode::GetNodeTitleColor() const
 	return EmitterHandleGuid.IsValid() ? EmitterColor : SystemColor;
 }
 
+static bool IsSystemAsset(UNiagaraSystem* System)
+{
+	FNiagaraEditorModule& NiagaraEditorModule = FModuleManager::Get().LoadModuleChecked<FNiagaraEditorModule>("NiagaraEditor");
+	TSharedPtr<FNiagaraSystemViewModel> OwningSystemViewModel = NiagaraEditorModule.GetExistingViewModelForSystem(System);
+	if (OwningSystemViewModel.IsValid())
+	{
+		if (OwningSystemViewModel->GetEditMode() == ENiagaraSystemViewModelEditMode::SystemAsset)
+		{
+			return true;
+		}
+	}
+
+	return false;
+}
+
+bool UNiagaraOverviewNode::GetCanRenameNode() const
+{
+	return IsSystemAsset(OwningSystem) && EmitterHandleGuid.IsValid();
+}
+
 bool UNiagaraOverviewNode::CanUserDeleteNode() const
 {
-	return EmitterHandleGuid.IsValid();
+	return IsSystemAsset(OwningSystem) && EmitterHandleGuid.IsValid();
 }
 
 bool UNiagaraOverviewNode::CanDuplicateNode() const
diff --git a/Engine/Plugins/FX/Niagara/Source/NiagaraEditor/Public/NiagaraOverviewNode.h b/Engine/Plugins/FX/Niagara/Source/NiagaraEditor/Public/NiagaraOverviewNode.h
--- a/Engine/Plugins/FX/Niagara/Source/NiagaraEditor/Public/NiagaraOverviewNode.h
+++ b/Engine/Plugins/FX/Niagara/Source/NiagaraEditor/Public/NiagaraOverviewNode.h
@@ -31,6 +31,8 @@ public:
 
 	virtual void OnRenameNode(const FString& NewName) override;
 
+	virtual bool GetCanRenameNode() const override;
+
 	virtual void GetNodeContextMenuActions(class UToolMenu* Menu, class UGraphNodeContextMenuContext* Context) const override;
 	//~ End UEdGraphNode Interface
 
