Add support for PrimitiveComponent's BoundsScale to NiagaraComponent

diff --git a/Engine/Plugins/FX/Niagara/Source/Niagara/Private/NiagaraComponent.cpp b/Engine/Plugins/FX/Niagara/Source/Niagara/Private/NiagaraComponent.cpp
--- a/Engine/Plugins/FX/Niagara/Source/Niagara/Private/NiagaraComponent.cpp
+++ b/Engine/Plugins/FX/Niagara/Source/Niagara/Private/NiagaraComponent.cpp
@@ -1263,6 +1263,8 @@ FBoxSphereBounds UNiagaraComponent::CalcBounds(const FTransform& LocalToWorld) c
 	if (SystemInstance.IsValid())
 	{
 		SystemBounds = SystemInstance->GetLocalBounds();
+		SystemBounds.BoxExtent *= BoundsScale;
+		SystemBounds.SphereRadius *= BoundsScale;
 	}
 	else
 	{
