- Fix for Volume Materials not working on Niagara Sprite Renderers

#rnx

diff --git a/Engine/Plugins/FX/Niagara/Shaders/Private/NiagaraSpriteVertexFactory.ush b/Engine/Plugins/FX/Niagara/Shaders/Private/NiagaraSpriteVertexFactory.ush
--- a/Engine/Plugins/FX/Niagara/Shaders/Private/NiagaraSpriteVertexFactory.ush
+++ b/Engine/Plugins/FX/Niagara/Shaders/Private/NiagaraSpriteVertexFactory.ush
@@ -10,7 +10,7 @@
 #include "NiagaraVFParticleAccess.usf"
 
 #define USE_PARTICLE_LIGHTING_OFFSET ((FEATURE_LEVEL >= FEATURE_LEVEL_SM5 || FEATURE_LEVEL == FEATURE_LEVEL_ES3_1) && !MATERIAL_SHADINGMODEL_UNLIT)
-#define USE_PARTICLE_POSITION ((FEATURE_LEVEL >= FEATURE_LEVEL_SM5 || FEATURE_LEVEL == FEATURE_LEVEL_ES3_1) && NEEDS_PARTICLE_POSITION)
+#define USE_PARTICLE_POSITION ((FEATURE_LEVEL >= FEATURE_LEVEL_SM5 || FEATURE_LEVEL == FEATURE_LEVEL_ES3_1) && (NEEDS_PARTICLE_POSITION || PASS_NEEDS_PRIMITIVE_VOLUME_BOUNDS))
 #define USE_PARTICLE_VELOCITY ((FEATURE_LEVEL >= FEATURE_LEVEL_SM5 || FEATURE_LEVEL == FEATURE_LEVEL_ES3_1) && NEEDS_PARTICLE_VELOCITY)
 #define USE_PARTICLE_TIME ((FEATURE_LEVEL >= FEATURE_LEVEL_SM5 || FEATURE_LEVEL == FEATURE_LEVEL_ES3_1) && NEEDS_PARTICLE_TIME)
 #define USE_PARTICLE_SIZE ((FEATURE_LEVEL >= FEATURE_LEVEL_SM5 || FEATURE_LEVEL == FEATURE_LEVEL_ES3_1) && NEEDS_PARTICLE_SIZE)
