Fix for crash caused by potentially bad state in scalability manager.

diff --git a/Engine/Plugins/FX/Niagara/Source/Niagara/Private/NiagaraScalabilityManager.cpp b/Engine/Plugins/FX/Niagara/Source/Niagara/Private/NiagaraScalabilityManager.cpp
--- a/Engine/Plugins/FX/Niagara/Source/Niagara/Private/NiagaraScalabilityManager.cpp
+++ b/Engine/Plugins/FX/Niagara/Source/Niagara/Private/NiagaraScalabilityManager.cpp
@@ -193,10 +193,11 @@ void FNiagaraScalabilityManager::Update(FNiagaraWorldManager* WorldMan)
 		while (CompIdx < ManagedComponents.Num())
 		{
 			FNiagaraScalabilityState& CompState = State[CompIdx];
+			UNiagaraComponent* Component = ManagedComponents[CompIdx];
 			bool bRepeatIndex = false;
-			if (CompState.bDirty)
+			if (Component && CompState.bDirty)
 			{
-				UNiagaraComponent* Component = ManagedComponents[CompIdx];
+				CompState.bDirty = false;
 				if (CompState.bCulled)
 				{
 					switch (EffectType->CullReaction)
