Niagara - Copy/paste improvements in the system/emitter editor.
+ Can now cut, copy, and paste modules and renderers.
+ Can now select and copy sets of inputs and paste them onto modules.
+ Rows in the full stack view can now be selected so you can use the engine keyboard shortcuts for cut, copy, paste, and delete in both the overview and selection stacks.
+ Fixed a crash bug when copy and pasting inputs for emitter modules.

diff --git a/Engine/Plugins/FX/Niagara/Source/NiagaraEditor/Private/NiagaraClipboard.cpp b/Engine/Plugins/FX/Niagara/Source/NiagaraEditor/Private/NiagaraClipboard.cpp
--- a/Engine/Plugins/FX/Niagara/Source/NiagaraEditor/Private/NiagaraClipboard.cpp
+++ b/Engine/Plugins/FX/Niagara/Source/NiagaraEditor/Private/NiagaraClipboard.cpp
@@ -38,82 +38,83 @@ protected:
 	}
 };
 
-UNiagaraClipboardFunctionInput* MakeNewInput(UObject* InOuter, FName InInputName, FNiagaraTypeDefinition InInputType, ENiagaraClipboardFunctionInputValueMode InValueMode)
+UNiagaraClipboardFunctionInput* MakeNewInput(UObject* InOuter, FName InInputName, FNiagaraTypeDefinition InInputType, TOptional<bool> bInEditConditionValue, ENiagaraClipboardFunctionInputValueMode InValueMode)
 {
 	UNiagaraClipboardFunctionInput* NewInput = Cast<UNiagaraClipboardFunctionInput>(NewObject<UNiagaraClipboardFunctionInput>(InOuter));
 	NewInput->InputName = InInputName;
 	NewInput->InputType = InInputType;
+	NewInput->bHasEditCondition = bInEditConditionValue.IsSet();
+	NewInput->bEditConditionValue = bInEditConditionValue.Get(false);
 	NewInput->ValueMode = InValueMode;
 	return NewInput;
 }
 
-const UNiagaraClipboardFunctionInput* UNiagaraClipboardFunctionInput::CreateLocalValue(UObject* InOuter, FName InInputName, FNiagaraTypeDefinition InInputType, TArray<uint8>& InLocalValueData)
+const UNiagaraClipboardFunctionInput* UNiagaraClipboardFunctionInput::CreateLocalValue(UObject* InOuter, FName InInputName, FNiagaraTypeDefinition InInputType, TOptional<bool> bInEditConditionValue, TArray<uint8>& InLocalValueData)
 {
 	checkf(InLocalValueData.Num() == InInputType.GetSize(), TEXT("Input data size didn't match type size."))
-	UNiagaraClipboardFunctionInput* NewInput = MakeNewInput(InOuter, InInputName, InInputType, ENiagaraClipboardFunctionInputValueMode::Local);
+	UNiagaraClipboardFunctionInput* NewInput = MakeNewInput(InOuter, InInputName, InInputType, bInEditConditionValue, ENiagaraClipboardFunctionInputValueMode::Local);
 	NewInput->Local = InLocalValueData;
 	return NewInput;
 }
 
-const UNiagaraClipboardFunctionInput* UNiagaraClipboardFunctionInput::CreateLinkedValue(UObject* InOuter, FName InInputName, FNiagaraTypeDefinition InInputType, FName InLinkedValue)
+const UNiagaraClipboardFunctionInput* UNiagaraClipboardFunctionInput::CreateLinkedValue(UObject* InOuter, FName InInputName, FNiagaraTypeDefinition InInputType, TOptional<bool> bInEditConditionValue, FName InLinkedValue)
 {
-	UNiagaraClipboardFunctionInput* NewInput = MakeNewInput(InOuter, InInputName, InInputType, ENiagaraClipboardFunctionInputValueMode::Linked);
+	UNiagaraClipboardFunctionInput* NewInput = MakeNewInput(InOuter, InInputName, InInputType, bInEditConditionValue, ENiagaraClipboardFunctionInputValueMode::Linked);
 	NewInput->Linked = InLinkedValue;
 	return NewInput;
 }
 
-const UNiagaraClipboardFunctionInput* UNiagaraClipboardFunctionInput::CreateDataValue(UObject* InOuter, FName InInputName, FNiagaraTypeDefinition InInputType, UNiagaraDataInterface* InDataValue)
+const UNiagaraClipboardFunctionInput* UNiagaraClipboardFunctionInput::CreateDataValue(UObject* InOuter, FName InInputName, FNiagaraTypeDefinition InInputType, TOptional<bool> bInEditConditionValue, UNiagaraDataInterface* InDataValue)
 {
-	UNiagaraClipboardFunctionInput* NewInput = MakeNewInput(InOuter, InInputName, InInputType, ENiagaraClipboardFunctionInputValueMode::Data);
+	UNiagaraClipboardFunctionInput* NewInput = MakeNewInput(InOuter, InInputName, InInputType, bInEditConditionValue, ENiagaraClipboardFunctionInputValueMode::Data);
 	NewInput->Data = NewObject<UNiagaraDataInterface>(NewInput, InDataValue->GetClass());
 	InDataValue->CopyTo(NewInput->Data);
 	return NewInput;
 }
 
-const UNiagaraClipboardFunctionInput* UNiagaraClipboardFunctionInput::CreateExpressionValue(UObject* InOuter, FName InInputName, FNiagaraTypeDefinition InInputType, const FString& InExpressionValue)
+const UNiagaraClipboardFunctionInput* UNiagaraClipboardFunctionInput::CreateExpressionValue(UObject* InOuter, FName InInputName, FNiagaraTypeDefinition InInputType, TOptional<bool> bInEditConditionValue, const FString& InExpressionValue)
 {
-	UNiagaraClipboardFunctionInput* NewInput = MakeNewInput(InOuter, InInputName, InInputType, ENiagaraClipboardFunctionInputValueMode::Expression);
+	UNiagaraClipboardFunctionInput* NewInput = MakeNewInput(InOuter, InInputName, InInputType, bInEditConditionValue, ENiagaraClipboardFunctionInputValueMode::Expression);
 	NewInput->Expression = InExpressionValue;
 	return NewInput;
 }
 
-const UNiagaraClipboardFunctionInput* UNiagaraClipboardFunctionInput::CreateDynamicValue(UObject* InOuter, FName InInputName, FNiagaraTypeDefinition InInputType, UNiagaraScript* InDynamicValue)
+const UNiagaraClipboardFunctionInput* UNiagaraClipboardFunctionInput::CreateDynamicValue(UObject* InOuter, FName InInputName, FNiagaraTypeDefinition InInputType, TOptional<bool> bInEditConditionValue, FString InDynamicValueName, UNiagaraScript* InDynamicValue)
 {
-	UNiagaraClipboardFunctionInput* NewInput = MakeNewInput(InOuter, InInputName, InInputType, ENiagaraClipboardFunctionInputValueMode::Dynamic);
-	NewInput->Dynamic = UNiagaraClipboardFunction::CreateScriptFunction(NewInput, InDynamicValue);
+	UNiagaraClipboardFunctionInput* NewInput = MakeNewInput(InOuter, InInputName, InInputType, bInEditConditionValue, ENiagaraClipboardFunctionInputValueMode::Dynamic);
+	NewInput->Dynamic = UNiagaraClipboardFunction::CreateScriptFunction(NewInput, InDynamicValueName, InDynamicValue);
 	return NewInput;
 }
 
-UNiagaraClipboardFunction* UNiagaraClipboardFunction::CreateScriptFunction(UObject* InOuter, UNiagaraScript* InScript)
+UNiagaraClipboardFunction* UNiagaraClipboardFunction::CreateScriptFunction(UObject* InOuter, FString InFunctionName, UNiagaraScript* InScript)
 {
 	UNiagaraClipboardFunction* NewFunction = Cast<UNiagaraClipboardFunction>(NewObject<UNiagaraClipboardFunction>(InOuter));
 	NewFunction->ScriptMode = ENiagaraClipboardFunctionScriptMode::ScriptAsset;
+	NewFunction->FunctionName = InFunctionName;
 	NewFunction->Script = InScript;
 	return NewFunction;
 }
 
-FNiagaraClipboard::FNiagaraClipboard()
+UNiagaraClipboardFunction* UNiagaraClipboardFunction::CreateAssignmentFunction(UObject* InOuter, FString InFunctionName, const TArray<FNiagaraVariable>& InAssignmentTargets, const TArray<FString>& InAssignmentDefaults)
 {
+	UNiagaraClipboardFunction* NewFunction = Cast<UNiagaraClipboardFunction>(NewObject<UNiagaraClipboardFunction>(InOuter));
+	NewFunction->ScriptMode = ENiagaraClipboardFunctionScriptMode::Assignment;
+	NewFunction->FunctionName = InFunctionName;
+	NewFunction->AssignmentTargets = InAssignmentTargets;
+	NewFunction->AssignmentDefaults = InAssignmentDefaults;
+	return NewFunction;
 }
 
-void FNiagaraClipboard::Copy(const UNiagaraClipboardFunctionInput* Input)
+UNiagaraClipboardContent* UNiagaraClipboardContent::Create()
 {
-	UNiagaraClipboardContent* ClipboardContent = NewObject<UNiagaraClipboardContent>(GetTransientPackage());
-	ClipboardContent->FunctionInputs.Add(CastChecked<UNiagaraClipboardFunctionInput>(StaticDuplicateObject(Input, ClipboardContent)));
-	SetClipboardContentInternal(ClipboardContent);
+	return NewObject<UNiagaraClipboardContent>(GetTransientPackage());
 }
 
-const UNiagaraClipboardFunctionInput* FNiagaraClipboard::GetCopiedInput() const
+FNiagaraClipboard::FNiagaraClipboard()
 {
-	const UNiagaraClipboardContent* ClipboardContent = GetClipboardContentInternal();
-	if (ClipboardContent != nullptr &&ClipboardContent->FunctionInputs.Num() == 1)
-	{
-		return ClipboardContent->FunctionInputs[0];
-	}
-	return nullptr;
 }
 
-void FNiagaraClipboard::SetClipboardContentInternal(UNiagaraClipboardContent* ClipboardContent)
+void FNiagaraClipboard::SetClipboardContent(UNiagaraClipboardContent* ClipboardContent)
 {
 	// Clear the mark state for saving.
 	UnMarkAllObjects(EObjectMark(OBJECTMARK_TagExp | OBJECTMARK_TagImp));
@@ -125,7 +126,7 @@ void FNiagaraClipboard::SetClipboardContentInternal(UNiagaraClipboardContent* Cl
 	FPlatformApplicationMisc::ClipboardCopy(*Archive);
 }
 
-const UNiagaraClipboardContent* FNiagaraClipboard::GetClipboardContentInternal() const
+const UNiagaraClipboardContent* FNiagaraClipboard::GetClipboardContent() const
 {
 	// Get the text from the clipboard.
 	FString ClipboardText;
diff --git a/Engine/Plugins/FX/Niagara/Source/NiagaraEditor/Private/ViewModels/NiagaraSystemSelectionViewModel.cpp b/Engine/Plugins/FX/Niagara/Source/NiagaraEditor/Private/ViewModels/NiagaraSystemSelectionViewModel.cpp
--- a/Engine/Plugins/FX/Niagara/Source/NiagaraEditor/Private/ViewModels/NiagaraSystemSelectionViewModel.cpp
+++ b/Engine/Plugins/FX/Niagara/Source/NiagaraEditor/Private/ViewModels/NiagaraSystemSelectionViewModel.cpp
@@ -148,10 +148,31 @@ void  UNiagaraSystemSelectionViewModel::EmptySelection()
 	}
 }
 
-void UNiagaraSystemSelectionViewModel::RemoveEntryFromSelectionByDisplayedObject(const UObject* InObject)
+void UNiagaraSystemSelectionViewModel::RemoveEntriesFromSelection(const TArray<UNiagaraStackEntry*>& InEntriesToRemove)
 {
-	bool bSelectionChanged = false;
+	TArray<UNiagaraStackEntry*> EntriesToDeselect;
+	for (FSelectionEntry& SelectionEntry : SelectionEntries)
+	{
+		if (InEntriesToRemove.Contains(SelectionEntry.Entry.Get()))
+		{
+			EntriesToDeselect.Add(SelectionEntry.Entry.Get());
+		}
+	}
 
+	bool bSelectionChanged = EntriesToDeselect.Num() > 0;
+	for (UNiagaraStackEntry* EntryToDeselect : EntriesToDeselect)
+	{
+		RemoveEntryFromSelectionInternal(EntryToDeselect);
+	}
+
+	if (bSelectionChanged)
+	{
+		UpdateExternalSelectionState();
+	}
+}
+
+void UNiagaraSystemSelectionViewModel::RemoveEntryFromSelectionByDisplayedObject(const UObject* InObject)
+{
 	TArray<UNiagaraStackEntry*> EntriesToDeselect;
 	for (FSelectionEntry& SelectionEntry : SelectionEntries)
 	{
@@ -161,10 +182,10 @@ void UNiagaraSystemSelectionViewModel::RemoveEntryFromSelectionByDisplayedObject
 		}
 	}
 
+	bool bSelectionChanged = EntriesToDeselect.Num() > 0;
 	for (UNiagaraStackEntry* EntryToDeselect : EntriesToDeselect)
 	{
 		RemoveEntryFromSelectionInternal(EntryToDeselect);
-		bSelectionChanged = true;
 	}
 
 	if (bSelectionChanged)
diff --git a/Engine/Plugins/FX/Niagara/Source/NiagaraEditor/Private/ViewModels/NiagaraSystemViewModel.cpp b/Engine/Plugins/FX/Niagara/Source/NiagaraEditor/Private/ViewModels/NiagaraSystemViewModel.cpp
--- a/Engine/Plugins/FX/Niagara/Source/NiagaraEditor/Private/ViewModels/NiagaraSystemViewModel.cpp
+++ b/Engine/Plugins/FX/Niagara/Source/NiagaraEditor/Private/ViewModels/NiagaraSystemViewModel.cpp
@@ -1336,13 +1336,20 @@ void FNiagaraSystemViewModel::EmitterParameterStoreChanged(const FNiagaraParamet
 {
 	if (bUpdatingEmittersFromSequencerDataChange == false)
 	{
-		TGuardValue<bool> UpdateGuard(bUpdatingSequencerFromEmitterDataChange, true);
-		for (UMovieSceneTrack* Track : NiagaraSequence->GetMovieScene()->GetMasterTracks())
+		if (EmitterIdsRequiringSequencerTrackUpdate.Num() > 0)
 		{
-			UMovieSceneNiagaraEmitterTrack* EmitterTrack = CastChecked<UMovieSceneNiagaraEmitterTrack>(Track);
-			EmitterTrack->UpdateTrackFromEmitterParameterChange(NiagaraSequence->GetMovieScene()->GetTickResolution());
+			UpdateSequencerTracksForEmitters(EmitterIdsRequiringSequencerTrackUpdate);
+		}
+		else
+		{
+			TGuardValue<bool> UpdateGuard(bUpdatingSequencerFromEmitterDataChange, true);
+			for (UMovieSceneTrack* Track : NiagaraSequence->GetMovieScene()->GetMasterTracks())
+			{
+				UMovieSceneNiagaraEmitterTrack* EmitterTrack = CastChecked<UMovieSceneNiagaraEmitterTrack>(Track);
+				EmitterTrack->UpdateTrackFromEmitterParameterChange(NiagaraSequence->GetMovieScene()->GetTickResolution());
+			}
+			Sequencer->NotifyMovieSceneDataChanged(EMovieSceneDataChangeType::TrackValueChanged);
 		}
-		Sequencer->NotifyMovieSceneDataChanged(EMovieSceneDataChangeType::TrackValueChanged);
 	}
 	UpdateSimulationFromParameterChange();
 }
diff --git a/Engine/Plugins/FX/Niagara/Source/NiagaraEditor/Private/ViewModels/Stack/NiagaraStackClipboardUtilities.cpp b/Engine/Plugins/FX/Niagara/Source/NiagaraEditor/Private/ViewModels/Stack/NiagaraStackClipboardUtilities.cpp
new file mode 100644
--- /dev/null
+++ b/Engine/Plugins/FX/Niagara/Source/NiagaraEditor/Private/ViewModels/Stack/NiagaraStackClipboardUtilities.cpp
@@ -0,0 +1,284 @@
+// Copyright 1998-2019 Epic Games, Inc. All Rights Reserved.
+
+#include "ViewModels/Stack/NiagaraStackClipboardUtilities.h"
+#include "ViewModels/NiagaraSystemViewModel.h"
+#include "ViewModels/NiagaraSystemSelectionViewModel.h"
+#include "ViewModels/Stack/NiagaraStackEntry.h"
+#include "NiagaraEditorModule.h"
+#include "NiagaraClipboard.h"
+
+#include "ScopedTransaction.h"
+
+#define LOCTEXT_NAMESPACE "NiagaraStackClipboardUtilities"
+
+bool FNiagaraStackClipboardUtilities::TestCanCutSelectionWithMessage(const TArray<UNiagaraStackEntry*>& SelectedEntries, FText& OutCanCutMessage)
+{
+	if (SelectedEntries.Num() == 1)
+	{
+		return SelectedEntries[0]->TestCanCutWithMessage(OutCanCutMessage);
+	}
+
+	UClass* EntryClass = nullptr;
+	bool bAnySupportCut = false;
+	bool bAnyCanCut = false;
+	bool bClassesMatch = true;
+	for (UNiagaraStackEntry* SelectedEntry : SelectedEntries)
+	{
+		if (EntryClass == nullptr)
+		{
+			EntryClass = SelectedEntry->GetClass();
+		}
+		else
+		{
+			if (SelectedEntry->GetClass() != EntryClass)
+			{
+				bClassesMatch = false;
+			}
+		}
+		if (SelectedEntry->SupportsCut())
+		{
+			bAnySupportCut = true;
+			FText UnusedMessage;
+			bAnyCanCut |= SelectedEntry->TestCanCutWithMessage(UnusedMessage);
+		}
+	}
+
+	if (bClassesMatch == false)
+	{
+		OutCanCutMessage = LOCTEXT("CantCutMixedSelection", "Can only Cut selections of similar items, e.g. modules, inputs, renderers, etc.");
+		return false;
+	}
+
+	if (bAnySupportCut)
+	{
+		if (bAnyCanCut)
+		{
+			OutCanCutMessage = LOCTEXT("CutSelection", "Cut the selected items.");
+			return true;
+		}
+		else
+		{
+			OutCanCutMessage = LOCTEXT("CantCutSelection", "None of the items in the selection can be cut.");
+			return false;
+		}
+	}
+	else
+	{
+		OutCanCutMessage = FText();
+		return false;
+	}
+}
+
+void FNiagaraStackClipboardUtilities::CutSelection(const TArray<UNiagaraStackEntry*>& SelectedEntries)
+{
+	FText TransactionMessage;
+	if (SelectedEntries.Num() == 1)
+	{
+		TransactionMessage = SelectedEntries[0]->GetCutTransactionText();
+	}
+	if (TransactionMessage.IsEmptyOrWhitespace())
+	{
+		TransactionMessage = LOCTEXT("CutSelectionTransaction", "Cut the niagara selection");
+	}
+
+	FScopedTransaction CutTransaction(TransactionMessage);
+	UNiagaraClipboardContent* ClipboardContent = UNiagaraClipboardContent::Create();
+	for (UNiagaraStackEntry* SelectedEntry : SelectedEntries)
+	{
+		SelectedEntry->CopyForCut(ClipboardContent);
+	}
+
+	TArray<UNiagaraStackEntry*> EntriesToRemove = SelectedEntries;
+	SelectedEntries[0]->GetSystemViewModel()->GetSelectionViewModel()->RemoveEntriesFromSelection(SelectedEntries);
+	for (UNiagaraStackEntry* EntryToRemove : EntriesToRemove)
+	{
+		EntryToRemove->RemoveForCut();
+	}
+	FNiagaraEditorModule::Get().GetClipboard().SetClipboardContent(ClipboardContent);
+}
+
+bool FNiagaraStackClipboardUtilities::TestCanCopySelectionWithMessage(const TArray<UNiagaraStackEntry*>& SelectedEntries, FText& OutCanCopyMessage)
+{
+	if (SelectedEntries.Num() == 1)
+	{
+		return SelectedEntries[0]->TestCanCopyWithMessage(OutCanCopyMessage);
+	}
+
+	UClass* EntryClass = nullptr;
+	bool bAnySupportCopy = false;
+	bool bAnyCanCopy = false;
+	bool bClassesMatch = true;
+	for (UNiagaraStackEntry* SelectedEntry : SelectedEntries)
+	{
+		if (EntryClass == nullptr)
+		{
+			EntryClass = SelectedEntry->GetClass();
+		}
+		else
+		{
+			if (SelectedEntry->GetClass() != EntryClass)
+			{
+				bClassesMatch = false;
+			}
+		}
+		if (SelectedEntry->SupportsCopy())
+		{
+			bAnySupportCopy = true;
+			FText UnusedMessage;
+			bAnyCanCopy |= SelectedEntry->TestCanCopyWithMessage(UnusedMessage);
+		}
+	}
+
+	if (bClassesMatch == false)
+	{
+		OutCanCopyMessage = LOCTEXT("CantCopyMixedSelection", "Can only copy selections of similar items, e.g. modules, inputs, renderers, etc.");
+		return false;
+	}
+
+	if (bAnySupportCopy)
+	{
+		if (bAnyCanCopy)
+		{
+			OutCanCopyMessage = LOCTEXT("CopySelection", "Copy the selected items.");
+			return true;
+		}
+		else
+		{
+			OutCanCopyMessage = LOCTEXT("CantCopySelection", "None of the items in the selection can be copied.");
+			return false;
+		}
+	}
+	else
+	{
+		OutCanCopyMessage = FText();
+		return false;
+	}
+}
+
+void FNiagaraStackClipboardUtilities::CopySelection(const TArray<UNiagaraStackEntry*>& SelectedEntries)
+{
+	UNiagaraClipboardContent* ClipboardContent = UNiagaraClipboardContent::Create();
+	for (UNiagaraStackEntry* SelectedEntry : SelectedEntries)
+	{
+		SelectedEntry->Copy(ClipboardContent);
+	}
+	FNiagaraEditorModule::Get().GetClipboard().SetClipboardContent(ClipboardContent);
+}
+
+bool FNiagaraStackClipboardUtilities::TestCanPasteSelectionWithMessage(const TArray<UNiagaraStackEntry*>& SelectedEntries, FText& OutCanPasteMessage)
+{
+	if (SelectedEntries.Num() > 1)
+	{
+		OutCanPasteMessage = LOCTEXT("CantPasteToMultipleTargets", "Can not paste with multiple rows selected.");
+		return false;
+	}
+	if (SelectedEntries.Num() == 1 && SelectedEntries[0]->SupportsPaste())
+	{
+		const UNiagaraClipboardContent* ClipboardContent = FNiagaraEditorModule::Get().GetClipboard().GetClipboardContent();
+		if (ClipboardContent != nullptr)
+		{
+			return SelectedEntries[0]->TestCanPasteWithMessage(ClipboardContent, OutCanPasteMessage);
+		}
+	}
+	return false;
+}
+
+void FNiagaraStackClipboardUtilities::PasteSelection(const TArray<UNiagaraStackEntry*>& SelectedEntries)
+{
+	if (SelectedEntries.Num() == 1)
+	{
+		const UNiagaraClipboardContent* ClipboardContent = FNiagaraEditorModule::Get().GetClipboard().GetClipboardContent();
+
+		FText TransactionMessage = SelectedEntries[0]->GetPasteTransactionText(ClipboardContent);
+		if (TransactionMessage.IsEmptyOrWhitespace())
+		{
+			TransactionMessage = LOCTEXT("PasteSelectionTransaction", "Paste to niagara stack.");
+		}
+
+		FScopedTransaction PasteTransaction(TransactionMessage);
+		if (ClipboardContent != nullptr)
+		{
+			SelectedEntries[0]->Paste(ClipboardContent);
+		}
+	}
+}
+
+bool FNiagaraStackClipboardUtilities::TestCanDeleteSelectionWithMessage(const TArray<UNiagaraStackEntry*>& SelectedEntries, FText& OutCanDeleteMessage)
+{
+	if (SelectedEntries.Num() == 1)
+	{
+		if (SelectedEntries[0]->SupportsDelete())
+		{
+			return SelectedEntries[0]->TestCanDeleteWithMessage(OutCanDeleteMessage);
+		}
+		return false;
+	}
+	else
+	{
+		bool bAnySupportDelete = false;
+		bool bAnyCanDelete = false;
+		for (UNiagaraStackEntry* SelectedEntry : SelectedEntries)
+		{
+			if (SelectedEntry->SupportsDelete())
+			{
+				bAnySupportDelete = true;
+				FText UnusedMessage;
+				if (SelectedEntry->TestCanDeleteWithMessage(UnusedMessage))
+				{
+					bAnyCanDelete = true;
+					break;
+				}
+			}
+		}
+
+		if (bAnyCanDelete)
+		{
+			OutCanDeleteMessage = LOCTEXT("DeleteMulti", "Delete items in the selection which can be deleted.");
+			return true;
+		}
+		else
+		{
+			if (bAnySupportDelete)
+			{
+				// Only provide a message if an item in the selection actually supports being deleted.
+				OutCanDeleteMessage = LOCTEXT("CantDeleteMulti", "None of the selected items can be deleted.");
+			}
+			return false;
+		}
+	}
+}
+
+void FNiagaraStackClipboardUtilities::DeleteSelection(const TArray<UNiagaraStackEntry*>& SelectedEntries)
+{
+	FText TransactionMessage;
+	if (SelectedEntries.Num() == 1)
+	{
+		TransactionMessage = SelectedEntries[0]->GetDeleteTransactionText();
+	}
+	if (TransactionMessage.IsEmptyOrWhitespace())
+	{
+		TransactionMessage = LOCTEXT("DeleteSelectionTransaction", "Delete the niagara selection");
+	}
+
+	TArray<UNiagaraStackEntry*> EntriesToDelete;
+	for (UNiagaraStackEntry* SelectedEntry : SelectedEntries)
+	{
+		FText DeleteMessage;
+		if (SelectedEntry->SupportsDelete() && SelectedEntry->TestCanDeleteWithMessage(DeleteMessage))
+		{
+			EntriesToDelete.Add(SelectedEntry);
+		}
+	}
+
+	if (EntriesToDelete.Num() > 0)
+	{
+		EntriesToDelete[0]->GetSystemViewModel()->GetSelectionViewModel()->RemoveEntriesFromSelection(EntriesToDelete);
+		const FScopedTransaction DeleteTransaction(TransactionMessage);
+		for (UNiagaraStackEntry* EntryToDelete : EntriesToDelete)
+		{
+			EntryToDelete->Delete();
+		}
+	}
+}
+
+#undef LOCTEXT_NAMESPACE
\ No newline at end of file
diff --git a/Engine/Plugins/FX/Niagara/Source/NiagaraEditor/Private/ViewModels/Stack/NiagaraStackEventScriptItemGroup.cpp b/Engine/Plugins/FX/Niagara/Source/NiagaraEditor/Private/ViewModels/Stack/NiagaraStackEventScriptItemGroup.cpp
--- a/Engine/Plugins/FX/Niagara/Source/NiagaraEditor/Private/ViewModels/Stack/NiagaraStackEventScriptItemGroup.cpp
+++ b/Engine/Plugins/FX/Niagara/Source/NiagaraEditor/Private/ViewModels/Stack/NiagaraStackEventScriptItemGroup.cpp
@@ -225,12 +225,21 @@ void UNiagaraStackEventScriptItemGroup::RefreshChildrenInternal(const TArray<UNi
 	Super::RefreshChildrenInternal(CurrentChildren, NewChildren, NewIssues);
 }
 
-bool UNiagaraStackEventScriptItemGroup::CanDelete() const
+bool UNiagaraStackEventScriptItemGroup::TestCanDeleteWithMessage(FText& OutCanDeleteMessage) const
 {
-	return HasBaseEventHandler() == false;
+	if (HasBaseEventHandler())
+	{
+		OutCanDeleteMessage = LOCTEXT("CantDeleteInherited", "Can not delete this event handler because it's inherited.");
+		return false;
+	}
+	else
+	{
+		OutCanDeleteMessage = LOCTEXT("CanDelete", "Delete this event handler.");
+		return true;
+	}
 }
 
-bool UNiagaraStackEventScriptItemGroup::Delete()
+void UNiagaraStackEventScriptItemGroup::Delete()
 {
 	TSharedPtr<FNiagaraScriptViewModel> ScriptViewModelPinned = ScriptViewModel.Pin();
 	checkf(ScriptViewModelPinned.IsValid(), TEXT("Can not delete when the script view model has been deleted."));
@@ -240,7 +249,7 @@ bool UNiagaraStackEventScriptItemGroup::Delete()
 
 	if (!Source || !Source->NodeGraph)
 	{
-		return false;
+		return;
 	}
 
 	FScopedTransaction Transaction(FText::Format(LOCTEXT("DeleteEventHandler", "Deleted {0}"), GetDisplayName()));
@@ -267,8 +276,6 @@ bool UNiagaraStackEventScriptItemGroup::Delete()
 	ScriptViewModelPinned->SetScripts(Emitter);
 	
 	OnModifiedEventHandlersDelegate.ExecuteIfBound();
-
-	return true;
 }
 
 bool UNiagaraStackEventScriptItemGroup::HasBaseEventHandler() const
diff --git a/Engine/Plugins/FX/Niagara/Source/NiagaraEditor/Private/ViewModels/Stack/NiagaraStackFunctionInput.cpp b/Engine/Plugins/FX/Niagara/Source/NiagaraEditor/Private/ViewModels/Stack/NiagaraStackFunctionInput.cpp
--- a/Engine/Plugins/FX/Niagara/Source/NiagaraEditor/Private/ViewModels/Stack/NiagaraStackFunctionInput.cpp
+++ b/Engine/Plugins/FX/Niagara/Source/NiagaraEditor/Private/ViewModels/Stack/NiagaraStackFunctionInput.cpp
@@ -267,14 +267,27 @@ bool UNiagaraStackFunctionInput::TestCanCutWithMessage(FText& OutMessage) const
 		OutMessage = LOCTEXT("CantCutInvalidMessage", "The current input state doesn't support cutting.");
 		return false;
 	}
+	if (GetIsEnabledAndOwnerIsEnabled() == false)
+	{
+		OutMessage = LOCTEXT("CantCutDisabled", "Can not cut and input when it's owner is disabled.");
+		return false;
+	}
 	OutMessage = LOCTEXT("CanCutMessage", "Cut will copy the value of this input including\nany data objects and dynamic inputs, and will reset it to default.");
 	return true;
 }
 
-void UNiagaraStackFunctionInput::Cut()
+FText UNiagaraStackFunctionInput::GetCutTransactionText() const
+{
+	return LOCTEXT("CutInputTransaction", "Cut niagara input");
+}
+
+void UNiagaraStackFunctionInput::CopyForCut(UNiagaraClipboardContent* ClipboardContent) const
+{
+	Copy(ClipboardContent);
+}
+
+void UNiagaraStackFunctionInput::RemoveForCut()
 {
-	FScopedTransaction ScopedTransaction(LOCTEXT("CutInput", "Cut niagara input"));
-	Copy();
 	Reset();
 }
 
@@ -289,32 +302,58 @@ bool UNiagaraStackFunctionInput::TestCanCopyWithMessage(FText& OutMessage) const
 	return true;
 }
 
-void UNiagaraStackFunctionInput::Copy() const
+void UNiagaraStackFunctionInput::Copy(UNiagaraClipboardContent* ClipboardContent) const
 {
 	const UNiagaraClipboardFunctionInput* ClipboardInput = ToClipboardFunctionInput(GetTransientPackage());
 	if (ClipboardInput != nullptr)
 	{
-		FNiagaraEditorModule::Get().GetClipboard().Copy(ClipboardInput);
+		ClipboardContent->FunctionInputs.Add(ClipboardInput);
 	}
 }
 
-bool UNiagaraStackFunctionInput::TestCanPasteWithMessage(FText& OutMessage) const
+bool UNiagaraStackFunctionInput::TestCanPasteWithMessage(const UNiagaraClipboardContent* ClipboardContent, FText& OutMessage) const
 {
-	const UNiagaraClipboardFunctionInput* ClipboardInput = FNiagaraEditorModule::Get().GetClipboard().GetCopiedInput();
-	if (GetIsEnabledAndOwnerIsEnabled() &&  ClipboardInput != nullptr && ClipboardInput->InputType == InputType)
+	if (ClipboardContent->FunctionInputs.Num() == 0 || GetIsEnabledAndOwnerIsEnabled() == false)
 	{
-		OutMessage = LOCTEXT("PastMessage", "Paste the input from the clipboard here.");
-		return true;
+		// Empty clipboard, or disabled don't allow paste, but be silent.
+		return false;
+	}
+	else if (ClipboardContent->FunctionInputs.Num() == 1)
+	{
+		if (ClipboardContent->FunctionInputs[0] != nullptr)
+		{
+			if (ClipboardContent->FunctionInputs[0]->InputType == InputType)
+			{
+				OutMessage = LOCTEXT("PastMessage", "Paste the input from the clipboard here.");
+				return true;
+			}
+			else
+			{
+				OutMessage = LOCTEXT("CantPasteIncorrectType", "Can not paste inputs with mismatched types.");
+				return false;
+			}
+		}
+		return false;
+	}
+	else
+	{
+		OutMessage = LOCTEXT("CantPasteMultipleInputs", "Can't paste multiple inputs onto a single input.");
+		return false;
 	}
-	return false;
 }
 
-void UNiagaraStackFunctionInput::Paste()
+FText UNiagaraStackFunctionInput::GetPasteTransactionText(const UNiagaraClipboardContent* ClipboardContent) const
 {
-	const UNiagaraClipboardFunctionInput* ClipboardInput = FNiagaraEditorModule::Get().GetClipboard().GetCopiedInput();
+	return LOCTEXT("PasteInputTransactionText", "Paste niagara inputs");
+}
+
+void UNiagaraStackFunctionInput::Paste(const UNiagaraClipboardContent* ClipboardContent)
+{
+	checkf(ClipboardContent != nullptr && ClipboardContent->FunctionInputs.Num() == 1, TEXT("Clipboard must not be null, and must contain a single input.  Call TestCanPasteWithMessage to validate"));
+
+	const UNiagaraClipboardFunctionInput* ClipboardInput = ClipboardContent->FunctionInputs[0];
 	if (ClipboardInput != nullptr && ClipboardInput->InputType == InputType)
 	{
-		FScopedTransaction ScopedTransaction(LOCTEXT("PasteInput", "Paste niagara input"));
 		SetValueFromClipboardFunctionInput(*ClipboardInput);
 	}
 }
@@ -1022,7 +1061,7 @@ void UNiagaraStackFunctionInput::GetAvailableDynamicInputs(TArray<UNiagaraScript
 	}
 }
 
-void UNiagaraStackFunctionInput::SetDynamicInput(UNiagaraScript* DynamicInput)
+void UNiagaraStackFunctionInput::SetDynamicInput(UNiagaraScript* DynamicInput, FString SuggestedName)
 {
 	FScopedTransaction ScopedTransaction(LOCTEXT("SetDynamicInput", "Make dynamic input"));
 
@@ -1034,7 +1073,7 @@ void UNiagaraStackFunctionInput::SetDynamicInput(UNiagaraScript* DynamicInput)
 	}
 
 	UNiagaraNodeFunctionCall* FunctionCallNode;
-	FNiagaraStackGraphUtilities::SetDynamicInputForFunctionInput(OverridePin, DynamicInput, FunctionCallNode);
+	FNiagaraStackGraphUtilities::SetDynamicInputForFunctionInput(OverridePin, DynamicInput, FunctionCallNode, FGuid(), SuggestedName);
 	FNiagaraStackGraphUtilities::InitializeStackFunctionInputs(GetSystemViewModel(), GetEmitterViewModel(), GetStackEditorData(), *OwningModuleNode, *FunctionCallNode);
 	FNiagaraStackGraphUtilities::RelayoutGraph(*OwningFunctionCallNode->GetGraph());
 
@@ -1734,6 +1773,7 @@ const UNiagaraClipboardFunctionInput* UNiagaraStackFunctionInput::ToClipboardFun
 {
 	const UNiagaraClipboardFunctionInput* ClipboardInput = nullptr;
 	FName InputName = InputParameterHandle.GetName();
+	TOptional<bool> bEditConditionValue = GetHasEditCondition() ? GetEditConditionEnabled() : TOptional<bool>();
 	switch (InputValues.Mode)
 	{
 	case EValueMode::Local:
@@ -1741,21 +1781,21 @@ const UNiagaraClipboardFunctionInput* UNiagaraStackFunctionInput::ToClipboardFun
 		TArray<uint8> LocalValueData;
 		LocalValueData.AddUninitialized(InputType.GetSize());
 		FMemory::Memcpy(LocalValueData.GetData(), InputValues.LocalStruct->GetStructMemory(), InputType.GetSize());
-		ClipboardInput = UNiagaraClipboardFunctionInput::CreateLocalValue(InOuter, InputName, InputType, LocalValueData);
+		ClipboardInput = UNiagaraClipboardFunctionInput::CreateLocalValue(InOuter, InputName, InputType, bEditConditionValue, LocalValueData);
 		break;
 	}
 	case EValueMode::Linked:
-		ClipboardInput = UNiagaraClipboardFunctionInput::CreateLinkedValue(InOuter, InputName, InputType, InputValues.LinkedHandle.GetParameterHandleString());
+		ClipboardInput = UNiagaraClipboardFunctionInput::CreateLinkedValue(InOuter, InputName, InputType, bEditConditionValue, InputValues.LinkedHandle.GetParameterHandleString());
 		break;
 	case EValueMode::Data:
-		ClipboardInput = UNiagaraClipboardFunctionInput::CreateDataValue(InOuter, InputName, InputType, InputValues.DataObjects.GetValueObject());
+		ClipboardInput = UNiagaraClipboardFunctionInput::CreateDataValue(InOuter, InputName, InputType, bEditConditionValue, InputValues.DataObjects.GetValueObject());
 		break;
 	case EValueMode::Expression:
-		ClipboardInput = UNiagaraClipboardFunctionInput::CreateExpressionValue(InOuter, InputName, InputType, InputValues.ExpressionNode->GetHlslText().ToString());
+		ClipboardInput = UNiagaraClipboardFunctionInput::CreateExpressionValue(InOuter, InputName, InputType, bEditConditionValue, InputValues.ExpressionNode->GetHlslText().ToString());
 		break;
 	case EValueMode::Dynamic:
 	{
-		ClipboardInput = UNiagaraClipboardFunctionInput::CreateDynamicValue(InOuter, InputName, InputType, InputValues.DynamicNode->FunctionScript);
+		ClipboardInput = UNiagaraClipboardFunctionInput::CreateDynamicValue(InOuter, InputName, InputType, bEditConditionValue, InputValues.DynamicNode->GetFunctionName(), InputValues.DynamicNode->FunctionScript);
 
 		TArray<UNiagaraStackFunctionInputCollection*> DynamicInputCollections;
 		GetUnfilteredChildrenOfType(DynamicInputCollections);
@@ -1766,6 +1806,9 @@ const UNiagaraClipboardFunctionInput* UNiagaraStackFunctionInput::ToClipboardFun
 
 		break;
 	}
+	case EValueMode::Invalid:
+		// Do nothing.
+		break;
 	default:
 		ensureMsgf(false, TEXT("A new value mode was added without adding support for copy paste."));
 		break;
@@ -1799,7 +1842,7 @@ void UNiagaraStackFunctionInput::SetValueFromClipboardFunctionInput(const UNiaga
 			if (ensureMsgf(ClipboardFunctionInput.Dynamic->ScriptMode == ENiagaraClipboardFunctionScriptMode::ScriptAsset,
 				TEXT("Can not set dynamic input value from clipboard, only script asset funcitons can be set.")))
 			{
-				SetDynamicInput(ClipboardFunctionInput.Dynamic->Script);
+				SetDynamicInput(ClipboardFunctionInput.Dynamic->Script, ClipboardFunctionInput.Dynamic->FunctionName);
 				TArray<UNiagaraStackFunctionInputCollection*> DynamicInputCollections;
 				GetUnfilteredChildrenOfType(DynamicInputCollections);
 				for (UNiagaraStackFunctionInputCollection* DynamicInputCollection : DynamicInputCollections)
@@ -1813,6 +1856,11 @@ void UNiagaraStackFunctionInput::SetValueFromClipboardFunctionInput(const UNiaga
 			break;
 		}
 	}
+
+	if (GetHasEditCondition() && ClipboardFunctionInput.bHasEditCondition)
+	{
+		SetEditConditionEnabled(ClipboardFunctionInput.bEditConditionValue);
+	}
 }
 
 void UNiagaraStackFunctionInput::GetSearchItems(TArray<FStackSearchItem>& SearchItems) const
diff --git a/Engine/Plugins/FX/Niagara/Source/NiagaraEditor/Private/ViewModels/Stack/NiagaraStackGraphUtilities.cpp b/Engine/Plugins/FX/Niagara/Source/NiagaraEditor/Private/ViewModels/Stack/NiagaraStackGraphUtilities.cpp
--- a/Engine/Plugins/FX/Niagara/Source/NiagaraEditor/Private/ViewModels/Stack/NiagaraStackGraphUtilities.cpp
+++ b/Engine/Plugins/FX/Niagara/Source/NiagaraEditor/Private/ViewModels/Stack/NiagaraStackGraphUtilities.cpp
@@ -702,6 +702,38 @@ void FNiagaraStackGraphUtilities::GetStackFunctionStaticSwitchPins(UNiagaraNodeF
 	}
 }
 
+void FNiagaraStackGraphUtilities::GetStackFunctionOutputVariables(UNiagaraNodeFunctionCall& FunctionCallNode, FCompileConstantResolver ConstantResolver, TArray<FNiagaraVariable>& OutOutputVariables, TArray<FNiagaraVariable>& OutOutputVariablesWithOriginalAliasesIntact)
+{
+	FNiagaraParameterMapHistoryBuilder Builder;
+	Builder.SetIgnoreDisabled(false);
+	Builder.ConstantResolver = ConstantResolver;
+	FunctionCallNode.BuildParameterMapHistory(Builder, false);
+
+	if (ensureMsgf(Builder.Histories.Num() == 1, TEXT("Invalid Stack Graph - Function call node has invalid history count!")))
+	{
+		for (int32 i = 0; i < Builder.Histories[0].Variables.Num(); i++)
+		{
+			bool bHasParameterMapSetWrite = false;
+			for (const UEdGraphPin* WritePin : Builder.Histories[0].PerVariableWriteHistory[i])
+			{
+				if (WritePin != nullptr && WritePin->GetOwningNode() != nullptr && WritePin->GetOwningNode()->IsA<UNiagaraNodeParameterMapSet>())
+				{
+					bHasParameterMapSetWrite = true;
+					break;
+				}
+			}
+
+			if (bHasParameterMapSetWrite)
+			{
+				FNiagaraVariable& Variable = Builder.Histories[0].Variables[i];
+				FNiagaraVariable& VariableWithOriginalAliasIntact = Builder.Histories[0].VariablesWithOriginalAliasesIntact[i];
+				OutOutputVariables.Add(Variable);
+				OutOutputVariablesWithOriginalAliasesIntact.Add(VariableWithOriginalAliasIntact);
+			}
+		}
+	}
+}
+
 UNiagaraNodeParameterMapSet* FNiagaraStackGraphUtilities::GetStackFunctionOverrideNode(UNiagaraNodeFunctionCall& FunctionCallNode)
 {
 	UEdGraphPin* ParameterMapInput = FNiagaraStackGraphUtilities::GetParameterMapInputPin(FunctionCallNode);
diff --git a/Engine/Plugins/FX/Niagara/Source/NiagaraEditor/Private/ViewModels/Stack/NiagaraStackInputCategory.cpp b/Engine/Plugins/FX/Niagara/Source/NiagaraEditor/Private/ViewModels/Stack/NiagaraStackInputCategory.cpp
--- a/Engine/Plugins/FX/Niagara/Source/NiagaraEditor/Private/ViewModels/Stack/NiagaraStackInputCategory.cpp
+++ b/Engine/Plugins/FX/Niagara/Source/NiagaraEditor/Private/ViewModels/Stack/NiagaraStackInputCategory.cpp
@@ -97,7 +97,11 @@ void UNiagaraStackInputCategory::ToClipboardFunctionInputs(UObject* InOuter, TAr
 	GetUnfilteredChildrenOfType(ChildInputs);
 	for (UNiagaraStackFunctionInput* ChildInput : ChildInputs)
 	{
-		OutClipboardFunctionInputs.Add(ChildInput->ToClipboardFunctionInput(InOuter));
+		const UNiagaraClipboardFunctionInput* FunctionInput = ChildInput->ToClipboardFunctionInput(InOuter);
+		if (FunctionInput != nullptr)
+		{
+			OutClipboardFunctionInputs.Add(FunctionInput);
+		}
 	}
 }
 
diff --git a/Engine/Plugins/FX/Niagara/Source/NiagaraEditor/Private/ViewModels/Stack/NiagaraStackItem.cpp b/Engine/Plugins/FX/Niagara/Source/NiagaraEditor/Private/ViewModels/Stack/NiagaraStackItem.cpp
--- a/Engine/Plugins/FX/Niagara/Source/NiagaraEditor/Private/ViewModels/Stack/NiagaraStackItem.cpp
+++ b/Engine/Plugins/FX/Niagara/Source/NiagaraEditor/Private/ViewModels/Stack/NiagaraStackItem.cpp
@@ -27,6 +27,11 @@ UNiagaraStackItem::FOnModifiedGroupItems& UNiagaraStackItem::OnModifiedGroupItem
 	return ModifiedGroupItemsDelegate;
 }
 
+UNiagaraStackItem::FOnRequestPaste& UNiagaraStackItem::OnRequestPaste()
+{
+	return RequestPasteDelegate;
+}
+
 void UNiagaraStackItem::SetIsEnabled(bool bInIsEnabled)
 {
 	if (ItemFooter != nullptr)
@@ -36,15 +41,6 @@ void UNiagaraStackItem::SetIsEnabled(bool bInIsEnabled)
 	SetIsEnabledInternal(bInIsEnabled);
 }
 
-void UNiagaraStackItem::Delete()
-{
-	if (GetDisplayedObject() != nullptr)
-	{
-		GetSystemViewModel()->GetSelectionViewModel()->RemoveEntryFromSelectionByDisplayedObject(GetDisplayedObject());
-	}
-	DeleteInternal();
-}
-
 const TArray<FNiagaraScriptHighlight>& UNiagaraStackItem::GetHighlights() const
 {
 	static TArray<FNiagaraScriptHighlight> Empty;
diff --git a/Engine/Plugins/FX/Niagara/Source/NiagaraEditor/Private/ViewModels/Stack/NiagaraStackModuleItem.cpp b/Engine/Plugins/FX/Niagara/Source/NiagaraEditor/Private/ViewModels/Stack/NiagaraStackModuleItem.cpp
--- a/Engine/Plugins/FX/Niagara/Source/NiagaraEditor/Private/ViewModels/Stack/NiagaraStackModuleItem.cpp
+++ b/Engine/Plugins/FX/Niagara/Source/NiagaraEditor/Private/ViewModels/Stack/NiagaraStackModuleItem.cpp
@@ -31,6 +31,7 @@
 #include "NiagaraConstants.h"
 #include "Widgets/SWidget.h"
 #include "NiagaraActions.h"
+#include "NiagaraClipboard.h"
 
 #include "ScopedTransaction.h"
 #include "NiagaraScriptVariable.h"
@@ -844,55 +845,6 @@
 	OnRequestFullRefreshDeferred().Broadcast();
 }
 
-bool UNiagaraStackModuleItem::TestCanDeleteWithMessage(FText& OutCanDeleteMessage) const
-{
-	if (GetOwnerIsEnabled() == false)
-	{
-		OutCanDeleteMessage = LOCTEXT("CantDeleteOwnerDisabledToolTip", "This module can not be deleted because its owner is disabled.");
-		return false;
-	}
-	else if (CanMoveAndDelete())
-	{
-		OutCanDeleteMessage = LOCTEXT("DeleteToolTip", "Delete this module.");
-		return true;
-	}
-	else
-	{
-		OutCanDeleteMessage = LOCTEXT("CantDeleteToolTip", "This module can not be deleted becaue it is inherited.");
-		return false;
-	}
-}
-
-void UNiagaraStackModuleItem::DeleteInternal()
-{
-	checkf(CanMoveAndDelete(), TEXT("This module can't be deleted"));
-
-	FScopedTransaction ScopedTransaction(LOCTEXT("RemoveAModuleFromTheStack", "Remove a module from the stack"));
-
-	const FNiagaraEmitterHandle* EmitterHandle = GetEmitterViewModel().IsValid()
-		? FNiagaraEditorUtilities::GetEmitterHandleForEmitter(GetSystemViewModel()->GetSystem(), *GetEmitterViewModel()->GetEmitter())
-		: nullptr;
-	FGuid EmitterHandleId = EmitterHandle != nullptr ? EmitterHandle->GetId() : FGuid();
-
-	TArray<TWeakObjectPtr<UNiagaraNodeInput>> RemovedNodes;
-	bool bRemoved = FNiagaraStackGraphUtilities::RemoveModuleFromStack(GetSystemViewModel()->GetSystem(), EmitterHandleId, *FunctionCallNode, RemovedNodes);
-	if (bRemoved)
-	{
-		UNiagaraGraph* Graph = FunctionCallNode->GetNiagaraGraph();
-		Graph->NotifyGraphNeedsRecompile();
-		FNiagaraStackGraphUtilities::RelayoutGraph(*FunctionCallNode->GetGraph());
-		for (auto InputNode: RemovedNodes)
-		{
-			if (InputNode != nullptr && InputNode->Usage == ENiagaraInputNodeUsage::Parameter)
-			{
-				GetSystemViewModel()->NotifyDataObjectChanged(InputNode->GetDataInterface());
-			}
-		}
-		Finalize();
-		ModifiedGroupItemsDelegate.Broadcast();
-	}
-}
-
 bool UNiagaraStackModuleItem::SupportsHighlights() const
 {
 	return FunctionCallNode != nullptr && FunctionCallNode->FunctionScript != nullptr;
@@ -969,6 +921,190 @@
 	}
 }
 
+void UNiagaraStackModuleItem::SetInputValuesFromClipboardFunctionInputs(const TArray<const UNiagaraClipboardFunctionInput*>& ClipboardFunctionInputs)
+{
+	InputCollection->SetValuesFromClipboardFunctionInputs(ClipboardFunctionInputs);
+}
+
+bool UNiagaraStackModuleItem::TestCanCutWithMessage(FText& OutMessage) const
+{
+	FText CanCopyMessage;
+	if (TestCanCopyWithMessage(CanCopyMessage) == false)
+	{
+		OutMessage = FText::Format(LOCTEXT("CantCutBecauseCantCopyFormat", "This module can not be cut because it can't be copied.  {0}"), CanCopyMessage);
+		return false;
+	}
+
+	FText CanDeleteMessage;
+	if (TestCanDeleteWithMessage(CanDeleteMessage) == false)
+	{
+		OutMessage = FText::Format(LOCTEXT("CantCutBecauseCantDeleteFormat", "This module can't be cut because it can't be deleted.  {0}"), CanDeleteMessage);
+		return false;
+	}
+
+	OutMessage = LOCTEXT("CanCut", "Cut this module.");
+	return true;
+}
+
+FText UNiagaraStackModuleItem::GetCutTransactionText() const
+{
+	return LOCTEXT("CutModuleTransactionText", "Cut modules");
+}
+
+void UNiagaraStackModuleItem::CopyForCut(UNiagaraClipboardContent* ClipboardContent) const
+{
+	Copy(ClipboardContent);
+}
+
+void UNiagaraStackModuleItem::RemoveForCut()
+{
+	Delete();
+}
+
+bool UNiagaraStackModuleItem::TestCanCopyWithMessage(FText& OutMessage) const
+{
+	if (FunctionCallNode->GetClass() == UNiagaraNodeFunctionCall::StaticClass())
+	{
+		if (FunctionCallNode->FunctionScript == nullptr)
+		{
+			OutMessage = LOCTEXT("CantCopyInvalidModule", "This module can't be copied because it's referenced script is not valid.");
+			return false;
+		}
+	}
+	OutMessage = LOCTEXT("CopyModule", "Copy this module.");
+	return true;
+}
+
+void UNiagaraStackModuleItem::Copy(UNiagaraClipboardContent* ClipboardContent) const
+{
+	UNiagaraClipboardFunction* ClipboardFunction;
+	UNiagaraNodeAssignment* AssignmentNode = Cast<UNiagaraNodeAssignment>(FunctionCallNode);
+	if (AssignmentNode != nullptr)
+	{
+		ClipboardFunction = UNiagaraClipboardFunction::CreateAssignmentFunction(ClipboardContent, AssignmentNode->GetFunctionName(), AssignmentNode->GetAssignmentTargets(), AssignmentNode->GetAssignmentDefaults());
+	}
+	else
+	{
+		checkf(FunctionCallNode->FunctionScript != nullptr, TEXT("Can't copy this module because it's script is invalid.  Call TestCanCopyWithMessage to check this."));
+		ClipboardFunction = UNiagaraClipboardFunction::CreateScriptFunction(ClipboardContent, FunctionCallNode->GetFunctionName(), FunctionCallNode->FunctionScript);
+	}
+
+	InputCollection->ToClipboardFunctionInputs(ClipboardFunction, ClipboardFunction->Inputs);
+	ClipboardContent->Functions.Add(ClipboardFunction);
+}
+
+bool UNiagaraStackModuleItem::TestCanPasteWithMessage(const UNiagaraClipboardContent* ClipboardContent, FText& OutMessage) const
+{
+	if (ClipboardContent->FunctionInputs.Num() > 0)
+	{
+		OutMessage = LOCTEXT("PasteInputs", "Paste inputs from the clipboard which match inputs on this module by name and type.");
+		return true;
+	}
+	if (ClipboardContent->Functions.Num() > 0)
+	{
+		bool bValidUsage = false;
+		for (const UNiagaraClipboardFunction* Function : ClipboardContent->Functions)
+		{
+			if (Function != nullptr && Function->Script->GetSupportedUsageContexts().Contains(OutputNode->GetUsage()))
+			{
+				bValidUsage = true;
+			}
+		}
+		if (bValidUsage)
+		{
+			OutMessage = LOCTEXT("PasteModules", "Paste modules from the clipboard which have a valid usage.");
+			return true;
+		}
+		else
+		{
+			OutMessage = LOCTEXT("CantPasteModulesForUsage", "Can't paste the copied modules because they don't support the correct usage.");
+			return false;
+		}
+	}
+
+	OutMessage = FText();
+	return false;
+}
+
+FText UNiagaraStackModuleItem::GetPasteTransactionText(const UNiagaraClipboardContent* ClipboardContent) const
+{
+	if (ClipboardContent->FunctionInputs.Num() > 0)
+	{
+		return LOCTEXT("PasteInputsTransactionText", "Paste inputs to module.");
+	}
+	else
+	{
+		return LOCTEXT("PasteModuleTransactionText", "Paste niagara modules");
+	}
+}
+
+void UNiagaraStackModuleItem::Paste(const UNiagaraClipboardContent* ClipboardContent)
+{
+	if (ClipboardContent->FunctionInputs.Num() > 0)
+	{
+		SetInputValuesFromClipboardFunctionInputs(ClipboardContent->FunctionInputs);
+	}
+
+	if (ClipboardContent->Functions.Num() > 0)
+	{
+		// Pasted modules should go after this module, so add 1 to the index.
+		int32 PasteIndex = GetModuleIndex() + 1;
+		RequestPasteDelegate.Broadcast(ClipboardContent, PasteIndex);
+	}
+}
+
+bool UNiagaraStackModuleItem::TestCanDeleteWithMessage(FText& OutCanDeleteMessage) const
+{
+	if (GetOwnerIsEnabled() == false)
+	{
+		OutCanDeleteMessage = LOCTEXT("CantDeleteOwnerDisabledToolTip", "This module can not be deleted because its owner is disabled.");
+		return false;
+	}
+	else if (CanMoveAndDelete())
+	{
+		OutCanDeleteMessage = LOCTEXT("DeleteToolTip", "Delete this module.");
+		return true;
+	}
+	else
+	{
+		OutCanDeleteMessage = LOCTEXT("CantDeleteToolTip", "This module can not be deleted becaue it is inherited.");
+		return false;
+	}
+}
+
+FText UNiagaraStackModuleItem::GetDeleteTransactionText() const
+{
+	return LOCTEXT("DeleteModuleTransaction", "Delete modules");
+}
+
+void UNiagaraStackModuleItem::Delete()
+{
+	checkf(CanMoveAndDelete(), TEXT("This module can't be deleted"));
+
+	const FNiagaraEmitterHandle* EmitterHandle = GetEmitterViewModel().IsValid()
+		? FNiagaraEditorUtilities::GetEmitterHandleForEmitter(GetSystemViewModel()->GetSystem(), *GetEmitterViewModel()->GetEmitter())
+		: nullptr;
+	FGuid EmitterHandleId = EmitterHandle != nullptr ? EmitterHandle->GetId() : FGuid();
+
+	TArray<TWeakObjectPtr<UNiagaraNodeInput>> RemovedNodes;
+	bool bRemoved = FNiagaraStackGraphUtilities::RemoveModuleFromStack(GetSystemViewModel()->GetSystem(), EmitterHandleId, *FunctionCallNode, RemovedNodes);
+	if (bRemoved)
+	{
+		UNiagaraGraph* Graph = FunctionCallNode->GetNiagaraGraph();
+		Graph->NotifyGraphNeedsRecompile();
+		FNiagaraStackGraphUtilities::RelayoutGraph(*FunctionCallNode->GetGraph());
+		for (auto InputNode : RemovedNodes)
+		{
+			if (InputNode != nullptr && InputNode->Usage == ENiagaraInputNodeUsage::Parameter)
+			{
+				GetSystemViewModel()->NotifyDataObjectChanged(InputNode->GetDataInterface());
+			}
+		}
+		Finalize();
+		ModifiedGroupItemsDelegate.Broadcast();
+	}
+}
+
 UObject* UNiagaraStackModuleItem::GetExternalAsset() const
 {
 	if (GetModuleNode().FunctionScript != nullptr && GetModuleNode().FunctionScript->IsAsset())
diff --git a/Engine/Plugins/FX/Niagara/Source/NiagaraEditor/Private/ViewModels/Stack/NiagaraStackModuleItemOutputCollection.cpp b/Engine/Plugins/FX/Niagara/Source/NiagaraEditor/Private/ViewModels/Stack/NiagaraStackModuleItemOutputCollection.cpp
--- a/Engine/Plugins/FX/Niagara/Source/NiagaraEditor/Private/ViewModels/Stack/NiagaraStackModuleItemOutputCollection.cpp
+++ b/Engine/Plugins/FX/Niagara/Source/NiagaraEditor/Private/ViewModels/Stack/NiagaraStackModuleItemOutputCollection.cpp
@@ -48,38 +48,28 @@ void UNiagaraStackModuleItemOutputCollection::RefreshChildrenInternal(const TArr
 	UEdGraphPin* OutputParameterMapPin = FNiagaraStackGraphUtilities::GetParameterMapOutputPin(*FunctionCallNode);
 	if (ensureMsgf(OutputParameterMapPin != nullptr, TEXT("Invalid Stack Graph - Function call node has no output pin.")))
 	{
-		FNiagaraParameterMapHistoryBuilder Builder;
-		Builder.SetIgnoreDisabled(false);
-		Builder.ConstantResolver = GetEmitterViewModel().IsValid()
+		TArray<FNiagaraVariable> OutputVariables;
+		TArray<FNiagaraVariable> Unused;
+		FCompileConstantResolver ConstantResolver = GetEmitterViewModel().IsValid()
 			? FCompileConstantResolver(GetEmitterViewModel()->GetEmitter())
 			: FCompileConstantResolver();
-		FunctionCallNode->BuildParameterMapHistory(Builder, false);
-
-		if (ensureMsgf(Builder.Histories.Num() == 1, TEXT("Invalid Stack Graph - Function call node has invalid history count!")))
+		FNiagaraStackGraphUtilities::GetStackFunctionOutputVariables(*FunctionCallNode, ConstantResolver, OutputVariables, Unused);
+		
+		for(int32 OutputIndex = 0; OutputIndex < OutputVariables.Num(); OutputIndex++)
 		{
-			for (int32 i = 0; i < Builder.Histories[0].Variables.Num(); i++)
-			{
-				FNiagaraVariable& Variable = Builder.Histories[0].Variables[i];
-				TArray<const UEdGraphPin*>& WriteHistory = Builder.Histories[0].PerVariableWriteHistory[i];
-
-				for (const UEdGraphPin* WritePin : WriteHistory)
-				{
-					if (Cast<UNiagaraNodeParameterMapSet>(WritePin->GetOwningNode()) != nullptr)
-					{
-						UNiagaraStackModuleItemOutput* Output = FindCurrentChildOfTypeByPredicate<UNiagaraStackModuleItemOutput>(CurrentChildren,
-							[&](UNiagaraStackModuleItemOutput* CurrentOutput) { return CurrentOutput->GetOutputParameterHandle().GetParameterHandleString() == Variable.GetName(); });
+			FNiagaraVariable& Variable = OutputVariables[OutputIndex];
 
-						if (Output == nullptr)
-						{
-							Output = NewObject<UNiagaraStackModuleItemOutput>(this);
-							Output->Initialize(CreateDefaultChildRequiredData(), *FunctionCallNode, Variable.GetName(), Variable.GetType());
-						}
+			UNiagaraStackModuleItemOutput* Output = FindCurrentChildOfTypeByPredicate<UNiagaraStackModuleItemOutput>(CurrentChildren,
+				[&](UNiagaraStackModuleItemOutput* CurrentOutput) { return CurrentOutput->GetOutputParameterHandle().GetParameterHandleString() == Variable.GetName(); });
 
-						NewChildren.Add(Output);
-						break;
-					}
-				}
+			if (Output == nullptr)
+			{
+				Output = NewObject<UNiagaraStackModuleItemOutput>(this);
+				Output->Initialize(CreateDefaultChildRequiredData(), *FunctionCallNode, Variable.GetName(), Variable.GetType());
 			}
+
+			NewChildren.Add(Output);
+			break;
 		}
 	}
 }
diff --git a/Engine/Plugins/FX/Niagara/Source/NiagaraEditor/Private/ViewModels/Stack/NiagaraStackRenderItemGroup.cpp b/Engine/Plugins/FX/Niagara/Source/NiagaraEditor/Private/ViewModels/Stack/NiagaraStackRenderItemGroup.cpp
--- a/Engine/Plugins/FX/Niagara/Source/NiagaraEditor/Private/ViewModels/Stack/NiagaraStackRenderItemGroup.cpp
+++ b/Engine/Plugins/FX/Niagara/Source/NiagaraEditor/Private/ViewModels/Stack/NiagaraStackRenderItemGroup.cpp
@@ -7,6 +7,7 @@
 #include "NiagaraEmitterEditorData.h"
 #include "ViewModels/NiagaraSystemViewModel.h"
 #include "NiagaraRendererProperties.h"
+#include "NiagaraClipboard.h"
 
 #include "ScopedTransaction.h"
 #include "Widgets/Notifications/SNotificationList.h"
@@ -128,6 +129,36 @@ void UNiagaraStackRenderItemGroup::Initialize(FRequiredEntryData InRequiredEntry
 	EmitterWeak->OnRenderersChanged().AddUObject(this, &UNiagaraStackRenderItemGroup::EmitterRenderersChanged);
 }
 
+bool UNiagaraStackRenderItemGroup::TestCanPasteWithMessage(const UNiagaraClipboardContent* ClipboardContent, FText& OutMessage) const
+{
+	if (ClipboardContent->Renderers.Num() > 0)
+	{
+		OutMessage = LOCTEXT("PasteRenderers", "Paste renderers from the clipboard.");
+		return true;
+	}
+	OutMessage = LOCTEXT("NoRenderers", "No renderers on the clipboard.");
+	return false;
+}
+
+FText UNiagaraStackRenderItemGroup::GetPasteTransactionText(const UNiagaraClipboardContent* ClipboardContent) const
+{
+	return LOCTEXT("PasteRenderersTransactionText", "Paste renderers");
+}
+
+void UNiagaraStackRenderItemGroup::Paste(const UNiagaraClipboardContent* ClipboardContent)
+{
+	if (EmitterWeak.IsValid())
+	{
+		for (const UNiagaraRendererProperties* ClipboardRenderer : ClipboardContent->Renderers)
+		{
+			if (ClipboardRenderer != nullptr)
+			{
+				EmitterWeak->AddRenderer(CastChecked<UNiagaraRendererProperties>(StaticDuplicateObject(ClipboardRenderer, EmitterWeak.Get())));
+			}
+		}
+	}
+}
+
 void UNiagaraStackRenderItemGroup::RefreshChildrenInternal(const TArray<UNiagaraStackEntry*>& CurrentChildren, TArray<UNiagaraStackEntry*>& NewChildren, TArray<FStackIssue>& NewIssues)
 {
 	int32 RendererIndex = 0;
@@ -140,6 +171,7 @@ void UNiagaraStackRenderItemGroup::RefreshChildrenInternal(const TArray<UNiagara
 		{
 			RendererItem = NewObject<UNiagaraStackRendererItem>(this);
 			RendererItem->Initialize(CreateDefaultChildRequiredData(), RendererProperties);
+			RendererItem->OnRequestPaste().AddUObject(this, &UNiagaraStackRenderItemGroup::ChildRequestPaste);
 		}
 
 		NewChildren.Add(RendererItem);
@@ -160,6 +192,11 @@ void UNiagaraStackRenderItemGroup::EmitterRenderersChanged()
 	}
 }
 
+void UNiagaraStackRenderItemGroup::ChildRequestPaste(const UNiagaraClipboardContent* ClipboardContent, int32 PasteIndex)
+{
+	Paste(ClipboardContent);
+}
+
 void UNiagaraStackRenderItemGroup::FinalizeInternal()
 {
 	if (EmitterWeak.IsValid())
diff --git a/Engine/Plugins/FX/Niagara/Source/NiagaraEditor/Private/ViewModels/Stack/NiagaraStackRendererItem.cpp b/Engine/Plugins/FX/Niagara/Source/NiagaraEditor/Private/ViewModels/Stack/NiagaraStackRendererItem.cpp
--- a/Engine/Plugins/FX/Niagara/Source/NiagaraEditor/Private/ViewModels/Stack/NiagaraStackRendererItem.cpp
+++ b/Engine/Plugins/FX/Niagara/Source/NiagaraEditor/Private/ViewModels/Stack/NiagaraStackRendererItem.cpp
@@ -24,6 +24,7 @@
 #include "Widgets/Notifications/SNotificationList.h"
 #include "ViewModels/Stack/NiagaraStackGraphUtilities.h"
 #include "NiagaraScriptMergeManager.h"
+#include "NiagaraClipboard.h"
 
 #include "Styling/SlateIconFinder.h"
 
@@ -146,6 +147,73 @@ FText UNiagaraStackRendererItem::GetDisplayName() const
 	}
 }
 
+bool UNiagaraStackRendererItem::TestCanCutWithMessage(FText& OutMessage) const
+{
+	FText CanCopyMessage;
+	if (TestCanCopyWithMessage(CanCopyMessage) == false)
+	{
+		OutMessage = FText::Format(LOCTEXT("CantCutBecauseCantCopyFormat", "This renderer can not be cut because it can't be copied.  {0}"), CanCopyMessage);
+		return false;
+	}
+
+	FText CanDeleteMessage;
+	if (TestCanDeleteWithMessage(CanDeleteMessage) == false)
+	{
+		OutMessage = FText::Format(LOCTEXT("CantCutBecauseCantDeleteFormat", "This renderer can't be cut because it can't be deleted.  {0}"), CanDeleteMessage);
+		return false;
+	}
+
+	OutMessage = LOCTEXT("CanCut", "Cut this renderer.");
+	return true;
+}
+
+FText UNiagaraStackRendererItem::GetCutTransactionText() const
+{
+	return LOCTEXT("CutRendererTransactionText", "Cut renderers");
+}
+
+void UNiagaraStackRendererItem::CopyForCut(UNiagaraClipboardContent* ClipboardContent) const
+{
+	Copy(ClipboardContent);
+}
+
+void UNiagaraStackRendererItem::RemoveForCut()
+{
+	Delete();
+}
+
+bool UNiagaraStackRendererItem::TestCanCopyWithMessage(FText& OutMessage) const
+{
+	OutMessage = LOCTEXT("CopyRenderer", "Copy this renderer.");
+	return true;
+}
+
+void UNiagaraStackRendererItem::Copy(UNiagaraClipboardContent* ClipboardContent) const
+{
+	ClipboardContent->Renderers.Add(CastChecked<UNiagaraRendererProperties>(StaticDuplicateObject(RendererProperties.Get(), ClipboardContent)));
+}
+
+bool UNiagaraStackRendererItem::TestCanPasteWithMessage(const UNiagaraClipboardContent* ClipboardContent, FText& OutMessage) const
+{
+	if (ClipboardContent->Renderers.Num() > 0)
+	{
+		OutMessage = LOCTEXT("PasteRenderers", "Paste renderers from the clipboard.");
+		return true;
+	}
+	OutMessage = LOCTEXT("NoRenderers", "No renderers on the clipboard.");
+	return false;
+}
+
+FText UNiagaraStackRendererItem::GetPasteTransactionText(const UNiagaraClipboardContent* ClipboardContent) const
+{
+	return LOCTEXT("PasteRenderersTransactionText", "Paste renderers");
+}
+
+void UNiagaraStackRendererItem::Paste(const UNiagaraClipboardContent* ClipboardContent)
+{
+	OnRequestPaste().Broadcast(ClipboardContent, INDEX_NONE);
+}
+
 bool UNiagaraStackRendererItem::TestCanDeleteWithMessage(FText& OutCanDeleteMessage) const
 {
 	if (GetOwnerIsEnabled() == false)
@@ -165,10 +233,13 @@ bool UNiagaraStackRendererItem::TestCanDeleteWithMessage(FText& OutCanDeleteMess
 	}
 }
 
-void UNiagaraStackRendererItem::DeleteInternal()
+FText UNiagaraStackRendererItem::GetDeleteTransactionText() const
 {
-	const FScopedTransaction Transaction(LOCTEXT("DeleteRenderer", "Delete Renderer"));
+	return LOCTEXT("DeleteRenderer", "Delete Renderer");
+}
 
+void UNiagaraStackRendererItem::Delete()
+{
 	UNiagaraEmitter* Emitter = GetEmitterViewModel()->GetEmitter();
 	Emitter->Modify();
 	Emitter->RemoveRenderer(RendererProperties.Get());
diff --git a/Engine/Plugins/FX/Niagara/Source/NiagaraEditor/Private/ViewModels/Stack/NiagaraStackScriptItemGroup.cpp b/Engine/Plugins/FX/Niagara/Source/NiagaraEditor/Private/ViewModels/Stack/NiagaraStackScriptItemGroup.cpp
--- a/Engine/Plugins/FX/Niagara/Source/NiagaraEditor/Private/ViewModels/Stack/NiagaraStackScriptItemGroup.cpp
+++ b/Engine/Plugins/FX/Niagara/Source/NiagaraEditor/Private/ViewModels/Stack/NiagaraStackScriptItemGroup.cpp
@@ -19,6 +19,7 @@
 #include "NiagaraSystem.h"
 #include "NiagaraEmitterHandle.h"
 #include "NiagaraEditorSettings.h"
+#include "NiagaraClipboard.h"
 
 #include "Internationalization/Internationalization.h"
 #include "ScopedTransaction.h"
@@ -327,6 +328,49 @@ void UNiagaraStackScriptItemGroup::FinalizeInternal()
 	Super::FinalizeInternal();
 }
 
+bool UNiagaraStackScriptItemGroup::TestCanPasteWithMessage(const UNiagaraClipboardContent* ClipboardContent, FText& OutMessage) const
+{
+	if (ClipboardContent->Functions.Num() > 0)
+	{
+		bool bValidUsage = false;
+		for (const UNiagaraClipboardFunction* Function : ClipboardContent->Functions)
+		{
+			if (Function != nullptr && Function->Script->GetSupportedUsageContexts().Contains(ScriptUsage))
+			{
+				bValidUsage = true;
+			}
+		}
+		if (bValidUsage)
+		{
+			OutMessage = LOCTEXT("PasteModules", "Paste modules from the clipboard which have a valid usage.");
+			return true;
+		}
+		else
+		{
+			OutMessage = LOCTEXT("CantPasteModulesForUsage", "Can't paste the copied modules because they don't support the correct usage.");
+			return false;
+		}
+	}
+	OutMessage = FText();
+	return false;
+}
+
+FText UNiagaraStackScriptItemGroup::GetPasteTransactionText(const UNiagaraClipboardContent* ClipboardContent) const
+{
+	return LOCTEXT("PasteModuleTransactionText", "Paste niagara modules");
+}
+
+void UNiagaraStackScriptItemGroup::Paste(const UNiagaraClipboardContent* ClipboardContent)
+{
+	if (ClipboardContent->Functions.Num() > 0)
+	{
+		TArray<UNiagaraStackModuleItem*> ModuleItems;
+		GetUnfilteredChildrenOfType(ModuleItems);
+		int32 PasteIndex = ModuleItems.Num() > 0 ? ModuleItems.Last()->GetModuleIndex() + 1 : 0;
+		PasteModules(ClipboardContent, PasteIndex);
+	}
+}
+
 void UNiagaraStackScriptItemGroup::RefreshChildrenInternal(const TArray<UNiagaraStackEntry*>& CurrentChildren, TArray<UNiagaraStackEntry*>& NewChildren, TArray<FStackIssue>& NewIssues)
 {
 	TSharedPtr<FNiagaraScriptViewModel> ScriptViewModelPinned = ScriptViewModel.Pin();
@@ -354,6 +398,7 @@ void UNiagaraStackScriptItemGroup::RefreshChildrenInternal(const TArray<UNiagara
 				ModuleItem = NewObject<UNiagaraStackModuleItem>(this);
 				ModuleItem->Initialize(CreateDefaultChildRequiredData(), GetAddUtilities(), *ModuleNode);
 				ModuleItem->OnModifiedGroupItems().AddUObject(this, &UNiagaraStackScriptItemGroup::ChildModifiedGroupItems);
+				ModuleItem->OnRequestPaste().AddUObject(this, &UNiagaraStackScriptItemGroup::ChildRequestPaste);
 			}
 
 			NewChildren.Add(ModuleItem);
@@ -896,6 +941,11 @@ void UNiagaraStackScriptItemGroup::ChildModifiedGroupItems()
 	RefreshChildren();
 }
 
+void UNiagaraStackScriptItemGroup::ChildRequestPaste(const UNiagaraClipboardContent* ClipboardContent, int32 PasteIndex)
+{
+	PasteModules(ClipboardContent, PasteIndex);
+}
+
 void UNiagaraStackScriptItemGroup::OnScriptGraphChanged(const struct FEdGraphEditAction& InAction)
 {
 	if (InAction.Action == GRAPHACTION_RemoveNode)
@@ -904,4 +954,138 @@ void UNiagaraStackScriptItemGroup::OnScriptGraphChanged(const struct FEdGraphEdi
 	}
 }
 
+void GatherRenamedModuleOutputs(
+	UNiagaraStackScriptItemGroup* InOwnerEntry,
+	TArray<TPair<const UNiagaraClipboardFunction*, UNiagaraNodeFunctionCall*>> InClipboardFunctionAndNodeFunctionPairs,
+	TMap<FName, FName>& OutOldModuleOutputNameToNewModuleOutputNameMap)
+{
+	for (TPair<const UNiagaraClipboardFunction*, UNiagaraNodeFunctionCall*>& ClipboardFunctionAndNodeFunctionPair : InClipboardFunctionAndNodeFunctionPairs)
+	{
+		const UNiagaraClipboardFunction* ClipboardFunction = ClipboardFunctionAndNodeFunctionPair.Key;
+		UNiagaraNodeFunctionCall* NewFunctionCallNode = ClipboardFunctionAndNodeFunctionPair.Value;
+		if (NewFunctionCallNode->GetFunctionName() != ClipboardFunction->FunctionName)
+		{
+			TArray<FNiagaraVariable> OutputVariables;
+			TArray<FNiagaraVariable> OutputVariablesWithOriginalAliasesIntact;
+			FCompileConstantResolver ConstantResolver = InOwnerEntry->GetEmitterViewModel().IsValid()
+				? FCompileConstantResolver(InOwnerEntry->GetEmitterViewModel()->GetEmitter())
+				: FCompileConstantResolver();
+			FNiagaraStackGraphUtilities::GetStackFunctionOutputVariables(*NewFunctionCallNode, ConstantResolver, OutputVariables, OutputVariablesWithOriginalAliasesIntact);
+
+			for (FNiagaraVariable& OutputVariableWithOriginalAliasesIntact : OutputVariablesWithOriginalAliasesIntact)
+			{
+				TArray<FString> SplitAliasedVariableName;
+				OutputVariableWithOriginalAliasesIntact.GetName().ToString().ParseIntoArray(SplitAliasedVariableName, TEXT("."));
+				if (SplitAliasedVariableName.Contains(TEXT("Module")))
+				{
+					TArray<FString> SplitOldVariableName = SplitAliasedVariableName;
+					TArray<FString> SplitNewVariableName = SplitAliasedVariableName;
+					for (int32 i = 0; i < SplitAliasedVariableName.Num(); i++)
+					{
+						if (SplitAliasedVariableName[i] == TEXT("Module"))
+						{
+							SplitOldVariableName[i] = ClipboardFunction->FunctionName;
+							SplitNewVariableName[i] = NewFunctionCallNode->GetFunctionName();
+						}
+					}
+					OutOldModuleOutputNameToNewModuleOutputNameMap.Add(*FString::Join(SplitOldVariableName, TEXT(".")), *FString::Join(SplitNewVariableName, TEXT(".")));
+				}
+			}
+		}
+	}
+}
+
+void RenameInputsFromClipboard(TMap<FName, FName> OldModuleOutputNameToNewModuleOutputNameMap, UObject* InOuter, const TArray<const UNiagaraClipboardFunctionInput*>& InSourceInputs, TArray<const UNiagaraClipboardFunctionInput*>& OutRenamedInputs)
+{
+	for (const UNiagaraClipboardFunctionInput* SourceInput : InSourceInputs)
+	{
+		TOptional<bool> bEditConditionValue = SourceInput->bHasEditCondition 
+			? TOptional<bool>(SourceInput->bEditConditionValue)
+			: TOptional<bool>();
+		if (SourceInput->ValueMode == ENiagaraClipboardFunctionInputValueMode::Linked && OldModuleOutputNameToNewModuleOutputNameMap.Contains(SourceInput->Linked))
+		{
+			FName RenamedLinkedValue = OldModuleOutputNameToNewModuleOutputNameMap[SourceInput->Linked];
+			OutRenamedInputs.Add(UNiagaraClipboardFunctionInput::CreateLinkedValue(InOuter, SourceInput->InputName, SourceInput->InputType, bEditConditionValue, RenamedLinkedValue));
+		}
+		else if (SourceInput->ValueMode == ENiagaraClipboardFunctionInputValueMode::Dynamic)
+		{
+			const UNiagaraClipboardFunctionInput* RenamedDynamicInput = UNiagaraClipboardFunctionInput::CreateDynamicValue(InOuter, SourceInput->InputName, SourceInput->InputType, bEditConditionValue, SourceInput->Dynamic->FunctionName, SourceInput->Dynamic->Script);
+			RenameInputsFromClipboard(OldModuleOutputNameToNewModuleOutputNameMap, RenamedDynamicInput->Dynamic, SourceInput->Dynamic->Inputs, RenamedDynamicInput->Dynamic->Inputs);
+			OutRenamedInputs.Add(RenamedDynamicInput);
+		}
+		else
+		{
+			OutRenamedInputs.Add(CastChecked<UNiagaraClipboardFunctionInput>(StaticDuplicateObject(SourceInput, InOuter)));
+		}
+	}
+}
+
+void UNiagaraStackScriptItemGroup::PasteModules(const UNiagaraClipboardContent* ClipboardContent, int32 PasteIndex)
+{
+	UNiagaraNodeOutput* OutputNode = GetScriptOutputNode();
+	
+	// Add the new function call nodes from the clipboard functions and collected the pairs.
+	TArray<TPair<const UNiagaraClipboardFunction*, UNiagaraNodeFunctionCall*>> ClipboardFunctionAndNodeFunctionPairs;
+	int32 CurrentPasteIndex = PasteIndex;
+	for (const UNiagaraClipboardFunction* ClipboardFunction : ClipboardContent->Functions)
+	{
+		if (ClipboardFunction != nullptr)
+		{
+			UNiagaraNodeFunctionCall* NewFunctionCallNode = nullptr;
+			switch (ClipboardFunction->ScriptMode)
+			{
+			case ENiagaraClipboardFunctionScriptMode::Assignment:
+			{
+				NewFunctionCallNode = FNiagaraStackGraphUtilities::AddParameterModuleToStack(ClipboardFunction->AssignmentTargets, *OutputNode, CurrentPasteIndex, ClipboardFunction->AssignmentDefaults);
+				break;
+			}
+			case ENiagaraClipboardFunctionScriptMode::ScriptAsset:
+			{
+				if (ClipboardFunction->Script->GetSupportedUsageContexts().Contains(OutputNode->GetUsage()))
+				{
+					NewFunctionCallNode = FNiagaraStackGraphUtilities::AddScriptModuleToStack(ClipboardFunction->Script, *OutputNode, CurrentPasteIndex);
+				}
+				break;
+			}
+			default:
+				checkf(false, TEXT("Unsupported clipboard function mode."));
+			}
+
+			NewFunctionCallNode->SuggestName(ClipboardFunction->FunctionName);
+			ClipboardFunctionAndNodeFunctionPairs.Add(TPair<const UNiagaraClipboardFunction*, UNiagaraNodeFunctionCall*>(ClipboardFunction, NewFunctionCallNode));
+			CurrentPasteIndex++;
+		}
+	}
+	RefreshChildren();
+
+	// Check to see if any of the functions were renamed, and if they were, check if they had aliased outputs that need to be fixed up.
+	TMap<FName, FName> OldModuleOutputNameToNewModuleOutputNameMap;
+	GatherRenamedModuleOutputs(this, ClipboardFunctionAndNodeFunctionPairs, OldModuleOutputNameToNewModuleOutputNameMap);
+
+	// Apply the function inputs now that we have the map of old module output names to new module output names.
+	for(TPair<const UNiagaraClipboardFunction*, UNiagaraNodeFunctionCall*>& ClipboardFunctionAndNodeFunctionPair : ClipboardFunctionAndNodeFunctionPairs)
+	{
+		const UNiagaraClipboardFunction* ClipboardFunction = ClipboardFunctionAndNodeFunctionPair.Key;
+		UNiagaraNodeFunctionCall* NewFunctionCallNode = ClipboardFunctionAndNodeFunctionPair.Value;
+
+		TArray<UNiagaraStackModuleItem*> ModuleItems;
+		GetUnfilteredChildrenOfType(ModuleItems);
+
+		UNiagaraStackModuleItem** NewModuleItemPtr = ModuleItems.FindByPredicate([NewFunctionCallNode](UNiagaraStackModuleItem* ModuleItem) { return &ModuleItem->GetModuleNode() == NewFunctionCallNode; });
+		if (NewModuleItemPtr != nullptr)
+		{
+			TArray<const UNiagaraClipboardFunctionInput*> FunctionInputs;
+			if (OldModuleOutputNameToNewModuleOutputNameMap.Num() > 0)
+			{
+				RenameInputsFromClipboard(OldModuleOutputNameToNewModuleOutputNameMap, GetTransientPackage(), ClipboardFunction->Inputs, FunctionInputs);
+			}
+			else
+			{
+				FunctionInputs = ClipboardFunction->Inputs;
+			}
+			(*NewModuleItemPtr)->SetInputValuesFromClipboardFunctionInputs(FunctionInputs);
+		}
+	}
+}
+
 #undef LOCTEXT_NAMESPACE
diff --git a/Engine/Plugins/FX/Niagara/Source/NiagaraEditor/Public/NiagaraClipboard.h b/Engine/Plugins/FX/Niagara/Source/NiagaraEditor/Public/NiagaraClipboard.h
--- a/Engine/Plugins/FX/Niagara/Source/NiagaraEditor/Public/NiagaraClipboard.h
+++ b/Engine/Plugins/FX/Niagara/Source/NiagaraEditor/Public/NiagaraClipboard.h
@@ -8,6 +8,7 @@
 
 class UNiagaraDataInterface;
 class UNiagaraScript;
+class UNiagaraRendererProperties;
 
 UENUM()
 enum class ENiagaraClipboardFunctionInputValueMode
@@ -27,15 +28,15 @@ class NIAGARAEDITOR_API UNiagaraClipboardFunctionInput : public UObject
 	GENERATED_BODY()
 
 public:
-	static const UNiagaraClipboardFunctionInput* CreateLocalValue(UObject* InOuter, FName InInputName, FNiagaraTypeDefinition InInputType, TArray<uint8>& InLocalValueData);
+	static const UNiagaraClipboardFunctionInput* CreateLocalValue(UObject* InOuter, FName InInputName, FNiagaraTypeDefinition InInputType, TOptional<bool> bInEditConditionValue, TArray<uint8>& InLocalValueData);
 
-	static const UNiagaraClipboardFunctionInput* CreateLinkedValue(UObject* InOuter, FName InInputName, FNiagaraTypeDefinition InInputType, FName InLinkedValue);
+	static const UNiagaraClipboardFunctionInput* CreateLinkedValue(UObject* InOuter, FName InInputName, FNiagaraTypeDefinition InInputType, TOptional<bool> bInEditConditionValue, FName InLinkedValue);
 
-	static const UNiagaraClipboardFunctionInput* CreateDataValue(UObject* InOuter, FName InInputName, FNiagaraTypeDefinition InInputType, UNiagaraDataInterface* InDataValue);
+	static const UNiagaraClipboardFunctionInput* CreateDataValue(UObject* InOuter, FName InInputName, FNiagaraTypeDefinition InInputType, TOptional<bool> bInEditConditionValue, UNiagaraDataInterface* InDataValue);
 
-	static const UNiagaraClipboardFunctionInput* CreateExpressionValue(UObject* InOuter, FName InInputName, FNiagaraTypeDefinition InInputType, const FString& InExpressionValue);
+	static const UNiagaraClipboardFunctionInput* CreateExpressionValue(UObject* InOuter, FName InInputName, FNiagaraTypeDefinition InInputType, TOptional<bool> bInEditConditionValue, const FString& InExpressionValue);
 
-	static const UNiagaraClipboardFunctionInput* CreateDynamicValue(UObject* InOuter, FName InInputName, FNiagaraTypeDefinition InInputType, UNiagaraScript* InDynamicValue);
+	static const UNiagaraClipboardFunctionInput* CreateDynamicValue(UObject* InOuter, FName InInputName, FNiagaraTypeDefinition InInputType, TOptional<bool> bInEditConditionValue, FString InDynamicValueName, UNiagaraScript* InDynamicValue);
 
 	UPROPERTY()
 	FName InputName;
@@ -43,6 +44,12 @@ public:
 	UPROPERTY()
 	FNiagaraTypeDefinition InputType;
 
+	UPROPERTY()
+	bool bHasEditCondition;
+
+	UPROPERTY()
+	bool bEditConditionValue;
+
 	UPROPERTY()
 	ENiagaraClipboardFunctionInputValueMode ValueMode;
 
@@ -75,7 +82,12 @@ class NIAGARAEDITOR_API UNiagaraClipboardFunction : public UObject
 	GENERATED_BODY()
 
 public:
-	static UNiagaraClipboardFunction* CreateScriptFunction(UObject* InOuter, UNiagaraScript* InScript);
+	static UNiagaraClipboardFunction* CreateScriptFunction(UObject* InOuter, FString InFunctionName, UNiagaraScript* InScript);
+
+	static UNiagaraClipboardFunction* CreateAssignmentFunction(UObject* InOuter, FString InFunctionName, const TArray<FNiagaraVariable>& InAssignmentTargets, const TArray<FString>& InAssignmentDefaults);
+
+	UPROPERTY()
+	FString FunctionName;
 
 	UPROPERTY()
 	ENiagaraClipboardFunctionScriptMode ScriptMode;
@@ -86,6 +98,9 @@ public:
 	UPROPERTY()
 	TArray<FNiagaraVariable> AssignmentTargets;
 
+	UPROPERTY()
+	TArray<FString> AssignmentDefaults;
+
 	UPROPERTY()
 	TArray<const UNiagaraClipboardFunctionInput*> Inputs;
 };
@@ -96,11 +111,16 @@ class UNiagaraClipboardContent : public UObject
 	GENERATED_BODY()
 
 public:
+	static UNiagaraClipboardContent* Create();
+
 	UPROPERTY()
 	TArray<const UNiagaraClipboardFunction*> Functions;
 
 	UPROPERTY()
 	TArray<const UNiagaraClipboardFunctionInput*> FunctionInputs;
+
+	UPROPERTY()
+	TArray<const UNiagaraRendererProperties*> Renderers;
 };
 
 class FNiagaraClipboard
@@ -108,11 +128,7 @@ class FNiagaraClipboard
 public:
 	FNiagaraClipboard();
 
-	void Copy(const UNiagaraClipboardFunctionInput* Input);
-
-	const UNiagaraClipboardFunctionInput* GetCopiedInput() const;
+	void SetClipboardContent(UNiagaraClipboardContent* ClipboardContent);
 
-private:
-	void SetClipboardContentInternal(UNiagaraClipboardContent* ClipboardContent);
-	const UNiagaraClipboardContent* GetClipboardContentInternal() const;
+	const UNiagaraClipboardContent* GetClipboardContent() const;
 };
\ No newline at end of file
diff --git a/Engine/Plugins/FX/Niagara/Source/NiagaraEditor/Public/ViewModels/NiagaraSystemSelectionViewModel.h b/Engine/Plugins/FX/Niagara/Source/NiagaraEditor/Public/ViewModels/NiagaraSystemSelectionViewModel.h
--- a/Engine/Plugins/FX/Niagara/Source/NiagaraEditor/Public/ViewModels/NiagaraSystemSelectionViewModel.h
+++ b/Engine/Plugins/FX/Niagara/Source/NiagaraEditor/Public/ViewModels/NiagaraSystemSelectionViewModel.h
@@ -48,6 +48,8 @@ public:
 
 	void EmptySelection();
 
+	void RemoveEntriesFromSelection(const TArray<UNiagaraStackEntry*>& InEntriesToRemove);
+
 	void RemoveEntryFromSelectionByDisplayedObject(const UObject* InObject);
 
 	void AddEntriesToSelectionByDisplayedObjectsDeferred(const TArray<const UObject*>& InObjects);
diff --git a/Engine/Plugins/FX/Niagara/Source/NiagaraEditor/Public/ViewModels/Stack/NiagaraStackClipboardUtilities.h b/Engine/Plugins/FX/Niagara/Source/NiagaraEditor/Public/ViewModels/Stack/NiagaraStackClipboardUtilities.h
new file mode 100644
--- /dev/null
+++ b/Engine/Plugins/FX/Niagara/Source/NiagaraEditor/Public/ViewModels/Stack/NiagaraStackClipboardUtilities.h
@@ -0,0 +1,24 @@
+// Copyright 1998-2019 Epic Games, Inc. All Rights Reserved.
+
+#pragma once
+
+class UNiagaraStackEntry;
+
+namespace FNiagaraStackClipboardUtilities
+{
+	bool NIAGARAEDITOR_API TestCanCutSelectionWithMessage(const TArray<UNiagaraStackEntry*>& SelectedEntries, FText& OutCanCutMessage);
+
+	void NIAGARAEDITOR_API CutSelection(const TArray<UNiagaraStackEntry*>& SelectedEntries);
+
+	bool NIAGARAEDITOR_API TestCanCopySelectionWithMessage(const TArray<UNiagaraStackEntry*>& SelectedEntries, FText& OutCanCopyMessage);
+
+	void NIAGARAEDITOR_API CopySelection(const TArray<UNiagaraStackEntry*>& SelectedEntries);
+
+	bool NIAGARAEDITOR_API TestCanPasteSelectionWithMessage(const TArray<UNiagaraStackEntry*>& SelectedEntries, FText& OutCanPasteMessage);
+
+	void NIAGARAEDITOR_API PasteSelection(const TArray<UNiagaraStackEntry*>& SelectedEntries);
+
+	bool NIAGARAEDITOR_API TestCanDeleteSelectionWithMessage(const TArray<UNiagaraStackEntry*>& SelectedEntries, FText& OutMessage);
+
+	void NIAGARAEDITOR_API DeleteSelection(const TArray<UNiagaraStackEntry*>& SelectedEntries);
+}
\ No newline at end of file
diff --git a/Engine/Plugins/FX/Niagara/Source/NiagaraEditor/Public/ViewModels/Stack/NiagaraStackEntry.h b/Engine/Plugins/FX/Niagara/Source/NiagaraEditor/Public/ViewModels/Stack/NiagaraStackEntry.h
--- a/Engine/Plugins/FX/Niagara/Source/NiagaraEditor/Public/ViewModels/Stack/NiagaraStackEntry.h
+++ b/Engine/Plugins/FX/Niagara/Source/NiagaraEditor/Public/ViewModels/Stack/NiagaraStackEntry.h
@@ -11,6 +11,7 @@ class FNiagaraEmitterViewModel;
 class FNiagaraScriptViewModel;
 class UNiagaraStackEditorData;
 class UNiagaraStackErrorItem;
+class UNiagaraClipboardContent;
 
 UENUM()
 enum class EStackIssueSeverity : uint8
@@ -348,19 +349,32 @@ public:
 
 	virtual bool TestCanCutWithMessage(FText& OutMessage) const { return false; }
 
-	virtual void Cut() { }
+	virtual FText GetCutTransactionText() const { return FText(); }
+
+	virtual void CopyForCut(UNiagaraClipboardContent* ClipboardContent) const { }
+	virtual void RemoveForCut() { }
 
 	virtual bool SupportsCopy() const { return false; }
 
 	virtual bool TestCanCopyWithMessage(FText& OutMessage) const { return false; }
 
-	virtual void Copy() const { }
+	virtual void Copy(UNiagaraClipboardContent* ClipboardContent) const { }
 
 	virtual bool SupportsPaste() const { return false; }
 
-	virtual bool TestCanPasteWithMessage(FText& OutMessage) const { return false; }
+	virtual bool TestCanPasteWithMessage(const UNiagaraClipboardContent* ClipboardContent, FText& OutMessage) const { return false; }
+
+	virtual FText GetPasteTransactionText(const UNiagaraClipboardContent* ClipboardContent) const { return FText(); }
+
+	virtual void Paste(const UNiagaraClipboardContent* ClipboardContent) { }
+
+	virtual bool SupportsDelete() const { return false; }
+
+	virtual bool TestCanDeleteWithMessage(FText& OutCanDeleteMessage) const { return false; }
+
+	virtual FText GetDeleteTransactionText() const { return FText(); }
 
-	virtual void Paste() { }
+	virtual void Delete() { }
 
 protected:
 	virtual void BeginDestroy() override;
diff --git a/Engine/Plugins/FX/Niagara/Source/NiagaraEditor/Public/ViewModels/Stack/NiagaraStackEventScriptItemGroup.h b/Engine/Plugins/FX/Niagara/Source/NiagaraEditor/Public/ViewModels/Stack/NiagaraStackEventScriptItemGroup.h
--- a/Engine/Plugins/FX/Niagara/Source/NiagaraEditor/Public/ViewModels/Stack/NiagaraStackEventScriptItemGroup.h
+++ b/Engine/Plugins/FX/Niagara/Source/NiagaraEditor/Public/ViewModels/Stack/NiagaraStackEventScriptItemGroup.h
@@ -72,14 +72,15 @@ public:
 
 	void SetOnModifiedEventHandlers(FOnModifiedEventHandlers OnModifiedEventHandlers);
 
+	virtual bool SupportsDelete() const override { return true; }
+	virtual bool TestCanDeleteWithMessage(FText& OutCanDeleteMessage) const override;
+	virtual void Delete() override;
+
 protected:
 	FOnModifiedEventHandlers OnModifiedEventHandlersDelegate;
 
 	virtual void RefreshChildrenInternal(const TArray<UNiagaraStackEntry*>& CurrentChildren, TArray<UNiagaraStackEntry*>& NewChildren, TArray<FStackIssue>& NewIssues) override;
 
-	virtual bool CanDelete() const override;
-	virtual bool Delete() override;
-
 private:
 	bool HasBaseEventHandler() const;
 	
diff --git a/Engine/Plugins/FX/Niagara/Source/NiagaraEditor/Public/ViewModels/Stack/NiagaraStackFunctionInput.h b/Engine/Plugins/FX/Niagara/Source/NiagaraEditor/Public/ViewModels/Stack/NiagaraStackFunctionInput.h
--- a/Engine/Plugins/FX/Niagara/Source/NiagaraEditor/Public/ViewModels/Stack/NiagaraStackFunctionInput.h
+++ b/Engine/Plugins/FX/Niagara/Source/NiagaraEditor/Public/ViewModels/Stack/NiagaraStackFunctionInput.h
@@ -91,13 +91,16 @@ public:
 	virtual UObject* GetExternalAsset() const override;
 	virtual bool SupportsCut() const override { return true; }
 	virtual bool TestCanCutWithMessage(FText& OutMessage) const override;
-	virtual void Cut() override;
+	virtual FText GetCutTransactionText() const override;
+	virtual void CopyForCut(UNiagaraClipboardContent* ClipboardContent) const override;
+	virtual void RemoveForCut() override;
 	virtual bool SupportsCopy() const override { return true; }
 	virtual bool TestCanCopyWithMessage(FText& OutMessage) const override;
-	virtual void Copy() const override;
+	virtual void Copy(UNiagaraClipboardContent* ClipboardContent) const override;
 	virtual bool SupportsPaste() const override { return true; }
-	virtual bool TestCanPasteWithMessage(FText& OutMessage) const override;
-	virtual void Paste() override;
+	virtual bool TestCanPasteWithMessage(const UNiagaraClipboardContent* ClipboardContent, FText& OutMessage) const override;
+	virtual FText GetPasteTransactionText(const UNiagaraClipboardContent* ClipboardContent) const override;
+	virtual void Paste(const UNiagaraClipboardContent* ClipboardContent) override;
 
 	FText GetTooltipText(EValueMode InValueMode) const;
 
@@ -123,7 +126,7 @@ public:
 	void GetAvailableDynamicInputs(TArray<UNiagaraScript*>& AvailableDynamicInputs, bool bIncludeNonLibraryInputs = false);
 
 	/** Sets the dynamic input script for this input. */
-	void SetDynamicInput(UNiagaraScript* DynamicInput);
+	void SetDynamicInput(UNiagaraScript* DynamicInput, FString SuggestedName = FString());
 
 	/** Gets the expression providing the value for this input, if one is available. */
 	FText GetCustomExpressionText() const;
@@ -383,8 +386,6 @@ private:
 
 	FNiagaraVariable CreateRapidIterationVariable(const FName& InName);
 
-	void PrepareFunctionInputForNewValue();
-
 private:
 	/** The module function call which owns this input entry. NOTE: This input might not be an input to the module function
 		call, it may be an input to a dynamic input function call which is owned by the module. */
diff --git a/Engine/Plugins/FX/Niagara/Source/NiagaraEditor/Public/ViewModels/Stack/NiagaraStackGraphUtilities.h b/Engine/Plugins/FX/Niagara/Source/NiagaraEditor/Public/ViewModels/Stack/NiagaraStackGraphUtilities.h
--- a/Engine/Plugins/FX/Niagara/Source/NiagaraEditor/Public/ViewModels/Stack/NiagaraStackGraphUtilities.h
+++ b/Engine/Plugins/FX/Niagara/Source/NiagaraEditor/Public/ViewModels/Stack/NiagaraStackGraphUtilities.h
@@ -85,6 +85,8 @@ namespace FNiagaraStackGraphUtilities
 
 	void GetStackFunctionStaticSwitchPins(UNiagaraNodeFunctionCall& FunctionCallNode, TArray<UEdGraphPin*>& OutInputPins, TSet<UEdGraphPin*>& OutHiddenPins);
 
+	void GetStackFunctionOutputVariables(UNiagaraNodeFunctionCall& FunctionCallNode, FCompileConstantResolver ConstantResolver, TArray<FNiagaraVariable>& OutOutputVariables, TArray<FNiagaraVariable>& OutOutputVariablesWithOriginalAliasesIntact);
+
 	UNiagaraNodeParameterMapSet* GetStackFunctionOverrideNode(UNiagaraNodeFunctionCall& FunctionCallNode);
 
 	UNiagaraNodeParameterMapSet& GetOrCreateStackFunctionOverrideNode(UNiagaraNodeFunctionCall& FunctionCallNode, const FGuid& PreferredOverrideNodeGuid = FGuid());
diff --git a/Engine/Plugins/FX/Niagara/Source/NiagaraEditor/Public/ViewModels/Stack/NiagaraStackItem.h b/Engine/Plugins/FX/Niagara/Source/NiagaraEditor/Public/ViewModels/Stack/NiagaraStackItem.h
--- a/Engine/Plugins/FX/Niagara/Source/NiagaraEditor/Public/ViewModels/Stack/NiagaraStackItem.h
+++ b/Engine/Plugins/FX/Niagara/Source/NiagaraEditor/Public/ViewModels/Stack/NiagaraStackItem.h
@@ -9,6 +9,7 @@
 
 class UNiagaraStackItemFooter;
 class UNiagaraNode;
+class UNiagaraClipboardContent;
 
 UCLASS()
 class NIAGARAEDITOR_API UNiagaraStackItem : public UNiagaraStackEntry
@@ -17,6 +18,7 @@ class NIAGARAEDITOR_API UNiagaraStackItem : public UNiagaraStackEntry
 
 public:
 	DECLARE_MULTICAST_DELEGATE(FOnModifiedGroupItems);
+	DECLARE_MULTICAST_DELEGATE_TwoParams(FOnRequestPaste, const UNiagaraClipboardContent* /* ClipboardContent */, int32 /* PasteIndex */);
 
 public:
 	void Initialize(FRequiredEntryData InRequiredEntryData, FString InStackEditorDataKey);
@@ -24,14 +26,11 @@ public:
 	virtual EStackRowStyle GetStackRowStyle() const override;
 
 	FOnModifiedGroupItems& OnModifiedGroupItems();
+	FOnRequestPaste& OnRequestPaste();
 
 	virtual bool SupportsChangeEnabled() const { return false; }
 	void SetIsEnabled(bool bInIsEnabled);
 
-	virtual bool SupportsDelete() const { return false; }
-	virtual bool TestCanDeleteWithMessage(FText& OutCanDeleteMessage) const { return false; }
-	void Delete();
-
 	virtual bool SupportsHighlights() const { return false; }
 	virtual const TArray<FNiagaraScriptHighlight>& GetHighlights() const;
 
@@ -46,7 +45,6 @@ protected:
 	virtual int32 GetChildIndentLevel() const override;
 
 	virtual void SetIsEnabledInternal(bool bInIsEnabled) { }
-	virtual void DeleteInternal() { }
 
 private:
 	bool FilterAdvancedChildren(const UNiagaraStackEntry& Child) const;
@@ -55,6 +53,7 @@ private:
 
 protected:
 	FOnModifiedGroupItems ModifiedGroupItemsDelegate;
+	FOnRequestPaste RequestPasteDelegate;
 
 private:
 	UPROPERTY()
diff --git a/Engine/Plugins/FX/Niagara/Source/NiagaraEditor/Public/ViewModels/Stack/NiagaraStackItemGroup.h b/Engine/Plugins/FX/Niagara/Source/NiagaraEditor/Public/ViewModels/Stack/NiagaraStackItemGroup.h
--- a/Engine/Plugins/FX/Niagara/Source/NiagaraEditor/Public/ViewModels/Stack/NiagaraStackItemGroup.h
+++ b/Engine/Plugins/FX/Niagara/Source/NiagaraEditor/Public/ViewModels/Stack/NiagaraStackItemGroup.h
@@ -23,9 +23,6 @@ public:
 	virtual EStackRowStyle GetStackRowStyle() const override;
 	virtual FText GetTooltipText() const override;
 
-	virtual bool CanDelete() const { return false; }
-	virtual bool Delete() { return false; }
-
 	virtual bool GetIsEnabled() const override;
 
 	INiagaraStackItemGroupAddUtilities* GetAddUtilities() const;
diff --git a/Engine/Plugins/FX/Niagara/Source/NiagaraEditor/Public/ViewModels/Stack/NiagaraStackModuleItem.h b/Engine/Plugins/FX/Niagara/Source/NiagaraEditor/Public/ViewModels/Stack/NiagaraStackModuleItem.h
--- a/Engine/Plugins/FX/Niagara/Source/NiagaraEditor/Public/ViewModels/Stack/NiagaraStackModuleItem.h
+++ b/Engine/Plugins/FX/Niagara/Source/NiagaraEditor/Public/ViewModels/Stack/NiagaraStackModuleItem.h
@@ -13,6 +13,7 @@ class UNiagaraStackModuleItemOutputCollection;
 class UNiagaraScript;
 class INiagaraStackItemGroupAddUtilities;
 struct FAssetData;
+class UNiagaraClipboardFunctionInput;
 
 UCLASS()
 class NIAGARAEDITOR_API UNiagaraStackModuleItem : public UNiagaraStackItem
@@ -39,9 +40,6 @@ public:
 	virtual bool SupportsChangeEnabled() const override { return true; }
 	virtual bool GetIsEnabled() const override;
 
-	virtual bool SupportsDelete() const override { return true; }
-	virtual bool TestCanDeleteWithMessage(FText& OutCanDeleteMessage) const override;
-
 	virtual bool SupportsHighlights() const override;
 	virtual const TArray<FNiagaraScriptHighlight>& GetHighlights() const override;
 
@@ -67,9 +65,30 @@ public:
 	/** Reassigns the function script for the module without resetting the inputs. */
 	void ReassignModuleScript(UNiagaraScript* ModuleScript);
 
+	void SetInputValuesFromClipboardFunctionInputs(const TArray<const UNiagaraClipboardFunctionInput*>& ClipboardFunctionInputs);
+
+	virtual bool SupportsCut() const override { return true; }
+	virtual bool TestCanCutWithMessage(FText& OutMessage) const override;
+	virtual FText GetCutTransactionText() const override;
+	virtual void CopyForCut(UNiagaraClipboardContent* ClipboardContent) const override;
+	virtual void RemoveForCut() override;
+
+	virtual bool SupportsCopy() const override { return true; }
+	virtual bool TestCanCopyWithMessage(FText& OutMessage) const override;
+	virtual void Copy(UNiagaraClipboardContent* ClipboardContent) const override;
+
+	virtual bool SupportsPaste() const override { return true; }
+	virtual bool TestCanPasteWithMessage(const UNiagaraClipboardContent* ClipboardContent, FText& OutMessage) const override;
+	virtual FText GetPasteTransactionText(const UNiagaraClipboardContent* ClipboardContent) const override;
+	virtual void Paste(const UNiagaraClipboardContent* ClipboardContent) override;
+
+	virtual bool SupportsDelete() const override { return true; }
+	virtual bool TestCanDeleteWithMessage(FText& OutCanDeleteMessage) const override;
+	virtual FText GetDeleteTransactionText() const override;
+	virtual void Delete() override;
+
 protected:
 	virtual void RefreshChildrenInternal(const TArray<UNiagaraStackEntry*>& CurrentChildren, TArray<UNiagaraStackEntry*>& NewChildren, TArray<FStackIssue>& NewIssues) override;
-	virtual void DeleteInternal() override;
 	virtual void SetIsEnabledInternal(bool bInIsEnabled) override;
 
 	virtual TOptional<FDropRequestResponse> CanDropInternal(const FDropRequest& DropRequest) override;
diff --git a/Engine/Plugins/FX/Niagara/Source/NiagaraEditor/Public/ViewModels/Stack/NiagaraStackRenderItemGroup.h b/Engine/Plugins/FX/Niagara/Source/NiagaraEditor/Public/ViewModels/Stack/NiagaraStackRenderItemGroup.h
--- a/Engine/Plugins/FX/Niagara/Source/NiagaraEditor/Public/ViewModels/Stack/NiagaraStackRenderItemGroup.h
+++ b/Engine/Plugins/FX/Niagara/Source/NiagaraEditor/Public/ViewModels/Stack/NiagaraStackRenderItemGroup.h
@@ -7,6 +7,7 @@
 
 class UNiagaraRendererProperties;
 class UNiagaraEmitter;
+class UNiagaraClipboardContent;
 
 UCLASS()
 class NIAGARAEDITOR_API UNiagaraStackRenderItemGroup : public UNiagaraStackItemGroup
@@ -16,13 +17,18 @@ class NIAGARAEDITOR_API UNiagaraStackRenderItemGroup : public UNiagaraStackItemG
 public:
 	void Initialize(FRequiredEntryData InRequiredEntryData);
 
+	virtual bool SupportsPaste() const override { return true; }
+	virtual bool TestCanPasteWithMessage(const UNiagaraClipboardContent* ClipboardContent, FText& OutMessage) const override;
+	virtual FText GetPasteTransactionText(const UNiagaraClipboardContent* ClipboardContent) const override;
+	virtual void Paste(const UNiagaraClipboardContent* ClipboardContent) override;
+
 protected:
 	virtual void RefreshChildrenInternal(const TArray<UNiagaraStackEntry*>& CurrentChildren, TArray<UNiagaraStackEntry*>& NewChildren, TArray<FStackIssue>& NewIssues) override;
-
+	virtual void FinalizeInternal() override;
 private:
 	void EmitterRenderersChanged();
 
-	virtual void FinalizeInternal() override;
+	void ChildRequestPaste(const UNiagaraClipboardContent* ClipboardContent, int32 PasteIndex);
 
 private:
 	TSharedPtr<INiagaraStackItemGroupAddUtilities> AddUtilities;
diff --git a/Engine/Plugins/FX/Niagara/Source/NiagaraEditor/Public/ViewModels/Stack/NiagaraStackRendererItem.h b/Engine/Plugins/FX/Niagara/Source/NiagaraEditor/Public/ViewModels/Stack/NiagaraStackRendererItem.h
--- a/Engine/Plugins/FX/Niagara/Source/NiagaraEditor/Public/ViewModels/Stack/NiagaraStackRendererItem.h
+++ b/Engine/Plugins/FX/Niagara/Source/NiagaraEditor/Public/ViewModels/Stack/NiagaraStackRendererItem.h
@@ -24,8 +24,25 @@ public:
 
 	virtual FText GetDisplayName() const override;
 
+	virtual bool SupportsCut() const override { return true; }
+	virtual bool TestCanCutWithMessage(FText& OutMessage) const override;
+	virtual FText GetCutTransactionText() const override;
+	virtual void CopyForCut(UNiagaraClipboardContent* ClipboardContent) const override;
+	virtual void RemoveForCut() override;
+
+	virtual bool SupportsCopy() const override { return true; }
+	virtual bool TestCanCopyWithMessage(FText& OutMessage) const override;
+	virtual void Copy(UNiagaraClipboardContent* ClipboardContent) const override;
+
+	virtual bool SupportsPaste() const { return true; }
+	virtual bool TestCanPasteWithMessage(const UNiagaraClipboardContent* ClipboardContent, FText& OutMessage) const override;
+	virtual FText GetPasteTransactionText(const UNiagaraClipboardContent* ClipboardContent) const override;
+	virtual void Paste(const UNiagaraClipboardContent* ClipboardContent) override;
+
 	virtual bool SupportsDelete() const override { return true; }
 	virtual bool TestCanDeleteWithMessage(FText& OutCanDeleteMessage) const;
+	virtual FText GetDeleteTransactionText() const override;
+	virtual void Delete();
 
 	bool HasBaseRenderer() const;
 
@@ -48,7 +65,6 @@ protected:
 	virtual void RefreshChildrenInternal(const TArray<UNiagaraStackEntry*>& CurrentChildren, TArray<UNiagaraStackEntry*>& NewChildren, TArray<FStackIssue>& NewIssues) override;
 
 	virtual void SetIsEnabledInternal(bool bInIsEnabled) override;
-	virtual void DeleteInternal() override;
 
 private:
 	void RendererChanged();
diff --git a/Engine/Plugins/FX/Niagara/Source/NiagaraEditor/Public/ViewModels/Stack/NiagaraStackScriptItemGroup.h b/Engine/Plugins/FX/Niagara/Source/NiagaraEditor/Public/ViewModels/Stack/NiagaraStackScriptItemGroup.h
--- a/Engine/Plugins/FX/Niagara/Source/NiagaraEditor/Public/ViewModels/Stack/NiagaraStackScriptItemGroup.h
+++ b/Engine/Plugins/FX/Niagara/Source/NiagaraEditor/Public/ViewModels/Stack/NiagaraStackScriptItemGroup.h
@@ -31,6 +31,11 @@ public:
 	FGuid GetScriptUsageId() const { return ScriptUsageId; }
 	UNiagaraNodeOutput* GetScriptOutputNode() const;
 
+	virtual bool SupportsPaste() const override { return true; }
+	virtual bool TestCanPasteWithMessage(const UNiagaraClipboardContent* ClipboardContent, FText& OutMessage) const override;
+	virtual FText GetPasteTransactionText(const UNiagaraClipboardContent* ClipboardContent) const override;
+	virtual void Paste(const UNiagaraClipboardContent* ClipboardContent) override;
+
 protected:
 	virtual void RefreshChildrenInternal(const TArray<UNiagaraStackEntry*>& CurrentChildren, TArray<UNiagaraStackEntry*>& NewChildren, TArray<FStackIssue>& NewIssues) override;
 
@@ -49,6 +54,8 @@ private:
 
 	void ChildModifiedGroupItems();
 
+	void ChildRequestPaste(const UNiagaraClipboardContent* ClipboardContent, int32 PasteIndex);
+
 	void OnScriptGraphChanged(const struct FEdGraphEditAction& InAction);
 
 	TOptional<FDropRequestResponse> CanDropOnTarget(const UNiagaraStackEntry& TargetEntry, const FDropRequest& DropRequest);
@@ -71,6 +78,9 @@ protected:
 	TWeakPtr<FNiagaraScriptViewModel> ScriptViewModel;
 	void RefreshIssues(TArray<FStackIssue>& NewIssues);
 
+private:
+	void PasteModules(const UNiagaraClipboardContent* ClipboardContent, int32 PasteIndex);
+
 private:
 	TSharedPtr<FScriptItemGroupAddUtilities> AddUtilities;
 
diff --git a/Engine/Plugins/FX/Niagara/Source/NiagaraEditorWidgets/Private/NiagaraEditorWidgetsStyle.cpp b/Engine/Plugins/FX/Niagara/Source/NiagaraEditorWidgets/Private/NiagaraEditorWidgetsStyle.cpp
--- a/Engine/Plugins/FX/Niagara/Source/NiagaraEditorWidgets/Private/NiagaraEditorWidgetsStyle.cpp
+++ b/Engine/Plugins/FX/Niagara/Source/NiagaraEditorWidgets/Private/NiagaraEditorWidgetsStyle.cpp
@@ -127,11 +127,13 @@ TSharedRef< FSlateStyleSet > FNiagaraEditorWidgetsStyle::Create()
 
 	const FTableRowStyle& NormalTableRowStyle = FEditorStyle::Get().GetWidgetStyle<FTableRowStyle>("TableView.Row");
 
+	FSlateBrush StackRowSelectionBrush = BOX_PLUGIN_BRUSH("Icons/StackSelectionBorder", FMargin(2.0f / 8.0f), GetDefault<UEditorStyleSettings>()->SelectionColor);
+	FSlateBrush StackRowSubduedSelectionBrush = BOX_PLUGIN_BRUSH("Icons/StackSelectionBorder", FMargin(2.0f / 8.0f), GetDefault<UEditorStyleSettings>()->GetSubduedSelectionColor());
 	Style->Set("NiagaraEditor.Stack.TableViewRow", FTableRowStyle(NormalTableRowStyle)
-		.SetActiveBrush(*FEditorStyle::GetBrush("NoBrush"))
-		.SetActiveHoveredBrush(*FEditorStyle::GetBrush("NoBrush"))
-		.SetInactiveBrush(*FEditorStyle::GetBrush("NoBrush"))
-		.SetInactiveHoveredBrush(*FEditorStyle::GetBrush("NoBrush")));
+		.SetActiveBrush(StackRowSelectionBrush)
+		.SetActiveHoveredBrush(StackRowSelectionBrush)
+		.SetInactiveBrush(StackRowSubduedSelectionBrush)
+		.SetInactiveHoveredBrush(StackRowSelectionBrush));
 
 	Style->Set("NiagaraEditor.SystemOverview.TableViewRow", FTableRowStyle(NormalTableRowStyle)
 		.SetInactiveBrush(IMAGE_CORE_BRUSH("Common/Selection", Icon8x8, GetDefault<UEditorStyleSettings>()->GetSubduedSelectionColor())));
@@ -201,7 +203,7 @@ TSharedRef< FSlateStyleSet > FNiagaraEditorWidgetsStyle::Create()
 		.SetHandleHighlightBrush(IMAGE_CORE_BRUSH("Common/SplitterHandleHighlight", Icon8x8, FLinearColor::White))
 	);
 
-	Style->Set("NiagaraEditor.Stack.SearchHighlightColor", FLinearColor(FColor::Orange));
+	Style->Set("NiagaraEditor.Stack.SearchHighlightColor", FEditorStyle::GetColor("TextBlock.HighlighColor"));
 	Style->Set("NiagaraEditor.Stack.SearchResult", new BOX_PLUGIN_BRUSH("Icons/SearchResultBorder", FMargin(1.f/8.f)));
 
 	Style->Set("NiagaraEditor.Stack.AddButton", FButtonStyle()
diff --git a/Engine/Plugins/FX/Niagara/Source/NiagaraEditorWidgets/Private/NiagaraEditorWidgetsUtilities.cpp b/Engine/Plugins/FX/Niagara/Source/NiagaraEditorWidgets/Private/NiagaraEditorWidgetsUtilities.cpp
--- a/Engine/Plugins/FX/Niagara/Source/NiagaraEditorWidgets/Private/NiagaraEditorWidgetsUtilities.cpp
+++ b/Engine/Plugins/FX/Niagara/Source/NiagaraEditorWidgets/Private/NiagaraEditorWidgetsUtilities.cpp
@@ -6,6 +6,9 @@
 #include "ViewModels/Stack/NiagaraStackModuleItem.h"
 #include "Stack/SNiagaraStackItemGroupAddMenu.h"
 #include "NiagaraEditorWidgetsStyle.h"
+#include "NiagaraEditorCommon.h"
+#include "NiagaraClipboard.h"
+#include "NiagaraEditorModule.h"
 
 #include "Framework/Application/SlateApplication.h"
 #include "Framework/MultiBox/MultiBoxBuilder.h"
@@ -161,73 +164,6 @@ bool FNiagaraStackEditorWidgetsUtilities::AddStackEntryAssetContextMenuActions(F
 	return false;
 }
 
-void CopyEntry(TWeakObjectPtr<UNiagaraStackEntry> StackEntryWeak)
-{
-	UNiagaraStackEntry* StackEntry = StackEntryWeak.Get();
-	if (StackEntry != nullptr)
-	{
-		StackEntry->Copy();
-	}
-}
-
-void PasteEntry(TWeakObjectPtr<UNiagaraStackEntry> StackEntryWeak)
-{
-	UNiagaraStackEntry* StackEntry = StackEntryWeak.Get();
-	if (StackEntry != nullptr)
-	{
-		StackEntry->Paste();
-	}
-}
-
-bool FNiagaraStackEditorWidgetsUtilities::AddStackEntryContextMenuActions(FMenuBuilder& MenuBuilder, UNiagaraStackEntry& StackEntry)
-{
-	if (StackEntry.SupportsCut() || StackEntry.SupportsCopy() || StackEntry.SupportsPaste())
-	{
-		MenuBuilder.BeginSection("EntryEdit", LOCTEXT("EntryEditActions", "Edit"));
-		{
-			if (StackEntry.SupportsCut())
-			{
-				FText CutMessage;
-				bool bCanCut = StackEntry.TestCanCutWithMessage(CutMessage);
-				MenuBuilder.AddMenuEntry(
-					LOCTEXT("CutAction", "Cut"),
-					CutMessage,
-					FSlateIcon(),
-					FUIAction(
-						FExecuteAction::CreateUObject(&StackEntry, &UNiagaraStackEntry::Cut),
-						FCanExecuteAction::CreateLambda([bCanCut]() { return bCanCut; })));
-			}
-			if (StackEntry.SupportsCopy())
-			{
-				FText CopyMessage;
-				bool bCanCopy = StackEntry.TestCanCopyWithMessage(CopyMessage);
-				MenuBuilder.AddMenuEntry(
-					LOCTEXT("CopyAction", "Copy"),
-					CopyMessage,
-					FSlateIcon(),
-					FUIAction(
-						FExecuteAction::CreateStatic(&CopyEntry, TWeakObjectPtr<UNiagaraStackEntry>(&StackEntry)),
-						FCanExecuteAction::CreateLambda([bCanCopy]() { return bCanCopy; })));
-			}
-			if (StackEntry.SupportsPaste())
-			{
-				FText PasteMessage;
-				bool bCanPaste = StackEntry.TestCanPasteWithMessage(PasteMessage);
-				MenuBuilder.AddMenuEntry(
-					LOCTEXT("PasteAction", "Paste"),
-					PasteMessage,
-					FSlateIcon(),
-					FUIAction(
-						FExecuteAction::CreateStatic(&PasteEntry, TWeakObjectPtr<UNiagaraStackEntry>(&StackEntry)),
-						FCanExecuteAction::CreateLambda([bCanPaste]() { return bCanPaste; })));
-			}
-		}
-		MenuBuilder.EndSection();
-		return true;
-	}
-	return false;
-}
-
 void DeleteItem(TWeakObjectPtr<UNiagaraStackItem> StackItemWeak)
 {
 	UNiagaraStackItem* StackItem = StackItemWeak.Get();
@@ -249,7 +185,7 @@ void ToggleEnabledState(TWeakObjectPtr<UNiagaraStackItem> StackItemWeak)
 
 bool FNiagaraStackEditorWidgetsUtilities::AddStackItemContextMenuActions(FMenuBuilder& MenuBuilder, UNiagaraStackItem& StackItem)
 {
-	if (StackItem.SupportsDelete() || StackItem.SupportsChangeEnabled())
+	if (StackItem.SupportsChangeEnabled())
 	{
 		MenuBuilder.BeginSection("ItemActions", LOCTEXT("ItemActions", "Item Actions"));
 		{
@@ -266,20 +202,6 @@ bool FNiagaraStackEditorWidgetsUtilities::AddStackItemContextMenuActions(FMenuBu
 					NAME_None,
 					EUserInterfaceActionType::Check);
 			}
-
-			if (StackItem.SupportsDelete())
-			{
-				FText CanDeleteMessage;
-				bool bCanDelete = StackItem.TestCanDeleteWithMessage(CanDeleteMessage);
-				MenuBuilder.AddMenuEntry(
-					LOCTEXT("DeleteModule", "Delete Item"),
-					CanDeleteMessage,
-					FSlateIcon(),
-					FUIAction(
-						FExecuteAction::CreateStatic(&DeleteItem, TWeakObjectPtr<UNiagaraStackItem>(&StackItem)),
-						FCanExecuteAction::CreateLambda([=]() { return bCanDelete; })));
-
-			}
 		}
 		MenuBuilder.EndSection();
 		return true;
diff --git a/Engine/Plugins/FX/Niagara/Source/NiagaraEditorWidgets/Private/NiagaraEditorWidgetsUtilities.h b/Engine/Plugins/FX/Niagara/Source/NiagaraEditorWidgets/Private/NiagaraEditorWidgetsUtilities.h
--- a/Engine/Plugins/FX/Niagara/Source/NiagaraEditorWidgets/Private/NiagaraEditorWidgetsUtilities.h
+++ b/Engine/Plugins/FX/Niagara/Source/NiagaraEditorWidgets/Private/NiagaraEditorWidgetsUtilities.h
@@ -22,8 +22,6 @@ namespace FNiagaraStackEditorWidgetsUtilities
 	
 	bool AddStackEntryAssetContextMenuActions(FMenuBuilder& MenuBuilder, UNiagaraStackEntry& StackEntry);
 
-	bool AddStackEntryContextMenuActions(FMenuBuilder& MenuBuilder, UNiagaraStackEntry& StackEntry);
-
 	bool AddStackItemContextMenuActions(FMenuBuilder& MenuBuilder, UNiagaraStackItem& StackItem);
 
 	bool AddStackModuleItemContextMenuActions(FMenuBuilder& MenuBuilder, UNiagaraStackModuleItem& StackItem, TSharedRef<SWidget> TargetWidget);
diff --git a/Engine/Plugins/FX/Niagara/Source/NiagaraEditorWidgets/Private/NiagaraStackCommandContext.cpp b/Engine/Plugins/FX/Niagara/Source/NiagaraEditorWidgets/Private/NiagaraStackCommandContext.cpp
new file mode 100644
--- /dev/null
+++ b/Engine/Plugins/FX/Niagara/Source/NiagaraEditorWidgets/Private/NiagaraStackCommandContext.cpp
@@ -0,0 +1,203 @@
+// Copyright 1998-2019 Epic Games, Inc. All Rights Reserved.
+
+#include "NiagaraStackCommandContext.h"
+#include "NiagaraEditorCommon.h"
+#include "ViewModels/Stack/NiagaraStackClipboardUtilities.h"
+#include "ViewModels/Stack/NiagaraStackEntry.h"
+
+#include "Framework/Notifications/NotificationManager.h"
+#include "Logging/LogMacros.h"
+#include "Widgets/Notifications/SNotificationList.h"
+#include "Framework/Commands/GenericCommands.h"
+#include "Framework/Commands/UICommandList.h"
+#include "Framework/MultiBox/MultiBoxBuilder.h"
+
+#define LOCTEXT_NAMESPACE "NiagaraStackCommandContext"
+
+void WarnWithToastAndLog(FText WarningMessage)
+{
+	FText IncompleteDeleteMessage = LOCTEXT("DeleteIncompleteMessage", "Not all items could be deleted because they either\ndon't support being deleted or they are inherited.");
+	FNotificationInfo WarningNotification(WarningMessage);
+	WarningNotification.ExpireDuration = 5.0f;
+	WarningNotification.bFireAndForget = true;
+	WarningNotification.bUseLargeFont = false;
+	WarningNotification.Image = FCoreStyle::Get().GetBrush(TEXT("MessageLog.Warning"));
+	FSlateNotificationManager::Get().AddNotification(WarningNotification);
+	UE_LOG(LogNiagaraEditor, Warning, TEXT("%s"), *WarningMessage.ToString());
+}
+
+FNiagaraStackCommandContext::FNiagaraStackCommandContext()
+	: Commands(MakeShared<FUICommandList>())
+	, bCommandsAreSetup(false)
+	, bProcessingCommandBindings(false)
+{
+}
+
+void FNiagaraStackCommandContext::SetSelectedEntries(const TArray<UNiagaraStackEntry*> InSelectedEntries)
+{
+	SelectedEntries = InSelectedEntries;
+}
+
+TSharedRef<FUICommandList> FNiagaraStackCommandContext::GetCommands()
+{
+	if (bCommandsAreSetup == false)
+	{
+		SetupCommands();
+		bCommandsAreSetup = true;
+	}
+	return Commands;
+}
+
+bool FNiagaraStackCommandContext::ProcessCommandBindings(const FKeyEvent& InKeyEvent)
+{
+	TGuardValue<bool> ProcessingGuard(bProcessingCommandBindings, true);
+	return GetCommands()->ProcessCommandBindings(InKeyEvent);
+}
+
+bool FNiagaraStackCommandContext::AddEditMenuItems(FMenuBuilder& MenuBuilder)
+{
+	bool bSupportsCut = SelectedEntries.ContainsByPredicate([](const UNiagaraStackEntry* SelectedEntry) { return SelectedEntry->SupportsCut(); });
+	bool bSupportsCopy = SelectedEntries.ContainsByPredicate([](const UNiagaraStackEntry* SelectedEntry) { return SelectedEntry->SupportsCopy(); });
+	bool bSupportsPaste = SelectedEntries.ContainsByPredicate([](const UNiagaraStackEntry* SelectedEntry) { return SelectedEntry->SupportsPaste(); });
+	bool bSupportsDelete = SelectedEntries.ContainsByPredicate([](const UNiagaraStackEntry* SelectedEntry) { return SelectedEntry->SupportsDelete(); });
+	if (bSupportsCut || bSupportsCopy || bSupportsPaste || bSupportsDelete)
+	{
+		MenuBuilder.BeginSection("EntryEdit", LOCTEXT("EntryEditActions", "Edit"));
+		{
+			if (bSupportsCut)
+			{
+				TAttribute<FText> CanCutToolTip = TAttribute<FText>::Create(TAttribute<FText>::FGetter::CreateSP(this, &FNiagaraStackCommandContext::GetCanCutSelectedEntriesToolTip));
+				MenuBuilder.AddMenuEntry(FGenericCommands::Get().Cut, NAME_None, TAttribute<FText>(), CanCutToolTip);
+			}
+			if (bSupportsCopy)
+			{
+				TAttribute<FText> CanCopyToolTip = TAttribute<FText>::Create(TAttribute<FText>::FGetter::CreateSP(this, &FNiagaraStackCommandContext::GetCanCopySelectedEntriesToolTip));
+				MenuBuilder.AddMenuEntry(FGenericCommands::Get().Copy, NAME_None, TAttribute<FText>(), CanCopyToolTip);
+			}
+			if (bSupportsPaste)
+			{
+				TAttribute<FText> CanPasteToolTip = TAttribute<FText>::Create(TAttribute<FText>::FGetter::CreateSP(this, &FNiagaraStackCommandContext::GetCanPasteSelectedEntriesToolTip));
+				MenuBuilder.AddMenuEntry(FGenericCommands::Get().Paste, NAME_None, TAttribute<FText>(), CanPasteToolTip);
+			}
+			if (bSupportsDelete)
+			{
+				TAttribute<FText> CanDeleteToolTip = TAttribute<FText>::Create(TAttribute<FText>::FGetter::CreateSP(this, &FNiagaraStackCommandContext::GetCanDeleteSelectedEntriesToolTip));
+				MenuBuilder.AddMenuEntry(FGenericCommands::Get().Delete, NAME_None, TAttribute<FText>(), CanDeleteToolTip);
+			}
+		}
+		MenuBuilder.EndSection();
+		return true;
+	}
+	return false;
+}
+
+void FNiagaraStackCommandContext::SetupCommands()
+{
+	Commands->MapAction(FGenericCommands::Get().Cut, FUIAction(
+		FExecuteAction::CreateSP(this, &FNiagaraStackCommandContext::CutSelectedEntries),
+		FCanExecuteAction::CreateSP(this, &FNiagaraStackCommandContext::CanCutSelectedEntries)));
+	Commands->MapAction(FGenericCommands::Get().Copy, FUIAction(
+		FExecuteAction::CreateSP(this, &FNiagaraStackCommandContext::CopySelectedEntries),
+		FCanExecuteAction::CreateSP(this, &FNiagaraStackCommandContext::CanCopySelectedEntries)));
+	Commands->MapAction(FGenericCommands::Get().Paste, FUIAction(
+		FExecuteAction::CreateSP(this, &FNiagaraStackCommandContext::PasteSelectedEntries),
+		FCanExecuteAction::CreateSP(this, &FNiagaraStackCommandContext::CanPasteSelectedEntries)));
+	Commands->MapAction(FGenericCommands::Get().Delete, FUIAction(
+		FExecuteAction::CreateSP(this, &FNiagaraStackCommandContext::DeleteSelectedEntries),
+		FCanExecuteAction::CreateSP(this, &FNiagaraStackCommandContext::CanDeleteSelectedEntries)));
+}
+
+bool FNiagaraStackCommandContext::CanCutSelectedEntries() const
+{
+	FText CanCutMessage;
+	bool bCanCut = FNiagaraStackClipboardUtilities::TestCanCutSelectionWithMessage(SelectedEntries, CanCutMessage);
+	if (bProcessingCommandBindings && bCanCut == false && CanCutMessage.IsEmptyOrWhitespace() == false)
+	{
+		WarnWithToastAndLog(CanCutMessage);
+	}
+	return bCanCut;
+}
+
+FText FNiagaraStackCommandContext::GetCanCutSelectedEntriesToolTip() const
+{
+	FText CanCutMessage;
+	FNiagaraStackClipboardUtilities::TestCanCutSelectionWithMessage(SelectedEntries, CanCutMessage);
+	return CanCutMessage;
+}
+
+void FNiagaraStackCommandContext::CutSelectedEntries() const
+{
+	FNiagaraStackClipboardUtilities::CutSelection(SelectedEntries);
+}
+
+bool FNiagaraStackCommandContext::CanCopySelectedEntries() const
+{
+	FText CanCopyMessage;
+	bool bCanCopy = FNiagaraStackClipboardUtilities::TestCanCopySelectionWithMessage(SelectedEntries, CanCopyMessage);
+	if (bProcessingCommandBindings && bCanCopy == false && CanCopyMessage.IsEmptyOrWhitespace() == false)
+	{
+		WarnWithToastAndLog(CanCopyMessage);
+	}
+	return bCanCopy;
+}
+
+FText FNiagaraStackCommandContext::GetCanCopySelectedEntriesToolTip() const
+{
+	FText CanCopyMessage;
+	FNiagaraStackClipboardUtilities::TestCanCopySelectionWithMessage(SelectedEntries, CanCopyMessage);
+	return CanCopyMessage;
+}
+
+void FNiagaraStackCommandContext::CopySelectedEntries() const
+{
+	FNiagaraStackClipboardUtilities::CopySelection(SelectedEntries);
+}
+
+bool FNiagaraStackCommandContext::CanPasteSelectedEntries() const
+{
+	FText CanPasteMessage;
+	bool bCanPaste = FNiagaraStackClipboardUtilities::TestCanPasteSelectionWithMessage(SelectedEntries, CanPasteMessage);
+	if (bProcessingCommandBindings && bCanPaste == false && CanPasteMessage.IsEmptyOrWhitespace() == false)
+	{
+		WarnWithToastAndLog(CanPasteMessage);
+	}
+	return bCanPaste;
+}
+
+FText FNiagaraStackCommandContext::GetCanPasteSelectedEntriesToolTip() const
+{
+	FText CanPasteMessage;
+	FNiagaraStackClipboardUtilities::TestCanPasteSelectionWithMessage(SelectedEntries, CanPasteMessage);
+	return CanPasteMessage;
+}
+
+void FNiagaraStackCommandContext::PasteSelectedEntries() const
+{
+	FNiagaraStackClipboardUtilities::PasteSelection(SelectedEntries);
+}
+
+bool FNiagaraStackCommandContext::CanDeleteSelectedEntries() const
+{
+	FText CanDeleteMessage;
+	bool bCanDelete = FNiagaraStackClipboardUtilities::TestCanDeleteSelectionWithMessage(SelectedEntries, CanDeleteMessage);
+	if (bProcessingCommandBindings && bCanDelete == false && CanDeleteMessage.IsEmptyOrWhitespace() == false)
+	{
+		WarnWithToastAndLog(CanDeleteMessage);
+	}
+	return bCanDelete;
+}
+
+FText FNiagaraStackCommandContext::GetCanDeleteSelectedEntriesToolTip() const
+{
+	FText CanDeleteMessage;
+	FNiagaraStackClipboardUtilities::TestCanDeleteSelectionWithMessage(SelectedEntries, CanDeleteMessage);
+	return CanDeleteMessage;
+}
+
+void FNiagaraStackCommandContext::DeleteSelectedEntries() const
+{
+	FNiagaraStackClipboardUtilities::DeleteSelection(SelectedEntries);
+}
+
+#undef LOCTEXT_NAMESPACE
+
diff --git a/Engine/Plugins/FX/Niagara/Source/NiagaraEditorWidgets/Private/NiagaraStackCommandContext.h b/Engine/Plugins/FX/Niagara/Source/NiagaraEditorWidgets/Private/NiagaraStackCommandContext.h
new file mode 100644
--- /dev/null
+++ b/Engine/Plugins/FX/Niagara/Source/NiagaraEditorWidgets/Private/NiagaraStackCommandContext.h
@@ -0,0 +1,59 @@
+// Copyright 1998-2019 Epic Games, Inc. All Rights Reserved.
+
+#pragma once
+
+#include "Templates/SharedPointer.h"
+
+class UNiagaraStackEntry;
+class FUICommandList;
+class FMenuBuilder;
+
+class FNiagaraStackCommandContext : public TSharedFromThis<FNiagaraStackCommandContext>
+{
+public:
+	FNiagaraStackCommandContext();
+
+	void SetSelectedEntries(const TArray<UNiagaraStackEntry*> InSelectedEntries);
+
+	TSharedRef<FUICommandList> GetCommands();
+
+	bool ProcessCommandBindings(const FKeyEvent& InKeyEvent);
+
+	bool AddEditMenuItems(FMenuBuilder& MenuBuilder);
+
+private:
+	void SetupCommands();
+
+	bool CanCutSelectedEntries() const;
+
+	FText GetCanCutSelectedEntriesToolTip() const;
+
+	void CutSelectedEntries() const;
+
+	bool CanCopySelectedEntries() const;
+
+	FText GetCanCopySelectedEntriesToolTip() const;
+
+	void CopySelectedEntries() const;
+
+	bool CanPasteSelectedEntries() const;
+
+	FText GetCanPasteSelectedEntriesToolTip() const;
+
+	void PasteSelectedEntries() const;
+
+	bool CanDeleteSelectedEntries() const;
+
+	FText GetCanDeleteSelectedEntriesToolTip() const;
+
+	void DeleteSelectedEntries() const;
+
+private:
+	TSharedRef<FUICommandList> Commands;
+
+	bool bCommandsAreSetup;
+
+	bool bProcessingCommandBindings;
+
+	TArray<UNiagaraStackEntry*> SelectedEntries;
+};
\ No newline at end of file
diff --git a/Engine/Plugins/FX/Niagara/Source/NiagaraEditorWidgets/Private/SNiagaraOverviewStack.cpp b/Engine/Plugins/FX/Niagara/Source/NiagaraEditorWidgets/Private/SNiagaraOverviewStack.cpp
--- a/Engine/Plugins/FX/Niagara/Source/NiagaraEditorWidgets/Private/SNiagaraOverviewStack.cpp
+++ b/Engine/Plugins/FX/Niagara/Source/NiagaraEditorWidgets/Private/SNiagaraOverviewStack.cpp
@@ -18,22 +18,18 @@
 #include "SNiagaraStack.h"
 #include "Stack/SNiagaraStackItemGroupAddButton.h"
 #include "Stack/SNiagaraStackIssueIcon.h"
-#include "NiagaraEditorCommon.h"
+#include "NiagaraStackCommandContext.h"
 
 #include "Widgets/Layout/SScrollBox.h"
 #include "Widgets/Layout/SGridPanel.h"
 #include "Widgets/Images/SImage.h"
 #include "Widgets/Input/SButton.h"
-#include "Widgets/Notifications/SNotificationList.h"
 #include "Widgets/SToolTip.h"
 #include "EditorStyleSet.h"
 #include "Styling/CoreStyle.h"
 #include "Framework/Application/SlateApplication.h"
 #include "Framework/MultiBox/MultiBoxBuilder.h"
 #include "Framework/Commands/GenericCommands.h"
-#include "Framework/Notifications/NotificationManager.h"
-#include "ScopedTransaction.h"
-#include "Logging/LogMacros.h"
 
 #define LOCTEXT_NAMESPACE "NiagaraOverviewStack"
 
@@ -48,10 +44,11 @@ class SNiagaraSystemOverviewEntryListRow : public STableRow<UNiagaraStackEntry*>
 		SLATE_DEFAULT_SLOT(FArguments, Content);
 	SLATE_END_ARGS();
 
-	void Construct(const FArguments& InArgs, UNiagaraStackViewModel* InStackViewModel, UNiagaraStackEntry* InStackEntry, const TSharedRef<STableViewBase>& InOwnerTableView)
+	void Construct(const FArguments& InArgs, UNiagaraStackViewModel* InStackViewModel, UNiagaraStackEntry* InStackEntry, TSharedRef<FNiagaraStackCommandContext> InStackCommandContext, const TSharedRef<STableViewBase>& InOwnerTableView)
 	{
 		StackViewModel = InStackViewModel;
 		StackEntry = InStackEntry;
+		StackCommandContext = InStackCommandContext;
 		IssueIconVisibility = InArgs._IssueIconVisibility;
 		FSlateColor IconColor = FNiagaraEditorWidgetsStyle::Get().GetColor(FNiagaraStackEditorWidgetsUtilities::GetColorNameForExecutionCategory(StackEntry->GetExecutionCategoryName()));
 
@@ -118,12 +115,33 @@ class SNiagaraSystemOverviewEntryListRow : public STableRow<UNiagaraStackEntry*>
 		InOwnerTableView);
 	}
 
-	// Need to use the mouse down event here since the graph eats the mouse up to show its context menu.
 	virtual FReply OnMouseButtonDown(const FGeometry& MyGeometry, const FPointerEvent& MouseEvent) override
 	{
 		if (MouseEvent.GetEffectingButton() == EKeys::RightMouseButton)
 		{
-			FMenuBuilder MenuBuilder(true, nullptr);
+			return FReply::Handled()
+				.SetUserFocus(OwnerTablePtr.Pin()->AsWidget(), EFocusCause::Mouse)
+				.CaptureMouse(SharedThis(this));
+		}
+		return STableRow<UNiagaraStackEntry*>::OnMouseButtonDown(MyGeometry, MouseEvent);
+	}
+
+	virtual FReply OnMouseButtonUp(const FGeometry& MyGeometry, const FPointerEvent& MouseEvent) override
+	{
+		if (MouseEvent.GetEffectingButton() == EKeys::RightMouseButton)
+		{
+			// Add the item represented by this row to the selection.
+			UNiagaraSystemSelectionViewModel* Selection = StackEntry->GetSystemViewModel()->GetSelectionViewModel();
+			if (Selection->ContainsEntry(StackEntry) == false)
+			{
+				TArray<UNiagaraStackEntry*> SelectedEntries;
+				TArray<UNiagaraStackEntry*> DeselectedEntries;
+				SelectedEntries.Add(StackEntry);
+				bool bClearSelection = true;
+				Selection->UpdateSelectedEntries(SelectedEntries, DeselectedEntries, bClearSelection);
+			}
+
+			FMenuBuilder MenuBuilder(true, StackCommandContext->GetCommands());
 			bool bMenuItemsAdded = false;
 
 			if (StackEntry->IsA<UNiagaraStackModuleItem>())
@@ -137,16 +155,17 @@ class SNiagaraSystemOverviewEntryListRow : public STableRow<UNiagaraStackEntry*>
 			}
 
 			bMenuItemsAdded |= FNiagaraStackEditorWidgetsUtilities::AddStackEntryAssetContextMenuActions(MenuBuilder, *StackEntry);
-			bMenuItemsAdded |= FNiagaraStackEditorWidgetsUtilities::AddStackEntryContextMenuActions(MenuBuilder, *StackEntry);
+			bMenuItemsAdded |= StackCommandContext->AddEditMenuItems(MenuBuilder);
 		
 			if (bMenuItemsAdded)
 			{
 				FWidgetPath WidgetPath = MouseEvent.GetEventPath() != nullptr ? *MouseEvent.GetEventPath() : FWidgetPath();
 				FSlateApplication::Get().PushMenu(AsShared(), WidgetPath, MenuBuilder.MakeWidget(), MouseEvent.GetScreenSpacePosition(), FPopupTransitionEffect(FPopupTransitionEffect::ContextMenu));
-				return FReply::Handled();
+				return FReply::Handled().ReleaseMouseCapture();
 			}
+			return STableRow<UNiagaraStackEntry*>::OnMouseButtonUp(MyGeometry, MouseEvent).ReleaseMouseCapture();
 		}
-		return STableRow<UNiagaraStackEntry*>::OnMouseButtonDown(MyGeometry, MouseEvent);
+		return STableRow<UNiagaraStackEntry*>::OnMouseButtonUp(MyGeometry, MouseEvent);
 	}
 
 private:
@@ -180,6 +199,7 @@ private:
 private:
 	UNiagaraStackViewModel* StackViewModel;
 	UNiagaraStackEntry* StackEntry;
+	TSharedPtr<FNiagaraStackCommandContext> StackCommandContext;
 	FLinearColor BackgroundColor;
 	FLinearColor DisabledBackgroundColor;
 	FLinearColor IsolatedBackgroundColor;
@@ -239,9 +259,7 @@ private:
 
 void SNiagaraOverviewStack::Construct(const FArguments& InArgs, UNiagaraStackViewModel& InStackViewModel, UNiagaraSystemSelectionViewModel& InOverviewSelectionViewModel)
 {
-	Commands = MakeShared<FUICommandList>();
-
-	SetupCommands();
+	StackCommandContext = MakeShared<FNiagaraStackCommandContext>();
 
 	bUpdatingOverviewSelectionFromStackSelection = false;
 	bUpdatingStackSelectionFromOverviewSelection = false;
@@ -287,7 +305,7 @@ bool SNiagaraOverviewStack::SupportsKeyboardFocus() const
 
 FReply SNiagaraOverviewStack::OnKeyDown(const FGeometry& MyGeometry, const FKeyEvent& InKeyEvent)
 {
-	if (Commands->ProcessCommandBindings(InKeyEvent))
+	if (StackCommandContext->ProcessCommandBindings(InKeyEvent))
 	{
 		return FReply::Handled();
 	}
@@ -299,61 +317,6 @@ void SNiagaraOverviewStack::Tick(const FGeometry& AllottedGeometry, const double
 	RefreshEntryList();
 }
 
-void SNiagaraOverviewStack::SetupCommands()
-{
-	Commands->MapAction(FGenericCommands::Get().Delete, FUIAction(
-		FExecuteAction::CreateSP(this, &SNiagaraOverviewStack::DeleteSelectedEntries)));
-}
-
-void SNiagaraOverviewStack::DeleteSelectedEntries()
-{
-	bool bIncompleteDelete = false;
-	TArray<UNiagaraStackItem*> ItemsToDelete;
-	TArray<UNiagaraStackEntry*> SelectedEntries;
-	OverviewSelectionViewModel->GetSelectedEntries(SelectedEntries);
-	for (UNiagaraStackEntry* SelectedEntry : SelectedEntries)
-	{
-		UNiagaraStackItem* SelectedItem = Cast<UNiagaraStackItem>(SelectedEntry);
-		if (SelectedItem != nullptr)
-		{
-			FText DeleteMessage;
-			if (SelectedItem->SupportsDelete() && SelectedItem->TestCanDeleteWithMessage(DeleteMessage))
-			{
-				ItemsToDelete.Add(SelectedItem);
-			}
-			else
-			{
-				bIncompleteDelete = true;
-			}
-		}
-		else
-		{
-			bIncompleteDelete = true;
-		}
-	}
-
-	if (ItemsToDelete.Num() > 0)
-	{
-		const FScopedTransaction Transaction(LOCTEXT("DeleteSelected", "Delete items from the system overview"));
-		for (UNiagaraStackItem* ItemToDelete : ItemsToDelete)
-		{
-			ItemToDelete->Delete();
-		}
-	}
-
-	if (bIncompleteDelete)
-	{
-		FText IncompleteDeleteMessage = LOCTEXT("DeleteIncompleteMessage", "Not all items could be deleted because they either\ndon't support being deleted or they are inherited.");
-		FNotificationInfo Warning(IncompleteDeleteMessage);
-		Warning.ExpireDuration = 5.0f;
-		Warning.bFireAndForget = true;
-		Warning.bUseLargeFont = false;
-		Warning.Image = FCoreStyle::Get().GetBrush(TEXT("MessageLog.Warning"));
-		FSlateNotificationManager::Get().AddNotification(Warning);
-		UE_LOG(LogNiagaraEditor, Warning, TEXT("%s"), *IncompleteDeleteMessage.ToString());
-	}
-}
-
 void SNiagaraOverviewStack::AddEntriesRecursive(UNiagaraStackEntry& EntryToAdd, TArray<UNiagaraStackEntry*>& EntryList, const TArray<UClass*>& AcceptableClasses, TArray<UNiagaraStackEntry*> ParentChain)
 {
 	if (AcceptableClasses.ContainsByPredicate([&] (UClass* Class) { return EntryToAdd.IsA(Class); }))
@@ -599,7 +562,7 @@ TSharedRef<ITableRow> SNiagaraOverviewStack::OnGenerateRowForEntry(UNiagaraStack
 	}
 
 
-	return SNew(SNiagaraSystemOverviewEntryListRow, StackViewModel, Item, OwnerTable)
+	return SNew(SNiagaraSystemOverviewEntryListRow, StackViewModel, Item, StackCommandContext.ToSharedRef(), OwnerTable)
 		.OnDragDetected(this, &SNiagaraOverviewStack::OnRowDragDetected, TWeakObjectPtr<UNiagaraStackEntry>(Item))
 		.OnDragLeave(this, &SNiagaraOverviewStack::OnRowDragLeave)
 		.OnCanAcceptDrop(this, &SNiagaraOverviewStack::OnRowCanAcceptDrop)
@@ -641,6 +604,8 @@ void SNiagaraOverviewStack::OnSelectionChanged(UNiagaraStackEntry* InNewSelectio
 		{
 			PreviousSelection.Add(SelectedEntry);
 		}
+
+		UpdateCommandContextSelection();
 	}
 }
 
@@ -683,9 +648,18 @@ void SNiagaraOverviewStack::SystemSelectionChanged()
 		{
 			EntryListView->SetItemSelection(EntryToSelect, true);
 		}
+
+		UpdateCommandContextSelection();
 	}
 }
 
+void SNiagaraOverviewStack::UpdateCommandContextSelection()
+{
+	TArray<UNiagaraStackEntry*> SelectedListViewStackEntries;
+	EntryListView->GetSelectedItems(SelectedListViewStackEntries);
+	StackCommandContext->SetSelectedEntries(SelectedListViewStackEntries);
+}
+
 FReply SNiagaraOverviewStack::OnRowDragDetected(const FGeometry& InGeometry, const FPointerEvent& InPointerEvent, TWeakObjectPtr<UNiagaraStackEntry> InStackEntryWeak)
 {
 	UNiagaraStackEntry* StackEntry = InStackEntryWeak.Get();
diff --git a/Engine/Plugins/FX/Niagara/Source/NiagaraEditorWidgets/Private/SNiagaraOverviewStack.h b/Engine/Plugins/FX/Niagara/Source/NiagaraEditorWidgets/Private/SNiagaraOverviewStack.h
--- a/Engine/Plugins/FX/Niagara/Source/NiagaraEditorWidgets/Private/SNiagaraOverviewStack.h
+++ b/Engine/Plugins/FX/Niagara/Source/NiagaraEditorWidgets/Private/SNiagaraOverviewStack.h
@@ -3,6 +3,7 @@
 #pragma once
 
 #include "ViewModels/Stack/NiagaraStackEntry.h"
+
 #include "Widgets/SCompoundWidget.h"
 #include "Widgets/Views/SListView.h"
 #include "UObject/ObjectKey.h"
@@ -12,6 +13,7 @@ class UNiagaraStackEntry;
 class UNiagaraStackItem;
 class UNiagaraStackViewModel;
 class UNiagaraSystemSelectionViewModel;
+class FNiagaraStackCommandContext;
 
 class SNiagaraOverviewStack : public SCompoundWidget
 {
@@ -31,10 +33,6 @@ public:
 	virtual void Tick(const FGeometry& AllottedGeometry, const double InCurrentTime, const float InDeltaTime) override;
 
 private:
-	void SetupCommands();
-
-	void DeleteSelectedEntries();
-
 	void AddEntriesRecursive(UNiagaraStackEntry& EntryToAdd, TArray<UNiagaraStackEntry*>& EntryList, const TArray<UClass*>& AcceptableClasses, TArray<UNiagaraStackEntry*> ParentChain);
 
 	void RefreshEntryList();
@@ -49,6 +47,8 @@ private:
 
 	void SystemSelectionChanged();
 
+	void UpdateCommandContextSelection();
+
 	FReply OnRowDragDetected(const FGeometry& InGeometry, const FPointerEvent& InPointerEvent, TWeakObjectPtr<UNiagaraStackEntry> InStackEntryWeak);
 
 	void OnRowDragLeave(const FDragDropEvent& InDragDropEvent);
@@ -69,7 +69,7 @@ private:
 
 	TArray<TWeakObjectPtr<UNiagaraStackEntry>> PreviousSelection;
 
-	TSharedPtr<FUICommandList> Commands;
+	TSharedPtr<FNiagaraStackCommandContext> StackCommandContext;
 
 	bool bRefreshEntryListPending;
 	bool bUpdatingOverviewSelectionFromStackSelection;
diff --git a/Engine/Plugins/FX/Niagara/Source/NiagaraEditorWidgets/Private/SNiagaraStack.cpp b/Engine/Plugins/FX/Niagara/Source/NiagaraEditorWidgets/Private/SNiagaraStack.cpp
--- a/Engine/Plugins/FX/Niagara/Source/NiagaraEditorWidgets/Private/SNiagaraStack.cpp
+++ b/Engine/Plugins/FX/Niagara/Source/NiagaraEditorWidgets/Private/SNiagaraStack.cpp
@@ -54,6 +54,7 @@
 #include "ScopedTransaction.h"
 #include "Widgets/Layout/SWrapBox.h"
 #include "ViewModels/Stack/NiagaraStackRoot.h"
+#include "NiagaraStackCommandContext.h"
 #include "NiagaraEditorUtilities.h"
 #include "Subsystems/AssetEditorSubsystem.h"
 
@@ -253,6 +254,7 @@ void SNiagaraStack::Construct(const FArguments& InArgs, UNiagaraStackViewModel*
 	NameColumnWidth = .3f;
 	ContentColumnWidth = .7f;
 	bNeedsJumpToNextOccurence = false;
+	StackCommandContext = MakeShared<FNiagaraStackCommandContext>();
 
 	ChildSlot
 	[
@@ -280,9 +282,9 @@ void SNiagaraStack::Construct(const FArguments& InArgs, UNiagaraStackViewModel*
 				.OnGetChildren(this, &SNiagaraStack::OnGetChildren)
 				.TreeItemsSource(&StackViewModel->GetRootEntryAsArray())
 				.OnTreeViewScrolled(this, &SNiagaraStack::StackTreeScrolled)
-				.SelectionMode(ESelectionMode::Single)
+				.OnSelectionChanged(this, &SNiagaraStack::StackTreeSelectionChanged)
+				.SelectionMode(ESelectionMode::Multi)
 				.OnItemToString_Debug_Static(&FNiagaraStackEditorWidgetsUtilities::StackEntryToStringForListDebug)
-				.OnMouseButtonClick(this, &SNiagaraStack::OnStackItemClicked)
 			]
 		]
 	];
@@ -842,7 +844,7 @@ TSharedRef<SNiagaraStackTableRow> SNiagaraStack::ConstructContainerForItem(UNiag
 		break;
 	}
 
-	return SNew(SNiagaraStackTableRow, StackViewModel, Item, StackTree.ToSharedRef())
+	return SNew(SNiagaraStackTableRow, StackViewModel, Item, StackCommandContext.ToSharedRef(), StackTree.ToSharedRef())
 		.RowPadding(RowPadding)
 		.ContentPadding(ContentPadding)
 		.ItemBackgroundColor(ItemBackgroundColor)
@@ -860,6 +862,16 @@ TSharedRef<SNiagaraStackTableRow> SNiagaraStack::ConstructContainerForItem(UNiag
 		.IssueIconVisibility(this, &SNiagaraStack::GetIssueIconVisibility);
 }
 
+
+FReply SNiagaraStack::OnKeyDown(const FGeometry& MyGeometry, const FKeyEvent& InKeyEvent)
+{
+	if (StackCommandContext->ProcessCommandBindings(InKeyEvent))
+	{
+		return FReply::Handled();
+	}
+	return SCompoundWidget::OnKeyDown(MyGeometry, InKeyEvent);
+}
+
 void SNiagaraStack::Tick(const FGeometry& AllottedGeometry, const double InCurrentTime, const float InDeltaTime)
 {
 	SCompoundWidget::Tick(AllottedGeometry, InCurrentTime, InDeltaTime);
@@ -1027,6 +1039,13 @@ void SNiagaraStack::StackTreeScrolled(double ScrollValue)
 	StackViewModel->SetLastScrollPosition(ScrollValue);
 }
 
+void SNiagaraStack::StackTreeSelectionChanged(UNiagaraStackEntry* InNewSelection, ESelectInfo::Type SelectInfo)
+{
+	TArray<UNiagaraStackEntry*> SelectedStackEntries;
+	StackTree->GetSelectedItems(SelectedStackEntries);
+	StackCommandContext->SetSelectedEntries(SelectedStackEntries);
+}
+
 float SNiagaraStack::GetNameColumnWidth() const
 {
 	return NameColumnWidth;
@@ -1059,13 +1078,6 @@ EVisibility SNiagaraStack::GetIssueIconVisibility() const
 	return StackViewModel->HasIssues() ? EVisibility::Visible : EVisibility::Collapsed;
 }
 
-void SNiagaraStack::OnStackItemClicked(UNiagaraStackEntry* Item)
-{
-	// Focus the stack search if we are just clicking on an item
-	// This won't be hit by a bubbling-up click if a widget with text entry is clicked
-	FSlateApplication::Get().SetKeyboardFocus(SearchBox, EFocusCause::OtherWidgetLostFocus);
-}
-
 void SNiagaraStack::OnCycleThroughIssues()
 {
 	UNiagaraStackEntry* StackEntry = StackViewModel->GetCurrentFocusedIssue();
diff --git a/Engine/Plugins/FX/Niagara/Source/NiagaraEditorWidgets/Private/SNiagaraStack.h b/Engine/Plugins/FX/Niagara/Source/NiagaraEditorWidgets/Private/SNiagaraStack.h
--- a/Engine/Plugins/FX/Niagara/Source/NiagaraEditorWidgets/Private/SNiagaraStack.h
+++ b/Engine/Plugins/FX/Niagara/Source/NiagaraEditorWidgets/Private/SNiagaraStack.h
@@ -16,6 +16,7 @@ class UNiagaraStackViewModel;
 class SNiagaraStackTableRow;
 class SSearchBox;
 class FReply;
+class FNiagaraStackCommandContext;
 
 class SNiagaraStack : public SCompoundWidget
 {
@@ -61,6 +62,8 @@ private:
 
 	void StackTreeScrolled(double ScrollValue);
 
+	void StackTreeSelectionChanged(UNiagaraStackEntry* InNewSelection, ESelectInfo::Type SelectInfo);
+
 	float GetNameColumnWidth() const;
 	float GetContentColumnWidth() const;
 
@@ -70,6 +73,8 @@ private:
 
 	EVisibility GetVisibilityForItem(UNiagaraStackEntry* Item) const;
 
+
+	virtual FReply OnKeyDown(const FGeometry& MyGeometry, const FKeyEvent& InKeyEvent) override;
 	virtual void Tick(const FGeometry& AllottedGeometry, const double InCurrentTime, const float InDeltaTime) override;
 
 	TSharedRef<SWidget> ConstructHeaderWidget();
@@ -107,7 +112,6 @@ private:
 
 	EVisibility GetIssueIconVisibility() const;
 
-	void OnStackItemClicked(UNiagaraStackEntry* Item);
 	void OnCycleThroughIssues();
 
 private:
@@ -124,4 +128,6 @@ private:
 	// ~ search stuff
 	TSharedPtr<SSearchBox> SearchBox;
 	bool bNeedsJumpToNextOccurence;
+
+	TSharedPtr<FNiagaraStackCommandContext> StackCommandContext;
 };
diff --git a/Engine/Plugins/FX/Niagara/Source/NiagaraEditorWidgets/Private/Stack/SNiagaraStackItemGroup.cpp b/Engine/Plugins/FX/Niagara/Source/NiagaraEditorWidgets/Private/Stack/SNiagaraStackItemGroup.cpp
--- a/Engine/Plugins/FX/Niagara/Source/NiagaraEditorWidgets/Private/Stack/SNiagaraStackItemGroup.cpp
+++ b/Engine/Plugins/FX/Niagara/Source/NiagaraEditorWidgets/Private/Stack/SNiagaraStackItemGroup.cpp
@@ -35,12 +35,12 @@ void SNiagaraStackItemGroup::Construct(const FArguments& InArgs, UNiagaraStackIt
 		.AutoWidth()
 		[
 			SNew(SButton)
-			.Visibility(this, &SNiagaraStackItemGroup::GetDeleteButtonVisibility)
 			.ButtonStyle(FEditorStyle::Get(), "HoverHintOnly")
 			.IsFocusable(false)
 			.ForegroundColor(FNiagaraEditorWidgetsStyle::Get().GetColor("NiagaraEditor.Stack.ForegroundColor"))
-			.ToolTipText(LOCTEXT("DeleteGroupToolTip", "Delete this group"))
+			.ToolTipText(this, &SNiagaraStackItemGroup::GetDeleteButtonToolTip)
 			.OnClicked(this, &SNiagaraStackItemGroup::DeleteClicked)
+			.IsEnabled(this, &SNiagaraStackItemGroup::GetDeleteButtonIsEnabled)
 			.Content()
 			[
 				SNew(STextBlock)
@@ -68,16 +68,17 @@ TSharedRef<SWidget> SNiagaraStackItemGroup::ConstructAddButton()
 	return SNullWidget::NullWidget;
 }
 
-EVisibility SNiagaraStackItemGroup::GetDeleteButtonVisibility() const
+FText SNiagaraStackItemGroup::GetDeleteButtonToolTip() const
 {
-	if (Group->CanDelete())
-	{
-		return EVisibility::Visible;
-	}
-	else
-	{
-		return EVisibility::Collapsed;
-	}
+	FText Message;
+	Group->TestCanDeleteWithMessage(Message);
+	return Message;
+}
+
+bool SNiagaraStackItemGroup::GetDeleteButtonIsEnabled() const
+{
+	FText UnusedMessage;
+	return Group->TestCanDeleteWithMessage(UnusedMessage);
 }
 
 FReply SNiagaraStackItemGroup::DeleteClicked()
diff --git a/Engine/Plugins/FX/Niagara/Source/NiagaraEditorWidgets/Private/Stack/SNiagaraStackItemGroup.h b/Engine/Plugins/FX/Niagara/Source/NiagaraEditorWidgets/Private/Stack/SNiagaraStackItemGroup.h
--- a/Engine/Plugins/FX/Niagara/Source/NiagaraEditorWidgets/Private/Stack/SNiagaraStackItemGroup.h
+++ b/Engine/Plugins/FX/Niagara/Source/NiagaraEditorWidgets/Private/Stack/SNiagaraStackItemGroup.h
@@ -19,7 +19,9 @@ public:
 private:
 	TSharedRef<SWidget> ConstructAddButton();
 
-	EVisibility GetDeleteButtonVisibility() const;
+	FText GetDeleteButtonToolTip() const;
+
+	bool GetDeleteButtonIsEnabled() const;
 
 	FReply DeleteClicked();
 
diff --git a/Engine/Plugins/FX/Niagara/Source/NiagaraEditorWidgets/Private/Stack/SNiagaraStackTableRow.cpp b/Engine/Plugins/FX/Niagara/Source/NiagaraEditorWidgets/Private/Stack/SNiagaraStackTableRow.cpp
--- a/Engine/Plugins/FX/Niagara/Source/NiagaraEditorWidgets/Private/Stack/SNiagaraStackTableRow.cpp
+++ b/Engine/Plugins/FX/Niagara/Source/NiagaraEditorWidgets/Private/Stack/SNiagaraStackTableRow.cpp
@@ -8,6 +8,7 @@
 #include "ViewModels/Stack/NiagaraStackItemGroup.h"
 #include "ViewModels/Stack/NiagaraStackEntry.h"
 #include "NiagaraEditorWidgetsUtilities.h"
+#include "NiagaraStackCommandContext.h"
 #include "Stack/SNiagaraStackIssueIcon.h"
 #include "Widgets/Input/SButton.h"
 #include "Widgets/Images/SImage.h"
@@ -18,7 +19,7 @@
 
 const float IndentSize = 12;
 
-void SNiagaraStackTableRow::Construct(const FArguments& InArgs, UNiagaraStackViewModel* InStackViewModel, UNiagaraStackEntry* InStackEntry, const TSharedRef<STreeView<UNiagaraStackEntry*>>& InOwnerTree)
+void SNiagaraStackTableRow::Construct(const FArguments& InArgs, UNiagaraStackViewModel* InStackViewModel, UNiagaraStackEntry* InStackEntry, TSharedRef<FNiagaraStackCommandContext> InStackCommandContext, const TSharedRef<STreeView<UNiagaraStackEntry*>>& InOwnerTree)
 {
 	ContentPadding = InArgs._ContentPadding;
 	bIsCategoryIconHighlighted = InArgs._IsCategoryIconHighlighted;
@@ -31,14 +32,14 @@ void SNiagaraStackTableRow::Construct(const FArguments& InArgs, UNiagaraStackVie
 	RowPadding = InArgs._RowPadding;
 	StackViewModel = InStackViewModel;
 	StackEntry = InStackEntry;
+	StackCommandContext = InStackCommandContext;
 	OwnerTree = InOwnerTree;
 
 	ExpandedImage = FCoreStyle::Get().GetBrush("TreeArrow_Expanded");
 	CollapsedImage = FCoreStyle::Get().GetBrush("TreeArrow_Collapsed");
 
-	InactiveItemBackgroundColor = InArgs._ItemBackgroundColor;
-	ActiveItemBackgroundColor = InactiveItemBackgroundColor + FLinearColor(.05f, .05f, .05f, 0.0f);
-	DisabledItemBackgroundColor = InactiveItemBackgroundColor + FLinearColor(.02f, .02f, .02f, 0.0f);
+	ItemBackgroundColor = InArgs._ItemBackgroundColor;
+	DisabledItemBackgroundColor = ItemBackgroundColor + FLinearColor(.02f, .02f, .02f, 0.0f);
 	ForegroundColor = InArgs._ItemForegroundColor;
 
 	ExecutionCategoryToolTipText = InStackEntry->GetExecutionSubcategoryName() != NAME_None
@@ -212,66 +213,60 @@ void SNiagaraStackTableRow::SetNameAndValueContent(TSharedRef<SWidget> InNameWid
 	FLinearColor AccentColor = AccentColorName != NAME_None ? FNiagaraEditorWidgetsStyle::Get().GetColor(AccentColorName) : FLinearColor::Transparent;
 	ChildSlot
 	[
-		SNew(SOverlay)
-		+ SOverlay::Slot()
+		SNew(SHorizontalBox)
+		.Visibility(this, &SNiagaraStackTableRow::GetRowVisibility)
+		// Accent color.
+		+ SHorizontalBox::Slot()
+		.AutoWidth()
+		.Padding(1, 0, 6, 0)
 		[
-			SNew(SHorizontalBox)
-			.Visibility(this, &SNiagaraStackTableRow::GetRowVisibility)
-			+ SHorizontalBox::Slot()
-			.AutoWidth()
-			.Padding(1, 0, 6, 0)
+			SNew(SBorder)
+			.BorderImage(FEditorStyle::GetBrush("WhiteBrush"))
+			.BorderBackgroundColor(AccentColor)
+			.Padding(0)
+			[
+				SNew(SBox)
+				.WidthOverride(4)
+			]
+		]
+		// Content
+		+ SHorizontalBox::Slot()
+		[
+			// Row content
+			SNew(SBox)
+			.Padding(RowPadding)
 			[
 				SNew(SBorder)
 				.BorderImage(FEditorStyle::GetBrush("WhiteBrush"))
-				.BorderBackgroundColor(AccentColor)
+				.BorderBackgroundColor(this, &SNiagaraStackTableRow::GetItemBackgroundColor)
+				.ForegroundColor(ForegroundColor)
 				.Padding(0)
-				[
-					SNew(SBox)
-					.WidthOverride(4)
-				]
-			]
-			+ SHorizontalBox::Slot()
-			[
-				SNew(SBox)
-				.Padding(RowPadding)
 				[
 					SNew(SBorder)
-					.BorderImage(FEditorStyle::GetBrush("WhiteBrush"))
-					.BorderBackgroundColor(this, &SNiagaraStackTableRow::GetItemBackgroundColor)
-					.ForegroundColor(ForegroundColor)
+					.BorderImage(this, &SNiagaraStackTableRow::GetSearchResultBorderBrush)
+					.BorderBackgroundColor(FNiagaraEditorWidgetsStyle::Get().GetColor("NiagaraEditor.Stack.SearchHighlightColor"))
 					.Padding(0)
 					[
-						ChildContent.ToSharedRef()
+						SNew(SBorder)
+						.BorderImage(this, &SNiagaraStackTableRow::GetSelectionBorderBrush)
+						.Padding(0)
+						[
+							ChildContent.ToSharedRef()
+						]
 					]
 				]
 			]
-			+ SHorizontalBox::Slot()
-			.AutoWidth()
-			.Padding(3, 0, 0, 0)
-			[
-				SNew(SNiagaraStackIssueIcon, StackViewModel, StackEntry)
-				.Visibility(IssueIconVisibility)
-			]
 		]
-		+ SOverlay::Slot()
+		+ SHorizontalBox::Slot()
+		.AutoWidth()
+		.Padding(3, 0, 0, 0)
 		[
-			SNew(SBorder)
-			.BorderImage(FNiagaraEditorWidgetsStyle::Get().GetBrush("NiagaraEditor.Stack.SearchResult"))
-			.BorderBackgroundColor(FNiagaraEditorWidgetsStyle::Get().GetColor("NiagaraEditor.Stack.SearchHighlightColor"))
-			.Visibility(this, &SNiagaraStackTableRow::GetSearchResultBorderVisibility)
-			.Padding(FMargin(0))
-			[
-				SNullWidget::NullWidget
-			]
+			SNew(SNiagaraStackIssueIcon, StackViewModel, StackEntry)
+			.Visibility(IssueIconVisibility)
 		]
 	];
 }
 
-bool SNiagaraStackTableRow::GetIsRowActive() const
-{
-	return IsHovered();
-}
-
 void SNiagaraStackTableRow::AddFillRowContextMenuHandler(FOnFillRowContextMenu FillRowContextMenuHandler)
 {
 	OnFillRowContextMenuHanders.Add(FillRowContextMenuHandler);
@@ -286,14 +281,25 @@ FReply SNiagaraStackTableRow::OnMouseButtonUp(const FGeometry& MyGeometry, const
 {
 	if (MouseEvent.GetEffectingButton() == EKeys::RightMouseButton)
 	{
-		FMenuBuilder MenuBuilder(true, nullptr);
+		TSharedPtr<ITypedTableView<UNiagaraStackEntry*>> OwnerTable = OwnerTablePtr.Pin();
+		if (OwnerTable.IsValid())
+		{
+			if (OwnerTable->GetSelectedItems().Contains(StackEntry) == false)
+			{
+				OwnerTable->Private_ClearSelection();
+				OwnerTable->Private_SetItemSelection(StackEntry, true, true);
+				OwnerTable->Private_SignalSelectionChanged(ESelectInfo::OnMouseClick);
+			}
+		}
+
+		FMenuBuilder MenuBuilder(true, StackCommandContext->GetCommands());
 		for (FOnFillRowContextMenu& OnFillRowContextMenuHandler : OnFillRowContextMenuHanders)
 		{
 			OnFillRowContextMenuHandler.ExecuteIfBound(MenuBuilder);
 		}
 
 		FNiagaraStackEditorWidgetsUtilities::AddStackEntryAssetContextMenuActions(MenuBuilder, *StackEntry);
-		FNiagaraStackEditorWidgetsUtilities::AddStackEntryContextMenuActions(MenuBuilder, *StackEntry);
+		StackCommandContext->AddEditMenuItems(MenuBuilder);
 
 		TArray<UNiagaraStackEntry*> EntriesToProcess;
 		TArray<UNiagaraStackEntry*> NavigationEntries;
@@ -347,6 +353,12 @@ FReply SNiagaraStackTableRow::OnMouseButtonUp(const FGeometry& MyGeometry, const
 	return STableRow<UNiagaraStackEntry*>::OnMouseButtonUp(MyGeometry, MouseEvent);
 }
 
+const FSlateBrush* SNiagaraStackTableRow::GetBorder() const
+{
+	// Return no brush here so that the background doesn't change.  The border color changing will be handled by an internal border.
+	return FEditorStyle::GetBrush("NoBrush");
+}
+
 void SNiagaraStackTableRow::CollapseChildren()
 {
 	TArray<UNiagaraStackEntry*> Children;
@@ -451,21 +463,21 @@ void SNiagaraStackTableRow::OnValueColumnWidthChanged(float Width)
 
 FSlateColor SNiagaraStackTableRow::GetItemBackgroundColor() const
 {
-	if (StackEntry->GetIsEnabled() && StackEntry->GetOwnerIsEnabled())
-	{
-		return GetIsRowActive() ? ActiveItemBackgroundColor : InactiveItemBackgroundColor;
-	}
-	else
-	{
-		return DisabledItemBackgroundColor;
-	}
+	return StackEntry->GetIsEnabledAndOwnerIsEnabled() 
+		? ItemBackgroundColor 
+		: DisabledItemBackgroundColor;
+}
+
+const FSlateBrush* SNiagaraStackTableRow::GetSelectionBorderBrush() const
+{
+	return STableRow<UNiagaraStackEntry*>::GetBorder();
 }
 
-EVisibility SNiagaraStackTableRow::GetSearchResultBorderVisibility() const
+const FSlateBrush* SNiagaraStackTableRow::GetSearchResultBorderBrush() const
 {
 	return StackViewModel->GetCurrentFocusedEntry() == StackEntry
-		? EVisibility::HitTestInvisible
-		: EVisibility::Hidden;
+		? FNiagaraEditorWidgetsStyle::Get().GetBrush("NiagaraEditor.Stack.SearchResult")
+		: FEditorStyle::GetBrush("NoBrush");
 }
 
 void SNiagaraStackTableRow::NavigateTo(UNiagaraStackEntry* Item)
diff --git a/Engine/Plugins/FX/Niagara/Source/NiagaraEditorWidgets/Private/Stack/SNiagaraStackTableRow.h b/Engine/Plugins/FX/Niagara/Source/NiagaraEditorWidgets/Private/Stack/SNiagaraStackTableRow.h
--- a/Engine/Plugins/FX/Niagara/Source/NiagaraEditorWidgets/Private/Stack/SNiagaraStackTableRow.h
+++ b/Engine/Plugins/FX/Niagara/Source/NiagaraEditorWidgets/Private/Stack/SNiagaraStackTableRow.h
@@ -12,6 +12,7 @@
 
 class UNiagaraStackViewModel;
 class UNiagaraStackEntry;
+class FNiagaraStackCommandContext;
 
 class SNiagaraStackTableRow: public STableRow<UNiagaraStackEntry*>
 {
@@ -27,7 +28,7 @@ public:
 		, _ShowExecutionCategoryIcon(false)
 		, _RowPadding(FMargin(0, 0, 0, 0))
 	{}
-		SLATE_ARGUMENT(FMargin, ContentPadding)
+	SLATE_ARGUMENT(FMargin, ContentPadding)
 		SLATE_ARGUMENT(FLinearColor, ItemBackgroundColor)
 		SLATE_ARGUMENT(FLinearColor, ItemForegroundColor)
 		SLATE_ARGUMENT(bool, IsCategoryIconHighlighted)
@@ -41,10 +42,10 @@ public:
 		SLATE_EVENT(FOnDragDetected, OnDragDetected)
 		SLATE_EVENT(FOnTableRowDragLeave, OnDragLeave)
 		SLATE_EVENT(FOnCanAcceptDrop, OnCanAcceptDrop)
-		SLATE_EVENT(FOnAcceptDrop, OnAcceptDrop)
+		SLATE_EVENT(FOnAcceptDrop, OnAcceptDrop);
 	SLATE_END_ARGS();
 
-	void Construct(const FArguments& InArgs, UNiagaraStackViewModel* InStackViewModel, UNiagaraStackEntry* InStackEntry, const TSharedRef<STreeView<UNiagaraStackEntry*>>& InOwnerTree);
+	void Construct(const FArguments& InArgs, UNiagaraStackViewModel* InStackViewModel, UNiagaraStackEntry* InStackEntry, TSharedRef<FNiagaraStackCommandContext> InStackCommandContext, const TSharedRef<STreeView<UNiagaraStackEntry*>>& InOwnerTree);
 
 	void SetOverrideNameWidth(TOptional<float> InMinWidth, TOptional<float> InMaxWidth);
 
@@ -60,14 +61,14 @@ public:
 
 	void SetNameAndValueContent(TSharedRef<SWidget> InNameWidget, TSharedPtr<SWidget> InValueWidget);
 
-	bool GetIsRowActive() const;
-
 	void AddFillRowContextMenuHandler(FOnFillRowContextMenu FillRowContextMenuHandler);
 
 	virtual FReply OnMouseButtonDoubleClick(const FGeometry& InMyGeometry, const FPointerEvent& InMouseEvent) override;
 
 	FReply OnMouseButtonUp(const FGeometry& MyGeometry, const FPointerEvent& MouseEvent);
 
+	virtual const FSlateBrush* GetBorder() const override;
+
 private:
 	void CollapseChildren();
 
@@ -91,7 +92,9 @@ private:
 
 	FSlateColor GetItemBackgroundColor() const;
 
-	EVisibility GetSearchResultBorderVisibility() const;
+	const FSlateBrush* GetSelectionBorderBrush() const;
+
+	const FSlateBrush* GetSearchResultBorderBrush() const;
 
 	void NavigateTo(UNiagaraStackEntry* Item);
 
@@ -111,8 +114,7 @@ private:
 	const FSlateBrush* ExpandedImage;
 	const FSlateBrush* CollapsedImage;
 
-	FLinearColor InactiveItemBackgroundColor;
-	FLinearColor ActiveItemBackgroundColor;
+	FLinearColor ItemBackgroundColor;
 	FLinearColor DisabledItemBackgroundColor;
 	FLinearColor ForegroundColor;
 
@@ -136,4 +138,6 @@ private:
 	bool bShowExecutionCategoryIcon;
 
 	TArray<FOnFillRowContextMenu> OnFillRowContextMenuHanders;
+
+	TSharedPtr<FNiagaraStackCommandContext> StackCommandContext;
 };
\ No newline at end of file
