Niagara - Fix non-unity.

diff --git a/Engine/Plugins/FX/Niagara/Source/NiagaraEditor/Public/ViewModels/Stack/NiagaraStackClipboardUtilities.h b/Engine/Plugins/FX/Niagara/Source/NiagaraEditor/Public/ViewModels/Stack/NiagaraStackClipboardUtilities.h
--- a/Engine/Plugins/FX/Niagara/Source/NiagaraEditor/Public/ViewModels/Stack/NiagaraStackClipboardUtilities.h
+++ b/Engine/Plugins/FX/Niagara/Source/NiagaraEditor/Public/ViewModels/Stack/NiagaraStackClipboardUtilities.h
@@ -6,19 +6,19 @@ class UNiagaraStackEntry;
 
 namespace FNiagaraStackClipboardUtilities
 {
-	bool NIAGARAEDITOR_API TestCanCutSelectionWithMessage(const TArray<UNiagaraStackEntry*>& SelectedEntries, FText& OutCanCutMessage);
+	NIAGARAEDITOR_API bool TestCanCutSelectionWithMessage(const TArray<UNiagaraStackEntry*>& SelectedEntries, FText& OutCanCutMessage);
 
-	void NIAGARAEDITOR_API CutSelection(const TArray<UNiagaraStackEntry*>& SelectedEntries);
+	NIAGARAEDITOR_API void CutSelection(const TArray<UNiagaraStackEntry*>& SelectedEntries);
 
-	bool NIAGARAEDITOR_API TestCanCopySelectionWithMessage(const TArray<UNiagaraStackEntry*>& SelectedEntries, FText& OutCanCopyMessage);
+	NIAGARAEDITOR_API bool TestCanCopySelectionWithMessage(const TArray<UNiagaraStackEntry*>& SelectedEntries, FText& OutCanCopyMessage);
 
-	void NIAGARAEDITOR_API CopySelection(const TArray<UNiagaraStackEntry*>& SelectedEntries);
+	NIAGARAEDITOR_API void CopySelection(const TArray<UNiagaraStackEntry*>& SelectedEntries);
 
-	bool NIAGARAEDITOR_API TestCanPasteSelectionWithMessage(const TArray<UNiagaraStackEntry*>& SelectedEntries, FText& OutCanPasteMessage);
+	NIAGARAEDITOR_API bool TestCanPasteSelectionWithMessage(const TArray<UNiagaraStackEntry*>& SelectedEntries, FText& OutCanPasteMessage);
 
-	void NIAGARAEDITOR_API PasteSelection(const TArray<UNiagaraStackEntry*>& SelectedEntries);
+	NIAGARAEDITOR_API void PasteSelection(const TArray<UNiagaraStackEntry*>& SelectedEntries);
 
-	bool NIAGARAEDITOR_API TestCanDeleteSelectionWithMessage(const TArray<UNiagaraStackEntry*>& SelectedEntries, FText& OutMessage);
+	NIAGARAEDITOR_API bool TestCanDeleteSelectionWithMessage(const TArray<UNiagaraStackEntry*>& SelectedEntries, FText& OutMessage);
 
-	void NIAGARAEDITOR_API DeleteSelection(const TArray<UNiagaraStackEntry*>& SelectedEntries);
+	NIAGARAEDITOR_API void DeleteSelection(const TArray<UNiagaraStackEntry*>& SelectedEntries);
 }
\ No newline at end of file
