#UE4 IWYU fix

diff --git a/Engine/Plugins/FX/Niagara/Source/NiagaraEditorWidgets/Private/NiagaraStackCommandContext.h b/Engine/Plugins/FX/Niagara/Source/NiagaraEditorWidgets/Private/NiagaraStackCommandContext.h
--- a/Engine/Plugins/FX/Niagara/Source/NiagaraEditorWidgets/Private/NiagaraStackCommandContext.h
+++ b/Engine/Plugins/FX/Niagara/Source/NiagaraEditorWidgets/Private/NiagaraStackCommandContext.h
@@ -7,6 +7,7 @@
 class UNiagaraStackEntry;
 class FUICommandList;
 class FMenuBuilder;
+struct FKeyEvent;
 
 class FNiagaraStackCommandContext : public TSharedFromThis<FNiagaraStackCommandContext>
 {
