Persistent IDs for GPU particles.

"Requires persistent IDs" is now supported for GPU emitters. Each emitter maintains a list of free IDs and an ID to index mapping table. The ID to index mapping is filled in by the simulation step. The free IDs are computed after each simulation step by compacting the slots in the ID mapping table which are not touched by the simulation.

diff --git a/Engine/Plugins/FX/Niagara/Shaders/Private/NiagaraComputeFreeIDs.usf b/Engine/Plugins/FX/Niagara/Shaders/Private/NiagaraComputeFreeIDs.usf
new file mode 100644
--- /dev/null
+++ b/Engine/Plugins/FX/Niagara/Shaders/Private/NiagaraComputeFreeIDs.usf
@@ -0,0 +1,134 @@
+// Copyright 1998-2019 Epic Games, Inc. All Rights Reserved.
+
+#include "/Engine/Private/Common.ush"
+
+#define USE_WAVE_INTRISICS (PS4_PROFILE || XBOXONE_PROFILE)
+
+// ID mapping table produced by the simulate step.
+Buffer<int> IDToIndexTable;
+
+// Output list of free IDs.
+RWBuffer<int> RWFreeIDList;
+
+// Buffer containing the current size of each list that's currently being processed.
+// Multiple invocations of this shader might be running concurrently, for independent emitters.
+RWBuffer<int> RWFreeIDListSizes;
+// Index of the free list that's being computed by this invocation. The current size of the
+// list is at RWFreeIDListSizes[FreeIDListIndex].
+uint FreeIDListIndex;
+
+#if USE_WAVE_INTRISICS
+
+// This happens to be fixed on PS4 and XB1, but we'll need to configure it externally
+// when we start using this path on PC.
+#define WAVE_SIZE 64
+
+[numthreads(WAVE_SIZE, 1, 1)]
+void ComputeFreeIDs(uint3 GroupId : SV_GroupID, uint3 GroupThreadId : SV_GroupThreadID, uint3 DispatchThreadId : SV_DispatchThreadID)
+{
+	uint InputStart = GroupId.x * WAVE_SIZE;
+	uint Thread = GroupThreadId.x;
+
+	uint ParticleID = InputStart + Thread;
+	uint IsFree = (IDToIndexTable[ParticleID] == -1) ? 1 : 0;
+	
+	// Count how many predicates are true across all lanes.
+	uint NumGroupFreeIDs = WaveActiveCountBits(IsFree);
+
+	// Compute a prefix sum on the predicate bit mask. This gives us the
+	// local write offset for each lane.
+	uint WriteOffset = WavePrefixCountBits(IsFree);
+
+	// Skip the whole wave if all the IDs are in use.
+	if(NumGroupFreeIDs > 0)
+	{
+		uint GroupWriteStart;
+		if(WaveIsFirstLane())
+		{
+			// Add to the global write offset. The previous offset is where we start writing in the output buffer.
+			// We can reduce the number of atomic adds by processing multiple sub-groups in the shader, such that each lane
+			// computes the offsets and counts of a sub-group. However, this doesn't seem to be the bottleneck right now,
+			// so it's not worth the effort.
+			InterlockedAdd(RWFreeIDListSizes[FreeIDListIndex], NumGroupFreeIDs, GroupWriteStart);
+		}
+
+		// Broadcast the global write offset to all lanes.
+		// TODO: WaveReadFirstLane is actually called WaveReadLaneFirst in HLSL, despite what the docs say, so we need to change our defines for
+		// other platforms to use the HLSL name or we'll run into trouble when we'll compile for SM6.
+		GroupWriteStart = WaveReadFirstLane(GroupWriteStart);
+
+		if(IsFree)
+		{
+			RWFreeIDList[GroupWriteStart + WriteOffset] = ParticleID;
+		}
+	}
+}
+
+#else
+
+#define THREAD_COUNT 128
+
+// Temporary buffer where we compute the number of elements to write before each input element.
+groupshared uint GroupWriteOffsets[2*THREAD_COUNT];
+// The offset where this thread group will start writing its output.
+groupshared uint GroupWriteStart;
+
+[numthreads(THREAD_COUNT, 1, 1)]
+void ComputeFreeIDs(uint3 GroupId : SV_GroupID, uint3 GroupThreadId : SV_GroupThreadID, uint3 DispatchThreadId : SV_DispatchThreadID)
+{
+	uint InputStart = GroupId.x * THREAD_COUNT;
+	uint Thread = GroupThreadId.x;
+	uint ParticleID = InputStart + Thread;
+
+	int OutBuffer = 0, InBuffer = 1;
+
+	// Evaluate predicate for each input slot, then shift down by one position because we need an exclusive sum.
+	if(Thread > 0)
+	{
+		uint IsFree = (IDToIndexTable[ParticleID - 1] == -1) ? 1 : 0;
+		GroupWriteOffsets[Thread] = IsFree;
+	}
+	else
+	{
+		GroupWriteOffsets[0] = 0;
+	}
+	GroupMemoryBarrierWithGroupSync();
+
+	// Compute prefix sum.
+	[unroll]
+	for(uint Offset = 1; Offset < THREAD_COUNT; Offset = Offset << 1)
+	{
+		// Swap the input and output buffers.
+		OutBuffer = 1 - OutBuffer;
+		InBuffer = 1 - InBuffer;
+		if(Thread >= Offset)
+		{
+			GroupWriteOffsets[OutBuffer*THREAD_COUNT + Thread] = GroupWriteOffsets[InBuffer*THREAD_COUNT + Thread - Offset] + GroupWriteOffsets[InBuffer*THREAD_COUNT + Thread];
+		}
+		else
+		{
+			GroupWriteOffsets[OutBuffer*THREAD_COUNT + Thread] = GroupWriteOffsets[InBuffer*THREAD_COUNT + Thread];
+		}
+		GroupMemoryBarrierWithGroupSync();
+	}
+
+	if(Thread == 0)
+	{
+		// Add the value of the last predicate to get the total number of IDs that this group will write.
+		uint LastIsFree = (IDToIndexTable[InputStart + THREAD_COUNT - 1] == -1) ? 1 : 0;
+		uint NumGroupFreeIDs = GroupWriteOffsets[(OutBuffer + 1)*THREAD_COUNT - 1] + LastIsFree;
+		// Add to the global write offset. The previous offset is where we start writing in the output buffer.
+		InterlockedAdd(RWFreeIDListSizes[FreeIDListIndex], NumGroupFreeIDs, GroupWriteStart);
+	}
+	GroupMemoryBarrierWithGroupSync();
+
+	if(IDToIndexTable[ParticleID] == -1)
+	{
+		// This is a free ID, output it.
+		uint WriteOffset = GroupWriteStart + GroupWriteOffsets[OutBuffer*THREAD_COUNT + Thread];
+		// TODO: consider buffering the results before writing to the output UAV.
+		RWFreeIDList[WriteOffset] = ParticleID;
+	}
+}
+
+#endif
diff --git a/Engine/Plugins/FX/Niagara/Shaders/Private/NiagaraEmitterInstanceShader.usf b/Engine/Plugins/FX/Niagara/Shaders/Private/NiagaraEmitterInstanceShader.usf
--- a/Engine/Plugins/FX/Niagara/Shaders/Private/NiagaraEmitterInstanceShader.usf
+++ b/Engine/Plugins/FX/Niagara/Shaders/Private/NiagaraEmitterInstanceShader.usf
@@ -418,9 +418,9 @@ int rand_int(int x, int Seed1, int Seed2, int Seed3)
 	uint ReadInstanceCountOffset;
 	uint WriteInstanceCountOffset;
 
-	//TODO: IDs for GPU
-//	void AcquireID(int DataSetID, out int IDIndex, out int IDAcquireTag);
-// 	void UpdateID(int DataSetID, int IDIndex, int InstanceIndex);
+	Buffer<int> FreeIDList;
+	RWBuffer<int> RWIDToIndexTable;
+
 
 	/* Returns the current instance index relative to the operation (spawn/update)
 	 */
@@ -436,6 +436,30 @@ int rand_int(int x, int Seed1, int Seed2, int Seed3)
 		}
 	}
 
+	void AcquireID(int DataSetID, out int IDIndex, out int IDAcquireTag)
+	{
+		// Begin static assert : GPU particles only support DataSetID 0
+		int MustBe0[1];
+		MustBe0[DataSetID] = 0;
+		// End static assert
+
+		// This is the same as ExecIndex() right now, but that function may change in the future to accommodate multiple
+		// spawn infos. Revisit this computation if the change affects the meaning of UpdateStartInstance or GSpawnStartInstance.
+		int SpawnIndex = UpdateStartInstance + GDispatchThreadId.x - GSpawnStartInstance;
+		IDIndex = FreeIDList[SpawnIndex];
+		IDAcquireTag = EmitterTickCounter;
+	}
+
+ 	void UpdateID(int DataSetID, int IDIndex, int InstanceIndex)
+	{
+		// Begin static assert : GPU particles only support DataSetID 0
+		int MustBe0[1];
+		MustBe0[DataSetID] = 0;
+		// End static assert
+
+		RWIDToIndexTable[IDIndex] = InstanceIndex;
+	}
+
 #define USE_WAVE_INTRISICS (PS4_PROFILE || XBOXONE_PROFILE)
 #define USE_GROUP_SHARED ((THREADGROUP_SIZE == 64 || THREADGROUP_SIZE == 32)  && !USE_WAVE_INTRISICS)
 
diff --git a/Engine/Plugins/FX/Niagara/Shaders/Private/NiagaraInitFreeIDBuffer.usf b/Engine/Plugins/FX/Niagara/Shaders/Private/NiagaraInitFreeIDBuffer.usf
new file mode 100644
--- /dev/null
+++ b/Engine/Plugins/FX/Niagara/Shaders/Private/NiagaraInitFreeIDBuffer.usf
@@ -0,0 +1,24 @@
+// Copyright 1998-2019 Epic Games, Inc. All Rights Reserved.
+
+#include "/Engine/Private/Common.ush"
+
+RWBuffer<int> RWNewBuffer;
+Buffer<int> ExistingBuffer;
+uint NumNewElements;
+uint NumExistingElements;
+
+[numthreads(THREAD_COUNT, 1, 1)]
+void InitIDBufferCS(uint Index : SV_DispatchThreadID)
+{
+	if(Index.x < NumNewElements)
+	{
+		// Place new IDs at the start of the buffer, so we don't have to worry about how many
+		// elements are actually in use in the existing buffer.
+		RWNewBuffer[Index.x] = NumExistingElements + Index.x;
+	}
+	else
+	{
+		// Copy existing IDs after the new IDs.
+		RWNewBuffer[Index.x] = ExistingBuffer[Index.x];
+	}
+}
diff --git a/Engine/Plugins/FX/Niagara/Source/Niagara/Classes/NiagaraDataSet.h b/Engine/Plugins/FX/Niagara/Source/Niagara/Classes/NiagaraDataSet.h
--- a/Engine/Plugins/FX/Niagara/Source/Niagara/Classes/NiagaraDataSet.h
+++ b/Engine/Plugins/FX/Niagara/Source/Niagara/Classes/NiagaraDataSet.h
@@ -103,7 +103,7 @@ protected:
 public:
 	FNiagaraDataBuffer(FNiagaraDataSet* InOwner);
 	void Allocate(uint32 NumInstances, bool bMaintainExisting = false);
-	void AllocateGPU(uint32 InNumInstances, FNiagaraGPUInstanceCountManager& GPUInstanceCountManager, FRHICommandList &RHICmdList);
+	void AllocateGPU(uint32 InNumInstances, FNiagaraGPUInstanceCountManager& GPUInstanceCountManager, FRHICommandList& RHICmdList, ERHIFeatureLevel::Type FeatureLevel);
 	void SwapInstances(uint32 OldIndex, uint32 NewIndex);
 	void KillInstance(uint32 InstanceIdx);
 	void CopyTo(FNiagaraDataBuffer& DestBuffer, int32 SrcStartIdx, int32 DestStartIdx, int32 NumInstances)const;
@@ -139,6 +139,9 @@ public:
 	FORCEINLINE FRWBuffer& GetGPUBufferInt() { return GPUBufferInt;	}
 	FORCEINLINE uint32 GetGPUInstanceCountBufferOffset() const { return GPUInstanceCountBufferOffset; }
 	FORCEINLINE void ClearGPUInstanceCountBufferOffset() { GPUInstanceCountBufferOffset = INDEX_NONE; }
+	FORCEINLINE uint32 GetGPUNumAllocatedIDs() const { return NumIDsAllocatedForGPU; }
+	FORCEINLINE FRWBuffer& GetGPUFreeIDs() { return GPUFreeIDs; }
+	FORCEINLINE FRWBuffer& GetGPUIDToIndexTable() { return GPUIDToIndexTable; }
 
 	FORCEINLINE int32 GetSafeComponentBufferSize() const { return GetSafeComponentBufferSize(GetNumInstancesAllocated()); }
 	FORCEINLINE uint32 GetFloatStride() const { return FloatStride; }
@@ -195,6 +198,12 @@ private:
 	FRWBuffer GPUBufferFloat;
 	/** GPU Buffer containing floating point values for GPU simulations. */
 	FRWBuffer GPUBufferInt;
+	/** Size of the GPU ID buffers. */
+	uint32 NumIDsAllocatedForGPU;
+	/** GPU list of free particle IDs. */
+	FRWBuffer GPUFreeIDs;
+	/** GPU table which maps particle ID to index. */
+	FRWBuffer GPUIDToIndexTable;
 	//////////////////////////////////////////////////////////////////////////
 
 	/** Number of instances in data. */
diff --git a/Engine/Plugins/FX/Niagara/Source/Niagara/Classes/NiagaraEmitterInstanceBatcher.h b/Engine/Plugins/FX/Niagara/Source/Niagara/Classes/NiagaraEmitterInstanceBatcher.h
--- a/Engine/Plugins/FX/Niagara/Source/Niagara/Classes/NiagaraEmitterInstanceBatcher.h
+++ b/Engine/Plugins/FX/Niagara/Source/Niagara/Classes/NiagaraEmitterInstanceBatcher.h
@@ -57,6 +57,7 @@ public:
 		, ParticleSortBuffers(true)
 		// @todo REMOVE THIS HACK
 		, LastFrameThatDrainedData(GFrameNumberRenderThread)
+		, NumAllocatedFreeIDListSizes(0)
 	{
 	}
 
@@ -140,7 +141,7 @@ public:
 				uint32 ShaderStageIndex = 0,
 				FNiagaraDataInterfaceProxy *IterationInterface = nullptr,
 				bool HasRunParticleStage = false
-			) const;
+			);
 
 	void ResizeCurrentBuffer(FRHICommandList &RHICmdList, FNiagaraComputeExecutionContext *Context, uint32 NewNumInstances, uint32 PrevNumInstances) const;
 
@@ -162,12 +163,12 @@ public:
 	void PostStageInterface(const FNiagaraGPUSystemTick& Tick, FNiagaraComputeInstanceData *Instance, FRHICommandList &RHICmdList, FNiagaraShader* ComputeShader, const uint32 ShaderStageIndex) const;
 
 	/** Run the dispatch over multiple stages */
-	void DispatchMultipleStages(const FNiagaraGPUSystemTick& Tick, FNiagaraComputeInstanceData *Instance, FRHICommandList &RHICmdList, FRHIUniformBuffer* ViewUniformBuffer, FNiagaraShader* ComputeShader) const;
+	void DispatchMultipleStages(const FNiagaraGPUSystemTick& Tick, FNiagaraComputeInstanceData *Instance, FRHICommandList &RHICmdList, FRHIUniformBuffer* ViewUniformBuffer, FNiagaraShader* ComputeShader);
 
 private:
-	void ExecuteAll(FRHICommandList &RHICmdList, FRHIUniformBuffer* ViewUniformBuffer, bool bSetReadback, ETickStage TickStage);
-	void ResizeBuffersAndGatherResources(FOverlappableTicks& OverlappableTick, FRHICommandList& RHICmdList, FNiagaraBufferArray& DestDataBuffers, FNiagaraBufferArray& CurrDataBuffers, FNiagaraBufferArray& DestBufferIntFloat, FNiagaraBufferArray& CurrBufferIntFloat);
-	void DispatchAllOnCompute(FOverlappableTicks& OverlappableTick, FRHICommandList& RHICmdList, FRHIUniformBuffer* ViewUniformBuffer, FNiagaraBufferArray& DestDataBuffers, FNiagaraBufferArray& CurrDataBuffers, FNiagaraBufferArray& DestBufferIntFloat, FNiagaraBufferArray& CurrBufferIntFloat, bool bSetReadback);
+	void ExecuteAll(FRHICommandList& RHICmdList, FRHIUniformBuffer* ViewUniformBuffer, bool bSetReadback, ETickStage TickStage);
+	void ResizeBuffersAndGatherResources(FOverlappableTicks& OverlappableTick, FRHICommandList& RHICmdList, FNiagaraBufferArray& ReadBuffers, FNiagaraBufferArray& WriteBuffers);
+	void DispatchAllOnCompute(FOverlappableTicks& OverlappableTick, FRHICommandList& RHICmdList, FRHIUniformBuffer* ViewUniformBuffer, FNiagaraBufferArray& ReadBuffers, FNiagaraBufferArray& WriteBuffers, bool bSetReadback);
 
 	bool ShouldTickForStage(const FNiagaraGPUSystemTick& Tick, ETickStage TickStage) const;
 
@@ -198,7 +199,7 @@ private:
 	FParticleSortBuffers ParticleSortBuffers;
 
 	// GPU emitter instance count buffer. Contains the actual particle / instance count generate in the GPU tick.
-	mutable FNiagaraGPUInstanceCountManager GPUInstanceCounterManager;
+	FNiagaraGPUInstanceCountManager GPUInstanceCounterManager;
 
 	// @todo REMOVE THIS HACK
 	uint32 LastFrameThatDrainedData;
@@ -215,4 +216,9 @@ private:
 	TSet<FNiagaraDataSet*> DataSetsToDestroy_RT;
 
 	TSet<TSharedPtr<FNiagaraDataInterfaceProxy, ESPMode::ThreadSafe>> DIProxyDeferredDeletes_RT;
+
+	/** A buffer used by the compute shader which determines the list of free particle IDs for each emitter. */
+	FRWBuffer FreeIDListSizesBuffer;
+	uint32 NumAllocatedFreeIDListSizes;
+	uint32 CurrentFreeIDListIndex;
 };
diff --git a/Engine/Plugins/FX/Niagara/Source/Niagara/Private/NiagaraDataSet.cpp b/Engine/Plugins/FX/Niagara/Source/Niagara/Private/NiagaraDataSet.cpp
--- a/Engine/Plugins/FX/Niagara/Source/Niagara/Private/NiagaraDataSet.cpp
+++ b/Engine/Plugins/FX/Niagara/Source/Niagara/Private/NiagaraDataSet.cpp
@@ -10,6 +10,7 @@
 #include "NiagaraRenderer.h"
 #include "NiagaraGPUInstanceCountManager.h"
 #include "NiagaraEmitterInstanceBatcher.h"
+#include "NiagaraShaderParticleID.h"
 
 DECLARE_CYCLE_STAT(TEXT("InitRenderData"), STAT_InitRenderData, STATGROUP_Niagara);
 
@@ -430,6 +431,7 @@ FNiagaraDataBuffer::FNiagaraDataBuffer(FNiagaraDataSet* InOwner)
 	: Owner(InOwner)
 	, GPUInstanceCountBufferOffset(INDEX_NONE)
 	, NumChunksAllocatedForGPU(0)
+	, NumIDsAllocatedForGPU(0)
 	, NumInstances(0)
 	, NumInstancesAllocated(0)
 	, FloatStride(0)
@@ -598,7 +600,7 @@ void FNiagaraDataBuffer::Allocate(uint32 InNumInstances, bool bMaintainExisting)
 	}
 }
 
-void FNiagaraDataBuffer::AllocateGPU(uint32 InNumInstances, FNiagaraGPUInstanceCountManager& GPUInstanceCountManager, FRHICommandList &RHICmdList)
+void FNiagaraDataBuffer::AllocateGPU(uint32 InNumInstances, FNiagaraGPUInstanceCountManager& GPUInstanceCountManager, FRHICommandList& RHICmdList, ERHIFeatureLevel::Type FeatureLevel)
 {
 	CheckUsage(false);
 
@@ -633,6 +635,17 @@ void FNiagaraDataBuffer::AllocateGPU(uint32 InNumInstances, FNiagaraGPUInstanceC
 		{
 			GPUBufferInt.Release();
 		}
+		if (GPUFreeIDs.Buffer)
+		{
+			GPUFreeIDs.Release();
+		}
+		if (GPUIDToIndexTable.Buffer)
+		{
+			GPUIDToIndexTable.Release();
+		}
+
+		NumChunksAllocatedForGPU = 0;
+		NumIDsAllocatedForGPU = 0;
 	}
 	else // Otherwise check for growing and possibly shrinking (if GNiagaraGPUDataBufferBufferSlack > 1) .
 	{
@@ -659,6 +672,33 @@ void FNiagaraDataBuffer::AllocateGPU(uint32 InNumInstances, FNiagaraGPUInstanceC
 				GPUBufferInt.Initialize(sizeof(int32), NumElementsToAlloc * Owner->GetNumInt32Components(), EPixelFormat::PF_R32_SINT, BUF_Static);
 			}
 		}
+
+		// The ID buffers cannot be shrinked, because the maximum ID could be currently in use even if we have a single alive particle.
+		if (Owner->GetNeedsPersistentIDs() && PaddedNumInstances > NumIDsAllocatedForGPU)
+		{
+			// Smaller chunk size for ID allocations, because clearing and compacting the ID table must run over all the allocated
+			// elements, so too much slack memory hurts performance.
+			constexpr uint32 ID_ALLOC_CHUNKSIZE = 1024;
+			const uint32 NumIDsToAlloc = FMath::DivideAndRoundUp(PaddedNumInstances, ID_ALLOC_CHUNKSIZE) * ID_ALLOC_CHUNKSIZE;
+
+			// We must maintain the existing list of free IDs.
+			FRWBuffer NewFreeIDsBuffer;
+			NewFreeIDsBuffer.Initialize(sizeof(int32), NumIDsToAlloc, EPixelFormat::PF_R32_SINT, BUF_Static, TEXT("NiagaraFreeIDList"));
+
+			FRHIShaderResourceView* ExistingBuffer = NumIDsAllocatedForGPU > 0 ? GPUFreeIDs.SRV : FNiagaraRenderer::GetDummyIntBuffer().SRV;
+			NiagaraInitGPUFreeIDList(RHICmdList, FeatureLevel, NumIDsToAlloc, NewFreeIDsBuffer, NumIDsAllocatedForGPU, ExistingBuffer);
+
+			GPUFreeIDs = MoveTemp(NewFreeIDsBuffer);
+
+			// The ID to index table is filled in each time the emitter runs, so we don't need to keep its contents.
+			if (GPUIDToIndexTable.Buffer)
+			{
+				GPUIDToIndexTable.Release();
+			}
+			GPUIDToIndexTable.Initialize(sizeof(int32), NumIDsToAlloc, EPixelFormat::PF_R32_SINT, BUF_Static, TEXT("NiagaraIDToIndexTable"));
+
+			NumIDsAllocatedForGPU = NumIDsToAlloc;
+		}
 	}
 }
 
@@ -904,12 +944,18 @@ void FNiagaraDataBuffer::SetShaderParams(FNiagaraShader *Shader, FRHICommandList
 		SetSRVParameter(CommandList, ComputeShader, Shader->FloatInputBufferParam, InstancesAllocated ? GetGPUBufferFloat().SRV : FNiagaraRenderer::GetDummyFloatBuffer().SRV);
 		SetSRVParameter(CommandList, ComputeShader, Shader->IntInputBufferParam, InstancesAllocated ? GetGPUBufferInt().SRV : FNiagaraRenderer::GetDummyIntBuffer().SRV);
 		SetShaderValue(CommandList, ComputeShader, Shader->ComponentBufferSizeReadParam, SafeBufferSize);
+		SetSRVParameter(CommandList, ComputeShader, Shader->FreeIDBufferParam, GPUFreeIDs.SRV ? GPUFreeIDs.SRV : FNiagaraRenderer::GetDummyIntBuffer().SRV);
 	}
 	else
 	{
 		Shader->FloatOutputBufferParam.SetBuffer(CommandList, ComputeShader, GetGPUBufferFloat());
 		Shader->IntOutputBufferParam.SetBuffer(CommandList, ComputeShader, GetGPUBufferInt());
 		SetShaderValue(CommandList, ComputeShader, Shader->ComponentBufferSizeWriteParam, SafeBufferSize);
+		if (Shader->IDToIndexBufferParam.IsUAVBound())
+		{
+			check(GPUIDToIndexTable.Buffer);
+			Shader->IDToIndexBufferParam.SetBuffer(CommandList, ComputeShader, GPUIDToIndexTable);
+		}
 	}
 }
 
@@ -932,6 +978,13 @@ void FNiagaraDataBuffer::UnsetShaderParams(FNiagaraShader *Shader, FRHICommandLi
 #endif
 		//RHICmdList.TransitionResource(EResourceTransitionAccess::EReadable, EResourceTransitionPipeline::EGfxToGfx, CurrDataRender().GetGPUBufferInt()->UAV);
 	}
+
+	if (Shader->IDToIndexBufferParam.IsUAVBound())
+	{
+#if !PLATFORM_PS4
+		Shader->IDToIndexBufferParam.UnsetUAV(RHICmdList, Shader->GetComputeShader());
+#endif
+	}
 }
 
 void FNiagaraDataBuffer::ReleaseGPUInstanceCount(FNiagaraGPUInstanceCountManager& GPUInstanceCountManager)
diff --git a/Engine/Plugins/FX/Niagara/Source/Niagara/Private/NiagaraEmitterInstanceBatcher.cpp b/Engine/Plugins/FX/Niagara/Source/Niagara/Private/NiagaraEmitterInstanceBatcher.cpp
--- a/Engine/Plugins/FX/Niagara/Source/Niagara/Private/NiagaraEmitterInstanceBatcher.cpp
+++ b/Engine/Plugins/FX/Niagara/Source/Niagara/Private/NiagaraEmitterInstanceBatcher.cpp
@@ -8,6 +8,7 @@
 #include "NiagaraShader.h"
 #include "NiagaraSortingGPU.h"
 #include "NiagaraWorldManager.h"
+#include "NiagaraShaderParticleID.h"
 #include "ShaderParameterUtils.h"
 #include "SceneUtils.h"
 #include "ShaderParameterUtils.h"
@@ -25,6 +26,8 @@ DECLARE_DWORD_COUNTER_STAT(TEXT("Readback latency (frames)"), STAT_NiagaraReadba
 
 DECLARE_GPU_STAT_NAMED(NiagaraGPU, TEXT("Niagara"));
 DECLARE_GPU_STAT_NAMED(NiagaraGPUSimulation, TEXT("Niagara GPU Simulation"));
+DECLARE_GPU_STAT_NAMED(NiagaraGPUClearIDTable, TEXT("NiagaraGPU Clear ID Table"));
+DECLARE_GPU_STAT_NAMED(NiagaraGPUComputeFreeIDs, TEXT("Niagara GPU Compute Free IDs"));
 DECLARE_GPU_STAT_NAMED(NiagaraGPUSorting, TEXT("Niagara GPU sorting"));
 
 uint32 FNiagaraComputeExecutionContext::TickCounter = 0;
@@ -266,7 +269,7 @@ void NiagaraEmitterInstanceBatcher::PostStageInterface(const FNiagaraGPUSystemTi
 	}
 }
 
-void NiagaraEmitterInstanceBatcher::DispatchMultipleStages(const FNiagaraGPUSystemTick& Tick, FNiagaraComputeInstanceData *Instance, FRHICommandList &RHICmdList, FRHIUniformBuffer* ViewUniformBuffer, FNiagaraShader* ComputeShader) const
+void NiagaraEmitterInstanceBatcher::DispatchMultipleStages(const FNiagaraGPUSystemTick& Tick, FNiagaraComputeInstanceData *Instance, FRHICommandList &RHICmdList, FRHIUniformBuffer* ViewUniformBuffer, FNiagaraShader* ComputeShader)
 {
 	FNiagaraComputeExecutionContext* Context = Instance->Context;
 
@@ -346,10 +349,12 @@ void NiagaraEmitterInstanceBatcher::DispatchMultipleStages(const FNiagaraGPUSyst
 	}
 }
 
-void NiagaraEmitterInstanceBatcher::ResizeBuffersAndGatherResources(FOverlappableTicks& OverlappableTick, FRHICommandList& RHICmdList, FNiagaraBufferArray& DestDataBuffers, FNiagaraBufferArray& CurrDataBuffers, FNiagaraBufferArray& DestBufferIntFloat, FNiagaraBufferArray& CurrBufferIntFloat)
+void NiagaraEmitterInstanceBatcher::ResizeBuffersAndGatherResources(FOverlappableTicks& OverlappableTick, FRHICommandList& RHICmdList, FNiagaraBufferArray& ReadBuffers, FNiagaraBufferArray& WriteBuffers)
 {
 	SCOPE_CYCLE_COUNTER(STAT_NiagaraGPUDispatchSetup_RT);
 
+	uint32 NumFreeIDListsToCompute = 0;
+
 	//UE_LOG(LogNiagara, Warning, TEXT("NiagaraEmitterInstanceBatcher::ResizeBuffersAndGatherResources:  %0xP"), this);
 	for (FNiagaraGPUSystemTick* Tick : OverlappableTick)
 	{
@@ -373,6 +378,8 @@ void NiagaraEmitterInstanceBatcher::ResizeBuffersAndGatherResources(FOverlappabl
 				continue;
 			}
 
+			const bool bNeedsPersistentIDs = Context->MainDataSet->GetNeedsPersistentIDs();
+
 			//The buffer containing current simulation state.
 			Instance.CurrentData = Context->MainDataSet->GetCurrentData();
 			//The buffer we're going to write simulation results to.
@@ -390,25 +397,34 @@ void NiagaraEmitterInstanceBatcher::ResizeBuffersAndGatherResources(FOverlappabl
 			const uint32 RequiredInstances = FMath::Max(PrevNumInstances, NewNumInstances);
 			const uint32 AllocatedInstances = FMath::Max(RequiredInstances, Instance.SpawnInfo.MaxParticleCount);
 
-			DestinationData.AllocateGPU(AllocatedInstances + 1, GPUInstanceCounterManager, RHICmdList);
+			DestinationData.AllocateGPU(AllocatedInstances + 1, GPUInstanceCounterManager, RHICmdList, FeatureLevel);
 			DestinationData.SetNumInstances(RequiredInstances);
 
+			WriteBuffers.Add(GPUInstanceCounterManager.GetInstanceCountBuffer().UAV);
+
 			if ( Shader->FloatInputBufferParam.IsBound() )
 			{
-				CurrDataBuffers.Add(CurrentData.GetGPUBufferFloat().UAV);
+				ReadBuffers.Add(CurrentData.GetGPUBufferFloat().UAV);
 			}
 			if ( Shader->IntInputBufferParam.IsBound() )
 			{
-				CurrBufferIntFloat.Add(CurrentData.GetGPUBufferInt().UAV);
+				ReadBuffers.Add(CurrentData.GetGPUBufferInt().UAV);
 			}
 
 			if ( Shader->FloatOutputBufferParam.IsBound() )
 			{
-				DestDataBuffers.Add(DestinationData.GetGPUBufferFloat().UAV);
+				WriteBuffers.Add(DestinationData.GetGPUBufferFloat().UAV);
 			}
 			if ( Shader->IntOutputBufferParam.IsBound() )
 			{
-				DestBufferIntFloat.Add(DestinationData.GetGPUBufferInt().UAV);
+				WriteBuffers.Add(DestinationData.GetGPUBufferInt().UAV);
+			}
+
+			if (bNeedsPersistentIDs)
+			{
+				ReadBuffers.Add(CurrentData.GetGPUFreeIDs().UAV);
+				WriteBuffers.Add(CurrentData.GetGPUIDToIndexTable().UAV);
+				++NumFreeIDListsToCompute;
 			}
 
 			Context->MainDataSet->EndSimulate();
@@ -418,9 +434,28 @@ void NiagaraEmitterInstanceBatcher::ResizeBuffersAndGatherResources(FOverlappabl
 			}
 		}
 	}
+
+	if (NumFreeIDListsToCompute > NumAllocatedFreeIDListSizes)
+	{
+		constexpr uint32 ALLOC_CHUNK_SIZE = 128;
+		NumAllocatedFreeIDListSizes = FMath::DivideAndRoundUp(NumFreeIDListsToCompute, ALLOC_CHUNK_SIZE) * ALLOC_CHUNK_SIZE;
+		if (FreeIDListSizesBuffer.Buffer)
+		{
+			FreeIDListSizesBuffer.Release();
+		}
+		FreeIDListSizesBuffer.Initialize(sizeof(uint32), NumAllocatedFreeIDListSizes, EPixelFormat::PF_R32_UINT, BUF_Static, TEXT("NiagaraFreeIDListSizes"));
+	}
+
+	CurrentFreeIDListIndex = 0;
+	if (NumFreeIDListsToCompute > 0)
+	{
+		SCOPED_DRAW_EVENT(RHICmdList, NiagaraGPUClearFreeIDListSizeTable);
+		RHICmdList.TransitionResource(EResourceTransitionAccess::ERWBarrier, EResourceTransitionPipeline::EComputeToCompute, FreeIDListSizesBuffer.UAV);
+		ClearUAV(RHICmdList, FreeIDListSizesBuffer, 0);
+	}
 }
 
-void NiagaraEmitterInstanceBatcher::DispatchAllOnCompute(FOverlappableTicks& OverlappableTick, FRHICommandList& RHICmdList, FRHIUniformBuffer* ViewUniformBuffer, FNiagaraBufferArray& DestDataBuffers, FNiagaraBufferArray& CurrDataBuffers, FNiagaraBufferArray& DestBufferIntFloat, FNiagaraBufferArray& CurrBufferIntFloat, bool bSetReadback)
+void NiagaraEmitterInstanceBatcher::DispatchAllOnCompute(FOverlappableTicks& OverlappableTick, FRHICommandList& RHICmdList, FRHIUniformBuffer* ViewUniformBuffer, FNiagaraBufferArray& ReadBuffers, FNiagaraBufferArray& WriteBuffers, bool bSetReadback)
 {
 	FRHICommandListImmediate& RHICmdListImmediate = FRHICommandListExecutor::GetImmediateCommandList();
 
@@ -430,11 +465,6 @@ void NiagaraEmitterInstanceBatcher::DispatchAllOnCompute(FOverlappableTicks& Ove
 
 #if WITH_EDITORONLY_DATA
 	{
-		//
-		//	Transition current index buffer ready for compute and clear then all using overlapping compute work items.
-		//
-		RHICmdList.TransitionResources(EResourceTransitionAccess::ERWBarrier, EResourceTransitionPipeline::EGfxToCompute, CurrDataBuffers.GetData(), CurrDataBuffers.Num());
-
 		for (FNiagaraGPUSystemTick* Tick : OverlappableTick)
 		{
 			uint32 DispatchCount = Tick->Count;
@@ -455,18 +485,9 @@ void NiagaraEmitterInstanceBatcher::DispatchAllOnCompute(FOverlappableTicks& Ove
 	}
 #endif // WITH_EDITORONLY_DATA
 
-	//
-	//	Add a rw barrier for the destination data buffers we just cleared and mark others as read/write as needed for particles simulation.
-	//
-	{
-		TArray<FRHIUnorderedAccessView*, TInlineAllocator<5>> TransitionUAVs;
-		TransitionUAVs.Add(GPUInstanceCounterManager.GetInstanceCountBuffer().UAV);
-		TransitionUAVs.Append(DestDataBuffers.GetData(), DestDataBuffers.Num());
-		TransitionUAVs.Append(DestBufferIntFloat.GetData(), DestBufferIntFloat.Num());
-		TransitionUAVs.Append(CurrDataBuffers.GetData(), CurrDataBuffers.Num());
-		TransitionUAVs.Append(CurrBufferIntFloat.GetData(), CurrBufferIntFloat.Num());
-		RHICmdList.TransitionResources(EResourceTransitionAccess::ERWBarrier, EResourceTransitionPipeline::EGfxToCompute, TransitionUAVs.GetData(), TransitionUAVs.Num());
-	}
+	// Transition the input buffers to readable and the outputs to writable on the compute pipe.
+	RHICmdList.TransitionResources(EResourceTransitionAccess::EReadable, EResourceTransitionPipeline::EComputeToCompute, ReadBuffers.GetData(), ReadBuffers.Num());
+	RHICmdList.TransitionResources(EResourceTransitionAccess::ERWBarrier, EResourceTransitionPipeline::EComputeToCompute, WriteBuffers.GetData(), WriteBuffers.Num());
 
 	for (FNiagaraGPUSystemTick* Tick : OverlappableTick)
 	{
@@ -499,16 +520,9 @@ void NiagaraEmitterInstanceBatcher::DispatchAllOnCompute(FOverlappableTicks& Ove
 		}
 	}
 
-	//
-	//	Now Copy to staging buffer the data we want to read back (alive particle count). And make buffer ready for that and draw commands on the graphics pipe too.
-	//
-	{
-		TArray<FRHIUnorderedAccessView*, TInlineAllocator<5>> TransitionUAVs;
-		TransitionUAVs.Add(GPUInstanceCounterManager.GetInstanceCountBuffer().UAV);
-		TransitionUAVs.Append(DestDataBuffers.GetData(), DestDataBuffers.Num());
-		TransitionUAVs.Append(DestBufferIntFloat.GetData(), DestBufferIntFloat.Num());
-		RHICmdList.TransitionResources(EResourceTransitionAccess::EReadable, EResourceTransitionPipeline::EComputeToGfx, TransitionUAVs.GetData(), TransitionUAVs.Num());
-	}
+	// The output buffers will be used by rendering, so transition them to readable on the graphics pipe.
+	RHICmdList.TransitionResources(EResourceTransitionAccess::EReadable, EResourceTransitionPipeline::EComputeToGfx, WriteBuffers.GetData(), WriteBuffers.Num());
+
 	// We have done all our overlapping compute work on this list so go back to default behavior and flush.
 	RHICmdList.AutomaticCacheFlushAfterComputeShader(true);
 
@@ -646,18 +660,15 @@ void NiagaraEmitterInstanceBatcher::ExecuteAll(FRHICommandList &RHICmdList, FRHI
 
 	for (auto& SimPass : SimPasses)
 	{
-		FNiagaraBufferArray DestDataBuffers;
-		FNiagaraBufferArray CurrDataBuffers;
-		FNiagaraBufferArray DestBufferIntFloat;
-		FNiagaraBufferArray CurrBufferIntFloat;
+		FNiagaraBufferArray ReadBuffers, WriteBuffers;
 
 		// This initial pass gathers all the buffers that are read from and written to so we can do batch resource transitions.
 		// It also ensures the GPU buffers are large enough to hold everything.
-		ResizeBuffersAndGatherResources(SimPass, RHICmdList, DestDataBuffers, CurrDataBuffers, DestBufferIntFloat, CurrBufferIntFloat);
+		ResizeBuffersAndGatherResources(SimPass, RHICmdList, ReadBuffers, WriteBuffers);
 
 		SCOPED_DRAW_EVENT(RHICmdList, NiagaraGPUSimulation);
 		SCOPED_GPU_STAT(RHICmdList, NiagaraGPUSimulation);
-		DispatchAllOnCompute(SimPass, RHICmdList, ViewUniformBuffer, DestDataBuffers, CurrDataBuffers, DestBufferIntFloat, CurrBufferIntFloat, bSetReadback);
+		DispatchAllOnCompute(SimPass, RHICmdList, ViewUniformBuffer, ReadBuffers, WriteBuffers, bSetReadback);
 	}
 }
 
@@ -1161,17 +1172,17 @@ void NiagaraEmitterInstanceBatcher::UnsetDataInterfaceParameters(const TArray<FN
 /* Kick off a simulation/spawn run
  */
 void NiagaraEmitterInstanceBatcher::Run(const FNiagaraGPUSystemTick& Tick, const FNiagaraComputeInstanceData* Instance, uint32 UpdateStartInstance, const uint32 TotalNumInstances, FNiagaraShader* Shader,
-	FRHICommandList &RHICmdList, FRHIUniformBuffer* ViewUniformBuffer, const FNiagaraGpuSpawnInfo& SpawnInfo, bool bCopyBeforeStart, uint32 DefaultShaderStageIndex, uint32 ShaderStageIndex, FNiagaraDataInterfaceProxy *IterationInterface, bool HasRunParticleStage) const
+	FRHICommandList &RHICmdList, FRHIUniformBuffer* ViewUniformBuffer, const FNiagaraGpuSpawnInfo& SpawnInfo, bool bCopyBeforeStart, uint32 DefaultShaderStageIndex, uint32 ShaderStageIndex, FNiagaraDataInterfaceProxy *IterationInterface, bool HasRunParticleStage)
 {
 	FNiagaraComputeExecutionContext* Context = Instance->Context;
+
+	SCOPED_DRAW_EVENTF(RHICmdList, NiagaraGPUSimulationCS, TEXT("Niagara Gpu Sim - %s - NumInstances: %u - StageNumber: %u"),
+		Context->GetDebugSimName(),
+		TotalNumInstances,
+		ShaderStageIndex);
+
 	if (TotalNumInstances == 0)
 	{
-#if !UE_BUILD_SHIPPING
-		SCOPED_DRAW_EVENTF(RHICmdList, NiagaraGPUSimulationCS, TEXT("Niagara Gpu Sim - %s - NumInstances: %u - StageNumber: %u"),
-			Context->GetDebugSimName(),
-			TotalNumInstances,
-			ShaderStageIndex);
-#endif
 		return;
 	}
 
@@ -1183,6 +1194,16 @@ void NiagaraEmitterInstanceBatcher::Run(const FNiagaraGPUSystemTick& Tick, const
 	FNiagaraDataBuffer& DestinationData = *Instance->DestinationData;
 	FNiagaraDataBuffer& CurrentData = *Instance->CurrentData;
 
+	if (Context->MainDataSet->GetNeedsPersistentIDs())
+	{
+		// Put INDEX_NONE in all the slots of the ID to index table. The simulation will fill in the
+		// indices for the IDs which are currently in use, so we can compute the free ID list based on
+		// the unused slots. No need for an explicit barrier here, ClearUAV() inserts a RW barrier.
+		SCOPED_DRAW_EVENT(RHICmdList, NiagaraGPUClearIDTable);
+		SCOPED_GPU_STAT(RHICmdList, NiagaraGPUClearIDTable);
+		ClearUAV(RHICmdList, DestinationData.GetGPUIDToIndexTable(), INDEX_NONE);
+	}
+
 	RHICmdList.SetComputeShader(Shader->GetComputeShader());
 
 	// #todo(dmp): clean up this logic for shader stages on first frame
@@ -1300,12 +1321,6 @@ void NiagaraEmitterInstanceBatcher::Run(const FNiagaraGPUSystemTick& Tick, const
 	// Dispatch, if anything needs to be done
 	if (TotalNumInstances)
 	{
-#if !UE_BUILD_SHIPPING
-		SCOPED_DRAW_EVENTF(RHICmdList, NiagaraGPUSimulationCS, TEXT("Niagara Gpu Sim - %s - NumInstances: %u - StageNumber: %u"),
-			Context->GetDebugSimName(),
-			TotalNumInstances,
-			ShaderStageIndex);
-#endif
 		DispatchComputeShader(RHICmdList, Shader, NumThreadGroups, 1, 1);
 	}
 
@@ -1360,4 +1375,18 @@ void NiagaraEmitterInstanceBatcher::Run(const FNiagaraGPUSystemTick& Tick, const
 	CurrentData.UnsetShaderParams(Shader, RHICmdList);
 	DestinationData.UnsetShaderParams(Shader, RHICmdList);
 	Shader->InstanceCountsParam.UnsetUAV(RHICmdList, Shader->GetComputeShader());
+
+	if (Context->MainDataSet->GetNeedsPersistentIDs() && ShaderStageIndex == Context->MaxUpdateIterations - 1)
+	{
+		SCOPED_DRAW_EVENT(RHICmdList, NiagaraGPUComputeFreeIDs);
+		SCOPED_GPU_STAT(RHICmdList, NiagaraGPUComputeFreeIDs);
+
+		check(CurrentFreeIDListIndex < NumAllocatedFreeIDListSizes);
+
+		// Make the ID to index table readable and the free ID list writable.
+		RHICmdList.TransitionResource(EResourceTransitionAccess::EReadable, EResourceTransitionPipeline::EComputeToCompute, DestinationData.GetGPUIDToIndexTable().UAV);
+		RHICmdList.TransitionResource(EResourceTransitionAccess::ERWBarrier, EResourceTransitionPipeline::EComputeToCompute, DestinationData.GetGPUFreeIDs().UAV);
+		NiagaraComputeGPUFreeIDs(RHICmdList, FeatureLevel, DestinationData.GetGPUNumAllocatedIDs(), DestinationData.GetGPUIDToIndexTable().SRV, DestinationData.GetGPUFreeIDs(), FreeIDListSizesBuffer, CurrentFreeIDListIndex);
+		++CurrentFreeIDListIndex;
+	}
 }
diff --git a/Engine/Plugins/FX/Niagara/Source/NiagaraCore/Private/NiagaraCustomVersion.cpp b/Engine/Plugins/FX/Niagara/Source/NiagaraCore/Private/NiagaraCustomVersion.cpp
--- a/Engine/Plugins/FX/Niagara/Source/NiagaraCore/Private/NiagaraCustomVersion.cpp
+++ b/Engine/Plugins/FX/Niagara/Source/NiagaraCore/Private/NiagaraCustomVersion.cpp
@@ -8,5 +8,4 @@ const FGuid FNiagaraCustomVersion::GUID(0xFCF57AFA, 0x50764283, 0xB9A9E658, 0xFF
 // Register the custom version with core
 FCustomVersionRegistration GRegisterNiagaraCustomVersion(FNiagaraCustomVersion::GUID, FNiagaraCustomVersion::LatestVersion, TEXT("NiagaraVer"));
 
-const FGuid FNiagaraCustomVersion::LatestScriptCompileVersion(0x7698EEE8, 0x01EA483B, 0xAA2841E4, 0xCC6BED0D);
-
+const FGuid FNiagaraCustomVersion::LatestScriptCompileVersion(0x5A6DBD4C, 0x817E476A, 0x91039D69, 0xC9918D2E);
diff --git a/Engine/Plugins/FX/Niagara/Source/NiagaraEditor/Private/NiagaraHlslTranslator.cpp b/Engine/Plugins/FX/Niagara/Source/NiagaraEditor/Private/NiagaraHlslTranslator.cpp
--- a/Engine/Plugins/FX/Niagara/Source/NiagaraEditor/Private/NiagaraHlslTranslator.cpp
+++ b/Engine/Plugins/FX/Niagara/Source/NiagaraEditor/Private/NiagaraHlslTranslator.cpp
@@ -809,10 +809,6 @@ const FNiagaraTranslateResults &FHlslNiagaraTranslator::Translate(const FNiagara
 	}
 
 	bool bNeedsPersistentIDs = CompileOptions.AdditionalDefines.Contains(TEXT("RequiresPersistentIDs"));
-	if (bNeedsPersistentIDs && CompilationTarget == ENiagaraSimTarget::GPUComputeSim) {
-		Error(LOCTEXT("GPUPersistentIDFail", "GPU particles do not support persistent IDs. Change to a CPU simulation or disable persistent IDs."), nullptr, nullptr);
-		return TranslateResults;
-	}
 
 	TranslationStages.Empty();
 	ActiveStageIdx = 0;
@@ -1839,6 +1835,8 @@ void FHlslNiagaraTranslator::DefineMainGPUFunctions(
 		return false;
 	}();
 
+	const bool bNeedsPersistentIDs = CompileOptions.AdditionalDefines.Contains(TEXT("RequiresPersistentIDs"));
+
 	// A list of constant to reset after Emitter_SpawnGroup gets modified by GetEmitterSpawnInfoForParticle()
 	TArray<FString> EmitterSpawnGroupReinit;
 
@@ -1895,6 +1893,13 @@ void FHlslNiagaraTranslator::DefineMainGPUFunctions(
 		{
 			HlslOutput += TEXT("\n") + ContextName + TEXT("DataInstance.Alive=true;\n");
 		}
+
+		if (bNeedsPersistentIDs)
+		{
+			HlslOutput += TEXT("\n\tint IDIndex, IDAcquireTag;\n\tAcquireID(0, IDIndex, IDAcquireTag);\n");
+			HlslOutput += ContextName + TEXT("Particles.ID.Index = IDIndex;\n");
+			HlslOutput += ContextName + TEXT("Particles.ID.AcquireTag = IDAcquireTag;\n");
+		}
 	}
 	HlslOutput += TEXT("}\n\n");
 
@@ -1992,13 +1997,19 @@ void FHlslNiagaraTranslator::DefineMainGPUFunctions(
 			HlslOutput += TEXT("\tconst int WriteIndex = OutputIndex(0, true, bValid);\n");
 		}
 
-		HlslOutput += TEXT("\tif (bValid)\n\t{\n");
 		FString ContextName = TEXT("Context.Map.");
 		if (TranslationStages.Num() > 1) // Last context is "MapUpdate"
 		{
 			ContextName = FString::Printf(TEXT("Context.%s."), *TranslationStages.Last().PassNamespace);
 		}
 
+		if (bNeedsPersistentIDs)
+		{
+			HlslOutput += FString::Printf(TEXT("\tUpdateID(0, %sParticles.ID.Index, WriteIndex);\n"), *ContextName);
+		}
+
+		HlslOutput += TEXT("\tif (bValid)\n\t{\n");
+
 		for (int32 DataSetIndex = 0, IntCounter = 0, FloatCounter = 0; DataSetIndex < DataSetWrites.Num(); ++DataSetIndex)
 		{
 			const FNiagaraDataSetID DataSetID = ReadDataSetIDs[DataSetIndex];
diff --git a/Engine/Plugins/FX/Niagara/Source/NiagaraShader/Private/NiagaraShader.cpp b/Engine/Plugins/FX/Niagara/Source/NiagaraShader/Private/NiagaraShader.cpp
--- a/Engine/Plugins/FX/Niagara/Source/NiagaraShader/Private/NiagaraShader.cpp
+++ b/Engine/Plugins/FX/Niagara/Source/NiagaraShader/Private/NiagaraShader.cpp
@@ -1314,6 +1314,9 @@ void FNiagaraShader::BindParams(const FShaderParameterMap &ParameterMap)
 	InstanceCountsParam.Bind(ParameterMap, TEXT("InstanceCounts"));
 	ReadInstanceCountOffsetParam.Bind(ParameterMap, TEXT("ReadInstanceCountOffset"));
 	WriteInstanceCountOffsetParam.Bind(ParameterMap, TEXT("WriteInstanceCountOffset"));
+
+	FreeIDBufferParam.Bind(ParameterMap, TEXT("FreeIDList"));
+	IDToIndexBufferParam.Bind(ParameterMap, TEXT("IDToIndexTable"));
 	
 	SimStartParam.Bind(ParameterMap, TEXT("SimStart"));
 	EmitterTickCounterParam.Bind(ParameterMap, TEXT("EmitterTickCounter"));
@@ -1387,6 +1390,9 @@ bool FNiagaraShader::Serialize(FArchive& Ar)
 	Ar << ReadInstanceCountOffsetParam;
 	Ar << WriteInstanceCountOffsetParam;
 
+	Ar << FreeIDBufferParam;
+	Ar << IDToIndexBufferParam;
+
 	Ar << SimStartParam;
 	Ar << EmitterTickCounterParam;
 	Ar << EmitterSpawnInfoOffsetsParam;
diff --git a/Engine/Plugins/FX/Niagara/Source/NiagaraShader/Private/NiagaraShaderParticleID.cpp b/Engine/Plugins/FX/Niagara/Source/NiagaraShader/Private/NiagaraShaderParticleID.cpp
new file mode 100644
--- /dev/null
+++ b/Engine/Plugins/FX/Niagara/Source/NiagaraShader/Private/NiagaraShaderParticleID.cpp
@@ -0,0 +1,151 @@
+#include "NiagaraShaderParticleID.h"
+#include "GlobalShader.h"
+
+class FNiagaraInitFreeIDBufferCS : public FGlobalShader
+{
+	DECLARE_GLOBAL_SHADER(FNiagaraInitFreeIDBufferCS);
+
+public:
+	FNiagaraInitFreeIDBufferCS() : FGlobalShader() {}
+
+	FNiagaraInitFreeIDBufferCS(const ShaderMetaType::CompiledShaderInitializerType& Initializer) : FGlobalShader(Initializer)
+	{
+		NewBufferParam.Bind(Initializer.ParameterMap, TEXT("NewBuffer"));
+		ExistingBufferParam.Bind(Initializer.ParameterMap, TEXT("ExistingBuffer"));
+		NumNewElementsParam.Bind(Initializer.ParameterMap, TEXT("NumNewElements"));
+		NumExistingElementsParam.Bind(Initializer.ParameterMap, TEXT("NumExistingElements"));
+	}
+
+	static bool ShouldCompilePermutation(const FGlobalShaderPermutationParameters& Parameters)
+	{
+		return RHISupportsComputeShaders(Parameters.Platform);
+	}
+
+	static void ModifyCompilationEnvironment(const FGlobalShaderPermutationParameters& Parameters, FShaderCompilerEnvironment& OutEnvironment)
+	{
+		FGlobalShader::ModifyCompilationEnvironment(Parameters, OutEnvironment);
+		OutEnvironment.SetDefine(TEXT("THREAD_COUNT"), THREAD_COUNT);
+	}
+
+	bool Serialize(FArchive& Ar)
+	{
+		bool bShaderHasOutdatedParameters = FGlobalShader::Serialize(Ar);
+		Ar << NewBufferParam;
+		Ar << ExistingBufferParam;
+		Ar << NumNewElementsParam;
+		Ar << NumExistingElementsParam;
+		return bShaderHasOutdatedParameters;
+	}
+
+	void Execute(FRHICommandList& RHICmdList, uint32 NumElementsToAlloc, FRWBuffer& NewBuffer, uint32 NumExistingElements, FRHIShaderResourceView* ExistingBuffer)
+	{
+		// To simplify the shader code, the size of the ID table must be a multiple of the thread count.
+		check(NumElementsToAlloc % THREAD_COUNT == 0);
+
+		// Shrinking is not supported.
+		check(NumElementsToAlloc >= NumExistingElements);
+		uint32 NumNewElements = NumElementsToAlloc - NumExistingElements;
+
+		FRHIComputeShader* ComputeShader = GetComputeShader();
+		RHICmdList.SetComputeShader(ComputeShader);
+
+		NewBufferParam.SetBuffer(RHICmdList, ComputeShader, NewBuffer);
+		SetSRVParameter(RHICmdList, ComputeShader, ExistingBufferParam, ExistingBuffer);
+		SetShaderValue(RHICmdList, ComputeShader, NumNewElementsParam, NumNewElements);
+		SetShaderValue(RHICmdList, ComputeShader, NumExistingElementsParam, NumExistingElements);
+
+		DispatchComputeShader(RHICmdList, this, FMath::DivideAndRoundUp(NumElementsToAlloc, THREAD_COUNT), 1, 1);
+
+		RHICmdList.SetUAVParameter(ComputeShader, NewBufferParam.GetUAVIndex(), nullptr);
+		RHICmdList.SetShaderResourceViewParameter(ComputeShader, ExistingBufferParam.GetBaseIndex(), nullptr);
+	}
+
+private:
+	static constexpr uint32 THREAD_COUNT = 64;
+
+	FRWShaderParameter NewBufferParam;
+	FShaderResourceParameter ExistingBufferParam;
+	FShaderParameter NumNewElementsParam;
+	FShaderParameter NumExistingElementsParam;
+};
+
+IMPLEMENT_GLOBAL_SHADER(FNiagaraInitFreeIDBufferCS, "/Plugin/FX/Niagara/Private/NiagaraInitFreeIDBuffer.usf", "InitIDBufferCS", SF_Compute);
+
+void NiagaraInitGPUFreeIDList(FRHICommandList& RHICmdList, ERHIFeatureLevel::Type FeatureLevel, uint32 NumElementsToAlloc, FRWBuffer& NewBuffer, uint32 NumExistingElements, FRHIShaderResourceView* ExistingBuffer)
+{
+	TShaderMapRef<FNiagaraInitFreeIDBufferCS> InitIDBufferCS(GetGlobalShaderMap(FeatureLevel));
+	InitIDBufferCS->Execute(RHICmdList, NumElementsToAlloc, NewBuffer, NumExistingElements, ExistingBuffer);
+}
+
+class NiagaraComputeFreeIDsCS : public FGlobalShader
+{
+	DECLARE_GLOBAL_SHADER(NiagaraComputeFreeIDsCS);
+
+public:
+	NiagaraComputeFreeIDsCS() : FGlobalShader() {}
+
+	NiagaraComputeFreeIDsCS(const ShaderMetaType::CompiledShaderInitializerType& Initializer) : FGlobalShader(Initializer)
+	{
+		IDToIndexTableParam.Bind(Initializer.ParameterMap, TEXT("IDToIndexTable"));
+		FreeIDListParam.Bind(Initializer.ParameterMap, TEXT("FreeIDList"));
+		FreeIDListSizesParam.Bind(Initializer.ParameterMap, TEXT("FreeIDListSizes"));
+		FreeIDListIndexParam.Bind(Initializer.ParameterMap, TEXT("FreeIDListIndex"));
+	}
+
+	static bool ShouldCompilePermutation(const FGlobalShaderPermutationParameters& Parameters)
+	{
+		return RHISupportsComputeShaders(Parameters.Platform);
+	}
+
+	static void ModifyCompilationEnvironment(const FGlobalShaderPermutationParameters& Parameters, FShaderCompilerEnvironment& OutEnvironment)
+	{
+		FGlobalShader::ModifyCompilationEnvironment(Parameters, OutEnvironment);
+	}
+
+	bool Serialize(FArchive& Ar)
+	{
+		bool bShaderHasOutdatedParameters = FGlobalShader::Serialize(Ar);
+		Ar << IDToIndexTableParam;
+		Ar << FreeIDListParam;
+		Ar << FreeIDListSizesParam;
+		Ar << FreeIDListIndexParam;
+		return bShaderHasOutdatedParameters;
+	}
+
+	void Execute(FRHICommandList& RHICmdList, ERHIFeatureLevel::Type FeatureLevel, uint32 NumIDs, FRHIShaderResourceView* IDToIndexTable, FRWBuffer& FreeIDList, FRWBuffer& FreeIDListSizes, uint32 FreeIDListIndex)
+	{
+		const EShaderPlatform Platform = GShaderPlatformForFeatureLevel[FeatureLevel];
+		const uint32 THREAD_COUNT = (Platform == SP_XBOXONE_D3D12 || Platform == SP_PS4) ? 64 : 128;
+
+		// To simplify the shader code, the size of the ID table must be a multiple of the thread count.
+		check(NumIDs % THREAD_COUNT == 0);
+
+		FRHIComputeShader* ComputeShader = GetComputeShader();
+		RHICmdList.SetComputeShader(ComputeShader);
+
+		SetSRVParameter(RHICmdList, ComputeShader, IDToIndexTableParam, IDToIndexTable);
+		FreeIDListParam.SetBuffer(RHICmdList, ComputeShader, FreeIDList);
+		FreeIDListSizesParam.SetBuffer(RHICmdList, ComputeShader, FreeIDListSizes);
+		SetShaderValue(RHICmdList, ComputeShader, FreeIDListIndexParam, FreeIDListIndex);
+
+		DispatchComputeShader(RHICmdList, this, NumIDs / THREAD_COUNT, 1, 1);
+
+		RHICmdList.SetShaderResourceViewParameter(ComputeShader, IDToIndexTableParam.GetBaseIndex(), nullptr);
+		RHICmdList.SetUAVParameter(ComputeShader, FreeIDListParam.GetUAVIndex(), nullptr);
+		RHICmdList.SetUAVParameter(ComputeShader, FreeIDListSizesParam.GetUAVIndex(), nullptr);
+	}
+
+private:
+	FShaderResourceParameter IDToIndexTableParam;
+	FRWShaderParameter FreeIDListParam;
+	FRWShaderParameter FreeIDListSizesParam;
+	FShaderParameter FreeIDListIndexParam;
+};
+
+IMPLEMENT_GLOBAL_SHADER(NiagaraComputeFreeIDsCS, "/Plugin/FX/Niagara/Private/NiagaraComputeFreeIDs.usf", "ComputeFreeIDs", SF_Compute);
+
+void NiagaraComputeGPUFreeIDs(FRHICommandList& RHICmdList, ERHIFeatureLevel::Type FeatureLevel, uint32 NumIDs, FRHIShaderResourceView* IDToIndexTable, FRWBuffer& FreeIDList, FRWBuffer& FreeIDListSizes, uint32 FreeIDListIndex)
+{
+	TShaderMapRef<NiagaraComputeFreeIDsCS> ComputeFreeIDsCS(GetGlobalShaderMap(FeatureLevel));
+	ComputeFreeIDsCS->Execute(RHICmdList, FeatureLevel, NumIDs, IDToIndexTable, FreeIDList, FreeIDListSizes, FreeIDListIndex);
+}
diff --git a/Engine/Plugins/FX/Niagara/Source/NiagaraShader/Public/NiagaraShader.h b/Engine/Plugins/FX/Niagara/Source/NiagaraShader/Public/NiagaraShader.h
--- a/Engine/Plugins/FX/Niagara/Source/NiagaraShader/Public/NiagaraShader.h
+++ b/Engine/Plugins/FX/Niagara/Source/NiagaraShader/Public/NiagaraShader.h
@@ -107,6 +107,8 @@ public:
 	FRWShaderParameter InstanceCountsParam;
 	FShaderParameter ReadInstanceCountOffsetParam;
 	FShaderParameter WriteInstanceCountOffsetParam;
+	FShaderResourceParameter FreeIDBufferParam;
+	FRWShaderParameter IDToIndexBufferParam;
 	FShaderUniformBufferParameter EmitterConstantBufferParam;
 	FShaderUniformBufferParameter DataInterfaceUniformBufferParam;
 	FShaderUniformBufferParameter ViewUniformBufferParam;
diff --git a/Engine/Plugins/FX/Niagara/Source/NiagaraShader/Public/NiagaraShaderParticleID.h b/Engine/Plugins/FX/Niagara/Source/NiagaraShader/Public/NiagaraShaderParticleID.h
new file mode 100644
--- /dev/null
+++ b/Engine/Plugins/FX/Niagara/Source/NiagaraShader/Public/NiagaraShaderParticleID.h
@@ -0,0 +1,8 @@
+// Copyright 1998-2019 Epic Games, Inc. All Rights Reserved.
+
+#pragma once
+
+#include "CoreMinimal.h"
+
+NIAGARASHADER_API void NiagaraInitGPUFreeIDList(FRHICommandList& RHICmdList, ERHIFeatureLevel::Type FeatureLevel, uint32 NumElementsToAlloc, FRWBuffer& NewBuffer, uint32 NumExistingElements, FRHIShaderResourceView* ExistingBuffer);
+NIAGARASHADER_API void NiagaraComputeGPUFreeIDs(FRHICommandList& RHICmdList, ERHIFeatureLevel::Type FeatureLevel, uint32 NumIDs, FRHIShaderResourceView* IDToIndexTable, FRWBuffer& FreeIDList, FRWBuffer& FreeIDListSizes, uint32 FreeIDListIndex);
