Fix non-unity build and missing copyright notice.

diff --git a/Engine/Plugins/FX/Niagara/Source/NiagaraShader/Private/NiagaraShaderParticleID.cpp b/Engine/Plugins/FX/Niagara/Source/NiagaraShader/Private/NiagaraShaderParticleID.cpp
--- a/Engine/Plugins/FX/Niagara/Source/NiagaraShader/Private/NiagaraShaderParticleID.cpp
+++ b/Engine/Plugins/FX/Niagara/Source/NiagaraShader/Private/NiagaraShaderParticleID.cpp
@@ -1,3 +1,5 @@
+// Copyright 1998-2019 Epic Games, Inc. All Rights Reserved.
+
 #include "NiagaraShaderParticleID.h"
 #include "GlobalShader.h"
 
diff --git a/Engine/Plugins/FX/Niagara/Source/NiagaraShader/Public/NiagaraShaderParticleID.h b/Engine/Plugins/FX/Niagara/Source/NiagaraShader/Public/NiagaraShaderParticleID.h
--- a/Engine/Plugins/FX/Niagara/Source/NiagaraShader/Public/NiagaraShaderParticleID.h
+++ b/Engine/Plugins/FX/Niagara/Source/NiagaraShader/Public/NiagaraShaderParticleID.h
@@ -2,7 +2,7 @@
 
 #pragma once
 
-#include "CoreMinimal.h"
+#include "RHI.h"
 
 NIAGARASHADER_API void NiagaraInitGPUFreeIDList(FRHICommandList& RHICmdList, ERHIFeatureLevel::Type FeatureLevel, uint32 NumElementsToAlloc, FRWBuffer& NewBuffer, uint32 NumExistingElements, FRHIShaderResourceView* ExistingBuffer);
 NIAGARASHADER_API void NiagaraComputeGPUFreeIDs(FRHICommandList& RHICmdList, ERHIFeatureLevel::Type FeatureLevel, uint32 NumIDs, FRHIShaderResourceView* IDToIndexTable, FRWBuffer& FreeIDList, FRWBuffer& FreeIDListSizes, uint32 FreeIDListIndex);
