More fixes for non-unity builds.

diff --git a/Engine/Plugins/FX/Niagara/Source/NiagaraShader/Private/NiagaraShaderParticleID.cpp b/Engine/Plugins/FX/Niagara/Source/NiagaraShader/Private/NiagaraShaderParticleID.cpp
--- a/Engine/Plugins/FX/Niagara/Source/NiagaraShader/Private/NiagaraShaderParticleID.cpp
+++ b/Engine/Plugins/FX/Niagara/Source/NiagaraShader/Private/NiagaraShaderParticleID.cpp
@@ -2,6 +2,7 @@
 
 #include "NiagaraShaderParticleID.h"
 #include "GlobalShader.h"
+#include "ShaderParameterUtils.h"
 
 class FNiagaraInitFreeIDBufferCS : public FGlobalShader
 {
