Moved Isolate toggle to renderers section in the emitter node. Added Go to Parent button to emitters that have parents.

diff --git a/Engine/Plugins/FX/Niagara/Source/NiagaraEditorWidgets/Private/SNiagaraOverviewStackNode.cpp b/Engine/Plugins/FX/Niagara/Source/NiagaraEditorWidgets/Private/SNiagaraOverviewStackNode.cpp
--- a/Engine/Plugins/FX/Niagara/Source/NiagaraEditorWidgets/Private/SNiagaraOverviewStackNode.cpp
+++ b/Engine/Plugins/FX/Niagara/Source/NiagaraEditorWidgets/Private/SNiagaraOverviewStackNode.cpp
@@ -12,6 +12,7 @@
 #include "NiagaraEditorModule.h"
 #include "NiagaraEditorStyle.h"
 #include "NiagaraEditorWidgetsStyle.h"
+#include "NiagaraEmitter.h"
 #include "Stack/SNiagaraStackIssueIcon.h"
 #include "Modules/ModuleManager.h"
 #include "Widgets/Layout/SBox.h"
@@ -22,6 +23,7 @@
 #include "Widgets/SBoxPanel.h"
 #include "NiagaraRendererProperties.h"
 #include "Widgets/SWidget.h"
+#include "Subsystems/AssetEditorSubsystem.h"
 
 #define LOCTEXT_NAMESPACE "NiagaraOverviewStackNode"
 
@@ -97,39 +99,45 @@ TSharedRef<SWidget> SNiagaraOverviewStackNode::CreateTitleWidget(TSharedPtr<SNod
 			.IsChecked(this, &SNiagaraOverviewStackNode::GetEnabledCheckState)
 			.OnCheckStateChanged(this, &SNiagaraOverviewStackNode::OnEnabledCheckStateChanged)
 		]
-		// Isolate toggle
+		// Name
+		+ SHorizontalBox::Slot()
+		.Padding(3, 0, 0, 0)
+		.FillWidth(1.0f)
+		[
+			DefaultTitle
+		];
+}
+
+TSharedRef<SWidget> SNiagaraOverviewStackNode::CreateTitleRightWidget()
+{
+	return SNew(SHorizontalBox)
 		+ SHorizontalBox::Slot()
 		.AutoWidth()
+		.HAlign(HAlign_Right)
 		.VAlign(VAlign_Center)
-		.Padding(3, 0, 0, 0)
+		.Padding(0, 0, 1, 0)
 		[
 			SNew(SButton)
+			.IsFocusable(false)
+			.ToolTipText(LOCTEXT("OpenAndFocusParentEmitterToolTip", "Open and Focus Parent Emitter"))
 			.ButtonStyle(FEditorStyle::Get(), "HoverHintOnly")
-			.HAlign(HAlign_Center)
-			.ContentPadding(1)
-			.ToolTipText(this, &SNiagaraOverviewStackNode::GetToggleIsolateToolTip)
-			.OnClicked(this, &SNiagaraOverviewStackNode::OnToggleIsolateButtonClicked)
-			.Visibility(this, &SNiagaraOverviewStackNode::GetToggleIsolateVisibility)
+			.ForegroundColor(FNiagaraEditorWidgetsStyle::Get().GetColor("NiagaraEditor.Stack.ForegroundColor"))
+			.ContentPadding(2)
+			.OnClicked(this, &SNiagaraOverviewStackNode::OpenParentEmitter)
+			.Visibility(this, &SNiagaraOverviewStackNode::GetOpenParentEmitterVisibility)
+			.DesiredSizeScale(FVector2D(14.0f / 30.0f, 14.0f / 30.0f)) // GoToSourceIcon is 30x30, scale down
 			.Content()
 			[
 				SNew(SImage)
-				.Image(FNiagaraEditorStyle::Get().GetBrush("NiagaraEditor.Isolate"))
-				.ColorAndOpacity(this, &SNiagaraOverviewStackNode::GetToggleIsolateImageColor)
+				.Image(FNiagaraEditorWidgetsStyle::Get().GetBrush("NiagaraEditor.Stack.GoToSourceIcon"))
+				.ColorAndOpacity(FNiagaraEditorWidgetsStyle::Get().GetColor("NiagaraEditor.Stack.FlatButtonColor"))
 			]
 		]
-		// Name
 		+ SHorizontalBox::Slot()
-		.Padding(3, 0, 0, 0)
-		.FillWidth(1.0f)
-		[
-			DefaultTitle
-		];
-}
-
-TSharedRef<SWidget> SNiagaraOverviewStackNode::CreateTitleRightWidget()
-{
-	return SNew(SBox)
-		.Content()
+		.AutoWidth()
+		.HAlign(HAlign_Right)
+		.VAlign(VAlign_Center)
+		.Padding(0, 0, 1, 0)
 		[
 			SNew(SNiagaraStackIssueIcon, StackViewModel, StackViewModel->GetRootEntry())
 			.Visibility(this, &SNiagaraOverviewStackNode::GetIssueIconVisibility)
@@ -318,6 +326,28 @@ void SNiagaraOverviewStackNode::FillThumbnailBar(UObject* ChangedObject, const b
 		{
 			EmitterHandleViewModelWeak.Pin()->GetRendererPreviewData(PreviewData);
 
+			// Isolate toggle button
+			ThumbnailBar->AddSlot()
+				.AutoWidth()
+				.HAlign(HAlign_Left)
+				.VAlign(VAlign_Center)
+				.Padding(2, 0, 0, 0)
+				[
+					SNew(SButton)
+					.ButtonStyle(FEditorStyle::Get(), "HoverHintOnly")
+					.HAlign(HAlign_Center)
+					.ContentPadding(1)
+					.ToolTipText(this, &SNiagaraOverviewStackNode::GetToggleIsolateToolTip)
+					.OnClicked(this, &SNiagaraOverviewStackNode::OnToggleIsolateButtonClicked)
+					.Visibility(this, &SNiagaraOverviewStackNode::GetToggleIsolateVisibility)
+					.Content()
+					[
+						SNew(SImage)
+						.Image(FNiagaraEditorStyle::Get().GetBrush("NiagaraEditor.Isolate"))
+						.ColorAndOpacity(this, &SNiagaraOverviewStackNode::GetToggleIsolateImageColor)
+					]
+				];
+
 			for (FRendererPreviewData* Preview : PreviewData)
 			{
 				if (Preview->RenderingObject)
@@ -413,4 +443,28 @@ FSlateColor SNiagaraOverviewStackNode::GetToggleIsolateImageColor() const
 		: FLinearColor::Gray;
 }
 
+FReply SNiagaraOverviewStackNode::OpenParentEmitter()
+{
+	TSharedPtr<FNiagaraEmitterHandleViewModel> EmitterHandleViewModel = EmitterHandleViewModelWeak.Pin();
+
+	if (EmitterHandleViewModel.IsValid())
+	{
+		UNiagaraEmitter* ParentEmitter = const_cast<UNiagaraEmitter*>(EmitterHandleViewModel->GetEmitterViewModel()->GetParentEmitter());
+		if (ParentEmitter != nullptr)
+		{
+			GEditor->GetEditorSubsystem<UAssetEditorSubsystem>()->OpenEditorForAsset(ParentEmitter);
+		}
+	}
+	return FReply::Handled();
+}
+
+EVisibility SNiagaraOverviewStackNode::GetOpenParentEmitterVisibility() const
+{
+	TSharedPtr<FNiagaraEmitterHandleViewModel> EmitterHandleViewModel = EmitterHandleViewModelWeak.Pin();
+	return EmitterHandleViewModel.IsValid() && 
+		EmitterHandleViewModel->GetEmitterViewModel()->GetParentEmitter() != nullptr
+		? EVisibility::Visible 
+		: EVisibility::Collapsed;
+}
+
 #undef LOCTEXT_NAMESPACE
\ No newline at end of file
diff --git a/Engine/Plugins/FX/Niagara/Source/NiagaraEditorWidgets/Private/SNiagaraOverviewStackNode.h b/Engine/Plugins/FX/Niagara/Source/NiagaraEditorWidgets/Private/SNiagaraOverviewStackNode.h
--- a/Engine/Plugins/FX/Niagara/Source/NiagaraEditorWidgets/Private/SNiagaraOverviewStackNode.h
+++ b/Engine/Plugins/FX/Niagara/Source/NiagaraEditorWidgets/Private/SNiagaraOverviewStackNode.h
@@ -55,6 +55,8 @@ private:
 	}
 
 	FReply OnCycleThroughIssues();
+	FReply OpenParentEmitter();
+	EVisibility GetOpenParentEmitterVisibility() const;
 
 private:
 	UNiagaraOverviewNode* OverviewStackNode;
