Fix duplicated loctext key with different value
#rnx

diff --git a/Engine/Plugins/FX/Niagara/Source/NiagaraEditor/Private/NiagaraEditorUtilities.cpp b/Engine/Plugins/FX/Niagara/Source/NiagaraEditor/Private/NiagaraEditorUtilities.cpp
--- a/Engine/Plugins/FX/Niagara/Source/NiagaraEditor/Private/NiagaraEditorUtilities.cpp
+++ b/Engine/Plugins/FX/Niagara/Source/NiagaraEditor/Private/NiagaraEditorUtilities.cpp
@@ -1529,8 +1529,8 @@ bool FNiagaraEditorUtilities::AddEmitterContextMenuActions(FMenuBuilder& MenuBui
 			);
 
 			MenuBuilder.AddMenuEntry(
-				LOCTEXT("CreateAssetFromEmitter", "Create Asset From This"),
-				LOCTEXT("CreateAssetFromEmitterToolTip", "Create an emitter asset from this emitter."),
+				LOCTEXT("CreateAssetFromThisEmitter", "Create Asset From This"),
+				LOCTEXT("CreateAssetFromThisEmitterToolTip", "Create an emitter asset from this emitter."),
 				FSlateIcon(),
 				FUIAction(
 					FExecuteAction::CreateStatic(&FNiagaraEditorUtilities::CreateAssetFromEmitter, EmitterHandleViewModel),
