Niagara - Shader stage data and UI for future work.

diff --git a/Engine/Plugins/FX/Niagara/Source/Niagara/Classes/NiagaraEmitter.h b/Engine/Plugins/FX/Niagara/Source/Niagara/Classes/NiagaraEmitter.h
--- a/Engine/Plugins/FX/Niagara/Source/Niagara/Classes/NiagaraEmitter.h
+++ b/Engine/Plugins/FX/Niagara/Source/Niagara/Classes/NiagaraEmitter.h
@@ -16,6 +16,7 @@ class UMaterial;
 class UNiagaraEmitter;
 class UNiagaraEventReceiverEmitterAction;
 class UNiagaraRendererProperties;
+class UNiagaraShaderStageBase;
 
 //TODO: Event action that spawns other whole Systems?
 //One that calls a BP exposed delegate?
@@ -448,6 +449,16 @@ public:
 
 	void NIAGARA_API RemoveEventHandlerByUsageId(FGuid EventHandlerUsageId);
 
+	NIAGARA_API const TArray<UNiagaraShaderStageBase*>& GetShaderStages() const { return ShaderStages; }
+
+	NIAGARA_API UNiagaraShaderStageBase* GetShaderStageById(FGuid ScriptUsageId) const;
+
+	void NIAGARA_API AddShaderStage(UNiagaraShaderStageBase* ShaderStage);
+
+	void NIAGARA_API RemoveShaderStage(UNiagaraShaderStageBase* ShaderStage);
+
+	void NIAGARA_API MoveShaderStageToIndex(UNiagaraShaderStageBase* ShaderStage, int32 TargetIndex);
+
 	/* Gets whether or not the supplied event generator id matches an event generator which is shared between the particle spawn and update scrips. */
 	bool IsEventGeneratorShared(FName EventGeneratorId) const;
 
@@ -501,6 +512,8 @@ private:
 
 	void ScriptRapidIterationParameterChanged();
 
+	void ShaderStageChanged();
+
 	void RendererChanged();
 
 	void GraphSourceChanged();
@@ -523,6 +536,9 @@ private:
 	UPROPERTY(EditAnywhere, Category = "Events", meta=(NiagaraNoMerge))
 	TArray<FNiagaraEventScriptProperties> EventHandlerScriptProps;
 
+	UPROPERTY(meta = (NiagaraNoMerge))
+	TArray<UNiagaraShaderStageBase*> ShaderStages;
+
 	UPROPERTY()
 	UNiagaraScript* GPUComputeScript;
 
diff --git a/Engine/Plugins/FX/Niagara/Source/Niagara/Private/NiagaraEmitter.cpp b/Engine/Plugins/FX/Niagara/Source/Niagara/Private/NiagaraEmitter.cpp
--- a/Engine/Plugins/FX/Niagara/Source/Niagara/Private/NiagaraEmitter.cpp
+++ b/Engine/Plugins/FX/Niagara/Source/Niagara/Private/NiagaraEmitter.cpp
@@ -4,7 +4,8 @@
 #include "NiagaraEmitter.h"
 #include "NiagaraScript.h"
 #include "NiagaraScriptSourceBase.h"
-#include "NiagaraSpriteRendererProperties.h"
+#include "NiagaraRendererProperties.h"
+#include "NiagaraShaderStageBase.h"
 #include "NiagaraCustomVersion.h"
 #include "UObject/Package.h"
 #include "UObject/Linker.h"
@@ -453,6 +454,13 @@ void UNiagaraEmitter::PostLoad()
 				FNiagaraParameterStore::FOnChanged::FDelegate::CreateUObject(this, &UNiagaraEmitter::ScriptRapidIterationParameterChanged));
 		}
 
+		for (UNiagaraShaderStageBase* ShaderStage : ShaderStages)
+		{
+			ShaderStage->OnChanged().AddUObject(this, &UNiagaraEmitter::ShaderStageChanged);
+			ShaderStage->Script->RapidIterationParameters.AddOnChangedHandler(
+				FNiagaraParameterStore::FOnChanged::FDelegate::CreateUObject(this, &UNiagaraEmitter::ScriptRapidIterationParameterChanged));
+		}
+
 		for (UNiagaraRendererProperties* Renderer : RendererProperties)
 		{
 			Renderer->OnChanged().AddUObject(this, &UNiagaraEmitter::RendererChanged);
@@ -1009,6 +1017,23 @@ void UNiagaraEmitter::UpdateFromMergedCopy(const INiagaraMergeManager& MergeMana
 			FNiagaraParameterStore::FOnChanged::FDelegate::CreateUObject(this, &UNiagaraEmitter::ScriptRapidIterationParameterChanged));
 	}
 
+	// Copy shader stages
+	for (UNiagaraShaderStageBase*& ShaderStage : ShaderStages)
+	{
+		ShaderStage->OnChanged().RemoveAll(this);
+		ShaderStage->Script->RapidIterationParameters.RemoveAllOnChangedHandlers(this);
+	}
+	ShaderStages.Empty();
+
+	for (UNiagaraShaderStageBase* MergedShaderStage : MergedEmitter->ShaderStages)
+	{
+		ReouterMergedObject(this, MergedShaderStage);
+		ShaderStages.Add(MergedShaderStage);
+		MergedShaderStage->OnChanged().AddUObject(this, &UNiagaraEmitter::ShaderStageChanged);
+		MergedShaderStage->Script->RapidIterationParameters.AddOnChangedHandler(
+			FNiagaraParameterStore::FOnChanged::FDelegate::CreateUObject(this, &UNiagaraEmitter::ScriptRapidIterationParameterChanged));
+	}
+
 	// Copy renderers
 	for (UNiagaraRendererProperties* Renderer : RendererProperties)
 	{
@@ -1143,6 +1168,55 @@ void UNiagaraEmitter::RemoveEventHandlerByUsageId(FGuid EventHandlerUsageId)
 #endif
 }
 
+UNiagaraShaderStageBase* UNiagaraEmitter::GetShaderStageById(FGuid ScriptUsageId) const
+{
+	UNiagaraShaderStageBase*const* FoundShaderStagePtr = ShaderStages.FindByPredicate([&ScriptUsageId](UNiagaraShaderStageBase* ShaderStage) { return ShaderStage->Script->GetUsageId() == ScriptUsageId; });
+	return FoundShaderStagePtr != nullptr ? *FoundShaderStagePtr : nullptr;
+}
+
+void UNiagaraEmitter::AddShaderStage(UNiagaraShaderStageBase* ShaderStage)
+{
+	Modify();
+	ShaderStages.Add(ShaderStage);
+#if WITH_EDITOR
+	ShaderStage->OnChanged().AddUObject(this, &UNiagaraEmitter::ShaderStageChanged);
+	ShaderStage->Script->RapidIterationParameters.AddOnChangedHandler(
+		FNiagaraParameterStore::FOnChanged::FDelegate::CreateUObject(this, &UNiagaraEmitter::ScriptRapidIterationParameterChanged));
+	UpdateChangeId(TEXT("Shader stage added"));
+#endif
+}
+
+void UNiagaraEmitter::RemoveShaderStage(UNiagaraShaderStageBase* ShaderStage)
+{
+	Modify();
+	bool bRemoved = ShaderStages.Remove(ShaderStage) != 0;
+#if WITH_EDITOR
+	if (bRemoved)
+	{
+		ShaderStage->OnChanged().RemoveAll(this);
+		ShaderStage->Script->RapidIterationParameters.RemoveAllOnChangedHandlers(this);
+		UpdateChangeId(TEXT("Shader stage removed"));
+	}
+#endif
+}
+
+void UNiagaraEmitter::MoveShaderStageToIndex(UNiagaraShaderStageBase* ShaderStageToMove, int32 TargetIndex)
+{
+	int32 CurrentIndex = ShaderStages.IndexOfByKey(ShaderStageToMove);
+	checkf(CurrentIndex != INDEX_NONE, TEXT("Shader stage could not be moved because it is not owned by this emitter."));
+	if (TargetIndex != CurrentIndex)
+	{
+		int32 AdjustedTargetIndex = CurrentIndex < TargetIndex
+			? TargetIndex - 1 // If the current index is less than the target index, the target index needs to be decreased to make up for the item being removed.
+			: TargetIndex;
+
+		ShaderStages.Remove(ShaderStageToMove);
+		ShaderStages.Insert(ShaderStageToMove, AdjustedTargetIndex);
+
+		UpdateChangeId("Shader stage moved.");
+	}
+}
+
 bool UNiagaraEmitter::IsEventGeneratorShared(FName EventGeneratorId) const
 {
 	return SharedEventGeneratorIds.Contains(EventGeneratorId);
@@ -1180,6 +1254,11 @@ void UNiagaraEmitter::ScriptRapidIterationParameterChanged()
 	UpdateChangeId(TEXT("Script rapid iteration parameter changed."));
 }
 
+void UNiagaraEmitter::ShaderStageChanged()
+{
+	UpdateChangeId(TEXT("Shader Stage Changed"));
+}
+
 void UNiagaraEmitter::RendererChanged()
 {
 	UpdateChangeId(TEXT("Renderer changed."));
diff --git a/Engine/Plugins/FX/Niagara/Source/Niagara/Private/NiagaraShaderStageBase.cpp b/Engine/Plugins/FX/Niagara/Source/Niagara/Private/NiagaraShaderStageBase.cpp
new file mode 100644
--- /dev/null
+++ b/Engine/Plugins/FX/Niagara/Source/Niagara/Private/NiagaraShaderStageBase.cpp
@@ -0,0 +1,3 @@
+// Copyright 1998-2019 Epic Games, Inc. All Rights Reserved.
+
+#include "NiagaraShaderStageBase.h"
\ No newline at end of file
diff --git a/Engine/Plugins/FX/Niagara/Source/Niagara/Public/NiagaraCommon.h b/Engine/Plugins/FX/Niagara/Source/Niagara/Public/NiagaraCommon.h
--- a/Engine/Plugins/FX/Niagara/Source/Niagara/Public/NiagaraCommon.h
+++ b/Engine/Plugins/FX/Niagara/Source/Niagara/Public/NiagaraCommon.h
@@ -598,6 +598,8 @@ enum class ENiagaraScriptUsage : uint8
 	ParticleUpdateScript UMETA(Hidden),
 	/** The script is called to update particles in response to an event. */
 	ParticleEventScript UMETA(Hidden),
+	/** The script is called as a particle shader stage. */
+	ParticleShaderStageScript UMETA(Hidden),
 	/** The script is called to update particles on the GPU. */
 	ParticleGPUComputeScript UMETA(Hidden),
 	/** The script is called once when the emitter spawns. */
diff --git a/Engine/Plugins/FX/Niagara/Source/Niagara/Public/NiagaraShaderStageBase.h b/Engine/Plugins/FX/Niagara/Source/Niagara/Public/NiagaraShaderStageBase.h
new file mode 100644
--- /dev/null
+++ b/Engine/Plugins/FX/Niagara/Source/Niagara/Public/NiagaraShaderStageBase.h
@@ -0,0 +1,44 @@
+// Copyright 1998-2019 Epic Games, Inc. All Rights Reserved.
+
+#pragma once
+
+#include "NiagaraMergeable.h"
+#include "NiagaraShaderStageBase.generated.h"
+
+class UNiagaraScript;
+
+/**
+* A base class for niagara shader stages.  This class should be derived to add stage specific information.
+*/
+UCLASS()
+class NIAGARA_API UNiagaraShaderStageBase : public UNiagaraMergeable
+{
+	GENERATED_BODY()
+
+public:
+	UPROPERTY()
+	UNiagaraScript* Script;
+
+	UPROPERTY(EditAnywhere, Category = "Shader Stage")
+	FName ShaderStageName;
+};
+
+UCLASS(meta = (DisplayName = "Placeholder Stage A"))
+class NIAGARA_API UNiagaraShaderStagePlaceholderA : public UNiagaraShaderStageBase
+{
+	GENERATED_BODY()
+
+public:
+	UPROPERTY(EditAnywhere, Category="Shader Stage")
+	int32 Iterations;
+};
+
+UCLASS(meta = (DisplayName = "Placeholder Stage B"))
+class NIAGARA_API UNiagaraShaderStagePlaceholderB: public UNiagaraShaderStageBase
+{
+	GENERATED_BODY()
+
+public:
+	UPROPERTY(EditAnywhere, Category="Shader Stage")
+	FVector Direction;
+};
\ No newline at end of file
diff --git a/Engine/Plugins/FX/Niagara/Source/NiagaraEditor/Private/NiagaraScriptMergeManager.cpp b/Engine/Plugins/FX/Niagara/Source/NiagaraEditor/Private/NiagaraScriptMergeManager.cpp
--- a/Engine/Plugins/FX/Niagara/Source/NiagaraEditor/Private/NiagaraScriptMergeManager.cpp
+++ b/Engine/Plugins/FX/Niagara/Source/NiagaraEditor/Private/NiagaraScriptMergeManager.cpp
@@ -4,6 +4,7 @@
 #include "NiagaraEmitter.h"
 #include "NiagaraScript.h"
 #include "NiagaraScriptSource.h"
+#include "NiagaraShaderStageBase.h"
 #include "NiagaraRendererProperties.h"
 #include "EdGraphSchema_Niagara.h"
 #include "NiagaraGraph.h"
@@ -411,6 +412,79 @@ TSharedPtr<FNiagaraScriptStackMergeAdapter> FNiagaraEventHandlerMergeAdapter::Ge
 	return EventStack;
 }
 
+FNiagaraShaderStageMergeAdapter::FNiagaraShaderStageMergeAdapter(const UNiagaraEmitter& InEmitter, const UNiagaraShaderStageBase* InShaderStage, UNiagaraNodeOutput* InOutputNode)
+{
+	Initialize(InEmitter, InShaderStage, nullptr, InOutputNode);
+}
+
+FNiagaraShaderStageMergeAdapter::FNiagaraShaderStageMergeAdapter(const UNiagaraEmitter& InEmitter, UNiagaraShaderStageBase* InShaderStage, UNiagaraNodeOutput* InOutputNode)
+{
+	Initialize(InEmitter, InShaderStage, InShaderStage, InOutputNode);
+}
+
+FNiagaraShaderStageMergeAdapter::FNiagaraShaderStageMergeAdapter(const UNiagaraEmitter& InEmitter, UNiagaraNodeOutput* InOutputNode)
+{
+	Initialize(InEmitter, nullptr, nullptr, InOutputNode);
+}
+
+void FNiagaraShaderStageMergeAdapter::Initialize(const UNiagaraEmitter& InEmitter, const UNiagaraShaderStageBase* InShaderStage, UNiagaraShaderStageBase* InEditableShaderStage, UNiagaraNodeOutput* InOutputNode)
+{
+	Emitter = MakeWeakObjectPtr(const_cast<UNiagaraEmitter*>(&InEmitter));
+
+	ShaderStage = InShaderStage;
+	EditableShaderStage = InEditableShaderStage;
+
+	OutputNode = InOutputNode;
+
+	if (ShaderStage != nullptr && OutputNode != nullptr)
+	{
+		ShaderStageStack = MakeShared<FNiagaraScriptStackMergeAdapter>(*OutputNode.Get(), *ShaderStage->Script, Emitter->GetUniqueEmitterName());
+		InputNode = ShaderStageStack->GetInputNode();
+	}
+}
+
+const UNiagaraEmitter* FNiagaraShaderStageMergeAdapter::GetEmitter() const
+{
+	return Emitter.Get();
+}
+
+FGuid FNiagaraShaderStageMergeAdapter::GetUsageId() const
+{
+	if (ShaderStage != nullptr)
+	{
+		return ShaderStage->Script->GetUsageId();
+	}
+	else
+	{
+		return OutputNode->GetUsageId();
+	}
+}
+
+const UNiagaraShaderStageBase* FNiagaraShaderStageMergeAdapter::GetShaderStage() const
+{
+	return ShaderStage;
+}
+
+UNiagaraShaderStageBase* FNiagaraShaderStageMergeAdapter::GetEditableShaderStage() const
+{
+	return EditableShaderStage;
+}
+
+UNiagaraNodeOutput* FNiagaraShaderStageMergeAdapter::GetOutputNode() const
+{
+	return OutputNode.Get();
+}
+
+UNiagaraNodeInput* FNiagaraShaderStageMergeAdapter::GetInputNode() const
+{
+	return InputNode.Get();
+}
+
+TSharedPtr<FNiagaraScriptStackMergeAdapter> FNiagaraShaderStageMergeAdapter::GetShaderStageStack() const
+{
+	return ShaderStageStack;
+}
+
 FNiagaraRendererMergeAdapter::FNiagaraRendererMergeAdapter(UNiagaraRendererProperties& InRenderer)
 {
 	Renderer = &InRenderer;
@@ -441,6 +515,7 @@ void FNiagaraEmitterMergeAdapter::Initialize(const UNiagaraEmitter& InEmitter, U
 	Graph->GetNodesOfClass<UNiagaraNodeOutput>(OutputNodes);
 
 	TArray<UNiagaraNodeOutput*> EventOutputNodes;
+	TArray<UNiagaraNodeOutput*> ShaderStageOutputNodes;
 	for (UNiagaraNodeOutput* OutputNode : OutputNodes)
 	{
 		if (UNiagaraScript::IsEquivalentUsage(OutputNode->GetUsage(), ENiagaraScriptUsage::EmitterSpawnScript))
@@ -463,6 +538,10 @@ void FNiagaraEmitterMergeAdapter::Initialize(const UNiagaraEmitter& InEmitter, U
 		{
 			EventOutputNodes.Add(OutputNode);
 		}
+		else if (UNiagaraScript::IsEquivalentUsage(OutputNode->GetUsage(), ENiagaraScriptUsage::ParticleShaderStageScript))
+		{
+			ShaderStageOutputNodes.Add(OutputNode);
+		}
 	}
 
 	// Create an event handler adapter for each usage id even if it's missing an event script properties struct or an output node.  These
@@ -495,6 +574,36 @@ void FNiagaraEmitterMergeAdapter::Initialize(const UNiagaraEmitter& InEmitter, U
 		EventHandlers.Add(MakeShared<FNiagaraEventHandlerMergeAdapter>(*Emitter.Get(), EventOutputNode));
 	}
 
+	// Create an shader stage adapter for each usage id even if it's missing a shader stage object or an output node.  These
+	// incomplete adapters will be caught if they are diffed.
+	for (const UNiagaraShaderStageBase* ShaderStage : Emitter->GetShaderStages())
+	{
+		UNiagaraNodeOutput** MatchingOutputNodePtr = ShaderStageOutputNodes.FindByPredicate(
+			[=](UNiagaraNodeOutput* ShaderStageOutputNode) { return ShaderStageOutputNode->GetUsageId() == ShaderStage->Script->GetUsageId(); });
+
+		UNiagaraNodeOutput* MatchingOutputNode = MatchingOutputNodePtr != nullptr ? *MatchingOutputNodePtr : nullptr;
+
+		if (EditableEmitter == nullptr)
+		{
+			ShaderStages.Add(MakeShared<FNiagaraShaderStageMergeAdapter>(*Emitter.Get(), ShaderStage, MatchingOutputNode));
+		}
+		else
+		{
+			UNiagaraShaderStageBase* EditableShaderStage = EditableEmitter->GetShaderStageById(ShaderStage->Script->GetUsageId());
+			ShaderStages.Add(MakeShared<FNiagaraShaderStageMergeAdapter>(*Emitter.Get(), EditableShaderStage, MatchingOutputNode));
+		}
+
+		if (MatchingOutputNode != nullptr)
+		{
+			ShaderStageOutputNodes.Remove(MatchingOutputNode);
+		}
+	}
+
+	for (UNiagaraNodeOutput* ShaderStageOutputNode : ShaderStageOutputNodes)
+	{
+		ShaderStages.Add(MakeShared<FNiagaraShaderStageMergeAdapter>(*Emitter.Get(), ShaderStageOutputNode));
+	}
+
 	// Renderers
 	for (UNiagaraRendererProperties* RendererProperties : Emitter->GetRenderers())
 	{
@@ -532,6 +641,11 @@ const TArray<TSharedRef<FNiagaraEventHandlerMergeAdapter>> FNiagaraEmitterMergeA
 	return EventHandlers;
 }
 
+const TArray<TSharedRef<FNiagaraShaderStageMergeAdapter>> FNiagaraEmitterMergeAdapter::GetShaderStages() const
+{
+	return ShaderStages;
+}
+
 const TArray<TSharedRef<FNiagaraRendererMergeAdapter>> FNiagaraEmitterMergeAdapter::GetRenderers() const
 {
 	return Renderers;
@@ -577,6 +691,18 @@ TSharedPtr<FNiagaraEventHandlerMergeAdapter> FNiagaraEmitterMergeAdapter::GetEve
 	return TSharedPtr<FNiagaraEventHandlerMergeAdapter>();
 }
 
+TSharedPtr<FNiagaraShaderStageMergeAdapter> FNiagaraEmitterMergeAdapter::GetShaderStage(FGuid ShaderStageUsageId)
+{
+	for (TSharedRef<FNiagaraShaderStageMergeAdapter> ShaderStage : ShaderStages)
+	{
+		if (ShaderStage->GetUsageId() == ShaderStageUsageId)
+		{
+			return ShaderStage;
+		}
+	}
+	return TSharedPtr<FNiagaraShaderStageMergeAdapter>();
+}
+
 TSharedPtr<FNiagaraRendererMergeAdapter> FNiagaraEmitterMergeAdapter::GetRenderer(FGuid RendererMergeId)
 {
 	for (TSharedRef<FNiagaraRendererMergeAdapter> Renderer : Renderers)
@@ -642,7 +768,18 @@ bool FNiagaraEmitterDiffResults::IsValid() const
 			break;
 		}
 	}
-	return bIsValid && bEventHandlerDiffsAreValid &&
+	bool bShaderStageDiffsAreValid = true;
+	for (const FNiagaraModifiedShaderStageDiffResults& ShaderStageDiffResults : ModifiedShaderStages)
+	{
+		if (ShaderStageDiffResults.ScriptDiffResults.IsValid() == false)
+		{
+			bShaderStageDiffsAreValid = false;
+			break;
+		}
+	}
+	return bIsValid &&
+		bEventHandlerDiffsAreValid &&
+		bShaderStageDiffsAreValid &&
 		EmitterSpawnDiffResults.IsValid() &&
 		EmitterUpdateDiffResults.IsValid() &&
 		ParticleSpawnDiffResults.IsValid() &&
@@ -659,6 +796,9 @@ bool FNiagaraEmitterDiffResults::IsEmpty() const
 		RemovedBaseEventHandlers.Num() == 0 &&
 		AddedOtherEventHandlers.Num() == 0 &&
 		ModifiedEventHandlers.Num() == 0 &&
+		RemovedBaseShaderStages.Num() == 0 &&
+		AddedOtherShaderStages.Num() == 0 &&
+		ModifiedShaderStages.Num() == 0 &&
 		RemovedBaseRenderers.Num() == 0 &&
 		AddedOtherRenderers.Num() == 0 &&
 		ModifiedBaseRenderers.Num() == 0 &&
@@ -931,6 +1071,14 @@ INiagaraMergeManager::FMergeEmitterResults FNiagaraScriptMergeManager::MergeEmit
 		MergeResults.bModifiedGraph |= EventHandlerResults.bModifiedGraph;
 		MergeResults.ErrorMessages.Append(EventHandlerResults.ErrorMessages);
 
+		FApplyDiffResults ShaderStageResults = ApplyShaderStageDiff(MergedInstanceAdapter, DiffResults);
+		if (ShaderStageResults.bSucceeded == false)
+		{
+			MergeResults.MergeResult = INiagaraMergeManager::EMergeEmitterResult::FailedToMerge;
+		}
+		MergeResults.bModifiedGraph |= ShaderStageResults.bModifiedGraph;
+		MergeResults.ErrorMessages.Append(ShaderStageResults.ErrorMessages);
+
 		FApplyDiffResults RendererResults = ApplyRendererDiff(*MergedInstance, DiffResults);
 		if (RendererResults.bSucceeded == false)
 		{
@@ -1167,6 +1315,50 @@ void FNiagaraScriptMergeManager::ResetEventHandlerPropertySetToBase(UNiagaraEmit
 	Emitter.PostEditChange();
 }
 
+bool FNiagaraScriptMergeManager::HasBaseShaderStage(const UNiagaraEmitter& BaseEmitter, FGuid ShaderStageScriptUsageId)
+{
+	TSharedRef<FNiagaraEmitterMergeAdapter> BaseEmitterAdapter = GetEmitterMergeAdapterUsingCache(BaseEmitter);
+	return BaseEmitterAdapter->GetShaderStage(ShaderStageScriptUsageId).IsValid();
+}
+
+bool FNiagaraScriptMergeManager::IsShaderStagePropertySetDifferentFromBase(UNiagaraEmitter& Emitter, const UNiagaraEmitter& BaseEmitter, FGuid ShaderStageScriptUsageId)
+{
+	TSharedRef<FNiagaraEmitterMergeAdapter> EmitterAdapter = GetEmitterMergeAdapterUsingCache(Emitter);
+	TSharedRef<FNiagaraEmitterMergeAdapter> BaseEmitterAdapter = GetEmitterMergeAdapterUsingCache(BaseEmitter);
+
+	TSharedPtr<FNiagaraShaderStageMergeAdapter> ShaderStageAdapter = EmitterAdapter->GetShaderStage(ShaderStageScriptUsageId);
+	TSharedPtr<FNiagaraShaderStageMergeAdapter> BaseShaderStageAdapter = BaseEmitterAdapter->GetShaderStage(ShaderStageScriptUsageId);
+
+	if (ShaderStageAdapter->GetEditableShaderStage() == nullptr || BaseShaderStageAdapter->GetShaderStage() == nullptr)
+	{
+		return true;
+	}
+
+	TArray<UProperty*> DifferentProperties;
+	DiffEditableProperties(BaseShaderStageAdapter->GetShaderStage(), ShaderStageAdapter->GetShaderStage(), *BaseShaderStageAdapter->GetShaderStage()->GetClass(), DifferentProperties);
+	return DifferentProperties.Num() > 0;
+}
+
+void FNiagaraScriptMergeManager::ResetShaderStagePropertySetToBase(UNiagaraEmitter& Emitter, const UNiagaraEmitter& BaseEmitter, FGuid ShaderStageScriptUsageId)
+{
+	TSharedRef<FNiagaraEmitterMergeAdapter> EmitterAdapter = GetEmitterMergeAdapterUsingCache(Emitter);
+	TSharedRef<FNiagaraEmitterMergeAdapter> BaseEmitterAdapter = GetEmitterMergeAdapterUsingCache(BaseEmitter);
+
+	TSharedPtr<FNiagaraShaderStageMergeAdapter> ShaderStageAdapter = EmitterAdapter->GetShaderStage(ShaderStageScriptUsageId);
+	TSharedPtr<FNiagaraShaderStageMergeAdapter> BaseShaderStageAdapter = BaseEmitterAdapter->GetShaderStage(ShaderStageScriptUsageId);
+
+	if (ShaderStageAdapter->GetEditableShaderStage() == nullptr || BaseShaderStageAdapter->GetShaderStage() == nullptr)
+	{
+		// TODO: Display an error to the user.
+		return;
+	}
+
+	TArray<UProperty*> DifferentProperties;
+	DiffEditableProperties(BaseShaderStageAdapter->GetShaderStage(), ShaderStageAdapter->GetShaderStage(), *BaseShaderStageAdapter->GetShaderStage()->GetClass(), DifferentProperties);
+	CopyPropertiesToBase(ShaderStageAdapter->GetEditableShaderStage(), BaseShaderStageAdapter->GetShaderStage(), DifferentProperties);
+	Emitter.PostEditChange();
+}
+
 bool FNiagaraScriptMergeManager::HasBaseRenderer(const UNiagaraEmitter& BaseEmitter, FGuid RendererMergeId)
 {
 	TSharedRef<FNiagaraEmitterMergeAdapter> BaseEmitterAdapter = GetEmitterMergeAdapterUsingCache(BaseEmitter);
@@ -1278,6 +1470,7 @@ FNiagaraEmitterDiffResults FNiagaraScriptMergeManager::DiffEmitters(UNiagaraEmit
 	}
 
 	DiffEventHandlers(BaseEmitterAdapter->GetEventHandlers(), OtherEmitterAdapter->GetEventHandlers(), EmitterDiffResults);
+	DiffShaderStages(BaseEmitterAdapter->GetShaderStages(), OtherEmitterAdapter->GetShaderStages(), EmitterDiffResults);
 	DiffRenderers(BaseEmitterAdapter->GetRenderers(), OtherEmitterAdapter->GetRenderers(), EmitterDiffResults);
 	DiffEditableProperties(&BaseEmitter, &OtherEmitter, *UNiagaraEmitter::StaticClass(), EmitterDiffResults.DifferentEmitterProperties);
 
@@ -1397,6 +1590,57 @@ void FNiagaraScriptMergeManager::DiffEventHandlers(const TArray<TSharedRef<FNiag
 	}
 }
 
+void FNiagaraScriptMergeManager::DiffShaderStages(const TArray<TSharedRef<FNiagaraShaderStageMergeAdapter>>& BaseShaderStages, const TArray<TSharedRef<FNiagaraShaderStageMergeAdapter>>& OtherShaderStages, FNiagaraEmitterDiffResults& DiffResults) const
+{
+	FListDiffResults<TSharedRef<FNiagaraShaderStageMergeAdapter>> ShaderStageListDiffResults = DiffLists<TSharedRef<FNiagaraShaderStageMergeAdapter>, FGuid>(
+		BaseShaderStages,
+		OtherShaderStages,
+		[](TSharedRef<FNiagaraShaderStageMergeAdapter> ShaderStage) { return ShaderStage->GetUsageId(); });
+
+	DiffResults.RemovedBaseShaderStages.Append(ShaderStageListDiffResults.RemovedBaseValues);
+	DiffResults.AddedOtherShaderStages.Append(ShaderStageListDiffResults.AddedOtherValues);
+
+	for (const FCommonValuePair<TSharedRef<FNiagaraShaderStageMergeAdapter>>& CommonValuePair : ShaderStageListDiffResults.CommonValuePairs)
+	{
+		if (CommonValuePair.BaseValue->GetShaderStage() == nullptr || CommonValuePair.BaseValue->GetOutputNode() == nullptr)
+		{
+			DiffResults.AddError(FText::Format(LOCTEXT("InvalidBaseShaderStageDiffFailedFormat", "Failed to diff shader stages, the base shader stage was invalid.  Script Usage Id: {0}"),
+				FText::FromString(CommonValuePair.BaseValue->GetUsageId().ToString())));
+		}
+		else if (CommonValuePair.OtherValue->GetShaderStage() == nullptr || CommonValuePair.OtherValue->GetOutputNode() == nullptr)
+		{
+			DiffResults.AddError(FText::Format(LOCTEXT("InvalidOtherShaderStageDiffFailedFormat", "Failed to diff shader stage, the other shader stage was invalid.  Script Usage Id: {0}"),
+				FText::FromString(CommonValuePair.OtherValue->GetUsageId().ToString())));
+		}
+		else
+		{
+			TArray<UProperty*> DifferentProperties;
+			DiffEditableProperties(CommonValuePair.BaseValue->GetShaderStage(), CommonValuePair.OtherValue->GetShaderStage(), *CommonValuePair.BaseValue->GetShaderStage()->GetClass(), DifferentProperties);
+
+			FNiagaraScriptStackDiffResults ShaderStageScriptStackDiffResults;
+			DiffScriptStacks(CommonValuePair.BaseValue->GetShaderStageStack().ToSharedRef(), CommonValuePair.OtherValue->GetShaderStageStack().ToSharedRef(), ShaderStageScriptStackDiffResults);
+
+			if (DifferentProperties.Num() > 0 || ShaderStageScriptStackDiffResults.IsValid() == false || ShaderStageScriptStackDiffResults.IsEmpty() == false)
+			{
+				FNiagaraModifiedShaderStageDiffResults ModifiedShaderStageResults;
+				ModifiedShaderStageResults.BaseAdapter = CommonValuePair.BaseValue;
+				ModifiedShaderStageResults.OtherAdapter = CommonValuePair.OtherValue;
+				ModifiedShaderStageResults.ChangedProperties.Append(DifferentProperties);
+				ModifiedShaderStageResults.ScriptDiffResults = ShaderStageScriptStackDiffResults;
+				DiffResults.ModifiedShaderStages.Add(ModifiedShaderStageResults);
+			}
+
+			if (ShaderStageScriptStackDiffResults.IsValid() == false)
+			{
+				for (const FText& ShaderStageScriptStackDiffErrorMessage : ShaderStageScriptStackDiffResults.GetErrorMessages())
+				{
+					DiffResults.AddError(ShaderStageScriptStackDiffErrorMessage);
+				}
+			}
+		}
+	}
+}
+
 void FNiagaraScriptMergeManager::DiffRenderers(const TArray<TSharedRef<FNiagaraRendererMergeAdapter>>& BaseRenderers, const TArray<TSharedRef<FNiagaraRendererMergeAdapter>>& OtherRenderers, FNiagaraEmitterDiffResults& DiffResults) const
 {
 	FListDiffResults<TSharedRef<FNiagaraRendererMergeAdapter>> RendererListDiffResults = DiffLists<TSharedRef<FNiagaraRendererMergeAdapter>, FGuid>(
@@ -2053,12 +2297,16 @@ FNiagaraScriptMergeManager::FApplyDiffResults FNiagaraScriptMergeManager::ApplyE
 		if (ModifiedEventHandler.OtherAdapter->GetEventScriptProperties() == nullptr)
 		{
 			Results.bSucceeded = false;
-			Results.ErrorMessages.Add(LOCTEXT("MissingModifiedEventProperties", "Apply diff failed.  The modified event handler with id: {0} was missing it's event properties."));
+			Results.ErrorMessages.Add(FText::Format(
+				LOCTEXT("MissingModifiedEventPropertiesFormat", "Apply diff failed.  The modified event handler with id: {0} was missing it's event properties."),
+				FText::FromString(ModifiedEventHandler.OtherAdapter->GetUsageId().ToString(EGuidFormats::DigitsWithHyphens))));
 		}
 		else if (ModifiedEventHandler.OtherAdapter->GetOutputNode() == nullptr)
 		{
 			Results.bSucceeded = false;
-			Results.ErrorMessages.Add(LOCTEXT("MissingModifiedEventOutputNode", "Apply diff failed.  The modified event handler with id: {0} was missing it's output node."));
+			Results.ErrorMessages.Add(FText::Format(
+				LOCTEXT("MissingModifiedEventOutputNodeFormat", "Apply diff failed.  The modified event handler with id: {0} was missing it's output node."),
+				FText::FromString(ModifiedEventHandler.OtherAdapter->GetUsageId().ToString(EGuidFormats::DigitsWithHyphens))));
 		}
 		else
 		{
@@ -2087,12 +2335,16 @@ FNiagaraScriptMergeManager::FApplyDiffResults FNiagaraScriptMergeManager::ApplyE
 		if (AddedEventHandler->GetEventScriptProperties() == nullptr)
 		{
 			Results.bSucceeded = false;
-			Results.ErrorMessages.Add(LOCTEXT("MissingAddedEventProperties", "Apply diff failed.  The added event handler with id: {0} was missing it's event properties."));
+			Results.ErrorMessages.Add(FText::Format(
+				LOCTEXT("MissingAddedEventPropertiesFormat", "Apply diff failed.  The added event handler with id: {0} was missing it's event properties."),
+				FText::FromString(AddedEventHandler->GetUsageId().ToString(EGuidFormats::DigitsWithHyphens))));
 		}
 		else if (AddedEventHandler->GetOutputNode() == nullptr)
 		{
 			Results.bSucceeded = false;
-			Results.ErrorMessages.Add(LOCTEXT("MissingAddedEventOutputNode", "Apply diff failed.  The added event handler with id: {0} was missing it's output node."));
+			Results.ErrorMessages.Add(FText::Format(
+				LOCTEXT("MissingAddedEventOutputNodeFormat", "Apply diff failed.  The added event handler with id: {0} was missing it's output node."),
+				FText::FromString(AddedEventHandler->GetUsageId().ToString(EGuidFormats::DigitsWithHyphens))));
 		}
 		else
 		{
@@ -2128,6 +2380,105 @@ FNiagaraScriptMergeManager::FApplyDiffResults FNiagaraScriptMergeManager::ApplyE
 	return Results;
 }
 
+FNiagaraScriptMergeManager::FApplyDiffResults FNiagaraScriptMergeManager::ApplyShaderStageDiff(TSharedRef<FNiagaraEmitterMergeAdapter> BaseEmitterAdapter, const FNiagaraEmitterDiffResults& DiffResults) const
+{
+	FApplyDiffResults Results;
+	if (DiffResults.RemovedBaseShaderStages.Num() > 0)
+	{
+		Results.bSucceeded = false;
+		Results.bModifiedGraph = false;
+		Results.ErrorMessages.Add(LOCTEXT("RemovedShaderStagesUnsupported", "Apply diff failed, removed shader stages are currently unsupported."));
+		return Results;
+	}
+
+	for (const FNiagaraModifiedShaderStageDiffResults& ModifiedShaderStage : DiffResults.ModifiedShaderStages)
+	{
+		if (ModifiedShaderStage.OtherAdapter->GetShaderStage() == nullptr)
+		{
+			Results.bSucceeded = false;
+			Results.ErrorMessages.Add(FText::Format(
+				LOCTEXT("MissingModifiedShaderStageObjectFormat", "Apply diff failed.  The modified shader stage with id: {0} was missing it's shader stage object."),
+				FText::FromString(ModifiedShaderStage.OtherAdapter->GetUsageId().ToString(EGuidFormats::DigitsWithHyphens))));
+		}
+		else if (ModifiedShaderStage.OtherAdapter->GetOutputNode() == nullptr)
+		{
+			Results.bSucceeded = false;
+			Results.ErrorMessages.Add(FText::Format(
+				LOCTEXT("MissingModifiedShaderStageOutputNodeFormat", "Apply diff failed.  The modified shader stage with id: {0} was missing it's output node."),
+				FText::FromString(ModifiedShaderStage.OtherAdapter->GetUsageId().ToString(EGuidFormats::DigitsWithHyphens))));
+		}
+		else
+		{
+			TSharedPtr<FNiagaraShaderStageMergeAdapter> MatchingBaseShaderStageAdapter = BaseEmitterAdapter->GetShaderStage(ModifiedShaderStage.OtherAdapter->GetUsageId());
+			if (MatchingBaseShaderStageAdapter.IsValid())
+			{
+				if (ModifiedShaderStage.ChangedProperties.Num() > 0)
+				{
+					CopyPropertiesToBase(MatchingBaseShaderStageAdapter->GetEditableShaderStage(), ModifiedShaderStage.OtherAdapter->GetEditableShaderStage(), ModifiedShaderStage.ChangedProperties);
+				}
+				if (ModifiedShaderStage.ScriptDiffResults.IsEmpty() == false)
+				{
+					FApplyDiffResults ApplyShaderStageStackDiffResults = ApplyScriptStackDiff(MatchingBaseShaderStageAdapter->GetShaderStageStack().ToSharedRef(), ModifiedShaderStage.ScriptDiffResults);
+					Results.bSucceeded &= ApplyShaderStageStackDiffResults.bSucceeded;
+					Results.bModifiedGraph |= ApplyShaderStageStackDiffResults.bModifiedGraph;
+					Results.ErrorMessages.Append(ApplyShaderStageStackDiffResults.ErrorMessages);
+				}
+			}
+		}
+	}
+
+	UNiagaraScriptSource* EmitterSource = CastChecked<UNiagaraScriptSource>(BaseEmitterAdapter->GetEditableEmitter()->GraphSource);
+	UNiagaraGraph* EmitterGraph = EmitterSource->NodeGraph;
+	for (TSharedRef<FNiagaraShaderStageMergeAdapter> AddedOtherShaderStage : DiffResults.AddedOtherShaderStages)
+	{
+		if (AddedOtherShaderStage->GetShaderStage() == nullptr)
+		{
+			Results.bSucceeded = false;
+			Results.ErrorMessages.Add(FText::Format(
+				LOCTEXT("MissingAddedShaderStageObjectFormat", "Apply diff failed.  The added shader stage with id: {0} was missing it's shader stage object."),
+				FText::FromString(AddedOtherShaderStage->GetUsageId().ToString(EGuidFormats::DigitsWithHyphens))));
+		}
+		else if (AddedOtherShaderStage->GetOutputNode() == nullptr)
+		{
+			Results.bSucceeded = false;
+			Results.ErrorMessages.Add(FText::Format(
+				LOCTEXT("MissingAddedShaderStageOutputNodeFormat", "Apply diff failed.  The added shader stage with id: {0} was missing it's output node."),
+				FText::FromString(AddedOtherShaderStage->GetUsageId().ToString(EGuidFormats::DigitsWithHyphens))));
+		}
+		else
+		{
+			UNiagaraEmitter* BaseEmitter = BaseEmitterAdapter->GetEditableEmitter();
+			UNiagaraShaderStageBase* AddedShaderStage = CastChecked<UNiagaraShaderStageBase>(StaticDuplicateObject(AddedOtherShaderStage->GetShaderStage(), BaseEmitter));
+			AddedShaderStage->Script = NewObject<UNiagaraScript>(BaseEmitter, MakeUniqueObjectName(BaseEmitter, UNiagaraScript::StaticClass(), "ShaderStage"), EObjectFlags::RF_Transactional);
+			AddedShaderStage->Script->SetUsage(ENiagaraScriptUsage::ParticleEventScript);
+			AddedShaderStage->Script->SetUsageId(AddedOtherShaderStage->GetUsageId());
+			AddedShaderStage->Script->SetSource(EmitterSource);
+			BaseEmitter->AddShaderStage(AddedShaderStage);
+
+			FGuid PreferredOutputNodeGuid = AddedOtherShaderStage->GetOutputNode()->NodeGuid;
+			FGuid PreferredInputNodeGuid = AddedOtherShaderStage->GetInputNode()->NodeGuid;
+			UNiagaraNodeOutput* ShaderStageOutputNode = FNiagaraStackGraphUtilities::ResetGraphForOutput(*EmitterGraph, ENiagaraScriptUsage::ParticleShaderStageScript, AddedShaderStage->Script->GetUsageId(), PreferredOutputNodeGuid, PreferredInputNodeGuid);
+			for (TSharedRef<FNiagaraStackFunctionMergeAdapter> ModuleAdapter : AddedOtherShaderStage->GetShaderStageStack()->GetModuleFunctions())
+			{
+				FApplyDiffResults AddModuleResults = AddModule(BaseEmitter->GetUniqueEmitterName(), *AddedShaderStage->Script, *ShaderStageOutputNode, ModuleAdapter);
+				Results.bSucceeded &= AddModuleResults.bSucceeded;
+				Results.ErrorMessages.Append(AddModuleResults.ErrorMessages);
+			}
+
+			// Force the base compile id of the new shader stage to match the added instance shader stage.
+			UNiagaraScriptSource* AddedShaderStageSourceFromDiff = Cast<UNiagaraScriptSource>(AddedOtherShaderStage->GetShaderStage()->Script->GetSource());
+			UNiagaraGraph* AddedShaderStageGraphFromDiff = AddedShaderStageSourceFromDiff->NodeGraph;
+			FGuid ScriptBaseIdFromDiff = AddedShaderStageGraphFromDiff->GetBaseId(ENiagaraScriptUsage::ParticleShaderStageScript, AddedOtherShaderStage->GetUsageId());
+			UNiagaraScriptSource* AddedShaderStageSource = Cast<UNiagaraScriptSource>(AddedShaderStage->Script->GetSource());
+			UNiagaraGraph* AddedShaderStageGraph = AddedShaderStageSource->NodeGraph;
+			AddedShaderStageGraph->ForceBaseId(ENiagaraScriptUsage::ParticleShaderStageScript, AddedOtherShaderStage->GetUsageId(), ScriptBaseIdFromDiff);
+
+			Results.bModifiedGraph = true;
+		}
+	}
+	return Results;
+}
+
 FNiagaraScriptMergeManager::FApplyDiffResults FNiagaraScriptMergeManager::ApplyRendererDiff(UNiagaraEmitter& BaseEmitter, const FNiagaraEmitterDiffResults& DiffResults) const
 {
 	TArray<UNiagaraRendererProperties*> RenderersToRemove;
diff --git a/Engine/Plugins/FX/Niagara/Source/NiagaraEditor/Private/NiagaraScriptMergeManager.h b/Engine/Plugins/FX/Niagara/Source/NiagaraEditor/Private/NiagaraScriptMergeManager.h
--- a/Engine/Plugins/FX/Niagara/Source/NiagaraEditor/Private/NiagaraScriptMergeManager.h
+++ b/Engine/Plugins/FX/Niagara/Source/NiagaraEditor/Private/NiagaraScriptMergeManager.h
@@ -23,6 +23,7 @@ class UEdGraphNode;
 class UEdGraphPin;
 class UNiagaraDataInterface;
 struct FNiagaraEventScriptProperties;
+class UNiagaraShaderStageBase;
 class UNiagaraRendererProperties;
 
 class FNiagaraStackFunctionMergeAdapter;
@@ -154,6 +155,34 @@ private:
 	TSharedPtr<FNiagaraScriptStackMergeAdapter> EventStack;
 };
 
+class FNiagaraShaderStageMergeAdapter
+{
+public:
+	FNiagaraShaderStageMergeAdapter(const UNiagaraEmitter& InEmitter, const UNiagaraShaderStageBase* InShaderStage, UNiagaraNodeOutput* InOutputNode);
+	FNiagaraShaderStageMergeAdapter(const UNiagaraEmitter& InEmitter, UNiagaraShaderStageBase* InShaderStage, UNiagaraNodeOutput* InOutputNode);
+	FNiagaraShaderStageMergeAdapter(const UNiagaraEmitter& InEmitter, UNiagaraNodeOutput* InOutputNode);
+
+	FGuid GetUsageId() const;
+	const UNiagaraEmitter* GetEmitter() const;
+	const UNiagaraShaderStageBase* GetShaderStage() const;
+	UNiagaraShaderStageBase* GetEditableShaderStage() const;
+	UNiagaraNodeOutput* GetOutputNode() const;
+	UNiagaraNodeInput* GetInputNode() const;
+	TSharedPtr<FNiagaraScriptStackMergeAdapter> GetShaderStageStack() const;
+
+private:
+	void Initialize(const UNiagaraEmitter& InEmitter, const UNiagaraShaderStageBase* InShaderStage, UNiagaraShaderStageBase* InEditableShaderStage, UNiagaraNodeOutput* InOutputNode);
+
+private:
+	TWeakObjectPtr<UNiagaraEmitter> Emitter;
+	const UNiagaraShaderStageBase* ShaderStage;
+	UNiagaraShaderStageBase* EditableShaderStage;
+	TWeakObjectPtr<UNiagaraNodeOutput> OutputNode;
+	TWeakObjectPtr<UNiagaraNodeInput> InputNode;
+
+	TSharedPtr<FNiagaraScriptStackMergeAdapter> ShaderStageStack;
+};
+
 class FNiagaraRendererMergeAdapter
 {
 public:
@@ -178,12 +207,15 @@ public:
 	TSharedPtr<FNiagaraScriptStackMergeAdapter> GetParticleSpawnStack() const;
 	TSharedPtr<FNiagaraScriptStackMergeAdapter> GetParticleUpdateStack() const;
 	const TArray<TSharedRef<FNiagaraEventHandlerMergeAdapter>> GetEventHandlers() const;
+	const TArray<TSharedRef<FNiagaraShaderStageMergeAdapter>> GetShaderStages() const;
 	const TArray<TSharedRef<FNiagaraRendererMergeAdapter>> GetRenderers() const;
 
 	TSharedPtr<FNiagaraScriptStackMergeAdapter> GetScriptStack(ENiagaraScriptUsage Usage, FGuid UsageId);
 
 	TSharedPtr<FNiagaraEventHandlerMergeAdapter> GetEventHandler(FGuid EventScriptUsageId);
 
+	TSharedPtr<FNiagaraShaderStageMergeAdapter> GetShaderStage(FGuid ShaderStageUsageId);
+
 	TSharedPtr<FNiagaraRendererMergeAdapter> GetRenderer(FGuid RendererMergeId);
 
 private:
@@ -197,6 +229,7 @@ private:
 	TSharedPtr<FNiagaraScriptStackMergeAdapter> ParticleSpawnStack;
 	TSharedPtr<FNiagaraScriptStackMergeAdapter> ParticleUpdateStack;
 	TArray<TSharedRef<FNiagaraEventHandlerMergeAdapter>> EventHandlers;
+	TArray<TSharedRef<FNiagaraShaderStageMergeAdapter>> ShaderStages;
 	TArray<TSharedRef<FNiagaraRendererMergeAdapter>> Renderers;
 };
 
@@ -242,6 +275,14 @@ struct FNiagaraModifiedEventHandlerDiffResults
 	FNiagaraScriptStackDiffResults ScriptDiffResults;
 };
 
+struct FNiagaraModifiedShaderStageDiffResults
+{
+	TSharedPtr<FNiagaraShaderStageMergeAdapter> BaseAdapter;
+	TSharedPtr<FNiagaraShaderStageMergeAdapter> OtherAdapter;
+	TArray<UProperty*> ChangedProperties;
+	FNiagaraScriptStackDiffResults ScriptDiffResults;
+};
+
 struct FNiagaraEmitterDiffResults
 {
 	FNiagaraEmitterDiffResults();
@@ -267,6 +308,10 @@ struct FNiagaraEmitterDiffResults
 	TArray<TSharedRef<FNiagaraEventHandlerMergeAdapter>> AddedOtherEventHandlers;
 	TArray<FNiagaraModifiedEventHandlerDiffResults> ModifiedEventHandlers;
 
+	TArray<TSharedRef<FNiagaraShaderStageMergeAdapter>> RemovedBaseShaderStages;
+	TArray<TSharedRef<FNiagaraShaderStageMergeAdapter>> AddedOtherShaderStages;
+	TArray<FNiagaraModifiedShaderStageDiffResults> ModifiedShaderStages;
+
 	TArray<TSharedRef<FNiagaraRendererMergeAdapter>> RemovedBaseRenderers;
 	TArray<TSharedRef<FNiagaraRendererMergeAdapter>> AddedOtherRenderers;
 	TArray<TSharedRef<FNiagaraRendererMergeAdapter>> ModifiedBaseRenderers;
@@ -314,6 +359,12 @@ public:
 
 	void ResetEventHandlerPropertySetToBase(UNiagaraEmitter& Emitter, const UNiagaraEmitter& BaseEmitter, FGuid EventScriptUsageId);
 
+	bool HasBaseShaderStage(const UNiagaraEmitter& BaseEmitter, FGuid ShaderStageScriptUsageId);
+
+	bool IsShaderStagePropertySetDifferentFromBase(UNiagaraEmitter& Emitter, const UNiagaraEmitter& BaseEmitter, FGuid ShaderStageScriptUsageId);
+
+	void ResetShaderStagePropertySetToBase(UNiagaraEmitter& Emitter, const UNiagaraEmitter& BaseEmitter, FGuid ShaderStageScriptUsageId);
+
 	bool HasBaseRenderer(const UNiagaraEmitter& BaseEmitter, FGuid RendererMergeId);
 
 	bool IsRendererDifferentFromBase(UNiagaraEmitter& Emitter, const UNiagaraEmitter& BaseEmitter, FGuid RendererMergeId);
@@ -326,6 +377,8 @@ public:
 
 	void DiffEventHandlers(const TArray<TSharedRef<FNiagaraEventHandlerMergeAdapter>>& BaseEventHandlers, const TArray<TSharedRef<FNiagaraEventHandlerMergeAdapter>>& OtherEventHandlers, FNiagaraEmitterDiffResults& DiffResults) const;
 
+	void DiffShaderStages(const TArray<TSharedRef<FNiagaraShaderStageMergeAdapter>>& BaseShaderStages, const TArray<TSharedRef<FNiagaraShaderStageMergeAdapter>>& OtherShaderStages, FNiagaraEmitterDiffResults& DiffResults) const;
+
 	void DiffRenderers(const TArray<TSharedRef<FNiagaraRendererMergeAdapter>>& BaseRenderers, const TArray<TSharedRef<FNiagaraRendererMergeAdapter>>& OtherRenderers, FNiagaraEmitterDiffResults& DiffResults) const;
 
 	void DiffScriptStacks(TSharedRef<FNiagaraScriptStackMergeAdapter> BaseScriptStackAdapter, TSharedRef<FNiagaraScriptStackMergeAdapter> OtherScriptStackAdapter, FNiagaraScriptStackDiffResults& DiffResults) const;
@@ -343,6 +396,8 @@ private:
 
 	FApplyDiffResults ApplyEventHandlerDiff(TSharedRef<FNiagaraEmitterMergeAdapter> BaseEmitterAdapter, const FNiagaraEmitterDiffResults& DiffResults) const;
 
+	FApplyDiffResults ApplyShaderStageDiff(TSharedRef<FNiagaraEmitterMergeAdapter> BaseEmitterAdapter, const FNiagaraEmitterDiffResults& DiffResults) const;
+
 	FApplyDiffResults ApplyRendererDiff(UNiagaraEmitter& BaseEmitter, const FNiagaraEmitterDiffResults& DiffResults) const;
 
 	FApplyDiffResults AddModule(FString UniqueEmitterName, UNiagaraScript& OwningScript, UNiagaraNodeOutput& TargetOutputNode, TSharedRef<FNiagaraStackFunctionMergeAdapter> AddModule) const;
diff --git a/Engine/Plugins/FX/Niagara/Source/NiagaraEditor/Private/ViewModels/Stack/NiagaraStackEmitterSettingsGroup.cpp b/Engine/Plugins/FX/Niagara/Source/NiagaraEditor/Private/ViewModels/Stack/NiagaraStackEmitterSettingsGroup.cpp
--- a/Engine/Plugins/FX/Niagara/Source/NiagaraEditor/Private/ViewModels/Stack/NiagaraStackEmitterSettingsGroup.cpp
+++ b/Engine/Plugins/FX/Niagara/Source/NiagaraEditor/Private/ViewModels/Stack/NiagaraStackEmitterSettingsGroup.cpp
@@ -38,7 +38,7 @@ FText UNiagaraStackEmitterPropertiesItem::GetTooltipText() const
 	return LOCTEXT("EmitterPropertiesTooltip", "Properties that are handled per Emitter. These cannot change at runtime.");
 }
 
-bool UNiagaraStackEmitterPropertiesItem::CanResetToBase() const
+bool UNiagaraStackEmitterPropertiesItem::TestCanResetToBaseWithMessage(FText& OutCanResetToBaseMessage) const
 {
 	if (bCanResetToBaseCache.IsSet() == false)
 	{
@@ -53,12 +53,22 @@ bool UNiagaraStackEmitterPropertiesItem::CanResetToBase() const
 			bCanResetToBaseCache = false;
 		}
 	}
-	return bCanResetToBaseCache.GetValue();
+	if (bCanResetToBaseCache.GetValue())
+	{
+		OutCanResetToBaseMessage = LOCTEXT("CanResetToBase", "Reset the emitter properties to the state defined by the parent emitter.");
+		return true;
+	}
+	else
+	{
+		OutCanResetToBaseMessage = LOCTEXT("CanNotResetToBase", "No parent to reset to, or not different from parent.");
+		return false;
+	}
 }
 
 void UNiagaraStackEmitterPropertiesItem::ResetToBase()
 {
-	if (CanResetToBase())
+	FText Unused;
+	if (TestCanResetToBaseWithMessage(Unused))
 	{
 		const UNiagaraEmitter* BaseEmitter = GetEmitterViewModel()->GetEmitter()->GetParent();
 		TSharedRef<FNiagaraScriptMergeManager> MergeManager = FNiagaraScriptMergeManager::Get();
diff --git a/Engine/Plugins/FX/Niagara/Source/NiagaraEditor/Private/ViewModels/Stack/NiagaraStackEntry.cpp b/Engine/Plugins/FX/Niagara/Source/NiagaraEditor/Private/ViewModels/Stack/NiagaraStackEntry.cpp
--- a/Engine/Plugins/FX/Niagara/Source/NiagaraEditor/Private/ViewModels/Stack/NiagaraStackEntry.cpp
+++ b/Engine/Plugins/FX/Niagara/Source/NiagaraEditor/Private/ViewModels/Stack/NiagaraStackEntry.cpp
@@ -18,6 +18,7 @@ const FName UNiagaraStackEntry::FExecutionSubcategoryNames::Settings = TEXT("Set
 const FName UNiagaraStackEntry::FExecutionSubcategoryNames::Spawn = TEXT("Spawn");
 const FName UNiagaraStackEntry::FExecutionSubcategoryNames::Update = TEXT("Update");
 const FName UNiagaraStackEntry::FExecutionSubcategoryNames::Event = TEXT("Event");
+const FName UNiagaraStackEntry::FExecutionSubcategoryNames::ShaderStage = TEXT("Shader Stage");
 const FName UNiagaraStackEntry::FExecutionSubcategoryNames::Render = TEXT("Render");
 
 UNiagaraStackEntry::FStackIssueFix::FStackIssueFix()
diff --git a/Engine/Plugins/FX/Niagara/Source/NiagaraEditor/Private/ViewModels/Stack/NiagaraStackEventScriptItemGroup.cpp b/Engine/Plugins/FX/Niagara/Source/NiagaraEditor/Private/ViewModels/Stack/NiagaraStackEventScriptItemGroup.cpp
--- a/Engine/Plugins/FX/Niagara/Source/NiagaraEditor/Private/ViewModels/Stack/NiagaraStackEventScriptItemGroup.cpp
+++ b/Engine/Plugins/FX/Niagara/Source/NiagaraEditor/Private/ViewModels/Stack/NiagaraStackEventScriptItemGroup.cpp
@@ -50,7 +50,7 @@ FText UNiagaraStackEventHandlerPropertiesItem::GetDisplayName() const
 	return LOCTEXT("EventHandlerPropertiesDisplayName", "Event Handler Properties");
 }
 
-bool UNiagaraStackEventHandlerPropertiesItem::CanResetToBase() const
+bool UNiagaraStackEventHandlerPropertiesItem::TestCanResetToBaseWithMessage(FText& OutCanResetToBaseMessage) const
 {
 	if (bCanResetToBaseCache.IsSet() == false)
 	{
@@ -72,12 +72,22 @@ bool UNiagaraStackEventHandlerPropertiesItem::CanResetToBase() const
 			bCanResetToBaseCache = false;
 		}
 	}
-	return bCanResetToBaseCache.GetValue();
+	if (bCanResetToBaseCache.GetValue())
+	{
+		OutCanResetToBaseMessage = LOCTEXT("CanResetToBase", "Reset the event handler properties to the state defined by the parent emitter.");
+		return true;
+	}
+	else
+	{
+		OutCanResetToBaseMessage = LOCTEXT("CanNotResetToBase", "No parent to reset to, or not different from parent.");
+		return false;
+	}
 }
 
 void UNiagaraStackEventHandlerPropertiesItem::ResetToBase()
 {
-	if (CanResetToBase())
+	FText Unused;
+	if (TestCanResetToBaseWithMessage(Unused))
 	{
 		const UNiagaraEmitter* BaseEmitter = GetEmitterViewModel()->GetEmitter()->GetParent();
 		TSharedRef<FNiagaraScriptMergeManager> MergeManager = FNiagaraScriptMergeManager::Get();
diff --git a/Engine/Plugins/FX/Niagara/Source/NiagaraEditor/Private/ViewModels/Stack/NiagaraStackRendererItem.cpp b/Engine/Plugins/FX/Niagara/Source/NiagaraEditor/Private/ViewModels/Stack/NiagaraStackRendererItem.cpp
--- a/Engine/Plugins/FX/Niagara/Source/NiagaraEditor/Private/ViewModels/Stack/NiagaraStackRendererItem.cpp
+++ b/Engine/Plugins/FX/Niagara/Source/NiagaraEditor/Private/ViewModels/Stack/NiagaraStackRendererItem.cpp
@@ -263,24 +263,37 @@ bool UNiagaraStackRendererItem::HasBaseRenderer() const
 	return false;
 }
 
-bool UNiagaraStackRendererItem::CanResetToBase() const
+bool UNiagaraStackRendererItem::TestCanResetToBaseWithMessage(FText& OutCanResetToBaseMessage) const
 {
-	if (HasBaseRenderer())
+	if (bCanResetToBaseCache.IsSet() == false)
 	{
-		if (bCanResetToBaseCache.IsSet() == false)
+		if (HasBaseRenderer())
 		{
 			TSharedRef<FNiagaraScriptMergeManager> MergeManager = FNiagaraScriptMergeManager::Get();
 			const UNiagaraEmitter* BaseEmitter = GetEmitterViewModel()->GetEmitter()->GetParent();
 			bCanResetToBaseCache = BaseEmitter != nullptr && MergeManager->IsRendererDifferentFromBase(*GetEmitterViewModel()->GetEmitter(), *BaseEmitter, RendererProperties->GetMergeId());
 		}
-		return bCanResetToBaseCache.GetValue();
+		else
+		{
+			bCanResetToBaseCache = false;
+		}
+	}
+	if (bCanResetToBaseCache.GetValue())
+	{
+		OutCanResetToBaseMessage = LOCTEXT("CanResetToBase", "Reset this renderer to the state defined by the parent emitter.");
+		return true;
+	}
+	else
+	{
+		OutCanResetToBaseMessage = LOCTEXT("CanNotResetToBase", "No parent to reset to, or not different from parent.");
+		return false;
 	}
-	return false;
 }
 
 void UNiagaraStackRendererItem::ResetToBase()
 {
-	if (CanResetToBase())
+	FText Unused;
+	if (TestCanResetToBaseWithMessage(Unused))
 	{
 		TSharedRef<FNiagaraScriptMergeManager> MergeManager = FNiagaraScriptMergeManager::Get();
 		const UNiagaraEmitter* BaseEmitter = GetEmitterViewModel()->GetEmitter()->GetParent();
diff --git a/Engine/Plugins/FX/Niagara/Source/NiagaraEditor/Private/ViewModels/Stack/NiagaraStackRoot.cpp b/Engine/Plugins/FX/Niagara/Source/NiagaraEditor/Private/ViewModels/Stack/NiagaraStackRoot.cpp
--- a/Engine/Plugins/FX/Niagara/Source/NiagaraEditor/Private/ViewModels/Stack/NiagaraStackRoot.cpp
+++ b/Engine/Plugins/FX/Niagara/Source/NiagaraEditor/Private/ViewModels/Stack/NiagaraStackRoot.cpp
@@ -6,6 +6,8 @@
 #include "ViewModels/Stack/NiagaraStackRenderItemGroup.h"
 #include "ViewModels/Stack/NiagaraStackEventHandlerGroup.h"
 #include "ViewModels/Stack/NiagaraStackEventScriptItemGroup.h"
+#include "ViewModels/Stack/NiagaraStackShaderStagesGroup.h"
+#include "ViewModels/Stack/NiagaraStackShaderStageGroup.h"
 #include "ViewModels/NiagaraSystemViewModel.h"
 #include "ViewModels/NiagaraEmitterViewModel.h"
 #include "NiagaraSystemScriptViewModel.h"
@@ -17,6 +19,14 @@
 
 #define LOCTEXT_NAMESPACE "NiagaraStackViewModel"
 
+int32 GbShowNewShaderStageStackUIOptions = 0;
+static FAutoConsoleVariableRef CVarShowNewShaderStageStackUIOptions(
+	TEXT("fx.Niagara.ShowNewShaderStageStackUIOptions"),
+	GbShowNewShaderStageStackUIOptions,
+	TEXT("If > 0 the new shader state UI will be shown in the stack view.  This UI will generate data, but that data is currently unused.\n"),
+	ECVF_Default
+);
+
 UNiagaraStackRoot::UNiagaraStackRoot()
 	: SystemSettingsGroup(nullptr)
 	, EmitterSettingsGroup(nullptr)
@@ -25,6 +35,7 @@ UNiagaraStackRoot::UNiagaraStackRoot()
 	, ParticleSpawnGroup(nullptr)
 	, ParticleUpdateGroup(nullptr)
 	, AddEventHandlerGroup(nullptr)
+	, AddShaderStageGroup(nullptr)
 	, RenderGroup(nullptr)
 {
 }
@@ -41,6 +52,7 @@ void UNiagaraStackRoot::Initialize(FRequiredEntryData InRequiredEntryData, bool
 	ParticleSpawnGroup = nullptr;
 	ParticleUpdateGroup = nullptr;
 	AddEventHandlerGroup = nullptr;
+	AddShaderStageGroup = nullptr;
 	RenderGroup = nullptr;
 }
 
@@ -150,7 +162,17 @@ void UNiagaraStackRoot::RefreshChildrenInternal(const TArray<UNiagaraStackEntry*
 			FExecutionCategoryNames::Particle, FExecutionSubcategoryNames::Event,
 			GetEmitterViewModel()->GetOrCreateEditorData().GetStackEditorData());
 		AddEventHandlerGroup->Initialize(RequiredEntryData);
-		AddEventHandlerGroup->SetOnItemAdded(UNiagaraStackEventHandlerGroup::FOnItemAdded::CreateUObject(this, &UNiagaraStackRoot::EmitterEventArraysChanged));
+		AddEventHandlerGroup->SetOnItemAdded(UNiagaraStackEventHandlerGroup::FOnItemAdded::CreateUObject(this, &UNiagaraStackRoot::EmitterArraysChanged));
+	}
+
+	if (bIncludeEmitterInformation && AddShaderStageGroup == nullptr && GbShowNewShaderStageStackUIOptions > 0)
+	{
+		AddShaderStageGroup = NewObject<UNiagaraStackShaderStagesGroup>(this);
+		FRequiredEntryData RequiredEntryData(GetSystemViewModel(), GetEmitterViewModel(),
+			FExecutionCategoryNames::Particle, FExecutionSubcategoryNames::ShaderStage,
+			GetEmitterViewModel()->GetOrCreateEditorData().GetStackEditorData());
+		AddShaderStageGroup->Initialize(RequiredEntryData);
+		AddShaderStageGroup->SetOnItemAdded(UNiagaraStackShaderStagesGroup::FOnItemAdded::CreateUObject(this, &UNiagaraStackRoot::EmitterArraysChanged));
 	}
 
 	if (bIncludeEmitterInformation && RenderGroup == nullptr)
@@ -186,12 +208,12 @@ void UNiagaraStackRoot::RefreshChildrenInternal(const TArray<UNiagaraStackEntry*
 
 			if (EventHandlerGroup == nullptr)
 			{
-				EventHandlerGroup = NewObject<UNiagaraStackEventScriptItemGroup>(this, NAME_None, RF_Transactional);
+				EventHandlerGroup = NewObject<UNiagaraStackEventScriptItemGroup>(this);
 				FRequiredEntryData RequiredEntryData(GetSystemViewModel(), GetEmitterViewModel(),
 					FExecutionCategoryNames::Particle, FExecutionSubcategoryNames::Event,
 					GetEmitterViewModel()->GetEditorData().GetStackEditorData());
 				EventHandlerGroup->Initialize(RequiredEntryData, GetEmitterViewModel()->GetSharedScriptViewModel(), ENiagaraScriptUsage::ParticleEventScript, EventScriptProperties.Script->GetUsageId());
-				EventHandlerGroup->SetOnModifiedEventHandlers(UNiagaraStackEventScriptItemGroup::FOnModifiedEventHandlers::CreateUObject(this, &UNiagaraStackRoot::EmitterEventArraysChanged));
+				EventHandlerGroup->SetOnModifiedEventHandlers(UNiagaraStackEventScriptItemGroup::FOnModifiedEventHandlers::CreateUObject(this, &UNiagaraStackRoot::EmitterArraysChanged));
 			}
 
 			NewChildren.Add(EventHandlerGroup);
@@ -199,11 +221,34 @@ void UNiagaraStackRoot::RefreshChildrenInternal(const TArray<UNiagaraStackEntry*
 
 		NewChildren.Add(AddEventHandlerGroup);
 
+		if (GbShowNewShaderStageStackUIOptions > 0)
+		{
+			for (UNiagaraShaderStageBase* ShaderStage : GetEmitterViewModel()->GetEmitter()->GetShaderStages())
+			{
+				UNiagaraStackShaderStageGroup* ShaderStageGroup = FindCurrentChildOfTypeByPredicate<UNiagaraStackShaderStageGroup>(CurrentChildren,
+					[ShaderStage](UNiagaraStackShaderStageGroup* CurrentShaderStageGroup) { return CurrentShaderStageGroup->GetShaderStage() == ShaderStage; });
+
+				if (ShaderStageGroup == nullptr)
+				{
+					ShaderStageGroup = NewObject<UNiagaraStackShaderStageGroup>(this);
+					FRequiredEntryData RequiredEntryData(GetSystemViewModel(), GetEmitterViewModel(),
+						FExecutionCategoryNames::Particle, FExecutionSubcategoryNames::ShaderStage,
+						GetEmitterViewModel()->GetEditorData().GetStackEditorData());
+					ShaderStageGroup->Initialize(RequiredEntryData, GetEmitterViewModel()->GetSharedScriptViewModel(), ShaderStage);
+					ShaderStageGroup->SetOnModifiedShaderStages(UNiagaraStackShaderStageGroup::FOnModifiedShaderStages::CreateUObject(this, &UNiagaraStackRoot::EmitterArraysChanged));
+				}
+
+				NewChildren.Add(ShaderStageGroup);
+			}
+
+			NewChildren.Add(AddShaderStageGroup);
+		}
+
 		NewChildren.Add(RenderGroup);
 	}
 }
 
-void UNiagaraStackRoot::EmitterEventArraysChanged()
+void UNiagaraStackRoot::EmitterArraysChanged()
 {
 	RefreshChildren();
 }
diff --git a/Engine/Plugins/FX/Niagara/Source/NiagaraEditor/Private/ViewModels/Stack/NiagaraStackShaderStageGroup.cpp b/Engine/Plugins/FX/Niagara/Source/NiagaraEditor/Private/ViewModels/Stack/NiagaraStackShaderStageGroup.cpp
new file mode 100644
--- /dev/null
+++ b/Engine/Plugins/FX/Niagara/Source/NiagaraEditor/Private/ViewModels/Stack/NiagaraStackShaderStageGroup.cpp
@@ -0,0 +1,353 @@
+// Copyright 1998-2019 Epic Games, Inc. All Rights Reserved.
+
+#include "ViewModels/Stack/NiagaraStackShaderStageGroup.h"
+#include "NiagaraEmitter.h"
+#include "NiagaraShaderStageBase.h"
+#include "ViewModels/Stack/NiagaraStackObject.h"
+#include "ViewModels/Stack/NiagaraStackModuleItem.h"
+#include "NiagaraScriptViewModel.h"
+#include "NiagaraScriptGraphViewModel.h"
+#include "ViewModels/NiagaraEmitterViewModel.h"
+#include "NiagaraGraph.h"
+#include "NiagaraNodeOutput.h"
+#include "NiagaraNodeFunctionCall.h"
+#include "EdGraphSchema_Niagara.h"
+#include "ViewModels/Stack/NiagaraStackErrorItem.h"
+#include "NiagaraScriptSource.h"
+#include "ViewModels/Stack/NiagaraStackGraphUtilities.h"
+#include "NiagaraScriptMergeManager.h"
+#include "NiagaraStackEditorData.h"
+#include "ViewModels/NiagaraSystemViewModel.h"
+
+#include "Internationalization/Internationalization.h"
+#include "ScopedTransaction.h"
+#include "IDetailTreeNode.h"
+
+#define LOCTEXT_NAMESPACE "UNiagaraStackShaderStageGroup"
+
+void UNiagaraStackShaderStagePropertiesItem::Initialize(FRequiredEntryData InRequiredEntryData, UNiagaraShaderStageBase* InShaderStage)
+{
+	checkf(ShaderStage.IsValid() == false, TEXT("Can not initialize more than once."));
+	ShaderStage = InShaderStage;
+	FString ShaderStageStackEditorDataKey = FString::Printf(TEXT("ShaderStage-%s-Properties"), *ShaderStage->Script->GetUsageId().ToString(EGuidFormats::DigitsWithHyphens));
+	Super::Initialize(InRequiredEntryData, ShaderStageStackEditorDataKey);
+
+	ShaderStage->OnChanged().AddUObject(this, &UNiagaraStackShaderStagePropertiesItem::ShaderStagePropertiesChanged);
+}
+
+void UNiagaraStackShaderStagePropertiesItem::FinalizeInternal()
+{
+	if (ShaderStage.IsValid())
+	{
+		ShaderStage->OnChanged().RemoveAll(this);
+	}
+	Super::FinalizeInternal();
+}
+
+FText UNiagaraStackShaderStagePropertiesItem::GetDisplayName() const
+{
+	return FText::Format(LOCTEXT("ShaderStagePropertiesDisplayNameFormat", "{0} Settings"), ShaderStage->GetClass()->GetDisplayNameText());
+}
+
+bool UNiagaraStackShaderStagePropertiesItem::TestCanResetToBaseWithMessage(FText& OutCanResetToBaseMessage) const
+{
+	if (bCanResetToBaseCache.IsSet() == false)
+	{
+		if (HasBaseShaderStage())
+		{
+			UNiagaraEmitter* Emitter = GetEmitterViewModel()->GetEmitter();
+			const UNiagaraEmitter* BaseEmitter = Emitter->GetParent();
+			if (BaseEmitter != nullptr && Emitter != BaseEmitter)
+			{
+				TSharedRef<FNiagaraScriptMergeManager> MergeManager = FNiagaraScriptMergeManager::Get();
+				bCanResetToBaseCache = MergeManager->IsShaderStagePropertySetDifferentFromBase(*Emitter, *BaseEmitter, ShaderStage->Script->GetUsageId());
+			}
+			else
+			{
+				bCanResetToBaseCache = false;
+			}
+		}
+		else
+		{
+			bCanResetToBaseCache = false;
+		}
+	}
+	if (bCanResetToBaseCache.GetValue())
+	{
+		OutCanResetToBaseMessage = LOCTEXT("CanResetToBase", "Reset this shader stage to the one defined in the parent emitter.");
+		return true;
+	}
+	else
+	{
+		OutCanResetToBaseMessage = LOCTEXT("CanNotResetToBase", "No parent to reset to, or not different from parent.");
+		return false;
+	}
+}
+
+void UNiagaraStackShaderStagePropertiesItem::ResetToBase()
+{
+	FText Unused;
+	if (TestCanResetToBaseWithMessage(Unused))
+	{
+		UNiagaraEmitter* Emitter = GetEmitterViewModel()->GetEmitter();
+		const UNiagaraEmitter* BaseEmitter = Emitter->GetParent();
+		TSharedRef<FNiagaraScriptMergeManager> MergeManager = FNiagaraScriptMergeManager::Get();
+		MergeManager->ResetShaderStagePropertySetToBase(*Emitter, *BaseEmitter, ShaderStage->Script->GetUsageId());
+		RefreshChildren();
+	}
+}
+
+void UNiagaraStackShaderStagePropertiesItem::RefreshChildrenInternal(const TArray<UNiagaraStackEntry*>& CurrentChildren, TArray<UNiagaraStackEntry*>& NewChildren, TArray<FStackIssue>& NewIssues)
+{
+	if (ShaderStageObject == nullptr)
+	{
+		ShaderStageObject = NewObject<UNiagaraStackObject>(this);
+		ShaderStageObject->Initialize(CreateDefaultChildRequiredData(), ShaderStage.Get(), GetStackEditorDataKey());
+	}
+
+	NewChildren.Add(ShaderStageObject);
+
+	bCanResetToBaseCache.Reset();
+	bHasBaseShaderStageCache.Reset();
+
+	Super::RefreshChildrenInternal(CurrentChildren, NewChildren, NewIssues);
+}
+
+void UNiagaraStackShaderStagePropertiesItem::ShaderStagePropertiesChanged()
+{
+	bCanResetToBaseCache.Reset();
+}
+
+bool UNiagaraStackShaderStagePropertiesItem::HasBaseShaderStage() const
+{
+	if (bHasBaseShaderStageCache.IsSet() == false)
+	{
+		UNiagaraEmitter* Emitter = GetEmitterViewModel()->GetEmitter();
+		const UNiagaraEmitter* BaseEmitter = Emitter->GetParent();
+		if (BaseEmitter != nullptr && Emitter != BaseEmitter)
+		{
+			TSharedRef<FNiagaraScriptMergeManager> MergeManager = FNiagaraScriptMergeManager::Get();
+			bHasBaseShaderStageCache = MergeManager->HasBaseShaderStage(*BaseEmitter, ShaderStage->Script->GetUsageId());
+		}
+		else
+		{
+			bHasBaseShaderStageCache = false;
+		}
+	}
+	return bHasBaseShaderStageCache.GetValue();
+}
+
+void UNiagaraStackShaderStageGroup::Initialize(
+	FRequiredEntryData InRequiredEntryData,
+	TSharedRef<FNiagaraScriptViewModel> InScriptViewModel,
+	UNiagaraShaderStageBase* InShaderStage)
+{
+	ShaderStage = InShaderStage;
+	ShaderStage->OnChanged().AddUObject(this, &UNiagaraStackShaderStageGroup::ShaderStagePropertiesChanged);
+
+	FText ToolTip = LOCTEXT("ShaderStageGroupTooltip", "Defines properties and script modules for a shader stage.");
+	FText DisplayName = FText::FromName(ShaderStage->ShaderStageName);
+	Super::Initialize(InRequiredEntryData, DisplayName, ToolTip, InScriptViewModel, ENiagaraScriptUsage::ParticleShaderStageScript, ShaderStage->Script->GetUsageId());
+}
+
+UNiagaraShaderStageBase* UNiagaraStackShaderStageGroup::GetShaderStage() const
+{
+	return ShaderStage.Get();
+}
+
+void UNiagaraStackShaderStageGroup::FinalizeInternal()
+{
+	if (ShaderStage.IsValid())
+	{
+		ShaderStage->OnChanged().RemoveAll(this);
+	}
+	Super::FinalizeInternal();
+}
+
+void UNiagaraStackShaderStageGroup::RefreshChildrenInternal(const TArray<UNiagaraStackEntry*>& CurrentChildren, TArray<UNiagaraStackEntry*>& NewChildren, TArray<FStackIssue>& NewIssues)
+{
+	bHasBaseShaderStageCache.Reset();
+
+	UNiagaraEmitter* Emitter = GetEmitterViewModel()->GetEmitter();
+
+	SetDisplayName(FText::FromName(ShaderStage->ShaderStageName));
+
+	if (ShaderStageProperties == nullptr)
+	{
+		ShaderStageProperties = NewObject<UNiagaraStackShaderStagePropertiesItem>(this);
+		ShaderStageProperties->Initialize(CreateDefaultChildRequiredData(), ShaderStage.Get());
+	}
+	NewChildren.Add(ShaderStageProperties);
+	
+	Super::RefreshChildrenInternal(CurrentChildren, NewChildren, NewIssues);
+}
+
+TOptional<UNiagaraStackEntry::FDropRequestResponse> UNiagaraStackShaderStageGroup::CanDropInternal(const FDropRequest& DropRequest)
+{
+	if (DropRequest.DragDropOperation->IsOfType<FNiagaraStackEntryDragDropOp>() && 
+		(DropRequest.DropZone == EItemDropZone::AboveItem || DropRequest.DropZone == EItemDropZone::BelowItem))
+	{
+		TSharedRef<const FNiagaraStackEntryDragDropOp> StackEntryDragDropOp = StaticCastSharedRef<const FNiagaraStackEntryDragDropOp>(DropRequest.DragDropOperation);
+		bool ShaderStageEntriesDragged = false;
+		for (UNiagaraStackEntry* DraggedEntry : StackEntryDragDropOp->GetDraggedEntries())
+		{
+			if (DraggedEntry->IsA<UNiagaraStackShaderStageGroup>())
+			{
+				ShaderStageEntriesDragged = true;
+				break;
+			}
+		}
+
+		if (ShaderStageEntriesDragged == false)
+		{
+			// Only handle dragged ShaderStage items.
+			return TOptional<FDropRequestResponse>();
+		}
+
+		if (DropRequest.DropOptions != UNiagaraStackEntry::EDropOptions::Overview)
+		{
+			// Only allow dropping in the overview stacks.
+			return FDropRequestResponse(TOptional<EItemDropZone>(), LOCTEXT("CantDropShaderStageOnStack", "Shader stages can only be dropped into the overview."));
+		}
+
+		if (StackEntryDragDropOp->GetDraggedEntries().Num() != 1)
+		{
+			// Only handle a single items.
+			return FDropRequestResponse(TOptional<EItemDropZone>(), LOCTEXT("CantDropMultipleShaderStages", "Only single shader stages can be dragged and dropped."));
+		}
+
+		UNiagaraStackShaderStageGroup* SourceShaderStageGroup = CastChecked<UNiagaraStackShaderStageGroup>(StackEntryDragDropOp->GetDraggedEntries()[0]);
+		if (DropRequest.DragOptions != EDragOptions::Copy && SourceShaderStageGroup->HasBaseShaderStage())
+		{
+			return FDropRequestResponse(TOptional<EItemDropZone>(), LOCTEXT("CantMoveShaderStageError", "This inherited shader stage can't be moved."));
+		}
+
+		if (SourceShaderStageGroup == this)
+		{
+			// Don't allow dropping on yourself.
+			return TOptional<FDropRequestResponse>();
+		}
+
+		UNiagaraEmitter* OwningEmitter = GetEmitterViewModel()->GetEmitter();
+		int32 SourceIndex = OwningEmitter->GetShaderStages().IndexOfByKey(SourceShaderStageGroup->GetShaderStage());
+		if (SourceIndex == INDEX_NONE)
+		{
+			return FDropRequestResponse(TOptional<EItemDropZone>(), LOCTEXT("CantDropShaderStageFromOtherEmitterError", "This shader stage can't be moved here\nbecause it's owned by a different emitter."));
+		}
+
+		int32 TargetIndexOffset = DropRequest.DropZone == EItemDropZone::BelowItem ? 1 : 0;
+		int32 TargetIndex = OwningEmitter->GetShaderStages().IndexOfByKey(ShaderStage.Get()) + TargetIndexOffset;
+		if (SourceIndex < TargetIndex)
+		{
+			// If the source index is less than the target index, the target index will decrease by 1 after the source is removed.
+			TargetIndex--;
+		}
+
+		if (SourceIndex == TargetIndex)
+		{
+			// Only handle the drag if the item would actually move.
+			return TOptional<FDropRequestResponse>();
+		}
+
+		return FDropRequestResponse(DropRequest.DropZone, LOCTEXT("MoveShaderStageDragMessage", "Move this shader stage here."));
+	}
+	return TOptional<FDropRequestResponse>();
+}
+
+TOptional<UNiagaraStackEntry::FDropRequestResponse> UNiagaraStackShaderStageGroup::DropInternal(const FDropRequest& DropRequest)
+{
+	UNiagaraEmitter* OwningEmitter = GetEmitterViewModel()->GetEmitter();
+
+	TSharedRef<const FNiagaraStackEntryDragDropOp> StackEntryDragDropOp = StaticCastSharedRef<const FNiagaraStackEntryDragDropOp>(DropRequest.DragDropOperation);
+	UNiagaraStackShaderStageGroup* SourceShaderStageGroup = CastChecked<UNiagaraStackShaderStageGroup>(StackEntryDragDropOp->GetDraggedEntries()[0]);
+
+	int32 SourceIndex = OwningEmitter->GetShaderStages().IndexOfByKey(SourceShaderStageGroup->GetShaderStage());
+	if (SourceIndex != INDEX_NONE)
+	{
+		int32 TargetOffset = DropRequest.DropZone == EItemDropZone::BelowItem ? 1 : 0;
+		int32 TargetIndex = OwningEmitter->GetShaderStages().IndexOfByKey(ShaderStage.Get()) + TargetOffset;
+
+		FScopedTransaction Transaction(FText::Format(LOCTEXT("MoveShaderStage", "Move Shader Stage {0}"), GetDisplayName()));
+		OwningEmitter->MoveShaderStageToIndex(SourceShaderStageGroup->GetShaderStage(), TargetIndex);
+
+		OnRequestFullRefreshDeferred().Broadcast();
+		return FDropRequestResponse(DropRequest.DropZone, FText());
+	}
+	return TOptional<FDropRequestResponse>();
+}
+
+bool UNiagaraStackShaderStageGroup::TestCanDeleteWithMessage(FText& OutCanDeleteMessage) const
+{
+	if (HasBaseShaderStage())
+	{
+		OutCanDeleteMessage = LOCTEXT("CantDeleteInherited", "Can not delete this shader stage because it's inherited.");
+		return false;
+	}
+	else
+	{
+		OutCanDeleteMessage = LOCTEXT("CanDelete", "Delete this shader stage.");
+		return true;
+	}
+}
+
+void UNiagaraStackShaderStageGroup::Delete()
+{
+	TSharedPtr<FNiagaraScriptViewModel> ScriptViewModelPinned = ScriptViewModel.Pin();
+	checkf(ScriptViewModelPinned.IsValid(), TEXT("Can not delete when the script view model has been deleted."));
+
+	UNiagaraEmitter* Emitter = GetEmitterViewModel()->GetEmitter();
+	UNiagaraScriptSource* Source = Cast<UNiagaraScriptSource>(Emitter->GraphSource);
+
+	if (!Source || !Source->NodeGraph)
+	{
+		return;
+	}
+
+	FScopedTransaction Transaction(FText::Format(LOCTEXT("DeleteShaderStage", "Delete {0}"), GetDisplayName()));
+
+	Emitter->Modify();
+	Source->NodeGraph->Modify();
+	TArray<UNiagaraNode*> ShaderStageNodes;
+	Source->NodeGraph->BuildTraversal(ShaderStageNodes, GetScriptUsage(), GetScriptUsageId());
+	for (UNiagaraNode* Node : ShaderStageNodes)
+	{
+		Node->Modify();
+	}
+	
+	// First, remove the shader stage object.
+	Emitter->RemoveShaderStage(ShaderStage.Get());
+	
+	// Now remove all graph nodes associated with the shader stage.
+	for (UNiagaraNode* Node : ShaderStageNodes)
+	{
+		Node->DestroyNode();
+	}
+
+	// Set the emitter here to that the internal state of the view model is updated.
+	// TODO: Move the logic for managing additional scripts into the emitter view model or script view model.
+	ScriptViewModelPinned->SetScripts(Emitter);
+	
+	OnModifiedShaderStagesDelegate.ExecuteIfBound();
+}
+
+void UNiagaraStackShaderStageGroup::ShaderStagePropertiesChanged()
+{
+	SetDisplayName(FText::FromName(ShaderStage->ShaderStageName));
+}
+
+bool UNiagaraStackShaderStageGroup::HasBaseShaderStage() const
+{
+	if (bHasBaseShaderStageCache.IsSet() == false)
+	{
+		const UNiagaraEmitter* BaseEmitter = GetEmitterViewModel()->GetEmitter()->GetParent();
+		bHasBaseShaderStageCache = BaseEmitter != nullptr && FNiagaraScriptMergeManager::Get()->HasBaseShaderStage(*BaseEmitter, GetScriptUsageId());
+	}
+	return bHasBaseShaderStageCache.GetValue();
+}
+
+void UNiagaraStackShaderStageGroup::SetOnModifiedShaderStages(FOnModifiedShaderStages OnModifiedShaderStages)
+{
+	OnModifiedShaderStagesDelegate = OnModifiedShaderStages;
+}
+
+#undef LOCTEXT_NAMESPACE
diff --git a/Engine/Plugins/FX/Niagara/Source/NiagaraEditor/Private/ViewModels/Stack/NiagaraStackShaderStagesGroup.cpp b/Engine/Plugins/FX/Niagara/Source/NiagaraEditor/Private/ViewModels/Stack/NiagaraStackShaderStagesGroup.cpp
new file mode 100644
--- /dev/null
+++ b/Engine/Plugins/FX/Niagara/Source/NiagaraEditor/Private/ViewModels/Stack/NiagaraStackShaderStagesGroup.cpp
@@ -0,0 +1,135 @@
+// Copyright 1998-2019 Epic Games, Inc. All Rights Reserved.
+
+#include "ViewModels/Stack/NiagaraStackShaderStagesGroup.h"
+#include "ViewModels/NiagaraEmitterViewModel.h"
+#include "NiagaraScriptViewModel.h"
+#include "NiagaraScriptGraphViewModel.h"
+#include "NiagaraEmitter.h"
+#include "NiagaraShaderStageBase.h"
+#include "NiagaraGraph.h"
+#include "NiagaraScriptSource.h"
+#include "ViewModels/Stack/INiagaraStackItemGroupAddUtilities.h"
+#include "ViewModels/Stack/NiagaraStackGraphUtilities.h"
+
+#include "ScopedTransaction.h"
+
+#define LOCTEXT_NAMESPACE "NiagaraStackShaderStagesGroup"
+
+class FShaderStagesGroupAddAction : public INiagaraStackItemGroupAddAction
+{
+public:
+	FShaderStagesGroupAddAction(UClass* InShaderStageClass)
+		: ShaderStageClass(InShaderStageClass)
+	{
+	}
+
+	virtual FText GetCategory() const override
+	{
+		return LOCTEXT("AddShaderStageCategory", "Add Shader Stage");
+	}
+
+	virtual FText GetDisplayName() const override
+	{
+		return ShaderStageClass->GetDisplayNameText();
+	}
+
+	virtual FText GetDescription() const override
+	{
+		return FText::FromString(ShaderStageClass->GetDescription());
+	}
+
+	virtual FText GetKeywords() const override
+	{
+		return FText();
+	}
+
+	UClass* GetShaderStageClass() const
+	{
+		return ShaderStageClass;
+	}
+
+private:
+	UClass* ShaderStageClass;
+};
+
+class FShaderStagesGroupAddUtilities : public TNiagaraStackItemGroupAddUtilities<UNiagaraShaderStageBase*>
+{
+public:
+	FShaderStagesGroupAddUtilities(TSharedRef<FNiagaraEmitterViewModel> InEmitterViewModel, FOnItemAdded InOnItemAdded)
+		: TNiagaraStackItemGroupAddUtilities(LOCTEXT("ShaderStagesGroupAddItemName", "Shader Stage"), EAddMode::AddFromAction, true, InOnItemAdded)
+		, EmitterViewModel(InEmitterViewModel)
+	{
+	}
+
+	virtual void AddItemDirectly() override	{ unimplemented(); }
+
+	virtual void GenerateAddActions(TArray<TSharedRef<INiagaraStackItemGroupAddAction>>& OutAddActions, const FNiagaraStackItemGroupAddOptions& AddProperties) const override 
+	{
+		TArray<UClass*> ShaderStageClasses;
+		GetDerivedClasses(UNiagaraShaderStageBase::StaticClass(), ShaderStageClasses);
+		for (UClass* ShaderStageClass : ShaderStageClasses)
+		{
+			OutAddActions.Add(MakeShared<FShaderStagesGroupAddAction>(ShaderStageClass));
+		}
+	}
+
+	virtual void ExecuteAddAction(TSharedRef<INiagaraStackItemGroupAddAction> AddAction, int32 TargetIndex) override 
+	{ 
+		TSharedPtr<FNiagaraEmitterViewModel> EmitterViewModelPinned = EmitterViewModel.Pin();
+		if (EmitterViewModelPinned.IsValid() == false)
+		{
+			return;
+		}
+
+		UNiagaraEmitter* Emitter = EmitterViewModelPinned->GetEmitter();
+		UNiagaraScriptSource* Source = EmitterViewModelPinned->GetSharedScriptViewModel()->GetGraphViewModel()->GetScriptSource();
+		UNiagaraGraph* Graph = EmitterViewModelPinned->GetSharedScriptViewModel()->GetGraphViewModel()->GetGraph();
+		TSharedRef<FShaderStagesGroupAddAction> ShaderStageAddAction = StaticCastSharedRef<FShaderStagesGroupAddAction>(AddAction);
+
+		// The stack should not have been created if any of these are null, so bail out if it happens somehow rather than try to handle all of these cases.
+		checkf(Emitter != nullptr && Source != nullptr && Graph != nullptr, TEXT("Stack created for invalid emitter or graph."));
+
+		FScopedTransaction ScopedTransaction(LOCTEXT("AddNewShaderStagesTransaction", "Add new shader stage"));
+
+		Emitter->Modify();
+		UNiagaraShaderStageBase* ShaderStage = NewObject<UNiagaraShaderStageBase>(Emitter, ShaderStageAddAction->GetShaderStageClass(), NAME_None, RF_Transactional);
+		ShaderStage->Script = NewObject<UNiagaraScript>(ShaderStage, MakeUniqueObjectName(ShaderStage, UNiagaraScript::StaticClass(), "ShaderStage"), EObjectFlags::RF_Transactional);
+		ShaderStage->Script->SetUsage(ENiagaraScriptUsage::ParticleShaderStageScript);
+		ShaderStage->Script->SetUsageId(ShaderStage->GetMergeId());
+		ShaderStage->Script->SetSource(Source);
+		Emitter->AddShaderStage(ShaderStage);
+		FNiagaraStackGraphUtilities::ResetGraphForOutput(*Graph, ENiagaraScriptUsage::ParticleShaderStageScript, ShaderStage->Script->GetUsageId());
+
+		// Set the emitter here so that the internal state of the view model is updated.
+		// TODO: Move the logic for managing additional scripts into the emitter view model or script view model.
+		TWeakPtr<FNiagaraEmitterInstance, ESPMode::ThreadSafe> Simulation = EmitterViewModelPinned->GetSimulation();
+		EmitterViewModelPinned->Reset();
+		EmitterViewModelPinned->Initialize(Emitter, Simulation);
+
+		OnItemAdded.ExecuteIfBound(ShaderStage);
+	}
+
+private:
+	TWeakPtr<FNiagaraEmitterViewModel> EmitterViewModel;
+};
+
+void UNiagaraStackShaderStagesGroup::Initialize(FRequiredEntryData InRequiredEntryData)
+{
+	FText DisplayName = LOCTEXT("ShaderStageGroupName", "Add Shader Stage");
+	FText ToolTip = LOCTEXT("ShaderStageGroupTooltip", "Shader stages for this emitter.");
+	AddUtilities = MakeShared<FShaderStagesGroupAddUtilities>(InRequiredEntryData.EmitterViewModel.ToSharedRef(), 
+		TNiagaraStackItemGroupAddUtilities<UNiagaraShaderStageBase*>::FOnItemAdded::CreateUObject(this, &UNiagaraStackShaderStagesGroup::ItemAddedFromUtilties));
+	Super::Initialize(InRequiredEntryData, DisplayName, ToolTip, AddUtilities.Get());
+}
+
+void UNiagaraStackShaderStagesGroup::SetOnItemAdded(FOnItemAdded InOnItemAdded)
+{
+	ItemAddedDelegate = InOnItemAdded;
+}
+
+void UNiagaraStackShaderStagesGroup::ItemAddedFromUtilties(UNiagaraShaderStageBase* AddedShaderStage)
+{
+	ItemAddedDelegate.ExecuteIfBound();
+}
+
+#undef LOCTEXT_NAMESPACE
diff --git a/Engine/Plugins/FX/Niagara/Source/NiagaraEditor/Public/ViewModels/Stack/NiagaraStackEmitterSettingsGroup.h b/Engine/Plugins/FX/Niagara/Source/NiagaraEditor/Public/ViewModels/Stack/NiagaraStackEmitterSettingsGroup.h
--- a/Engine/Plugins/FX/Niagara/Source/NiagaraEditor/Public/ViewModels/Stack/NiagaraStackEmitterSettingsGroup.h
+++ b/Engine/Plugins/FX/Niagara/Source/NiagaraEditor/Public/ViewModels/Stack/NiagaraStackEmitterSettingsGroup.h
@@ -22,9 +22,9 @@ public:
 
 	virtual FText GetTooltipText() const override;
 
-	bool CanResetToBase() const;
-
-	void ResetToBase();
+	virtual bool SupportsResetToBase() const override { return true; }
+	virtual bool TestCanResetToBaseWithMessage(FText& OutCanResetToBaseMessage) const override;
+	virtual void ResetToBase() override;
 
 	virtual bool IsExpandedByDefault() const override;
 
diff --git a/Engine/Plugins/FX/Niagara/Source/NiagaraEditor/Public/ViewModels/Stack/NiagaraStackEntry.h b/Engine/Plugins/FX/Niagara/Source/NiagaraEditor/Public/ViewModels/Stack/NiagaraStackEntry.h
--- a/Engine/Plugins/FX/Niagara/Source/NiagaraEditor/Public/ViewModels/Stack/NiagaraStackEntry.h
+++ b/Engine/Plugins/FX/Niagara/Source/NiagaraEditor/Public/ViewModels/Stack/NiagaraStackEntry.h
@@ -110,6 +110,7 @@ public:
 		static const FName Spawn;
 		static const FName Update;
 		static const FName Event;
+		static const FName ShaderStage;
 		static const FName Render;
 	};
 
diff --git a/Engine/Plugins/FX/Niagara/Source/NiagaraEditor/Public/ViewModels/Stack/NiagaraStackEventScriptItemGroup.h b/Engine/Plugins/FX/Niagara/Source/NiagaraEditor/Public/ViewModels/Stack/NiagaraStackEventScriptItemGroup.h
--- a/Engine/Plugins/FX/Niagara/Source/NiagaraEditor/Public/ViewModels/Stack/NiagaraStackEventScriptItemGroup.h
+++ b/Engine/Plugins/FX/Niagara/Source/NiagaraEditor/Public/ViewModels/Stack/NiagaraStackEventScriptItemGroup.h
@@ -23,9 +23,9 @@ public:
 
 	virtual FText GetDisplayName() const override;
 
-	bool CanResetToBase() const;
-
-	void ResetToBase();
+	virtual bool SupportsResetToBase() const override { return true; }
+	virtual bool TestCanResetToBaseWithMessage(FText& OutCanResetToBaseMessage) const override;
+	virtual void ResetToBase() override;
 
 protected:
 	virtual void FinalizeInternal() override;
diff --git a/Engine/Plugins/FX/Niagara/Source/NiagaraEditor/Public/ViewModels/Stack/NiagaraStackItem.h b/Engine/Plugins/FX/Niagara/Source/NiagaraEditor/Public/ViewModels/Stack/NiagaraStackItem.h
--- a/Engine/Plugins/FX/Niagara/Source/NiagaraEditor/Public/ViewModels/Stack/NiagaraStackItem.h
+++ b/Engine/Plugins/FX/Niagara/Source/NiagaraEditor/Public/ViewModels/Stack/NiagaraStackItem.h
@@ -36,6 +36,10 @@ public:
 
 	virtual bool SupportsIcon() const { return false; }
 	virtual const FSlateBrush* GetIconBrush() const;
+
+	virtual bool SupportsResetToBase() const { return false; }
+	virtual bool TestCanResetToBaseWithMessage(FText& OutCanResetToBaseMessage) const { return false; }
+	virtual void ResetToBase() { }
 	
 protected:
 	virtual void RefreshChildrenInternal(const TArray<UNiagaraStackEntry*>& CurrentChildren, TArray<UNiagaraStackEntry*>& NewChildren, TArray<FStackIssue>& NewIssues) override;
diff --git a/Engine/Plugins/FX/Niagara/Source/NiagaraEditor/Public/ViewModels/Stack/NiagaraStackRendererItem.h b/Engine/Plugins/FX/Niagara/Source/NiagaraEditor/Public/ViewModels/Stack/NiagaraStackRendererItem.h
--- a/Engine/Plugins/FX/Niagara/Source/NiagaraEditor/Public/ViewModels/Stack/NiagaraStackRendererItem.h
+++ b/Engine/Plugins/FX/Niagara/Source/NiagaraEditor/Public/ViewModels/Stack/NiagaraStackRendererItem.h
@@ -46,16 +46,16 @@ public:
 
 	bool HasBaseRenderer() const;
 
-	bool CanResetToBase() const;
-
-	void ResetToBase();
-
 	virtual bool SupportsChangeEnabled() const override { return true; }
 	virtual bool GetIsEnabled() const override;
 
 	virtual bool SupportsIcon() const override { return true; }
 	virtual const FSlateBrush* GetIconBrush() const override;
 
+	virtual bool SupportsResetToBase() const override { return true; }
+	virtual bool TestCanResetToBaseWithMessage(FText& OutCanResetToBaseMessage) const override;
+	virtual void ResetToBase() override;
+
 	static TArray<FNiagaraVariable> GetMissingVariables(UNiagaraRendererProperties* RendererProperties, UNiagaraEmitter* Emitter);
 	static bool AddMissingVariable(UNiagaraEmitter* Emitter, const FNiagaraVariable& Variable);
 
diff --git a/Engine/Plugins/FX/Niagara/Source/NiagaraEditor/Public/ViewModels/Stack/NiagaraStackRoot.h b/Engine/Plugins/FX/Niagara/Source/NiagaraEditor/Public/ViewModels/Stack/NiagaraStackRoot.h
--- a/Engine/Plugins/FX/Niagara/Source/NiagaraEditor/Public/ViewModels/Stack/NiagaraStackRoot.h
+++ b/Engine/Plugins/FX/Niagara/Source/NiagaraEditor/Public/ViewModels/Stack/NiagaraStackRoot.h
@@ -10,6 +10,7 @@ class UNiagaraStackSystemSettingsGroup;
 class UNiagaraStackEmitterSettingsGroup;
 class UNiagaraStackScriptItemGroup;
 class UNiagaraStackEventHandlerGroup;
+class UNiagaraStackShaderStagesGroup;
 class UNiagaraStackRenderItemGroup;
 
 UCLASS()
@@ -32,7 +33,7 @@ protected:
 	virtual void RefreshChildrenInternal(const TArray<UNiagaraStackEntry*>& CurrentChildren, TArray<UNiagaraStackEntry*>& NewChildren, TArray<FStackIssue>& NewIssues) override;
 
 private:
-	void EmitterEventArraysChanged();
+	void EmitterArraysChanged();
 
 private:
 	UPROPERTY()
@@ -62,6 +63,9 @@ private:
 	UPROPERTY()
 	UNiagaraStackEventHandlerGroup* AddEventHandlerGroup;
 
+	UPROPERTY()
+	UNiagaraStackShaderStagesGroup* AddShaderStageGroup;
+
 	UPROPERTY()
 	UNiagaraStackRenderItemGroup* RenderGroup;
 
diff --git a/Engine/Plugins/FX/Niagara/Source/NiagaraEditor/Public/ViewModels/Stack/NiagaraStackShaderStageGroup.h b/Engine/Plugins/FX/Niagara/Source/NiagaraEditor/Public/ViewModels/Stack/NiagaraStackShaderStageGroup.h
new file mode 100644
--- /dev/null
+++ b/Engine/Plugins/FX/Niagara/Source/NiagaraEditor/Public/ViewModels/Stack/NiagaraStackShaderStageGroup.h
@@ -0,0 +1,99 @@
+// Copyright 1998-2019 Epic Games, Inc. All Rights Reserved.
+
+#pragma once
+
+#include "NiagaraCommon.h"
+#include "NiagaraEmitter.h"
+#include "ViewModels/Stack/NiagaraStackItem.h"
+#include "ViewModels/Stack/NiagaraStackScriptItemGroup.h"
+#include "NiagaraStackShaderStageGroup.generated.h"
+
+class FNiagaraScriptViewModel;
+class UNiagaraStackObject;
+struct FNiagaraEventScriptProperties;
+class IDetailTreeNode;
+
+UCLASS()
+class NIAGARAEDITOR_API UNiagaraStackShaderStagePropertiesItem : public UNiagaraStackItem
+{
+	GENERATED_BODY()
+		 
+public:
+	void Initialize(FRequiredEntryData InRequiredEntryData, UNiagaraShaderStageBase* InShaderStage);
+
+	virtual FText GetDisplayName() const override;
+
+	virtual bool SupportsResetToBase() const override { return true; }
+	virtual bool TestCanResetToBaseWithMessage(FText& OutCanResetToBaseMessage) const override;
+	virtual void ResetToBase() override;
+
+protected:
+	virtual void FinalizeInternal() override;
+
+	virtual void RefreshChildrenInternal(const TArray<UNiagaraStackEntry*>& CurrentChildren, TArray<UNiagaraStackEntry*>& NewChildren, TArray<FStackIssue>& NewIssues) override;
+
+private:
+	void ShaderStagePropertiesChanged();
+
+	bool HasBaseShaderStage() const;
+
+private:
+	TWeakObjectPtr<UNiagaraShaderStageBase> ShaderStage;
+
+	mutable TOptional<bool> bHasBaseShaderStageCache;
+
+	mutable TOptional<bool> bCanResetToBaseCache;
+
+	UPROPERTY()
+	UNiagaraStackObject* ShaderStageObject;
+};
+
+UCLASS()
+/** Meant to contain a single binding of a Emitter::EventScriptProperties to the stack.*/
+class NIAGARAEDITOR_API UNiagaraStackShaderStageGroup : public UNiagaraStackScriptItemGroup
+{
+	GENERATED_BODY()
+
+public:
+	DECLARE_DELEGATE(FOnModifiedShaderStages);
+
+public:
+	void Initialize(
+		FRequiredEntryData InRequiredEntryData,
+		TSharedRef<FNiagaraScriptViewModel> InScriptViewModel,
+		UNiagaraShaderStageBase* InShaderStage);
+
+	UNiagaraShaderStageBase* GetShaderStage() const;
+
+	void SetOnModifiedShaderStages(FOnModifiedShaderStages OnModifiedShaderStages);
+
+	virtual bool SupportsDelete() const override { return true; }
+	virtual bool TestCanDeleteWithMessage(FText& OutCanDeleteMessage) const override;
+	virtual void Delete() override;
+
+	virtual bool CanDrag() const override { return true; }
+
+protected:
+	virtual void FinalizeInternal() override;
+
+	virtual void RefreshChildrenInternal(const TArray<UNiagaraStackEntry*>& CurrentChildren, TArray<UNiagaraStackEntry*>& NewChildren, TArray<FStackIssue>& NewIssues) override;
+
+	virtual TOptional<FDropRequestResponse> CanDropInternal(const FDropRequest& DropRequest) override;
+
+	virtual TOptional<FDropRequestResponse> DropInternal(const FDropRequest& DropRequest) override;
+
+private:
+	void ShaderStagePropertiesChanged();
+
+	bool HasBaseShaderStage() const;
+	
+private:
+	TWeakObjectPtr<UNiagaraShaderStageBase> ShaderStage;
+
+	mutable TOptional<bool> bHasBaseShaderStageCache;
+
+	FOnModifiedShaderStages OnModifiedShaderStagesDelegate;
+
+	UPROPERTY()
+	UNiagaraStackShaderStagePropertiesItem* ShaderStageProperties;
+};
diff --git a/Engine/Plugins/FX/Niagara/Source/NiagaraEditor/Public/ViewModels/Stack/NiagaraStackShaderStagesGroup.h b/Engine/Plugins/FX/Niagara/Source/NiagaraEditor/Public/ViewModels/Stack/NiagaraStackShaderStagesGroup.h
new file mode 100644
--- /dev/null
+++ b/Engine/Plugins/FX/Niagara/Source/NiagaraEditor/Public/ViewModels/Stack/NiagaraStackShaderStagesGroup.h
@@ -0,0 +1,30 @@
+// Copyright 1998-2019 Epic Games, Inc. All Rights Reserved.
+
+#pragma once
+
+#include "ViewModels/Stack/NiagaraStackItemGroup.h"
+#include "NiagaraStackShaderStagesGroup.generated.h"
+
+class UNiagaraShaderStageBase;
+
+UCLASS()
+/** Container for one or more NiagaraStackEventScriptItemGroups, allowing multiple event handlers per script.*/
+class NIAGARAEDITOR_API UNiagaraStackShaderStagesGroup : public UNiagaraStackItemGroup
+{
+	GENERATED_BODY()
+
+public:
+	DECLARE_DELEGATE(FOnItemAdded);
+
+public:
+	void Initialize(FRequiredEntryData InRequiredEntryData);
+
+	void SetOnItemAdded(FOnItemAdded InOnItemAdded);
+
+private:
+	void ItemAddedFromUtilties(UNiagaraShaderStageBase* AddedShaderStage);
+
+private:
+	TSharedPtr<INiagaraStackItemGroupAddUtilities> AddUtilities;
+	FOnItemAdded ItemAddedDelegate;
+};
diff --git a/Engine/Plugins/FX/Niagara/Source/NiagaraEditorWidgets/Private/NiagaraEditorWidgetsStyle.cpp b/Engine/Plugins/FX/Niagara/Source/NiagaraEditorWidgets/Private/NiagaraEditorWidgetsStyle.cpp
--- a/Engine/Plugins/FX/Niagara/Source/NiagaraEditorWidgets/Private/NiagaraEditorWidgetsStyle.cpp
+++ b/Engine/Plugins/FX/Niagara/Source/NiagaraEditorWidgets/Private/NiagaraEditorWidgetsStyle.cpp
@@ -188,12 +188,14 @@ TSharedRef< FSlateStyleSet > FNiagaraEditorWidgetsStyle::Create()
 	Style->Set("NiagaraEditor.Stack.SpawnIcon", new IMAGE_PLUGIN_BRUSH("Icons/Spawn", Icon12x12, FLinearColor::White));
 	Style->Set("NiagaraEditor.Stack.UpdateIcon", new IMAGE_PLUGIN_BRUSH("Icons/Update", Icon12x12, FLinearColor::White));
 	Style->Set("NiagaraEditor.Stack.EventIcon", new IMAGE_PLUGIN_BRUSH("Icons/Event", Icon12x12, FLinearColor::White));
+	Style->Set("NiagaraEditor.Stack.ShaderStageIcon", new IMAGE_PLUGIN_BRUSH("Icons/ShaderStage", Icon12x12, FLinearColor::White));
 	Style->Set("NiagaraEditor.Stack.RenderIcon", new IMAGE_PLUGIN_BRUSH("Icons/Render", Icon12x12, FLinearColor::White));
 
 	Style->Set("NiagaraEditor.Stack.ParametersIconHighlighted", new IMAGE_PLUGIN_BRUSH("Icons/SystemParams", Icon16x16, FLinearColor::White));
 	Style->Set("NiagaraEditor.Stack.SpawnIconHighlighted", new IMAGE_PLUGIN_BRUSH("Icons/Spawn", Icon16x16, FLinearColor::White));
 	Style->Set("NiagaraEditor.Stack.UpdateIconHighlighted", new IMAGE_PLUGIN_BRUSH("Icons/Update", Icon16x16, FLinearColor::White));
 	Style->Set("NiagaraEditor.Stack.EventIconHighlighted", new IMAGE_PLUGIN_BRUSH("Icons/Event", Icon16x16, FLinearColor::White));
+	Style->Set("NiagaraEditor.Stack.ShaderStageIconHighlighted", new IMAGE_PLUGIN_BRUSH("Icons/ShaderStage", Icon16x16, FLinearColor::White));
 	Style->Set("NiagaraEditor.Stack.RenderIconHighlighted", new IMAGE_PLUGIN_BRUSH("Icons/Render", Icon16x16, FLinearColor::White));
 
 	Style->Set("NiagaraEditor.Stack.IconHighlightedSize", 16.0f);
diff --git a/Engine/Plugins/FX/Niagara/Source/NiagaraEditorWidgets/Private/NiagaraEditorWidgetsUtilities.cpp b/Engine/Plugins/FX/Niagara/Source/NiagaraEditorWidgets/Private/NiagaraEditorWidgetsUtilities.cpp
--- a/Engine/Plugins/FX/Niagara/Source/NiagaraEditorWidgets/Private/NiagaraEditorWidgetsUtilities.cpp
+++ b/Engine/Plugins/FX/Niagara/Source/NiagaraEditorWidgets/Private/NiagaraEditorWidgetsUtilities.cpp
@@ -64,6 +64,10 @@ FName FNiagaraStackEditorWidgetsUtilities::GetIconNameForExecutionSubcategory(FN
 		{
 			return "NiagaraEditor.Stack.EventIconHighlighted";
 		}
+		else if (ExecutionSubcategoryName == UNiagaraStackEntry::FExecutionSubcategoryNames::ShaderStage)
+		{
+			return "NiagaraEditor.Stack.ShaderStageIconHighlighted";
+		}
 		else if (ExecutionSubcategoryName == UNiagaraStackEntry::FExecutionSubcategoryNames::Render)
 		{
 			return "NiagaraEditor.Stack.RenderIconHighlighted";
@@ -87,6 +91,10 @@ FName FNiagaraStackEditorWidgetsUtilities::GetIconNameForExecutionSubcategory(FN
 		{
 			return "NiagaraEditor.Stack.EventIcon";
 		}
+		else if (ExecutionSubcategoryName == UNiagaraStackEntry::FExecutionSubcategoryNames::ShaderStage)
+		{
+			return "NiagaraEditor.Stack.ShaderStageIcon";
+		}
 		else if (ExecutionSubcategoryName == UNiagaraStackEntry::FExecutionSubcategoryNames::Render)
 		{
 			return "NiagaraEditor.Stack.RenderIcon";
diff --git a/Engine/Plugins/FX/Niagara/Source/NiagaraEditorWidgets/Private/SNiagaraStack.cpp b/Engine/Plugins/FX/Niagara/Source/NiagaraEditorWidgets/Private/SNiagaraStack.cpp
--- a/Engine/Plugins/FX/Niagara/Source/NiagaraEditorWidgets/Private/SNiagaraStack.cpp
+++ b/Engine/Plugins/FX/Niagara/Source/NiagaraEditorWidgets/Private/SNiagaraStack.cpp
@@ -12,6 +12,7 @@
 #include "ViewModels/NiagaraEmitterHandleViewModel.h"
 #include "ViewModels/NiagaraSystemViewModel.h"
 #include "ViewModels/Stack/NiagaraStackViewModel.h"
+#include "ViewModels/Stack/NiagaraStackItemGroup.h"
 #include "ViewModels/Stack/NiagaraStackModuleItemOutputCollection.h"
 #include "ViewModels/Stack/NiagaraStackModuleItemLinkedInputCollection.h"
 #include "ViewModels/Stack/NiagaraStackModuleItemOutput.h"
@@ -21,10 +22,7 @@
 #include "ViewModels/Stack/NiagaraStackItemFooter.h"
 #include "ViewModels/Stack/NiagaraStackFunctionInputCollection.h"
 #include "ViewModels/Stack/NiagaraStackModuleItem.h"
-#include "ViewModels/Stack/NiagaraStackRendererItem.h"
 #include "ViewModels/Stack/NiagaraStackParameterStoreEntry.h"
-#include "ViewModels/Stack/NiagaraStackEmitterSettingsGroup.h"
-#include "ViewModels/Stack/NiagaraStackEventScriptItemGroup.h"
 #include "Framework/Application/SlateApplication.h"
 #include "IContentBrowserSingleton.h"
 #include "ContentBrowserModule.h"
@@ -37,14 +35,12 @@
 #include "IDetailTreeNode.h"
 #include "Stack/SNiagaraStackFunctionInputName.h"
 #include "Stack/SNiagaraStackFunctionInputValue.h"
-#include "Stack/SNiagaraStackEmitterPropertiesItem.h"
-#include "Stack/SNiagaraStackEventHandlerPropertiesItem.h"
+#include "Stack/SNiagaraStackItem.h"
 #include "Stack/SNiagaraStackItemFooter.h"
 #include "Stack/SNiagaraStackItemGroup.h"
 #include "Stack/SNiagaraStackModuleItem.h"
 #include "Stack/SNiagaraStackParameterStoreEntryName.h"
 #include "Stack/SNiagaraStackParameterStoreEntryValue.h"
-#include "Stack/SNiagaraStackRendererItem.h"
 #include "Stack/SNiagaraStackTableRow.h"
 #include "Stack/SNiagaraStackIssueIcon.h"
 #include "NiagaraEditorWidgetsUtilities.h"
@@ -895,10 +891,6 @@ SNiagaraStack::FRowWidgets SNiagaraStack::ConstructNameAndValueWidgetsForItem(UN
 		Container->AddFillRowContextMenuHandler(SNiagaraStackTableRow::FOnFillRowContextMenu::CreateSP(ModuleItemWidget, &SNiagaraStackModuleItem::FillRowContextMenu));
 		return FRowWidgets(ModuleItemWidget);
 	}
-	else if (Item->IsA<UNiagaraStackRendererItem>())
-	{
-		return FRowWidgets(SNew(SNiagaraStackRendererItem, *CastChecked<UNiagaraStackRendererItem>(Item), StackViewModel));
-	}
 	else if (Item->IsA<UNiagaraStackFunctionInput>())
 	{
 		UNiagaraStackFunctionInput* FunctionInput = CastChecked<UNiagaraStackFunctionInput>(Item);
@@ -930,16 +922,6 @@ SNiagaraStack::FRowWidgets SNiagaraStack::ConstructNameAndValueWidgetsForItem(UN
 		UNiagaraStackItemFooter* ItemExpander = CastChecked<UNiagaraStackItemFooter>(Item);
 		return FRowWidgets(SNew(SNiagaraStackItemFooter, *ItemExpander));
 	}
-	else if (Item->IsA<UNiagaraStackEmitterPropertiesItem>())
-	{
-		UNiagaraStackEmitterPropertiesItem* PropertiesItem = CastChecked<UNiagaraStackEmitterPropertiesItem>(Item);
-		return FRowWidgets(SNew(SNiagaraStackEmitterPropertiesItem, *PropertiesItem, StackViewModel));
-	}
-	else if (Item->IsA<UNiagaraStackEventHandlerPropertiesItem>())
-	{
-		UNiagaraStackEventHandlerPropertiesItem* PropertiesItem = CastChecked<UNiagaraStackEventHandlerPropertiesItem>(Item);
-		return FRowWidgets(SNew(SNiagaraStackEventHandlerPropertiesItem, *PropertiesItem, StackViewModel));
-	}
 	else if (Item->IsA<UNiagaraStackParameterStoreEntry>())
 	{
 		UNiagaraStackParameterStoreEntry* StackEntry = CastChecked<UNiagaraStackParameterStoreEntry>(Item);
@@ -1016,12 +998,8 @@ SNiagaraStack::FRowWidgets SNiagaraStack::ConstructNameAndValueWidgetsForItem(UN
 	}
 	else if (Item->IsA<UNiagaraStackItem>())
 	{
-		return FRowWidgets(SNew(STextBlock)
-			.TextStyle(FNiagaraEditorWidgetsStyle::Get(), "NiagaraEditor.Stack.ItemText")
-			.ToolTipText_UObject(Item, &UNiagaraStackEntry::GetTooltipText)
-			.Text_UObject(Item, &UNiagaraStackEntry::GetDisplayName)
-			.ColorAndOpacity(this, &SNiagaraStack::GetTextColorForItem, Item)
-			.HighlightText_UObject(StackViewModel, &UNiagaraStackViewModel::GetCurrentSearchText));
+		UNiagaraStackItem* StackItem = CastChecked<UNiagaraStackItem>(Item);
+		return FRowWidgets(SNew(SNiagaraStackItem, *StackItem, StackViewModel));
 	}
 	else
 	{
diff --git a/Engine/Plugins/FX/Niagara/Source/NiagaraEditorWidgets/Private/Stack/SNiagaraStackEmitterPropertiesItem.cpp b/Engine/Plugins/FX/Niagara/Source/NiagaraEditorWidgets/Private/Stack/SNiagaraStackEmitterPropertiesItem.cpp
deleted file mode 100644
diff --git a/Engine/Plugins/FX/Niagara/Source/NiagaraEditorWidgets/Private/Stack/SNiagaraStackEmitterPropertiesItem.h b/Engine/Plugins/FX/Niagara/Source/NiagaraEditorWidgets/Private/Stack/SNiagaraStackEmitterPropertiesItem.h
deleted file mode 100644
diff --git a/Engine/Plugins/FX/Niagara/Source/NiagaraEditorWidgets/Private/Stack/SNiagaraStackEventHandlerPropertiesItem.cpp b/Engine/Plugins/FX/Niagara/Source/NiagaraEditorWidgets/Private/Stack/SNiagaraStackEventHandlerPropertiesItem.cpp
deleted file mode 100644
diff --git a/Engine/Plugins/FX/Niagara/Source/NiagaraEditorWidgets/Private/Stack/SNiagaraStackEventHandlerPropertiesItem.h b/Engine/Plugins/FX/Niagara/Source/NiagaraEditorWidgets/Private/Stack/SNiagaraStackEventHandlerPropertiesItem.h
deleted file mode 100644
diff --git a/Engine/Plugins/FX/Niagara/Source/NiagaraEditorWidgets/Private/Stack/SNiagaraStackItem.cpp b/Engine/Plugins/FX/Niagara/Source/NiagaraEditorWidgets/Private/Stack/SNiagaraStackItem.cpp
new file mode 100644
--- /dev/null
+++ b/Engine/Plugins/FX/Niagara/Source/NiagaraEditorWidgets/Private/Stack/SNiagaraStackItem.cpp
@@ -0,0 +1,164 @@
+// Copyright 1998-2019 Epic Games, Inc. All Rights Reserved.
+
+#include "Stack/SNiagaraStackItem.h"
+#include "NiagaraEditorWidgetsStyle.h"
+#include "NiagaraEditorStyle.h"
+#include "EditorStyleSet.h"
+#include "ViewModels/Stack/NiagaraStackItem.h"
+#include "ViewModels/Stack/NiagaraStackViewModel.h"
+#include "Widgets/SBoxPanel.h"
+#include "Widgets/Input/SButton.h"
+#include "Widgets/Input/SCheckBox.h"
+#include "Widgets/Images/SImage.h"
+#include "Widgets/Text/STextBlock.h"
+#include "Widgets/Images/SImage.h"
+#include "Widgets/Layout/SBox.h"
+#include "SNiagaraStackIssueIcon.h"
+
+#define LOCTEXT_NAMESPACE "NiagaraStackItem"
+
+void SNiagaraStackItem::Construct(const FArguments& InArgs, UNiagaraStackItem& InItem, UNiagaraStackViewModel* InStackViewModel)
+{
+	Item = &InItem;
+	StackEntryItem = Item;
+	StackViewModel = InStackViewModel;
+
+	TSharedRef<SHorizontalBox> RowBox = SNew(SHorizontalBox);
+
+	// Icon
+	if (Item->SupportsIcon())
+	{
+		RowBox->AddSlot()
+		.Padding(2, 0, 3, 0)
+		.AutoWidth()
+		.VAlign(VAlign_Center)
+		[
+			SNew(SImage)
+			.Image(Item->GetIconBrush())
+		];
+	}
+
+	// Display name
+	RowBox->AddSlot()
+	.Padding(2, 0, 2, 0)
+	.VAlign(VAlign_Center)
+	[
+		SNew(SNiagaraStackDisplayName, InItem, *InStackViewModel, "NiagaraEditor.Stack.ItemText")
+	];
+
+	// Reset to base button
+	if (Item->SupportsResetToBase())
+	{
+		RowBox->AddSlot()
+		.AutoWidth()
+		.VAlign(VAlign_Center)
+		.Padding(0, 0, 2, 0)
+		[
+			SNew(SButton)
+			.Visibility(this, &SNiagaraStackItem::GetResetToBaseButtonVisibility)
+			.IsFocusable(false)
+			.ToolTipText(this, &SNiagaraStackItem::GetResetToBaseButtonToolTipText)
+			.ButtonStyle(FEditorStyle::Get(), "NoBorder")
+			.ContentPadding(0)
+			.OnClicked(this, &SNiagaraStackItem::ResetToBaseButtonClicked)
+			.Content()
+			[
+				SNew(SImage)
+				.Image(FEditorStyle::GetBrush("PropertyWindow.DiffersFromDefault"))
+				.ColorAndOpacity(FSlateColor(FLinearColor::Green))
+			]
+		];
+	}
+
+	// Delete button
+	if (Item->SupportsDelete())
+	{
+		RowBox->AddSlot()
+		.AutoWidth()
+		.VAlign(VAlign_Center)
+		.Padding(0, 0, 2, 0)
+		[
+			SNew(SButton)
+			.ButtonStyle(FEditorStyle::Get(), "HoverHintOnly")
+			.IsFocusable(false)
+			.ForegroundColor(FNiagaraEditorWidgetsStyle::Get().GetColor("NiagaraEditor.Stack.FlatButtonColor"))
+			.ToolTipText(this, &SNiagaraStackItem::GetDeleteButtonToolTipText)
+			.IsEnabled(this, &SNiagaraStackItem::GetDeleteButtonEnabled)
+			.OnClicked(this, &SNiagaraStackItem::DeleteClicked)
+			.Content()
+			[
+				SNew(STextBlock)
+				.Font(FEditorStyle::Get().GetFontStyle("FontAwesome.10"))
+				.Text(FText::FromString(FString(TEXT("\xf1f8"))))
+			]
+		];
+	}
+
+	// Enabled checkbox
+	if (Item->SupportsChangeEnabled())
+	{
+		RowBox->AddSlot()
+		.Padding(0)
+		.AutoWidth()
+		[
+			SNew(SCheckBox)
+			.IsChecked(this, &SNiagaraStackItem::CheckEnabledStatus)
+			.OnCheckStateChanged(this, &SNiagaraStackItem::OnCheckStateChanged)
+		];
+	}
+
+	ChildSlot
+	[
+		RowBox
+	];
+}
+
+EVisibility SNiagaraStackItem::GetResetToBaseButtonVisibility() const
+{
+	FText Unused;
+	return Item->TestCanResetToBaseWithMessage(Unused) ? EVisibility::Visible : EVisibility::Collapsed;
+}
+
+FText SNiagaraStackItem::GetResetToBaseButtonToolTipText() const
+{
+	FText CanResetToBaseMessage;
+	Item->TestCanCopyWithMessage(CanResetToBaseMessage);
+	return CanResetToBaseMessage;
+}
+
+FReply SNiagaraStackItem::ResetToBaseButtonClicked()
+{
+	Item->ResetToBase();
+	return FReply::Handled();
+}
+
+FText SNiagaraStackItem::GetDeleteButtonToolTipText() const
+{
+	FText CanDeleteMessage;
+	Item->TestCanDeleteWithMessage(CanDeleteMessage);
+	return CanDeleteMessage;
+}
+
+bool SNiagaraStackItem::GetDeleteButtonEnabled() const
+{
+	FText Unused;
+	return Item->TestCanDeleteWithMessage(Unused);
+}
+
+FReply SNiagaraStackItem::DeleteClicked()
+{
+	Item->Delete();
+	return FReply::Handled();
+}
+
+void SNiagaraStackItem::OnCheckStateChanged(ECheckBoxState InCheckState)
+{
+	Item->SetIsEnabled(InCheckState == ECheckBoxState::Checked);
+}
+
+ECheckBoxState SNiagaraStackItem::CheckEnabledStatus() const
+{
+	return Item->GetIsEnabled() ? ECheckBoxState::Checked : ECheckBoxState::Unchecked;
+}
+
+#undef LOCTEXT_NAMESPACE
\ No newline at end of file
diff --git a/Engine/Plugins/FX/Niagara/Source/NiagaraEditorWidgets/Private/Stack/SNiagaraStackItem.h b/Engine/Plugins/FX/Niagara/Source/NiagaraEditorWidgets/Private/Stack/SNiagaraStackItem.h
new file mode 100644
--- /dev/null
+++ b/Engine/Plugins/FX/Niagara/Source/NiagaraEditorWidgets/Private/Stack/SNiagaraStackItem.h
@@ -0,0 +1,41 @@
+// Copyright 1998-2019 Epic Games, Inc. All Rights Reserved.
+
+#pragma once
+
+#include "SNiagaraStackEntryWidget.h"
+#include "Styling/SlateTypes.h"
+#include "Layout/Visibility.h"
+
+class UNiagaraStackItem;
+class UNiagaraStackViewModel; 
+
+class SNiagaraStackItem: public SNiagaraStackEntryWidget
+{
+public:
+	SLATE_BEGIN_ARGS(SNiagaraStackItem) { }
+	SLATE_END_ARGS();
+
+	void Construct(const FArguments& InArgs, UNiagaraStackItem& InItem, UNiagaraStackViewModel* InStackViewModel);
+	
+private:
+	EVisibility GetResetToBaseButtonVisibility() const;
+
+	FText GetResetToBaseButtonToolTipText() const;
+
+	FReply ResetToBaseButtonClicked();
+
+	EVisibility GetDeleteButtonVisibility() const;
+
+	FText GetDeleteButtonToolTipText() const;
+
+	bool GetDeleteButtonEnabled() const;
+
+	FReply DeleteClicked();
+
+	void OnCheckStateChanged(ECheckBoxState InCheckState);
+
+	ECheckBoxState CheckEnabledStatus() const;
+
+private:
+	UNiagaraStackItem* Item;
+};
\ No newline at end of file
diff --git a/Engine/Plugins/FX/Niagara/Source/NiagaraEditorWidgets/Private/Stack/SNiagaraStackRendererItem.cpp b/Engine/Plugins/FX/Niagara/Source/NiagaraEditorWidgets/Private/Stack/SNiagaraStackRendererItem.cpp
deleted file mode 100644
diff --git a/Engine/Plugins/FX/Niagara/Source/NiagaraEditorWidgets/Private/Stack/SNiagaraStackRendererItem.h b/Engine/Plugins/FX/Niagara/Source/NiagaraEditorWidgets/Private/Stack/SNiagaraStackRendererItem.h
deleted file mode 100644
