Niagara - Fix monolithic build.

diff --git a/Engine/Plugins/FX/Niagara/Source/Niagara/Private/NiagaraEmitter.cpp b/Engine/Plugins/FX/Niagara/Source/Niagara/Private/NiagaraEmitter.cpp
--- a/Engine/Plugins/FX/Niagara/Source/Niagara/Private/NiagaraEmitter.cpp
+++ b/Engine/Plugins/FX/Niagara/Source/Niagara/Private/NiagaraEmitter.cpp
@@ -1212,8 +1212,9 @@ void UNiagaraEmitter::MoveShaderStageToIndex(UNiagaraShaderStageBase* ShaderStag
 
 		ShaderStages.Remove(ShaderStageToMove);
 		ShaderStages.Insert(ShaderStageToMove, AdjustedTargetIndex);
-
+#if WITH_EDITOR
 		UpdateChangeId("Shader stage moved.");
+#endif
 	}
 }
 
