Niagara - Fix CIS warnings due to UProperty deprecation.

diff --git a/Engine/Plugins/FX/Niagara/Source/NiagaraEditor/Private/NiagaraScriptMergeManager.cpp b/Engine/Plugins/FX/Niagara/Source/NiagaraEditor/Private/NiagaraScriptMergeManager.cpp
--- a/Engine/Plugins/FX/Niagara/Source/NiagaraEditor/Private/NiagaraScriptMergeManager.cpp
+++ b/Engine/Plugins/FX/Niagara/Source/NiagaraEditor/Private/NiagaraScriptMergeManager.cpp
@@ -1334,7 +1334,7 @@ bool FNiagaraScriptMergeManager::IsShaderStagePropertySetDifferentFromBase(UNiag
 		return true;
 	}
 
-	TArray<UProperty*> DifferentProperties;
+	TArray<FProperty*> DifferentProperties;
 	DiffEditableProperties(BaseShaderStageAdapter->GetShaderStage(), ShaderStageAdapter->GetShaderStage(), *BaseShaderStageAdapter->GetShaderStage()->GetClass(), DifferentProperties);
 	return DifferentProperties.Num() > 0;
 }
@@ -1353,7 +1353,7 @@ void FNiagaraScriptMergeManager::ResetShaderStagePropertySetToBase(UNiagaraEmitt
 		return;
 	}
 
-	TArray<UProperty*> DifferentProperties;
+	TArray<FProperty*> DifferentProperties;
 	DiffEditableProperties(BaseShaderStageAdapter->GetShaderStage(), ShaderStageAdapter->GetShaderStage(), *BaseShaderStageAdapter->GetShaderStage()->GetClass(), DifferentProperties);
 	CopyPropertiesToBase(ShaderStageAdapter->GetEditableShaderStage(), BaseShaderStageAdapter->GetShaderStage(), DifferentProperties);
 	Emitter.PostEditChange();
@@ -1614,7 +1614,7 @@ void FNiagaraScriptMergeManager::DiffShaderStages(const TArray<TSharedRef<FNiaga
 		}
 		else
 		{
-			TArray<UProperty*> DifferentProperties;
+			TArray<FProperty*> DifferentProperties;
 			DiffEditableProperties(CommonValuePair.BaseValue->GetShaderStage(), CommonValuePair.OtherValue->GetShaderStage(), *CommonValuePair.BaseValue->GetShaderStage()->GetClass(), DifferentProperties);
 
 			FNiagaraScriptStackDiffResults ShaderStageScriptStackDiffResults;
diff --git a/Engine/Plugins/FX/Niagara/Source/NiagaraEditor/Private/NiagaraScriptMergeManager.h b/Engine/Plugins/FX/Niagara/Source/NiagaraEditor/Private/NiagaraScriptMergeManager.h
--- a/Engine/Plugins/FX/Niagara/Source/NiagaraEditor/Private/NiagaraScriptMergeManager.h
+++ b/Engine/Plugins/FX/Niagara/Source/NiagaraEditor/Private/NiagaraScriptMergeManager.h
@@ -279,7 +279,7 @@ struct FNiagaraModifiedShaderStageDiffResults
 {
 	TSharedPtr<FNiagaraShaderStageMergeAdapter> BaseAdapter;
 	TSharedPtr<FNiagaraShaderStageMergeAdapter> OtherAdapter;
-	TArray<UProperty*> ChangedProperties;
+	TArray<FProperty*> ChangedProperties;
 	FNiagaraScriptStackDiffResults ScriptDiffResults;
 };
 
