- Pull basis vectors from matrix rather than rotating per vector

#rnx

diff --git a/Engine/Plugins/FX/Niagara/Source/Niagara/Private/NiagaraDataInterfaceCamera.cpp b/Engine/Plugins/FX/Niagara/Source/Niagara/Private/NiagaraDataInterfaceCamera.cpp
--- a/Engine/Plugins/FX/Niagara/Source/Niagara/Private/NiagaraDataInterfaceCamera.cpp
+++ b/Engine/Plugins/FX/Niagara/Source/Niagara/Private/NiagaraDataInterfaceCamera.cpp
@@ -312,9 +312,10 @@ void UNiagaraDataInterfaceCamera::GetCameraProperties(FVectorVMContext& Context)
 	float YPos = CamData->CameraLocation.Y;
 	float ZPos = CamData->CameraLocation.Z;
 
-	FVector Forward = CamData->CameraRotation.RotateVector(FVector::ForwardVector);
-	FVector Up = CamData->CameraRotation.RotateVector(FVector::UpVector);
-	FVector Right = CamData->CameraRotation.RotateVector(FVector::RightVector);
+	FRotationMatrix RotationMatrix(CamData->CameraRotation);
+	const FVector Forward = RotationMatrix.GetScaledAxis(EAxis::X);
+	const FVector Up = RotationMatrix.GetScaledAxis(EAxis::Z);
+	const FVector Right = RotationMatrix.GetScaledAxis(EAxis::Y);
 
 	for (int32 i = 0; i < Context.NumInstances; ++i)
 	{
