Fixed: UpdateID() was called for particles which weren't alive, which could overwrite slots belonging to other particles in the ID to index mapping table.
Added a separate draw event and GPU stat for the simulation compute shader, to get nicer grouping in RenderDoc.

diff --git a/Engine/Plugins/FX/Niagara/Source/Niagara/Private/NiagaraEmitterInstanceBatcher.cpp b/Engine/Plugins/FX/Niagara/Source/Niagara/Private/NiagaraEmitterInstanceBatcher.cpp
--- a/Engine/Plugins/FX/Niagara/Source/Niagara/Private/NiagaraEmitterInstanceBatcher.cpp
+++ b/Engine/Plugins/FX/Niagara/Source/Niagara/Private/NiagaraEmitterInstanceBatcher.cpp
@@ -27,6 +27,7 @@ DECLARE_DWORD_COUNTER_STAT(TEXT("Readback latency (frames)"), STAT_NiagaraReadba
 DECLARE_GPU_STAT_NAMED(NiagaraGPU, TEXT("Niagara"));
 DECLARE_GPU_STAT_NAMED(NiagaraGPUSimulation, TEXT("Niagara GPU Simulation"));
 DECLARE_GPU_STAT_NAMED(NiagaraGPUClearIDTable, TEXT("NiagaraGPU Clear ID Table"));
+DECLARE_GPU_STAT_NAMED(NiagaraGPURunEmitter, TEXT("NiagaraGPU Run Emitter"));
 DECLARE_GPU_STAT_NAMED(NiagaraGPUComputeFreeIDs, TEXT("Niagara GPU Compute Free IDs"));
 DECLARE_GPU_STAT_NAMED(NiagaraGPUSorting, TEXT("Niagara GPU sorting"));
 
@@ -1321,6 +1322,8 @@ void NiagaraEmitterInstanceBatcher::Run(const FNiagaraGPUSystemTick& Tick, const
 	// Dispatch, if anything needs to be done
 	if (TotalNumInstances)
 	{
+		SCOPED_DRAW_EVENT(RHICmdList, NiagaraGPURunEmitter);
+		SCOPED_GPU_STAT(RHICmdList, NiagaraGPURunEmitter);
 		DispatchComputeShader(RHICmdList, Shader, NumThreadGroups, 1, 1);
 	}
 
diff --git a/Engine/Plugins/FX/Niagara/Source/NiagaraCore/Private/NiagaraCustomVersion.cpp b/Engine/Plugins/FX/Niagara/Source/NiagaraCore/Private/NiagaraCustomVersion.cpp
--- a/Engine/Plugins/FX/Niagara/Source/NiagaraCore/Private/NiagaraCustomVersion.cpp
+++ b/Engine/Plugins/FX/Niagara/Source/NiagaraCore/Private/NiagaraCustomVersion.cpp
@@ -8,4 +8,4 @@ const FGuid FNiagaraCustomVersion::GUID(0xFCF57AFA, 0x50764283, 0xB9A9E658, 0xFF
 // Register the custom version with core
 FCustomVersionRegistration GRegisterNiagaraCustomVersion(FNiagaraCustomVersion::GUID, FNiagaraCustomVersion::LatestVersion, TEXT("NiagaraVer"));
 
-const FGuid FNiagaraCustomVersion::LatestScriptCompileVersion(0x5A6DBD4C, 0x817E476A, 0x91039D69, 0xC9918D2E);
+const FGuid FNiagaraCustomVersion::LatestScriptCompileVersion(0x3B2E872A, 0x9FDE4E10, 0xBB0194B8, 0x98DF161C);
diff --git a/Engine/Plugins/FX/Niagara/Source/NiagaraEditor/Private/NiagaraHlslTranslator.cpp b/Engine/Plugins/FX/Niagara/Source/NiagaraEditor/Private/NiagaraHlslTranslator.cpp
--- a/Engine/Plugins/FX/Niagara/Source/NiagaraEditor/Private/NiagaraHlslTranslator.cpp
+++ b/Engine/Plugins/FX/Niagara/Source/NiagaraEditor/Private/NiagaraHlslTranslator.cpp
@@ -2003,13 +2003,13 @@ void FHlslNiagaraTranslator::DefineMainGPUFunctions(
 			ContextName = FString::Printf(TEXT("Context.%s."), *TranslationStages.Last().PassNamespace);
 		}
 
+		HlslOutput += TEXT("\tif (bValid)\n\t{\n");
+
 		if (bNeedsPersistentIDs)
 		{
-			HlslOutput += FString::Printf(TEXT("\tUpdateID(0, %sParticles.ID.Index, WriteIndex);\n"), *ContextName);
+			HlslOutput += FString::Printf(TEXT("\t\tUpdateID(0, %sParticles.ID.Index, WriteIndex);\n"), *ContextName);
 		}
 
-		HlslOutput += TEXT("\tif (bValid)\n\t{\n");
-
 		for (int32 DataSetIndex = 0, IntCounter = 0, FloatCounter = 0; DataSetIndex < DataSetWrites.Num(); ++DataSetIndex)
 		{
 			const FNiagaraDataSetID DataSetID = ReadDataSetIDs[DataSetIndex];
