- Add flag to enable / disable motion blur per Niagara Renderer
- Fix local space motion blur for mesh / sprite particles
- Disable velocities on ribbons, they were not accurate and making accurate is a difficult problem due to tessellation

diff --git a/Engine/Plugins/FX/Niagara/Shaders/Private/NiagaraMeshVertexFactory.ush b/Engine/Plugins/FX/Niagara/Shaders/Private/NiagaraMeshVertexFactory.ush
--- a/Engine/Plugins/FX/Niagara/Shaders/Private/NiagaraMeshVertexFactory.ush
+++ b/Engine/Plugins/FX/Niagara/Shaders/Private/NiagaraMeshVertexFactory.ush
@@ -192,6 +192,7 @@ struct FVertexFactoryIntermediates
 #endif
 
 	float4x4 LocalToWorld;
+	float4x4 PrevLocalToWorld;
 };
 
 float4x4 GetParticleTransform(FVertexFactoryIntermediates Inters)
@@ -199,9 +200,9 @@ float4x4 GetParticleTransform(FVertexFactoryIntermediates Inters)
 	return Inters.LocalToWorld;
 }
 
-#if USE_PREVIOUS_TRANSFORM
-float4x4 GetParticlePrevTransform(FVertexFactoryInput Input)
+float4x4 GetParticlePrevTransform(FVertexFactoryIntermediates Inters)
 {
+#if USE_PREVIOUS_TRANSFORM
 	uint TransformBaseIdx = Input.InstanceId * 3;
 
 	float4x4 PrevTransform4x4 = float4x4(
@@ -213,8 +214,10 @@ float4x4 GetParticlePrevTransform(FVertexFactoryInput Input)
 
 	//Transform is transposed on CPU and position is packed into W. Transpose again to get the correct transform back.
 	return transpose(PrevTransform4x4);
-}
+#else
+	return Inters.PrevLocalToWorld;
 #endif
+}
 
 /** Converts from vertex factory specific interpolants to a FMaterialPixelParameters, which is used by material inputs. */
 FMaterialPixelParameters GetMaterialPixelParameters(FVertexFactoryInterpolantsVSToPS Interpolants, float4 SvPosition)
@@ -379,11 +382,7 @@ FMaterialVertexParameters GetMaterialVertexParameters(FVertexFactoryInput Input,
 
 	Result.TangentToWorld = mul(TangentToLocal, (float3x3)Transform );
 
-#if USE_PREVIOUS_TRANSFORM
-	Result.PrevFrameLocalToWorld = GetParticlePrevTransform(Input);
-#else
-	Result.PrevFrameLocalToWorld = GetParticleTransform(Intermediates);
-#endif
+	Result.PrevFrameLocalToWorld = GetParticlePrevTransform(Intermediates);
 
 #if NUM_MATERIAL_TEXCOORDS_VERTEX
 	for(int CoordinateIndex = 0; CoordinateIndex < NUM_MATERIAL_TEXCOORDS_VERTEX; CoordinateIndex++)
@@ -404,28 +403,9 @@ float4 CalcWorldPosition(FVertexFactoryInput Input, FVertexFactoryIntermediates
 
 float4 CalcPrevWorldPosition(FVertexFactoryInput Input, FVertexFactoryIntermediates Intermediates)
 {
-#if USE_PREVIOUS_TRANSFORM
-	if (MeshParticleVF.PrevTransformAvailable > 0)
-	{
-		float4x4 PrevTransform = GetParticlePrevTransform(Input);
-		float3 PrevWorldPosition = mul(Input.Position, PrevTransform).xyz + ResolvedView.PrevPreViewTranslation;
-		return float4(PrevWorldPosition, Input.Position.w);
-	}
-	else
-#endif
-	// if we don't have a full previous transform, we can at least calculate the position to get some motion blur
-	{
-		float3 WorldPosition = CalcWorldPosition(Input, Intermediates).xyz;
-
-		// ParticleVelocity.xyz - normalized direction vector
-		// ParticleVelocity.w - velocity magnitude
-		float3 WorldDisplacement = -(Intermediates.ParticleVelocity.xyz * (Intermediates.ParticleVelocity.w * NiagaraMeshVF.DeltaSeconds));
-
-		float3 PrevWorldPosition = WorldPosition - ResolvedView.PreViewTranslation;
-		PrevWorldPosition += WorldDisplacement + ResolvedView.PrevPreViewTranslation;
-
-		return float4(PrevWorldPosition.xyz, Input.Position.w);
-	}
+    float4x4 Transform = GetParticlePrevTransform(Intermediates);
+    float3 WorldPosition = mul(Input.Position, Transform).xyz + ResolvedView.PrevPreViewTranslation;
+    return float4(WorldPosition, Input.Position.w);
 }
 
 // From UnrealMath::FindBetween
@@ -682,15 +662,27 @@ FVertexFactoryIntermediates GetVertexFactoryIntermediates(FVertexFactoryInput In
 		Intermediates.LocalToWorld =  mul(LocalToWorldMat, ScaleMatrix);
 	}
 
+	Intermediates.LocalToWorld = transpose(Intermediates.LocalToWorld);
+
+    //-TODO: If / when we store previous transform data use them here (providing they are in local space, if not handle below)
+	//if (MeshParticleVF.PrevTransformAvailable > 0)
+	//{
+    //   //-TODO:
+	//}
+    //else
+    {
+        Intermediates.PrevLocalToWorld = Intermediates.LocalToWorld;
+        Intermediates.PrevLocalToWorld[3].xyz -= ParticleVelocity * NiagaraMeshVF.DeltaSeconds;
+    }
+
 	//Handle local space. TODO: Optimize behind define.
 	//Apply local space to local to world, also pretranspse here to avoid it in the vs&ps
-	Intermediates.LocalToWorld = transpose(Intermediates.LocalToWorld);
 	if (NiagaraMeshVF.bLocalSpace)
 	{
+		Intermediates.PrevLocalToWorld = mul(Intermediates.PrevLocalToWorld, Primitive.PreviousLocalToWorld);
 		Intermediates.LocalToWorld = mul(Intermediates.LocalToWorld, Primitive.LocalToWorld);
 		Intermediates.ParticlePosition = TransformLocalToWorld(Intermediates.ParticlePosition);
 	}
-	//Intermediates.ParticleVelocity.xyz = TransformVector(Intermediates.ParticleVelocity.xyz);
 
 #if NUM_MATERIAL_TEXCOORDS_VERTEX
 	for(int CoordinateIndex = 0; CoordinateIndex < NUM_MATERIAL_TEXCOORDS_VERTEX; CoordinateIndex++)
diff --git a/Engine/Plugins/FX/Niagara/Shaders/Private/NiagaraSpriteVertexFactory.ush b/Engine/Plugins/FX/Niagara/Shaders/Private/NiagaraSpriteVertexFactory.ush
--- a/Engine/Plugins/FX/Niagara/Shaders/Private/NiagaraSpriteVertexFactory.ush
+++ b/Engine/Plugins/FX/Niagara/Shaders/Private/NiagaraSpriteVertexFactory.ush
@@ -609,6 +609,18 @@ float3 TransformPosition(float3 InPosition)
 	}
 }
 
+float3 TransformPrevPosition(float3 InPosition)
+{
+	if (NiagaraSpriteVF.bLocalSpace)
+	{
+		return mul(float4(InPosition,1), Primitive.PreviousLocalToWorld).xyz;
+	}
+	else
+	{
+		return InPosition;
+	}
+}
+
 float3 TransformVector(float3 InVector)
 {
 	if (NiagaraSpriteVF.bLocalSpace)
@@ -634,6 +646,15 @@ float3 GetNiagaraParticlePosition(uint InstanceID)
 	}
 	return TransformPosition(GetVec3(NiagaraSpriteVF.PositionDataOffset, InstanceID));
 }
+float3 GetNiagaraParticlePrevPosition(uint InstanceID)
+{
+	if(NiagaraSpriteVF.PositionDataOffset  == -1)
+	{
+		// See FNiagaraConstants for a synchronized value...
+		return TransformPrevPosition(NiagaraSpriteVF.DefaultPos.xyz);
+	}
+	return TransformPrevPosition(GetVec3(NiagaraSpriteVF.PositionDataOffset, InstanceID));
+}
 float3 GetNiagaraParticleVelocity(uint InstanceID)
 {
 	if(NiagaraSpriteVF.VelocityDataOffset  == -1)
@@ -813,7 +834,7 @@ FVertexFactoryIntermediates GetVertexFactoryIntermediates(FVertexFactoryInput In
 	}
 #endif
 	Intermediates.Position = GetNiagaraParticlePosition(ParticleID);
-	Intermediates.OldPosition = Intermediates.Position.xyz - GetNiagaraParticleVelocity(ParticleID) * NiagaraSpriteVF.DeltaSeconds;
+	Intermediates.OldPosition = GetNiagaraParticlePrevPosition(ParticleID) - GetNiagaraParticleVelocity(ParticleID) * NiagaraSpriteVF.DeltaSeconds;
 	float2 ParticleSize = GetNiagaraParticleSize(ParticleID);
 	Intermediates.SizeRotSubImage = float4(ParticleSize.x, ParticleSize.y, GetNiagaraParticleRotation(ParticleID), GetNiagaraParticleSubimage(ParticleID));
 	Intermediates.Color	= GetNiagaraParticleColor(ParticleID);
diff --git a/Engine/Plugins/FX/Niagara/Source/Niagara/Private/NiagaraComponent.cpp b/Engine/Plugins/FX/Niagara/Source/Niagara/Private/NiagaraComponent.cpp
--- a/Engine/Plugins/FX/Niagara/Source/Niagara/Private/NiagaraComponent.cpp
+++ b/Engine/Plugins/FX/Niagara/Source/Niagara/Private/NiagaraComponent.cpp
@@ -21,6 +21,7 @@
 #include "EngineUtils.h"
 #include "ProfilingDebugging/CsvProfiler.h"
 #include "Engine/CollisionProfile.h"
+#include "PrimitiveSceneInfo.h"
 #include "NiagaraEmitterInstanceBatcher.h"
 
 DECLARE_CYCLE_STAT(TEXT("Sceneproxy create (GT)"), STAT_NiagaraCreateSceneProxy, STATGROUP_Niagara);
@@ -124,9 +125,7 @@ FNiagaraSceneProxy::FNiagaraSceneProxy(const UNiagaraComponent* InComponent)
 	if (SystemInst)
 	{
 		//UE_LOG(LogNiagara, Warning, TEXT("FNiagaraSceneProxy %p"), this);
-
 		CreateRenderers(InComponent);
-		bAlwaysHasVelocity = true;
 		Batcher = SystemInst->GetBatcher();
 
 #if STATS
@@ -183,6 +182,8 @@ void FNiagaraSceneProxy::CreateRenderers(const UNiagaraComponent* Component)
 	};
 	TArray<FSortInfo, TInlineAllocator<8>> RendererSortInfo;
 
+	bAlwaysHasVelocity = false;
+
 	ReleaseRenderers();
 	ERHIFeatureLevel::Type FeatureLevel = GetScene().GetFeatureLevel();
 	for(TSharedRef<const FNiagaraEmitterInstance, ESPMode::ThreadSafe> EmitterInst : Component->GetSystemInstance()->GetEmitters())
@@ -196,6 +197,7 @@ void FNiagaraSceneProxy::CreateRenderers(const UNiagaraComponent* Component)
 				if (Properties->GetIsEnabled() && EmitterInst->GetData().IsInitialized() && !EmitterInst->IsDisabled())
 				{
 					NewRenderer = Properties->CreateEmitterRenderer(FeatureLevel, &EmitterInst.Get());
+					bAlwaysHasVelocity |= Properties->bMotionBlurEnabled;
 				}
 				EmitterRenderers.Add(NewRenderer);
 			}
@@ -230,6 +232,7 @@ FNiagaraSceneProxy::~FNiagaraSceneProxy()
 			delete EmitterRenderer;
 		}
 	}
+	UniformBufferNoVelocity.ReleaseResource();
 	EmitterRenderers.Empty();
 }
 
@@ -242,6 +245,7 @@ void FNiagaraSceneProxy::ReleaseRenderThreadResources()
 			Renderer->ReleaseRenderThreadResources();
 		}
 	}
+	UniformBufferNoVelocity.ReleaseResource();
 }
 
 // FPrimitiveSceneProxy interface.
@@ -261,6 +265,7 @@ void FNiagaraSceneProxy::CreateRenderThreadResources()
 void FNiagaraSceneProxy::OnTransformChanged()
 {
 	LocalToWorldInverse = GetLocalToWorld().Inverse();
+	UniformBufferNoVelocity.ReleaseResource();
 }
 
 FPrimitiveViewRelevance FNiagaraSceneProxy::GetViewRelevance(const FSceneView* View) const
@@ -286,6 +291,43 @@ FPrimitiveViewRelevance FNiagaraSceneProxy::GetViewRelevance(const FSceneView* V
 	return Relevance;
 }
 
+FRHIUniformBuffer* FNiagaraSceneProxy::GetUniformBufferNoVelocity() const
+{
+	if ( !UniformBufferNoVelocity.IsInitialized() )
+	{
+		bool bHasPrecomputedVolumetricLightmap;
+		FMatrix PreviousLocalToWorld;
+		int32 SingleCaptureIndex;
+		bool bOutputVelocity;
+		FPrimitiveSceneInfo* LocalPrimitiveSceneInfo = GetPrimitiveSceneInfo();
+		GetScene().GetPrimitiveUniformShaderParameters_RenderThread(LocalPrimitiveSceneInfo, bHasPrecomputedVolumetricLightmap, PreviousLocalToWorld, SingleCaptureIndex, bOutputVelocity);
+
+		UniformBufferNoVelocity.SetContents(
+			GetPrimitiveUniformShaderParameters(
+				GetLocalToWorld(),
+				PreviousLocalToWorld,
+				GetActorPosition(),
+				GetBounds(),
+				GetLocalBounds(),
+				GetLocalBounds(),
+				ReceivesDecals(),
+				HasDistanceFieldRepresentation(),
+				HasDynamicIndirectShadowCasterRepresentation(),
+				UseSingleSampleShadowFromStationaryLights(),
+				bHasPrecomputedVolumetricLightmap,
+				DrawsVelocity(),
+				GetLightingChannelMask(),
+				LpvBiasMultiplier,
+				LocalPrimitiveSceneInfo ? LocalPrimitiveSceneInfo->GetLightmapDataOffset() : 0,
+				SingleCaptureIndex,
+				false,
+				GetCustomPrimitiveData()
+			)
+		);
+		UniformBufferNoVelocity.InitResource();
+	}
+	return UniformBufferNoVelocity.GetUniformBufferRHI();
+}
 
 uint32 FNiagaraSceneProxy::GetMemoryFootprint() const
 { 
diff --git a/Engine/Plugins/FX/Niagara/Source/Niagara/Private/NiagaraRenderer.cpp b/Engine/Plugins/FX/Niagara/Source/Niagara/Private/NiagaraRenderer.cpp
--- a/Engine/Plugins/FX/Niagara/Source/Niagara/Private/NiagaraRenderer.cpp
+++ b/Engine/Plugins/FX/Niagara/Source/Niagara/Private/NiagaraRenderer.cpp
@@ -246,6 +246,7 @@ FNiagaraRenderer::FNiagaraRenderer(ERHIFeatureLevel::Type InFeatureLevel, const
 	: DynamicDataRender(nullptr)
 	, bLocalSpace(Emitter->GetCachedEmitter()->bLocalSpace)
 	, bHasLights(false)
+	, bMotionBlurEnabled(InProps ? InProps->bMotionBlurEnabled : false)
 	, SimTarget(Emitter->GetCachedEmitter()->SimTarget)
 	, NumIndicesPerInstance(InProps ? InProps->GetNumIndicesPerInstance() : 0)
 	, FeatureLevel(InFeatureLevel)
diff --git a/Engine/Plugins/FX/Niagara/Source/Niagara/Private/NiagaraRendererMeshes.cpp b/Engine/Plugins/FX/Niagara/Source/Niagara/Private/NiagaraRendererMeshes.cpp
--- a/Engine/Plugins/FX/Niagara/Source/Niagara/Private/NiagaraRendererMeshes.cpp
+++ b/Engine/Plugins/FX/Niagara/Source/Niagara/Private/NiagaraRendererMeshes.cpp
@@ -413,7 +413,7 @@ void FNiagaraRendererMeshes::GetDynamicMeshElements(const TArray<const FSceneVie
 					Mesh.DepthPriorityGroup = (ESceneDepthPriorityGroup)SceneProxy->GetDepthPriorityGroup(View);
 
 					FMeshBatchElement& BatchElement = Mesh.Elements[0];
-					BatchElement.PrimitiveUniformBuffer = SceneProxy->GetUniformBuffer();
+					BatchElement.PrimitiveUniformBuffer = IsMotionBlurEnabled() ? SceneProxy->GetUniformBuffer() : SceneProxy->GetUniformBufferNoVelocity();
 					BatchElement.FirstIndex = 0;
 					BatchElement.MinVertexIndex = 0;
 					BatchElement.MaxVertexIndex = 0;
diff --git a/Engine/Plugins/FX/Niagara/Source/Niagara/Private/NiagaraRendererRibbons.cpp b/Engine/Plugins/FX/Niagara/Source/Niagara/Private/NiagaraRendererRibbons.cpp
--- a/Engine/Plugins/FX/Niagara/Source/Niagara/Private/NiagaraRendererRibbons.cpp
+++ b/Engine/Plugins/FX/Niagara/Source/Niagara/Private/NiagaraRendererRibbons.cpp
@@ -841,8 +841,7 @@ void FNiagaraRendererRibbons::SetupMeshBatchAndCollectorResourceForView(
 	MeshElement.NumInstances = 1;
 	MeshElement.MinVertexIndex = 0;
 	MeshElement.MaxVertexIndex = 0;
-	MeshElement.PrimitiveUniformBuffer = SceneProxy->GetUniformBuffer();
-
+	MeshElement.PrimitiveUniformBuffer = SceneProxy->GetUniformBufferNoVelocity();	// Note: Ribbons don't generate accurate velocities so disabling
 }
 
 FNiagaraRendererRibbons::FCPUSimParticleDataAllocation FNiagaraRendererRibbons::AllocateParticleDataIfCPUSim(FNiagaraDynamicDataRibbon* DynamicDataRibbon, FGlobalDynamicReadBuffer& DynamicReadBuffer) const
diff --git a/Engine/Plugins/FX/Niagara/Source/Niagara/Private/NiagaraRendererSprites.cpp b/Engine/Plugins/FX/Niagara/Source/Niagara/Private/NiagaraRendererSprites.cpp
--- a/Engine/Plugins/FX/Niagara/Source/Niagara/Private/NiagaraRendererSprites.cpp
+++ b/Engine/Plugins/FX/Niagara/Source/Niagara/Private/NiagaraRendererSprites.cpp
@@ -522,7 +522,7 @@ void FNiagaraRendererSprites::CreateMeshBatchForView(
 	MeshElement.NumInstances = FMath::Max(0, NumInstances);	//->VertexData.Num();
 	MeshElement.MinVertexIndex = 0;
 	MeshElement.MaxVertexIndex = 0;// MeshElement.NumInstances * 4 - 1;
-	MeshElement.PrimitiveUniformBuffer = SceneProxy->GetUniformBuffer();
+	MeshElement.PrimitiveUniformBuffer = IsMotionBlurEnabled() ? SceneProxy->GetUniformBuffer() : SceneProxy->GetUniformBufferNoVelocity();
 	if (IndirectArgsOffset != INDEX_NONE)
 	{
 		NiagaraEmitterInstanceBatcher* Batcher = SceneProxy->GetBatcher();
diff --git a/Engine/Plugins/FX/Niagara/Source/Niagara/Public/NiagaraComponent.h b/Engine/Plugins/FX/Niagara/Source/Niagara/Public/NiagaraComponent.h
--- a/Engine/Plugins/FX/Niagara/Source/Niagara/Public/NiagaraComponent.h
+++ b/Engine/Plugins/FX/Niagara/Source/Niagara/Public/NiagaraComponent.h
@@ -603,6 +603,8 @@ public:
 
 	FORCEINLINE const FMatrix& GetLocalToWorldInverse() const { return LocalToWorldInverse; }
 
+	FRHIUniformBuffer* GetUniformBufferNoVelocity() const;
+
 private:
 	void ReleaseRenderThreadResources();
 
@@ -626,12 +628,14 @@ private:
 	/** Callback from the renderer to gather simple lights that this proxy wants renderered. */
 	virtual void GatherSimpleLights(const FSceneViewFamily& ViewFamily, FSimpleLightArray& OutParticleLights) const override;
 
-
 	virtual uint32 GetMemoryFootprint() const override;
 
 	uint32 GetAllocatedSize() const;
 
 private:
+	/** Uniform Buffer with Velocity writes disabled.  Mutable as it is updated during GetDynamicMeshElements is required. */
+	mutable TUniformBuffer<FPrimitiveUniformShaderParameters> UniformBufferNoVelocity;
+
 	/** Emitter Renderers in the order they appear in the emitters. */
 	TArray<FNiagaraRenderer*> EmitterRenderers;
 	
diff --git a/Engine/Plugins/FX/Niagara/Source/Niagara/Public/NiagaraRenderer.h b/Engine/Plugins/FX/Niagara/Source/Niagara/Public/NiagaraRenderer.h
--- a/Engine/Plugins/FX/Niagara/Source/Niagara/Public/NiagaraRenderer.h
+++ b/Engine/Plugins/FX/Niagara/Source/Niagara/Public/NiagaraRenderer.h
@@ -106,9 +106,10 @@ public:
 	void SortIndices(const struct FNiagaraGPUSortInfo& SortInfo, int32 SortVarIdx, const FNiagaraDataBuffer& Buffer, FGlobalDynamicReadBuffer::FAllocation& OutIndices)const;
 
 	void SetDynamicData_RenderThread(FNiagaraDynamicDataBase* NewDynamicData);
-	FORCEINLINE FNiagaraDynamicDataBase *GetDynamicData()const { return DynamicDataRender; }
-	FORCEINLINE bool HasDynamicData()const { return DynamicDataRender != nullptr; }
-	FORCEINLINE bool HasLights()const { return bHasLights; }
+	FORCEINLINE FNiagaraDynamicDataBase *GetDynamicData() const { return DynamicDataRender; }
+	FORCEINLINE bool HasDynamicData() const { return DynamicDataRender != nullptr; }
+	FORCEINLINE bool HasLights() const { return bHasLights; }
+	FORCEINLINE bool IsMotionBlurEnabled() const { return bMotionBlurEnabled; }
 
 #if RHI_RAYTRACING
 	virtual void GetDynamicRayTracingInstances(FRayTracingMaterialGatheringContext& Context, TArray<FRayTracingInstance>& OutRayTracingInstances, const FNiagaraSceneProxy* Proxy) {}
@@ -132,6 +133,7 @@ protected:
 
 	uint32 bLocalSpace : 1;
 	uint32 bHasLights : 1;
+	uint32 bMotionBlurEnabled : 1;
 	const ENiagaraSimTarget SimTarget;
 	uint32 NumIndicesPerInstance;
 
diff --git a/Engine/Plugins/FX/Niagara/Source/Niagara/Public/NiagaraRendererProperties.h b/Engine/Plugins/FX/Niagara/Source/Niagara/Public/NiagaraRendererProperties.h
--- a/Engine/Plugins/FX/Niagara/Source/Niagara/Public/NiagaraRendererProperties.h
+++ b/Engine/Plugins/FX/Niagara/Source/Niagara/Public/NiagaraRendererProperties.h
@@ -31,7 +31,8 @@ class NIAGARA_API UNiagaraRendererProperties : public UNiagaraMergeable
 
 public:
 	UNiagaraRendererProperties()
-	: bIsEnabled(true)
+		: bIsEnabled(true)
+		, bMotionBlurEnabled(true)
 	{
 	}
 
@@ -65,6 +66,8 @@ public:
 	
 	UPROPERTY()
 	bool bIsEnabled;
-};
-
 
+	/** Is motion blur enabled on this renderer or not, the material must also have motion blur enabled. */
+	UPROPERTY(EditAnywhere, Category = "Rendering")
+	bool bMotionBlurEnabled;
+};
