- Fix issue around destroying a Niagara Component on the same frame as a tick group promotion has been requested
- The sequence to reproduce this issue is for the instance to change tick group during system sim ticking to one before the current tick group and the instance then being removed before post actor tick (i.e. DestroyComponent called on the owner)

diff --git a/Engine/Plugins/FX/Niagara/Source/Niagara/Private/NiagaraSystemSimulation.cpp b/Engine/Plugins/FX/Niagara/Source/Niagara/Private/NiagaraSystemSimulation.cpp
--- a/Engine/Plugins/FX/Niagara/Source/Niagara/Private/NiagaraSystemSimulation.cpp
+++ b/Engine/Plugins/FX/Niagara/Source/Niagara/Private/NiagaraSystemSimulation.cpp
@@ -871,8 +871,10 @@ void FNiagaraSystemSimulation::UpdateTickGroups_GameThread()
 
 	// Transfer promoted instances to the new tick group
 	//-OPT: This can be done async
-	for (FNiagaraSystemInstance* Instance : PendingTickGroupPromotions)
+	while (PendingTickGroupPromotions.Num() > 0)
 	{
+		FNiagaraSystemInstance* Instance = PendingTickGroupPromotions.Pop(false);
+
 		const ETickingGroup TickGroup = Instance->CalculateTickGroup();
 		if (TickGroup != SystemTickGroup)
 		{
@@ -1642,6 +1644,9 @@ void FNiagaraSystemSimulation::RemoveInstance(FNiagaraSystemInstance* Instance)
 		--EffectType->NumInstances;
 	}
 
+	// Remove from pending promotions list
+	PendingTickGroupPromotions.RemoveSingleSwap(Instance);
+
 	UNiagaraSystem* System = WeakSystem.Get();
 	if (Instance->IsPendingSpawn())
 	{
