Fixed a "No mesh assigned error" being erroneously reported when using the Sample Skeletal Mesh Skeleton module. The mesh was being flagged as missing when the DI did not require CPU access as well as when it was not set.

diff --git a/Engine/Plugins/FX/Niagara/Source/Niagara/Private/NiagaraDataInterfaceSkeletalMesh.cpp b/Engine/Plugins/FX/Niagara/Source/Niagara/Private/NiagaraDataInterfaceSkeletalMesh.cpp
--- a/Engine/Plugins/FX/Niagara/Source/Niagara/Private/NiagaraDataInterfaceSkeletalMesh.cpp
+++ b/Engine/Plugins/FX/Niagara/Source/Niagara/Private/NiagaraDataInterfaceSkeletalMesh.cpp
@@ -1745,17 +1745,23 @@ bool UNiagaraDataInterfaceSkeletalMesh::PerInstanceTick(void* PerInstanceData, F
 TArray<FNiagaraDataInterfaceError> UNiagaraDataInterfaceSkeletalMesh::GetErrors()
 {
 	TArray<FNiagaraDataInterfaceError> Errors;
-	bool bHasCPUAccessError= false;
+	bool bHasCPUAccessError = false;
 	bool bHasNoMeshAssignedError = false;
 	
 	// Collect Errors
 #if WITH_EDITORONLY_DATA
-	if (DefaultMesh != nullptr && bRequiresCPUAccess)
+	if (DefaultMesh != nullptr)
 	{
-		for (auto info : DefaultMesh->GetLODInfoArray())
+		if (bRequiresCPUAccess)
 		{
-			if (!info.bAllowCPUAccess)
-				bHasCPUAccessError = true;
+			for (const FSkeletalMeshLODInfo& LODInfo : DefaultMesh->GetLODInfoArray())
+			{
+				if (!LODInfo.bAllowCPUAccess)
+				{
+					bHasCPUAccessError = true;
+					break;
+				}
+			}
 		}
 	}
 	else
@@ -1771,11 +1777,9 @@ TArray<FNiagaraDataInterfaceError> UNiagaraDataInterfaceSkeletalMesh::GetErrors(
 			FNiagaraDataInterfaceFix::CreateLambda([=]()
 		{
 			DefaultMesh->Modify();
-			for (int i = 0; i < DefaultMesh->GetLODInfoArray().Num(); i++)
+			for (FSkeletalMeshLODInfo& LODInfo : DefaultMesh->GetLODInfoArray())
 			{
-				FSkeletalMeshLODInfo* info = &DefaultMesh->GetLODInfoArray()[i];
-				DefaultMesh->Modify();
-				info->bAllowCPUAccess = true;
+				LODInfo.bAllowCPUAccess = true;
 			}
 			return true;
 		}));
