- Update GPU side RandomTri to use NiagaraRandomInfo structure

#rnx

diff --git a/Engine/Plugins/FX/Niagara/Shaders/Private/NiagaraDataInterfaceSkeletalMesh.ush b/Engine/Plugins/FX/Niagara/Shaders/Private/NiagaraDataInterfaceSkeletalMesh.ush
--- a/Engine/Plugins/FX/Niagara/Shaders/Private/NiagaraDataInterfaceSkeletalMesh.ush
+++ b/Engine/Plugins/FX/Niagara/Shaders/Private/NiagaraDataInterfaceSkeletalMesh.ush
@@ -179,9 +179,9 @@ FDISkelMeshSkinnedVertex DISkelMesh_GetSkinnedPointOnTriangleWS(in FDISkelMeshCo
 	return FinalVertex;
 }
 
-void DISKelMesh_RandomTriCoord(in FDISkelMeshContext DIContext, out uint OutTriangle, out float3 OutBaryCoord)
+void DISKelMesh_RandomTriCoord(in FDISkelMeshContext DIContext, NiagaraRandInfo InRandomInfo, out uint OutTriangle, out float3 OutBaryCoord)
 {
-	float RandT0 = NiagaraInternalNoise(1, 2, 3);
+	float RandT0 = NiagaraInternalNoise(InRandomInfo.Seed1, InRandomInfo.Seed2, InRandomInfo.Seed3);
 
 	[branch]
 	if ( !DIContext.UniformTriangleSamplingEnable )
@@ -196,7 +196,7 @@ void DISKelMesh_RandomTriCoord(in FDISkelMeshContext DIContext, out uint OutTria
 		float TriangleProbability = DIContext.MeshTriangleSamplerProbaBuffer[TriangleIndex];
 
 		// Alias check
-		float RandT1 = NiagaraInternalNoise(1, 2, 3);
+		float RandT1 = NiagaraInternalNoise(InRandomInfo.Seed1, InRandomInfo.Seed2, InRandomInfo.Seed3);
 		if( RandT1 > TriangleProbability )
 		{
 			TriangleIndex = DIContext.MeshTriangleSamplerAliasBuffer[TriangleIndex];
@@ -204,8 +204,8 @@ void DISKelMesh_RandomTriCoord(in FDISkelMeshContext DIContext, out uint OutTria
 		OutTriangle = TriangleIndex;
 	}
 
-	float r0 = NiagaraInternalNoise(1, 2, 3);
-	float r1 = NiagaraInternalNoise(1, 2, 3);
+	float r0 = NiagaraInternalNoise(InRandomInfo.Seed1, InRandomInfo.Seed2, InRandomInfo.Seed3);
+	float r1 = NiagaraInternalNoise(InRandomInfo.Seed1, InRandomInfo.Seed2, InRandomInfo.Seed3);
 	float sqrt0 = sqrt(r0);
 	float sqrt1 = sqrt(r1);
 	OutBaryCoord = float3(1.0f - sqrt0, sqrt0 * (1.0 - r1), r1 * sqrt0);
diff --git a/Engine/Plugins/FX/Niagara/Source/Niagara/Private/NiagaraDataInterfaceSkeletalMesh.cpp b/Engine/Plugins/FX/Niagara/Source/Niagara/Private/NiagaraDataInterfaceSkeletalMesh.cpp
--- a/Engine/Plugins/FX/Niagara/Source/Niagara/Private/NiagaraDataInterfaceSkeletalMesh.cpp
+++ b/Engine/Plugins/FX/Niagara/Source/Niagara/Private/NiagaraDataInterfaceSkeletalMesh.cpp
@@ -1961,7 +1961,7 @@ bool UNiagaraDataInterfaceSkeletalMesh::GetFunctionHLSL(const FName& DefinitionF
 	// Triangle Sampling
 	if (DefinitionFunctionName == FSkeletalMeshInterfaceHelper::RandomTriCoordName)
 	{
-		static const TCHAR* FormatSample = TEXT("void {InstanceFunctionName} (out {MeshTriCoordinateStructName} OutCoord) { {GetDISkelMeshContextName} DISKelMesh_RandomTriCoord(DIContext, OutCoord.Tri, OutCoord.BaryCoord); }");
+		static const TCHAR* FormatSample = TEXT("void {InstanceFunctionName} (NiagaraRandInfo InRandomInfo, out {MeshTriCoordinateStructName} OutCoord) { {GetDISkelMeshContextName} DISKelMesh_RandomTriCoord(DIContext, InRandomInfo, OutCoord.Tri, OutCoord.BaryCoord); }");
 		OutHLSL += FString::Format(FormatSample, ArgsSample);
 	}
 	else if (DefinitionFunctionName == FSkeletalMeshInterfaceHelper::GetSkinnedTriangleDataWSName)
