Fix for incorrectly culled FX in Persona.

diff --git a/Engine/Plugins/FX/Niagara/Source/Niagara/Private/NiagaraWorldManager.cpp b/Engine/Plugins/FX/Niagara/Source/Niagara/Private/NiagaraWorldManager.cpp
--- a/Engine/Plugins/FX/Niagara/Source/Niagara/Private/NiagaraWorldManager.cpp
+++ b/Engine/Plugins/FX/Niagara/Source/Niagara/Private/NiagaraWorldManager.cpp
@@ -865,7 +865,7 @@ float FNiagaraWorldManager::DistanceSignificance(UNiagaraEffectType* EffectType,
 
 float FNiagaraWorldManager::DistanceSignificance(UNiagaraEffectType* EffectType, const FNiagaraScalabilitySettings& ScalabilitySettings, FVector Location)
 {
-	if (ScalabilitySettings.bCullByDistance)
+	if (ScalabilitySettings.bCullByDistance && bCachedPlayerViewLocationsValid)
 	{
 		float ClosestDistSq = FLT_MAX;
 		for (FVector ViewLocation : CachedPlayerViewLocations)
