Niagara - Fix module usage bitmask issue caused by modirying the usage enum.

diff --git a/Engine/Plugins/FX/Niagara/Source/Niagara/Private/NiagaraScript.cpp b/Engine/Plugins/FX/Niagara/Source/Niagara/Private/NiagaraScript.cpp
--- a/Engine/Plugins/FX/Niagara/Source/Niagara/Private/NiagaraScript.cpp
+++ b/Engine/Plugins/FX/Niagara/Source/Niagara/Private/NiagaraScript.cpp
@@ -746,6 +746,27 @@ void UNiagaraScript::PostLoad()
 	const int32 NiagaraVer = GetLinkerCustomVersion(FNiagaraCustomVersion::GUID);
 
 #if WITH_EDITORONLY_DATA
+	if (NiagaraVer < FNiagaraCustomVersion::AddShaderStageUsageEnum)
+	{
+		uint8 ShaderStageIndex = (uint8)ENiagaraScriptUsage::ParticleShaderStageScript;
+		uint8 MaxIndex = (uint8)ENiagaraScriptUsage::SystemUpdateScript;
+		// Start at the end and shift the bits down to account for the new shader stage bit.
+		for (uint8 CurrentIndex = MaxIndex; CurrentIndex > ShaderStageIndex; CurrentIndex--)
+		{
+			uint8 OldIndex = CurrentIndex - 1;
+			if ((ModuleUsageBitmask & (1 << OldIndex)) != 0)
+			{
+				ModuleUsageBitmask |= 1 << CurrentIndex;
+			}
+			else
+			{
+				ModuleUsageBitmask &= ~(1 << CurrentIndex);
+			}
+		}
+		// Clear the shader stage bit.
+		ModuleUsageBitmask &= ~(1 << ShaderStageIndex);
+	}
+
 	if (Source != nullptr)
 	{
 		Source->ConditionalPostLoad();
diff --git a/Engine/Plugins/FX/Niagara/Source/NiagaraCore/Public/NiagaraCustomVersion.h b/Engine/Plugins/FX/Niagara/Source/NiagaraCore/Public/NiagaraCustomVersion.h
--- a/Engine/Plugins/FX/Niagara/Source/NiagaraCore/Public/NiagaraCustomVersion.h
+++ b/Engine/Plugins/FX/Niagara/Source/NiagaraCore/Public/NiagaraCustomVersion.h
@@ -136,6 +136,8 @@ struct FNiagaraCustomVersion
 
 		MemorySaving, // Convert TMap into TArray to save memory, TMap contains an inline allocator which pushes the size to 80 bytes
 
+		AddShaderStageUsageEnum, // Added a new value to the script usage enum, and we need a custom version to fix the existing bitfields.
+
 		// DO NOT ADD A NEW VERSION UNLESS YOU HAVE TALKED TO THE NIAGARA LEAD. Mismanagement of these versions can lead to data loss if it is adjusted in multiple streams simultaneously.
 		// -----<new versions can be added above this line>  -------------------------------------------------
 		VersionPlusOne,
