Fix: structures were not handled correctly when converting emitter parameters to padded cbuffers for compute shaders

diff --git a/Engine/Plugins/FX/Niagara/Source/Niagara/Private/NiagaraScriptExecutionParameterStore.cpp b/Engine/Plugins/FX/Niagara/Source/Niagara/Private/NiagaraScriptExecutionParameterStore.cpp
--- a/Engine/Plugins/FX/Niagara/Source/Niagara/Private/NiagaraScriptExecutionParameterStore.cpp
+++ b/Engine/Plugins/FX/Niagara/Source/Niagara/Private/NiagaraScriptExecutionParameterStore.cpp
@@ -36,7 +36,7 @@ uint32 OffsetAlign(uint32 SrcOffset, uint32 Size)
 	}
 }
 
-void FNiagaraScriptExecutionParameterStore::GenerateLayoutInfoInternal(TArray<FNiagaraScriptExecutionPaddingInfo>& Members, uint32& NextMemberOffset, const UStruct* InSrcStruct, uint32 InSrcOffset)
+uint32 FNiagaraScriptExecutionParameterStore::GenerateLayoutInfoInternal(TArray<FNiagaraScriptExecutionPaddingInfo>& Members, uint32& NextMemberOffset, const UStruct* InSrcStruct, uint32 InSrcOffset)
 {
 	uint32 VectorPaddedSize = (TShaderParameterTypeInfo<FVector4>::NumRows * TShaderParameterTypeInfo<FVector4>::NumColumns) * sizeof(float);
 
@@ -115,10 +115,11 @@ void FNiagaraScriptExecutionParameterStore::GenerateLayoutInfoInternal(TArray<FN
 				check(false);
 			}
 
-			GenerateLayoutInfoInternal(Members, NextMemberOffset, Struct, InSrcOffset);
+			InSrcOffset = GenerateLayoutInfoInternal(Members, NextMemberOffset, Struct, InSrcOffset);
 		}
 	}
 
+	return InSrcOffset;
 }
 
 void FNiagaraScriptExecutionParameterStore::AddPaddedParamSize(const FNiagaraTypeDefinition& InParamType, uint32 InOffset)
diff --git a/Engine/Plugins/FX/Niagara/Source/Niagara/Public/NiagaraScriptExecutionParameterStore.h b/Engine/Plugins/FX/Niagara/Source/Niagara/Public/NiagaraScriptExecutionParameterStore.h
--- a/Engine/Plugins/FX/Niagara/Source/Niagara/Public/NiagaraScriptExecutionParameterStore.h
+++ b/Engine/Plugins/FX/Niagara/Source/Niagara/Public/NiagaraScriptExecutionParameterStore.h
@@ -107,7 +107,7 @@ private:
 	UPROPERTY()
 	uint32 PaddedParameterSize;
 
-	static void GenerateLayoutInfoInternal(TArray<FNiagaraScriptExecutionPaddingInfo>& Members, uint32& NextMemberOffset, const UStruct* InSrcStruct, uint32 InSrcOffset);
+	static uint32 GenerateLayoutInfoInternal(TArray<FNiagaraScriptExecutionPaddingInfo>& Members, uint32& NextMemberOffset, const UStruct* InSrcStruct, uint32 InSrcOffset);
 
 	UPROPERTY()
 	TArray<FNiagaraScriptExecutionPaddingInfo> PaddingInfo;
