Fixed a bug where modifing a statich switch node could cause the loss of all its parameter metadata

diff --git a/Engine/Plugins/FX/Niagara/Source/NiagaraEditor/Private/NiagaraGraph.cpp b/Engine/Plugins/FX/Niagara/Source/NiagaraEditor/Private/NiagaraGraph.cpp
--- a/Engine/Plugins/FX/Niagara/Source/NiagaraEditor/Private/NiagaraGraph.cpp
+++ b/Engine/Plugins/FX/Niagara/Source/NiagaraEditor/Private/NiagaraGraph.cpp
@@ -830,7 +830,7 @@ void UNiagaraGraph::AddParameter(const FNiagaraVariable& Parameter, bool bIsStat
 	}
 }
 
-void UNiagaraGraph::RemoveParameter(const FNiagaraVariable& Parameter, bool bFromStaticSwitch)
+void UNiagaraGraph::RemoveParameter(const FNiagaraVariable& Parameter)
 {
 	FNiagaraGraphParameterReferenceCollection* ReferenceCollection = ParameterToReferencesMap.Find(Parameter);
 	if (ReferenceCollection)
@@ -841,6 +841,12 @@ void UNiagaraGraph::RemoveParameter(const FNiagaraVariable& Parameter, bool bFro
 			UNiagaraNode* Node = Reference.Value.Get();
 			if (Node && Node->GetGraph() == this)
 			{
+				if (Node->IsA(UNiagaraNodeStaticSwitch::StaticClass()))
+				{
+					// Static switch parameters are automatically populated from the graph nodes and cannot be manually deleted
+					NotifyGraphChanged();
+					return;
+				}
 				UEdGraphPin* Pin = Node->GetPinByPersistentGuid(Reference.Key);
 				if (Pin)
 				{
@@ -858,21 +864,10 @@ void UNiagaraGraph::RemoveParameter(const FNiagaraVariable& Parameter, bool bFro
 	TArray<FNiagaraVariable> VarsToRemove;
 	for (auto It : VariableToScriptVariable)
 	{
-		if (bFromStaticSwitch)
-		{
-			if (It.Key.GetName() == Parameter.GetName() && It.Key.GetType() == Parameter.GetType())
-			{
-				VarsToRemove.Add(It.Key);
-			}
-		}
-		else
+		if (It.Key.GetName() == Parameter.GetName())
 		{
-			if (It.Key.GetName() == Parameter.GetName())
-			{
-				VarsToRemove.Add(It.Key);
-			}
+			VarsToRemove.Add(It.Key);
 		}
-		
 	}
 	for (FNiagaraVariable& Var : VarsToRemove)
 	{
diff --git a/Engine/Plugins/FX/Niagara/Source/NiagaraEditor/Private/NiagaraNodeStaticSwitch.cpp b/Engine/Plugins/FX/Niagara/Source/NiagaraEditor/Private/NiagaraNodeStaticSwitch.cpp
--- a/Engine/Plugins/FX/Niagara/Source/NiagaraEditor/Private/NiagaraNodeStaticSwitch.cpp
+++ b/Engine/Plugins/FX/Niagara/Source/NiagaraEditor/Private/NiagaraNodeStaticSwitch.cpp
@@ -14,7 +14,7 @@ UNiagaraNodeStaticSwitch::UNiagaraNodeStaticSwitch(const FObjectInitializer& Obj
 
 void UNiagaraNodeStaticSwitch::DestroyNode()
 {
-	GetNiagaraGraph()->RemoveParameter(FNiagaraVariable(GetInputType(), InputParameterName), true);
+	GetNiagaraGraph()->RemoveParameter(FNiagaraVariable(GetInputType(), InputParameterName));
 	Super::DestroyNode();
 }
 
@@ -102,7 +102,7 @@ void UNiagaraNodeStaticSwitch::RemoveUnusedGraphParameter(const FNiagaraVariable
 	int Index = GraphVariables.Find(OldParameter);
 	if (Index == INDEX_NONE)
 	{
-		GetNiagaraGraph()->RemoveParameter(OldParameter, true);
+		GetNiagaraGraph()->RemoveParameter(OldParameter);
 	}
 	else
 	{
diff --git a/Engine/Plugins/FX/Niagara/Source/NiagaraEditor/Public/NiagaraGraph.h b/Engine/Plugins/FX/Niagara/Source/NiagaraEditor/Public/NiagaraGraph.h
--- a/Engine/Plugins/FX/Niagara/Source/NiagaraEditor/Public/NiagaraGraph.h
+++ b/Engine/Plugins/FX/Niagara/Source/NiagaraEditor/Public/NiagaraGraph.h
@@ -233,7 +233,7 @@ class UNiagaraGraph : public UEdGraph
 	void AddParameterReference(const FNiagaraVariable& Parameter, const UEdGraphPin* Pin);
 
 	/** Remove parameter from map and all the pins associated. */
-	void RemoveParameter(const FNiagaraVariable& Parameter, bool bFromStaticSwitch = false);
+	void RemoveParameter(const FNiagaraVariable& Parameter);
 
 	/** Rename parameter from map and all the pins associated. */
 	bool RenameParameter(const FNiagaraVariable& Parameter, FName NewName, bool bFromStaticSwitch = false);
