Fix for niagara curves not rendering in curve editor if they are linear (all keys have the value)

diff --git a/Engine/Plugins/FX/Niagara/Source/NiagaraEditorWidgets/Private/DetailCustomizations/NiagaraDataInterfaceCurveDetails.cpp b/Engine/Plugins/FX/Niagara/Source/NiagaraEditorWidgets/Private/DetailCustomizations/NiagaraDataInterfaceCurveDetails.cpp
--- a/Engine/Plugins/FX/Niagara/Source/NiagaraEditorWidgets/Private/DetailCustomizations/NiagaraDataInterfaceCurveDetails.cpp
+++ b/Engine/Plugins/FX/Niagara/Source/NiagaraEditorWidgets/Private/DetailCustomizations/NiagaraDataInterfaceCurveDetails.cpp
@@ -283,6 +283,34 @@ private:
 			ViewMaxOutput = 1;
 		}
 
+		if (FMath::IsNearlyEqual(ViewMinInput, ViewMaxInput))
+		{
+			if (FMath::IsWithinInclusive(ViewMinInput, 0.0f, 1.0f))
+			{
+				ViewMinInput = 0;
+				ViewMaxInput = 1;
+			}
+			else
+			{
+				ViewMinInput -= 0.5f;
+				ViewMaxInput += 0.5f;
+			}
+		}
+
+		if (FMath::IsNearlyEqual(ViewMinOutput, ViewMaxOutput))
+		{
+			if (FMath::IsWithinInclusive(ViewMinOutput, 0.0f, 1.0f))
+			{
+				ViewMinOutput = 0;
+				ViewMaxOutput = 1;
+			}
+			else
+			{
+				ViewMinOutput -= 0.5f;
+				ViewMaxOutput += 0.5f;
+			}
+		}
+
 		float ViewInputRange = ViewMaxInput - ViewMinInput;
 		float ViewOutputRange = ViewMaxOutput - ViewMinOutput;
 		float ViewInputPadding = ViewInputRange * .05f;
