Niagara - Fix the system factory to use the required system update script from the niagara editor settings instead of hard coding system life cycle.  Also update the ini to default to system life cycle so that emitter editors still work.

diff --git a/Engine/Plugins/FX/Niagara/Config/BaseNiagara.ini b/Engine/Plugins/FX/Niagara/Config/BaseNiagara.ini
--- a/Engine/Plugins/FX/Niagara/Config/BaseNiagara.ini
+++ b/Engine/Plugins/FX/Niagara/Config/BaseNiagara.ini
@@ -29,6 +29,7 @@ DefaultScript=/Niagara/DefaultAssets/DefaultModule.DefaultModule
 DefaultDynamicInputScript=/Niagara/DefaultAssets/DefaultDynamicInputs.DefaultDynamicInputs
 DefaultFunctionScript=/Niagara/DefaultAssets/DefaultFunction.DefaultFunction
 DefaultModuleScript=/Niagara/DefaultAssets/DefaultModule.DefaultModule
+RequiredSystemUpdateScript=/Niagara/Modules/System/SystemLifeCycle.SystemLifeCycle
 +GraphCreationShortcuts=(Name="Numeric::Add",Input=(Key=A,bShift=False,bCtrl=False,bAlt=False,bCmd=False))
 +GraphCreationShortcuts=(Name="Numeric::Div",Input=(Key=D,bShift=False,bCtrl=False,bAlt=False,bCmd=False))
 +GraphCreationShortcuts=(Name="Numeric::Pow",Input=(Key=E,bShift=False,bCtrl=False,bAlt=False,bCmd=False))
diff --git a/Engine/Plugins/FX/Niagara/Source/NiagaraEditor/Private/NiagaraSystemFactoryNew.cpp b/Engine/Plugins/FX/Niagara/Source/NiagaraEditor/Private/NiagaraSystemFactoryNew.cpp
--- a/Engine/Plugins/FX/Niagara/Source/NiagaraEditor/Private/NiagaraSystemFactoryNew.cpp
+++ b/Engine/Plugins/FX/Niagara/Source/NiagaraEditor/Private/NiagaraSystemFactoryNew.cpp
@@ -163,7 +163,7 @@ void UNiagaraSystemFactoryNew::InitializeSystem(UNiagaraSystem* System, bool bCr
 
 	if (bCreateDefaultNodes)
 	{
-		FSoftObjectPath SystemUpdateScriptRef(TEXT("/Niagara/Modules/System/SystemLifeCycle.SystemLifeCycle"));
+		FSoftObjectPath SystemUpdateScriptRef = GetDefault<UNiagaraEditorSettings>()->RequiredSystemUpdateScript;
 		UNiagaraScript* Script = Cast<UNiagaraScript>(SystemUpdateScriptRef.TryLoad());
 
 		FAssetData ModuleScriptAsset(Script);
