Fixing UNiagaraScript GC crash

diff --git a/Engine/Plugins/FX/Niagara/Source/Niagara/Classes/NiagaraScript.h b/Engine/Plugins/FX/Niagara/Source/Niagara/Classes/NiagaraScript.h
--- a/Engine/Plugins/FX/Niagara/Source/Niagara/Classes/NiagaraScript.h
+++ b/Engine/Plugins/FX/Niagara/Source/Niagara/Classes/NiagaraScript.h
@@ -454,6 +454,7 @@ public:
 	//~ Begin UObject interface
 	void Serialize(FArchive& Ar)override;
 	virtual void PostLoad() override;
+	virtual bool IsDestructionThreadSafe() const override { return false; }
 #if WITH_EDITOR
 	virtual void PostEditChangeProperty(FPropertyChangedEvent& PropertyChangedEvent) override;
 #endif
