Attempted CIS fix for NiagaraEditor triggering FName comparison operator deprecation warning

diff --git a/Engine/Plugins/FX/Niagara/Source/NiagaraEditor/Private/Customizations/NiagaraMetaDataCustomNodeBuilder.cpp b/Engine/Plugins/FX/Niagara/Source/NiagaraEditor/Private/Customizations/NiagaraMetaDataCustomNodeBuilder.cpp
--- a/Engine/Plugins/FX/Niagara/Source/NiagaraEditor/Private/Customizations/NiagaraMetaDataCustomNodeBuilder.cpp
+++ b/Engine/Plugins/FX/Niagara/Source/NiagaraEditor/Private/Customizations/NiagaraMetaDataCustomNodeBuilder.cpp
@@ -86,7 +86,7 @@ void FNiagaraMetaDataCustomNodeBuilder::GenerateChildContent(IDetailChildrenBuil
 		{
 			const FNiagaraVariable& ParameterVariableA = MetaDataRowDataA.Get<0>();
 			const FNiagaraVariable& ParameterVariableB = MetaDataRowDataB.Get<0>();
-			return ParameterVariableA.GetName() < ParameterVariableB.GetName();
+			return FNameLexicalLess()(ParameterVariableA.GetName(), ParameterVariableB.GetName());
 		}
 	});
 
