Fixed up some code in NiagaraRenderer.h and ModuleService.cpp for VS2019 (implicit bool conversion)

diff --git a/Engine/Plugins/FX/Niagara/Source/Niagara/Public/NiagaraRenderer.h b/Engine/Plugins/FX/Niagara/Source/Niagara/Public/NiagaraRenderer.h
--- a/Engine/Plugins/FX/Niagara/Source/Niagara/Public/NiagaraRenderer.h
+++ b/Engine/Plugins/FX/Niagara/Source/Niagara/Public/NiagaraRenderer.h
@@ -106,7 +106,7 @@ public:
 
 	void SetDynamicData_RenderThread(FNiagaraDynamicDataBase* NewDynamicData);
 	FORCEINLINE FNiagaraDynamicDataBase *GetDynamicData()const { return DynamicDataRender; }
-	FORCEINLINE bool HasDynamicData()const { return DynamicDataRender; }
+	FORCEINLINE bool HasDynamicData()const { return !!DynamicDataRender; }
 	FORCEINLINE float GetCPUTimeMS() const { return CPUTimeMS; }
 	FORCEINLINE bool HasLights()const { return bHasLights; }
 
diff --git a/Engine/Source/Developer/TraceServices/Private/ModuleService.cpp b/Engine/Source/Developer/TraceServices/Private/ModuleService.cpp
--- a/Engine/Source/Developer/TraceServices/Private/ModuleService.cpp
+++ b/Engine/Source/Developer/TraceServices/Private/ModuleService.cpp
@@ -54,7 +54,7 @@ void FModuleService::SetModuleEnabled(const FName& ModuleName, bool bEnabled)
 	{
 		return;
 	}
-	bool bWasEnabled = EnabledModules.Find(*FindIt);
+	bool bWasEnabled = !!EnabledModules.Find(*FindIt);
 	if (bEnabled == bWasEnabled)
 	{
 		return;
