Merging //UE4/Main @ 6685008 to //UE4/Dev-Core (sbo182_DevCore)

(peafour-mergedown)

diff --git a/Engine/Build/Android/Java/src/com/epicgames/ue4/VideoDecoder.java b/Engine/Build/Android/Java/src/com/epicgames/ue4/VideoDecoder.java
--- a/Engine/Build/Android/Java/src/com/epicgames/ue4/VideoDecoder.java
+++ b/Engine/Build/Android/Java/src/com/epicgames/ue4/VideoDecoder.java
@@ -152,10 +152,24 @@ public class VideoDecoder
 					MediaCodec.BufferInfo info = new MediaCodec.BufferInfo();
 
 					int outIndex = mVideoCodec.dequeueOutputBuffer(info, -1);
-					if (outIndex >= 0) 
+
+					// manage the config frames
+					switch (outIndex) 
 					{
-						addFrameTimeDuration(timeStamp, duration); //us to ms
-						mVideoCodec.releaseOutputBuffer(outIndex, true);
+						case MediaCodec.INFO_OUTPUT_BUFFERS_CHANGED:
+							GameActivity.Log.debug("Android Video Decoder: INFO_OUTPUT_BUFFERS_CHANGED");
+							break;
+						case MediaCodec.INFO_OUTPUT_FORMAT_CHANGED:
+							GameActivity.Log.debug("Android Video Decoder: New format" + mVideoCodec.getOutputFormat());
+							break;
+						case MediaCodec.INFO_TRY_AGAIN_LATER:
+							GameActivity.Log.debug("Android Video Decoder: dequeueOutputBuffer timed out!");
+							break;
+						default:
+							addFrameTimeDuration(timeStamp, duration); //us to ms
+							
+							boolean doRender = (info.size != 0) && ((info.flags & MediaCodec.BUFFER_FLAG_CODEC_CONFIG) != MediaCodec.BUFFER_FLAG_CODEC_CONFIG);
+							mVideoCodec.releaseOutputBuffer(outIndex, doRender);
 					}
 
 					if ((info.flags & MediaCodec.BUFFER_FLAG_END_OF_STREAM) != 0) 
diff --git a/Engine/Plugins/FX/Niagara/Source/Niagara/Classes/NDISkeletalMeshCommon.h b/Engine/Plugins/FX/Niagara/Source/Niagara/Classes/NDISkeletalMeshCommon.h
--- a/Engine/Plugins/FX/Niagara/Source/Niagara/Classes/NDISkeletalMeshCommon.h
+++ b/Engine/Plugins/FX/Niagara/Source/Niagara/Classes/NDISkeletalMeshCommon.h
@@ -101,6 +101,16 @@ struct FSkinnedPositionAccessorHelper
 	}
 
 	FORCEINLINE FVector GetSkinnedBonePreviousPosition(FSkeletalMeshAccessorHelper& Accessor, int32 BoneIndex)
+	{
+		checkf(false, TEXT("Must provide a specialization for this template type"));
+	}	
+	
+	FORCEINLINE_DEBUGGABLE FQuat GetSkinnedBoneRotation(FSkeletalMeshAccessorHelper& Accessor, int32 BoneIndex)
+	{
+		checkf(false, TEXT("Must provide a specialization for this template type"));
+	}
+
+	FORCEINLINE_DEBUGGABLE FQuat GetSkinnedBonePreviousRotation(FSkeletalMeshAccessorHelper& Accessor, int32 BoneIndex)
 	{
 		checkf(false, TEXT("Must provide a specialization for this template type"));
 	}
@@ -151,6 +161,16 @@ struct FSkinnedPositionAccessorHelper<TIntegralConstant<ENDISkeletalMesh_Skinnin
 	{
 		return Accessor.Mesh->GetComposedRefPoseMatrix(BoneIndex).GetOrigin();
 	}
+
+	FORCEINLINE_DEBUGGABLE FQuat GetSkinnedBoneRotation(FSkeletalMeshAccessorHelper& Accessor, int32 BoneIndex)
+	{
+		return Accessor.Mesh->GetComposedRefPoseMatrix(BoneIndex).ToQuat();
+	}
+
+	FORCEINLINE_DEBUGGABLE FQuat GetSkinnedBonePreviousRotation(FSkeletalMeshAccessorHelper& Accessor, int32 BoneIndex)
+	{
+		return Accessor.Mesh->GetComposedRefPoseMatrix(BoneIndex).ToQuat();
+	}
 };
 
 template<>
@@ -198,6 +218,16 @@ struct FSkinnedPositionAccessorHelper<TIntegralConstant<ENDISkeletalMesh_Skinnin
 	{
 		return (*Accessor.PrevBoneComponentSpaceTransforms)[BoneIndex].GetLocation();
 	}
+
+	FORCEINLINE_DEBUGGABLE FQuat GetSkinnedBoneRotation(FSkeletalMeshAccessorHelper& Accessor, int32 BoneIndex)
+	{
+		return (*Accessor.BoneComponentSpaceTransforms)[BoneIndex].GetRotation();
+	}
+
+	FORCEINLINE_DEBUGGABLE FQuat GetSkinnedBonePreviousRotation(FSkeletalMeshAccessorHelper& Accessor, int32 BoneIndex)
+	{
+		return (*Accessor.PrevBoneComponentSpaceTransforms)[BoneIndex].GetRotation();
+	}
 };
 
 template<>
@@ -245,6 +275,16 @@ struct FSkinnedPositionAccessorHelper<TIntegralConstant<ENDISkeletalMesh_Skinnin
 	{
 		return (*Accessor.PrevBoneComponentSpaceTransforms)[BoneIndex].GetLocation();
 	}
+
+	FORCEINLINE_DEBUGGABLE FQuat GetSkinnedBoneRotation(FSkeletalMeshAccessorHelper& Accessor, int32 BoneIndex)
+	{
+		return (*Accessor.BoneComponentSpaceTransforms)[BoneIndex].GetRotation();
+	}
+
+	FORCEINLINE_DEBUGGABLE FQuat GetSkinnedBonePreviousRotation(FSkeletalMeshAccessorHelper& Accessor, int32 BoneIndex)
+	{
+		return (*Accessor.PrevBoneComponentSpaceTransforms)[BoneIndex].GetRotation();
+	}
 };
 
 //////////////////////////////////////////////////////////////////////////
diff --git a/Engine/Plugins/FX/Niagara/Source/Niagara/Private/NDISkeletalMesh_BoneSampling.cpp b/Engine/Plugins/FX/Niagara/Source/Niagara/Private/NDISkeletalMesh_BoneSampling.cpp
--- a/Engine/Plugins/FX/Niagara/Source/Niagara/Private/NDISkeletalMesh_BoneSampling.cpp
+++ b/Engine/Plugins/FX/Niagara/Source/Niagara/Private/NDISkeletalMesh_BoneSampling.cpp
@@ -70,6 +70,7 @@ void UNiagaraDataInterfaceSkeletalMesh::GetSkeletonSamplingFunctions(TArray<FNia
 		Sig.Inputs.Add(FNiagaraVariable(FNiagaraTypeDefinition(GetClass()), TEXT("SkeletalMesh")));
 		Sig.Inputs.Add(FNiagaraVariable(FNiagaraTypeDefinition::GetIntDef(), TEXT("Bone")));
 		Sig.Outputs.Add(FNiagaraVariable(FNiagaraTypeDefinition::GetVec3Def(), TEXT("Position")));
+		Sig.Outputs.Add(FNiagaraVariable(FNiagaraTypeDefinition::GetQuatDef(), TEXT("Rotation")));
 		Sig.Outputs.Add(FNiagaraVariable(FNiagaraTypeDefinition::GetVec3Def(), TEXT("Velocity")));
 		Sig.bMemberFunction = true;
 		Sig.bRequiresContext = false;
@@ -85,6 +86,7 @@ void UNiagaraDataInterfaceSkeletalMesh::GetSkeletonSamplingFunctions(TArray<FNia
 		Sig.Inputs.Add(FNiagaraVariable(FNiagaraTypeDefinition(GetClass()), TEXT("SkeletalMesh")));
 		Sig.Inputs.Add(FNiagaraVariable(FNiagaraTypeDefinition::GetIntDef(), TEXT("Bone")));
 		Sig.Outputs.Add(FNiagaraVariable(FNiagaraTypeDefinition::GetVec3Def(), TEXT("Position")));
+		Sig.Outputs.Add(FNiagaraVariable(FNiagaraTypeDefinition::GetQuatDef(), TEXT("Rotation")));
 		Sig.Outputs.Add(FNiagaraVariable(FNiagaraTypeDefinition::GetVec3Def(), TEXT("Velocity")));
 		Sig.bMemberFunction = true;
 		Sig.bRequiresContext = false;
@@ -102,6 +104,7 @@ void UNiagaraDataInterfaceSkeletalMesh::GetSkeletonSamplingFunctions(TArray<FNia
 		Sig.Inputs.Add(FNiagaraVariable(FNiagaraTypeDefinition::GetIntDef(), TEXT("Bone")));
 		Sig.Inputs.Add(FNiagaraVariable(FNiagaraTypeDefinition::GetFloatDef(), TEXT("Interpolation")));
 		Sig.Outputs.Add(FNiagaraVariable(FNiagaraTypeDefinition::GetVec3Def(), TEXT("Position")));
+		Sig.Outputs.Add(FNiagaraVariable(FNiagaraTypeDefinition::GetQuatDef(), TEXT("Rotation")));
 		Sig.Outputs.Add(FNiagaraVariable(FNiagaraTypeDefinition::GetVec3Def(), TEXT("Velocity")));
 		Sig.bMemberFunction = true;
 		Sig.bRequiresContext = false;
@@ -118,6 +121,7 @@ void UNiagaraDataInterfaceSkeletalMesh::GetSkeletonSamplingFunctions(TArray<FNia
 		Sig.Inputs.Add(FNiagaraVariable(FNiagaraTypeDefinition::GetIntDef(), TEXT("Bone")));
 		Sig.Inputs.Add(FNiagaraVariable(FNiagaraTypeDefinition::GetFloatDef(), TEXT("Interpolation")));
 		Sig.Outputs.Add(FNiagaraVariable(FNiagaraTypeDefinition::GetVec3Def(), TEXT("Position")));
+		Sig.Outputs.Add(FNiagaraVariable(FNiagaraTypeDefinition::GetQuatDef(), TEXT("Rotation")));
 		Sig.Outputs.Add(FNiagaraVariable(FNiagaraTypeDefinition::GetVec3Def(), TEXT("Velocity")));
 		Sig.bMemberFunction = true;
 		Sig.bRequiresContext = false;
@@ -214,22 +218,22 @@ void UNiagaraDataInterfaceSkeletalMesh::BindSkeletonSamplingFunction(const FVMEx
 	}
 	else if (BindingInfo.Name == GetSkinnedBoneDataName)
 	{
-		check(BindingInfo.GetNumInputs() == 2 && BindingInfo.GetNumOutputs() == 6);
+		ensure(BindingInfo.GetNumInputs() == 2 && BindingInfo.GetNumOutputs() == 10);
 		TSkinningModeBinder<TNDIExplicitBinder<FNDITransformHandlerNoop, TNDIExplicitBinder<TIntegralConstant<bool, false>, NDI_FUNC_BINDER(UNiagaraDataInterfaceSkeletalMesh, GetSkinnedBoneData)>>>::Bind(this, BindingInfo, InstanceData, OutFunc);
 	}
 	else if (BindingInfo.Name == GetSkinnedBoneDataWSName)
 	{
-		check(BindingInfo.GetNumInputs() == 2 && BindingInfo.GetNumOutputs() == 6);
+		ensure(BindingInfo.GetNumInputs() == 2 && BindingInfo.GetNumOutputs() == 10);
 		TSkinningModeBinder<TNDIExplicitBinder<FNDITransformHandler, TNDIExplicitBinder<TIntegralConstant<bool, false>, NDI_FUNC_BINDER(UNiagaraDataInterfaceSkeletalMesh, GetSkinnedBoneData)>>>::Bind(this, BindingInfo, InstanceData, OutFunc);
 	}
 	else if (BindingInfo.Name == GetSkinnedBoneDataInterpolatedName)
 	{
-		check(BindingInfo.GetNumInputs() == 3 && BindingInfo.GetNumOutputs() == 6);
+		ensure(BindingInfo.GetNumInputs() == 3 && BindingInfo.GetNumOutputs() == 10);
 		TSkinningModeBinder<TNDIExplicitBinder<FNDITransformHandlerNoop, TNDIExplicitBinder<TIntegralConstant<bool, true>, NDI_FUNC_BINDER(UNiagaraDataInterfaceSkeletalMesh, GetSkinnedBoneData)>>>::Bind(this, BindingInfo, InstanceData, OutFunc);
 	}
 	else if (BindingInfo.Name == GetSkinnedBoneDataWSInterpolatedName)
 	{
-		check(BindingInfo.GetNumInputs() == 3 && BindingInfo.GetNumOutputs() == 6);
+		ensure(BindingInfo.GetNumInputs() == 3 && BindingInfo.GetNumOutputs() == 10);
 		TSkinningModeBinder<TNDIExplicitBinder<FNDITransformHandler, TNDIExplicitBinder<TIntegralConstant<bool, true>, NDI_FUNC_BINDER(UNiagaraDataInterfaceSkeletalMesh, GetSkinnedBoneData)>>>::Bind(this, BindingInfo, InstanceData, OutFunc);
 	}
 	else if (BindingInfo.Name == GetSpecificBoneCountName)
@@ -360,18 +364,22 @@ struct FBoneSocketSkinnedDataOutputHandler
 {
 	FBoneSocketSkinnedDataOutputHandler(FVectorVMContext& Context)
 		: PosX(Context), PosY(Context), PosZ(Context)
+		, RotX(Context), RotY(Context), RotZ(Context), RotW(Context)
 		, VelX(Context), VelY(Context), VelZ(Context)
 		, bNeedsPosition(PosX.IsValid() || PosY.IsValid() || PosZ.IsValid())
+		, bNeedsRotation(RotX.IsValid() || RotY.IsValid() || RotZ.IsValid() || RotW.IsValid())
 		, bNeedsVelocity(VelX.IsValid() || VelY.IsValid() || VelZ.IsValid())
 	{
 	}
 
 	VectorVM::FExternalFuncRegisterHandler<float> PosX; VectorVM::FExternalFuncRegisterHandler<float> PosY; VectorVM::FExternalFuncRegisterHandler<float> PosZ;
+	VectorVM::FExternalFuncRegisterHandler<float> RotX; VectorVM::FExternalFuncRegisterHandler<float> RotY; VectorVM::FExternalFuncRegisterHandler<float> RotZ; VectorVM::FExternalFuncRegisterHandler<float> RotW;
 	VectorVM::FExternalFuncRegisterHandler<float> VelX; VectorVM::FExternalFuncRegisterHandler<float> VelY; VectorVM::FExternalFuncRegisterHandler<float> VelZ;
 
 	//TODO: Rotation + Scale too? Use quats so we can get proper interpolation between bone and parent.
 
 	const bool bNeedsPosition;
+	const bool bNeedsRotation;
 	const bool bNeedsVelocity;
 
 	FORCEINLINE void SetPosition(FVector Position)
@@ -381,6 +389,14 @@ struct FBoneSocketSkinnedDataOutputHandler
 		*PosZ.GetDestAndAdvance() = Position.Z;
 	}
 
+	FORCEINLINE void SetRotation(FQuat Rotation)
+	{
+		*RotX.GetDestAndAdvance() = Rotation.X;
+		*RotY.GetDestAndAdvance() = Rotation.Y;
+		*RotZ.GetDestAndAdvance() = Rotation.Z;
+		*RotW.GetDestAndAdvance() = Rotation.W;
+	}
+
 	FORCEINLINE void SetVelocity(FVector Velocity)
 	{
 		*VelX.GetDestAndAdvance() = Velocity.X;
@@ -469,6 +485,19 @@ void UNiagaraDataInterfaceSkeletalMesh::GetSkinnedBoneData(FVectorVMContext& Con
 			Velocity = (Pos - Prev) * InvDt;
 			Output.SetVelocity(Velocity);
 		}
+
+		if (Output.bNeedsRotation)
+		{
+			FQuat Rotation = SkinningHandler.GetSkinnedBoneRotation(Accessor, Bone);
+
+			if (bInterpolated::Value)
+			{
+				FQuat PrevRotation = SkinningHandler.GetSkinnedBonePreviousRotation(Accessor, Bone);
+				Rotation = FMath::Lerp(PrevRotation, Rotation, Interp);
+			}
+
+			Output.SetRotation(Rotation);
+		}
 	}
 }
 
diff --git a/Engine/Plugins/FX/Niagara/Source/NiagaraEditor/Private/Toolkits/NiagaraScriptToolkit.cpp b/Engine/Plugins/FX/Niagara/Source/NiagaraEditor/Private/Toolkits/NiagaraScriptToolkit.cpp
--- a/Engine/Plugins/FX/Niagara/Source/NiagaraEditor/Private/Toolkits/NiagaraScriptToolkit.cpp
+++ b/Engine/Plugins/FX/Niagara/Source/NiagaraEditor/Private/Toolkits/NiagaraScriptToolkit.cpp
@@ -88,7 +88,7 @@ void FNiagaraScriptToolkit::RegisterTabSpawners(const TSharedRef<class FTabManag
 		.SetIcon(FSlateIcon(FEditorStyle::GetStyleSetName(), "LevelEditor.Tabs.Details"));
 
 	InTabManager->RegisterTabSpawner(SelectedDetailsTabId, FOnSpawnTab::CreateSP(this, &FNiagaraScriptToolkit::SpawnTabSelectedDetails))
-		.SetDisplayName(LOCTEXT("ScriptDetailsTab", "Selected Details"))
+		.SetDisplayName(LOCTEXT("SelectedDetailsTab", "Selected Details"))
 		.SetGroup(WorkspaceMenuCategoryRef)
 		.SetIcon(FSlateIcon(FEditorStyle::GetStyleSetName(), "LevelEditor.Tabs.Details"));
 
diff --git a/Engine/Source/Runtime/Engine/Classes/Engine/NetConnection.h b/Engine/Source/Runtime/Engine/Classes/Engine/NetConnection.h
--- a/Engine/Source/Runtime/Engine/Classes/Engine/NetConnection.h
+++ b/Engine/Source/Runtime/Engine/Classes/Engine/NetConnection.h
@@ -1107,6 +1107,11 @@ public:
 	 * @param bWasSaturated		True if we failed to replicate all data because we were saturated.
 	 */
 	ENGINE_API void TrackReplicationForAnalytics(const bool bWasSaturated);
+
+	/**
+	 * Get the current number of sent packets for which we have received a delivery notification
+	 */
+	ENGINE_API uint32 GetOutTotalNotifiedPackets() const { return OutTotalNotifiedPackets; }
 	
 protected:
 
@@ -1190,6 +1195,9 @@ private:
 	/** Full PacketId  of last sent packet that we have received notification for (i.e. we know if it was delivered or not). Related to OutAckPacketId which is tha last successfully delivered PacketId */
 	int32 LastNotifiedPacketId;
 
+	/** Count the number of notified packets, i.e. packets that we know if they are delivered or not. Used to reliably measure outgoing packet loss */
+	uint32 OutTotalNotifiedPackets;
+
 	/** Keep old behavior where we send a packet with only acks even if we have no other outgoing data if we got incoming data */
 	uint32 HasDirtyAcks;
 	
diff --git a/Engine/Source/Runtime/Engine/Classes/Engine/NetDriver.h b/Engine/Source/Runtime/Engine/Classes/Engine/NetDriver.h
--- a/Engine/Source/Runtime/Engine/Classes/Engine/NetDriver.h
+++ b/Engine/Source/Runtime/Engine/Classes/Engine/NetDriver.h
@@ -1520,6 +1520,16 @@ public:
 		return bMaySendProperties;
 	}
 
+	/**
+	 * Get the current number of sent packets for which we have received a delivery notification
+	 */
+	ENGINE_API uint32 GetOutTotalNotifiedPackets() const { return OutTotalNotifiedPackets; }
+
+	/**
+	 * Increase the current number of sent packets for which we have received a delivery notification
+	 */
+	inline void IncreaseOutTotalNotifiedPackets() { ++OutTotalNotifiedPackets; }
+
 protected:
 
 	bool bMaySendProperties;
@@ -1551,4 +1561,7 @@ private:
 
 	/** NetDriver time to end packet loss burst simulation. */
 	float PacketLossBurstEndTime;
+
+	/** Count the number of notified packets, i.e. packets that we know if they are delivered or not. Used to reliably measure outgoing packet loss */
+	uint32 OutTotalNotifiedPackets;
 };
\ No newline at end of file
diff --git a/Engine/Source/Runtime/Engine/Private/Components/PrimitiveComponent.cpp b/Engine/Source/Runtime/Engine/Private/Components/PrimitiveComponent.cpp
--- a/Engine/Source/Runtime/Engine/Private/Components/PrimitiveComponent.cpp
+++ b/Engine/Source/Runtime/Engine/Private/Components/PrimitiveComponent.cpp
@@ -805,11 +805,6 @@ void UPrimitiveComponent::Serialize(FArchive& Ar)
 			LightmapType = ELightmapType::ForceSurface;
 		}
 	}
-	
-	if (Ar.CustomVer(FFortniteMainBranchObjectVersion::GUID) < FFortniteMainBranchObjectVersion::CullDistanceRefactor_RemovedDefaultDistance)
-	{
-		LDMaxDrawDistance = 0.f;
-	}
 }
 
 #if WITH_EDITOR
diff --git a/Engine/Source/Runtime/Engine/Private/NetConnection.cpp b/Engine/Source/Runtime/Engine/Private/NetConnection.cpp
--- a/Engine/Source/Runtime/Engine/Private/NetConnection.cpp
+++ b/Engine/Source/Runtime/Engine/Private/NetConnection.cpp
@@ -169,6 +169,7 @@ PRAGMA_ENABLE_DEPRECATION_WARNINGS
 ,	PlayerOnlinePlatformName( NAME_None )
 ,	ClientWorldPackageName( NAME_None )
 ,	LastNotifiedPacketId( -1 )
+,	OutTotalNotifiedPackets(0)
 ,	HasDirtyAcks(0u)
 ,	bHasWarnedAboutChannelLimit(false)
 ,	bConnectionPendingCloseDueToSocketSendFailure(false)
@@ -1764,6 +1765,8 @@ void UNetConnection::ReceivedPacket( FBitReader& Reader )
 		{
 			// Increase LastNotifiedPacketId, this is a full packet Id
 			++LastNotifiedPacketId;
+			++OutTotalNotifiedPackets;
+			Driver->IncreaseOutTotalNotifiedPackets();
 
 			// Sanity check
 			if (FNetPacketNotify::SequenceNumberT(LastNotifiedPacketId) != AckedSequence)
diff --git a/Engine/Source/Runtime/Engine/Private/NetworkDriver.cpp b/Engine/Source/Runtime/Engine/Private/NetworkDriver.cpp
--- a/Engine/Source/Runtime/Engine/Private/NetworkDriver.cpp
+++ b/Engine/Source/Runtime/Engine/Private/NetworkDriver.cpp
@@ -271,6 +271,7 @@ UNetDriver::UNetDriver(const FObjectInitializer& ObjectInitializer)
 ,	LagState(ENetworkLagState::NotLagging)
 ,	DuplicateLevelID(INDEX_NONE)
 ,	PacketLossBurstEndTime(-1.0f)
+,	OutTotalNotifiedPackets(0u)
 {
 PRAGMA_DISABLE_DEPRECATION_WARNINGS
 	ChannelClasses[CHTYPE_Control]	= UControlChannel::StaticClass();
diff --git a/Engine/Source/Runtime/SlateRHIRenderer/Private/SlateRHIRenderingPolicy.cpp b/Engine/Source/Runtime/SlateRHIRenderer/Private/SlateRHIRenderingPolicy.cpp
--- a/Engine/Source/Runtime/SlateRHIRenderer/Private/SlateRHIRenderingPolicy.cpp
+++ b/Engine/Source/Runtime/SlateRHIRenderer/Private/SlateRHIRenderingPolicy.cpp
@@ -720,7 +720,28 @@ void FSlateRHIRenderingPolicy::DrawElements(
 	*/
 	for (int32 BatchIndex = 0; BatchIndex < RenderBatches.Num(); ++BatchIndex)
 	{
-		check(RHICmdList.IsInsideRenderPass());
+		if (!RHICmdList.IsInsideRenderPass())
+		{
+			// Restart the renderpass since the CustomDrawer or post-process may have changed it in last iteration
+			FRHIRenderPassInfo RPInfo(BackBuffer.GetRenderTargetTexture(), ERenderTargetActions::Load_Store);
+			RPInfo.DepthStencilRenderTarget.DepthStencilTarget = DepthStencilTarget;
+			if (DepthStencilTarget)
+			{
+				RPInfo.DepthStencilRenderTarget.Action = EDepthStencilTargetActions::LoadDepthStencil_StoreDepthStencil;
+				RPInfo.DepthStencilRenderTarget.ExclusiveDepthStencil = FExclusiveDepthStencil::DepthWrite_StencilWrite;
+			}
+			else
+			{
+				RPInfo.DepthStencilRenderTarget.Action = EDepthStencilTargetActions::DontLoad_DontStore;
+				RPInfo.DepthStencilRenderTarget.ExclusiveDepthStencil = FExclusiveDepthStencil::DepthNop_StencilNop;
+			}
+			RHICmdList.BeginRenderPass(RPInfo, TEXT("RestartingSlateDrawElements"));
+
+			// Something may have messed with the viewport size so set it back to the full target.
+			RHICmdList.SetViewport(0, 0, 0, BackBuffer.GetSizeXY().X, BackBuffer.GetSizeXY().Y, 0.0f);
+			RHICmdList.SetStreamSource(0, VertexBuffer->VertexBufferRHI, 0);
+		}
+				
 #if WITH_SLATE_VISUALIZERS
 		FLinearColor BatchColor = FLinearColor(BatchColors.GetUnitVector());
 #endif
@@ -1227,21 +1248,7 @@ void FSlateRHIRenderingPolicy::DrawElements(
 				PostProcessor->BlurRect(RHICmdList, RendererModule, BlurParams, RectParams);
 
 				check(RHICmdList.IsOutsideRenderPass());
-				{
-					FRHIRenderPassInfo RPInfo(BackBuffer.GetRenderTargetTexture(), ERenderTargetActions::Load_Store);
-					RPInfo.DepthStencilRenderTarget.DepthStencilTarget = DepthStencilTarget;
-					if (DepthStencilTarget)
-					{
-						RPInfo.DepthStencilRenderTarget.Action = EDepthStencilTargetActions::LoadDepthStencil_StoreDepthStencil;
-						RPInfo.DepthStencilRenderTarget.ExclusiveDepthStencil = FExclusiveDepthStencil::DepthWrite_StencilWrite;
-					}
-					else
-					{
-						RPInfo.DepthStencilRenderTarget.Action = EDepthStencilTargetActions::DontLoad_DontStore;
-						RPInfo.DepthStencilRenderTarget.ExclusiveDepthStencil = FExclusiveDepthStencil::DepthNop_StencilNop;
-					}
-					RHICmdList.BeginRenderPass(RPInfo, TEXT("RestartingSlateDrawElements"));
-				}
+				// Render pass for slate elements will be restarted on a next loop iteration if any
 			}
 		}
 		else
@@ -1250,6 +1257,7 @@ void FSlateRHIRenderingPolicy::DrawElements(
 			if (CustomDrawer.IsValid())
 			{
 				// CustomDrawers will change the rendertarget. So we must close any outstanding renderpasses.
+				// Render pass for slate elements will be restarted on a next loop iteration if any
 				RHICmdList.EndRenderPass();	
 
 				SLATE_DRAW_EVENT(RHICmdList, CustomDrawer);
@@ -1263,25 +1271,6 @@ void FSlateRHIRenderingPolicy::DrawElements(
 
 				//We reset the maskingID here because otherwise the RT might not get re-set in the lines above see: if (bClearStencil || bForceStateChange)
 				MaskingID = 0;
-
-				// Restart the renderpass since the CustomDrawer will have drawn into its own.
-				FRHIRenderPassInfo RPInfo(ColorTarget, ERenderTargetActions::Load_Store);
-				RPInfo.DepthStencilRenderTarget.DepthStencilTarget = DepthStencilTarget;
-				if (DepthStencilTarget)
-				{
-					RPInfo.DepthStencilRenderTarget.Action = EDepthStencilTargetActions::LoadDepthStencil_StoreDepthStencil;
-					RPInfo.DepthStencilRenderTarget.ExclusiveDepthStencil = FExclusiveDepthStencil::DepthWrite_StencilWrite;
-				}
-				else
-				{
-					RPInfo.DepthStencilRenderTarget.Action = EDepthStencilTargetActions::DontLoad_DontStore;
-					RPInfo.DepthStencilRenderTarget.ExclusiveDepthStencil = FExclusiveDepthStencil::DepthNop_StencilNop;
-				}
-				RHICmdList.BeginRenderPass(RPInfo, TEXT("SlateDrawElements"));
-
-				// Something may have messed with the viewport size so set it back to the full target.
-				RHICmdList.SetViewport(0, 0, 0, BackBuffer.GetSizeXY().X, BackBuffer.GetSizeXY().Y, 0.0f);
-				RHICmdList.SetStreamSource(0, VertexBuffer->VertexBufferRHI, 0);
 			}
 		} // CustomDrawer
 	}
@@ -1290,7 +1279,10 @@ void FSlateRHIRenderingPolicy::DrawElements(
 #if !(PLATFORM_IOS || PLATFORM_ANDROID)
 	if (bApplyColorDeficiencyCorrection && GSlateColorDeficiencyType != EColorVisionDeficiency::NormalVision && GSlateColorDeficiencySeverity > 0)
 	{
-		RHICmdList.EndRenderPass();
+		if (RHICmdList.IsInsideRenderPass())
+		{
+			RHICmdList.EndRenderPass();
+		}
 
 		FPostProcessRectParams RectParams;
 		RectParams.SourceTexture = BackBuffer.GetRenderTargetTexture();
diff --git a/Engine/Source/Runtime/VulkanRHI/Private/VulkanRenderTarget.cpp b/Engine/Source/Runtime/VulkanRHI/Private/VulkanRenderTarget.cpp
--- a/Engine/Source/Runtime/VulkanRHI/Private/VulkanRenderTarget.cpp
+++ b/Engine/Source/Runtime/VulkanRHI/Private/VulkanRenderTarget.cpp
@@ -312,6 +312,12 @@ void FTransitionAndLayoutManager::BeginRealRenderPass(FVulkanCommandListContext&
 			{
 				VulkanRHI::ImagePipelineBarrier(CmdBuffer->GetHandle(), Surface.Image, EImageLayoutBarrier::Undefined, EImageLayoutBarrier::ColorAttachment, SetupImageSubresourceRange());
 			}
+			else if (*Found == VK_IMAGE_LAYOUT_COLOR_ATTACHMENT_OPTIMAL && GetLoadAction(RPInfo.ColorRenderTargets[Index].Action) == ERenderTargetLoadAction::ELoad)
+			{
+				// make sure we have dependency between multiple render-passes that use the same attachment
+				// otherwise GPU can be execute them in any order
+				VulkanRHI::ImagePipelineBarrier(CmdBuffer->GetHandle(), Surface.Image, EImageLayoutBarrier::ColorAttachment, EImageLayoutBarrier::ColorAttachment, SetupImageSubresourceRange());
+			}
 			else
 			{
 				Context.RHITransitionResources(EResourceTransitionAccess::EWritable, &Texture, 1);
