Fix for deprecated use of const-losing TWeakObjectPtr constructor.

diff --git a/Engine/Plugins/FX/Niagara/Source/NiagaraEditor/Private/ViewModels/NiagaraSystemViewModel.cpp b/Engine/Plugins/FX/Niagara/Source/NiagaraEditor/Private/ViewModels/NiagaraSystemViewModel.cpp
--- a/Engine/Plugins/FX/Niagara/Source/NiagaraEditor/Private/ViewModels/NiagaraSystemViewModel.cpp
+++ b/Engine/Plugins/FX/Niagara/Source/NiagaraEditor/Private/ViewModels/NiagaraSystemViewModel.cpp
@@ -678,7 +678,7 @@ void FNiagaraSystemViewModel::SendLastCompileMessageJobs() const
 				WarningCount++;
 			}
 
-			JobBatchToQueue.Add(MakeShared<FNiagaraMessageJobCompileEvent>(CompileEvent, TWeakObjectPtr<UNiagaraScript>(ScriptInfo.Script), ScriptInfo.OwningScriptNameString));
+			JobBatchToQueue.Add(MakeShared<FNiagaraMessageJobCompileEvent>(CompileEvent, MakeWeakObjectPtr(const_cast<UNiagaraScript*>(ScriptInfo.Script)), ScriptInfo.OwningScriptNameString));
 		}
 	}
 	JobBatchToQueue.Insert(MakeShared<FNiagaraMessageJobPostCompileSummary>(ErrorCount, WarningCount, GetLatestCompileStatus(), FText::FromString("System")), 0);
