BaseNiagara updated to deprecate SystemLifeCycle

diff --git a/Engine/Plugins/FX/Niagara/Config/BaseNiagara.ini b/Engine/Plugins/FX/Niagara/Config/BaseNiagara.ini
--- a/Engine/Plugins/FX/Niagara/Config/BaseNiagara.ini
+++ b/Engine/Plugins/FX/Niagara/Config/BaseNiagara.ini
@@ -29,7 +29,7 @@ DefaultScript=/Niagara/DefaultAssets/DefaultModule.DefaultModule
 DefaultDynamicInputScript=/Niagara/DefaultAssets/DefaultDynamicInputs.DefaultDynamicInputs
 DefaultFunctionScript=/Niagara/DefaultAssets/DefaultFunction.DefaultFunction
 DefaultModuleScript=/Niagara/DefaultAssets/DefaultModule.DefaultModule
-RequiredSystemUpdateScript=/Niagara/Modules/System/SystemLifeCycle.SystemLifeCycle
+RequiredSystemUpdateScript=/Niagara/Modules/System/SystemState.SystemState
 +GraphCreationShortcuts=(Name="Numeric::Add",Input=(Key=A,bShift=False,bCtrl=False,bAlt=False,bCmd=False))
 +GraphCreationShortcuts=(Name="Numeric::Div",Input=(Key=D,bShift=False,bCtrl=False,bAlt=False,bCmd=False))
 +GraphCreationShortcuts=(Name="Numeric::Pow",Input=(Key=E,bShift=False,bCtrl=False,bAlt=False,bCmd=False))
