- Fixed a crash where SystemSimulation::RemoveInstance will wait for the async task to complete and the finalize would complete the system also calling RemoveInstance, so the original call would have a SystemInstanceIndex of -1
- Separated WaitForAsyncTick into WaitForAsyncTickAndFinalize & WaitForAsyncTickDoNotFinalize and added more detailed comments

#rnx

diff --git a/Engine/Plugins/FX/Niagara/Source/Niagara/Private/NiagaraComponent.cpp b/Engine/Plugins/FX/Niagara/Source/Niagara/Private/NiagaraComponent.cpp
--- a/Engine/Plugins/FX/Niagara/Source/Niagara/Private/NiagaraComponent.cpp
+++ b/Engine/Plugins/FX/Niagara/Source/Niagara/Private/NiagaraComponent.cpp
@@ -1174,7 +1174,7 @@ void UNiagaraComponent::OnEndOfFrameUpdateDuringTick()
 	Super::OnEndOfFrameUpdateDuringTick();
 	if ( SystemInstance )
 	{
-		SystemInstance->WaitForAsyncTick();
+		SystemInstance->WaitForAsyncTickAndFinalize();
 	}
 }
 
diff --git a/Engine/Plugins/FX/Niagara/Source/Niagara/Private/NiagaraSystemInstance.cpp b/Engine/Plugins/FX/Niagara/Source/Niagara/Private/NiagaraSystemInstance.cpp
--- a/Engine/Plugins/FX/Niagara/Source/Niagara/Private/NiagaraSystemInstance.cpp
+++ b/Engine/Plugins/FX/Niagara/Source/Niagara/Private/NiagaraSystemInstance.cpp
@@ -112,7 +112,7 @@ FNiagaraSystemInstance::FNiagaraSystemInstance(UNiagaraComponent* InComponent)
 
 void FNiagaraSystemInstance::SetEmitterEnable(FName EmitterName, bool bNewEnableState)
 {
-	WaitForAsyncTick();
+	WaitForAsyncTickAndFinalize();
 
 	UNiagaraSystem* System = GetSystem();
 	if (System != nullptr)
@@ -160,7 +160,7 @@ void FNiagaraSystemInstance::SetEmitterEnable(FName EmitterName, bool bNewEnable
 
 void FNiagaraSystemInstance::Init(bool bInForceSolo)
 {
-	WaitForAsyncTick(true);
+	WaitForAsyncTickAndFinalize(true);
 
 	bForceSolo = bInForceSolo;
 	ActualExecutionState = ENiagaraExecutionState::Inactive;
@@ -240,7 +240,7 @@ void FNiagaraSystemInstance::Dump()const
 
 void FNiagaraSystemInstance::DumpTickInfo(FOutputDevice& Ar)
 {
-	WaitForAsyncTick();
+	WaitForAsyncTickAndFinalize();
 
 	static const UEnum* TickingGroupEnum = FindObjectChecked<UEnum>(ANY_PACKAGE, TEXT("ETickingGroup"));
 
@@ -274,7 +274,7 @@ bool FNiagaraSystemInstance::RequestCapture(const FGuid& RequestId)
 		return false;
 	}
 
-	WaitForAsyncTick();
+	WaitForAsyncTickAndFinalize();
 
 	UE_LOG(LogNiagara, Warning, TEXT("Capture requested!"));
 
@@ -312,7 +312,7 @@ bool FNiagaraSystemInstance::RequestCapture(const FGuid& RequestId)
 
 void FNiagaraSystemInstance::FinishCapture()
 {
-	WaitForAsyncTick();
+	WaitForAsyncTickAndFinalize();
 
 	if (!CurrentCapture.IsValid())
 	{
@@ -326,7 +326,7 @@ void FNiagaraSystemInstance::FinishCapture()
 
 bool FNiagaraSystemInstance::QueryCaptureResults(const FGuid& RequestId, TArray<TSharedPtr<struct FNiagaraScriptDebuggerInfo, ESPMode::ThreadSafe>>& OutCaptureResults)
 {
-	WaitForAsyncTick();
+	WaitForAsyncTickAndFinalize();
 
 	if (CurrentCaptureGuid.IsValid() && RequestId == *CurrentCaptureGuid.Get())
 	{
@@ -408,7 +408,7 @@ void FNiagaraSystemInstance::SetSolo(bool bInSolo)
 		return;
 	}
 
-	WaitForAsyncTick();
+	WaitForAsyncTickAndFinalize();
 
 	UNiagaraSystem* System = GetSystem();
 	if (bInSolo)
@@ -443,7 +443,7 @@ void FNiagaraSystemInstance::Activate(EResetMode InResetMode)
 {
 	SCOPE_CYCLE_COUNTER(STAT_NiagaraSystemActivate);
 
-	WaitForAsyncTick();
+	WaitForAsyncTickAndFinalize();
 
 	UNiagaraSystem* System = GetSystem();
 	if (System && System->IsValid() && IsReadyToRun())
@@ -460,7 +460,7 @@ void FNiagaraSystemInstance::Deactivate(bool bImmediate)
 {
 	SCOPE_CYCLE_COUNTER(STAT_NiagaraSystemDeactivate);
 
-	WaitForAsyncTick();
+	WaitForAsyncTickAndFinalize();
 
 	if (IsComplete())
 	{
@@ -566,7 +566,7 @@ void FNiagaraSystemInstance::SetPaused(bool bInPaused)
 		return;
 	}
 
-	WaitForAsyncTick();
+	WaitForAsyncTickAndFinalize();
 	
 	if (SystemInstanceIndex != INDEX_NONE)
 	{
@@ -602,7 +602,7 @@ void FNiagaraSystemInstance::Reset(FNiagaraSystemInstance::EResetMode Mode)
 		return;
 	}
 
-	WaitForAsyncTick();
+	WaitForAsyncTickAndFinalize();
 
 	if (Component)
 	{
@@ -751,7 +751,7 @@ void FNiagaraSystemInstance::AdvanceSimulation(int32 TickCountToSimulate, float
 {
 	if (TickCountToSimulate > 0 && !IsPaused())
 	{
-		WaitForAsyncTick();
+		WaitForAsyncTickAndFinalize();
 
 		SCOPE_CYCLE_COUNTER(STAT_NiagaraSystemAdvanceSim);
 		bool bWasSolo = bSolo;
@@ -975,7 +975,7 @@ FNiagaraSystemInstance::~FNiagaraSystemInstance()
 
 void FNiagaraSystemInstance::Cleanup()
 {
-	WaitForAsyncTick(true);
+	WaitForAsyncTickAndFinalize(true);
 
 	if (SystemInstanceIndex != INDEX_NONE)
 	{
@@ -1163,7 +1163,7 @@ void FNiagaraSystemInstance::InitDataInterfaces()
 		return;
 	}
 
-	WaitForAsyncTick(true);
+	WaitForAsyncTickAndFinalize(true);
 
 	Component->GetOverrideParameters().Tick();
 	
@@ -1767,28 +1767,46 @@ void FNiagaraSystemInstance::ComponentTick(float DeltaSeconds, const FGraphEvent
 	SystemSim->Tick_GameThread(DeltaSeconds, MyCompletionGraphEvent);
 }
 
-void FNiagaraSystemInstance::WaitForAsyncTick(bool bEnsureComplete)
+void FNiagaraSystemInstance::WaitForAsyncTickDoNotFinalize(bool bEnsureComplete)
 {
-	if (bAsyncWorkInProgress)
+	if (bAsyncWorkInProgress == false)
 	{
-		ensureAlwaysMsgf(!bEnsureComplete, TEXT("Niagara System Async Task should be complete by now. %s"), *GetSystem()->GetPathName());
-		ensureAlwaysMsgf(IsInGameThread(), TEXT("NiagaraSystemInstance::WaitForAsyncTick() call is assuming execution on GT but is not on GT. %s"), *GetSystem()->GetPathName());
+		return;
+	}
 
-		SCOPE_CYCLE_COUNTER(STAT_NiagaraSystemWaitForAsyncTick);
+	ensureAlwaysMsgf(!bEnsureComplete, TEXT("Niagara System Async Task should be complete by now. %s"), *GetSystem()->GetPathName());
+	ensureAlwaysMsgf(IsInGameThread(), TEXT("NiagaraSystemInstance::WaitForAsyncTick() call is assuming execution on GT but is not on GT. %s"), *GetSystem()->GetPathName());
 
-		const uint64 StartCycles = FPlatformTime::Cycles64();
-		while (bAsyncWorkInProgress)
-		{
-			FPlatformProcess::SleepNoStats(0.0f);
-		}
+	SCOPE_CYCLE_COUNTER(STAT_NiagaraSystemWaitForAsyncTick);
+
+	const uint64 StartCycles = FPlatformTime::Cycles64();
+	const double WarnSeconds = 5.0;
+	const uint64 WarnCycles = StartCycles + uint64(WarnSeconds / FPlatformTime::GetSecondsPerCycle());
+	bool bDoWarning = true;
 
-		const double StallTimeMS = FPlatformTime::ToMilliseconds64(FPlatformTime::Cycles64() - StartCycles);
-		if (StallTimeMS > GWaitForAsyncStallWarnThresholdMS)
+	while ( bAsyncWorkInProgress )
+	{
+		FPlatformProcess::SleepNoStats(0.001f);
+		if ( bDoWarning && (FPlatformTime::Cycles64() > WarnCycles) )
 		{
-			//-TODO: This should be put back to a warning once EngineTests no longer cause it show up.  The reason it's triggered is because we pause in latent actions right after a TG running Niagara sims.
-			UE_LOG(LogNiagara, Log, TEXT("Niagara Effect stalled GT for %g ms.\nComponent: %s \nSystem: %s"), StallTimeMS, *GetFullNameSafe(Component), *GetFullNameSafe(GetSystem()));
+			bDoWarning = false;
+			UE_LOG(LogNiagara, Warning, TEXT("Niagara Effect has stalled GT for %g and is not complete, this may result in a deadlock.\nComponent: %s \nSystem: %s"), WarnSeconds, *GetFullNameSafe(Component), *GetFullNameSafe(GetSystem()));
 		}
+	}
 
+	const double StallTimeMS = FPlatformTime::ToMilliseconds64(FPlatformTime::Cycles64() - StartCycles);
+	if (StallTimeMS > GWaitForAsyncStallWarnThresholdMS)
+	{
+		//-TODO: This should be put back to a warning once EngineTests no longer cause it show up.  The reason it's triggered is because we pause in latent actions right after a TG running Niagara sims.
+		UE_LOG(LogNiagara, Log, TEXT("Niagara Effect stalled GT for %g ms.\nComponent: %s \nSystem: %s"), StallTimeMS, *GetFullNameSafe(Component), *GetFullNameSafe(GetSystem()));
+	}
+}
+
+void FNiagaraSystemInstance::WaitForAsyncTickAndFinalize(bool bEnsureComplete)
+{
+	if (bAsyncWorkInProgress)
+	{
+		WaitForAsyncTickDoNotFinalize(bEnsureComplete);
 		FinalizeTick_GameThread();
 	}
 }
@@ -1826,7 +1844,7 @@ void FNiagaraSystemInstance::Tick_GameThread(float DeltaSeconds)
 	UNiagaraSystem* System = GetSystem();
 	FScopeCycleCounter SystemStat(System->GetStatID(true, false));
 
-	WaitForAsyncTick(true);
+	WaitForAsyncTickAndFinalize(true);
 
 	CachedDeltaSeconds = DeltaSeconds;
 	bNeedsFinalize = true;
diff --git a/Engine/Plugins/FX/Niagara/Source/Niagara/Private/NiagaraSystemSimulation.cpp b/Engine/Plugins/FX/Niagara/Source/Niagara/Private/NiagaraSystemSimulation.cpp
--- a/Engine/Plugins/FX/Niagara/Source/Niagara/Private/NiagaraSystemSimulation.cpp
+++ b/Engine/Plugins/FX/Niagara/Source/Niagara/Private/NiagaraSystemSimulation.cpp
@@ -1028,7 +1028,7 @@ void FNiagaraSystemSimulation::WaitForInstancesTickComplete(bool bEnsureComplete
 	{
 		// If we're in a spawn phase all existing instances should be complete already.
 		FNiagaraSystemInstance* Inst = SystemInstances[SystemInstIndex];
-		Inst->WaitForAsyncTick(bInSpawnPhase);
+		Inst->WaitForAsyncTickAndFinalize(bInSpawnPhase);
 
 		// If the system completes during finalize it can be removed from instances so we don't update the index.
 		if ( SystemInstances[SystemInstIndex] == Inst)
@@ -1713,12 +1713,12 @@ void FNiagaraSystemSimulation::RemoveInstance(FNiagaraSystemInstance* Instance)
 		// Wait for the system simulation & the system instances tick to complete as we are touching both the SystemInstances & DataSet
 		// Note: We do not need to wait for all instances to complete as the system simulation concurrent tick will have transfered data from the DataSet out to ParameterStores
 		WaitForSystemTickComplete();
-		Instance->WaitForAsyncTick();
+		Instance->WaitForAsyncTickDoNotFinalize();
 
-		int32 NumInstances = MainDataSet.GetCurrentDataChecked().GetNumInstances();
+		const int32 NumInstances = MainDataSet.GetCurrentDataChecked().GetNumInstances();
 		check(SystemInstances.Num() == NumInstances);
 
-		int32 SystemIndex = Instance->SystemInstanceIndex;
+		const int32 SystemIndex = Instance->SystemInstanceIndex;
 		check(Instance == SystemInstances[SystemIndex]);
 		check(SystemInstances.IsValidIndex(SystemIndex));
 
diff --git a/Engine/Plugins/FX/Niagara/Source/Niagara/Public/NiagaraSystemInstance.h b/Engine/Plugins/FX/Niagara/Source/Niagara/Public/NiagaraSystemInstance.h
--- a/Engine/Plugins/FX/Niagara/Source/Niagara/Public/NiagaraSystemInstance.h
+++ b/Engine/Plugins/FX/Niagara/Source/Niagara/Public/NiagaraSystemInstance.h
@@ -95,8 +95,19 @@ public:
 	/** Final phase of system instance tick. Must be executed on the game thread. */
 	void FinalizeTick_GameThread();
 
-	/** Blocks until any async work for this system instance has completed. Must be called on the game thread. */
-	void WaitForAsyncTick(bool bEnsureComplete=false);
+	/**
+		Blocks until any async work for this system instance has completed, must be called on the GameThread.
+		This will NOT call finalize on the instance, be very careful when using to avoid leaving the instance in an undefined state.
+		Note: This only waits for the instance to be safe to touch, it does not wait for the owning system simulation to be safe.
+	*/
+	void WaitForAsyncTickDoNotFinalize(bool bEnsureComplete = false);
+
+	/**
+		Blocks until any async work for this system instance has completed, must be called on the GameThread.
+		This will call finalize if required by the instance and can therefore complete leaving removing the instance from the owning system simulation.
+		Note: This only waits for the instance to be safe to touch, it does not wait for the owning system simulation to be safe.
+	*/
+	void WaitForAsyncTickAndFinalize(bool bEnsureComplete = false);
 
 	/** Handles completion of the system and returns true if the system is complete. */
 	bool HandleCompletion();
