- Don't enqueue a GPU tick if we have no component attached
- The component pointer will be nullptr once FNiagaraSystemInstance::DeallocateSystemInstance has been called in which case we don't want to tick the GPU anymore as we will not render the results

#rnx

diff --git a/Engine/Plugins/FX/Niagara/Source/Niagara/Private/NiagaraSystemInstance.cpp b/Engine/Plugins/FX/Niagara/Source/Niagara/Private/NiagaraSystemInstance.cpp
--- a/Engine/Plugins/FX/Niagara/Source/Niagara/Private/NiagaraSystemInstance.cpp
+++ b/Engine/Plugins/FX/Niagara/Source/Niagara/Private/NiagaraSystemInstance.cpp
@@ -1949,7 +1949,6 @@ void FNiagaraSystemInstance::FinalizeTick_GameThread()
 			//Post tick our interfaces.
 			TickDataInterfaces(CachedDeltaSeconds, true);
 
-			bool bIsRegistered = true;
 			if (Component)
 			{
 				if ( bIsTransformDirty )
@@ -1958,39 +1957,38 @@ void FNiagaraSystemInstance::FinalizeTick_GameThread()
 					Component->UpdateComponentToWorld();
 				}
 				Component->MarkRenderDynamicDataDirty();
-				bIsRegistered = Component->IsRegistered();
-			}
 
-			// Push any GPU ticks for this system instance.
-			if (ActiveGPUEmitterCount > 0 && NiagaraSupportsComputeShaders(GetBatcher()->GetShaderPlatform()) && bIsRegistered)
-			{
-				ensure(!IsComplete());
-				FNiagaraGPUSystemTick GPUTick;
-				GPUTick.Init(this);
-
-				//if (GPUTick.DIInstanceData)
-				//{
-				//	uint8* BasePointer = (uint8*)GPUTick.DIInstanceData->PerInstanceDataForRT;
-
-				//	//UE_LOG(LogNiagara, Log, TEXT("GT Testing (dipacket) %p (baseptr) %p"), GPUTick.DIInstanceData, BasePointer);
-				//	for (auto& Pair : GPUTick.DIInstanceData->InterfaceProxiesToOffsets)
-				//	{
-				//		FNiagaraDataInterfaceProxy* Proxy = Pair.Key;
-				//		UE_LOG(LogNiagara, Log, TEXT("\tGT (proxy) %p (size) %u"), Proxy, Proxy->PerInstanceDataPassedToRenderThreadSize());
-				//	}
-				//}
-
-				// We will give the data over to the render thread. It is responsible for freeing it.
-				// We no longer own it and cannot modify it after this point.
-				// @todo We are taking a copy of the object here. This object is small so this overhead should
-				// not be very high. And we avoid making a bunch of small allocations here.
-				NiagaraEmitterInstanceBatcher* TheBatcher = GetBatcher();
-				ENQUEUE_RENDER_COMMAND(FGiveSystemInstanceTickToRT)(
-					[TheBatcher, GPUTick](FRHICommandListImmediate& RHICmdList) mutable
+				// Push any GPU ticks for this system instance.
+				if (ActiveGPUEmitterCount > 0 && NiagaraSupportsComputeShaders(GetBatcher()->GetShaderPlatform()) && Component->IsRegistered())
 				{
-					TheBatcher->GiveSystemTick_RenderThread(GPUTick);
+					ensure(!IsComplete());
+					FNiagaraGPUSystemTick GPUTick;
+					GPUTick.Init(this);
+
+					//if (GPUTick.DIInstanceData)
+					//{
+					//	uint8* BasePointer = (uint8*)GPUTick.DIInstanceData->PerInstanceDataForRT;
+
+					//	//UE_LOG(LogNiagara, Log, TEXT("GT Testing (dipacket) %p (baseptr) %p"), GPUTick.DIInstanceData, BasePointer);
+					//	for (auto& Pair : GPUTick.DIInstanceData->InterfaceProxiesToOffsets)
+					//	{
+					//		FNiagaraDataInterfaceProxy* Proxy = Pair.Key;
+					//		UE_LOG(LogNiagara, Log, TEXT("\tGT (proxy) %p (size) %u"), Proxy, Proxy->PerInstanceDataPassedToRenderThreadSize());
+					//	}
+					//}
+
+					// We will give the data over to the render thread. It is responsible for freeing it.
+					// We no longer own it and cannot modify it after this point.
+					// @todo We are taking a copy of the object here. This object is small so this overhead should
+					// not be very high. And we avoid making a bunch of small allocations here.
+					NiagaraEmitterInstanceBatcher* TheBatcher = GetBatcher();
+					ENQUEUE_RENDER_COMMAND(FNiagaraGiveSystemInstanceTickToRT)(
+						[TheBatcher, GPUTick](FRHICommandListImmediate& RHICmdList) mutable
+						{
+							TheBatcher->GiveSystemTick_RenderThread(GPUTick);
+						}
+					);
 				}
-				);
 			}
 		}
 	}
