- Fix a deadlock where a Niagara System Instance is transfered into a new system simulation after it was ticked then destroyed before it's been ticked again
- The instance was left in a state where async work is flagged to complete (due to SpawnNew_GameThread calling Tick_GameThread)

#rnx

diff --git a/Engine/Plugins/FX/Niagara/Source/Niagara/Private/NiagaraSystemSimulation.cpp b/Engine/Plugins/FX/Niagara/Source/Niagara/Private/NiagaraSystemSimulation.cpp
--- a/Engine/Plugins/FX/Niagara/Source/Niagara/Private/NiagaraSystemSimulation.cpp
+++ b/Engine/Plugins/FX/Niagara/Source/Niagara/Private/NiagaraSystemSimulation.cpp
@@ -259,7 +259,6 @@ FNiagaraSystemSimulation::FNiagaraSystemSimulation()
 	, bBindingsInitialized(false)
 	, bInSpawnPhase(false)
 	, bIsSolo(false)
-	, bHasEverTicked(false)
 {
 
 }
@@ -284,7 +283,6 @@ bool FNiagaraSystemSimulation::Init(UNiagaraSystem* InSystem, UWorld* InWorld, b
 
 	bBindingsInitialized = false;
 	bInSpawnPhase = false;
-	bHasEverTicked = false;
 
 	FNiagaraWorldManager* WorldMan = FNiagaraWorldManager::Get(InWorld);
 	check(WorldMan);
@@ -622,6 +620,8 @@ void FNiagaraSystemSimulation::FlushTickBatch(FNiagaraSystemSimulationTickContex
 void FNiagaraSystemSimulation::Tick_GameThread(float DeltaSeconds, const FGraphEventRef& MyCompletionGraphEvent)
 {
 	check(IsInGameThread());
+	check(bInSpawnPhase == false);
+
 	WaitForSystemTickComplete(true);
 
 	SCOPE_CYCLE_COUNTER(STAT_NiagaraOverview_GT);
@@ -635,8 +635,6 @@ void FNiagaraSystemSimulation::Tick_GameThread(float DeltaSeconds, const FGraphE
 	PARTICLE_PERF_STAT_INSTANCE_COUNT(System, SystemInstances.Num());
 	PARTICLE_PERF_STAT_CYCLES(System, TickGameThread);
 
-	bHasEverTicked = true;
-
 	SystemTickGraphEvent = nullptr;
 
 	check(SystemInstances.Num() == MainDataSet.GetCurrentDataChecked().GetNumInstances());
@@ -655,7 +653,7 @@ void FNiagaraSystemSimulation::Tick_GameThread(float DeltaSeconds, const FGraphE
 	SystemUpdateScript->RapidIterationParameters.Tick();
 #endif
 
-	const bool bUpdateTickGroups = !bIsSolo && !bInSpawnPhase;
+	const bool bUpdateTickGroups = !bIsSolo;
 
 	// Update instances
 	int32 SystemIndex = 0;
@@ -722,130 +720,112 @@ void FNiagaraSystemSimulation::Tick_GameThread(float DeltaSeconds, const FGraphE
 		}
 	}
 
-	//Setup the few real constants like delta time.
-	const float InvDt = 1.0f / DeltaSeconds;
-
-	SpawnTimeParam.SetValue(World->TimeSeconds);
-	UpdateTimeParam.SetValue(World->TimeSeconds);
-	SpawnDeltaTimeParam.SetValue(DeltaSeconds);
-	UpdateDeltaTimeParam.SetValue(DeltaSeconds);
-	SpawnInvDeltaTimeParam.SetValue(InvDt);
-	UpdateInvDeltaTimeParam.SetValue(InvDt);
-	SpawnNumSystemInstancesParam.SetValue(SystemInstances.Num());
-	UpdateNumSystemInstancesParam.SetValue(SystemInstances.Num());
-	SpawnGlobalSpawnCountScaleParam.SetValue(INiagaraModule::GetGlobalSpawnCountScale());
-	UpdateGlobalSpawnCountScaleParam.SetValue(INiagaraModule::GetGlobalSpawnCountScale());
-	SpawnGlobalSystemCountScaleParam.SetValue(INiagaraModule::GetGlobalSystemCountScale());
-	UpdateGlobalSystemCountScaleParam.SetValue(INiagaraModule::GetGlobalSystemCountScale());
+	// Setup the few real constants like delta time.
+	SetupParameters_GameThread(DeltaSeconds);
 
-	// Somethings we don't want to happen during the spawn phase
-	if ( !bInSpawnPhase )
+	// Gather any pending spawn systems and add to the end of the system instances
+	int32 SpawnNum = 0;
+	if (PendingSystemInstances.Num() > 0)
 	{
-		int32 SpawnNum = 0;
+		SystemInstances.Reserve(SystemInstances.Num() + PendingSystemInstances.Num());
 
-		// Gather any pending spawn systems and add to the end of the system instances
-		if (PendingSystemInstances.Num() > 0)
+		SystemIndex = 0;
+		while (SystemIndex < PendingSystemInstances.Num())
 		{
-			SystemInstances.Reserve(SystemInstances.Num() + PendingSystemInstances.Num());
+			FNiagaraSystemInstance* Inst = PendingSystemInstances[SystemIndex];
 
-			SystemIndex = 0;
-			while (SystemIndex < PendingSystemInstances.Num())
+			// If we are paused continue
+			if (Inst->IsPaused())
 			{
-				FNiagaraSystemInstance* Inst = PendingSystemInstances[SystemIndex];
-
-				// If we are paused continue
-				if (Inst->IsPaused())
-				{
-					++SystemIndex;
-					continue;
-				}
+				++SystemIndex;
+				continue;
+			}
 
-				// If our tick group has changed we need to move this pending instance
-				// Note we don't do this with solo instances
-				if (!bIsSolo)
+			// If our tick group has changed we need to move this pending instance
+			// Note we don't do this with solo instances
+			if (!bIsSolo)
+			{
+				const ETickingGroup DesiredTickGroup = Inst->CalculateTickGroup();
+				if (DesiredTickGroup != SystemTickGroup)
 				{
-					const ETickingGroup DesiredTickGroup = Inst->CalculateTickGroup();
-					if (DesiredTickGroup != SystemTickGroup)
+					PendingSystemInstances.RemoveAtSwap(SystemIndex);
+					if (PendingSystemInstances.IsValidIndex(SystemIndex))
 					{
-						PendingSystemInstances.RemoveAtSwap(SystemIndex);
-						if (PendingSystemInstances.IsValidIndex(SystemIndex))
-						{
-							PendingSystemInstances[SystemIndex]->SystemInstanceIndex = SystemIndex;
-						}
-						Inst->SystemInstanceIndex = INDEX_NONE;
+						PendingSystemInstances[SystemIndex]->SystemInstanceIndex = SystemIndex;
+					}
+					Inst->SystemInstanceIndex = INDEX_NONE;
 
-						FNiagaraWorldManager* WorldManager = FNiagaraWorldManager::Get(World);
-						check(WorldManager != nullptr);
+					FNiagaraWorldManager* WorldManager = FNiagaraWorldManager::Get(World);
+					check(WorldManager != nullptr);
 
-						TSharedPtr<FNiagaraSystemSimulation, ESPMode::ThreadSafe> DestSim = WorldManager->GetSystemSimulation(DesiredTickGroup, System);
+					TSharedPtr<FNiagaraSystemSimulation, ESPMode::ThreadSafe> DestSim = WorldManager->GetSystemSimulation(DesiredTickGroup, System);
 
-						Inst->SystemSimulation = DestSim;
-						Inst->SystemInstanceIndex = DestSim->PendingSystemInstances.Add(Inst);
-						continue;
-					}
+					Inst->SystemSimulation = DestSim;
+					Inst->SystemInstanceIndex = DestSim->PendingSystemInstances.Add(Inst);
+					continue;
 				}
+			}
 
-				// Execute instance tick
-				Inst->Tick_GameThread(DeltaSeconds);
-
-				if (Inst->SystemInstanceIndex != INDEX_NONE)
-				{
-					// We should not move tick group during Tick_GameThread but let's be safe
-					check(Inst->SystemSimulation.Get() == this);
+			// Execute instance tick
+			Inst->Tick_GameThread(DeltaSeconds);
 
-					// When the first instance is added we need to initialize the parameter store to data set bindings.
-					if (!bBindingsInitialized)
-					{
-						InitParameterDataSetBindings(Inst);
-					}
+			if (Inst->SystemInstanceIndex != INDEX_NONE)
+			{
+				// We should not move tick group during Tick_GameThread but let's be safe
+				check(Inst->SystemSimulation.Get() == this);
 
-					check(PendingSystemInstances[SystemIndex] == Inst);
-					PendingSystemInstances.RemoveAtSwap(SystemIndex);
-					if (PendingSystemInstances.IsValidIndex(SystemIndex))
-					{
-						PendingSystemInstances[SystemIndex]->SystemInstanceIndex = SystemIndex;
-					}
+				// When the first instance is added we need to initialize the parameter store to data set bindings.
+				if (!bBindingsInitialized)
+				{
+					InitParameterDataSetBindings(Inst);
+				}
 
-					Inst->SystemInstanceIndex = SystemInstances.Add(Inst);
-					Inst->SetPendingSpawn(false);
-					++SpawnNum;
+				check(PendingSystemInstances[SystemIndex] == Inst);
+				PendingSystemInstances.RemoveAtSwap(SystemIndex);
+				if (PendingSystemInstances.IsValidIndex(SystemIndex))
+				{
+					PendingSystemInstances[SystemIndex]->SystemInstanceIndex = SystemIndex;
 				}
+
+				Inst->SystemInstanceIndex = SystemInstances.Add(Inst);
+				Inst->SetPendingSpawn(false);
+				++SpawnNum;
 			}
 		}
+	}
 
-		static const auto EffectsQualityCVar = IConsoleManager::Get().FindConsoleVariable(TEXT("sg.EffectsQuality"));
-		FNiagaraSystemSimulationTickContext Context(this, SystemInstances, MainDataSet, DeltaSeconds, SpawnNum, EffectsQualityCVar->GetInt(), MyCompletionGraphEvent);
+	static const auto EffectsQualityCVar = IConsoleManager::Get().FindConsoleVariable(TEXT("sg.EffectsQuality"));
+	FNiagaraSystemSimulationTickContext Context(this, SystemInstances, MainDataSet, DeltaSeconds, SpawnNum, EffectsQualityCVar->GetInt(), MyCompletionGraphEvent);
 
-		//Now kick of the concurrent tick.
-		if (Context.bTickAsync)
-		{
-			TGraphTask<FNiagaraSystemSimulationWaitAllFinalizeTask>* WaitAllFinalizeTask = TGraphTask<FNiagaraSystemSimulationWaitAllFinalizeTask>::CreateTask(nullptr, ENamedThreads::GameThread).ConstructAndHold(Context.FinalizeEvents);
-			FGraphEventRef FinalizeGraphEvent = WaitAllFinalizeTask->GetCompletionEvent();
+	//Now kick of the concurrent tick.
+	if (Context.bTickAsync)
+	{
+		TGraphTask<FNiagaraSystemSimulationWaitAllFinalizeTask>* WaitAllFinalizeTask = TGraphTask<FNiagaraSystemSimulationWaitAllFinalizeTask>::CreateTask(nullptr, ENamedThreads::GameThread).ConstructAndHold(Context.FinalizeEvents);
+		FGraphEventRef FinalizeGraphEvent = WaitAllFinalizeTask->GetCompletionEvent();
 
-			TGraphTask<FNiagaraSystemSimulationTickTask>* SimulationTickTask = TGraphTask<FNiagaraSystemSimulationTickTask>::CreateTask(nullptr, ENamedThreads::GameThread).ConstructAndHold(Context, WaitAllFinalizeTask);
-			SystemTickGraphEvent = SimulationTickTask->GetCompletionEvent();
-			Context.FinalizeEvents->Add(SystemTickGraphEvent);
+		TGraphTask<FNiagaraSystemSimulationTickTask>* SimulationTickTask = TGraphTask<FNiagaraSystemSimulationTickTask>::CreateTask(nullptr, ENamedThreads::GameThread).ConstructAndHold(Context, WaitAllFinalizeTask);
+		SystemTickGraphEvent = SimulationTickTask->GetCompletionEvent();
+		Context.FinalizeEvents->Add(SystemTickGraphEvent);
 
-			MyCompletionGraphEvent->SetGatherThreadForDontCompleteUntil(ENamedThreads::GameThread);
-			MyCompletionGraphEvent->DontCompleteUntil(FinalizeGraphEvent);
+		MyCompletionGraphEvent->SetGatherThreadForDontCompleteUntil(ENamedThreads::GameThread);
+		MyCompletionGraphEvent->DontCompleteUntil(FinalizeGraphEvent);
 
-			SimulationTickTask->Unlock(ENamedThreads::GameThread);
-		}
-		else
+		SimulationTickTask->Unlock(ENamedThreads::GameThread);
+	}
+	else
+	{
+		TGraphTask<FNiagaraSystemSimulationWaitAllFinalizeTask>* WaitAllFinalizeTask = nullptr;
+		if (Context.bTickInstancesAsync)
 		{
-			TGraphTask<FNiagaraSystemSimulationWaitAllFinalizeTask>* WaitAllFinalizeTask = nullptr;
-			if (Context.bTickInstancesAsync)
-			{
-				WaitAllFinalizeTask = TGraphTask<FNiagaraSystemSimulationWaitAllFinalizeTask>::CreateTask(nullptr, ENamedThreads::GameThread).ConstructAndHold(Context.FinalizeEvents);
-			}
+			WaitAllFinalizeTask = TGraphTask<FNiagaraSystemSimulationWaitAllFinalizeTask>::CreateTask(nullptr, ENamedThreads::GameThread).ConstructAndHold(Context.FinalizeEvents);
+		}
 
-			Tick_Concurrent(Context);
+		Tick_Concurrent(Context);
 
-			if (Context.bTickInstancesAsync)
-			{
-				WaitAllFinalizeTask->Unlock(ENamedThreads::GameThread);
-				Context.FinalizeEvents = nullptr;
-			}
+		if (Context.bTickInstancesAsync)
+		{
+			WaitAllFinalizeTask->Unlock(ENamedThreads::GameThread);
+			Context.FinalizeEvents = nullptr;
 		}
 	}
 }
@@ -932,11 +912,18 @@ void FNiagaraSystemSimulation::Spawn_GameThread(float DeltaSeconds)
 
 	bInSpawnPhase = true;
 
-	if (bHasEverTicked == false)
+	if (MaxDeltaTime.IsSet())
 	{
-		Tick_GameThread(DeltaSeconds, nullptr);
+		DeltaSeconds = FMath::Clamp(DeltaSeconds, 0.0f, MaxDeltaTime.GetValue());
 	}
 
+#if WITH_EDITOR
+	System->GetSystemSpawnScript()->RapidIterationParameters.Tick();
+	System->GetSystemUpdateScript()->RapidIterationParameters.Tick();
+#endif
+
+	SetupParameters_GameThread(DeltaSeconds);
+
 	FNiagaraScopedRuntimeCycleCounter RuntimeScope(System, true, false);
 
 	// Spawn instances
@@ -1343,6 +1330,25 @@ void FNiagaraSystemSimulation::TickFastPath(FNiagaraSystemSimulationTickContext&
 	}
 }
 
+void FNiagaraSystemSimulation::SetupParameters_GameThread(float DeltaSeconds)
+{
+	check(IsInGameThread());
+
+	const float InvDt = 1.0f / DeltaSeconds;
+	SpawnTimeParam.SetValue(World->TimeSeconds);
+	UpdateTimeParam.SetValue(World->TimeSeconds);
+	SpawnDeltaTimeParam.SetValue(DeltaSeconds);
+	UpdateDeltaTimeParam.SetValue(DeltaSeconds);
+	SpawnInvDeltaTimeParam.SetValue(InvDt);
+	UpdateInvDeltaTimeParam.SetValue(InvDt);
+	SpawnNumSystemInstancesParam.SetValue(SystemInstances.Num());
+	UpdateNumSystemInstancesParam.SetValue(SystemInstances.Num());
+	SpawnGlobalSpawnCountScaleParam.SetValue(INiagaraModule::GetGlobalSpawnCountScale());
+	UpdateGlobalSpawnCountScaleParam.SetValue(INiagaraModule::GetGlobalSpawnCountScale());
+	SpawnGlobalSystemCountScaleParam.SetValue(INiagaraModule::GetGlobalSystemCountScale());
+	UpdateGlobalSystemCountScaleParam.SetValue(INiagaraModule::GetGlobalSystemCountScale());
+}
+
 void FNiagaraSystemSimulation::PrepareForSystemSimulate(FNiagaraSystemSimulationTickContext& Context)
 {
 	SCOPE_CYCLE_COUNTER(STAT_NiagaraSystemSim_PrepareForSimulateCNC);
diff --git a/Engine/Plugins/FX/Niagara/Source/Niagara/Public/NiagaraSystemSimulation.h b/Engine/Plugins/FX/Niagara/Source/Niagara/Public/NiagaraSystemSimulation.h
--- a/Engine/Plugins/FX/Niagara/Source/Niagara/Public/NiagaraSystemSimulation.h
+++ b/Engine/Plugins/FX/Niagara/Source/Niagara/Public/NiagaraSystemSimulation.h
@@ -211,6 +211,9 @@ public:
 	void AddTickGroupPromotion(FNiagaraSystemInstance* Instance);
 
 protected:
+	/** Sets constant parameter values */
+	void SetupParameters_GameThread(float DeltaSeconds);
+
 	/** Does any prep work for system simulation such as pulling instance parameters into a dataset. */
 	void PrepareForSystemSimulate(FNiagaraSystemSimulationTickContext& Context);
 	/** Runs the system spawn script for new system instances. */
@@ -313,7 +316,6 @@ protected:
 	uint32 bBindingsInitialized : 1;
 	uint32 bInSpawnPhase : 1;
 	uint32 bIsSolo : 1;
-	uint32 bHasEverTicked : 1;
 
 	/** A parameter store which contains the data interfaces parameters which were defined by the scripts. */
 	FNiagaraParameterStore ScriptDefinedDataInterfaceParameters;
