Fixing Apply button not activating when changing defaults through the parameters panel

diff --git a/Engine/Plugins/FX/Niagara/Source/NiagaraEditor/Private/Customizations/NiagaraScriptVariableCustomization.cpp b/Engine/Plugins/FX/Niagara/Source/NiagaraEditor/Private/Customizations/NiagaraScriptVariableCustomization.cpp
--- a/Engine/Plugins/FX/Niagara/Source/NiagaraEditor/Private/Customizations/NiagaraScriptVariableCustomization.cpp
+++ b/Engine/Plugins/FX/Niagara/Source/NiagaraEditor/Private/Customizations/NiagaraScriptVariableCustomization.cpp
@@ -21,6 +21,8 @@
 #include "TypeEditorUtilities/NiagaraColorTypeEditorUtilities.h"
 #include "NiagaraEditorStyle.h"
 #include "ScopedTransaction.h"
+#include "NiagaraNode.h"
+#include "EdGraphSchema_Niagara.h"
 
 #define LOCTEXT_NAMESPACE "NiagaraScriptVariableVariableDetails"
 
@@ -110,6 +112,13 @@ void FNiagaraScriptVariableDetails::OnComboValueChanged()
 		TSharedRef<IPropertyUtilities> PropertyUtilities = DetailBuilderPtr->GetPropertyUtilities();
 		PropertyUtilities->ForceRefresh();
 	}
+
+#if WITH_EDITOR
+	if (UNiagaraGraph* Graph = Cast<UNiagaraGraph>(Variable->GetOuter()))
+	{
+		Graph->NotifyGraphNeedsRecompile();
+	}
+#endif	//#if WITH_EDITOR
 }
 
 void FNiagaraScriptVariableDetails::CustomizeDetails(IDetailLayoutBuilder& DetailBuilder)
@@ -247,6 +256,7 @@ void FNiagaraScriptVariableDetails::OnValueChanged()
 	{
 		if (UEdGraphPin* Pin = GetDefaultPin())
 		{
+			FString NewDefaultValue;
 			if (!ParameterEditorValue->CanChangeContinuously())
 			{
 				const FScopedTransaction Transaction( NSLOCTEXT("ScriptVariableCustomization", "ChangeValue", "Change Default Value" ) );
@@ -256,15 +266,17 @@ void FNiagaraScriptVariableDetails::OnValueChanged()
 				TSharedPtr<FStructOnScope> ParameterValue = MakeShareable(new FStructOnScope(Variable->Variable.GetType().GetStruct()));
 				ParameterEditorValue->UpdateStructFromInternalValue(ParameterValue.ToSharedRef());
 				Variable->Variable.SetData(ParameterValue->GetStructMemory());
-				Pin->DefaultValue = TypeUtilityValue->GetPinDefaultStringFromValue(Variable->Variable);
+				NewDefaultValue = TypeUtilityValue->GetPinDefaultStringFromValue(Variable->Variable);
 			}		
 			else
 			{
 				TSharedPtr<FStructOnScope> ParameterValue = MakeShareable(new FStructOnScope(Variable->Variable.GetType().GetStruct()));
 				ParameterEditorValue->UpdateStructFromInternalValue(ParameterValue.ToSharedRef());
 				Variable->Variable.SetData(ParameterValue->GetStructMemory());
-				Pin->DefaultValue = TypeUtilityValue->GetPinDefaultStringFromValue(Variable->Variable);	
+				NewDefaultValue = TypeUtilityValue->GetPinDefaultStringFromValue(Variable->Variable);
 			}
+
+			GetDefault<UEdGraphSchema_Niagara>()->TrySetDefaultValue(*Pin, NewDefaultValue, true);
 			
 		}
 	} 
@@ -288,7 +300,10 @@ void FNiagaraScriptVariableDetails::OnBeginValueChanged()
 			TSharedPtr<FStructOnScope> ParameterValue = MakeShareable(new FStructOnScope(Variable->Variable.GetType().GetStruct()));
 			ParameterEditorValue->UpdateStructFromInternalValue(ParameterValue.ToSharedRef());
 			Variable->Variable.SetData(ParameterValue->GetStructMemory());
-			Pin->DefaultValue = TypeUtilityValue->GetPinDefaultStringFromValue(Variable->Variable);
+
+			FString NewDefaultValue = TypeUtilityValue->GetPinDefaultStringFromValue(Variable->Variable);
+			GetDefault<UEdGraphSchema_Niagara>()->TrySetDefaultValue(*Pin, NewDefaultValue, true);
+
 		}
 	} 
 }
@@ -311,6 +326,13 @@ void FNiagaraScriptVariableDetails::OnStaticSwitchValueChanged()
 		ParameterEditorStaticSwitchValue->UpdateStructFromInternalValue(ParameterValue.ToSharedRef());
 		Variable->Variable.SetData(ParameterValue->GetStructMemory());
 		Variable->Metadata.StaticSwitchDefaultValue = Variable->Variable.GetValue<int>();
+		
+#if WITH_EDITOR
+		if (UNiagaraGraph* Graph = Cast<UNiagaraGraph>(Variable->GetOuter()))
+		{
+			Graph->NotifyGraphNeedsRecompile();
+		}
+#endif	//#if WITH_EDITOR
 	} 
 }
  
