Working on getting compile hashing working properly for Niagara script graphs.
+ Created Engine\Plugins\FX\Niagara\Shaders\Private\NiagaraShaderVersion.ush currently works in tandem with  FNiagaraCustomVersion::LatestScriptCompileVersion but is needed for desired in-editor workflow of making an edit to source files and having it update if proper console command is set. Might want to deprecate the old one in the future.
+ Adding console commands fx.InvalidateCachedScripts and fx.RebuildDirtyScripts similar to r.InvalidateCachedShaders and recompileshaders changed

diff --git a/Engine/Plugins/FX/Niagara/Shaders/Private/NiagaraShaderVersion.ush b/Engine/Plugins/FX/Niagara/Shaders/Private/NiagaraShaderVersion.ush
--- a/Engine/Plugins/FX/Niagara/Shaders/Private/NiagaraShaderVersion.ush
+++ b/Engine/Plugins/FX/Niagara/Shaders/Private/NiagaraShaderVersion.ush
@@ -3,4 +3,4 @@
 // in the DDC key, it will automatically invalidate.
 // 
 // If you are merging streams and there is a conflict with this GUID you should make a new GUID rather than taking one or the other.
-// GUID = 787C431B-53E1-4ACD-A12E-5F39EBE3F5DF
\ No newline at end of file
+// GUID = C0325EAC41B1CEA997E74A9708F15E84
\ No newline at end of file
diff --git a/Engine/Plugins/FX/Niagara/Source/Niagara/Classes/NiagaraDataInterface.h b/Engine/Plugins/FX/Niagara/Source/Niagara/Classes/NiagaraDataInterface.h
--- a/Engine/Plugins/FX/Niagara/Source/Niagara/Classes/NiagaraDataInterface.h
+++ b/Engine/Plugins/FX/Niagara/Source/Niagara/Classes/NiagaraDataInterface.h
@@ -211,6 +211,9 @@ public:
 	virtual bool PerInstanceTick(void* PerInstanceData, FNiagaraSystemInstance* SystemInstance, float DeltaSeconds) { return false; }
 	virtual bool PerInstanceTickPostSimulate(void* PerInstanceData, FNiagaraSystemInstance* SystemInstance, float DeltaSeconds) { return false; }
 
+	/** Allows the generic class defaults version of this class to specify any dependencies/version/etc that might invalidate the compile. It should never depend on the value of specific properties.*/
+	virtual bool AppendCompileHash(FNiagaraCompileHashVisitor* InVisitor) const;
+
 	/** 
 		Subclasses that wish to work with GPU systems/emitters must implement this.
 		Those interfaces must fill DataForRenderThread with the data needed to upload to the GPU. It will be the last thing called on this
diff --git a/Engine/Plugins/FX/Niagara/Source/Niagara/Classes/NiagaraDataInterfaceSkeletalMesh.h b/Engine/Plugins/FX/Niagara/Source/Niagara/Classes/NiagaraDataInterfaceSkeletalMesh.h
--- a/Engine/Plugins/FX/Niagara/Source/Niagara/Classes/NiagaraDataInterfaceSkeletalMesh.h
+++ b/Engine/Plugins/FX/Niagara/Source/Niagara/Classes/NiagaraDataInterfaceSkeletalMesh.h
@@ -528,6 +528,7 @@ public:
 #endif
 	virtual bool HasTickGroupPrereqs() const override { return true; }
 	virtual ETickingGroup CalculateTickGroup(void* PerInstanceData) const override;
+	virtual bool AppendCompileHash(FNiagaraCompileHashVisitor* InVisitor) const override;
 	//~ UNiagaraDataInterface interface END
 
 	USkeletalMesh* GetSkeletalMesh(class UNiagaraComponent* OwningComponent, TWeakObjectPtr<USceneComponent>& SceneComponent, USkeletalMeshComponent*& FoundSkelComp, FNDISkeletalMesh_InstanceData* InstData = nullptr);
diff --git a/Engine/Plugins/FX/Niagara/Source/Niagara/Classes/NiagaraScript.h b/Engine/Plugins/FX/Niagara/Source/Niagara/Classes/NiagaraScript.h
--- a/Engine/Plugins/FX/Niagara/Source/Niagara/Classes/NiagaraScript.h
+++ b/Engine/Plugins/FX/Niagara/Source/Niagara/Classes/NiagaraScript.h
@@ -156,21 +156,21 @@ public:
 	* The GUID of the subgraph this shader primarily represents.
 	*/
 	UPROPERTY()
-	FGuid BaseScriptID;
+	FGuid BaseScriptID_DEPRECATED;
 
-#if WITH_EDITORONLY_DATA
 	/**
 	* The hash of the subgraph this shader primarily represents.
 	*/
 	UPROPERTY()
 	FNiagaraCompileHash BaseScriptCompileHash;
 
+#if WITH_EDITORONLY_DATA
 	/** Compile hashes of any top level scripts the script was dependent on that might trigger a recompile if they change. */
 	UPROPERTY()
 	TArray<FNiagaraCompileHash> ReferencedCompileHashes;
 
 	/** Temp storage while generating the Id. This is NOT serialized and shouldn't be used in any comparisons*/
-	TArray<UObject*> ReferencedObjects;
+	TArray<FString> DebugReferencedObjects;
 #endif
 
 	FNiagaraVMExecutableDataId()
@@ -180,7 +180,7 @@ public:
 		, bUsesRapidIterationParams(true)
 		, bInterpolatedSpawn(false)
 		, bRequiresPersistentIDs(false)
-		, BaseScriptID(0, 0, 0, 0)
+		, BaseScriptID_DEPRECATED(0, 0, 0, 0)
 	{ }
 
 
@@ -192,7 +192,7 @@ public:
 	
 	friend uint32 GetTypeHash(const FNiagaraVMExecutableDataId& Ref)
 	{
-		return Ref.BaseScriptID.A;
+		return Ref.BaseScriptCompileHash.GetTypeHash();
 	}
 
 	SIZE_T GetSizeBytes() const
@@ -222,7 +222,8 @@ public:
 
 #if WITH_EDITORONLY_DATA
 	/** Appends string representations of this Id to a key string. */
-	void AppendKeyString(FString& KeyString) const;
+	void AppendKeyString(FString& KeyString, const FString& Delimiter = TEXT("_"), bool bAppendObjectForDebugging = false) const;
+
 #endif
 };
 
@@ -502,7 +503,7 @@ public:
 
 	NIAGARA_API FGuid GetBaseChangeID() const;
 	NIAGARA_API ENiagaraScriptCompileStatus GetLastCompileStatus() const;
-	void InvalidateCachedCompileIds();
+	void ForceGraphToRecompileOnNextCheck();
 
 	NIAGARA_API bool HandleVariableRenames(const TMap<FNiagaraVariable, FNiagaraVariable>& OldToNewVars, const FString& UniqueEmitterName);
 #endif
diff --git a/Engine/Plugins/FX/Niagara/Source/Niagara/Classes/NiagaraScriptSourceBase.h b/Engine/Plugins/FX/Niagara/Source/Niagara/Classes/NiagaraScriptSourceBase.h
--- a/Engine/Plugins/FX/Niagara/Source/Niagara/Classes/NiagaraScriptSourceBase.h
+++ b/Engine/Plugins/FX/Niagara/Source/Niagara/Classes/NiagaraScriptSourceBase.h
@@ -117,7 +117,7 @@ class UNiagaraScriptSourceBase : public UObject
 
 	FOnChanged& OnChanged() { return OnChangedDelegate; }
 
-	virtual void InvalidateCachedCompileIds() {}
+	virtual void ForceGraphToRecompileOnNextCheck() {}
 
 	virtual void RefreshFromExternalChanges() {}
 
diff --git a/Engine/Plugins/FX/Niagara/Source/Niagara/Classes/NiagaraSystem.h b/Engine/Plugins/FX/Niagara/Source/Niagara/Classes/NiagaraSystem.h
--- a/Engine/Plugins/FX/Niagara/Source/Niagara/Classes/NiagaraSystem.h
+++ b/Engine/Plugins/FX/Niagara/Source/Niagara/Classes/NiagaraSystem.h
@@ -266,7 +266,7 @@ public:
 #if WITH_EDITORONLY_DATA
 	bool UsesEmitter(const UNiagaraEmitter* Emitter) const;
 	bool UsesScript(const UNiagaraScript* Script)const; 
-	void InvalidateCachedCompileIds();
+	void ForceGraphToRecompileOnNextCheck();
 
 	static void RequestCompileForEmitter(UNiagaraEmitter* InEmitter);
 
diff --git a/Engine/Plugins/FX/Niagara/Source/Niagara/Private/NiagaraDataInterface.cpp b/Engine/Plugins/FX/Niagara/Source/Niagara/Private/NiagaraDataInterface.cpp
--- a/Engine/Plugins/FX/Niagara/Source/Niagara/Private/NiagaraDataInterface.cpp
+++ b/Engine/Plugins/FX/Niagara/Source/Niagara/Private/NiagaraDataInterface.cpp
@@ -29,6 +29,12 @@ UNiagaraDataInterface::~UNiagaraDataInterface()
 	}
 }
 
+bool UNiagaraDataInterface::AppendCompileHash(FNiagaraCompileHashVisitor* InVisitor) const
+{
+
+	return true;
+}
+
 void UNiagaraDataInterface::PostLoad()
 {
 	Super::PostLoad();
diff --git a/Engine/Plugins/FX/Niagara/Source/Niagara/Private/NiagaraDataInterfaceSkeletalMesh.cpp b/Engine/Plugins/FX/Niagara/Source/Niagara/Private/NiagaraDataInterfaceSkeletalMesh.cpp
--- a/Engine/Plugins/FX/Niagara/Source/Niagara/Private/NiagaraDataInterfaceSkeletalMesh.cpp
+++ b/Engine/Plugins/FX/Niagara/Source/Niagara/Private/NiagaraDataInterfaceSkeletalMesh.cpp
@@ -1946,6 +1946,16 @@ void UNiagaraDataInterfaceSkeletalMesh::GetCommonHLSL(FString& OutHLSL)
 	OutHLSL += TEXT("#include \"/Plugin/FX/Niagara/Private/NiagaraDataInterfaceSkeletalMesh.ush\"\n");
 }
 
+bool UNiagaraDataInterfaceSkeletalMesh::AppendCompileHash(FNiagaraCompileHashVisitor* InVisitor) const
+{
+	if (!Super::AppendCompileHash(InVisitor))
+		return false;
+
+	FSHAHash Hash = GetShaderFileHash((TEXT("/Plugin/FX/Niagara/Private/NiagaraDataInterfaceSkeletalMesh.ush")), EShaderPlatform::SP_PCD3D_SM5);
+	InVisitor->UpdateString(TEXT("NiagaraDataInterfaceSkeletalMeshHLSLSource"), Hash.ToString());
+	return true;
+}
+
 bool UNiagaraDataInterfaceSkeletalMesh::GetFunctionHLSL(const FName& DefinitionFunctionName, FString InstanceFunctionName, FNiagaraDataInterfaceGPUParamInfo& ParamInfo, FString& OutHLSL) 
 {
 	FNDISkeletalMeshParametersName ParamNames;
diff --git a/Engine/Plugins/FX/Niagara/Source/Niagara/Private/NiagaraEmitter.cpp b/Engine/Plugins/FX/Niagara/Source/Niagara/Private/NiagaraEmitter.cpp
--- a/Engine/Plugins/FX/Niagara/Source/Niagara/Private/NiagaraEmitter.cpp
+++ b/Engine/Plugins/FX/Niagara/Source/Niagara/Private/NiagaraEmitter.cpp
@@ -388,7 +388,7 @@ void UNiagaraEmitter::PostLoad()
 			UObject* OuterObj = GetOuter();
 			if (OuterObj == GetOutermost())
 			{
-				GraphSource->InvalidateCachedCompileIds();
+				GraphSource->ForceGraphToRecompileOnNextCheck();
 				bGenerateNewChangeId = true;
 				GenerateNewChangeIdReason = TEXT("PostLoad - Force compile on load");
 				if (GEnableVerboseNiagaraChangeIdLogging)
@@ -700,6 +700,17 @@ void UNiagaraEmitter::GetScripts(TArray<UNiagaraScript*>& OutScripts, bool bComp
 		}
 	}
 
+	if (!bCompilableOnly)
+	{
+		for (int32 i = 0; i < ShaderStages.Num(); i++)
+		{
+			if (ShaderStages[i] && ShaderStages[i]->Script)
+			{
+				OutScripts.Add(ShaderStages[i]->Script);
+			}
+		}
+	}
+
 	if (SimTarget == ENiagaraSimTarget::GPUComputeSim)
 	{
 		OutScripts.Add(GPUComputeScript);
@@ -772,6 +783,14 @@ bool UNiagaraEmitter::AreAllScriptAndSourcesSynchronized() const
 		}
 	}
 
+	for (int32 i = 0; i < ShaderStages.Num(); i++)
+	{
+		if (ShaderStages[i] && ShaderStages[i]->Script  && ShaderStages[i]->Script->IsCompilable() && !ShaderStages[i]->Script->AreScriptAndSourceSynchronized())
+		{
+			return false;
+		}
+	}
+
 	if (GPUComputeScript->IsCompilable() && !GPUComputeScript->AreScriptAndSourceSynchronized())
 	{
 		return false;
diff --git a/Engine/Plugins/FX/Niagara/Source/Niagara/Private/NiagaraEmitterHandle.cpp b/Engine/Plugins/FX/Niagara/Source/Niagara/Private/NiagaraEmitterHandle.cpp
--- a/Engine/Plugins/FX/Niagara/Source/Niagara/Private/NiagaraEmitterHandle.cpp
+++ b/Engine/Plugins/FX/Niagara/Source/Niagara/Private/NiagaraEmitterHandle.cpp
@@ -83,12 +83,12 @@ void FNiagaraEmitterHandle::SetName(FName InName, UNiagaraSystem& InOwnerSystem)
 	Name = UniqueName;
 	if (Instance->SetUniqueEmitterName(Name.ToString()))
 	{
-#if WITH_EDITOR
+ #if WITH_EDITOR
 		if (InOwnerSystem.GetSystemSpawnScript() && InOwnerSystem.GetSystemSpawnScript()->GetSource())
 		{
 			// Just invalidate the system scripts here. The emitter scripts have their important variables 
 			// changed in the SetUniqueEmitterName method above.
-			InOwnerSystem.GetSystemSpawnScript()->GetSource()->InvalidateCachedCompileIds();
+			InOwnerSystem.GetSystemSpawnScript()->GetSource()->MarkNotSynchronized(TEXT("EmitterHandleRenamed"));
 		}
 #endif
 	}
diff --git a/Engine/Plugins/FX/Niagara/Source/Niagara/Private/NiagaraModule.cpp b/Engine/Plugins/FX/Niagara/Source/Niagara/Private/NiagaraModule.cpp
--- a/Engine/Plugins/FX/Niagara/Source/Niagara/Private/NiagaraModule.cpp
+++ b/Engine/Plugins/FX/Niagara/Source/Niagara/Private/NiagaraModule.cpp
@@ -28,6 +28,15 @@ IMPLEMENT_MODULE(INiagaraModule, Niagara);
 
 #define LOCTEXT_NAMESPACE "NiagaraModule"
 
+int32 FNiagaraCompileHashVisitor::LogCompileIdGeneration = 0;
+static FAutoConsoleVariableRef CVarLogCompileIdGeneration(
+	TEXT("fx.LogCompileIdGeneration"),
+	FNiagaraCompileHashVisitor::LogCompileIdGeneration,
+	TEXT("If > 0 all compile id generation will be logged. If 2 or greater, log detailed info. \n"),
+	ECVF_Default
+);
+
+
 float INiagaraModule::EngineGlobalSpawnCountScale = 1.0f;
 float INiagaraModule::EngineGlobalSystemCountScale = 1.0f;
 int32 INiagaraModule::EngineDetailLevel = 4;
@@ -647,6 +656,75 @@ bool FNiagaraTypeDefinition::IsValidNumericInput(const FNiagaraTypeDefinition& T
 	return false;
 }
 
+
+bool FNiagaraTypeDefinition::AppendCompileHash(FNiagaraCompileHashVisitor* InVisitor) const
+{
+#if WITH_EDITORONLY_DATA
+	UStruct* TDStruct = GetStruct();
+	UClass* TDClass = GetClass();
+	UEnum* TDEnum = GetEnum();
+
+	if (TDEnum)
+	{
+		// Do we need to enumerate all the enum values and rebuild if that changes or are we ok with just knowing that there are the same  count of enum entries?
+		// For now, am just going to be ok with the number of entries. The actual string values don't matter so much.
+		FString CppType = TDEnum->CppType;
+		FString PathName = TDEnum->GetPathName();
+		InVisitor->UpdateString(TEXT("\tEnumPath"), PathName);
+		InVisitor->UpdateString(TEXT("\tEnumCppType"), CppType);
+		InVisitor->UpdatePOD(TEXT("\t\tNumEnums"),TDEnum->NumEnums());
+	}
+	else if (TDClass)
+	{
+		// For data interfaces, get the default object and the compile version so that we can properly update when code changes.
+		check(IsInGameThread());
+		UObject* TempObj = TDClass->GetDefaultObject(false);
+		check(TempObj);
+
+		FString ClassName = TDClass->GetPathName();
+		InVisitor->UpdateString(TEXT("\tClassName"), ClassName);
+
+		UNiagaraDataInterface* TempDI = Cast< UNiagaraDataInterface>(TempObj);
+		if (TempDI)
+		{
+			if (!TempDI->AppendCompileHash(InVisitor))
+			{
+				UE_LOG(LogNiagara, Warning, TEXT("Unable to generate AppendCompileHash for DI %s"), *TempDI->GetPathName());
+			}
+		}
+	}
+	else if (TDStruct)
+	{
+		FString ClassName = TDStruct->GetPathName();
+		InVisitor->UpdateString(TEXT("\tStructName"), ClassName);
+		// Structs are potentially changed, so we will want to register their actual types and variable names.
+		for (TFieldIterator<UProperty> PropertyIt(TDStruct, EFieldIteratorFlags::IncludeSuper, EFieldIteratorFlags::IncludeDeprecated); PropertyIt; ++PropertyIt)
+		{
+			UProperty* Property = *PropertyIt;
+			if (Property->HasMetaData(TEXT("SkipForCompileHash")))
+			{
+				continue;
+			}
+			/*if (Property->HasAnyFlags(RF_Transient))
+			{
+				continue;
+			}*/
+			InVisitor->UpdateString(TEXT("\t\tPropertyName"), Property->GetName());
+			InVisitor->UpdateString(TEXT("\t\tPropertyClass"), Property->GetClass()->GetName());
+		}
+	}
+	else
+	{
+		FString InvalidStr = TEXT("Invalid");
+		InVisitor->UpdateString(TEXT("\tTDName"), InvalidStr);
+	}
+
+	return true;
+#else
+	return false;
+#endif
+}
+
 void FNiagaraTypeDefinition::RecreateUserDefinedTypeRegistry()
 {
 	static auto* CoreUObjectPkg = FindObjectChecked<UPackage>(nullptr, TEXT("/Script/CoreUObject"));
diff --git a/Engine/Plugins/FX/Niagara/Source/Niagara/Private/NiagaraParameterCollection.cpp b/Engine/Plugins/FX/Niagara/Source/Niagara/Private/NiagaraParameterCollection.cpp
--- a/Engine/Plugins/FX/Niagara/Source/Niagara/Private/NiagaraParameterCollection.cpp
+++ b/Engine/Plugins/FX/Niagara/Source/Niagara/Private/NiagaraParameterCollection.cpp
@@ -2,6 +2,7 @@
 
 #include "NiagaraParameterCollection.h"
 #include "NiagaraDataInterface.h"
+#include "Misc/SecureHash.h"
 #if WITH_EDITORONLY_DATA
 	#include "IAssetTools.h"
 #endif
@@ -297,7 +298,7 @@ FNiagaraCompileHash UNiagaraParameterCollection::GetCompileHash() const
 	CompileHash.Final();
 
 	TArray<uint8> DataHash;
-	DataHash.AddUninitialized(20);
+	DataHash.AddUninitialized(FSHA1::DigestSize);
 	CompileHash.GetHash(DataHash.GetData());
 
 	return FNiagaraCompileHash(DataHash);
diff --git a/Engine/Plugins/FX/Niagara/Source/Niagara/Private/NiagaraScript.cpp b/Engine/Plugins/FX/Niagara/Source/Niagara/Private/NiagaraScript.cpp
--- a/Engine/Plugins/FX/Niagara/Source/Niagara/Private/NiagaraScript.cpp
+++ b/Engine/Plugins/FX/Niagara/Source/Niagara/Private/NiagaraScript.cpp
@@ -11,6 +11,7 @@
 #include "NiagaraCustomVersion.h"
 #include "NiagaraShaderCompilationManager.h"
 #include "Serialization/MemoryReader.h"
+#include "Misc/SecureHash.h"
 
 #include "Stats/Stats.h"
 #include "UObject/Linker.h"
@@ -36,6 +37,14 @@
 
 DECLARE_STATS_GROUP(TEXT("Niagara Detailed"), STATGROUP_NiagaraDetailed, STATCAT_Advanced);
 
+int32 GNiagaraDumpKeyGen = 0;
+static FAutoConsoleVariableRef CVarNiagaraDumpKeyGen(
+	TEXT("fx.DumpGraphKeyGen"),
+	GNiagaraDumpKeyGen,
+	TEXT("If > 0 the key generation will be dumped to the log. \n"),
+	ECVF_Default
+);
+
 FNiagaraScriptDebuggerInfo::FNiagaraScriptDebuggerInfo() : bWaitForGPU(false), FrameLastWriteId(-1), bWritten(false)
 {
 }
@@ -93,7 +102,7 @@ void FNiagaraVMExecutableData::SerializeData(FArchive& Ar, bool bDDCData)
 
 bool FNiagaraVMExecutableDataId::IsValid() const
 {
-	return BaseScriptID.IsValid();
+	return CompilerVersionID.IsValid();
 }
 
 void FNiagaraVMExecutableDataId::Invalidate()
@@ -122,7 +131,6 @@ bool FNiagaraVMExecutableDataId::operator==(const FNiagaraVMExecutableDataId& Re
 	if (CompilerVersionID != ReferenceSet.CompilerVersionID ||
 		ScriptUsageType != ReferenceSet.ScriptUsageType || 
 		ScriptUsageTypeID != ReferenceSet.ScriptUsageTypeID ||
-		BaseScriptID != ReferenceSet.BaseScriptID ||
 #if WITH_EDITORONLY_DATA
 		BaseScriptCompileHash != ReferenceSet.BaseScriptCompileHash ||
 #endif
@@ -169,17 +177,34 @@ bool FNiagaraVMExecutableDataId::operator==(const FNiagaraVMExecutableDataId& Re
 }
 
 #if WITH_EDITORONLY_DATA
-void FNiagaraVMExecutableDataId::AppendKeyString(FString& KeyString) const
+void FNiagaraVMExecutableDataId::AppendKeyString(FString& KeyString, const FString& Delimiter, bool bAppendObjectForDebugging) const
 {
-	KeyString += FString::Printf(TEXT("%d_"), (int32)ScriptUsageType);
+	KeyString += FString::Printf(TEXT("%d%s"), (int32)ScriptUsageType, *Delimiter);
 	KeyString += ScriptUsageTypeID.ToString();
-	KeyString += TEXT("_");
+	if (bAppendObjectForDebugging)
+	{
+		KeyString += TEXT(" [ScriptUsageType]");
+	}
+	KeyString += Delimiter;
+	
 	KeyString += CompilerVersionID.ToString();
-	KeyString += TEXT("_");
-	KeyString += BaseScriptID.ToString();
-	KeyString += TEXT("_");
+	if (bAppendObjectForDebugging)
+	{
+		KeyString += TEXT(" [CompilerVersionID]");
+	}
+	KeyString += Delimiter;
+	
 	KeyString += BaseScriptCompileHash.ToString();
-	KeyString += TEXT("_");
+	if (bAppendObjectForDebugging)
+	{
+		KeyString += TEXT(" [BaseScriptCompileHash]");
+	}
+	KeyString += Delimiter;
+
+	if (bAppendObjectForDebugging)
+	{
+		KeyString += TEXT("[AdditionalDefines]") + Delimiter;
+	}
 
 	if (DetailLevelMask != 0xFFFFFFFF)
 	{
@@ -187,26 +212,22 @@ void FNiagaraVMExecutableDataId::AppendKeyString(FString& KeyString) const
 	}
 	else
 	{
-		KeyString += TEXT("ALLDL_");
+		KeyString += TEXT("ALLDL") + Delimiter;
 	}
 
 	if (bUsesRapidIterationParams)
 	{
-		KeyString += TEXT("USESRI_");
+		KeyString += TEXT("USESRI") + Delimiter;
 	}
 	else
 	{
-		KeyString += TEXT("NORI_");
+		KeyString += TEXT("NORI") + Delimiter;
 	}
 
 	for (int32 Idx = 0; Idx < AdditionalDefines.Num(); Idx++)
 	{
 		KeyString += AdditionalDefines[Idx];
-
-		if (Idx < AdditionalDefines.Num() - 1)
-		{
-			KeyString += TEXT("_");
-		}
+		KeyString += Delimiter;
 	}
 	
 	// Add any referenced script compile hashes to the key so that we will recompile when they are changed
@@ -214,12 +235,18 @@ void FNiagaraVMExecutableDataId::AppendKeyString(FString& KeyString) const
 	{
 		KeyString += ReferencedCompileHashes[HashIndex].ToString();
 
+		if (bAppendObjectForDebugging && DebugReferencedObjects.Num() > HashIndex)
+		{
+			KeyString += TEXT(" [") + DebugReferencedObjects[HashIndex] + TEXT("]") ;
+		}
+
 		if (HashIndex < ReferencedCompileHashes.Num() - 1)
 		{
-			KeyString += TEXT("_");
+			KeyString += Delimiter;
 		}
 	}
 }
+
 #endif
 
 UNiagaraScript::UNiagaraScript(const FObjectInitializer& ObjectInitializer)
@@ -366,12 +393,12 @@ void UNiagaraScript::ComputeVMCompilationId(FNiagaraVMExecutableDataId& Id) cons
 		HashState.Final();
 
 		TArray<uint8> DataHash;
-		DataHash.AddUninitialized(20);
+		DataHash.AddUninitialized(FSHA1::DigestSize);
 		HashState.GetHash(DataHash.GetData());
 
 		FNiagaraCompileHash Hash(DataHash);
 		Id.ReferencedCompileHashes.Add(Hash);
-		Id.ReferencedObjects.Add(nullptr);
+		Id.DebugReferencedObjects.Add(TEXT("RIParams"));
 	}
 
 	static const auto UseShaderStagesCVar = IConsoleManager::Get().FindConsoleVariable(TEXT("fx.UseShaderStages"));
@@ -380,9 +407,24 @@ void UNiagaraScript::ComputeVMCompilationId(FNiagaraVMExecutableDataId& Id) cons
 		Id.AdditionalDefines.Add(TEXT("fx.UseShaderStages"));
 	}
 
-	Source->ComputeVMCompilationId(Id, Usage, UsageId);
+	Source->ComputeVMCompilationId(Id, Usage, UsageId);	
 	
+	if (GNiagaraDumpKeyGen == 1 && Id != LastGeneratedVMId)
+	{
+		TArray<FString> OutputByLines;
+		FString StrDump;
+		Id.AppendKeyString(StrDump, TEXT("\n"), true);
+		StrDump.ParseIntoArrayLines(OutputByLines, false);
+
+		UE_LOG(LogNiagara, Display, TEXT("KeyGen %s\n==================\n"), *GetPathName());
+		for (int32 i = 0; i < OutputByLines.Num(); i++)
+		{
+			UE_LOG(LogNiagara, Display, TEXT("/*%04d*/\t\t%s"), i + 1, *OutputByLines[i]);
+		}
+	}
+
 	LastGeneratedVMId = Id;
+
 }
 #endif
 
@@ -779,7 +821,7 @@ void UNiagaraScript::PostLoad()
 					TEXT("Invalidating compile ids for script %s because it doesn't have a valid base id.  The owning asset will continue to compile on load until it is resaved."),
 					*GetPathName());
 				bool bForceRebuild = true;
-				Source->InvalidateCachedCompileIds();
+				Source->ForceGraphToRecompileOnNextCheck();
 				Source->ComputeVMCompilationId(CachedScriptVMId, Usage, UsageId, bForceRebuild);
 			}
 			else
@@ -816,7 +858,7 @@ void UNiagaraScript::PostLoad()
 	// for now, force recompile until we can be sure everything is working
 	//bNeedsRecompile = true;
 #if WITH_EDITORONLY_DATA
-	if (CachedScriptVMId.BaseScriptID.IsValid() && CachedScriptVMId.BaseScriptCompileHash.IsValid())
+	if (CachedScriptVMId.BaseScriptCompileHash.IsValid())
 	{
 		CacheResourceShadersForRendering(false, bNeedsRecompile);
 	}
@@ -916,12 +958,6 @@ bool UNiagaraScript::AreScriptAndSourceSynchronized() const
 			{
 				if (GEnableVerboseNiagaraChangeIdLogging)
 				{
-					if (NewId.BaseScriptID != CachedScriptVMId.BaseScriptID)
-					{
-						UE_LOG(LogNiagara, Log, TEXT("AreScriptAndSourceSynchronized base script id's don't match. %s != %s, script %s"),
-							*NewId.BaseScriptID.ToString(), *CachedScriptVMId.BaseScriptID.ToString(), *GetPathName());
-					}
-
 					if (NewId.BaseScriptCompileHash != CachedScriptVMId.BaseScriptCompileHash)
 					{
 						UE_LOG(LogNiagara, Log, TEXT("AreScriptAndSourceSynchronized base script compile hashes don't match. %s != %s, script %s"),
@@ -941,7 +977,7 @@ bool UNiagaraScript::AreScriptAndSourceSynchronized() const
 							{
 								UE_LOG(LogNiagara, Log, TEXT("AreScriptAndSourceSynchronized referenced compile hash %d doesn't match. %s != %s, script %s, source %s"),
 									i, *NewId.ReferencedCompileHashes[i].ToString(), *CachedScriptVMId.ReferencedCompileHashes[i].ToString(), *GetPathName(),
-									NewId.ReferencedObjects[i] != nullptr ? *NewId.ReferencedObjects[i]->GetPathName() : TEXT("nullptr"));
+									*NewId.DebugReferencedObjects[i]);
 							}
 						}
 					}
@@ -1224,7 +1260,7 @@ void UNiagaraScript::SetVMCompilationResults(const FNiagaraVMExecutableDataId& I
 	// Now go ahead and trigger the GPU script compile now that we have a compiled GPU hlsl script.
 	if (Usage == ENiagaraScriptUsage::ParticleGPUComputeScript)
 	{
-		if (CachedScriptVMId.CompilerVersionID.IsValid() && CachedScriptVMId.BaseScriptID.IsValid() && CachedScriptVMId.BaseScriptCompileHash.IsValid())
+		if (CachedScriptVMId.CompilerVersionID.IsValid() && CachedScriptVMId.BaseScriptCompileHash.IsValid())
 		{
 			CacheResourceShadersForRendering(false, true);
 		}
@@ -1252,9 +1288,9 @@ void UNiagaraScript::InvalidateExecutionReadyParameterStores()
 #endif
 }
 
-void UNiagaraScript::InvalidateCachedCompileIds()
+void UNiagaraScript::ForceGraphToRecompileOnNextCheck()
 {
-	GetSource()->InvalidateCachedCompileIds();
+	GetSource()->ForceGraphToRecompileOnNextCheck();
 }
 
 void UNiagaraScript::RequestCompile(bool bForceCompile)
@@ -1431,7 +1467,7 @@ void UNiagaraScript::BeginCacheForCookedPlatformData(const ITargetPlatform *Targ
 			SystemOwner->WaitForCompilationComplete();
 		}
 
-		if (CachedScriptVMId.CompilerVersionID.IsValid() == false || CachedScriptVMId.BaseScriptID.IsValid() == false || CachedScriptVMId.BaseScriptCompileHash.IsValid() == false)
+		if (CachedScriptVMId.CompilerVersionID.IsValid() == false ||  CachedScriptVMId.BaseScriptCompileHash.IsValid() == false)
 		{
 			UE_LOG(LogNiagara, Error,
 				TEXT("Failed to cache cooked shader for script %s because it had an invalid cached script id.  This should be fixed by running the console command fx.PreventSystemRecompile with the owning system asset path as the argument and then resaving the assets."),
@@ -1510,10 +1546,9 @@ void UNiagaraScript::CacheResourceShadersForCooking(EShaderPlatform ShaderPlatfo
 
 			FNiagaraShaderScript* NewResource = AllocateResource();
 			check(CachedScriptVMId.CompilerVersionID.IsValid());
-			check(CachedScriptVMId.BaseScriptID.IsValid());
 			check(CachedScriptVMId.BaseScriptCompileHash.IsValid());
 
-			NewResource->SetScript(this, (ERHIFeatureLevel::Type)TargetFeatureLevel, CachedScriptVMId.CompilerVersionID, CachedScriptVMId.BaseScriptID, CachedScriptVMId.AdditionalDefines,
+			NewResource->SetScript(this, (ERHIFeatureLevel::Type)TargetFeatureLevel, CachedScriptVMId.CompilerVersionID, CachedScriptVMId.AdditionalDefines,
 				CachedScriptVMId.BaseScriptCompileHash,	CachedScriptVMId.ReferencedCompileHashes, 
 				CachedScriptVMId.bUsesRapidIterationParams, CachedScriptVMId.DetailLevelMask, GetFullName());
 			ResourceToCache = NewResource;
@@ -1579,7 +1614,7 @@ void UNiagaraScript::CacheResourceShadersForRendering(bool bRegenerateId, bool b
 		{
 			FNiagaraShaderScript* ResourceToCache;
 			ERHIFeatureLevel::Type CacheFeatureLevel = GMaxRHIFeatureLevel;
-			ScriptResource.SetScript(this, FeatureLevel, CachedScriptVMId.CompilerVersionID, CachedScriptVMId.BaseScriptID, CachedScriptVMId.AdditionalDefines,
+			ScriptResource.SetScript(this, FeatureLevel, CachedScriptVMId.CompilerVersionID, CachedScriptVMId.AdditionalDefines,
 				CachedScriptVMId.BaseScriptCompileHash, CachedScriptVMId.ReferencedCompileHashes, 
 				CachedScriptVMId.bUsesRapidIterationParams, CachedScriptVMId.DetailLevelMask, GetFullName());
 
diff --git a/Engine/Plugins/FX/Niagara/Source/Niagara/Private/NiagaraSystem.cpp b/Engine/Plugins/FX/Niagara/Source/Niagara/Private/NiagaraSystem.cpp
--- a/Engine/Plugins/FX/Niagara/Source/Niagara/Private/NiagaraSystem.cpp
+++ b/Engine/Plugins/FX/Niagara/Source/Niagara/Private/NiagaraSystem.cpp
@@ -412,7 +412,7 @@ void UNiagaraSystem::PostLoad()
 
 		if (UNiagaraEmitter::GetForceCompileOnLoad())
 		{
-			InvalidateCachedCompileIds();
+			ForceGraphToRecompileOnNextCheck();
 			UE_LOG(LogNiagara, Log, TEXT("System %s being rebuilt because UNiagaraEmitter::GetForceCompileOnLoad() == true."), *GetPathName());
 		}
 
@@ -761,17 +761,17 @@ UNiagaraSystem::FOnSystemCompiled& UNiagaraSystem::OnSystemCompiled()
 	return OnSystemCompiledDelegate;
 }
 
-void UNiagaraSystem::InvalidateCachedCompileIds()
+void UNiagaraSystem::ForceGraphToRecompileOnNextCheck()
 {
 	check(SystemSpawnScript->GetSource() == SystemUpdateScript->GetSource());
-	SystemSpawnScript->GetSource()->InvalidateCachedCompileIds();
+	SystemSpawnScript->GetSource()->ForceGraphToRecompileOnNextCheck();
 
 	for (FNiagaraEmitterHandle Handle : EmitterHandles)
 	{
 		if (Handle.GetInstance())
 		{
 			UNiagaraScriptSourceBase* GraphSource = Handle.GetInstance()->GraphSource;
-			GraphSource->InvalidateCachedCompileIds();
+			GraphSource->ForceGraphToRecompileOnNextCheck();
 		}
 	}
 }
@@ -829,9 +829,16 @@ bool UNiagaraSystem::QueryCompileComplete(bool bWait, bool bDoPost, bool bDoNotA
 				bool bBuiltLocally = false;
 				if (GetDerivedDataCacheRef().GetAsynchronousResults(EmitterCompiledScriptPair.PendingDDCID, OutData, &bBuiltLocally))
 				{
-					if (bBuiltLocally)
+					if (GEnableVerboseNiagaraChangeIdLogging)
 					{
-						UE_LOG(LogNiagara, Log, TEXT("UNiagraScript \'%s\' was built locally.."), *EmitterCompiledScriptPair.CompiledScript->GetFullName());
+						if (bBuiltLocally)
+						{
+							UE_LOG(LogNiagara, Log, TEXT("UNiagraScript \'%s\' was built locally.."), *EmitterCompiledScriptPair.CompiledScript->GetFullName());
+						}
+						else
+						{
+							UE_LOG(LogNiagara, Log, TEXT("UNiagraScript \'%s\' was pulled from DDC.."), *EmitterCompiledScriptPair.CompiledScript->GetFullName());
+						}
 					}
 
 					TSharedPtr<FNiagaraVMExecutableData> ExeData = MakeShared<FNiagaraVMExecutableData>();
@@ -1046,7 +1053,7 @@ bool UNiagaraSystem::RequestCompile(bool bForce)
 {
 	if (bForce)
 	{
-		InvalidateCachedCompileIds();
+		ForceGraphToRecompileOnNextCheck();
 		bForce = false;
 	}
 
diff --git a/Engine/Plugins/FX/Niagara/Source/Niagara/Public/NiagaraCommon.h b/Engine/Plugins/FX/Niagara/Source/Niagara/Public/NiagaraCommon.h
--- a/Engine/Plugins/FX/Niagara/Source/Niagara/Public/NiagaraCommon.h
+++ b/Engine/Plugins/FX/Niagara/Source/Niagara/Public/NiagaraCommon.h
@@ -272,7 +272,7 @@ struct NIAGARA_API FNiagaraFunctionSignature
 
 	/** Localized description of this node. Note that this is *not* used during the operator == below since it may vary from culture to culture.*/
 #if WITH_EDITORONLY_DATA
-	UPROPERTY()
+	UPROPERTY(meta = (SkipForCompileHash = "true"))
 	FText Description;
 #endif
 
diff --git a/Engine/Plugins/FX/Niagara/Source/Niagara/Public/NiagaraTypes.h b/Engine/Plugins/FX/Niagara/Source/Niagara/Public/NiagaraTypes.h
--- a/Engine/Plugins/FX/Niagara/Source/Niagara/Public/NiagaraTypes.h
+++ b/Engine/Plugins/FX/Niagara/Source/Niagara/Public/NiagaraTypes.h
@@ -6,6 +6,7 @@
 #include "UObject/UnrealType.h"
 #include "Engine/UserDefinedStruct.h"
 #include "Templates/SharedPointer.h"
+#include "Misc/SecureHash.h"
 #include "NiagaraTypes.generated.h"
 
 class UNiagaraDataInterfaceBase;
@@ -306,6 +307,138 @@ enum class ENiagaraExecutionState : uint32
 	Num UMETA(Hidden)
 };
 
+USTRUCT()
+struct NIAGARA_API FNiagaraCompileHashVisitorDebugInfo
+{
+	GENERATED_USTRUCT_BODY()
+public:
+	UPROPERTY()
+	FString Object;
+
+	UPROPERTY()
+	TArray<FString> PropertyKeys;
+
+	UPROPERTY()
+	TArray<FString> PropertyValues;
+};
+
+/**
+Used to store the state of a graph when deciding if it has been dirtied for recompile.
+*/
+struct NIAGARA_API FNiagaraCompileHashVisitor
+{
+public:
+	FNiagaraCompileHashVisitor(FSHA1& InHashState) : HashState(InHashState) {}
+
+	FSHA1& HashState;
+	TArray<const void*> ObjectList;
+
+	static int LogCompileIdGeneration;
+
+#if WITH_EDITORONLY_DATA
+
+	// Debug data about the compilation hash, including key value pairs to detect differences.
+	TArray<FNiagaraCompileHashVisitorDebugInfo> Values;
+
+	template<typename T>
+	void ToDebugString(const T* InData, uint64 InDataCount, FString& OutStr)
+	{
+		for (int32 i = 0; i < InDataCount; i++)
+		{
+			FString DataStr = LexToString(InData[i]);
+			OutStr.Appendf(TEXT("%s "), *DataStr);
+		}
+	}
+#endif
+	/**
+	Registers a pointer for later reference in the compile id in a deterministic manner.
+	*/
+	int32 RegisterReference(const void* InObject)
+	{
+		if (InObject == nullptr)
+		{
+			return -1;
+		}
+
+		int32 Index = ObjectList.Find(InObject);
+		if (Index < 0)
+		{
+			Index = ObjectList.Add(InObject);
+		}
+		return Index;
+	}
+
+	/**
+	We don't usually want to save GUID's or pointer values because they have nondeterministic values. Consider a PostLoad upgrade operation that creates a new node.
+	Each pin and node gets a unique ID. If you close the editor and reopen, you'll get a different set of values. One of the characteristics we want for compilation
+	behavior is that the same graph structure produces the same compile results, so we only want to embed information that is deterministic. This method is for use
+	when registering a pointer to an object that is serialized within the compile hash.
+	*/
+	bool UpdateReference(const TCHAR* InDebugName, const void* InObject)
+	{
+		int32 Index = RegisterReference(InObject);
+		return UpdatePOD(InDebugName, Index);
+	}
+
+	/**
+	Adds an array of POD (plain old data) values to the hash.
+	*/
+	template<typename T>
+	bool UpdateArray(const TCHAR* InDebugName, const T* InData, uint64 InDataCount = 1)
+	{
+		static_assert(TIsPODType<T>::Value, "UpdateArray does not support a constructor / destructor on the element class.");
+		HashState.Update((const uint8 *)InData, sizeof(T)*InDataCount);
+#if WITH_EDITORONLY_DATA
+		if (LogCompileIdGeneration != 0)
+		{
+			FString ValuesStr = InDebugName;
+			ValuesStr.Append(TEXT(" = "));
+			ToDebugString(InData, InDataCount, ValuesStr);
+			Values.Top().PropertyKeys.Push(InDebugName);
+			Values.Top().PropertyValues.Push(ValuesStr);
+		}
+#endif
+		return true;
+	}
+
+	/**
+	Adds a single value of typed POD (plain old data) to the hash.
+	*/
+	template<typename T>
+	bool UpdatePOD(const TCHAR* InDebugName, const T& InData)
+	{
+		static_assert(TIsPODType<T>::Value, "Update does not support a constructor / destructor on the element class.");
+		HashState.Update((const uint8 *)&InData, sizeof(T));
+#if WITH_EDITORONLY_DATA
+		if (LogCompileIdGeneration != 0)
+		{
+			FString ValuesStr;
+			ToDebugString(&InData, 1, ValuesStr);
+			Values.Top().PropertyKeys.Push(InDebugName);
+			Values.Top().PropertyValues.Push(ValuesStr);
+		}
+#endif
+		return true;
+	}
+
+	/**
+	Adds an string value to the hash.
+	*/
+	bool UpdateString(const TCHAR* InDebugName, const FString& InData)
+	{
+		HashState.Update((const uint8 *)(*InData), sizeof(TCHAR)*InData.Len());
+#if WITH_EDITORONLY_DATA
+		if (LogCompileIdGeneration != 0)
+		{
+			Values.Top().PropertyKeys.Push(InDebugName);
+			Values.Top().PropertyValues.Push(InData);
+		}
+#endif
+		return true;
+	}
+};
+
+
 /** Defines options for conditionally editing and showing script inputs in the UI. */
 USTRUCT()
 struct NIAGARA_API FNiagaraInputConditionMetadata
@@ -335,33 +468,33 @@ public:
 	{
 	}
 public:
-	UPROPERTY(EditAnywhere, Category = "Variable", meta = (MultiLine = true))
+	UPROPERTY(EditAnywhere, Category = "Variable", meta = (MultiLine = true, SkipForCompileHash = "true"))
 	FText Description;
 
-	UPROPERTY(EditAnywhere, Category = "Variable")
+	UPROPERTY(EditAnywhere, Category = "Variable", meta = (SkipForCompileHash = "true"))
 	FText CategoryName;
 
 	/** Declares that this input is advanced and should only be visible if expanded inputs have been expanded. */
-	UPROPERTY(EditAnywhere, Category = "Variable")
+	UPROPERTY(EditAnywhere, Category = "Variable", meta = (SkipForCompileHash = "true"))
 	bool bAdvancedDisplay;
 
-	UPROPERTY(EditAnywhere, Category = "Variable", meta = (ToolTip = "Affects the sort order in the editor stacks. Use a smaller number to push it to the top. Defaults to zero."))
+	UPROPERTY(EditAnywhere, Category = "Variable", meta = (ToolTip = "Affects the sort order in the editor stacks. Use a smaller number to push it to the top. Defaults to zero.", SkipForCompileHash = "true"))
 	int32 EditorSortPriority;
 
 	/** Declares the associated input is used as an inline edit condition toggle, so it should be hidden and edited as a 
 	checkbox inline with the input which was designated as its edit condition. */
-	UPROPERTY(EditAnywhere, Category = "Variable")
+	UPROPERTY(EditAnywhere, Category = "Variable", meta = (SkipForCompileHash = "true"))
 	bool bInlineEditConditionToggle;
 
 	/** Declares the associated input should be conditionally editable based on the value of another input. */
-	UPROPERTY(EditAnywhere, Category = "Variable")
+	UPROPERTY(EditAnywhere, Category = "Variable", meta = (SkipForCompileHash = "true"))
 	FNiagaraInputConditionMetadata EditCondition;
 
 	/** Declares the associated input should be conditionally visible based on the value of another input. */
-	UPROPERTY(EditAnywhere, Category = "Variable")
+	UPROPERTY(EditAnywhere, Category = "Variable", meta = (SkipForCompileHash = "true"))
 	FNiagaraInputConditionMetadata VisibleCondition;
 
-	UPROPERTY(EditAnywhere, Category = "Variable", DisplayName = "Property Metadata", meta = (ToolTip = "Property Metadata"))
+	UPROPERTY(EditAnywhere, Category = "Variable", DisplayName = "Property Metadata", meta = (ToolTip = "Property Metadata", SkipForCompileHash = "true"))
 	TMap<FName, FString> PropertyMetaData;
 
 	/** This is a read-only variable that designates if the metadata is tied to a static switch or not. */
@@ -551,6 +684,8 @@ public:
 		return ClassStructOrEnum != nullptr;
 	}
 
+	bool AppendCompileHash(FNiagaraCompileHashVisitor* InVisitor) const;
+
 	/*
 	Underlying type for this variable, use FUnderlyingType to determine type without casting
 	This can be a UClass, UStruct or UEnum.  Pointing to something like the struct for an FVector, etc.
@@ -1035,7 +1170,7 @@ struct FNiagaraVariable : public FNiagaraVariableBase
 private:
 	//This gets serialized but do we need to worry about endianness doing things like this? If not, where does that get handled?
 	//TODO: Remove storage here entirely and move everything to an FNiagaraParameterStore.
-	UPROPERTY()
+	UPROPERTY(meta = (SkipForCompileHash = "true"))
 	TArray<uint8> VarData;
 };
 
diff --git a/Engine/Plugins/FX/Niagara/Source/NiagaraCore/Private/NiagaraCustomVersion.cpp b/Engine/Plugins/FX/Niagara/Source/NiagaraCore/Private/NiagaraCustomVersion.cpp
--- a/Engine/Plugins/FX/Niagara/Source/NiagaraCore/Private/NiagaraCustomVersion.cpp
+++ b/Engine/Plugins/FX/Niagara/Source/NiagaraCore/Private/NiagaraCustomVersion.cpp
@@ -8,4 +8,4 @@ const FGuid FNiagaraCustomVersion::GUID(0xFCF57AFA, 0x50764283, 0xB9A9E658, 0xFF
 // Register the custom version with core
 FCustomVersionRegistration GRegisterNiagaraCustomVersion(FNiagaraCustomVersion::GUID, FNiagaraCustomVersion::LatestVersion, TEXT("NiagaraVer"));
 
-const FGuid FNiagaraCustomVersion::LatestScriptCompileVersion(0x3B2E872A, 0x9FDE4E10, 0xBB0194B8, 0x98DF161C);
+const FGuid FNiagaraCustomVersion::LatestScriptCompileVersion(0xB345792C, 0x05E94BFA, 0x93656C4C, 0xF8FFB377);
diff --git a/Engine/Plugins/FX/Niagara/Source/NiagaraCore/Public/NiagaraCompileHash.h b/Engine/Plugins/FX/Niagara/Source/NiagaraCore/Public/NiagaraCompileHash.h
--- a/Engine/Plugins/FX/Niagara/Source/NiagaraCore/Public/NiagaraCompileHash.h
+++ b/Engine/Plugins/FX/Niagara/Source/NiagaraCore/Public/NiagaraCompileHash.h
@@ -20,6 +20,16 @@ struct NIAGARACORE_API FNiagaraCompileHash
 		DataHash = InDataHash;
 	}
 
+	explicit FNiagaraCompileHash(const uint8* InDataHash, uint32 InCount) : DataHash(InDataHash, InCount)
+	{
+		checkf(InCount == HashSize, TEXT("Invalid hash data."));
+		/*.AddUninitialized(InCount);
+		for (uint32 i = 0; i < InCount; i++)
+		{
+			DataHash[i] = InDataHash[i];
+		}*/
+	}
+
 	bool operator==(const FNiagaraCompileHash& Other) const;
 
 	bool operator!=(const FNiagaraCompileHash& Other) const;
diff --git a/Engine/Plugins/FX/Niagara/Source/NiagaraEditor/Private/NiagaraEditorModule.cpp b/Engine/Plugins/FX/Niagara/Source/NiagaraEditor/Private/NiagaraEditorModule.cpp
--- a/Engine/Plugins/FX/Niagara/Source/NiagaraEditor/Private/NiagaraEditorModule.cpp
+++ b/Engine/Plugins/FX/Niagara/Source/NiagaraEditor/Private/NiagaraEditorModule.cpp
@@ -107,6 +107,18 @@
 #include "NiagaraNodeFunctionCall.h"
 #include "INiagaraEditorOnlyDataUtlities.h"
 
+#include "Editor.h"
+#include "Factories/Factory.h"
+#include "HAL/PlatformFilemanager.h"
+#include "HAL/FileManager.h"
+#include "Modules/ModuleManager.h"
+#include "Containers/ArrayView.h"
+#include "EditorReimportHandler.h"
+#include "ISourceControlOperation.h"
+#include "SourceControlOperations.h"
+#include "ISourceControlProvider.h"
+#include "ISourceControlModule.h"
+
 IMPLEMENT_MODULE( FNiagaraEditorModule, NiagaraEditor );
 
 PRAGMA_DISABLE_OPTIMIZATION
@@ -579,6 +591,101 @@ void DumpCompileIdDataForAsset(const TArray<FString>& Arguments)
 	}
 }
 
+void ExecuteInvalidateNiagaraCachedScripts(const TArray< FString >& Args)
+{
+	if (Args.Num() == 0)
+	{
+		// todo: log error, at least one command is needed
+		UE_LOG(LogConsoleResponse, Display, TEXT("fx.InvalidateCachedScripts failed\nAs this command should not be executed accidentally it requires you to specify an extra parameter."));
+		return;
+	}
+
+	FString FileName = FPaths::EngineDir() + TEXT("Plugins/FX/Niagara/Shaders/Private/NiagaraShaderVersion.ush");
+
+	FileName = IFileManager::Get().ConvertToAbsolutePathForExternalAppForRead(*FileName);
+
+	ISourceControlProvider& SourceControlProvider = ISourceControlModule::Get().GetProvider();
+	SourceControlProvider.Init();
+
+	FSourceControlStatePtr SourceControlState = SourceControlProvider.GetState(FileName, EStateCacheUsage::ForceUpdate);
+	if (SourceControlState.IsValid())
+	{
+		if (SourceControlState->CanCheckout() || SourceControlState->IsCheckedOutOther())
+		{
+			if (SourceControlProvider.Execute(ISourceControlOperation::Create<FCheckOut>(), FileName) == ECommandResult::Failed)
+			{
+				UE_LOG(LogConsoleResponse, Display, TEXT("fx.InvalidateCachedScripts failed\nCouldn't check out \"NiagaraShaderVersion.ush\""));
+				return;
+			}
+		}
+		else if (!SourceControlState->IsSourceControlled())
+		{
+			UE_LOG(LogConsoleResponse, Display, TEXT("fx.InvalidateCachedScripts failed\n\"NiagaraShaderVersion.ush\" is not under source control."));
+		}
+		else if (SourceControlState->IsCheckedOutOther())
+		{
+			UE_LOG(LogConsoleResponse, Display, TEXT("fx.InvalidateCachedScripts failed\n\"NiagaraShaderVersion.ush\" is already checked out by someone else\n(UE4 SourceControl needs to be fixed to allow multiple checkout.)"));
+			return;
+		}
+		else if (SourceControlState->IsDeleted())
+		{
+			UE_LOG(LogConsoleResponse, Display, TEXT("fx.InvalidateCachedScripts failed\n\"NiagaraShaderVersion.ush\" is marked for delete"));
+			return;
+		}
+	}
+
+	IPlatformFile& PlatformFile = FPlatformFileManager::Get().GetPlatformFile();
+
+	IFileHandle* FileHandle = PlatformFile.OpenWrite(*FileName);
+	if (FileHandle)
+	{
+		FString Guid = FString(
+			TEXT("// This file is automatically generated by the console command fx.InvalidateCachedScripts\n")
+			TEXT("// Each time the console command is executed it generates a new GUID. As a hash of this file is included\n")
+			TEXT("// in the DDC key, it will automatically invalidate.\n")
+			TEXT("// \n")
+			TEXT("// If you are merging streams and there is a conflict with this GUID you should make a new GUID rather than taking one or the other.\n")
+			TEXT("// GUID = "))
+			+ FGuid::NewGuid().ToString();
+
+		FileHandle->Write((const uint8*)TCHAR_TO_ANSI(*Guid), Guid.Len());
+		delete FileHandle;
+
+		UE_LOG(LogConsoleResponse, Display, TEXT("fx.InvalidateCachedScripts succeeded\n\"NiagaraShaderVersion.ush\" was updated.\n"));
+	}
+	else
+	{
+		UE_LOG(LogConsoleResponse, Display, TEXT("fx.InvalidateCachedScripts failed\nCouldn't open \"NiagaraShaderVersion.ush\".\n"));
+	}
+}
+
+FAutoConsoleCommand InvalidateCachedNiagaraScripts(
+	TEXT("fx.InvalidateCachedScripts"),
+	TEXT("Invalidate Niagara script cache by making a unique change to NiagaraShaderVersion.ush which is included in common.usf.")
+	TEXT("To initiate actual the recompile of all shaders use \"recompileshaders changed\" or press \"Ctrl Shift .\".\n")
+	TEXT("The NiagaraShaderVersion.ush file should be automatically checked out but it needs to be checked in to have effect on other machines."),
+	FConsoleCommandWithArgsDelegate::CreateStatic(ExecuteInvalidateNiagaraCachedScripts)
+);
+
+void ExecuteRebuildNiagaraCachedScripts(const TArray< FString >& Args)
+{
+	UE_LOG(LogConsoleResponse, Display, TEXT("fx.RebuildDirtyScripts started.\n"));
+
+	// Need to flush the cache to make sure that we have the latest files.
+	FlushShaderFileCache();
+	for (TObjectIterator<UNiagaraSystem> SystemIterator; SystemIterator; ++SystemIterator)
+	{
+		SystemIterator->RequestCompile(false);
+	}
+}
+
+FAutoConsoleCommand ExecuteRebuildNiagaraCachedScriptsCmd(
+	TEXT("fx.RebuildDirtyScripts"),
+	TEXT("Go through all loaded assets and force them to recompute their script hash. If dirty, regenerate."),
+	FConsoleCommandWithArgsDelegate::CreateStatic(ExecuteRebuildNiagaraCachedScripts)
+);
+
+
 class FNiagaraSystemBoolParameterTrackEditor : public FNiagaraSystemParameterTrackEditor<UMovieSceneNiagaraBoolParameterTrack, UMovieSceneBoolSection>
 {
 	virtual TSharedRef<ISequencerSection> MakeSectionInterface(UMovieSceneSection& SectionObject, UMovieSceneTrack& Track, FGuid ObjectBinding) override
diff --git a/Engine/Plugins/FX/Niagara/Source/NiagaraEditor/Private/NiagaraEditorUtilities.cpp b/Engine/Plugins/FX/Niagara/Source/NiagaraEditor/Private/NiagaraEditorUtilities.cpp
--- a/Engine/Plugins/FX/Niagara/Source/NiagaraEditor/Private/NiagaraEditorUtilities.cpp
+++ b/Engine/Plugins/FX/Niagara/Source/NiagaraEditor/Private/NiagaraEditorUtilities.cpp
@@ -332,6 +332,293 @@ void FNiagaraEditorUtilities::WriteTextFileToDisk(FString SaveDirectory, FString
 	}
 }
 
+
+bool FNiagaraEditorUtilities::PODPropertyAppendCompileHash(const void* Container, UProperty* Property, const FString& PropertyName, FNiagaraCompileHashVisitor* InVisitor) 
+{
+	if (Property->IsA(UFloatProperty::StaticClass()))
+	{
+		UFloatProperty* CastProp = CastChecked<UFloatProperty>(Property);
+		float Value = CastProp->GetPropertyValue_InContainer(Container, 0);
+		InVisitor->UpdatePOD(*PropertyName, Value);
+		return true;
+	}
+	else if (Property->IsA(UIntProperty::StaticClass()))
+	{
+		UIntProperty* CastProp = CastChecked<UIntProperty>(Property);
+		int32 Value = CastProp->GetPropertyValue_InContainer(Container, 0);
+		InVisitor->UpdatePOD(*PropertyName, Value);
+		return true;
+	}
+	else if (Property->IsA(UInt16Property::StaticClass()))
+	{
+		UInt16Property* CastProp = CastChecked<UInt16Property>(Property);
+		int16 Value = CastProp->GetPropertyValue_InContainer(Container, 0);
+		InVisitor->UpdatePOD(*PropertyName, Value);
+		return true;
+	}
+	else if (Property->IsA(UUInt32Property::StaticClass()))
+	{
+		UUInt32Property* CastProp = CastChecked<UUInt32Property>(Property);
+		uint32 Value = CastProp->GetPropertyValue_InContainer(Container, 0);
+		InVisitor->UpdatePOD(*PropertyName, Value);
+		return true;
+	}
+	else if (Property->IsA(UUInt16Property::StaticClass()))
+	{
+		UUInt16Property* CastProp = CastChecked<UUInt16Property>(Property);
+		uint16 Value = CastProp->GetPropertyValue_InContainer(Container, 0);
+		InVisitor->UpdatePOD(*PropertyName, Value);
+		return true;
+	}
+	else if (Property->IsA(UByteProperty::StaticClass()))
+	{
+		UByteProperty* CastProp = CastChecked<UByteProperty>(Property);
+		uint8 Value = CastProp->GetPropertyValue_InContainer(Container, 0);
+		InVisitor->UpdatePOD(*PropertyName, Value);
+		return true;
+	}
+	else if (Property->IsA(UBoolProperty::StaticClass()))
+	{
+		UBoolProperty* CastProp = CastChecked<UBoolProperty>(Property);
+		bool Value = CastProp->GetPropertyValue_InContainer(Container, 0);
+		InVisitor->UpdatePOD(*PropertyName, Value);
+		return true;
+	}
+	else if (Property->IsA(UNameProperty::StaticClass()))
+	{
+		UNameProperty* CastProp = CastChecked<UNameProperty>(Property);
+		FName Value = CastProp->GetPropertyValue_InContainer(Container);
+		InVisitor->UpdateString(*PropertyName, Value.ToString());
+		return true;
+	}
+	else if (Property->IsA(UStrProperty::StaticClass()))
+	{
+		UStrProperty* CastProp = CastChecked<UStrProperty>(Property);
+		FString Value = CastProp->GetPropertyValue_InContainer(Container);
+		InVisitor->UpdateString(*PropertyName, Value);
+		return true;
+	}
+	return false;
+}
+
+bool FNiagaraEditorUtilities::NestedPropertiesAppendCompileHash(const void* Container, const UStruct* Struct, EFieldIteratorFlags::SuperClassFlags IteratorFlags, const FString& BaseName, FNiagaraCompileHashVisitor* InVisitor) 
+{
+	// We special case FNiagaraTypeDefinitions here because they need to write out a lot more than just their standalone uproperties.
+	if (Struct == FNiagaraTypeDefinition::StaticStruct())
+	{
+		FNiagaraTypeDefinition* TypeDef = (FNiagaraTypeDefinition*)Container;
+		if (TypeDef)
+		{
+			TypeDef->AppendCompileHash(InVisitor);
+			return true;
+		}
+	}
+
+	TFieldIterator<UProperty> PropertyCountIt(Struct, IteratorFlags);
+	int32 NumProperties = 0;
+	for (; PropertyCountIt; ++PropertyCountIt)
+	{
+		NumProperties++;
+	}
+
+	for (TFieldIterator<UProperty> PropertyIt(Struct, IteratorFlags); PropertyIt; ++PropertyIt)
+	{
+		UProperty* Property = *PropertyIt;
+
+		static FName SkipMeta = TEXT("SkipForCompileHash");
+		if (Property->HasMetaData(SkipMeta))
+		{
+			continue;
+		}
+		
+		FString PropertyName = (NumProperties == 1) ? (BaseName) : (BaseName + "." + Property->GetName());
+
+		if (PODPropertyAppendCompileHash(Container, Property, PropertyName, InVisitor))
+		{
+			continue;
+		}
+		else if (Property->IsA(UMapProperty::StaticClass()))
+		{
+			UMapProperty* CastProp = CastChecked<UMapProperty>(Property);
+			FScriptMapHelper MapHelper(CastProp, CastProp->ContainerPtrToValuePtr<void>(Container));
+			InVisitor->UpdatePOD(*PropertyName, MapHelper.Num());
+			if (MapHelper.GetKeyProperty())
+			{
+				InVisitor->UpdateString(TEXT("KeyPathname"), MapHelper.GetKeyProperty()->GetPathName());
+				InVisitor->UpdateString(TEXT("ValuePathname"), MapHelper.GetValueProperty()->GetPathName());
+
+				// We currently only support maps with keys of FNames. Anything else should generate a warning.
+				if (MapHelper.GetKeyProperty()->GetClass() == UNameProperty::StaticClass())
+				{
+					// To be safe, let's gather up all the keys and sort them lexicographically so that this is stable across application runs.
+					TArray<FName> Names;
+					Names.AddUninitialized(MapHelper.Num());
+					for (int32 i = 0; i < MapHelper.Num(); i++)
+					{
+						FName* KeyPtr = (FName*)(MapHelper.GetKeyPtr(i));
+						if (KeyPtr)
+						{
+							Names[i] = *KeyPtr;
+						}
+						else
+						{
+							Names[i] = FName();
+							UE_LOG(LogNiagaraEditor, Warning, TEXT("Bad key in %s at %d"), *Property->GetName(), i);
+						}
+					}
+					// Sort stably over runs
+					Names.Sort(FNameLexicalLess());
+
+					// Now hash out the values directly.
+					// We support map values of POD types or map values of structs with POD types internally. Anything else we should generate a warning on.
+					if (MapHelper.GetValueProperty()->IsA(UStructProperty::StaticClass()))
+					{
+						bool bPassed = true;
+						UStructProperty* StructProp = CastChecked<UStructProperty>(MapHelper.GetValueProperty());
+
+						for (int32 ArrayIdx = 0; ArrayIdx < MapHelper.Num(); ArrayIdx++)
+						{
+							InVisitor->UpdateString(*FString::Printf(TEXT("Key[%d]"), ArrayIdx), Names[ArrayIdx].ToString());
+							if (!NestedPropertiesAppendCompileHash(MapHelper.GetValuePtr(ArrayIdx), StructProp->Struct, EFieldIteratorFlags::IncludeSuper, FString::Printf(TEXT("Value[%d]"), ArrayIdx), InVisitor))
+							{
+								UE_LOG(LogNiagaraEditor, Warning, TEXT("Skipping %s because it is an map value property of unsupported underlying type, please add \"meta = (SkipForCompileHash=\"true\")\" to avoid this warning in the future or handle it yourself in NestedPropertiesAppendCompileHash!"), *Property->GetName());
+								bPassed = false;
+								continue;
+							}
+						}
+						if (bPassed)
+						{
+							continue;
+						}
+					}
+					else
+					{
+						bool bPassed = true;
+						for (int32 ArrayIdx = 0; ArrayIdx < MapHelper.Num(); ArrayIdx++)
+						{
+							InVisitor->UpdateString(*FString::Printf(TEXT("Key[%d]"), ArrayIdx), Names[ArrayIdx].ToString());
+							if (!PODPropertyAppendCompileHash(MapHelper.GetValuePtr(ArrayIdx), MapHelper.GetValueProperty(), FString::Printf(TEXT("Value[%d]"), ArrayIdx), InVisitor))
+							{
+								UE_LOG(LogNiagaraEditor, Warning, TEXT("Skipping %s because it is an map value property of unsupported underlying type, please add \"meta = (SkipForCompileHash=\"true\")\" to avoid this warning in the future or handle it yourself in NestedPropertiesAppendCompileHash!"), *Property->GetName());
+								bPassed = false;
+								continue;
+							}
+						}
+						if (bPassed)
+						{
+							continue;
+						}
+					}
+				}
+				else
+				{
+					UE_LOG(LogNiagaraEditor, Warning, TEXT("Skipping %s because it is a map property, please add \"meta = (SkipForCompileHash=\"true\")\" to avoid this warning in the future or handle it yourself in NestedPropertiesAppendCompileHash!"), *Property->GetName());
+				}
+			}
+			continue;
+		}
+		else if (Property->IsA(UArrayProperty::StaticClass()))
+		{
+			UArrayProperty* CastProp = CastChecked<UArrayProperty>(Property);
+
+			FScriptArrayHelper ArrayHelper(CastProp, CastProp->ContainerPtrToValuePtr<void>(Container));
+			InVisitor->UpdatePOD(*PropertyName, ArrayHelper.Num());
+			InVisitor->UpdateString(TEXT("InnerPathname"), CastProp->Inner->GetPathName());
+
+			// We support arrays of POD types or arrays of structs with POD types internally. Anything else we should generate a warning on.
+			if (CastProp->Inner->IsA(UStructProperty::StaticClass()))
+			{
+				bool bPassed = true;
+				UStructProperty* StructProp = CastChecked<UStructProperty>(CastProp->Inner);
+
+				for (int32 ArrayIdx = 0; ArrayIdx < ArrayHelper.Num(); ArrayIdx++)
+				{
+					if (!NestedPropertiesAppendCompileHash(ArrayHelper.GetRawPtr(ArrayIdx), StructProp->Struct, EFieldIteratorFlags::IncludeSuper, PropertyName, InVisitor))
+					{
+						UE_LOG(LogNiagaraEditor, Warning, TEXT("Skipping %s because it is an array property of unsupported underlying type, please add \"meta = (SkipForCompileHash=\"true\")\" to avoid this warning in the future or handle it yourself in NestedPropertiesAppendCompileHash!"), *Property->GetName());
+						bPassed = false;
+						continue;
+					}
+				}
+				if (bPassed)
+				{
+					continue;
+				}
+			}
+			else
+			{
+				bool bPassed = true;
+				for (int32 ArrayIdx = 0; ArrayIdx < ArrayHelper.Num(); ArrayIdx++)
+				{
+					if (!PODPropertyAppendCompileHash(ArrayHelper.GetRawPtr(ArrayIdx), CastProp->Inner, PropertyName, InVisitor))
+					{
+						if (bPassed)
+						{
+							UE_LOG(LogNiagaraEditor, Warning, TEXT("Skipping %s because it is an array property of unsupported underlying type, please add \"meta = (SkipForCompileHash=\"true\")\" to avoid this warning in the future or handle it yourself in NestedPropertiesAppendCompileHash!"), *Property->GetName());
+						}
+						bPassed = false;
+						continue;
+					}
+				}
+				if (bPassed)
+				{
+					continue;
+				}
+			}
+
+			UE_LOG(LogNiagaraEditor, Warning, TEXT("Skipping %s because it is an array property, please add \"meta = (SkipForCompileHash=\"true\")\" to avoid this warning in the future or handle it yourself in NestedPropertiesAppendCompileHash!"), *Property->GetName());
+			continue;
+		}
+		else if (Property->IsA(UTextProperty::StaticClass()))
+		{
+			UTextProperty* CastProp = CastChecked<UTextProperty>(Property);
+			UE_LOG(LogNiagaraEditor, Warning, TEXT("Skipping %s because it is a UText property, please add \"meta = (SkipForCompileHash=\"true\")\" to avoid this warning in the future or handle it yourself in NestedPropertiesAppendCompileHash!"), *Property->GetName());
+			return true;
+		}
+		else if (Property->IsA(UEnumProperty::StaticClass()))
+		{
+			UEnumProperty* CastProp = CastChecked<UEnumProperty>(Property);
+			const void* EnumContainer = Property->ContainerPtrToValuePtr<uint8>(Container);
+			if (PODPropertyAppendCompileHash(EnumContainer, CastProp->GetUnderlyingProperty(), PropertyName, InVisitor))
+			{
+				continue;
+			}
+			check(false);
+			return false;
+		}
+		else if (Property->IsA(UObjectProperty::StaticClass()))
+		{
+			UObjectProperty* CastProp = CastChecked<UObjectProperty>(Property);
+			UObject* Obj = CastProp->GetObjectPropertyValue_InContainer(Container);
+			if (Obj != nullptr)
+			{
+				// We just do name here as sometimes things will be in a transient package or something tricky.
+				// Because we do nested id's for each called graph, it should work out in the end to have a different
+				// value in the compile array if the scripts are the same name but different locations.
+				InVisitor->UpdateString(*PropertyName, Obj->GetName());
+			}
+			else
+			{
+				InVisitor->UpdateString(*PropertyName, TEXT("nullptr"));
+			}
+			continue;
+		}
+		else if (Property->IsA(UStructProperty::StaticClass()))
+		{
+			UStructProperty* StructProp = CastChecked<UStructProperty>(Property);
+			const void* StructContainer = Property->ContainerPtrToValuePtr<uint8>(Container);
+			NestedPropertiesAppendCompileHash(StructContainer, StructProp->Struct, EFieldIteratorFlags::IncludeSuper, PropertyName, InVisitor);
+			continue;
+		}
+		else
+		{
+			check(false);
+			return false;
+		}
+	}
+	return true;
+}
+
 void FNiagaraEditorUtilities::GatherChangeIds(UNiagaraEmitter& Emitter, TMap<FGuid, FGuid>& ChangeIds, const FString& InDebugName, bool bWriteToLogDir)
 {
 	FString ExportText;
diff --git a/Engine/Plugins/FX/Niagara/Source/NiagaraEditor/Private/NiagaraGraph.cpp b/Engine/Plugins/FX/Niagara/Source/NiagaraEditor/Private/NiagaraGraph.cpp
--- a/Engine/Plugins/FX/Niagara/Source/NiagaraEditor/Private/NiagaraGraph.cpp
+++ b/Engine/Plugins/FX/Niagara/Source/NiagaraEditor/Private/NiagaraGraph.cpp
@@ -34,6 +34,8 @@
 #include "NiagaraScriptVariable.h"
 #include "NiagaraHlslTranslator.h"
 #include "NiagaraNodeFunctionCall.h"
+#include "NiagaraScriptVariable.h"
+#include "Misc/SecureHash.h"
 
 DECLARE_CYCLE_STAT(TEXT("NiagaraEditor - Graph - FindInputNodes"), STAT_NiagaraEditor_Graph_FindInputNodes, STATGROUP_NiagaraEditor);
 DECLARE_CYCLE_STAT(TEXT("NiagaraEditor - Graph - FindInputNodes_NotFilterUsage"), STAT_NiagaraEditor_Graph_FindInputNodes_NotFilterUsage, STATGROUP_NiagaraEditor);
@@ -47,6 +49,14 @@ bool bWriteToLog = false;
 
 #define LOCTEXT_NAMESPACE "NiagaraGraph"
 
+int32 GNiagaraUseGraphHash = 1;
+static FAutoConsoleVariableRef CVarNiagaraUseGraphHash(
+	TEXT("fx.UseNewGraphHash"),
+	GNiagaraUseGraphHash,
+	TEXT("If > 0 a hash of the graph node state will be used, otherwise will use the older code path. \n"),
+	ECVF_Default
+);
+
 FNiagaraGraphParameterReferenceCollection::FNiagaraGraphParameterReferenceCollection(const bool bInCreated)
 	: Graph(nullptr), bCreated(bInCreated)
 {
@@ -75,6 +85,7 @@ void FNiagaraGraphScriptUsageInfo::PostLoad(UObject* Owner)
 		if (CompileHash.IsValid() == false && DataHash_DEPRECATED.Num() == FNiagaraCompileHash::HashSize)
 		{
 			CompileHash = FNiagaraCompileHash(DataHash_DEPRECATED);
+			CompileHashFromGraph = FNiagaraCompileHash(DataHash_DEPRECATED);
 		}
 	}
 }
@@ -335,7 +346,14 @@ FNiagaraCompileHash UNiagaraGraph::GetCompileDataHash(ENiagaraScriptUsage InUsag
 	{
 		if (UNiagaraScript::IsEquivalentUsage(CachedUsageInfo[i].UsageType, InUsage) && CachedUsageInfo[i].UsageId == InUsageId)
 		{
-			return CachedUsageInfo[i].CompileHash;
+			if (GNiagaraUseGraphHash <= 0)
+			{
+				return CachedUsageInfo[i].CompileHash;
+			}
+			else
+			{
+				return CachedUsageInfo[i].CompileHashFromGraph;
+			}
 		}
 	}
 	return FNiagaraCompileHash();
@@ -1076,6 +1094,110 @@ FNiagaraTypeDefinition UNiagaraGraph::GetCachedNumericConversion(class UEdGraphP
 	return ReturnDef;
 }
 
+bool UNiagaraGraph::AppendCompileHash(FNiagaraCompileHashVisitor* InVisitor, const TArray<UNiagaraNode*>& InTraversal) const
+{
+#if WITH_EDITORONLY_DATA
+	int32 Index = InVisitor->Values.AddDefaulted();
+	InVisitor->Values[Index].Object = FString::Printf(TEXT("Class: \"%s\"  Name: \"%s\""), *this->GetClass()->GetName(),  *this->GetName());
+#endif
+	InVisitor->UpdateString(TEXT("ForceRebuildId"), ForceRebuildId.ToString());
+
+	// We need to sort the variables in a stable manner.
+	TArray<UNiagaraScriptVariable*> Values;
+	VariableToScriptVariable.GenerateValueArray(Values);
+	Values.Sort([&](const UNiagaraScriptVariable& A, const UNiagaraScriptVariable& B)
+	{
+		const FName& AName = A.Variable.GetName();
+		const FName& BName = B.Variable.GetName();
+
+		return AName.LexicalLess(BName);
+	});
+
+
+	// Write all the values of the local variables to the visitor as they could potentially influence compilation.
+	for (const UNiagaraScriptVariable* Var : Values)
+	{
+		FNiagaraGraphParameterReferenceCollection* Collection = ParameterToReferencesMap.Find(Var->Variable);
+		if (Collection && Collection->ParameterReferences.Num() > 0)
+		{
+			// Only add parameters to this hash that could potentially affect this compile.
+			bool bFoundInTraversal = false;
+			for (const FNiagaraGraphParameterReference& Ref : Collection->ParameterReferences)
+			{
+				if (InTraversal.Contains(Ref.Value.Get()))
+				{
+					bFoundInTraversal = true;
+					break;
+				}
+			}
+
+			if (!bFoundInTraversal)
+			{
+				continue;
+			}
+
+		#if WITH_EDITORONLY_DATA
+			Index = InVisitor->Values.AddDefaulted();
+			InVisitor->Values[Index].Object = FString::Printf(TEXT("Class: \"%s\"  Name: \"%s\""), *Var->GetClass()->GetName(), *Var->Variable.GetName().ToString());
+		#endif
+			verify(Var->AppendCompileHash(InVisitor));
+		}
+	}
+
+	// Write all the values of the nodes to the visitor as they could influence compilation.
+	for (const UNiagaraNode* Node : InTraversal)
+	{
+#if WITH_EDITORONLY_DATA
+		Index = InVisitor->Values.AddDefaulted();
+		InVisitor->Values[Index].Object = FString::Printf(TEXT("Class: \"%s\" Title: \"%s\" Name: \"%s\" Guid: %s"), *Node->GetClass()->GetName(), *Node->GetNodeTitle(ENodeTitleType::EditableTitle).ToString(), *Node->GetName(), *LexToString(Node->NodeGuid));
+#endif
+		verify(Node->AppendCompileHash(InVisitor));
+	}
+
+#if WITH_EDITORONLY_DATA
+	// Optionally log out the information for debugging.
+	if (FNiagaraCompileHashVisitor::LogCompileIdGeneration == 2 && InTraversal.Num() > 0)
+	{
+		UE_LOG(LogNiagaraEditor, Log, TEXT("UNiagaraGraph::AppendCompileHash %s %s\n==========================="), *GetFullName(), *InTraversal[InTraversal.Num()- 1]->GetNodeTitle(ENodeTitleType::ListView).ToString());
+		for (int32 i = 0; i < InVisitor->Values.Num(); i++)
+		{
+			UE_LOG(LogNiagaraEditor, Log, TEXT("Object[%d]: %s"), i, *InVisitor->Values[i].Object);
+			ensure(InVisitor->Values[i].PropertyKeys.Num() == InVisitor->Values[i].PropertyValues.Num());
+			for (int32 j = 0; j < InVisitor->Values[i].PropertyKeys.Num(); j++)
+			{
+				UE_LOG(LogNiagaraEditor, Log, TEXT("\tProperty[%d]: %s = %s"), j, *InVisitor->Values[i].PropertyKeys[j], *InVisitor->Values[i].PropertyValues[j]);
+			}
+		}
+	}
+#endif
+	return true;
+}
+
+void DiffProperties(const FNiagaraCompileHashVisitorDebugInfo& A, const FNiagaraCompileHashVisitorDebugInfo& B)
+{
+	if (A.PropertyKeys.Num() != B.PropertyKeys.Num())
+	{
+		UE_LOG(LogNiagaraEditor, Log, TEXT("Hash Difference: Property Count Mismatch %d vs %d on %s"), A.PropertyKeys.Num(), B.PropertyKeys.Num(), *A.Object);
+	}
+	else
+	{
+		bool bFoundMatch = false;
+		for (int32 i = 0; i < A.PropertyKeys.Num(); i++)
+		{
+			if (A.PropertyKeys[i] == B.PropertyKeys[i])
+			{
+				bFoundMatch = true;
+				if (A.PropertyValues[i] != B.PropertyValues[i])
+				{
+					UE_LOG(LogNiagaraEditor, Log, TEXT("Hash Difference: Property Value Mismatch %s vs %s on property %s of %s"), *A.PropertyValues[i], *B.PropertyValues[i], *A.PropertyKeys[i], *A.Object);
+				}
+			}
+		}
+		ensure(bFoundMatch);
+	}
+}
+
+
 void UNiagaraGraph::RebuildCachedCompileIds(bool bForce)
 {
 	// If the graph hasn't changed since last rebuild, then do nothing.
@@ -1128,6 +1250,11 @@ void UNiagaraGraph::RebuildCachedCompileIds(bool bForce)
 
 		// Now compare the change id's of all the nodes in the traversal by hashing them up and comparing the hash
 		// now with the hash from previous runs.
+		FSHA1 GraphHashState;
+		FNiagaraCompileHashVisitor Visitor(GraphHashState);
+		AppendCompileHash(&Visitor, NewUsageCache[i].Traversal);
+		GraphHashState.Final();
+
 		FSHA1 HashState;
 		for (UNiagaraNode* Node : NewUsageCache[i].Traversal)
 		{
@@ -1136,10 +1263,64 @@ void UNiagaraGraph::RebuildCachedCompileIds(bool bForce)
 		HashState.Final();
 
 		// We can't store in a FShaHash struct directly because you can't FProperty it. Using a standin of the same size.
-		TArray<uint8> DataHash;
-		DataHash.AddUninitialized(20);
-		HashState.GetHash(DataHash.GetData());
-		NewUsageCache[i].CompileHash = FNiagaraCompileHash(DataHash);
+		{
+			TArray<uint8> DataHash;
+			DataHash.AddUninitialized(FSHA1::DigestSize);
+			HashState.GetHash(DataHash.GetData());
+		
+			NewUsageCache[i].CompileHash = FNiagaraCompileHash(DataHash);
+		}
+
+
+		{
+			// We can't store in a FShaHash struct directly because you can't UProperty it. Using a standin of the same size.
+			TArray<uint8> DataHash;
+			DataHash.AddUninitialized(FSHA1::DigestSize);
+			GraphHashState.GetHash(DataHash.GetData());
+			NewUsageCache[i].CompileHashFromGraph = FNiagaraCompileHash(DataHash);
+			NewUsageCache[i].CompileLastObjects = Visitor.Values;
+
+#if WITH_EDITORONLY_DATA
+			// Log out all the entries that differ!
+			if (FNiagaraCompileHashVisitor::LogCompileIdGeneration != 0 && FoundMatchIdx != -1 && NewUsageCache[i].CompileHashFromGraph != CachedUsageInfo[FoundMatchIdx].CompileHashFromGraph)
+			{
+				TArray<FNiagaraCompileHashVisitorDebugInfo> OldDebugValues;
+				OldDebugValues = CachedUsageInfo[FoundMatchIdx].CompileLastObjects;
+
+				TArray<bool> bFoundIndices; // Record if we ever found these.
+				bFoundIndices.AddZeroed(OldDebugValues.Num());
+
+				for (int32 ObjIdx = 0; ObjIdx < NewUsageCache[i].CompileLastObjects.Num(); ObjIdx++)
+				{
+					bool bFound = false;
+					for (int32 OldObjIdx = 0; OldObjIdx < OldDebugValues.Num(); OldObjIdx++)
+					{
+						if (OldDebugValues[OldObjIdx].Object == NewUsageCache[i].CompileLastObjects[ObjIdx].Object)
+						{
+							bFound = true; // Record that we found a match for this object
+							bFoundIndices[OldObjIdx] = true; // Record that we found an overall match
+							DiffProperties(OldDebugValues[OldObjIdx], NewUsageCache[i].CompileLastObjects[ObjIdx]);
+						}
+					}
+
+					if (!bFound)
+					{
+						UE_LOG(LogNiagaraEditor, Log, TEXT("Hash Difference: New Object: %s"), *NewUsageCache[i].CompileLastObjects[ObjIdx].Object);
+					}
+				}
+
+
+				for (int32 ObjIdx = 0; ObjIdx < OldDebugValues.Num(); ObjIdx++)
+				{
+					if (bFoundIndices[ObjIdx] == false)
+					{
+						UE_LOG(LogNiagaraEditor, Log, TEXT("Hash Difference: Removed Object: %s"), *OldDebugValues[ObjIdx].Object);
+					}					
+				}
+			}
+#endif
+		}
+
 
 		// TODO sckime debug logic... should be disabled or put under a cvar in the future
 		{
@@ -1190,6 +1371,7 @@ void UNiagaraGraph::RebuildCachedCompileIds(bool bForce)
 		TArray<uint8> DataHash;
 		DataHash.AddZeroed(20);
 		GpuUsageInfo.CompileHash = FNiagaraCompileHash(DataHash);
+		GpuUsageInfo.CompileHashFromGraph = FNiagaraCompileHash(DataHash);
 
 		NewUsageCache.Add(GpuUsageInfo);
 	}
@@ -1284,14 +1466,15 @@ void UNiagaraGraph::ResolveNumerics(TMap<UNiagaraNode*, bool>& VisitedNodes, UEd
 	}
 }
 
-void UNiagaraGraph::InvalidateCachedCompileIds()
+void UNiagaraGraph::ForceGraphToRecompileOnNextCheck()
 {
 	Modify();
 	CachedUsageInfo.Empty();
+	ForceRebuildId = FGuid::NewGuid();
 	MarkGraphRequiresSynchronization(__FUNCTION__);
 }
 
-void UNiagaraGraph::GatherExternalDependencyData(ENiagaraScriptUsage InUsage, const FGuid& InUsageId, TArray<FNiagaraCompileHash>& InReferencedCompileHashes, TArray<UObject*>& InReferencedObjs)
+void UNiagaraGraph::GatherExternalDependencyData(ENiagaraScriptUsage InUsage, const FGuid& InUsageId, TArray<FNiagaraCompileHash>& InReferencedCompileHashes, TArray<FString>& InReferencedObjs)
 {
 	RebuildCachedCompileIds();
 	
@@ -1308,8 +1491,15 @@ void UNiagaraGraph::GatherExternalDependencyData(ENiagaraScriptUsage InUsage, co
 		// Now add any other dependency chains that we might have...
 		else if (UNiagaraScript::IsUsageDependentOn(InUsage, CachedUsageInfo[i].UsageType))
 		{
-			InReferencedCompileHashes.Add(CachedUsageInfo[i].CompileHash);
-			InReferencedObjs.Add(CachedUsageInfo[i].Traversal.Last());
+			if (GNiagaraUseGraphHash == 1)
+			{
+				InReferencedCompileHashes.Add(CachedUsageInfo[i].CompileHashFromGraph);
+			}
+			else
+			{
+				InReferencedCompileHashes.Add(CachedUsageInfo[i].CompileHash);
+			}
+			InReferencedObjs.Add(CachedUsageInfo[i].Traversal.Last()->GetPathName());
 
 			for (UNiagaraNode* Node : CachedUsageInfo[i].Traversal)
 			{
diff --git a/Engine/Plugins/FX/Niagara/Source/NiagaraEditor/Private/NiagaraNode.cpp b/Engine/Plugins/FX/Niagara/Source/NiagaraEditor/Private/NiagaraNode.cpp
--- a/Engine/Plugins/FX/Niagara/Source/NiagaraEditor/Private/NiagaraNode.cpp
+++ b/Engine/Plugins/FX/Niagara/Source/NiagaraEditor/Private/NiagaraNode.cpp
@@ -11,7 +11,10 @@
 #include "ToolMenu.h"
 #include "GraphEditorActions.h"
 #include "Serialization/PropertyLocalizationDataGathering.h"
+#include "NiagaraEditorUtilities.h"
 
+#include "Misc/Guid.h"
+#include "UObject/TextProperty.h"
 
 #define LOCTEXT_NAMESPACE "NiagaraNode"
 
@@ -35,6 +38,123 @@ void UNiagaraNode::PostLoad()
 	}
 }
 
+bool UNiagaraNode::NestedPropertiesAppendCompileHash(const void* Container, const UStruct* Struct, EFieldIteratorFlags::SuperClassFlags IteratorFlags, const FString& BaseName, FNiagaraCompileHashVisitor* InVisitor) const
+{
+	return FNiagaraEditorUtilities::NestedPropertiesAppendCompileHash(Container, Struct, IteratorFlags, BaseName, InVisitor);
+}
+
+bool UNiagaraNode::PODPropertyAppendCompileHash(const void* Container, UProperty* Property, const FString& PropertyName, FNiagaraCompileHashVisitor* InVisitor) const
+{
+	return FNiagaraEditorUtilities::PODPropertyAppendCompileHash(Container, Property, PropertyName, InVisitor);
+}
+
+
+bool UNiagaraNode::GenerateCompileHashForClassMembers(const UClass* InClass, FNiagaraCompileHashVisitor* InVisitor) const
+{
+	if (!NestedPropertiesAppendCompileHash(static_cast<const void*>(this), InClass, EFieldIteratorFlags::ExcludeSuper, InClass->GetName(), InVisitor))
+	{
+		return false;
+	}
+	return true;
+}
+
+bool UNiagaraNode::AppendCompileHash(FNiagaraCompileHashVisitor* InVisitor) const
+{
+	if (NiagaraNodeAppendCompileHash(InVisitor))
+	{
+		UClass* Class = GetClass();
+		UClass* NiagaraNodeClass = UNiagaraNode::StaticClass();
+
+		while (Class != NiagaraNodeClass && Class != nullptr)
+		{
+			if (!GenerateCompileHashForClassMembers(Class, InVisitor))
+			{
+				return false;
+			}
+
+			Class = Class->GetSuperClass();
+		}
+		return true;
+	}
+	return false;
+}
+
+
+
+bool UNiagaraNode::PinAppendCompileHash(const UEdGraphPin* InPin, FNiagaraCompileHashVisitor* InVisitor) const
+{
+#if WITH_EDITORONLY_DATA
+	InVisitor->UpdateReference(TEXT("\tPinId"), InPin);
+	InVisitor->UpdateString(TEXT("\tName"), InPin->GetName());
+	InVisitor->UpdatePOD(TEXT("\tDirection"), InPin->Direction);
+	
+	const UEdGraphSchema_Niagara* Schema = CastChecked<UEdGraphSchema_Niagara>(GetSchema());
+	check(Schema);
+
+	FNiagaraTypeDefinition PinType = Schema->PinToTypeDefinition(InPin);
+	PinType.AppendCompileHash(InVisitor);
+
+
+	if (InPin->LinkedTo.Num() > 0)
+	{
+		for (int32 i = 0; i < InPin->LinkedTo.Num(); i++)
+		{
+			InVisitor->UpdateReference(TEXT("\tLinkedTo"), InPin->LinkedTo[i]);
+		}
+	}
+	else
+	{
+		InVisitor->UpdateString(TEXT("\tDefaultValue"), InPin->DefaultValue);
+		InVisitor->UpdateString(TEXT("\tAutogeneratedDefaultValue"), InPin->AutogeneratedDefaultValue);
+		bool bDefaultObj = InPin->DefaultObject != nullptr;
+		InVisitor->UpdatePOD(TEXT("\tDefaultObject"), bDefaultObj);
+		if (InPin->DefaultObject)
+		{
+			UE_LOG(LogNiagaraEditor, Warning, TEXT("AppendCompileHash can't handle default objects properly, please investigate!"));
+		}
+	}
+
+	for (int32 i = 0; i < InPin->SubPins.Num(); i++)
+	{
+		PinAppendCompileHash(InPin->SubPins[i], InVisitor);
+	}
+	
+	if (InPin->ParentPin != nullptr)
+	{
+		InVisitor->UpdateReference(TEXT("\tParentPin"), InPin->ParentPin);
+	}
+	return true;
+#else
+	return false;
+#endif
+}
+
+bool UNiagaraNode::NiagaraNodeAppendCompileHash(FNiagaraCompileHashVisitor* InVisitor) const
+{
+#if WITH_EDITORONLY_DATA
+	UClass* Class = GetClass();
+	InVisitor->UpdateString(TEXT("Classname"), Class->GetName());
+	InVisitor->UpdateReference(TEXT("Node"), this);
+	// DO NOT include the name in the hash as it makes copy/paste duplicates not work out b/c the names internally will be incremented by a digit. 
+	// The state and iteration occurrence should be enough to uniquely identify.
+	//InVisitor->UpdateString(TEXT("NodeName"), GetName()); 
+	ENodeEnabledState NodeEnabledState = GetDesiredEnabledState();
+	InVisitor->UpdatePOD(TEXT("EnabledState"), NodeEnabledState);
+
+	for (const UEdGraphPin* Pin : Pins)
+	{
+		if (!PinAppendCompileHash(Pin, InVisitor))
+		{
+			return false;
+		}
+	}
+
+	return true;
+#else
+	return false;
+#endif
+}
+
 bool UNiagaraNode::ReallocatePins(bool bMarkNeedsResynchronizeOnChange)
 {
 	Modify();
diff --git a/Engine/Plugins/FX/Niagara/Source/NiagaraEditor/Private/NiagaraNodeAssignment.cpp b/Engine/Plugins/FX/Niagara/Source/NiagaraEditor/Private/NiagaraNodeAssignment.cpp
--- a/Engine/Plugins/FX/Niagara/Source/NiagaraEditor/Private/NiagaraNodeAssignment.cpp
+++ b/Engine/Plugins/FX/Niagara/Source/NiagaraEditor/Private/NiagaraNodeAssignment.cpp
@@ -182,7 +182,7 @@ void UNiagaraNodeAssignment::BuildParameterMapHistory(FNiagaraParameterMapHistor
 	Super::BuildParameterMapHistory(OutHistory, bRecursive, bFilterForCompilation);
 }
 
-void UNiagaraNodeAssignment::GatherExternalDependencyData(ENiagaraScriptUsage InMasterUsage, const FGuid& InMasterUsageId, TArray<FNiagaraCompileHash>& InReferencedCompileHashes, TArray<UObject*>& InReferencedObjs) const
+void UNiagaraNodeAssignment::GatherExternalDependencyData(ENiagaraScriptUsage InMasterUsage, const FGuid& InMasterUsageId, TArray<FNiagaraCompileHash>& InReferencedCompileHashes, TArray<FString>& InReferencedObjs) const
 {
 	// Assignment nodes own their function graphs and therefore have no external dependencies so we override the default function behavior here to avoid 
 	// adding additional non-deterministic guids to the compile id generation which can invalid the DDC for compiled scripts, especially during emitter merging.
diff --git a/Engine/Plugins/FX/Niagara/Source/NiagaraEditor/Private/NiagaraNodeEmitter.cpp b/Engine/Plugins/FX/Niagara/Source/NiagaraEditor/Private/NiagaraNodeEmitter.cpp
--- a/Engine/Plugins/FX/Niagara/Source/NiagaraEditor/Private/NiagaraNodeEmitter.cpp
+++ b/Engine/Plugins/FX/Niagara/Source/NiagaraEditor/Private/NiagaraNodeEmitter.cpp
@@ -88,6 +88,36 @@ bool UNiagaraNodeEmitter::CommitEditablePinName(const FText& InName, UEdGraphPin
 	return false;
 }
 
+bool UNiagaraNodeEmitter::GenerateCompileHashForClassMembers(const UClass* InClass, FNiagaraCompileHashVisitor* InVisitor) const
+{
+	if (InClass == UNiagaraNodeEmitter::StaticClass())
+	{
+		// For emitters, we really just want the emitter name.
+		FName EmitterName;
+		if (OwnerSystem != nullptr && EmitterHandleId.IsValid())
+		{
+			for (const FNiagaraEmitterHandle& EmitterHandle : OwnerSystem->GetEmitterHandles())
+			{
+				if (EmitterHandle.GetId() == EmitterHandleId)
+				{
+					EmitterName = (EmitterHandle.GetName());
+					break;
+				}
+			}
+		}
+		else if (CachedUniqueName.IsValid())
+		{
+			EmitterName = (CachedUniqueName);
+		}
+
+		InVisitor->UpdateString(TEXT("EmitterName"), EmitterName.ToString());
+		return true;
+	}
+	else
+	{
+		return Super::GenerateCompileHashForClassMembers(InClass, InVisitor);
+	}
+}
 
 void UNiagaraNodeEmitter::AllocateDefaultPins()
 {
@@ -460,7 +490,7 @@ void UNiagaraNodeEmitter::Compile(FHlslNiagaraTranslator *Translator, TArray<int
 	}
 }
 
-void UNiagaraNodeEmitter::GatherExternalDependencyData(ENiagaraScriptUsage InMasterUsage, const FGuid& InMasterUsageId, TArray<FNiagaraCompileHash>& InReferencedCompileHashes, TArray<UObject*>& InReferencedObjs) const
+void UNiagaraNodeEmitter::GatherExternalDependencyData(ENiagaraScriptUsage InMasterUsage, const FGuid& InMasterUsageId, TArray<FNiagaraCompileHash>& InReferencedCompileHashes, TArray<FString>& InReferencedObjs) const
 {
 	UNiagaraGraph* CalledGraph = GetCalledGraph();
 
@@ -468,7 +498,7 @@ void UNiagaraNodeEmitter::GatherExternalDependencyData(ENiagaraScriptUsage InMas
 	{
 		ENiagaraScriptUsage TargetUsage = InMasterUsage == ENiagaraScriptUsage::SystemSpawnScript ? ENiagaraScriptUsage::EmitterSpawnScript : ENiagaraScriptUsage::EmitterUpdateScript;
 		InReferencedCompileHashes.Add(CalledGraph->GetCompileDataHash(TargetUsage, FGuid(0,0,0,0)));
-		InReferencedObjs.Add(CalledGraph);
+		InReferencedObjs.Add(CalledGraph->GetPathName());
 		CalledGraph->GatherExternalDependencyData(TargetUsage, FGuid(0, 0, 0, 0), InReferencedCompileHashes, InReferencedObjs);
 	}
 }
diff --git a/Engine/Plugins/FX/Niagara/Source/NiagaraEditor/Private/NiagaraNodeEmitter.h b/Engine/Plugins/FX/Niagara/Source/NiagaraEditor/Private/NiagaraNodeEmitter.h
--- a/Engine/Plugins/FX/Niagara/Source/NiagaraEditor/Private/NiagaraNodeEmitter.h
+++ b/Engine/Plugins/FX/Niagara/Source/NiagaraEditor/Private/NiagaraNodeEmitter.h
@@ -60,22 +60,24 @@ public:
 	UNiagaraGraph* GetCalledGraph() const;
 	
 	virtual void Compile(FHlslNiagaraTranslator *Translator, TArray<int32>& Outputs) override;
-	virtual void GatherExternalDependencyData(ENiagaraScriptUsage InMasterUsage, const FGuid& InMasterUsageId, TArray<FNiagaraCompileHash>& InReferencedCompileHashes, TArray<UObject*>& InReferencedObjs) const override;
+	virtual void GatherExternalDependencyData(ENiagaraScriptUsage InMasterUsage, const FGuid& InMasterUsageId, TArray<FNiagaraCompileHash>& InReferencedCompileHashes, TArray<FString>& InReferencedObjs) const override;
 
 	void SetCachedVariablesForCompilation(const FName& InUniqueName, UNiagaraGraph* InGraph, UNiagaraScriptSourceBase* InSource);
 
 protected:
 	UEdGraphPin* PinPendingRename;
 
+	virtual bool GenerateCompileHashForClassMembers(const UClass* InClass, FNiagaraCompileHashVisitor* InVisitor) const;
 	
 private:
 	/** Looks up the name of the emitter and converts it to text. */
 	FText GetNameFromEmitter();
 
 private:
+	
 
 	/** The System that owns the emitter which is represented by this node. */
-	UPROPERTY()
+	UPROPERTY(meta = (SkipForCompileHash = "true"))
 	UNiagaraSystem* OwnerSystem;
 
 	/** The id of the emitter handle which points to the emitter represented by this node. */
@@ -83,7 +85,7 @@ private:
 	FGuid EmitterHandleId;
 
 	/** The display name for the title bar of this node. */
-	UPROPERTY()
+	UPROPERTY(meta = (SkipForCompileHash = "true"))
 	FText DisplayName;
 
 	UPROPERTY()
diff --git a/Engine/Plugins/FX/Niagara/Source/NiagaraEditor/Private/NiagaraNodeFunctionCall.cpp b/Engine/Plugins/FX/Niagara/Source/NiagaraEditor/Private/NiagaraNodeFunctionCall.cpp
--- a/Engine/Plugins/FX/Niagara/Source/NiagaraEditor/Private/NiagaraNodeFunctionCall.cpp
+++ b/Engine/Plugins/FX/Niagara/Source/NiagaraEditor/Private/NiagaraNodeFunctionCall.cpp
@@ -715,7 +715,7 @@ void UNiagaraNodeFunctionCall::SubsumeExternalDependencies(TMap<const UObject*,
 	}
 }
 
-void UNiagaraNodeFunctionCall::GatherExternalDependencyData(ENiagaraScriptUsage InMasterUsage, const FGuid& InMasterUsageId, TArray<FNiagaraCompileHash>& InReferencedCompileHashes, TArray<UObject*>& InReferencedObjs) const
+void UNiagaraNodeFunctionCall::GatherExternalDependencyData(ENiagaraScriptUsage InMasterUsage, const FGuid& InMasterUsageId, TArray<FNiagaraCompileHash>& InReferencedCompileHashes, TArray<FString>& InReferencedObjs) const
 {
 	if (FunctionScript && FunctionScript->GetOutermost() != this->GetOutermost())
 	{
@@ -733,7 +733,7 @@ void UNiagaraNodeFunctionCall::GatherExternalDependencyData(ENiagaraScriptUsage
 				if (FoundGuid.IsValid() && FoundCompileHash.IsValid())
 				{
 					InReferencedCompileHashes.Add(FoundCompileHash);
-					InReferencedObjs.Add(FunctionGraph);
+					InReferencedObjs.Add(FunctionGraph->GetPathName());
 					FunctionGraph->GatherExternalDependencyData((ENiagaraScriptUsage)i, FGuid(0, 0, 0, 0), InReferencedCompileHashes, InReferencedObjs);
 				}
 			}
diff --git a/Engine/Plugins/FX/Niagara/Source/NiagaraEditor/Private/NiagaraNodeIf.h b/Engine/Plugins/FX/Niagara/Source/NiagaraEditor/Private/NiagaraNodeIf.h
--- a/Engine/Plugins/FX/Niagara/Source/NiagaraEditor/Private/NiagaraNodeIf.h
+++ b/Engine/Plugins/FX/Niagara/Source/NiagaraEditor/Private/NiagaraNodeIf.h
@@ -43,10 +43,10 @@ public:
 	UPROPERTY()
 	TArray<FNiagaraVariable> OutputVars;
 
-	UPROPERTY()
+	UPROPERTY(meta = (SkipForCompileHash = "true"))
 	TArray<FPinGuidsForPath> PathAssociatedPinGuids;
 
-	UPROPERTY()
+	UPROPERTY(meta = (SkipForCompileHash = "true"))
 	FGuid ConditionPinGuid;
 
 	//~ Begin UObject Interface
diff --git a/Engine/Plugins/FX/Niagara/Source/NiagaraEditor/Private/NiagaraNodeParameterMapGet.cpp b/Engine/Plugins/FX/Niagara/Source/NiagaraEditor/Private/NiagaraNodeParameterMapGet.cpp
--- a/Engine/Plugins/FX/Niagara/Source/NiagaraEditor/Private/NiagaraNodeParameterMapGet.cpp
+++ b/Engine/Plugins/FX/Niagara/Source/NiagaraEditor/Private/NiagaraNodeParameterMapGet.cpp
@@ -414,7 +414,7 @@ bool UNiagaraNodeParameterMapGet::CommitEditablePinName(const FText& InName, UEd
 	return false;
 }
 
-void UNiagaraNodeParameterMapGet::GatherExternalDependencyData(ENiagaraScriptUsage InMasterUsage, const FGuid& InMasterUsageId, TArray<FNiagaraCompileHash>& InReferencedCompileHashes, TArray<UObject*>& InReferencedObjs) const
+void UNiagaraNodeParameterMapGet::GatherExternalDependencyData(ENiagaraScriptUsage InMasterUsage, const FGuid& InMasterUsageId, TArray<FNiagaraCompileHash>& InReferencedCompileHashes, TArray<FString>& InReferencedObjs) const
 {
 	// If we are referencing any parameter collections, we need to register them here... might want to speeed this up in the future 
 	// by caching any parameter collections locally.
@@ -434,7 +434,7 @@ void UNiagaraNodeParameterMapGet::GatherExternalDependencyData(ENiagaraScriptUsa
 		if (Collection)
 		{
 			InReferencedCompileHashes.Add(Collection->GetCompileHash());
-			InReferencedObjs.Add(Collection);
+			InReferencedObjs.Add(Collection->GetPathName());
 		}
 	}
 }
diff --git a/Engine/Plugins/FX/Niagara/Source/NiagaraEditor/Private/NiagaraNodeParameterMapGet.h b/Engine/Plugins/FX/Niagara/Source/NiagaraEditor/Private/NiagaraNodeParameterMapGet.h
--- a/Engine/Plugins/FX/Niagara/Source/NiagaraEditor/Private/NiagaraNodeParameterMapGet.h
+++ b/Engine/Plugins/FX/Niagara/Source/NiagaraEditor/Private/NiagaraNodeParameterMapGet.h
@@ -40,7 +40,7 @@ public:
 
 	virtual void PostLoad() override;
 
-	void GatherExternalDependencyData(ENiagaraScriptUsage InMasterUsage, const FGuid& InMasterUsageId, TArray<FNiagaraCompileHash>& InReferencedCompileHashes, TArray<UObject*>& InReferencedObjs) const override;
+	void GatherExternalDependencyData(ENiagaraScriptUsage InMasterUsage, const FGuid& InMasterUsageId, TArray<FNiagaraCompileHash>& InReferencedCompileHashes, TArray<FString>& InReferencedObjs) const override;
 	virtual void GetPinHoverText(const UEdGraphPin& Pin, FString& HoverTextOut) const override;
 
 protected:
@@ -59,6 +59,6 @@ protected:
 	/** Properly set up the default input pin for an output pin.*/
 	UEdGraphPin* CreateDefaultPin(UEdGraphPin* OutputPin);
 
-	UPROPERTY()
+	UPROPERTY(meta = (SkipForCompileHash="true"))
 	TMap<FGuid, FGuid> PinOutputToPinDefaultPersistentId;
 };
diff --git a/Engine/Plugins/FX/Niagara/Source/NiagaraEditor/Private/NiagaraScriptSource.cpp b/Engine/Plugins/FX/Niagara/Source/NiagaraEditor/Private/NiagaraScriptSource.cpp
--- a/Engine/Plugins/FX/Niagara/Source/NiagaraEditor/Private/NiagaraScriptSource.cpp
+++ b/Engine/Plugins/FX/Niagara/Source/NiagaraEditor/Private/NiagaraScriptSource.cpp
@@ -51,10 +51,20 @@ void UNiagaraScriptSource::ComputeVMCompilationId(FNiagaraVMExecutableDataId& Id
 	if (NodeGraph)
 	{
 		NodeGraph->RebuildCachedCompileIds(bForceRebuild);
-		Id.BaseScriptID = NodeGraph->GetBaseId(InUsage, InUsageId);
 		Id.BaseScriptCompileHash = FNiagaraCompileHash(NodeGraph->GetCompileDataHash(InUsage, InUsageId));
-		NodeGraph->GatherExternalDependencyData(InUsage, InUsageId, Id.ReferencedCompileHashes, Id.ReferencedObjects);
+		NodeGraph->GatherExternalDependencyData(InUsage, InUsageId, Id.ReferencedCompileHashes, Id.DebugReferencedObjects);
 	}
+
+	// Add in any referenced HLSL files.
+	FSHAHash Hash = GetShaderFileHash((TEXT("/Plugin/FX/Niagara/Private/NiagaraEmitterInstanceShader.usf")), EShaderPlatform::SP_PCD3D_SM5);
+	Id.ReferencedCompileHashes.Emplace(Hash.Hash, sizeof(Hash.Hash)/sizeof(uint8));
+	Id.DebugReferencedObjects.Emplace(TEXT("/Plugin/FX/Niagara/Private/NiagaraEmitterInstanceShader.usf"));
+	Hash = GetShaderFileHash((TEXT("/Plugin/FX/Niagara/Private/NiagaraShaderVersion.ush")), EShaderPlatform::SP_PCD3D_SM5);
+	Id.ReferencedCompileHashes.Emplace(Hash.Hash, sizeof(Hash.Hash) / sizeof(uint8));
+	Id.DebugReferencedObjects.Emplace(TEXT("/Plugin/FX/Niagara/Private/NiagaraShaderVersion.ush"));
+	Hash = GetShaderFileHash((TEXT("/Engine/Public/ShaderVersion.ush")), EShaderPlatform::SP_PCD3D_SM5);
+	Id.ReferencedCompileHashes.Emplace(Hash.Hash, sizeof(Hash.Hash) / sizeof(uint8));
+	Id.DebugReferencedObjects.Emplace(TEXT("/Engine/Public/ShaderVersion.ush"));
 }
 
 FGuid UNiagaraScriptSource::GetCompileBaseId(ENiagaraScriptUsage InUsage, const FGuid& InUsageId) const
@@ -67,9 +77,9 @@ FNiagaraCompileHash UNiagaraScriptSource::GetCompileHash(ENiagaraScriptUsage InU
 	return NodeGraph->GetCompileDataHash(InUsage, InUsageId);
 }
 
-void UNiagaraScriptSource::InvalidateCachedCompileIds() 
+void UNiagaraScriptSource::ForceGraphToRecompileOnNextCheck()
 {
-	NodeGraph->InvalidateCachedCompileIds();
+	NodeGraph->ForceGraphToRecompileOnNextCheck();
 }
 
 void UNiagaraScriptSource::RefreshFromExternalChanges()
diff --git a/Engine/Plugins/FX/Niagara/Source/NiagaraEditor/Private/NiagaraScriptVariable.cpp b/Engine/Plugins/FX/Niagara/Source/NiagaraEditor/Private/NiagaraScriptVariable.cpp
--- a/Engine/Plugins/FX/Niagara/Source/NiagaraEditor/Private/NiagaraScriptVariable.cpp
+++ b/Engine/Plugins/FX/Niagara/Source/NiagaraEditor/Private/NiagaraScriptVariable.cpp
@@ -2,6 +2,7 @@
 
 #include "NiagaraScriptVariable.h"
 #include "NiagaraCommon.h"
+#include "NiagaraEditorUtilities.h"
 
 UNiagaraScriptVariable::UNiagaraScriptVariable(const FObjectInitializer& ObjectInitializer)
 	: Super(ObjectInitializer)
@@ -10,6 +11,18 @@ UNiagaraScriptVariable::UNiagaraScriptVariable(const FObjectInitializer& ObjectI
 	
 }
 
+void UNiagaraScriptVariable::PostEditChangeProperty(struct FPropertyChangedEvent& PropertyChangedEvent)
+{
+	Super::PostEditChangeProperty(PropertyChangedEvent);
+
+#if WITH_EDITOR
+	if (UNiagaraGraph* Graph = Cast<UNiagaraGraph>(GetOuter()))
+	{
+		Graph->NotifyGraphNeedsRecompile();
+	}
+#endif	//#if WITH_EDITOR
+}
+
 void UNiagaraScriptVariable::PostLoad()
 {
 	Super::PostLoad();
@@ -19,3 +32,12 @@ void UNiagaraScriptVariable::PostLoad()
 		SetFlags(RF_Transactional);
 	}
 }
+
+bool UNiagaraScriptVariable::AppendCompileHash(FNiagaraCompileHashVisitor* InVisitor) const
+{
+	if (!FNiagaraEditorUtilities::NestedPropertiesAppendCompileHash(static_cast<const void*>(this), StaticClass(), EFieldIteratorFlags::ExcludeSuper, StaticClass()->GetName(), InVisitor))
+	{
+		return false;
+	}
+	return true;
+}
\ No newline at end of file
diff --git a/Engine/Plugins/FX/Niagara/Source/NiagaraEditor/Private/NiagaraScriptVariable.h b/Engine/Plugins/FX/Niagara/Source/NiagaraEditor/Private/NiagaraScriptVariable.h
--- a/Engine/Plugins/FX/Niagara/Source/NiagaraEditor/Private/NiagaraScriptVariable.h
+++ b/Engine/Plugins/FX/Niagara/Source/NiagaraEditor/Private/NiagaraScriptVariable.h
@@ -16,6 +16,8 @@ public:
 	GENERATED_UCLASS_BODY()
 
 	virtual void PostLoad() override;
+
+	virtual void PostEditChangeProperty(struct FPropertyChangedEvent& PropertyChangedEvent) override;
 	
 	/** The default mode. Can be Value, Binding or Custom. */
 	UPROPERTY(EditAnywhere, Category = "Default Value")
@@ -32,4 +34,7 @@ public:
 	/** The metadata associated with this script variable. */
 	UPROPERTY(EditAnywhere, Category = "Variable", meta=(ShowOnlyInnerProperties))
 	FNiagaraVariableMetaData Metadata;
+	
+	/** Entry point for generating the compile hash.*/
+	bool AppendCompileHash(FNiagaraCompileHashVisitor* InVisitor) const;
 }; 
diff --git a/Engine/Plugins/FX/Niagara/Source/NiagaraEditor/Public/NiagaraEditorUtilities.h b/Engine/Plugins/FX/Niagara/Source/NiagaraEditor/Public/NiagaraEditorUtilities.h
--- a/Engine/Plugins/FX/Niagara/Source/NiagaraEditor/Public/NiagaraEditorUtilities.h
+++ b/Engine/Plugins/FX/Niagara/Source/NiagaraEditor/Public/NiagaraEditorUtilities.h
@@ -150,6 +150,9 @@ namespace FNiagaraEditorUtilities
 
 	void PreprocessFunctionGraph(const UEdGraphSchema_Niagara* Schema, UNiagaraGraph* Graph, const TArray<UEdGraphPin*>& CallInputs, const TArray<UEdGraphPin*>& CallOutputs, ENiagaraScriptUsage ScriptUsage, const FCompileConstantResolver& ConstantResolver);
 
+	bool PODPropertyAppendCompileHash(const void* Container, UProperty* Property, const FString& PropertyName, struct FNiagaraCompileHashVisitor* InVisitor);
+	bool NestedPropertiesAppendCompileHash(const void* Container, const UStruct* Struct, EFieldIteratorFlags::SuperClassFlags IteratorFlags, const FString& BaseName, struct FNiagaraCompileHashVisitor* InVisitor);
+
 	/** Options for the GetScriptsByFilter function. 
 	** @Param ScriptUsageToInclude Only return Scripts that have this usage
 	** @Param (Optional) TargetUsageToMatch Only return Scripts that have this target usage (output node) 
diff --git a/Engine/Plugins/FX/Niagara/Source/NiagaraEditor/Public/NiagaraGraph.h b/Engine/Plugins/FX/Niagara/Source/NiagaraEditor/Public/NiagaraGraph.h
--- a/Engine/Plugins/FX/Niagara/Source/NiagaraEditor/Public/NiagaraGraph.h
+++ b/Engine/Plugins/FX/Niagara/Source/NiagaraEditor/Public/NiagaraGraph.h
@@ -42,6 +42,7 @@ private:
 	bool bCreated;
 };
 
+
 /** Container for UNiagaraGraph cached data for managing CompileIds and Traversals.*/
 USTRUCT()
 struct FNiagaraGraphScriptUsageInfo
@@ -67,6 +68,14 @@ public:
 	UPROPERTY()
 	FNiagaraCompileHash CompileHash;
 
+	/** The hash that we calculated last traversal. */
+	UPROPERTY()
+	FNiagaraCompileHash CompileHashFromGraph;
+
+	UPROPERTY(Transient)
+	TArray<FNiagaraCompileHashVisitorDebugInfo> CompileLastObjects;
+
+
 	/** The traversal of output to input nodes for this graph. This is not a recursive traversal, it just includes nodes from this graph.*/
 	UPROPERTY()
 	TArray<UNiagaraNode*> Traversal;
@@ -181,7 +190,7 @@ class UNiagaraGraph : public UEdGraph
 	void GetAllReferencedGraphs(TArray<const UNiagaraGraph*>& Graphs) const;
 
 	/** Gather all the change ids of external references for this specific graph traversal.*/
-	void GatherExternalDependencyData(ENiagaraScriptUsage InUsage, const FGuid& InUsageId, TArray<FNiagaraCompileHash>& InReferencedCompileHashes, TArray<UObject*>& InReferencedObjs);
+	void GatherExternalDependencyData(ENiagaraScriptUsage InUsage, const FGuid& InUsageId, TArray<FNiagaraCompileHash>& InReferencedCompileHashes, TArray<FString>& InReferencedObjs);
 
 	/** Determine if there are any external dependencies wrt to scripts and ensure that those dependencies are sucked into the existing package.*/
 	void SubsumeExternalDependencies(TMap<const UObject*, UObject*>& ExistingConversions);
@@ -241,7 +250,7 @@ class UNiagaraGraph : public UEdGraph
 	/** Gets a delegate which is called whenever a contained data interfaces changes. */
 	FOnDataInterfaceChanged& OnDataInterfaceChanged();
 
-	void InvalidateCachedCompileIds();
+	void ForceGraphToRecompileOnNextCheck();
 
 	/** Add a listener for OnGraphNeedsRecompile events */
 	FDelegateHandle AddOnGraphNeedsRecompileHandler(const FOnGraphChanged::FDelegate& InHandler);
@@ -270,6 +279,7 @@ protected:
 	bool bNeedNumericCacheRebuilt;
 	TMap<TPair<FGuid, UEdGraphNode*>, FNiagaraTypeDefinition> CachedNumericConversions;
 	void ResolveNumerics(TMap<UNiagaraNode*, bool>& VisitedNodes, UEdGraphNode* Node);
+	bool AppendCompileHash(FNiagaraCompileHashVisitor* InVisitor, const TArray<UNiagaraNode*>& InTraversal) const;
 
 private:
 	virtual void NotifyGraphChanged(const FEdGraphEditAction& InAction) override;
@@ -290,6 +300,10 @@ private:
 	UPROPERTY()
 	FGuid ChangeId;
 
+	/** Internal value used to invalidate a DDC key for the script no matter what.*/
+	UPROPERTY()
+	FGuid ForceRebuildId;
+
 	UPROPERTY()
 	FGuid LastBuiltTraversalDataChangeId;
 
diff --git a/Engine/Plugins/FX/Niagara/Source/NiagaraEditor/Public/NiagaraNode.h b/Engine/Plugins/FX/Niagara/Source/NiagaraEditor/Public/NiagaraNode.h
--- a/Engine/Plugins/FX/Niagara/Source/NiagaraEditor/Public/NiagaraNode.h
+++ b/Engine/Plugins/FX/Niagara/Source/NiagaraEditor/Public/NiagaraNode.h
@@ -9,6 +9,7 @@
 #include "NiagaraEditorCommon.h"
 #include "NiagaraParameterMapHistory.h"
 #include "Misc/Guid.h"
+#include "UObject/UnrealType.h"
 #include "NiagaraNode.generated.h"
 
 class UEdGraphPin;
@@ -101,7 +102,7 @@ public:
 	virtual void BuildParameterMapHistory(FNiagaraParameterMapHistoryBuilder& OutHistory, bool bRecursive = true, bool bFilterForCompilation = true) const;
 	
 	/** Go through all the external dependencies of this node in isolation and add them to the reference id list.*/
-	virtual void GatherExternalDependencyData(ENiagaraScriptUsage InMasterUsage, const FGuid& InMasterUsageId, TArray<FNiagaraCompileHash>& InReferencedCompileHashes, TArray<UObject*>& InReferencedObjs) const {};
+	virtual void GatherExternalDependencyData(ENiagaraScriptUsage InMasterUsage, const FGuid& InMasterUsageId, TArray<FNiagaraCompileHash>& InReferencedCompileHashes, TArray<FString>& InReferencedObjs) const {};
 
 	/** Traces one of this node's output pins to its source output pin if it is a reroute node output pin.*/
 	virtual UEdGraphPin* GetTracedOutputPin(UEdGraphPin* LocallyOwnedOutputPin) const {return LocallyOwnedOutputPin;}
@@ -126,9 +127,28 @@ public:
 
 	virtual void AppendFunctionAliasForContext(const FNiagaraGraphFunctionAliasContext& InFunctionAliasContext, FString& InOutFunctionAlias) { };
 
+	/** Old style compile hash code. To be removed in the future.*/
 	virtual void UpdateCompileHashForNode(FSHA1& HashState) const;
 
+	/** Entry point for generating the compile hash.*/
+	bool AppendCompileHash(FNiagaraCompileHashVisitor* InVisitor) const;
+
 protected:
+	/** Go through all class members for a given UClass on this object and hash them into the visitor.*/
+	virtual bool GenerateCompileHashForClassMembers(const UClass* InClass, FNiagaraCompileHashVisitor* InVisitor) const;
+
+	/** Write out the specific entries for UNiagaraNode into the visitor hash. */
+	bool NiagaraNodeAppendCompileHash(FNiagaraCompileHashVisitor* InVisitor) const;
+
+	/** Write out the specific entries of this pin to the visitor hash.*/
+	virtual bool PinAppendCompileHash(const UEdGraphPin* InPin, FNiagaraCompileHashVisitor* InVisitor) const;
+	
+	/** Helper function to hash arbitrary UProperty entries (Arrays, Maps, Structs, etc).*/
+	virtual bool NestedPropertiesAppendCompileHash(const void* Container, const UStruct* Struct, EFieldIteratorFlags::SuperClassFlags IteratorFlags, const FString& BaseName, FNiagaraCompileHashVisitor* InVisitor) const;
+	
+	/** For a simple Plain old data type UProperty, hash the data.*/
+	virtual bool PODPropertyAppendCompileHash(const void* Container, UProperty* Property, const FString& PropertyName, FNiagaraCompileHashVisitor* InVisitor) const;
+
 	virtual int32 CompileInputPin(class FHlslNiagaraTranslator *Translator, UEdGraphPin* Pin);
 	virtual bool IsValidPinToCompile(UEdGraphPin* Pin) const { return true; }; 
 
diff --git a/Engine/Plugins/FX/Niagara/Source/NiagaraEditor/Public/NiagaraNodeAssignment.h b/Engine/Plugins/FX/Niagara/Source/NiagaraEditor/Public/NiagaraNodeAssignment.h
--- a/Engine/Plugins/FX/Niagara/Source/NiagaraEditor/Public/NiagaraNodeAssignment.h
+++ b/Engine/Plugins/FX/Niagara/Source/NiagaraEditor/Public/NiagaraNodeAssignment.h
@@ -45,7 +45,7 @@ public:
 	//~ UNiagaraNodeFunctionCall interface
 	virtual bool RefreshFromExternalChanges() override;
 	virtual void BuildParameterMapHistory(FNiagaraParameterMapHistoryBuilder& OutHistory, bool bRecursive = true, bool bFilterForCompilation = true) const override;
-	virtual void GatherExternalDependencyData(ENiagaraScriptUsage InMasterUsage, const FGuid& InMasterUsageId, TArray<FNiagaraCompileHash>& InReferencedCompileHashes, TArray<UObject*>& InReferencedObjs) const override;
+	virtual void GatherExternalDependencyData(ENiagaraScriptUsage InMasterUsage, const FGuid& InMasterUsageId, TArray<FNiagaraCompileHash>& InReferencedCompileHashes, TArray<FString>& InReferencedObjs) const override;
 
 	void AddParameter(FNiagaraVariable InVar, FString InDefaultValue);
 	void RemoveParameter(const FNiagaraVariable& InVar);
diff --git a/Engine/Plugins/FX/Niagara/Source/NiagaraEditor/Public/NiagaraNodeFunctionCall.h b/Engine/Plugins/FX/Niagara/Source/NiagaraEditor/Public/NiagaraNodeFunctionCall.h
--- a/Engine/Plugins/FX/Niagara/Source/NiagaraEditor/Public/NiagaraNodeFunctionCall.h
+++ b/Engine/Plugins/FX/Niagara/Source/NiagaraEditor/Public/NiagaraNodeFunctionCall.h
@@ -63,7 +63,7 @@ public:
 	 * default pins are generated. This is used so that the function nodes can be populated in the graph context
 	 * menu without having to load all of the actual script assets.
 	 */
-	UPROPERTY(Transient)
+	UPROPERTY(Transient, meta = (SkipForCompileHash = "true"))
 	FName FunctionScriptAssetObjectPath;
 
 	/** Some functions can be provided a signature directly rather than a script. */
@@ -92,7 +92,7 @@ public:
 	virtual ENiagaraNumericOutputTypeSelectionMode GetNumericOutputTypeSelectionMode() const override;
 	virtual bool CanAddToGraph(UNiagaraGraph* TargetGraph, FString& OutErrorMsg) const override;
 	virtual void SubsumeExternalDependencies(TMap<const UObject*, UObject*>& ExistingConversions) override;
-	virtual void GatherExternalDependencyData(ENiagaraScriptUsage InMasterUsage, const FGuid& InMasterUsageId, TArray<FNiagaraCompileHash>& InReferencedCompileHashes, TArray<UObject*>& InReferencedObjs) const override;
+	virtual void GatherExternalDependencyData(ENiagaraScriptUsage InMasterUsage, const FGuid& InMasterUsageId, TArray<FNiagaraCompileHash>& InReferencedCompileHashes, TArray<FString>& InReferencedObjs) const override;
 	virtual void UpdateCompileHashForNode(FSHA1& HashState) const override;
 	//End UNiagaraNode interface
 
@@ -155,10 +155,10 @@ protected:
 	void UpdateNodeErrorMessage();
 
 	/** Adjusted every time that we compile this script. Lets us know that we might differ from any cached versions.*/
-	UPROPERTY()
+	UPROPERTY(meta = (SkipForCompileHash="true"))
 	FGuid CachedChangeId;
 
-	UPROPERTY()
+	UPROPERTY(meta = (SkipForCompileHash = "true"))
 	FString FunctionDisplayName;
 
 	FOnInputsChanged OnInputsChangedDelegate;
diff --git a/Engine/Plugins/FX/Niagara/Source/NiagaraEditor/Public/NiagaraNodeInput.h b/Engine/Plugins/FX/Niagara/Source/NiagaraEditor/Public/NiagaraNodeInput.h
--- a/Engine/Plugins/FX/Niagara/Source/NiagaraEditor/Public/NiagaraNodeInput.h
+++ b/Engine/Plugins/FX/Niagara/Source/NiagaraEditor/Public/NiagaraNodeInput.h
@@ -49,14 +49,14 @@ class UNiagaraNodeInput : public UNiagaraNode
 	GENERATED_UCLASS_BODY()
 
 public:
-	UPROPERTY(EditAnywhere, Category = Input) 
+	UPROPERTY(EditAnywhere, Category = Input, meta = (SkipForCompileHash = "true"))
 	FNiagaraVariable Input;
 
 	UPROPERTY(EditAnywhere, Category = Input)
 	ENiagaraInputNodeUsage Usage;
 
 	/** Controls where this input is relative to others in the calling node. */
-	UPROPERTY(EditAnywhere, Category = Input)
+	UPROPERTY(EditAnywhere, Category = Input, meta = (SkipForCompileHash = "true"))
 	int32 CallSortPriority;
 
 	/** Controls this inputs exposure to callers. */
@@ -115,7 +115,7 @@ private:
 	void ValidateDataInterface();
 
 private:
-	UPROPERTY()
+	UPROPERTY(meta = (SkipForCompileHash = "true"))
 	class UNiagaraDataInterface* DataInterface;
 };
 
diff --git a/Engine/Plugins/FX/Niagara/Source/NiagaraEditor/Public/NiagaraNodeOp.h b/Engine/Plugins/FX/Niagara/Source/NiagaraEditor/Public/NiagaraNodeOp.h
--- a/Engine/Plugins/FX/Niagara/Source/NiagaraEditor/Public/NiagaraNodeOp.h
+++ b/Engine/Plugins/FX/Niagara/Source/NiagaraEditor/Public/NiagaraNodeOp.h
@@ -41,7 +41,7 @@ public:
 	UPROPERTY()
 	FName OpName;
 
-	UPROPERTY()
+	UPROPERTY(meta = (SkipForCompileHash = "true"))
 	TArray<FAddedPinData> AddedPins;
 
 	//~ Begin UObject interface
diff --git a/Engine/Plugins/FX/Niagara/Source/NiagaraEditor/Public/NiagaraNodeOutput.h b/Engine/Plugins/FX/Niagara/Source/NiagaraEditor/Public/NiagaraNodeOutput.h
--- a/Engine/Plugins/FX/Niagara/Source/NiagaraEditor/Public/NiagaraNodeOutput.h
+++ b/Engine/Plugins/FX/Niagara/Source/NiagaraEditor/Public/NiagaraNodeOutput.h
@@ -15,7 +15,7 @@ class UNiagaraNodeOutput : public UNiagaraNode
 {
 	GENERATED_UCLASS_BODY()
 
-	UPROPERTY(VisibleAnywhere, Category = Output)
+	UPROPERTY(VisibleAnywhere, Category = Output, meta = (SkipForCompileHash = "true"))
 	TArray<FNiagaraVariable> Outputs;
 
 	UPROPERTY()
diff --git a/Engine/Plugins/FX/Niagara/Source/NiagaraEditor/Public/NiagaraScriptSource.h b/Engine/Plugins/FX/Niagara/Source/NiagaraEditor/Public/NiagaraScriptSource.h
--- a/Engine/Plugins/FX/Niagara/Source/NiagaraEditor/Public/NiagaraScriptSource.h
+++ b/Engine/Plugins/FX/Niagara/Source/NiagaraEditor/Public/NiagaraScriptSource.h
@@ -46,7 +46,7 @@ class UNiagaraScriptSource : public UNiagaraScriptSourceBase
 	NIAGARAEDITOR_API virtual bool AddModuleIfMissing(FString ModulePath, ENiagaraScriptUsage Usage, bool& bOutFoundModule)override;
 
 	virtual void CleanUpOldAndInitializeNewRapidIterationParameters(FString UniqueEmitterName, ENiagaraScriptUsage ScriptUsage, FGuid ScriptUsageId, FNiagaraParameterStore& RapidIterationParameters) const override;
-	virtual void InvalidateCachedCompileIds() override;
+	virtual void ForceGraphToRecompileOnNextCheck() override;
 	virtual void RefreshFromExternalChanges() override;
 
 private:
diff --git a/Engine/Plugins/FX/Niagara/Source/NiagaraShader/Private/NiagaraShader.cpp b/Engine/Plugins/FX/Niagara/Source/NiagaraShader/Private/NiagaraShader.cpp
--- a/Engine/Plugins/FX/Niagara/Source/NiagaraShader/Private/NiagaraShader.cpp
+++ b/Engine/Plugins/FX/Niagara/Source/NiagaraShader/Private/NiagaraShader.cpp
@@ -120,7 +120,7 @@ void FNiagaraShaderMapId::Serialize(FArchive& Ar)
 		Ar << CompilerVersionID;
 	}
 
-	Ar << BaseScriptID;
+	Ar << BaseScriptID_DEPRECATED;
 	Ar << (int32&)FeatureLevel;
 
 	if (Ar.IsLoading() && NiagaraVer < FNiagaraCustomVersion::RemoveGraphUsageCompileIds)
@@ -156,7 +156,7 @@ void FNiagaraShaderMapId::GetScriptHash(FSHAHash& OutHash) const
 	FSHA1 HashState;
 
 	HashState.Update((const uint8*)&CompilerVersionID, sizeof(CompilerVersionID));
-	HashState.Update((const uint8*)&BaseScriptID, sizeof(BaseScriptID));
+	//HashState.Update((const uint8*)&BaseScriptID, sizeof(BaseScriptID));
 	HashState.Update(BaseCompileHash.GetData(), FNiagaraCompileHash::HashSize);
 	HashState.Update((const uint8*)&FeatureLevel, sizeof(FeatureLevel));
 		
@@ -184,8 +184,8 @@ void FNiagaraShaderMapId::GetScriptHash(FSHAHash& OutHash) const
 */
 bool FNiagaraShaderMapId::operator==(const FNiagaraShaderMapId& ReferenceSet) const
 {
-	if (BaseScriptID != ReferenceSet.BaseScriptID 
-		|| BaseCompileHash != ReferenceSet.BaseCompileHash
+	if (/*BaseScriptID != ReferenceSet.BaseScriptID 
+		|| */BaseCompileHash != ReferenceSet.BaseCompileHash
 		|| FeatureLevel != ReferenceSet.FeatureLevel
 		|| CompilerVersionID != ReferenceSet.CompilerVersionID 
 		|| DetailLevelMask != ReferenceSet.DetailLevelMask 
@@ -245,8 +245,8 @@ bool FNiagaraShaderMapId::operator==(const FNiagaraShaderMapId& ReferenceSet) co
 
 void FNiagaraShaderMapId::AppendKeyString(FString& KeyString) const
 {
-	KeyString += BaseScriptID.ToString();
-	KeyString += TEXT("_");
+	//KeyString += BaseScriptID.ToString();
+	//KeyString += TEXT("_");
 
 	KeyString += BaseCompileHash.ToString();
 	KeyString += TEXT("_");
@@ -486,7 +486,7 @@ FShader* FNiagaraShaderType::FinishCompileShader(
 */
 FNiagaraShaderMap* FNiagaraShaderMap::FindId(const FNiagaraShaderMapId& ShaderMapId, EShaderPlatform InPlatform)
 {
-	check(ShaderMapId.BaseScriptID != FGuid());
+	check(ShaderMapId.BaseCompileHash.IsValid());
 	return GIdToNiagaraShaderMap[InPlatform].FindRef(ShaderMapId);
 }
 
diff --git a/Engine/Plugins/FX/Niagara/Source/NiagaraShader/Private/NiagaraShared.cpp b/Engine/Plugins/FX/Niagara/Source/NiagaraShader/Private/NiagaraShared.cpp
--- a/Engine/Plugins/FX/Niagara/Source/NiagaraShader/Private/NiagaraShared.cpp
+++ b/Engine/Plugins/FX/Niagara/Source/NiagaraShader/Private/NiagaraShared.cpp
@@ -87,8 +87,8 @@ NIAGARASHADER_API void FNiagaraShaderScript::LegacySerialize(FArchive& Ar)
 bool FNiagaraShaderScript::IsSame(const FNiagaraShaderMapId& InId) const
 {
 	return
-		InId.FeatureLevel == FeatureLevel &&
-		InId.BaseScriptID == BaseScriptId &&
+		InId.FeatureLevel == FeatureLevel &&/*
+		InId.BaseScriptID == BaseScriptId &&*/
 		InId.AdditionalDefines == AdditionalDefines &&
 		InId.DetailLevelMask == DetailLevelMask &&
 		InId.bUsesRapidIterationParams == bUsesRapidIterationParams &&
@@ -123,8 +123,8 @@ NIAGARASHADER_API void FNiagaraShaderScript::GetShaderMapId(EShaderPlatform Plat
 	{
 		TArray<FShaderType*> ShaderTypes;
 		GetDependentShaderTypes(Platform, ShaderTypes);
-		OutId.FeatureLevel = GetFeatureLevel();
-		OutId.BaseScriptID = BaseScriptId;
+		OutId.FeatureLevel = GetFeatureLevel();/*
+		OutId.BaseScriptID = BaseScriptId;*/
 		OutId.AdditionalDefines = AdditionalDefines;
 		OutId.DetailLevelMask = DetailLevelMask;
 		OutId.bUsesRapidIterationParams = bUsesRapidIterationParams;		
@@ -233,15 +233,14 @@ void FNiagaraShaderScript::SerializeShaderMap(FArchive& Ar)
 	}
 }
 
-void FNiagaraShaderScript::SetScript(UNiagaraScript *InScript, ERHIFeatureLevel::Type InFeatureLevel, const FGuid& InCompilerVersionID, const FGuid& InBaseScriptID, const TArray<FString>& InAdditionalDefines,
+void FNiagaraShaderScript::SetScript(UNiagaraScript *InScript, ERHIFeatureLevel::Type InFeatureLevel, const FGuid& InCompilerVersionID,  const TArray<FString>& InAdditionalDefines,
 		const FNiagaraCompileHash& InBaseCompileHash, const TArray<FNiagaraCompileHash>& InReferencedCompileHashes, 
 		bool bInUsesRapidIterationParams, uint32 InDetailLevelMask, FString InFriendlyName)
 {
-	checkf(InBaseScriptID.IsValid(), TEXT("Invalid base script id.  Script caching will fail."));
 	checkf(InBaseCompileHash.IsValid(), TEXT("Invalid base compile hash.  Script caching will fail."))
 	BaseVMScript = InScript;
 	CompilerVersionId = InCompilerVersionID;
-	BaseScriptId = InBaseScriptID;
+	//BaseScriptId = InBaseScriptID;
 	AdditionalDefines = InAdditionalDefines;
 	bUsesRapidIterationParams = bInUsesRapidIterationParams;
 	DetailLevelMask = InDetailLevelMask;
diff --git a/Engine/Plugins/FX/Niagara/Source/NiagaraShader/Public/NiagaraShared.h b/Engine/Plugins/FX/Niagara/Source/NiagaraShader/Public/NiagaraShared.h
--- a/Engine/Plugins/FX/Niagara/Source/NiagaraShader/Public/NiagaraShared.h
+++ b/Engine/Plugins/FX/Niagara/Source/NiagaraShader/Public/NiagaraShared.h
@@ -158,7 +158,7 @@ public:
 	/**
 	* The base id of the subgraph this shader primarily represents.
 	*/
-	FGuid BaseScriptID;
+	FGuid BaseScriptID_DEPRECATED;
 
 	/** Configuration options */
 	TArray<FString> AdditionalDefines;
@@ -192,7 +192,7 @@ public:
 
 	friend uint32 GetTypeHash(const FNiagaraShaderMapId& Ref)
 	{
-		return Ref.BaseScriptID.A;
+		return Ref.BaseCompileHash.GetTypeHash();
 	}
 
 	SIZE_T GetSizeBytes() const
@@ -673,7 +673,7 @@ public:
 	const FString& GetFriendlyName()	const { return FriendlyName; }
 
 
-	NIAGARASHADER_API void SetScript(UNiagaraScript *InScript, ERHIFeatureLevel::Type InFeatureLevel, const FGuid& InCompilerVersion, const FGuid& InBaseScriptID, const TArray<FString>& InAdditionalDefines,
+	NIAGARASHADER_API void SetScript(UNiagaraScript *InScript, ERHIFeatureLevel::Type InFeatureLevel, const FGuid& InCompilerVersion, const TArray<FString>& InAdditionalDefines,
 		const FNiagaraCompileHash& InBaseCompileHash, const TArray<FNiagaraCompileHash>& InReferencedCompileHashes, 
 		bool bInUsesRapidIterationParams, uint32 InDetailLevelMask, FString InFriendlyName);
 
@@ -748,7 +748,7 @@ private:
 	TArray< FNiagaraDataInterfaceGPUParamInfo > DIParamInfo;
 
 	/** Guid id for base script*/
-	FGuid BaseScriptId;
+	FGuid BaseScriptId_DEPRECATED;
 
 	/** Configuration options */
 	TArray<FString> AdditionalDefines;
diff --git a/Engine/Source/Runtime/Engine/Classes/EdGraph/EdGraphNode.h b/Engine/Source/Runtime/Engine/Classes/EdGraph/EdGraphNode.h
--- a/Engine/Source/Runtime/Engine/Classes/EdGraph/EdGraphNode.h
+++ b/Engine/Source/Runtime/Engine/Classes/EdGraph/EdGraphNode.h
@@ -137,6 +137,24 @@ enum class ENodeEnabledState : uint8
 	DevelopmentOnly
 };
 
+inline const TCHAR* const LexToString(const ENodeEnabledState State)
+{
+	switch (State)
+	{
+	case ENodeEnabledState::Enabled:
+		return TEXT("Enabled");
+	case ENodeEnabledState::Disabled:
+		return TEXT("Disabled");
+	case ENodeEnabledState::DevelopmentOnly:
+		return TEXT("DevelopmentOnly");
+	default:
+		break;
+	}
+
+	checkf(false, TEXT("Missing ENodeEnabledState Type: %d"), static_cast<const int32>(State));
+	return TEXT("");
+}
+
 /** Enum that defines what kind of orphaned pins should be retained. */
 enum class ESaveOrphanPinMode : uint8
 {
