Found and fixed another instance of UProperty to convert to FProperty

diff --git a/Engine/Plugins/FX/Niagara/Source/NiagaraEditor/Private/NiagaraNode.cpp b/Engine/Plugins/FX/Niagara/Source/NiagaraEditor/Private/NiagaraNode.cpp
--- a/Engine/Plugins/FX/Niagara/Source/NiagaraEditor/Private/NiagaraNode.cpp
+++ b/Engine/Plugins/FX/Niagara/Source/NiagaraEditor/Private/NiagaraNode.cpp
@@ -43,7 +43,7 @@ bool UNiagaraNode::NestedPropertiesAppendCompileHash(const void* Container, cons
 	return FNiagaraEditorUtilities::NestedPropertiesAppendCompileHash(Container, Struct, IteratorFlags, BaseName, InVisitor);
 }
 
-bool UNiagaraNode::PODPropertyAppendCompileHash(const void* Container, UProperty* Property, const FString& PropertyName, FNiagaraCompileHashVisitor* InVisitor) const
+bool UNiagaraNode::PODPropertyAppendCompileHash(const void* Container, FProperty* Property, const FString& PropertyName, FNiagaraCompileHashVisitor* InVisitor) const
 {
 	return FNiagaraEditorUtilities::PODPropertyAppendCompileHash(Container, Property, PropertyName, InVisitor);
 }
diff --git a/Engine/Plugins/FX/Niagara/Source/NiagaraEditor/Public/NiagaraNode.h b/Engine/Plugins/FX/Niagara/Source/NiagaraEditor/Public/NiagaraNode.h
--- a/Engine/Plugins/FX/Niagara/Source/NiagaraEditor/Public/NiagaraNode.h
+++ b/Engine/Plugins/FX/Niagara/Source/NiagaraEditor/Public/NiagaraNode.h
@@ -147,7 +147,7 @@ protected:
 	virtual bool NestedPropertiesAppendCompileHash(const void* Container, const UStruct* Struct, EFieldIteratorFlags::SuperClassFlags IteratorFlags, const FString& BaseName, FNiagaraCompileHashVisitor* InVisitor) const;
 	
 	/** For a simple Plain old data type UProperty, hash the data.*/
-	virtual bool PODPropertyAppendCompileHash(const void* Container, UProperty* Property, const FString& PropertyName, FNiagaraCompileHashVisitor* InVisitor) const;
+	virtual bool PODPropertyAppendCompileHash(const void* Container, FProperty* Property, const FString& PropertyName, FNiagaraCompileHashVisitor* InVisitor) const;
 
 	virtual int32 CompileInputPin(class FHlslNiagaraTranslator *Translator, UEdGraphPin* Pin);
 	virtual bool IsValidPinToCompile(UEdGraphPin* Pin) const { return true; }; 
