- Use GetSecondsPerCycle64 rather than GetSecondsPerCycle since we are using 64 elsewhere

#rnx

diff --git a/Engine/Plugins/FX/Niagara/Source/Niagara/Private/NiagaraSystemInstance.cpp b/Engine/Plugins/FX/Niagara/Source/Niagara/Private/NiagaraSystemInstance.cpp
--- a/Engine/Plugins/FX/Niagara/Source/Niagara/Private/NiagaraSystemInstance.cpp
+++ b/Engine/Plugins/FX/Niagara/Source/Niagara/Private/NiagaraSystemInstance.cpp
@@ -1781,7 +1781,7 @@ void FNiagaraSystemInstance::WaitForAsyncTickDoNotFinalize(bool bEnsureComplete)
 
 	const uint64 StartCycles = FPlatformTime::Cycles64();
 	const double WarnSeconds = 5.0;
-	const uint64 WarnCycles = StartCycles + uint64(WarnSeconds / FPlatformTime::GetSecondsPerCycle());
+	const uint64 WarnCycles = StartCycles + uint64(WarnSeconds / FPlatformTime::GetSecondsPerCycle64());
 	bool bDoWarning = true;
 
 	while ( bAsyncWorkInProgress )
@@ -1790,7 +1790,7 @@ void FNiagaraSystemInstance::WaitForAsyncTickDoNotFinalize(bool bEnsureComplete)
 		if ( bDoWarning && (FPlatformTime::Cycles64() > WarnCycles) )
 		{
 			bDoWarning = false;
-			UE_LOG(LogNiagara, Warning, TEXT("Niagara Effect has stalled GT for %g and is not complete, this may result in a deadlock.\nComponent: %s \nSystem: %s"), WarnSeconds, *GetFullNameSafe(Component), *GetFullNameSafe(GetSystem()));
+			UE_LOG(LogNiagara, Warning, TEXT("Niagara Effect has stalled GT for %g seconds and is not complete, this may result in a deadlock.\nComponent: %s \nSystem: %s"), WarnSeconds, *GetFullNameSafe(Component), *GetFullNameSafe(GetSystem()));
 		}
 	}
 
