- Fix for random skel mesh triangle sampling on some platforms where the int / uint difference will cause a compiler error

#rnx

diff --git a/Engine/Plugins/FX/Niagara/Shaders/Private/NiagaraDataInterfaceSkeletalMesh.ush b/Engine/Plugins/FX/Niagara/Shaders/Private/NiagaraDataInterfaceSkeletalMesh.ush
--- a/Engine/Plugins/FX/Niagara/Shaders/Private/NiagaraDataInterfaceSkeletalMesh.ush
+++ b/Engine/Plugins/FX/Niagara/Shaders/Private/NiagaraDataInterfaceSkeletalMesh.ush
@@ -179,7 +179,7 @@ FDISkelMeshSkinnedVertex DISkelMesh_GetSkinnedPointOnTriangleWS(in FDISkelMeshCo
 	return FinalVertex;
 }
 
-void DISKelMesh_RandomTriCoord(in FDISkelMeshContext DIContext, NiagaraRandInfo InRandomInfo, out uint OutTriangle, out float3 OutBaryCoord)
+void DISKelMesh_RandomTriCoord(in FDISkelMeshContext DIContext, NiagaraRandInfo InRandomInfo, out int OutTriangle, out float3 OutBaryCoord)
 {
 	float RandT0 = NiagaraInternalNoise(InRandomInfo.Seed1, InRandomInfo.Seed2, InRandomInfo.Seed3);
 
