- Removed Camera Facing Vector Mask
- Rather than using this masking method functions should be used inside the ParticleUpdate

diff --git a/Engine/Plugins/FX/Niagara/Shaders/Private/NiagaraSpriteVertexFactory.ush b/Engine/Plugins/FX/Niagara/Shaders/Private/NiagaraSpriteVertexFactory.ush
--- a/Engine/Plugins/FX/Niagara/Shaders/Private/NiagaraSpriteVertexFactory.ush
+++ b/Engine/Plugins/FX/Niagara/Shaders/Private/NiagaraSpriteVertexFactory.ush
@@ -28,7 +28,7 @@
 #define FACING_MODE_FACECAMERA 0
 /** The sprite billboard plane is completely parallel to the camera plane. Particle always looks "flat" */
 #define FACING_MODE_FACECAMERAPLANE 1
-/** The sprite billboard faces toward the "Particles.SpriteFacing" vector attribute, using the per-axis CustomFacingVectorMask as a lerp factor from the standard FaceCamera mode.*/
+/** The sprite billboard faces toward the "Particles.SpriteFacing" vector attribute. */
 #define FACING_MODE_CUSTOMFACINGVECTOR 2
 /** Faces the camera position, but is not dependent on the camera rotation.  This method produces more stable particles under camera rotation. Uses the up axis of (0,0,1).*/
 #define FACING_MODE_FACECAMERAPOSITION 3 
@@ -440,10 +440,7 @@ void GetTangents(FVertexFactoryInput Input, FVertexFactoryIntermediates Intermed
 	float3 CameraDirection = SafeNormalize(CameraVector);
 	float3 ParticleDirection	= SafeNormalizeWithFallback(WorldPosition.xyz - OldWorldPosition.xyz, float3(0,0,1));
 
-	// First interpolate between the base facing mode and the custom facing mode. We assume that the code setting up this call has the facing vector mask
-	// set to 0,0,0,0 if not in FACING_MODE_CUSTOMFACINGVECTOR.
-	float3 CustomFacingVectorMask =  NiagaraSpriteVF.CustomFacingVectorMask;
-	float3 CustomFacingVector = Intermediates.CustomFacingVector * CustomFacingVectorMask + CameraDirection * (float3(1.0,1.0,1.0)-CustomFacingVectorMask);
+	float3 CustomFacingVector = NiagaraSpriteVFLooseParameters.ParticleFacingMode == FACING_MODE_CUSTOMFACINGVECTOR ? Intermediates.CustomFacingVector : CameraDirection;
 	
 	// the particle's movement direction.
 	float3 DirectionVector	= lerp(ParticleDirection, Intermediates.CustomAlignmentVector, NiagaraSpriteVFLooseParameters.ParticleAlignmentMode == ALIGNMENT_MODE_CUSTOMALIGNMENT);
diff --git a/Engine/Plugins/FX/Niagara/Source/Niagara/Private/NiagaraConstants.cpp b/Engine/Plugins/FX/Niagara/Source/Niagara/Private/NiagaraConstants.cpp
--- a/Engine/Plugins/FX/Niagara/Source/Niagara/Private/NiagaraConstants.cpp
+++ b/Engine/Plugins/FX/Niagara/Source/Niagara/Private/NiagaraConstants.cpp
@@ -385,7 +385,7 @@ void FNiagaraConstants::Init()
 		AttrDescStrMap.Add(SYS_PARAM_PARTICLES_SPRITE_ROTATION, LOCTEXT("SpriteRotDesc", "The screen aligned roll of the particle in degrees."));
 		AttrDescStrMap.Add(SYS_PARAM_PARTICLES_NORMALIZED_AGE, LOCTEXT("NormalizedAgeDesc", "The age in seconds divided by lifetime in seconds. Useful for animation as the value is between 0 and 1."));
 		AttrDescStrMap.Add(SYS_PARAM_PARTICLES_SPRITE_SIZE, LOCTEXT("SpriteSizeDesc", "The size of the sprite quad."));
-		AttrDescStrMap.Add(SYS_PARAM_PARTICLES_SPRITE_FACING, LOCTEXT("FacingDesc", "Makes the surface of the sprite face towards a custom vector. Must be used with the SpriteRenderer's CustomFacingVector FacingMode and CustomFacingVectorMask options."));
+		AttrDescStrMap.Add(SYS_PARAM_PARTICLES_SPRITE_FACING, LOCTEXT("FacingDesc", "Makes the surface of the sprite face towards a custom vector. Must be used with the SpriteRenderer's CustomFacingVector FacingMode option."));
 		AttrDescStrMap.Add(SYS_PARAM_PARTICLES_SPRITE_ALIGNMENT, LOCTEXT("AlignmentDesc", "Imagine the texture having an arrow pointing up, this attribute makes the arrow point towards the alignment axis. Must be used with the SpriteRenderer's CustomAlignment Alignment option."));
 		AttrDescStrMap.Add(SYS_PARAM_PARTICLES_SUB_IMAGE_INDEX, LOCTEXT("SubImageIndexDesc", "A value from 0 to the number of entries in the table of SubUV images."));
 		FText DynParamText = LOCTEXT("DynamicMaterialParameterDesc", "The 4-float vector used to send custom data to renderer.");
diff --git a/Engine/Plugins/FX/Niagara/Source/Niagara/Private/NiagaraRendererSprites.cpp b/Engine/Plugins/FX/Niagara/Source/Niagara/Private/NiagaraRendererSprites.cpp
--- a/Engine/Plugins/FX/Niagara/Source/Niagara/Private/NiagaraRendererSprites.cpp
+++ b/Engine/Plugins/FX/Niagara/Source/Niagara/Private/NiagaraRendererSprites.cpp
@@ -98,7 +98,6 @@ FNiagaraRendererSprites::FNiagaraRendererSprites(ERHIFeatureLevel::Type FeatureL
 	: FNiagaraRenderer(FeatureLevel, InProps, Emitter)
 	, Alignment(ENiagaraSpriteAlignment::Unaligned)
 	, FacingMode(ENiagaraSpriteFacingMode::FaceCamera)
-	, CustomFacingVectorMask(ForceInitToZero)
 	, PivotInUVSpace(0.5f, 0.5f)
 	, SortMode(ENiagaraSortMode::ViewDistance)
 	, SubImageSize(1.0f, 1.0f)
@@ -115,7 +114,6 @@ FNiagaraRendererSprites::FNiagaraRendererSprites(ERHIFeatureLevel::Type FeatureL
 	
 	Alignment = Properties->Alignment;
 	FacingMode = Properties->FacingMode;
-	CustomFacingVectorMask = Properties->CustomFacingVectorMask;
 	PivotInUVSpace = Properties->PivotInUVSpace;
 	SortMode = Properties->SortMode;
 	SubImageSize = Properties->SubImageSize;
@@ -235,7 +233,6 @@ FNiagaraSpriteUniformBufferRef FNiagaraRendererSprites::CreatePerViewUniformBuff
 	PerViewUniformParameters.MacroUVParameters = FVector4(0.0f, 0.0f, 1.0f, 1.0f);
 	PerViewUniformParameters.CameraFacingBlend = FVector4(0.0f, 0.0f, 0.0f, 1.0f);
 	PerViewUniformParameters.RemoveHMDRoll = bRemoveHMDRollInVR;
-	PerViewUniformParameters.CustomFacingVectorMask = FVector4(0.0f, 0.0f, 0.0f, 0.0f);
 	PerViewUniformParameters.SubImageSize = FVector4(SubImageSize.X, SubImageSize.Y, 1.0f / SubImageSize.X, 1.0f / SubImageSize.Y);
 
 	PerViewUniformParameters.PositionDataOffset = VFVariables[ENiagaraSpriteVFLayout::Position].GetGPUOffset();
@@ -291,11 +288,6 @@ FNiagaraSpriteUniformBufferRef FNiagaraRendererSprites::CreatePerViewUniformBuff
 		PerViewUniformParameters.TangentSelector = FVector4(0.0f, 1.0f, 0.0f, 0.0f);
 	}
 
-	if (ActualFacingMode == ENiagaraSpriteFacingMode::CustomFacingVector)
-	{
-		PerViewUniformParameters.CustomFacingVectorMask = CustomFacingVectorMask;
-	}
-
 	return FNiagaraSpriteUniformBufferRef::CreateUniformBufferImmediate(PerViewUniformParameters, UniformBuffer_SingleFrame);
 }
 
diff --git a/Engine/Plugins/FX/Niagara/Source/Niagara/Private/NiagaraSpriteRendererProperties.cpp b/Engine/Plugins/FX/Niagara/Source/Niagara/Private/NiagaraSpriteRendererProperties.cpp
--- a/Engine/Plugins/FX/Niagara/Source/Niagara/Private/NiagaraSpriteRendererProperties.cpp
+++ b/Engine/Plugins/FX/Niagara/Source/Niagara/Private/NiagaraSpriteRendererProperties.cpp
@@ -32,7 +32,6 @@ FCookStatsManager::FAutoRegisterCallback NiagaraCutoutCookStats::RegisterCookSta
 UNiagaraSpriteRendererProperties::UNiagaraSpriteRendererProperties()
 	: Alignment(ENiagaraSpriteAlignment::Unaligned)
 	, FacingMode(ENiagaraSpriteFacingMode::FaceCamera)
-	, CustomFacingVectorMask(ForceInitToZero)
 	, PivotInUVSpace(0.5f, 0.5f)
 	, SortMode(ENiagaraSortMode::None)
 	, SubImageSize(1.0f, 1.0f)
diff --git a/Engine/Plugins/FX/Niagara/Source/Niagara/Public/NiagaraRendererSprites.h b/Engine/Plugins/FX/Niagara/Source/Niagara/Public/NiagaraRendererSprites.h
--- a/Engine/Plugins/FX/Niagara/Source/Niagara/Public/NiagaraRendererSprites.h
+++ b/Engine/Plugins/FX/Niagara/Source/Niagara/Public/NiagaraRendererSprites.h
@@ -60,7 +60,6 @@ private:
 	//Cached data from the properties struct.
 	ENiagaraSpriteAlignment Alignment;
 	ENiagaraSpriteFacingMode FacingMode;
-	FVector CustomFacingVectorMask;
 	FVector2D PivotInUVSpace;
 	ENiagaraSortMode SortMode;
 	FVector2D SubImageSize;
diff --git a/Engine/Plugins/FX/Niagara/Source/Niagara/Public/NiagaraSpriteRendererProperties.h b/Engine/Plugins/FX/Niagara/Source/Niagara/Public/NiagaraSpriteRendererProperties.h
--- a/Engine/Plugins/FX/Niagara/Source/Niagara/Public/NiagaraSpriteRendererProperties.h
+++ b/Engine/Plugins/FX/Niagara/Source/Niagara/Public/NiagaraSpriteRendererProperties.h
@@ -30,7 +30,7 @@ enum class ENiagaraSpriteFacingMode : uint8
 	FaceCamera,
 	/** The sprite billboard plane is completely parallel to the camera plane. Particle always looks "flat" */
 	FaceCameraPlane,
-	/** The sprite billboard faces toward the "Particles.SpriteFacing" vector attribute, using the per-axis CustomFacingVectorMask as a lerp factor from the standard FaceCamera mode. If the "Particles.SpriteFacing" attribute is missing, this falls back to FaceCamera mode.*/
+	/** The sprite billboard faces toward the "Particles.SpriteFacing" vector attribute. If the "Particles.SpriteFacing" attribute is missing, this falls back to FaceCamera mode.*/
 	CustomFacingVector,
 	/** Faces the camera position, but is not dependent on the camera rotation.  This method produces more stable particles under camera rotation. Uses the up axis of (0,0,1).*/
 	FaceCameraPosition,
@@ -90,10 +90,6 @@ public:
 	UPROPERTY(EditAnywhere, Category = "Sprite Rendering")
 	ENiagaraSpriteFacingMode FacingMode;
 
-	/** Used as a per-axis interpolation factor with the CustomFacingVector mode to determine how the billboard orients itself relative to the camera. A value of 1.0 is fully facing the custom vector. A value of 0.0 uses the standard facing strategy.*/
-	UPROPERTY(EditAnywhere, Category = "Sprite Rendering")
-	FVector CustomFacingVectorMask;
-
 	/** Determines the location of the pivot point of this particle. It follows Unreal's UV space, which has the upper left of the image at 0,0 and bottom right at 1,1. The middle is at 0.5, 0.5. */
 	UPROPERTY(EditAnywhere, Category = "Sprite Rendering")
 	FVector2D PivotInUVSpace;
diff --git a/Engine/Plugins/FX/Niagara/Source/NiagaraVertexFactories/Public/NiagaraSpriteVertexFactory.h b/Engine/Plugins/FX/Niagara/Source/NiagaraVertexFactories/Public/NiagaraSpriteVertexFactory.h
--- a/Engine/Plugins/FX/Niagara/Source/NiagaraVertexFactories/Public/NiagaraSpriteVertexFactory.h
+++ b/Engine/Plugins/FX/Niagara/Source/NiagaraVertexFactories/Public/NiagaraSpriteVertexFactory.h
@@ -20,7 +20,6 @@ class FMaterial;
  */
 BEGIN_GLOBAL_SHADER_PARAMETER_STRUCT( FNiagaraSpriteUniformParameters, NIAGARAVERTEXFACTORIES_API)
 	SHADER_PARAMETER( uint32, bLocalSpace)
-	SHADER_PARAMETER_EX( FVector, CustomFacingVectorMask, EShaderPrecisionModifier::Half)
 	SHADER_PARAMETER_EX( FVector4, TangentSelector, EShaderPrecisionModifier::Half )
 	SHADER_PARAMETER_EX( FVector4, NormalsSphereCenter, EShaderPrecisionModifier::Half )
 	SHADER_PARAMETER_EX( FVector4, NormalsCylinderUnitDirection, EShaderPrecisionModifier::Half )
