Initialized to zero any parameter data added in FNiagaraParameterStore.
This allows unitialized parameters to have deterministic cooking.

diff --git a/Engine/Plugins/FX/Niagara/Source/Niagara/Private/NiagaraParameterStore.cpp b/Engine/Plugins/FX/Niagara/Source/Niagara/Private/NiagaraParameterStore.cpp
--- a/Engine/Plugins/FX/Niagara/Source/Niagara/Private/NiagaraParameterStore.cpp
+++ b/Engine/Plugins/FX/Niagara/Source/Niagara/Private/NiagaraParameterStore.cpp
@@ -505,7 +505,9 @@ bool FNiagaraParameterStore::AddParameter(const FNiagaraVariable& Param, bool bI
 		int32 ParamAlignment = Param.GetAlignment();
 		//int32 Offset = AlignArbitrary(ParameterData.Num(), ParamAlignment);//TODO: We need to handle alignment better here. Need to both satisfy CPU and GPU alignment concerns. VM doesn't care but the VM complier needs to be aware. Probably best to have everything adhere to GPU alignment rules.
 		Offset = ParameterData.Num();
-		ParameterData.AddUninitialized(ParamSize);
+		// Memory must be initialized in order to have deterministic cooking. 
+		// This is because some system parameters never get initialized otherwise (particle count, owner rotation, ...)
+		ParameterData.AddZeroed(ParamSize);
 				
 		INC_MEMORY_STAT_BY(STAT_NiagaraParamStoreMemory, ParameterData.GetAllocatedSize());
 
