- Fix for not unregistering with scalability when Deactive is called and scalability had already deactivated the system
- This would leave the component inside scalability allowing it to be reactivated and never garbage collected

diff --git a/Engine/Plugins/FX/Niagara/Source/Niagara/Private/NiagaraComponent.cpp b/Engine/Plugins/FX/Niagara/Source/Niagara/Private/NiagaraComponent.cpp
--- a/Engine/Plugins/FX/Niagara/Source/Niagara/Private/NiagaraComponent.cpp
+++ b/Engine/Plugins/FX/Niagara/Source/Niagara/Private/NiagaraComponent.cpp
@@ -914,41 +914,31 @@ void UNiagaraComponent::Deactivate()
 
 void UNiagaraComponent::DeactivateInternal(bool bIsScalabilityCull /* = false */)
 {
-	if (IsActive())
+	// Unregister with the scalability manager if this is a genuine deactivation from outside.
+	// The scalability manager itself can call this function when culling systems.
+	if (bIsScalabilityCull == false)
+	{
+		UnregisterWithScalabilityManager();
+	}
+
+	if (IsActive() && SystemInstance)
 	{
 		SCOPE_CYCLE_COUNTER(STAT_NiagaraComponentDeactivate);
 		CSV_SCOPED_TIMING_STAT_EXCLUSIVE(Niagara);
 
 		//UE_LOG(LogNiagara, Log, TEXT("Deactivate: %u - %s"), this, *Asset->GetName());
 
-		if (SystemInstance)
-		{
-			// Don't deactivate in solo mode as we are not ticked by the world but rather the component
-			// Deactivating will cause the system to never Complete
-			if (SystemInstance->IsSolo() == false)
-			{
-				Super::Deactivate();
-			}
-
-			SystemInstance->Deactivate();
-
-			// We are considered active until we are complete
-			SetActiveFlag(!SystemInstance->IsComplete());
-		}
-		else
+		// Don't deactivate in solo mode as we are not ticked by the world but rather the component
+		// Deactivating will cause the system to never Complete
+		if (SystemInstance->IsSolo() == false)
 		{
 			Super::Deactivate();
-			SetActiveFlag(false);
-		}
-
-		//Unregister with the scalability manager if this is a genuine deactivation from outside.
-		//The scalability manager itself can call this function when culling systems.
-		if (bIsScalabilityCull == false)
-		{
-			UnregisterWithScalabilityManager();
 		}
 
 		SystemInstance->Deactivate();
+
+		// We are considered active until we are complete
+		SetActiveFlag(!SystemInstance->IsComplete());
 	}
 	else
 	{
