Niagara - Fix an issue where input rows in the stack weren't being refreshed correctly when updating a function script due to deprecation or a missing script asset.

diff --git a/Engine/Plugins/FX/Niagara/Source/NiagaraEditor/Private/ViewModels/Stack/NiagaraStackFunctionInput.cpp b/Engine/Plugins/FX/Niagara/Source/NiagaraEditor/Private/ViewModels/Stack/NiagaraStackFunctionInput.cpp
--- a/Engine/Plugins/FX/Niagara/Source/NiagaraEditor/Private/ViewModels/Stack/NiagaraStackFunctionInput.cpp
+++ b/Engine/Plugins/FX/Niagara/Source/NiagaraEditor/Private/ViewModels/Stack/NiagaraStackFunctionInput.cpp
@@ -134,6 +134,7 @@ void UNiagaraStackFunctionInput::Initialize(
 	Super::Initialize(InRequiredEntryData, bInputIsAdvanced, InOwnerStackItemEditorDataKey, InputStackEditorDataKey);
 	OwningModuleNode = &InModuleNode;
 	OwningFunctionCallNode = &InInputFunctionCallNode;
+	OwningFunctionCallInitialScript = OwningFunctionCallNode->FunctionScript;
 	OwningAssignmentNode = Cast<UNiagaraNodeAssignment>(OwningFunctionCallNode.Get());
 
 	UNiagaraNodeOutput* OutputNode = FNiagaraStackGraphUtilities::GetEmitterOutputNodeForStackNode(*OwningModuleNode.Get());
@@ -231,6 +232,11 @@ const UNiagaraNodeFunctionCall& UNiagaraStackFunctionInput::GetInputFunctionCall
 	return *OwningFunctionCallNode.Get();
 }
 
+UNiagaraScript* UNiagaraStackFunctionInput::GetInputFunctionCallInitialScript() const
+{
+	return OwningFunctionCallInitialScript.Get();
+}
+
 UNiagaraStackFunctionInput::EValueMode UNiagaraStackFunctionInput::GetValueMode()
 {
 	return InputValues.Mode;
diff --git a/Engine/Plugins/FX/Niagara/Source/NiagaraEditor/Private/ViewModels/Stack/NiagaraStackInputCategory.cpp b/Engine/Plugins/FX/Niagara/Source/NiagaraEditor/Private/ViewModels/Stack/NiagaraStackInputCategory.cpp
--- a/Engine/Plugins/FX/Niagara/Source/NiagaraEditor/Private/ViewModels/Stack/NiagaraStackInputCategory.cpp
+++ b/Engine/Plugins/FX/Niagara/Source/NiagaraEditor/Private/ViewModels/Stack/NiagaraStackInputCategory.cpp
@@ -43,8 +43,10 @@ void UNiagaraStackInputCategory::RefreshChildrenInternal(const TArray<UNiagaraSt
 {
 	for (FInputParameterHandleAndType& Input : Inputs)
 	{
-		UNiagaraStackFunctionInput* InputChild = FindCurrentChildOfTypeByPredicate<UNiagaraStackFunctionInput>(CurrentChildren,
-			[&](UNiagaraStackFunctionInput* CurrentInput) { return CurrentInput->GetInputParameterHandle() == Input.ParameterHandle; });
+		UNiagaraStackFunctionInput* InputChild = FindCurrentChildOfTypeByPredicate<UNiagaraStackFunctionInput>(CurrentChildren, [&](UNiagaraStackFunctionInput* CurrentInput) 
+		{ 
+			return CurrentInput->GetInputParameterHandle() == Input.ParameterHandle && CurrentInput->GetInputFunctionCallInitialScript() == InputFunctionCallNode->FunctionScript; 
+		});
 
 		if (InputChild == nullptr)
 		{
diff --git a/Engine/Plugins/FX/Niagara/Source/NiagaraEditor/Public/ViewModels/Stack/NiagaraStackFunctionInput.h b/Engine/Plugins/FX/Niagara/Source/NiagaraEditor/Public/ViewModels/Stack/NiagaraStackFunctionInput.h
--- a/Engine/Plugins/FX/Niagara/Source/NiagaraEditor/Public/ViewModels/Stack/NiagaraStackFunctionInput.h
+++ b/Engine/Plugins/FX/Niagara/Source/NiagaraEditor/Public/ViewModels/Stack/NiagaraStackFunctionInput.h
@@ -78,6 +78,9 @@ public:
 	/** Gets the function call node which owns this input. */
 	const UNiagaraNodeFunctionCall& GetInputFunctionCallNode() const;
 
+	/** Gets the script that the function call node was referencing when this input was initialized. */
+	UNiagaraScript* GetInputFunctionCallInitialScript() const;
+
 	/** Gets the current value mode */
 	EValueMode GetValueMode();
 
@@ -394,6 +397,9 @@ private:
 	/** The function call which this entry is an input to. NOTE: This node can be a module function call node or a dynamic input node. */
 	TWeakObjectPtr<UNiagaraNodeFunctionCall> OwningFunctionCallNode;
 
+	/** The script which the owning function call is referencing. */
+	TWeakObjectPtr<UNiagaraScript> OwningFunctionCallInitialScript;
+
 	/** The assignment node which owns this input.  This is only valid for inputs of assignment modules. */
 	TWeakObjectPtr<UNiagaraNodeAssignment> OwningAssignmentNode;
 
