FNM - Allo Niagara Vertgexfactories to bind differnt buffertypes per UseCase. This will eventually allow us to reduce the memory footprint with a later check-in.

diff --git a/Engine/Plugins/FX/Niagara/Shaders/Private/NiagaraMeshVertexFactory.ush b/Engine/Plugins/FX/Niagara/Shaders/Private/NiagaraMeshVertexFactory.ush
--- a/Engine/Plugins/FX/Niagara/Shaders/Private/NiagaraMeshVertexFactory.ush
+++ b/Engine/Plugins/FX/Niagara/Shaders/Private/NiagaraMeshVertexFactory.ush
@@ -13,6 +13,20 @@
 #define USE_PARTICLE_RANDOM ((FEATURE_LEVEL >= FEATURE_LEVEL_SM5 || FEATURE_LEVEL == FEATURE_LEVEL_ES3_1) && NEEDS_PARTICLE_RANDOM)
 #define USE_PARTICLE_TRANSFORM ((FEATURE_LEVEL >= FEATURE_LEVEL_SM5 || FEATURE_LEVEL == FEATURE_LEVEL_ES3_1) && NEEDS_PARTICLE_TRANSFORM)
 
+#if NIAGARA_LOOSEPARAM_BINDING
+	Buffer<float> NiagaraParticleDataPosition;
+	Buffer<float> NiagaraParticleDataVelocity;
+	Buffer<float> NiagaraParticleDataColor;
+	Buffer<float> NiagaraParticleDataScale;
+	Buffer<float> NiagaraParticleDataTransform;
+	Buffer<float> NiagaraParticleDataNormalizedAge;
+	Buffer<float> NiagaraParticleDataMaterialRandom;
+	Buffer<float> NiagaraParticleDataMaterialParam0;
+	Buffer<float> NiagaraParticleDataMaterialParam1;
+	Buffer<float> NiagaraParticleDataMaterialParam2;
+	Buffer<float> NiagaraParticleDataMaterialParam3;
+#endif
+
 float3 TransformPosition(float3 InPosition)
 {
 	if (NiagaraMeshVF.bLocalSpace)
@@ -510,7 +524,7 @@ FVertexFactoryIntermediates GetVertexFactoryIntermediates(FVertexFactoryInput In
 	}
 	else
 	{
-		Intermediates.ParticlePosition = float4(GetVec3(NiagaraMeshVF.PositionDataOffset, Input.InstanceId), 1);
+		Intermediates.ParticlePosition = float4(GetVec3(NIAGARA_LOOSEPARAM_STRUCT(NiagaraParticleDataPosition), NiagaraMeshVF.PositionDataOffset, Input.InstanceId), 1);
 	}
 
 	if (NiagaraMeshVF.ColorDataOffset == -1)
@@ -519,7 +533,7 @@ FVertexFactoryIntermediates GetVertexFactoryIntermediates(FVertexFactoryInput In
 	}
 	else
 	{
-		Intermediates.ParticleColor = GetVec4(NiagaraMeshVF.ColorDataOffset, Input.InstanceId);
+		Intermediates.ParticleColor = GetVec4(NIAGARA_LOOSEPARAM_STRUCT(NiagaraParticleDataColor), NiagaraMeshVF.ColorDataOffset, Input.InstanceId);
 	}
 
 	float3 ParticleVelocity;
@@ -529,7 +543,7 @@ FVertexFactoryIntermediates GetVertexFactoryIntermediates(FVertexFactoryInput In
 	}
 	else
 	{
-		ParticleVelocity = float3(GetVec3(NiagaraMeshVF.VelocityDataOffset, Input.InstanceId));
+		ParticleVelocity = float3(GetVec3(NIAGARA_LOOSEPARAM_STRUCT(NiagaraParticleDataVelocity), NiagaraMeshVF.VelocityDataOffset, Input.InstanceId));
 	}
 
 #if (DYNAMIC_PARAMETERS_MASK & 1)
@@ -539,7 +553,7 @@ FVertexFactoryIntermediates GetVertexFactoryIntermediates(FVertexFactoryInput In
 	}
 	else
 	{
-		Intermediates.DynamicParameter = GetVec4(NiagaraMeshVF.MaterialParamDataOffset, Input.InstanceId);
+		Intermediates.DynamicParameter = GetVec4(NIAGARA_LOOSEPARAM_STRUCT(NiagaraParticleDataMaterialParam0), NiagaraMeshVF.MaterialParamDataOffset, Input.InstanceId);
 	}
 #endif
 #if (DYNAMIC_PARAMETERS_MASK & 2)
@@ -549,7 +563,7 @@ FVertexFactoryIntermediates GetVertexFactoryIntermediates(FVertexFactoryInput In
 	}	
 	else
 	{
-		Intermediates.DynamicParameter1 = GetVec4(NiagaraMeshVF.MaterialParam1DataOffset, Input.InstanceId);
+		Intermediates.DynamicParameter1 = GetVec4(NIAGARA_LOOSEPARAM_STRUCT(NiagaraParticleDataMaterialParam1), NiagaraMeshVF.MaterialParam1DataOffset, Input.InstanceId);
 	}
 #endif
 #if (DYNAMIC_PARAMETERS_MASK & 4)
@@ -559,7 +573,7 @@ FVertexFactoryIntermediates GetVertexFactoryIntermediates(FVertexFactoryInput In
 	}
 	else
 	{
-		Intermediates.DynamicParameter2 = GetVec4(NiagaraMeshVF.MaterialParam2DataOffset, Input.InstanceId);
+		Intermediates.DynamicParameter2 = GetVec4(NIAGARA_LOOSEPARAM_STRUCT(NiagaraParticleDataMaterialParam2), NiagaraMeshVF.MaterialParam2DataOffset, Input.InstanceId);
 	}
 #endif
 #if (DYNAMIC_PARAMETERS_MASK & 8)
@@ -569,7 +583,7 @@ FVertexFactoryIntermediates GetVertexFactoryIntermediates(FVertexFactoryInput In
 	}
 	else
 	{
-		Intermediates.DynamicParameter3 = GetVec4(NiagaraMeshVF.MaterialParam3DataOffset, Input.InstanceId);
+		Intermediates.DynamicParameter3 = GetVec4(NIAGARA_LOOSEPARAM_STRUCT(NiagaraParticleDataMaterialParam3), NiagaraMeshVF.MaterialParam3DataOffset, Input.InstanceId);
 	}
 #endif
 
@@ -580,7 +594,7 @@ FVertexFactoryIntermediates GetVertexFactoryIntermediates(FVertexFactoryInput In
 	}
 	else
 	{
-		Intermediates.RelativeTime = GetFloat(NiagaraMeshVF.NormalizedAgeDataOffset, Input.InstanceId);
+		Intermediates.RelativeTime = GetFloat(NIAGARA_LOOSEPARAM_STRUCT(NiagaraParticleDataNormalizedAge), NiagaraMeshVF.NormalizedAgeDataOffset, Input.InstanceId);
 	}
 #endif
 
@@ -591,7 +605,7 @@ FVertexFactoryIntermediates GetVertexFactoryIntermediates(FVertexFactoryInput In
 	}
 	else
 	{
-		Intermediates.ParticleRandom = GetFloat(NiagaraMeshVF.MaterialRandomDataOffset, Input.InstanceId);
+		Intermediates.ParticleRandom = GetFloat(NIAGARA_LOOSEPARAM_STRUCT(NiagaraParticleDataMaterialRandom), NiagaraMeshVF.MaterialRandomDataOffset, Input.InstanceId);
 	}
 #endif
 
@@ -608,7 +622,7 @@ FVertexFactoryIntermediates GetVertexFactoryIntermediates(FVertexFactoryInput In
 
 	if(NiagaraMeshVF.TransformDataOffset != -1)
 	{
-		float4 Quat = GetVec4(NiagaraMeshVF.TransformDataOffset, Input.InstanceId);
+		float4 Quat = GetVec4(NIAGARA_LOOSEPARAM_STRUCT(NiagaraParticleDataTransform), NiagaraMeshVF.TransformDataOffset, Input.InstanceId);
 
 		//Transform axis to handle local space.
 		//AxisAngle.xyz = TransformVector(AxisAngle.xyz);
@@ -622,7 +636,7 @@ FVertexFactoryIntermediates GetVertexFactoryIntermediates(FVertexFactoryInput In
 	float4x4 ScaleMatrix = float4x4(float4(1.0f, 0.0f, 0.0f, 0.0f), float4(0.0f, 1.0f, 0.0f, 0.0f), float4(0.0f, 0.0f, 1.0f, 0.0f), float4(0.0f, 0.0f, 0.0f, 1.0f));
 	if(NiagaraMeshVF.ScaleDataOffset != -1)
 	{
-		float3 Scale = GetVec3(NiagaraMeshVF.ScaleDataOffset, Input.InstanceId);
+		float3 Scale = GetVec3(NIAGARA_LOOSEPARAM_STRUCT(NiagaraParticleDataScale), NiagaraMeshVF.ScaleDataOffset, Input.InstanceId);
 		ScaleMatrix[0][0] *= Scale.x;
 		ScaleMatrix[1][1] *= Scale.y;
 		ScaleMatrix[2][2] *= Scale.z;
diff --git a/Engine/Plugins/FX/Niagara/Shaders/Private/NiagaraRibbonVertexFactory.ush b/Engine/Plugins/FX/Niagara/Shaders/Private/NiagaraRibbonVertexFactory.ush
--- a/Engine/Plugins/FX/Niagara/Shaders/Private/NiagaraRibbonVertexFactory.ush
+++ b/Engine/Plugins/FX/Niagara/Shaders/Private/NiagaraRibbonVertexFactory.ush
@@ -14,6 +14,21 @@
 #define USE_PARTICLE_RANDOM ((FEATURE_LEVEL >= FEATURE_LEVEL_SM5 || FEATURE_LEVEL == FEATURE_LEVEL_ES3_1) && NEEDS_PARTICLE_RANDOM)
 #define USE_PARTICLE_INTERPOLATION (FEATURE_LEVEL >= FEATURE_LEVEL_SM5 || FEATURE_LEVEL == FEATURE_LEVEL_ES3_1)
 
+#if NIAGARA_LOOSEPARAM_BINDING
+	Buffer<float> NiagaraParticleDataPosition;
+	Buffer<float> NiagaraParticleDataVelocity;
+	Buffer<float> NiagaraParticleDataColor;
+	Buffer<float> NiagaraParticleDataWidth;
+	Buffer<float> NiagaraParticleDataTwist;
+	Buffer<float> NiagaraParticleDataFacing;
+	Buffer<float> NiagaraParticleDataNormalizedAge;
+	Buffer<float> NiagaraParticleDataMaterialRandom;
+	Buffer<float> NiagaraParticleDataMaterialParam0;
+	Buffer<float> NiagaraParticleDataMaterialParam1;
+	Buffer<float> NiagaraParticleDataMaterialParam2;
+	Buffer<float> NiagaraParticleDataMaterialParam3;
+#endif
+
 float3 TransformPosition(float3 InPosition)
 {
 	if (NiagaraRibbonVF.bLocalSpace)
@@ -144,8 +159,8 @@ void GetNiagaraParticlePositionAndDirection(in FInterpVtxData InterpVtxData, in
 	BRANCH // Also acts as ISOLATE
 	if (NiagaraRibbonVF.PositionDataOffset != -1)
 	{
-		const float3 Pos0 = GetVec3(NiagaraRibbonVF.PositionDataOffset, InterpVtxData.ParticleId0);
-		const float3 Pos1 = GetVec3(NiagaraRibbonVF.PositionDataOffset, InterpVtxData.ParticleId1);
+		const float3 Pos0 = GetVec3(NIAGARA_LOOSEPARAM_STRUCT(NiagaraParticleDataPosition), NiagaraRibbonVF.PositionDataOffset, InterpVtxData.ParticleId0);
+		const float3 Pos1 = GetVec3(NIAGARA_LOOSEPARAM_STRUCT(NiagaraParticleDataPosition), NiagaraRibbonVF.PositionDataOffset, InterpVtxData.ParticleId1);
 		const float4 TangentAndDist0 = NiagaraRibbonVFLooseParameters.TangentsAndDistances[InterpVtxData.RawParticleId0];
 		const float4 TangentAndDist1 = NiagaraRibbonVFLooseParameters.TangentsAndDistances[InterpVtxData.RawParticleId1];
 		
@@ -161,8 +176,8 @@ void GetNiagaraParticleVelocity(in FInterpVtxData InterpVtxData, in out FVertexF
 	BRANCH
 	if(NiagaraRibbonVF.VelocityDataOffset != -1)
 	{
-		const float3 Velocity0 = GetVec3(NiagaraRibbonVF.VelocityDataOffset, InterpVtxData.ParticleId0);
-		const float3 Velocity1 = GetVec3(NiagaraRibbonVF.VelocityDataOffset, InterpVtxData.ParticleId1);
+		const float3 Velocity0 = GetVec3(NIAGARA_LOOSEPARAM_STRUCT(NiagaraParticleDataVelocity), NiagaraRibbonVF.VelocityDataOffset, InterpVtxData.ParticleId0);
+		const float3 Velocity1 = GetVec3(NIAGARA_LOOSEPARAM_STRUCT(NiagaraParticleDataVelocity), NiagaraRibbonVF.VelocityDataOffset, InterpVtxData.ParticleId1);
 		const float3 Velocity = TransformVector(lerp(Velocity0, Velocity1, InterpVtxData.Alpha));
 
 		Intermediates.ParticleVelocity = float4(normalize(Velocity), length(Velocity));
@@ -176,8 +191,8 @@ void GetNiagaraParticleColor(in FInterpVtxData InterpVtxData, in out FVertexFact
 
 	if (NiagaraRibbonVF.ColorDataOffset != -1)
 	{
-		const float4 Color0 = GetVec4(NiagaraRibbonVF.ColorDataOffset, InterpVtxData.ParticleId0);
-		const float4 Color1 = GetVec4(NiagaraRibbonVF.ColorDataOffset, InterpVtxData.ParticleId1);
+		const float4 Color0 = GetVec4(NIAGARA_LOOSEPARAM_STRUCT(NiagaraParticleDataColor), NiagaraRibbonVF.ColorDataOffset, InterpVtxData.ParticleId0);
+		const float4 Color1 = GetVec4(NIAGARA_LOOSEPARAM_STRUCT(NiagaraParticleDataColor), NiagaraRibbonVF.ColorDataOffset, InterpVtxData.ParticleId1);
 		Intermediates.Color = lerp(Color0, Color1, InterpVtxData.Alpha);
 	}
 }
@@ -191,7 +206,12 @@ void GetNiagaraParticleDynamicParameters(in FInterpVtxData InterpVtxData, in out
 	FLATTEN
 	if (NiagaraRibbonVF.MaterialParamDataOffset != -1)
 	{
-		Intermediates.DynamicParameter = lerp(GetVec4(NiagaraRibbonVF.MaterialParamDataOffset, InterpVtxData.ParticleId0), GetVec4(NiagaraRibbonVF.MaterialParamDataOffset, InterpVtxData.ParticleId1), InterpVtxData.Alpha);
+		Intermediates.DynamicParameter = lerp
+		(
+			GetVec4(NIAGARA_LOOSEPARAM_STRUCT(NiagaraParticleDataMaterialParam0), NiagaraRibbonVF.MaterialParamDataOffset, InterpVtxData.ParticleId0), 
+			GetVec4(NIAGARA_LOOSEPARAM_STRUCT(NiagaraParticleDataMaterialParam0), NiagaraRibbonVF.MaterialParamDataOffset, InterpVtxData.ParticleId1), 
+			InterpVtxData.Alpha
+		);
 	}
 #endif
 #if (DYNAMIC_PARAMETERS_MASK & 2)
@@ -201,7 +221,12 @@ void GetNiagaraParticleDynamicParameters(in FInterpVtxData InterpVtxData, in out
 	FLATTEN
 	if (NiagaraRibbonVF.MaterialParam1DataOffset != -1)
 	{
-		Intermediates.DynamicParameter1 = lerp(GetVec4(NiagaraRibbonVF.MaterialParam1DataOffset, InterpVtxData.ParticleId0), GetVec4(NiagaraRibbonVF.MaterialParam1DataOffset, InterpVtxData.ParticleId1), InterpVtxData.Alpha);
+		Intermediates.DynamicParameter1 = lerp
+		(
+			GetVec4(NIAGARA_LOOSEPARAM_STRUCT(NiagaraParticleDataMaterialParam1), NiagaraRibbonVF.MaterialParam1DataOffset, InterpVtxData.ParticleId0), 
+			GetVec4(NIAGARA_LOOSEPARAM_STRUCT(NiagaraParticleDataMaterialParam1), NiagaraRibbonVF.MaterialParam1DataOffset, InterpVtxData.ParticleId1), 
+			InterpVtxData.Alpha
+		);
 	}
 #endif
 #if (DYNAMIC_PARAMETERS_MASK & 4)
@@ -211,7 +236,12 @@ void GetNiagaraParticleDynamicParameters(in FInterpVtxData InterpVtxData, in out
 	FLATTEN
 	if (NiagaraRibbonVF.MaterialParam2DataOffset != -1)
 	{
-		Intermediates.DynamicParameter2 = lerp(GetVec4(NiagaraRibbonVF.MaterialParam2DataOffset, InterpVtxData.ParticleId0), GetVec4(NiagaraRibbonVF.MaterialParam2DataOffset, InterpVtxData.ParticleId1), InterpVtxData.Alpha);
+		Intermediates.DynamicParameter2 = lerp
+		(
+			GetVec4(NIAGARA_LOOSEPARAM_STRUCT(NiagaraParticleDataMaterialParam2), NiagaraRibbonVF.MaterialParam2DataOffset, InterpVtxData.ParticleId0), 
+			GetVec4(NIAGARA_LOOSEPARAM_STRUCT(NiagaraParticleDataMaterialParam2), NiagaraRibbonVF.MaterialParam2DataOffset, InterpVtxData.ParticleId1), 
+			InterpVtxData.Alpha
+		);
 	}
 #endif
 #if (DYNAMIC_PARAMETERS_MASK & 8)
@@ -221,7 +251,12 @@ void GetNiagaraParticleDynamicParameters(in FInterpVtxData InterpVtxData, in out
 	FLATTEN
 	if (NiagaraRibbonVF.MaterialParam3DataOffset != -1)
 	{
-		Intermediates.DynamicParameter3 = lerp(GetVec4(NiagaraRibbonVF.MaterialParam3DataOffset, InterpVtxData.ParticleId0), GetVec4(NiagaraRibbonVF.MaterialParam3DataOffset, InterpVtxData.ParticleId1), InterpVtxData.Alpha);
+		Intermediates.DynamicParameter3 = lerp
+		(
+			GetVec4(NIAGARA_LOOSEPARAM_STRUCT(NiagaraParticleDataMaterialParam3), NiagaraRibbonVF.MaterialParam3DataOffset, InterpVtxData.ParticleId0), 
+			GetVec4(NIAGARA_LOOSEPARAM_STRUCT(NiagaraParticleDataMaterialParam3), NiagaraRibbonVF.MaterialParam3DataOffset, InterpVtxData.ParticleId1), 
+			InterpVtxData.Alpha
+		);
 	}
 #endif
 }
@@ -267,7 +302,12 @@ void GetNiagaraParticleRibbonWidth(in FInterpVtxData InterpVtxData, in out FVert
 	FLATTEN
 	if (NiagaraRibbonVF.WidthDataOffset != -1)
 	{
-		Intermediates.RibbonWidth = lerp(GetFloat(NiagaraRibbonVF.WidthDataOffset, InterpVtxData.ParticleId0), GetFloat(NiagaraRibbonVF.WidthDataOffset, InterpVtxData.ParticleId1), InterpVtxData.Alpha);
+		Intermediates.RibbonWidth = lerp
+		(
+			GetFloat(NIAGARA_LOOSEPARAM_STRUCT(NiagaraParticleDataWidth), NiagaraRibbonVF.WidthDataOffset, InterpVtxData.ParticleId0), 
+			GetFloat(NIAGARA_LOOSEPARAM_STRUCT(NiagaraParticleDataWidth), NiagaraRibbonVF.WidthDataOffset, InterpVtxData.ParticleId1), 
+			InterpVtxData.Alpha
+		);
 	}
 }
 
@@ -277,7 +317,12 @@ void GetNiagaraNormalizedAge(in FInterpVtxData InterpVtxData, in out FVertexFact
 	FLATTEN
 	if (NiagaraRibbonVF.NormalizedAgeDataOffset != -1)
 	{
-		Intermediates.RelativeTime = lerp(GetFloat(NiagaraRibbonVF.NormalizedAgeDataOffset, InterpVtxData.ParticleId0), GetFloat(NiagaraRibbonVF.NormalizedAgeDataOffset, InterpVtxData.ParticleId1), InterpVtxData.Alpha);
+		Intermediates.RelativeTime = lerp
+		(
+			GetFloat(NIAGARA_LOOSEPARAM_STRUCT(NiagaraParticleDataNormalizedAge), NiagaraRibbonVF.NormalizedAgeDataOffset, InterpVtxData.ParticleId0), 
+			GetFloat(NIAGARA_LOOSEPARAM_STRUCT(NiagaraParticleDataNormalizedAge), NiagaraRibbonVF.NormalizedAgeDataOffset, InterpVtxData.ParticleId1), 
+			InterpVtxData.Alpha
+		);
 	}
 }
 #endif
@@ -288,7 +333,12 @@ void GetNiagaraMaterialRandom(in FInterpVtxData InterpVtxData, in out FVertexFac
 	FLATTEN
 	if (NiagaraRibbonVF.MaterialRandomDataOffset != -1)
 	{
-		Intermediates.ParticleRandom = lerp(GetFloat(NiagaraRibbonVF.MaterialRandomDataOffset, InterpVtxData.ParticleId0), GetFloat(NiagaraRibbonVF.MaterialRandomDataOffset, InterpVtxData.ParticleId1), InterpVtxData.Alpha);
+		Intermediates.ParticleRandom = lerp
+		(
+			GetFloat(NIAGARA_LOOSEPARAM_STRUCT(NiagaraParticleDataMaterialRandom), NiagaraRibbonVF.MaterialRandomDataOffset, InterpVtxData.ParticleId0), 
+			GetFloat(NIAGARA_LOOSEPARAM_STRUCT(NiagaraParticleDataMaterialRandom), NiagaraRibbonVF.MaterialRandomDataOffset, InterpVtxData.ParticleId1), 
+			InterpVtxData.Alpha
+		);
 	}
 }
 #endif
@@ -578,8 +628,8 @@ void GetNiagaraParticleTangents(in FInterpVtxData InterpVtxData, in out FVertexF
 	BRANCH
 	if(NiagaraRibbonVF.FacingDataOffset != -1)
 	{
-		const float3 Facing0 = GetVec3(NiagaraRibbonVF.FacingDataOffset, InterpVtxData.ParticleId0);
-		const float3 Facing1 = GetVec3(NiagaraRibbonVF.FacingDataOffset, InterpVtxData.ParticleId1);
+		const float3 Facing0 = GetVec3(NIAGARA_LOOSEPARAM_STRUCT(NiagaraParticleDataFacing), NiagaraRibbonVF.FacingDataOffset, InterpVtxData.ParticleId0);
+		const float3 Facing1 = GetVec3(NIAGARA_LOOSEPARAM_STRUCT(NiagaraParticleDataFacing), NiagaraRibbonVF.FacingDataOffset, InterpVtxData.ParticleId1);
 		Intermediates.Facing = normalize(TransformVector(lerp(Facing0, Facing1, InterpVtxData.Alpha)));
 	}
 	else
@@ -605,8 +655,8 @@ void GetNiagaraParticleTangents(in FInterpVtxData InterpVtxData, in out FVertexF
 	BRANCH
 	if (NiagaraRibbonVF.TwistDataOffset != -1)
 	{
-		const float Twist0 = GetFloat(NiagaraRibbonVF.TwistDataOffset, InterpVtxData.ParticleId0);
-		const float Twist1 = GetFloat(NiagaraRibbonVF.TwistDataOffset, InterpVtxData.ParticleId1);
+		const float Twist0 = GetFloat(NIAGARA_LOOSEPARAM_STRUCT(NiagaraParticleDataTwist), NiagaraRibbonVF.TwistDataOffset, InterpVtxData.ParticleId0);
+		const float Twist1 = GetFloat(NIAGARA_LOOSEPARAM_STRUCT(NiagaraParticleDataTwist), NiagaraRibbonVF.TwistDataOffset, InterpVtxData.ParticleId1);
 		float SinTwist, CosTwist;
 		sincos(lerp(Twist0, Twist1, InterpVtxData.Alpha) /* + ResolvedView.RealTime * .5 */, SinTwist, CosTwist);
 
diff --git a/Engine/Plugins/FX/Niagara/Shaders/Private/NiagaraSpriteVertexFactory.ush b/Engine/Plugins/FX/Niagara/Shaders/Private/NiagaraSpriteVertexFactory.ush
--- a/Engine/Plugins/FX/Niagara/Shaders/Private/NiagaraSpriteVertexFactory.ush
+++ b/Engine/Plugins/FX/Niagara/Shaders/Private/NiagaraSpriteVertexFactory.ush
@@ -35,6 +35,26 @@
 /** Blends between FaceCamera and FaceCameraPosition.*/
 #define FACING_MODE_FACECAMERADISTANCEBLEND 4
 
+#if NIAGARA_LOOSEPARAM_BINDING
+	Buffer<float> NiagaraParticleDataPosition;
+	Buffer<float> NiagaraParticleDataVelocity;
+	Buffer<float> NiagaraParticleDataColor;
+	Buffer<float> NiagaraParticleDataRotation;
+	Buffer<float> NiagaraParticleDataSize;
+	Buffer<float> NiagaraParticleDataFacing;
+	Buffer<float> NiagaraParticleDataAlignment;
+	Buffer<float> NiagaraParticleDataSubImage;
+	Buffer<float> NiagaraParticleDataCameraOffset;
+	Buffer<float> NiagaraParticleDataUVScale;
+	Buffer<float> NiagaraParticleDataNormalizedAge;
+	Buffer<float> NiagaraParticleDataMaterialRandom;
+	Buffer<float> NiagaraParticleDataCustomSorting;
+	Buffer<float> NiagaraParticleDataMaterialParam0;
+	Buffer<float> NiagaraParticleDataMaterialParam1;
+	Buffer<float> NiagaraParticleDataMaterialParam2;
+	Buffer<float> NiagaraParticleDataMaterialParam3;
+#endif
+
 struct FVertexFactoryInput
 {
 	float2	TexCoord				: ATTRIBUTE0;
@@ -641,7 +661,7 @@ float3 GetNiagaraParticlePosition(uint InstanceID)
 		// See FNiagaraConstants for a synchronized value...
 		return TransformPosition(NiagaraSpriteVF.DefaultPos.xyz);
 	}
-	return TransformPosition(GetVec3(NiagaraSpriteVF.PositionDataOffset, InstanceID));
+	return TransformPosition(GetVec3(NIAGARA_LOOSEPARAM_STRUCT(NiagaraParticleDataPosition), NiagaraSpriteVF.PositionDataOffset, InstanceID));
 }
 float3 GetNiagaraParticlePrevPosition(uint InstanceID)
 {
@@ -650,7 +670,7 @@ float3 GetNiagaraParticlePrevPosition(uint InstanceID)
 		// See FNiagaraConstants for a synchronized value...
 		return TransformPrevPosition(NiagaraSpriteVF.DefaultPos.xyz);
 	}
-	return TransformPrevPosition(GetVec3(NiagaraSpriteVF.PositionDataOffset, InstanceID));
+	return TransformPrevPosition(GetVec3(NIAGARA_LOOSEPARAM_STRUCT(NiagaraParticleDataPosition), NiagaraSpriteVF.PositionDataOffset, InstanceID));
 }
 float3 GetNiagaraParticleVelocity(uint InstanceID)
 {
@@ -659,7 +679,7 @@ float3 GetNiagaraParticleVelocity(uint InstanceID)
 		// See FNiagaraConstants for a synchronized value...
 		return TransformVector(float3(0.0f,0.0f,0.0f));
 	}
-	return TransformVector(GetVec3(NiagaraSpriteVF.VelocityDataOffset, InstanceID));
+	return TransformVector(GetVec3(NIAGARA_LOOSEPARAM_STRUCT(NiagaraParticleDataVelocity), NiagaraSpriteVF.VelocityDataOffset, InstanceID));
 }
 float2 GetNiagaraParticleSize(uint InstanceID)
 {
@@ -668,7 +688,7 @@ float2 GetNiagaraParticleSize(uint InstanceID)
 		// See FNiagaraConstants for a synchronized value...
 		return float2(50.0f, 50.0f);
 	}
-	return GetVec2(NiagaraSpriteVF.SizeDataOffset, InstanceID);
+	return GetVec2(NIAGARA_LOOSEPARAM_STRUCT(NiagaraParticleDataSize), NiagaraSpriteVF.SizeDataOffset, InstanceID);
 }
 float GetNiagaraParticleRotation(uint InstanceID)
 {
@@ -677,7 +697,7 @@ float GetNiagaraParticleRotation(uint InstanceID)
 		// See FNiagaraConstants for a synchronized value...
 		return 0.0f;
 	}
-	return (GetFloat(NiagaraSpriteVF.RotationDataOffset, InstanceID) / 180.0f) * PI;
+	return (GetFloat(NIAGARA_LOOSEPARAM_STRUCT(NiagaraParticleDataRotation), NiagaraSpriteVF.RotationDataOffset, InstanceID) / 180.0f) * PI;
 }
 float4 GetNiagaraParticleColor(uint InstanceID)
 {
@@ -686,7 +706,7 @@ float4 GetNiagaraParticleColor(uint InstanceID)
 		// See FNiagaraConstants for a synchronized value...
 		return float4(1.0f, 1.0f, 1.0f, 1.0f);
 	}
-	return GetVec4(NiagaraSpriteVF.ColorDataOffset, InstanceID);
+	return GetVec4(NIAGARA_LOOSEPARAM_STRUCT(NiagaraParticleDataColor), NiagaraSpriteVF.ColorDataOffset, InstanceID);
 }
 
 float2 GetNiagaraUVScale(uint InstanceID)
@@ -696,7 +716,7 @@ float2 GetNiagaraUVScale(uint InstanceID)
 		// See FNiagaraConstants for a synchronized value...
 		return float2(1.0f, 1.0f);
 	}
-	return GetVec2(NiagaraSpriteVF.UVScaleDataOffset, InstanceID);	
+	return GetVec2(NIAGARA_LOOSEPARAM_STRUCT(NiagaraParticleDataUVScale), NiagaraSpriteVF.UVScaleDataOffset, InstanceID);	
 }
 
 float GetNiagaraParticleRandom(uint InstanceID)
@@ -706,7 +726,7 @@ float GetNiagaraParticleRandom(uint InstanceID)
 		// See FNiagaraConstants for a synchronized value...
 		return 0.0f;
 	}
-	return GetFloat(NiagaraSpriteVF.MaterialRandomDataOffset, InstanceID);	
+	return GetFloat(NIAGARA_LOOSEPARAM_STRUCT(NiagaraParticleDataMaterialRandom), NiagaraSpriteVF.MaterialRandomDataOffset, InstanceID);	
 }
 
 float GetNiagaraCameraOffset(uint InstanceID)
@@ -716,7 +736,7 @@ float GetNiagaraCameraOffset(uint InstanceID)
 		// See FNiagaraConstants for a synchronized value...
 		return 0.0f;
 	}
-	return GetFloat(NiagaraSpriteVF.CameraOffsetDataOffset, InstanceID);
+	return GetFloat(NIAGARA_LOOSEPARAM_STRUCT(NiagaraParticleDataCameraOffset), NiagaraSpriteVF.CameraOffsetDataOffset, InstanceID);
 }
 
 #if USE_PARTICLE_TIME
@@ -726,7 +746,7 @@ float GetNiagaraNormalizedAge(uint InstanceID)
 	{
 		return 0;
 	}
-	return GetFloat(NiagaraSpriteVF.NormalizedAgeDataOffset, InstanceID);
+	return GetFloat(NIAGARA_LOOSEPARAM_STRUCT(NiagaraParticleDataNormalizedAge), NiagaraSpriteVF.NormalizedAgeDataOffset, InstanceID);
 }
 #endif
 
@@ -735,7 +755,7 @@ float4 GetNiagaraParticleDynamicParameters(uint InstanceID)
 {
 	if (NiagaraSpriteVF.MaterialParamDataOffset != -1)
 	{
-		return GetVec4(NiagaraSpriteVF.MaterialParamDataOffset, InstanceID);
+		return GetVec4(NIAGARA_LOOSEPARAM_STRUCT(NiagaraParticleDataMaterialParam0), NiagaraSpriteVF.MaterialParamDataOffset, InstanceID);
 	}
 	else
 	{
@@ -748,7 +768,7 @@ float4 GetNiagaraParticleDynamicParameters1(uint InstanceID)
 {
 	if (NiagaraSpriteVF.MaterialParam1DataOffset != -1)
 	{
-		return GetVec4(NiagaraSpriteVF.MaterialParam1DataOffset, InstanceID);
+		return GetVec4(NIAGARA_LOOSEPARAM_STRUCT(NiagaraParticleDataMaterialParam1), NiagaraSpriteVF.MaterialParam1DataOffset, InstanceID);
 	}
 	else
 	{
@@ -761,7 +781,7 @@ float4 GetNiagaraParticleDynamicParameters2(uint InstanceID)
 {
 	if (NiagaraSpriteVF.MaterialParam2DataOffset != -1)
 	{
-		return GetVec4(NiagaraSpriteVF.MaterialParam2DataOffset, InstanceID);
+		return GetVec4(NIAGARA_LOOSEPARAM_STRUCT(NiagaraParticleDataMaterialParam2), NiagaraSpriteVF.MaterialParam2DataOffset, InstanceID);
 	}
 	else
 	{
@@ -774,7 +794,7 @@ float4 GetNiagaraParticleDynamicParameters3(uint InstanceID)
 {
 	if (NiagaraSpriteVF.MaterialParam3DataOffset != -1)
 	{
-		return GetVec4(NiagaraSpriteVF.MaterialParam3DataOffset, InstanceID);
+		return GetVec4(NIAGARA_LOOSEPARAM_STRUCT(NiagaraParticleDataMaterialParam3), NiagaraSpriteVF.MaterialParam3DataOffset, InstanceID);
 	}
 	else
 	{
@@ -790,7 +810,7 @@ float GetNiagaraParticleSubimage(uint InstanceID)
 	float Val = 0.0f;
     if(NiagaraSpriteVF.SubimageDataOffset != -1)
 	{
-		Val = GetFloat(NiagaraSpriteVF.SubimageDataOffset, InstanceID);
+		Val = GetFloat(NIAGARA_LOOSEPARAM_STRUCT(NiagaraParticleDataSubImage), NiagaraSpriteVF.SubimageDataOffset, InstanceID);
 	}
 	return Val;
 }
@@ -801,7 +821,7 @@ float3 GetNiagaraParticleFacingVector(uint InstanceID)
 	float3 Val = float3(1.0f, 0.0f, 0.0f);
     if(NiagaraSpriteVF.FacingDataOffset != -1)
 	{
-		Val = GetVec3(NiagaraSpriteVF.FacingDataOffset, InstanceID);
+		Val = GetVec3(NIAGARA_LOOSEPARAM_STRUCT(NiagaraParticleDataFacing), NiagaraSpriteVF.FacingDataOffset, InstanceID);
 	}
 	return TransformVector(Val);
 }
@@ -811,7 +831,7 @@ float3 GetNiagaraParticleAlignmentVector(uint InstanceID)
 	float3 Val = float3(1.0f, 0.0f, 0.0f);
     if(NiagaraSpriteVF.AlignmentDataOffset != -1)
 	{
-		Val = GetVec3(NiagaraSpriteVF.AlignmentDataOffset, InstanceID);
+		Val = GetVec3(NIAGARA_LOOSEPARAM_STRUCT(NiagaraParticleDataAlignment), NiagaraSpriteVF.AlignmentDataOffset, InstanceID);
 	}
 	return TransformVector(Val);
 }
diff --git a/Engine/Plugins/FX/Niagara/Shaders/Private/NiagaraVFParticleAccess.usf b/Engine/Plugins/FX/Niagara/Shaders/Private/NiagaraVFParticleAccess.usf
--- a/Engine/Plugins/FX/Niagara/Shaders/Private/NiagaraVFParticleAccess.usf
+++ b/Engine/Plugins/FX/Niagara/Shaders/Private/NiagaraVFParticleAccess.usf
@@ -12,38 +12,41 @@
 #define NIAGARA_KEYGEN 0
 #endif
 
-#if (!RAYHITGROUPSHADER && !COMPUTESHADER) || NIAGARA_KEYGEN
+#define NIAGARA_LOOSEPARAM_BINDING (!RAYHITGROUPSHADER && !COMPUTESHADER) || NIAGARA_KEYGEN
 
-Buffer<float> NiagaraParticleDataFloat;
+#if NIAGARA_LOOSEPARAM_BINDING
+
+#define NIAGARA_LOOSEPARAM_STRUCT(x) x
 uint NiagaraFloatDataOffset;
 uint NiagaraFloatDataStride;
 
-float GetFloat(int RegisterIdx, uint InstanceID)
+float GetFloat(Buffer<float> NiagaraBuffer, int RegisterIdx, uint InstanceID)
 {
-	return NiagaraParticleDataFloat[NiagaraFloatDataOffset + (RegisterIdx * NiagaraFloatDataStride + InstanceID)];
+	return NiagaraBuffer[NiagaraFloatDataOffset + (RegisterIdx * NiagaraFloatDataStride + InstanceID)];
 }
 
 #else
+#define NIAGARA_LOOSEPARAM_STRUCT(x) NiagaraVFLooseParameters.x
 
-float GetFloat(int RegisterIdx, uint InstanceID)
+float GetFloat(Buffer<float> NiagaraBuffer, int RegisterIdx, uint InstanceID)
 {
-	return NiagaraVFLooseParameters.NiagaraParticleDataFloat[NiagaraVFLooseParameters.NiagaraFloatDataOffset + (RegisterIdx * NiagaraVFLooseParameters.NiagaraFloatDataStride + InstanceID)];
+	return NiagaraBuffer[NiagaraVFLooseParameters.NiagaraFloatDataOffset + (RegisterIdx * NiagaraVFLooseParameters.NiagaraFloatDataStride + InstanceID)];
 }
 
 #endif
 
-float2 GetVec2(int RegisterIndex, uint InstanceID)
+float2 GetVec2(Buffer<float> NiagaraBuffer, int RegisterIndex, uint InstanceID)
 {
-	return float2(GetFloat(RegisterIndex, InstanceID), GetFloat(RegisterIndex+1, InstanceID));
+	return float2(GetFloat(NiagaraBuffer, RegisterIndex, InstanceID), GetFloat(NiagaraBuffer, RegisterIndex+1, InstanceID));
 }
 
-float3 GetVec3(int RegisterIndex, uint InstanceID)
+float3 GetVec3(Buffer<float> NiagaraBuffer, int RegisterIndex, uint InstanceID)
 {
-	return float3(GetFloat(RegisterIndex, InstanceID), GetFloat(RegisterIndex+1, InstanceID), GetFloat(RegisterIndex+2, InstanceID));
+	return float3(GetFloat(NiagaraBuffer, RegisterIndex, InstanceID), GetFloat(NiagaraBuffer, RegisterIndex+1, InstanceID), GetFloat(NiagaraBuffer, RegisterIndex+2, InstanceID));
 }
 
-float4 GetVec4(int RegisterIndex, uint InstanceID)
+float4 GetVec4(Buffer<float> NiagaraBuffer, int RegisterIndex, uint InstanceID)
 {
-	return float4(GetFloat(RegisterIndex, InstanceID), GetFloat(RegisterIndex+1, InstanceID), GetFloat(RegisterIndex+2, InstanceID), GetFloat(RegisterIndex+3, InstanceID));
+	return float4(GetFloat(NiagaraBuffer, RegisterIndex, InstanceID), GetFloat(NiagaraBuffer, RegisterIndex+1, InstanceID), GetFloat(NiagaraBuffer, RegisterIndex+2, InstanceID), GetFloat(NiagaraBuffer, RegisterIndex+3, InstanceID));
 }
 
diff --git a/Engine/Plugins/FX/Niagara/Source/Niagara/Private/NiagaraRendererRibbons.cpp b/Engine/Plugins/FX/Niagara/Source/Niagara/Private/NiagaraRendererRibbons.cpp
--- a/Engine/Plugins/FX/Niagara/Source/Niagara/Private/NiagaraRendererRibbons.cpp
+++ b/Engine/Plugins/FX/Niagara/Source/Niagara/Private/NiagaraRendererRibbons.cpp
@@ -796,7 +796,18 @@ void FNiagaraRendererRibbons::SetupMeshBatchAndCollectorResourceForView(
 	VFLooseParams.TangentsAndDistances = TangentsAndDistancesBuffer.SRV;
 	VFLooseParams.MultiRibbonIndices = MultiRibbonIndicesBuffer.SRV;
 	VFLooseParams.PackedPerRibbonDataByIndex = PackedPerRibbonDataByIndexBuffer.SRV;
-	VFLooseParams.NiagaraParticleDataFloat = CollectorResources.VertexFactory.GetParticleDataFloatSRV();
+	VFLooseParams.NiagaraParticleDataPosition = CollectorResources.VertexFactory.GetParticleDataFloatSRV();
+	VFLooseParams.NiagaraParticleDataVelocity = CollectorResources.VertexFactory.GetParticleDataFloatSRV();
+	VFLooseParams.NiagaraParticleDataColor = CollectorResources.VertexFactory.GetParticleDataFloatSRV();
+	VFLooseParams.NiagaraParticleDataWidth = CollectorResources.VertexFactory.GetParticleDataFloatSRV();
+	VFLooseParams.NiagaraParticleDataTwist = CollectorResources.VertexFactory.GetParticleDataFloatSRV();
+	VFLooseParams.NiagaraParticleDataFacing = CollectorResources.VertexFactory.GetParticleDataFloatSRV();
+	VFLooseParams.NiagaraParticleDataNormalizedAge = CollectorResources.VertexFactory.GetParticleDataFloatSRV();
+	VFLooseParams.NiagaraParticleDataMaterialRandom = CollectorResources.VertexFactory.GetParticleDataFloatSRV();
+	VFLooseParams.NiagaraParticleDataMaterialParam0 = CollectorResources.VertexFactory.GetParticleDataFloatSRV();
+	VFLooseParams.NiagaraParticleDataMaterialParam1 = CollectorResources.VertexFactory.GetParticleDataFloatSRV();
+	VFLooseParams.NiagaraParticleDataMaterialParam2 = CollectorResources.VertexFactory.GetParticleDataFloatSRV();
+	VFLooseParams.NiagaraParticleDataMaterialParam3 = CollectorResources.VertexFactory.GetParticleDataFloatSRV();
 	VFLooseParams.NiagaraFloatDataOffset = CollectorResources.VertexFactory.GetFloatDataOffset();
 	VFLooseParams.NiagaraFloatDataStride = CollectorResources.VertexFactory.GetFloatDataStride();
 	VFLooseParams.SortedIndicesOffset = CollectorResources.VertexFactory.GetSortedIndicesOffset();
diff --git a/Engine/Plugins/FX/Niagara/Source/Niagara/Private/NiagaraRendererSprites.cpp b/Engine/Plugins/FX/Niagara/Source/Niagara/Private/NiagaraRendererSprites.cpp
--- a/Engine/Plugins/FX/Niagara/Source/Niagara/Private/NiagaraRendererSprites.cpp
+++ b/Engine/Plugins/FX/Niagara/Source/Niagara/Private/NiagaraRendererSprites.cpp
@@ -463,12 +463,30 @@ void FNiagaraRendererSprites::CreateMeshBatchForView(
 	FNiagaraSpriteVFLooseParameters VFLooseParams;
 	VFLooseParams.NumCutoutVerticesPerFrame = CollectorResources.VertexFactory.GetNumCutoutVerticesPerFrame();
 	VFLooseParams.CutoutGeometry = CollectorResources.VertexFactory.GetCutoutGeometrySRV() ? CollectorResources.VertexFactory.GetCutoutGeometrySRV() : GFNiagaraNullCutoutVertexBuffer.VertexBufferSRV.GetReference();
-	VFLooseParams.NiagaraParticleDataFloat = CollectorResources.VertexFactory.GetParticleDataFloatSRV();
+	
+	VFLooseParams.NiagaraParticleDataPosition = CollectorResources.VertexFactory.GetParticleDataFloatSRV();
+	VFLooseParams.NiagaraParticleDataVelocity = CollectorResources.VertexFactory.GetParticleDataFloatSRV();
+	VFLooseParams.NiagaraParticleDataColor = CollectorResources.VertexFactory.GetParticleDataFloatSRV();
+	VFLooseParams.NiagaraParticleDataRotation = CollectorResources.VertexFactory.GetParticleDataFloatSRV();
+	VFLooseParams.NiagaraParticleDataSize = CollectorResources.VertexFactory.GetParticleDataFloatSRV();
+	VFLooseParams.NiagaraParticleDataFacing = CollectorResources.VertexFactory.GetParticleDataFloatSRV();
+	VFLooseParams.NiagaraParticleDataAlignment = CollectorResources.VertexFactory.GetParticleDataFloatSRV();
+	VFLooseParams.NiagaraParticleDataSubImage = CollectorResources.VertexFactory.GetParticleDataFloatSRV();
+	VFLooseParams.NiagaraParticleDataCameraOffset = CollectorResources.VertexFactory.GetParticleDataFloatSRV();
+	VFLooseParams.NiagaraParticleDataUVScale = CollectorResources.VertexFactory.GetParticleDataFloatSRV();
+	VFLooseParams.NiagaraParticleDataNormalizedAge = CollectorResources.VertexFactory.GetParticleDataFloatSRV();
+	VFLooseParams.NiagaraParticleDataMaterialRandom = CollectorResources.VertexFactory.GetParticleDataFloatSRV();
+	VFLooseParams.NiagaraParticleDataCustomSorting = CollectorResources.VertexFactory.GetParticleDataFloatSRV();
+	VFLooseParams.NiagaraParticleDataMaterialParam0 = CollectorResources.VertexFactory.GetParticleDataFloatSRV();
+	VFLooseParams.NiagaraParticleDataMaterialParam1 = CollectorResources.VertexFactory.GetParticleDataFloatSRV();
+	VFLooseParams.NiagaraParticleDataMaterialParam2 = CollectorResources.VertexFactory.GetParticleDataFloatSRV();
+	VFLooseParams.NiagaraParticleDataMaterialParam3 = CollectorResources.VertexFactory.GetParticleDataFloatSRV();
+
 	VFLooseParams.NiagaraFloatDataOffset = CollectorResources.VertexFactory.GetFloatDataOffset();
 	VFLooseParams.NiagaraFloatDataStride = CollectorResources.VertexFactory.GetFloatDataStride();
 	VFLooseParams.ParticleAlignmentMode = CollectorResources.VertexFactory.GetAlignmentMode();
 	VFLooseParams.ParticleFacingMode = CollectorResources.VertexFactory.GetFacingMode();
-	VFLooseParams.SortedIndices = CollectorResources.VertexFactory.GetSortedIndicesSRV() ? CollectorResources.VertexFactory.GetSortedIndicesSRV().GetReference() : GFNiagaraNullSortedIndicesVertexBuffer.VertexBufferSRV.GetReference();
+	VFLooseParams.SortedIndices = CollectorResources.VertexFactory.GetSortedIndicesSRV() ? CollectorResources.VertexFactory.GetSortedIndicesSRV() : GFNiagaraNullSortedIndicesVertexBuffer.VertexBufferSRV.GetReference();
 	VFLooseParams.SortedIndicesOffset = CollectorResources.VertexFactory.GetSortedIndicesOffset();
 	if (IndirectArgsOffset != INDEX_NONE)
 	{
diff --git a/Engine/Plugins/FX/Niagara/Source/NiagaraVertexFactories/Private/NiagaraMeshVertexFactory.cpp b/Engine/Plugins/FX/Niagara/Source/NiagaraVertexFactories/Private/NiagaraMeshVertexFactory.cpp
--- a/Engine/Plugins/FX/Niagara/Source/NiagaraVertexFactories/Private/NiagaraMeshVertexFactory.cpp
+++ b/Engine/Plugins/FX/Niagara/Source/NiagaraVertexFactories/Private/NiagaraMeshVertexFactory.cpp
@@ -18,7 +18,18 @@ public:
 	virtual void Bind(const FShaderParameterMap& ParameterMap) override
 	{
 		//PrevTransformBuffer.Bind(ParameterMap, TEXT("PrevTransformBuffer"));
-		NiagaraParticleDataFloat.Bind(ParameterMap, TEXT("NiagaraParticleDataFloat"));
+		NiagaraParticleDataPosition.Bind(ParameterMap, TEXT("NiagaraParticleDataPosition"));
+		NiagaraParticleDataVelocity.Bind(ParameterMap, TEXT("NiagaraParticleDataVelocity"));
+		NiagaraParticleDataColor.Bind(ParameterMap, TEXT("NiagaraParticleDataColor"));
+		NiagaraParticleDataScale.Bind(ParameterMap, TEXT("NiagaraParticleDataScale"));
+		NiagaraParticleDataTransform.Bind(ParameterMap, TEXT("NiagaraParticleDataTransform"));
+		NiagaraParticleDataNormalizedAge.Bind(ParameterMap, TEXT("NiagaraParticleDataNormalizedAge"));
+		NiagaraParticleDataMaterialRandom.Bind(ParameterMap, TEXT("NiagaraParticleDataMaterialRandom"));
+		NiagaraParticleDataMaterialParam0.Bind(ParameterMap, TEXT("NiagaraParticleDataMaterialParam0"));
+		NiagaraParticleDataMaterialParam1.Bind(ParameterMap, TEXT("NiagaraParticleDataMaterialParam1"));
+		NiagaraParticleDataMaterialParam2.Bind(ParameterMap, TEXT("NiagaraParticleDataMaterialParam2"));
+		NiagaraParticleDataMaterialParam3.Bind(ParameterMap, TEXT("NiagaraParticleDataMaterialParam3"));
+
 		FloatDataOffset.Bind(ParameterMap, TEXT("NiagaraFloatDataOffset"));
 		FloatDataStride.Bind(ParameterMap, TEXT("NiagaraFloatDataStride"));
 
@@ -35,7 +46,17 @@ public:
 	virtual void Serialize(FArchive& Ar) override
 	{
 		//Ar << PrevTransformBuffer;
-		Ar << NiagaraParticleDataFloat;
+		Ar << NiagaraParticleDataPosition;
+		Ar << NiagaraParticleDataVelocity;
+		Ar << NiagaraParticleDataColor;
+		Ar << NiagaraParticleDataScale;
+		Ar << NiagaraParticleDataTransform;
+		Ar << NiagaraParticleDataNormalizedAge;
+		Ar << NiagaraParticleDataMaterialRandom;
+		Ar << NiagaraParticleDataMaterialParam0;
+		Ar << NiagaraParticleDataMaterialParam1;
+		Ar << NiagaraParticleDataMaterialParam2;
+		Ar << NiagaraParticleDataMaterialParam3;
 		Ar << FloatDataOffset;
 		Ar << FloatDataStride;
 
@@ -66,7 +87,18 @@ public:
 
 		ShaderBindings.Add(MeshFacingMode, NiagaraMeshVF->GetMeshFacingMode());
 
-		ShaderBindings.Add(NiagaraParticleDataFloat, NiagaraMeshVF->GetParticleDataFloatSRV());
+		ShaderBindings.Add(NiagaraParticleDataPosition, NiagaraMeshVF->GetParticleDataFloatSRV());
+		ShaderBindings.Add(NiagaraParticleDataVelocity, NiagaraMeshVF->GetParticleDataFloatSRV());
+		ShaderBindings.Add(NiagaraParticleDataColor, NiagaraMeshVF->GetParticleDataFloatSRV());
+		ShaderBindings.Add(NiagaraParticleDataScale, NiagaraMeshVF->GetParticleDataFloatSRV());
+		ShaderBindings.Add(NiagaraParticleDataTransform, NiagaraMeshVF->GetParticleDataFloatSRV());
+		ShaderBindings.Add(NiagaraParticleDataNormalizedAge, NiagaraMeshVF->GetParticleDataFloatSRV());
+		ShaderBindings.Add(NiagaraParticleDataMaterialRandom, NiagaraMeshVF->GetParticleDataFloatSRV());
+		ShaderBindings.Add(NiagaraParticleDataMaterialParam0, NiagaraMeshVF->GetParticleDataFloatSRV());
+		ShaderBindings.Add(NiagaraParticleDataMaterialParam1, NiagaraMeshVF->GetParticleDataFloatSRV());
+		ShaderBindings.Add(NiagaraParticleDataMaterialParam2, NiagaraMeshVF->GetParticleDataFloatSRV());
+		ShaderBindings.Add(NiagaraParticleDataMaterialParam3, NiagaraMeshVF->GetParticleDataFloatSRV());
+
 		ShaderBindings.Add(FloatDataOffset, NiagaraMeshVF->GetFloatDataOffset());
 		ShaderBindings.Add(FloatDataStride, NiagaraMeshVF->GetFloatDataStride());
 
@@ -79,7 +111,17 @@ private:
 
 	//FShaderResourceParameter PrevTransformBuffer;
 
-	FShaderResourceParameter NiagaraParticleDataFloat;
+	FShaderResourceParameter NiagaraParticleDataPosition;
+	FShaderResourceParameter NiagaraParticleDataVelocity;
+	FShaderResourceParameter NiagaraParticleDataColor;
+	FShaderResourceParameter NiagaraParticleDataScale;
+	FShaderResourceParameter NiagaraParticleDataTransform;
+	FShaderResourceParameter NiagaraParticleDataNormalizedAge;
+	FShaderResourceParameter NiagaraParticleDataMaterialRandom;
+	FShaderResourceParameter NiagaraParticleDataMaterialParam0;
+	FShaderResourceParameter NiagaraParticleDataMaterialParam1;
+	FShaderResourceParameter NiagaraParticleDataMaterialParam2;
+	FShaderResourceParameter NiagaraParticleDataMaterialParam3;
 	FShaderParameter FloatDataOffset;
 	FShaderParameter FloatDataStride;
 
diff --git a/Engine/Plugins/FX/Niagara/Source/NiagaraVertexFactories/Private/NiagaraRibbonVertexFactory.cpp b/Engine/Plugins/FX/Niagara/Source/NiagaraVertexFactories/Private/NiagaraRibbonVertexFactory.cpp
--- a/Engine/Plugins/FX/Niagara/Source/NiagaraVertexFactories/Private/NiagaraRibbonVertexFactory.cpp
+++ b/Engine/Plugins/FX/Niagara/Source/NiagaraVertexFactories/Private/NiagaraRibbonVertexFactory.cpp
@@ -34,14 +34,37 @@ class FNiagaraRibbonVertexFactoryShaderParametersVS : public FNiagaraRibbonVerte
 public:
 	virtual void Bind(const FShaderParameterMap& ParameterMap) override
 	{
-		NiagaraParticleDataFloat.Bind(ParameterMap, TEXT("NiagaraParticleDataFloat"));
+		NiagaraParticleDataPosition.Bind(ParameterMap, TEXT("NiagaraParticleDataPosition"));
+		NiagaraParticleDataVelocity.Bind(ParameterMap, TEXT("NiagaraParticleDataVelocity"));
+		NiagaraParticleDataColor.Bind(ParameterMap, TEXT("NiagaraParticleDataColor"));
+		NiagaraParticleDataWidth.Bind(ParameterMap, TEXT("NiagaraParticleDataWidth"));
+		NiagaraParticleDataTwist.Bind(ParameterMap, TEXT("NiagaraParticleDataTwist"));
+		NiagaraParticleDataFacing.Bind(ParameterMap, TEXT("NiagaraParticleDataFacing"));
+		NiagaraParticleDataNormalizedAge.Bind(ParameterMap, TEXT("NiagaraParticleDataNormalizedAge"));
+		NiagaraParticleDataMaterialRandom.Bind(ParameterMap, TEXT("NiagaraParticleDataMaterialRandom"));
+		NiagaraParticleDataMaterialParam0.Bind(ParameterMap, TEXT("NiagaraParticleDataMaterialParam0"));
+		NiagaraParticleDataMaterialParam1.Bind(ParameterMap, TEXT("NiagaraParticleDataMaterialParam1"));
+		NiagaraParticleDataMaterialParam2.Bind(ParameterMap, TEXT("NiagaraParticleDataMaterialParam2"));
+		NiagaraParticleDataMaterialParam3.Bind(ParameterMap, TEXT("NiagaraParticleDataMaterialParam3"));
+
 		FloatDataOffset.Bind(ParameterMap, TEXT("NiagaraFloatDataOffset"));
 		FloatDataStride.Bind(ParameterMap, TEXT("NiagaraFloatDataStride"));
 	}
 
 	virtual void Serialize(FArchive& Ar) override
 	{
-		Ar << NiagaraParticleDataFloat;
+		Ar << NiagaraParticleDataPosition;
+		Ar << NiagaraParticleDataVelocity;
+		Ar << NiagaraParticleDataColor;
+		Ar << NiagaraParticleDataWidth;
+		Ar << NiagaraParticleDataTwist;
+		Ar << NiagaraParticleDataFacing;
+		Ar << NiagaraParticleDataNormalizedAge;
+		Ar << NiagaraParticleDataMaterialRandom;
+		Ar << NiagaraParticleDataMaterialParam0;
+		Ar << NiagaraParticleDataMaterialParam1;
+		Ar << NiagaraParticleDataMaterialParam2;
+		Ar << NiagaraParticleDataMaterialParam3;
 		Ar << FloatDataOffset;
 		Ar << FloatDataStride;
 	}
@@ -60,13 +83,36 @@ public:
 		FNiagaraRibbonVertexFactory* RibbonVF = (FNiagaraRibbonVertexFactory*)VertexFactory;
 		ShaderBindings.Add(Shader->GetUniformBufferParameter<FNiagaraRibbonUniformParameters>(), RibbonVF->GetRibbonUniformBuffer());
 		ShaderBindings.Add(Shader->GetUniformBufferParameter<FNiagaraRibbonVFLooseParameters>(), RibbonVF->LooseParameterUniformBuffer);
-		ShaderBindings.Add(NiagaraParticleDataFloat, RibbonVF->GetParticleDataFloatSRV());
+		ShaderBindings.Add(NiagaraParticleDataPosition, RibbonVF->GetParticleDataFloatSRV());
+		ShaderBindings.Add(NiagaraParticleDataVelocity, RibbonVF->GetParticleDataFloatSRV());
+		ShaderBindings.Add(NiagaraParticleDataColor, RibbonVF->GetParticleDataFloatSRV());
+		ShaderBindings.Add(NiagaraParticleDataWidth, RibbonVF->GetParticleDataFloatSRV());
+		ShaderBindings.Add(NiagaraParticleDataTwist, RibbonVF->GetParticleDataFloatSRV());
+		ShaderBindings.Add(NiagaraParticleDataFacing, RibbonVF->GetParticleDataFloatSRV());
+		ShaderBindings.Add(NiagaraParticleDataNormalizedAge, RibbonVF->GetParticleDataFloatSRV());
+		ShaderBindings.Add(NiagaraParticleDataMaterialRandom, RibbonVF->GetParticleDataFloatSRV());
+		ShaderBindings.Add(NiagaraParticleDataMaterialParam0, RibbonVF->GetParticleDataFloatSRV());
+		ShaderBindings.Add(NiagaraParticleDataMaterialParam1, RibbonVF->GetParticleDataFloatSRV());
+		ShaderBindings.Add(NiagaraParticleDataMaterialParam2, RibbonVF->GetParticleDataFloatSRV());
+		ShaderBindings.Add(NiagaraParticleDataMaterialParam3, RibbonVF->GetParticleDataFloatSRV());
 		ShaderBindings.Add(FloatDataOffset, RibbonVF->GetFloatDataOffset());
 		ShaderBindings.Add(FloatDataStride, RibbonVF->GetFloatDataStride());
 	}
 
 private:
-	FShaderResourceParameter NiagaraParticleDataFloat;
+	FShaderResourceParameter NiagaraParticleDataPosition;
+	FShaderResourceParameter NiagaraParticleDataVelocity;
+	FShaderResourceParameter NiagaraParticleDataColor;
+	FShaderResourceParameter NiagaraParticleDataWidth;
+	FShaderResourceParameter NiagaraParticleDataTwist;
+	FShaderResourceParameter NiagaraParticleDataFacing;
+	FShaderResourceParameter NiagaraParticleDataNormalizedAge;
+	FShaderResourceParameter NiagaraParticleDataMaterialRandom;
+	FShaderResourceParameter NiagaraParticleDataMaterialParam0;
+	FShaderResourceParameter NiagaraParticleDataMaterialParam1;
+	FShaderResourceParameter NiagaraParticleDataMaterialParam2;
+	FShaderResourceParameter NiagaraParticleDataMaterialParam3;
+
 	FShaderParameter FloatDataOffset;
 	FShaderParameter FloatDataStride;
 };
diff --git a/Engine/Plugins/FX/Niagara/Source/NiagaraVertexFactories/Private/NiagaraSpriteVertexFactory.cpp b/Engine/Plugins/FX/Niagara/Source/NiagaraVertexFactories/Private/NiagaraSpriteVertexFactory.cpp
--- a/Engine/Plugins/FX/Niagara/Source/NiagaraVertexFactories/Private/NiagaraSpriteVertexFactory.cpp
+++ b/Engine/Plugins/FX/Niagara/Source/NiagaraVertexFactories/Private/NiagaraSpriteVertexFactory.cpp
@@ -40,7 +40,24 @@ public:
 		NumCutoutVerticesPerFrame.Bind(ParameterMap, TEXT("NumCutoutVerticesPerFrame"));
 		CutoutGeometry.Bind(ParameterMap, TEXT("CutoutGeometry"));
 
-		NiagaraParticleDataFloat.Bind(ParameterMap, TEXT("NiagaraParticleDataFloat"));
+		NiagaraParticleDataPosition.Bind(ParameterMap, TEXT("NiagaraParticleDataPosition"));
+		NiagaraParticleDataVelocity.Bind(ParameterMap, TEXT("NiagaraParticleDataVelocity"));
+		NiagaraParticleDataColor.Bind(ParameterMap, TEXT("NiagaraParticleDataColor"));
+		NiagaraParticleDataRotation.Bind(ParameterMap, TEXT("NiagaraParticleDataRotation"));
+		NiagaraParticleDataSize.Bind(ParameterMap, TEXT("NiagaraParticleDataSize"));
+		NiagaraParticleDataFacing.Bind(ParameterMap, TEXT("NiagaraParticleDataFacing"));
+		NiagaraParticleDataAlignment.Bind(ParameterMap, TEXT("NiagaraParticleDataAlignment"));
+		NiagaraParticleDataSubImage.Bind(ParameterMap, TEXT("NiagaraParticleDataSubImage"));
+		NiagaraParticleDataCameraOffset.Bind(ParameterMap, TEXT("NiagaraParticleDataCameraOffset"));
+		NiagaraParticleDataUVScale.Bind(ParameterMap, TEXT("NiagaraParticleDataUVScale"));
+		NiagaraParticleDataNormalizedAge.Bind(ParameterMap, TEXT("NiagaraParticleDataNormalizedAge"));
+		NiagaraParticleDataMaterialRandom.Bind(ParameterMap, TEXT("NiagaraParticleDataMaterialRandom"));
+		NiagaraParticleDataCustomSorting.Bind(ParameterMap, TEXT("NiagaraParticleDataCustomSorting"));
+		NiagaraParticleDataMaterialParam0.Bind(ParameterMap, TEXT("NiagaraParticleDataMaterialParam0"));
+		NiagaraParticleDataMaterialParam1.Bind(ParameterMap, TEXT("NiagaraParticleDataMaterialParam1"));
+		NiagaraParticleDataMaterialParam2.Bind(ParameterMap, TEXT("NiagaraParticleDataMaterialParam2"));
+		NiagaraParticleDataMaterialParam3.Bind(ParameterMap, TEXT("NiagaraParticleDataMaterialParam3"));
+
 		FloatDataOffset.Bind(ParameterMap, TEXT("NiagaraFloatDataOffset"));
 		FloatDataStride.Bind(ParameterMap, TEXT("NiagaraFloatDataStride"));
 
@@ -62,7 +79,24 @@ public:
 		Ar << ParticleFacingMode;
 		Ar << ParticleAlignmentMode;
 
-		Ar << NiagaraParticleDataFloat;
+		Ar << NiagaraParticleDataPosition;
+		Ar << NiagaraParticleDataVelocity;
+		Ar << NiagaraParticleDataColor;
+		Ar << NiagaraParticleDataRotation;
+		Ar << NiagaraParticleDataSize;
+		Ar << NiagaraParticleDataFacing;
+		Ar << NiagaraParticleDataAlignment;
+		Ar << NiagaraParticleDataSubImage;
+		Ar << NiagaraParticleDataCameraOffset;
+		Ar << NiagaraParticleDataUVScale;
+		Ar << NiagaraParticleDataNormalizedAge;
+		Ar << NiagaraParticleDataMaterialRandom;
+		Ar << NiagaraParticleDataCustomSorting;
+		Ar << NiagaraParticleDataMaterialParam0;
+		Ar << NiagaraParticleDataMaterialParam1;
+		Ar << NiagaraParticleDataMaterialParam2;
+		Ar << NiagaraParticleDataMaterialParam3;
+
 		Ar << FloatDataOffset;
 		Ar << FloatDataStride;
 
@@ -97,11 +131,28 @@ public:
 		ShaderBindings.Add(ParticleAlignmentMode, SpriteVF->GetAlignmentMode());
 		ShaderBindings.Add(ParticleFacingMode, SpriteVF->GetFacingMode());
 
-		ShaderBindings.Add(NiagaraParticleDataFloat, SpriteVF->GetParticleDataFloatSRV());
+		ShaderBindings.Add(NiagaraParticleDataPosition, SpriteVF->GetParticleDataFloatSRV());
+		ShaderBindings.Add(NiagaraParticleDataVelocity, SpriteVF->GetParticleDataFloatSRV());
+		ShaderBindings.Add(NiagaraParticleDataColor, SpriteVF->GetParticleDataFloatSRV());
+		ShaderBindings.Add(NiagaraParticleDataRotation, SpriteVF->GetParticleDataFloatSRV());
+		ShaderBindings.Add(NiagaraParticleDataSize, SpriteVF->GetParticleDataFloatSRV());
+		ShaderBindings.Add(NiagaraParticleDataFacing, SpriteVF->GetParticleDataFloatSRV());
+		ShaderBindings.Add(NiagaraParticleDataAlignment, SpriteVF->GetParticleDataFloatSRV());
+		ShaderBindings.Add(NiagaraParticleDataSubImage, SpriteVF->GetParticleDataFloatSRV());
+		ShaderBindings.Add(NiagaraParticleDataCameraOffset, SpriteVF->GetParticleDataFloatSRV());
+		ShaderBindings.Add(NiagaraParticleDataUVScale, SpriteVF->GetParticleDataFloatSRV());
+		ShaderBindings.Add(NiagaraParticleDataNormalizedAge, SpriteVF->GetParticleDataFloatSRV());
+		ShaderBindings.Add(NiagaraParticleDataMaterialRandom, SpriteVF->GetParticleDataFloatSRV());
+		ShaderBindings.Add(NiagaraParticleDataCustomSorting, SpriteVF->GetParticleDataFloatSRV());
+		ShaderBindings.Add(NiagaraParticleDataMaterialParam0, SpriteVF->GetParticleDataFloatSRV());
+		ShaderBindings.Add(NiagaraParticleDataMaterialParam1, SpriteVF->GetParticleDataFloatSRV());
+		ShaderBindings.Add(NiagaraParticleDataMaterialParam2, SpriteVF->GetParticleDataFloatSRV());
+		ShaderBindings.Add(NiagaraParticleDataMaterialParam3, SpriteVF->GetParticleDataFloatSRV());
+
 		ShaderBindings.Add(FloatDataOffset, SpriteVF->GetFloatDataOffset());
 		ShaderBindings.Add(FloatDataStride, SpriteVF->GetFloatDataStride());
 
-		ShaderBindings.Add(SortedIndices, SpriteVF->GetSortedIndicesSRV() ? SpriteVF->GetSortedIndicesSRV() : GFNiagaraNullSortedIndicesVertexBuffer.VertexBufferSRV);
+		ShaderBindings.Add(SortedIndices, SpriteVF->GetSortedIndicesSRV() ? SpriteVF->GetSortedIndicesSRV() : GFNiagaraNullSortedIndicesVertexBuffer.VertexBufferSRV.GetReference());
 		ShaderBindings.Add(SortedIndicesOffset, SpriteVF->GetSortedIndicesOffset());
 	}
 private:
@@ -112,7 +163,24 @@ private:
 
 	FShaderResourceParameter CutoutGeometry;
 
-	FShaderResourceParameter NiagaraParticleDataFloat;
+	FShaderResourceParameter NiagaraParticleDataPosition;
+	FShaderResourceParameter NiagaraParticleDataVelocity;
+	FShaderResourceParameter NiagaraParticleDataColor;
+	FShaderResourceParameter NiagaraParticleDataRotation;
+	FShaderResourceParameter NiagaraParticleDataSize;
+	FShaderResourceParameter NiagaraParticleDataFacing;
+	FShaderResourceParameter NiagaraParticleDataAlignment;
+	FShaderResourceParameter NiagaraParticleDataSubImage;
+	FShaderResourceParameter NiagaraParticleDataCameraOffset;
+	FShaderResourceParameter NiagaraParticleDataUVScale;
+	FShaderResourceParameter NiagaraParticleDataNormalizedAge;
+	FShaderResourceParameter NiagaraParticleDataMaterialRandom;
+	FShaderResourceParameter NiagaraParticleDataCustomSorting;
+	FShaderResourceParameter NiagaraParticleDataMaterialParam0;
+	FShaderResourceParameter NiagaraParticleDataMaterialParam1;
+	FShaderResourceParameter NiagaraParticleDataMaterialParam2;
+	FShaderResourceParameter NiagaraParticleDataMaterialParam3;
+
 	FShaderParameter FloatDataOffset;
 	FShaderParameter FloatDataStride;
 
diff --git a/Engine/Plugins/FX/Niagara/Source/NiagaraVertexFactories/Public/NiagaraRibbonVertexFactory.h b/Engine/Plugins/FX/Niagara/Source/NiagaraVertexFactories/Public/NiagaraRibbonVertexFactory.h
--- a/Engine/Plugins/FX/Niagara/Source/NiagaraVertexFactories/Public/NiagaraRibbonVertexFactory.h
+++ b/Engine/Plugins/FX/Niagara/Source/NiagaraVertexFactories/Public/NiagaraRibbonVertexFactory.h
@@ -62,7 +62,18 @@ BEGIN_GLOBAL_SHADER_PARAMETER_STRUCT(FNiagaraRibbonVFLooseParameters, NIAGARAVER
 	SHADER_PARAMETER_SRV(Buffer<float4>, TangentsAndDistances)
 	SHADER_PARAMETER_SRV(Buffer<uint>, MultiRibbonIndices)
 	SHADER_PARAMETER_SRV(Buffer<float>, PackedPerRibbonDataByIndex)
-	SHADER_PARAMETER_SRV(Buffer<float>, NiagaraParticleDataFloat)
+	SHADER_PARAMETER_SRV(Buffer<float>, NiagaraParticleDataPosition)
+	SHADER_PARAMETER_SRV(Buffer<float>, NiagaraParticleDataVelocity)
+	SHADER_PARAMETER_SRV(Buffer<float>, NiagaraParticleDataColor)
+	SHADER_PARAMETER_SRV(Buffer<float>, NiagaraParticleDataWidth)
+	SHADER_PARAMETER_SRV(Buffer<float>, NiagaraParticleDataTwist)
+	SHADER_PARAMETER_SRV(Buffer<float>, NiagaraParticleDataFacing)
+	SHADER_PARAMETER_SRV(Buffer<float>, NiagaraParticleDataNormalizedAge)
+	SHADER_PARAMETER_SRV(Buffer<float>, NiagaraParticleDataMaterialRandom)
+	SHADER_PARAMETER_SRV(Buffer<float>, NiagaraParticleDataMaterialParam0)
+	SHADER_PARAMETER_SRV(Buffer<float>, NiagaraParticleDataMaterialParam1)
+	SHADER_PARAMETER_SRV(Buffer<float>, NiagaraParticleDataMaterialParam2)
+	SHADER_PARAMETER_SRV(Buffer<float>, NiagaraParticleDataMaterialParam3)
 END_GLOBAL_SHADER_PARAMETER_STRUCT()
 typedef TUniformBufferRef<FNiagaraRibbonVFLooseParameters> FNiagaraRibbonVFLooseParametersRef;
 
@@ -170,7 +181,7 @@ public:
 		FacingMode = InFacingMode;
 	}
 
-	FORCEINLINE FShaderResourceViewRHIRef GetParticleDataFloatSRV()
+	FORCEINLINE FRHIShaderResourceView* GetParticleDataFloatSRV()
 	{
 		return ParticleDataFloatSRV;
 	}
@@ -185,7 +196,7 @@ public:
 		return FloatDataStride;
 	}
 
-	FORCEINLINE FShaderResourceViewRHIRef GetSortedIndicesSRV()
+	FORCEINLINE FRHIShaderResourceView* GetSortedIndicesSRV()
 	{
 		return SortedIndicesSRV;
 	}
@@ -195,17 +206,17 @@ public:
 		return SortedIndicesOffset;
 	}
 
-	FORCEINLINE FShaderResourceViewRHIRef GetTangentAndDistancesSRV()
+	FORCEINLINE FRHIShaderResourceView* GetTangentAndDistancesSRV()
 	{
 		return TangentAndDistancesSRV;
 	}
 
-	FORCEINLINE FShaderResourceViewRHIRef GetMultiRibbonIndicesSRV()
+	FORCEINLINE FRHIShaderResourceView* GetMultiRibbonIndicesSRV()
 	{
 		return MultiRibbonIndicesSRV;
 	}
 
-	FORCEINLINE FShaderResourceViewRHIRef GetPackedPerRibbonDataByIndexSRV()
+	FORCEINLINE FRHIShaderResourceView* GetPackedPerRibbonDataByIndexSRV()
 	{
 		return PackedPerRibbonDataByIndexSRV;
 	}
diff --git a/Engine/Plugins/FX/Niagara/Source/NiagaraVertexFactories/Public/NiagaraSpriteVertexFactory.h b/Engine/Plugins/FX/Niagara/Source/NiagaraVertexFactories/Public/NiagaraSpriteVertexFactory.h
--- a/Engine/Plugins/FX/Niagara/Source/NiagaraVertexFactories/Public/NiagaraSpriteVertexFactory.h
+++ b/Engine/Plugins/FX/Niagara/Source/NiagaraVertexFactories/Public/NiagaraSpriteVertexFactory.h
@@ -64,7 +64,23 @@ BEGIN_GLOBAL_SHADER_PARAMETER_STRUCT(FNiagaraSpriteVFLooseParameters, NIAGARAVER
 	SHADER_PARAMETER(uint32, SortedIndicesOffset)
 	SHADER_PARAMETER(uint32, IndirectArgsOffset)
 	SHADER_PARAMETER_SRV(Buffer<float2>, CutoutGeometry)
-	SHADER_PARAMETER_SRV(Buffer<float>, NiagaraParticleDataFloat)
+	SHADER_PARAMETER_SRV(Buffer<float>, NiagaraParticleDataPosition)
+	SHADER_PARAMETER_SRV(Buffer<float>, NiagaraParticleDataVelocity)
+	SHADER_PARAMETER_SRV(Buffer<float>, NiagaraParticleDataColor)
+	SHADER_PARAMETER_SRV(Buffer<float>, NiagaraParticleDataRotation)
+	SHADER_PARAMETER_SRV(Buffer<float>, NiagaraParticleDataSize)
+	SHADER_PARAMETER_SRV(Buffer<float>, NiagaraParticleDataFacing)
+	SHADER_PARAMETER_SRV(Buffer<float>, NiagaraParticleDataAlignment)
+	SHADER_PARAMETER_SRV(Buffer<float>, NiagaraParticleDataSubImage)
+	SHADER_PARAMETER_SRV(Buffer<float>, NiagaraParticleDataCameraOffset)
+	SHADER_PARAMETER_SRV(Buffer<float>, NiagaraParticleDataUVScale)
+	SHADER_PARAMETER_SRV(Buffer<float>, NiagaraParticleDataNormalizedAge)
+	SHADER_PARAMETER_SRV(Buffer<float>, NiagaraParticleDataMaterialRandom)
+	SHADER_PARAMETER_SRV(Buffer<float>, NiagaraParticleDataCustomSorting)
+	SHADER_PARAMETER_SRV(Buffer<float>, NiagaraParticleDataMaterialParam0)
+	SHADER_PARAMETER_SRV(Buffer<float>, NiagaraParticleDataMaterialParam1)
+	SHADER_PARAMETER_SRV(Buffer<float>, NiagaraParticleDataMaterialParam2)
+	SHADER_PARAMETER_SRV(Buffer<float>, NiagaraParticleDataMaterialParam3)
 	SHADER_PARAMETER_SRV(Buffer<int>, SortedIndices)
 	SHADER_PARAMETER_SRV(Buffer<uint>, IndirectArgsBuffer)
 END_GLOBAL_SHADER_PARAMETER_STRUCT()
@@ -161,7 +177,7 @@ public:
 		SortedIndicesOffset = InSortedIndicesOffset;
 	}
 
-	FORCEINLINE FShaderResourceViewRHIRef GetParticleDataFloatSRV()
+	FORCEINLINE FRHIShaderResourceView* GetParticleDataFloatSRV()
 	{
 		return ParticleDataFloatSRV;
 	}
@@ -176,7 +192,7 @@ public:
 		return FloatDataStride;
 	}
 
-	FORCEINLINE FShaderResourceViewRHIRef GetSortedIndicesSRV()
+	FORCEINLINE FRHIShaderResourceView* GetSortedIndicesSRV()
 	{
 		return SortedIndicesSRV;
 	}
