FNM - Allo Niagara Vertgexfactories to bind differnt buffertypes per UseCase. This will eventually allow us to reduce the memory footprint with a later check-in.

diff --git a/Engine/Plugins/FX/Niagara/Shaders/Private/NiagaraSortKeyGen.usf b/Engine/Plugins/FX/Niagara/Shaders/Private/NiagaraSortKeyGen.usf
--- a/Engine/Plugins/FX/Niagara/Shaders/Private/NiagaraSortKeyGen.usf
+++ b/Engine/Plugins/FX/Niagara/Shaders/Private/NiagaraSortKeyGen.usf
@@ -18,6 +18,10 @@
 #define NIAGARA_KEYGEN 1
 #include "NiagaraVFParticleAccess.usf"
 
+#if NIAGARA_LOOSEPARAM_BINDING
+Buffer<float> NiagaraParticleDataFloat;
+#endif
+
 Buffer<uint> GPUParticleCountBuffer;
 
 // (ParticleCount, GPUParticleCountOffset)
@@ -70,7 +74,7 @@ void GenerateParticleSortKeys(uint ParticleIndex : SV_DispatchThreadID)
 		{
 			if (SortMode == SORT_VIEW_DEPTH || SortMode == SORT_VIEW_DISTANCE)
 			{ 
-				float3 ParticlePosition = GetVec3(SortAttributeOffset, ParticleIndex); 
+				float3 ParticlePosition = GetVec3(NIAGARA_LOOSEPARAM_STRUCT(NiagaraParticleDataFloat), SortAttributeOffset, ParticleIndex); 
 				float3 CameraVector = ParticlePosition - CameraPosition;
 				SortValue = SortMode == SORT_VIEW_DEPTH ? dot(CameraVector, CameraDirection) : length(CameraVector);
 				
@@ -84,7 +88,7 @@ void GenerateParticleSortKeys(uint ParticleIndex : SV_DispatchThreadID)
 			}
 			else
 			{
-				SortValue = GetFloat(SortAttributeOffset, ParticleIndex);
+				SortValue = GetFloat(NIAGARA_LOOSEPARAM_STRUCT(NiagaraParticleDataFloat), SortAttributeOffset, ParticleIndex);
 
 #if SORT_MAX_PRECISION
 				SortKey = (asuint(SortValue)) >> SortKeyShift;
