- Pass seed into RandomrTriCoord as we don't have the structure available unless the function is used

#rnx

diff --git a/Engine/Plugins/FX/Niagara/Shaders/Private/NiagaraDataInterfaceSkeletalMesh.ush b/Engine/Plugins/FX/Niagara/Shaders/Private/NiagaraDataInterfaceSkeletalMesh.ush
--- a/Engine/Plugins/FX/Niagara/Shaders/Private/NiagaraDataInterfaceSkeletalMesh.ush
+++ b/Engine/Plugins/FX/Niagara/Shaders/Private/NiagaraDataInterfaceSkeletalMesh.ush
@@ -179,9 +179,9 @@ FDISkelMeshSkinnedVertex DISkelMesh_GetSkinnedPointOnTriangleWS(in FDISkelMeshCo
 	return FinalVertex;
 }
 
-void DISKelMesh_RandomTriCoord(in FDISkelMeshContext DIContext, NiagaraRandInfo InRandomInfo, out int OutTriangle, out float3 OutBaryCoord)
+void DISKelMesh_RandomTriCoord(in FDISkelMeshContext DIContext, uint Seed1, uint Seed2, uint Seed3, out int OutTriangle, out float3 OutBaryCoord)
 {
-	float RandT0 = NiagaraInternalNoise(InRandomInfo.Seed1, InRandomInfo.Seed2, InRandomInfo.Seed3);
+	float RandT0 = NiagaraInternalNoise(Seed1, Seed2, Seed3);
 
 	[branch]
 	if ( !DIContext.UniformTriangleSamplingEnable )
@@ -196,7 +196,7 @@ void DISKelMesh_RandomTriCoord(in FDISkelMeshContext DIContext, NiagaraRandInfo
 		float TriangleProbability = DIContext.MeshTriangleSamplerProbaBuffer[TriangleIndex];
 
 		// Alias check
-		float RandT1 = NiagaraInternalNoise(InRandomInfo.Seed1, InRandomInfo.Seed2, InRandomInfo.Seed3);
+		float RandT1 = NiagaraInternalNoise(Seed1, Seed2, Seed3);
 		if( RandT1 > TriangleProbability )
 		{
 			TriangleIndex = DIContext.MeshTriangleSamplerAliasBuffer[TriangleIndex];
@@ -204,8 +204,8 @@ void DISKelMesh_RandomTriCoord(in FDISkelMeshContext DIContext, NiagaraRandInfo
 		OutTriangle = TriangleIndex;
 	}
 
-	float r0 = NiagaraInternalNoise(InRandomInfo.Seed1, InRandomInfo.Seed2, InRandomInfo.Seed3);
-	float r1 = NiagaraInternalNoise(InRandomInfo.Seed1, InRandomInfo.Seed2, InRandomInfo.Seed3);
+	float r0 = NiagaraInternalNoise(Seed1, Seed2, Seed3);
+	float r1 = NiagaraInternalNoise(Seed1, Seed2, Seed3);
 	float sqrt0 = sqrt(r0);
 	float sqrt1 = sqrt(r1);
 	OutBaryCoord = float3(1.0f - sqrt0, sqrt0 * (1.0 - r1), r1 * sqrt0);
diff --git a/Engine/Plugins/FX/Niagara/Source/Niagara/Private/NiagaraDataInterfaceSkeletalMesh.cpp b/Engine/Plugins/FX/Niagara/Source/Niagara/Private/NiagaraDataInterfaceSkeletalMesh.cpp
--- a/Engine/Plugins/FX/Niagara/Source/Niagara/Private/NiagaraDataInterfaceSkeletalMesh.cpp
+++ b/Engine/Plugins/FX/Niagara/Source/Niagara/Private/NiagaraDataInterfaceSkeletalMesh.cpp
@@ -1971,7 +1971,7 @@ bool UNiagaraDataInterfaceSkeletalMesh::GetFunctionHLSL(const FName& DefinitionF
 	// Triangle Sampling
 	if (DefinitionFunctionName == FSkeletalMeshInterfaceHelper::RandomTriCoordName)
 	{
-		static const TCHAR* FormatSample = TEXT("void {InstanceFunctionName} (NiagaraRandInfo InRandomInfo, out {MeshTriCoordinateStructName} OutCoord) { {GetDISkelMeshContextName} DISKelMesh_RandomTriCoord(DIContext, InRandomInfo, OutCoord.Tri, OutCoord.BaryCoord); }");
+		static const TCHAR* FormatSample = TEXT("void {InstanceFunctionName} (NiagaraRandInfo InRandomInfo, out {MeshTriCoordinateStructName} OutCoord) { {GetDISkelMeshContextName} DISKelMesh_RandomTriCoord(DIContext, InRandomInfo.Seed1, InRandomInfo.Seed2, InRandomInfo.Seed3, OutCoord.Tri, OutCoord.BaryCoord); }");
 		OutHLSL += FString::Format(FormatSample, ArgsSample);
 	}
 	else if (DefinitionFunctionName == FSkeletalMeshInterfaceHelper::GetSkinnedTriangleDataWSName)
