Implemented the particle read data interface on GPU.

The method UNiagaraDataInterface::GetFunctionHLSL receives more information about the function being generated, such as the specifiers and the index inside the list of functions produced by the DI. The extra information is also stored inside FNiagaraDataInterfaceGPUParamInfo, so it's accessible in the Bind() method of the DI's compute shader. This allows the particle read DI to retrieve the field names for each function when the shader is bound, and to map them to register indices when the shader is used. The register indices are cached, so that they're only computed when the source emitter changes.

Also moved the GPU free ID list out of FNiagaraDataBuffer and into FNiagaraDataSet. The previous code was wrong, because when the particle count changes, we need to resize the *input* free ID list, not only the output. Having a list per data buffer doesn't make sense, and a single list per emitter is easier to manage. This fixes duplicated IDs at the beginning of the simulation, and other problems which appeared whenever the buffer had to grow.

diff --git a/Engine/Plugins/Experimental/ChaosNiagara/Source/ChaosNiagara/Classes/NiagaraDataInterfaceChaosDestruction.h b/Engine/Plugins/Experimental/ChaosNiagara/Source/ChaosNiagara/Classes/NiagaraDataInterfaceChaosDestruction.h
--- a/Engine/Plugins/Experimental/ChaosNiagara/Source/ChaosNiagara/Classes/NiagaraDataInterfaceChaosDestruction.h
+++ b/Engine/Plugins/Experimental/ChaosNiagara/Source/ChaosNiagara/Classes/NiagaraDataInterfaceChaosDestruction.h
@@ -554,8 +554,8 @@ public:
 
 	//----------------------------------------------------------------------------
 	// GPU sim functionality
-	virtual void GetParameterDefinitionHLSL(FNiagaraDataInterfaceGPUParamInfo& ParamInfo, FString& OutHLSL) override; 
-	virtual bool GetFunctionHLSL(const FName&  DefinitionFunctionName, FString InstanceFunctionName, FNiagaraDataInterfaceGPUParamInfo& ParamInfo, FString& OutHLSL) override;
+	virtual void GetParameterDefinitionHLSL(const FNiagaraDataInterfaceGPUParamInfo& ParamInfo, FString& OutHLSL) override;
+	virtual bool GetFunctionHLSL(const FNiagaraDataInterfaceGPUParamInfo& ParamInfo, const FNiagaraDataInterfaceGeneratedFunction& FunctionInfo, int FunctionInstanceIndex, FString& OutHLSL) override;
 	virtual FNiagaraDataInterfaceParametersCS* ConstructComputeParameters() const override;
 
 	//----------------------------------------------------------------------------
diff --git a/Engine/Plugins/Experimental/ChaosNiagara/Source/ChaosNiagara/Private/NiagaraDataInterfaceChaosDestruction.cpp b/Engine/Plugins/Experimental/ChaosNiagara/Source/ChaosNiagara/Private/NiagaraDataInterfaceChaosDestruction.cpp
--- a/Engine/Plugins/Experimental/ChaosNiagara/Source/ChaosNiagara/Private/NiagaraDataInterfaceChaosDestruction.cpp
+++ b/Engine/Plugins/Experimental/ChaosNiagara/Source/ChaosNiagara/Private/NiagaraDataInterfaceChaosDestruction.cpp
@@ -3370,7 +3370,7 @@ void UNiagaraDataInterfaceChaosDestruction::GetTrailingData(FVectorVMContext& Co
 //----------------------------------------------------------------------------
 // GPU sim functionality
 //
-void UNiagaraDataInterfaceChaosDestruction::GetParameterDefinitionHLSL(FNiagaraDataInterfaceGPUParamInfo& ParamInfo, FString& OutHLSL)
+void UNiagaraDataInterfaceChaosDestruction::GetParameterDefinitionHLSL(const FNiagaraDataInterfaceGPUParamInfo& ParamInfo, FString& OutHLSL)
 {
 
 	// This will get indented in the generated HLSL, which won't look good. 
@@ -3424,9 +3424,9 @@ void UNiagaraDataInterfaceChaosDestruction::GetParameterDefinitionHLSL(FNiagaraD
 	*/
 }
 
-bool UNiagaraDataInterfaceChaosDestruction::GetFunctionHLSL(const FName& DefinitionFunctionName, FString InstanceFunctionName, FNiagaraDataInterfaceGPUParamInfo& ParamInfo, FString& OutHLSL)
+bool UNiagaraDataInterfaceChaosDestruction::GetFunctionHLSL(const FNiagaraDataInterfaceGPUParamInfo& ParamInfo, const FNiagaraDataInterfaceGeneratedFunction& FunctionInfo, int FunctionInstanceIndex, FString& OutHLSL)
 {
-	if (DefinitionFunctionName == GetPositionName)
+	if (FunctionInfo.DefinitionName == GetPositionName)
 	{
 		static const TCHAR *Format = TEXT(R"(
 			void {FunctionName}(in int ParticleID, out float3 Out_Position) 
@@ -3437,13 +3437,13 @@ bool UNiagaraDataInterfaceChaosDestruction::GetFunctionHLSL(const FName& Definit
 		)");
 
 		TMap<FString, FStringFormatArg> Args = {
-			{ TEXT("FunctionName"), InstanceFunctionName },
+			{ TEXT("FunctionName"), FunctionInfo.InstanceName },
 			{ TEXT("Symbol"), ParamInfo.DataInterfaceHLSLSymbol },
 		};
 		OutHLSL += FString::Format(Format, Args);
 		return true;
 	}
-	else if (DefinitionFunctionName == GetNormalName)
+	else if (FunctionInfo.DefinitionName == GetNormalName)
 	{
 		static const TCHAR *Format = TEXT(R"(
 			void {FunctionName}(in int ParticleID, out float3 Out_Normal) 
@@ -3454,13 +3454,13 @@ bool UNiagaraDataInterfaceChaosDestruction::GetFunctionHLSL(const FName& Definit
 		)");
 
 		TMap<FString, FStringFormatArg> Args = {
-			{ TEXT("FunctionName"), InstanceFunctionName },
+			{ TEXT("FunctionName"), FunctionInfo.InstanceName },
 			{ TEXT("Symbol"), ParamInfo.DataInterfaceHLSLSymbol },
 		};
 		OutHLSL += FString::Format(Format, Args);
 		return true;
 	}
-	else if (DefinitionFunctionName == GetVelocityName)
+	else if (FunctionInfo.DefinitionName == GetVelocityName)
 	{
 		static const TCHAR *Format = TEXT(R"(
 			void {FunctionName}(in int ParticleID, out float3 Out_Velocity) 
@@ -3471,13 +3471,13 @@ bool UNiagaraDataInterfaceChaosDestruction::GetFunctionHLSL(const FName& Definit
 		)");
 
 		TMap<FString, FStringFormatArg> Args = {
-			{ TEXT("FunctionName"), InstanceFunctionName },
+			{ TEXT("FunctionName"), FunctionInfo.InstanceName },
 			{ TEXT("Symbol"), ParamInfo.DataInterfaceHLSLSymbol },
 		};
 		OutHLSL += FString::Format(Format, Args);
 		return true;
 	}
-	else if (DefinitionFunctionName == GetAngularVelocityName)
+	else if (FunctionInfo.DefinitionName == GetAngularVelocityName)
 	{
 		static const TCHAR *Format = TEXT(R"(
 			void {FunctionName}(in int ParticleID, out float3 Out_AngularVelocity) 
@@ -3488,13 +3488,13 @@ bool UNiagaraDataInterfaceChaosDestruction::GetFunctionHLSL(const FName& Definit
 		)");
 
 		TMap<FString, FStringFormatArg> Args = {
-			{ TEXT("FunctionName"), InstanceFunctionName },
+			{ TEXT("FunctionName"), FunctionInfo.InstanceName },
 			{ TEXT("Symbol"), ParamInfo.DataInterfaceHLSLSymbol },
 		};
 		OutHLSL += FString::Format(Format, Args);
 		return true;
 	}
-	else if (DefinitionFunctionName == GetExtentMinName)
+	else if (FunctionInfo.DefinitionName == GetExtentMinName)
 	{
 		static const TCHAR *Format = TEXT(R"(
 			void {FunctionName}(in int ParticleID, out float Out_ExtentMin) 
@@ -3505,13 +3505,13 @@ bool UNiagaraDataInterfaceChaosDestruction::GetFunctionHLSL(const FName& Definit
 		)");
 
 		TMap<FString, FStringFormatArg> Args = {
-			{ TEXT("FunctionName"), InstanceFunctionName },
+			{ TEXT("FunctionName"), FunctionInfo.InstanceName },
 			{ TEXT("Symbol"), ParamInfo.DataInterfaceHLSLSymbol },
 		};
 		OutHLSL += FString::Format(Format, Args);
 		return true;
 	}
-	else if (DefinitionFunctionName == GetExtentMaxName)
+	else if (FunctionInfo.DefinitionName == GetExtentMaxName)
 	{
 		static const TCHAR *Format = TEXT(R"(
 			void {FunctionName}(in int ParticleID, out float Out_ExtentMax) 
@@ -3522,13 +3522,13 @@ bool UNiagaraDataInterfaceChaosDestruction::GetFunctionHLSL(const FName& Definit
 		)");
 
 		TMap<FString, FStringFormatArg> Args = {
-			{ TEXT("FunctionName"), InstanceFunctionName },
+			{ TEXT("FunctionName"), FunctionInfo.InstanceName },
 			{ TEXT("Symbol"), ParamInfo.DataInterfaceHLSLSymbol },
 		};
 		OutHLSL += FString::Format(Format, Args);
 		return true;
 	}
-	else if (DefinitionFunctionName == GetVolumeName)
+	else if (FunctionInfo.DefinitionName == GetVolumeName)
 	{
 		static const TCHAR *Format = TEXT(R"(
 			void {FunctionName}(in int ParticleID, out float Out_Volume) 
@@ -3539,13 +3539,13 @@ bool UNiagaraDataInterfaceChaosDestruction::GetFunctionHLSL(const FName& Definit
 		)");
 
 		TMap<FString, FStringFormatArg> Args = {
-			{ TEXT("FunctionName"), InstanceFunctionName },
+			{ TEXT("FunctionName"), FunctionInfo.InstanceName },
 			{ TEXT("Symbol"), ParamInfo.DataInterfaceHLSLSymbol },
 		};
 		OutHLSL += FString::Format(Format, Args);
 		return true;
 	}
-	else if (DefinitionFunctionName == GetParticleIdsToSpawnAtTimeName)
+	else if (FunctionInfo.DefinitionName == GetParticleIdsToSpawnAtTimeName)
 	{
 
 		static const TCHAR *Format = TEXT(R"(
@@ -3562,13 +3562,13 @@ bool UNiagaraDataInterfaceChaosDestruction::GetFunctionHLSL(const FName& Definit
 		)");
 
 		TMap<FString, FStringFormatArg> Args = {
-			{ TEXT("FunctionName"), InstanceFunctionName },
+			{ TEXT("FunctionName"), FunctionInfo.InstanceName },
 			{ TEXT("Symbol"), ParamInfo.DataInterfaceHLSLSymbol },
 		};
 		OutHLSL += FString::Format(Format, Args);
 		return true;
 	}
-	else if (DefinitionFunctionName == GetPointTypeName)
+	else if (FunctionInfo.DefinitionName == GetPointTypeName)
 	{
 		static const TCHAR *Format = TEXT(R"(
 			void {FunctionName}(in int ParticleID, out int Out_PointType) 
@@ -3579,12 +3579,12 @@ bool UNiagaraDataInterfaceChaosDestruction::GetFunctionHLSL(const FName& Definit
 		)");
 
 		TMap<FString, FStringFormatArg> Args = {
-			{ TEXT("FunctionName"), InstanceFunctionName },
+			{ TEXT("FunctionName"), FunctionInfo.InstanceName },
 		};
 		OutHLSL += FString::Format(Format, Args);
 		return true;
 	}
-	else if (DefinitionFunctionName == GetColorName)
+	else if (FunctionInfo.DefinitionName == GetColorName)
 	{
 		static const TCHAR *Format = TEXT(R"(
 			void {FunctionName}(in int ParticleID, out float4 Out_Color) 
@@ -3595,13 +3595,13 @@ bool UNiagaraDataInterfaceChaosDestruction::GetFunctionHLSL(const FName& Definit
 		)");
 
 		TMap<FString, FStringFormatArg> Args = {
-			{ TEXT("FunctionName"), InstanceFunctionName },
+			{ TEXT("FunctionName"), FunctionInfo.InstanceName },
 			{ TEXT("Symbol"), ParamInfo.DataInterfaceHLSLSymbol },
 		};
 		OutHLSL += FString::Format(Format, Args);
 		return true;
 	}
-	else if (DefinitionFunctionName == GetSolverTimeName)
+	else if (FunctionInfo.DefinitionName == GetSolverTimeName)
 	{
 		static const TCHAR *Format = TEXT(R"(
 			void {FunctionName}(out float Out_SolverTime) 
@@ -3611,13 +3611,13 @@ bool UNiagaraDataInterfaceChaosDestruction::GetFunctionHLSL(const FName& Definit
 		)");
 
 		TMap<FString, FStringFormatArg> Args = {
-			{ TEXT("FunctionName"), InstanceFunctionName },
+			{ TEXT("FunctionName"), FunctionInfo.InstanceName },
 			{ TEXT("Symbol"), ParamInfo.DataInterfaceHLSLSymbol },
 		};
 		OutHLSL += FString::Format(Format, Args);
 		return true;
 	}
-	else if (DefinitionFunctionName == GetDensityName)
+	else if (FunctionInfo.DefinitionName == GetDensityName)
 	{
 		static const TCHAR *Format = TEXT(R"(
 			void {FunctionName}(in int ParticleID, out float Out_Density) 
@@ -3628,13 +3628,13 @@ bool UNiagaraDataInterfaceChaosDestruction::GetFunctionHLSL(const FName& Definit
 		)");
 
 		TMap<FString, FStringFormatArg> Args = {
-			{ TEXT("FunctionName"), InstanceFunctionName },
+			{ TEXT("FunctionName"), FunctionInfo.InstanceName },
 			{ TEXT("Symbol"), ParamInfo.DataInterfaceHLSLSymbol },
 		};
 		OutHLSL += FString::Format(Format, Args);
 		return true;
 	}
-	else if (DefinitionFunctionName == GetFrictionName)
+	else if (FunctionInfo.DefinitionName == GetFrictionName)
 	{
 		static const TCHAR *Format = TEXT(R"(
 			void {FunctionName}(in int ParticleID, out float Out_Friction) 
@@ -3645,13 +3645,13 @@ bool UNiagaraDataInterfaceChaosDestruction::GetFunctionHLSL(const FName& Definit
 		)");
 
 		TMap<FString, FStringFormatArg> Args = {
-			{ TEXT("FunctionName"), InstanceFunctionName },
+			{ TEXT("FunctionName"), FunctionInfo.InstanceName },
 			{ TEXT("Symbol"), ParamInfo.DataInterfaceHLSLSymbol },
 		};
 		OutHLSL += FString::Format(Format, Args);
 		return true;
 	}
-	else if (DefinitionFunctionName == GetRestitutionName)
+	else if (FunctionInfo.DefinitionName == GetRestitutionName)
 	{
 		static const TCHAR *Format = TEXT(R"(
 			void {FunctionName}(in int ParticleID, out float Out_Restitution) 
@@ -3662,13 +3662,13 @@ bool UNiagaraDataInterfaceChaosDestruction::GetFunctionHLSL(const FName& Definit
 		)");
 
 		TMap<FString, FStringFormatArg> Args = {
-			{ TEXT("FunctionName"), InstanceFunctionName },
+			{ TEXT("FunctionName"), FunctionInfo.InstanceName },
 			{ TEXT("Symbol"), ParamInfo.DataInterfaceHLSLSymbol },
 		};
 		OutHLSL += FString::Format(Format, Args);
 		return true;
 	}
-	else if (DefinitionFunctionName == GetTransformName)
+	else if (FunctionInfo.DefinitionName == GetTransformName)
 	{
 		static const TCHAR *Format = TEXT(R"(
 			void {FunctionName}(in int ParticleID, out float3 Out_Translation, out float4 Out_Rotation, out float3 Out_Scale) 
@@ -3681,13 +3681,13 @@ bool UNiagaraDataInterfaceChaosDestruction::GetFunctionHLSL(const FName& Definit
 		)");
 
 		TMap<FString, FStringFormatArg> Args = {
-			{ TEXT("FunctionName"), InstanceFunctionName },
+			{ TEXT("FunctionName"), FunctionInfo.InstanceName },
 			{ TEXT("Symbol"), ParamInfo.DataInterfaceHLSLSymbol },
 		};
 		OutHLSL += FString::Format(Format, Args);
 		return true;
 	}
-	else if (DefinitionFunctionName == GetSizeName)
+	else if (FunctionInfo.DefinitionName == GetSizeName)
 	{
 		static const TCHAR *Format = TEXT(R"(
 			void {FunctionName}(in int ParticleID, out float3 Out_Size) 
@@ -3698,13 +3698,13 @@ bool UNiagaraDataInterfaceChaosDestruction::GetFunctionHLSL(const FName& Definit
 		)");
 
 		TMap<FString, FStringFormatArg> Args = {
-			{ TEXT("FunctionName"), InstanceFunctionName },
+			{ TEXT("FunctionName"), FunctionInfo.InstanceName },
 			{ TEXT("Symbol"), ParamInfo.DataInterfaceHLSLSymbol },
 		};
 		OutHLSL += FString::Format(Format, Args);
 		return true;
 	}
-	else if (DefinitionFunctionName == GetSurfaceTypeName)
+	else if (FunctionInfo.DefinitionName == GetSurfaceTypeName)
 	{
 		static const TCHAR *Format = TEXT(R"(
 			void {FunctionName}(in int ParticleID, out int Out_SurfaceType) 
@@ -3715,13 +3715,13 @@ bool UNiagaraDataInterfaceChaosDestruction::GetFunctionHLSL(const FName& Definit
 		)");
 
 		TMap<FString, FStringFormatArg> Args = {
-			{ TEXT("FunctionName"), InstanceFunctionName },
+			{ TEXT("FunctionName"), FunctionInfo.InstanceName },
 			{ TEXT("Symbol"), ParamInfo.DataInterfaceHLSLSymbol },
 		};
 		OutHLSL += FString::Format(Format, Args);
 		return true;
 	}
-	else if (DefinitionFunctionName == GetCollisionDataName)
+	else if (FunctionInfo.DefinitionName == GetCollisionDataName)
 	{
 		static const TCHAR *Format = TEXT(R"(
 			void {FunctionName}(in int ParticleID, out float3 Out_Location, 
@@ -3750,13 +3750,13 @@ bool UNiagaraDataInterfaceChaosDestruction::GetFunctionHLSL(const FName& Definit
 		)");
 
 		TMap<FString, FStringFormatArg> Args = {
-			{ TEXT("FunctionName"), InstanceFunctionName },
+			{ TEXT("FunctionName"), FunctionInfo.InstanceName },
 			{ TEXT("Symbol"), ParamInfo.DataInterfaceHLSLSymbol },
 		};
 		OutHLSL += FString::Format(Format, Args);
 		return true;
 	}
-	else if (DefinitionFunctionName == GetBreakingDataName)
+	else if (FunctionInfo.DefinitionName == GetBreakingDataName)
 	{
 		static const TCHAR *Format = TEXT(R"(
 			void {FunctionName}(in int ParticleID, out float3 Out_Location,
@@ -3775,13 +3775,13 @@ bool UNiagaraDataInterfaceChaosDestruction::GetFunctionHLSL(const FName& Definit
 		)");
 
 		TMap<FString, FStringFormatArg> Args = {
-			{ TEXT("FunctionName"), InstanceFunctionName },
+			{ TEXT("FunctionName"), FunctionInfo.InstanceName },
 			{ TEXT("Symbol"), ParamInfo.DataInterfaceHLSLSymbol },
 		};
 		OutHLSL += FString::Format(Format, Args);
 		return true;
 	}
-	else if (DefinitionFunctionName == GetTrailingDataName)
+	else if (FunctionInfo.DefinitionName == GetTrailingDataName)
 	{
 		static const TCHAR *Format = TEXT(R"(
 			void {FunctionName}(in int ParticleID, out float3 Out_Location,
@@ -3800,7 +3800,7 @@ bool UNiagaraDataInterfaceChaosDestruction::GetFunctionHLSL(const FName& Definit
 		)");
 
 		TMap<FString, FStringFormatArg> Args = {
-			{ TEXT("FunctionName"), InstanceFunctionName },
+			{ TEXT("FunctionName"), FunctionInfo.InstanceName },
 			{ TEXT("Symbol"), ParamInfo.DataInterfaceHLSLSymbol },
 		};
 		OutHLSL += FString::Format(Format, Args);
diff --git a/Engine/Plugins/Experimental/HairStrands/Source/HairStrandsNiagara/Private/NiagaraDataInterfaceHairStrands.cpp b/Engine/Plugins/Experimental/HairStrands/Source/HairStrandsNiagara/Private/NiagaraDataInterfaceHairStrands.cpp
--- a/Engine/Plugins/Experimental/HairStrands/Source/HairStrandsNiagara/Private/NiagaraDataInterfaceHairStrands.cpp
+++ b/Engine/Plugins/Experimental/HairStrands/Source/HairStrandsNiagara/Private/NiagaraDataInterfaceHairStrands.cpp
@@ -2235,7 +2235,7 @@ void UNiagaraDataInterfaceHairStrands::NeedSimulationReset(FVectorVMContext& Con
 	// @todo : implement function for cpu 
 }
 
-bool UNiagaraDataInterfaceHairStrands::GetFunctionHLSL(const FName& DefinitionFunctionName, FString InstanceFunctionName, FNiagaraDataInterfaceGPUParamInfo& ParamInfo, FString& OutHLSL)
+bool UNiagaraDataInterfaceHairStrands::GetFunctionHLSL(const FNiagaraDataInterfaceGPUParamInfo& ParamInfo, const FNiagaraDataInterfaceGeneratedFunction& FunctionInfo, FString& OutHLSL)
 {
 	FNDIHairStrandsParametersName ParamNames(ParamInfo.DataInterfaceHLSLSymbol);
 
@@ -2812,7 +2812,7 @@ void UNiagaraDataInterfaceHairStrands::GetCommonHLSL(FString& OutHLSL)
 	OutHLSL += TEXT("#include \"/Plugin/Experimental/HairStrands/Private/NiagaraDataInterfaceHairStrands.ush\"\n");
 }
 
-void UNiagaraDataInterfaceHairStrands::GetParameterDefinitionHLSL(FNiagaraDataInterfaceGPUParamInfo& ParamInfo, FString& OutHLSL)
+void UNiagaraDataInterfaceHairStrands::GetParameterDefinitionHLSL(const FNiagaraDataInterfaceGPUParamInfo& ParamInfo, FString& OutHLSL)
 {
 	OutHLSL += TEXT("DIHAIRSTRANDS_DECLARE_CONSTANTS(") + ParamInfo.DataInterfaceHLSLSymbol + TEXT(")\n");
 }
diff --git a/Engine/Plugins/Experimental/HairStrands/Source/HairStrandsNiagara/Private/NiagaraDataInterfacePhysicsAsset.cpp b/Engine/Plugins/Experimental/HairStrands/Source/HairStrandsNiagara/Private/NiagaraDataInterfacePhysicsAsset.cpp
--- a/Engine/Plugins/Experimental/HairStrands/Source/HairStrandsNiagara/Private/NiagaraDataInterfacePhysicsAsset.cpp
+++ b/Engine/Plugins/Experimental/HairStrands/Source/HairStrandsNiagara/Private/NiagaraDataInterfacePhysicsAsset.cpp
@@ -790,7 +790,7 @@ void UNiagaraDataInterfacePhysicsAsset::GetProjectionPoint(FVectorVMContext& Con
 {
 }
 
-bool UNiagaraDataInterfacePhysicsAsset::GetFunctionHLSL(const FName& DefinitionFunctionName, FString InstanceFunctionName, FNiagaraDataInterfaceGPUParamInfo& ParamInfo, FString& OutHLSL)
+bool UNiagaraDataInterfacePhysicsAsset::GetFunctionHLSL(const FNiagaraDataInterfaceGPUParamInfo& ParamInfo, const FNiagaraDataInterfaceGeneratedFunction& FunctionInfo, FString& OutHLSL)
 {
 	FNDIPhysicsAssetParametersName ParamNames(ParamInfo.DataInterfaceHLSLSymbol);
 
@@ -887,7 +887,7 @@ void UNiagaraDataInterfacePhysicsAsset::GetCommonHLSL(FString& OutHLSL)
 	OutHLSL += TEXT("#include \"/Plugin/Experimental/HairStrands/Private/NiagaraDataInterfacePhysicsAsset.ush\"\n");
 }
 
-void UNiagaraDataInterfacePhysicsAsset::GetParameterDefinitionHLSL(FNiagaraDataInterfaceGPUParamInfo& ParamInfo, FString& OutHLSL)
+void UNiagaraDataInterfacePhysicsAsset::GetParameterDefinitionHLSL(const FNiagaraDataInterfaceGPUParamInfo& ParamInfo, FString& OutHLSL)
 {
 	OutHLSL += TEXT("DIPHYSICSASSET_DECLARE_CONSTANTS(") + ParamInfo.DataInterfaceHLSLSymbol + TEXT(")\n");
 }
diff --git a/Engine/Plugins/Experimental/HairStrands/Source/HairStrandsNiagara/Private/NiagaraDataInterfacePressureGrid.cpp b/Engine/Plugins/Experimental/HairStrands/Source/HairStrandsNiagara/Private/NiagaraDataInterfacePressureGrid.cpp
--- a/Engine/Plugins/Experimental/HairStrands/Source/HairStrandsNiagara/Private/NiagaraDataInterfacePressureGrid.cpp
+++ b/Engine/Plugins/Experimental/HairStrands/Source/HairStrandsNiagara/Private/NiagaraDataInterfacePressureGrid.cpp
@@ -686,7 +686,7 @@ void UNiagaraDataInterfacePressureGrid::UpdateGridTransform(FVectorVMContext& Co
 	//UE_LOG(LogPressureGrid, Warning, TEXT("Get Grid Transform : %s"), *InstData->WorldTransform.ToString() );
 }
 
-bool UNiagaraDataInterfacePressureGrid::GetFunctionHLSL(const FName& DefinitionFunctionName, FString InstanceFunctionName, FNiagaraDataInterfaceGPUParamInfo& ParamInfo, FString& OutHLSL)
+bool UNiagaraDataInterfacePressureGrid::GetFunctionHLSL(const FNiagaraDataInterfaceGPUParamInfo& ParamInfo, const FNiagaraDataInterfaceGeneratedFunction& FunctionInfo, FString& OutHLSL)
 {
 	FNDIPressureGridParametersName ParamNames(ParamInfo.DataInterfaceHLSLSymbol);
 
@@ -800,7 +800,7 @@ void UNiagaraDataInterfacePressureGrid::GetCommonHLSL(FString& OutHLSL)
 	OutHLSL += TEXT("#include \"/Plugin/Experimental/HairStrands/Private/NiagaraDataInterfacePressureGrid.ush\"\n");
 }
 
-void UNiagaraDataInterfacePressureGrid::GetParameterDefinitionHLSL(FNiagaraDataInterfaceGPUParamInfo& ParamInfo, FString& OutHLSL)
+void UNiagaraDataInterfacePressureGrid::GetParameterDefinitionHLSL(const FNiagaraDataInterfaceGPUParamInfo& ParamInfo, FString& OutHLSL)
 {
 	OutHLSL += TEXT("DIPRESSUREGRID_DECLARE_CONSTANTS(") + ParamInfo.DataInterfaceHLSLSymbol + TEXT(")\n");
 }
diff --git a/Engine/Plugins/Experimental/HairStrands/Source/HairStrandsNiagara/Public/NiagaraDataInterfaceHairStrands.h b/Engine/Plugins/Experimental/HairStrands/Source/HairStrandsNiagara/Public/NiagaraDataInterfaceHairStrands.h
--- a/Engine/Plugins/Experimental/HairStrands/Source/HairStrandsNiagara/Public/NiagaraDataInterfaceHairStrands.h
+++ b/Engine/Plugins/Experimental/HairStrands/Source/HairStrandsNiagara/Public/NiagaraDataInterfaceHairStrands.h
@@ -159,8 +159,8 @@ public:
 	virtual bool Equals(const UNiagaraDataInterface* Other) const override;
 
 	/** GPU simulation  functionality */
-	virtual bool GetFunctionHLSL(const FName&  DefinitionFunctionName, FString InstanceFunctionName, FNiagaraDataInterfaceGPUParamInfo& ParamInfo, FString& OutHLSL) override;
-	virtual void GetParameterDefinitionHLSL(FNiagaraDataInterfaceGPUParamInfo& ParamInfo, FString& OutHLSL) override;
+	virtual void GetParameterDefinitionHLSL(const FNiagaraDataInterfaceGPUParamInfo& ParamInfo, FString& OutHLSL) override;
+	virtual bool GetFunctionHLSL(const FNiagaraDataInterfaceGPUParamInfo& ParamInfo, const FNiagaraDataInterfaceGeneratedFunction& FunctionInfo, int FunctionInstanceIndex, FString& OutHLSL) override;
 	virtual FNiagaraDataInterfaceParametersCS* ConstructComputeParameters() const override;
 	virtual void ProvidePerInstanceDataForRenderThread(void* DataForRenderThread, void* PerInstanceData, const FNiagaraSystemInstanceID& SystemInstance) override;
 	virtual void GetCommonHLSL(FString& OutHLSL) override;
diff --git a/Engine/Plugins/Experimental/HairStrands/Source/HairStrandsNiagara/Public/NiagaraDataInterfacePhysicsAsset.h b/Engine/Plugins/Experimental/HairStrands/Source/HairStrandsNiagara/Public/NiagaraDataInterfacePhysicsAsset.h
--- a/Engine/Plugins/Experimental/HairStrands/Source/HairStrandsNiagara/Public/NiagaraDataInterfacePhysicsAsset.h
+++ b/Engine/Plugins/Experimental/HairStrands/Source/HairStrandsNiagara/Public/NiagaraDataInterfacePhysicsAsset.h
@@ -137,8 +137,8 @@ public:
 	virtual bool Equals(const UNiagaraDataInterface* Other) const override;
 
 	/** GPU simulation  functionality */
-	virtual bool GetFunctionHLSL(const FName&  DefinitionFunctionName, FString InstanceFunctionName, FNiagaraDataInterfaceGPUParamInfo& ParamInfo, FString& OutHLSL) override;
-	virtual void GetParameterDefinitionHLSL(FNiagaraDataInterfaceGPUParamInfo& ParamInfo, FString& OutHLSL) override;
+	virtual void GetParameterDefinitionHLSL(const FNiagaraDataInterfaceGPUParamInfo& ParamInfo, FString& OutHLSL) override;
+	virtual bool GetFunctionHLSL(const FNiagaraDataInterfaceGPUParamInfo& ParamInfo, const FNiagaraDataInterfaceGeneratedFunction& FunctionInfo, int FunctionInstanceIndex, FString& OutHLSL) override;
 	virtual FNiagaraDataInterfaceParametersCS* ConstructComputeParameters() const override;
 	virtual void ProvidePerInstanceDataForRenderThread(void* DataForRenderThread, void* PerInstanceData, const FNiagaraSystemInstanceID& SystemInstance) override;
 	virtual void GetCommonHLSL(FString& OutHLSL) override;
diff --git a/Engine/Plugins/Experimental/HairStrands/Source/HairStrandsNiagara/Public/NiagaraDataInterfacePressureGrid.h b/Engine/Plugins/Experimental/HairStrands/Source/HairStrandsNiagara/Public/NiagaraDataInterfacePressureGrid.h
--- a/Engine/Plugins/Experimental/HairStrands/Source/HairStrandsNiagara/Public/NiagaraDataInterfacePressureGrid.h
+++ b/Engine/Plugins/Experimental/HairStrands/Source/HairStrandsNiagara/Public/NiagaraDataInterfacePressureGrid.h
@@ -93,8 +93,8 @@ public:
 	virtual bool Equals(const UNiagaraDataInterface* Other) const override;
 
 	/** GPU simulation  functionality */
-	virtual bool GetFunctionHLSL(const FName&  DefinitionFunctionName, FString InstanceFunctionName, FNiagaraDataInterfaceGPUParamInfo& ParamInfo, FString& OutHLSL) override;
-	virtual void GetParameterDefinitionHLSL(FNiagaraDataInterfaceGPUParamInfo& ParamInfo, FString& OutHLSL) override;
+	virtual void GetParameterDefinitionHLSL(const FNiagaraDataInterfaceGPUParamInfo& ParamInfo, FString& OutHLSL) override;
+	virtual bool GetFunctionHLSL(const FNiagaraDataInterfaceGPUParamInfo& ParamInfo, const FNiagaraDataInterfaceGeneratedFunction& FunctionInfo, int FunctionInstanceIndex, FString& OutHLSL) override;
 	virtual FNiagaraDataInterfaceParametersCS* ConstructComputeParameters() const override;
 	virtual void ProvidePerInstanceDataForRenderThread(void* DataForRenderThread, void* PerInstanceData, const FNiagaraSystemInstanceID& SystemInstance) override;
 	virtual void GetCommonHLSL(FString& OutHLSL) override;
diff --git a/Engine/Plugins/FX/HoudiniNiagara/Source/HoudiniNiagara/Classes/NiagaraDataInterfaceHoudiniCSV.h b/Engine/Plugins/FX/HoudiniNiagara/Source/HoudiniNiagara/Classes/NiagaraDataInterfaceHoudiniCSV.h
--- a/Engine/Plugins/FX/HoudiniNiagara/Source/HoudiniNiagara/Classes/NiagaraDataInterfaceHoudiniCSV.h
+++ b/Engine/Plugins/FX/HoudiniNiagara/Source/HoudiniNiagara/Classes/NiagaraDataInterfaceHoudiniCSV.h
@@ -205,8 +205,8 @@ public:
 	
 	//----------------------------------------------------------------------------
 	// GPU / HLSL Functions
-	virtual bool GetFunctionHLSL(const FName& DefinitionFunctionName, FString InstanceFunctionName, FNiagaraDataInterfaceGPUParamInfo& ParamInfo, FString& OutHLSL) override;
-	virtual void GetParameterDefinitionHLSL(FNiagaraDataInterfaceGPUParamInfo& ParamInfo, FString& OutHLSL) override;
+	virtual void GetParameterDefinitionHLSL(const FNiagaraDataInterfaceGPUParamInfo& ParamInfo, FString& OutHLSL) override;
+	virtual bool GetFunctionHLSL(const FNiagaraDataInterfaceGPUParamInfo& ParamInfo, const FNiagaraDataInterfaceGeneratedFunction& FunctionInfo, int FunctionInstanceIndex, FString& OutHLSL) override;
 	virtual FNiagaraDataInterfaceParametersCS* ConstructComputeParameters() const override;
 
 	virtual bool CanExecuteOnTarget(ENiagaraSimTarget Target)const override { return true; }
diff --git a/Engine/Plugins/FX/HoudiniNiagara/Source/HoudiniNiagara/Private/NiagaraDataInterfaceHoudiniCSV.cpp b/Engine/Plugins/FX/HoudiniNiagara/Source/HoudiniNiagara/Private/NiagaraDataInterfaceHoudiniCSV.cpp
--- a/Engine/Plugins/FX/HoudiniNiagara/Source/HoudiniNiagara/Private/NiagaraDataInterfaceHoudiniCSV.cpp
+++ b/Engine/Plugins/FX/HoudiniNiagara/Source/HoudiniNiagara/Private/NiagaraDataInterfaceHoudiniCSV.cpp
@@ -1336,7 +1336,7 @@ void UNiagaraDataInterfaceHoudiniCSV::GetNumberOfPoints(FVectorVMContext& Contex
 }
 
 // Build the shader function HLSL Code.
-bool UNiagaraDataInterfaceHoudiniCSV::GetFunctionHLSL(const FName& DefinitionFunctionName, FString InstanceFunctionName, FNiagaraDataInterfaceGPUParamInfo& ParamInfo, FString& OutHLSL)
+bool UNiagaraDataInterfaceHoudiniCSV::GetFunctionHLSL(const FNiagaraDataInterfaceGPUParamInfo& ParamInfo, const FNiagaraDataInterfaceGeneratedFunction& FunctionInfo, int FunctionInstanceIndex, FString& OutHLSL)
 {
 	// Build the buffer/variable names declared for this DI
 	FString NumberOfRowsVar = NumberOfRowsBaseName + ParamInfo.DataInterfaceHLSLSymbol;
@@ -1437,10 +1437,10 @@ bool UNiagaraDataInterfaceHoudiniCSV::GetFunctionHLSL(const FName& DefinitionFun
 	};
 
 	// Build each function's HLSL code
-	if (DefinitionFunctionName == GetFloatValueName)
+	if (FunctionInfo.DefinitionName == GetFloatValueName)
 	{
 		// GetFloatValue(int In_Row, int In_Col, out float Out_Value)		
-		OutHLSL += TEXT("void ") + InstanceFunctionName + TEXT("(int In_Row, int In_Col, out float Out_Value) \n{\n");
+		OutHLSL += TEXT("void ") + FunctionInfo.InstanceName + TEXT("(int In_Row, int In_Col, out float Out_Value) \n{\n");
 
 			OutHLSL += ReadFloatInBuffer(TEXT("Out_Value"), TEXT("In_Row"), TEXT("In_Col"));
 
@@ -1448,20 +1448,20 @@ bool UNiagaraDataInterfaceHoudiniCSV::GetFunctionHLSL(const FName& DefinitionFun
 		return true;
 	}
 	/*
-	else if (DefinitionFunctionName == GetFloatValueByStringName)
+	else if (FunctionInfo.DefinitionName == GetFloatValueByStringName)
 	{
 		// GetFloatValueByString(int In_Row, string In_ColTitle, out float Out_Value)
 		FString BufferName = FloatValuesBufferBaseName + ParamInfo.DataInterfaceHLSLSymbol;
-		OutHLSL += TEXT("void ") + InstanceFunctionName + TEXT("(int In_Row, string In_ColTitle, out float Out_Value) \n{\n");
+		OutHLSL += TEXT("void ") + FunctionInfo.InstanceName + TEXT("(int In_Row, string In_ColTitle, out float Out_Value) \n{\n");
 		// NEED CODE
 		OutHLSL += TEXT("\n}\n");
 		return true;
 	}
 	*/
-	else if (DefinitionFunctionName == GetVectorValueName)
+	else if (FunctionInfo.DefinitionName == GetVectorValueName)
 	{
 		// GetVectorValue(int In_Row, int In_Col, out float3 Out_Value)
-		OutHLSL += TEXT("void ") + InstanceFunctionName + TEXT("(int In_Row, int In_Col, out float3 Out_Value) \n{\n");
+		OutHLSL += TEXT("void ") + FunctionInfo.InstanceName + TEXT("(int In_Row, int In_Col, out float3 Out_Value) \n{\n");
 
 			OutHLSL += TEXT("\tint In_DoSwap = 1;\n");
 			OutHLSL += TEXT("\tint In_DoScale = 1;\n");
@@ -1470,20 +1470,20 @@ bool UNiagaraDataInterfaceHoudiniCSV::GetFunctionHLSL(const FName& DefinitionFun
 		OutHLSL += TEXT("\n}\n");
 		return true;
 	}
-	else if (DefinitionFunctionName == GetVectorValueExName)
+	else if (FunctionInfo.DefinitionName == GetVectorValueExName)
 	{
 		// GetVectorValueEx(int In_Row, int In_Col, int In_DoSwap, int In_DoScale, out float3 Out_Value)
-		OutHLSL += TEXT("void ") + InstanceFunctionName + TEXT("(int In_Row, int In_Col, int In_DoSwap, int In_DoScale, out float3 Out_Value) \n{\n");
+		OutHLSL += TEXT("void ") + FunctionInfo.InstanceName + TEXT("(int In_Row, int In_Col, int In_DoSwap, int In_DoScale, out float3 Out_Value) \n{\n");
 		
 			OutHLSL += ReadVectorInBuffer(TEXT("Out_Value"), TEXT("In_Row"), TEXT("In_Col"));
 
 		OutHLSL += TEXT("\n}\n");
 		return true;
 	}
-	else if (DefinitionFunctionName == GetPositionName)
+	else if (FunctionInfo.DefinitionName == GetPositionName)
 	{
 		// GetPosition(int In_Row, out float3 Out_Position)
-		OutHLSL += TEXT("void ") + InstanceFunctionName + TEXT("(int In_Row, out float3 Out_Position) \n{\n");
+		OutHLSL += TEXT("void ") + FunctionInfo.InstanceName + TEXT("(int In_Row, out float3 Out_Position) \n{\n");
 
 			OutHLSL += TEXT("\tint In_DoSwap = 1;\n");
 			OutHLSL += TEXT("\tint In_DoScale = 1;\n");
@@ -1493,10 +1493,10 @@ bool UNiagaraDataInterfaceHoudiniCSV::GetFunctionHLSL(const FName& DefinitionFun
 		OutHLSL += TEXT("\n}\n");
 		return true;
 	}
-	else if (DefinitionFunctionName == GetPositionAndTimeName)
+	else if (FunctionInfo.DefinitionName == GetPositionAndTimeName)
 	{
 		// GetPositionAndTime(int In_Row, out float3 Out_Position, out float Out_Time)
-		OutHLSL += TEXT("void ") + InstanceFunctionName + TEXT("(int In_Row, out float3 Out_Position, out float Out_Time) \n{\n");
+		OutHLSL += TEXT("void ") + FunctionInfo.InstanceName + TEXT("(int In_Row, out float3 Out_Position, out float Out_Time) \n{\n");
 
 			OutHLSL += TEXT("\tint In_DoSwap = 1;\n");
 			OutHLSL += TEXT("\tint In_DoScale = 1;\n");
@@ -1509,10 +1509,10 @@ bool UNiagaraDataInterfaceHoudiniCSV::GetFunctionHLSL(const FName& DefinitionFun
 		OutHLSL += TEXT("\n}\n");
 		return true;
 	}
-	else if (DefinitionFunctionName == GetNormalName)
+	else if (FunctionInfo.DefinitionName == GetNormalName)
 	{
 		// GetNormal(int In_Row, out float3 Out_Value)
-		OutHLSL += TEXT("void ") + InstanceFunctionName + TEXT("(int In_Row, out float3 Out_Normal) \n{\n");
+		OutHLSL += TEXT("void ") + FunctionInfo.InstanceName + TEXT("(int In_Row, out float3 Out_Normal) \n{\n");
 
 			OutHLSL += TEXT("\tint In_DoSwap = 1;\n");
 			OutHLSL += TEXT("\tint In_DoScale = 0;\n");
@@ -1522,10 +1522,10 @@ bool UNiagaraDataInterfaceHoudiniCSV::GetFunctionHLSL(const FName& DefinitionFun
 		OutHLSL += TEXT("\n}\n");
 		return true;
 	}
-	else if (DefinitionFunctionName == GetTimeName)
+	else if (FunctionInfo.DefinitionName == GetTimeName)
 	{
 		// GetTime(int In_Row, out float Out_Value)
-		OutHLSL += TEXT("void ") + InstanceFunctionName + TEXT("(int In_Row, out float Out_Time) \n{\n");
+		OutHLSL += TEXT("void ") + FunctionInfo.InstanceName + TEXT("(int In_Row, out float Out_Time) \n{\n");
 
 			OutHLSL += TEXT("\tint In_TimeCol = ") + GetSpecAttributeColumnIndex(EHoudiniAttributes::TIME) + TEXT(";\n");
 			OutHLSL += ReadFloatInBuffer(TEXT("Out_Time"), TEXT("In_Row"), TEXT("In_TimeCol"));
@@ -1533,10 +1533,10 @@ bool UNiagaraDataInterfaceHoudiniCSV::GetFunctionHLSL(const FName& DefinitionFun
 		OutHLSL += TEXT("\n}\n");
 		return true;
 	}
-	else if (DefinitionFunctionName == GetVelocityName)
+	else if (FunctionInfo.DefinitionName == GetVelocityName)
 	{
 		// GetVelocity(int In_Row, out float3 Out_Value)
-		OutHLSL += TEXT("void ") + InstanceFunctionName + TEXT("(int In_Row, out float3 Out_Velocity) \n{\n");
+		OutHLSL += TEXT("void ") + FunctionInfo.InstanceName + TEXT("(int In_Row, out float3 Out_Velocity) \n{\n");
 
 			OutHLSL += TEXT("\tint In_DoSwap = 1;\n");
 			OutHLSL += TEXT("\tint In_DoScale = 0;\n");
@@ -1546,10 +1546,10 @@ bool UNiagaraDataInterfaceHoudiniCSV::GetFunctionHLSL(const FName& DefinitionFun
 		OutHLSL += TEXT("\n}\n");
 		return true;
 	}
-	else if (DefinitionFunctionName == GetImpulseName)
+	else if (FunctionInfo.DefinitionName == GetImpulseName)
 	{
 		// GetImpulse(int In_Row, out float Out_Value)
-		OutHLSL += TEXT("void ") + InstanceFunctionName + TEXT("(int In_Row, out float Out_Impulse) \n{\n");
+		OutHLSL += TEXT("void ") + FunctionInfo.InstanceName + TEXT("(int In_Row, out float Out_Impulse) \n{\n");
 
 		OutHLSL += TEXT("\tint In_ImpulseCol = ") + GetSpecAttributeColumnIndex(EHoudiniAttributes::IMPULSE) + TEXT(";\n");
 		OutHLSL += ReadFloatInBuffer(TEXT("Out_Impulse"), TEXT("In_Row"), TEXT("In_ImpulseCol"));
@@ -1557,10 +1557,10 @@ bool UNiagaraDataInterfaceHoudiniCSV::GetFunctionHLSL(const FName& DefinitionFun
 		OutHLSL += TEXT("\n}\n");
 		return true;
 	}
-	else if (DefinitionFunctionName == GetColorName)
+	else if (FunctionInfo.DefinitionName == GetColorName)
 	{
 		// GetColor(int In_Row, out float4 Out_Value)
-		OutHLSL += TEXT("void ") + InstanceFunctionName + TEXT("(int In_Row, out float4 Out_Color) \n{\n");
+		OutHLSL += TEXT("void ") + FunctionInfo.InstanceName + TEXT("(int In_Row, out float4 Out_Color) \n{\n");
 
 			OutHLSL += TEXT("\tfloat3 temp_color = float3(1.0, 1.0, 1.0);\n");
 			OutHLSL += TEXT("\tint In_ColorCol = ") + GetSpecAttributeColumnIndex(EHoudiniAttributes::COLOR) + TEXT(";\n");
@@ -1580,34 +1580,34 @@ bool UNiagaraDataInterfaceHoudiniCSV::GetFunctionHLSL(const FName& DefinitionFun
 		OutHLSL += TEXT("\n}\n");
 		return true;
 	}
-	else if (DefinitionFunctionName == GetNumberOfPointsName)
+	else if (FunctionInfo.DefinitionName == GetNumberOfPointsName)
 	{
 		// GetNumberOfPoints(out int Out_Value)
-		OutHLSL += TEXT("void ") + InstanceFunctionName + TEXT("(out int Out_NumPoints) \n{\n");
+		OutHLSL += TEXT("void ") + FunctionInfo.InstanceName + TEXT("(out int Out_NumPoints) \n{\n");
 			OutHLSL += TEXT("\tOut_NumPoints = ") + NumberOfPointsVar + TEXT(";\n");
 		OutHLSL += TEXT("\n}\n");
 		return true;
 	}
-	else if (DefinitionFunctionName == GetNumberOfRowsName)
+	else if (FunctionInfo.DefinitionName == GetNumberOfRowsName)
 	{
 		// GetNumberOfRows(out int Out_Value)
-		OutHLSL += TEXT("void ") + InstanceFunctionName + TEXT("(out int Out_NumRows) \n{\n");
+		OutHLSL += TEXT("void ") + FunctionInfo.InstanceName + TEXT("(out int Out_NumRows) \n{\n");
 			OutHLSL += TEXT("\tOut_NumRows = ") + NumberOfRowsVar + TEXT(";\n");
 		OutHLSL += TEXT("\n}\n");
 		return true;
 	}
-	else if (DefinitionFunctionName == GetNumberOfColumnsName)
+	else if (FunctionInfo.DefinitionName == GetNumberOfColumnsName)
 	{
 		// GetNumberOfColumns(out int Out_Value)
-		OutHLSL += TEXT("void ") + InstanceFunctionName + TEXT("(out int Out_NumColumns) \n{\n");
+		OutHLSL += TEXT("void ") + FunctionInfo.InstanceName + TEXT("(out int Out_NumColumns) \n{\n");
 			OutHLSL += TEXT("\tOut_NumColumns = ") + NumberOfColumnsVar + TEXT(";\n");
 		OutHLSL += TEXT("\n}\n");
 		return true;
 	}
-	else if (DefinitionFunctionName == GetLastRowIndexAtTimeName)
+	else if (FunctionInfo.DefinitionName == GetLastRowIndexAtTimeName)
 	{
 		// GetLastRowIndexAtTime(float In_Time, out int Out_Value)
-		OutHLSL += TEXT("void ") + InstanceFunctionName + TEXT("(float In_Time, out int Out_Value) \n{\n");
+		OutHLSL += TEXT("void ") + FunctionInfo.InstanceName + TEXT("(float In_Time, out int Out_Value) \n{\n");
 
 			OutHLSL += TEXT("\tint In_TimeCol = ") + GetSpecAttributeColumnIndex( EHoudiniAttributes::TIME ) + TEXT(";\n");
 			OutHLSL += TEXT("\tfloat temp_time = 1.0;\n");
@@ -1624,10 +1624,10 @@ bool UNiagaraDataInterfaceHoudiniCSV::GetFunctionHLSL(const FName& DefinitionFun
 		OutHLSL += TEXT("\n}\n");
 		return true;
 	}
-	else if (DefinitionFunctionName == GetPointIDsToSpawnAtTimeName)
+	else if (FunctionInfo.DefinitionName == GetPointIDsToSpawnAtTimeName)
 	{
 		// GetPointIDsToSpawnAtTime(float In_Time, out int Out_MinID, out int Out_MaxID, out int Out_Count)
-		OutHLSL += TEXT("void ") + InstanceFunctionName + TEXT("(float In_Time, out int Out_MinID, out int Out_MaxID, out int Out_Count) \n{\n");
+		OutHLSL += TEXT("void ") + FunctionInfo.InstanceName + TEXT("(float In_Time, out int Out_MinID, out int Out_MaxID, out int Out_Count) \n{\n");
 			
 			OutHLSL += TEXT("\tint time_col = ") + GetSpecAttributeColumnIndex(EHoudiniAttributes::TIME) + TEXT(";\n");
 			OutHLSL += TEXT("\tif( time_col < 0 )\n");
@@ -1665,10 +1665,10 @@ bool UNiagaraDataInterfaceHoudiniCSV::GetFunctionHLSL(const FName& DefinitionFun
 		OutHLSL += TEXT("\n}\n");
 		return true;
 	}
-	else if (DefinitionFunctionName == GetRowIndexesForPointAtTimeName)
+	else if (FunctionInfo.DefinitionName == GetRowIndexesForPointAtTimeName)
 	{
 		// GetRowIndexesForPointAtTime(int In_PointID, float In_Time, out int Out_PreviousRow, out int Out_NextRow, out float Out_Weight)
-		OutHLSL += TEXT("void ") + InstanceFunctionName + TEXT("(int In_PointID, float In_Time, out int Out_PreviousRow, out int Out_NextRow, out float Out_Weight) \n{\n");
+		OutHLSL += TEXT("void ") + FunctionInfo.InstanceName + TEXT("(int In_PointID, float In_Time, out int Out_PreviousRow, out int Out_NextRow, out float Out_Weight) \n{\n");
 
 			OutHLSL += TEXT("\tOut_PreviousRow = 0;\n\tOut_NextRow = 0;\n\tOut_Weight = 0;\n");
 			OutHLSL += GetRowIndexesForPointAtTime(TEXT("In_PointID"), TEXT("In_Time"), TEXT("Out_PreviousRow"), TEXT("Out_NextRow"), TEXT("Out_Weight"));
@@ -1676,10 +1676,10 @@ bool UNiagaraDataInterfaceHoudiniCSV::GetFunctionHLSL(const FName& DefinitionFun
 		OutHLSL += TEXT("\n}\n");
 		return true;
 	}
-	else if (DefinitionFunctionName == GetPointPositionAtTimeName)
+	else if (FunctionInfo.DefinitionName == GetPointPositionAtTimeName)
 	{
 		// GetPointPositionAtTime(int In_PointID, float In_Time, out float3 Out_Value)
-		OutHLSL += TEXT("void ") + InstanceFunctionName + TEXT("(int In_PointID, float In_Time, out float3 Out_Value) \n{\n");
+		OutHLSL += TEXT("void ") + FunctionInfo.InstanceName + TEXT("(int In_PointID, float In_Time, out float3 Out_Value) \n{\n");
 		OutHLSL += TEXT("Out_Value = float3( 0.0, 0.0, 0.0 );\n");
 			OutHLSL += TEXT("\tint pos_col = ") + GetSpecAttributeColumnIndex(EHoudiniAttributes::POSITION) + TEXT(";\n");
 			OutHLSL += TEXT("\tint prev_index = 0;int next_index = 0;float weight = 0.0f;\n");
@@ -1698,10 +1698,10 @@ bool UNiagaraDataInterfaceHoudiniCSV::GetFunctionHLSL(const FName& DefinitionFun
 		OutHLSL += TEXT("\n}\n");
 		return true;
 	}
-	else if (DefinitionFunctionName == GetPointValueAtTimeName)
+	else if (FunctionInfo.DefinitionName == GetPointValueAtTimeName)
 	{
 		// GetPointValueAtTime(int In_PointID, int In_Col, float In_Time, out float Out_Value)
-		OutHLSL += TEXT("void ") + InstanceFunctionName + TEXT("(int In_PointID, int In_Col, float In_Time, out float Out_Value) \n{\n");
+		OutHLSL += TEXT("void ") + FunctionInfo.InstanceName + TEXT("(int In_PointID, int In_Col, float In_Time, out float Out_Value) \n{\n");
 		
 		OutHLSL += TEXT("\tint prev_index = -1;int next_index = -1;float weight = 1.0f;\n");
 		OutHLSL += GetRowIndexesForPointAtTime(TEXT("In_PointID"), TEXT("In_Time"), TEXT("prev_index"), TEXT("next_index"), TEXT("weight"));
@@ -1716,10 +1716,10 @@ bool UNiagaraDataInterfaceHoudiniCSV::GetFunctionHLSL(const FName& DefinitionFun
 		OutHLSL += TEXT("\n}\n");
 		return true;
 	}
-	else if (DefinitionFunctionName == GetPointVectorValueAtTimeName)
+	else if (FunctionInfo.DefinitionName == GetPointVectorValueAtTimeName)
 	{
 		// GetPointVectorValueAtTime(int In_PointID, int In_Col, float In_Time, out float3 Out_Value)
-		OutHLSL += TEXT("void ") + InstanceFunctionName + TEXT("(int In_PointID, int In_Col, float In_Time, out float3 Out_Value) \n{\n");
+		OutHLSL += TEXT("void ") + FunctionInfo.InstanceName + TEXT("(int In_PointID, int In_Col, float In_Time, out float3 Out_Value) \n{\n");
 
 		OutHLSL += TEXT("\tint prev_index = -1;int next_index = -1;float weight = 1.0f;\n");
 		OutHLSL += GetRowIndexesForPointAtTime(TEXT("In_PointID"), TEXT("In_Time"), TEXT("prev_index"), TEXT("next_index"), TEXT("weight"));
@@ -1737,10 +1737,10 @@ bool UNiagaraDataInterfaceHoudiniCSV::GetFunctionHLSL(const FName& DefinitionFun
 		OutHLSL += TEXT("\n}\n");
 		return true;
 	}
-	else if (DefinitionFunctionName == GetPointVectorValueAtTimeExName)
+	else if (FunctionInfo.DefinitionName == GetPointVectorValueAtTimeExName)
 	{
 		// GetPointVectorValueAtTimeEx(int In_PointID, int In_Col, float In_Time, int In_DoSwap, int In_DoScale, out float3 Out_Value)
-		OutHLSL += TEXT("void ") + InstanceFunctionName + TEXT("(int In_PointID, int In_Col, float In_Time, int In_DoSwap, int In_DoScale, out float3 Out_Value) \n{\n");
+		OutHLSL += TEXT("void ") + FunctionInfo.InstanceName + TEXT("(int In_PointID, int In_Col, float In_Time, int In_DoSwap, int In_DoScale, out float3 Out_Value) \n{\n");
 
 			OutHLSL += TEXT("\tint prev_index = -1;int next_index = -1;float weight = 1.0f;\n");
 			OutHLSL += GetRowIndexesForPointAtTime(TEXT("In_PointID"), TEXT("In_Time"), TEXT("prev_index"), TEXT("next_index"), TEXT("weight"));
@@ -1756,10 +1756,10 @@ bool UNiagaraDataInterfaceHoudiniCSV::GetFunctionHLSL(const FName& DefinitionFun
 
 		return true;
 	}
-	else if (DefinitionFunctionName == GetPointLifeName)
+	else if (FunctionInfo.DefinitionName == GetPointLifeName)
 	{
 		// GetPointLife(int In_PointID, out float Out_Value)
-		OutHLSL += TEXT("void ") + InstanceFunctionName + TEXT("(int In_PointID, out float Out_Value) \n{\n");
+		OutHLSL += TEXT("void ") + FunctionInfo.InstanceName + TEXT("(int In_PointID, out float Out_Value) \n{\n");
 
 			OutHLSL += TEXT("\tif ( ( In_PointID < 0 ) || ( In_PointID >= ") + NumberOfPointsVar + TEXT(") )\n");
 				OutHLSL += TEXT("\t\t{Out_Value = -1.0f; return;}\n");
@@ -1769,19 +1769,19 @@ bool UNiagaraDataInterfaceHoudiniCSV::GetFunctionHLSL(const FName& DefinitionFun
 		return true;
 	}
 	/*
-	else if (DefinitionFunctionName == GetPointLifeAtTimeName)
+	else if (FunctionInfo.DefinitionName == GetPointLifeAtTimeName)
 	{
 		// GetPointLifeAtTime(int In_PointID, float In_Time, out float Out_Value)
-		OutHLSL += TEXT("void ") + InstanceFunctionName + TEXT("(int In_PointID, float In_Time, out float Out_Value) \n{\n");
+		OutHLSL += TEXT("void ") + FunctionInfo.InstanceName + TEXT("(int In_PointID, float In_Time, out float Out_Value) \n{\n");
 
 		OutHLSL += TEXT("\n}\n");
 		return true;
 	}
 	*/
-	else if (DefinitionFunctionName == GetPointTypeName)
+	else if (FunctionInfo.DefinitionName == GetPointTypeName)
 	{
 		// GetPointType(int In_PointID, out int Out_Value)
-		OutHLSL += TEXT("void ") + InstanceFunctionName + TEXT("(int In_PointID, out int Out_Value) \n{\n");
+		OutHLSL += TEXT("void ") + FunctionInfo.InstanceName + TEXT("(int In_PointID, out int Out_Value) \n{\n");
 
 			OutHLSL += TEXT("\tif ( ( In_PointID < 0 ) || ( In_PointID >= ") + NumberOfPointsVar+ TEXT(") )\n");
 				OutHLSL += TEXT("\t\t{Out_Value = -1; return;}\n");
@@ -1796,7 +1796,7 @@ bool UNiagaraDataInterfaceHoudiniCSV::GetFunctionHLSL(const FName& DefinitionFun
 
 // Build buffers and member variables HLSL definition
 // Always use the BaseName + the DataInterfaceHLSL indentifier!
-void UNiagaraDataInterfaceHoudiniCSV::GetParameterDefinitionHLSL(FNiagaraDataInterfaceGPUParamInfo& ParamInfo, FString& OutHLSL)
+void UNiagaraDataInterfaceHoudiniCSV::GetParameterDefinitionHLSL(const FNiagaraDataInterfaceGPUParamInfo& ParamInfo, FString& OutHLSL)
 {
 	// int NumberOfRows_XX;
 	FString BufferName = UNiagaraDataInterfaceHoudiniCSV::NumberOfRowsBaseName + ParamInfo.DataInterfaceHLSLSymbol;
diff --git a/Engine/Plugins/FX/Niagara/Source/Niagara/Classes/NiagaraDataInterface.h b/Engine/Plugins/FX/Niagara/Source/Niagara/Classes/NiagaraDataInterface.h
--- a/Engine/Plugins/FX/Niagara/Source/Niagara/Classes/NiagaraDataInterface.h
+++ b/Engine/Plugins/FX/Niagara/Source/Niagara/Classes/NiagaraDataInterface.h
@@ -279,14 +279,13 @@ public:
 	{
 	}
 
-	virtual bool GetFunctionHLSL(const FName& DefinitionFunctionName, FString InstanceFunctionName, FNiagaraDataInterfaceGPUParamInfo& ParamInfo, FString& OutHLSL)
+	virtual void GetParameterDefinitionHLSL(const FNiagaraDataInterfaceGPUParamInfo& ParamInfo, FString& OutHLSL)
 	{
-//		checkf(false, TEXT("Unefined HLSL in data interface. Interfaces need to be able to return HLSL for each function they define in GetFunctions."));
-		return false;
 	}
-	virtual void GetParameterDefinitionHLSL(FNiagaraDataInterfaceGPUParamInfo& ParamInfo, FString& OutHLSL)
+
+	virtual bool GetFunctionHLSL(const FNiagaraDataInterfaceGPUParamInfo& ParamInfo, const FNiagaraDataInterfaceGeneratedFunction& FunctionInfo, int FunctionInstanceIndex, FString& OutHLSL)
 	{
-//		checkf(false, TEXT("Unefined HLSL in data interface. Interfaces need to define HLSL for uniforms their functions access."));
+		return false;
 	}
 
 	virtual void PostExecute() {}
diff --git a/Engine/Plugins/FX/Niagara/Source/Niagara/Classes/NiagaraDataInterfaceCamera.h b/Engine/Plugins/FX/Niagara/Source/Niagara/Classes/NiagaraDataInterfaceCamera.h
--- a/Engine/Plugins/FX/Niagara/Source/Niagara/Classes/NiagaraDataInterfaceCamera.h
+++ b/Engine/Plugins/FX/Niagara/Source/Niagara/Classes/NiagaraDataInterfaceCamera.h
@@ -34,7 +34,7 @@ public:
 	virtual bool InitPerInstanceData(void* PerInstanceData, FNiagaraSystemInstance* SystemInstance) override;
 	virtual int32 PerInstanceDataSize() const override { return sizeof(CameraDataInterface_InstanceData); }
 	virtual bool PerInstanceTick(void* PerInstanceData, FNiagaraSystemInstance* SystemInstance, float DeltaSeconds) override;
-	virtual bool GetFunctionHLSL(const FName&  DefinitionFunctionName, FString InstanceFunctionName, FNiagaraDataInterfaceGPUParamInfo& ParamInfo, FString& OutHLSL) override;
+	virtual bool GetFunctionHLSL(const FNiagaraDataInterfaceGPUParamInfo& ParamInfo, const FNiagaraDataInterfaceGeneratedFunction& FunctionInfo, int FunctionInstanceIndex, FString& OutHLSL) override;
 	virtual bool CanExecuteOnTarget(ENiagaraSimTarget Target) const override { return true; }
 	virtual bool HasTickGroupPrereqs() const override { return true; }
 	virtual ETickingGroup CalculateTickGroup(void* PerInstanceData) const override;
diff --git a/Engine/Plugins/FX/Niagara/Source/Niagara/Classes/NiagaraDataInterfaceCollisionQuery.h b/Engine/Plugins/FX/Niagara/Source/Niagara/Classes/NiagaraDataInterfaceCollisionQuery.h
--- a/Engine/Plugins/FX/Niagara/Source/Niagara/Classes/NiagaraDataInterfaceCollisionQuery.h
+++ b/Engine/Plugins/FX/Niagara/Source/Niagara/Classes/NiagaraDataInterfaceCollisionQuery.h
@@ -64,8 +64,8 @@ public:
 	virtual bool RequiresDistanceFieldData() const override { return true; }
 	virtual bool RequiresDepthBuffer() const override { return true; }
 
-	virtual bool GetFunctionHLSL(const FName&  DefinitionFunctionName, FString InstanceFunctionName, FNiagaraDataInterfaceGPUParamInfo& ParamInfo, FString& OutHLSL) override;
-	virtual void GetParameterDefinitionHLSL(FNiagaraDataInterfaceGPUParamInfo& ParamInfo, FString& OutHLSL) override;
+	virtual void GetParameterDefinitionHLSL(const FNiagaraDataInterfaceGPUParamInfo& ParamInfo, FString& OutHLSL) override;
+	virtual bool GetFunctionHLSL(const FNiagaraDataInterfaceGPUParamInfo& ParamInfo, const FNiagaraDataInterfaceGeneratedFunction& FunctionInfo, int FunctionInstanceIndex, FString& OutHLSL) override;
 	virtual FNiagaraDataInterfaceParametersCS* ConstructComputeParameters() const override;
 	
 private:
diff --git a/Engine/Plugins/FX/Niagara/Source/Niagara/Classes/NiagaraDataInterfaceColorCurve.h b/Engine/Plugins/FX/Niagara/Source/Niagara/Classes/NiagaraDataInterfaceColorCurve.h
--- a/Engine/Plugins/FX/Niagara/Source/Niagara/Classes/NiagaraDataInterfaceColorCurve.h
+++ b/Engine/Plugins/FX/Niagara/Source/Niagara/Classes/NiagaraDataInterfaceColorCurve.h
@@ -60,7 +60,7 @@ public:
 	//~ UNiagaraDataInterfaceCurveBase interface
 	virtual void GetCurveData(TArray<FCurveData>& OutCurveData) override;
 
-	virtual bool GetFunctionHLSL(const FName&  DefinitionFunctionName, FString InstanceFunctionName, FNiagaraDataInterfaceGPUParamInfo& ParamInfo, FString& OutHLSL) override;
+	virtual bool GetFunctionHLSL(const FNiagaraDataInterfaceGPUParamInfo& ParamInfo, const FNiagaraDataInterfaceGeneratedFunction& FunctionInfo, int FunctionInstanceIndex, FString& OutHLSL) override;
 
 	virtual int32 GetCurveNumElems()const { return CurveLUTNumElems; }
 protected:
diff --git a/Engine/Plugins/FX/Niagara/Source/Niagara/Classes/NiagaraDataInterfaceCurlNoise.h b/Engine/Plugins/FX/Niagara/Source/Niagara/Classes/NiagaraDataInterfaceCurlNoise.h
--- a/Engine/Plugins/FX/Niagara/Source/Niagara/Classes/NiagaraDataInterfaceCurlNoise.h
+++ b/Engine/Plugins/FX/Niagara/Source/Niagara/Classes/NiagaraDataInterfaceCurlNoise.h
@@ -39,8 +39,8 @@ public:
 	void SampleNoiseField(FVectorVMContext& Context);
 
 	// GPU sim functionality
-	virtual bool GetFunctionHLSL(const FName&  DefinitionFunctionName, FString InstanceFunctionName, FNiagaraDataInterfaceGPUParamInfo& ParamInfo, FString& OutHLSL) override;
-	virtual void GetParameterDefinitionHLSL(FNiagaraDataInterfaceGPUParamInfo& ParamInfo, FString& OutHLSL) override;
+	virtual void GetParameterDefinitionHLSL(const FNiagaraDataInterfaceGPUParamInfo& ParamInfo, FString& OutHLSL) override;
+	virtual bool GetFunctionHLSL(const FNiagaraDataInterfaceGPUParamInfo& ParamInfo, const FNiagaraDataInterfaceGeneratedFunction& FunctionInfo, int FunctionInstanceIndex, FString& OutHLSL) override;
 	virtual FNiagaraDataInterfaceParametersCS* ConstructComputeParameters() const override;
 
 protected:
diff --git a/Engine/Plugins/FX/Niagara/Source/Niagara/Classes/NiagaraDataInterfaceCurve.h b/Engine/Plugins/FX/Niagara/Source/Niagara/Classes/NiagaraDataInterfaceCurve.h
--- a/Engine/Plugins/FX/Niagara/Source/Niagara/Classes/NiagaraDataInterfaceCurve.h
+++ b/Engine/Plugins/FX/Niagara/Source/Niagara/Classes/NiagaraDataInterfaceCurve.h
@@ -44,7 +44,7 @@ public:
 	//~ UNiagaraDataInterfaceCurveBase interface
 	virtual void GetCurveData(TArray<FCurveData>& OutCurveData) override;
 
-	virtual bool GetFunctionHLSL(const FName&  DefinitionFunctionName, FString InstanceFunctionName, FNiagaraDataInterfaceGPUParamInfo& ParamInfo, FString& OutHLSL) override;
+	virtual bool GetFunctionHLSL(const FNiagaraDataInterfaceGPUParamInfo& ParamInfo, const FNiagaraDataInterfaceGeneratedFunction& FunctionInfo, int FunctionInstanceIndex, FString& OutHLSL) override;
 
 	virtual int32 GetCurveNumElems()const { return CurveLUTNumElems; }
 protected:
diff --git a/Engine/Plugins/FX/Niagara/Source/Niagara/Classes/NiagaraDataInterfaceCurveBase.h b/Engine/Plugins/FX/Niagara/Source/Niagara/Classes/NiagaraDataInterfaceCurveBase.h
--- a/Engine/Plugins/FX/Niagara/Source/Niagara/Classes/NiagaraDataInterfaceCurveBase.h
+++ b/Engine/Plugins/FX/Niagara/Source/Niagara/Classes/NiagaraDataInterfaceCurveBase.h
@@ -151,7 +151,7 @@ public:
 	/** Gets information for all of the curves owned by this curve data interface. */
 	virtual void GetCurveData(TArray<FCurveData>& OutCurveData) { }
 
-	virtual void GetParameterDefinitionHLSL(FNiagaraDataInterfaceGPUParamInfo& ParamInfo, FString& OutHLSL) override;
+	virtual void GetParameterDefinitionHLSL(const FNiagaraDataInterfaceGPUParamInfo& ParamInfo, FString& OutHLSL) override;
 	virtual FNiagaraDataInterfaceParametersCS* ConstructComputeParameters() const override;
 
 	void SetDefaultLUT();
diff --git a/Engine/Plugins/FX/Niagara/Source/Niagara/Classes/NiagaraDataInterfaceExport.h b/Engine/Plugins/FX/Niagara/Source/Niagara/Classes/NiagaraDataInterfaceExport.h
--- a/Engine/Plugins/FX/Niagara/Source/Niagara/Classes/NiagaraDataInterfaceExport.h
+++ b/Engine/Plugins/FX/Niagara/Source/Niagara/Classes/NiagaraDataInterfaceExport.h
@@ -73,7 +73,7 @@ public:
 	virtual bool PerInstanceTick(void* PerInstanceData, FNiagaraSystemInstance* SystemInstance, float DeltaSeconds) override;
 	virtual bool PerInstanceTickPostSimulate(void* PerInstanceData, FNiagaraSystemInstance* SystemInstance, float DeltaSeconds) override;
 	virtual bool Equals(const UNiagaraDataInterface* Other) const override;
-	virtual bool GetFunctionHLSL(const FName&  DefinitionFunctionName, FString InstanceFunctionName, FNiagaraDataInterfaceGPUParamInfo& ParamInfo, FString& OutHLSL) override;
+	virtual bool GetFunctionHLSL(const FNiagaraDataInterfaceGPUParamInfo& ParamInfo, const FNiagaraDataInterfaceGeneratedFunction& FunctionInfo, int FunctionInstanceIndex, FString& OutHLSL) override;
 	virtual bool CanExecuteOnTarget(ENiagaraSimTarget Target) const override { return Target == ENiagaraSimTarget::CPUSim; }
 	//UNiagaraDataInterface Interface
 
diff --git a/Engine/Plugins/FX/Niagara/Source/Niagara/Classes/NiagaraDataInterfaceGrid2DCollection.h b/Engine/Plugins/FX/Niagara/Source/Niagara/Classes/NiagaraDataInterfaceGrid2DCollection.h
--- a/Engine/Plugins/FX/Niagara/Source/Niagara/Classes/NiagaraDataInterfaceGrid2DCollection.h
+++ b/Engine/Plugins/FX/Niagara/Source/Niagara/Classes/NiagaraDataInterfaceGrid2DCollection.h
@@ -120,8 +120,8 @@ public:
 	virtual bool Equals(const UNiagaraDataInterface* Other) const override;
 
 	// GPU sim functionality
-	virtual void GetParameterDefinitionHLSL(FNiagaraDataInterfaceGPUParamInfo& ParamInfo, FString& OutHLSL) override;
-	virtual bool GetFunctionHLSL(const FName&  DefinitionFunctionName, FString InstanceFunctionName, FNiagaraDataInterfaceGPUParamInfo& ParamInfo, FString& OutHLSL) override;
+	virtual void GetParameterDefinitionHLSL(const FNiagaraDataInterfaceGPUParamInfo& ParamInfo, FString& OutHLSL) override;
+	virtual bool GetFunctionHLSL(const FNiagaraDataInterfaceGPUParamInfo& ParamInfo, const FNiagaraDataInterfaceGeneratedFunction& FunctionInfo, int FunctionInstanceIndex, FString& OutHLSL) override;
 	virtual FNiagaraDataInterfaceParametersCS* ConstructComputeParameters() const override;
 	
 	virtual void ProvidePerInstanceDataForRenderThread(void* DataForRenderThread, void* PerInstanceData, const FNiagaraSystemInstanceID& SystemInstance) override {}
diff --git a/Engine/Plugins/FX/Niagara/Source/Niagara/Classes/NiagaraDataInterfaceNeighborGrid3D.h b/Engine/Plugins/FX/Niagara/Source/Niagara/Classes/NiagaraDataInterfaceNeighborGrid3D.h
--- a/Engine/Plugins/FX/Niagara/Source/Niagara/Classes/NiagaraDataInterfaceNeighborGrid3D.h
+++ b/Engine/Plugins/FX/Niagara/Source/Niagara/Classes/NiagaraDataInterfaceNeighborGrid3D.h
@@ -77,8 +77,8 @@ public:
 	virtual bool Equals(const UNiagaraDataInterface* Other) const override;
 
 	// GPU sim functionality
-	virtual void GetParameterDefinitionHLSL(FNiagaraDataInterfaceGPUParamInfo& ParamInfo, FString& OutHLSL) override;
-	virtual bool GetFunctionHLSL(const FName&  DefinitionFunctionName, FString InstanceFunctionName, FNiagaraDataInterfaceGPUParamInfo& ParamInfo, FString& OutHLSL) override;
+	virtual void GetParameterDefinitionHLSL(const FNiagaraDataInterfaceGPUParamInfo& ParamInfo, FString& OutHLSL) override;
+	virtual bool GetFunctionHLSL(const FNiagaraDataInterfaceGPUParamInfo& ParamInfo, const FNiagaraDataInterfaceGeneratedFunction& FunctionInfo, int FunctionInstanceIndex, FString& OutHLSL) override;
 	virtual FNiagaraDataInterfaceParametersCS* ConstructComputeParameters() const override;
 		
 	virtual void ProvidePerInstanceDataForRenderThread(void* DataForRenderThread, void* PerInstanceData, const FNiagaraSystemInstanceID& SystemInstance) override {}
diff --git a/Engine/Plugins/FX/Niagara/Source/Niagara/Classes/NiagaraDataInterfaceOcclusion.h b/Engine/Plugins/FX/Niagara/Source/Niagara/Classes/NiagaraDataInterfaceOcclusion.h
--- a/Engine/Plugins/FX/Niagara/Source/Niagara/Classes/NiagaraDataInterfaceOcclusion.h
+++ b/Engine/Plugins/FX/Niagara/Source/Niagara/Classes/NiagaraDataInterfaceOcclusion.h
@@ -19,7 +19,7 @@ public:
 	//UNiagaraDataInterface Interface
 	virtual void GetFunctions(TArray<FNiagaraFunctionSignature>& OutFunctions)override;
 	virtual void GetVMExternalFunction(const FVMExternalFunctionBindingInfo& BindingInfo, void* InstanceData, FVMExternalFunction &OutFunc) override;
-	virtual bool GetFunctionHLSL(const FName&  DefinitionFunctionName, FString InstanceFunctionName, FNiagaraDataInterfaceGPUParamInfo& ParamInfo, FString& OutHLSL) override;
+	virtual bool GetFunctionHLSL(const FNiagaraDataInterfaceGPUParamInfo& ParamInfo, const FNiagaraDataInterfaceGeneratedFunction& FunctionInfo, int FunctionInstanceIndex, FString& OutHLSL) override;
 	virtual bool CanExecuteOnTarget(ENiagaraSimTarget Target) const override { return Target == ENiagaraSimTarget::GPUComputeSim; }
 	virtual FNiagaraDataInterfaceParametersCS* ConstructComputeParameters() const override;
 	virtual bool RequiresDepthBuffer() const override { return true; }
diff --git a/Engine/Plugins/FX/Niagara/Source/Niagara/Classes/NiagaraDataInterfaceParticleRead.h b/Engine/Plugins/FX/Niagara/Source/Niagara/Classes/NiagaraDataInterfaceParticleRead.h
--- a/Engine/Plugins/FX/Niagara/Source/Niagara/Classes/NiagaraDataInterfaceParticleRead.h
+++ b/Engine/Plugins/FX/Niagara/Source/Niagara/Classes/NiagaraDataInterfaceParticleRead.h
@@ -40,14 +40,15 @@ public:
 
 	virtual void GetFunctions(TArray<FNiagaraFunctionSignature>& OutFunctions) override;
 	virtual void GetVMExternalFunction(const FVMExternalFunctionBindingInfo& BindingInfo, void* InstanceData, FVMExternalFunction &OutFunc) override;
-	virtual bool CanExecuteOnTarget(ENiagaraSimTarget Target)const override { return Target == ENiagaraSimTarget::CPUSim; }
+	virtual bool CanExecuteOnTarget(ENiagaraSimTarget Target) const override { return true; }
 	virtual bool Equals(const UNiagaraDataInterface* Other) const override;
 
 	// GPU sim functionality
-	virtual bool GetFunctionHLSL(const FName& DefinitionFunctionName, FString InstanceFunctionName, FNiagaraDataInterfaceGPUParamInfo& ParamInfo, FString& OutHLSL) override;
-	virtual void GetParameterDefinitionHLSL(FNiagaraDataInterfaceGPUParamInfo& ParamInfo, FString& OutHLSL) override;
+	virtual void GetCommonHLSL(FString& OutHLSL) override;
+	virtual void GetParameterDefinitionHLSL(const FNiagaraDataInterfaceGPUParamInfo& ParamInfo, FString& OutHLSL) override;
+	virtual bool GetFunctionHLSL(const FNiagaraDataInterfaceGPUParamInfo& ParamInfo, const FNiagaraDataInterfaceGeneratedFunction& FunctionInfo, int FunctionInstanceIndex, FString& OutHLSL) override;
 	virtual FNiagaraDataInterfaceParametersCS* ConstructComputeParameters() const override;
-
+	virtual void ProvidePerInstanceDataForRenderThread(void* DataForRenderThread, void* PerInstanceData, const FNiagaraSystemInstanceID& SystemInstance) override;
 	//UNiagaraDataInterface Interface End
 
 	virtual void ReadFloat(FVectorVMContext& Context, FName AttributeToRead);
diff --git a/Engine/Plugins/FX/Niagara/Source/Niagara/Classes/NiagaraDataInterfaceRW.h b/Engine/Plugins/FX/Niagara/Source/Niagara/Classes/NiagaraDataInterfaceRW.h
--- a/Engine/Plugins/FX/Niagara/Source/Niagara/Classes/NiagaraDataInterfaceRW.h
+++ b/Engine/Plugins/FX/Niagara/Source/Niagara/Classes/NiagaraDataInterfaceRW.h
@@ -151,8 +151,8 @@ public:
 	virtual bool Equals(const UNiagaraDataInterface* Other) const override;
 
 	// GPU sim functionality
-	virtual void GetParameterDefinitionHLSL(FNiagaraDataInterfaceGPUParamInfo& ParamInfo, FString& OutHLSL) override;
-	virtual bool GetFunctionHLSL(const FName&  DefinitionFunctionName, FString InstanceFunctionName, FNiagaraDataInterfaceGPUParamInfo& ParamInfo, FString& OutHLSL) override;
+	virtual void GetParameterDefinitionHLSL(const FNiagaraDataInterfaceGPUParamInfo& ParamInfo, FString& OutHLSL) override;
+	virtual bool GetFunctionHLSL(const FNiagaraDataInterfaceGPUParamInfo& ParamInfo, const FNiagaraDataInterfaceGeneratedFunction& FunctionInfo, int FunctionInstanceIndex, FString& OutHLSL) override;
 	virtual FNiagaraDataInterfaceParametersCS* ConstructComputeParameters() const override;
 	//~ UNiagaraDataInterface interface END
 
@@ -200,8 +200,8 @@ public:
 	virtual bool Equals(const UNiagaraDataInterface* Other) const override;
 
 	// GPU sim functionality
-	virtual void GetParameterDefinitionHLSL(FNiagaraDataInterfaceGPUParamInfo& ParamInfo, FString& OutHLSL) override;
-	virtual bool GetFunctionHLSL(const FName&  DefinitionFunctionName, FString InstanceFunctionName, FNiagaraDataInterfaceGPUParamInfo& ParamInfo, FString& OutHLSL) override;
+	virtual void GetParameterDefinitionHLSL(const FNiagaraDataInterfaceGPUParamInfo& ParamInfo, FString& OutHLSL) override;
+	virtual bool GetFunctionHLSL(const FNiagaraDataInterfaceGPUParamInfo& ParamInfo, const FNiagaraDataInterfaceGeneratedFunction& FunctionInfo, int FunctionInstanceIndex, FString& OutHLSL) override;
 	virtual FNiagaraDataInterfaceParametersCS* ConstructComputeParameters() const override;
 	//~ UNiagaraDataInterface interface END
 
diff --git a/Engine/Plugins/FX/Niagara/Source/Niagara/Classes/NiagaraDataInterfaceSimpleCounter.h b/Engine/Plugins/FX/Niagara/Source/Niagara/Classes/NiagaraDataInterfaceSimpleCounter.h
--- a/Engine/Plugins/FX/Niagara/Source/Niagara/Classes/NiagaraDataInterfaceSimpleCounter.h
+++ b/Engine/Plugins/FX/Niagara/Source/Niagara/Classes/NiagaraDataInterfaceSimpleCounter.h
@@ -57,9 +57,9 @@ public:
 	virtual bool Equals(const UNiagaraDataInterface* Other) const override;
 	virtual bool CanExecuteOnTarget(ENiagaraSimTarget Target) const override { return true; }
 
-	virtual bool GetFunctionHLSL(const FName&  DefinitionFunctionName, FString InstanceFunctionName, FNiagaraDataInterfaceGPUParamInfo& ParamInfo, FString& OutHLSL) override;
-	virtual void GetParameterDefinitionHLSL(FNiagaraDataInterfaceGPUParamInfo& ParamInfo, FString& OutHLSL) override;
-	
+	virtual void GetParameterDefinitionHLSL(const FNiagaraDataInterfaceGPUParamInfo& ParamInfo, FString& OutHLSL) override;
+	virtual bool GetFunctionHLSL(const FNiagaraDataInterfaceGPUParamInfo& ParamInfo, const FNiagaraDataInterfaceGeneratedFunction& FunctionInfo, int FunctionInstanceIndex, FString& OutHLSL) override;
+
 	virtual void ProvidePerInstanceDataForRenderThread(void* DataForRenderThread, void* PerInstanceData, const FNiagaraSystemInstanceID& SystemInstance) override
 	{
 		check(false);
diff --git a/Engine/Plugins/FX/Niagara/Source/Niagara/Classes/NiagaraDataInterfaceSkeletalMesh.h b/Engine/Plugins/FX/Niagara/Source/Niagara/Classes/NiagaraDataInterfaceSkeletalMesh.h
--- a/Engine/Plugins/FX/Niagara/Source/Niagara/Classes/NiagaraDataInterfaceSkeletalMesh.h
+++ b/Engine/Plugins/FX/Niagara/Source/Niagara/Classes/NiagaraDataInterfaceSkeletalMesh.h
@@ -534,8 +534,8 @@ public:
 	USkeletalMesh* GetSkeletalMesh(class UNiagaraComponent* OwningComponent, TWeakObjectPtr<USceneComponent>& SceneComponent, USkeletalMeshComponent*& FoundSkelComp, FNDISkeletalMesh_InstanceData* InstData = nullptr);
 
 	virtual void GetCommonHLSL(FString& OutHLSL) override;
-	virtual bool GetFunctionHLSL(const FName&  DefinitionFunctionName, FString InstanceFunctionName, FNiagaraDataInterfaceGPUParamInfo& ParamInfo, FString& OutHLSL) override;
-	virtual void GetParameterDefinitionHLSL(FNiagaraDataInterfaceGPUParamInfo& ParamInfo, FString& OutHLSL) override;
+	virtual void GetParameterDefinitionHLSL(const FNiagaraDataInterfaceGPUParamInfo& ParamInfo, FString& OutHLSL) override;
+	virtual bool GetFunctionHLSL(const FNiagaraDataInterfaceGPUParamInfo& ParamInfo, const FNiagaraDataInterfaceGeneratedFunction& FunctionInfo, int FunctionInstanceIndex, FString& OutHLSL) override;
 	virtual FNiagaraDataInterfaceParametersCS* ConstructComputeParameters() const override;
 
 	virtual void ProvidePerInstanceDataForRenderThread(void* DataForRenderThread, void* PerInstanceData, const FNiagaraSystemInstanceID& SystemInstance) override;
diff --git a/Engine/Plugins/FX/Niagara/Source/Niagara/Classes/NiagaraDataInterfaceStaticMesh.h b/Engine/Plugins/FX/Niagara/Source/Niagara/Classes/NiagaraDataInterfaceStaticMesh.h
--- a/Engine/Plugins/FX/Niagara/Source/Niagara/Classes/NiagaraDataInterfaceStaticMesh.h
+++ b/Engine/Plugins/FX/Niagara/Source/Niagara/Classes/NiagaraDataInterfaceStaticMesh.h
@@ -207,8 +207,8 @@ public:
 	virtual TArray<FNiagaraDataInterfaceError> GetErrors() override;
 #endif
 
-	virtual bool GetFunctionHLSL(const FName&  DefinitionFunctionName, FString InstanceFunctionName, FNiagaraDataInterfaceGPUParamInfo& ParamInfo, FString& OutHLSL) override;
-	virtual void GetParameterDefinitionHLSL(FNiagaraDataInterfaceGPUParamInfo& ParamInfo, FString& OutHLSL) override;
+	virtual void GetParameterDefinitionHLSL(const FNiagaraDataInterfaceGPUParamInfo& ParamInfo, FString& OutHLSL) override;
+	virtual bool GetFunctionHLSL(const FNiagaraDataInterfaceGPUParamInfo& ParamInfo, const FNiagaraDataInterfaceGeneratedFunction& FunctionInfo, int FunctionInstanceIndex, FString& OutHLSL) override;
 	virtual FNiagaraDataInterfaceParametersCS* ConstructComputeParameters()const override;
 
 	virtual void ProvidePerInstanceDataForRenderThread(void* DataForRenderThread, void* PerInstanceData, const FNiagaraSystemInstanceID& SystemInstance) override;
diff --git a/Engine/Plugins/FX/Niagara/Source/Niagara/Classes/NiagaraDataInterfaceTexture.h b/Engine/Plugins/FX/Niagara/Source/Niagara/Classes/NiagaraDataInterfaceTexture.h
--- a/Engine/Plugins/FX/Niagara/Source/Niagara/Classes/NiagaraDataInterfaceTexture.h
+++ b/Engine/Plugins/FX/Niagara/Source/Niagara/Classes/NiagaraDataInterfaceTexture.h
@@ -39,8 +39,8 @@ public:
 	virtual bool Equals(const UNiagaraDataInterface* Other) const override;
 
 	// GPU sim functionality
-	virtual bool GetFunctionHLSL(const FName&  DefinitionFunctionName, FString InstanceFunctionName, FNiagaraDataInterfaceGPUParamInfo& ParamInfo, FString& OutHLSL) override;
-	virtual void GetParameterDefinitionHLSL(FNiagaraDataInterfaceGPUParamInfo& ParamInfo, FString& OutHLSL) override;
+	virtual void GetParameterDefinitionHLSL(const FNiagaraDataInterfaceGPUParamInfo& ParamInfo, FString& OutHLSL) override;
+	virtual bool GetFunctionHLSL(const FNiagaraDataInterfaceGPUParamInfo& ParamInfo, const FNiagaraDataInterfaceGeneratedFunction& FunctionInfo, int FunctionInstanceIndex, FString& OutHLSL) override;
 	virtual FNiagaraDataInterfaceParametersCS* ConstructComputeParameters()const override;
 
 	//FRWBuffer& GetGPUBuffer();
diff --git a/Engine/Plugins/FX/Niagara/Source/Niagara/Classes/NiagaraDataInterfaceVector2DCurve.h b/Engine/Plugins/FX/Niagara/Source/Niagara/Classes/NiagaraDataInterfaceVector2DCurve.h
--- a/Engine/Plugins/FX/Niagara/Source/Niagara/Classes/NiagaraDataInterfaceVector2DCurve.h
+++ b/Engine/Plugins/FX/Niagara/Source/Niagara/Classes/NiagaraDataInterfaceVector2DCurve.h
@@ -49,7 +49,7 @@ public:
 	//~ UNiagaraDataInterfaceCurveBase interface
 	virtual void GetCurveData(TArray<FCurveData>& OutCurveData) override;
 
-	virtual bool GetFunctionHLSL(const FName&  DefinitionFunctionName, FString InstanceFunctionName, FNiagaraDataInterfaceGPUParamInfo& ParamInfo, FString& OutHLSL) override;
+	virtual bool GetFunctionHLSL(const FNiagaraDataInterfaceGPUParamInfo& ParamInfo, const FNiagaraDataInterfaceGeneratedFunction& FunctionInfo, int FunctionInstanceIndex, FString& OutHLSL) override;
 
 	virtual int32 GetCurveNumElems()const { return CurveLUTNumElems; }
 protected:
diff --git a/Engine/Plugins/FX/Niagara/Source/Niagara/Classes/NiagaraDataInterfaceVector4Curve.h b/Engine/Plugins/FX/Niagara/Source/Niagara/Classes/NiagaraDataInterfaceVector4Curve.h
--- a/Engine/Plugins/FX/Niagara/Source/Niagara/Classes/NiagaraDataInterfaceVector4Curve.h
+++ b/Engine/Plugins/FX/Niagara/Source/Niagara/Classes/NiagaraDataInterfaceVector4Curve.h
@@ -57,7 +57,7 @@ public:
 
 	//~ UNiagaraDataInterfaceCurveBase interface
 	virtual void GetCurveData(TArray<FCurveData>& OutCurveData) override;
-	virtual bool GetFunctionHLSL (const FName& DefinitionFunctionName, FString InstanceFunctionName, FNiagaraDataInterfaceGPUParamInfo& ParamInfo, FString& OutHLSL) override;
+	virtual bool GetFunctionHLSL(const FNiagaraDataInterfaceGPUParamInfo& ParamInfo, const FNiagaraDataInterfaceGeneratedFunction& FunctionInfo, int FunctionInstanceIndex, FString& OutHLSL) override;
 
 	virtual int32 GetCurveNumElems()const { return CurveLUTNumElems; }
 protected:
diff --git a/Engine/Plugins/FX/Niagara/Source/Niagara/Classes/NiagaraDataInterfaceVectorCurve.h b/Engine/Plugins/FX/Niagara/Source/Niagara/Classes/NiagaraDataInterfaceVectorCurve.h
--- a/Engine/Plugins/FX/Niagara/Source/Niagara/Classes/NiagaraDataInterfaceVectorCurve.h
+++ b/Engine/Plugins/FX/Niagara/Source/Niagara/Classes/NiagaraDataInterfaceVectorCurve.h
@@ -57,7 +57,7 @@ public:
 	
 	virtual int32 GetCurveNumElems()const { return CurveLUTNumElems; }
 
-	virtual bool GetFunctionHLSL(const FName&  DefinitionFunctionName, FString InstanceFunctionName, FNiagaraDataInterfaceGPUParamInfo& ParamInfo, FString& OutHLSL) override;
+	virtual bool GetFunctionHLSL(const FNiagaraDataInterfaceGPUParamInfo& ParamInfo, const FNiagaraDataInterfaceGeneratedFunction& FunctionInfo, int FunctionInstanceIndex, FString& OutHLSL) override;
 protected:
 	virtual bool CopyToInternal(UNiagaraDataInterface* Destination) const override;
 
diff --git a/Engine/Plugins/FX/Niagara/Source/Niagara/Classes/NiagaraDataInterfaceVectorField.h b/Engine/Plugins/FX/Niagara/Source/Niagara/Classes/NiagaraDataInterfaceVectorField.h
--- a/Engine/Plugins/FX/Niagara/Source/Niagara/Classes/NiagaraDataInterfaceVectorField.h
+++ b/Engine/Plugins/FX/Niagara/Source/Niagara/Classes/NiagaraDataInterfaceVectorField.h
@@ -48,8 +48,8 @@ public:
 #endif
 
 	// GPU sim functionality
-	virtual void GetParameterDefinitionHLSL(FNiagaraDataInterfaceGPUParamInfo& ParamInfo, FString& OutHLSL) override; 
-	virtual bool GetFunctionHLSL(const FName&  DefinitionFunctionName, FString InstanceFunctionName, FNiagaraDataInterfaceGPUParamInfo& ParamInfo, FString& OutHLSL) override;
+	virtual void GetParameterDefinitionHLSL(const FNiagaraDataInterfaceGPUParamInfo& ParamInfo, FString& OutHLSL) override;
+	virtual bool GetFunctionHLSL(const FNiagaraDataInterfaceGPUParamInfo& ParamInfo, const FNiagaraDataInterfaceGeneratedFunction& FunctionInfo, int FunctionInstanceIndex, FString& OutHLSL) override;
 	virtual FNiagaraDataInterfaceParametersCS* ConstructComputeParameters() const override;
 	//~ UNiagaraDataInterface interface END
 
diff --git a/Engine/Plugins/FX/Niagara/Source/Niagara/Classes/NiagaraDataSet.h b/Engine/Plugins/FX/Niagara/Source/Niagara/Classes/NiagaraDataSet.h
--- a/Engine/Plugins/FX/Niagara/Source/Niagara/Classes/NiagaraDataSet.h
+++ b/Engine/Plugins/FX/Niagara/Source/Niagara/Classes/NiagaraDataSet.h
@@ -103,7 +103,7 @@ protected:
 public:
 	FNiagaraDataBuffer(FNiagaraDataSet* InOwner);
 	void Allocate(uint32 NumInstances, bool bMaintainExisting = false);
-	void AllocateGPU(uint32 InNumInstances, FNiagaraGPUInstanceCountManager& GPUInstanceCountManager, FRHICommandList& RHICmdList, ERHIFeatureLevel::Type FeatureLevel);
+	void AllocateGPU(uint32 InNumInstances, FNiagaraGPUInstanceCountManager& GPUInstanceCountManager, FRHICommandList& RHICmdList, ERHIFeatureLevel::Type FeatureLevel, const TCHAR* DebugSimName);
 	void SwapInstances(uint32 OldIndex, uint32 NewIndex);
 	void KillInstance(uint32 InstanceIdx);
 	void CopyTo(FNiagaraDataBuffer& DestBuffer, int32 SrcStartIdx, int32 DestStartIdx, int32 NumInstances)const;
@@ -139,8 +139,6 @@ public:
 	FORCEINLINE FRWBuffer& GetGPUBufferInt() { return GPUBufferInt;	}
 	FORCEINLINE uint32 GetGPUInstanceCountBufferOffset() const { return GPUInstanceCountBufferOffset; }
 	FORCEINLINE void ClearGPUInstanceCountBufferOffset() { GPUInstanceCountBufferOffset = INDEX_NONE; }
-	FORCEINLINE uint32 GetGPUNumAllocatedIDs() const { return NumIDsAllocatedForGPU; }
-	FORCEINLINE FRWBuffer& GetGPUFreeIDs() { return GPUFreeIDs; }
 	FORCEINLINE FRWBuffer& GetGPUIDToIndexTable() { return GPUIDToIndexTable; }
 
 	FORCEINLINE int32 GetSafeComponentBufferSize() const { return GetSafeComponentBufferSize(GetNumInstancesAllocated()); }
@@ -198,10 +196,6 @@ private:
 	FRWBuffer GPUBufferFloat;
 	/** GPU Buffer containing floating point values for GPU simulations. */
 	FRWBuffer GPUBufferInt;
-	/** Size of the GPU ID buffers. */
-	uint32 NumIDsAllocatedForGPU;
-	/** GPU list of free particle IDs. */
-	FRWBuffer GPUFreeIDs;
 	/** GPU table which maps particle ID to index. */
 	FRWBuffer GPUIDToIndexTable;
 	//////////////////////////////////////////////////////////////////////////
@@ -309,6 +303,8 @@ public:
 	FORCEINLINE int32& GetMaxUsedID() { return MaxUsedID; }
 	FORCEINLINE int32& GetIDAcquireTag() { return IDAcquireTag; }
 	FORCEINLINE void SetIDAcquireTag(int32 InTag) { IDAcquireTag = InTag; }
+	FORCEINLINE FRWBuffer& GetGPUFreeIDs() { return GPUFreeIDs; }
+	FORCEINLINE uint32 GetGPUNumAllocatedIDs() const { return GPUNumAllocatedIDs; }
 
 	FORCEINLINE const TArray<FNiagaraVariable>& GetVariables()const { return CompiledData.Variables; }
 	FORCEINLINE uint32 GetNumVariables()const { return CompiledData.Variables.Num(); }
@@ -347,6 +343,8 @@ public:
 	/** Release the GPU instance counts so that they can be reused */
 	void ReleaseGPUInstanceCounts(FNiagaraGPUInstanceCountManager& GPUInstanceCountManager);
 
+	void AllocateGPUFreeIDs(uint32 InNumInstances, FRHICommandList& RHICmdList, ERHIFeatureLevel::Type FeatureLevel, const TCHAR* DebugSimName);
+
 private:
 
 	void Reset();
@@ -384,6 +382,12 @@ private:
 	/** Tag to use when new IDs are acquired. Should be unique per tick. */
 	int32 IDAcquireTag;
 
+	/** GPU buffer of free IDs available on the next tick. */
+	FRWBuffer GPUFreeIDs;
+
+	/** NUmber of IDs allocated for the GPU simulation. */
+	uint32 GPUNumAllocatedIDs;
+
 	/** Buffer containing the current simulation state. */
 	FNiagaraDataBuffer* CurrentData;
 
diff --git a/Engine/Plugins/FX/Niagara/Source/Niagara/Private/NiagaraDataInterfaceCamera.cpp b/Engine/Plugins/FX/Niagara/Source/Niagara/Private/NiagaraDataInterfaceCamera.cpp
--- a/Engine/Plugins/FX/Niagara/Source/Niagara/Private/NiagaraDataInterfaceCamera.cpp
+++ b/Engine/Plugins/FX/Niagara/Source/Niagara/Private/NiagaraDataInterfaceCamera.cpp
@@ -154,12 +154,12 @@ void UNiagaraDataInterfaceCamera::GetFunctions(TArray<FNiagaraFunctionSignature>
 	OutFunctions.Add(Sig);
 }
 
-bool UNiagaraDataInterfaceCamera::GetFunctionHLSL(const FName& DefinitionFunctionName, FString InstanceFunctionName, FNiagaraDataInterfaceGPUParamInfo& ParamInfo, FString& OutHLSL)
+bool UNiagaraDataInterfaceCamera::GetFunctionHLSL(const FNiagaraDataInterfaceGPUParamInfo& ParamInfo, const FNiagaraDataInterfaceGeneratedFunction& FunctionInfo, int FunctionInstanceIndex, FString& OutHLSL)
 {
 	TMap<FString, FStringFormatArg> ArgsSample;
-	ArgsSample.Add(TEXT("FunctionName"), InstanceFunctionName);
+	ArgsSample.Add(TEXT("FunctionName"), FunctionInfo.InstanceName);
 
-	if (DefinitionFunctionName == GetViewPropertiesName)
+	if (FunctionInfo.DefinitionName == GetViewPropertiesName)
 	{
 		static const TCHAR *FormatSample = TEXT(R"(
 			void {FunctionName}(out float3 Out_ViewPositionWorld, out float3 Out_ViewForwardVector, out float3 Out_ViewUpVector, out float3 Out_ViewRightVector, out float4 Out_ViewSizeAndInverseSize, out float4 Out_ScreenToViewSpace, out float2 Out_Current_TAAJitter, out float2 Out_Previous_TAAJitter)
@@ -177,7 +177,7 @@ bool UNiagaraDataInterfaceCamera::GetFunctionHLSL(const FName& DefinitionFunctio
 		OutHLSL += FString::Format(FormatSample, ArgsSample);
 		return true;
 	}
-	if (DefinitionFunctionName == GetFieldOfViewName)
+	if (FunctionInfo.DefinitionName == GetFieldOfViewName)
 	{
 		static const TCHAR *FormatSample = TEXT(R"(
 			void {FunctionName}(out float Out_FieldOfViewAngle)
@@ -188,7 +188,7 @@ bool UNiagaraDataInterfaceCamera::GetFunctionHLSL(const FName& DefinitionFunctio
 		OutHLSL += FString::Format(FormatSample, ArgsSample);
 		return true;
 	}
-	if (DefinitionFunctionName == GetClipSpaceTransformsName)
+	if (FunctionInfo.DefinitionName == GetClipSpaceTransformsName)
 	{
 		static const TCHAR *FormatSample = TEXT(R"(
 			void {FunctionName}(out float4x4 Out_WorldToClipTransform, out float4x4 Out_TranslatedWorldToClipTransform, out float4x4 Out_ClipToWorldTransform, out float4x4 Out_ClipToViewTransform,
@@ -207,7 +207,7 @@ bool UNiagaraDataInterfaceCamera::GetFunctionHLSL(const FName& DefinitionFunctio
 		OutHLSL += FString::Format(FormatSample, ArgsSample);
 		return true;
 	}
-	if (DefinitionFunctionName == GetViewSpaceTransformsName)
+	if (FunctionInfo.DefinitionName == GetViewSpaceTransformsName)
 	{
 		static const TCHAR *FormatSample = TEXT(R"(
 			void {FunctionName}(out float4x4 Out_TranslatedWorldToViewTransform, out float4x4 Out_ViewToTranslatedWorldTransform, out float4x4 Out_TranslatedWorldToCameraViewTransform,
@@ -224,7 +224,7 @@ bool UNiagaraDataInterfaceCamera::GetFunctionHLSL(const FName& DefinitionFunctio
 		OutHLSL += FString::Format(FormatSample, ArgsSample);
 		return true;
 	}
-	if (DefinitionFunctionName == GetCameraPropertiesName)
+	if (FunctionInfo.DefinitionName == GetCameraPropertiesName)
 	{
 		static const TCHAR *FormatSample = TEXT(R"(
 			void {FunctionName}(out float3 Out_CameraPositionWorld, out float3 Out_ViewForwardVector, out float3 Out_ViewUpVector, out float3 Out_ViewRightVector)
diff --git a/Engine/Plugins/FX/Niagara/Source/Niagara/Private/NiagaraDataInterfaceCollisionQuery.cpp b/Engine/Plugins/FX/Niagara/Source/Niagara/Private/NiagaraDataInterfaceCollisionQuery.cpp
--- a/Engine/Plugins/FX/Niagara/Source/Niagara/Private/NiagaraDataInterfaceCollisionQuery.cpp
+++ b/Engine/Plugins/FX/Niagara/Source/Niagara/Private/NiagaraDataInterfaceCollisionQuery.cpp
@@ -183,13 +183,13 @@ void UNiagaraDataInterfaceCollisionQuery::GetFunctions(TArray<FNiagaraFunctionSi
 // TODO: need a way to identify each specific function here
 
 // 
-bool UNiagaraDataInterfaceCollisionQuery::GetFunctionHLSL(const FName& DefinitionFunctionName, FString InstanceFunctionName, FNiagaraDataInterfaceGPUParamInfo& ParamInfo, FString& OutHLSL)
+bool UNiagaraDataInterfaceCollisionQuery::GetFunctionHLSL(const FNiagaraDataInterfaceGPUParamInfo& ParamInfo, const FNiagaraDataInterfaceGeneratedFunction& FunctionInfo, int FunctionInstanceIndex, FString& OutHLSL)
 {
 	// a little tricky, since we've got two functions for submitting and retrieving a query; we store submitted queries per thread group,
 	// assuming it'll usually be the same thread trying to call ReadQuery for a particular QueryID, that submitted it in the first place.
-	if (DefinitionFunctionName == TEXT("PerformCollisionQuery"))
+	if (FunctionInfo.DefinitionName == TEXT("PerformCollisionQuery"))
 	{
-		OutHLSL += TEXT("void ") + InstanceFunctionName + TEXT("(in int InQueryID, in float3 In_ParticlePos, in float3 In_ParticleVel, in float In_DeltaSeconds, float CollisionRadius, in float CollisionDepthBounds, \
+		OutHLSL += TEXT("void ") + FunctionInfo.InstanceName + TEXT("(in int InQueryID, in float3 In_ParticlePos, in float3 In_ParticleVel, in float In_DeltaSeconds, float CollisionRadius, in float CollisionDepthBounds, \
 			out int Out_QueryID, out bool OutCollisionValid, out float3 Out_CollisionPos, out float3 Out_CollisionNormal, out float Out_Friction, out float Out_Restitution) \n{\n");
 		// get the screen position
 		OutHLSL += TEXT("\
@@ -240,10 +240,10 @@ bool UNiagaraDataInterfaceCollisionQuery::GetFunctionHLSL(const FName& Definitio
 			}\n\
 		}\n}\n\n");
 	}
-	else if (DefinitionFunctionName == TEXT("PerformCollisionQueryGPUShader"))
+	else if (FunctionInfo.DefinitionName == TEXT("PerformCollisionQueryGPUShader"))
 	{
-		FString SceneDepthFunction = InstanceFunctionName + TEXT("_SceneDepthCollision");
-		FString DistanceFieldFunction = InstanceFunctionName + TEXT("_DistanceFieldCollision");
+		FString SceneDepthFunction = FunctionInfo.InstanceName + TEXT("_SceneDepthCollision");
+		FString DistanceFieldFunction = FunctionInfo.InstanceName + TEXT("_DistanceFieldCollision");
 
 		OutHLSL += TEXT("void ") + SceneDepthFunction + TEXT("(in float3 In_SamplePos, in float3 In_TraceEndPos, in float CollisionDepthBounds, in float ParticleRadius, out bool OutCollisionValid, out float3 Out_CollisionPos, out float3 Out_CollisionNormal) \n{\n\
 		OutCollisionValid = false;\n\
@@ -293,7 +293,7 @@ bool UNiagaraDataInterfaceCollisionQuery::GetFunctionHLSL(const FName& Definitio
 			Out_CollisionNormal = float3(0.0, 0.0, 1.0);\n\
 			Out_CollisionPos = InPosition;\n\
 		}\n}\n\n");
-		OutHLSL += TEXT("void ") + InstanceFunctionName + TEXT("(in float3 In_SamplePos, in float3 In_TraceEndPos, in float CollisionDepthBounds, ") +
+		OutHLSL += TEXT("void ") + FunctionInfo.InstanceName + TEXT("(in float3 In_SamplePos, in float3 In_TraceEndPos, in float CollisionDepthBounds, ") +
 			TEXT("in float ParticleRadius, in bool UseMeshDistanceField, out bool OutCollisionValid, out float3 Out_CollisionPos, out float3 Out_CollisionNormal) \n{\n");
 		OutHLSL += TEXT("\
 			if (UseMeshDistanceField)\n\
@@ -305,9 +305,9 @@ bool UNiagaraDataInterfaceCollisionQuery::GetFunctionHLSL(const FName& Definitio
 				") + SceneDepthFunction + TEXT("(In_SamplePos, In_TraceEndPos, CollisionDepthBounds, ParticleRadius, OutCollisionValid, Out_CollisionPos, Out_CollisionNormal);\n\
 			}\n}\n\n");
 	}
-	else if (DefinitionFunctionName == TEXT("QuerySceneDepthGPU"))
+	else if (FunctionInfo.DefinitionName == TEXT("QuerySceneDepthGPU"))
 	{
-		OutHLSL += TEXT("void ") + InstanceFunctionName + TEXT("(in float3 In_SamplePos, out float Out_SceneDepth, out float3 Out_CameraPosWorld, out bool Out_IsInsideView, out float3 Out_WorldPos, out float3 Out_WorldNormal) \n{\n");
+		OutHLSL += TEXT("void ") + FunctionInfo.InstanceName + TEXT("(in float3 In_SamplePos, out float Out_SceneDepth, out float3 Out_CameraPosWorld, out bool Out_IsInsideView, out float3 Out_WorldPos, out float3 Out_WorldNormal) \n{\n");
 		OutHLSL += TEXT("\
 			Out_SceneDepth = -1;\n\
 			Out_WorldPos = float3(0.0, 0.0, 0.0);\n\
@@ -334,9 +334,9 @@ bool UNiagaraDataInterfaceCollisionQuery::GetFunctionHLSL(const FName& Definitio
 				Out_IsInsideView = false;\n\
 			}\n}\n\n");
 	}
-	else if (DefinitionFunctionName == TEXT("QueryMeshDistanceFieldGPU"))
+	else if (FunctionInfo.DefinitionName == TEXT("QueryMeshDistanceFieldGPU"))
 	{
-		OutHLSL += TEXT("void ") + InstanceFunctionName + TEXT("(in float3 In_SamplePos, out float Out_DistanceToNearestSurface, out float3 Out_FieldGradient) \n{\n");
+		OutHLSL += TEXT("void ") + FunctionInfo.InstanceName + TEXT("(in float3 In_SamplePos, out float Out_DistanceToNearestSurface, out float3 Out_FieldGradient) \n{\n");
 		OutHLSL += TEXT("\
 			Out_DistanceToNearestSurface = GetDistanceToNearestSurfaceGlobal(In_SamplePos);\n\
 			Out_FieldGradient = GetDistanceFieldGradientGlobal(In_SamplePos);\
@@ -346,7 +346,7 @@ bool UNiagaraDataInterfaceCollisionQuery::GetFunctionHLSL(const FName& Definitio
 	return true;
 }
 
-void UNiagaraDataInterfaceCollisionQuery::GetParameterDefinitionHLSL(FNiagaraDataInterfaceGPUParamInfo& ParamInfo, FString& OutHLSL)
+void UNiagaraDataInterfaceCollisionQuery::GetParameterDefinitionHLSL(const FNiagaraDataInterfaceGPUParamInfo& ParamInfo, FString& OutHLSL)
 {
 	// we don't need to add these to hlsl, as they're already in common.ush
 }
diff --git a/Engine/Plugins/FX/Niagara/Source/Niagara/Private/NiagaraDataInterfaceColorCurve.cpp b/Engine/Plugins/FX/Niagara/Source/Niagara/Private/NiagaraDataInterfaceColorCurve.cpp
--- a/Engine/Plugins/FX/Niagara/Source/Niagara/Private/NiagaraDataInterfaceColorCurve.cpp
+++ b/Engine/Plugins/FX/Niagara/Source/Niagara/Private/NiagaraDataInterfaceColorCurve.cpp
@@ -172,7 +172,7 @@ void UNiagaraDataInterfaceColorCurve::GetFunctions(TArray<FNiagaraFunctionSignat
 // the HLSL in the spirit of a static switch
 // TODO: need a way to identify each specific function here
 // 
-bool UNiagaraDataInterfaceColorCurve::GetFunctionHLSL(const FName& DefinitionFunctionName, FString InstanceFunctionName, FNiagaraDataInterfaceGPUParamInfo& ParamInfo, FString& OutHLSL)
+bool UNiagaraDataInterfaceColorCurve::GetFunctionHLSL(const FNiagaraDataInterfaceGPUParamInfo& ParamInfo, const FNiagaraDataInterfaceGeneratedFunction& FunctionInfo, int FunctionInstanceIndex, FString& OutHLSL)
 {
 	FString TimeToLUTFrac = TEXT("TimeToLUTFraction_") + ParamInfo.DataInterfaceHLSLSymbol;
 	FString Sample = TEXT("SampleCurve_") + ParamInfo.DataInterfaceHLSLSymbol;
@@ -190,7 +190,7 @@ void %s(in float In_X, out float4 Out_Value) \n\
 	float4 B = float4(%s(Next), %s(Next + 1), %s(Next + 2), %s(Next + 3)); \n\
 	Out_Value = lerp(A, B, Interp); \n\
 }\n")
-, *InstanceFunctionName, *TimeToLUTFrac, *NumSamples, *NumSamples, CurveLUTNumElems, CurveLUTNumElems
+, *FunctionInfo.InstanceName, *TimeToLUTFrac, *NumSamples, *NumSamples, CurveLUTNumElems, CurveLUTNumElems
 , *Sample, *Sample, *Sample, *Sample, *Sample, *Sample, *Sample, *Sample);
 
 	return true;
diff --git a/Engine/Plugins/FX/Niagara/Source/Niagara/Private/NiagaraDataInterfaceCurlNoise.cpp b/Engine/Plugins/FX/Niagara/Source/Niagara/Private/NiagaraDataInterfaceCurlNoise.cpp
--- a/Engine/Plugins/FX/Niagara/Source/Niagara/Private/NiagaraDataInterfaceCurlNoise.cpp
+++ b/Engine/Plugins/FX/Niagara/Source/Niagara/Private/NiagaraDataInterfaceCurlNoise.cpp
@@ -423,7 +423,7 @@ void UNiagaraDataInterfaceCurlNoise::SampleNoiseField(FVectorVMContext& Context)
 	}
 }
 
-bool UNiagaraDataInterfaceCurlNoise::GetFunctionHLSL(const FName& DefinitionFunctionName, FString InstanceFunctionName, FNiagaraDataInterfaceGPUParamInfo& ParamInfo, FString& OutHLSL)
+bool UNiagaraDataInterfaceCurlNoise::GetFunctionHLSL(const FNiagaraDataInterfaceGPUParamInfo& ParamInfo, const FNiagaraDataInterfaceGeneratedFunction& FunctionInfo, int FunctionInstanceIndex, FString& OutHLSL)
 {
 	static const TCHAR *FormatSample = TEXT(R"(
 		void {FunctionName}(float3 In_XYZ, out float3 Out_Value)
@@ -434,13 +434,13 @@ bool UNiagaraDataInterfaceCurlNoise::GetFunctionHLSL(const FName& DefinitionFunc
 		}
 	)");
 	TMap<FString, FStringFormatArg> ArgsSample;
-	ArgsSample.Add(TEXT("FunctionName"), InstanceFunctionName);
+	ArgsSample.Add(TEXT("FunctionName"), FunctionInfo.InstanceName);
 	ArgsSample.Add(TEXT("OffsetFromSeedName"), OffsetFromSeedBaseName + ParamInfo.DataInterfaceHLSLSymbol);
 	OutHLSL += FString::Format(FormatSample, ArgsSample);
 	return true;
 }
 
-void UNiagaraDataInterfaceCurlNoise::GetParameterDefinitionHLSL(FNiagaraDataInterfaceGPUParamInfo& ParamInfo, FString& OutHLSL)
+void UNiagaraDataInterfaceCurlNoise::GetParameterDefinitionHLSL(const FNiagaraDataInterfaceGPUParamInfo& ParamInfo, FString& OutHLSL)
 {
 	static const TCHAR *FormatDeclarations = TEXT(R"(
 		float3 {OffsetFromSeedName};
diff --git a/Engine/Plugins/FX/Niagara/Source/Niagara/Private/NiagaraDataInterfaceCurve.cpp b/Engine/Plugins/FX/Niagara/Source/Niagara/Private/NiagaraDataInterfaceCurve.cpp
--- a/Engine/Plugins/FX/Niagara/Source/Niagara/Private/NiagaraDataInterfaceCurve.cpp
+++ b/Engine/Plugins/FX/Niagara/Source/Niagara/Private/NiagaraDataInterfaceCurve.cpp
@@ -134,7 +134,7 @@ TArray<float> UNiagaraDataInterfaceCurve::BuildLUT(int32 NumEntries) const
 // the HLSL in the spirit of a static switch
 // TODO: need a way to identify each specific function here
 // 
-bool UNiagaraDataInterfaceCurve::GetFunctionHLSL(const FName& DefinitionFunctionName, FString InstanceFunctionName, FNiagaraDataInterfaceGPUParamInfo& ParamInfo, FString& OutHLSL)
+bool UNiagaraDataInterfaceCurve::GetFunctionHLSL(const FNiagaraDataInterfaceGPUParamInfo& ParamInfo, const FNiagaraDataInterfaceGeneratedFunction& FunctionInfo, int FunctionInstanceIndex, FString& OutHLSL)
 {
 	FString TimeToLUTFrac = TEXT("TimeToLUTFraction_") + ParamInfo.DataInterfaceHLSLSymbol;
 	FString Sample = TEXT("SampleCurve_") + ParamInfo.DataInterfaceHLSLSymbol;
@@ -150,7 +150,7 @@ void %s(in float In_X, out float Out_Value) \n\
 	float B = %s(Next); \n\
 	Out_Value = lerp(A, B, Interp); \n\
 }\n")
-, *InstanceFunctionName, *TimeToLUTFrac, *NumSamples, *NumSamples, *Sample, *Sample);
+, *FunctionInfo.InstanceName, *TimeToLUTFrac, *NumSamples, *NumSamples, *Sample, *Sample);
 
 	return true;
 }
diff --git a/Engine/Plugins/FX/Niagara/Source/Niagara/Private/NiagaraDataInterfaceCurveBase.cpp b/Engine/Plugins/FX/Niagara/Source/Niagara/Private/NiagaraDataInterfaceCurveBase.cpp
--- a/Engine/Plugins/FX/Niagara/Source/Niagara/Private/NiagaraDataInterfaceCurveBase.cpp
+++ b/Engine/Plugins/FX/Niagara/Source/Niagara/Private/NiagaraDataInterfaceCurveBase.cpp
@@ -233,7 +233,7 @@ bool UNiagaraDataInterfaceCurveBase::Equals(const UNiagaraDataInterface* Other)
 }
 
 
-void UNiagaraDataInterfaceCurveBase::GetParameterDefinitionHLSL(FNiagaraDataInterfaceGPUParamInfo& ParamInfo, FString& OutHLSL)
+void UNiagaraDataInterfaceCurveBase::GetParameterDefinitionHLSL(const FNiagaraDataInterfaceGPUParamInfo& ParamInfo, FString& OutHLSL)
 {
 	OutHLSL += TEXT("\n");
 	FString MinTimeStr = TEXT("MinTime_") + ParamInfo.DataInterfaceHLSLSymbol;
diff --git a/Engine/Plugins/FX/Niagara/Source/Niagara/Private/NiagaraDataInterfaceExport.cpp b/Engine/Plugins/FX/Niagara/Source/Niagara/Private/NiagaraDataInterfaceExport.cpp
--- a/Engine/Plugins/FX/Niagara/Source/Niagara/Private/NiagaraDataInterfaceExport.cpp
+++ b/Engine/Plugins/FX/Niagara/Source/Niagara/Private/NiagaraDataInterfaceExport.cpp
@@ -126,7 +126,7 @@ void UNiagaraDataInterfaceExport::GetFunctions(TArray<FNiagaraFunctionSignature>
 	OutFunctions.Add(Sig);
 }
 
-bool UNiagaraDataInterfaceExport::GetFunctionHLSL(const FName& DefinitionFunctionName, FString InstanceFunctionName, FNiagaraDataInterfaceGPUParamInfo& ParamInfo, FString& OutHLSL)
+bool UNiagaraDataInterfaceExport::GetFunctionHLSL(const FNiagaraDataInterfaceGPUParamInfo& ParamInfo, const FNiagaraDataInterfaceGeneratedFunction& FunctionInfo, int FunctionInstanceIndex, FString& OutHLSL)
 {
 	return false;
 }
diff --git a/Engine/Plugins/FX/Niagara/Source/Niagara/Private/NiagaraDataInterfaceGrid2DCollection.cpp b/Engine/Plugins/FX/Niagara/Source/Niagara/Private/NiagaraDataInterfaceGrid2DCollection.cpp
--- a/Engine/Plugins/FX/Niagara/Source/Niagara/Private/NiagaraDataInterfaceGrid2DCollection.cpp
+++ b/Engine/Plugins/FX/Niagara/Source/Niagara/Private/NiagaraDataInterfaceGrid2DCollection.cpp
@@ -228,7 +228,7 @@ bool UNiagaraDataInterfaceGrid2DCollection::Equals(const UNiagaraDataInterface*
 	return OtherTyped != nullptr;		
 }
 
-void UNiagaraDataInterfaceGrid2DCollection::GetParameterDefinitionHLSL(FNiagaraDataInterfaceGPUParamInfo& ParamInfo, FString& OutHLSL)
+void UNiagaraDataInterfaceGrid2DCollection::GetParameterDefinitionHLSL(const FNiagaraDataInterfaceGPUParamInfo& ParamInfo, FString& OutHLSL)
 {
 	Super::GetParameterDefinitionHLSL(ParamInfo, OutHLSL);
 
@@ -248,14 +248,14 @@ void UNiagaraDataInterfaceGrid2DCollection::GetParameterDefinitionHLSL(FNiagaraD
 	OutHLSL += FString::Format(FormatDeclarations, ArgsDeclarations);
 }
 
-bool UNiagaraDataInterfaceGrid2DCollection::GetFunctionHLSL(const FName& DefinitionFunctionName, FString InstanceFunctionName, FNiagaraDataInterfaceGPUParamInfo& ParamInfo, FString& OutHLSL)
+bool UNiagaraDataInterfaceGrid2DCollection::GetFunctionHLSL(const FNiagaraDataInterfaceGPUParamInfo& ParamInfo, const FNiagaraDataInterfaceGeneratedFunction& FunctionInfo, int FunctionInstanceIndex, FString& OutHLSL)
 {
-	bool ParentRet = Super::GetFunctionHLSL(DefinitionFunctionName, InstanceFunctionName, ParamInfo, OutHLSL);
+	bool ParentRet = Super::GetFunctionHLSL(ParamInfo, FunctionInfo, FunctionInstanceIndex, OutHLSL);
 	if (ParentRet)
 	{
 		return true;
 	} 
-	else if (DefinitionFunctionName == GetValueFunctionName)
+	else if (FunctionInfo.DefinitionName == GetValueFunctionName)
 	{
 		static const TCHAR *FormatBounds = TEXT(R"(
 			void {FunctionName}(int In_IndexX, int In_IndexY, int In_AttributeIndex, out float Out_Val)
@@ -267,7 +267,7 @@ bool UNiagaraDataInterfaceGrid2DCollection::GetFunctionHLSL(const FName& Definit
 			}
 		)");
 		TMap<FString, FStringFormatArg> ArgsBounds = {
-			{TEXT("FunctionName"), InstanceFunctionName},
+			{TEXT("FunctionName"), FunctionInfo.InstanceName},
 			{TEXT("Grid"), GridName + ParamInfo.DataInterfaceHLSLSymbol},
 			{TEXT("NumCellsName"), NumCellsName + ParamInfo.DataInterfaceHLSLSymbol},
 			{TEXT("NumTiles"),    NumTilesName + ParamInfo.DataInterfaceHLSLSymbol},
@@ -275,7 +275,7 @@ bool UNiagaraDataInterfaceGrid2DCollection::GetFunctionHLSL(const FName& Definit
 		OutHLSL += FString::Format(FormatBounds, ArgsBounds);
 		return true;
 	}
-	else if (DefinitionFunctionName == SetValueFunctionName)
+	else if (FunctionInfo.DefinitionName == SetValueFunctionName)
 	{
 		static const TCHAR *FormatBounds = TEXT(R"(
 			void {FunctionName}(int In_IndexX, int In_IndexY, int In_AttributeIndex, float In_Value, out int val)
@@ -288,7 +288,7 @@ bool UNiagaraDataInterfaceGrid2DCollection::GetFunctionHLSL(const FName& Definit
 			}
 		)");
 		TMap<FString, FStringFormatArg> ArgsBounds = {
-			{TEXT("FunctionName"), InstanceFunctionName},
+			{TEXT("FunctionName"), FunctionInfo.InstanceName},
 			{TEXT("OutputGrid"), OutputGridName + ParamInfo.DataInterfaceHLSLSymbol},
 			{TEXT("NumCellsName"), NumCellsName + ParamInfo.DataInterfaceHLSLSymbol},
 			{TEXT("NumTiles"),    NumTilesName + ParamInfo.DataInterfaceHLSLSymbol},
@@ -297,7 +297,7 @@ bool UNiagaraDataInterfaceGrid2DCollection::GetFunctionHLSL(const FName& Definit
 		OutHLSL += FString::Format(FormatBounds, ArgsBounds);
 		return true;
 	}
-	else if (DefinitionFunctionName == SampleGridFunctionName)
+	else if (FunctionInfo.DefinitionName == SampleGridFunctionName)
 	{
 		static const TCHAR *FormatBounds = TEXT(R"(
 			void {FunctionName}(float In_UnitX, float In_UnitY, int In_AttributeIndex, out float Out_Val)
@@ -309,7 +309,7 @@ bool UNiagaraDataInterfaceGrid2DCollection::GetFunctionHLSL(const FName& Definit
 			}
 		)");
 		TMap<FString, FStringFormatArg> ArgsBounds = {
-			{TEXT("FunctionName"), InstanceFunctionName},
+			{TEXT("FunctionName"), FunctionInfo.InstanceName},
 			{TEXT("Grid"), GridName + ParamInfo.DataInterfaceHLSLSymbol},
 			{TEXT("SamplerName"),    SamplerName + ParamInfo.DataInterfaceHLSLSymbol },
 			{TEXT("NumTiles"),    NumTilesName + ParamInfo.DataInterfaceHLSLSymbol},
diff --git a/Engine/Plugins/FX/Niagara/Source/Niagara/Private/NiagaraDataInterfaceNeighborGrid3D.cpp b/Engine/Plugins/FX/Niagara/Source/Niagara/Private/NiagaraDataInterfaceNeighborGrid3D.cpp
--- a/Engine/Plugins/FX/Niagara/Source/Niagara/Private/NiagaraDataInterfaceNeighborGrid3D.cpp
+++ b/Engine/Plugins/FX/Niagara/Source/Niagara/Private/NiagaraDataInterfaceNeighborGrid3D.cpp
@@ -272,7 +272,7 @@ bool UNiagaraDataInterfaceNeighborGrid3D::Equals(const UNiagaraDataInterface* Ot
 	return OtherTyped->MaxNeighborsPerVoxel == MaxNeighborsPerVoxel;
 }
 
-void UNiagaraDataInterfaceNeighborGrid3D::GetParameterDefinitionHLSL(FNiagaraDataInterfaceGPUParamInfo& ParamInfo, FString& OutHLSL)
+void UNiagaraDataInterfaceNeighborGrid3D::GetParameterDefinitionHLSL(const FNiagaraDataInterfaceGPUParamInfo& ParamInfo, FString& OutHLSL)
 {
 	Super::GetParameterDefinitionHLSL(ParamInfo, OutHLSL);
 
@@ -293,13 +293,13 @@ void UNiagaraDataInterfaceNeighborGrid3D::GetParameterDefinitionHLSL(FNiagaraDat
 	OutHLSL += FString::Format(FormatDeclarations, ArgsDeclarations);
 }
 
-bool UNiagaraDataInterfaceNeighborGrid3D::GetFunctionHLSL(const FName& DefinitionFunctionName, FString InstanceFunctionName, FNiagaraDataInterfaceGPUParamInfo& ParamInfo, FString& OutHLSL)
+bool UNiagaraDataInterfaceNeighborGrid3D::GetFunctionHLSL(const FNiagaraDataInterfaceGPUParamInfo& ParamInfo, const FNiagaraDataInterfaceGeneratedFunction& FunctionInfo, int FunctionInstanceIndex, FString& OutHLSL)
 {
-	bool ParentRet = Super::GetFunctionHLSL(DefinitionFunctionName, InstanceFunctionName, ParamInfo, OutHLSL);
+	bool ParentRet = Super::GetFunctionHLSL(ParamInfo, FunctionInfo, FunctionInstanceIndex, OutHLSL);
 	if (ParentRet)
 	{
 		return true;
-	} else if (DefinitionFunctionName == MaxNeighborsPerVoxelFunctionName)
+	} else if (FunctionInfo.DefinitionName == MaxNeighborsPerVoxelFunctionName)
 	{
 		static const TCHAR *FormatSample = TEXT(R"(
 			void {FunctionName}(out int Out_MaxNeighborsPerVoxel)
@@ -308,14 +308,14 @@ bool UNiagaraDataInterfaceNeighborGrid3D::GetFunctionHLSL(const FName& Definitio
 			}
 		)");
 		TMap<FString, FStringFormatArg> ArgsSample = {
-			{TEXT("FunctionName"), InstanceFunctionName},
+			{TEXT("FunctionName"), FunctionInfo.InstanceName},
 			{TEXT("MaxNeighborsPerVoxelName"), MaxNeighborsPerVoxelName + ParamInfo.DataInterfaceHLSLSymbol},
 
 		};
 		OutHLSL += FString::Format(FormatSample, ArgsSample);
 		return true;
 	}
-	else if (DefinitionFunctionName == NeighborGridIndexToLinearFunctionName)
+	else if (FunctionInfo.DefinitionName == NeighborGridIndexToLinearFunctionName)
 	{
 		static const TCHAR *FormatBounds = TEXT(R"(
 			void {FunctionName}(int In_IndexX, int In_IndexY, int In_IndexZ, int In_Neighbor, out int Out_Linear)
@@ -324,14 +324,14 @@ bool UNiagaraDataInterfaceNeighborGrid3D::GetFunctionHLSL(const FName& Definitio
 			}
 		)");
 		TMap<FString, FStringFormatArg> ArgsBounds = {
-			{TEXT("FunctionName"), InstanceFunctionName},
+			{TEXT("FunctionName"), FunctionInfo.InstanceName},
 			{TEXT("MaxNeighborsPerVoxelName"), MaxNeighborsPerVoxelName + ParamInfo.DataInterfaceHLSLSymbol},
 			{TEXT("NumVoxelsName"), NumVoxelsName + ParamInfo.DataInterfaceHLSLSymbol},
 		};
 		OutHLSL += FString::Format(FormatBounds, ArgsBounds);
 		return true;
 	}
-	else if (DefinitionFunctionName == GetParticleNeighborFunctionName)
+	else if (FunctionInfo.DefinitionName == GetParticleNeighborFunctionName)
 	{
 		static const TCHAR *FormatBounds = TEXT(R"(
 			void {FunctionName}(int In_Index, out int Out_ParticleNeighborIndex)
@@ -340,13 +340,13 @@ bool UNiagaraDataInterfaceNeighborGrid3D::GetFunctionHLSL(const FName& Definitio
 			}
 		)");
 		TMap<FString, FStringFormatArg> ArgsBounds = {
-			{TEXT("FunctionName"), InstanceFunctionName},
+			{TEXT("FunctionName"), FunctionInfo.InstanceName},
 			{TEXT("ParticleNeighbors"), ParticleNeighborsName + ParamInfo.DataInterfaceHLSLSymbol},
 		};
 		OutHLSL += FString::Format(FormatBounds, ArgsBounds);
 		return true;
 	}
-	else if (DefinitionFunctionName == SetParticleNeighborFunctionName)
+	else if (FunctionInfo.DefinitionName == SetParticleNeighborFunctionName)
 	{
 		static const TCHAR *FormatBounds = TEXT(R"(
 			void {FunctionName}(int In_Index, int In_ParticleNeighborIndex, out int Out_Ignore)
@@ -356,13 +356,13 @@ bool UNiagaraDataInterfaceNeighborGrid3D::GetFunctionHLSL(const FName& Definitio
 			}
 		)");
 		TMap<FString, FStringFormatArg> ArgsBounds = {
-			{TEXT("FunctionName"), InstanceFunctionName},
+			{TEXT("FunctionName"), FunctionInfo.InstanceName},
 			{TEXT("OutputParticleNeighbors"), OutputParticleNeighborsName + ParamInfo.DataInterfaceHLSLSymbol},
 		};
 		OutHLSL += FString::Format(FormatBounds, ArgsBounds);
 		return true;
 	}
-	else if (DefinitionFunctionName == GetParticleNeighborCountFunctionName)
+	else if (FunctionInfo.DefinitionName == GetParticleNeighborCountFunctionName)
 	{
 		static const TCHAR *FormatBounds = TEXT(R"(
 			void {FunctionName}(int In_Index, out int Out_ParticleNeighborIndex)
@@ -371,13 +371,13 @@ bool UNiagaraDataInterfaceNeighborGrid3D::GetFunctionHLSL(const FName& Definitio
 			}
 		)");
 		TMap<FString, FStringFormatArg> ArgsBounds = {
-			{TEXT("FunctionName"), InstanceFunctionName},
+			{TEXT("FunctionName"), FunctionInfo.InstanceName},
 			{TEXT("ParticleNeighborCount"), ParticleNeighborCountName + ParamInfo.DataInterfaceHLSLSymbol},
 		};
 		OutHLSL += FString::Format(FormatBounds, ArgsBounds);
 		return true;
 	}
-	else if (DefinitionFunctionName == SetParticleNeighborCountFunctionName)
+	else if (FunctionInfo.DefinitionName == SetParticleNeighborCountFunctionName)
 	{
 		static const TCHAR *FormatBounds = TEXT(R"(
 			void {FunctionName}(int In_Index, int In_Increment, out int PreviousNeighborCount)
@@ -386,7 +386,7 @@ bool UNiagaraDataInterfaceNeighborGrid3D::GetFunctionHLSL(const FName& Definitio
 			}
 		)");
 		TMap<FString, FStringFormatArg> ArgsBounds = {
-			{TEXT("FunctionName"), InstanceFunctionName},
+			{TEXT("FunctionName"), FunctionInfo.InstanceName},
 			{TEXT("OutputParticleNeighborCount"), OutputParticleNeighborCountName + ParamInfo.DataInterfaceHLSLSymbol},
 		};
 		OutHLSL += FString::Format(FormatBounds, ArgsBounds);
diff --git a/Engine/Plugins/FX/Niagara/Source/Niagara/Private/NiagaraDataInterfaceOcclusion.cpp b/Engine/Plugins/FX/Niagara/Source/Niagara/Private/NiagaraDataInterfaceOcclusion.cpp
--- a/Engine/Plugins/FX/Niagara/Source/Niagara/Private/NiagaraDataInterfaceOcclusion.cpp
+++ b/Engine/Plugins/FX/Niagara/Source/Niagara/Private/NiagaraDataInterfaceOcclusion.cpp
@@ -63,12 +63,12 @@ void UNiagaraDataInterfaceOcclusion::GetFunctions(TArray<FNiagaraFunctionSignatu
 	OutFunctions.Add(Sig);
 }
 
-bool UNiagaraDataInterfaceOcclusion::GetFunctionHLSL(const FName& DefinitionFunctionName, FString InstanceFunctionName, FNiagaraDataInterfaceGPUParamInfo& ParamInfo, FString& OutHLSL)
+bool UNiagaraDataInterfaceOcclusion::GetFunctionHLSL(const FNiagaraDataInterfaceGPUParamInfo& ParamInfo, const FNiagaraDataInterfaceGeneratedFunction& FunctionInfo, int FunctionInstanceIndex, FString& OutHLSL)
 {
 	TMap<FString, FStringFormatArg> ArgsSample;
-	ArgsSample.Add(TEXT("FunctionName"), InstanceFunctionName);
+	ArgsSample.Add(TEXT("FunctionName"), FunctionInfo.InstanceName);
 
-	if (DefinitionFunctionName == GetCameraOcclusionRectangleName)
+	if (FunctionInfo.DefinitionName == GetCameraOcclusionRectangleName)
 	{
 		static const TCHAR *FormatSample = TEXT(R"(
 			void {FunctionName}(in float3 In_SampleCenterWorldPos, in float In_SampleWindowWidthWorld, in float In_SampleWindowHeightWorld, in float In_SampleSteps, out float Out_VisibilityFraction, out float Out_SampleFraction)
@@ -125,7 +125,7 @@ bool UNiagaraDataInterfaceOcclusion::GetFunctionHLSL(const FName& DefinitionFunc
 		OutHLSL += FString::Format(FormatSample, ArgsSample);
 		return true;
 	}
-	if (DefinitionFunctionName == GetCameraOcclusionCircleName)
+	if (FunctionInfo.DefinitionName == GetCameraOcclusionCircleName)
 	{
 		static const TCHAR *FormatSample = TEXT(R"(
 			void {FunctionName}(in float3 In_SampleCenterWorldPos, in float In_SampleWindowDiameterWorld, in float In_SampleRays, in float In_SampleStepsPerRay, out float Out_VisibilityFraction, out float Out_SampleFraction)
diff --git a/Engine/Plugins/FX/Niagara/Source/Niagara/Private/NiagaraDataInterfaceParticleRead.cpp b/Engine/Plugins/FX/Niagara/Source/Niagara/Private/NiagaraDataInterfaceParticleRead.cpp
--- a/Engine/Plugins/FX/Niagara/Source/Niagara/Private/NiagaraDataInterfaceParticleRead.cpp
+++ b/Engine/Plugins/FX/Niagara/Source/Niagara/Private/NiagaraDataInterfaceParticleRead.cpp
@@ -3,6 +3,8 @@
 #include "NiagaraDataInterfaceParticleRead.h"
 #include "NiagaraConstants.h"
 #include "NiagaraSystemInstance.h"
+#include "ShaderParameterUtils.h"
+#include "NiagaraRenderer.h"
 
 #define LOCTEXT_NAMESPACE "NiagaraDataInterfaceParticleRead"
 
@@ -15,9 +17,320 @@ static const FName GetBoolAttributeFunctionName("Get bool Attribute");
 static const FName GetColorAttributeFunctionName("Get Color Attribute");
 static const FName GetQuatAttributeFunctionName("Get Quaternion Attribute");
 
+static const FString IDToIndexTableBaseName(TEXT("IDToIndexTable_"));
+static const FString InputFloatBufferBaseName(TEXT("InputFloatBuffer_"));
+static const FString InputIntBufferBaseName(TEXT("InputIntBuffer_"));
+static const FString ParticleStrideFloatBaseName(TEXT("ParticleStrideFloat_"));
+static const FString ParticleStrideIntBaseName(TEXT("ParticleStrideInt_"));
+static const FString AttributeIndicesBaseName(TEXT("AttributeIndices_"));
+static const FString AcquireTagRegisterIndexBaseName(TEXT("AcquireTagRegisterIndex_"));
+
+enum class ENiagaraParticleDataComponentType : uint8
+{
+	Float,
+	Int,
+	Bool
+};
+
+enum class ENiagaraParticleDataValueType : uint8
+{
+	Invalid,
+	Int,
+	Float,
+	Vec2,
+	Vec3,
+	Vec4,
+	Bool,
+	Color,
+	Quat
+};
+
+static const TCHAR* NiagaraParticleDataValueTypeName(ENiagaraParticleDataValueType Type)
+{
+	switch (Type)
+	{
+		case ENiagaraParticleDataValueType::Invalid:	return TEXT("INVALID");
+		case ENiagaraParticleDataValueType::Int:		return TEXT("int");
+		case ENiagaraParticleDataValueType::Float:		return TEXT("float");
+		case ENiagaraParticleDataValueType::Vec2:		return TEXT("vec2");
+		case ENiagaraParticleDataValueType::Vec3:		return TEXT("vec3");
+		case ENiagaraParticleDataValueType::Vec4:		return TEXT("vec4");
+		case ENiagaraParticleDataValueType::Bool:		return TEXT("bool");
+		case ENiagaraParticleDataValueType::Color:		return TEXT("color");
+		case ENiagaraParticleDataValueType::Quat:		return TEXT("quaternion");
+		default:										return TEXT("UNKNOWN");
+	}
+}
+
+struct FNDIParticleRead_InstanceDataGPU
+{
+	FNiagaraComputeExecutionContext* SourceEmitterGPUContext;
+	FString SourceEmitterName;
+};
+
+struct FNiagaraDataInterfaceProxyParticleRead : public FNiagaraDataInterfaceProxy
+{
+	FNiagaraDataInterfaceProxyParticleRead() : SourceEmitterGPUContext(nullptr)
+	{
+	}
+
+	virtual void ConsumePerInstanceDataFromGameThread(void* PerInstanceData, const FNiagaraSystemInstanceID& Instance) override
+	{
+		FNDIParticleRead_InstanceDataGPU* InstanceData = static_cast<FNDIParticleRead_InstanceDataGPU*>(PerInstanceData);
+		if (InstanceData)
+		{
+			SourceEmitterGPUContext = InstanceData->SourceEmitterGPUContext;
+			SourceEmitterName = InstanceData->SourceEmitterName;
+		}
+	}
+	
+	virtual int32 PerInstanceDataPassedToRenderThreadSize() const override
+	{
+		return sizeof(FNDIParticleRead_InstanceDataGPU);
+	}
+
+	FNiagaraComputeExecutionContext* SourceEmitterGPUContext;
+	FString SourceEmitterName;
+};
+
+struct FNiagaraDataInterfaceParametersCS_ParticleRead : public FNiagaraDataInterfaceParametersCS
+{
+	FNiagaraDataInterfaceParametersCS_ParticleRead() : 
+		CachedDataSet(nullptr),
+		bSourceEmitterNotGPUErrorShown(false)
+	{
+	}
+
+	ENiagaraParticleDataValueType GetValueTypeFromFuncName(const FName& FuncName)
+	{
+		if (FuncName == GetIntAttributeFunctionName) return ENiagaraParticleDataValueType::Int;
+		if (FuncName == GetFloatAttributeFunctionName) return ENiagaraParticleDataValueType::Float;
+		if (FuncName == GetVec2AttributeFunctionName) return ENiagaraParticleDataValueType::Vec2;
+		if (FuncName == GetVec3AttributeFunctionName) return ENiagaraParticleDataValueType::Vec3;
+		if (FuncName == GetVec4AttributeFunctionName) return ENiagaraParticleDataValueType::Vec4;
+		if (FuncName == GetBoolAttributeFunctionName) return ENiagaraParticleDataValueType::Bool;
+		if (FuncName == GetColorAttributeFunctionName) return ENiagaraParticleDataValueType::Color;
+		if (FuncName == GetQuatAttributeFunctionName) return ENiagaraParticleDataValueType::Quat;
+		return ENiagaraParticleDataValueType::Invalid;
+	}
+
+	virtual void Bind(const FNiagaraDataInterfaceParamRef& ParamRef, const class FShaderParameterMap& ParameterMap) override
+	{
+		IDToIndexTableParam.Bind(ParameterMap, *(IDToIndexTableBaseName + ParamRef.ParameterInfo.DataInterfaceHLSLSymbol));
+		InputFloatBufferParam.Bind(ParameterMap, *(InputFloatBufferBaseName + ParamRef.ParameterInfo.DataInterfaceHLSLSymbol));
+		InputIntBufferParam.Bind(ParameterMap, *(InputIntBufferBaseName + ParamRef.ParameterInfo.DataInterfaceHLSLSymbol));
+		ParticleStrideFloatParam.Bind(ParameterMap, *(ParticleStrideFloatBaseName + ParamRef.ParameterInfo.DataInterfaceHLSLSymbol));
+		ParticleStrideIntParam.Bind(ParameterMap, *(ParticleStrideIntBaseName + ParamRef.ParameterInfo.DataInterfaceHLSLSymbol));
+		AttributeIndicesParam.Bind(ParameterMap, *(AttributeIndicesBaseName + ParamRef.ParameterInfo.DataInterfaceHLSLSymbol));
+		AcquireTagRegisterIndexParam.Bind(ParameterMap, *(AcquireTagRegisterIndexBaseName + ParamRef.ParameterInfo.DataInterfaceHLSLSymbol));
+
+		int32 NumFuncs = ParamRef.ParameterInfo.GeneratedFunctions.Num();
+		AttributeNames.SetNum(NumFuncs);
+		AttributeTypes.SetNum(NumFuncs);
+		for (int32 FuncIdx = 0; FuncIdx < NumFuncs; ++FuncIdx)
+		{
+			const FNiagaraDataInterfaceGeneratedFunction& Func = ParamRef.ParameterInfo.GeneratedFunctions[FuncIdx];
+			static const FName NAME_Attribute("Attribute");
+			const FName* AttributeName = Func.FindSpecifierValue(NAME_Attribute);
+			if (AttributeName != nullptr)
+			{
+				AttributeNames[FuncIdx] = *AttributeName;
+				AttributeTypes[FuncIdx] = GetValueTypeFromFuncName(Func.DefinitionName);
+			}
+			else
+			{
+				AttributeNames[FuncIdx] = NAME_None;
+				AttributeTypes[FuncIdx] = ENiagaraParticleDataValueType::Invalid;
+			}
+		}
+
+		AttributeIndices.SetNum(AttributeNames.Num());
+	}
+
+	virtual void Serialize(FArchive& Ar) override
+	{
+		Ar << IDToIndexTableParam;
+		Ar << InputFloatBufferParam;
+		Ar << InputIntBufferParam;
+		Ar << ParticleStrideFloatParam;
+		Ar << ParticleStrideIntParam;
+		Ar << AttributeIndicesParam;
+		Ar << AcquireTagRegisterIndexParam;
+		Ar << AttributeNames;
+		Ar << AttributeTypes;
+
+		AttributeIndices.SetNum(AttributeNames.Num());
+	}
+
+	void SetErrorParams(FRHICommandList& RHICmdList, FRHIComputeShader* ComputeShader) const
+	{
+		CachedDataSet = nullptr;
+
+		for (int AttrIdx = 0; AttrIdx < AttributeIndices.Num(); ++AttrIdx)
+		{
+			AttributeIndices[AttrIdx] = -1;
+		}
+		AcquireTagRegisterIndex = -1;
+
+		SetSRVParameter(RHICmdList, ComputeShader, IDToIndexTableParam, FNiagaraRenderer::GetDummyIntBuffer().SRV);
+		SetSRVParameter(RHICmdList, ComputeShader, InputFloatBufferParam, FNiagaraRenderer::GetDummyFloatBuffer().SRV);
+		SetSRVParameter(RHICmdList, ComputeShader, InputIntBufferParam, FNiagaraRenderer::GetDummyIntBuffer().SRV);
+		SetShaderValue(RHICmdList, ComputeShader, ParticleStrideFloatParam, 0);
+		SetShaderValue(RHICmdList, ComputeShader, ParticleStrideIntParam, 0);
+		SetShaderValueArray(RHICmdList, ComputeShader, AttributeIndicesParam, AttributeIndices.GetData(), AttributeIndices.Num());
+		SetShaderValue(RHICmdList, ComputeShader, AcquireTagRegisterIndexParam, AcquireTagRegisterIndex);
+	}
+
+	bool CheckVariableType(const FNiagaraTypeDefinition& VarType, ENiagaraParticleDataValueType AttributeType) const
+	{
+		switch (AttributeType)
+		{
+			case ENiagaraParticleDataValueType::Int: return VarType == FNiagaraTypeDefinition::GetIntDef();
+			case ENiagaraParticleDataValueType::Float: return VarType == FNiagaraTypeDefinition::GetFloatDef();
+			case ENiagaraParticleDataValueType::Vec2: return VarType == FNiagaraTypeDefinition::GetVec2Def();
+			case ENiagaraParticleDataValueType::Vec3: return VarType == FNiagaraTypeDefinition::GetVec3Def();
+			case ENiagaraParticleDataValueType::Vec4: return VarType == FNiagaraTypeDefinition::GetVec4Def();
+			case ENiagaraParticleDataValueType::Bool: return VarType == FNiagaraTypeDefinition::GetBoolDef();
+			case ENiagaraParticleDataValueType::Color: return VarType == FNiagaraTypeDefinition::GetColorDef();
+			case ENiagaraParticleDataValueType::Quat: return VarType == FNiagaraTypeDefinition::GetQuatDef();
+			default: return false;
+		}
+	}
+
+	void FindAttributeIndices(FNiagaraDataSet* SourceDataSet, const TCHAR* SourceEmitterName) const
+	{
+		check(AttributeIndices.Num() == AttributeNames.Num());
+
+		AcquireTagRegisterIndex = -1;
+		const TArray<FNiagaraVariable>& SourceEmitterVariables = SourceDataSet->GetVariables();
+		const TArray<FNiagaraVariableLayoutInfo>& SourceEmitterVariableLayouts = SourceDataSet->GetVariableLayouts();
+		for (int AttrNameIdx = 0; AttrNameIdx < AttributeNames.Num(); ++AttrNameIdx)
+		{
+			bool FoundVariable = false;
+			for (int VarIdx = 0; VarIdx < SourceEmitterVariables.Num(); ++VarIdx)
+			{
+				const FNiagaraVariable& Var = SourceEmitterVariables[VarIdx];
+
+				if (AcquireTagRegisterIndex == -1 && Var.GetName().ToString() == TEXT("ID"))
+				{
+					AcquireTagRegisterIndex = SourceEmitterVariableLayouts[VarIdx].Int32ComponentStart + 1;
+				}
+
+				if (Var.GetName() == AttributeNames[AttrNameIdx])
+				{
+					ENiagaraParticleDataValueType AttributeType = AttributeTypes[AttrNameIdx];
+					if (CheckVariableType(Var.GetType(), AttributeType))
+					{
+						const FNiagaraVariableLayoutInfo& Layout = SourceEmitterVariableLayouts[VarIdx];
+						AttributeIndices[AttrNameIdx] = (AttributeType == ENiagaraParticleDataValueType::Int || AttributeType == ENiagaraParticleDataValueType::Bool) ? Layout.Int32ComponentStart : Layout.FloatComponentStart;
+					}
+					else
+					{
+						UE_LOG(LogNiagara, Error, TEXT("Variable '%s' in emitter '%s' has type '%s', but particle read DI tried to access it as '%s'."),
+							*Var.GetName().ToString(), SourceEmitterName, *Var.GetType().GetName(), NiagaraParticleDataValueTypeName(AttributeType)
+						);
+						AttributeIndices[AttrNameIdx] = -1;
+					}
+					FoundVariable = true;
+					break;
+				}
+			}
+
+			if (!FoundVariable)
+			{
+				UE_LOG(LogNiagara, Error, TEXT("Particle read DI is trying to access inexistent variable '%s' in emitter '%s'."), *AttributeNames[AttrNameIdx].ToString(), SourceEmitterName);
+				AttributeIndices[AttrNameIdx] = -1;
+			}
+		}
+
+		if (AcquireTagRegisterIndex == -1)
+		{
+			UE_LOG(LogNiagara, Error, TEXT("Particle read DI cannot find ID variable in emitter '%s'."), SourceEmitterName);
+		}
+	}
+
+	virtual void Set(FRHICommandList& RHICmdList, const FNiagaraDataInterfaceSetArgs& Context) const override
+	{
+		check(IsInRenderingThread());
+
+		FRHIComputeShader* ComputeShader = Context.Shader->GetComputeShader();
+
+		if (!InputFloatBufferParam.IsBound() && !InputIntBufferParam.IsBound())
+		{
+			// This DI instance didn't generate any reachable code, probably because all the values it sets are overwritten by
+			// other DIs. Don't bother with it.
+			SetErrorParams(RHICmdList, ComputeShader);
+			return;
+		}
+
+		FNiagaraDataInterfaceProxyParticleRead* Proxy = static_cast<FNiagaraDataInterfaceProxyParticleRead*>(Context.DataInterface);
+		check(Proxy);
+		if (Proxy->SourceEmitterGPUContext == nullptr)
+		{
+			// This means the source emitter isn't running on GPU.
+			if (!bSourceEmitterNotGPUErrorShown)
+			{
+				UE_LOG(LogNiagara, Error, TEXT("GPU particle read DI is set to access CPU emitter '%s'."), *Proxy->SourceEmitterName);
+				bSourceEmitterNotGPUErrorShown = true;
+			}
+			SetErrorParams(RHICmdList, ComputeShader);
+			return;
+		}
+
+		bSourceEmitterNotGPUErrorShown = false;
+
+		FNiagaraDataSet* SourceDataSet = Proxy->SourceEmitterGPUContext->MainDataSet;
+		if (!SourceDataSet)
+		{
+			SetErrorParams(RHICmdList, ComputeShader);
+			return;
+		}
+
+		FNiagaraDataBuffer* SourceData = SourceDataSet->GetCurrentData();
+		if (!SourceData)
+		{
+			SetErrorParams(RHICmdList, ComputeShader);
+			return;
+		}
+
+		if (CachedDataSet != SourceDataSet)
+		{
+			FindAttributeIndices(SourceDataSet, *Proxy->SourceEmitterName);
+			CachedDataSet = SourceDataSet;
+		}
+
+		const uint32 ParticleStrideFloat = SourceData->GetFloatStride() / sizeof(float);
+		const uint32 ParticleStrideInt = SourceData->GetInt32Stride() / sizeof(int32);
+
+		SetSRVParameter(RHICmdList, ComputeShader, IDToIndexTableParam, SourceData->GetGPUIDToIndexTable().SRV);
+		SetSRVParameter(RHICmdList, ComputeShader, InputFloatBufferParam, SourceData->GetGPUBufferFloat().SRV);
+		SetSRVParameter(RHICmdList, ComputeShader, InputIntBufferParam, SourceData->GetGPUBufferInt().SRV);
+		SetShaderValue(RHICmdList, ComputeShader, ParticleStrideFloatParam, ParticleStrideFloat);
+		SetShaderValue(RHICmdList, ComputeShader, ParticleStrideIntParam, ParticleStrideInt);
+		SetShaderValueArray(RHICmdList, ComputeShader, AttributeIndicesParam, AttributeIndices.GetData(), AttributeIndices.Num());
+		SetShaderValue(RHICmdList, ComputeShader, AcquireTagRegisterIndexParam, AcquireTagRegisterIndex);
+	}
+
+private:
+	FShaderResourceParameter IDToIndexTableParam;
+	FShaderResourceParameter InputFloatBufferParam;
+	FShaderResourceParameter InputIntBufferParam;
+	FShaderParameter ParticleStrideFloatParam;
+	FShaderParameter ParticleStrideIntParam;
+	FShaderParameter AttributeIndicesParam;
+	FShaderParameter AcquireTagRegisterIndexParam;
+	TArray<FName> AttributeNames;
+	TArray<ENiagaraParticleDataValueType> AttributeTypes;
+	mutable TArray<int32> AttributeIndices;
+	mutable int32 AcquireTagRegisterIndex;
+	mutable FNiagaraDataSet* CachedDataSet;
+	mutable bool bSourceEmitterNotGPUErrorShown;
+};
+
 UNiagaraDataInterfaceParticleRead::UNiagaraDataInterfaceParticleRead(FObjectInitializer const& ObjectInitializer)
 	: Super(ObjectInitializer)
 {
+	Proxy = MakeShared<FNiagaraDataInterfaceProxyParticleRead, ESPMode::ThreadSafe>();
 }
 
 void UNiagaraDataInterfaceParticleRead::PostInitProperties()
@@ -567,19 +880,179 @@ bool UNiagaraDataInterfaceParticleRead::CopyToInternal(UNiagaraDataInterface* De
 	return true;
 }
 
-bool UNiagaraDataInterfaceParticleRead::GetFunctionHLSL(const FName& DefinitionFunctionName, FString InstanceFunctionName, FNiagaraDataInterfaceGPUParamInfo& ParamInfo, FString& OutHLSL)
+void UNiagaraDataInterfaceParticleRead::GetCommonHLSL(FString& OutHLSL)
 {
-	return false;
 }
 
-void UNiagaraDataInterfaceParticleRead::GetParameterDefinitionHLSL(FNiagaraDataInterfaceGPUParamInfo& ParamInfo, FString& OutHLSL)
+void UNiagaraDataInterfaceParticleRead::GetParameterDefinitionHLSL(const FNiagaraDataInterfaceGPUParamInfo& ParamInfo, FString& OutHLSL)
 {
+	static const TCHAR *FormatDeclarations = TEXT(
+		"Buffer<int> {IDToIndexTableName};\n"
+		"Buffer<float> {InputFloatBufferName};\n"
+		"Buffer<int> {InputIntBufferName};\n"
+		"uint {ParticleStrideFloatName};\n"
+		"uint {ParticleStrideIntName};\n"
+		"int {AttributeIndicesName}[{AttributeCount}];\n"
+		"int {AcquireTagRegisterIndexName};\n\n"
+	);
+
+	TMap<FString, FStringFormatArg> ArgsDeclarations;
+	ArgsDeclarations.Add(TEXT("IDToIndexTableName"), IDToIndexTableBaseName + ParamInfo.DataInterfaceHLSLSymbol);
+	ArgsDeclarations.Add(TEXT("InputFloatBufferName"), InputFloatBufferBaseName + ParamInfo.DataInterfaceHLSLSymbol);
+	ArgsDeclarations.Add(TEXT("InputIntBufferName"), InputIntBufferBaseName + ParamInfo.DataInterfaceHLSLSymbol);
+	ArgsDeclarations.Add(TEXT("ParticleStrideFloatName"), ParticleStrideFloatBaseName + ParamInfo.DataInterfaceHLSLSymbol);
+	ArgsDeclarations.Add(TEXT("ParticleStrideIntName"), ParticleStrideIntBaseName + ParamInfo.DataInterfaceHLSLSymbol);
+	ArgsDeclarations.Add(TEXT("AttributeIndicesName"), AttributeIndicesBaseName + ParamInfo.DataInterfaceHLSLSymbol);
+	ArgsDeclarations.Add(TEXT("AttributeCount"), ParamInfo.GeneratedFunctions.Num());
+	ArgsDeclarations.Add(TEXT("AcquireTagRegisterIndexName"), AcquireTagRegisterIndexBaseName + ParamInfo.DataInterfaceHLSLSymbol);
+
+	OutHLSL += FString::Format(FormatDeclarations, ArgsDeclarations);
+}
 
+static bool GenerateGetFunctionHLSL(const FNiagaraDataInterfaceGPUParamInfo& ParamInfo, const FNiagaraDataInterfaceGeneratedFunction& FunctionInfo, int FunctionInstanceIndex, ENiagaraParticleDataComponentType ComponentType, int NumComponents, FString& OutHLSL)
+{
+	static const TCHAR* FuncTemplate = TEXT(
+		"void {FunctionName}(NiagaraID In_ParticleID, out bool Out_Valid, out {ValueType} Out_Value)\n"
+		"{\n"
+		"    int RegisterIndex = {AttributeIndicesName}[{FunctionInstanceIndex}];\n"
+		"    int ParticleIndex = (RegisterIndex != -1) ? {IDToIndexTableName}[In_ParticleID.Index] : -1;\n"
+		"    int AcquireTag = (ParticleIndex != -1) ? {InputIntBufferName}[{AcquireTagRegisterIndexName}*{ParticleStrideIntName} + ParticleIndex] : 0;\n"
+		"    if(ParticleIndex != -1 && In_ParticleID.AcquireTag == AcquireTag)\n"
+		"    {\n"
+		"        Out_Valid = true;\n"
+		"{FetchValueCode}"
+		"    }\n"
+		"    else\n"
+		"    {\n"
+		"        Out_Valid = false;\n"
+		"        Out_Value = {ValueType}(0{ExtraDefaultValues});\n"
+		"    }\n"
+		"}\n\n"
+	);
+
+	const FString ParticleStrideFloatName = ParticleStrideFloatBaseName + ParamInfo.DataInterfaceHLSLSymbol;
+	const FString ParticleStrideIntName = ParticleStrideIntBaseName + ParamInfo.DataInterfaceHLSLSymbol;
+	const FString InputFloatBufferName = InputFloatBufferBaseName + ParamInfo.DataInterfaceHLSLSymbol;
+	const FString InputIntBufferName = InputIntBufferBaseName + ParamInfo.DataInterfaceHLSLSymbol;
+
+	const TCHAR* ComponentTypeName;
+	const TCHAR* InputBufferName;
+	const TCHAR* InputBufferStrideName;
+	switch (ComponentType)
+	{
+		case ENiagaraParticleDataComponentType::Float:
+			ComponentTypeName = TEXT("float");
+			InputBufferName = *InputFloatBufferName;
+			InputBufferStrideName = *ParticleStrideFloatName;
+			break;
+		case ENiagaraParticleDataComponentType::Int:
+			ComponentTypeName = TEXT("int");
+			InputBufferName = *InputIntBufferName;
+			InputBufferStrideName = *ParticleStrideIntName;
+			break;
+		case ENiagaraParticleDataComponentType::Bool:
+			ComponentTypeName = TEXT("bool");
+			InputBufferName = *InputIntBufferName;
+			InputBufferStrideName = *ParticleStrideIntName;
+			break;
+		default:
+			UE_LOG(LogNiagara, Error, TEXT("Unknown component type %d while generating function %s"), ComponentType, *FunctionInfo.InstanceName);
+			return false;
+	}
+
+	FString ExtraDefaultValues;
+	for (int ComponentIndex = 1; ComponentIndex < NumComponents; ++ComponentIndex)
+	{
+		ExtraDefaultValues += TEXT(", 0");
+	}
+
+	FString FetchValueCode;
+	for (int ComponentIndex = 0; ComponentIndex < NumComponents; ++ComponentIndex)
+	{
+		static const TCHAR* ComponentNames[] = { TEXT(".x"), TEXT(".y"), TEXT(".z"), TEXT(".w") };
+		const TCHAR* ComponentName = (NumComponents > 1) ? ComponentNames[ComponentIndex] : TEXT("");
+		const FString FetchComponentCode = FString::Printf(TEXT("        Out_Value%s = %s(%s[(RegisterIndex + %d)*%s + ParticleIndex]);\n"), ComponentName, ComponentTypeName, InputBufferName, ComponentIndex, InputBufferStrideName);
+		FetchValueCode += FetchComponentCode;
+	}
+
+	FString ValueTypeName = (NumComponents > 1) ? FString::Printf(TEXT("%s%d"), ComponentTypeName, NumComponents) : FString(ComponentTypeName);
+
+	TMap<FString, FStringFormatArg> FuncTemplateArgs;
+	FuncTemplateArgs.Add(TEXT("FunctionName"), FunctionInfo.InstanceName);
+	FuncTemplateArgs.Add(TEXT("ValueType"), ValueTypeName);
+	FuncTemplateArgs.Add(TEXT("AttributeIndicesName"), AttributeIndicesBaseName + ParamInfo.DataInterfaceHLSLSymbol);
+	FuncTemplateArgs.Add(TEXT("FunctionInstanceIndex"), FunctionInstanceIndex);
+	FuncTemplateArgs.Add(TEXT("IDToIndexTableName"), IDToIndexTableBaseName + ParamInfo.DataInterfaceHLSLSymbol);
+	FuncTemplateArgs.Add(TEXT("InputIntBufferName"), InputIntBufferName);
+	FuncTemplateArgs.Add(TEXT("AcquireTagRegisterIndexName"), AcquireTagRegisterIndexBaseName + ParamInfo.DataInterfaceHLSLSymbol);
+	FuncTemplateArgs.Add(TEXT("ParticleStrideIntName"), ParticleStrideIntName);
+	FuncTemplateArgs.Add(TEXT("FetchValueCode"), FetchValueCode);
+	FuncTemplateArgs.Add(TEXT("ExtraDefaultValues"), ExtraDefaultValues);
+	
+	OutHLSL += FString::Format(FuncTemplate, FuncTemplateArgs);
+
+	return true;
+}
+
+bool UNiagaraDataInterfaceParticleRead::GetFunctionHLSL(const FNiagaraDataInterfaceGPUParamInfo& ParamInfo, const FNiagaraDataInterfaceGeneratedFunction& FunctionInfo, int FunctionInstanceIndex, FString& OutHLSL)
+{
+	if (FunctionInfo.DefinitionName == GetIntAttributeFunctionName)
+	{
+		return GenerateGetFunctionHLSL(ParamInfo, FunctionInfo, FunctionInstanceIndex, ENiagaraParticleDataComponentType::Int, 1, OutHLSL);
+	}
+
+	if (FunctionInfo.DefinitionName == GetFloatAttributeFunctionName)
+	{
+		return GenerateGetFunctionHLSL(ParamInfo, FunctionInfo, FunctionInstanceIndex, ENiagaraParticleDataComponentType::Float, 1, OutHLSL);
+	}
+
+	if (FunctionInfo.DefinitionName == GetVec2AttributeFunctionName)
+	{
+		return GenerateGetFunctionHLSL(ParamInfo, FunctionInfo, FunctionInstanceIndex, ENiagaraParticleDataComponentType::Float, 2, OutHLSL);
+	}
+
+	if (FunctionInfo.DefinitionName == GetVec3AttributeFunctionName)
+	{
+		return GenerateGetFunctionHLSL(ParamInfo, FunctionInfo, FunctionInstanceIndex, ENiagaraParticleDataComponentType::Float, 3, OutHLSL);
+	}
+
+	if (FunctionInfo.DefinitionName == GetVec4AttributeFunctionName)
+	{
+		return GenerateGetFunctionHLSL(ParamInfo, FunctionInfo, FunctionInstanceIndex, ENiagaraParticleDataComponentType::Float, 4, OutHLSL);
+	}
+
+	if (FunctionInfo.DefinitionName == GetBoolAttributeFunctionName)
+	{
+		return GenerateGetFunctionHLSL(ParamInfo, FunctionInfo, FunctionInstanceIndex, ENiagaraParticleDataComponentType::Bool, 1, OutHLSL);
+	}
+
+	if (FunctionInfo.DefinitionName == GetColorAttributeFunctionName)
+	{
+		return GenerateGetFunctionHLSL(ParamInfo, FunctionInfo, FunctionInstanceIndex, ENiagaraParticleDataComponentType::Float, 4, OutHLSL);
+	}
+
+	if (FunctionInfo.DefinitionName == GetQuatAttributeFunctionName)
+	{
+		return GenerateGetFunctionHLSL(ParamInfo, FunctionInfo, FunctionInstanceIndex, ENiagaraParticleDataComponentType::Float, 4, OutHLSL);
+	}
+
+	return false;
 }
 
 FNiagaraDataInterfaceParametersCS* UNiagaraDataInterfaceParticleRead::ConstructComputeParameters() const
 {
-	return nullptr;
+	return new FNiagaraDataInterfaceParametersCS_ParticleRead();
+}
+
+void UNiagaraDataInterfaceParticleRead::ProvidePerInstanceDataForRenderThread(void* DataForRenderThread, void* PerInstanceData, const FNiagaraSystemInstanceID& SystemInstance)
+{
+	FNDIParticleRead_InstanceDataGPU* DataToPass = new (DataForRenderThread) FNDIParticleRead_InstanceDataGPU;
+	FNDIParticleRead_InstanceData* PIData = static_cast<FNDIParticleRead_InstanceData*>(PerInstanceData);
+	if (PIData && PIData->EmitterInstance)
+	{
+		DataToPass->SourceEmitterGPUContext = PIData->EmitterInstance->GetGPUContext();
+		DataToPass->SourceEmitterName = PIData->EmitterInstance->GetCachedEmitter()->GetUniqueEmitterName();
+	}
 }
 
-#undef LOCTEXT_NAMESPACE
\ No newline at end of file
+#undef LOCTEXT_NAMESPACE
diff --git a/Engine/Plugins/FX/Niagara/Source/Niagara/Private/NiagaraDataInterfaceRW.cpp b/Engine/Plugins/FX/Niagara/Source/Niagara/Private/NiagaraDataInterfaceRW.cpp
--- a/Engine/Plugins/FX/Niagara/Source/Niagara/Private/NiagaraDataInterfaceRW.cpp
+++ b/Engine/Plugins/FX/Niagara/Source/Niagara/Private/NiagaraDataInterfaceRW.cpp
@@ -172,7 +172,7 @@ bool UNiagaraDataInterfaceGrid3D::Equals(const UNiagaraDataInterface* Other) con
 		OtherTyped->WorldBBoxSize.Equals(WorldBBoxSize);
 }
 
-void UNiagaraDataInterfaceGrid3D::GetParameterDefinitionHLSL(FNiagaraDataInterfaceGPUParamInfo& ParamInfo, FString& OutHLSL)
+void UNiagaraDataInterfaceGrid3D::GetParameterDefinitionHLSL(const FNiagaraDataInterfaceGPUParamInfo& ParamInfo, FString& OutHLSL)
 {
 	static const TCHAR *FormatDeclarations = TEXT(R"(
 		int3 {NumVoxelsName};
@@ -187,16 +187,16 @@ void UNiagaraDataInterfaceGrid3D::GetParameterDefinitionHLSL(FNiagaraDataInterfa
 	OutHLSL += FString::Format(FormatDeclarations, ArgsDeclarations);
 }
 
-bool UNiagaraDataInterfaceGrid3D::GetFunctionHLSL(const FName& DefinitionFunctionName, FString InstanceFunctionName, FNiagaraDataInterfaceGPUParamInfo& ParamInfo, FString& OutHLSL)
+bool UNiagaraDataInterfaceGrid3D::GetFunctionHLSL(const FNiagaraDataInterfaceGPUParamInfo& ParamInfo, const FNiagaraDataInterfaceGeneratedFunction& FunctionInfo, int FunctionInstanceIndex, FString& OutHLSL)
 {
 	TMap<FString, FStringFormatArg> ArgsDeclarations = {
-		{ TEXT("FunctionName"), InstanceFunctionName},
+		{ TEXT("FunctionName"), FunctionInfo.InstanceName},
 		{ TEXT("NumVoxelsName"), NumVoxelsName + ParamInfo.DataInterfaceHLSLSymbol },
 		{ TEXT("VoxelSizeName"), VoxelSizeName + ParamInfo.DataInterfaceHLSLSymbol },
 		{ TEXT("WorldBBoxSizeName"),    WorldBBoxSizeName + ParamInfo.DataInterfaceHLSLSymbol },
 	};
 
-	if (DefinitionFunctionName == WorldBBoxSizeFunctionName)
+	if (FunctionInfo.DefinitionName == WorldBBoxSizeFunctionName)
 	{
 		static const TCHAR *FormatSample = TEXT(R"(
 			void {FunctionName}(out float3 Out_WorldBBox)
@@ -208,7 +208,7 @@ bool UNiagaraDataInterfaceGrid3D::GetFunctionHLSL(const FName& DefinitionFunctio
 		OutHLSL += FString::Format(FormatSample, ArgsDeclarations);
 		return true;
 	}
-	else if (DefinitionFunctionName == NumVoxelsFunctionName)
+	else if (FunctionInfo.DefinitionName == NumVoxelsFunctionName)
 	{
 		static const TCHAR *FormatSample = TEXT(R"(
 			void {FunctionName}(out int Out_NumVoxelsX, out int Out_NumVoxelsY, out int Out_NumVoxelsZ)
@@ -222,7 +222,7 @@ bool UNiagaraDataInterfaceGrid3D::GetFunctionHLSL(const FName& DefinitionFunctio
 		OutHLSL += FString::Format(FormatSample, ArgsDeclarations);
 		return true;
 	}
-	else if (DefinitionFunctionName == SimulationToUnitFunctionName)
+	else if (FunctionInfo.DefinitionName == SimulationToUnitFunctionName)
 	{
 		static const TCHAR *FormatSample = TEXT(R"(
 			void {FunctionName}(float3 In_Simulation, float4x4 In_SimulationToUnitTransform, out float3 Out_Unit)
@@ -234,7 +234,7 @@ bool UNiagaraDataInterfaceGrid3D::GetFunctionHLSL(const FName& DefinitionFunctio
 		OutHLSL += FString::Format(FormatSample, ArgsDeclarations);
 		return true;
 	}
-	else if (DefinitionFunctionName == UnitToIndexFunctionName)
+	else if (FunctionInfo.DefinitionName == UnitToIndexFunctionName)
 	{
 		static const TCHAR *FormatSample = TEXT(R"(
 			void {FunctionName}(float3 In_Unit, out int Out_IndexX, out int Out_IndexY, out int Out_IndexZ)
@@ -249,7 +249,7 @@ bool UNiagaraDataInterfaceGrid3D::GetFunctionHLSL(const FName& DefinitionFunctio
 		OutHLSL += FString::Format(FormatSample, ArgsDeclarations);
 		return true;
 	}
-	else if (DefinitionFunctionName == IndexToLinearFunctionName)
+	else if (FunctionInfo.DefinitionName == IndexToLinearFunctionName)
 	{
 		static const TCHAR *FormatSample = TEXT(R"(
 			void {FunctionName}(int In_IndexX, int In_IndexY, int In_IndexZ, out int Out_Linear)
@@ -261,7 +261,7 @@ bool UNiagaraDataInterfaceGrid3D::GetFunctionHLSL(const FName& DefinitionFunctio
 		OutHLSL += FString::Format(FormatSample, ArgsDeclarations);
 		return true;
 	}
-	else if (DefinitionFunctionName == VoxelSizeFunctionName)
+	else if (FunctionInfo.DefinitionName == VoxelSizeFunctionName)
 	{
 		static const TCHAR *FormatSample = TEXT(R"(
 			void {FunctionName}(out float3 Out_VoxelSize)
@@ -511,7 +511,7 @@ bool UNiagaraDataInterfaceGrid2D::Equals(const UNiagaraDataInterface* Other) con
 		OtherTyped->WorldBBoxSize.Equals(WorldBBoxSize);
 }
 
-void UNiagaraDataInterfaceGrid2D::GetParameterDefinitionHLSL(FNiagaraDataInterfaceGPUParamInfo& ParamInfo, FString& OutHLSL)
+void UNiagaraDataInterfaceGrid2D::GetParameterDefinitionHLSL(const FNiagaraDataInterfaceGPUParamInfo& ParamInfo, FString& OutHLSL)
 {
 	static const TCHAR *FormatDeclarations = TEXT(R"(
 		int2 {NumCellsName};
@@ -526,16 +526,16 @@ void UNiagaraDataInterfaceGrid2D::GetParameterDefinitionHLSL(FNiagaraDataInterfa
 	OutHLSL += FString::Format(FormatDeclarations, ArgsDeclarations);
 }
 
-bool UNiagaraDataInterfaceGrid2D::GetFunctionHLSL(const FName& DefinitionFunctionName, FString InstanceFunctionName, FNiagaraDataInterfaceGPUParamInfo& ParamInfo, FString& OutHLSL)
+bool UNiagaraDataInterfaceGrid2D::GetFunctionHLSL(const FNiagaraDataInterfaceGPUParamInfo& ParamInfo, const FNiagaraDataInterfaceGeneratedFunction& FunctionInfo, int FunctionInstanceIndex, FString& OutHLSL)
 {
 	TMap<FString, FStringFormatArg> ArgsDeclarations = {
-		{ TEXT("FunctionName"), InstanceFunctionName},
+		{ TEXT("FunctionName"), FunctionInfo.InstanceName},
 		{ TEXT("NumCellsName"), NumCellsName + ParamInfo.DataInterfaceHLSLSymbol },
 		{ TEXT("CellSizeName"), CellSizeName + ParamInfo.DataInterfaceHLSLSymbol },		
 		{ TEXT("WorldBBoxSizeName"),    WorldBBoxSizeName + ParamInfo.DataInterfaceHLSLSymbol },
 	};
 
-	if (DefinitionFunctionName == WorldBBoxSizeFunctionName)
+	if (FunctionInfo.DefinitionName == WorldBBoxSizeFunctionName)
 	{
 		static const TCHAR *FormatSample = TEXT(R"(
 			void {FunctionName}(out float2 Out_WorldBBox)
@@ -547,7 +547,7 @@ bool UNiagaraDataInterfaceGrid2D::GetFunctionHLSL(const FName& DefinitionFunctio
 		OutHLSL += FString::Format(FormatSample, ArgsDeclarations);
 		return true;
 	}
-	else if (DefinitionFunctionName == NumCellsFunctionName)
+	else if (FunctionInfo.DefinitionName == NumCellsFunctionName)
 	{
 		static const TCHAR *FormatSample = TEXT(R"(
 			void {FunctionName}(out int Out_NumCellsX, out int Out_NumCellsY)
@@ -560,7 +560,7 @@ bool UNiagaraDataInterfaceGrid2D::GetFunctionHLSL(const FName& DefinitionFunctio
 		OutHLSL += FString::Format(FormatSample, ArgsDeclarations);
 		return true;
 	}
-	else if (DefinitionFunctionName == SimulationToUnitFunctionName)
+	else if (FunctionInfo.DefinitionName == SimulationToUnitFunctionName)
 	{
 		static const TCHAR *FormatSample = TEXT(R"(
 			void {FunctionName}(float3 In_Simulation, float4x4 In_SimulationToUnitTransform, out float3 Out_Unit)
@@ -572,7 +572,7 @@ bool UNiagaraDataInterfaceGrid2D::GetFunctionHLSL(const FName& DefinitionFunctio
 		OutHLSL += FString::Format(FormatSample, ArgsDeclarations);
 		return true;
 	}
-	else if (DefinitionFunctionName == UnitToSimulationFunctionName)
+	else if (FunctionInfo.DefinitionName == UnitToSimulationFunctionName)
 	{
 		static const TCHAR *FormatSample = TEXT(R"(
 			void {FunctionName}(float3 In_Unit, float4x4 In_UnitToSimulationTransform, out float3 Out_Simulation)
@@ -584,7 +584,7 @@ bool UNiagaraDataInterfaceGrid2D::GetFunctionHLSL(const FName& DefinitionFunctio
 		OutHLSL += FString::Format(FormatSample, ArgsDeclarations);
 		return true;
 	}
-	else if (DefinitionFunctionName == UnitToIndexFunctionName)
+	else if (FunctionInfo.DefinitionName == UnitToIndexFunctionName)
 	{
 		static const TCHAR *FormatSample = TEXT(R"(
 			void {FunctionName}(float2 In_Unit, out int Out_IndexX, out int Out_IndexY)
@@ -598,7 +598,7 @@ bool UNiagaraDataInterfaceGrid2D::GetFunctionHLSL(const FName& DefinitionFunctio
 		OutHLSL += FString::Format(FormatSample, ArgsDeclarations);
 		return true;
 	}
-	else if (DefinitionFunctionName == IndexToUnitFunctionName)
+	else if (FunctionInfo.DefinitionName == IndexToUnitFunctionName)
 	{
 		static const TCHAR *FormatSample = TEXT(R"(
 			void {FunctionName}(float In_IndexX, float In_IndexY, out float3 Out_Unit)
@@ -610,7 +610,7 @@ bool UNiagaraDataInterfaceGrid2D::GetFunctionHLSL(const FName& DefinitionFunctio
 		OutHLSL += FString::Format(FormatSample, ArgsDeclarations);
 		return true;
 	}
-	else if (DefinitionFunctionName == IndexToUnitStaggeredXFunctionName)
+	else if (FunctionInfo.DefinitionName == IndexToUnitStaggeredXFunctionName)
 	{
 		static const TCHAR *FormatSample = TEXT(R"(
 			void {FunctionName}(float In_IndexX, float In_IndexY, out float3 Out_Unit)
@@ -622,7 +622,7 @@ bool UNiagaraDataInterfaceGrid2D::GetFunctionHLSL(const FName& DefinitionFunctio
 		OutHLSL += FString::Format(FormatSample, ArgsDeclarations);
 		return true;
 	}
-	else if (DefinitionFunctionName == IndexToUnitStaggeredYFunctionName)
+	else if (FunctionInfo.DefinitionName == IndexToUnitStaggeredYFunctionName)
 	{
 		static const TCHAR *FormatSample = TEXT(R"(
 			void {FunctionName}(float In_IndexX, float In_IndexY, out float3 Out_Unit)
@@ -634,7 +634,7 @@ bool UNiagaraDataInterfaceGrid2D::GetFunctionHLSL(const FName& DefinitionFunctio
 		OutHLSL += FString::Format(FormatSample, ArgsDeclarations);
 		return true;
 	}
-	else if (DefinitionFunctionName == IndexToLinearFunctionName)
+	else if (FunctionInfo.DefinitionName == IndexToLinearFunctionName)
 	{
 		static const TCHAR *FormatSample = TEXT(R"(
 			void {FunctionName}(int In_IndexX, int In_IndexY, out int Out_Linear)
@@ -646,7 +646,7 @@ bool UNiagaraDataInterfaceGrid2D::GetFunctionHLSL(const FName& DefinitionFunctio
 		OutHLSL += FString::Format(FormatSample, ArgsDeclarations);
 		return true;
 	}
-	else if (DefinitionFunctionName == LinearToIndexFunctionName)
+	else if (FunctionInfo.DefinitionName == LinearToIndexFunctionName)
 	{
 		static const TCHAR *FormatSample = TEXT(R"(
 			void {FunctionName}(int In_Linear, out int Out_IndexX, out int Out_IndexY)
@@ -659,7 +659,7 @@ bool UNiagaraDataInterfaceGrid2D::GetFunctionHLSL(const FName& DefinitionFunctio
 		OutHLSL += FString::Format(FormatSample, ArgsDeclarations);
 		return true;
 	}
-	else if (DefinitionFunctionName == CellSizeFunctionName)
+	else if (FunctionInfo.DefinitionName == CellSizeFunctionName)
 	{
 		static const TCHAR *FormatSample = TEXT(R"(
 			void {FunctionName}(out float2 Out_CellSize)
diff --git a/Engine/Plugins/FX/Niagara/Source/Niagara/Private/NiagaraDataInterfaceSimpleCounter.cpp b/Engine/Plugins/FX/Niagara/Source/Niagara/Private/NiagaraDataInterfaceSimpleCounter.cpp
--- a/Engine/Plugins/FX/Niagara/Source/Niagara/Private/NiagaraDataInterfaceSimpleCounter.cpp
+++ b/Engine/Plugins/FX/Niagara/Source/Niagara/Private/NiagaraDataInterfaceSimpleCounter.cpp
@@ -84,12 +84,12 @@ void UNiagaraDataInterfaceSimpleCounter::GetFunctions(TArray<FNiagaraFunctionSig
 	OutFunctions.Add(Sig2);
 }
 
-bool UNiagaraDataInterfaceSimpleCounter::GetFunctionHLSL(const FName& DefinitionFunctionName, FString InstanceFunctionName, FNiagaraDataInterfaceGPUParamInfo& ParamInfo, FString& OutHLSL)
+bool UNiagaraDataInterfaceSimpleCounter::GetFunctionHLSL(const FNiagaraDataInterfaceGPUParamInfo& ParamInfo, const FNiagaraDataInterfaceGeneratedFunction& FunctionInfo, int FunctionInstanceIndex, FString& OutHLSL)
 {
 	return false;
 }
 
-void UNiagaraDataInterfaceSimpleCounter::GetParameterDefinitionHLSL(FNiagaraDataInterfaceGPUParamInfo& ParamInfo, FString& OutHLSL)
+void UNiagaraDataInterfaceSimpleCounter::GetParameterDefinitionHLSL(const FNiagaraDataInterfaceGPUParamInfo& ParamInfo, FString& OutHLSL)
 {
 }
 
diff --git a/Engine/Plugins/FX/Niagara/Source/Niagara/Private/NiagaraDataInterfaceSkeletalMesh.cpp b/Engine/Plugins/FX/Niagara/Source/Niagara/Private/NiagaraDataInterfaceSkeletalMesh.cpp
--- a/Engine/Plugins/FX/Niagara/Source/Niagara/Private/NiagaraDataInterfaceSkeletalMesh.cpp
+++ b/Engine/Plugins/FX/Niagara/Source/Niagara/Private/NiagaraDataInterfaceSkeletalMesh.cpp
@@ -1956,12 +1956,12 @@ bool UNiagaraDataInterfaceSkeletalMesh::AppendCompileHash(FNiagaraCompileHashVis
 	return true;
 }
 
-bool UNiagaraDataInterfaceSkeletalMesh::GetFunctionHLSL(const FName& DefinitionFunctionName, FString InstanceFunctionName, FNiagaraDataInterfaceGPUParamInfo& ParamInfo, FString& OutHLSL) 
+bool UNiagaraDataInterfaceSkeletalMesh::GetFunctionHLSL(const FNiagaraDataInterfaceGPUParamInfo& ParamInfo, const FNiagaraDataInterfaceGeneratedFunction& FunctionInfo, int FunctionInstanceIndex, FString& OutHLSL)
 {
 	FNDISkeletalMeshParametersName ParamNames;
 	GetNiagaraDataInterfaceParametersName(ParamNames, ParamInfo.DataInterfaceHLSLSymbol);
 	TMap<FString, FStringFormatArg> ArgsSample = {
-		{TEXT("InstanceFunctionName"), InstanceFunctionName},
+		{TEXT("InstanceFunctionName"), FunctionInfo.InstanceName},
 		{TEXT("MeshTriCoordinateStructName"), TEXT("MeshTriCoordinate")},
 		{TEXT("MeshTriangleCount"), ParamNames.MeshTriangleCountName},
 		{TEXT("MeshVertexCount"), ParamNames.MeshVertexCountName},
@@ -1969,148 +1969,148 @@ bool UNiagaraDataInterfaceSkeletalMesh::GetFunctionHLSL(const FName& DefinitionF
 	};
 
 	// Triangle Sampling
-	if (DefinitionFunctionName == FSkeletalMeshInterfaceHelper::RandomTriCoordName)
+	if (FunctionInfo.DefinitionName == FSkeletalMeshInterfaceHelper::RandomTriCoordName)
 	{
 		static const TCHAR* FormatSample = TEXT("void {InstanceFunctionName} (NiagaraRandInfo InRandomInfo, out {MeshTriCoordinateStructName} OutCoord) { {GetDISkelMeshContextName} DISKelMesh_RandomTriCoord(DIContext, InRandomInfo.Seed1, InRandomInfo.Seed2, InRandomInfo.Seed3, OutCoord.Tri, OutCoord.BaryCoord); }");
 		OutHLSL += FString::Format(FormatSample, ArgsSample);
 	}
-	else if (DefinitionFunctionName == FSkeletalMeshInterfaceHelper::GetSkinnedTriangleDataWSName)
+	else if (FunctionInfo.DefinitionName == FSkeletalMeshInterfaceHelper::GetSkinnedTriangleDataWSName)
 	{
 		static const TCHAR* FormatSample = TEXT("void {InstanceFunctionName} (in {MeshTriCoordinateStructName} InCoord, out float3 OutPosition, out float3 OutVelocity, out float3 OutNormal, out float3 OutBinormal, out float3 OutTangent) { {GetDISkelMeshContextName} DISKelMesh_GetSkinnedTriangleDataWS(DIContext, InCoord.Tri, InCoord.BaryCoord, OutPosition, OutVelocity, OutNormal, OutBinormal, OutTangent); }");
 		OutHLSL += FString::Format(FormatSample, ArgsSample);
 	}
-	else if (DefinitionFunctionName == FSkeletalMeshInterfaceHelper::GetSkinnedTriangleDataWSInterpName)
+	else if (FunctionInfo.DefinitionName == FSkeletalMeshInterfaceHelper::GetSkinnedTriangleDataWSInterpName)
 	{
 		static const TCHAR* FormatSample = TEXT("void {InstanceFunctionName} (in {MeshTriCoordinateStructName} InCoord, in float InInterp, out float3 OutPosition, out float3 OutVelocity, out float3 OutNormal, out float3 OutBinormal, out float3 OutTangent) { {GetDISkelMeshContextName} DISKelMesh_GetSkinnedTriangleDataInterpolatedWS(DIContext, InCoord.Tri, InCoord.BaryCoord, InInterp, OutPosition, OutVelocity, OutNormal, OutBinormal, OutTangent); }");
 		OutHLSL += FString::Format(FormatSample, ArgsSample);
 	}
-	else if (DefinitionFunctionName == FSkeletalMeshInterfaceHelper::GetSkinnedTriangleDataName)
+	else if (FunctionInfo.DefinitionName == FSkeletalMeshInterfaceHelper::GetSkinnedTriangleDataName)
 	{
 		static const TCHAR* FormatSample = TEXT("void {InstanceFunctionName} (in {MeshTriCoordinateStructName} InCoord, out float3 OutPosition, out float3 OutVelocity, out float3 OutNormal, out float3 OutBinormal, out float3 OutTangent) { {GetDISkelMeshContextName} DISKelMesh_GetSkinnedTriangleData(DIContext, InCoord.Tri, InCoord.BaryCoord, OutPosition, OutVelocity, OutNormal, OutBinormal, OutTangent); }");
 		OutHLSL += FString::Format(FormatSample, ArgsSample);
 	}
-	else if (DefinitionFunctionName == FSkeletalMeshInterfaceHelper::GetSkinnedTriangleDataInterpName)
+	else if (FunctionInfo.DefinitionName == FSkeletalMeshInterfaceHelper::GetSkinnedTriangleDataInterpName)
 	{
 		static const TCHAR* FormatSample = TEXT("void {InstanceFunctionName} (in {MeshTriCoordinateStructName} InCoord, in float InInterp, out float3 OutPosition, out float3 OutVelocity, out float3 OutNormal, out float3 OutBinormal, out float3 OutTangent) { {GetDISkelMeshContextName} DISKelMesh_GetSkinnedTriangleDataInterpolated(DIContext, InCoord.Tri, InCoord.BaryCoord, InInterp, OutPosition, OutVelocity, OutNormal, OutBinormal, OutTangent); }");
 		OutHLSL += FString::Format(FormatSample, ArgsSample);
 	}
-	else if (DefinitionFunctionName == FSkeletalMeshInterfaceHelper::GetTriUVName)
+	else if (FunctionInfo.DefinitionName == FSkeletalMeshInterfaceHelper::GetTriUVName)
 	{
 		static const TCHAR* FormatSample = TEXT("void {InstanceFunctionName} (in {MeshTriCoordinateStructName} InCoord, in int InUVSet, out float2 OutUV) { {GetDISkelMeshContextName} DISKelMesh_GetTriUV(DIContext, InCoord.Tri, InCoord.BaryCoord, InUVSet, OutUV); }");
 		OutHLSL += FString::Format(FormatSample, ArgsSample);
 	}
-	else if (DefinitionFunctionName == FSkeletalMeshInterfaceHelper::GetTriColorName)
+	else if (FunctionInfo.DefinitionName == FSkeletalMeshInterfaceHelper::GetTriColorName)
 	{
 		static const TCHAR* FormatSample = TEXT("void {InstanceFunctionName} (in {MeshTriCoordinateStructName} InCoord, out float4 OutColor) { {GetDISkelMeshContextName} DISkelMesh_GetTriColor(DIContext, InCoord.Tri, InCoord.BaryCoord, OutColor); }");
 		OutHLSL += FString::Format(FormatSample, ArgsSample);
 	}
-	else if (DefinitionFunctionName == FSkeletalMeshInterfaceHelper::IsValidTriCoordName)
+	else if (FunctionInfo.DefinitionName == FSkeletalMeshInterfaceHelper::IsValidTriCoordName)
 	{
 		static const TCHAR* FormatSample = TEXT("void {InstanceFunctionName} (in {MeshTriCoordinateStructName} InCoord, out bool IsValid) { {GetDISkelMeshContextName} IsValid = InCoord.Tri < DIContext.MeshTriangleCount; }");
 		OutHLSL += FString::Format(FormatSample, ArgsSample);
 	}
-	else if (DefinitionFunctionName == FSkeletalMeshInterfaceHelper::GetTriCoordVerticesName)
+	else if (FunctionInfo.DefinitionName == FSkeletalMeshInterfaceHelper::GetTriCoordVerticesName)
 	{
 		static const TCHAR* FormatSample = TEXT("void {InstanceFunctionName} (in int TriangleIndex, out int OutVertexIndex0, out int OutVertexIndex1, out int OutVertexIndex2) { {GetDISkelMeshContextName} DISkelMesh_GetTriVertices(DIContext, TriangleIndex, OutVertexIndex0, OutVertexIndex1, OutVertexIndex2); }");
 		OutHLSL += FString::Format(FormatSample, ArgsSample);
 	}
 	// Bone Sampling
-	else if (DefinitionFunctionName == FSkeletalMeshInterfaceHelper::GetSkinnedBoneDataName)
+	else if (FunctionInfo.DefinitionName == FSkeletalMeshInterfaceHelper::GetSkinnedBoneDataName)
 	{
 		static const TCHAR* FormatSample = TEXT("void {InstanceFunctionName} (in int InBone, out float3 OutPosition, out float4 OutRotation, out float3 OutVelocity) { {GetDISkelMeshContextName} DISkelMesh_GetSkinnedBone(DIContext, InBone, OutPosition, OutRotation, OutVelocity); }");
 		OutHLSL += FString::Format(FormatSample, ArgsSample);
 	}
-	else if (DefinitionFunctionName == FSkeletalMeshInterfaceHelper::GetSkinnedBoneDataInterpolatedName)
+	else if (FunctionInfo.DefinitionName == FSkeletalMeshInterfaceHelper::GetSkinnedBoneDataInterpolatedName)
 	{
 		static const TCHAR* FormatSample = TEXT("void {InstanceFunctionName} (in int InBone, in float Interp, out float3 OutPosition, out float4 OutRotation, out float3 OutVelocity) { {GetDISkelMeshContextName} DISkelMesh_GetSkinnedBoneInterpolated(DIContext, InBone, Interp, OutPosition, OutRotation, OutVelocity); }");
 		OutHLSL += FString::Format(FormatSample, ArgsSample);
 	}
-	else if (DefinitionFunctionName == FSkeletalMeshInterfaceHelper::GetSkinnedBoneDataWSName)
+	else if (FunctionInfo.DefinitionName == FSkeletalMeshInterfaceHelper::GetSkinnedBoneDataWSName)
 	{
 		static const TCHAR* FormatSample = TEXT("void {InstanceFunctionName} (in int InBone, out float3 OutPosition, out float4 OutRotation, out float3 OutVelocity) { {GetDISkelMeshContextName} DISkelMesh_GetSkinnedBoneWS(DIContext, InBone, OutPosition, OutRotation, OutVelocity); }");
 		OutHLSL += FString::Format(FormatSample, ArgsSample);
 	}
-	else if (DefinitionFunctionName == FSkeletalMeshInterfaceHelper::GetSkinnedBoneDataWSInterpolatedName)
+	else if (FunctionInfo.DefinitionName == FSkeletalMeshInterfaceHelper::GetSkinnedBoneDataWSInterpolatedName)
 	{
 		static const TCHAR* FormatSample = TEXT("void {InstanceFunctionName} (in int InBone, in float Interp, out float3 OutPosition, out float4 OutRotation, out float3 OutVelocity) { {GetDISkelMeshContextName} DISkelMesh_GetSkinnedBoneInterpolatedWS(DIContext, InBone, Interp, OutPosition, OutRotation, OutVelocity); }");
 		OutHLSL += FString::Format(FormatSample, ArgsSample);
 	}
 	// Vertex Sampling
-	else if (DefinitionFunctionName == FSkeletalMeshInterfaceHelper::RandomVertexName)
+	else if (FunctionInfo.DefinitionName == FSkeletalMeshInterfaceHelper::RandomVertexName)
 	{
 		static const TCHAR* FormatSample = TEXT("void {InstanceFunctionName}(out int OutVertex) { {GetDISkelMeshContextName} DISkelMesh_GetRandomVertex(DIContext, OutVertex); }");
 		OutHLSL += FString::Format(FormatSample, ArgsSample);
 	}
-	else if (DefinitionFunctionName == FSkeletalMeshInterfaceHelper::IsValidVertexName)
+	else if (FunctionInfo.DefinitionName == FSkeletalMeshInterfaceHelper::IsValidVertexName)
 	{
 		static const TCHAR* FormatSample = TEXT("void {InstanceFunctionName} (in int Vertex, out bool IsValid) { {GetDISkelMeshContextName} DISkelMesh_IsValidVertex(DIContext, Vertex, IsValid); }");
 		OutHLSL += FString::Format(FormatSample, ArgsSample);
 	}
-	else if (DefinitionFunctionName == FSkeletalMeshInterfaceHelper::GetSkinnedVertexDataName)
+	else if (FunctionInfo.DefinitionName == FSkeletalMeshInterfaceHelper::GetSkinnedVertexDataName)
 	{
 		static const TCHAR* FormatSample = TEXT("void {InstanceFunctionName} (in int Vertex, out float3 OutPosition, out float3 OutVelocity) { {GetDISkelMeshContextName} DISkelMesh_GetSkinnedVertex(DIContext, Vertex, OutPosition, OutVelocity); }");
 		OutHLSL += FString::Format(FormatSample, ArgsSample);
 	}
-	else if (DefinitionFunctionName == FSkeletalMeshInterfaceHelper::GetSkinnedVertexDataWSName)
+	else if (FunctionInfo.DefinitionName == FSkeletalMeshInterfaceHelper::GetSkinnedVertexDataWSName)
 	{
 		static const TCHAR* FormatSample = TEXT("void {InstanceFunctionName} (in int Vertex, out float3 OutPosition, out float3 OutVelocity) { {GetDISkelMeshContextName} DISkelMesh_GetSkinnedVertexWS(DIContext, Vertex, OutPosition, OutVelocity); }");
 		OutHLSL += FString::Format(FormatSample, ArgsSample);
 	}
-	else if (DefinitionFunctionName == FSkeletalMeshInterfaceHelper::GetVertexColorName)
+	else if (FunctionInfo.DefinitionName == FSkeletalMeshInterfaceHelper::GetVertexColorName)
 	{
 		static const TCHAR* FormatSample = TEXT("void {InstanceFunctionName} (in int Vertex, out float4 OutColor) { {GetDISkelMeshContextName} DISkelMesh_GetVertexColor(DIContext, Vertex, OutColor); }");
 		OutHLSL += FString::Format(FormatSample, ArgsSample);
 	}
-	else if (DefinitionFunctionName == FSkeletalMeshInterfaceHelper::GetVertexUVName)
+	else if (FunctionInfo.DefinitionName == FSkeletalMeshInterfaceHelper::GetVertexUVName)
 	{
 		static const TCHAR* FormatSample = TEXT("void {InstanceFunctionName} (in int Vertex, in int UVSet, out float2 OutUV) { {GetDISkelMeshContextName} DISkelMesh_GetVertexUV(DIContext, Vertex, UVSet, OutUV); }");
 		OutHLSL += FString::Format(FormatSample, ArgsSample);
 	}
 	// Specific Bone
-	else if (DefinitionFunctionName == FSkeletalMeshInterfaceHelper::GetSpecificBoneCountName)
+	else if (FunctionInfo.DefinitionName == FSkeletalMeshInterfaceHelper::GetSpecificBoneCountName)
 	{
 		static const TCHAR* FormatSample = TEXT("void {InstanceFunctionName} (out int Count) { {GetDISkelMeshContextName} DISkelMesh_GetSpecificBoneCount(DIContext, Count); }");
 		OutHLSL += FString::Format(FormatSample, ArgsSample);
 	}
-	else if (DefinitionFunctionName == FSkeletalMeshInterfaceHelper::GetSpecificBoneAtName)
+	else if (FunctionInfo.DefinitionName == FSkeletalMeshInterfaceHelper::GetSpecificBoneAtName)
 	{
 		static const TCHAR* FormatSample = TEXT("void {InstanceFunctionName} (in int BoneIndex, out int Bone) { {GetDISkelMeshContextName} DISkelMesh_GetSpecificBoneAt(DIContext, BoneIndex, Bone); }");
 		OutHLSL += FString::Format(FormatSample, ArgsSample);
 	}
-	else if (DefinitionFunctionName == FSkeletalMeshInterfaceHelper::RandomSpecificBoneName)
+	else if (FunctionInfo.DefinitionName == FSkeletalMeshInterfaceHelper::RandomSpecificBoneName)
 	{
 		static const TCHAR* FormatSample = TEXT("void {InstanceFunctionName} (out int Bone) { {GetDISkelMeshContextName} DISkelMesh_RandomSpecificBone(DIContext, Bone); }");
 		OutHLSL += FString::Format(FormatSample, ArgsSample);
 	}
 	// Specific Socket
-	else if (DefinitionFunctionName == FSkeletalMeshInterfaceHelper::GetSpecificSocketCountName)
+	else if (FunctionInfo.DefinitionName == FSkeletalMeshInterfaceHelper::GetSpecificSocketCountName)
 	{
 		static const TCHAR* FormatSample = TEXT("void {InstanceFunctionName} (out int Count) { {GetDISkelMeshContextName} DISkelMesh_GetSpecificSocketCount(DIContext, Count); }");
 		OutHLSL += FString::Format(FormatSample, ArgsSample);
 	}
-	else if (DefinitionFunctionName == FSkeletalMeshInterfaceHelper::GetSpecificSocketBoneAtName)
+	else if (FunctionInfo.DefinitionName == FSkeletalMeshInterfaceHelper::GetSpecificSocketBoneAtName)
 	{
 		static const TCHAR* FormatSample = TEXT("void {InstanceFunctionName} (in int SocketIndex, out int Bone) { {GetDISkelMeshContextName} DISkelMesh_GetSpecificSocketBoneAt(DIContext, SocketIndex, Bone); }");
 		OutHLSL += FString::Format(FormatSample, ArgsSample);
 	} 
-	else if (DefinitionFunctionName == FSkeletalMeshInterfaceHelper::GetSpecificSocketTransformName)
+	else if (FunctionInfo.DefinitionName == FSkeletalMeshInterfaceHelper::GetSpecificSocketTransformName)
 	{
 		// TODO: This just returns the Identity transform.
 		// TODO: Make this work on the GPU?
 		static const TCHAR* FormatSample = TEXT("void {InstanceFunctionName} (in int SocketIndex, in int bShouldApplyTransform, out float3 Translation, out float4 Rotation, out float4 Scale) { Translation = float3(0.0, 0.0, 0.0); Rotation = float4(0.0, 0.0, 0.0, 1.0); Scale = float3(1.0, 1.0, 1.0); }");
 		OutHLSL += FString::Format(FormatSample, ArgsSample);
 	} 
-	else if (DefinitionFunctionName == FSkeletalMeshInterfaceHelper::RandomSpecificSocketBoneName)
+	else if (FunctionInfo.DefinitionName == FSkeletalMeshInterfaceHelper::RandomSpecificSocketBoneName)
 	{
 		static const TCHAR* FormatSample = TEXT("void {InstanceFunctionName} (out int SocketBone) { {GetDISkelMeshContextName} DISkelMesh_RandomSpecificSocketBone(DIContext, SocketBone); }");
 		OutHLSL += FString::Format(FormatSample, ArgsSample);
 	}
 	// Unsupported functionality
-	//else if (DefinitionFunctionName == FSkeletalMeshInterfaceHelper::GetVertexCountName)				// void GetFilteredVertexCount(out int VertexCount)
-	//else if (DefinitionFunctionName == FSkeletalMeshInterfaceHelper::GetVertexAtName)					// void GetFilteredVertex(in int FilterdIndex, out int VertexIndex)
-	//else if (DefinitionFunctionName == FSkeletalMeshInterfaceHelper::GetTriangleCountName)			// void GetFilteredTriangleCount(out int OutTriangleCount)
-	//else if (DefinitionFunctionName == FSkeletalMeshInterfaceHelper::GetTriangleAtName)				// void GetFilteredTriangle(in int Triangle, out {MeshTriCoordinateStructName} OutCoord)
-	//else if (DefinitionFunctionName == FSkeletalMeshInterfaceHelper::IsValidBoneName)					// void IsValidBoneName(in int BoneIndex, out bool IsValid)
+	//else if (FunctionInfo.DefinitionName == FSkeletalMeshInterfaceHelper::GetVertexCountName)				// void GetFilteredVertexCount(out int VertexCount)
+	//else if (FunctionInfo.DefinitionName == FSkeletalMeshInterfaceHelper::GetVertexAtName)					// void GetFilteredVertex(in int FilterdIndex, out int VertexIndex)
+	//else if (FunctionInfo.DefinitionName == FSkeletalMeshInterfaceHelper::GetTriangleCountName)			// void GetFilteredTriangleCount(out int OutTriangleCount)
+	//else if (FunctionInfo.DefinitionName == FSkeletalMeshInterfaceHelper::GetTriangleAtName)				// void GetFilteredTriangle(in int Triangle, out {MeshTriCoordinateStructName} OutCoord)
+	//else if (FunctionInfo.DefinitionName == FSkeletalMeshInterfaceHelper::IsValidBoneName)					// void IsValidBoneName(in int BoneIndex, out bool IsValid)
 	else
 	{
 		// This function is not support
@@ -2121,7 +2121,7 @@ bool UNiagaraDataInterfaceSkeletalMesh::GetFunctionHLSL(const FName& DefinitionF
 	return true;
 }
 
-void UNiagaraDataInterfaceSkeletalMesh::GetParameterDefinitionHLSL(FNiagaraDataInterfaceGPUParamInfo& ParamInfo, FString& OutHLSL)
+void UNiagaraDataInterfaceSkeletalMesh::GetParameterDefinitionHLSL(const FNiagaraDataInterfaceGPUParamInfo& ParamInfo, FString& OutHLSL)
 {
 	OutHLSL += TEXT("DISKELMESH_DECLARE_CONSTANTS(") + ParamInfo.DataInterfaceHLSLSymbol + TEXT(")\n");
 }
diff --git a/Engine/Plugins/FX/Niagara/Source/Niagara/Private/NiagaraDataInterfaceStaticMesh.cpp b/Engine/Plugins/FX/Niagara/Source/Niagara/Private/NiagaraDataInterfaceStaticMesh.cpp
--- a/Engine/Plugins/FX/Niagara/Source/Niagara/Private/NiagaraDataInterfaceStaticMesh.cpp
+++ b/Engine/Plugins/FX/Niagara/Source/Niagara/Private/NiagaraDataInterfaceStaticMesh.cpp
@@ -1915,7 +1915,7 @@ void UNiagaraDataInterfaceStaticMesh::SetDefaultMeshFromBlueprints(UStaticMesh*
 	DefaultMesh = MeshToUse;
 }
 
-bool UNiagaraDataInterfaceStaticMesh::GetFunctionHLSL(const FName&  DefinitionFunctionName, FString InstanceFunctionName, FNiagaraDataInterfaceGPUParamInfo& ParamInfo, FString& OutHLSL)
+bool UNiagaraDataInterfaceStaticMesh::GetFunctionHLSL(const FNiagaraDataInterfaceGPUParamInfo& ParamInfo, const FNiagaraDataInterfaceGeneratedFunction& FunctionInfo, int FunctionInstanceIndex, FString& OutHLSL)
 {
 	FNDIStaticMeshParametersName ParamNames;
 	GetNiagaraDataInterfaceParametersName(ParamNames, ParamInfo.DataInterfaceHLSLSymbol);
@@ -1923,7 +1923,7 @@ bool UNiagaraDataInterfaceStaticMesh::GetFunctionHLSL(const FName&  DefinitionFu
 	FString MeshTriCoordinateStructName = "MeshTriCoordinate";
 
 	TMap<FString, FStringFormatArg> ArgsSample = {
-		{TEXT("InstanceFunctionName"), InstanceFunctionName},
+		{TEXT("InstanceFunctionName"), FunctionInfo.InstanceName},
 		{TEXT("MeshTriCoordinateStructName"), MeshTriCoordinateStructName},
 		{TEXT("SectionCountName"), ParamNames.SectionCountName},
 		{TEXT("MeshSectionBufferName"), ParamNames.MeshSectionBufferName},
@@ -1942,7 +1942,7 @@ bool UNiagaraDataInterfaceStaticMesh::GetFunctionHLSL(const FName&  DefinitionFu
 		{TEXT("NumTexCoordName"), ParamNames.NumTexCoordName},
 	};
 
-	if (DefinitionFunctionName == StaticMeshHelpers::RandomSectionName)
+	if (FunctionInfo.DefinitionName == StaticMeshHelpers::RandomSectionName)
 	{
 		static const TCHAR *FormatSample = TEXT(R"(
 			void {InstanceFunctionName} (out int Out_Section)
@@ -1963,7 +1963,7 @@ bool UNiagaraDataInterfaceStaticMesh::GetFunctionHLSL(const FName&  DefinitionFu
 			)");
 		OutHLSL += FString::Format(FormatSample, ArgsSample);
 	}
-	else if (DefinitionFunctionName == StaticMeshHelpers::RandomTriCoordName)
+	else if (FunctionInfo.DefinitionName == StaticMeshHelpers::RandomTriCoordName)
 	{
 		static const TCHAR *FormatSample = TEXT(R"(
 			void {InstanceFunctionName} (out {MeshTriCoordinateStructName} Out_Coord)
@@ -2015,10 +2015,10 @@ bool UNiagaraDataInterfaceStaticMesh::GetFunctionHLSL(const FName&  DefinitionFu
 			)");
 		OutHLSL += FString::Format(FormatSample, ArgsSample);
 	}
-	//else if (DefinitionFunctionName == StaticMeshHelpers::RandomTriCoordVCFilteredName)
+	//else if (FunctionInfo.DefinitionName == StaticMeshHelpers::RandomTriCoordVCFilteredName)
 	//{
 	//}
-	else if (DefinitionFunctionName == StaticMeshHelpers::RandomTriCoordOnSectionName)
+	else if (FunctionInfo.DefinitionName == StaticMeshHelpers::RandomTriCoordOnSectionName)
 	{
 		static const TCHAR *FormatSample = TEXT(R"(
 			void {InstanceFunctionName} (in int In_Section, out {MeshTriCoordinateStructName} Out_Coord)
@@ -2060,7 +2060,7 @@ bool UNiagaraDataInterfaceStaticMesh::GetFunctionHLSL(const FName&  DefinitionFu
 			)");
 		OutHLSL += FString::Format(FormatSample, ArgsSample);
 	}
-	else if (DefinitionFunctionName == StaticMeshHelpers::GetTriPositionName)
+	else if (FunctionInfo.DefinitionName == StaticMeshHelpers::GetTriPositionName)
 	{
 		static const TCHAR *FormatSample = TEXT(R"(
 			void {InstanceFunctionName} (in {MeshTriCoordinateStructName} In_Coord, out float3 Out_Position)
@@ -2079,7 +2079,7 @@ bool UNiagaraDataInterfaceStaticMesh::GetFunctionHLSL(const FName&  DefinitionFu
 			)");
 		OutHLSL += FString::Format(FormatSample, ArgsSample);
 	}
-	else if (DefinitionFunctionName == StaticMeshHelpers::GetTriPositionWSName)
+	else if (FunctionInfo.DefinitionName == StaticMeshHelpers::GetTriPositionWSName)
 	{
 		static const TCHAR *FormatSample = TEXT(R"(
 			void {InstanceFunctionName} (in {MeshTriCoordinateStructName} In_Coord, out float3 Out_Position)
@@ -2100,7 +2100,7 @@ bool UNiagaraDataInterfaceStaticMesh::GetFunctionHLSL(const FName&  DefinitionFu
 			)");
 		OutHLSL += FString::Format(FormatSample, ArgsSample);
 	}
-	else if (DefinitionFunctionName == StaticMeshHelpers::GetTriNormalName)
+	else if (FunctionInfo.DefinitionName == StaticMeshHelpers::GetTriNormalName)
 	{
 		static const TCHAR *FormatSample = TEXT(R"(
 		void {InstanceFunctionName} (in {MeshTriCoordinateStructName} In_Coord, out float3 Out_Normal)
@@ -2121,7 +2121,7 @@ bool UNiagaraDataInterfaceStaticMesh::GetFunctionHLSL(const FName&  DefinitionFu
 		)");
 		OutHLSL += FString::Format(FormatSample, ArgsSample);
 	}
-	else if (DefinitionFunctionName == StaticMeshHelpers::GetTriNormalWSName)
+	else if (FunctionInfo.DefinitionName == StaticMeshHelpers::GetTriNormalWSName)
 	{
 		static const TCHAR *FormatSample = TEXT(R"(
 		void {InstanceFunctionName} (in {MeshTriCoordinateStructName} In_Coord, out float3 Out_Normal)
@@ -2142,7 +2142,7 @@ bool UNiagaraDataInterfaceStaticMesh::GetFunctionHLSL(const FName&  DefinitionFu
 		)");
 		OutHLSL += FString::Format(FormatSample, ArgsSample);
 	}
-	else if (DefinitionFunctionName == StaticMeshHelpers::GetTriTangentsName)
+	else if (FunctionInfo.DefinitionName == StaticMeshHelpers::GetTriTangentsName)
 	{
 		static const TCHAR *FormatSample = TEXT(R"(
 		void {InstanceFunctionName} (in {MeshTriCoordinateStructName} In_Coord, out float3 Out_Tangent, out float3 Out_Binormal, out float3 Out_Normal)
@@ -2170,7 +2170,7 @@ bool UNiagaraDataInterfaceStaticMesh::GetFunctionHLSL(const FName&  DefinitionFu
 		)");
 		OutHLSL += FString::Format(FormatSample, ArgsSample);
 	}
-	else if (DefinitionFunctionName == StaticMeshHelpers::GetTriTangentsWSName)
+	else if (FunctionInfo.DefinitionName == StaticMeshHelpers::GetTriTangentsWSName)
 	{
 		static const TCHAR *FormatSample = TEXT(R"(
 		void {InstanceFunctionName} (in {MeshTriCoordinateStructName} In_Coord, out float3 Out_Tangent, out float3 Out_Binormal, out float3 Out_Normal)
@@ -2206,7 +2206,7 @@ bool UNiagaraDataInterfaceStaticMesh::GetFunctionHLSL(const FName&  DefinitionFu
 		)");
 		OutHLSL += FString::Format(FormatSample, ArgsSample);
 	}
-	else if (DefinitionFunctionName == StaticMeshHelpers::GetTriColorName)
+	else if (FunctionInfo.DefinitionName == StaticMeshHelpers::GetTriColorName)
 	{
 		static const TCHAR *FormatSample = TEXT(R"(
 			void {InstanceFunctionName} (in {MeshTriCoordinateStructName} In_Coord, out float4 Out_Color)
@@ -2225,7 +2225,7 @@ bool UNiagaraDataInterfaceStaticMesh::GetFunctionHLSL(const FName&  DefinitionFu
 			)");
 		OutHLSL += FString::Format(FormatSample, ArgsSample);
 	}
-	else if (DefinitionFunctionName == StaticMeshHelpers::GetTriUVName)
+	else if (FunctionInfo.DefinitionName == StaticMeshHelpers::GetTriUVName)
 	{
 		static const TCHAR *FormatSample = TEXT(R"(
 			void {InstanceFunctionName} (in {MeshTriCoordinateStructName} In_Coord, in int In_UVSet, out float2 Out_UV)
@@ -2253,7 +2253,7 @@ bool UNiagaraDataInterfaceStaticMesh::GetFunctionHLSL(const FName&  DefinitionFu
 			)");
 		OutHLSL += FString::Format(FormatSample, ArgsSample);
 	}
-	else if (DefinitionFunctionName == StaticMeshHelpers::GetTriPositionAndVelocityName)
+	else if (FunctionInfo.DefinitionName == StaticMeshHelpers::GetTriPositionAndVelocityName)
 	{
 		static const TCHAR *FormatSample = TEXT(R"(
 			void {InstanceFunctionName} (in {MeshTriCoordinateStructName} In_Coord, out float3 Out_Position, out float3 Out_Velocity)
@@ -2279,7 +2279,7 @@ bool UNiagaraDataInterfaceStaticMesh::GetFunctionHLSL(const FName&  DefinitionFu
 			)");
 		OutHLSL += FString::Format(FormatSample, ArgsSample);
 	}
-	else if (DefinitionFunctionName == StaticMeshHelpers::GetMeshLocalToWorldName)
+	else if (FunctionInfo.DefinitionName == StaticMeshHelpers::GetMeshLocalToWorldName)
 	{
 		static const TCHAR *FormatSample = TEXT(R"(
 			void {InstanceFunctionName} (out float4x4 Out_Transform)
@@ -2289,7 +2289,7 @@ bool UNiagaraDataInterfaceStaticMesh::GetFunctionHLSL(const FName&  DefinitionFu
 			)");
 		OutHLSL += FString::Format(FormatSample, ArgsSample);
 	}
-	else if (DefinitionFunctionName == StaticMeshHelpers::GetMeshLocalToWorldInverseTransposedName)
+	else if (FunctionInfo.DefinitionName == StaticMeshHelpers::GetMeshLocalToWorldInverseTransposedName)
 	{
 		static const TCHAR *FormatSample = TEXT(R"(
 			void {InstanceFunctionName} (out float4x4 Out_Transform)
@@ -2299,7 +2299,7 @@ bool UNiagaraDataInterfaceStaticMesh::GetFunctionHLSL(const FName&  DefinitionFu
 			)");
 		OutHLSL += FString::Format(FormatSample, ArgsSample);
 	}
-	else if (DefinitionFunctionName == StaticMeshHelpers::GetMeshWorldVelocityName)
+	else if (FunctionInfo.DefinitionName == StaticMeshHelpers::GetMeshWorldVelocityName)
 	{
 		static const TCHAR *FormatSample = TEXT(R"(
 			void {InstanceFunctionName} (out float3 Out_Velocity)
@@ -2309,7 +2309,7 @@ bool UNiagaraDataInterfaceStaticMesh::GetFunctionHLSL(const FName&  DefinitionFu
 			)");
 		OutHLSL += FString::Format(FormatSample, ArgsSample);
 	}
-	else if (DefinitionFunctionName == StaticMeshHelpers::GetVertexPositionName)
+	else if (FunctionInfo.DefinitionName == StaticMeshHelpers::GetVertexPositionName)
 	{
 		static const TCHAR *FormatSample = TEXT(R"(
 				void {InstanceFunctionName} (in int VertexIndex, out float3 Out_Position)
@@ -2320,7 +2320,7 @@ bool UNiagaraDataInterfaceStaticMesh::GetFunctionHLSL(const FName&  DefinitionFu
 				)");
 		OutHLSL += FString::Format(FormatSample, ArgsSample);
 	}
-	else if (DefinitionFunctionName == StaticMeshHelpers::GetVertexPositionWSName)
+	else if (FunctionInfo.DefinitionName == StaticMeshHelpers::GetVertexPositionWSName)
 	{
 		static const TCHAR *FormatSample = TEXT(R"(
 				void {InstanceFunctionName} (in int VertexIndex, out float3 Out_Position)
@@ -2342,7 +2342,7 @@ bool UNiagaraDataInterfaceStaticMesh::GetFunctionHLSL(const FName&  DefinitionFu
 	return true;
 }
 
-void UNiagaraDataInterfaceStaticMesh::GetParameterDefinitionHLSL(FNiagaraDataInterfaceGPUParamInfo& ParamInfo, FString& OutHLSL)
+void UNiagaraDataInterfaceStaticMesh::GetParameterDefinitionHLSL(const FNiagaraDataInterfaceGPUParamInfo& ParamInfo, FString& OutHLSL)
 {
 	FNDIStaticMeshParametersName ParamNames;
 	GetNiagaraDataInterfaceParametersName(ParamNames, ParamInfo.DataInterfaceHLSLSymbol);
diff --git a/Engine/Plugins/FX/Niagara/Source/Niagara/Private/NiagaraDataInterfaceTexture.cpp b/Engine/Plugins/FX/Niagara/Source/Niagara/Private/NiagaraDataInterfaceTexture.cpp
--- a/Engine/Plugins/FX/Niagara/Source/Niagara/Private/NiagaraDataInterfaceTexture.cpp
+++ b/Engine/Plugins/FX/Niagara/Source/Niagara/Private/NiagaraDataInterfaceTexture.cpp
@@ -296,39 +296,39 @@ void UNiagaraDataInterfaceTexture::SamplePseudoVolumeTexture(FVectorVMContext& C
 	}
 }
 
-bool UNiagaraDataInterfaceTexture::GetFunctionHLSL(const FName& DefinitionFunctionName, FString InstanceFunctionName, FNiagaraDataInterfaceGPUParamInfo& ParamInfo, FString& OutHLSL)
+bool UNiagaraDataInterfaceTexture::GetFunctionHLSL(const FNiagaraDataInterfaceGPUParamInfo& ParamInfo, const FNiagaraDataInterfaceGeneratedFunction& FunctionInfo, int FunctionInstanceIndex, FString& OutHLSL)
 {
-	if (DefinitionFunctionName == SampleTexture2DName)
+	if (FunctionInfo.DefinitionName == SampleTexture2DName)
 	{
 		FString HLSLTextureName = TextureName + ParamInfo.DataInterfaceHLSLSymbol;
 		FString HLSLSamplerName = SamplerName + ParamInfo.DataInterfaceHLSLSymbol;
-		OutHLSL += TEXT("void ") + InstanceFunctionName + TEXT("(in float2 In_UV, out float4 Out_Value) \n{\n");
+		OutHLSL += TEXT("void ") + FunctionInfo.InstanceName + TEXT("(in float2 In_UV, out float4 Out_Value) \n{\n");
 		OutHLSL += TEXT("\t Out_Value = ") + HLSLTextureName + TEXT(".SampleLevel(") + HLSLSamplerName + TEXT(", In_UV, 0);\n");
 		OutHLSL += TEXT("\n}\n");
 		return true;
 	}
-	/*else if (DefinitionFunctionName == SampleVolumeTextureName)
+	/*else if (FunctionInfo.DefinitionName == SampleVolumeTextureName)
 	{
 		FString HLSLTextureName = TextureName + ParamInfo.DataInterfaceHLSLSymbol;
 		FString HLSLSamplerName = SamplerName + ParamInfo.DataInterfaceHLSLSymbol;
-		OutHLSL += TEXT("void ") + InstanceFunctionName + TEXT("(in float3 In_UV, out float4 Out_Value) \n{\n");
+		OutHLSL += TEXT("void ") + FunctionInfo.InstanceName + TEXT("(in float3 In_UV, out float4 Out_Value) \n{\n");
 		OutHLSL += TEXT("\t Out_Value = ") + HLSLTextureName + TEXT(".SampleLevel(") + HLSLSamplerName + TEXT(", In_UV, 0);\n");
 		OutHLSL += TEXT("\n}\n");
 		return true;
 	}*/
-	else if (DefinitionFunctionName == SamplePseudoVolumeTextureName)
+	else if (FunctionInfo.DefinitionName == SamplePseudoVolumeTextureName)
 	{
 		FString HLSLTextureName = TextureName + ParamInfo.DataInterfaceHLSLSymbol;
 		FString HLSLSamplerName = SamplerName + ParamInfo.DataInterfaceHLSLSymbol;
-		OutHLSL += TEXT("void ") + InstanceFunctionName + TEXT("(in float3 In_UVW, in float2 In_XYNumFrames, in float In_TotalNumFrames, in int In_MipMode, in float In_MipLevel, in float2 In_DDX, in float2 In_DDY, out float4 Out_Value) \n{\n");
+		OutHLSL += TEXT("void ") + FunctionInfo.InstanceName + TEXT("(in float3 In_UVW, in float2 In_XYNumFrames, in float In_TotalNumFrames, in int In_MipMode, in float In_MipLevel, in float2 In_DDX, in float2 In_DDY, out float4 Out_Value) \n{\n");
 		OutHLSL += TEXT("\t Out_Value = PseudoVolumeTexture(") + HLSLTextureName + TEXT(", ") + HLSLSamplerName + TEXT(", In_UVW, In_XYNumFrames, In_TotalNumFrames, (uint) In_MipMode, In_MipLevel, In_DDX, In_DDY); \n");
 		OutHLSL += TEXT("\n}\n");
 		return true;
 	}
-	else if (DefinitionFunctionName == TextureDimsName)
+	else if (FunctionInfo.DefinitionName == TextureDimsName)
 	{
 		FString DimsVar = DimensionsBaseName + ParamInfo.DataInterfaceHLSLSymbol;
-		OutHLSL += TEXT("void ") + InstanceFunctionName + TEXT("(out float2 Out_Value) \n{\n");
+		OutHLSL += TEXT("void ") + FunctionInfo.InstanceName + TEXT("(out float2 Out_Value) \n{\n");
 		OutHLSL += TEXT("\t Out_Value = ") + DimsVar + TEXT(";\n");
 		OutHLSL += TEXT("\n}\n");
 		return true;
@@ -336,7 +336,7 @@ bool UNiagaraDataInterfaceTexture::GetFunctionHLSL(const FName& DefinitionFuncti
 	return false;
 }
 
-void UNiagaraDataInterfaceTexture::GetParameterDefinitionHLSL(FNiagaraDataInterfaceGPUParamInfo& ParamInfo, FString& OutHLSL)
+void UNiagaraDataInterfaceTexture::GetParameterDefinitionHLSL(const FNiagaraDataInterfaceGPUParamInfo& ParamInfo, FString& OutHLSL)
 {
 	FString HLSLTextureName = TextureName + ParamInfo.DataInterfaceHLSLSymbol;
 	FString HLSLSamplerName = SamplerName + ParamInfo.DataInterfaceHLSLSymbol;
diff --git a/Engine/Plugins/FX/Niagara/Source/Niagara/Private/NiagaraDataInterfaceVector2DCurve.cpp b/Engine/Plugins/FX/Niagara/Source/Niagara/Private/NiagaraDataInterfaceVector2DCurve.cpp
--- a/Engine/Plugins/FX/Niagara/Source/Niagara/Private/NiagaraDataInterfaceVector2DCurve.cpp
+++ b/Engine/Plugins/FX/Niagara/Source/Niagara/Private/NiagaraDataInterfaceVector2DCurve.cpp
@@ -149,7 +149,7 @@ void UNiagaraDataInterfaceVector2DCurve::GetFunctions(TArray<FNiagaraFunctionSig
 // the HLSL in the spirit of a static switch
 // TODO: need a way to identify each specific function here
 // 
-bool UNiagaraDataInterfaceVector2DCurve::GetFunctionHLSL(const FName& DefinitionFunctionName, FString InstanceFunctionName, FNiagaraDataInterfaceGPUParamInfo& ParamInfo, FString& OutHLSL)
+bool UNiagaraDataInterfaceVector2DCurve::GetFunctionHLSL(const FNiagaraDataInterfaceGPUParamInfo& ParamInfo, const FNiagaraDataInterfaceGeneratedFunction& FunctionInfo, int FunctionInstanceIndex, FString& OutHLSL)
 {
 	FString TimeToLUTFrac = TEXT("TimeToLUTFraction_") + ParamInfo.DataInterfaceHLSLSymbol;
 	FString Sample = TEXT("SampleCurve_") + ParamInfo.DataInterfaceHLSLSymbol;
@@ -167,7 +167,7 @@ void %s(in float In_X, out float2 Out_Value) \n\
 	float2 B = float2(%s(Next), %s(Next + 1)); \n\
 	Out_Value = lerp(A, B, Interp); \n\
 }\n")
-, *InstanceFunctionName, *TimeToLUTFrac, *NumSamples, *NumSamples, CurveLUTNumElems, CurveLUTNumElems, *Sample, *Sample, *Sample, *Sample);
+, *FunctionInfo.InstanceName, *TimeToLUTFrac, *NumSamples, *NumSamples, CurveLUTNumElems, CurveLUTNumElems, *Sample, *Sample, *Sample, *Sample);
 
 	return true;
 }
diff --git a/Engine/Plugins/FX/Niagara/Source/Niagara/Private/NiagaraDataInterfaceVector4Curve.cpp b/Engine/Plugins/FX/Niagara/Source/Niagara/Private/NiagaraDataInterfaceVector4Curve.cpp
--- a/Engine/Plugins/FX/Niagara/Source/Niagara/Private/NiagaraDataInterfaceVector4Curve.cpp
+++ b/Engine/Plugins/FX/Niagara/Source/Niagara/Private/NiagaraDataInterfaceVector4Curve.cpp
@@ -170,7 +170,7 @@ void UNiagaraDataInterfaceVector4Curve::GetFunctions(TArray<FNiagaraFunctionSign
 // the HLSL in the spirit of a static switch
 // TODO: need a way to identify each specific function here
 // 
-bool UNiagaraDataInterfaceVector4Curve::GetFunctionHLSL(const FName& DefinitionFunctionName, FString InstanceFunctionName, FNiagaraDataInterfaceGPUParamInfo& ParamInfo, FString& OutHLSL)
+bool UNiagaraDataInterfaceVector4Curve::GetFunctionHLSL(const FNiagaraDataInterfaceGPUParamInfo& ParamInfo, const FNiagaraDataInterfaceGeneratedFunction& FunctionInfo, int FunctionInstanceIndex, FString& OutHLSL)
 {
 	FString TimeToLUTFrac = TEXT("TimeToLUTFraction_") + ParamInfo.DataInterfaceHLSLSymbol;
 	FString Sample = TEXT("SampleCurve_") + ParamInfo.DataInterfaceHLSLSymbol;
@@ -188,7 +188,7 @@ void %s(in float In_X, out float4 Out_Value) \n\
 	float4 B = float4(%s(Next), %s(Next + 1), %s(Next + 2), %s(Next + 3)); \n\
 	Out_Value = lerp(A, B, Interp); \n\
 }\n")
-, *InstanceFunctionName, *TimeToLUTFrac, *NumSamples, *NumSamples, CurveLUTNumElems, CurveLUTNumElems
+, *FunctionInfo.InstanceName, *TimeToLUTFrac, *NumSamples, *NumSamples, CurveLUTNumElems, CurveLUTNumElems
 , *Sample, *Sample, *Sample, *Sample, *Sample, *Sample, *Sample, *Sample);
 
 	return true;
diff --git a/Engine/Plugins/FX/Niagara/Source/Niagara/Private/NiagaraDataInterfaceVectorCurve.cpp b/Engine/Plugins/FX/Niagara/Source/Niagara/Private/NiagaraDataInterfaceVectorCurve.cpp
--- a/Engine/Plugins/FX/Niagara/Source/Niagara/Private/NiagaraDataInterfaceVectorCurve.cpp
+++ b/Engine/Plugins/FX/Niagara/Source/Niagara/Private/NiagaraDataInterfaceVectorCurve.cpp
@@ -161,7 +161,7 @@ void UNiagaraDataInterfaceVectorCurve::GetFunctions(TArray<FNiagaraFunctionSigna
 // the HLSL in the spirit of a static switch
 // TODO: need a way to identify each specific function here
 // 
-bool UNiagaraDataInterfaceVectorCurve::GetFunctionHLSL(const FName& DefinitionFunctionName, FString InstanceFunctionName, FNiagaraDataInterfaceGPUParamInfo& ParamInfo, FString& OutHLSL)
+bool UNiagaraDataInterfaceVectorCurve::GetFunctionHLSL(const FNiagaraDataInterfaceGPUParamInfo& ParamInfo, const FNiagaraDataInterfaceGeneratedFunction& FunctionInfo, int FunctionInstanceIndex, FString& OutHLSL)
 {
 	FString TimeToLUTFrac = TEXT("TimeToLUTFraction_") + ParamInfo.DataInterfaceHLSLSymbol;
 	FString Sample = TEXT("SampleCurve_") + ParamInfo.DataInterfaceHLSLSymbol;
@@ -179,7 +179,7 @@ void %s(in float In_X, out float3 Out_Value) \n\
 	float3 B = float3(%s(Next), %s(Next + 1), %s(Next + 2)); \n\
 	Out_Value = lerp(A, B, Interp); \n\
 }\n") 
-, *InstanceFunctionName, *TimeToLUTFrac, *NumSamples, *NumSamples, CurveLUTNumElems, CurveLUTNumElems
+, *FunctionInfo.InstanceName, *TimeToLUTFrac, *NumSamples, *NumSamples, CurveLUTNumElems, CurveLUTNumElems
 , *Sample, *Sample, *Sample, *Sample, *Sample, *Sample);
 
 	return true;
diff --git a/Engine/Plugins/FX/Niagara/Source/Niagara/Private/NiagaraDataInterfaceVectorField.cpp b/Engine/Plugins/FX/Niagara/Source/Niagara/Private/NiagaraDataInterfaceVectorField.cpp
--- a/Engine/Plugins/FX/Niagara/Source/Niagara/Private/NiagaraDataInterfaceVectorField.cpp
+++ b/Engine/Plugins/FX/Niagara/Source/Niagara/Private/NiagaraDataInterfaceVectorField.cpp
@@ -222,7 +222,7 @@ TArray<FNiagaraDataInterfaceError> UNiagaraDataInterfaceVectorField::GetErrors()
 
 /*--------------------------------------------------------------------------------------------------------------------------*/
 
-void UNiagaraDataInterfaceVectorField::GetParameterDefinitionHLSL(FNiagaraDataInterfaceGPUParamInfo& ParamInfo, FString& OutHLSL)
+void UNiagaraDataInterfaceVectorField::GetParameterDefinitionHLSL(const FNiagaraDataInterfaceGPUParamInfo& ParamInfo, FString& OutHLSL)
 {
 	static const TCHAR *FormatDeclarations = TEXT(R"(
 		float3 {TilingAxesName};
@@ -242,9 +242,9 @@ void UNiagaraDataInterfaceVectorField::GetParameterDefinitionHLSL(FNiagaraDataIn
 	OutHLSL += FString::Format(FormatDeclarations, ArgsDeclarations);
 }
 
-bool UNiagaraDataInterfaceVectorField::GetFunctionHLSL(const FName& DefinitionFunctionName, FString InstanceFunctionName, FNiagaraDataInterfaceGPUParamInfo& ParamInfo, FString& OutHLSL)
+bool UNiagaraDataInterfaceVectorField::GetFunctionHLSL(const FNiagaraDataInterfaceGPUParamInfo& ParamInfo, const FNiagaraDataInterfaceGeneratedFunction& FunctionInfo, int FunctionInstanceIndex, FString& OutHLSL)
 {
-	if (DefinitionFunctionName == SampleVectorFieldName)
+	if (FunctionInfo.DefinitionName == SampleVectorFieldName)
 	{
 		static const TCHAR *FormatSample = TEXT(R"(
 			void {FunctionName}(float3 In_SamplePoint, out float3 Out_Sample)
@@ -254,7 +254,7 @@ bool UNiagaraDataInterfaceVectorField::GetFunctionHLSL(const FName& DefinitionFu
 			}
 		)");
 		TMap<FString, FStringFormatArg> ArgsSample;
-		ArgsSample.Add(TEXT("FunctionName"), InstanceFunctionName);
+		ArgsSample.Add(TEXT("FunctionName"), FunctionInfo.InstanceName);
 		ArgsSample.Add(TEXT("TextureName"), TextureBaseName + ParamInfo.DataInterfaceHLSLSymbol);
 		ArgsSample.Add(TEXT("MinBoundsName"), MinBoundsBaseName + ParamInfo.DataInterfaceHLSLSymbol);
 		ArgsSample.Add(TEXT("MaxBoundsName"), MaxBoundsBaseName + ParamInfo.DataInterfaceHLSLSymbol);
@@ -262,7 +262,7 @@ bool UNiagaraDataInterfaceVectorField::GetFunctionHLSL(const FName& DefinitionFu
 		OutHLSL += FString::Format(FormatSample, ArgsSample);
 		return true;
 	}
-	else if (DefinitionFunctionName == GetVectorFieldTilingAxesName)
+	else if (FunctionInfo.DefinitionName == GetVectorFieldTilingAxesName)
 	{
 		static const TCHAR *FormatTilingAxes = TEXT(R"(
 			void {FunctionName}(out float3 Out_TilingAxes)
@@ -271,12 +271,12 @@ bool UNiagaraDataInterfaceVectorField::GetFunctionHLSL(const FName& DefinitionFu
 			}
 		)");
 		TMap<FString, FStringFormatArg> ArgsTilingAxes;
-		ArgsTilingAxes.Add(TEXT("FunctionName"), InstanceFunctionName);
+		ArgsTilingAxes.Add(TEXT("FunctionName"), FunctionInfo.InstanceName);
 		ArgsTilingAxes.Add(TEXT("TilingAxesName"), TilingAxesBaseName + ParamInfo.DataInterfaceHLSLSymbol);
 		OutHLSL += FString::Format(FormatTilingAxes, ArgsTilingAxes);
 		return true;
 	}
-	else if (DefinitionFunctionName == GetVectorFieldDimensionsName)
+	else if (FunctionInfo.DefinitionName == GetVectorFieldDimensionsName)
 	{
 		static const TCHAR *FormatDimensions = TEXT(R"(
 			void {FunctionName}(out float3 Out_Dimensions)
@@ -285,12 +285,12 @@ bool UNiagaraDataInterfaceVectorField::GetFunctionHLSL(const FName& DefinitionFu
 			}
 		)");
 		TMap<FString, FStringFormatArg> ArgsDimensions;
-		ArgsDimensions.Add(TEXT("FunctionName"), InstanceFunctionName);
+		ArgsDimensions.Add(TEXT("FunctionName"), FunctionInfo.InstanceName);
 		ArgsDimensions.Add(TEXT("DimensionsName"), DimensionsBaseName + ParamInfo.DataInterfaceHLSLSymbol);
 		OutHLSL += FString::Format(FormatDimensions, ArgsDimensions);
 		return true;
 	}
-	else if (DefinitionFunctionName == GetVectorFieldBoundsName)
+	else if (FunctionInfo.DefinitionName == GetVectorFieldBoundsName)
 	{
 		static const TCHAR *FormatBounds = TEXT(R"(
 			void {FunctionName}(out float3 Out_MinBounds, out float3 Out_MaxBounds)
@@ -300,7 +300,7 @@ bool UNiagaraDataInterfaceVectorField::GetFunctionHLSL(const FName& DefinitionFu
 			}
 		)");
 		TMap<FString, FStringFormatArg> ArgsBounds;
-		ArgsBounds.Add(TEXT("FunctionName"), InstanceFunctionName);
+		ArgsBounds.Add(TEXT("FunctionName"), FunctionInfo.InstanceName);
 		ArgsBounds.Add(TEXT("MinBoundsName"), MinBoundsBaseName + ParamInfo.DataInterfaceHLSLSymbol);
 		ArgsBounds.Add(TEXT("MaxBoundsName"), MaxBoundsBaseName + ParamInfo.DataInterfaceHLSLSymbol);
 		OutHLSL += FString::Format(FormatBounds, ArgsBounds);
diff --git a/Engine/Plugins/FX/Niagara/Source/Niagara/Private/NiagaraDataSet.cpp b/Engine/Plugins/FX/Niagara/Source/Niagara/Private/NiagaraDataSet.cpp
--- a/Engine/Plugins/FX/Niagara/Source/Niagara/Private/NiagaraDataSet.cpp
+++ b/Engine/Plugins/FX/Niagara/Source/Niagara/Private/NiagaraDataSet.cpp
@@ -80,6 +80,7 @@ FNiagaraDataSet::FNiagaraDataSet()
 	, NumFreeIDs(0)
 	, MaxUsedID(0)
 	, IDAcquireTag(0)
+	, GPUNumAllocatedIDs(0)
 	, CurrentData(nullptr)
 	, DestinationData(nullptr)
 	, bInitialized(false)
@@ -143,6 +144,13 @@ void FNiagaraDataSet::ReleaseBuffers()
 		}
 		Data.Empty();
 	}
+
+	if (GPUFreeIDs.Buffer)
+	{
+		GPUFreeIDs.Release();
+	}
+
+	GPUNumAllocatedIDs = 0;
 }
 
 FNiagaraDataBuffer& FNiagaraDataSet::BeginSimulate()
@@ -339,6 +347,48 @@ void FNiagaraDataSet::ReleaseGPUInstanceCounts(FNiagaraGPUInstanceCountManager&
 	}
 }
 
+void FNiagaraDataSet::AllocateGPUFreeIDs(uint32 InNumInstances, FRHICommandList& RHICmdList, ERHIFeatureLevel::Type FeatureLevel, const TCHAR* DebugSimName)
+{
+	checkSlow(GetSimTarget() == ENiagaraSimTarget::GPUComputeSim && GetNeedsPersistentIDs());
+
+	// Clearing and compacting the ID table must run over all the allocated elements, so we must use a chunk size which balances
+	// between reallocation frequency and the cost of processing unused elements.
+	constexpr uint32 ID_ALLOC_CHUNKSIZE = 1024;
+	const uint32 NumIDsToAlloc = FMath::DivideAndRoundUp(InNumInstances, ID_ALLOC_CHUNKSIZE) * ID_ALLOC_CHUNKSIZE;
+
+	if (NumIDsToAlloc <= GPUNumAllocatedIDs)
+	{
+		// We can never shrink the ID buffer, because IDs with numeric values larger than the current number of instances
+		// might still be in use.
+		return;
+	}
+
+	TCHAR DebugBufferName[128];
+	FCString::Snprintf(DebugBufferName, UE_ARRAY_COUNT(DebugBufferName), TEXT("NiagaraFreeIDList_%s"), DebugSimName ? DebugSimName : TEXT(""));
+	FRWBuffer NewFreeIDsBuffer;
+	NewFreeIDsBuffer.Initialize(sizeof(int32), NumIDsToAlloc, EPixelFormat::PF_R32_SINT, BUF_Static, DebugBufferName);
+
+	// We must maintain the existing list of free IDs.
+	FRHIShaderResourceView* ExistingBuffer;
+	if (GPUNumAllocatedIDs > 0)
+	{
+		// The free IDs buffer was written in the previous simulation step, but hasn't been transitioned to read yet, so we must
+		// transition it explicitly here. The new buffer will be transitioned by NiagaraEmitterInstanceBatcher::DispatchAllOnCompute(),
+		// so there's no need for a barrier at the end of this function.
+		RHICmdList.TransitionResource(EResourceTransitionAccess::EReadable, EResourceTransitionPipeline::EComputeToCompute, GPUFreeIDs.UAV);
+		ExistingBuffer = GPUFreeIDs.SRV;
+	}
+	else
+	{
+		ExistingBuffer = FNiagaraRenderer::GetDummyIntBuffer().SRV;
+	}
+	
+	NiagaraInitGPUFreeIDList(RHICmdList, FeatureLevel, NumIDsToAlloc, NewFreeIDsBuffer, GPUNumAllocatedIDs, ExistingBuffer);
+
+	GPUFreeIDs = MoveTemp(NewFreeIDsBuffer);
+	GPUNumAllocatedIDs = NumIDsToAlloc;
+}
+
 const FNiagaraVariableLayoutInfo* FNiagaraDataSet::GetVariableLayout(const FNiagaraVariable& Var)const
 {
 	int32 VarLayoutIndex = GetVariables().IndexOfByKey(Var);
@@ -431,7 +481,6 @@ FNiagaraDataBuffer::FNiagaraDataBuffer(FNiagaraDataSet* InOwner)
 	: Owner(InOwner)
 	, GPUInstanceCountBufferOffset(INDEX_NONE)
 	, NumChunksAllocatedForGPU(0)
-	, NumIDsAllocatedForGPU(0)
 	, NumInstances(0)
 	, NumInstancesAllocated(0)
 	, FloatStride(0)
@@ -600,7 +649,7 @@ void FNiagaraDataBuffer::Allocate(uint32 InNumInstances, bool bMaintainExisting)
 	}
 }
 
-void FNiagaraDataBuffer::AllocateGPU(uint32 InNumInstances, FNiagaraGPUInstanceCountManager& GPUInstanceCountManager, FRHICommandList& RHICmdList, ERHIFeatureLevel::Type FeatureLevel)
+void FNiagaraDataBuffer::AllocateGPU(uint32 InNumInstances, FNiagaraGPUInstanceCountManager& GPUInstanceCountManager, FRHICommandList& RHICmdList, ERHIFeatureLevel::Type FeatureLevel, const TCHAR* DebugSimName)
 {
 	CheckUsage(false);
 
@@ -635,17 +684,12 @@ void FNiagaraDataBuffer::AllocateGPU(uint32 InNumInstances, FNiagaraGPUInstanceC
 		{
 			GPUBufferInt.Release();
 		}
-		if (GPUFreeIDs.Buffer)
-		{
-			GPUFreeIDs.Release();
-		}
 		if (GPUIDToIndexTable.Buffer)
 		{
 			GPUIDToIndexTable.Release();
 		}
 
 		NumChunksAllocatedForGPU = 0;
-		NumIDsAllocatedForGPU = 0;
 	}
 	else // Otherwise check for growing and possibly shrinking (if GNiagaraGPUDataBufferBufferSlack > 1) .
 	{
@@ -673,31 +717,20 @@ void FNiagaraDataBuffer::AllocateGPU(uint32 InNumInstances, FNiagaraGPUInstanceC
 			}
 		}
 
-		// The ID buffers cannot be shrinked, because the maximum ID could be currently in use even if we have a single alive particle.
-		if (Owner->GetNeedsPersistentIDs() && PaddedNumInstances > NumIDsAllocatedForGPU)
+		if (Owner->GetNeedsPersistentIDs())
 		{
-			// Smaller chunk size for ID allocations, because clearing and compacting the ID table must run over all the allocated
-			// elements, so too much slack memory hurts performance.
-			constexpr uint32 ID_ALLOC_CHUNKSIZE = 1024;
-			const uint32 NumIDsToAlloc = FMath::DivideAndRoundUp(PaddedNumInstances, ID_ALLOC_CHUNKSIZE) * ID_ALLOC_CHUNKSIZE;
-
-			// We must maintain the existing list of free IDs.
-			FRWBuffer NewFreeIDsBuffer;
-			NewFreeIDsBuffer.Initialize(sizeof(int32), NumIDsToAlloc, EPixelFormat::PF_R32_SINT, BUF_Static, TEXT("NiagaraFreeIDList"));
-
-			FRHIShaderResourceView* ExistingBuffer = NumIDsAllocatedForGPU > 0 ? GPUFreeIDs.SRV : FNiagaraRenderer::GetDummyIntBuffer().SRV;
-			NiagaraInitGPUFreeIDList(RHICmdList, FeatureLevel, NumIDsToAlloc, NewFreeIDsBuffer, NumIDsAllocatedForGPU, ExistingBuffer);
-
-			GPUFreeIDs = MoveTemp(NewFreeIDsBuffer);
-
-			// The ID to index table is filled in each time the emitter runs, so we don't need to keep its contents.
-			if (GPUIDToIndexTable.Buffer)
+			uint32 NumExistingElems = GPUIDToIndexTable.Buffer ? GPUIDToIndexTable.Buffer->GetSize() / sizeof(int32) : 0;
+			uint32 NumNeededElems = Owner->GetGPUNumAllocatedIDs();
+			if (NumExistingElems < NumNeededElems)
 			{
-				GPUIDToIndexTable.Release();
+				if (GPUIDToIndexTable.Buffer)
+				{
+					GPUIDToIndexTable.Release();
+				}
+				TCHAR DebugBufferName[128];
+				FCString::Snprintf(DebugBufferName, UE_ARRAY_COUNT(DebugBufferName), TEXT("NiagaraIDToIndexTable_%s_%p"), DebugSimName ? DebugSimName : TEXT(""), this);
+				GPUIDToIndexTable.Initialize(sizeof(int32), NumNeededElems, EPixelFormat::PF_R32_SINT, BUF_Static, DebugBufferName);
 			}
-			GPUIDToIndexTable.Initialize(sizeof(int32), NumIDsToAlloc, EPixelFormat::PF_R32_SINT, BUF_Static, TEXT("NiagaraIDToIndexTable"));
-
-			NumIDsAllocatedForGPU = NumIDsToAlloc;
 		}
 	}
 }
@@ -944,7 +977,6 @@ void FNiagaraDataBuffer::SetShaderParams(FNiagaraShader *Shader, FRHICommandList
 		SetSRVParameter(CommandList, ComputeShader, Shader->FloatInputBufferParam, InstancesAllocated ? GetGPUBufferFloat().SRV : FNiagaraRenderer::GetDummyFloatBuffer().SRV);
 		SetSRVParameter(CommandList, ComputeShader, Shader->IntInputBufferParam, InstancesAllocated ? GetGPUBufferInt().SRV : FNiagaraRenderer::GetDummyIntBuffer().SRV);
 		SetShaderValue(CommandList, ComputeShader, Shader->ComponentBufferSizeReadParam, SafeBufferSize);
-		SetSRVParameter(CommandList, ComputeShader, Shader->FreeIDBufferParam, GPUFreeIDs.SRV ? GPUFreeIDs.SRV : FNiagaraRenderer::GetDummyIntBuffer().SRV);
 	}
 	else
 	{
diff --git a/Engine/Plugins/FX/Niagara/Source/Niagara/Private/NiagaraEmitterInstanceBatcher.cpp b/Engine/Plugins/FX/Niagara/Source/Niagara/Private/NiagaraEmitterInstanceBatcher.cpp
--- a/Engine/Plugins/FX/Niagara/Source/Niagara/Private/NiagaraEmitterInstanceBatcher.cpp
+++ b/Engine/Plugins/FX/Niagara/Source/Niagara/Private/NiagaraEmitterInstanceBatcher.cpp
@@ -9,6 +9,7 @@
 #include "NiagaraSortingGPU.h"
 #include "NiagaraWorldManager.h"
 #include "NiagaraShaderParticleID.h"
+#include "NiagaraRenderer.h"
 #include "ShaderParameterUtils.h"
 #include "SceneUtils.h"
 #include "ShaderParameterUtils.h"
@@ -398,7 +399,12 @@ void NiagaraEmitterInstanceBatcher::ResizeBuffersAndGatherResources(FOverlappabl
 			const uint32 RequiredInstances = FMath::Max(PrevNumInstances, NewNumInstances);
 			const uint32 AllocatedInstances = FMath::Max(RequiredInstances, Instance.SpawnInfo.MaxParticleCount);
 
-			DestinationData.AllocateGPU(AllocatedInstances + 1, GPUInstanceCounterManager, RHICmdList, FeatureLevel);
+			if (bNeedsPersistentIDs)
+			{
+				Context->MainDataSet->AllocateGPUFreeIDs(AllocatedInstances + 1, RHICmdList, FeatureLevel, Context->GetDebugSimName());
+			}
+
+			DestinationData.AllocateGPU(AllocatedInstances + 1, GPUInstanceCounterManager, RHICmdList, FeatureLevel, Context->GetDebugSimName());
 			DestinationData.SetNumInstances(RequiredInstances);
 
 			WriteBuffers.Add(GPUInstanceCounterManager.GetInstanceCountBuffer().UAV);
@@ -423,8 +429,8 @@ void NiagaraEmitterInstanceBatcher::ResizeBuffersAndGatherResources(FOverlappabl
 
 			if (bNeedsPersistentIDs)
 			{
-				ReadBuffers.Add(CurrentData.GetGPUFreeIDs().UAV);
-				WriteBuffers.Add(CurrentData.GetGPUIDToIndexTable().UAV);
+				ReadBuffers.Add(Context->MainDataSet->GetGPUFreeIDs().UAV);
+				WriteBuffers.Add(DestinationData.GetGPUIDToIndexTable().UAV);
 				++NumFreeIDListsToCompute;
 			}
 
@@ -1195,7 +1201,8 @@ void NiagaraEmitterInstanceBatcher::Run(const FNiagaraGPUSystemTick& Tick, const
 	FNiagaraDataBuffer& DestinationData = *Instance->DestinationData;
 	FNiagaraDataBuffer& CurrentData = *Instance->CurrentData;
 
-	if (Context->MainDataSet->GetNeedsPersistentIDs())
+	const bool bNeedsPersistentIDs = Context->MainDataSet->GetNeedsPersistentIDs();
+	if (bNeedsPersistentIDs)
 	{
 		// Put INDEX_NONE in all the slots of the ID to index table. The simulation will fill in the
 		// indices for the IDs which are currently in use, so we can compute the free ID list based on
@@ -1225,6 +1232,7 @@ void NiagaraEmitterInstanceBatcher::Run(const FNiagaraGPUSystemTick& Tick, const
 
 	// set the shader and data set params 
 	//
+	SetSRVParameter(RHICmdList, Shader->GetComputeShader(), Shader->FreeIDBufferParam, bNeedsPersistentIDs ? Context->MainDataSet->GetGPUFreeIDs().SRV : FNiagaraRenderer::GetDummyIntBuffer().SRV);
 	CurrentData.SetShaderParams(Shader, RHICmdList, true);
 	DestinationData.SetShaderParams(Shader, RHICmdList, false);
 
@@ -1379,7 +1387,7 @@ void NiagaraEmitterInstanceBatcher::Run(const FNiagaraGPUSystemTick& Tick, const
 	DestinationData.UnsetShaderParams(Shader, RHICmdList);
 	Shader->InstanceCountsParam.UnsetUAV(RHICmdList, Shader->GetComputeShader());
 
-	if (Context->MainDataSet->GetNeedsPersistentIDs() && ShaderStageIndex == Context->MaxUpdateIterations - 1)
+	if (bNeedsPersistentIDs && ShaderStageIndex == Context->MaxUpdateIterations - 1)
 	{
 		SCOPED_DRAW_EVENT(RHICmdList, NiagaraGPUComputeFreeIDs);
 		SCOPED_GPU_STAT(RHICmdList, NiagaraGPUComputeFreeIDs);
@@ -1388,8 +1396,8 @@ void NiagaraEmitterInstanceBatcher::Run(const FNiagaraGPUSystemTick& Tick, const
 
 		// Make the ID to index table readable and the free ID list writable.
 		RHICmdList.TransitionResource(EResourceTransitionAccess::EReadable, EResourceTransitionPipeline::EComputeToCompute, DestinationData.GetGPUIDToIndexTable().UAV);
-		RHICmdList.TransitionResource(EResourceTransitionAccess::ERWBarrier, EResourceTransitionPipeline::EComputeToCompute, DestinationData.GetGPUFreeIDs().UAV);
-		NiagaraComputeGPUFreeIDs(RHICmdList, FeatureLevel, DestinationData.GetGPUNumAllocatedIDs(), DestinationData.GetGPUIDToIndexTable().SRV, DestinationData.GetGPUFreeIDs(), FreeIDListSizesBuffer, CurrentFreeIDListIndex);
+		RHICmdList.TransitionResource(EResourceTransitionAccess::ERWBarrier, EResourceTransitionPipeline::EComputeToCompute, Context->MainDataSet->GetGPUFreeIDs().UAV);
+		NiagaraComputeGPUFreeIDs(RHICmdList, FeatureLevel, Context->MainDataSet->GetGPUNumAllocatedIDs(), DestinationData.GetGPUIDToIndexTable().SRV, Context->MainDataSet->GetGPUFreeIDs(), FreeIDListSizesBuffer, CurrentFreeIDListIndex);
 		++CurrentFreeIDListIndex;
 	}
 }
diff --git a/Engine/Plugins/FX/Niagara/Source/NiagaraCore/Private/NiagaraCustomVersion.cpp b/Engine/Plugins/FX/Niagara/Source/NiagaraCore/Private/NiagaraCustomVersion.cpp
--- a/Engine/Plugins/FX/Niagara/Source/NiagaraCore/Private/NiagaraCustomVersion.cpp
+++ b/Engine/Plugins/FX/Niagara/Source/NiagaraCore/Private/NiagaraCustomVersion.cpp
@@ -8,4 +8,4 @@ const FGuid FNiagaraCustomVersion::GUID(0xFCF57AFA, 0x50764283, 0xB9A9E658, 0xFF
 // Register the custom version with core
 FCustomVersionRegistration GRegisterNiagaraCustomVersion(FNiagaraCustomVersion::GUID, FNiagaraCustomVersion::LatestVersion, TEXT("NiagaraVer"));
 
-const FGuid FNiagaraCustomVersion::LatestScriptCompileVersion(0x61632118, 0x26994E14, 0x964DB60F, 0xD4705785);
+const FGuid FNiagaraCustomVersion::LatestScriptCompileVersion(0x64FCB871, 0xB2954D47, 0x8721FF75, 0x20496D93);
diff --git a/Engine/Plugins/FX/Niagara/Source/NiagaraCore/Public/NiagaraCustomVersion.h b/Engine/Plugins/FX/Niagara/Source/NiagaraCore/Public/NiagaraCustomVersion.h
--- a/Engine/Plugins/FX/Niagara/Source/NiagaraCore/Public/NiagaraCustomVersion.h
+++ b/Engine/Plugins/FX/Niagara/Source/NiagaraCore/Public/NiagaraCustomVersion.h
@@ -138,6 +138,8 @@ struct FNiagaraCustomVersion
 
 		AddShaderStageUsageEnum, // Added a new value to the script usage enum, and we need a custom version to fix the existing bitfields.
 
+		AddGeneratedFunctionsToGPUParamInfo, // Save the functions generated by a GPU data interface inside FNiagaraDataInterfaceGPUParamInfo
+
 		// DO NOT ADD A NEW VERSION UNLESS YOU HAVE TALKED TO THE NIAGARA LEAD. Mismanagement of these versions can lead to data loss if it is adjusted in multiple streams simultaneously.
 		// -----<new versions can be added above this line>  -------------------------------------------------
 		VersionPlusOne,
diff --git a/Engine/Plugins/FX/Niagara/Source/NiagaraEditor/Private/NiagaraHlslTranslator.cpp b/Engine/Plugins/FX/Niagara/Source/NiagaraEditor/Private/NiagaraHlslTranslator.cpp
--- a/Engine/Plugins/FX/Niagara/Source/NiagaraEditor/Private/NiagaraHlslTranslator.cpp
+++ b/Engine/Plugins/FX/Niagara/Source/NiagaraEditor/Private/NiagaraHlslTranslator.cpp
@@ -1751,35 +1751,54 @@ void FHlslNiagaraTranslator::DefineDataInterfaceHLSL(FString &InHlslOutput)
 
 			// grab the buffer definition from the interface
 			//
-			int32 NewIdx = DIParamInfo.AddDefaulted(1);
-			DIParamInfo[NewIdx].DataInterfaceHLSLSymbol = SanitizedOwnerIDString;
-			DIParamInfo[NewIdx].DIClassName = Info.Type.GetClass()->GetName();
-			CDO->GetParameterDefinitionHLSL(DIParamInfo[NewIdx], InterfaceUniformHLSL);
+			FNiagaraDataInterfaceGPUParamInfo& DIInstanceInfo = DIParamInfo.AddDefaulted_GetRef();
+			DIInstanceInfo.DataInterfaceHLSLSymbol = SanitizedOwnerIDString;
+			DIInstanceInfo.DIClassName = Info.Type.GetClass()->GetName();
 
-			// grab the function hlsl from the interface
-			//
+			// Build a list of function instances that will be generated for this DI.
 			const TSet<FNiagaraFunctionSignature>* DataInterfaceFunctions = DataInterfaceRegisteredFunctions.Find(Info.Type.GetFName());
 			if (DataInterfaceFunctions != nullptr)
 			{
+				DIInstanceInfo.GeneratedFunctions.Reserve(DataInterfaceFunctions->Num());
 				for (const FNiagaraFunctionSignature& OriginalSig : *DataInterfaceFunctions)
 				{
-					FNiagaraFunctionSignature Sig = OriginalSig;	// make a copy so we can modify the owner id and get the correct hlsl signature
+					if (!OriginalSig.bSupportsGPU)
+					{
+						Error(FText::Format(LOCTEXT("GPUDataInterfaceFunctionNotSupported", "DataInterface {0} function {1} cannot run on the GPU."), FText::FromName(Info.Type.GetFName()), FText::FromName(OriginalSig.Name)), nullptr, nullptr);
+						continue;
+					}
+
+					// make a copy so we can modify the owner id and get the correct hlsl signature
+					FNiagaraFunctionSignature Sig = OriginalSig;
 					Sig.OwnerName = Info.Name;
-					FString DefStr = GetFunctionSignatureSymbol(Sig);
 
-					const bool HlslOK = CDO->GetFunctionHLSL(Sig.Name, DefStr, DIParamInfo[NewIdx], InterfaceFunctionHLSL);
-					if (OriginalSig.bSupportsGPU == false || HlslOK == false)
+					FNiagaraDataInterfaceGeneratedFunction& DIFunc = DIInstanceInfo.GeneratedFunctions.AddDefaulted_GetRef();
+					DIFunc.DefinitionName = Sig.Name;
+					DIFunc.InstanceName = GetFunctionSignatureSymbol(Sig);
+					DIFunc.Specifiers.Empty(Sig.FunctionSpecifiers.Num());
+					for (const TTuple<FName, FName>& Specifier : Sig.FunctionSpecifiers)
 					{
-						Error(FText::Format(LOCTEXT("GPUDataInterfaceFunctionNotSupported", "DataInterface {0} function {1} cannot run on the GPU or is not implemented."), FText::FromName(Info.Type.GetFName()), FText::FromName(Sig.Name))
-							, nullptr, nullptr);
+						DIFunc.Specifiers.Add(Specifier);
 					}
 				}
 			}
+
+			CDO->GetParameterDefinitionHLSL(DIInstanceInfo, InterfaceUniformHLSL);
+
+			// Ask the DI to generate HLSL.
+			for(int FunctionInstanceIndex = 0; FunctionInstanceIndex < DIInstanceInfo.GeneratedFunctions.Num(); ++FunctionInstanceIndex)
+			{
+				const FNiagaraDataInterfaceGeneratedFunction& DIFunc = DIInstanceInfo.GeneratedFunctions[FunctionInstanceIndex];
+				const bool HlslOK = CDO->GetFunctionHLSL(DIInstanceInfo, DIFunc, FunctionInstanceIndex, InterfaceFunctionHLSL);
+				if (!HlslOK)
+				{
+					Error(FText::Format(LOCTEXT("GPUDataInterfaceFunctionNotSupported", "DataInterface {0} function {1} is not implemented for GPU."), FText::FromName(Info.Type.GetFName()), FText::FromName(DIFunc.DefinitionName)), nullptr, nullptr);
+				}
+			}
 		}
 		else
 		{
-			Error(FText::Format(LOCTEXT("NonGPUDataInterfaceError", "DataInterface {0} ({1}) cannot run on the GPU."), FText::FromName(Info.Name), FText::FromString(CDO ? CDO->GetClass()->GetName() : TEXT("")))
-				, nullptr, nullptr);
+			Error(FText::Format(LOCTEXT("NonGPUDataInterfaceError", "DataInterface {0} ({1}) cannot run on the GPU."), FText::FromName(Info.Name), FText::FromString(CDO ? CDO->GetClass()->GetName() : TEXT(""))), nullptr, nullptr);
 		}
 	}
 	InHlslOutput += InterfaceCommonHLSL + InterfaceUniformHLSL + InterfaceFunctionHLSL;
diff --git a/Engine/Plugins/FX/Niagara/Source/NiagaraShader/Private/NiagaraShader.cpp b/Engine/Plugins/FX/Niagara/Source/NiagaraShader/Private/NiagaraShader.cpp
--- a/Engine/Plugins/FX/Niagara/Source/NiagaraShader/Private/NiagaraShader.cpp
+++ b/Engine/Plugins/FX/Niagara/Source/NiagaraShader/Private/NiagaraShader.cpp
@@ -1490,7 +1490,7 @@ void FNiagaraDataInterfaceParamRef::Bind(const class FShaderParameterMap& Parame
 	}
 }
 
-bool operator<<(FArchive &Ar, FNiagaraDataInterfaceParamRef& ParamRef)
+bool operator<<(FArchive& Ar, FNiagaraDataInterfaceParamRef& ParamRef)
 {
 	ParamRef.ParameterInfo.Serialize(Ar);
 
@@ -1508,9 +1508,54 @@ bool operator<<(FArchive &Ar, FNiagaraDataInterfaceParamRef& ParamRef)
 	return true;
 }
 
+bool FNiagaraDataInterfaceGeneratedFunction::Serialize(FArchive& Ar)
+{
+	Ar.UsingCustomVersion(FNiagaraCustomVersion::GUID);
+	const int32 NiagaraVer = Ar.CustomVer(FNiagaraCustomVersion::GUID);
+
+	Ar << DefinitionName;
+	Ar << InstanceName;
+
+	//if (Ar.IsLoading() && NiagaraVer == FNiagaraCustomVersion::AddGeneratedFunctionsToGPUParamInfo)
+	//{
+	//	TMap<FName, FName> cacat;
+	//	Ar << cacat;
+	//	Specifiers.SetNum(cacat.Num());
+	//	for (const auto& muie : cacat)
+	//	{
+	//		Specifiers.Add(muie);
+	//	}
+	//}
+	//else
+	{
+		Ar << Specifiers;
+	}
+	return true;
+}
+
+bool operator<<(FArchive& Ar, FNiagaraDataInterfaceGeneratedFunction& DIFunction)
+{
+	return DIFunction.Serialize(Ar);
+}
+
 bool FNiagaraDataInterfaceGPUParamInfo::Serialize(FArchive& Ar)
 {
+	Ar.UsingCustomVersion(FNiagaraCustomVersion::GUID);
+	const int32 NiagaraVer = Ar.CustomVer(FNiagaraCustomVersion::GUID);
+
 	Ar << DataInterfaceHLSLSymbol;
 	Ar << DIClassName;
+
+	// If FNiagaraDataInterfaceGPUParamInfo was only included in the DI parameters of FNiagaraShader, we wouldn't need to worry about
+	// custom versions, because bumping FNiagaraCustomVersion::LatestScriptCompileVersion is enough to cause all shaders to be
+	// rebuilt and things to be serialized correctly. However, there's a property of type FNiagaraVMExecutableData inside UNiagaraScript,
+	// which in turn contains an array of FNiagaraDataInterfaceGPUParamInfo, so we must check the version in order to be able to load
+	// UNiagaraScript objects saved before GeneratedFunctions was introduced.
+	const bool SkipGeneratedFunctions = Ar.IsLoading() && (NiagaraVer < FNiagaraCustomVersion::AddGeneratedFunctionsToGPUParamInfo);
+	if (!SkipGeneratedFunctions)
+	{
+		Ar << GeneratedFunctions;
+	}
+
 	return true;
 }
diff --git a/Engine/Plugins/FX/Niagara/Source/NiagaraShader/Public/NiagaraShared.h b/Engine/Plugins/FX/Niagara/Source/NiagaraShader/Public/NiagaraShared.h
--- a/Engine/Plugins/FX/Niagara/Source/NiagaraShader/Public/NiagaraShared.h
+++ b/Engine/Plugins/FX/Niagara/Source/NiagaraShader/Public/NiagaraShared.h
@@ -77,6 +77,48 @@ public:
 	TArray<FGuid> StackGuids;
 };
 
+/**
+* Information about a data interface function generated by the translator.
+*/
+USTRUCT()
+struct NIAGARASHADER_API FNiagaraDataInterfaceGeneratedFunction
+{
+	GENERATED_USTRUCT_BODY()
+	
+	/** Name of the function as defined by the data interface. */
+	FName DefinitionName;
+
+	/** Name of the instance. Derived from the definition name but made unique for this DI instance and specifier values. */
+	FString InstanceName;
+
+	/** Specifier values for this instance. */
+	using FunctionSpecifier = TTuple<FName, FName>;
+	TArray<FunctionSpecifier> Specifiers;
+
+	const FName* FindSpecifierValue(const FName& Name) const
+	{
+		for (const TTuple<FName, FName>& Specifier : Specifiers)
+		{
+			if (Specifier.template Get<0>() == Name)
+			{
+				return &Specifier.template Get<1>();
+			}
+		}
+		return nullptr;
+	}
+
+	bool Serialize(FArchive& Ar);
+	friend bool operator<<(FArchive& Ar, FNiagaraDataInterfaceGeneratedFunction& DIFunction);
+};
+
+template<> struct TStructOpsTypeTraits<FNiagaraDataInterfaceGeneratedFunction> : public TStructOpsTypeTraitsBase2<FNiagaraDataInterfaceGeneratedFunction>
+{
+	enum
+	{
+		WithSerializer = true,
+	};
+};
+
 /**
 * Data coming from that translator that describes parameters needed for each data interface.
 */
@@ -93,6 +135,10 @@ struct NIAGARASHADER_API FNiagaraDataInterfaceGPUParamInfo
 	UPROPERTY()
 	FString DIClassName;
 
+	/** Information about all the functions generated by the translator for this data interface. */
+	UPROPERTY()
+	TArray<FNiagaraDataInterfaceGeneratedFunction> GeneratedFunctions;
+
 	bool Serialize(FArchive& Ar);
 };
 
