Niagara - Fix an issue where renderers duplicated with copy/paste wouldn't serialize correctly, also add guard code to prevent crashing when loading assets missing renderers.

diff --git a/Engine/Plugins/FX/Niagara/Source/Niagara/Private/NiagaraEmitter.cpp b/Engine/Plugins/FX/Niagara/Source/Niagara/Private/NiagaraEmitter.cpp
--- a/Engine/Plugins/FX/Niagara/Source/Niagara/Private/NiagaraEmitter.cpp
+++ b/Engine/Plugins/FX/Niagara/Source/Niagara/Private/NiagaraEmitter.cpp
@@ -293,6 +293,14 @@ void UNiagaraEmitter::PostLoad()
 		SetFlags(RF_Transactional);
 	}
 
+	for (int32 RendererIndex = RendererProperties.Num() - 1; RendererIndex >= 0; --RendererIndex)
+	{
+		if (ensureMsgf(RendererProperties[RendererIndex] != nullptr, TEXT("Null renderer found in %s at index %i, removing it to prevent crashes."), *GetPathName(), RendererIndex) == false)
+		{
+			RendererProperties.RemoveAt(RendererIndex);
+		}
+	}
+
 	if (!GPUComputeScript)
 	{
 		GPUComputeScript = NewObject<UNiagaraScript>(this, "GPUComputeScript", EObjectFlags::RF_Transactional);
diff --git a/Engine/Plugins/FX/Niagara/Source/Niagara/Public/NiagaraRendererProperties.h b/Engine/Plugins/FX/Niagara/Source/Niagara/Public/NiagaraRendererProperties.h
--- a/Engine/Plugins/FX/Niagara/Source/Niagara/Public/NiagaraRendererProperties.h
+++ b/Engine/Plugins/FX/Niagara/Source/Niagara/Public/NiagaraRendererProperties.h
@@ -50,6 +50,11 @@ public:
 
 	virtual const TArray<FNiagaraVariable>& GetRequiredAttributes() { static TArray<FNiagaraVariable> Vars; return Vars; };
 	virtual const TArray<FNiagaraVariable>& GetOptionalAttributes() { static TArray<FNiagaraVariable> Vars; return Vars; };
+
+	UNiagaraRendererProperties* StaticDuplicateWithNewMergeId(UObject* InOuter) const
+	{
+		return CastChecked<UNiagaraRendererProperties>(Super::StaticDuplicateWithNewMergeIdInternal(InOuter));
+	}
 #endif // WITH_EDITORONLY_DATA
 
 
diff --git a/Engine/Plugins/FX/Niagara/Source/NiagaraCore/Private/NiagaraMergeable.cpp b/Engine/Plugins/FX/Niagara/Source/NiagaraCore/Private/NiagaraMergeable.cpp
--- a/Engine/Plugins/FX/Niagara/Source/NiagaraCore/Private/NiagaraMergeable.cpp
+++ b/Engine/Plugins/FX/Niagara/Source/NiagaraCore/Private/NiagaraMergeable.cpp
@@ -46,4 +46,11 @@ FGuid UNiagaraMergeable::GetMergeId()
 {
 	return MergeId;
 }
+
+UNiagaraMergeable* UNiagaraMergeable::StaticDuplicateWithNewMergeIdInternal(UObject* InOuter) const
+{
+	UNiagaraMergeable* Duplicate = CastChecked<UNiagaraMergeable>(StaticDuplicateObject(this, InOuter));
+	Duplicate->MergeId = FGuid::NewGuid();
+	return Duplicate;
+}
 #endif
\ No newline at end of file
diff --git a/Engine/Plugins/FX/Niagara/Source/NiagaraCore/Public/NiagaraMergeable.h b/Engine/Plugins/FX/Niagara/Source/NiagaraCore/Public/NiagaraMergeable.h
--- a/Engine/Plugins/FX/Niagara/Source/NiagaraCore/Public/NiagaraMergeable.h
+++ b/Engine/Plugins/FX/Niagara/Source/NiagaraCore/Public/NiagaraMergeable.h
@@ -25,6 +25,9 @@ public:
 	FOnChanged& OnChanged();
 
 	FGuid GetMergeId();
+
+protected:
+	UNiagaraMergeable* StaticDuplicateWithNewMergeIdInternal(UObject* InOuter) const;
 #endif
 
 private:
diff --git a/Engine/Plugins/FX/Niagara/Source/NiagaraEditor/Private/NiagaraClipboard.cpp b/Engine/Plugins/FX/Niagara/Source/NiagaraEditor/Private/NiagaraClipboard.cpp
--- a/Engine/Plugins/FX/Niagara/Source/NiagaraEditor/Private/NiagaraClipboard.cpp
+++ b/Engine/Plugins/FX/Niagara/Source/NiagaraEditor/Private/NiagaraClipboard.cpp
@@ -136,7 +136,7 @@ const UNiagaraClipboardContent* FNiagaraClipboard::GetClipboardContent() const
 	FNiagaraClipboardContentTextObjectFactory ClipboardContentFactory;
 	if (ClipboardContentFactory.CanCreateObjectsFromText(ClipboardText))
 	{
-		ClipboardContentFactory.ProcessBuffer(GetTransientPackage(), RF_Transient, ClipboardText);
+		ClipboardContentFactory.ProcessBuffer(GetTransientPackage(), RF_Transactional, ClipboardText);
 		return ClipboardContentFactory.ClipboardContent;
 	}
 
diff --git a/Engine/Plugins/FX/Niagara/Source/NiagaraEditor/Private/ViewModels/Stack/NiagaraStackRenderItemGroup.cpp b/Engine/Plugins/FX/Niagara/Source/NiagaraEditor/Private/ViewModels/Stack/NiagaraStackRenderItemGroup.cpp
--- a/Engine/Plugins/FX/Niagara/Source/NiagaraEditor/Private/ViewModels/Stack/NiagaraStackRenderItemGroup.cpp
+++ b/Engine/Plugins/FX/Niagara/Source/NiagaraEditor/Private/ViewModels/Stack/NiagaraStackRenderItemGroup.cpp
@@ -149,11 +149,12 @@ void UNiagaraStackRenderItemGroup::Paste(const UNiagaraClipboardContent* Clipboa
 {
 	if (EmitterWeak.IsValid())
 	{
+		UNiagaraEmitter* Emitter = EmitterWeak.Get();
 		for (const UNiagaraRendererProperties* ClipboardRenderer : ClipboardContent->Renderers)
 		{
 			if (ClipboardRenderer != nullptr)
 			{
-				EmitterWeak->AddRenderer(CastChecked<UNiagaraRendererProperties>(StaticDuplicateObject(ClipboardRenderer, EmitterWeak.Get())));
+				EmitterWeak->AddRenderer(ClipboardRenderer->StaticDuplicateWithNewMergeId(Emitter));
 			}
 		}
 	}
