#UE4 Add some info to a log

diff --git a/Engine/Plugins/FX/Niagara/Source/Niagara/Private/NiagaraScript.cpp b/Engine/Plugins/FX/Niagara/Source/Niagara/Private/NiagaraScript.cpp
--- a/Engine/Plugins/FX/Niagara/Source/Niagara/Private/NiagaraScript.cpp
+++ b/Engine/Plugins/FX/Niagara/Source/Niagara/Private/NiagaraScript.cpp
@@ -1249,7 +1249,7 @@ void UNiagaraScript::SetVMCompilationResults(const FNiagaraVMExecutableDataId& I
 
 			if (Info.bIsPlaceholder == false)
 			{
-				UE_LOG(LogNiagara, Error, TEXT("We somehow ended up with a data interface that we couldn't match post compile. This shouldn't happen. Creating a dummy to prevent crashes. %s"), *Info.Name.ToString());
+				UE_LOG(LogNiagara, Error, TEXT("We somehow ended up with a data interface that we couldn't match post compile. This shouldn't happen. Creating a dummy to prevent crashes. DataInterfaceInfoName:%s Object:%s"), *Info.Name.ToString(), *GetPathNameSafe(this));
 			}
 		}
 		check(CachedDefaultDataInterfaces[Idx].DataInterface != nullptr);
