Niagara - Fixed an issue where scripts for shader stages would get outered incorrectly during merging which would prevent them from saving correctly.  Also added some guard code with an ensure to the emitter which will remove invalid shader stages to prevent crashes.

diff --git a/Engine/Plugins/FX/Niagara/Source/Niagara/Private/NiagaraEmitter.cpp b/Engine/Plugins/FX/Niagara/Source/Niagara/Private/NiagaraEmitter.cpp
--- a/Engine/Plugins/FX/Niagara/Source/Niagara/Private/NiagaraEmitter.cpp
+++ b/Engine/Plugins/FX/Niagara/Source/Niagara/Private/NiagaraEmitter.cpp
@@ -301,6 +301,14 @@ void UNiagaraEmitter::PostLoad()
 		}
 	}
 
+	for (int32 ShaderStageIndex = ShaderStages.Num() - 1; ShaderStageIndex >= 0; --ShaderStageIndex)
+	{
+		if (ensureMsgf(ShaderStages[ShaderStageIndex] != nullptr && ShaderStages[ShaderStageIndex]->Script != nullptr, TEXT("Null shader stage, or shader stage with a null script found in %s at index %i, removing it to prevent crashes."), *GetPathName(), ShaderStageIndex) == false)
+		{
+			ShaderStages.RemoveAt(ShaderStageIndex);
+		}
+	}
+
 	if (!GPUComputeScript)
 	{
 		GPUComputeScript = NewObject<UNiagaraScript>(this, "GPUComputeScript", EObjectFlags::RF_Transactional);
diff --git a/Engine/Plugins/FX/Niagara/Source/NiagaraEditor/Private/NiagaraScriptMergeManager.cpp b/Engine/Plugins/FX/Niagara/Source/NiagaraEditor/Private/NiagaraScriptMergeManager.cpp
--- a/Engine/Plugins/FX/Niagara/Source/NiagaraEditor/Private/NiagaraScriptMergeManager.cpp
+++ b/Engine/Plugins/FX/Niagara/Source/NiagaraEditor/Private/NiagaraScriptMergeManager.cpp
@@ -2449,8 +2449,8 @@ FNiagaraScriptMergeManager::FApplyDiffResults FNiagaraScriptMergeManager::ApplyS
 		{
 			UNiagaraEmitter* BaseEmitter = BaseEmitterAdapter->GetEditableEmitter();
 			UNiagaraShaderStageBase* AddedShaderStage = CastChecked<UNiagaraShaderStageBase>(StaticDuplicateObject(AddedOtherShaderStage->GetShaderStage(), BaseEmitter));
-			AddedShaderStage->Script = NewObject<UNiagaraScript>(BaseEmitter, MakeUniqueObjectName(BaseEmitter, UNiagaraScript::StaticClass(), "ShaderStage"), EObjectFlags::RF_Transactional);
-			AddedShaderStage->Script->SetUsage(ENiagaraScriptUsage::ParticleEventScript);
+			AddedShaderStage->Script = NewObject<UNiagaraScript>(AddedShaderStage, MakeUniqueObjectName(AddedShaderStage, UNiagaraScript::StaticClass(), "ShaderStage"), EObjectFlags::RF_Transactional);
+			AddedShaderStage->Script->SetUsage(ENiagaraScriptUsage::ParticleShaderStageScript);
 			AddedShaderStage->Script->SetUsageId(AddedOtherShaderStage->GetUsageId());
 			AddedShaderStage->Script->SetSource(EmitterSource);
 			BaseEmitter->AddShaderStage(AddedShaderStage);
