- Fix crash when attemping to copy / paste a module with custom HLSL inside

diff --git a/Engine/Plugins/FX/Niagara/Source/NiagaraEditor/Private/ViewModels/Stack/NiagaraStackModuleItem.cpp b/Engine/Plugins/FX/Niagara/Source/NiagaraEditor/Private/ViewModels/Stack/NiagaraStackModuleItem.cpp
--- a/Engine/Plugins/FX/Niagara/Source/NiagaraEditor/Private/ViewModels/Stack/NiagaraStackModuleItem.cpp
+++ b/Engine/Plugins/FX/Niagara/Source/NiagaraEditor/Private/ViewModels/Stack/NiagaraStackModuleItem.cpp
@@ -1005,7 +1005,7 @@ bool UNiagaraStackModuleItem::TestCanPasteWithMessage(const UNiagaraClipboardCon
 		bool bValidUsage = false;
 		for (const UNiagaraClipboardFunction* Function : ClipboardContent->Functions)
 		{
-			if (Function != nullptr && Function->Script->GetSupportedUsageContexts().Contains(OutputNode->GetUsage()))
+			if (Function != nullptr && Function->Script != nullptr && Function->Script->GetSupportedUsageContexts().Contains(OutputNode->GetUsage()))
 			{
 				bValidUsage = true;
 			}
diff --git a/Engine/Plugins/FX/Niagara/Source/NiagaraEditor/Private/ViewModels/Stack/NiagaraStackScriptItemGroup.cpp b/Engine/Plugins/FX/Niagara/Source/NiagaraEditor/Private/ViewModels/Stack/NiagaraStackScriptItemGroup.cpp
--- a/Engine/Plugins/FX/Niagara/Source/NiagaraEditor/Private/ViewModels/Stack/NiagaraStackScriptItemGroup.cpp
+++ b/Engine/Plugins/FX/Niagara/Source/NiagaraEditor/Private/ViewModels/Stack/NiagaraStackScriptItemGroup.cpp
@@ -335,7 +335,7 @@ bool UNiagaraStackScriptItemGroup::TestCanPasteWithMessage(const UNiagaraClipboa
 		bool bValidUsage = false;
 		for (const UNiagaraClipboardFunction* Function : ClipboardContent->Functions)
 		{
-			if (Function != nullptr && Function->Script->GetSupportedUsageContexts().Contains(ScriptUsage))
+			if (Function != nullptr && Function->Script != nullptr && Function->Script->GetSupportedUsageContexts().Contains(ScriptUsage))
 			{
 				bValidUsage = true;
 			}
