Niagara - Fix a crash when copying items with different allowed usages.

diff --git a/Engine/Plugins/FX/Niagara/Source/NiagaraEditor/Private/ViewModels/Stack/NiagaraParameterHandle.cpp b/Engine/Plugins/FX/Niagara/Source/NiagaraEditor/Private/ViewModels/Stack/NiagaraParameterHandle.cpp
--- a/Engine/Plugins/FX/Niagara/Source/NiagaraEditor/Private/ViewModels/Stack/NiagaraParameterHandle.cpp
+++ b/Engine/Plugins/FX/Niagara/Source/NiagaraEditor/Private/ViewModels/Stack/NiagaraParameterHandle.cpp
@@ -132,4 +132,25 @@ bool FNiagaraParameterHandle::IsModuleHandle() const
 bool FNiagaraParameterHandle::IsParameterCollectionHandle() const
 {
 	return Namespace == ParameterCollectionNamespace;
+}
+
+bool FNiagaraParameterHandle::IsReadOnlyHandle() const
+{
+	return
+		Namespace == UserNamespace ||
+		Namespace == EngineNamespace ||
+		Namespace == ParameterCollectionNamespace ||
+		Namespace == ModuleNamespace;
+}
+
+bool FNiagaraParameterHandle::IsTransientHandle() const
+{
+	return
+		Namespace != UserNamespace &&
+		Namespace != EngineNamespace &&
+		Namespace != SystemNamespace &&
+		Namespace != EmitterNamespace &&
+		Namespace != ParticleAttributeNamespace &&
+		Namespace != ModuleNamespace &&
+		Namespace != ParameterCollectionNamespace;
 }
\ No newline at end of file
diff --git a/Engine/Plugins/FX/Niagara/Source/NiagaraEditor/Private/ViewModels/Stack/NiagaraStackClipboardUtilities.cpp b/Engine/Plugins/FX/Niagara/Source/NiagaraEditor/Private/ViewModels/Stack/NiagaraStackClipboardUtilities.cpp
--- a/Engine/Plugins/FX/Niagara/Source/NiagaraEditor/Private/ViewModels/Stack/NiagaraStackClipboardUtilities.cpp
+++ b/Engine/Plugins/FX/Niagara/Source/NiagaraEditor/Private/ViewModels/Stack/NiagaraStackClipboardUtilities.cpp
@@ -183,7 +183,7 @@ bool FNiagaraStackClipboardUtilities::TestCanPasteSelectionWithMessage(const TAr
 	return false;
 }
 
-void FNiagaraStackClipboardUtilities::PasteSelection(const TArray<UNiagaraStackEntry*>& SelectedEntries)
+void FNiagaraStackClipboardUtilities::PasteSelection(const TArray<UNiagaraStackEntry*>& SelectedEntries, FText& OutPasteWarning)
 {
 	if (SelectedEntries.Num() == 1)
 	{
@@ -198,7 +198,7 @@ void FNiagaraStackClipboardUtilities::PasteSelection(const TArray<UNiagaraStackE
 		FScopedTransaction PasteTransaction(TransactionMessage);
 		if (ClipboardContent != nullptr)
 		{
-			SelectedEntries[0]->Paste(ClipboardContent);
+			SelectedEntries[0]->Paste(ClipboardContent, OutPasteWarning);
 		}
 	}
 }
diff --git a/Engine/Plugins/FX/Niagara/Source/NiagaraEditor/Private/ViewModels/Stack/NiagaraStackFunctionInput.cpp b/Engine/Plugins/FX/Niagara/Source/NiagaraEditor/Private/ViewModels/Stack/NiagaraStackFunctionInput.cpp
--- a/Engine/Plugins/FX/Niagara/Source/NiagaraEditor/Private/ViewModels/Stack/NiagaraStackFunctionInput.cpp
+++ b/Engine/Plugins/FX/Niagara/Source/NiagaraEditor/Private/ViewModels/Stack/NiagaraStackFunctionInput.cpp
@@ -353,7 +353,7 @@ FText UNiagaraStackFunctionInput::GetPasteTransactionText(const UNiagaraClipboar
 	return LOCTEXT("PasteInputTransactionText", "Paste niagara inputs");
 }
 
-void UNiagaraStackFunctionInput::Paste(const UNiagaraClipboardContent* ClipboardContent)
+void UNiagaraStackFunctionInput::Paste(const UNiagaraClipboardContent* ClipboardContent, FText& OutPasteWarning)
 {
 	checkf(ClipboardContent != nullptr && ClipboardContent->FunctionInputs.Num() == 1, TEXT("Clipboard must not be null, and must contain a single input.  Call TestCanPasteWithMessage to validate"));
 
diff --git a/Engine/Plugins/FX/Niagara/Source/NiagaraEditor/Private/ViewModels/Stack/NiagaraStackGraphUtilities.cpp b/Engine/Plugins/FX/Niagara/Source/NiagaraEditor/Private/ViewModels/Stack/NiagaraStackGraphUtilities.cpp
--- a/Engine/Plugins/FX/Niagara/Source/NiagaraEditor/Private/ViewModels/Stack/NiagaraStackGraphUtilities.cpp
+++ b/Engine/Plugins/FX/Niagara/Source/NiagaraEditor/Private/ViewModels/Stack/NiagaraStackGraphUtilities.cpp
@@ -1753,9 +1753,20 @@ bool FNiagaraStackGraphUtilities::IsValidDefaultDynamicInput(UNiagaraScript& Own
 	return TryGetStackFunctionInputValue(OwningScript, nullptr, DefaultPin, NAME_None, FRapidIterationParameterContext(), InputValue) && InputValue.DynamicValue.IsSet();
 }
 
-bool FNiagaraStackGraphUtilities::ParameterIsCompatibleWithScriptUsage(FNiagaraVariable Parameter, ENiagaraScriptUsage Usage)
+bool FNiagaraStackGraphUtilities::CanWriteParameterFromUsage(FNiagaraVariable Parameter, ENiagaraScriptUsage Usage)
 {
 	const FNiagaraParameterHandle ParameterHandle(Parameter.GetName());
+
+	if (ParameterHandle.IsReadOnlyHandle())
+	{
+		return false;
+	}
+
+	if (ParameterHandle.IsTransientHandle())
+	{
+		return true;
+	}
+
 	switch (Usage)
 	{
 	case ENiagaraScriptUsage::SystemSpawnScript:
@@ -1768,6 +1779,7 @@ bool FNiagaraStackGraphUtilities::ParameterIsCompatibleWithScriptUsage(FNiagaraV
 	case ENiagaraScriptUsage::ParticleSpawnScriptInterpolated:
 	case ENiagaraScriptUsage::ParticleUpdateScript:
 	case ENiagaraScriptUsage::ParticleEventScript:
+	case ENiagaraScriptUsage::ParticleShaderStageScript:
 		return ParameterHandle.IsParticleAttributeHandle();
 	default:
 		return false;
diff --git a/Engine/Plugins/FX/Niagara/Source/NiagaraEditor/Private/ViewModels/Stack/NiagaraStackItem.cpp b/Engine/Plugins/FX/Niagara/Source/NiagaraEditor/Private/ViewModels/Stack/NiagaraStackItem.cpp
--- a/Engine/Plugins/FX/Niagara/Source/NiagaraEditor/Private/ViewModels/Stack/NiagaraStackItem.cpp
+++ b/Engine/Plugins/FX/Niagara/Source/NiagaraEditor/Private/ViewModels/Stack/NiagaraStackItem.cpp
@@ -27,9 +27,14 @@ UNiagaraStackItem::FOnModifiedGroupItems& UNiagaraStackItem::OnModifiedGroupItem
 	return ModifiedGroupItemsDelegate;
 }
 
-UNiagaraStackItem::FOnRequestPaste& UNiagaraStackItem::OnRequestPaste()
+void UNiagaraStackItem::SetOnRequestCanPaste(FOnRequestCanPaste InOnRequestCanPaste)
 {
-	return RequestPasteDelegate;
+	RequestCanPasteDelegete = InOnRequestCanPaste;
+}
+
+void UNiagaraStackItem::SetOnRequestPaste(FOnRequestPaste InOnRequestCanPaste)
+{
+	RequestPasteDelegate = InOnRequestCanPaste;
 }
 
 void UNiagaraStackItem::SetIsEnabled(bool bInIsEnabled)
diff --git a/Engine/Plugins/FX/Niagara/Source/NiagaraEditor/Private/ViewModels/Stack/NiagaraStackModuleItem.cpp b/Engine/Plugins/FX/Niagara/Source/NiagaraEditor/Private/ViewModels/Stack/NiagaraStackModuleItem.cpp
--- a/Engine/Plugins/FX/Niagara/Source/NiagaraEditor/Private/ViewModels/Stack/NiagaraStackModuleItem.cpp
+++ b/Engine/Plugins/FX/Niagara/Source/NiagaraEditor/Private/ViewModels/Stack/NiagaraStackModuleItem.cpp
@@ -260,13 +260,13 @@ TOptional<UNiagaraStackEntry::FDropRequestResponse> UNiagaraStackModuleItem::Can
 		TSharedPtr<const FNiagaraParameterAction> ParameterAction = StaticCastSharedPtr<const FNiagaraParameterAction>(ParameterDragDropOp->GetSourceAction());
 		if (ParameterAction.IsValid())
 		{
-			if (FNiagaraStackGraphUtilities::ParameterIsCompatibleWithScriptUsage(ParameterAction->GetParameter(), OutputNode->GetUsage()))
+			if (FNiagaraStackGraphUtilities::CanWriteParameterFromUsage(ParameterAction->GetParameter(), OutputNode->GetUsage()))
 			{
 				return FDropRequestResponse(EItemDropZone::OntoItem, LOCTEXT("DropParameterToAdd", "Add this parameter to this 'Set Variables' node."));
 			}
 			else
 			{
-				return FDropRequestResponse(TOptional<EItemDropZone>(), LOCTEXT("CantDropParameterByUsage", "Can not drop this parameter here because\nit's not valid for this usage context."));
+				return FDropRequestResponse(TOptional<EItemDropZone>(), LOCTEXT("CantDropParameterByUsage", "Can not drop this parameter here because\nit can't be written in this usage context."));
 			}
 		}
 	}
@@ -884,7 +884,7 @@ bool UNiagaraStackModuleItem::CanAddInput(FNiagaraVariable InputParameter) const
 	UNiagaraNodeAssignment* AssignmentModule = Cast<UNiagaraNodeAssignment>(FunctionCallNode);
 	return AssignmentModule != nullptr &&
 		AssignmentModule->GetAssignmentTargets().Contains(InputParameter) == false &&
-		FNiagaraStackGraphUtilities::ParameterIsCompatibleWithScriptUsage(InputParameter, OutputNode->GetUsage());
+		FNiagaraStackGraphUtilities::CanWriteParameterFromUsage(InputParameter, OutputNode->GetUsage());
 }
 
 void UNiagaraStackModuleItem::AddInput(FNiagaraVariable InputParameter)
@@ -1000,26 +1000,10 @@ bool UNiagaraStackModuleItem::TestCanPasteWithMessage(const UNiagaraClipboardCon
 		OutMessage = LOCTEXT("PasteInputs", "Paste inputs from the clipboard which match inputs on this module by name and type.");
 		return true;
 	}
-	if (ClipboardContent->Functions.Num() > 0)
+
+	if (RequestCanPasteDelegete.IsBound())
 	{
-		bool bValidUsage = false;
-		for (const UNiagaraClipboardFunction* Function : ClipboardContent->Functions)
-		{
-			if (Function != nullptr && Function->Script != nullptr && Function->Script->GetSupportedUsageContexts().Contains(OutputNode->GetUsage()))
-			{
-				bValidUsage = true;
-			}
-		}
-		if (bValidUsage)
-		{
-			OutMessage = LOCTEXT("PasteModules", "Paste modules from the clipboard which have a valid usage.");
-			return true;
-		}
-		else
-		{
-			OutMessage = LOCTEXT("CantPasteModulesForUsage", "Can't paste the copied modules because they don't support the correct usage.");
-			return false;
-		}
+		return RequestCanPasteDelegete.Execute(ClipboardContent, OutMessage);
 	}
 
 	OutMessage = FText();
@@ -1038,18 +1022,17 @@ FText UNiagaraStackModuleItem::GetPasteTransactionText(const UNiagaraClipboardCo
 	}
 }
 
-void UNiagaraStackModuleItem::Paste(const UNiagaraClipboardContent* ClipboardContent)
+void UNiagaraStackModuleItem::Paste(const UNiagaraClipboardContent* ClipboardContent, FText& OutPasteWarning)
 {
 	if (ClipboardContent->FunctionInputs.Num() > 0)
 	{
 		SetInputValuesFromClipboardFunctionInputs(ClipboardContent->FunctionInputs);
 	}
-
-	if (ClipboardContent->Functions.Num() > 0)
+	else if (RequestCanPasteDelegete.IsBound())
 	{
 		// Pasted modules should go after this module, so add 1 to the index.
 		int32 PasteIndex = GetModuleIndex() + 1;
-		RequestPasteDelegate.Broadcast(ClipboardContent, PasteIndex);
+		RequestPasteDelegate.Execute(ClipboardContent, PasteIndex, OutPasteWarning);
 	}
 }
 
diff --git a/Engine/Plugins/FX/Niagara/Source/NiagaraEditor/Private/ViewModels/Stack/NiagaraStackRenderItemGroup.cpp b/Engine/Plugins/FX/Niagara/Source/NiagaraEditor/Private/ViewModels/Stack/NiagaraStackRenderItemGroup.cpp
--- a/Engine/Plugins/FX/Niagara/Source/NiagaraEditor/Private/ViewModels/Stack/NiagaraStackRenderItemGroup.cpp
+++ b/Engine/Plugins/FX/Niagara/Source/NiagaraEditor/Private/ViewModels/Stack/NiagaraStackRenderItemGroup.cpp
@@ -136,7 +136,7 @@ bool UNiagaraStackRenderItemGroup::TestCanPasteWithMessage(const UNiagaraClipboa
 		OutMessage = LOCTEXT("PasteRenderers", "Paste renderers from the clipboard.");
 		return true;
 	}
-	OutMessage = LOCTEXT("NoRenderers", "No renderers on the clipboard.");
+	OutMessage = FText();
 	return false;
 }
 
@@ -145,7 +145,7 @@ FText UNiagaraStackRenderItemGroup::GetPasteTransactionText(const UNiagaraClipbo
 	return LOCTEXT("PasteRenderersTransactionText", "Paste renderers");
 }
 
-void UNiagaraStackRenderItemGroup::Paste(const UNiagaraClipboardContent* ClipboardContent)
+void UNiagaraStackRenderItemGroup::Paste(const UNiagaraClipboardContent* ClipboardContent, FText& OutPasteWarning)
 {
 	if (EmitterWeak.IsValid())
 	{
@@ -172,7 +172,8 @@ void UNiagaraStackRenderItemGroup::RefreshChildrenInternal(const TArray<UNiagara
 		{
 			RendererItem = NewObject<UNiagaraStackRendererItem>(this);
 			RendererItem->Initialize(CreateDefaultChildRequiredData(), RendererProperties);
-			RendererItem->OnRequestPaste().AddUObject(this, &UNiagaraStackRenderItemGroup::ChildRequestPaste);
+			RendererItem->SetOnRequestCanPaste(UNiagaraStackRendererItem::FOnRequestCanPaste::CreateUObject(this, &UNiagaraStackRenderItemGroup::ChildRequestCanPaste));
+			RendererItem->SetOnRequestPaste(UNiagaraStackRendererItem::FOnRequestPaste::CreateUObject(this, &UNiagaraStackRenderItemGroup::ChildRequestPaste));
 		}
 
 		NewChildren.Add(RendererItem);
@@ -193,9 +194,14 @@ void UNiagaraStackRenderItemGroup::EmitterRenderersChanged()
 	}
 }
 
-void UNiagaraStackRenderItemGroup::ChildRequestPaste(const UNiagaraClipboardContent* ClipboardContent, int32 PasteIndex)
+bool UNiagaraStackRenderItemGroup::ChildRequestCanPaste(const UNiagaraClipboardContent* ClipboardContent, FText& OutCanPasteMessage)
 {
-	Paste(ClipboardContent);
+	return TestCanPasteWithMessage(ClipboardContent, OutCanPasteMessage);
+}
+
+void UNiagaraStackRenderItemGroup::ChildRequestPaste(const UNiagaraClipboardContent* ClipboardContent, int32 PasteIndex, FText& OutPasteWarning)
+{
+	Paste(ClipboardContent, OutPasteWarning);
 }
 
 void UNiagaraStackRenderItemGroup::FinalizeInternal()
diff --git a/Engine/Plugins/FX/Niagara/Source/NiagaraEditor/Private/ViewModels/Stack/NiagaraStackRendererItem.cpp b/Engine/Plugins/FX/Niagara/Source/NiagaraEditor/Private/ViewModels/Stack/NiagaraStackRendererItem.cpp
--- a/Engine/Plugins/FX/Niagara/Source/NiagaraEditor/Private/ViewModels/Stack/NiagaraStackRendererItem.cpp
+++ b/Engine/Plugins/FX/Niagara/Source/NiagaraEditor/Private/ViewModels/Stack/NiagaraStackRendererItem.cpp
@@ -195,12 +195,11 @@ void UNiagaraStackRendererItem::Copy(UNiagaraClipboardContent* ClipboardContent)
 
 bool UNiagaraStackRendererItem::TestCanPasteWithMessage(const UNiagaraClipboardContent* ClipboardContent, FText& OutMessage) const
 {
-	if (ClipboardContent->Renderers.Num() > 0)
+	if (RequestCanPasteDelegete.IsBound())
 	{
-		OutMessage = LOCTEXT("PasteRenderers", "Paste renderers from the clipboard.");
-		return true;
+		return RequestCanPasteDelegete.Execute(ClipboardContent, OutMessage);
 	}
-	OutMessage = LOCTEXT("NoRenderers", "No renderers on the clipboard.");
+	OutMessage = FText();
 	return false;
 }
 
@@ -209,9 +208,9 @@ FText UNiagaraStackRendererItem::GetPasteTransactionText(const UNiagaraClipboard
 	return LOCTEXT("PasteRenderersTransactionText", "Paste renderers");
 }
 
-void UNiagaraStackRendererItem::Paste(const UNiagaraClipboardContent* ClipboardContent)
+void UNiagaraStackRendererItem::Paste(const UNiagaraClipboardContent* ClipboardContent, FText& OutPasteWarning)
 {
-	OnRequestPaste().Broadcast(ClipboardContent, INDEX_NONE);
+	RequestPasteDelegate.ExecuteIfBound(ClipboardContent, INDEX_NONE, OutPasteWarning);
 }
 
 bool UNiagaraStackRendererItem::TestCanDeleteWithMessage(FText& OutCanDeleteMessage) const
diff --git a/Engine/Plugins/FX/Niagara/Source/NiagaraEditor/Private/ViewModels/Stack/NiagaraStackScriptItemGroup.cpp b/Engine/Plugins/FX/Niagara/Source/NiagaraEditor/Private/ViewModels/Stack/NiagaraStackScriptItemGroup.cpp
--- a/Engine/Plugins/FX/Niagara/Source/NiagaraEditor/Private/ViewModels/Stack/NiagaraStackScriptItemGroup.cpp
+++ b/Engine/Plugins/FX/Niagara/Source/NiagaraEditor/Private/ViewModels/Stack/NiagaraStackScriptItemGroup.cpp
@@ -333,11 +333,26 @@ bool UNiagaraStackScriptItemGroup::TestCanPasteWithMessage(const UNiagaraClipboa
 	if (ClipboardContent->Functions.Num() > 0)
 	{
 		bool bValidUsage = false;
+		UNiagaraNodeOutput* OutputNode = GetScriptOutputNode();
 		for (const UNiagaraClipboardFunction* Function : ClipboardContent->Functions)
 		{
-			if (Function != nullptr && Function->Script != nullptr && Function->Script->GetSupportedUsageContexts().Contains(ScriptUsage))
+			if (Function->ScriptMode == ENiagaraClipboardFunctionScriptMode::ScriptAsset)
 			{
-				bValidUsage = true;
+				if (Function != nullptr && Function->Script != nullptr && Function->Script->GetSupportedUsageContexts().Contains(OutputNode->GetUsage()))
+				{
+					bValidUsage = true;
+				}
+			}
+			else if (Function->ScriptMode == ENiagaraClipboardFunctionScriptMode::Assignment)
+			{
+				for (const FNiagaraVariable& AssignmentTarget : Function->AssignmentTargets)
+				{
+					if (FNiagaraStackGraphUtilities::CanWriteParameterFromUsage(AssignmentTarget, OutputNode->GetUsage()))
+					{
+						bValidUsage = true;
+						continue;
+					}
+				}
 			}
 		}
 		if (bValidUsage)
@@ -360,14 +375,14 @@ FText UNiagaraStackScriptItemGroup::GetPasteTransactionText(const UNiagaraClipbo
 	return LOCTEXT("PasteModuleTransactionText", "Paste niagara modules");
 }
 
-void UNiagaraStackScriptItemGroup::Paste(const UNiagaraClipboardContent* ClipboardContent)
+void UNiagaraStackScriptItemGroup::Paste(const UNiagaraClipboardContent* ClipboardContent, FText& OutPasteWarning)
 {
 	if (ClipboardContent->Functions.Num() > 0)
 	{
 		TArray<UNiagaraStackModuleItem*> ModuleItems;
 		GetUnfilteredChildrenOfType(ModuleItems);
 		int32 PasteIndex = ModuleItems.Num() > 0 ? ModuleItems.Last()->GetModuleIndex() + 1 : 0;
-		PasteModules(ClipboardContent, PasteIndex);
+		PasteModules(ClipboardContent, PasteIndex, OutPasteWarning);
 	}
 }
 
@@ -398,7 +413,8 @@ void UNiagaraStackScriptItemGroup::RefreshChildrenInternal(const TArray<UNiagara
 				ModuleItem = NewObject<UNiagaraStackModuleItem>(this);
 				ModuleItem->Initialize(CreateDefaultChildRequiredData(), GetAddUtilities(), *ModuleNode);
 				ModuleItem->OnModifiedGroupItems().AddUObject(this, &UNiagaraStackScriptItemGroup::ChildModifiedGroupItems);
-				ModuleItem->OnRequestPaste().AddUObject(this, &UNiagaraStackScriptItemGroup::ChildRequestPaste);
+				ModuleItem->SetOnRequestCanPaste(UNiagaraStackModuleItem::FOnRequestCanPaste::CreateUObject(this, &UNiagaraStackScriptItemGroup::ChildRequestCanPaste));
+				ModuleItem->SetOnRequestPaste(UNiagaraStackModuleItem::FOnRequestPaste::CreateUObject(this, &UNiagaraStackScriptItemGroup::ChildRequestPaste));
 			}
 
 			NewChildren.Add(ModuleItem);
@@ -800,9 +816,9 @@ TOptional<UNiagaraStackEntry::FDropRequestResponse> UNiagaraStackScriptItemGroup
 	}
 	if (ParameterAction.IsValid())
 	{
-		if (FNiagaraStackGraphUtilities::ParameterIsCompatibleWithScriptUsage(ParameterAction->GetParameter(), ScriptUsage) == false)
+		if (FNiagaraStackGraphUtilities::CanWriteParameterFromUsage(ParameterAction->GetParameter(), ScriptUsage) == false)
 		{
-			return FDropRequestResponse(TOptional<EItemDropZone>(), LOCTEXT("CantDropParameterByUsage", "Can not drop this parameter here because\nit's not valid for this usage context."));
+			return FDropRequestResponse(TOptional<EItemDropZone>(), LOCTEXT("CantDropParameterByUsage", "Can not drop this parameter here because\nit can't be written in this usage context."));
 		}
 
 		{
@@ -941,9 +957,14 @@ void UNiagaraStackScriptItemGroup::ChildModifiedGroupItems()
 	RefreshChildren();
 }
 
-void UNiagaraStackScriptItemGroup::ChildRequestPaste(const UNiagaraClipboardContent* ClipboardContent, int32 PasteIndex)
+bool UNiagaraStackScriptItemGroup::ChildRequestCanPaste(const UNiagaraClipboardContent* ClipboardContent, FText& OutCanPasteMessage)
+{
+	return TestCanPasteWithMessage(ClipboardContent, OutCanPasteMessage);
+}
+
+void UNiagaraStackScriptItemGroup::ChildRequestPaste(const UNiagaraClipboardContent* ClipboardContent, int32 PasteIndex, FText& OutPasteWarning)
 {
-	PasteModules(ClipboardContent, PasteIndex);
+	PasteModules(ClipboardContent, PasteIndex, OutPasteWarning);
 }
 
 void UNiagaraStackScriptItemGroup::OnScriptGraphChanged(const struct FEdGraphEditAction& InAction)
@@ -1020,13 +1041,15 @@ void RenameInputsFromClipboard(TMap<FName, FName> OldModuleOutputNameToNewModule
 	}
 }
 
-void UNiagaraStackScriptItemGroup::PasteModules(const UNiagaraClipboardContent* ClipboardContent, int32 PasteIndex)
+void UNiagaraStackScriptItemGroup::PasteModules(const UNiagaraClipboardContent* ClipboardContent, int32 PasteIndex, FText& OutPasteWarning)
 {
 	UNiagaraNodeOutput* OutputNode = GetScriptOutputNode();
 	
 	// Add the new function call nodes from the clipboard functions and collected the pairs.
 	TArray<TPair<const UNiagaraClipboardFunction*, UNiagaraNodeFunctionCall*>> ClipboardFunctionAndNodeFunctionPairs;
 	int32 CurrentPasteIndex = PasteIndex;
+	int32 ModulesSkippedForIncorrectUsage = 0;
+	int32 AssignmentTargetsSkippedForIncorrectUsage = 0;
 	for (const UNiagaraClipboardFunction* ClipboardFunction : ClipboardContent->Functions)
 	{
 		if (ClipboardFunction != nullptr)
@@ -1036,7 +1059,24 @@ void UNiagaraStackScriptItemGroup::PasteModules(const UNiagaraClipboardContent*
 			{
 			case ENiagaraClipboardFunctionScriptMode::Assignment:
 			{
-				NewFunctionCallNode = FNiagaraStackGraphUtilities::AddParameterModuleToStack(ClipboardFunction->AssignmentTargets, *OutputNode, CurrentPasteIndex, ClipboardFunction->AssignmentDefaults);
+				TArray<FNiagaraVariable> CompatibleAssignmentTargets;
+				TArray<FString> CompatibleAssignmentDefaults;
+				for (int32 AssignmentTargetIndex = 0; AssignmentTargetIndex < ClipboardFunction->AssignmentTargets.Num(); ++AssignmentTargetIndex)
+				{
+					if (FNiagaraStackGraphUtilities::CanWriteParameterFromUsage(ClipboardFunction->AssignmentTargets[AssignmentTargetIndex], OutputNode->GetUsage()))
+					{
+						CompatibleAssignmentTargets.Add(ClipboardFunction->AssignmentTargets[AssignmentTargetIndex]);
+						CompatibleAssignmentDefaults.Add(ClipboardFunction->AssignmentDefaults[AssignmentTargetIndex]);
+					}
+					else
+					{
+						AssignmentTargetsSkippedForIncorrectUsage++;
+					}
+				}
+				if (CompatibleAssignmentTargets.Num() > 0)
+				{
+					NewFunctionCallNode = FNiagaraStackGraphUtilities::AddParameterModuleToStack(CompatibleAssignmentTargets, *OutputNode, CurrentPasteIndex, CompatibleAssignmentDefaults);
+				}
 				break;
 			}
 			case ENiagaraClipboardFunctionScriptMode::ScriptAsset:
@@ -1045,15 +1085,22 @@ void UNiagaraStackScriptItemGroup::PasteModules(const UNiagaraClipboardContent*
 				{
 					NewFunctionCallNode = FNiagaraStackGraphUtilities::AddScriptModuleToStack(ClipboardFunction->Script, *OutputNode, CurrentPasteIndex);
 				}
+				else
+				{
+					ModulesSkippedForIncorrectUsage++;
+				}
 				break;
 			}
 			default:
 				checkf(false, TEXT("Unsupported clipboard function mode."));
 			}
 
-			NewFunctionCallNode->SuggestName(ClipboardFunction->FunctionName);
-			ClipboardFunctionAndNodeFunctionPairs.Add(TPair<const UNiagaraClipboardFunction*, UNiagaraNodeFunctionCall*>(ClipboardFunction, NewFunctionCallNode));
-			CurrentPasteIndex++;
+			if (NewFunctionCallNode != nullptr)
+			{
+				NewFunctionCallNode->SuggestName(ClipboardFunction->FunctionName);
+				ClipboardFunctionAndNodeFunctionPairs.Add(TPair<const UNiagaraClipboardFunction*, UNiagaraNodeFunctionCall*>(ClipboardFunction, NewFunctionCallNode));
+				CurrentPasteIndex++;
+			}
 		}
 	}
 	RefreshChildren();
@@ -1086,6 +1133,12 @@ void UNiagaraStackScriptItemGroup::PasteModules(const UNiagaraClipboardContent*
 			(*NewModuleItemPtr)->SetInputValuesFromClipboardFunctionInputs(FunctionInputs);
 		}
 	}
+
+	if (ModulesSkippedForIncorrectUsage > 0 || AssignmentTargetsSkippedForIncorrectUsage > 0)
+	{
+		OutPasteWarning = FText::Format(LOCTEXT("NonPastedModulesWarningFormat", "{0} modules not pasted due to incorrect usage.\n{1} inputs not pasted in Set Variables nodes due to incorrect usage."),
+			FText::AsNumber(ModulesSkippedForIncorrectUsage), FText::AsNumber(AssignmentTargetsSkippedForIncorrectUsage));
+	}
 }
 
 #undef LOCTEXT_NAMESPACE
diff --git a/Engine/Plugins/FX/Niagara/Source/NiagaraEditor/Public/ViewModels/Stack/NiagaraParameterHandle.h b/Engine/Plugins/FX/Niagara/Source/NiagaraEditor/Public/ViewModels/Stack/NiagaraParameterHandle.h
--- a/Engine/Plugins/FX/Niagara/Source/NiagaraEditor/Public/ViewModels/Stack/NiagaraParameterHandle.h
+++ b/Engine/Plugins/FX/Niagara/Source/NiagaraEditor/Public/ViewModels/Stack/NiagaraParameterHandle.h
@@ -53,6 +53,10 @@ public:
 
 	bool IsParameterCollectionHandle() const;
 
+	bool IsReadOnlyHandle() const;
+
+	bool IsTransientHandle() const;
+
 public:
 	static const FName UserNamespace;
 	static const FName EngineNamespace;
diff --git a/Engine/Plugins/FX/Niagara/Source/NiagaraEditor/Public/ViewModels/Stack/NiagaraStackClipboardUtilities.h b/Engine/Plugins/FX/Niagara/Source/NiagaraEditor/Public/ViewModels/Stack/NiagaraStackClipboardUtilities.h
--- a/Engine/Plugins/FX/Niagara/Source/NiagaraEditor/Public/ViewModels/Stack/NiagaraStackClipboardUtilities.h
+++ b/Engine/Plugins/FX/Niagara/Source/NiagaraEditor/Public/ViewModels/Stack/NiagaraStackClipboardUtilities.h
@@ -18,7 +18,7 @@ namespace FNiagaraStackClipboardUtilities
 
 	NIAGARAEDITOR_API bool TestCanPasteSelectionWithMessage(const TArray<UNiagaraStackEntry*>& SelectedEntries, FText& OutCanPasteMessage);
 
-	NIAGARAEDITOR_API void PasteSelection(const TArray<UNiagaraStackEntry*>& SelectedEntries);
+	NIAGARAEDITOR_API void PasteSelection(const TArray<UNiagaraStackEntry*>& SelectedEntries, FText& OutPasteWarning);
 
 	NIAGARAEDITOR_API bool TestCanDeleteSelectionWithMessage(const TArray<UNiagaraStackEntry*>& SelectedEntries, FText& OutMessage);
 
diff --git a/Engine/Plugins/FX/Niagara/Source/NiagaraEditor/Public/ViewModels/Stack/NiagaraStackEntry.h b/Engine/Plugins/FX/Niagara/Source/NiagaraEditor/Public/ViewModels/Stack/NiagaraStackEntry.h
--- a/Engine/Plugins/FX/Niagara/Source/NiagaraEditor/Public/ViewModels/Stack/NiagaraStackEntry.h
+++ b/Engine/Plugins/FX/Niagara/Source/NiagaraEditor/Public/ViewModels/Stack/NiagaraStackEntry.h
@@ -367,7 +367,7 @@ public:
 
 	virtual FText GetPasteTransactionText(const UNiagaraClipboardContent* ClipboardContent) const { return FText(); }
 
-	virtual void Paste(const UNiagaraClipboardContent* ClipboardContent) { }
+	virtual void Paste(const UNiagaraClipboardContent* ClipboardContent, FText& OutPasteWarning) { }
 
 	virtual bool SupportsDelete() const { return false; }
 
diff --git a/Engine/Plugins/FX/Niagara/Source/NiagaraEditor/Public/ViewModels/Stack/NiagaraStackFunctionInput.h b/Engine/Plugins/FX/Niagara/Source/NiagaraEditor/Public/ViewModels/Stack/NiagaraStackFunctionInput.h
--- a/Engine/Plugins/FX/Niagara/Source/NiagaraEditor/Public/ViewModels/Stack/NiagaraStackFunctionInput.h
+++ b/Engine/Plugins/FX/Niagara/Source/NiagaraEditor/Public/ViewModels/Stack/NiagaraStackFunctionInput.h
@@ -103,7 +103,7 @@ public:
 	virtual bool SupportsPaste() const override { return true; }
 	virtual bool TestCanPasteWithMessage(const UNiagaraClipboardContent* ClipboardContent, FText& OutMessage) const override;
 	virtual FText GetPasteTransactionText(const UNiagaraClipboardContent* ClipboardContent) const override;
-	virtual void Paste(const UNiagaraClipboardContent* ClipboardContent) override;
+	virtual void Paste(const UNiagaraClipboardContent* ClipboardContent, FText& OutPasteWarning) override;
 
 	FText GetTooltipText(EValueMode InValueMode) const;
 
diff --git a/Engine/Plugins/FX/Niagara/Source/NiagaraEditor/Public/ViewModels/Stack/NiagaraStackGraphUtilities.h b/Engine/Plugins/FX/Niagara/Source/NiagaraEditor/Public/ViewModels/Stack/NiagaraStackGraphUtilities.h
--- a/Engine/Plugins/FX/Niagara/Source/NiagaraEditor/Public/ViewModels/Stack/NiagaraStackGraphUtilities.h
+++ b/Engine/Plugins/FX/Niagara/Source/NiagaraEditor/Public/ViewModels/Stack/NiagaraStackGraphUtilities.h
@@ -149,7 +149,7 @@ namespace FNiagaraStackGraphUtilities
 
 	bool IsValidDefaultDynamicInput(UNiagaraScript& OwningScript, UEdGraphPin& DefaultPin);
 
-	bool ParameterIsCompatibleWithScriptUsage(FNiagaraVariable Parameter, ENiagaraScriptUsage Usage);
+	bool CanWriteParameterFromUsage(FNiagaraVariable Parameter, ENiagaraScriptUsage Usage);
 
 	bool DoesDynamicInputMatchDefault(
 		FString EmitterUniqueName,
diff --git a/Engine/Plugins/FX/Niagara/Source/NiagaraEditor/Public/ViewModels/Stack/NiagaraStackItem.h b/Engine/Plugins/FX/Niagara/Source/NiagaraEditor/Public/ViewModels/Stack/NiagaraStackItem.h
--- a/Engine/Plugins/FX/Niagara/Source/NiagaraEditor/Public/ViewModels/Stack/NiagaraStackItem.h
+++ b/Engine/Plugins/FX/Niagara/Source/NiagaraEditor/Public/ViewModels/Stack/NiagaraStackItem.h
@@ -18,7 +18,8 @@ class NIAGARAEDITOR_API UNiagaraStackItem : public UNiagaraStackEntry
 
 public:
 	DECLARE_MULTICAST_DELEGATE(FOnModifiedGroupItems);
-	DECLARE_MULTICAST_DELEGATE_TwoParams(FOnRequestPaste, const UNiagaraClipboardContent* /* ClipboardContent */, int32 /* PasteIndex */);
+	DECLARE_DELEGATE_RetVal_TwoParams(bool, FOnRequestCanPaste, const UNiagaraClipboardContent* /* ClipboardContent */, FText& /* OutCanPasteMessage */);
+	DECLARE_DELEGATE_ThreeParams(FOnRequestPaste, const UNiagaraClipboardContent* /* ClipboardContent */, int32 /* PasteIndex */, FText& /* OutPasteWarning */);
 
 public:
 	void Initialize(FRequiredEntryData InRequiredEntryData, FString InStackEditorDataKey);
@@ -26,7 +27,9 @@ public:
 	virtual EStackRowStyle GetStackRowStyle() const override;
 
 	FOnModifiedGroupItems& OnModifiedGroupItems();
-	FOnRequestPaste& OnRequestPaste();
+
+	void SetOnRequestCanPaste(FOnRequestCanPaste InOnRequestCanPaste);
+	void SetOnRequestPaste(FOnRequestPaste InOnRequestCanPaste);
 
 	virtual bool SupportsChangeEnabled() const { return false; }
 	void SetIsEnabled(bool bInIsEnabled);
@@ -57,6 +60,7 @@ private:
 
 protected:
 	FOnModifiedGroupItems ModifiedGroupItemsDelegate;
+	FOnRequestCanPaste RequestCanPasteDelegete;
 	FOnRequestPaste RequestPasteDelegate;
 
 private:
diff --git a/Engine/Plugins/FX/Niagara/Source/NiagaraEditor/Public/ViewModels/Stack/NiagaraStackModuleItem.h b/Engine/Plugins/FX/Niagara/Source/NiagaraEditor/Public/ViewModels/Stack/NiagaraStackModuleItem.h
--- a/Engine/Plugins/FX/Niagara/Source/NiagaraEditor/Public/ViewModels/Stack/NiagaraStackModuleItem.h
+++ b/Engine/Plugins/FX/Niagara/Source/NiagaraEditor/Public/ViewModels/Stack/NiagaraStackModuleItem.h
@@ -80,7 +80,7 @@ public:
 	virtual bool SupportsPaste() const override { return true; }
 	virtual bool TestCanPasteWithMessage(const UNiagaraClipboardContent* ClipboardContent, FText& OutMessage) const override;
 	virtual FText GetPasteTransactionText(const UNiagaraClipboardContent* ClipboardContent) const override;
-	virtual void Paste(const UNiagaraClipboardContent* ClipboardContent) override;
+	virtual void Paste(const UNiagaraClipboardContent* ClipboardContent, FText& OutPasteWarning) override;
 
 	virtual bool SupportsDelete() const override { return true; }
 	virtual bool TestCanDeleteWithMessage(FText& OutCanDeleteMessage) const override;
diff --git a/Engine/Plugins/FX/Niagara/Source/NiagaraEditor/Public/ViewModels/Stack/NiagaraStackRenderItemGroup.h b/Engine/Plugins/FX/Niagara/Source/NiagaraEditor/Public/ViewModels/Stack/NiagaraStackRenderItemGroup.h
--- a/Engine/Plugins/FX/Niagara/Source/NiagaraEditor/Public/ViewModels/Stack/NiagaraStackRenderItemGroup.h
+++ b/Engine/Plugins/FX/Niagara/Source/NiagaraEditor/Public/ViewModels/Stack/NiagaraStackRenderItemGroup.h
@@ -20,7 +20,7 @@ public:
 	virtual bool SupportsPaste() const override { return true; }
 	virtual bool TestCanPasteWithMessage(const UNiagaraClipboardContent* ClipboardContent, FText& OutMessage) const override;
 	virtual FText GetPasteTransactionText(const UNiagaraClipboardContent* ClipboardContent) const override;
-	virtual void Paste(const UNiagaraClipboardContent* ClipboardContent) override;
+	virtual void Paste(const UNiagaraClipboardContent* ClipboardContent, FText& OutPasteWarning) override;
 
 protected:
 	virtual void RefreshChildrenInternal(const TArray<UNiagaraStackEntry*>& CurrentChildren, TArray<UNiagaraStackEntry*>& NewChildren, TArray<FStackIssue>& NewIssues) override;
@@ -28,7 +28,8 @@ protected:
 private:
 	void EmitterRenderersChanged();
 
-	void ChildRequestPaste(const UNiagaraClipboardContent* ClipboardContent, int32 PasteIndex);
+	bool ChildRequestCanPaste(const UNiagaraClipboardContent* ClipboardContent,FText& OutCanPasteMessage);
+	void ChildRequestPaste(const UNiagaraClipboardContent* ClipboardContent, int32 PasteIndex, FText& OutPasteWarning);
 
 private:
 	TSharedPtr<INiagaraStackItemGroupAddUtilities> AddUtilities;
diff --git a/Engine/Plugins/FX/Niagara/Source/NiagaraEditor/Public/ViewModels/Stack/NiagaraStackRendererItem.h b/Engine/Plugins/FX/Niagara/Source/NiagaraEditor/Public/ViewModels/Stack/NiagaraStackRendererItem.h
--- a/Engine/Plugins/FX/Niagara/Source/NiagaraEditor/Public/ViewModels/Stack/NiagaraStackRendererItem.h
+++ b/Engine/Plugins/FX/Niagara/Source/NiagaraEditor/Public/ViewModels/Stack/NiagaraStackRendererItem.h
@@ -37,7 +37,7 @@ public:
 	virtual bool SupportsPaste() const { return true; }
 	virtual bool TestCanPasteWithMessage(const UNiagaraClipboardContent* ClipboardContent, FText& OutMessage) const override;
 	virtual FText GetPasteTransactionText(const UNiagaraClipboardContent* ClipboardContent) const override;
-	virtual void Paste(const UNiagaraClipboardContent* ClipboardContent) override;
+	virtual void Paste(const UNiagaraClipboardContent* ClipboardContent, FText& OutPasteWarning) override;
 
 	virtual bool SupportsDelete() const override { return true; }
 	virtual bool TestCanDeleteWithMessage(FText& OutCanDeleteMessage) const;
diff --git a/Engine/Plugins/FX/Niagara/Source/NiagaraEditor/Public/ViewModels/Stack/NiagaraStackScriptItemGroup.h b/Engine/Plugins/FX/Niagara/Source/NiagaraEditor/Public/ViewModels/Stack/NiagaraStackScriptItemGroup.h
--- a/Engine/Plugins/FX/Niagara/Source/NiagaraEditor/Public/ViewModels/Stack/NiagaraStackScriptItemGroup.h
+++ b/Engine/Plugins/FX/Niagara/Source/NiagaraEditor/Public/ViewModels/Stack/NiagaraStackScriptItemGroup.h
@@ -34,7 +34,7 @@ public:
 	virtual bool SupportsPaste() const override { return true; }
 	virtual bool TestCanPasteWithMessage(const UNiagaraClipboardContent* ClipboardContent, FText& OutMessage) const override;
 	virtual FText GetPasteTransactionText(const UNiagaraClipboardContent* ClipboardContent) const override;
-	virtual void Paste(const UNiagaraClipboardContent* ClipboardContent) override;
+	virtual void Paste(const UNiagaraClipboardContent* ClipboardContent, FText& OutPasteWarning) override;
 
 protected:
 	virtual void RefreshChildrenInternal(const TArray<UNiagaraStackEntry*>& CurrentChildren, TArray<UNiagaraStackEntry*>& NewChildren, TArray<FStackIssue>& NewIssues) override;
@@ -54,7 +54,9 @@ private:
 
 	void ChildModifiedGroupItems();
 
-	void ChildRequestPaste(const UNiagaraClipboardContent* ClipboardContent, int32 PasteIndex);
+	bool ChildRequestCanPaste(const UNiagaraClipboardContent* ClipboardContent, FText& OutCanPasteMessage);
+
+	void ChildRequestPaste(const UNiagaraClipboardContent* ClipboardContent, int32 PasteIndex, FText& OutPasteWarning);
 
 	void OnScriptGraphChanged(const struct FEdGraphEditAction& InAction);
 
@@ -79,7 +81,7 @@ protected:
 	void RefreshIssues(TArray<FStackIssue>& NewIssues);
 
 private:
-	void PasteModules(const UNiagaraClipboardContent* ClipboardContent, int32 PasteIndex);
+	void PasteModules(const UNiagaraClipboardContent* ClipboardContent, int32 PasteIndex, FText& OutPasteWarning);
 
 private:
 	TSharedPtr<FScriptItemGroupAddUtilities> AddUtilities;
diff --git a/Engine/Plugins/FX/Niagara/Source/NiagaraEditorWidgets/Private/NiagaraStackCommandContext.cpp b/Engine/Plugins/FX/Niagara/Source/NiagaraEditorWidgets/Private/NiagaraStackCommandContext.cpp
--- a/Engine/Plugins/FX/Niagara/Source/NiagaraEditorWidgets/Private/NiagaraStackCommandContext.cpp
+++ b/Engine/Plugins/FX/Niagara/Source/NiagaraEditorWidgets/Private/NiagaraStackCommandContext.cpp
@@ -173,7 +173,12 @@ FText FNiagaraStackCommandContext::GetCanPasteSelectedEntriesToolTip() const
 
 void FNiagaraStackCommandContext::PasteSelectedEntries() const
 {
-	FNiagaraStackClipboardUtilities::PasteSelection(SelectedEntries);
+	FText PasteWarning;
+	FNiagaraStackClipboardUtilities::PasteSelection(SelectedEntries, PasteWarning);
+	if (PasteWarning.IsEmptyOrWhitespace() == false)
+	{
+		WarnWithToastAndLog(PasteWarning);
+	}
 }
 
 bool FNiagaraStackCommandContext::CanDeleteSelectedEntries() const
