Added "Create Duplicate Parent" for Niagara Emitters, which duplicates the selected emitter and reparents it to the newly-created emitter. This allows you to insert parents into the inheritance chain of emitters.

diff --git a/Engine/Plugins/FX/Niagara/Source/Niagara/Classes/NiagaraEmitter.h b/Engine/Plugins/FX/Niagara/Source/Niagara/Classes/NiagaraEmitter.h
--- a/Engine/Plugins/FX/Niagara/Source/Niagara/Classes/NiagaraEmitter.h
+++ b/Engine/Plugins/FX/Niagara/Source/Niagara/Classes/NiagaraEmitter.h
@@ -237,10 +237,10 @@ public:
 
 public:
 #if WITH_EDITOR
-	/** Creates a new emitter with the supplied emitter as a parent emitter and the supplied system as it's owner. */
+	/** Creates a new emitter with the supplied emitter as a parent emitter and the supplied system as its owner. */
 	NIAGARA_API static UNiagaraEmitter* CreateWithParentAndOwner(UNiagaraEmitter& InParentEmitter, UObject* InOwner, FName InName, EObjectFlags FlagMask);
 
-	/** Creates a new emitter by duplicating an existing emitter.  The new emitter  will reference the same parent emitter if one is available. */
+	/** Creates a new emitter by duplicating an existing emitter. The new emitter will reference the same parent emitter if one is available. */
 	static UNiagaraEmitter* CreateAsDuplicate(const UNiagaraEmitter& InEmitterToDuplicate, FName InDuplicateName, UNiagaraSystem& InDuplicateOwnerSystem);
 
 	//Begin UObject Interface
@@ -251,6 +251,7 @@ public:
 	void Serialize(FArchive& Ar)override;
 	virtual void PostInitProperties() override;
 	virtual void PostLoad() override;
+	virtual void PostDuplicate(bool bDuplicateForPIE) override;
 	//End UObject Interface
 
 	/** Toggles whether or not the particles within this emitter are relative to the emitter origin or in global space. */ 
@@ -474,6 +475,8 @@ public:
 	NIAGARA_API UNiagaraEmitter* GetParent() const;
 
 	NIAGARA_API void RemoveParent();
+
+	NIAGARA_API void SetParent(UNiagaraEmitter& InParent);
 #endif
 
 	void InitFastPathAttributeNames();
diff --git a/Engine/Plugins/FX/Niagara/Source/Niagara/Private/NiagaraEmitter.cpp b/Engine/Plugins/FX/Niagara/Source/Niagara/Private/NiagaraEmitter.cpp
--- a/Engine/Plugins/FX/Niagara/Source/Niagara/Private/NiagaraEmitter.cpp
+++ b/Engine/Plugins/FX/Niagara/Source/Niagara/Private/NiagaraEmitter.cpp
@@ -170,6 +170,16 @@ void UNiagaraEmitter::PostInitProperties()
 	UniqueEmitterName = TEXT("Emitter");
 }
 
+void UNiagaraEmitter::PostDuplicate(bool bDuplicateForPIE)
+{
+	Super::PostDuplicate(bDuplicateForPIE);
+
+	if (!bDuplicateForPIE)
+	{
+		SetUniqueEmitterName(GetName());
+	}
+}
+
 #if WITH_EDITORONLY_DATA
 bool UNiagaraEmitter::GetForceCompileOnLoad()
 {
@@ -1414,6 +1424,14 @@ void UNiagaraEmitter::RemoveParent()
 	Parent = nullptr;
 	ParentAtLastMerge = nullptr;
 }
+
+void UNiagaraEmitter::SetParent(UNiagaraEmitter& InParent)
+{
+	Parent = &InParent;
+	ParentAtLastMerge = Cast<UNiagaraEmitter>(StaticDuplicateObject(&InParent, this));
+	ParentAtLastMerge->ClearFlags(RF_Standalone | RF_Public);
+	GraphSource->MarkNotSynchronized(TEXT("Emitter parent changed"));
+}
 #endif
 
 void UNiagaraEmitter::InitFastPathAttributeNames()
diff --git a/Engine/Plugins/FX/Niagara/Source/NiagaraEditor/Private/AssetTypeActions/AssetTypeActions_NiagaraEmitter.cpp b/Engine/Plugins/FX/Niagara/Source/NiagaraEditor/Private/AssetTypeActions/AssetTypeActions_NiagaraEmitter.cpp
--- a/Engine/Plugins/FX/Niagara/Source/NiagaraEditor/Private/AssetTypeActions/AssetTypeActions_NiagaraEmitter.cpp
+++ b/Engine/Plugins/FX/Niagara/Source/NiagaraEditor/Private/AssetTypeActions/AssetTypeActions_NiagaraEmitter.cpp
@@ -2,6 +2,7 @@
 
 #include "AssetTypeActions_NiagaraEmitter.h"
 #include "NiagaraEmitter.h"
+#include "NiagaraScriptSourceBase.h"
 #include "NiagaraSystem.h"
 #include "NiagaraSystemToolkit.h"
 #include "NiagaraEditorStyle.h"
@@ -38,7 +39,7 @@ void FAssetTypeActions_NiagaraEmitter::OpenAssetEditor(const TArray<UObject*>& I
 
 	for (auto ObjIt = InObjects.CreateConstIterator(); ObjIt; ++ObjIt)
 	{
-		auto Emitter = Cast<UNiagaraEmitter>(*ObjIt);
+		UNiagaraEmitter* Emitter = Cast<UNiagaraEmitter>(*ObjIt);
 		if (Emitter != nullptr)
 		{
 			TSharedRef<FNiagaraSystemToolkit> SystemToolkit(new FNiagaraSystemToolkit());
@@ -54,7 +55,7 @@ UClass* FAssetTypeActions_NiagaraEmitter::GetSupportedClass() const
 
 void FAssetTypeActions_NiagaraEmitter::GetActions(const TArray<UObject*>& InObjects, FToolMenuSection& Section)
 {
-	auto NiagaraEmitters = GetTypedWeakObjectPtrs<UNiagaraEmitter>(InObjects);
+	TArray<TWeakObjectPtr<UNiagaraEmitter>> NiagaraEmitters = GetTypedWeakObjectPtrs<UNiagaraEmitter>(InObjects);
 
 	Section.AddMenuEntry(
 		"Emitter_NewNiagaraSystem",
@@ -66,21 +67,23 @@ void FAssetTypeActions_NiagaraEmitter::GetActions(const TArray<UObject*>& InObje
 		)
 	);
 
-	bool bInObjectsAreCompilable = true;
-	for (UObject* InObject : InObjects)
-	{
-		if (FNiagaraEditorUtilities::IsCompilableAssetClass(InObject->GetClass()) == false)
-		{
-			bInObjectsAreCompilable = false;
-		}
-	}
+	Section.AddMenuEntry(
+		"Emitter_CreateDuplicateParent",
+		LOCTEXT("Emitter_CreateDuplicateParent", "Create Duplicate Parent"),
+		LOCTEXT("Emitter_CreateDuplicateParentTooltip", "Duplicate this emitter and set this emitter's parent to the new emitter."),
+		FSlateIcon(),
+		FUIAction(
+			FExecuteAction::CreateSP(this, &FAssetTypeActions_NiagaraEmitter::ExecuteCreateDuplicateParent, NiagaraEmitters)
+		)
+	);
 
 	Section.AddMenuEntry(
 		"MarkDependentCompilableAssetsDirty",
-		LOCTEXT("MarkDependentCompilableAssetsDirtyLabel", "Mark dependent compilable assets dirty"),
+		LOCTEXT("MarkDependentCompilableAssetsDirtyLabel", "Mark Dependent Compilable Assets Dirty"),
 		LOCTEXT("MarkDependentCompilableAssetsDirtyToolTip", "Finds all niagara assets which depend on this asset either directly or indirectly,\n and marks them dirty so they can be saved with the latest version."),
 		FSlateIcon(),
-		FUIAction(FExecuteAction::CreateStatic(&FNiagaraEditorUtilities::MarkDependentCompilableAssetsDirty, InObjects)));
+		FUIAction(FExecuteAction::CreateStatic(&FNiagaraEditorUtilities::MarkDependentCompilableAssetsDirty, InObjects))
+	);
 }
 
 void FAssetTypeActions_NiagaraEmitter::ExecuteNewNiagaraSystem(TArray<TWeakObjectPtr<UNiagaraEmitter>> Objects)
@@ -92,10 +95,9 @@ void FAssetTypeActions_NiagaraEmitter::ExecuteNewNiagaraSystem(TArray<TWeakObjec
 	SystemOptions.EditMode = ENiagaraSystemViewModelEditMode::SystemAsset;
 
 	TArray<UObject*> ObjectsToSync;
-	for (auto ObjIt = Objects.CreateConstIterator(); ObjIt; ++ObjIt)
+	for (TWeakObjectPtr<UNiagaraEmitter> Emitter : Objects)
 	{
-		auto Emitter = (*ObjIt).Get();
-		if (Emitter)
+		if (Emitter.IsValid())
 		{
 			// Determine an appropriate names
 			FString Name;
@@ -104,7 +106,7 @@ void FAssetTypeActions_NiagaraEmitter::ExecuteNewNiagaraSystem(TArray<TWeakObjec
 
 			// Create the factory used to generate the asset
 			UNiagaraSystemFactoryNew* Factory = NewObject<UNiagaraSystemFactoryNew>();
-			Factory->EmittersToAddToNewSystem.Add(Emitter);
+			Factory->EmittersToAddToNewSystem.Add(Emitter.Get());
 			FAssetToolsModule& AssetToolsModule = FModuleManager::GetModuleChecked<FAssetToolsModule>("AssetTools");
 			UObject* NewAsset = AssetToolsModule.Get().CreateAsset(Name, FPackageName::GetLongPackagePath(PackageName), UNiagaraSystem::StaticClass(), Factory);
 			
@@ -123,4 +125,46 @@ void FAssetTypeActions_NiagaraEmitter::ExecuteNewNiagaraSystem(TArray<TWeakObjec
 	}
 }
 
+void FAssetTypeActions_NiagaraEmitter::ExecuteCreateDuplicateParent(TArray<TWeakObjectPtr<UNiagaraEmitter>> Emitters)
+{
+	const FString DefaultSuffix = TEXT("_Parent");
+
+	TArray<UObject*> ObjectsToSync;
+	FAssetToolsModule& AssetToolsModule = FModuleManager::GetModuleChecked<FAssetToolsModule>("AssetTools");
+
+	for (TWeakObjectPtr<UNiagaraEmitter> Emitter : Emitters)
+	{
+		if (Emitter.IsValid())
+		{
+			FString Name;
+			FString PackageName;
+			CreateUniqueAssetName(Emitter->GetOutermost()->GetName(), DefaultSuffix, PackageName, Name);
+
+			// Create the factory used to generate the asset
+			UNiagaraEmitter* NewEmitter = Cast<UNiagaraEmitter>(AssetToolsModule.Get().DuplicateAssetWithDialog(Name, FPackageName::GetLongPackagePath(PackageName), Emitter.Get()));
+
+			if (NewEmitter != nullptr)
+			{
+				NewEmitter->Modify();
+				NewEmitter->SetUniqueEmitterName(Name);
+				Emitter->GraphSource->MarkNotSynchronized(TEXT("Emitter created"));
+
+				Emitter->Modify();
+				Emitter->SetParent(*NewEmitter);
+				Emitter->GraphSource->MarkNotSynchronized(TEXT("Emitter parent changed"));
+
+				ObjectsToSync.Add(NewEmitter);
+			}
+		}
+	}
+
+	if (ObjectsToSync.Num() > 0)
+	{
+		AssetToolsModule.Get().OpenEditorForAssets(ObjectsToSync);
+
+		FContentBrowserModule& ContentBrowserModule = FModuleManager::LoadModuleChecked<FContentBrowserModule>("ContentBrowser");
+		ContentBrowserModule.Get().SyncBrowserToAssets(ObjectsToSync);
+	}
+}
+
 #undef LOCTEXT_NAMESPACE
diff --git a/Engine/Plugins/FX/Niagara/Source/NiagaraEditor/Private/AssetTypeActions/AssetTypeActions_NiagaraEmitter.h b/Engine/Plugins/FX/Niagara/Source/NiagaraEditor/Private/AssetTypeActions/AssetTypeActions_NiagaraEmitter.h
--- a/Engine/Plugins/FX/Niagara/Source/NiagaraEditor/Private/AssetTypeActions/AssetTypeActions_NiagaraEmitter.h
+++ b/Engine/Plugins/FX/Niagara/Source/NiagaraEditor/Private/AssetTypeActions/AssetTypeActions_NiagaraEmitter.h
@@ -23,4 +23,5 @@ public:
 
 private:
 	void ExecuteNewNiagaraSystem(TArray<TWeakObjectPtr<UNiagaraEmitter>> Emitters);
+	void ExecuteCreateDuplicateParent(TArray<TWeakObjectPtr<UNiagaraEmitter>> Emitters);
 };
\ No newline at end of file
