Fix for skeletal mesh DIs not handling attachement change correctly.

diff --git a/Engine/Plugins/FX/Niagara/Source/Niagara/Classes/NiagaraDataInterfaceSkeletalMesh.h b/Engine/Plugins/FX/Niagara/Source/Niagara/Classes/NiagaraDataInterfaceSkeletalMesh.h
--- a/Engine/Plugins/FX/Niagara/Source/Niagara/Classes/NiagaraDataInterfaceSkeletalMesh.h
+++ b/Engine/Plugins/FX/Niagara/Source/Niagara/Classes/NiagaraDataInterfaceSkeletalMesh.h
@@ -358,12 +358,15 @@ private:
 
 struct FNDISkeletalMesh_InstanceData
 {
-	//Cached ptr to component we sample from. 
+	//Cached ptr to component we sample from. TODO: This should not need to be a weak ptr. We should always be clearing out DIs when the component is destroyed.
 	TWeakObjectPtr<USceneComponent> Component;
 
 	/** A binding to the user ptr we're reading the mesh from (if we are). */
 	FNiagaraParameterDirectBinding<UObject*> UserParamBinding;
 
+	//Always reset the DI when the attach parent changes.
+	USceneComponent* CachedAttachParent;
+
 	UObject* CachedUserParam;
 
 	USkeletalMesh* Mesh;
diff --git a/Engine/Plugins/FX/Niagara/Source/Niagara/Private/NiagaraDataInterfaceSkeletalMesh.cpp b/Engine/Plugins/FX/Niagara/Source/Niagara/Private/NiagaraDataInterfaceSkeletalMesh.cpp
--- a/Engine/Plugins/FX/Niagara/Source/Niagara/Private/NiagaraDataInterfaceSkeletalMesh.cpp
+++ b/Engine/Plugins/FX/Niagara/Source/Niagara/Private/NiagaraDataInterfaceSkeletalMesh.cpp
@@ -1092,6 +1092,7 @@ bool FNDISkeletalMesh_InstanceData::Init(UNiagaraDataInterfaceSkeletalMesh* Inte
 
 	// Initialize members
 	Component = nullptr;
+	CachedAttachParent = nullptr;
 	Mesh = nullptr;
 	MeshSafe = nullptr;
 	Transform = FMatrix::Identity;
@@ -1132,6 +1133,8 @@ bool FNDISkeletalMesh_InstanceData::Init(UNiagaraDataInterfaceSkeletalMesh* Inte
 	TransformInverseTransposed = Transform.InverseFast().GetTransposed();
 	PrevTransform = Transform;
 	PrevTransformInverseTransposed = TransformInverseTransposed;
+	
+	CachedAttachParent = Component->GetAttachParent();
 
 #if WITH_EDITOR
 	MeshSafe->GetOnMeshChanged().AddUObject(SystemInstance->GetComponent(), &UNiagaraComponent::ReinitializeSystem);
@@ -1388,6 +1391,12 @@ bool FNDISkeletalMesh_InstanceData::ResetRequired(UNiagaraDataInterfaceSkeletalM
 		//The component we were bound to is no longer valid so we have to trigger a reset.
 		return true;
 	}
+
+	//Detect and reset on any attachment change.
+	if (Comp->GetAttachParent() != CachedAttachParent)
+	{
+		return true;
+	}
 	
 	if (USkeletalMeshComponent* SkelComp = Cast<USkeletalMeshComponent>(Comp))
 	{
