Remove commented out debug code which should never have been submitted.

diff --git a/Engine/Plugins/FX/Niagara/Source/NiagaraShader/Private/NiagaraShader.cpp b/Engine/Plugins/FX/Niagara/Source/NiagaraShader/Private/NiagaraShader.cpp
--- a/Engine/Plugins/FX/Niagara/Source/NiagaraShader/Private/NiagaraShader.cpp
+++ b/Engine/Plugins/FX/Niagara/Source/NiagaraShader/Private/NiagaraShader.cpp
@@ -1510,26 +1510,9 @@ bool operator<<(FArchive& Ar, FNiagaraDataInterfaceParamRef& ParamRef)
 
 bool FNiagaraDataInterfaceGeneratedFunction::Serialize(FArchive& Ar)
 {
-	Ar.UsingCustomVersion(FNiagaraCustomVersion::GUID);
-	const int32 NiagaraVer = Ar.CustomVer(FNiagaraCustomVersion::GUID);
-
 	Ar << DefinitionName;
 	Ar << InstanceName;
-
-	//if (Ar.IsLoading() && NiagaraVer == FNiagaraCustomVersion::AddGeneratedFunctionsToGPUParamInfo)
-	//{
-	//	TMap<FName, FName> cacat;
-	//	Ar << cacat;
-	//	Specifiers.SetNum(cacat.Num());
-	//	for (const auto& muie : cacat)
-	//	{
-	//		Specifiers.Add(muie);
-	//	}
-	//}
-	//else
-	{
-		Ar << Specifiers;
-	}
+	Ar << Specifiers;
 	return true;
 }
 
