Fixed duplicate localization string.

diff --git a/Engine/Plugins/FX/Niagara/Source/NiagaraEditor/Private/NiagaraHlslTranslator.cpp b/Engine/Plugins/FX/Niagara/Source/NiagaraEditor/Private/NiagaraHlslTranslator.cpp
--- a/Engine/Plugins/FX/Niagara/Source/NiagaraEditor/Private/NiagaraHlslTranslator.cpp
+++ b/Engine/Plugins/FX/Niagara/Source/NiagaraEditor/Private/NiagaraHlslTranslator.cpp
@@ -1792,7 +1792,7 @@ void FHlslNiagaraTranslator::DefineDataInterfaceHLSL(FString &InHlslOutput)
 				const bool HlslOK = CDO->GetFunctionHLSL(DIInstanceInfo, DIFunc, FunctionInstanceIndex, InterfaceFunctionHLSL);
 				if (!HlslOK)
 				{
-					Error(FText::Format(LOCTEXT("GPUDataInterfaceFunctionNotSupported", "DataInterface {0} function {1} is not implemented for GPU."), FText::FromName(Info.Type.GetFName()), FText::FromName(DIFunc.DefinitionName)), nullptr, nullptr);
+					Error(FText::Format(LOCTEXT("GPUDataInterfaceFunctionNotImplemented", "DataInterface {0} function {1} is not implemented for GPU."), FText::FromName(Info.Type.GetFName()), FText::FromName(DIFunc.DefinitionName)), nullptr, nullptr);
 				}
 			}
 		}
