- Clean up ExecIndex and implement Engine.ExecutionCount for GPU systems

diff --git a/Engine/Plugins/FX/Niagara/Shaders/Private/NiagaraEmitterInstanceShader.usf b/Engine/Plugins/FX/Niagara/Shaders/Private/NiagaraEmitterInstanceShader.usf
--- a/Engine/Plugins/FX/Niagara/Shaders/Private/NiagaraEmitterInstanceShader.usf
+++ b/Engine/Plugins/FX/Niagara/Shaders/Private/NiagaraEmitterInstanceShader.usf
@@ -90,20 +90,52 @@ float Reciprocal(float x){ return 1.0f/x; }
 	static float Emitter_InterpSpawnStartDt;
 	static int Emitter_SpawnGroup;
 
-	void GetEmitterSpawnInfoForParticle(int ParticleIndex)
+	static int Engine_ExecutionCount;
+	static int GGPUExecIndex;
+
+	static uint GSpawnStartInstance;
+	uint SpawnedInstances;
+	uint UpdateStartInstance;
+
+	void SetupExecIndexForGPU()
 	{
+		GGPUExecIndex = UpdateStartInstance + GDispatchThreadId.x;
+		Engine_ExecutionCount = GSpawnStartInstance;
+	}
+
+	void SetupExecIndexAndSpawnInfoForGPU()
+	{
+		GGPUExecIndex = UpdateStartInstance + GDispatchThreadId.x - GSpawnStartInstance;
+
 		int SpawnInfoIndex = 0;
 		UNROLL
 		for ( int i=0; i < NIAGARA_MAX_GPU_SPAWN_INFOS_V4; ++i )
 		{
-			SpawnInfoIndex += (int)dot(1.0f, ParticleIndex >= EmitterSpawnInfoOffsets[i]);
+			SpawnInfoIndex += (int)dot(1.0f, GGPUExecIndex >= EmitterSpawnInfoOffsets[i]);
 		}
 
 		Emitter_SpawnInterval		= EmitterSpawnInfoParams[SpawnInfoIndex].x;
 		Emitter_InterpSpawnStartDt	= EmitterSpawnInfoParams[SpawnInfoIndex].y;
 		Emitter_SpawnGroup			= asint(EmitterSpawnInfoParams[SpawnInfoIndex].z);
 
-		GInterpSpawnIndex		= ParticleIndex - asint(EmitterSpawnInfoParams[SpawnInfoIndex].w);
+		int GroupSpawnStartIndex    = asint(EmitterSpawnInfoParams[SpawnInfoIndex].w);
+		GGPUExecIndex		        = GGPUExecIndex - GroupSpawnStartIndex;
+		if ( SpawnInfoIndex == (NIAGARA_MAX_GPU_SPAWN_INFOS - 1) )
+		{
+			Engine_ExecutionCount   = int(SpawnedInstances) - GroupSpawnStartIndex;
+		}
+		else
+		{
+			int NextGroupSpawnStartIndex = asint(EmitterSpawnInfoParams[SpawnInfoIndex + 1].w);
+			Engine_ExecutionCount       = NextGroupSpawnStartIndex - GroupSpawnStartIndex;
+		}
+	}
+
+	/* Returns the current instance index relative to the operation (spawn/update)
+	 */
+	int ExecIndex()
+	{
+		return GGPUExecIndex;
 	}
 #endif
 
@@ -382,12 +414,6 @@ int rand_int(int x, int Seed1, int Seed2, int Seed3)
  * -----------------------------------------------------------------
  */
 #if GPU_SIMULATION
-	static uint GSpawnStartInstance;
-		
-	/* Parameters for spawning and simulations
-	 */
-	uint UpdateStartInstance;
-	
 	/*
 	 * declared as int instead of uint to avoid compiler warnings because we can't expose uint on Niagara module parameters.  
 	 * This is still declared as a uint in the corresponding C++ classes
@@ -395,9 +421,7 @@ int rand_int(int x, int Seed1, int Seed2, int Seed3)
 	int DefaultShaderStageIndex;
 	int ShaderStageIndex;
 
-
 	int IterationInterfaceCount;
-	uint SpawnedInstances;
 
 	uint ComponentBufferSizeRead;
 	uint ComponentBufferSizeWrite;
@@ -421,21 +445,6 @@ int rand_int(int x, int Seed1, int Seed2, int Seed3)
 	Buffer<int> FreeIDList;
 	RWBuffer<int> RWIDToIndexTable;
 
-
-	/* Returns the current instance index relative to the operation (spawn/update)
-	 */
-	int ExecIndex()
-	{
-		if (GCurrentPhase == 0)
-		{
-			return UpdateStartInstance + GDispatchThreadId.x - GSpawnStartInstance;
-		}
-		else 
-		{
-			return UpdateStartInstance + GDispatchThreadId.x;
-		}
-	}
-
 	void AcquireID(int DataSetID, out int IDIndex, out int IDAcquireTag)
 	{
 		// Begin static assert : GPU particles only support DataSetID 0
@@ -719,13 +728,15 @@ void SimulateMainComputeCS(
 	const float RandomSeedInitialisation = NiagaraInternalNoise(InstanceID*16384, 0 * 8196, (bRunPhase1 ? 4096 : 0) + EmitterTickCounter);	// initialise the random state seed
 
 	FSimulationContext Context = (FSimulationContext)0;
-	InitConstants(Context);
 
 	BRANCH
 	if (bRunPhase1)
 	{
 		GCurrentPhase = 1;
 
+		SetupExecIndexForGPU();
+
+		InitConstants(Context);
 		LoadUpdateVariables(Context, InstanceID);
 		ReadDataSets(Context);
 
@@ -736,8 +747,9 @@ void SimulateMainComputeCS(
 	{
 		GCurrentPhase = 0;
 
-		GetEmitterSpawnInfoForParticle(ExecIndex());
+		SetupExecIndexAndSpawnInfoForGPU();
 
+		InitConstants(Context);
 		InitSpawnVariables(Context);
 		ReadDataSets(Context);
 
diff --git a/Engine/Plugins/FX/Niagara/Source/Niagara/Private/NiagaraEmitterInstance.cpp b/Engine/Plugins/FX/Niagara/Source/Niagara/Private/NiagaraEmitterInstance.cpp
--- a/Engine/Plugins/FX/Niagara/Source/Niagara/Private/NiagaraEmitterInstance.cpp
+++ b/Engine/Plugins/FX/Niagara/Source/Niagara/Private/NiagaraEmitterInstance.cpp
@@ -1263,10 +1263,13 @@ void FNiagaraEmitterInstance::Tick(float DeltaSeconds)
 			// If we have spawning make sure we clear out the remaining data and leave the end slot as MAX to avoid reading off end of the array on the GPU
 			if ( GpuSpawnInfo.EventSpawnTotal + GpuSpawnInfo.SpawnRateInstances >  0 )
 			{
-				while (NumSpawnInfos < NIAGARA_MAX_GPU_SPAWN_INFOS - 1)
+				while (NumSpawnInfos < NIAGARA_MAX_GPU_SPAWN_INFOS)
 				{
-					reinterpret_cast<float*>(GpuSpawnInfo.SpawnInfoStartOffsets)[NumSpawnInfos] = MAX_FLT;
-					GpuSpawnInfo.SpawnInfoParams[NumSpawnInfos] = FVector4(ForceInitToZero);
+					GpuSpawnInfo.SpawnInfoParams[NumSpawnInfos].X = 0.0f;
+					GpuSpawnInfo.SpawnInfoParams[NumSpawnInfos].Y = 0.0f;
+					reinterpret_cast<int32&>(GpuSpawnInfo.SpawnInfoParams[NumSpawnInfos].Z) = 0;
+					reinterpret_cast<int32&>(GpuSpawnInfo.SpawnInfoParams[NumSpawnInfos].W) = GpuSpawnInfo.SpawnRateInstances;
+					reinterpret_cast<float*>(GpuSpawnInfo.SpawnInfoStartOffsets)[NumSpawnInfos] = (float)MAX_FLT;
 					++NumSpawnInfos;
 				}
 				reinterpret_cast<float*>(GpuSpawnInfo.SpawnInfoStartOffsets)[NIAGARA_MAX_GPU_SPAWN_INFOS - 1] = MAX_FLT;
diff --git a/Engine/Plugins/FX/Niagara/Source/NiagaraCore/Private/NiagaraCustomVersion.cpp b/Engine/Plugins/FX/Niagara/Source/NiagaraCore/Private/NiagaraCustomVersion.cpp
--- a/Engine/Plugins/FX/Niagara/Source/NiagaraCore/Private/NiagaraCustomVersion.cpp
+++ b/Engine/Plugins/FX/Niagara/Source/NiagaraCore/Private/NiagaraCustomVersion.cpp
@@ -8,4 +8,4 @@ const FGuid FNiagaraCustomVersion::GUID(0xFCF57AFA, 0x50764283, 0xB9A9E658, 0xFF
 // Register the custom version with core
 FCustomVersionRegistration GRegisterNiagaraCustomVersion(FNiagaraCustomVersion::GUID, FNiagaraCustomVersion::LatestVersion, TEXT("NiagaraVer"));
 
-const FGuid FNiagaraCustomVersion::LatestScriptCompileVersion(0x64FCB871, 0xB2954D47, 0x8721FF75, 0x20496D93);
+const FGuid FNiagaraCustomVersion::LatestScriptCompileVersion(0x4953F984, 0x8FEC104D, 0xAF87DB5C, 0x44A7F79F);
diff --git a/Engine/Plugins/FX/Niagara/Source/NiagaraEditor/Private/NiagaraHlslTranslator.cpp b/Engine/Plugins/FX/Niagara/Source/NiagaraEditor/Private/NiagaraHlslTranslator.cpp
--- a/Engine/Plugins/FX/Niagara/Source/NiagaraEditor/Private/NiagaraHlslTranslator.cpp
+++ b/Engine/Plugins/FX/Niagara/Source/NiagaraEditor/Private/NiagaraHlslTranslator.cpp
@@ -752,11 +752,16 @@ void FHlslNiagaraTranslator::HandleNamespacedExternalVariablesToDataSetRead(TArr
 
 bool FHlslNiagaraTranslator::IsVariableInUniformBuffer(const FNiagaraVariable& Variable) const
 {
-	static FNiagaraVariable ExcludeVariables[] =
+	static FNiagaraVariable GpuExcludeVariables[] =
 	{
+		// Variables that must be calcualted on the GPU
+		FNiagaraVariable(FNiagaraTypeDefinition::GetIntDef(),   TEXT("Engine.ExecutionCount")),
+		FNiagaraVariable(FNiagaraTypeDefinition::GetIntDef(),   TEXT("Engine_ExecutionCount")),
+
+		// Spawn variables
 		FNiagaraVariable(FNiagaraTypeDefinition::GetFloatDef(), TEXT("Emitter_SpawnInterval")),
-		FNiagaraVariable(FNiagaraTypeDefinition::GetFloatDef(), TEXT("Emitter_InterpSpawnStartDt")),
 		FNiagaraVariable(FNiagaraTypeDefinition::GetFloatDef(), TEXT("Emitter.SpawnInterval")),
+		FNiagaraVariable(FNiagaraTypeDefinition::GetFloatDef(), TEXT("Emitter_InterpSpawnStartDt")),
 		FNiagaraVariable(FNiagaraTypeDefinition::GetFloatDef(), TEXT("Emitter.InterpSpawnStartDt")),
 		FNiagaraVariable(FNiagaraTypeDefinition::GetIntDef(),   TEXT("Emitter_SpawnGroup")),
 		FNiagaraVariable(FNiagaraTypeDefinition::GetIntDef(),   TEXT("Emitter.SpawnGroup")),
@@ -764,7 +769,7 @@ bool FHlslNiagaraTranslator::IsVariableInUniformBuffer(const FNiagaraVariable& V
 
 	if (CompilationTarget == ENiagaraSimTarget::GPUComputeSim)
 	{
-		for ( const FNiagaraVariable& ExcludeVar : ExcludeVariables )
+		for ( const FNiagaraVariable& ExcludeVar : GpuExcludeVariables)
 		{
 			if ( Variable == ExcludeVar )
 			{
@@ -1520,16 +1525,7 @@ void FHlslNiagaraTranslator::DefineInterpolatedParametersFunction(FString &HlslO
 		FString PrevMap = TranslationStages[i - 1].PassNamespace;
 		FString CurMap = TranslationStages[i].PassNamespace;
 		{
-			// GPU simulation is slightly different as we run all spawns at once rather than separate invocations so we can not ExecIndex() as it has to be biased into the correct spawn group's index
-			if ( CompilationTarget == ENiagaraSimTarget::GPUComputeSim )
-			{
-				HlslOutputString += TEXT("\tint InterpSpawn_Index = GInterpSpawnIndex;\n");
-			}
-			else
-			{
-				HlslOutputString += TEXT("\tint InterpSpawn_Index = ExecIndex();\n");
-			}
-
+			HlslOutputString += TEXT("\tint InterpSpawn_Index = ExecIndex();\n");
 			HlslOutputString += TEXT("\tfloat InterpSpawn_SpawnTime = ") + Emitter_InterpSpawnStartDt + TEXT(" + (") + Emitter_SpawnInterval + TEXT(" * InterpSpawn_Index);\n");
 			HlslOutputString += TEXT("\tfloat InterpSpawn_UpdateTime = Engine_DeltaTime - InterpSpawn_SpawnTime;\n");
 			HlslOutputString += TEXT("\tfloat InterpSpawn_InvSpawnTime = 1.0 / InterpSpawn_SpawnTime;\n");
