- Add support for SubUV's on Niagara Mesh Renderer
- Add nointerpolation to various attributes that don't require interpolation

#test Ran EngineTests for Niagara

diff --git a/Engine/Plugins/FX/Niagara/Shaders/Private/NiagaraMeshVertexFactory.ush b/Engine/Plugins/FX/Niagara/Shaders/Private/NiagaraMeshVertexFactory.ush
--- a/Engine/Plugins/FX/Niagara/Shaders/Private/NiagaraMeshVertexFactory.ush
+++ b/Engine/Plugins/FX/Niagara/Shaders/Private/NiagaraMeshVertexFactory.ush
@@ -20,6 +20,7 @@
 	Buffer<float> NiagaraParticleDataScale;
 	Buffer<float> NiagaraParticleDataTransform;
 	Buffer<float> NiagaraParticleDataNormalizedAge;
+	Buffer<float> NiagaraParticleDataSubImage;
 	Buffer<float> NiagaraParticleDataMaterialRandom;
 	Buffer<float> NiagaraParticleDataMaterialParam0;
 	Buffer<float> NiagaraParticleDataMaterialParam1;
@@ -111,7 +112,7 @@ struct FVertexFactoryInterpolantsVSToPS
 #endif
 
 #if NEEDS_PARTICLE_COLOR
-	float4	ParticleColor : COLOR1;
+	nointerpolation float4	ParticleColor : COLOR1;
 #endif
 
 #if (DYNAMIC_PARAMETERS_MASK & 1)
@@ -129,25 +130,25 @@ struct FVertexFactoryInterpolantsVSToPS
 
 #if USE_PARTICLE_POSITION
 	/** Particle position in camera-centered translated world space */
-	float3 ParticleTranslatedWorldPosition	: PARTICLE_POSITION;
+	nointerpolation float3 ParticleTranslatedWorldPosition	: PARTICLE_POSITION;
 #endif
 
 #if USE_PARTICLE_VELOCITY
 	/** The velocity of the particle, XYZ: direction, W: speed. */
-	float4 ParticleVelocity	: PARTICLE_VELOCITY;
+	nointerpolation float4 ParticleVelocity	: PARTICLE_VELOCITY;
 #endif
 
 #if USE_PARTICLE_TIME
 	/** Relative alive time of the particle */
-	float RelativeTime : PARTICLE_TIME;
+	nointerpolation float RelativeTime : PARTICLE_TIME;
 #endif
 
 #if USE_PARTICLE_RANDOM
-	float ParticleRandom : PARTICLE_RANDOM;
+	nointerpolation float ParticleRandom : PARTICLE_RANDOM;
 #endif
 
 #if USE_PARTICLE_TRANSFORM
-	float4 LocalToWorld[4] : PARTICLE_TRANSFORM;
+	nointerpolation float4 LocalToWorld[4] : PARTICLE_TRANSFORM;
 #endif
 
 #if INSTANCED_STEREO
@@ -169,7 +170,7 @@ struct FVertexFactoryIntermediates
 #endif
 #if USE_PARTICLE_SUBUVS
 	float4 SubUVCoords;
-	float SubUVLerp;
+	nointerpolation float SubUVLerp;
 #endif
 
 	/** Optional dynamic parameters for the particle. */
@@ -706,13 +707,29 @@ FVertexFactoryIntermediates GetVertexFactoryIntermediates(FVertexFactoryInput In
 #endif
 
 #if USE_PARTICLE_SUBUVS
-		Intermediates.SubUVCoords.xy = (float2( Input.SubUVParams.x, Input.SubUVParams.y ) + Input.TexCoords[0].xy) * MeshParticleVF.SubImageSize.xy;
-		Intermediates.SubUVCoords.zw = (float2( Input.SubUVParams.z, Input.SubUVParams.w ) + Input.TexCoords[0].xy) * MeshParticleVF.SubImageSize.xy;
-		//Intermediates.SubUVLerp = Input.SubUVLerpAndRelTime.x;
+	if (NiagaraMeshVF.SubImageDataOffset == -1)
+	{
+		Intermediates.SubUVCoords.xy = 0.0f;
+		Intermediates.SubUVCoords.zw = 0.0f;
+		Intermediates.SubUVLerp = 0;
+	}
+	else
+	{
+		// SubUV.
+		float SubImageIndex = GetFloat(NIAGARA_LOOSEPARAM_STRUCT(NiagaraParticleDataSubImage), NiagaraMeshVF.SubImageDataOffset, Input.InstanceId);
+		float SubImageLerp = frac(SubImageIndex);
+		float SubImageA = SubImageIndex - SubImageLerp;
+		float SubImageB = SubImageA + 1;
+		float SubImageAH = fmod(SubImageA, NiagaraMeshVF.SubImageSize.x);
+		float SubImageBH = fmod(SubImageB, NiagaraMeshVF.SubImageSize.x);
+		float SubImageAV = floor(SubImageA * NiagaraMeshVF.SubImageSize.z);
+		float SubImageBV = floor(SubImageB * NiagaraMeshVF.SubImageSize.z);
+		Intermediates.SubUVCoords.xy = (float2(SubImageAH, SubImageAV) + Input.TexCoords[0].xy) * NiagaraMeshVF.SubImageSize.zw;
+		Intermediates.SubUVCoords.zw = (float2(SubImageBH, SubImageBV) + Input.TexCoords[0].xy) * NiagaraMeshVF.SubImageSize.zw;
+		Intermediates.SubUVLerp = (NiagaraMeshVF.SubImageBlendMode == 1) ? SubImageLerp : 0.0f;
+	}
 #endif
 
-	//Intermediates.RelativeTime = Input.SubUVLerpAndRelTime.y;
-
 	return Intermediates;
 }
 
diff --git a/Engine/Plugins/FX/Niagara/Shaders/Private/NiagaraSpriteVertexFactory.ush b/Engine/Plugins/FX/Niagara/Shaders/Private/NiagaraSpriteVertexFactory.ush
--- a/Engine/Plugins/FX/Niagara/Shaders/Private/NiagaraSpriteVertexFactory.ush
+++ b/Engine/Plugins/FX/Niagara/Shaders/Private/NiagaraSpriteVertexFactory.ush
@@ -127,15 +127,15 @@ struct FVertexFactoryInterpolantsVSToPS
 
 #if USE_PARTICLE_POSITION
 	/** Cam-relative (translated) particle center and radius */
-	float4 ParticleTranslatedWorldPositionAndSize	: PARTICLE_POSITION;
+	nointerpolation float4 ParticleTranslatedWorldPositionAndSize	: PARTICLE_POSITION;
 #endif
 
 #if USE_PARTICLE_VELOCITY
-	float4 ParticleVelocity : PARTICLE_VELOCITY;
+	nointerpolation float4 ParticleVelocity : PARTICLE_VELOCITY;
 #endif
 
 #if USE_PARTICLE_TIME
-	float RelativeTime : PARTICLE_TIME;
+	nointerpolation float RelativeTime : PARTICLE_TIME;
 #endif
 
 #if USE_PARTICLE_LIGHTING_OFFSET
@@ -143,11 +143,11 @@ struct FVertexFactoryInterpolantsVSToPS
 #endif
 
 #if USE_PARTICLE_SIZE
-	float2 ParticleSize : PARTICLE_SIZE;
+	nointerpolation float2 ParticleSize : PARTICLE_SIZE;
 #endif
 
 #if USE_PARTICLE_RANDOM
-	float ParticleRandom : PARTICLE_RANDOM;
+	nointerpolation float ParticleRandom : PARTICLE_RANDOM;
 #endif
 
 #if INSTANCED_STEREO
diff --git a/Engine/Plugins/FX/Niagara/Source/Niagara/Private/NiagaraMeshRendererProperties.cpp b/Engine/Plugins/FX/Niagara/Source/Niagara/Private/NiagaraMeshRendererProperties.cpp
--- a/Engine/Plugins/FX/Niagara/Source/Niagara/Private/NiagaraMeshRendererProperties.cpp
+++ b/Engine/Plugins/FX/Niagara/Source/Niagara/Private/NiagaraMeshRendererProperties.cpp
@@ -33,6 +33,8 @@ UNiagaraMeshRendererProperties::UNiagaraMeshRendererProperties()
 	: ParticleMesh(nullptr)
 	, SortMode(ENiagaraSortMode::None)
 	, bSortOnlyWhenTranslucent(true)
+	, SubImageSize(1.0f, 1.0f)
+	, bSubImageBlend(false)
 {
 }
 
@@ -111,6 +113,7 @@ void UNiagaraMeshRendererProperties::InitBindings()
 		PositionBinding = FNiagaraConstants::GetAttributeDefaultBinding(SYS_PARAM_PARTICLES_POSITION);
 		ColorBinding = FNiagaraConstants::GetAttributeDefaultBinding(SYS_PARAM_PARTICLES_COLOR);
 		VelocityBinding = FNiagaraConstants::GetAttributeDefaultBinding(SYS_PARAM_PARTICLES_VELOCITY);
+		SubImageIndexBinding = FNiagaraConstants::GetAttributeDefaultBinding(SYS_PARAM_PARTICLES_SUB_IMAGE_INDEX);
 		DynamicMaterialBinding = FNiagaraConstants::GetAttributeDefaultBinding(SYS_PARAM_PARTICLES_DYNAMIC_MATERIAL_PARAM);
 		DynamicMaterial1Binding = FNiagaraConstants::GetAttributeDefaultBinding(SYS_PARAM_PARTICLES_DYNAMIC_MATERIAL_PARAM_1);
 		DynamicMaterial2Binding = FNiagaraConstants::GetAttributeDefaultBinding(SYS_PARAM_PARTICLES_DYNAMIC_MATERIAL_PARAM_2);
@@ -239,6 +242,7 @@ const TArray<FNiagaraVariable>& UNiagaraMeshRendererProperties::GetOptionalAttri
 		Attrs.Add(SYS_PARAM_PARTICLES_NORMALIZED_AGE);
 		Attrs.Add(SYS_PARAM_PARTICLES_SCALE);
 		Attrs.Add(SYS_PARAM_PARTICLES_MESH_ORIENTATION);
+		Attrs.Add(SYS_PARAM_PARTICLES_SUB_IMAGE_INDEX);
 		Attrs.Add(SYS_PARAM_PARTICLES_DYNAMIC_MATERIAL_PARAM);
 		Attrs.Add(SYS_PARAM_PARTICLES_DYNAMIC_MATERIAL_PARAM_1);
 		Attrs.Add(SYS_PARAM_PARTICLES_DYNAMIC_MATERIAL_PARAM_2);
@@ -275,6 +279,9 @@ void UNiagaraMeshRendererProperties::PreEditChange(class FProperty* PropertyThat
 
 void UNiagaraMeshRendererProperties::PostEditChangeProperty(FPropertyChangedEvent& PropertyChangedEvent)
 {
+	SubImageSize.X = FMath::Max<float>(SubImageSize.X, 1.f);
+	SubImageSize.Y = FMath::Max<float>(SubImageSize.Y, 1.f);
+
 	static FName ParticleMeshName(TEXT("ParticleMesh"));
 	if (ParticleMesh && PropertyChangedEvent.Property && PropertyChangedEvent.Property->GetFName() == ParticleMeshName)
 	{
diff --git a/Engine/Plugins/FX/Niagara/Source/Niagara/Private/NiagaraRendererMeshes.cpp b/Engine/Plugins/FX/Niagara/Source/Niagara/Private/NiagaraRendererMeshes.cpp
--- a/Engine/Plugins/FX/Niagara/Source/Niagara/Private/NiagaraRendererMeshes.cpp
+++ b/Engine/Plugins/FX/Niagara/Source/Niagara/Private/NiagaraRendererMeshes.cpp
@@ -70,6 +70,7 @@ namespace ENiagaraMeshVFLayout
 		MaterialRandom,
 		NormalizedAge,
 		CustomSorting,
+		SubImage,
 		DynamicParam0,
 		DynamicParam1,
 		DynamicParam2,
@@ -95,6 +96,8 @@ FNiagaraRendererMeshes::FNiagaraRendererMeshes(ERHIFeatureLevel::Type FeatureLev
 	SortMode = Properties->SortMode;
 	bSortOnlyWhenTranslucent = Properties->bSortOnlyWhenTranslucent;
 	bOverrideMaterials = Properties->bOverrideMaterials;
+	SubImageSize = Properties->SubImageSize;
+	bSubImageBlend = Properties->bSubImageBlend;
 
 	const FNiagaraDataSet& Data = Emitter->GetData();
 
@@ -109,11 +112,12 @@ FNiagaraRendererMeshes::FNiagaraRendererMeshes(ERHIFeatureLevel::Type FeatureLev
 	SetVertexFactoryVariable(Data, Properties->MaterialRandomBinding.DataSetVariable, ENiagaraMeshVFLayout::MaterialRandom);
 	SetVertexFactoryVariable(Data, Properties->NormalizedAgeBinding.DataSetVariable, ENiagaraMeshVFLayout::NormalizedAge);
 	SetVertexFactoryVariable(Data, Properties->CustomSortingBinding.DataSetVariable, ENiagaraMeshVFLayout::CustomSorting);
+	SetVertexFactoryVariable(Data, Properties->SubImageIndexBinding.DataSetVariable, ENiagaraMeshVFLayout::SubImage);
 	MaterialParamValidMask |= SetVertexFactoryVariable(Data, Properties->DynamicMaterialBinding.DataSetVariable, ENiagaraMeshVFLayout::DynamicParam0) ? 0x1 : 0;
 	MaterialParamValidMask |= SetVertexFactoryVariable(Data, Properties->DynamicMaterial1Binding.DataSetVariable, ENiagaraMeshVFLayout::DynamicParam1) ? 0x2 : 0;
 	MaterialParamValidMask |= SetVertexFactoryVariable(Data, Properties->DynamicMaterial2Binding.DataSetVariable, ENiagaraMeshVFLayout::DynamicParam2) ? 0x4 : 0;
 	MaterialParamValidMask |= SetVertexFactoryVariable(Data, Properties->DynamicMaterial3Binding.DataSetVariable, ENiagaraMeshVFLayout::DynamicParam3) ? 0x8 : 0;
-	
+
 	MeshMinimumLOD = Properties->ParticleMesh->MinLOD.GetValueForFeatureLevel(FeatureLevel);
 }
 
@@ -290,6 +294,7 @@ void FNiagaraRendererMeshes::GetDynamicMeshElements(const TArray<const FSceneVie
 				PerViewUniformParameters.TransformDataOffset = VFVariables[ENiagaraMeshVFLayout::Transform].GetGPUOffset();
 				PerViewUniformParameters.NormalizedAgeDataOffset = VFVariables[ENiagaraMeshVFLayout::NormalizedAge].GetGPUOffset();
 				PerViewUniformParameters.MaterialRandomDataOffset = VFVariables[ENiagaraMeshVFLayout::MaterialRandom].GetGPUOffset();
+				PerViewUniformParameters.SubImageDataOffset = VFVariables[ENiagaraMeshVFLayout::SubImage].GetGPUOffset();
 				PerViewUniformParameters.MaterialParamDataOffset = VFVariables[ENiagaraMeshVFLayout::DynamicParam0].GetGPUOffset();
 				PerViewUniformParameters.MaterialParam1DataOffset = VFVariables[ENiagaraMeshVFLayout::DynamicParam1].GetGPUOffset();
 				PerViewUniformParameters.MaterialParam2DataOffset = VFVariables[ENiagaraMeshVFLayout::DynamicParam2].GetGPUOffset();
@@ -298,6 +303,8 @@ void FNiagaraRendererMeshes::GetDynamicMeshElements(const TArray<const FSceneVie
 				PerViewUniformParameters.MaterialParamValidMask = MaterialParamValidMask;
 				PerViewUniformParameters.SizeDataOffset = INDEX_NONE;
 				PerViewUniformParameters.DefaultPos = bLocalSpace ? FVector4(0.0f, 0.0f, 0.0f, 1.0f) : FVector4(SceneProxy->GetLocalToWorld().GetOrigin());
+				PerViewUniformParameters.SubImageSize = FVector4(SubImageSize.X, SubImageSize.Y, 1.0f / SubImageSize.X, 1.0f / SubImageSize.Y);
+				PerViewUniformParameters.SubImageBlendMode = bSubImageBlend;
 
 				//Sort particles if needed.
 				CollectorResources.VertexFactory->SetSortedIndices(nullptr, 0xFFFFFFFF);
diff --git a/Engine/Plugins/FX/Niagara/Source/Niagara/Public/NiagaraMeshRendererProperties.h b/Engine/Plugins/FX/Niagara/Source/Niagara/Public/NiagaraMeshRendererProperties.h
--- a/Engine/Plugins/FX/Niagara/Source/Niagara/Public/NiagaraMeshRendererProperties.h
+++ b/Engine/Plugins/FX/Niagara/Source/Niagara/Public/NiagaraMeshRendererProperties.h
@@ -113,6 +113,14 @@ public:
 	UPROPERTY(EditAnywhere, Category = "Mesh Rendering", meta = (EditCondition = "bOverrideMaterials"))
 	TArray<FNiagaraMeshMaterialOverride> OverrideMaterials;
 
+	/** When using SubImage lookups for particles, this variable contains the number of columns in X and the number of rows in Y.*/
+	UPROPERTY(EditAnywhere, Category = "SubUV")
+	FVector2D SubImageSize;
+
+	/** If true, blends the sub-image UV lookup with its next adjacent member using the fractional part of the SubImageIndex float value as the linear interpolation factor.*/
+	UPROPERTY(EditAnywhere, Category = "SubUV", meta = (DisplayName = "Sub UV Blending Enabled"))
+	uint32 bSubImageBlend : 1;
+
 	/** Determines how the mesh orients itself relative to the camera.*/
 	UPROPERTY(EditAnywhere, Category = "Mesh Rendering")
 	ENiagaraMeshFacingMode FacingMode; 
@@ -137,6 +145,10 @@ public:
 	UPROPERTY(EditAnywhere, Category = "Bindings")
 	FNiagaraVariableAttributeBinding ScaleBinding;
 	
+	/** Which attribute should we use for sprite sub-image indexing when generating sprites?*/
+	UPROPERTY(EditAnywhere, Category = "Bindings")
+	FNiagaraVariableAttributeBinding SubImageIndexBinding;
+
 	/** Which attribute should we use for dynamic material parameters when generating instanced meshes?*/
 	UPROPERTY(EditAnywhere, Category = "Bindings")
 	FNiagaraVariableAttributeBinding DynamicMaterialBinding;
diff --git a/Engine/Plugins/FX/Niagara/Source/Niagara/Public/NiagaraRendererMeshes.h b/Engine/Plugins/FX/Niagara/Source/Niagara/Public/NiagaraRendererMeshes.h
--- a/Engine/Plugins/FX/Niagara/Source/Niagara/Public/NiagaraRendererMeshes.h
+++ b/Engine/Plugins/FX/Niagara/Source/Niagara/Public/NiagaraRendererMeshes.h
@@ -41,6 +41,9 @@ private:
 	uint32 bOverrideMaterials : 1;
 	uint32 bSortOnlyWhenTranslucent : 1;
 
+	uint32 bSubImageBlend : 1;
+	FVector2D SubImageSize;
+
 	uint32 MaterialParamValidMask;
 
 	int32 MeshMinimumLOD = 0;
diff --git a/Engine/Plugins/FX/Niagara/Source/NiagaraVertexFactories/Private/NiagaraMeshVertexFactory.cpp b/Engine/Plugins/FX/Niagara/Source/NiagaraVertexFactories/Private/NiagaraMeshVertexFactory.cpp
--- a/Engine/Plugins/FX/Niagara/Source/NiagaraVertexFactories/Private/NiagaraMeshVertexFactory.cpp
+++ b/Engine/Plugins/FX/Niagara/Source/NiagaraVertexFactories/Private/NiagaraMeshVertexFactory.cpp
@@ -29,6 +29,7 @@ public:
 		NiagaraParticleDataMaterialParam1.Bind(ParameterMap, TEXT("NiagaraParticleDataMaterialParam1"));
 		NiagaraParticleDataMaterialParam2.Bind(ParameterMap, TEXT("NiagaraParticleDataMaterialParam2"));
 		NiagaraParticleDataMaterialParam3.Bind(ParameterMap, TEXT("NiagaraParticleDataMaterialParam3"));
+		NiagaraParticleDataSubImage.Bind(ParameterMap, TEXT("NiagaraParticleDataSubImage"));
 
 		FloatDataOffset.Bind(ParameterMap, TEXT("NiagaraFloatDataOffset"));
 		FloatDataStride.Bind(ParameterMap, TEXT("NiagaraFloatDataStride"));
@@ -40,7 +41,6 @@ public:
 		MeshFacingMode.Bind(ParameterMap, TEXT("MeshFacingMode"));
 		SortedIndices.Bind(ParameterMap, TEXT("SortedIndices"));
 		SortedIndicesOffset.Bind(ParameterMap, TEXT("SortedIndicesOffset"));
-
 	}
 
 	virtual void Serialize(FArchive& Ar) override
@@ -57,6 +57,7 @@ public:
 		Ar << NiagaraParticleDataMaterialParam1;
 		Ar << NiagaraParticleDataMaterialParam2;
 		Ar << NiagaraParticleDataMaterialParam3;
+		Ar << NiagaraParticleDataSubImage;
 		Ar << FloatDataOffset;
 		Ar << FloatDataStride;
 
@@ -98,6 +99,7 @@ public:
 		ShaderBindings.Add(NiagaraParticleDataMaterialParam1, NiagaraMeshVF->GetParticleDataFloatSRV());
 		ShaderBindings.Add(NiagaraParticleDataMaterialParam2, NiagaraMeshVF->GetParticleDataFloatSRV());
 		ShaderBindings.Add(NiagaraParticleDataMaterialParam3, NiagaraMeshVF->GetParticleDataFloatSRV());
+		ShaderBindings.Add(NiagaraParticleDataSubImage, NiagaraMeshVF->GetParticleDataFloatSRV());
 
 		ShaderBindings.Add(FloatDataOffset, NiagaraMeshVF->GetFloatDataOffset());
 		ShaderBindings.Add(FloatDataStride, NiagaraMeshVF->GetFloatDataStride());
@@ -122,6 +124,7 @@ private:
 	FShaderResourceParameter NiagaraParticleDataMaterialParam1;
 	FShaderResourceParameter NiagaraParticleDataMaterialParam2;
 	FShaderResourceParameter NiagaraParticleDataMaterialParam3;
+	FShaderResourceParameter NiagaraParticleDataSubImage;
 	FShaderParameter FloatDataOffset;
 	FShaderParameter FloatDataStride;
 
diff --git a/Engine/Plugins/FX/Niagara/Source/NiagaraVertexFactories/Public/NiagaraMeshVertexFactory.h b/Engine/Plugins/FX/Niagara/Source/NiagaraVertexFactories/Public/NiagaraMeshVertexFactory.h
--- a/Engine/Plugins/FX/Niagara/Source/NiagaraVertexFactories/Public/NiagaraMeshVertexFactory.h
+++ b/Engine/Plugins/FX/Niagara/Source/NiagaraVertexFactories/Public/NiagaraMeshVertexFactory.h
@@ -41,6 +41,7 @@ BEGIN_GLOBAL_SHADER_PARAMETER_STRUCT(FNiagaraMeshUniformParameters, NIAGARAVERTE
 	SHADER_PARAMETER(int, ScaleDataOffset)
 	SHADER_PARAMETER(int, SizeDataOffset)
 	SHADER_PARAMETER(uint32, MaterialParamValidMask)
+	SHADER_PARAMETER(int, SubImageDataOffset)
 	SHADER_PARAMETER(int, MaterialParamDataOffset)
 	SHADER_PARAMETER(int, MaterialParam1DataOffset)
 	SHADER_PARAMETER(int, MaterialParam2DataOffset)
@@ -48,6 +49,7 @@ BEGIN_GLOBAL_SHADER_PARAMETER_STRUCT(FNiagaraMeshUniformParameters, NIAGARAVERTE
 	SHADER_PARAMETER(int, NormalizedAgeDataOffset)
 	SHADER_PARAMETER(int, MaterialRandomDataOffset)
 	SHADER_PARAMETER(FVector4, DefaultPos)
+	SHADER_PARAMETER(int, SubImageBlendMode)
 END_GLOBAL_SHADER_PARAMETER_STRUCT()
 
 typedef TUniformBufferRef<FNiagaraMeshUniformParameters> FNiagaraMeshUniformBufferRef;
