Cycling through issues with the issue button in the Niagara selection view now cycles through the specific emitter when multi-selecting emitters.

diff --git a/Engine/Plugins/FX/Niagara/Source/NiagaraEditor/Private/ViewModels/Stack/NiagaraStackViewModel.cpp b/Engine/Plugins/FX/Niagara/Source/NiagaraEditor/Private/ViewModels/Stack/NiagaraStackViewModel.cpp
--- a/Engine/Plugins/FX/Niagara/Source/NiagaraEditor/Private/ViewModels/Stack/NiagaraStackViewModel.cpp
+++ b/Engine/Plugins/FX/Niagara/Source/NiagaraEditor/Private/ViewModels/Stack/NiagaraStackViewModel.cpp
@@ -780,22 +780,47 @@ UNiagaraStackEntry* UNiagaraStackViewModel::GetCurrentFocusedIssue() const
 {
 	if (CurrentIssueCycleIndex >= 0)
 	{
-		const TArray<UNiagaraStackEntry*>& Issues = RootEntry->GetAllChildrenWithIssues();
-		return Issues[CurrentIssueCycleIndex];
+		UNiagaraStackEntry* CyclingRootEntry = CyclingIssuesForTopLevel.Pin()->RootEntry.Get();
+		if (CyclingRootEntry != nullptr)
+		{
+			const TArray<UNiagaraStackEntry*>& Issues = CyclingRootEntry->GetAllChildrenWithIssues();
+			return Issues[CurrentIssueCycleIndex];
+		}
 	}
 	
 	return nullptr;
 }
 
-void UNiagaraStackViewModel::OnCycleThroughIssues()
+void UNiagaraStackViewModel::OnCycleThroughIssues(TSharedPtr<UNiagaraStackViewModel::FTopLevelViewModel> TopLevelToCycle)
 {
-	if (RootEntry == nullptr)
+	if (RootEntries.Num() == 0)
+	{
+		CurrentIssueCycleIndex = -1;
+		CyclingIssuesForTopLevel.Reset();
+		return;
+	}
+
+	if (CyclingIssuesForTopLevel.IsValid() && CyclingIssuesForTopLevel != TopLevelToCycle)
+	{
+		CurrentIssueCycleIndex = -1;
+	}
+
+	CyclingIssuesForTopLevel = TopLevelToCycle;
+
+	UNiagaraStackEntry* CyclingRootEntry = nullptr;
+	if (CyclingIssuesForTopLevel.IsValid())
+	{
+		CyclingRootEntry = CyclingIssuesForTopLevel.Pin()->RootEntry.Get();
+	}
+	
+	if (CyclingRootEntry == nullptr)
 	{
 		CurrentIssueCycleIndex = -1;
+		CyclingIssuesForTopLevel.Reset();
 		return;
 	}
 
-	const TArray<UNiagaraStackEntry*>& Issues = RootEntry->GetAllChildrenWithIssues();
+	const TArray<UNiagaraStackEntry*>& Issues = CyclingRootEntry->GetAllChildrenWithIssues();
 	if (Issues.Num() > 0)
 	{
 		++CurrentIssueCycleIndex;
diff --git a/Engine/Plugins/FX/Niagara/Source/NiagaraEditor/Public/ViewModels/Stack/NiagaraStackViewModel.h b/Engine/Plugins/FX/Niagara/Source/NiagaraEditor/Public/ViewModels/Stack/NiagaraStackViewModel.h
--- a/Engine/Plugins/FX/Niagara/Source/NiagaraEditor/Public/ViewModels/Stack/NiagaraStackViewModel.h
+++ b/Engine/Plugins/FX/Niagara/Source/NiagaraEditor/Public/ViewModels/Stack/NiagaraStackViewModel.h
@@ -123,7 +123,7 @@ public:
 	UNiagaraStackEntry* GetCurrentFocusedEntry();
 	void AddSearchScrollOffset(int NumberOfSteps);
 
-	void OnCycleThroughIssues();
+	void OnCycleThroughIssues(TSharedPtr<FTopLevelViewModel> TopLevelToCycle);
 	UNiagaraStackEntry* GetCurrentFocusedIssue() const;
 
 	void GetPathForEntry(UNiagaraStackEntry* Entry, TArray<UNiagaraStackEntry*>& EntryPath) const;
@@ -201,6 +201,7 @@ private:
 	bool bRefreshPending;
 	bool bHasIssues;
 	int32 CurrentIssueCycleIndex;
+	TWeakPtr<FTopLevelViewModel> CyclingIssuesForTopLevel;
 
 	bool bUsesTopLevelViewModels;
 	TArray<TSharedRef<FTopLevelViewModel>> TopLevelViewModels;
diff --git a/Engine/Plugins/FX/Niagara/Source/NiagaraEditorWidgets/Private/SNiagaraOverviewStackNode.cpp b/Engine/Plugins/FX/Niagara/Source/NiagaraEditorWidgets/Private/SNiagaraOverviewStackNode.cpp
--- a/Engine/Plugins/FX/Niagara/Source/NiagaraEditorWidgets/Private/SNiagaraOverviewStackNode.cpp
+++ b/Engine/Plugins/FX/Niagara/Source/NiagaraEditorWidgets/Private/SNiagaraOverviewStackNode.cpp
@@ -6,6 +6,7 @@
 #include "ViewModels/NiagaraEmitterViewModel.h"
 #include "ViewModels/NiagaraEmitterHandleViewModel.h"
 #include "ViewModels/NiagaraSystemSelectionViewModel.h"
+#include "ViewModels/Stack/NiagaraStackItemGroup.h"
 #include "ViewModels/Stack/NiagaraStackModuleItem.h"
 #include "ViewModels/Stack/NiagaraStackViewModel.h"
 #include "SNiagaraOverviewStack.h"
@@ -147,40 +148,38 @@ TSharedRef<SWidget> SNiagaraOverviewStackNode::CreateTitleRightWidget()
 
 FReply SNiagaraOverviewStackNode::OnCycleThroughIssues()
 {
-	TSharedPtr<FNiagaraEmitterHandleViewModel> EmitterHandleViewModel = EmitterHandleViewModelWeak.Pin();
-	if (EmitterHandleViewModel.IsValid())
+	const TArray<UNiagaraStackEntry*>& ChildrenWithIssues = StackViewModel->GetRootEntry()->GetAllChildrenWithIssues();
+	if (ChildrenWithIssues.Num() > 0)
 	{
-		const TArray<UNiagaraStackEntry*>& ChildrenWithIssues = StackViewModel->GetRootEntry()->GetAllChildrenWithIssues();
-		if (ChildrenWithIssues.Num() > 0)
+		++CurrentIssueIndex;
+
+		if (CurrentIssueIndex >= ChildrenWithIssues.Num())
 		{
-			++CurrentIssueIndex;
+			CurrentIssueIndex = 0;
+		}
 
-			if (CurrentIssueIndex >= ChildrenWithIssues.Num())
+		if (ChildrenWithIssues.IsValidIndex(CurrentIssueIndex))
+		{
+			UNiagaraStackEntry* ChildIssue = ChildrenWithIssues[CurrentIssueIndex];
+			UNiagaraStackEntry* ChildToSelect = Cast<UNiagaraStackModuleItem>(ChildIssue);
+			if (ChildToSelect == nullptr)
 			{
-				CurrentIssueIndex = 0;
+				ChildToSelect = ChildIssue->GetTypedOuter<UNiagaraStackModuleItem>();
 			}
 
-			if (ChildrenWithIssues.IsValidIndex(CurrentIssueIndex))
+			if (ChildToSelect == nullptr)
 			{
-				UNiagaraStackEntry* ChildToSelect = ChildrenWithIssues[CurrentIssueIndex];
-				UNiagaraStackModuleItem* ModuleToSelect = Cast<UNiagaraStackModuleItem>(ChildToSelect);
-				if (ModuleToSelect == nullptr)
-				{
-					ModuleToSelect = ChildToSelect->GetTypedOuter<UNiagaraStackModuleItem>();
-				}
-
-				if (ModuleToSelect != nullptr)
-				{
-					OverviewSelectionViewModel->UpdateSelectedEntries(TArray<UNiagaraStackEntry*> { ModuleToSelect },
-						TArray<UNiagaraStackEntry*>(), true /* bClearCurrentSelection */ );
-				}
+				ChildToSelect = Cast<UNiagaraStackItemGroup>(ChildIssue);
+			}
+			
+			if (ChildToSelect != nullptr)
+			{
+				OverviewSelectionViewModel->UpdateSelectedEntries(TArray<UNiagaraStackEntry*> { ChildToSelect }, TArray<UNiagaraStackEntry*>(), true /* bClearCurrentSelection */ );
 			}
 		}
-
-		return FReply::Handled();
 	}
 
-	return FReply::Unhandled();
+	return FReply::Handled();
 }
 
 TSharedRef<SWidget> SNiagaraOverviewStackNode::CreateThumbnailWidget(float InThumbnailSize, FRendererPreviewData* InData)
diff --git a/Engine/Plugins/FX/Niagara/Source/NiagaraEditorWidgets/Private/SNiagaraStack.cpp b/Engine/Plugins/FX/Niagara/Source/NiagaraEditorWidgets/Private/SNiagaraStack.cpp
--- a/Engine/Plugins/FX/Niagara/Source/NiagaraEditorWidgets/Private/SNiagaraStack.cpp
+++ b/Engine/Plugins/FX/Niagara/Source/NiagaraEditorWidgets/Private/SNiagaraStack.cpp
@@ -69,6 +69,7 @@ public:
 		EmitterHandleViewModel = InEmitterHandleViewModel;
 		OnCycleThroughIssues = InArgs._OnCycleThroughIssues;
 		StackViewModel = InStackViewModel;
+		TopLevelViewModel = StackViewModel->GetTopLevelViewModelForEntry(*InRootEntry);
 
 		ChildSlot
 		[
@@ -227,13 +228,14 @@ private:
 
 	FReply OnIssueIconClicked() const
 	{
-		StackViewModel->OnCycleThroughIssues();
+		StackViewModel->OnCycleThroughIssues(TopLevelViewModel);
 		OnCycleThroughIssues.ExecuteIfBound();
 		return FReply::Handled();
 	}
 
 private:
 	TSharedPtr<FNiagaraEmitterHandleViewModel> EmitterHandleViewModel;
+	TSharedPtr<UNiagaraStackViewModel::FTopLevelViewModel> TopLevelViewModel;
 
 	TSharedPtr<SInlineEditableTextBlock> EmitterNameTextBlock;
 	FSimpleDelegate OnCycleThroughIssues;
@@ -695,11 +697,12 @@ TSharedRef<ITableRow> SNiagaraStack::OnGenerateRowForTopLevelObject(TSharedRef<U
 			.VAlign(VAlign_Center)
 			.Padding(4, 0, 2, 0)
 			[
-				SNew(SNiagaraStackIssueIcon, Item->SystemViewModel->GetSystemStackViewModel(), Item->RootEntry.Get())
+				SNew(SNiagaraStackIssueIcon, StackViewModel, Item->RootEntry.Get())
 				.Visibility(this, &SNiagaraStack::GetIssueIconVisibility)
+				.OnClicked(this, &SNiagaraStack::OnCycleThroughSystemIssues, Item->SystemViewModel)
 			];
 	}
-	else if(Item->EmitterHandleViewModel.IsValid())
+	else if (Item->EmitterHandleViewModel.IsValid())
 	{
 		Content = SNew(SNiagaraStackEmitterHeader, Item->EmitterHandleViewModel.ToSharedRef(), Item->RootEntry.Get(), StackViewModel)
 			.IssueIconVisibility(this, &SNiagaraStack::GetIssueIconVisibility)
@@ -1056,6 +1059,19 @@ EVisibility SNiagaraStack::GetIssueIconVisibility() const
 	return StackViewModel->HasIssues() ? EVisibility::Visible : EVisibility::Collapsed;
 }
 
+FReply SNiagaraStack::OnCycleThroughSystemIssues(TSharedPtr<FNiagaraSystemViewModel> SystemViewModel)
+{
+	UNiagaraStackEntry* RootEntry = SystemViewModel->GetSystemStackViewModel()->GetRootEntry();
+	if (RootEntry != nullptr)
+	{
+		TSharedPtr<UNiagaraStackViewModel::FTopLevelViewModel> TopLevelViewModel = StackViewModel->GetTopLevelViewModelForEntry(*RootEntry);
+		StackViewModel->OnCycleThroughIssues(TopLevelViewModel);
+		OnCycleThroughIssues();
+	}
+
+	return FReply::Handled();
+}
+
 void SNiagaraStack::OnCycleThroughIssues()
 {
 	UNiagaraStackEntry* StackEntry = StackViewModel->GetCurrentFocusedIssue();
diff --git a/Engine/Plugins/FX/Niagara/Source/NiagaraEditorWidgets/Private/SNiagaraStack.h b/Engine/Plugins/FX/Niagara/Source/NiagaraEditorWidgets/Private/SNiagaraStack.h
--- a/Engine/Plugins/FX/Niagara/Source/NiagaraEditorWidgets/Private/SNiagaraStack.h
+++ b/Engine/Plugins/FX/Niagara/Source/NiagaraEditorWidgets/Private/SNiagaraStack.h
@@ -112,6 +112,7 @@ private:
 
 	EVisibility GetIssueIconVisibility() const;
 
+	FReply OnCycleThroughSystemIssues(TSharedPtr<FNiagaraSystemViewModel> SystemViewModel);
 	void OnCycleThroughIssues();
 
 private:
