Made skel mesh DIs cached attach parent a weak ptr to avoid possible aliasing issue.

diff --git a/Engine/Plugins/FX/Niagara/Source/Niagara/Classes/NiagaraDataInterfaceSkeletalMesh.h b/Engine/Plugins/FX/Niagara/Source/Niagara/Classes/NiagaraDataInterfaceSkeletalMesh.h
--- a/Engine/Plugins/FX/Niagara/Source/Niagara/Classes/NiagaraDataInterfaceSkeletalMesh.h
+++ b/Engine/Plugins/FX/Niagara/Source/Niagara/Classes/NiagaraDataInterfaceSkeletalMesh.h
@@ -365,7 +365,7 @@ struct FNDISkeletalMesh_InstanceData
 	FNiagaraParameterDirectBinding<UObject*> UserParamBinding;
 
 	//Always reset the DI when the attach parent changes.
-	USceneComponent* CachedAttachParent;
+	TWeakObjectPtr<USceneComponent> CachedAttachParent;
 
 	UObject* CachedUserParam;
 
diff --git a/Engine/Plugins/FX/Niagara/Source/Niagara/Private/NiagaraDataInterfaceSkeletalMesh.cpp b/Engine/Plugins/FX/Niagara/Source/Niagara/Private/NiagaraDataInterfaceSkeletalMesh.cpp
--- a/Engine/Plugins/FX/Niagara/Source/Niagara/Private/NiagaraDataInterfaceSkeletalMesh.cpp
+++ b/Engine/Plugins/FX/Niagara/Source/Niagara/Private/NiagaraDataInterfaceSkeletalMesh.cpp
@@ -1393,7 +1393,7 @@ bool FNDISkeletalMesh_InstanceData::ResetRequired(UNiagaraDataInterfaceSkeletalM
 	}
 
 	//Detect and reset on any attachment change.
-	if (Comp->GetAttachParent() != CachedAttachParent)
+	if (CachedAttachParent.IsValid() && Comp->GetAttachParent() != CachedAttachParent.Get())
 	{
 		return true;
 	}
