Fix for crash when clearing the asset ref on a NiagaraComponent in a BP.

diff --git a/Engine/Plugins/FX/Niagara/Source/Niagara/Private/NiagaraComponent.cpp b/Engine/Plugins/FX/Niagara/Source/Niagara/Private/NiagaraComponent.cpp
--- a/Engine/Plugins/FX/Niagara/Source/Niagara/Private/NiagaraComponent.cpp
+++ b/Engine/Plugins/FX/Niagara/Source/Niagara/Private/NiagaraComponent.cpp
@@ -1625,6 +1625,7 @@ void UNiagaraComponent::PreEditChange(FProperty* PropertyAboutToChange)
 	if (PropertyAboutToChange != nullptr && PropertyAboutToChange->GetFName() == GET_MEMBER_NAME_CHECKED(UNiagaraComponent, Asset) && Asset != nullptr)
 	{
 		Asset->GetExposedParameters().RemoveOnChangedHandler(AssetExposedParametersChangedHandle);
+		DestroyInstance();
 	}
 }
 
