Unshelved from pending changelist '10969490':

Fix transform sampled skel mesh rotations from component to worldspace.
Save one transform op in building socket sampling tables by avoiding sampling bone in component space and transforming back to world space.
Change sampled rotation interpolation from lerp to slerp.

diff --git a/Engine/Plugins/FX/Niagara/Source/Niagara/Classes/NiagaraDataInterface.h b/Engine/Plugins/FX/Niagara/Source/Niagara/Classes/NiagaraDataInterface.h
--- a/Engine/Plugins/FX/Niagara/Source/Niagara/Classes/NiagaraDataInterface.h
+++ b/Engine/Plugins/FX/Niagara/Source/Niagara/Classes/NiagaraDataInterface.h
@@ -19,14 +19,16 @@ struct FNiagaraDataInterfaceProxy;
 
 struct FNDITransformHandlerNoop
 {
-	FORCEINLINE void TransformPosition(FVector& V, const FMatrix& M) {  }
+	FORCEINLINE void TransformPosition(FVector& V, const FMatrix& M) { }
 	FORCEINLINE void TransformVector(FVector& V, const FMatrix& M) { }
+	FORCEINLINE void TransformRotation(FQuat& Q, const FMatrix& M) { }
 };
 
 struct FNDITransformHandler
 {
 	FORCEINLINE void TransformPosition(FVector& P, const FMatrix& M) { P = M.TransformPosition(P); }
 	FORCEINLINE void TransformVector(FVector& V, const FMatrix& M) { V = M.TransformVector(V).GetUnsafeNormal3(); }
+	FORCEINLINE void TransformRotation(FQuat& Q, const FMatrix& M) { Q = M.ToQuat() * Q; }
 };
 
 //////////////////////////////////////////////////////////////////////////
diff --git a/Engine/Plugins/FX/Niagara/Source/Niagara/Private/NDISkeletalMesh_BoneSampling.cpp b/Engine/Plugins/FX/Niagara/Source/Niagara/Private/NDISkeletalMesh_BoneSampling.cpp
--- a/Engine/Plugins/FX/Niagara/Source/Niagara/Private/NDISkeletalMesh_BoneSampling.cpp
+++ b/Engine/Plugins/FX/Niagara/Source/Niagara/Private/NDISkeletalMesh_BoneSampling.cpp
@@ -526,10 +526,12 @@ void UNiagaraDataInterfaceSkeletalMesh::GetSkinnedBoneData(FVectorVMContext& Con
 			if (Output.bNeedsRotation)
 			{
 				FQuat Rotation = CurrSocketTransform.GetRotation();
+				TransformHandler.TransformRotation(Rotation, Transform);
 				if (bInterpolated::Value)
 				{
 					FQuat PrevRotation = PrevSocketTransform.GetRotation();
-					Rotation = FMath::Lerp(PrevRotation, Rotation, Interp);
+					TransformHandler.TransformRotation(PrevRotation, Transform);
+					Rotation = FQuat::Slerp(PrevRotation, Rotation, Interp);
 				}
 
 				Output.SetRotation(Rotation);
@@ -550,10 +552,12 @@ void UNiagaraDataInterfaceSkeletalMesh::GetSkinnedBoneData(FVectorVMContext& Con
 			if (Output.bNeedsRotation)
 			{
 				FQuat Rotation = SkinningHandler.GetSkinnedBoneRotation(Accessor, Bone);
+				TransformHandler.TransformRotation(Rotation, Transform);
 				if (bInterpolated::Value)
 				{
 					FQuat PrevRotation = SkinningHandler.GetSkinnedBonePreviousRotation(Accessor, Bone);
-					Rotation = FMath::Lerp(PrevRotation, Rotation, Interp);
+					TransformHandler.TransformRotation(PrevRotation, Transform);
+					Rotation = FQuat::Slerp(PrevRotation, Rotation, Interp);
 				}
 
 				Output.SetRotation(Rotation);
diff --git a/Engine/Plugins/FX/Niagara/Source/Niagara/Private/NiagaraDataInterfaceSkeletalMesh.cpp b/Engine/Plugins/FX/Niagara/Source/Niagara/Private/NiagaraDataInterfaceSkeletalMesh.cpp
--- a/Engine/Plugins/FX/Niagara/Source/Niagara/Private/NiagaraDataInterfaceSkeletalMesh.cpp
+++ b/Engine/Plugins/FX/Niagara/Source/Niagara/Private/NiagaraDataInterfaceSkeletalMesh.cpp
@@ -1500,7 +1500,10 @@ void FNDISkeletalMesh_InstanceData::UpdateSpecificSocketTransforms()
 	{
 		for (int32 i = 0; i < SpecificSockets.Num(); ++i)
 		{
-			WriteBuffer[i] = SkelComp->GetSocketTransform(SpecificSockets[i], RTS_Component);
+			FTransform SocketTransform;
+			int32 BoneIdx;
+			SkelComp->GetSocketInfoByName(SpecificSockets[i], SocketTransform, BoneIdx);
+			WriteBuffer[i] = SocketTransform * SkelComp->GetBoneTransform(BoneIdx, FTransform::Identity);
 		}
 	}
 	else if ( Mesh != nullptr )
