Integrated 10944364 : Disabled groupshared logic for Niagara GPU emitters on Vulkan.

diff --git a/Engine/Plugins/FX/Niagara/Shaders/Private/NiagaraEmitterInstanceShader.usf b/Engine/Plugins/FX/Niagara/Shaders/Private/NiagaraEmitterInstanceShader.usf
--- a/Engine/Plugins/FX/Niagara/Shaders/Private/NiagaraEmitterInstanceShader.usf
+++ b/Engine/Plugins/FX/Niagara/Shaders/Private/NiagaraEmitterInstanceShader.usf
@@ -470,7 +470,7 @@ int rand_int(int x, int Seed1, int Seed2, int Seed3)
 	}
 
 #define USE_WAVE_INTRISICS (PS4_PROFILE || XBOXONE_PROFILE)
-#define USE_GROUP_SHARED ((THREADGROUP_SIZE == 64 || THREADGROUP_SIZE == 32)  && !USE_WAVE_INTRISICS)
+#define USE_GROUP_SHARED ((THREADGROUP_SIZE == 64 || THREADGROUP_SIZE == 32)  && !USE_WAVE_INTRISICS && !VULKAN_PROFILE_SM5 && !VULKAN_PROFILE)
 
 #if USE_GROUP_SHARED
 	#if THREADGROUP_SIZE == 64
