- When we destroy an instance of a GPU system on the game thread ensure any pending ticks are removed from the render thread (i.e. instance batcher)
- This removes the need for various deferred deletion queues as we are now guaranteed that resources such as data sets, data interfaces, etc, will no longer be in use by a pending tick

#rnx

diff --git a/Engine/Plugins/FX/Niagara/Source/Niagara/Classes/NiagaraEmitterInstanceBatcher.h b/Engine/Plugins/FX/Niagara/Source/Niagara/Classes/NiagaraEmitterInstanceBatcher.h
--- a/Engine/Plugins/FX/Niagara/Source/Niagara/Classes/NiagaraEmitterInstanceBatcher.h
+++ b/Engine/Plugins/FX/Niagara/Source/Niagara/Classes/NiagaraEmitterInstanceBatcher.h
@@ -66,16 +66,17 @@ public:
 	static const FName Name;
 	virtual FFXSystemInterface* GetInterface(const FName& InName) override;
 
+	/** Notification that the InstanceID has been removed. */
+	void InstanceDeallocated_RenderThread(const FNiagaraSystemInstanceID InstanceID);
+
 	// The batcher assumes ownership of the data here.
+
 	void GiveSystemTick_RenderThread(FNiagaraGPUSystemTick& Tick);
 
-	void GiveEmitterContextToDestroy_RenderThread(FNiagaraComputeExecutionContext* Context);
-	void GiveDataSetToDestroy_RenderThread(FNiagaraDataSet* DataSet);
+	/** Called to release GPU instance counts that the batcher is tracking. */
+	void ReleaseInstanceCounts_RenderThread(FNiagaraComputeExecutionContext* ExecContext, FNiagaraDataSet* DataSet);
 
-	void EnqueueDeferredDeletesForDI_RenderThread(TSharedPtr<FNiagaraDataInterfaceProxy, ESPMode::ThreadSafe> Proxy)
-	{
-		DIProxyDeferredDeletes_RT.Add(MoveTemp(Proxy));
-	}
+	NIAGARA_API void EnqueueDeferredDeletesForDI_RenderThread(TSharedPtr<FNiagaraDataInterfaceProxy, ESPMode::ThreadSafe> Proxy);
 
 #if WITH_EDITOR
 	virtual void Suspend() override {}
diff --git a/Engine/Plugins/FX/Niagara/Source/Niagara/Private/NiagaraEmitterInstance.cpp b/Engine/Plugins/FX/Niagara/Source/Niagara/Private/NiagaraEmitterInstance.cpp
--- a/Engine/Plugins/FX/Niagara/Source/Niagara/Private/NiagaraEmitterInstance.cpp
+++ b/Engine/Plugins/FX/Niagara/Source/Niagara/Private/NiagaraEmitterInstance.cpp
@@ -137,36 +137,21 @@ FNiagaraEmitterInstance::~FNiagaraEmitterInstance()
 	if (GPUExecContext != nullptr)
 	{
 		/** We defer the deletion of the particle dataset and the compute context to the RT to be sure all in-flight RT commands have finished using it.*/
-		NiagaraEmitterInstanceBatcher* B = Batcher && !Batcher->IsPendingKill() ? Batcher : nullptr;
-		FNiagaraComputeExecutionContext* Context = GPUExecContext;
-		FNiagaraDataSet* DataSet = ParticleDataSet;
+		NiagaraEmitterInstanceBatcher* Batcher_RT = Batcher && !Batcher->IsPendingKill() ? Batcher : nullptr;
 		ENQUEUE_RENDER_COMMAND(FDeleteContextCommand)(
-			[B, Context, DataSet](FRHICommandListImmediate& RHICmdList)
+			[Batcher_RT, ExecContext=GPUExecContext, DataSet= ParticleDataSet](FRHICommandListImmediate& RHICmdList)
 			{
-				if (Context)
+				if ( Batcher_RT != nullptr )
 				{
-					if (B)
-					{
-						B->GiveEmitterContextToDestroy_RenderThread(Context);
-					}
-					else
-					{
-						delete Context;
-					}
+					Batcher_RT->ReleaseInstanceCounts_RenderThread(ExecContext, DataSet);
 				}
-
-				//TODO: deleting these on the RT shouldn't be needed any more.
-				if (DataSet)
+				if ( ExecContext != nullptr )
 				{
-					if (B)
-					{
-						B->GiveDataSetToDestroy_RenderThread(DataSet);
-					}
-					else
-					{
-						delete DataSet;
-					}
-
+					delete ExecContext;
+				}
+				if ( DataSet != nullptr )
+				{
+					delete DataSet;
 				}
 			}
 		);
diff --git a/Engine/Plugins/FX/Niagara/Source/Niagara/Private/NiagaraEmitterInstanceBatcher.cpp b/Engine/Plugins/FX/Niagara/Source/Niagara/Private/NiagaraEmitterInstanceBatcher.cpp
--- a/Engine/Plugins/FX/Niagara/Source/Niagara/Private/NiagaraEmitterInstanceBatcher.cpp
+++ b/Engine/Plugins/FX/Niagara/Source/Niagara/Private/NiagaraEmitterInstanceBatcher.cpp
@@ -91,6 +91,25 @@ NiagaraEmitterInstanceBatcher::~NiagaraEmitterInstanceBatcher()
 	ParticleSortBuffers.ReleaseRHI();
 }
 
+void NiagaraEmitterInstanceBatcher::InstanceDeallocated_RenderThread(const FNiagaraSystemInstanceID InstanceID)
+{
+	int iTick = 0;
+	while ( iTick < Ticks_RT.Num() )
+	{
+		FNiagaraGPUSystemTick& Tick = Ticks_RT[iTick];
+		if (Tick.SystemInstanceID == InstanceID)
+		{
+			//-OPT: Since we can't RemoveAtSwap (due to ordering issues) if may be better to not remove and flag as dead
+			Tick.Destroy();
+			Ticks_RT.RemoveAt(iTick);
+		}
+		else
+		{
+			++iTick;
+		}
+	}
+}
+
 void NiagaraEmitterInstanceBatcher::GiveSystemTick_RenderThread(FNiagaraGPUSystemTick& Tick)
 {
 	check(IsInRenderingThread());
@@ -124,16 +143,23 @@ void NiagaraEmitterInstanceBatcher::GiveSystemTick_RenderThread(FNiagaraGPUSyste
 	Ticks_RT.Add(Tick);
 }
 
-void NiagaraEmitterInstanceBatcher::GiveEmitterContextToDestroy_RenderThread(FNiagaraComputeExecutionContext* Context)
+void NiagaraEmitterInstanceBatcher::ReleaseInstanceCounts_RenderThread(FNiagaraComputeExecutionContext* ExecContext, FNiagaraDataSet* DataSet)
 {
 	LLM_SCOPE(ELLMTag::Niagara);
-	ContextsToDestroy_RT.Add(Context);
+
+	if ( ExecContext != nullptr )
+	{
+		GPUInstanceCounterManager.FreeEntry(ExecContext->EmitterInstanceReadback.GPUCountOffset);
+	}
+	if ( DataSet != nullptr )
+	{
+		DataSet->ReleaseGPUInstanceCounts(GPUInstanceCounterManager);
+	}
 }
 
-void NiagaraEmitterInstanceBatcher::GiveDataSetToDestroy_RenderThread(FNiagaraDataSet* DataSet)
+void NiagaraEmitterInstanceBatcher::EnqueueDeferredDeletesForDI_RenderThread(TSharedPtr<FNiagaraDataInterfaceProxy, ESPMode::ThreadSafe> Proxy)
 {
-	LLM_SCOPE(ELLMTag::Niagara);
-	DataSetsToDestroy_RT.Add(DataSet);
+	DIProxyDeferredDeletes_RT.Add(MoveTemp(Proxy));
 }
 
 void NiagaraEmitterInstanceBatcher::FinishDispatches()
diff --git a/Engine/Plugins/FX/Niagara/Source/Niagara/Private/NiagaraSystemInstance.cpp b/Engine/Plugins/FX/Niagara/Source/Niagara/Private/NiagaraSystemInstance.cpp
--- a/Engine/Plugins/FX/Niagara/Source/Niagara/Private/NiagaraSystemInstance.cpp
+++ b/Engine/Plugins/FX/Niagara/Source/Niagara/Private/NiagaraSystemInstance.cpp
@@ -496,6 +496,19 @@ bool FNiagaraSystemInstance::DeallocateSystemInstance(TUniquePtr< FNiagaraSystem
 			SystemInstanceAllocation->UnbindParameters();
 		}
 
+		// If we have active GPU emitters make sure we remove any pending ticks from the RT
+		NiagaraEmitterInstanceBatcher* InstanceBatcher = SystemInstanceAllocation->GetBatcher();
+		if (SystemInstanceAllocation->bHasGPUEmitters && NiagaraSupportsComputeShaders(InstanceBatcher->GetShaderPlatform()))
+		{
+			ENQUEUE_RENDER_COMMAND(NiagaraRemoveGPUSystem)
+			(
+				[InstanceBatcher, InstanceID=SystemInstanceAllocation->GetId()](FRHICommandListImmediate& RHICmdList) mutable
+				{
+					InstanceBatcher->InstanceDeallocated_RenderThread(InstanceID);
+				}
+			);
+		}
+		
 		// Queue deferred deletion from the WorldManager
 		FNiagaraWorldManager* WorldManager = SystemInstanceAllocation->GetWorldManager();
 		check(WorldManager != nullptr);
diff --git a/Engine/Plugins/FX/Niagara/Source/Niagara/Private/NiagaraWorldManager.cpp b/Engine/Plugins/FX/Niagara/Source/Niagara/Private/NiagaraWorldManager.cpp
--- a/Engine/Plugins/FX/Niagara/Source/Niagara/Private/NiagaraWorldManager.cpp
+++ b/Engine/Plugins/FX/Niagara/Source/Niagara/Private/NiagaraWorldManager.cpp
@@ -317,7 +317,7 @@ void FNiagaraWorldManager::DestroySystemInstance(TUniquePtr<FNiagaraSystemInstan
 {
 	check(IsInGameThread());
 	check(InPtr != nullptr);
-	DeferredDeletionQueue[DeferredDeletionQueueIndex].Queue.Emplace(MoveTemp(InPtr));
+	DeferredDeletionQueue.Emplace(MoveTemp(InPtr));
 }
 
 void FNiagaraWorldManager::OnBatcherDestroyed_Internal(NiagaraEmitterInstanceBatcher* InBatcher)
@@ -326,11 +326,7 @@ void FNiagaraWorldManager::OnBatcherDestroyed_Internal(NiagaraEmitterInstanceBat
 	// This is required because the batcher is accessed in FNiagaraEmitterInstance::~FNiagaraEmitterInstance
 	if (World && World->FXSystem && World->FXSystem->GetInterface(NiagaraEmitterInstanceBatcher::Name) == InBatcher)
 	{
-		for ( int32 i=0; i < NumDeferredQueues; ++i)
-		{
-			DeferredDeletionQueue[DeferredDeletionQueueIndex].Fence.Wait();
-			DeferredDeletionQueue[i].Queue.Empty();
-		}
+		DeferredDeletionQueue.Empty();
 	}
 }
 
@@ -348,11 +344,7 @@ void FNiagaraWorldManager::OnWorldCleanup(bool bSessionEnded, bool bCleanupResou
 	}
 	CleanupParameterCollections();
 
-	for ( int32 i=0; i < NumDeferredQueues; ++i)
-	{
-		DeferredDeletionQueue[i].Fence.Wait();
-		DeferredDeletionQueue[i].Queue.Empty();
-	}
+	DeferredDeletionQueue.Empty();
 
 	ScalabilityManagers.Empty();
 }
@@ -470,23 +462,9 @@ void FNiagaraWorldManager::PostActorTick(float DeltaSeconds)
 	bCachedPlayerViewLocationsValid = false;
 	CachedPlayerViewLocations.Reset();
 
-	// Enqueue fence for deferred deletion if we need to wait on anything
-	if (DeferredDeletionQueue[DeferredDeletionQueueIndex].Queue.Num() > 0)
-	{
-		DeferredDeletionQueue[DeferredDeletionQueueIndex].Fence.BeginFence();
-	}
-
-	// Remove instances from oldest frame making sure they aren't in use on the RT
-	DeferredDeletionQueueIndex = (DeferredDeletionQueueIndex + 1) % NumDeferredQueues;
-	if (DeferredDeletionQueue[DeferredDeletionQueueIndex].Queue.Num() > 0)
-	{
-		if (!DeferredDeletionQueue[DeferredDeletionQueueIndex].Fence.IsFenceComplete())
-		{
-			SCOPE_CYCLE_COUNTER(STAT_NiagaraWorldManWaitOnRender);
-			DeferredDeletionQueue[DeferredDeletionQueueIndex].Fence.Wait();
-		}
-		DeferredDeletionQueue[DeferredDeletionQueueIndex].Queue.Empty();
-	}
+	// Delete any instances that were pending deletion
+	//-TODO: This could be done after each system sim has run
+	DeferredDeletionQueue.Empty();
 
 	// Update tick groups
 	for (FNiagaraWorldManagerTickFunction& TickFunc : TickFunctions )
diff --git a/Engine/Plugins/FX/Niagara/Source/Niagara/Public/NiagaraWorldManager.h b/Engine/Plugins/FX/Niagara/Source/Niagara/Public/NiagaraWorldManager.h
--- a/Engine/Plugins/FX/Niagara/Source/Niagara/Public/NiagaraWorldManager.h
+++ b/Engine/Plugins/FX/Niagara/Source/Niagara/Public/NiagaraWorldManager.h
@@ -218,20 +218,11 @@ private:
 
 	UNiagaraComponentPool* ComponentPool;
 
-	/** Generated data used by data interfaces*/
+	/** Generated data used by data interfaces */
 	FNDI_SkeletalMesh_GeneratedData SkeletalMeshGeneratedData;
 
-	// Deferred deletion queue for system instances
-	// We need to make sure that any enqueued GPU ticks have been processed before we remove the system instances
-	struct FDeferredDeletionQueue
-	{
-		FRenderCommandFence							Fence;
-		TArray<TUniquePtr<FNiagaraSystemInstance>>	Queue;
-	};
-
-	static constexpr int NumDeferredQueues = 3;
-	int DeferredDeletionQueueIndex = 0;
-	FDeferredDeletionQueue DeferredDeletionQueue[NumDeferredQueues];
+	/** Instances that have been queued for deletion this frame, serviced in PostActorTick */
+	TArray<TUniquePtr<FNiagaraSystemInstance>> DeferredDeletionQueue;
 
 	UPROPERTY(transient)
 	TMap<UNiagaraEffectType*, FNiagaraScalabilityManager> ScalabilityManagers;
