Fixed input buffers when a particle read DI reads from the same emitter it's applied to.
Also fixed a case when it couldn't find the ID variable in the source emitter because of an early out in FindAttributeIndices().

diff --git a/Engine/Plugins/FX/Niagara/Source/Niagara/Private/NiagaraDataInterfaceParticleRead.cpp b/Engine/Plugins/FX/Niagara/Source/Niagara/Private/NiagaraDataInterfaceParticleRead.cpp
--- a/Engine/Plugins/FX/Niagara/Source/Niagara/Private/NiagaraDataInterfaceParticleRead.cpp
+++ b/Engine/Plugins/FX/Niagara/Source/Niagara/Private/NiagaraDataInterfaceParticleRead.cpp
@@ -201,7 +201,6 @@ struct FNiagaraDataInterfaceParametersCS_ParticleRead : public FNiagaraDataInter
 	{
 		check(AttributeIndices.Num() == AttributeNames.Num());
 
-		AcquireTagRegisterIndex = -1;
 		const TArray<FNiagaraVariable>& SourceEmitterVariables = SourceDataSet->GetVariables();
 		const TArray<FNiagaraVariableLayoutInfo>& SourceEmitterVariableLayouts = SourceDataSet->GetVariableLayouts();
 		for (int AttrNameIdx = 0; AttrNameIdx < AttributeNames.Num(); ++AttrNameIdx)
@@ -210,12 +209,6 @@ struct FNiagaraDataInterfaceParametersCS_ParticleRead : public FNiagaraDataInter
 			for (int VarIdx = 0; VarIdx < SourceEmitterVariables.Num(); ++VarIdx)
 			{
 				const FNiagaraVariable& Var = SourceEmitterVariables[VarIdx];
-
-				if (AcquireTagRegisterIndex == -1 && Var.GetName().ToString() == TEXT("ID"))
-				{
-					AcquireTagRegisterIndex = SourceEmitterVariableLayouts[VarIdx].Int32ComponentStart + 1;
-				}
-
 				if (Var.GetName() == AttributeNames[AttrNameIdx])
 				{
 					ENiagaraParticleDataValueType AttributeType = AttributeTypes[AttrNameIdx];
@@ -243,6 +236,17 @@ struct FNiagaraDataInterfaceParametersCS_ParticleRead : public FNiagaraDataInter
 			}
 		}
 
+		AcquireTagRegisterIndex = -1;
+		for (int VarIdx = 0; VarIdx < SourceEmitterVariables.Num(); ++VarIdx)
+		{
+			const FNiagaraVariable& Var = SourceEmitterVariables[VarIdx];
+			if (Var.GetName().ToString() == TEXT("ID"))
+			{
+				AcquireTagRegisterIndex = SourceEmitterVariableLayouts[VarIdx].Int32ComponentStart + 1;
+				break;
+			}
+		}
+
 		if (AcquireTagRegisterIndex == -1)
 		{
 			UE_LOG(LogNiagara, Error, TEXT("Particle read DI cannot find ID variable in emitter '%s'."), SourceEmitterName);
@@ -286,7 +290,20 @@ struct FNiagaraDataInterfaceParametersCS_ParticleRead : public FNiagaraDataInter
 			return;
 		}
 
-		FNiagaraDataBuffer* SourceData = SourceDataSet->GetCurrentData();
+		FNiagaraDataBuffer* SourceData;
+		if (Context.ComputeInstanceData->Context == Proxy->SourceEmitterGPUContext)
+		{
+			// If the current execution context is the same as the source emitter's context, it means we're reading from
+			// ourselves. We can't use SourceDataSet->GetCurrentData() in that case, because EndSimulate() has already been
+			// called on the current emitter, and the current data has been set to the destination data. We need to use the
+			// current compute instance data to get to the input buffers.
+			SourceData = Context.ComputeInstanceData->CurrentData;
+		}
+		else
+		{
+			SourceData = SourceDataSet->GetCurrentData();
+		}
+
 		if (!SourceData)
 		{
 			SetErrorParams(RHICmdList, ComputeShader);
@@ -299,6 +316,15 @@ struct FNiagaraDataInterfaceParametersCS_ParticleRead : public FNiagaraDataInter
 			CachedDataSet = SourceDataSet;
 		}
 
+		if (!SourceData->GetGPUIDToIndexTable().Buffer)
+		{
+			// This can happen in the first frame, when there's no previous data yet. The DI shouldn't be
+			// queried in this case, because there's no way to have particle IDs (since there are no particles),
+			// but if it is it will just return failure and default values.
+			SetErrorParams(RHICmdList, ComputeShader);
+			return;
+		}
+
 		const uint32 ParticleStrideFloat = SourceData->GetFloatStride() / sizeof(float);
 		const uint32 ParticleStrideInt = SourceData->GetInt32Stride() / sizeof(int32);
 
@@ -915,7 +941,7 @@ static bool GenerateGetFunctionHLSL(const FNiagaraDataInterfaceGPUParamInfo& Par
 		"void {FunctionName}(NiagaraID In_ParticleID, out bool Out_Valid, out {ValueType} Out_Value)\n"
 		"{\n"
 		"    int RegisterIndex = {AttributeIndicesName}[{FunctionInstanceIndex}];\n"
-		"    int ParticleIndex = (RegisterIndex != -1) ? {IDToIndexTableName}[In_ParticleID.Index] : -1;\n"
+		"    int ParticleIndex = (RegisterIndex != -1) && (In_ParticleID.Index >= 0) ? {IDToIndexTableName}[In_ParticleID.Index] : -1;\n"
 		"    int AcquireTag = (ParticleIndex != -1) ? {InputIntBufferName}[{AcquireTagRegisterIndexName}*{ParticleStrideIntName} + ParticleIndex] : 0;\n"
 		"    if(ParticleIndex != -1 && In_ParticleID.AcquireTag == AcquireTag)\n"
 		"    {\n"
diff --git a/Engine/Plugins/FX/Niagara/Source/Niagara/Private/NiagaraEmitterInstanceBatcher.cpp b/Engine/Plugins/FX/Niagara/Source/Niagara/Private/NiagaraEmitterInstanceBatcher.cpp
--- a/Engine/Plugins/FX/Niagara/Source/Niagara/Private/NiagaraEmitterInstanceBatcher.cpp
+++ b/Engine/Plugins/FX/Niagara/Source/Niagara/Private/NiagaraEmitterInstanceBatcher.cpp
@@ -1156,6 +1156,7 @@ void NiagaraEmitterInstanceBatcher::SetDataInterfaceParameters(const TArray<FNia
 			Context.DataInterface = Interface;
 			Context.SystemInstance = SystemInstance;
 			Context.Batcher = this;
+			Context.ComputeInstanceData = Instance;
 			Context.ShaderStageIndex = ShaderStageIndex;
 			Context.IsOutputStage = Interface->IsOutputStage(ShaderStageIndex);
 			DIParam.Parameters->Set(RHICmdList, Context);
diff --git a/Engine/Plugins/FX/Niagara/Source/NiagaraCore/Public/NiagaraDataInterfaceBase.h b/Engine/Plugins/FX/Niagara/Source/NiagaraCore/Public/NiagaraDataInterfaceBase.h
--- a/Engine/Plugins/FX/Niagara/Source/NiagaraCore/Public/NiagaraDataInterfaceBase.h
+++ b/Engine/Plugins/FX/Niagara/Source/NiagaraCore/Public/NiagaraDataInterfaceBase.h
@@ -16,6 +16,7 @@ class FRHICommandList;
 
 struct FNiagaraDataInterfaceProxy;
 class NiagaraEmitterInstanceBatcher;
+struct FNiagaraComputeInstanceData;
 
 struct FNiagaraDataInterfaceSetArgs
 {
@@ -23,6 +24,7 @@ struct FNiagaraDataInterfaceSetArgs
 	FNiagaraDataInterfaceProxy* DataInterface;
 	FNiagaraSystemInstanceID SystemInstance;
 	const NiagaraEmitterInstanceBatcher* Batcher;
+	const FNiagaraComputeInstanceData* ComputeInstanceData;
 	uint32 ShaderStageIndex;
 	bool IsOutputStage;
 	bool IsIterationStage;
