Clarifying new system and new emitter workflows and adding checkboxes to the system graph and timeline menus for filtering library and template assets.
#rnx

diff --git a/Engine/Plugins/FX/Niagara/Source/Niagara/Classes/NiagaraEmitter.h b/Engine/Plugins/FX/Niagara/Source/Niagara/Classes/NiagaraEmitter.h
--- a/Engine/Plugins/FX/Niagara/Source/Niagara/Classes/NiagaraEmitter.h
+++ b/Engine/Plugins/FX/Niagara/Source/Niagara/Classes/NiagaraEmitter.h
@@ -394,6 +394,10 @@ public:
 	UPROPERTY()
 	uint32 ThumbnailImageOutOfDate : 1;
 
+	/* If this emitter is exposed to the library. */
+	UPROPERTY(EditAnywhere, AdvancedDisplay, Category = "Asset Options", AssetRegistrySearchable)
+	bool bExposeToLibrary;
+
 	UPROPERTY(EditAnywhere, AdvancedDisplay, Category = "Asset Options", AssetRegistrySearchable)
 	bool bIsTemplateAsset;
 
diff --git a/Engine/Plugins/FX/Niagara/Source/Niagara/Classes/NiagaraSystem.h b/Engine/Plugins/FX/Niagara/Source/Niagara/Classes/NiagaraSystem.h
--- a/Engine/Plugins/FX/Niagara/Source/Niagara/Classes/NiagaraSystem.h
+++ b/Engine/Plugins/FX/Niagara/Source/Niagara/Classes/NiagaraSystem.h
@@ -242,6 +242,10 @@ public:
 	UPROPERTY()
 	uint32 ThumbnailImageOutOfDate : 1;
 
+	/* If this system is exposed to the library. */
+	UPROPERTY(EditAnywhere, AdvancedDisplay, Category = "Asset Options", AssetRegistrySearchable)
+	bool bExposeToLibrary;
+
 	UPROPERTY(EditAnywhere, AdvancedDisplay, Category = "Asset Options", AssetRegistrySearchable)
 	bool bIsTemplateAsset;
 
diff --git a/Engine/Plugins/FX/Niagara/Source/NiagaraEditor/Private/NiagaraEditorModule.cpp b/Engine/Plugins/FX/Niagara/Source/NiagaraEditor/Private/NiagaraEditorModule.cpp
--- a/Engine/Plugins/FX/Niagara/Source/NiagaraEditor/Private/NiagaraEditorModule.cpp
+++ b/Engine/Plugins/FX/Niagara/Source/NiagaraEditor/Private/NiagaraEditorModule.cpp
@@ -777,6 +777,13 @@ void FNiagaraEditorModule::StartupModule()
 	);
 
 	FNiagaraEditorStyle::Initialize();
+	ReinitializeStyleCommand = IConsoleManager::Get().RegisterConsoleCommand(
+		TEXT("fx.NiagaraEditor.ReinitializeStyle"),
+		TEXT("Reinitializes the style for the niagara editor module.  Used in conjuction with live coding for UI tweaks.  May crash the editor if style objects are in use."),
+		FConsoleCommandDelegate::CreateRaw(this, &FNiagaraEditorModule::ReinitializeStyle));
+
+
+
 	FNiagaraEditorCommands::Register();
 
 	TSharedPtr<FNiagaraScriptGraphPanelPinFactory> GraphPanelPinFactory = MakeShareable(new FNiagaraScriptGraphPanelPinFactory());
@@ -1352,6 +1359,13 @@ void FNiagaraEditorModule::OnExecParticleInvoked(const TCHAR* Str)
 		}
 	}
 }
+
+void FNiagaraEditorModule::ReinitializeStyle()
+{
+	FNiagaraEditorStyle::Shutdown();
+	FNiagaraEditorStyle::Initialize();
+}
+
 PRAGMA_ENABLE_OPTIMIZATION
 
 #undef LOCTEXT_NAMESPACE
diff --git a/Engine/Plugins/FX/Niagara/Source/NiagaraEditor/Private/NiagaraEditorStyle.cpp b/Engine/Plugins/FX/Niagara/Source/NiagaraEditor/Private/NiagaraEditorStyle.cpp
--- a/Engine/Plugins/FX/Niagara/Source/NiagaraEditor/Private/NiagaraEditorStyle.cpp
+++ b/Engine/Plugins/FX/Niagara/Source/NiagaraEditor/Private/NiagaraEditorStyle.cpp
@@ -70,19 +70,29 @@ TSharedRef< FSlateStyleSet > FNiagaraEditorStyle::Create()
 	Style->Set("NiagaraEditor.StatsText", CategoryText);
 
 	// Asset picker
+	FTextBlockStyle AssetPickerBoldAssetNameText = FTextBlockStyle(NormalText)
+		.SetColorAndOpacity(FLinearColor::White)
+		.SetFont(DEFAULT_FONT("Bold", 9));
+
+	Style->Set("NiagaraEditor.AssetPickerBoldAssetNameText", AssetPickerBoldAssetNameText);
+
 	FTextBlockStyle AssetPickerAssetNameText = FTextBlockStyle(NormalText)
 		.SetColorAndOpacity(FLinearColor::White)
-		.SetFont(DEFAULT_FONT("Bold", 14))
-		.SetShadowOffset(FVector2D(0, 1))
-		.SetShadowColorAndOpacity(FLinearColor(0, 0, 0, 0.9f));
+		.SetFont(DEFAULT_FONT("Regular", 9));
 
 	Style->Set("NiagaraEditor.AssetPickerAssetNameText", AssetPickerAssetNameText);
 
 	FTextBlockStyle AssetPickerAssetCategoryText = FTextBlockStyle(NormalText)
-		.SetFont(DEFAULT_FONT("Regular", 10));
+		.SetFont(DEFAULT_FONT("Bold", 11));
 
 	Style->Set("NiagaraEditor.AssetPickerAssetCategoryText", AssetPickerAssetCategoryText);
 
+
+	FTextBlockStyle AssetPickerAssetSubcategoryText = FTextBlockStyle(NormalText)
+		.SetFont(DEFAULT_FONT("Bold", 10));
+
+	Style->Set("NiagaraEditor.AssetPickerAssetSubcategoryText", AssetPickerAssetSubcategoryText);
+
 	// New Asset Dialog
 	FTextBlockStyle NewAssetDialogOptionText = FTextBlockStyle(NormalText)
 		.SetFont(DEFAULT_FONT("Bold", 11));
@@ -91,17 +101,13 @@ TSharedRef< FSlateStyleSet > FNiagaraEditorStyle::Create()
 
 	FTextBlockStyle NewAssetDialogHeaderText = FTextBlockStyle(NormalText)
 		.SetColorAndOpacity(FLinearColor::White)
-		.SetFont(DEFAULT_FONT("Bold", 10))
-		.SetShadowOffset(FVector2D(0, 1))
-		.SetShadowColorAndOpacity(FLinearColor(0, 0, 0, 0.9f));
+		.SetFont(DEFAULT_FONT("Bold", 10));
 
 	Style->Set("NiagaraEditor.NewAssetDialog.HeaderText", NewAssetDialogHeaderText);
 
 	FTextBlockStyle NewAssetDialogSubHeaderText = FTextBlockStyle(NormalText)
 		.SetColorAndOpacity(FLinearColor::White)
-		.SetFont(DEFAULT_FONT("Bold", 10))
-		.SetShadowOffset(FVector2D(0, 1))
-		.SetShadowColorAndOpacity(FLinearColor(0, 0, 0, 0.9f));
+		.SetFont(DEFAULT_FONT("Bold", 11));
 
 	Style->Set("NiagaraEditor.NewAssetDialog.SubHeaderText", NewAssetDialogSubHeaderText);
 
@@ -118,9 +124,7 @@ TSharedRef< FSlateStyleSet > FNiagaraEditorStyle::Create()
 
 	// Emitter Header
 	FTextBlockStyle HeadingText = FTextBlockStyle(NormalText)
-		.SetFont(DEFAULT_FONT("Regular", 14))
-		.SetShadowOffset(FVector2D(0, 1))
-		.SetShadowColorAndOpacity(FLinearColor(0, 0, 0, 0.9f));
+		.SetFont(DEFAULT_FONT("Regular", 14));
 
 	FEditableTextBoxStyle HeadingEditableTextBox = FEditableTextBoxStyle(NormalEditableTextBox)
 		.SetFont(DEFAULT_FONT("Regular", 14));
@@ -142,17 +146,13 @@ TSharedRef< FSlateStyleSet > FNiagaraEditorStyle::Create()
 
 	FTextBlockStyle SubduedHeadingText = FTextBlockStyle(NormalText)
 		.SetFont(DEFAULT_FONT("Regular", 14))
-		.SetShadowOffset(FVector2D(0, 1))
-		.SetColorAndOpacity(FSlateColor(FLinearColor(0.5f, 0.5f, 0.5f, 1.0f)))
-		.SetShadowColorAndOpacity(FLinearColor(0, 0, 0, 0.9f));
+		.SetColorAndOpacity(FSlateColor(FLinearColor(0.5f, 0.5f, 0.5f, 1.0f)));
 	
 	Style->Set("NiagaraEditor.SubduedHeadingTextBox", SubduedHeadingText);
 
 	// Details
 	FTextBlockStyle DetailsHeadingText = FTextBlockStyle(NormalText)
-		.SetFont(DEFAULT_FONT("Bold", 9))
-		.SetShadowOffset(FVector2D(0, 1))
-		.SetShadowColorAndOpacity(FLinearColor(0, 0, 0, 0.9f));
+		.SetFont(DEFAULT_FONT("Bold", 9));
 
 	Style->Set("NiagaraEditor.DetailsHeadingText", DetailsHeadingText);
 
@@ -201,9 +201,7 @@ TSharedRef< FSlateStyleSet > FNiagaraEditorStyle::Create()
 	FSlateFontInfo SelectedEmitterUnsupportedSelectionFont = DEFAULT_FONT("Regular", 10);
 	FTextBlockStyle SelectedEmitterUnsupportedSelectionText = FTextBlockStyle(NormalText)
 		.SetFont(SelectedEmitterUnsupportedSelectionFont)
-		.SetColorAndOpacity(FLinearColor(1.0f, 1.0f, 1.0f, 1.0f))
-		.SetShadowOffset(FVector2D(0, 1))
-		.SetShadowColorAndOpacity(FLinearColor(0, 0, 0, 0.9f));
+		.SetColorAndOpacity(FLinearColor(1.0f, 1.0f, 1.0f, 1.0f));
 	Style->Set("NiagaraEditor.SelectedEmitter.UnsupportedSelectionText", SelectedEmitterUnsupportedSelectionText);
 
 	// Toolbar Icons
diff --git a/Engine/Plugins/FX/Niagara/Source/NiagaraEditor/Private/SNewEmitterDialog.cpp b/Engine/Plugins/FX/Niagara/Source/NiagaraEditor/Private/SNewEmitterDialog.cpp
--- a/Engine/Plugins/FX/Niagara/Source/NiagaraEditor/Private/SNewEmitterDialog.cpp
+++ b/Engine/Plugins/FX/Niagara/Source/NiagaraEditor/Private/SNewEmitterDialog.cpp
@@ -3,7 +3,7 @@
 #include "SNewEmitterDialog.h"
 #include "NiagaraEmitter.h"
 #include "NiagaraEditorStyle.h"
-#include "SNiagaraTemplateAssetPicker.h"
+#include "SNiagaraAssetPickerList.h"
 #include "SNiagaraNewAssetDialog.h"
 
 #include "AssetData.h"
@@ -17,49 +17,34 @@ typedef SItemSelector<FText, FAssetData> SNiagaraAssetItemSelector;
 
 void SNewEmitterDialog::Construct(const FArguments& InArgs)
 {
-	FAssetPickerConfig AssetPickerConfig;
-	AssetPickerConfig.SelectionMode = ESelectionMode::SingleToggle;
-	AssetPickerConfig.InitialAssetViewType = EAssetViewType::List;
-	AssetPickerConfig.Filter.ClassNames.Add(UNiagaraEmitter::StaticClass()->GetFName());
-	AssetPickerConfig.GetCurrentSelectionDelegates.Add(&GetSelectedEmitterAssetsFromPicker);
-	AssetPickerConfig.OnAssetsActivated.BindSP(this, &SNewEmitterDialog::OnEmitterAssetsActivated);
-
-	FContentBrowserModule& ContentBrowserModule = FModuleManager::Get().LoadModuleChecked<FContentBrowserModule>(TEXT("ContentBrowser"));
-	TSharedRef<SWidget> AssetPicker = ContentBrowserModule.Get().CreateAssetPicker(AssetPickerConfig);
-
 	SNiagaraNewAssetDialog::Construct(SNiagaraNewAssetDialog::FArguments(), UNiagaraEmitter::StaticClass()->GetFName(), LOCTEXT("AssetTypeName", "emitter"),
 		{
 			SNiagaraNewAssetDialog::FNiagaraNewAssetDialogOption(
-				LOCTEXT("CreateFromTemplateLabel", "Create a new emitter from an emitter template"),
+				LOCTEXT("CreateFromTemplateLabel", "New emitter from a template"),
+				LOCTEXT("CreateFromTemplateDescription", "Create a new emitter from an emitter template (no inheritance)"),
 				LOCTEXT("TemplatesPickerHeader", "Select a Template Emitter"),
 				SNiagaraNewAssetDialog::FOnGetSelectedAssetsFromPicker::CreateSP(this, &SNewEmitterDialog::GetSelectedEmitterTemplateAssets),
 				SNiagaraNewAssetDialog::FOnSelectionConfirmed(),
-				SAssignNew(TemplateAssetPicker, SNiagaraTemplateAssetPicker, UNiagaraEmitter::StaticClass())
-				.OnTemplateAssetActivated(this, &SNewEmitterDialog::OnTemplateAssetActivated)),
+				SAssignNew(TemplateAssetPicker, SNiagaraAssetPickerList, UNiagaraEmitter::StaticClass())
+				.OnTemplateAssetActivated(this, &SNewEmitterDialog::OnTemplateAssetActivated)
+				.bTemplateOnly(true)),
 			SNiagaraNewAssetDialog::FNiagaraNewAssetDialogOption(
-				LOCTEXT("CreateFromOtherEmitterLabel", "Copy an existing emitter from your project content"),
-				LOCTEXT("ProjectEmitterPickerHeader", "Select a Project Emitter"),
-				SNiagaraNewAssetDialog::FOnGetSelectedAssetsFromPicker::CreateSP(this, &SNewEmitterDialog::GetSelectedProjectEmiterAssets),
-				SNiagaraNewAssetDialog::FOnSelectionConfirmed(),
-				AssetPicker),
-			SNiagaraNewAssetDialog::FNiagaraNewAssetDialogOption(
-				LOCTEXT("InheritFromOtherEmitterLabel", "Inherit from an existing emitter in your project content"),
+				LOCTEXT("InheritFromOtherEmitterLabel", "Inherit from an existing emitter"),
+				LOCTEXT("InheritFromOtherEmitterDescription", "Create an inheritance chain between the new emitter and an existing emitter"),
 				LOCTEXT("InheritProjectEmitterPickerHeader", "Select a Parent Project Emitter"),
 				SNiagaraNewAssetDialog::FOnGetSelectedAssetsFromPicker::CreateSP(this, &SNewEmitterDialog::GetSelectedProjectEmiterAssets),
 				SNiagaraNewAssetDialog::FOnSelectionConfirmed::CreateSP(this, &SNewEmitterDialog::InheritanceOptionConfirmed),
-				AssetPicker),
+				SAssignNew(TemplateAssetPicker, SNiagaraAssetPickerList, UNiagaraEmitter::StaticClass())
+				.bTemplateOnly(false)),
 			SNiagaraNewAssetDialog::FNiagaraNewAssetDialogOption(
-				LOCTEXT("CreateEmptyLabel", "Create an empty emitter with no modules or renderers (Advanced)"),
-				LOCTEXT("EmptyLabel", "Empty Emitter"),
-				SNiagaraNewAssetDialog::FOnGetSelectedAssetsFromPicker(),
+				LOCTEXT("CreateFromOtherEmitterLabel", "Copy existing emitter"),
+				LOCTEXT("CreateFromOtherEmitterDescription", "Copies an existing emitter from your project content"),
+				LOCTEXT("ProjectEmitterPickerHeader", "Select a Project Emitter"),
+				SNiagaraNewAssetDialog::FOnGetSelectedAssetsFromPicker::CreateSP(this, &SNewEmitterDialog::GetSelectedProjectEmiterAssets),
 				SNiagaraNewAssetDialog::FOnSelectionConfirmed(),
-				SNew(SBox)
-				.HAlign(HAlign_Center)
-				.VAlign(VAlign_Center)
-				[
-					SNew(STextBlock)
-					.Text(LOCTEXT("NoOptionsLabel", "No Options"))
-				])
+				SAssignNew(TemplateAssetPicker, SNiagaraAssetPickerList, UNiagaraEmitter::StaticClass())
+				.OnTemplateAssetActivated(this, &SNewEmitterDialog::OnEmitterAssetsActivated)
+				.bTemplateOnly(false))
 		});
 }
 
@@ -89,7 +74,7 @@ void SNewEmitterDialog::GetSelectedEmitterTemplateAssets(TArray<FAssetData>& Out
 
 void SNewEmitterDialog::GetSelectedProjectEmiterAssets(TArray<FAssetData>& OutSelectedAssets)
 {
-	OutSelectedAssets.Append(GetSelectedEmitterAssetsFromPicker.Execute());
+	OutSelectedAssets.Append(TemplateAssetPicker->GetSelectedAssets());
 	if (ActivatedProjectAsset.IsValid())
 	{
 		OutSelectedAssets.AddUnique(ActivatedProjectAsset);
@@ -104,15 +89,12 @@ void SNewEmitterDialog::OnTemplateAssetActivated(const FAssetData& InActivatedTe
 	ConfirmSelection();
 }
 
-void SNewEmitterDialog::OnEmitterAssetsActivated(const TArray<FAssetData>& ActivatedAssets, EAssetTypeActivationMethod::Type ActivationMethod)
+void SNewEmitterDialog::OnEmitterAssetsActivated(const FAssetData& InActivatedTemplateAsset)
 {
-	if ((ActivationMethod == EAssetTypeActivationMethod::DoubleClicked || ActivationMethod == EAssetTypeActivationMethod::Opened) && ActivatedAssets.Num() == 1)
-	{
-		// Input handling issues with the list view widget can allow items to be activated but not added to the selection so cache this here
-		// so it can be included in the selection set.
-		ActivatedProjectAsset = ActivatedAssets[0];
-		ConfirmSelection();
-	}
+	// Input handling issues with the list view widget can allow items to be activated but not added to the selection so cache this here
+	// so it can be included in the selection set.
+	ActivatedProjectAsset = InActivatedTemplateAsset;
+	ConfirmSelection();
 }
 
 void SNewEmitterDialog::InheritanceOptionConfirmed()
diff --git a/Engine/Plugins/FX/Niagara/Source/NiagaraEditor/Private/SNewEmitterDialog.h b/Engine/Plugins/FX/Niagara/Source/NiagaraEditor/Private/SNewEmitterDialog.h
--- a/Engine/Plugins/FX/Niagara/Source/NiagaraEditor/Private/SNewEmitterDialog.h
+++ b/Engine/Plugins/FX/Niagara/Source/NiagaraEditor/Private/SNewEmitterDialog.h
@@ -9,7 +9,7 @@
 #include "AssetData.h"
 #include "ContentBrowserDelegates.h"
 
-class SNiagaraTemplateAssetPicker;
+class SNiagaraAssetPickerList;
 
 /** A modal dialog to collect information needed to create a new niagara system. */
 class SNewEmitterDialog : public SNiagaraNewAssetDialog
@@ -32,14 +32,12 @@ private:
 
 	void OnTemplateAssetActivated(const FAssetData& InActivatedTemplateAsset);
 
-	void OnEmitterAssetsActivated(const TArray<FAssetData>& ActivatedAssets, EAssetTypeActivationMethod::Type ActivationMethod);
+	void OnEmitterAssetsActivated(const FAssetData& InActivatedTemplateAsset);
 
 	void InheritanceOptionConfirmed();
 
 private:
-	TSharedPtr<SNiagaraTemplateAssetPicker> TemplateAssetPicker;
-
-	FGetCurrentSelectionDelegate GetSelectedEmitterAssetsFromPicker;
+	TSharedPtr<SNiagaraAssetPickerList> TemplateAssetPicker;
 
 	FAssetData ActivatedTemplateAsset;
 
diff --git a/Engine/Plugins/FX/Niagara/Source/NiagaraEditor/Private/SNewSystemDialog.cpp b/Engine/Plugins/FX/Niagara/Source/NiagaraEditor/Private/SNewSystemDialog.cpp
--- a/Engine/Plugins/FX/Niagara/Source/NiagaraEditor/Private/SNewSystemDialog.cpp
+++ b/Engine/Plugins/FX/Niagara/Source/NiagaraEditor/Private/SNewSystemDialog.cpp
@@ -3,7 +3,7 @@
 #include "SNewSystemDialog.h"
 #include "NiagaraSystem.h"
 #include "NiagaraEditorStyle.h"
-#include "SNiagaraTemplateAssetPicker.h"
+#include "SNiagaraAssetPickerList.h"
 
 #include "AssetData.h"
 
@@ -24,35 +24,11 @@ void SNewSystemDialog::Construct(const FArguments& InArgs)
 {
 	FContentBrowserModule& ContentBrowserModule = FModuleManager::Get().LoadModuleChecked<FContentBrowserModule>(TEXT("ContentBrowser"));
 
-	FAssetPickerConfig SystemAssetPickerConfig;
-	SystemAssetPickerConfig.SelectionMode = ESelectionMode::SingleToggle;
-	SystemAssetPickerConfig.InitialAssetViewType = EAssetViewType::List;
-	SystemAssetPickerConfig.Filter.ClassNames.Add(UNiagaraSystem::StaticClass()->GetFName());
-	SystemAssetPickerConfig.GetCurrentSelectionDelegates.Add(&GetSelectedSystemAssetsFromPicker);
-	SystemAssetPickerConfig.OnAssetsActivated.BindSP(this, &SNewSystemDialog::OnSystemAssetsActivated);
-
-	TSharedRef<SWidget> SystemAssetPicker = ContentBrowserModule.Get().CreateAssetPicker(SystemAssetPickerConfig);
-
-	FAssetPickerConfig EmitterAssetPickerConfig;
-	EmitterAssetPickerConfig.SelectionMode = ESelectionMode::Multi;
-	EmitterAssetPickerConfig.InitialAssetViewType = EAssetViewType::List;
-	EmitterAssetPickerConfig.Filter.ClassNames.Add(UNiagaraEmitter::StaticClass()->GetFName());
-	EmitterAssetPickerConfig.GetCurrentSelectionDelegates.Add(&GetSelectedEmitterAssetsFromPicker);
-	EmitterAssetPickerConfig.OnAssetsActivated.BindSP(this, &SNewSystemDialog::OnEmitterAssetsActivated);
-
-	TSharedRef<SWidget> EmitterAssetPicker = ContentBrowserModule.Get().CreateAssetPicker(EmitterAssetPickerConfig);
-
 	SNiagaraNewAssetDialog::Construct(SNiagaraNewAssetDialog::FArguments(), UNiagaraSystem::StaticClass()->GetFName(), LOCTEXT("AssetTypeName", "system"),
 		{
 			SNiagaraNewAssetDialog::FNiagaraNewAssetDialogOption(
-				LOCTEXT("CreateFromTemplateLabel", "Create a new system from a system template"),
-				LOCTEXT("TemplateLabel", "Select a System Template"),
-				SNiagaraNewAssetDialog::FOnGetSelectedAssetsFromPicker::CreateSP(this, &SNewSystemDialog::GetSelectedSystemTemplateAssets),
-				SNiagaraNewAssetDialog::FOnSelectionConfirmed(),
-				SAssignNew(TemplateAssetPicker, SNiagaraTemplateAssetPicker, UNiagaraSystem::StaticClass())
-				.OnTemplateAssetActivated(this, &SNewSystemDialog::OnTemplateAssetActivated)),
-			SNiagaraNewAssetDialog::FNiagaraNewAssetDialogOption(
-				LOCTEXT("CreateFromSelectedEmittersLabel", "Create a new system from a set of selected emitters"),
+				LOCTEXT("CreateFromSelectedEmittersLabel", "New system from selected emitter(s)"),
+				LOCTEXT("CreateFromSelectedEmittersDescription", "Choose a mix of emitters (inherited) and emitter templates (no inheritance)"),
 				LOCTEXT("ProjectEmittersLabel", "Select Emitters to Add"),
 				SNiagaraNewAssetDialog::FOnGetSelectedAssetsFromPicker::CreateSP(this, &SNewSystemDialog::GetSelectedProjectEmiterAssets),
 				SNiagaraNewAssetDialog::FOnSelectionConfirmed(),
@@ -60,7 +36,9 @@ void SNewSystemDialog::Construct(const FArguments& InArgs)
 				+ SVerticalBox::Slot()
 				.Padding(0, 0, 0, 10)
 				[
-					EmitterAssetPicker
+					SAssignNew(EmitterAssetPicker, SNiagaraAssetPickerList, UNiagaraEmitter::StaticClass())
+					.bTemplateOnly(false)
+					.bAllowMultiSelect(true)
 				]
 				+ SVerticalBox::Slot()
 				.AutoHeight()
@@ -80,7 +58,6 @@ void SNewSystemDialog::Construct(const FArguments& InArgs)
 					.VAlign(VAlign_Center)
 					[
 						SNew(SButton)
-						//.ButtonStyle(FNiagaraEditorStyle::Get(), "NiagaraEditor.NewAssetDialog.AddButton")
 						.ButtonStyle(FEditorStyle::Get(), "FlatButton.Success")
 						.IsEnabled(this, &SNewSystemDialog::IsAddEmittersToSelectionButtonEnabled)
 						.OnClicked(this, &SNewSystemDialog::AddEmittersToSelectionButtonClicked)
@@ -108,13 +85,26 @@ void SNewSystemDialog::Construct(const FArguments& InArgs)
 					.UseAllottedWidth(true)
 				]),
 			SNiagaraNewAssetDialog::FNiagaraNewAssetDialogOption(
-				LOCTEXT("CreateFromOtherSystemLabel", "Copy an existing system from your project content"),
+				LOCTEXT("CreateFromTemplateLabel", "New system from a template"),
+				LOCTEXT("CreateFromTemplateDescription", "The new system will be derived from a system template"),
+				LOCTEXT("TemplateLabel", "Select a System Template"),
+				SNiagaraNewAssetDialog::FOnGetSelectedAssetsFromPicker::CreateSP(this, &SNewSystemDialog::GetSelectedSystemTemplateAssets),
+				SNiagaraNewAssetDialog::FOnSelectionConfirmed(),
+				SAssignNew(TemplateAssetPicker, SNiagaraAssetPickerList, UNiagaraSystem::StaticClass())
+				.OnTemplateAssetActivated(this, &SNewSystemDialog::OnTemplateAssetActivated)
+				.bTemplateOnly(true)),
+			SNiagaraNewAssetDialog::FNiagaraNewAssetDialogOption(
+				LOCTEXT("CreateFromOtherSystemLabel", "Copy existing system"),
+				LOCTEXT("CreateFromOtherSystemDescription", "Copies an existing system from your project content and maintains any inheritance of the included emitters"),
 				LOCTEXT("ProjectSystemsLabel", "Select a Project System"),
 				SNiagaraNewAssetDialog::FOnGetSelectedAssetsFromPicker::CreateSP(this, &SNewSystemDialog::GetSelectedProjectSystemAssets),
 				SNiagaraNewAssetDialog::FOnSelectionConfirmed(),
-				SystemAssetPicker),
+				SAssignNew(SystemAssetPicker, SNiagaraAssetPickerList, UNiagaraSystem::StaticClass())
+				.OnTemplateAssetActivated(this, &SNewSystemDialog::OnSystemAssetsActivated)
+				.bTemplateOnly(false)),
 			SNiagaraNewAssetDialog::FNiagaraNewAssetDialogOption(
-				LOCTEXT("CreateEmptyLabel", "Create an empty system with no emitters"),
+				LOCTEXT("CreateEmptyLabel", "Create empty system"),
+				LOCTEXT("CreateEmptyDescription", "Create an empty system with no emitters or emitter templates"),
 				LOCTEXT("EmptyLabel", "Empty System"),
 				SNiagaraNewAssetDialog::FOnGetSelectedAssetsFromPicker(),
 				SNiagaraNewAssetDialog::FOnSelectionConfirmed(),
@@ -171,7 +161,7 @@ void SNewSystemDialog::GetSelectedSystemTemplateAssets(TArray<FAssetData>& OutSe
 
 void SNewSystemDialog::GetSelectedProjectSystemAssets(TArray<FAssetData>& OutSelectedAssets)
 {
-	OutSelectedAssets.Append(GetSelectedSystemAssetsFromPicker.Execute());
+	OutSelectedAssets.Append(SystemAssetPicker->GetSelectedAssets());
 	if (ActivatedProjectSystemAsset.IsValid())
 	{
 		OutSelectedAssets.AddUnique(ActivatedProjectSystemAsset);
@@ -191,33 +181,22 @@ void SNewSystemDialog::OnTemplateAssetActivated(const FAssetData& ActivatedTempl
 	ConfirmSelection();
 }
 
-void SNewSystemDialog::OnSystemAssetsActivated(const TArray<FAssetData>& ActivatedAssets, EAssetTypeActivationMethod::Type ActivationMethod)
-{
-	if ((ActivationMethod == EAssetTypeActivationMethod::DoubleClicked || ActivationMethod == EAssetTypeActivationMethod::Opened) && ActivatedAssets.Num() == 1)
-	{
-		// Input handling issues with the list view widget can allow items to be activated but not added to the selection so cache this here
-		// so it can be included in the selection set.
-		ActivatedProjectSystemAsset = ActivatedAssets[0];
-		ConfirmSelection();
-	}
-}
-
-void SNewSystemDialog::OnEmitterAssetsActivated(const TArray<FAssetData>& ActivatedAssets, EAssetTypeActivationMethod::Type ActivationMethod)
+void SNewSystemDialog::OnSystemAssetsActivated(const FAssetData& ActivatedTemplateAsset)
 {
-	if (ActivationMethod == EAssetTypeActivationMethod::DoubleClicked || ActivationMethod == EAssetTypeActivationMethod::Opened)
-	{
-		AddEmitterAssetsToSelection(ActivatedAssets);
-	}
+	// Input handling issues with the list view widget can allow items to be activated but not added to the selection so cache this here
+	// so it can be included in the selection set.
+	ActivatedProjectSystemAsset = ActivatedTemplateAsset;
+	ConfirmSelection();
 }
 
 bool SNewSystemDialog::IsAddEmittersToSelectionButtonEnabled() const
 {
-	return GetSelectedEmitterAssetsFromPicker.Execute().Num() > 0;
+	return EmitterAssetPicker->GetSelectedAssets().Num() > 0;
 }
 
 FReply SNewSystemDialog::AddEmittersToSelectionButtonClicked()
 {
-	TArray<FAssetData> SelectedEmitterAssetsFromPicker = GetSelectedEmitterAssetsFromPicker.Execute();
+	TArray<FAssetData> SelectedEmitterAssetsFromPicker = EmitterAssetPicker->GetSelectedAssets();
 	AddEmitterAssetsToSelection(SelectedEmitterAssetsFromPicker);
 	return FReply::Handled();
 }
diff --git a/Engine/Plugins/FX/Niagara/Source/NiagaraEditor/Private/SNewSystemDialog.h b/Engine/Plugins/FX/Niagara/Source/NiagaraEditor/Private/SNewSystemDialog.h
--- a/Engine/Plugins/FX/Niagara/Source/NiagaraEditor/Private/SNewSystemDialog.h
+++ b/Engine/Plugins/FX/Niagara/Source/NiagaraEditor/Private/SNewSystemDialog.h
@@ -9,7 +9,7 @@
 #include "ContentBrowserDelegates.h"
 #include "SNiagaraNewAssetDialog.h"
 
-class SNiagaraTemplateAssetPicker;
+class SNiagaraAssetPickerList;
 class SWrapBox;
 
 /** A modal dialog to collect information needed to create a new niagara system. */
@@ -35,9 +35,7 @@ private:
 
 	void OnTemplateAssetActivated(const FAssetData& ActivatedTemplateAsset);
 
-	void OnSystemAssetsActivated(const TArray<FAssetData>& ActivatedAssets, EAssetTypeActivationMethod::Type ActivationMethod);
-
-	void OnEmitterAssetsActivated(const TArray<FAssetData>& ActivatedAssets, EAssetTypeActivationMethod::Type ActivationMethod);
+	void OnSystemAssetsActivated(const FAssetData& ActivatedTemplateAsset);
 
 	bool IsAddEmittersToSelectionButtonEnabled() const;
 	
@@ -48,10 +46,6 @@ private:
 	FReply RemoveEmitterFromSelectionButtonClicked(FAssetData EmitterAsset);
 
 private:
-	FGetCurrentSelectionDelegate GetSelectedSystemAssetsFromPicker;
-
-	FGetCurrentSelectionDelegate GetSelectedEmitterAssetsFromPicker;
-
 	FAssetData ActivatedTemplateSystemAsset;
 
 	FAssetData ActivatedProjectSystemAsset;
@@ -62,5 +56,9 @@ private:
 
 	TSharedPtr<SWrapBox> SelectedEmitterBox;
 
-	TSharedPtr<SNiagaraTemplateAssetPicker> TemplateAssetPicker;
+	TSharedPtr<SNiagaraAssetPickerList> TemplateAssetPicker;
+
+	TSharedPtr<SNiagaraAssetPickerList> EmitterAssetPicker;
+
+	TSharedPtr<SNiagaraAssetPickerList> SystemAssetPicker;
 };
\ No newline at end of file
diff --git a/Engine/Plugins/FX/Niagara/Source/NiagaraEditor/Private/SNiagaraAssetPickerList.cpp b/Engine/Plugins/FX/Niagara/Source/NiagaraEditor/Private/SNiagaraAssetPickerList.cpp
new file mode 100644
--- /dev/null
+++ b/Engine/Plugins/FX/Niagara/Source/NiagaraEditor/Private/SNiagaraAssetPickerList.cpp
@@ -0,0 +1,298 @@
+// Copyright Epic Games, Inc. All Rights Reserved.
+
+#include "SNiagaraAssetPickerList.h"
+#include "NiagaraEmitter.h"
+#include "NiagaraEditorStyle.h"
+
+#include "AssetData.h"
+#include "AssetRegistryModule.h"
+#include "Modules/ModuleManager.h"
+#include "AssetThumbnail.h"
+
+#include "Widgets/Input/SButton.h"
+#include "Widgets/Layout/SUniformGridPanel.h"
+#include "EditorStyleSet.h"
+
+#define LOCTEXT_NAMESPACE "SNiagaraAssetSelector"
+
+void SNiagaraAssetPickerList::Construct(const FArguments& InArgs, UClass* AssetClass)
+{
+	FAssetRegistryModule& AssetRegistryModule = FModuleManager::LoadModuleChecked<FAssetRegistryModule>(TEXT("AssetRegistry"));
+	TArray<FAssetData> EmitterAssets;
+	AssetRegistryModule.Get().GetAssetsByClass(AssetClass->GetFName(), EmitterAssets);
+
+	NiagaraPluginCategory = LOCTEXT("NiagaraCategory", "Engine (Niagara Plugin)");
+	ProjectCategory = LOCTEXT("ProjectCategory", "Project");
+	TemplateCategory = LOCTEXT("Template", "Template");
+	LibraryCategory = LOCTEXT("Library", "Library");
+	NonLibraryCategory = LOCTEXT("NotInLibrary", "Not in Library");
+	AssetThumbnailPool = MakeShareable(new FAssetThumbnailPool(24));
+
+	OnTemplateAssetActivated = InArgs._OnTemplateAssetActivated;
+	bTemplateOnly = InArgs._bTemplateOnly;
+	TArray<FAssetData> EmittersToShow;
+	for (const FAssetData& EmitterAsset : EmitterAssets)
+	{
+		bool bShowEmitter = false;
+		EmitterAsset.GetTagValue("bIsTemplateAsset", bShowEmitter);
+		if (bShowEmitter || !bTemplateOnly)
+		{
+			EmittersToShow.Add(EmitterAsset);
+		}
+	}
+
+	ChildSlot
+	[
+		SAssignNew(ItemSelector, SNiagaraAssetItemSelector)
+		.Items(EmittersToShow)
+		.OnGetCategoriesForItem(this, &SNiagaraAssetPickerList::OnGetCategoriesForItem)
+		.OnCompareCategoriesForEquality(this, &SNiagaraAssetPickerList::OnCompareCategoriesForEquality)
+		.OnCompareCategoriesForSorting(this, &SNiagaraAssetPickerList::OnCompareCategoriesForSorting)
+		.OnCompareItemsForSorting(this, &SNiagaraAssetPickerList::OnCompareItemsForSorting)
+		.OnDoesItemMatchFilterText(this, &SNiagaraAssetPickerList::OnDoesItemMatchFilterText)
+		.OnGenerateWidgetForCategory(this, &SNiagaraAssetPickerList::OnGenerateWidgetForCategory)
+		.OnGenerateWidgetForItem(this, &SNiagaraAssetPickerList::OnGenerateWidgetForItem)
+		.OnItemActivated(this, &SNiagaraAssetPickerList::OnItemActivated)
+		.AllowMultiselect(InArgs._bAllowMultiSelect)
+	];
+}
+
+void SNiagaraAssetPickerList::Tick(const FGeometry& AllottedGeometry, const double InCurrentTime, const float InDeltaTime)
+{
+	AssetThumbnailPool->Tick(InDeltaTime);
+}
+
+TArray<FAssetData> SNiagaraAssetPickerList::GetSelectedAssets() const
+{
+	return ItemSelector->GetSelectedItems();
+}
+
+TArray<FText> SNiagaraAssetPickerList::OnGetCategoriesForItem(const FAssetData& Item)
+{
+	TArray<FText> Categories;
+	if (bTemplateOnly)
+	{
+		TArray<FString> AssetPathParts;
+		Item.ObjectPath.ToString().ParseIntoArray(AssetPathParts, TEXT("/"));
+		if (AssetPathParts.Num() > 0)
+		{
+			if (AssetPathParts[0] == TEXT("Niagara"))
+			{
+				Categories.Add(LOCTEXT("NiagaraCategory", "Engine (Niagara Plugin)"));
+			}
+			else if (AssetPathParts[0] == TEXT("Game"))
+			{
+				Categories.Add(LOCTEXT("ProjectCategory", "Project"));
+			}
+			else
+			{
+				Categories.Add(FText::Format(LOCTEXT("OtherPluginFormat", "Plugin - {0}"), FText::FromString(AssetPathParts[0])));
+			}
+		}
+	}
+	else
+	{
+
+		bool bIsTemplate;
+		Item.GetTagValue("bIsTemplateAsset", bIsTemplate);
+		if (bIsTemplate)
+		{
+			Categories.Add(TemplateCategory);
+		}
+		else
+		{
+			Categories.Add(LOCTEXT("NonTemplateEmitters", "Non-Template"));
+		}
+	}
+
+	bool bIsLibrary;
+	bool bFoundLibScriptTag = Item.GetTagValue("bExposeToLibrary", bIsLibrary);
+
+	if (bFoundLibScriptTag && bIsLibrary)
+	{
+		Categories.Add(LibraryCategory);
+	}
+	else
+	{
+		Categories.Add(NonLibraryCategory);
+	}
+
+
+	return Categories;
+}
+
+bool SNiagaraAssetPickerList::OnCompareCategoriesForEquality(const FText& CategoryA, const FText& CategoryB) const
+{
+	return CategoryA.CompareTo(CategoryB) == 0;
+}
+
+bool SNiagaraAssetPickerList::OnCompareCategoriesForSorting(const FText& CategoryA, const FText& CategoryB) const
+{
+	int32 CompareResult = CategoryA.CompareTo(CategoryB);
+	if (bTemplateOnly)
+	{
+		if (CompareResult != 0)
+		{
+			// Project first
+			if (CategoryA.CompareTo(ProjectCategory) == 0)
+			{
+				return true;
+			}
+			if (CategoryB.CompareTo(ProjectCategory) == 0)
+			{
+				return false;
+			}
+
+			// Niagara plugin second.
+			if (CategoryA.CompareTo(NiagaraPluginCategory) == 0)
+			{
+				return true;
+			}
+			if (CategoryB.CompareTo(NiagaraPluginCategory) == 0)
+			{
+				return false;
+			}
+		}
+	}
+	else
+	{
+		if (CompareResult != 0)
+		{
+			// Template first
+			if (CategoryA.CompareTo(TemplateCategory) == 0)
+			{
+				return true;
+			}
+			if (CategoryB.CompareTo(TemplateCategory) == 0)
+			{
+				return false;
+			}
+
+			// Library second.
+			if (CategoryA.CompareTo(LibraryCategory) == 0)
+			{
+				return true;
+			}
+			if (CategoryB.CompareTo(LibraryCategory) == 0)
+			{
+				return false;
+			}
+		}
+	}
+	// Otherwise just return the actual result.
+	return CompareResult < 0;
+}
+
+bool SNiagaraAssetPickerList::OnCompareItemsForSorting(const FAssetData& ItemA, const FAssetData& ItemB) const
+{
+	return ItemA.AssetName.ToString().Compare(ItemB.AssetName.ToString()) < 0;
+}
+
+bool SNiagaraAssetPickerList::OnDoesItemMatchFilterText(const FText& FilterText, const FAssetData& Item)
+{
+	return Item.AssetName.ToString().Contains(FilterText.ToString());
+}
+
+TSharedRef<SWidget> SNiagaraAssetPickerList::OnGenerateWidgetForCategory(const FText& Category)
+{
+	if (Category.EqualTo(LibraryCategory) || Category.EqualTo(NonLibraryCategory))
+	{
+		return SNew(SBox)
+			.Padding(FMargin(5, 5, 5, 3))
+			[
+				SNew(STextBlock)
+				.TextStyle(FNiagaraEditorStyle::Get(), "NiagaraEditor.AssetPickerAssetSubcategoryText")
+				.Text(Category)
+			];
+	}
+	return SNew(SBox)
+		.Padding(FMargin(5, 5, 5, 3))
+		[
+			SNew(STextBlock)
+			.TextStyle(FNiagaraEditorStyle::Get(), "NiagaraEditor.AssetPickerAssetCategoryText")
+			.Text(Category)
+		];
+}
+
+const int32 ThumbnailSize = 72;
+
+TSharedRef<SWidget> SNiagaraAssetPickerList::OnGenerateWidgetForItem(const FAssetData& Item)
+{
+	TSharedRef<FAssetThumbnail> AssetThumbnail = MakeShared<FAssetThumbnail>(Item, ThumbnailSize, ThumbnailSize, AssetThumbnailPool);
+	FAssetThumbnailConfig ThumbnailConfig;
+	ThumbnailConfig.bAllowFadeIn = false;
+
+	FText AssetDescription;
+	Item.GetTagValue("TemplateAssetDescription", AssetDescription);
+
+	bool bIsTemplate;
+	bool bIsLibrary;
+	Item.GetTagValue("bIsTemplateAsset", bIsTemplate);
+	bool bFoundLibScriptTag = Item.GetTagValue("bExposeToLibrary", bIsLibrary);
+
+	if (bTemplateOnly
+		|| (bIsTemplate)
+		|| (bFoundLibScriptTag && bIsLibrary))
+	{
+		return
+				SNew(SHorizontalBox)
+				+ SHorizontalBox::Slot()
+				.AutoWidth()
+				.Padding(15, 3, 10, 3)
+				[
+					SNew(SBox)
+					.WidthOverride(ThumbnailSize)
+					.HeightOverride(ThumbnailSize)
+					[
+						AssetThumbnail->MakeThumbnailWidget(ThumbnailConfig)
+					]
+				]
+				+ SHorizontalBox::Slot()
+				.Padding(5, 0, 5, 0)
+				.VAlign(VAlign_Center)
+				[
+					SNew(SVerticalBox)
+					+ SVerticalBox::Slot()
+					.AutoHeight()
+					[
+						SNew(STextBlock)
+						.TextStyle(FNiagaraEditorStyle::Get(), "NiagaraEditor.AssetPickerBoldAssetNameText")
+						.Text(FText::FromString(FName::NameToDisplayString(Item.AssetName.ToString(), false)))
+					]
+					+ SVerticalBox::Slot()
+					[
+						SNew(STextBlock)
+						.Text(AssetDescription)
+						.TextStyle(FNiagaraEditorStyle::Get(), "NiagaraEditor.AssetPickerAssetNameText")
+						.AutoWrapText(true)
+					]
+				];
+	}
+	return 	SNew(SHorizontalBox)
+		+ SHorizontalBox::Slot()
+		.AutoWidth()
+		.Padding(15, 3, 3, 3)
+		[
+			SNew(SBox)
+			.WidthOverride(32.0)
+			.HeightOverride(32.0)
+			[
+				AssetThumbnail->MakeThumbnailWidget(ThumbnailConfig)
+			]
+		]
+		+ SHorizontalBox::Slot()
+		.Padding(3)
+		.VAlign(VAlign_Center)
+		[
+			SNew(STextBlock)
+			.Text(FText::FromString(FName::NameToDisplayString(Item.AssetName.ToString(), false)))
+			.TextStyle(FNiagaraEditorStyle::Get(), "NiagaraEditor.AssetPickerAssetNameText")
+		];
+}
+
+void SNiagaraAssetPickerList::OnItemActivated(const FAssetData& Item)
+{
+	OnTemplateAssetActivated.ExecuteIfBound(Item);
+}
+
+#undef LOCTEXT_NAMESPACE
diff --git a/Engine/Plugins/FX/Niagara/Source/NiagaraEditor/Private/SNiagaraAssetPickerList.h b/Engine/Plugins/FX/Niagara/Source/NiagaraEditor/Private/SNiagaraAssetPickerList.h
new file mode 100644
--- /dev/null
+++ b/Engine/Plugins/FX/Niagara/Source/NiagaraEditor/Private/SNiagaraAssetPickerList.h
@@ -0,0 +1,60 @@
+// Copyright Epic Games, Inc. All Rights Reserved.
+
+#pragma once
+
+#include "SItemSelector.h"
+#include "AssetData.h"
+
+typedef SItemSelector<FText, FAssetData> SNiagaraAssetItemSelector;
+
+class FAssetThumbnailPool;
+
+class SNiagaraAssetPickerList : public SCompoundWidget
+{
+public:
+	DECLARE_DELEGATE_OneParam(FOnTemplateAssetActivated, const FAssetData&);
+
+public:
+	SLATE_BEGIN_ARGS(SNiagaraAssetPickerList) 
+		: _bTemplateOnly(true)
+		, _bAllowMultiSelect(false)
+	{}
+		SLATE_EVENT(FOnTemplateAssetActivated, OnTemplateAssetActivated);
+		SLATE_ARGUMENT(bool, bTemplateOnly)
+		SLATE_ARGUMENT(bool, bAllowMultiSelect);
+	SLATE_END_ARGS();
+
+	void Construct(const FArguments& InArgs, UClass* AssetClass);
+
+	virtual void Tick(const FGeometry& AllottedGeometry, const double InCurrentTime, const float InDeltaTime) override;
+
+	TArray<FAssetData> GetSelectedAssets() const;
+
+private:
+	TArray<FText> OnGetCategoriesForItem(const FAssetData& Item);
+
+	bool OnCompareCategoriesForEquality(const FText& CategoryA, const FText& CategoryB) const;
+
+	bool OnCompareCategoriesForSorting(const FText& CategoryA, const FText& CategoryB) const;
+
+	bool OnCompareItemsForSorting(const FAssetData& ItemA, const FAssetData& ItemB) const;
+
+	bool OnDoesItemMatchFilterText(const FText& FilterText, const FAssetData& Item);
+
+	TSharedRef<SWidget> OnGenerateWidgetForCategory(const FText& Category);
+
+	TSharedRef<SWidget> OnGenerateWidgetForItem(const FAssetData& Item);
+
+	void OnItemActivated(const FAssetData& Item);
+
+private:
+	TSharedPtr<SNiagaraAssetItemSelector> ItemSelector;
+	FText NiagaraPluginCategory;
+	FText ProjectCategory;
+	FText TemplateCategory;
+	FText LibraryCategory;
+	FText NonLibraryCategory;
+	TSharedPtr<FAssetThumbnailPool> AssetThumbnailPool;
+	FOnTemplateAssetActivated OnTemplateAssetActivated;
+	bool bTemplateOnly;
+};
\ No newline at end of file
diff --git a/Engine/Plugins/FX/Niagara/Source/NiagaraEditor/Private/SNiagaraNewAssetDialog.cpp b/Engine/Plugins/FX/Niagara/Source/NiagaraEditor/Private/SNiagaraNewAssetDialog.cpp
--- a/Engine/Plugins/FX/Niagara/Source/NiagaraEditor/Private/SNiagaraNewAssetDialog.cpp
+++ b/Engine/Plugins/FX/Niagara/Source/NiagaraEditor/Private/SNiagaraNewAssetDialog.cpp
@@ -3,7 +3,7 @@
 #include "SNiagaraNewAssetDialog.h"
 #include "NiagaraEmitter.h"
 #include "NiagaraEditorStyle.h"
-#include "SNiagaraTemplateAssetPicker.h"
+#include "SNiagaraAssetPickerList.h"
 #include "NiagaraEditorSettings.h"
 
 #include "AssetData.h"
@@ -12,7 +12,7 @@
 #include "Widgets/Input/SButton.h"
 #include "Widgets/Layout/SUniformGridPanel.h"
 #include "EditorStyleSet.h"
-
+#include "Widgets/Workflow/SWizard.h"
 #include "ContentBrowserModule.h"
 #include "IContentBrowserSingleton.h"
 
@@ -38,21 +38,7 @@ void SNiagaraNewAssetDialog::Construct(const FArguments& InArgs, FName InSaveCon
 	TSharedRef<SVerticalBox> RootBox =
 		SNew(SVerticalBox)
 
-		// Options label
-		+ SVerticalBox::Slot()
-		.Padding(0, 7, 0, 0)
-		.AutoHeight()
-		[
-			SNew(SBox)
-			.Padding(FEditorStyle::GetMargin("StandardDialog.SlotPadding"))
-			[
-				SNew(STextBlock)
-				.TextStyle(FNiagaraEditorStyle::Get(), "NiagaraEditor.NewAssetDialog.HeaderText")
-				.Text(LOCTEXT("OptionsLabel", "Select an Option"))
-			]
-		]
-
-		// Creation mode radio buttons.
+		// Creation mode toggle buttons.
 		+ SVerticalBox::Slot()
 		.Padding(0, 5, 0, 5)
 		.AutoHeight()
@@ -68,67 +54,9 @@ void SNiagaraNewAssetDialog::Construct(const FArguments& InArgs, FName InSaveCon
 					SAssignNew(OptionsBox, SVerticalBox)
 				]
 			]
-		]
-
-		// Asset pickers label
-		+ SVerticalBox::Slot()
-		.Padding(0, 5, 0, 0)
-		.AutoHeight()
-		[
-			SNew(SBox)
-			.Padding(FEditorStyle::GetMargin("StandardDialog.SlotPadding"))
-			[
-				SNew(STextBlock)
-				.TextStyle(FNiagaraEditorStyle::Get(), "NiagaraEditor.NewAssetDialog.HeaderText")
-				.Text(this, &SNiagaraNewAssetDialog::GetAssetPickersLabelText)
-			]
-		]
-
-		// Asset pickers
-		+ SVerticalBox::Slot()
-		.Padding(0, 5, 0, 5)
-		[
-			SNew(SBox)
-			.Padding(FEditorStyle::GetMargin("StandardDialog.SlotPadding"))
-			[
-				SNew(SBorder)
-				.BorderImage(FEditorStyle::GetBrush("ToolPanel.GroupBorder"))
-				.Padding(FMargin(7))
-				[
-					SAssignNew(AssetPickerOverlay, SOverlay)
-				]
-			]
-		]
-
-		// OK/Cancel buttons
-		+SVerticalBox::Slot()
-		.AutoHeight()
-		.HAlign(HAlign_Right)
-		.Padding(0, 5, 0, 5)
-		[
-			SNew(SUniformGridPanel)
-			.SlotPadding(FEditorStyle::GetMargin("StandardDialog.SlotPadding"))
-			.MinDesiredSlotWidth(FEditorStyle::GetFloat("StandardDialog.MinDesiredSlotWidth"))
-			.MinDesiredSlotHeight(FEditorStyle::GetFloat("StandardDialog.MinDesiredSlotHeight"))
-			+SUniformGridPanel::Slot(0, 0)
-			[
-				SNew(SButton)
-				.HAlign(HAlign_Center)
-				.ContentPadding(FEditorStyle::GetMargin("StandardDialog.ContentPadding"))
-				.Text(LOCTEXT("OK", "OK"))
-				.OnClicked(this, &SNiagaraNewAssetDialog::OnOkButtonClicked)
-				.IsEnabled(this, &SNiagaraNewAssetDialog::IsOkButtonEnabled)
-			]
-			+SUniformGridPanel::Slot(1, 0)
-			[
-				SNew(SButton)
-				.HAlign(HAlign_Center)
-				.ContentPadding(FEditorStyle::GetMargin("StandardDialog.ContentPadding"))
-				.Text(LOCTEXT("Cancel", "Cancel"))
-				.OnClicked(this, &SNiagaraNewAssetDialog::OnCancelButtonClicked)
-			]
 		];
 
+
 	int32 OptionIndex = 0;
 	for (FNiagaraNewAssetDialogOption& Option : Options)
 	{
@@ -140,28 +68,35 @@ void SNiagaraNewAssetDialog::Construct(const FArguments& InArgs, FName InSaveCon
 				.BorderImage(FNiagaraEditorStyle::Get().GetBrush("NiagaraEditor.NewAssetDialog.SubBorder"))
 				.BorderBackgroundColor(this, &SNiagaraNewAssetDialog::GetOptionBorderColor, OptionIndex)
 				[
-					SNew(SCheckBox)
-					.Style(FCoreStyle::Get(), "RadioButton")
-					.CheckBoxContentUsesAutoWidth(false)
-					.IsChecked(this, &SNiagaraNewAssetDialog::GetOptionCheckBoxState, OptionIndex)
-					.OnCheckStateChanged(this, &SNiagaraNewAssetDialog::OptionCheckBoxStateChanged, OptionIndex)
-					.Content()
-					[
-						SNew(STextBlock)
-						.TextStyle(FNiagaraEditorStyle::Get(), "NiagaraEditor.NewAssetDialog.OptionText")
-						.ColorAndOpacity(this, &SNiagaraNewAssetDialog::GetOptionTextColor, OptionIndex)
-						.Text(Option.OptionText)
-						.AutoWrapText(true)
-					]
-				]
-			];
-
-		AssetPickerOverlay->AddSlot()
-			[
-				SNew(SBox)
-				.Visibility(this, &SNiagaraNewAssetDialog::GetAssetPickerVisibility, OptionIndex)
-				[
-					Option.AssetPicker
+						SNew(SCheckBox)
+						.Style(FEditorStyle::Get(), "ToggleButtonCheckbox")
+						.CheckBoxContentUsesAutoWidth(false)
+						.IsChecked(this, &SNiagaraNewAssetDialog::GetOptionCheckBoxState, OptionIndex)
+						.OnCheckStateChanged(this, &SNiagaraNewAssetDialog::OptionCheckBoxStateChanged, OptionIndex)
+						.Content()
+						[
+							SNew(SVerticalBox)
+							+ SVerticalBox::Slot()
+							.AutoHeight()
+							.Padding(5, 2)
+							[		
+								SNew(STextBlock)
+								.TextStyle(FNiagaraEditorStyle::Get(), "NiagaraEditor.NewAssetDialog.OptionText")
+								.ColorAndOpacity(this, &SNiagaraNewAssetDialog::GetOptionTextColor, OptionIndex)
+								.Text(Option.OptionText)
+								.AutoWrapText(true)
+							]
+							+ SVerticalBox::Slot()
+							.AutoHeight()
+							.Padding(5, 2, 5, 7)
+							[
+								SNew(STextBlock)
+								.ColorAndOpacity(this, &SNiagaraNewAssetDialog::GetOptionTextColor, OptionIndex)
+								.Text(Option.OptionDescription)
+								.AutoWrapText(true)
+							]
+						]
+					
 				]
 			];
 
@@ -175,10 +110,40 @@ void SNiagaraNewAssetDialog::Construct(const FArguments& InArgs, FName InSaveCon
 		.SupportsMaximize(false)
 		.SupportsMinimize(false)
 		[
-			RootBox
+			SNew(SWizard)
+			.OnCanceled(this, &SNiagaraNewAssetDialog::OnCancelButtonClicked)
+			.OnFinished(this, &SNiagaraNewAssetDialog::OnOkButtonClicked)
+			.CanFinish(this, &SNiagaraNewAssetDialog::IsOkButtonEnabled)
+			.ShowPageList(false)
+			.ButtonStyle(FEditorStyle::Get(), "FlatButton.Default")
+			.CancelButtonStyle(FEditorStyle::Get(), "FlatButton.Default")
+			.FinishButtonStyle(FEditorStyle::Get(), "FlatButton.Success")
+			+SWizard::Page()
+			.CanShow(true)
+			.OnEnter(this, &SNiagaraNewAssetDialog::ResetStage)
+			[
+				RootBox
+			]
+			+ SWizard::Page()
+			.CanShow(this, &SNiagaraNewAssetDialog::HasAssetPage)
+			.OnEnter(this, &SNiagaraNewAssetDialog::GetAssetPicker)
+			[
+				SAssignNew(AssetSettingsPage, SBox)
+			]
 		]);
 }
 
+void SNiagaraNewAssetDialog::GetAssetPicker()
+{
+	bOnAssetStage = true;
+	AssetSettingsPage->SetContent(Options[SelectedOptionIndex].AssetPicker);
+}
+
+void SNiagaraNewAssetDialog::ResetStage()
+{
+	bOnAssetStage = false;
+}
+
 bool SNiagaraNewAssetDialog::GetUserConfirmedSelection() const
 {
 	return bUserConfirmedSelection;
@@ -239,11 +204,6 @@ FText SNiagaraNewAssetDialog::GetAssetPickersLabelText() const
 	return Options[SelectedOptionIndex].AssetPickerHeader;
 }
 
-EVisibility SNiagaraNewAssetDialog::GetAssetPickerVisibility(int32 OptionIndex) const
-{
-	return  SelectedOptionIndex == OptionIndex ? EVisibility::Visible : EVisibility::Collapsed;
-}
-
 bool SNiagaraNewAssetDialog::IsOkButtonEnabled() const
 {
 	const FNiagaraNewAssetDialogOption& SelectedOption = Options[SelectedOptionIndex];
@@ -251,7 +211,7 @@ bool SNiagaraNewAssetDialog::IsOkButtonEnabled() const
 	{
 		TArray<FAssetData> TempSelectedAssets;
 		SelectedOption.OnGetSelectedAssetsFromPicker.Execute(TempSelectedAssets);
-		return TempSelectedAssets.Num() != 0;
+		return bOnAssetStage && TempSelectedAssets.Num() != 0;
 	}
 	else
 	{
@@ -259,19 +219,22 @@ bool SNiagaraNewAssetDialog::IsOkButtonEnabled() const
 	}
 }
 
-FReply SNiagaraNewAssetDialog::OnOkButtonClicked()
+void SNiagaraNewAssetDialog::OnOkButtonClicked()
 {
 	ConfirmSelection();
-	return FReply::Handled();
 }
 
-FReply SNiagaraNewAssetDialog::OnCancelButtonClicked()
+void SNiagaraNewAssetDialog::OnCancelButtonClicked()
 {
 	bUserConfirmedSelection = false;
 	SelectedAssets.Empty();
 
 	RequestDestroyWindow();
-	return FReply::Handled();
+}
+
+bool SNiagaraNewAssetDialog::HasAssetPage() const
+{
+	return !IsOkButtonEnabled();
 }
 
 void SNiagaraNewAssetDialog::SaveConfig()
diff --git a/Engine/Plugins/FX/Niagara/Source/NiagaraEditor/Private/SNiagaraNewAssetDialog.h b/Engine/Plugins/FX/Niagara/Source/NiagaraEditor/Private/SNiagaraNewAssetDialog.h
--- a/Engine/Plugins/FX/Niagara/Source/NiagaraEditor/Private/SNiagaraNewAssetDialog.h
+++ b/Engine/Plugins/FX/Niagara/Source/NiagaraEditor/Private/SNiagaraNewAssetDialog.h
@@ -6,8 +6,11 @@
 #include "Widgets/SWindow.h"
 #include "AssetData.h"
 #include "IAssetTypeActions.h"
+#include "ContentBrowserDelegates.h"
 
-class SNiagaraTemplateAssetPicker;
+
+class SNiagaraAssetPickerList;
+class SBox;
 
 /** A modal dialog to collect information needed to create a new niagara system. */
 class SNiagaraNewAssetDialog : public SWindow
@@ -21,13 +24,15 @@ public:
 	{
 	public:
 		FText OptionText;
+		FText OptionDescription;
 		FText AssetPickerHeader;
 		TSharedRef<SWidget> AssetPicker;
 		FOnGetSelectedAssetsFromPicker OnGetSelectedAssetsFromPicker;
 		FOnSelectionConfirmed OnSelectionConfirmed;
 
-		FNiagaraNewAssetDialogOption(FText InOptionText, FText InAssetPickerHeader, FOnGetSelectedAssetsFromPicker InOnGetSelectedAssetsFromPicker, FOnSelectionConfirmed InOnSelecitonConfirmed, TSharedRef<SWidget> InAssetPicker)
+		FNiagaraNewAssetDialogOption(FText InOptionText, FText InOptionDescription, FText InAssetPickerHeader, FOnGetSelectedAssetsFromPicker InOnGetSelectedAssetsFromPicker, FOnSelectionConfirmed InOnSelecitonConfirmed, TSharedRef<SWidget> InAssetPicker)
 			: OptionText(InOptionText)
+			, OptionDescription(InOptionDescription)
 			, AssetPickerHeader(InAssetPickerHeader)
 			, AssetPicker(InAssetPicker)
 			, OnGetSelectedAssetsFromPicker(InOnGetSelectedAssetsFromPicker)
@@ -42,15 +47,21 @@ public:
 	SLATE_END_ARGS()
 
 	void Construct(const FArguments& InArgs, FName InSaveConfigKey, FText AssetTypeDisplayName, TArray<FNiagaraNewAssetDialogOption> InOptions);
-
+	void GetAssetPicker();
+	void ResetStage();
 	bool GetUserConfirmedSelection() const;
 
 protected:
 	const TArray<FAssetData>& GetSelectedAssets() const;
 
 	void ConfirmSelection();
+	int32 GetSelectedObjectIndex() const { return SelectedOptionIndex; };
+
+protected:
+	TSharedPtr<SBox> AssetSettingsPage;
 
 private:
+
 	void OnWindowClosed(const TSharedRef<SWindow>& Window);
 
 	FSlateColor GetOptionBorderColor(int32 OptionIndex) const;
@@ -58,28 +69,20 @@ private:
 	ECheckBoxState GetOptionCheckBoxState(int32 OptionIndex) const;
 
 	void OptionCheckBoxStateChanged(ECheckBoxState InCheckBoxState, int32 OptionIndex);
-
 	FSlateColor GetOptionTextColor(int32 OptionIndex) const;
 
 	FText GetAssetPickersLabelText() const;
-
-	EVisibility GetAssetPickerVisibility(int32 OptionIndex) const;
-
 	bool IsOkButtonEnabled() const;
-
-	FReply OnOkButtonClicked();
-
-	FReply OnCancelButtonClicked();
-
+	void OnOkButtonClicked();
+	void OnCancelButtonClicked();
+	bool HasAssetPage() const;
 	void SaveConfig();
 
 private:
 	FName SaveConfigKey;
-
 	TArray<FNiagaraNewAssetDialogOption> Options;
-
 	int32 SelectedOptionIndex;
-
 	bool bUserConfirmedSelection;
 	TArray<FAssetData> SelectedAssets;
+	bool bOnAssetStage;
 };
\ No newline at end of file
diff --git a/Engine/Plugins/FX/Niagara/Source/NiagaraEditor/Private/SNiagaraTemplateAssetPicker.cpp b/Engine/Plugins/FX/Niagara/Source/NiagaraEditor/Private/SNiagaraTemplateAssetPicker.cpp
deleted file mode 100644
diff --git a/Engine/Plugins/FX/Niagara/Source/NiagaraEditor/Private/SNiagaraTemplateAssetPicker.h b/Engine/Plugins/FX/Niagara/Source/NiagaraEditor/Private/SNiagaraTemplateAssetPicker.h
deleted file mode 100644
diff --git a/Engine/Plugins/FX/Niagara/Source/NiagaraEditor/Private/Toolkits/NiagaraSystemToolkit.cpp b/Engine/Plugins/FX/Niagara/Source/NiagaraEditor/Private/Toolkits/NiagaraSystemToolkit.cpp
--- a/Engine/Plugins/FX/Niagara/Source/NiagaraEditor/Private/Toolkits/NiagaraSystemToolkit.cpp
+++ b/Engine/Plugins/FX/Niagara/Source/NiagaraEditor/Private/Toolkits/NiagaraSystemToolkit.cpp
@@ -52,6 +52,7 @@
 #include "NiagaraMessageLogViewModel.h"
 #include "ViewModels/NiagaraOverviewGraphViewModel.h"
 #include "NiagaraScriptSourceBase.h"
+#include "Widgets/Input/SCheckBox.h"
 
 #define LOCTEXT_NAMESPACE "NiagaraSystemEditor"
 
@@ -87,6 +88,9 @@ static FAutoConsoleVariableRef CVarShowNiagaraDeveloperWindows(
 	ECVF_Default
 );
 
+bool FNiagaraSystemToolkit::bShowLibraryOnly = false;
+bool FNiagaraSystemToolkit::bShowTemplateOnly = false;
+
 void FNiagaraSystemToolkit::RegisterTabSpawners(const TSharedRef<class FTabManager>& InTabManager)
 {
 	WorkspaceMenuCategory = InTabManager->AddLocalWorkspaceMenuCategory(LOCTEXT("WorkspaceMenu_NiagaraSystemEditor", "Niagara System"));
@@ -880,18 +884,124 @@ TSharedRef<SWidget> FNiagaraSystemToolkit::CreateAddEmitterMenuContent()
 		AssetPickerConfig.bAllowNullSelection = false;
 		AssetPickerConfig.InitialAssetViewType = EAssetViewType::List;
 		AssetPickerConfig.Filter.ClassNames.Add(UNiagaraEmitter::StaticClass()->GetFName());
+		AssetPickerConfig.OnShouldFilterAsset.BindSP(this, &FNiagaraSystemToolkit::ShouldFilterEmitter);
+		AssetPickerConfig.RefreshAssetViewDelegates.Add(&RefreshAssetView);
 	}
 
 	FContentBrowserModule& ContentBrowserModule = FModuleManager::Get().LoadModuleChecked<FContentBrowserModule>(TEXT("ContentBrowser"));
 
-	return SNew(SBox)
-		.WidthOverride(300.0f)
-		.HeightOverride(300.f)
+	return SNew(SVerticalBox)
+		+ SVerticalBox::Slot()
+		.Padding(3)
+		.AutoHeight()
+		[
+			SNew(SHorizontalBox)
+			+ SHorizontalBox::Slot()
+			.Padding(FEditorStyle::GetMargin("StandardDialog.SlotPadding"))
+			.FillWidth(1.0f)
+			.HAlign(HAlign_Right)
+			[
+				SNew(SCheckBox)
+				.OnCheckStateChanged(this, &FNiagaraSystemToolkit::TemplateCheckBoxStateChanged)
+				.IsChecked(this, &FNiagaraSystemToolkit::GetTemplateCheckBoxState)
+				[
+					SNew(STextBlock)
+					.Text(LOCTEXT("TemplateOnly", "Template Only"))
+				]
+			]
+			+ SHorizontalBox::Slot()
+			.Padding(FEditorStyle::GetMargin("StandardDialog.SlotPadding"))
+			.AutoWidth()
+			.HAlign(HAlign_Right)
+			[
+				SNew(SCheckBox)
+		 		.OnCheckStateChanged(this, &FNiagaraSystemToolkit::LibraryCheckBoxStateChanged)
+				.IsChecked(this, &FNiagaraSystemToolkit::GetLibraryCheckBoxState)
+				[
+					SNew(STextBlock)
+					.Text(LOCTEXT("LibraryOnly", "Library Only"))
+				]
+			]
+		]
+		+SVerticalBox::Slot()
+		.FillHeight(1.0f)
 		[
-			ContentBrowserModule.Get().CreateAssetPicker(AssetPickerConfig)
+			SNew(SBox)
+			.HeightOverride(300.f)
+			.WidthOverride(300.f)
+			[
+				ContentBrowserModule.Get().CreateAssetPicker(AssetPickerConfig)
+			]
 		];
 }
 
+void FNiagaraSystemToolkit::LibraryCheckBoxStateChanged(ECheckBoxState InCheckbox)
+{
+	FNiagaraSystemToolkit::bShowLibraryOnly = (InCheckbox == ECheckBoxState::Checked);
+	RefreshAssetView.ExecuteIfBound(true);
+}
+
+ECheckBoxState FNiagaraSystemToolkit::GetLibraryCheckBoxState() const
+{
+	return FNiagaraSystemToolkit::bShowLibraryOnly ? ECheckBoxState::Checked : ECheckBoxState::Unchecked;
+}
+
+void FNiagaraSystemToolkit::TemplateCheckBoxStateChanged(ECheckBoxState InCheckbox)
+{
+	FNiagaraSystemToolkit::bShowTemplateOnly = (InCheckbox == ECheckBoxState::Checked);
+	RefreshAssetView.ExecuteIfBound(true);
+}
+
+ECheckBoxState FNiagaraSystemToolkit::GetTemplateCheckBoxState() const
+{
+	return FNiagaraSystemToolkit::bShowTemplateOnly ? ECheckBoxState::Checked : ECheckBoxState::Unchecked;
+}
+
+bool FNiagaraSystemToolkit::ShouldFilterEmitter(const FAssetData& AssetData)
+{
+	// Check if library script
+	bool bScriptAllowed = true;
+	bool bInLibrary = false;
+	bool bIsTemplate = false;
+	if (FNiagaraSystemToolkit::bShowLibraryOnly == true)
+	{
+		bool bFoundLibraryTag = AssetData.GetTagValue(GET_MEMBER_NAME_CHECKED(UNiagaraEmitter, bExposeToLibrary), bInLibrary);
+
+		if (bFoundLibraryTag == false)
+		{
+			if (AssetData.IsAssetLoaded())
+			{
+				UNiagaraEmitter* EmitterAsset = static_cast<UNiagaraEmitter*>(AssetData.GetAsset());
+				if (EmitterAsset != nullptr)
+				{
+					bInLibrary = EmitterAsset->bExposeToLibrary;
+				}
+			}
+		}
+		bScriptAllowed &= bFoundLibraryTag;
+	}
+	if (FNiagaraSystemToolkit::bShowTemplateOnly == true)
+	{
+		bool bFoundTemplateTag = AssetData.GetTagValue(GET_MEMBER_NAME_CHECKED(UNiagaraEmitter, bIsTemplateAsset), bIsTemplate);
+
+		if (bFoundTemplateTag == false)
+		{
+			if (AssetData.IsAssetLoaded())
+			{
+				UNiagaraEmitter* EmitterAsset = static_cast<UNiagaraEmitter*>(AssetData.GetAsset());
+				if (EmitterAsset != nullptr)
+				{
+					bIsTemplate = EmitterAsset->bIsTemplateAsset;
+				}
+			}
+		}
+		bScriptAllowed &= bIsTemplate;
+	}
+
+	return !bScriptAllowed;
+}
+
+
 TSharedRef<SWidget> FNiagaraSystemToolkit::GenerateCompileMenuContent()
 {
 	const bool bShouldCloseWindowAfterMenuSelection = true;
diff --git a/Engine/Plugins/FX/Niagara/Source/NiagaraEditor/Private/Toolkits/NiagaraSystemToolkit.h b/Engine/Plugins/FX/Niagara/Source/NiagaraEditor/Private/Toolkits/NiagaraSystemToolkit.h
--- a/Engine/Plugins/FX/Niagara/Source/NiagaraEditor/Private/Toolkits/NiagaraSystemToolkit.h
+++ b/Engine/Plugins/FX/Niagara/Source/NiagaraEditor/Private/Toolkits/NiagaraSystemToolkit.h
@@ -14,6 +14,7 @@
 #include "ISequencerTrackEditor.h"
 
 #include "NiagaraScript.h"
+#include "ContentBrowserDelegates.h"
 
 class FNiagaraSystemInstance;
 class FNiagaraSystemViewModel;
@@ -71,6 +72,9 @@ public:
 
 	TSharedPtr<FNiagaraSystemViewModel> GetSystemViewModel();
 
+public:
+	FRefreshAssetViewDelegate RefreshAssetView;
+
 protected:
 	void OnToggleBounds();
 	bool IsToggleBoundsChecked() const;
@@ -112,6 +116,11 @@ private:
 
 	void GetSequencerAddMenuContent(FMenuBuilder& MenuBuilder, TSharedRef<ISequencer> Sequencer);
 	TSharedRef<SWidget> CreateAddEmitterMenuContent();
+	void LibraryCheckBoxStateChanged(ECheckBoxState InCheckbox);
+	ECheckBoxState GetLibraryCheckBoxState() const;
+	void TemplateCheckBoxStateChanged(ECheckBoxState InCheckbox);
+	ECheckBoxState GetTemplateCheckBoxState() const;
+	bool ShouldFilterEmitter(const FAssetData& AssetData);
 	TSharedRef<SWidget> GenerateCompileMenuContent();
 
 	void EmitterAssetSelected(const FAssetData& AssetData);
@@ -179,4 +188,7 @@ private:
 	static const FName GeneratedCodeTabID;
 	static const FName MessageLogTabID;
 	static const FName SystemOverviewTabID;
+
+	static bool bShowLibraryOnly;
+	static bool bShowTemplateOnly;
 };
diff --git a/Engine/Plugins/FX/Niagara/Source/NiagaraEditor/Public/NiagaraEditorModule.h b/Engine/Plugins/FX/Niagara/Source/NiagaraEditor/Public/NiagaraEditorModule.h
--- a/Engine/Plugins/FX/Niagara/Source/NiagaraEditor/Public/NiagaraEditorModule.h
+++ b/Engine/Plugins/FX/Niagara/Source/NiagaraEditor/Public/NiagaraEditorModule.h
@@ -131,6 +131,7 @@ private:
 	}
 
 	void TestCompileScriptFromConsole(const TArray<FString>& Arguments);
+	void ReinitializeStyle();
 
 private:
 	TSharedPtr<FExtensibilityManager> MenuExtensibilityManager;
@@ -181,4 +182,6 @@ private:
 	bool bThumbnailRenderersRegistered;
 
 	TSharedRef<FNiagaraClipboard> Clipboard;
+
+	IConsoleCommand* ReinitializeStyleCommand;
 };
diff --git a/Engine/Plugins/FX/Niagara/Source/NiagaraEditorWidgets/Private/SNiagaraOverviewGraph.cpp b/Engine/Plugins/FX/Niagara/Source/NiagaraEditorWidgets/Private/SNiagaraOverviewGraph.cpp
--- a/Engine/Plugins/FX/Niagara/Source/NiagaraEditorWidgets/Private/SNiagaraOverviewGraph.cpp
+++ b/Engine/Plugins/FX/Niagara/Source/NiagaraEditorWidgets/Private/SNiagaraOverviewGraph.cpp
@@ -25,9 +25,13 @@
 #include "EdGraph/EdGraph.h"
 #include "EdGraph/EdGraphNode.h"
 #include "Framework/Commands/GenericCommands.h"
+#include "Widgets/Input/SCheckBox.h"
 
 #define LOCTEXT_NAMESPACE "NiagaraOverviewGraph"
 
+bool SNiagaraOverviewGraph::bShowLibraryOnly = false;
+bool SNiagaraOverviewGraph::bShowTemplateOnly = false;
+
 void SNiagaraOverviewGraph::Construct(const FArguments& InArgs, TSharedRef<FNiagaraOverviewGraphViewModel> InViewModel)
 {
 	ViewModel = InViewModel;
@@ -294,6 +298,8 @@ void SNiagaraOverviewGraph::CreateAddEmitterMenuContent(FMenuBuilder& MenuBuilde
 		AssetPickerConfig.bAllowNullSelection = false;
 		AssetPickerConfig.InitialAssetViewType = EAssetViewType::List;
 		AssetPickerConfig.Filter.ClassNames.Add(UNiagaraEmitter::StaticClass()->GetFName());
+		AssetPickerConfig.OnShouldFilterAsset.BindSP(this, &SNiagaraOverviewGraph::ShouldFilterEmitter);
+		AssetPickerConfig.RefreshAssetViewDelegates.Add(&RefreshAssetView);
 	}
 
 	FContentBrowserModule& ContentBrowserModule = FModuleManager::Get().LoadModuleChecked<FContentBrowserModule>(TEXT("ContentBrowser"));
@@ -301,13 +307,51 @@ void SNiagaraOverviewGraph::CreateAddEmitterMenuContent(FMenuBuilder& MenuBuilde
 	TSharedRef<SWidget> EmitterAddSubMenu =
 		SNew(SBorder)
 		.BorderImage(FEditorStyle::GetBrush("Menu.Background"))
-		.Padding(5.0f)
+		.Padding(0.0f)
 		[
-			SNew(SBox)
-			.WidthOverride(300.0f)
-			.HeightOverride(300.0f)
+			SNew(SVerticalBox)
+			+ SVerticalBox::Slot()
+			.Padding(3)
+			.AutoHeight()
 			[
-				ContentBrowserModule.Get().CreateAssetPicker(AssetPickerConfig)
+				SNew(SHorizontalBox)
+				+ SHorizontalBox::Slot()
+				.Padding(FEditorStyle::GetMargin("StandardDialog.SlotPadding"))
+				.FillWidth(1.0f)
+				.HAlign(HAlign_Right)
+				[
+					SNew(SCheckBox)
+					.OnCheckStateChanged(this, &SNiagaraOverviewGraph::TemplateCheckBoxStateChanged)
+					.IsChecked(this, &SNiagaraOverviewGraph::GetTemplateCheckBoxState)
+					[
+						SNew(STextBlock)
+						.Text(LOCTEXT("TemplateOnly", "Template Only"))
+					]
+				]
+				+ SHorizontalBox::Slot()
+				.Padding(FEditorStyle::GetMargin("StandardDialog.SlotPadding"))
+				.AutoWidth()
+				.HAlign(HAlign_Right)
+				[
+					SNew(SCheckBox)
+					.OnCheckStateChanged(this, &SNiagaraOverviewGraph::LibraryCheckBoxStateChanged)
+					.IsChecked(this, &SNiagaraOverviewGraph::GetLibraryCheckBoxState)
+					[
+						SNew(STextBlock)
+						.Text(LOCTEXT("LibraryOnly", "Library Only"))
+					]
+				]
+
+			]
+			+ SVerticalBox::Slot()
+			.FillHeight(1.0f)
+			[
+				SNew(SBox)
+				.WidthOverride(300.0f)
+				.HeightOverride(300.0f)
+				[
+					ContentBrowserModule.Get().CreateAssetPicker(AssetPickerConfig)
+				]
 			]
 		];
 
@@ -365,4 +409,72 @@ void SNiagaraOverviewGraph::OnDistributeNodesV()
 	}
 }
 
+void SNiagaraOverviewGraph::LibraryCheckBoxStateChanged(ECheckBoxState InCheckbox)
+{
+	SNiagaraOverviewGraph::bShowLibraryOnly = (InCheckbox == ECheckBoxState::Checked);
+	RefreshAssetView.ExecuteIfBound(true);
+}
+
+ECheckBoxState SNiagaraOverviewGraph::GetLibraryCheckBoxState() const
+{
+	return SNiagaraOverviewGraph::bShowLibraryOnly ? ECheckBoxState::Checked : ECheckBoxState::Unchecked;
+}
+
+void SNiagaraOverviewGraph::TemplateCheckBoxStateChanged(ECheckBoxState InCheckbox)
+{
+	SNiagaraOverviewGraph::bShowTemplateOnly = (InCheckbox == ECheckBoxState::Checked);
+	RefreshAssetView.ExecuteIfBound(true);
+}
+
+ECheckBoxState SNiagaraOverviewGraph::GetTemplateCheckBoxState() const
+{
+	return SNiagaraOverviewGraph::bShowTemplateOnly ? ECheckBoxState::Checked : ECheckBoxState::Unchecked;
+}
+
+
+bool SNiagaraOverviewGraph::ShouldFilterEmitter(const FAssetData& AssetData)
+{
+	// Check if library script
+	bool bScriptAllowed = true;
+	bool bInLibrary = false;
+	bool bIsTemplate = false;
+	if (SNiagaraOverviewGraph::bShowLibraryOnly == true)
+	{
+		bool bFoundLibraryTag = AssetData.GetTagValue(GET_MEMBER_NAME_CHECKED(UNiagaraEmitter, bExposeToLibrary), bInLibrary);
+
+		if (bFoundLibraryTag == false)
+		{
+			if (AssetData.IsAssetLoaded())
+			{
+				UNiagaraEmitter* EmitterAsset = static_cast<UNiagaraEmitter*>(AssetData.GetAsset());
+				if (EmitterAsset != nullptr)
+				{
+					bInLibrary = EmitterAsset->bExposeToLibrary;
+				}
+			}
+		}
+		bScriptAllowed &= bFoundLibraryTag;
+	}
+	if (SNiagaraOverviewGraph::bShowTemplateOnly == true)
+	{
+		bool bFoundTemplateTag = AssetData.GetTagValue(GET_MEMBER_NAME_CHECKED(UNiagaraEmitter, bIsTemplateAsset), bIsTemplate);
+
+		if (bFoundTemplateTag == false)
+		{
+			if (AssetData.IsAssetLoaded())
+			{
+				UNiagaraEmitter* EmitterAsset = static_cast<UNiagaraEmitter*>(AssetData.GetAsset());
+				if (EmitterAsset != nullptr)
+				{
+					bIsTemplate = EmitterAsset->bIsTemplateAsset;
+				}
+			}
+		}
+		bScriptAllowed &= bIsTemplate;
+	}
+
+	return !bScriptAllowed;
+}
+
+
 #undef LOCTEXT_NAMESPACE // "NiagaraOverviewGraph"
\ No newline at end of file
diff --git a/Engine/Plugins/FX/Niagara/Source/NiagaraEditorWidgets/Private/SNiagaraOverviewGraph.h b/Engine/Plugins/FX/Niagara/Source/NiagaraEditorWidgets/Private/SNiagaraOverviewGraph.h
--- a/Engine/Plugins/FX/Niagara/Source/NiagaraEditorWidgets/Private/SNiagaraOverviewGraph.h
+++ b/Engine/Plugins/FX/Niagara/Source/NiagaraEditorWidgets/Private/SNiagaraOverviewGraph.h
@@ -4,6 +4,7 @@
 
 #include "Widgets/SCompoundWidget.h"
 #include "GraphEditor.h"
+#include "ContentBrowserDelegates.h"
 
 class FNiagaraOverviewGraphViewModel;
 struct FActionMenuContent;
@@ -21,6 +22,9 @@ public:
 
 	virtual void Tick(const FGeometry& AllottedGeometry, const double InCurrentTime, const float InDeltaTime) override;
 
+public:
+	FRefreshAssetViewDelegate RefreshAssetView;
+
 private:
 	void ViewModelSelectionChanged();
 
@@ -48,6 +52,12 @@ private:
 	void OnDistributeNodesH();
 	void OnDistributeNodesV();
 
+	void LibraryCheckBoxStateChanged(ECheckBoxState InCheckbox);
+	ECheckBoxState GetLibraryCheckBoxState() const;
+	void TemplateCheckBoxStateChanged(ECheckBoxState InCheckbox);
+	ECheckBoxState GetTemplateCheckBoxState() const;
+
+	bool ShouldFilterEmitter(const FAssetData& AssetData);
 private:
 	TSharedPtr<FNiagaraOverviewGraphViewModel> ViewModel;
 	TSharedPtr<SGraphEditor> GraphEditor;
@@ -56,4 +66,7 @@ private:
 	bool bUpdatingGraphSelectionFromViewModel;
 
 	int32 ZoomToFitFrameDelay;
+
+	static bool bShowLibraryOnly;
+	static bool bShowTemplateOnly;
 };
\ No newline at end of file
