Increasing maximum stack node width and decreasing item row padding
#rnx

diff --git a/Engine/Plugins/FX/Niagara/Source/NiagaraEditorWidgets/Private/NiagaraEditorWidgetsStyle.cpp b/Engine/Plugins/FX/Niagara/Source/NiagaraEditorWidgets/Private/NiagaraEditorWidgetsStyle.cpp
--- a/Engine/Plugins/FX/Niagara/Source/NiagaraEditorWidgets/Private/NiagaraEditorWidgetsStyle.cpp
+++ b/Engine/Plugins/FX/Niagara/Source/NiagaraEditorWidgets/Private/NiagaraEditorWidgetsStyle.cpp
@@ -67,12 +67,10 @@ TSharedRef< FSlateStyleSet > FNiagaraEditorWidgetsStyle::Create()
 	// Stack
 	Style->Set("NiagaraEditor.Stack.IconSize", FVector2D(18.0f, 18.0f));
 
-	FSlateFontInfo StackGroupFont = DEFAULT_FONT("Bold", 10);
+	FSlateFontInfo StackGroupFont = DEFAULT_FONT("Bold", 12);
 	FTextBlockStyle StackGroupText = FTextBlockStyle(NormalText)
 		.SetFont(StackGroupFont)
-		.SetColorAndOpacity(FLinearColor(1.0f, 1.0f, 1.0f, 1.0f))
-		.SetShadowOffset(FVector2D(0, 1))
-		.SetShadowColorAndOpacity(FLinearColor(0, 0, 0, 0.9f));
+		.SetColorAndOpacity(FLinearColor(1.0f, 1.0f, 1.0f, 1.0f));
 	Style->Set("NiagaraEditor.Stack.GroupText", StackGroupText);
 
 	FSlateFontInfo StackDefaultFont = DEFAULT_FONT("Regular", 10);
@@ -82,9 +80,9 @@ TSharedRef< FSlateStyleSet > FNiagaraEditorWidgetsStyle::Create()
 
 	FSlateFontInfo StackCategoryFont = DEFAULT_FONT("Bold", 10);
 	FTextBlockStyle StackCategoryText = FTextBlockStyle(NormalText)
-		.SetFont(StackCategoryFont)
-		.SetShadowOffset(FVector2D(1, 1));
+		.SetFont(StackCategoryFont);
 	Style->Set("NiagaraEditor.Stack.CategoryText", StackCategoryText);
+	Style->Set("NiagaraEditor.SystemOverview.GroupHeaderText", StackCategoryText);
 
 	FSlateFontInfo ParameterFont = DEFAULT_FONT("Regular", 8);
 	FTextBlockStyle ParameterText = FTextBlockStyle(NormalText)
@@ -106,12 +104,6 @@ TSharedRef< FSlateStyleSet > FNiagaraEditorWidgetsStyle::Create()
 		.SetFont(SystemOverviewListHeaderFont);
 	Style->Set("NiagaraEditor.SystemOverview.ListHeaderText", SystemOverviewListHeaderText);
 
-	FSlateFontInfo SystemOverviewGroupHeaderFont = DEFAULT_FONT("Bold", 9);
-	FTextBlockStyle SystemOverviewGroupHeaderText = FTextBlockStyle(NormalText)
-		.SetFont(SystemOverviewGroupHeaderFont)
-		.SetShadowOffset(FVector2D(0, 1))
-		.SetShadowColorAndOpacity(FLinearColor(0, 0, 0, 0.9f));;
-	Style->Set("NiagaraEditor.SystemOverview.GroupHeaderText", SystemOverviewGroupHeaderText);
 
 	FSlateFontInfo SystemOverviewItemFont = DEFAULT_FONT("Regular", 9);
 	FTextBlockStyle SystemOverviewItemText = FTextBlockStyle(NormalText)
diff --git a/Engine/Plugins/FX/Niagara/Source/NiagaraEditorWidgets/Private/SNiagaraOverviewStack.cpp b/Engine/Plugins/FX/Niagara/Source/NiagaraEditorWidgets/Private/SNiagaraOverviewStack.cpp
--- a/Engine/Plugins/FX/Niagara/Source/NiagaraEditorWidgets/Private/SNiagaraOverviewStack.cpp
+++ b/Engine/Plugins/FX/Niagara/Source/NiagaraEditorWidgets/Private/SNiagaraOverviewStack.cpp
@@ -56,12 +56,12 @@ class SNiagaraSystemOverviewEntryListRow : public STableRow<UNiagaraStackEntry*>
 		if (StackEntry->IsA<UNiagaraStackItem>())
 		{
 			BackgroundColor = FNiagaraEditorWidgetsStyle::Get().GetColor("NiagaraEditor.SystemOverview.Item.BackgroundColor");
-			ContentPadding = FMargin(0, 2, 0, 2);
+			ContentPadding = FMargin(0, 1, 0, 1);
 		}
 		else
 		{
 			BackgroundColor = FNiagaraEditorWidgetsStyle::Get().GetColor("NiagaraEditor.SystemOverview.Group.BackgroundColor");
-			ContentPadding = FMargin(0, 5, 0, 5);
+			ContentPadding = FMargin(0, 2, 0, 2);
 		}
 
 		DisabledBackgroundColor = BackgroundColor + FLinearColor(.02f, .02f, .02f, 0.0f);
@@ -490,7 +490,7 @@ TSharedRef<ITableRow> SNiagaraOverviewStack::OnGenerateRowForEntry(UNiagaraStack
 			// Name
 			+ SHorizontalBox::Slot()
 			.VAlign(VAlign_Center)
-			.Padding(3, 3, 0, 3)
+			.Padding(3, 2, 0, 2)
 			[
 				SNew(STextBlock)
 				.TextStyle(FNiagaraEditorWidgetsStyle::Get(), "NiagaraEditor.SystemOverview.ItemText")
diff --git a/Engine/Plugins/FX/Niagara/Source/NiagaraEditorWidgets/Private/SNiagaraOverviewStackNode.cpp b/Engine/Plugins/FX/Niagara/Source/NiagaraEditorWidgets/Private/SNiagaraOverviewStackNode.cpp
--- a/Engine/Plugins/FX/Niagara/Source/NiagaraEditorWidgets/Private/SNiagaraOverviewStackNode.cpp
+++ b/Engine/Plugins/FX/Niagara/Source/NiagaraEditorWidgets/Private/SNiagaraOverviewStackNode.cpp
@@ -243,7 +243,7 @@ TSharedRef<SWidget> SNiagaraOverviewStackNode::CreateNodeContentArea()
 	if (StackViewModel != nullptr && OverviewSelectionViewModel != nullptr)
 	{
 		ContentWidget = SNew(SBox)
-			.WidthOverride(210)
+			.MaxDesiredWidth(300)
 			[
 				SNew(SNiagaraOverviewStack, *StackViewModel, *OverviewSelectionViewModel)
 			];
