- Enable FNiagaraDrawIndirectArgsGenCS when RHISupportsComputeShaders is true

#rnx

diff --git a/Engine/Plugins/FX/Niagara/Source/NiagaraVertexFactories/Public/NiagaraDrawIndirect.h b/Engine/Plugins/FX/Niagara/Source/NiagaraVertexFactories/Public/NiagaraDrawIndirect.h
--- a/Engine/Plugins/FX/Niagara/Source/NiagaraVertexFactories/Public/NiagaraDrawIndirect.h
+++ b/Engine/Plugins/FX/Niagara/Source/NiagaraVertexFactories/Public/NiagaraDrawIndirect.h
@@ -51,7 +51,7 @@ public:
 
 	static bool ShouldCompilePermutation(const FGlobalShaderPermutationParameters& Parameters)
 	{
-		return (RHISupportsComputeShaders(Parameters.Platform) && !(Parameters.Platform == EShaderPlatform::SP_METAL || Parameters.Platform == EShaderPlatform::SP_METAL_TVOS));
+		return RHISupportsComputeShaders(Parameters.Platform);
 	}
 
 	static void ModifyCompilationEnvironment(const FGlobalShaderPermutationParameters& Parameters, FShaderCompilerEnvironment& OutEnvironment);
