Fixing New Emitter Dialog wizard to use selection and not double-clicking to select assets when you're not selecting multiple emitters.
#rnx

diff --git a/Engine/Plugins/FX/Niagara/Source/NiagaraEditor/Private/SNewEmitterDialog.cpp b/Engine/Plugins/FX/Niagara/Source/NiagaraEditor/Private/SNewEmitterDialog.cpp
--- a/Engine/Plugins/FX/Niagara/Source/NiagaraEditor/Private/SNewEmitterDialog.cpp
+++ b/Engine/Plugins/FX/Niagara/Source/NiagaraEditor/Private/SNewEmitterDialog.cpp
@@ -26,15 +26,14 @@ void SNewEmitterDialog::Construct(const FArguments& InArgs)
 				SNiagaraNewAssetDialog::FOnGetSelectedAssetsFromPicker::CreateSP(this, &SNewEmitterDialog::GetSelectedEmitterTemplateAssets),
 				SNiagaraNewAssetDialog::FOnSelectionConfirmed(),
 				SAssignNew(TemplateAssetPicker, SNiagaraAssetPickerList, UNiagaraEmitter::StaticClass())
-				.OnTemplateAssetActivated(this, &SNewEmitterDialog::OnTemplateAssetActivated)
 				.bTemplateOnly(true)),
 			SNiagaraNewAssetDialog::FNiagaraNewAssetDialogOption(
 				LOCTEXT("InheritFromOtherEmitterLabel", "Inherit from an existing emitter"),
 				LOCTEXT("InheritFromOtherEmitterDescription", "Create an inheritance chain between the new emitter and an existing emitter"),
 				LOCTEXT("InheritProjectEmitterPickerHeader", "Select a Parent Project Emitter"),
-				SNiagaraNewAssetDialog::FOnGetSelectedAssetsFromPicker::CreateSP(this, &SNewEmitterDialog::GetSelectedProjectEmiterAssets),
+				SNiagaraNewAssetDialog::FOnGetSelectedAssetsFromPicker::CreateSP(this, &SNewEmitterDialog::GetSelectedParentEmitterAssets),
 				SNiagaraNewAssetDialog::FOnSelectionConfirmed::CreateSP(this, &SNewEmitterDialog::InheritanceOptionConfirmed),
-				SAssignNew(TemplateAssetPicker, SNiagaraAssetPickerList, UNiagaraEmitter::StaticClass())
+				SAssignNew(InheritAssetPicker, SNiagaraAssetPickerList, UNiagaraEmitter::StaticClass())
 				.bTemplateOnly(false)),
 			SNiagaraNewAssetDialog::FNiagaraNewAssetDialogOption(
 				LOCTEXT("CreateFromOtherEmitterLabel", "Copy existing emitter"),
@@ -42,8 +41,7 @@ void SNewEmitterDialog::Construct(const FArguments& InArgs)
 				LOCTEXT("ProjectEmitterPickerHeader", "Select a Project Emitter"),
 				SNiagaraNewAssetDialog::FOnGetSelectedAssetsFromPicker::CreateSP(this, &SNewEmitterDialog::GetSelectedProjectEmiterAssets),
 				SNiagaraNewAssetDialog::FOnSelectionConfirmed(),
-				SAssignNew(TemplateAssetPicker, SNiagaraAssetPickerList, UNiagaraEmitter::StaticClass())
-				.OnTemplateAssetActivated(this, &SNewEmitterDialog::OnEmitterAssetsActivated)
+				SAssignNew(CopyAssetPicker, SNiagaraAssetPickerList, UNiagaraEmitter::StaticClass())
 				.bTemplateOnly(false))
 		});
 }
@@ -66,36 +64,18 @@ bool SNewEmitterDialog::GetUseInheritance() const
 void SNewEmitterDialog::GetSelectedEmitterTemplateAssets(TArray<FAssetData>& OutSelectedAssets)
 {
 	OutSelectedAssets.Append(TemplateAssetPicker->GetSelectedAssets());
-	if (ActivatedTemplateAsset.IsValid())
-	{
-		OutSelectedAssets.AddUnique(ActivatedTemplateAsset);
-	}
 }
 
-void SNewEmitterDialog::GetSelectedProjectEmiterAssets(TArray<FAssetData>& OutSelectedAssets)
+void SNewEmitterDialog::GetSelectedParentEmitterAssets(TArray<FAssetData>& OutSelectedAssets)
 {
-	OutSelectedAssets.Append(TemplateAssetPicker->GetSelectedAssets());
-	if (ActivatedProjectAsset.IsValid())
-	{
-		OutSelectedAssets.AddUnique(ActivatedProjectAsset);
-	}
+	OutSelectedAssets.Append(InheritAssetPicker->GetSelectedAssets());
 }
 
-void SNewEmitterDialog::OnTemplateAssetActivated(const FAssetData& InActivatedTemplateAsset)
+void SNewEmitterDialog::GetSelectedProjectEmiterAssets(TArray<FAssetData>& OutSelectedAssets)
 {
-	// Input handling issues with the list view widget can allow items to be activated but not added to the selection so cache this here
-	// so it can be included in the selection set.
-	ActivatedTemplateAsset = InActivatedTemplateAsset;
-	ConfirmSelection();
+	OutSelectedAssets.Append(CopyAssetPicker->GetSelectedAssets());
 }
 
-void SNewEmitterDialog::OnEmitterAssetsActivated(const FAssetData& InActivatedTemplateAsset)
-{
-	// Input handling issues with the list view widget can allow items to be activated but not added to the selection so cache this here
-	// so it can be included in the selection set.
-	ActivatedProjectAsset = InActivatedTemplateAsset;
-	ConfirmSelection();
-}
 
 void SNewEmitterDialog::InheritanceOptionConfirmed()
 {
diff --git a/Engine/Plugins/FX/Niagara/Source/NiagaraEditor/Private/SNewEmitterDialog.h b/Engine/Plugins/FX/Niagara/Source/NiagaraEditor/Private/SNewEmitterDialog.h
--- a/Engine/Plugins/FX/Niagara/Source/NiagaraEditor/Private/SNewEmitterDialog.h
+++ b/Engine/Plugins/FX/Niagara/Source/NiagaraEditor/Private/SNewEmitterDialog.h
@@ -28,20 +28,16 @@ public:
 private:
 	void GetSelectedEmitterTemplateAssets(TArray<FAssetData>& OutSelectedAssets);
 
-	void GetSelectedProjectEmiterAssets(TArray<FAssetData>& OutSelectedAssets);
-
-	void OnTemplateAssetActivated(const FAssetData& InActivatedTemplateAsset);
+	void GetSelectedParentEmitterAssets(TArray<FAssetData>& OutSelectedAssets);
 
-	void OnEmitterAssetsActivated(const FAssetData& InActivatedTemplateAsset);
+	void GetSelectedProjectEmiterAssets(TArray<FAssetData>& OutSelectedAssets);
 
 	void InheritanceOptionConfirmed();
 
 private:
 	TSharedPtr<SNiagaraAssetPickerList> TemplateAssetPicker;
-
-	FAssetData ActivatedTemplateAsset;
-
-	FAssetData ActivatedProjectAsset;
+	TSharedPtr<SNiagaraAssetPickerList> InheritAssetPicker;
+	TSharedPtr<SNiagaraAssetPickerList> CopyAssetPicker;
 
 	bool bUseInheritance;
 };
\ No newline at end of file
diff --git a/Engine/Plugins/FX/Niagara/Source/NiagaraEditor/Private/SNewSystemDialog.cpp b/Engine/Plugins/FX/Niagara/Source/NiagaraEditor/Private/SNewSystemDialog.cpp
--- a/Engine/Plugins/FX/Niagara/Source/NiagaraEditor/Private/SNewSystemDialog.cpp
+++ b/Engine/Plugins/FX/Niagara/Source/NiagaraEditor/Private/SNewSystemDialog.cpp
@@ -37,6 +37,7 @@ void SNewSystemDialog::Construct(const FArguments& InArgs)
 				.Padding(0, 0, 0, 10)
 				[
 					SAssignNew(EmitterAssetPicker, SNiagaraAssetPickerList, UNiagaraEmitter::StaticClass())
+					.OnTemplateAssetActivated(this, &SNewSystemDialog::OnEmitterAssetsActivated)
 					.bTemplateOnly(false)
 					.bAllowMultiSelect(true)
 				]
@@ -91,7 +92,6 @@ void SNewSystemDialog::Construct(const FArguments& InArgs)
 				SNiagaraNewAssetDialog::FOnGetSelectedAssetsFromPicker::CreateSP(this, &SNewSystemDialog::GetSelectedSystemTemplateAssets),
 				SNiagaraNewAssetDialog::FOnSelectionConfirmed(),
 				SAssignNew(TemplateAssetPicker, SNiagaraAssetPickerList, UNiagaraSystem::StaticClass())
-				.OnTemplateAssetActivated(this, &SNewSystemDialog::OnTemplateAssetActivated)
 				.bTemplateOnly(true)),
 			SNiagaraNewAssetDialog::FNiagaraNewAssetDialogOption(
 				LOCTEXT("CreateFromOtherSystemLabel", "Copy existing system"),
@@ -100,7 +100,6 @@ void SNewSystemDialog::Construct(const FArguments& InArgs)
 				SNiagaraNewAssetDialog::FOnGetSelectedAssetsFromPicker::CreateSP(this, &SNewSystemDialog::GetSelectedProjectSystemAssets),
 				SNiagaraNewAssetDialog::FOnSelectionConfirmed(),
 				SAssignNew(SystemAssetPicker, SNiagaraAssetPickerList, UNiagaraSystem::StaticClass())
-				.OnTemplateAssetActivated(this, &SNewSystemDialog::OnSystemAssetsActivated)
 				.bTemplateOnly(false)),
 			SNiagaraNewAssetDialog::FNiagaraNewAssetDialogOption(
 				LOCTEXT("CreateEmptyLabel", "Create empty system"),
@@ -150,22 +149,21 @@ TArray<FAssetData> SNewSystemDialog::GetSelectedEmitterAssets() const
 	return ConfirmedSelectedEmitterAssets;
 }
 
+void SNewSystemDialog::OnEmitterAssetsActivated(const FAssetData& ActivatedTemplateAsset)
+{
+	TArray<FAssetData> ActivatedTemplateAssets;
+	ActivatedTemplateAssets.Add(ActivatedTemplateAsset);
+	AddEmitterAssetsToSelection(ActivatedTemplateAssets);
+}
+
 void SNewSystemDialog::GetSelectedSystemTemplateAssets(TArray<FAssetData>& OutSelectedAssets)
 {
 	OutSelectedAssets.Append(TemplateAssetPicker->GetSelectedAssets());
-	if (ActivatedTemplateSystemAsset.IsValid())
-	{
-		OutSelectedAssets.AddUnique(ActivatedTemplateSystemAsset);
-	}
 }
 
 void SNewSystemDialog::GetSelectedProjectSystemAssets(TArray<FAssetData>& OutSelectedAssets)
 {
 	OutSelectedAssets.Append(SystemAssetPicker->GetSelectedAssets());
-	if (ActivatedProjectSystemAsset.IsValid())
-	{
-		OutSelectedAssets.AddUnique(ActivatedProjectSystemAsset);
-	}
 }
 
 void SNewSystemDialog::GetSelectedProjectEmiterAssets(TArray<FAssetData>& OutSelectedAssets)
@@ -173,22 +171,6 @@ void SNewSystemDialog::GetSelectedProjectEmiterAssets(TArray<FAssetData>& OutSel
 	OutSelectedAssets.Append(SelectedEmitterAssets);
 }
 
-void SNewSystemDialog::OnTemplateAssetActivated(const FAssetData& ActivatedTemplateAsset)
-{
-	// Input handling issues with the list view widget can allow items to be activated but not added to the selection so cache this here
-	// so it can be included in the selection set.
-	ActivatedTemplateSystemAsset = ActivatedTemplateAsset;
-	ConfirmSelection();
-}
-
-void SNewSystemDialog::OnSystemAssetsActivated(const FAssetData& ActivatedTemplateAsset)
-{
-	// Input handling issues with the list view widget can allow items to be activated but not added to the selection so cache this here
-	// so it can be included in the selection set.
-	ActivatedProjectSystemAsset = ActivatedTemplateAsset;
-	ConfirmSelection();
-}
-
 bool SNewSystemDialog::IsAddEmittersToSelectionButtonEnabled() const
 {
 	return EmitterAssetPicker->GetSelectedAssets().Num() > 0;
diff --git a/Engine/Plugins/FX/Niagara/Source/NiagaraEditor/Private/SNewSystemDialog.h b/Engine/Plugins/FX/Niagara/Source/NiagaraEditor/Private/SNewSystemDialog.h
--- a/Engine/Plugins/FX/Niagara/Source/NiagaraEditor/Private/SNewSystemDialog.h
+++ b/Engine/Plugins/FX/Niagara/Source/NiagaraEditor/Private/SNewSystemDialog.h
@@ -27,16 +27,15 @@ public:
 	TArray<FAssetData> GetSelectedEmitterAssets() const;
 
 private:
+
+	void OnEmitterAssetsActivated(const FAssetData& ActivatedTemplateAsset);
+
 	void GetSelectedSystemTemplateAssets(TArray<FAssetData>& OutSelectedAssets);
 
 	void GetSelectedProjectSystemAssets(TArray<FAssetData>& OutSelectedAssets);
 
 	void GetSelectedProjectEmiterAssets(TArray<FAssetData>& OutSelectedAssets);
 
-	void OnTemplateAssetActivated(const FAssetData& ActivatedTemplateAsset);
-
-	void OnSystemAssetsActivated(const FAssetData& ActivatedTemplateAsset);
-
 	bool IsAddEmittersToSelectionButtonEnabled() const;
 	
 	FReply AddEmittersToSelectionButtonClicked();
@@ -46,9 +45,6 @@ private:
 	FReply RemoveEmitterFromSelectionButtonClicked(FAssetData EmitterAsset);
 
 private:
-	FAssetData ActivatedTemplateSystemAsset;
-
-	FAssetData ActivatedProjectSystemAsset;
 
 	TArray<FAssetData> SelectedEmitterAssets;
 
