Disabled the check for "requires persistent IDs" when compiling modules, since they don't know anything about the emitter at that point.

diff --git a/Engine/Plugins/FX/Niagara/Source/NiagaraEditor/Private/NiagaraHlslTranslator.cpp b/Engine/Plugins/FX/Niagara/Source/NiagaraEditor/Private/NiagaraHlslTranslator.cpp
--- a/Engine/Plugins/FX/Niagara/Source/NiagaraEditor/Private/NiagaraHlslTranslator.cpp
+++ b/Engine/Plugins/FX/Niagara/Source/NiagaraEditor/Private/NiagaraHlslTranslator.cpp
@@ -984,7 +984,7 @@ const FNiagaraTranslateResults &FHlslNiagaraTranslator::Translate(const FNiagara
 		}
 	}
 	ENiagaraScriptUsage Usage = CompileOptions.TargetUsage;
-	if (Usage != ENiagaraScriptUsage::SystemSpawnScript && Usage != ENiagaraScriptUsage::SystemUpdateScript)
+	if (Usage != ENiagaraScriptUsage::SystemSpawnScript && Usage != ENiagaraScriptUsage::SystemUpdateScript && Usage != ENiagaraScriptUsage::Module && Usage != ENiagaraScriptUsage::DynamicInput)
 	{
 		ValidateParticleIDUsage();
 	}
