Moved SetPinName to NiagaraNodeParameterMapBase, since it was required for both the Get and Set nodes.

#rnx

diff --git a/Engine/Plugins/FX/Niagara/Source/NiagaraEditor/Private/NiagaraNodeParameterMapBase.cpp b/Engine/Plugins/FX/Niagara/Source/NiagaraEditor/Private/NiagaraNodeParameterMapBase.cpp
--- a/Engine/Plugins/FX/Niagara/Source/NiagaraEditor/Private/NiagaraNodeParameterMapBase.cpp
+++ b/Engine/Plugins/FX/Niagara/Source/NiagaraEditor/Private/NiagaraNodeParameterMapBase.cpp
@@ -190,6 +190,13 @@ void UNiagaraNodeParameterMapBase::GetPinHoverText(const UEdGraphPin& Pin, FStri
 	}
 }
 
+void UNiagaraNodeParameterMapBase::SetPinName(UEdGraphPin* InPin, const FName& InName)
+{
+	FName OldName = InPin->PinName;
+	InPin->PinName = InName;
+	OnPinRenamed(InPin, OldName.ToString());
+}
+
 void UNiagaraNodeParameterMapBase::OnPinRenamed(UEdGraphPin* RenamedPin, const FString& OldName)
 {
 	RenamedPin->PinFriendlyName = FText::FromName(RenamedPin->PinName);
diff --git a/Engine/Plugins/FX/Niagara/Source/NiagaraEditor/Private/NiagaraNodeParameterMapBase.h b/Engine/Plugins/FX/Niagara/Source/NiagaraEditor/Private/NiagaraNodeParameterMapBase.h
--- a/Engine/Plugins/FX/Niagara/Source/NiagaraEditor/Private/NiagaraNodeParameterMapBase.h
+++ b/Engine/Plugins/FX/Niagara/Source/NiagaraEditor/Private/NiagaraNodeParameterMapBase.h
@@ -34,10 +34,11 @@ public:
 
 	virtual void CollectAddPinActions(FGraphActionListBuilderBase& OutActions, bool& bOutCreateRemainingActions, UEdGraphPin* Pin) override;
 	virtual void GetPinHoverText(const UEdGraphPin& Pin, FString& HoverTextOut) const override;
-	
 
 	virtual bool IncludeParentNodeContextMenu() const { return true; }
 
+	void SetPinName(UEdGraphPin* InPin, const FName& InName);
+
 protected:
 	virtual void OnPinRenamed(UEdGraphPin* RenamedPin, const FString& OldName) override;
 
diff --git a/Engine/Plugins/FX/Niagara/Source/NiagaraEditor/Private/NiagaraNodeParameterMapSet.cpp b/Engine/Plugins/FX/Niagara/Source/NiagaraEditor/Private/NiagaraNodeParameterMapSet.cpp
--- a/Engine/Plugins/FX/Niagara/Source/NiagaraEditor/Private/NiagaraNodeParameterMapSet.cpp
+++ b/Engine/Plugins/FX/Niagara/Source/NiagaraEditor/Private/NiagaraNodeParameterMapSet.cpp
@@ -107,13 +107,6 @@ void UNiagaraNodeParameterMapSet::OnPinRenamed(UEdGraphPin* RenamedPin, const FS
 	MarkNodeRequiresSynchronization(__FUNCTION__, true);
 }
 
-void UNiagaraNodeParameterMapSet::SetPinName(UEdGraphPin* InPin, const FName& InName)
-{
-	FName OldName = InPin->PinName;
-	InPin->PinName = InName;
-	OnPinRenamed(InPin, OldName.ToString());
-}
-
 bool UNiagaraNodeParameterMapSet::CancelEditablePinName(const FText& InName, UEdGraphPin* InGraphPinObj)
 {
 	if (InGraphPinObj == PinPendingRename)
diff --git a/Engine/Plugins/FX/Niagara/Source/NiagaraEditor/Private/NiagaraNodeParameterMapSet.h b/Engine/Plugins/FX/Niagara/Source/NiagaraEditor/Private/NiagaraNodeParameterMapSet.h
--- a/Engine/Plugins/FX/Niagara/Source/NiagaraEditor/Private/NiagaraNodeParameterMapSet.h
+++ b/Engine/Plugins/FX/Niagara/Source/NiagaraEditor/Private/NiagaraNodeParameterMapSet.h
@@ -36,8 +36,6 @@ public:
 	virtual bool IncludeParentNodeContextMenu() const { return true; }
 	virtual void PostLoad() override;
 
-	void SetPinName(UEdGraphPin* InPin, const FName& InName);
-
 protected:
 	virtual void OnNewTypedPinAdded(UEdGraphPin* NewPin); 
 	virtual void OnPinRenamed(UEdGraphPin* RenamedPin, const FString& OldName) override;
