Replacing a deprecated module or dynamic input with the suggested replacement now also renames the module, and tries to retain any dynamic inputs or input values that were set on that module previously.

diff --git a/Engine/Plugins/FX/Niagara/Source/NiagaraEditor/Private/ViewModels/Stack/NiagaraStackFunctionInput.cpp b/Engine/Plugins/FX/Niagara/Source/NiagaraEditor/Private/ViewModels/Stack/NiagaraStackFunctionInput.cpp
--- a/Engine/Plugins/FX/Niagara/Source/NiagaraEditor/Private/ViewModels/Stack/NiagaraStackFunctionInput.cpp
+++ b/Engine/Plugins/FX/Niagara/Source/NiagaraEditor/Private/ViewModels/Stack/NiagaraStackFunctionInput.cpp
@@ -137,48 +137,21 @@ void UNiagaraStackFunctionInput::Initialize(
 	OwningFunctionCallInitialScript = OwningFunctionCallNode->FunctionScript;
 	OwningAssignmentNode = Cast<UNiagaraNodeAssignment>(OwningFunctionCallNode.Get());
 
-	UNiagaraNodeOutput* OutputNode = FNiagaraStackGraphUtilities::GetEmitterOutputNodeForStackNode(*OwningModuleNode.Get());
+	UNiagaraSystem& ParentSystem = GetSystemViewModel()->GetSystem();
 	UNiagaraEmitter* ParentEmitter = GetEmitterViewModel().IsValid() ? GetEmitterViewModel()->GetEmitter() : nullptr;
-	UNiagaraSystem* ParentSystem = &GetSystemViewModel()->GetSystem();
-	if (OutputNode)
-	{
-		TArray<UNiagaraScript*> Scripts;
-		if (ParentEmitter != nullptr)
-		{
-			ParentEmitter->GetScripts(Scripts, false);
-		}
-		if (ParentSystem != nullptr)
-		{
-			Scripts.Add(ParentSystem->GetSystemSpawnScript());
-			Scripts.Add(ParentSystem->GetSystemUpdateScript());
-		}
 
-		for (UNiagaraScript* Script : Scripts)
-		{
-			if (OutputNode->GetUsage() == ENiagaraScriptUsage::ParticleEventScript)
-			{
-				if (Script->GetUsage() == ENiagaraScriptUsage::ParticleEventScript && Script->GetUsageId() == OutputNode->GetUsageId())
-				{
-					AffectedScripts.Add(Script);
-					break;
-				}
-			}
-			else if (Script->ContainsUsage(OutputNode->GetUsage()))
-			{
-				AffectedScripts.Add(Script);
-			}
-		}
+	FNiagaraStackGraphUtilities::FindAffectedScripts(ParentSystem, ParentEmitter, *OwningModuleNode.Get(), AffectedScripts);
 
-		for (TWeakObjectPtr<UNiagaraScript> AffectedScript : AffectedScripts)
+	UNiagaraNodeOutput* OutputNode = FNiagaraStackGraphUtilities::GetEmitterOutputNodeForStackNode(*OwningModuleNode.Get());
+	for (TWeakObjectPtr<UNiagaraScript> AffectedScript : AffectedScripts)
+	{
+		if (AffectedScript.IsValid() && AffectedScript->IsEquivalentUsage(OutputNode->GetUsage()) && AffectedScript->GetUsageId() == OutputNode->GetUsageId())
 		{
-			if (AffectedScript.IsValid() && AffectedScript->IsEquivalentUsage(OutputNode->GetUsage()) && AffectedScript->GetUsageId() == OutputNode->GetUsageId())
-			{
-				SourceScript = AffectedScript;
-				RapidIterationParametersChangedHandle = SourceScript->RapidIterationParameters.AddOnChangedHandler(
-					FNiagaraParameterStore::FOnChanged::FDelegate::CreateUObject(this, &UNiagaraStackFunctionInput::OnRapidIterationParametersChanged));
-				SourceScript->GetSource()->OnChanged().AddUObject(this, &UNiagaraStackFunctionInput::OnScriptSourceChanged);
-				break;
-			}
+			SourceScript = AffectedScript;
+			RapidIterationParametersChangedHandle = SourceScript->RapidIterationParameters.AddOnChangedHandler(
+				FNiagaraParameterStore::FOnChanged::FDelegate::CreateUObject(this, &UNiagaraStackFunctionInput::OnRapidIterationParametersChanged));
+			SourceScript->GetSource()->OnChanged().AddUObject(this, &UNiagaraStackFunctionInput::OnScriptSourceChanged);
+			break;
 		}
 	}
 
@@ -446,12 +419,12 @@ void UNiagaraStackFunctionInput::RefreshChildrenInternal(const TArray<UNiagaraSt
 				if (InputValues.DynamicNode->FunctionScript->bDeprecated)
 				{
 					FText LongMessage = InputValues.DynamicNode->FunctionScript->DeprecationRecommendation != nullptr ?
-						FText::Format(LOCTEXT("ModuleScriptDeprecationLong", "The script asset for the assigned module {0} has been deprecated. Suggested replacement: {1}"), FText::FromString(InputValues.DynamicNode->GetName()), FText::FromString(InputValues.DynamicNode->FunctionScript->DeprecationRecommendation->GetPathName())) :
-						FText::Format(LOCTEXT("ModuleScriptDeprecationUnknownLong", "The script asset for the assigned module {0} has been deprecated."), FText::FromString(InputValues.DynamicNode->GetName()));
+						FText::Format(LOCTEXT("DynamicInputScriptDeprecationLong", "The script asset for this dynamic input has been deprecated. Suggested replacement: {1}"), FText::FromString(InputValues.DynamicNode->GetName()), FText::FromString(InputValues.DynamicNode->FunctionScript->DeprecationRecommendation->GetPathName())) :
+						FText::Format(LOCTEXT("DynamicInputScriptDeprecationUnknownLong", "The script asset for the assigned module {0} has been deprecated."), FText::FromString(InputValues.DynamicNode->GetName()));
 
 					int32 AddIdx = NewIssues.Add(FStackIssue(
 						EStackIssueSeverity::Warning,
-						LOCTEXT("ModuleScriptDeprecationShort", "Deprecated module"),
+						LOCTEXT("DynamicInputScriptDeprecationShort", "Deprecated dynamic input"),
 						LongMessage,
 						GetStackEditorDataKey(),
 						false,
@@ -460,7 +433,7 @@ void UNiagaraStackFunctionInput::RefreshChildrenInternal(const TArray<UNiagaraSt
 								LOCTEXT("SelectNewDynamicInputScriptFix", "Select a new dynamic input script"),
 								FStackIssueFixDelegate::CreateLambda([this]() { this->bIsDynamicInputScriptReassignmentPending = true; })),
 							FStackIssueFix(
-								LOCTEXT("ResetFix", "Reset this input to it's default value"),
+								LOCTEXT("ResetDynamicInputFix", "Reset this input to its default value"),
 								FStackIssueFixDelegate::CreateLambda([this]() { this->Reset(); }))
 						}));
 
@@ -468,7 +441,7 @@ void UNiagaraStackFunctionInput::RefreshChildrenInternal(const TArray<UNiagaraSt
 					{
 						NewIssues[AddIdx].InsertFix(0,
 							FStackIssueFix(
-							LOCTEXT("SelectNewModuleScriptFixUseRecommended", "Use recommended replacement"),
+							LOCTEXT("SelectNewDynamicInputScriptFixUseRecommended", "Use recommended replacement"),
 							FStackIssueFixDelegate::CreateLambda([this]() { ReassignDynamicInputScript(InputValues.DynamicNode->FunctionScript->DeprecationRecommendation); })));
 					}
 				}
@@ -477,8 +450,8 @@ void UNiagaraStackFunctionInput::RefreshChildrenInternal(const TArray<UNiagaraSt
 				{
 					NewIssues.Add(FStackIssue(
 						EStackIssueSeverity::Info,
-						LOCTEXT("ModuleScriptExperimentalShort", "Experimental module"),
-						FText::Format(LOCTEXT("ModuleScriptExperimental", "The script asset for the assigned module {0} is experimental, use with care!"), FText::FromString(InputValues.DynamicNode->GetName())),
+						LOCTEXT("DynamicInputScriptExperimentalShort", "Experimental dynamic input"),
+						FText::Format(LOCTEXT("DynamicInputScriptExperimental", "The script asset for this dynamic input is experimental, use with care!"), FText::FromString(InputValues.DynamicNode->GetName())),
 						GetStackEditorDataKey(),
 						true));
 				}
@@ -1762,8 +1735,22 @@ void UNiagaraStackFunctionInput::ReassignDynamicInputScript(UNiagaraScript* Dyna
 		TEXT("Can not reassign the dynamic input script when tne input doesn't have a valid dynamic input.")))
 	{
 		FScopedTransaction ScopedTransaction(LOCTEXT("ReassignDynamicInputTransaction", "Reassign dynamic input script"));
+
+		const FString OldName = InputValues.DynamicNode->GetFunctionName();
+
 		InputValues.DynamicNode->Modify();
 		InputValues.DynamicNode->FunctionScript = DynamicInputScript;
+
+		// intermediate refresh to purge any rapid iteration parameters that have been removed in the new script
+		RefreshChildren();
+
+		InputValues.DynamicNode->SuggestName(FString());
+
+		const FString NewName = InputValues.DynamicNode->GetFunctionName();
+		UNiagaraSystem& System = GetSystemViewModel()->GetSystem();
+		UNiagaraEmitter* Emitter = GetEmitterViewModel().IsValid() ? GetEmitterViewModel()->GetEmitter() : nullptr;
+		FNiagaraStackGraphUtilities::RenameReferencingParameters(System, Emitter, *InputValues.DynamicNode.Get(), OldName, NewName);
+
 		InputValues.DynamicNode->RefreshFromExternalChanges();
 		InputValues.DynamicNode->MarkNodeRequiresSynchronization(TEXT("Dynamic input script reassigned."), true);
 		RefreshChildren();
diff --git a/Engine/Plugins/FX/Niagara/Source/NiagaraEditor/Private/ViewModels/Stack/NiagaraStackGraphUtilities.cpp b/Engine/Plugins/FX/Niagara/Source/NiagaraEditor/Private/ViewModels/Stack/NiagaraStackGraphUtilities.cpp
--- a/Engine/Plugins/FX/Niagara/Source/NiagaraEditor/Private/ViewModels/Stack/NiagaraStackGraphUtilities.cpp
+++ b/Engine/Plugins/FX/Niagara/Source/NiagaraEditor/Private/ViewModels/Stack/NiagaraStackGraphUtilities.cpp
@@ -534,7 +534,8 @@ void ExtractInputPinsFromHistory(FNiagaraParameterMapHistory& History, UEdGraph*
 			const UEdGraphPin* InputPin = ReadHistory[0].Get<0>();
 
 			// Make sure that the module input is from the called graph, and not a nested graph.
-			if (InputPin->GetOwningNode()->GetGraph() == FunctionGraph &&
+			UEdGraph* NodeGraph = InputPin->GetOwningNode()->GetGraph();
+			if (NodeGraph == FunctionGraph &&
 				(Options == FNiagaraStackGraphUtilities::ENiagaraGetStackFunctionInputPinsOptions::AllInputs || FNiagaraParameterHandle(InputPin->PinName).IsModuleHandle()))
 			{
 				OutPins.Add(InputPin);
@@ -734,6 +735,56 @@ void FNiagaraStackGraphUtilities::GetStackFunctionOutputVariables(UNiagaraNodeFu
 	}
 }
 
+void FNiagaraStackGraphUtilities::GetStackFunctionInputAndOutputVariables(UNiagaraNodeFunctionCall& FunctionCallNode, FCompileConstantResolver ConstantResolver, TArray<FNiagaraVariable>& OutVariables, TArray<FNiagaraVariable>& OutVariablesWithOriginalAliasesIntact)
+{
+	FNiagaraParameterMapHistoryBuilder Builder;
+	Builder.SetIgnoreDisabled(false);
+	Builder.ConstantResolver = ConstantResolver;
+	FunctionCallNode.BuildParameterMapHistory(Builder, false);
+
+	if (ensureMsgf(Builder.Histories.Num() == 1, TEXT("Invalid Stack Graph - Function call node has invalid history count!")))
+	{
+		for (int32 i = 0; i < Builder.Histories[0].Variables.Num(); ++i)
+		{
+			bool bHasParameterMapSetWrite = false;
+			for (const UEdGraphPin* WritePin : Builder.Histories[0].PerVariableWriteHistory[i])
+			{
+				if (WritePin != nullptr && WritePin->GetOwningNode() != nullptr && 
+					WritePin->GetOwningNode()->IsA<UNiagaraNodeParameterMapSet>())
+				{
+					bHasParameterMapSetWrite = true;
+					break;
+				}
+			}
+
+			if (bHasParameterMapSetWrite)
+			{
+				FNiagaraVariable& Variable = Builder.Histories[0].Variables[i];
+				FNiagaraVariable& VariableWithOriginalAliasIntact = Builder.Histories[0].VariablesWithOriginalAliasesIntact[i];
+				OutVariables.Add(Variable);
+				OutVariablesWithOriginalAliasesIntact.Add(VariableWithOriginalAliasIntact);
+			}
+		}
+
+		TArray<const UEdGraphPin*> InputPins;
+		ExtractInputPinsFromHistory(Builder.Histories[0], FunctionCallNode.GetCalledGraph(), FNiagaraStackGraphUtilities::ENiagaraGetStackFunctionInputPinsOptions::ModuleInputsOnly, InputPins);
+
+		for (const UEdGraphPin* Pin : InputPins)
+		{
+			for (int32 i = 0; i < Builder.Histories[0].Variables.Num(); ++i)
+			{
+				FNiagaraVariable& VariableWithOriginalAliasIntact = Builder.Histories[0].VariablesWithOriginalAliasesIntact[i];
+				if (VariableWithOriginalAliasIntact.GetName() == Pin->PinName)
+				{
+					FNiagaraVariable& Variable = Builder.Histories[0].Variables[i];
+					OutVariables.AddUnique(Variable);
+					OutVariablesWithOriginalAliasesIntact.AddUnique(VariableWithOriginalAliasIntact);
+				}
+			}
+		}
+	}
+}
+
 UNiagaraNodeParameterMapSet* FNiagaraStackGraphUtilities::GetStackFunctionOverrideNode(UNiagaraNodeFunctionCall& FunctionCallNode)
 {
 	UEdGraphPin* ParameterMapInput = FNiagaraStackGraphUtilities::GetParameterMapInputPin(FunctionCallNode);
@@ -2292,4 +2343,175 @@ void FNiagaraStackGraphUtilities::RebuildEmitterNodes(UNiagaraSystem& System)
 	RelayoutGraph(*SystemGraph);
 }
 
+void FNiagaraStackGraphUtilities::FindAffectedScripts(UNiagaraSystem& System, UNiagaraEmitter* Emitter, UNiagaraNodeFunctionCall& ModuleNode, TArray<TWeakObjectPtr<UNiagaraScript>>& OutAffectedScripts)
+{
+	UNiagaraNodeOutput* OutputNode = FNiagaraStackGraphUtilities::GetEmitterOutputNodeForStackNode(ModuleNode);
+
+	if (OutputNode)
+	{
+		TArray<UNiagaraScript*> Scripts;
+		if (Emitter != nullptr)
+		{
+			Emitter->GetScripts(Scripts, false);
+		}
+
+		OutAffectedScripts.Add(System.GetSystemSpawnScript());
+		OutAffectedScripts.Add(System.GetSystemUpdateScript());
+
+		for (UNiagaraScript* Script : Scripts)
+		{
+			if (OutputNode->GetUsage() == ENiagaraScriptUsage::ParticleEventScript)
+			{
+				if (Script->GetUsage() == ENiagaraScriptUsage::ParticleEventScript && Script->GetUsageId() == OutputNode->GetUsageId())
+				{
+					OutAffectedScripts.Add(Script);
+					break;
+				}
+			}
+			else if (Script->ContainsUsage(OutputNode->GetUsage()))
+			{
+				OutAffectedScripts.Add(Script);
+			}
+		}
+	}
+}
+
+void FNiagaraStackGraphUtilities::GatherRenamedStackFunctionOutputVariableNames(UNiagaraEmitter* Emitter, UNiagaraNodeFunctionCall& FunctionCallNode, const FString& OldFunctionName, const FString& NewFunctionName, TMap<FName, FName>& OutOldToNewNameMap)
+{
+	if (OldFunctionName == NewFunctionName)
+	{
+		return;
+	}
+
+	TArray<FNiagaraVariable> OutputVariables;
+	TArray<FNiagaraVariable> OutputVariablesWithOriginalAliasesIntact;
+	FCompileConstantResolver ConstantResolver(Emitter);
+	FNiagaraStackGraphUtilities::GetStackFunctionOutputVariables(FunctionCallNode, ConstantResolver, OutputVariables, OutputVariablesWithOriginalAliasesIntact);
+
+	for (FNiagaraVariable& OutputVariableWithOriginalAliasesIntact : OutputVariablesWithOriginalAliasesIntact)
+	{
+		TArray<FString> SplitAliasedVariableName;
+		OutputVariableWithOriginalAliasesIntact.GetName().ToString().ParseIntoArray(SplitAliasedVariableName, TEXT("."));
+		if (SplitAliasedVariableName.Contains(TEXT("Module")))
+		{
+			TArray<FString> SplitOldVariableName = SplitAliasedVariableName;
+			TArray<FString> SplitNewVariableName = SplitAliasedVariableName;
+			for (int32 i = 0; i < SplitAliasedVariableName.Num(); i++)
+			{
+				if (SplitAliasedVariableName[i] == TEXT("Module"))
+				{
+					SplitOldVariableName[i] = OldFunctionName;
+					SplitNewVariableName[i] = NewFunctionName;
+				}
+			}
+
+			OutOldToNewNameMap.Add(*FString::Join(SplitOldVariableName, TEXT(".")), *FString::Join(SplitNewVariableName, TEXT(".")));
+		}
+	}
+}
+
+void FNiagaraStackGraphUtilities::GatherRenamedStackFunctionInputAndOutputVariableNames(UNiagaraEmitter* Emitter, UNiagaraNodeFunctionCall& FunctionCallNode, const FString& OldFunctionName, const FString& NewFunctionName, TMap<FName, FName>& OutOldToNewNameMap)
+{
+	if (OldFunctionName == NewFunctionName)
+	{
+		return;
+	}
+
+	TArray<FNiagaraVariable> Variables;
+	TArray<FNiagaraVariable> VariablesWithOriginalAliasesIntact;
+	FCompileConstantResolver ConstantResolver(Emitter);
+	FNiagaraStackGraphUtilities::GetStackFunctionInputAndOutputVariables(FunctionCallNode, ConstantResolver, Variables, VariablesWithOriginalAliasesIntact);
+
+	for (FNiagaraVariable& Variable : VariablesWithOriginalAliasesIntact)
+	{
+		TArray<FString> SplitAliasedVariableName;
+		Variable.GetName().ToString().ParseIntoArray(SplitAliasedVariableName, TEXT("."));
+		if (SplitAliasedVariableName.Contains(TEXT("Module")))
+		{
+			TArray<FString> SplitOldVariableName = SplitAliasedVariableName;
+			TArray<FString> SplitNewVariableName = SplitAliasedVariableName;
+			for (int32 i = 0; i < SplitAliasedVariableName.Num(); i++)
+			{
+				if (SplitAliasedVariableName[i] == TEXT("Module"))
+				{
+					SplitOldVariableName[i] = OldFunctionName;
+					SplitNewVariableName[i] = NewFunctionName;
+				}
+			}
+
+			OutOldToNewNameMap.Add(*FString::Join(SplitOldVariableName, TEXT(".")), *FString::Join(SplitNewVariableName, TEXT(".")));
+		}
+	}
+}
+
+void FNiagaraStackGraphUtilities::RenameReferencingParameters(UNiagaraSystem& System, UNiagaraEmitter* Emitter, UNiagaraNodeFunctionCall& FunctionCallNode, const FString& OldModuleName, const FString& NewModuleName)
+{
+	TMap<FName, FName> OldNameToNewNameMap;
+	FNiagaraStackGraphUtilities::GatherRenamedStackFunctionInputAndOutputVariableNames(Emitter, FunctionCallNode, OldModuleName, NewModuleName, OldNameToNewNameMap);
+
+	// local function to rename pins referencing the given module
+	auto RenamePinsReferencingModule = [&OldNameToNewNameMap](UNiagaraNodeParameterMapBase* Node)
+	{
+		for (UEdGraphPin* Pin : Node->Pins)
+		{
+			FName* NewName = OldNameToNewNameMap.Find(Pin->PinName);
+			if (NewName != nullptr)
+			{
+				Node->SetPinName(Pin, *NewName);
+			}
+		}
+	};
+
+	UNiagaraNodeParameterMapSet* ParameterMapSet = FNiagaraStackGraphUtilities::GetStackFunctionOverrideNode(FunctionCallNode);
+	if (ParameterMapSet != nullptr)
+	{
+		RenamePinsReferencingModule(ParameterMapSet);
+	}
+
+	TArray<TWeakObjectPtr<UNiagaraScript>> Scripts;
+	FindAffectedScripts(System, Emitter, FunctionCallNode, Scripts);
+
+	const UNiagaraNodeOutput* OutputNode = GetEmitterOutputNodeForStackNode(FunctionCallNode);
+	UNiagaraGraph* OwningGraph = FunctionCallNode.GetNiagaraGraph();
+
+	FString OwningEmitterName = Emitter != nullptr ? Emitter->GetUniqueEmitterName() : FString();
+
+	for (TWeakObjectPtr<UNiagaraScript> Script : Scripts)
+	{
+		if (!Script.IsValid(false))
+		{
+			continue;
+		}
+
+		TArray<FNiagaraVariable> RapidIterationVariables;
+		Script->RapidIterationParameters.GetParameters(RapidIterationVariables);
+
+		for (FNiagaraVariable& Variable : RapidIterationVariables)
+		{
+			FString EmitterName, FunctionCallName, InputName;
+			if (FNiagaraParameterMapHistory::SplitRapidIterationParameterName(Variable, EmitterName, FunctionCallName, InputName))
+			{
+				if (EmitterName == OwningEmitterName && FunctionCallName == OldModuleName)
+				{
+					FName NewParameterName(*(NewModuleName + TEXT(".") + InputName));
+					FNiagaraVariable NewParameter = FNiagaraStackGraphUtilities::CreateRapidIterationParameter(EmitterName, Script->GetUsage(), NewParameterName, Variable.GetType());
+					Script->RapidIterationParameters.RenameParameter(Variable, NewParameter.GetName());
+				}
+			}
+		}
+
+		if (UNiagaraScriptSource* ScriptSource = Cast<UNiagaraScriptSource>(Script->GetSource()))
+		{
+			// rename all parameter map get nodes that use the parameter name
+			TArray<UNiagaraNodeParameterMapGet*> ParameterGetNodes;
+			ScriptSource->NodeGraph->GetNodesOfClass<UNiagaraNodeParameterMapGet>(ParameterGetNodes);
+
+			for (UNiagaraNodeParameterMapGet* Node : ParameterGetNodes)
+			{
+				RenamePinsReferencingModule(Node);
+			}
+		}
+	}
+}
+
 #undef LOCTEXT_NAMESPACE
diff --git a/Engine/Plugins/FX/Niagara/Source/NiagaraEditor/Private/ViewModels/Stack/NiagaraStackModuleItem.cpp b/Engine/Plugins/FX/Niagara/Source/NiagaraEditor/Private/ViewModels/Stack/NiagaraStackModuleItem.cpp
--- a/Engine/Plugins/FX/Niagara/Source/NiagaraEditor/Private/ViewModels/Stack/NiagaraStackModuleItem.cpp
+++ b/Engine/Plugins/FX/Niagara/Source/NiagaraEditor/Private/ViewModels/Stack/NiagaraStackModuleItem.cpp
@@ -913,8 +913,23 @@ void UNiagaraStackModuleItem::ReassignModuleScript(UNiagaraScript* ModuleScript)
 		TEXT("Can not reassign the module script when the module isn't a valid function call module.")))
 	{
 		FScopedTransaction ScopedTransaction(LOCTEXT("ReassignModuleTransaction", "Reassign module script"));
+
+		const FString OldName = FunctionCallNode->GetFunctionName();
+		const UNiagaraScript* OldScript = FunctionCallNode->FunctionScript;
+
 		FunctionCallNode->Modify();
 		FunctionCallNode->FunctionScript = ModuleScript;
+		
+		// intermediate refresh to purge any rapid iteration parameters that have been removed in the new script
+		RefreshChildren();
+
+		FunctionCallNode->SuggestName(FString());
+
+		const FString NewName = FunctionCallNode->GetFunctionName();
+		UNiagaraSystem& System = GetSystemViewModel()->GetSystem();
+		UNiagaraEmitter* Emitter = GetEmitterViewModel().IsValid() ? GetEmitterViewModel()->GetEmitter() : nullptr;
+		FNiagaraStackGraphUtilities::RenameReferencingParameters(System, Emitter, *FunctionCallNode, OldName, NewName);
+
 		FunctionCallNode->RefreshFromExternalChanges();
 		FunctionCallNode->MarkNodeRequiresSynchronization(TEXT("Module script reassigned."), true);
 		RefreshChildren();
diff --git a/Engine/Plugins/FX/Niagara/Source/NiagaraEditor/Private/ViewModels/Stack/NiagaraStackScriptItemGroup.cpp b/Engine/Plugins/FX/Niagara/Source/NiagaraEditor/Private/ViewModels/Stack/NiagaraStackScriptItemGroup.cpp
--- a/Engine/Plugins/FX/Niagara/Source/NiagaraEditor/Private/ViewModels/Stack/NiagaraStackScriptItemGroup.cpp
+++ b/Engine/Plugins/FX/Niagara/Source/NiagaraEditor/Private/ViewModels/Stack/NiagaraStackScriptItemGroup.cpp
@@ -174,6 +174,7 @@ class FScriptItemGroupAddUtilities : public TNiagaraStackItemGroupAddUtilities<U
 public:
 	FScriptItemGroupAddUtilities(TSharedRef<FNiagaraSystemViewModel> InSystemViewModel, TSharedPtr<FNiagaraEmitterViewModel> InEmitterViewModel, UNiagaraStackEditorData& InStackEditorData, FOnItemAdded InOnItemAdded)
 		: TNiagaraStackItemGroupAddUtilities(LOCTEXT("ScriptGroupAddItemName", "Module"), EAddMode::AddFromAction, false, InOnItemAdded)
+		, OutputNode(nullptr)
 		, SystemViewModel(InSystemViewModel)
 		, EmitterViewModel(InEmitterViewModel)
 		, StackEditorData(InStackEditorData)
@@ -980,38 +981,15 @@ void GatherRenamedModuleOutputs(
 	TArray<TPair<const UNiagaraClipboardFunction*, UNiagaraNodeFunctionCall*>> InClipboardFunctionAndNodeFunctionPairs,
 	TMap<FName, FName>& OutOldModuleOutputNameToNewModuleOutputNameMap)
 {
+	UNiagaraEmitter* Emitter = InOwnerEntry->GetEmitterViewModel().IsValid() ? InOwnerEntry->GetEmitterViewModel()->GetEmitter() : nullptr;
+
 	for (TPair<const UNiagaraClipboardFunction*, UNiagaraNodeFunctionCall*>& ClipboardFunctionAndNodeFunctionPair : InClipboardFunctionAndNodeFunctionPairs)
 	{
 		const UNiagaraClipboardFunction* ClipboardFunction = ClipboardFunctionAndNodeFunctionPair.Key;
 		UNiagaraNodeFunctionCall* NewFunctionCallNode = ClipboardFunctionAndNodeFunctionPair.Value;
 		if (NewFunctionCallNode->GetFunctionName() != ClipboardFunction->FunctionName)
 		{
-			TArray<FNiagaraVariable> OutputVariables;
-			TArray<FNiagaraVariable> OutputVariablesWithOriginalAliasesIntact;
-			FCompileConstantResolver ConstantResolver = InOwnerEntry->GetEmitterViewModel().IsValid()
-				? FCompileConstantResolver(InOwnerEntry->GetEmitterViewModel()->GetEmitter())
-				: FCompileConstantResolver();
-			FNiagaraStackGraphUtilities::GetStackFunctionOutputVariables(*NewFunctionCallNode, ConstantResolver, OutputVariables, OutputVariablesWithOriginalAliasesIntact);
-
-			for (FNiagaraVariable& OutputVariableWithOriginalAliasesIntact : OutputVariablesWithOriginalAliasesIntact)
-			{
-				TArray<FString> SplitAliasedVariableName;
-				OutputVariableWithOriginalAliasesIntact.GetName().ToString().ParseIntoArray(SplitAliasedVariableName, TEXT("."));
-				if (SplitAliasedVariableName.Contains(TEXT("Module")))
-				{
-					TArray<FString> SplitOldVariableName = SplitAliasedVariableName;
-					TArray<FString> SplitNewVariableName = SplitAliasedVariableName;
-					for (int32 i = 0; i < SplitAliasedVariableName.Num(); i++)
-					{
-						if (SplitAliasedVariableName[i] == TEXT("Module"))
-						{
-							SplitOldVariableName[i] = ClipboardFunction->FunctionName;
-							SplitNewVariableName[i] = NewFunctionCallNode->GetFunctionName();
-						}
-					}
-					OutOldModuleOutputNameToNewModuleOutputNameMap.Add(*FString::Join(SplitOldVariableName, TEXT(".")), *FString::Join(SplitNewVariableName, TEXT(".")));
-				}
-			}
+			FNiagaraStackGraphUtilities::GatherRenamedStackFunctionOutputVariableNames(Emitter, *NewFunctionCallNode, ClipboardFunction->FunctionName, NewFunctionCallNode->GetFunctionName(), OutOldModuleOutputNameToNewModuleOutputNameMap);
 		}
 	}
 }
diff --git a/Engine/Plugins/FX/Niagara/Source/NiagaraEditor/Public/ViewModels/Stack/NiagaraStackGraphUtilities.h b/Engine/Plugins/FX/Niagara/Source/NiagaraEditor/Public/ViewModels/Stack/NiagaraStackGraphUtilities.h
--- a/Engine/Plugins/FX/Niagara/Source/NiagaraEditor/Public/ViewModels/Stack/NiagaraStackGraphUtilities.h
+++ b/Engine/Plugins/FX/Niagara/Source/NiagaraEditor/Public/ViewModels/Stack/NiagaraStackGraphUtilities.h
@@ -87,6 +87,8 @@ namespace FNiagaraStackGraphUtilities
 
 	void GetStackFunctionOutputVariables(UNiagaraNodeFunctionCall& FunctionCallNode, FCompileConstantResolver ConstantResolver, TArray<FNiagaraVariable>& OutOutputVariables, TArray<FNiagaraVariable>& OutOutputVariablesWithOriginalAliasesIntact);
 
+	void GetStackFunctionInputAndOutputVariables(UNiagaraNodeFunctionCall& FunctionCallNode, FCompileConstantResolver ConstantResolver, TArray<FNiagaraVariable>& OutVariables, TArray<FNiagaraVariable>& OutVariablesWithOriginalAliasesIntact);
+
 	UNiagaraNodeParameterMapSet* GetStackFunctionOverrideNode(UNiagaraNodeFunctionCall& FunctionCallNode);
 
 	UNiagaraNodeParameterMapSet& GetOrCreateStackFunctionOverrideNode(UNiagaraNodeFunctionCall& FunctionCallNode, const FGuid& PreferredOverrideNodeGuid = FGuid());
@@ -179,4 +181,11 @@ namespace FNiagaraStackGraphUtilities
 	NIAGARAEDITOR_API bool ParameterAllowedInExecutionCategory(const FName InParameterName, const FName ExecutionCategory);
 
 	void RebuildEmitterNodes(UNiagaraSystem& System);
+
+	void FindAffectedScripts(UNiagaraSystem& System, UNiagaraEmitter* Emitter, UNiagaraNodeFunctionCall& ModuleNode, TArray<TWeakObjectPtr<UNiagaraScript>>& OutAffectedScripts);
+
+	void RenameReferencingParameters(UNiagaraSystem& System, UNiagaraEmitter* Emitter, UNiagaraNodeFunctionCall& FunctionCallNode, const FString& OldName, const FString& NewName);
+
+	void GatherRenamedStackFunctionOutputVariableNames(UNiagaraEmitter* Emitter, UNiagaraNodeFunctionCall& FunctionCallNode, const FString& OldFunctionName, const FString& NewFunctionName, TMap<FName, FName>& OutOldToNewNameMap);
+	void GatherRenamedStackFunctionInputAndOutputVariableNames(UNiagaraEmitter* Emitter, UNiagaraNodeFunctionCall& FunctionCallNode, const FString& OldFunctionName, const FString& NewFunctionName, TMap<FName, FName>& OutOldToNewNameMap);
 }
