Fix for broken engine tests.
Reverted change in burst spawn module.
Moved logic to determin which emitters should tick out of finalize and into the concurrent tick.

diff --git a/Engine/Plugins/FX/Niagara/Source/Niagara/Private/NiagaraSystemInstance.cpp b/Engine/Plugins/FX/Niagara/Source/Niagara/Private/NiagaraSystemInstance.cpp
--- a/Engine/Plugins/FX/Niagara/Source/Niagara/Private/NiagaraSystemInstance.cpp
+++ b/Engine/Plugins/FX/Niagara/Source/Niagara/Private/NiagaraSystemInstance.cpp
@@ -74,7 +74,6 @@ FNiagaraSystemInstance::FNiagaraSystemInstance(UNiagaraComponent* InComponent)
 	, bSolo(false)
 	, bForceSolo(false)
 	, bPendingSpawn(false)
-	, bHasTickingEmitters(true)
 	, bPaused(false)
 	, bDataInterfacesHaveTickPrereqs(false)
 	, bIsTransformDirty(true)
@@ -712,7 +711,6 @@ void FNiagaraSystemInstance::ResetInternal(bool bResetSimulations)
 
 	Age = 0;
 	TickCount = 0;
-	bHasTickingEmitters = true;
 	CachedDeltaSeconds = 0.0f;
 	bLODDistanceIsValid = false;
 	// Note: We do not need to update our bounds here as they are still valid
@@ -841,7 +839,6 @@ void FNiagaraSystemInstance::ReInitInternal()
 
 	Age = 0;
 	TickCount = 0;
-	bHasTickingEmitters = true;
 	bIsTransformDirty = true;
 	TimeSinceLastForceUpdateTransform = 0.0f;
 	LocalBounds = FBox(FVector::ZeroVector, FVector::ZeroVector);
@@ -1827,12 +1824,10 @@ void FNiagaraSystemInstance::WaitForAsyncTickAndFinalize(bool bEnsureComplete)
 bool FNiagaraSystemInstance::HandleCompletion()
 {
 	bool bEmittersCompleteOrDisabled = true;
-	bHasTickingEmitters = false;
 	for (TSharedRef<FNiagaraEmitterInstance, ESPMode::ThreadSafe>&it : Emitters)
 	{
 		FNiagaraEmitterInstance& Inst = *it;
 		bEmittersCompleteOrDisabled &= Inst.HandleCompletion();
-		bHasTickingEmitters |= Inst.ShouldTick();
 	}
 
 	bool bCompletedAlready = IsComplete();
@@ -1888,6 +1883,17 @@ void FNiagaraSystemInstance::Tick_Concurrent()
 	ActiveGPUEmitterCount = 0;
 	UNiagaraSystem* System = GetSystem();
 
+	//Determine if any of our emitters should be ticking.
+	TBitArray<TInlineAllocator<8>> EmittersShouldTick;
+	bool bHasTickingEmitters = false;
+	for (int32 EmitterIdx = 0; EmitterIdx < Emitters.Num(); ++EmitterIdx)
+	{
+		FNiagaraEmitterInstance& Inst = Emitters[EmitterIdx].Get();
+		bool bShouldTick = Inst.ShouldTick();
+		bHasTickingEmitters |= bShouldTick;
+		EmittersShouldTick.Add(bShouldTick);
+	}
+
 	if (IsComplete() || !bHasTickingEmitters || GetSystem() == nullptr || Component == nullptr || CachedDeltaSeconds < SMALL_NUMBER)
 	{
 		bAsyncWorkInProgress = false;
@@ -1898,20 +1904,26 @@ void FNiagaraSystemInstance::Tick_Concurrent()
 
 	for (int32 EmitterIdx = 0; EmitterIdx < Emitters.Num(); EmitterIdx++)
 	{
-		FNiagaraEmitterInstance& Inst = Emitters[EmitterIdx].Get();
-		Inst.PreTick();
+		if (EmittersShouldTick[EmitterIdx])
+		{
+			FNiagaraEmitterInstance& Inst = Emitters[EmitterIdx].Get();
+			Inst.PreTick();
+		}
 	}
 
 	// now tick all emitters
 	for (int32 EmitterIdx = 0; EmitterIdx < Emitters.Num(); EmitterIdx++)
 	{
-		FNiagaraEmitterInstance& Inst = Emitters[EmitterIdx].Get();
-		Inst.Tick(CachedDeltaSeconds);
-
-		if (Inst.GetCachedEmitter() && Inst.GetCachedEmitter()->SimTarget == ENiagaraSimTarget::GPUComputeSim && Inst.GetGPUContext() != nullptr && (Inst.GetExecutionState() != ENiagaraExecutionState::Complete))
+		if (EmittersShouldTick[EmitterIdx])
 		{
-			TotalParamSize += Inst.GetGPUContext()->CombinedParamStore.GetPaddedParameterSizeInBytes();
-			ActiveGPUEmitterCount++;
+			FNiagaraEmitterInstance& Inst = Emitters[EmitterIdx].Get();
+			Inst.Tick(CachedDeltaSeconds);
+
+			if (Inst.GetCachedEmitter() && Inst.GetCachedEmitter()->SimTarget == ENiagaraSimTarget::GPUComputeSim && Inst.GetGPUContext() != nullptr && (Inst.GetExecutionState() != ENiagaraExecutionState::Complete))
+			{
+				TotalParamSize += Inst.GetGPUContext()->CombinedParamStore.GetPaddedParameterSizeInBytes();
+				ActiveGPUEmitterCount++;
+			}
 		}
 	}
 
diff --git a/Engine/Plugins/FX/Niagara/Source/Niagara/Public/NiagaraSystemInstance.h b/Engine/Plugins/FX/Niagara/Source/Niagara/Public/NiagaraSystemInstance.h
--- a/Engine/Plugins/FX/Niagara/Source/Niagara/Public/NiagaraSystemInstance.h
+++ b/Engine/Plugins/FX/Niagara/Source/Niagara/Public/NiagaraSystemInstance.h
@@ -202,8 +202,6 @@ public:
 
 	bool IsReadyToRun() const;
 
-	FORCEINLINE bool HasTickingEmitters()const { return bHasTickingEmitters; }
-
 	UNiagaraParameterCollectionInstance* GetParameterCollectionInstance(UNiagaraParameterCollection* Collection);
 
 	/** 
@@ -391,9 +389,6 @@ public:
 	uint32 bPendingSpawn : 1;
 	uint32 bNotifyOnCompletion : 1;
 
-	/** If this instance has any currently ticking emitters. If false, allows us to skip some work. */
-	uint32 bHasTickingEmitters : 1;
-
 	/** If this system is paused. When paused it will not tick and never complete etc. */
 	uint32 bPaused : 1;
 	/** If this system has emitters that will run GPU Simulations */
