Preview cpu camera DI in the level editor without PIE

diff --git a/Engine/Plugins/FX/Niagara/Source/Niagara/Niagara.Build.cs b/Engine/Plugins/FX/Niagara/Source/Niagara/Niagara.Build.cs
--- a/Engine/Plugins/FX/Niagara/Source/Niagara/Niagara.Build.cs
+++ b/Engine/Plugins/FX/Niagara/Source/Niagara/Niagara.Build.cs
@@ -58,6 +58,7 @@ public class Niagara : ModuleRules
             PrivateDependencyModuleNames.AddRange(
                 new string[] {
                 "TargetPlatform",
+                "UnrealEd",
             });
         }
     }
diff --git a/Engine/Plugins/FX/Niagara/Source/Niagara/Private/NiagaraDataInterfaceCamera.cpp b/Engine/Plugins/FX/Niagara/Source/Niagara/Private/NiagaraDataInterfaceCamera.cpp
--- a/Engine/Plugins/FX/Niagara/Source/Niagara/Private/NiagaraDataInterfaceCamera.cpp
+++ b/Engine/Plugins/FX/Niagara/Source/Niagara/Private/NiagaraDataInterfaceCamera.cpp
@@ -7,6 +7,11 @@
 #include "NiagaraSystemInstance.h"
 #include "GameFramework/PlayerController.h"
 
+#if WITH_EDITORONLY_DATA
+#include "EditorViewportClient.h"
+#include "LevelEditorViewport.h"
+#endif
+
 const FName UNiagaraDataInterfaceCamera::GetViewPropertiesName(TEXT("GetViewPropertiesGPU"));
 const FName UNiagaraDataInterfaceCamera::GetClipSpaceTransformsName(TEXT("GetClipSpaceTransformsGPU"));
 const FName UNiagaraDataInterfaceCamera::GetViewSpaceTransformsName(TEXT("GetViewSpaceTransformsGPU"));
@@ -55,11 +60,22 @@ bool UNiagaraDataInterfaceCamera::PerInstanceTick(void* PerInstanceData, FNiagar
 				PIData->CameraLocation = PlayerController->PlayerCameraManager->GetCameraLocation();
 				PIData->CameraRotation = PlayerController->PlayerCameraManager->GetCameraRotation();
 				PIData->CameraFOV = PlayerController->PlayerCameraManager->GetFOVAngle();
+				return false;
 			}
 			i++;
 		}
+	}
+#if WITH_EDITORONLY_DATA
+	if (GCurrentLevelEditingViewportClient)
+	{
+		const FViewportCameraTransform& ViewTransform = GCurrentLevelEditingViewportClient->GetViewTransform();
+		PIData->CameraLocation = ViewTransform.GetLocation();
+		PIData->CameraRotation = ViewTransform.GetRotation();
+		PIData->CameraFOV = GCurrentLevelEditingViewportClient->ViewFOV;
 		return false;
 	}
+#endif
+
 	PIData->CameraLocation = FVector::ZeroVector;
 	PIData->CameraRotation = FRotator(0);
 	PIData->CameraFOV = 0;
