Niagara - Fix some inconsistencies with the way input nodes and data interfaces are created which in specific circumstances could cause data interfaces to fail to bind correctly post compile.  Fixes recent cook errors related to data interfaces.

diff --git a/Engine/Plugins/FX/Niagara/Source/Niagara/Private/NiagaraScript.cpp b/Engine/Plugins/FX/Niagara/Source/Niagara/Private/NiagaraScript.cpp
--- a/Engine/Plugins/FX/Niagara/Source/Niagara/Private/NiagaraScript.cpp
+++ b/Engine/Plugins/FX/Niagara/Source/Niagara/Private/NiagaraScript.cpp
@@ -1249,7 +1249,12 @@ void UNiagaraScript::SetVMCompilationResults(const FNiagaraVMExecutableDataId& I
 
 			if (Info.bIsPlaceholder == false)
 			{
-				UE_LOG(LogNiagara, Error, TEXT("We somehow ended up with a data interface that we couldn't match post compile. This shouldn't happen. Creating a dummy to prevent crashes. DataInterfaceInfoName:%s Object:%s"), *Info.Name.ToString(), *GetPathNameSafe(this));
+				UE_LOG(LogNiagara, Warning, TEXT("We somehow ended up with a data interface that we couldn't match post compile. This shouldn't happen. Creating a dummy to prevent crashes. DataInterfaceInfoName:%s Object:%s"), *Info.Name.ToString(), *GetPathNameSafe(this));
+				UE_LOG(LogNiagara, Log, TEXT("Object to Name map contents:"));
+				for (const TPair<FName, UNiagaraDataInterface*>& Pair : InRequestData->GetObjectNameMap())
+				{
+					UE_LOG(LogNiagara, Log, TEXT("%s -> %s"), *Pair.Key.ToString(), *GetPathNameSafe(Pair.Value));
+				}
 			}
 		}
 		check(CachedDefaultDataInterfaces[Idx].DataInterface != nullptr);
diff --git a/Engine/Plugins/FX/Niagara/Source/NiagaraCore/Private/NiagaraCustomVersion.cpp b/Engine/Plugins/FX/Niagara/Source/NiagaraCore/Private/NiagaraCustomVersion.cpp
--- a/Engine/Plugins/FX/Niagara/Source/NiagaraCore/Private/NiagaraCustomVersion.cpp
+++ b/Engine/Plugins/FX/Niagara/Source/NiagaraCore/Private/NiagaraCustomVersion.cpp
@@ -8,4 +8,4 @@ const FGuid FNiagaraCustomVersion::GUID(0xFCF57AFA, 0x50764283, 0xB9A9E658, 0xFF
 // Register the custom version with core
 FCustomVersionRegistration GRegisterNiagaraCustomVersion(FNiagaraCustomVersion::GUID, FNiagaraCustomVersion::LatestVersion, TEXT("NiagaraVer"));
 
-const FGuid FNiagaraCustomVersion::LatestScriptCompileVersion(0xA6B5A395, 0xF276467E, 0x9D03C2C2, 0x99CAE1E3);
+const FGuid FNiagaraCustomVersion::LatestScriptCompileVersion(0xC1786B5B, 0xFAF74AD2, 0x92CF74E9, 0x65C2658F);
diff --git a/Engine/Plugins/FX/Niagara/Source/NiagaraEditor/Private/NiagaraScriptMergeManager.cpp b/Engine/Plugins/FX/Niagara/Source/NiagaraEditor/Private/NiagaraScriptMergeManager.cpp
--- a/Engine/Plugins/FX/Niagara/Source/NiagaraEditor/Private/NiagaraScriptMergeManager.cpp
+++ b/Engine/Plugins/FX/Niagara/Source/NiagaraEditor/Private/NiagaraScriptMergeManager.cpp
@@ -66,6 +66,7 @@ FNiagaraStackFunctionInputOverrideMergeAdapter::FNiagaraStackFunctionInputOverri
 		else if (OverridePin->LinkedTo[0]->GetOwningNode()->IsA<UNiagaraNodeInput>())
 		{
 			UNiagaraNodeInput* DataInputNode = CastChecked<UNiagaraNodeInput>(OverridePin->LinkedTo[0]->GetOwningNode());
+			DataValueInputName = DataInputNode->Input.GetName();
 			DataValueObject = DataInputNode->GetDataInterface();
 		}
 		else if (OverridePin->LinkedTo[0]->GetOwningNode()->IsA<UNiagaraNodeFunctionCall>())
@@ -163,6 +164,11 @@ TOptional<FNiagaraParameterHandle> FNiagaraStackFunctionInputOverrideMergeAdapte
 	return LinkedValueHandle;
 }
 
+TOptional<FName> FNiagaraStackFunctionInputOverrideMergeAdapter::GetDataValueInputName() const
+{
+	return DataValueInputName;
+}
+
 UNiagaraDataInterface* FNiagaraStackFunctionInputOverrideMergeAdapter::GetDataValueObject() const
 {
 	return DataValueObject;
@@ -1819,15 +1825,20 @@ TOptional<bool> FNiagaraScriptMergeManager::DoFunctionInputOverridesMatch(TShare
 	}
 
 	// Data value
-	if ((BaseFunctionInputAdapter->GetDataValueObject()!= nullptr && OtherFunctionInputAdapter->GetDataValueObject()== nullptr) ||
-		(BaseFunctionInputAdapter->GetDataValueObject()== nullptr && OtherFunctionInputAdapter->GetDataValueObject()!= nullptr))
+	if ((BaseFunctionInputAdapter->GetDataValueInputName().IsSet() && OtherFunctionInputAdapter->GetDataValueInputName().IsSet() == false) ||
+		(BaseFunctionInputAdapter->GetDataValueInputName().IsSet() == false && OtherFunctionInputAdapter->GetDataValueInputName().IsSet()) ||
+		(BaseFunctionInputAdapter->GetDataValueObject() != nullptr && OtherFunctionInputAdapter->GetDataValueObject() == nullptr) ||
+		(BaseFunctionInputAdapter->GetDataValueObject() == nullptr && OtherFunctionInputAdapter->GetDataValueObject() != nullptr))
 	{
 		return false;
 	}
 
-	if (BaseFunctionInputAdapter->GetDataValueObject()!= nullptr && OtherFunctionInputAdapter->GetDataValueObject()!= nullptr)
+	if (BaseFunctionInputAdapter->GetDataValueInputName().IsSet() && OtherFunctionInputAdapter->GetDataValueInputName().IsSet() &&
+		BaseFunctionInputAdapter->GetDataValueObject() != nullptr && OtherFunctionInputAdapter->GetDataValueObject() != nullptr)
 	{
-		return BaseFunctionInputAdapter->GetDataValueObject()->Equals(OtherFunctionInputAdapter->GetDataValueObject());
+		return 
+			BaseFunctionInputAdapter->GetDataValueInputName().GetValue() == OtherFunctionInputAdapter->GetDataValueInputName().GetValue() &&
+			BaseFunctionInputAdapter->GetDataValueObject()->Equals(OtherFunctionInputAdapter->GetDataValueObject());
 	}
 
 	// Dynamic value
@@ -2012,11 +2023,12 @@ FNiagaraScriptMergeManager::FApplyDiffResults FNiagaraScriptMergeManager::AddInp
 				FNiagaraStackGraphUtilities::SetLinkedValueHandleForFunctionInput(InputOverridePin, OverrideToAdd->GetLinkedValueHandle().GetValue(), OverrideToAdd->GetOverrideNodeId());
 				Results.bSucceeded = true;
 			}
-			else if (OverrideToAdd->GetDataValueObject() != nullptr)
+			else if (OverrideToAdd->GetDataValueInputName().IsSet() && OverrideToAdd->GetDataValueObject() != nullptr)
 			{
+				FName OverrideValueInputName = OverrideToAdd->GetDataValueInputName().GetValue();
 				UNiagaraDataInterface* OverrideValueObject = OverrideToAdd->GetDataValueObject();
 				UNiagaraDataInterface* NewOverrideValueObject;
-				FNiagaraStackGraphUtilities::SetDataValueObjectForFunctionInput(InputOverridePin, OverrideToAdd->GetDataValueObject()->GetClass(), OverrideValueObject->GetName(), NewOverrideValueObject, OverrideToAdd->GetOverrideNodeId());
+				FNiagaraStackGraphUtilities::SetDataValueObjectForFunctionInput(InputOverridePin, OverrideToAdd->GetDataValueObject()->GetClass(), OverrideValueInputName.ToString(), NewOverrideValueObject, OverrideToAdd->GetOverrideNodeId());
 				OverrideValueObject->CopyTo(NewOverrideValueObject);
 				Results.bSucceeded = true;
 			}
diff --git a/Engine/Plugins/FX/Niagara/Source/NiagaraEditor/Private/NiagaraScriptMergeManager.h b/Engine/Plugins/FX/Niagara/Source/NiagaraEditor/Private/NiagaraScriptMergeManager.h
--- a/Engine/Plugins/FX/Niagara/Source/NiagaraEditor/Private/NiagaraScriptMergeManager.h
+++ b/Engine/Plugins/FX/Niagara/Source/NiagaraEditor/Private/NiagaraScriptMergeManager.h
@@ -57,6 +57,7 @@ public:
 	TOptional<FString> GetLocalValueString() const;
 	TOptional<FNiagaraVariable> GetLocalValueRapidIterationParameter() const;
 	TOptional<FNiagaraParameterHandle> GetLinkedValueHandle() const;
+	TOptional<FName> GetDataValueInputName() const;
 	UNiagaraDataInterface* GetDataValueObject() const;
 	TSharedPtr<FNiagaraStackFunctionMergeAdapter> GetDynamicValueFunction() const;
 	TOptional<FString> GetStaticSwitchValue() const;
@@ -74,6 +75,7 @@ private:
 	TOptional<FString> LocalValueString;
 	TOptional<FNiagaraVariable> LocalValueRapidIterationParameter;
 	TOptional<FNiagaraParameterHandle> LinkedValueHandle;
+	TOptional<FName> DataValueInputName;
 	UNiagaraDataInterface* DataValueObject;
 	TSharedPtr<FNiagaraStackFunctionMergeAdapter> DynamicValueFunction;
 	TOptional<FString> StaticSwitchValue;
diff --git a/Engine/Plugins/FX/Niagara/Source/NiagaraEditor/Private/ViewModels/Stack/NiagaraStackGraphUtilities.cpp b/Engine/Plugins/FX/Niagara/Source/NiagaraEditor/Private/ViewModels/Stack/NiagaraStackGraphUtilities.cpp
--- a/Engine/Plugins/FX/Niagara/Source/NiagaraEditor/Private/ViewModels/Stack/NiagaraStackGraphUtilities.cpp
+++ b/Engine/Plugins/FX/Niagara/Source/NiagaraEditor/Private/ViewModels/Stack/NiagaraStackGraphUtilities.cpp
@@ -30,6 +30,7 @@
 #include "EdGraph/EdGraphPin.h"
 #include "ViewModels/NiagaraEmitterViewModel.h"
 #include "AssetRegistryModule.h"
+#include "ObjectTools.h"
 
 DECLARE_CYCLE_STAT(TEXT("Niagara - StackGraphUtilities - RelayoutGraph"), STAT_NiagaraEditor_StackGraphUtilities_RelayoutGraph, STATGROUP_NiagaraEditor);
 
@@ -986,7 +987,7 @@ void FNiagaraStackGraphUtilities::SetLinkedValueHandleForFunctionInput(UEdGraphP
 	}
 }
 
-void FNiagaraStackGraphUtilities::SetDataValueObjectForFunctionInput(UEdGraphPin& OverridePin, UClass* DataObjectType, FString DataObjectName, UNiagaraDataInterface*& OutDataObject, const FGuid& NewNodePersistentId)
+void FNiagaraStackGraphUtilities::SetDataValueObjectForFunctionInput(UEdGraphPin& OverridePin, UClass* DataObjectType, FString InputNodeInputName, UNiagaraDataInterface*& OutDataObject, const FGuid& NewNodePersistentId)
 {
 	checkf(OverridePin.LinkedTo.Num() == 0, TEXT("Can't set a data value when the override pin already has a value."));
 	checkf(DataObjectType->IsChildOf<UNiagaraDataInterface>(), TEXT("Can only set a function input to a data interface value object"));
@@ -996,9 +997,9 @@ void FNiagaraStackGraphUtilities::SetDataValueObjectForFunctionInput(UEdGraphPin
 	Graph->Modify();
 	FGraphNodeCreator<UNiagaraNodeInput> InputNodeCreator(*Graph);
 	UNiagaraNodeInput* InputNode = InputNodeCreator.CreateNode();
-	FNiagaraEditorUtilities::InitializeParameterInputNode(*InputNode, FNiagaraTypeDefinition(DataObjectType), CastChecked<UNiagaraGraph>(Graph), *DataObjectName);
+	FNiagaraEditorUtilities::InitializeParameterInputNode(*InputNode, FNiagaraTypeDefinition(DataObjectType), CastChecked<UNiagaraGraph>(Graph), *InputNodeInputName);
 
-	OutDataObject = NewObject<UNiagaraDataInterface>(InputNode, DataObjectType, *DataObjectName, RF_Transactional | RF_Public);
+	OutDataObject = NewObject<UNiagaraDataInterface>(InputNode, DataObjectType, *ObjectTools::SanitizeObjectName(InputNodeInputName), RF_Transactional | RF_Public);
 	InputNode->SetDataInterface(OutDataObject);
 
 	InputNodeCreator.Finalize();
diff --git a/Engine/Plugins/FX/Niagara/Source/NiagaraEditor/Private/ViewModels/Stack/NiagaraStackParameterStoreEntry.cpp b/Engine/Plugins/FX/Niagara/Source/NiagaraEditor/Private/ViewModels/Stack/NiagaraStackParameterStoreEntry.cpp
--- a/Engine/Plugins/FX/Niagara/Source/NiagaraEditor/Private/ViewModels/Stack/NiagaraStackParameterStoreEntry.cpp
+++ b/Engine/Plugins/FX/Niagara/Source/NiagaraEditor/Private/ViewModels/Stack/NiagaraStackParameterStoreEntry.cpp
@@ -351,8 +351,8 @@ void UNiagaraStackParameterStoreEntry::RemovePins(TArray<UEdGraphPin*> OwningPin
 					UNiagaraDataInterface* CurrentValueDataInterface = Cast<UNiagaraDataInterface>(GetCurrentValueObject());
 					if (CurrentValueDataInterface != nullptr)
 					{
-						UNiagaraDataInterface* OverrideObj = NewObject<UNiagaraDataInterface>(this, const_cast<UClass*>(InputType.GetClass()));
-						FNiagaraStackGraphUtilities::SetDataValueObjectForFunctionInput(*OverridePin, const_cast<UClass*>(InputType.GetClass()), CurrentValueDataInterface->GetName(), OverrideObj);
+						UNiagaraDataInterface* OverrideObj;
+						FNiagaraStackGraphUtilities::SetDataValueObjectForFunctionInput(*OverridePin, const_cast<UClass*>(InputType.GetClass()), InputType.GetClass()->GetName(), OverrideObj);
 						CurrentValueDataInterface->CopyTo(OverrideObj);
 					}
 				}
diff --git a/Engine/Plugins/FX/Niagara/Source/NiagaraEditor/Public/ViewModels/Stack/NiagaraStackGraphUtilities.h b/Engine/Plugins/FX/Niagara/Source/NiagaraEditor/Public/ViewModels/Stack/NiagaraStackGraphUtilities.h
--- a/Engine/Plugins/FX/Niagara/Source/NiagaraEditor/Public/ViewModels/Stack/NiagaraStackGraphUtilities.h
+++ b/Engine/Plugins/FX/Niagara/Source/NiagaraEditor/Public/ViewModels/Stack/NiagaraStackGraphUtilities.h
@@ -103,7 +103,7 @@ namespace FNiagaraStackGraphUtilities
 
 	void SetLinkedValueHandleForFunctionInput(UEdGraphPin& OverridePin, FNiagaraParameterHandle LinkedParameterHandle, const FGuid& NewNodePersistentId = FGuid());
 
-	void SetDataValueObjectForFunctionInput(UEdGraphPin& OverridePin, UClass* DataObjectType, FString DataObjectName, UNiagaraDataInterface*& OutDataObject, const FGuid& NewNodePersistentId = FGuid());
+	void SetDataValueObjectForFunctionInput(UEdGraphPin& OverridePin, UClass* DataObjectType, FString InputNodeInputName, UNiagaraDataInterface*& OutDataObject, const FGuid& NewNodePersistentId = FGuid());
 
 	void SetDynamicInputForFunctionInput(UEdGraphPin& OverridePin, UNiagaraScript* DynamicInput, UNiagaraNodeFunctionCall*& OutDynamicInputFunctionCall, const FGuid& NewNodePersistentId = FGuid(), FString SuggestedName = FString());
 
