Making custom alignment + custom facing allow direct control of sprite orientation.

diff --git a/Engine/Plugins/FX/Niagara/Shaders/Private/NiagaraSpriteVertexFactory.ush b/Engine/Plugins/FX/Niagara/Shaders/Private/NiagaraSpriteVertexFactory.ush
--- a/Engine/Plugins/FX/Niagara/Shaders/Private/NiagaraSpriteVertexFactory.ush
+++ b/Engine/Plugins/FX/Niagara/Shaders/Private/NiagaraSpriteVertexFactory.ush
@@ -486,7 +486,12 @@ void GetTangents(FVertexFactoryInput Input, FVertexFactoryIntermediates Intermed
 	
 	/**This could well be better as a branch*/
 	BRANCH
-	if(NiagaraSpriteVFLooseParameters.ParticleAlignmentMode == ALIGNMENT_MODE_UNALIGNED)
+	if(NiagaraSpriteVFLooseParameters.ParticleAlignmentMode == ALIGNMENT_MODE_CUSTOMALIGNMENT && NiagaraSpriteVFLooseParameters.ParticleFacingMode == FACING_MODE_CUSTOMFACINGVECTOR)
+	{
+		Right = cross(Intermediates.CustomAlignmentVector, CustomFacingVector);
+		Up = Intermediates.CustomAlignmentVector;
+	}
+	else if(NiagaraSpriteVFLooseParameters.ParticleAlignmentMode == ALIGNMENT_MODE_UNALIGNED)
 	{
 		FLATTEN
 		if (NiagaraSpriteVFLooseParameters.ParticleFacingMode == FACING_MODE_FACECAMERADISTANCEBLEND)
