Adds a set of constant buffers to hold engine driven niagara variables so that we can reduce the cost of managing the parameters through the script execution pipeline
-reduces fixed concurrent CPU costs of emitters by about 15%
-emitter parameter storage remains in place for non-system parameters

diff --git a/Engine/Plugins/FX/Niagara/Source/Niagara/Classes/NiagaraDataSet.h b/Engine/Plugins/FX/Niagara/Source/Niagara/Classes/NiagaraDataSet.h
--- a/Engine/Plugins/FX/Niagara/Source/Niagara/Classes/NiagaraDataSet.h
+++ b/Engine/Plugins/FX/Niagara/Source/Niagara/Classes/NiagaraDataSet.h
@@ -169,7 +169,7 @@ private:
 	{
 		//Round up to VECTOR_WIDTH_BYTES.
 		//Both aligns the component buffers to the vector width but also ensures their ops cannot stomp over one another.		
-		return RequiredSize + VECTOR_WIDTH_BYTES - (RequiredSize % VECTOR_WIDTH_BYTES) + VECTOR_WIDTH_BYTES;
+		return Align(RequiredSize, VECTOR_WIDTH_BYTES) + VECTOR_WIDTH_BYTES;
 	}
 
 	/** Back ptr to our owning data set. Used to access layout info for the buffer. */
diff --git a/Engine/Plugins/FX/Niagara/Source/Niagara/Classes/NiagaraEmitterInstance.h b/Engine/Plugins/FX/Niagara/Source/Niagara/Classes/NiagaraEmitterInstance.h
--- a/Engine/Plugins/FX/Niagara/Source/Niagara/Classes/NiagaraEmitterInstance.h
+++ b/Engine/Plugins/FX/Niagara/Source/Niagara/Classes/NiagaraEmitterInstance.h
@@ -130,6 +130,9 @@ private:
 	void CheckForErrors();
 
 	void InitFastPathParameterBindingsInternal(const FNiagaraFastPathAttributeNames& SourceParameterNames, FNiagaraParameterStore& TargetParameterStore);
+	void BuildConstantBufferTable(
+		const FNiagaraScriptExecutionContext& ExecContext,
+		FScriptExecutionConstantBufferTable& ConstantBufferTable) const;
 
 	template<typename TBindingType, typename TVariableType>
 	static void AddBinding(FName ParameterName, TVariableType ParameterType, TBindingType* SourceValuePtr, FNiagaraParameterStore& TargetParameterStore, TArray<TNiagaraFastPathAttributeBinding<TBindingType>>& TargetBindings)
@@ -148,7 +151,7 @@ private:
 	int32 EmitterIdx;
 
 	/* The age of the emitter*/
-	float Age;
+	float EmitterAge;
 
 	int32 TickCount;
 
@@ -178,30 +181,9 @@ private:
 	FNiagaraParameterDirectBinding<float> InterpSpawnStartBinding;
 	FNiagaraParameterDirectBinding<int32> SpawnGroupBinding;
 
-	FNiagaraParameterDirectBinding<float> EmitterAgeBindingGPU;
-
-	FNiagaraParameterDirectBinding<float> SpawnEmitterAgeBinding;
-	FNiagaraParameterDirectBinding<float> UpdateEmitterAgeBinding;
-	TArray<FNiagaraParameterDirectBinding<float>> EventEmitterAgeBindings;
-
-	FNiagaraParameterDirectBinding<int32> SpawnRandomSeedBinding;
-	FNiagaraParameterDirectBinding<int32> UpdateRandomSeedBinding;
-	FNiagaraParameterDirectBinding<int32> GPURandomSeedBinding;
-	TArray<FNiagaraParameterDirectBinding<int>> EventRandomSeedBindings;
-
-	FNiagaraParameterDirectBinding<int32> SpawnTotalSpawnedParticlesBinding;
-	FNiagaraParameterDirectBinding<int32> UpdateTotalSpawnedParticlesBinding;
-	TArray<FNiagaraParameterDirectBinding<int32>> EventTotalSpawnedParticlesBindings;
-
 	FNiagaraParameterDirectBinding<int32> SpawnExecCountBinding;
 	FNiagaraParameterDirectBinding<int32> UpdateExecCountBinding;
 	TArray<FNiagaraParameterDirectBinding<int32>> EventExecCountBindings;
-
-	/*
-	FNiagaraParameterDirectBinding<int32> SpawnDeterminismBinding;
-	FNiagaraParameterDirectBinding<int32> UpdateDeterminismBinding;
-	FNiagaraParameterDirectBinding<int32> GPUDeterminismBinding;
-	*/
 	
 	/** particle simulation data. Must be a shared ref as various things on the RT can have direct ref to it. */
 	FNiagaraDataSet* ParticleDataSet;
diff --git a/Engine/Plugins/FX/Niagara/Source/Niagara/Classes/NiagaraScriptExecutionContext.h b/Engine/Plugins/FX/Niagara/Source/Niagara/Classes/NiagaraScriptExecutionContext.h
--- a/Engine/Plugins/FX/Niagara/Source/Niagara/Classes/NiagaraScriptExecutionContext.h
+++ b/Engine/Plugins/FX/Niagara/Source/Niagara/Classes/NiagaraScriptExecutionContext.h
@@ -120,6 +120,31 @@ struct FNiagaraDataSetExecutionInfo
 	bool bUpdateInstanceCount;
 };
 
+struct FScriptExecutionConstantBufferTable
+{
+	TArray<const uint8*, TInlineAllocator<12>> Buffers;
+	TArray<int32, TInlineAllocator<12>> BufferSizes;
+
+	void Reset(int32 ResetSize)
+	{
+		Buffers.Reset(ResetSize);
+		BufferSizes.Reset(ResetSize);
+	}
+
+	template<typename T>
+	void AddTypedBuffer(const T& Buffer)
+	{
+		Buffers.Add(reinterpret_cast<const uint8*>(&Buffer));
+		BufferSizes.Add(sizeof(T));
+	}
+
+	void AddRawBuffer(const uint8* BufferData, int32 BufferSize)
+	{
+		Buffers.Add(BufferData);
+		BufferSizes.Add(BufferSize);
+	}
+};
+
 struct FNiagaraScriptExecutionContext
 {
 	UNiagaraScript* Script;
@@ -139,6 +164,8 @@ struct FNiagaraScriptExecutionContext
 
 	static uint32 TickCounter;
 
+	int32 HasInterpolationParameters : 1;
+
 	FNiagaraScriptExecutionContext();
 	~FNiagaraScriptExecutionContext();
 
@@ -149,7 +176,7 @@ struct FNiagaraScriptExecutionContext
 
 	void BindData(int32 Index, FNiagaraDataSet& DataSet, int32 StartInstance, bool bUpdateInstanceCounts);
 	void BindData(int32 Index, FNiagaraDataBuffer* Input, int32 StartInstance, bool bUpdateInstanceCounts);
-	bool Execute(uint32 NumInstances);
+	bool Execute(uint32 NumInstances, const FScriptExecutionConstantBufferTable& ConstantBufferTable);
 
 	const TArray<UNiagaraDataInterface*>& GetDataInterfaces()const { return Parameters.GetDataInterfaces(); }
 
@@ -208,7 +235,13 @@ public:
 	class FNiagaraDataSet *MainDataSet;
 	UNiagaraScript* GPUScript;
 	class FNiagaraShaderScript*  GPUScript_RT;
-	FRHIUniformBufferLayout CBufferLayout; // Persistent layouts used to create Compute Sim CBuffer
+
+	// persistent layouts used to create the constant buffers for the compute sim shader
+	FRHIUniformBufferLayout GlobalCBufferLayout;
+	FRHIUniformBufferLayout SystemCBufferLayout;
+	FRHIUniformBufferLayout OwnerCBufferLayout;
+	FRHIUniformBufferLayout EmitterCBufferLayout;
+	FRHIUniformBufferLayout ExternalCBufferLayout;
 
 	//Dynamic state updated either from GT via RT commands or from the RT side sim code itself.
 	//TArray<uint8, TAlignedHeapAllocator<16>> ParamData_RT;		// RT side copy of the parameter data
@@ -229,6 +262,8 @@ public:
 	uint32 MaxUpdateIterations;
 	TSet<uint32> SpawnStages;
 
+	bool HasInterpolationParameters;
+
 	/** Temp data used in NiagaraEmitterInstanceBatcher::ExecuteAll() to avoid creating a map per FNiagaraComputeExecutionContext */
 	mutable int32 ScratchIndex = INDEX_NONE;
 
@@ -262,21 +297,18 @@ struct FNiagaraDataInterfaceInstanceData
 struct FNiagaraComputeInstanceData
 {
 	FNiagaraGpuSpawnInfo SpawnInfo;
-	uint8* ParamData;
-	FNiagaraComputeExecutionContext* Context;
+	uint8* GlobalParamData = nullptr;
+	uint8* SystemParamData = nullptr;
+	uint8* OwnerParamData = nullptr;
+	uint8* EmitterParamData = nullptr;
+	uint8* ExternalParamData = nullptr;
+	FNiagaraComputeExecutionContext* Context = nullptr;
 	TArray<FNiagaraDataInterfaceProxy*> DataInterfaceProxies;
 
 	//Buffer containing current state that this tick will read from. Initialized at the start of processing this tick on the RT.
-	FNiagaraDataBuffer* CurrentData;
+	FNiagaraDataBuffer* CurrentData = nullptr;
 	//Buffer into which we'll write the new simulation state. Initialized at the start of processing this tick on the RT.
-	FNiagaraDataBuffer* DestinationData;
-
-	FNiagaraComputeInstanceData()
-		: ParamData(nullptr)
-		, Context(nullptr)
-		, CurrentData(nullptr)
-		, DestinationData(nullptr)
-	{}
+	FNiagaraDataBuffer* DestinationData = nullptr;
 };
 
 /*
diff --git a/Engine/Plugins/FX/Niagara/Source/Niagara/Classes/NiagaraSystem.h b/Engine/Plugins/FX/Niagara/Source/Niagara/Classes/NiagaraSystem.h
--- a/Engine/Plugins/FX/Niagara/Source/Niagara/Classes/NiagaraSystem.h
+++ b/Engine/Plugins/FX/Niagara/Source/Niagara/Classes/NiagaraSystem.h
@@ -63,21 +63,54 @@ struct FNiagaraEmitterCompiledData
 };
 
 USTRUCT()
-struct FNiagaraSystemCompiledData
+struct FNiagaraParameterDataSetBinding
 {
 	GENERATED_USTRUCT_BODY()
 
 	UPROPERTY()
-	TArray<FNiagaraVariable> NumParticleVars;
+	int32 ParameterOffset;
 
 	UPROPERTY()
-	TArray<FNiagaraVariable> TotalSpawnedParticlesVars;
+	int32 DataSetComponentOffset;
+};
+
+USTRUCT()
+struct FNiagaraParameterDataSetBindingCollection
+{
+	GENERATED_USTRUCT_BODY()
 
 	UPROPERTY()
-	FNiagaraParameterStore InstanceParamStore;
+	TArray<FNiagaraParameterDataSetBinding> FloatOffsets;
 
 	UPROPERTY()
-	TArray<FNiagaraVariable> SpawnCountScaleVars;
+	TArray<FNiagaraParameterDataSetBinding> Int32Offsets;
+
+#if WITH_EDITORONLY_DATA
+	template<typename BufferType>
+	void Build(const FNiagaraDataSetCompiledData& DataSet)
+	{
+		BuildInternal(BufferType::GetVariables(), DataSet, TEXT(""), TEXT(""));
+	}
+
+	template<typename BufferType>
+	void Build(const FNiagaraDataSetCompiledData& DataSet, const FString& NamespaceBase, const FString& NamespaceReplacement)
+	{
+		BuildInternal(BufferType::GetVariables(), DataSet, NamespaceBase, NamespaceReplacement);
+	}
+
+protected:
+	void BuildInternal(const TArray<FNiagaraVariable>& ParameterVars, const FNiagaraDataSetCompiledData& DataSet, const FString& NamespaceBase, const FString& NamespaceReplacement);
+
+#endif
+};
+
+USTRUCT()
+struct FNiagaraSystemCompiledData
+{
+	GENERATED_USTRUCT_BODY()
+
+	UPROPERTY()
+	FNiagaraParameterStore InstanceParamStore;
 
 	UPROPERTY()
 	FNiagaraDataSetCompiledData DataSetCompiledData;
@@ -87,6 +120,24 @@ struct FNiagaraSystemCompiledData
 
 	UPROPERTY()
 	FNiagaraDataSetCompiledData UpdateInstanceParamsDataSetCompiledData;
+
+	UPROPERTY()
+	FNiagaraParameterDataSetBindingCollection SpawnInstanceGlobalBinding;
+	UPROPERTY()
+	FNiagaraParameterDataSetBindingCollection SpawnInstanceSystemBinding;
+	UPROPERTY()
+	FNiagaraParameterDataSetBindingCollection SpawnInstanceOwnerBinding;
+	UPROPERTY()
+	TArray<FNiagaraParameterDataSetBindingCollection> SpawnInstanceEmitterBindings;
+
+	UPROPERTY()
+	FNiagaraParameterDataSetBindingCollection UpdateInstanceGlobalBinding;
+	UPROPERTY()
+	FNiagaraParameterDataSetBindingCollection UpdateInstanceSystemBinding;
+	UPROPERTY()
+	FNiagaraParameterDataSetBindingCollection UpdateInstanceOwnerBinding;
+	UPROPERTY()
+	TArray<FNiagaraParameterDataSetBindingCollection> UpdateInstanceEmitterBindings;
 };
 
 USTRUCT()
diff --git a/Engine/Plugins/FX/Niagara/Source/Niagara/Private/NiagaraEmitterInstance.cpp b/Engine/Plugins/FX/Niagara/Source/Niagara/Private/NiagaraEmitterInstance.cpp
--- a/Engine/Plugins/FX/Niagara/Source/Niagara/Private/NiagaraEmitterInstance.cpp
+++ b/Engine/Plugins/FX/Niagara/Source/Niagara/Private/NiagaraEmitterInstance.cpp
@@ -370,38 +370,15 @@ void FNiagaraEmitterInstance::Init(int32 InEmitterIdx, FNiagaraSystemInstanceID
 	}
 
 
+	//Setup direct bindings for setting parameter values.
 	{
 		SCOPE_CYCLE_COUNTER(STAT_NiagaraEmitterInit_DirectBindings);
+
 		//Setup direct bindings for setting parameter values.
 		SpawnIntervalBinding.Init(SpawnExecContext.Parameters, CachedEmitterCompiledData->EmitterSpawnIntervalVar);
 		InterpSpawnStartBinding.Init(SpawnExecContext.Parameters, CachedEmitterCompiledData->EmitterInterpSpawnStartDTVar);
 		SpawnGroupBinding.Init(SpawnExecContext.Parameters, CachedEmitterCompiledData->EmitterSpawnGroupVar);
 
-		SpawnEmitterAgeBinding.Init(SpawnExecContext.Parameters, CachedEmitterCompiledData->EmitterAgeVar);
-		UpdateEmitterAgeBinding.Init(UpdateExecContext.Parameters, CachedEmitterCompiledData->EmitterAgeVar);
-		EventEmitterAgeBindings.SetNum(NumEvents);
-		for (int32 i = 0; i < NumEvents; i++)
-		{
-			EventEmitterAgeBindings[i].Init(EventExecContexts[i].Parameters, CachedEmitterCompiledData->EmitterAgeVar);
-		}
-
-		if (CachedEmitter->SimTarget == ENiagaraSimTarget::GPUComputeSim && GPUExecContext != nullptr)
-		{
-			EmitterAgeBindingGPU.Init(GPUExecContext->CombinedParamStore, CachedEmitterCompiledData->EmitterAgeVar);
-		}
-
-		SpawnRandomSeedBinding.Init(SpawnExecContext.Parameters, CachedEmitterCompiledData->EmitterRandomSeedVar);
-		UpdateRandomSeedBinding.Init(UpdateExecContext.Parameters, CachedEmitterCompiledData->EmitterRandomSeedVar);
-		EventRandomSeedBindings.SetNum(NumEvents);
-		for (int32 i = 0; i < NumEvents; i++)
-		{
-			EventRandomSeedBindings[i].Init(EventExecContexts[i].Parameters, CachedEmitterCompiledData->EmitterRandomSeedVar);
-		}
-		if (CachedEmitter->SimTarget == ENiagaraSimTarget::GPUComputeSim && GPUExecContext != nullptr)
-		{
-			GPURandomSeedBinding.Init(GPUExecContext->CombinedParamStore, CachedEmitterCompiledData->EmitterRandomSeedVar);
-		}
-
 		// Initialize the exec count
 		SpawnExecCountBinding.Init(SpawnExecContext.Parameters, SYS_PARAM_ENGINE_EXEC_COUNT);
 		UpdateExecCountBinding.Init(UpdateExecContext.Parameters, SYS_PARAM_ENGINE_EXEC_COUNT);
@@ -410,14 +387,6 @@ void FNiagaraEmitterInstance::Init(int32 InEmitterIdx, FNiagaraSystemInstanceID
 		{
 			EventExecCountBindings[i].Init(EventExecContexts[i].Parameters, SYS_PARAM_ENGINE_EXEC_COUNT);
 		}
-
-		SpawnTotalSpawnedParticlesBinding.Init(SpawnExecContext.Parameters, CachedEmitterCompiledData->EmitterTotalSpawnedParticlesVar);
-		UpdateTotalSpawnedParticlesBinding.Init(UpdateExecContext.Parameters, CachedEmitterCompiledData->EmitterTotalSpawnedParticlesVar);
-		EventTotalSpawnedParticlesBindings.SetNum(NumEvents);
-		for (int32 i = 0; i < NumEvents; i++)
-		{
-			EventTotalSpawnedParticlesBindings[i].Init(EventExecContexts[i].Parameters, CachedEmitterCompiledData->EmitterTotalSpawnedParticlesVar);
-		}
 	}
 
 	{
@@ -474,7 +443,7 @@ void FNiagaraEmitterInstance::Init(int32 InEmitterIdx, FNiagaraSystemInstanceID
 
 void FNiagaraEmitterInstance::ResetSimulation(bool bKillExisting /*= true*/)
 {
-	Age = 0;
+	EmitterAge = 0;
 	TickCount = 0;
 	CachedBounds.Init();
 
@@ -711,7 +680,7 @@ void FNiagaraEmitterInstance::BindParameters(bool bExternalOnly)
 	if (!bExternalOnly)
 	{
 		//Now bind parameters from the component and system.
-		FNiagaraParameterStore& InstanceParams = ParentSystemInstance->GetParameters();		
+		FNiagaraParameterStore& InstanceParams = ParentSystemInstance->GetInstanceParameters();		
 	
 		InstanceParams.Bind(&SpawnExecContext.Parameters);
 		InstanceParams.Bind(&UpdateExecContext.Parameters);
@@ -896,7 +865,7 @@ void FNiagaraEmitterInstance::PostTick()
 			}
 			else
 			{
-				CachedBounds = DynamicBounds.TransformBy(ParentSystemInstance->GetComponent()->GetComponentToWorld().Inverse());
+				CachedBounds = DynamicBounds.TransformBy(ParentSystemInstance->GetOwnerParameters().EngineWorldToLocal);
 			}
 		}
 		else
@@ -1061,7 +1030,6 @@ bool FNiagaraEmitterInstance::WaitForDebugInfo()
 	return false;
 }
 
-
 void FNiagaraEmitterInstance::SetSystemFixedBoundsOverride(FBox SystemFixedBounds)
 {
 	CachedSystemFixedBounds = SystemFixedBounds;
@@ -1083,9 +1051,9 @@ void FNiagaraEmitterInstance::Tick(float DeltaSeconds)
 
 	checkSlow(ParticleDataSet);
 	FNiagaraDataSet& Data = *ParticleDataSet;
-	Age += DeltaSeconds;
+	EmitterAge += DeltaSeconds;
 
-	//UE_LOG(LogNiagara, Warning, TEXT("Emitter Tick %f"), Age);
+	//UE_LOG(LogNiagara, Warning, TEXT("Emitter Tick %f"), EmitterAge);
 
 	if (ExecutionState == ENiagaraExecutionState::InactiveClear)
 	{
@@ -1118,35 +1086,11 @@ void FNiagaraEmitterInstance::Tick(float DeltaSeconds)
 	// add system constants
 	{
 		SCOPE_CYCLE_COUNTER(STAT_NiagaraConstants);
-		float InvDT = 1.0f / DeltaSeconds;
-
-		//TODO: Create a binding helper object for these to avoid the search.
-		SpawnEmitterAgeBinding.SetValue(Age);
-		UpdateEmitterAgeBinding.SetValue(Age);
-		for (FNiagaraParameterDirectBinding<float>& Binding : EventEmitterAgeBindings)
-		{
-			Binding.SetValue(Age);
-		}
-		
-		SpawnRandomSeedBinding.SetValue(CachedEmitter->RandomSeed);
-		UpdateRandomSeedBinding.SetValue(CachedEmitter->RandomSeed);
-		for (FNiagaraParameterDirectBinding<int32>& Binding : EventRandomSeedBindings)
-		{
-			Binding.SetValue(CachedEmitter->RandomSeed);
-		}
 
-		if (CachedEmitter->SimTarget == ENiagaraSimTarget::GPUComputeSim && GPUExecContext != nullptr)
-		{
-			EmitterAgeBindingGPU.SetValue(Age);
-			GPURandomSeedBinding.SetValue(CachedEmitter->RandomSeed);
-		}
-
-		SpawnTotalSpawnedParticlesBinding.SetValue(TotalSpawnedParticles);
-		UpdateTotalSpawnedParticlesBinding.SetValue(TotalSpawnedParticles);
-		for (FNiagaraParameterDirectBinding<int32>& Binding : EventTotalSpawnedParticlesBindings)
-		{
-			Binding.SetValue(TotalSpawnedParticles);
-		}
+		auto& EmitterParameters = ParentSystemInstance->EditEmitterParameters(EmitterIdx);
+		EmitterParameters.EmitterTotalSpawnedParticles = TotalSpawnedParticles;
+		EmitterParameters.EmitterAge = EmitterAge;
+		EmitterParameters.EmitterRandomSeed = CachedEmitter->RandomSeed;
 	}
 	
 	// Calculate number of new particles from regular spawning 
@@ -1279,8 +1223,8 @@ void FNiagaraEmitterInstance::Tick(float DeltaSeconds)
 		// Because each context is only ran once each frame, the CBuffer layout stays constant for the lifetime duration of the CBuffer (one frame).
 
 		// @todo-threadsafety do this once during init. Should not change during runtime...
-		GPUExecContext->CBufferLayout.ConstantBufferSize = ParmSize;
-		GPUExecContext->CBufferLayout.ComputeHash();
+		GPUExecContext->ExternalCBufferLayout.ConstantBufferSize = ParmSize / (GPUExecContext->HasInterpolationParameters ? 2 : 1);
+		GPUExecContext->ExternalCBufferLayout.ComputeHash();
 
 		// Need to call post-tick, which calls the copy to previous for interpolated spawning
 		SpawnExecContext.PostTick();
@@ -1312,7 +1256,6 @@ void FNiagaraEmitterInstance::Tick(float DeltaSeconds)
 		EventSpawnTotal = 0;
 	}
 
-
 	Data.BeginSimulate();
 	Data.Allocate(AllocationSize);
 
@@ -1359,7 +1302,11 @@ void FNiagaraEmitterInstance::Tick(float DeltaSeconds)
 			EventDataSet->GetDestinationDataChecked().SetNumInstances(OrigNumParticles);
 			UpdateExecContext.BindData(EventDataSetIdx++, *EventDataSet, 0, true);
 		}
-		UpdateExecContext.Execute(OrigNumParticles);
+
+		FScriptExecutionConstantBufferTable UpdateConstantBufferTable;
+		BuildConstantBufferTable(UpdateExecContext, UpdateConstantBufferTable);
+
+		UpdateExecContext.Execute(OrigNumParticles, UpdateConstantBufferTable);
 		int32 DeltaParticles = Data.GetDestinationDataChecked().GetNumInstances() - OrigNumParticles;
 
 		ensure(DeltaParticles <= 0); // We either lose particles or stay the same, we should never add particles in update!
@@ -1389,6 +1336,11 @@ void FNiagaraEmitterInstance::Tick(float DeltaSeconds)
 	{
 		SCOPE_CYCLE_COUNTER(STAT_NiagaraSpawn);
 
+		// note that this constant buffer table is used for each invocation of the spawn, the data within the
+		// table will get modified between invocations (TotalSpawnedParticles).
+		FScriptExecutionConstantBufferTable SpawnConstantBufferTable;
+		BuildConstantBufferTable(SpawnExecContext, SpawnConstantBufferTable);
+
 		//Handle main spawn rate spawning
 		auto SpawnParticles = [&](int32 Num, const TCHAR* DumpLabel)
 		{
@@ -1396,7 +1348,7 @@ void FNiagaraEmitterInstance::Tick(float DeltaSeconds)
 			Data.GetDestinationDataChecked().SetNumInstances(OrigNum + Num);
 
 			// We need to update Engine.Emitter.TotalSpawnedParticles for each event spawn invocation.
-			SpawnTotalSpawnedParticlesBinding.SetValue(TotalSpawnedParticles); 
+			ParentSystemInstance->EditEmitterParameters(EmitterIdx).EmitterTotalSpawnedParticles = TotalSpawnedParticles;
 			
 			// NOTE(mv): Updates the count after setting the variable, such that the TotalSpawnedParticles value read 
 			//           in the script has the count at the start of the frame. 
@@ -1417,7 +1369,7 @@ void FNiagaraEmitterInstance::Tick(float DeltaSeconds)
 				SpawnExecContext.BindData(EventDataSetIdx++, *EventDataSet, EventOrigNum, true);
 			}
 
-			SpawnExecContext.Execute(Num);
+			SpawnExecContext.Execute(Num, SpawnConstantBufferTable);
 
 			if (GbDumpParticleData || System->bDumpDebugEmitterInfo)
 			{
@@ -1440,10 +1392,10 @@ void FNiagaraEmitterInstance::Tick(float DeltaSeconds)
 		{
 			if ( Info.Count > 0 )
 			{
+				auto& EmitterParameters = ParentSystemInstance->EditEmitterParameters(EmitterIdx);
 				SpawnIntervalBinding.SetValue(Info.IntervalDt);
 				InterpSpawnStartBinding.SetValue(Info.InterpStartDt);
 				SpawnGroupBinding.SetValue(Info.SpawnGroup);
-
 				SpawnParticles(Info.Count, TEXT("Regular Spawn"));
 			}
 		}
@@ -1462,10 +1414,10 @@ void FNiagaraEmitterInstance::Tick(float DeltaSeconds)
 					const int32 CurrNumParticles = Data.GetDestinationDataChecked().GetNumInstances();
 
 					//Event spawns are instantaneous at the middle of the frame?
+					auto& EmitterParameters = ParentSystemInstance->EditEmitterParameters(EmitterIdx);
 					SpawnIntervalBinding.SetValue(0.0f);
 					InterpSpawnStartBinding.SetValue(DeltaSeconds * 0.5f);
 					SpawnGroupBinding.SetValue(0);
-
 					SpawnParticles(EventNumToSpawn, TEXT("Event Spawn"));
 
 					//Update EventSpawnCounts to the number actually spawned.
@@ -1549,6 +1501,9 @@ void FNiagaraEmitterInstance::Tick(float DeltaSeconds)
 		const FNiagaraEventScriptProperties &EventHandlerProps = CachedEmitter->GetEventHandlers()[EventScriptIdx];
 		FNiagaraEventHandlingInfo& Info = EventHandlingInfo[EventScriptIdx];
 
+		FScriptExecutionConstantBufferTable EventConstantBufferTable;
+		BuildConstantBufferTable(EventExecContexts[EventScriptIdx], EventConstantBufferTable);
+
 		if (Info.EventData && Info.SpawnCounts.Num() > 0)
 		{
 			SCOPE_CYCLE_COUNTER(STAT_NiagaraEventHandle);
@@ -1562,7 +1517,8 @@ void FNiagaraEmitterInstance::Tick(float DeltaSeconds)
 
 					EventExecContexts[EventScriptIdx].BindData(0, Data, EventSpawnStart, false);
 					EventExecContexts[EventScriptIdx].BindData(1, EventHandlingInfo[EventScriptIdx].EventData, i, false);
-					EventExecContexts[EventScriptIdx].Execute(EventNumToSpawn);
+
+					EventExecContexts[EventScriptIdx].Execute(EventNumToSpawn, EventConstantBufferTable);
 
 					if (GbDumpParticleData || System->bDumpDebugEmitterInfo)
 					{
@@ -1602,6 +1558,10 @@ void FNiagaraEmitterInstance::Tick(float DeltaSeconds)
 	{
 		const FNiagaraEventScriptProperties &EventHandlerProps = CachedEmitter->GetEventHandlers()[EventScriptIdx];
 		FNiagaraDataBuffer* EventData = EventHandlingInfo[EventScriptIdx].EventData;
+
+		FScriptExecutionConstantBufferTable EventConstantBufferTable;
+		BuildConstantBufferTable(EventExecContexts[EventScriptIdx], EventConstantBufferTable);
+
 		// handle all-particle events
 		if (EventHandlerProps.Script && EventHandlerProps.ExecutionMode == EScriptExecutionMode::EveryParticle && EventData)
 		{
@@ -1621,7 +1581,8 @@ void FNiagaraEmitterInstance::Tick(float DeltaSeconds)
 
 					EventExecContexts[EventScriptIdx].BindData(0, Data, 0, true);
 					EventExecContexts[EventScriptIdx].BindData(1, EventData, i, false);
-					EventExecContexts[EventScriptIdx].Execute(NumInstancesPrev);
+
+					EventExecContexts[EventScriptIdx].Execute(NumInstancesPrev, EventConstantBufferTable);
 
 					Data.EndSimulate();
 
@@ -1841,3 +1802,35 @@ void FNiagaraEmitterInstance::TickFastPathAttributeBindings()
 		FloatBinding.Tick();
 	}
 }
+
+void FNiagaraEmitterInstance::BuildConstantBufferTable(
+	const FNiagaraScriptExecutionContext& ExecContext,
+	FScriptExecutionConstantBufferTable& ConstantBufferTable) const
+{
+	const auto& ExternalParameterData = ExecContext.Parameters.GetParameterDataArray();
+	uint8* ExternalParameterBuffer = const_cast<uint8*>(ExternalParameterData.GetData());
+
+	const uint32 ExternalParameterSize = ExecContext.Parameters.GetExternalParameterSize();
+	const uint32 LiteralConstantOffset = (ExecContext.HasInterpolationParameters ? 2 : 1) * ExternalParameterSize;
+	const uint32 LiteralConstantSize = ExternalParameterData.Num() - LiteralConstantOffset;
+
+	const uint32 TableCount = 5 * (ExecContext.HasInterpolationParameters ? 2 : 1) + 1;
+	ConstantBufferTable.Reset(TableCount);
+
+	ConstantBufferTable.AddTypedBuffer(ParentSystemInstance->GetGlobalParameters());
+	ConstantBufferTable.AddTypedBuffer(ParentSystemInstance->GetSystemParameters());
+	ConstantBufferTable.AddTypedBuffer(ParentSystemInstance->GetOwnerParameters());
+	ConstantBufferTable.AddTypedBuffer(ParentSystemInstance->GetEmitterParameters(EmitterIdx));
+	ConstantBufferTable.AddRawBuffer(ExternalParameterBuffer, ExternalParameterSize);
+
+	if (ExecContext.HasInterpolationParameters)
+	{
+		ConstantBufferTable.AddTypedBuffer(ParentSystemInstance->GetGlobalParameters(true));
+		ConstantBufferTable.AddTypedBuffer(ParentSystemInstance->GetSystemParameters(true));
+		ConstantBufferTable.AddTypedBuffer(ParentSystemInstance->GetOwnerParameters(true));
+		ConstantBufferTable.AddTypedBuffer(ParentSystemInstance->GetEmitterParameters(EmitterIdx, true));
+		ConstantBufferTable.AddRawBuffer(ExternalParameterBuffer + ExternalParameterSize, ExternalParameterSize);
+	}
+
+	ConstantBufferTable.AddRawBuffer(ExternalParameterBuffer + LiteralConstantOffset, LiteralConstantSize);
+}
diff --git a/Engine/Plugins/FX/Niagara/Source/Niagara/Private/NiagaraEmitterInstanceBatcher.cpp b/Engine/Plugins/FX/Niagara/Source/Niagara/Private/NiagaraEmitterInstanceBatcher.cpp
--- a/Engine/Plugins/FX/Niagara/Source/Niagara/Private/NiagaraEmitterInstanceBatcher.cpp
+++ b/Engine/Plugins/FX/Niagara/Source/Niagara/Private/NiagaraEmitterInstanceBatcher.cpp
@@ -1,4 +1,4 @@
-// Copyright Epic Games, Inc. All Rights Reserved.
+// Copyright 1998-2019 Epic Games, Inc. All Rights Reserved.
 
 #include "NiagaraEmitterInstanceBatcher.h"
 #include "NiagaraScriptExecutionContext.h"
@@ -518,7 +518,7 @@ void NiagaraEmitterInstanceBatcher::DispatchAllOnCompute(FOverlappableTicks& Ove
 	}
 #endif // WITH_EDITORONLY_DATA
 
-	// Transition the input buffers to readable and the outputs to writable on the compute pipe.
+	//
 	RHICmdList.TransitionResources(EResourceTransitionAccess::EReadable, EResourceTransitionPipeline::EComputeToCompute, ReadBuffers.GetData(), ReadBuffers.Num());
 	RHICmdList.TransitionResources(EResourceTransitionAccess::ERWBarrier, EResourceTransitionPipeline::EComputeToCompute, WriteBuffers.GetData(), WriteBuffers.Num());
 
@@ -553,10 +553,10 @@ void NiagaraEmitterInstanceBatcher::DispatchAllOnCompute(FOverlappableTicks& Ove
 		}
 	}
 
-	// The output buffers will be used by rendering, so transition them to readable on the graphics pipe.
+	//
 	RHICmdList.TransitionResources(EResourceTransitionAccess::EReadable, EResourceTransitionPipeline::EComputeToGfx, WriteBuffers.GetData(), WriteBuffers.Num());
-
-	// We have done all our overlapping compute work on this list so go back to default behavior and flush.
+	//	Now Copy to staging buffer the data we want to read back (alive particle count). And make buffer ready for that and draw commands on the graphics pipe too.
+	//
 	RHICmdList.AutomaticCacheFlushAfterComputeShader(true);
 
 	if (GNiagaraSubmitCommands)
@@ -1202,6 +1202,18 @@ void NiagaraEmitterInstanceBatcher::UnsetDataInterfaceParameters(const TArray<FN
 	}
 }
 
+static void SetConstantBuffer(FRHICommandList &RHICmdList, FRHIComputeShader* ComputeShader, const FShaderUniformBufferParameter& BufferParam, const FRHIUniformBufferLayout& Layout, const uint8* ParamData)
+{
+	if (!BufferParam.IsBound())
+		return;
+
+	if (Layout.ConstantBufferSize)
+	{
+		check(Layout.Resources.Num() == 0);
+		FUniformBufferRHIRef CBuffer = RHICreateUniformBuffer(ParamData, Layout, EUniformBufferUsage::UniformBuffer_SingleDraw);
+		RHICmdList.SetShaderUniformBuffer(ComputeShader, BufferParam.GetBaseIndex(), CBuffer);
+	}
+}
 
 /* Kick off a simulation/spawn run
  */
@@ -1223,7 +1235,6 @@ void NiagaraEmitterInstanceBatcher::Run(const FNiagaraGPUSystemTick& Tick, const
 	//UE_LOG(LogNiagara, Warning, TEXT("Run"));
 	
 	const TArray<FNiagaraDataInterfaceProxy*>& DataInterfaceProxies = Instance->DataInterfaceProxies;
-	const FRHIUniformBufferLayout& CBufferLayout = Context->CBufferLayout;
 	check(Instance->CurrentData && Instance->DestinationData);
 	FNiagaraDataBuffer& DestinationData = *Instance->DestinationData;
 	FNiagaraDataBuffer& CurrentData = *Instance->CurrentData;
@@ -1239,20 +1250,17 @@ void NiagaraEmitterInstanceBatcher::Run(const FNiagaraGPUSystemTick& Tick, const
 		ClearUAV(RHICmdList, DestinationData.GetGPUIDToIndexTable(), INDEX_NONE);
 	}
 
+	FRHIComputeShader* ComputeShader = Shader->GetComputeShader();
+
 	RHICmdList.SetComputeShader(Shader->GetComputeShader());
 
 	// #todo(dmp): clean up this logic for shader stages on first frame
-	if (Shader->SimStartParam.IsBound())
-	{
-		int v = Tick.bNeedsReset ? 1 : 0;
-		RHICmdList.SetShaderParameter(Shader->GetComputeShader(), Shader->SimStartParam.GetBufferIndex(), Shader->SimStartParam.GetBaseIndex(), Shader->SimStartParam.GetNumBytes(), &v);
-	}
-	
+	SetShaderValue(RHICmdList, ComputeShader, Shader->SimStartParam, Tick.bNeedsReset ? 1U : 0U);
 
 	// set the view uniform buffer param
 	if (Shader->ViewUniformBufferParam.IsBound() && ViewUniformBuffer)
 	{
-		RHICmdList.SetShaderUniformBuffer(Shader->GetComputeShader(), Shader->ViewUniformBufferParam.GetBaseIndex(), ViewUniformBuffer);
+		RHICmdList.SetShaderUniformBuffer(ComputeShader, Shader->ViewUniformBufferParam.GetBaseIndex(), ViewUniformBuffer);
 	}
 
 	SetDataInterfaceParameters(DataInterfaceProxies, Shader, RHICmdList, Instance, Tick, ShaderStageIndex);
@@ -1268,28 +1276,25 @@ void NiagaraEmitterInstanceBatcher::Run(const FNiagaraGPUSystemTick& Tick, const
 	if (Shader->InstanceCountsParam.IsBound() && !IterationInterface )
 	{
 		RHICmdList.TransitionResource(EResourceTransitionAccess::ERWNoBarrier, EResourceTransitionPipeline::EComputeToCompute, GPUInstanceCounterManager.GetInstanceCountBuffer().UAV);
-		RHICmdList.SetUAVParameter(Shader->GetComputeShader(), Shader->InstanceCountsParam.GetUAVIndex(), GPUInstanceCounterManager.GetInstanceCountBuffer().UAV);
+		Shader->InstanceCountsParam.SetBuffer(RHICmdList, ComputeShader, GPUInstanceCounterManager.GetInstanceCountBuffer());
 		const uint32 ReadOffset = Tick.bNeedsReset ? INDEX_NONE : CurrentData.GetGPUInstanceCountBufferOffset();
 		const uint32 WriteOffset = DestinationData.GetGPUInstanceCountBufferOffset();
-		RHICmdList.SetShaderParameter(Shader->GetComputeShader(), Shader->ReadInstanceCountOffsetParam.GetBufferIndex(), Shader->ReadInstanceCountOffsetParam.GetBaseIndex(), Shader->ReadInstanceCountOffsetParam.GetNumBytes(), &ReadOffset);
-		RHICmdList.SetShaderParameter(Shader->GetComputeShader(), Shader->WriteInstanceCountOffsetParam.GetBufferIndex(), Shader->WriteInstanceCountOffsetParam.GetBaseIndex(), Shader->WriteInstanceCountOffsetParam.GetNumBytes(), &WriteOffset);
+		SetShaderValue(RHICmdList, ComputeShader, Shader->ReadInstanceCountOffsetParam, ReadOffset);
+		SetShaderValue(RHICmdList, ComputeShader, Shader->WriteInstanceCountOffsetParam, WriteOffset);
 	}
 
 	// set the execution parameters
 	//
-	if (Shader->EmitterTickCounterParam.IsBound())
-	{
-		RHICmdList.SetShaderParameter(Shader->GetComputeShader(), Shader->EmitterTickCounterParam.GetBufferIndex(), Shader->EmitterTickCounterParam.GetBaseIndex(), Shader->EmitterTickCounterParam.GetNumBytes(), &FNiagaraComputeExecutionContext::TickCounter);
-	}
+	SetShaderValue(RHICmdList, ComputeShader, Shader->EmitterTickCounterParam, FNiagaraComputeExecutionContext::TickCounter);
 
 	// set spawn info
 	//
 	static_assert((sizeof(SpawnInfo.SpawnInfoStartOffsets) % SHADER_PARAMETER_ARRAY_ELEMENT_ALIGNMENT) == 0, "sizeof SpawnInfoStartOffsets should be a multiple of SHADER_PARAMETER_ARRAY_ELEMENT_ALIGNMENT");
 	static_assert((sizeof(SpawnInfo.SpawnInfoParams) % SHADER_PARAMETER_ARRAY_ELEMENT_ALIGNMENT) == 0, "sizeof SpawnInfoParams should be a multiple of SHADER_PARAMETER_ARRAY_ELEMENT_ALIGNMENT");
-	SetShaderValueArray(RHICmdList, Shader->GetComputeShader(), Shader->EmitterSpawnInfoOffsetsParam, SpawnInfo.SpawnInfoStartOffsets, NIAGARA_MAX_GPU_SPAWN_INFOS_V4);
-	SetShaderValueArray(RHICmdList, Shader->GetComputeShader(), Shader->EmitterSpawnInfoParamsParam, SpawnInfo.SpawnInfoParams, NIAGARA_MAX_GPU_SPAWN_INFOS);
+	SetShaderValueArray(RHICmdList, ComputeShader, Shader->EmitterSpawnInfoOffsetsParam, SpawnInfo.SpawnInfoStartOffsets, NIAGARA_MAX_GPU_SPAWN_INFOS_V4);
+	SetShaderValueArray(RHICmdList, ComputeShader, Shader->EmitterSpawnInfoParamsParam, SpawnInfo.SpawnInfoParams, NIAGARA_MAX_GPU_SPAWN_INFOS);
 
-	RHICmdList.SetShaderParameter(Shader->GetComputeShader(), Shader->UpdateStartInstanceParam.GetBufferIndex(), Shader->UpdateStartInstanceParam.GetBaseIndex(), Shader->UpdateStartInstanceParam.GetNumBytes(), &UpdateStartInstance);					// 0, except for event handler runs
+	SetShaderValue(RHICmdList, ComputeShader, Shader->UpdateStartInstanceParam, UpdateStartInstance);					// 0, except for event handler runs
 	int32 InstancesToSpawnThisFrame = Instance->SpawnInfo.SpawnRateInstances + Instance->SpawnInfo.EventSpawnTotal;
 
 	// Only spawn particles on the first stage
@@ -1298,29 +1303,18 @@ void NiagaraEmitterInstanceBatcher::Run(const FNiagaraGPUSystemTick& Tick, const
 		InstancesToSpawnThisFrame = 0;
 	}
 
-	RHICmdList.SetShaderParameter(Shader->GetComputeShader(), Shader->NumSpawnedInstancesParam.GetBufferIndex(), Shader->NumSpawnedInstancesParam.GetBaseIndex(), Shader->NumSpawnedInstancesParam.GetNumBytes(), &InstancesToSpawnThisFrame);				// number of instances in the spawn run
-	
-	if (Shader->DefaultShaderStageIndexParam.IsBound())
-	{
-		RHICmdList.SetShaderParameter(Shader->GetComputeShader(), Shader->DefaultShaderStageIndexParam.GetBufferIndex(), Shader->DefaultShaderStageIndexParam.GetBaseIndex(), Shader->DefaultShaderStageIndexParam.GetNumBytes(), &DefaultShaderStageIndex);					// 0, except if several stages are defined
-	}
-	
-	if (Shader->ShaderStageIndexParam.IsBound())
-	{
-		RHICmdList.SetShaderParameter(Shader->GetComputeShader(), Shader->ShaderStageIndexParam.GetBufferIndex(), Shader->ShaderStageIndexParam.GetBaseIndex(), Shader->ShaderStageIndexParam.GetNumBytes(), &ShaderStageIndex);					// 0, except if several stages are defined
-	}
+	SetShaderValue(RHICmdList, ComputeShader, Shader->NumSpawnedInstancesParam, InstancesToSpawnThisFrame);				// number of instances in the spawn run
+	SetShaderValue(RHICmdList, ComputeShader, Shader->DefaultShaderStageIndexParam, DefaultShaderStageIndex);					// 0, except if several stages are defined
+	SetShaderValue(RHICmdList, ComputeShader, Shader->ShaderStageIndexParam, ShaderStageIndex);					// 0, except if several stages are defined
 	const int32 DefaultIterationCount = -1;
-	if (Shader->IterationInterfaceCount.IsBound())
-	{
-		RHICmdList.SetShaderParameter(Shader->GetComputeShader(), Shader->IterationInterfaceCount.GetBufferIndex(), Shader->IterationInterfaceCount.GetBaseIndex(), Shader->IterationInterfaceCount.GetNumBytes(), &DefaultIterationCount);					// 0, except if several stages are defined
-	}
+	SetShaderValue(RHICmdList, ComputeShader, Shader->IterationInterfaceCount, DefaultIterationCount);					// 0, except if several stages are defined
 
 	const uint32 ShaderThreadGroupSize = FNiagaraShader::GetGroupSize(ShaderPlatform);
 	if (IterationInterface)
 	{
 		if (TotalNumInstances > ShaderThreadGroupSize)
 		{
-			RHICmdList.SetShaderParameter(Shader->GetComputeShader(), Shader->IterationInterfaceCount.GetBufferIndex(), Shader->IterationInterfaceCount.GetBaseIndex(), Shader->IterationInterfaceCount.GetNumBytes(), &TotalNumInstances);					// 0, except if several stages are defined
+			SetShaderValue(RHICmdList, ComputeShader, Shader->IterationInterfaceCount, TotalNumInstances);					// 0, except if several stages are defined
 		}
 	}
 
@@ -1330,19 +1324,23 @@ void NiagaraEmitterInstanceBatcher::Run(const FNiagaraGPUSystemTick& Tick, const
 		NumThreadGroups = FMath::Min(NIAGARA_MAX_COMPUTE_THREADGROUPS, FMath::DivideAndRoundUp(TotalNumInstances, ShaderThreadGroupSize));
 	}
 
+	SetConstantBuffer(RHICmdList, ComputeShader, Shader->GlobalConstantBufferParam[0], Context->GlobalCBufferLayout, Instance->GlobalParamData);
+	SetConstantBuffer(RHICmdList, ComputeShader, Shader->SystemConstantBufferParam[0], Context->SystemCBufferLayout, Instance->SystemParamData);
+	SetConstantBuffer(RHICmdList, ComputeShader, Shader->OwnerConstantBufferParam[0], Context->OwnerCBufferLayout, Instance->OwnerParamData);
+	SetConstantBuffer(RHICmdList, ComputeShader, Shader->EmitterConstantBufferParam[0], Context->EmitterCBufferLayout, Instance->EmitterParamData);
+	SetConstantBuffer(RHICmdList, ComputeShader, Shader->ExternalConstantBufferParam[0], Context->ExternalCBufferLayout, Instance->ExternalParamData);
 	// setup script parameters
-	if (CBufferLayout.ConstantBufferSize)
+	if (Context->HasInterpolationParameters)
 	{
-		check(CBufferLayout.Resources.Num() == 0);
-		const uint8* ParamData = Instance->ParamData;
-		FUniformBufferRHIRef CBuffer = RHICreateUniformBuffer(ParamData, CBufferLayout, EUniformBufferUsage::UniformBuffer_SingleDraw);
-		RHICmdList.SetShaderUniformBuffer(Shader->GetComputeShader(), Shader->EmitterConstantBufferParam.GetBaseIndex(), CBuffer);
-	}
-	else
-	{
-		ensure(!Shader->EmitterConstantBufferParam.IsBound());
+		SetConstantBuffer(RHICmdList, ComputeShader, Shader->GlobalConstantBufferParam[1], Context->GlobalCBufferLayout, Instance->GlobalParamData + sizeof(FNiagaraGlobalParameters));
+		SetConstantBuffer(RHICmdList, ComputeShader, Shader->SystemConstantBufferParam[1], Context->SystemCBufferLayout, Instance->SystemParamData + sizeof(FNiagaraSystemParameters));
+		SetConstantBuffer(RHICmdList, ComputeShader, Shader->OwnerConstantBufferParam[1], Context->OwnerCBufferLayout, Instance->OwnerParamData + sizeof(FNiagaraOwnerParameters));
+		SetConstantBuffer(RHICmdList, ComputeShader, Shader->EmitterConstantBufferParam[1], Context->EmitterCBufferLayout, Instance->EmitterParamData + sizeof(FNiagaraEmitterParameters));
+		SetConstantBuffer(RHICmdList, ComputeShader, Shader->ExternalConstantBufferParam[1], Context->ExternalCBufferLayout, Instance->ExternalParamData + Context->ExternalCBufferLayout.ConstantBufferSize);
 	}
 
+	// setup script parameters
+
 	// #todo(dmp): temporary hack -- unbind UAVs if we have a valid iteration DI.  This way, when we are outputting with a different iteration count, we don't
 	// mess up particle state
 	if (IterationInterface)
@@ -1365,7 +1363,7 @@ void NiagaraEmitterInstanceBatcher::Run(const FNiagaraGPUSystemTick& Tick, const
 	// reset iteration count
 	if (IterationInterface)
 	{
-		RHICmdList.SetShaderParameter(Shader->GetComputeShader(), Shader->IterationInterfaceCount.GetBufferIndex(), Shader->IterationInterfaceCount.GetBaseIndex(), Shader->IterationInterfaceCount.GetNumBytes(), &DefaultIterationCount);					// 0, except if several stages are defined
+		SetShaderValue(RHICmdList, ComputeShader, Shader->IterationInterfaceCount, DefaultIterationCount);					// 0, except if several stages are defined
 	}
 
 #if WITH_EDITORONLY_DATA
@@ -1412,7 +1410,7 @@ void NiagaraEmitterInstanceBatcher::Run(const FNiagaraGPUSystemTick& Tick, const
 	UnsetDataInterfaceParameters(DataInterfaceProxies, Shader, RHICmdList, Instance, Tick);
 	CurrentData.UnsetShaderParams(Shader, RHICmdList);
 	DestinationData.UnsetShaderParams(Shader, RHICmdList);
-	Shader->InstanceCountsParam.UnsetUAV(RHICmdList, Shader->GetComputeShader());
+	Shader->InstanceCountsParam.UnsetUAV(RHICmdList, ComputeShader);
 
 	if (bNeedsPersistentIDs && ShaderStageIndex == Context->MaxUpdateIterations - 1)
 	{
diff --git a/Engine/Plugins/FX/Niagara/Source/Niagara/Private/NiagaraModule.cpp b/Engine/Plugins/FX/Niagara/Source/Niagara/Private/NiagaraModule.cpp
--- a/Engine/Plugins/FX/Niagara/Source/Niagara/Private/NiagaraModule.cpp
+++ b/Engine/Plugins/FX/Niagara/Source/Niagara/Private/NiagaraModule.cpp
@@ -49,10 +49,6 @@ static FAutoConsoleVariableRef CVarEnableVerboseNiagaraChangeIdLogging(
 	ECVF_Default
 );
 
-#if WITH_EDITORONLY_DATA
-FNiagaraParameterStore INiagaraModule::FixedSystemInstanceParameters = FNiagaraParameterStore();
-#endif
-
 /**
 Use Shader Stages CVar.
 Enable the custom dispatch for multiple shader stages 
@@ -307,10 +303,6 @@ void INiagaraModule::StartupModule()
 	{
 		return new NiagaraEmitterInstanceBatcher(InFeatureLevel, InShaderPlatform);
 	}));
-
-#if WITH_EDITORONLY_DATA
-	InitFixedSystemInstanceParameterStore();
-#endif
 }
 
 void INiagaraModule::ShutdownRenderingResources()
@@ -442,37 +434,6 @@ bool INiagaraModule::IsTargetPlatformIncludedInLevelRangeForCook(const ITargetPl
 	return true;
 }
 
-#if WITH_EDITORONLY_DATA
-void INiagaraModule::InitFixedSystemInstanceParameterStore()
-{
-	FixedSystemInstanceParameters.AddParameter(SYS_PARAM_ENGINE_POSITION, true, false);
-	FixedSystemInstanceParameters.AddParameter(SYS_PARAM_ENGINE_ROTATION, true, false);
-	FixedSystemInstanceParameters.AddParameter(SYS_PARAM_ENGINE_SCALE, true, false);
-	FixedSystemInstanceParameters.AddParameter(SYS_PARAM_ENGINE_VELOCITY, true, false);
-	FixedSystemInstanceParameters.AddParameter(SYS_PARAM_ENGINE_X_AXIS, true, false);
-	FixedSystemInstanceParameters.AddParameter(SYS_PARAM_ENGINE_Y_AXIS, true, false);
-	FixedSystemInstanceParameters.AddParameter(SYS_PARAM_ENGINE_Z_AXIS, true, false);
-	FixedSystemInstanceParameters.AddParameter(SYS_PARAM_ENGINE_LOCAL_TO_WORLD, true, false);
-	FixedSystemInstanceParameters.AddParameter(SYS_PARAM_ENGINE_WORLD_TO_LOCAL, true, false);
-	FixedSystemInstanceParameters.AddParameter(SYS_PARAM_ENGINE_LOCAL_TO_WORLD_TRANSPOSED, true, false);
-	FixedSystemInstanceParameters.AddParameter(SYS_PARAM_ENGINE_WORLD_TO_LOCAL_TRANSPOSED, true, false);
-	FixedSystemInstanceParameters.AddParameter(SYS_PARAM_ENGINE_LOCAL_TO_WORLD_NO_SCALE, true, false);
-	FixedSystemInstanceParameters.AddParameter(SYS_PARAM_ENGINE_WORLD_TO_LOCAL_NO_SCALE, true, false);
-	FixedSystemInstanceParameters.AddParameter(SYS_PARAM_ENGINE_DELTA_TIME, true, false);
-	FixedSystemInstanceParameters.AddParameter(SYS_PARAM_ENGINE_TIME, true, false);
-	FixedSystemInstanceParameters.AddParameter(SYS_PARAM_ENGINE_REAL_TIME, true, false);
-	FixedSystemInstanceParameters.AddParameter(SYS_PARAM_ENGINE_INV_DELTA_TIME, true, false);
-	FixedSystemInstanceParameters.AddParameter(SYS_PARAM_ENGINE_TIME_SINCE_RENDERED, true, false);
-	FixedSystemInstanceParameters.AddParameter(SYS_PARAM_ENGINE_EXECUTION_STATE, true, false);
-	FixedSystemInstanceParameters.AddParameter(SYS_PARAM_ENGINE_LOD_DISTANCE, true, false);
-	FixedSystemInstanceParameters.AddParameter(SYS_PARAM_ENGINE_LOD_DISTANCE_FRACTION, true, false);
-	FixedSystemInstanceParameters.AddParameter(SYS_PARAM_ENGINE_SYSTEM_NUM_EMITTERS, true, false);
-	FixedSystemInstanceParameters.AddParameter(SYS_PARAM_ENGINE_SYSTEM_NUM_EMITTERS_ALIVE, true, false);
-	FixedSystemInstanceParameters.AddParameter(SYS_PARAM_ENGINE_SYSTEM_AGE);
-	FixedSystemInstanceParameters.AddParameter(SYS_PARAM_ENGINE_SYSTEM_TICK_COUNT);
-}
-#endif
-
 void INiagaraModule::OnChangeDetailLevel(class IConsoleVariable* CVar)
 {
 	int32 NewDetailLevel = CVar->GetInt();
@@ -1037,6 +998,80 @@ void INiagaraModule::ProcessShaderCompilationQueue()
 	return OnProcessQueue.Execute();
 }
 
-#undef LOCTEXT_NAMESPACE
+#if WITH_EDITOR
+const TArray<FNiagaraVariable>& FNiagaraGlobalParameters::GetVariables()
+{
+	static const TArray<FNiagaraVariable> Variables =
+	{
+		SYS_PARAM_ENGINE_DELTA_TIME,
+		SYS_PARAM_ENGINE_INV_DELTA_TIME,
+		SYS_PARAM_ENGINE_TIME,
+		SYS_PARAM_ENGINE_REAL_TIME,
+	};
 
+	return Variables;
+}
+
+const TArray<FNiagaraVariable>& FNiagaraSystemParameters::GetVariables()
+{
+	static const TArray<FNiagaraVariable> Variables =
+	{
+		SYS_PARAM_ENGINE_TIME_SINCE_RENDERED,
+		SYS_PARAM_ENGINE_LOD_DISTANCE,
+		SYS_PARAM_ENGINE_LOD_DISTANCE_FRACTION,
+		SYS_PARAM_ENGINE_SYSTEM_AGE,
+		SYS_PARAM_ENGINE_EXECUTION_STATE,
+		SYS_PARAM_ENGINE_SYSTEM_TICK_COUNT,
+		SYS_PARAM_ENGINE_SYSTEM_NUM_EMITTERS,
+		SYS_PARAM_ENGINE_SYSTEM_NUM_EMITTERS_ALIVE,
+	};
+
+	return Variables;
+}
+
+const TArray<FNiagaraVariable>& FNiagaraOwnerParameters::GetVariables()
+{
+	static const TArray<FNiagaraVariable> Variables =
+	{
+		SYS_PARAM_ENGINE_LOCAL_TO_WORLD,
+		SYS_PARAM_ENGINE_WORLD_TO_LOCAL,
+		SYS_PARAM_ENGINE_LOCAL_TO_WORLD_TRANSPOSED,
+		SYS_PARAM_ENGINE_WORLD_TO_LOCAL_TRANSPOSED,
+		SYS_PARAM_ENGINE_LOCAL_TO_WORLD_NO_SCALE,
+		SYS_PARAM_ENGINE_WORLD_TO_LOCAL_NO_SCALE,
+		SYS_PARAM_ENGINE_ROTATION,
+		SYS_PARAM_ENGINE_POSITION,
+		SYS_PARAM_ENGINE_VELOCITY,
+		SYS_PARAM_ENGINE_X_AXIS,
+		SYS_PARAM_ENGINE_Y_AXIS,
+		SYS_PARAM_ENGINE_Z_AXIS,
+		SYS_PARAM_ENGINE_SCALE,
+	};
+
+	return Variables;
+}
+
+const TArray<FNiagaraVariable>& FNiagaraEmitterParameters::GetVariables()
+{
+	static const FName NAME_NiagaraStructPadding0 = "Engine.Emitter.PaddingInt32_0";
+	static const FName NAME_NiagaraStructPadding1 = "Engine.Emitter.PaddingInt32_1";
+	static const FName NAME_NiagaraStructPadding2 = "Engine.Emitter.PaddingInt32_2";
+
+	static const TArray<FNiagaraVariable> Variables =
+	{
+		SYS_PARAM_ENGINE_EMITTER_NUM_PARTICLES,
+		SYS_PARAM_ENGINE_EMITTER_TOTAL_SPAWNED_PARTICLES,
+		SYS_PARAM_ENGINE_EMITTER_SPAWN_COUNT_SCALE,
+		SYS_PARAM_EMITTER_AGE,
+
+		SYS_PARAM_EMITTER_RANDOM_SEED,
+		FNiagaraVariable(FNiagaraTypeDefinition::GetIntDef(), NAME_NiagaraStructPadding0),
+		FNiagaraVariable(FNiagaraTypeDefinition::GetIntDef(), NAME_NiagaraStructPadding1),
+		FNiagaraVariable(FNiagaraTypeDefinition::GetIntDef(), NAME_NiagaraStructPadding2),
+	};
+
+	return Variables;
+}
+#endif
 
+#undef LOCTEXT_NAMESPACE
\ No newline at end of file
diff --git a/Engine/Plugins/FX/Niagara/Source/Niagara/Private/NiagaraParameterStore.cpp b/Engine/Plugins/FX/Niagara/Source/Niagara/Private/NiagaraParameterStore.cpp
--- a/Engine/Plugins/FX/Niagara/Source/Niagara/Private/NiagaraParameterStore.cpp
+++ b/Engine/Plugins/FX/Niagara/Source/Niagara/Private/NiagaraParameterStore.cpp
@@ -123,9 +123,7 @@ void FNiagaraParameterStore::SetOwner(UObject* InOwner)
 
 FNiagaraParameterStore::FNiagaraParameterStore(const FNiagaraParameterStore& Other)
 {
-	DEC_MEMORY_STAT_BY(STAT_NiagaraParamStoreMemory, ParameterData.GetAllocatedSize());
 	*this = Other;
-	INC_MEMORY_STAT_BY(STAT_NiagaraParamStoreMemory, ParameterData.GetAllocatedSize());
 }
 
 FNiagaraParameterStore& FNiagaraParameterStore::operator=(const FNiagaraParameterStore& Other)
diff --git a/Engine/Plugins/FX/Niagara/Source/Niagara/Private/NiagaraScriptExecutionContext.cpp b/Engine/Plugins/FX/Niagara/Source/Niagara/Private/NiagaraScriptExecutionContext.cpp
--- a/Engine/Plugins/FX/Niagara/Source/Niagara/Private/NiagaraScriptExecutionContext.cpp
+++ b/Engine/Plugins/FX/Niagara/Source/Niagara/Private/NiagaraScriptExecutionContext.cpp
@@ -44,6 +44,8 @@ bool FNiagaraScriptExecutionContext::Init(UNiagaraScript* InScript, ENiagaraSimT
 
 	Parameters.InitFromOwningContext(Script, InTarget, true);
 
+	HasInterpolationParameters = Script && Script->GetComputedVMCompilationId().HasInterpolatedParameters();
+
 	return true;//TODO: Error cases?
 }
 
@@ -140,7 +142,7 @@ bool FNiagaraScriptExecutionContext::Tick(FNiagaraSystemInstance* ParentSystemIn
 void FNiagaraScriptExecutionContext::PostTick()
 {
 	//If we're for interpolated spawn, copy over the previous frame's parameters into the Prev parameters.
-	if (Script && Script->GetComputedVMCompilationId().HasInterpolatedParameters())
+	if (HasInterpolationParameters)
 	{
 		Parameters.CopyCurrToPrev();
 	}
@@ -173,7 +175,7 @@ void FNiagaraScriptExecutionContext::BindData(int32 Index, FNiagaraDataBuffer* I
 	DataSetMetaTable[Index].Init(Input->GetRegisterTable().GetData(), nullptr, StartInstance, nullptr, nullptr, &DataSet->GetNumFreeIDs(), &DataSet->GetMaxUsedID(), DataSet->GetIDAcquireTag());
 }
 
-bool FNiagaraScriptExecutionContext::Execute(uint32 NumInstances)
+bool FNiagaraScriptExecutionContext::Execute(uint32 NumInstances, const FScriptExecutionConstantBufferTable& ConstantBufferTable)
 {
 	if (NumInstances == 0)
 	{
@@ -190,7 +192,9 @@ bool FNiagaraScriptExecutionContext::Execute(uint32 NumInstances)
 			ExecData.ByteCode.GetData(),
 			ExecData.OptimizedByteCode.Num() > 0 ? ExecData.OptimizedByteCode.GetData() : nullptr,
 			ExecData.NumTempRegisters,
-			Parameters.GetParameterDataArray().GetData(),
+			ConstantBufferTable.Buffers.Num(),
+			ConstantBufferTable.Buffers.GetData(),
+			ConstantBufferTable.BufferSizes.GetData(),
 			DataSetMetaTable,
 			FunctionTable.GetData(),
 			DataInterfaceInstDataTable.GetData(),
@@ -295,7 +299,7 @@ void FNiagaraGPUSystemTick::Init(FNiagaraSystemInstance* InSystemInstance)
 	const uint32 PackedDispatchesSize = InSystemInstance->ActiveGPUEmitterCount * sizeof(FNiagaraComputeInstanceData);
 	// We want the Params after the instance data to be aligned so we can upload to the gpu.
 	uint32 PackedDispatchesSizeAligned = Align(PackedDispatchesSize, SHADER_PARAMETER_STRUCT_ALIGNMENT);
-	uint32 TotalParamSize = InSystemInstance->TotalParamSize;
+	uint32 TotalParamSize = InSystemInstance->TotalGPUParamSize;
 
 	uint32 TotalPackedBufferSize = PackedDispatchesSizeAligned + TotalParamSize;
 
@@ -308,6 +312,12 @@ void FNiagaraGPUSystemTick::Init(FNiagaraSystemInstance* InSystemInstance)
 	check(TickCount > 0);
 	bNeedsReset = ( TickCount == 1);
 
+	uint8* GlobalParamData = ParamDataBufferPtr;
+	ParamDataBufferPtr += 2 * sizeof(FNiagaraGlobalParameters);
+	uint8* SystemParamData = ParamDataBufferPtr;
+	ParamDataBufferPtr += 2 * sizeof(FNiagaraSystemParameters);
+	uint8* OwnerParamData = ParamDataBufferPtr;
+	ParamDataBufferPtr += 2 * sizeof(FNiagaraOwnerParameters);
 
 	// Now we will generate instance data for every GPU simulation we want to run on the render thread.
 	// This is spawn rate as well as DataInterface per instance data and the ParameterData for the emitter.
@@ -319,24 +329,39 @@ void FNiagaraGPUSystemTick::Init(FNiagaraSystemInstance* InSystemInstance)
 
 		if (Emitter && Emitter->GetCachedEmitter()->SimTarget == ENiagaraSimTarget::GPUComputeSim && Emitter->GetGPUContext() != nullptr && Emitter->GetExecutionState() != ENiagaraExecutionState::Complete)
 		{
+			FNiagaraComputeExecutionContext* GPUContext = Emitter->GetGPUContext();
+
 			FNiagaraComputeInstanceData* InstanceData = new (&Instances[InstanceIndex]) FNiagaraComputeInstanceData;
 			InstanceIndex++;
 
-			InstanceData->Context = Emitter->GetGPUContext();
-			check(InstanceData->Context->MainDataSet);
+			InstanceData->Context = GPUContext;
+			check(GPUContext->MainDataSet);
 
-			InstanceData->SpawnInfo = Emitter->GetGPUContext()->GpuSpawnInfo_GT;
+			InstanceData->SpawnInfo = GPUContext->GpuSpawnInfo_GT;
 
-			int32 ParmSize = Emitter->GetGPUContext()->CombinedParamStore.GetPaddedParameterSizeInBytes();
+			int32 ParmSize = GPUContext->CombinedParamStore.GetPaddedParameterSizeInBytes();
 
-			Emitter->GetGPUContext()->CombinedParamStore.CopyParameterDataToPaddedBuffer(ParamDataBufferPtr, ParmSize);
+			InstanceData->GlobalParamData = GlobalParamData;
+			InstanceData->SystemParamData = SystemParamData;
+			InstanceData->OwnerParamData = OwnerParamData;
+			InstanceData->EmitterParamData = ParamDataBufferPtr;
+			ParamDataBufferPtr += 2 * sizeof(FNiagaraEmitterParameters);
+			InstanceData->ExternalParamData = ParamDataBufferPtr;
+			ParamDataBufferPtr += ParmSize;
 
-			InstanceData->ParamData = ParamDataBufferPtr;
+			// actually copy all of the data over
+			FMemory::Memcpy(InstanceData->GlobalParamData, &InSystemInstance->GetGlobalParameters(), sizeof(FNiagaraGlobalParameters));
+			FMemory::Memcpy(InstanceData->GlobalParamData + sizeof(FNiagaraGlobalParameters), &InSystemInstance->GetGlobalParameters(true), sizeof(FNiagaraGlobalParameters));
+			FMemory::Memcpy(InstanceData->SystemParamData, &InSystemInstance->GetSystemParameters(), sizeof(FNiagaraSystemParameters));
+			FMemory::Memcpy(InstanceData->SystemParamData + sizeof(FNiagaraSystemParameters), &InSystemInstance->GetSystemParameters(true), sizeof(FNiagaraSystemParameters));
+			FMemory::Memcpy(InstanceData->OwnerParamData, &InSystemInstance->GetOwnerParameters(), sizeof(FNiagaraOwnerParameters));
+			FMemory::Memcpy(InstanceData->OwnerParamData + sizeof(FNiagaraOwnerParameters), &InSystemInstance->GetOwnerParameters(true), sizeof(FNiagaraOwnerParameters));
+			FMemory::Memcpy(InstanceData->EmitterParamData, &InSystemInstance->GetEmitterParameters(i), sizeof(FNiagaraEmitterParameters));
+			FMemory::Memcpy(InstanceData->EmitterParamData + sizeof(FNiagaraEmitterParameters), &InSystemInstance->GetEmitterParameters(i, true), sizeof(FNiagaraEmitterParameters));
 
-			ParamDataBufferPtr += ParmSize;
+			GPUContext->CombinedParamStore.CopyParameterDataToPaddedBuffer(InstanceData->ExternalParamData, ParmSize);
 
 			// @todo-threadsafety Think of a better way to do this!
-			FNiagaraComputeExecutionContext* GPUContext = Emitter->GetGPUContext();
 			const TArray<UNiagaraDataInterface*>& DataInterfaces = GPUContext->CombinedParamStore.GetDataInterfaces();
 			InstanceData->DataInterfaceProxies.Reserve(DataInterfaces.Num());
 			for (UNiagaraDataInterface* DI : DataInterfaces)
@@ -374,7 +399,11 @@ FNiagaraComputeExecutionContext::FNiagaraComputeExecutionContext()
 	: MainDataSet(nullptr)
 	, GPUScript(nullptr)
 	, GPUScript_RT(nullptr)
-	, CBufferLayout(TEXT("Niagara Compute Sim CBuffer"))
+	, GlobalCBufferLayout(TEXT("Niagara GPU Global CBuffer"))
+	, SystemCBufferLayout(TEXT("Niagara GPU System CBuffer"))
+	, OwnerCBufferLayout(TEXT("Niagara GPU Owner CBuffer"))
+	, EmitterCBufferLayout(TEXT("Niagara GPU Emitter CBuffer"))
+	, ExternalCBufferLayout(TEXT("Niagara GPU External CBuffer"))
 	, DataToRender(nullptr)
 #if WITH_EDITORONLY_DATA
 	, GPUDebugDataReadbackFloat(nullptr)
@@ -387,6 +416,17 @@ FNiagaraComputeExecutionContext::FNiagaraComputeExecutionContext()
 	, GPUDebugDataCountOffset(INDEX_NONE)
 #endif	  
 {
+	GlobalCBufferLayout.ConstantBufferSize = sizeof(FNiagaraGlobalParameters);
+	GlobalCBufferLayout.ComputeHash();
+
+	SystemCBufferLayout.ConstantBufferSize = sizeof(FNiagaraSystemParameters);
+	SystemCBufferLayout.ComputeHash();
+
+	OwnerCBufferLayout.ConstantBufferSize = sizeof(FNiagaraOwnerParameters);
+	OwnerCBufferLayout.ComputeHash();
+
+	EmitterCBufferLayout.ConstantBufferSize = sizeof(FNiagaraEmitterParameters);
+	EmitterCBufferLayout.ComputeHash();
 }
 
 FNiagaraComputeExecutionContext::~FNiagaraComputeExecutionContext()
@@ -441,6 +481,7 @@ void FNiagaraComputeExecutionContext::InitParams(UNiagaraScript* InGPUComputeScr
 
 	SpawnStages.Append(InSpawnStages);
 	
+	HasInterpolationParameters = GPUScript && GPUScript->GetComputedVMCompilationId().HasInterpolatedParameters();
 	
 #if DO_CHECK
 	FNiagaraShader *Shader = InGPUComputeScript->GetRenderThreadScript()->GetShaderGameThread();
@@ -497,7 +538,7 @@ bool FNiagaraComputeExecutionContext::Tick(FNiagaraSystemInstance* ParentSystemI
 void FNiagaraComputeExecutionContext::PostTick()
 {
 	//If we're for interpolated spawn, copy over the previous frame's parameters into the Prev parameters.
-	if (GPUScript && GPUScript->GetComputedVMCompilationId().HasInterpolatedParameters())
+	if (HasInterpolationParameters)
 	{
 		CombinedParamStore.CopyCurrToPrev();
 	}
diff --git a/Engine/Plugins/FX/Niagara/Source/Niagara/Private/NiagaraScriptExecutionParameterStore.cpp b/Engine/Plugins/FX/Niagara/Source/Niagara/Private/NiagaraScriptExecutionParameterStore.cpp
--- a/Engine/Plugins/FX/Niagara/Source/Niagara/Private/NiagaraScriptExecutionParameterStore.cpp
+++ b/Engine/Plugins/FX/Niagara/Source/Niagara/Private/NiagaraScriptExecutionParameterStore.cpp
@@ -147,6 +147,20 @@ void FNiagaraScriptExecutionParameterStore::AddPaddedParamSize(const FNiagaraTyp
 	}
 }
 
+void FNiagaraScriptExecutionParameterStore::AddAlignmentPadding()
+{
+	if (PaddingInfo.Num())
+	{
+		const auto& LastEntry = PaddingInfo.Last();
+		const uint32 CurrentOffset = LastEntry.DestOffset + LastEntry.DestSize;
+		const uint32 AlignedOffset = Align(CurrentOffset, SHADER_PARAMETER_STRUCT_ALIGNMENT);
+
+		if (CurrentOffset != AlignedOffset)
+		{
+			PaddingInfo.Emplace(GetParameterDataArray().Num(), CurrentOffset, 0, AlignedOffset - CurrentOffset);
+		}
+	}
+}
 
 void FNiagaraScriptExecutionParameterStore::InitFromOwningScript(UNiagaraScript* Script, ENiagaraSimTarget SimTarget, bool bNotifyAsDirty)
 {
@@ -224,11 +238,15 @@ void FNiagaraScriptExecutionParameterStore::AddScriptParams(UNiagaraScript* Scri
 	DebugName = FString::Printf(TEXT("ScriptExecParamStore %s %p"), *Script->GetFullName(), this);
 #endif
 	
+	ParameterSize = GetParameterDataArray().Num();
+
 	//Add previous frame values if we're interpolated spawn.
 	bool bIsInterpolatedSpawn = Script->GetVMExecutableDataCompilationId().HasInterpolatedParameters();
 
 	if (bIsInterpolatedSpawn)
 	{
+		AddAlignmentPadding();
+
 		for (FNiagaraVariable& Param : Script->GetVMExecutableData().Parameters.Parameters)
 		{
 			FNiagaraVariable PrevParam(Param.GetType(), FName(*(INTERPOLATED_PARAMETER_PREFIX + Param.GetName().ToString())));
@@ -236,7 +254,6 @@ void FNiagaraScriptExecutionParameterStore::AddScriptParams(UNiagaraScript* Scri
 		}
 	}
 	
-	ParameterSize = GetParameterDataArray().Num();
 	if (bIsInterpolatedSpawn)
 	{
 		CopyCurrToPrev();
@@ -281,8 +298,8 @@ void FNiagaraScriptExecutionParameterStore::AddScriptParams(UNiagaraScript* Scri
 
 void FNiagaraScriptExecutionParameterStore::CopyCurrToPrev()
 {
-	int32 ParamStart = ParameterSize / 2;
-	checkSlow(FMath::Frac((float)ParameterSize / 2) == 0.0f);
+	int32 ParamStart = ParameterSize;
+	checkSlow(FMath::Frac((float)ParameterSize) == 0.0f);
 
 	FMemory::Memcpy(GetParameterData_Internal(ParamStart), GetParameterData(0), ParamStart);
 }
diff --git a/Engine/Plugins/FX/Niagara/Source/Niagara/Private/NiagaraSystem.cpp b/Engine/Plugins/FX/Niagara/Source/Niagara/Private/NiagaraSystem.cpp
--- a/Engine/Plugins/FX/Niagara/Source/Niagara/Private/NiagaraSystem.cpp
+++ b/Engine/Plugins/FX/Niagara/Source/Niagara/Private/NiagaraSystem.cpp
@@ -8,6 +8,7 @@
 #include "NiagaraScriptSourceBase.h"
 #include "NiagaraCustomVersion.h"
 #include "NiagaraModule.h"
+#include "NiagaraTypes.h"
 #include "Modules/ModuleManager.h"
 #include "NiagaraEmitter.h"
 #include "UObject/Package.h"
@@ -580,7 +581,7 @@ bool UNiagaraSystem::IsReadyToRunInternal() const
 	/* Check that our post compile data is in sync with the current emitter handles count. If we have just added a new emitter handle, we will not have any outstanding compilation requests as the new compile
 	 * will not be added to the outstanding compilation requests until the next tick.
 	 */
-	if (EmitterHandles.Num() != EmitterCompiledData.Num() || EmitterHandles.Num() != SystemCompiledData.NumParticleVars.Num() || EmitterHandles.Num() != SystemCompiledData.TotalSpawnedParticlesVars.Num() || EmitterHandles.Num() != SystemCompiledData.SpawnCountScaleVars.Num())
+	if (EmitterHandles.Num() != EmitterCompiledData.Num())
 	{
 		return false;
 	}
@@ -1240,44 +1241,10 @@ void UNiagaraSystem::InitEmitterCompiledData()
 
 void UNiagaraSystem::InitSystemCompiledData()
 {
-	SystemCompiledData.NumParticleVars.Empty();
-	SystemCompiledData.TotalSpawnedParticlesVars.Empty();
-	SystemCompiledData.SpawnCountScaleVars.Empty();
 	SystemCompiledData.InstanceParamStore.Empty();
 
-	SystemCompiledData.InstanceParamStore = INiagaraModule::GetFixedSystemInstanceParameterStore();
 	ExposedParameters.CopyParametersTo(SystemCompiledData.InstanceParamStore, false, FNiagaraParameterStore::EDataInterfaceCopyMethod::Reference);
 
-	for (const FNiagaraEmitterHandle& PerEmitterHandle : EmitterHandles)
-	{
-		const UNiagaraEmitter* Emitter = PerEmitterHandle.GetInstance();
-		if (ensureMsgf(Emitter != nullptr, TEXT("Failed to get Emitter Instance from Emitter Handle when post compiling Niagara System!")))
-		{
-			const FString EmitterName = Emitter->GetUniqueEmitterName();
-			{
-				FNiagaraVariable Var = SYS_PARAM_ENGINE_EMITTER_NUM_PARTICLES;
-				const FString ParamName = Var.GetName().ToString().Replace(TEXT("Emitter"), *EmitterName);
-				Var.SetName(*ParamName);
-				SystemCompiledData.InstanceParamStore.AddParameter(Var, true, false);
-				SystemCompiledData.NumParticleVars.Add(Var);
-			}
-			{
-				FNiagaraVariable Var = SYS_PARAM_ENGINE_EMITTER_TOTAL_SPAWNED_PARTICLES;
-				const FString ParamName = Var.GetName().ToString().Replace(TEXT("Emitter"), *EmitterName);
-				Var.SetName(*ParamName);
-				SystemCompiledData.InstanceParamStore.AddParameter(Var, true, false);
-				SystemCompiledData.TotalSpawnedParticlesVars.Add(Var);
-			}
-			{
-				FNiagaraVariable Var = SYS_PARAM_ENGINE_EMITTER_SPAWN_COUNT_SCALE;
-				const FString ParamName = Var.GetName().ToString().Replace(TEXT("Emitter"), *EmitterName);
-				Var.SetName(*ParamName);
-				SystemCompiledData.InstanceParamStore.AddParameter(Var, true, false);
-				SystemCompiledData.SpawnCountScaleVars.Add(Var);
-			}
-		}
-	}
-
 	auto CreateDataSetCompiledData = [&](FNiagaraDataSetCompiledData& CompiledData, TArrayView<FNiagaraVariable> Vars)
 	{
 		CompiledData.Empty();
@@ -1301,6 +1268,37 @@ void UNiagaraSystem::InitSystemCompiledData()
 	CreateDataSetCompiledData(SystemCompiledData.SpawnInstanceParamsDataSetCompiledData, EngineParamsSpawn ? EngineParamsSpawn->Parameters : TArrayView<FNiagaraVariable>());
 	FNiagaraParameters* EngineParamsUpdate = GetSystemUpdateScript()->GetVMExecutableData().DataSetToParameters.Find(TEXT("Engine"));
 	CreateDataSetCompiledData(SystemCompiledData.UpdateInstanceParamsDataSetCompiledData, EngineParamsUpdate ? EngineParamsUpdate->Parameters : TArrayView<FNiagaraVariable>());
+
+	// create the bindings to be used with our constant buffers; geenrating the offsets to/from the data sets; we need
+	// editor data to build these bindings because of the constant buffer structs only having their variable definitions
+	// with editor data.
+	SystemCompiledData.SpawnInstanceGlobalBinding.Build<FNiagaraGlobalParameters>(SystemCompiledData.SpawnInstanceParamsDataSetCompiledData);
+	SystemCompiledData.SpawnInstanceSystemBinding.Build<FNiagaraSystemParameters>(SystemCompiledData.SpawnInstanceParamsDataSetCompiledData);
+	SystemCompiledData.SpawnInstanceOwnerBinding.Build<FNiagaraOwnerParameters>(SystemCompiledData.SpawnInstanceParamsDataSetCompiledData);
+
+	SystemCompiledData.UpdateInstanceGlobalBinding.Build<FNiagaraGlobalParameters>(SystemCompiledData.UpdateInstanceParamsDataSetCompiledData);
+	SystemCompiledData.UpdateInstanceSystemBinding.Build<FNiagaraSystemParameters>(SystemCompiledData.UpdateInstanceParamsDataSetCompiledData);
+	SystemCompiledData.UpdateInstanceOwnerBinding.Build<FNiagaraOwnerParameters>(SystemCompiledData.UpdateInstanceParamsDataSetCompiledData);
+
+	const int32 EmitterCount = EmitterHandles.Num();
+
+	SystemCompiledData.SpawnInstanceEmitterBindings.SetNum(EmitterHandles.Num());
+	SystemCompiledData.UpdateInstanceEmitterBindings.SetNum(EmitterHandles.Num());
+
+	const FString EmitterNamespace = TEXT("Emitter");
+	for (int32 EmitterIdx = 0; EmitterIdx < EmitterCount; ++EmitterIdx)
+	{
+		const FNiagaraEmitterHandle& PerEmitterHandle = EmitterHandles[EmitterIdx];
+		const UNiagaraEmitter* Emitter = PerEmitterHandle.GetInstance();
+		if (ensureMsgf(Emitter != nullptr, TEXT("Failed to get Emitter Instance from Emitter Handle when post compiling Niagara System!")))
+		{
+			const FString EmitterName = Emitter->GetUniqueEmitterName();
+
+			SystemCompiledData.SpawnInstanceEmitterBindings[EmitterIdx].Build<FNiagaraEmitterParameters>(SystemCompiledData.SpawnInstanceParamsDataSetCompiledData, EmitterNamespace, EmitterName);
+			SystemCompiledData.UpdateInstanceEmitterBindings[EmitterIdx].Build<FNiagaraEmitterParameters>(SystemCompiledData.UpdateInstanceParamsDataSetCompiledData, EmitterNamespace, EmitterName);
+		}
+	}
+
 }
 #endif
 
@@ -1465,3 +1463,48 @@ FNiagaraEmitterCompiledData::FNiagaraEmitterCompiledData()
 	EmitterRandomSeedVar = SYS_PARAM_EMITTER_RANDOM_SEED;
 	EmitterTotalSpawnedParticlesVar = SYS_PARAM_ENGINE_EMITTER_TOTAL_SPAWNED_PARTICLES;
 }
+
+#if WITH_EDITORONLY_DATA
+void FNiagaraParameterDataSetBindingCollection::BuildInternal(const TArray<FNiagaraVariable>& ParameterVars, const FNiagaraDataSetCompiledData& DataSet, const FString& NamespaceBase, const FString& NamespaceReplacement)
+{
+	const bool DoNameReplacement = !NamespaceBase.IsEmpty() && !NamespaceReplacement.IsEmpty();
+
+	int32 ParameterOffset = 0;
+	for (FNiagaraVariable Var : ParameterVars)
+	{
+		if (DoNameReplacement)
+		{
+			const FString ParamName = Var.GetName().ToString().Replace(*NamespaceBase, *NamespaceReplacement);
+			Var.SetName(*ParamName);
+		}
+
+		int32 VariableIndex = DataSet.Variables.IndexOfByKey(Var);
+
+		if (DataSet.VariableLayouts.IsValidIndex(VariableIndex))
+		{
+			const FNiagaraVariableLayoutInfo& Layout = DataSet.VariableLayouts[VariableIndex];
+			int32 NumFloats = 0;
+			int32 NumInts = 0;
+
+			for (uint32 CompIdx = 0; CompIdx < Layout.GetNumFloatComponents(); ++CompIdx)
+			{
+				int32 ParamOffset = ParameterOffset + Layout.LayoutInfo.FloatComponentByteOffsets[CompIdx];
+				int32 DataSetOffset = Layout.FloatComponentStart + NumFloats++;
+				auto& Binding = FloatOffsets.AddDefaulted_GetRef();
+				Binding.ParameterOffset = ParamOffset;
+				Binding.DataSetComponentOffset = DataSetOffset;
+			}
+			for (uint32 CompIdx = 0; CompIdx < Layout.GetNumInt32Components(); ++CompIdx)
+			{
+				int32 ParamOffset = ParameterOffset + Layout.LayoutInfo.Int32ComponentByteOffsets[CompIdx];
+				int32 DataSetOffset = Layout.Int32ComponentStart + NumInts++;
+				auto& Binding = Int32Offsets.AddDefaulted_GetRef();
+				Binding.ParameterOffset = ParamOffset;
+				Binding.DataSetComponentOffset = DataSetOffset;
+			}
+		}
+
+		ParameterOffset += Var.GetSizeInBytes();
+	}
+}
+#endif
\ No newline at end of file
diff --git a/Engine/Plugins/FX/Niagara/Source/Niagara/Private/NiagaraSystemInstance.cpp b/Engine/Plugins/FX/Niagara/Source/Niagara/Private/NiagaraSystemInstance.cpp
--- a/Engine/Plugins/FX/Niagara/Source/Niagara/Private/NiagaraSystemInstance.cpp
+++ b/Engine/Plugins/FX/Niagara/Source/Niagara/Private/NiagaraSystemInstance.cpp
@@ -71,6 +71,8 @@ FNiagaraSystemInstance::FNiagaraSystemInstance(UNiagaraComponent* InComponent)
 	, TickBehavior(InComponent ? InComponent->GetTickBehavior() : ENiagaraTickBehavior::UsePrereqs)
 	, Age(0.0f)
 	, TickCount(0)
+	, CurrentFrameIndex(1)
+	, ParametersValid(false)
 	, bSolo(false)
 	, bForceSolo(false)
 	, bPendingSpawn(false)
@@ -662,6 +664,7 @@ void FNiagaraSystemInstance::Reset(FNiagaraSystemInstance::EResetMode Mode)
 	
 	if (bBindParams)
 	{
+		ResetParameters();
 		BindParameters();
 	}
 
@@ -898,61 +901,6 @@ void FNiagaraSystemInstance::ReInitInternal()
 	
 	const FNiagaraSystemCompiledData& SystemCompiledData = System->GetSystemCompiledData();
 	InstanceParameters = SystemCompiledData.InstanceParamStore;
-	const TArray<FNiagaraVariable>& NumParticleVars = SystemCompiledData.NumParticleVars;
-	const TArray<FNiagaraVariable>& TotalSpawnedParticlesVars = SystemCompiledData.TotalSpawnedParticlesVars;
-	const TArray<FNiagaraVariable>& SpawnCountScaleVars = SystemCompiledData.SpawnCountScaleVars;
-
-	// Make sure all parameters are added before initializing the bindings, otherwise parameter store layout changes might invalidate the bindings.
-	OwnerPositionParam.Init(InstanceParameters, SYS_PARAM_ENGINE_POSITION);
-	OwnerScaleParam.Init(InstanceParameters, SYS_PARAM_ENGINE_SCALE);
-	OwnerVelocityParam.Init(InstanceParameters, SYS_PARAM_ENGINE_VELOCITY);
-	OwnerXAxisParam.Init(InstanceParameters, SYS_PARAM_ENGINE_X_AXIS);
-	OwnerYAxisParam.Init(InstanceParameters, SYS_PARAM_ENGINE_Y_AXIS);
-	OwnerZAxisParam.Init(InstanceParameters, SYS_PARAM_ENGINE_Z_AXIS);
-
-	OwnerRotationParam.Init(InstanceParameters, SYS_PARAM_ENGINE_ROTATION);
-
-	OwnerTransformParam.Init(InstanceParameters, SYS_PARAM_ENGINE_LOCAL_TO_WORLD);
-	OwnerInverseParam.Init(InstanceParameters, SYS_PARAM_ENGINE_WORLD_TO_LOCAL);
-	OwnerTransposeParam.Init(InstanceParameters, SYS_PARAM_ENGINE_LOCAL_TO_WORLD_TRANSPOSED);
-	OwnerInverseTransposeParam.Init(InstanceParameters, SYS_PARAM_ENGINE_WORLD_TO_LOCAL_TRANSPOSED);
-	OwnerTransformNoScaleParam.Init(InstanceParameters, SYS_PARAM_ENGINE_LOCAL_TO_WORLD_NO_SCALE);
-	OwnerInverseNoScaleParam.Init(InstanceParameters, SYS_PARAM_ENGINE_WORLD_TO_LOCAL_NO_SCALE);
-
-	OwnerDeltaSecondsParam.Init(InstanceParameters, SYS_PARAM_ENGINE_DELTA_TIME);
-	OwnerInverseDeltaSecondsParam.Init(InstanceParameters, SYS_PARAM_ENGINE_INV_DELTA_TIME);
-
-	SystemAgeParam.Init(InstanceParameters, SYS_PARAM_ENGINE_SYSTEM_AGE);
-	SystemTickCountParam.Init(InstanceParameters, SYS_PARAM_ENGINE_SYSTEM_TICK_COUNT);
-	OwnerEngineTimeParam.Init(InstanceParameters, SYS_PARAM_ENGINE_TIME);
-	OwnerEngineRealtimeParam.Init(InstanceParameters, SYS_PARAM_ENGINE_REAL_TIME);
-
-	OwnerLODDistanceParam.Init(InstanceParameters, SYS_PARAM_ENGINE_LOD_DISTANCE);
-	OwnerLODDistanceFractionParam.Init(InstanceParameters, SYS_PARAM_ENGINE_LOD_DISTANCE_FRACTION);
-	SystemNumEmittersParam.Init(InstanceParameters, SYS_PARAM_ENGINE_SYSTEM_NUM_EMITTERS);
-	SystemNumEmittersAliveParam.Init(InstanceParameters, SYS_PARAM_ENGINE_SYSTEM_NUM_EMITTERS_ALIVE);
-
-	SystemTimeSinceRenderedParam.Init(InstanceParameters, SYS_PARAM_ENGINE_TIME_SINCE_RENDERED);
-
-	OwnerExecutionStateParam.Init(InstanceParameters, SYS_PARAM_ENGINE_EXECUTION_STATE);
-
-	ParameterNumParticleBindings.SetNum(NumParticleVars.Num());
-	for (int32 i = 0; i < NumParticleVars.Num(); i++)
-	{
-		ParameterNumParticleBindings[i].Init(InstanceParameters, NumParticleVars[i]);
-	}
-
-	ParameterTotalSpawnedParticlesBindings.SetNum(TotalSpawnedParticlesVars.Num());
-	for (int32 i = 0; i < TotalSpawnedParticlesVars.Num(); i++)
-	{
-		ParameterTotalSpawnedParticlesBindings[i].Init(InstanceParameters, TotalSpawnedParticlesVars[i]);
-	}
-
-	ParameterSpawnCountScaleBindings.SetNum(SpawnCountScaleVars.Num());
-	for (int32 i = 0; i < SpawnCountScaleVars.Num(); i++)
-	{
-		ParameterSpawnCountScaleBindings[i].Init(InstanceParameters, SpawnCountScaleVars[i]);
-	}
 
 	// rebind now after all parameters have been added
 	InstanceParameters.Rebind();
@@ -970,6 +918,29 @@ void FNiagaraSystemInstance::ReInitInternal()
 
 }
 
+void FNiagaraSystemInstance::ResetParameters()
+{
+	UNiagaraSystem* System = GetSystem();
+	if (!System)
+	{
+		return;
+	}
+
+	CurrentFrameIndex = 1;
+	ParametersValid = false;
+
+	GlobalParameters[0] = FNiagaraGlobalParameters();
+	GlobalParameters[1] = FNiagaraGlobalParameters();
+	SystemParameters[0] = FNiagaraSystemParameters();
+	SystemParameters[1] = FNiagaraSystemParameters();
+	OwnerParameters[0] = FNiagaraOwnerParameters();
+	OwnerParameters[1] = FNiagaraOwnerParameters();
+
+	EmitterParameters.Reset(Emitters.Num() * 2);
+	EmitterParameters.AddDefaulted(Emitters.Num() * 2);
+	GatheredInstanceParameters.Init(Emitters.Num());
+}
+
 FNiagaraSystemInstance::~FNiagaraSystemInstance()
 {
 	//UE_LOG(LogNiagara, Log, TEXT("~FNiagaraSystemInstance %p"), this);
@@ -1482,55 +1453,65 @@ void FNiagaraSystemInstance::TickInstanceParameters_GameThread(float DeltaSecond
 	{
 		return;
 	}
-	check(Component->GetAsset());
-	UNiagaraSystem* System = Component->GetAsset();
 
-	GatheredInstanceParameters.DeltaSeconds = DeltaSeconds;
+	const int EffectsQualityLevel = EffectsQualityLevelCVar->GetInt();
+
+	const FTransform& ComponentTransform = Component->GetComponentTransform();
+	const bool TransformMatches = GatheredInstanceParameters.ComponentTrans.Equals(ComponentTransform);
+	if (TransformMatches)
+	{
+		GatheredInstanceParameters.TransformMatchCount = FMath::Min(ParameterBufferCount, GatheredInstanceParameters.TransformMatchCount + 1);
+	}
+	else
+	{
+		GatheredInstanceParameters.ComponentTrans = ComponentTransform;
+		GatheredInstanceParameters.TransformMatchCount = 0;
+	}
 
-	GatheredInstanceParameters.ComponentTrans = Component->GetComponentTransform();
-	GatheredInstanceParameters.OldPos = OwnerPositionParam.GetValue();
+	GatheredInstanceParameters.EmitterCount = Emitters.Num();
+	GatheredInstanceParameters.DeltaSeconds = DeltaSeconds;
+	GatheredInstanceParameters.NumAlive = 0;
 
-	UWorld* World = Component->GetWorld();
-	if (World != NULL)
+	//Bias the LastRenderTime slightly to account for any delay as it's written by the RT.
+	if (const UWorld* World = Component->GetWorld())
 	{
-		GatheredInstanceParameters.LODDistance = GetLODDistance();
-		GatheredInstanceParameters.MaxLODDistance = MaxLODDistance;
 		GatheredInstanceParameters.TimeSeconds = World->TimeSeconds;
 		GatheredInstanceParameters.RealTimeSeconds = World->RealTimeSeconds;
 	}
 	else
 	{
-		GatheredInstanceParameters.LODDistance = 0;
-		GatheredInstanceParameters.MaxLODDistance = 1;
 		GatheredInstanceParameters.TimeSeconds = Age;
 		GatheredInstanceParameters.RealTimeSeconds = Age;
 	}
 
-	GatheredInstanceParameters.Age = Age;
-	GatheredInstanceParameters.TickCount = TickCount;
+	// flip our buffered parameters
+	FlipParameterBuffers();
+	uint32 ParameterIndex = GetParameterIndex();
 
-	int EffectsQualityLevel = EffectsQualityLevelCVar->GetInt();
-
-	GatheredInstanceParameters.NumAlive = 0;
-	for (int32 i = 0; i < Emitters.Num(); ++i)
+	for (int32 i = 0; i < GatheredInstanceParameters.EmitterCount; ++i)
 	{
+		auto& CurrentEmitterParameters = EditEmitterParameters(i);
+
 		const TSharedRef<FNiagaraEmitterInstance, ESPMode::ThreadSafe>& Emitter = Emitters[i];
 		if (Emitter->GetExecutionState() != ENiagaraExecutionState::Disabled)
 		{
-			GatheredInstanceParameters.EmitterNumParticles[i] = Emitter->GetNumParticles();
-			GatheredInstanceParameters.EmitterTotalSpawnedParticles[i] = Emitter->GetTotalSpawnedParticles();
-			GatheredInstanceParameters.EmitterSpawnCountScale[i] = Emitter->GetSpawnCountScale(EffectsQualityLevel);
-			if (!Emitter->IsComplete())
-			{
-				GatheredInstanceParameters.NumAlive++;
-			}
+			CurrentEmitterParameters.EmitterNumParticles = Emitter->GetNumParticles();
+			CurrentEmitterParameters.EmitterTotalSpawnedParticles = Emitter->GetTotalSpawnedParticles();
+			CurrentEmitterParameters.EmitterSpawnCountScale = Emitter->GetSpawnCountScale(EffectsQualityLevel);
+		}
+		else
+		{
+			CurrentEmitterParameters.EmitterNumParticles = 0;
 		}
 	}
 
-	//Bias the LastRenderTime slightly to account for any delay as it's written by the RT.
-	GatheredInstanceParameters.SafeTimeSinceRendererd = Component->GetSafeTimeSinceRendered(GatheredInstanceParameters.TimeSeconds);
-
-	GatheredInstanceParameters.RequestedExecutionState = RequestedExecutionState;
+	FNiagaraSystemParameters& CurrentSystemParameters = SystemParameters[ParameterIndex];
+	CurrentSystemParameters.EngineSystemAge = Age;
+	CurrentSystemParameters.EngineTickCount = TickCount;
+	CurrentSystemParameters.EngineTimeSinceRendered = FMath::Max(0.0f, GatheredInstanceParameters.TimeSeconds - Component->GetLastRenderTime() - GLastRenderTimeSafetyBias);
+	CurrentSystemParameters.EngineExecutionState = static_cast<uint32>(RequestedExecutionState);
+	CurrentSystemParameters.EngineLodDistance = GetLODDistance();
+	CurrentSystemParameters.EngineLodDistanceFraction = CurrentSystemParameters.EngineLodDistance / MaxLODDistance;
 
 	Component->GetOverrideParameters().Tick();
 }
@@ -1538,64 +1519,42 @@ void FNiagaraSystemInstance::TickInstanceParameters_GameThread(float DeltaSecond
 
 void FNiagaraSystemInstance::TickInstanceParameters_Concurrent()
 {
-	FVector CurrPos = GatheredInstanceParameters.ComponentTrans.GetLocation();
-	FVector OldPos = FMath::IsNearlyZero(Age) ? CurrPos : GatheredInstanceParameters.OldPos; // The first frame the value in OwnerPositionParam is uninitialized memory, we need to make sure that we don't use it
-
-	OwnerPositionParam.SetValue(CurrPos);
-	OwnerScaleParam.SetValue(GatheredInstanceParameters.ComponentTrans.GetScale3D());
-	OwnerVelocityParam.SetValue((CurrPos - OldPos) / GatheredInstanceParameters.DeltaSeconds);
-
-	FQuat Rotation = GatheredInstanceParameters.ComponentTrans.GetRotation();
-	OwnerXAxisParam.SetValue(Rotation.GetAxisX());
-	OwnerYAxisParam.SetValue(Rotation.GetAxisY());
-	OwnerZAxisParam.SetValue(Rotation.GetAxisZ());
+	uint32 ParameterIndex = GetParameterIndex();
+	FNiagaraSystemParameters& CurrentSystemParameters = SystemParameters[ParameterIndex];
+	FNiagaraOwnerParameters& CurrentOwnerParameters = OwnerParameters[ParameterIndex];
 
-	OwnerRotationParam.SetValue(Rotation);
-
-	FMatrix Transform = GatheredInstanceParameters.ComponentTrans.ToMatrixWithScale();
-	FMatrix Inverse = Transform.Inverse();
-	FMatrix Transpose = Transform.GetTransposed();
-	FMatrix InverseTranspose = Inverse.GetTransposed();
-	OwnerTransformParam.SetValue(Transform);
-	OwnerInverseParam.SetValue(Inverse);
-	OwnerTransposeParam.SetValue(Transpose);
-	OwnerInverseTransposeParam.SetValue(InverseTranspose);
-
-	FMatrix TransformNoScale = GatheredInstanceParameters.ComponentTrans.ToMatrixNoScale();
-	FMatrix InverseNoScale = TransformNoScale.Inverse();
-	OwnerTransformNoScaleParam.SetValue(TransformNoScale);
-	OwnerInverseNoScaleParam.SetValue(InverseNoScale);
-
-	OwnerDeltaSecondsParam.SetValue(GatheredInstanceParameters.DeltaSeconds);
-	OwnerInverseDeltaSecondsParam.SetValue(1.0f / GatheredInstanceParameters.DeltaSeconds);
-
-	OwnerLODDistanceParam.SetValue(GatheredInstanceParameters.LODDistance);
-	float SystemDistanceFraction = FMath::Clamp(GatheredInstanceParameters.LODDistance / GatheredInstanceParameters.MaxLODDistance, 0.0f, 1.0f);
-	OwnerLODDistanceFractionParam.SetValue(SystemDistanceFraction);
-	OwnerEngineTimeParam.SetValue(GatheredInstanceParameters.TimeSeconds);
-	OwnerEngineRealtimeParam.SetValue(GatheredInstanceParameters.RealTimeSeconds);
+	if (GatheredInstanceParameters.TransformMatchCount < ParameterBufferCount)
+	{
+		const FMatrix LocalToWorld = GatheredInstanceParameters.ComponentTrans.ToMatrixWithScale();
+		const FMatrix LocalToWorldNoScale = GatheredInstanceParameters.ComponentTrans.ToMatrixNoScale();
 
-	SystemAgeParam.SetValue(GatheredInstanceParameters.Age);
-	SystemTickCountParam.SetValue(GatheredInstanceParameters.TickCount);
+		const FVector Location = GatheredInstanceParameters.ComponentTrans.GetLocation();
+		const FVector LastLocation = FMath::IsNearlyZero(CurrentSystemParameters.EngineSystemAge) ? Location : FVector(OwnerParameters[GetParameterIndex(true)].EnginePosition);
 
-	for (int32 i = 0; i < Emitters.Num(); i++)
-	{
-		const TSharedRef<FNiagaraEmitterInstance, ESPMode::ThreadSafe>& Emitter = Emitters[i];
-		if (Emitter->GetExecutionState() != ENiagaraExecutionState::Disabled)
-		{
-			ParameterNumParticleBindings[i].SetValue(GatheredInstanceParameters.EmitterNumParticles[i]);
-			ParameterTotalSpawnedParticlesBindings[i].SetValue(GatheredInstanceParameters.EmitterTotalSpawnedParticles[i]);
-			ParameterSpawnCountScaleBindings[i].SetValue(GatheredInstanceParameters.EmitterSpawnCountScale[i]);
-		}
+		CurrentOwnerParameters.EngineLocalToWorld = LocalToWorld;
+		CurrentOwnerParameters.EngineWorldToLocal = LocalToWorld.Inverse();
+		CurrentOwnerParameters.EngineLocalToWorldTransposed = LocalToWorld.GetTransposed();
+		CurrentOwnerParameters.EngineWorldToLocalTransposed = CurrentOwnerParameters.EngineWorldToLocal.GetTransposed();
+		CurrentOwnerParameters.EngineLocalToWorldNoScale = LocalToWorldNoScale;
+		CurrentOwnerParameters.EngineWorldToLocalNoScale = LocalToWorldNoScale.Inverse();
+		CurrentOwnerParameters.EngineRotation = GatheredInstanceParameters.ComponentTrans.GetRotation();
+		CurrentOwnerParameters.EnginePosition = GatheredInstanceParameters.ComponentTrans.GetLocation();
+		CurrentOwnerParameters.EngineVelocity = (Location - LastLocation) / GatheredInstanceParameters.DeltaSeconds;
+		CurrentOwnerParameters.EngineXAxis = CurrentOwnerParameters.EngineRotation.GetAxisX();
+		CurrentOwnerParameters.EngineYAxis = CurrentOwnerParameters.EngineRotation.GetAxisY();
+		CurrentOwnerParameters.EngineZAxis = CurrentOwnerParameters.EngineRotation.GetAxisZ();
+		CurrentOwnerParameters.EngineScale = GatheredInstanceParameters.ComponentTrans.GetScale3D();
 	}
 
-	SystemNumEmittersParam.SetValue(Emitters.Num());
-	SystemNumEmittersAliveParam.SetValue(GatheredInstanceParameters.NumAlive);
+	CurrentSystemParameters.EngineEmitterCount = GatheredInstanceParameters.EmitterCount;
+	CurrentSystemParameters.EngineAliveEmitterCount = GatheredInstanceParameters.NumAlive;
 
-	SystemTimeSinceRenderedParam.SetValue(GatheredInstanceParameters.SafeTimeSinceRendererd);
+	FNiagaraGlobalParameters& CurrentGlobalParameter = GlobalParameters[ParameterIndex];
+	CurrentGlobalParameter.EngineDeltaTime = GatheredInstanceParameters.DeltaSeconds;
+	CurrentGlobalParameter.EngineInvDeltaTime = 1.0f / GatheredInstanceParameters.DeltaSeconds;
+	CurrentGlobalParameter.EngineRealTime = GatheredInstanceParameters.RealTimeSeconds;
+	CurrentGlobalParameter.EngineTime = GatheredInstanceParameters.TimeSeconds;
 
-	OwnerExecutionStateParam.SetValue((int32)GatheredInstanceParameters.RequestedExecutionState);
-	
 	InstanceParameters.Tick();
 	InstanceParameters.MarkParametersDirty();
 }
@@ -1753,7 +1712,7 @@ void FNiagaraSystemInstance::InitEmitters()
 		}
 	}
 
-	GatheredInstanceParameters.Init(Emitters.Num());
+	ResetParameters();
 }
 
 void FNiagaraSystemInstance::ComponentTick(float DeltaSeconds, const FGraphEventRef& MyCompletionGraphEvent)
@@ -1879,7 +1838,7 @@ void FNiagaraSystemInstance::Tick_Concurrent()
 	FScopeCycleCounterUObject AdditionalScope(GetSystem(), GET_STATID(STAT_NiagaraOverview_GT_CNC));
 
 	// Reset values that will be accumulated during emitter tick.
-	TotalParamSize = 0;
+	TotalGPUParamSize = 0;
 	ActiveGPUEmitterCount = 0;
 	UNiagaraSystem* System = GetSystem();
 
@@ -1911,6 +1870,8 @@ void FNiagaraSystemInstance::Tick_Concurrent()
 		}
 	}
 
+	bool FirstGpuEmitter = true;
+
 	// now tick all emitters
 	for (int32 EmitterIdx = 0; EmitterIdx < Emitters.Num(); EmitterIdx++)
 	{
@@ -1921,7 +1882,17 @@ void FNiagaraSystemInstance::Tick_Concurrent()
 
 			if (Inst.GetCachedEmitter() && Inst.GetCachedEmitter()->SimTarget == ENiagaraSimTarget::GPUComputeSim && Inst.GetGPUContext() != nullptr && (Inst.GetExecutionState() != ENiagaraExecutionState::Complete))
 			{
-				TotalParamSize += Inst.GetGPUContext()->CombinedParamStore.GetPaddedParameterSizeInBytes();
+				if (FirstGpuEmitter)
+				{
+					TotalGPUParamSize += 2 * sizeof(FNiagaraGlobalParameters);
+					TotalGPUParamSize += 2 * sizeof(FNiagaraSystemParameters);
+					TotalGPUParamSize += 2 * sizeof(FNiagaraOwnerParameters);
+					FirstGpuEmitter = false;
+				}
+
+				TotalGPUParamSize += 2 * sizeof(FNiagaraEmitterParameters);
+				
+				TotalGPUParamSize += Inst.GetGPUContext()->CombinedParamStore.GetPaddedParameterSizeInBytes();
 				ActiveGPUEmitterCount++;
 			}
 		}
diff --git a/Engine/Plugins/FX/Niagara/Source/Niagara/Private/NiagaraSystemSimulation.cpp b/Engine/Plugins/FX/Niagara/Source/Niagara/Private/NiagaraSystemSimulation.cpp
--- a/Engine/Plugins/FX/Niagara/Source/Niagara/Private/NiagaraSystemSimulation.cpp
+++ b/Engine/Plugins/FX/Niagara/Source/Niagara/Private/NiagaraSystemSimulation.cpp
@@ -1,4 +1,4 @@
-// Copyright Epic Games, Inc. All Rights Reserved.
+// Copyright 1998-2019 Epic Games, Inc. All Rights Reserved.
 
 #include "NiagaraSystemSimulation.h"
 #include "NiagaraModule.h"
@@ -309,7 +309,8 @@ bool FNiagaraSystemSimulation::Init(UNiagaraSystem* InSystem, UWorld* InWorld, b
 			SpawnInstanceParameterDataSet.Init(&SystemCompiledData.SpawnInstanceParamsDataSetCompiledData);
 
 			UpdateInstanceParameterDataSet.Init(&SystemCompiledData.UpdateInstanceParamsDataSetCompiledData);
-			
+
+			ConstantBufferToDataSetBinding.Init(SystemCompiledData);
 		}
 
 		UNiagaraScript* SpawnScript = System->GetSystemSpawnScript();
@@ -400,12 +401,6 @@ bool FNiagaraSystemSimulation::Init(UNiagaraSystem* InSystem, UWorld* InWorld, b
 		{
 			SCOPE_CYCLE_COUNTER(STAT_NiagaraSystemSim_Init_DirectBindings);
 
-			SpawnTimeParam.Init(SpawnExecContext.Parameters, SYS_PARAM_ENGINE_TIME);
-			UpdateTimeParam.Init(UpdateExecContext.Parameters, SYS_PARAM_ENGINE_TIME);
-			SpawnDeltaTimeParam.Init(SpawnExecContext.Parameters, SYS_PARAM_ENGINE_DELTA_TIME);
-			UpdateDeltaTimeParam.Init(UpdateExecContext.Parameters, SYS_PARAM_ENGINE_DELTA_TIME);
-			SpawnInvDeltaTimeParam.Init(SpawnExecContext.Parameters, SYS_PARAM_ENGINE_INV_DELTA_TIME);
-			UpdateInvDeltaTimeParam.Init(UpdateExecContext.Parameters, SYS_PARAM_ENGINE_INV_DELTA_TIME);
 			SpawnNumSystemInstancesParam.Init(SpawnExecContext.Parameters, SYS_PARAM_ENGINE_NUM_SYSTEM_INSTANCES);
 			UpdateNumSystemInstancesParam.Init(UpdateExecContext.Parameters, SYS_PARAM_ENGINE_NUM_SYSTEM_INSTANCES);
 			SpawnGlobalSpawnCountScaleParam.Init(SpawnExecContext.Parameters, SYS_PARAM_ENGINE_GLOBAL_SPAWN_COUNT_SCALE);
@@ -720,10 +715,10 @@ void FNiagaraSystemSimulation::Tick_GameThread(float DeltaSeconds, const FGraphE
 		}
 	}
 
-	// Setup the few real constants like delta time.
+	//Setup the few real constants like delta time.
 	SetupParameters_GameThread(DeltaSeconds);
 
-	// Gather any pending spawn systems and add to the end of the system instances
+	// Somethings we don't want to happen during the spawn phase
 	int32 SpawnNum = 0;
 	if (PendingSystemInstances.Num() > 0)
 	{
@@ -733,16 +728,16 @@ void FNiagaraSystemSimulation::Tick_GameThread(float DeltaSeconds, const FGraphE
 		while (SystemIndex < PendingSystemInstances.Num())
 		{
 			FNiagaraSystemInstance* Inst = PendingSystemInstances[SystemIndex];
+			// Gather any pending spawn systems and add to the end of the system instances
 
-			// If we are paused continue
 			if (Inst->IsPaused())
 			{
 				++SystemIndex;
 				continue;
 			}
 
-			// If our tick group has changed we need to move this pending instance
-			// Note we don't do this with solo instances
+			// If we are paused continue
+
 			if (!bIsSolo)
 			{
 				const ETickingGroup DesiredTickGroup = Inst->CalculateTickGroup();
@@ -1119,9 +1114,11 @@ void FNiagaraSystemSimulation::TickFastPath(FNiagaraSystemSimulationTickContext&
 	// PrepareForSystemSimulate
 	for (FNiagaraSystemInstance* SystemInstance : Context.Instances)
 	{
+		const auto& SystemParameters = SystemInstance->GetSystemParameters();
+
 		FNiagaraSystemFastPath::FParamMap0& SystemMap = SystemInstance->GetFastPathMap();
 		SystemMap.Engine.Owner.ExecutionState = SystemInstance->GetRequestedExecutionState();
-		SystemMap.Engine.Owner.TimeSinceRendered = SystemInstance->GetSystemTimeSinceRendered();
+		SystemMap.Engine.Owner.TimeSinceRendered = SystemInstance->GetSystemParameters().EngineTimeSinceRendered;
 
 		TArray<TSharedRef<FNiagaraEmitterInstance, ESPMode::ThreadSafe>>& EmitterInstances = SystemInstance->GetEmitters();
 		for (int32 EmitterIndex = 0; EmitterIndex < EmitterInstances.Num(); ++EmitterIndex)
@@ -1132,11 +1129,11 @@ void FNiagaraSystemSimulation::TickFastPath(FNiagaraSystemSimulationTickContext&
 				continue;
 			}
 			FNiagaraEmitterFastPath::FParamMap0& EmitterMap = EmitterInstance->GetFastPathMap();
-			EmitterMap.Engine.Owner.LODDistance = SystemInstance->GetOwnerLODDistance();
-			EmitterMap.Engine.Owner.MaxLODDistance = SystemInstance->GetOwnerMaxLODDistance();
+			EmitterMap.Engine.Owner.LODDistance = SystemParameters.EngineLodDistance;
+			EmitterMap.Engine.Owner.MaxLODDistance = SystemParameters.EngineLodDistance / SystemParameters.EngineLodDistanceFraction;
 			EmitterMap.Engine.DeltaTime = Context.DeltaSeconds;
-			EmitterMap.Engine.Emitter.NumParticles = SystemInstance->GetNumParticles(EmitterIndex);
-			EmitterMap.Engine.Owner.Velocity = SystemInstance->GetOwnerVelocity();
+			EmitterMap.Engine.Emitter.NumParticles = SystemInstance->GetEmitterParameters(EmitterIndex).EmitterNumParticles;
+			EmitterMap.Engine.Owner.Velocity = SystemInstance->GetOwnerParameters().EngineVelocity;
 			EmitterMap.Emitter.ExecutionState = EmitterInstance->GetExecutionState();
 			EmitterMap.Engine.GlobalSpawnCountScale = INiagaraModule::GetGlobalSpawnCountScale();
 
@@ -1334,13 +1331,6 @@ void FNiagaraSystemSimulation::SetupParameters_GameThread(float DeltaSeconds)
 {
 	check(IsInGameThread());
 
-	const float InvDt = 1.0f / DeltaSeconds;
-	SpawnTimeParam.SetValue(World->TimeSeconds);
-	UpdateTimeParam.SetValue(World->TimeSeconds);
-	SpawnDeltaTimeParam.SetValue(DeltaSeconds);
-	UpdateDeltaTimeParam.SetValue(DeltaSeconds);
-	SpawnInvDeltaTimeParam.SetValue(InvDt);
-	UpdateInvDeltaTimeParam.SetValue(InvDt);
 	SpawnNumSystemInstancesParam.SetValue(SystemInstances.Num());
 	UpdateNumSystemInstancesParam.SetValue(SystemInstances.Num());
 	SpawnGlobalSpawnCountScaleParam.SetValue(INiagaraModule::GetGlobalSpawnCountScale());
@@ -1379,13 +1369,16 @@ void FNiagaraSystemSimulation::PrepareForSystemSimulate(FNiagaraSystemSimulation
 	auto TransferInstanceParameters = [&](int32 SystemIndex)
 	{
 		FNiagaraSystemInstance* Inst = Context.Instances[SystemIndex];
+		const FNiagaraParameterStore& InstParameters = Inst->GetInstanceParameters();
 
-		if (Inst->GetParameters().GetParametersDirty() && bCanExecute)
+		if (InstParameters.GetParametersDirty() && bCanExecute)
 		{
-			SpawnInstanceParameterToDataSetBinding.ParameterStoreToDataSet(Inst->GetParameters(), SpawnInstanceParameterDataSet, SystemIndex);
-			UpdateInstanceParameterToDataSetBinding.ParameterStoreToDataSet(Inst->GetParameters(), UpdateInstanceParameterDataSet, SystemIndex);
+			SpawnInstanceParameterToDataSetBinding.ParameterStoreToDataSet(InstParameters, SpawnInstanceParameterDataSet, SystemIndex);
+			UpdateInstanceParameterToDataSetBinding.ParameterStoreToDataSet(InstParameters, UpdateInstanceParameterDataSet, SystemIndex);
 		}
 
+		ConstantBufferToDataSetBinding.CopyToDataSets(*Inst, SpawnInstanceParameterDataSet, UpdateInstanceParameterDataSet, SystemIndex);
+
 		//TODO: Find good way to check that we're not using any instance parameter data interfaces in the system scripts here.
 		//In that case we need to solo and will never get here.
 
@@ -1437,7 +1430,11 @@ void FNiagaraSystemSimulation::SpawnSystemInstances(FNiagaraSystemSimulationTick
 	SpawnExecContext.Tick(SoloSystemInstance);//We can't require a specific instance here as these are for all instances.
 	SpawnExecContext.BindData(0, Context.DataSet, OrigNum, false);
 	SpawnExecContext.BindData(1, SpawnInstanceParameterDataSet, OrigNum, false);
-	SpawnExecContext.Execute(SpawnNum);
+
+	FScriptExecutionConstantBufferTable SpawnConstantBufferTable;
+	BuildConstantBufferTable(Context.Instances[0]->GetGlobalParameters(), SpawnExecContext, SpawnConstantBufferTable);
+
+	SpawnExecContext.Execute(SpawnNum, SpawnConstantBufferTable);
 
 	if (GbDumpSystemData || Context.System->bDumpDebugSystemInfo)
 	{
@@ -1486,7 +1483,11 @@ void FNiagaraSystemSimulation::UpdateSystemInstances(FNiagaraSystemSimulationTic
 			UpdateExecContext.Tick(Context.Instances[0]);
 			UpdateExecContext.BindData(0, Context.DataSet, 0, false);
 			UpdateExecContext.BindData(1, UpdateInstanceParameterDataSet, 0, false);
-			UpdateExecContext.Execute(OrigNum);
+
+			FScriptExecutionConstantBufferTable UpdateConstantBufferTable;
+			BuildConstantBufferTable(Context.Instances[0]->GetGlobalParameters(), UpdateExecContext, UpdateConstantBufferTable);
+
+			UpdateExecContext.Execute(OrigNum, UpdateConstantBufferTable);
 		}
 
 		if (GbDumpSystemData || Context.System->bDumpDebugSystemInfo)
@@ -1502,12 +1503,18 @@ void FNiagaraSystemSimulation::UpdateSystemInstances(FNiagaraSystemSimulationTic
 		if ( (SpawnNum > 0) && GbSystemUpdateOnSpawn)
 		{
 			UpdateExecContext.Tick(Context.Instances[0]);
-			UpdateExecContext.Parameters.SetParameterValue(0.0001f, SYS_PARAM_ENGINE_DELTA_TIME);
-			UpdateExecContext.Parameters.SetParameterValue(10000.0f, SYS_PARAM_ENGINE_INV_DELTA_TIME);
 
 			UpdateExecContext.BindData(0, Context.DataSet, OrigNum, false);
 			UpdateExecContext.BindData(1, UpdateInstanceParameterDataSet, OrigNum, false);
-			UpdateExecContext.Execute(SpawnNum);
+
+			FNiagaraGlobalParameters UpdateOnSpawnParameters(Context.Instances[0]->GetGlobalParameters());
+			UpdateOnSpawnParameters.EngineDeltaTime = 0.0001f;
+			UpdateOnSpawnParameters.EngineInvDeltaTime = 10000.0f;
+
+			FScriptExecutionConstantBufferTable UpdateConstantBufferTable;
+			BuildConstantBufferTable(UpdateOnSpawnParameters, UpdateExecContext, UpdateConstantBufferTable);
+
+			UpdateExecContext.Execute(SpawnNum, UpdateConstantBufferTable);
 
 			if (GbDumpSystemData || Context.System->bDumpDebugSystemInfo)
 			{
@@ -1884,12 +1891,19 @@ void FNiagaraSystemSimulation::InitParameterDataSetBindings(FNiagaraSystemInstan
 		UpdateInstanceParameterToDataSetBinding.Init(UpdateInstanceParameterDataSet, SystemInst->GetInstanceParameters());
 
 		TArray<TSharedRef<FNiagaraEmitterInstance, ESPMode::ThreadSafe>>& Emitters = SystemInst->GetEmitters();
-		DataSetToEmitterSpawnParameters.SetNum(Emitters.Num());
-		DataSetToEmitterUpdateParameters.SetNum(Emitters.Num());
-		DataSetToEmitterEventParameters.SetNum(Emitters.Num());
-		for (int32 EmitterIdx = 0; EmitterIdx < Emitters.Num(); ++EmitterIdx)
+		const int32 EmitterCount = Emitters.Num();
+
+		DataSetToEmitterSpawnParameters.SetNum(EmitterCount);
+		DataSetToEmitterUpdateParameters.SetNum(EmitterCount);
+		DataSetToEmitterEventParameters.SetNum(EmitterCount);
+
+		const FString EmitterNamespace = TEXT("Emitter");
+
+		for (int32 EmitterIdx = 0; EmitterIdx < EmitterCount; ++EmitterIdx)
 		{
 			FNiagaraEmitterInstance& EmitterInst = Emitters[EmitterIdx].Get();
+			const FString EmitterName = EmitterInst.GetCachedEmitter()->GetUniqueEmitterName();
+
 			FNiagaraScriptExecutionContext& SpawnContext = EmitterInst.GetSpawnExecutionContext();
 			DataSetToEmitterSpawnParameters[EmitterIdx].Init(MainDataSet, SpawnContext.Parameters);
 
@@ -1897,8 +1911,10 @@ void FNiagaraSystemSimulation::InitParameterDataSetBindings(FNiagaraSystemInstan
 			DataSetToEmitterUpdateParameters[EmitterIdx].Init(MainDataSet, UpdateContext.Parameters);
 
 			TArray<FNiagaraScriptExecutionContext>& EventContexts = EmitterInst.GetEventExecutionContexts();
-			DataSetToEmitterEventParameters[EmitterIdx].SetNum(EventContexts.Num());
-			for (int32 EventIdx = 0; EventIdx < EventContexts.Num(); ++EventIdx)
+			const int32 EventCount = EventContexts.Num();
+			DataSetToEmitterEventParameters[EmitterIdx].SetNum(EventCount);
+
+			for (int32 EventIdx = 0; EventIdx < EventCount; ++EventIdx)
 			{
 				FNiagaraScriptExecutionContext& EventContext = EventContexts[EventIdx];
 				DataSetToEmitterEventParameters[EmitterIdx][EventIdx].Init(MainDataSet, EventContext.Parameters);
@@ -1906,3 +1922,95 @@ void FNiagaraSystemSimulation::InitParameterDataSetBindings(FNiagaraSystemInstan
 		}
 	}
 }
+
+void FNiagaraConstantBufferToDataSetBinding::Init(const FNiagaraSystemCompiledData& CompiledData)
+{
+	// for now we'll copy the data to our local structure so that we don't have to worry about the lifetime of the compiled data
+	SpawnInstanceGlobalBinding = CompiledData.SpawnInstanceGlobalBinding;
+	SpawnInstanceSystemBinding = CompiledData.SpawnInstanceSystemBinding;
+	SpawnInstanceOwnerBinding = CompiledData.SpawnInstanceOwnerBinding;
+	SpawnInstanceEmitterBindings = CompiledData.SpawnInstanceEmitterBindings;
+
+	UpdateInstanceGlobalBinding = CompiledData.UpdateInstanceGlobalBinding;
+	UpdateInstanceSystemBinding = CompiledData.UpdateInstanceSystemBinding;
+	UpdateInstanceOwnerBinding = CompiledData.UpdateInstanceOwnerBinding;
+	UpdateInstanceEmitterBindings = CompiledData.UpdateInstanceEmitterBindings;
+}
+
+void FNiagaraConstantBufferToDataSetBinding::CopyToDataSets(
+	const FNiagaraSystemInstance& SystemInstance,
+	FNiagaraDataSet& SpawnDataSet,
+	FNiagaraDataSet& UpdateDataSet,
+	int32 DataSestInstanceIndex) const
+{
+	{
+		const uint8* GlobalParameters = reinterpret_cast<const uint8*>(&SystemInstance.GetGlobalParameters());
+		ApplyOffsets(SpawnInstanceGlobalBinding, GlobalParameters, SpawnDataSet, DataSestInstanceIndex);
+		ApplyOffsets(UpdateInstanceGlobalBinding, GlobalParameters, UpdateDataSet, DataSestInstanceIndex);
+	}
+
+	{
+		const uint8* SystemParameters = reinterpret_cast<const uint8*>(&SystemInstance.GetSystemParameters());
+		ApplyOffsets(SpawnInstanceSystemBinding, SystemParameters, SpawnDataSet, DataSestInstanceIndex);
+		ApplyOffsets(UpdateInstanceSystemBinding, SystemParameters, UpdateDataSet, DataSestInstanceIndex);
+	}
+
+	{
+		const uint8* OwnerParameters = reinterpret_cast<const uint8*>(&SystemInstance.GetOwnerParameters());
+		ApplyOffsets(SpawnInstanceOwnerBinding, OwnerParameters, SpawnDataSet, DataSestInstanceIndex);
+		ApplyOffsets(UpdateInstanceOwnerBinding, OwnerParameters, UpdateDataSet, DataSestInstanceIndex);
+	}
+
+	const auto& Emitters = SystemInstance.GetEmitters();
+	const int32 EmitterCount = Emitters.Num();
+
+	for (int32 EmitterIdx = 0; EmitterIdx < EmitterCount; ++EmitterIdx)
+	{
+		const uint8* EmitterParameters = reinterpret_cast<const uint8*>(&SystemInstance.GetEmitterParameters(EmitterIdx));
+		ApplyOffsets(SpawnInstanceEmitterBindings[EmitterIdx], EmitterParameters, SpawnDataSet, DataSestInstanceIndex);
+		ApplyOffsets(UpdateInstanceEmitterBindings[EmitterIdx], EmitterParameters, UpdateDataSet, DataSestInstanceIndex);
+	}
+}
+
+void FNiagaraConstantBufferToDataSetBinding::ApplyOffsets(
+	const FNiagaraParameterDataSetBindingCollection& Offsets,
+	const uint8* SourceData,
+	FNiagaraDataSet& DataSet,
+	int32 DataSetInstanceIndex) const
+{
+	FNiagaraDataBuffer& CurrBuffer = DataSet.GetDestinationDataChecked();
+
+	for (const auto& DataOffsets : Offsets.FloatOffsets)
+	{
+		float* ParamPtr = (float*)(SourceData + DataOffsets.ParameterOffset);
+		float* DataSetPtr = CurrBuffer.GetInstancePtrFloat(DataOffsets.DataSetComponentOffset, DataSetInstanceIndex);
+		*DataSetPtr = *ParamPtr;
+	}
+	for (const auto& DataOffsets : Offsets.Int32Offsets)
+	{
+		int32* ParamPtr = (int32*)(SourceData + DataOffsets.ParameterOffset);
+		int32* DataSetPtr = CurrBuffer.GetInstancePtrInt32(DataOffsets.DataSetComponentOffset, DataSetInstanceIndex);
+		*DataSetPtr = *ParamPtr;
+	}
+}
+
+
+void FNiagaraSystemSimulation::BuildConstantBufferTable(
+	const FNiagaraGlobalParameters& GlobalParameters,
+	FNiagaraScriptExecutionContext& ExecContext,
+	FScriptExecutionConstantBufferTable& ConstantBufferTable) const
+{
+	check(!ExecContext.HasInterpolationParameters);
+
+	const auto& ExternalParameterData = ExecContext.Parameters.GetParameterDataArray();
+	uint8* ExternalParameterBuffer = const_cast<uint8*>(ExternalParameterData.GetData());
+
+	const uint32 ExternalParameterSize = ExecContext.Parameters.GetExternalParameterSize();
+	const uint32 LiteralConstantOffset = ExternalParameterSize;
+	const uint32 LiteralConstantSize = ExternalParameterData.Num() - LiteralConstantOffset;
+
+	ConstantBufferTable.Reset(3);
+	ConstantBufferTable.AddTypedBuffer(GlobalParameters);
+	ConstantBufferTable.AddRawBuffer(ExternalParameterBuffer, ExternalParameterSize);
+	ConstantBufferTable.AddRawBuffer(ExternalParameterBuffer + LiteralConstantOffset, LiteralConstantSize);
+}
diff --git a/Engine/Plugins/FX/Niagara/Source/Niagara/Public/NiagaraModule.h b/Engine/Plugins/FX/Niagara/Source/Niagara/Public/NiagaraModule.h
--- a/Engine/Plugins/FX/Niagara/Source/Niagara/Public/NiagaraModule.h
+++ b/Engine/Plugins/FX/Niagara/Source/Niagara/Public/NiagaraModule.h
@@ -159,15 +159,6 @@ public:
 	FORCEINLINE static const FNiagaraVariable&  GetVar_DataInstance_Alive() { return DataInstance_Alive; }
 	FORCEINLINE static const FNiagaraVariable&  GetVar_BeginDefaults() { return Translator_BeginDefaults; }
 
-#if WITH_EDITORONLY_DATA
-	FORCEINLINE static const FNiagaraParameterStore& GetFixedSystemInstanceParameterStore() { return FixedSystemInstanceParameters; }
-#endif
-private:
-
-#if WITH_EDITORONLY_DATA
-	void InitFixedSystemInstanceParameterStore();
-#endif
-
 	FOnProcessQueue OnProcessQueue;
 
 #if WITH_EDITORONLY_DATA
@@ -266,9 +257,5 @@ private:
 	static FNiagaraVariable ScriptUsage;
 	static FNiagaraVariable DataInstance_Alive;
 	static FNiagaraVariable Translator_BeginDefaults;
-
-#if WITH_EDITORONLY_DATA
-	static FNiagaraParameterStore FixedSystemInstanceParameters;
-#endif
 };
 
diff --git a/Engine/Plugins/FX/Niagara/Source/Niagara/Public/NiagaraParameterStore.h b/Engine/Plugins/FX/Niagara/Source/Niagara/Public/NiagaraParameterStore.h
--- a/Engine/Plugins/FX/Niagara/Source/Niagara/Public/NiagaraParameterStore.h
+++ b/Engine/Plugins/FX/Niagara/Source/Niagara/Public/NiagaraParameterStore.h
@@ -138,49 +138,49 @@ struct NIAGARA_API FNiagaraParameterStore
 	GENERATED_USTRUCT_BODY()
 
 private:
-	/** Owner of this store. Used to provide an outer to data interfaces in this store. */
-	UPROPERTY(Transient)
-	UObject* Owner;
-	
+/** Owner of this store. Used to provide an outer to data interfaces in this store. */
+UPROPERTY(Transient)
+UObject* Owner;
+
 #if WITH_EDITORONLY_DATA
-	/** Map from parameter defs to their offset in the data table or the data interface. TODO: Separate out into a layout and instance class to reduce duplicated data for this?  */
-	UPROPERTY()
-	TMap<FNiagaraVariable, int32> ParameterOffsets;
+/** Map from parameter defs to their offset in the data table or the data interface. TODO: Separate out into a layout and instance class to reduce duplicated data for this?  */
+UPROPERTY()
+TMap<FNiagaraVariable, int32> ParameterOffsets;
 #endif // WITH_EDITORONLY_DATA
 
-	UPROPERTY()
-	TArray<FNiagaraVariableWithOffset> SortedParameterOffsets;
+UPROPERTY()
+TArray<FNiagaraVariableWithOffset> SortedParameterOffsets;
 
-	/** Buffer containing parameter data. Indexed using offsets in ParameterOffsets */
-	UPROPERTY()
-	TArray<uint8> ParameterData;
-	
-	/** Data interfaces for this script. Possibly overridden with externally owned interfaces. Also indexed by ParameterOffsets. */
-	UPROPERTY()
-	TArray<UNiagaraDataInterface*> DataInterfaces;
+/** Buffer containing parameter data. Indexed using offsets in ParameterOffsets */
+UPROPERTY()
+TArray<uint8> ParameterData;
 
-	/** UObjects referenced by this store. Also indexed by ParameterOffsets.*/
-	UPROPERTY()
-	TArray<UObject*> UObjects;
+/** Data interfaces for this script. Possibly overridden with externally owned interfaces. Also indexed by ParameterOffsets. */
+UPROPERTY()
+TArray<UNiagaraDataInterface*> DataInterfaces;
 
-	/** Bindings between this parameter store and others we push data into when we tick. */
-	TMap<FNiagaraParameterStore*, FNiagaraParameterStoreBinding> Bindings;
+/** UObjects referenced by this store. Also indexed by ParameterOffsets.*/
+UPROPERTY()
+TArray<UObject*> UObjects;
 
-	/** Parameter stores we've been bound to and are feeding data into us. */
-	TArray<FNiagaraParameterStore*> SourceStores;
+/** Bindings between this parameter store and others we push data into when we tick. */
+TMap<FNiagaraParameterStore*, FNiagaraParameterStoreBinding> Bindings;
 
-	/** Marks our parameters as dirty. They will be pushed to any bound stores on tick if true. */
-	uint32 bParametersDirty : 1;
-	/** Marks our interfaces as dirty. They will be pushed to any bound stores on tick if true. */
-	uint32 bInterfacesDirty : 1;
-	/** Marks our UObjects as dirty. They will be pushed to any bound stores on tick if true. */
-	uint32 bUObjectsDirty : 1;
+/** Parameter stores we've been bound to and are feeding data into us. */
+TArray<FNiagaraParameterStore*> SourceStores;
 
-	/** Uniquely identifies the current layout of this parameter store for detecting layout changes. */
-	uint32 LayoutVersion;
+/** Marks our parameters as dirty. They will be pushed to any bound stores on tick if true. */
+uint32 bParametersDirty : 1;
+/** Marks our interfaces as dirty. They will be pushed to any bound stores on tick if true. */
+uint32 bInterfacesDirty : 1;
+/** Marks our UObjects as dirty. They will be pushed to any bound stores on tick if true. */
+uint32 bUObjectsDirty : 1;
+
+/** Uniquely identifies the current layout of this parameter store for detecting layout changes. */
+uint32 LayoutVersion;
 
 #if WITH_EDITOR
-	FOnChanged OnChangedDelegate;
+FOnChanged OnChangedDelegate;
 #endif
 
 public:
@@ -189,10 +189,10 @@ public:
 	FNiagaraParameterStore& operator=(const FNiagaraParameterStore& Other);
 
 	virtual ~FNiagaraParameterStore();
-	
+
 #if WITH_EDITORONLY_DATA
 	UPROPERTY()
-	FString DebugName;
+		FString DebugName;
 #endif
 
 	void SetOwner(UObject* InOwner);
@@ -223,7 +223,7 @@ public:
 	FORCEINLINE_DEBUGGABLE void Tick();
 	/** Unbinds this store from all stores it's being driven by. */
 	void UnbindFromSourceStores();
-	
+
 	bool VerifyBinding(const FNiagaraParameterStore* InDestStore) const;
 
 	void CheckForNaNs() const;
@@ -233,7 +233,18 @@ public:
 	Does nothing if this parameter is already present.
 	Returns true if we added a new parameter.
 	*/
-	virtual bool AddParameter(const FNiagaraVariable& Param, bool bInitialize=true, bool bTriggerRebind = true, int32* OutOffset = nullptr);
+	virtual bool AddParameter(const FNiagaraVariable& Param, bool bInitialize = true, bool bTriggerRebind = true, int32* OutOffset = nullptr);
+
+#if WITH_EDITORONLY_DATA
+	template<typename BufferType>
+	void AddConstantBuffer()
+	{
+		for (const FNiagaraVariable& BufferVariable : BufferType::GetVariables())
+		{
+			AddParameter(BufferVariable, true, false);
+		}
+	}
+#endif
 
 	/** Removes the passed parameter if it exists in the store. */
 	virtual bool RemoveParameter(const FNiagaraVariable& Param);
diff --git a/Engine/Plugins/FX/Niagara/Source/Niagara/Public/NiagaraScriptExecutionParameterStore.h b/Engine/Plugins/FX/Niagara/Source/Niagara/Public/NiagaraScriptExecutionParameterStore.h
--- a/Engine/Plugins/FX/Niagara/Source/Niagara/Public/NiagaraScriptExecutionParameterStore.h
+++ b/Engine/Plugins/FX/Niagara/Source/Niagara/Public/NiagaraScriptExecutionParameterStore.h
@@ -84,10 +84,10 @@ public:
 	}
 
 	// Just the external parameters, not previous or internal...
-	uint32 GetExternalParameterSize() { return ParameterSize; }
+	uint32 GetExternalParameterSize() const { return ParameterSize; }
 
 	// The entire buffer padded out by the required alignment of the types..
-	uint32 GetPaddedParameterSizeInBytes() { return PaddedParameterSize; }
+	uint32 GetPaddedParameterSizeInBytes() const { return PaddedParameterSize; }
 
 	// Helper that converts the data from the base type array internally into the padded out renderer-ready format.
 	void CopyParameterDataToPaddedBuffer(uint8* InTargetBuffer, uint32 InTargetBufferSizeInBytes);
@@ -97,6 +97,7 @@ public:
 	void SetAsInitialized() { bInitialized = true; }
 protected:
 	void AddPaddedParamSize(const FNiagaraTypeDefinition& InParamType, uint32 InOffset);
+	void AddAlignmentPadding();
 
 private:
 
diff --git a/Engine/Plugins/FX/Niagara/Source/Niagara/Public/NiagaraSystemInstance.h b/Engine/Plugins/FX/Niagara/Source/Niagara/Public/NiagaraSystemInstance.h
--- a/Engine/Plugins/FX/Niagara/Source/Niagara/Public/NiagaraSystemInstance.h
+++ b/Engine/Plugins/FX/Niagara/Source/Niagara/Public/NiagaraSystemInstance.h
@@ -77,6 +77,28 @@ public:
 	void UnbindParameters(bool bFromComplete = false);
 
 	FORCEINLINE FNiagaraParameterStore& GetInstanceParameters() { return InstanceParameters; }
+
+	FORCEINLINE uint32 GetParameterIndex(bool PreviousFrame = false) const
+	{
+		return (!!(PreviousFrame && ParametersValid) ^ !!CurrentFrameIndex) ? 1 : 0;
+	}
+
+	FORCEINLINE void FlipParameterBuffers()
+	{
+		CurrentFrameIndex = ~CurrentFrameIndex;
+		
+		// when we've hit both buffers, we'll mark the parameters as being valid
+		if (CurrentFrameIndex == 1)
+		{
+			ParametersValid = true;
+		}
+	}
+
+	FORCEINLINE const FNiagaraGlobalParameters& GetGlobalParameters(bool PreviousFrame = false) const { return GlobalParameters[GetParameterIndex(PreviousFrame)]; }
+	FORCEINLINE const FNiagaraSystemParameters& GetSystemParameters(bool PreviousFrame = false) const { return SystemParameters[GetParameterIndex(PreviousFrame)]; }
+	FORCEINLINE const FNiagaraOwnerParameters& GetOwnerParameters(bool PreviousFrame = false) const { return OwnerParameters[GetParameterIndex(PreviousFrame)]; }
+	FORCEINLINE const FNiagaraEmitterParameters& GetEmitterParameters(int32 EmitterIdx, bool PreviousFrame = false) const { return EmitterParameters[EmitterIdx * 2 + GetParameterIndex(PreviousFrame)]; }
+	FORCEINLINE FNiagaraEmitterParameters& EditEmitterParameters(int32 EmitterIdx) { return EmitterParameters[EmitterIdx * 2 + GetParameterIndex()]; }
 	
 	FNiagaraWorldManager* GetWorldManager()const;
 	bool RequiresDistanceFieldData() const;
@@ -123,15 +145,14 @@ public:
 	ENiagaraExecutionState GetActualExecutionState() { return ActualExecutionState; }
 	void SetActualExecutionState(ENiagaraExecutionState InState);
 
-	float GetSystemTimeSinceRendered() const { return SystemTimeSinceRenderedParam.GetValue(); }
+//	float GetSystemTimeSinceRendered() const { return SystemTimeSinceRenderedParam.GetValue(); }
 
-	float GetOwnerLODDistance() const { return GatheredInstanceParameters.LODDistance; }
-	float GetOwnerMaxLODDistance() const { return GatheredInstanceParameters.MaxLODDistance; }
+	//float GetOwnerLODDistance() const { return OwnerLODDistanceParam.GetValue(); }
 
-	int32 GetNumParticles(int32 EmitterIndex) const { return ParameterNumParticleBindings[EmitterIndex].GetValue(); }
-	float GetSpawnCountScale(int32 EmitterIndex) const { return ParameterSpawnCountScaleBindings[EmitterIndex].GetValue(); }
+	//int32 GetNumParticles(int32 EmitterIndex) const { return ParameterNumParticleBindings[EmitterIndex].GetValue(); }
+	//float GetSpawnCountScale(int32 EmitterIndex) const { return ParameterSpawnCountScaleBindings[EmitterIndex].GetValue(); }
 
-	FVector GetOwnerVelocity() const { return OwnerVelocityParam.GetValue(); }
+//	FVector GetOwnerVelocity() const { return OwnerVelocityParam.GetValue(); }
 
 	FORCEINLINE bool IsComplete()const { return ActualExecutionState == ENiagaraExecutionState::Complete || ActualExecutionState == ENiagaraExecutionState::Disabled; }
 	FORCEINLINE bool IsDisabled()const { return ActualExecutionState == ENiagaraExecutionState::Disabled; }
@@ -141,17 +162,14 @@ public:
 
 	UNiagaraSystem* GetSystem()const;
 	FORCEINLINE UNiagaraComponent *GetComponent() { return Component; }
-	FORCEINLINE TArray<TSharedRef<FNiagaraEmitterInstance, ESPMode::ThreadSafe> > &GetEmitters()	{ return Emitters; }
+	FORCEINLINE TArray<TSharedRef<FNiagaraEmitterInstance, ESPMode::ThreadSafe> > &GetEmitters() { return Emitters; }
+	FORCEINLINE const TArray<TSharedRef<FNiagaraEmitterInstance, ESPMode::ThreadSafe> >& GetEmitters() const { return Emitters; }
 	FORCEINLINE const FBox& GetLocalBounds() { return LocalBounds;  }
 
 	FNiagaraEmitterInstance* GetEmitterByID(FGuid InID);
 
 	FORCEINLINE bool IsSolo()const { return bSolo; }
 
-	//TEMPORARY. We wont have a single set of parameters when we're executing system scripts.
-	//System params will be pulled in from a data set.
-	FORCEINLINE FNiagaraParameterStore& GetParameters() { return InstanceParameters; }
-
 	/** Gets a data set either from another emitter or one owned by the System itself. */
 	FNiagaraDataSet* GetDataSet(FNiagaraDataSetID SetID, FName EmitterName = NAME_None);
 
@@ -270,6 +288,8 @@ public:
 
 	FORCEINLINE void SetLODDistance(float InLODDistance, float InMaxLODDistance);
 
+private:
+
 	void DestroyDataInterfaceInstanceData();
 
 	/** Builds the emitter simulations. */
@@ -280,6 +300,9 @@ public:
 	/** Resets for restart, assumes no change in emitter setup */
 	void ResetInternal(bool bResetSimulations);
 
+	/** Resets the parameter structrs */
+	void ResetParameters();
+
 	/** Call PrepareForSImulation on each data source from the simulations and determine which need per-tick updates.*/
 	void InitDataInterfaces();	
 	
@@ -341,41 +364,15 @@ public:
 	/** Per system instance parameters. These can be fed by the component and are placed into a dataset for execution for the system scripts. */
 	FNiagaraParameterStore InstanceParameters;
 	
-	FNiagaraParameterDirectBinding<FVector> OwnerPositionParam;
-	FNiagaraParameterDirectBinding<FVector> OwnerScaleParam;
-	FNiagaraParameterDirectBinding<FVector> OwnerVelocityParam;
-	FNiagaraParameterDirectBinding<FVector> OwnerXAxisParam;
-	FNiagaraParameterDirectBinding<FVector> OwnerYAxisParam;
-	FNiagaraParameterDirectBinding<FVector> OwnerZAxisParam;
-
-	FNiagaraParameterDirectBinding<FQuat> OwnerRotationParam;
-
-	FNiagaraParameterDirectBinding<FMatrix> OwnerTransformParam;
-	FNiagaraParameterDirectBinding<FMatrix> OwnerInverseParam;
-	FNiagaraParameterDirectBinding<FMatrix> OwnerTransposeParam;
-	FNiagaraParameterDirectBinding<FMatrix> OwnerInverseTransposeParam;
-	FNiagaraParameterDirectBinding<FMatrix> OwnerTransformNoScaleParam;
-	FNiagaraParameterDirectBinding<FMatrix> OwnerInverseNoScaleParam;
-
-	FNiagaraParameterDirectBinding<float> OwnerDeltaSecondsParam;
-	FNiagaraParameterDirectBinding<float> OwnerInverseDeltaSecondsParam;
-	FNiagaraParameterDirectBinding<float> OwnerEngineTimeParam;
-	FNiagaraParameterDirectBinding<float> OwnerEngineRealtimeParam;
-	FNiagaraParameterDirectBinding<float> SystemAgeParam;
-	FNiagaraParameterDirectBinding<int32> SystemTickCountParam;
-
-	FNiagaraParameterDirectBinding<float> OwnerLODDistanceParam;
-	FNiagaraParameterDirectBinding<float> OwnerLODDistanceFractionParam;
-	FNiagaraParameterDirectBinding<int32> SystemNumEmittersParam;
-	FNiagaraParameterDirectBinding<int32> SystemNumEmittersAliveParam;
-
-	FNiagaraParameterDirectBinding<float> SystemTimeSinceRenderedParam;
-
-	FNiagaraParameterDirectBinding<int32> OwnerExecutionStateParam;
-
-	TArray<FNiagaraParameterDirectBinding<float>> ParameterSpawnCountScaleBindings;
-	TArray<FNiagaraParameterDirectBinding<int32>> ParameterNumParticleBindings;
-	TArray<FNiagaraParameterDirectBinding<int32>> ParameterTotalSpawnedParticlesBindings;
+	static constexpr int32 ParameterBufferCount = 2;
+	FNiagaraGlobalParameters GlobalParameters[ParameterBufferCount];
+	FNiagaraSystemParameters SystemParameters[ParameterBufferCount];
+	FNiagaraOwnerParameters OwnerParameters[ParameterBufferCount];
+	TArray<FNiagaraEmitterParameters> EmitterParameters;
+
+	/** Used for double buffered global/system/emitter parameters */
+	uint32 CurrentFrameIndex : 1;
+	uint32 ParametersValid : 1;
 
 	// registered events for each of the emitters
 	typedef TPair<FName, FName> EmitterEventKey;
@@ -430,39 +427,36 @@ public:
 
 public:
 	// Transient data that is accumulated during tick.
-	uint32 TotalParamSize = 0;
+	uint32 TotalGPUParamSize = 0;
 	uint32 ActiveGPUEmitterCount = 0;
 	int32 GPUDataInterfaceInstanceDataSize = 0;
 
 	struct FInstanceParameters
 	{
-		float DeltaSeconds;
-
-		FTransform ComponentTrans;
-		FVector OldPos;
+		FTransform ComponentTrans = FTransform::Identity;
 
-		float LODDistance;
-		float MaxLODDistance;
-		float TimeSeconds;
-		float RealTimeSeconds;
+		float DeltaSeconds = 0.0f;
+		float TimeSeconds = 0.0f;
+		float RealTimeSeconds = 0.0f;
 
-		float Age;
-		int32 TickCount;
+		int32 EmitterCount = 0;
+		int32 NumAlive = 0;
+		int32 TransformMatchCount = 0;
 
-		TArray<int32> EmitterNumParticles;
-		TArray<int32> EmitterTotalSpawnedParticles;
-		TArray<float> EmitterSpawnCountScale;
-		int32 NumAlive;
-
-		float SafeTimeSinceRendererd;
-
-		ENiagaraExecutionState RequestedExecutionState;
+		ENiagaraExecutionState RequestedExecutionState = ENiagaraExecutionState::Active;
 
 		void Init(int32 NumEmitters)
 		{
-			EmitterNumParticles.AddUninitialized(NumEmitters);
-			EmitterTotalSpawnedParticles.AddUninitialized(NumEmitters);
-			EmitterSpawnCountScale.AddUninitialized(NumEmitters);
+			ComponentTrans = FTransform::Identity;
+			DeltaSeconds = 0.0f;
+			TimeSeconds = 0.0f;
+			RealTimeSeconds = 0.0f;
+
+			EmitterCount = 0;
+			NumAlive = 0;
+			TransformMatchCount = 0;
+
+			RequestedExecutionState = ENiagaraExecutionState::Active;
 		}
 	};
 
diff --git a/Engine/Plugins/FX/Niagara/Source/Niagara/Public/NiagaraSystemSimulation.h b/Engine/Plugins/FX/Niagara/Source/Niagara/Public/NiagaraSystemSimulation.h
--- a/Engine/Plugins/FX/Niagara/Source/Niagara/Public/NiagaraSystemSimulation.h
+++ b/Engine/Plugins/FX/Niagara/Source/Niagara/Public/NiagaraSystemSimulation.h
@@ -7,6 +7,7 @@
 #include "UObject/GCObject.h"
 #include "NiagaraDataSet.h"
 #include "NiagaraScriptExecutionContext.h"
+#include "NiagaraSystem.h"
 
 class UNiagaraEffectType;
 class UWorld;
@@ -96,7 +97,7 @@ struct FNiagaraParameterStoreToDataSetBinding
 		ParameterStore.OnParameterChange();
 	}
 
-	FORCEINLINE_DEBUGGABLE void ParameterStoreToDataSet(FNiagaraParameterStore& ParameterStore, FNiagaraDataSet& DataSet, int32 DataSetInstanceIndex)
+	FORCEINLINE_DEBUGGABLE void ParameterStoreToDataSet(const FNiagaraParameterStore& ParameterStore, FNiagaraDataSet& DataSet, int32 DataSetInstanceIndex)
 	{
 		FNiagaraDataBuffer& CurrBuffer = DataSet.GetDestinationDataChecked();
 		const uint8* ParameterData = ParameterStore.GetParameterDataArray().GetData();
@@ -124,6 +125,29 @@ struct FNiagaraParameterStoreToDataSetBinding
 	}
 };
 
+#if 1
+struct FNiagaraConstantBufferToDataSetBinding
+{
+	void Init(const FNiagaraSystemCompiledData& CompiledData);
+
+	void CopyToDataSets(const FNiagaraSystemInstance& SystemInstance, FNiagaraDataSet& SpawnDataSet, FNiagaraDataSet& UpdateDataSet, int32 DataSetInstanceIndex) const;
+	
+
+protected:
+	void ApplyOffsets(const FNiagaraParameterDataSetBindingCollection& Offsets, const uint8* SourceData, FNiagaraDataSet& DataSet, int32 DataSetInstanceIndex) const;
+
+	FNiagaraParameterDataSetBindingCollection SpawnInstanceGlobalBinding;
+	FNiagaraParameterDataSetBindingCollection SpawnInstanceSystemBinding;
+	FNiagaraParameterDataSetBindingCollection SpawnInstanceOwnerBinding;
+	TArray<FNiagaraParameterDataSetBindingCollection> SpawnInstanceEmitterBindings;
+
+	FNiagaraParameterDataSetBindingCollection UpdateInstanceGlobalBinding;
+	FNiagaraParameterDataSetBindingCollection UpdateInstanceSystemBinding;
+	FNiagaraParameterDataSetBindingCollection UpdateInstanceOwnerBinding;
+	TArray<FNiagaraParameterDataSetBindingCollection> UpdateInstanceEmitterBindings;
+};
+#endif
+
 struct FNiagaraSystemSimulationTickContext
 {
 	FNiagaraSystemSimulationTickContext(class FNiagaraSystemSimulation* Owner, TArray<FNiagaraSystemInstance*>& Instances, FNiagaraDataSet& DataSet, float DeltaSeconds, int32 SpawnNum, int EffectsQuality, const FGraphEventRef& MyCompletionGraphEvent);
@@ -222,6 +246,11 @@ protected:
 	void UpdateSystemInstances(FNiagaraSystemSimulationTickContext& Context);
 	/** Transfers the results of the system simulation into the emitter instances. */
 	void TransferSystemSimResults(FNiagaraSystemSimulationTickContext& Context);
+	/** Builds the constant buffer table for a given script execution */
+	void BuildConstantBufferTable(
+		const FNiagaraGlobalParameters& GlobalParameters,
+		FNiagaraScriptExecutionContext& ExecContext,
+		FScriptExecutionConstantBufferTable& ConstantBufferTable) const;
 
 	/** Should we push the system sim tick off the game thread. */
 	FORCEINLINE bool ShouldTickAsync(const FNiagaraSystemSimulationTickContext& Context);
@@ -266,6 +295,8 @@ protected:
 	/** Bindings that pull per component parameters into the update parameter dataset. */
 	FNiagaraParameterStoreToDataSetBinding UpdateInstanceParameterToDataSetBinding;
 
+	FNiagaraConstantBufferToDataSetBinding ConstantBufferToDataSetBinding;
+
 	/** Binding to push system attributes into each emitter spawn parameters. */
 	TArray<FNiagaraParameterStoreToDataSetBinding> DataSetToEmitterSpawnParameters;
 	/** Binding to push system attributes into each emitter update parameters. */
@@ -274,15 +305,6 @@ protected:
 	TArray<TArray<FNiagaraParameterStoreToDataSetBinding>> DataSetToEmitterEventParameters;
 
 	/** Direct bindings for Engine variables in System Spawn and Update scripts. */
-	FNiagaraParameterDirectBinding<float> SpawnTimeParam;
-	FNiagaraParameterDirectBinding<float> UpdateTimeParam;
-
-	FNiagaraParameterDirectBinding<float> SpawnDeltaTimeParam;
-	FNiagaraParameterDirectBinding<float> UpdateDeltaTimeParam;
-
-	FNiagaraParameterDirectBinding<float> SpawnInvDeltaTimeParam;
-	FNiagaraParameterDirectBinding<float> UpdateInvDeltaTimeParam;
-	
 	FNiagaraParameterDirectBinding<int32> SpawnNumSystemInstancesParam;
 	FNiagaraParameterDirectBinding<int32> UpdateNumSystemInstancesParam;
 
diff --git a/Engine/Plugins/FX/Niagara/Source/Niagara/Public/NiagaraTypes.h b/Engine/Plugins/FX/Niagara/Source/Niagara/Public/NiagaraTypes.h
--- a/Engine/Plugins/FX/Niagara/Source/Niagara/Public/NiagaraTypes.h
+++ b/Engine/Plugins/FX/Niagara/Source/Niagara/Public/NiagaraTypes.h
@@ -1195,4 +1195,79 @@ inline void FNiagaraVariable::SetValue<bool>(const bool& Data)
 	AllocateData();
 	FNiagaraBool* BoolStruct = (FNiagaraBool*)GetData();
 	BoolStruct->SetValue(Data);
-}
\ No newline at end of file
+}
+
+// Any change to this structure, or it's GetVariables implementation will require a bump in the CustomNiagaraVersion so that we
+// properly rebuild the scripts
+struct alignas(16) FNiagaraGlobalParameters
+{
+#if WITH_EDITOR
+	NIAGARA_API static const TArray<FNiagaraVariable>& GetVariables();
+#endif
+
+	float EngineDeltaTime =  0.0f;
+	float EngineInvDeltaTime = 0.0f;
+	float EngineTime = 0.0f;
+	float EngineRealTime = 0.0f;
+};
+
+// Any change to this structure, or it's GetVariables implementation will require a bump in the CustomNiagaraVersion so that we
+// properly rebuild the scripts
+struct alignas(16) FNiagaraSystemParameters
+{
+#if WITH_EDITOR
+	NIAGARA_API static const TArray<FNiagaraVariable>& GetVariables();
+#endif
+
+	float EngineTimeSinceRendered = 0.0f;
+	float EngineLodDistance = 0.0f;
+	float EngineLodDistanceFraction = 0.0f;
+	float EngineSystemAge = 0.0f;
+	uint32 EngineExecutionState = 0;
+	int32 EngineTickCount = 0;
+	int32 EngineEmitterCount = 0;
+	int32 EngineAliveEmitterCount = 0;
+};
+
+// Any change to this structure, or it's GetVariables implementation will require a bump in the CustomNiagaraVersion so that we
+// properly rebuild the scripts
+struct alignas(16) FNiagaraOwnerParameters
+{
+#if WITH_EDITOR
+	NIAGARA_API static const TArray<FNiagaraVariable>& GetVariables();
+#endif
+
+	FMatrix EngineLocalToWorld = FMatrix::Identity;
+	FMatrix EngineWorldToLocal = FMatrix::Identity;
+	FMatrix EngineLocalToWorldTransposed = FMatrix::Identity;
+	FMatrix EngineWorldToLocalTransposed = FMatrix::Identity;
+	FMatrix EngineLocalToWorldNoScale = FMatrix::Identity;
+	FMatrix EngineWorldToLocalNoScale = FMatrix::Identity;
+	FQuat EngineRotation = FQuat::Identity;
+	FVector4 EnginePosition = FVector4(EForceInit::ForceInitToZero);
+	FVector4 EngineVelocity = FVector4(EForceInit::ForceInitToZero);
+	FVector4 EngineXAxis = FVector4(1.0f, 0.0f, 0.0f, 0.0f);
+	FVector4 EngineYAxis = FVector4(0.0f, 1.0f, 0.0f, 0.0f);
+	FVector4 EngineZAxis = FVector4(0.0f, 0.0f, 1.0f, 0.0f);
+	FVector4 EngineScale = FVector4(1.0f, 1.0f, 1.0f, 0.0f);
+};
+
+// Any change to this structure, or it's GetVariables implementation will require a bump in the CustomNiagaraVersion so that we
+// properly rebuild the scripts
+struct alignas(16) FNiagaraEmitterParameters
+{
+#if WITH_EDITOR
+	NIAGARA_API static const TArray<FNiagaraVariable>& GetVariables();
+#endif
+
+	int32 EmitterNumParticles = 0;
+	int32 EmitterTotalSpawnedParticles = 0;
+	float EmitterSpawnCountScale = 1.0f;
+	float EmitterAge = 0.0f;
+	int32 EmitterRandomSeed = 0;
+
+	// todo - what else should be inserted here?  we could put an array of spawninfos/interp spawn values
+	int32 _Pad0;
+	int32 _Pad1;
+	int32 _Pad2;
+};
\ No newline at end of file
diff --git a/Engine/Plugins/FX/Niagara/Source/NiagaraCore/Private/NiagaraCustomVersion.cpp b/Engine/Plugins/FX/Niagara/Source/NiagaraCore/Private/NiagaraCustomVersion.cpp
--- a/Engine/Plugins/FX/Niagara/Source/NiagaraCore/Private/NiagaraCustomVersion.cpp
+++ b/Engine/Plugins/FX/Niagara/Source/NiagaraCore/Private/NiagaraCustomVersion.cpp
@@ -8,4 +8,4 @@ const FGuid FNiagaraCustomVersion::GUID(0xFCF57AFA, 0x50764283, 0xB9A9E658, 0xFF
 // Register the custom version with core
 FCustomVersionRegistration GRegisterNiagaraCustomVersion(FNiagaraCustomVersion::GUID, FNiagaraCustomVersion::LatestVersion, TEXT("NiagaraVer"));
 
-const FGuid FNiagaraCustomVersion::LatestScriptCompileVersion(0xC1786B5B, 0xFAF74AD2, 0x92CF74E9, 0x65C2658F);
+const FGuid FNiagaraCustomVersion::LatestScriptCompileVersion(0x85227085, 0xDE6D4AC5, 0x917C9745, 0x9E948906);
\ No newline at end of file
diff --git a/Engine/Plugins/FX/Niagara/Source/NiagaraEditor/Private/NiagaraHlslTranslator.cpp b/Engine/Plugins/FX/Niagara/Source/NiagaraEditor/Private/NiagaraHlslTranslator.cpp
--- a/Engine/Plugins/FX/Niagara/Source/NiagaraEditor/Private/NiagaraHlslTranslator.cpp
+++ b/Engine/Plugins/FX/Niagara/Source/NiagaraEditor/Private/NiagaraHlslTranslator.cpp
@@ -507,10 +507,9 @@ FString FHlslNiagaraTranslator::BuildParameterMapHlslDefinitions(TArray<FNiagara
 	TArray<FNiagaraVariable> UniqueVariables;
 
 	// Add in currently defined system vars.
-	TArray<FNiagaraVariable> ValueArray;
-	ParamMapDefinedSystemToNamespaceVars.GenerateValueArray(ValueArray);
-	for (FNiagaraVariable& Var : ValueArray)
+	for (const auto& SystemVarPair : ParamMapDefinedSystemVars)
 	{
+		const auto& Var = SystemVarPair.Value.Variable;
 		if (Var.GetType().GetClass() != nullptr)
 		{
 			continue;
@@ -519,6 +518,7 @@ FString FHlslNiagaraTranslator::BuildParameterMapHlslDefinitions(TArray<FNiagara
 	}
 
 	// Add in currently defined emitter vars.
+	TArray<FNiagaraVariable> ValueArray;
 	ParamMapDefinedEmitterParameterToNamespaceVars.GenerateValueArray(ValueArray);
 	for (FNiagaraVariable& Var : ValueArray)
 	{
@@ -780,6 +780,16 @@ bool FHlslNiagaraTranslator::IsVariableInUniformBuffer(const FNiagaraVariable& V
 	return true;
 }
 
+template<typename T>
+void FHlslNiagaraTranslator::BuildConstantBuffer(ENiagaraCodeChunkMode ChunkMode)
+{
+	for (const FNiagaraVariable& Variable : T::GetVariables())
+	{
+		const FString SymbolName = GetSanitizedSymbolName(Variable.GetName().ToString(), true);
+		AddUniformChunk(SymbolName, Variable, ChunkMode, true);
+	}
+}
+
 const FNiagaraTranslateResults &FHlslNiagaraTranslator::Translate(const FNiagaraCompileRequestData* InCompileData, const FNiagaraCompileOptions& InCompileOptions, FHlslNiagaraTranslatorOptions InTranslateOptions)
 {
 	SCOPE_CYCLE_COUNTER(STAT_NiagaraEditor_HlslTranslator_Translate);
@@ -989,6 +999,16 @@ const FNiagaraTranslateResults &FHlslNiagaraTranslator::Translate(const FNiagara
 		ValidateParticleIDUsage();
 	}
 
+	BuildConstantBuffer<FNiagaraGlobalParameters>(ENiagaraCodeChunkMode::GlobalConstant);
+	// only use the SystemConstantBuffer if we are doing particle scripts (for system scripts the data should come from
+	// the datasets)
+	if (!IsBulkSystemScript())
+	{
+		BuildConstantBuffer<FNiagaraSystemParameters>(ENiagaraCodeChunkMode::SystemConstant);
+		BuildConstantBuffer<FNiagaraOwnerParameters>(ENiagaraCodeChunkMode::OwnerConstant);
+		BuildConstantBuffer<FNiagaraEmitterParameters>(ENiagaraCodeChunkMode::EmitterConstant);
+	}
+
 	//Create main scope pin cache.
 	PinToCodeChunks.AddDefaulted(1);
 
@@ -1121,34 +1141,37 @@ const FNiagaraTranslateResults &FHlslNiagaraTranslator::Translate(const FNiagara
 				Error(ErrorMessage, nullptr, nullptr);
 			}
 		}
-		//Declare parameters.
-		//TODO: Separate Cbuffer for Global, System and Emitter parameters.
-		{
-			HlslOutput += TEXT("cbuffer FEmitterParameters\n{\n");
 
-			for (int32 i = 0; i < ChunksByMode[(int32)ENiagaraCodeChunkMode::Uniform].Num(); ++i)
-			{
-				FNiagaraVariable BufferVariable(CodeChunks[ChunksByMode[(int32)ENiagaraCodeChunkMode::Uniform][i]].Type, FName(*CodeChunks[ChunksByMode[(int32)ENiagaraCodeChunkMode::Uniform][i]].SymbolName));
-				if ( IsVariableInUniformBuffer(BufferVariable) )
-				{
-					FString Chunk = GetCode(ChunksByMode[(int32)ENiagaraCodeChunkMode::Uniform][i]);
-					HlslOutput += TEXT("\t") + Chunk;
-				}
-			}
+		ENiagaraCodeChunkMode ChunkModes[] = { ENiagaraCodeChunkMode::GlobalConstant, ENiagaraCodeChunkMode::SystemConstant, ENiagaraCodeChunkMode::OwnerConstant, ENiagaraCodeChunkMode::EmitterConstant, ENiagaraCodeChunkMode::Uniform };
+		FString ConstantBufferNames[] = { TEXT("FNiagaraGlobalParameters"), TEXT("FNiagaraSystemParameters"), TEXT("FNiagaraOwnerParameters"), TEXT("FNiagaraEmitterParameters"), TEXT("FNiagaraExternalParameters") };
+
+		static_assert(UE_ARRAY_COUNT(ChunkModes) == UE_ARRAY_COUNT(ConstantBufferNames), "Mismatch between ChunkModes and ConstantBufferNames");
 
-			if (bInterpolateParams)
+		FString SymbolPrefix[] = { TEXT(""), INTERPOLATED_PARAMETER_PREFIX };
+
+		for (int32 PrevIt = 0; PrevIt < (bInterpolateParams ? 2 : 1); ++PrevIt)
+		{
+			for (int32 ChunkModeIt = 0; ChunkModeIt < UE_ARRAY_COUNT(ChunkModes); ++ChunkModeIt)
 			{
-				//Define the params from the previous frame after the main parameters.
-				for (int32 i = 0; i < ChunksByMode[(int32)ENiagaraCodeChunkMode::Uniform].Num(); ++i)
+				int32 ChunkMode = static_cast<int32>(ChunkModes[ChunkModeIt]);
+				const FString BufferName = SymbolPrefix[PrevIt] + ConstantBufferNames[ChunkModeIt];
+
+				HlslOutput += TEXT("cbuffer ") + BufferName + TEXT("\n{\n");
+
+				for (int32 i = 0; i < ChunksByMode[ChunkMode].Num(); ++i)
 				{
-					//Copy the chunk so we can fiddle it's symbol name.
-					FNiagaraCodeChunk Chunk = CodeChunks[ChunksByMode[(int32)ENiagaraCodeChunkMode::Uniform][i]];
-					Chunk.SymbolName = INTERPOLATED_PARAMETER_PREFIX + Chunk.SymbolName;
-					HlslOutput += TEXT("\t") + GetCode(Chunk);
+					FNiagaraVariable BufferVariable(CodeChunks[ChunksByMode[ChunkMode][i]].Type, FName(*CodeChunks[ChunksByMode[ChunkMode][i]].SymbolName));
+					if (IsVariableInUniformBuffer(BufferVariable))
+					{
+						FNiagaraCodeChunk Chunk = CodeChunks[ChunksByMode[ChunkMode][i]];
+						Chunk.SymbolName = SymbolPrefix[PrevIt] + Chunk.SymbolName;
+
+						HlslOutput += TEXT("\t") + GetCode(Chunk);
+					}
 				}
-			}
 
-			HlslOutput += TEXT("}\n\n");
+				HlslOutput += TEXT("}\n\n");
+			}
 		}
 
 		WriteDataSetStructDeclarations(DataSetReadInfo[0], true, HlslOutput);
@@ -1237,7 +1260,7 @@ const FNiagaraTranslateResults &FHlslNiagaraTranslator::Translate(const FNiagara
 			HlslOutput += TEXT("struct FSimulationContext") TEXT("\n{\n");
 
 			// We need to reserve a place in the simulation context for the base Parameter Map.
-			if (PrimaryDataSetOutputEntries.Num() != 0 || ParamMapDefinedSystemToNamespaceVars.Num() != 0 || ParamMapDefinedEmitterParameterToNamespaceVars.Num() != 0 || (ParamMapSetVariablesToChunks.Num() != 0 && ParamMapSetVariablesToChunks[0].Num() > 0))
+			if (PrimaryDataSetOutputEntries.Num() != 0 || ParamMapDefinedSystemVars.Num() != 0 || ParamMapDefinedEmitterParameterToNamespaceVars.Num() != 0 || (ParamMapSetVariablesToChunks.Num() != 0 && ParamMapSetVariablesToChunks[0].Num() > 0))
 			{
 				for (int32 i = 0; i < TranslationStages.Num(); i++)
 				{
@@ -1543,40 +1566,45 @@ void FHlslNiagaraTranslator::DefineInterpolatedParametersFunction(FString &HlslO
 			HlslOutputString += TEXT("\tContext.") + PrevMap + TEXT(".Interpolation.Emitter_InterpSpawnStartDt = Emitter_InterpSpawnStartDt;\n");
 			HlslOutputString += TEXT("\tContext.") + PrevMap + TEXT(".Interpolation.Emitter_SpawnGroup = Emitter_SpawnGroup;\n");
 
-			for (int32 UniformIdx = 0; UniformIdx < ChunksByMode[(int32)ENiagaraCodeChunkMode::Uniform].Num(); ++UniformIdx)
+			int32 ModesToInterpolate[] =
 			{
-				int32 ChunkIdx = ChunksByMode[(int32)ENiagaraCodeChunkMode::Uniform][UniformIdx];
-				if (UniformIdx != INDEX_NONE)
-				{
-					FNiagaraVariable* FoundNamespacedVar = nullptr;
-					const FName* FoundSystemKey = ParamMapDefinedSystemVarsToUniformChunks.FindKey(UniformIdx);
+				static_cast<int32>(ENiagaraCodeChunkMode::GlobalConstant),
+				static_cast<int32>(ENiagaraCodeChunkMode::SystemConstant),
+				static_cast<int32>(ENiagaraCodeChunkMode::OwnerConstant),
+				static_cast<int32>(ENiagaraCodeChunkMode::EmitterConstant),
+				static_cast<int32>(ENiagaraCodeChunkMode::Uniform)
+			};
 
-					// This uniform was either an emitter uniform parameter or a system uniform parameter. Search our maps to find out which one it was 
-					// so that we can properly deal with accessors.
-					if (FoundSystemKey != nullptr)
-					{
-						FoundNamespacedVar = ParamMapDefinedSystemToNamespaceVars.Find(*FoundSystemKey);
-					}
-					/*else
+			for (int32 ChunkMode : ModesToInterpolate)
+			{
+				for (int32 UniformIdx = 0; UniformIdx < ChunksByMode[ChunkMode].Num(); ++UniformIdx)
+				{
+					int32 ChunkIdx = ChunksByMode[ChunkMode][UniformIdx];
+					if (ChunkIdx != INDEX_NONE)
 					{
-						const FName* FoundEmitterParameterKey = ParamMapDefinedEmitterParameterVarsToUniformChunks.FindKey(UniformIdx);
-						if (FoundEmitterParameterKey != nullptr)
-						{
-							FoundNamespacedVar = ParamMapDefinedEmitterParameterToNamespaceVars.Find(*FoundEmitterParameterKey);
-						}
-					}*/
+						const FNiagaraVariable* FoundNamespacedVar = nullptr;
 
-					if (FoundNamespacedVar != nullptr)
-					{
-						FString FoundName = GetSanitizedSymbolName(FoundNamespacedVar->GetName().ToString());
-						FNiagaraCodeChunk& Chunk = CodeChunks[ChunkIdx];
-						if (ShouldInterpolateParameter(*FoundNamespacedVar))
+						for (const auto& SystemVarPair : ParamMapDefinedSystemVars)
 						{
-							HlslOutputString += TEXT("\tContext.") + PrevMap + TEXT(".") + FoundName + TEXT(" = lerp(") + INTERPOLATED_PARAMETER_PREFIX + Chunk.SymbolName + Chunk.ComponentMask + TEXT(", ") + Chunk.SymbolName + Chunk.ComponentMask + TEXT(", ") + TEXT("SpawnInterp);\n");
+							if (SystemVarPair.Value.ChunkIndex == ChunkIdx)
+							{
+								FoundNamespacedVar = &SystemVarPair.Value.Variable;
+								break;
+							}
 						}
-						else
+
+						if (FoundNamespacedVar != nullptr)
 						{
-							// For now, we do nothing for non-floating point variables..
+							FString FoundName = GetSanitizedSymbolName(FoundNamespacedVar->GetName().ToString());
+							FNiagaraCodeChunk& Chunk = CodeChunks[ChunkIdx];
+							if (ShouldInterpolateParameter(*FoundNamespacedVar))
+							{
+								HlslOutputString += TEXT("\tContext.") + PrevMap + TEXT(".") + FoundName + TEXT(" = lerp(") + INTERPOLATED_PARAMETER_PREFIX + Chunk.SymbolName + Chunk.ComponentMask + TEXT(", ") + Chunk.SymbolName + Chunk.ComponentMask + TEXT(", ") + TEXT("SpawnInterp);\n");
+							}
+							else
+							{
+								// For now, we do nothing for non-floating point variables..
+							}
 						}
 					}
 				}
@@ -2721,12 +2749,14 @@ FString FHlslNiagaraTranslator::GeneratedConstantString(FVector4 Constant)
 	return FString::Format(TEXT("float4({0}, {1}, {2}, {3})"), Args);
 }
 
-int32 FHlslNiagaraTranslator::AddUniformChunk(FString SymbolName, const FNiagaraTypeDefinition& Type)
+int32 FHlslNiagaraTranslator::AddUniformChunk(FString SymbolName, const FNiagaraVariable& InVariable, ENiagaraCodeChunkMode ChunkMode, bool AddPadding)
 {
+	const FNiagaraTypeDefinition& Type = InVariable.GetType();
+
 	int32 Ret = CodeChunks.IndexOfByPredicate(
 		[&](const FNiagaraCodeChunk& Chunk)
 	{
-		return Chunk.Mode == ENiagaraCodeChunkMode::Uniform && Chunk.SymbolName == SymbolName && Chunk.Type == Type;
+		return Chunk.Mode == ChunkMode && Chunk.SymbolName == SymbolName && Chunk.Type == Type;
 	}
 	);
 
@@ -2737,7 +2767,7 @@ int32 FHlslNiagaraTranslator::AddUniformChunk(FString SymbolName, const FNiagara
 		Chunk.SymbolName = GetSanitizedSymbolName(SymbolName);
 		Chunk.Type = Type;
 
-		if (UNiagaraScript::IsGPUScript(CompileOptions.TargetUsage))
+		if (AddPadding)
 		{
 			if (Type == FNiagaraTypeDefinition::GetVec2Def())
 			{
@@ -2751,9 +2781,14 @@ int32 FHlslNiagaraTranslator::AddUniformChunk(FString SymbolName, const FNiagara
 			}
 		}
 
-		Chunk.Mode = ENiagaraCodeChunkMode::Uniform;
+		Chunk.Mode = ChunkMode;
+
+		ChunksByMode[static_cast<int32>(ChunkMode)].Add(Ret);
 
-		ChunksByMode[(int32)ENiagaraCodeChunkMode::Uniform].Add(Ret);
+		auto& SystemVar = ParamMapDefinedSystemVars.Add(InVariable.GetName());
+		SystemVar.ChunkIndex = Ret;
+		SystemVar.ChunkMode = static_cast<int32>(ChunkMode);
+		SystemVar.Variable = InVariable;
 	}
 	return Ret;
 }
@@ -3713,10 +3748,9 @@ bool FHlslNiagaraTranslator::ParameterMapRegisterExternalConstantNamespaceVariab
 
 		if (!bIsPerInstanceBulkSystemParam)
 		{
-			int32 UniformIdx = 0;
 			int32 UniformChunk = 0;
 
-			if (false == ParamMapDefinedSystemVarsToUniformChunks.Contains(InVariable.GetName()))
+			if (false == ParamMapDefinedSystemVars.Contains(InVariable.GetName()))
 			{
 				FString SymbolNameDefined = FlattenedName;
 
@@ -3769,16 +3803,12 @@ bool FHlslNiagaraTranslator::ParameterMapRegisterExternalConstantNamespaceVariab
 					CompilationOutput.ScriptData.Parameters.SetOrAdd(InVariable);
 				}
 
-				UniformIdx = ChunksByMode[(int32)ENiagaraCodeChunkMode::Uniform].Num();
-
-				UniformChunk = AddUniformChunk(SymbolNameDefined, InVariable.GetType());
-				ParamMapDefinedSystemVarsToUniformChunks.Add(InVariable.GetName(), UniformIdx);
-				ParamMapDefinedSystemToNamespaceVars.Add(InVariable.GetName(), InVariable);
+				UniformChunk = AddUniformChunk(SymbolNameDefined, InVariable, ENiagaraCodeChunkMode::Uniform, UNiagaraScript::IsGPUScript(CompileOptions.TargetUsage));
 			}
 			else
 			{
-				UniformIdx = ParamMapDefinedSystemVarsToUniformChunks.FindChecked(InVariable.GetName());
-				UniformChunk = ChunksByMode[(int32)ENiagaraCodeChunkMode::Uniform][UniformIdx];
+				const auto& SystemVar = ParamMapDefinedSystemVars.FindChecked(InVariable.GetName());
+				UniformChunk = SystemVar.ChunkIndex;
 			}
 			static const auto UseShaderStagesCVar = IConsoleManager::Get().FindConsoleVariable(TEXT("fx.UseShaderStages"));
 			// Avoid overriding the register indices
@@ -5379,17 +5409,23 @@ void FHlslNiagaraTranslator::RegisterFunctionCall(ENiagaraScriptUsage ScriptUsag
 
 				//Now remove all chunks for the function again.			
 				//This is super hacky. Should move chunks etc into a proper scoping system.
-
+				const int32 UniformMode = static_cast<int32>(ENiagaraCodeChunkMode::Uniform);
+				const int32 FuncUniformCount = ChunksByMode[UniformMode].Num() - ChunkStartsByMode[UniformMode];
 				TArray<FNiagaraCodeChunk> FuncUniforms;
-				FuncUniforms.Reserve(1024);
+				TArray<int32> OriginalUniformChunkIndices;
+
+				FuncUniforms.Reserve(FuncUniformCount);
+				OriginalUniformChunkIndices.Reserve(FuncUniformCount);
+
 				for (int32 i = 0; i < (int32)ENiagaraCodeChunkMode::Num; ++i)
 				{
 					//Keep uniform chunks.
-					if (i == (int32)ENiagaraCodeChunkMode::Uniform)
+					if (i == UniformMode)
 					{
 						for (int32 ChunkIdx = ChunkStartsByMode[i]; ChunkIdx < ChunksByMode[i].Num(); ++ChunkIdx)
 						{
 							FuncUniforms.Add(CodeChunks[ChunksByMode[i][ChunkIdx]]);
+							OriginalUniformChunkIndices.Add(ChunksByMode[i][ChunkIdx]);
 						}
 					}
 
@@ -5398,9 +5434,22 @@ void FHlslNiagaraTranslator::RegisterFunctionCall(ENiagaraScriptUsage ScriptUsag
 				CodeChunks.RemoveAt(ChunkStart, CodeChunks.Num() - ChunkStart, false);
 
 				//Re-add the uniforms. Really this is horrible. Rework soon.
-				for (FNiagaraCodeChunk& Chunk : FuncUniforms)
+				for (int32 FuncUniformIt = 0; FuncUniformIt < FuncUniformCount; ++FuncUniformIt)
 				{
-					ChunksByMode[(int32)ENiagaraCodeChunkMode::Uniform].Add(CodeChunks.Add(Chunk));
+					const FNiagaraCodeChunk& Chunk = FuncUniforms[FuncUniformIt];
+					const int32 OriginalChunkIndex = OriginalUniformChunkIndices[FuncUniformIt];
+
+					const int32 NewChunkIndex = CodeChunks.Add(Chunk);
+					ChunksByMode[UniformMode].Add(NewChunkIndex);
+
+					for (auto& SystemVarPair : ParamMapDefinedSystemVars)
+					{
+						if ((SystemVarPair.Value.ChunkIndex == OriginalChunkIndex)
+							&& (SystemVarPair.Value.ChunkMode == UniformMode))
+						{
+							SystemVarPair.Value.ChunkIndex = NewChunkIndex;
+						}
+					}
 				}
 
 				// We don't support an empty function definition when calling a real function.
diff --git a/Engine/Plugins/FX/Niagara/Source/NiagaraEditor/Private/NiagaraHlslTranslator.h b/Engine/Plugins/FX/Niagara/Source/NiagaraEditor/Private/NiagaraHlslTranslator.h
--- a/Engine/Plugins/FX/Niagara/Source/NiagaraEditor/Private/NiagaraHlslTranslator.h
+++ b/Engine/Plugins/FX/Niagara/Source/NiagaraEditor/Private/NiagaraHlslTranslator.h
@@ -183,6 +183,10 @@ struct FCompiledPin
 
 enum class ENiagaraCodeChunkMode : uint8
 {
+	GlobalConstant,
+	SystemConstant,
+	OwnerConstant,
+	EmitterConstant,
 	Uniform,
 	Source,
 	Body,
@@ -315,6 +319,12 @@ public:
 		TArray<int32> CodeChunks;
 	};
 
+	struct FNiagaraConstantBuffer
+	{
+		TArray<FNiagaraVariable> Variables;
+		FString StructName;
+	};
+
 protected:
 	const FNiagaraCompileRequestData* CompileData;
 	FNiagaraCompileOptions CompileOptions;
@@ -411,7 +421,7 @@ protected:
 	int32 AddSourceChunk(FString SymbolName, const FNiagaraTypeDefinition& Type, bool bSanitize = true);
 
 	/** Add a chunk defining a uniform value. */
-	int32 AddUniformChunk(FString SymbolName, const FNiagaraTypeDefinition& Type);
+	int32 AddUniformChunk(FString SymbolName, const FNiagaraVariable& InVariable, ENiagaraCodeChunkMode ChunkMode, bool AddPadding);
 
 	/* Add a chunk that is written to the body of the shader code. */
 	int32 AddBodyChunk(FString SymbolName, FString Definition, const FNiagaraTypeDefinition& Type, TArray<int32>& SourceChunks, bool bDecl = true, bool bIsTerminated = true);
@@ -588,6 +598,9 @@ private:
 	bool IsSpawnScript() const;
 	bool RequiresInterpolation() const;
 
+	template<typename T>
+	void BuildConstantBuffer(ENiagaraCodeChunkMode ChunkMode);
+
 	/** Map of symbol names to count of times it's been used. Used for generating unique symbol names. */
 	TMap<FName, uint32> SymbolCounts;
 
@@ -611,8 +624,14 @@ private:
 	TArray<TSet<FName>> FunctionNodeStack;
 
 	// Synced to the System uniforms encountered for parameter maps thus far.
-	TMap<FName, int32> ParamMapDefinedSystemVarsToUniformChunks; // Map from the defined constants to the uniform chunk expressing them (i.e. have we encountered before in this graph?)
-	TMap<FName, FNiagaraVariable> ParamMapDefinedSystemToNamespaceVars; // Map from defined constants to the Namespaced variable expressing it.
+	struct UniformVariableInfo
+	{
+		FNiagaraVariable Variable;
+		int32 ChunkIndex;
+		int32 ChunkMode;
+	};
+
+	TMap<FName, UniformVariableInfo> ParamMapDefinedSystemVars; // Map from the defined constants to the uniform chunk expressing them (i.e. have we encountered before in this graph?)
 
 	// Synced to the EmitterParameter uniforms encountered for parameter maps thus far.
 	TMap<FName, int32> ParamMapDefinedEmitterParameterVarsToUniformChunks; // Map from the variable name exposed by the emitter as a parameter to the uniform chunk expressing it (i.e. have we encountered before in this graph?)
diff --git a/Engine/Plugins/FX/Niagara/Source/NiagaraShader/Private/NiagaraShader.cpp b/Engine/Plugins/FX/Niagara/Source/NiagaraShader/Private/NiagaraShader.cpp
--- a/Engine/Plugins/FX/Niagara/Source/NiagaraShader/Private/NiagaraShader.cpp
+++ b/Engine/Plugins/FX/Niagara/Source/NiagaraShader/Private/NiagaraShader.cpp
@@ -1287,7 +1287,6 @@ const FNiagaraShaderMap* FNiagaraShaderMap::GetShaderMapBeingCompiled(const FNia
 
 FNiagaraShader::FNiagaraShader(const FNiagaraShaderType::CompiledShaderInitializerType& Initializer)
 	: FShader(Initializer)
-	, CBufferLayout(TEXT("Niagara Compute Sim CBuffer"))
 	, DebugDescription(Initializer.DebugDescription)
 {
 	check(!DebugDescription.IsEmpty());
@@ -1335,9 +1334,20 @@ void FNiagaraShader::BindParams(const FShaderParameterMap &ParameterMap)
 
 	ComponentBufferSizeReadParam.Bind(ParameterMap, TEXT("ComponentBufferSizeRead"));
 	ComponentBufferSizeWriteParam.Bind(ParameterMap, TEXT("ComponentBufferSizeWrite"));
-	EmitterConstantBufferParam.Bind(ParameterMap, TEXT("FEmitterParameters"));
 	ViewUniformBufferParam.Bind(ParameterMap, TEXT("View"));
 
+	GlobalConstantBufferParam[0].Bind(ParameterMap, TEXT("FNiagaraGlobalParameters"));
+	SystemConstantBufferParam[0].Bind(ParameterMap, TEXT("FNiagaraSystemParameters"));
+	OwnerConstantBufferParam[0].Bind(ParameterMap, TEXT("FNiagaraOwnerParameters"));
+	EmitterConstantBufferParam[0].Bind(ParameterMap, TEXT("FNiagaraEmitterParameters"));
+	ExternalConstantBufferParam[0].Bind(ParameterMap, TEXT("FNiagaraExternalParameters"));
+
+	GlobalConstantBufferParam[1].Bind(ParameterMap, TEXT("PREV_FNiagaraGlobalParameters"));
+	SystemConstantBufferParam[1].Bind(ParameterMap, TEXT("PREV_FNiagaraSystemParameters"));
+	OwnerConstantBufferParam[1].Bind(ParameterMap, TEXT("PREV_FNiagaraOwnerParameters"));
+	EmitterConstantBufferParam[1].Bind(ParameterMap, TEXT("PREV_FNiagaraEmitterParameters"));
+	ExternalConstantBufferParam[1].Bind(ParameterMap, TEXT("PREV_FNiagaraExternalParameters"));
+
 	// params for event buffers
 	// this is horrendous; need to do this in a uniform buffer instead.
 	//
@@ -1416,7 +1426,18 @@ bool FNiagaraShader::Serialize(FArchive& Ar)
 
 	Ar << DataInterfaceParameters;
 
-	Ar << EmitterConstantBufferParam;
+	Ar << GlobalConstantBufferParam[0];
+	Ar << SystemConstantBufferParam[0];
+	Ar << OwnerConstantBufferParam[0];
+	Ar << EmitterConstantBufferParam[0];
+	Ar << ExternalConstantBufferParam[0];
+
+	Ar << GlobalConstantBufferParam[1];
+	Ar << SystemConstantBufferParam[1];
+	Ar << OwnerConstantBufferParam[1];
+	Ar << EmitterConstantBufferParam[1];
+	Ar << ExternalConstantBufferParam[1];
+
 	Ar << DataInterfaceUniformBufferParam;
 	Ar << NumEventsPerParticleParam;
 	Ar << NumParticlesPerEventParam;
diff --git a/Engine/Plugins/FX/Niagara/Source/NiagaraShader/Public/NiagaraShader.h b/Engine/Plugins/FX/Niagara/Source/NiagaraShader/Public/NiagaraShader.h
--- a/Engine/Plugins/FX/Niagara/Source/NiagaraShader/Public/NiagaraShader.h
+++ b/Engine/Plugins/FX/Niagara/Source/NiagaraShader/Public/NiagaraShader.h
@@ -1,4 +1,4 @@
-// Copyright Epic Games, Inc. All Rights Reserved.
+// Copyright 1998-2019 Epic Games, Inc. All Rights Reserved.
 
 /*=============================================================================
 	MaterialShader.h: Shader base classes
@@ -31,7 +31,6 @@ public:
 	static FName UniformBufferLayoutName;
 
 	FNiagaraShader()
-		: CBufferLayout(TEXT("Niagara Compute Sim CBuffer"))
 	{
 	}
 
@@ -99,7 +98,6 @@ public:
 	virtual bool Serialize(FArchive& Ar) override;
 	virtual uint32 GetAllocatedSize() const override;
 
-	FRHIUniformBufferLayout CBufferLayout;
 	FShaderResourceParameter FloatInputBufferParam;
 	FShaderResourceParameter IntInputBufferParam;
 	FRWShaderParameter FloatOutputBufferParam;
@@ -109,7 +107,11 @@ public:
 	FShaderParameter WriteInstanceCountOffsetParam;
 	FShaderResourceParameter FreeIDBufferParam;
 	FRWShaderParameter IDToIndexBufferParam;
-	FShaderUniformBufferParameter EmitterConstantBufferParam;
+	FShaderUniformBufferParameter GlobalConstantBufferParam[2];
+	FShaderUniformBufferParameter SystemConstantBufferParam[2];
+	FShaderUniformBufferParameter OwnerConstantBufferParam[2];
+	FShaderUniformBufferParameter EmitterConstantBufferParam[2];
+	FShaderUniformBufferParameter ExternalConstantBufferParam[2];
 	FShaderUniformBufferParameter DataInterfaceUniformBufferParam;
 	FShaderUniformBufferParameter ViewUniformBufferParam;
 	FShaderParameter SimStartParam;
diff --git a/Engine/Source/Runtime/VectorVM/Private/VectorVM.cpp b/Engine/Source/Runtime/VectorVM/Private/VectorVM.cpp
--- a/Engine/Source/Runtime/VectorVM/Private/VectorVM.cpp
+++ b/Engine/Source/Runtime/VectorVM/Private/VectorVM.cpp
@@ -24,8 +24,11 @@ DEFINE_LOG_CATEGORY_STATIC(LogVectorVM, All, All);
 //I don't expect us to ever be waiting long
 #define FREE_TABLE_LOCK_CONTENTION_WARN_THRESHOLD_MS (0.01)
 
-//#define VM_FORCEINLINE
+#if UE_BUILD_DEBUG
+#define VM_FORCEINLINE FORCENOINLINE
+#else
 #define VM_FORCEINLINE FORCEINLINE
+#endif
 
 #define OP_REGISTER (0)
 #define OP0_CONST (1 << 0)
@@ -228,7 +231,7 @@ struct FVectorVMCodeOptimizerContext
 		, OptimizedCode(InOptimizedCode)
 		, ExternalFunctionRegisterCounts(InExternalFunctionRegisterCounts)
 	{
-		BaseContext.PrepareForExec(0, nullptr, nullptr, nullptr, TArrayView<FDataSetMeta>(), 0, false);
+		BaseContext.PrepareForExec(0, 0, nullptr, nullptr, nullptr, nullptr, TArrayView<FDataSetMeta>(), 0, false);
 		BaseContext.PrepareForChunk(ByteCode, 0, 0);
 	}
 	FVectorVMCodeOptimizerContext(const FVectorVMCodeOptimizerContext&) = delete;
@@ -305,36 +308,51 @@ struct FConstantHandler : public FConstantHandlerBase
 	const T Constant;
 	FConstantHandler(FVectorVMContext& Context)
 		: FConstantHandlerBase(Context)
-		, Constant(*((T*)(Context.ConstantTable + ConstantIndex)))
-	{}
+		, Constant(*Context.GetConstant<T>(ConstantIndex))
+	{
+	}
 
-	FORCEINLINE const T& Get() { return Constant; }
+	FORCEINLINE const T& Get() const { return Constant; }
 	FORCEINLINE const T& GetAndAdvance() { return Constant; }
 };
 
 template<>
 struct FConstantHandler<VectorRegister> : public FConstantHandlerBase
 {
+	static VectorRegister LoadConstant(const FVectorVMContext& Context, uint16 ConstantIndex)
+	{
+		float ConstantValue = *Context.GetConstant<float>(ConstantIndex);
+
+		return MakeVectorRegister(ConstantValue, ConstantValue, ConstantValue, ConstantValue);
+	}
+
 	const VectorRegister Constant;
 	FConstantHandler(FVectorVMContext& Context)
 		: FConstantHandlerBase(Context)
-		, Constant(VectorLoadFloat1(&Context.ConstantTable[ConstantIndex]))
+		, Constant(LoadConstant(Context, ConstantIndex))
 	{}
 
-	FORCEINLINE const VectorRegister Get() { return Constant; }
+	FORCEINLINE const VectorRegister Get() const { return Constant; }
 	FORCEINLINE const VectorRegister GetAndAdvance() { return Constant; }
 };
 
 template<>
 struct FConstantHandler<VectorRegisterInt> : public FConstantHandlerBase
 {
+	static VectorRegisterInt LoadConstant(const FVectorVMContext& Context, uint16 ConstantIndex)
+	{
+		int32 ConstantValue = *Context.GetConstant<int32>(ConstantIndex);
+
+		return MakeVectorRegisterInt(ConstantValue, ConstantValue, ConstantValue, ConstantValue);
+	}
+
 	const VectorRegisterInt Constant;
 	FConstantHandler(FVectorVMContext& Context)
 		: FConstantHandlerBase(Context)
-		, Constant(VectorIntLoad1(&Context.ConstantTable[ConstantIndex]))
+		, Constant(LoadConstant(Context, ConstantIndex))
 	{}
 
-	FORCEINLINE const VectorRegisterInt Get() { return Constant; }
+	FORCEINLINE const VectorRegisterInt Get() const { return Constant; }
 	FORCEINLINE const VectorRegisterInt GetAndAdvance() { return Constant; }
 };
 
@@ -400,7 +418,9 @@ FVectorVMContext::FVectorVMContext()
 
 void FVectorVMContext::PrepareForExec(
 	int32 InNumTempRegisters,
-	const uint8* InConstantTable,
+	int32 InConstantTableCount,
+	const uint8* const* InConstantTable,
+	const int32* InConstantTableSizes,
 	FVMExternalFunction* InExternalFunctionTable,
 	void** InUserPtrTable,
 	TArrayView<FDataSetMeta> InDataSetMetaTable,
@@ -409,6 +429,8 @@ void FVectorVMContext::PrepareForExec(
 )
 {
 	NumTempRegisters = InNumTempRegisters;
+	ConstantTableCount = InConstantTableCount;
+	ConstantTableSizes = InConstantTableSizes;
 	ConstantTable = InConstantTable;
 	ExternalFunctionTable = InExternalFunctionTable;
 	UserPtrTable = InUserPtrTable;
@@ -2238,7 +2260,9 @@ void VectorVM::Exec(
 	uint8 const* ByteCode,
 	uint8 const* OptimizedByteCode,
 	int32 NumTempRegisters,
-	uint8 const* ConstantTable,
+	int32 ConstantTableCount,
+	const uint8* const* ConstantTable,
+	const int32* ConstantTableSizes,
 	TArrayView<FDataSetMeta> DataSetMetaTable,
 	FVMExternalFunction* ExternalFunctionTable,
 	void** UserPtrTable,
@@ -2263,7 +2287,7 @@ void VectorVM::Exec(
 		//SCOPE_CYCLE_COUNTER(STAT_VVMExecChunk);
 
 		FVectorVMContext& Context = FVectorVMContext::Get();
-		Context.PrepareForExec(NumTempRegisters, ConstantTable, ExternalFunctionTable, UserPtrTable, DataSetMetaTable, MaxInstances, bParallel);
+		Context.PrepareForExec(NumTempRegisters, ConstantTableCount, ConstantTable, ConstantTableSizes, ExternalFunctionTable, UserPtrTable, DataSetMetaTable, MaxInstances, bParallel);
 #if STATS
 		Context.SetStatScopes(&StatScopes);
 #endif
diff --git a/Engine/Source/Runtime/VectorVM/Public/VectorVM.h b/Engine/Source/Runtime/VectorVM/Public/VectorVM.h
--- a/Engine/Source/Runtime/VectorVM/Public/VectorVM.h
+++ b/Engine/Source/Runtime/VectorVM/Public/VectorVM.h
@@ -267,7 +267,10 @@ private:
 
 public:
 	/** Pointer to the constant table. */
-	uint8 const* RESTRICT ConstantTable;
+	const uint8* const* RESTRICT ConstantTable;
+	const int32* ConstantTableSizes;
+	int32 ConstantTableCount;
+
 	/** Num temp registers required by this script. */
 	int32 NumTempRegisters;
 
@@ -313,7 +316,9 @@ public:
 
 	void PrepareForExec(
 		int32 InNumTempRegisters,
-		const uint8* InConstantTable,
+		int32 ConstantTableCount,
+		const uint8* const* InConstantTables,
+		const int32* InConstantTableSizes,
 		FVMExternalFunction* InExternalFunctionTable,
 		void** InUserPtrTable,
 		TArrayView<FDataSetMeta> InDataSetMetaTable,
@@ -389,6 +394,28 @@ public:
 	{
 		return bIsParallelExecution;
 	}
+
+	template<typename T = uint8>
+	FORCEINLINE const T* GetConstant(int32 TableIndex, int32 TableOffset) const
+	{
+		check(TableIndex < ConstantTableCount);
+		return reinterpret_cast<const T*>(ConstantTable[TableIndex] + TableOffset);
+	}
+
+	template<typename T = uint8>
+	FORCEINLINE const T* GetConstant(int32 Offset) const
+	{
+		int32 TableIndex = 0;
+
+		while (Offset >= ConstantTableSizes[TableIndex])
+		{
+			Offset -= ConstantTableSizes[TableIndex];
+			++TableIndex;
+		}
+
+		check(Offset < ConstantTableSizes[TableIndex]);
+		return reinterpret_cast<const T*>(ConstantTable[TableIndex] + Offset);
+	}
 };
 
 namespace VectorVM
@@ -410,7 +437,9 @@ namespace VectorVM
 		uint8 const* ByteCode,
 		uint8 const* OptimizedByteCode,
 		int32 NumTempRegisters,
-		uint8 const* ConstantTable,
+		int32 ConstantTableCount,
+		const uint8* const* ConstantTable,
+		const int32* ConstantTableSizes,
 		TArrayView<FDataSetMeta> DataSetMetaTable,
 		FVMExternalFunction* ExternalFunctionTable,
 		void** UserPtrTable,
@@ -433,11 +462,17 @@ namespace VectorVM
 		int32 UserPtrIdx;
 		T* Ptr;
 		FUserPtrHandler(FVectorVMContext& Context)
-			: UserPtrIdx(*(int32*)(Context.ConstantTable + (Context.DecodeU16())))
-			, Ptr((T*)Context.UserPtrTable[UserPtrIdx])
 		{
+			const uint16 VariableOffset = Context.DecodeU16();
+			check(!(VariableOffset & VVM_EXT_FUNC_INPUT_LOC_BIT));
+
+			const uint16 ConstantTableOffset = VariableOffset & VVM_EXT_FUNC_INPUT_LOC_MASK;
+			UserPtrIdx = *Context.GetConstant<int32>(ConstantTableOffset);
 			check(UserPtrIdx != INDEX_NONE);
+
+			Ptr = reinterpret_cast<T*>(Context.UserPtrTable[UserPtrIdx]);
 		}
+
 		FORCEINLINE T* Get() { return Ptr; }
 		FORCEINLINE T* operator->() { return Ptr; }
 		FORCEINLINE operator T*() { return Ptr; }
@@ -450,7 +485,7 @@ namespace VectorVM
 	private:
 		/** Either byte offset into constant table or offset into register table deepening on VVM_INPUT_LOCATION_BIT */
 		int32 InputOffset;
-		T* RESTRICT InputPtr;
+		const T* RESTRICT InputPtr;
 		int32 AdvanceOffset;
 
 	public:
@@ -468,7 +503,9 @@ namespace VectorVM
 		void Init(FVectorVMContext& Context)
 		{
 			InputOffset = Context.DecodeU16();
-			InputPtr = IsConstant() ? (T*)(Context.ConstantTable + GetOffset()) : (T*)Context.GetTempRegister(GetOffset());
+
+			const int32 Offset = GetOffset();
+			InputPtr = IsConstant() ? Context.GetConstant<T>(Offset) : reinterpret_cast<T*>(Context.GetTempRegister(Offset));
 			AdvanceOffset = IsConstant() ? 0 : 1;
 		}
 
@@ -477,17 +514,17 @@ namespace VectorVM
 		FORCEINLINE int32 GetOffset()const { return InputOffset & VVM_EXT_FUNC_INPUT_LOC_MASK; }
 
 		FORCEINLINE const T Get() { return *InputPtr; }
-		FORCEINLINE T* GetDest() { return InputPtr; }
+		FORCEINLINE const T* GetDest() { return InputPtr; }
 		FORCEINLINE void Advance() { InputPtr += AdvanceOffset; }
 		FORCEINLINE const T GetAndAdvance()
 		{
-			T* Ret = InputPtr;
+			const T* Ret = InputPtr;
 			InputPtr += AdvanceOffset;
 			return *Ret;
 		}
-		FORCEINLINE T* GetDestAndAdvance()
+		FORCEINLINE const T* GetDestAndAdvance()
 		{
-			T* Ret = InputPtr;
+			const T* Ret = InputPtr;
 			InputPtr += AdvanceOffset;
 			return Ret;
 		}
@@ -543,7 +580,7 @@ namespace VectorVM
 		T Constant;
 		FExternalFuncConstHandler(FVectorVMContext& Context)
 			: ConstantIndex(Context.DecodeU16() & VVM_EXT_FUNC_INPUT_LOC_MASK)
-			, Constant(*((T*)(Context.ConstantTable + ConstantIndex)))
+			, Constant(*Context.GetConstant<T>(ConstantIndex))
 		{}
 		FORCEINLINE const T& Get() { return Constant; }
 		FORCEINLINE const T& GetAndAdvance() { return Constant; }
