Added Experimental Message field to modules to allow module creators to give a reason for a Niagara script being marked as experimental.

diff --git a/Engine/Plugins/FX/Niagara/Source/Niagara/Classes/NiagaraScript.h b/Engine/Plugins/FX/Niagara/Source/Niagara/Classes/NiagaraScript.h
--- a/Engine/Plugins/FX/Niagara/Source/Niagara/Classes/NiagaraScript.h
+++ b/Engine/Plugins/FX/Niagara/Source/Niagara/Classes/NiagaraScript.h
@@ -390,9 +390,14 @@ public:
 	UPROPERTY(EditAnywhere, Category = Script, meta = (EditCondition = "bDeprecated"))
 	UNiagaraScript* DeprecationRecommendation;
 
+	/** Is this script experimental and less supported? */
 	UPROPERTY(EditAnywhere, Category = Script)
 	uint32 bExperimental : 1;
 
+	/** The message to display when a function is marked experimental. */
+	UPROPERTY(EditAnywhere, Category = Script, meta = (EditCondition = "bExperimental", MultiLine = true))
+	FText ExperimentalMessage;
+
 	/* If this script is exposed to the library. */
 	UPROPERTY(AssetRegistrySearchable, EditAnywhere, Category = Script)
 	uint32 bExposeToLibrary : 1;
diff --git a/Engine/Plugins/FX/Niagara/Source/Niagara/Public/NiagaraCommon.h b/Engine/Plugins/FX/Niagara/Source/Niagara/Public/NiagaraCommon.h
--- a/Engine/Plugins/FX/Niagara/Source/Niagara/Public/NiagaraCommon.h
+++ b/Engine/Plugins/FX/Niagara/Source/Niagara/Public/NiagaraCommon.h
@@ -259,6 +259,13 @@ struct NIAGARA_API FNiagaraFunctionSignature
 	/** Is this function experimental? */
 	UPROPERTY()
 	uint32 bExperimental : 1;
+
+#if WITH_EDITORONLY_DATA
+	/** The message to display when a function is marked experimental. */
+	UPROPERTY(EditAnywhere, Category = Script, meta = (EditCondition = "bExperimental", MultiLine = true))
+	FText ExperimentalMessage;
+#endif
+
 	/** Support running on the CPU. */
 	UPROPERTY()
 	uint32 bSupportsCPU : 1;
@@ -272,7 +279,7 @@ struct NIAGARA_API FNiagaraFunctionSignature
 
 	/** Localized description of this node. Note that this is *not* used during the operator == below since it may vary from culture to culture.*/
 #if WITH_EDITORONLY_DATA
-	UPROPERTY(meta = (SkipForCompileHash = "true"))
+	UPROPERTY(meta = (SkipForCompileHash = true))
 	FText Description;
 #endif
 
diff --git a/Engine/Plugins/FX/Niagara/Source/NiagaraEditor/Private/NiagaraNodeFunctionCall.cpp b/Engine/Plugins/FX/Niagara/Source/NiagaraEditor/Private/NiagaraNodeFunctionCall.cpp
--- a/Engine/Plugins/FX/Niagara/Source/NiagaraEditor/Private/NiagaraNodeFunctionCall.cpp
+++ b/Engine/Plugins/FX/Niagara/Source/NiagaraEditor/Private/NiagaraNodeFunctionCall.cpp
@@ -624,7 +624,17 @@ void UNiagaraNodeFunctionCall::UpdateNodeErrorMessage()
 		{
 			UEdGraphNode::bHasCompilerMessage = true;
 			UEdGraphNode::ErrorType = EMessageSeverity::Info;
-			UEdGraphNode::NodeUpgradeMessage = LOCTEXT("FunctionExperimental", "This function is marked as experimental, use with care!");
+
+			if (FunctionScript->ExperimentalMessage.IsEmpty())
+			{
+				UEdGraphNode::NodeUpgradeMessage = LOCTEXT("FunctionExperimental", "This function is marked as experimental, use with care!");
+			}
+			else
+			{
+				FFormatNamedArguments Args;
+				Args.Add(TEXT("Message"), FunctionScript->ExperimentalMessage);
+				UEdGraphNode::NodeUpgradeMessage = FText::Format(LOCTEXT("FunctionExperimentalReason", "This function is marked as experimental, reason: {Message}."), Args);
+			}
 		}
 		else
 		{
@@ -638,7 +648,17 @@ void UNiagaraNodeFunctionCall::UpdateNodeErrorMessage()
 		{
 			UEdGraphNode::bHasCompilerMessage = true;
 			UEdGraphNode::ErrorType = EMessageSeverity::Info;
-			UEdGraphNode::NodeUpgradeMessage = LOCTEXT("FunctionExperimental", "This function is marked as experimental, use with care!");
+
+			if (FunctionScript->ExperimentalMessage.IsEmpty())
+			{
+				UEdGraphNode::NodeUpgradeMessage = LOCTEXT("FunctionExperimental", "This function is marked as experimental, use with care!");
+			}
+			else
+			{
+				FFormatNamedArguments Args;
+				Args.Add(TEXT("Message"), Signature.ExperimentalMessage);
+				UEdGraphNode::NodeUpgradeMessage = FText::Format(LOCTEXT("ModuleScriptExperimentalReason", "This function is marked as experimental, reason: {Message}"), Args);
+			}
 		}
 	}
 }
diff --git a/Engine/Plugins/FX/Niagara/Source/NiagaraEditor/Private/ViewModels/Stack/NiagaraStackErrorItem.cpp b/Engine/Plugins/FX/Niagara/Source/NiagaraEditor/Private/ViewModels/Stack/NiagaraStackErrorItem.cpp
--- a/Engine/Plugins/FX/Niagara/Source/NiagaraEditor/Private/ViewModels/Stack/NiagaraStackErrorItem.cpp
+++ b/Engine/Plugins/FX/Niagara/Source/NiagaraEditor/Private/ViewModels/Stack/NiagaraStackErrorItem.cpp
@@ -77,7 +77,7 @@ void UNiagaraStackErrorItem::RefreshChildrenInternal(const TArray<UNiagaraStackE
 			ErrorEntryFix->SetFixDelegate(CurrentFix.GetFixDelegate());
 		}
 		if (ensureMsgf(NewChildren.Contains(ErrorEntryFix) == false,
-			TEXT("Duplicate stack issue fix rows detected, This is caused by two different issue fixes with the same description which is used to generate their unique id. Issue Fix description: %s.  This issue fix will not be shown in the UI."),
+			TEXT("Duplicate stack issue fix rows detected. This is caused by two different issue fixes with the same description which is used to generate their unique ID. Issue Fix description: %s.  This issue fix will not be shown in the UI."),
 			*CurrentFix.GetDescription().ToString()))
 		{
 			ErrorEntryFix->OnIssueFixed().AddUObject(this, &UNiagaraStackErrorItem::IssueFixed);
@@ -172,7 +172,7 @@ void UNiagaraStackErrorItemDismiss::Initialize(FRequiredEntryData InRequiredEntr
 	UNiagaraStackEntry::Initialize(InRequiredEntryData, ErrorStackEditorDataKey);
 	StackIssue = InStackIssue;
 	IssueFix = FStackIssueFix(
-		LOCTEXT("DismissError", "Dismiss the issue without fixing (I know what I'm doing)."),
+		LOCTEXT("DismissError", "Dismiss the issue without fixing (I know what I'm doing)"),
 		FStackIssueFixDelegate::CreateUObject(this, &UNiagaraStackErrorItemDismiss::DismissIssue));
 }
 
diff --git a/Engine/Plugins/FX/Niagara/Source/NiagaraEditor/Private/ViewModels/Stack/NiagaraStackFunctionInput.cpp b/Engine/Plugins/FX/Niagara/Source/NiagaraEditor/Private/ViewModels/Stack/NiagaraStackFunctionInput.cpp
--- a/Engine/Plugins/FX/Niagara/Source/NiagaraEditor/Private/ViewModels/Stack/NiagaraStackFunctionInput.cpp
+++ b/Engine/Plugins/FX/Niagara/Source/NiagaraEditor/Private/ViewModels/Stack/NiagaraStackFunctionInput.cpp
@@ -419,8 +419,8 @@ void UNiagaraStackFunctionInput::RefreshChildrenInternal(const TArray<UNiagaraSt
 				if (InputValues.DynamicNode->FunctionScript->bDeprecated)
 				{
 					FText LongMessage = InputValues.DynamicNode->FunctionScript->DeprecationRecommendation != nullptr ?
-						FText::Format(LOCTEXT("DynamicInputScriptDeprecationLong", "The script asset for this dynamic input has been deprecated. Suggested replacement: {1}"), FText::FromString(InputValues.DynamicNode->GetName()), FText::FromString(InputValues.DynamicNode->FunctionScript->DeprecationRecommendation->GetPathName())) :
-						FText::Format(LOCTEXT("DynamicInputScriptDeprecationUnknownLong", "The script asset for the assigned module {0} has been deprecated."), FText::FromString(InputValues.DynamicNode->GetName()));
+						FText::Format(LOCTEXT("DynamicInputScriptDeprecationLong", "The script asset for the assigned dynamic input {0} has been deprecated. Suggested replacement: {1}"), FText::FromString(InputValues.DynamicNode->GetFunctionName()), FText::FromString(InputValues.DynamicNode->FunctionScript->DeprecationRecommendation->GetPathName())) :
+						FText::Format(LOCTEXT("DynamicInputScriptDeprecationUnknownLong", "The script asset for the assigned dynamic input {0} has been deprecated."), FText::FromString(InputValues.DynamicNode->GetFunctionName()));
 
 					int32 AddIdx = NewIssues.Add(FStackIssue(
 						EStackIssueSeverity::Warning,
@@ -448,10 +448,23 @@ void UNiagaraStackFunctionInput::RefreshChildrenInternal(const TArray<UNiagaraSt
 
 				if (InputValues.DynamicNode->FunctionScript->bExperimental)
 				{
+					FText ErrorMessage;
+					if (InputValues.DynamicNode->FunctionScript->ExperimentalMessage.IsEmpty())
+					{
+						ErrorMessage = FText::Format(LOCTEXT("DynamicInputScriptExperimental", "The script asset for the dynamic input {0} is experimental, use with care!"), FText::FromString(InputValues.DynamicNode->GetFunctionName()));
+					}
+					else
+					{
+						FFormatNamedArguments Args;
+						Args.Add(TEXT("Function"), FText::FromString(InputValues.DynamicNode->GetFunctionName()));
+						Args.Add(TEXT("Message"), InputValues.DynamicNode->FunctionScript->ExperimentalMessage);
+						ErrorMessage = FText::Format(LOCTEXT("DynamicInputScriptExperimentalReason", "The script asset for the dynamic input {Function} is experimental, reason: {Message}."), Args);
+					}
+
 					NewIssues.Add(FStackIssue(
 						EStackIssueSeverity::Info,
 						LOCTEXT("DynamicInputScriptExperimentalShort", "Experimental dynamic input"),
-						FText::Format(LOCTEXT("DynamicInputScriptExperimental", "The script asset for this dynamic input is experimental, use with care!"), FText::FromString(InputValues.DynamicNode->GetName())),
+						ErrorMessage,
 						GetStackEditorDataKey(),
 						true));
 				}
@@ -472,7 +485,7 @@ void UNiagaraStackFunctionInput::RefreshChildrenInternal(const TArray<UNiagaraSt
 						LOCTEXT("SelectNewDynamicInputScriptFix", "Select a new dynamic input script"),
 						FStackIssueFixDelegate::CreateLambda([this]() { this->bIsDynamicInputScriptReassignmentPending = true; })),
 					FStackIssueFix(
-						LOCTEXT("ResetFix", "Reset this input to it's default value"),
+						LOCTEXT("ResetFix", "Reset this input to its default value"),
 						FStackIssueFixDelegate::CreateLambda([this]() { this->Reset(); }))
 				}));
 		}
diff --git a/Engine/Plugins/FX/Niagara/Source/NiagaraEditor/Private/ViewModels/Stack/NiagaraStackModuleItem.cpp b/Engine/Plugins/FX/Niagara/Source/NiagaraEditor/Private/ViewModels/Stack/NiagaraStackModuleItem.cpp
--- a/Engine/Plugins/FX/Niagara/Source/NiagaraEditor/Private/ViewModels/Stack/NiagaraStackModuleItem.cpp
+++ b/Engine/Plugins/FX/Niagara/Source/NiagaraEditor/Private/ViewModels/Stack/NiagaraStackModuleItem.cpp
@@ -146,7 +146,7 @@ FText UNiagaraStackModuleItem::GetDisplayName() const
 	}
 	else
 	{
-		return FText::FromName(NAME_None);
+	return FText::FromName(NAME_None);
 	}
 }
 
@@ -342,10 +342,23 @@ void UNiagaraStackModuleItem::RefreshIssues(TArray<FStackIssue>& NewIssues)
 
 			if (FunctionCallNode->FunctionScript->bExperimental)
 			{
+				FText ErrorMessage;
+				if (FunctionCallNode->FunctionScript->ExperimentalMessage.IsEmpty())
+				{
+					ErrorMessage = FText::Format(LOCTEXT("ModuleScriptExperimental", "The script asset for this module is experimental, use with care!"), FText::FromString(FunctionCallNode->GetFunctionName()));
+				}
+				else
+				{
+					FFormatNamedArguments Args;
+					Args.Add(TEXT("Module"), FText::FromString(FunctionCallNode->GetFunctionName()));
+					Args.Add(TEXT("Message"), FunctionCallNode->FunctionScript->ExperimentalMessage);
+					ErrorMessage = FText::Format(LOCTEXT("ModuleScriptExperimentalReason", "The script asset for this module is marked as experimental, reason: {Message}."), Args);
+				}
+
 				NewIssues.Add(FStackIssue(
 					EStackIssueSeverity::Info,
 					LOCTEXT("ModuleScriptExperimentalShort", "Experimental module"),
-					FText::Format(LOCTEXT("ModuleScriptExperimental", "The script asset for the assigned module {0} is experimental, use with care!"), FText::FromString(FunctionCallNode->GetName())),
+					ErrorMessage,
 					GetStackEditorDataKey(),
 					true));
 			}
@@ -408,8 +421,8 @@ void UNiagaraStackModuleItem::RefreshIssues(TArray<FStackIssue>& NewIssues)
 			}));
 			FStackIssue InconsistentEnabledError(
 				EStackIssueSeverity::Error,
-				LOCTEXT("InconsistentEnabledErrorSummary", "The enabled state for module is inconsistent."),
-				LOCTEXT("InconsistentEnabledError", "This module is using multiple functions and their enabled state is inconsistent.\nClick fix to make all of the functions for this module enabled."),
+				LOCTEXT("InconsistentEnabledErrorSummary", "The enabled state for this module is inconsistent."),
+				LOCTEXT("InconsistentEnabledError", "This module is using multiple functions and their enabled states are inconsistent.\nClick \"Fix issue\" to make all of the functions for this module enabled."),
 				GetStackEditorDataKey(),
 				false,
 				EnableFix);
