Niagara Emitter name displayed in the Emitter view is now kept up to date when renaming the emitter from the content browser.

#rnx

diff --git a/Engine/Plugins/FX/Niagara/Source/Niagara/Classes/NiagaraEmitter.h b/Engine/Plugins/FX/Niagara/Source/Niagara/Classes/NiagaraEmitter.h
--- a/Engine/Plugins/FX/Niagara/Source/Niagara/Classes/NiagaraEmitter.h
+++ b/Engine/Plugins/FX/Niagara/Source/Niagara/Classes/NiagaraEmitter.h
@@ -244,6 +244,8 @@ public:
 	static UNiagaraEmitter* CreateAsDuplicate(const UNiagaraEmitter& InEmitterToDuplicate, FName InDuplicateName, UNiagaraSystem& InDuplicateOwnerSystem);
 
 	//Begin UObject Interface
+	virtual void PostRename(UObject* OldOuter, const FName OldName) override;
+	virtual void PostDuplicate(EDuplicateMode::Type DuplicateMode) override;
 	virtual void PostEditChangeProperty(struct FPropertyChangedEvent& PropertyChangedEvent) override;
 	NIAGARA_API FOnPropertiesChanged& OnPropertiesChanged();
 	NIAGARA_API FOnRenderersChanged& OnRenderersChanged();
@@ -251,7 +253,6 @@ public:
 	void Serialize(FArchive& Ar)override;
 	virtual void PostInitProperties() override;
 	virtual void PostLoad() override;
-	virtual void PostDuplicate(bool bDuplicateForPIE) override;
 	//End UObject Interface
 
 	/** Toggles whether or not the particles within this emitter are relative to the emitter origin or in global space. */ 
diff --git a/Engine/Plugins/FX/Niagara/Source/Niagara/Private/NiagaraEmitter.cpp b/Engine/Plugins/FX/Niagara/Source/Niagara/Private/NiagaraEmitter.cpp
--- a/Engine/Plugins/FX/Niagara/Source/Niagara/Private/NiagaraEmitter.cpp
+++ b/Engine/Plugins/FX/Niagara/Source/Niagara/Private/NiagaraEmitter.cpp
@@ -170,16 +170,6 @@ void UNiagaraEmitter::PostInitProperties()
 	UniqueEmitterName = TEXT("Emitter");
 }
 
-void UNiagaraEmitter::PostDuplicate(bool bDuplicateForPIE)
-{
-	Super::PostDuplicate(bDuplicateForPIE);
-
-	if (!bDuplicateForPIE)
-	{
-		SetUniqueEmitterName(GetName());
-	}
-}
-
 #if WITH_EDITORONLY_DATA
 bool UNiagaraEmitter::GetForceCompileOnLoad()
 {
@@ -284,7 +274,6 @@ UNiagaraEmitter* UNiagaraEmitter::DuplicateWithoutMerging(UObject* InOuter)
 	}
 	return Duplicate;
 }
-
 #endif
 
 void UNiagaraEmitter::Serialize(FArchive& Ar)
@@ -503,7 +492,7 @@ UNiagaraEmitter* UNiagaraEmitter::CreateWithParentAndOwner(UNiagaraEmitter& InPa
 	NewEmitter->Parent = &InParentEmitter;
 	NewEmitter->ParentAtLastMerge = Cast<UNiagaraEmitter>(StaticDuplicateObject(&InParentEmitter, NewEmitter));
 	NewEmitter->ParentAtLastMerge->ClearFlags(RF_Standalone | RF_Public);
-	NewEmitter->SetUniqueEmitterName(InName.ToString());
+	NewEmitter->SetUniqueEmitterName(InName.GetPlainNameString());
 	NewEmitter->GraphSource->MarkNotSynchronized(InitialNotSynchronizedReason);
 	return NewEmitter;
 }
@@ -519,12 +508,33 @@ UNiagaraEmitter* UNiagaraEmitter::CreateAsDuplicate(const UNiagaraEmitter& InEmi
 		NewEmitter->ParentAtLastMerge = Cast<UNiagaraEmitter>(StaticDuplicateObject(InEmitterToDuplicate.ParentAtLastMerge, NewEmitter));
 		NewEmitter->ParentAtLastMerge->ClearFlags(RF_Standalone | RF_Public);
 	}
-	NewEmitter->SetUniqueEmitterName(InDuplicateName.ToString());
+	NewEmitter->SetUniqueEmitterName(InDuplicateName.GetPlainNameString());
 	NewEmitter->GraphSource->MarkNotSynchronized(InitialNotSynchronizedReason);
 
 	return NewEmitter;
 }
 
+
+void UNiagaraEmitter::PostDuplicate(EDuplicateMode::Type DuplicateMode)
+{
+	Super::PostDuplicate(DuplicateMode);
+
+	if (IsAsset() && DuplicateMode == EDuplicateMode::Normal)
+	{
+		SetUniqueEmitterName(GetFName().GetPlainNameString());
+	}
+}
+
+void UNiagaraEmitter::PostRename(UObject* OldOuter, const FName OldName)
+{
+	Super::PostRename(OldOuter, OldName);
+
+	if (IsAsset())
+	{
+		SetUniqueEmitterName(GetFName().GetPlainNameString());
+	}
+}
+
 void UNiagaraEmitter::PostEditChangeProperty(struct FPropertyChangedEvent& PropertyChangedEvent)
 {
 	Super::PostEditChangeProperty(PropertyChangedEvent);
