Fix for copywrite text regression from CL 11082540

diff --git a/Engine/Plugins/FX/Niagara/Source/Niagara/Private/NiagaraEmitterInstanceBatcher.cpp b/Engine/Plugins/FX/Niagara/Source/Niagara/Private/NiagaraEmitterInstanceBatcher.cpp
--- a/Engine/Plugins/FX/Niagara/Source/Niagara/Private/NiagaraEmitterInstanceBatcher.cpp
+++ b/Engine/Plugins/FX/Niagara/Source/Niagara/Private/NiagaraEmitterInstanceBatcher.cpp
@@ -1,4 +1,4 @@
-// Copyright 1998-2019 Epic Games, Inc. All Rights Reserved.
+// Copyright Epic Games, Inc. All Rights Reserved.
 
 #include "NiagaraEmitterInstanceBatcher.h"
 #include "NiagaraScriptExecutionContext.h"
diff --git a/Engine/Plugins/FX/Niagara/Source/Niagara/Private/NiagaraSystemSimulation.cpp b/Engine/Plugins/FX/Niagara/Source/Niagara/Private/NiagaraSystemSimulation.cpp
--- a/Engine/Plugins/FX/Niagara/Source/Niagara/Private/NiagaraSystemSimulation.cpp
+++ b/Engine/Plugins/FX/Niagara/Source/Niagara/Private/NiagaraSystemSimulation.cpp
@@ -1,4 +1,4 @@
-// Copyright 1998-2019 Epic Games, Inc. All Rights Reserved.
+// Copyright Epic Games, Inc. All Rights Reserved.
 
 #include "NiagaraSystemSimulation.h"
 #include "NiagaraModule.h"
diff --git a/Engine/Plugins/FX/Niagara/Source/NiagaraShader/Public/NiagaraShader.h b/Engine/Plugins/FX/Niagara/Source/NiagaraShader/Public/NiagaraShader.h
--- a/Engine/Plugins/FX/Niagara/Source/NiagaraShader/Public/NiagaraShader.h
+++ b/Engine/Plugins/FX/Niagara/Source/NiagaraShader/Public/NiagaraShader.h
@@ -1,4 +1,4 @@
-// Copyright 1998-2019 Epic Games, Inc. All Rights Reserved.
+// Copyright Epic Games, Inc. All Rights Reserved.
 
 /*=============================================================================
 	MaterialShader.h: Shader base classes
