Don't use FName::ToString() when creating an emitter from a template, because it appends the numeric suffix.

Missed file from previous changelist.

#rnx

diff --git a/Engine/Plugins/FX/Niagara/Source/NiagaraEditor/Private/NiagaraEmitterFactoryNew.cpp b/Engine/Plugins/FX/Niagara/Source/NiagaraEditor/Private/NiagaraEmitterFactoryNew.cpp
--- a/Engine/Plugins/FX/Niagara/Source/NiagaraEditor/Private/NiagaraEmitterFactoryNew.cpp
+++ b/Engine/Plugins/FX/Niagara/Source/NiagaraEditor/Private/NiagaraEmitterFactoryNew.cpp
@@ -129,7 +129,7 @@ UObject* UNiagaraEmitterFactoryNew::FactoryCreateNew(UClass* Class, UObject* InP
 		else
 		{
 			NewEmitter = Cast<UNiagaraEmitter>(StaticDuplicateObject(EmitterToCopy, InParent, Name, Flags, Class));
-			NewEmitter->SetUniqueEmitterName(Name.ToString());
+			NewEmitter->SetUniqueEmitterName(Name.GetPlainNameString());
 		}
 
 		NewEmitter->bIsTemplateAsset = false;
