- Convert transform -> rotation once and remove scaling from matrix

#rnx

diff --git a/Engine/Plugins/FX/Niagara/Source/Niagara/Classes/NiagaraDataInterface.h b/Engine/Plugins/FX/Niagara/Source/Niagara/Classes/NiagaraDataInterface.h
--- a/Engine/Plugins/FX/Niagara/Source/Niagara/Classes/NiagaraDataInterface.h
+++ b/Engine/Plugins/FX/Niagara/Source/Niagara/Classes/NiagaraDataInterface.h
@@ -21,14 +21,14 @@ struct FNDITransformHandlerNoop
 {
 	FORCEINLINE void TransformPosition(FVector& V, const FMatrix& M) { }
 	FORCEINLINE void TransformVector(FVector& V, const FMatrix& M) { }
-	FORCEINLINE void TransformRotation(FQuat& Q, const FMatrix& M) { }
+	FORCEINLINE void TransformRotation(FQuat& Q1, const FQuat& Q2) { }
 };
 
 struct FNDITransformHandler
 {
 	FORCEINLINE void TransformPosition(FVector& P, const FMatrix& M) { P = M.TransformPosition(P); }
 	FORCEINLINE void TransformVector(FVector& V, const FMatrix& M) { V = M.TransformVector(V).GetUnsafeNormal3(); }
-	FORCEINLINE void TransformRotation(FQuat& Q, const FMatrix& M) { Q = M.ToQuat() * Q; }
+	FORCEINLINE void TransformRotation(FQuat& Q1, const FQuat& Q2) { Q1 = Q2 * Q1; }
 };
 
 //////////////////////////////////////////////////////////////////////////
diff --git a/Engine/Plugins/FX/Niagara/Source/Niagara/Private/NDISkeletalMesh_BoneSampling.cpp b/Engine/Plugins/FX/Niagara/Source/Niagara/Private/NDISkeletalMesh_BoneSampling.cpp
--- a/Engine/Plugins/FX/Niagara/Source/Niagara/Private/NDISkeletalMesh_BoneSampling.cpp
+++ b/Engine/Plugins/FX/Niagara/Source/Niagara/Private/NDISkeletalMesh_BoneSampling.cpp
@@ -445,12 +445,14 @@ void UNiagaraDataInterfaceSkeletalMesh::GetSkinnedBoneData(FVectorVMContext& Con
 	checkfSlow(InstData.Get(), TEXT("Skeletal Mesh Interface has invalid instance data. %s"), *GetPathName());
 	checkfSlow(InstData->Mesh, TEXT("Skeletal Mesh Interface has invalid mesh. %s"), *GetPathName());
 
+	FBoneSocketSkinnedDataOutputHandler Output(Context);
+
 	//TODO: Replace this by storing off FTransforms and doing a proper lerp to get a final transform.
 	//Also need to pull in a per particle interpolation factor.
-	const FMatrix& Transform = InstData->Transform;
-	const FMatrix& PrevTransform = InstData->PrevTransform;
-
-	FBoneSocketSkinnedDataOutputHandler Output(Context);
+	const FMatrix& InstanceTransform = InstData->Transform;
+	const FMatrix& PrevInstanceTransform = InstData->PrevTransform;
+	const FQuat InstanceRotation = Output.bNeedsRotation ? InstData->Transform.GetMatrixWithoutScale().ToQuat() : FQuat::Identity;
+	const FQuat PrevInstanceRotation = Output.bNeedsRotation ? InstData->Transform.GetMatrixWithoutScale().ToQuat() : FQuat::Identity;
 
 	FSkinWeightVertexBuffer* SkinWeightBuffer;
 	FSkeletalMeshLODRenderData& LODData = InstData->GetLODRenderDataAndSkinWeights(SkinWeightBuffer);
@@ -468,13 +470,6 @@ void UNiagaraDataInterfaceSkeletalMesh::GetSkinnedBoneData(FVectorVMContext& Con
 	const TArray<FTransform>& SpecificSocketCurrTransforms = InstData->GetSpecificSocketsCurrBuffer();
 	const TArray<FTransform>& SpecificSocketPrevTransforms = InstData->GetSpecificSocketsPrevBuffer();
 
-	FVector BonePos;
-	FVector BonePrev;
-
-	FVector Pos;
-	FVector Prev;
-	FVector Velocity;
-
 	for (int32 i = 0; i < Context.NumInstances; ++i)
 	{
 		const float Interp = bInterpolated::Value ? InterpParam.GetAndAdvance() : 1.0f;
@@ -485,17 +480,21 @@ void UNiagaraDataInterfaceSkeletalMesh::GetSkinnedBoneData(FVectorVMContext& Con
 		const bool bIsSocket = Bone > BoneMax;
 		const int32 Socket = Bone - SpecificSocketBoneOffset;
 
+		FVector Pos;
+		FVector Prev;
+		FVector Velocity;
+
 		// Handle edge cases first...
 		if ((!bIsSocket && Bone >= BoneMax) ||
 			(bIsSocket && (Socket >= SpecificSocketCurrTransforms.Num() || Socket < 0)))
 		{
 			Pos = FVector::ZeroVector;
-			TransformHandler.TransformPosition(Pos, Transform);
+			TransformHandler.TransformPosition(Pos, InstanceTransform);
 
 			if (Output.bNeedsVelocity || bInterpolated::Value)
 			{
 				Prev = FVector::ZeroVector;
-				TransformHandler.TransformPosition(Prev, PrevTransform);
+				TransformHandler.TransformPosition(Prev, PrevInstanceTransform);
 			}
 			if (Output.bNeedsRotation)
 			{
@@ -515,22 +514,22 @@ void UNiagaraDataInterfaceSkeletalMesh::GetSkinnedBoneData(FVectorVMContext& Con
 			FTransform PrevSocketTransform = SpecificSocketPrevTransforms[Socket];
 
 			Pos = CurrSocketTransform.GetLocation();
-			TransformHandler.TransformPosition(Pos, Transform);
+			TransformHandler.TransformPosition(Pos, InstanceTransform);
 
 			if (Output.bNeedsVelocity || bInterpolated::Value)
 			{
 				Prev = PrevSocketTransform.GetLocation();
-				TransformHandler.TransformPosition(Prev, PrevTransform);
+				TransformHandler.TransformPosition(Prev, PrevInstanceTransform);
 			}
 
 			if (Output.bNeedsRotation)
 			{
 				FQuat Rotation = CurrSocketTransform.GetRotation();
-				TransformHandler.TransformRotation(Rotation, Transform);
+				TransformHandler.TransformRotation(Rotation, InstanceRotation);
 				if (bInterpolated::Value)
 				{
 					FQuat PrevRotation = PrevSocketTransform.GetRotation();
-					TransformHandler.TransformRotation(PrevRotation, Transform);
+					TransformHandler.TransformRotation(PrevRotation, PrevInstanceRotation);
 					Rotation = FQuat::Slerp(PrevRotation, Rotation, Interp);
 				}
 
@@ -541,22 +540,22 @@ void UNiagaraDataInterfaceSkeletalMesh::GetSkinnedBoneData(FVectorVMContext& Con
 		else
 		{
 			Pos = SkinningHandler.GetSkinnedBonePosition(Accessor, Bone);
-			TransformHandler.TransformPosition(Pos, Transform);
+			TransformHandler.TransformPosition(Pos, InstanceTransform);
 
 			if (Output.bNeedsVelocity || bInterpolated::Value)
 			{
 				Prev = SkinningHandler.GetSkinnedBonePreviousPosition(Accessor, Bone);
-				TransformHandler.TransformPosition(Prev, PrevTransform);
+				TransformHandler.TransformPosition(Prev, PrevInstanceTransform);
 			}
 
 			if (Output.bNeedsRotation)
 			{
 				FQuat Rotation = SkinningHandler.GetSkinnedBoneRotation(Accessor, Bone);
-				TransformHandler.TransformRotation(Rotation, Transform);
+				TransformHandler.TransformRotation(Rotation, InstanceRotation);
 				if (bInterpolated::Value)
 				{
 					FQuat PrevRotation = SkinningHandler.GetSkinnedBonePreviousRotation(Accessor, Bone);
-					TransformHandler.TransformRotation(PrevRotation, Transform);
+					TransformHandler.TransformRotation(PrevRotation, PrevInstanceRotation);
 					Rotation = FQuat::Slerp(PrevRotation, Rotation, Interp);
 				}
 
