Added some runtime switches to control whether Niagara is allowed to use compute shaders and GPU particles:
- "fx.NiagaraAllowComputeShaders"
- "fx.NiagaraAllowGPUParticles"

diff --git a/Engine/Plugins/FX/Niagara/Source/Niagara/Private/NiagaraCommon.cpp b/Engine/Plugins/FX/Niagara/Source/Niagara/Private/NiagaraCommon.cpp
--- a/Engine/Plugins/FX/Niagara/Source/Niagara/Private/NiagaraCommon.cpp
+++ b/Engine/Plugins/FX/Niagara/Source/Niagara/Private/NiagaraCommon.cpp
@@ -16,6 +16,23 @@
 
 DECLARE_CYCLE_STAT(TEXT("Niagara - Utilities - PrepareRapidIterationParameters"), STAT_Niagara_Utilities_PrepareRapidIterationParameters, STATGROUP_Niagara);
 
+
+//////////////////////////////////////////////////////////////////////////
+
+int32 GNiagaraAllowComputeShaders = 1;
+FAutoConsoleVariableRef CVarAllowComputeShaders(
+	TEXT("fx.NiagaraAllowComputeShaders"),
+	GNiagaraAllowComputeShaders,
+	TEXT("If true, allow the usage compute shaders within Niagara."),
+	ECVF_Default);
+
+int32 GNiagaraAllowGPUParticles = 1;
+FAutoConsoleVariableRef CVarAllowGPUParticles(
+	TEXT("fx.NiagaraAllowGPUParticles"),
+	GNiagaraAllowGPUParticles,
+	TEXT("If true, allow the usage of GPU particles for Niagara."),
+	ECVF_Default);
+
 //////////////////////////////////////////////////////////////////////////
 
 FString FNiagaraTypeHelper::ToString(const uint8* ValueData, const UObject* StructOrEnum)
@@ -215,6 +232,16 @@ void FNiagaraSystemUpdateContext::AddInternal(UNiagaraComponent* Comp, bool bReI
 
 //////////////////////////////////////////////////////////////////////////
 
+bool FNiagaraUtilities::AllowGPUParticles(EShaderPlatform ShaderPlatform)
+{
+	return SupportsGPUParticles(ShaderPlatform) && GNiagaraAllowGPUParticles && GNiagaraAllowComputeShaders;
+}
+
+bool FNiagaraUtilities::AllowComputeShaders(EShaderPlatform ShaderPlatform)
+{
+	return RHISupportsComputeShaders(ShaderPlatform) && GNiagaraAllowComputeShaders;
+}
+
 FName NIAGARA_API FNiagaraUtilities::GetUniqueName(FName CandidateName, const TSet<FName>& ExistingNames)
 {
 	if (ExistingNames.Contains(CandidateName) == false)
diff --git a/Engine/Plugins/FX/Niagara/Source/Niagara/Private/NiagaraComponent.cpp b/Engine/Plugins/FX/Niagara/Source/Niagara/Private/NiagaraComponent.cpp
--- a/Engine/Plugins/FX/Niagara/Source/Niagara/Private/NiagaraComponent.cpp
+++ b/Engine/Plugins/FX/Niagara/Source/Niagara/Private/NiagaraComponent.cpp
@@ -371,7 +371,7 @@ void FNiagaraSceneProxy::GetDynamicMeshElements(const TArray<const FSceneView*>&
 	for (int32 RendererIdx : RendererDrawOrder)
 	{
 		FNiagaraRenderer* Renderer = EmitterRenderers[RendererIdx];
-		if (Renderer && (Renderer->GetSimTarget() == ENiagaraSimTarget::CPUSim || ViewFamily.GetFeatureLevel() >= ERHIFeatureLevel::ES3_1))
+		if (Renderer && (Renderer->GetSimTarget() != ENiagaraSimTarget::GPUComputeSim || FNiagaraUtilities::AllowGPUParticles(ViewFamily.GetShaderPlatform())))
 		{
 			Renderer->GetDynamicMeshElements(Views, ViewFamily, VisibilityMap, Collector, this);
 		}
diff --git a/Engine/Plugins/FX/Niagara/Source/Niagara/Private/NiagaraEmitterInstance.cpp b/Engine/Plugins/FX/Niagara/Source/Niagara/Private/NiagaraEmitterInstance.cpp
--- a/Engine/Plugins/FX/Niagara/Source/Niagara/Private/NiagaraEmitterInstance.cpp
+++ b/Engine/Plugins/FX/Niagara/Source/Niagara/Private/NiagaraEmitterInstance.cpp
@@ -203,14 +203,10 @@ bool FNiagaraEmitterInstance::IsAllowedToExecute() const
 {
 	int32 DetailLevel = ParentSystemInstance->GetDetailLevel();
 	const FNiagaraEmitterHandle& EmitterHandle = GetEmitterHandle();
-	if (!EmitterHandle.GetIsEnabled()
-		|| !CachedEmitter->IsAllowedByDetailLevel(DetailLevel)
-		|| (!FNiagaraUtilities::SupportsGPUParticles(GMaxRHIFeatureLevel) && CachedEmitter->SimTarget == ENiagaraSimTarget::GPUComputeSim)  // skip if GPU sim and <SM5. TODO: fall back to CPU sim instead once we have scalability functionality to do so
-		)
-	{
-		return false;
-	}
-	return true;
+	return EmitterHandle.GetIsEnabled() &&
+		CachedEmitter->IsAllowedByDetailLevel(DetailLevel) &&
+		// TODO: fall back to CPU sim instead once we have scalability functionality to do so
+		(CachedEmitter->SimTarget != ENiagaraSimTarget::GPUComputeSim || (Batcher && FNiagaraUtilities::AllowGPUParticles(Batcher->GetShaderPlatform())));
 }
 
 void FNiagaraEmitterInstance::Init(int32 InEmitterIdx, FNiagaraSystemInstanceID InSystemInstanceID)
diff --git a/Engine/Plugins/FX/Niagara/Source/Niagara/Private/NiagaraGPUInstanceCountManager.cpp b/Engine/Plugins/FX/Niagara/Source/Niagara/Private/NiagaraGPUInstanceCountManager.cpp
--- a/Engine/Plugins/FX/Niagara/Source/Niagara/Private/NiagaraGPUInstanceCountManager.cpp
+++ b/Engine/Plugins/FX/Niagara/Source/Niagara/Private/NiagaraGPUInstanceCountManager.cpp
@@ -196,7 +196,7 @@ void FNiagaraGPUInstanceCountManager::UpdateDrawIndirectBuffer(FRHICommandList&
 {
 	if (DrawIndirectArgGenTasks.Num() || InstanceCountClearTasks.Num())
 	{
-		if (NiagaraSupportsComputeShaders(GShaderPlatformForFeatureLevel[FeatureLevel]))
+		if (FNiagaraUtilities::AllowGPUParticles(GShaderPlatformForFeatureLevel[FeatureLevel]))
 		{
 			RHICmdList.TransitionResource(EResourceTransitionAccess::ERWBarrier, EResourceTransitionPipeline::EComputeToCompute, CountBuffer.UAV);
 
diff --git a/Engine/Plugins/FX/Niagara/Source/Niagara/Private/NiagaraRendererMeshes.cpp b/Engine/Plugins/FX/Niagara/Source/Niagara/Private/NiagaraRendererMeshes.cpp
--- a/Engine/Plugins/FX/Niagara/Source/Niagara/Private/NiagaraRendererMeshes.cpp
+++ b/Engine/Plugins/FX/Niagara/Source/Niagara/Private/NiagaraRendererMeshes.cpp
@@ -343,8 +343,9 @@ void FNiagaraRendererMeshes::GetDynamicMeshElements(const TArray<const FSceneVie
 					if (SortInfo.SortMode != ENiagaraSortMode::None && SortInfo.SortAttributeOffset != INDEX_NONE)
 					{
 						if (GNiagaraGPUSorting &&
-							GNiagaraGPUSortingCPUToGPUThreshold != INDEX_NONE &&
-							SortInfo.ParticleCount >= GNiagaraGPUSortingCPUToGPUThreshold)
+							GNiagaraGPUSortingCPUToGPUThreshold >= 0 &&
+							SortInfo.ParticleCount >= GNiagaraGPUSortingCPUToGPUThreshold &&
+							FNiagaraUtilities::AllowComputeShaders(Batcher->GetShaderPlatform()))
 						{
 							SortInfo.ParticleCount = NumInstances;
 							SortInfo.ParticleDataFloatSRV = ParticleData.ReadBuffer->SRV;
diff --git a/Engine/Plugins/FX/Niagara/Source/Niagara/Private/NiagaraRendererSprites.cpp b/Engine/Plugins/FX/Niagara/Source/Niagara/Private/NiagaraRendererSprites.cpp
--- a/Engine/Plugins/FX/Niagara/Source/Niagara/Private/NiagaraRendererSprites.cpp
+++ b/Engine/Plugins/FX/Niagara/Source/Niagara/Private/NiagaraRendererSprites.cpp
@@ -368,8 +368,9 @@ void FNiagaraRendererSprites::SetVertexFactoryParticleData(
 			if (SortInfo.SortMode != ENiagaraSortMode::None && SortInfo.SortAttributeOffset != INDEX_NONE)
 			{
 				if (GNiagaraGPUSorting &&
-					GNiagaraGPUSortingCPUToGPUThreshold != INDEX_NONE &&
-					SortInfo.ParticleCount >= GNiagaraGPUSortingCPUToGPUThreshold)
+					GNiagaraGPUSortingCPUToGPUThreshold >= 0 &&
+					SortInfo.ParticleCount >= GNiagaraGPUSortingCPUToGPUThreshold &&
+					FNiagaraUtilities::AllowComputeShaders(Batcher->GetShaderPlatform()))
 				{
 					SortInfo.ParticleCount = NumInstances;
 					SortInfo.ParticleDataFloatSRV = CPUSimParticleDataAllocation.ParticleData.ReadBuffer->SRV;
diff --git a/Engine/Plugins/FX/Niagara/Source/Niagara/Private/NiagaraScript.cpp b/Engine/Plugins/FX/Niagara/Source/Niagara/Private/NiagaraScript.cpp
--- a/Engine/Plugins/FX/Niagara/Source/Niagara/Private/NiagaraScript.cpp
+++ b/Engine/Plugins/FX/Niagara/Source/Niagara/Private/NiagaraScript.cpp
@@ -1625,9 +1625,9 @@ void UNiagaraScript::CacheResourceShadersForRendering(bool bRegenerateId, bool b
 
 			//if (ScriptResourcesByFeatureLevel[FeatureLevel])
 			{
-				if (FNiagaraUtilities::SupportsGPUParticles(CacheFeatureLevel))
+				const EShaderPlatform ShaderPlatform = GShaderPlatformForFeatureLevel[CacheFeatureLevel];
+				if (FNiagaraUtilities::AllowGPUParticles(ShaderPlatform))
 				{
-					EShaderPlatform ShaderPlatform = GShaderPlatformForFeatureLevel[CacheFeatureLevel];
 					ResourceToCache = ScriptResourcesByFeatureLevel[CacheFeatureLevel];
 					CacheShadersForResources(ShaderPlatform, &ScriptResource, true);
 					ScriptResourcesByFeatureLevel[CacheFeatureLevel] = &ScriptResource;
diff --git a/Engine/Plugins/FX/Niagara/Source/Niagara/Private/NiagaraSystemInstance.cpp b/Engine/Plugins/FX/Niagara/Source/Niagara/Private/NiagaraSystemInstance.cpp
--- a/Engine/Plugins/FX/Niagara/Source/Niagara/Private/NiagaraSystemInstance.cpp
+++ b/Engine/Plugins/FX/Niagara/Source/Niagara/Private/NiagaraSystemInstance.cpp
@@ -499,7 +499,7 @@ bool FNiagaraSystemInstance::DeallocateSystemInstance(TUniquePtr< FNiagaraSystem
 
 		// If we have active GPU emitters make sure we remove any pending ticks from the RT
 		NiagaraEmitterInstanceBatcher* InstanceBatcher = SystemInstanceAllocation->GetBatcher();
-		if (SystemInstanceAllocation->bHasGPUEmitters && NiagaraSupportsComputeShaders(InstanceBatcher->GetShaderPlatform()))
+		if (SystemInstanceAllocation->bHasGPUEmitters)
 		{
 			ENQUEUE_RENDER_COMMAND(NiagaraRemoveGPUSystem)
 			(
@@ -1959,7 +1959,7 @@ void FNiagaraSystemInstance::FinalizeTick_GameThread()
 				Component->MarkRenderDynamicDataDirty();
 
 				// Push any GPU ticks for this system instance.
-				if (ActiveGPUEmitterCount > 0 && NiagaraSupportsComputeShaders(GetBatcher()->GetShaderPlatform()) && Component->IsRegistered())
+				if (ActiveGPUEmitterCount > 0 && Batcher && FNiagaraUtilities::AllowGPUParticles(Batcher->GetShaderPlatform()) && Component->IsRegistered())
 				{
 					ensure(!IsComplete());
 					FNiagaraGPUSystemTick GPUTick;
diff --git a/Engine/Plugins/FX/Niagara/Source/Niagara/Public/NiagaraCommon.h b/Engine/Plugins/FX/Niagara/Source/Niagara/Public/NiagaraCommon.h
--- a/Engine/Plugins/FX/Niagara/Source/Niagara/Public/NiagaraCommon.h
+++ b/Engine/Plugins/FX/Niagara/Source/Niagara/Public/NiagaraCommon.h
@@ -133,11 +133,6 @@ enum class ENiagaraScriptCompileStatus : uint8
 	NCS_MAX,
 };
 
-FORCEINLINE bool NiagaraSupportsComputeShaders(EShaderPlatform ShaderPlatform)
-{
-	return RHISupportsComputeShaders(ShaderPlatform);
-}
-
 USTRUCT()
 struct FNiagaraDataSetID
 {
@@ -729,17 +724,18 @@ namespace FNiagaraUtilities
 		return IsFeatureLevelSupported(ShaderPlatform, ERHIFeatureLevel::SM5) || IsFeatureLevelSupported(ShaderPlatform, ERHIFeatureLevel::ES3_1);
 	}
 
-	inline bool SupportsGPUParticles(ERHIFeatureLevel::Type FeatureLevel)
-	{
-		EShaderPlatform ShaderPlatform = GShaderPlatformForFeatureLevel[FeatureLevel];
-		return NiagaraSupportsComputeShaders(ShaderPlatform);
-	}
-
+	// Whether the platform supports GPU particles. A static function that doesn't not rely on any runtime switches.
 	inline bool SupportsGPUParticles(EShaderPlatform ShaderPlatform)
 	{
-		return NiagaraSupportsComputeShaders(ShaderPlatform);
+		return RHISupportsComputeShaders(ShaderPlatform);
 	}
 
+	// Whether GPU particles are currently allowed. Could change depending on config and runtime switches.
+	bool AllowGPUParticles(EShaderPlatform ShaderPlatform);
+
+	// Whether compute shaders are allowed. Could change depending on config and runtime switches.
+	bool AllowComputeShaders(EShaderPlatform ShaderPlatform);
+	
 #if WITH_EDITORONLY_DATA
 	/**
 	 * Prepares rapid iteration parameter stores for simulation by removing old parameters no longer used by functions, by initializing new parameters
