Fix for data cooking issue that could lead to a memory stomp with the new constant buffer niagara changes

diff --git a/Engine/Plugins/FX/Niagara/Source/Niagara/Private/NiagaraSystem.cpp b/Engine/Plugins/FX/Niagara/Source/Niagara/Private/NiagaraSystem.cpp
--- a/Engine/Plugins/FX/Niagara/Source/Niagara/Private/NiagaraSystem.cpp
+++ b/Engine/Plugins/FX/Niagara/Source/Niagara/Private/NiagaraSystem.cpp
@@ -1467,6 +1467,10 @@ FNiagaraEmitterCompiledData::FNiagaraEmitterCompiledData()
 #if WITH_EDITORONLY_DATA
 void FNiagaraParameterDataSetBindingCollection::BuildInternal(const TArray<FNiagaraVariable>& ParameterVars, const FNiagaraDataSetCompiledData& DataSet, const FString& NamespaceBase, const FString& NamespaceReplacement)
 {
+	// be sure to reset the offsets first
+	FloatOffsets.Empty();
+	Int32Offsets.Empty();
+
 	const bool DoNameReplacement = !NamespaceBase.IsEmpty() && !NamespaceReplacement.IsEmpty();
 
 	int32 ParameterOffset = 0;
@@ -1506,5 +1510,8 @@ void FNiagaraParameterDataSetBindingCollection::BuildInternal(const TArray<FNiag
 
 		ParameterOffset += Var.GetSizeInBytes();
 	}
+
+	FloatOffsets.Shrink();
+	Int32Offsets.Shrink();
 }
 #endif
\ No newline at end of file
