The Niagara new emitter dialog now supports searching by multiple disjoint strings, eg. typing "water foam" will successfully find "Water_Foam_Base".

#rnx

diff --git a/Engine/Plugins/FX/Niagara/Source/NiagaraEditor/Private/SNiagaraAssetPickerList.cpp b/Engine/Plugins/FX/Niagara/Source/NiagaraEditor/Private/SNiagaraAssetPickerList.cpp
--- a/Engine/Plugins/FX/Niagara/Source/NiagaraEditor/Private/SNiagaraAssetPickerList.cpp
+++ b/Engine/Plugins/FX/Niagara/Source/NiagaraEditor/Private/SNiagaraAssetPickerList.cpp
@@ -190,7 +190,19 @@ bool SNiagaraAssetPickerList::OnCompareItemsForSorting(const FAssetData& ItemA,
 
 bool SNiagaraAssetPickerList::OnDoesItemMatchFilterText(const FText& FilterText, const FAssetData& Item)
 {
-	return Item.AssetName.ToString().Contains(FilterText.ToString());
+	TArray<FString> FilterStrings;
+	FilterText.ToString().ParseIntoArrayWS(FilterStrings, TEXT(","));
+
+	FString AssetNameString = Item.AssetName.ToString();
+	for (const FString& FilterString : FilterStrings)
+	{
+		if (!AssetNameString.Contains(FilterString))
+		{
+			return false;
+		}
+	}
+
+	return true;
 }
 
 TSharedRef<SWidget> SNiagaraAssetPickerList::OnGenerateWidgetForCategory(const FText& Category)
