- CopyTo checks the allocation count rather than the num instances
- This fixes a buffer overrun where the destination buffer has less instances than the current buffer but the destination buffer was presized larger to handle new instances, in this case the CopyTo would size down the buffer causing an buffer overrun inside the VM
- Added a check into SetNumInstances to ensure we don't attempt to set the number of instances above our allocation size (this also catches the issue we are fixing)

#rnx

diff --git a/Engine/Plugins/FX/Niagara/Source/Niagara/Classes/NiagaraDataSet.h b/Engine/Plugins/FX/Niagara/Source/Niagara/Classes/NiagaraDataSet.h
--- a/Engine/Plugins/FX/Niagara/Source/Niagara/Classes/NiagaraDataSet.h
+++ b/Engine/Plugins/FX/Niagara/Source/Niagara/Classes/NiagaraDataSet.h
@@ -133,7 +133,7 @@ public:
 	FORCEINLINE uint32 GetNumInstances()const { return NumInstances; }
 	FORCEINLINE uint32 GetNumInstancesAllocated()const { return NumInstancesAllocated; }
 
-	FORCEINLINE void SetNumInstances(uint32 InNumInstances) { NumInstances = InNumInstances; }
+	FORCEINLINE void SetNumInstances(uint32 InNumInstances) { check(InNumInstances <= NumInstancesAllocated); NumInstances = InNumInstances; }
 	FORCEINLINE uint32 GetSizeBytes()const { return FloatData.Num() + Int32Data.Num(); }
 	FORCEINLINE FRWBuffer& GetGPUBufferFloat() { return GPUBufferFloat; }
 	FORCEINLINE FRWBuffer& GetGPUBufferInt() { return GPUBufferInt;	}
diff --git a/Engine/Plugins/FX/Niagara/Source/Niagara/Private/NiagaraDataSet.cpp b/Engine/Plugins/FX/Niagara/Source/Niagara/Private/NiagaraDataSet.cpp
--- a/Engine/Plugins/FX/Niagara/Source/Niagara/Private/NiagaraDataSet.cpp
+++ b/Engine/Plugins/FX/Niagara/Source/Niagara/Private/NiagaraDataSet.cpp
@@ -802,8 +802,11 @@ void FNiagaraDataBuffer::CopyTo(FNiagaraDataBuffer& DestBuffer, int32 StartIdx,
 
 	if (InstancesToCopy != 0)
 	{
-		uint32 NewNumInstances = DestStartIdx + InstancesToCopy;
-		if (DestStartIdx < 0 || NewNumInstances >= DestBuffer.GetNumInstances())
+		const uint32 NewNumInstances = DestStartIdx + InstancesToCopy;
+
+		// Only allocate if we need to increase the number of instances as the caller may have previously
+		// allocated the array and may not be expecting it to shrink inside this call.
+		if (DestStartIdx < 0 || NewNumInstances >= DestBuffer.GetNumInstancesAllocated())
 		{
 			DestBuffer.Allocate(NewNumInstances, true);
 		}
