SpawnEffect moved from protected to public.

Simon.Tovey
#test compile/PIE
#rnx

diff --git a/Engine/Plugins/FX/Niagara/Source/NiagaraAnimNotifies/Public/AnimNotify_PlayNiagaraEffect.h b/Engine/Plugins/FX/Niagara/Source/NiagaraAnimNotifies/Public/AnimNotify_PlayNiagaraEffect.h
--- a/Engine/Plugins/FX/Niagara/Source/NiagaraAnimNotifies/Public/AnimNotify_PlayNiagaraEffect.h
+++ b/Engine/Plugins/FX/Niagara/Source/NiagaraAnimNotifies/Public/AnimNotify_PlayNiagaraEffect.h
@@ -54,13 +54,17 @@ public:
 	// Scale to spawn the Niagara system at
 	UPROPERTY(EditAnywhere, Category = "AnimNotify")
 	FVector Scale;
+	
+	// Spawns the NiagaraSystemComponent. Called from Notify.
+	UFUNCTION(BlueprintCallable, Category = "AnimNotify")
+	virtual UFXSystemComponent* SpawnEffect(USkeletalMeshComponent* MeshComp, UAnimSequenceBase* Animation);
 
 protected:
 	// Cached version of the Rotation Offset already in Quat form
 	FQuat RotationOffsetQuat;
 
-	// Spawns the NiagaraSystemComponent. Called from Notify.
-	virtual UFXSystemComponent* SpawnEffect(USkeletalMeshComponent* MeshComp, UAnimSequenceBase* Animation);
+	
+	
 
 public:
 
