- Remove warning about a null system asset by default, this mirrors Cascade
- Added console variable "fx.Niagara.ComponentWarnNullAsset" to enable warning

diff --git a/Engine/Plugins/FX/Niagara/Source/Niagara/Private/NiagaraComponent.cpp b/Engine/Plugins/FX/Niagara/Source/Niagara/Private/NiagaraComponent.cpp
--- a/Engine/Plugins/FX/Niagara/Source/Niagara/Private/NiagaraComponent.cpp
+++ b/Engine/Plugins/FX/Niagara/Source/Niagara/Private/NiagaraComponent.cpp
@@ -58,6 +58,14 @@ static FAutoConsoleVariableRef CVarSuppressNiagaraSystems(
 	ECVF_Default
 );
 
+static int32 GNiagaraComponentWarnNullAsset = 0;
+static FAutoConsoleVariableRef CVarNiagaraComponentWarnNullAsset(
+	TEXT("fx.Niagara.ComponentWarnNullAsset"),
+	GNiagaraComponentWarnNullAsset,
+	TEXT("When enabled we will warn if a NiagaraComponent is activate with a null asset.  This is sometimes useful for tracking down components that can be removed."),
+	ECVF_Default
+);
+
 void DumpNiagaraComponents(UWorld* World)
 {
 	for (TActorIterator<AActor> ActorItr(World); ActorItr; ++ActorItr)
@@ -783,7 +791,7 @@ void UNiagaraComponent::ActivateInternal(bool bReset /* = false */, bool bIsScal
 	if (Asset == nullptr)
 	{
 		DestroyInstance();
-		if (!HasAnyFlags(RF_DefaultSubObject | RF_ArchetypeObject | RF_ClassDefaultObject))
+		if (GNiagaraComponentWarnNullAsset && !HasAnyFlags(RF_DefaultSubObject | RF_ArchetypeObject | RF_ClassDefaultObject))
 		{
 			UE_LOG(LogNiagara, Warning, TEXT("Failed to activate Niagara Component due to missing or invalid asset! (%s)"), *GetFullName());
 		}
