Changed AllowGPUParticles to SupportsGPUParticles when cooking Niagara GPU particles, since AllowGPUParticles is reserved for runtime use only.

diff --git a/Engine/Plugins/FX/Niagara/Source/Niagara/Private/NiagaraScript.cpp b/Engine/Plugins/FX/Niagara/Source/Niagara/Private/NiagaraScript.cpp
--- a/Engine/Plugins/FX/Niagara/Source/Niagara/Private/NiagaraScript.cpp
+++ b/Engine/Plugins/FX/Niagara/Source/Niagara/Private/NiagaraScript.cpp
@@ -1626,7 +1626,7 @@ void UNiagaraScript::CacheResourceShadersForRendering(bool bRegenerateId, bool b
 			//if (ScriptResourcesByFeatureLevel[FeatureLevel])
 			{
 				const EShaderPlatform ShaderPlatform = GShaderPlatformForFeatureLevel[CacheFeatureLevel];
-				if (FNiagaraUtilities::AllowGPUParticles(ShaderPlatform))
+				if (FNiagaraUtilities::SupportsGPUParticles(ShaderPlatform))
 				{
 					ResourceToCache = ScriptResourcesByFeatureLevel[CacheFeatureLevel];
 					CacheShadersForResources(ShaderPlatform, &ScriptResource, true);
