Fixed possible Z fighting when aligning Niagara mesh emitters on particle velocity.

diff --git a/Engine/Plugins/FX/Niagara/Shaders/Private/NiagaraMeshVertexFactory.ush b/Engine/Plugins/FX/Niagara/Shaders/Private/NiagaraMeshVertexFactory.ush
--- a/Engine/Plugins/FX/Niagara/Shaders/Private/NiagaraMeshVertexFactory.ush
+++ b/Engine/Plugins/FX/Niagara/Shaders/Private/NiagaraMeshVertexFactory.ush
@@ -426,6 +426,10 @@ float4 CalcPrevWorldPosition(FVertexFactoryInput Input, FVertexFactoryIntermedia
 // From UnrealMath::FindBetween
 float4 FindQuatBetween(float3 From, float3 To)
 {
+	// It appears that length(normalize(To)) is not always exactly equal to 1 and can create issues. 
+	// See comments around NEEDS_PARTICLE_VELOCITY
+	// float NormAB = sqrt(dot(NormAB, NormAB) * dot(To, To));
+
 	float NormAB = 1.0f;
 	float w = NormAB + dot(From, To);
 	float4 Result;
@@ -610,9 +614,8 @@ FVertexFactoryIntermediates GetVertexFactoryIntermediates(FVertexFactoryInput In
 	}
 #endif
 
+	const float ParticleVelocityMagnitude = length(ParticleVelocity);
 #if NEEDS_PARTICLE_VELOCITY
-	float ParticleVelocityMagnitude = length(ParticleVelocity);
-
 	Intermediates.ParticleVelocity.xyz = (ParticleVelocityMagnitude > 0) ? normalize(ParticleVelocity) : float3(0.0f, 0.0f, 0.0f);
 	Intermediates.ParticleVelocity.w = ParticleVelocityMagnitude;
 #endif
@@ -652,7 +655,10 @@ FVertexFactoryIntermediates GetVertexFactoryIntermediates(FVertexFactoryInput In
 
 	if (MeshFacingMode == MESH_FACING_VELOCITY)
 	{
-		float4 ToAlignmentQuat = FindQuatBetween(float3(1.0f, 0.0f, 0.0f), normalize(ParticleVelocity.xyz));
+		// The result of normalize(ParticleVelocity) seems affected by whether lenght(ParticleVelocity) is computed above or not (see NEEDS_PARTICLE_VELOCITY).
+		// This makes z fighting possible between the basepass and the depthpass (which never defines NEEDS_PARTICLE_VELOCITY) for opaque materials.
+		// The workaround for this issue is to ensure that lenght(ParticleVelocity) gets computed in both shaders.
+		float4 ToAlignmentQuat = FindQuatBetween(float3(1.0f, 0.0f, 0.0f), (ParticleVelocityMagnitude > 0) ? normalize(ParticleVelocity) : float3(0.0f, 0.0f, 0.0f));
 		float4x4 ToAlignmentMat = QuatTo4x4( normalize(ToAlignmentQuat) );
 		Intermediates.LocalToWorld =  mul(LocalToWorldMat, mul(ToAlignmentMat, ScaleMatrix));
 	}
