Fixing some issues caused by GC in scalability manager.

diff --git a/Engine/Plugins/FX/Niagara/Source/Niagara/Private/NiagaraComponent.cpp b/Engine/Plugins/FX/Niagara/Source/Niagara/Private/NiagaraComponent.cpp
--- a/Engine/Plugins/FX/Niagara/Source/Niagara/Private/NiagaraComponent.cpp
+++ b/Engine/Plugins/FX/Niagara/Source/Niagara/Private/NiagaraComponent.cpp
@@ -1124,6 +1124,8 @@ void UNiagaraComponent::OnComponentDestroyed(bool bDestroyingHierarchy)
 	//UE_LOG(LogNiagara, Log, TEXT("OnComponentDestroyed %p %p"), this, SystemInstance.Get());
 	//DestroyInstance();//Can't do this here as we can call this from inside the system instance currently during completion 
 
+	UnregisterWithScalabilityManager();
+
 	Super::OnComponentDestroyed(bDestroyingHierarchy);
 }
 
@@ -1151,7 +1153,11 @@ void UNiagaraComponent::OnUnregister()
 
 void UNiagaraComponent::BeginDestroy()
 {
-	//UE_LOG(LogNiagara, Log, TEXT("UNiagaraComponent::BeginDestroy(): %0xP  %s\n"), this, *GetAsset()->GetFullName());
+	//UE_LOG(LogNiagara, Log, TEXT("UNiagaraComponent::BeginDestroy(): %0xP - %d - %s\n"), this, ScalabilityManagerHandle, *GetAsset()->GetFullName());
+
+	//By now we will have already unregisted with the scalability manger. Either directly in OnComponentDestroyed, or via the post GC callbacks in the manager it's self in the case of someone calling MarkPendingKill() directly on a component.
+	ScalabilityManagerHandle = INDEX_NONE;
+
 	DestroyInstance();
 
 	Super::BeginDestroy();
diff --git a/Engine/Plugins/FX/Niagara/Source/Niagara/Private/NiagaraScalabilityManager.cpp b/Engine/Plugins/FX/Niagara/Source/Niagara/Private/NiagaraScalabilityManager.cpp
--- a/Engine/Plugins/FX/Niagara/Source/Niagara/Private/NiagaraScalabilityManager.cpp
+++ b/Engine/Plugins/FX/Niagara/Source/Niagara/Private/NiagaraScalabilityManager.cpp
@@ -38,6 +38,24 @@ void FNiagaraScalabilityManager::AddReferencedObjects(FReferenceCollector& Colle
 	Collector.AddReferencedObjects(ManagedComponents);
 }
 
+
+void FNiagaraScalabilityManager::PreGarbageCollectBeginDestroy()
+{
+	//After the GC has potentially nulled out references to the components we were tracking we clear them out here.
+	//This should only be in the case where MarkPendingKill() is called directly. Typical component destruction will unregister in OnComponentDestroyed().
+	//Components then just clear their handle in BeginDestroy knowing they've already been removed from the manager.
+	//I would prefer some pre BeginDestroy() callback into the component in which I could cleanly unregister with the manager in all cases but I don't think that's possible.
+	int32 CompIdx = ManagedComponents.Num();
+	while (--CompIdx >= 0)
+	{
+		if (ManagedComponents[CompIdx] == nullptr)
+		{
+			//UE_LOG(LogNiagara, Warning, TEXT("Unregister from PreGCBeginDestroy @%d/%d - %s"), CompIdx, ManagedComponents.Num(), *EffectType->GetName());
+			UnregisterAt(CompIdx);
+		}
+	}
+}
+
 void FNiagaraScalabilityManager::Register(UNiagaraComponent* Component)
 {
 	check(Component->ScalabilityManagerHandle == INDEX_NONE);
@@ -118,9 +136,17 @@ void FNiagaraScalabilityManager::Update(FNiagaraWorldManager* WorldMan)
 		UNiagaraComponent* Component = ManagedComponents[CompIdx];
 		if (Component)
 		{
+			//Belt and braces GC safety. If someone calls MarkPendingKill() directly and we get here before we clear these out in the post GC callback.
+			if (Component->IsPendingKill())
+			{
+				//UE_LOG(LogNiagara, Warning, TEXT("Unregisteded a pending kill Niagara component from the scalability manager. \nComponent: 0x%P - %s\nEffectType: 0x%P - %s"),
+				//	Component, *Component->GetName(), EffectType, *EffectType->GetName());
+				Unregister(Component);
+				continue;
+			}
 			if (Component->GetAsset() == nullptr)
 			{
-				UE_LOG(LogNiagara, Warning, TEXT("Niagara System has been destroyed with components still registered to the scalability manager. Unregistering this component.\nComponent: %0xP - %s\nEffectType: %0xP - %s"),
+				UE_LOG(LogNiagara, Warning, TEXT("Niagara System has been destroyed with components still registered to the scalability manager. Unregistering this component.\nComponent: 0x%P - %s\nEffectType: 0x%P - %s"),
 					Component, *Component->GetName(), EffectType, *EffectType->GetName());
 				Unregister(Component);
 				continue;
diff --git a/Engine/Plugins/FX/Niagara/Source/Niagara/Private/NiagaraWorldManager.cpp b/Engine/Plugins/FX/Niagara/Source/Niagara/Private/NiagaraWorldManager.cpp
--- a/Engine/Plugins/FX/Niagara/Source/Niagara/Private/NiagaraWorldManager.cpp
+++ b/Engine/Plugins/FX/Niagara/Source/Niagara/Private/NiagaraWorldManager.cpp
@@ -52,7 +52,8 @@ FDelegateHandle FNiagaraWorldManager::OnWorldCleanupHandle;
 FDelegateHandle FNiagaraWorldManager::OnPreWorldFinishDestroyHandle;
 FDelegateHandle FNiagaraWorldManager::OnWorldBeginTearDownHandle;
 FDelegateHandle FNiagaraWorldManager::TickWorldHandle;
-FDelegateHandle FNiagaraWorldManager::PostGCHandle;
+FDelegateHandle FNiagaraWorldManager::PostGCHandle; 
+FDelegateHandle FNiagaraWorldManager::PreGCBeginDestroyHandle; 
 TMap<class UWorld*, class FNiagaraWorldManager*> FNiagaraWorldManager::WorldManagers;
 
 TGlobalResource<FNiagaraViewDataMgr> GNiagaraViewDataManager;
@@ -169,6 +170,7 @@ void FNiagaraWorldManager::OnStartup()
 	TickWorldHandle = FWorldDelegates::OnWorldPostActorTick.AddStatic(&FNiagaraWorldManager::TickWorld);
 
 	PostGCHandle = FCoreUObjectDelegates::GetPostGarbageCollect().AddStatic(&FNiagaraWorldManager::OnPostGarbageCollect);
+	PreGCBeginDestroyHandle = FCoreUObjectDelegates::PreGarbageCollectConditionalBeginDestroy.AddStatic(&FNiagaraWorldManager::OnPreGarbageCollectBeginDestroy);
 }
 
 void FNiagaraWorldManager::OnShutdown()
@@ -180,6 +182,8 @@ void FNiagaraWorldManager::OnShutdown()
 	FWorldDelegates::OnWorldPostActorTick.Remove(TickWorldHandle);
 
 	FCoreUObjectDelegates::GetPostGarbageCollect().Remove(PostGCHandle);
+	FCoreUObjectDelegates::PreGarbageCollectConditionalBeginDestroy.Remove(PreGCBeginDestroyHandle);
+
 
 	//Should have cleared up all world managers by now.
 	check(WorldManagers.Num() == 0);
@@ -355,6 +359,20 @@ void FNiagaraWorldManager::PostGarbageCollect()
 	while (ScalabilityManagers.Remove(nullptr)) {}
 }
 
+void FNiagaraWorldManager::PreGarbageCollectBeginDestroy()
+{
+	//Clear out and scalability managers who's EffectTypes have been GCd.
+	while (ScalabilityManagers.Remove(nullptr)) {}
+
+	//Also tell the scalability managers to clear out any references the GC has nulled.
+	for (auto& Pair : ScalabilityManagers)
+	{
+		FNiagaraScalabilityManager& ScalabilityMan = Pair.Value;
+		UNiagaraEffectType* EffectType = Pair.Key;
+		ScalabilityMan.PreGarbageCollectBeginDestroy();
+	}
+}
+
 void FNiagaraWorldManager::OnWorldInit(UWorld* World, const UWorld::InitializationValues IVS)
 {
 	check(WorldManagers.Find(World) == nullptr);
@@ -425,6 +443,14 @@ void FNiagaraWorldManager::OnPostGarbageCollect()
 	}
 }
 
+void FNiagaraWorldManager::OnPreGarbageCollectBeginDestroy()
+{
+	for (TPair<UWorld*, FNiagaraWorldManager*>& Pair : WorldManagers)
+	{
+		Pair.Value->PreGarbageCollectBeginDestroy();
+	}
+}
+
 void FNiagaraWorldManager::PostActorTick(float DeltaSeconds)
 {
 	CSV_SCOPED_TIMING_STAT_EXCLUSIVE(Niagara);
diff --git a/Engine/Plugins/FX/Niagara/Source/Niagara/Public/NiagaraScalabilityManager.h b/Engine/Plugins/FX/Niagara/Source/Niagara/Public/NiagaraScalabilityManager.h
--- a/Engine/Plugins/FX/Niagara/Source/Niagara/Public/NiagaraScalabilityManager.h
+++ b/Engine/Plugins/FX/Niagara/Source/Niagara/Public/NiagaraScalabilityManager.h
@@ -131,6 +131,7 @@ struct FNiagaraScalabilityManager
 	void Unregister(UNiagaraComponent* Component);
 
 	void AddReferencedObjects(FReferenceCollector& Collector);
+	void PreGarbageCollectBeginDestroy();
 
 #if DEBUG_SCALABILITY_STATE
 	void Dump();
diff --git a/Engine/Plugins/FX/Niagara/Source/Niagara/Public/NiagaraWorldManager.h b/Engine/Plugins/FX/Niagara/Source/Niagara/Public/NiagaraWorldManager.h
--- a/Engine/Plugins/FX/Niagara/Source/Niagara/Public/NiagaraWorldManager.h
+++ b/Engine/Plugins/FX/Niagara/Source/Niagara/Public/NiagaraWorldManager.h
@@ -126,6 +126,7 @@ public:
 	void OnWorldCleanup(bool bSessionEnded, bool bCleanupResources);
 
 	void PostGarbageCollect();
+	void PreGarbageCollectBeginDestroy();
 	
 	FORCEINLINE FNDI_SkeletalMesh_GeneratedData& GetSkeletalMeshGeneratedData() { return SkeletalMeshGeneratedData; }
 
@@ -179,7 +180,10 @@ private:
 
 	// Callback to handle any post GC processing needed.
 	static void OnPostGarbageCollect();
-	
+
+	// Callback to handle any pre GC processing needed.
+	static void OnPreGarbageCollectBeginDestroy();
+		
 	// Gamethread callback to cleanup references to the given batcher before it gets deleted on the renderthread.
 	void OnBatcherDestroyed_Internal(NiagaraEmitterInstanceBatcher* InBatcher);
 
@@ -200,6 +204,7 @@ private:
 	static FDelegateHandle OnWorldBeginTearDownHandle;
 	static FDelegateHandle TickWorldHandle;
 	static FDelegateHandle PostGCHandle;
+	static FDelegateHandle PreGCBeginDestroyHandle;
 
 	static TMap<class UWorld*, class FNiagaraWorldManager*> WorldManagers;
 
