Fixing copy-paste mistake causing CIS to warn (correctly) about a nullptr access.

#rnx

diff --git a/Engine/Plugins/FX/Niagara/Source/NiagaraEditor/Private/NiagaraNodeFunctionCall.cpp b/Engine/Plugins/FX/Niagara/Source/NiagaraEditor/Private/NiagaraNodeFunctionCall.cpp
--- a/Engine/Plugins/FX/Niagara/Source/NiagaraEditor/Private/NiagaraNodeFunctionCall.cpp
+++ b/Engine/Plugins/FX/Niagara/Source/NiagaraEditor/Private/NiagaraNodeFunctionCall.cpp
@@ -649,7 +649,7 @@ void UNiagaraNodeFunctionCall::UpdateNodeErrorMessage()
 			UEdGraphNode::bHasCompilerMessage = true;
 			UEdGraphNode::ErrorType = EMessageSeverity::Info;
 
-			if (FunctionScript->ExperimentalMessage.IsEmpty())
+			if (Signature.ExperimentalMessage.IsEmpty())
 			{
 				UEdGraphNode::NodeUpgradeMessage = LOCTEXT("FunctionExperimental", "This function is marked as experimental, use with care!");
 			}
@@ -657,7 +657,7 @@ void UNiagaraNodeFunctionCall::UpdateNodeErrorMessage()
 			{
 				FFormatNamedArguments Args;
 				Args.Add(TEXT("Message"), Signature.ExperimentalMessage);
-				UEdGraphNode::NodeUpgradeMessage = FText::Format(LOCTEXT("ModuleScriptExperimentalReason", "This function is marked as experimental, reason: {Message}"), Args);
+				UEdGraphNode::NodeUpgradeMessage = FText::Format(LOCTEXT("FunctionExperimentalReason", "This function is marked as experimental, reason: {Message}"), Args);
 			}
 		}
 	}
