SpawnEffect returns component to a protected variable.
Added a function to return the spawned effect component.

#test compile/PIE
#rnx

diff --git a/Engine/Plugins/FX/Niagara/Source/NiagaraAnimNotifies/Private/AnimNotify_PlayNiagaraEffect.cpp b/Engine/Plugins/FX/Niagara/Source/NiagaraAnimNotifies/Private/AnimNotify_PlayNiagaraEffect.cpp
--- a/Engine/Plugins/FX/Niagara/Source/NiagaraAnimNotifies/Private/AnimNotify_PlayNiagaraEffect.cpp
+++ b/Engine/Plugins/FX/Niagara/Source/NiagaraAnimNotifies/Private/AnimNotify_PlayNiagaraEffect.cpp
@@ -75,8 +75,13 @@ void UAnimNotify_PlayNiagaraEffect::ValidateAssociatedAssets()
 
 void UAnimNotify_PlayNiagaraEffect::Notify(class USkeletalMeshComponent* MeshComp, class UAnimSequenceBase* Animation)
 {
-	// Don't call super to avoid unnecessary call in to blueprints
-	SpawnEffect(MeshComp, Animation);
+	//Store the spawned effect in a protected variable
+	SpawnedEffect = SpawnEffect(MeshComp, Animation);
+	
+	//Call to BP to allows setting of Niagara User Variables
+	Super::Notify(MeshComp, Animation);
+	
+	
 }
 
 FString UAnimNotify_PlayNiagaraEffect::GetNotifyName_Implementation() const
@@ -120,3 +125,8 @@ UFXSystemComponent* UAnimNotify_PlayNiagaraEffect::SpawnEffect(USkeletalMeshComp
 
 	return ReturnComp;
 }
+
+UFXSystemComponent* UAnimNotify_PlayNiagaraEffect::GetSpawnedEffect() 
+{
+	return SpawnedEffect;
+}
diff --git a/Engine/Plugins/FX/Niagara/Source/NiagaraAnimNotifies/Public/AnimNotify_PlayNiagaraEffect.h b/Engine/Plugins/FX/Niagara/Source/NiagaraAnimNotifies/Public/AnimNotify_PlayNiagaraEffect.h
--- a/Engine/Plugins/FX/Niagara/Source/NiagaraAnimNotifies/Public/AnimNotify_PlayNiagaraEffect.h
+++ b/Engine/Plugins/FX/Niagara/Source/NiagaraAnimNotifies/Public/AnimNotify_PlayNiagaraEffect.h
@@ -5,8 +5,12 @@
 #include "CoreMinimal.h"
 #include "UObject/ObjectMacros.h"
 #include "Animation/AnimNotifies/AnimNotify.h"
+
+
 #include "AnimNotify_PlayNiagaraEffect.generated.h"
 
+
+
 class UAnimSequenceBase;
 class UNiagaraSystem;
 class USkeletalMeshComponent;
@@ -54,17 +58,22 @@ public:
 	// Scale to spawn the Niagara system at
 	UPROPERTY(EditAnywhere, Category = "AnimNotify")
 	FVector Scale;
-	
-	// Spawns the NiagaraSystemComponent. Called from Notify.
+
+	// Return FXSystemComponent created from SpawnEffect
 	UFUNCTION(BlueprintCallable, Category = "AnimNotify")
-	virtual UFXSystemComponent* SpawnEffect(USkeletalMeshComponent* MeshComp, UAnimSequenceBase* Animation);
+	UFXSystemComponent* GetSpawnedEffect();
 
 protected:
+
+	//FXSystem Pointer to Spawned Effect called from Notify.
+	UFXSystemComponent* SpawnedEffect;
+
 	// Cached version of the Rotation Offset already in Quat form
 	FQuat RotationOffsetQuat;
 
-	
-	
+	// Spawns the NiagaraSystemComponent. Called from Notify.
+	virtual UFXSystemComponent* SpawnEffect(USkeletalMeshComponent* MeshComp, UAnimSequenceBase* Animation);
+
 
 public:
 
