- Warn if the sim target is not what we expect it to be, this is not fatal and can be handled by the calling code
- One reason why this occurs is if an event is used on a GPU system

#rnx

diff --git a/Engine/Plugins/FX/Niagara/Source/Niagara/Private/NiagaraScript.cpp b/Engine/Plugins/FX/Niagara/Source/Niagara/Private/NiagaraScript.cpp
--- a/Engine/Plugins/FX/Niagara/Source/Niagara/Private/NiagaraScript.cpp
+++ b/Engine/Plugins/FX/Niagara/Source/Niagara/Private/NiagaraScript.cpp
@@ -499,8 +499,17 @@ FNiagaraScriptExecutionParameterStore* UNiagaraScript::GetExecutionReadyParamete
 		return nullptr;
 	}
 #else
-	check(GetSimTarget().GetValue() == SimTarget);
-	return &ScriptExecutionParamStore;
+	TOptional<ENiagaraSimTarget> ActualSimTarget = GetSimTarget();
+	if (ActualSimTarget.IsSet() )
+	{
+		if ( ActualSimTarget == SimTarget )
+		{
+			return &ScriptExecutionParamStore;
+		}
+
+		UE_LOG(LogNiagara, Warning, TEXT("SimTarget is '%d' but expecting '%d' on Script '%s' Usage '%d'"), ActualSimTarget.GetValue(), SimTarget, *GetFullName(), Usage);
+	}
+	return nullptr;
 #endif
 }
 
