- Niagara Skel Mesh DI, fix for when a MasterPoseComponent is used

#rnx

diff --git a/Engine/Plugins/FX/Niagara/Source/Niagara/Private/NiagaraDataInterfaceSkeletalMesh.cpp b/Engine/Plugins/FX/Niagara/Source/Niagara/Private/NiagaraDataInterfaceSkeletalMesh.cpp
--- a/Engine/Plugins/FX/Niagara/Source/Niagara/Private/NiagaraDataInterfaceSkeletalMesh.cpp
+++ b/Engine/Plugins/FX/Niagara/Source/Niagara/Private/NiagaraDataInterfaceSkeletalMesh.cpp
@@ -187,7 +187,55 @@ bool FSkeletalMeshSkinningData::Tick(float InDeltaSeconds, bool bRequirePreskin)
 
 	if (BoneMatrixUsers > 0)
 	{
-		SkelComp->CacheRefToLocalMatrices(CurrBoneRefToLocals());
+		TArray<FMatrix>& CurrBones = CurrBoneRefToLocals();
+		if ( USkinnedMeshComponent* MasterComponent = SkelComp->MasterPoseComponent.Get() )
+		{
+			const USkeletalMesh* SkelMesh = SkelComp->SkeletalMesh;
+			const TArray<int32>& MasterBoneMap = SkelComp->GetMasterBoneMap();
+			const int32 NumBones = MasterBoneMap.Num();
+
+			check(SkelMesh);
+			if ( NumBones == 0 )
+			{
+				// This case indicates an invalid master pose component (e.g. no skeletal mesh)
+				CurrBones.Empty(SkelMesh->RefSkeleton.GetNum());
+				CurrBones.AddDefaulted(SkelMesh->RefSkeleton.GetNum());
+			}
+			else
+			{
+				CurrBones.SetNumUninitialized(NumBones);
+
+				const TArray<FTransform>& MasterTransforms = MasterComponent->GetComponentSpaceTransforms();
+				for (int32 BoneIndex = 0; BoneIndex < NumBones; ++BoneIndex)
+				{
+					bool bFoundMaster = false;
+					if (MasterBoneMap.IsValidIndex(BoneIndex))
+					{
+						const int32 MasterIndex = MasterBoneMap[BoneIndex];
+						if (MasterIndex != INDEX_NONE && MasterIndex < MasterTransforms.Num())
+						{
+							CurrBones[BoneIndex] = MasterTransforms[MasterIndex].ToMatrixWithScale();
+							bFoundMaster = true;
+						}
+					}
+
+					if ( !bFoundMaster )
+					{
+						const int32 ParentIndex = SkelMesh->RefSkeleton.GetParentIndex(BoneIndex);
+
+						CurrBones[BoneIndex] = SkelMesh->RefSkeleton.GetRefBonePose()[BoneIndex].ToMatrixWithScale();
+						if ( CurrBones.IsValidIndex(ParentIndex) && ParentIndex < BoneIndex )
+						{
+							CurrBones[BoneIndex] = CurrBones[ParentIndex] * CurrBones[BoneIndex];
+						}
+					}
+				}
+			}
+		}
+		else
+		{
+			SkelComp->CacheRefToLocalMatrices(CurrBones);
+		}
 	}
 
 	//Prime the prev matrices if they're missing.
