More explicitly check for the template and library flags (including ones newly set on loaded assets) and using GET_MEMBER_NAME_CHECKED.
#rnx

diff --git a/Engine/Plugins/FX/Niagara/Source/NiagaraEditor/Private/SNiagaraAssetPickerList.cpp b/Engine/Plugins/FX/Niagara/Source/NiagaraEditor/Private/SNiagaraAssetPickerList.cpp
--- a/Engine/Plugins/FX/Niagara/Source/NiagaraEditor/Private/SNiagaraAssetPickerList.cpp
+++ b/Engine/Plugins/FX/Niagara/Source/NiagaraEditor/Private/SNiagaraAssetPickerList.cpp
@@ -92,9 +92,22 @@ TArray<FText> SNiagaraAssetPickerList::OnGetCategoriesForItem(const FAssetData&
 	}
 	else
 	{
+		bool bIsTemplate = false;
+
+		bool bFoundTemplateScriptTag = Item.GetTagValue(GET_MEMBER_NAME_CHECKED(UNiagaraEmitter, bIsTemplateAsset), bIsTemplate);
+
+		if (bFoundTemplateScriptTag == false)
+		{
+			if (Item.IsAssetLoaded())
+			{
+				UNiagaraEmitter* EmitterAsset = static_cast<UNiagaraEmitter*>(Item.GetAsset());
+				if (EmitterAsset != nullptr)
+				{
+					bIsTemplate = EmitterAsset->bIsTemplateAsset;
+				}
+			}
+		}
 
-		bool bIsTemplate;
-		Item.GetTagValue("bIsTemplateAsset", bIsTemplate);
 		if (bIsTemplate)
 		{
 			Categories.Add(TemplateCategory);
@@ -103,18 +116,32 @@ TArray<FText> SNiagaraAssetPickerList::OnGetCategoriesForItem(const FAssetData&
 		{
 			Categories.Add(LOCTEXT("NonTemplateEmitters", "Non-Template"));
 		}
-	}
 
-	bool bIsLibrary;
-	bool bFoundLibScriptTag = Item.GetTagValue("bExposeToLibrary", bIsLibrary);
 
-	if (bFoundLibScriptTag && bIsLibrary)
-	{
-		Categories.Add(LibraryCategory);
-	}
-	else
-	{
-		Categories.Add(NonLibraryCategory);
+		bool bInLibrary = false;
+		bool bFoundLibraryTag = Item.GetTagValue(GET_MEMBER_NAME_CHECKED(UNiagaraEmitter, bExposeToLibrary), bInLibrary);
+
+		if (bFoundLibraryTag == false)
+		{
+			if (Item.IsAssetLoaded())
+			{
+				UNiagaraEmitter* EmitterAsset = static_cast<UNiagaraEmitter*>(Item.GetAsset());
+				if (EmitterAsset != nullptr)
+				{
+					bInLibrary = EmitterAsset->bExposeToLibrary;
+				}
+			}
+		}
+
+		if (bFoundLibraryTag && bInLibrary)
+		{
+			Categories.Add(LibraryCategory);
+		}
+		else
+		{
+			Categories.Add(NonLibraryCategory);
+		}
+
 	}
 
 
@@ -237,14 +264,40 @@ TSharedRef<SWidget> SNiagaraAssetPickerList::OnGenerateWidgetForItem(const FAsse
 	FText AssetDescription;
 	Item.GetTagValue("TemplateAssetDescription", AssetDescription);
 
-	bool bIsTemplate;
-	bool bIsLibrary;
-	Item.GetTagValue("bIsTemplateAsset", bIsTemplate);
-	bool bFoundLibScriptTag = Item.GetTagValue("bExposeToLibrary", bIsLibrary);
+	bool bIsTemplate = false;
+	bool bInLibrary = false;
+	bool bFoundTemplateScriptTag = Item.GetTagValue(GET_MEMBER_NAME_CHECKED(UNiagaraEmitter, bIsTemplateAsset), bIsTemplate);
+
+	if (bFoundTemplateScriptTag == false)
+	{
+		if (Item.IsAssetLoaded())
+		{
+			UNiagaraEmitter* EmitterAsset = static_cast<UNiagaraEmitter*>(Item.GetAsset());
+			if (EmitterAsset != nullptr)
+			{
+				bIsTemplate = EmitterAsset->bIsTemplateAsset;
+			}
+		}
+	}
+
+	bool bFoundLibraryTag = Item.GetTagValue(GET_MEMBER_NAME_CHECKED(UNiagaraEmitter, bExposeToLibrary), bInLibrary);
+
+	if (bFoundLibraryTag == false)
+	{
+		if (Item.IsAssetLoaded())
+		{
+			UNiagaraEmitter* EmitterAsset = static_cast<UNiagaraEmitter*>(Item.GetAsset());
+			if (EmitterAsset != nullptr)
+			{
+				bInLibrary = EmitterAsset->bExposeToLibrary;
+			}
+		}
+	}
+
 
 	if (bTemplateOnly
-		|| (bIsTemplate)
-		|| (bFoundLibScriptTag && bIsLibrary))
+		|| (bFoundTemplateScriptTag && bIsTemplate)
+		|| (bFoundLibraryTag && bInLibrary))
 	{
 		return
 				SNew(SHorizontalBox)
