Hiding Unspecified as an Input Node usage flag, and setting the new default to Parameter
#rnx

diff --git a/Engine/Plugins/FX/Niagara/Source/Niagara/Public/NiagaraCommon.h b/Engine/Plugins/FX/Niagara/Source/Niagara/Public/NiagaraCommon.h
--- a/Engine/Plugins/FX/Niagara/Source/Niagara/Public/NiagaraCommon.h
+++ b/Engine/Plugins/FX/Niagara/Source/Niagara/Public/NiagaraCommon.h
@@ -102,7 +102,7 @@ enum class ENiagaraDataSetType : uint8
 UENUM()
 enum class ENiagaraInputNodeUsage : uint8
 {
-	Undefined = 0,
+	Undefined = 0 UMETA(Hidden),
 	Parameter,
 	Attribute,
 	SystemConstant,
diff --git a/Engine/Plugins/FX/Niagara/Source/NiagaraEditor/Private/NiagaraNodeInput.cpp b/Engine/Plugins/FX/Niagara/Source/NiagaraEditor/Private/NiagaraNodeInput.cpp
--- a/Engine/Plugins/FX/Niagara/Source/NiagaraEditor/Private/NiagaraNodeInput.cpp
+++ b/Engine/Plugins/FX/Niagara/Source/NiagaraEditor/Private/NiagaraNodeInput.cpp
@@ -20,7 +20,7 @@ DECLARE_CYCLE_STAT(TEXT("NiagaraEditor - UNiagaraNodeInput - SortNodes"), STAT_N
 
 UNiagaraNodeInput::UNiagaraNodeInput(const FObjectInitializer& ObjectInitializer)
 	: Super(ObjectInitializer)
-	, Usage(ENiagaraInputNodeUsage::Undefined)
+	, Usage(ENiagaraInputNodeUsage::Parameter)
 	, CallSortPriority(0)
 	, DataInterface(nullptr)
 {
