Niagara - Prevent a cook hang when a niagara asset does not have valid IDs for shader caching.

diff --git a/Engine/Plugins/FX/Niagara/Source/Niagara/Classes/NiagaraScript.h b/Engine/Plugins/FX/Niagara/Source/Niagara/Classes/NiagaraScript.h
--- a/Engine/Plugins/FX/Niagara/Source/Niagara/Classes/NiagaraScript.h
+++ b/Engine/Plugins/FX/Niagara/Source/Niagara/Classes/NiagaraScript.h
@@ -661,6 +661,8 @@ private:
 #if WITH_EDITORONLY_DATA
 	class UNiagaraSystem* FindRootSystem();
 
+	bool HasIdsRequiredForShaderCaching() const;
+
 	/** 'Source' data/graphs for this script */
 	UPROPERTY()
 	class UNiagaraScriptSourceBase*	Source;
diff --git a/Engine/Plugins/FX/Niagara/Source/Niagara/Private/NiagaraScript.cpp b/Engine/Plugins/FX/Niagara/Source/Niagara/Private/NiagaraScript.cpp
--- a/Engine/Plugins/FX/Niagara/Source/Niagara/Private/NiagaraScript.cpp
+++ b/Engine/Plugins/FX/Niagara/Source/Niagara/Private/NiagaraScript.cpp
@@ -292,6 +292,11 @@ class UNiagaraSystem* UNiagaraScript::FindRootSystem()
 	return nullptr;
 }
 
+bool UNiagaraScript::HasIdsRequiredForShaderCaching() const
+{
+	return CachedScriptVMId.CompilerVersionID.IsValid() && CachedScriptVMId.BaseScriptCompileHash.IsValid();
+}
+
 FString UNiagaraScript::GetNiagaraDDCKeyString()
 {
 	enum { UE_NIAGARA_COMPILATION_DERIVEDDATA_VER = 2 };
@@ -1503,10 +1508,10 @@ void UNiagaraScript::BeginCacheForCookedPlatformData(const ITargetPlatform *Targ
 			SystemOwner->WaitForCompilationComplete();
 		}
 
-		if (CachedScriptVMId.CompilerVersionID.IsValid() == false ||  CachedScriptVMId.BaseScriptCompileHash.IsValid() == false)
+		if (HasIdsRequiredForShaderCaching() == false)
 		{
-			UE_LOG(LogNiagara, Error,
-				TEXT("Failed to cache cooked shader for script %s because it had an invalid cached script id.  This should be fixed by running the console command fx.PreventSystemRecompile with the owning system asset path as the argument and then resaving the assets."),
+			UE_LOG(LogNiagara, Warning,
+				TEXT("Could not cache cooked shader for script %s because it had an invalid cached script id.  This should be fixed by running the console command fx.PreventSystemRecompile with the owning system asset path as the argument and then resaving the assets."),
 				*GetPathName());
 			return;
 		}
@@ -1530,7 +1535,7 @@ void UNiagaraScript::BeginCacheForCookedPlatformData(const ITargetPlatform *Targ
 
 bool UNiagaraScript::IsCachedCookedPlatformDataLoaded(const ITargetPlatform* TargetPlatform)
 {
-	if (ShouldCacheShadersForCooking())
+	if (ShouldCacheShadersForCooking() && HasIdsRequiredForShaderCaching())
 	{
 		bool bHasOutstandingCompilationRequests = false;
 		if (UNiagaraSystem* SystemOwner = FindRootSystem())
