Tweak for custom facing + alignment.

diff --git a/Engine/Plugins/FX/Niagara/Shaders/Private/NiagaraSpriteVertexFactory.ush b/Engine/Plugins/FX/Niagara/Shaders/Private/NiagaraSpriteVertexFactory.ush
--- a/Engine/Plugins/FX/Niagara/Shaders/Private/NiagaraSpriteVertexFactory.ush
+++ b/Engine/Plugins/FX/Niagara/Shaders/Private/NiagaraSpriteVertexFactory.ush
@@ -488,8 +488,8 @@ void GetTangents(FVertexFactoryInput Input, FVertexFactoryIntermediates Intermed
 	BRANCH
 	if(NiagaraSpriteVFLooseParameters.ParticleAlignmentMode == ALIGNMENT_MODE_CUSTOMALIGNMENT && NiagaraSpriteVFLooseParameters.ParticleFacingMode == FACING_MODE_CUSTOMFACINGVECTOR)
 	{
-		Right = cross(Intermediates.CustomAlignmentVector, CustomFacingVector);
-		Up = Intermediates.CustomAlignmentVector;
+		Right = Intermediates.CustomAlignmentVector;
+		Up = cross(CustomFacingVector, Right);
 	}
 	else if(NiagaraSpriteVFLooseParameters.ParticleAlignmentMode == ALIGNMENT_MODE_UNALIGNED)
 	{
