Merging 11120879 - Niagara - Prevent niagara script nodes and niagara system overview nodes from being pasted into the wrong types of graphs and also post load fixup any overview graphs with niagara script nodes in them.

diff --git a/Engine/Plugins/FX/Niagara/Source/NiagaraEditor/Private/NiagaraNode.cpp b/Engine/Plugins/FX/Niagara/Source/NiagaraEditor/Private/NiagaraNode.cpp
--- a/Engine/Plugins/FX/Niagara/Source/NiagaraEditor/Private/NiagaraNode.cpp
+++ b/Engine/Plugins/FX/Niagara/Source/NiagaraEditor/Private/NiagaraNode.cpp
@@ -354,7 +354,7 @@ TSharedPtr<SGraphNode> UNiagaraNode::CreateVisualWidget()
 
 void UNiagaraNode::GetPinHoverText(const UEdGraphPin& Pin, FString& HoverTextOut) const
 {
-	const UNiagaraGraph* NiagaraGraph = GetNiagaraGraph();
+	const UNiagaraGraph* NiagaraGraph = Cast<UNiagaraGraph>(GetGraph());
 	if (NiagaraGraph)
 	{
 		const UEdGraphSchema_Niagara* Schema = Cast<UEdGraphSchema_Niagara>(NiagaraGraph->GetSchema());
@@ -399,14 +399,19 @@ void UNiagaraNode::GetNodeContextMenuActions(class UToolMenu* Menu, class UGraph
 	}
 }
 
+bool UNiagaraNode::CanCreateUnderSpecifiedSchema(const UEdGraphSchema* Schema) const
+{
+	return Schema->IsA<UEdGraphSchema_Niagara>();
+}
+
 void UNiagaraNode::MarkNodeRequiresSynchronization(FString Reason, bool bRaiseGraphNeedsRecompile)
 {
 	Modify();
 	ChangeId = FGuid::NewGuid();
 	//UE_LOG(LogNiagaraEditor, Verbose, TEXT("Node %s was marked requires synchronization.  Reason: %s"), *GetPathName(), *Reason);
 
-	UNiagaraGraph* Graph = GetNiagaraGraph();
-	if (bRaiseGraphNeedsRecompile)
+	UNiagaraGraph* Graph = Cast<UNiagaraGraph>(GetGraph());
+	if (Graph != nullptr && bRaiseGraphNeedsRecompile)
 	{
 		Graph->NotifyGraphNeedsRecompile();
 	}
diff --git a/Engine/Plugins/FX/Niagara/Source/NiagaraEditor/Private/NiagaraOverviewNode.cpp b/Engine/Plugins/FX/Niagara/Source/NiagaraEditor/Private/NiagaraOverviewNode.cpp
--- a/Engine/Plugins/FX/Niagara/Source/NiagaraEditor/Private/NiagaraOverviewNode.cpp
+++ b/Engine/Plugins/FX/Niagara/Source/NiagaraEditor/Private/NiagaraOverviewNode.cpp
@@ -10,6 +10,7 @@
 #include "ViewModels/NiagaraEmitterHandleViewModel.h"
 #include "ViewModels/NiagaraSystemViewModel.h"
 #include "ViewModels/NiagaraSystemSelectionViewModel.h"
+#include "EdGraphSchema_NiagaraSystemOverview.h"
 #include "Modules/ModuleManager.h"
 #include "ToolMenuSection.h"
 #include "ToolMenu.h"
@@ -270,6 +271,11 @@ void UNiagaraOverviewNode::GetNodeContextMenuActions(class UToolMenu* Menu, clas
 	}
 }
 
+bool UNiagaraOverviewNode::CanCreateUnderSpecifiedSchema(const UEdGraphSchema* Schema) const
+{
+	return Schema->IsA<UEdGraphSchema_NiagaraSystemOverview>();
+}
+
 UNiagaraSystem* UNiagaraOverviewNode::GetOwningSystem()
 {
 	return OwningSystem;
diff --git a/Engine/Plugins/FX/Niagara/Source/NiagaraEditor/Private/NiagaraSystemEditorData.cpp b/Engine/Plugins/FX/Niagara/Source/NiagaraEditor/Private/NiagaraSystemEditorData.cpp
--- a/Engine/Plugins/FX/Niagara/Source/NiagaraEditor/Private/NiagaraSystemEditorData.cpp
+++ b/Engine/Plugins/FX/Niagara/Source/NiagaraEditor/Private/NiagaraSystemEditorData.cpp
@@ -8,6 +8,7 @@
 #include "NiagaraEmitter.h"
 #include "NiagaraEmitterEditorData.h"
 #include "NiagaraOverviewNode.h"
+#include "NiagaraNode.h"
 #include "EdGraphSchema_NiagaraSystemOverview.h"
 #include "EdGraph/EdGraph.h"
 
@@ -107,6 +108,14 @@ void UNiagaraSystemEditorData::PostLoadFromOwner(UObject* InOwner)
 	{
 		SynchronizeOverviewGraphWithSystem(*OwningSystem);
 	}
+
+	// Remove any niagara nodes which may have been pasted into the overview graph in error.  This is no longer possible.
+	TArray<UNiagaraNode*> NiagaraNodes;
+	SystemOverviewGraph->GetNodesOfClass<UNiagaraNode>(NiagaraNodes);
+	for (UNiagaraNode* NiagaraNode : NiagaraNodes)
+	{
+		NiagaraNode->DestroyNode();
+	}
 }
 
 UNiagaraSystemEditorFolder& UNiagaraSystemEditorData::GetRootFolder() const
diff --git a/Engine/Plugins/FX/Niagara/Source/NiagaraEditor/Public/NiagaraNode.h b/Engine/Plugins/FX/Niagara/Source/NiagaraEditor/Public/NiagaraNode.h
--- a/Engine/Plugins/FX/Niagara/Source/NiagaraEditor/Public/NiagaraNode.h
+++ b/Engine/Plugins/FX/Niagara/Source/NiagaraEditor/Public/NiagaraNode.h
@@ -47,6 +47,7 @@ public:
 	virtual TSharedPtr<SGraphNode> CreateVisualWidget() override; 
 	virtual void GetPinHoverText(const UEdGraphPin& Pin, FString& HoverTextOut) const override;
 	virtual void GetNodeContextMenuActions(class UToolMenu* Menu, class UGraphNodeContextMenuContext* Context) const override;
+	virtual bool CanCreateUnderSpecifiedSchema(const UEdGraphSchema* Schema) const override;
 	//~ End EdGraphNode Interface
 
 	/** Get the Niagara graph that owns this node */
diff --git a/Engine/Plugins/FX/Niagara/Source/NiagaraEditor/Public/NiagaraOverviewNode.h b/Engine/Plugins/FX/Niagara/Source/NiagaraEditor/Public/NiagaraOverviewNode.h
--- a/Engine/Plugins/FX/Niagara/Source/NiagaraEditor/Public/NiagaraOverviewNode.h
+++ b/Engine/Plugins/FX/Niagara/Source/NiagaraEditor/Public/NiagaraOverviewNode.h
@@ -34,6 +34,8 @@ public:
 	virtual bool GetCanRenameNode() const override;
 
 	virtual void GetNodeContextMenuActions(class UToolMenu* Menu, class UGraphNodeContextMenuContext* Context) const override;
+
+	virtual bool CanCreateUnderSpecifiedSchema(const UEdGraphSchema* Schema) const override;
 	//~ End UEdGraphNode Interface
 
 	UNiagaraSystem* GetOwningSystem();
