Merging CL#11123575

#UE4 Demoting another error to warning

diff --git a/Engine/Plugins/FX/Niagara/Source/Niagara/Private/NiagaraScript.cpp b/Engine/Plugins/FX/Niagara/Source/Niagara/Private/NiagaraScript.cpp
--- a/Engine/Plugins/FX/Niagara/Source/Niagara/Private/NiagaraScript.cpp
+++ b/Engine/Plugins/FX/Niagara/Source/Niagara/Private/NiagaraScript.cpp
@@ -1327,8 +1327,8 @@ void UNiagaraScript::SetVMCompilationResults(const FNiagaraVMExecutableDataId& I
 		}
 		else
 		{
-			UE_LOG(LogNiagara, Error,
-				TEXT("Failed to cache resource shaders for rendering for script %s because it had an invalid cached script id. This should be fixed by force recompiling the owning asset using the 'Full Rebuild' option and then saving the asset."),
+			UE_LOG(LogNiagara, Warning,
+				TEXT("Could not cache resource shaders for rendering for script %s because it had an invalid cached script id. This should be fixed by force recompiling the owning asset using the 'Full Rebuild' option and then saving the asset."),
 				*GetPathName());
 		}
 	}
