- Fix for GPU emitter ticks being counted differently between Tick_Concurrent & Finalize

#rnx

diff --git a/Engine/Plugins/FX/Niagara/Source/Niagara/Private/NiagaraSystemInstance.cpp b/Engine/Plugins/FX/Niagara/Source/Niagara/Private/NiagaraSystemInstance.cpp
--- a/Engine/Plugins/FX/Niagara/Source/Niagara/Private/NiagaraSystemInstance.cpp
+++ b/Engine/Plugins/FX/Niagara/Source/Niagara/Private/NiagaraSystemInstance.cpp
@@ -716,6 +716,8 @@ void FNiagaraSystemInstance::ResetInternal(bool bResetSimulations)
 	TickCount = 0;
 	CachedDeltaSeconds = 0.0f;
 	bLODDistanceIsValid = false;
+	TotalGPUParamSize = 0;
+	ActiveGPUEmitterCount = 0;
 	// Note: We do not need to update our bounds here as they are still valid
 
 	UNiagaraSystem* System = GetSystem();
@@ -1876,26 +1878,26 @@ void FNiagaraSystemInstance::Tick_Concurrent()
 	// now tick all emitters
 	for (int32 EmitterIdx = 0; EmitterIdx < Emitters.Num(); EmitterIdx++)
 	{
+		FNiagaraEmitterInstance& Inst = Emitters[EmitterIdx].Get();
 		if (EmittersShouldTick[EmitterIdx])
 		{
-			FNiagaraEmitterInstance& Inst = Emitters[EmitterIdx].Get();
 			Inst.Tick(CachedDeltaSeconds);
+		}
 
-			if (Inst.GetCachedEmitter() && Inst.GetCachedEmitter()->SimTarget == ENiagaraSimTarget::GPUComputeSim && Inst.GetGPUContext() != nullptr && (Inst.GetExecutionState() != ENiagaraExecutionState::Complete))
+		if (Inst.GetCachedEmitter() && Inst.GetCachedEmitter()->SimTarget == ENiagaraSimTarget::GPUComputeSim && Inst.GetGPUContext() != nullptr && (Inst.GetExecutionState() != ENiagaraExecutionState::Complete))
+		{
+			if (FirstGpuEmitter)
 			{
-				if (FirstGpuEmitter)
-				{
-					TotalGPUParamSize += 2 * sizeof(FNiagaraGlobalParameters);
-					TotalGPUParamSize += 2 * sizeof(FNiagaraSystemParameters);
-					TotalGPUParamSize += 2 * sizeof(FNiagaraOwnerParameters);
-					FirstGpuEmitter = false;
-				}
+				TotalGPUParamSize += 2 * sizeof(FNiagaraGlobalParameters);
+				TotalGPUParamSize += 2 * sizeof(FNiagaraSystemParameters);
+				TotalGPUParamSize += 2 * sizeof(FNiagaraOwnerParameters);
+				FirstGpuEmitter = false;
+			}
 
-				TotalGPUParamSize += 2 * sizeof(FNiagaraEmitterParameters);
+			TotalGPUParamSize += 2 * sizeof(FNiagaraEmitterParameters);
 				
-				TotalGPUParamSize += Inst.GetGPUContext()->CombinedParamStore.GetPaddedParameterSizeInBytes();
-				ActiveGPUEmitterCount++;
-			}
+			TotalGPUParamSize += Inst.GetGPUContext()->CombinedParamStore.GetPaddedParameterSizeInBytes();
+			ActiveGPUEmitterCount++;
 		}
 	}
 
