Fix for non-deterministic PostLoad when compiling Niagara script during cook

diff --git a/Engine/Plugins/FX/Niagara/Source/Niagara/Private/NiagaraParameterCollection.cpp b/Engine/Plugins/FX/Niagara/Source/Niagara/Private/NiagaraParameterCollection.cpp
--- a/Engine/Plugins/FX/Niagara/Source/Niagara/Private/NiagaraParameterCollection.cpp
+++ b/Engine/Plugins/FX/Niagara/Source/Niagara/Private/NiagaraParameterCollection.cpp
@@ -370,6 +370,9 @@ void UNiagaraParameterCollection::MakeNamespaceNameUnique()
 void UNiagaraParameterCollection::PostLoad()
 {
 	Super::PostLoad();
+
+	DefaultInstance->ConditionalPostLoad();
+
 	if (CompileId.IsValid() == false)
 	{
 		CompileId = FGuid::NewGuid();
diff --git a/Engine/Plugins/FX/Niagara/Source/NiagaraEditor/Private/EdGraphSchema_Niagara.cpp b/Engine/Plugins/FX/Niagara/Source/NiagaraEditor/Private/EdGraphSchema_Niagara.cpp
--- a/Engine/Plugins/FX/Niagara/Source/NiagaraEditor/Private/EdGraphSchema_Niagara.cpp
+++ b/Engine/Plugins/FX/Niagara/Source/NiagaraEditor/Private/EdGraphSchema_Niagara.cpp
@@ -1366,6 +1366,8 @@ UNiagaraParameterCollection* UEdGraphSchema_Niagara::VariableIsFromParameterColl
 		{
 			if (UNiagaraParameterCollection* Collection = CastChecked<UNiagaraParameterCollection>(CollectionAsset.GetAsset()))
 			{
+				// asset may not have been fully loaded so give it a chance to do it's PostLoad
+				Collection->ConditionalPostLoad();
 				if (VarName.StartsWith(Collection->GetFullNamespace()))
 				{
 					return Collection;
@@ -1390,6 +1392,8 @@ UNiagaraParameterCollection* UEdGraphSchema_Niagara::VariableIsFromParameterColl
 		{
 			if (UNiagaraParameterCollection* Collection = CastChecked<UNiagaraParameterCollection>(CollectionAsset.GetAsset()))
 			{
+				// asset may not have been fully loaded so give it a chance to do it's PostLoad
+				Collection->ConditionalPostLoad();
 				if (VarName.StartsWith(Collection->GetFullNamespace()))
 				{
 					const TArray<FNiagaraVariable>& CollectionVariables = Collection->GetParameters();
