Added Deprecation Message to Niagara scripts to give module creators a way to signal why a module is deprecated and what users can do about it.

diff --git a/Engine/Plugins/FX/Niagara/Source/Niagara/Classes/NiagaraScript.h b/Engine/Plugins/FX/Niagara/Source/Niagara/Classes/NiagaraScript.h
--- a/Engine/Plugins/FX/Niagara/Source/Niagara/Classes/NiagaraScript.h
+++ b/Engine/Plugins/FX/Niagara/Source/Niagara/Classes/NiagaraScript.h
@@ -386,6 +386,10 @@ public:
 	UPROPERTY(AssetRegistrySearchable, EditAnywhere, Category = Script)
 	uint32 bDeprecated : 1;
 
+	/* Message to display when the script is deprecated. */
+	UPROPERTY(EditAnywhere, Category = Script, meta = (EditCondition = "bDeprecated", MultiLine = true))
+	FText DeprecationMessage;
+
 	/* Which script to use if this is deprecated.*/
 	UPROPERTY(EditAnywhere, Category = Script, meta = (EditCondition = "bDeprecated"))
 	UNiagaraScript* DeprecationRecommendation;
diff --git a/Engine/Plugins/FX/Niagara/Source/Niagara/Private/NiagaraScript.cpp b/Engine/Plugins/FX/Niagara/Source/Niagara/Private/NiagaraScript.cpp
--- a/Engine/Plugins/FX/Niagara/Source/Niagara/Private/NiagaraScript.cpp
+++ b/Engine/Plugins/FX/Niagara/Source/Niagara/Private/NiagaraScript.cpp
@@ -962,7 +962,9 @@ void UNiagaraScript::PostEditChangeProperty(FPropertyChangedEvent& PropertyChang
 
 	CacheResourceShadersForRendering(true);
 
-	if (PropertyName == GET_MEMBER_NAME_CHECKED(UNiagaraScript, bDeprecated) || PropertyName == GET_MEMBER_NAME_CHECKED(UNiagaraScript, DeprecationRecommendation))
+	if (PropertyName == GET_MEMBER_NAME_CHECKED(UNiagaraScript, bDeprecated) || 
+		PropertyName == GET_MEMBER_NAME_CHECKED(UNiagaraScript, DeprecationMessage) ||
+		PropertyName == GET_MEMBER_NAME_CHECKED(UNiagaraScript, DeprecationRecommendation))
 	{
 		if (Source)
 		{
@@ -970,6 +972,15 @@ void UNiagaraScript::PostEditChangeProperty(FPropertyChangedEvent& PropertyChang
 		}
 	}
 
+	if (PropertyName == GET_MEMBER_NAME_CHECKED(UNiagaraScript, bExperimental) || 
+		PropertyName == GET_MEMBER_NAME_CHECKED(UNiagaraScript, ExperimentalMessage))
+	{
+		if (Source)
+		{
+			Source->MarkNotSynchronized(TEXT("Experimental changed."));
+		}
+	}
+
 	CustomAssetRegistryTagCache.Reset();
 }
 
diff --git a/Engine/Plugins/FX/Niagara/Source/NiagaraEditor/Private/NiagaraNodeFunctionCall.cpp b/Engine/Plugins/FX/Niagara/Source/NiagaraEditor/Private/NiagaraNodeFunctionCall.cpp
--- a/Engine/Plugins/FX/Niagara/Source/NiagaraEditor/Private/NiagaraNodeFunctionCall.cpp
+++ b/Engine/Plugins/FX/Niagara/Source/NiagaraEditor/Private/NiagaraNodeFunctionCall.cpp
@@ -416,6 +416,39 @@ ENiagaraScriptUsage UNiagaraNodeFunctionCall::GetCalledUsage() const
 	return ENiagaraScriptUsage::Function;
 }
 
+static FText GetFormattedDeprecationMessage(const UNiagaraScript* FunctionScript, const FString& FunctionDisplayName)
+{
+	FFormatNamedArguments Args;
+	Args.Add(TEXT("NodeName"), FText::FromString(FunctionDisplayName));
+
+	if (FunctionScript->DeprecationRecommendation != nullptr)
+	{
+		Args.Add(TEXT("Recommendation"), FText::FromString(FunctionScript->DeprecationRecommendation->GetPathName()));
+	}
+
+	if (FunctionScript->DeprecationMessage.IsEmptyOrWhitespace() == false)
+	{
+		Args.Add(TEXT("Message"), FunctionScript->DeprecationMessage);
+	}
+
+	FText FormatString = LOCTEXT("DeprecationErrorFmtUnknown", "Function call \"{NodeName}\" is deprecated. No recommendation was provided.");
+
+	if (FunctionScript->DeprecationRecommendation != nullptr && FunctionScript->DeprecationMessage.IsEmptyOrWhitespace() == false)
+	{
+		FormatString = LOCTEXT("DeprecationErrorFmtMessageAndRecommendation", "Function call \"{NodeName}\" is deprecated. Reason: {Message}. Please use {Recommendation} instead.");
+	}
+	else if (FunctionScript->DeprecationRecommendation != nullptr)
+	{
+		FormatString = LOCTEXT("DeprecationErrorFmtRecommendation", "Function call \"{NodeName}\" is deprecated. Please use {Recommendation} instead.");
+	}
+	else if (FunctionScript->DeprecationMessage.IsEmptyOrWhitespace() == false)
+	{
+		FormatString = LOCTEXT("DeprecationErrorFmtMessage", "Function call \"{NodeName}\" is deprecated. Reason: {Message} ");
+	}
+
+	return FText::Format(FormatString, Args);
+}
+
 void UNiagaraNodeFunctionCall::Compile(class FHlslNiagaraTranslator* Translator, TArray<int32>& Outputs)
 {
 	TArray<int32> Inputs;
@@ -428,17 +461,9 @@ void UNiagaraNodeFunctionCall::Compile(class FHlslNiagaraTranslator* Translator,
 	{
 		if (FunctionScript->bDeprecated && IsNodeEnabled())
 		{
-			if (FunctionScript->DeprecationRecommendation)
-			{
-				Translator->Warning(FText::Format(LOCTEXT("DeprecationErrorFmtRecommendation", "Function call \"{0}\" is deprecated. Please use {1} instead."), FText::FromString(FunctionDisplayName),
-					FText::FromString(FunctionScript->DeprecationRecommendation->GetPathName())),
-					this, nullptr);
-			}
-			else
-			{
-				Translator->Warning(FText::Format(LOCTEXT("DeprecationErrorFmtUnknown", "Function call \"{0}\" is deprecated. No recommendation was provided."), FText::FromString(FunctionDisplayName)),
-					this, nullptr);
-			}
+			FText DeprecationMessage = GetFormattedDeprecationMessage(FunctionScript, FunctionDisplayName);
+
+			Translator->Warning(DeprecationMessage, this, nullptr);
 		}
 
 		TArray<UEdGraphPin*> CallerInputPins;
@@ -615,17 +640,15 @@ void UNiagaraNodeFunctionCall::UpdateNodeErrorMessage()
 		{
 			UEdGraphNode::bHasCompilerMessage = true;
 			ErrorType = EMessageSeverity::Warning;
-			FFormatNamedArguments Args;
-			Args.Add(TEXT("NodeName"), GetNodeTitle(ENodeTitleType::ListView));
-			Args.Add(TEXT("Recommendation"), FunctionScript->DeprecationRecommendation != nullptr ? FText::FromString(FunctionScript->DeprecationRecommendation->GetPathName()) : LOCTEXT("UnspecifiedName","Unspecified"));
-			UEdGraphNode::ErrorMsg = FText::Format(LOCTEXT("DeprecationWarning", "{NodeName} is deprecated. Suggested replacement: {Recommendation}"), Args).ToString();
+			
+			UEdGraphNode::ErrorMsg = GetFormattedDeprecationMessage(FunctionScript, FunctionDisplayName).ToString();
 		}
 		else if (FunctionScript->bExperimental)
 		{
 			UEdGraphNode::bHasCompilerMessage = true;
 			UEdGraphNode::ErrorType = EMessageSeverity::Info;
 
-			if (FunctionScript->ExperimentalMessage.IsEmpty())
+			if (FunctionScript->ExperimentalMessage.IsEmptyOrWhitespace())
 			{
 				UEdGraphNode::NodeUpgradeMessage = LOCTEXT("FunctionExperimental", "This function is marked as experimental, use with care!");
 			}
@@ -649,7 +672,7 @@ void UNiagaraNodeFunctionCall::UpdateNodeErrorMessage()
 			UEdGraphNode::bHasCompilerMessage = true;
 			UEdGraphNode::ErrorType = EMessageSeverity::Info;
 
-			if (Signature.ExperimentalMessage.IsEmpty())
+			if (Signature.ExperimentalMessage.IsEmptyOrWhitespace())
 			{
 				UEdGraphNode::NodeUpgradeMessage = LOCTEXT("FunctionExperimental", "This function is marked as experimental, use with care!");
 			}
diff --git a/Engine/Plugins/FX/Niagara/Source/NiagaraEditor/Private/ViewModels/Stack/NiagaraStackFunctionInput.cpp b/Engine/Plugins/FX/Niagara/Source/NiagaraEditor/Private/ViewModels/Stack/NiagaraStackFunctionInput.cpp
--- a/Engine/Plugins/FX/Niagara/Source/NiagaraEditor/Private/ViewModels/Stack/NiagaraStackFunctionInput.cpp
+++ b/Engine/Plugins/FX/Niagara/Source/NiagaraEditor/Private/ViewModels/Stack/NiagaraStackFunctionInput.cpp
@@ -418,9 +418,36 @@ void UNiagaraStackFunctionInput::RefreshChildrenInternal(const TArray<UNiagaraSt
 			{
 				if (InputValues.DynamicNode->FunctionScript->bDeprecated)
 				{
-					FText LongMessage = InputValues.DynamicNode->FunctionScript->DeprecationRecommendation != nullptr ?
-						FText::Format(LOCTEXT("DynamicInputScriptDeprecationLong", "The script asset for the assigned dynamic input {0} has been deprecated. Suggested replacement: {1}"), FText::FromString(InputValues.DynamicNode->GetFunctionName()), FText::FromString(InputValues.DynamicNode->FunctionScript->DeprecationRecommendation->GetPathName())) :
-						FText::Format(LOCTEXT("DynamicInputScriptDeprecationUnknownLong", "The script asset for the assigned dynamic input {0} has been deprecated."), FText::FromString(InputValues.DynamicNode->GetFunctionName()));
+					FFormatNamedArguments Args;
+					Args.Add(TEXT("ScriptName"), FText::FromString(InputValues.DynamicNode->GetFunctionName()));
+
+					if (InputValues.DynamicNode->FunctionScript->DeprecationRecommendation != nullptr)
+					{
+						Args.Add(TEXT("Recommendation"), FText::FromString(InputValues.DynamicNode->FunctionScript->DeprecationRecommendation->GetPathName()));
+					}
+
+					if (InputValues.DynamicNode->FunctionScript->DeprecationMessage.IsEmptyOrWhitespace() == false)
+					{
+						Args.Add(TEXT("Message"), InputValues.DynamicNode->FunctionScript->DeprecationMessage);
+					}
+
+					FText FormatString = LOCTEXT("DynamicInputScriptDeprecationUnknownLong", "The script asset for the assigned dynamic input {ScriptName} has been deprecated.");
+
+					if (InputValues.DynamicNode->FunctionScript->DeprecationRecommendation != nullptr &&
+						InputValues.DynamicNode->FunctionScript->DeprecationMessage.IsEmptyOrWhitespace() == false)
+					{
+						FormatString = LOCTEXT("DynamicInputScriptDeprecationMessageAndRecommendationLong", "The script asset for the assigned dynamic input {ScriptName} has been deprecated. Reason: {Message}. Suggested replacement: {Recommendation}");
+					}
+					else if (InputValues.DynamicNode->FunctionScript->DeprecationRecommendation != nullptr)
+					{
+						FormatString = LOCTEXT("DynamicInputScriptDeprecationLong", "The script asset for the assigned dynamic input {ScriptName} has been deprecated. Suggested replacement: {Recommendation}");
+					}
+					else if (InputValues.DynamicNode->FunctionScript->DeprecationMessage.IsEmptyOrWhitespace() == false)
+					{
+						FormatString = LOCTEXT("DynamicInputScriptDeprecationMessageLong", "The script asset for the assigned dynamic input {ScriptName} has been deprecated. Reason: {Message}");
+					}
+
+					FText LongMessage = FText::Format(FormatString, Args);
 
 					int32 AddIdx = NewIssues.Add(FStackIssue(
 						EStackIssueSeverity::Warning,
@@ -449,7 +476,7 @@ void UNiagaraStackFunctionInput::RefreshChildrenInternal(const TArray<UNiagaraSt
 				if (InputValues.DynamicNode->FunctionScript->bExperimental)
 				{
 					FText ErrorMessage;
-					if (InputValues.DynamicNode->FunctionScript->ExperimentalMessage.IsEmpty())
+					if (InputValues.DynamicNode->FunctionScript->ExperimentalMessage.IsEmptyOrWhitespace())
 					{
 						ErrorMessage = FText::Format(LOCTEXT("DynamicInputScriptExperimental", "The script asset for the dynamic input {0} is experimental, use with care!"), FText::FromString(InputValues.DynamicNode->GetFunctionName()));
 					}
diff --git a/Engine/Plugins/FX/Niagara/Source/NiagaraEditor/Private/ViewModels/Stack/NiagaraStackModuleItem.cpp b/Engine/Plugins/FX/Niagara/Source/NiagaraEditor/Private/ViewModels/Stack/NiagaraStackModuleItem.cpp
--- a/Engine/Plugins/FX/Niagara/Source/NiagaraEditor/Private/ViewModels/Stack/NiagaraStackModuleItem.cpp
+++ b/Engine/Plugins/FX/Niagara/Source/NiagaraEditor/Private/ViewModels/Stack/NiagaraStackModuleItem.cpp
@@ -301,9 +301,36 @@ void UNiagaraStackModuleItem::RefreshIssues(TArray<FStackIssue>& NewIssues)
 				FText ModuleScriptDeprecationShort = LOCTEXT("ModuleScriptDeprecationShort", "Deprecated module");
 				if (CanMoveAndDelete())
 				{
-					FText LongMessage = FunctionCallNode->FunctionScript->DeprecationRecommendation != nullptr ?
-						FText::Format(LOCTEXT("ModuleScriptDeprecationLong", "The script asset for the assigned module {0} has been deprecated. Suggested replacement: {1}"), FText::FromString(FunctionCallNode->GetFunctionName()), FText::FromString(FunctionCallNode->FunctionScript->DeprecationRecommendation->GetPathName())) :
-						FText::Format(LOCTEXT("ModuleScriptDeprecationUnknownLong", "The script asset for the assigned module {0} has been deprecated."), FText::FromString(FunctionCallNode->GetFunctionName()));
+					FFormatNamedArguments Args;
+					Args.Add(TEXT("ScriptName"), FText::FromString(FunctionCallNode->GetFunctionName()));
+
+					if (FunctionCallNode->FunctionScript->DeprecationRecommendation != nullptr)
+					{
+						Args.Add(TEXT("Recommendation"), FText::FromString(FunctionCallNode->FunctionScript->DeprecationRecommendation->GetPathName()));
+					}
+
+					if (FunctionCallNode->FunctionScript->DeprecationMessage.IsEmptyOrWhitespace() == false)
+					{
+						Args.Add(TEXT("Message"), FunctionCallNode->FunctionScript->DeprecationMessage);
+					}
+
+					FText FormatString = LOCTEXT("ModuleScriptDeprecationUnknownLong", "The script asset for the assigned module {ScriptName} has been deprecated.");
+
+					if (FunctionCallNode->FunctionScript->DeprecationRecommendation != nullptr &&
+						FunctionCallNode->FunctionScript->DeprecationMessage.IsEmptyOrWhitespace() == false)
+					{
+						FormatString = LOCTEXT("ModuleScriptDeprecationMessageAndRecommendationLong", "The script asset for the assigned module {ScriptName} has been deprecated. Reason: {Message}. Suggested replacement: {Recommendation}");
+					}
+					else if (FunctionCallNode->FunctionScript->DeprecationRecommendation != nullptr)
+					{
+						FormatString = LOCTEXT("ModuleScriptDeprecationLong", "The script asset for the assigned module {ScriptName} has been deprecated. Suggested replacement: {Recommendation}");
+					}
+					else if (FunctionCallNode->FunctionScript->DeprecationMessage.IsEmptyOrWhitespace() == false)
+					{
+						FormatString = LOCTEXT("ModuleScriptDeprecationMessageLong", "The script asset for the assigned module {ScriptName} has been deprecated. Reason: {Message}");
+					}
+
+					FText LongMessage = FText::Format(FormatString, Args);
 
 					int32 AddIdx = NewIssues.Add(FStackIssue(
 						EStackIssueSeverity::Warning,
@@ -343,7 +370,7 @@ void UNiagaraStackModuleItem::RefreshIssues(TArray<FStackIssue>& NewIssues)
 			if (FunctionCallNode->FunctionScript->bExperimental)
 			{
 				FText ErrorMessage;
-				if (FunctionCallNode->FunctionScript->ExperimentalMessage.IsEmpty())
+				if (FunctionCallNode->FunctionScript->ExperimentalMessage.IsEmptyOrWhitespace())
 				{
 					ErrorMessage = FText::Format(LOCTEXT("ModuleScriptExperimental", "The script asset for this module is experimental, use with care!"), FText::FromString(FunctionCallNode->GetFunctionName()));
 				}
