Niagara - Add a speculative fix for a cook crash and add some code to detect the crash situation if it can still happen to avoid it and print a bunch of log information.

diff --git a/Engine/Plugins/FX/Niagara/Source/Niagara/Private/NiagaraScript.cpp b/Engine/Plugins/FX/Niagara/Source/Niagara/Private/NiagaraScript.cpp
--- a/Engine/Plugins/FX/Niagara/Source/Niagara/Private/NiagaraScript.cpp
+++ b/Engine/Plugins/FX/Niagara/Source/Niagara/Private/NiagaraScript.cpp
@@ -616,9 +616,40 @@ void UNiagaraScript::PreSave(const class ITargetPlatform* TargetPlatform)
 	Super::PreSave(TargetPlatform);
 
 #if WITH_EDITORONLY_DATA
+	// Pre-save can happen in any order for objects in the package and since this is now used to cache data for execution we need to make sure that the system compilation
+	// is complete before caching the executable data.
+	UNiagaraSystem* SystemOwner = FindRootSystem();
+	if (SystemOwner)
+	{
+		SystemOwner->WaitForCompilationComplete();
+	}
+
 	ScriptExecutionParamStore.Empty();
 	ScriptExecutionBoundParameters.Empty();
 
+	// Make sure the data interfaces are consistent to prevent crashes in later caching operations.
+	if (CachedScriptVM.DataInterfaceInfo.Num() != CachedDefaultDataInterfaces.Num())
+	{
+		UE_LOG(LogNiagara, Warning, TEXT("Data interface count mistmatch during script presave. Invaliding compile results (see full log for details).  Script: %s"), *GetPathName());
+		UE_LOG(LogNiagara, Log, TEXT("Compiled DataInterfaceInfos:"));
+		for (const FNiagaraScriptDataInterfaceCompileInfo& DataInterfaceCompileInfo : CachedScriptVM.DataInterfaceInfo)
+		{
+			UE_LOG(LogNiagara, Log, TEXT("Name:%s, Type: %s"), *DataInterfaceCompileInfo.Name.ToString(), *DataInterfaceCompileInfo.Type.GetName());
+		}
+		UE_LOG(LogNiagara, Log, TEXT("Cached DataInterfaceInfos:"));
+		for (const FNiagaraScriptDataInterfaceInfo& DataInterfaceCacheInfo : CachedDefaultDataInterfaces)
+		{
+			UE_LOG(LogNiagara, Log, TEXT("Name:%s, Type: %s, Path:%s"), 
+				*DataInterfaceCacheInfo.Name.ToString(), *DataInterfaceCacheInfo.Type.GetName(), 
+				DataInterfaceCacheInfo.DataInterface != nullptr 
+					? *DataInterfaceCacheInfo.DataInterface->GetPathName() 
+					: TEXT("None"));
+		}
+
+		InvalidateCompileResults();
+		return;
+	}
+
 	if (TargetPlatform && TargetPlatform->RequiresCookedData())
 	{
 		TOptional<ENiagaraSimTarget> SimTarget = GetSimTarget();
@@ -1786,6 +1817,7 @@ void UNiagaraScript::InvalidateCompileResults()
 	ScriptResource.Invalidate();
 	CachedScriptVMId.Invalidate();
 	LastGeneratedVMId.Invalidate();
+	CachedDefaultDataInterfaces.Reset();
 }
 
 
