Fix for a mesh with multiple materials causing a mesh icon to show up in the timeline twice
#rnx

diff --git a/Engine/Plugins/FX/Niagara/Source/NiagaraEditor/Private/Sequencer/NiagaraSequence/NiagaraEmitterTrackEditor.cpp b/Engine/Plugins/FX/Niagara/Source/NiagaraEditor/Private/Sequencer/NiagaraSequence/NiagaraEmitterTrackEditor.cpp
--- a/Engine/Plugins/FX/Niagara/Source/NiagaraEditor/Private/Sequencer/NiagaraSequence/NiagaraEmitterTrackEditor.cpp
+++ b/Engine/Plugins/FX/Niagara/Source/NiagaraEditor/Private/Sequencer/NiagaraSequence/NiagaraEmitterTrackEditor.cpp
@@ -80,10 +80,10 @@ public:
 
 		// Renderer buttons
 
-		EmitterTrack->GetEmitterHandleViewModel()->GetRendererPreviewData(RendererPreviewData);
-		for (FRendererPreviewData* RendererPreview : RendererPreviewData)
+		EmitterTrack->GetEmitterHandleViewModel()->GetRendererEntries(RendererEntryData);
+		for (UNiagaraStackEntry* RendererEntry : RendererEntryData)
 		{
-			if (UNiagaraStackRendererItem* RendererItem = Cast<UNiagaraStackRendererItem>(RendererPreview->RenderingEntry))
+			if (UNiagaraStackRendererItem* RendererItem = Cast<UNiagaraStackRendererItem>(RendererEntry))
 			{
 				UNiagaraRendererProperties* Renderer = RendererItem->GetRendererProperties();
 				TrackBox->AddSlot()
@@ -95,7 +95,7 @@ public:
 						.ButtonStyle(FEditorStyle::Get(), "HoverHintOnly")
 						.IsFocusable(false)
 						.ToolTipText(FText::Format(LOCTEXT("RenderButtonToolTip", "{0} - Press to select."), FText::FromString(FName::NameToDisplayString(Renderer->GetName(), false))))
-						.OnClicked(this, &SEmitterTrackWidget::OnRenderButtonClicked, RendererPreview)
+						.OnClicked(this, &SEmitterTrackWidget::OnRenderButtonClicked, RendererEntry)
 						[
 							SNew(SImage)
 							.Image(FSlateIconFinder::FindIconBrushForClass(Renderer->GetClass()))
@@ -190,12 +190,12 @@ private:
 			: FLinearColor::Gray;
 	}
 
-	FReply OnRenderButtonClicked(FRendererPreviewData* InRendererPreview)
+	FReply OnRenderButtonClicked(UNiagaraStackEntry* InRendererEntry)
 	{
 		if (EmitterTrack.IsValid())
 		{
 			TArray<UNiagaraStackEntry*> SelectedEntries;
-			SelectedEntries.Add(InRendererPreview->RenderingEntry);
+			SelectedEntries.Add(InRendererEntry);
 			TArray<UNiagaraStackEntry*> DeselectedEntries;
 			EmitterTrack->GetSystemViewModel().GetSelectionViewModel()->UpdateSelectedEntries(SelectedEntries, DeselectedEntries, true);
 		}
@@ -215,7 +215,7 @@ private:
 private:
 	TWeakObjectPtr<UMovieSceneNiagaraEmitterTrack> EmitterTrack;
 	mutable TOptional<FText> TrackErrorIconToolTip;
-	TArray<FRendererPreviewData*> RendererPreviewData;
+	TArray<UNiagaraStackEntry*> RendererEntryData;
 };
 
 FNiagaraEmitterTrackEditor::FNiagaraEmitterTrackEditor(TSharedPtr<ISequencer> Sequencer) 
diff --git a/Engine/Plugins/FX/Niagara/Source/NiagaraEditor/Private/ViewModels/NiagaraEmitterHandleViewModel.cpp b/Engine/Plugins/FX/Niagara/Source/NiagaraEditor/Private/ViewModels/NiagaraEmitterHandleViewModel.cpp
--- a/Engine/Plugins/FX/Niagara/Source/NiagaraEditor/Private/ViewModels/NiagaraEmitterHandleViewModel.cpp
+++ b/Engine/Plugins/FX/Niagara/Source/NiagaraEditor/Private/ViewModels/NiagaraEmitterHandleViewModel.cpp
@@ -73,6 +73,25 @@ void FNiagaraEmitterHandleViewModel::GetRendererPreviewData(TArray<FRendererPrev
 	}
 }
 
+void FNiagaraEmitterHandleViewModel::GetRendererEntries(TArray<UNiagaraStackEntry*>& InRenderingEntries)
+{
+	InRenderingEntries.Empty();
+	UNiagaraEmitter* Emitter = GetEmitterHandle()->GetInstance();
+	UNiagaraStackRoot* StackRoot = Cast<UNiagaraStackRoot>(EmitterStackViewModel->GetRootEntry());
+	if (StackRoot)
+	{
+		TArray<UNiagaraStackEntry*> Children;
+		StackRoot->GetRenderGroup()->GetUnfilteredChildren(Children);
+		for (UNiagaraStackEntry* Child : Children)
+		{
+			if (UNiagaraStackRendererItem* RendererItem = Cast<UNiagaraStackRendererItem>(Child))
+			{
+				InRenderingEntries.Add(Child);
+			}
+		}
+	}
+}
+
 TSharedRef<FNiagaraSystemViewModel> FNiagaraEmitterHandleViewModel::GetOwningSystemViewModel() const
 {
 	TSharedPtr<FNiagaraSystemViewModel> OwningSystemViewModelPinned = OwningSystemViewModelWeak.Pin();
diff --git a/Engine/Plugins/FX/Niagara/Source/NiagaraEditor/Public/ViewModels/NiagaraEmitterHandleViewModel.h b/Engine/Plugins/FX/Niagara/Source/NiagaraEditor/Public/ViewModels/NiagaraEmitterHandleViewModel.h
--- a/Engine/Plugins/FX/Niagara/Source/NiagaraEditor/Public/ViewModels/NiagaraEmitterHandleViewModel.h
+++ b/Engine/Plugins/FX/Niagara/Source/NiagaraEditor/Public/ViewModels/NiagaraEmitterHandleViewModel.h
@@ -134,6 +134,7 @@ public:
 
 	NIAGARAEDITOR_API void GetRendererPreviewData(TArray<FRendererPreviewData*>& InRendererPreviewData);
 
+	NIAGARAEDITOR_API void GetRendererEntries(TArray<UNiagaraStackEntry*>& InRenderingEntries);
 	NIAGARAEDITOR_API TSharedRef<FNiagaraSystemViewModel> GetOwningSystemViewModel() const;
 
 private:
