Fixed up support for having GPU mesh emitters using a mesh that has multiple sections

diff --git a/Engine/Plugins/FX/Niagara/Shaders/Private/NiagaraDrawIndirectArgsGen.usf b/Engine/Plugins/FX/Niagara/Shaders/Private/NiagaraDrawIndirectArgsGen.usf
--- a/Engine/Plugins/FX/Niagara/Shaders/Private/NiagaraDrawIndirectArgsGen.usf
+++ b/Engine/Plugins/FX/Niagara/Shaders/Private/NiagaraDrawIndirectArgsGen.usf
@@ -7,7 +7,7 @@
 #include "/Engine/Private/Common.ush"
 
 // Array of task info. Each being a draw indirect arg gen or instance count reset.
-// ArgGen = (InstanceCountBufferOffset, NumIndicesPerInstance)
+// ArgGen = (InstanceCountBufferOffset, NumIndicesPerInstance, StartIndexLocation)
 // InstanceCountClear = (InstanceCountBufferOffset)
 Buffer<uint> TaskInfos;
 
@@ -22,16 +22,16 @@ void MainCS(uint TaskIndex : SV_DispatchThreadID)
 	if (TaskIndex < TaskCount.x)
 	{
 		const uint ArgOffset = NIAGARA_DRAW_INDIRECT_ARGS_SIZE * TaskIndex;
-		const uint InfoOffset = TaskIndex * 2;
+		const uint InfoOffset = TaskIndex * NIAGARA_DRAW_INDIRECT_TASK_INFO_SIZE;
 		RWDrawIndirectArgs[ArgOffset + 0] = TaskInfos[InfoOffset + 1]; // NumIndicesPerInstance
 		RWDrawIndirectArgs[ArgOffset + 1] = RWInstanceCounts[TaskInfos[InfoOffset + 0]]; // InstanceCount
-		RWDrawIndirectArgs[ArgOffset + 2] = 0; // StartIndexLocation
+		RWDrawIndirectArgs[ArgOffset + 2] = TaskInfos[InfoOffset + 2]; // StartIndexLocation
 		RWDrawIndirectArgs[ArgOffset + 3] = 0; // BaseVertexLocation
 		RWDrawIndirectArgs[ArgOffset + 4] = 0; // StartInstanceLocation
 	}
 	else if (TaskIndex < TaskCount.z)
 	{
-		const uint InfoOffset = TaskCount.x * 2 + (TaskIndex - TaskCount.x);
+		const uint InfoOffset = TaskCount.x * NIAGARA_DRAW_INDIRECT_TASK_INFO_SIZE + (TaskIndex - TaskCount.x);
 		RWInstanceCounts[TaskInfos[InfoOffset + 0]] = 0;
 	}
 }
\ No newline at end of file
diff --git a/Engine/Plugins/FX/Niagara/Source/Niagara/Classes/NiagaraGPUInstanceCountManager.h b/Engine/Plugins/FX/Niagara/Source/Niagara/Classes/NiagaraGPUInstanceCountManager.h
--- a/Engine/Plugins/FX/Niagara/Source/Niagara/Classes/NiagaraGPUInstanceCountManager.h
+++ b/Engine/Plugins/FX/Niagara/Source/Niagara/Classes/NiagaraGPUInstanceCountManager.h
@@ -23,7 +23,7 @@ public:
 
 FORCEINLINE uint32 GetTypeHash(const FNiagaraDrawIndirectArgsGenCS::FArgGenTaskInfo& Info)
 {
-	return (Info.InstanceCountBufferOffset << 8) | Info.NumIndicesPerInstance;
+	return HashCombine(Info.InstanceCountBufferOffset, HashCombine(Info.NumIndicesPerInstance, Info.StartIndexLocation));
 }
 
 /** 
@@ -58,7 +58,7 @@ public:
 	bool HasPendingGPUReadback() const;
 
 	/** Add a draw indirect task to generate the draw indirect args. Returns the draw indirect arg buffer offset. */
-	uint32 AddDrawIndirect(uint32 InstanceCountBufferOffset, uint32 NumIndicesPerInstance);
+	uint32 AddDrawIndirect(uint32 InstanceCountBufferOffset, uint32 NumIndicesPerInstance, uint32 StartIndexLocation);
 	FRWBuffer& GetDrawIndirectBuffer() { return DrawIndirectBuffer; }
 
 	/** 
diff --git a/Engine/Plugins/FX/Niagara/Source/Niagara/Private/NiagaraGPUInstanceCountManager.cpp b/Engine/Plugins/FX/Niagara/Source/Niagara/Private/NiagaraGPUInstanceCountManager.cpp
--- a/Engine/Plugins/FX/Niagara/Source/Niagara/Private/NiagaraGPUInstanceCountManager.cpp
+++ b/Engine/Plugins/FX/Niagara/Source/Niagara/Private/NiagaraGPUInstanceCountManager.cpp
@@ -161,11 +161,11 @@ void FNiagaraGPUInstanceCountManager::ResizeBuffers(FRHICommandListImmediate& RH
 	INC_DWORD_STAT_BY(STAT_NiagaraMaxNumGPURenderers, MaxDrawIndirectArgs);
 }
 
-uint32 FNiagaraGPUInstanceCountManager::AddDrawIndirect(uint32 InstanceCountBufferOffset, uint32 NumIndicesPerInstance)
+uint32 FNiagaraGPUInstanceCountManager::AddDrawIndirect(uint32 InstanceCountBufferOffset, uint32 NumIndicesPerInstance, uint32 StartIndexLocation)
 {
 	checkSlow(IsInRenderingThread());
 
-	const FArgGenTaskInfo Info(InstanceCountBufferOffset, NumIndicesPerInstance);
+	const FArgGenTaskInfo Info(InstanceCountBufferOffset, NumIndicesPerInstance, StartIndexLocation);
 
 	uint32& CachedOffset = DrawIndirectArgMap.FindOrAdd(Info, INDEX_NONE);
 	if (CachedOffset != INDEX_NONE)
diff --git a/Engine/Plugins/FX/Niagara/Source/Niagara/Private/NiagaraMeshRendererProperties.cpp b/Engine/Plugins/FX/Niagara/Source/Niagara/Private/NiagaraMeshRendererProperties.cpp
--- a/Engine/Plugins/FX/Niagara/Source/Niagara/Private/NiagaraMeshRendererProperties.cpp
+++ b/Engine/Plugins/FX/Niagara/Source/Niagara/Private/NiagaraMeshRendererProperties.cpp
@@ -185,7 +185,26 @@ void UNiagaraMeshRendererProperties::GetUsedMaterials(const FNiagaraEmitterInsta
 uint32 UNiagaraMeshRendererProperties::GetNumIndicesPerInstance() const
 {
 	// TODO: Add proper support for multiple mesh sections for GPU mesh particles.
-	return ParticleMesh ? ParticleMesh->RenderData->LODResources[0].Sections[0].NumTriangles * 3 : 0;
+	//return ParticleMesh ? ParticleMesh->RenderData->LODResources[0].Sections[0].NumTriangles * 3 : 0;
+	return ParticleMesh ? ParticleMesh->RenderData->LODResources[0].IndexBuffer.GetNumIndices() : 0;
+}
+
+void UNiagaraMeshRendererProperties::GetIndexInfoPerSection(int32 LODIndex, TArray<TPair<int32, int32>>& IndexInfoPerSection) const
+{
+	check(ParticleMesh && ParticleMesh->RenderData->LODResources.IsValidIndex(LODIndex));
+
+	if (ParticleMesh && ParticleMesh->RenderData->LODResources.IsValidIndex(LODIndex))
+	{
+		const FStaticMeshLODResources& MeshLod = ParticleMesh->RenderData->LODResources[LODIndex];
+		const int32 SectionCount = MeshLod.Sections.Num();
+		IndexInfoPerSection.SetNum(SectionCount);
+
+		for (int32 SectionIdx = 0; SectionIdx < SectionCount; ++SectionIdx)
+		{
+			IndexInfoPerSection[SectionIdx].Key = MeshLod.Sections[SectionIdx].NumTriangles * 3;
+			IndexInfoPerSection[SectionIdx].Value = MeshLod.Sections[SectionIdx].FirstIndex;
+		}
+	}
 }
 
 
diff --git a/Engine/Plugins/FX/Niagara/Source/Niagara/Private/NiagaraRenderer.cpp b/Engine/Plugins/FX/Niagara/Source/Niagara/Private/NiagaraRenderer.cpp
--- a/Engine/Plugins/FX/Niagara/Source/Niagara/Private/NiagaraRenderer.cpp
+++ b/Engine/Plugins/FX/Niagara/Source/Niagara/Private/NiagaraRenderer.cpp
@@ -284,7 +284,7 @@ void FNiagaraRenderer::CreateRenderThreadResources(NiagaraEmitterInstanceBatcher
 		NumRegisteredGPURenderers = Batcher->GetGPUInstanceCounterManager().GetGPURendererCount();
 		if (NumRegisteredGPURenderers)
 		{
-			NumRegisteredGPURenderers->Value += 1;
+			NumRegisteredGPURenderers->Value += GetMaxIndirectArgs();
 		}
 	}
 }
@@ -293,7 +293,7 @@ void FNiagaraRenderer::ReleaseRenderThreadResources()
 {
 	if (NumRegisteredGPURenderers)
 	{
-		NumRegisteredGPURenderers->Value -= 1;
+		NumRegisteredGPURenderers->Value -= GetMaxIndirectArgs();
 		NumRegisteredGPURenderers.SafeRelease();
 	}
 }
diff --git a/Engine/Plugins/FX/Niagara/Source/Niagara/Private/NiagaraRendererMeshes.cpp b/Engine/Plugins/FX/Niagara/Source/Niagara/Private/NiagaraRendererMeshes.cpp
--- a/Engine/Plugins/FX/Niagara/Source/Niagara/Private/NiagaraRendererMeshes.cpp
+++ b/Engine/Plugins/FX/Niagara/Source/Niagara/Private/NiagaraRendererMeshes.cpp
@@ -122,6 +122,14 @@ FNiagaraRendererMeshes::FNiagaraRendererMeshes(ERHIFeatureLevel::Type FeatureLev
 	MaterialParamValidMask |= SetVertexFactoryVariable(Data, Properties->DynamicMaterial3Binding.DataSetVariable, ENiagaraMeshVFLayout::DynamicParam3) ? 0x8 : 0;
 
 	MeshMinimumLOD = Properties->ParticleMesh->MinLOD.GetValueForFeatureLevel(FeatureLevel);
+
+	const int32 LODCount = MeshRenderData->LODResources.Num();
+	IndexInfoPerSection.SetNum(LODCount);
+
+	for (int32 LODIdx = 0; LODIdx < LODCount; ++LODIdx)
+	{
+		Properties->GetIndexInfoPerSection(LODIdx, IndexInfoPerSection[LODIdx]);
+	}
 }
 
 FNiagaraRendererMeshes::~FNiagaraRendererMeshes()
@@ -144,9 +152,17 @@ void FNiagaraRendererMeshes::ReleaseRenderThreadResources()
 	VertexFactories.Empty();
 }
 
-void FNiagaraRendererMeshes::CreateRenderThreadResources(NiagaraEmitterInstanceBatcher* Batcher)
+int32 FNiagaraRendererMeshes::GetMaxIndirectArgs() const
 {
-	FNiagaraRenderer::CreateRenderThreadResources(Batcher);
+	// currently the most indirect args we can add would be for a single lod, so search for the LOD with the highest number of sections
+	// this value should be constant for the life of the renderer as it is being used for NumRegisteredGPURenderers
+	int32 MaxSectionCount = 0;
+
+	for (const auto& IndexInfo : IndexInfoPerSection)
+	{
+		MaxSectionCount = FMath::Max(MaxSectionCount, IndexInfo.Num());
+	}
+	return MaxSectionCount;
 }
 
 void FNiagaraRendererMeshes::SetupVertexFactory(FNiagaraMeshVertexFactory *InVertexFactory, const FStaticMeshLODResources& LODResources) const
@@ -160,6 +176,24 @@ void FNiagaraRendererMeshes::SetupVertexFactory(FNiagaraMeshVertexFactory *InVer
 	InVertexFactory->SetData(Data);
 }
 
+int32 FNiagaraRendererMeshes::GetLODIndex() const
+{
+	const int32 LODCount = MeshRenderData->LODResources.Num();
+
+	// Doesn't seem to work for some reason. See comment in FDynamicMeshEmitterData::GetMeshLODIndexFromProxy()
+	// const int32 LODIndex = FMath::Max<int32>((int32)MeshRenderData->CurrentFirstLODIdx, MeshMinimumLOD);
+	int32 LODIndex = FMath::Clamp<int32>(MeshRenderData->CurrentFirstLODIdx, 0, LODCount - 1);
+
+	while (LODIndex < LODCount && !MeshRenderData->LODResources[LODIndex].GetNumVertices())
+	{
+		++LODIndex;
+	}
+
+	check(MeshRenderData->LODResources[LODIndex].GetNumVertices());
+
+	return LODIndex;
+}
+
 void FNiagaraRendererMeshes::GetDynamicMeshElements(const TArray<const FSceneView*>& Views, const FSceneViewFamily& ViewFamily, uint32 VisibilityMap, FMeshElementCollector& Collector, const FNiagaraSceneProxy *SceneProxy) const
 {
 	check(SceneProxy);
@@ -191,8 +225,8 @@ void FNiagaraRendererMeshes::GetDynamicMeshElements(const TArray<const FSceneVie
 	FScopeCycleCounter EmitterStatsCounter(EmitterStatID);
 #endif
 
-	// @TODO : support multiple LOD and section, using an inlined array and/or the SceneRenderingAllocator
-	uint32 IndirectArgsOffset = INDEX_NONE;
+	// @TODO : support multiple LOD
+	TArray<uint32, TInlineAllocator<8>> IndirectArgsOffsets;
 	int32 NumInstances = SourceParticleData->GetNumInstances();
 
 	FGlobalDynamicReadBuffer& DynamicReadBuffer = Collector.GetDynamicReadBuffer();
@@ -228,11 +262,22 @@ void FNiagaraRendererMeshes::GetDynamicMeshElements(const TArray<const FSceneVie
 			FMemory::Memcpy(ParticleData.Buffer, SourceParticleData->GetFloatBuffer().GetData(), SourceParticleData->GetFloatBuffer().Num());
 		}
 	}
-
+	
+	const int32 LODIndex = GetLODIndex();
+	const FStaticMeshLODResources& LODModel = MeshRenderData->LODResources[LODIndex];
 
 	if (SimTarget == ENiagaraSimTarget::GPUComputeSim)
 	{
-		IndirectArgsOffset = Batcher->GetGPUInstanceCounterManager().AddDrawIndirect(SourceParticleData->GetGPUInstanceCountBufferOffset(), NumIndicesPerInstance);
+		const int32 SectionCount = LODModel.Sections.Num();
+
+		IndirectArgsOffsets.SetNum(SectionCount);
+		for (int32 SectionIdx = 0; SectionIdx < SectionCount; ++SectionIdx)
+		{
+			IndirectArgsOffsets[SectionIdx] = Batcher->GetGPUInstanceCounterManager().AddDrawIndirect(
+				SourceParticleData->GetGPUInstanceCountBufferOffset(),
+				IndexInfoPerSection[LODIndex][SectionIdx].Key,
+				IndexInfoPerSection[LODIndex][SectionIdx].Value);
+		}
 	}
 
 	{
@@ -245,15 +290,6 @@ void FNiagaraRendererMeshes::GetDynamicMeshElements(const TArray<const FSceneVie
 			{
 				const FSceneView* View = Views[ViewIndex];
 
-				// Doesn't seem to work for some reason. See comment in FDynamicMeshEmitterData::GetMeshLODIndexFromProxy()
-				// const int32 LODIndex = FMath::Max<int32>((int32)MeshRenderData->CurrentFirstLODIdx, MeshMinimumLOD);
-				int32 LODIndex = (int32)MeshRenderData->CurrentFirstLODIdx;
-				while (LODIndex < MeshRenderData->LODResources.Num() - 1 && !MeshRenderData->LODResources[LODIndex].GetNumVertices())
-				{
-					++LODIndex;
-				}
-				const FStaticMeshLODResources& LODModel = MeshRenderData->LODResources[LODIndex];
-
 				// Get the next vertex factory to use
 				FNiagaraMeshVertexFactory* VertexFactory = nullptr;
 				if (VertexFactories.IsValidIndex(VertexFactoryIndex))
@@ -402,9 +438,9 @@ void FNiagaraRendererMeshes::GetDynamicMeshElements(const TArray<const FSceneVie
 
 				// GPU mesh rendering currently only supports one mesh section.
 				// TODO: Add proper support for multiple mesh sections for GPU mesh particles.
-				int32 MaxSection = SimTarget == ENiagaraSimTarget::GPUComputeSim ? 1 : LODModel.Sections.Num();
+				const int32 SectionCount = LODModel.Sections.Num();
 				const bool bIsWireframe = AllowDebugViewmodes() && View && View->Family->EngineShowFlags.Wireframe;
-				for (int32 SectionIndex = 0; SectionIndex < MaxSection; SectionIndex++)
+				for (int32 SectionIndex = 0; SectionIndex < SectionCount; SectionIndex++)
 				{
 					const FStaticMeshSection& Section = LODModel.Sections[SectionIndex];
 					FMaterialRenderProxy* MaterialProxy = DynamicDataMesh->Materials[SectionIndex];
@@ -461,10 +497,10 @@ void FNiagaraRendererMeshes::GetDynamicMeshElements(const TArray<const FSceneVie
 						BatchElement.NumPrimitives = Section.NumTriangles;
 					}
 
-					if (IndirectArgsOffset != INDEX_NONE)
+					if (IndirectArgsOffsets.IsValidIndex(SectionIndex))
 					{
 						BatchElement.NumPrimitives = 0;
-						BatchElement.IndirectArgsOffset = IndirectArgsOffset;
+						BatchElement.IndirectArgsOffset = IndirectArgsOffsets[SectionIndex];
 						BatchElement.IndirectArgsBuffer = Batcher->GetGPUInstanceCounterManager().GetDrawIndirectBuffer().Buffer;
 					}
 					else
@@ -669,13 +705,7 @@ FNiagaraDynamicDataBase *FNiagaraRendererMeshes::GenerateDynamicData(const FNiag
 	{
 		DynamicData = new FNiagaraDynamicDataMesh(Emitter);
 
-		// Doesn't seem to work for some reason. See comment in FDynamicMeshEmitterData::GetMeshLODIndexFromProxy()
-		// const int32 LODIndex = FMath::Max<int32>((int32)MeshRenderData->CurrentFirstLODIdx, MeshMinimumLOD);
-		int32 LODIndex = (int32)MeshRenderData->CurrentFirstLODIdx;
-		while (LODIndex < MeshRenderData->LODResources.Num() - 1 && !MeshRenderData->LODResources[LODIndex].GetNumVertices())
-		{
-			++LODIndex;
-		}
+		const int32 LODIndex = GetLODIndex();
 		const FStaticMeshLODResources& LODModel = MeshRenderData->LODResources[LODIndex];
 
 		check(BaseMaterials_GT.Num() == LODModel.Sections.Num());
diff --git a/Engine/Plugins/FX/Niagara/Source/Niagara/Private/NiagaraRendererSprites.cpp b/Engine/Plugins/FX/Niagara/Source/Niagara/Private/NiagaraRendererSprites.cpp
--- a/Engine/Plugins/FX/Niagara/Source/Niagara/Private/NiagaraRendererSprites.cpp
+++ b/Engine/Plugins/FX/Niagara/Source/Niagara/Private/NiagaraRendererSprites.cpp
@@ -590,7 +590,7 @@ void FNiagaraRendererSprites::GetDynamicMeshElements(const TArray<const FSceneVi
 	uint32 IndirectArgsOffset = INDEX_NONE;
 	if (SimTarget == ENiagaraSimTarget::GPUComputeSim)
 	{
-		IndirectArgsOffset = Batcher->GetGPUInstanceCounterManager().AddDrawIndirect(SourceParticleData->GetGPUInstanceCountBufferOffset(), NumIndicesPerInstance);
+		IndirectArgsOffset = Batcher->GetGPUInstanceCounterManager().AddDrawIndirect(SourceParticleData->GetGPUInstanceCountBufferOffset(), NumIndicesPerInstance, 0);
 	}
 
 	for (int32 ViewIndex = 0; ViewIndex < Views.Num(); ViewIndex++)
@@ -637,7 +637,7 @@ void FNiagaraRendererSprites::GetDynamicRayTracingInstances(FRayTracingMaterialG
 	uint32 IndirectArgsOffset = INDEX_NONE;
 	if (SimTarget == ENiagaraSimTarget::GPUComputeSim)
 	{
-		IndirectArgsOffset = Batcher->GetGPUInstanceCounterManager().AddDrawIndirect(SourceParticleData->GetGPUInstanceCountBufferOffset(), NumIndicesPerInstance);
+		IndirectArgsOffset = Batcher->GetGPUInstanceCounterManager().AddDrawIndirect(SourceParticleData->GetGPUInstanceCountBufferOffset(), NumIndicesPerInstance, 0);
 	}
 
 	FRayTracingInstance RayTracingInstance;
diff --git a/Engine/Plugins/FX/Niagara/Source/Niagara/Public/NiagaraMeshRendererProperties.h b/Engine/Plugins/FX/Niagara/Source/Niagara/Public/NiagaraMeshRendererProperties.h
--- a/Engine/Plugins/FX/Niagara/Source/Niagara/Public/NiagaraMeshRendererProperties.h
+++ b/Engine/Plugins/FX/Niagara/Source/Niagara/Public/NiagaraMeshRendererProperties.h
@@ -91,6 +91,7 @@ public:
 #endif // WITH_EDITORONLY_DATA
 
 	virtual uint32 GetNumIndicesPerInstance() const final override;
+	void GetIndexInfoPerSection(int32 LODIndex, TArray<TPair<int32, int32>>& InfoPerSection) const;
 
 	/** The static mesh to be instanced when rendering mesh particles. If OverrideMaterial is not specified, the mesh's material is used. Note that the material must have the Niagara Mesh Particles flag checked.*/
 	UPROPERTY(EditAnywhere, Category = "Mesh Rendering")
diff --git a/Engine/Plugins/FX/Niagara/Source/Niagara/Public/NiagaraRenderer.h b/Engine/Plugins/FX/Niagara/Source/Niagara/Public/NiagaraRenderer.h
--- a/Engine/Plugins/FX/Niagara/Source/Niagara/Public/NiagaraRenderer.h
+++ b/Engine/Plugins/FX/Niagara/Source/Niagara/Public/NiagaraRenderer.h
@@ -143,6 +143,8 @@ protected:
 	TStatId EmitterStatID;
 #endif
 
+	virtual int32 GetMaxIndirectArgs() const { return 1; }
+
 	bool SetVertexFactoryVariable(const FNiagaraDataSet& DataSet, const FNiagaraVariable& Var, int32 VFVarOffset);
 	FGlobalDynamicReadBuffer::FAllocation TransferDataToGPU(FGlobalDynamicReadBuffer& DynamicReadBuffer, FNiagaraDataBuffer* SrcData)const;
 	
diff --git a/Engine/Plugins/FX/Niagara/Source/Niagara/Public/NiagaraRendererMeshes.h b/Engine/Plugins/FX/Niagara/Source/Niagara/Public/NiagaraRendererMeshes.h
--- a/Engine/Plugins/FX/Niagara/Source/Niagara/Public/NiagaraRendererMeshes.h
+++ b/Engine/Plugins/FX/Niagara/Source/Niagara/Public/NiagaraRendererMeshes.h
@@ -20,7 +20,6 @@ public:
 	~FNiagaraRendererMeshes();
 	
 	//FNiagaraRenderer Interface
-	virtual void CreateRenderThreadResources(NiagaraEmitterInstanceBatcher* Batcher) override;
 	virtual void ReleaseRenderThreadResources() override;
 
 	virtual void GetDynamicMeshElements(const TArray<const FSceneView*>& Views, const FSceneViewFamily& ViewFamily, uint32 VisibilityMap, FMeshElementCollector& Collector, const FNiagaraSceneProxy *SceneProxy) const override;
@@ -34,11 +33,16 @@ public:
 
 	void SetupVertexFactory(FNiagaraMeshVertexFactory *InVertexFactory, const FStaticMeshLODResources& LODResources) const;
 
+protected:
+	virtual int32 GetMaxIndirectArgs() const override;
+	int32 GetLODIndex() const;
+
 private:
 	mutable TArray<class FNiagaraMeshVertexFactory*, TInlineAllocator<2>> VertexFactories;
 	/** Render data of the static mesh we use. */
 	FStaticMeshRenderData* MeshRenderData;
 
+	TArray<TArray<TPair<int32 /*Count*/, int32 /*Offset*/>>> IndexInfoPerSection;
 	ENiagaraSortMode SortMode;
 	ENiagaraMeshFacingMode FacingMode;
 	uint32 bOverrideMaterials : 1;
diff --git a/Engine/Plugins/FX/Niagara/Source/NiagaraVertexFactories/Private/NiagaraDrawIndirect.cpp b/Engine/Plugins/FX/Niagara/Source/NiagaraVertexFactories/Private/NiagaraDrawIndirect.cpp
--- a/Engine/Plugins/FX/Niagara/Source/NiagaraVertexFactories/Private/NiagaraDrawIndirect.cpp
+++ b/Engine/Plugins/FX/Niagara/Source/NiagaraVertexFactories/Private/NiagaraDrawIndirect.cpp
@@ -15,6 +15,7 @@ void FNiagaraDrawIndirectArgsGenCS::ModifyCompilationEnvironment(const FGlobalSh
 	FGlobalShader::ModifyCompilationEnvironment(Parameters, OutEnvironment);
 	OutEnvironment.SetDefine(TEXT("THREAD_COUNT"), NIAGARA_DRAW_INDIRECT_ARGS_GEN_THREAD_COUNT);
 	OutEnvironment.SetDefine(TEXT("NIAGARA_DRAW_INDIRECT_ARGS_SIZE"), NIAGARA_DRAW_INDIRECT_ARGS_SIZE);
+	OutEnvironment.SetDefine(TEXT("NIAGARA_DRAW_INDIRECT_TASK_INFO_SIZE"), NIAGARA_DRAW_INDIRECT_TASK_INFO_SIZE);
 }
 
 FNiagaraDrawIndirectArgsGenCS::FNiagaraDrawIndirectArgsGenCS(const ShaderMetaType::CompiledShaderInitializerType& Initializer)
diff --git a/Engine/Plugins/FX/Niagara/Source/NiagaraVertexFactories/Public/NiagaraDrawIndirect.h b/Engine/Plugins/FX/Niagara/Source/NiagaraVertexFactories/Public/NiagaraDrawIndirect.h
--- a/Engine/Plugins/FX/Niagara/Source/NiagaraVertexFactories/Public/NiagaraDrawIndirect.h
+++ b/Engine/Plugins/FX/Niagara/Source/NiagaraVertexFactories/Public/NiagaraDrawIndirect.h
@@ -12,6 +12,7 @@ NiagaraDrawIndirect.h: Niagara shader to generate the draw indirect args for Nia
 
 #define NIAGARA_DRAW_INDIRECT_ARGS_GEN_THREAD_COUNT 64
 #define NIAGARA_DRAW_INDIRECT_ARGS_SIZE 5
+#define NIAGARA_DRAW_INDIRECT_TASK_INFO_SIZE 3
 
 // #define NIAGARA_COPY_BUFFER_THREAD_COUNT 64
 // #define NIAGARA_COPY_BUFFER_BUFFER_COUNT 3
@@ -33,17 +34,21 @@ public:
 	 */
 	struct FArgGenTaskInfo
 	{
-		FArgGenTaskInfo(uint32 InInstanceCountBufferOffset, uint32 InNumIndicesPerInstance) 
+		FArgGenTaskInfo(uint32 InInstanceCountBufferOffset, uint32 InNumIndicesPerInstance, uint32 InStartIndexLocation) 
 			: InstanceCountBufferOffset(InInstanceCountBufferOffset)
 			, NumIndicesPerInstance(InNumIndicesPerInstance)
+			, StartIndexLocation(InStartIndexLocation)
 		{}
 
 		uint32 InstanceCountBufferOffset;
 		uint32 NumIndicesPerInstance; // When -1 the counter needs to be reset to 0.
+		uint32 StartIndexLocation;
 
 		bool operator==(const FArgGenTaskInfo& Rhs) const
 		{
-			return InstanceCountBufferOffset == Rhs.InstanceCountBufferOffset && NumIndicesPerInstance == Rhs.NumIndicesPerInstance;
+			return InstanceCountBufferOffset == Rhs.InstanceCountBufferOffset
+				&& NumIndicesPerInstance == Rhs.NumIndicesPerInstance
+				&& StartIndexLocation == Rhs.StartIndexLocation;
 		}
 	};
 
