Fix for EngineAliveEmitterCount counter

diff --git a/Engine/Plugins/FX/Niagara/Source/Niagara/Private/NiagaraSystemInstance.cpp b/Engine/Plugins/FX/Niagara/Source/Niagara/Private/NiagaraSystemInstance.cpp
--- a/Engine/Plugins/FX/Niagara/Source/Niagara/Private/NiagaraSystemInstance.cpp
+++ b/Engine/Plugins/FX/Niagara/Source/Niagara/Private/NiagaraSystemInstance.cpp
@@ -1501,6 +1501,7 @@ void FNiagaraSystemInstance::TickInstanceParameters_GameThread(float DeltaSecond
 			CurrentEmitterParameters.EmitterNumParticles = Emitter->GetNumParticles();
 			CurrentEmitterParameters.EmitterTotalSpawnedParticles = Emitter->GetTotalSpawnedParticles();
 			CurrentEmitterParameters.EmitterSpawnCountScale = Emitter->GetSpawnCountScale(EffectsQualityLevel);
+			++GatheredInstanceParameters.NumAlive;
 		}
 		else
 		{
