Niagara - Fix an issue where the cached script vm data wasn't getting fully invalidated when the niagara compiler version changed.

diff --git a/Engine/Plugins/FX/Niagara/Source/Niagara/Classes/NiagaraScript.h b/Engine/Plugins/FX/Niagara/Source/Niagara/Classes/NiagaraScript.h
--- a/Engine/Plugins/FX/Niagara/Source/Niagara/Classes/NiagaraScript.h
+++ b/Engine/Plugins/FX/Niagara/Source/Niagara/Classes/NiagaraScript.h
@@ -569,7 +569,7 @@ public:
 	NIAGARA_API bool DidScriptCompilationSucceed(bool bGPUScript) const;
 
 #if WITH_EDITORONLY_DATA
-	NIAGARA_API void InvalidateCompileResults();
+	NIAGARA_API void InvalidateCompileResults(const FString& Reason);
 	FText GetDescription() { return Description.IsEmpty() ? FText::FromString(GetName()) : Description; }
 
 	/** Helper to convert the struct from its binary data out of the DDC to it's actual in-memory version.
diff --git a/Engine/Plugins/FX/Niagara/Source/Niagara/Private/NiagaraEmitter.cpp b/Engine/Plugins/FX/Niagara/Source/Niagara/Private/NiagaraEmitter.cpp
--- a/Engine/Plugins/FX/Niagara/Source/Niagara/Private/NiagaraEmitter.cpp
+++ b/Engine/Plugins/FX/Niagara/Source/Niagara/Private/NiagaraEmitter.cpp
@@ -345,7 +345,7 @@ void UNiagaraEmitter::PostLoad()
 			if (bActualInterpolatedSpawning)
 			{
 #if WITH_EDITORONLY_DATA
-				SpawnScriptProps.Script->InvalidateCompileResults();//clear out the script as it was compiled with interpolated spawn.
+				SpawnScriptProps.Script->InvalidateCompileResults(TEXT("Interpolated spawn changed."));//clear out the script as it was compiled with interpolated spawn.
 #endif
 				SpawnScriptProps.Script->SetUsage(ENiagaraScriptUsage::ParticleSpawnScript);
 			}
@@ -847,7 +847,7 @@ void  UNiagaraEmitter::InvalidateCompileResults()
 	GetScripts(Scripts, false);
 	for (int32 i = 0; i < Scripts.Num(); i++)
 	{
-		Scripts[i]->InvalidateCompileResults();
+		Scripts[i]->InvalidateCompileResults(TEXT("Emitter compile invalidated."));
 	}
 }
 
@@ -886,7 +886,7 @@ void UNiagaraEmitter::OnPostCompile()
 		GetScripts(Scripts, false);
 		for (int32 i = 0; i < Scripts.Num(); i++)
 		{
-			Scripts[i]->InvalidateCompileResults(); 
+			Scripts[i]->InvalidateCompileResults(TEXT("Console variable forced recompile.")); 
 		}
 	}
 
diff --git a/Engine/Plugins/FX/Niagara/Source/Niagara/Private/NiagaraEmitterHandle.cpp b/Engine/Plugins/FX/Niagara/Source/Niagara/Private/NiagaraEmitterHandle.cpp
--- a/Engine/Plugins/FX/Niagara/Source/Niagara/Private/NiagaraEmitterHandle.cpp
+++ b/Engine/Plugins/FX/Niagara/Source/Niagara/Private/NiagaraEmitterHandle.cpp
@@ -112,8 +112,9 @@ bool FNiagaraEmitterHandle::SetIsEnabled(bool bInIsEnabled, UNiagaraSystem& InOw
 			InOwnerSystem.GetSystemSpawnScript()->GetSource()->RefreshFromExternalChanges();
 
 			// Need to cause us to recompile in the future if necessary...
-			InOwnerSystem.GetSystemSpawnScript()->InvalidateCompileResults();
-			InOwnerSystem.GetSystemUpdateScript()->InvalidateCompileResults();
+			FString InvalidateReason = TEXT("Emitter enabled changed.");
+			InOwnerSystem.GetSystemSpawnScript()->InvalidateCompileResults(InvalidateReason);
+			InOwnerSystem.GetSystemUpdateScript()->InvalidateCompileResults(InvalidateReason);
 
 			// Clean out the emitter's compile results for cleanliness.
 			if (Instance)
diff --git a/Engine/Plugins/FX/Niagara/Source/Niagara/Private/NiagaraScript.cpp b/Engine/Plugins/FX/Niagara/Source/Niagara/Private/NiagaraScript.cpp
--- a/Engine/Plugins/FX/Niagara/Source/Niagara/Private/NiagaraScript.cpp
+++ b/Engine/Plugins/FX/Niagara/Source/Niagara/Private/NiagaraScript.cpp
@@ -646,7 +646,7 @@ void UNiagaraScript::PreSave(const class ITargetPlatform* TargetPlatform)
 					: TEXT("None"));
 		}
 
-		InvalidateCompileResults();
+		InvalidateCompileResults(TEXT("Data interface count mismatch during script presave."));
 		return;
 	}
 
@@ -874,6 +874,8 @@ void UNiagaraScript::PostLoad()
 	if (Source != nullptr)
 	{
 		Source->ConditionalPostLoad();
+		bool bScriptVMNeedsRebuild = false;
+		FString RebuildReason;
 		if (NiagaraVer < FNiagaraCustomVersion::UseHashesToIdentifyCompileStateOfTopLevelScripts && CachedScriptVMId.CompilerVersionID.IsValid())
 		{
 			FGuid BaseId = Source->GetCompileBaseId(Usage, UsageId);
@@ -882,9 +884,8 @@ void UNiagaraScript::PostLoad()
 				UE_LOG(LogNiagara, Warning,
 					TEXT("Invalidating compile ids for script %s because it doesn't have a valid base id.  The owning asset will continue to compile on load until it is resaved."),
 					*GetPathName());
-				bool bForceRebuild = true;
+				InvalidateCompileResults(TEXT("Script didn't have a valid base id."));
 				Source->ForceGraphToRecompileOnNextCheck();
-				Source->ComputeVMCompilationId(CachedScriptVMId, Usage, UsageId, bForceRebuild);
 			}
 			else
 			{
@@ -895,23 +896,32 @@ void UNiagaraScript::PostLoad()
 				}
 				else
 				{
-					// If the compile hash isn't valid, recompute the entire cached VM Id.
-					bool bForceRebuild = true;
-					Source->ComputeVMCompilationId(CachedScriptVMId, Usage, UsageId, bForceRebuild);
+					// If the compile hash isn't valid, the vm id needs to be recalculated and the cached vm needs to be invalidated.
+					bScriptVMNeedsRebuild = true;
+					RebuildReason = TEXT("Script did not have a valid compile hash.");
 				}
 			}
 		}
+
+		if (CachedScriptVMId.CompilerVersionID != FNiagaraCustomVersion::LatestScriptCompileVersion)
+		{
+			bScriptVMNeedsRebuild = true;
+			RebuildReason = TEXT("Niagara compiler version changed since the last time the script was compiled.");
+		}
+
+		if (bScriptVMNeedsRebuild)
+		{
+			// Force a rebuild on the source vm ids, and then invalidate the current cache to force the script to be unsynchronized.
+			bool bForceRebuild = true;
+			Source->ComputeVMCompilationId(CachedScriptVMId, Usage, UsageId, bForceRebuild);
+			InvalidateCompileResults(RebuildReason);
+		}
+
 		if (NiagaraVer < FNiagaraCustomVersion::AddLibraryAssetProperty)
 		{
 			bExposeToLibrary = true;
 		}
 	}
-
-	// Invalidate the CachedScriptVM if it's out of date to fix some cook errors, a further investigation is required in how to handle this correctly
-	if (CachedScriptVMId.CompilerVersionID != FNiagaraCustomVersion::LatestScriptCompileVersion)
-	{
-		CachedScriptVM.LastCompileStatus = ENiagaraScriptCompileStatus::NCS_Unknown;
-	}
 #endif
 	
 	// Resources can be processed / registered now that we're back on the main thread
@@ -1810,9 +1820,9 @@ bool UNiagaraScript::SynchronizeExecutablesWithMaster(const UNiagaraScript* Scri
 	return false;
 }
 
-void UNiagaraScript::InvalidateCompileResults()
+void UNiagaraScript::InvalidateCompileResults(const FString& Reason)
 {
-	UE_LOG(LogNiagara, Log, TEXT("InvalidateCompileResults %s"), *GetPathName());
+	UE_LOG(LogNiagara, Verbose, TEXT("InvalidateCompileResults Script:%s Reason:%s"), *GetPathName());
 	CachedScriptVM.Reset();
 	ScriptResource.Invalidate();
 	CachedScriptVMId.Invalidate();
