#Niagara Remove null entries in AppendCompileHash. Fixes a crash in cook

diff --git a/Engine/Plugins/FX/Niagara/Source/NiagaraEditor/Private/NiagaraGraph.cpp b/Engine/Plugins/FX/Niagara/Source/NiagaraEditor/Private/NiagaraGraph.cpp
--- a/Engine/Plugins/FX/Niagara/Source/NiagaraEditor/Private/NiagaraGraph.cpp
+++ b/Engine/Plugins/FX/Niagara/Source/NiagaraEditor/Private/NiagaraGraph.cpp
@@ -1105,6 +1105,7 @@ bool UNiagaraGraph::AppendCompileHash(FNiagaraCompileHashVisitor* InVisitor, con
 	// We need to sort the variables in a stable manner.
 	TArray<UNiagaraScriptVariable*> Values;
 	VariableToScriptVariable.GenerateValueArray(Values);
+	Values.Remove(nullptr);
 	Values.Sort([&](const UNiagaraScriptVariable& A, const UNiagaraScriptVariable& B)
 	{
 		const FName& AName = A.Variable.GetName();
