2 of 3: Add the Audio Osc Niagara Data Interface.

diff --git a/Engine/Plugins/FX/Niagara/Source/Niagara/Classes/NiagaraDataInterfaceAudioOscilloscope.h b/Engine/Plugins/FX/Niagara/Source/Niagara/Classes/NiagaraDataInterfaceAudioOscilloscope.h
new file mode 100644
--- /dev/null
+++ b/Engine/Plugins/FX/Niagara/Source/Niagara/Classes/NiagaraDataInterfaceAudioOscilloscope.h
@@ -0,0 +1,169 @@
+// Copyright 1998-2019 Epic Games, Inc. All Rights Reserved.
+#pragma once
+
+#include "NiagaraCommon.h"
+#include "NiagaraShared.h"
+#include "VectorVM.h"
+#include "StaticMeshResources.h"
+#include "NiagaraDataInterface.h"
+#include "AudioDevice.h"
+#include "AudioDeviceManager.h"
+#include "DSP/MultithreadedPatching.h"
+#include "NiagaraDataInterfaceAudioOscilloscope.generated.h"
+
+// Class used to to capture the audio stream of an arbitrary submix.
+class NIAGARA_API FNiagaraSubmixListener : public ISubmixBufferListener
+{
+public:
+	FNiagaraSubmixListener(Audio::FPatchMixer& InMixer, int32 InNumSamplesToBuffer);
+	FNiagaraSubmixListener(const FNiagaraSubmixListener& Other)
+	{
+		// Copy constructor technically required to compile TMap, but not used during runtime if move constructor is available.
+		// If you're hitting this, consider using Emplace or Add(MoveTemp()).
+		checkNoEntry();
+	}
+
+	FNiagaraSubmixListener(FNiagaraSubmixListener&& Other);
+
+	virtual ~FNiagaraSubmixListener();
+
+	float GetSampleRate();
+	int32 GetNumChannels();
+
+	// Begin ISubmixBufferListener overrides
+	virtual void OnNewSubmixBuffer(const USoundSubmix* OwningSubmix, float* AudioData, int32 NumSamples, int32 NumChannels, const int32 SampleRate, double AudioClock) override;
+	// End ISubmixBufferListener overrides
+
+private:
+	FNiagaraSubmixListener();
+
+	int32 NumChannelsInSubmix;
+	int32 SubmixSampleRate;
+
+	Audio::FPatchInput MixerInput;
+};
+
+struct FNiagaraDataInterfaceProxyOscilloscope : public FNiagaraDataInterfaceProxy
+{
+	FNiagaraDataInterfaceProxyOscilloscope(int32 InResolution, float InScopeInMillseconds);
+
+	~FNiagaraDataInterfaceProxyOscilloscope();
+
+	void OnBeginDestroy();
+
+	/**
+	 * Sample vertical displacement from the oscilloscope buffer.
+	 * @param NormalizedPositionInBuffer Horizontal position in the Oscilloscope buffer, from 0.0 to 1.0.
+	 * @param Channel channel index.
+	 * @return Amplitude at this position.
+	 */
+	float SampleAudio(float NormalizedPositionInBuffer, int32 Channel);
+
+	/**
+	 * @return the number of channels in the buffer.
+	 */
+	int32 GetNumChannels();
+
+	// Called when the Submix property changes.
+	void OnUpdateSubmix(USoundSubmix* Submix);
+
+	void RegisterToAllAudioDevices();
+	void UnregisterFromAllAudioDevices(USoundSubmix* Submix);
+
+	// Called when Resolution or Zoom are changed.
+	void OnUpdateResampling(int32 InResolution, float InScopeInMilliseconds);
+
+	// This function enqueues a render thread command to decimate the pop audio off of the SubmixListener, downsample it, and post it to the GPUAudioBuffer.
+	void PostAudioToGPU();
+	FReadBuffer& ComputeAndPostSRV();
+
+	// This function pops audio and downsamples it to our specified resolution.
+	void DownsampleAudioToBuffer();
+
+	virtual int32 PerInstanceDataPassedToRenderThreadSize() const override
+	{
+		return 0;
+	}
+
+private:
+	void OnNewDeviceCreated(Audio::FDeviceId InID);
+	void OnDeviceDestroyed(Audio::FDeviceId InID);
+
+	TMap<Audio::FDeviceId, FNiagaraSubmixListener> SubmixListeners;
+
+	// This 
+	Audio::FPatchMixer PatchMixer;
+
+	USoundSubmix* SubmixRegisteredTo;
+	bool bIsSubmixListenerRegistered;
+
+	int32 Resolution;
+	float ScopeInMilliseconds;
+
+	// The buffer we downsample PoppedBuffer to based on the Resolution property.
+	Audio::AlignedFloatBuffer PopBuffer;
+	Audio::AlignedFloatBuffer DownsampledBuffer;
+
+	// Handle for the SRV used by the generated HLSL.
+	FReadBuffer GPUDownsampledBuffer;
+	int32 NumChannelsInDownsampledBuffer;
+	
+	FDelegateHandle DeviceCreatedHandle;
+	FDelegateHandle DeviceDestroyedHandle;
+
+	FCriticalSection DownsampleBufferLock;
+};
+
+/** Data Interface allowing sampling of recent audio data. */
+UCLASS(EditInlineNew, Category = "Audio", meta = (DisplayName = "Audio Oscilloscope"))
+class NIAGARA_API UNiagaraDataInterfaceAudioOscilloscope : public UNiagaraDataInterface
+{
+	GENERATED_UCLASS_BODY()
+public:
+	
+	UPROPERTY(EditAnywhere, Category = "Oscilloscope")
+	USoundSubmix* Submix;
+
+	// The number of samples of audio to pass to the GPU. Audio will be resampled to fit this resolution.
+	UPROPERTY(EditAnywhere, Category = "Oscilloscope")
+	int32 Resolution;
+
+	// The number of milliseconds of audio to show.
+	UPROPERTY(EditAnywhere, Category = "Oscilloscope")
+	float ScopeInMilliseconds;
+
+	//VM function overrides:
+	void SampleAudio(FVectorVMContext& Context);
+	void GetNumChannels(FVectorVMContext& Context);
+
+	virtual void GetFunctions(TArray<FNiagaraFunctionSignature>& OutFunctions)override;
+	virtual void GetVMExternalFunction(const FVMExternalFunctionBindingInfo& BindingInfo, void* InstanceData, FVMExternalFunction &OutFunc) override;
+
+	virtual bool CanExecuteOnTarget(ENiagaraSimTarget Target) const override { 
+		if (Target == ENiagaraSimTarget::GPUComputeSim)
+		{
+			return true;
+		}
+		else
+		{
+			return false;
+		}
+	}
+	virtual bool RequiresDistanceFieldData() const override { return false; }
+
+	virtual bool GetFunctionHLSL(const FNiagaraDataInterfaceGPUParamInfo& ParamInfo, const FNiagaraDataInterfaceGeneratedFunction& FunctionInfo, int FunctionInstanceIndex, FString& OutHLSL) override;
+	virtual void GetParameterDefinitionHLSL(const FNiagaraDataInterfaceGPUParamInfo& ParamInfo, FString& OutHLSL) override;
+	virtual FNiagaraDataInterfaceParametersCS* ConstructComputeParameters() const override;
+
+#if WITH_EDITOR
+	virtual void PostEditChangeProperty(struct FPropertyChangedEvent& PropertyChangedEvent) override;
+#endif // WITH_EDITOR
+
+	virtual void PostInitProperties() override;
+	virtual void BeginDestroy() override;
+	virtual void PostLoad() override;
+
+protected:
+	virtual bool CopyToInternal(UNiagaraDataInterface* Destination) const override;
+};
+
diff --git a/Engine/Plugins/FX/Niagara/Source/Niagara/Niagara.Build.cs b/Engine/Plugins/FX/Niagara/Source/Niagara/Niagara.Build.cs
--- a/Engine/Plugins/FX/Niagara/Source/Niagara/Niagara.Build.cs
+++ b/Engine/Plugins/FX/Niagara/Source/Niagara/Niagara.Build.cs
@@ -18,6 +18,8 @@ public class Niagara : ModuleRules
                 "Renderer",
                 "JsonUtilities",
 				"Json",
+				"AudioPlatformConfiguration",
+				"SignalProcessing"
             }
         );
 
diff --git a/Engine/Plugins/FX/Niagara/Source/Niagara/Private/NiagaraDataInterfaceAudioOscilloscope.cpp b/Engine/Plugins/FX/Niagara/Source/Niagara/Private/NiagaraDataInterfaceAudioOscilloscope.cpp
new file mode 100644
--- /dev/null
+++ b/Engine/Plugins/FX/Niagara/Source/Niagara/Private/NiagaraDataInterfaceAudioOscilloscope.cpp
@@ -0,0 +1,596 @@
+// Copyright 1998-2019 Epic Games, Inc. All Rights Reserved.
+
+#include "NiagaraDataInterfaceAudioOscilloscope.h"
+#include "NiagaraTypes.h"
+#include "NiagaraCustomVersion.h"
+#include "AudioResampler.h"
+#include "NiagaraShader.h"
+#include "ShaderParameterUtils.h"
+#include "ClearQuad.h"
+#include "TextureResource.h"
+#include "Engine/Texture2D.h"
+#include "NiagaraEmitterInstanceBatcher.h"
+#include "NiagaraSystemInstance.h"
+#include "Engine/TextureRenderTarget2D.h"
+#include "Engine/Engine.h"
+#include "NiagaraComponent.h"
+
+#define LOCTEXT_NAMESPACE "NiagaraDataInterfaceGridAudioOscilloscope"
+
+
+// Global VM function names, also used by the shaders code generation methods.
+static const FName SampleAudioBufferFunctionName("SampleAudioBuffer");
+static const FName GetAudioBufferNumChannelsFunctionName("GetAudioBufferNumChannels");
+
+// Global variable prefixes, used in HLSL parameter declarations.
+static const FString AudioBufferName(TEXT("AudioBuffer_"));
+static const FString NumChannelsName(TEXT("NumChannels_"));
+
+FNiagaraDataInterfaceProxyOscilloscope::FNiagaraDataInterfaceProxyOscilloscope(int32 InResolution, float InScopeInMillseconds)
+	: PatchMixer()
+	, SubmixRegisteredTo(nullptr)
+	, bIsSubmixListenerRegistered(false)
+	, Resolution(InResolution)
+	, ScopeInMilliseconds(InScopeInMillseconds)
+	, NumChannelsInDownsampledBuffer(0)
+{
+	DeviceCreatedHandle = FAudioDeviceManagerDelegates::OnAudioDeviceCreated.AddRaw(this, &FNiagaraDataInterfaceProxyOscilloscope::OnNewDeviceCreated);
+	DeviceDestroyedHandle = FAudioDeviceManagerDelegates::OnAudioDeviceDestroyed.AddRaw(this, &FNiagaraDataInterfaceProxyOscilloscope::OnDeviceDestroyed);
+}
+
+FNiagaraDataInterfaceProxyOscilloscope::~FNiagaraDataInterfaceProxyOscilloscope()
+{
+	check(IsInRenderingThread());
+	GPUDownsampledBuffer.Release();
+
+	FAudioDeviceManagerDelegates::OnAudioDeviceCreated.Remove(DeviceCreatedHandle);
+	FAudioDeviceManagerDelegates::OnAudioDeviceDestroyed.Remove(DeviceDestroyedHandle);
+}
+
+void FNiagaraDataInterfaceProxyOscilloscope::RegisterToAllAudioDevices()
+{
+	if (FAudioDeviceManager* DeviceManager = FAudioDeviceManager::Get())
+	{
+		// Register a new submix listener for every audio device that currently exists.
+		DeviceManager->IterateOverAllDevices([&](Audio::FDeviceId DeviceId, FAudioDevice* InDevice)
+		{
+			check(!SubmixListeners.Contains(DeviceId));
+			FNiagaraSubmixListener& Listener = SubmixListeners.Emplace(DeviceId, FNiagaraSubmixListener(PatchMixer, Resolution * 2));
+			InDevice->RegisterSubmixBufferListener(&Listener, SubmixRegisteredTo);
+		});
+	}
+}
+
+void FNiagaraDataInterfaceProxyOscilloscope::UnregisterFromAllAudioDevices(USoundSubmix* Submix)
+{
+	if (FAudioDeviceManager* DeviceManager = FAudioDeviceManager::Get())
+	{
+		// Register a new submix listener for every audio device that currently exists.
+		DeviceManager->IterateOverAllDevices([&](Audio::FDeviceId DeviceId, FAudioDevice* InDevice)
+		{
+			check(SubmixListeners.Contains(DeviceId));
+
+			InDevice->UnregisterSubmixBufferListener(&SubmixListeners[DeviceId], Submix);
+			SubmixListeners.Remove(DeviceId);
+		});
+	}
+}
+
+void FNiagaraDataInterfaceProxyOscilloscope::OnUpdateSubmix(USoundSubmix* Submix)
+{
+	if (bIsSubmixListenerRegistered)
+	{
+		UnregisterFromAllAudioDevices(Submix);
+	}
+
+	SubmixRegisteredTo = Submix;
+	
+	RegisterToAllAudioDevices();
+	bIsSubmixListenerRegistered = true;
+}
+
+void FNiagaraDataInterfaceProxyOscilloscope::OnNewDeviceCreated(Audio::FDeviceId InID)
+{
+	if (bIsSubmixListenerRegistered)
+	{
+		check(!SubmixListeners.Contains(InID));
+		FAudioDeviceHandle DeviceHandle = FAudioDeviceManager::Get()->GetAudioDevice(InID);
+		check(DeviceHandle);
+
+		FNiagaraSubmixListener& Listener = SubmixListeners.Emplace(InID, FNiagaraSubmixListener(PatchMixer, Resolution * 2));
+		DeviceHandle->RegisterSubmixBufferListener(&Listener, SubmixRegisteredTo);
+	}
+}
+
+void FNiagaraDataInterfaceProxyOscilloscope::OnDeviceDestroyed(Audio::FDeviceId InID)
+{
+	if (bIsSubmixListenerRegistered)
+	{
+		if (SubmixListeners.Contains(InID))
+		{
+			FAudioDeviceHandle DeviceHandle = FAudioDeviceManager::Get()->GetAudioDevice(InID);
+			check(DeviceHandle);
+
+			DeviceHandle->UnregisterSubmixBufferListener(&SubmixListeners[InID], SubmixRegisteredTo);
+			SubmixListeners.Remove(InID);
+		}
+	}
+}
+
+UNiagaraDataInterfaceAudioOscilloscope::UNiagaraDataInterfaceAudioOscilloscope(FObjectInitializer const& ObjectInitializer)
+	: Super(ObjectInitializer)
+	, Submix(nullptr)
+	, Resolution(512)
+	, ScopeInMilliseconds(20.0f)
+	
+{
+	Proxy = TUniquePtr<FNiagaraDataInterfaceProxyOscilloscope>(new FNiagaraDataInterfaceProxyOscilloscope(Resolution, ScopeInMilliseconds));
+}
+
+float FNiagaraDataInterfaceProxyOscilloscope::SampleAudio(float NormalizedPositionInBuffer, int32 ChannelIndex)
+{
+	if (DownsampledBuffer.Num() == 0)
+	{
+		return 0.0f;
+	}
+
+	NormalizedPositionInBuffer = FMath::Clamp(NormalizedPositionInBuffer, 0.0f, 1.0f - SMALL_NUMBER);
+	float FrameIndex = NormalizedPositionInBuffer * DownsampledBuffer.Num();
+	int32 LowerFrameIndex = FMath::FloorToInt(FrameIndex);
+	float LowerFrameAmplitude = DownsampledBuffer[LowerFrameIndex * NumChannelsInDownsampledBuffer + ChannelIndex];
+	int32 HigherFrameIndex = FMath::CeilToInt(FrameIndex);
+	float HigherFrameAmplitude = DownsampledBuffer[HigherFrameIndex * NumChannelsInDownsampledBuffer + ChannelIndex];
+	float Fraction = HigherFrameIndex - FrameIndex;
+	return FMath::Lerp<float>(LowerFrameAmplitude, HigherFrameAmplitude, Fraction);
+}
+
+void UNiagaraDataInterfaceAudioOscilloscope::SampleAudio(FVectorVMContext& Context)
+{
+	//VectorVM::FUserPtrHandler<UNiagaraDataInterfaceAudioOscilloscope> InstData(Context);
+	VectorVM::FExternalFuncRegisterHandler<float> InNormalizedPos(Context);
+	VectorVM::FExternalFuncRegisterHandler<int32> InChannel(Context);
+	VectorVM::FExternalFuncRegisterHandler<float> OutAmplitude(Context);
+
+	for (int32 InstanceIdx = 0; InstanceIdx < Context.NumInstances; ++InstanceIdx)
+	{
+		float Position = InNormalizedPos.Get();
+		int32 Channel = InChannel.Get();
+		*OutAmplitude.GetDest() = GetProxyAs<FNiagaraDataInterfaceProxyOscilloscope>()->SampleAudio(Position, Channel);
+
+		InNormalizedPos.Advance();
+		InChannel.Advance();
+		OutAmplitude.Advance();
+	}
+}
+
+int32 FNiagaraDataInterfaceProxyOscilloscope::GetNumChannels()
+{
+	return NumChannelsInDownsampledBuffer;
+}
+
+void UNiagaraDataInterfaceAudioOscilloscope::GetNumChannels(FVectorVMContext& Context)
+{
+	VectorVM::FUserPtrHandler<UNiagaraDataInterfaceAudioOscilloscope> InstData(Context);
+	VectorVM::FExternalFuncRegisterHandler<int32> OutChannel(Context);
+
+	for (int32 InstanceIdx = 0; InstanceIdx < Context.NumInstances; ++InstanceIdx)
+	{
+		*OutChannel.GetDestAndAdvance() = GetProxyAs<FNiagaraDataInterfaceProxyOscilloscope>()->GetNumChannels();
+	}
+}
+
+void UNiagaraDataInterfaceAudioOscilloscope::GetFunctions(TArray<FNiagaraFunctionSignature>& OutFunctions)
+{
+	Super::GetFunctions(OutFunctions);
+
+	{
+		FNiagaraFunctionSignature SampleAudioBufferSignature;
+		SampleAudioBufferSignature.Name = SampleAudioBufferFunctionName;
+		SampleAudioBufferSignature.Inputs.Add(FNiagaraVariable(GetClass(), TEXT("Oscilloscope")));
+		SampleAudioBufferSignature.Inputs.Add(FNiagaraVariable(FNiagaraTypeDefinition::GetFloatDef(), TEXT("NormalizedPositionInBuffer")));
+		SampleAudioBufferSignature.Inputs.Add(FNiagaraVariable(FNiagaraTypeDefinition::GetIntDef(), TEXT("ChannelIndex")));
+		SampleAudioBufferSignature.Outputs.Add(FNiagaraVariable(FNiagaraTypeDefinition::GetFloatDef(), TEXT("Amplitude")));
+
+		SampleAudioBufferSignature.bMemberFunction = true;
+		SampleAudioBufferSignature.bRequiresContext = false;
+		OutFunctions.Add(SampleAudioBufferSignature);
+	}
+
+	{
+		FNiagaraFunctionSignature NumChannelsSignature;
+		NumChannelsSignature.Name = GetAudioBufferNumChannelsFunctionName;
+		NumChannelsSignature.Inputs.Add(FNiagaraVariable(GetClass(), TEXT("Oscilloscope")));
+		NumChannelsSignature.Outputs.Add(FNiagaraVariable(FNiagaraTypeDefinition::GetIntDef(), TEXT("NumChannels")));
+
+		NumChannelsSignature.bMemberFunction = true;
+		NumChannelsSignature.bRequiresContext = false;
+		OutFunctions.Add(NumChannelsSignature);
+	}
+}
+
+DEFINE_NDI_DIRECT_FUNC_BINDER(UNiagaraDataInterfaceAudioOscilloscope, SampleAudio);
+DEFINE_NDI_DIRECT_FUNC_BINDER(UNiagaraDataInterfaceAudioOscilloscope, GetNumChannels);
+
+void UNiagaraDataInterfaceAudioOscilloscope::GetVMExternalFunction(const FVMExternalFunctionBindingInfo& BindingInfo, void* InstanceData, FVMExternalFunction &OutFunc)
+{
+	if (BindingInfo.Name == SampleAudioBufferFunctionName)
+	{
+		//TNDIParamBinder<0, int32, NDI_FUNC_BINDER(UNiagaraDataInterfaceAudioOscilloscope, SampleAudio)>::Bind(this, BindingInfo, InstanceData, OutFunc);
+		NDI_FUNC_BINDER(UNiagaraDataInterfaceAudioOscilloscope, SampleAudio)::Bind(this, OutFunc);
+	}
+	else if (BindingInfo.Name == GetAudioBufferNumChannelsFunctionName)
+	{
+		//TNDIParamBinder<0, int32, NDI_FUNC_BINDER(UNiagaraDataInterfaceAudioOscilloscope, GetNumChannels)>::Bind(this, BindingInfo, InstanceData, OutFunc);
+		NDI_FUNC_BINDER(UNiagaraDataInterfaceAudioOscilloscope, GetNumChannels)::Bind(this, OutFunc);
+	}
+	else
+	{
+		ensureMsgf(false, TEXT("Error! Function defined for this class but not bound."));
+	}
+}
+
+bool UNiagaraDataInterfaceAudioOscilloscope::GetFunctionHLSL(const FNiagaraDataInterfaceGPUParamInfo& ParamInfo, const FNiagaraDataInterfaceGeneratedFunction& FunctionInfo, int FunctionInstanceIndex, FString& OutHLSL)
+{
+	bool ParentRet = Super::GetFunctionHLSL(ParamInfo, FunctionInfo, FunctionInstanceIndex, OutHLSL);
+	if (ParentRet)
+	{
+		return true;
+	}
+	else if (FunctionInfo.DefinitionName == SampleAudioBufferFunctionName)
+	{
+		// See SampleBuffer(float InNormalizedPosition, int32 InChannel)
+		/**
+		 * float FrameIndex = NormalizedPositionInBuffer * DownsampledBuffer.Num();
+		 * int32 LowerFrameIndex = FMath::FloorToInt(FrameIndex);
+		 * float LowerFrameAmplitude = DownsampledBuffer[LowerFrameIndex * NumChannelsInDownsampledBuffer + ChannelIndex];
+		 * int32 HigherFrameIndex = FMath::CeilToInt(FrameIndex);
+		 * float HigherFrameAmplitude = DownsampledBuffer[HigherFrameIndex * NumChannelsInDownsampledBuffer + ChannelIndex];
+		 * float Fraction = HigherFrameIndex - FrameIndex;
+		 * return FMath::Lerp<float>(LowerFrameAmplitude, HigherFrameAmplitude, Fraction);
+		 */
+		static const TCHAR *FormatBounds = TEXT(R"(
+			void {FunctionName}(float In_NormalizedPosition, int In_ChannelIndex, out float Out_Val)
+			{
+				float FrameIndex = In_NormalizedPosition * {AudioBufferNumFrames};
+				int LowerIndex = floor(FrameIndex);
+				int UpperIndex =  LowerIndex < {AudioBufferNumFrames} ? LowerIndex + 1.0 : LowerIndex;
+				float Fraction = FrameIndex - LowerIndex;
+				float LowerValue = {AudioBuffer}.Load(LowerIndex * {ChannelCount} + In_ChannelIndex);
+				float UpperValue = {AudioBuffer}.Load(UpperIndex * {ChannelCount} + In_ChannelIndex);
+				Out_Val = lerp(LowerValue, UpperValue, Fraction);
+			}
+		)");
+		TMap<FString, FStringFormatArg> ArgsBounds = {
+			{TEXT("FunctionName"), FStringFormatArg(FunctionInfo.InstanceName)},
+			{TEXT("ChannelCount"), FStringFormatArg(NumChannelsName + ParamInfo.DataInterfaceHLSLSymbol)},
+			{TEXT("AudioBufferNumFrames"), FStringFormatArg(Resolution)},
+			{TEXT("AudioBuffer"), FStringFormatArg(AudioBufferName + ParamInfo.DataInterfaceHLSLSymbol)},
+		};
+		OutHLSL += FString::Format(FormatBounds, ArgsBounds);
+		return true;
+	}
+	else if (FunctionInfo.DefinitionName == GetAudioBufferNumChannelsFunctionName)
+	{
+		// See GetNumChannels()
+		static const TCHAR *FormatBounds = TEXT(R"(
+			void {FunctionName}(out int Out_Val)
+			{
+				Out_Val = {ChannelCount};
+			}
+		)");
+		TMap<FString, FStringFormatArg> ArgsBounds = {
+			{TEXT("FunctionName"), FStringFormatArg(FunctionInfo.InstanceName)},
+			{TEXT("ChannelCount"), FStringFormatArg(NumChannelsName + ParamInfo.DataInterfaceHLSLSymbol)},
+		};
+		OutHLSL += FString::Format(FormatBounds, ArgsBounds);
+		return true;
+	}
+	else
+	{
+		return false;
+	}
+}
+
+void UNiagaraDataInterfaceAudioOscilloscope::GetParameterDefinitionHLSL(const FNiagaraDataInterfaceGPUParamInfo& ParamInfo, FString& OutHLSL)
+{
+	Super::GetParameterDefinitionHLSL(ParamInfo, OutHLSL);
+
+	static const TCHAR *FormatDeclarations = TEXT(R"(				
+		Buffer<float> {AudioBufferName};
+		int {NumChannelsName};
+
+	)");
+	TMap<FString, FStringFormatArg> ArgsDeclarations = {
+		{ TEXT("AudioBufferName"),    FStringFormatArg(AudioBufferName + ParamInfo.DataInterfaceHLSLSymbol) },
+		{ TEXT("NumChannelsName"),    FStringFormatArg(NumChannelsName + ParamInfo.DataInterfaceHLSLSymbol) },
+	};
+	OutHLSL += FString::Format(FormatDeclarations, ArgsDeclarations);
+}
+
+struct FNiagaraDataInterfaceParametersCS_AudioOscilloscope : public FNiagaraDataInterfaceParametersCS
+{
+	
+
+	FNiagaraDataInterfaceParametersCS_AudioOscilloscope(UNiagaraDataInterfaceAudioOscilloscope* InDataInterface)
+	{
+	}
+
+	virtual ~FNiagaraDataInterfaceParametersCS_AudioOscilloscope() {}
+	virtual void Bind(const FNiagaraDataInterfaceParamRef& ParamRef, const class FShaderParameterMap& ParameterMap) override
+	{
+		NumChannels.Bind(ParameterMap, *(NumChannelsName + ParamRef.ParameterInfo.DataInterfaceHLSLSymbol));
+		AudioBuffer.Bind(ParameterMap, *(AudioBufferName + ParamRef.ParameterInfo.DataInterfaceHLSLSymbol));
+	}
+
+	virtual void Serialize(FArchive& Ar) override
+	{
+		Ar << NumChannels;
+		Ar << AudioBuffer;
+	}
+
+	virtual void Set(FRHICommandList& RHICmdList, const FNiagaraDataInterfaceSetArgs& Context) const override
+	{
+		check(IsInRenderingThread());
+
+		FRHIComputeShader* ComputeShaderRHI = Context.Shader->GetComputeShader();
+
+		FNiagaraDataInterfaceProxyOscilloscope* NDI = (FNiagaraDataInterfaceProxyOscilloscope*)Context.DataInterface;
+		FReadBuffer& AudioBufferSRV = NDI->ComputeAndPostSRV();
+
+		SetShaderValue(RHICmdList, ComputeShaderRHI, NumChannels, NDI->GetNumChannels());
+		RHICmdList.SetShaderResourceViewParameter(ComputeShaderRHI, AudioBuffer.GetBaseIndex(), AudioBufferSRV.SRV);
+	}
+
+	FShaderParameter NumChannels;
+	FShaderResourceParameter AudioBuffer;
+};
+
+FNiagaraDataInterfaceParametersCS* UNiagaraDataInterfaceAudioOscilloscope::ConstructComputeParameters()const
+{
+	return new FNiagaraDataInterfaceParametersCS_AudioOscilloscope(const_cast<UNiagaraDataInterfaceAudioOscilloscope*>(this));
+}
+
+void FNiagaraDataInterfaceProxyOscilloscope::OnUpdateResampling(int32 InResolution, float InScopeInMilliseconds)
+{
+	Resolution = InResolution;
+	ScopeInMilliseconds = InScopeInMilliseconds;
+
+	const int32 NumSamplesInBuffer = Resolution * NumChannelsInDownsampledBuffer;
+	if (NumSamplesInBuffer)
+	{
+		ENQUEUE_RENDER_COMMAND(FUpdateDIAudioBuffer)(
+			[&, NumSamplesInBuffer](FRHICommandListImmediate& RHICmdList)
+		{
+			if (GPUDownsampledBuffer.NumBytes > 0)
+			{
+				GPUDownsampledBuffer.Release();
+			}
+
+			GPUDownsampledBuffer.Initialize(sizeof(float), NumSamplesInBuffer, EPixelFormat::PF_R32_FLOAT, BUF_Static);
+		});
+	}
+	
+}
+
+FReadBuffer& FNiagaraDataInterfaceProxyOscilloscope::ComputeAndPostSRV()
+{
+	// Copy to GPUDownsampledBuffer:
+	PostAudioToGPU();
+	return GPUDownsampledBuffer;
+}
+
+#if WITH_EDITOR
+void UNiagaraDataInterfaceAudioOscilloscope::PostEditChangeProperty(struct FPropertyChangedEvent& PropertyChangedEvent)
+{
+	Super::PostEditChangeProperty(PropertyChangedEvent);
+	
+	static FName SubmixFName = GET_MEMBER_NAME_CHECKED(UNiagaraDataInterfaceAudioOscilloscope, Submix);
+	static FName ResolutionFName = GET_MEMBER_NAME_CHECKED(UNiagaraDataInterfaceAudioOscilloscope, Resolution);
+	static FName ScopeInMillisecondsFName = GET_MEMBER_NAME_CHECKED(UNiagaraDataInterfaceAudioOscilloscope, ScopeInMilliseconds);
+
+	// Regenerate on save any compressed sound formats or if analysis needs to be re-done
+	if (UProperty* PropertyThatChanged = PropertyChangedEvent.Property)
+	{
+		const FName& Name = PropertyThatChanged->GetFName();
+		if (Name == SubmixFName)
+		{
+			GetProxyAs<FNiagaraDataInterfaceProxyOscilloscope>()->OnUpdateSubmix(Submix);
+		}
+		else if (Name == ResolutionFName || Name == ScopeInMillisecondsFName)
+		{
+			GetProxyAs<FNiagaraDataInterfaceProxyOscilloscope>()->OnUpdateResampling(Resolution, ScopeInMilliseconds);
+		}
+	}
+}
+#endif //WITH_EDITOR
+
+void UNiagaraDataInterfaceAudioOscilloscope::PostInitProperties()
+{
+	Super::PostInitProperties();
+
+	if (HasAnyFlags(RF_ClassDefaultObject))
+	{
+		FNiagaraTypeRegistry::Register(FNiagaraTypeDefinition(GetClass()), /*bCanBeParameter*/ true, /*bCanBePayload*/ false, /*bIsUserDefined*/ false);
+	}
+
+	GetProxyAs<FNiagaraDataInterfaceProxyOscilloscope>()->OnUpdateResampling(Resolution, ScopeInMilliseconds);
+	GetProxyAs<FNiagaraDataInterfaceProxyOscilloscope>()->OnUpdateSubmix(Submix);
+}
+
+void FNiagaraDataInterfaceProxyOscilloscope::OnBeginDestroy()
+{
+	if (bIsSubmixListenerRegistered)
+	{
+		UnregisterFromAllAudioDevices(SubmixRegisteredTo);
+	}
+}
+
+void UNiagaraDataInterfaceAudioOscilloscope::BeginDestroy()
+{
+	GetProxyAs<FNiagaraDataInterfaceProxyOscilloscope>()->OnBeginDestroy();
+
+	Super::BeginDestroy();
+}
+
+void UNiagaraDataInterfaceAudioOscilloscope::PostLoad()
+{
+	Super::PostLoad();
+	GetProxyAs<FNiagaraDataInterfaceProxyOscilloscope>()->OnUpdateResampling(Resolution, ScopeInMilliseconds);
+	GetProxyAs<FNiagaraDataInterfaceProxyOscilloscope>()->OnUpdateSubmix(Submix);
+}
+
+bool UNiagaraDataInterfaceAudioOscilloscope::CopyToInternal(UNiagaraDataInterface* Destination) const
+{
+	Super::CopyToInternal(Destination);
+
+	UNiagaraDataInterfaceAudioOscilloscope* CastedDestination = Cast<UNiagaraDataInterfaceAudioOscilloscope>(Destination);
+
+	if (CastedDestination)
+	{
+		CastedDestination->Submix = Submix;
+		CastedDestination->Resolution = Resolution;
+		CastedDestination->ScopeInMilliseconds = ScopeInMilliseconds;
+	}
+	
+	return true;
+}
+
+void FNiagaraDataInterfaceProxyOscilloscope::PostAudioToGPU()
+{
+	ENQUEUE_RENDER_COMMAND(FUpdateDIAudioBuffer)(
+		[&](FRHICommandListImmediate& RHICmdList)
+	{
+		DownsampleAudioToBuffer();
+		uint32 BufferSize = Resolution * NumChannelsInDownsampledBuffer * sizeof(float);
+		if (BufferSize != 0 && !GPUDownsampledBuffer.NumBytes)
+		{
+			GPUDownsampledBuffer.Initialize(sizeof(float), Resolution * NumChannelsInDownsampledBuffer, EPixelFormat::PF_R32_FLOAT, BUF_Static);
+		}
+
+		if (GPUDownsampledBuffer.NumBytes > 0)
+		{
+			float *BufferData = static_cast<float*>(RHILockVertexBuffer(GPUDownsampledBuffer.Buffer, 0, BufferSize, EResourceLockMode::RLM_WriteOnly));
+			FScopeLock ScopeLock(&DownsampleBufferLock);
+			FPlatformMemory::Memcpy(BufferData, DownsampledBuffer.GetData(), BufferSize);
+			RHIUnlockVertexBuffer(GPUDownsampledBuffer.Buffer);
+		}
+	});
+}
+
+void FNiagaraDataInterfaceProxyOscilloscope::DownsampleAudioToBuffer()
+{
+	FScopeLock ScopeLock(&DownsampleBufferLock);
+
+	// Get the number of channels from the first listener we find.
+	// If the num channels hasn't been set, either SubmixListeners is empty
+	// or we haven't pushed audio to any of them yet.
+	NumChannelsInDownsampledBuffer = 0;
+	float SourceSampleRate = 0.0f;
+
+	for (auto& Listener : SubmixListeners)
+	{
+		NumChannelsInDownsampledBuffer = Listener.Value.GetNumChannels();
+		SourceSampleRate = Listener.Value.GetSampleRate();
+
+		if (NumChannelsInDownsampledBuffer != 0)
+		{
+			break;
+		}
+	}
+
+	if (NumChannelsInDownsampledBuffer == 0 || FMath::IsNearlyZero(SourceSampleRate))
+	{
+		return;
+	}
+
+	check(NumChannelsInDownsampledBuffer > 0);
+
+	// Get the number of frames of audio in the original sample rate to show.
+	int32 NumFramesToPop = FMath::FloorToInt(SourceSampleRate * (ScopeInMilliseconds / 1000.0f));
+
+	// If we have enough frames buffered to pop, update the PopBuffer. Otherwise, use the previous frames buffer.
+	int32 NumSamplesToPop = NumFramesToPop * NumChannelsInDownsampledBuffer;
+	if (PatchMixer.MaxNumberOfSamplesThatCanBePopped() >= NumFramesToPop)
+	{
+		PopBuffer.Reset();
+		PopBuffer.AddZeroed(NumSamplesToPop);
+		PatchMixer.PopAudio(PopBuffer.GetData(), NumSamplesToPop, true);
+	}
+	else
+	{
+		PopBuffer.SetNumZeroed(NumFramesToPop * NumChannelsInDownsampledBuffer);
+	}
+
+	// Downsample buffer in place.
+	float SampleRateRatio = ((float)Resolution) / ((float)NumFramesToPop);
+	float DestinationSampleRate = SourceSampleRate * SampleRateRatio;
+
+	Audio::FResamplingParameters ResampleParameters = {
+		Audio::EResamplingMethod::Linear,
+		NumChannelsInDownsampledBuffer,
+		SourceSampleRate,
+		DestinationSampleRate,
+		PopBuffer
+	};
+
+	int32 DownsampleBufferSize = Audio::GetOutputBufferSize(ResampleParameters);
+	DownsampledBuffer.Reset();
+	DownsampledBuffer.AddZeroed(DownsampleBufferSize);
+
+
+	Audio::FResamplerResults ResampleResults;
+	ResampleResults.OutBuffer = &DownsampledBuffer;
+
+	bool bResampleResult = Audio::Resample(ResampleParameters, ResampleResults);
+	check(bResampleResult);
+
+	// here we may have an extra sample or two due to roundoff.
+	DownsampledBuffer.SetNumZeroed(Resolution * NumChannelsInDownsampledBuffer);
+}
+
+FNiagaraSubmixListener::FNiagaraSubmixListener(Audio::FPatchMixer& InMixer, int32 InNumSamplesToBuffer)
+	: NumChannelsInSubmix(0)
+	, SubmixSampleRate(0)
+	, MixerInput(InMixer.AddNewInput(InNumSamplesToBuffer, 1.0f))
+{
+
+}
+
+FNiagaraSubmixListener::FNiagaraSubmixListener()
+	: NumChannelsInSubmix(0)
+	, SubmixSampleRate(0)
+{
+
+}
+
+FNiagaraSubmixListener::FNiagaraSubmixListener(FNiagaraSubmixListener&& Other)
+{
+	NumChannelsInSubmix = Other.NumChannelsInSubmix;
+	Other.NumChannelsInSubmix = 0;
+
+	SubmixSampleRate = Other.SubmixSampleRate;
+	Other.SubmixSampleRate = 0;
+
+	MixerInput = MoveTemp(Other.MixerInput);
+}
+
+FNiagaraSubmixListener::~FNiagaraSubmixListener()
+{
+}
+
+float FNiagaraSubmixListener::GetSampleRate()
+{
+	return (float) SubmixSampleRate;
+}
+
+int32 FNiagaraSubmixListener::GetNumChannels()
+{
+	return NumChannelsInSubmix;
+}
+
+void FNiagaraSubmixListener::OnNewSubmixBuffer(const USoundSubmix* OwningSubmix, float* AudioData, int32 NumSamples, int32 NumChannels, const int32 SampleRate, double AudioClock)
+{
+	NumChannelsInSubmix = NumChannels;
+	SubmixSampleRate = SampleRate;
+	MixerInput.PushAudio(AudioData, NumSamples);
+}
diff --git a/Engine/Source/Runtime/SignalProcessing/Private/MultithreadedPatching.cpp b/Engine/Source/Runtime/SignalProcessing/Private/MultithreadedPatching.cpp
--- a/Engine/Source/Runtime/SignalProcessing/Private/MultithreadedPatching.cpp
+++ b/Engine/Source/Runtime/SignalProcessing/Private/MultithreadedPatching.cpp
@@ -120,6 +120,15 @@ namespace Audio
 	{
 	}
 
+	FPatchInput::FPatchInput(FPatchInput&& Other)
+	{
+		OutputHandle = Other.OutputHandle;
+		Other.OutputHandle.Reset();
+
+		PushCallsCounter = Other.PushCallsCounter;
+		Other.PushCallsCounter = 0;
+	}
+
 	FPatchInput& FPatchInput::operator=(const FPatchInput& Other)
 	{
 		OutputHandle = Other.OutputHandle;
diff --git a/Engine/Source/Runtime/SignalProcessing/Public/DSP/MultithreadedPatching.h b/Engine/Source/Runtime/SignalProcessing/Public/DSP/MultithreadedPatching.h
--- a/Engine/Source/Runtime/SignalProcessing/Public/DSP/MultithreadedPatching.h
+++ b/Engine/Source/Runtime/SignalProcessing/Public/DSP/MultithreadedPatching.h
@@ -83,6 +83,7 @@ namespace Audio
 		FPatchInput(const FPatchOutputStrongPtr& InOutput);
 		FPatchInput(const FPatchInput& Other);
 		FPatchInput& operator=(const FPatchInput& Other);
+		FPatchInput(FPatchInput&& Other);
 
 		/** Default constructed FPatchInput instances will always return -1 for PushAudio and true for IsOutputStillActive. */
 		FPatchInput();
