Reverting sprite alignment changes until we can investigate issues more.

diff --git a/Engine/Plugins/FX/Niagara/Shaders/Private/NiagaraSpriteVertexFactory.ush b/Engine/Plugins/FX/Niagara/Shaders/Private/NiagaraSpriteVertexFactory.ush
--- a/Engine/Plugins/FX/Niagara/Shaders/Private/NiagaraSpriteVertexFactory.ush
+++ b/Engine/Plugins/FX/Niagara/Shaders/Private/NiagaraSpriteVertexFactory.ush
@@ -486,12 +486,7 @@ void GetTangents(FVertexFactoryInput Input, FVertexFactoryIntermediates Intermed
 	
 	/**This could well be better as a branch*/
 	BRANCH
-	if(NiagaraSpriteVFLooseParameters.ParticleAlignmentMode == ALIGNMENT_MODE_CUSTOMALIGNMENT && NiagaraSpriteVFLooseParameters.ParticleFacingMode == FACING_MODE_CUSTOMFACINGVECTOR)
-	{
-		Right = Intermediates.CustomAlignmentVector;
-		Up = cross(CustomFacingVector, Right);
-	}
-	else if(NiagaraSpriteVFLooseParameters.ParticleAlignmentMode == ALIGNMENT_MODE_UNALIGNED)
+	if(NiagaraSpriteVFLooseParameters.ParticleAlignmentMode == ALIGNMENT_MODE_UNALIGNED)
 	{
 		FLATTEN
 		if (NiagaraSpriteVFLooseParameters.ParticleFacingMode == FACING_MODE_FACECAMERADISTANCEBLEND)
