Particle read DI:
* added functions which can retrieve the IDs of the particles spawned in the last tick
* added missing barriers on source buffers when reading from another emitter
* fixed bug when more than one DI function was used in a module (bad alignment on the elements of the array of attribute indices)

Persistent IDs:
* replaced ClearUAV with a custom shader, so we can set sane barriers instead of the stuff that ClearUAV() forces. Dev-Rendering replaced ClearUAV() with a better API, but I didn't want to wait for that merge.
* grouped together all the ID buffer clears before the simulation dispatches, and all the free ID list updates after rendering, so they can overlap and to avoid stalls

Barriers:
* fixed a bogus compute to graphics transition on the ID to index buffer
* when running multiple ticks, only the outputs of the final tick are transitioned to graphics
* automatic CS cache flushes are disabled once before running all the ticks and enabled at the end, instead of doing it for each tick

diff --git a/Engine/Plugins/FX/Niagara/Shaders/Private/NiagaraFillIntBuffer.usf b/Engine/Plugins/FX/Niagara/Shaders/Private/NiagaraFillIntBuffer.usf
new file mode 100644
--- /dev/null
+++ b/Engine/Plugins/FX/Niagara/Shaders/Private/NiagaraFillIntBuffer.usf
@@ -0,0 +1,17 @@
+// Copyright Epic Games, Inc. All Rights Reserved.
+
+#include "/Engine/Private/Common.ush"
+
+RWBuffer<int> TargetBuffer;
+int FillValue;
+uint BufferSize;
+
+[numthreads(64, 1, 1)]
+void FillIntBuffer(uint3 DispatchThreadId : SV_DispatchThreadID)
+{
+	uint Index = DispatchThreadId.x;
+	if(Index < BufferSize)
+	{
+		TargetBuffer[Index] = FillValue;
+	}
+}
diff --git a/Engine/Plugins/FX/Niagara/Source/Niagara/Classes/NiagaraDataInterfaceParticleRead.h b/Engine/Plugins/FX/Niagara/Source/Niagara/Classes/NiagaraDataInterfaceParticleRead.h
--- a/Engine/Plugins/FX/Niagara/Source/Niagara/Classes/NiagaraDataInterfaceParticleRead.h
+++ b/Engine/Plugins/FX/Niagara/Source/Niagara/Classes/NiagaraDataInterfaceParticleRead.h
@@ -23,48 +23,36 @@ public:
 
 	//UObject Interface
 	virtual void PostInitProperties()override;
-	virtual void PostLoad() override;
-#if WITH_EDITOR
-	virtual void PreEditChange(FProperty* PropertyAboutToChange) override;
-	virtual void PostEditChangeProperty(struct FPropertyChangedEvent& PropertyChangedEvent) override;
-#endif
 	//UObject Interface End
 
 	//UNiagaraDataInterface Interface
 	virtual bool InitPerInstanceData(void* PerInstanceData, FNiagaraSystemInstance* SystemInstance) override;
-	virtual void DestroyPerInstanceData(void* PerInstanceData, FNiagaraSystemInstance* SystemInstance) override;
 	virtual int32 PerInstanceDataSize() const { return sizeof(FNDIParticleRead_InstanceData); }
-	
-	virtual bool PerInstanceTick(void* PerInstanceData, FNiagaraSystemInstance* SystemInstance, float DeltaSeconds);
-	virtual bool PerInstanceTickPostSimulate(void* PerInstanceData, FNiagaraSystemInstance* SystemInstance, float DeltaSeconds);
-
 	virtual void GetFunctions(TArray<FNiagaraFunctionSignature>& OutFunctions) override;
-	virtual void GetVMExternalFunction(const FVMExternalFunctionBindingInfo& BindingInfo, void* InstanceData, FVMExternalFunction &OutFunc) override;
+	virtual void GetVMExternalFunction(const FVMExternalFunctionBindingInfo& BindingInfo, void* InstanceData, FVMExternalFunction& OutFunc) override;
 	virtual bool CanExecuteOnTarget(ENiagaraSimTarget Target) const override { return true; }
 	virtual bool Equals(const UNiagaraDataInterface* Other) const override;
-
-	// GPU sim functionality
-	virtual void GetCommonHLSL(FString& OutHLSL) override;
 	virtual void GetParameterDefinitionHLSL(const FNiagaraDataInterfaceGPUParamInfo& ParamInfo, FString& OutHLSL) override;
 	virtual bool GetFunctionHLSL(const FNiagaraDataInterfaceGPUParamInfo& ParamInfo, const FNiagaraDataInterfaceGeneratedFunction& FunctionInfo, int FunctionInstanceIndex, FString& OutHLSL) override;
 	virtual FNiagaraDataInterfaceParametersCS* ConstructComputeParameters() const override;
 	virtual void ProvidePerInstanceDataForRenderThread(void* DataForRenderThread, void* PerInstanceData, const FNiagaraSystemInstanceID& SystemInstance) override;
 	//UNiagaraDataInterface Interface End
 
-	virtual void ReadFloat(FVectorVMContext& Context, FName AttributeToRead);
-	virtual void ReadVector2(FVectorVMContext& Context, FName AttributeToRead);
-	virtual void ReadVector3(FVectorVMContext& Context, FName AttributeToRead);
-	virtual void ReadVector4(FVectorVMContext& Context, FName AttributeToRead);
-	virtual void ReadInt(FVectorVMContext& Context, FName AttributeToRead);
-	virtual void ReadBool(FVectorVMContext& Context, FName AttributeToRead);
-	virtual void ReadColor(FVectorVMContext& Context, FName AttributeToRead);
-	virtual void ReadQuat(FVectorVMContext& Context, FName AttributeToRead);
+	void GetNumSpawnedParticles(FVectorVMContext& Context);
+	void GetSpawnedIDAtIndex(FVectorVMContext& Context);
+	void ReadFloat(FVectorVMContext& Context, FName AttributeToRead);
+	void ReadVector2(FVectorVMContext& Context, FName AttributeToRead);
+	void ReadVector3(FVectorVMContext& Context, FName AttributeToRead);
+	void ReadVector4(FVectorVMContext& Context, FName AttributeToRead);
+	void ReadInt(FVectorVMContext& Context, FName AttributeToRead);
+	void ReadBool(FVectorVMContext& Context, FName AttributeToRead);
+	void ReadColor(FVectorVMContext& Context, FName AttributeToRead);
+	void ReadQuat(FVectorVMContext& Context, FName AttributeToRead);
 
 protected:
 	virtual bool CopyToInternal(UNiagaraDataInterface* Destination) const override;
 
 private:
 	template<typename T>
-	T RetrieveValueWithCheck(FNiagaraEmitterInstance* EmitterInstance, const FNiagaraTypeDefinition& Type, const FName& Attr, const FNiagaraID& ParticleID, bool &bValid);
-
+	T RetrieveValueWithCheck(FNiagaraEmitterInstance* EmitterInstance, const FNiagaraTypeDefinition& Type, const FName& Attr, const FNiagaraID& ParticleID, bool& bValid);
 };
\ No newline at end of file
diff --git a/Engine/Plugins/FX/Niagara/Source/Niagara/Classes/NiagaraDataSet.h b/Engine/Plugins/FX/Niagara/Source/Niagara/Classes/NiagaraDataSet.h
--- a/Engine/Plugins/FX/Niagara/Source/Niagara/Classes/NiagaraDataSet.h
+++ b/Engine/Plugins/FX/Niagara/Source/Niagara/Classes/NiagaraDataSet.h
@@ -132,8 +132,10 @@ public:
 
 	FORCEINLINE uint32 GetNumInstances()const { return NumInstances; }
 	FORCEINLINE uint32 GetNumInstancesAllocated()const { return NumInstancesAllocated; }
+	FORCEINLINE uint32 GetNumSpawnedInstances() const { return NumSpawnedInstances; }
 
 	FORCEINLINE void SetNumInstances(uint32 InNumInstances) { check(InNumInstances <= NumInstancesAllocated); NumInstances = InNumInstances; }
+	FORCEINLINE void SetNumSpawnedInstances(uint32 InNumSpawnedInstances) { NumSpawnedInstances = InNumSpawnedInstances; }
 	FORCEINLINE uint32 GetSizeBytes()const { return FloatData.Num() + Int32Data.Num(); }
 	FORCEINLINE FRWBuffer& GetGPUBufferFloat() { return GPUBufferFloat; }
 	FORCEINLINE FRWBuffer& GetGPUBufferInt() { return GPUBufferInt;	}
@@ -145,6 +147,9 @@ public:
 	FORCEINLINE uint32 GetFloatStride() const { return FloatStride; }
 	FORCEINLINE uint32 GetInt32Stride() const { return Int32Stride; }
 
+	FORCEINLINE uint32 GetIDAcquireTag() const { return IDAcquireTag; }
+	FORCEINLINE void SetIDAcquireTag(uint32 InTag) { IDAcquireTag = InTag; }
+
 	FORCEINLINE FNiagaraDataSet* GetOwner()const { return Owner; }
 
 	int32 TransferInstance(FNiagaraDataBuffer& SourceBuffer, int32 InstanceIndex, bool bRemoveFromSource=true);
@@ -208,6 +213,10 @@ private:
 	uint32 FloatStride;
 	/** Stride between components in the int32 buffer. */
 	uint32 Int32Stride;
+	/** Number of instances spawned in the last tick. */
+	uint32 NumSpawnedInstances;
+	/** ID acquire tag used in the last tick. */
+	uint32 IDAcquireTag;
 
 	/** Table containing current base locations for all registers in this dataset. */
 	TArray<uint8*> RegisterTable;//TODO: Should make inline? Feels like a useful size to keep local would be too big.
@@ -299,6 +308,7 @@ public:
 	FORCEINLINE bool GetNeedsPersistentIDs()const { return CompiledData.bNeedsPersistentIDs; }
 
 	FORCEINLINE TArray<int32>& GetFreeIDTable() { return FreeIDsTable; }
+	FORCEINLINE TArray<int32>& GetSpawnedIDsTable() { return SpawnedIDsTable; }
 	FORCEINLINE int32& GetNumFreeIDs() { return NumFreeIDs; }
 	FORCEINLINE int32& GetMaxUsedID() { return MaxUsedID; }
 	FORCEINLINE int32& GetIDAcquireTag() { return IDAcquireTag; }
@@ -382,6 +392,9 @@ private:
 	/** Tag to use when new IDs are acquired. Should be unique per tick. */
 	int32 IDAcquireTag;
 
+	/** Table of IDs spawned in the last tick (just the index part, the acquire tag is IDAcquireTag for all of them). */
+	TArray<int32> SpawnedIDsTable;
+
 	/** GPU buffer of free IDs available on the next tick. */
 	FRWBuffer GPUFreeIDs;
 
diff --git a/Engine/Plugins/FX/Niagara/Source/Niagara/Classes/NiagaraEmitterInstanceBatcher.h b/Engine/Plugins/FX/Niagara/Source/Niagara/Classes/NiagaraEmitterInstanceBatcher.h
--- a/Engine/Plugins/FX/Niagara/Source/Niagara/Classes/NiagaraEmitterInstanceBatcher.h
+++ b/Engine/Plugins/FX/Niagara/Source/Niagara/Classes/NiagaraEmitterInstanceBatcher.h
@@ -58,6 +58,7 @@ public:
 		// @todo REMOVE THIS HACK
 		, LastFrameThatDrainedData(GFrameNumberRenderThread)
 		, NumAllocatedFreeIDListSizes(0)
+		, bFreeIDListSizesBufferCleared(false)
 	{
 	}
 
@@ -165,8 +166,10 @@ public:
 	void DispatchMultipleStages(const FNiagaraGPUSystemTick& Tick, FNiagaraComputeInstanceData *Instance, FRHICommandList &RHICmdList, FRHIUniformBuffer* ViewUniformBuffer, FNiagaraShader* ComputeShader);
 
 private:
+	using FEmitterInstanceList = TArray<FNiagaraComputeInstanceData*>;
+
 	void ExecuteAll(FRHICommandList& RHICmdList, FRHIUniformBuffer* ViewUniformBuffer, bool bSetReadback, ETickStage TickStage);
-	void ResizeBuffersAndGatherResources(FOverlappableTicks& OverlappableTick, FRHICommandList& RHICmdList, FNiagaraBufferArray& ReadBuffers, FNiagaraBufferArray& WriteBuffers);
+	void ResizeBuffersAndGatherResources(FOverlappableTicks& OverlappableTick, FRHICommandList& RHICmdList, FNiagaraBufferArray& ReadBuffers, FNiagaraBufferArray& WriteBuffers, FNiagaraBufferArray& OutputGraphicsBuffers, FEmitterInstanceList& InstancesWithPersistentIDs);
 	void DispatchAllOnCompute(FOverlappableTicks& OverlappableTick, FRHICommandList& RHICmdList, FRHIUniformBuffer* ViewUniformBuffer, FNiagaraBufferArray& ReadBuffers, FNiagaraBufferArray& WriteBuffers, bool bSetReadback);
 
 	bool ShouldTickForStage(const FNiagaraGPUSystemTick& Tick, ETickStage TickStage) const;
@@ -185,6 +188,9 @@ private:
 
 	void FinishDispatches();
 	void ReleaseTicks();
+	void ResizeFreeIDsListSizesBuffer(uint32 NumInstances);
+	void ClearFreeIDsListSizesBuffer(FRHICommandList& RHICmdList);
+	void UpdateFreeIDBuffers(FRHICommandList& RHICmdList, FEmitterInstanceList& Instances);
 
 	/** Feature level of this effects system */
 	ERHIFeatureLevel::Type FeatureLevel;
@@ -210,8 +216,11 @@ private:
 	TArray<FNiagaraGPUSystemTick> Ticks_RT;
 	FGlobalDistanceFieldParameterData GlobalDistanceFieldParams;
 
-	/** A buffer used by the compute shader which determines the list of free particle IDs for each emitter. */
+	/** A buffer of list sizes used by UpdateFreeIDBuffers to allow overlapping several dispatches. */
 	FRWBuffer FreeIDListSizesBuffer;
 	uint32 NumAllocatedFreeIDListSizes;
-	uint32 CurrentFreeIDListIndex;
+	bool bFreeIDListSizesBufferCleared;
+
+	/** List of emitter instances which need their free ID buffers updated post render. */
+	FEmitterInstanceList DeferredIDBufferUpdates;
 };
diff --git a/Engine/Plugins/FX/Niagara/Source/Niagara/Private/NiagaraDataInterfaceParticleRead.cpp b/Engine/Plugins/FX/Niagara/Source/Niagara/Private/NiagaraDataInterfaceParticleRead.cpp
--- a/Engine/Plugins/FX/Niagara/Source/Niagara/Private/NiagaraDataInterfaceParticleRead.cpp
+++ b/Engine/Plugins/FX/Niagara/Source/Niagara/Private/NiagaraDataInterfaceParticleRead.cpp
@@ -8,6 +8,8 @@
 
 #define LOCTEXT_NAMESPACE "NiagaraDataInterfaceParticleRead"
 
+static const FName GetNumSpawnedParticlesFunctionName("Get Num Spawned Particles");
+static const FName GetSpawnedIDAtIndexFunctionName("Get Spawned ID At Index");
 static const FName GetIntAttributeFunctionName("Get int Attribute");
 static const FName GetFloatAttributeFunctionName("Get float Attribute");
 static const FName GetVec2AttributeFunctionName("Get Vector2 Attribute");
@@ -17,6 +19,9 @@ static const FName GetBoolAttributeFunctionName("Get bool Attribute");
 static const FName GetColorAttributeFunctionName("Get Color Attribute");
 static const FName GetQuatAttributeFunctionName("Get Quaternion Attribute");
 
+static const FString NumSpawnedParticlesBaseName(TEXT("NumSpawnedParticles_"));
+static const FString SpawnedParticlesAcquireTagBaseName(TEXT("SpawnedParticlesAcquireTag_"));
+static const FString SpawnedIDsBufferBaseName(TEXT("SpawnedIDsBuffer_"));
 static const FString IDToIndexTableBaseName(TEXT("IDToIndexTable_"));
 static const FString InputFloatBufferBaseName(TEXT("InputFloatBuffer_"));
 static const FString InputIntBufferBaseName(TEXT("InputIntBuffer_"));
@@ -116,6 +121,9 @@ struct FNiagaraDataInterfaceParametersCS_ParticleRead : public FNiagaraDataInter
 
 	virtual void Bind(const FNiagaraDataInterfaceParamRef& ParamRef, const class FShaderParameterMap& ParameterMap) override
 	{
+		NumSpawnedParticlesParam.Bind(ParameterMap, *(NumSpawnedParticlesBaseName + ParamRef.ParameterInfo.DataInterfaceHLSLSymbol));
+		SpawnedParticlesAcquireTagParam.Bind(ParameterMap, *(SpawnedParticlesAcquireTagBaseName + ParamRef.ParameterInfo.DataInterfaceHLSLSymbol));
+		SpawnedIDsBufferParam.Bind(ParameterMap, *(SpawnedIDsBufferBaseName + ParamRef.ParameterInfo.DataInterfaceHLSLSymbol));
 		IDToIndexTableParam.Bind(ParameterMap, *(IDToIndexTableBaseName + ParamRef.ParameterInfo.DataInterfaceHLSLSymbol));
 		InputFloatBufferParam.Bind(ParameterMap, *(InputFloatBufferBaseName + ParamRef.ParameterInfo.DataInterfaceHLSLSymbol));
 		InputIntBufferParam.Bind(ParameterMap, *(InputIntBufferBaseName + ParamRef.ParameterInfo.DataInterfaceHLSLSymbol));
@@ -139,16 +147,22 @@ struct FNiagaraDataInterfaceParametersCS_ParticleRead : public FNiagaraDataInter
 			}
 			else
 			{
+				// This is not an error. GetNumSpawnedParticles and GetSpawnedIDAtIndexFunctionName don't use specifiers,
+				// but they take up slots in the attribute indices array for simplicity. Just stick NAME_None in here to ignore them.
 				AttributeNames[FuncIdx] = NAME_None;
 				AttributeTypes[FuncIdx] = ENiagaraParticleDataValueType::Invalid;
 			}
 		}
 
-		AttributeIndices.SetNum(AttributeNames.Num());
+		int NumAttrIndices = Align(AttributeNames.Num(), 4);
+		AttributeIndices.SetNum(NumAttrIndices);
 	}
 
 	virtual void Serialize(FArchive& Ar) override
 	{
+		Ar << NumSpawnedParticlesParam;
+		Ar << SpawnedParticlesAcquireTagParam;
+		Ar << SpawnedIDsBufferParam;
 		Ar << IDToIndexTableParam;
 		Ar << InputFloatBufferParam;
 		Ar << InputIntBufferParam;
@@ -159,10 +173,11 @@ struct FNiagaraDataInterfaceParametersCS_ParticleRead : public FNiagaraDataInter
 		Ar << AttributeNames;
 		Ar << AttributeTypes;
 
-		AttributeIndices.SetNum(AttributeNames.Num());
+		int NumAttrIndices = Align(AttributeNames.Num(), 4);
+		AttributeIndices.SetNum(NumAttrIndices);
 	}
 
-	void SetErrorParams(FRHICommandList& RHICmdList, FRHIComputeShader* ComputeShader) const
+	void SetErrorParams(FRHICommandList& RHICmdList, FRHIComputeShader* ComputeShader, bool bSkipSpawnInfo) const
 	{
 		CachedDataSet = nullptr;
 
@@ -172,6 +187,13 @@ struct FNiagaraDataInterfaceParametersCS_ParticleRead : public FNiagaraDataInter
 		}
 		AcquireTagRegisterIndex = -1;
 
+		if (!bSkipSpawnInfo)
+		{
+			SetShaderValue(RHICmdList, ComputeShader, NumSpawnedParticlesParam, 0);
+			SetShaderValue(RHICmdList, ComputeShader, SpawnedParticlesAcquireTagParam, 0);
+			SetSRVParameter(RHICmdList, ComputeShader, SpawnedIDsBufferParam, FNiagaraRenderer::GetDummyIntBuffer().SRV);
+		}
+
 		SetSRVParameter(RHICmdList, ComputeShader, IDToIndexTableParam, FNiagaraRenderer::GetDummyIntBuffer().SRV);
 		SetSRVParameter(RHICmdList, ComputeShader, InputFloatBufferParam, FNiagaraRenderer::GetDummyFloatBuffer().SRV);
 		SetSRVParameter(RHICmdList, ComputeShader, InputIntBufferParam, FNiagaraRenderer::GetDummyIntBuffer().SRV);
@@ -197,19 +219,26 @@ struct FNiagaraDataInterfaceParametersCS_ParticleRead : public FNiagaraDataInter
 		}
 	}
 
-	void FindAttributeIndices(FNiagaraDataSet* SourceDataSet, const TCHAR* SourceEmitterName) const
+	void FindAttributeIndices(const FNiagaraDataSet* SourceDataSet, const TCHAR* SourceEmitterName) const
 	{
-		check(AttributeIndices.Num() == AttributeNames.Num());
+		check(AttributeIndices.Num() == Align(AttributeNames.Num(), 4));
 
 		const TArray<FNiagaraVariable>& SourceEmitterVariables = SourceDataSet->GetVariables();
 		const TArray<FNiagaraVariableLayoutInfo>& SourceEmitterVariableLayouts = SourceDataSet->GetVariableLayouts();
 		for (int AttrNameIdx = 0; AttrNameIdx < AttributeNames.Num(); ++AttrNameIdx)
 		{
+			const FName& AttrName = AttributeNames[AttrNameIdx];
+			if (AttrName == NAME_None)
+			{
+				AttributeIndices[AttrNameIdx] = -1;
+				continue;
+			}
+
 			bool FoundVariable = false;
 			for (int VarIdx = 0; VarIdx < SourceEmitterVariables.Num(); ++VarIdx)
 			{
 				const FNiagaraVariable& Var = SourceEmitterVariables[VarIdx];
-				if (Var.GetName() == AttributeNames[AttrNameIdx])
+				if (Var.GetName() == AttrName)
 				{
 					ENiagaraParticleDataValueType AttributeType = AttributeTypes[AttrNameIdx];
 					if (CheckVariableType(Var.GetType(), AttributeType))
@@ -231,7 +260,7 @@ struct FNiagaraDataInterfaceParametersCS_ParticleRead : public FNiagaraDataInter
 
 			if (!FoundVariable)
 			{
-				UE_LOG(LogNiagara, Error, TEXT("Particle read DI is trying to access inexistent variable '%s' in emitter '%s'."), *AttributeNames[AttrNameIdx].ToString(), SourceEmitterName);
+				UE_LOG(LogNiagara, Error, TEXT("Particle read DI is trying to access inexistent variable '%s' in emitter '%s'."), *AttrName.ToString(), SourceEmitterName);
 				AttributeIndices[AttrNameIdx] = -1;
 			}
 		}
@@ -251,6 +280,11 @@ struct FNiagaraDataInterfaceParametersCS_ParticleRead : public FNiagaraDataInter
 		{
 			UE_LOG(LogNiagara, Error, TEXT("Particle read DI cannot find ID variable in emitter '%s'."), SourceEmitterName);
 		}
+
+		for (int AttrIdx = AttributeNames.Num(); AttrIdx < AttributeIndices.Num(); ++AttrIdx)
+		{
+			AttributeIndices[AttrIdx] = -1;
+		}
 	}
 
 	virtual void Set(FRHICommandList& RHICmdList, const FNiagaraDataInterfaceSetArgs& Context) const override
@@ -259,14 +293,6 @@ struct FNiagaraDataInterfaceParametersCS_ParticleRead : public FNiagaraDataInter
 
 		FRHIComputeShader* ComputeShader = Context.Shader->GetComputeShader();
 
-		if (!InputFloatBufferParam.IsBound() && !InputIntBufferParam.IsBound())
-		{
-			// This DI instance didn't generate any reachable code, probably because all the values it sets are overwritten by
-			// other DIs. Don't bother with it.
-			SetErrorParams(RHICmdList, ComputeShader);
-			return;
-		}
-
 		FNiagaraDataInterfaceProxyParticleRead* Proxy = static_cast<FNiagaraDataInterfaceProxyParticleRead*>(Context.DataInterface);
 		check(Proxy);
 		if (Proxy->SourceEmitterGPUContext == nullptr)
@@ -277,7 +303,7 @@ struct FNiagaraDataInterfaceParametersCS_ParticleRead : public FNiagaraDataInter
 				UE_LOG(LogNiagara, Error, TEXT("GPU particle read DI is set to access CPU emitter '%s'."), *Proxy->SourceEmitterName);
 				bSourceEmitterNotGPUErrorShown = true;
 			}
-			SetErrorParams(RHICmdList, ComputeShader);
+			SetErrorParams(RHICmdList, ComputeShader, false);
 			return;
 		}
 
@@ -286,27 +312,46 @@ struct FNiagaraDataInterfaceParametersCS_ParticleRead : public FNiagaraDataInter
 		FNiagaraDataSet* SourceDataSet = Proxy->SourceEmitterGPUContext->MainDataSet;
 		if (!SourceDataSet)
 		{
-			SetErrorParams(RHICmdList, ComputeShader);
+			SetErrorParams(RHICmdList, ComputeShader, false);
 			return;
 		}
 
 		FNiagaraDataBuffer* SourceData;
-		if (Context.ComputeInstanceData->Context == Proxy->SourceEmitterGPUContext)
+		uint32 NumSpawnedInstances = 0, IDAcquireTag = 0;
+		const bool bReadingOwnEmitter = Context.ComputeInstanceData->Context == Proxy->SourceEmitterGPUContext;
+		if (bReadingOwnEmitter)
 		{
 			// If the current execution context is the same as the source emitter's context, it means we're reading from
 			// ourselves. We can't use SourceDataSet->GetCurrentData() in that case, because EndSimulate() has already been
 			// called on the current emitter, and the current data has been set to the destination data. We need to use the
 			// current compute instance data to get to the input buffers.
 			SourceData = Context.ComputeInstanceData->CurrentData;
+
+			// We still want to get the spawn count and ID acquire tag from the destination data, because that's where
+			// NiagaraEmitterInstanceBatcher::Run() stores them.
+			if (Context.ComputeInstanceData->DestinationData != nullptr)
+			{
+				NumSpawnedInstances = Context.ComputeInstanceData->DestinationData->GetNumSpawnedInstances();
+				IDAcquireTag = Context.ComputeInstanceData->DestinationData->GetIDAcquireTag();
+			}
 		}
 		else
 		{
 			SourceData = SourceDataSet->GetCurrentData();
+			if (SourceData)
+			{
+				NumSpawnedInstances = SourceData->GetNumSpawnedInstances();
+				IDAcquireTag = SourceData->GetIDAcquireTag();
+			}
 		}
 
+		SetShaderValue(RHICmdList, ComputeShader, NumSpawnedParticlesParam, NumSpawnedInstances);
+		SetShaderValue(RHICmdList, ComputeShader, SpawnedParticlesAcquireTagParam, IDAcquireTag);
+		SetSRVParameter(RHICmdList, ComputeShader, SpawnedIDsBufferParam, SourceDataSet->GetGPUFreeIDs().SRV);
+
 		if (!SourceData)
 		{
-			SetErrorParams(RHICmdList, ComputeShader);
+			SetErrorParams(RHICmdList, ComputeShader, true);
 			return;
 		}
 
@@ -321,10 +366,16 @@ struct FNiagaraDataInterfaceParametersCS_ParticleRead : public FNiagaraDataInter
 			// This can happen in the first frame, when there's no previous data yet. The DI shouldn't be
 			// queried in this case, because there's no way to have particle IDs (since there are no particles),
 			// but if it is it will just return failure and default values.
-			SetErrorParams(RHICmdList, ComputeShader);
+			SetErrorParams(RHICmdList, ComputeShader, true);
 			return;
 		}
 
+		if (!bReadingOwnEmitter)
+		{
+			FRHIUnorderedAccessView* InputBuffers[] = { SourceData->GetGPUIDToIndexTable().UAV, SourceData->GetGPUBufferFloat().UAV, SourceData->GetGPUBufferInt().UAV };
+			RHICmdList.TransitionResources(EResourceTransitionAccess::EReadable, EResourceTransitionPipeline::EComputeToCompute, InputBuffers, UE_ARRAY_COUNT(InputBuffers));
+		}
+
 		const uint32 ParticleStrideFloat = SourceData->GetFloatStride() / sizeof(float);
 		const uint32 ParticleStrideInt = SourceData->GetInt32Stride() / sizeof(int32);
 
@@ -336,8 +387,11 @@ struct FNiagaraDataInterfaceParametersCS_ParticleRead : public FNiagaraDataInter
 		SetShaderValueArray(RHICmdList, ComputeShader, AttributeIndicesParam, AttributeIndices.GetData(), AttributeIndices.Num());
 		SetShaderValue(RHICmdList, ComputeShader, AcquireTagRegisterIndexParam, AcquireTagRegisterIndex);
 	}
-
+	
 private:
+	FShaderParameter NumSpawnedParticlesParam;
+	FShaderParameter SpawnedParticlesAcquireTagParam;
+	FShaderResourceParameter SpawnedIDsBufferParam;
 	FShaderResourceParameter IDToIndexTableParam;
 	FShaderResourceParameter InputFloatBufferParam;
 	FShaderResourceParameter InputIntBufferParam;
@@ -369,25 +423,6 @@ void UNiagaraDataInterfaceParticleRead::PostInitProperties()
 	}
 }
 
-void UNiagaraDataInterfaceParticleRead::PostLoad()
-{
-	Super::PostLoad();
-}
-
-#if WITH_EDITOR
-
-void UNiagaraDataInterfaceParticleRead::PreEditChange(FProperty* PropertyAboutToChange)
-{
-	Super::PreEditChange(PropertyAboutToChange);
-}
-
-void UNiagaraDataInterfaceParticleRead::PostEditChangeProperty(struct FPropertyChangedEvent& PropertyChangedEvent)
-{
-	Super::PostEditChangeProperty(PropertyChangedEvent);
-}
-
-#endif
-
 bool UNiagaraDataInterfaceParticleRead::InitPerInstanceData(void* PerInstanceData, FNiagaraSystemInstance* SystemInstance)
 {
 	FNDIParticleRead_InstanceData* PIData = new (PerInstanceData) FNDIParticleRead_InstanceData;
@@ -401,26 +436,41 @@ bool UNiagaraDataInterfaceParticleRead::InitPerInstanceData(void* PerInstanceDat
 			break;
 		}
 	}
-	return (PIData->EmitterInstance != nullptr);
-}
 
-void UNiagaraDataInterfaceParticleRead::DestroyPerInstanceData(void* PerInstanceData, FNiagaraSystemInstance* SystemInstance)
-{
+	if (PIData->EmitterInstance == nullptr)
+	{
+		UE_LOG(LogNiagara, Error, TEXT("Source emitter '%s' not found."), *EmitterName);
+		return false;
+	}
 
+	return true;
 }
 
-bool UNiagaraDataInterfaceParticleRead::PerInstanceTick(void* PerInstanceData, FNiagaraSystemInstance* SystemInstance, float DeltaSeconds)
+void UNiagaraDataInterfaceParticleRead::GetFunctions(TArray<FNiagaraFunctionSignature>& OutFunctions)
 {
-	return false;
-}
+	{
+		FNiagaraFunctionSignature Sig;
+		Sig.Name = GetNumSpawnedParticlesFunctionName;
+		Sig.Inputs.Add(FNiagaraVariable(FNiagaraTypeDefinition(GetClass()), TEXT("Particle Reader")));
+		Sig.Outputs.Add(FNiagaraVariable(FNiagaraTypeDefinition::GetIntDef(), TEXT("Num Spawned")));
+		Sig.bMemberFunction = true;
+		Sig.bRequiresContext = false;
+		OutFunctions.Add(Sig);
+	}
+	{
+		FNiagaraFunctionSignature Sig;
+		Sig.Name = GetSpawnedIDAtIndexFunctionName;
 
-bool UNiagaraDataInterfaceParticleRead::PerInstanceTickPostSimulate(void* PerInstanceData, FNiagaraSystemInstance* SystemInstance, float DeltaSeconds)
-{
-	return false;
-}
+		Sig.Inputs.Add(FNiagaraVariable(FNiagaraTypeDefinition(GetClass()), TEXT("Particle Reader")));
+		Sig.Inputs.Add(FNiagaraVariable(FNiagaraTypeDefinition::GetIntDef(), TEXT("Spawn Index")));
 
-void UNiagaraDataInterfaceParticleRead::GetFunctions(TArray<FNiagaraFunctionSignature>& OutFunctions)
-{
+		Sig.Outputs.Add(FNiagaraVariable(FNiagaraTypeDefinition::GetBoolDef(), TEXT("Valid")));
+		Sig.Outputs.Add(FNiagaraVariable(FNiagaraTypeDefinition::GetIDDef(), TEXT("ID")));
+
+		Sig.bMemberFunction = true;
+		Sig.bRequiresContext = false;
+		OutFunctions.Add(Sig);
+	}
 	{
 		FNiagaraFunctionSignature Sig;
 		Sig.Name = GetFloatAttributeFunctionName;
@@ -542,6 +592,9 @@ void UNiagaraDataInterfaceParticleRead::GetFunctions(TArray<FNiagaraFunctionSign
 		OutFunctions.Add(Sig);
 	}
 }
+
+DEFINE_NDI_DIRECT_FUNC_BINDER(UNiagaraDataInterfaceParticleRead, GetNumSpawnedParticles);
+DEFINE_NDI_DIRECT_FUNC_BINDER(UNiagaraDataInterfaceParticleRead, GetSpawnedIDAtIndex);
 DEFINE_NDI_DIRECT_FUNC_BINDER_WITH_PAYLOAD(UNiagaraDataInterfaceParticleRead, ReadFloat);
 DEFINE_NDI_DIRECT_FUNC_BINDER_WITH_PAYLOAD(UNiagaraDataInterfaceParticleRead, ReadVector2);
 DEFINE_NDI_DIRECT_FUNC_BINDER_WITH_PAYLOAD(UNiagaraDataInterfaceParticleRead, ReadVector3);
@@ -550,8 +603,21 @@ DEFINE_NDI_DIRECT_FUNC_BINDER_WITH_PAYLOAD(UNiagaraDataInterfaceParticleRead, Re
 DEFINE_NDI_DIRECT_FUNC_BINDER_WITH_PAYLOAD(UNiagaraDataInterfaceParticleRead, ReadBool);
 DEFINE_NDI_DIRECT_FUNC_BINDER_WITH_PAYLOAD(UNiagaraDataInterfaceParticleRead, ReadColor);
 DEFINE_NDI_DIRECT_FUNC_BINDER_WITH_PAYLOAD(UNiagaraDataInterfaceParticleRead, ReadQuat);
-void UNiagaraDataInterfaceParticleRead::GetVMExternalFunction(const FVMExternalFunctionBindingInfo& BindingInfo, void* InstanceData, FVMExternalFunction &OutFunc)
+
+void UNiagaraDataInterfaceParticleRead::GetVMExternalFunction(const FVMExternalFunctionBindingInfo& BindingInfo, void* InstanceData, FVMExternalFunction& OutFunc)
 {
+	if (BindingInfo.Name == GetNumSpawnedParticlesFunctionName)
+	{
+		NDI_FUNC_BINDER(UNiagaraDataInterfaceParticleRead, GetNumSpawnedParticles)::Bind(this, OutFunc);
+		return;
+	}
+
+	if (BindingInfo.Name == GetSpawnedIDAtIndexFunctionName)
+	{
+		NDI_FUNC_BINDER(UNiagaraDataInterfaceParticleRead, GetSpawnedIDAtIndex)::Bind(this, OutFunc);
+		return;
+	}
+
 	bool bBindSuccessful = false;
 	FNDIParticleRead_InstanceData* PIData = static_cast<FNDIParticleRead_InstanceData*>(InstanceData);
 	static const FName NAME_Attribute("Attribute");
@@ -643,6 +709,61 @@ void UNiagaraDataInterfaceParticleRead::GetVMExternalFunction(const FVMExternalF
 	}
 }
 
+void UNiagaraDataInterfaceParticleRead::GetNumSpawnedParticles(FVectorVMContext& Context)
+{
+	VectorVM::FUserPtrHandler<FNDIParticleRead_InstanceData> InstData(Context);
+	VectorVM::FExternalFuncRegisterHandler<int32> OutNumSpawned(Context);
+
+	FNiagaraEmitterInstance* EmitterInstance = InstData.Get()->EmitterInstance;
+	const TArray<int32>& SpawnedIDsTable = EmitterInstance->GetData().GetSpawnedIDsTable();
+	int32 NumSpawned = SpawnedIDsTable.Num();
+
+	for (int32 InstanceIdx = 0; InstanceIdx < Context.NumInstances; ++InstanceIdx)
+	{
+		*OutNumSpawned.GetDestAndAdvance() = NumSpawned;
+	}
+}
+
+void UNiagaraDataInterfaceParticleRead::GetSpawnedIDAtIndex(FVectorVMContext& Context)
+{
+	VectorVM::FExternalFuncInputHandler<int32> InIndex(Context);
+
+	VectorVM::FUserPtrHandler<FNDIParticleRead_InstanceData> InstData(Context);
+
+	VectorVM::FExternalFuncRegisterHandler<FNiagaraBool> OutValid(Context);
+	VectorVM::FExternalFuncRegisterHandler<int32> OutIDIndex(Context);
+	VectorVM::FExternalFuncRegisterHandler<int32> OutIDAcquireTag(Context);
+
+	FNiagaraEmitterInstance* EmitterInstance = InstData.Get()->EmitterInstance;
+	const TArray<int32>& SpawnedIDsTable = EmitterInstance->GetData().GetSpawnedIDsTable();
+	int32 NumSpawned = SpawnedIDsTable.Num();
+	int32 IDAcquireTag = EmitterInstance->GetData().GetIDAcquireTag();
+
+	for (int32 InstanceIdx = 0; InstanceIdx < Context.NumInstances; ++InstanceIdx)
+	{
+		FNiagaraBool ValidValue;
+		FNiagaraID IDValue;
+
+		int32 SpawnIndex = InIndex.GetAndAdvance();
+		if (SpawnIndex >= 0 && SpawnIndex < NumSpawned)
+		{
+			ValidValue.SetValue(true);
+			IDValue.Index = SpawnedIDsTable[SpawnIndex];
+			IDValue.AcquireTag = IDAcquireTag;
+		}
+		else
+		{
+			ValidValue.SetValue(false);
+			IDValue.Index = 0;
+			IDValue.AcquireTag = 0;
+		}
+
+		*OutValid.GetDestAndAdvance() = ValidValue;
+		*OutIDIndex.GetDestAndAdvance() = IDValue.Index;
+		*OutIDAcquireTag.GetDestAndAdvance() = IDValue.AcquireTag;
+	}
+}
+
 void UNiagaraDataInterfaceParticleRead::ReadFloat(FVectorVMContext& Context, FName AttributeToRead)
 {
 	VectorVM::FExternalFuncInputHandler<int32> ParticleIDIndexParam(Context);
@@ -906,31 +1027,38 @@ bool UNiagaraDataInterfaceParticleRead::CopyToInternal(UNiagaraDataInterface* De
 	return true;
 }
 
-void UNiagaraDataInterfaceParticleRead::GetCommonHLSL(FString& OutHLSL)
-{
-}
-
 void UNiagaraDataInterfaceParticleRead::GetParameterDefinitionHLSL(const FNiagaraDataInterfaceGPUParamInfo& ParamInfo, FString& OutHLSL)
 {
 	static const TCHAR *FormatDeclarations = TEXT(
+		"int {NumSpawnedParticlesName};\n"
+		"int {SpawnedParticlesAcquireTagName};\n"
+		"uint {ParticleStrideFloatName};\n"
+		"uint {ParticleStrideIntName};\n"
+		"int {AcquireTagRegisterIndexName};\n"
+		"Buffer<int> {SpawnedIDsBufferName};\n"
 		"Buffer<int> {IDToIndexTableName};\n"
 		"Buffer<float> {InputFloatBufferName};\n"
 		"Buffer<int> {InputIntBufferName};\n"
-		"uint {ParticleStrideFloatName};\n"
-		"uint {ParticleStrideIntName};\n"
-		"int {AttributeIndicesName}[{AttributeCount}];\n"
-		"int {AcquireTagRegisterIndexName};\n\n"
+		"int4 {AttributeIndicesName}[{AttributeInt4Count}];\n\n"
 	);
 
+	// If we use an int array for the attribute indices, the shader compiler will actually use int4 due to the packing rules,
+	// and leave 3 elements unused. Besides being wasteful, this means that the array we send to the CS would need to be padded,
+	// which is a hassle. Instead, use int4 explicitly, and access individual components in the generated code.
+	const int AttributeInt4Count = FMath::DivideAndRoundUp(ParamInfo.GeneratedFunctions.Num(), 4);
+
 	TMap<FString, FStringFormatArg> ArgsDeclarations;
+	ArgsDeclarations.Add(TEXT("NumSpawnedParticlesName"), NumSpawnedParticlesBaseName + ParamInfo.DataInterfaceHLSLSymbol);
+	ArgsDeclarations.Add(TEXT("SpawnedParticlesAcquireTagName"), SpawnedParticlesAcquireTagBaseName + ParamInfo.DataInterfaceHLSLSymbol);
+	ArgsDeclarations.Add(TEXT("ParticleStrideFloatName"), ParticleStrideFloatBaseName + ParamInfo.DataInterfaceHLSLSymbol);
+	ArgsDeclarations.Add(TEXT("ParticleStrideIntName"), ParticleStrideIntBaseName + ParamInfo.DataInterfaceHLSLSymbol);
+	ArgsDeclarations.Add(TEXT("AcquireTagRegisterIndexName"), AcquireTagRegisterIndexBaseName + ParamInfo.DataInterfaceHLSLSymbol);
+	ArgsDeclarations.Add(TEXT("SpawnedIDsBufferName"), SpawnedIDsBufferBaseName + ParamInfo.DataInterfaceHLSLSymbol);
 	ArgsDeclarations.Add(TEXT("IDToIndexTableName"), IDToIndexTableBaseName + ParamInfo.DataInterfaceHLSLSymbol);
 	ArgsDeclarations.Add(TEXT("InputFloatBufferName"), InputFloatBufferBaseName + ParamInfo.DataInterfaceHLSLSymbol);
 	ArgsDeclarations.Add(TEXT("InputIntBufferName"), InputIntBufferBaseName + ParamInfo.DataInterfaceHLSLSymbol);
-	ArgsDeclarations.Add(TEXT("ParticleStrideFloatName"), ParticleStrideFloatBaseName + ParamInfo.DataInterfaceHLSLSymbol);
-	ArgsDeclarations.Add(TEXT("ParticleStrideIntName"), ParticleStrideIntBaseName + ParamInfo.DataInterfaceHLSLSymbol);
 	ArgsDeclarations.Add(TEXT("AttributeIndicesName"), AttributeIndicesBaseName + ParamInfo.DataInterfaceHLSLSymbol);
-	ArgsDeclarations.Add(TEXT("AttributeCount"), ParamInfo.GeneratedFunctions.Num());
-	ArgsDeclarations.Add(TEXT("AcquireTagRegisterIndexName"), AcquireTagRegisterIndexBaseName + ParamInfo.DataInterfaceHLSLSymbol);
+	ArgsDeclarations.Add(TEXT("AttributeInt4Count"), AttributeInt4Count);
 
 	OutHLSL += FString::Format(FormatDeclarations, ArgsDeclarations);
 }
@@ -940,7 +1068,7 @@ static bool GenerateGetFunctionHLSL(const FNiagaraDataInterfaceGPUParamInfo& Par
 	static const TCHAR* FuncTemplate = TEXT(
 		"void {FunctionName}(NiagaraID In_ParticleID, out bool Out_Valid, out {ValueType} Out_Value)\n"
 		"{\n"
-		"    int RegisterIndex = {AttributeIndicesName}[{FunctionInstanceIndex}];\n"
+		"    int RegisterIndex = {AttributeIndicesName}[{AttributeIndexGroup}]{AttributeIndexComponent};\n"
 		"    int ParticleIndex = (RegisterIndex != -1) && (In_ParticleID.Index >= 0) ? {IDToIndexTableName}[In_ParticleID.Index] : -1;\n"
 		"    int AcquireTag = (ParticleIndex != -1) ? {InputIntBufferName}[{AcquireTagRegisterIndexName}*{ParticleStrideIntName} + ParticleIndex] : 0;\n"
 		"    if(ParticleIndex != -1 && In_ParticleID.AcquireTag == AcquireTag)\n"
@@ -956,6 +1084,8 @@ static bool GenerateGetFunctionHLSL(const FNiagaraDataInterfaceGPUParamInfo& Par
 		"}\n\n"
 	);
 
+	static const TCHAR* ComponentNames[] = { TEXT(".x"), TEXT(".y"), TEXT(".z"), TEXT(".w") };
+
 	const FString ParticleStrideFloatName = ParticleStrideFloatBaseName + ParamInfo.DataInterfaceHLSLSymbol;
 	const FString ParticleStrideIntName = ParticleStrideIntBaseName + ParamInfo.DataInterfaceHLSLSymbol;
 	const FString InputFloatBufferName = InputFloatBufferBaseName + ParamInfo.DataInterfaceHLSLSymbol;
@@ -995,7 +1125,6 @@ static bool GenerateGetFunctionHLSL(const FNiagaraDataInterfaceGPUParamInfo& Par
 	FString FetchValueCode;
 	for (int ComponentIndex = 0; ComponentIndex < NumComponents; ++ComponentIndex)
 	{
-		static const TCHAR* ComponentNames[] = { TEXT(".x"), TEXT(".y"), TEXT(".z"), TEXT(".w") };
 		const TCHAR* ComponentName = (NumComponents > 1) ? ComponentNames[ComponentIndex] : TEXT("");
 		const FString FetchComponentCode = FString::Printf(TEXT("        Out_Value%s = %s(%s[(RegisterIndex + %d)*%s + ParticleIndex]);\n"), ComponentName, ComponentTypeName, InputBufferName, ComponentIndex, InputBufferStrideName);
 		FetchValueCode += FetchComponentCode;
@@ -1007,7 +1136,8 @@ static bool GenerateGetFunctionHLSL(const FNiagaraDataInterfaceGPUParamInfo& Par
 	FuncTemplateArgs.Add(TEXT("FunctionName"), FunctionInfo.InstanceName);
 	FuncTemplateArgs.Add(TEXT("ValueType"), ValueTypeName);
 	FuncTemplateArgs.Add(TEXT("AttributeIndicesName"), AttributeIndicesBaseName + ParamInfo.DataInterfaceHLSLSymbol);
-	FuncTemplateArgs.Add(TEXT("FunctionInstanceIndex"), FunctionInstanceIndex);
+	FuncTemplateArgs.Add(TEXT("AttributeIndexGroup"), FunctionInstanceIndex / 4);
+	FuncTemplateArgs.Add(TEXT("AttributeIndexComponent"), ComponentNames[FunctionInstanceIndex % 4]);
 	FuncTemplateArgs.Add(TEXT("IDToIndexTableName"), IDToIndexTableBaseName + ParamInfo.DataInterfaceHLSLSymbol);
 	FuncTemplateArgs.Add(TEXT("InputIntBufferName"), InputIntBufferName);
 	FuncTemplateArgs.Add(TEXT("AcquireTagRegisterIndexName"), AcquireTagRegisterIndexBaseName + ParamInfo.DataInterfaceHLSLSymbol);
@@ -1022,6 +1152,53 @@ static bool GenerateGetFunctionHLSL(const FNiagaraDataInterfaceGPUParamInfo& Par
 
 bool UNiagaraDataInterfaceParticleRead::GetFunctionHLSL(const FNiagaraDataInterfaceGPUParamInfo& ParamInfo, const FNiagaraDataInterfaceGeneratedFunction& FunctionInfo, int FunctionInstanceIndex, FString& OutHLSL)
 {
+	if (FunctionInfo.DefinitionName == GetNumSpawnedParticlesFunctionName)
+	{
+		static const TCHAR* FuncTemplate = TEXT(
+			"void {FunctionName}(out int Out_NumSpawned)\n"
+			"{\n"
+			"    return {NumSpawnedParticlesName};\n"
+			"}\n\n"
+		);
+
+		TMap<FString, FStringFormatArg> FuncTemplateArgs;
+		FuncTemplateArgs.Add(TEXT("FunctionName"), FunctionInfo.InstanceName);
+		FuncTemplateArgs.Add(TEXT("NumSpawnedParticlesName"), NumSpawnedParticlesBaseName + ParamInfo.DataInterfaceHLSLSymbol);
+
+		OutHLSL += FString::Format(FuncTemplate, FuncTemplateArgs);
+		return true;
+	}
+	
+	if (FunctionInfo.DefinitionName == GetSpawnedIDAtIndexFunctionName)
+	{
+		static const TCHAR* FuncTemplate = TEXT(
+			"void {FunctionName}(int In_SpawnIndex, out bool Out_Valid, out NiagaraID Out_ID)\n"
+			"{\n"
+			"    if(In_SpawnIndex >= 0 && In_SpawnIndex < {NumSpawnedParticlesName})\n"
+			"    {\n"
+			"        Out_Valid = true;\n"
+			"        Out_ID.Index = {SpawnedIDsBufferName}[In_SpawnIndex];\n"
+			"        Out_ID.AcquireTag = {SpawnedParticlesAcquireTagName};\n"
+			"    }\n"
+			"    else\n"
+			"    {\n"
+			"        Out_Valid = false;\n"
+			"        Out_ID.Index = 0;\n"
+			"        Out_ID.AcquireTag = 0;\n"
+			"    }\n"
+			"}\n\n"
+		);
+
+		TMap<FString, FStringFormatArg> FuncTemplateArgs;
+		FuncTemplateArgs.Add(TEXT("FunctionName"), FunctionInfo.InstanceName);
+		FuncTemplateArgs.Add(TEXT("NumSpawnedParticlesName"), NumSpawnedParticlesBaseName + ParamInfo.DataInterfaceHLSLSymbol);
+		FuncTemplateArgs.Add(TEXT("SpawnedParticlesAcquireTagName"), SpawnedParticlesAcquireTagBaseName + ParamInfo.DataInterfaceHLSLSymbol);
+		FuncTemplateArgs.Add(TEXT("SpawnedIDsBufferName"), SpawnedIDsBufferBaseName + ParamInfo.DataInterfaceHLSLSymbol);
+
+		OutHLSL += FString::Format(FuncTemplate, FuncTemplateArgs);
+		return true;
+	}
+
 	if (FunctionInfo.DefinitionName == GetIntAttributeFunctionName)
 	{
 		return GenerateGetFunctionHLSL(ParamInfo, FunctionInfo, FunctionInstanceIndex, ENiagaraParticleDataComponentType::Int, 1, OutHLSL);
diff --git a/Engine/Plugins/FX/Niagara/Source/Niagara/Private/NiagaraDataSet.cpp b/Engine/Plugins/FX/Niagara/Source/Niagara/Private/NiagaraDataSet.cpp
--- a/Engine/Plugins/FX/Niagara/Source/Niagara/Private/NiagaraDataSet.cpp
+++ b/Engine/Plugins/FX/Niagara/Source/Niagara/Private/NiagaraDataSet.cpp
@@ -127,6 +127,7 @@ void FNiagaraDataSet::ResetBuffersInternal()
 	FreeIDsTable.Reset();
 	NumFreeIDs = 0;
 	MaxUsedID = INDEX_NONE;
+	SpawnedIDsTable.Reset();
 
 	//Ensure we have a valid current buffer
 	BeginSimulate();
@@ -248,6 +249,9 @@ void FNiagaraDataSet::Allocate(int32 NumInstances, bool bMaintainExisting)
 			//Free ID Table must always be at least as large as the data buffer + it's current size in the case all particles die this frame.
 			FreeIDsTable.AddUninitialized(NumNewIDs);
 
+			// The spawned IDs table must be as large as the free IDs table, in case we allocate all of them this tick.
+			SpawnedIDsTable.Reserve(FreeIDsTable.Num());
+
 			//Free table should always have enough room for these new IDs.
 			check(NumFreeIDs + NumNewIDs <= FreeIDsTable.Num());
 
@@ -363,6 +367,8 @@ void FNiagaraDataSet::AllocateGPUFreeIDs(uint32 InNumInstances, FRHICommandList&
 		return;
 	}
 
+	SCOPED_DRAW_EVENTF(RHICmdList, NiagaraGPUComputeInitFreeIDs, TEXT("Init Free IDs - %s"), DebugSimName ? DebugSimName : TEXT(""));
+
 	TCHAR DebugBufferName[128];
 	FCString::Snprintf(DebugBufferName, UE_ARRAY_COUNT(DebugBufferName), TEXT("NiagaraFreeIDList_%s"), DebugSimName ? DebugSimName : TEXT(""));
 	FRWBuffer NewFreeIDsBuffer;
@@ -485,6 +491,8 @@ FNiagaraDataBuffer::FNiagaraDataBuffer(FNiagaraDataSet* InOwner)
 	, NumInstancesAllocated(0)
 	, FloatStride(0)
 	, Int32Stride(0)
+	, NumSpawnedInstances(0)
+	, IDAcquireTag(0)
 {
 }
 
diff --git a/Engine/Plugins/FX/Niagara/Source/Niagara/Private/NiagaraEmitterInstance.cpp b/Engine/Plugins/FX/Niagara/Source/Niagara/Private/NiagaraEmitterInstance.cpp
--- a/Engine/Plugins/FX/Niagara/Source/Niagara/Private/NiagaraEmitterInstance.cpp
+++ b/Engine/Plugins/FX/Niagara/Source/Niagara/Private/NiagaraEmitterInstance.cpp
@@ -1327,6 +1327,8 @@ void FNiagaraEmitterInstance::Tick(float DeltaSeconds)
 	int32 NumBeforeSpawn = Data.GetDestinationDataChecked().GetNumInstances();
 	uint32 TotalActualEventSpawns = 0;
 
+	Data.GetSpawnedIDsTable().SetNum(0, false);
+
 	//Init new particles with the spawn script.
 	if (SpawnTotal + EventSpawnTotal > 0)
 	{
@@ -1425,6 +1427,11 @@ void FNiagaraEmitterInstance::Tick(float DeltaSeconds)
 		}
 	}
 
+	// The spawn count and ID acquire tag in FNiagaraDataBuffer are currently only used by the GPU path,
+	// but we'll set them here anyway, to prevent surprises in case someone sees them and decides to use them.
+	Data.GetDestinationDataChecked().SetNumSpawnedInstances(Data.GetSpawnedIDsTable().Num());
+	Data.GetDestinationDataChecked().SetIDAcquireTag(Data.GetIDAcquireTag());
+
 	//We're done with this simulation pass.
 	Data.EndSimulate();
 	for (FNiagaraDataSet* SpawnEventDataSet : SpawnScriptEventDataSets)
diff --git a/Engine/Plugins/FX/Niagara/Source/Niagara/Private/NiagaraEmitterInstanceBatcher.cpp b/Engine/Plugins/FX/Niagara/Source/Niagara/Private/NiagaraEmitterInstanceBatcher.cpp
--- a/Engine/Plugins/FX/Niagara/Source/Niagara/Private/NiagaraEmitterInstanceBatcher.cpp
+++ b/Engine/Plugins/FX/Niagara/Source/Niagara/Private/NiagaraEmitterInstanceBatcher.cpp
@@ -27,9 +27,9 @@ DECLARE_DWORD_COUNTER_STAT(TEXT("Readback latency (frames)"), STAT_NiagaraReadba
 
 DECLARE_GPU_STAT_NAMED(NiagaraGPU, TEXT("Niagara"));
 DECLARE_GPU_STAT_NAMED(NiagaraGPUSimulation, TEXT("Niagara GPU Simulation"));
-DECLARE_GPU_STAT_NAMED(NiagaraGPUClearIDTable, TEXT("NiagaraGPU Clear ID Table"));
-DECLARE_GPU_STAT_NAMED(NiagaraGPURunEmitter, TEXT("NiagaraGPU Run Emitter"));
-DECLARE_GPU_STAT_NAMED(NiagaraGPUComputeFreeIDs, TEXT("Niagara GPU Compute Free IDs"));
+DECLARE_GPU_STAT_NAMED(NiagaraGPUClearIDTables, TEXT("NiagaraGPU Clear ID Tables"));
+DECLARE_GPU_STAT_NAMED(NiagaraGPUComputeFreeIDs, TEXT("Niagara GPU Compute All Free IDs"));
+DECLARE_GPU_STAT_NAMED(NiagaraGPUComputeFreeIDsEmitter, TEXT("Niagara GPU Compute Emitter Free IDs"));
 DECLARE_GPU_STAT_NAMED(NiagaraGPUSorting, TEXT("Niagara GPU sorting"));
 
 uint32 FNiagaraComputeExecutionContext::TickCounter = 0;
@@ -347,12 +347,10 @@ void NiagaraEmitterInstanceBatcher::DispatchMultipleStages(const FNiagaraGPUSyst
 	}
 }
 
-void NiagaraEmitterInstanceBatcher::ResizeBuffersAndGatherResources(FOverlappableTicks& OverlappableTick, FRHICommandList& RHICmdList, FNiagaraBufferArray& ReadBuffers, FNiagaraBufferArray& WriteBuffers)
+void NiagaraEmitterInstanceBatcher::ResizeBuffersAndGatherResources(FOverlappableTicks& OverlappableTick, FRHICommandList& RHICmdList, FNiagaraBufferArray& ReadBuffers, FNiagaraBufferArray& WriteBuffers, FNiagaraBufferArray& OutputGraphicsBuffers, FEmitterInstanceList& InstancesWithPersistentIDs)
 {
 	SCOPE_CYCLE_COUNTER(STAT_NiagaraGPUDispatchSetup_RT);
 
-	uint32 NumFreeIDListsToCompute = 0;
-
 	//UE_LOG(LogNiagara, Warning, TEXT("NiagaraEmitterInstanceBatcher::ResizeBuffersAndGatherResources:  %0xP"), this);
 	for (FNiagaraGPUSystemTick* Tick : OverlappableTick)
 	{
@@ -398,13 +396,13 @@ void NiagaraEmitterInstanceBatcher::ResizeBuffersAndGatherResources(FOverlappabl
 			if (bNeedsPersistentIDs)
 			{
 				Context->MainDataSet->AllocateGPUFreeIDs(AllocatedInstances + 1, RHICmdList, FeatureLevel, Context->GetDebugSimName());
+				ReadBuffers.Add(Context->MainDataSet->GetGPUFreeIDs().UAV);
+				InstancesWithPersistentIDs.Add(&Instance);
 			}
 
 			DestinationData.AllocateGPU(AllocatedInstances + 1, GPUInstanceCounterManager, RHICmdList, FeatureLevel, Context->GetDebugSimName());
 			DestinationData.SetNumInstances(RequiredInstances);
 
-			WriteBuffers.Add(GPUInstanceCounterManager.GetInstanceCountBuffer().UAV);
-
 			if ( Shader->FloatInputBufferParam.IsBound() )
 			{
 				ReadBuffers.Add(CurrentData.GetGPUBufferFloat().UAV);
@@ -417,17 +415,23 @@ void NiagaraEmitterInstanceBatcher::ResizeBuffersAndGatherResources(FOverlappabl
 			if ( Shader->FloatOutputBufferParam.IsBound() )
 			{
 				WriteBuffers.Add(DestinationData.GetGPUBufferFloat().UAV);
+				if (bIsFinalTick)
+				{
+					OutputGraphicsBuffers.Add(DestinationData.GetGPUBufferFloat().UAV);
+				}
 			}
 			if ( Shader->IntOutputBufferParam.IsBound() )
 			{
 				WriteBuffers.Add(DestinationData.GetGPUBufferInt().UAV);
+				if (bIsFinalTick)
+				{
+					OutputGraphicsBuffers.Add(DestinationData.GetGPUBufferInt().UAV);
+				}
 			}
 
 			if (bNeedsPersistentIDs)
 			{
-				ReadBuffers.Add(Context->MainDataSet->GetGPUFreeIDs().UAV);
 				WriteBuffers.Add(DestinationData.GetGPUIDToIndexTable().UAV);
-				++NumFreeIDListsToCompute;
 			}
 
 			Context->MainDataSet->EndSimulate();
@@ -438,24 +442,30 @@ void NiagaraEmitterInstanceBatcher::ResizeBuffersAndGatherResources(FOverlappabl
 		}
 	}
 
-	if (NumFreeIDListsToCompute > NumAllocatedFreeIDListSizes)
+	uint32 NumInstancesWithPersistentIDs = (uint32)InstancesWithPersistentIDs.Num();
+	if (NumInstancesWithPersistentIDs > 0)
 	{
-		constexpr uint32 ALLOC_CHUNK_SIZE = 128;
-		NumAllocatedFreeIDListSizes = FMath::DivideAndRoundUp(NumFreeIDListsToCompute, ALLOC_CHUNK_SIZE) * ALLOC_CHUNK_SIZE;
-		if (FreeIDListSizesBuffer.Buffer)
+		// These buffers will be needed by the simulation dispatches which come immediately after, so there will be a stall, but
+		// moving this step to a different place is difficult, and the stall is not large, so we'll live with it for now.
+		SCOPED_DRAW_EVENT(RHICmdList, NiagaraGPUClearIDTables);
+		SCOPED_GPU_STAT(RHICmdList, NiagaraGPUClearIDTables);
+
+		FNiagaraBufferArray IDToIndexTables;
+		IDToIndexTables.SetNum(NumInstancesWithPersistentIDs);
+		for (uint32 i = 0; i < NumInstancesWithPersistentIDs; ++i)
 		{
-			FreeIDListSizesBuffer.Release();
+			FNiagaraComputeInstanceData* Instance = InstancesWithPersistentIDs[i];
+			IDToIndexTables[i] = Instance->DestinationData->GetGPUIDToIndexTable().UAV;
 		}
-		FreeIDListSizesBuffer.Initialize(sizeof(uint32), NumAllocatedFreeIDListSizes, EPixelFormat::PF_R32_UINT, BUF_Static, TEXT("NiagaraFreeIDListSizes"));
-	}
+		// TODO: is it sufficient to do a CS cache flush before all this and get rid of these explicit barriers?
+		RHICmdList.TransitionResources(EResourceTransitionAccess::EWritable, EResourceTransitionPipeline::EComputeToCompute, IDToIndexTables.GetData(), IDToIndexTables.Num());
 
-	CurrentFreeIDListIndex = 0;
-	if (NumFreeIDListsToCompute > 0)
-	{
-		SCOPED_DRAW_EVENT(RHICmdList, NiagaraGPUClearFreeIDListSizeTable);
-		RHICmdList.TransitionResource(EResourceTransitionAccess::ERWBarrier, EResourceTransitionPipeline::EComputeToCompute, FreeIDListSizesBuffer.UAV);
-		RHICmdList.ClearUAVUint(FreeIDListSizesBuffer.UAV, FUintVector4(ForceInitToZero));
-		RHICmdList.TransitionResource(EResourceTransitionAccess::ERWBarrier, EResourceTransitionPipeline::EComputeToCompute, FreeIDListSizesBuffer.UAV);
+		for (uint32 i = 0; i < NumInstancesWithPersistentIDs; ++i)
+		{
+			FNiagaraComputeInstanceData* Instance = InstancesWithPersistentIDs[i];
+			SCOPED_DRAW_EVENTF(RHICmdList, NiagaraGPUComputeClearIDToIndexBuffer, TEXT("Clear ID To Index Table - %s"), Instance->Context->GetDebugSimName());
+			NiagaraFillGPUIntBuffer(RHICmdList, FeatureLevel, Instance->DestinationData->GetGPUIDToIndexTable(), INDEX_NONE);
+		}
 	}
 }
 
@@ -464,8 +474,6 @@ void NiagaraEmitterInstanceBatcher::DispatchAllOnCompute(FOverlappableTicks& Ove
 	FRHICommandListImmediate& RHICmdListImmediate = FRHICommandListExecutor::GetImmediateCommandList();
 
 	//UE_LOG(LogNiagara, Warning, TEXT("NiagaraEmitterInstanceBatcher::DispatchAllOnCompute:  %0xP"), this);
-	// Disable automatic cache flush so that we can have our compute work overlapping. Barrier will be used as a sync mechanism.
-	RHICmdList.AutomaticCacheFlushAfterComputeShader(false);
 
 #if WITH_EDITORONLY_DATA
 	{
@@ -524,12 +532,6 @@ void NiagaraEmitterInstanceBatcher::DispatchAllOnCompute(FOverlappableTicks& Ove
 		}
 	}
 
-	//
-	RHICmdList.TransitionResources(EResourceTransitionAccess::EReadable, EResourceTransitionPipeline::EComputeToGfx, WriteBuffers.GetData(), WriteBuffers.Num());
-	//	Now Copy to staging buffer the data we want to read back (alive particle count). And make buffer ready for that and draw commands on the graphics pipe too.
-	//
-	RHICmdList.AutomaticCacheFlushAfterComputeShader(true);
-
 	if (GNiagaraSubmitCommands)
 	{
 		RHICmdList.SubmitCommandsHint();
@@ -545,6 +547,12 @@ void NiagaraEmitterInstanceBatcher::PostRenderOpaque(FRHICommandListImmediate& R
 		// Setup new readback since if there is no pending request, there is no risk of having invalid data read (offset being allocated after the readback was sent).
 		ExecuteAll(RHICmdList, ViewUniformBuffer, !GPUInstanceCounterManager.HasPendingGPUReadback(), ETickStage::PostOpaqueRender);
 
+		RHICmdList.AutomaticCacheFlushAfterComputeShader(false);
+		UpdateFreeIDBuffers(RHICmdList, DeferredIDBufferUpdates);
+		RHICmdList.AutomaticCacheFlushAfterComputeShader(true);
+
+		DeferredIDBufferUpdates.SetNum(0, false);
+
 		FinishDispatches();
 	}
 
@@ -574,7 +582,73 @@ bool NiagaraEmitterInstanceBatcher::ShouldTickForStage(const FNiagaraGPUSystemTi
 	return TickStage == ETickStage::PreInitViews;
 }
 
-void NiagaraEmitterInstanceBatcher::ExecuteAll(FRHICommandList &RHICmdList, FRHIUniformBuffer* ViewUniformBuffer, bool bSetReadback, ETickStage TickStage)
+void NiagaraEmitterInstanceBatcher::ResizeFreeIDsListSizesBuffer(uint32 NumInstances)
+{
+	if (NumInstances <= NumAllocatedFreeIDListSizes)
+	{
+		return;
+	}
+
+	constexpr uint32 ALLOC_CHUNK_SIZE = 128;
+	NumAllocatedFreeIDListSizes = Align(NumInstances, ALLOC_CHUNK_SIZE);
+	if (FreeIDListSizesBuffer.Buffer)
+	{
+		FreeIDListSizesBuffer.Release();
+	}
+	FreeIDListSizesBuffer.Initialize(sizeof(uint32), NumAllocatedFreeIDListSizes, EPixelFormat::PF_R32_SINT, BUF_Static, TEXT("NiagaraFreeIDListSizes"));
+	bFreeIDListSizesBufferCleared = false;
+}
+
+void NiagaraEmitterInstanceBatcher::ClearFreeIDsListSizesBuffer(FRHICommandList& RHICmdList)
+{
+	if (bFreeIDListSizesBufferCleared)
+	{
+		return;
+	}
+
+	SCOPED_DRAW_EVENT(RHICmdList, NiagaraGPUComputeClearFreeIDListSizes);
+	RHICmdList.TransitionResource(EResourceTransitionAccess::EWritable, EResourceTransitionPipeline::EComputeToCompute, FreeIDListSizesBuffer.UAV);
+	NiagaraFillGPUIntBuffer(RHICmdList, FeatureLevel, FreeIDListSizesBuffer, 0);
+	bFreeIDListSizesBufferCleared = true;
+}
+
+void NiagaraEmitterInstanceBatcher::UpdateFreeIDBuffers(FRHICommandList& RHICmdList, FEmitterInstanceList& Instances)
+{
+	if (Instances.Num() == 0)
+	{
+		return;
+	}
+
+	SCOPED_DRAW_EVENT(RHICmdList, NiagaraGPUComputeFreeIDs);
+	SCOPED_GPU_STAT(RHICmdList, NiagaraGPUComputeFreeIDs);
+
+	FNiagaraBufferArray ReadBuffers, WriteBuffers;
+	for(FNiagaraComputeInstanceData* Instance : Instances)
+	{
+		ReadBuffers.Add(Instance->DestinationData->GetGPUIDToIndexTable().UAV);
+		WriteBuffers.Add(Instance->Context->MainDataSet->GetGPUFreeIDs().UAV);
+	}
+
+	RHICmdList.TransitionResource(EResourceTransitionAccess::ERWBarrier, EResourceTransitionPipeline::EComputeToCompute, FreeIDListSizesBuffer.UAV);
+	RHICmdList.TransitionResources(EResourceTransitionAccess::EReadable, EResourceTransitionPipeline::EComputeToCompute, ReadBuffers.GetData(), ReadBuffers.Num());
+	RHICmdList.TransitionResources(EResourceTransitionAccess::EWritable, EResourceTransitionPipeline::EComputeToCompute, WriteBuffers.GetData(), WriteBuffers.Num());
+
+	check((uint32)Instances.Num() <= NumAllocatedFreeIDListSizes);
+
+	for (uint32 InstanceIdx = 0; InstanceIdx < (uint32)Instances.Num(); ++InstanceIdx)
+	{
+		FNiagaraComputeInstanceData* Instance = Instances[InstanceIdx];
+		FNiagaraDataSet* MainDataSet = Instance->Context->MainDataSet;
+		FNiagaraDataBuffer* DestinationData = Instance->DestinationData;
+
+		SCOPED_DRAW_EVENTF(RHICmdList, NiagaraGPUComputeFreeIDsEmitter, TEXT("Update Free ID Buffer - %s"), Instance->Context->GetDebugSimName());
+		NiagaraComputeGPUFreeIDs(RHICmdList, FeatureLevel, MainDataSet->GetGPUNumAllocatedIDs(), DestinationData->GetGPUIDToIndexTable().SRV, MainDataSet->GetGPUFreeIDs(), FreeIDListSizesBuffer, InstanceIdx);
+	}
+
+	bFreeIDListSizesBufferCleared = false;
+}
+
+void NiagaraEmitterInstanceBatcher::ExecuteAll(FRHICommandList& RHICmdList, FRHIUniformBuffer* ViewUniformBuffer, bool bSetReadback, ETickStage TickStage)
 {
 	SCOPE_CYCLE_COUNTER(STAT_NiagaraGPUSimTick_RT);
 
@@ -662,18 +736,60 @@ void NiagaraEmitterInstanceBatcher::ExecuteAll(FRHICommandList &RHICmdList, FRHI
 		}
 	}
 
-	for (auto& SimPass : SimPasses)
+	RHICmdList.AutomaticCacheFlushAfterComputeShader(false);
+
+	FEmitterInstanceList InstancesWithPersistentIDs;
+	FNiagaraBufferArray ReadBuffers, WriteBuffers, OutputGraphicsBuffers;
+
+	for (int32 SimPassIdx = 0; SimPassIdx < SimPasses.Num(); ++SimPassIdx)
 	{
-		FNiagaraBufferArray ReadBuffers, WriteBuffers;
+		FOverlappableTicks& SimPass = SimPasses[SimPassIdx];
+		ReadBuffers.SetNum(0, false);
+		WriteBuffers.SetNum(0, false);
+		InstancesWithPersistentIDs.SetNum(0, false);
 
 		// This initial pass gathers all the buffers that are read from and written to so we can do batch resource transitions.
 		// It also ensures the GPU buffers are large enough to hold everything.
-		ResizeBuffersAndGatherResources(SimPass, RHICmdList, ReadBuffers, WriteBuffers);
+		ResizeBuffersAndGatherResources(SimPass, RHICmdList, ReadBuffers, WriteBuffers, OutputGraphicsBuffers, InstancesWithPersistentIDs);
 
-		SCOPED_DRAW_EVENT(RHICmdList, NiagaraGPUSimulation);
-		SCOPED_GPU_STAT(RHICmdList, NiagaraGPUSimulation);
-		DispatchAllOnCompute(SimPass, RHICmdList, ViewUniformBuffer, ReadBuffers, WriteBuffers, bSetReadback);
+		{
+			SCOPED_DRAW_EVENT(RHICmdList, NiagaraGPUSimulation);
+			SCOPED_GPU_STAT(RHICmdList, NiagaraGPUSimulation);
+			DispatchAllOnCompute(SimPass, RHICmdList, ViewUniformBuffer, ReadBuffers, WriteBuffers, bSetReadback);
+		}
+
+		if (InstancesWithPersistentIDs.Num() == 0)
+		{
+			continue;
+		}
+
+		// If we're doing multiple ticks (e.g. when scrubbing the timeline in the editor), we must update the free ID buffers before running
+		// the next tick, which will cause stalls (because the ID to index buffer is written by DispatchAllOnCompute and read by UpdateFreeIDBuffers).
+		// However, when we're at the last tick, we can postpone the update until later in the frame and avoid the stall. This will be the case when
+		// running normally, with one tick per frame.
+		if (SimPassIdx < SimPasses.Num() - 1)
+		{
+			ResizeFreeIDsListSizesBuffer(InstancesWithPersistentIDs.Num());
+			ClearFreeIDsListSizesBuffer(RHICmdList);
+			UpdateFreeIDBuffers(RHICmdList, InstancesWithPersistentIDs);
+		}
+		else
+		{
+			DeferredIDBufferUpdates.Append(InstancesWithPersistentIDs);
+			ResizeFreeIDsListSizesBuffer(DeferredIDBufferUpdates.Num());
+
+			// Speculatively clear the list sizes buffer here. Under normal circumstances, this happens in the first stage which finds instances with persistent IDs
+			// (usually PreInitViews) and it's finished by the time the deferred updates need to be processed. If a subsequent tick stage runs multiple time ticks,
+			// the first step will find the buffer already cleared and will not clear again. The only time when this clear is superfluous is when a following stage
+			// reallocates the buffer, but that's unlikely (and amortized) because we allocate in chunks.
+			ClearFreeIDsListSizesBuffer(RHICmdList);
+		}
 	}
+
+	OutputGraphicsBuffers.Add(GPUInstanceCounterManager.GetInstanceCountBuffer().UAV);
+	RHICmdList.TransitionResources(EResourceTransitionAccess::EReadable, EResourceTransitionPipeline::EComputeToGfx, OutputGraphicsBuffers.GetData(), OutputGraphicsBuffers.Num());
+
+	RHICmdList.AutomaticCacheFlushAfterComputeShader(true);
 }
 
 void NiagaraEmitterInstanceBatcher::PreInitViews(FRHICommandListImmediate& RHICmdList, bool bAllowGPUParticleUpdate)
@@ -1210,20 +1326,17 @@ void NiagaraEmitterInstanceBatcher::Run(const FNiagaraGPUSystemTick& Tick, const
 	FNiagaraDataBuffer& DestinationData = *Instance->DestinationData;
 	FNiagaraDataBuffer& CurrentData = *Instance->CurrentData;
 
-	const bool bNeedsPersistentIDs = Context->MainDataSet->GetNeedsPersistentIDs();
-	if (bNeedsPersistentIDs)
+	int32 InstancesToSpawnThisFrame = Instance->SpawnInfo.SpawnRateInstances + Instance->SpawnInfo.EventSpawnTotal;
+
+	// Only spawn particles on the first stage
+	if (HasRunParticleStage)
 	{
-		// Put INDEX_NONE in all the slots of the ID to index table. The simulation will fill in the
-		// indices for the IDs which are currently in use, so we can compute the free ID list based on
-		// the unused slots.
-		SCOPED_DRAW_EVENT(RHICmdList, NiagaraGPUClearIDTable);
-		SCOPED_GPU_STAT(RHICmdList, NiagaraGPUClearIDTable);
-		RHICmdList.TransitionResource(EResourceTransitionAccess::ERWBarrier, EResourceTransitionPipeline::EComputeToCompute, DestinationData.GetGPUIDToIndexTable().UAV);
-		RHICmdList.ClearUAVUint(DestinationData.GetGPUIDToIndexTable().UAV, FUintVector4(INDEX_NONE));
-		RHICmdList.TransitionResource(EResourceTransitionAccess::ERWBarrier, EResourceTransitionPipeline::EComputeToCompute, DestinationData.GetGPUIDToIndexTable().UAV);
+		InstancesToSpawnThisFrame = 0;
 	}
 
 	FRHIComputeShader* ComputeShader = Shader->GetComputeShader();
+	DestinationData.SetNumSpawnedInstances(InstancesToSpawnThisFrame);
+	DestinationData.SetIDAcquireTag(FNiagaraComputeExecutionContext::TickCounter);
 
 	RHICmdList.SetComputeShader(Shader->GetComputeShader());
 
@@ -1240,6 +1353,7 @@ void NiagaraEmitterInstanceBatcher::Run(const FNiagaraGPUSystemTick& Tick, const
 
 	// set the shader and data set params 
 	//
+	const bool bNeedsPersistentIDs = Context->MainDataSet->GetNeedsPersistentIDs();
 	SetSRVParameter(RHICmdList, Shader->GetComputeShader(), Shader->FreeIDBufferParam, bNeedsPersistentIDs ? Context->MainDataSet->GetGPUFreeIDs().SRV : FNiagaraRenderer::GetDummyIntBuffer().SRV);
 	CurrentData.SetShaderParams(Shader, RHICmdList, true);
 	DestinationData.SetShaderParams(Shader, RHICmdList, false);
@@ -1268,14 +1382,6 @@ void NiagaraEmitterInstanceBatcher::Run(const FNiagaraGPUSystemTick& Tick, const
 	SetShaderValueArray(RHICmdList, ComputeShader, Shader->EmitterSpawnInfoParamsParam, SpawnInfo.SpawnInfoParams, NIAGARA_MAX_GPU_SPAWN_INFOS);
 
 	SetShaderValue(RHICmdList, ComputeShader, Shader->UpdateStartInstanceParam, UpdateStartInstance);					// 0, except for event handler runs
-	int32 InstancesToSpawnThisFrame = Instance->SpawnInfo.SpawnRateInstances + Instance->SpawnInfo.EventSpawnTotal;
-
-	// Only spawn particles on the first stage
-	if (HasRunParticleStage)
-	{
-		InstancesToSpawnThisFrame = 0;
-	}
-
 	SetShaderValue(RHICmdList, ComputeShader, Shader->NumSpawnedInstancesParam, InstancesToSpawnThisFrame);				// number of instances in the spawn run
 	SetShaderValue(RHICmdList, ComputeShader, Shader->DefaultShaderStageIndexParam, DefaultShaderStageIndex);					// 0, except if several stages are defined
 	SetShaderValue(RHICmdList, ComputeShader, Shader->ShaderStageIndexParam, ShaderStageIndex);					// 0, except if several stages are defined
@@ -1328,8 +1434,6 @@ void NiagaraEmitterInstanceBatcher::Run(const FNiagaraGPUSystemTick& Tick, const
 	// Dispatch, if anything needs to be done
 	if (TotalNumInstances)
 	{
-		SCOPED_DRAW_EVENT(RHICmdList, NiagaraGPURunEmitter);
-		SCOPED_GPU_STAT(RHICmdList, NiagaraGPURunEmitter);
 		DispatchComputeShader(RHICmdList, Shader, NumThreadGroups, 1, 1);
 	}
 
@@ -1384,18 +1488,4 @@ void NiagaraEmitterInstanceBatcher::Run(const FNiagaraGPUSystemTick& Tick, const
 	CurrentData.UnsetShaderParams(Shader, RHICmdList);
 	DestinationData.UnsetShaderParams(Shader, RHICmdList);
 	Shader->InstanceCountsParam.UnsetUAV(RHICmdList, ComputeShader);
-
-	if (bNeedsPersistentIDs && ShaderStageIndex == Context->MaxUpdateIterations - 1)
-	{
-		SCOPED_DRAW_EVENT(RHICmdList, NiagaraGPUComputeFreeIDs);
-		SCOPED_GPU_STAT(RHICmdList, NiagaraGPUComputeFreeIDs);
-
-		check(CurrentFreeIDListIndex < NumAllocatedFreeIDListSizes);
-
-		// Make the ID to index table readable and the free ID list writable.
-		RHICmdList.TransitionResource(EResourceTransitionAccess::EReadable, EResourceTransitionPipeline::EComputeToCompute, DestinationData.GetGPUIDToIndexTable().UAV);
-		RHICmdList.TransitionResource(EResourceTransitionAccess::ERWBarrier, EResourceTransitionPipeline::EComputeToCompute, Context->MainDataSet->GetGPUFreeIDs().UAV);
-		NiagaraComputeGPUFreeIDs(RHICmdList, FeatureLevel, Context->MainDataSet->GetGPUNumAllocatedIDs(), DestinationData.GetGPUIDToIndexTable().SRV, Context->MainDataSet->GetGPUFreeIDs(), FreeIDListSizesBuffer, CurrentFreeIDListIndex);
-		++CurrentFreeIDListIndex;
-	}
 }
diff --git a/Engine/Plugins/FX/Niagara/Source/Niagara/Private/NiagaraScriptExecutionContext.cpp b/Engine/Plugins/FX/Niagara/Source/Niagara/Private/NiagaraScriptExecutionContext.cpp
--- a/Engine/Plugins/FX/Niagara/Source/Niagara/Private/NiagaraScriptExecutionContext.cpp
+++ b/Engine/Plugins/FX/Niagara/Source/Niagara/Private/NiagaraScriptExecutionContext.cpp
@@ -161,7 +161,7 @@ void FNiagaraScriptExecutionContext::BindData(int32 Index, FNiagaraDataSet& Data
 	uint8** OutputRegisters = Output ? Output->GetRegisterTable().GetData() : nullptr;
 	DataSetMetaTable.SetNum(FMath::Max(DataSetMetaTable.Num(), Index + 1));
 	DataSetMetaTable[Index].Init(InputRegisters, OutputRegisters, StartInstance,
-		Output ? &Output->GetIDTable() : nullptr, &DataSet.GetFreeIDTable(), &DataSet.GetNumFreeIDs(), &DataSet.GetMaxUsedID(), DataSet.GetIDAcquireTag());
+		Output ? &Output->GetIDTable() : nullptr, &DataSet.GetFreeIDTable(), &DataSet.GetNumFreeIDs(), &DataSet.GetMaxUsedID(), DataSet.GetIDAcquireTag(), &DataSet.GetSpawnedIDsTable());
 }
 
 void FNiagaraScriptExecutionContext::BindData(int32 Index, FNiagaraDataBuffer* Input, int32 StartInstance, bool bUpdateInstanceCounts)
@@ -172,7 +172,7 @@ void FNiagaraScriptExecutionContext::BindData(int32 Index, FNiagaraDataBuffer* I
 	DataSetInfo[Index].Init(DataSet, Input, nullptr, StartInstance, bUpdateInstanceCounts);
 
 	DataSetMetaTable.SetNum(FMath::Max(DataSetMetaTable.Num(), Index + 1));
-	DataSetMetaTable[Index].Init(Input->GetRegisterTable().GetData(), nullptr, StartInstance, nullptr, nullptr, &DataSet->GetNumFreeIDs(), &DataSet->GetMaxUsedID(), DataSet->GetIDAcquireTag());
+	DataSetMetaTable[Index].Init(Input->GetRegisterTable().GetData(), nullptr, StartInstance, nullptr, nullptr, &DataSet->GetNumFreeIDs(), &DataSet->GetMaxUsedID(), DataSet->GetIDAcquireTag(), &DataSet->GetSpawnedIDsTable());
 }
 
 bool FNiagaraScriptExecutionContext::Execute(uint32 NumInstances, const FScriptExecutionConstantBufferTable& ConstantBufferTable)
diff --git a/Engine/Plugins/FX/Niagara/Source/NiagaraCore/Private/NiagaraCustomVersion.cpp b/Engine/Plugins/FX/Niagara/Source/NiagaraCore/Private/NiagaraCustomVersion.cpp
--- a/Engine/Plugins/FX/Niagara/Source/NiagaraCore/Private/NiagaraCustomVersion.cpp
+++ b/Engine/Plugins/FX/Niagara/Source/NiagaraCore/Private/NiagaraCustomVersion.cpp
@@ -8,4 +8,4 @@ const FGuid FNiagaraCustomVersion::GUID(0xFCF57AFA, 0x50764283, 0xB9A9E658, 0xFF
 // Register the custom version with core
 FCustomVersionRegistration GRegisterNiagaraCustomVersion(FNiagaraCustomVersion::GUID, FNiagaraCustomVersion::LatestVersion, TEXT("NiagaraVer"));
 
-const FGuid FNiagaraCustomVersion::LatestScriptCompileVersion(0xBCBD4D65, 0xF6FA42CB, 0xA9F3851F, 0x15E98637);
+const FGuid FNiagaraCustomVersion::LatestScriptCompileVersion(0xCA9F2224, 0x728A48DD, 0x87238CFD, 0x5DF7DED7);
diff --git a/Engine/Plugins/FX/Niagara/Source/NiagaraShader/Private/NiagaraShaderParticleID.cpp b/Engine/Plugins/FX/Niagara/Source/NiagaraShader/Private/NiagaraShaderParticleID.cpp
--- a/Engine/Plugins/FX/Niagara/Source/NiagaraShader/Private/NiagaraShaderParticleID.cpp
+++ b/Engine/Plugins/FX/Niagara/Source/NiagaraShader/Private/NiagaraShaderParticleID.cpp
@@ -152,3 +152,63 @@ void NiagaraComputeGPUFreeIDs(FRHICommandList& RHICmdList, ERHIFeatureLevel::Typ
 	TShaderMapRef<NiagaraComputeFreeIDsCS> ComputeFreeIDsCS(GetGlobalShaderMap(FeatureLevel));
 	ComputeFreeIDsCS->Execute(RHICmdList, FeatureLevel, NumIDs, IDToIndexTable, FreeIDList, FreeIDListSizes, FreeIDListIndex);
 }
+
+class NiagaraFillIntBufferCS : public FGlobalShader
+{
+	DECLARE_GLOBAL_SHADER(NiagaraFillIntBufferCS);
+
+public:
+	NiagaraFillIntBufferCS() : FGlobalShader() {}
+
+	NiagaraFillIntBufferCS(const ShaderMetaType::CompiledShaderInitializerType& Initializer) : FGlobalShader(Initializer)
+	{
+		TargetBufferParam.Bind(Initializer.ParameterMap, TEXT("TargetBuffer"));
+		FillValueParam.Bind(Initializer.ParameterMap, TEXT("FillValue"));
+		BufferSizeParam.Bind(Initializer.ParameterMap, TEXT("BufferSize"));
+	}
+
+	static bool ShouldCompilePermutation(const FGlobalShaderPermutationParameters& Parameters)
+	{
+		return RHISupportsComputeShaders(Parameters.Platform);
+	}
+
+	bool Serialize(FArchive& Ar)
+	{
+		bool bShaderHasOutdatedParameters = FGlobalShader::Serialize(Ar);
+		Ar << TargetBufferParam;
+		Ar << FillValueParam;
+		Ar << BufferSizeParam;
+		return bShaderHasOutdatedParameters;
+	}
+
+	void Execute(FRHICommandList& RHICmdList, FRWBuffer& Buffer, int Value)
+	{
+		const uint32 THREAD_COUNT = 64;
+		const uint32 NumInts = Buffer.NumBytes / sizeof(int);
+		const uint32 ThreadGroups = FMath::DivideAndRoundUp(NumInts, THREAD_COUNT);
+
+		FRHIComputeShader* ComputeShader = GetComputeShader();
+		RHICmdList.SetComputeShader(ComputeShader);
+
+		SetUAVParameter(RHICmdList, ComputeShader, TargetBufferParam, Buffer.UAV);
+		SetShaderValue(RHICmdList, ComputeShader, FillValueParam, Value);
+		SetShaderValue(RHICmdList, ComputeShader, BufferSizeParam, NumInts);
+
+		DispatchComputeShader(RHICmdList, this, ThreadGroups, 1, 1);
+
+		SetUAVParameter(RHICmdList, ComputeShader, TargetBufferParam, nullptr);
+	}
+
+private:
+	FShaderResourceParameter TargetBufferParam;
+	FShaderParameter FillValueParam;
+	FShaderParameter BufferSizeParam;
+};
+
+IMPLEMENT_GLOBAL_SHADER(NiagaraFillIntBufferCS, "/Plugin/FX/Niagara/Private/NiagaraFillIntBuffer.usf", "FillIntBuffer", SF_Compute);
+
+void NiagaraFillGPUIntBuffer(FRHICommandList& RHICmdList, ERHIFeatureLevel::Type FeatureLevel, FRWBuffer& Buffer, int Value)
+{
+	TShaderMapRef<NiagaraFillIntBufferCS> FillCS(GetGlobalShaderMap(FeatureLevel));
+	FillCS->Execute(RHICmdList, Buffer, Value);
+}
diff --git a/Engine/Plugins/FX/Niagara/Source/NiagaraShader/Public/NiagaraShaderParticleID.h b/Engine/Plugins/FX/Niagara/Source/NiagaraShader/Public/NiagaraShaderParticleID.h
--- a/Engine/Plugins/FX/Niagara/Source/NiagaraShader/Public/NiagaraShaderParticleID.h
+++ b/Engine/Plugins/FX/Niagara/Source/NiagaraShader/Public/NiagaraShaderParticleID.h
@@ -6,3 +6,4 @@
 
 NIAGARASHADER_API void NiagaraInitGPUFreeIDList(FRHICommandList& RHICmdList, ERHIFeatureLevel::Type FeatureLevel, uint32 NumElementsToAlloc, FRWBuffer& NewBuffer, uint32 NumExistingElements, FRHIShaderResourceView* ExistingBuffer);
 NIAGARASHADER_API void NiagaraComputeGPUFreeIDs(FRHICommandList& RHICmdList, ERHIFeatureLevel::Type FeatureLevel, uint32 NumIDs, FRHIShaderResourceView* IDToIndexTable, FRWBuffer& FreeIDList, FRWBuffer& FreeIDListSizes, uint32 FreeIDListIndex);
+NIAGARASHADER_API void NiagaraFillGPUIntBuffer(FRHICommandList& RHICmdList, ERHIFeatureLevel::Type FeatureLevel, FRWBuffer& Buffer, int Value);
diff --git a/Engine/Source/Runtime/VectorVM/Private/VectorVM.cpp b/Engine/Source/Runtime/VectorVM/Private/VectorVM.cpp
--- a/Engine/Source/Runtime/VectorVM/Private/VectorVM.cpp
+++ b/Engine/Source/Runtime/VectorVM/Private/VectorVM.cpp
@@ -1421,6 +1421,7 @@ struct FScalarKernelAcquireID
 		const int32 DataSetIndex = Context.DecodeU16();
 		const TArrayView<FDataSetMeta> MetaTable = Context.DataSetMetaTable;
 		TArray<int32>&RESTRICT FreeIDTable = *MetaTable[DataSetIndex].FreeIDTable;
+		TArray<int32>&RESTRICT SpawnedIDsTable = *MetaTable[DataSetIndex].SpawnedIDsTable;
 
 		const int32 Tag = MetaTable[DataSetIndex].IDAcquireTag;
 
@@ -1453,6 +1454,8 @@ struct FScalarKernelAcquireID
 			*IDTag = Tag;
 			++IDIndex;
 			++IDTag;
+
+			SpawnedIDsTable.Add(AcquiredID);
 		}
 
 		MetaTable[DataSetIndex].UnlockFreeTable();
diff --git a/Engine/Source/Runtime/VectorVM/Public/VectorVM.h b/Engine/Source/Runtime/VectorVM/Public/VectorVM.h
--- a/Engine/Source/Runtime/VectorVM/Public/VectorVM.h
+++ b/Engine/Source/Runtime/VectorVM/Public/VectorVM.h
@@ -165,6 +165,7 @@ struct FDataSetMeta
 	
 	TArray<int32>*RESTRICT IDTable;
 	TArray<int32>*RESTRICT FreeIDTable;
+	TArray<int32>*RESTRICT SpawnedIDsTable;
 
 	/** Number of free IDs in the FreeIDTable */
 	int32* NumFreeIDs;
@@ -190,6 +191,7 @@ struct FDataSetMeta
 		, InstanceOffset(INDEX_NONE)
 		, IDTable(nullptr)
 		, FreeIDTable(nullptr)
+		, SpawnedIDsTable(nullptr)
 		, NumFreeIDs(nullptr)
 		, MaxUsedID(nullptr)
 		, IDAcquireTag(INDEX_NONE)
@@ -204,12 +206,13 @@ struct FDataSetMeta
 		InstanceOffset = INDEX_NONE;
 		IDTable = nullptr;
 		FreeIDTable = nullptr;
+		SpawnedIDsTable = nullptr;
 		NumFreeIDs = nullptr;
 		MaxUsedID = nullptr;
 		IDAcquireTag = INDEX_NONE;
 	}
 
-	FORCEINLINE void Init(uint8*RESTRICT *RESTRICT InInputRegisters, uint8*RESTRICT *RESTRICT InOutputRegisters, int32 InInstanceOffset, TArray<int32>* InIDTable, TArray<int32>* InFreeIDTable, int32* InNumFreeIDs, int32* InMaxUsedID, int32 InIDAcquireTag)
+	FORCEINLINE void Init(uint8*RESTRICT *RESTRICT InInputRegisters, uint8*RESTRICT *RESTRICT InOutputRegisters, int32 InInstanceOffset, TArray<int32>* InIDTable, TArray<int32>* InFreeIDTable, int32* InNumFreeIDs, int32* InMaxUsedID, int32 InIDAcquireTag, TArray<int32>* InSpawnedIDsTable)
 	{
 		InputRegisters = InInputRegisters;
 		OutputRegisters = InOutputRegisters;
@@ -221,6 +224,7 @@ struct FDataSetMeta
 		NumFreeIDs = InNumFreeIDs;
 		MaxUsedID = InMaxUsedID;
 		IDAcquireTag = InIDAcquireTag;
+		SpawnedIDsTable = InSpawnedIDsTable;
 	}
 
 private:
