- Copy component space transforms to avoid race around the skeletal mesh component being destroyed

diff --git a/Engine/Plugins/FX/Niagara/Source/Niagara/Classes/NDISkeletalMeshCommon.h b/Engine/Plugins/FX/Niagara/Source/Niagara/Classes/NDISkeletalMeshCommon.h
--- a/Engine/Plugins/FX/Niagara/Source/Niagara/Classes/NDISkeletalMeshCommon.h
+++ b/Engine/Plugins/FX/Niagara/Source/Niagara/Classes/NDISkeletalMeshCommon.h
@@ -15,13 +15,6 @@ struct FSkeletalMeshAccessorHelper
 		IndexBuffer = LODData->MultiSizeIndexContainer.GetIndexBuffer();
 		SkinningData = InstData->SkinningData.SkinningData.Get();
 		Usage = InstData->SkinningData.Usage;
-
-		if (Comp)
-		{
-			const USkinnedMeshComponent* BaseComp = Comp->GetBaseComponent();
-			BoneComponentSpaceTransforms = &BaseComp->GetComponentSpaceTransforms();
-			PrevBoneComponentSpaceTransforms = &BaseComp->GetPreviousComponentTransformsArray();
-		}
 	}
 
 	USkeletalMeshComponent* Comp = nullptr;
@@ -34,8 +27,6 @@ struct FSkeletalMeshAccessorHelper
 	const FSkeletalMeshSamplingRegionBuiltData* SamplingRegionBuiltData = nullptr;
 	FSkeletalMeshSkinningData* SkinningData = nullptr;
 	FSkeletalMeshSkinningDataUsage Usage;
-	const TArray<FTransform>* BoneComponentSpaceTransforms = nullptr;
-	const TArray<FTransform>* PrevBoneComponentSpaceTransforms = nullptr;
 };
 
 template<>
@@ -161,30 +152,22 @@ struct FSkinnedPositionAccessorHelper<TIntegralConstant<ENDISkeletalMesh_Skinnin
 
 	FORCEINLINE_DEBUGGABLE FVector GetSkinnedBonePosition(FSkeletalMeshAccessorHelper& Accessor, int32 BoneIndex)
 	{
-		return (*Accessor.BoneComponentSpaceTransforms)[BoneIndex].GetLocation();
+		return Accessor.SkinningData->CurrComponentTransforms()[BoneIndex].GetLocation();
 	}
 
 	FORCEINLINE_DEBUGGABLE FVector GetSkinnedBonePreviousPosition(FSkeletalMeshAccessorHelper& Accessor, int32 BoneIndex)
 	{
-		if (Accessor.PrevBoneComponentSpaceTransforms != nullptr && Accessor.PrevBoneComponentSpaceTransforms->Num() > 0)
-		{
-			return (*Accessor.PrevBoneComponentSpaceTransforms)[BoneIndex].GetLocation();
-		}
-		return GetSkinnedBonePosition(Accessor, BoneIndex);
+		return Accessor.SkinningData->PrevComponentTransforms()[BoneIndex].GetLocation();
 	}
 
 	FORCEINLINE_DEBUGGABLE FQuat GetSkinnedBoneRotation(FSkeletalMeshAccessorHelper& Accessor, int32 BoneIndex)
 	{
-		return (*Accessor.BoneComponentSpaceTransforms)[BoneIndex].GetRotation();
+		return Accessor.SkinningData->CurrComponentTransforms()[BoneIndex].GetRotation();
 	}
 
 	FORCEINLINE_DEBUGGABLE FQuat GetSkinnedBonePreviousRotation(FSkeletalMeshAccessorHelper& Accessor, int32 BoneIndex)
 	{
-		if (Accessor.PrevBoneComponentSpaceTransforms != nullptr && Accessor.PrevBoneComponentSpaceTransforms->Num() > 0)
-		{
-			return (*Accessor.PrevBoneComponentSpaceTransforms)[BoneIndex].GetRotation();
-		}
-		return GetSkinnedBoneRotation(Accessor, BoneIndex);
+		return Accessor.SkinningData->PrevComponentTransforms()[BoneIndex].GetRotation();
 	}
 };
 
@@ -226,22 +209,22 @@ struct FSkinnedPositionAccessorHelper<TIntegralConstant<ENDISkeletalMesh_Skinnin
 
 	FORCEINLINE_DEBUGGABLE FVector GetSkinnedBonePosition(FSkeletalMeshAccessorHelper& Accessor, int32 BoneIndex)
 	{
-		return (*Accessor.BoneComponentSpaceTransforms)[BoneIndex].GetLocation();
+		return Accessor.SkinningData->CurrComponentTransforms()[BoneIndex].GetLocation();
 	}
 
 	FORCEINLINE_DEBUGGABLE FVector GetSkinnedBonePreviousPosition(FSkeletalMeshAccessorHelper& Accessor, int32 BoneIndex)
 	{
-		return (*Accessor.PrevBoneComponentSpaceTransforms)[BoneIndex].GetLocation();
+		return Accessor.SkinningData->PrevComponentTransforms()[BoneIndex].GetLocation();
 	}
 
 	FORCEINLINE_DEBUGGABLE FQuat GetSkinnedBoneRotation(FSkeletalMeshAccessorHelper& Accessor, int32 BoneIndex)
 	{
-		return (*Accessor.BoneComponentSpaceTransforms)[BoneIndex].GetRotation();
+		return Accessor.SkinningData->CurrComponentTransforms()[BoneIndex].GetRotation();
 	}
 
 	FORCEINLINE_DEBUGGABLE FQuat GetSkinnedBonePreviousRotation(FSkeletalMeshAccessorHelper& Accessor, int32 BoneIndex)
 	{
-		return (*Accessor.PrevBoneComponentSpaceTransforms)[BoneIndex].GetRotation();
+		return Accessor.SkinningData->PrevComponentTransforms()[BoneIndex].GetRotation();
 	}
 };
 
diff --git a/Engine/Plugins/FX/Niagara/Source/Niagara/Classes/NiagaraDataInterfaceSkeletalMesh.h b/Engine/Plugins/FX/Niagara/Source/Niagara/Classes/NiagaraDataInterfaceSkeletalMesh.h
--- a/Engine/Plugins/FX/Niagara/Source/Niagara/Classes/NiagaraDataInterfaceSkeletalMesh.h
+++ b/Engine/Plugins/FX/Niagara/Source/Niagara/Classes/NiagaraDataInterfaceSkeletalMesh.h
@@ -134,8 +134,17 @@ struct FSkeletalMeshSkinningData
 		return BoneRefToLocals[CurrIndex ^ 1];
 	}
 
-private:
+	FORCEINLINE TArray<FTransform>& CurrComponentTransforms()
+	{
+		return ComponentTransforms[CurrIndex];
+	}
+
+	FORCEINLINE TArray<FTransform>& PrevComponentTransforms()
+	{
+		return ComponentTransforms[CurrIndex ^ 1];
+	}
 
+private:
 	FCriticalSection CriticalSection; 
 	
 	TWeakObjectPtr<USkeletalMeshComponent> MeshComp;
@@ -152,6 +161,9 @@ private:
 	/** Cached bone matrices. */
 	TArray<FMatrix> BoneRefToLocals[2];
 
+	/** Component space transforms */
+	TArray<FTransform> ComponentTransforms[2];
+
 	struct FLODData
 	{
 		FLODData() : PreSkinnedVertsUsers(0) { }
diff --git a/Engine/Plugins/FX/Niagara/Source/Niagara/Private/NiagaraDataInterfaceSkeletalMesh.cpp b/Engine/Plugins/FX/Niagara/Source/Niagara/Private/NiagaraDataInterfaceSkeletalMesh.cpp
--- a/Engine/Plugins/FX/Niagara/Source/Niagara/Private/NiagaraDataInterfaceSkeletalMesh.cpp
+++ b/Engine/Plugins/FX/Niagara/Source/Niagara/Private/NiagaraDataInterfaceSkeletalMesh.cpp
@@ -128,6 +128,15 @@ void FSkeletalMeshSkinningData::RegisterUser(FSkeletalMeshSkinningDataUsage Usag
 			PrevBoneRefToLocals() = CurrBoneRefToLocals();
 		}
 
+		if (CurrComponentTransforms().Num() == 0)
+		{
+			CurrComponentTransforms() = SkelComp->GetComponentSpaceTransforms();
+		}
+		if (PrevComponentTransforms().Num() != CurrComponentTransforms().Num())
+		{
+			PrevComponentTransforms() = CurrComponentTransforms();
+		}
+
 		if (Usage.NeedPreSkinnedVerts() && CurrSkinnedPositions(LODIndex).Num() == 0)
 		{
 			FSkeletalMeshLODRenderData& SkelMeshLODData = SkelComp->SkeletalMesh->GetResourceForRendering()->LODRenderData[LODIndex];
@@ -236,14 +245,21 @@ bool FSkeletalMeshSkinningData::Tick(float InDeltaSeconds, bool bRequirePreskin)
 		{
 			SkelComp->CacheRefToLocalMatrices(CurrBones);
 		}
+
+		CurrComponentTransforms() = SkelComp->GetComponentSpaceTransforms();
 	}
 
 	//Prime the prev matrices if they're missing.
-	if (PrevBoneRefToLocals().Num() != CurrBoneRefToLocals().Num() || bForceDataRefresh)
+	if ((PrevBoneRefToLocals().Num() != CurrBoneRefToLocals().Num()) || bForceDataRefresh)
 	{
 		PrevBoneRefToLocals() = CurrBoneRefToLocals();
 	}
 
+	if ((PrevComponentTransforms().Num() != CurrComponentTransforms().Num()) || bForceDataRefresh)
+	{
+		PrevComponentTransforms() = CurrComponentTransforms();
+	}
+
 	if (bRequirePreskin)
 	{
 		for (int32 LODIndex = 0; LODIndex < LODData.Num(); ++LODIndex)
@@ -2235,13 +2251,6 @@ void FSkeletalMeshAccessorHelper::Init<
 	SkinningData = InstData->SkinningData.SkinningData.Get();
 	Usage = InstData->SkinningData.Usage;
 
-	if (Comp)
-	{
-		const USkinnedMeshComponent* BaseComp = Comp->GetBaseComponent();
-		BoneComponentSpaceTransforms = &BaseComp->GetComponentSpaceTransforms();
-		PrevBoneComponentSpaceTransforms = &BaseComp->GetPreviousComponentTransformsArray();
-	}
-
 	const FSkeletalMeshSamplingInfo& SamplingInfo = InstData->Mesh->GetSamplingInfo();
 	SamplingRegion = &SamplingInfo.GetRegion(InstData->SamplingRegionIndices[0]);
 	SamplingRegionBuiltData = &SamplingInfo.GetRegionBuiltData(InstData->SamplingRegionIndices[0]);
@@ -2260,13 +2269,6 @@ void FSkeletalMeshAccessorHelper::Init<
 	SkinningData = InstData->SkinningData.SkinningData.Get();
 	Usage = InstData->SkinningData.Usage;
 
-	if (Comp)
-	{
-		const USkinnedMeshComponent* BaseComp = Comp->GetBaseComponent();
-		BoneComponentSpaceTransforms = &BaseComp->GetComponentSpaceTransforms();
-		PrevBoneComponentSpaceTransforms = &BaseComp->GetPreviousComponentTransformsArray();
-	}
-
 	const FSkeletalMeshSamplingInfo& SamplingInfo = InstData->Mesh->GetSamplingInfo();
 	SamplingRegion = &SamplingInfo.GetRegion(InstData->SamplingRegionIndices[0]);
 	SamplingRegionBuiltData = &SamplingInfo.GetRegionBuiltData(InstData->SamplingRegionIndices[0]);
