- Fix bool access to NiagaraParameterCollection

diff --git a/Engine/Plugins/FX/Niagara/Source/Niagara/Private/NiagaraParameterCollection.cpp b/Engine/Plugins/FX/Niagara/Source/Niagara/Private/NiagaraParameterCollection.cpp
--- a/Engine/Plugins/FX/Niagara/Source/Niagara/Private/NiagaraParameterCollection.cpp
+++ b/Engine/Plugins/FX/Niagara/Source/Niagara/Private/NiagaraParameterCollection.cpp
@@ -143,7 +143,7 @@ void UNiagaraParameterCollectionInstance::SetOverridesParameter(const FNiagaraVa
 //Blueprint Accessors
 bool UNiagaraParameterCollectionInstance::GetBoolParameter(const FString& InVariableName)
 {
-	return ParameterStorage.GetParameterValue<bool>(FNiagaraVariable(FNiagaraTypeDefinition::GetBoolDef(), *Collection->ParameterNameFromFriendlyName(InVariableName)));
+	return ParameterStorage.GetParameterValue<int32>(FNiagaraVariable(FNiagaraTypeDefinition::GetBoolDef(), *Collection->ParameterNameFromFriendlyName(InVariableName))) == FNiagaraBool::True;
 }
 
 float UNiagaraParameterCollectionInstance::GetFloatParameter(const FString& InVariableName)
@@ -184,7 +184,7 @@ FLinearColor UNiagaraParameterCollectionInstance::GetColorParameter(const FStrin
 
 void UNiagaraParameterCollectionInstance::SetBoolParameter(const FString& InVariableName, bool InValue)
 {
-	ParameterStorage.SetParameterValue(InValue, FNiagaraVariable(FNiagaraTypeDefinition::GetBoolDef(), *Collection->ParameterNameFromFriendlyName(InVariableName)));
+	ParameterStorage.SetParameterValue(InValue ? FNiagaraBool::True : FNiagaraBool::False, FNiagaraVariable(FNiagaraTypeDefinition::GetBoolDef(), *Collection->ParameterNameFromFriendlyName(InVariableName)));
 }
 
 void UNiagaraParameterCollectionInstance::SetFloatParameter(const FString& InVariableName, float InValue)
