- Mark Niagara System as invalid if we have no emitters
- This fixes a crash where the spawn compiled data layout does not match the update compiled data layout leading to memory read / write out of bounds

#rnx

diff --git a/Engine/Plugins/FX/Niagara/Source/Niagara/Private/NiagaraSystem.cpp b/Engine/Plugins/FX/Niagara/Source/Niagara/Private/NiagaraSystem.cpp
--- a/Engine/Plugins/FX/Niagara/Source/Niagara/Private/NiagaraSystem.cpp
+++ b/Engine/Plugins/FX/Niagara/Source/Niagara/Private/NiagaraSystem.cpp
@@ -675,6 +675,10 @@ bool UNiagaraSystem::IsValid()const
 		return false;
 	}
 
+	if (EmitterHandles.Num() == 0)
+	{
+		return false;
+	}
 
 	for (const FNiagaraEmitterHandle& Handle : EmitterHandles)
 	{
