Fixed a bug with integer type static switches that had exactly two input values

diff --git a/Engine/Plugins/FX/Niagara/Source/NiagaraEditor/Private/NiagaraNodeStaticSwitch.cpp b/Engine/Plugins/FX/Niagara/Source/NiagaraEditor/Private/NiagaraNodeStaticSwitch.cpp
--- a/Engine/Plugins/FX/Niagara/Source/NiagaraEditor/Private/NiagaraNodeStaticSwitch.cpp
+++ b/Engine/Plugins/FX/Niagara/Source/NiagaraEditor/Private/NiagaraNodeStaticSwitch.cpp
@@ -291,7 +291,7 @@ bool UNiagaraNodeStaticSwitch::GetVarIndex(FHlslNiagaraTranslator* Translator, i
 			{
 				Translator->Warning(FText::Format(LOCTEXT("InvalidStaticSwitchIntValue", "The supplied int value {0} is outside the bounds for the static switch."), FText::FromString(FString::FromInt(Value))), this, nullptr);
 			}
-			VarIndexOut = FMath::Clamp(Value, 0, MaxValue) * (InputPinCount / MaxValue);
+			VarIndexOut = FMath::Clamp(Value, 0, MaxValue) * (InputPinCount / (MaxValue + 1));
 			Success = true;
 		}
 		else if (Translator)
