Duplicate Edit section removed from context menu when right-clicking in the Niagara system timeline track list. Create Asset from This and Rename Emitter now also exist in the system timeline track list.

diff --git a/Engine/Plugins/FX/Niagara/Source/NiagaraEditor/Private/NiagaraEditorUtilities.cpp b/Engine/Plugins/FX/Niagara/Source/NiagaraEditor/Private/NiagaraEditorUtilities.cpp
--- a/Engine/Plugins/FX/Niagara/Source/NiagaraEditor/Private/NiagaraEditorUtilities.cpp
+++ b/Engine/Plugins/FX/Niagara/Source/NiagaraEditor/Private/NiagaraEditorUtilities.cpp
@@ -1833,15 +1833,6 @@ bool FNiagaraEditorUtilities::AddEmitterContextMenuActions(FMenuBuilder& MenuBui
 		}
 		MenuBuilder.EndSection();
 
-		MenuBuilder.BeginSection("EmitterEditSection", LOCTEXT("Edit", "Edit"));
-
-		MenuBuilder.AddMenuEntry(FGenericCommands::Get().Cut);
-		MenuBuilder.AddMenuEntry(FGenericCommands::Get().Copy);
-		MenuBuilder.AddMenuEntry(FGenericCommands::Get().Delete);
-		MenuBuilder.AddMenuEntry(FGenericCommands::Get().Rename);
-
-		MenuBuilder.EndSection();
-
 		return true;
 	}
 
diff --git a/Engine/Plugins/FX/Niagara/Source/NiagaraEditor/Private/NiagaraOverviewNode.cpp b/Engine/Plugins/FX/Niagara/Source/NiagaraEditor/Private/NiagaraOverviewNode.cpp
--- a/Engine/Plugins/FX/Niagara/Source/NiagaraEditor/Private/NiagaraOverviewNode.cpp
+++ b/Engine/Plugins/FX/Niagara/Source/NiagaraEditor/Private/NiagaraOverviewNode.cpp
@@ -201,6 +201,16 @@ void UNiagaraOverviewNode::GetNodeContextMenuActions(class UToolMenu* Menu, clas
 						),
 						EUserInterfaceActionType::ToggleButton
 					);
+
+					Section.AddMenuEntry(
+						"RenameEmitter",
+						LOCTEXT("RenameEmitter", "Rename Emitter"),
+						LOCTEXT("RenameEmitterToolTip", "Rename this local emitter copy."),
+						FSlateIcon(),
+						FUIAction(
+							FExecuteAction::CreateSP(EmitterHandleViewModel, &FNiagaraEmitterHandleViewModel::SetIsRenamePending, true)
+						)
+					);
 				}
 
 				Section.AddMenuEntry(
@@ -234,6 +244,22 @@ void UNiagaraOverviewNode::GetNodeContextMenuActions(class UToolMenu* Menu, clas
 						)
 					)
 				);
+
+				Section.AddMenuEntry(
+					"CreateAssetFromThis",
+					LOCTEXT("CreateAssetFromThisEmitter", "Create Asset From This"),
+					LOCTEXT("CreateAssetFromThisEmitterToolTip", "Create an emitter asset from this emitter."),
+					FSlateIcon(),
+					FUIAction(
+						FExecuteAction::CreateStatic(&FNiagaraEditorUtilities::CreateAssetFromEmitter, EmitterHandleViewModel),
+						FCanExecuteAction::CreateLambda(
+							[bSingleSelection, EmitterHandleViewModel]()
+							{
+								return bSingleSelection && EmitterHandleViewModel->GetOwningSystemEditMode() == ENiagaraSystemViewModelEditMode::SystemAsset;
+							}
+						)
+					)
+				);
 			}
 		}
 	}
diff --git a/Engine/Plugins/FX/Niagara/Source/NiagaraEditorWidgets/Private/SNiagaraStack.cpp b/Engine/Plugins/FX/Niagara/Source/NiagaraEditorWidgets/Private/SNiagaraStack.cpp
--- a/Engine/Plugins/FX/Niagara/Source/NiagaraEditorWidgets/Private/SNiagaraStack.cpp
+++ b/Engine/Plugins/FX/Niagara/Source/NiagaraEditorWidgets/Private/SNiagaraStack.cpp
@@ -24,6 +24,7 @@
 #include "ViewModels/Stack/NiagaraStackModuleItem.h"
 #include "ViewModels/Stack/NiagaraStackParameterStoreEntry.h"
 #include "Framework/Application/SlateApplication.h"
+#include "Framework/Commands/GenericCommands.h"
 #include "IContentBrowserSingleton.h"
 #include "ContentBrowserModule.h"
 #include "Widgets/Input/SButton.h"
@@ -726,6 +727,15 @@ FReply SNiagaraStack::OnTopLevelRowMouseButtonDown(const FGeometry&, const FPoin
 
 		FNiagaraEditorUtilities::AddEmitterContextMenuActions(MenuBuilder, TopLevelViewModel->EmitterHandleViewModel);
 
+		MenuBuilder.BeginSection("EmitterEditSection", LOCTEXT("Edit", "Edit"));
+
+		MenuBuilder.AddMenuEntry(FGenericCommands::Get().Cut);
+		MenuBuilder.AddMenuEntry(FGenericCommands::Get().Copy);
+		MenuBuilder.AddMenuEntry(FGenericCommands::Get().Delete);
+		MenuBuilder.AddMenuEntry(FGenericCommands::Get().Rename);
+
+		MenuBuilder.EndSection();
+
 		MenuBuilder.BeginSection("StackActions", LOCTEXT("StackActions", "Stack Actions"));
 		{
 			if (StackViewModel->HasDismissedStackIssues())
