Fix missing transitions on ID buffers.

diff --git a/Engine/Plugins/FX/Niagara/Source/Niagara/Private/NiagaraDataSet.cpp b/Engine/Plugins/FX/Niagara/Source/Niagara/Private/NiagaraDataSet.cpp
--- a/Engine/Plugins/FX/Niagara/Source/Niagara/Private/NiagaraDataSet.cpp
+++ b/Engine/Plugins/FX/Niagara/Source/Niagara/Private/NiagaraDataSet.cpp
@@ -375,21 +375,14 @@ void FNiagaraDataSet::AllocateGPUFreeIDs(uint32 InNumInstances, FRHICommandList&
 	NewFreeIDsBuffer.Initialize(sizeof(int32), NumIDsToAlloc, EPixelFormat::PF_R32_SINT, BUF_Static, DebugBufferName);
 
 	// We must maintain the existing list of free IDs.
-	FRHIShaderResourceView* ExistingBuffer;
-	if (GPUNumAllocatedIDs > 0)
-	{
-		// The free IDs buffer was written in the previous simulation step, but hasn't been transitioned to read yet, so we must
-		// transition it explicitly here. The new buffer will be transitioned by NiagaraEmitterInstanceBatcher::DispatchAllOnCompute(),
-		// so there's no need for a barrier at the end of this function.
-		RHICmdList.TransitionResource(EResourceTransitionAccess::EReadable, EResourceTransitionPipeline::EComputeToCompute, GPUFreeIDs.UAV);
-		ExistingBuffer = GPUFreeIDs.SRV;
-	}
-	else
-	{
-		ExistingBuffer = FNiagaraRenderer::GetDummyIntBuffer().SRV;
-	}
-	
-	NiagaraInitGPUFreeIDList(RHICmdList, FeatureLevel, NumIDsToAlloc, NewFreeIDsBuffer, GPUNumAllocatedIDs, ExistingBuffer);
+	FRWBuffer& ExistingBuffer = (GPUNumAllocatedIDs > 0) ? GPUFreeIDs : FNiagaraRenderer::GetDummyIntBuffer();
+
+	// The free IDs buffer was written in the previous simulation step, but hasn't been transitioned to read yet, so we must
+	// transition it explicitly here. The new buffer will be transitioned by NiagaraEmitterInstanceBatcher::DispatchAllOnCompute(),
+	// so there's no need for a barrier at the end of this function.
+	RHICmdList.TransitionResource(EResourceTransitionAccess::EReadable, EResourceTransitionPipeline::EComputeToCompute, ExistingBuffer.UAV);
+	RHICmdList.TransitionResource(EResourceTransitionAccess::EWritable, EResourceTransitionPipeline::EComputeToCompute, NewFreeIDsBuffer.UAV);
+	NiagaraInitGPUFreeIDList(RHICmdList, FeatureLevel, NumIDsToAlloc, NewFreeIDsBuffer, GPUNumAllocatedIDs, ExistingBuffer.SRV);
 
 	GPUFreeIDs = MoveTemp(NewFreeIDsBuffer);
 	GPUNumAllocatedIDs = NumIDsToAlloc;
diff --git a/Engine/Plugins/FX/Niagara/Source/Niagara/Private/NiagaraEmitterInstanceBatcher.cpp b/Engine/Plugins/FX/Niagara/Source/Niagara/Private/NiagaraEmitterInstanceBatcher.cpp
--- a/Engine/Plugins/FX/Niagara/Source/Niagara/Private/NiagaraEmitterInstanceBatcher.cpp
+++ b/Engine/Plugins/FX/Niagara/Source/Niagara/Private/NiagaraEmitterInstanceBatcher.cpp
@@ -629,7 +629,7 @@ void NiagaraEmitterInstanceBatcher::UpdateFreeIDBuffers(FRHICommandList& RHICmdL
 		WriteBuffers.Add(Instance->Context->MainDataSet->GetGPUFreeIDs().UAV);
 	}
 
-	RHICmdList.TransitionResource(EResourceTransitionAccess::ERWBarrier, EResourceTransitionPipeline::EComputeToCompute, FreeIDListSizesBuffer.UAV);
+	RHICmdList.TransitionResource(EResourceTransitionAccess::ERWNoBarrier, EResourceTransitionPipeline::EComputeToCompute, FreeIDListSizesBuffer.UAV);
 	RHICmdList.TransitionResources(EResourceTransitionAccess::EReadable, EResourceTransitionPipeline::EComputeToCompute, ReadBuffers.GetData(), ReadBuffers.Num());
 	RHICmdList.TransitionResources(EResourceTransitionAccess::EWritable, EResourceTransitionPipeline::EComputeToCompute, WriteBuffers.GetData(), WriteBuffers.Num());
 
