- Integrate 11188683: Use WeakPtrs to UObjects for callback handler and Niagara system

#rnx

diff --git a/Engine/Plugins/FX/Niagara/Source/Niagara/Classes/NiagaraDataInterfaceExport.h b/Engine/Plugins/FX/Niagara/Source/Niagara/Classes/NiagaraDataInterfaceExport.h
--- a/Engine/Plugins/FX/Niagara/Source/Niagara/Classes/NiagaraDataInterfaceExport.h
+++ b/Engine/Plugins/FX/Niagara/Source/Niagara/Classes/NiagaraDataInterfaceExport.h
@@ -40,7 +40,7 @@ public:
 
 struct ExportInterface_InstanceData
 {
-	UObject* CallbackHandler = nullptr;
+	TWeakObjectPtr<UObject> CallbackHandler;
 
 	/** We use a lock-free queue here because multiple threads might try to push data to it at the same time. */
 	TQueue<FBasicParticleData, EQueueMode::Mpsc> GatheredData;
diff --git a/Engine/Plugins/FX/Niagara/Source/Niagara/Private/NiagaraDataInterfaceExport.cpp b/Engine/Plugins/FX/Niagara/Source/Niagara/Private/NiagaraDataInterfaceExport.cpp
--- a/Engine/Plugins/FX/Niagara/Source/Niagara/Private/NiagaraDataInterfaceExport.cpp
+++ b/Engine/Plugins/FX/Niagara/Source/Niagara/Private/NiagaraDataInterfaceExport.cpp
@@ -13,15 +13,15 @@ Async task to call the blueprint callback on the game thread and isolate the nia
 */
 class FNiagaraExportCallbackAsyncTask
 {
-	UObject* CallbackHandler;
+	TWeakObjectPtr<UObject> WeakCallbackHandler;
 	TArray<FBasicParticleData> Data;
-	UNiagaraSystem* System;
+	TWeakObjectPtr<UNiagaraSystem> WeakSystem;
 
 public:
-	FNiagaraExportCallbackAsyncTask(UObject* CallbackHandler, TArray<FBasicParticleData>& Data, UNiagaraSystem* System)
-		: CallbackHandler(CallbackHandler)
+	FNiagaraExportCallbackAsyncTask(TWeakObjectPtr<UObject> InCallbackHandler, TArray<FBasicParticleData>& Data, TWeakObjectPtr<UNiagaraSystem> InSystem)
+		: WeakCallbackHandler(InCallbackHandler)
 		, Data(Data)
-		, System(System)
+		, WeakSystem(InSystem)
 	{
 	}
 
@@ -31,11 +31,20 @@ public:
 
 	void DoTask(ENamedThreads::Type CurrentThread, const FGraphEventRef& MyCompletionGraphEvent)
 	{
-		if (!CallbackHandler->IsValidLowLevelFast())
+		UNiagaraSystem* System = WeakSystem.Get();
+		if (System == nullptr)
 		{
-			UE_LOG(LogNiagara, Warning, TEXT("Received invalid callback handle in data interface - was the callback object deleted? System %s"), *System->GetName());
+			UE_LOG(LogNiagara, Warning, TEXT("Invalid system handle in export data interface callback, skipping"));
 			return;
 		}
+
+		UObject* CallbackHandler = WeakCallbackHandler.Get();
+		if (CallbackHandler == nullptr)
+		{
+			UE_LOG(LogNiagara, Warning, TEXT("Invalid CallbackHandler in export data interface callback, skipping"));
+			return;
+		}
+
 		INiagaraParticleCallbackHandler::Execute_ReceiveParticleData(CallbackHandler, Data, System);
 	}
 };
@@ -75,7 +84,7 @@ bool UNiagaraDataInterfaceExport::PerInstanceTick(void* PerInstanceData, FNiagar
 	}
 	else
 	{
-		PIData->CallbackHandler = nullptr;
+		PIData->CallbackHandler.Reset();
 	}
 	return false;
 }
@@ -83,17 +92,21 @@ bool UNiagaraDataInterfaceExport::PerInstanceTick(void* PerInstanceData, FNiagar
 bool UNiagaraDataInterfaceExport::PerInstanceTickPostSimulate(void* PerInstanceData, FNiagaraSystemInstance* SystemInstance, float DeltaSeconds)
 {
 	ExportInterface_InstanceData* PIData = (ExportInterface_InstanceData*) PerInstanceData;
-	if (!PIData->GatheredData.IsEmpty() && PIData->CallbackHandler && PIData->CallbackHandler->GetClass()->ImplementsInterface(UNiagaraParticleCallbackHandler::StaticClass()))
+	if ( !PIData->GatheredData.IsEmpty() )
 	{
-		//Drain the queue into an array here
-		TArray<FBasicParticleData> Data;
-		FBasicParticleData Value;
-		while (PIData->GatheredData.Dequeue(Value))
+		UObject* CallbackHandler = PIData->CallbackHandler.Get();
+		if ( CallbackHandler && CallbackHandler->GetClass()->ImplementsInterface(UNiagaraParticleCallbackHandler::StaticClass()) )
 		{
-			Data.Add(Value);
+			//Drain the queue into an array here
+			TArray<FBasicParticleData> Data;
+			FBasicParticleData Value;
+			while (PIData->GatheredData.Dequeue(Value))
+			{
+				Data.Add(Value);
+			}
+
+			TGraphTask<FNiagaraExportCallbackAsyncTask>::CreateTask().ConstructAndDispatchWhenReady(PIData->CallbackHandler, Data, SystemInstance->GetSystem());
 		}
-
-		TGraphTask<FNiagaraExportCallbackAsyncTask>::CreateTask().ConstructAndDispatchWhenReady(PIData->CallbackHandler, Data, SystemInstance->GetSystem());
 	}
 	return false;
 }
@@ -163,7 +176,7 @@ void UNiagaraDataInterfaceExport::StoreData(FVectorVMContext& Context)
 	VectorVM::FExternalFuncRegisterHandler<FNiagaraBool> OutSample(Context);
 
 	checkfSlow(InstData.Get(), TEXT("Export data interface has invalid instance data. %s"), *GetPathName());
-	bool ValidHandlerData = InstData->UserParamBinding.BoundVariable.IsValid() && InstData->CallbackHandler;
+	bool ValidHandlerData = InstData->UserParamBinding.BoundVariable.IsValid() && InstData->CallbackHandler.IsValid();
 
 	for (int32 i = 0; i < Context.NumInstances; ++i)
 	{
