Allow static switch parameters in the stack view to be reset to their default value

diff --git a/Engine/Plugins/FX/Niagara/Source/NiagaraEditor/Private/ViewModels/Stack/NiagaraStackFunctionInput.cpp b/Engine/Plugins/FX/Niagara/Source/NiagaraEditor/Private/ViewModels/Stack/NiagaraStackFunctionInput.cpp
--- a/Engine/Plugins/FX/Niagara/Source/NiagaraEditor/Private/ViewModels/Stack/NiagaraStackFunctionInput.cpp
+++ b/Engine/Plugins/FX/Niagara/Source/NiagaraEditor/Private/ViewModels/Stack/NiagaraStackFunctionInput.cpp
@@ -1265,11 +1265,6 @@ void UNiagaraStackFunctionInput::SetLocalValue(TSharedRef<FStructOnScope> InLoca
 
 bool UNiagaraStackFunctionInput::CanReset() const
 {
-	if (IsStaticParameter())
-	{
-		// Static switch parameters only hold a single value in the default pin, so we disable resetting to prevent a special implementation for them
-		return false;
-	}
 	if (bCanResetCache.IsSet() == false)
 	{
 		bool bNewCanReset;
@@ -1288,8 +1283,37 @@ bool UNiagaraStackFunctionInput::CanReset() const
 		{
 			UEdGraphPin* DefaultPin = GetDefaultPin();
 			if (ensure(DefaultPin != nullptr))
-			{			
-				if(DefaultPin->LinkedTo.Num() == 0)
+			{
+				if (IsStaticParameter())
+				{
+					bNewCanReset = true;
+
+					// Static switch parameters only hold a single value in the default pin, so we compare it to the default value in the switch metadata
+					UNiagaraGraph* FunctionGraph = CastChecked<UNiagaraScriptSource>(OwningFunctionCallNode->FunctionScript->GetSource())->NodeGraph;
+					const UEdGraphSchema_Niagara* Schema = CastChecked<UEdGraphSchema_Niagara>(OwningFunctionCallNode->GetSchema());
+
+					for (const FNiagaraVariable& SwitchVar : FunctionGraph->FindStaticSwitchInputs())
+					{
+						FEdGraphPinType PinType = Schema->TypeDefinitionToPinType(SwitchVar.GetType());
+						if (DefaultPin->PinName.IsEqual(SwitchVar.GetName()) && DefaultPin->PinType == PinType)
+						{
+							TOptional<FNiagaraVariableMetaData> MetaData = FunctionGraph->GetMetaData(SwitchVar);
+
+							int32 DefaultValue = MetaData.IsSet() ? MetaData->StaticSwitchDefaultValue  : 0;
+							int32 CurrentPinValue = 0;
+							if (FNiagaraEditorUtilities::ResolveConstantValue(DefaultPin, CurrentPinValue))
+							{
+								bNewCanReset = DefaultValue != CurrentPinValue;
+							}
+							else
+							{
+								bNewCanReset = false;
+							}
+							break;
+						}
+					}
+				}
+				else if (DefaultPin->LinkedTo.Num() == 0)
 				{
 					if (UEdGraphPin* OverridePin = GetOverridePin())
 					{
@@ -1416,6 +1440,34 @@ void UNiagaraStackFunctionInput::Reset()
 			FNiagaraStackGraphUtilities::RelayoutGraph(*OwningFunctionCallNode->GetGraph());
 		}
 	}
+	else if (IsStaticParameter())
+	{
+		// Static switch parameters only hold a single value in the default pin, so we directly set it to the default value from the switch metadata
+		UNiagaraGraph* FunctionGraph = CastChecked<UNiagaraScriptSource>(OwningFunctionCallNode->FunctionScript->GetSource())->NodeGraph;
+		const UEdGraphSchema_Niagara* Schema = CastChecked<UEdGraphSchema_Niagara>(OwningFunctionCallNode->GetSchema());
+		UEdGraphPin* DefaultPin = GetDefaultPin();
+
+		for (FNiagaraVariable SwitchVar : FunctionGraph->FindStaticSwitchInputs())
+		{
+			FEdGraphPinType PinType = Schema->TypeDefinitionToPinType(SwitchVar.GetType());
+			if (DefaultPin->PinName.IsEqual(SwitchVar.GetName()) && DefaultPin->PinType == PinType)
+			{
+				TOptional<FNiagaraVariableMetaData> MetaData = FunctionGraph->GetMetaData(SwitchVar);
+
+				int32 DefaultValue = MetaData.IsSet() ? MetaData->StaticSwitchDefaultValue : 0;
+				SwitchVar.SetValue(DefaultValue);
+				FString PinDefaultValue;
+				if (ensureMsgf(Schema->TryGetPinDefaultValueFromNiagaraVariable(SwitchVar, PinDefaultValue), TEXT("Could not generate value string to reset static switch parameter.")))
+				{
+					DefaultPin->DefaultValue = PinDefaultValue;
+				}
+
+				OwningFunctionCallNode->GetNiagaraGraph()->NotifyGraphNeedsRecompile();
+				FNiagaraStackGraphUtilities::RelayoutGraph(*OwningFunctionCallNode->GetGraph());
+				break;
+			}
+		}
+	}
 	else
 	{
 		// For all other value modes removing the nodes connected to the override pin resets them.
