- Do not return true unless the function is implemented

diff --git a/Engine/Plugins/FX/Niagara/Source/Niagara/Private/NiagaraDataInterfaceCollisionQuery.cpp b/Engine/Plugins/FX/Niagara/Source/Niagara/Private/NiagaraDataInterfaceCollisionQuery.cpp
--- a/Engine/Plugins/FX/Niagara/Source/Niagara/Private/NiagaraDataInterfaceCollisionQuery.cpp
+++ b/Engine/Plugins/FX/Niagara/Source/Niagara/Private/NiagaraDataInterfaceCollisionQuery.cpp
@@ -239,6 +239,7 @@ bool UNiagaraDataInterfaceCollisionQuery::GetFunctionHLSL(const FNiagaraDataInte
 				}\n\
 			}\n\
 		}\n}\n\n");
+		return true;
 	}
 	else if (FunctionInfo.DefinitionName == TEXT("PerformCollisionQueryGPUShader"))
 	{
@@ -304,6 +305,7 @@ bool UNiagaraDataInterfaceCollisionQuery::GetFunctionHLSL(const FNiagaraDataInte
 			{\n\
 				") + SceneDepthFunction + TEXT("(In_SamplePos, In_TraceEndPos, CollisionDepthBounds, ParticleRadius, OutCollisionValid, Out_CollisionPos, Out_CollisionNormal);\n\
 			}\n}\n\n");
+		return true;
 	}
 	else if (FunctionInfo.DefinitionName == TEXT("QuerySceneDepthGPU"))
 	{
@@ -333,6 +335,7 @@ bool UNiagaraDataInterfaceCollisionQuery::GetFunctionHLSL(const FNiagaraDataInte
 			{\n\
 				Out_IsInsideView = false;\n\
 			}\n}\n\n");
+		return true;
 	}
 	else if (FunctionInfo.DefinitionName == TEXT("QueryMeshDistanceFieldGPU"))
 	{
@@ -341,9 +344,10 @@ bool UNiagaraDataInterfaceCollisionQuery::GetFunctionHLSL(const FNiagaraDataInte
 			Out_DistanceToNearestSurface = GetDistanceToNearestSurfaceGlobal(In_SamplePos);\n\
 			Out_FieldGradient = GetDistanceFieldGradientGlobal(In_SamplePos);\
 			\n}\n\n");
+		return true;
 	}
 
-	return true;
+	return false;
 }
 
 void UNiagaraDataInterfaceCollisionQuery::GetParameterDefinitionHLSL(const FNiagaraDataInterfaceGPUParamInfo& ParamInfo, FString& OutHLSL)
