- Add console variable "fx.MaxNiagaraGPUParticlesSpawnPerFrame" to control per frame spawn capacity on the GPU

diff --git a/Engine/Plugins/FX/Niagara/Source/Niagara/Private/NiagaraEmitterInstance.cpp b/Engine/Plugins/FX/Niagara/Source/Niagara/Private/NiagaraEmitterInstance.cpp
--- a/Engine/Plugins/FX/Niagara/Source/Niagara/Private/NiagaraEmitterInstance.cpp
+++ b/Engine/Plugins/FX/Niagara/Source/Niagara/Private/NiagaraEmitterInstance.cpp
@@ -74,6 +74,14 @@ static FAutoConsoleVariableRef CVarMaxNiagaraCPUParticlesPerEmitter(
 	TEXT("The max number of supported CPU particles per emitter in Niagara. \n"),
 	ECVF_Default
 );
+
+static int32 GMaxNiagaraGPUParticlesSpawnPerFrame = 2000000;
+static FAutoConsoleVariableRef CVarMaxNiagaraGPUParticlesSpawnPerFrame(
+	TEXT("fx.MaxNiagaraGPUParticlesSpawnPerFrame"),
+	GMaxNiagaraGPUParticlesSpawnPerFrame,
+	TEXT("The max number of GPU particles we expect to spawn in a single frame.\n"),
+	ECVF_Default
+);
 //////////////////////////////////////////////////////////////////////////
 
 template<bool bAccumulate>
@@ -1164,6 +1172,12 @@ void FNiagaraEmitterInstance::Tick(float DeltaSeconds)
 				{
 					if (Info.Count > 0 && (NumSpawnInfos < NIAGARA_MAX_GPU_SPAWN_INFOS))
 					{
+						if ((TotalSpawnedParticles + Info.Count) > GMaxNiagaraGPUParticlesSpawnPerFrame)
+						{
+							UE_LOG(LogNiagara, Warning, TEXT("%s has attempted to execeed max GPU per frame spawn! | Max: %d | Requested: %d"), *CachedEmitter->GetUniqueEmitterName(), GMaxNiagaraGPUParticlesSpawnPerFrame, TotalSpawnedParticles + Info.Count);
+							break;
+						}
+
 						GpuSpawnInfo.SpawnInfoParams[NumSpawnInfos].X = Info.IntervalDt;
 						GpuSpawnInfo.SpawnInfoParams[NumSpawnInfos].Y = Info.InterpStartDt;
 						reinterpret_cast<int32&>(GpuSpawnInfo.SpawnInfoParams[NumSpawnInfos].Z) = Info.SpawnGroup;
@@ -1176,7 +1190,7 @@ void FNiagaraEmitterInstance::Tick(float DeltaSeconds)
 					}
 					else if (Info.Count > 0)
 					{
-						UE_LOG(LogNiagara, Log, TEXT("Exceeded Gpu spawn info count, see NIAGARA_MAX_GPU_SPAWN_INFOS for more information!"));
+						UE_LOG(LogNiagara, Warning, TEXT("%s Exceeded Gpu spawn info count, see NIAGARA_MAX_GPU_SPAWN_INFOS for more information!"), *CachedEmitter->GetUniqueEmitterName());
 						break;
 					}
 
@@ -1245,7 +1259,8 @@ void FNiagaraEmitterInstance::Tick(float DeltaSeconds)
 	//TODO: These current limits can be improved relatively easily. Though perf in at these counts will obviously be an issue anyway.
 	if (CachedEmitter->SimTarget == ENiagaraSimTarget::CPUSim && AllocationSize > GMaxNiagaraCPUParticlesPerEmitter)
 	{
-		UE_LOG(LogNiagara, Warning, TEXT("Emitter %s has attempted to exceed the max CPU particle count! | Max: %d | Requested: %u"), *CachedEmitter->GetUniqueEmitterName(), GMaxNiagaraCPUParticlesPerEmitter, AllocationSize);
+		UE_LOG(LogNiagara, Warning, TEXT("%s has attempted to exceed the max CPU particle count! | Max: %d | Requested: %u"), *CachedEmitter->GetFullName(), GMaxNiagaraCPUParticlesPerEmitter, AllocationSize);
+
 		//For now we completely bail out of spawning new particles. Possibly should improve this in future.
 		AllocationSize = OrigNumParticles;
 		SpawnTotal = 0;
