Niagara - Move some header files for view models and widgets into the public folder so that they can be accessed by other modules in the plugin.  Also fix up all of the referencing files.

diff --git a/Engine/Plugins/FX/Niagara/Source/NiagaraEditor/Private/Customizations/NiagaraParameterCollectionCustomNodeBuilder.h b/Engine/Plugins/FX/Niagara/Source/NiagaraEditor/Private/Customizations/NiagaraParameterCollectionCustomNodeBuilder.h
--- a/Engine/Plugins/FX/Niagara/Source/NiagaraEditor/Private/Customizations/NiagaraParameterCollectionCustomNodeBuilder.h
+++ b/Engine/Plugins/FX/Niagara/Source/NiagaraEditor/Private/Customizations/NiagaraParameterCollectionCustomNodeBuilder.h
@@ -9,7 +9,7 @@
 #include "NiagaraParameterCollectionViewModel.h"
 #include "NiagaraScriptInputCollectionViewModel.h"
 #include "NiagaraScriptOutputCollectionViewModel.h"
-#include "NiagaraScriptViewModel.h"
+#include "ViewModels/NiagaraScriptViewModel.h"
 #include "NiagaraParameterViewModel.h"
 #include "NiagaraEditorStyle.h"
 #include "IDetailChildrenBuilder.h"
diff --git a/Engine/Plugins/FX/Niagara/Source/NiagaraEditor/Private/Customizations/NiagaraScriptDetails.cpp b/Engine/Plugins/FX/Niagara/Source/NiagaraEditor/Private/Customizations/NiagaraScriptDetails.cpp
--- a/Engine/Plugins/FX/Niagara/Source/NiagaraEditor/Private/Customizations/NiagaraScriptDetails.cpp
+++ b/Engine/Plugins/FX/Niagara/Source/NiagaraEditor/Private/Customizations/NiagaraScriptDetails.cpp
@@ -9,7 +9,7 @@
 #include "NiagaraParameterCollectionViewModel.h"
 #include "NiagaraScriptInputCollectionViewModel.h"
 #include "NiagaraScriptOutputCollectionViewModel.h"
-#include "NiagaraScriptViewModel.h"
+#include "ViewModels/NiagaraScriptViewModel.h"
 #include "NiagaraScriptGraphViewModel.h"
 #include "NiagaraParameterViewModel.h"
 #include "NiagaraEditorStyle.h"
@@ -28,7 +28,7 @@
 #include "NiagaraScript.h"
 #include "NiagaraParameterCollectionCustomNodeBuilder.h"
 #include "NiagaraMetaDataCustomNodeBuilder.h"
-#include "SItemSelector.h"
+#include "Widgets/SItemSelector.h"
 #include "PropertyCustomizationHelpers.h"
 #include "ScopedTransaction.h"
 
diff --git a/Engine/Plugins/FX/Niagara/Source/NiagaraEditor/Private/NiagaraEditorModule.cpp b/Engine/Plugins/FX/Niagara/Source/NiagaraEditor/Private/NiagaraEditorModule.cpp
--- a/Engine/Plugins/FX/Niagara/Source/NiagaraEditor/Private/NiagaraEditorModule.cpp
+++ b/Engine/Plugins/FX/Niagara/Source/NiagaraEditor/Private/NiagaraEditorModule.cpp
@@ -60,7 +60,7 @@
 #include "NiagaraDataInterfaceVectorCurve.h"
 #include "NiagaraDataInterfaceVector4Curve.h"
 #include "NiagaraDataInterfaceColorCurve.h"
-#include "NiagaraScriptViewModel.h"
+#include "ViewModels/NiagaraScriptViewModel.h"
 #include "ViewModels/NiagaraSystemViewModel.h"
 #include "ViewModels/NiagaraEmitterViewModel.h"
 #include "TNiagaraGraphPinEditableName.h"
diff --git a/Engine/Plugins/FX/Niagara/Source/NiagaraEditor/Private/SItemSelector.h b/Engine/Plugins/FX/Niagara/Source/NiagaraEditor/Private/SItemSelector.h
deleted file mode 100644
diff --git a/Engine/Plugins/FX/Niagara/Source/NiagaraEditor/Private/SNiagaraAssetPickerList.h b/Engine/Plugins/FX/Niagara/Source/NiagaraEditor/Private/SNiagaraAssetPickerList.h
--- a/Engine/Plugins/FX/Niagara/Source/NiagaraEditor/Private/SNiagaraAssetPickerList.h
+++ b/Engine/Plugins/FX/Niagara/Source/NiagaraEditor/Private/SNiagaraAssetPickerList.h
@@ -2,7 +2,7 @@
 
 #pragma once
 
-#include "SItemSelector.h"
+#include "Widgets/SItemSelector.h"
 #include "AssetData.h"
 
 typedef SItemSelector<FText, FAssetData> SNiagaraAssetItemSelector;
diff --git a/Engine/Plugins/FX/Niagara/Source/NiagaraEditor/Private/Toolkits/NiagaraParameterCollectionToolkit.cpp b/Engine/Plugins/FX/Niagara/Source/NiagaraEditor/Private/Toolkits/NiagaraParameterCollectionToolkit.cpp
--- a/Engine/Plugins/FX/Niagara/Source/NiagaraEditor/Private/Toolkits/NiagaraParameterCollectionToolkit.cpp
+++ b/Engine/Plugins/FX/Niagara/Source/NiagaraEditor/Private/Toolkits/NiagaraParameterCollectionToolkit.cpp
@@ -5,7 +5,7 @@
 #include "SNiagaraParameterCollection.h"
 
 #include "NiagaraObjectSelection.h"
-#include "SNiagaraSelectedObjectsDetails.h"
+#include "Widgets/SNiagaraSelectedObjectsDetails.h"
 #include "NiagaraEditorCommands.h"
 #include "NiagaraEditorStyle.h"
 #include "ScopedTransaction.h"
diff --git a/Engine/Plugins/FX/Niagara/Source/NiagaraEditor/Private/Toolkits/NiagaraScriptToolkit.cpp b/Engine/Plugins/FX/Niagara/Source/NiagaraEditor/Private/Toolkits/NiagaraScriptToolkit.cpp
--- a/Engine/Plugins/FX/Niagara/Source/NiagaraEditor/Private/Toolkits/NiagaraScriptToolkit.cpp
+++ b/Engine/Plugins/FX/Niagara/Source/NiagaraEditor/Private/Toolkits/NiagaraScriptToolkit.cpp
@@ -4,14 +4,14 @@
 #include "NiagaraEditorModule.h"
 #include "NiagaraEditorUtilities.h"
 #include "NiagaraScript.h"
-#include "NiagaraScriptViewModel.h"
+#include "ViewModels/NiagaraScriptViewModel.h"
 #include "NiagaraScriptGraphViewModel.h"
 #include "NiagaraObjectSelection.h"
 #include "NiagaraEditorCommands.h"
 #include "NiagaraScriptInputCollectionViewModel.h"
 #include "NiagaraScriptOutputCollectionViewModel.h"
-#include "SNiagaraScriptGraph.h"
-#include "SNiagaraSelectedObjectsDetails.h"
+#include "Widgets/SNiagaraScriptGraph.h"
+#include "Widgets/SNiagaraSelectedObjectsDetails.h"
 #include "SNiagaraParameterCollection.h"
 #include "UObject/Package.h"
 #include "SNiagaraParameterMapView.h"
diff --git a/Engine/Plugins/FX/Niagara/Source/NiagaraEditor/Private/ViewModels/NiagaraEmitterHandleViewModel.cpp b/Engine/Plugins/FX/Niagara/Source/NiagaraEditor/Private/ViewModels/NiagaraEmitterHandleViewModel.cpp
--- a/Engine/Plugins/FX/Niagara/Source/NiagaraEditor/Private/ViewModels/NiagaraEmitterHandleViewModel.cpp
+++ b/Engine/Plugins/FX/Niagara/Source/NiagaraEditor/Private/ViewModels/NiagaraEmitterHandleViewModel.cpp
@@ -6,7 +6,7 @@
 #include "ViewModels/NiagaraSystemViewModel.h"
 #include "ViewModels/NiagaraEmitterViewModel.h"
 #include "ViewModels/Stack/NiagaraStackViewModel.h"
-#include "NiagaraScriptViewModel.h"
+#include "ViewModels/NiagaraScriptViewModel.h"
 #include "NiagaraScriptGraphViewModel.h"
 #include "NiagaraObjectSelection.h"
 #include "NiagaraScriptSource.h"
diff --git a/Engine/Plugins/FX/Niagara/Source/NiagaraEditor/Private/ViewModels/NiagaraEmitterViewModel.cpp b/Engine/Plugins/FX/Niagara/Source/NiagaraEditor/Private/ViewModels/NiagaraEmitterViewModel.cpp
--- a/Engine/Plugins/FX/Niagara/Source/NiagaraEditor/Private/ViewModels/NiagaraEmitterViewModel.cpp
+++ b/Engine/Plugins/FX/Niagara/Source/NiagaraEditor/Private/ViewModels/NiagaraEmitterViewModel.cpp
@@ -5,13 +5,13 @@
 #include "NiagaraEmitterInstance.h"
 #include "NiagaraScriptSourceBase.h"
 #include "ViewModels/NiagaraSystemViewModel.h"
-#include "NiagaraScriptViewModel.h"
+#include "ViewModels/NiagaraScriptViewModel.h"
 #include "NiagaraScriptGraphViewModel.h"
 #include "NiagaraObjectSelection.h"
 #include "NiagaraNodeOutput.h"
 #include "NiagaraScriptSource.h"
 #include "NiagaraEmitterEditorData.h"
-#include "NiagaraParameterEditMode.h"
+#include "ViewModels/NiagaraParameterEditMode.h"
 #include "NiagaraGraph.h"
 #include "NiagaraSystem.h"
 #include "NiagaraSystemInstance.h"
diff --git a/Engine/Plugins/FX/Niagara/Source/NiagaraEditor/Private/ViewModels/NiagaraParameterEditMode.h b/Engine/Plugins/FX/Niagara/Source/NiagaraEditor/Private/ViewModels/NiagaraParameterEditMode.h
deleted file mode 100644
diff --git a/Engine/Plugins/FX/Niagara/Source/NiagaraEditor/Private/ViewModels/NiagaraParameterViewModel.h b/Engine/Plugins/FX/Niagara/Source/NiagaraEditor/Private/ViewModels/NiagaraParameterViewModel.h
--- a/Engine/Plugins/FX/Niagara/Source/NiagaraEditor/Private/ViewModels/NiagaraParameterViewModel.h
+++ b/Engine/Plugins/FX/Niagara/Source/NiagaraEditor/Private/ViewModels/NiagaraParameterViewModel.h
@@ -2,7 +2,7 @@
 
 #pragma once
 
-#include "NiagaraParameterEditMode.h"
+#include "ViewModels/NiagaraParameterEditMode.h"
 #include "Delegates/DelegateCombinations.h"
 #include "Delegates/Delegate.h"
 #include "Styling/SlateTypes.h"
diff --git a/Engine/Plugins/FX/Niagara/Source/NiagaraEditor/Private/ViewModels/NiagaraScriptViewModel.cpp b/Engine/Plugins/FX/Niagara/Source/NiagaraEditor/Private/ViewModels/NiagaraScriptViewModel.cpp
--- a/Engine/Plugins/FX/Niagara/Source/NiagaraEditor/Private/ViewModels/NiagaraScriptViewModel.cpp
+++ b/Engine/Plugins/FX/Niagara/Source/NiagaraEditor/Private/ViewModels/NiagaraScriptViewModel.cpp
@@ -1,6 +1,6 @@
 // Copyright Epic Games, Inc. All Rights Reserved.
 
-#include "NiagaraScriptViewModel.h"
+#include "ViewModels/NiagaraScriptViewModel.h"
 #include "NiagaraScript.h"
 #include "NiagaraScriptSource.h"
 #include "NiagaraGraph.h"
diff --git a/Engine/Plugins/FX/Niagara/Source/NiagaraEditor/Private/ViewModels/NiagaraScriptViewModel.h b/Engine/Plugins/FX/Niagara/Source/NiagaraEditor/Private/ViewModels/NiagaraScriptViewModel.h
deleted file mode 100644
diff --git a/Engine/Plugins/FX/Niagara/Source/NiagaraEditor/Private/ViewModels/NiagaraStandaloneScriptViewModel.h b/Engine/Plugins/FX/Niagara/Source/NiagaraEditor/Private/ViewModels/NiagaraStandaloneScriptViewModel.h
--- a/Engine/Plugins/FX/Niagara/Source/NiagaraEditor/Private/ViewModels/NiagaraStandaloneScriptViewModel.h
+++ b/Engine/Plugins/FX/Niagara/Source/NiagaraEditor/Private/ViewModels/NiagaraStandaloneScriptViewModel.h
@@ -3,8 +3,8 @@
 #pragma once
 
 #include "Internationalization/Text.h"
-#include "NiagaraParameterEditMode.h"
-#include "NiagaraScriptViewModel.h"
+#include "ViewModels/NiagaraParameterEditMode.h"
+#include "ViewModels/NiagaraScriptViewModel.h"
 
 class FNiagaraMessageLogViewModel;
 class UNiagaraScript;
diff --git a/Engine/Plugins/FX/Niagara/Source/NiagaraEditor/Private/ViewModels/NiagaraSystemScriptViewModel.cpp b/Engine/Plugins/FX/Niagara/Source/NiagaraEditor/Private/ViewModels/NiagaraSystemScriptViewModel.cpp
--- a/Engine/Plugins/FX/Niagara/Source/NiagaraEditor/Private/ViewModels/NiagaraSystemScriptViewModel.cpp
+++ b/Engine/Plugins/FX/Niagara/Source/NiagaraEditor/Private/ViewModels/NiagaraSystemScriptViewModel.cpp
@@ -6,7 +6,7 @@
 #include "NiagaraGraph.h"
 #include "NiagaraTypes.h"
 #include "NiagaraScriptSource.h"
-#include "NiagaraScriptViewModel.h"
+#include "ViewModels/NiagaraScriptViewModel.h"
 #include "NiagaraScriptGraphViewModel.h"
 #include "NiagaraScriptInputCollectionViewModel.h"
 #include "NiagaraNodeEmitter.h"
diff --git a/Engine/Plugins/FX/Niagara/Source/NiagaraEditor/Private/ViewModels/NiagaraSystemScriptViewModel.h b/Engine/Plugins/FX/Niagara/Source/NiagaraEditor/Private/ViewModels/NiagaraSystemScriptViewModel.h
--- a/Engine/Plugins/FX/Niagara/Source/NiagaraEditor/Private/ViewModels/NiagaraSystemScriptViewModel.h
+++ b/Engine/Plugins/FX/Niagara/Source/NiagaraEditor/Private/ViewModels/NiagaraSystemScriptViewModel.h
@@ -2,7 +2,7 @@
 
 #pragma once
 
-#include "NiagaraScriptViewModel.h"
+#include "ViewModels/NiagaraScriptViewModel.h"
 
 class FNiagaraSystemViewModel;
 class FNiagaraScriptViewModel;
diff --git a/Engine/Plugins/FX/Niagara/Source/NiagaraEditor/Private/ViewModels/Stack/NiagaraStackEventHandlerGroup.cpp b/Engine/Plugins/FX/Niagara/Source/NiagaraEditor/Private/ViewModels/Stack/NiagaraStackEventHandlerGroup.cpp
--- a/Engine/Plugins/FX/Niagara/Source/NiagaraEditor/Private/ViewModels/Stack/NiagaraStackEventHandlerGroup.cpp
+++ b/Engine/Plugins/FX/Niagara/Source/NiagaraEditor/Private/ViewModels/Stack/NiagaraStackEventHandlerGroup.cpp
@@ -4,7 +4,7 @@
 #include "ViewModels/Stack/NiagaraStackRendererItem.h"
 #include "ViewModels/Stack/NiagaraStackEventScriptItemGroup.h"
 #include "ViewModels/NiagaraEmitterViewModel.h"
-#include "NiagaraScriptViewModel.h"
+#include "ViewModels/NiagaraScriptViewModel.h"
 #include "NiagaraScriptGraphViewModel.h"
 #include "NiagaraEmitter.h"
 #include "NiagaraGraph.h"
diff --git a/Engine/Plugins/FX/Niagara/Source/NiagaraEditor/Private/ViewModels/Stack/NiagaraStackEventScriptItemGroup.cpp b/Engine/Plugins/FX/Niagara/Source/NiagaraEditor/Private/ViewModels/Stack/NiagaraStackEventScriptItemGroup.cpp
--- a/Engine/Plugins/FX/Niagara/Source/NiagaraEditor/Private/ViewModels/Stack/NiagaraStackEventScriptItemGroup.cpp
+++ b/Engine/Plugins/FX/Niagara/Source/NiagaraEditor/Private/ViewModels/Stack/NiagaraStackEventScriptItemGroup.cpp
@@ -4,7 +4,7 @@
 #include "NiagaraEmitter.h"
 #include "ViewModels/Stack/NiagaraStackObject.h"
 #include "ViewModels/Stack/NiagaraStackModuleItem.h"
-#include "NiagaraScriptViewModel.h"
+#include "ViewModels/NiagaraScriptViewModel.h"
 #include "NiagaraScriptGraphViewModel.h"
 #include "ViewModels/NiagaraEmitterViewModel.h"
 #include "NiagaraGraph.h"
diff --git a/Engine/Plugins/FX/Niagara/Source/NiagaraEditor/Private/ViewModels/Stack/NiagaraStackRendererItem.cpp b/Engine/Plugins/FX/Niagara/Source/NiagaraEditor/Private/ViewModels/Stack/NiagaraStackRendererItem.cpp
--- a/Engine/Plugins/FX/Niagara/Source/NiagaraEditor/Private/ViewModels/Stack/NiagaraStackRendererItem.cpp
+++ b/Engine/Plugins/FX/Niagara/Source/NiagaraEditor/Private/ViewModels/Stack/NiagaraStackRendererItem.cpp
@@ -8,7 +8,7 @@
 #include "NiagaraScript.h"
 #include "ViewModels/NiagaraSystemViewModel.h"
 #include "ViewModels/NiagaraEmitterViewModel.h"
-#include "NiagaraScriptViewModel.h"
+#include "ViewModels/NiagaraScriptViewModel.h"
 #include "Internationalization/Internationalization.h"
 #include "NiagaraNodeAssignment.h"
 #include "NiagaraNodeOutput.h"
diff --git a/Engine/Plugins/FX/Niagara/Source/NiagaraEditor/Private/ViewModels/Stack/NiagaraStackRoot.cpp b/Engine/Plugins/FX/Niagara/Source/NiagaraEditor/Private/ViewModels/Stack/NiagaraStackRoot.cpp
--- a/Engine/Plugins/FX/Niagara/Source/NiagaraEditor/Private/ViewModels/Stack/NiagaraStackRoot.cpp
+++ b/Engine/Plugins/FX/Niagara/Source/NiagaraEditor/Private/ViewModels/Stack/NiagaraStackRoot.cpp
@@ -11,7 +11,7 @@
 #include "ViewModels/NiagaraSystemViewModel.h"
 #include "ViewModels/NiagaraEmitterViewModel.h"
 #include "NiagaraSystemScriptViewModel.h"
-#include "NiagaraScriptViewModel.h"
+#include "ViewModels/NiagaraScriptViewModel.h"
 #include "NiagaraSystem.h"
 #include "NiagaraSystemEditorData.h"
 #include "NiagaraEmitterEditorData.h"
diff --git a/Engine/Plugins/FX/Niagara/Source/NiagaraEditor/Private/ViewModels/Stack/NiagaraStackScriptItemGroup.cpp b/Engine/Plugins/FX/Niagara/Source/NiagaraEditor/Private/ViewModels/Stack/NiagaraStackScriptItemGroup.cpp
--- a/Engine/Plugins/FX/Niagara/Source/NiagaraEditor/Private/ViewModels/Stack/NiagaraStackScriptItemGroup.cpp
+++ b/Engine/Plugins/FX/Niagara/Source/NiagaraEditor/Private/ViewModels/Stack/NiagaraStackScriptItemGroup.cpp
@@ -5,7 +5,7 @@
 #include "ViewModels/NiagaraSystemViewModel.h"
 #include "ViewModels/NiagaraSystemSelectionViewModel.h"
 #include "ViewModels/NiagaraEmitterViewModel.h"
-#include "NiagaraScriptViewModel.h"
+#include "ViewModels/NiagaraScriptViewModel.h"
 #include "NiagaraScriptGraphViewModel.h"
 #include "NiagaraGraph.h"
 #include "NiagaraNodeOutput.h"
diff --git a/Engine/Plugins/FX/Niagara/Source/NiagaraEditor/Private/ViewModels/Stack/NiagaraStackShaderStageGroup.cpp b/Engine/Plugins/FX/Niagara/Source/NiagaraEditor/Private/ViewModels/Stack/NiagaraStackShaderStageGroup.cpp
--- a/Engine/Plugins/FX/Niagara/Source/NiagaraEditor/Private/ViewModels/Stack/NiagaraStackShaderStageGroup.cpp
+++ b/Engine/Plugins/FX/Niagara/Source/NiagaraEditor/Private/ViewModels/Stack/NiagaraStackShaderStageGroup.cpp
@@ -5,7 +5,7 @@
 #include "NiagaraShaderStageBase.h"
 #include "ViewModels/Stack/NiagaraStackObject.h"
 #include "ViewModels/Stack/NiagaraStackModuleItem.h"
-#include "NiagaraScriptViewModel.h"
+#include "ViewModels/NiagaraScriptViewModel.h"
 #include "NiagaraScriptGraphViewModel.h"
 #include "ViewModels/NiagaraEmitterViewModel.h"
 #include "NiagaraGraph.h"
diff --git a/Engine/Plugins/FX/Niagara/Source/NiagaraEditor/Private/ViewModels/Stack/NiagaraStackShaderStagesGroup.cpp b/Engine/Plugins/FX/Niagara/Source/NiagaraEditor/Private/ViewModels/Stack/NiagaraStackShaderStagesGroup.cpp
--- a/Engine/Plugins/FX/Niagara/Source/NiagaraEditor/Private/ViewModels/Stack/NiagaraStackShaderStagesGroup.cpp
+++ b/Engine/Plugins/FX/Niagara/Source/NiagaraEditor/Private/ViewModels/Stack/NiagaraStackShaderStagesGroup.cpp
@@ -2,7 +2,7 @@
 
 #include "ViewModels/Stack/NiagaraStackShaderStagesGroup.h"
 #include "ViewModels/NiagaraEmitterViewModel.h"
-#include "NiagaraScriptViewModel.h"
+#include "ViewModels/NiagaraScriptViewModel.h"
 #include "NiagaraScriptGraphViewModel.h"
 #include "NiagaraEmitter.h"
 #include "NiagaraShaderStageBase.h"
diff --git a/Engine/Plugins/FX/Niagara/Source/NiagaraEditor/Private/ViewModels/Stack/NiagaraStackSystemSettingsGroup.cpp b/Engine/Plugins/FX/Niagara/Source/NiagaraEditor/Private/ViewModels/Stack/NiagaraStackSystemSettingsGroup.cpp
--- a/Engine/Plugins/FX/Niagara/Source/NiagaraEditor/Private/ViewModels/Stack/NiagaraStackSystemSettingsGroup.cpp
+++ b/Engine/Plugins/FX/Niagara/Source/NiagaraEditor/Private/ViewModels/Stack/NiagaraStackSystemSettingsGroup.cpp
@@ -2,7 +2,7 @@
 
 #include "ViewModels/Stack/NiagaraStackSystemSettingsGroup.h"
 #include "ViewModels/Stack/NiagaraStackModuleItem.h"
-#include "NiagaraScriptViewModel.h"
+#include "ViewModels/NiagaraScriptViewModel.h"
 #include "NiagaraScriptGraphViewModel.h"
 #include "NiagaraGraph.h"
 #include "NiagaraNodeOutput.h"
diff --git a/Engine/Plugins/FX/Niagara/Source/NiagaraEditor/Private/ViewModels/Stack/NiagaraStackViewModel.cpp b/Engine/Plugins/FX/Niagara/Source/NiagaraEditor/Private/ViewModels/Stack/NiagaraStackViewModel.cpp
--- a/Engine/Plugins/FX/Niagara/Source/NiagaraEditor/Private/ViewModels/Stack/NiagaraStackViewModel.cpp
+++ b/Engine/Plugins/FX/Niagara/Source/NiagaraEditor/Private/ViewModels/Stack/NiagaraStackViewModel.cpp
@@ -6,7 +6,7 @@
 #include "ViewModels/NiagaraEmitterViewModel.h"
 #include "ViewModels/NiagaraEmitterHandleViewModel.h"
 #include "NiagaraScriptGraphViewModel.h"
-#include "NiagaraScriptViewModel.h"
+#include "ViewModels/NiagaraScriptViewModel.h"
 #include "NiagaraSystemEditorData.h"
 #include "NiagaraEmitterEditorData.h"
 #include "NiagaraStackEditorData.h"
diff --git a/Engine/Plugins/FX/Niagara/Source/NiagaraEditor/Private/Widgets/SNiagaraScriptGraph.cpp b/Engine/Plugins/FX/Niagara/Source/NiagaraEditor/Private/Widgets/SNiagaraScriptGraph.cpp
--- a/Engine/Plugins/FX/Niagara/Source/NiagaraEditor/Private/Widgets/SNiagaraScriptGraph.cpp
+++ b/Engine/Plugins/FX/Niagara/Source/NiagaraEditor/Private/Widgets/SNiagaraScriptGraph.cpp
@@ -1,6 +1,6 @@
 // Copyright Epic Games, Inc. All Rights Reserved.
 
-#include "SNiagaraScriptGraph.h"
+#include "Widgets/SNiagaraScriptGraph.h"
 #include "NiagaraScriptGraphViewModel.h"
 #include "NiagaraEditorUtilities.h"
 #include "NiagaraGraph.h"
diff --git a/Engine/Plugins/FX/Niagara/Source/NiagaraEditor/Private/Widgets/SNiagaraScriptGraph.h b/Engine/Plugins/FX/Niagara/Source/NiagaraEditor/Private/Widgets/SNiagaraScriptGraph.h
deleted file mode 100644
diff --git a/Engine/Plugins/FX/Niagara/Source/NiagaraEditor/Private/Widgets/SNiagaraSelectedObjectsDetails.cpp b/Engine/Plugins/FX/Niagara/Source/NiagaraEditor/Private/Widgets/SNiagaraSelectedObjectsDetails.cpp
--- a/Engine/Plugins/FX/Niagara/Source/NiagaraEditor/Private/Widgets/SNiagaraSelectedObjectsDetails.cpp
+++ b/Engine/Plugins/FX/Niagara/Source/NiagaraEditor/Private/Widgets/SNiagaraSelectedObjectsDetails.cpp
@@ -1,6 +1,6 @@
 // Copyright Epic Games, Inc. All Rights Reserved.
 
-#include "SNiagaraSelectedObjectsDetails.h"
+#include "Widgets/SNiagaraSelectedObjectsDetails.h"
 #include "NiagaraObjectSelection.h"
 
 #include "Modules/ModuleManager.h"
diff --git a/Engine/Plugins/FX/Niagara/Source/NiagaraEditor/Private/Widgets/SNiagaraSelectedObjectsDetails.h b/Engine/Plugins/FX/Niagara/Source/NiagaraEditor/Private/Widgets/SNiagaraSelectedObjectsDetails.h
deleted file mode 100644
diff --git a/Engine/Plugins/FX/Niagara/Source/NiagaraEditor/Private/Widgets/SNiagaraSystemScript.cpp b/Engine/Plugins/FX/Niagara/Source/NiagaraEditor/Private/Widgets/SNiagaraSystemScript.cpp
--- a/Engine/Plugins/FX/Niagara/Source/NiagaraEditor/Private/Widgets/SNiagaraSystemScript.cpp
+++ b/Engine/Plugins/FX/Niagara/Source/NiagaraEditor/Private/Widgets/SNiagaraSystemScript.cpp
@@ -3,10 +3,10 @@
 #include "SNiagaraSystemScript.h"
 #include "ViewModels/NiagaraSystemViewModel.h"
 #include "NiagaraSystemScriptViewModel.h"
-#include "NiagaraScriptViewModel.h"
+#include "ViewModels/NiagaraScriptViewModel.h"
 #include "NiagaraScriptInputCollectionViewModel.h"
 #include "SNiagaraParameterCollection.h"
-#include "SNiagaraScriptGraph.h"
+#include "Widgets/SNiagaraScriptGraph.h"
 
 #include "Widgets/Layout/SSplitter.h"
 
diff --git a/Engine/Plugins/FX/Niagara/Source/NiagaraEditor/Public/ViewModels/NiagaraParameterEditMode.h b/Engine/Plugins/FX/Niagara/Source/NiagaraEditor/Public/ViewModels/NiagaraParameterEditMode.h
new file mode 100644
--- /dev/null
+++ b/Engine/Plugins/FX/Niagara/Source/NiagaraEditor/Public/ViewModels/NiagaraParameterEditMode.h
@@ -0,0 +1,12 @@
+// Copyright Epic Games, Inc. All Rights Reserved.
+
+#pragma once
+
+/** Defines edit modes for parameters in the view model. */
+enum class ENiagaraParameterEditMode
+{
+	/** User can edit all aspects include name, type, and value, and can also add and remove. */
+	EditAll,
+	/** The user can only edit the value of the parameters. */
+	EditValueOnly
+};
\ No newline at end of file
diff --git a/Engine/Plugins/FX/Niagara/Source/NiagaraEditor/Public/ViewModels/NiagaraScriptViewModel.h b/Engine/Plugins/FX/Niagara/Source/NiagaraEditor/Public/ViewModels/NiagaraScriptViewModel.h
new file mode 100644
--- /dev/null
+++ b/Engine/Plugins/FX/Niagara/Source/NiagaraEditor/Public/ViewModels/NiagaraScriptViewModel.h
@@ -0,0 +1,132 @@
+// Copyright Epic Games, Inc. All Rights Reserved.
+
+#pragma once
+
+#include "ViewModels/NiagaraParameterEditMode.h"
+#include "INiagaraCompiler.h"
+#include "ViewModels/TNiagaraViewModelManager.h"
+#include "EditorUndoClient.h"
+#include "Misc/NotifyHook.h"
+
+class UNiagaraScript;
+class UNiagaraScriptSource;
+class INiagaraParameterCollectionViewModel;
+class FNiagaraScriptGraphViewModel;
+class FNiagaraScriptInputCollectionViewModel;
+class FNiagaraScriptOutputCollectionViewModel;
+class FNiagaraMetaDataCollectionViewModel;
+class UNiagaraEmitter;
+class FNiagaraObjectSelection;
+
+
+/** A view model for Niagara scripts which manages other script related view models. */
+class FNiagaraScriptViewModel : public TSharedFromThis<FNiagaraScriptViewModel>, public FEditorUndoClient, public TNiagaraViewModelManager<UNiagaraScript, FNiagaraScriptViewModel>
+{
+public:
+	FNiagaraScriptViewModel(FText DisplayName, ENiagaraParameterEditMode InParameterEditMode);
+
+	~FNiagaraScriptViewModel();
+
+	/** Sets the view model to a different script. */
+	void SetScript(UNiagaraScript* InScript);
+
+	void SetScripts(UNiagaraEmitter* InEmitter);
+
+	/** Gets the view model for the input parameter collection. */
+	TSharedRef<FNiagaraScriptInputCollectionViewModel> GetInputCollectionViewModel();
+	
+	/** Gets the view model for the output parameter collection. */
+	TSharedRef<FNiagaraScriptOutputCollectionViewModel> GetOutputCollectionViewModel();
+
+	TSharedRef<INiagaraParameterCollectionViewModel> GetInputParameterMapViewModel();
+	TSharedRef<INiagaraParameterCollectionViewModel> GetOutputParameterMapViewModel();
+
+	/** Gets the view model for the graph. */
+	TSharedRef<FNiagaraScriptGraphViewModel> GetGraphViewModel();
+
+	/** Gets the currently selected script variables. */
+	TSharedRef<FNiagaraObjectSelection> GetVariableSelection();
+
+	/** Updates the script with the latest compile status. */
+	void UpdateCompileStatus(ENiagaraScriptCompileStatus InAggregateCompileStatus, const FString& InAggregateCompileErrors,
+		const TArray<ENiagaraScriptCompileStatus>& InCompileStatuses, const TArray<FString>& InCompileErrors, const TArray<FString>& InCompilePaths,
+		const TArray<UNiagaraScript*>& InCompileSources);
+
+	/** Compiles a script that isn't part of an emitter or System. */
+	void CompileStandaloneScript(bool bForceCompile = false);
+
+	/** Get the latest status of this view-model's script compilation.*/
+	virtual ENiagaraScriptCompileStatus GetLatestCompileStatus();
+
+	/** Refreshes the nodes in the script graph, updating the pins to match external changes. */
+	void RefreshNodes();
+
+	//~ FEditorUndoClient Interface
+	virtual void PostUndo(bool bSuccess) override;
+	virtual void PostRedo(bool bSuccess) override { PostUndo(bSuccess); }
+
+	UNiagaraScript* GetContainerScript(ENiagaraScriptUsage InUsage, FGuid InUsageId = FGuid());
+	UNiagaraScript* GetScript(ENiagaraScriptUsage InUsage, FGuid InUsageId = FGuid());
+
+	ENiagaraScriptCompileStatus GetScriptCompileStatus(ENiagaraScriptUsage InUsage, FGuid InUsageId) const;
+	FText GetScriptErrors(ENiagaraScriptUsage InUsage, FGuid InUsageId) const;
+
+	/** If this is editing a standalone script, returns the script being edited.*/
+	UNiagaraScript* GetStandaloneScript();
+
+private:
+	/** Handles the selection changing in the graph view model. */
+	void GraphViewModelSelectedNodesChanged();
+
+	/** Handles the selection changing in the input collection view model. */
+	void InputViewModelSelectionChanged();
+
+	/** Marks this script view model as dirty and marks the scripts as needing synchrnozation. */
+	void MarkAllDirty(FString Reason);
+
+	void SetScripts(UNiagaraScriptSource* InScriptSource, TArray<UNiagaraScript*>& InScripts);
+
+	/** Handles when a value in the input parameter collection changes. */
+	void InputParameterValueChanged(FName ParameterName);
+
+	/** Handles when a value in the output parameter collection changes. */
+	void OutputParameterValueChanged(FName ParameterName);
+
+protected:
+	/** The script which provides the data for this view model. */
+	TArray<TWeakObjectPtr<UNiagaraScript>> Scripts;
+
+	virtual void OnVMScriptCompiled(UNiagaraScript* InScript);
+
+	TWeakObjectPtr<UNiagaraScriptSource> Source;
+
+	/** The view model for the input parameter collection. */
+	TSharedRef<FNiagaraScriptInputCollectionViewModel> InputCollectionViewModel;
+
+	/** The view model for the output parameter collection .*/
+	TSharedRef<FNiagaraScriptOutputCollectionViewModel> OutputCollectionViewModel;
+
+	/** The view model for the graph. */
+	TSharedRef<FNiagaraScriptGraphViewModel> GraphViewModel;
+
+	/** The set of variables currently selected in the graph or the parameters panel. */
+	TSharedRef<FNiagaraObjectSelection> VariableSelection;
+
+	/** A flag for preventing reentrancy when synchronizing selection. */
+	bool bUpdatingSelectionInternally;
+
+	/** The stored latest compile status.*/
+	ENiagaraScriptCompileStatus LastCompileStatus;
+	
+	/** The handle to the graph changed delegate needed for removing. */
+	FDelegateHandle OnGraphChangedHandle;
+
+	TArray<TNiagaraViewModelManager<UNiagaraScript, FNiagaraScriptViewModel>::Handle> RegisteredHandles;
+
+	bool IsGraphDirty() const;
+
+	TArray<ENiagaraScriptCompileStatus> CompileStatuses;
+	TArray<FString> CompileErrors;
+	TArray<FString> CompilePaths;
+	TArray<TPair<ENiagaraScriptUsage, FGuid>> CompileTypes;
+};
diff --git a/Engine/Plugins/FX/Niagara/Source/NiagaraEditor/Public/Widgets/SItemSelector.h b/Engine/Plugins/FX/Niagara/Source/NiagaraEditor/Public/Widgets/SItemSelector.h
new file mode 100644
--- /dev/null
+++ b/Engine/Plugins/FX/Niagara/Source/NiagaraEditor/Public/Widgets/SItemSelector.h
@@ -0,0 +1,573 @@
+// Copyright Epic Games, Inc. All Rights Reserved.
+
+#pragma once
+
+#include "Widgets/Views/STableRow.h"
+#include "Widgets/Input/SSearchBox.h"
+#include "Widgets/Views/STreeView.h"
+
+enum class EItemSelectorClickActivateMode
+{
+	SingleClick,
+	DoubleClick
+};
+
+/** A generic widget for selecting an item from an array of items including optional filtering and categorization. */
+template<typename CategoryType, typename ItemType>
+class SItemSelector : public SCompoundWidget
+{
+public:
+	DECLARE_DELEGATE_RetVal_OneParam(TArray<CategoryType>, FOnGetCategoriesForItem, const ItemType& /* Item */);
+	DECLARE_DELEGATE_RetVal_TwoParams(bool, FOnCompareCategoriesForEquality, const CategoryType& /* CategoryA */, const CategoryType& /* CategoryB */);
+	DECLARE_DELEGATE_RetVal_TwoParams(bool, FOnCompareCategoriesForSorting, const CategoryType& /* CategoryA */, const CategoryType& /* CateogoryB */);
+	DECLARE_DELEGATE_RetVal_TwoParams(bool, FOnCompareItemsForSorting, const ItemType& /* ItemA */, const ItemType& /* ItemB */);
+	DECLARE_DELEGATE_RetVal_TwoParams(bool, FOnDoesItemMatchFilterText, const FText& /* Filter text */, const ItemType& /* Item */);
+	DECLARE_DELEGATE_RetVal_OneParam(TSharedRef<SWidget>, FOnGenerateWidgetForCategory, const CategoryType& /* Category */);
+	DECLARE_DELEGATE_RetVal_OneParam(TSharedRef<SWidget>, FOnGenerateWidgetForItem, const ItemType& /* Item */);
+	DECLARE_DELEGATE_OneParam(FOnItemActivated, const ItemType& /* Item */);
+
+public:
+	SLATE_BEGIN_ARGS(SItemSelector)
+		: _AllowMultiselect(false)
+		, _ClickActivateMode(EItemSelectorClickActivateMode::DoubleClick)
+	{}
+
+	SLATE_ARGUMENT(TArray<ItemType>, Items)
+		/** Whether or not this item selector should allow multiple items to be selected. */
+		SLATE_ARGUMENT(bool, AllowMultiselect)
+
+		/** Whether or not a single click activates an item. */
+		SLATE_ARGUMENT(EItemSelectorClickActivateMode, ClickActivateMode)
+
+		/** An optional delegate to get an array of categories for the specified item. Each category in the returned array represents one level of nested categories. 
+		NOTE: The OnCompareCategoriesForEquality, and OnGenerateWidgetForCategory delegates must be bound if this delegate is bound. */
+		SLATE_EVENT(FOnGetCategoriesForItem, OnGetCategoriesForItem)
+
+		/** An optional delegate to compare two categories for equality which must be supplied when generating categories for items.  This equality comparer will be used to collate items into matching categories. */
+		SLATE_EVENT(FOnCompareCategoriesForEquality, OnCompareCategoriesForEquality)
+
+		/** An optional delegate which determines the sorting for categories.  If not bound categories will be ordered by the order they're encountered while processing items. */
+		SLATE_EVENT(FOnCompareCategoriesForSorting, OnCompareCategoriesForSorting)
+
+		/** An optional delegate which determines the sorting for items within each category. */
+		SLATE_EVENT(FOnCompareItemsForSorting, OnCompareItemsForSorting)
+
+		/** An optional delegate which can be used to filter items available for selection.  If not bound the search box will not be shown. */
+		SLATE_EVENT(FOnDoesItemMatchFilterText, OnDoesItemMatchFilterText)
+
+		/** An optional delegate which generates widgets for categories which must be bound when generating categories for items. */
+		SLATE_EVENT(FOnGenerateWidgetForCategory, OnGenerateWidgetForCategory)
+
+		/** The delegate which is used to generate widgets for the items to be selected. */
+		SLATE_EVENT(FOnGenerateWidgetForItem, OnGenerateWidgetForItem)
+
+		/** A delegate which is called when an item is activated by either double clicking on it or by pressing enter while it's selected. */
+		SLATE_EVENT(FOnItemActivated, OnItemActivated)
+	SLATE_END_ARGS();
+
+private:
+	enum class EItemSelectorItemViewModelType
+	{
+		Category,
+		Item
+	};
+
+	class IItemSelectorItemViewModelUtilities
+	{
+	public:
+		virtual ~IItemSelectorItemViewModelUtilities()
+		{
+		}
+
+		virtual bool DoesItemMatchFilterText(const ItemType& InItem) const = 0;
+		virtual bool CompareCategoriesForEquality(const CategoryType& CategoryA, const CategoryType& CategoryB) const = 0;
+		virtual const FOnCompareCategoriesForSorting& GetOnCompareCategoriesForSorting() const = 0;
+		virtual const FOnCompareItemsForSorting& GetOnCompareItemsForSorting() const = 0;
+	};
+
+	class FItemSelectorItemViewModel
+	{
+	public:
+		FItemSelectorItemViewModel(TSharedRef<IItemSelectorItemViewModelUtilities> InItemUtilities, EItemSelectorItemViewModelType InType)
+			: ItemUtilities(InItemUtilities)
+			, Type(InType)
+		{
+		}
+
+		virtual ~FItemSelectorItemViewModel()
+		{
+		}
+
+		TSharedRef<IItemSelectorItemViewModelUtilities> GetItemUtilities() const
+		{
+			TSharedPtr<IItemSelectorItemViewModelUtilities> ItemUtilitiesPinned = ItemUtilities.Pin();
+			checkf(ItemUtilities.IsValid(), TEXT("Item utilities deleted before item."));
+			return ItemUtilitiesPinned.ToSharedRef();
+		}
+
+		EItemSelectorItemViewModelType GetType() const
+		{
+			return Type;
+		}
+
+		void GetChildren(TArray<TSharedRef<FItemSelectorItemViewModel>>& OutChildren) const
+		{
+			TArray<TSharedRef<FItemSelectorItemViewModel>> Children;
+			GetChildrenInternal(Children);
+			for (TSharedRef<FItemSelectorItemViewModel> Child : Children)
+			{
+				if (Child->PassesFilter())
+				{
+					OutChildren.Add(Child);
+				}
+			}
+		}
+
+		virtual bool PassesFilter() const
+		{
+			return true;
+		}
+
+	protected:
+		virtual void GetChildrenInternal(TArray<TSharedRef<FItemSelectorItemViewModel>>& OutChildren) const
+		{
+		}
+
+	private:
+		TWeakPtr<IItemSelectorItemViewModelUtilities> ItemUtilities;
+		EItemSelectorItemViewModelType Type;
+	};
+
+	class FItemSelectorItemContainerViewModel : public FItemSelectorItemViewModel
+	{
+	public:
+		FItemSelectorItemContainerViewModel(TSharedRef<IItemSelectorItemViewModelUtilities> InItemUtilities, const ItemType& InItem)
+			: FItemSelectorItemViewModel(InItemUtilities, EItemSelectorItemViewModelType::Item)
+			, Item(InItem)
+		{
+		}
+
+		const ItemType& GetItem() const
+		{
+			return Item;
+		}
+
+		virtual bool PassesFilter() const
+		{
+			return this->GetItemUtilities()->DoesItemMatchFilterText(Item);
+		}
+
+	private:
+		const ItemType& Item;
+	};
+
+	class FItemSelectorItemCategoryViewModel : public FItemSelectorItemViewModel
+	{
+	public:
+		FItemSelectorItemCategoryViewModel(TSharedRef<IItemSelectorItemViewModelUtilities> InItemUtilities, const CategoryType& InCategory)
+			: FItemSelectorItemViewModel(InItemUtilities, EItemSelectorItemViewModelType::Category)
+			, Category(InCategory)
+		{
+		}
+
+		const CategoryType& GetCategory() const
+		{
+			return Category;
+		}
+
+		TSharedRef<FItemSelectorItemCategoryViewModel> AddCategory(const CategoryType& InCategory)
+		{
+			TSharedRef<FItemSelectorItemCategoryViewModel> NewCategoryViewModel = MakeShared<FItemSelectorItemCategoryViewModel>(this->GetItemUtilities(), InCategory);
+			ChildCategoryViewModels.Add(NewCategoryViewModel);
+			return NewCategoryViewModel;
+		}
+
+		void AddItem(const ItemType& InItem)
+		{
+			ChildItemViewModels.Add(MakeShared<FItemSelectorItemContainerViewModel>(this->GetItemUtilities(), InItem));
+		}
+
+		TSharedPtr<FItemSelectorItemCategoryViewModel> FindChildCategory(const CategoryType& InCategory)
+		{
+			for (TSharedRef<FItemSelectorItemCategoryViewModel> ChildCategoryViewModel : ChildCategoryViewModels)
+			{
+				if (this->GetItemUtilities()->CompareCategoriesForEquality(ChildCategoryViewModel->GetCategory(), InCategory))
+				{
+					return ChildCategoryViewModel;
+				}
+			}
+			return TSharedPtr<FItemSelectorItemCategoryViewModel>();
+		}
+
+		void SortChildren()
+		{
+			TSharedRef<IItemSelectorItemViewModelUtilities> Utilities = this->GetItemUtilities();
+			if (ChildCategoryViewModels.Num() > 0 && Utilities->GetOnCompareCategoriesForSorting().IsBound())
+			{
+				ChildCategoryViewModels.Sort([Utilities](const TSharedRef<FItemSelectorItemCategoryViewModel>& CategoryViewModelA, const TSharedRef<FItemSelectorItemCategoryViewModel>& CategoryViewModelB)
+				{
+					return Utilities->GetOnCompareCategoriesForSorting().Execute(CategoryViewModelA->GetCategory(), CategoryViewModelB->GetCategory());
+				});
+
+			}
+			if (ChildItemViewModels.Num() > 0 && Utilities->GetOnCompareItemsForSorting().IsBound())
+			{
+				ChildItemViewModels.Sort([Utilities](const TSharedRef<FItemSelectorItemContainerViewModel>& ItemViewModelA, const TSharedRef<FItemSelectorItemContainerViewModel>& ItemViewModelB)
+				{
+					return Utilities->GetOnCompareItemsForSorting().Execute(ItemViewModelA->GetItem(), ItemViewModelB->GetItem());
+				});
+			}
+
+			for (TSharedRef<FItemSelectorItemCategoryViewModel> ChildCategoryViewModel : ChildCategoryViewModels)
+			{
+				ChildCategoryViewModel->SortChildren();
+			}
+		}
+
+		virtual bool PassesFilter() const
+		{
+			for (const TSharedRef<FItemSelectorItemContainerViewModel> ChildItemViewModel : ChildItemViewModels)
+			{
+				if (ChildItemViewModel->PassesFilter())
+				{
+					return true;
+				}
+			}
+			for (const TSharedRef<FItemSelectorItemCategoryViewModel> ChildCategoryViewModel : ChildCategoryViewModels)
+			{
+				if (ChildCategoryViewModel->PassesFilter())
+				{
+					return true;
+				}
+			}
+			return false;
+		}
+
+	protected:
+		virtual void GetChildrenInternal(TArray<TSharedRef<FItemSelectorItemViewModel>>& OutChildren) const
+		{
+			OutChildren.Append(ChildCategoryViewModels);
+			OutChildren.Append(ChildItemViewModels);
+		}
+
+	private:
+		const CategoryType& Category;
+		TArray<TSharedRef<FItemSelectorItemCategoryViewModel>> ChildCategoryViewModels;
+		TArray<TSharedRef<FItemSelectorItemContainerViewModel>> ChildItemViewModels;
+	};
+
+	class FItemSelectorViewModel : public IItemSelectorItemViewModelUtilities, public TSharedFromThis<FItemSelectorViewModel>
+	{
+	public:
+		FItemSelectorViewModel(TArray<ItemType> InItems, FOnGetCategoriesForItem InOnGetCategoriesForItem, FOnCompareCategoriesForEquality InOnCompareCategoriesForEquality, FOnCompareCategoriesForSorting InOnCompareCategoriesForSorting, FOnCompareItemsForSorting InOnCompareItemsForSorting, FOnDoesItemMatchFilterText InOnDoesItemMatchFilterText)
+			: Items(InItems)
+			, OnGetCategoriesForItem(InOnGetCategoriesForItem)
+			, OnCompareCategoriesForEquality(InOnCompareCategoriesForEquality)
+			, OnCompareCategoriesForSorting(InOnCompareCategoriesForSorting)
+			, OnCompareItemsForSorting(InOnCompareItemsForSorting)
+			, OnDoesItemMatchFilterText(InOnDoesItemMatchFilterText)
+		{
+		}
+
+		const TArray<TSharedRef<FItemSelectorItemViewModel>>* GetRootItems()
+		{
+			if (RootCategoryViewModel.IsValid() == false)
+			{
+				RootCategoryViewModel = MakeShared<FItemSelectorItemCategoryViewModel>(this->AsShared(), RootCategory);
+				for (const ItemType& Item : Items)
+				{
+					// Cache the category array since the category view models hold a const reference to the generated categories and without this they
+					// are deleted.
+					TSharedRef<TArray<CategoryType>> ItemCategories = MakeShared<TArray<CategoryType>>();
+					ItemCategoriesCache.Add(ItemCategories);
+
+					if (OnGetCategoriesForItem.IsBound())
+					{
+						ItemCategories->Append(OnGetCategoriesForItem.Execute(Item));
+					}
+
+					TSharedPtr<FItemSelectorItemCategoryViewModel> CurrentCategoryViewModel = RootCategoryViewModel;
+					if (ItemCategories->Num() > 0)
+					{
+
+						for (const CategoryType& ItemCategory : (*ItemCategories))
+						{
+							TSharedPtr<FItemSelectorItemCategoryViewModel> ExistingCategoryViewModel = CurrentCategoryViewModel->FindChildCategory(ItemCategory);
+							if (ExistingCategoryViewModel.IsValid())
+							{
+								CurrentCategoryViewModel = ExistingCategoryViewModel;
+							}
+							else
+							{
+								TSharedRef<FItemSelectorItemCategoryViewModel> NewItemCategoryViewModel = CurrentCategoryViewModel->AddCategory(ItemCategory);
+								CurrentCategoryViewModel = NewItemCategoryViewModel;
+							}
+						}
+					}
+					CurrentCategoryViewModel->AddItem(Item);
+				}
+				RootCategoryViewModel->SortChildren();
+				RootCategoryViewModel->GetChildren(RootTreeCategories);
+			}
+			return &RootTreeCategories;
+		}
+
+		const FText& GetFilterText() const
+		{
+			return FilterText;
+		}
+
+		void SetFilterText(FText InFilterText)
+		{
+			FilterText = InFilterText;
+			RootTreeCategories.Empty();
+			RootCategoryViewModel->GetChildren(RootTreeCategories);
+		}
+
+		virtual bool DoesItemMatchFilterText(const ItemType& InItem) const override
+		{
+			return FilterText.IsEmpty() || OnDoesItemMatchFilterText.IsBound() == false || OnDoesItemMatchFilterText.Execute(FilterText, InItem);
+		}
+
+		virtual bool CompareCategoriesForEquality(const CategoryType& CategoryA, const CategoryType& CategoryB) const override
+		{
+			return OnCompareCategoriesForEquality.Execute(CategoryA, CategoryB);
+		}
+
+		virtual const FOnCompareCategoriesForSorting& GetOnCompareCategoriesForSorting() const override
+		{
+			return OnCompareCategoriesForSorting;
+		}
+
+		virtual const FOnCompareItemsForSorting& GetOnCompareItemsForSorting() const override
+		{
+			return OnCompareItemsForSorting;
+		}
+
+		bool CanCompareItems() const
+		{
+			return OnCompareItemsForSorting.IsBound();
+		}
+
+		int32 CompareItems(const ItemType& ItemA, const ItemType& ItemB)
+		{
+			return OnCompareItemsForSorting.Execute(ItemA, ItemB);
+		}
+
+	private:
+		CategoryType RootCategory;
+
+		TArray<ItemType> Items;
+		TArray<TSharedRef<TArray<CategoryType>>> ItemCategoriesCache;
+
+		FOnGetCategoriesForItem OnGetCategoriesForItem;
+		FOnCompareCategoriesForEquality OnCompareCategoriesForEquality;
+		FOnCompareCategoriesForSorting OnCompareCategoriesForSorting;
+		FOnCompareItemsForSorting OnCompareItemsForSorting;
+		FOnDoesItemMatchFilterText OnDoesItemMatchFilterText;
+
+		TSharedPtr<FItemSelectorItemCategoryViewModel> RootCategoryViewModel;
+		TArray<TSharedRef<FItemSelectorItemViewModel>> RootTreeCategories;
+		FText FilterText;
+	};
+
+	typedef STableRow<TSharedRef<FItemSelectorItemViewModel>> SItemSelectorTableRow;
+
+	class SItemSelectorItemContainerTableRow : public SItemSelectorTableRow
+	{
+	public:
+		SLATE_BEGIN_ARGS(SItemSelectorItemContainerTableRow)
+		{}
+			SLATE_DEFAULT_SLOT(typename SItemSelectorItemContainerTableRow::FArguments, Content)
+		SLATE_END_ARGS();
+
+		void Construct(const FArguments& InArgs, const TSharedRef<STableViewBase>& OwnerTree)
+		{
+			typename SItemSelectorTableRow::FArguments Arguments;
+			Arguments = Arguments[InArgs._Content.Widget];
+			SItemSelectorTableRow::Construct(Arguments,	OwnerTree);
+		}
+
+		virtual int32 GetIndentLevel() const override
+		{
+			return 0;
+		}
+	};
+
+public:
+	void Construct(const FArguments& InArgs)
+	{
+		Items = InArgs._Items;
+		ClickActivateMode = InArgs._ClickActivateMode;
+		OnGetCategoriesForItem = InArgs._OnGetCategoriesForItem;
+		OnCompareCategoriesForEquality = InArgs._OnCompareCategoriesForEquality;
+		OnCompareCategoriesForSorting = InArgs._OnCompareCategoriesForSorting;
+		OnCompareItemsForSorting = InArgs._OnCompareItemsForSorting;
+		OnDoesItemMatchFilterText = InArgs._OnDoesItemMatchFilterText;
+		OnGenerateWidgetForCategory = InArgs._OnGenerateWidgetForCategory;
+		OnGenerateWidgetForItem = InArgs._OnGenerateWidgetForItem;
+		OnItemActivated = InArgs._OnItemActivated;
+
+		checkf(OnGetCategoriesForItem.IsBound() == false || OnCompareCategoriesForEquality.IsBound(), TEXT("OnCompareCategoriesForEquality must be bound if OnGenerateCategoriesForItem is bound."));
+		checkf(OnGetCategoriesForItem.IsBound() == false || OnGenerateWidgetForCategory.IsBound(), TEXT("OnGenerateWidgetForCategory must be bound if OnGenerateCategoriesForItem is bound."));
+		checkf(OnGenerateWidgetForItem.IsBound(), TEXT("OnGenerateWidgetForItem must be bound"));
+
+		ViewModel = MakeShared<FItemSelectorViewModel>(Items, OnGetCategoriesForItem, OnCompareCategoriesForEquality, OnCompareCategoriesForSorting, OnCompareItemsForSorting, OnDoesItemMatchFilterText);
+
+		ChildSlot
+		[
+			SNew(SVerticalBox)
+			+ SVerticalBox::Slot()
+			.AutoHeight()
+			.Padding(0, 0, 0, 5)
+			[
+				SAssignNew(SearchBox, SSearchBox)
+				.Visibility(this, &SItemSelector::GetSearchBoxVisibility)
+				.OnTextChanged(this, &SItemSelector::OnSearchTextChanged)
+			]
+			+ SVerticalBox::Slot()
+			.Padding(0)
+			[
+				SAssignNew(ItemTree, STreeView<TSharedRef<FItemSelectorItemViewModel>>)
+				.SelectionMode(InArgs._AllowMultiselect ? ESelectionMode::Multi : ESelectionMode::SingleToggle)
+				.OnGenerateRow(this, &SItemSelector::OnGenerateRow)
+				.OnGetChildren(this, &SItemSelector::OnGetChildren)
+				.OnMouseButtonClick(this, &SItemSelector::OnMouseClick)
+				.OnMouseButtonDoubleClick(this, &SItemSelector::OnMouseDoubleClick)
+				.TreeItemsSource(ViewModel->GetRootItems())
+			]
+		];
+
+		ExpandTree();
+	}
+
+	TArray<ItemType> GetSelectedItems()
+	{
+		TArray<TSharedRef<FItemSelectorItemViewModel>> SelectedItemViewModels;
+		ItemTree->GetSelectedItems(SelectedItemViewModels);
+
+		TArray<ItemType> SelectedItems;
+		for (TSharedRef<FItemSelectorItemViewModel> SelectedItemViewModel : SelectedItemViewModels)
+		{
+			if (SelectedItemViewModel->GetType() == EItemSelectorItemViewModelType::Item)
+			{
+				TSharedRef<FItemSelectorItemContainerViewModel> SelectedItemContainerViewModel = StaticCastSharedRef<FItemSelectorItemContainerViewModel>(SelectedItemViewModel);
+				SelectedItems.Add(SelectedItemContainerViewModel->GetItem());
+			}
+		}
+		return SelectedItems;
+	}
+
+	virtual FReply OnKeyDown(const FGeometry& MyGeometry, const FKeyEvent& InKeyEvent)
+	{
+		if (InKeyEvent.GetKey() == EKeys::Enter && OnItemActivated.IsBound())
+		{
+			TArray<TSharedRef<FItemSelectorItemViewModel>> SelectedItemViewModels;
+			ItemTree->GetSelectedItems(SelectedItemViewModels);
+			if (SelectedItemViewModels.Num() == 1 && SelectedItemViewModels[0]->GetType() == EItemSelectorItemViewModelType::Item)
+			{
+				TSharedRef<FItemSelectorItemContainerViewModel> ItemContainer = StaticCastSharedRef<FItemSelectorItemContainerViewModel>(SelectedItemViewModels[0]);
+				OnItemActivated.Execute(ItemContainer->GetItem());
+				return FReply::Handled();
+			}
+		}
+		return SCompoundWidget::OnKeyDown(MyGeometry, InKeyEvent);
+	}
+
+private:
+	EVisibility GetSearchBoxVisibility() const
+	{
+		return OnDoesItemMatchFilterText.IsBound() ? EVisibility::Visible : EVisibility::Collapsed;
+	}
+
+	void OnSearchTextChanged(const FText& SearchText)
+	{
+		if (ViewModel->GetFilterText().CompareTo(SearchText) != 0)
+		{
+			ViewModel->SetFilterText(SearchText);
+			ExpandTree();
+			ItemTree->RequestTreeRefresh();
+		}
+	}
+
+	TSharedRef<ITableRow> OnGenerateRow(TSharedRef<FItemSelectorItemViewModel> Item, const TSharedRef<STableViewBase>& OwnerTable)
+	{
+		switch (Item->GetType())
+		{
+		case EItemSelectorItemViewModelType::Category:
+		{
+			TSharedRef<FItemSelectorItemCategoryViewModel> ItemCategoryViewModel = StaticCastSharedRef<FItemSelectorItemCategoryViewModel>(Item);
+			return SNew(STableRow<TSharedRef<FItemSelectorItemViewModel>>, OwnerTable)
+				.ShowSelection(false)
+				[
+					OnGenerateWidgetForCategory.Execute(ItemCategoryViewModel->GetCategory())
+				];
+		}
+		case EItemSelectorItemViewModelType::Item:
+		{
+			TSharedRef<FItemSelectorItemContainerViewModel> ItemContainerViewModel = StaticCastSharedRef<FItemSelectorItemContainerViewModel>(Item);
+			return SNew(SItemSelectorItemContainerTableRow, OwnerTable)
+				[
+					OnGenerateWidgetForItem.Execute(ItemContainerViewModel->GetItem())
+				];
+		}
+		default:
+			checkf(false, TEXT("Unsupported type"));
+			return SNew(STableRow<TSharedRef<FItemSelectorItemViewModel>>, OwnerTable);
+		}
+	}
+
+	void OnGetChildren(TSharedRef<FItemSelectorItemViewModel> Item, TArray<TSharedRef<FItemSelectorItemViewModel>>& OutChildren)
+	{
+		Item->GetChildren(OutChildren);
+	}
+
+	void OnMouseClick(TSharedRef<FItemSelectorItemViewModel> ItemClicked)
+	{
+		if (ClickActivateMode == EItemSelectorClickActivateMode::SingleClick && OnItemActivated.IsBound() && ItemClicked->GetType() == EItemSelectorItemViewModelType::Item)
+		{
+			TSharedRef<FItemSelectorItemContainerViewModel> ItemContainer = StaticCastSharedRef<FItemSelectorItemContainerViewModel>(ItemClicked);
+			OnItemActivated.Execute(ItemContainer->GetItem());
+		}
+	}
+
+	void OnMouseDoubleClick(TSharedRef<FItemSelectorItemViewModel> ItemDoubleClicked)
+	{
+		if (ClickActivateMode == EItemSelectorClickActivateMode::DoubleClick && OnItemActivated.IsBound() && ItemDoubleClicked->GetType() == EItemSelectorItemViewModelType::Item)
+		{
+			TSharedRef<FItemSelectorItemContainerViewModel> ItemContainer = StaticCastSharedRef<FItemSelectorItemContainerViewModel>(ItemDoubleClicked);
+			OnItemActivated.Execute(ItemContainer->GetItem());
+		}
+	}
+
+	void ExpandTree()
+	{
+		TArray<TSharedRef<FItemSelectorItemViewModel>> ItemsToProcess;
+		ItemsToProcess.Append(*ViewModel->GetRootItems());
+		while (ItemsToProcess.Num() > 0)
+		{
+			TSharedRef<FItemSelectorItemViewModel> ItemToProcess = ItemsToProcess[0];
+			ItemsToProcess.RemoveAtSwap(0);
+			ItemTree->SetItemExpansion(ItemToProcess, true);
+			ItemToProcess->GetChildren(ItemsToProcess);
+		}
+	}
+
+private:
+	TArray<ItemType> Items;
+
+	EItemSelectorClickActivateMode ClickActivateMode;
+
+	FOnGetCategoriesForItem OnGetCategoriesForItem;
+	FOnCompareCategoriesForEquality OnCompareCategoriesForEquality;
+	FOnCompareCategoriesForSorting OnCompareCategoriesForSorting;
+	FOnCompareItemsForSorting OnCompareItemsForSorting;
+	FOnDoesItemMatchFilterText OnDoesItemMatchFilterText;
+	FOnGenerateWidgetForCategory OnGenerateWidgetForCategory;
+	FOnGenerateWidgetForItem OnGenerateWidgetForItem;
+	FOnItemActivated OnItemActivated;
+
+	TSharedPtr<FItemSelectorViewModel> ViewModel;
+	TSharedPtr<SSearchBox> SearchBox;
+	TSharedPtr<STreeView<TSharedRef<FItemSelectorItemViewModel>>> ItemTree;
+};
\ No newline at end of file
diff --git a/Engine/Plugins/FX/Niagara/Source/NiagaraEditor/Public/Widgets/SNiagaraScriptGraph.h b/Engine/Plugins/FX/Niagara/Source/NiagaraEditor/Public/Widgets/SNiagaraScriptGraph.h
new file mode 100644
--- /dev/null
+++ b/Engine/Plugins/FX/Niagara/Source/NiagaraEditor/Public/Widgets/SNiagaraScriptGraph.h
@@ -0,0 +1,122 @@
+// Copyright Epic Games, Inc. All Rights Reserved.
+
+#pragma once
+
+#include "Widgets/SCompoundWidget.h"
+#include "Widgets/DeclarativeSyntaxSupport.h"
+#include "GraphEditor.h"
+#include "EdGraph/EdGraphNode.h"
+#include "Widgets/Input/SSearchBox.h"
+#include "NiagaraEditorCommon.h"
+
+class FNiagaraScriptGraphViewModel;
+
+/** A widget for editing a UNiagaraScript with a graph. */
+class SNiagaraScriptGraph : public SCompoundWidget
+{
+public:
+	SLATE_BEGIN_ARGS(SNiagaraScriptGraph)
+	{}
+		/** The text displayed in the title bar of the graph. */
+		SLATE_ATTRIBUTE(FText, GraphTitle)
+
+	SLATE_END_ARGS();
+
+	void Construct(const FArguments& InArgs, TSharedRef<FNiagaraScriptGraphViewModel> InViewModel);
+
+	const TSharedPtr<SGraphEditor> GetGraphEditor() { return GraphEditor; };
+
+	void FocusGraphElement(const INiagaraScriptGraphFocusInfo* FocusInfo);
+
+	void FocusGraphSearchBox();
+
+	void OnCreateComment();
+
+private:
+	/** Constructs the graph editor widget for the current graph. */
+	TSharedRef<SGraphEditor> ConstructGraphEditor();
+
+	/** Called whenever the selected nodes on the script view model changes. */
+	void ViewModelSelectedNodesChanged();
+
+	/** Called whenever the selected nodes in the graph editor changes. */
+	void GraphEditorSelectedNodesChanged(const TSet<UObject*>& SelectedNodes);
+
+	/** Called when a node is double clicked. */
+	void OnNodeDoubleClicked(UEdGraphNode* ClickedNode);
+
+	/** Called when nodes are based in the script view model. */
+	void NodesPasted(const TSet<UEdGraphNode*>& PastedNodes);
+
+	/** Sets the position on a group of newly pasted nodes. */
+	void PositionPastedNodes(const TSet<UEdGraphNode*>& PastedNodes);
+
+	/** Called whenever the view model's graph changes to a different graph. */
+	void GraphChanged();
+
+	/** Called whenever a user edits the name inline of a node.*/
+	void OnNodeTitleCommitted(const FText& NewText, ETextCommit::Type CommitInfo, UEdGraphNode* NodeBeingChanged);
+
+	/** Called whenever a user left-clicks on the graph with a chord pressed.*/
+	FReply OnSpawnGraphNodeByShortcut(FInputChord InChord, const FVector2D& InPosition);
+
+	/** Called whenever a user is trying to edit the name inline of a node and we want to make sure that it is valid.*/
+	bool OnVerifyNodeTextCommit(const FText& NewText, UEdGraphNode* NodeBeingChanged, FText& OutErrorMessage);
+
+	void OnSearchTextChanged(const FText& SearchText);
+
+	void OnSearchBoxTextCommitted(const FText& NewText, ETextCommit::Type CommitInfo);
+
+	void OnSearchBoxSearch(SSearchBox::SearchDirection Direction);
+
+	TOptional<SSearchBox::FSearchResultData> GetSearchResultData() const;
+
+	int GetCurrentFocusedMatchIndex() const { return CurrentFocusedSearchMatchIndex; }
+
+	FText GetCurrentSearchText() const { return CurrentSearchText; };
+
+	EVisibility GetGraphSearchBoxVisibility() const { return bGraphSearchBoxActive ? EVisibility::Visible : EVisibility::Collapsed; };
+
+	FReply CloseGraphSearchBoxPressed();
+
+	FReply HandleGraphSearchBoxKeyDown(const FGeometry& MyGeometry, const FKeyEvent& InKeyEvent);
+
+	void OnAlignTop();
+	void OnAlignMiddle();
+	void OnAlignBottom();
+	void OnAlignLeft();
+	void OnAlignCenter();
+	void OnAlignRight();
+
+	void OnStraightenConnections();
+
+	void OnDistributeNodesH();
+	void OnDistributeNodesV();
+
+
+private:
+	/** The combined commands used by the graph editor. */
+	TSharedPtr<FUICommandList> Commands;
+
+	/** An attribute for the title text of the graph. */
+	TAttribute<FText> GraphTitle;
+
+	/** The view model which exposes the data used by the widget. */
+	TSharedPtr<FNiagaraScriptGraphViewModel> ViewModel;
+
+	/** The graph editor which is editing the script graph. */
+	TSharedPtr<SGraphEditor> GraphEditor;
+
+	/** The search box for searching the script graph. */
+	TSharedPtr<SSearchBox> SearchBox;
+
+	/** Flag to prevent modifying the view model selection when updating the graph
+		editor selection due to a view model selection change. */
+	bool bUpdatingGraphSelectionFromViewModel;
+
+	// ~Search stuff
+	FText CurrentSearchText;
+	int CurrentFocusedSearchMatchIndex;
+	TArray<TSharedPtr<INiagaraScriptGraphFocusInfo>> CurrentSearchResults;
+	bool bGraphSearchBoxActive;
+};
diff --git a/Engine/Plugins/FX/Niagara/Source/NiagaraEditor/Public/Widgets/SNiagaraSelectedObjectsDetails.h b/Engine/Plugins/FX/Niagara/Source/NiagaraEditor/Public/Widgets/SNiagaraSelectedObjectsDetails.h
new file mode 100644
--- /dev/null
+++ b/Engine/Plugins/FX/Niagara/Source/NiagaraEditor/Public/Widgets/SNiagaraSelectedObjectsDetails.h
@@ -0,0 +1,45 @@
+// Copyright Epic Games, Inc. All Rights Reserved.
+
+#pragma once
+
+#include "Widgets/SCompoundWidget.h"
+#include "Delegates/Delegate.h"
+#include "Widgets/DeclarativeSyntaxSupport.h"
+#include "PropertyEditorDelegates.h"
+
+class FNiagaraObjectSelection;
+class IDetailsView;
+
+/** A widget for viewing and editing a set of selected objects with a details panel. */
+class SNiagaraSelectedObjectsDetails : public SCompoundWidget
+{
+public:
+	SLATE_BEGIN_ARGS(SNiagaraSelectedObjectsDetails)
+	{}
+
+	SLATE_END_ARGS();
+
+	void Construct(const FArguments& InArgs, TSharedRef<FNiagaraObjectSelection> InSelectedObjects);
+	void Construct(const FArguments& InArgs, TSharedRef<FNiagaraObjectSelection> InSelectedObjects, TSharedRef<FNiagaraObjectSelection> InSelectedObjects2);
+
+	/** Delegate to know when one of the properties has been changed.*/
+	FOnFinishedChangingProperties& OnFinishedChangingProperties() { return OnFinishedChangingPropertiesDelegate; }
+
+private:
+	/** Called whenever the object selection changes. */
+	void SelectedObjectsChanged();
+	void SelectedObjectsChangedSecond();
+
+	/** Internal delegate to route to third parties. */
+	void OnDetailsPanelFinishedChangingProperties(const FPropertyChangedEvent& InEvent);
+
+private:
+	/** The selected objects being viewed and edited by this widget. */
+	TArray<TSharedPtr<FNiagaraObjectSelection>> SelectedObjectsArray;
+
+	/** The details view for the selected object. */
+	TSharedPtr<IDetailsView> DetailsView;
+
+	/** Delegate for third parties to be notified when properties have changed.*/
+	FOnFinishedChangingProperties OnFinishedChangingPropertiesDelegate;
+};
