Allow texture data interface to change textures from BP.

diff --git a/Engine/Plugins/FX/Niagara/Source/Niagara/Classes/NiagaraDataInterfaceTexture.h b/Engine/Plugins/FX/Niagara/Source/Niagara/Classes/NiagaraDataInterfaceTexture.h
--- a/Engine/Plugins/FX/Niagara/Source/Niagara/Classes/NiagaraDataInterfaceTexture.h
+++ b/Engine/Plugins/FX/Niagara/Source/Niagara/Classes/NiagaraDataInterfaceTexture.h
@@ -43,6 +43,8 @@ public:
 	virtual bool GetFunctionHLSL(const FNiagaraDataInterfaceGPUParamInfo& ParamInfo, const FNiagaraDataInterfaceGeneratedFunction& FunctionInfo, int FunctionInstanceIndex, FString& OutHLSL) override;
 	virtual FNiagaraDataInterfaceParametersCS* ConstructComputeParameters()const override;
 
+	void SetTexture(UTexture* InTexture);
+
 	//FRWBuffer& GetGPUBuffer();
 	static const FString TextureName;
 	static const FString SamplerName;
diff --git a/Engine/Plugins/FX/Niagara/Source/Niagara/Private/NiagaraDataInterfaceTexture.cpp b/Engine/Plugins/FX/Niagara/Source/Niagara/Private/NiagaraDataInterfaceTexture.cpp
--- a/Engine/Plugins/FX/Niagara/Source/Niagara/Private/NiagaraDataInterfaceTexture.cpp
+++ b/Engine/Plugins/FX/Niagara/Source/Niagara/Private/NiagaraDataInterfaceTexture.cpp
@@ -468,4 +468,11 @@ FNiagaraDataInterfaceParametersCS* UNiagaraDataInterfaceTexture::ConstructComput
 	return new FNiagaraDataInterfaceParametersCS_Texture();
 }
 
+
+void UNiagaraDataInterfaceTexture::SetTexture(UTexture* InTexture)
+{
+	Texture = InTexture;
+	PushToRenderThread();
+}
+
 #undef LOCTEXT_NAMESPACE
\ No newline at end of file
diff --git a/Engine/Plugins/FX/Niagara/Source/Niagara/Private/NiagaraFunctionLibrary.cpp b/Engine/Plugins/FX/Niagara/Source/Niagara/Private/NiagaraFunctionLibrary.cpp
--- a/Engine/Plugins/FX/Niagara/Source/Niagara/Private/NiagaraFunctionLibrary.cpp
+++ b/Engine/Plugins/FX/Niagara/Source/Niagara/Private/NiagaraFunctionLibrary.cpp
@@ -12,6 +12,7 @@
 #include "NiagaraWorldManager.h"
 #include "NiagaraDataInterfaceStaticMesh.h"
 #include "NiagaraStats.h"
+#include "NiagaraDataInterfaceTexture.h"
 
 #define LOCTEXT_NAMESPACE "NiagaraFunctionLibrary"
 
@@ -387,6 +388,41 @@ void UNiagaraFunctionLibrary::OverrideSystemUserVariableSkeletalMeshComponent(UN
 	SkeletalMeshInterface->SetSourceComponentFromBlueprints(SkeletalMeshComponent);
 }
 
+void UNiagaraFunctionLibrary::SetTextureObject(UNiagaraComponent* NiagaraSystem, const FString& OverrideName, UTexture* Texture)
+{
+	if (!NiagaraSystem)
+	{
+		UE_LOG(LogNiagara, Warning, TEXT("NiagaraSystem in \"SetTextureObject\" is NULL, OverrideName \"%s\" and Texture \"%s\", skipping."), *OverrideName, Texture ? *Texture->GetName() : TEXT("NULL"));
+		return;
+	}
+
+	if (!Texture)
+	{
+		UE_LOG(LogNiagara, Warning, TEXT("SkeletalMeshComponent in \"SetTextureObject\" is NULL, OverrideName \"%s\" and NiagaraSystem \"%s\", skipping."), *OverrideName, *NiagaraSystem->GetOwner()->GetName());
+		return;
+	}
+
+	const FNiagaraParameterStore& OverrideParameters = NiagaraSystem->GetOverrideParameters();
+
+	FNiagaraVariable Variable(FNiagaraTypeDefinition(UNiagaraDataInterfaceTexture::StaticClass()), *OverrideName);
+
+	int32 Index = OverrideParameters.IndexOf(Variable);
+	if (Index == INDEX_NONE)
+	{
+		UE_LOG(LogNiagara, Warning, TEXT("Could not find index of variable \"%s\" in the OverrideParameters map of NiagaraSystem \"%s\"."), *OverrideName, *NiagaraSystem->GetOwner()->GetName());
+		return;
+	}
+
+	UNiagaraDataInterfaceTexture* TexureDI = Cast<UNiagaraDataInterfaceTexture>(OverrideParameters.GetDataInterface(Index));
+	if (!TexureDI)
+	{
+		UE_LOG(LogNiagara, Warning, TEXT("Did not find a matching Texture Data Interface variable named \"%s\" in the User variables of NiagaraSystem \"%s\" ."), *OverrideName, *NiagaraSystem->GetOwner()->GetName());
+		return;
+	}
+
+	TexureDI->SetTexture(Texture);
+}
+
 UNiagaraParameterCollectionInstance* UNiagaraFunctionLibrary::GetNiagaraParameterCollection(UObject* WorldContextObject, UNiagaraParameterCollection* Collection)
 {
 	UWorld* World = GEngine->GetWorldFromContextObject(WorldContextObject, EGetWorldErrorMode::LogAndReturnNull);
diff --git a/Engine/Plugins/FX/Niagara/Source/Niagara/Public/NiagaraFunctionLibrary.h b/Engine/Plugins/FX/Niagara/Source/Niagara/Public/NiagaraFunctionLibrary.h
--- a/Engine/Plugins/FX/Niagara/Source/Niagara/Public/NiagaraFunctionLibrary.h
+++ b/Engine/Plugins/FX/Niagara/Source/Niagara/Public/NiagaraFunctionLibrary.h
@@ -47,6 +47,10 @@ public:
 	/** Sets a Niagara StaticMesh parameter by name, overriding locally if necessary.*/
 	UFUNCTION(BlueprintCallable, Category = Niagara, meta = (DisplayName = "Set Niagara Skeletal Mesh Component"))
 	static void OverrideSystemUserVariableSkeletalMeshComponent(UNiagaraComponent* NiagaraSystem, const FString& OverrideName, USkeletalMeshComponent* SkeletalMeshComponent);
+
+	/** Overrides the Texture Object for a Niagara Texture Data Interface User Parameter.*/
+	UFUNCTION(BlueprintCallable, Category = Niagara)
+	static void SetTextureObject(UNiagaraComponent* NiagaraSystem, const FString& OverrideName, UTexture* Texture);
 	
 	//This is gonna be totally reworked
 // 	UFUNCTION(BlueprintCallable, Category = Niagara, meta = (Keywords = "niagara System", UnsafeDuringActorConstruction = "true"))
