Fixed: the CPU version of the particle read DI didn't check the acquire tag part of the ID.

diff --git a/Engine/Plugins/FX/Niagara/Source/Niagara/Private/NiagaraDataInterfaceParticleRead.cpp b/Engine/Plugins/FX/Niagara/Source/Niagara/Private/NiagaraDataInterfaceParticleRead.cpp
--- a/Engine/Plugins/FX/Niagara/Source/Niagara/Private/NiagaraDataInterfaceParticleRead.cpp
+++ b/Engine/Plugins/FX/Niagara/Source/Niagara/Private/NiagaraDataInterfaceParticleRead.cpp
@@ -995,15 +995,20 @@ T UNiagaraDataInterfaceParticleRead::RetrieveValueWithCheck(FNiagaraEmitterInsta
 		FNiagaraVariable IDVar(FNiagaraTypeDefinition::GetIDDef(), "ID");
 		FNiagaraDataSetAccessor<FNiagaraID> IDData(EmitterInstance->GetData(), IDVar);
 
-		int32 CorrectIndex = IDTable[ParticleID.Index];
+		int32 ParticleIndex = IDTable[ParticleID.Index];
 		T Value = T();
-		FNiagaraID ID = NIAGARA_INVALID_ID;
-		if (CorrectIndex >= 0)
+		bValid = false;
+
+		if (ParticleIndex >= 0)
 		{
-			ID = IDData.GetSafe(CorrectIndex, NIAGARA_INVALID_ID);
-			Value = ValueData.GetSafe(CorrectIndex, T());
+			FNiagaraID ActualID = IDData.GetSafe(ParticleIndex, NIAGARA_INVALID_ID);
+			if (ActualID == ParticleID)
+			{
+				Value = ValueData.GetSafe(ParticleIndex, T());
+				bValid = true;
+			}
 		}
-		bValid = (ID != NIAGARA_INVALID_ID);
+
 		return Value;
 	}
 }
