Fix deprecation warning
#rnx

diff --git a/Engine/Plugins/FX/Niagara/Source/Niagara/Private/NiagaraDataInterfaceAudioOscilloscope.cpp b/Engine/Plugins/FX/Niagara/Source/Niagara/Private/NiagaraDataInterfaceAudioOscilloscope.cpp
--- a/Engine/Plugins/FX/Niagara/Source/Niagara/Private/NiagaraDataInterfaceAudioOscilloscope.cpp
+++ b/Engine/Plugins/FX/Niagara/Source/Niagara/Private/NiagaraDataInterfaceAudioOscilloscope.cpp
@@ -389,7 +389,7 @@ void UNiagaraDataInterfaceAudioOscilloscope::PostEditChangeProperty(struct FProp
 	static FName ScopeInMillisecondsFName = GET_MEMBER_NAME_CHECKED(UNiagaraDataInterfaceAudioOscilloscope, ScopeInMilliseconds);
 
 	// Regenerate on save any compressed sound formats or if analysis needs to be re-done
-	if (UProperty* PropertyThatChanged = PropertyChangedEvent.Property)
+	if (FProperty* PropertyThatChanged = PropertyChangedEvent.Property)
 	{
 		const FName& Name = PropertyThatChanged->GetFName();
 		if (Name == SubmixFName)
