Niagara - Fix an issue where stack inputs could enter rename mode when clicked on even if they weren't selected.

diff --git a/Engine/Plugins/FX/Niagara/Source/NiagaraEditorWidgets/Private/SNiagaraStack.cpp b/Engine/Plugins/FX/Niagara/Source/NiagaraEditorWidgets/Private/SNiagaraStack.cpp
--- a/Engine/Plugins/FX/Niagara/Source/NiagaraEditorWidgets/Private/SNiagaraStack.cpp
+++ b/Engine/Plugins/FX/Niagara/Source/NiagaraEditorWidgets/Private/SNiagaraStack.cpp
@@ -908,7 +908,8 @@ SNiagaraStack::FRowWidgets SNiagaraStack::ConstructNameAndValueWidgetsForItem(UN
 	{
 		UNiagaraStackFunctionInput* FunctionInput = CastChecked<UNiagaraStackFunctionInput>(Item);
 		return FRowWidgets(
-			SNew(SNiagaraStackFunctionInputName, FunctionInput, StackViewModel),
+			SNew(SNiagaraStackFunctionInputName, FunctionInput, StackViewModel)
+			.IsSelected(Container, &SNiagaraStackTableRow::IsSelected),
 			SNew(SNiagaraStackFunctionInputValue, FunctionInput));
 	}
 	else if (Item->IsA<UNiagaraStackErrorItem>())
diff --git a/Engine/Plugins/FX/Niagara/Source/NiagaraEditorWidgets/Private/Stack/SNiagaraStackFunctionInputName.cpp b/Engine/Plugins/FX/Niagara/Source/NiagaraEditorWidgets/Private/Stack/SNiagaraStackFunctionInputName.cpp
--- a/Engine/Plugins/FX/Niagara/Source/NiagaraEditorWidgets/Private/Stack/SNiagaraStackFunctionInputName.cpp
+++ b/Engine/Plugins/FX/Niagara/Source/NiagaraEditorWidgets/Private/Stack/SNiagaraStackFunctionInputName.cpp
@@ -14,6 +14,7 @@ void SNiagaraStackFunctionInputName::Construct(const FArguments& InArgs, UNiagar
 	FunctionInput = InFunctionInput;
 	StackViewModel = InStackViewModel;
 	StackEntryItem = InFunctionInput;
+	IsSelected = InArgs._IsSelected;
 
 	ChildSlot
 	[
@@ -81,7 +82,7 @@ bool SNiagaraStackFunctionInputName::GetIsNameReadOnly() const
 
 bool SNiagaraStackFunctionInputName::GetIsNameWidgetSelected() const
 {
-	return true;
+	return IsSelected.Get();
 }
 
 bool SNiagaraStackFunctionInputName::GetIsEnabled() const
diff --git a/Engine/Plugins/FX/Niagara/Source/NiagaraEditorWidgets/Private/Stack/SNiagaraStackFunctionInputName.h b/Engine/Plugins/FX/Niagara/Source/NiagaraEditorWidgets/Private/Stack/SNiagaraStackFunctionInputName.h
--- a/Engine/Plugins/FX/Niagara/Source/NiagaraEditorWidgets/Private/Stack/SNiagaraStackFunctionInputName.h
+++ b/Engine/Plugins/FX/Niagara/Source/NiagaraEditorWidgets/Private/Stack/SNiagaraStackFunctionInputName.h
@@ -16,6 +16,7 @@ public:
 
 public:
 	SLATE_BEGIN_ARGS(SNiagaraStackFunctionInputName) { }
+		SLATE_ATTRIBUTE(bool, IsSelected);
 	SLATE_END_ARGS();
 
 	void Construct(const FArguments& InArgs, UNiagaraStackFunctionInput* InFunctionInput, UNiagaraStackViewModel* InStackViewModel);
@@ -43,4 +44,6 @@ private:
 	UNiagaraStackFunctionInput* FunctionInput;
 
 	TSharedPtr<SInlineEditableTextBlock> NameTextBlock;
+
+	TAttribute<bool> IsSelected;
 };
\ No newline at end of file
