Undef loctext namespace
#rnx

diff --git a/Engine/Plugins/FX/Niagara/Source/Niagara/Private/NiagaraDataInterfaceAudioOscilloscope.cpp b/Engine/Plugins/FX/Niagara/Source/Niagara/Private/NiagaraDataInterfaceAudioOscilloscope.cpp
--- a/Engine/Plugins/FX/Niagara/Source/Niagara/Private/NiagaraDataInterfaceAudioOscilloscope.cpp
+++ b/Engine/Plugins/FX/Niagara/Source/Niagara/Private/NiagaraDataInterfaceAudioOscilloscope.cpp
@@ -594,3 +594,5 @@ void FNiagaraSubmixListener::OnNewSubmixBuffer(const USoundSubmix* OwningSubmix,
 	SubmixSampleRate = SampleRate;
 	MixerInput.PushAudio(AudioData, NumSamples);
 }
+
+#undef LOCTEXT_NAMESPACE
\ No newline at end of file
