Fixed: the direct reads DI didn't work correctly when multiple copies of an emitter were used in the same scene.

DI instances are shared between emitters, so we must manually track per-instance data inside the render thread proxy.

diff --git a/Engine/Plugins/FX/Niagara/Source/Niagara/Classes/NiagaraDataInterfaceParticleRead.h b/Engine/Plugins/FX/Niagara/Source/Niagara/Classes/NiagaraDataInterfaceParticleRead.h
--- a/Engine/Plugins/FX/Niagara/Source/Niagara/Classes/NiagaraDataInterfaceParticleRead.h
+++ b/Engine/Plugins/FX/Niagara/Source/Niagara/Classes/NiagaraDataInterfaceParticleRead.h
@@ -7,12 +7,6 @@
 #include "NiagaraEmitterInstance.h"
 #include "NiagaraDataInterfaceParticleRead.generated.h"
 
-struct FNDIParticleRead_InstanceData
-{
-	FNiagaraSystemInstance* SystemInstance;
-	FNiagaraEmitterInstance* EmitterInstance;
-};
-
 UCLASS(EditInlineNew, Category = "ParticleRead", meta = (DisplayName = "Particle Attribute Reader"))
 class NIAGARA_API UNiagaraDataInterfaceParticleRead : public UNiagaraDataInterface
 {
@@ -27,7 +21,8 @@ public:
 
 	//UNiagaraDataInterface Interface
 	virtual bool InitPerInstanceData(void* PerInstanceData, FNiagaraSystemInstance* SystemInstance) override;
-	virtual int32 PerInstanceDataSize() const { return sizeof(FNDIParticleRead_InstanceData); }
+	virtual void DestroyPerInstanceData(void* PerInstanceData, FNiagaraSystemInstance* SystemInstance) override;
+	virtual int32 PerInstanceDataSize() const override;
 	virtual void GetFunctions(TArray<FNiagaraFunctionSignature>& OutFunctions) override;
 	virtual void GetVMExternalFunction(const FVMExternalFunctionBindingInfo& BindingInfo, void* InstanceData, FVMExternalFunction& OutFunc) override;
 	virtual bool CanExecuteOnTarget(ENiagaraSimTarget Target) const override { return true; }
diff --git a/Engine/Plugins/FX/Niagara/Source/Niagara/Private/NiagaraDataInterfaceParticleRead.cpp b/Engine/Plugins/FX/Niagara/Source/Niagara/Private/NiagaraDataInterfaceParticleRead.cpp
--- a/Engine/Plugins/FX/Niagara/Source/Niagara/Private/NiagaraDataInterfaceParticleRead.cpp
+++ b/Engine/Plugins/FX/Niagara/Source/Niagara/Private/NiagaraDataInterfaceParticleRead.cpp
@@ -67,35 +67,65 @@ static const TCHAR* NiagaraParticleDataValueTypeName(ENiagaraParticleDataValueTy
 	}
 }
 
-struct FNDIParticleRead_InstanceDataGPU
+struct FNDIParticleRead_InstanceData
 {
+	FNiagaraSystemInstance* SystemInstance;
+	FNiagaraEmitterInstance* EmitterInstance;
+};
+
+struct FNDIParticleRead_RenderThreadData
+{
+	FNDIParticleRead_RenderThreadData() : SourceEmitterGPUContext(nullptr) {}
+
 	FNiagaraComputeExecutionContext* SourceEmitterGPUContext;
 	FString SourceEmitterName;
 };
 
 struct FNiagaraDataInterfaceProxyParticleRead : public FNiagaraDataInterfaceProxy
 {
-	FNiagaraDataInterfaceProxyParticleRead() : SourceEmitterGPUContext(nullptr)
+	virtual void ConsumePerInstanceDataFromGameThread(void* PerInstanceData, const FNiagaraSystemInstanceID& InstanceID) override
 	{
-	}
+		FNDIParticleRead_RenderThreadData* SystemData = SystemsRenderData.Find(InstanceID);
+		if (!ensure(SystemData))
+		{
+			return;
+		}
 
-	virtual void ConsumePerInstanceDataFromGameThread(void* PerInstanceData, const FNiagaraSystemInstanceID& Instance) override
-	{
-		FNDIParticleRead_InstanceDataGPU* InstanceData = static_cast<FNDIParticleRead_InstanceDataGPU*>(PerInstanceData);
-		if (InstanceData)
+		FNDIParticleRead_RenderThreadData* IncomingData = static_cast<FNDIParticleRead_RenderThreadData*>(PerInstanceData);
+		if (!IncomingData)
 		{
-			SourceEmitterGPUContext = InstanceData->SourceEmitterGPUContext;
-			SourceEmitterName = InstanceData->SourceEmitterName;
+			return;
 		}
+
+		*SystemData = *IncomingData;
 	}
 	
 	virtual int32 PerInstanceDataPassedToRenderThreadSize() const override
 	{
-		return sizeof(FNDIParticleRead_InstanceDataGPU);
+		return sizeof(FNDIParticleRead_RenderThreadData);
 	}
 
-	FNiagaraComputeExecutionContext* SourceEmitterGPUContext;
-	FString SourceEmitterName;
+	void CreateRenderThreadSystemData(const FNiagaraSystemInstanceID& InstanceID)
+	{
+		check(IsInRenderingThread());
+		check(!SystemsRenderData.Contains(InstanceID));
+		SystemsRenderData.Add(InstanceID);
+	}
+
+	void DestroyRenderThreadSystemData(const FNiagaraSystemInstanceID& InstanceID)
+	{
+		check(IsInRenderingThread());
+		int32 NumRemoved = SystemsRenderData.Remove(InstanceID);
+		check(NumRemoved == 1);
+	}
+
+	FNDIParticleRead_RenderThreadData* GetRenderDataForSystem(const FNiagaraSystemInstanceID& InstanceID)
+	{
+		return SystemsRenderData.Find(InstanceID);
+	}
+
+private:
+	TMap<FNiagaraSystemInstanceID, FNDIParticleRead_RenderThreadData> SystemsRenderData;
 };
 
 struct FNiagaraDataInterfaceParametersCS_ParticleRead : public FNiagaraDataInterfaceParametersCS
@@ -295,12 +325,16 @@ struct FNiagaraDataInterfaceParametersCS_ParticleRead : public FNiagaraDataInter
 
 		FNiagaraDataInterfaceProxyParticleRead* Proxy = static_cast<FNiagaraDataInterfaceProxyParticleRead*>(Context.DataInterface);
 		check(Proxy);
-		if (Proxy->SourceEmitterGPUContext == nullptr)
+
+		FNDIParticleRead_RenderThreadData* RTData = Proxy->GetRenderDataForSystem(Context.SystemInstance);
+		check(RTData);
+
+		if (RTData->SourceEmitterGPUContext == nullptr)
 		{
 			// This means the source emitter isn't running on GPU.
 			if (!bSourceEmitterNotGPUErrorShown)
 			{
-				UE_LOG(LogNiagara, Error, TEXT("GPU particle read DI is set to access CPU emitter '%s'."), *Proxy->SourceEmitterName);
+				UE_LOG(LogNiagara, Error, TEXT("GPU particle read DI is set to access CPU emitter '%s'."), *RTData->SourceEmitterName);
 				bSourceEmitterNotGPUErrorShown = true;
 			}
 			SetErrorParams(RHICmdList, ComputeShader, false);
@@ -309,7 +343,7 @@ struct FNiagaraDataInterfaceParametersCS_ParticleRead : public FNiagaraDataInter
 
 		bSourceEmitterNotGPUErrorShown = false;
 
-		FNiagaraDataSet* SourceDataSet = Proxy->SourceEmitterGPUContext->MainDataSet;
+		FNiagaraDataSet* SourceDataSet = RTData->SourceEmitterGPUContext->MainDataSet;
 		if (!SourceDataSet)
 		{
 			SetErrorParams(RHICmdList, ComputeShader, false);
@@ -318,7 +352,7 @@ struct FNiagaraDataInterfaceParametersCS_ParticleRead : public FNiagaraDataInter
 
 		FNiagaraDataBuffer* SourceData;
 		uint32 NumSpawnedInstances = 0, IDAcquireTag = 0;
-		const bool bReadingOwnEmitter = Context.ComputeInstanceData->Context == Proxy->SourceEmitterGPUContext;
+		const bool bReadingOwnEmitter = Context.ComputeInstanceData->Context == RTData->SourceEmitterGPUContext;
 		if (bReadingOwnEmitter)
 		{
 			// If the current execution context is the same as the source emitter's context, it means we're reading from
@@ -357,7 +391,7 @@ struct FNiagaraDataInterfaceParametersCS_ParticleRead : public FNiagaraDataInter
 
 		if (CachedDataSet != SourceDataSet)
 		{
-			FindAttributeIndices(SourceDataSet, *Proxy->SourceEmitterName);
+			FindAttributeIndices(SourceDataSet, *RTData->SourceEmitterName);
 			CachedDataSet = SourceDataSet;
 		}
 
@@ -443,9 +477,36 @@ bool UNiagaraDataInterfaceParticleRead::InitPerInstanceData(void* PerInstanceDat
 		return false;
 	}
 
+	FNiagaraDataInterfaceProxyParticleRead* ThisProxy = GetProxyAs<FNiagaraDataInterfaceProxyParticleRead>();
+	ENQUEUE_RENDER_COMMAND(FNDIParticleReadCreateRTInstance)(
+		[ThisProxy, InstanceID = SystemInstance->GetId()](FRHICommandList& CmdList)
+		{
+			ThisProxy->CreateRenderThreadSystemData(InstanceID);
+		}
+	);
+
 	return true;
 }
 
+void UNiagaraDataInterfaceParticleRead::DestroyPerInstanceData(void* PerInstanceData, FNiagaraSystemInstance* SystemInstance)
+{
+	FNDIParticleRead_InstanceData* InstData = (FNDIParticleRead_InstanceData*)PerInstanceData;
+	InstData->~FNDIParticleRead_InstanceData();
+
+	FNiagaraDataInterfaceProxyParticleRead* ThisProxy = GetProxyAs<FNiagaraDataInterfaceProxyParticleRead>();
+	ENQUEUE_RENDER_COMMAND(FNDIParticleReadDestroyRTInstance) (
+		[ThisProxy, InstanceID = SystemInstance->GetId()](FRHICommandListImmediate& CmdList)
+		{
+			ThisProxy->DestroyRenderThreadSystemData(InstanceID);
+		}
+	);
+}
+
+int32 UNiagaraDataInterfaceParticleRead::PerInstanceDataSize() const
+{
+	return sizeof(FNDIParticleRead_InstanceData);
+}
+
 void UNiagaraDataInterfaceParticleRead::GetFunctions(TArray<FNiagaraFunctionSignature>& OutFunctions)
 {
 	{
@@ -1254,12 +1315,12 @@ FNiagaraDataInterfaceParametersCS* UNiagaraDataInterfaceParticleRead::ConstructC
 
 void UNiagaraDataInterfaceParticleRead::ProvidePerInstanceDataForRenderThread(void* DataForRenderThread, void* PerInstanceData, const FNiagaraSystemInstanceID& SystemInstance)
 {
-	FNDIParticleRead_InstanceDataGPU* DataToPass = new (DataForRenderThread) FNDIParticleRead_InstanceDataGPU;
+	FNDIParticleRead_RenderThreadData* RTData = new (DataForRenderThread) FNDIParticleRead_RenderThreadData;
 	FNDIParticleRead_InstanceData* PIData = static_cast<FNDIParticleRead_InstanceData*>(PerInstanceData);
 	if (PIData && PIData->EmitterInstance)
 	{
-		DataToPass->SourceEmitterGPUContext = PIData->EmitterInstance->GetGPUContext();
-		DataToPass->SourceEmitterName = PIData->EmitterInstance->GetCachedEmitter()->GetUniqueEmitterName();
+		RTData->SourceEmitterGPUContext = PIData->EmitterInstance->GetGPUContext();
+		RTData->SourceEmitterName = PIData->EmitterInstance->GetCachedEmitter()->GetUniqueEmitterName();
 	}
 }
 
