Removed overzealous assert.

It's possible to reach DestroyRenderThreadSystemData with an empty map when the system is being destroyed.

diff --git a/Engine/Plugins/FX/Niagara/Source/Niagara/Private/NiagaraDataInterfaceParticleRead.cpp b/Engine/Plugins/FX/Niagara/Source/Niagara/Private/NiagaraDataInterfaceParticleRead.cpp
--- a/Engine/Plugins/FX/Niagara/Source/Niagara/Private/NiagaraDataInterfaceParticleRead.cpp
+++ b/Engine/Plugins/FX/Niagara/Source/Niagara/Private/NiagaraDataInterfaceParticleRead.cpp
@@ -94,6 +94,8 @@ struct FNiagaraDataInterfaceProxyParticleRead : public FNiagaraDataInterfaceProx
 		FNDIParticleRead_RenderThreadData* IncomingData = static_cast<FNDIParticleRead_RenderThreadData*>(PerInstanceData);
 		if (!IncomingData)
 		{
+			SystemData->SourceEmitterGPUContext = nullptr;
+			SystemData->SourceEmitterName = TEXT("");
 			return;
 		}
 
@@ -115,8 +117,7 @@ struct FNiagaraDataInterfaceProxyParticleRead : public FNiagaraDataInterfaceProx
 	void DestroyRenderThreadSystemData(const FNiagaraSystemInstanceID& InstanceID)
 	{
 		check(IsInRenderingThread());
-		int32 NumRemoved = SystemsRenderData.Remove(InstanceID);
-		check(NumRemoved == 1);
+		SystemsRenderData.Remove(InstanceID);
 	}
 
 	FNDIParticleRead_RenderThreadData* GetRenderDataForSystem(const FNiagaraSystemInstanceID& InstanceID)
