Disabled USE_WAVE_INTRISICS on PS4 as it creates visual artefacts.

diff --git a/Engine/Plugins/FX/Niagara/Shaders/Private/NiagaraEmitterInstanceShader.usf b/Engine/Plugins/FX/Niagara/Shaders/Private/NiagaraEmitterInstanceShader.usf
--- a/Engine/Plugins/FX/Niagara/Shaders/Private/NiagaraEmitterInstanceShader.usf
+++ b/Engine/Plugins/FX/Niagara/Shaders/Private/NiagaraEmitterInstanceShader.usf
@@ -469,7 +469,7 @@ int rand_int(int x, int Seed1, int Seed2, int Seed3)
 		RWIDToIndexTable[IDIndex] = InstanceIndex;
 	}
 
-#define USE_WAVE_INTRISICS (PS4_PROFILE || XBOXONE_PROFILE)
+#define USE_WAVE_INTRISICS (XBOXONE_PROFILE)
 #define USE_GROUP_SHARED ((THREADGROUP_SIZE == 64 || THREADGROUP_SIZE == 32)  && !USE_WAVE_INTRISICS && !VULKAN_PROFILE_SM5 && !VULKAN_PROFILE)
 
 #if USE_GROUP_SHARED
