- Various fixes to matrix related functionality, some were transposed and some were not
- You can now generate a matrix inside BP and pass in the vectors to a Niagara system and have it work as expected on both CPU & GPU
- Added MatrixToQuaternion, this uses the fast path library under the hood but gives us a location to replace once we address matrix access issues inside the VM

diff --git a/Engine/Plugins/FX/Niagara/Source/Niagara/Private/NiagaraFunctionLibrary.cpp b/Engine/Plugins/FX/Niagara/Source/Niagara/Private/NiagaraFunctionLibrary.cpp
--- a/Engine/Plugins/FX/Niagara/Source/Niagara/Private/NiagaraFunctionLibrary.cpp
+++ b/Engine/Plugins/FX/Niagara/Source/Niagara/Private/NiagaraFunctionLibrary.cpp
@@ -718,9 +718,9 @@ struct FVectorKernelFastMatrixToQuaternion
 			float InvS = rsqrt(tr + 1.f);
 			float s = 0.5f * InvS;
 
-			Quat.x = (Mat[2][1] - Mat[1][2]) * s;
-			Quat.y = (Mat[0][2] - Mat[2][0]) * s;
-			Quat.z = (Mat[1][0] - Mat[0][1]) * s;
+			Quat.x = (Mat[1][2] - Mat[2][1]) * s;
+			Quat.y = (Mat[2][0] - Mat[0][2]) * s;
+			Quat.z = (Mat[0][1] - Mat[1][0]) * s;
 			Quat.w = 0.5f * rcp(InvS);
 		}
 		else if ( (Mat[0][0] > Mat[1][1]) && (Mat[0][0] > Mat[2][2]) )
@@ -730,9 +730,9 @@ struct FVectorKernelFastMatrixToQuaternion
 			s = 0.5f * InvS;
 
 			Quat.x = 0.5f * rcp(InvS);
-			Quat.y = (Mat[1][0] + Mat[0][1]) * s;
-			Quat.z = (Mat[2][0] + Mat[0][2]) * s;
-			Quat.w = (Mat[2][1] - Mat[1][2]) * s;
+			Quat.y = (Mat[0][1] + Mat[1][0]) * s;
+			Quat.z = (Mat[0][2] + Mat[2][0]) * s;
+			Quat.w = (Mat[1][2] - Mat[2][1]) * s;
 		}
 		else if ( Mat[1][1] > Mat[2][2] )
 		{
@@ -740,10 +740,10 @@ struct FVectorKernelFastMatrixToQuaternion
 			float InvS = rsqrt(s);
 			s = 0.5f * InvS;
 
-			Quat.x = (Mat[0][1] + Mat[1][0]) * s;
+			Quat.x = (Mat[1][0] + Mat[0][1]) * s;
 			Quat.y = 0.5f * rcp(InvS);
-			Quat.z = (Mat[2][1] + Mat[1][2]) * s;
-			Quat.w = (Mat[0][2] - Mat[2][0]) * s;
+			Quat.z = (Mat[1][2] + Mat[2][1]) * s;
+			Quat.w = (Mat[2][0] - Mat[0][2]) * s;
 
 		}
 		else
@@ -752,10 +752,10 @@ struct FVectorKernelFastMatrixToQuaternion
 			float InvS = rsqrt(s);
 			s = 0.5f * InvS;
 
-			Quat.x = (Mat[0][2] + Mat[2][0]) * s;
-			Quat.y = (Mat[1][2] + Mat[2][1]) * s;
+			Quat.x = (Mat[2][0] + Mat[0][2]) * s;
+			Quat.y = (Mat[2][1] + Mat[1][2]) * s;
 			Quat.z = 0.5f * rcp(InvS);
-			Quat.w = (Mat[1][0] - Mat[0][1]) * s;
+			Quat.w = (Mat[0][1] - Mat[1][0]) * s;
 		}
 	}
 )");
@@ -777,15 +777,65 @@ struct FVectorKernelFastMatrixToQuaternion
 			OutQuat.Emplace(Context);
 		}
 
-		for ( int32 i=0; i < Context.GetNumInstances(); ++i )
+		for ( int32 iInstance=0; iInstance < Context.GetNumInstances(); ++iInstance)
 		{
+			// Read Matrix
 			FMatrix Mat;
 			for ( int32 j=0; j < 16; ++j )
 			{
-				Mat.M[j & 3][j >> 2] = InMatrix[j].GetAndAdvance();
+				Mat.M[j >> 2][j & 3] = InMatrix[j].GetAndAdvance();
+			}
+
+			// Generate Quat (without checking for consistency with the GPU)
+			FQuat Quat;
+			{
+				// Check diagonal (trace)
+				const float tr = Mat.M[0][0] + Mat.M[1][1] + Mat.M[2][2];
+				if (tr > 0.0f)
+				{
+					float InvS = FMath::InvSqrt(tr + 1.f);
+					Quat.W = 0.5f * (1.f / InvS);
+					const float s = 0.5f * InvS;
+
+					Quat.X = (Mat.M[1][2] - Mat.M[2][1]) * s;
+					Quat.Y = (Mat.M[2][0] - Mat.M[0][2]) * s;
+					Quat.Z = (Mat.M[0][1] - Mat.M[1][0]) * s;
+				}
+				else
+				{
+					// diagonal is negative
+					int32 i = 0;
+
+					if (Mat.M[1][1] > Mat.M[0][0])
+						i = 1;
+
+					if (Mat.M[2][2] > Mat.M[i][i])
+						i = 2;
+
+					static const int32 nxt[3] = { 1, 2, 0 };
+					const int32 j = nxt[i];
+					const int32 k = nxt[j];
+
+					float s = Mat.M[i][i] - Mat.M[j][j] - Mat.M[k][k] + 1.0f;
+
+					float InvS = FMath::InvSqrt(s);
+
+					float qt[4];
+					qt[i] = 0.5f * (1.f / InvS);
+
+					s = 0.5f * InvS;
+
+					qt[3] = (Mat.M[j][k] - Mat.M[k][j]) * s;
+					qt[j] = (Mat.M[i][j] + Mat.M[j][i]) * s;
+					qt[k] = (Mat.M[i][k] + Mat.M[k][i]) * s;
+
+					Quat.X = qt[0];
+					Quat.Y = qt[1];
+					Quat.Z = qt[2];
+					Quat.W = qt[3];
+				}
 			}
 
-			FQuat Quat(Mat);
 			*OutQuat[0].GetDestAndAdvance() = Quat.X;
 			*OutQuat[1].GetDestAndAdvance() = Quat.Y;
 			*OutQuat[2].GetDestAndAdvance() = Quat.Z;
diff --git a/Engine/Plugins/FX/Niagara/Source/NiagaraCore/Private/NiagaraCustomVersion.cpp b/Engine/Plugins/FX/Niagara/Source/NiagaraCore/Private/NiagaraCustomVersion.cpp
--- a/Engine/Plugins/FX/Niagara/Source/NiagaraCore/Private/NiagaraCustomVersion.cpp
+++ b/Engine/Plugins/FX/Niagara/Source/NiagaraCore/Private/NiagaraCustomVersion.cpp
@@ -8,4 +8,4 @@ const FGuid FNiagaraCustomVersion::GUID(0xFCF57AFA, 0x50764283, 0xB9A9E658, 0xFF
 // Register the custom version with core
 FCustomVersionRegistration GRegisterNiagaraCustomVersion(FNiagaraCustomVersion::GUID, FNiagaraCustomVersion::LatestVersion, TEXT("NiagaraVer"));
 
-const FGuid FNiagaraCustomVersion::LatestScriptCompileVersion(0xCA9F2224, 0x728A48DD, 0x87238CFD, 0x5DF7DED7);
+const FGuid FNiagaraCustomVersion::LatestScriptCompileVersion(0x6E2BA7D2, 0xEDAC9A43, 0xB0972E7C, 0x95A69C9C);
