Niagara renderer properties now produce the right widgets for their renderer preview
#rnx

diff --git a/Engine/Plugins/FX/Niagara/Source/Niagara/Niagara.Build.cs b/Engine/Plugins/FX/Niagara/Source/Niagara/Niagara.Build.cs
--- a/Engine/Plugins/FX/Niagara/Source/Niagara/Niagara.Build.cs
+++ b/Engine/Plugins/FX/Niagara/Source/Niagara/Niagara.Build.cs
@@ -61,6 +61,8 @@ public class Niagara : ModuleRules
                 new string[] {
                 "TargetPlatform",
                 "UnrealEd",
+				"SlateCore",
+				"Slate"
             });
         }
     }
diff --git a/Engine/Plugins/FX/Niagara/Source/Niagara/Private/NiagaraLightRendererProperties.cpp b/Engine/Plugins/FX/Niagara/Source/Niagara/Private/NiagaraLightRendererProperties.cpp
--- a/Engine/Plugins/FX/Niagara/Source/Niagara/Private/NiagaraLightRendererProperties.cpp
+++ b/Engine/Plugins/FX/Niagara/Source/Niagara/Private/NiagaraLightRendererProperties.cpp
@@ -5,6 +5,17 @@
 #include "NiagaraConstants.h"
 #include "NiagaraRendererLights.h"
 #include "Modules/ModuleManager.h"
+#if WITH_EDITOR
+#include "Widgets/Images/SImage.h"
+#include "Styling/SlateIconFinder.h"
+#include "Widgets/SWidget.h"
+#include "Styling/SlateBrush.h"
+#include "AssetThumbnail.h"
+#include "Widgets/Text/STextBlock.h"
+#endif
+
+
+#define LOCTEXT_NAMESPACE "UNiagaraLightRendererProperties"
 
 TArray<TWeakObjectPtr<UNiagaraLightRendererProperties>> UNiagaraLightRendererProperties::LightRendererPropertiesToDeferredInit;
 
@@ -101,6 +112,20 @@ const TArray<FNiagaraVariable>& UNiagaraLightRendererProperties::GetOptionalAttr
 	return Attrs;
 }
 
+void UNiagaraLightRendererProperties::GetRendererWidgets(const FNiagaraEmitterInstance* InEmitter, TArray<TSharedPtr<SWidget>>& OutWidgets, TSharedPtr<FAssetThumbnailPool> InThumbnailPool) const
+{
+	TSharedRef<SWidget> LightWidget = SNew(SImage)
+		.Image(FSlateIconFinder::FindIconBrushForClass(GetClass()));
+	OutWidgets.Add(LightWidget);
+}
+
+void UNiagaraLightRendererProperties::GetRendererTooltipWidgets(const FNiagaraEmitterInstance* InEmitter, TArray<TSharedPtr<SWidget>>& OutWidgets, TSharedPtr<FAssetThumbnailPool> InThumbnailPool) const
+{
+	TSharedRef<SWidget> LightTooltip = SNew(STextBlock)
+		.Text(LOCTEXT("LightRenderer", "Light Renderer"));
+	OutWidgets.Add(LightTooltip);
+}
+
 bool UNiagaraLightRendererProperties::IsMaterialValidForRenderer(UMaterial* Material, FText& InvalidMessage)
 {
 	return true;
@@ -112,3 +137,4 @@ void UNiagaraLightRendererProperties::FixMaterial(UMaterial* Material)
 
 #endif // WITH_EDITORONLY_DATA
 
+#undef LOCTEXT_NAMESPACE
\ No newline at end of file
diff --git a/Engine/Plugins/FX/Niagara/Source/Niagara/Private/NiagaraMeshRendererProperties.cpp b/Engine/Plugins/FX/Niagara/Source/Niagara/Private/NiagaraMeshRendererProperties.cpp
--- a/Engine/Plugins/FX/Niagara/Source/Niagara/Private/NiagaraMeshRendererProperties.cpp
+++ b/Engine/Plugins/FX/Niagara/Source/Niagara/Private/NiagaraMeshRendererProperties.cpp
@@ -6,6 +6,18 @@
 #include "NiagaraConstants.h"
 #include "NiagaraBoundsCalculatorHelper.h"
 #include "Modules/ModuleManager.h"
+#if WITH_EDITOR
+#include "Widgets/Images/SImage.h"
+#include "Styling/SlateIconFinder.h"
+#include "Widgets/SWidget.h"
+#include "Styling/SlateBrush.h"
+#include "AssetThumbnail.h"
+#include "Widgets/Text/STextBlock.h"
+#endif
+
+
+#define LOCTEXT_NAMESPACE "UNiagaraMeshRendererProperties"
+
 
 TArray<TWeakObjectPtr<UNiagaraMeshRendererProperties>> UNiagaraMeshRendererProperties::MeshRendererPropertiesToDeferredInit;
 
@@ -271,6 +283,46 @@ const TArray<FNiagaraVariable>& UNiagaraMeshRendererProperties::GetOptionalAttri
 	return Attrs;
 }
 
+void UNiagaraMeshRendererProperties::GetRendererWidgets(const FNiagaraEmitterInstance* InEmitter, TArray<TSharedPtr<SWidget>>& OutWidgets, TSharedPtr<FAssetThumbnailPool> InThumbnailPool) const
+{
+	TSharedRef<SWidget> ThumbnailWidget = SNullWidget::NullWidget;
+	int32 ThumbnailSize = 32;
+	TArray<UMaterialInterface*> Materials;
+	GetUsedMaterials(InEmitter, Materials);
+	for (UMaterialInterface* Material : Materials)
+	{
+		TSharedPtr<FAssetThumbnail> AssetThumbnail = MakeShareable(new FAssetThumbnail(Material, ThumbnailSize, ThumbnailSize, InThumbnailPool));
+		if (AssetThumbnail)
+		{
+			ThumbnailWidget = AssetThumbnail->MakeThumbnailWidget();
+		}
+		OutWidgets.Add(ThumbnailWidget);
+	}
+
+	if (Materials.Num() == 0)
+	{
+		TSharedRef<SWidget> SpriteWidget = SNew(SImage)
+			.Image(FSlateIconFinder::FindIconBrushForClass(GetClass()));
+		OutWidgets.Add(SpriteWidget);
+	}
+}
+
+void UNiagaraMeshRendererProperties::GetRendererTooltipWidgets(const FNiagaraEmitterInstance* InEmitter, TArray<TSharedPtr<SWidget>>& OutWidgets, TSharedPtr<FAssetThumbnailPool> InThumbnailPool) const
+{
+	TArray<UMaterialInterface*> Materials;
+	GetUsedMaterials(InEmitter, Materials);
+	if (Materials.Num() > 0)
+	{
+		GetRendererWidgets(InEmitter, OutWidgets, InThumbnailPool);
+	}
+	else
+	{
+		TSharedRef<SWidget> MeshTooltip = SNew(STextBlock)
+			.Text(LOCTEXT("MeshRendererNoMat", "Mesh Renderer (No Material Set)"));
+		OutWidgets.Add(MeshTooltip);
+	}
+}
+
 void UNiagaraMeshRendererProperties::BeginDestroy()
 {
 	Super::BeginDestroy();
@@ -344,3 +396,5 @@ void UNiagaraMeshRendererProperties::CheckMaterialUsage()
 }
 
 #endif // WITH_EDITORONLY_DATA
+
+#undef LOCTEXT_NAMESPACE
\ No newline at end of file
diff --git a/Engine/Plugins/FX/Niagara/Source/Niagara/Private/NiagaraRibbonRendererProperties.cpp b/Engine/Plugins/FX/Niagara/Source/Niagara/Private/NiagaraRibbonRendererProperties.cpp
--- a/Engine/Plugins/FX/Niagara/Source/Niagara/Private/NiagaraRibbonRendererProperties.cpp
+++ b/Engine/Plugins/FX/Niagara/Source/Niagara/Private/NiagaraRibbonRendererProperties.cpp
@@ -5,6 +5,16 @@
 #include "NiagaraConstants.h"
 #include "NiagaraBoundsCalculatorHelper.h"
 #include "Modules/ModuleManager.h"
+#if WITH_EDITOR
+#include "Widgets/Images/SImage.h"
+#include "Styling/SlateIconFinder.h"
+#include "Widgets/SWidget.h"
+#include "Styling/SlateBrush.h"
+#include "Widgets/Text/STextBlock.h"
+#endif
+
+
+#define LOCTEXT_NAMESPACE "UNiagaraRibbonRendererProperties"
 
 TArray<TWeakObjectPtr<UNiagaraRibbonRendererProperties>> UNiagaraRibbonRendererProperties::RibbonRendererPropertiesToDeferredInit;
 
@@ -168,6 +178,45 @@ const TArray<FNiagaraVariable>& UNiagaraRibbonRendererProperties::GetOptionalAtt
 }
 
 
+void UNiagaraRibbonRendererProperties::GetRendererWidgets(const FNiagaraEmitterInstance* InEmitter, TArray<TSharedPtr<SWidget>>& OutWidgets, TSharedPtr<FAssetThumbnailPool> InThumbnailPool) const
+{
+	TSharedRef<SWidget> ThumbnailWidget = SNullWidget::NullWidget;
+	int32 ThumbnailSize = 32;
+	TArray<UMaterialInterface*> Materials;
+	GetUsedMaterials(InEmitter, Materials);
+	for (UMaterialInterface* PreviewedMaterial : Materials)
+	{
+		TSharedPtr<FAssetThumbnail> AssetThumbnail = MakeShareable(new FAssetThumbnail(PreviewedMaterial, ThumbnailSize, ThumbnailSize, InThumbnailPool));
+		if (AssetThumbnail)
+		{
+			ThumbnailWidget = AssetThumbnail->MakeThumbnailWidget();
+		}
+		OutWidgets.Add(ThumbnailWidget);
+	}
+
+	if (Materials.Num() == 0)
+	{
+		TSharedRef<SWidget> SpriteWidget = SNew(SImage)
+			.Image(FSlateIconFinder::FindIconBrushForClass(GetClass()));
+		OutWidgets.Add(SpriteWidget);
+	}
+}
+
+void UNiagaraRibbonRendererProperties::GetRendererTooltipWidgets(const FNiagaraEmitterInstance* InEmitter, TArray<TSharedPtr<SWidget>>& OutWidgets, TSharedPtr<FAssetThumbnailPool> InThumbnailPool) const
+{
+	TArray<UMaterialInterface*> Materials;
+	GetUsedMaterials(InEmitter, Materials);
+	if (Materials.Num() > 0)
+	{
+		GetRendererWidgets(InEmitter, OutWidgets, InThumbnailPool);
+	}
+	else
+	{
+		TSharedRef<SWidget> RibbonTooltip = SNew(STextBlock)
+			.Text(LOCTEXT("RibbonRendererNoMat", "Ribbon Renderer (No Material Set)"));
+		OutWidgets.Add(RibbonTooltip);
+	}
+}
 
 
 bool UNiagaraRibbonRendererProperties::IsMaterialValidForRenderer(UMaterial* InMaterial, FText& InvalidMessage)
@@ -211,3 +260,4 @@ bool UNiagaraRibbonRendererProperties::CanEditChange(const FProperty* InProperty
 }
 
 #endif // WITH_EDITORONLY_DATA
+#undef LOCTEXT_NAMESPACE
\ No newline at end of file
diff --git a/Engine/Plugins/FX/Niagara/Source/Niagara/Private/NiagaraSpriteRendererProperties.cpp b/Engine/Plugins/FX/Niagara/Source/Niagara/Private/NiagaraSpriteRendererProperties.cpp
--- a/Engine/Plugins/FX/Niagara/Source/Niagara/Private/NiagaraSpriteRendererProperties.cpp
+++ b/Engine/Plugins/FX/Niagara/Source/Niagara/Private/NiagaraSpriteRendererProperties.cpp
@@ -9,7 +9,14 @@
 #include "Modules/ModuleManager.h"
 #if WITH_EDITOR
 #include "DerivedDataCacheInterface.h"
+#include "Widgets/Images/SImage.h"
+#include "Styling/SlateIconFinder.h"
+#include "Widgets/SWidget.h"
+#include "Styling/SlateBrush.h"
+#include "AssetThumbnail.h"
+#include "Widgets/Text/STextBlock.h"
 #endif
+
 #define LOCTEXT_NAMESPACE "UNiagaraSpriteRendererProperties"
 
 TArray<TWeakObjectPtr<UNiagaraSpriteRendererProperties>> UNiagaraSpriteRendererProperties::SpriteRendererPropertiesToDeferredInit;
@@ -246,6 +253,46 @@ const TArray<FNiagaraVariable>& UNiagaraSpriteRendererProperties::GetOptionalAtt
 	return Attrs;
 }
 
+void UNiagaraSpriteRendererProperties::GetRendererWidgets(const FNiagaraEmitterInstance* InEmitter, TArray<TSharedPtr<SWidget>>& OutWidgets, TSharedPtr<FAssetThumbnailPool> InThumbnailPool) const
+{
+	TSharedRef<SWidget> ThumbnailWidget = SNullWidget::NullWidget;
+	int32 ThumbnailSize = 32;
+	TArray<UMaterialInterface*> Materials;
+	GetUsedMaterials(InEmitter, Materials);
+	for (UMaterialInterface* PreviewedMaterial : Materials)
+	{
+		TSharedPtr<FAssetThumbnail> AssetThumbnail = MakeShareable(new FAssetThumbnail(PreviewedMaterial, ThumbnailSize, ThumbnailSize, InThumbnailPool));
+		if (AssetThumbnail)
+		{
+			ThumbnailWidget = AssetThumbnail->MakeThumbnailWidget();
+		}
+		OutWidgets.Add(ThumbnailWidget);
+	}
+
+	if (Materials.Num() == 0)
+	{
+		TSharedRef<SWidget> SpriteWidget = SNew(SImage)
+			.Image(FSlateIconFinder::FindIconBrushForClass(GetClass()));
+		OutWidgets.Add(SpriteWidget);
+	}
+}
+
+void UNiagaraSpriteRendererProperties::GetRendererTooltipWidgets(const FNiagaraEmitterInstance* InEmitter, TArray<TSharedPtr<SWidget>>& OutWidgets, TSharedPtr<FAssetThumbnailPool> InThumbnailPool) const
+{
+	TArray<UMaterialInterface*> Materials;
+	GetUsedMaterials(InEmitter, Materials);
+	if (Materials.Num() > 0)
+	{
+		GetRendererWidgets(InEmitter, OutWidgets, InThumbnailPool);
+	}
+	else
+	{
+		TSharedRef<SWidget> SpriteTooltip = SNew(STextBlock)
+			.Text(LOCTEXT("SpriteRendererNoMat", "Sprite Renderer (No Material Set)"));
+		OutWidgets.Add(SpriteTooltip);
+	}
+}
+
 bool UNiagaraSpriteRendererProperties::IsMaterialValidForRenderer(UMaterial* InMaterial, FText& InvalidMessage)
 {
 	if (InMaterial->bUsedWithNiagaraSprites == false)
diff --git a/Engine/Plugins/FX/Niagara/Source/Niagara/Public/NiagaraLightRendererProperties.h b/Engine/Plugins/FX/Niagara/Source/Niagara/Public/NiagaraLightRendererProperties.h
--- a/Engine/Plugins/FX/Niagara/Source/Niagara/Public/NiagaraLightRendererProperties.h
+++ b/Engine/Plugins/FX/Niagara/Source/Niagara/Public/NiagaraLightRendererProperties.h
@@ -9,6 +9,7 @@
 #include "NiagaraLightRendererProperties.generated.h"
 
 class FNiagaraEmitterInstance;
+class SWidget;
 
 UCLASS(editinlinenew, meta = (DisplayName = "Light Renderer"))
 class UNiagaraLightRendererProperties : public UNiagaraRendererProperties
@@ -34,6 +35,8 @@ public:
 	virtual void FixMaterial(UMaterial* Material) override;
 	virtual const TArray<FNiagaraVariable>& GetRequiredAttributes() override;
 	virtual const TArray<FNiagaraVariable>& GetOptionalAttributes() override;
+	virtual void GetRendererWidgets(const FNiagaraEmitterInstance* InEmitter, TArray<TSharedPtr<SWidget>>& OutWidgets, TSharedPtr<FAssetThumbnailPool> InThumbnailPool) const override;
+	virtual void GetRendererTooltipWidgets(const FNiagaraEmitterInstance* InEmitter, TArray<TSharedPtr<SWidget>>& OutWidgets, TSharedPtr<FAssetThumbnailPool> InThumbnailPool) const override;
 #endif // WITH_EDITORONLY_DATA
 
 	/** Whether to use physically based inverse squared falloff from the light.  If unchecked, the value from the LightExponent binding will be used instead. */
diff --git a/Engine/Plugins/FX/Niagara/Source/Niagara/Public/NiagaraMeshRendererProperties.h b/Engine/Plugins/FX/Niagara/Source/Niagara/Public/NiagaraMeshRendererProperties.h
--- a/Engine/Plugins/FX/Niagara/Source/Niagara/Public/NiagaraMeshRendererProperties.h
+++ b/Engine/Plugins/FX/Niagara/Source/Niagara/Public/NiagaraMeshRendererProperties.h
@@ -10,6 +10,8 @@
 #include "NiagaraMeshRendererProperties.generated.h"
 
 class FNiagaraEmitterInstance;
+class FAssetThumbnailPool;
+class SWidget;
 
 /** This enum decides how a mesh particle will orient its "facing" axis relative to camera. Must keep these in sync with NiagaraMeshVertexFactory.ush*/
 UENUM()
@@ -85,7 +87,8 @@ public:
 	virtual void FixMaterial(UMaterial* Material) override;
 	virtual const TArray<FNiagaraVariable>& GetRequiredAttributes() override;
 	virtual const TArray<FNiagaraVariable>& GetOptionalAttributes() override;
-
+	virtual	void GetRendererWidgets(const FNiagaraEmitterInstance* InEmitter, TArray<TSharedPtr<SWidget>>& OutWidgets, TSharedPtr<FAssetThumbnailPool> InThumbnailPool) const override;
+	virtual	void GetRendererTooltipWidgets(const FNiagaraEmitterInstance* InEmitter, TArray<TSharedPtr<SWidget>>& OutWidgets, TSharedPtr<FAssetThumbnailPool> InThumbnailPool) const override;
 	void OnMeshChanged();
 	void CheckMaterialUsage();
 #endif // WITH_EDITORONLY_DATA
diff --git a/Engine/Plugins/FX/Niagara/Source/Niagara/Public/NiagaraRendererProperties.h b/Engine/Plugins/FX/Niagara/Source/Niagara/Public/NiagaraRendererProperties.h
--- a/Engine/Plugins/FX/Niagara/Source/Niagara/Public/NiagaraRendererProperties.h
+++ b/Engine/Plugins/FX/Niagara/Source/Niagara/Public/NiagaraRendererProperties.h
@@ -23,6 +23,8 @@ class FNiagaraRenderer;
 class UMaterial;
 class UMaterialInterface;
 class FNiagaraEmitterInstance;
+class SWidget;
+class FAssetThumbnailPool;
 
 UCLASS(ABSTRACT)
 class NIAGARA_API UNiagaraRendererProperties : public UNiagaraMergeable
@@ -55,6 +57,9 @@ public:
 	{
 		return CastChecked<UNiagaraRendererProperties>(Super::StaticDuplicateWithNewMergeIdInternal(InOuter));
 	}
+
+	virtual void GetRendererWidgets(const FNiagaraEmitterInstance* InEmitter, TArray<TSharedPtr<SWidget>>& OutWidgets, TSharedPtr<FAssetThumbnailPool> InThumbnailPool) const PURE_VIRTUAL(UNiagaraRendererProperties::GetRendererWidgets, );
+	virtual void GetRendererTooltipWidgets(const FNiagaraEmitterInstance* InEmitter, TArray<TSharedPtr<SWidget>>& OutWidgets, TSharedPtr<FAssetThumbnailPool> InThumbnailPool) const PURE_VIRTUAL(UNiagaraRendererProperties::GetRendererTooltipWidgets, );
 #endif // WITH_EDITORONLY_DATA
 
 
diff --git a/Engine/Plugins/FX/Niagara/Source/Niagara/Public/NiagaraRibbonRendererProperties.h b/Engine/Plugins/FX/Niagara/Source/Niagara/Public/NiagaraRibbonRendererProperties.h
--- a/Engine/Plugins/FX/Niagara/Source/Niagara/Public/NiagaraRibbonRendererProperties.h
+++ b/Engine/Plugins/FX/Niagara/Source/Niagara/Public/NiagaraRibbonRendererProperties.h
@@ -9,6 +9,8 @@
 #include "NiagaraRibbonRendererProperties.generated.h"
 
 class FNiagaraEmitterInstance;
+class FAssetThumbnailPool;
+class SWidget;
 
 UENUM()
 enum class ENiagaraRibbonFacingMode : uint8
@@ -83,6 +85,8 @@ public:
 	virtual void FixMaterial(UMaterial* Material);
 	virtual const TArray<FNiagaraVariable>& GetRequiredAttributes() override;
 	virtual const TArray<FNiagaraVariable>& GetOptionalAttributes() override;
+	virtual void GetRendererWidgets(const FNiagaraEmitterInstance* InEmitter, TArray<TSharedPtr<SWidget>>& OutWidgets, TSharedPtr<FAssetThumbnailPool> InThumbnailPool) const override;
+	virtual void GetRendererTooltipWidgets(const FNiagaraEmitterInstance* InEmitter, TArray<TSharedPtr<SWidget>>& OutWidgets, TSharedPtr<FAssetThumbnailPool> InThumbnailPool) const override;
 #endif
 	//UNiagaraRendererProperties Interface END
 
diff --git a/Engine/Plugins/FX/Niagara/Source/Niagara/Public/NiagaraSpriteRendererProperties.h b/Engine/Plugins/FX/Niagara/Source/Niagara/Public/NiagaraSpriteRendererProperties.h
--- a/Engine/Plugins/FX/Niagara/Source/Niagara/Public/NiagaraSpriteRendererProperties.h
+++ b/Engine/Plugins/FX/Niagara/Source/Niagara/Public/NiagaraSpriteRendererProperties.h
@@ -38,6 +38,8 @@ enum class ENiagaraSpriteFacingMode : uint8
 	FaceCameraDistanceBlend
 };
 
+class FAssetThumbnailPool;
+class SWidget;
 
 UCLASS(editinlinenew, meta = (DisplayName = "Sprite Renderer"))
 class NIAGARA_API UNiagaraSpriteRendererProperties : public UNiagaraRendererProperties
@@ -68,6 +70,8 @@ public:
 	virtual void FixMaterial(UMaterial* Material) override;
 	virtual const TArray<FNiagaraVariable>& GetRequiredAttributes() override;
 	virtual const TArray<FNiagaraVariable>& GetOptionalAttributes() override;
+	virtual void GetRendererWidgets(const FNiagaraEmitterInstance* InEmitter, TArray<TSharedPtr<SWidget>>& OutWidgets, TSharedPtr<FAssetThumbnailPool> InThumbnailPool) const override;
+	virtual void GetRendererTooltipWidgets(const FNiagaraEmitterInstance* InEmitter, TArray<TSharedPtr<SWidget>>& OutWidgets, TSharedPtr<FAssetThumbnailPool> InThumbnailPool) const override;
 #endif
 	//UNiagaraMaterialRendererProperties interface END
 
diff --git a/Engine/Plugins/FX/Niagara/Source/NiagaraEditor/Private/ViewModels/NiagaraEmitterHandleViewModel.cpp b/Engine/Plugins/FX/Niagara/Source/NiagaraEditor/Private/ViewModels/NiagaraEmitterHandleViewModel.cpp
--- a/Engine/Plugins/FX/Niagara/Source/NiagaraEditor/Private/ViewModels/NiagaraEmitterHandleViewModel.cpp
+++ b/Engine/Plugins/FX/Niagara/Source/NiagaraEditor/Private/ViewModels/NiagaraEmitterHandleViewModel.cpp
@@ -48,31 +48,6 @@ void FNiagaraEmitterHandleViewModel::Cleanup()
 	}
 }
 
-void FNiagaraEmitterHandleViewModel::GetRendererPreviewData(TArray<FRendererPreviewData*>& InRendererPreviewData)
-{
-	InRendererPreviewData.Empty();
-	UNiagaraEmitter* Emitter = GetEmitterHandle()->GetInstance();
-	UNiagaraStackRoot* StackRoot = Cast<UNiagaraStackRoot>(EmitterStackViewModel->GetRootEntry());
-	if (StackRoot)
-	{
-		TArray<UNiagaraStackEntry*> Children;
-		StackRoot->GetRenderGroup()->GetUnfilteredChildren(Children);
-		for (UNiagaraStackEntry* Child : Children)
-		{
-			if (UNiagaraStackRendererItem* RendererItem = Cast<UNiagaraStackRendererItem>(Child))
-			{
-				TArray<UMaterialInterface*> Materials;
-				FNiagaraEmitterInstance* InInstance = GetEmitterViewModel()->GetSimulation().IsValid() ? GetEmitterViewModel()->GetSimulation().Pin().Get() : nullptr;
-				RendererItem->GetRendererProperties()->GetUsedMaterials(InInstance, Materials);
-				for (UMaterialInterface* Material : Materials)
-				{
-					InRendererPreviewData.Add(new FRendererPreviewData(Child, Material));
-				}
-			}
-		}
-	}
-}
-
 void FNiagaraEmitterHandleViewModel::GetRendererEntries(TArray<UNiagaraStackEntry*>& InRenderingEntries)
 {
 	InRenderingEntries.Empty();
diff --git a/Engine/Plugins/FX/Niagara/Source/NiagaraEditor/Public/ViewModels/NiagaraEmitterHandleViewModel.h b/Engine/Plugins/FX/Niagara/Source/NiagaraEditor/Public/ViewModels/NiagaraEmitterHandleViewModel.h
--- a/Engine/Plugins/FX/Niagara/Source/NiagaraEditor/Public/ViewModels/NiagaraEmitterHandleViewModel.h
+++ b/Engine/Plugins/FX/Niagara/Source/NiagaraEditor/Public/ViewModels/NiagaraEmitterHandleViewModel.h
@@ -17,20 +17,6 @@ class UNiagaraStackViewModel;
 class UNiagaraStackEntry;
 enum class ENiagaraSystemViewModelEditMode;
 
-struct FRendererPreviewData
-{
-	UNiagaraStackEntry* RenderingEntry;
-	UObject* RenderingObject;
-
-
-	FRendererPreviewData(UNiagaraStackEntry* InRenderingEntry, UObject* InRenderingObject)
-		: RenderingEntry(InRenderingEntry)
-		, RenderingObject(InRenderingObject)
-	{
-
-	};
-};
-
 /** The view model for the FNiagaraEmitterEditorWidget. */
 class FNiagaraEmitterHandleViewModel : public TSharedFromThis<FNiagaraEmitterHandleViewModel>, public FGCObject
 {
@@ -131,9 +117,6 @@ public:
 	FOnNameChanged& OnNameChanged();
 
 	void Cleanup();
-
-	NIAGARAEDITOR_API void GetRendererPreviewData(TArray<FRendererPreviewData*>& InRendererPreviewData);
-
 	NIAGARAEDITOR_API void GetRendererEntries(TArray<UNiagaraStackEntry*>& InRenderingEntries);
 	NIAGARAEDITOR_API TSharedRef<FNiagaraSystemViewModel> GetOwningSystemViewModel() const;
 
diff --git a/Engine/Plugins/FX/Niagara/Source/NiagaraEditor/Public/ViewModels/NiagaraEmitterViewModel.h b/Engine/Plugins/FX/Niagara/Source/NiagaraEditor/Public/ViewModels/NiagaraEmitterViewModel.h
--- a/Engine/Plugins/FX/Niagara/Source/NiagaraEditor/Public/ViewModels/NiagaraEmitterViewModel.h
+++ b/Engine/Plugins/FX/Niagara/Source/NiagaraEditor/Public/ViewModels/NiagaraEmitterViewModel.h
@@ -40,7 +40,7 @@ public:
 	void Reset();
 
 	/** Gets the currently assigned simulation if there is one. */
-	TWeakPtr<FNiagaraEmitterInstance, ESPMode::ThreadSafe> GetSimulation() const;
+	NIAGARAEDITOR_API TWeakPtr<FNiagaraEmitterInstance, ESPMode::ThreadSafe> GetSimulation() const;
 
 	/** Sets the current simulation for the emitter. */
 	void SetSimulation(TWeakPtr<FNiagaraEmitterInstance, ESPMode::ThreadSafe> InSimulation);
@@ -77,7 +77,7 @@ public:
 	FOnEmitterChanged& OnEmitterChanged();
 
 	/** Gets a delegate which is called when a property on the emitter changes. */
-	FOnPropertyChanged& OnPropertyChanged();
+	NIAGARAEDITOR_API FOnPropertyChanged& OnPropertyChanged();
 
 	/** Gets a delegate which is called when the shared script is compiled. */
 	FOnScriptCompiled& OnScriptCompiled();
diff --git a/Engine/Plugins/FX/Niagara/Source/NiagaraEditorWidgets/Private/SNiagaraOverviewStackNode.cpp b/Engine/Plugins/FX/Niagara/Source/NiagaraEditorWidgets/Private/SNiagaraOverviewStackNode.cpp
--- a/Engine/Plugins/FX/Niagara/Source/NiagaraEditorWidgets/Private/SNiagaraOverviewStackNode.cpp
+++ b/Engine/Plugins/FX/Niagara/Source/NiagaraEditorWidgets/Private/SNiagaraOverviewStackNode.cpp
@@ -25,6 +25,7 @@
 #include "NiagaraRendererProperties.h"
 #include "Widgets/SWidget.h"
 #include "Subsystems/AssetEditorSubsystem.h"
+#include "ViewModels/Stack/NiagaraStackRendererItem.h"
 
 #define LOCTEXT_NAMESPACE "NiagaraOverviewStackNode"
 
@@ -39,7 +40,7 @@ void SNiagaraOverviewStackNode::Construct(const FArguments& InArgs, UNiagaraOver
 	CurrentIssueIndex = -1;
 
 	EmitterHandleViewModelWeak.Reset();
-	ThumbnailPool = MakeShared<FAssetThumbnailPool>(100, TAttribute<bool>::Create(TAttribute<bool>::FGetter::CreateSP(this, &SNiagaraOverviewStackNode::IsHoveringThumbnail)));
+	ThumbnailPool = MakeShared<FAssetThumbnailPool>(10, TAttribute<bool>::Create(TAttribute<bool>::FGetter::CreateSP(this, &SNiagaraOverviewStackNode::IsHoveringThumbnail)));
 	ThumbnailBar = SNew(SHorizontalBox);
 	if (OverviewStackNode->GetOwningSystem() != nullptr)
 	{
@@ -182,32 +183,40 @@ FReply SNiagaraOverviewStackNode::OnCycleThroughIssues()
 	return FReply::Handled();
 }
 
-TSharedRef<SWidget> SNiagaraOverviewStackNode::CreateThumbnailWidget(float InThumbnailSize, FRendererPreviewData* InData)
+TSharedRef<SWidget> SNiagaraOverviewStackNode::CreateThumbnailWidget(UNiagaraStackEntry* InData, TSharedPtr<SWidget> InWidget, TSharedPtr<SWidget> InTooltipWidget)
 {
-	TSharedRef<SWidget> ThumbnailWidget = SNullWidget::NullWidget;
-	TSharedPtr<FAssetThumbnail> AssetThumbnail = MakeShareable(new FAssetThumbnail(InData->RenderingObject, InThumbnailSize, InThumbnailSize, ThumbnailPool));
-	if (AssetThumbnail)
+	TSharedPtr<SToolTip> ThumbnailTooltipWidget;
+	// If this is just text, don't constrain the size
+	if (InTooltipWidget->GetType() == TEXT("STextBlock"))
 	{
-		ThumbnailWidget = AssetThumbnail->MakeThumbnailWidget();
-		TSharedPtr<FAssetThumbnail> TooltipThumbnail = MakeShareable(new FAssetThumbnail(InData->RenderingObject, 64.0f, 64.0f, ThumbnailPool));
-		TSharedRef<SToolTip> ThumbnailTooltipWidget = SNew(SToolTip)
+		ThumbnailTooltipWidget = SNew(SToolTip)
+			.Content()
+			[
+				InTooltipWidget.ToSharedRef()
+			];
+	}
+	else
+	{
+
+		ThumbnailTooltipWidget = SNew(SToolTip)
 			.Content()
 			[
 				SNew(SBox)
-				.MinDesiredHeight(64.0f)
 				.MaxDesiredHeight(64.0f)
-				.MinDesiredWidth(64.0f)
+				.MinDesiredHeight(64.0f)
 				.MaxDesiredWidth(64.0f)
+				.MinDesiredWidth(64.0f)
 				[
-					TooltipThumbnail->MakeThumbnailWidget()
+					InTooltipWidget.ToSharedRef()
 				]
 			];
-		ThumbnailWidget->SetOnMouseButtonDown(FPointerEventHandler::CreateSP(this, &SNiagaraOverviewStackNode::OnClickedRenderingPreview, InData->RenderingEntry));
-		ThumbnailWidget->SetOnMouseEnter(FNoReplyPointerEventHandler::CreateSP(this, &SNiagaraOverviewStackNode::SetIsHoveringThumbnail, true));
-		ThumbnailWidget->SetOnMouseLeave(FSimpleNoReplyPointerEventHandler::CreateSP(this, &SNiagaraOverviewStackNode::SetIsHoveringThumbnail, false));
-		ThumbnailWidget->SetToolTip(ThumbnailTooltipWidget);
 	}
-	return ThumbnailWidget;
+	InWidget->SetOnMouseButtonDown(FPointerEventHandler::CreateSP(this, &SNiagaraOverviewStackNode::OnClickedRenderingPreview, InData));
+	InWidget->SetOnMouseEnter(FNoReplyPointerEventHandler::CreateSP(this, &SNiagaraOverviewStackNode::SetIsHoveringThumbnail, true));
+	InWidget->SetOnMouseLeave(FSimpleNoReplyPointerEventHandler::CreateSP(this, &SNiagaraOverviewStackNode::SetIsHoveringThumbnail, false));
+	InWidget->SetToolTip(ThumbnailTooltipWidget);
+
+	return InWidget.ToSharedRef();
 }
 
 FReply SNiagaraOverviewStackNode::OnClickedRenderingPreview(const FGeometry& InGeometry, const FPointerEvent& InEvent, UNiagaraStackEntry* InEntry)
@@ -323,8 +332,7 @@ void SNiagaraOverviewStackNode::FillThumbnailBar(UObject* ChangedObject, const b
 		}
 		if (EmitterHandleViewModelWeak.IsValid())
 		{
-			EmitterHandleViewModelWeak.Pin()->GetRendererPreviewData(PreviewData);
-
+	
 			// Isolate toggle button
 			ThumbnailBar->AddSlot()
 				.AutoWidth()
@@ -347,26 +355,35 @@ void SNiagaraOverviewStackNode::FillThumbnailBar(UObject* ChangedObject, const b
 					]
 				];
 
-			for (FRendererPreviewData* Preview : PreviewData)
+			EmitterHandleViewModelWeak.Pin()->GetRendererEntries(PreviewStackEntries);
+			FNiagaraEmitterInstance* InInstance = EmitterHandleViewModelWeak.Pin()->GetEmitterViewModel()->GetSimulation().IsValid() ? EmitterHandleViewModelWeak.Pin()->GetEmitterViewModel()->GetSimulation().Pin().Get() : nullptr;
+			for (UNiagaraStackEntry* Entry : PreviewStackEntries)
 			{
-				if (Preview->RenderingObject)
+				if (UNiagaraStackRendererItem* RendererItem = Cast<UNiagaraStackRendererItem>(Entry))
 				{
-					ThumbnailBar->AddSlot()
-						.AutoWidth()
-						.HAlign(HAlign_Left)
-						.VAlign(VAlign_Center)
-						.Padding(2.0f, 2.0f, 4.0f, 4.0f)
-						.MaxWidth(ThumbnailSize)
-						[
-							SNew(SBox)
-							.MinDesiredHeight(ThumbnailSize)
-							.MaxDesiredHeight(ThumbnailSize)
-							.MinDesiredWidth(ThumbnailSize)
-							.Visibility(this, &SNiagaraOverviewStackNode::GetEnabledCheckBoxVisibility)
+					TArray<TSharedPtr<SWidget>> Widgets;
+					RendererItem->GetRendererProperties()->GetRendererWidgets(InInstance, Widgets, ThumbnailPool);
+					TArray<TSharedPtr<SWidget>> TooltipWidgets;
+					RendererItem->GetRendererProperties()->GetRendererTooltipWidgets(InInstance, TooltipWidgets, ThumbnailPool);
+					for (int32 WidgetIndex = 0; WidgetIndex < Widgets.Num(); WidgetIndex++)
+					{
+						ThumbnailBar->AddSlot()
+							.AutoWidth()
+							.HAlign(HAlign_Left)
+							.VAlign(VAlign_Center)
+							.Padding(2.0f, 2.0f, 4.0f, 4.0f)
+							.MaxWidth(ThumbnailSize)
 							[
-								CreateThumbnailWidget(ThumbnailSize, Preview)
-							]
-						];
+								SNew(SBox)
+								.MinDesiredHeight(ThumbnailSize)
+								.MaxDesiredHeight(ThumbnailSize)
+								.MinDesiredWidth(ThumbnailSize)
+								.Visibility(this, &SNiagaraOverviewStackNode::GetEnabledCheckBoxVisibility)
+								[
+									CreateThumbnailWidget(Entry, Widgets[WidgetIndex], TooltipWidgets[WidgetIndex])
+								]
+							];
+					}
 				}
 			}
 		}
diff --git a/Engine/Plugins/FX/Niagara/Source/NiagaraEditorWidgets/Private/SNiagaraOverviewStackNode.h b/Engine/Plugins/FX/Niagara/Source/NiagaraEditorWidgets/Private/SNiagaraOverviewStackNode.h
--- a/Engine/Plugins/FX/Niagara/Source/NiagaraEditorWidgets/Private/SNiagaraOverviewStackNode.h
+++ b/Engine/Plugins/FX/Niagara/Source/NiagaraEditorWidgets/Private/SNiagaraOverviewStackNode.h
@@ -12,7 +12,7 @@ class UNiagaraSystemSelectionViewModel;
 class FNiagaraEmitterHandleViewModel;
 class FAssetThumbnailPool;
 class FAssetThumbnail;
-struct FRendererPreviewData;
+class UNiagaraStackEntry;
 
 class SNiagaraOverviewStackNode : public SGraphNode
 {
@@ -35,7 +35,7 @@ private:
 	EVisibility GetEnabledCheckBoxVisibility() const;
 	ECheckBoxState GetEnabledCheckState() const;
 	void OnEnabledCheckStateChanged(ECheckBoxState InCheckState);
-	TSharedRef<SWidget> CreateThumbnailWidget(float InThumbnailSize, FRendererPreviewData* InData);
+	TSharedRef<SWidget> CreateThumbnailWidget(UNiagaraStackEntry* InData, TSharedPtr<SWidget> InWidget, TSharedPtr<SWidget> InTooltipWidget);
 	FReply OnClickedRenderingPreview(const FGeometry& InGeometry, const FPointerEvent& InEvent, class UNiagaraStackEntry* InEntry);
 	FText GetToggleIsolateToolTip() const;
 	FReply OnToggleIsolateButtonClicked();
@@ -66,7 +66,7 @@ private:
 	TSharedPtr<FAssetThumbnailPool> ThumbnailPool;
 	/** Thumbnail widget containers */
 	TSharedPtr<SHorizontalBox> ThumbnailBar;
-	TArray<FRendererPreviewData*> PreviewData;
+	TArray<UNiagaraStackEntry*> PreviewStackEntries;
 	bool bIsHoveringThumbnail;
 	int32 CurrentIssueIndex;
 };
\ No newline at end of file
