Adding missing header
#rnx

diff --git a/Engine/Plugins/FX/Niagara/Source/Niagara/Private/NiagaraRibbonRendererProperties.cpp b/Engine/Plugins/FX/Niagara/Source/Niagara/Private/NiagaraRibbonRendererProperties.cpp
--- a/Engine/Plugins/FX/Niagara/Source/Niagara/Private/NiagaraRibbonRendererProperties.cpp
+++ b/Engine/Plugins/FX/Niagara/Source/Niagara/Private/NiagaraRibbonRendererProperties.cpp
@@ -10,10 +10,12 @@
 #include "Styling/SlateIconFinder.h"
 #include "Widgets/SWidget.h"
 #include "Styling/SlateBrush.h"
+#include "AssetThumbnail.h"
 #include "Widgets/Text/STextBlock.h"
 #endif
 
 
+
 #define LOCTEXT_NAMESPACE "UNiagaraRibbonRendererProperties"
 
 TArray<TWeakObjectPtr<UNiagaraRibbonRendererProperties>> UNiagaraRibbonRendererProperties::RibbonRendererPropertiesToDeferredInit;
