Fixed int32 to bool conversion warning.

diff --git a/Engine/Plugins/FX/Niagara/Source/Niagara/Private/NiagaraRendererMeshes.cpp b/Engine/Plugins/FX/Niagara/Source/Niagara/Private/NiagaraRendererMeshes.cpp
--- a/Engine/Plugins/FX/Niagara/Source/Niagara/Private/NiagaraRendererMeshes.cpp
+++ b/Engine/Plugins/FX/Niagara/Source/Niagara/Private/NiagaraRendererMeshes.cpp
@@ -355,7 +355,7 @@ void FNiagaraRendererMeshes::GetDynamicMeshElements(const TArray<const FSceneVie
 				{
 					SortInfo.ParticleCount = NumInstances;
 					SortInfo.SortMode = SortMode;
-					SortInfo.SetSortFlags(GNiagaraGPUSortingUseMaxPrecision, bHasTranslucentMaterials);
+					SortInfo.SetSortFlags(GNiagaraGPUSortingUseMaxPrecision != 0, bHasTranslucentMaterials); 
 					if (bCustomSorting)
 					{
 						SortVarIdx = ENiagaraMeshVFLayout::CustomSorting;
diff --git a/Engine/Plugins/FX/Niagara/Source/Niagara/Private/NiagaraRendererSprites.cpp b/Engine/Plugins/FX/Niagara/Source/Niagara/Private/NiagaraRendererSprites.cpp
--- a/Engine/Plugins/FX/Niagara/Source/Niagara/Private/NiagaraRendererSprites.cpp
+++ b/Engine/Plugins/FX/Niagara/Source/Niagara/Private/NiagaraRendererSprites.cpp
@@ -340,7 +340,7 @@ void FNiagaraRendererSprites::SetVertexFactoryParticleData(
 		{
 			SortInfo.ParticleCount = NumInstances;
 			SortInfo.SortMode = SortMode;
-			SortInfo.SetSortFlags(GNiagaraGPUSortingUseMaxPrecision, bHasTranslucentMaterials); 
+			SortInfo.SetSortFlags(GNiagaraGPUSortingUseMaxPrecision != 0, bHasTranslucentMaterials); 
 			if (SortInfo.SortMode == ENiagaraSortMode::CustomAscending || SortInfo.SortMode == ENiagaraSortMode::CustomDecending)
 			{
 				SortingVarIdx = ENiagaraSpriteVFLayout::CustomSorting;	
