Fixed a bug where changing the type of a static switch could result in a crash

diff --git a/Engine/Plugins/FX/Niagara/Source/NiagaraEditor/Private/NiagaraNodeStaticSwitch.cpp b/Engine/Plugins/FX/Niagara/Source/NiagaraEditor/Private/NiagaraNodeStaticSwitch.cpp
--- a/Engine/Plugins/FX/Niagara/Source/NiagaraEditor/Private/NiagaraNodeStaticSwitch.cpp
+++ b/Engine/Plugins/FX/Niagara/Source/NiagaraEditor/Private/NiagaraNodeStaticSwitch.cpp
@@ -48,7 +48,10 @@ void UNiagaraNodeStaticSwitch::OnSwitchParameterTypeChanged(const FNiagaraTypeDe
 {
 	TOptional<FNiagaraVariableMetaData> OldMetaData = GetNiagaraGraph()->GetMetaData(FNiagaraVariable(OldType, InputParameterName));
 	RefreshFromExternalChanges(); // Magick happens here: The old pins are destroyed and new ones are created.
-	GetNiagaraGraph()->SetMetaData(FNiagaraVariable(GetInputType(), InputParameterName), OldMetaData.GetValue());
+	if (OldMetaData.IsSet())
+	{
+		GetNiagaraGraph()->SetMetaData(FNiagaraVariable(GetInputType(), InputParameterName), OldMetaData.GetValue());
+	}
 
 	VisualsChangedDelegate.Broadcast(this);
 	RemoveUnusedGraphParameter(FNiagaraVariable(OldType, InputParameterName));
