Added a way for data interfaces to specify inter-emitter dependencies.

FNiagaraSystemInstance uses this information to determine the correct order in which emitters need to be ticked.

diff --git a/Engine/Plugins/FX/Niagara/Source/Niagara/Classes/NiagaraDataInterface.h b/Engine/Plugins/FX/Niagara/Source/Niagara/Classes/NiagaraDataInterface.h
--- a/Engine/Plugins/FX/Niagara/Source/Niagara/Classes/NiagaraDataInterface.h
+++ b/Engine/Plugins/FX/Niagara/Source/Niagara/Classes/NiagaraDataInterface.h
@@ -15,6 +15,7 @@ class UCurveVector;
 class UCurveLinearColor;
 class UCurveFloat;
 class FNiagaraSystemInstance;
+class FNiagaraEmitterInstance;
 struct FNiagaraDataInterfaceProxy;
 
 struct FNDITransformHandlerNoop
@@ -304,6 +305,14 @@ public:
 		return Proxy.Get();
 	}
 
+	/**
+	* Allows a DI to specify data dependencies between emitters, so the system can ensure that the emitter instances are executed in the correct order.
+	* The Dependencies array may already contain items, and this method should only append to it.
+	*/
+	virtual void GetEmitterDependencies(void* PerInstanceData, FNiagaraSystemInstance* SystemInstance, TArray<FNiagaraEmitterInstance*>& Dependencies) const
+	{
+	}
+
 protected:
 	template<typename T>
 	T* GetProxyAs()
diff --git a/Engine/Plugins/FX/Niagara/Source/Niagara/Classes/NiagaraDataInterfaceParticleRead.h b/Engine/Plugins/FX/Niagara/Source/Niagara/Classes/NiagaraDataInterfaceParticleRead.h
--- a/Engine/Plugins/FX/Niagara/Source/Niagara/Classes/NiagaraDataInterfaceParticleRead.h
+++ b/Engine/Plugins/FX/Niagara/Source/Niagara/Classes/NiagaraDataInterfaceParticleRead.h
@@ -31,6 +31,7 @@ public:
 	virtual bool GetFunctionHLSL(const FNiagaraDataInterfaceGPUParamInfo& ParamInfo, const FNiagaraDataInterfaceGeneratedFunction& FunctionInfo, int FunctionInstanceIndex, FString& OutHLSL) override;
 	virtual FNiagaraDataInterfaceParametersCS* ConstructComputeParameters() const override;
 	virtual void ProvidePerInstanceDataForRenderThread(void* DataForRenderThread, void* PerInstanceData, const FNiagaraSystemInstanceID& SystemInstance) override;
+	virtual void GetEmitterDependencies(void* PerInstanceData, FNiagaraSystemInstance* SystemInstance, TArray<FNiagaraEmitterInstance*>& Dependencies) const override;
 	//UNiagaraDataInterface Interface End
 
 	void GetNumSpawnedParticles(FVectorVMContext& Context);
diff --git a/Engine/Plugins/FX/Niagara/Source/Niagara/Private/NiagaraDataInterfaceParticleRead.cpp b/Engine/Plugins/FX/Niagara/Source/Niagara/Private/NiagaraDataInterfaceParticleRead.cpp
--- a/Engine/Plugins/FX/Niagara/Source/Niagara/Private/NiagaraDataInterfaceParticleRead.cpp
+++ b/Engine/Plugins/FX/Niagara/Source/Niagara/Private/NiagaraDataInterfaceParticleRead.cpp
@@ -1325,4 +1325,13 @@ void UNiagaraDataInterfaceParticleRead::ProvidePerInstanceDataForRenderThread(vo
 	}
 }
 
+void UNiagaraDataInterfaceParticleRead::GetEmitterDependencies(void* PerInstanceData, FNiagaraSystemInstance* SystemInstance, TArray<FNiagaraEmitterInstance*>& Dependencies) const
+{
+	FNDIParticleRead_InstanceData* PIData = static_cast<FNDIParticleRead_InstanceData*>(PerInstanceData);
+	if (PIData && PIData->EmitterInstance)
+	{
+		Dependencies.Add(PIData->EmitterInstance);
+	}
+}
+
 #undef LOCTEXT_NAMESPACE
diff --git a/Engine/Plugins/FX/Niagara/Source/Niagara/Private/NiagaraScriptExecutionContext.cpp b/Engine/Plugins/FX/Niagara/Source/Niagara/Private/NiagaraScriptExecutionContext.cpp
--- a/Engine/Plugins/FX/Niagara/Source/Niagara/Private/NiagaraScriptExecutionContext.cpp
+++ b/Engine/Plugins/FX/Niagara/Source/Niagara/Private/NiagaraScriptExecutionContext.cpp
@@ -318,9 +318,9 @@ void FNiagaraGPUSystemTick::Init(FNiagaraSystemInstance* InSystemInstance)
 	// This is spawn rate as well as DataInterface per instance data and the ParameterData for the emitter.
 	// @todo Ideally we would only update DataInterface and ParameterData bits if they have changed.
 	uint32 InstanceIndex = 0;
-	for (int32 i = 0; i < InSystemInstance->GetEmitters().Num(); i++)
+	for (int32 EmitterIdx : InSystemInstance->GetEmitterExecutionOrder())
 	{
-		FNiagaraEmitterInstance* Emitter = &InSystemInstance->GetEmitters()[i].Get();
+		FNiagaraEmitterInstance* Emitter = &InSystemInstance->GetEmitters()[EmitterIdx].Get();
 
 		if (Emitter && Emitter->GetCachedEmitter()->SimTarget == ENiagaraSimTarget::GPUComputeSim && Emitter->GetGPUContext() != nullptr && Emitter->GetExecutionState() != ENiagaraExecutionState::Complete)
 		{
@@ -351,8 +351,8 @@ void FNiagaraGPUSystemTick::Init(FNiagaraSystemInstance* InSystemInstance)
 			FMemory::Memcpy(InstanceData->SystemParamData + sizeof(FNiagaraSystemParameters), &InSystemInstance->GetSystemParameters(true), sizeof(FNiagaraSystemParameters));
 			FMemory::Memcpy(InstanceData->OwnerParamData, &InSystemInstance->GetOwnerParameters(), sizeof(FNiagaraOwnerParameters));
 			FMemory::Memcpy(InstanceData->OwnerParamData + sizeof(FNiagaraOwnerParameters), &InSystemInstance->GetOwnerParameters(true), sizeof(FNiagaraOwnerParameters));
-			FMemory::Memcpy(InstanceData->EmitterParamData, &InSystemInstance->GetEmitterParameters(i), sizeof(FNiagaraEmitterParameters));
-			FMemory::Memcpy(InstanceData->EmitterParamData + sizeof(FNiagaraEmitterParameters), &InSystemInstance->GetEmitterParameters(i, true), sizeof(FNiagaraEmitterParameters));
+			FMemory::Memcpy(InstanceData->EmitterParamData, &InSystemInstance->GetEmitterParameters(EmitterIdx), sizeof(FNiagaraEmitterParameters));
+			FMemory::Memcpy(InstanceData->EmitterParamData + sizeof(FNiagaraEmitterParameters), &InSystemInstance->GetEmitterParameters(EmitterIdx, true), sizeof(FNiagaraEmitterParameters));
 
 			GPUContext->CombinedParamStore.CopyParameterDataToPaddedBuffer(InstanceData->ExternalParamData, ParmSize);
 
diff --git a/Engine/Plugins/FX/Niagara/Source/Niagara/Private/NiagaraSystemInstance.cpp b/Engine/Plugins/FX/Niagara/Source/Niagara/Private/NiagaraSystemInstance.cpp
--- a/Engine/Plugins/FX/Niagara/Source/Niagara/Private/NiagaraSystemInstance.cpp
+++ b/Engine/Plugins/FX/Niagara/Source/Niagara/Private/NiagaraSystemInstance.cpp
@@ -601,6 +601,134 @@ void FNiagaraSystemInstance::SetPaused(bool bInPaused)
 	bPaused = bInPaused;
 }
 
+bool FNiagaraSystemInstance::ComputeEmitterPriority(int32 EmitterIdx, TArray<int32, TInlineAllocator<32>>& EmitterPriorities, const TBitArray<TInlineAllocator<32>>& EmitterDependencyGraph)
+{
+	// Mark this node as being evaluated.
+	EmitterPriorities[EmitterIdx] = 0;
+
+	int32 MaxPriority = 0;
+
+	// Examine all the nodes we depend on. We must run after all of them, so our priority
+	// will be 1 higher than the maximum priority of all our dependencies.
+	const int32 NumEmitters = Emitters.Num();
+	int32 DepStartIndex = EmitterIdx * NumEmitters;
+	TConstSetBitIterator<TInlineAllocator<32>> DepIt(EmitterDependencyGraph, DepStartIndex);
+	while (DepIt.GetIndex() < DepStartIndex + NumEmitters)
+	{
+		int32 OtherEmitterIdx = DepIt.GetIndex() - DepStartIndex;
+
+		// This can't happen, because we explicitly skip self-dependencies when building the edge table.
+		checkSlow(OtherEmitterIdx != EmitterIdx);
+
+		if (EmitterPriorities[OtherEmitterIdx] == 0)
+		{
+			// This node is currently being evaluated, which means we've found a cycle.
+			return false;
+		}
+
+		if (EmitterPriorities[OtherEmitterIdx] < 0)
+		{
+			// Node not evaluated yet, recurse.
+			if (!ComputeEmitterPriority(OtherEmitterIdx, EmitterPriorities, EmitterDependencyGraph))
+			{
+				return false;
+			}
+		}
+
+		if (MaxPriority < EmitterPriorities[OtherEmitterIdx])
+		{
+			MaxPriority = EmitterPriorities[OtherEmitterIdx];
+		}
+
+		++DepIt;
+	}
+
+	EmitterPriorities[EmitterIdx] = MaxPriority + 1;
+	return true;
+}
+
+void FNiagaraSystemInstance::FindDataInterfaceDependencies(const TArray<UNiagaraDataInterface*>& DataInterfaces, TArray<FNiagaraEmitterInstance*>& Dependencies)
+{
+	for (const UNiagaraDataInterface* DI : DataInterfaces)
+	{
+		void* InterfaceInstanceData = FindDataInterfaceInstanceData(DI);
+		int32 NumDepsBefore = Dependencies.Num();
+		DI->GetEmitterDependencies(InterfaceInstanceData, this, Dependencies);
+		// Make sure the DI appended to the array, instead of resetting it.
+		check(Dependencies.Num() >= NumDepsBefore);
+	}
+}
+
+void FNiagaraSystemInstance::ComputeEmittersExecutionOrder()
+{
+	const int32 NumEmitters = Emitters.Num();
+
+	TArray<int32, TInlineAllocator<32>> EmitterPriorities;
+	TBitArray<TInlineAllocator<32>> EmitterDependencyGraph;
+
+	EmitterExecutionOrder.SetNum(NumEmitters);
+	EmitterPriorities.SetNum(NumEmitters);
+	EmitterDependencyGraph.Init(false, NumEmitters * NumEmitters);
+
+	TArray<FNiagaraEmitterInstance*> EmitterDependencies;
+	EmitterDependencies.Reserve(3 * NumEmitters);
+
+	bool bHasEmitterDependencies = false;
+	for (int32 EmitterIdx = 0; EmitterIdx < NumEmitters; ++EmitterIdx)
+	{
+		FNiagaraEmitterInstance& Inst = Emitters[EmitterIdx].Get();
+		EmitterExecutionOrder[EmitterIdx] = EmitterIdx;
+		EmitterPriorities[EmitterIdx] = -1;
+
+		EmitterDependencies.SetNum(0, false);
+		FindDataInterfaceDependencies(Inst.GetSpawnExecutionContext().GetDataInterfaces(), EmitterDependencies);
+		FindDataInterfaceDependencies(Inst.GetUpdateExecutionContext().GetDataInterfaces(), EmitterDependencies);
+
+		// Map the pointers returned by the emitter to indices inside the Emitters array. This is O(N^2), but we expect
+		// to have few dependencies, so in practice it should be faster than a TMap. If it gets out of hand, we can also
+		// ask the DIs to give us indices directly, since they probably got the pointers by scanning the array we gave them
+		// through GetEmitters() anyway.
+		for (int32 DepIdx = 0; DepIdx < EmitterDependencies.Num(); ++DepIdx)
+		{
+			for (int32 OtherEmitterIdx = 0; OtherEmitterIdx < NumEmitters; ++OtherEmitterIdx)
+			{
+				if (EmitterDependencies[DepIdx] == &Emitters[OtherEmitterIdx].Get())
+				{
+					// Some DIs might read from the same emitter they're applied to. We don't care about dependencies on self.
+					if (EmitterIdx != OtherEmitterIdx)
+					{
+						EmitterDependencyGraph.SetRange(EmitterIdx * NumEmitters + OtherEmitterIdx, 1, true);
+						bHasEmitterDependencies = true;
+					}
+					break;
+				}
+			}
+		}
+	}
+
+	if (!bHasEmitterDependencies)
+	{
+		return;
+	}
+
+	for (int32 EmitterIdx = 0; EmitterIdx < NumEmitters; ++EmitterIdx)
+	{
+		if (EmitterPriorities[EmitterIdx] < 0)
+		{
+			if (!ComputeEmitterPriority(EmitterIdx, EmitterPriorities, EmitterDependencyGraph))
+			{
+				FName EmitterName = GetSystem()->GetEmitterHandles()[EmitterIdx].GetName();
+				UE_LOG(LogNiagara, Error, TEXT("Found circular dependency involving emitter '%s' in system '%s'. The execution order will be undefined."), *EmitterName.ToString(), *GetSystem()->GetName());
+				break;
+			}
+		}
+	}
+
+	// Sort the emitter indices in the execution order array so that dependencies are satisfied. Also, emitters with the same priority value don't have any
+	// inter-dependencies, so we can use that if we ever want to parallelize emitter execution.
+	Algo::Sort(EmitterExecutionOrder, [&EmitterPriorities](int32 IdxA, int32 IdxB) { return EmitterPriorities[IdxA] < EmitterPriorities[IdxB]; });
+}
+
 void FNiagaraSystemInstance::Reset(FNiagaraSystemInstance::EResetMode Mode)
 {
 	SCOPE_CYCLE_COUNTER(STAT_NiagaraSystemReset);
@@ -679,6 +807,8 @@ void FNiagaraSystemInstance::Reset(FNiagaraSystemInstance::EResetMode Mode)
 	//Interface init can disable the system.
 	if (!IsComplete())
 	{
+		ComputeEmittersExecutionOrder();
+
 		bPendingSpawn = true;
 		SystemSimulation->AddInstance(this);
 
@@ -1690,8 +1820,10 @@ void FNiagaraSystemInstance::InitEmitters()
 	if (System != nullptr)
 	{
 		const TArray<FNiagaraEmitterHandle>& EmitterHandles = System->GetEmitterHandles();
-		Emitters.Reserve(EmitterHandles.Num());
-		for (int32 EmitterIdx = 0; EmitterIdx < System->GetEmitterHandles().Num(); ++EmitterIdx)
+
+		const int32 NumEmitters = EmitterHandles.Num();
+		Emitters.Reserve(NumEmitters);
+		for (int32 EmitterIdx = 0; EmitterIdx < NumEmitters; ++EmitterIdx)
 		{
 			TSharedRef<FNiagaraEmitterInstance, ESPMode::ThreadSafe> Sim = MakeShared<FNiagaraEmitterInstance, ESPMode::ThreadSafe>(this);
 			Sim->Init(EmitterIdx, ID);
@@ -1846,18 +1978,25 @@ void FNiagaraSystemInstance::Tick_Concurrent()
 	ActiveGPUEmitterCount = 0;
 	UNiagaraSystem* System = GetSystem();
 
+	const int32 NumEmitters = Emitters.Num();
+	checkSlow(EmitterExecutionOrder.Num() == NumEmitters);
+
 	//Determine if any of our emitters should be ticking.
 	TBitArray<TInlineAllocator<8>> EmittersShouldTick;
+	EmittersShouldTick.Init(false, NumEmitters);
+
 	bool bHasTickingEmitters = false;
-	for (int32 EmitterIdx = 0; EmitterIdx < Emitters.Num(); ++EmitterIdx)
+	for (int32 EmitterIdx : EmitterExecutionOrder)
 	{
 		FNiagaraEmitterInstance& Inst = Emitters[EmitterIdx].Get();
-		bool bShouldTick = Inst.ShouldTick();
-		bHasTickingEmitters |= bShouldTick;
-		EmittersShouldTick.Add(bShouldTick);
+		if (Inst.ShouldTick())
+		{
+			bHasTickingEmitters = true;
+			EmittersShouldTick.SetRange(EmitterIdx, 1, true);
+		}
 	}
 
-	if (IsComplete() || !bHasTickingEmitters || GetSystem() == nullptr || Component == nullptr || CachedDeltaSeconds < SMALL_NUMBER)
+	if (IsComplete() || !bHasTickingEmitters || System == nullptr || Component == nullptr || CachedDeltaSeconds < SMALL_NUMBER)
 	{
 		bAsyncWorkInProgress = false;
 		return;
@@ -1865,7 +2004,7 @@ void FNiagaraSystemInstance::Tick_Concurrent()
 
 	FScopeCycleCounter SystemStat(System->GetStatID(true, true));
 
-	for (int32 EmitterIdx = 0; EmitterIdx < Emitters.Num(); EmitterIdx++)
+	for (int32 EmitterIdx : EmitterExecutionOrder)
 	{
 		if (EmittersShouldTick[EmitterIdx])
 		{
@@ -1877,7 +2016,7 @@ void FNiagaraSystemInstance::Tick_Concurrent()
 	bool FirstGpuEmitter = true;
 
 	// now tick all emitters
-	for (int32 EmitterIdx = 0; EmitterIdx < Emitters.Num(); EmitterIdx++)
+	for (int32 EmitterIdx : EmitterExecutionOrder)
 	{
 		FNiagaraEmitterInstance& Inst = Emitters[EmitterIdx].Get();
 		if (EmittersShouldTick[EmitterIdx])
@@ -1896,7 +2035,6 @@ void FNiagaraSystemInstance::Tick_Concurrent()
 			}
 
 			TotalGPUParamSize += 2 * sizeof(FNiagaraEmitterParameters);
-				
 			TotalGPUParamSize += Inst.GetGPUContext()->CombinedParamStore.GetPaddedParameterSizeInBytes();
 			ActiveGPUEmitterCount++;
 		}
diff --git a/Engine/Plugins/FX/Niagara/Source/Niagara/Public/NiagaraSystemInstance.h b/Engine/Plugins/FX/Niagara/Source/Niagara/Public/NiagaraSystemInstance.h
--- a/Engine/Plugins/FX/Niagara/Source/Niagara/Public/NiagaraSystemInstance.h
+++ b/Engine/Plugins/FX/Niagara/Source/Niagara/Public/NiagaraSystemInstance.h
@@ -164,6 +164,7 @@ public:
 	FORCEINLINE UNiagaraComponent *GetComponent() { return Component; }
 	FORCEINLINE TArray<TSharedRef<FNiagaraEmitterInstance, ESPMode::ThreadSafe> > &GetEmitters() { return Emitters; }
 	FORCEINLINE const TArray<TSharedRef<FNiagaraEmitterInstance, ESPMode::ThreadSafe> >& GetEmitters() const { return Emitters; }
+	FORCEINLINE const TArray<int32>& GetEmitterExecutionOrder() const { return EmitterExecutionOrder; }
 	FORCEINLINE const FBox& GetLocalBounds() { return LocalBounds;  }
 
 	FNiagaraEmitterInstance* GetEmitterByID(FGuid InID);
@@ -193,7 +194,7 @@ public:
 	FNiagaraSystemInstanceID GetId() { return ID; }
 
 	/** Returns the instance data for a particular interface for this System. */
-	FORCEINLINE void* FindDataInterfaceInstanceData(UNiagaraDataInterface* Interface) 
+	FORCEINLINE void* FindDataInterfaceInstanceData(const UNiagaraDataInterface* Interface) 
 	{
 		if (int32* InstDataOffset = DataInterfaceInstanceDataOffsets.Find(Interface))
 		{
@@ -312,6 +313,15 @@ private:
 	/** Calculates which tick group the instance should be in. */
 	ETickingGroup CalculateTickGroup();
 
+	/** Computes emitter priorities based on the dependency information. */
+	bool ComputeEmitterPriority(int32 EmitterIdx, TArray<int32, TInlineAllocator<32>>& EmitterPriorities, const TBitArray<TInlineAllocator<32>>& EmitterDependencyGraph);
+
+	/** Queries all the data interfaces in the array for emitter dependencies. */
+	void FindDataInterfaceDependencies(const TArray<UNiagaraDataInterface*>& DataInterfaces, TArray<FNiagaraEmitterInstance*>& Dependencies);
+
+	/** Computes the order in which the emitters in the Emitters array will be ticked and stores the results in EmitterExecutionOrder. */
+	void ComputeEmittersExecutionOrder();
+
 	/** Index of this instance in the system simulation. */
 	int32 SystemInstanceIndex;
 
@@ -425,6 +435,9 @@ private:
 
 	NiagaraEmitterInstanceBatcher* Batcher = nullptr;
 
+	/** Array of emitter indices sorted by execution priority. The emitters will be ticked in this order. */
+	TArray<int32> EmitterExecutionOrder;
+
 public:
 	// Transient data that is accumulated during tick.
 	uint32 TotalGPUParamSize = 0;
