Merging //UE4/Main @ 9375351 to //UE4/Dev-Core (stefan.boberg_rdu-p4merge-02_6162)

(ushell-p4-mergedown)

diff --git a/Engine/Plugins/Experimental/ControlRig/Source/ControlRigEditor/Private/Sequencer/ControlRigParameterTrackEditor.cpp b/Engine/Plugins/Experimental/ControlRig/Source/ControlRigEditor/Private/Sequencer/ControlRigParameterTrackEditor.cpp
--- a/Engine/Plugins/Experimental/ControlRig/Source/ControlRigEditor/Private/Sequencer/ControlRigParameterTrackEditor.cpp
+++ b/Engine/Plugins/Experimental/ControlRig/Source/ControlRigEditor/Private/Sequencer/ControlRigParameterTrackEditor.cpp
@@ -214,16 +214,14 @@ void FControlRigParameterTrackEditor::OnRelease()
 			}
 		}
 	}
-	FControlRigEditMode* ControlRigEditMode = static_cast<FControlRigEditMode*>(GLevelEditorModeTools().FindMode(FControlRigEditMode::ModeName));
+	FControlRigEditMode* ControlRigEditMode = static_cast<FControlRigEditMode*>(GLevelEditorModeTools().GetActiveMode(FControlRigEditMode::ModeName));
 	if (ControlRigEditMode)
 	{
-		if (GLevelEditorModeTools().GetActiveMode(FControlRigEditMode::ModeName))
+		if (GLevelEditorModeTools().HasToolkitHost())
 		{
-			if (GLevelEditorModeTools().HasToolkitHost())
-			{
-				GLevelEditorModeTools().DeactivateMode(FControlRigEditMode::ModeName);
-			}
+			GLevelEditorModeTools().DeactivateMode(FControlRigEditMode::ModeName);
 		}
+
 		ControlRigEditMode->SetObjects(nullptr, FGuid(), nullptr);
 	}
 
@@ -574,7 +572,7 @@ void FControlRigParameterTrackEditor::OnAddTransformKeysForSelectedObjects(EMovi
 void FControlRigParameterTrackEditor::OnSequencerDataChanged(EMovieSceneDataChangeType DataChangeType)
 {
 	UMovieScene* MovieScene = GetSequencer()->GetFocusedMovieSceneSequence()->GetMovieScene();
-	FControlRigEditMode* ControlRigEditMode = static_cast<FControlRigEditMode*>(GLevelEditorModeTools().FindMode(FControlRigEditMode::ModeName));
+	FControlRigEditMode* ControlRigEditMode = static_cast<FControlRigEditMode*>(GLevelEditorModeTools().GetActiveMode(FControlRigEditMode::ModeName));
 	
 	//if we have a valid control rig edit mode need to check and see the control rig in that mode is still in a track
 	//if not we get rid of it.
@@ -591,12 +589,9 @@ void FControlRigParameterTrackEditor::OnSequencerDataChanged(EMovieSceneDataChan
 			}
 		}
 		//okay no good track so deactive it and delete it's Control Rig and bingings.
-		if (GLevelEditorModeTools().GetActiveMode(FControlRigEditMode::ModeName))
+		if (GLevelEditorModeTools().HasToolkitHost())
 		{
-			if (GLevelEditorModeTools().HasToolkitHost())
-			{
-				GLevelEditorModeTools().DeactivateMode(FControlRigEditMode::ModeName);
-			}
+			GLevelEditorModeTools().DeactivateMode(FControlRigEditMode::ModeName);
 		}
 		ControlRigEditMode->SetObjects(nullptr, FGuid(), nullptr);
 	}
diff --git a/Engine/Plugins/Experimental/Landmass/Landmass.uplugin b/Engine/Plugins/Experimental/Landmass/Landmass.uplugin
new file mode 100644
--- /dev/null
+++ b/Engine/Plugins/Experimental/Landmass/Landmass.uplugin
@@ -0,0 +1,23 @@
+{
+	"FileVersion": 3,
+	"Version": 1,
+	"VersionName": "1.0",
+	"FriendlyName": "Landmass",
+	"Description": "",
+	"Category": "Other",
+	"CreatedBy": "Epic Games",
+	"CreatedByURL": "",
+	"DocsURL": "",
+	"MarketplaceURL": "",
+	"SupportURL": "",
+	"CanContainContent": true,
+	"IsBetaVersion": true,
+	"Installed": false,
+	"Modules": [
+		{
+			"Name": "Landmass",
+			"Type": "Editor",
+			"LoadingPhase": "PreDefault"
+		}
+	]
+}
\ No newline at end of file
diff --git a/Engine/Plugins/Experimental/Landmass/Source/Runtime/Landmass.Build.cs b/Engine/Plugins/Experimental/Landmass/Source/Runtime/Landmass.Build.cs
new file mode 100644
--- /dev/null
+++ b/Engine/Plugins/Experimental/Landmass/Source/Runtime/Landmass.Build.cs
@@ -0,0 +1,26 @@
+// Copyright 1998-2019 Epic Games, Inc. All Rights Reserved.
+
+using UnrealBuildTool;
+
+public class Landmass : ModuleRules
+{
+	public Landmass(ReadOnlyTargetRules Target) : base(Target)
+	{
+		PCHUsage = ModuleRules.PCHUsageMode.UseExplicitOrSharedPCHs;
+		
+		PrivateIncludePaths.Add("Runtime/Private");
+
+		PublicDependencyModuleNames.AddRange(
+			new string[] {
+				"Core",
+				"CoreUObject",
+				"Engine",
+			}
+		);
+
+		PrivateDependencyModuleNames.AddRange(
+			new string[] {
+			}
+		);
+	}
+}
\ No newline at end of file
diff --git a/Engine/Plugins/Experimental/Landmass/Source/Runtime/Private/BrushEffectsList.cpp b/Engine/Plugins/Experimental/Landmass/Source/Runtime/Private/BrushEffectsList.cpp
new file mode 100644
--- /dev/null
+++ b/Engine/Plugins/Experimental/Landmass/Source/Runtime/Private/BrushEffectsList.cpp
@@ -0,0 +1,8 @@
+// Copyright 1998-2019 Epic Games, Inc. All Rights Reserved.
+
+#include "BrushEffectsList.h"
+
+FLandmassBrushEffectsList::FLandmassBrushEffectsList()
+{
+
+}
\ No newline at end of file
diff --git a/Engine/Plugins/Experimental/Landmass/Source/Runtime/Private/FalloffSettings.cpp b/Engine/Plugins/Experimental/Landmass/Source/Runtime/Private/FalloffSettings.cpp
new file mode 100644
--- /dev/null
+++ b/Engine/Plugins/Experimental/Landmass/Source/Runtime/Private/FalloffSettings.cpp
@@ -0,0 +1,13 @@
+// Copyright 1998-2019 Epic Games, Inc. All Rights Reserved.
+
+#include "FalloffSettings.h"
+
+FLandmassFalloffSettings::FLandmassFalloffSettings()
+	: FalloffMode(EBrushFalloffMode::Angle)
+	, FalloffAngle(45.0f)
+	, FalloffWidth(1024.0f)
+	, EdgeOffset(0)
+	, ZOffset(0)
+{
+
+}
diff --git a/Engine/Plugins/Experimental/Landmass/Source/Runtime/Private/LandmassModule.cpp b/Engine/Plugins/Experimental/Landmass/Source/Runtime/Private/LandmassModule.cpp
new file mode 100644
--- /dev/null
+++ b/Engine/Plugins/Experimental/Landmass/Source/Runtime/Private/LandmassModule.cpp
@@ -0,0 +1,22 @@
+// Copyright 1998-2019 Epic Games, Inc. All Rights Reserved.
+
+#include "LandmassModule.h"
+#include "Modules/ModuleManager.h"
+
+
+class FLandmassModule : public ILandmassModuleInterface
+{
+public:
+	virtual void StartupModule() override
+	{
+	}
+
+	virtual void ShutdownModule() override
+	{
+	}
+};
+
+//////////////////////////////////////////////////////////////////////////
+
+IMPLEMENT_MODULE(FLandmassModule, Landmass);
+
diff --git a/Engine/Plugins/Experimental/Landmass/Source/Runtime/Private/TerrainCarvingSettings.cpp b/Engine/Plugins/Experimental/Landmass/Source/Runtime/Private/TerrainCarvingSettings.cpp
new file mode 100644
--- /dev/null
+++ b/Engine/Plugins/Experimental/Landmass/Source/Runtime/Private/TerrainCarvingSettings.cpp
@@ -0,0 +1,12 @@
+// Copyright 1998-2019 Epic Games, Inc. All Rights Reserved.
+
+#include "TerrainCarvingSettings.h"
+
+FLandmassTerrainCarvingSettings::FLandmassTerrainCarvingSettings()
+	: BlendMode(EBrushBlendType::AlphaBlend)
+	, bInvertShape(false)
+	, Priority(0)
+{
+
+
+}
diff --git a/Engine/Plugins/Experimental/Landmass/Source/Runtime/Public/BrushEffectsList.h b/Engine/Plugins/Experimental/Landmass/Source/Runtime/Public/BrushEffectsList.h
new file mode 100644
--- /dev/null
+++ b/Engine/Plugins/Experimental/Landmass/Source/Runtime/Public/BrushEffectsList.h
@@ -0,0 +1,182 @@
+// Copyright 1998-2019 Epic Games, Inc. All Rights Reserved.
+
+#pragma once
+
+#include "CoreMinimal.h"
+#include "UObject/ObjectMacros.h"
+#include "BrushEffectsList.generated.h"
+
+class UTexture2D;
+class UCurveFloat;
+
+USTRUCT(BlueprintType)
+struct FBrushEffectBlurring
+{
+	GENERATED_BODY()
+
+	FBrushEffectBlurring()
+		: bBlurShape(true)
+		, Radius(2)
+	{}
+
+	UPROPERTY(EditAnywhere, BlueprintReadWrite, Category = BrushEffects)
+	bool bBlurShape;
+
+	UPROPERTY(EditAnywhere, BlueprintReadWrite, Category = BrushEffects)
+	int32 Radius;
+};
+
+
+USTRUCT(BlueprintType)
+struct FBrushEffectCurlNoise
+{
+	GENERATED_BODY()
+
+	FBrushEffectCurlNoise()
+		: Curl1Amount(0)
+		, Curl2Amount(0)
+		, Curl1Tiling(16.0)
+		, Curl2Tiling(3.0)
+	{}
+
+	UPROPERTY(EditAnywhere, BlueprintReadWrite, Category = BrushEffects)
+	float Curl1Amount;
+
+	UPROPERTY(EditAnywhere, BlueprintReadWrite, Category = BrushEffects)
+	float Curl2Amount;
+
+	UPROPERTY(EditAnywhere, BlueprintReadWrite, Category = BrushEffects)
+	float Curl1Tiling;
+
+	UPROPERTY(EditAnywhere, BlueprintReadWrite, Category = BrushEffects)
+	float Curl2Tiling;
+};
+
+USTRUCT(BlueprintType)
+struct FBrushEffectCurves
+{
+	GENERATED_BODY()
+
+	FBrushEffectCurves()
+		: bUseCurveChannel(true)
+		, ElevationCurveAsset(nullptr)
+		, ChannelEdgeOffset(0.0f)
+		, ChannelDepth(0.0f)
+		, CurveRampWidth(512.0f)
+	{}
+
+	UPROPERTY(EditAnywhere, BlueprintReadWrite, Category = BrushEffects)
+	bool bUseCurveChannel;
+
+	UPROPERTY(EditAnywhere, BlueprintReadWrite, Category = BrushEffects)
+	UCurveFloat* ElevationCurveAsset;
+
+	UPROPERTY(EditAnywhere, BlueprintReadWrite, Category = BrushEffects)
+	float ChannelEdgeOffset;
+
+	UPROPERTY(EditAnywhere, BlueprintReadWrite, Category = BrushEffects)
+	float ChannelDepth;
+
+	UPROPERTY(EditAnywhere, BlueprintReadWrite, Category = BrushEffects)
+	float CurveRampWidth;
+};
+
+
+USTRUCT(BlueprintType)
+struct FBrushEffectDisplacement
+{
+	GENERATED_BODY()
+
+	FBrushEffectDisplacement()
+		: DisplacementHeight(0)
+		, DisplacementTiling(0)
+		, Texture(nullptr)
+		, Midpoint(-128.0f)
+		, Channel(0,0,0,1)
+		, WeightmapInfluence(0.0f)
+	{}
+
+	UPROPERTY(EditAnywhere, BlueprintReadWrite, Category = BrushEffects)
+	float DisplacementHeight;
+
+	UPROPERTY(EditAnywhere, BlueprintReadWrite, Category = BrushEffects)
+	float DisplacementTiling;
+
+	UPROPERTY(EditAnywhere, BlueprintReadWrite, Category = BrushEffects)
+	UTexture2D* Texture;
+
+	UPROPERTY(EditAnywhere, BlueprintReadWrite, Category = BrushEffects)
+	float Midpoint;
+
+	UPROPERTY(EditAnywhere, BlueprintReadWrite, Category = BrushEffects)
+	FLinearColor Channel;
+
+	UPROPERTY(EditAnywhere, BlueprintReadWrite, Category = BrushEffects)
+	float WeightmapInfluence;
+};
+
+
+USTRUCT(BlueprintType)
+struct FBrushEffectSmoothBlending
+{
+	GENERATED_BODY()
+
+	UPROPERTY(EditAnywhere, BlueprintReadWrite, Category = BrushEffects)
+	float InnerSmoothDistance;
+
+	UPROPERTY(EditAnywhere, BlueprintReadWrite, Category = BrushEffects)
+	float OuterSmoothDistance;
+};
+
+USTRUCT(BlueprintType)
+struct FBrushEffectTerracing
+{
+	GENERATED_BODY()
+
+	FBrushEffectTerracing()
+		: TerraceAlpha(0.0f)
+		, TerraceSpacing(256.0f)
+		, TerraceSmoothness(0.0f)
+		, MaskLength(0.0f)
+		, MaskStartOffset(0.0f)
+	{}
+
+	UPROPERTY(EditAnywhere, BlueprintReadWrite, Category = BrushEffects)
+	float TerraceAlpha;
+
+	UPROPERTY(EditAnywhere, BlueprintReadWrite, Category = BrushEffects)
+	float TerraceSpacing;
+
+	UPROPERTY(EditAnywhere, BlueprintReadWrite, Category = BrushEffects)
+	float TerraceSmoothness;
+
+	UPROPERTY(EditAnywhere, BlueprintReadWrite, Category = BrushEffects)
+	float MaskLength;
+
+	UPROPERTY(EditAnywhere, BlueprintReadWrite, Category = BrushEffects)
+	float MaskStartOffset;
+};
+
+
+USTRUCT(BlueprintType)
+struct FLandmassBrushEffectsList
+{
+	GENERATED_BODY()
+
+	FLandmassBrushEffectsList();
+
+	UPROPERTY(EditAnywhere, BlueprintReadWrite, Category = BrushEffects)
+	FBrushEffectBlurring Blurring;
+
+	UPROPERTY(EditAnywhere, BlueprintReadWrite, Category = FalloffSettings)
+	FBrushEffectCurlNoise CurlNoise;
+
+	UPROPERTY(EditAnywhere, BlueprintReadWrite, Category = FalloffSettings)
+	FBrushEffectDisplacement Displacement;
+
+	UPROPERTY(EditAnywhere, BlueprintReadWrite, Category = FalloffSettings)
+	FBrushEffectSmoothBlending SmoothBlending;
+
+	UPROPERTY(EditAnywhere, BlueprintReadWrite, Category = FalloffSettings)
+	FBrushEffectTerracing Terracing;
+};
\ No newline at end of file
diff --git a/Engine/Plugins/Experimental/Landmass/Source/Runtime/Public/FalloffSettings.h b/Engine/Plugins/Experimental/Landmass/Source/Runtime/Public/FalloffSettings.h
new file mode 100644
--- /dev/null
+++ b/Engine/Plugins/Experimental/Landmass/Source/Runtime/Public/FalloffSettings.h
@@ -0,0 +1,36 @@
+// Copyright 1998-2019 Epic Games, Inc. All Rights Reserved.
+
+#pragma once
+
+#include "CoreMinimal.h"
+#include "UObject/ObjectMacros.h"
+#include "FalloffSettings.generated.h"
+
+UENUM(BlueprintType)
+enum class EBrushFalloffMode : uint8
+{
+	Angle,
+	Width,
+};
+
+USTRUCT(BlueprintType)
+struct FLandmassFalloffSettings
+{
+	GENERATED_BODY()
+	FLandmassFalloffSettings();
+
+	UPROPERTY(EditAnywhere, BlueprintReadWrite, Category = FalloffSettings)
+	EBrushFalloffMode FalloffMode;
+
+	UPROPERTY(EditAnywhere, BlueprintReadWrite, Category = FalloffSettings)
+	float FalloffAngle;
+
+	UPROPERTY(EditAnywhere, BlueprintReadWrite, Category = FalloffSettings)
+	float FalloffWidth;
+
+	UPROPERTY(EditAnywhere, BlueprintReadWrite, Category = FalloffSettings)
+	float EdgeOffset;
+
+	UPROPERTY(EditAnywhere, BlueprintReadWrite, Category = FalloffSettings)
+	float ZOffset;
+};
\ No newline at end of file
diff --git a/Engine/Plugins/Experimental/Landmass/Source/Runtime/Public/LandmassModule.h b/Engine/Plugins/Experimental/Landmass/Source/Runtime/Public/LandmassModule.h
new file mode 100644
--- /dev/null
+++ b/Engine/Plugins/Experimental/Landmass/Source/Runtime/Public/LandmassModule.h
@@ -0,0 +1,12 @@
+// Copyright 1998-2019 Epic Games, Inc. All Rights Reserved.
+
+#pragma once
+
+#include "CoreMinimal.h"
+#include "Stats/Stats.h"
+#include "Modules/ModuleInterface.h"
+
+
+class ILandmassModuleInterface : public IModuleInterface
+{
+};
\ No newline at end of file
diff --git a/Engine/Plugins/Experimental/Landmass/Source/Runtime/Public/TerrainCarvingSettings.h b/Engine/Plugins/Experimental/Landmass/Source/Runtime/Public/TerrainCarvingSettings.h
new file mode 100644
--- /dev/null
+++ b/Engine/Plugins/Experimental/Landmass/Source/Runtime/Public/TerrainCarvingSettings.h
@@ -0,0 +1,47 @@
+// Copyright 1998-2019 Epic Games, Inc. All Rights Reserved.
+
+#pragma once
+
+
+#include "CoreMinimal.h"
+#include "UObject/ObjectMacros.h"
+#include "FalloffSettings.h"
+#include "BrushEffectsList.h"
+#include "TerrainCarvingSettings.generated.h"
+
+/** The blend mode changes how the brush material is applied to the terrain. */
+UENUM(BlueprintType)
+enum class EBrushBlendType : uint8
+{
+	/** Alpha Blend will affect the heightmap both upwards and downwards. */
+	AlphaBlend,
+	/** Limits the brush to only lowering the terrain. */
+	Min,
+	/** Limits the brush to only raising the terrain. */
+	Max,
+	/** Performs an additive blend, using a flat Z=0 terrain as the input. Useful when you want to preserve underlying detail or ramps. */
+	Additive,
+};
+
+USTRUCT(BlueprintType)
+struct FLandmassTerrainCarvingSettings
+{
+	GENERATED_BODY()
+
+	LANDMASS_API FLandmassTerrainCarvingSettings();
+
+	UPROPERTY(EditAnywhere, BlueprintReadWrite, Category=TerrainCarvingSettings)
+	EBrushBlendType BlendMode;
+
+	UPROPERTY(EditAnywhere, BlueprintReadWrite, Category = TerrainCarvingSettings)
+	bool bInvertShape;
+
+	UPROPERTY(EditAnywhere, BlueprintReadWrite, Category = TerrainCarvingSettings)
+	FLandmassFalloffSettings FalloffSettings;
+
+	UPROPERTY(EditAnywhere, BlueprintReadWrite, Category = TerrainCarvingSettings)
+	FLandmassBrushEffectsList Effects;
+
+	UPROPERTY(EditAnywhere, BlueprintReadWrite, Category = TerrainCarvingSettings)
+	int32 Priority;
+};
diff --git a/Engine/Plugins/FX/Niagara/Source/Niagara/Classes/NiagaraCollision.h b/Engine/Plugins/FX/Niagara/Source/Niagara/Classes/NiagaraCollision.h
--- a/Engine/Plugins/FX/Niagara/Source/Niagara/Classes/NiagaraCollision.h
+++ b/Engine/Plugins/FX/Niagara/Source/Niagara/Classes/NiagaraCollision.h
@@ -24,74 +24,6 @@ struct FNiagaraCollisionTrace
 };
 
 
-class FNiagaraCollisionBatch
-{
-public:
-	FNiagaraCollisionBatch()
-		: CollisionEventDataSet(nullptr)
-		, EmitterName(NAME_None)
-	{
-	}
-
-	~FNiagaraCollisionBatch()
-	{
-		FNiagaraEventDataSetMgr::Reset(OwnerSystemInstanceID, EmitterName);
-	}
-
-	void Reset()
-	{
-		if (CollisionEventDataSet)
-		{
-			CollisionEventDataSet->ResetBuffers();
-		}
-	}
-
-	void Init(FNiagaraSystemInstanceID InOwnerSystemInstanceID, FName InEmitterName)
-	{
-		if (CollisionEventDataSet)
-		{
-#if !UE_BUILD_SHIPPING && !UE_BUILD_TEST
-			CollisionEventDataSet->Init(FNiagaraDataSetID(), ENiagaraSimTarget::CPUSim, TEXT("NiagaraCollisionEventDataset"));
-#else
-			CollisionEventDataSet->Init(FNiagaraDataSetID(), ENiagaraSimTarget::CPUSim);
-#endif
-		}
-		EmitterName = InEmitterName;
-		OwnerSystemInstanceID = InOwnerSystemInstanceID;
-		CollisionEventDataSet = FNiagaraEventDataSetMgr::CreateEventDataSet(OwnerSystemInstanceID, EmitterName, NIAGARA_BUILTIN_EVENTNAME_COLLISION);
-		// TODO: this should go away once we can use the FNiagaraCollisionEventPayload
-		// UStruct to create the data set
-		//FNiagaraVariable ValidVar(FNiagaraTypeDefinition::GetIntDef(), "Valid");
-		FNiagaraVariable PosVar(FNiagaraTypeDefinition::GetVec3Def(), "CollisionLocation");
-		FNiagaraVariable NrmVar(FNiagaraTypeDefinition::GetVec3Def(), "CollisionNormal");
-		FNiagaraVariable PhysMatIdxVar(FNiagaraTypeDefinition::GetIntDef(), "PhysicalMaterialIndex");
-		FNiagaraVariable VelVar(FNiagaraTypeDefinition::GetVec3Def(), "CollisionVelocity");
-		FNiagaraVariable ParticleIndexVar(FNiagaraTypeDefinition::GetIntDef(), "ParticleIndex");
-		//CollisionEventDataSet->AddVariable(ValidVar);
-		CollisionEventDataSet->AddVariable(PosVar);
-		CollisionEventDataSet->AddVariable(NrmVar);
-		CollisionEventDataSet->AddVariable(PhysMatIdxVar);
-		CollisionEventDataSet->AddVariable(VelVar);
-		CollisionEventDataSet->AddVariable(ParticleIndexVar);
-		CollisionEventDataSet->Finalize();
-	}
-	
-	void KickoffNewBatch(class FNiagaraEmitterInstance *Sim, float DeltaSeconds);
-
-	void GenerateEventsFromResults(class FNiagaraEmitterInstance *Sim);
-
-	const FNiagaraDataSet *GetDataSet() const { return CollisionEventDataSet; }
-private:
-	TArray<FTraceHandle> CollisionTraceHandles;
-	TArray<FNiagaraCollisionTrace> CollisionTraces;
-	TArray<FNiagaraCollisionEventPayload> CollisionEvents;
-	FNiagaraDataSet *CollisionEventDataSet;
-	FName EmitterName;
-	FNiagaraSystemInstanceID OwnerSystemInstanceID;
-};
-
-
-
 struct FNiagaraDICollsionQueryResult
 {
 	uint32 TraceID;
diff --git a/Engine/Plugins/FX/Niagara/Source/Niagara/Classes/NiagaraDataSet.h b/Engine/Plugins/FX/Niagara/Source/Niagara/Classes/NiagaraDataSet.h
--- a/Engine/Plugins/FX/Niagara/Source/Niagara/Classes/NiagaraDataSet.h
+++ b/Engine/Plugins/FX/Niagara/Source/Niagara/Classes/NiagaraDataSet.h
@@ -217,7 +217,7 @@ private:
 //////////////////////////////////////////////////////////////////////////
 
 USTRUCT()
-struct FNiagaraDataSetCompiledData
+struct NIAGARA_API FNiagaraDataSetCompiledData
 {
 	GENERATED_BODY()
 
@@ -248,10 +248,14 @@ struct FNiagaraDataSetCompiledData
 	UPROPERTY()
 	ENiagaraSimTarget SimTarget;
 
+	FNiagaraDataSetCompiledData();
 	void BuildLayout();
 	void Empty();
+
+	static FNiagaraDataSetCompiledData DummyCompiledData;
 };
 
+
 //////////////////////////////////////////////////////////////////////////
 
 /**
@@ -266,13 +270,12 @@ public:
 	~FNiagaraDataSet();
 	FNiagaraDataSet& operator=(const FNiagaraDataSet&) = delete;
 
-#if !UE_BUILD_SHIPPING && !UE_BUILD_TEST
-	void Init(FNiagaraDataSetID InID, ENiagaraSimTarget InSimTarget, const FString& InDebugName);
-	void InitFromCompiledData(const FNiagaraDataSetCompiledData& InDataSetCompiledData, const FString& InDebugName);
-#else
-	void Init(FNiagaraDataSetID InID, ENiagaraSimTarget InSimTarget);
-	void InitFromCompiledData(const FNiagaraDataSetCompiledData& InDataSetCompiledData);
-#endif
+	FORCEINLINE void Init(const FNiagaraDataSetCompiledData* InDataSetCompiledData)
+	{
+		CompiledData = InDataSetCompiledData != nullptr ? InDataSetCompiledData : &FNiagaraDataSetCompiledData::DummyCompiledData;
+	
+		Reset();
+	}
 
 	/** Resets current data but leaves variable/layout information etc intact. */
 	void ResetBuffers();
@@ -289,13 +292,10 @@ public:
 	/** Returns size in bytes for all data buffers currently allocated by this dataset. */
 	uint32 GetSizeBytes()const;
 
-	FORCEINLINE bool IsInitialized()const { return bFinalized; }
-	FORCEINLINE ENiagaraSimTarget GetSimTarget()const { return SimTarget; }
-	FORCEINLINE FNiagaraDataSetID GetID()const { return ID; }
-	FORCEINLINE void SetID(FNiagaraDataSetID InID) { ID = InID; }
-
-	FORCEINLINE void SetNeedsPersistentIDs(bool bNeedsIDs) { bNeedsPersistentIDs = bNeedsIDs; }
-	FORCEINLINE bool GetNeedsPersistentIDs()const { return bNeedsPersistentIDs; }
+	FORCEINLINE bool IsInitialized()const { return CompiledData != nullptr; }
+	FORCEINLINE ENiagaraSimTarget GetSimTarget()const { return CompiledData->SimTarget; }
+	FORCEINLINE FNiagaraDataSetID GetID()const { return CompiledData->ID; }	
+	FORCEINLINE bool GetNeedsPersistentIDs()const { return CompiledData->bNeedsPersistentIDs; }
 
 	FORCEINLINE TArray<int32>& GetFreeIDTable() { return FreeIDsTable; }
 	FORCEINLINE int32& GetNumFreeIDs() { return NumFreeIDs; }
@@ -303,18 +303,13 @@ public:
 	FORCEINLINE int32& GetIDAcquireTag() { return IDAcquireTag; }
 	FORCEINLINE void SetIDAcquireTag(int32 InTag) { IDAcquireTag = InTag; }
 
-	FORCEINLINE TArray<FNiagaraVariable>& GetVariables() { return Variables; }
-	FORCEINLINE uint32 GetNumVariables()const { return Variables.Num(); }
-	FORCEINLINE bool HasVariable(const FNiagaraVariable& Var)const { return Variables.Contains(Var); }
-	FORCEINLINE uint32 GetNumFloatComponents()const { return TotalFloatComponents; }
-	FORCEINLINE uint32 GetNumInt32Components()const { return TotalInt32Components; }
-
-	void AddVariable(FNiagaraVariable& Variable);
-	void AddVariables(const TArray<FNiagaraVariable>& Vars);
-
-	/** Finalize the addition of variables and other setup before this data set can be used. */
-	void Finalize();
+	FORCEINLINE const TArray<FNiagaraVariable>& GetVariables()const { return CompiledData->Variables; }
+	FORCEINLINE uint32 GetNumVariables()const { return CompiledData->Variables.Num(); }
+	FORCEINLINE bool HasVariable(const FNiagaraVariable& Var)const { return CompiledData->Variables.Contains(Var); }
+	FORCEINLINE uint32 GetNumFloatComponents()const { return CompiledData->TotalFloatComponents; }
+	FORCEINLINE uint32 GetNumInt32Components()const { return CompiledData->TotalInt32Components; }
 
+	const TArray<FNiagaraVariableLayoutInfo>& GetVariableLayouts()const { return CompiledData->VariableLayouts; }
 	const FNiagaraVariableLayoutInfo* GetVariableLayout(const FNiagaraVariable& Var)const;
 	bool GetVariableComponentOffsets(const FNiagaraVariable& Var, int32 &FloatStart, int32 &IntStart) const;
 
@@ -359,31 +354,15 @@ private:
 		// In some rare occasions, the render thread might be null, like when offloading work to Lightmass 
 		// The final GRenderingThread check keeps us from inadvertently failing when that happens.
 #if DO_GUARD_SLOW
+		ENiagaraSimTarget SimTarget = GetSimTarget();
 		bool CPUSimOK = (SimTarget == ENiagaraSimTarget::CPUSim && !IsInRenderingThread());
 		bool GPUSimOK = (SimTarget == ENiagaraSimTarget::GPUComputeSim && IsInRenderingThread());
 		checkfSlow(!GRenderingThread || CPUSimOK || GPUSimOK, TEXT("NiagaraDataSet function being called on incorrect thread."));
 #endif
 	}
 
-	/** Unique ID for this data set. Used to allow referencing from other emitters and Systems. */
-	FNiagaraDataSetID ID;
-
-	//////////////////////////////////////////////////////////////////////////
-	//TODO: All this layout is known per emitter / system so doesn't need to be generated and stored for every dataset!
-	/** Variables in the data set. */
-	TArray<FNiagaraVariable> Variables;
-	/** Data describing the layout of variable data. */
-	TArray<FNiagaraVariableLayoutInfo> VariableLayouts;
-	/** Total number of components of each type in the data set. */
-	uint32 TotalFloatComponents;
-	uint32 TotalInt32Components;
-	//////////////////////////////////////////////////////////////////////////
-
-	ENiagaraSimTarget SimTarget;
+	const FNiagaraDataSetCompiledData* CompiledData;
 
-	/** Once finalized, the data layout etc is built and no more variables can be added. */
-	uint32 bFinalized : 1;
-	uint32 bNeedsPersistentIDs : 1;
 
 	/** Table of free IDs available to allocate next tick. */
 	TArray<int32> FreeIDsTable;
@@ -411,10 +390,6 @@ private:
 	Additional buffers may be in here if they are currently being used by the render thread.
 	*/
 	TArray<FNiagaraDataBuffer*, TInlineAllocator<2>> Data;
-
-#if !UE_BUILD_SHIPPING && !UE_BUILD_TEST
-	FString DebugName;
-#endif
 };
 
 /**
@@ -1476,10 +1451,10 @@ struct FNiagaraDataVariableIterator
 	{
 		for (int32 VarIdx = 0; VarIdx < Variables.Num(); ++VarIdx)
 		{
-			FNiagaraVariable* Var = Variables[VarIdx];
+			const FNiagaraVariable* Var = Variables[VarIdx];
 			const FNiagaraVariableLayoutInfo* Layout = VarLayouts[VarIdx];
 			check(Var && Layout);
-			uint8* ValuePtr = Var->GetData();
+			const uint8* ValuePtr = Var->GetData();
 
 			for (uint32 CompIdx = 0; CompIdx < Layout->GetNumFloatComponents(); ++CompIdx)
 			{
@@ -1526,7 +1501,7 @@ struct FNiagaraDataVariableIterator
 private:
 
 	const FNiagaraDataBuffer* Data;
-	TArray<FNiagaraVariable*> Variables;
+	TArray<const FNiagaraVariable*> Variables;
 	TArray<const FNiagaraVariableLayoutInfo*> VarLayouts;
 
 	uint32 CurrIdx;
@@ -1562,8 +1537,8 @@ FORCEINLINE void FNiagaraDataBuffer::CheckUsage(bool bReadOnly)const
 
 	//We can read on the RT but any modifications must be GT (or GT Task).
 	//For GPU sims we must be on the RT.
-	checkSlow((Owner->SimTarget == ENiagaraSimTarget::CPUSim  && (IsInGameThread() || (bReadOnly || !IsInRenderingThread()))) ||
-		(Owner->SimTarget == ENiagaraSimTarget::GPUComputeSim) && IsInRenderingThread());
+	checkSlow((Owner->GetSimTarget() == ENiagaraSimTarget::CPUSim  && (IsInGameThread() || (bReadOnly || !IsInRenderingThread()))) ||
+		(Owner->GetSimTarget() == ENiagaraSimTarget::GPUComputeSim) && IsInRenderingThread());
 }
 
 FORCEINLINE_DEBUGGABLE void FNiagaraDataBuffer::AppendToRegisterTable(FNiagaraRegisterTable& Registers, int32 StartInstance)
diff --git a/Engine/Plugins/FX/Niagara/Source/Niagara/Classes/NiagaraEmitter.h b/Engine/Plugins/FX/Niagara/Source/Niagara/Classes/NiagaraEmitter.h
--- a/Engine/Plugins/FX/Niagara/Source/Niagara/Classes/NiagaraEmitter.h
+++ b/Engine/Plugins/FX/Niagara/Source/Niagara/Classes/NiagaraEmitter.h
@@ -70,9 +70,11 @@ struct FNiagaraEventGeneratorProperties
 
 	FNiagaraEventGeneratorProperties(FNiagaraDataSetProperties &Props, FName InEventGenerator)
 		: ID(Props.ID.Name)
-		, SetProps(Props)		
 	{
-
+		DataSetCompiledData.Variables = Props.Variables;
+		DataSetCompiledData.ID = Props.ID;
+		DataSetCompiledData.SimTarget = ENiagaraSimTarget::CPUSim;
+		DataSetCompiledData.BuildLayout();
 	}
 
 	/** Max Number of Events that can be generated per frame. */
@@ -83,7 +85,7 @@ struct FNiagaraEventGeneratorProperties
 	FName ID;
 
 	UPROPERTY()
-	FNiagaraDataSetProperties SetProps;
+	FNiagaraDataSetCompiledData DataSetCompiledData;
 };
 
 
@@ -215,6 +217,13 @@ public:
 	/** An emitter-based seed for the deterministic random number generator. */
 	UPROPERTY(EditAnywhere, Category = "Emitter", meta = (EditCondition = "bDeterminism"))
 	int32 RandomSeed;
+	
+	/** 
+	The emitter will allocate at least this many particles on it's first tick.
+	This can aid performance by avoiding many allocations as an emitter ramps up to it's max size.
+	*/
+	UPROPERTY(EditAnywhere, Category = "Emitter")
+	int32 PreAllocationCount;
 
 	UPROPERTY()
 	FNiagaraEmitterScriptProperties UpdateScriptProps;
diff --git a/Engine/Plugins/FX/Niagara/Source/Niagara/Classes/NiagaraEmitterInstance.h b/Engine/Plugins/FX/Niagara/Source/Niagara/Classes/NiagaraEmitterInstance.h
--- a/Engine/Plugins/FX/Niagara/Source/Niagara/Classes/NiagaraEmitterInstance.h
+++ b/Engine/Plugins/FX/Niagara/Source/Niagara/Classes/NiagaraEmitterInstance.h
@@ -45,8 +45,8 @@ public:
 
 	/** Replaces the binding for a single parameter colleciton instance. If for example the component begins to override the global instance. */
 	//void RebindParameterCollection(UNiagaraParameterCollectionInstance* OldInstance, UNiagaraParameterCollectionInstance* NewInstance);
-	void BindParameters();
-	void UnbindParameters();
+	void BindParameters(bool bExternalOnly);
+	void UnbindParameters(bool bExternalOnly);
 
 	bool IsAllowedToExecute() const;
 
@@ -127,7 +127,6 @@ public:
 	void TickFastPathAttributeBindings();
 
 	FNiagaraEmitterFastPath::FParamMap0& GetFastPathMap() { return FastPathMap; }
-
 private:
 	void CheckForErrors();
 
diff --git a/Engine/Plugins/FX/Niagara/Source/Niagara/Classes/NiagaraSystem.h b/Engine/Plugins/FX/Niagara/Source/Niagara/Classes/NiagaraSystem.h
--- a/Engine/Plugins/FX/Niagara/Source/Niagara/Classes/NiagaraSystem.h
+++ b/Engine/Plugins/FX/Niagara/Source/Niagara/Classes/NiagaraSystem.h
@@ -49,6 +49,11 @@ struct FNiagaraEmitterCompiledData
 	/** Per-Emitter DataSet Data. */
 	UPROPERTY()
 	FNiagaraDataSetCompiledData DataSetCompiledData;
+
+#if WITH_EDITORONLY_DATA
+	UPROPERTY()
+	FNiagaraDataSetCompiledData GPUCaptureDataSetCompiledData;
+#endif
 };
 
 USTRUCT()
@@ -64,6 +69,15 @@ struct FNiagaraSystemCompiledData
 
 	UPROPERTY()
 	FNiagaraParameterStore InstanceParamStore;
+
+	UPROPERTY()
+	FNiagaraDataSetCompiledData DataSetCompiledData;
+
+	UPROPERTY()
+	FNiagaraDataSetCompiledData SpawnInstanceParamsDataSetCompiledData;
+
+	UPROPERTY()
+	FNiagaraDataSetCompiledData UpdateInstanceParamsDataSetCompiledData;
 };
 
 USTRUCT()
diff --git a/Engine/Plugins/FX/Niagara/Source/Niagara/Private/NiagaraCollision.cpp b/Engine/Plugins/FX/Niagara/Source/Niagara/Private/NiagaraCollision.cpp
--- a/Engine/Plugins/FX/Niagara/Source/Niagara/Private/NiagaraCollision.cpp
+++ b/Engine/Plugins/FX/Niagara/Source/Niagara/Private/NiagaraCollision.cpp
@@ -10,167 +10,6 @@
 DECLARE_CYCLE_STAT(TEXT("Collision"), STAT_NiagaraCollision, STATGROUP_Niagara);
 DECLARE_CYCLE_STAT(TEXT("Event Emission"), STAT_NiagaraEventWrite, STATGROUP_Niagara);
 
-void FNiagaraCollisionBatch::KickoffNewBatch(FNiagaraEmitterInstance *Sim, float DeltaSeconds)
-{
-	SCOPE_CYCLE_COUNTER(STAT_NiagaraCollision);
-	FNiagaraVariable PosVar(FNiagaraTypeDefinition::GetVec3Def(), "Position");
-	FNiagaraVariable VelVar(FNiagaraTypeDefinition::GetVec3Def(), "Velocity");
-	FNiagaraVariable SizeVar(FNiagaraTypeDefinition::GetVec2Def(), "SpriteSize");
-	FNiagaraVariable TstVar(FNiagaraTypeDefinition::GetBoolDef(), "PerformCollision");
-	FNiagaraDataSetAccessor<FVector> PosData(Sim->GetData(), PosVar);
-	FNiagaraDataSetAccessor<FVector> VelData(Sim->GetData(), VelVar);
-	FNiagaraDataSetAccessor<FVector2D> SizeData(Sim->GetData(), SizeVar);
-	//FNiagaraDataSetIterator<int32> TstIt(Sim->GetData(), TstVar, 0, false);
-
-	bool bUseSize = SizeData.IsValidForRead();
-
-	if (!PosData.IsValidForRead() || !VelData.IsValidForRead() /*|| !TstIt.IsValid()*/)
-	{
-		return;
-	}
-
-	UWorld *SystemWorld = Sim->GetParentSystemInstance()->GetComponent()->GetWorld();
-	if (SystemWorld)
-	{
-		CollisionTraces.Empty();
-
-		for (uint32 i = 0; i < Sim->GetData().GetCurrentDataChecked().GetNumInstances(); i++)
-		{
-			//int32 TestCollision = *TstIt;
-			//if (TestCollision)
-			{
-				FVector Position;
-				FVector EndPosition;
-				FVector Velocity;
-
-				PosData.Get(i, Position);
-				VelData.Get(i, Velocity);
-				EndPosition = Position + Velocity * DeltaSeconds;
-
-				if (bUseSize)
-				{
-					//TODO:  Handle mesh particles too.  Also this can probably be better and or faster.
-					FVector2D SpriteSize;
-					SizeData.Get(i, SpriteSize);
-					float MaxSize = FMath::Max(SpriteSize.X, SpriteSize.Y);
-
-					float Length;
-					FVector Direction;
-					Velocity.ToDirectionAndLength(Direction, Length);
-
-					Position -= Direction * (MaxSize / 2);
-					EndPosition += Direction * (MaxSize / 2);
-				}
-
-				FCollisionQueryParams QueryParams(SCENE_QUERY_STAT(NiagraAsync));
-				QueryParams.OwnerTag = "Niagara";
-				FTraceHandle Handle = SystemWorld->AsyncLineTraceByChannel(EAsyncTraceType::Single, Position, EndPosition, ECollisionChannel::ECC_WorldStatic, QueryParams, FCollisionResponseParams::DefaultResponseParam, nullptr, i);
-				FNiagaraCollisionTrace Trace;
-				Trace.CollisionTraceHandle = Handle;
-				Trace.SourceParticleIndex = i;
-				CollisionTraces.Add(Trace);
-			}
-		}
-	}
-}
-
-void FNiagaraCollisionBatch::GenerateEventsFromResults(FNiagaraEmitterInstance *Sim)
-{
-	SCOPE_CYCLE_COUNTER(STAT_NiagaraCollision);
-	//CollisionEventDataSet.Allocate(CollisionTraceHandles.Num());
-
-	UWorld *SystemWorld = Sim->GetParentSystemInstance()->GetComponent()->GetWorld();
-	if (SystemWorld)
-	{
-		TArray<FNiagaraCollisionEventPayload> Payloads;
-
-		// generate events for last frame's collisions
-		//
-		for (FNiagaraCollisionTrace CurCheck: CollisionTraces)
-		{
-			FTraceHandle Handle = CurCheck.CollisionTraceHandle;
-			FTraceDatum CurTrace;
-			// wait for trace handles; this should block rarely to never
-			bool bReady = false;
-			while (!bReady)
-			{
-				bReady = SystemWorld->QueryTraceData(Handle, CurTrace);
-				if (!bReady)
-				{
-					// if the query came back false, it's possible that the hanle is invalid for some reason; skip in that case
-					// TODO: handle this more gracefully
-					if (!SystemWorld->IsTraceHandleValid(Handle, false))
-					{
-						break;
-					}
-					break;
-				}
-			}
-
-			if (bReady && CurTrace.OutHits.Num())
-			{
-				// grab the first hit that blocks
-				FHitResult *Hit = FHitResult::GetFirstBlockingHit(CurTrace.OutHits);
-				if (Hit && Hit->IsValidBlockingHit())
-				{
-					FNiagaraCollisionEventPayload Event;
-					Event.CollisionNormal = Hit->ImpactNormal;
-					Event.CollisionPos = Hit->ImpactPoint;
-					Event.ParticleIndex = CurTrace.UserData;
-					check(!Event.CollisionNormal.ContainsNaN());
-					check(Event.CollisionNormal.IsNormalized());
-					check(!Event.CollisionPos.ContainsNaN());
-					check(!Event.CollisionVelocity.ContainsNaN());
-
-					// TODO add to unique list of physical materials for Blueprint
-					Event.PhysicalMaterialIndex = 0;// Hit->PhysMaterial->GetUniqueID();
-
-					Payloads.Add(Event);
-				}
-			}
-		}
-
-		int32 NumInstances = Payloads.Num();
-		if (NumInstances)
-		{
-			// now allocate the data set and write all the event structs
-			//
-			CollisionEventDataSet->Allocate(NumInstances);
-			CollisionEventDataSet->GetDestinationDataChecked().SetNumInstances(NumInstances);
-			//FNiagaraVariable ValidVar(FNiagaraTypeDefinition::GetIntDef(), "Valid");
-			FNiagaraVariable PosVar(FNiagaraTypeDefinition::GetVec3Def(), "CollisionLocation");
-			FNiagaraVariable VelVar(FNiagaraTypeDefinition::GetVec3Def(), "CollisionVelocity");
-			FNiagaraVariable NormVar(FNiagaraTypeDefinition::GetVec3Def(), "CollisionNormal");
-			FNiagaraVariable PhysMatIdxVar(FNiagaraTypeDefinition::GetIntDef(), "PhysicalMaterialIndex");
-			FNiagaraVariable ParticleIndexVar(FNiagaraTypeDefinition::GetIntDef(), "ParticleIndex");
-			//FNiagaraDataSetIterator<int32> ValidItr(*CollisionEventDataSet, ValidVar, 0, true);
-			FNiagaraDataSetAccessor<FVector> PosItr(*CollisionEventDataSet, PosVar);
-			FNiagaraDataSetAccessor<FVector> NormItr(*CollisionEventDataSet, NormVar);
-			FNiagaraDataSetAccessor<FVector> VelItr(*CollisionEventDataSet, VelVar);
-			FNiagaraDataSetAccessor<int32> PhysMatItr(*CollisionEventDataSet, PhysMatIdxVar);
-			FNiagaraDataSetAccessor<int32> ParticleIndexItr(*CollisionEventDataSet, ParticleIndexVar);
-
-			for (int32 i = 0; i < NumInstances; ++i)
-			{
-				SCOPE_CYCLE_COUNTER(STAT_NiagaraEventWrite);
-				FNiagaraCollisionEventPayload& Payload = Payloads[i];
-
-				check(/*ValidItr.IsValid() && */PosItr.IsValid() && VelItr.IsValid() && NormItr.IsValid() && PhysMatItr.IsValid());
-				//ValidItr.Set(1);
-				PosItr.Set(i, Payload.CollisionPos);
-				VelItr.Set(i, Payload.CollisionVelocity);
-				NormItr.Set(i, Payload.CollisionNormal);
-				ParticleIndexItr.Set(i, Payload.ParticleIndex);
-				PhysMatItr.Set(i, 0);
-			}
-		}
-		else
-		{
-			CollisionEventDataSet->ResetBuffers();
-		}
-
-	}
-}
 
 int32 FNiagaraDICollisionQueryBatch::SubmitQuery(FVector Position, FVector Direction, float CollisionSize, float DeltaSeconds)
 {
diff --git a/Engine/Plugins/FX/Niagara/Source/Niagara/Private/NiagaraComponent.cpp b/Engine/Plugins/FX/Niagara/Source/Niagara/Private/NiagaraComponent.cpp
--- a/Engine/Plugins/FX/Niagara/Source/Niagara/Private/NiagaraComponent.cpp
+++ b/Engine/Plugins/FX/Niagara/Source/Niagara/Private/NiagaraComponent.cpp
@@ -123,6 +123,11 @@ FNiagaraSceneProxy::FNiagaraSceneProxy(const UNiagaraComponent* InComponent)
 		CreateRenderers(InComponent);
 		bAlwaysHasVelocity = true;
 		Batcher = SystemInst->GetBatcher();
+
+#if STATS
+		SystemStatID = InComponent->GetAsset()->GetStatID(false, false);
+#endif
+
 	}
 }
 
@@ -183,7 +188,12 @@ void FNiagaraSceneProxy::CreateRenderers(const UNiagaraComponent* Component)
 			for (UNiagaraRendererProperties* Properties : Emitter->GetEnabledRenderers())
 			{
 				RendererSortInfo.Emplace(Properties->SortOrderHint, EmitterRenderers.Num());
-				EmitterRenderers.Add(Properties->CreateEmitterRenderer(FeatureLevel, &EmitterInst.Get()));
+				FNiagaraRenderer* NewRenderer = nullptr;
+				if (Properties->GetIsEnabled() && EmitterInst->GetData().IsInitialized() && !EmitterInst->IsDisabled())
+				{
+					NewRenderer = Properties->CreateEmitterRenderer(FeatureLevel, &EmitterInst.Get());
+				}
+				EmitterRenderers.Add(NewRenderer);
 			}
 		}
 	}
@@ -304,6 +314,11 @@ void FNiagaraSceneProxy::GetDynamicMeshElements(const TArray<const FSceneView*>&
 {
 	SCOPE_CYCLE_COUNTER(STAT_NiagaraOverview_RT);
 	SCOPE_CYCLE_COUNTER(STAT_NiagaraComponentGetDynamicMeshElements);
+
+#if STATS
+	FScopeCycleCounter SystemStatCounter(SystemStatID);
+#endif
+
 	for (int32 RendererIdx : RendererDrawOrder)
 	{
 		FNiagaraRenderer* Renderer = EmitterRenderers[RendererIdx];
diff --git a/Engine/Plugins/FX/Niagara/Source/Niagara/Private/NiagaraDataSet.cpp b/Engine/Plugins/FX/Niagara/Source/Niagara/Private/NiagaraDataSet.cpp
--- a/Engine/Plugins/FX/Niagara/Source/Niagara/Private/NiagaraDataSet.cpp
+++ b/Engine/Plugins/FX/Niagara/Source/Niagara/Private/NiagaraDataSet.cpp
@@ -67,11 +67,7 @@ static FAutoConsoleVariableRef CVarNiagaraRenderBufferShrinkFactor(
 );
 
 FNiagaraDataSet::FNiagaraDataSet()
-	: TotalFloatComponents(0)
-	, TotalInt32Components(0)
-	, SimTarget(ENiagaraSimTarget::CPUSim)
-	, bFinalized(0)
-	, bNeedsPersistentIDs(0)
+	: CompiledData(&FNiagaraDataSetCompiledData::DummyCompiledData)
 	, NumFreeIDs(0)
 	, MaxUsedID(0)
 	, IDAcquireTag(0)
@@ -87,69 +83,22 @@ FNiagaraDataSet::~FNiagaraDataSet()
 	ReleaseBuffers();
 }
 
-#if !UE_BUILD_SHIPPING && !UE_BUILD_TEST
-void FNiagaraDataSet::Init(FNiagaraDataSetID InID, ENiagaraSimTarget InSimTarget, const FString& InDebugName)
-{
-	Reset();
-	ID = InID;
-	SimTarget = InSimTarget;
-	DebugName = InDebugName;
-}
-
-void FNiagaraDataSet::InitFromCompiledData(const FNiagaraDataSetCompiledData& InDataSetCompiledData, const FString& InDebugName)
-{
-	Variables = InDataSetCompiledData.Variables;
-	VariableLayouts = InDataSetCompiledData.VariableLayouts;
-	TotalFloatComponents = InDataSetCompiledData.TotalFloatComponents;
-	TotalInt32Components = InDataSetCompiledData.TotalInt32Components;
-	bNeedsPersistentIDs = InDataSetCompiledData.bNeedsPersistentIDs;
-	ID = InDataSetCompiledData.ID;
-	SimTarget = InDataSetCompiledData.SimTarget;
-	DebugName = InDebugName;
-	bFinalized = true;
-}
-#else
-void FNiagaraDataSet::Init(FNiagaraDataSetID InID, ENiagaraSimTarget InSimTarget)
-{
-	Reset();
-	ID = InID;
-	SimTarget = InSimTarget;
-}
-
-void FNiagaraDataSet::InitFromCompiledData(const FNiagaraDataSetCompiledData& InDataSetCompiledData)
-{
-	Variables = InDataSetCompiledData.Variables;
-	VariableLayouts = InDataSetCompiledData.VariableLayouts;
-	TotalFloatComponents = InDataSetCompiledData.TotalFloatComponents;
-	TotalInt32Components = InDataSetCompiledData.TotalInt32Components;
-	bNeedsPersistentIDs = InDataSetCompiledData.bNeedsPersistentIDs;
-	ID = InDataSetCompiledData.ID;
-	SimTarget = InDataSetCompiledData.SimTarget;
-	bFinalized = true;
-}
-#endif
-
 void FNiagaraDataSet::Reset()
 {
 	ResetBuffers();
-
-	Variables.Empty();
-	VariableLayouts.Empty();
-	bFinalized = false;
-	TotalFloatComponents = 0;
-	TotalInt32Components = 0;
-	bNeedsPersistentIDs = 0;
 }
 
 void FNiagaraDataSet::ResetBuffers()
 {
-	if (SimTarget == ENiagaraSimTarget::CPUSim)
+	checkSlow(CompiledData);
+	
+	if (GetSimTarget() == ENiagaraSimTarget::CPUSim)
 	{
 		ResetBuffersInternal();
 	}
 	else
 	{
-		check(SimTarget == ENiagaraSimTarget::GPUComputeSim);
+		checkSlow(GetSimTarget() == ENiagaraSimTarget::GPUComputeSim);
 		ENQUEUE_RENDER_COMMAND(ResetBuffersCommand)([=](FRHICommandListImmediate& RHICmdList)
 		{
 			ResetBuffersInternal();
@@ -171,6 +120,11 @@ void FNiagaraDataSet::ResetBuffersInternal()
 	//Ensure we have a valid current buffer
 	BeginSimulate();
 	EndSimulate();
+
+	if (GetSimTarget() == ENiagaraSimTarget::CPUSim)
+	{
+		GetCurrentDataChecked().BuildRegisterTable();
+	}
 }
 
 void FNiagaraDataSet::ReleaseBuffers()
@@ -207,8 +161,8 @@ FNiagaraDataBuffer& FNiagaraDataSet::BeginSimulate()
 	{
 		Data.Add(new FNiagaraDataBuffer(this));
 		DestinationData = Data.Last();
-		verify(DestinationData->TryLock());
-		check(DestinationData->IsBeingWritten());
+		verifySlow(DestinationData->TryLock());
+		checkSlow(DestinationData->IsBeingWritten());
 	}
 
 	DestinationData->SetNumInstances(0);
@@ -222,7 +176,7 @@ void FNiagaraDataSet::EndSimulate(bool SetCurrentData)
 	//CheckCorrectThread();
 	//Destination is now complete so make it the current simulation state.
 	DestinationData->Unlock();
-	check(!DestinationData->IsInUse());
+	checkSlow(!DestinationData->IsInUse());
 
 	if (SetCurrentData)
 	{
@@ -235,9 +189,9 @@ void FNiagaraDataSet::EndSimulate(bool SetCurrentData)
 
 void FNiagaraDataSet::Allocate(int32 NumInstances, bool bMaintainExisting)
 {
-	check(bFinalized);
+	checkSlow(IsInitialized());
 	CheckCorrectThread();
-	check(DestinationData);
+	checkSlow(DestinationData);
 
 	DestinationData->GetIDTable().Reset();
 	if (bMaintainExisting)
@@ -251,7 +205,7 @@ void FNiagaraDataSet::Allocate(int32 NumInstances, bool bMaintainExisting)
 	CheckForNaNs();
 #endif
 
-	if (bNeedsPersistentIDs)
+	if (GetNeedsPersistentIDs())
 	{
 		TArray<int32>& CurrentIDTable = CurrentData->GetIDTable();
 		TArray<int32>& DestinationIDTable = DestinationData->GetIDTable();
@@ -382,57 +336,10 @@ void FNiagaraDataSet::ReleaseGPUInstanceCounts(FNiagaraGPUInstanceCountManager&
 	}
 }
 
-// note, this method is also implemented in FNiagaraDataSetCompiledData::BuildLayout()
-void FNiagaraDataSet::BuildLayout()
-{
-	VariableLayouts.Empty();
-	TotalFloatComponents = 0;
-	TotalInt32Components = 0;
-
-	VariableLayouts.Reserve(Variables.Num());
-	for (FNiagaraVariable& Var : Variables)
-	{
-		FNiagaraVariableLayoutInfo& VarInfo = VariableLayouts[VariableLayouts.AddDefaulted()];
-		FNiagaraTypeLayoutInfo::GenerateLayoutInfo(VarInfo.LayoutInfo, Var.GetType().GetScriptStruct());
-		VarInfo.FloatComponentStart = TotalFloatComponents;
-		VarInfo.Int32ComponentStart = TotalInt32Components;
-		TotalFloatComponents += VarInfo.GetNumFloatComponents();
-		TotalInt32Components += VarInfo.GetNumInt32Components();
-	}
-}
-
-void FNiagaraDataSet::AddVariable(FNiagaraVariable& Variable)
-{
-	check(!bFinalized);
-	Variables.AddUnique(Variable);
-}
-
-void FNiagaraDataSet::AddVariables(const TArray<FNiagaraVariable>& Vars)
-{
-	check(!bFinalized);
-	for (const FNiagaraVariable& Var : Vars)
-	{
-		Variables.AddUnique(Var);
-	}
-}
-
-void FNiagaraDataSet::Finalize()
-{
-	check(!bFinalized);
-	bFinalized = true;
-	BuildLayout();
-
-	ResetBuffers();
-	if (SimTarget == ENiagaraSimTarget::CPUSim)
-	{
-		GetCurrentDataChecked().BuildRegisterTable();
-	}
-}
-
 const FNiagaraVariableLayoutInfo* FNiagaraDataSet::GetVariableLayout(const FNiagaraVariable& Var)const
 {
-	int32 VarLayoutIndex = Variables.IndexOfByKey(Var);
-	return VarLayoutIndex != INDEX_NONE ? &VariableLayouts[VarLayoutIndex] : nullptr;
+	int32 VarLayoutIndex = GetVariables().IndexOfByKey(Var);
+	return VarLayoutIndex != INDEX_NONE ? &CompiledData->VariableLayouts[VarLayoutIndex] : nullptr;
 }
 
 bool FNiagaraDataSet::GetVariableComponentOffsets(const FNiagaraVariable& Var, int32 &FloatStart, int32 &IntStart) const
@@ -452,20 +359,18 @@ bool FNiagaraDataSet::GetVariableComponentOffsets(const FNiagaraVariable& Var, i
 
 void FNiagaraDataSet::CopyTo(FNiagaraDataSet& Other, int32 StartIdx, int32 NumInstances, bool bResetOther)const
 {
+	check(CompiledData);
+
 	CheckCorrectThread();
 
 	if (bResetOther)
 	{
+		Other.CompiledData = CompiledData;
 		Other.Reset();
-		Other.Variables = Variables;
-		Other.VariableLayouts = VariableLayouts;
-		Other.TotalFloatComponents = TotalFloatComponents;
-		Other.TotalInt32Components = TotalInt32Components;
-		Other.Finalize();
 	}
 	else
 	{
-		checkSlow(Other.GetVariables() == Variables);
+		checkSlow(Other.GetVariables() == GetVariables());
 	}
 
 	//Read the most current data. Even if it's possibly partially complete simulation data.
@@ -510,7 +415,7 @@ void FNiagaraDataSet::CopyTo(FNiagaraDataSet& Other, int32 StartIdx, int32 NumIn
 void FNiagaraDataSet::CopyFromGPUReadback(float* GPUReadBackFloat, int* GPUReadBackInt, int32 StartIdx /* = 0 */, int32 NumInstances /* = INDEX_NONE */, uint32 FloatStride, uint32 IntStride)
 {
 	check(IsInRenderingThread());
-	check(bFinalized);//We should be finalized with proper layout information already.
+	check(IsInitialized());//We should be finalized with proper layout information already.
 
 	FNiagaraDataBuffer& DestBuffer = BeginSimulate();
 	DestBuffer.GPUCopyFrom(GPUReadBackFloat, GPUReadBackInt, StartIdx, NumInstances, FloatStride, IntStride);
@@ -555,13 +460,15 @@ int32 FNiagaraDataBuffer::TransferInstance(FNiagaraDataBuffer& SourceBuffer, int
 		SetNumInstances(OldNumInstances + 1);
 
 		/** Copy the instance data. */
-		for (int32 CompIdx = (int32)Owner->TotalFloatComponents - 1; CompIdx >= 0; --CompIdx)
+		int32 FloatComponents = Owner->GetNumFloatComponents();
+		for (int32 CompIdx = FloatComponents - 1; CompIdx >= 0; --CompIdx)
 		{
 			float* Src = SourceBuffer.GetInstancePtrFloat(CompIdx, InstanceIndex);
 			float* Dst = GetInstancePtrFloat(CompIdx, OldNumInstances);
 			*Dst = *Src;
 		}
-		for (int32 CompIdx = (int32)Owner->TotalInt32Components - 1; CompIdx >= 0; --CompIdx)
+		int32 IntComponents = Owner->GetNumInt32Components();
+		for (int32 CompIdx = IntComponents - 1; CompIdx >= 0; --CompIdx)
 		{
 			int32* Src = SourceBuffer.GetInstancePtrInt32(CompIdx, InstanceIndex);
 			int32* Dst = GetInstancePtrInt32(CompIdx, OldNumInstances);
@@ -599,7 +506,7 @@ bool FNiagaraDataBuffer::CheckForNaNs()const
 void FNiagaraDataBuffer::Allocate(uint32 InNumInstances, bool bMaintainExisting)
 {
 	//CheckUsage(false);
-	check(Owner->SimTarget == ENiagaraSimTarget::CPUSim);
+	checkSlow(Owner->GetSimTarget() == ENiagaraSimTarget::CPUSim);
 
 	NumInstancesAllocated = InNumInstances;
 	NumInstances = 0;
@@ -639,18 +546,21 @@ void FNiagaraDataBuffer::Allocate(uint32 InNumInstances, bool bMaintainExisting)
 	{
 		if (FloatStride != OldFloatStride && FloatStride > 0 && OldFloatStride > 0)
 		{
+			int32 FloatComponents = Owner->GetNumFloatComponents();
 			const uint32 BytesToCopy = FMath::Min(OldFloatStride, FloatStride);
-			for (int32 CompIdx = (int32)Owner->TotalFloatComponents-1; CompIdx >= 0; --CompIdx)
+			for (int32 CompIdx = FloatComponents - 1; CompIdx >= 0; --CompIdx)
 			{
 				uint8* Src = OldFloatData.GetData() + OldFloatStride * CompIdx;
 				uint8* Dst = FloatData.GetData() + FloatStride * CompIdx;
 				FMemory::Memcpy(Dst, Src, BytesToCopy);
 			}
 		}
+
 		if (Int32Stride != OldInt32Stride && Int32Stride > 0 && OldInt32Stride > 0)
 		{
+			int32 IntComponents = Owner->GetNumInt32Components();
 			const uint32 BytesToCopy = FMath::Min(OldInt32Stride, Int32Stride);
-			for (int32 CompIdx = (int32)Owner->TotalInt32Components - 1; CompIdx >= 0; --CompIdx)
+			for (int32 CompIdx = IntComponents - 1; CompIdx >= 0; --CompIdx)
 			{
 				uint8* Src = OldIntData.GetData() + OldInt32Stride * CompIdx;
 				uint8* Dst = Int32Data.GetData() + Int32Stride * CompIdx;
@@ -671,7 +581,7 @@ void FNiagaraDataBuffer::AllocateGPU(uint32 InNumInstances, FNiagaraGPUInstanceC
 {
 	CheckUsage(false);
 
-	check(Owner->SimTarget == ENiagaraSimTarget::GPUComputeSim);
+	checkSlow(Owner->GetSimTarget() == ENiagaraSimTarget::GPUComputeSim);
 
 	// Release previous entry if any.
 	GPUInstanceCountManager.FreeEntry(GPUInstanceCountBufferOffset);
@@ -707,11 +617,8 @@ void FNiagaraDataBuffer::AllocateGPU(uint32 InNumInstances, FNiagaraGPUInstanceC
 			{
 				GPUBufferFloat.Release();
 			}
-#if !UE_BUILD_SHIPPING && !UE_BUILD_TEST 
-			GPUBufferFloat.Initialize(sizeof(float), NumElementsToAlloc * Owner->GetNumFloatComponents(), EPixelFormat::PF_R32_FLOAT, BUF_Static, *Owner->DebugName );
-#else
 			GPUBufferFloat.Initialize(sizeof(float), NumElementsToAlloc * Owner->GetNumFloatComponents(), EPixelFormat::PF_R32_FLOAT, BUF_Static);
-#endif
+
 		}
 		if (Owner->GetNumInt32Components())
 		{
@@ -719,11 +626,7 @@ void FNiagaraDataBuffer::AllocateGPU(uint32 InNumInstances, FNiagaraGPUInstanceC
 			{
 				GPUBufferInt.Release();
 			}
-#if !UE_BUILD_SHIPPING && !UE_BUILD_TEST 
-			GPUBufferInt.Initialize(sizeof(int32), NumElementsToAlloc * Owner->GetNumInt32Components(), EPixelFormat::PF_R32_SINT, BUF_Static, *Owner->DebugName);
-#else
 			GPUBufferInt.Initialize(sizeof(int32), NumElementsToAlloc * Owner->GetNumInt32Components(), EPixelFormat::PF_R32_SINT, BUF_Static);
-#endif
 		}
 	}
 }
@@ -732,7 +635,8 @@ void FNiagaraDataBuffer::SwapInstances(uint32 OldIndex, uint32 NewIndex)
 {
 	CheckUsage(false);
 
-	for (uint32 CompIdx = 0; CompIdx < Owner->TotalFloatComponents; ++CompIdx)
+	uint32 FloatComponents = Owner->GetNumFloatComponents();
+	for (uint32 CompIdx = 0; CompIdx < FloatComponents; ++CompIdx)
 	{
 		float* Src = GetInstancePtrFloat(CompIdx, OldIndex);
 		float* Dst = GetInstancePtrFloat(CompIdx, NewIndex);
@@ -740,7 +644,8 @@ void FNiagaraDataBuffer::SwapInstances(uint32 OldIndex, uint32 NewIndex)
 		*Dst = *Src;
 		*Src = Temp;
 	}
-	for (uint32 CompIdx = 0; CompIdx < Owner->TotalInt32Components; ++CompIdx)
+	uint32 IntComponents = Owner->GetNumInt32Components();
+	for (uint32 CompIdx = 0; CompIdx < IntComponents; ++CompIdx)
 	{
 		int32* Src = GetInstancePtrInt32(CompIdx, OldIndex);
 		int32* Dst = GetInstancePtrInt32(CompIdx, NewIndex);
@@ -756,13 +661,15 @@ void FNiagaraDataBuffer::KillInstance(uint32 InstanceIdx)
 	check(InstanceIdx < NumInstances);
 	--NumInstances;
 
-	for (uint32 CompIdx = 0; CompIdx < Owner->TotalFloatComponents; ++CompIdx)
+	uint32 FloatComponents = Owner->GetNumFloatComponents();
+	for (uint32 CompIdx = 0; CompIdx < FloatComponents; ++CompIdx)
 	{
 		float* Src = GetInstancePtrFloat(CompIdx, NumInstances);
 		float* Dst = GetInstancePtrFloat(CompIdx, InstanceIdx);
 		*Dst = *Src;
 	}
-	for (uint32 CompIdx = 0; CompIdx < Owner->TotalInt32Components; ++CompIdx)
+	uint32 IntComponents = Owner->GetNumInt32Components();
+	for (uint32 CompIdx = 0; CompIdx < IntComponents; ++CompIdx)
 	{
 		int32* Src = GetInstancePtrInt32(CompIdx, NumInstances);
 		int32* Dst = GetInstancePtrInt32(CompIdx, InstanceIdx);
@@ -798,7 +705,8 @@ void FNiagaraDataBuffer::CopyTo(FNiagaraDataBuffer& DestBuffer, int32 StartIdx,
 		}
 		DestBuffer.SetNumInstances(NewNumInstances);
 
-		for (uint32 CompIdx = 0; CompIdx < Owner->TotalFloatComponents; ++CompIdx)
+		uint32 FloatComponents = Owner->GetNumFloatComponents();
+		for (uint32 CompIdx = 0; CompIdx < FloatComponents; ++CompIdx)
 		{
 			const float* SrcStart = GetInstancePtrFloat(CompIdx, StartIdx);
 			const float* SrcEnd = GetInstancePtrFloat(CompIdx, StartIdx + InstancesToCopy);
@@ -814,7 +722,8 @@ void FNiagaraDataBuffer::CopyTo(FNiagaraDataBuffer& DestBuffer, int32 StartIdx,
 				}
 			}
 		}
-		for (uint32 CompIdx = 0; CompIdx < Owner->TotalInt32Components; ++CompIdx)
+		uint32 IntComponents = Owner->GetNumInt32Components();
+		for (uint32 CompIdx = 0; CompIdx < IntComponents; ++CompIdx)
 		{
 			const int32* SrcStart = GetInstancePtrInt32(CompIdx, StartIdx);
 			const int32* SrcEnd = GetInstancePtrInt32(CompIdx, StartIdx + InstancesToCopy);
@@ -847,7 +756,8 @@ void FNiagaraDataBuffer::GPUCopyFrom(float* GPUReadBackFloat, int* GPUReadBackIn
 
 	if (GPUReadBackFloat)
 	{
-		for (uint32 CompIdx = 0; CompIdx < Owner->TotalFloatComponents; ++CompIdx)
+		uint32 FloatComponents = Owner->GetNumFloatComponents();
+		for (uint32 CompIdx = 0; CompIdx < FloatComponents; ++CompIdx)
 		{
 			// We have to reimplement the logic from GetInstancePtrFloat here because the incoming stride may be different than this 
 			// data buffer's stride.
@@ -868,7 +778,8 @@ void FNiagaraDataBuffer::GPUCopyFrom(float* GPUReadBackFloat, int* GPUReadBackIn
 	}
 	if (GPUReadBackInt)
 	{
-		for (uint32 CompIdx = 0; CompIdx < Owner->TotalInt32Components; ++CompIdx)
+		uint32 IntComponents = Owner->GetNumInt32Components();
+		for (uint32 CompIdx = 0; CompIdx < IntComponents; ++CompIdx)
 		{
 			// We have to reimplement the logic from GetInstancePtrInt here because the incoming stride may be different than this 
 			// data buffer's stride.
@@ -904,7 +815,7 @@ void FNiagaraDataBuffer::CopyTo(FNiagaraDataBuffer& DestBuffer)const
 
 void FNiagaraDataBuffer::Dump(int32 StartIndex, int32 InNumInstances, const FString& Label)const
 {
-	TArray<FNiagaraVariable>& Variables = Owner->GetVariables();
+	TArray<FNiagaraVariable> Variables = Owner->GetVariables();
 	FNiagaraDataVariableIterator Itr(this, StartIndex);
 	Itr.AddVariables(Variables);
 
@@ -922,7 +833,7 @@ void FNiagaraDataBuffer::Dump(int32 StartIndex, int32 InNumInstances, const FStr
 		Itr.Get();
 
 		FString Line = TEXT("| ");
-		for (FNiagaraVariable& Var : Owner->GetVariables())
+		for (const FNiagaraVariable& Var : Owner->GetVariables())
 		{
 			Line += Var.ToString() + TEXT(" | ");
 		}
@@ -1120,7 +1031,7 @@ void FNiagaraDataBuffer::BuildRegisterTable()
 	RegisterTable.Reset();
 	RegisterTable.SetNumUninitialized(TotalRegisters);
 	int32 NumRegisters = 0;
-	for (FNiagaraVariableLayoutInfo& VarLayout : Owner->VariableLayouts)
+	for (const FNiagaraVariableLayoutInfo& VarLayout : Owner->GetVariableLayouts())
 	{
 		int32 NumFloats = VarLayout.GetNumFloatComponents();
 		int32 NumInts = VarLayout.GetNumInt32Components();
@@ -1140,6 +1051,10 @@ void FNiagaraDataBuffer::BuildRegisterTable()
 	}
 }
 
+//////////////////////////////////////////////////////////////////////////
+
+FNiagaraDataSetCompiledData FNiagaraDataSetCompiledData::DummyCompiledData;
+
 // note, this method is also implemented in FNiagaraDataSet::BuildLayout()
 void FNiagaraDataSetCompiledData::BuildLayout()
 {
@@ -1159,6 +1074,11 @@ void FNiagaraDataSetCompiledData::BuildLayout()
 	}
 }
 
+FNiagaraDataSetCompiledData::FNiagaraDataSetCompiledData()
+{
+	Empty();
+}
+
 void FNiagaraDataSetCompiledData::Empty()
 {
 	bNeedsPersistentIDs = 0;
diff --git a/Engine/Plugins/FX/Niagara/Source/Niagara/Private/NiagaraEmitter.cpp b/Engine/Plugins/FX/Niagara/Source/Niagara/Private/NiagaraEmitter.cpp
--- a/Engine/Plugins/FX/Niagara/Source/Niagara/Private/NiagaraEmitter.cpp
+++ b/Engine/Plugins/FX/Niagara/Source/Niagara/Private/NiagaraEmitter.cpp
@@ -98,6 +98,7 @@ bool FNiagaraEmitterScriptProperties::DataSetAccessSynchronized() const
 
 UNiagaraEmitter::UNiagaraEmitter(const FObjectInitializer& Initializer)
 : Super(Initializer)
+, PreAllocationCount(0)
 , FixedBounds(FBox(FVector(-100), FVector(100)))
 , MinDetailLevel(0)
 , MaxDetailLevel(4)
diff --git a/Engine/Plugins/FX/Niagara/Source/Niagara/Private/NiagaraEmitterInstance.cpp b/Engine/Plugins/FX/Niagara/Source/Niagara/Private/NiagaraEmitterInstance.cpp
--- a/Engine/Plugins/FX/Niagara/Source/Niagara/Private/NiagaraEmitterInstance.cpp
+++ b/Engine/Plugins/FX/Niagara/Source/Niagara/Private/NiagaraEmitterInstance.cpp
@@ -25,6 +25,12 @@ DECLARE_CYCLE_STAT(TEXT("Emitter Event Handling [CNC]"), STAT_NiagaraEventHandle
 DECLARE_CYCLE_STAT(TEXT("Emitter Error Check [CNC]"), STAT_NiagaraEmitterErrorCheck, STATGROUP_Niagara);
 DECLARE_CYCLE_STAT(TEXT("Init Emitters [GT]"), STAT_NiagaraEmitterInit, STATGROUP_Niagara);
 
+DECLARE_CYCLE_STAT(TEXT("Init Emitters (DataSets) [GT]"), STAT_NiagaraEmitterInit_Datasets, STATGROUP_Niagara);
+DECLARE_CYCLE_STAT(TEXT("Init Emitters (ExecContexts) [GT]"), STAT_NiagaraEmitterInit_ExecContexts, STATGROUP_Niagara);
+DECLARE_CYCLE_STAT(TEXT("Init Emitters (DirectBindings) [GT]"), STAT_NiagaraEmitterInit_DirectBindings, STATGROUP_Niagara);
+DECLARE_CYCLE_STAT(TEXT("Init Emitters (DIDefaultsAndBoundCalcs) [GT]"), STAT_NiagaraEmitterInit_DIDefaultsAndBoundsCalcs, STATGROUP_Niagara);
+
+
 static int32 GbDumpParticleData = 0;
 static FAutoConsoleVariableRef CVarNiagaraDumpParticleData(
 	TEXT("fx.DumpParticleData"),
@@ -118,7 +124,7 @@ FNiagaraEmitterInstance::~FNiagaraEmitterInstance()
 
 	//UE_LOG(LogNiagara, Warning, TEXT("~Simulator %p"), this);
 	CachedBounds.Init();
-	UnbindParameters();
+	UnbindParameters(false);
 
 	if (GPUExecContext != nullptr)
 	{
@@ -260,166 +266,191 @@ void FNiagaraEmitterInstance::Init(int32 InEmitterIdx, FNiagaraSystemInstanceID
 		SpawnInfos.SetNum(EmitterCompiledData[EmitterIdx].SpawnAttributes.Num());
 	}
 
-#if !UE_BUILD_SHIPPING && !UE_BUILD_TEST 
-	Data.InitFromCompiledData(EmitterCompiledData[EmitterIdx].DataSetCompiledData, ParentSystemInstance->GetSystem()->GetName() + TEXT("/") + CachedEmitter->GetName());
-#else
-	Data.InitFromCompiledData(EmitterCompiledData[EmitterIdx].DataSetCompiledData);
-#endif
+	{
+		SCOPE_CYCLE_COUNTER(STAT_NiagaraEmitterInit_Datasets);
+		Data.Init(&EmitterCompiledData[EmitterIdx].DataSetCompiledData);
 
-	ResetSimulation();
+		ResetSimulation();
 
-	DataSetMap.Empty();
+		DataSetMap.Empty();
 
-	//Add the particle data to the data set map.
-	//Currently just used for the tick loop but will also allow access directly to the particle data from other emitters.
-	DataSetMap.Add(Data.GetID(), &Data);
-	//Warn the user if there are any attributes used in the update script that are not initialized in the spawn script.
-	//TODO: We need some window in the System editor and possibly the graph editor for warnings and errors.
+		//Add the particle data to the data set map.
+		//Currently just used for the tick loop but will also allow access directly to the particle data from other emitters.
+		DataSetMap.Add(Data.GetID(), &Data);
+		//Warn the user if there are any attributes used in the update script that are not initialized in the spawn script.
+		//TODO: We need some window in the System editor and possibly the graph editor for warnings and errors.
 
-	const bool bVerboseAttributeLogging = false;
+		const bool bVerboseAttributeLogging = false;
 
 #if !UE_BUILD_SHIPPING && !UE_BUILD_TEST 
-	if (bVerboseAttributeLogging)
-	{
-		for (FNiagaraVariable& Attr : CachedEmitter->UpdateScriptProps.Script->GetVMExecutableData().Attributes)
+		if (bVerboseAttributeLogging)
 		{
-			int32 FoundIdx;
-			if (!CachedEmitter->SpawnScriptProps.Script->GetVMExecutableData().Attributes.Find(Attr, FoundIdx))
+			for (FNiagaraVariable& Attr : CachedEmitter->UpdateScriptProps.Script->GetVMExecutableData().Attributes)
 			{
-				UE_LOG(LogNiagara, Warning, TEXT("Attribute %s is used in the Update script for %s but it is not initialised in the Spawn script!"), *Attr.GetName().ToString(), *EmitterHandle.GetName().ToString());
-			}
-			for (int32 i = 0; i < CachedEmitter->GetEventHandlers().Num(); i++)
-			{
-				if (CachedEmitter->GetEventHandlers()[i].Script && !CachedEmitter->GetEventHandlers()[i].Script->GetVMExecutableData().Attributes.Find(Attr, FoundIdx))
+				int32 FoundIdx;
+				if (!CachedEmitter->SpawnScriptProps.Script->GetVMExecutableData().Attributes.Find(Attr, FoundIdx))
+				{
+					UE_LOG(LogNiagara, Warning, TEXT("Attribute %s is used in the Update script for %s but it is not initialised in the Spawn script!"), *Attr.GetName().ToString(), *EmitterHandle.GetName().ToString());
+				}
+				for (int32 i = 0; i < CachedEmitter->GetEventHandlers().Num(); i++)
 				{
-					UE_LOG(LogNiagara, Warning, TEXT("Attribute %s is used in the event handler script for %s but it is not initialised in the Spawn script!"), *Attr.GetName().ToString(), *EmitterHandle.GetName().ToString());
+					if (CachedEmitter->GetEventHandlers()[i].Script && !CachedEmitter->GetEventHandlers()[i].Script->GetVMExecutableData().Attributes.Find(Attr, FoundIdx))
+					{
+						UE_LOG(LogNiagara, Warning, TEXT("Attribute %s is used in the event handler script for %s but it is not initialised in the Spawn script!"), *Attr.GetName().ToString(), *EmitterHandle.GetName().ToString());
+					}
 				}
 			}
 		}
-	}
 #endif
 
-	ensure(CachedEmitter->UpdateScriptProps.DataSetAccessSynchronized());
-	UpdateScriptEventDataSets.Empty();
-	UpdateEventGeneratorIsSharedByIndex.SetNumZeroed(CachedEmitter->UpdateScriptProps.EventGenerators.Num());
-	int32 UpdateEventGeneratorIndex = 0;
-	for (const FNiagaraEventGeneratorProperties &GeneratorProps : CachedEmitter->UpdateScriptProps.EventGenerators)
-	{
-		FNiagaraDataSet *Set = FNiagaraEventDataSetMgr::CreateEventDataSet(ParentSystemInstance->GetId(), EmitterHandle.GetIdName(), GeneratorProps.SetProps.ID.Name);
-#if !UE_BUILD_SHIPPING && !UE_BUILD_TEST 
-		Set->Init(FNiagaraDataSetID(), ENiagaraSimTarget::CPUSim, CachedEmitter->GetFullName() + TEXT("/") + GeneratorProps.SetProps.ID.Name.ToString());
-#else
-		Set->Init(FNiagaraDataSetID(), ENiagaraSimTarget::CPUSim);
-#endif
-		Set->AddVariables(GeneratorProps.SetProps.Variables);
-		Set->Finalize();
-		UpdateScriptEventDataSets.Add(Set);
-		UpdateEventGeneratorIsSharedByIndex[UpdateEventGeneratorIndex] = CachedEmitter->IsEventGeneratorShared(GeneratorProps.ID);
-		UpdateEventGeneratorIndex++;
-	}
+		ensure(CachedEmitter->UpdateScriptProps.DataSetAccessSynchronized());
+		UpdateScriptEventDataSets.Empty();
+		UpdateEventGeneratorIsSharedByIndex.SetNumZeroed(CachedEmitter->UpdateScriptProps.EventGenerators.Num());
+		int32 UpdateEventGeneratorIndex = 0;
+		for (const FNiagaraEventGeneratorProperties &GeneratorProps : CachedEmitter->UpdateScriptProps.EventGenerators)
+		{
+			FNiagaraDataSet *Set = FNiagaraEventDataSetMgr::CreateEventDataSet(ParentSystemInstance->GetId(), EmitterHandle.GetIdName(), GeneratorProps.ID);
+			Set->Init(&GeneratorProps.DataSetCompiledData);
 
-	ensure(CachedEmitter->SpawnScriptProps.DataSetAccessSynchronized());
-	SpawnScriptEventDataSets.Empty();
-	SpawnEventGeneratorIsSharedByIndex.SetNumZeroed(CachedEmitter->SpawnScriptProps.EventGenerators.Num());
-	int32 SpawnEventGeneratorIndex = 0;
-	for (const FNiagaraEventGeneratorProperties &GeneratorProps : CachedEmitter->SpawnScriptProps.EventGenerators)
-	{
-		FNiagaraDataSet *Set = FNiagaraEventDataSetMgr::CreateEventDataSet(ParentSystemInstance->GetId(), EmitterHandle.GetIdName(), GeneratorProps.SetProps.ID.Name);
-#if !UE_BUILD_SHIPPING && !UE_BUILD_TEST 
-		Set->Init(FNiagaraDataSetID(), ENiagaraSimTarget::CPUSim, CachedEmitter->GetFullName() + TEXT("/") + GeneratorProps.SetProps.ID.Name.ToString());
-#else
-		Set->Init(FNiagaraDataSetID(), ENiagaraSimTarget::CPUSim);
-#endif
-		Set->AddVariables(GeneratorProps.SetProps.Variables);
-		Set->Finalize();
-		SpawnScriptEventDataSets.Add(Set);
-		SpawnEventGeneratorIsSharedByIndex[SpawnEventGeneratorIndex] = CachedEmitter->IsEventGeneratorShared(GeneratorProps.ID);
-		SpawnEventGeneratorIndex++;
-	}
+			UpdateScriptEventDataSets.Add(Set);
+			UpdateEventGeneratorIsSharedByIndex[UpdateEventGeneratorIndex] = CachedEmitter->IsEventGeneratorShared(GeneratorProps.ID);
+			UpdateEventGeneratorIndex++;
+		}
 
-	SpawnExecContext.Init(CachedEmitter->SpawnScriptProps.Script, CachedEmitter->SimTarget);
-	UpdateExecContext.Init(CachedEmitter->UpdateScriptProps.Script, CachedEmitter->SimTarget);
+		ensure(CachedEmitter->SpawnScriptProps.DataSetAccessSynchronized());
+		SpawnScriptEventDataSets.Empty();
+		SpawnEventGeneratorIsSharedByIndex.SetNumZeroed(CachedEmitter->SpawnScriptProps.EventGenerators.Num());
+		int32 SpawnEventGeneratorIndex = 0;
+		for (const FNiagaraEventGeneratorProperties &GeneratorProps : CachedEmitter->SpawnScriptProps.EventGenerators)
+		{
+			FNiagaraDataSet *Set = FNiagaraEventDataSetMgr::CreateEventDataSet(ParentSystemInstance->GetId(), EmitterHandle.GetIdName(), GeneratorProps.ID);
 
-	// setup the parameer store for the GPU execution context; since spawn and update are combined here, we build one with params from both script props
-	if (CachedEmitter->SimTarget == ENiagaraSimTarget::GPUComputeSim)
-	{
-		GPUExecContext = new FNiagaraComputeExecutionContext();
-		GPUExecContext->InitParams(CachedEmitter->GetGPUComputeScript(), CachedEmitter->SimTarget, CachedEmitter->GetUniqueEmitterName(), CachedEmitter->MaxUpdateIterations, CachedEmitter->SpawnStages);
-		GPUExecContext->MainDataSet = &Data;
-		GPUExecContext->GPUScript_RT = CachedEmitter->GetGPUComputeScript()->GetRenderThreadScript();
+			Set->Init(&GeneratorProps.DataSetCompiledData);
 
-		SpawnExecContext.Parameters.Bind(&GPUExecContext->CombinedParamStore);
-		UpdateExecContext.Parameters.Bind(&GPUExecContext->CombinedParamStore);
+			SpawnScriptEventDataSets.Add(Set);
+			SpawnEventGeneratorIsSharedByIndex[SpawnEventGeneratorIndex] = CachedEmitter->IsEventGeneratorShared(GeneratorProps.ID);
+			SpawnEventGeneratorIndex++;
+		}
 	}
 
-	EventExecContexts.SetNum(CachedEmitter->GetEventHandlers().Num());
-	for (int32 i = 0; i < NumEvents; i++)
 	{
-		ensure(CachedEmitter->GetEventHandlers()[i].DataSetAccessSynchronized());
+		SCOPE_CYCLE_COUNTER(STAT_NiagaraEmitterInit_ExecContexts);
 
-		UNiagaraScript* EventScript = CachedEmitter->GetEventHandlers()[i].Script;
+		SpawnExecContext.Init(CachedEmitter->SpawnScriptProps.Script, CachedEmitter->SimTarget);
+		UpdateExecContext.Init(CachedEmitter->UpdateScriptProps.Script, CachedEmitter->SimTarget);
 
-		//This is cpu explicitly? Are we doing event handlers on GPU?
-		EventExecContexts[i].Init(EventScript, ENiagaraSimTarget::CPUSim);
-	}
+		// setup the parameer store for the GPU execution context; since spawn and update are combined here, we build one with params from both script props
+		if (CachedEmitter->SimTarget == ENiagaraSimTarget::GPUComputeSim)
+		{
+			GPUExecContext = new FNiagaraComputeExecutionContext();
+			GPUExecContext->InitParams(CachedEmitter->GetGPUComputeScript(), CachedEmitter->SimTarget, CachedEmitter->GetUniqueEmitterName(), CachedEmitter->MaxUpdateIterations, CachedEmitter->SpawnStages);
+			GPUExecContext->MainDataSet = &Data;
+			GPUExecContext->GPUScript_RT = CachedEmitter->GetGPUComputeScript()->GetRenderThreadScript();
 
-	//Setup direct bindings for setting parameter values.
-	SpawnIntervalBinding.Init(SpawnExecContext.Parameters, EmitterCompiledData[EmitterIdx].EmitterSpawnIntervalVar);
-	InterpSpawnStartBinding.Init(SpawnExecContext.Parameters, EmitterCompiledData[EmitterIdx].EmitterInterpSpawnStartDTVar);
-	SpawnGroupBinding.Init(SpawnExecContext.Parameters, EmitterCompiledData[EmitterIdx].EmitterSpawnGroupVar);
+			SpawnExecContext.Parameters.Bind(&GPUExecContext->CombinedParamStore);
+			UpdateExecContext.Parameters.Bind(&GPUExecContext->CombinedParamStore);
+		}
+	
+		EventExecContexts.SetNum(CachedEmitter->GetEventHandlers().Num());
+		for (int32 i = 0; i < NumEvents; i++)
+		{
+			ensure(CachedEmitter->GetEventHandlers()[i].DataSetAccessSynchronized());
 
-	SpawnEmitterAgeBinding.Init(SpawnExecContext.Parameters, EmitterCompiledData[EmitterIdx].EmitterAgeVar);
-	UpdateEmitterAgeBinding.Init(UpdateExecContext.Parameters, EmitterCompiledData[EmitterIdx].EmitterAgeVar);
-	EventEmitterAgeBindings.SetNum(NumEvents);
-	for (int32 i = 0; i < NumEvents; i++)
-	{
-		EventEmitterAgeBindings[i].Init(EventExecContexts[i].Parameters, EmitterCompiledData[EmitterIdx].EmitterAgeVar);
-	}
+			UNiagaraScript* EventScript = CachedEmitter->GetEventHandlers()[i].Script;
 
-	if (CachedEmitter->SimTarget == ENiagaraSimTarget::GPUComputeSim && GPUExecContext != nullptr)
-	{
-		EmitterAgeBindingGPU.Init(GPUExecContext->CombinedParamStore, EmitterCompiledData[EmitterIdx].EmitterAgeVar);
+			//This is cpu explicitly? Are we doing event handlers on GPU?
+			EventExecContexts[i].Init(EventScript, ENiagaraSimTarget::CPUSim);
+		}
 	}
 
-	SpawnRandomSeedBinding.Init(SpawnExecContext.Parameters, EmitterCompiledData[EmitterIdx].EmitterRandomSeedVar);
-	UpdateRandomSeedBinding.Init(UpdateExecContext.Parameters, EmitterCompiledData[EmitterIdx].EmitterRandomSeedVar);
-	if (CachedEmitter->SimTarget == ENiagaraSimTarget::GPUComputeSim && GPUExecContext != nullptr)
-	{
-		GPURandomSeedBinding.Init(GPUExecContext->CombinedParamStore, EmitterCompiledData[EmitterIdx].EmitterRandomSeedVar);
-	}
 
-	// Initialize the exec count
-	SpawnExecCountBinding.Init(SpawnExecContext.Parameters, SYS_PARAM_ENGINE_EXEC_COUNT);
-	UpdateExecCountBinding.Init(UpdateExecContext.Parameters, SYS_PARAM_ENGINE_EXEC_COUNT);
-	EventExecCountBindings.SetNum(NumEvents);
-	for (int32 i = 0; i < NumEvents; i++)
 	{
-		EventExecCountBindings[i].Init(EventExecContexts[i].Parameters, SYS_PARAM_ENGINE_EXEC_COUNT);
-	}
+		SCOPE_CYCLE_COUNTER(STAT_NiagaraEmitterInit_DirectBindings);
+		//Setup direct bindings for setting parameter values.
+		SpawnIntervalBinding.Init(SpawnExecContext.Parameters, EmitterCompiledData[EmitterIdx].EmitterSpawnIntervalVar);
+		InterpSpawnStartBinding.Init(SpawnExecContext.Parameters, EmitterCompiledData[EmitterIdx].EmitterInterpSpawnStartDTVar);
+		SpawnGroupBinding.Init(SpawnExecContext.Parameters, EmitterCompiledData[EmitterIdx].EmitterSpawnGroupVar);
 
-	// Collect script defined data interface parameters.
-	TArray<UNiagaraScript*, TInlineAllocator<4>> Scripts;
-	Scripts.Add(CachedEmitter->SpawnScriptProps.Script);
-	Scripts.Add(CachedEmitter->UpdateScriptProps.Script);
-	for (const FNiagaraEventScriptProperties& EventHandler : CachedEmitter->GetEventHandlers())
-	{
-		Scripts.Add(EventHandler.Script);
+		SpawnEmitterAgeBinding.Init(SpawnExecContext.Parameters, EmitterCompiledData[EmitterIdx].EmitterAgeVar);
+		UpdateEmitterAgeBinding.Init(UpdateExecContext.Parameters, EmitterCompiledData[EmitterIdx].EmitterAgeVar);
+		EventEmitterAgeBindings.SetNum(NumEvents);
+		for (int32 i = 0; i < NumEvents; i++)
+		{
+			EventEmitterAgeBindings[i].Init(EventExecContexts[i].Parameters, EmitterCompiledData[EmitterIdx].EmitterAgeVar);
+		}
+
+		if (CachedEmitter->SimTarget == ENiagaraSimTarget::GPUComputeSim && GPUExecContext != nullptr)
+		{
+			EmitterAgeBindingGPU.Init(GPUExecContext->CombinedParamStore, EmitterCompiledData[EmitterIdx].EmitterAgeVar);
+		}
+
+		SpawnRandomSeedBinding.Init(SpawnExecContext.Parameters, EmitterCompiledData[EmitterIdx].EmitterRandomSeedVar);
+		UpdateRandomSeedBinding.Init(UpdateExecContext.Parameters, EmitterCompiledData[EmitterIdx].EmitterRandomSeedVar);
+		if (CachedEmitter->SimTarget == ENiagaraSimTarget::GPUComputeSim && GPUExecContext != nullptr)
+		{
+			GPURandomSeedBinding.Init(GPUExecContext->CombinedParamStore, EmitterCompiledData[EmitterIdx].EmitterRandomSeedVar);
+		}
+
+		// Initialize the exec count
+		SpawnExecCountBinding.Init(SpawnExecContext.Parameters, SYS_PARAM_ENGINE_EXEC_COUNT);
+		UpdateExecCountBinding.Init(UpdateExecContext.Parameters, SYS_PARAM_ENGINE_EXEC_COUNT);
+		EventExecCountBindings.SetNum(NumEvents);
+		for (int32 i = 0; i < NumEvents; i++)
+		{
+			EventExecCountBindings[i].Init(EventExecContexts[i].Parameters, SYS_PARAM_ENGINE_EXEC_COUNT);
+		}
 	}
-	FNiagaraUtilities::CollectScriptDataInterfaceParameters(*CachedEmitter, MakeArrayView(Scripts), ScriptDefinedDataInterfaceParameters);
 
-	// Initialize bounds calculators
-	//-OPT: Could skip creating this if we won't ever use it
-	const TArray<UNiagaraRendererProperties*>& Renderers = CachedEmitter->GetEnabledRenderers();
-	BoundsCalculators.Reserve(Renderers.Num());
-	for (UNiagaraRendererProperties* RendererProperties : Renderers)
 	{
-		FNiagaraBoundsCalculator* BoundsCalculator = RendererProperties->CreateBoundsCalculator();
-		if (BoundsCalculator != nullptr)
+		SCOPE_CYCLE_COUNTER(STAT_NiagaraEmitterInit_DIDefaultsAndBoundsCalcs);
+		// Collect script defined data interface parameters.
+		TArray<UNiagaraScript*, TInlineAllocator<8>> Scripts;
+		Scripts.Add(CachedEmitter->SpawnScriptProps.Script);
+		Scripts.Add(CachedEmitter->UpdateScriptProps.Script);
+		for (const FNiagaraEventScriptProperties& EventHandler : CachedEmitter->GetEventHandlers())
 		{
-			BoundsCalculator->InitAccessors(*ParticleDataSet);
-			BoundsCalculators.Emplace(BoundsCalculator);
+			Scripts.Add(EventHandler.Script);
 		}
-	}
+		FNiagaraUtilities::CollectScriptDataInterfaceParameters(*CachedEmitter, MakeArrayView(Scripts), ScriptDefinedDataInterfaceParameters);
+
+		// Initialize bounds calculators
+		//-OPT: Could skip creating this if we won't ever use it
+		bool bUseDynamicBounds = !CachedSystemFixedBounds.IsSet() && !CachedEmitter->bFixedBounds && CachedEmitter->SimTarget == ENiagaraSimTarget::CPUSim;
+		if (bUseDynamicBounds)
+		{
+			BoundsCalculators.Reserve(CachedEmitter->GetRenderers().Num());
+			for (UNiagaraRendererProperties* RendererProperties : CachedEmitter->GetRenderers())
+			{
+				if ((RendererProperties != nullptr) && RendererProperties->GetIsEnabled())
+				{
+					FNiagaraBoundsCalculator* BoundsCalculator = RendererProperties->CreateBoundsCalculator();
+					if (BoundsCalculator != nullptr)
+					{
+						BoundsCalculator->InitAccessors(*ParticleDataSet);
+						BoundsCalculators.Emplace(BoundsCalculator);
+					}
+				}
+			}
+		}
+
+		//Bind some stores and unbind immediately just to prime some data from those stores.
+		FNiagaraParameterStore& SystemScriptDefinedDataInterfaceParameters = ParentSystemInstance->GetSystemSimulation()->GetScriptDefinedDataInterfaceParameters();
+		
+		SystemScriptDefinedDataInterfaceParameters.Bind(&SpawnExecContext.Parameters);
+		ScriptDefinedDataInterfaceParameters.Bind(&SpawnExecContext.Parameters);
+		SpawnExecContext.Parameters.UnbindFromSourceStores();
+
+		SystemScriptDefinedDataInterfaceParameters.Bind(&UpdateExecContext.Parameters);
+		ScriptDefinedDataInterfaceParameters.Bind(&UpdateExecContext.Parameters);
+		UpdateExecContext.Parameters.UnbindFromSourceStores();
+
+		for (FNiagaraScriptExecutionContext& EventContext : EventExecContexts)
+		{
+			SystemScriptDefinedDataInterfaceParameters.Bind(&EventContext.Parameters);
+			ScriptDefinedDataInterfaceParameters.Bind(&EventContext.Parameters);
+			EventContext.Parameters.UnbindFromSourceStores();
+		}
+	}	
 }
 
 void FNiagaraEmitterInstance::ResetSimulation(bool bKillExisting /*= true*/)
@@ -428,6 +459,12 @@ void FNiagaraEmitterInstance::ResetSimulation(bool bKillExisting /*= true*/)
 	Loops = 0;
 	TickCount = 0;
 	CachedBounds.Init();
+
+	if (IsDisabled())
+	{
+		return;
+	}
+
 	SetExecutionState(ENiagaraExecutionState::Active);
 
 	if (bKillExisting)
@@ -576,31 +613,59 @@ void FNiagaraEmitterInstance::DirtyDataInterfaces()
 // 	}
 // }
 
-void FNiagaraEmitterInstance::UnbindParameters()
+void FNiagaraEmitterInstance::UnbindParameters(bool bExternalOnly)
 {
-	SpawnExecContext.Parameters.UnbindFromSourceStores();
-	UpdateExecContext.Parameters.UnbindFromSourceStores();
-	if (GPUExecContext != nullptr)
+	if (bExternalOnly && !IsDisabled())
 	{
-		GPUExecContext->CombinedParamStore.UnbindFromSourceStores();
-	}
+		for (UNiagaraParameterCollection* Collection : SpawnExecContext.Script->GetCachedParameterCollectionReferences())
+		{
+			if (UNiagaraParameterCollectionInstance* NPCInst = ParentSystemInstance->GetParameterCollectionInstance(Collection))
+			{
+				NPCInst->GetParameterStore().Unbind(&SpawnExecContext.Parameters);
+			}
+		}
+		for (UNiagaraParameterCollection* Collection : UpdateExecContext.Script->GetCachedParameterCollectionReferences())
+		{
+			if (UNiagaraParameterCollectionInstance* NPCInst = ParentSystemInstance->GetParameterCollectionInstance(Collection))
+			{
+				NPCInst->GetParameterStore().Unbind(&UpdateExecContext.Parameters);
+			}
+		}
 
-	for (int32 EventIdx = 0; EventIdx < EventExecContexts.Num(); ++EventIdx)
+		for (int32 EventIdx = 0; EventIdx < EventExecContexts.Num(); ++EventIdx)
+		{
+			for (UNiagaraParameterCollection* Collection : EventExecContexts[EventIdx].Script->GetCachedParameterCollectionReferences())
+			{
+				if (UNiagaraParameterCollectionInstance* NPCInst = ParentSystemInstance->GetParameterCollectionInstance(Collection))
+				{
+					NPCInst->GetParameterStore().Unbind(&EventExecContexts[EventIdx].Parameters);
+				}
+			}
+		}
+	}
+	else
 	{
-		EventExecContexts[EventIdx].Parameters.UnbindFromSourceStores();
+		SpawnExecContext.Parameters.UnbindFromSourceStores();
+		UpdateExecContext.Parameters.UnbindFromSourceStores();
+		if (GPUExecContext != nullptr)
+		{
+			GPUExecContext->CombinedParamStore.UnbindFromSourceStores();
+		}
+
+		for (int32 EventIdx = 0; EventIdx < EventExecContexts.Num(); ++EventIdx)
+		{
+			EventExecContexts[EventIdx].Parameters.UnbindFromSourceStores();
+		}
 	}
 }
 
-void FNiagaraEmitterInstance::BindParameters()
+void FNiagaraEmitterInstance::BindParameters(bool bExternalOnly)
 {
 	if (IsDisabled())
 	{
 		return;
 	}
 
-	FNiagaraWorldManager* WorldMan = ParentSystemInstance->GetWorldManager();
-	check(WorldMan);
-
 	for (UNiagaraParameterCollection* Collection : SpawnExecContext.Script->GetCachedParameterCollectionReferences())
 	{
 		ParentSystemInstance->GetParameterCollectionInstance(Collection)->GetParameterStore().Bind(&SpawnExecContext.Parameters);
@@ -618,39 +683,34 @@ void FNiagaraEmitterInstance::BindParameters()
 		}
 	}
 
-	//Now bind parameters from the component and system.
-	FNiagaraParameterStore& InstanceParams = ParentSystemInstance->GetParameters();
-	FNiagaraParameterStore& SystemScriptDefinedDataInterfaceParameters = ParentSystemInstance->GetSystemSimulation()->GetScriptDefinedDataInterfaceParameters();
+	if (!bExternalOnly)
+	{
+		//Now bind parameters from the component and system.
+		FNiagaraParameterStore& InstanceParams = ParentSystemInstance->GetParameters();		
 	
-	InstanceParams.Bind(&SpawnExecContext.Parameters);
-	SystemScriptDefinedDataInterfaceParameters.Bind(&SpawnExecContext.Parameters);
-	ScriptDefinedDataInterfaceParameters.Bind(&SpawnExecContext.Parameters);
+		InstanceParams.Bind(&SpawnExecContext.Parameters);
+		InstanceParams.Bind(&UpdateExecContext.Parameters);
 
-	InstanceParams.Bind(&UpdateExecContext.Parameters);
-	SystemScriptDefinedDataInterfaceParameters.Bind(&UpdateExecContext.Parameters);
-	ScriptDefinedDataInterfaceParameters.Bind(&UpdateExecContext.Parameters);
-
-	for (FNiagaraScriptExecutionContext& EventContext : EventExecContexts)
-	{
-		InstanceParams.Bind(&EventContext.Parameters);
-		SystemScriptDefinedDataInterfaceParameters.Bind(&EventContext.Parameters);
-		ScriptDefinedDataInterfaceParameters.Bind(&EventContext.Parameters);
-	}
+		for (FNiagaraScriptExecutionContext& EventContext : EventExecContexts)
+		{
+			InstanceParams.Bind(&EventContext.Parameters);
+		}
 
-#if WITH_EDITORONLY_DATA
-	CachedEmitter->SpawnScriptProps.Script->RapidIterationParameters.Bind(&SpawnExecContext.Parameters);
-	CachedEmitter->UpdateScriptProps.Script->RapidIterationParameters.Bind(&UpdateExecContext.Parameters);
-	ensure(CachedEmitter->GetEventHandlers().Num() == EventExecContexts.Num());
-	for (int32 i = 0; i < CachedEmitter->GetEventHandlers().Num(); i++)
-	{
-		CachedEmitter->GetEventHandlers()[i].Script->RapidIterationParameters.Bind(&EventExecContexts[i].Parameters);
-	}
-#endif
+	#if WITH_EDITORONLY_DATA
+		CachedEmitter->SpawnScriptProps.Script->RapidIterationParameters.Bind(&SpawnExecContext.Parameters);
+		CachedEmitter->UpdateScriptProps.Script->RapidIterationParameters.Bind(&UpdateExecContext.Parameters);
+		ensure(CachedEmitter->GetEventHandlers().Num() == EventExecContexts.Num());
+		for (int32 i = 0; i < CachedEmitter->GetEventHandlers().Num(); i++)
+		{
+			CachedEmitter->GetEventHandlers()[i].Script->RapidIterationParameters.Bind(&EventExecContexts[i].Parameters);
+		}
+	#endif
 
-	if (CachedEmitter->SimTarget == ENiagaraSimTarget::GPUComputeSim)
-	{
-		SpawnExecContext.Parameters.Bind(&GPUExecContext->CombinedParamStore);
-		UpdateExecContext.Parameters.Bind(&GPUExecContext->CombinedParamStore);
+		if (CachedEmitter->SimTarget == ENiagaraSimTarget::GPUComputeSim)
+		{
+			SpawnExecContext.Parameters.Bind(&GPUExecContext->CombinedParamStore);
+			UpdateExecContext.Parameters.Bind(&GPUExecContext->CombinedParamStore);
+		}
 	}
 }
 
@@ -1097,14 +1157,14 @@ void FNiagaraEmitterInstance::Tick(float DeltaSeconds)
 			TSharedPtr<struct FNiagaraScriptDebuggerInfo, ESPMode::ThreadSafe> DebugInfo = ParentSystemInstance->GetActiveCaptureWrite(CachedIDName, ENiagaraScriptUsage::ParticleGPUComputeScript, FGuid());
 			if (DebugInfo)
 			{
+				const TArray<FNiagaraEmitterCompiledData>& EmitterCompiledData = ParentSystemInstance->GetSystem()->GetEmitterCompiledData();
+
 				//Data.Dump(DebugInfo->Frame, true, 0, OrigNumParticles);
 				//DebugInfo->Frame.Dump(true, 0, OrigNumParticles);
 				DebugInfo->Parameters = GPUExecContext->CombinedParamStore;
 				
 				//TODO: This layout info can be pulled into the emitter/systems etc and all sets just refer to them. They are becoming an annoyance here.
-				DebugInfo->Frame.Init(FNiagaraDataSetID(TEXT("GPU Capture Data"), ENiagaraDataSetType::ParticleData), ENiagaraSimTarget::CPUSim, TEXT("GPU Capture Data"));
-				DebugInfo->Frame.AddVariables(Data.GetVariables());
-				DebugInfo->Frame.Finalize();
+				DebugInfo->Frame.Init(&EmitterCompiledData[EmitterIdx].GPUCaptureDataSetCompiledData);
 
 				GPUExecContext->DebugInfo = DebugInfo;
 			}
@@ -1202,7 +1262,7 @@ void FNiagaraEmitterInstance::Tick(float DeltaSeconds)
 
 	int32 OrigNumParticles = Data.GetCurrentDataChecked().GetNumInstances();
 
-	int32 AllocationSize = OrigNumParticles + SpawnTotal + EventSpawnTotal;
+	int32 AllocationSize = FMath::Max<int32>(CachedEmitter->PreAllocationCount, OrigNumParticles + SpawnTotal + EventSpawnTotal);
 	//Ensure we don't blow our current hard limits on cpu particle count.
 	//TODO: These current limits can be improved relatively easily. Though perf in at these counts will obviously be an issue anyway.
 	if (CachedEmitter->SimTarget == ENiagaraSimTarget::CPUSim && AllocationSize > GMaxNiagaraCPUParticlesPerEmitter)
diff --git a/Engine/Plugins/FX/Niagara/Source/Niagara/Private/NiagaraScript.cpp b/Engine/Plugins/FX/Niagara/Source/Niagara/Private/NiagaraScript.cpp
--- a/Engine/Plugins/FX/Niagara/Source/Niagara/Private/NiagaraScript.cpp
+++ b/Engine/Plugins/FX/Niagara/Source/Niagara/Private/NiagaraScript.cpp
@@ -662,6 +662,31 @@ void UNiagaraScript::PostLoad()
 
 	//FNiagaraUtilities::DumpHLSLText(RapidIterationParameters.ToString(), *GetPathName());
 
+	// Init ScriptExecutionParamStoreCPU or ScriptExecutionParamStoreGPU to avoid first time cost
+	if (FPlatformProperties::RequiresCookedData())
+	{
+	    switch (Usage)
+	    {
+	    case ENiagaraScriptUsage::ParticleSpawnScript:
+	    case ENiagaraScriptUsage::ParticleSpawnScriptInterpolated:
+	    case ENiagaraScriptUsage::ParticleUpdateScript:
+	    case ENiagaraScriptUsage::ParticleEventScript:
+	    case ENiagaraScriptUsage::ParticleGPUComputeScript:
+		    if (UNiagaraEmitter* OwningEmitter = GetTypedOuter<UNiagaraEmitter>())
+		    {
+				GetExecutionReadyParameterStore(OwningEmitter->SimTarget);
+		    }
+		    break;
+	    case ENiagaraScriptUsage::EmitterSpawnScript:
+	    case ENiagaraScriptUsage::EmitterUpdateScript:
+	    case ENiagaraScriptUsage::SystemSpawnScript:
+	    case ENiagaraScriptUsage::SystemUpdateScript:
+			GetExecutionReadyParameterStore(ENiagaraSimTarget::CPUSim);
+		    break;
+	    default:
+		    break;
+		};
+	}
 }
 
 bool UNiagaraScript::IsReadyToRun(ENiagaraSimTarget SimTarget) const
diff --git a/Engine/Plugins/FX/Niagara/Source/Niagara/Private/NiagaraSystem.cpp b/Engine/Plugins/FX/Niagara/Source/Niagara/Private/NiagaraSystem.cpp
--- a/Engine/Plugins/FX/Niagara/Source/Niagara/Private/NiagaraSystem.cpp
+++ b/Engine/Plugins/FX/Niagara/Source/Niagara/Private/NiagaraSystem.cpp
@@ -369,12 +369,12 @@ void UNiagaraSystem::PostLoad()
 			UE_LOG(LogNiagara, Log, TEXT("System %s being compiled because there were changes to an emitter script Change ID."), *GetPathName());
 		}
 
-	if (EmitterCompiledData.Num() == 0 || EmitterCompiledData[0].DataSetCompiledData.Variables.Num() == 0)
-	{
-		InitEmitterCompiledData();
-	}
+		if (EmitterCompiledData.Num() == 0 || EmitterCompiledData[0].DataSetCompiledData.Variables.Num() == 0)
+		{
+			InitEmitterCompiledData();
+		}
 
-		if (SystemCompiledData.InstanceParamStore.GetNumParameters() == 0)
+		if (SystemCompiledData.InstanceParamStore.GetNumParameters() == 0 ||SystemCompiledData.DataSetCompiledData.Variables.Num() == 0)
 		{
 			InitSystemCompiledData();
 		}
@@ -1130,10 +1130,16 @@ void UNiagaraSystem::InitEmitterCompiledData()
 			const FNiagaraEmitterHandle& EmitterHandle = EmitterHandles[EmitterIdx];
 			const UNiagaraEmitter* Emitter = EmitterHandle.GetInstance();
 			FNiagaraDataSetCompiledData& EmitterDataSetCompiledData = EmitterCompiledData[EmitterIdx].DataSetCompiledData;
+			FNiagaraDataSetCompiledData& GPUCaptureCompiledData = EmitterCompiledData[EmitterIdx].GPUCaptureDataSetCompiledData;
 			if ensureMsgf(Emitter != nullptr, TEXT("Failed to get Emitter Instance from Emitter Handle in post compile, please investigate."))
 			{
+				static FName GPUCaptureDataSetName = TEXT("GPU Capture Dataset");
 				InitEmitterVariableAliasNames(EmitterCompiledData[EmitterIdx], Emitter);
 				InitEmitterDataSetCompiledData(EmitterDataSetCompiledData, Emitter, EmitterHandle);
+				GPUCaptureCompiledData.ID = FNiagaraDataSetID(GPUCaptureDataSetName, ENiagaraDataSetType::ParticleData);
+				GPUCaptureCompiledData.Variables = EmitterDataSetCompiledData.Variables;
+				GPUCaptureCompiledData.SimTarget = ENiagaraSimTarget::CPUSim;
+				GPUCaptureCompiledData.BuildLayout();				
 			}
 		}
 	}
@@ -1170,6 +1176,30 @@ void UNiagaraSystem::InitSystemCompiledData()
 			}
 		}
 	}
+
+	auto CreateDataSetCompiledData = [&](FNiagaraDataSetCompiledData& CompiledData, TArrayView<FNiagaraVariable> Vars)
+	{
+		CompiledData.Empty();
+
+		CompiledData.Variables.Reset(Vars.Num());
+		for (const FNiagaraVariable& Var : Vars)
+		{
+			CompiledData.Variables.AddUnique(Var);
+		}
+
+		CompiledData.bNeedsPersistentIDs = false;
+		CompiledData.ID = FNiagaraDataSetID();
+		CompiledData.SimTarget = ENiagaraSimTarget::CPUSim;
+
+		CompiledData.BuildLayout();
+	};
+
+	CreateDataSetCompiledData(SystemCompiledData.DataSetCompiledData, GetSystemUpdateScript()->GetVMExecutableData().Attributes);
+
+	FNiagaraParameters* EngineParamsSpawn = GetSystemSpawnScript()->GetVMExecutableData().DataSetToParameters.Find(TEXT("Engine"));
+	CreateDataSetCompiledData(SystemCompiledData.SpawnInstanceParamsDataSetCompiledData, EngineParamsSpawn ? EngineParamsSpawn->Parameters : TArrayView<FNiagaraVariable>());
+	FNiagaraParameters* EngineParamsUpdate = GetSystemUpdateScript()->GetVMExecutableData().DataSetToParameters.Find(TEXT("Engine"));
+	CreateDataSetCompiledData(SystemCompiledData.UpdateInstanceParamsDataSetCompiledData, EngineParamsUpdate ? EngineParamsUpdate->Parameters : TArrayView<FNiagaraVariable>());
 }
 #endif
 
diff --git a/Engine/Plugins/FX/Niagara/Source/Niagara/Private/NiagaraSystemInstance.cpp b/Engine/Plugins/FX/Niagara/Source/Niagara/Private/NiagaraSystemInstance.cpp
--- a/Engine/Plugins/FX/Niagara/Source/Niagara/Private/NiagaraSystemInstance.cpp
+++ b/Engine/Plugins/FX/Niagara/Source/Niagara/Private/NiagaraSystemInstance.cpp
@@ -80,6 +80,7 @@ FNiagaraSystemInstance::FNiagaraSystemInstance(UNiagaraComponent* InComponent)
 	, bIsTransformDirty(true)
 	, bNeedsFinalize(false)
 	, bDataInterfacesInitialized(false)
+	, bAlreadyBound(false)
 	, bAsyncWorkInProgress(false)
 	, CachedDeltaSeconds(0.0f)
 	, RequestedExecutionState(ENiagaraExecutionState::Complete)
@@ -160,6 +161,7 @@ void FNiagaraSystemInstance::Init(bool bInForceSolo)
 	bForceSolo = bInForceSolo;
 	ActualExecutionState = ENiagaraExecutionState::Inactive;
 	RequestedExecutionState = ENiagaraExecutionState::Inactive;
+	bAlreadyBound = false;
 
 	//InstanceParameters = GetSystem()->GetInstanceParameters();
 	// In order to get user data interface parameters in the component to work properly,
@@ -534,7 +536,7 @@ void FNiagaraSystemInstance::Complete()
 
 	DestroyDataInterfaceInstanceData();
 
-	UnbindParameters();
+	UnbindParameters(true);
 
 	bPendingSpawn = false;
 
@@ -827,6 +829,8 @@ void FNiagaraSystemInstance::ReInitInternal()
 	LocalBounds = FBox(FVector::ZeroVector, FVector::ZeroVector);
 	CachedDeltaSeconds = 0.0f;
 
+	bAlreadyBound = false;
+
 	FastPathIntUserParameterInputBindings.Empty();
 	FastPathFloatUserParameterInputBindings.Empty();
 	FastPathIntUpdateRangedInputBindings.Empty();
@@ -997,7 +1001,11 @@ void FNiagaraSystemInstance::BindParameters()
 		return;
 	}
 
-	Component->GetOverrideParameters().Bind(&InstanceParameters);
+	if (!bAlreadyBound)
+	{
+		// NOTE: We don't rebind if it's already bound to improve reset times.
+		Component->GetOverrideParameters().Bind(&InstanceParameters); 
+	}
 
 	if (SystemSimulation->GetIsSolo())
 	{
@@ -1006,27 +1014,20 @@ void FNiagaraSystemInstance::BindParameters()
 		Component->GetOverrideParameters().Bind(&SystemSimulation->GetSpawnExecutionContext().Parameters);
 		Component->GetOverrideParameters().Bind(&SystemSimulation->GetUpdateExecutionContext().Parameters);
 	}
-	else 
-	{
-		// If this simulation is not solo than we have bind the source system parameters to the system simulation contexts so that
-		// the system and emitter scripts use the default shared data interfaces.
-		GetSystem()->GetExposedParameters().Bind(&SystemSimulation->GetSpawnExecutionContext().Parameters);
-		GetSystem()->GetExposedParameters().Bind(&SystemSimulation->GetUpdateExecutionContext().Parameters);
-	}
 
 	for (TSharedRef<FNiagaraEmitterInstance, ESPMode::ThreadSafe> Simulation : Emitters)
 	{
-		Simulation->BindParameters();
+		Simulation->BindParameters(bAlreadyBound);
 	}
+
+	bAlreadyBound = true;
 }
 
-void FNiagaraSystemInstance::UnbindParameters()
+void FNiagaraSystemInstance::UnbindParameters(bool bFromComplete)
 {
-	//TODO: Split unbind into internal and external.
-	//We need to unbind from external sources when we complete but internal can stay bound to improve reset times.
-
-	if (Component)
+	if (Component && !bFromComplete)
 	{
+		// NOTE: We don't unbind this on complete to improve reset times.
 		Component->GetOverrideParameters().Unbind(&InstanceParameters);
 	}
 
@@ -1040,20 +1041,12 @@ void FNiagaraSystemInstance::UnbindParameters()
 				Component->GetOverrideParameters().Unbind(&SystemSimulation->GetUpdateExecutionContext().Parameters);
 			}
 		}
-		else 
-		{
-			UNiagaraSystem* System = GetSystem();
-			if (System)
-			{
-				System->GetExposedParameters().Unbind(&SystemSimulation->GetSpawnExecutionContext().Parameters);
-				System->GetExposedParameters().Unbind(&SystemSimulation->GetSpawnExecutionContext().Parameters);
-			}
-		}
 	}
 
+	bAlreadyBound = bFromComplete && bAlreadyBound;
 	for (TSharedRef<FNiagaraEmitterInstance, ESPMode::ThreadSafe> Simulation : Emitters)
 	{
-		Simulation->UnbindParameters();
+		Simulation->UnbindParameters(bFromComplete);
 	}
 }
 
@@ -1577,31 +1570,30 @@ void FNiagaraSystemInstance::InitEmitters()
 
 	LocalBounds = FBox(FVector::ZeroVector, FVector::ZeroVector);
 
-	Emitters.Empty();
+	Emitters.Empty(false);
 	UNiagaraSystem* System = GetSystem();
 	if (System != nullptr)
 	{
 		const TArray<FNiagaraEmitterHandle>& EmitterHandles = System->GetEmitterHandles();
-		for (int32 EmitterIdx=0; EmitterIdx < System->GetEmitterHandles().Num(); ++EmitterIdx)
+		Emitters.Reserve(EmitterHandles.Num());
+		for (int32 EmitterIdx = 0; EmitterIdx < System->GetEmitterHandles().Num(); ++EmitterIdx)
 		{
+			TSharedRef<FNiagaraEmitterInstance, ESPMode::ThreadSafe> Sim = MakeShared<FNiagaraEmitterInstance, ESPMode::ThreadSafe>(this);
+			Sim->Init(EmitterIdx, ID);
+			if (System->bFixedBounds)
 			{
-				TSharedRef<FNiagaraEmitterInstance, ESPMode::ThreadSafe> Sim = MakeShared<FNiagaraEmitterInstance, ESPMode::ThreadSafe>(this);
-				Sim->Init(EmitterIdx, ID);
-				if (System->bFixedBounds)
-				{
-					Sim->SetSystemFixedBoundsOverride(System->GetFixedBounds());
-				}
-				Emitters.Add(Sim);
+				Sim->SetSystemFixedBoundsOverride(System->GetFixedBounds());
 			}
+			Emitters.Add(Sim);
 		}
 
-		for (TSharedRef<FNiagaraEmitterInstance, ESPMode::ThreadSafe> Simulation : Emitters)
+		for (TSharedRef<FNiagaraEmitterInstance, ESPMode::ThreadSafe>& Simulation : Emitters)
 		{
-			if (Simulation->GetCachedEmitter() != nullptr)
+			FNiagaraEmitterInstance& Inst = Simulation.Get();
+			if (Inst.GetCachedEmitter() != nullptr)
 			{
-				bHasGPUEmitters |= Simulation->GetCachedEmitter()->SimTarget == ENiagaraSimTarget::GPUComputeSim;
-
-				Simulation->PostInitSimulation();
+				bHasGPUEmitters |= Inst.GetCachedEmitter()->SimTarget == ENiagaraSimTarget::GPUComputeSim;
+				Inst.PostInitSimulation();
 			}
 		}
 
diff --git a/Engine/Plugins/FX/Niagara/Source/Niagara/Private/NiagaraSystemSimulation.cpp b/Engine/Plugins/FX/Niagara/Source/Niagara/Private/NiagaraSystemSimulation.cpp
--- a/Engine/Plugins/FX/Niagara/Source/Niagara/Private/NiagaraSystemSimulation.cpp
+++ b/Engine/Plugins/FX/Niagara/Source/Niagara/Private/NiagaraSystemSimulation.cpp
@@ -33,6 +33,13 @@ DECLARE_CYCLE_STAT(TEXT("System Sim Transfer Results [CNC]"), STAT_NiagaraSystem
 DECLARE_CYCLE_STAT(TEXT("System Sim Init [GT]"), STAT_NiagaraSystemSim_Init, STATGROUP_Niagara);
 DECLARE_CYCLE_STAT(TEXT("System Sim FastPath [CNC]"), STAT_NiagaraSystemSim_FastPathCNC, STATGROUP_Niagara);
 
+DECLARE_CYCLE_STAT(TEXT("System Sim Init (DataSets) [GT]"), STAT_NiagaraSystemSim_Init_DataSets, STATGROUP_Niagara);
+DECLARE_CYCLE_STAT(TEXT("System Sim Init (ExecContexts) [GT]"), STAT_NiagaraSystemSim_Init_ExecContexts, STATGROUP_Niagara);
+DECLARE_CYCLE_STAT(TEXT("System Sim Init (BindParams) [GT]"), STAT_NiagaraSystemSim_Init_BindParams, STATGROUP_Niagara);
+DECLARE_CYCLE_STAT(TEXT("System Sim Init (DatasetAccessors) [GT]"), STAT_NiagaraSystemSim_Init_DatasetAccessors, STATGROUP_Niagara);
+DECLARE_CYCLE_STAT(TEXT("System Sim Init (DirectBindings) [GT]"), STAT_NiagaraSystemSim_Init_DirectBindings, STATGROUP_Niagara);
+
+
 DECLARE_CYCLE_STAT(TEXT("ForcedWaitForAsync"), STAT_NiagaraSystemSim_ForceWaitForAsync, STATGROUP_Niagara);
 DECLARE_CYCLE_STAT(TEXT("ForcedWait Fake Stall"), STAT_NiagaraSystemSim_ForceWaitFakeStall, STATGROUP_Niagara);
 
@@ -260,134 +267,127 @@ bool FNiagaraSystemSimulation::Init(UNiagaraSystem* InSystem, UWorld* InWorld, b
 
 	if (bCanExecute)
 	{
-		//TODO: Move all layout data into the System!
-
-		//Initialize the main simulation dataset.
-#if !UE_BUILD_SHIPPING && !UE_BUILD_TEST 
-		MainDataSet.Init(FNiagaraDataSetID(), ENiagaraSimTarget::CPUSim, System->GetFullName());
-#else
-		MainDataSet.Init(FNiagaraDataSetID(), ENiagaraSimTarget::CPUSim);
-#endif
-		MainDataSet.AddVariables(System->GetSystemSpawnScript()->GetVMExecutableData().Attributes);
-		MainDataSet.AddVariables(System->GetSystemUpdateScript()->GetVMExecutableData().Attributes);
-		MainDataSet.Finalize();
-
-		//Initialize the main simulation dataset.
-#if !UE_BUILD_SHIPPING && !UE_BUILD_TEST 
-		SpawningDataSet.Init(FNiagaraDataSetID(), ENiagaraSimTarget::CPUSim, System->GetFullName());
-#else
-		SpawningDataSet.Init(FNiagaraDataSetID(), ENiagaraSimTarget::CPUSim);
-#endif
-		SpawningDataSet.AddVariables(System->GetSystemSpawnScript()->GetVMExecutableData().Attributes);
-		SpawningDataSet.AddVariables(System->GetSystemUpdateScript()->GetVMExecutableData().Attributes);
-		SpawningDataSet.Finalize();
-
-		//Initialize the dataset for paused systems.
-#if !UE_BUILD_SHIPPING && !UE_BUILD_TEST 
-		PausedInstanceData.Init(FNiagaraDataSetID(), ENiagaraSimTarget::CPUSim, System->GetFullName());
-#else
-		PausedInstanceData.Init(FNiagaraDataSetID(), ENiagaraSimTarget::CPUSim);
-#endif
-		PausedInstanceData.AddVariables(System->GetSystemSpawnScript()->GetVMExecutableData().Attributes);
-		PausedInstanceData.AddVariables(System->GetSystemUpdateScript()->GetVMExecutableData().Attributes);
-		PausedInstanceData.Finalize();
 
 		{
-#if !UE_BUILD_SHIPPING && !UE_BUILD_TEST 
-			SpawnInstanceParameterDataSet.Init(FNiagaraDataSetID(), ENiagaraSimTarget::CPUSim, System->GetFullName());
-#else
-			SpawnInstanceParameterDataSet.Init(FNiagaraDataSetID(), ENiagaraSimTarget::CPUSim);
-#endif
-			FNiagaraParameters* EngineParamsSpawn = System->GetSystemSpawnScript()->GetVMExecutableData().DataSetToParameters.Find(TEXT("Engine"));
-			if (EngineParamsSpawn != nullptr)
-			{
-				SpawnInstanceParameterDataSet.AddVariables(EngineParamsSpawn->Parameters);
-			}
-			SpawnInstanceParameterDataSet.Finalize();
-#if !UE_BUILD_SHIPPING && !UE_BUILD_TEST 
-			UpdateInstanceParameterDataSet.Init(FNiagaraDataSetID(), ENiagaraSimTarget::CPUSim, System->GetFullName());
-#else
-			UpdateInstanceParameterDataSet.Init(FNiagaraDataSetID(), ENiagaraSimTarget::CPUSim);
-#endif
-			FNiagaraParameters* EngineParamsUpdate = System->GetSystemUpdateScript()->GetVMExecutableData().DataSetToParameters.Find(TEXT("Engine"));
-			if (EngineParamsUpdate != nullptr)
-			{
-				UpdateInstanceParameterDataSet.AddVariables(EngineParamsUpdate->Parameters);
-			}
-			UpdateInstanceParameterDataSet.Finalize();
+			SCOPE_CYCLE_COUNTER(STAT_NiagaraSystemSim_Init_DataSets);
+
+			const FNiagaraSystemCompiledData& SystemCompiledData = System->GetSystemCompiledData();
+			//Initialize the main simulation dataset.
+			MainDataSet.Init(&SystemCompiledData.DataSetCompiledData);
+
+			//Initialize the main simulation dataset.
+			SpawningDataSet.Init(&SystemCompiledData.DataSetCompiledData);
+
+			//Initialize the dataset for paused systems.
+			PausedInstanceData.Init(&SystemCompiledData.DataSetCompiledData);
+			
+			SpawnInstanceParameterDataSet.Init(&SystemCompiledData.SpawnInstanceParamsDataSetCompiledData);
+
+			UpdateInstanceParameterDataSet.Init(&SystemCompiledData.UpdateInstanceParamsDataSetCompiledData);
+			
 		}
 
 		UNiagaraScript* SpawnScript = System->GetSystemSpawnScript();
 		UNiagaraScript* UpdateScript = System->GetSystemUpdateScript();
 
-		SpawnExecContext.Init(SpawnScript, ENiagaraSimTarget::CPUSim);
-		UpdateExecContext.Init(UpdateScript, ENiagaraSimTarget::CPUSim);
-
-		//Bind parameter collections.
-		for (UNiagaraParameterCollection* Collection : SpawnScript->GetCachedParameterCollectionReferences())
 		{
-			GetParameterCollectionInstance(Collection)->GetParameterStore().Bind(&SpawnExecContext.Parameters);
+			SCOPE_CYCLE_COUNTER(STAT_NiagaraSystemSim_Init_ExecContexts);
+
+
+
+			SpawnExecContext.Init(SpawnScript, ENiagaraSimTarget::CPUSim);
+			UpdateExecContext.Init(UpdateScript, ENiagaraSimTarget::CPUSim);
 		}
-		for (UNiagaraParameterCollection* Collection : UpdateScript->GetCachedParameterCollectionReferences())
+
+
 		{
-			GetParameterCollectionInstance(Collection)->GetParameterStore().Bind(&UpdateExecContext.Parameters);
-		}
+			SCOPE_CYCLE_COUNTER(STAT_NiagaraSystemSim_Init_BindParams);
+
+
+			//Bind parameter collections.
+			for (UNiagaraParameterCollection* Collection : SpawnScript->GetCachedParameterCollectionReferences())
+			{
+				GetParameterCollectionInstance(Collection)->GetParameterStore().Bind(&SpawnExecContext.Parameters);
+			}
+			for (UNiagaraParameterCollection* Collection : UpdateScript->GetCachedParameterCollectionReferences())
+			{
+				GetParameterCollectionInstance(Collection)->GetParameterStore().Bind(&UpdateExecContext.Parameters);
+			}
+
+			TArray<UNiagaraScript*, TInlineAllocator<2>> Scripts;
+			Scripts.Add(SpawnScript);
+			Scripts.Add(UpdateScript);
+			FNiagaraUtilities::CollectScriptDataInterfaceParameters(*System, Scripts, ScriptDefinedDataInterfaceParameters);
 
-		TArray<UNiagaraScript*> Scripts;
-		Scripts.Add(SpawnScript);
-		Scripts.Add(UpdateScript);
-		FNiagaraUtilities::CollectScriptDataInterfaceParameters(*System, Scripts, ScriptDefinedDataInterfaceParameters);
+			ScriptDefinedDataInterfaceParameters.Bind(&SpawnExecContext.Parameters);
+			ScriptDefinedDataInterfaceParameters.Bind(&UpdateExecContext.Parameters);
 
-		ScriptDefinedDataInterfaceParameters.Bind(&SpawnExecContext.Parameters);
-		ScriptDefinedDataInterfaceParameters.Bind(&UpdateExecContext.Parameters);
+			SpawnScript->RapidIterationParameters.Bind(&SpawnExecContext.Parameters);
+			UpdateScript->RapidIterationParameters.Bind(&UpdateExecContext.Parameters);
 
-		SpawnScript->RapidIterationParameters.Bind(&SpawnExecContext.Parameters);
-		UpdateScript->RapidIterationParameters.Bind(&UpdateExecContext.Parameters);
+			// If this simulation is not solo than we have bind the source system parameters to the system simulation contexts so that
+			// the system and emitter scripts use the default shared data interfaces.
+			if (!bIsSolo)
+			{
+				GetSystem()->GetExposedParameters().Bind(&GetSpawnExecutionContext().Parameters);
+				GetSystem()->GetExposedParameters().Bind(&GetUpdateExecutionContext().Parameters);
+			}
+		}
 
-		SystemExecutionStateAccessor.Create(&MainDataSet, FNiagaraVariable(EnumPtr, TEXT("System.ExecutionState")));
-		EmitterSpawnInfoAccessors.Reset();
-		EmitterExecutionStateAccessors.Reset();
-		EmitterSpawnInfoAccessors.SetNum(System->GetNumEmitters());
 
-		for (int32 EmitterIdx = 0; EmitterIdx < System->GetNumEmitters(); ++EmitterIdx)
 		{
-			FNiagaraEmitterHandle& EmitterHandle = System->GetEmitterHandle(EmitterIdx);
-			UNiagaraEmitter* Emitter = EmitterHandle.GetInstance();
-			if (Emitter)
-			{
-				FString EmitterName = Emitter->GetUniqueEmitterName();
-				EmitterExecutionStateAccessors.Emplace(MainDataSet, FNiagaraVariable(EnumPtr, *(EmitterName + TEXT(".ExecutionState"))));
-				const TArray<FNiagaraEmitterCompiledData>& EmitterCompiledData = System->GetEmitterCompiledData();
+			SCOPE_CYCLE_COUNTER(STAT_NiagaraSystemSim_Init_DatasetAccessors);
+
+			SystemExecutionStateAccessor.Create(&MainDataSet, FNiagaraVariable(EnumPtr, TEXT("System.ExecutionState")));
+			EmitterSpawnInfoAccessors.Reset();
+			EmitterExecutionStateAccessors.Reset();
+			EmitterSpawnInfoAccessors.SetNum(System->GetNumEmitters());
 
-				check(EmitterCompiledData.Num() == System->GetNumEmitters());
-				for (FName AttrName : EmitterCompiledData[EmitterIdx].SpawnAttributes)
+			for (int32 EmitterIdx = 0; EmitterIdx < System->GetNumEmitters(); ++EmitterIdx)
+			{
+				FNiagaraEmitterHandle& EmitterHandle = System->GetEmitterHandle(EmitterIdx);
+				UNiagaraEmitter* Emitter = EmitterHandle.GetInstance();
+				if (Emitter)
 				{
-					EmitterSpawnInfoAccessors[EmitterIdx].Emplace(MainDataSet, FNiagaraVariable(FNiagaraTypeDefinition(FNiagaraSpawnInfo::StaticStruct()), AttrName));
-				}
+					FString EmitterName = Emitter->GetUniqueEmitterName();
+					EmitterExecutionStateAccessors.Emplace(MainDataSet, FNiagaraVariable(EnumPtr, *(EmitterName + TEXT(".ExecutionState"))));
+					const TArray<FNiagaraEmitterCompiledData>& EmitterCompiledData = System->GetEmitterCompiledData();
+
+					check(EmitterCompiledData.Num() == System->GetNumEmitters());
+					for (FName AttrName : EmitterCompiledData[EmitterIdx].SpawnAttributes)
+					{
+						EmitterSpawnInfoAccessors[EmitterIdx].Emplace(MainDataSet, FNiagaraVariable(FNiagaraTypeDefinition(FNiagaraSpawnInfo::StaticStruct()), AttrName));
+					}
 
-				if (Emitter->bLimitDeltaTime)
+					if (Emitter->bLimitDeltaTime)
+					{
+						MaxDeltaTime = MaxDeltaTime.IsSet() ? FMath::Min(MaxDeltaTime.GetValue(), Emitter->MaxDeltaTimePerTick) : Emitter->MaxDeltaTimePerTick;
+					}
+				}
+				else
 				{
-					MaxDeltaTime = MaxDeltaTime.IsSet() ? FMath::Min(MaxDeltaTime.GetValue(), Emitter->MaxDeltaTimePerTick) : Emitter->MaxDeltaTimePerTick;
+					EmitterExecutionStateAccessors.Emplace();
 				}
 			}
-			else
-			{
-				EmitterExecutionStateAccessors.Emplace();
-			}
 		}
 
-		SpawnTimeParam.Init(SpawnExecContext.Parameters, SYS_PARAM_ENGINE_TIME);
-		UpdateTimeParam.Init(UpdateExecContext.Parameters, SYS_PARAM_ENGINE_TIME);
-		SpawnDeltaTimeParam.Init(SpawnExecContext.Parameters, SYS_PARAM_ENGINE_DELTA_TIME);
-		UpdateDeltaTimeParam.Init(UpdateExecContext.Parameters, SYS_PARAM_ENGINE_DELTA_TIME);
-		SpawnInvDeltaTimeParam.Init(SpawnExecContext.Parameters, SYS_PARAM_ENGINE_INV_DELTA_TIME);
-		UpdateInvDeltaTimeParam.Init(UpdateExecContext.Parameters, SYS_PARAM_ENGINE_INV_DELTA_TIME);
-		SpawnNumSystemInstancesParam.Init(SpawnExecContext.Parameters, SYS_PARAM_ENGINE_NUM_SYSTEM_INSTANCES);
-		UpdateNumSystemInstancesParam.Init(UpdateExecContext.Parameters, SYS_PARAM_ENGINE_NUM_SYSTEM_INSTANCES);
-		SpawnGlobalSpawnCountScaleParam.Init(SpawnExecContext.Parameters, SYS_PARAM_ENGINE_GLOBAL_SPAWN_COUNT_SCALE);
-		UpdateGlobalSpawnCountScaleParam.Init(UpdateExecContext.Parameters, SYS_PARAM_ENGINE_GLOBAL_SPAWN_COUNT_SCALE);
-		SpawnGlobalSystemCountScaleParam.Init(SpawnExecContext.Parameters, SYS_PARAM_ENGINE_GLOBAL_SYSTEM_COUNT_SCALE);
-		UpdateGlobalSystemCountScaleParam.Init(UpdateExecContext.Parameters, SYS_PARAM_ENGINE_GLOBAL_SYSTEM_COUNT_SCALE);
+
+		{
+			SCOPE_CYCLE_COUNTER(STAT_NiagaraSystemSim_Init_DirectBindings);
+
+			SpawnTimeParam.Init(SpawnExecContext.Parameters, SYS_PARAM_ENGINE_TIME);
+			UpdateTimeParam.Init(UpdateExecContext.Parameters, SYS_PARAM_ENGINE_TIME);
+			SpawnDeltaTimeParam.Init(SpawnExecContext.Parameters, SYS_PARAM_ENGINE_DELTA_TIME);
+			UpdateDeltaTimeParam.Init(UpdateExecContext.Parameters, SYS_PARAM_ENGINE_DELTA_TIME);
+			SpawnInvDeltaTimeParam.Init(SpawnExecContext.Parameters, SYS_PARAM_ENGINE_INV_DELTA_TIME);
+			UpdateInvDeltaTimeParam.Init(UpdateExecContext.Parameters, SYS_PARAM_ENGINE_INV_DELTA_TIME);
+			SpawnNumSystemInstancesParam.Init(SpawnExecContext.Parameters, SYS_PARAM_ENGINE_NUM_SYSTEM_INSTANCES);
+			UpdateNumSystemInstancesParam.Init(UpdateExecContext.Parameters, SYS_PARAM_ENGINE_NUM_SYSTEM_INSTANCES);
+			SpawnGlobalSpawnCountScaleParam.Init(SpawnExecContext.Parameters, SYS_PARAM_ENGINE_GLOBAL_SPAWN_COUNT_SCALE);
+			UpdateGlobalSpawnCountScaleParam.Init(UpdateExecContext.Parameters, SYS_PARAM_ENGINE_GLOBAL_SPAWN_COUNT_SCALE);
+			SpawnGlobalSystemCountScaleParam.Init(SpawnExecContext.Parameters, SYS_PARAM_ENGINE_GLOBAL_SYSTEM_COUNT_SCALE);
+			UpdateGlobalSystemCountScaleParam.Init(UpdateExecContext.Parameters, SYS_PARAM_ENGINE_GLOBAL_SYSTEM_COUNT_SCALE);
+		}
 	}
 
 	return true;
@@ -429,8 +429,10 @@ UNiagaraParameterCollectionInstance* FNiagaraSystemSimulation::GetParameterColle
 	//If no explicit override from the system, just get the current instance set on the world.
 	if (!Ret)
 	{
-		FNiagaraWorldManager* WorldMan = FNiagaraWorldManager::Get(World);
-		Ret = WorldMan->GetParameterCollection(Collection);
+		if(FNiagaraWorldManager* WorldMan = FNiagaraWorldManager::Get(World))
+		{
+			Ret = WorldMan->GetParameterCollection(Collection);
+		}
 	}
 
 	return Ret;
diff --git a/Engine/Plugins/FX/Niagara/Source/Niagara/Public/NiagaraSystemInstance.h b/Engine/Plugins/FX/Niagara/Source/Niagara/Public/NiagaraSystemInstance.h
--- a/Engine/Plugins/FX/Niagara/Source/Niagara/Public/NiagaraSystemInstance.h
+++ b/Engine/Plugins/FX/Niagara/Source/Niagara/Public/NiagaraSystemInstance.h
@@ -74,7 +74,7 @@ public:
 
 	//void RebindParameterCollection(UNiagaraParameterCollectionInstance* OldInstance, UNiagaraParameterCollectionInstance* NewInstance);
 	void BindParameters();
-	void UnbindParameters();
+	void UnbindParameters(bool bFromComplete = false);
 
 	FORCEINLINE FNiagaraParameterStore& GetInstanceParameters() { return InstanceParameters; }
 	
@@ -377,6 +377,8 @@ private:
 
 	uint32 bDataInterfacesInitialized : 1;
 
+	uint32 bAlreadyBound : 1;
+
 	/** True if we have async work in flight. */
 	volatile bool bAsyncWorkInProgress;
 
diff --git a/Engine/Plugins/FX/Niagara/Source/NiagaraCore/Private/NiagaraCustomVersion.cpp b/Engine/Plugins/FX/Niagara/Source/NiagaraCore/Private/NiagaraCustomVersion.cpp
--- a/Engine/Plugins/FX/Niagara/Source/NiagaraCore/Private/NiagaraCustomVersion.cpp
+++ b/Engine/Plugins/FX/Niagara/Source/NiagaraCore/Private/NiagaraCustomVersion.cpp
@@ -8,4 +8,4 @@ const FGuid FNiagaraCustomVersion::GUID(0xFCF57AFA, 0x50764283, 0xB9A9E658, 0xFF
 // Register the custom version with core
 FCustomVersionRegistration GRegisterNiagaraCustomVersion(FNiagaraCustomVersion::GUID, FNiagaraCustomVersion::LatestVersion, TEXT("NiagaraVer"));
 
-const FGuid FNiagaraCustomVersion::LatestScriptCompileVersion(0x66E6CB72, 0xE08011E9, 0x94E21072, 0xE014E1EA);
+const FGuid FNiagaraCustomVersion::LatestScriptCompileVersion(0x908C3A78, 0x6E39F14C, 0x8E1C1A93, 0xDD667695);
\ No newline at end of file
diff --git a/Engine/Plugins/FX/Niagara/Source/NiagaraEditor/Private/Widgets/SNiagaraSpreadsheetView.cpp b/Engine/Plugins/FX/Niagara/Source/NiagaraEditor/Private/Widgets/SNiagaraSpreadsheetView.cpp
--- a/Engine/Plugins/FX/Niagara/Source/NiagaraEditor/Private/Widgets/SNiagaraSpreadsheetView.cpp
+++ b/Engine/Plugins/FX/Niagara/Source/NiagaraEditor/Private/Widgets/SNiagaraSpreadsheetView.cpp
@@ -710,7 +710,7 @@ void SNiagaraSpreadsheetView::SetTarget(UNiagaraComponent* InComponent)
 
 	for (int32 i = 0; i < (int32)UIMax; i++)
 	{
-		CaptureData[i].DataSet.Init(FNiagaraDataSetID(), ENiagaraSimTarget::CPUSim, InComponent->GetFullName() + TEXT("/Spreadsheet"));
+		CaptureData[i].DataSet.Init(nullptr);
 		CaptureData[i].InputParams.Reset();
 		CaptureData[i].CaptureData.Reset();
 	}
@@ -902,7 +902,7 @@ void SNiagaraSpreadsheetView::SystemSelectionChanged(UNiagaraSystemSelectionView
 	// Need to reset the attributes list...
 	for (int32 i = 0; i < (int32)UIMax; i++)
 	{
-		CaptureData[i].DataSet.Init(FNiagaraDataSetID(), ENiagaraSimTarget::CPUSim, TEXT("/Spreadsheet"));
+		CaptureData[i].DataSet.Init(nullptr);
 		CaptureData[i].SupportedInputIndices.SetNum(0);
 		CaptureData[i].SupportedOutputIndices.SetNum(0);
 		CaptureData[i].OutputsListView->RequestTreeRefresh();
@@ -1035,8 +1035,7 @@ void SNiagaraSpreadsheetView::HandleTimeChange()
 					else
 					{
 						CaptureData[i].CaptureData.Reset();
-						CaptureData[i].DataSet.Init(FNiagaraDataSetID(), ENiagaraSimTarget::CPUSim, TargetComponent->GetFullName() + TEXT("/Spreadsheet"));
-						CaptureData[i].DataSet.Finalize();
+						CaptureData[i].DataSet.Init(nullptr);
 
 						ResetColumns((EUITab)i);
 						ResetEntries((EUITab)i);
diff --git a/Engine/Plugins/Online/IOS/OnlineSubsystemIOS/Source/Private/OnlineIdentityInterfaceIOS.h b/Engine/Plugins/Online/IOS/OnlineSubsystemIOS/Source/Private/OnlineIdentityInterfaceIOS.h
--- a/Engine/Plugins/Online/IOS/OnlineSubsystemIOS/Source/Private/OnlineIdentityInterfaceIOS.h
+++ b/Engine/Plugins/Online/IOS/OnlineSubsystemIOS/Source/Private/OnlineIdentityInterfaceIOS.h
@@ -3,6 +3,7 @@
 
 #pragma once
 
+#import <UIKit/UIKit.h>
 #import <GameKit/GameKit.h>
 
 #include "OnlineSubsystemIOSTypes.h"
diff --git a/Engine/Plugins/Online/IOS/OnlineSubsystemIOS/Source/Private/OnlineSessionInterfaceIOS.h b/Engine/Plugins/Online/IOS/OnlineSubsystemIOS/Source/Private/OnlineSessionInterfaceIOS.h
--- a/Engine/Plugins/Online/IOS/OnlineSubsystemIOS/Source/Private/OnlineSessionInterfaceIOS.h
+++ b/Engine/Plugins/Online/IOS/OnlineSubsystemIOS/Source/Private/OnlineSessionInterfaceIOS.h
@@ -2,6 +2,7 @@
 
 #pragma once
 
+#import <UIKit/UIKit.h>
 #import <GameKit/GKLocalPlayer.h>
 
 #include "Interfaces/OnlineSessionInterface.h"
diff --git a/Engine/Plugins/Runtime/AR/Microsoft/HoloLensAR/Source/HoloLensAR/Private/HoloLensARSystem.cpp b/Engine/Plugins/Runtime/AR/Microsoft/HoloLensAR/Source/HoloLensAR/Private/HoloLensARSystem.cpp
--- a/Engine/Plugins/Runtime/AR/Microsoft/HoloLensAR/Source/HoloLensAR/Private/HoloLensARSystem.cpp
+++ b/Engine/Plugins/Runtime/AR/Microsoft/HoloLensAR/Source/HoloLensAR/Private/HoloLensARSystem.cpp
@@ -697,6 +697,13 @@ void FHoloLensARSystem::AddOrUpdateMesh(FMeshUpdate* CurrentMesh)
 		{
 			// Attach this component to our single origin actor
 			AAROriginActor* OriginActor = AAROriginActor::GetOriginActor();
+
+			// During shutdown we can get a mesh update after the OriginActor has been destroyed, just return in that case.
+			if (OriginActor == nullptr)
+			{
+				return;
+			}
+
 			UMRMeshComponent* MRMesh = NewObject<UMRMeshComponent>(OriginActor);
 
 			// Set the occlusion and wireframe defaults
diff --git a/Engine/Plugins/Runtime/AudioSynesthesia/Source/AudioSynesthesia/AudioSynesthesia.Build.cs b/Engine/Plugins/Runtime/AudioSynesthesia/Source/AudioSynesthesia/AudioSynesthesia.Build.cs
--- a/Engine/Plugins/Runtime/AudioSynesthesia/Source/AudioSynesthesia/AudioSynesthesia.Build.cs
+++ b/Engine/Plugins/Runtime/AudioSynesthesia/Source/AudioSynesthesia/AudioSynesthesia.Build.cs
@@ -16,7 +16,8 @@ namespace UnrealBuildTool.Rules
                     "Core",
 					"CoreUObject",
 					"Engine",
-					"AudioMixer",
+					"SignalProcessing",
+					"AudioMixerCore",
                     "AudioAnalyzer"
                 }
             );
diff --git a/Engine/Plugins/Runtime/AudioSynesthesia/Source/AudioSynesthesia/Classes/LoudnessNRT.h b/Engine/Plugins/Runtime/AudioSynesthesia/Source/AudioSynesthesia/Classes/LoudnessNRT.h
new file mode 100644
--- /dev/null
+++ b/Engine/Plugins/Runtime/AudioSynesthesia/Source/AudioSynesthesia/Classes/LoudnessNRT.h
@@ -0,0 +1,117 @@
+// Copyright 1998-2019 Epic Games, Inc. All Rights Reserved.
+
+#pragma once
+
+#include "AudioSynesthesiaNRT.h"
+#include "LoudnessNRT.generated.h"
+
+/** ELoudnessNRTCurveTypeEnum
+ *
+ * Enumeration of available equal loudness curves. Loudness curves based on IEC 61672-1:2013 Electroacoustics - Sound level meters - Part 1: Specifications.
+ */
+UENUM(BlueprintType)
+enum class ELoudnessNRTCurveTypeEnum : uint8
+{
+	/** Loudness Curve A Weighting. Commonly used in environmental sound measurements. Best for low level sounds. */
+	A 		UMETA(DisplayName="A"),
+	/** Loudness Curve B Weighting. Relative to "A", gives more precedence to frequencies below 1kHz. */
+	B		UMETA(DisplayName="B"),
+	/** Loudness Curve C Weighting. Relative to "A" and "B", gives more precedence to frequencies below 1kHz. */
+	C		UMETA(DisplayName="C"),
+	/** Loudness Curve D Weighting. Simlar to "B" but with an emphasis on presence in the 2kHz-6KHz frequency range. */
+	D		UMETA(DisplayName="D"),
+	/** No loudness curve weighting. */
+	None	UMETA(DisplayName="None"),
+};
+
+/** ULoudnessNRTSettings
+ *
+ * Settings for a ULoudnessNRT analyzer.
+ */
+UCLASS(Blueprintable)
+class AUDIOSYNESTHESIA_API ULoudnessNRTSettings : public UAudioSynesthesiaNRTSettings
+{
+	GENERATED_BODY()
+	public:
+
+		ULoudnessNRTSettings();
+
+		/** Number of seconds between loudness measurements */
+		UPROPERTY(EditAnywhere, BlueprintReadWrite, Category=AudioAnalyzer, meta = (ClampMin = "0.01", ClampMax = "0.25"))
+		float AnalysisPeriod;
+
+		/** Minimum analysis frequency for calculating loudness. */
+		UPROPERTY(EditAnywhere, BlueprintReadWrite, Category=AudioAnalyzer, meta = (ClampMin = "20.0", ClampMax = "20000"))
+		float MinimumFrequency;
+
+		/** Maximum analysis frequency for calculating loudness. */
+		UPROPERTY(EditAnywhere, BlueprintReadWrite, Category=AudioAnalyzer, meta = (ClampMin = "20.0", ClampMax = "20000"))
+		float MaximumFrequency;
+
+		/** Type of equal loudness curve used in calculations */
+		UPROPERTY(EditAnywhere, BlueprintReadWrite, Category=AudioAnalyzer)
+		ELoudnessNRTCurveTypeEnum CurveType;
+
+		/** dB level to denote silence.  Used when calculating normalized loudness. */
+		UPROPERTY(EditAnywhere, BlueprintReadWrite, AdvancedDisplay, Category=AudioAnalyzer, meta = (ClampMin = "-100.0", ClampMax="0"))
+		float NoiseFloorDb;
+
+		/** Convert ULoudnessNRTSettings to FLoudnessNRTSettings */
+		TUniquePtr<Audio::IAnalyzerNRTSettings> GetSettings();
+};
+
+
+/** ULoudnessNRT
+ *
+ * ULoudnessNRT calculates the temporal evolution of perceptual loudness for a given 
+ * sound. Loudness is available for individual channels or the overall sound asset. Normalized 
+ * loudness values convert the range to 0.0 to 1.0 where 0.0 is the noise floor and 1.0 is the
+ * maximum loudness of the particular sound.
+ */
+UCLASS(Blueprintable)
+class AUDIOSYNESTHESIA_API ULoudnessNRT : public UAudioSynesthesiaNRT
+{
+	GENERATED_BODY()
+
+	public:
+
+		ULoudnessNRT();
+
+		/** The settings for the audio analyzer.  */
+		UPROPERTY(EditAnywhere, BlueprintReadWrite, Category=AudioAnalyzer)
+		ULoudnessNRTSettings* Settings;
+
+		/** Get the overall loudness of the analyzed sound at a given time. */
+		UFUNCTION(BlueprintCallable, Category="Audio Analyzer")
+		void GetLoudnessAtTime(const float InSeconds, float& OutLoudness);
+
+		/** Get a specific channel loudness of the analyzed sound at a given time. */
+		UFUNCTION(BlueprintCallable, Category="Audio Analyzer")
+		void GetChannelLoudnessAtTime(const float InSeconds, const int32 InChannel, float& OutLoudness);
+
+		/**
+		 * Get the normalized overall loudness of the analyzed sound at a given time. Normalized loudness
+		 * is always between 0.0 to 1.0. 0.0 refers to the noise floor while 1.0 refers to the maximum 
+		 * loudness in the sound.
+		 */
+		UFUNCTION(BlueprintCallable, Category="Audio Analyzer")
+		void GetNormalizedLoudnessAtTime(const float InSeconds, float& OutLoudness);
+
+		/**
+		 * Get a specific channel normalized loudness of the analyzed sound at a given time. Normalized 
+		 * loudness is always between 0.0 to 1.0. 0.0 refers to the noise floor while 1.0 refers to the 
+		 * maximum loudness in the sound.
+		 */
+		UFUNCTION(BlueprintCallable, Category="Audio Analyzer")
+		void GetNormalizedChannelLoudnessAtTime(const float InSeconds, const int32 InChannel, float& OutLoudness);
+
+		/** Convert ULoudnessNRTSettings to FLoudnessNRTSettings */
+ 		TUniquePtr<Audio::IAnalyzerNRTSettings> GetSettings() override;
+
+	protected:
+
+		/** Return the name of the IAudioAnalyzerNRTFactory associated with this UAudioAnalyzerNRT */
+		FName GetAnalyzerNRTFactoryName() const override;
+
+};
+
diff --git a/Engine/Plugins/Runtime/AudioSynesthesia/Source/AudioSynesthesia/Private/AudioSynesthesiaModule.cpp b/Engine/Plugins/Runtime/AudioSynesthesia/Source/AudioSynesthesia/Private/AudioSynesthesiaModule.cpp
--- a/Engine/Plugins/Runtime/AudioSynesthesia/Source/AudioSynesthesia/Private/AudioSynesthesiaModule.cpp
+++ b/Engine/Plugins/Runtime/AudioSynesthesia/Source/AudioSynesthesia/Private/AudioSynesthesiaModule.cpp
@@ -4,6 +4,7 @@
 #include "CoreMinimal.h"
 #include "Modules/ModuleInterface.h"
 #include "Features/IModularFeatures.h"
+#include "LoudnessNRTFactory.h"
 
 DEFINE_LOG_CATEGORY(LogAudioSynesthesia);
 
@@ -13,9 +14,13 @@ namespace Audio
 {
 	void FAudioSynesthesiaModule::StartupModule()
 	{
+		// Register factories on startup
+		IModularFeatures::Get().RegisterModularFeature(FLoudnessNRTFactory::GetModularFeatureName(), &LoudnessFactory);
 	}
 
 	void FAudioSynesthesiaModule::ShutdownModule()
 	{
+		// Unregister factories on shutdown
+		IModularFeatures::Get().UnregisterModularFeature(FLoudnessNRTFactory::GetModularFeatureName(), &LoudnessFactory);
 	}
 }
diff --git a/Engine/Plugins/Runtime/AudioSynesthesia/Source/AudioSynesthesia/Private/LoudnessAnalyzer.cpp b/Engine/Plugins/Runtime/AudioSynesthesia/Source/AudioSynesthesia/Private/LoudnessAnalyzer.cpp
new file mode 100644
--- /dev/null
+++ b/Engine/Plugins/Runtime/AudioSynesthesia/Source/AudioSynesthesia/Private/LoudnessAnalyzer.cpp
@@ -0,0 +1,271 @@
+// Copyright 1998-2019 Epic Games, Inc. All Rights Reserved.
+
+
+#include "LoudnessAnalyzer.h"
+#include "AudioMixer.h"
+#include "DSP/BufferVectorOperations.h"
+#include "DSP/DeinterleaveView.h"
+
+namespace 
+{
+	static const float Log10Scale = 1.0f / FMath::Loge(10.f);
+}
+
+namespace Audio
+{
+    /* Equal loudness curves from: ANSI Standards S1.4-1983 and S1.42-2001 */
+    /* https://en.wikipedia.org/wiki/A-weighting */
+    float GetEqualLoudnessAWeightForFrequency(const float Freq)
+    {
+        const float FreqSqrd = Freq * Freq;
+
+        return 1.258925411794167f * ((148840000.f * FreqSqrd * FreqSqrd) / ((FreqSqrd + 424.36f) * (FreqSqrd + 148840000.f) * FMath::Sqrt((FreqSqrd + 11599.29f) * (FreqSqrd + 544496.41f))));
+    }
+
+    float GetEqualLoudnessBWeightForFrequency(const float Freq)
+    {
+        const float FreqSqrd = Freq * Freq;
+
+        return 1.019764760044717f * ((148840000.f * Freq * FreqSqrd) / ((FreqSqrd + 424.36f) * (FreqSqrd + 148840000.f) * FMath::Sqrt(FreqSqrd + 25122.25f)));
+    }
+
+    float GetEqualLoudnessCWeightForFrequency(const float Freq)
+    {
+        const float FreqSqrd = Freq * Freq;
+
+        return 1.006931668851804f * ((148840000.f * FreqSqrd) / ((FreqSqrd + 424.36f) * (FreqSqrd + 148840000.f)));
+    }
+
+    float GetEqualLoudnessDWeightForFrequency(const float Freq)
+    {
+        const float FreqSqrd = Freq * Freq;
+        const float H1 = (1037918.48f - FreqSqrd);
+        const float H2 = (9837328.f - FreqSqrd);
+        const float H = (H1 + 1080768.16f * FreqSqrd) / (H2 + 11723776.f * FreqSqrd);
+		return (Freq / 6.8966888496476e-5f) * FMath::Sqrt(H / ((FreqSqrd + 79919.29f) * (FreqSqrd + 1345600.f)));
+    }
+
+    float GetEqualLoudnessNoneWeightForFrequency(const float Freq)
+    {
+        return 1.0f;
+    }
+    
+
+	/****************************************************************************/
+	/***********************    FLoudnessAnalyzer    ****************************/
+	/****************************************************************************/
+
+    FLoudnessAnalyzer::FLoudnessAnalyzer(float InSampleRate, const FLoudnessAnalyzerSettings& InSettings)
+    :   Settings(InSettings)
+    ,   EnergyScale(1.0f)
+    ,   MinFreqIndex(0)
+    ,   MaxFreqIndex(0)
+    ,   FFTFreqSize((Settings.FFTSize / 2) + 1)
+    ,   Window(Settings.WindowType, Settings.FFTSize, 1, false)
+    {
+		checkf(FMath::IsPowerOfTwo(Settings.FFTSize), TEXT("FFT size must be a power of two"));
+		checkf(Settings.FFTSize >= 4, TEXT("FFT size must be atleast 4"))
+        check(InSampleRate > 0);
+
+        // Determine which fft bins will be used for calculating loudness
+        const float FreqPerBin = InSampleRate / Settings.FFTSize;
+        MinFreqIndex = static_cast<int32>(FMath::RoundHalfFromZero(Settings.MinAnalysisFrequency / FreqPerBin));
+        MinFreqIndex = FMath::Max(0, MinFreqIndex);
+        MaxFreqIndex = static_cast<int32>(FMath::RoundHalfFromZero(Settings.MaxAnalysisFrequency / FreqPerBin)) + 1;
+        MaxFreqIndex = FMath::Min(MaxFreqIndex, FFTFreqSize);
+
+        // Resize internal buffers
+        WindowedSamples.Reset(Settings.FFTSize);
+        WindowedSamples.AddUninitialized(Settings.FFTSize);
+
+        FFTOutputRealData.Reset(Settings.FFTSize);
+        FFTOutputRealData.AddUninitialized(Settings.FFTSize);
+
+        FFTOutputImagData.Reset(Settings.FFTSize);
+        FFTOutputImagData.AddUninitialized(Settings.FFTSize);
+
+        CurveWeights.Reset(FFTFreqSize);
+        CurveWeights.AddUninitialized(FFTFreqSize);
+
+        PowerSpectrumBuffer.Reset(FFTFreqSize);
+        PowerSpectrumBuffer.AddUninitialized(FFTFreqSize);
+
+        float(*LoudnessMapFn)(const float);
+        switch (Settings.LoudnessCurveType)
+        {
+            case ELoudnessCurveType::A:
+                LoudnessMapFn = GetEqualLoudnessAWeightForFrequency;
+                break;
+
+            case ELoudnessCurveType::B:
+                LoudnessMapFn = GetEqualLoudnessBWeightForFrequency;
+                break;
+
+            case ELoudnessCurveType::C:
+                LoudnessMapFn = GetEqualLoudnessCWeightForFrequency;
+                break;
+
+            case ELoudnessCurveType::D:
+                LoudnessMapFn = GetEqualLoudnessDWeightForFrequency;
+                break;
+
+            case ELoudnessCurveType::None:
+            default:
+                LoudnessMapFn = GetEqualLoudnessNoneWeightForFrequency;
+                break;
+        }
+
+        /* Calculate equal loudness weighting curve */
+
+        // Convention is to normalize 1kHz to have a loudness of 1.
+        float EqualLoudnessNorm = 1.0f / LoudnessMapFn(1000.0f);
+
+        float* CurveWeightsPtr = CurveWeights.GetData();
+
+        const float FFTSize = static_cast<float>(Settings.FFTSize);
+
+        for (int32 i = 0; i < FFTFreqSize; i++)
+        {
+            const float Freq = static_cast<float>(i) * InSampleRate / FFTSize;
+			float Weighting = EqualLoudnessNorm;
+
+			switch (Settings.LoudnessCurveType)
+			{
+				case ELoudnessCurveType::A:
+					Weighting *= GetEqualLoudnessAWeightForFrequency(Freq);
+					break;
+
+				case ELoudnessCurveType::B:
+					Weighting *= GetEqualLoudnessBWeightForFrequency(Freq);
+					break;
+
+				case ELoudnessCurveType::C:
+					Weighting *= GetEqualLoudnessCWeightForFrequency(Freq);
+					break;
+
+				case ELoudnessCurveType::D:
+					Weighting *= GetEqualLoudnessDWeightForFrequency(Freq);
+					break;
+
+				case ELoudnessCurveType::None:
+				default:
+					Weighting *= GetEqualLoudnessNoneWeightForFrequency(Freq);
+					break;
+			}
+
+            // Curve designed for magnitude domain, but applied to power spectrum. The curve is squared to be applied to power spectrum.
+            CurveWeightsPtr[i] = Weighting * Weighting;
+        }
+
+		// Normalize by FFT Scaling (1 / FFTSize) and by number of samples (1 / WindowSize) to calculate mean squared value.
+		EnergyScale = 1.f / static_cast<float>(Settings.FFTSize * Settings.FFTSize);
+    }
+
+    const FLoudnessAnalyzerSettings& FLoudnessAnalyzer::GetSettings() const
+    {
+        return Settings;
+    }
+
+    float FLoudnessAnalyzer::CalculatePerceptualEnergy(TArrayView<const float> InView)
+    {
+		check(InView.Num() == Settings.FFTSize);
+
+		// Copy input samples and apply window
+        FMemory::Memcpy(WindowedSamples.GetData(), InView.GetData(), sizeof(float) * Settings.FFTSize);
+        Window.ApplyToBuffer(WindowedSamples.GetData());
+
+		// Take FFT of windowed samples
+        const FFTTimeDomainData TimeData = {WindowedSamples.GetData(), Settings.FFTSize};
+        FFTFreqDomainData FreqData = {FFTOutputRealData.GetData(), FFTOutputImagData.GetData()};
+
+        PerformFFT(TimeData, FreqData);
+
+		// Accumulate energy as weighted sum of power spectrum in frequency analysis range.
+        const float* FFTOutputRealPtr = FFTOutputRealData.GetData();
+        const float* FFTOutputImagPtr = FFTOutputImagData.GetData();
+        float* PowerSpectrumPtr = PowerSpectrumBuffer.GetData();
+		const float* CurveWeightsPtr = CurveWeights.GetData();
+
+		float Total = 0.0f;
+        for (int32 i = MinFreqIndex; i < MaxFreqIndex; i++)
+        {
+            PowerSpectrumPtr[i] = FFTOutputRealPtr[i] * FFTOutputRealPtr[i] + FFTOutputImagPtr[i] * FFTOutputImagPtr[i];
+			Total += PowerSpectrumPtr[i] * CurveWeightsPtr[i];
+        }
+
+		// Normalize by FFT Scaling and by number of samples to calculate mean squared value.
+		Total *= EnergyScale;
+
+        return Total;
+    }
+
+    float FLoudnessAnalyzer::CalculateLoudness(TArrayView<const float> InView)
+    {
+        float Energy = CalculatePerceptualEnergy(InView);
+		return ConvertPerceptualEnergyToLoudness(Energy);
+    }
+
+	float FLoudnessAnalyzer::ConvertPerceptualEnergyToLoudness(float InPerceptualEnergy)
+	{
+        return 10.0f * FMath::Loge(InPerceptualEnergy) * Log10Scale;
+	}
+
+	/****************************************************************************/
+	/****************    FMultichannelLoudnessAnalyzer    ***********************/
+	/****************************************************************************/
+
+    FMultichannelLoudnessAnalyzer::FMultichannelLoudnessAnalyzer(float InSampleRate, const FLoudnessAnalyzerSettings& InSettings)
+    :   Analyzer(InSampleRate, InSettings)
+    {
+        // Initialize all channel weights to default to 1.0f
+        ChannelWeights.Init(1.0f, EAudioMixerChannel::ChannelTypeCount);
+
+        // If elevation angle is less than 30 degrees and azimuth is between 60 and 120, set channel weight to 1.5f. 
+        ChannelWeights[EAudioMixerChannel::BackLeft] = 1.5f;
+        ChannelWeights[EAudioMixerChannel::BackRight] = 1.5f;
+
+        MonoBuffer.Reset(InSettings.FFTSize);
+        MonoBuffer.AddUninitialized(InSettings.FFTSize);
+    }
+
+    float FMultichannelLoudnessAnalyzer::CalculatePerceptualEnergy(TArrayView<const float> InView, const int32 InNumChannels, TArray<float>& OutChannelEnergies)
+    {
+		check(InView.Num() == (InNumChannels * Analyzer.GetSettings().FFTSize));
+
+        OutChannelEnergies.Reset(InNumChannels);
+        OutChannelEnergies.AddUninitialized(InNumChannels);
+
+        // Deinterleave audio and send to mono loudness analyzer.
+		TAutoDeinterleaveView<float, FAudioBufferAlignedAllocator> DeinterleaveView(InView, MonoBuffer, InNumChannels);
+		for (TAutoDeinterleaveView<float>::TChannel<FAudioBufferAlignedAllocator> Channel : DeinterleaveView)
+		{
+            OutChannelEnergies[Channel.ChannelIndex] = Analyzer.CalculatePerceptualEnergy(Channel.Values);
+		}
+				
+        // Combine channel energies into overall perceptual energy
+        float Energy = 0.0f;
+        for (int32 ChannelNum = 0; ChannelNum < InNumChannels; ChannelNum++)
+        {
+            Energy += ChannelWeights[ChannelNum] * OutChannelEnergies[ChannelNum];
+        }
+
+        return Energy;
+    }
+
+    float FMultichannelLoudnessAnalyzer::CalculateLoudness(TArrayView<const float> InView, const int32 InNumChannels, TArray<float>& OutChannelLoudness)
+    {
+        // Convert perceptual energy to loudness 
+        float Energy = CalculatePerceptualEnergy(InView, InNumChannels, OutChannelLoudness);
+        for (float& Value : OutChannelLoudness)
+        {
+            Value = FLoudnessAnalyzer::ConvertPerceptualEnergyToLoudness(Value);
+        }
+        return FLoudnessAnalyzer::ConvertPerceptualEnergyToLoudness(Energy);
+    }
+
+    const FLoudnessAnalyzerSettings& FMultichannelLoudnessAnalyzer::GetSettings() const
+    {
+        return Analyzer.GetSettings();
+    }
+}
+
diff --git a/Engine/Plugins/Runtime/AudioSynesthesia/Source/AudioSynesthesia/Private/LoudnessNRT.cpp b/Engine/Plugins/Runtime/AudioSynesthesia/Source/AudioSynesthesia/Private/LoudnessNRT.cpp
new file mode 100644
--- /dev/null
+++ b/Engine/Plugins/Runtime/AudioSynesthesia/Source/AudioSynesthesia/Private/LoudnessNRT.cpp
@@ -0,0 +1,161 @@
+// Copyright 1998-2019 Epic Games, Inc. All Rights Reserved.
+
+#include "LoudnessNRT.h"
+#include "LoudnessNRTFactory.h"
+#include "InterpolateSorted.h"
+
+namespace
+{
+	float InterpolateLoudness(TArrayView<const Audio::FLoudnessDatum> InLoudnessArray, const float InTimestamp)
+	{
+		int32 LowerIndex = INDEX_NONE;
+		int32 UpperIndex = INDEX_NONE;
+		float Alpha = 0.f;
+
+		GetInterpolationParametersAtTimestamp(InLoudnessArray, InTimestamp, LowerIndex, UpperIndex, Alpha);
+
+		if ((INDEX_NONE != LowerIndex) && (INDEX_NONE != UpperIndex))
+		{
+			return FMath::Lerp(InLoudnessArray[LowerIndex].Loudness, InLoudnessArray[UpperIndex].Loudness, Alpha);
+		}
+
+		return 0.f;
+	}
+}
+
+
+/***************************************************************************/
+/**********************    ULoudnessNRTSettings     ************************/
+/***************************************************************************/
+
+ULoudnessNRTSettings::ULoudnessNRTSettings()
+:	AnalysisPeriod(0.01f)	
+,	MinimumFrequency(20.f)
+,	MaximumFrequency(20000.f)
+,	CurveType(ELoudnessNRTCurveTypeEnum::D)
+,	NoiseFloorDb(-60.f)
+{}
+
+TUniquePtr<Audio::IAnalyzerNRTSettings> ULoudnessNRTSettings::GetSettings()
+{
+	TUniquePtr<Audio::FLoudnessNRTSettings> Settings = MakeUnique<Audio::FLoudnessNRTSettings>();
+
+	Settings->AnalysisPeriod = AnalysisPeriod;
+	Settings->MinAnalysisFrequency = MinimumFrequency;
+	Settings->MaxAnalysisFrequency = MaximumFrequency;
+
+	switch (CurveType)
+	{
+		case ELoudnessNRTCurveTypeEnum::A:
+			Settings->LoudnessCurveType = Audio::ELoudnessCurveType::A;
+			break;
+
+		case ELoudnessNRTCurveTypeEnum::B:
+			Settings->LoudnessCurveType = Audio::ELoudnessCurveType::B;
+			break;
+			
+		case ELoudnessNRTCurveTypeEnum::C:
+			Settings->LoudnessCurveType = Audio::ELoudnessCurveType::C;
+			break;
+
+		case ELoudnessNRTCurveTypeEnum::D:
+			Settings->LoudnessCurveType = Audio::ELoudnessCurveType::D;
+			break;
+
+		case ELoudnessNRTCurveTypeEnum::None:
+		default:
+			Settings->LoudnessCurveType = Audio::ELoudnessCurveType::None;
+			break;
+	}
+	
+	return Settings;
+}
+
+/***************************************************************************/
+/**********************        ULoudnessNRT         ************************/
+/***************************************************************************/
+
+ULoudnessNRT::ULoudnessNRT()
+{
+	Settings = CreateDefaultSubobject<ULoudnessNRTSettings>(TEXT("DefaultLoudnessNRTSettings"));
+#if WITH_EDITOR
+	// Bind settings to audio analyze so changes to default settings will trigger analysis.
+	Settings->AnalyzeAudioDelegate.BindUObject(this, &UAudioAnalyzerNRT::AnalyzeAudio);
+#endif
+}
+
+void ULoudnessNRT::GetLoudnessAtTime(const float InSeconds, float& OutLoudness)
+{
+	GetChannelLoudnessAtTime(InSeconds, Audio::FLoudnessNRTResult::ChannelIndexOverall, OutLoudness);
+}
+
+void ULoudnessNRT::GetChannelLoudnessAtTime(const float InSeconds, const int32 InChannel, float& OutLoudness)
+{
+	OutLoudness = 0.0f;
+
+	TSharedPtr<Audio::FLoudnessNRTResult, ESPMode::ThreadSafe> LoudnessResult = GetResult<Audio::FLoudnessNRTResult>();
+
+	if (LoudnessResult.IsValid())
+	{
+		if (!LoudnessResult->IsSortedChronologically())
+		{
+			LoudnessResult->SortChronologically();
+		}
+
+		const TArray<Audio::FLoudnessDatum>& LoudnessArray = LoudnessResult->GetChannelLoudnessArray(InChannel);
+
+		OutLoudness = InterpolateLoudness(LoudnessArray, InSeconds);
+	}
+}
+
+void ULoudnessNRT::GetNormalizedLoudnessAtTime(const float InSeconds, float& OutLoudness)
+{
+	GetNormalizedChannelLoudnessAtTime(InSeconds, Audio::FLoudnessNRTResult::ChannelIndexOverall, OutLoudness);
+}
+
+void ULoudnessNRT::GetNormalizedChannelLoudnessAtTime(const float InSeconds, const int32 InChannel, float& OutLoudness)
+{
+	OutLoudness = 0.0f;
+
+	TSharedPtr<Audio::FLoudnessNRTResult, ESPMode::ThreadSafe> LoudnessResult = GetResult<Audio::FLoudnessNRTResult>();
+
+	if (LoudnessResult.IsValid())
+	{
+		if (!LoudnessResult->IsSortedChronologically())
+		{
+			LoudnessResult->SortChronologically();
+		}
+
+		const TArray<Audio::FLoudnessDatum>& LoudnessArray = LoudnessResult->GetChannelLoudnessArray(InChannel);
+
+		OutLoudness = InterpolateLoudness(LoudnessArray, InSeconds);
+
+		// Subtract offset
+		OutLoudness -= Settings->NoiseFloorDb;
+
+		// Divide by range
+		OutLoudness /= FMath::Max(LoudnessResult->GetChannelLoudnessRange(InChannel, Settings->NoiseFloorDb), SMALL_NUMBER);
+
+		// Clamp to desired range in case value was below noise floor
+		OutLoudness = FMath::Clamp(OutLoudness, 0.f, 1.f);
+	}
+}
+
+TUniquePtr<Audio::IAnalyzerNRTSettings> ULoudnessNRT::GetSettings()
+{
+	TUniquePtr<Audio::IAnalyzerNRTSettings> AnalyzerSettings;
+
+	if (Settings)
+	{
+		AnalyzerSettings = Settings->GetSettings();	
+	}
+
+	return AnalyzerSettings;
+}
+
+FName ULoudnessNRT::GetAnalyzerNRTFactoryName() const
+{
+	static const FName FactoryName(TEXT("LoudnessNRTFactory"));
+	return FactoryName;
+}
+
diff --git a/Engine/Plugins/Runtime/AudioSynesthesia/Source/AudioSynesthesia/Private/LoudnessNRTFactory.cpp b/Engine/Plugins/Runtime/AudioSynesthesia/Source/AudioSynesthesia/Private/LoudnessNRTFactory.cpp
new file mode 100644
--- /dev/null
+++ b/Engine/Plugins/Runtime/AudioSynesthesia/Source/AudioSynesthesia/Private/LoudnessNRTFactory.cpp
@@ -0,0 +1,249 @@
+// Copyright 1998-2019 Epic Games, Inc. All Rights Reserved.
+
+#include "LoudnessNRTFactory.h"
+#include "DSP/SlidingWindow.h"
+
+namespace Audio
+{
+	/************************************************************************/
+	/********************** FLoudnessNRTResult ******************************/
+	/************************************************************************/
+	FArchive &operator <<(FArchive& Ar, FLoudnessDatum& Datum)
+	{
+		Ar << Datum.Channel;
+		Ar << Datum.Timestamp;
+		Ar << Datum.Energy;
+		Ar << Datum.Loudness;
+
+		return Ar;
+	}
+
+	// Use for overall loudness storage
+	const int32 FLoudnessNRTResult::ChannelIndexOverall = -1;
+
+
+	FLoudnessNRTResult::FLoudnessNRTResult()
+	: bIsSortedChronologically(false)
+	{
+		// By  default, overall loudness exists
+		ChannelLoudnessArrays.Add(ChannelIndexOverall);
+		ChannelLoudnessIntervals.Add(ChannelIndexOverall);
+	}
+
+
+	void FLoudnessNRTResult::Serialize(FArchive& Archive)
+	{
+		Archive << bIsSortedChronologically;
+		Archive << ChannelLoudnessArrays;
+		Archive << ChannelLoudnessIntervals;
+	}
+
+	void FLoudnessNRTResult::Add(const FLoudnessDatum& InDatum)
+	{
+		// Store loudness data in appropriate channel
+		TArray<FLoudnessDatum>& LoudnessArray = ChannelLoudnessArrays.FindOrAdd(InDatum.Channel);
+		LoudnessArray.Add(InDatum);
+
+		// Update intervals per channel
+		FFloatInterval& Interval = ChannelLoudnessIntervals.FindOrAdd(InDatum.Channel);
+		Interval.Include(InDatum.Loudness);
+
+		// Mark as not sorted
+		bIsSortedChronologically = false;
+	}
+
+	const TArray<FLoudnessDatum>& FLoudnessNRTResult::GetChannelLoudnessArray(int32 ChannelIdx) const
+	{
+		return ChannelLoudnessArrays[ChannelIdx];
+	}
+
+	const TArray<FLoudnessDatum>& FLoudnessNRTResult::GetLoudnessArray() const
+	{
+		return ChannelLoudnessArrays[ChannelIndexOverall];
+	}
+
+	float FLoudnessNRTResult::GetLoudnessRange(float InNoiseFloor) const
+	{
+		return GetChannelLoudnessRange(ChannelIndexOverall, InNoiseFloor);
+	}
+
+	float FLoudnessNRTResult::GetChannelLoudnessRange(int32 InChannelIdx, float InNoiseFloor) const
+	{
+		const FFloatInterval& Interval = ChannelLoudnessIntervals.FindRef(InChannelIdx);
+
+		if (Interval.Contains(InNoiseFloor))
+		{
+			// Noise floor is greater than minimum, so range must be relative to noise floor
+			return Interval.Max - InNoiseFloor;
+		}
+		else if (InNoiseFloor > Interval.Max)
+		{
+			return 0.f;
+		}
+
+		// Noise floor is less than minimum, so use entire range.
+		return Interval.Size();
+	}
+
+	int32 FLoudnessNRTResult::GetNumChannels() const 
+	{
+		return FMath::Max(0, ChannelLoudnessArrays.Num() - 1);
+	}
+
+	bool FLoudnessNRTResult::IsSortedChronologically() const
+	{
+		return bIsSortedChronologically;
+	}
+
+	void FLoudnessNRTResult::SortChronologically()
+	{
+		for (auto& KeyValue : ChannelLoudnessArrays)
+		{
+			KeyValue.Value.Sort([](const FLoudnessDatum& A, const FLoudnessDatum& B) { return A.Timestamp < B.Timestamp; });
+		}
+
+		bIsSortedChronologically = true;
+	}
+
+	/************************************************************************/
+	/********************** FLoudnessNRTWorker ******************************/
+	/************************************************************************/
+
+	FLoudnessNRTWorker::FLoudnessNRTWorker(const FAnalyzerNRTParameters& InParams, const FLoudnessNRTSettings& InAnalyzerSettings)
+	: NumChannels(InParams.NumChannels)
+	, NumAnalyzedBuffers(0)
+	, NumHopFrames(0)
+	, SampleRate(InParams.SampleRate)
+	{
+		check(NumChannels >= 0);
+
+		// Include NumChannels in calculating window size because windows are generated 
+		// with interleaved samples and deinterleaved later. 
+		int32 NumWindowSamples = InAnalyzerSettings.FFTSize * NumChannels;
+
+		NumHopFrames = FMath::CeilToInt(InAnalyzerSettings.AnalysisPeriod * InParams.SampleRate);
+
+		NumHopFrames = FMath::Max(1, NumHopFrames);
+
+		int32 NumHopSamples = NumHopFrames * NumChannels;
+
+		InternalBuffer = MakeUnique<TSlidingBuffer<float>>(NumWindowSamples, NumHopSamples);
+
+		Analyzer = MakeUnique<FMultichannelLoudnessAnalyzer>(InParams.SampleRate, InAnalyzerSettings);
+	}
+
+	void FLoudnessNRTWorker::Analyze(TArrayView<const float> InAudio, IAnalyzerNRTResult* OutResult) 
+	{
+		// Assume that outer layers ensured that this is of correct type.
+		FLoudnessNRTResult* LoudnessResult = static_cast<FLoudnessNRTResult*>(OutResult);
+
+		check(nullptr != LoudnessResult);
+		
+		TAutoSlidingWindow<float> SlidingWindow(*InternalBuffer, InAudio, InternalWindow);
+
+		// Loop through entire array of input samples. 
+		for (const TArray<float>& Window : SlidingWindow)
+		{
+			AnalyzeWindow(Window, *LoudnessResult);
+		}
+	}
+
+	// Called when analysis of audio asset is complete. 
+	void FLoudnessNRTWorker::Finalize(IAnalyzerNRTResult* OutResult) 
+	{
+		FLoudnessNRTResult* LoudnessResult = static_cast<FLoudnessNRTResult*>(OutResult);
+
+		check(nullptr != LoudnessResult);
+
+		TArray<float> EmptyArray;
+		bool bDoFlush = true; // Add zeros to last window until all data has been processed.
+		TAutoSlidingWindow<float> SlidingWindow(*InternalBuffer, EmptyArray, InternalWindow, bDoFlush);
+
+		// Loop through entire array of input samples. 
+		for (const TArray<float>& Window : SlidingWindow)
+		{
+			AnalyzeWindow(Window, *LoudnessResult);
+		}
+
+		// Prepare result for fast lookup.
+		LoudnessResult->SortChronologically(); 
+
+		// Reset internal counters
+		NumAnalyzedBuffers = 0; 
+
+		// Reset sliding buffer
+		InternalBuffer->Reset();
+	}
+
+	void FLoudnessNRTWorker::AnalyzeWindow(TArrayView<const float> InWindow, FLoudnessNRTResult& OutResult)
+	{
+		TArray<float> ChannelPerceptualEnergy;
+		float OverallPerceptualEnergy = 0.0f;
+
+		// Calculate perceptual energy
+		OverallPerceptualEnergy = Analyzer->CalculatePerceptualEnergy(InWindow, NumChannels, ChannelPerceptualEnergy);
+
+		// Convert to perceptual energy to loudness
+		TArray<float> ChannelLoudness;
+		float OverallLoudness = 0.f;
+		
+		OverallLoudness = FLoudnessAnalyzer::ConvertPerceptualEnergyToLoudness(OverallPerceptualEnergy);
+
+		for (float PerceptualEnergy : ChannelPerceptualEnergy)
+		{
+			ChannelLoudness.Add(FLoudnessAnalyzer::ConvertPerceptualEnergyToLoudness(PerceptualEnergy));
+		}
+
+		// Calculate timestamp
+		float Timestamp = ((NumAnalyzedBuffers * NumHopFrames) + (0.5f * Analyzer->GetSettings().FFTSize)) / SampleRate;
+
+		// Add overall FLoudnessDatum to result
+		const FLoudnessDatum OverallDatum = {FLoudnessNRTResult::ChannelIndexOverall, Timestamp, OverallPerceptualEnergy, OverallLoudness};
+		OutResult.Add(OverallDatum);
+
+		// Add channel FLoudnessDatum to result
+		for (int32 ChannelIdx = 0; ChannelIdx < NumChannels; ChannelIdx++)
+		{
+			const FLoudnessDatum ChannelDatum = {ChannelIdx, Timestamp, ChannelPerceptualEnergy[ChannelIdx], ChannelLoudness[ChannelIdx]};
+
+			OutResult.Add(ChannelDatum);
+		}	
+
+		// update counters
+		NumAnalyzedBuffers++;
+	}
+
+
+	/************************************************************************/
+	/************************* FLoudnessNRTFactory **************************/
+	/************************************************************************/
+
+	/** Name of specific analyzer type. */
+	FName FLoudnessNRTFactory::GetName() const 
+	{
+		static FName FactoryName(TEXT("LoudnessNRTFactory"));
+		return FactoryName;
+	}
+
+	/** Human readable name of analyzer. */
+	FString FLoudnessNRTFactory::GetTitle() const
+	{
+		return TEXT("Loudness Analyzer Non-Real-Time");
+	}
+
+	TUniquePtr<IAnalyzerNRTResult> FLoudnessNRTFactory::NewResult()
+	{
+		TUniquePtr<FLoudnessNRTResult> Result = MakeUnique<FLoudnessNRTResult>();
+		return Result;
+	}
+
+	TUniquePtr<IAnalyzerNRTWorker> FLoudnessNRTFactory::NewWorker(const FAnalyzerNRTParameters& InParams, const IAnalyzerNRTSettings* InSettings)
+	{
+		const FLoudnessNRTSettings* LoudnessSettings = static_cast<const FLoudnessNRTSettings*>(InSettings);
+
+		check(nullptr != LoudnessSettings);
+
+		return MakeUnique<FLoudnessNRTWorker>(InParams, *LoudnessSettings);
+	}
+}
+
diff --git a/Engine/Plugins/Runtime/AudioSynesthesia/Source/AudioSynesthesia/Public/AudioSynesthesiaModule.h b/Engine/Plugins/Runtime/AudioSynesthesia/Source/AudioSynesthesia/Public/AudioSynesthesiaModule.h
--- a/Engine/Plugins/Runtime/AudioSynesthesia/Source/AudioSynesthesia/Public/AudioSynesthesiaModule.h
+++ b/Engine/Plugins/Runtime/AudioSynesthesia/Source/AudioSynesthesia/Public/AudioSynesthesiaModule.h
@@ -5,6 +5,7 @@
 #include "CoreMinimal.h"
 #include "Modules/ModuleInterface.h"
 #include "Modules/ModuleManager.h"
+#include "LoudnessNRTFactory.h"
 
 DECLARE_LOG_CATEGORY_EXTERN(LogAudioSynesthesia, Log, All);
 
@@ -19,5 +20,6 @@ namespace Audio
 		virtual void ShutdownModule() override;
 
 	private:
+		FLoudnessNRTFactory LoudnessFactory;
 	};
 }
diff --git a/Engine/Plugins/Runtime/AudioSynesthesia/Source/AudioSynesthesia/Public/InterpolateSorted.h b/Engine/Plugins/Runtime/AudioSynesthesia/Source/AudioSynesthesia/Public/InterpolateSorted.h
new file mode 100644
--- /dev/null
+++ b/Engine/Plugins/Runtime/AudioSynesthesia/Source/AudioSynesthesia/Public/InterpolateSorted.h
@@ -0,0 +1,48 @@
+// Copyright 1998-2019 Epic Games, Inc. All Rights Reserved.
+
+#pragma once
+
+namespace Audio
+{
+	/** GetInterpolationParametersAtTimestamp
+	 *
+	 *  Calculates the indices and alpha parameter for interpolating between to elements in the array.
+	 *
+	 *  ElementType must be a type with a "float Timestamp;" member.
+	 *  InArray is a array of ElementTypes which are sorted in ascending order via the Timestamp member.
+	 *  InTimestamp is the query timestamp.
+	 *
+	 *  OutLowerIndex is populated with the index of the last element in InArray with a Timestamp less than or equal to InTimestamp. On error, this is set to INDEX_NONE
+	 *  OutUpperIndex is populated with the index of the first element in InArray with a Timestamp greater than InTimestamp. On error, this is set to INDEX_NONE
+	 *  OutAlpha is a number between 0.f and 1.f which describes the blending between the lower and upper indices. 0.f refers OutLowerIndex and 1.f to OutUpperIndex.
+	 */
+	template<typename ElementType>
+	void GetInterpolationParametersAtTimestamp(TArrayView<const ElementType> InArray, float InTimestamp, int32& OutLowerIndex, int32& OutUpperIndex, float& OutAlpha)
+	{
+		if (0 == InArray.Num())
+		{
+			OutLowerIndex = INDEX_NONE;
+			OutUpperIndex = INDEX_NONE;
+			OutAlpha = 1.f;
+			return;
+		}
+		
+		// Get upper and lower array indices with FLoudnessData closest to the query timestamp.
+		OutUpperIndex = Algo::UpperBoundBy(InArray, InTimestamp, [](const ElementType& Datum) { return Datum.Timestamp; });
+		OutLowerIndex = OutUpperIndex - 1;
+
+		// Ensure indices are within range.
+		int32 MaxIdx = InArray.Num() - 1;
+		OutLowerIndex = FMath::Clamp(OutLowerIndex, 0, MaxIdx);
+		OutUpperIndex = FMath::Clamp(OutUpperIndex, 0, MaxIdx);
+
+		const ElementType& LowerDatum = InArray[OutLowerIndex];
+		const ElementType& UpperDatum = InArray[OutUpperIndex];
+
+		// Determine time difference and linear interpolation factor from data.
+		const float TimeDelta = FMath::Max(UpperDatum.Timestamp - LowerDatum.Timestamp, SMALL_NUMBER);
+		OutAlpha = (InTimestamp - LowerDatum.Timestamp) / TimeDelta;
+		OutAlpha = FMath::Clamp(OutAlpha, 0.0f, 1.0f);
+	}
+}
+
diff --git a/Engine/Plugins/Runtime/AudioSynesthesia/Source/AudioSynesthesia/Public/LoudnessAnalyzer.h b/Engine/Plugins/Runtime/AudioSynesthesia/Source/AudioSynesthesia/Public/LoudnessAnalyzer.h
new file mode 100644
--- /dev/null
+++ b/Engine/Plugins/Runtime/AudioSynesthesia/Source/AudioSynesthesia/Public/LoudnessAnalyzer.h
@@ -0,0 +1,172 @@
+// Copyright 1998-2019 Epic Games, Inc. All Rights Reserved.
+
+#pragma once
+
+#include "DSP/BufferVectorOperations.h"
+#include "DSP/AudioFFT.h"
+
+namespace Audio
+{
+    /** Equal loudness curves from: ANSI Standards S1.4-1983 and S1.42-2001 */
+	enum class ELoudnessCurveType : uint8
+	{
+		A,
+		B,
+		C,
+		D,
+		None
+	};
+
+	/** Generate equal loudness weight for magnitude spectrum using curve A */
+    float GetEqualLoudnessAWeigthForFrequency(const float Freq);
+
+	/** Generate equal loudness weight for magnitude spectrum using curve B */
+    float GetEqualLoudnessBWeigthForFrequency(const float Freq);
+
+	/** Generate equal loudness weight for magnitude spectrum using curve C */
+    float GetEqualLoudnessCWeigthForFrequency(const float Freq);
+
+	/** Generate equal loudness weight for magnitude spectrum using curve D */
+    float GetEqualLoudnessDWeigthForFrequency(const float Freq);
+
+	/** Generate equal loudness weight for magnitude spectrum using curve None */
+    float GetEqualLoudnessNoneWeigthForFrequency(const float Freq);
+
+	/** Settings for FLoudnessAnalyzer */
+    struct FLoudnessAnalyzerSettings
+    {
+		/** Size of FFT describe in number of samples */
+        int32 FFTSize;
+
+		/** Type of window to be applied to input audio */
+        EWindowType WindowType;
+
+		/** Type of equal loudness curve to weight power spectrum */
+        ELoudnessCurveType LoudnessCurveType;
+
+		/** Minimum frequency for analyzing audio */
+        float MinAnalysisFrequency;
+
+		/** Maximum frequency for analyzing audio */
+        float MaxAnalysisFrequency;
+
+        /** Default settings */
+		FLoudnessAnalyzerSettings()
+			: FFTSize(8192)
+			, WindowType(EWindowType::Blackman)
+			, LoudnessCurveType(ELoudnessCurveType::D)
+			, MinAnalysisFrequency(20.0f)
+			, MaxAnalysisFrequency(10000.0f)
+		{}
+    };
+
+	/** FLoudnessAnalyzer
+	 *
+	 * FLoudnessAnalyzer will calculate perceptual energy or loudness one window at a time. 
+	 * This analyzer only handles mono audio. See FMultichannelLoudnessAnalyzer for multichannel
+	 * support.
+	 */
+    class FLoudnessAnalyzer
+    {
+        public:
+			/** Construct analyzer */
+            FLoudnessAnalyzer(float InSampleRate, const FLoudnessAnalyzerSettings& InSettings);
+
+			/**
+			 * Calculats the perceptual energy for the input samples.  The number of 
+			 * input samples must be equal to the FFTSize supplied in FLoudnessAnalyzerSettings.
+			 */
+            float CalculatePerceptualEnergy(TArrayView<const float> InSampleView);
+
+			/**
+			 * Calculats the loudness for the input samples.  The number of  input samples must 
+			 * be equal to the FFTSize supplied in FLoudnessAnalyzerSettings.
+			 */
+            float CalculateLoudness(TArrayView<const float> InSampleView);
+
+			/**
+			 * Convert from perceptual energy to loudness [Loudness = 10 * log_10(Energy)]
+			 */
+			static float ConvertPerceptualEnergyToLoudness(float InPerceptualEnergy);
+
+			/**
+			 * Return const reference to settings used inside this analyzer.
+			 */
+            const FLoudnessAnalyzerSettings& GetSettings() const;
+        private:
+
+            FLoudnessAnalyzerSettings Settings;
+
+			/** Scale factor used to normalize between differing window sizes and fft sizes */
+            float EnergyScale;
+
+            int32 MinFreqIndex;
+            int32 MaxFreqIndex;
+
+			/** Number of FFTFrequency bins worth looking at. */
+			int32 FFTFreqSize;
+
+            FWindow Window;
+
+            AlignedFloatBuffer WindowedSamples;
+            AlignedFloatBuffer FFTOutputRealData;
+            AlignedFloatBuffer FFTOutputImagData;
+
+            AlignedFloatBuffer CurveWeights;
+            AlignedFloatBuffer PowerSpectrumBuffer;
+    };
+
+
+	/** FMultichannelLoudnessAnalyzer
+	 *
+	 * FMultichannelLoudnessAnalyzer will calculate perceptual energy or loudness one window 
+	 * at a time. This module can handle one or more channels at a time. It expects 
+	 * interleaved audio samples.
+	 */
+    class FMultichannelLoudnessAnalyzer 
+    {
+        public:
+			/** Construct analyzer. */
+            FMultichannelLoudnessAnalyzer(float InSampleRate, const FLoudnessAnalyzerSettings& InSettings);
+
+			/** 
+			 * Calculate perceptual energy. 
+			 *
+			 * InSampleView		  InSampleView is expected to contain interleaved audio with 
+			 *                    (InSettings.FFTSize * InNumChannels) number of samples.
+			 *
+			 * InNumChannels      Number of input channels.
+			 *
+			 * OutChannelEnergies Will be populated with energies on a per channel basis.
+			 *
+			 * The overall perceptual energy is returned.
+			 */
+            float CalculatePerceptualEnergy(TArrayView<const float> InSampleView, const int32 InNumChannels, TArray<float>& OutChannelEnergies); 
+
+			/** 
+			 * Calculate loudness
+			 *
+			 * InSampleView		  InSampleView is expected to contain interleaved audio with 
+			 *                    (InSettings.FFTSize * InNumChannels) number of samples.
+			 *
+			 * InNumChannels      Number of input channels.
+			 *
+			 * OutChannelLoudness Will be populated with energies on a per channel basis.
+			 *
+			 * The overall loudness is returned.
+			 */
+            float CalculateLoudness(TArrayView<const float> InSampleView, const int32 InNumChannels, TArray<float>& OutChannelLoudness);
+
+			/**
+			 * Return const reference to settings used inside this analyzer.
+			 */
+            const FLoudnessAnalyzerSettings& GetSettings() const;
+
+        private:
+
+            AlignedFloatBuffer MonoBuffer;
+            TArray<float> ChannelWeights;
+            FLoudnessAnalyzer Analyzer;
+    };
+}
+
diff --git a/Engine/Plugins/Runtime/AudioSynesthesia/Source/AudioSynesthesia/Public/LoudnessNRTFactory.h b/Engine/Plugins/Runtime/AudioSynesthesia/Source/AudioSynesthesia/Public/LoudnessNRTFactory.h
new file mode 100644
--- /dev/null
+++ b/Engine/Plugins/Runtime/AudioSynesthesia/Source/AudioSynesthesia/Public/LoudnessNRTFactory.h
@@ -0,0 +1,171 @@
+// Copyright 1998-2019 Epic Games, Inc. All Rights Reserved.
+
+#pragma once
+
+#include "LoudnessAnalyzer.h"
+#include "IAudioAnalyzerNRTInterface.h"
+#include "DSP/SlidingWindow.h"
+
+namespace Audio
+{
+	/** FLoudnessNRTSettings
+	 *
+	 * Contains settings for loudness analyzer.
+	 */
+	class FLoudnessNRTSettings : public IAnalyzerNRTSettings, public FLoudnessAnalyzerSettings
+	{
+		public:
+			/** Number of seconds between loudness measurements */
+			float AnalysisPeriod;
+
+			FLoudnessNRTSettings()
+			:	AnalysisPeriod(0.01f)
+			{}
+	};
+
+
+	/** FLoudnessNRTResult
+	 *
+	 * Holds the loundess values per a time step.
+	 */
+	struct FLoudnessDatum
+	{
+		int32 Channel = 0;
+		float Timestamp = 0.f;
+		float Energy = 0.f;
+		float Loudness = 0.f;
+	};
+	/** De/Serialize single loudness datum into archive. */
+	AUDIOSYNESTHESIA_API FArchive &operator <<(FArchive& Ar, FLoudnessDatum& Datum);
+
+	/** FLoudnessNRTResult
+	 *
+	 * FLoudnessNRTResult contains the temporal evolution of loudness.
+	 */
+	class AUDIOSYNESTHESIA_API FLoudnessNRTResult : public IAnalyzerNRTResult
+	{
+	public:
+		/** ChannelIndexOverall
+		 *
+		 * Denotes the overall loudness channel index as opposed individual channel indices.
+		 */
+		static const int32 ChannelIndexOverall;
+
+		FLoudnessNRTResult();
+
+		/**
+		 * Defines how to serialize result.
+		 */
+		void Serialize(FArchive& Archive) override;
+
+		/**
+		 * Appends an FLoudnessDatum to the container.
+		 */
+		void Add(const FLoudnessDatum& InDatum);
+
+		/**
+		 * Returns const reference to FLoudnessDatum array for individual channel.
+		 */
+		const TArray<FLoudnessDatum>& GetChannelLoudnessArray(int32 ChannelIdx) const;
+
+		/**
+		 * Returns const reference to FLoudnessDatum array associated with overall loudness.
+		 */
+		const TArray<FLoudnessDatum>& GetLoudnessArray() const;
+
+		/**
+		 * Returns range in dB of overall loudness result given the noise floor.
+		 */
+		float GetLoudnessRange(float InNoiseFloor) const;
+
+		/**
+		 * Returns range in dB of loudness result given the noise floor.
+		 */
+		float GetChannelLoudnessRange(int32 InChannelIdx, float InNoiseFloor) const;
+
+		/**
+		 * Get Number of Channels
+		 */
+		int32 GetNumChannels() const;
+
+		/**
+		 * Returns true if FLoudnessDatum arrays are sorted in chronologically ascending order via their timestamp.
+		 */
+		bool IsSortedChronologically() const;
+
+		/**
+		 * Sorts FLoudnessDatum arrays in chronologically ascnding order via their timestamp.
+		 */
+		void SortChronologically();
+
+	private:
+		TMap<int32, TArray<FLoudnessDatum> > ChannelLoudnessArrays;
+
+		TMap<int32, FFloatInterval> ChannelLoudnessIntervals;
+
+		bool bIsSortedChronologically;
+	};
+
+	/** FLoudnessNRTWorker
+	 *
+	 * FLoudnessNRTWorker performs loudness analysis on input sample buffers.
+	 */
+	class FLoudnessNRTWorker : public IAnalyzerNRTWorker
+	{
+	public:
+		/** Construct a worker */
+		FLoudnessNRTWorker(const FAnalyzerNRTParameters& InParams, const FLoudnessNRTSettings& InAnalyzerSettings);
+
+		/**
+		 * Analyzes input sample buffer and updates result. 
+		 */
+		void Analyze(TArrayView<const float> InAudio, IAnalyzerNRTResult* OutResult) override;
+
+		/**
+		 * Call when all audio data has been analyzed. 
+		 */
+		void Finalize(IAnalyzerNRTResult* OutResult) override;
+
+	private:
+
+		// Analyze a single window
+		void AnalyzeWindow(TArrayView<const float> InWindow, FLoudnessNRTResult& OutResult);
+
+		int32 NumChannels;
+		int32 NumAnalyzedBuffers;
+		int32 NumHopFrames;
+
+		float SampleRate;
+
+		TArray<float> InternalWindow;
+
+		TUniquePtr<TSlidingBuffer<float>> InternalBuffer;
+
+		TUniquePtr<FMultichannelLoudnessAnalyzer> Analyzer;
+	};
+
+	/** FLoudnessNRTFactory
+	 *
+	 * Defines the LoudnessNRT analyzer and creates related classes.
+	 */
+	class FLoudnessNRTFactory : public IAnalyzerNRTFactory
+	{
+		public:
+
+		/** Name of specific analyzer type. */
+		FName GetName() const override;
+
+		/** Human readable name of analyzer. */
+		FString GetTitle() const override;
+
+		/** Creates a new FLoudnessNRTResult */
+		TUniquePtr<IAnalyzerNRTResult> NewResult() override;
+
+		/** 
+		 * Creates a new FLoudnessNRTWorker. This expects IAnalyzerNRTSettings to be a valid pointer to
+		 * a FLoudnessNRTSettings object.
+		 */
+		TUniquePtr<IAnalyzerNRTWorker> NewWorker(const FAnalyzerNRTParameters& InParams, const IAnalyzerNRTSettings* InSettings) override;
+	};
+}
+
diff --git a/Engine/Plugins/Runtime/WindowsMixedReality/Source/WindowsMixedRealityHMD/Private/WindowsMixedRealityFunctionLibrary.cpp b/Engine/Plugins/Runtime/WindowsMixedReality/Source/WindowsMixedRealityHMD/Private/WindowsMixedRealityFunctionLibrary.cpp
--- a/Engine/Plugins/Runtime/WindowsMixedReality/Source/WindowsMixedRealityHMD/Private/WindowsMixedRealityFunctionLibrary.cpp
+++ b/Engine/Plugins/Runtime/WindowsMixedReality/Source/WindowsMixedRealityHMD/Private/WindowsMixedRealityFunctionLibrary.cpp
@@ -186,4 +186,5 @@ static FAutoConsoleCommand CSetFocusPointForFrameCmd(
 	TEXT("vr.SetFocusPointForFrame"),
 	*LOCTEXT("CVarText_SetFocusPointForFrame",
 		"Set the reference point for the stabilization plane on hololens 2. You must set it each frame to activate the feature for that frame.").ToString(),
-	FConsoleCommandWithWorldArgsAndOutputDeviceDelegate::CreateStatic(SetFocusPointForFrame));
\ No newline at end of file
+	FConsoleCommandWithWorldArgsAndOutputDeviceDelegate::CreateStatic(SetFocusPointForFrame));
+#undef LOCTEXT_NAMESPACE
\ No newline at end of file
diff --git a/Engine/Source/Editor/UnrealEd/Private/Commandlets/AssetRegistryGenerator.cpp b/Engine/Source/Editor/UnrealEd/Private/Commandlets/AssetRegistryGenerator.cpp
--- a/Engine/Source/Editor/UnrealEd/Private/Commandlets/AssetRegistryGenerator.cpp
+++ b/Engine/Source/Editor/UnrealEd/Private/Commandlets/AssetRegistryGenerator.cpp
@@ -339,7 +339,7 @@ bool FAssetRegistryGenerator::GenerateStreamingInstallManifest(int64 InExtraFlav
 				int32 PakchunkIndex = GetPakchunkIndex(ChunkID);
 				if (PakchunkIndex >= FinalChunkManifests.Num())
 				{
-					FinalChunkManifests.AddZeroed(PakchunkIndex - FinalChunkManifests.Num() + 1);
+					FinalChunkManifests.AddDefaulted(PakchunkIndex - FinalChunkManifests.Num() + 1);
 				}
 				checkf(PakchunkIndex < FinalChunkManifests.Num(), TEXT("Chunk %i out of range. %i manifests available"), PakchunkIndex, FinalChunkManifests.Num() - 1);
 				checkf(FinalChunkManifests[PakchunkIndex] == nullptr, TEXT("Manifest already exists for chunk %i"), PakchunkIndex);
@@ -368,7 +368,7 @@ bool FAssetRegistryGenerator::GenerateStreamingInstallManifest(int64 InExtraFlav
 		ChunkLayerFile->Serialize(TCHAR_TO_ANSI(*LayerString), LayerString.Len());
 
 		// Is this chunk empty?
-		if (!FinalChunkManifests[PakchunkIndex] || FinalChunkManifests[PakchunkIndex]->Num() == 0)
+		if (!FinalChunkManifests[PakchunkIndex])
 		{
 			continue;
 		}
@@ -626,7 +626,8 @@ bool FAssetRegistryGenerator::SaveManifests(FSandboxPlatformFile* InSandboxFile,
 		// Collect all unique chunk indices and map all files to their chunks
 		for (int32 PakchunkIndex = 0; PakchunkIndex < FinalChunkManifests.Num(); ++PakchunkIndex)
 		{
-			if (FinalChunkManifests[PakchunkIndex] && FinalChunkManifests[PakchunkIndex]->Num())
+			check(FinalChunkManifests[PakchunkIndex]);
+			if (FinalChunkManifests[PakchunkIndex]->Num())
 			{
 				PakchunkIndicesInUse.Add(PakchunkIndex);
 				for (auto& Filename : *FinalChunkManifests[PakchunkIndex])
@@ -1440,7 +1441,7 @@ bool FAssetRegistryGenerator::GenerateAssetChunkInformationCSV(const FString& Ou
 	TArray<TUniquePtr<FArchive>> ChunkFiles;
 	if (bWriteIndividualFiles)
 	{
-		for (int32 PakchunkIndex = 0, ChunkNum = FinalChunkManifests.Num(); PakchunkIndex < ChunkNum; ++PakchunkIndex)
+		for (int32 PakchunkIndex = 0; PakchunkIndex < FinalChunkManifests.Num(); ++PakchunkIndex)
 		{
 			FArchive* ChunkFile = IFileManager::Get().CreateFileWriter(*FPaths::Combine(*OutputPath, *FString::Printf(TEXT("Chunks%dInfo.csv"), PakchunkIndex)));
 			if (ChunkFile == nullptr)
@@ -1689,12 +1690,12 @@ void FAssetRegistryGenerator::FixupPackageDependenciesForChunks(FSandboxPlatform
 
 	for (int32 PakchunkIndex = 0, MaxPakchunk = ChunkManifests.Num(); PakchunkIndex < MaxPakchunk; ++PakchunkIndex)
 	{
-		FinalChunkManifests.Add(nullptr);
+		FinalChunkManifests.Add(new FChunkPackageSet());
 		if (!ChunkManifests[PakchunkIndex])
 		{
 			continue;
 		}
-		FinalChunkManifests[PakchunkIndex] = new FChunkPackageSet();
+		
 		for (auto It = ChunkManifests[PakchunkIndex]->CreateConstIterator(); It; ++It)
 		{
 			AddPackageAndDependenciesToChunk(FinalChunkManifests[PakchunkIndex], It.Key(), It.Value(), PakchunkIndex, InSandboxFile);
@@ -1711,12 +1712,10 @@ void FAssetRegistryGenerator::FixupPackageDependenciesForChunks(FSandboxPlatform
 	//Once complete, Add any remaining assets (that are not assigned to a chunk) to the first chunk.
 	if (FinalChunkManifests.Num() == 0)
 	{
-		FinalChunkManifests.Add(nullptr);
-	}
-	if (!FinalChunkManifests[0])
-	{
-		FinalChunkManifests[0] = new FChunkPackageSet();
+		FinalChunkManifests.Add(new FChunkPackageSet());
 	}
+	check(FinalChunkManifests[0]);
+	
 	// Copy the remaining assets
 	auto RemainingAssets = UnassignedPackageSet;
 	for (auto It = RemainingAssets.CreateConstIterator(); It; ++It)
@@ -1750,17 +1749,18 @@ void FAssetRegistryGenerator::FixupPackageDependenciesForChunks(FSandboxPlatform
 	for (int32 PakchunkIndex = 0, MaxPakchunk = ChunkManifests.Num(); PakchunkIndex < MaxPakchunk; ++PakchunkIndex)
 	{
 		const int32 ChunkManifestNum = ChunkManifests[PakchunkIndex] ? ChunkManifests[PakchunkIndex]->Num() : 0;
-		const int32 FinalChunkManifestNum = FinalChunkManifests[PakchunkIndex] ? FinalChunkManifests[PakchunkIndex]->Num() : 0;
+		const int32 FinalChunkManifestNum = FinalChunkManifests[PakchunkIndex]->Num();
 		UE_LOG(LogAssetRegistryGenerator, Log, TEXT("Chunk: %i, Started with %i packages, Final after dependency resolve: %i"), PakchunkIndex, ChunkManifestNum, FinalChunkManifestNum);
 	}
 	
 	// Fix up the asset registry to reflect this chunk layout
 	for (int32 PakchunkIndex = 0 ; PakchunkIndex < FinalChunkManifests.Num(); ++PakchunkIndex)
 	{
-		if (PakchunkIndex >= FinalChunkManifests.Num() || !FinalChunkManifests[PakchunkIndex])
+		if (PakchunkIndex >= FinalChunkManifests.Num())
 		{
 			continue;
 		}
+		check(FinalChunkManifests[PakchunkIndex]);
 		for (const TPair<FName, FString>& Asset : *FinalChunkManifests[PakchunkIndex])
 		{
 			const TArray<const FAssetData*> AssetIndexArray = State.GetAssetsByPackageName(Asset.Key);
diff --git a/Engine/Source/Programs/AutomationTool/AutomationUtils/CommandUtils.cs b/Engine/Source/Programs/AutomationTool/AutomationUtils/CommandUtils.cs
--- a/Engine/Source/Programs/AutomationTool/AutomationUtils/CommandUtils.cs
+++ b/Engine/Source/Programs/AutomationTool/AutomationUtils/CommandUtils.cs
@@ -2365,6 +2365,7 @@ namespace AutomationTool
 		Minutes,
 		Bytes,
 		Megabytes,
+		Percentage,
 	}
 
 	/// <summary>
diff --git a/Engine/Source/Runtime/AudioAnalyzer/Private/AudioAnalyzerNRT.cpp b/Engine/Source/Runtime/AudioAnalyzer/Private/AudioAnalyzerNRT.cpp
--- a/Engine/Source/Runtime/AudioAnalyzer/Private/AudioAnalyzerNRT.cpp
+++ b/Engine/Source/Runtime/AudioAnalyzer/Private/AudioAnalyzerNRT.cpp
@@ -219,10 +219,9 @@ void UAudioAnalyzerNRT::Serialize(FArchive& Ar)
 			if (nullptr != Factory)
 			{
 				// Create result and worker from factory
-				TUniquePtr<Audio::IAnalyzerNRTResult> NewResult = Factory->NewResult();
 				{
 					FScopeLock ResultLock(&ResultCriticalSection);
-					Result = TSharedPtr<Audio::IAnalyzerNRTResult, ESPMode::ThreadSafe>(NewResult.Release());
+					Result = Factory->NewResultShared<ESPMode::ThreadSafe>();
 				}
 			}
 		}
diff --git a/Engine/Source/Runtime/AudioAnalyzer/Private/AudioAnalyzerNRTFacade.cpp b/Engine/Source/Runtime/AudioAnalyzer/Private/AudioAnalyzerNRTFacade.cpp
--- a/Engine/Source/Runtime/AudioAnalyzer/Private/AudioAnalyzerNRTFacade.cpp
+++ b/Engine/Source/Runtime/AudioAnalyzer/Private/AudioAnalyzerNRTFacade.cpp
@@ -5,6 +5,7 @@
 #include "AudioAnalyzerModule.h"
 #include "IAudioAnalyzerNRTInterface.h"
 #include "Features/IModularFeatures.h"
+#include "SampleBuffer.h"
 
 namespace Audio
 {
diff --git a/Engine/Source/Runtime/AudioAnalyzer/Public/IAudioAnalyzerNRTInterface.h b/Engine/Source/Runtime/AudioAnalyzer/Public/IAudioAnalyzerNRTInterface.h
--- a/Engine/Source/Runtime/AudioAnalyzer/Public/IAudioAnalyzerNRTInterface.h
+++ b/Engine/Source/Runtime/AudioAnalyzer/Public/IAudioAnalyzerNRTInterface.h
@@ -5,7 +5,6 @@
 #pragma once
 
 #include "CoreMinimal.h"
-#include "SampleBuffer.h"
 #include "Features/IModularFeature.h"
 
 namespace Audio
@@ -104,6 +103,15 @@ namespace Audio
 		// Create a new result.
 		virtual TUniquePtr<IAnalyzerNRTResult> NewResult() = 0;
 
+		// Convenience function to create a new shared result by calling NewResult.
+		template<ESPMode Mode = ESPMode::Fast>
+		TSharedPtr<IAnalyzerNRTResult, Mode> NewResultShared()
+		{
+			TUniquePtr<Audio::IAnalyzerNRTResult> Result = NewResult();
+
+			return TSharedPtr<Audio::IAnalyzerNRTResult, Mode>(Result.Release());
+		}
+
 		// Create a new worker.
 		virtual TUniquePtr<IAnalyzerNRTWorker> NewWorker(const FAnalyzerNRTParameters& InParams, const IAnalyzerNRTSettings* InSettings) = 0;
 	};
diff --git a/Engine/Source/Runtime/Core/Private/UObject/UnrealNames.cpp b/Engine/Source/Runtime/Core/Private/UObject/UnrealNames.cpp
--- a/Engine/Source/Runtime/Core/Private/UObject/UnrealNames.cpp
+++ b/Engine/Source/Runtime/Core/Private/UObject/UnrealNames.cpp
@@ -806,7 +806,7 @@ static bool IsPureAnsi(const WIDECHAR* Str, const int32 Len)
 
 FNameEntryId FNamePool::Find(EName Ename) const
 {
-	check(Ename < NAME_MaxHardcodedNameIndex);
+	checkSlow(Ename < NAME_MaxHardcodedNameIndex);
 	return ENameToEntry[Ename];
 }
 
diff --git a/Engine/Source/Runtime/Experimental/Chaos/Public/Chaos/KinematicGeometryParticles.h b/Engine/Source/Runtime/Experimental/Chaos/Public/Chaos/KinematicGeometryParticles.h
--- a/Engine/Source/Runtime/Experimental/Chaos/Public/Chaos/KinematicGeometryParticles.h
+++ b/Engine/Source/Runtime/Experimental/Chaos/Public/Chaos/KinematicGeometryParticles.h
@@ -63,8 +63,13 @@ FChaosArchive& operator<<(FChaosArchive& Ar, TKinematicGeometryParticlesImp<T, d
 	return Ar;
 }
 
+#ifdef __clang__
+extern template class CHAOS_API Chaos::TKinematicGeometryParticlesImp<float, 3, Chaos::EGeometryParticlesSimType::RigidBodySim>;
+extern template class CHAOS_API Chaos::TKinematicGeometryParticlesImp<float, 3, Chaos::EGeometryParticlesSimType::Other>;
+#else
 extern template class Chaos::TKinematicGeometryParticlesImp<float, 3, Chaos::EGeometryParticlesSimType::RigidBodySim>;
 extern template class Chaos::TKinematicGeometryParticlesImp<float, 3, Chaos::EGeometryParticlesSimType::Other>;
+#endif
 
 template <typename T, int d>
 using TKinematicGeometryParticles = TKinematicGeometryParticlesImp<T, d, EGeometryParticlesSimType::RigidBodySim>;
diff --git a/Engine/Source/Runtime/Landscape/Private/Landscape.cpp b/Engine/Source/Runtime/Landscape/Private/Landscape.cpp
--- a/Engine/Source/Runtime/Landscape/Private/Landscape.cpp
+++ b/Engine/Source/Runtime/Landscape/Private/Landscape.cpp
@@ -1627,6 +1627,8 @@ void ULandscapeComponent::AddLayerData(const FGuid& InLayerGuid, const FLandscap
 	check(!LandscapeEditingLayer.IsValid());
 	FLandscapeLayerComponentData& Data = LayersData.FindOrAdd(InLayerGuid);
 	Data = InData;
+	CachedEditingLayer.Invalidate();
+	CachedEditingLayerData = nullptr;
 }
 
 void ULandscapeComponent::AddDefaultLayerData(const FGuid& InLayerGuid, const TArray<ULandscapeComponent*>& InComponentsUsingHeightmap, TMap<UTexture2D*, UTexture2D*>& InOutCreatedHeightmapTextures)
@@ -1712,6 +1714,8 @@ void ULandscapeComponent::RemoveLayerData(const FGuid& InLayerGuid)
 	Modify();
 	check(!LandscapeEditingLayer.IsValid());
 	LayersData.Remove(InLayerGuid);
+	CachedEditingLayer.Invalidate();
+	CachedEditingLayerData = nullptr;
 }
 
 void ULandscapeComponent::SetHeightmap(UTexture2D* NewHeightmap)
diff --git a/Engine/Source/Runtime/Landscape/Private/LandscapeEditLayers.cpp b/Engine/Source/Runtime/Landscape/Private/LandscapeEditLayers.cpp
--- a/Engine/Source/Runtime/Landscape/Private/LandscapeEditLayers.cpp
+++ b/Engine/Source/Runtime/Landscape/Private/LandscapeEditLayers.cpp
@@ -1403,9 +1403,6 @@ void ALandscape::CreateLayersRenderingResource()
 		}
 	});
 
-	// Flush because TickLayers can access CPU Readbacks in the same frame (see ResolveLayersTexture)
-	FlushRenderingCommands();
-
 	const FIntPoint ComponentCounts = ComputeComponentCounts();
 
 	ALandscape* Landscape = GetLandscapeActor();
@@ -2691,9 +2688,9 @@ void ALandscape::PrintLayersDebugTextureResource(const FString& InContext, FText
 	Flags.SetMip(InMipRender);
 
 	ENQUEUE_RENDER_COMMAND(FLandscapeLayersDebugReadSurfaceCommand)(
-		[SourceTextureRHI = InTextureResource->TextureRHI, SampleRect = SampleRect, OutData = &OutputTexels, ReadFlags = Flags](FRHICommandListImmediate& RHICmdList) mutable
+		[InTextureResource, SampleRect = SampleRect, OutData = &OutputTexels, ReadFlags = Flags](FRHICommandListImmediate& RHICmdList) mutable
 	{
-		RHICmdList.ReadSurfaceData(SourceTextureRHI, SampleRect, *OutData, ReadFlags);
+		RHICmdList.ReadSurfaceData(InTextureResource->TextureRHI, SampleRect, *OutData, ReadFlags);
 	});
 
 	FlushRenderingCommands();
@@ -3186,11 +3183,12 @@ bool ALandscape::ResolveLayersTexture(FLandscapeLayersTexture2DCPUReadBackResour
 	SCOPE_CYCLE_COUNTER(STAT_LandscapeLayersResolveTexture);
 
 	TArray<TArray<FColor>> OutMipsData;
-	OutMipsData.AddDefaulted(InCPUReadBackTexture->TextureRHI->GetNumMips());
-
+	
 	ENQUEUE_RENDER_COMMAND(FLandscapeLayersReadSurfaceCommand)(
 		[InCPUReadBackTexture, &OutMipsData](FRHICommandListImmediate& RHICmdList) mutable
 	{
+		OutMipsData.AddDefaulted(InCPUReadBackTexture->TextureRHI->GetNumMips());
+
 		int32 MipSizeU = InCPUReadBackTexture->GetSizeX();
 		int32 MipSizeV = InCPUReadBackTexture->GetSizeY();
 		int32 MipIndex = 0;
@@ -3996,7 +3994,6 @@ int32 ALandscape::RegenerateLayersWeightmaps(const TArray<ULandscapeComponent*>&
 		if (bHasWeightmapData)
 		{
 			// Will generate CPU read back resource, if required
-			bool bHasPendingInitResource = false;
 			for (ULandscapeComponent* Component : InLandscapeComponentsToRender)
 			{
 				const TArray<UTexture2D*>& ComponentWeightmapTextures = Component->GetWeightmapTextures();
@@ -4014,17 +4011,10 @@ int32 ALandscape::RegenerateLayersWeightmaps(const TArray<ULandscapeComponent*>&
 						WeightmapTexture->Source.UnlockMip(0);
 						BeginInitResource(NewWeightmapCPUReadBack);
 						Proxy->WeightmapsCPUReadBack.Add(WeightmapTexture, NewWeightmapCPUReadBack);
-						bHasPendingInitResource = true;
 					}
 				}
 			}
-
-			if (bHasPendingInitResource)
-			{
-				// Flush because TickLayers can access CPU Readbacks in the same frame (see ResolveLayersTexture)
-				FlushRenderingCommands();
-			}
-
+			
 			int8 CurrentWeightmapToProcessIndex = 0;
 			bool HasFoundWeightmapToProcess = true; // try processing at least once
 
diff --git a/Engine/Source/Runtime/Renderer/Private/PostProcess/DiaphragmDOF.cpp b/Engine/Source/Runtime/Renderer/Private/PostProcess/DiaphragmDOF.cpp
--- a/Engine/Source/Runtime/Renderer/Private/PostProcess/DiaphragmDOF.cpp
+++ b/Engine/Source/Runtime/Renderer/Private/PostProcess/DiaphragmDOF.cpp
@@ -1498,7 +1498,7 @@ FRDGTextureRef DiaphragmDOF::AddPasses(
 		else if (bOutputHalfResolution)
 		{
 			PermutationVector.Set<FDiaphragmDOFSetupCS::FOutputResDivisor>(2);
-			PassViewSize = PreprocessViewSize;
+			PassViewSize = GatheringViewSize;
 			CocRadiusBasis = PreprocessViewSize.X;
 		}
 		else
diff --git a/Engine/Source/Runtime/SignalProcessing/Private/FloatArrayMath.cpp b/Engine/Source/Runtime/SignalProcessing/Private/FloatArrayMath.cpp
--- a/Engine/Source/Runtime/SignalProcessing/Private/FloatArrayMath.cpp
+++ b/Engine/Source/Runtime/SignalProcessing/Private/FloatArrayMath.cpp
@@ -10,10 +10,10 @@ namespace
 
 namespace Audio
 {
-	void ArrayCumulativeSum(const TArray<float>& InData, TArray<float>& OutData)
+	void ArrayCumulativeSum(TArrayView<const float> InView, TArray<float>& OutData)
 	{
 		// Initialize output data
-		int32 Num = InData.Num();
+		int32 Num = InView.Num();
 		OutData.Reset();
 		OutData.AddUninitialized(Num);
 
@@ -23,19 +23,19 @@ namespace Audio
 		}
 
 		float* OutDataPtr = OutData.GetData();
-		const float* InDataPtr = InData.GetData();
+		const float* InViewPtr = InView.GetData();
 
 		// Start summing
-		*OutDataPtr = *InDataPtr++;
+		*OutDataPtr = *InViewPtr++;
 
 		for (int32 i = 1; i < Num; i++)
 		{
-			float Temp = *OutDataPtr++ + *InDataPtr++;
+			float Temp = *OutDataPtr++ + *InViewPtr++;
 			*OutDataPtr = Temp;
 		}
 	}
 
-	void ArrayMeanFilter(const TArray<float>& InData, int32 WindowSize, int32 WindowOrigin, TArray<float>& OutData)
+	void ArrayMeanFilter(TArrayView<const float> InView, int32 WindowSize, int32 WindowOrigin, TArray<float>& OutData)
 	{
 		// a quick but sinful implementation of a mean filter. encourages floating point rounding errors. 
 		check(WindowOrigin < WindowSize);
@@ -43,7 +43,7 @@ namespace Audio
 		check(WindowSize > 0);
 
 		// Initialize output data
-		const int32 Num = InData.Num();
+		const int32 Num = InView.Num();
 		OutData.Reset();
 		OutData.AddUninitialized(Num);
 
@@ -53,11 +53,11 @@ namespace Audio
 		}
 		
 		// Use cumulative sum to avoid multiple summations 
-		// Instead of summing over InData[StartIndex:EndIndex], avoid all that
+		// Instead of summing over InView[StartIndex:EndIndex], avoid all that
 		// calculation by taking difference of cumulative sum at those two points:
 		//  cumsum(X[0:b]) - cumsum(X[0:a]) = sum(X[a:b])
 		TArray<float> SummedData;
-		ArrayCumulativeSum(InData, SummedData);
+		ArrayCumulativeSum(InView, SummedData);
 		const float LastSummedData = SummedData.Last();
 		
 		
@@ -101,7 +101,7 @@ namespace Audio
 		}
 	}
 
-	void ArrayMaxFilter(const TArray<float>& InData, int32 WindowSize, int32 WindowOrigin, TArray<float>& OutData)
+	void ArrayMaxFilter(TArrayView<const float> InView, int32 WindowSize, int32 WindowOrigin, TArray<float>& OutData)
 	{
 		// A reasonable implementation of a max filter for the data we're interested in, though surely not the fastest.
 		check(WindowOrigin < WindowSize);
@@ -112,7 +112,7 @@ namespace Audio
 		int32 EndIndex = StartIndex + WindowSize;
 
 		// Initialize output
-		int32 Num = InData.Num();
+		int32 Num = InView.Num();
 		OutData.Reset();
 		OutData.AddUninitialized(Num);
 
@@ -125,16 +125,16 @@ namespace Audio
 		int32 ActualStartIndex = 0;
 		int32 ActualEndIndex = FMath::Min(EndIndex, Num);
 
-		const float* InDataPtr = InData.GetData();
+		const float* InViewPtr = InView.GetData();
 		float* OutDataPtr = OutData.GetData();
 		int32 MaxIndex = 0;
-		float MaxValue = InData[0];
+		float MaxValue = InView[0];
 
 		for (int32 i = ActualStartIndex; i < ActualEndIndex; i++)
 		{
-			if (InDataPtr[i] > MaxValue)
+			if (InViewPtr[i] > MaxValue)
 			{
-				MaxValue = InDataPtr[i];
+				MaxValue = InViewPtr[i];
 				MaxIndex = i;
 			}		
 		}
@@ -153,23 +153,23 @@ namespace Audio
 			{
 				// We need to evaluate the entire window because the previous maximum value was not in this window.
 				MaxIndex = ActualStartIndex;
-				MaxValue = InDataPtr[MaxIndex];
+				MaxValue = InViewPtr[MaxIndex];
 				for (int32 j = ActualStartIndex + 1; j < ActualEndIndex; j++)
 				{
-					if (InDataPtr[j] > MaxValue)
+					if (InViewPtr[j] > MaxValue)
 					{
 						MaxIndex = j;
-						MaxValue = InDataPtr[MaxIndex];
+						MaxValue = InViewPtr[MaxIndex];
 					}
 				}
 			}
 			else
 			{
 				// We only need to inspect the newest sample because the previous maximum value was in this window.
-				if (InDataPtr[ActualEndIndex - 1] > MaxValue)
+				if (InViewPtr[ActualEndIndex - 1] > MaxValue)
 				{
 					MaxIndex = ActualEndIndex - 1;
-					MaxValue = InDataPtr[MaxIndex];
+					MaxValue = InViewPtr[MaxIndex];
 				}
 			}
 
@@ -180,59 +180,70 @@ namespace Audio
 		}
 	}
 
-	void ArrayGetEuclideanNorm(const TArray<float>& InArray, float& OutEuclideanNorm)
+	void ArrayGetEuclideanNorm(TArrayView<const float> InView, float& OutEuclideanNorm)
 	{
 		// Initialize output.
 		OutEuclideanNorm = 0.0f;
-		const int32 Num = InArray.Num();
-		const float* InArrayData = InArray.GetData();
+		const int32 Num = InView.Num();
+		const float* InViewData = InView.GetData();
 		
 		// Sum it up.
 		for (int32 i = 0; i < Num; i++)
 		{
-			OutEuclideanNorm += InArrayData[i] * InArrayData[i];
+			OutEuclideanNorm += InViewData[i] * InViewData[i];
 		}
 
 		OutEuclideanNorm = FMath::Sqrt(OutEuclideanNorm);
 	}
 
-	void ArrayMultiplyByConstantInPlace(TArray<float>& InArray, float InMultiplier)
+	void ArrayClampInPlace(TArrayView<float> InView, float InMin, float InMax)
 	{
-		const int32 Num = InArray.Num();
-		float* InArrayData = InArray.GetData();
+		const int32 Num = InView.Num();
+		float* Data = InView.GetData();
+
+		for (int32 i = 0; i < Num; i++)
+		{
+			Data[i] = FMath::Clamp(Data[i], InMin, InMax);
+		}
+	}
+
+	void ArrayMultiplyByConstantInPlace(TArrayView<float> InView, float InMultiplier)
+	{
+		const int32 Num = InView.Num();
+		float* InViewData = InView.GetData();
 		for (int32 i = 0; i < Num; i++)
 		{
-			InArrayData[i] *= InMultiplier;
+			InViewData[i] *= InMultiplier;
 		}
 	}
 
-	void ArraySubtractByConstantInPlace(TArray<float>& InArray, float InSubtrahend)
+	void ArraySubtractByConstantInPlace(TArrayView<float> InView, float InSubtrahend)
 	{
-		const int32 Num = InArray.Num();
-		float* InArrayData = InArray.GetData();
+		const int32 Num = InView.Num();
+		float* InViewData = InView.GetData();
 		for (int32 i = 0; i < Num; i++)
 		{
-			InArrayData[i] -= InSubtrahend;
+			InViewData[i] -= InSubtrahend;
 		}
 	}
 
-	void ArrayMagnitudeToDecibelInPlace(TArray<float>& InArray)
+	void ArrayMagnitudeToDecibelInPlace(TArrayView<float> InView)
 	{
-		const int32 Num = InArray.Num();
-		float* InArrayData = InArray.GetData();
+		const int32 Num = InView.Num();
+		float* InViewData = InView.GetData();
 		for (int32 i = 0; i < Num; i++)
 		{
-			InArrayData[i] = 20.f * FMath::Loge(InArrayData[i]) / LOGE10;
+			InViewData[i] = 20.f * FMath::Loge(InViewData[i]) / LOGE10;
 		}
 	}
 
-	void ArrayPowerToDecibelInPlace(TArray<float>& InArray)
+	void ArrayPowerToDecibelInPlace(TArrayView<float> InView)
 	{
-		const int32 Num = InArray.Num();
-		float* InArrayData = InArray.GetData();
+		const int32 Num = InView.Num();
+		float* InViewData = InView.GetData();
 		for (int32 i = 0; i < Num; i++)
 		{
-			InArrayData[i] = 10.f * FMath::Loge(InArrayData[i]) / LOGE10;
+			InViewData[i] = 10.f * FMath::Loge(InViewData[i]) / LOGE10;
 		}
 	}
 
@@ -269,9 +280,9 @@ namespace Audio
 		Kernel[RowIndex].OffsetValues = TArray<float>(OffsetValues.GetData(), OffsetValues.Num());
 	}
 
-	void FContiguousSparse2DKernelTransform::TransformArray(TArrayView<const float> InArray, TArray<float>& OutArray) const
+	void FContiguousSparse2DKernelTransform::TransformArray(TArrayView<const float> InView, TArray<float>& OutArray) const
 	{
-		check(InArray.Num() == NumIn);
+		check(InView.Num() == NumIn);
 
 		// Resize output
 		OutArray.Reset(NumOut);
@@ -280,12 +291,12 @@ namespace Audio
 			OutArray.AddUninitialized(NumOut);
 		}
 
-		TransformArray(InArray.GetData(), OutArray.GetData());
+		TransformArray(InView.GetData(), OutArray.GetData());
 	}
 
-	void FContiguousSparse2DKernelTransform::TransformArray(TArrayView<const float> InArray, AlignedFloatBuffer& OutArray) const
+	void FContiguousSparse2DKernelTransform::TransformArray(TArrayView<const float> InView, AlignedFloatBuffer& OutArray) const
 	{	
-		check(InArray.Num() == NumIn);
+		check(InView.Num() == NumIn);
 
 		// Resize output
 		OutArray.Reset(NumOut);
@@ -294,7 +305,7 @@ namespace Audio
 			OutArray.AddUninitialized(NumOut);
 		}
 
-		TransformArray(InArray.GetData(), OutArray.GetData());
+		TransformArray(InView.GetData(), OutArray.GetData());
 	}
 
 	void FContiguousSparse2DKernelTransform::TransformArray(const float* InArray, float* OutArray) const
diff --git a/Engine/Source/Runtime/SignalProcessing/Public/DSP/DeinterleaveView.h b/Engine/Source/Runtime/SignalProcessing/Public/DSP/DeinterleaveView.h
--- a/Engine/Source/Runtime/SignalProcessing/Public/DSP/DeinterleaveView.h
+++ b/Engine/Source/Runtime/SignalProcessing/Public/DSP/DeinterleaveView.h
@@ -185,6 +185,8 @@ namespace Audio
 	{
 		TArray<T, InAllocator>& ArrayToFill;
 
+		typedef typename TDeinterleaveView<T>::template TChannelIterator<InAllocator> TAutoChannelIterator;
+
 	public:
 		/** TAutoDeinterleaveView Constructor.
 		 *
@@ -193,20 +195,20 @@ namespace Audio
 		 *  InNumChannels is the number of channels in the deinterleaved array.
 		 */
 		TAutoDeinterleaveView(TArrayView<const T> InInterleavedArray, TArray<T, InAllocator>& InArrayToFill, int32 InNumChannels)
-		:	TDeinterleaveView(InInterleavedArray, InNumChannels)
+		:	TDeinterleaveView<T>(InInterleavedArray, InNumChannels)
 		,	ArrayToFill(InArrayToFill)
 		{}
 
 		/** Return an STL iterator to the first channel. */
-		TChannelIterator<InAllocator> begin() 
+		TAutoChannelIterator begin() 
 		{
-			return TDeinterleaveView<T>::begin<InAllocator>(ArrayToFill);
+			return TDeinterleaveView<T>::template begin<InAllocator>(ArrayToFill);
 		}
 
 		/** Return an STL iterator to the end. */
-		TChannelIterator<InAllocator> end()
+		TAutoChannelIterator end()
 		{
-			return TDeinterleaveView<T>::end<InAllocator>(ArrayToFill);
+			return TDeinterleaveView<T>::template end<InAllocator>(ArrayToFill);
 		}
 	};
 }
diff --git a/Engine/Source/Runtime/SignalProcessing/Public/DSP/FloatArrayMath.h b/Engine/Source/Runtime/SignalProcessing/Public/DSP/FloatArrayMath.h
--- a/Engine/Source/Runtime/SignalProcessing/Public/DSP/FloatArrayMath.h
+++ b/Engine/Source/Runtime/SignalProcessing/Public/DSP/FloatArrayMath.h
@@ -9,53 +9,61 @@ namespace Audio
 {
 	/** Cumulative sum of array.
 	 *
-	 *  InData contains data to be cumulatively summed.
-	 *  OutData contains sum and is same size as InData.
+	 *  InView contains data to be cumulatively summed.
+	 *  OutData contains sum and is same size as InView.
 	 */
-	SIGNALPROCESSING_API void ArrayCumulativeSum(const TArray<float>& InData, TArray<float>& OutData);
+	SIGNALPROCESSING_API void ArrayCumulativeSum(TArrayView<const float> InView, TArray<float>& OutData);
 
 	/** Mean filter of array.
 	 *
 	 *  Note: Uses standard biased mean estimator of Sum(x) / Count(x).
 	 *  Note: At array boundaries, this algorithm truncates windows where no valid array data exists. Values calculated with truncated windows have corresponding increased variances. 
 	 *
-	 *  InData contains data to be filtered.
-	 *  WindowSize determines the number of samples from InData analyzed to produce a value in OutData.
-	 *  WindowOrigin describes the offset from the windows first sample to the index of OutData. For example, if WindowOrigin = WindowSize/4, then OutData[i] = Mean(InData[i - Window/4 : i + 3 * Window / 4]).
+	 *  InView contains data to be filtered.
+	 *  WindowSize determines the number of samples from InView analyzed to produce a value in OutData.
+	 *  WindowOrigin describes the offset from the windows first sample to the index of OutData. For example, if WindowOrigin = WindowSize/4, then OutData[i] = Mean(InView[i - Window/4 : i + 3 * Window / 4]).
 	 *  OutData contains the produceds data.
 	 */
-	SIGNALPROCESSING_API void ArrayMeanFilter(const TArray<float>& InData, int32 WindowSize, int32 WindowOrigin, TArray<float>& OutData);
+	SIGNALPROCESSING_API void ArrayMeanFilter(TArrayView<const float> InView, int32 WindowSize, int32 WindowOrigin, TArray<float>& OutData);
 
 	/** Max filter of array.
 	 *
 	 *  Note: At array boundaries, this algorithm truncates windows where no valid array data exists. 
 	 *
-	 *  InData contains data to be filtered.
-	 *  WindowSize determines the number of samples from InData analyzed to produce a value in OutData.
-	 *  WindowOrigin describes the offset from the windows first sample to the index of OutData. For example, if WindowOrigin = WindowSize/4, then OutData[i] = Max(InData[i - Window/4 : i + 3 * Window / 4]).
+	 *  InView contains data to be filtered.
+	 *  WindowSize determines the number of samples from InView analyzed to produce a value in OutData.
+	 *  WindowOrigin describes the offset from the windows first sample to the index of OutData. For example, if WindowOrigin = WindowSize/4, then OutData[i] = Max(InView[i - Window/4 : i + 3 * Window / 4]).
 	 *  OutData contains the produceds data.
 	 */
-	SIGNALPROCESSING_API void ArrayMaxFilter(const TArray<float>& InData, int32 WindowSize, int32 WindowOrigin, TArray<float>& OutData);
+	SIGNALPROCESSING_API void ArrayMaxFilter(TArrayView<const float> InView, int32 WindowSize, int32 WindowOrigin, TArray<float>& OutData);
 
-	/** Computes the EuclideanNorm of the InArray. Same as calculating the energy in window. */
-	SIGNALPROCESSING_API void ArrayGetEuclideanNorm(const TArray<float>& InArray, float& OutEuclideanNorm);
+	/** Computes the EuclideanNorm of the InView. Same as calculating the energy in window. */
+	SIGNALPROCESSING_API void ArrayGetEuclideanNorm(TArrayView<const float> InView, float& OutEuclideanNorm);
 
-	/** Multiplies each element in InArray by InMultiplier */
-	SIGNALPROCESSING_API void ArrayMultiplyByConstantInPlace(TArray<float>& InArray, float InMultiplier);
+	/** Clamp values in an array.
+	 *
+	 *  InView is a view of a float array to be clamped.
+	 *  InMin is the minimum value allowable in the array.
+	 *  InMax is the maximum value allowable in the array.
+	 */
+	SIGNALPROCESSING_API void ArrayClampInPlace(TArrayView<float> InView, float InMin, float InMax);
+
+	/** Multiplies each element in InView by InMultiplier */
+	SIGNALPROCESSING_API void ArrayMultiplyByConstantInPlace(TArrayView<float> InView, float InMultiplier);
 
-	/** Subtract value from each element in InArray */
-	SIGNALPROCESSING_API void ArraySubtractByConstantInPlace(TArray<float>& InArray, float InSubtrahend);
+	/** Subtract value from each element in InView */
+	SIGNALPROCESSING_API void ArraySubtractByConstantInPlace(TArrayView<float> InView, float InSubtrahend);
 
 	/** Convert magnitude values to decibel values in place. db = 20 * log10(val) */
-	SIGNALPROCESSING_API void ArrayMagnitudeToDecibelInPlace(TArray<float>& InArray);
+	SIGNALPROCESSING_API void ArrayMagnitudeToDecibelInPlace(TArrayView<float> InView);
 
 	/** Convert power values to decibel values in place. db = 10 * log10(val) */
-	SIGNALPROCESSING_API void ArrayPowerToDecibelInPlace(TArray<float>& InArray);
+	SIGNALPROCESSING_API void ArrayPowerToDecibelInPlace(TArrayView<float> InView);
 
 	/** FContiguousSparse2DKernelTransform
 	 *
 	 *  FContiguousSparse2DKernelTransform applies a matrix transformation to an input array. 
-	 *  [OutArray] = [[Kernal]][InArray]  
+	 *  [OutArray] = [[Kernal]][InView]  
 	 *
 	 *  It provides some optimization by exploit the contiguous and sparse qualities of the kernel rows,
 	 *  which allows it to skip multiplications with the number zero. 
@@ -95,17 +103,17 @@ namespace Audio
 
 		/** Transforms the input array given the kernel.
 		 *
-		 *  InArray is the array to be transformed. It must have `NumInElements` number of elements.
+		 *  InView is the array to be transformed. It must have `NumInElements` number of elements.
 		 *  OutArray is the transformed array. It will have `NumOutElements` number of elements.
 		 */
-		void TransformArray(TArrayView<const float> InArray, TArray<float>& OutArray) const;
+		void TransformArray(TArrayView<const float> InView, TArray<float>& OutArray) const;
 
 		/** Transforms the input array given the kernel.
 		 *
-		 *  InArray is the array to be transformed. It must have `NumInElements` number of elements.
+		 *  InView is the array to be transformed. It must have `NumInElements` number of elements.
 		 *  OutArray is the transformed array. It will have `NumOutElements` number of elements.
 		 */
-		void TransformArray(TArrayView<const float> InArray, AlignedFloatBuffer& OutArray) const;
+		void TransformArray(TArrayView<const float> InView, AlignedFloatBuffer& OutArray) const;
 
 		/** Transforms the input array given the kernel.
 		 *
diff --git a/Engine/Source/Runtime/SignalProcessing/Public/DSP/SlidingWindow.h b/Engine/Source/Runtime/SignalProcessing/Public/DSP/SlidingWindow.h
--- a/Engine/Source/Runtime/SignalProcessing/Public/DSP/SlidingWindow.h
+++ b/Engine/Source/Runtime/SignalProcessing/Public/DSP/SlidingWindow.h
@@ -400,7 +400,7 @@ namespace Audio
 			 * bDoFlush Controls whether zeros to the final output windows until all possible windows with data from InNewBuffer have been covered.
 			 */
 			TScopedSlidingWindow(TSlidingBuffer<InSampleType>& InSlidingBuffer, TArrayView<const InSampleType> InNewBuffer, bool bDoFlush = false)
-			:	TSlidingWindow(InSlidingBuffer, InNewBuffer, bDoFlush)
+			:	TSlidingWindow<InSampleType>(InSlidingBuffer, InNewBuffer, bDoFlush)
 			,	SlidingBuffer(InSlidingBuffer)
 			{}
 
@@ -409,7 +409,7 @@ namespace Audio
 			 */
 			virtual ~TScopedSlidingWindow()
 			{
-				SlidingBuffer.StoreForFutureWindows(NewBuffer);
+				SlidingBuffer.StoreForFutureWindows(TSlidingWindow<InSampleType>::NewBuffer);
 			}
 	};
 
@@ -450,6 +450,8 @@ namespace Audio
 	{
 		TArray<InSampleType, InAllocator>& WindowBuffer;
 
+		typedef typename TSlidingWindow<InSampleType>::template TSlidingWindowIterator<InAllocator> TAutoSlidingWindowIterator;
+
 		public:
 			/**
 			 * TAutoSlidingWindow constructor
@@ -460,7 +462,7 @@ namespace Audio
 			 * bDoFlush Controls whether zeros to the final output windows until all possible windows with data from InNewBuffer have been covered.
 			 */
 			TAutoSlidingWindow(TSlidingBuffer<InSampleType>& InBuffer, TArrayView<const InSampleType> InNewBuffer, TArray<InSampleType, InAllocator>& OutWindow, bool bDoFlush = false)
-			:	TScopedSlidingWindow(InBuffer, InNewBuffer, bDoFlush)
+			:	TScopedSlidingWindow<InSampleType>(InBuffer, InNewBuffer, bDoFlush)
 			,	WindowBuffer(OutWindow)
 			{}
 
@@ -469,9 +471,9 @@ namespace Audio
 			 *
 			 * This iterator maintains a reference to the OutWindow passed into the constructor. That array will be manipulated when the iterator's * operator is called. 
 			 */
-			TSlidingWindow<InSampleType>::TSlidingWindowIterator<InAllocator> begin() 
+			TAutoSlidingWindowIterator begin() 
 			{
-				return TSlidingWindow<InSampleType>::begin<InAllocator>(WindowBuffer);
+				return TSlidingWindow<InSampleType>::template begin<InAllocator>(WindowBuffer);
 			}
 
 			/**
@@ -479,9 +481,9 @@ namespace Audio
 			 *
 			 * This iterator maintains a reference to the OutWindow passed into the constructor. That array will be manipulated when the iterator's * operator is called. 
 			 */
-			TSlidingWindow<InSampleType>::TSlidingWindowIterator<InAllocator> end()
+			TAutoSlidingWindowIterator end()
 			{
-				return TSlidingWindow<InSampleType>::end<InAllocator>(WindowBuffer);
+				return TSlidingWindow<InSampleType>::template end<InAllocator>(WindowBuffer);
 			}
 	};
 }
