[UE-85519] Remove overaggressive checkSlow in SetDataSet.  This was an erroneous check left in that fails due to some recent other changes.

#rnx

diff --git a/Engine/Plugins/FX/Niagara/Source/Niagara/Classes/NiagaraDataSet.h b/Engine/Plugins/FX/Niagara/Source/Niagara/Classes/NiagaraDataSet.h
--- a/Engine/Plugins/FX/Niagara/Source/Niagara/Classes/NiagaraDataSet.h
+++ b/Engine/Plugins/FX/Niagara/Source/Niagara/Classes/NiagaraDataSet.h
@@ -422,7 +422,6 @@ struct FNiagaraDataSetAccessorBase
 		if (Var.IsValid())
 		{
 			DataSet = &InDataSet;
-			checkSlow(VarLayout == InDataSet.GetVariableLayout(Var));
 		}
 		else
 		{
