Fixed merge-compile error.

#rnx

diff --git a/Engine/Plugins/FX/Niagara/Source/NiagaraEditor/Private/Customizations/NiagaraTypeCustomizations.cpp b/Engine/Plugins/FX/Niagara/Source/NiagaraEditor/Private/Customizations/NiagaraTypeCustomizations.cpp
--- a/Engine/Plugins/FX/Niagara/Source/NiagaraEditor/Private/Customizations/NiagaraTypeCustomizations.cpp
+++ b/Engine/Plugins/FX/Niagara/Source/NiagaraEditor/Private/Customizations/NiagaraTypeCustomizations.cpp
@@ -695,7 +695,7 @@ void FNiagaraScriptVariableBindingCustomization::CustomizeHeader(TSharedRef<IPro
 			.MaxDesiredWidth(200.f)
 			[
 				SNew(STextBlock)
-				.Text(FText::FromString(FName::NameToDisplayString(Cast<UStructProperty>(PropertyHandle->GetProperty())->Struct->GetName(), false)))
+				.Text(FText::FromString(FName::NameToDisplayString(CastField<FStructProperty>(PropertyHandle->GetProperty())->Struct->GetName(), false)))
 				.Font(IDetailLayoutBuilder::GetDetailFont())
 			];
 	}
