Fixed deprecation warnings

diff --git a/Engine/Plugins/FX/Niagara/Source/Niagara/Private/NiagaraSystemInstance.cpp b/Engine/Plugins/FX/Niagara/Source/Niagara/Private/NiagaraSystemInstance.cpp
--- a/Engine/Plugins/FX/Niagara/Source/Niagara/Private/NiagaraSystemInstance.cpp
+++ b/Engine/Plugins/FX/Niagara/Source/Niagara/Private/NiagaraSystemInstance.cpp
@@ -649,7 +649,7 @@ bool FNiagaraSystemInstance::ComputeEmitterPriority(int32 EmitterIdx, TArray<int
 
 void FNiagaraSystemInstance::FindDataInterfaceDependencies(const TArray<UNiagaraDataInterface*>& DataInterfaces, TArray<FNiagaraEmitterInstance*>& Dependencies)
 {
-	for (const UNiagaraDataInterface* DI : DataInterfaces)
+	for (UNiagaraDataInterface* DI : DataInterfaces)
 	{
 		void* InterfaceInstanceData = FindDataInterfaceInstanceData(DI);
 		int32 NumDepsBefore = Dependencies.Num();
diff --git a/Engine/Plugins/FX/Niagara/Source/Niagara/Public/NiagaraSystemInstance.h b/Engine/Plugins/FX/Niagara/Source/Niagara/Public/NiagaraSystemInstance.h
--- a/Engine/Plugins/FX/Niagara/Source/Niagara/Public/NiagaraSystemInstance.h
+++ b/Engine/Plugins/FX/Niagara/Source/Niagara/Public/NiagaraSystemInstance.h
@@ -194,7 +194,7 @@ public:
 	FNiagaraSystemInstanceID GetId() { return ID; }
 
 	/** Returns the instance data for a particular interface for this System. */
-	FORCEINLINE void* FindDataInterfaceInstanceData(const UNiagaraDataInterface* Interface) 
+	FORCEINLINE void* FindDataInterfaceInstanceData(UNiagaraDataInterface* Interface) 
 	{
 		if (int32* InstDataOffset = DataInterfaceInstanceDataOffsets.Find(MakeWeakObjectPtr(const_cast<UNiagaraDataInterface*>(Interface))))
 		{
