Niagara: avoid compiling another shadermap for the same id.

- If we found an incomplete shadermap (that is, one that is already being compiled), use it instead of creating a new one for the same id.

diff --git a/Engine/Plugins/FX/Niagara/Source/NiagaraShader/Private/NiagaraShared.cpp b/Engine/Plugins/FX/Niagara/Source/NiagaraShader/Private/NiagaraShared.cpp
--- a/Engine/Plugins/FX/Niagara/Source/NiagaraShader/Private/NiagaraShared.cpp
+++ b/Engine/Plugins/FX/Niagara/Source/NiagaraShader/Private/NiagaraShared.cpp
@@ -293,7 +293,7 @@ bool FNiagaraShaderScript::CacheShaders(const FNiagaraShaderMapId& ShaderMapId,
 		GameThreadShaderMap = FNiagaraShaderMap::FindId(ShaderMapId, Platform);
 
 		// Attempt to load from the derived data cache if we are uncooked
-		if (!bForceRecompile && (!GameThreadShaderMap || !GameThreadShaderMap->IsComplete(this, true)) && !FPlatformProperties::RequiresCookedData())
+		if (!bForceRecompile && !GameThreadShaderMap && !FPlatformProperties::RequiresCookedData())
 		{
 			FNiagaraShaderMap::LoadFromDerivedDataCache(this, ShaderMapId, Platform, GameThreadShaderMap);
 			if (GameThreadShaderMap && GameThreadShaderMap->IsValid())
