- Fix for crash when failing to bind VM functions inside emitter spawn / update scripts

diff --git a/Engine/Plugins/FX/Niagara/Source/Niagara/Private/NiagaraSystemSimulation.cpp b/Engine/Plugins/FX/Niagara/Source/Niagara/Private/NiagaraSystemSimulation.cpp
--- a/Engine/Plugins/FX/Niagara/Source/Niagara/Private/NiagaraSystemSimulation.cpp
+++ b/Engine/Plugins/FX/Niagara/Source/Niagara/Private/NiagaraSystemSimulation.cpp
@@ -1426,8 +1426,17 @@ void FNiagaraSystemSimulation::SpawnSystemInstances(FNiagaraSystemSimulationTick
 	Context.DataSet.Allocate(NumInstances, true);
 	Context.DataSet.GetDestinationDataChecked().SetNumInstances(NumInstances);
 
-	//Run Spawn
-	SpawnExecContext.Tick(SoloSystemInstance);//We can't require a specific instance here as these are for all instances.
+	// Run Spawn
+	if (SpawnExecContext.Tick(SoloSystemInstance) == false)
+	{
+		for (FNiagaraSystemInstance* SystemInst : Context.Instances)
+		{
+			SystemInst->SetActualExecutionState(ENiagaraExecutionState::Disabled);
+		}
+		Context.DataSet.EndSimulate();
+		return;
+	}
+
 	SpawnExecContext.BindData(0, Context.DataSet, OrigNum, false);
 	SpawnExecContext.BindData(1, SpawnInstanceParameterDataSet, OrigNum, false);
 
@@ -1477,10 +1486,20 @@ void FNiagaraSystemSimulation::UpdateSystemInstances(FNiagaraSystemSimulationTic
 		DestinationData.Allocate(NumInstances);
 		DestinationData.SetNumInstances(NumInstances);
 
-		//Run update.
+		// Tick UpdateExecContext, this can fail to bind VM functions if this happens we become invalid so mark all instances as disabled
+		if (UpdateExecContext.Tick(Context.Instances[0]) == false)
+		{
+			for (FNiagaraSystemInstance* SystemInst : Context.Instances)
+			{
+				SystemInst->SetActualExecutionState(ENiagaraExecutionState::Disabled);
+			}
+			Context.DataSet.EndSimulate();
+			return;
+		}
+
+		// Run update.
 		if (OrigNum > 0)
 		{
-			UpdateExecContext.Tick(Context.Instances[0]);
 			UpdateExecContext.BindData(0, Context.DataSet, 0, false);
 			UpdateExecContext.BindData(1, UpdateInstanceParameterDataSet, 0, false);
 
@@ -1502,7 +1521,6 @@ void FNiagaraSystemSimulation::UpdateSystemInstances(FNiagaraSystemSimulationTic
 		//Ideally this should be compiled directly into the script similarly to interpolated particle spawning.
 		if ( (SpawnNum > 0) && GbSystemUpdateOnSpawn)
 		{
-			UpdateExecContext.Tick(Context.Instances[0]);
 
 			UpdateExecContext.BindData(0, Context.DataSet, OrigNum, false);
 			UpdateExecContext.BindData(1, UpdateInstanceParameterDataSet, OrigNum, false);
