- Remove roll when using camera facing plane / position on Niagara mesh renderers

diff --git a/Engine/Plugins/FX/Niagara/Shaders/Private/NiagaraMeshVertexFactory.ush b/Engine/Plugins/FX/Niagara/Shaders/Private/NiagaraMeshVertexFactory.ush
--- a/Engine/Plugins/FX/Niagara/Shaders/Private/NiagaraMeshVertexFactory.ush
+++ b/Engine/Plugins/FX/Niagara/Shaders/Private/NiagaraMeshVertexFactory.ush
@@ -666,16 +666,25 @@ FVertexFactoryIntermediates GetVertexFactoryIntermediates(FVertexFactoryInput In
 	{
 		float3 WorldPosition = TransformPosition(Intermediates.ParticlePosition);
 		float3 CameraVector = ResolvedView.WorldCameraOrigin.xyz - WorldPosition.xyz;
-		float3 CameraDirection = normalize(CameraVector);
-		float4 ToAlignmentQuat = FindQuatBetween(float3(1.0f, 0.0f, 0.0f), CameraDirection);
-		float4x4 ToAlignmentMat = QuatTo4x4( normalize(ToAlignmentQuat) );
+		float3 CameraForward = normalize(CameraVector);
+		float3 CameraRight = normalize(cross(CameraForward, ResolvedView.ViewUp));
+		float3 CameraUp = cross(CameraRight, CameraForward);
+
+		float4x4 ToAlignmentMat;
+		ToAlignmentMat[0] = float4(CameraForward.x, -CameraRight.x, CameraUp.x, 0.0f);
+		ToAlignmentMat[1] = float4(CameraForward.y, -CameraRight.y, CameraUp.y, 0.0f);
+		ToAlignmentMat[2] = float4(CameraForward.z, -CameraRight.z, CameraUp.z, 0.0f);
+		ToAlignmentMat[3] = float4(0.0f, 0.0f, 0.0f, 1.0f);
+
 		Intermediates.LocalToWorld = mul(LocalToWorldMat, mul(ToAlignmentMat, ScaleMatrix));
 	}
 	else if (MeshFacingMode == MESH_FACING_CAMERA_PLANE)
 	{
-		float3 CameraFacingVector = -ResolvedView.ViewForward;
-		float4 ToAlignmentQuat = FindQuatBetween(float3(1.0f, 0.0f, 0.0f), CameraFacingVector);
-		float4x4 ToAlignmentMat = QuatTo4x4( normalize(ToAlignmentQuat) );
+		float4x4 ToAlignmentMat;
+		ToAlignmentMat[0] = float4(-ResolvedView.ViewForward.x, -ResolvedView.ViewRight.x, ResolvedView.ViewUp.x, 0.0f);
+		ToAlignmentMat[1] = float4(-ResolvedView.ViewForward.y, -ResolvedView.ViewRight.y, ResolvedView.ViewUp.y, 0.0f);
+		ToAlignmentMat[2] = float4(-ResolvedView.ViewForward.z, -ResolvedView.ViewRight.z, ResolvedView.ViewUp.z, 0.0f);
+		ToAlignmentMat[3] = float4(0.0f, 0.0f, 0.0f, 1.0f);
 		Intermediates.LocalToWorld =  mul(LocalToWorldMat, mul(ToAlignmentMat, ScaleMatrix));
 	}
 	else
