Niagara - Merging 10906567 from main.  Should fix issues where systems are losing their data interfaces when saving, and will fix crashes for systems which are already broken.

diff --git a/Engine/Plugins/FX/Niagara/Source/Niagara/Private/NiagaraSystem.cpp b/Engine/Plugins/FX/Niagara/Source/Niagara/Private/NiagaraSystem.cpp
--- a/Engine/Plugins/FX/Niagara/Source/Niagara/Private/NiagaraSystem.cpp
+++ b/Engine/Plugins/FX/Niagara/Source/Niagara/Private/NiagaraSystem.cpp
@@ -285,6 +285,8 @@ void UNiagaraSystem::PostLoad()
 	Super::PostLoad();
 
 	ExposedParameters.PostLoad();
+	ExposedParameters.SanityCheckData();
+
 	SystemCompiledData.InstanceParamStore.PostLoad();
 
 	if (GIsEditor)
