Niagara - Fix an issue with the parameter store entries which were keeping references to objects after the UI was closed.

diff --git a/Engine/Plugins/FX/Niagara/Source/NiagaraEditor/Private/ViewModels/Stack/NiagaraStackParameterStoreEntry.cpp b/Engine/Plugins/FX/Niagara/Source/NiagaraEditor/Private/ViewModels/Stack/NiagaraStackParameterStoreEntry.cpp
--- a/Engine/Plugins/FX/Niagara/Source/NiagaraEditor/Private/ViewModels/Stack/NiagaraStackParameterStoreEntry.cpp
+++ b/Engine/Plugins/FX/Niagara/Source/NiagaraEditor/Private/ViewModels/Stack/NiagaraStackParameterStoreEntry.cpp
@@ -58,12 +58,12 @@ void UNiagaraStackParameterStoreEntry::RefreshChildrenInternal(const TArray<UNia
 {
 	RefreshValueAndHandle();
 
-	if (ValueObject != nullptr && ValueObject->IsA<UNiagaraDataInterface>())
+	if (ValueObject.IsValid() && ValueObject->IsA<UNiagaraDataInterface>())
 	{
 		if(ValueObjectEntry == nullptr || ValueObjectEntry->GetObject() != ValueObject)
 		{
 			ValueObjectEntry = NewObject<UNiagaraStackObject>(this);
-			ValueObjectEntry->Initialize(CreateDefaultChildRequiredData(), ValueObject, GetOwnerStackItemEditorDataKey());
+			ValueObjectEntry->Initialize(CreateDefaultChildRequiredData(), ValueObject.Get(), GetOwnerStackItemEditorDataKey());
 		}
 		NewChildren.Add(ValueObjectEntry);
 	}
@@ -106,7 +106,7 @@ TSharedPtr<FStructOnScope> UNiagaraStackParameterStoreEntry::GetValueStruct()
 
 UObject* UNiagaraStackParameterStoreEntry::GetValueObject()
 {
-	return ValueObject;
+	return ValueObject.Get();
 }
 
 void UNiagaraStackParameterStoreEntry::NotifyBeginValueChange()
diff --git a/Engine/Plugins/FX/Niagara/Source/NiagaraEditor/Public/ViewModels/Stack/NiagaraStackParameterStoreEntry.h b/Engine/Plugins/FX/Niagara/Source/NiagaraEditor/Public/ViewModels/Stack/NiagaraStackParameterStoreEntry.h
--- a/Engine/Plugins/FX/Niagara/Source/NiagaraEditor/Public/ViewModels/Stack/NiagaraStackParameterStoreEntry.h
+++ b/Engine/Plugins/FX/Niagara/Source/NiagaraEditor/Public/ViewModels/Stack/NiagaraStackParameterStoreEntry.h
@@ -131,11 +131,10 @@ private:
 	FOnParameterDeleted ParameterDeletedDelegate;
 
 	/** A pointer to the data interface object for this input if one is available. */
-	UPROPERTY()
-	UObject* ValueObject;
+	TWeakObjectPtr<UObject> ValueObject;
 
-	UPROPERTY()
-	UObject* Owner;
+	/** A pointer to the owner of the parameter store that owns this entry. */
+	TWeakObjectPtr<UObject> Owner;
 
 	FNiagaraParameterStore* ParameterStore;
 
