Prevent static switches from receiving drag and drop inputs, because it leaves them in a permanently broken state

diff --git a/Engine/Plugins/FX/Niagara/Source/NiagaraEditorWidgets/Private/Stack/SNiagaraStackFunctionInputValue.cpp b/Engine/Plugins/FX/Niagara/Source/NiagaraEditorWidgets/Private/Stack/SNiagaraStackFunctionInputValue.cpp
--- a/Engine/Plugins/FX/Niagara/Source/NiagaraEditorWidgets/Private/Stack/SNiagaraStackFunctionInputValue.cpp
+++ b/Engine/Plugins/FX/Niagara/Source/NiagaraEditorWidgets/Private/Stack/SNiagaraStackFunctionInputValue.cpp
@@ -848,6 +848,11 @@ bool SNiagaraStackFunctionInputValue::OnFunctionInputAllowDrop(TSharedPtr<FDragD
 {
 	if (FunctionInput && DragDropOperation->IsOfType<FNiagaraParameterDragOperation>())
 	{
+		if (FunctionInput->IsStaticParameter())
+		{
+			return false;
+		}
+
 		TSharedPtr<FNiagaraParameterDragOperation> InputDragDropOperation = StaticCastSharedPtr<FNiagaraParameterDragOperation>(DragDropOperation);
 		TSharedPtr<FNiagaraParameterAction> Action = StaticCastSharedPtr<FNiagaraParameterAction>(InputDragDropOperation->GetSourceAction());
 		if (Action->GetParameter().GetType() == FunctionInput->GetInputType() && FNiagaraStackGraphUtilities::ParameterAllowedInExecutionCategory(Action->GetParameter().GetName(), FunctionInput->GetExecutionCategoryName()))
