Copying Private-LoadTimes-4.24 stream to Main. Biggest changes are in Materials/Shader memory freezing.

diff --git a/Engine/Plugins/FX/Niagara/Source/Niagara/Private/NiagaraDataInterfaceNeighborGrid3D.cpp b/Engine/Plugins/FX/Niagara/Source/Niagara/Private/NiagaraDataInterfaceNeighborGrid3D.cpp
--- a/Engine/Plugins/FX/Niagara/Source/Niagara/Private/NiagaraDataInterfaceNeighborGrid3D.cpp
+++ b/Engine/Plugins/FX/Niagara/Source/Niagara/Private/NiagaraDataInterfaceNeighborGrid3D.cpp
@@ -28,41 +28,30 @@ static const FName SetParticleNeighborCountFunctionName("SetParticleNeighborCoun
 /*--------------------------------------------------------------------------------------------------------------------------*/
 struct FNiagaraDataInterfaceParametersCS_NeighborGrid3D : public FNiagaraDataInterfaceParametersCS
 {
-	virtual void Bind(const FNiagaraDataInterfaceParamRef& ParamRef, const class FShaderParameterMap& ParameterMap) override
+	DECLARE_TYPE_LAYOUT(FNiagaraDataInterfaceParametersCS_NeighborGrid3D, NonVirtual);
+public:
+	void Bind(const FNiagaraDataInterfaceGPUParamInfo& ParameterInfo, const class FShaderParameterMap& ParameterMap)
 	{
-		NumVoxelsParam.Bind(ParameterMap, *(NumVoxelsName + ParamRef.ParameterInfo.DataInterfaceHLSLSymbol));		
-		VoxelSizeParam.Bind(ParameterMap, *(VoxelSizeName + ParamRef.ParameterInfo.DataInterfaceHLSLSymbol));
+		NumVoxelsParam.Bind(ParameterMap, *(NumVoxelsName + ParameterInfo.DataInterfaceHLSLSymbol));		
+		VoxelSizeParam.Bind(ParameterMap, *(VoxelSizeName + ParameterInfo.DataInterfaceHLSLSymbol));
 
-		MaxNeighborsPerVoxelParam.Bind(ParameterMap, *(MaxNeighborsPerVoxelName + ParamRef.ParameterInfo.DataInterfaceHLSLSymbol));		
-		WorldBBoxSizeParam.Bind(ParameterMap, *(WorldBBoxSizeName + ParamRef.ParameterInfo.DataInterfaceHLSLSymbol));
+		MaxNeighborsPerVoxelParam.Bind(ParameterMap, *(MaxNeighborsPerVoxelName + ParameterInfo.DataInterfaceHLSLSymbol));
+		WorldBBoxSizeParam.Bind(ParameterMap, *(WorldBBoxSizeName + ParameterInfo.DataInterfaceHLSLSymbol));
 		
-		ParticleNeighborsGridParam.Bind(ParameterMap,  *(ParticleNeighborsName + ParamRef.ParameterInfo.DataInterfaceHLSLSymbol));
-		ParticleNeighborCountGridParam.Bind(ParameterMap,  *(ParticleNeighborCountName + ParamRef.ParameterInfo.DataInterfaceHLSLSymbol));
+		ParticleNeighborsGridParam.Bind(ParameterMap,  *(ParticleNeighborsName + ParameterInfo.DataInterfaceHLSLSymbol));
+		ParticleNeighborCountGridParam.Bind(ParameterMap,  *(ParticleNeighborCountName + ParameterInfo.DataInterfaceHLSLSymbol));
 
-		OutputParticleNeighborsGridParam.Bind(ParameterMap, *(OutputParticleNeighborsName + ParamRef.ParameterInfo.DataInterfaceHLSLSymbol));
-		OutputParticleNeighborCountGridParam.Bind(ParameterMap, *(OutputParticleNeighborCountName + ParamRef.ParameterInfo.DataInterfaceHLSLSymbol));		
-	}
-
-	virtual void Serialize(FArchive& Ar)override
-	{
-		Ar << NumVoxelsParam;
-		Ar << VoxelSizeParam;
-		Ar << MaxNeighborsPerVoxelParam;		
-		Ar << WorldBBoxSizeParam;
-		
-		Ar << ParticleNeighborsGridParam;
-		Ar << ParticleNeighborCountGridParam;
-		Ar << OutputParticleNeighborsGridParam;
-		Ar << OutputParticleNeighborCountGridParam;		
+		OutputParticleNeighborsGridParam.Bind(ParameterMap, *(OutputParticleNeighborsName + ParameterInfo.DataInterfaceHLSLSymbol));
+		OutputParticleNeighborCountGridParam.Bind(ParameterMap, *(OutputParticleNeighborCountName + ParameterInfo.DataInterfaceHLSLSymbol));		
 	}
 
 	// #todo(dmp): make resource transitions batched
-	virtual void Set(FRHICommandList& RHICmdList, const FNiagaraDataInterfaceSetArgs& Context) const override
+	void Set(FRHICommandList& RHICmdList, const FNiagaraDataInterfaceSetArgs& Context) const
 	{
 		check(IsInRenderingThread());
 
 		// Get shader and DI
-		FRHIComputeShader* ComputeShaderRHI = Context.Shader->GetComputeShader();
+		FRHIComputeShader* ComputeShaderRHI = Context.Shader.GetComputeShader();
 		FNiagaraDataInterfaceProxyNeighborGrid3D* VFDI = static_cast<FNiagaraDataInterfaceProxyNeighborGrid3D*>(Context.DataInterface);
 
 		NeighborGrid3DRWInstanceData* ProxyData = VFDI->SystemInstancesToProxyData.Find(Context.SystemInstance);
@@ -85,13 +74,13 @@ struct FNiagaraDataInterfaceParametersCS_NeighborGrid3D : public FNiagaraDataInt
 			if (ParticleNeighborsGridParam.IsBound())
 			{
 				RHICmdList.TransitionResource(EResourceTransitionAccess::EReadable, EResourceTransitionPipeline::EComputeToCompute, ProxyData->NeighborhoodBuffer.UAV);
-				SetSRVParameter(RHICmdList, Context.Shader->GetComputeShader(), ParticleNeighborsGridParam, ProxyData->NeighborhoodBuffer.SRV);
+				SetSRVParameter(RHICmdList, Context.Shader.GetComputeShader(), ParticleNeighborsGridParam, ProxyData->NeighborhoodBuffer.SRV);
 			}
 
 			if (ParticleNeighborCountGridParam.IsBound())
 			{
 				RHICmdList.TransitionResource(EResourceTransitionAccess::EReadable, EResourceTransitionPipeline::EComputeToCompute, ProxyData->NeighborhoodCountBuffer.UAV);
-				SetSRVParameter(RHICmdList, Context.Shader->GetComputeShader(), ParticleNeighborCountGridParam, ProxyData->NeighborhoodCountBuffer.SRV);
+				SetSRVParameter(RHICmdList, Context.Shader.GetComputeShader(), ParticleNeighborCountGridParam, ProxyData->NeighborhoodCountBuffer.SRV);
 			}
 		}
 		else
@@ -99,42 +88,45 @@ struct FNiagaraDataInterfaceParametersCS_NeighborGrid3D : public FNiagaraDataInt
 			if (OutputParticleNeighborsGridParam.IsBound())
 			{
 				RHICmdList.TransitionResource(EResourceTransitionAccess::EWritable, EResourceTransitionPipeline::EComputeToCompute, ProxyData->NeighborhoodBuffer.UAV);
-				OutputParticleNeighborsGridParam.SetBuffer(RHICmdList, Context.Shader->GetComputeShader(), ProxyData->NeighborhoodBuffer);
+				OutputParticleNeighborsGridParam.SetBuffer(RHICmdList, Context.Shader.GetComputeShader(), ProxyData->NeighborhoodBuffer);
 			}
 
 			if (OutputParticleNeighborCountGridParam.IsBound())
 			{
 				RHICmdList.TransitionResource(EResourceTransitionAccess::EWritable, EResourceTransitionPipeline::EComputeToCompute, ProxyData->NeighborhoodCountBuffer.UAV);
-				OutputParticleNeighborCountGridParam.SetBuffer(RHICmdList, Context.Shader->GetComputeShader(), ProxyData->NeighborhoodCountBuffer);
+				OutputParticleNeighborCountGridParam.SetBuffer(RHICmdList, Context.Shader.GetComputeShader(), ProxyData->NeighborhoodCountBuffer);
 			}
 		}
 		// Note: There is a flush in PreEditChange to make sure everything is synced up at this point 
 	}
 
-	virtual void Unset(FRHICommandList& RHICmdList, const FNiagaraDataInterfaceSetArgs& Context) const 
+	void Unset(FRHICommandList& RHICmdList, const FNiagaraDataInterfaceSetArgs& Context) const 
 	{
 		if (OutputParticleNeighborsGridParam.IsBound())
 		{
-			OutputParticleNeighborsGridParam.UnsetUAV(RHICmdList, Context.Shader->GetComputeShader());
+			OutputParticleNeighborsGridParam.UnsetUAV(RHICmdList, Context.Shader.GetComputeShader());
 		}
 
 		if (OutputParticleNeighborCountGridParam.IsBound())
 		{
-			OutputParticleNeighborCountGridParam.UnsetUAV(RHICmdList, Context.Shader->GetComputeShader());
+			OutputParticleNeighborCountGridParam.UnsetUAV(RHICmdList, Context.Shader.GetComputeShader());
 		}
 	}
 
 private:
-	FShaderParameter NumVoxelsParam;
-	FShaderParameter VoxelSizeParam;
-	FShaderParameter MaxNeighborsPerVoxelParam;	
-	FShaderParameter WorldBBoxSizeParam;
-	FShaderResourceParameter ParticleNeighborsGridParam;
-	FShaderResourceParameter ParticleNeighborCountGridParam;
-	FRWShaderParameter OutputParticleNeighborCountGridParam;
-	FRWShaderParameter OutputParticleNeighborsGridParam;
+	LAYOUT_FIELD(FShaderParameter, NumVoxelsParam);
+	LAYOUT_FIELD(FShaderParameter, VoxelSizeParam);
+	LAYOUT_FIELD(FShaderParameter, MaxNeighborsPerVoxelParam);
+	LAYOUT_FIELD(FShaderParameter, WorldBBoxSizeParam);
+	LAYOUT_FIELD(FShaderResourceParameter, ParticleNeighborsGridParam);
+	LAYOUT_FIELD(FShaderResourceParameter, ParticleNeighborCountGridParam);
+	LAYOUT_FIELD(FRWShaderParameter, OutputParticleNeighborCountGridParam);
+	LAYOUT_FIELD(FRWShaderParameter, OutputParticleNeighborsGridParam);
 };
 
+IMPLEMENT_TYPE_LAYOUT(FNiagaraDataInterfaceParametersCS_NeighborGrid3D);
+
+IMPLEMENT_NIAGARA_DI_PARAMETER(UNiagaraDataInterfaceNeighborGrid3D, FNiagaraDataInterfaceParametersCS_NeighborGrid3D);
 
 UNiagaraDataInterfaceNeighborGrid3D::UNiagaraDataInterfaceNeighborGrid3D(FObjectInitializer const& ObjectInitializer)
 	: Super(ObjectInitializer)
@@ -396,13 +388,6 @@ bool UNiagaraDataInterfaceNeighborGrid3D::GetFunctionHLSL(const FNiagaraDataInte
 	return false;
 }
 
-FNiagaraDataInterfaceParametersCS* UNiagaraDataInterfaceNeighborGrid3D::ConstructComputeParameters()const
-{
-	return new FNiagaraDataInterfaceParametersCS_NeighborGrid3D();
-}
-
-
-
 bool UNiagaraDataInterfaceNeighborGrid3D::InitPerInstanceData(void* PerInstanceData, FNiagaraSystemInstance* SystemInstance)
 {
 
diff --git a/Engine/Plugins/FX/Niagara/Source/Niagara/Private/NiagaraDataInterfaceOcclusion.cpp b/Engine/Plugins/FX/Niagara/Source/Niagara/Private/NiagaraDataInterfaceOcclusion.cpp
--- a/Engine/Plugins/FX/Niagara/Source/Niagara/Private/NiagaraDataInterfaceOcclusion.cpp
+++ b/Engine/Plugins/FX/Niagara/Source/Niagara/Private/NiagaraDataInterfaceOcclusion.cpp
@@ -274,32 +274,28 @@ void UNiagaraDataInterfaceOcclusion::QueryOcclusionFactorCircleGPU(FVectorVMCont
 
 struct FNiagaraDataInterfaceParametersCS_OcclusionQuery : public FNiagaraDataInterfaceParametersCS
 {
-	virtual void Bind(const FNiagaraDataInterfaceParamRef& ParamRef, const class FShaderParameterMap& ParameterMap) override
+	DECLARE_TYPE_LAYOUT(FNiagaraDataInterfaceParametersCS_OcclusionQuery, NonVirtual);
+public:
+	void Bind(const FNiagaraDataInterfaceGPUParamInfo& ParameterInfo, const class FShaderParameterMap& ParameterMap)
 	{
 		PassUniformBuffer.Bind(ParameterMap, FSceneTexturesUniformParameters::StaticStructMetadata.GetShaderVariableName());
 	}
 
-	virtual void Serialize(FArchive& Ar) override
-	{
-		Ar << PassUniformBuffer;
-	}
-
-	virtual void Set(FRHICommandList& RHICmdList, const FNiagaraDataInterfaceSetArgs& Context) const override
+	void Set(FRHICommandList& RHICmdList, const FNiagaraDataInterfaceSetArgs& Context) const
 	{
 		check(IsInRenderingThread());
-		FRHIComputeShader* ComputeShaderRHI = Context.Shader->GetComputeShader();
+		FRHIComputeShader* ComputeShaderRHI = RHICmdList.GetBoundComputeShader();
 
 		TUniformBufferRef<FSceneTexturesUniformParameters> SceneTextureUniformParams = GNiagaraViewDataManager.GetSceneTextureUniformParameters();
 		SetUniformBufferParameter(RHICmdList, ComputeShaderRHI, PassUniformBuffer, SceneTextureUniformParams);
 	}
 
 private:
-
 	/** The SceneDepthTexture parameter for depth buffer query. */
-	FShaderUniformBufferParameter PassUniformBuffer;
+	LAYOUT_FIELD(FShaderUniformBufferParameter, PassUniformBuffer);
 };
 
-FNiagaraDataInterfaceParametersCS* UNiagaraDataInterfaceOcclusion::ConstructComputeParameters() const
-{
-	return new FNiagaraDataInterfaceParametersCS_OcclusionQuery();
-}
+IMPLEMENT_TYPE_LAYOUT(FNiagaraDataInterfaceParametersCS_OcclusionQuery);
+
+IMPLEMENT_NIAGARA_DI_PARAMETER(UNiagaraDataInterfaceOcclusion, FNiagaraDataInterfaceParametersCS_OcclusionQuery);
+
diff --git a/Engine/Plugins/FX/Niagara/Source/Niagara/Private/NiagaraDataInterfaceParticleRead.cpp b/Engine/Plugins/FX/Niagara/Source/Niagara/Private/NiagaraDataInterfaceParticleRead.cpp
--- a/Engine/Plugins/FX/Niagara/Source/Niagara/Private/NiagaraDataInterfaceParticleRead.cpp
+++ b/Engine/Plugins/FX/Niagara/Source/Niagara/Private/NiagaraDataInterfaceParticleRead.cpp
@@ -150,25 +150,25 @@ struct FNiagaraDataInterfaceParametersCS_ParticleRead : public FNiagaraDataInter
 		return ENiagaraParticleDataValueType::Invalid;
 	}
 
-	virtual void Bind(const FNiagaraDataInterfaceParamRef& ParamRef, const class FShaderParameterMap& ParameterMap) override
-	{
-		NumSpawnedParticlesParam.Bind(ParameterMap, *(NumSpawnedParticlesBaseName + ParamRef.ParameterInfo.DataInterfaceHLSLSymbol));
-		SpawnedParticlesAcquireTagParam.Bind(ParameterMap, *(SpawnedParticlesAcquireTagBaseName + ParamRef.ParameterInfo.DataInterfaceHLSLSymbol));
-		SpawnedIDsBufferParam.Bind(ParameterMap, *(SpawnedIDsBufferBaseName + ParamRef.ParameterInfo.DataInterfaceHLSLSymbol));
-		IDToIndexTableParam.Bind(ParameterMap, *(IDToIndexTableBaseName + ParamRef.ParameterInfo.DataInterfaceHLSLSymbol));
-		InputFloatBufferParam.Bind(ParameterMap, *(InputFloatBufferBaseName + ParamRef.ParameterInfo.DataInterfaceHLSLSymbol));
-		InputIntBufferParam.Bind(ParameterMap, *(InputIntBufferBaseName + ParamRef.ParameterInfo.DataInterfaceHLSLSymbol));
-		ParticleStrideFloatParam.Bind(ParameterMap, *(ParticleStrideFloatBaseName + ParamRef.ParameterInfo.DataInterfaceHLSLSymbol));
-		ParticleStrideIntParam.Bind(ParameterMap, *(ParticleStrideIntBaseName + ParamRef.ParameterInfo.DataInterfaceHLSLSymbol));
-		AttributeIndicesParam.Bind(ParameterMap, *(AttributeIndicesBaseName + ParamRef.ParameterInfo.DataInterfaceHLSLSymbol));
-		AcquireTagRegisterIndexParam.Bind(ParameterMap, *(AcquireTagRegisterIndexBaseName + ParamRef.ParameterInfo.DataInterfaceHLSLSymbol));
-
-		int32 NumFuncs = ParamRef.ParameterInfo.GeneratedFunctions.Num();
+	void Bind(const FNiagaraDataInterfaceGPUParamInfo& ParameterInfo, const class FShaderParameterMap& ParameterMap)
+	{
+		NumSpawnedParticlesParam.Bind(ParameterMap, *(NumSpawnedParticlesBaseName + ParameterInfo.DataInterfaceHLSLSymbol));
+		SpawnedParticlesAcquireTagParam.Bind(ParameterMap, *(SpawnedParticlesAcquireTagBaseName + ParameterInfo.DataInterfaceHLSLSymbol));
+		SpawnedIDsBufferParam.Bind(ParameterMap, *(SpawnedIDsBufferBaseName + ParameterInfo.DataInterfaceHLSLSymbol));
+		IDToIndexTableParam.Bind(ParameterMap, *(IDToIndexTableBaseName + ParameterInfo.DataInterfaceHLSLSymbol));
+		InputFloatBufferParam.Bind(ParameterMap, *(InputFloatBufferBaseName + ParameterInfo.DataInterfaceHLSLSymbol));
+		InputIntBufferParam.Bind(ParameterMap, *(InputIntBufferBaseName + ParameterInfo.DataInterfaceHLSLSymbol));
+		ParticleStrideFloatParam.Bind(ParameterMap, *(ParticleStrideFloatBaseName + ParameterInfo.DataInterfaceHLSLSymbol));
+		ParticleStrideIntParam.Bind(ParameterMap, *(ParticleStrideIntBaseName + ParameterInfo.DataInterfaceHLSLSymbol));
+		AttributeIndicesParam.Bind(ParameterMap, *(AttributeIndicesBaseName + ParameterInfo.DataInterfaceHLSLSymbol));
+		AcquireTagRegisterIndexParam.Bind(ParameterMap, *(AcquireTagRegisterIndexBaseName + ParameterInfo.DataInterfaceHLSLSymbol));
+
+		int32 NumFuncs = ParameterInfo.GeneratedFunctions.Num();
 		AttributeNames.SetNum(NumFuncs);
 		AttributeTypes.SetNum(NumFuncs);
 		for (int32 FuncIdx = 0; FuncIdx < NumFuncs; ++FuncIdx)
 		{
-			const FNiagaraDataInterfaceGeneratedFunction& Func = ParamRef.ParameterInfo.GeneratedFunctions[FuncIdx];
+			const FNiagaraDataInterfaceGeneratedFunction& Func = ParameterInfo.GeneratedFunctions[FuncIdx];
 			static const FName NAME_Attribute("Attribute");
 			const FName* AttributeName = Func.FindSpecifierValue(NAME_Attribute);
 			if (AttributeName != nullptr)
@@ -189,25 +189,6 @@ struct FNiagaraDataInterfaceParametersCS_ParticleRead : public FNiagaraDataInter
 		AttributeIndices.SetNum(NumAttrIndices);
 	}
 
-	virtual void Serialize(FArchive& Ar) override
-	{
-		Ar << NumSpawnedParticlesParam;
-		Ar << SpawnedParticlesAcquireTagParam;
-		Ar << SpawnedIDsBufferParam;
-		Ar << IDToIndexTableParam;
-		Ar << InputFloatBufferParam;
-		Ar << InputIntBufferParam;
-		Ar << ParticleStrideFloatParam;
-		Ar << ParticleStrideIntParam;
-		Ar << AttributeIndicesParam;
-		Ar << AcquireTagRegisterIndexParam;
-		Ar << AttributeNames;
-		Ar << AttributeTypes;
-
-		int NumAttrIndices = Align(AttributeNames.Num(), 4);
-		AttributeIndices.SetNum(NumAttrIndices);
-	}
-
 	void SetErrorParams(FRHICommandList& RHICmdList, FRHIComputeShader* ComputeShader, bool bSkipSpawnInfo) const
 	{
 		CachedDataSet = nullptr;
@@ -318,11 +299,11 @@ struct FNiagaraDataInterfaceParametersCS_ParticleRead : public FNiagaraDataInter
 		}
 	}
 
-	virtual void Set(FRHICommandList& RHICmdList, const FNiagaraDataInterfaceSetArgs& Context) const override
+	void Set(FRHICommandList& RHICmdList, const FNiagaraDataInterfaceSetArgs& Context) const
 	{
 		check(IsInRenderingThread());
 
-		FRHIComputeShader* ComputeShader = Context.Shader->GetComputeShader();
+		FRHIComputeShader* ComputeShader = RHICmdList.GetBoundComputeShader();
 
 		FNiagaraDataInterfaceProxyParticleRead* Proxy = static_cast<FNiagaraDataInterfaceProxyParticleRead*>(Context.DataInterface);
 		check(Proxy);
@@ -1309,11 +1290,6 @@ bool UNiagaraDataInterfaceParticleRead::GetFunctionHLSL(const FNiagaraDataInterf
 	return false;
 }
 
-FNiagaraDataInterfaceParametersCS* UNiagaraDataInterfaceParticleRead::ConstructComputeParameters() const
-{
-	return new FNiagaraDataInterfaceParametersCS_ParticleRead();
-}
-
 void UNiagaraDataInterfaceParticleRead::ProvidePerInstanceDataForRenderThread(void* DataForRenderThread, void* PerInstanceData, const FNiagaraSystemInstanceID& SystemInstance)
 {
 	FNDIParticleRead_RenderThreadData* RTData = new (DataForRenderThread) FNDIParticleRead_RenderThreadData;
diff --git a/Engine/Plugins/FX/Niagara/Source/Niagara/Private/NiagaraDataInterfaceRW.cpp b/Engine/Plugins/FX/Niagara/Source/Niagara/Private/NiagaraDataInterfaceRW.cpp
--- a/Engine/Plugins/FX/Niagara/Source/Niagara/Private/NiagaraDataInterfaceRW.cpp
+++ b/Engine/Plugins/FX/Niagara/Source/Niagara/Private/NiagaraDataInterfaceRW.cpp
@@ -277,11 +277,6 @@ bool UNiagaraDataInterfaceGrid3D::GetFunctionHLSL(const FNiagaraDataInterfaceGPU
 	return false;
 }
 
-FNiagaraDataInterfaceParametersCS* UNiagaraDataInterfaceGrid3D::ConstructComputeParameters()const
-{
-	return nullptr;
-}
-
 
 
 bool UNiagaraDataInterfaceGrid3D::CopyToInternal(UNiagaraDataInterface* Destination) const
@@ -675,12 +670,6 @@ bool UNiagaraDataInterfaceGrid2D::GetFunctionHLSL(const FNiagaraDataInterfaceGPU
 	return false;
 }
 
-FNiagaraDataInterfaceParametersCS* UNiagaraDataInterfaceGrid2D::ConstructComputeParameters()const
-{
-	return nullptr;
-}
-
-
 
 bool UNiagaraDataInterfaceGrid2D::CopyToInternal(UNiagaraDataInterface* Destination) const
 {
diff --git a/Engine/Plugins/FX/Niagara/Source/Niagara/Private/NiagaraDataInterfaceSkeletalMesh.cpp b/Engine/Plugins/FX/Niagara/Source/Niagara/Private/NiagaraDataInterfaceSkeletalMesh.cpp
--- a/Engine/Plugins/FX/Niagara/Source/Niagara/Private/NiagaraDataInterfaceSkeletalMesh.cpp
+++ b/Engine/Plugins/FX/Niagara/Source/Niagara/Private/NiagaraDataInterfaceSkeletalMesh.cpp
@@ -488,8 +488,8 @@ void FSkeletalMeshGpuSpawnStaticBuffers::InitRHI()
 	if (bUseGpuUniformlyDistributedSampling)
 	{
 		const FSkeletalMeshAreaWeightedTriangleSampler& triangleSampler = SkeletalMeshSamplingLODBuiltData->AreaWeightedTriangleSampler;
-		const TArray<float>& Prob = triangleSampler.GetProb();
-		const TArray<int32>& Alias = triangleSampler.GetAlias();
+		const TArray<float, FMemoryImageAllocator>& Prob = triangleSampler.GetProb();
+		const TArray<int32, FMemoryImageAllocator>& Alias = triangleSampler.GetAlias();
 		check(TriangleCount == triangleSampler.GetNumEntries());
 
 		FRHIResourceCreateInfo CreateInfo;
@@ -767,10 +767,12 @@ static void GetNiagaraDataInterfaceParametersName(FNDISkeletalMeshParametersName
 
 struct FNiagaraDataInterfaceParametersCS_SkeletalMesh : public FNiagaraDataInterfaceParametersCS
 {
-	virtual void Bind(const FNiagaraDataInterfaceParamRef& ParamRef, const class FShaderParameterMap& ParameterMap) override
+	DECLARE_TYPE_LAYOUT(FNiagaraDataInterfaceParametersCS_SkeletalMesh, NonVirtual);
+public:
+	void Bind(const FNiagaraDataInterfaceGPUParamInfo& ParameterInfo, const class FShaderParameterMap& ParameterMap)
 	{
 		FNDISkeletalMeshParametersName ParamNames;
-		GetNiagaraDataInterfaceParametersName(ParamNames, ParamRef.ParameterInfo.DataInterfaceHLSLSymbol);
+		GetNiagaraDataInterfaceParametersName(ParamNames, ParameterInfo.DataInterfaceHLSLSymbol);
 
 		MeshIndexBuffer.Bind(ParameterMap, *ParamNames.MeshIndexBufferName);
 		MeshVertexBuffer.Bind(ParameterMap, *ParamNames.MeshVertexBufferName);
@@ -802,43 +804,11 @@ struct FNiagaraDataInterfaceParametersCS_SkeletalMesh : public FNiagaraDataInter
 		EnabledFeatures.Bind(ParameterMap, *ParamNames.EnabledFeaturesName);
 	}
 
-	virtual void Serialize(FArchive& Ar)override
-	{
-		Ar << MeshIndexBuffer;
-		Ar << MeshVertexBuffer;
-		Ar << MeshSkinWeightBuffer;
-		Ar << MeshSkinWeightLookupBuffer;
-		Ar << MeshCurrBonesBuffer;
-		Ar << MeshPrevBonesBuffer;
-		Ar << MeshCurrSamplingBonesBuffer;
-		Ar << MeshPrevSamplingBonesBuffer;
-		Ar << MeshTangentBuffer;
-		Ar << MeshTexCoordBuffer;
-		Ar << MeshColorBuffer;
-		Ar << MeshTriangleSamplerProbaBuffer;
-		Ar << MeshTriangleSamplerAliasBuffer;
-		Ar << MeshTriangleMatricesOffsetBuffer;
-		Ar << MeshTriangleCount;
-		Ar << MeshVertexCount;
-		Ar << MeshWeightStride;
-		Ar << MeshSkinWeightIndexSize;
-		Ar << MeshNumTexCoord;
-		Ar << MeshNumWeights;
-		Ar << NumSpecificBones;
-		Ar << SpecificBones;
-		Ar << NumSpecificSockets;
-		Ar << SpecificSocketBoneOffset;
-		Ar << InstanceTransform;
-		Ar << InstancePrevTransform;
-		Ar << InstanceInvDeltaTime;
-		Ar << EnabledFeatures;
-	}
-
-	virtual void Set(FRHICommandList& RHICmdList, const FNiagaraDataInterfaceSetArgs& Context) const override
+	void Set(FRHICommandList& RHICmdList, const FNiagaraDataInterfaceSetArgs& Context) const
 	{
 		check(IsInRenderingThread());
 
-		FRHIComputeShader* ComputeShaderRHI = Context.Shader->GetComputeShader();
+		FRHIComputeShader* ComputeShaderRHI = Context.Shader.GetComputeShader();
 		FNiagaraDataInterfaceProxySkeletalMesh* InterfaceProxy = static_cast<FNiagaraDataInterfaceProxySkeletalMesh*>(Context.DataInterface);
 		FNiagaraDataInterfaceProxySkeletalMeshData* InstanceData = InterfaceProxy->SystemInstancesToData.Find(Context.SystemInstance);
 		if (InstanceData && InstanceData->StaticBuffers)
@@ -966,39 +936,41 @@ struct FNiagaraDataInterfaceParametersCS_SkeletalMesh : public FNiagaraDataInter
 		}
 	}
 
-
 private:
-
-	FShaderResourceParameter MeshIndexBuffer;
-	FShaderResourceParameter MeshVertexBuffer;
-	FShaderResourceParameter MeshSkinWeightBuffer;
-	FShaderResourceParameter MeshSkinWeightLookupBuffer;
-	FShaderResourceParameter MeshCurrBonesBuffer;
-	FShaderResourceParameter MeshPrevBonesBuffer;
-	FShaderResourceParameter MeshCurrSamplingBonesBuffer;
-	FShaderResourceParameter MeshPrevSamplingBonesBuffer;
-	FShaderResourceParameter MeshTangentBuffer;
-	FShaderResourceParameter MeshTexCoordBuffer;
-	FShaderResourceParameter MeshColorBuffer;
-	FShaderResourceParameter MeshTriangleSamplerProbaBuffer;
-	FShaderResourceParameter MeshTriangleSamplerAliasBuffer;
-	FShaderResourceParameter MeshTriangleMatricesOffsetBuffer;
-	FShaderParameter MeshTriangleCount;
-	FShaderParameter MeshVertexCount;
-	FShaderParameter MeshWeightStride;
-	FShaderParameter MeshSkinWeightIndexSize;
-	FShaderParameter MeshNumTexCoord;
-	FShaderParameter MeshNumWeights;
-	FShaderParameter NumSpecificBones;
-	FShaderResourceParameter SpecificBones;
-	FShaderParameter NumSpecificSockets;
-	FShaderParameter SpecificSocketBoneOffset;
-	FShaderParameter InstanceTransform;
-	FShaderParameter InstancePrevTransform;
-	FShaderParameter InstanceInvDeltaTime;
-	FShaderParameter EnabledFeatures;
+	LAYOUT_FIELD(FShaderResourceParameter, MeshIndexBuffer);
+	LAYOUT_FIELD(FShaderResourceParameter, MeshVertexBuffer);
+	LAYOUT_FIELD(FShaderResourceParameter, MeshSkinWeightBuffer);
+	LAYOUT_FIELD(FShaderResourceParameter, MeshSkinWeightLookupBuffer);
+	LAYOUT_FIELD(FShaderResourceParameter, MeshCurrBonesBuffer);
+	LAYOUT_FIELD(FShaderResourceParameter, MeshPrevBonesBuffer);
+	LAYOUT_FIELD(FShaderResourceParameter, MeshCurrSamplingBonesBuffer);
+	LAYOUT_FIELD(FShaderResourceParameter, MeshPrevSamplingBonesBuffer);
+	LAYOUT_FIELD(FShaderResourceParameter, MeshTangentBuffer);
+	LAYOUT_FIELD(FShaderResourceParameter, MeshTexCoordBuffer);
+	LAYOUT_FIELD(FShaderResourceParameter, MeshColorBuffer);
+	LAYOUT_FIELD(FShaderResourceParameter, MeshTriangleSamplerProbaBuffer);
+	LAYOUT_FIELD(FShaderResourceParameter, MeshTriangleSamplerAliasBuffer);
+	LAYOUT_FIELD(FShaderResourceParameter, MeshTriangleMatricesOffsetBuffer);
+	LAYOUT_FIELD(FShaderParameter, MeshTriangleCount);
+	LAYOUT_FIELD(FShaderParameter, MeshVertexCount);
+	LAYOUT_FIELD(FShaderParameter, MeshWeightStride);
+	LAYOUT_FIELD(FShaderParameter, MeshSkinWeightIndexSize);
+	LAYOUT_FIELD(FShaderParameter, MeshNumTexCoord);
+	LAYOUT_FIELD(FShaderParameter, MeshNumWeights);
+	LAYOUT_FIELD(FShaderParameter, NumSpecificBones);
+	LAYOUT_FIELD(FShaderResourceParameter, SpecificBones);
+	LAYOUT_FIELD(FShaderParameter, NumSpecificSockets);
+	LAYOUT_FIELD(FShaderParameter, SpecificSocketBoneOffset);
+	LAYOUT_FIELD(FShaderParameter, InstanceTransform);
+	LAYOUT_FIELD(FShaderParameter, InstancePrevTransform);
+	LAYOUT_FIELD(FShaderParameter, InstanceInvDeltaTime);
+	LAYOUT_FIELD(FShaderParameter, EnabledFeatures);
 };
 
+IMPLEMENT_TYPE_LAYOUT(FNiagaraDataInterfaceParametersCS_SkeletalMesh);
+
+IMPLEMENT_NIAGARA_DI_PARAMETER(UNiagaraDataInterfaceSkeletalMesh, FNiagaraDataInterfaceParametersCS_SkeletalMesh);
+
 //////////////////////////////////////////////////////////////////////////
 
 void FNiagaraDataInterfaceProxySkeletalMesh::ConsumePerInstanceDataFromGameThread(void* PerInstanceData, const FNiagaraSystemInstanceID& Instance)
@@ -2239,12 +2211,6 @@ void UNiagaraDataInterfaceSkeletalMesh::GetParameterDefinitionHLSL(const FNiagar
 	OutHLSL += TEXT("DISKELMESH_DECLARE_CONSTANTS(") + ParamInfo.DataInterfaceHLSLSymbol + TEXT(")\n");
 }
 
-FNiagaraDataInterfaceParametersCS* UNiagaraDataInterfaceSkeletalMesh::ConstructComputeParameters() const
-{
-	return new FNiagaraDataInterfaceParametersCS_SkeletalMesh();
-}
-
-
 void UNiagaraDataInterfaceSkeletalMesh::SetSourceComponentFromBlueprints(USkeletalMeshComponent* ComponentToUse)
 {
 	// NOTE: When ChangeId changes the next tick will be skipped and a reset of the per-instance data will be initiated. 
diff --git a/Engine/Plugins/FX/Niagara/Source/Niagara/Private/NiagaraDataInterfaceStaticMesh.cpp b/Engine/Plugins/FX/Niagara/Source/Niagara/Private/NiagaraDataInterfaceStaticMesh.cpp
--- a/Engine/Plugins/FX/Niagara/Source/Niagara/Private/NiagaraDataInterfaceStaticMesh.cpp
+++ b/Engine/Plugins/FX/Niagara/Source/Niagara/Private/NiagaraDataInterfaceStaticMesh.cpp
@@ -73,8 +73,8 @@ void FStaticMeshGpuSpawnBuffer::Initialise(const FStaticMeshLODResources& Res, c
 	SectionRenderData = &Res;
 
 	const uint32 ValidSectionCount = ValidSection.Num();
-	const TArray<float>& Prob = SectionSamplerParam.GetProb();
-	const TArray<int32>& Alias = SectionSamplerParam.GetAlias();
+	const TArray<float, FMemoryImageAllocator>& Prob = SectionSamplerParam.GetProb();
+	const TArray<int32, FMemoryImageAllocator>& Alias = SectionSamplerParam.GetAlias();
 	check(ValidSectionCount == Prob.Num());
 	// Build data that will be uploaded to GPU later from the render thread.
 	// The array contains data used to select regions for uniform particle spawning on them, as well as section triangle ranges.
@@ -444,10 +444,12 @@ static void GetNiagaraDataInterfaceParametersName(FNDIStaticMeshParametersName&
 
 struct FNiagaraDataInterfaceParametersCS_StaticMesh : public FNiagaraDataInterfaceParametersCS
 {
-	virtual void Bind(const FNiagaraDataInterfaceParamRef& ParamRef, const class FShaderParameterMap& ParameterMap) override
+	DECLARE_TYPE_LAYOUT(FNiagaraDataInterfaceParametersCS_StaticMesh, NonVirtual);
+public:
+	void Bind(const FNiagaraDataInterfaceGPUParamInfo& ParameterInfo, const class FShaderParameterMap& ParameterMap)
 	{
 		FNDIStaticMeshParametersName ParamNames;
-		GetNiagaraDataInterfaceParametersName(ParamNames, ParamRef.ParameterInfo.DataInterfaceHLSLSymbol);
+		GetNiagaraDataInterfaceParametersName(ParamNames, ParameterInfo.DataInterfaceHLSLSymbol);
 
 		MeshIndexBuffer.Bind(ParameterMap, *ParamNames.MeshIndexBufferName);
 		MeshVertexBuffer.Bind(ParameterMap, *ParamNames.MeshVertexBufferName);
@@ -466,30 +468,11 @@ struct FNiagaraDataInterfaceParametersCS_StaticMesh : public FNiagaraDataInterfa
 		NumTexCoord.Bind(ParameterMap, *ParamNames.NumTexCoordName);
 	}
 
-	virtual void Serialize(FArchive& Ar)override
-	{
-		Ar << MeshIndexBuffer;
-		Ar << MeshVertexBuffer;
-		Ar << MeshTangentBuffer;
-		Ar << MeshTexCoordBuffer;
-		Ar << MeshColorBuffer;
-		Ar << SectionCount;
-		Ar << MeshSectionBuffer;
-		Ar << MeshTriangleBuffer;
-		Ar << InstanceTransform;
-		Ar << InstanceTransformInverseTransposed;
-		Ar << InstancePrevTransform;
-		Ar << InstanceInvDeltaTime;
-		Ar << InstanceWorldVelocity;
-		Ar << AreaWeightedSampling;
-		Ar << NumTexCoord;
-	}
-
-	virtual void Set(FRHICommandList& RHICmdList, const FNiagaraDataInterfaceSetArgs& Context) const override
+	void Set(FRHICommandList& RHICmdList, const FNiagaraDataInterfaceSetArgs& Context) const
 	{
 		check(IsInRenderingThread());
 
-		FRHIComputeShader* ComputeShaderRHI = Context.Shader->GetComputeShader();
+		FRHIComputeShader* ComputeShaderRHI = Context.Shader.GetComputeShader();
 		
 		{
 			FNiagaraDataInterfaceProxyStaticMesh* InterfaceProxy = static_cast<FNiagaraDataInterfaceProxyStaticMesh*>(Context.DataInterface);
@@ -568,26 +551,29 @@ struct FNiagaraDataInterfaceParametersCS_StaticMesh : public FNiagaraDataInterfa
 		}
 	}
 
-
 private:
-
-	FShaderResourceParameter MeshIndexBuffer;
-	FShaderResourceParameter MeshVertexBuffer;
-	FShaderResourceParameter MeshTangentBuffer;
-	FShaderResourceParameter MeshTexCoordBuffer;
-	FShaderResourceParameter MeshColorBuffer;
-	FShaderResourceParameter MeshSectionBuffer;
-	FShaderResourceParameter MeshTriangleBuffer;
-	FShaderParameter SectionCount;
-	FShaderParameter InstanceTransform;
-	FShaderParameter InstanceTransformInverseTransposed;
-	FShaderParameter InstancePrevTransform;
-	FShaderParameter InstanceInvDeltaTime;
-	FShaderParameter InstanceWorldVelocity;
-	FShaderParameter AreaWeightedSampling;
-	FShaderParameter NumTexCoord;
+	LAYOUT_FIELD(FShaderResourceParameter, MeshIndexBuffer);
+	LAYOUT_FIELD(FShaderResourceParameter, MeshVertexBuffer);
+	LAYOUT_FIELD(FShaderResourceParameter, MeshTangentBuffer);
+	LAYOUT_FIELD(FShaderResourceParameter, MeshTexCoordBuffer);
+	LAYOUT_FIELD(FShaderResourceParameter, MeshColorBuffer);
+	LAYOUT_FIELD(FShaderResourceParameter, MeshSectionBuffer);
+	LAYOUT_FIELD(FShaderResourceParameter, MeshTriangleBuffer);
+	LAYOUT_FIELD(FShaderParameter, SectionCount);
+	LAYOUT_FIELD(FShaderParameter, InstanceTransform);
+	LAYOUT_FIELD(FShaderParameter, InstanceTransformInverseTransposed);
+	LAYOUT_FIELD(FShaderParameter, InstancePrevTransform);
+	LAYOUT_FIELD(FShaderParameter, InstanceInvDeltaTime);
+	LAYOUT_FIELD(FShaderParameter, InstanceWorldVelocity);
+	LAYOUT_FIELD(FShaderParameter, AreaWeightedSampling);
+	LAYOUT_FIELD(FShaderParameter, NumTexCoord);
 };
 
+IMPLEMENT_TYPE_LAYOUT(FNiagaraDataInterfaceParametersCS_StaticMesh);
+
+
+IMPLEMENT_NIAGARA_DI_PARAMETER(UNiagaraDataInterfaceStaticMesh, FNiagaraDataInterfaceParametersCS_StaticMesh);
+
 //////////////////////////////////////////////////////////////////////////
 
 void FNiagaraDataInterfaceProxyStaticMesh::InitializePerInstanceData(const FNiagaraSystemInstanceID& SystemInstance, FStaticMeshGpuSpawnBuffer* MeshGPUSpawnBuffer)
@@ -2360,11 +2346,6 @@ void UNiagaraDataInterfaceStaticMesh::ProvidePerInstanceDataForRenderThread(void
 	DataToPass->PrevTransform = InstanceData->PrevTransform;
 }
 
-FNiagaraDataInterfaceParametersCS* UNiagaraDataInterfaceStaticMesh::ConstructComputeParameters() const 
-{
-	return new FNiagaraDataInterfaceParametersCS_StaticMesh();
-}
-
 //////////////////////////////////////////////////////////////////////////
 
 bool FDynamicVertexColorFilterData::Init(FNDIStaticMesh_InstanceData* Owner)
diff --git a/Engine/Plugins/FX/Niagara/Source/Niagara/Private/NiagaraDataInterfaceTexture.cpp b/Engine/Plugins/FX/Niagara/Source/Niagara/Private/NiagaraDataInterfaceTexture.cpp
--- a/Engine/Plugins/FX/Niagara/Source/Niagara/Private/NiagaraDataInterfaceTexture.cpp
+++ b/Engine/Plugins/FX/Niagara/Source/Niagara/Private/NiagaraDataInterfaceTexture.cpp
@@ -350,10 +350,12 @@ void UNiagaraDataInterfaceTexture::GetParameterDefinitionHLSL(const FNiagaraData
 
 struct FNiagaraDataInterfaceParametersCS_Texture : public FNiagaraDataInterfaceParametersCS
 {
-	virtual void Bind(const FNiagaraDataInterfaceParamRef& ParamRef, const class FShaderParameterMap& ParameterMap) override
+	DECLARE_TYPE_LAYOUT(FNiagaraDataInterfaceParametersCS_Texture, NonVirtual);
+public:
+	void Bind(const FNiagaraDataInterfaceGPUParamInfo& ParameterInfo, const class FShaderParameterMap& ParameterMap)
 	{
-		FString TexName = UNiagaraDataInterfaceTexture::TextureName + ParamRef.ParameterInfo.DataInterfaceHLSLSymbol;
-		FString SampleName = (UNiagaraDataInterfaceTexture::SamplerName + ParamRef.ParameterInfo.DataInterfaceHLSLSymbol);
+		FString TexName = UNiagaraDataInterfaceTexture::TextureName + ParameterInfo.DataInterfaceHLSLSymbol;
+		FString SampleName = (UNiagaraDataInterfaceTexture::SamplerName + ParameterInfo.DataInterfaceHLSLSymbol);
 		TextureParam.Bind(ParameterMap, *TexName);
 		SamplerParam.Bind(ParameterMap, *SampleName);
 		
@@ -367,25 +369,15 @@ struct FNiagaraDataInterfaceParametersCS_Texture : public FNiagaraDataInterfaceP
 			UE_LOG(LogNiagara, Warning, TEXT("Binding failed for FNiagaraDataInterfaceParametersCS_Texture Sampler %s. Was it optimized out?"), *SampleName)
 		}
 
-		Dimensions.Bind(ParameterMap, *(UNiagaraDataInterfaceTexture::DimensionsBaseName + ParamRef.ParameterInfo.DataInterfaceHLSLSymbol));
+		Dimensions.Bind(ParameterMap, *(UNiagaraDataInterfaceTexture::DimensionsBaseName + ParameterInfo.DataInterfaceHLSLSymbol));
 
 	}
 
-	virtual void Serialize(FArchive& Ar)override
-	{
-		Ar << TextureParam;
-		Ar << SamplerParam;
-		if (Ar.IsLoading() == false || Ar.CustomVer(FNiagaraCustomVersion::GUID) >= FNiagaraCustomVersion::TextureDataInterfaceSizeSerialize)
-		{
-			Ar << Dimensions;
-		}
-	}
-
-	virtual void Set(FRHICommandList& RHICmdList, const FNiagaraDataInterfaceSetArgs& Context) const override
+	void Set(FRHICommandList& RHICmdList, const FNiagaraDataInterfaceSetArgs& Context) const
 	{
 		check(IsInRenderingThread());
 
-		FRHIComputeShader* ComputeShaderRHI = Context.Shader->GetComputeShader();
+		FRHIComputeShader* ComputeShaderRHI = Context.Shader.GetComputeShader();
 		FNiagaraDataInterfaceProxyTexture* TextureDI = static_cast<FNiagaraDataInterfaceProxyTexture*>(Context.DataInterface);
 
 		FRHITexture* TextureRHI = (TextureDI && TextureDI->TextureReferenceRHI) ? TextureDI->TextureReferenceRHI->GetReferencedTexture() : nullptr;
@@ -422,13 +414,17 @@ struct FNiagaraDataInterfaceParametersCS_Texture : public FNiagaraDataInterfaceP
 			SetShaderValue(RHICmdList, ComputeShaderRHI, Dimensions, TexDims);
 		}
 	}
-private:
 
-	FShaderResourceParameter TextureParam;
-	FShaderResourceParameter SamplerParam;
-	FShaderParameter Dimensions;
+private:
+	LAYOUT_FIELD(FShaderResourceParameter, TextureParam);
+	LAYOUT_FIELD(FShaderResourceParameter, SamplerParam);
+	LAYOUT_FIELD(FShaderParameter, Dimensions);
 };
 
+IMPLEMENT_TYPE_LAYOUT(FNiagaraDataInterfaceParametersCS_Texture);
+
+IMPLEMENT_NIAGARA_DI_PARAMETER(UNiagaraDataInterfaceTexture, FNiagaraDataInterfaceParametersCS_Texture);
+
 void UNiagaraDataInterfaceTexture::PushToRenderThread()
 {
 	FNiagaraDataInterfaceProxyTexture* RT_Proxy = GetProxyAs<FNiagaraDataInterfaceProxyTexture>();
@@ -463,11 +459,6 @@ void UNiagaraDataInterfaceTexture::PushToRenderThread()
 	);
 }
 
-FNiagaraDataInterfaceParametersCS* UNiagaraDataInterfaceTexture::ConstructComputeParameters()const
-{
-	return new FNiagaraDataInterfaceParametersCS_Texture();
-}
-
 
 void UNiagaraDataInterfaceTexture::SetTexture(UTexture* InTexture)
 {
diff --git a/Engine/Plugins/FX/Niagara/Source/Niagara/Private/NiagaraDataInterfaceVectorField.cpp b/Engine/Plugins/FX/Niagara/Source/Niagara/Private/NiagaraDataInterfaceVectorField.cpp
--- a/Engine/Plugins/FX/Niagara/Source/Niagara/Private/NiagaraDataInterfaceVectorField.cpp
+++ b/Engine/Plugins/FX/Niagara/Source/Niagara/Private/NiagaraDataInterfaceVectorField.cpp
@@ -309,91 +309,68 @@ bool UNiagaraDataInterfaceVectorField::GetFunctionHLSL(const FNiagaraDataInterfa
 	return false;
 }
 
-struct FNiagaraDataInterfaceParametersCS_VectorField : public FNiagaraDataInterfaceParametersCS
+void FNiagaraDataInterfaceParametersCS_VectorField::Bind(const FNiagaraDataInterfaceGPUParamInfo& ParameterInfo, const class FShaderParameterMap& ParameterMap)
 {
-	virtual void Bind(const FNiagaraDataInterfaceParamRef& ParamRef, const class FShaderParameterMap& ParameterMap) override
-	{
-		VectorFieldSampler.Bind(ParameterMap, *(SamplerBaseName + ParamRef.ParameterInfo.DataInterfaceHLSLSymbol));
-		VectorFieldTexture.Bind(ParameterMap, *(TextureBaseName + ParamRef.ParameterInfo.DataInterfaceHLSLSymbol));
-		TilingAxes.Bind(ParameterMap, *(TilingAxesBaseName + ParamRef.ParameterInfo.DataInterfaceHLSLSymbol));
-		Dimensions.Bind(ParameterMap, *(DimensionsBaseName + ParamRef.ParameterInfo.DataInterfaceHLSLSymbol));
-		MinBounds.Bind(ParameterMap, *(MinBoundsBaseName + ParamRef.ParameterInfo.DataInterfaceHLSLSymbol));
-		MaxBounds.Bind(ParameterMap, *(MaxBoundsBaseName + ParamRef.ParameterInfo.DataInterfaceHLSLSymbol));
-	}
+	VectorFieldSampler.Bind(ParameterMap, *(SamplerBaseName + ParameterInfo.DataInterfaceHLSLSymbol));
+	VectorFieldTexture.Bind(ParameterMap, *(TextureBaseName + ParameterInfo.DataInterfaceHLSLSymbol));
+	TilingAxes.Bind(ParameterMap, *(TilingAxesBaseName + ParameterInfo.DataInterfaceHLSLSymbol));
+	Dimensions.Bind(ParameterMap, *(DimensionsBaseName + ParameterInfo.DataInterfaceHLSLSymbol));
+	MinBounds.Bind(ParameterMap, *(MinBoundsBaseName + ParameterInfo.DataInterfaceHLSLSymbol));
+	MaxBounds.Bind(ParameterMap, *(MaxBoundsBaseName + ParameterInfo.DataInterfaceHLSLSymbol));
+}
 
-	virtual void Serialize(FArchive& Ar)override
+void FNiagaraDataInterfaceParametersCS_VectorField::Set(FRHICommandList& RHICmdList, const FNiagaraDataInterfaceSetArgs& Context) const
+{
+	check(IsInRenderingThread());
+
+	// Different sampler states used by the computer shader to sample 3D vector field. 
+	// Encoded as bitflags. To sample: 
+	//     1st bit: X-axis tiling flag
+	//     2nd bit: Y-axis tiling flag
+	//     3rd bit: Z-axis tiling flag
+	static FRHISamplerState* SamplerStates[8] = { nullptr };
+	if (SamplerStates[0] == nullptr)
 	{
-		Ar << VectorFieldSampler;
-		Ar << VectorFieldTexture;
-		Ar << TilingAxes;
-		Ar << Dimensions;
-		Ar << MinBounds;
-		Ar << MaxBounds;
+		SamplerStates[0] = TStaticSamplerState<SF_Bilinear, AM_Clamp, AM_Clamp, AM_Clamp>::GetRHI();
+		SamplerStates[1] = TStaticSamplerState<SF_Bilinear, AM_Wrap,  AM_Clamp, AM_Clamp>::GetRHI();
+		SamplerStates[2] = TStaticSamplerState<SF_Bilinear, AM_Clamp, AM_Wrap,  AM_Clamp>::GetRHI();
+		SamplerStates[3] = TStaticSamplerState<SF_Bilinear, AM_Wrap,  AM_Wrap,  AM_Clamp>::GetRHI();
+		SamplerStates[4] = TStaticSamplerState<SF_Bilinear, AM_Clamp, AM_Clamp, AM_Wrap>::GetRHI();
+		SamplerStates[5] = TStaticSamplerState<SF_Bilinear, AM_Wrap,  AM_Clamp, AM_Wrap>::GetRHI();
+		SamplerStates[6] = TStaticSamplerState<SF_Bilinear, AM_Clamp, AM_Wrap,  AM_Wrap>::GetRHI();
+		SamplerStates[7] = TStaticSamplerState<SF_Bilinear, AM_Wrap,  AM_Wrap,  AM_Wrap>::GetRHI();
 	}
 
-	virtual void Set(FRHICommandList& RHICmdList, const FNiagaraDataInterfaceSetArgs& Context) const override
-	{
-		check(IsInRenderingThread());
-
-		// Different sampler states used by the computer shader to sample 3D vector field. 
-		// Encoded as bitflags. To sample: 
-		//     1st bit: X-axis tiling flag
-		//     2nd bit: Y-axis tiling flag
-		//     3rd bit: Z-axis tiling flag
-		static FRHISamplerState* SamplerStates[8] = { nullptr };
-		if (SamplerStates[0] == nullptr)
-		{
-			SamplerStates[0] = TStaticSamplerState<SF_Bilinear, AM_Clamp, AM_Clamp, AM_Clamp>::GetRHI();
-			SamplerStates[1] = TStaticSamplerState<SF_Bilinear, AM_Wrap,  AM_Clamp, AM_Clamp>::GetRHI();
-			SamplerStates[2] = TStaticSamplerState<SF_Bilinear, AM_Clamp, AM_Wrap,  AM_Clamp>::GetRHI();
-			SamplerStates[3] = TStaticSamplerState<SF_Bilinear, AM_Wrap,  AM_Wrap,  AM_Clamp>::GetRHI();
-			SamplerStates[4] = TStaticSamplerState<SF_Bilinear, AM_Clamp, AM_Clamp, AM_Wrap>::GetRHI();
-			SamplerStates[5] = TStaticSamplerState<SF_Bilinear, AM_Wrap,  AM_Clamp, AM_Wrap>::GetRHI();
-			SamplerStates[6] = TStaticSamplerState<SF_Bilinear, AM_Clamp, AM_Wrap,  AM_Wrap>::GetRHI();
-			SamplerStates[7] = TStaticSamplerState<SF_Bilinear, AM_Wrap,  AM_Wrap,  AM_Wrap>::GetRHI();
-		}
-
-		// Get shader and DI
-		FRHIComputeShader* ComputeShaderRHI = Context.Shader->GetComputeShader();
-		FNiagaraDataInterfaceProxyVectorField* VFDI = static_cast<FNiagaraDataInterfaceProxyVectorField*>(Context.DataInterface);
+	// Get shader and DI
+	FRHIComputeShader* ComputeShaderRHI = Context.Shader.GetComputeShader();
+	FNiagaraDataInterfaceProxyVectorField* VFDI = static_cast<FNiagaraDataInterfaceProxyVectorField*>(Context.DataInterface);
 		
-		// Note: There is a flush in PreEditChange to make sure everything is synced up at this point 
+	// Note: There is a flush in PreEditChange to make sure everything is synced up at this point 
 
-		// Get and set 3D texture handle from the currently bound vector field.
-		if (VFDI->TextureRHI)
-		{
-			SetTextureParameter(RHICmdList, ComputeShaderRHI, VectorFieldTexture, VFDI->TextureRHI);
-		}
-		else
-		{
-			SetTextureParameter(RHICmdList, ComputeShaderRHI, VectorFieldTexture, GBlackVolumeTexture->TextureRHI);
-		}
-		
-		// Get and set sampler state
-		FRHISamplerState* SamplerState = SamplerStates[int(VFDI->bTileX) + 2 * int(VFDI->bTileY) + 4 * int(VFDI->bTileZ)];
-		SetSamplerParameter(RHICmdList, ComputeShaderRHI, VectorFieldSampler, SamplerState);
-
-		//
-		SetShaderValue(RHICmdList, ComputeShaderRHI, TilingAxes, VFDI->GetTilingAxes());
-		SetShaderValue(RHICmdList, ComputeShaderRHI, Dimensions, VFDI->Dimensions);
-		SetShaderValue(RHICmdList, ComputeShaderRHI, MinBounds, VFDI->MinBounds);
-		SetShaderValue(RHICmdList, ComputeShaderRHI, MaxBounds, VFDI->MaxBounds);
+	// Get and set 3D texture handle from the currently bound vector field.
+	if (VFDI->TextureRHI)
+	{
+		SetTextureParameter(RHICmdList, ComputeShaderRHI, VectorFieldTexture, VFDI->TextureRHI);
 	}
+	else
+	{
+		SetTextureParameter(RHICmdList, ComputeShaderRHI, VectorFieldTexture, GBlackVolumeTexture->TextureRHI);
+	}
+		
+	// Get and set sampler state
+	FRHISamplerState* SamplerState = SamplerStates[int(VFDI->bTileX) + 2 * int(VFDI->bTileY) + 4 * int(VFDI->bTileZ)];
+	SetSamplerParameter(RHICmdList, ComputeShaderRHI, VectorFieldSampler, SamplerState);
+
+	//
+	SetShaderValue(RHICmdList, ComputeShaderRHI, TilingAxes, VFDI->GetTilingAxes());
+	SetShaderValue(RHICmdList, ComputeShaderRHI, Dimensions, VFDI->Dimensions);
+	SetShaderValue(RHICmdList, ComputeShaderRHI, MinBounds, VFDI->MinBounds);
+	SetShaderValue(RHICmdList, ComputeShaderRHI, MaxBounds, VFDI->MaxBounds);
+}
 
-private:
-
-	FShaderResourceParameter VectorFieldSampler;
-	FShaderResourceParameter VectorFieldTexture;
-	FShaderParameter TilingAxes;
-	FShaderParameter Dimensions;
-	FShaderParameter MinBounds;
-	FShaderParameter MaxBounds;
-};
+IMPLEMENT_TYPE_LAYOUT(FNiagaraDataInterfaceParametersCS_VectorField);
 
-FNiagaraDataInterfaceParametersCS* UNiagaraDataInterfaceVectorField::ConstructComputeParameters()const
-{
-	return new FNiagaraDataInterfaceParametersCS_VectorField();
-}
+IMPLEMENT_NIAGARA_DI_PARAMETER(UNiagaraDataInterfaceVectorField, FNiagaraDataInterfaceParametersCS_VectorField);
 
 /*--------------------------------------------------------------------------------------------------------------------------*/
 
diff --git a/Engine/Plugins/FX/Niagara/Source/Niagara/Private/NiagaraDataSet.cpp b/Engine/Plugins/FX/Niagara/Source/Niagara/Private/NiagaraDataSet.cpp
--- a/Engine/Plugins/FX/Niagara/Source/Niagara/Private/NiagaraDataSet.cpp
+++ b/Engine/Plugins/FX/Niagara/Source/Niagara/Private/NiagaraDataSet.cpp
@@ -972,7 +972,7 @@ void FNiagaraDataBuffer::SetShaderParams(FNiagaraShader *Shader, FRHICommandList
 	check(IsInRenderingThread());
 
 	const uint32 SafeBufferSize = GetFloatStride() / sizeof(float);
-	FRHIComputeShader* ComputeShader = Shader->GetComputeShader();
+	FRHIComputeShader* ComputeShader = CommandList.GetBoundComputeShader();
 
 	if (bInput)
 	{
@@ -998,11 +998,12 @@ void FNiagaraDataBuffer::SetShaderParams(FNiagaraShader *Shader, FRHICommandList
 void FNiagaraDataBuffer::UnsetShaderParams(FNiagaraShader *Shader, FRHICommandList &RHICmdList)
 {
 	check(IsInRenderingThread());
+	FRHIComputeShader* ShaderRHI = RHICmdList.GetBoundComputeShader();
 
 	if (Shader->FloatOutputBufferParam.IsUAVBound())
 	{
 #if !PLATFORM_PS4
-		Shader->FloatOutputBufferParam.UnsetUAV(RHICmdList, Shader->GetComputeShader());
+		Shader->FloatOutputBufferParam.UnsetUAV(RHICmdList, ShaderRHI);
 #endif
 		//RHICmdList.TransitionResource(EResourceTransitionAccess::EReadable, EResourceTransitionPipeline::EGfxToGfx, CurrDataRender().GetGPUBufferFloat()->UAV);
 	}
@@ -1010,7 +1011,7 @@ void FNiagaraDataBuffer::UnsetShaderParams(FNiagaraShader *Shader, FRHICommandLi
 	if (Shader->IntOutputBufferParam.IsUAVBound())
 	{
 #if !PLATFORM_PS4
-		Shader->IntOutputBufferParam.UnsetUAV(RHICmdList, Shader->GetComputeShader());
+		Shader->IntOutputBufferParam.UnsetUAV(RHICmdList, ShaderRHI);
 #endif
 		//RHICmdList.TransitionResource(EResourceTransitionAccess::EReadable, EResourceTransitionPipeline::EGfxToGfx, CurrDataRender().GetGPUBufferInt()->UAV);
 	}
@@ -1018,7 +1019,7 @@ void FNiagaraDataBuffer::UnsetShaderParams(FNiagaraShader *Shader, FRHICommandLi
 	if (Shader->IDToIndexBufferParam.IsUAVBound())
 	{
 #if !PLATFORM_PS4
-		Shader->IDToIndexBufferParam.UnsetUAV(RHICmdList, Shader->GetComputeShader());
+		Shader->IDToIndexBufferParam.UnsetUAV(RHICmdList, RHICmdList.GetBoundComputeShader());
 #endif
 	}
 }
diff --git a/Engine/Plugins/FX/Niagara/Source/Niagara/Private/NiagaraEmitterInstanceBatcher.cpp b/Engine/Plugins/FX/Niagara/Source/Niagara/Private/NiagaraEmitterInstanceBatcher.cpp
--- a/Engine/Plugins/FX/Niagara/Source/Niagara/Private/NiagaraEmitterInstanceBatcher.cpp
+++ b/Engine/Plugins/FX/Niagara/Source/Niagara/Private/NiagaraEmitterInstanceBatcher.cpp
@@ -188,7 +188,7 @@ void NiagaraEmitterInstanceBatcher::ReleaseTicks()
 	Ticks_RT.Empty(0);
 }
 
-bool NiagaraEmitterInstanceBatcher::ResetDataInterfaces(const FNiagaraGPUSystemTick& Tick, FNiagaraComputeInstanceData *Instance, FRHICommandList &RHICmdList, FNiagaraShader* ComputeShader ) const
+bool NiagaraEmitterInstanceBatcher::ResetDataInterfaces(const FNiagaraGPUSystemTick& Tick, FNiagaraComputeInstanceData *Instance, FRHICommandList &RHICmdList, const FNiagaraShaderRef& ComputeShader ) const
 {
 	bool ValidSpawnStage = true;
 	FNiagaraComputeExecutionContext* Context = Instance->Context;
@@ -199,8 +199,8 @@ bool NiagaraEmitterInstanceBatcher::ResetDataInterfaces(const FNiagaraGPUSystemT
 		uint32 InterfaceIndex = 0;
 		for (FNiagaraDataInterfaceProxy* Interface : Instance->DataInterfaceProxies)
 		{
-			FNiagaraDataInterfaceParamRef& DIParam = ComputeShader->GetDIParameters()[InterfaceIndex];
-			if (DIParam.Parameters)
+			const FNiagaraDataInterfaceParamRef& DIParam = ComputeShader->GetDIParameters()[InterfaceIndex];
+			if (DIParam.Parameters.IsValid())
 			{
 				FNiagaraDataInterfaceSetArgs TmpContext;
 				TmpContext.Shader = ComputeShader;
@@ -237,13 +237,13 @@ FNiagaraDataInterfaceProxy* NiagaraEmitterInstanceBatcher::FindIterationInterfac
 	return IterationInterface;
 }
 
-void NiagaraEmitterInstanceBatcher::PreStageInterface(const FNiagaraGPUSystemTick& Tick, FNiagaraComputeInstanceData *Instance, FRHICommandList &RHICmdList, FNiagaraShader* ComputeShader, const uint32 ShaderStageIndex) const
+void NiagaraEmitterInstanceBatcher::PreStageInterface(const FNiagaraGPUSystemTick& Tick, FNiagaraComputeInstanceData *Instance, FRHICommandList &RHICmdList, const FNiagaraShaderRef& ComputeShader, const uint32 ShaderStageIndex) const
 {
 	uint32 InterfaceIndex = 0;
 	for (FNiagaraDataInterfaceProxy* Interface : Instance->DataInterfaceProxies)
 	{
-		FNiagaraDataInterfaceParamRef& DIParam = ComputeShader->GetDIParameters()[InterfaceIndex];
-		if (DIParam.Parameters)
+		const FNiagaraDataInterfaceParamRef& DIParam = ComputeShader->GetDIParameters()[InterfaceIndex];
+		if (DIParam.Parameters.IsValid())
 		{
 			FNiagaraDataInterfaceSetArgs TmpContext;
 			TmpContext.Shader = ComputeShader;
@@ -259,13 +259,13 @@ void NiagaraEmitterInstanceBatcher::PreStageInterface(const FNiagaraGPUSystemTic
 	}
 }
 
-void NiagaraEmitterInstanceBatcher::PostStageInterface(const FNiagaraGPUSystemTick& Tick, FNiagaraComputeInstanceData *Instance, FRHICommandList &RHICmdList, FNiagaraShader* ComputeShader, const uint32 ShaderStageIndex) const
+void NiagaraEmitterInstanceBatcher::PostStageInterface(const FNiagaraGPUSystemTick& Tick, FNiagaraComputeInstanceData *Instance, FRHICommandList &RHICmdList, const FNiagaraShaderRef& ComputeShader, const uint32 ShaderStageIndex) const
 {
 	uint32 InterfaceIndex = 0;
 	for (FNiagaraDataInterfaceProxy* Interface : Instance->DataInterfaceProxies)
 	{
-		FNiagaraDataInterfaceParamRef& DIParam = ComputeShader->GetDIParameters()[InterfaceIndex];
-		if (DIParam.Parameters)
+		const FNiagaraDataInterfaceParamRef& DIParam = ComputeShader->GetDIParameters()[InterfaceIndex];
+		if (DIParam.Parameters.IsValid())
 		{
 			FNiagaraDataInterfaceSetArgs TmpContext;
 			TmpContext.Shader = ComputeShader;
@@ -281,7 +281,7 @@ void NiagaraEmitterInstanceBatcher::PostStageInterface(const FNiagaraGPUSystemTi
 	}
 }
 
-void NiagaraEmitterInstanceBatcher::DispatchMultipleStages(const FNiagaraGPUSystemTick& Tick, FNiagaraComputeInstanceData *Instance, FRHICommandList &RHICmdList, FRHIUniformBuffer* ViewUniformBuffer, FNiagaraShader* ComputeShader)
+void NiagaraEmitterInstanceBatcher::DispatchMultipleStages(const FNiagaraGPUSystemTick& Tick, FNiagaraComputeInstanceData *Instance, FRHICommandList &RHICmdList, FRHIUniformBuffer* ViewUniformBuffer, const FNiagaraShaderRef& ComputeShader)
 {
 	FNiagaraComputeExecutionContext* Context = Instance->Context;
 
@@ -382,8 +382,8 @@ void NiagaraEmitterInstanceBatcher::ResizeBuffersAndGatherResources(FOverlappabl
 				continue;
 			}
 
-			FNiagaraShader* Shader = Context->GPUScript_RT->GetShader();
-			if ( Shader == nullptr )
+			FNiagaraShaderRef Shader = Context->GPUScript_RT->GetShader();
+			if ( Shader.IsNull() )
 			{
 				continue;
 			}
@@ -499,7 +499,7 @@ void NiagaraEmitterInstanceBatcher::DispatchAllOnCompute(FOverlappableTicks& Ove
 			{
 				FNiagaraComputeInstanceData& Instance = Instances[Index];
 				FNiagaraComputeExecutionContext* Context = Instance.Context;
-				if (Context && Context->GPUScript_RT->GetShader())
+				if (Context && Context->GPUScript_RT->GetShader().IsValid())
 				{
 					if (Context->DebugInfo.IsValid())
 					{
@@ -523,7 +523,7 @@ void NiagaraEmitterInstanceBatcher::DispatchAllOnCompute(FOverlappableTicks& Ove
 		{
 			FNiagaraComputeInstanceData& Instance = Instances[Index];
 			FNiagaraComputeExecutionContext* Context = Instance.Context;
-			if (Context && Context->GPUScript_RT->GetShader())
+			if (Context && Context->GPUScript_RT->GetShader().IsValid())
 			{
 				FNiagaraComputeExecutionContext::TickCounter++;
 
@@ -588,7 +588,7 @@ bool NiagaraEmitterInstanceBatcher::ShouldTickForStage(const FNiagaraGPUSystemTi
 		return TickStage == ETickStage::PostInitViews;
 	}
 
-	FNiagaraShader* ComputeShader = Tick.GetInstanceData()->Context->GPUScript_RT->GetShader();
+	FNiagaraShaderRef ComputeShader = Tick.GetInstanceData()->Context->GPUScript_RT->GetShader();
 	if (ComputeShader->ViewUniformBufferParam.IsBound())
 	{
 		return TickStage == ETickStage::PostOpaqueRender;
@@ -684,8 +684,8 @@ void NiagaraEmitterInstanceBatcher::ExecuteAll(FRHICommandList& RHICmdList, FRHI
 			FNiagaraComputeInstanceData* Data = Tick.GetInstanceData();
 			FNiagaraComputeExecutionContext* Context = Data->Context;
 			// This assumes all emitters fallback to the same FNiagaraShaderScript*.
-			FNiagaraShader* ComputeShader = Context->GPUScript_RT->GetShader();
-			if (!ComputeShader || !ShouldTickForStage(Tick, TickStage))
+			FNiagaraShaderRef ComputeShader = Context->GPUScript_RT->GetShader();
+			if (ComputeShader.IsNull() || !ShouldTickForStage(Tick, TickStage))
 			{
 				continue;
 			}
@@ -982,7 +982,7 @@ void NiagaraEmitterInstanceBatcher::GenerateSortKeys(FRHICommandListImmediate& R
 	FNiagaraSortKeyGenCS::FPermutationDomain PermutationVector;
 	PermutationVector.Set<FNiagaraSortKeyGenCS::FSortUsingMaxPrecision>(EnumHasAnyFlags(Flags, EGPUSortFlags::HighPrecisionKeys));
 	TShaderMapRef<FNiagaraSortKeyGenCS> KeyGenCS(GetGlobalShaderMap(FeatureLevel), PermutationVector);
-	RHICmdList.SetComputeShader(KeyGenCS->GetComputeShader());
+	RHICmdList.SetComputeShader(KeyGenCS.GetComputeShader());
 	KeyGenCS->SetOutput(RHICmdList, KeysUAV, ValuesUAV);
 
 	FRHIUnorderedAccessView* OutputUAVs[] = { KeysUAV, ValuesUAV };
@@ -991,7 +991,7 @@ void NiagaraEmitterInstanceBatcher::GenerateSortKeys(FRHICommandListImmediate& R
 		if (SortInfo.AllocationInfo.SortBatchId == BatchId)
 		{
 			KeyGenCS->SetParameters(RHICmdList, SortInfo, (uint32)SortInfo.AllocationInfo.ElementIndex << KeyGenInfo.ElementKeyShift, SortInfo.AllocationInfo.BufferOffset, KeyGenInfo.SortKeyParams);
-			DispatchComputeShader(RHICmdList, *KeyGenCS, FMath::DivideAndRoundUp(SortInfo.ParticleCount, NIAGARA_KEY_GEN_THREAD_COUNT), 1, 1);
+			DispatchComputeShader(RHICmdList, KeyGenCS, FMath::DivideAndRoundUp(SortInfo.ParticleCount, NIAGARA_KEY_GEN_THREAD_COUNT), 1, 1);
 			// TR-KeyGen : No sync needed between tasks since they update different parts of the data (assuming it's ok if cache lines overlap).
 			RHICmdList.TransitionResources(EResourceTransitionAccess::ERWNoBarrier, EResourceTransitionPipeline::EComputeToCompute, OutputUAVs, UE_ARRAY_COUNT(OutputUAVs));
 		}
@@ -1077,19 +1077,20 @@ void NiagaraEmitterInstanceBatcher::ProcessDebugInfo(FRHICommandList &RHICmdList
 
 /* Set shader parameters for data interfaces
  */
-void NiagaraEmitterInstanceBatcher::SetDataInterfaceParameters(const TArray<FNiagaraDataInterfaceProxy*>& DataInterfaceProxies, FNiagaraShader* Shader, FRHICommandList &RHICmdList, const FNiagaraComputeInstanceData* Instance, const FNiagaraGPUSystemTick& Tick, uint32 ShaderStageIndex) const
+void NiagaraEmitterInstanceBatcher::SetDataInterfaceParameters(const TArray<FNiagaraDataInterfaceProxy*>& DataInterfaceProxies, const FNiagaraShaderRef& Shader, FRHICommandList &RHICmdList, const FNiagaraComputeInstanceData* Instance, const FNiagaraGPUSystemTick& Tick, uint32 ShaderStageIndex) const
 {
 	// set up data interface buffers, as defined by the DIs during compilation
 	//
 
 	// @todo-threadsafety This is a bit gross. Need to rethink this api.
 	const FNiagaraSystemInstanceID& SystemInstance = Tick.SystemInstanceID;
+	const FNiagaraShaderMapPointerTable& PointerTable = Shader.GetPointerTable();
 
 	uint32 InterfaceIndex = 0;
 	for (FNiagaraDataInterfaceProxy* Interface : DataInterfaceProxies)
 	{
-		FNiagaraDataInterfaceParamRef& DIParam = Shader->GetDIParameters()[InterfaceIndex];
-		if (DIParam.Parameters)
+		const FNiagaraDataInterfaceParamRef& DIParam = Shader->GetDIParameters()[InterfaceIndex];
+		if (DIParam.Parameters.IsValid())
 		{
 			FNiagaraDataInterfaceSetArgs Context;
 			Context.Shader = Shader;
@@ -1099,26 +1100,27 @@ void NiagaraEmitterInstanceBatcher::SetDataInterfaceParameters(const TArray<FNia
 			Context.ComputeInstanceData = Instance;
 			Context.ShaderStageIndex = ShaderStageIndex;
 			Context.IsOutputStage = Interface->IsOutputStage(ShaderStageIndex);
-			DIParam.Parameters->Set(RHICmdList, Context);
+			DIParam.DIType.Get(PointerTable.DITypes)->SetParameters(DIParam.Parameters.Get(), RHICmdList, Context);
 		}
 
 		InterfaceIndex++;
 	}
 }
 
-void NiagaraEmitterInstanceBatcher::UnsetDataInterfaceParameters(const TArray<FNiagaraDataInterfaceProxy*> &DataInterfaceProxies, FNiagaraShader* Shader, FRHICommandList &RHICmdList, const FNiagaraComputeInstanceData* Instance, const FNiagaraGPUSystemTick& Tick) const
+void NiagaraEmitterInstanceBatcher::UnsetDataInterfaceParameters(const TArray<FNiagaraDataInterfaceProxy*> &DataInterfaceProxies, const FNiagaraShaderRef& Shader, FRHICommandList &RHICmdList, const FNiagaraComputeInstanceData* Instance, const FNiagaraGPUSystemTick& Tick) const
 {
 	// set up data interface buffers, as defined by the DIs during compilation
 	//
 
 	// @todo-threadsafety This is a bit gross. Need to rethink this api.
 	const FNiagaraSystemInstanceID& SystemInstance = Tick.SystemInstanceID;
+	const FNiagaraShaderMapPointerTable& PointerTable = Shader.GetPointerTable();
 
 	uint32 InterfaceIndex = 0;
 	for (FNiagaraDataInterfaceProxy* Interface : DataInterfaceProxies)
 	{
-		FNiagaraDataInterfaceParamRef& DIParam = Shader->GetDIParameters()[InterfaceIndex];
-		if (DIParam.Parameters)
+		const FNiagaraDataInterfaceParamRef& DIParam = Shader->GetDIParameters()[InterfaceIndex];
+		if (DIParam.Parameters.IsValid())
 		{
 			void* PerInstanceData = nullptr;
 			int32* OffsetFound = nullptr;
@@ -1135,7 +1137,7 @@ void NiagaraEmitterInstanceBatcher::UnsetDataInterfaceParameters(const TArray<FN
 			Context.DataInterface = Interface;
 			Context.SystemInstance = SystemInstance;
 			Context.Batcher = this;
-			DIParam.Parameters->Unset(RHICmdList, Context);
+			DIParam.DIType.Get(PointerTable.DITypes)->UnsetParameters(DIParam.Parameters.Get(), RHICmdList, Context);
 		}
 
 		InterfaceIndex++;
@@ -1157,7 +1159,7 @@ static void SetConstantBuffer(FRHICommandList &RHICmdList, FRHIComputeShader* Co
 
 /* Kick off a simulation/spawn run
  */
-void NiagaraEmitterInstanceBatcher::Run(const FNiagaraGPUSystemTick& Tick, const FNiagaraComputeInstanceData* Instance, uint32 UpdateStartInstance, const uint32 TotalNumInstances, FNiagaraShader* Shader,
+void NiagaraEmitterInstanceBatcher::Run(const FNiagaraGPUSystemTick& Tick, const FNiagaraComputeInstanceData* Instance, uint32 UpdateStartInstance, const uint32 TotalNumInstances, const FNiagaraShaderRef& Shader,
 	FRHICommandList &RHICmdList, FRHIUniformBuffer* ViewUniformBuffer, const FNiagaraGpuSpawnInfo& SpawnInfo, bool bCopyBeforeStart, uint32 DefaultShaderStageIndex, uint32 ShaderStageIndex, FNiagaraDataInterfaceProxy *IterationInterface, bool HasRunParticleStage)
 {
 	FNiagaraComputeExecutionContext* Context = Instance->Context;
@@ -1187,11 +1189,11 @@ void NiagaraEmitterInstanceBatcher::Run(const FNiagaraGPUSystemTick& Tick, const
 		InstancesToSpawnThisFrame = 0;
 	}
 
-	FRHIComputeShader* ComputeShader = Shader->GetComputeShader();
+	FRHIComputeShader* ComputeShader = Shader.GetComputeShader();
 	DestinationData.SetNumSpawnedInstances(InstancesToSpawnThisFrame);
 	DestinationData.SetIDAcquireTag(FNiagaraComputeExecutionContext::TickCounter);
 
-	RHICmdList.SetComputeShader(Shader->GetComputeShader());
+	RHICmdList.SetComputeShader(ComputeShader);
 
 	// #todo(dmp): clean up this logic for shader stages on first frame
 	SetShaderValue(RHICmdList, ComputeShader, Shader->SimStartParam, Tick.bNeedsReset ? 1U : 0U);
@@ -1207,9 +1209,9 @@ void NiagaraEmitterInstanceBatcher::Run(const FNiagaraGPUSystemTick& Tick, const
 	// set the shader and data set params 
 	//
 	const bool bNeedsPersistentIDs = Context->MainDataSet->GetNeedsPersistentIDs();
-	SetSRVParameter(RHICmdList, Shader->GetComputeShader(), Shader->FreeIDBufferParam, bNeedsPersistentIDs ? Context->MainDataSet->GetGPUFreeIDs().SRV : FNiagaraRenderer::GetDummyIntBuffer().SRV);
-	CurrentData.SetShaderParams(Shader, RHICmdList, true);
-	DestinationData.SetShaderParams(Shader, RHICmdList, false);
+	SetSRVParameter(RHICmdList, Shader.GetComputeShader(), Shader->FreeIDBufferParam, bNeedsPersistentIDs ? Context->MainDataSet->GetGPUFreeIDs().SRV : FNiagaraRenderer::GetDummyIntBuffer().SRV);
+	CurrentData.SetShaderParams(Shader.GetShader(), RHICmdList, true);
+	DestinationData.SetShaderParams(Shader.GetShader(), RHICmdList, false);
 
 	// set the index buffer uav
 	//
@@ -1277,8 +1279,8 @@ void NiagaraEmitterInstanceBatcher::Run(const FNiagaraGPUSystemTick& Tick, const
 	// mess up particle state
 	if (IterationInterface)
 	{
-		CurrentData.UnsetShaderParams(Shader, RHICmdList);
-		DestinationData.UnsetShaderParams(Shader, RHICmdList);
+		CurrentData.UnsetShaderParams(Shader.GetShader(), RHICmdList);
+		DestinationData.UnsetShaderParams(Shader.GetShader(), RHICmdList);
 	}
 
 	//UE_LOG(LogNiagara, Log, TEXT("Num Instance : %d | Num Group : %d | Spawned Istance : %d | Start Instance : %d | Num Indices : %d | Stage Index : %d"), 
@@ -1287,7 +1289,7 @@ void NiagaraEmitterInstanceBatcher::Run(const FNiagaraGPUSystemTick& Tick, const
 	// Dispatch, if anything needs to be done
 	if (TotalNumInstances)
 	{
-		DispatchComputeShader(RHICmdList, Shader, NumThreadGroups, 1, 1);
+		DispatchComputeShader(RHICmdList, Shader.GetShader(), NumThreadGroups, 1, 1);
 	}
 
 	// reset iteration count
@@ -1338,8 +1340,8 @@ void NiagaraEmitterInstanceBatcher::Run(const FNiagaraGPUSystemTick& Tick, const
 	// Unset UAV parameters and transition resources (TODO: resource transition should be moved to the renderer)
 	// 
 	UnsetDataInterfaceParameters(DataInterfaceProxies, Shader, RHICmdList, Instance, Tick);
-	CurrentData.UnsetShaderParams(Shader, RHICmdList);
-	DestinationData.UnsetShaderParams(Shader, RHICmdList);
+	CurrentData.UnsetShaderParams(Shader.GetShader(), RHICmdList);
+	DestinationData.UnsetShaderParams(Shader.GetShader(), RHICmdList);
 	Shader->InstanceCountsParam.UnsetUAV(RHICmdList, ComputeShader);
 }
 
diff --git a/Engine/Plugins/FX/Niagara/Source/Niagara/Private/NiagaraRendererMeshes.cpp b/Engine/Plugins/FX/Niagara/Source/Niagara/Private/NiagaraRendererMeshes.cpp
--- a/Engine/Plugins/FX/Niagara/Source/Niagara/Private/NiagaraRendererMeshes.cpp
+++ b/Engine/Plugins/FX/Niagara/Source/Niagara/Private/NiagaraRendererMeshes.cpp
@@ -734,7 +734,7 @@ int FNiagaraRendererMeshes::GetDynamicDataSize()const
 
 bool FNiagaraRendererMeshes::IsMaterialValid(UMaterialInterface* Mat)const
 {
-	return Mat && Mat->CheckMaterialUsage(MATUSAGE_NiagaraMeshParticles);
+	return Mat && Mat->CheckMaterialUsage_Concurrent(MATUSAGE_NiagaraMeshParticles);
 }
 
 
diff --git a/Engine/Plugins/FX/Niagara/Source/Niagara/Private/NiagaraRendererRibbons.cpp b/Engine/Plugins/FX/Niagara/Source/Niagara/Private/NiagaraRendererRibbons.cpp
--- a/Engine/Plugins/FX/Niagara/Source/Niagara/Private/NiagaraRendererRibbons.cpp
+++ b/Engine/Plugins/FX/Niagara/Source/Niagara/Private/NiagaraRendererRibbons.cpp
@@ -731,7 +731,7 @@ void FNiagaraRendererRibbons::AddDynamicParam(TArray<FNiagaraRibbonVertexDynamic
 
 bool FNiagaraRendererRibbons::IsMaterialValid(UMaterialInterface* Mat)const
 {
-	return Mat && Mat->CheckMaterialUsage(MATUSAGE_NiagaraRibbons);
+	return Mat && Mat->CheckMaterialUsage_Concurrent(MATUSAGE_NiagaraRibbons);
 }
 
 
diff --git a/Engine/Plugins/FX/Niagara/Source/Niagara/Private/NiagaraRendererSprites.cpp b/Engine/Plugins/FX/Niagara/Source/Niagara/Private/NiagaraRendererSprites.cpp
--- a/Engine/Plugins/FX/Niagara/Source/Niagara/Private/NiagaraRendererSprites.cpp
+++ b/Engine/Plugins/FX/Niagara/Source/Niagara/Private/NiagaraRendererSprites.cpp
@@ -712,5 +712,5 @@ int FNiagaraRendererSprites::GetDynamicDataSize()const
 
 bool FNiagaraRendererSprites::IsMaterialValid(UMaterialInterface* Mat)const
 {
-	return Mat && Mat->CheckMaterialUsage(MATUSAGE_NiagaraSprites);
+	return Mat && Mat->CheckMaterialUsage_Concurrent(MATUSAGE_NiagaraSprites);
 }
\ No newline at end of file
diff --git a/Engine/Plugins/FX/Niagara/Source/Niagara/Private/NiagaraScript.cpp b/Engine/Plugins/FX/Niagara/Source/Niagara/Private/NiagaraScript.cpp
--- a/Engine/Plugins/FX/Niagara/Source/Niagara/Private/NiagaraScript.cpp
+++ b/Engine/Plugins/FX/Niagara/Source/Niagara/Private/NiagaraScript.cpp
@@ -1845,8 +1845,8 @@ NIAGARA_API bool UNiagaraScript::IsScriptCompilationPending(bool bGPUScript) con
 {
 	if (bGPUScript)
 	{
-		FNiagaraShader *Shader = ScriptResource.GetShaderGameThread();
-		if (Shader)
+		FNiagaraShaderRef Shader = ScriptResource.GetShaderGameThread();
+		if (Shader.IsValid())
 		{
 			return false;
 		}
@@ -1866,8 +1866,8 @@ NIAGARA_API bool UNiagaraScript::DidScriptCompilationSucceed(bool bGPUScript) co
 {
 	if (bGPUScript)
 	{
-		FNiagaraShader *Shader = ScriptResource.GetShaderGameThread();
-		if (Shader)
+		FNiagaraShaderRef Shader = ScriptResource.GetShaderGameThread();
+		if (Shader.IsValid())
 		{
 			return true;
 		}
@@ -1942,8 +1942,6 @@ void ProcessSerializedShaderMaps(UNiagaraScript* Owner, TArray<FNiagaraShaderScr
 
 	for (FNiagaraShaderScript& LoadedResource : LoadedResources)
 	{
-		LoadedResource.RegisterShaderMap();
-
 		FNiagaraShaderMap* LoadedShaderMap = LoadedResource.GetGameThreadShaderMap();
 		if (LoadedShaderMap && LoadedShaderMap->GetShaderPlatform() == GMaxRHIShaderPlatform)
 		{
@@ -1956,7 +1954,7 @@ void ProcessSerializedShaderMaps(UNiagaraScript* Owner, TArray<FNiagaraShaderScr
 			}
 
 			OutScriptResourcesLoaded[LoadedFeatureLevel]->SetShaderMap(LoadedShaderMap);
-			OutResourceForCurrentPlatform.SetDataInterfaceParamInfo(LoadedResource.GetShaderGameThread()->GetDIParameters());
+			OutResourceForCurrentPlatform.SetDataInterfaceParamInfo(Owner->GetVMExecutableData().DIParamInfo);
 
 			break;
 		}
diff --git a/Engine/Plugins/FX/Niagara/Source/Niagara/Private/NiagaraScriptExecutionContext.cpp b/Engine/Plugins/FX/Niagara/Source/Niagara/Private/NiagaraScriptExecutionContext.cpp
--- a/Engine/Plugins/FX/Niagara/Source/Niagara/Private/NiagaraScriptExecutionContext.cpp
+++ b/Engine/Plugins/FX/Niagara/Source/Niagara/Private/NiagaraScriptExecutionContext.cpp
@@ -479,18 +479,22 @@ void FNiagaraComputeExecutionContext::InitParams(UNiagaraScript* InGPUComputeScr
 	HasInterpolationParameters = GPUScript && GPUScript->GetComputedVMCompilationId().HasInterpolatedParameters();
 	
 #if DO_CHECK
-	FNiagaraShader *Shader = InGPUComputeScript->GetRenderThreadScript()->GetShaderGameThread();
-	DIParamInfo.Empty();
-	if (Shader)
+	FNiagaraShaderRef Shader = InGPUComputeScript->GetRenderThreadScript()->GetShaderGameThread();
+	if (Shader.IsValid())
 	{
-		for (FNiagaraDataInterfaceParamRef& DIParams : Shader->GetDIParameters())
+		DIClassNames.Empty(Shader->GetDIParameters().Num());
+		for (const FNiagaraDataInterfaceParamRef& DIParams : Shader->GetDIParameters())
 		{
-			DIParamInfo.Add(DIParams.ParameterInfo);
+			DIClassNames.Add(DIParams.DIType.Get(Shader.GetPointerTable().DITypes)->GetClass()->GetName());
 		}
 	}
 	else
 	{
-		DIParamInfo = InGPUComputeScript->GetRenderThreadScript()->GetDataInterfaceParamInfo();
+		DIClassNames.Empty(InGPUComputeScript->GetRenderThreadScript()->GetDataInterfaceParamInfo().Num());
+		for (const FNiagaraDataInterfaceGPUParamInfo& DIParams : InGPUComputeScript->GetRenderThreadScript()->GetDataInterfaceParamInfo())
+		{
+			DIClassNames.Add(DIParams.DIClassName);
+		}
 	}
 #endif
 }
@@ -507,19 +511,19 @@ bool FNiagaraComputeExecutionContext::Tick(FNiagaraSystemInstance* ParentSystemI
 #if DO_CHECK
 		const TArray<UNiagaraDataInterface*> &DataInterfaces = CombinedParamStore.GetDataInterfaces();
 		// We must make sure that the data interfaces match up between the original script values and our overrides...
-		if (DIParamInfo.Num() != DataInterfaces.Num())
+		if (DIClassNames.Num() != DataInterfaces.Num())
 		{
 			UE_LOG(LogNiagara, Warning, TEXT("Mismatch between Niagara GPU Execution Context data interfaces and those in its script!"));
 			return false;
 		}
 
-		for (int32 i = 0; i < DIParamInfo.Num(); ++i)
+		for (int32 i = 0; i < DIClassNames.Num(); ++i)
 		{
 			FString UsedClassName = DataInterfaces[i]->GetClass()->GetName();
-			if (DIParamInfo[i].DIClassName != UsedClassName)
+			if (DIClassNames[i] != UsedClassName)
 			{
 				UE_LOG(LogNiagara, Warning, TEXT("Mismatched class between Niagara GPU Execution Context data interfaces and those in its script!\nIndex:%d\nShader:%s\nScript:%s")
-					, i, *DIParamInfo[i].DIClassName, *UsedClassName);
+					, i, *DIClassNames[i], *UsedClassName);
 			}
 		}
 #endif
diff --git a/Engine/Plugins/FX/Niagara/Source/Niagara/Private/NiagaraSystem.cpp b/Engine/Plugins/FX/Niagara/Source/Niagara/Private/NiagaraSystem.cpp
--- a/Engine/Plugins/FX/Niagara/Source/Niagara/Private/NiagaraSystem.cpp
+++ b/Engine/Plugins/FX/Niagara/Source/Niagara/Private/NiagaraSystem.cpp
@@ -1290,9 +1290,9 @@ void UNiagaraSystem::InitSystemCompiledData()
 	CreateDataSetCompiledData(SystemCompiledData.DataSetCompiledData, GetSystemUpdateScript()->GetVMExecutableData().Attributes);
 
 	FNiagaraParameters* EngineParamsSpawn = GetSystemSpawnScript()->GetVMExecutableData().DataSetToParameters.Find(TEXT("Engine"));
-	CreateDataSetCompiledData(SystemCompiledData.SpawnInstanceParamsDataSetCompiledData, EngineParamsSpawn ? EngineParamsSpawn->Parameters : TArrayView<FNiagaraVariable>());
+	CreateDataSetCompiledData(SystemCompiledData.SpawnInstanceParamsDataSetCompiledData, EngineParamsSpawn ? TArrayView<FNiagaraVariable>(EngineParamsSpawn->Parameters) : TArrayView<FNiagaraVariable>());
 	FNiagaraParameters* EngineParamsUpdate = GetSystemUpdateScript()->GetVMExecutableData().DataSetToParameters.Find(TEXT("Engine"));
-	CreateDataSetCompiledData(SystemCompiledData.UpdateInstanceParamsDataSetCompiledData, EngineParamsUpdate ? EngineParamsUpdate->Parameters : TArrayView<FNiagaraVariable>());
+	CreateDataSetCompiledData(SystemCompiledData.UpdateInstanceParamsDataSetCompiledData, EngineParamsUpdate ? TArrayView<FNiagaraVariable>(EngineParamsUpdate->Parameters) : TArrayView<FNiagaraVariable>());
 
 	// create the bindings to be used with our constant buffers; geenrating the offsets to/from the data sets; we need
 	// editor data to build these bindings because of the constant buffer structs only having their variable definitions
diff --git a/Engine/Plugins/FX/Niagara/Source/Niagara/Public/NiagaraComponent.h b/Engine/Plugins/FX/Niagara/Source/Niagara/Public/NiagaraComponent.h
--- a/Engine/Plugins/FX/Niagara/Source/Niagara/Public/NiagaraComponent.h
+++ b/Engine/Plugins/FX/Niagara/Source/Niagara/Public/NiagaraComponent.h
@@ -120,7 +120,7 @@ protected:
 	virtual void OnRegister() override;
 	virtual void OnUnregister() override;
 	virtual void OnEndOfFrameUpdateDuringTick() override;
-	virtual void CreateRenderState_Concurrent() override;
+	virtual void CreateRenderState_Concurrent(FRegisterComponentContext* Context) override;
 	virtual void SendRenderDynamicData_Concurrent() override;
 	virtual void BeginDestroy() override;
 	//virtual void OnAttachmentChanged() override;
diff --git a/Engine/Plugins/FX/Niagara/Source/NiagaraCore/Private/NiagaraCompileHash.cpp b/Engine/Plugins/FX/Niagara/Source/NiagaraCore/Private/NiagaraCompileHash.cpp
--- a/Engine/Plugins/FX/Niagara/Source/NiagaraCore/Private/NiagaraCompileHash.cpp
+++ b/Engine/Plugins/FX/Niagara/Source/NiagaraCore/Private/NiagaraCompileHash.cpp
@@ -1,6 +1,7 @@
 // Copyright Epic Games, Inc. All Rights Reserved.
 
 #include "NiagaraCompileHash.h"
+#include "Misc/SecureHash.h"
 
 const uint32 FNiagaraCompileHash::HashSize = 20;
 
@@ -14,6 +15,25 @@ bool FNiagaraCompileHash::operator!=(const FNiagaraCompileHash& Other) const
 	return DataHash != Other.DataHash;
 }
 
+bool FNiagaraCompileHash::operator==(const FSHAHash& Other) const
+{
+	if (DataHash.Num() != HashSize)
+	{
+		return false;
+	}
+	return FMemory::Memcmp(DataHash.GetData(), Other.Hash, HashSize) == 0;
+}
+
+bool FNiagaraCompileHash::ToSHAHash(FSHAHash& OutHash) const
+{
+	if (DataHash.Num() == HashSize)
+	{
+		FMemory::Memcpy(OutHash.Hash, DataHash.GetData(), HashSize);
+		return true;
+	}
+	return false;
+}
+
 bool FNiagaraCompileHash::IsValid() const
 {
 	return DataHash.Num() == HashSize;
@@ -53,4 +73,4 @@ FString FNiagaraCompileHash::ToString() const
 FArchive& operator<<(FArchive& Ar, FNiagaraCompileHash& Id)
 {
 	return Ar << Id.DataHash;
-}
\ No newline at end of file
+}
diff --git a/Engine/Plugins/FX/Niagara/Source/NiagaraCore/Private/NiagaraCoreModule.cpp b/Engine/Plugins/FX/Niagara/Source/NiagaraCore/Private/NiagaraCoreModule.cpp
--- a/Engine/Plugins/FX/Niagara/Source/NiagaraCore/Private/NiagaraCoreModule.cpp
+++ b/Engine/Plugins/FX/Niagara/Source/NiagaraCore/Private/NiagaraCoreModule.cpp
@@ -6,6 +6,9 @@
 
 IMPLEMENT_MODULE(INiagaraCoreModule, NiagaraCore);
 
+IMPLEMENT_TYPE_LAYOUT(FNiagaraDataInterfaceParametersCS);
+
+IMPLEMENT_EXPORTED_INTRINSIC_TYPE_LAYOUT(TIndexedPtr<UNiagaraDataInterfaceBase>);
 
 UNiagaraDataInterfaceBase::UNiagaraDataInterfaceBase(class FObjectInitializer const & Initializer)
 {
diff --git a/Engine/Plugins/FX/Niagara/Source/NiagaraCore/Private/NiagaraCustomVersion.cpp b/Engine/Plugins/FX/Niagara/Source/NiagaraCore/Private/NiagaraCustomVersion.cpp
--- a/Engine/Plugins/FX/Niagara/Source/NiagaraCore/Private/NiagaraCustomVersion.cpp
+++ b/Engine/Plugins/FX/Niagara/Source/NiagaraCore/Private/NiagaraCustomVersion.cpp
@@ -8,4 +8,4 @@ const FGuid FNiagaraCustomVersion::GUID(0xFCF57AFA, 0x50764283, 0xB9A9E658, 0xFF
 // Register the custom version with core
 FCustomVersionRegistration GRegisterNiagaraCustomVersion(FNiagaraCustomVersion::GUID, FNiagaraCustomVersion::LatestVersion, TEXT("NiagaraVer"));
 
-const FGuid FNiagaraCustomVersion::LatestScriptCompileVersion(0x6E2BA7D2, 0xEDAC9A43, 0xB0972E7C, 0x95A69C9C);
+const FGuid FNiagaraCustomVersion::LatestScriptCompileVersion(0x6B79BB03, 0x495F4956, 0x9F7C5D3E, 0x42B9AC36);
\ No newline at end of file
diff --git a/Engine/Plugins/FX/Niagara/Source/NiagaraCore/Public/NiagaraCompileHash.h b/Engine/Plugins/FX/Niagara/Source/NiagaraCore/Public/NiagaraCompileHash.h
--- a/Engine/Plugins/FX/Niagara/Source/NiagaraCore/Public/NiagaraCompileHash.h
+++ b/Engine/Plugins/FX/Niagara/Source/NiagaraCore/Public/NiagaraCompileHash.h
@@ -31,8 +31,11 @@ struct NIAGARACORE_API FNiagaraCompileHash
 	}
 
 	bool operator==(const FNiagaraCompileHash& Other) const;
-
 	bool operator!=(const FNiagaraCompileHash& Other) const;
+	bool operator==(const FSHAHash& Other) const;
+	inline bool operator!=(const FSHAHash& Other) const { return !operator==(Other); }
+
+	bool ToSHAHash(FSHAHash& OutHash) const;
 
 	bool IsValid() const;
 
@@ -49,4 +52,13 @@ struct NIAGARACORE_API FNiagaraCompileHash
 private:
 	UPROPERTY()
 	TArray<uint8> DataHash;
-};
\ No newline at end of file
+};
+
+inline bool operator==(const FSHAHash& Lhs, const FNiagaraCompileHash& Rhs)
+{
+	return Rhs.operator==(Lhs);
+}
+inline bool operator!=(const FSHAHash& Lhs, const FNiagaraCompileHash& Rhs)
+{
+	return !operator==(Lhs, Rhs);
+}
diff --git a/Engine/Plugins/FX/Niagara/Source/NiagaraCore/Public/NiagaraDataInterfaceBase.h b/Engine/Plugins/FX/Niagara/Source/NiagaraCore/Public/NiagaraDataInterfaceBase.h
--- a/Engine/Plugins/FX/Niagara/Source/NiagaraCore/Public/NiagaraDataInterfaceBase.h
+++ b/Engine/Plugins/FX/Niagara/Source/NiagaraCore/Public/NiagaraDataInterfaceBase.h
@@ -3,6 +3,8 @@
 
 #include "NiagaraCore.h"
 #include "NiagaraMergeable.h"
+#include "Shader.h"
+#include "Serialization/MemoryImage.h"
 #include "NiagaraDataInterfaceBase.generated.h"
 
 class INiagaraCompiler;
@@ -11,16 +13,20 @@ class UCurveLinearColor;
 class UCurveFloat;
 class FNiagaraSystemInstance;
 class FNiagaraShader;
-struct FNiagaraDataInterfaceParamRef;
+class FNiagaraShaderMapPointerTable;
+class UNiagaraDataInterfaceBase;
+struct FNiagaraDataInterfaceGPUParamInfo;
 class FRHICommandList;
 
 struct FNiagaraDataInterfaceProxy;
 class NiagaraEmitterInstanceBatcher;
 struct FNiagaraComputeInstanceData;
 
+DECLARE_EXPORTED_TEMPLATE_INTRINSIC_TYPE_LAYOUT(template<>, TIndexedPtr<UNiagaraDataInterfaceBase>, NIAGARACORE_API);
+
 struct FNiagaraDataInterfaceSetArgs
 {
-	FNiagaraShader* Shader;
+	TShaderRefBase<FNiagaraShader, FNiagaraShaderMapPointerTable> Shader;
 	FNiagaraDataInterfaceProxy* DataInterface;
 	FNiagaraSystemInstanceID SystemInstance;
 	const NiagaraEmitterInstanceBatcher* Batcher;
@@ -31,16 +37,19 @@ struct FNiagaraDataInterfaceSetArgs
 };
 
 /**
-* An interface to the parameter bindings for the data interface used by a Niagara compute shader.
-*/
+ * An interface to the parameter bindings for the data interface used by a Niagara compute shader.
+ * This is not using virtual methods, but derived classes may still override the methods listed below.
+ * Overriden methods will be correctly called via indirection through UNiagaraDataInterfaceBase's vtable
+ */
 struct FNiagaraDataInterfaceParametersCS
 {
-	virtual ~FNiagaraDataInterfaceParametersCS() {}
-	virtual void Bind(const FNiagaraDataInterfaceParamRef& ParamRef, const class FShaderParameterMap& ParameterMap) {}
-	virtual void Serialize(FArchive& Ar) { }
-	virtual void Set(FRHICommandList& RHICmdList, const FNiagaraDataInterfaceSetArgs& Context) const {}
-	virtual void Unset(FRHICommandList& RHICmdList, const FNiagaraDataInterfaceSetArgs& Context) const {}
+	DECLARE_EXPORTED_TYPE_LAYOUT(FNiagaraDataInterfaceParametersCS, NIAGARACORE_API, NonVirtual);
+public:
+	void Bind(const FNiagaraDataInterfaceGPUParamInfo& ParameterInfo, const class FShaderParameterMap& ParameterMap) {}
+	void Set(FRHICommandList& RHICmdList, const FNiagaraDataInterfaceSetArgs& Context) const {}
+	void Unset(FRHICommandList& RHICmdList, const FNiagaraDataInterfaceSetArgs& Context) const {}
 
+	LAYOUT_FIELD(TIndexedPtr<UNiagaraDataInterfaceBase>, DIType);
 };
 
 //////////////////////////////////////////////////////////////////////////
@@ -54,10 +63,29 @@ class NIAGARACORE_API UNiagaraDataInterfaceBase : public UNiagaraMergeable
 public:
 	
 	/** Constructs the correct CS parameter type for this DI (if any). The object type returned by this can only vary by class and not per object data. */
-	virtual FNiagaraDataInterfaceParametersCS* ConstructComputeParameters()const
-	{
-		return nullptr;
-	}
+	virtual FNiagaraDataInterfaceParametersCS* ConstructComputeParameters() const { return nullptr; }
+	virtual const FTypeLayoutDesc* GetComputeParametersTypeDesc() const { return nullptr; }
+
+	/** Methods that operate on an instance of type FNiagaraDataInterfaceParametersCS*, created by the above ConstructComputeParameters() method */
+	virtual void BindParameters(FNiagaraDataInterfaceParametersCS* Base, const FNiagaraDataInterfaceGPUParamInfo& ParameterInfo, const class FShaderParameterMap& ParameterMap) {}
+	virtual void SetParameters(const FNiagaraDataInterfaceParametersCS* Base, FRHICommandList& RHICmdList, const FNiagaraDataInterfaceSetArgs& Context) const {}
+	virtual void UnsetParameters(const FNiagaraDataInterfaceParametersCS* Base, FRHICommandList& RHICmdList, const FNiagaraDataInterfaceSetArgs& Context) const {}
 };
 
+/** This goes in class declaration for UNiagaraDataInterfaceBase-derived types, that need custom parameter type */
+#define DECLARE_NIAGARA_DI_PARAMETER() \
+	virtual FNiagaraDataInterfaceParametersCS* ConstructComputeParameters() const override; \
+	virtual const FTypeLayoutDesc* GetComputeParametersTypeDesc() const override; \
+	virtual void BindParameters(FNiagaraDataInterfaceParametersCS* Base, const FNiagaraDataInterfaceGPUParamInfo& ParameterInfo, const class FShaderParameterMap& ParameterMap) override; \
+	virtual void SetParameters(const FNiagaraDataInterfaceParametersCS* Base, FRHICommandList& RHICmdList, const FNiagaraDataInterfaceSetArgs& Context) const override; \
+	virtual void UnsetParameters(const FNiagaraDataInterfaceParametersCS* Base, FRHICommandList& RHICmdList, const FNiagaraDataInterfaceSetArgs& Context) const override
+
+/** This goes in cpp file matched with class declartion using DECLARE_NIAGARA_DI_PARAMETER() */
+#define IMPLEMENT_NIAGARA_DI_PARAMETER(T, ParameterType) \
+	static_assert(ParameterType::InterfaceType == ETypeLayoutInterface::NonVirtual, "DI ParameterType must be non-virtual"); \
+	FNiagaraDataInterfaceParametersCS* T::ConstructComputeParameters() const { return new ParameterType(); } \
+	const FTypeLayoutDesc* T::GetComputeParametersTypeDesc() const { return &StaticGetTypeLayoutDesc<ParameterType>(); } \
+	void T::BindParameters(FNiagaraDataInterfaceParametersCS* Base, const FNiagaraDataInterfaceGPUParamInfo& ParameterInfo, const class FShaderParameterMap& ParameterMap) { static_cast<ParameterType*>(Base)->Bind(ParameterInfo, ParameterMap); } \
+	void T::SetParameters(const FNiagaraDataInterfaceParametersCS* Base, FRHICommandList& RHICmdList, const FNiagaraDataInterfaceSetArgs& Context) const { static_cast<const ParameterType*>(Base)->Set(RHICmdList, Context); } \
+	void T::UnsetParameters(const FNiagaraDataInterfaceParametersCS* Base, FRHICommandList& RHICmdList, const FNiagaraDataInterfaceSetArgs& Context) const { static_cast<const ParameterType*>(Base)->Unset(RHICmdList, Context); }
 
diff --git a/Engine/Plugins/FX/Niagara/Source/NiagaraEditor/Private/Widgets/SNiagaraSystemViewport.cpp b/Engine/Plugins/FX/Niagara/Source/NiagaraEditor/Private/Widgets/SNiagaraSystemViewport.cpp
--- a/Engine/Plugins/FX/Niagara/Source/NiagaraEditor/Private/Widgets/SNiagaraSystemViewport.cpp
+++ b/Engine/Plugins/FX/Niagara/Source/NiagaraEditor/Private/Widgets/SNiagaraSystemViewport.cpp
@@ -190,8 +190,8 @@ void FNiagaraSystemViewportClient::DrawInstructionCounts(UNiagaraSystem* Particl
 			uint32 NumInstructions = 0;
 			if (Script->GetUsage() == ENiagaraScriptUsage::ParticleGPUComputeScript)
 			{
-				FNiagaraShader* Shader = Script->GetRenderThreadScript()->GetShaderGameThread();
-				if (Shader != nullptr)
+				FNiagaraShaderRef Shader = Script->GetRenderThreadScript()->GetShaderGameThread();
+				if (Shader.IsValid())
 				{
 					NumInstructions = Shader->GetNumInstructions();
 				}
