Copying Private-LoadTimes-4.24 stream to Main. Biggest changes are in Materials/Shader memory freezing.

diff --git a/Engine/Plugins/FX/Niagara/Source/NiagaraShader/Private/NiagaraShader.cpp b/Engine/Plugins/FX/Niagara/Source/NiagaraShader/Private/NiagaraShader.cpp
--- a/Engine/Plugins/FX/Niagara/Source/NiagaraShader/Private/NiagaraShader.cpp
+++ b/Engine/Plugins/FX/Niagara/Source/NiagaraShader/Private/NiagaraShader.cpp
@@ -85,7 +85,56 @@ void UpdateNiagaraShaderCompilingStats(const FNiagaraShaderScript* Script)
 	INC_DWORD_STAT_BY(STAT_ShaderCompiling_NumTotalNiagaraShaders,1);
 }
 
+int32 FNiagaraShaderMapPointerTable::AddIndexedPointer(const FTypeLayoutDesc& TypeDesc, void* Ptr)
+{
+	int32 Index = INDEX_NONE;
+	if (DITypes.TryAddIndexedPtr(TypeDesc, Ptr, Index)) return Index;
+	return Super::AddIndexedPointer(TypeDesc, Ptr);
+}
+
+void* FNiagaraShaderMapPointerTable::GetIndexedPointer(const FTypeLayoutDesc& TypeDesc, uint32 i) const
+{
+	void* Ptr = nullptr;
+	if (DITypes.TryGetIndexedPtr(TypeDesc, i, Ptr)) return Ptr;
+	return Super::GetIndexedPointer(TypeDesc, i);
+}
+
+void FNiagaraShaderMapPointerTable::SaveToArchive(FArchive& Ar, void* FrozenContent, bool bInlineShaderResources) const
+{
+	Super::SaveToArchive(Ar, FrozenContent, bInlineShaderResources);
+
+	int32 NumDITypes = DITypes.Num();
+
+	Ar << NumDITypes;
+
+	for (int32 TypeIndex = 0; TypeIndex < NumDITypes; ++TypeIndex)
+	{
+		const UNiagaraDataInterfaceBase* DIType = DITypes.GetIndexedPointer(TypeIndex);
+		FString DIClassName = DIType->GetClass()->GetFullName();
+		Ar << DIClassName;
+	}
+}
+
+void FNiagaraShaderMapPointerTable::LoadFromArchive(FArchive& Ar, void* FrozenContent, bool bInlineShaderResources, bool bLoadedByCookedMaterial)
+{
+	Super::LoadFromArchive(Ar, FrozenContent, bInlineShaderResources, bLoadedByCookedMaterial);
+	INiagaraShaderModule * Module = INiagaraShaderModule::Get();
 
+	int32 NumDITypes = 0;
+
+	Ar << NumDITypes;
+
+	DITypes.Empty(NumDITypes);
+	for (int32 TypeIndex = 0; TypeIndex < NumDITypes; ++TypeIndex)
+	{
+		FString DIClassName;
+		Ar << DIClassName;
+		UNiagaraDataInterfaceBase* DIType = Module->RequestDefaultDataInterface(*DIClassName);
+		DITypes.LoadIndexedPointer(DIType);
+	}
+}
+
+#if 0
 void FNiagaraShaderMapId::Serialize(FArchive& Ar)
 {
 	// Handle niagara custom version and the compiler version id.
@@ -146,6 +195,7 @@ void FNiagaraShaderMapId::Serialize(FArchive& Ar)
 	
 	//ParameterSet.Serialize(Ar);		// NIAGARATODO: at some point we'll need stuff for static switches here
 }
+#endif
 
 /** Hashes the script-specific part of this shader map Id. */
 void FNiagaraShaderMapId::GetScriptHash(FSHAHash& OutHash) const
@@ -154,7 +204,7 @@ void FNiagaraShaderMapId::GetScriptHash(FSHAHash& OutHash) const
 
 	HashState.Update((const uint8*)&CompilerVersionID, sizeof(CompilerVersionID));
 	//HashState.Update((const uint8*)&BaseScriptID, sizeof(BaseScriptID));
-	HashState.Update(BaseCompileHash.GetData(), FNiagaraCompileHash::HashSize);
+	HashState.Update(BaseCompileHash.Hash, FNiagaraCompileHash::HashSize);
 	HashState.Update((const uint8*)&FeatureLevel, sizeof(FeatureLevel));
 		
 	for (int32 Index = 0; Index < AdditionalDefines.Num(); Index++)
@@ -164,7 +214,7 @@ void FNiagaraShaderMapId::GetScriptHash(FSHAHash& OutHash) const
 
 	for (int32 Index = 0; Index < ReferencedCompileHashes.Num(); Index++)
 	{
-		HashState.Update(ReferencedCompileHashes[Index].GetData(), FNiagaraCompileHash::HashSize);
+		HashState.Update(ReferencedCompileHashes[Index].Hash, FNiagaraCompileHash::HashSize);
 	}
 
 	//ParameterSet.UpdateHash(HashState);		// will need for static switches
@@ -442,35 +492,27 @@ FShader* FNiagaraShaderType::FinishCompileShader(
 	//const FUniformExpressionSet& UniformExpressionSet,
 	const FSHAHash& ShaderMapHash,
 	const FShaderCompileJob& CurrentJob,
-	const FString& InDebugDescription
+	const FString& InDebugDescription,
+	FShaderMapResourceBuilder& ResourceBuilder
 	)
 {
 	check(CurrentJob.bSucceeded);
 
 	// Reuse an existing resource with the same key or create a new one based on the compile output
 	// This allows FShaders to share compiled bytecode and RHI shader references
-	TRefCountPtr<FShaderResource> Resource = FShaderResource::FindOrCreate(CurrentJob.Output, /* SpecificPermutationId = */ 0);
+	const int32 ResourceIndex = ResourceBuilder.FindOrAddCode(CurrentJob.Output);
 
-	// Find a shader with the same key in memory
-	FShader* Shader = CurrentJob.ShaderType->FindShaderByKey(FShaderKey(ShaderMapHash, nullptr, nullptr, /* SpecificPermutationId = */ 0, CurrentJob.Input.Target.GetPlatform()));
-
-	// There was no shader with the same key so create a new one with the compile output, which will bind shader parameters
-	if (!Shader)
+	const int32 PermutationId = 0;
+	//CurrentJob.Id
+	TArray<FNiagaraDataInterfaceGPUParamInfo> DIParamInfo;
+	if (!ExtraParamInfo.RemoveAndCopyValue(&CurrentJob, DIParamInfo))
 	{
-		const int32 PermutationId = 0;
-		//CurrentJob.Id
-		TArray<FNiagaraDataInterfaceGPUParamInfo> DIParamInfo;
-
-		if (!ExtraParamInfo.RemoveAndCopyValue(&CurrentJob, DIParamInfo))
-		{
-			check(false);
-		}
-
-		Shader = (*ConstructCompiledRef)(FNiagaraShaderType::CompiledShaderInitializerType(this, PermutationId, CurrentJob.Output, MoveTemp(Resource), ShaderMapHash, InDebugDescription, DIParamInfo));
-		//CurrentJob.Output.ParameterMap.VerifyBindingsAreComplete(GetName(), CurrentJob.Output.Target, nullptr); // b/c we don't bind data interfaces yet...
-
+		check(false);
 	}
 
+	FShader* Shader = ConstructCompiled(FNiagaraShaderType::CompiledShaderInitializerType(this, PermutationId, CurrentJob.Output, ResourceIndex, ShaderMapHash, InDebugDescription, DIParamInfo));
+	//CurrentJob.Output.ParameterMap.VerifyBindingsAreComplete(GetName(), CurrentJob.Output.Target, nullptr); // b/c we don't bind data interfaces yet...
+
 	return Shader;
 }
 
@@ -482,12 +524,11 @@ FShader* FNiagaraShaderType::FinishCompileShader(
 */
 FNiagaraShaderMap* FNiagaraShaderMap::FindId(const FNiagaraShaderMapId& ShaderMapId, EShaderPlatform InPlatform)
 {
-	check(ShaderMapId.BaseCompileHash.IsValid());
 	return GIdToNiagaraShaderMap[InPlatform].FindRef(ShaderMapId);
 }
 
 /** Flushes the given shader types from any loaded FNiagaraShaderMap's. */
-void FNiagaraShaderMap::FlushShaderTypes(TArray<FShaderType*>& ShaderTypesToFlush)
+void FNiagaraShaderMap::FlushShaderTypes(TArray<const FShaderType*>& ShaderTypesToFlush)
 {
 	for (int32 ShaderMapIndex = 0; ShaderMapIndex < AllNiagaraShaderMaps.Num(); ShaderMapIndex++)
 	{
@@ -500,6 +541,7 @@ void FNiagaraShaderMap::FlushShaderTypes(TArray<FShaderType*>& ShaderTypesToFlus
 	}
 }
 
+#if 0
 void FNiagaraShaderMap::FixupShaderTypes(EShaderPlatform Platform, const TMap<FShaderType*, FString>& ShaderTypeNames)
 {
 	TArray<FNiagaraShaderMapId> Keys;
@@ -517,7 +559,7 @@ void FNiagaraShaderMap::FixupShaderTypes(EShaderPlatform Platform, const TMap<FS
 		FNiagaraShaderMap::GIdToNiagaraShaderMap[Platform].Add(Key, Values[PairIndex]);
 	}
 }
-
+#endif
 
 void NiagaraShaderMapAppendKeyString(EShaderPlatform Platform, FString& KeyString)
 {
@@ -543,7 +585,7 @@ void FNiagaraShaderMap::LoadFromDerivedDataCache(const FNiagaraShaderScript* Scr
 {
 	if (InOutShaderMap != NULL)
 	{
-		check(InOutShaderMap->Platform == Platform);
+		check(InOutShaderMap->GetShaderPlatform() == Platform);
 		// If the shader map was non-NULL then it was found in memory but is incomplete, attempt to load the missing entries from memory
 		InOutShaderMap->LoadMissingShadersFromMemory(Script);
 	}
@@ -566,7 +608,7 @@ void FNiagaraShaderMap::LoadFromDerivedDataCache(const FNiagaraShaderScript* Scr
 
 				// Deserialize from the cached data
 				InOutShaderMap->Serialize(Ar);
-				InOutShaderMap->RegisterSerializedShaders(false);
+				//InOutShaderMap->RegisterSerializedShaders(false);
 
 				checkSlow(InOutShaderMap->GetShaderMapId() == ShaderMapId);
 
@@ -591,7 +633,7 @@ void FNiagaraShaderMap::SaveToDerivedDataCache()
 	FMemoryWriter Ar(SaveData, true);
 	Serialize(Ar);
 
-	GetDerivedDataCacheRef().Put(*GetNiagaraShaderMapKeyString(ShaderMapId, Platform), SaveData, FriendlyName);
+	GetDerivedDataCacheRef().Put(*GetNiagaraShaderMapKeyString(GetContent()->ShaderMapId, GetShaderPlatform()), SaveData, FStringView(*GetFriendlyName()));
 	COOK_STAT(Timer.AddMiss(SaveData.Num()));
 }
 
@@ -600,9 +642,10 @@ TArray<uint8>* FNiagaraShaderMap::BackupShadersToMemory()
 	TArray<uint8>* SavedShaderData = new TArray<uint8>();
 	FMemoryWriter Ar(*SavedShaderData);
 
-	SerializeInline(Ar, true, true, false);
-	RegisterSerializedShaders(false);
-	Empty();
+	check(false);
+	//SerializeInline(Ar, true, true, false);
+	//RegisterSerializedShaders(false);
+	//Empty();
 
 	return SavedShaderData;
 }
@@ -610,14 +653,16 @@ TArray<uint8>* FNiagaraShaderMap::BackupShadersToMemory()
 void FNiagaraShaderMap::RestoreShadersFromMemory(const TArray<uint8>& ShaderData)
 {
 	FMemoryReader Ar(ShaderData);
-	SerializeInline(Ar, true, true, false);
-	RegisterSerializedShaders(false);
+	check(false);
+	//SerializeInline(Ar, true, true, false);
+	//RegisterSerializedShaders(false);
 }
 
 void FNiagaraShaderMap::SaveForRemoteRecompile(FArchive& Ar, const TMap<FString, TArray<TRefCountPtr<FNiagaraShaderMap> > >& CompiledShaderMaps, const TArray<FShaderResourceId>& ClientResourceIds)
 {
 	UE_LOG(LogShaders, Display, TEXT("Niagara shader map looking for unique resources, %d were on client"), ClientResourceIds.Num());
-
+	check(false);
+#if 0
 	// first, we look for the unique shader resources
 	TArray<FShaderResource*> UniqueResources;
 	int32 NumSkippedResources = 0;
@@ -633,21 +678,21 @@ void FNiagaraShaderMap::SaveForRemoteRecompile(FArchive& Ar, const TMap<FString,
 			if (ShaderMap)
 			{
 				// get all shaders in the shader map
-				TMap<FShaderId, FShader*> ShaderList;
+				TMap<FShaderId, TShaderRef<FShader>> ShaderList;
 				ShaderMap->GetShaderList(ShaderList);
 
 				// get the resources from the shaders
 				for (auto& KeyValue : ShaderList)
 				{
-					FShader* Shader = KeyValue.Value;
-					FShaderResourceId ShaderId = Shader->GetResourceId();
+					const TShaderRef<FShader> Shader = KeyValue.Value;
+					FShaderResourceId ShaderId = Shader.GetResourceId();
 
 					// skip this shader if the Id was already on the client (ie, it didn't change)
 					if (ClientResourceIds.Contains(ShaderId) == false)
 					{
 						// lookup the resource by ID
 						// add it if it's unique
-						UniqueResources.AddUnique(FShaderResource::FindById(ShaderId));
+						UniqueResources.AddUnique(Shader.GetResourceChecked());
 					}
 					else
 					{
@@ -666,7 +711,7 @@ void FNiagaraShaderMap::SaveForRemoteRecompile(FArchive& Ar, const TMap<FString,
 
 	for (int32 Index = 0; Index < NumUniqueResources; Index++)
 	{
-		UniqueResources[Index]->Serialize(Ar, false);
+		UniqueResources[Index]->SaveCodeToArchive(Ar);
 	}
 
 	// now we serialize a map (for each script)
@@ -700,6 +745,7 @@ void FNiagaraShaderMap::SaveForRemoteRecompile(FArchive& Ar, const TMap<FString,
 			}
 		}
 	}
+#endif
 }
 
 void FNiagaraShaderMap::LoadForRemoteRecompile(FArchive& Ar, EShaderPlatform ShaderPlatform, const TArray<FString>& ScriptsForShaderMaps)
@@ -849,10 +895,11 @@ void FNiagaraShaderMap::Compile(
 			Script->SetupShaderCompilationEnvironment(InPlatform, *CompilationEnvironment);
   
 			// Store the script name for debugging purposes.
-			FriendlyName = Script->GetFriendlyName();
-			NiagaraCompilationOutput = InNiagaraCompilationOutput;
-			ShaderMapId = InShaderMapId;
-			Platform = InPlatform;
+			FNiagaraShaderMapContent* NewContent = new FNiagaraShaderMapContent(InPlatform);
+			NewContent->FriendlyName = Script->GetFriendlyName();
+			NewContent->NiagaraCompilationOutput = InNiagaraCompilationOutput;
+			NewContent->ShaderMapId = InShaderMapId;
+			AssignContent(NewContent);
 
 			uint32 NumShaders = 0;
 			TArray<TSharedRef<FShaderCommonCompileJob, ESPMode::ThreadSafe>> NewJobs;
@@ -868,15 +915,15 @@ void FNiagaraShaderMap::Compile(
 					TArray<FString> ShaderErrors;
   
 					// Only compile the shader if we don't already have it
-					if (!HasShader(ShaderType, /* PermutationId = */ 0))
+					if (!NewContent->HasShader(ShaderType, /* PermutationId = */ 0))
 					{
 						TSharedRef<FShaderCommonCompileJob, ESPMode::ThreadSafe> Job = ShaderType->BeginCompileShader(
 							CompilingId,
 							Script,
 							CompilationEnvironment,
-							Platform,
+							InPlatform,
 							NewJobs,
-							FShaderTarget(ShaderType->GetFrequency(), Platform),
+							FShaderTarget(ShaderType->GetFrequency(), InPlatform),
 							Script->GetDataInterfaceParamInfo()
 							);
 						check(!SharedShaderJobs.Find(ShaderType));
@@ -913,13 +960,13 @@ void FNiagaraShaderMap::Compile(
 			{
 				TArray<int32> CurrentShaderMapId;
 				CurrentShaderMapId.Add(CompilingId);
-				GNiagaraShaderCompilationManager.FinishCompilation(*FriendlyName, CurrentShaderMapId);
+				GNiagaraShaderCompilationManager.FinishCompilation(*NewContent->FriendlyName, CurrentShaderMapId);
 			}
 		}
 	}
 }
 
-FShader* FNiagaraShaderMap::ProcessCompilationResultsForSingleJob(TSharedRef<FShaderCommonCompileJob, ESPMode::ThreadSafe> SingleJob, const FSHAHash& ShaderMapHash)
+FShader* FNiagaraShaderMap::ProcessCompilationResultsForSingleJob(TSharedRef<FShaderCommonCompileJob, ESPMode::ThreadSafe> SingleJob, const FSHAHash& ShaderMapHash, FShaderMapResourceBuilder& InResourceBuilder)
 {
 	TSharedRef<FShaderCompileJob, ESPMode::ThreadSafe> CurrentJob = StaticCastSharedRef<FShaderCompileJob>(SingleJob);
 	check(CurrentJob->Id == CompilingId);
@@ -928,16 +975,15 @@ FShader* FNiagaraShaderMap::ProcessCompilationResultsForSingleJob(TSharedRef<FSh
 
 	FNiagaraShaderType* NiagaraShaderType = CurrentJob->ShaderType->GetNiagaraShaderType();
 	check(NiagaraShaderType);
-	Shader = NiagaraShaderType->FinishCompileShader(ShaderMapHash, *CurrentJob, FriendlyName);
+	Shader = NiagaraShaderType->FinishCompileShader(ShaderMapHash, *CurrentJob, GetContent()->FriendlyName, InResourceBuilder);
 	bCompiledSuccessfully = CurrentJob->bSucceeded;
 
 	FNiagaraShader *NiagaraShader = static_cast<FNiagaraShader*>(Shader);
-	// UE-67395 - we had a case where we polluted the DDC with a shader containing no bytecode.
-	check(Shader && Shader->GetCode().Num() > 0);
-	check(!HasShader(NiagaraShaderType, /* PermutationId = */ 0));
-	AddShader(NiagaraShaderType, /* PermutationId = */ 0, Shader);
 
-	return Shader;
+	// UE-67395 - we had a case where we polluted the DDC with a shader containing no bytecode.
+	check(Shader && Shader->GetCodeSize() > 0);
+	check(!GetContent()->HasShader(NiagaraShaderType, /* PermutationId = */ 0));
+	return GetMutableContent()->FindOrAddShader(Shader);
 }
 
 bool FNiagaraShaderMap::ProcessCompilationResults(const TArray<TSharedRef<FShaderCommonCompileJob, ESPMode::ThreadSafe>>& InCompilationResults, int32& InOutJobIndex, float& TimeBudget)
@@ -947,12 +993,13 @@ bool FNiagaraShaderMap::ProcessCompilationResults(const TArray<TSharedRef<FShade
 	double StartTime = FPlatformTime::Seconds();
 
 	FSHAHash ShaderMapHash;
-	ShaderMapId.GetScriptHash(ShaderMapHash);
+	GetContent()->ShaderMapId.GetScriptHash(ShaderMapHash);
 
+	FShaderMapResourceBuilder ResourceBuilder(GetResourceCode());
 	do
 	{
 		{
-			ProcessCompilationResultsForSingleJob(InCompilationResults[InOutJobIndex], ShaderMapHash);
+			ProcessCompilationResultsForSingleJob(InCompilationResults[InOutJobIndex], ShaderMapHash, ResourceBuilder);
 		}
 
 		InOutJobIndex++;
@@ -965,6 +1012,8 @@ bool FNiagaraShaderMap::ProcessCompilationResults(const TArray<TSharedRef<FShade
 
 	if (InOutJobIndex == InCompilationResults.Num())
 	{
+		FinalizeContent();
+
 		SaveToDerivedDataCache();
 		// The shader map can now be used on the rendering thread
 		bCompilationFinalized = true;
@@ -999,7 +1048,7 @@ bool FNiagaraShaderMap::TryToAddToExistingCompilationTask(FNiagaraShaderScript*
 bool FNiagaraShaderMap::IsNiagaraShaderComplete(const FNiagaraShaderScript* Script, const FNiagaraShaderType* ShaderType, bool bSilent)
 {
 	// If we should cache this script, it's incomplete if the shader is missing
-	if (ShouldCacheNiagaraShader(ShaderType, Platform, Script) &&	!HasShader((FShaderType*)ShaderType, /* PermutationId = */ 0))
+	if (ShouldCacheNiagaraShader(ShaderType, GetShaderPlatform(), Script) &&	!GetContent()->HasShader((FShaderType*)ShaderType, /* PermutationId = */ 0))
 	{
 		if (!bSilent)
 		{
@@ -1043,6 +1092,7 @@ bool FNiagaraShaderMap::IsComplete(const FNiagaraShaderScript* Script, bool bSil
 
 void FNiagaraShaderMap::LoadMissingShadersFromMemory(const FNiagaraShaderScript* Script)
 {
+#if 0
 	// Make sure we are operating on a referenced shader map or the below Find will cause this shader map to be deleted,
 	// Since it creates a temporary ref counted pointer.
 	check(NumRefs > 0);
@@ -1074,12 +1124,13 @@ void FNiagaraShaderMap::LoadMissingShadersFromMemory(const FNiagaraShaderScript*
 			}
 		}
 	}
+#endif
 }
 #endif
 
-void FNiagaraShaderMap::GetShaderList(TMap<FShaderId, FShader*>& OutShaders) const
+void FNiagaraShaderMap::GetShaderList(TMap<FShaderId, TShaderRef<FShader>>& OutShaders) const
 {
-	TShaderMap<FNiagaraShaderType>::GetShaderList(OutShaders);
+	GetContent()->GetShaderList(*this, FSHAHash(), OutShaders);
 }
 
 /**
@@ -1088,25 +1139,17 @@ void FNiagaraShaderMap::GetShaderList(TMap<FShaderId, FShader*>& OutShaders) con
 void FNiagaraShaderMap::Register(EShaderPlatform InShaderPlatform)
 {
 	extern int32 GCreateNiagaraShadersOnLoad;
-	if (GCreateNiagaraShadersOnLoad && Platform == InShaderPlatform)
+	if (GCreateNiagaraShadersOnLoad && GetShaderPlatform() == InShaderPlatform)
 	{
-		for (auto KeyValue : GetShaders())
-		{
-			FShader* Shader = KeyValue.Value;
-			if (Shader)
-			{
-				Shader->BeginInitializeResources();
-			}
-		}
+		// TODO
 	}
 
 	if (!bRegistered)
 	{
 		INC_DWORD_STAT(STAT_Shaders_NumShaderMaps);
-		INC_DWORD_STAT_BY(STAT_Shaders_ShaderMapMemory, GetSizeBytes());
 	}
 
-	GIdToNiagaraShaderMap[Platform].Add(ShaderMapId,this);
+	GIdToNiagaraShaderMap[GetShaderPlatform()].Add(GetContent()->ShaderMapId,this);
 	bRegistered = true;
 }
 
@@ -1124,9 +1167,8 @@ void FNiagaraShaderMap::Release()
 		if (bRegistered)
 		{
 			DEC_DWORD_STAT(STAT_Shaders_NumShaderMaps);
-			DEC_DWORD_STAT_BY(STAT_Shaders_ShaderMapMemory, GetSizeBytes());
 
-			GIdToNiagaraShaderMap[Platform].Remove(ShaderMapId);
+			GIdToNiagaraShaderMap[GetShaderPlatform()].Remove(GetContent()->ShaderMapId);
 			bRegistered = false;
 		}
 
@@ -1137,8 +1179,6 @@ void FNiagaraShaderMap::Release()
 }
 
 FNiagaraShaderMap::FNiagaraShaderMap() :
-	TShaderMap<FNiagaraShaderType>(SP_NumPlatforms),
-	Platform(SP_NumPlatforms),
 	CompilingId(1),
 	NumRefs(0),
 	bDeletedThroughDeferredCleanup(false),
@@ -1163,11 +1203,11 @@ FNiagaraShaderMap::~FNiagaraShaderMap()
  * Removes all entries in the cache with exceptions based on a shader type
  * @param ShaderType - The shader type to flush
  */
-void FNiagaraShaderMap::FlushShadersByShaderType(FShaderType* ShaderType)
+void FNiagaraShaderMap::FlushShadersByShaderType(const FShaderType* ShaderType)
 {
 	if (ShaderType->GetNiagaraShaderType())
 	{
-		RemoveShaderTypePermutaion(ShaderType->GetNiagaraShaderType(), /* PermutationId = */ 0);	
+		GetMutableContent()->RemoveShaderTypePermutaion(ShaderType->GetNiagaraShaderType(), /* PermutationId = */ 0);	
 	}
 }
 
@@ -1178,36 +1218,7 @@ void FNiagaraShaderMap::Serialize(FArchive& Ar, bool bInlineShaderResources, boo
 	// Note: This is saved to the DDC, not into packages (except when cooked)
 	// Backwards compatibility therefore will not work based on the version of Ar
 	// Instead, just bump NIAGARASHADERMAP_DERIVEDDATA_VER
-
-	ShaderMapId.Serialize(Ar);
-
-	// serialize the platform enum as a uint8
-	int32 TempPlatform = (int32)Platform;
-	Ar << TempPlatform;
-	Platform = (EShaderPlatform)TempPlatform;
-
-	Ar << FriendlyName;
-
-	NiagaraCompilationOutput.Serialize(Ar);
-
-	Ar << DebugDescription;
-
-	if (Ar.IsSaving() || Ar.IsLoading())
-	{
-		TShaderMap<FNiagaraShaderType>::SerializeInline(Ar, bInlineShaderResources, false, bLoadedByCookedMaterial);
-	}
-}
-
-void FNiagaraShaderMap::RegisterSerializedShaders(bool bCookedMaterial)
-{
-	check(IsInGameThread());
-
-	TShaderMap<FNiagaraShaderType>::RegisterSerializedShaders(bCookedMaterial);
-}
-
-void FNiagaraShaderMap::DiscardSerializedShaders()
-{
-	TShaderMap<FNiagaraShaderType>::DiscardSerializedShaders();
+	Super::Serialize(Ar, bInlineShaderResources, bLoadedByCookedMaterial);
 }
 
 void FNiagaraShaderMap::RemovePendingScript(FNiagaraShaderScript* Script)
@@ -1279,20 +1290,10 @@ FNiagaraShader::FNiagaraShader(const FNiagaraShaderType::CompiledShaderInitializ
 	, DebugDescription(Initializer.DebugDescription)
 {
 	check(!DebugDescription.IsEmpty());
-	SetDataInterfaceParameterInfo(Initializer.DIParamInfo);
-	BindParams(Initializer.ParameterMap);
-}
-
-void FNiagaraShader::SetDataInterfaceParameterInfo(const TArray<FNiagaraDataInterfaceGPUParamInfo>& InDIParamInfo)
-{
-	DataInterfaceParameters.Empty();
-	for (const FNiagaraDataInterfaceGPUParamInfo& DIParamInfo : InDIParamInfo)
-	{
-		DataInterfaceParameters.Emplace(DIParamInfo);
-	}
+	BindParams(Initializer.DIParamInfo, Initializer.ParameterMap);
 }
 
-void FNiagaraShader::BindParams(const FShaderParameterMap &ParameterMap)
+void FNiagaraShader::BindParams(const TArray<FNiagaraDataInterfaceGPUParamInfo>& InDIParamInfo, const FShaderParameterMap &ParameterMap)
 {
 	FloatInputBufferParam.Bind(ParameterMap, TEXT("InputFloat"));
 	IntInputBufferParam.Bind(ParameterMap, TEXT("InputInt"));
@@ -1362,9 +1363,11 @@ void FNiagaraShader::BindParams(const FShaderParameterMap &ParameterMap)
 		EventWriteIntStrideParams[i].Bind(ParameterMap, *VarName);
 	}
 
-	for(FNiagaraDataInterfaceParamRef& Param : DataInterfaceParameters)
+	DataInterfaceParameters.Empty(InDIParamInfo.Num());
+	for(const FNiagaraDataInterfaceGPUParamInfo& DIParamInfo : InDIParamInfo)
 	{
-		Param.Bind(ParameterMap);
+		FNiagaraDataInterfaceParamRef& ParamRef = DataInterfaceParameters.AddDefaulted_GetRef();
+		ParamRef.Bind(DIParamInfo, ParameterMap);
 	}
 
 	ensure(FloatOutputBufferParam.IsBound() || IntOutputBufferParam.IsBound());	// we should have at least one output buffer we're writing to
@@ -1373,149 +1376,31 @@ void FNiagaraShader::BindParams(const FShaderParameterMap &ParameterMap)
 	ensure(NumSpawnedInstancesParam.IsBound());
 }
 
-bool FNiagaraShader::Serialize(FArchive& Ar)
-{
-	const bool bShaderHasOutdatedParameters = FShader::Serialize(Ar);
-	Ar << NiagaraUniformBuffer;	// do we need to keep about uniform buffer parameters?
-
-	Ar << DebugDescription;
-
-	Ar << FloatInputBufferParam;
-	Ar << IntInputBufferParam;
-	Ar << FloatOutputBufferParam;
-	Ar << IntOutputBufferParam;
-
-	Ar << InstanceCountsParam;
-	Ar << ReadInstanceCountOffsetParam;
-	Ar << WriteInstanceCountOffsetParam;
-
-	Ar << FreeIDBufferParam;
-	Ar << IDToIndexBufferParam;
-
-	Ar << SimStartParam;
-	Ar << EmitterTickCounterParam;
-	Ar << EmitterSpawnInfoOffsetsParam;
-	Ar << EmitterSpawnInfoParamsParam;
-
-	Ar << NumSpawnedInstancesParam;
-	Ar << UpdateStartInstanceParam;
-	Ar << DefaultShaderStageIndexParam;
-	Ar << ShaderStageIndexParam;
-	Ar << IterationInterfaceCount;
-	Ar << ComponentBufferSizeReadParam;
-	Ar << ComponentBufferSizeWriteParam;
-
-	for (uint32 i = 0; i < MAX_CONCURRENT_EVENT_DATASETS; i++)
-	{
-		Ar << EventFloatUAVParams[i];
-		Ar << EventIntUAVParams[i];
-		Ar << EventFloatSRVParams[i];
-		Ar << EventIntSRVParams[i];
-	}
-
-	Ar << DataInterfaceParameters;
-
-	Ar << GlobalConstantBufferParam[0];
-	Ar << SystemConstantBufferParam[0];
-	Ar << OwnerConstantBufferParam[0];
-	Ar << EmitterConstantBufferParam[0];
-	Ar << ExternalConstantBufferParam[0];
-
-	Ar << GlobalConstantBufferParam[1];
-	Ar << SystemConstantBufferParam[1];
-	Ar << OwnerConstantBufferParam[1];
-	Ar << EmitterConstantBufferParam[1];
-	Ar << ExternalConstantBufferParam[1];
-
-	Ar << DataInterfaceUniformBufferParam;
-	Ar << NumEventsPerParticleParam;
-	Ar << NumParticlesPerEventParam;
-	Ar << CopyInstancesBeforeStartParam;
-	Ar << ViewUniformBufferParam;
-
-	return bShaderHasOutdatedParameters;
-}
-
-uint32 FNiagaraShader::GetAllocatedSize() const
-{
-	return FShader::GetAllocatedSize()
-		//+ ParameterCollectionUniformBuffers.GetAllocatedSize()
-		+ DebugDescription.GetAllocatedSize();
-}
 
 //////////////////////////////////////////////////////////////////////////
-FNiagaraDataInterfaceParamRef::FNiagaraDataInterfaceParamRef(const FNiagaraDataInterfaceGPUParamInfo& InParameterInfo)
-	: ParameterInfo(InParameterInfo)
-	, Parameters(nullptr)
-{
-	InitDIClass();
-}
-
 FNiagaraDataInterfaceParamRef::FNiagaraDataInterfaceParamRef()
-	: DIClass(nullptr)
-	, Parameters(nullptr)
 {
 
 }
 
 FNiagaraDataInterfaceParamRef::~FNiagaraDataInterfaceParamRef()
 {
-	if (Parameters)
-	{
-		delete Parameters;
-	}
-}
-
-void FNiagaraDataInterfaceParamRef::ConstructParameters()
-{
-	if (Parameters)
-	{
-		delete Parameters;
-	}
-
-	Parameters = CastChecked<UNiagaraDataInterfaceBase>(DIClass->GetDefaultObject())->ConstructComputeParameters();
+	Parameters.SafeDelete();
 }
 
-void FNiagaraDataInterfaceParamRef::InitDIClass()
+void FNiagaraDataInterfaceParamRef::Bind(const FNiagaraDataInterfaceGPUParamInfo& InParameterInfo, const class FShaderParameterMap& ParameterMap)
 {
 	INiagaraShaderModule* Module = INiagaraShaderModule::Get();
-	check(Module != nullptr);
-	
-	// Getting the base here in hopes that in the future we would just reference it directly rather than going through
-	// the class intermediary.
-	UNiagaraDataInterfaceBase* Base = Module->RequestDefaultDataInterface(ParameterInfo.DIClassName);
+	UNiagaraDataInterfaceBase* Base = Module->RequestDefaultDataInterface(*InParameterInfo.DIClassName);
+	ensureMsgf(Base, TEXT("Failed to load class for FNiagaraDataInterfaceParamRef. %s"), *InParameterInfo.DIClassName);
 	if (Base)
 	{
-		DIClass = Base->GetClass();
+		DIType = TIndexedPtr<UNiagaraDataInterfaceBase>(Base); // TODO - clean up TIndexedPtr::operator=()
+		Parameters = Base->ConstructComputeParameters();
+		checkf(Parameters != nullptr, TEXT("Niagara type '%s' is missing ComputeParameters, missing a IMPLEMENT_NIAGARA_DI_PARAMETER?"));
+		Parameters->DIType = DIType;
+		Base->BindParameters(Parameters, InParameterInfo, ParameterMap);
 	}
-	ensureMsgf(DIClass, TEXT("Failed to load class for FNiagaraDataInterfaceParamRef. %s"), *ParameterInfo.DIClassName);
-}
-
-void FNiagaraDataInterfaceParamRef::Bind(const class FShaderParameterMap& ParameterMap)
-{
-	ConstructParameters();
-	if (Parameters)
-	{
-		Parameters->Bind(*this, ParameterMap);
-	}
-}
-
-bool operator<<(FArchive& Ar, FNiagaraDataInterfaceParamRef& ParamRef)
-{
-	ParamRef.ParameterInfo.Serialize(Ar);
-
-	if (Ar.IsLoading())
-	{
-		ParamRef.InitDIClass();
-		ParamRef.ConstructParameters();
-	}
-
-	if (ParamRef.Parameters)
-	{
-		ParamRef.Parameters->Serialize(Ar);
-	}
-
-	return true;
 }
 
 bool FNiagaraDataInterfaceGeneratedFunction::Serialize(FArchive& Ar)
@@ -1531,6 +1416,12 @@ bool operator<<(FArchive& Ar, FNiagaraDataInterfaceGeneratedFunction& DIFunction
 	return DIFunction.Serialize(Ar);
 }
 
+void FNiagaraDataInterfaceParamRef::WriteFrozenParameters(FMemoryImageWriter& Writer, const TMemoryImagePtr<FNiagaraDataInterfaceParametersCS>& InParameters) const
+{
+	const UNiagaraDataInterfaceBase* Base = DIType.Get(Writer.TryGetPrevPointerTable());
+	InParameters.WriteMemoryImageWithDerivedType(Writer, Base->GetComputeParametersTypeDesc());
+}
+
 bool FNiagaraDataInterfaceGPUParamInfo::Serialize(FArchive& Ar)
 {
 	Ar.UsingCustomVersion(FNiagaraCustomVersion::GUID);
diff --git a/Engine/Plugins/FX/Niagara/Source/NiagaraShader/Private/NiagaraShaderDerivedDataVersion.h b/Engine/Plugins/FX/Niagara/Source/NiagaraShader/Private/NiagaraShaderDerivedDataVersion.h
--- a/Engine/Plugins/FX/Niagara/Source/NiagaraShader/Private/NiagaraShaderDerivedDataVersion.h
+++ b/Engine/Plugins/FX/Niagara/Source/NiagaraShader/Private/NiagaraShaderDerivedDataVersion.h
@@ -11,4 +11,4 @@ NiagaraShaderDerivedDataVersion.h: Shader derived data version for Niagara.
 // In case of merge conflicts with DDC versions, you MUST generate a new GUID and set this new
 // guid as version
 
-#define NIAGARASHADERMAP_DERIVEDDATA_VER		TEXT("A27AE38E639D4954AC529C61BD67058A")
+#define NIAGARASHADERMAP_DERIVEDDATA_VER		TEXT("{C6FAD17E-D2B1-4A9D-8242-D5FAC8A833EB}")
diff --git a/Engine/Plugins/FX/Niagara/Source/NiagaraShader/Private/NiagaraShaderParticleID.cpp b/Engine/Plugins/FX/Niagara/Source/NiagaraShader/Private/NiagaraShaderParticleID.cpp
--- a/Engine/Plugins/FX/Niagara/Source/NiagaraShader/Private/NiagaraShaderParticleID.cpp
+++ b/Engine/Plugins/FX/Niagara/Source/NiagaraShader/Private/NiagaraShaderParticleID.cpp
@@ -30,16 +30,6 @@ public:
 		OutEnvironment.SetDefine(TEXT("THREAD_COUNT"), THREAD_COUNT);
 	}
 
-	bool Serialize(FArchive& Ar)
-	{
-		bool bShaderHasOutdatedParameters = FGlobalShader::Serialize(Ar);
-		Ar << NewBufferParam;
-		Ar << ExistingBufferParam;
-		Ar << NumNewElementsParam;
-		Ar << NumExistingElementsParam;
-		return bShaderHasOutdatedParameters;
-	}
-
 	void Execute(FRHICommandList& RHICmdList, uint32 NumElementsToAlloc, FRWBuffer& NewBuffer, uint32 NumExistingElements, FRHIShaderResourceView* ExistingBuffer)
 	{
 		// To simplify the shader code, the size of the ID table must be a multiple of the thread count.
@@ -49,7 +39,7 @@ public:
 		check(NumElementsToAlloc >= NumExistingElements);
 		uint32 NumNewElements = NumElementsToAlloc - NumExistingElements;
 
-		FRHIComputeShader* ComputeShader = GetComputeShader();
+		FRHIComputeShader* ComputeShader = RHICmdList.GetBoundComputeShader();
 		RHICmdList.SetComputeShader(ComputeShader);
 
 		NewBufferParam.SetBuffer(RHICmdList, ComputeShader, NewBuffer);
@@ -66,10 +56,10 @@ public:
 private:
 	static constexpr uint32 THREAD_COUNT = 64;
 
-	FRWShaderParameter NewBufferParam;
-	FShaderResourceParameter ExistingBufferParam;
-	FShaderParameter NumNewElementsParam;
-	FShaderParameter NumExistingElementsParam;
+	LAYOUT_FIELD(FRWShaderParameter, NewBufferParam);
+	LAYOUT_FIELD(FShaderResourceParameter, ExistingBufferParam);
+	LAYOUT_FIELD(FShaderParameter, NumNewElementsParam);
+	LAYOUT_FIELD(FShaderParameter, NumExistingElementsParam);
 };
 
 IMPLEMENT_GLOBAL_SHADER(FNiagaraInitFreeIDBufferCS, "/Plugin/FX/Niagara/Private/NiagaraInitFreeIDBuffer.usf", "InitIDBufferCS", SF_Compute);
@@ -105,16 +95,6 @@ public:
 		FGlobalShader::ModifyCompilationEnvironment(Parameters, OutEnvironment);
 	}
 
-	bool Serialize(FArchive& Ar)
-	{
-		bool bShaderHasOutdatedParameters = FGlobalShader::Serialize(Ar);
-		Ar << IDToIndexTableParam;
-		Ar << FreeIDListParam;
-		Ar << FreeIDListSizesParam;
-		Ar << FreeIDListIndexParam;
-		return bShaderHasOutdatedParameters;
-	}
-
 	void Execute(FRHICommandList& RHICmdList, ERHIFeatureLevel::Type FeatureLevel, uint32 NumIDs, FRHIShaderResourceView* IDToIndexTable, FRWBuffer& FreeIDList, FRWBuffer& FreeIDListSizes, uint32 FreeIDListIndex)
 	{
 		const EShaderPlatform Platform = GShaderPlatformForFeatureLevel[FeatureLevel];
@@ -123,7 +103,7 @@ public:
 		// To simplify the shader code, the size of the ID table must be a multiple of the thread count.
 		check(NumIDs % THREAD_COUNT == 0);
 
-		FRHIComputeShader* ComputeShader = GetComputeShader();
+		FRHIComputeShader* ComputeShader = RHICmdList.GetBoundComputeShader();
 		RHICmdList.SetComputeShader(ComputeShader);
 
 		SetSRVParameter(RHICmdList, ComputeShader, IDToIndexTableParam, IDToIndexTable);
@@ -139,10 +119,10 @@ public:
 	}
 
 private:
-	FShaderResourceParameter IDToIndexTableParam;
-	FRWShaderParameter FreeIDListParam;
-	FRWShaderParameter FreeIDListSizesParam;
-	FShaderParameter FreeIDListIndexParam;
+	LAYOUT_FIELD(FShaderResourceParameter, IDToIndexTableParam);
+	LAYOUT_FIELD(FRWShaderParameter, FreeIDListParam);
+	LAYOUT_FIELD(FRWShaderParameter, FreeIDListSizesParam);
+	LAYOUT_FIELD(FShaderParameter, FreeIDListIndexParam);
 };
 
 IMPLEMENT_GLOBAL_SHADER(NiagaraComputeFreeIDsCS, "/Plugin/FX/Niagara/Private/NiagaraComputeFreeIDs.usf", "ComputeFreeIDs", SF_Compute);
@@ -172,22 +152,13 @@ public:
 		return RHISupportsComputeShaders(Parameters.Platform);
 	}
 
-	bool Serialize(FArchive& Ar)
-	{
-		bool bShaderHasOutdatedParameters = FGlobalShader::Serialize(Ar);
-		Ar << TargetBufferParam;
-		Ar << FillValueParam;
-		Ar << BufferSizeParam;
-		return bShaderHasOutdatedParameters;
-	}
-
 	void Execute(FRHICommandList& RHICmdList, FRWBuffer& Buffer, int Value)
 	{
 		const uint32 THREAD_COUNT = 64;
 		const uint32 NumInts = Buffer.NumBytes / sizeof(int);
 		const uint32 ThreadGroups = FMath::DivideAndRoundUp(NumInts, THREAD_COUNT);
 
-		FRHIComputeShader* ComputeShader = GetComputeShader();
+		FRHIComputeShader* ComputeShader = RHICmdList.GetBoundComputeShader();
 		RHICmdList.SetComputeShader(ComputeShader);
 
 		SetUAVParameter(RHICmdList, ComputeShader, TargetBufferParam, Buffer.UAV);
@@ -200,9 +171,9 @@ public:
 	}
 
 private:
-	FShaderResourceParameter TargetBufferParam;
-	FShaderParameter FillValueParam;
-	FShaderParameter BufferSizeParam;
+	LAYOUT_FIELD(FShaderResourceParameter, TargetBufferParam);
+	LAYOUT_FIELD(FShaderParameter, FillValueParam);
+	LAYOUT_FIELD(FShaderParameter, BufferSizeParam);
 };
 
 IMPLEMENT_GLOBAL_SHADER(NiagaraFillIntBufferCS, "/Plugin/FX/Niagara/Private/NiagaraFillIntBuffer.usf", "FillIntBuffer", SF_Compute);
diff --git a/Engine/Plugins/FX/Niagara/Source/NiagaraShader/Private/NiagaraShared.cpp b/Engine/Plugins/FX/Niagara/Source/NiagaraShader/Private/NiagaraShared.cpp
--- a/Engine/Plugins/FX/Niagara/Source/NiagaraShader/Private/NiagaraShared.cpp
+++ b/Engine/Plugins/FX/Niagara/Source/NiagaraShader/Private/NiagaraShared.cpp
@@ -19,6 +19,11 @@
 #include "Modules/ModuleManager.h"
 #include "NiagaraCustomVersion.h"
 
+IMPLEMENT_TYPE_LAYOUT(FNiagaraDataInterfaceParamRef);
+IMPLEMENT_TYPE_LAYOUT(FNiagaraShaderMapContent);
+IMPLEMENT_TYPE_LAYOUT(FNiagaraShaderMapId);
+IMPLEMENT_TYPE_LAYOUT(FNiagaraComputeShaderCompilationOutput);
+
 #if WITH_EDITOR
 	NIAGARASHADER_API FNiagaraCompilationQueue* FNiagaraCompilationQueue::Singleton = nullptr;
 #endif
@@ -86,14 +91,32 @@ NIAGARASHADER_API void FNiagaraShaderScript::LegacySerialize(FArchive& Ar)
 
 bool FNiagaraShaderScript::IsSame(const FNiagaraShaderMapId& InId) const
 {
+	if (InId.ReferencedCompileHashes.Num() != ReferencedCompileHashes.Num() ||
+		InId.AdditionalDefines.Num() != AdditionalDefines.Num())
+	{
+		return false;
+	}
+	for (int32 i = 0; i < ReferencedCompileHashes.Num(); ++i)
+	{
+		if (ReferencedCompileHashes[i] != InId.ReferencedCompileHashes[i])
+		{
+			return false;
+		}
+	}
+	for (int32 i = 0; i < AdditionalDefines.Num(); ++i)
+	{
+		if (AdditionalDefines[i] != *InId.AdditionalDefines[i])
+		{
+			return false;
+		}
+	}
+
 	return
 		InId.FeatureLevel == FeatureLevel &&/*
 		InId.BaseScriptID == BaseScriptId &&*/
-		InId.AdditionalDefines == AdditionalDefines &&
 		InId.DetailLevelMask == DetailLevelMask &&
 		InId.bUsesRapidIterationParams == bUsesRapidIterationParams &&
 		InId.BaseCompileHash == BaseCompileHash &&
-		InId.ReferencedCompileHashes == ReferencedCompileHashes &&
 		InId.CompilerVersionID == CompilerVersionId;
 }
 
@@ -125,12 +148,22 @@ NIAGARASHADER_API void FNiagaraShaderScript::GetShaderMapId(EShaderPlatform Plat
 		GetDependentShaderTypes(Platform, ShaderTypes);
 		OutId.FeatureLevel = GetFeatureLevel();/*
 		OutId.BaseScriptID = BaseScriptId;*/
-		OutId.AdditionalDefines = AdditionalDefines;
 		OutId.DetailLevelMask = DetailLevelMask;
 		OutId.bUsesRapidIterationParams = bUsesRapidIterationParams;		
-		OutId.BaseCompileHash = BaseCompileHash;
-		OutId.ReferencedCompileHashes = ReferencedCompileHashes;
+		BaseCompileHash.ToSHAHash(OutId.BaseCompileHash);
 		OutId.CompilerVersionID = FNiagaraCustomVersion::LatestScriptCompileVersion;
+
+		OutId.ReferencedCompileHashes.Reserve(ReferencedCompileHashes.Num());
+		for (const FNiagaraCompileHash& Hash : ReferencedCompileHashes)
+		{
+			Hash.ToSHAHash(OutId.ReferencedCompileHashes.AddDefaulted_GetRef());
+		}
+
+		OutId.AdditionalDefines.Empty(AdditionalDefines.Num());
+		for(const FString& Define : AdditionalDefines)
+		{
+			OutId.AdditionalDefines.Emplace(Define);
+		}
 	}
 }
 
@@ -140,20 +173,11 @@ void FNiagaraShaderScript::AddReferencedObjects(FReferenceCollector& Collector)
 {
 }
 
-
-void FNiagaraShaderScript::RegisterShaderMap()
-{
-	if (GameThreadShaderMap)
-	{
-		GameThreadShaderMap->RegisterSerializedShaders(bLoadedFromCookedMaterial);
-	}
-}
-
 void  FNiagaraShaderScript::DiscardShaderMap()
 {
 	if (GameThreadShaderMap)
 	{
-		GameThreadShaderMap->DiscardSerializedShaders();
+		//GameThreadShaderMap->DiscardSerializedShaders();
 	}
 }
 
@@ -226,7 +250,7 @@ void FNiagaraShaderScript::SerializeShaderMap(FArchive& Ar)
 				}
 				else
 				{
-					LoadedShaderMap->DiscardSerializedShaders();
+					//LoadedShaderMap->DiscardSerializedShaders();
 				}
 			}
 		}
@@ -391,38 +415,29 @@ void FNiagaraShaderScript::FinishCompilation()
 
 #endif
 
-void FNiagaraShaderScript::SetDataInterfaceParamInfo(TArray< FNiagaraDataInterfaceGPUParamInfo >& InDIParamInfo)
+void FNiagaraShaderScript::SetDataInterfaceParamInfo(const TArray< FNiagaraDataInterfaceGPUParamInfo >& InDIParamInfo)
 {
 	DIParamInfo = InDIParamInfo;
 }
 
-void FNiagaraShaderScript::SetDataInterfaceParamInfo(TArray<FNiagaraDataInterfaceParamRef>& InDIParamRefs)
-{
-	DIParamInfo.Empty();
-	for (FNiagaraDataInterfaceParamRef& DIParam : InDIParamRefs)
-	{
-		DIParamInfo.Add(DIParam.ParameterInfo);
-	}
-}
-
-NIAGARASHADER_API  FNiagaraShader* FNiagaraShaderScript::GetShader() const
+NIAGARASHADER_API  FNiagaraShaderRef FNiagaraShaderScript::GetShader() const
 {
 	check(!GIsThreadedRendering || !IsInGameThread());
 	if (!GIsEditor || RenderingThreadShaderMap /*&& RenderingThreadShaderMap->IsComplete(this, true)*/)
 	{
 		return RenderingThreadShaderMap->GetShader<FNiagaraShader>();
 	}
-	return nullptr;
+	return FNiagaraShaderRef();
 };
 
-NIAGARASHADER_API  FNiagaraShader* FNiagaraShaderScript::GetShaderGameThread() const
+NIAGARASHADER_API  FNiagaraShaderRef FNiagaraShaderScript::GetShaderGameThread() const
 {
 	if (GameThreadShaderMap)
 	{
 		return GameThreadShaderMap->GetShader<FNiagaraShader>();
 	}
 
-	return nullptr;
+	return FNiagaraShaderRef();
 };
 
 
diff --git a/Engine/Plugins/FX/Niagara/Source/NiagaraShader/Public/NiagaraShader.h b/Engine/Plugins/FX/Niagara/Source/NiagaraShader/Public/NiagaraShader.h
--- a/Engine/Plugins/FX/Niagara/Source/NiagaraShader/Public/NiagaraShader.h
+++ b/Engine/Plugins/FX/Niagara/Source/NiagaraShader/Public/NiagaraShader.h
@@ -17,7 +17,6 @@
 #include "NiagaraShaderType.h"
 #include "SceneRenderTargetParameters.h"
 
-struct FNiagaraDataInterfaceParametersCS;
 class UClass;
 
 template<typename TBufferStruct> class TUniformBufferRef;
@@ -28,23 +27,14 @@ class NIAGARASHADER_API FNiagaraShader : public FShader
 public:
 	DECLARE_SHADER_TYPE(FNiagaraShader, Niagara);
 
+	using FPermutationParameters = FNiagaraShaderPermutationParameters;
+
 	static FName UniformBufferLayoutName;
 
 	FNiagaraShader()
 	{
 	}
 
-	static bool ShouldCompilePermutation(EShaderPlatform Platform, const FNiagaraShaderScript*  Script)
-	{
-		//@todo - lit materials only 
-		return FNiagaraUtilities::SupportsGPUParticles(Platform);
-	}
-
-
-	FNiagaraShader(const FNiagaraShaderType::CompiledShaderInitializerType& Initializer);
-
-	typedef void (*ModifyCompilationEnvironmentType)(EShaderPlatform, const FNiagaraShaderScript* , FShaderCompilerEnvironment&);
-
 	static uint32 GetGroupSize(EShaderPlatform Platform)
 	{
 		if (Platform == SP_XBOXONE_D3D12 || Platform == SP_PS4)
@@ -57,17 +47,19 @@ public:
 		}
 	}
 
-	static void ModifyCompilationEnvironment(EShaderPlatform Platform, const FNiagaraShaderScript*  Script, FShaderCompilerEnvironment& OutEnvironment)
+	static void ModifyCompilationEnvironment(const FNiagaraShaderPermutationParameters& Parameters, FShaderCompilerEnvironment& OutEnvironment)
 	{
-		OutEnvironment.SetDefine(TEXT("THREADGROUP_SIZE"), GetGroupSize(Platform));
+		FShader::ModifyCompilationEnvironment(Parameters, OutEnvironment);
+		OutEnvironment.SetDefine(TEXT("THREADGROUP_SIZE"), GetGroupSize(Parameters.Platform));
 	}
 
-	static bool ValidateCompiledResult(EShaderPlatform Platform, const FShaderParameterMap& ParameterMap, TArray<FString>& OutError)
+	static bool ShouldCompilePermutation(const FNiagaraShaderPermutationParameters& Parameters)
 	{
-		return true;
+		//@todo - lit materials only 
+		return FNiagaraUtilities::SupportsGPUParticles(Parameters.Platform);
 	}
-	
-	void SetDataInterfaceParameterInfo(const TArray<FNiagaraDataInterfaceGPUParamInfo>& InDIParamInfo);
+
+	FNiagaraShader(const FNiagaraShaderType::CompiledShaderInitializerType& Initializer);
 
 //	FRHIUniformBuffer* GetParameterCollectionBuffer(const FGuid& Id, const FSceneInterface* SceneInterface) const;
 	/*
@@ -92,68 +84,63 @@ public:
 	*/
 
 	// Bind parameters
-	void BindParams(const FShaderParameterMap &ParameterMap);
-
-	// FShader interface.
-	virtual bool Serialize(FArchive& Ar) override;
-	virtual uint32 GetAllocatedSize() const override;
-
-	FShaderResourceParameter FloatInputBufferParam;
-	FShaderResourceParameter IntInputBufferParam;
-	FRWShaderParameter FloatOutputBufferParam;
-	FRWShaderParameter IntOutputBufferParam;
-	FRWShaderParameter InstanceCountsParam;
-	FShaderParameter ReadInstanceCountOffsetParam;
-	FShaderParameter WriteInstanceCountOffsetParam;
-	FShaderResourceParameter FreeIDBufferParam;
-	FRWShaderParameter IDToIndexBufferParam;
-	FShaderUniformBufferParameter GlobalConstantBufferParam[2];
-	FShaderUniformBufferParameter SystemConstantBufferParam[2];
-	FShaderUniformBufferParameter OwnerConstantBufferParam[2];
-	FShaderUniformBufferParameter EmitterConstantBufferParam[2];
-	FShaderUniformBufferParameter ExternalConstantBufferParam[2];
-	FShaderUniformBufferParameter DataInterfaceUniformBufferParam;
-	FShaderUniformBufferParameter ViewUniformBufferParam;
-	FShaderParameter SimStartParam;
-	FShaderParameter EmitterTickCounterParam;
-	FShaderParameter EmitterSpawnInfoOffsetsParam;
-	FShaderParameter EmitterSpawnInfoParamsParam;
-	FShaderParameter NumEventsPerParticleParam;
-	FShaderParameter NumParticlesPerEventParam;
-	FShaderParameter CopyInstancesBeforeStartParam;
-	FShaderParameter NumSpawnedInstancesParam;
-	FShaderParameter UpdateStartInstanceParam;
-	FShaderParameter DefaultShaderStageIndexParam;
-	FShaderParameter ShaderStageIndexParam;
-	FShaderParameter IterationInterfaceCount;
-	FShaderParameter ComponentBufferSizeReadParam;
-	FShaderParameter ComponentBufferSizeWriteParam;
-	FRWShaderParameter EventIntUAVParams[MAX_CONCURRENT_EVENT_DATASETS];
-	FRWShaderParameter EventFloatUAVParams[MAX_CONCURRENT_EVENT_DATASETS];
-	FShaderResourceParameter EventIntSRVParams[MAX_CONCURRENT_EVENT_DATASETS];
-	FShaderResourceParameter EventFloatSRVParams[MAX_CONCURRENT_EVENT_DATASETS];
-	FShaderParameter EventWriteFloatStrideParams[MAX_CONCURRENT_EVENT_DATASETS];
-	FShaderParameter EventWriteIntStrideParams[MAX_CONCURRENT_EVENT_DATASETS];
-	FShaderParameter EventReadFloatStrideParams[MAX_CONCURRENT_EVENT_DATASETS];
-	FShaderParameter EventReadIntStrideParams[MAX_CONCURRENT_EVENT_DATASETS];
-
-	TArray< FNiagaraDataInterfaceParamRef >& GetDIParameters()
+	void BindParams(const TArray<FNiagaraDataInterfaceGPUParamInfo>& InDIParamInfo, const FShaderParameterMap &ParameterMap);
+
+	const TMemoryImageArray<FNiagaraDataInterfaceParamRef>& GetDIParameters()
 	{
 		return DataInterfaceParameters;
 	}
 
+	LAYOUT_FIELD(FShaderResourceParameter, FloatInputBufferParam);
+	LAYOUT_FIELD(FShaderResourceParameter, IntInputBufferParam);
+	LAYOUT_FIELD(FRWShaderParameter, FloatOutputBufferParam);
+	LAYOUT_FIELD(FRWShaderParameter, IntOutputBufferParam);
+	LAYOUT_FIELD(FRWShaderParameter, InstanceCountsParam);
+	LAYOUT_FIELD(FShaderParameter, ReadInstanceCountOffsetParam);
+	LAYOUT_FIELD(FShaderParameter, WriteInstanceCountOffsetParam);
+	LAYOUT_FIELD(FShaderResourceParameter, FreeIDBufferParam);
+	LAYOUT_FIELD(FRWShaderParameter, IDToIndexBufferParam);
+	LAYOUT_ARRAY(FShaderUniformBufferParameter, GlobalConstantBufferParam, 2);
+	LAYOUT_ARRAY(FShaderUniformBufferParameter, SystemConstantBufferParam, 2);
+	LAYOUT_ARRAY(FShaderUniformBufferParameter, OwnerConstantBufferParam, 2);
+	LAYOUT_ARRAY(FShaderUniformBufferParameter, EmitterConstantBufferParam, 2);
+	LAYOUT_ARRAY(FShaderUniformBufferParameter, ExternalConstantBufferParam, 2);
+	LAYOUT_FIELD(FShaderUniformBufferParameter, DataInterfaceUniformBufferParam);
+	LAYOUT_FIELD(FShaderUniformBufferParameter, ViewUniformBufferParam);
+	LAYOUT_FIELD(FShaderParameter, SimStartParam);
+	LAYOUT_FIELD(FShaderParameter, EmitterTickCounterParam);
+	LAYOUT_FIELD(FShaderParameter, EmitterSpawnInfoOffsetsParam);
+	LAYOUT_FIELD(FShaderParameter, EmitterSpawnInfoParamsParam);
+	LAYOUT_FIELD(FShaderParameter, NumEventsPerParticleParam);
+	LAYOUT_FIELD(FShaderParameter, NumParticlesPerEventParam);
+	LAYOUT_FIELD(FShaderParameter, CopyInstancesBeforeStartParam);
+	LAYOUT_FIELD(FShaderParameter, NumSpawnedInstancesParam);
+	LAYOUT_FIELD(FShaderParameter, UpdateStartInstanceParam);
+	LAYOUT_FIELD(FShaderParameter, DefaultShaderStageIndexParam);
+	LAYOUT_FIELD(FShaderParameter, ShaderStageIndexParam);
+	LAYOUT_FIELD(FShaderParameter, IterationInterfaceCount);
+	LAYOUT_FIELD(FShaderParameter, ComponentBufferSizeReadParam);
+	LAYOUT_FIELD(FShaderParameter, ComponentBufferSizeWriteParam);
+	LAYOUT_ARRAY(FRWShaderParameter, EventIntUAVParams, MAX_CONCURRENT_EVENT_DATASETS);
+	LAYOUT_ARRAY(FRWShaderParameter, EventFloatUAVParams, MAX_CONCURRENT_EVENT_DATASETS);
+	LAYOUT_ARRAY(FShaderResourceParameter, EventIntSRVParams, MAX_CONCURRENT_EVENT_DATASETS);
+	LAYOUT_ARRAY(FShaderResourceParameter, EventFloatSRVParams, MAX_CONCURRENT_EVENT_DATASETS);
+	LAYOUT_ARRAY(FShaderParameter, EventWriteFloatStrideParams, MAX_CONCURRENT_EVENT_DATASETS);
+	LAYOUT_ARRAY(FShaderParameter, EventWriteIntStrideParams, MAX_CONCURRENT_EVENT_DATASETS);
+	LAYOUT_ARRAY(FShaderParameter, EventReadFloatStrideParams, MAX_CONCURRENT_EVENT_DATASETS);
+	LAYOUT_ARRAY(FShaderParameter, EventReadIntStrideParams, MAX_CONCURRENT_EVENT_DATASETS);
 
 private:
-	FShaderUniformBufferParameter NiagaraUniformBuffer;
+	LAYOUT_FIELD(FShaderUniformBufferParameter, NiagaraUniformBuffer);
 
 	// Data about parameters used for each Data Interface.
-	TArray< FNiagaraDataInterfaceParamRef > DataInterfaceParameters;
+	LAYOUT_FIELD(TMemoryImageArray<FNiagaraDataInterfaceParamRef>, DataInterfaceParameters);
 
 	/*
 	FDebugUniformExpressionSet	DebugUniformExpressionSet;
 	FRHIUniformBufferLayout		DebugUniformExpressionUBLayout;
 	*/
-	FString						DebugDescription;
+	LAYOUT_FIELD(FMemoryImageString, DebugDescription);
 
 	/* OPTODO: ? */
 	/*
@@ -163,7 +150,6 @@ private:
 	static FAutoConsoleVariableRef CVarAllowCachedUniformExpressions;
 	*/
 
-
 	/*
 #if !(UE_BUILD_TEST || UE_BUILD_SHIPPING || !WITH_EDITOR)
 	void VerifyExpressionAndShaderMaps(const FMaterialRenderProxy* MaterialRenderProxy, const FMaterial& Material, const FUniformExpressionCache* UniformExpressionCache);
@@ -171,7 +157,6 @@ private:
 	*/
 };
 
-
 class FNiagaraEmitterInstanceShader : public FNiagaraShader
 {
 
diff --git a/Engine/Plugins/FX/Niagara/Source/NiagaraShader/Public/NiagaraShaderType.h b/Engine/Plugins/FX/Niagara/Source/NiagaraShader/Public/NiagaraShaderType.h
--- a/Engine/Plugins/FX/Niagara/Source/NiagaraShader/Public/NiagaraShaderType.h
+++ b/Engine/Plugins/FX/Niagara/Source/NiagaraShader/Public/NiagaraShaderType.h
@@ -37,8 +37,15 @@ extern void UpdateNiagaraShaderCompilingStats(const FNiagaraShaderScript* Script
  */
 extern ENGINE_API void DumpComputeShaderStats( EShaderPlatform Platform );
 
+struct FNiagaraShaderPermutationParameters : public FShaderPermutationParameters
+{
+	const FNiagaraShaderScript* Script;
 
-
+	FNiagaraShaderPermutationParameters(EShaderPlatform InPlatform, const FNiagaraShaderScript* InScript)
+		: FShaderPermutationParameters(InPlatform)
+		, Script(InScript)
+	{}
+};
 
 /**
  * A shader meta type for niagara-linked shaders.
@@ -56,22 +63,19 @@ public:
 			FShaderType* InType,
 			int32 InPermutationId,
 			const FShaderCompilerOutput& CompilerOutput,
-			TRefCountPtr<FShaderResource>&& InResource,
+			int32 InResourceIndex,
 			const FSHAHash& InNiagaraShaderMapHash,
 			const FString& InDebugDescription,
 			const TArray< FNiagaraDataInterfaceGPUParamInfo > &InDIParamInfo
 			)
-		: FGlobalShaderType::CompiledShaderInitializerType(InType,InPermutationId,CompilerOutput, MoveTemp(InResource), InNiagaraShaderMapHash,nullptr,nullptr)
+		: FGlobalShaderType::CompiledShaderInitializerType(InType,InPermutationId,CompilerOutput, InResourceIndex, InNiagaraShaderMapHash,nullptr,nullptr)
 		, DebugDescription(InDebugDescription)
 		, DIParamInfo(InDIParamInfo)
 		{}
 	};
-	typedef FShader* (*ConstructCompiledType)(const CompiledShaderInitializerType&);
-	typedef bool (*ShouldCompilePermutationType)(EShaderPlatform,const FNiagaraShaderScript*);
-	typedef bool(*ValidateCompiledResultType)(EShaderPlatform, const FShaderParameterMap&, TArray<FString>&);
-	typedef void (*ModifyCompilationEnvironmentType)(EShaderPlatform,const FNiagaraShaderScript*, FShaderCompilerEnvironment&);
 
 	FNiagaraShaderType(
+		FTypeLayoutDesc& InTypeLayout,
 		const TCHAR* InName,
 		const TCHAR* InSourceFilename,
 		const TCHAR* InFunctionName,
@@ -81,13 +85,18 @@ public:
 		ConstructCompiledType InConstructCompiledRef,
 		ModifyCompilationEnvironmentType InModifyCompilationEnvironmentRef,
 		ShouldCompilePermutationType InShouldCompilePermutationRef,
-		ValidateCompiledResultType InValidateCompiledResultRef
+		ValidateCompiledResultType InValidateCompiledResultRef,
+		uint32 InTypeSize,
+		const FShaderParametersMetadata* InRootParametersMetadata = nullptr
 		):
-		FShaderType(EShaderTypeForDynamicCast::Niagara, InName, InSourceFilename, InFunctionName, SF_Compute, InTotalPermutationCount, InConstructSerializedRef, nullptr),
-		ConstructCompiledRef(InConstructCompiledRef),
-		ShouldCompilePermutationRef(InShouldCompilePermutationRef),
-		ValidateCompiledResultRef(InValidateCompiledResultRef),
-		ModifyCompilationEnvironmentRef(InModifyCompilationEnvironmentRef)
+		FShaderType(EShaderTypeForDynamicCast::Niagara, InTypeLayout, InName, InSourceFilename, InFunctionName, SF_Compute, InTotalPermutationCount,
+			InConstructSerializedRef,
+			InConstructCompiledRef,
+			InModifyCompilationEnvironmentRef,
+			InShouldCompilePermutationRef,
+			InValidateCompiledResultRef,
+			InTypeSize,
+			InRootParametersMetadata)
 	{
 		check(InTotalPermutationCount == 1);
 	}
@@ -114,7 +123,8 @@ public:
 	FShader* FinishCompileShader(
 		const FSHAHash& NiagaraShaderMapHash,
 		const FShaderCompileJob& CurrentJob,
-		const FString& InDebugDescription
+		const FString& InDebugDescription,
+		FShaderMapResourceBuilder& ResourceBuilder
 		);
 
 	/**
@@ -125,7 +135,7 @@ public:
 	 */
 	bool ShouldCache(EShaderPlatform Platform,const FNiagaraShaderScript* Script) const
 	{
-		return (*ShouldCompilePermutationRef)(Platform, Script);
+		return ShouldCompilePermutation(FNiagaraShaderPermutationParameters(Platform, Script));
 	}
 
 	/** Adds include statements for uniform buffers that this shader type references, and builds a prefix for the shader file with the include statements. */
@@ -141,17 +151,11 @@ protected:
 	 * @param Platform - Platform to compile for.
 	 * @param Environment - The shader compile environment that the function modifies.
 	 */
-	void SetupCompileEnvironment(EShaderPlatform Platform, const FNiagaraShaderScript* Script, FShaderCompilerEnvironment& Environment)
+	void SetupCompileEnvironment(EShaderPlatform Platform, const FNiagaraShaderScript* Script, FShaderCompilerEnvironment& Environment) const
 	{
-		// Allow the shader type to modify its compile environment.
-		(*ModifyCompilationEnvironmentRef)(Platform, Script, Environment);
+		ModifyCompilationEnvironment(FNiagaraShaderPermutationParameters(Platform, Script), Environment);
 	}
 
 private:
-	ConstructCompiledType ConstructCompiledRef;
-	ShouldCompilePermutationType ShouldCompilePermutationRef;
-	ValidateCompiledResultType ValidateCompiledResultRef;
-	ModifyCompilationEnvironmentType ModifyCompilationEnvironmentRef;
-
 	static TMap<const FShaderCompileJob*, TArray<FNiagaraDataInterfaceGPUParamInfo> > ExtraParamInfo;
 };
diff --git a/Engine/Plugins/FX/Niagara/Source/NiagaraShader/Public/NiagaraShared.h b/Engine/Plugins/FX/Niagara/Source/NiagaraShader/Public/NiagaraShared.h
--- a/Engine/Plugins/FX/Niagara/Source/NiagaraShader/Public/NiagaraShared.h
+++ b/Engine/Plugins/FX/Niagara/Source/NiagaraShader/Public/NiagaraShared.h
@@ -22,6 +22,7 @@
 #include "StaticParameterSet.h"
 #include "Misc/Optional.h"
 #include "NiagaraCompileHash.h"
+#include "NiagaraDataInterfaceBase.h"
 #include "NiagaraShared.generated.h"
 
 class FNiagaraShaderScript;
@@ -29,8 +30,6 @@ class FNiagaraShaderMap;
 class FNiagaraShader;
 class FNiagaraShaderMapId;
 class UNiagaraScript;
-struct FNiagaraDataInterfaceParametersCS;
-
 #define MAX_CONCURRENT_EVENT_DATASETS 4
 
 /** Defines the compile event types for translation/compilation.*/
@@ -77,6 +76,25 @@ public:
 	TArray<FGuid> StackGuids;
 };
 
+//
+class FNiagaraShaderMapPointerTable : public FShaderMapPointerTable
+{
+public:
+	using Super = FShaderMapPointerTable;
+
+	virtual int32 AddIndexedPointer(const FTypeLayoutDesc& TypeDesc, void* Ptr) override;
+	virtual void* GetIndexedPointer(const FTypeLayoutDesc& TypeDesc, uint32 i) const override;
+
+	virtual void SaveToArchive(FArchive& Ar, void* FrozenContent, bool bInlineShaderResources) const override;
+	virtual void LoadFromArchive(FArchive& Ar, void* FrozenContent, bool bInlineShaderResources, bool bLoadedByCookedMaterial) override;
+
+	TPtrTable<UNiagaraDataInterfaceBase> DITypes;
+};
+
+template<typename ShaderType>
+using TNiagaraShaderRef = TShaderRefBase<ShaderType, FNiagaraShaderMapPointerTable>;
+using FNiagaraShaderRef = TNiagaraShaderRef<FNiagaraShader>;
+
 /**
 * Information about a data interface function generated by the translator.
 */
@@ -155,75 +173,66 @@ template<> struct TStructOpsTypeTraits<FNiagaraDataInterfaceGPUParamInfo> : publ
 */
 struct FNiagaraDataInterfaceParamRef
 {
+	DECLARE_TYPE_LAYOUT(FNiagaraDataInterfaceParamRef, NonVirtual);
 public:
-
-	FNiagaraDataInterfaceParamRef(const FNiagaraDataInterfaceGPUParamInfo& InParameterInfo);
 	FNiagaraDataInterfaceParamRef();
 	~FNiagaraDataInterfaceParamRef();
-	friend bool operator<<(FArchive& Ar, FNiagaraDataInterfaceParamRef& Desc);
-
-	void Bind(const class FShaderParameterMap& ParameterMap);
-	void ConstructParameters();
-	void InitDIClass();
 
-	FNiagaraDataInterfaceGPUParamInfo ParameterInfo;
+	void Bind(const FNiagaraDataInterfaceGPUParamInfo& InParameterInfo, const class FShaderParameterMap& ParameterMap);
 
-	/** The class of this DI which we can use to create the parameter struct. */
-	UClass* DIClass;
+public:
+	/** Type of Parameters */
+	LAYOUT_FIELD(TIndexedPtr<UNiagaraDataInterfaceBase>, DIType);
 
 	/** Pointer to parameters struct for this data interface. */
-	FNiagaraDataInterfaceParametersCS* Parameters;
+	void WriteFrozenParameters(FMemoryImageWriter& Writer, const TMemoryImagePtr<FNiagaraDataInterfaceParametersCS>& InParameters) const;
+	LAYOUT_FIELD_WITH_WRITER(TMemoryImagePtr<FNiagaraDataInterfaceParametersCS>, Parameters, WriteFrozenParameters);
 };
 
 
-
-
 /** Stores outputs from the script compile that need to be saved. */
 class FNiagaraComputeShaderCompilationOutput
 {
+	DECLARE_TYPE_LAYOUT(FNiagaraComputeShaderCompilationOutput, NonVirtual);
 public:
 	FNiagaraComputeShaderCompilationOutput()
 	{}
-
-	NIAGARASHADER_API void Serialize(FArchive& Ar)
-	{}
 };
 
-
-
 /** Contains all the information needed to uniquely identify a FNiagaraShaderMapID. */
 class FNiagaraShaderMapId
 {
+	DECLARE_TYPE_LAYOUT(FNiagaraShaderMapId, NonVirtual);
 public:
 	/** The version of the compiler that this needs to be built against.*/
-	FGuid CompilerVersionID;
+	LAYOUT_FIELD(FGuid, CompilerVersionID);
 
 	/** Feature level that the shader map is going to be compiled for.  */
-	ERHIFeatureLevel::Type FeatureLevel;
+	LAYOUT_FIELD(ERHIFeatureLevel::Type, FeatureLevel);
 
 	/**
 	* The base id of the subgraph this shader primarily represents.
 	*/
-	FGuid BaseScriptID_DEPRECATED;
+	LAYOUT_FIELD(FGuid, BaseScriptID_DEPRECATED);
 
 	/** Configuration options */
-	TArray<FString> AdditionalDefines;
+	LAYOUT_FIELD(TMemoryImageArray<FMemoryImageString>, AdditionalDefines);
 
 	/**
 	* The hash of the subgraph this shader primarily represents.
 	*/
-	FNiagaraCompileHash BaseCompileHash;
+	LAYOUT_FIELD(FSHAHash, BaseCompileHash);
 
 	/** The compile hashes of the top level scripts the script is dependent on. */
-	TArray<FNiagaraCompileHash> ReferencedCompileHashes;
+	LAYOUT_FIELD(TMemoryImageArray<FSHAHash>, ReferencedCompileHashes);
 
 	
 	/** Bitfield of supported detail levels in this compile.*/
-	uint32 DetailLevelMask = 0xFFFFFFFF;
+	LAYOUT_FIELD_INITIALIZED(uint32, DetailLevelMask, 0xFFFFFFFF);
 
 	/** Whether or not we need to bake Rapid Iteration params. True to keep params, false to bake.*/
-	bool bUsesRapidIterationParams = true;
-	
+	LAYOUT_FIELD_INITIALIZED(bool, bUsesRapidIterationParams, true);
+
 	FNiagaraShaderMapId()
 		: CompilerVersionID()
 		, FeatureLevel(GMaxRHIFeatureLevel)
@@ -234,11 +243,11 @@ public:
 
 	//ENGINE_API void SetShaderDependencies(const TArray<FShaderType*>& ShaderTypes, const TArray<const FShaderPipelineType*>& ShaderPipelineTypes, const TArray<FVertexFactoryType*>& VFTypes);
 
-	void Serialize(FArchive& Ar);
+	//void Serialize(FArchive& Ar);
 
 	friend uint32 GetTypeHash(const FNiagaraShaderMapId& Ref)
 	{
-		return Ref.BaseCompileHash.GetTypeHash();
+		return GetTypeHash(Ref.BaseCompileHash.Hash);
 	}
 
 	SIZE_T GetSizeBytes() const
@@ -335,17 +344,34 @@ private:
 
 #endif
 
+class FNiagaraShaderMapContent : public FShaderMapContent
+{
+	using Super = FShaderMapContent;
+	friend class FNiagaraShaderMap;
+	DECLARE_TYPE_LAYOUT(FNiagaraShaderMapContent, NonVirtual);
+private:
+	explicit FNiagaraShaderMapContent(const EShaderPlatform InPlatform) : Super(InPlatform) {}
+	
+	/** The script's user friendly name, typically the object name. */
+	LAYOUT_FIELD(FMemoryImageString, FriendlyName);
 
+	/** Debug information about how the shader map was compiled. */
+	LAYOUT_FIELD(FMemoryImageString, DebugDescription);
 
+	/** The static parameter set that this shader map was compiled with */
+	LAYOUT_FIELD(FNiagaraShaderMapId, ShaderMapId);
 
-
+	/** Shader compilation output */
+	LAYOUT_FIELD(FNiagaraComputeShaderCompilationOutput, NiagaraCompilationOutput);
+};
 
 /**
 * The set of shaders for a single script.
 */
-class FNiagaraShaderMap : public TShaderMap<FNiagaraShaderType>, public FDeferredCleanupInterface
+class FNiagaraShaderMap : public TShaderMap<FNiagaraShaderMapContent, FNiagaraShaderMapPointerTable>, public FDeferredCleanupInterface
 {
 public:
+	using Super = TShaderMap<FNiagaraShaderMapContent, FNiagaraShaderMapPointerTable>;
 
 	/**
 	* Finds the shader map for a script.
@@ -355,10 +381,13 @@ public:
 	static FNiagaraShaderMap* FindId(const FNiagaraShaderMapId& ShaderMapId, EShaderPlatform Platform);
 
 	/** Flushes the given shader types from any loaded FNiagaraShaderMaps. */
-	static void FlushShaderTypes(TArray<FShaderType*>& ShaderTypesToFlush);
+	static void FlushShaderTypes(TArray<const FShaderType*>& ShaderTypesToFlush);
+
+	// ShaderMap interface
+	template<typename ShaderType> TNiagaraShaderRef<ShaderType> GetShader() const { return TNiagaraShaderRef<ShaderType>(GetContent()->GetShader<ShaderType>(), *this); }
+	TNiagaraShaderRef<FShader> GetShader(FShaderType* ShaderType) const { return TNiagaraShaderRef<FShader>(GetContent()->GetShader(ShaderType), *this); }
 
-	static void FixupShaderTypes(EShaderPlatform Platform,
-		const TMap<FShaderType*, FString>& ShaderTypeNames);
+	//static void FixupShaderTypes(EShaderPlatform Platform, const TMap<FShaderType*, FString>& ShaderTypeNames);
 
 	/**
 	* Attempts to load the shader map for the given script from the Derived Data Cache.
@@ -409,7 +438,7 @@ public:
 	bool TryToAddToExistingCompilationTask(FNiagaraShaderScript* Script);
 
 	/** Builds a list of the shaders in a shader map. */
-	NIAGARASHADER_API  void GetShaderList(TMap<FShaderId, FShader*>& OutShaders) const;
+	NIAGARASHADER_API  void GetShaderList(TMap<FShaderId, TShaderRef<FShader>>& OutShaders) const;
 
 	/** Builds a list of the shader pipelines in a shader map. */
 	//ENGINE_API void GetShaderPipelineList(TArray<FShaderPipeline*>& OutShaderPipelines) const;
@@ -425,7 +454,7 @@ public:
 	* Removes all entries in the cache with exceptions based on a shader type
 	* @param ShaderType - The shader type to flush
 	*/
-	void FlushShadersByShaderType(FShaderType* ShaderType);
+	void FlushShadersByShaderType(const FShaderType* ShaderType);
 
 	/** Removes a Script from NiagaraShaderMapsBeingCompiled. */
 	NIAGARASHADER_API static void RemovePendingScript(FNiagaraShaderScript* Script);
@@ -440,10 +469,6 @@ public:
 	/** Saves this shader map to the derived data cache. */
 	void SaveToDerivedDataCache();
 
-	/** Registers all shaders that have been loaded in Serialize */
-	virtual void RegisterSerializedShaders(bool bCookedMaterial) override;
-	virtual void DiscardSerializedShaders() override;
-
 	/** Backs up any FShaders in this shader map to memory through serialization and clears FShader references. */
 	TArray<uint8>* BackupShadersToMemory();
 	/** Recreates FShaders from the passed in memory, handling shader key changes. */
@@ -453,25 +478,18 @@ public:
 	NIAGARASHADER_API  static void SaveForRemoteRecompile(FArchive& Ar, const TMap<FString, TArray<TRefCountPtr<FNiagaraShaderMap> > >& CompiledShaderMaps, const TArray<FShaderResourceId>& ClientResourceIds);
 	NIAGARASHADER_API  static void LoadForRemoteRecompile(FArchive& Ar, EShaderPlatform ShaderPlatform, const TArray<FString>& ScriptsForShaderMaps);
 
-	/** Computes the memory used by this shader map without counting the shaders themselves. */
-	uint32 GetSizeBytes() const
-	{
-		return sizeof(*this)
-			+FriendlyName.GetAllocatedSize();
-	}
-
 	/** Returns the maximum number of texture samplers used by any shader in this shader map. */
 	uint32 GetMaxTextureSamplers() const;
 
 
 	// Accessors.
-	const FNiagaraShaderMapId& GetShaderMapId() const { return ShaderMapId; }
-	EShaderPlatform GetShaderPlatform() const	{ return Platform; }
-	const FString& GetFriendlyName() const		{ return FriendlyName; }
-	uint32 GetCompilingId() const				{ return CompilingId; }
-	bool IsCompilationFinalized() const			{ return bCompilationFinalized; }
-	bool CompiledSuccessfully() const			{ return bCompiledSuccessfully; }
-	const FString& GetDebugDescription() const	{ return DebugDescription; }
+	const FNiagaraShaderMapId& GetShaderMapId() const		{ return GetContent()->ShaderMapId; }
+	EShaderPlatform GetShaderPlatform() const				{ return GetContent()->GetShaderPlatform(); }
+	const FMemoryImageString& GetFriendlyName() const		{ return GetContent()->FriendlyName; }
+	uint32 GetCompilingId() const							{ return CompilingId; }
+	bool IsCompilationFinalized() const						{ return bCompilationFinalized; }
+	bool CompiledSuccessfully() const						{ return bCompiledSuccessfully; }
+	const FMemoryImageString& GetDebugDescription() const	{ return GetContent()->DebugDescription; }
 
 	bool IsValid() const
 	{
@@ -505,18 +523,6 @@ private:
 	*/
 	static TArray<FNiagaraShaderMap*> AllNiagaraShaderMaps;
 
-	/** The script's user friendly name, typically the object name. */
-	FString FriendlyName;
-
-	/** The platform this shader map was compiled with */
-	EShaderPlatform Platform;
-
-	/** The static parameter set that this shader map was compiled with */
-	FNiagaraShaderMapId ShaderMapId;
-
-	/** Shader compilation output */
-	FNiagaraComputeShaderCompilationOutput NiagaraCompilationOutput;
-
 	/** Tracks resources and their shader maps that need to be compiled but whose compilation is being deferred. */
 	static TMap<TRefCountPtr<FNiagaraShaderMap>, TArray<FNiagaraShaderScript*> > NiagaraShaderMapsBeingCompiled;
 
@@ -542,16 +548,10 @@ private:
 	/** Indicates whether the shader map should be stored in the shader cache. */
 	uint32 bIsPersistent : 1;
 
-	/** Debug information about how the shader map was compiled. */
-	FString DebugDescription;
-
-	FShader* ProcessCompilationResultsForSingleJob(TSharedRef<class FShaderCommonCompileJob, ESPMode::ThreadSafe> SingleJob, const FSHAHash& ShaderMapHash);
+	FShader* ProcessCompilationResultsForSingleJob(TSharedRef<class FShaderCommonCompileJob, ESPMode::ThreadSafe> SingleJob, const FSHAHash& ShaderMapHash, FShaderMapResourceBuilder& InResourceBuilder);
 
 	bool IsNiagaraShaderComplete(const FNiagaraShaderScript* Script, const FNiagaraShaderType* ShaderType, bool bSilent);
 
-	/** Initializes OrderedMeshShaderMaps from the contents of MeshShaderMaps. */
-	void InitOrderedMeshShaderMaps();
-
 	friend NIAGARASHADER_API  void DumpNiagaraStats(EShaderPlatform Platform);
 	friend class FShaderCompilingManager;
 };
@@ -660,8 +660,6 @@ public:
 		return GameThreadShaderMap;
 	}
 
-	NIAGARASHADER_API void RegisterShaderMap();
-
 	NIAGARASHADER_API void DiscardShaderMap();
 
 	/** Note: SetRenderingThreadShaderMap must also be called with the same value, but from the rendering thread. */
@@ -734,11 +732,10 @@ public:
 	
 	FString HlslOutput;
 
-	NIAGARASHADER_API  FNiagaraShader* GetShader() const;
-	NIAGARASHADER_API  FNiagaraShader* GetShaderGameThread() const;
+	NIAGARASHADER_API  FNiagaraShaderRef GetShader() const;
+	NIAGARASHADER_API  FNiagaraShaderRef GetShaderGameThread() const;
 	
-	NIAGARASHADER_API void SetDataInterfaceParamInfo(TArray< FNiagaraDataInterfaceGPUParamInfo >& InDIParamInfo);
-	NIAGARASHADER_API void SetDataInterfaceParamInfo(TArray< FNiagaraDataInterfaceParamRef >& InDIParamRefs);
+	NIAGARASHADER_API void SetDataInterfaceParamInfo(const TArray< FNiagaraDataInterfaceGPUParamInfo >& InDIParamInfo);
 
 	TArray< FNiagaraDataInterfaceGPUParamInfo >& GetDataInterfaceParamInfo()
 	{
diff --git a/Engine/Plugins/FX/Niagara/Source/NiagaraVertexFactories/Private/NiagaraDrawIndirect.cpp b/Engine/Plugins/FX/Niagara/Source/NiagaraVertexFactories/Private/NiagaraDrawIndirect.cpp
--- a/Engine/Plugins/FX/Niagara/Source/NiagaraVertexFactories/Private/NiagaraDrawIndirect.cpp
+++ b/Engine/Plugins/FX/Niagara/Source/NiagaraVertexFactories/Private/NiagaraDrawIndirect.cpp
@@ -27,7 +27,7 @@ FNiagaraDrawIndirectArgsGenCS::FNiagaraDrawIndirectArgsGenCS(const ShaderMetaTyp
 	TaskCountParam.Bind(Initializer.ParameterMap, TEXT("TaskCount"));
 }
 
-bool FNiagaraDrawIndirectArgsGenCS::Serialize(FArchive& Ar)
+/*bool FNiagaraDrawIndirectArgsGenCS::Serialize(FArchive& Ar)
 {
 	bool bShaderHasOutdatedParameters = FGlobalShader::Serialize(Ar);
 	Ar << TaskInfosParam;
@@ -35,11 +35,11 @@ bool FNiagaraDrawIndirectArgsGenCS::Serialize(FArchive& Ar)
 	Ar << DrawIndirectArgsParam;
 	Ar << TaskCountParam;
 	return bShaderHasOutdatedParameters;
-}
+}*/
 
 void FNiagaraDrawIndirectArgsGenCS::SetOutput(FRHICommandList& RHICmdList, FRHIUnorderedAccessView* DrawIndirectArgsUAV, FRHIUnorderedAccessView* InstanceCountsUAV)
 {
-	FRHIComputeShader* ComputeShaderRHI = GetComputeShader();
+	FRHIComputeShader* ComputeShaderRHI = RHICmdList.GetBoundComputeShader();
 	if (DrawIndirectArgsParam.IsBound())
 	{
 		RHICmdList.SetUAVParameter(ComputeShaderRHI, DrawIndirectArgsParam.GetUAVIndex(), DrawIndirectArgsUAV);
@@ -52,7 +52,7 @@ void FNiagaraDrawIndirectArgsGenCS::SetOutput(FRHICommandList& RHICmdList, FRHIU
 
 void FNiagaraDrawIndirectArgsGenCS::SetParameters(FRHICommandList& RHICmdList, FRHIShaderResourceView* TaskInfosBuffer, int32 NumArgGenTasks, int32 NumInstanceCountClearTasks)
 {
-	FRHIComputeShader* ComputeShaderRHI = GetComputeShader();
+	FRHIComputeShader* ComputeShaderRHI = RHICmdList.GetBoundComputeShader();
 
 	SetSRVParameter(RHICmdList, ComputeShaderRHI, TaskInfosParam, TaskInfosBuffer);
 
@@ -62,7 +62,7 @@ void FNiagaraDrawIndirectArgsGenCS::SetParameters(FRHICommandList& RHICmdList, F
 
 void FNiagaraDrawIndirectArgsGenCS::UnbindBuffers(FRHICommandList& RHICmdList)
 {
-	FRHIComputeShader* ComputeShaderRHI = GetComputeShader();
+	FRHIComputeShader* ComputeShaderRHI = RHICmdList.GetBoundComputeShader();
 
 	SetSRVParameter(RHICmdList, ComputeShaderRHI, TaskInfosParam, nullptr);
 	if (DrawIndirectArgsParam.IsBound())
diff --git a/Engine/Plugins/FX/Niagara/Source/NiagaraVertexFactories/Private/NiagaraMeshVertexFactory.cpp b/Engine/Plugins/FX/Niagara/Source/NiagaraVertexFactories/Private/NiagaraMeshVertexFactory.cpp
--- a/Engine/Plugins/FX/Niagara/Source/NiagaraVertexFactories/Private/NiagaraMeshVertexFactory.cpp
+++ b/Engine/Plugins/FX/Niagara/Source/NiagaraVertexFactories/Private/NiagaraMeshVertexFactory.cpp
@@ -14,8 +14,9 @@ IMPLEMENT_GLOBAL_SHADER_PARAMETER_STRUCT(FNiagaraMeshUniformParameters,"NiagaraM
 
 class FNiagaraMeshVertexFactoryShaderParametersVS : public FVertexFactoryShaderParameters
 {
+	DECLARE_INLINE_TYPE_LAYOUT(FNiagaraMeshVertexFactoryShaderParametersVS, NonVirtual);
 public:
-	virtual void Bind(const FShaderParameterMap& ParameterMap) override
+	void Bind(const FShaderParameterMap& ParameterMap)
 	{
 		//PrevTransformBuffer.Bind(ParameterMap, TEXT("PrevTransformBuffer"));
 		NiagaraParticleDataPosition.Bind(ParameterMap, TEXT("NiagaraParticleDataPosition"));
@@ -42,34 +43,7 @@ public:
 		SortedIndicesOffset.Bind(ParameterMap, TEXT("SortedIndicesOffset"));
 	}
 
-	virtual void Serialize(FArchive& Ar) override
-	{
-		//Ar << PrevTransformBuffer;
-		Ar << NiagaraParticleDataPosition;
-		Ar << NiagaraParticleDataVelocity;
-		Ar << NiagaraParticleDataColor;
-		Ar << NiagaraParticleDataScale;
-		Ar << NiagaraParticleDataTransform;
-		Ar << NiagaraParticleDataNormalizedAge;
-		Ar << NiagaraParticleDataMaterialRandom;
-		Ar << NiagaraParticleDataMaterialParam0;
-		Ar << NiagaraParticleDataMaterialParam1;
-		Ar << NiagaraParticleDataMaterialParam2;
-		Ar << NiagaraParticleDataMaterialParam3;
-		Ar << NiagaraParticleDataSubImage;
-		Ar << FloatDataStride;
-
-		// 		Ar << NiagaraParticleDataInt;
-		// 		Ar << Int32DataOffset;
-		// 		Ar << Int32DataStride;
-
-		Ar << MeshFacingMode;
-
-		Ar << SortedIndices;
-		Ar << SortedIndicesOffset;
-	}
-
-	virtual void GetElementShaderBindings(
+	void GetElementShaderBindings(
 		const FSceneInterface* Scene,
 		const FSceneView* View,
 		const FMeshMaterialShader* Shader,
@@ -78,7 +52,7 @@ public:
 		const FVertexFactory* VertexFactory,
 		const FMeshBatchElement& BatchElement,
 		class FMeshDrawSingleShaderBindings& ShaderBindings,
-		FVertexInputStreamArray& VertexStreams) const override
+		FVertexInputStreamArray& VertexStreams) const
 	{
 		FNiagaraMeshVertexFactory* NiagaraMeshVF = (FNiagaraMeshVertexFactory*)VertexFactory;
 		ShaderBindings.Add(Shader->GetUniformBufferParameter<FNiagaraMeshUniformParameters>(), NiagaraMeshVF->GetUniformBuffer());
@@ -107,43 +81,38 @@ public:
 
 private:
 
-	//FShaderResourceParameter PrevTransformBuffer;
-
-	FShaderResourceParameter NiagaraParticleDataPosition;
-	FShaderResourceParameter NiagaraParticleDataVelocity;
-	FShaderResourceParameter NiagaraParticleDataColor;
-	FShaderResourceParameter NiagaraParticleDataScale;
-	FShaderResourceParameter NiagaraParticleDataTransform;
-	FShaderResourceParameter NiagaraParticleDataNormalizedAge;
-	FShaderResourceParameter NiagaraParticleDataMaterialRandom;
-	FShaderResourceParameter NiagaraParticleDataMaterialParam0;
-	FShaderResourceParameter NiagaraParticleDataMaterialParam1;
-	FShaderResourceParameter NiagaraParticleDataMaterialParam2;
-	FShaderResourceParameter NiagaraParticleDataMaterialParam3;
-	FShaderResourceParameter NiagaraParticleDataSubImage;
-	FShaderParameter FloatDataStride;
-
-	// 	FShaderResourceParameter NiagaraParticleDataInt;
-	// 	FShaderParameter Int32DataOffset;
-	// 	FShaderParameter Int32DataStride;
-
-	FShaderParameter MeshFacingMode;
-	FShaderResourceParameter SortedIndices;
-	FShaderParameter SortedIndicesOffset;
+	
+		//LAYOUT_FIELD(FShaderResourceParameter, PrevTransformBuffer)
+
+	LAYOUT_FIELD(FShaderResourceParameter, NiagaraParticleDataPosition);
+	LAYOUT_FIELD(FShaderResourceParameter, NiagaraParticleDataVelocity);
+	LAYOUT_FIELD(FShaderResourceParameter, NiagaraParticleDataColor);
+	LAYOUT_FIELD(FShaderResourceParameter, NiagaraParticleDataScale);
+	LAYOUT_FIELD(FShaderResourceParameter, NiagaraParticleDataTransform);
+	LAYOUT_FIELD(FShaderResourceParameter, NiagaraParticleDataNormalizedAge);
+	LAYOUT_FIELD(FShaderResourceParameter, NiagaraParticleDataMaterialRandom);
+	LAYOUT_FIELD(FShaderResourceParameter, NiagaraParticleDataMaterialParam0);
+	LAYOUT_FIELD(FShaderResourceParameter, NiagaraParticleDataMaterialParam1);
+	LAYOUT_FIELD(FShaderResourceParameter, NiagaraParticleDataMaterialParam2);
+	LAYOUT_FIELD(FShaderResourceParameter, NiagaraParticleDataMaterialParam3);
+	LAYOUT_FIELD(FShaderResourceParameter, NiagaraParticleDataSubImage);
+	LAYOUT_FIELD(FShaderParameter, FloatDataStride);
+
+		// 	LAYOUT_FIELD(FShaderResourceParameter, NiagaraParticleDataInt)
+		// 	LAYOUT_FIELD(FShaderParameter, Int32DataOffset)
+		// 	LAYOUT_FIELD(FShaderParameter, Int32DataStride)
+
+		LAYOUT_FIELD(FShaderParameter, MeshFacingMode)
+		LAYOUT_FIELD(FShaderResourceParameter, SortedIndices)
+		LAYOUT_FIELD(FShaderParameter, SortedIndicesOffset)
+	
 };
 
 class FNiagaraMeshVertexFactoryShaderParametersPS : public FVertexFactoryShaderParameters
 {
+	DECLARE_INLINE_TYPE_LAYOUT(FNiagaraMeshVertexFactoryShaderParametersPS, NonVirtual);
 public:
-	virtual void Bind(const FShaderParameterMap& ParameterMap) override
-	{
-	}
-
-	virtual void Serialize(FArchive& Ar) override
-	{
-	}
-
-	virtual void GetElementShaderBindings(
+	void GetElementShaderBindings(
 		const FSceneInterface* Scene,
 		const FSceneView* View,
 		const FMeshMaterialShader* Shader,
@@ -152,11 +121,14 @@ public:
 		const FVertexFactory* VertexFactory,
 		const FMeshBatchElement& BatchElement,
 		class FMeshDrawSingleShaderBindings& ShaderBindings,
-		FVertexInputStreamArray& VertexStreams) const override
+		FVertexInputStreamArray& VertexStreams) const
 	{
 		FNiagaraMeshVertexFactory* NiagaraMeshVF = (FNiagaraMeshVertexFactory*)VertexFactory;
 		ShaderBindings.Add(Shader->GetUniformBufferParameter<FNiagaraMeshUniformParameters>(), NiagaraMeshVF->GetUniformBuffer());
 	}
+
+	
+	
 };
 
 void FNiagaraMeshVertexFactory::InitRHI()
@@ -259,9 +231,15 @@ FRHIShaderResourceView* FNiagaraMeshVertexFactory::GetPreviousTransformBufferSRV
 }
 */
 
-bool FNiagaraMeshVertexFactory::ShouldCompilePermutation(EShaderPlatform Platform, const class FMaterial* Material, const class FShaderType* ShaderType)
+bool FNiagaraMeshVertexFactory::ShouldCompilePermutation(const FVertexFactoryShaderPermutationParameters& Parameters)
 {
-	return (FNiagaraUtilities::SupportsNiagaraRendering(Platform)) && (Material->IsUsedWithNiagaraMeshParticles() || Material->IsSpecialEngineMaterial());
+	return (FNiagaraUtilities::SupportsNiagaraRendering(Parameters.Platform)) && (Parameters.MaterialParameters.bIsUsedWithNiagaraMeshParticles || Parameters.MaterialParameters.bIsSpecialEngineMaterial);
+}
+
+bool FNiagaraMeshVertexFactoryEmulatedInstancing::ShouldCompilePermutation(const FVertexFactoryShaderPermutationParameters& Parameters)
+{
+	return (Parameters.Platform == SP_OPENGL_ES2_ANDROID || Parameters.Platform == SP_OPENGL_ES2_WEBGL) // Those are only platforms that might not support hardware instancing
+		&& FNiagaraMeshVertexFactory::ShouldCompilePermutation(Parameters);
 }
 
 void FNiagaraMeshVertexFactory::SetData(const FStaticMeshDataType& InData)
@@ -271,19 +249,11 @@ void FNiagaraMeshVertexFactory::SetData(const FStaticMeshDataType& InData)
 	UpdateRHI();
 }
 
+IMPLEMENT_VERTEX_FACTORY_PARAMETER_TYPE(FNiagaraMeshVertexFactory, SF_Vertex, FNiagaraMeshVertexFactoryShaderParametersVS);
+IMPLEMENT_VERTEX_FACTORY_PARAMETER_TYPE(FNiagaraMeshVertexFactory, SF_Pixel, FNiagaraMeshVertexFactoryShaderParametersPS);
 
-FVertexFactoryShaderParameters* FNiagaraMeshVertexFactory::ConstructShaderParameters(EShaderFrequency ShaderFrequency)
-{
-	if (ShaderFrequency == SF_Vertex)
-	{
-		return new FNiagaraMeshVertexFactoryShaderParametersVS();
-	}
-	else if (ShaderFrequency == SF_Pixel)
-	{
-		return new FNiagaraMeshVertexFactoryShaderParametersPS();
-	}
-	return nullptr;
-}
+IMPLEMENT_VERTEX_FACTORY_PARAMETER_TYPE(FNiagaraMeshVertexFactoryEmulatedInstancing, SF_Vertex, FNiagaraMeshVertexFactoryShaderParametersVS);
+IMPLEMENT_VERTEX_FACTORY_PARAMETER_TYPE(FNiagaraMeshVertexFactoryEmulatedInstancing, SF_Pixel, FNiagaraMeshVertexFactoryShaderParametersPS);
 
 IMPLEMENT_VERTEX_FACTORY_TYPE(FNiagaraMeshVertexFactory, "/Plugin/FX/Niagara/Private/NiagaraMeshVertexFactory.ush", true, false, true, false, false);
 IMPLEMENT_VERTEX_FACTORY_TYPE(FNiagaraMeshVertexFactoryEmulatedInstancing, "/Plugin/FX/Niagara/Private/NiagaraMeshVertexFactory.ush", true, false, true, false, false);
diff --git a/Engine/Plugins/FX/Niagara/Source/NiagaraVertexFactories/Private/NiagaraRibbonVertexFactory.cpp b/Engine/Plugins/FX/Niagara/Source/NiagaraVertexFactories/Private/NiagaraRibbonVertexFactory.cpp
--- a/Engine/Plugins/FX/Niagara/Source/NiagaraVertexFactories/Private/NiagaraRibbonVertexFactory.cpp
+++ b/Engine/Plugins/FX/Niagara/Source/NiagaraVertexFactories/Private/NiagaraRibbonVertexFactory.cpp
@@ -16,14 +16,9 @@ IMPLEMENT_GLOBAL_SHADER_PARAMETER_STRUCT(FNiagaraRibbonVFLooseParameters, "Niaga
 
 class FNiagaraRibbonVertexFactoryShaderParameters : public FVertexFactoryShaderParameters
 {
+	DECLARE_INLINE_TYPE_LAYOUT(FNiagaraRibbonVertexFactoryShaderParameters, NonVirtual);
 public:
-	virtual void Bind(const FShaderParameterMap& ParameterMap) override
-	{
-	}
 
-	virtual void Serialize(FArchive& Ar) override
-	{
-	}
 };
 
 /**
@@ -31,8 +26,9 @@ public:
 */
 class FNiagaraRibbonVertexFactoryShaderParametersVS : public FNiagaraRibbonVertexFactoryShaderParameters
 {
+	DECLARE_INLINE_TYPE_LAYOUT(FNiagaraRibbonVertexFactoryShaderParametersVS, NonVirtual);
 public:
-	virtual void Bind(const FShaderParameterMap& ParameterMap) override
+	void Bind(const FShaderParameterMap& ParameterMap)
 	{
 		NiagaraParticleDataPosition.Bind(ParameterMap, TEXT("NiagaraParticleDataPosition"));
 		NiagaraParticleDataVelocity.Bind(ParameterMap, TEXT("NiagaraParticleDataVelocity"));
@@ -50,24 +46,7 @@ public:
 		FloatDataStride.Bind(ParameterMap, TEXT("NiagaraFloatDataStride"));
 	}
 
-	virtual void Serialize(FArchive& Ar) override
-	{
-		Ar << NiagaraParticleDataPosition;
-		Ar << NiagaraParticleDataVelocity;
-		Ar << NiagaraParticleDataColor;
-		Ar << NiagaraParticleDataWidth;
-		Ar << NiagaraParticleDataTwist;
-		Ar << NiagaraParticleDataFacing;
-		Ar << NiagaraParticleDataNormalizedAge;
-		Ar << NiagaraParticleDataMaterialRandom;
-		Ar << NiagaraParticleDataMaterialParam0;
-		Ar << NiagaraParticleDataMaterialParam1;
-		Ar << NiagaraParticleDataMaterialParam2;
-		Ar << NiagaraParticleDataMaterialParam3;
-		Ar << FloatDataStride;
-	}
-
-	virtual void GetElementShaderBindings(
+	void GetElementShaderBindings(
 		const FSceneInterface* Scene,
 		const FSceneView* View,
 		const FMeshMaterialShader* Shader,
@@ -76,7 +55,7 @@ public:
 		const FVertexFactory* VertexFactory,
 		const FMeshBatchElement& BatchElement,
 		class FMeshDrawSingleShaderBindings& ShaderBindings,
-		FVertexInputStreamArray& VertexStreams) const override
+		FVertexInputStreamArray& VertexStreams) const
 	{
 		FNiagaraRibbonVertexFactory* RibbonVF = (FNiagaraRibbonVertexFactory*)VertexFactory;
 		ShaderBindings.Add(Shader->GetUniformBufferParameter<FNiagaraRibbonUniformParameters>(), RibbonVF->GetRibbonUniformBuffer());
@@ -97,20 +76,20 @@ public:
 	}
 
 private:
-	FShaderResourceParameter NiagaraParticleDataPosition;
-	FShaderResourceParameter NiagaraParticleDataVelocity;
-	FShaderResourceParameter NiagaraParticleDataColor;
-	FShaderResourceParameter NiagaraParticleDataWidth;
-	FShaderResourceParameter NiagaraParticleDataTwist;
-	FShaderResourceParameter NiagaraParticleDataFacing;
-	FShaderResourceParameter NiagaraParticleDataNormalizedAge;
-	FShaderResourceParameter NiagaraParticleDataMaterialRandom;
-	FShaderResourceParameter NiagaraParticleDataMaterialParam0;
-	FShaderResourceParameter NiagaraParticleDataMaterialParam1;
-	FShaderResourceParameter NiagaraParticleDataMaterialParam2;
-	FShaderResourceParameter NiagaraParticleDataMaterialParam3;
-
-	FShaderParameter FloatDataStride;
+	LAYOUT_FIELD(FShaderResourceParameter, NiagaraParticleDataPosition);
+	LAYOUT_FIELD(FShaderResourceParameter, NiagaraParticleDataVelocity);
+	LAYOUT_FIELD(FShaderResourceParameter, NiagaraParticleDataColor);
+	LAYOUT_FIELD(FShaderResourceParameter, NiagaraParticleDataWidth);
+	LAYOUT_FIELD(FShaderResourceParameter, NiagaraParticleDataTwist);
+	LAYOUT_FIELD(FShaderResourceParameter, NiagaraParticleDataFacing);
+	LAYOUT_FIELD(FShaderResourceParameter, NiagaraParticleDataNormalizedAge);
+	LAYOUT_FIELD(FShaderResourceParameter, NiagaraParticleDataMaterialRandom);
+	LAYOUT_FIELD(FShaderResourceParameter, NiagaraParticleDataMaterialParam0);
+	LAYOUT_FIELD(FShaderResourceParameter, NiagaraParticleDataMaterialParam1);
+	LAYOUT_FIELD(FShaderResourceParameter, NiagaraParticleDataMaterialParam2);
+	LAYOUT_FIELD(FShaderResourceParameter, NiagaraParticleDataMaterialParam3);
+
+	LAYOUT_FIELD(FShaderParameter, FloatDataStride);
 };
 
 
@@ -120,16 +99,9 @@ private:
 */
 class FNiagaraRibbonVertexFactoryShaderParametersPS : public FNiagaraRibbonVertexFactoryShaderParameters
 {
+	DECLARE_INLINE_TYPE_LAYOUT(FNiagaraRibbonVertexFactoryShaderParametersPS, NonVirtual);
 public:
-	virtual void Bind(const FShaderParameterMap& ParameterMap) override
-	{
-	}
-
-	virtual void Serialize(FArchive& Ar) override
-	{
-	}
-
-	virtual void GetElementShaderBindings(
+	 void GetElementShaderBindings(
 		const FSceneInterface* Scene,
 		const FSceneView* View,
 		const FMeshMaterialShader* Shader,
@@ -138,7 +110,7 @@ public:
 		const FVertexFactory* VertexFactory,
 		const FMeshBatchElement& BatchElement,
 		class FMeshDrawSingleShaderBindings& ShaderBindings,
-		FVertexInputStreamArray& VertexStreams) const override
+		FVertexInputStreamArray& VertexStreams) const
 	{
 		FNiagaraRibbonVertexFactory* RibbonVF = (FNiagaraRibbonVertexFactory*)VertexFactory;
 		ShaderBindings.Add(Shader->GetUniformBufferParameter<FNiagaraRibbonUniformParameters>(), RibbonVF->GetRibbonUniformBuffer());
@@ -185,17 +157,17 @@ static TGlobalResource<FNiagaraRibbonVertexDeclaration> GNiagaraRibbonVertexDecl
 
 ///////////////////////////////////////////////////////////////////////////////
 
-bool FNiagaraRibbonVertexFactory::ShouldCompilePermutation(EShaderPlatform Platform, const class FMaterial* Material, const class FShaderType* ShaderType)
+bool FNiagaraRibbonVertexFactory::ShouldCompilePermutation(const FVertexFactoryShaderPermutationParameters& Parameters)
 {
-	return (FNiagaraUtilities::SupportsNiagaraRendering(Platform)) && (Material->IsUsedWithNiagaraRibbons() || Material->IsSpecialEngineMaterial());
+	return (FNiagaraUtilities::SupportsNiagaraRendering(Parameters.Platform)) && (Parameters.MaterialParameters.bIsUsedWithNiagaraRibbons || Parameters.MaterialParameters.bIsSpecialEngineMaterial);
 }
 
 /**
 * Can be overridden by FVertexFactory subclasses to modify their compile environment just before compilation occurs.
 */
-void FNiagaraRibbonVertexFactory::ModifyCompilationEnvironment(const FVertexFactoryType* Type, EShaderPlatform Platform, const class FMaterial* Material, FShaderCompilerEnvironment& OutEnvironment)
+void FNiagaraRibbonVertexFactory::ModifyCompilationEnvironment(const FVertexFactoryShaderPermutationParameters& Parameters, FShaderCompilerEnvironment& OutEnvironment)
 {
-	FNiagaraVertexFactoryBase::ModifyCompilationEnvironment(Type, Platform, Material, OutEnvironment);
+	FNiagaraVertexFactoryBase::ModifyCompilationEnvironment(Parameters, OutEnvironment);
 
 	OutEnvironment.SetDefine(TEXT("NiagaraVFLooseParameters"), TEXT("NiagaraRibbonVFLooseParameters"));
 	
@@ -216,29 +188,6 @@ void FNiagaraRibbonVertexFactory::InitRHI()
 	FVertexStream* DynamicParameter3Stream = new(Streams) FVertexStream;
 }
 
-FVertexFactoryShaderParameters* FNiagaraRibbonVertexFactory::ConstructShaderParameters(EShaderFrequency ShaderFrequency)
-{
-	if (ShaderFrequency == SF_Vertex)
-	{
-		return new FNiagaraRibbonVertexFactoryShaderParametersVS();
-	}
-	else if (ShaderFrequency == SF_Pixel)
-	{
-		return new FNiagaraRibbonVertexFactoryShaderParametersPS();
-	}
-#if RHI_RAYTRACING
-	else if (ShaderFrequency == SF_Compute)
-	{
-		return new FNiagaraRibbonVertexFactoryShaderParametersVS();
-	}
-	else if (ShaderFrequency == SF_RayHitGroup)
-	{
-		return new FNiagaraRibbonVertexFactoryShaderParametersVS();
-	}
-#endif
-	return NULL;
-}
-
 void FNiagaraRibbonVertexFactory::SetVertexBuffer(const FVertexBuffer* InBuffer, uint32 StreamOffset, uint32 Stride)
 {
 	check(Streams.Num() == 5);
@@ -268,4 +217,10 @@ void FNiagaraRibbonVertexFactory::SetDynamicParameterBuffer(const FVertexBuffer*
 
 ///////////////////////////////////////////////////////////////////////////////
 
+IMPLEMENT_VERTEX_FACTORY_PARAMETER_TYPE(FNiagaraRibbonVertexFactory, SF_Vertex, FNiagaraRibbonVertexFactoryShaderParametersVS);
+#if RHI_RAYTRACING
+IMPLEMENT_VERTEX_FACTORY_PARAMETER_TYPE(FNiagaraRibbonVertexFactory, SF_Compute, FNiagaraRibbonVertexFactoryShaderParametersVS);
+IMPLEMENT_VERTEX_FACTORY_PARAMETER_TYPE(FNiagaraRibbonVertexFactory, SF_RayHitGroup, FNiagaraRibbonVertexFactoryShaderParametersVS);
+#endif
+IMPLEMENT_VERTEX_FACTORY_PARAMETER_TYPE(FNiagaraRibbonVertexFactory, SF_Pixel, FNiagaraRibbonVertexFactoryShaderParametersPS);
 IMPLEMENT_VERTEX_FACTORY_TYPE(FNiagaraRibbonVertexFactory, "/Plugin/FX/Niagara/Private/NiagaraRibbonVertexFactory.ush", true, false, true, false, false);
diff --git a/Engine/Plugins/FX/Niagara/Source/NiagaraVertexFactories/Private/NiagaraSortingGPU.cpp b/Engine/Plugins/FX/Niagara/Source/NiagaraVertexFactories/Private/NiagaraSortingGPU.cpp
--- a/Engine/Plugins/FX/Niagara/Source/NiagaraVertexFactories/Private/NiagaraSortingGPU.cpp
+++ b/Engine/Plugins/FX/Niagara/Source/NiagaraVertexFactories/Private/NiagaraSortingGPU.cpp
@@ -52,7 +52,7 @@ FNiagaraSortKeyGenCS::FNiagaraSortKeyGenCS(const ShaderMetaType::CompiledShaderI
 	OutParticleIndices.Bind(Initializer.ParameterMap, TEXT("OutParticleIndices"));
 }
 
-bool FNiagaraSortKeyGenCS::Serialize(FArchive& Ar)
+/*bool FNiagaraSortKeyGenCS::Serialize(FArchive& Ar)
 {
 	bool bShaderHasOutdatedParameters = FGlobalShader::Serialize(Ar);
 	Ar << NiagaraParticleDataFloat;
@@ -66,18 +66,18 @@ bool FNiagaraSortKeyGenCS::Serialize(FArchive& Ar)
 	Ar << OutKeys;
 	Ar << OutParticleIndices;
 	return bShaderHasOutdatedParameters;
-}
+}*/
 
 void FNiagaraSortKeyGenCS::SetOutput(FRHICommandList& RHICmdList, FRHIUnorderedAccessView* OutKeysUAV, FRHIUnorderedAccessView* OutIndicesUAV)
 {
-	FRHIComputeShader* ComputeShaderRHI = GetComputeShader();
+	FRHIComputeShader* ComputeShaderRHI = RHICmdList.GetBoundComputeShader();
 	SetUAVParameter(RHICmdList, ComputeShaderRHI, OutKeys, OutKeysUAV);
 	SetUAVParameter(RHICmdList, ComputeShaderRHI, OutParticleIndices, OutIndicesUAV);
 }
 
 void FNiagaraSortKeyGenCS::SetParameters(FRHICommandList& RHICmdList, const FNiagaraGPUSortInfo& SortInfo, uint32 EmitterKey, int32 OutputOffset, const FUintVector4& SortKeyParamsValue)
 {
-	FRHIComputeShader* ComputeShaderRHI = GetComputeShader();
+	FRHIComputeShader* ComputeShaderRHI = RHICmdList.GetBoundComputeShader();
 
 	SetSRVParameter(RHICmdList, ComputeShaderRHI, NiagaraParticleDataFloat, SortInfo.ParticleDataFloatSRV);
 	SetShaderValue(RHICmdList, ComputeShaderRHI, FloatDataStride, SortInfo.FloatDataStride);
@@ -99,7 +99,7 @@ void FNiagaraSortKeyGenCS::SetParameters(FRHICommandList& RHICmdList, const FNia
 
 void FNiagaraSortKeyGenCS::UnbindBuffers(FRHICommandList& RHICmdList)
 {
-	FRHIComputeShader* ComputeShaderRHI = GetComputeShader();
+	FRHIComputeShader* ComputeShaderRHI = RHICmdList.GetBoundComputeShader();
 	SetUAVParameter(RHICmdList, ComputeShaderRHI, NiagaraParticleDataFloat, nullptr);
 	SetUAVParameter(RHICmdList, ComputeShaderRHI, GPUParticleCountBuffer, nullptr);
 	SetUAVParameter(RHICmdList, ComputeShaderRHI, OutKeys, nullptr);
diff --git a/Engine/Plugins/FX/Niagara/Source/NiagaraVertexFactories/Private/NiagaraSpriteVertexFactory.cpp b/Engine/Plugins/FX/Niagara/Source/NiagaraVertexFactories/Private/NiagaraSpriteVertexFactory.cpp
--- a/Engine/Plugins/FX/Niagara/Source/NiagaraVertexFactories/Private/NiagaraSpriteVertexFactory.cpp
+++ b/Engine/Plugins/FX/Niagara/Source/NiagaraVertexFactories/Private/NiagaraSpriteVertexFactory.cpp
@@ -21,21 +21,17 @@ TGlobalResource<FNullDynamicParameterVertexBuffer> GNullNiagaraDynamicParameterV
  */
 class FNiagaraSpriteVertexFactoryShaderParameters : public FVertexFactoryShaderParameters
 {
+	DECLARE_INLINE_TYPE_LAYOUT(FNiagaraSpriteVertexFactoryShaderParameters, NonVirtual);
 public:
-
-	virtual void Bind(const FShaderParameterMap& ParameterMap) override
-	{
-	}
-
-	virtual void Serialize(FArchive& Ar) override
-	{
-	}
+	
+	
 };
 
 class FNiagaraSpriteVertexFactoryShaderParametersVS : public FNiagaraSpriteVertexFactoryShaderParameters
 {
+	DECLARE_INLINE_TYPE_LAYOUT(FNiagaraSpriteVertexFactoryShaderParametersVS, NonVirtual);
 public:
-	virtual void Bind(const FShaderParameterMap& ParameterMap) override
+	void Bind(const FShaderParameterMap& ParameterMap)
 	{
 		NumCutoutVerticesPerFrame.Bind(ParameterMap, TEXT("NumCutoutVerticesPerFrame"));
 		CutoutGeometry.Bind(ParameterMap, TEXT("CutoutGeometry"));
@@ -71,42 +67,7 @@ public:
 		SortedIndicesOffset.Bind(ParameterMap, TEXT("SortedIndicesOffset"));
 	}
 
-	virtual void Serialize(FArchive& Ar) override
-	{
-		Ar << NumCutoutVerticesPerFrame;
-		Ar << CutoutGeometry;
-		Ar << ParticleFacingMode;
-		Ar << ParticleAlignmentMode;
-
-		Ar << NiagaraParticleDataPosition;
-		Ar << NiagaraParticleDataVelocity;
-		Ar << NiagaraParticleDataColor;
-		Ar << NiagaraParticleDataRotation;
-		Ar << NiagaraParticleDataSize;
-		Ar << NiagaraParticleDataFacing;
-		Ar << NiagaraParticleDataAlignment;
-		Ar << NiagaraParticleDataSubImage;
-		Ar << NiagaraParticleDataCameraOffset;
-		Ar << NiagaraParticleDataUVScale;
-		Ar << NiagaraParticleDataNormalizedAge;
-		Ar << NiagaraParticleDataMaterialRandom;
-		Ar << NiagaraParticleDataCustomSorting;
-		Ar << NiagaraParticleDataMaterialParam0;
-		Ar << NiagaraParticleDataMaterialParam1;
-		Ar << NiagaraParticleDataMaterialParam2;
-		Ar << NiagaraParticleDataMaterialParam3;
-
-		Ar << FloatDataStride;
-
-//  		Ar << NiagaraParticleDataInt;
-//  		Ar << Int32DataOffset;
-//  		Ar << Int32DataStride;
-
-		Ar << SortedIndices;
-		Ar << SortedIndicesOffset;
-	}
-
-	virtual void GetElementShaderBindings(
+	void GetElementShaderBindings(
 		const FSceneInterface* Scene,
 		const FSceneView* View,
 		const FMeshMaterialShader* Shader,
@@ -115,7 +76,7 @@ public:
 		const FVertexFactory* VertexFactory,
 		const FMeshBatchElement& BatchElement,
 		class FMeshDrawSingleShaderBindings& ShaderBindings,
-		FVertexInputStreamArray& VertexStreams) const override
+		FVertexInputStreamArray& VertexStreams) const
 	{
 		FNiagaraSpriteVertexFactory* SpriteVF = (FNiagaraSpriteVertexFactory*)VertexFactory;
 		ShaderBindings.Add(Shader->GetUniformBufferParameter<FNiagaraSpriteUniformParameters>(), SpriteVF->GetSpriteUniformBuffer() );
@@ -151,46 +112,48 @@ public:
 		ShaderBindings.Add(SortedIndices, SpriteVF->GetSortedIndicesSRV() ? SpriteVF->GetSortedIndicesSRV() : GFNiagaraNullSortedIndicesVertexBuffer.VertexBufferSRV.GetReference());
 		ShaderBindings.Add(SortedIndicesOffset, SpriteVF->GetSortedIndicesOffset());
 	}
+
 private:
-	FShaderParameter NumCutoutVerticesPerFrame;
-
-	FShaderParameter ParticleAlignmentMode;
-	FShaderParameter ParticleFacingMode;
-
-	FShaderResourceParameter CutoutGeometry;
-
-	FShaderResourceParameter NiagaraParticleDataPosition;
-	FShaderResourceParameter NiagaraParticleDataVelocity;
-	FShaderResourceParameter NiagaraParticleDataColor;
-	FShaderResourceParameter NiagaraParticleDataRotation;
-	FShaderResourceParameter NiagaraParticleDataSize;
-	FShaderResourceParameter NiagaraParticleDataFacing;
-	FShaderResourceParameter NiagaraParticleDataAlignment;
-	FShaderResourceParameter NiagaraParticleDataSubImage;
-	FShaderResourceParameter NiagaraParticleDataCameraOffset;
-	FShaderResourceParameter NiagaraParticleDataUVScale;
-	FShaderResourceParameter NiagaraParticleDataNormalizedAge;
-	FShaderResourceParameter NiagaraParticleDataMaterialRandom;
-	FShaderResourceParameter NiagaraParticleDataCustomSorting;
-	FShaderResourceParameter NiagaraParticleDataMaterialParam0;
-	FShaderResourceParameter NiagaraParticleDataMaterialParam1;
-	FShaderResourceParameter NiagaraParticleDataMaterialParam2;
-	FShaderResourceParameter NiagaraParticleDataMaterialParam3;
-	FShaderParameter FloatDataStride;
-
-//  	FShaderResourceParameter NiagaraParticleDataInt;
-//  	FShaderParameter Int32DataOffset;
-//  	FShaderParameter Int32DataStride;
 	
-	FShaderResourceParameter SortedIndices;
-	FShaderParameter SortedIndicesOffset;
+		LAYOUT_FIELD(FShaderParameter, NumCutoutVerticesPerFrame);
+
+		LAYOUT_FIELD(FShaderParameter, ParticleAlignmentMode);
+		LAYOUT_FIELD(FShaderParameter, ParticleFacingMode);
+
+		LAYOUT_FIELD(FShaderResourceParameter, CutoutGeometry);
+
+		LAYOUT_FIELD(FShaderResourceParameter, NiagaraParticleDataPosition);
+		LAYOUT_FIELD(FShaderResourceParameter, NiagaraParticleDataVelocity);
+		LAYOUT_FIELD(FShaderResourceParameter, NiagaraParticleDataColor);
+		LAYOUT_FIELD(FShaderResourceParameter, NiagaraParticleDataRotation);
+		LAYOUT_FIELD(FShaderResourceParameter, NiagaraParticleDataSize);
+		LAYOUT_FIELD(FShaderResourceParameter, NiagaraParticleDataFacing);
+		LAYOUT_FIELD(FShaderResourceParameter, NiagaraParticleDataAlignment);
+		LAYOUT_FIELD(FShaderResourceParameter, NiagaraParticleDataSubImage);
+		LAYOUT_FIELD(FShaderResourceParameter, NiagaraParticleDataCameraOffset);
+		LAYOUT_FIELD(FShaderResourceParameter, NiagaraParticleDataUVScale);
+		LAYOUT_FIELD(FShaderResourceParameter, NiagaraParticleDataNormalizedAge);
+		LAYOUT_FIELD(FShaderResourceParameter, NiagaraParticleDataMaterialRandom);
+		LAYOUT_FIELD(FShaderResourceParameter, NiagaraParticleDataCustomSorting);
+		LAYOUT_FIELD(FShaderResourceParameter, NiagaraParticleDataMaterialParam0);
+		LAYOUT_FIELD(FShaderResourceParameter, NiagaraParticleDataMaterialParam1);
+		LAYOUT_FIELD(FShaderResourceParameter, NiagaraParticleDataMaterialParam2);
+		LAYOUT_FIELD(FShaderResourceParameter, NiagaraParticleDataMaterialParam3);
+		LAYOUT_FIELD(FShaderParameter, FloatDataStride);
+
+	//  	LAYOUT_FIELD(FShaderResourceParameter, NiagaraParticleDataInt);
+	//  	LAYOUT_FIELD(FShaderParameter, Int32DataOffset);
+	//  	LAYOUT_FIELD(FShaderParameter, Int32DataStride);
+	
+		LAYOUT_FIELD(FShaderResourceParameter, SortedIndices);
+		LAYOUT_FIELD(FShaderParameter, SortedIndicesOffset);
+	
 };
 
 class FNiagaraSpriteVertexFactoryShaderParametersPS : public FNiagaraSpriteVertexFactoryShaderParameters
 {
 public:
-
-	virtual void GetElementShaderBindings(
+	void GetElementShaderBindings(
 		const FSceneInterface* Scene,
 		const FSceneView* View,
 		const FMeshMaterialShader* Shader,
@@ -199,7 +162,7 @@ public:
 		const FVertexFactory* VertexFactory,
 		const FMeshBatchElement& BatchElement,
 		class FMeshDrawSingleShaderBindings& ShaderBindings,
-		FVertexInputStreamArray& VertexStreams) const override
+		FVertexInputStreamArray& VertexStreams) const
 	{
 		FNiagaraSpriteVertexFactory* SpriteVF = (FNiagaraSpriteVertexFactory*)VertexFactory;
 		ShaderBindings.Add(Shader->GetUniformBufferParameter<FNiagaraSpriteUniformParameters>(), SpriteVF->GetSpriteUniformBuffer() );
@@ -251,17 +214,17 @@ public:
 /** The simple element vertex declaration. */
 static TGlobalResource<FNiagaraSpriteVertexDeclaration> GParticleSpriteVertexDeclaration;
 
-bool FNiagaraSpriteVertexFactory::ShouldCompilePermutation(EShaderPlatform Platform, const class FMaterial* Material, const class FShaderType* ShaderType)
+bool FNiagaraSpriteVertexFactory::ShouldCompilePermutation(const FVertexFactoryShaderPermutationParameters& Parameters)
 {
-	return (FNiagaraUtilities::SupportsNiagaraRendering(Platform)) && (Material->IsUsedWithNiagaraSprites() || Material->IsSpecialEngineMaterial());
+	return (FNiagaraUtilities::SupportsNiagaraRendering(Parameters.Platform)) && (Parameters.MaterialParameters.bIsUsedWithNiagaraSprites || Parameters.MaterialParameters.bIsSpecialEngineMaterial);
 }
 
 /**
  * Can be overridden by FVertexFactory subclasses to modify their compile environment just before compilation occurs.
  */
-void FNiagaraSpriteVertexFactory::ModifyCompilationEnvironment(const FVertexFactoryType* Type, EShaderPlatform Platform, const class FMaterial* Material, FShaderCompilerEnvironment& OutEnvironment)
+void FNiagaraSpriteVertexFactory::ModifyCompilationEnvironment(const FVertexFactoryShaderPermutationParameters& Parameters, FShaderCompilerEnvironment& OutEnvironment)
 {
-	FNiagaraVertexFactoryBase::ModifyCompilationEnvironment(Type, Platform, Material, OutEnvironment);
+	FNiagaraVertexFactoryBase::ModifyCompilationEnvironment(Parameters, OutEnvironment);
 
 	OutEnvironment.SetDefine(TEXT("NiagaraVFLooseParameters"),TEXT("NiagaraSpriteVFLooseParameters"));
 
@@ -293,27 +256,10 @@ void FNiagaraSpriteVertexFactory::SetTexCoordBuffer(const FVertexBuffer* InTexCo
 	TexCoordStream.VertexBuffer = InTexCoordBuffer;
 }
 
-FVertexFactoryShaderParameters* FNiagaraSpriteVertexFactory::ConstructShaderParameters(EShaderFrequency ShaderFrequency)
-{
-	if (ShaderFrequency == SF_Vertex)
-	{
-		return new FNiagaraSpriteVertexFactoryShaderParametersVS();
-	}
-	else if (ShaderFrequency == SF_Pixel)
-	{
-		return new FNiagaraSpriteVertexFactoryShaderParametersPS();
-	}
+IMPLEMENT_VERTEX_FACTORY_PARAMETER_TYPE(FNiagaraSpriteVertexFactory, SF_Vertex, FNiagaraSpriteVertexFactoryShaderParametersVS);
+IMPLEMENT_VERTEX_FACTORY_PARAMETER_TYPE(FNiagaraSpriteVertexFactory, SF_Pixel, FNiagaraSpriteVertexFactoryShaderParametersPS);
 #if RHI_RAYTRACING
-	else if (ShaderFrequency == SF_Compute)
-	{
-		return new FNiagaraSpriteVertexFactoryShaderParametersVS();
-	}
-	else if (ShaderFrequency == SF_RayHitGroup)
-	{
-		return new FNiagaraSpriteVertexFactoryShaderParametersVS();
-	}
-#endif
-	return NULL;
-}
-
+IMPLEMENT_VERTEX_FACTORY_PARAMETER_TYPE(FNiagaraSpriteVertexFactory, SF_Compute, FNiagaraSpriteVertexFactoryShaderParametersVS);
+IMPLEMENT_VERTEX_FACTORY_PARAMETER_TYPE(FNiagaraSpriteVertexFactory, SF_RayHitGroup, FNiagaraSpriteVertexFactoryShaderParametersVS);
+#endif // RHI_RAYTRACING
 IMPLEMENT_VERTEX_FACTORY_TYPE(FNiagaraSpriteVertexFactory,"/Plugin/FX/Niagara/Private/NiagaraSpriteVertexFactory.ush",true,false,true,false,false);
