Fixed vector comparisons.

The comparison operators return boolean vectors, but we were assigning the result to scalar booleans. We need to use all() around comparisons, to make sure they work correctly. Since there's no all() opcode in the VM, I've added NiagaraAll() which emulates it on CPU.

Also added support for comparing boolean values (new operator nodes plus support in the VM codegen).

diff --git a/Engine/Plugins/FX/Niagara/Shaders/Private/NiagaraEmitterInstanceShader.usf b/Engine/Plugins/FX/Niagara/Shaders/Private/NiagaraEmitterInstanceShader.usf
--- a/Engine/Plugins/FX/Niagara/Shaders/Private/NiagaraEmitterInstanceShader.usf
+++ b/Engine/Plugins/FX/Niagara/Shaders/Private/NiagaraEmitterInstanceShader.usf
@@ -31,6 +31,20 @@ float3 Reciprocal(float3 x){ return 1.0f/x; }
 float2 Reciprocal(float2 x){ return 1.0f/x;}
 float Reciprocal(float x){ return 1.0f/x; }
 
+// Annoyingly, all(bool) doesn't exist, so we'll make a NiagaraAll() function which
+// works with both scalars and vectors.
+bool NiagaraAll(bool b) { return b; }
+#if GPU_SIMULATION
+bool NiagaraAll(bool2 b) { return all(b); }
+bool NiagaraAll(bool3 b) { return all(b); }
+bool NiagaraAll(bool4 b) { return all(b); }
+#else
+// No all() opcode in our VM, emulate it.
+bool NiagaraAll(bool2 b) { return b.x && b.y; }
+bool NiagaraAll(bool3 b) { return b.x && b.y && b.z; }
+bool NiagaraAll(bool4 b) { return b.x && b.y && b.z && b.w; }
+#endif
+
 /* -----------------------------------------------------------------
  * GPU simulation utility functions
  * -----------------------------------------------------------------
diff --git a/Engine/Plugins/FX/Niagara/Source/NiagaraEditor/Private/NiagaraEditorCommon.cpp b/Engine/Plugins/FX/Niagara/Source/NiagaraEditor/Private/NiagaraEditorCommon.cpp
--- a/Engine/Plugins/FX/Niagara/Source/NiagaraEditor/Private/NiagaraEditorCommon.cpp
+++ b/Engine/Plugins/FX/Niagara/Source/NiagaraEditor/Private/NiagaraEditorCommon.cpp
@@ -874,7 +874,7 @@ void FNiagaraOpInfo::Init()
 		Op->Keywords = FText::FromString(TEXT("<"));
 		Op->Inputs.Add(FNiagaraOpInOutInfo(A, Type, AText, AText, DefaultStr_Zero));
 		Op->Inputs.Add(FNiagaraOpInOutInfo(B, Type, BText, BText, DefaultStr_One));
-		Op->Outputs.Add(FNiagaraOpInOutInfo(Result, FNiagaraTypeDefinition::GetBoolDef(), ResultText, ResultText, DefaultStr_One, TEXT("{0} < {1}")));
+		Op->Outputs.Add(FNiagaraOpInOutInfo(Result, FNiagaraTypeDefinition::GetBoolDef(), ResultText, ResultText, DefaultStr_One, TEXT("NiagaraAll({0} < {1})")));
 		Op->BuildName(TEXT("CmpLT"), CategoryName);
 		OpInfoMap.Add(Op->Name) = Idx;
 
@@ -886,7 +886,7 @@ void FNiagaraOpInfo::Init()
 		Op->Keywords = FText::FromString(TEXT("<="));
 		Op->Inputs.Add(FNiagaraOpInOutInfo(A, Type, AText, AText, DefaultStr_Zero));
 		Op->Inputs.Add(FNiagaraOpInOutInfo(B, Type, BText, BText, DefaultStr_One));
-		Op->Outputs.Add(FNiagaraOpInOutInfo(Result, FNiagaraTypeDefinition::GetBoolDef(), ResultText, ResultText, DefaultStr_One, TEXT("{0} <= {1}")));
+		Op->Outputs.Add(FNiagaraOpInOutInfo(Result, FNiagaraTypeDefinition::GetBoolDef(), ResultText, ResultText, DefaultStr_One, TEXT("NiagaraAll({0} <= {1})")));
 		Op->BuildName(TEXT("CmpLE"), CategoryName);
 		OpInfoMap.Add(Op->Name) = Idx;
 
@@ -898,7 +898,7 @@ void FNiagaraOpInfo::Init()
 		Op->Keywords = FText::FromString(TEXT(">"));
 		Op->Inputs.Add(FNiagaraOpInOutInfo(A, Type, AText, AText, DefaultStr_Zero));
 		Op->Inputs.Add(FNiagaraOpInOutInfo(B, Type, BText, BText, DefaultStr_One));
-		Op->Outputs.Add(FNiagaraOpInOutInfo(Result, FNiagaraTypeDefinition::GetBoolDef(), ResultText, ResultText, DefaultStr_One, TEXT("{0} > {1}")));
+		Op->Outputs.Add(FNiagaraOpInOutInfo(Result, FNiagaraTypeDefinition::GetBoolDef(), ResultText, ResultText, DefaultStr_One, TEXT("NiagaraAll({0} > {1})")));
 		Op->BuildName(TEXT("CmpGT"), CategoryName);
 		OpInfoMap.Add(Op->Name) = Idx;
 
@@ -910,7 +910,7 @@ void FNiagaraOpInfo::Init()
 		Op->Keywords = FText::FromString(TEXT(">="));
 		Op->Inputs.Add(FNiagaraOpInOutInfo(A, Type, AText, AText, DefaultStr_Zero));
 		Op->Inputs.Add(FNiagaraOpInOutInfo(B, Type, BText, BText, DefaultStr_One));
-		Op->Outputs.Add(FNiagaraOpInOutInfo(Result, FNiagaraTypeDefinition::GetBoolDef(), ResultText, ResultText, DefaultStr_One, TEXT("{0} >= {1}")));
+		Op->Outputs.Add(FNiagaraOpInOutInfo(Result, FNiagaraTypeDefinition::GetBoolDef(), ResultText, ResultText, DefaultStr_One, TEXT("NiagaraAll({0} >= {1})")));
 		Op->BuildName(TEXT("CmpGE"), CategoryName);
 		OpInfoMap.Add(Op->Name) = Idx;
 
@@ -922,7 +922,7 @@ void FNiagaraOpInfo::Init()
 		Op->Keywords = FText::FromString(TEXT("=="));
 		Op->Inputs.Add(FNiagaraOpInOutInfo(A, Type, AText, AText, DefaultStr_Zero));
 		Op->Inputs.Add(FNiagaraOpInOutInfo(B, Type, BText, BText, DefaultStr_One));
-		Op->Outputs.Add(FNiagaraOpInOutInfo(Result, FNiagaraTypeDefinition::GetBoolDef(), ResultText, ResultText, DefaultStr_One, TEXT("{0} == {1}")));
+		Op->Outputs.Add(FNiagaraOpInOutInfo(Result, FNiagaraTypeDefinition::GetBoolDef(), ResultText, ResultText, DefaultStr_One, TEXT("NiagaraAll({0} == {1})")));
 		Op->BuildName(TEXT("CmpEQ"), CategoryName);
 		OpInfoMap.Add(Op->Name) = Idx;
 		
@@ -934,7 +934,7 @@ void FNiagaraOpInfo::Init()
 		Op->Keywords = FText::FromString(TEXT("!="));
 		Op->Inputs.Add(FNiagaraOpInOutInfo(A, Type, AText, AText, DefaultStr_Zero));
 		Op->Inputs.Add(FNiagaraOpInOutInfo(B, Type, BText, BText, DefaultStr_One));
-		Op->Outputs.Add(FNiagaraOpInOutInfo(Result, FNiagaraTypeDefinition::GetBoolDef(), ResultText, ResultText, DefaultStr_One, TEXT("{0} != {1}")));
+		Op->Outputs.Add(FNiagaraOpInOutInfo(Result, FNiagaraTypeDefinition::GetBoolDef(), ResultText, ResultText, DefaultStr_One, TEXT("NiagaraAll({0} != {1})")));
 		Op->BuildName(TEXT("CmpNEQ"), CategoryName);
 		OpInfoMap.Add(Op->Name) = Idx;
 	}
@@ -1083,6 +1083,30 @@ void FNiagaraOpInfo::Init()
 	Op->BuildName(TEXT("LogicNot"), BoolCategoryName);
 	OpInfoMap.Add(Op->Name) = Idx;
 
+	Idx = OpInfos.AddDefaulted();
+	Op = &OpInfos[Idx];
+	Op->Category = BoolCategory;
+	Op->FriendlyName = NSLOCTEXT("NiagaraOpInfo", "LogicEq Name", "Bool Equal");
+	Op->Description = NSLOCTEXT("NiagaraOpInfo", "LogicEq Desc", "Result = A == B");
+	Op->Keywords = FText::FromString(TEXT("=="));
+	Op->Inputs.Add(FNiagaraOpInOutInfo(A, BoolType, AText, AText, Default_BoolZero));
+	Op->Inputs.Add(FNiagaraOpInOutInfo(B, BoolType, BText, BText, Default_BoolOne));
+	Op->Outputs.Add(FNiagaraOpInOutInfo(Result, BoolType, ResultText, ResultText, Default_BoolOne, TEXT("NiagaraAll({0} == {1})")));
+	Op->BuildName(TEXT("LogicEq"), BoolCategoryName);
+	OpInfoMap.Add(Op->Name) = Idx;
+
+	Idx = OpInfos.AddDefaulted();
+	Op = &OpInfos[Idx];
+	Op->Category = BoolCategory;
+	Op->FriendlyName = NSLOCTEXT("NiagaraOpInfo", "LogicNEq Name", "Bool Not Equal");
+	Op->Description = NSLOCTEXT("NiagaraOpInfo", "LogicNEq Desc", "Result = A != B");
+	Op->Keywords = FText::FromString(TEXT("!="));
+	Op->Inputs.Add(FNiagaraOpInOutInfo(A, BoolType, AText, AText, Default_BoolZero));
+	Op->Inputs.Add(FNiagaraOpInOutInfo(B, BoolType, BText, BText, Default_BoolOne));
+	Op->Outputs.Add(FNiagaraOpInOutInfo(Result, BoolType, ResultText, ResultText, Default_BoolOne, TEXT("NiagaraAll({0} != {1})")));
+	Op->BuildName(TEXT("LogicNEq"), BoolCategoryName);
+	OpInfoMap.Add(Op->Name) = Idx;
+
 	//////////////////////////////////////////////////////////////////////////
 	//Matrix only ops
 	FString Default_MatrixZero(TEXT(
diff --git a/Engine/Source/Developer/ShaderFormatVectorVM/Private/ir_vm_gen_bytecode_visitor.cpp b/Engine/Source/Developer/ShaderFormatVectorVM/Private/ir_vm_gen_bytecode_visitor.cpp
--- a/Engine/Source/Developer/ShaderFormatVectorVM/Private/ir_vm_gen_bytecode_visitor.cpp
+++ b/Engine/Source/Developer/ShaderFormatVectorVM/Private/ir_vm_gen_bytecode_visitor.cpp
@@ -258,9 +258,11 @@ void BuildExpressionMap()
 		Info = &VMExpressionMap.Add(ir_binop_equal);
 		Info->Add(FVMExpresssionInfo(EVectorVMOp::cmpeq, glsl_type::bool_type, glsl_type::float_type, glsl_type::float_type));
 		Info->Add(FVMExpresssionInfo(EVectorVMOp::cmpeqi, glsl_type::bool_type, glsl_type::int_type, glsl_type::int_type));
+		Info->Add(FVMExpresssionInfo(EVectorVMOp::cmpeqi, glsl_type::bool_type, glsl_type::bool_type, glsl_type::bool_type));
 		Info = &VMExpressionMap.Add(ir_binop_nequal);
 		Info->Add(FVMExpresssionInfo(EVectorVMOp::cmpneq, glsl_type::bool_type, glsl_type::float_type, glsl_type::float_type));
 		Info->Add(FVMExpresssionInfo(EVectorVMOp::cmpneqi, glsl_type::bool_type, glsl_type::int_type, glsl_type::int_type));
+		Info->Add(FVMExpresssionInfo(EVectorVMOp::cmpneqi, glsl_type::bool_type, glsl_type::bool_type, glsl_type::bool_type));
 
 		/**
 		* Returns single boolean for whether all components of operands[0]
